Distance Field AO is now evaluated on a screen grid and using a global distance field for far traces

diff --git a/Engine/Source/Runtime/Renderer/Private/DistanceFieldSurfaceCacheLighting.cpp b/Engine/Source/Runtime/Renderer/Private/DistanceFieldSurfaceCacheLighting.cpp
--- a/Engine/Source/Runtime/Renderer/Private/DistanceFieldSurfaceCacheLighting.cpp
+++ b/Engine/Source/Runtime/Renderer/Private/DistanceFieldSurfaceCacheLighting.cpp
@@ -12,6 +12,7 @@
 #include "SceneFilterRendering.h"
 #include "DistanceFieldLightingShared.h"
 #include "DistanceFieldSurfaceCacheLighting.h"
+#include "DistanceFieldLightingPost.h"
 #include "DistanceFieldGlobalIllumination.h"
 #include "PostProcessAmbientOcclusion.h"
 #include "RHICommandList.h"
@@ -19,6 +20,7 @@
 #include "OneColorShader.h"
 #include "BasePassRendering.h"
 #include "HeightfieldLighting.h"
+#include "GlobalDistanceField.h"
 
 int32 GDistanceFieldAO = 1;
 FAutoConsoleVariableRef CVarDistanceFieldAO(
@@ -34,6 +36,14 @@ bool IsDistanceFieldGIAllowed(const FViewInfo& View)
 		&& (View.Family->EngineShowFlags.VisualizeDistanceFieldGI || (View.Family->EngineShowFlags.DistanceFieldGI && GDistanceFieldGI && View.Family->EngineShowFlags.GlobalIllumination));
 }
 
+int32 GAOUseSurfaceCache = 0;
+FAutoConsoleVariableRef CVarAOUseSurfaceCache(
+	TEXT("r.AOUseSurfaceCache"),
+	GAOUseSurfaceCache,
+	TEXT("Whether to use the surface cache or screen grid for cone tracing.  The screen grid has more constant performance characteristics."),
+	ECVF_Cheat | ECVF_RenderThreadSafe
+	);
+
 int32 GAOPowerOfTwoBetweenLevels = 2;
 FAutoConsoleVariableRef CVarAOPowerOfTwoBetweenLevels(
 	TEXT("r.AOPowerOfTwoBetweenLevels"),
@@ -77,55 +87,6 @@ FAutoConsoleVariableRef CVarAOTrimOldRecordsFraction(
 	ECVF_Cheat | ECVF_RenderThreadSafe
 	);
 
-int32 GAOFillGaps = 1;
-FAutoConsoleVariableRef CVarAOFillGaps(
-	TEXT("r.AOFillGaps"),
-	GAOFillGaps,
-	TEXT("Whether to fill in pixels using a screen space filter that had no valid world space interpolation weight from surface cache samples.\n")
-	TEXT("This is needed whenever r.AOMinLevel is not 0."),
-	ECVF_Cheat | ECVF_RenderThreadSafe
-	);
-
-int32 GAOFillGapsHighQuality = 1;
-FAutoConsoleVariableRef CVarAOFillGapsHighQuality(
-	TEXT("r.AOFillGapsHighQuality"),
-	GAOFillGapsHighQuality,
-	TEXT("Whether to use the higher quality gap filling method that does a 5x5 gather, or the cheaper method that just looks up the 4 grid samples."),
-	ECVF_Cheat | ECVF_RenderThreadSafe
-	);
-
-int32 GAOUseHistory = 1;
-FAutoConsoleVariableRef CVarAOUseHistory(
-	TEXT("r.AOUseHistory"),
-	GAOUseHistory,
-	TEXT("Whether to apply a temporal filter to the distance field AO, which reduces flickering but also adds trails when occluders are moving."),
-	ECVF_Cheat | ECVF_RenderThreadSafe
-	);
-
-int32 GAOHistoryStabilityPass = 1;
-FAutoConsoleVariableRef CVarAOHistoryStabilityPass(
-	TEXT("r.AOHistoryStabilityPass"),
-	GAOHistoryStabilityPass,
-	TEXT("Whether to gather stable results to fill in holes in the temporal reprojection.  Adds some GPU cost but improves temporal stability with foliage."),
-	ECVF_Cheat | ECVF_RenderThreadSafe
-	);
-
-float GAOHistoryWeight = .7f;
-FAutoConsoleVariableRef CVarAOHistoryWeight(
-	TEXT("r.AOHistoryWeight"),
-	GAOHistoryWeight,
-	TEXT("Amount of last frame's AO to lerp into the final result.  Higher values increase stability, lower values have less streaking under occluder movement."),
-	ECVF_Cheat | ECVF_RenderThreadSafe
-	);
-
-float GAOHistoryDistanceThreshold = 20;
-FAutoConsoleVariableRef CVarAOHistoryDistanceThreshold(
-	TEXT("r.AOHistoryDistanceThreshold"),
-	GAOHistoryDistanceThreshold,
-	TEXT("World space distance threshold needed to discard last frame's DFAO results.  Lower values reduce ghosting from characters when near a wall but increase flickering artifacts."),
-	ECVF_Cheat | ECVF_RenderThreadSafe
-	);
-
 float GAORecordRadiusScale = .3f;
 FAutoConsoleVariableRef CVarAORecordRadiusScale(
 	TEXT("r.AORecordRadiusScale"),
@@ -175,7 +136,7 @@ FAutoConsoleVariableRef CVarAOStepExponentScale(
 	ECVF_Cheat | ECVF_RenderThreadSafe
 	);
 
-float GAOMaxViewDistance = 10000;
+float GAOMaxViewDistance = 20000;
 FAutoConsoleVariableRef CVarAOMaxViewDistance(
 	TEXT("r.AOMaxViewDistance"),
 	GAOMaxViewDistance,
@@ -183,14 +144,6 @@ FAutoConsoleVariableRef CVarAOMaxViewDistance(
 	ECVF_Cheat | ECVF_RenderThreadSafe
 	);
 
-float GAOViewFadeDistanceScale = .7f;
-FAutoConsoleVariableRef CVarAOViewFadeDistanceScale(
-	TEXT("r.AOViewFadeDistanceScale"),
-	GAOViewFadeDistanceScale,
-	TEXT("Distance over which AO will fade out as it approaches r.AOMaxViewDistance, as a fraction of r.AOMaxViewDistance."),
-	ECVF_Cheat | ECVF_RenderThreadSafe
-	);
-
 int32 GAOScatterTileCulling = 1;
 FAutoConsoleVariableRef CVarAOScatterTileCulling(
 	TEXT("r.AOScatterTileCulling"),
@@ -243,6 +196,24 @@ DEFINE_LOG_CATEGORY(LogDistanceField);
 
 IMPLEMENT_UNIFORM_BUFFER_STRUCT(FAOSampleData2,TEXT("AOSamples2"));
 
+FDistanceFieldAOParameters::FDistanceFieldAOParameters(float InOcclusionMaxDistance, float InContrast)
+{
+	Contrast = FMath::Clamp(InContrast, .01f, 2.0f);
+	InOcclusionMaxDistance = FMath::Clamp(InOcclusionMaxDistance, 2.0f, 3000.0f);
+
+	if (GAOGlobalDistanceField != 0 && GAOUseSurfaceCache == 0)
+	{
+		extern float GAOGlobalDFStartDistance;
+		ObjectMaxOcclusionDistance = FMath::Min(InOcclusionMaxDistance, GAOGlobalDFStartDistance);
+		GlobalMaxOcclusionDistance = InOcclusionMaxDistance >= GAOGlobalDFStartDistance ? InOcclusionMaxDistance : 0;
+	}
+	else
+	{
+		ObjectMaxOcclusionDistance = InOcclusionMaxDistance;
+		GlobalMaxOcclusionDistance = 0;
+	}
+}
+
 FIntPoint GetBufferSizeForAO()
 {
 	return FIntPoint::DivideAndRoundDown(GSceneRenderTargets.GetBufferSizeXY(), GAODownsampleFactor);
@@ -505,7 +476,6 @@ public:
 		FGlobalShader::ModifyCompilationEnvironment(Platform,OutEnvironment);
 		OutEnvironment.SetDefine(TEXT("THREADGROUP_SIZEX"), GDistanceFieldAOTileSizeX);
 		OutEnvironment.SetDefine(TEXT("THREADGROUP_SIZEY"), GDistanceFieldAOTileSizeY);
-		OutEnvironment.SetDefine(TEXT("NUM_CONE_DIRECTIONS"), NumConeSampleDirections);
 		OutEnvironment.SetDefine(TEXT("DOWNSAMPLE_FACTOR"), GAODownsampleFactor);
 		
 		// To reduce shader compile time of compute shaders with shared memory, doesn't have an impact on generated code with current compiler (June 2010 DX SDK)
@@ -679,6 +649,7 @@ public:
 	FObjectCullPS(const ShaderMetaType::CompiledShaderInitializerType& Initializer)
 		: FGlobalShader(Initializer)
 	{
+		ObjectParameters.Bind(Initializer.ParameterMap);
 		AOParameters.Bind(Initializer.ParameterMap);
 		TileHeadDataUnpacked.Bind(Initializer.ParameterMap, TEXT("TileHeadDataUnpacked"));
 		TileArrayData.Bind(Initializer.ParameterMap, TEXT("TileArrayData"));
@@ -691,6 +662,7 @@ public:
 	{
 		const FPixelShaderRHIParamRef ShaderRHI = GetPixelShader();
 		FGlobalShader::SetParameters(RHICmdList, ShaderRHI, View);
+		ObjectParameters.Set(RHICmdList, ShaderRHI, GAOCulledObjectBuffers.Buffers);
 		AOParameters.Set(RHICmdList, ShaderRHI, Parameters);
 
 		FTileIntersectionResources* TileIntersectionResources = ((FSceneViewState*)View.State)->AOTileIntersectionResources;
@@ -725,6 +697,7 @@ public:
 	virtual bool Serialize(FArchive& Ar) override
 	{
 		bool bShaderHasOutdatedParameters = FGlobalShader::Serialize(Ar);
+		Ar << ObjectParameters;
 		Ar << AOParameters;
 		Ar << TileHeadDataUnpacked;
 		Ar << TileArrayData;
@@ -735,6 +708,7 @@ public:
 	}
 
 private:
+	FDistanceFieldCulledObjectBufferParameters ObjectParameters;
 	FAOParameters AOParameters;
 	FRWShaderParameter TileHeadDataUnpacked;
 	FRWShaderParameter TileArrayData;
@@ -762,7 +736,6 @@ public:
 		FGlobalShader::ModifyCompilationEnvironment(Platform,OutEnvironment);
 		OutEnvironment.SetDefine(TEXT("THREADGROUP_SIZEX"), GDistanceFieldAOTileSizeX);
 		OutEnvironment.SetDefine(TEXT("THREADGROUP_SIZEY"), GDistanceFieldAOTileSizeY);
-		OutEnvironment.SetDefine(TEXT("NUM_CONE_DIRECTIONS"), NumConeSampleDirections);
 		OutEnvironment.SetDefine(TEXT("DOWNSAMPLE_FACTOR"), GAODownsampleFactor);
 		// To reduce shader compile time of compute shaders with shared memory, doesn't have an impact on generated code with current compiler (June 2010 DX SDK)
 		OutEnvironment.CompilerFlags.Add(CFLAG_StandardOptimization);
@@ -775,6 +748,7 @@ public:
 		ObjectParameters.Bind(Initializer.ParameterMap);
 		AOParameters.Bind(Initializer.ParameterMap);
 		TileHeadData.Bind(Initializer.ParameterMap, TEXT("TileHeadData"));
+		TileHeadDataUnpacked.Bind(Initializer.ParameterMap, TEXT("TileHeadDataUnpacked"));
 		TileArrayData.Bind(Initializer.ParameterMap, TEXT("TileArrayData"));
 		TileArrayNextAllocation.Bind(Initializer.ParameterMap, TEXT("TileArrayNextAllocation"));
 		NumGroups.Bind(Initializer.ParameterMap, TEXT("NumGroups"));
@@ -796,6 +770,7 @@ public:
 		FTileIntersectionResources* TileIntersectionResources = ((FSceneViewState*)View.State)->AOTileIntersectionResources;
 
 		TileHeadData.SetBuffer(RHICmdList, ShaderRHI, TileIntersectionResources->TileHeadData);
+		TileHeadDataUnpacked.SetBuffer(RHICmdList, ShaderRHI, TileIntersectionResources->TileHeadDataUnpacked);
 		TileArrayData.SetBuffer(RHICmdList, ShaderRHI, TileIntersectionResources->TileArrayData);
 		TileArrayNextAllocation.SetBuffer(RHICmdList, ShaderRHI, TileIntersectionResources->TileArrayNextAllocation);
 
@@ -807,6 +782,7 @@ public:
 	void UnsetParameters(FRHICommandList& RHICmdList)
 	{
 		TileHeadData.UnsetUAV(RHICmdList, GetComputeShader());
+		TileHeadDataUnpacked.UnsetUAV(RHICmdList, GetComputeShader());
 		TileArrayData.UnsetUAV(RHICmdList, GetComputeShader());
 		TileArrayNextAllocation.UnsetUAV(RHICmdList, GetComputeShader());
 	}
@@ -818,6 +794,7 @@ public:
 		Ar << ObjectParameters;
 		Ar << AOParameters;
 		Ar << TileHeadData;
+		Ar << TileHeadDataUnpacked;
 		Ar << TileArrayData;
 		Ar << TileArrayNextAllocation;
 		Ar << NumGroups;
@@ -831,6 +808,7 @@ private:
 	FDistanceFieldCulledObjectBufferParameters ObjectParameters;
 	FAOParameters AOParameters;
 	FRWShaderParameter TileHeadData;
+	FRWShaderParameter TileHeadDataUnpacked;
 	FRWShaderParameter TileArrayData;
 	FRWShaderParameter TileArrayNextAllocation;
 	FShaderParameter ViewDimensionsParameter;
@@ -839,46 +817,6 @@ private:
 
 IMPLEMENT_SHADER_TYPE(,FDistanceFieldBuildTileListCS,TEXT("DistanceFieldSurfaceCacheLightingCompute"),TEXT("DistanceFieldAOBuildTileListMain"),SF_Compute);
 
-class FAOLevelParameters
-{
-public:
-	void Bind(const FShaderParameterMap& ParameterMap)
-	{
-		CurrentLevelDownsampleFactor.Bind(ParameterMap, TEXT("CurrentLevelDownsampleFactor"));
-		AOBufferSize.Bind(ParameterMap, TEXT("AOBufferSize"));
-		DownsampleFactorToBaseLevel.Bind(ParameterMap, TEXT("DownsampleFactorToBaseLevel"));
-		BaseLevelTexelSize.Bind(ParameterMap, TEXT("BaseLevelTexelSize"));
-	}
-
-	template<typename TParamRef>
-	void Set(FRHICommandList& RHICmdList, const TParamRef& ShaderRHI, const FSceneView& View, int32 CurrentLevelDownsampleFactorValue)
-	{
-		SetShaderValue(RHICmdList, ShaderRHI, CurrentLevelDownsampleFactor, CurrentLevelDownsampleFactorValue);
-
-		// Round up, to match render target allocation
-		const FVector2D AOBufferSizeValue = FIntPoint::DivideAndRoundUp(GSceneRenderTargets.GetBufferSizeXY(), CurrentLevelDownsampleFactorValue);
-		SetShaderValue(RHICmdList, ShaderRHI, AOBufferSize, AOBufferSizeValue);
-
-		SetShaderValue(RHICmdList, ShaderRHI, DownsampleFactorToBaseLevel, CurrentLevelDownsampleFactorValue / GAODownsampleFactor);
-
-		const FIntPoint DownsampledBufferSize = GetBufferSizeForAO();
-		const FVector2D BaseLevelBufferSizeValue(1.0f / DownsampledBufferSize.X, 1.0f / DownsampledBufferSize.Y);
-		SetShaderValue(RHICmdList, ShaderRHI, BaseLevelTexelSize, BaseLevelBufferSizeValue);
-	}
-
-	friend FArchive& operator<<(FArchive& Ar,FAOLevelParameters& P)
-	{
-		Ar << P.CurrentLevelDownsampleFactor << P.AOBufferSize << P.DownsampleFactorToBaseLevel << P.BaseLevelTexelSize;
-		return Ar;
-	}
-
-private:
-	FShaderParameter CurrentLevelDownsampleFactor;
-	FShaderParameter AOBufferSize;
-	FShaderParameter DownsampleFactorToBaseLevel;
-	FShaderParameter BaseLevelTexelSize;
-};
-
 class FComputeDistanceFieldNormalPS : public FGlobalShader
 {
 	DECLARE_SHADER_TYPE(FComputeDistanceFieldNormalPS, Global);
@@ -991,8 +929,6 @@ public:
 private:
 
 	FRWShaderParameter DistanceFieldNormal;
-	FDistanceFieldCulledObjectBufferParameters ObjectParameters;
-	FShaderParameter UVToTileHead;
 	FDeferredPixelShaderParameters DeferredParameters;
 	FAOParameters AOParameters;
 };
@@ -1153,12 +1089,7 @@ public:
 	TCopyIrradianceCacheSamplesCS(const ShaderMetaType::CompiledShaderInitializerType& Initializer)
 		: FGlobalShader(Initializer)
 	{
-		IrradianceCachePositionRadius.Bind(Initializer.ParameterMap, TEXT("IrradianceCachePositionRadius"));
-		IrradianceCacheNormal.Bind(Initializer.ParameterMap, TEXT("IrradianceCacheNormal"));
-		OccluderRadius.Bind(Initializer.ParameterMap, TEXT("OccluderRadius"));
-		IrradianceCacheBentNormal.Bind(Initializer.ParameterMap, TEXT("IrradianceCacheBentNormal"));
-		IrradianceCacheIrradiance.Bind(Initializer.ParameterMap, TEXT("IrradianceCacheIrradiance"));
-		IrradianceCacheTileCoordinate.Bind(Initializer.ParameterMap, TEXT("IrradianceCacheTileCoordinate"));
+		CacheParameters.Bind(Initializer.ParameterMap);
 		DrawParameters.Bind(Initializer.ParameterMap, TEXT("DrawParameters"));
 		CopyIrradianceCachePositionRadius.Bind(Initializer.ParameterMap, TEXT("CopyIrradianceCachePositionRadius"));
 		CopyIrradianceCacheNormal.Bind(Initializer.ParameterMap, TEXT("CopyIrradianceCacheNormal"));
@@ -1182,12 +1113,7 @@ public:
 		const FScene* Scene = (const FScene*)View.Family->Scene;
 		FSurfaceCacheResources& SurfaceCacheResources = *Scene->SurfaceCacheResources;
 
-		SetSRVParameter(RHICmdList, ShaderRHI, IrradianceCachePositionRadius, SurfaceCacheResources.Level[DepthLevel]->PositionAndRadius.SRV);
-		SetSRVParameter(RHICmdList, ShaderRHI, IrradianceCacheNormal, SurfaceCacheResources.Level[DepthLevel]->Normal.SRV);
-		SetSRVParameter(RHICmdList, ShaderRHI, OccluderRadius, SurfaceCacheResources.Level[DepthLevel]->OccluderRadius.SRV);
-		SetSRVParameter(RHICmdList, ShaderRHI, IrradianceCacheBentNormal, SurfaceCacheResources.Level[DepthLevel]->BentNormal.SRV);
-		SetSRVParameter(RHICmdList, ShaderRHI, IrradianceCacheIrradiance, SurfaceCacheResources.Level[DepthLevel]->Irradiance.SRV);
-		SetSRVParameter(RHICmdList, ShaderRHI, IrradianceCacheTileCoordinate, SurfaceCacheResources.Level[DepthLevel]->TileCoordinate.SRV);
+		CacheParameters.Set(RHICmdList, ShaderRHI, *SurfaceCacheResources.Level[DepthLevel]);
 		SetSRVParameter(RHICmdList, ShaderRHI, DrawParameters, SurfaceCacheResources.Level[DepthLevel]->ScatterDrawParameters.SRV);
 
 		CopyIrradianceCachePositionRadius.SetBuffer(RHICmdList, ShaderRHI, SurfaceCacheResources.TempResources->PositionAndRadius);
@@ -1215,12 +1141,7 @@ public:
 	virtual bool Serialize(FArchive& Ar) override
 	{		
 		bool bShaderHasOutdatedParameters = FGlobalShader::Serialize(Ar);
-		Ar << IrradianceCachePositionRadius;
-		Ar << IrradianceCacheNormal;
-		Ar << OccluderRadius;
-		Ar << IrradianceCacheBentNormal;
-		Ar << IrradianceCacheIrradiance;
-		Ar << IrradianceCacheTileCoordinate;
+		Ar << CacheParameters;
 		Ar << DrawParameters;
 		Ar << CopyIrradianceCachePositionRadius;
 		Ar << CopyIrradianceCacheNormal;
@@ -1235,12 +1156,7 @@ public:
 
 private:
 
-	FShaderResourceParameter IrradianceCachePositionRadius;
-	FShaderResourceParameter IrradianceCacheNormal;
-	FShaderResourceParameter OccluderRadius;
-	FShaderResourceParameter IrradianceCacheBentNormal;
-	FShaderResourceParameter IrradianceCacheIrradiance;
-	FShaderResourceParameter IrradianceCacheTileCoordinate;
+	FIrradianceCacheParameters CacheParameters;
 	FShaderResourceParameter DrawParameters;
 	FRWShaderParameter CopyIrradianceCachePositionRadius;
 	FRWShaderParameter CopyIrradianceCacheNormal;
@@ -1324,7 +1240,6 @@ public:
 		FGlobalShader::ModifyCompilationEnvironment(Platform,OutEnvironment);
 		OutEnvironment.SetDefine(TEXT("THREADGROUP_SIZEX"), GDistanceFieldAOTileSizeX);
 		OutEnvironment.SetDefine(TEXT("THREADGROUP_SIZEY"), GDistanceFieldAOTileSizeY);
-		OutEnvironment.SetDefine(TEXT("NUM_CONE_DIRECTIONS"), NumConeSampleDirections);
 		OutEnvironment.SetDefine(TEXT("DOWNSAMPLE_FACTOR"), GAODownsampleFactor);
 
 		// To reduce shader compile time of compute shaders with shared memory, doesn't have an impact on generated code with current compiler (June 2010 DX SDK)
@@ -1472,34 +1387,10 @@ IMPLEMENT_SHADER_TYPE(,FPopulateCacheCS,TEXT("DistanceFieldSurfaceCacheLightingC
 IMPLEMENT_SHADER_TYPE(template<>,TSetupFinalGatherIndirectArgumentsCS<true>,TEXT("DistanceFieldSurfaceCacheLightingCompute"),TEXT("SetupFinalGatherIndirectArgumentsCS"),SF_Compute);
 IMPLEMENT_SHADER_TYPE(template<>,TSetupFinalGatherIndirectArgumentsCS<false>,TEXT("DistanceFieldSurfaceCacheLightingCompute"),TEXT("SetupFinalGatherIndirectArgumentsCS"),SF_Compute);
 
-class FDebugBuffer : public FRenderResource
-{
-public:
-	virtual void InitDynamicRHI() override
-	{
-        if(GMaxRHIFeatureLevel >= ERHIFeatureLevel::SM5)
-        {
-            DebugData.Initialize(sizeof(float)* 4, 50, PF_A32B32G32R32F, BUF_Static);
-        }
-	}
-
-	virtual void ReleaseDynamicRHI() override
-	{
-        if(GMaxRHIFeatureLevel >= ERHIFeatureLevel::SM5)
-        {
-            DebugData.Release();
-        }
-	}
-
-	FRWBuffer DebugData;
-};
-
-TGlobalResource<FDebugBuffer> GDebugBuffer;
-
 template<bool bSupportIrradiance>
-class TConeTraceOcclusionCS : public FGlobalShader
+class TConeTraceSurfaceCacheOcclusionCS : public FGlobalShader
 {
-	DECLARE_SHADER_TYPE(TConeTraceOcclusionCS,Global)
+	DECLARE_SHADER_TYPE(TConeTraceSurfaceCacheOcclusionCS,Global)
 public:
 
 	static bool ShouldCache(EShaderPlatform Platform)
@@ -1510,7 +1401,6 @@ public:
 	static void ModifyCompilationEnvironment(EShaderPlatform Platform, FShaderCompilerEnvironment& OutEnvironment)
 	{
 		FGlobalShader::ModifyCompilationEnvironment(Platform,OutEnvironment);
-		OutEnvironment.SetDefine(TEXT("NUM_CONE_DIRECTIONS"), NumConeSampleDirections);
 		OutEnvironment.SetDefine(TEXT("DOWNSAMPLE_FACTOR"), GAODownsampleFactor);
 		OutEnvironment.SetDefine(TEXT("SUPPORT_IRRADIANCE"), bSupportIrradiance ? TEXT("1") : TEXT("0"));
 
@@ -1518,18 +1408,18 @@ public:
 		OutEnvironment.CompilerFlags.Add(CFLAG_StandardOptimization);
 	}
 
-	TConeTraceOcclusionCS(const ShaderMetaType::CompiledShaderInitializerType& Initializer)
+	TConeTraceSurfaceCacheOcclusionCS(const ShaderMetaType::CompiledShaderInitializerType& Initializer)
 		: FGlobalShader(Initializer)
 	{
 		DeferredParameters.Bind(Initializer.ParameterMap);
 		ObjectParameters.Bind(Initializer.ParameterMap);
 		AOParameters.Bind(Initializer.ParameterMap);
 		IrradianceCachePositionRadius.Bind(Initializer.ParameterMap, TEXT("IrradianceCachePositionRadius"));
+		TileConeDepthRanges.Bind(Initializer.ParameterMap, TEXT("TileConeDepthRanges"));
 		IrradianceCacheNormal.Bind(Initializer.ParameterMap, TEXT("IrradianceCacheNormal"));
 		OccluderRadius.Bind(Initializer.ParameterMap, TEXT("OccluderRadius"));
 		RecordConeVisibility.Bind(Initializer.ParameterMap, TEXT("RecordConeVisibility"));
 		RecordConeData.Bind(Initializer.ParameterMap, TEXT("RecordConeData"));
-		DebugBuffer.Bind(Initializer.ParameterMap, TEXT("DebugBuffer"));
 		IrradianceCacheTileCoordinate.Bind(Initializer.ParameterMap, TEXT("IrradianceCacheTileCoordinate"));
 		ScatterDrawParameters.Bind(Initializer.ParameterMap, TEXT("ScatterDrawParameters"));
 		SavedStartIndex.Bind(Initializer.ParameterMap, TEXT("SavedStartIndex"));
@@ -1550,7 +1440,7 @@ public:
 		RecordRadiusScale.Bind(Initializer.ParameterMap, TEXT("RecordRadiusScale"));
 	}
 
-	TConeTraceOcclusionCS()
+	TConeTraceSurfaceCacheOcclusionCS()
 	{
 	}
 
@@ -1582,7 +1472,6 @@ public:
 		OccluderRadius.SetBuffer(RHICmdList, ShaderRHI, SurfaceCacheResources.Level[DepthLevel]->OccluderRadius);
 		RecordConeVisibility.SetBuffer(RHICmdList, ShaderRHI, GTemporaryIrradianceCacheResources.ConeVisibility);
 		RecordConeData.SetBuffer(RHICmdList, ShaderRHI, GTemporaryIrradianceCacheResources.ConeData);
-		DebugBuffer.SetBuffer(RHICmdList, ShaderRHI, GDebugBuffer.DebugData);
 
 		SetShaderValue(RHICmdList, ShaderRHI, ViewDimensionsParameter, View.ViewRect);
 
@@ -1607,6 +1496,7 @@ public:
 		SetSRVParameter(RHICmdList, ShaderRHI, TileHeadData, TileIntersectionResources->TileHeadData.SRV);
 		SetSRVParameter(RHICmdList, ShaderRHI, TileArrayData, TileIntersectionResources->TileArrayData.SRV);
 		SetSRVParameter(RHICmdList, ShaderRHI, IrradianceCacheTileCoordinate, SurfaceCacheResources.Level[DepthLevel]->TileCoordinate.SRV);
+		SetSRVParameter(RHICmdList, ShaderRHI, TileConeDepthRanges, TileIntersectionResources->TileConeDepthRanges.SRV);
 
 		SetShaderValue(RHICmdList, ShaderRHI, TileListGroupSize, TileListGroupSizeValue);
 
@@ -1654,7 +1544,6 @@ public:
 		OccluderRadius.UnsetUAV(RHICmdList, GetComputeShader());
 		RecordConeVisibility.UnsetUAV(RHICmdList, GetComputeShader());
 		RecordConeData.UnsetUAV(RHICmdList, GetComputeShader());
-		DebugBuffer.UnsetUAV(RHICmdList, GetComputeShader());
 	}
 
 	virtual bool Serialize(FArchive& Ar) override
@@ -1665,11 +1554,11 @@ public:
 		Ar << AOParameters;
 		Ar << AOLevelParameters;
 		Ar << IrradianceCachePositionRadius;
+		Ar << TileConeDepthRanges;
 		Ar << IrradianceCacheNormal;
 		Ar << OccluderRadius;
 		Ar << RecordConeVisibility;
 		Ar << RecordConeData;
-		Ar << DebugBuffer;
 		Ar << IrradianceCacheTileCoordinate;
 		Ar << ScatterDrawParameters;
 		Ar << SavedStartIndex;
@@ -1698,10 +1587,10 @@ private:
 	FAOLevelParameters AOLevelParameters;
 	FShaderResourceParameter IrradianceCacheNormal;
 	FShaderResourceParameter IrradianceCachePositionRadius;
+	FShaderResourceParameter TileConeDepthRanges;
 	FRWShaderParameter OccluderRadius;
 	FRWShaderParameter RecordConeVisibility;
 	FRWShaderParameter RecordConeData;
-	FRWShaderParameter DebugBuffer;
 	FShaderResourceParameter ScatterDrawParameters;
 	FShaderResourceParameter SavedStartIndex;
 	FShaderResourceParameter IrradianceCacheTileCoordinate;
@@ -1721,8 +1610,8 @@ private:
 	FShaderParameter RecordRadiusScale;
 };
 
-IMPLEMENT_SHADER_TYPE(template<>,TConeTraceOcclusionCS<true>,TEXT("DistanceFieldSurfaceCacheLightingCompute"),TEXT("ConeTraceOcclusionCS"),SF_Compute);
-IMPLEMENT_SHADER_TYPE(template<>,TConeTraceOcclusionCS<false>,TEXT("DistanceFieldSurfaceCacheLightingCompute"),TEXT("ConeTraceOcclusionCS"),SF_Compute);
+IMPLEMENT_SHADER_TYPE(template<>,TConeTraceSurfaceCacheOcclusionCS<true>,TEXT("DistanceFieldSurfaceCacheLightingCompute"),TEXT("ConeTraceOcclusionCS"),SF_Compute);
+IMPLEMENT_SHADER_TYPE(template<>,TConeTraceSurfaceCacheOcclusionCS<false>,TEXT("DistanceFieldSurfaceCacheLightingCompute"),TEXT("ConeTraceOcclusionCS"),SF_Compute);
 
 class FCombineConesCS : public FGlobalShader
 {
@@ -1737,7 +1626,6 @@ public:
 	static void ModifyCompilationEnvironment(EShaderPlatform Platform, FShaderCompilerEnvironment& OutEnvironment)
 	{
 		FGlobalShader::ModifyCompilationEnvironment(Platform,OutEnvironment);
-		OutEnvironment.SetDefine(TEXT("NUM_CONE_DIRECTIONS"), NumConeSampleDirections);
 
 		// To reduce shader compile time of compute shaders with shared memory, doesn't have an impact on generated code with current compiler (June 2010 DX SDK)
 		OutEnvironment.CompilerFlags.Add(CFLAG_StandardOptimization);
@@ -1748,11 +1636,12 @@ public:
 	{
 		AOParameters.Bind(Initializer.ParameterMap);
 		IrradianceCacheBentNormal.Bind(Initializer.ParameterMap, TEXT("IrradianceCacheBentNormal"));
-		IrradianceCacheNormal.Bind(Initializer.ParameterMap, TEXT("IrradianceCacheNormal"));
+		CacheParameters.Bind(Initializer.ParameterMap);
 		RecordConeVisibility.Bind(Initializer.ParameterMap, TEXT("RecordConeVisibility"));
 		ScatterDrawParameters.Bind(Initializer.ParameterMap, TEXT("ScatterDrawParameters"));
 		SavedStartIndex.Bind(Initializer.ParameterMap, TEXT("SavedStartIndex"));
 		BentNormalNormalizeFactor.Bind(Initializer.ParameterMap, TEXT("BentNormalNormalizeFactor"));
+		TanConeHalfAngle.Bind(Initializer.ParameterMap, TEXT("TanConeHalfAngle"));
 	}
 
 	FCombineConesCS()
@@ -1777,7 +1666,7 @@ public:
 		const FScene* Scene = (const FScene*)View.Family->Scene;
 		FSurfaceCacheResources& SurfaceCacheResources = *Scene->SurfaceCacheResources;
 
-		SetSRVParameter(RHICmdList, ShaderRHI, IrradianceCacheNormal, SurfaceCacheResources.Level[DepthLevel]->Normal.SRV);
+		CacheParameters.Set(RHICmdList, ShaderRHI, *SurfaceCacheResources.Level[DepthLevel]);
 		SetSRVParameter(RHICmdList, ShaderRHI, ScatterDrawParameters, SurfaceCacheResources.Level[DepthLevel]->ScatterDrawParameters.SRV);
 		SetSRVParameter(RHICmdList, ShaderRHI, SavedStartIndex, SurfaceCacheResources.Level[DepthLevel]->SavedStartIndex.SRV);
 		SetSRVParameter(RHICmdList, ShaderRHI, RecordConeVisibility, GTemporaryIrradianceCacheResources.ConeVisibility.SRV);
@@ -1805,6 +1694,8 @@ public:
 
 		float BentNormalNormalizeFactorValue = 1.0f / (UnoccludedVector / NumConeSampleDirections).Size();
 		SetShaderValue(RHICmdList, ShaderRHI, BentNormalNormalizeFactor, BentNormalNormalizeFactorValue);
+
+		SetShaderValue(RHICmdList, ShaderRHI, TanConeHalfAngle, FMath::Tan(GAOConeHalfAngle));
 	}
 
 	void UnsetParameters(FRHICommandList& RHICmdList)
@@ -1817,11 +1708,12 @@ public:
 		bool bShaderHasOutdatedParameters = FGlobalShader::Serialize(Ar);
 		Ar << AOParameters;
 		Ar << IrradianceCacheBentNormal;
-		Ar << IrradianceCacheNormal;
+		Ar << CacheParameters;
 		Ar << RecordConeVisibility;
 		Ar << ScatterDrawParameters;
 		Ar << SavedStartIndex;
 		Ar << BentNormalNormalizeFactor;
+		Ar << TanConeHalfAngle;
 		return bShaderHasOutdatedParameters;
 	}
 
@@ -1829,11 +1721,12 @@ private:
 
 	FAOParameters AOParameters;
 	FRWShaderParameter IrradianceCacheBentNormal;
-	FShaderResourceParameter IrradianceCacheNormal;
+	FIrradianceCacheParameters CacheParameters;
 	FShaderResourceParameter RecordConeVisibility;
 	FShaderResourceParameter ScatterDrawParameters;
 	FShaderResourceParameter SavedStartIndex;
 	FShaderParameter BentNormalNormalizeFactor;
+	FShaderParameter TanConeHalfAngle;
 };
 
 IMPLEMENT_SHADER_TYPE(,FCombineConesCS,TEXT("DistanceFieldSurfaceCacheLightingCompute"),TEXT("CombineConesCS"),SF_Compute);
@@ -1917,11 +1810,7 @@ public:
 	TIrradianceCacheSplatVS(const ShaderMetaType::CompiledShaderInitializerType& Initializer):
 		FGlobalShader(Initializer)
 	{
-		IrradianceCachePositionRadius.Bind(Initializer.ParameterMap, TEXT("IrradianceCachePositionRadius"));
-		IrradianceCacheOccluderRadius.Bind(Initializer.ParameterMap, TEXT("IrradianceCacheOccluderRadius"));
-		IrradianceCacheNormal.Bind(Initializer.ParameterMap, TEXT("IrradianceCacheNormal"));
-		IrradianceCacheBentNormal.Bind(Initializer.ParameterMap, TEXT("IrradianceCacheBentNormal"));
-		IrradianceCacheIrradiance.Bind(Initializer.ParameterMap, TEXT("IrradianceCacheIrradiance"));
+		CacheParameters.Bind(Initializer.ParameterMap);
 		InterpolationRadiusScale.Bind(Initializer.ParameterMap, TEXT("InterpolationRadiusScale"));
 		NormalizedOffsetToPixelCenter.Bind(Initializer.ParameterMap, TEXT("NormalizedOffsetToPixelCenter"));
 		HackExpand.Bind(Initializer.ParameterMap, TEXT("HackExpand"));
@@ -1939,11 +1828,7 @@ public:
 		const FScene* Scene = (const FScene*)View.Family->Scene;
 		FSurfaceCacheResources& SurfaceCacheResources = *Scene->SurfaceCacheResources;
 
-		SetSRVParameter(RHICmdList, ShaderRHI, IrradianceCachePositionRadius, SurfaceCacheResources.Level[DepthLevel]->PositionAndRadius.SRV);
-		SetSRVParameter(RHICmdList, ShaderRHI, IrradianceCacheOccluderRadius, SurfaceCacheResources.Level[DepthLevel]->OccluderRadius.SRV);
-		SetSRVParameter(RHICmdList, ShaderRHI, IrradianceCacheNormal, SurfaceCacheResources.Level[DepthLevel]->Normal.SRV);
-		SetSRVParameter(RHICmdList, ShaderRHI, IrradianceCacheBentNormal, SurfaceCacheResources.Level[DepthLevel]->BentNormal.SRV);
-		SetSRVParameter(RHICmdList, ShaderRHI, IrradianceCacheIrradiance, SurfaceCacheResources.Level[DepthLevel]->Irradiance.SRV);
+		CacheParameters.Set(RHICmdList, ShaderRHI, *SurfaceCacheResources.Level[DepthLevel]);
 
 		SetShaderValue(RHICmdList, ShaderRHI, InterpolationRadiusScale, (bFinalInterpolationPass ? GAOInterpolationRadiusScale : 1.0f));
 		SetShaderValue(RHICmdList, ShaderRHI, NormalizedOffsetToPixelCenter, NormalizedOffsetToPixelCenterValue);
@@ -1961,11 +1846,7 @@ public:
 	virtual bool Serialize(FArchive& Ar) override
 	{
 		bool bShaderHasOutdatedParameters = FGlobalShader::Serialize(Ar);
-		Ar << IrradianceCachePositionRadius;
-		Ar << IrradianceCacheOccluderRadius;
-		Ar << IrradianceCacheNormal;
-		Ar << IrradianceCacheBentNormal;
-		Ar << IrradianceCacheIrradiance;
+		Ar << CacheParameters;
 		Ar << InterpolationRadiusScale;
 		Ar << NormalizedOffsetToPixelCenter;
 		Ar << HackExpand;
@@ -1975,11 +1856,7 @@ public:
 	}
 
 private:
-	FShaderResourceParameter IrradianceCachePositionRadius;
-	FShaderResourceParameter IrradianceCacheOccluderRadius;
-	FShaderResourceParameter IrradianceCacheNormal;
-	FShaderResourceParameter IrradianceCacheBentNormal;
-	FShaderResourceParameter IrradianceCacheIrradiance;
+	FIrradianceCacheParameters CacheParameters;
 	FShaderParameter InterpolationRadiusScale;
 	FShaderParameter NormalizedOffsetToPixelCenter;
 	FShaderParameter HackExpand;
@@ -2092,1132 +1969,109 @@ IMPLEMENT_SPLAT_PS_TYPE(false, true)
 IMPLEMENT_SPLAT_PS_TYPE(true, false)
 IMPLEMENT_SPLAT_PS_TYPE(false, false)
 
-template<bool bSupportIrradiance>
-class TDistanceFieldAOCombinePS : public FGlobalShader
+/** Generates a pseudo-random position inside the unit sphere, uniformly distributed over the volume of the sphere. */
+FVector GetUnitPosition2(FRandomStream& RandomStream)
 {
-	DECLARE_SHADER_TYPE(TDistanceFieldAOCombinePS, Global);
-public:
-
-	static bool ShouldCache(EShaderPlatform Platform)
+	FVector Result;
+	// Use rejection sampling to generate a valid sample
+	do
 	{
-		return IsFeatureLevelSupported(Platform, ERHIFeatureLevel::SM5) && DoesPlatformSupportDistanceFieldAO(Platform);
-	}
+		Result.X = RandomStream.GetFraction() * 2 - 1;
+		Result.Y = RandomStream.GetFraction() * 2 - 1;
+		Result.Z = RandomStream.GetFraction() * 2 - 1;
+	} while( Result.SizeSquared() > 1.f );
+	return Result;
+}
 
-	static void ModifyCompilationEnvironment(EShaderPlatform Platform, FShaderCompilerEnvironment& OutEnvironment)
-	{
-		OutEnvironment.SetDefine(TEXT("DOWNSAMPLE_FACTOR"), GAODownsampleFactor);
-		OutEnvironment.SetDefine(TEXT("SUPPORT_IRRADIANCE"), bSupportIrradiance ? TEXT("1") : TEXT("0"));
-	}
+/** Generates a pseudo-random unit vector, uniformly distributed over all directions. */
+FVector GetUnitVector2(FRandomStream& RandomStream)
+{
+	return GetUnitPosition2(RandomStream).GetUnsafeNormal();
+}
 
-	/** Default constructor. */
-	TDistanceFieldAOCombinePS() {}
+void GenerateBestSpacedVectors()
+{
+	static bool bGenerated = false;
+	bool bApplyRepulsion = false;
 
-	/** Initialization constructor. */
-	TDistanceFieldAOCombinePS(const ShaderMetaType::CompiledShaderInitializerType& Initializer)
-		: FGlobalShader(Initializer)
+	if (bApplyRepulsion && !bGenerated)
 	{
-		DeferredParameters.Bind(Initializer.ParameterMap);
-		AOParameters.Bind(Initializer.ParameterMap);
-		AOLevelParameters.Bind(Initializer.ParameterMap);
-		BentNormalAOTexture.Bind(Initializer.ParameterMap,TEXT("BentNormalAOTexture"));
-		BentNormalAOSampler.Bind(Initializer.ParameterMap,TEXT("BentNormalAOSampler"));
-		IrradianceTexture.Bind(Initializer.ParameterMap,TEXT("IrradianceTexture"));
-		IrradianceSampler.Bind(Initializer.ParameterMap,TEXT("IrradianceSampler"));
-		DistanceFieldNormalTexture.Bind(Initializer.ParameterMap, TEXT("DistanceFieldNormalTexture"));
-		DistanceFieldNormalSampler.Bind(Initializer.ParameterMap, TEXT("DistanceFieldNormalSampler"));
-		DebugBuffer.Bind(Initializer.ParameterMap, TEXT("DebugBuffer"));
-		DistanceFadeScale.Bind(Initializer.ParameterMap, TEXT("DistanceFadeScale"));
-		SelfOcclusionReplacement.Bind(Initializer.ParameterMap, TEXT("SelfOcclusionReplacement"));
-	}
+		bGenerated = true;
 
-	void SetParameters(
-		FRHICommandList& RHICmdList, 
-		const FSceneView& View, 
-		FSceneRenderTargetItem& InBentNormalTexture, 
-		IPooledRenderTarget* InIrradianceTexture, 
-		FSceneRenderTargetItem& DistanceFieldNormal, 
-		const FDistanceFieldAOParameters& Parameters)
-	{
-		const FPixelShaderRHIParamRef ShaderRHI = GetPixelShader();
+		FVector OriginalSpacedVectors9[ARRAY_COUNT(SpacedVectors9)];
 
-		FGlobalShader::SetParameters(RHICmdList, ShaderRHI, View);
-		DeferredParameters.Set(RHICmdList, ShaderRHI, View);
-		AOParameters.Set(RHICmdList, ShaderRHI, Parameters);
-		AOLevelParameters.Set(RHICmdList, ShaderRHI, View, GAODownsampleFactor);
+		for (int32 i = 0; i < ARRAY_COUNT(OriginalSpacedVectors9); i++)
+		{
+			OriginalSpacedVectors9[i] = SpacedVectors9[i];
+		}
 
-		SetTextureParameter(RHICmdList, ShaderRHI, BentNormalAOTexture, BentNormalAOSampler, TStaticSamplerState<SF_Point>::GetRHI(), InBentNormalTexture.ShaderResourceTexture);
+		float CosHalfAngle = 1 - 1.0f / (float)ARRAY_COUNT(OriginalSpacedVectors9);
+		// Used to prevent self-shadowing on a plane
+		float AngleBias = .03f;
+		float MinAngle = FMath::Acos(CosHalfAngle) + AngleBias;
+		float MinZ = FMath::Sin(MinAngle);
 
-		if (IrradianceTexture.IsBound())
+		// Relaxation iterations by repulsion
+		for (int32 Iteration = 0; Iteration < 10000; Iteration++)
 		{
-			SetTextureParameter(RHICmdList, ShaderRHI, IrradianceTexture, IrradianceSampler, TStaticSamplerState<SF_Point>::GetRHI(), InIrradianceTexture->GetRenderTargetItem().ShaderResourceTexture);
-		}
+			for (int32 i = 0; i < ARRAY_COUNT(OriginalSpacedVectors9); i++)
+			{
+				FVector Force(0.0f, 0.0f, 0.0f);
 
-		SetTextureParameter(
-			RHICmdList,
-			ShaderRHI,
-			DistanceFieldNormalTexture,
-			DistanceFieldNormalSampler,
-			TStaticSamplerState<SF_Point,AM_Clamp,AM_Clamp,AM_Clamp>::GetRHI(),
-			DistanceFieldNormal.ShaderResourceTexture
-			);
+				for (int32 j = 0; j < ARRAY_COUNT(OriginalSpacedVectors9); j++)
+				{
+					if (i != j)
+					{
+						FVector Distance = OriginalSpacedVectors9[i] - OriginalSpacedVectors9[j];
+						float Dot = OriginalSpacedVectors9[i] | OriginalSpacedVectors9[j];
 
-		SetSRVParameter(RHICmdList, ShaderRHI, DebugBuffer, GDebugBuffer.DebugData.SRV);
+						if (Dot > 0)
+						{
+							// Repulsion force
+							Force += .001f * Distance.GetSafeNormal() * Dot * Dot * Dot * Dot;
+						}
+					}
+				}
 
-		const float DistanceFadeScaleValue = 1.0f / ((1.0f - GAOViewFadeDistanceScale) * GAOMaxViewDistance);
-		SetShaderValue(RHICmdList, ShaderRHI, DistanceFadeScale, DistanceFadeScaleValue);
+				FVector NewPosition = OriginalSpacedVectors9[i] + Force;
+				NewPosition.Z = FMath::Max(NewPosition.Z, MinZ);
+				NewPosition = NewPosition.GetSafeNormal();
+				OriginalSpacedVectors9[i] = NewPosition;
+			}
+		}
 
-		extern float GVPLSelfOcclusionReplacement;
-		SetShaderValue(RHICmdList, ShaderRHI, SelfOcclusionReplacement, GVPLSelfOcclusionReplacement);
-	}
+		for (int32 i = 0; i < ARRAY_COUNT(OriginalSpacedVectors9); i++)
+		{
+			UE_LOG(LogDistanceField, Log, TEXT("FVector(%f, %f, %f),"), OriginalSpacedVectors9[i].X, OriginalSpacedVectors9[i].Y, OriginalSpacedVectors9[i].Z);
+		}
 
-	// FShader interface.
-	virtual bool Serialize(FArchive& Ar) override
-	{
-		bool bShaderHasOutdatedParameters = FGlobalShader::Serialize(Ar);
-		Ar << DeferredParameters;
-		Ar << AOParameters;
-		Ar << AOLevelParameters;
-		Ar << BentNormalAOTexture;
-		Ar << BentNormalAOSampler;
-		Ar << IrradianceTexture;
-		Ar << IrradianceSampler;
-		Ar << DistanceFieldNormalTexture;
-		Ar << DistanceFieldNormalSampler;
-		Ar << DebugBuffer;
-		Ar << DistanceFadeScale;
-		Ar << SelfOcclusionReplacement;
-		return bShaderHasOutdatedParameters;
+		int32 temp = 0;
 	}
 
-private:
-
-	FDeferredPixelShaderParameters DeferredParameters;
-	FAOParameters AOParameters;
-	FAOLevelParameters AOLevelParameters;
-	FShaderResourceParameter BentNormalAOTexture;
-	FShaderResourceParameter BentNormalAOSampler;
-	FShaderResourceParameter IrradianceTexture;
-	FShaderResourceParameter IrradianceSampler;
-	FShaderResourceParameter DistanceFieldNormalTexture;
-	FShaderResourceParameter DistanceFieldNormalSampler;
-	FShaderResourceParameter DebugBuffer;
-	FShaderParameter DistanceFadeScale;
-	FShaderParameter SelfOcclusionReplacement;
-};
+	bool bBruteForceGenerateConeDirections = false;
 
-IMPLEMENT_SHADER_TYPE(template<>,TDistanceFieldAOCombinePS<true>,TEXT("DistanceFieldSurfaceCacheLighting"),TEXT("AOCombinePS"),SF_Pixel);
-IMPLEMENT_SHADER_TYPE(template<>,TDistanceFieldAOCombinePS<false>,TEXT("DistanceFieldSurfaceCacheLighting"),TEXT("AOCombinePS"),SF_Pixel);
+	if (bBruteForceGenerateConeDirections)
+	{
+		FVector BestSpacedVectors9[9];
+		float BestCoverage = 0;
+		// Each cone covers an area of ConeSolidAngle = HemisphereSolidAngle / NumCones
+		// HemisphereSolidAngle = 2 * PI
+		// ConeSolidAngle = 2 * PI * (1 - cos(ConeHalfAngle))
+		// cos(ConeHalfAngle) = 1 - 1 / NumCones
+		float CosHalfAngle = 1 - 1.0f / (float)ARRAY_COUNT(BestSpacedVectors9);
+		// Prevent self-intersection in sample set
+		float MinAngle = FMath::Acos(CosHalfAngle);
+		float MinZ = FMath::Sin(MinAngle);
+		FRandomStream RandomStream(123567);
 
-template<bool bSupportIrradiance, bool bHighQuality>
-class TFillGapsPS : public FGlobalShader
-{
-	DECLARE_SHADER_TYPE(TFillGapsPS, Global);
-public:
+		// Super slow random brute force search
+		for (int i = 0; i < 1000000; i++)
+		{
+			FVector CandidateSpacedVectors[ARRAY_COUNT(BestSpacedVectors9)];
 
-	static bool ShouldCache(EShaderPlatform Platform)
-	{
-		return IsFeatureLevelSupported(Platform, ERHIFeatureLevel::SM5) && DoesPlatformSupportDistanceFieldAO(Platform);
-	}
-
-	static void ModifyCompilationEnvironment(EShaderPlatform Platform, FShaderCompilerEnvironment& OutEnvironment)
-	{
-		OutEnvironment.SetDefine(TEXT("DOWNSAMPLE_FACTOR"), GAODownsampleFactor);
-		OutEnvironment.SetDefine(TEXT("SUPPORT_IRRADIANCE"), bSupportIrradiance ? TEXT("1") : TEXT("0"));
-		OutEnvironment.SetDefine(TEXT("HIGH_QUALITY_FILL_GAPS"), bHighQuality ? TEXT("1") : TEXT("0"));
-	}
-
-	/** Default constructor. */
-	TFillGapsPS() {}
-
-	/** Initialization constructor. */
-	TFillGapsPS(const ShaderMetaType::CompiledShaderInitializerType& Initializer)
-		: FGlobalShader(Initializer)
-	{
-		BentNormalAOTexture.Bind(Initializer.ParameterMap, TEXT("BentNormalAOTexture"));
-		BentNormalAOSampler.Bind(Initializer.ParameterMap, TEXT("BentNormalAOSampler"));
-		IrradianceTexture.Bind(Initializer.ParameterMap, TEXT("IrradianceTexture"));
-		IrradianceSampler.Bind(Initializer.ParameterMap, TEXT("IrradianceSampler"));
-		BentNormalAOTexelSize.Bind(Initializer.ParameterMap, TEXT("BentNormalAOTexelSize"));
-		MinDownsampleFactorToBaseLevel.Bind(Initializer.ParameterMap, TEXT("MinDownsampleFactorToBaseLevel"));
-	}
-
-	void SetParameters(FRHICommandList& RHICmdList, const FSceneView& View, FSceneRenderTargetItem& DistanceFieldAOBentNormal2, IPooledRenderTarget* DistanceFieldIrradiance2)
-	{
-		const FPixelShaderRHIParamRef ShaderRHI = GetPixelShader();
-
-		FGlobalShader::SetParameters(RHICmdList, ShaderRHI, View);
-
-		SetTextureParameter(
-			RHICmdList,
-			ShaderRHI,
-			BentNormalAOTexture,
-			BentNormalAOSampler,
-			TStaticSamplerState<SF_Bilinear,AM_Clamp,AM_Clamp,AM_Clamp>::GetRHI(),
-			DistanceFieldAOBentNormal2.ShaderResourceTexture
-			);
-
-		if (IrradianceTexture.IsBound())
-		{
-			SetTextureParameter(
-				RHICmdList,
-				ShaderRHI,
-				IrradianceTexture,
-				IrradianceSampler,
-				TStaticSamplerState<SF_Bilinear,AM_Clamp,AM_Clamp,AM_Clamp>::GetRHI(),
-				DistanceFieldIrradiance2->GetRenderTargetItem().ShaderResourceTexture
-				);
-		}
-
-		const FIntPoint DownsampledBufferSize(GSceneRenderTargets.GetBufferSizeXY() / FIntPoint(GAODownsampleFactor, GAODownsampleFactor));
-		const FVector2D BaseLevelTexelSizeValue(1.0f / DownsampledBufferSize.X, 1.0f / DownsampledBufferSize.Y);
-		SetShaderValue(RHICmdList, ShaderRHI, BentNormalAOTexelSize, BaseLevelTexelSizeValue);
-
-		const float MinDownsampleFactor = 1 << (GAOMinLevel * GAOPowerOfTwoBetweenLevels); 
-		SetShaderValue(RHICmdList, ShaderRHI, MinDownsampleFactorToBaseLevel, MinDownsampleFactor);
-	}
-	// FShader interface.
-	virtual bool Serialize(FArchive& Ar) override
-	{
-		bool bShaderHasOutdatedParameters = FGlobalShader::Serialize(Ar);
-		Ar << BentNormalAOTexture;
-		Ar << BentNormalAOSampler;
-		Ar << IrradianceTexture;
-		Ar << IrradianceSampler;
-		Ar << BentNormalAOTexelSize;
-		Ar << MinDownsampleFactorToBaseLevel;
-		return bShaderHasOutdatedParameters;
-	}
-
-private:
-
-	FShaderResourceParameter BentNormalAOTexture;
-	FShaderResourceParameter BentNormalAOSampler;
-	FShaderResourceParameter IrradianceTexture;
-	FShaderResourceParameter IrradianceSampler;
-	FShaderParameter BentNormalAOTexelSize;
-	FShaderParameter MinDownsampleFactorToBaseLevel;
-};
-
-#define IMPLEMENT_FILLGAPS_PS_TYPE(bSupportIrradiance,bHighQuality) \
-	typedef TFillGapsPS<bSupportIrradiance, bHighQuality> TFillGapsPS##bSupportIrradiance##bHighQuality; \
-	IMPLEMENT_SHADER_TYPE(template<>,TFillGapsPS##bSupportIrradiance##bHighQuality,TEXT("DistanceFieldLightingPost"),TEXT("FillGapsPS"),SF_Pixel); 
-
-IMPLEMENT_FILLGAPS_PS_TYPE(true, true);
-IMPLEMENT_FILLGAPS_PS_TYPE(true, false);
-IMPLEMENT_FILLGAPS_PS_TYPE(false, true);
-IMPLEMENT_FILLGAPS_PS_TYPE(false, false);
-
-template<bool bSupportIrradiance>
-class TUpdateHistoryPS : public FGlobalShader
-{
-	DECLARE_SHADER_TYPE(TUpdateHistoryPS, Global);
-public:
-
-	static bool ShouldCache(EShaderPlatform Platform)
-	{
-		return IsFeatureLevelSupported(Platform, ERHIFeatureLevel::SM5) && DoesPlatformSupportDistanceFieldAO(Platform);
-	}
-
-	static void ModifyCompilationEnvironment(EShaderPlatform Platform, FShaderCompilerEnvironment& OutEnvironment)
-	{
-		OutEnvironment.SetDefine(TEXT("DOWNSAMPLE_FACTOR"), GAODownsampleFactor);
-		OutEnvironment.SetDefine(TEXT("SUPPORT_IRRADIANCE"), bSupportIrradiance ? TEXT("1") : TEXT("0"));
-	}
-
-	/** Default constructor. */
-	TUpdateHistoryPS() {}
-
-	/** Initialization constructor. */
-	TUpdateHistoryPS(const ShaderMetaType::CompiledShaderInitializerType& Initializer)
-		: FGlobalShader(Initializer)
-	{
-		DeferredParameters.Bind(Initializer.ParameterMap);
-		BentNormalAOTexture.Bind(Initializer.ParameterMap, TEXT("BentNormalAOTexture"));
-		BentNormalAOSampler.Bind(Initializer.ParameterMap, TEXT("BentNormalAOSampler"));
-		BentNormalHistoryTexture.Bind(Initializer.ParameterMap, TEXT("BentNormalHistoryTexture"));
-		BentNormalHistorySampler.Bind(Initializer.ParameterMap, TEXT("BentNormalHistorySampler"));
-		IrradianceTexture.Bind(Initializer.ParameterMap, TEXT("IrradianceTexture"));
-		IrradianceSampler.Bind(Initializer.ParameterMap, TEXT("IrradianceSampler"));
-		IrradianceHistoryTexture.Bind(Initializer.ParameterMap, TEXT("IrradianceHistoryTexture"));
-		IrradianceHistorySampler.Bind(Initializer.ParameterMap, TEXT("IrradianceHistorySampler"));
-		HistoryWeight.Bind(Initializer.ParameterMap, TEXT("HistoryWeight"));
-		HistoryDistanceThreshold.Bind(Initializer.ParameterMap, TEXT("HistoryDistanceThreshold"));
-		UseHistoryFilter.Bind(Initializer.ParameterMap, TEXT("UseHistoryFilter"));
-		VelocityTexture.Bind(Initializer.ParameterMap, TEXT("VelocityTexture"));
-		VelocityTextureSampler.Bind(Initializer.ParameterMap, TEXT("VelocityTextureSampler"));
-	}
-
-	void SetParameters(
-		FRHICommandList& RHICmdList, 
-		const FSceneView& View, 
-		FSceneRenderTargetItem& BentNormalHistoryTextureValue, 
-		TRefCountPtr<IPooledRenderTarget>* IrradianceHistoryRT, 
-		FSceneRenderTargetItem& DistanceFieldAOBentNormal, 
-		IPooledRenderTarget* DistanceFieldIrradiance,
-		FSceneRenderTargetItem& VelocityTextureValue)
-	{
-		const FPixelShaderRHIParamRef ShaderRHI = GetPixelShader();
-
-		FGlobalShader::SetParameters(RHICmdList, ShaderRHI, View);
-		DeferredParameters.Set(RHICmdList, ShaderRHI, View);
-
-		SetTextureParameter(
-			RHICmdList,
-			ShaderRHI,
-			BentNormalAOTexture,
-			BentNormalAOSampler,
-			TStaticSamplerState<SF_Bilinear,AM_Clamp,AM_Clamp,AM_Clamp>::GetRHI(),
-			DistanceFieldAOBentNormal.ShaderResourceTexture
-			);
-
-		SetTextureParameter(
-			RHICmdList,
-			ShaderRHI,
-			BentNormalHistoryTexture,
-			BentNormalHistorySampler,
-			TStaticSamplerState<SF_Bilinear,AM_Clamp,AM_Clamp,AM_Clamp>::GetRHI(),
-			BentNormalHistoryTextureValue.ShaderResourceTexture
-			);
-
-		if (IrradianceTexture.IsBound())
-		{
-			SetTextureParameter(
-				RHICmdList,
-				ShaderRHI,
-				IrradianceTexture,
-				IrradianceSampler,
-				TStaticSamplerState<SF_Bilinear,AM_Clamp,AM_Clamp,AM_Clamp>::GetRHI(),
-				DistanceFieldIrradiance->GetRenderTargetItem().ShaderResourceTexture
-				);
-		}
-		
-		if (IrradianceHistoryTexture.IsBound())
-		{
-			SetTextureParameter(
-				RHICmdList,
-				ShaderRHI,
-				IrradianceHistoryTexture,
-				IrradianceHistorySampler,
-				TStaticSamplerState<SF_Bilinear,AM_Clamp,AM_Clamp,AM_Clamp>::GetRHI(),
-				(*IrradianceHistoryRT)->GetRenderTargetItem().ShaderResourceTexture
-				);
-		}
-
-		SetUniformBufferParameter(RHICmdList, ShaderRHI, GetUniformBufferParameter<FCameraMotionParameters>(), CreateCameraMotionParametersUniformBuffer(View));
-
-		SetShaderValue(RHICmdList, ShaderRHI, HistoryWeight, GAOHistoryWeight);
-		SetShaderValue(RHICmdList, ShaderRHI, HistoryDistanceThreshold, GAOHistoryDistanceThreshold);
-		SetShaderValue(RHICmdList, ShaderRHI, UseHistoryFilter, (GAOHistoryStabilityPass ? 1.0f : 0.0f));
-
-		SetTextureParameter(
-			RHICmdList,
-			ShaderRHI,
-			VelocityTexture,
-			VelocityTextureSampler,
-			TStaticSamplerState<SF_Bilinear,AM_Clamp,AM_Clamp,AM_Clamp>::GetRHI(),
-			VelocityTextureValue.ShaderResourceTexture
-			);
-	}
-	// FShader interface.
-	virtual bool Serialize(FArchive& Ar)
-	{
-		bool bShaderHasOutdatedParameters = FGlobalShader::Serialize(Ar);
-		Ar << DeferredParameters;
-		Ar << BentNormalAOTexture;
-		Ar << BentNormalAOSampler;
-		Ar << BentNormalHistoryTexture;
-		Ar << BentNormalHistorySampler;
-		Ar << IrradianceTexture;
-		Ar << IrradianceSampler;
-		Ar << IrradianceHistoryTexture;
-		Ar << IrradianceHistorySampler;
-		Ar << HistoryWeight;
-		Ar << HistoryDistanceThreshold;
-		Ar << UseHistoryFilter;
-		Ar << VelocityTexture;
-		Ar << VelocityTextureSampler;
-		return bShaderHasOutdatedParameters;
-	}
-
-private:
-
-	FDeferredPixelShaderParameters DeferredParameters;
-	FShaderResourceParameter BentNormalAOTexture;
-	FShaderResourceParameter BentNormalAOSampler;
-	FShaderResourceParameter BentNormalHistoryTexture;
-	FShaderResourceParameter BentNormalHistorySampler;
-	FShaderResourceParameter IrradianceTexture;
-	FShaderResourceParameter IrradianceSampler;
-	FShaderResourceParameter IrradianceHistoryTexture;
-	FShaderResourceParameter IrradianceHistorySampler;
-	FShaderParameter HistoryWeight;
-	FShaderParameter HistoryDistanceThreshold;
-	FShaderParameter UseHistoryFilter;
-	FShaderResourceParameter VelocityTexture;
-	FShaderResourceParameter VelocityTextureSampler;
-};
-
-IMPLEMENT_SHADER_TYPE(template<>,TUpdateHistoryPS<true>,TEXT("DistanceFieldLightingPost"),TEXT("UpdateHistoryPS"),SF_Pixel);
-IMPLEMENT_SHADER_TYPE(template<>,TUpdateHistoryPS<false>,TEXT("DistanceFieldLightingPost"),TEXT("UpdateHistoryPS"),SF_Pixel);
-
-
-template<bool bSupportIrradiance>
-class TFilterHistoryPS : public FGlobalShader
-{
-	DECLARE_SHADER_TYPE(TFilterHistoryPS, Global);
-public:
-
-	static bool ShouldCache(EShaderPlatform Platform)
-	{
-		return IsFeatureLevelSupported(Platform, ERHIFeatureLevel::SM5) && DoesPlatformSupportDistanceFieldAO(Platform);
-	}
-
-	static void ModifyCompilationEnvironment(EShaderPlatform Platform, FShaderCompilerEnvironment& OutEnvironment)
-	{
-		OutEnvironment.SetDefine(TEXT("DOWNSAMPLE_FACTOR"), GAODownsampleFactor);
-		OutEnvironment.SetDefine(TEXT("SUPPORT_IRRADIANCE"), bSupportIrradiance ? TEXT("1") : TEXT("0"));
-	}
-
-	/** Default constructor. */
-	TFilterHistoryPS() {}
-
-	/** Initialization constructor. */
-	TFilterHistoryPS(const ShaderMetaType::CompiledShaderInitializerType& Initializer)
-		: FGlobalShader(Initializer)
-	{
-		BentNormalAOTexture.Bind(Initializer.ParameterMap, TEXT("BentNormalAOTexture"));
-		BentNormalAOSampler.Bind(Initializer.ParameterMap, TEXT("BentNormalAOSampler"));
-		IrradianceTexture.Bind(Initializer.ParameterMap, TEXT("IrradianceTexture"));
-		IrradianceSampler.Bind(Initializer.ParameterMap, TEXT("IrradianceSampler"));
-		HistoryWeight.Bind(Initializer.ParameterMap, TEXT("HistoryWeight"));
-		BentNormalAOTexelSize.Bind(Initializer.ParameterMap, TEXT("BentNormalAOTexelSize"));
-	}
-
-	void SetParameters(
-		FRHICommandList& RHICmdList, 
-		const FSceneView& View, 
-		FSceneRenderTargetItem& BentNormalHistoryTextureValue, 
-		IPooledRenderTarget* IrradianceHistoryRT)
-	{
-		const FPixelShaderRHIParamRef ShaderRHI = GetPixelShader();
-
-		FGlobalShader::SetParameters(RHICmdList, ShaderRHI, View);
-
-		SetTextureParameter(
-			RHICmdList,
-			ShaderRHI,
-			BentNormalAOTexture,
-			BentNormalAOSampler,
-			TStaticSamplerState<SF_Bilinear,AM_Clamp,AM_Clamp,AM_Clamp>::GetRHI(),
-			BentNormalHistoryTextureValue.ShaderResourceTexture
-			);
-
-		if (IrradianceTexture.IsBound())
-		{
-			SetTextureParameter(
-				RHICmdList,
-				ShaderRHI,
-				IrradianceTexture,
-				IrradianceSampler,
-				TStaticSamplerState<SF_Bilinear,AM_Clamp,AM_Clamp,AM_Clamp>::GetRHI(),
-				IrradianceHistoryRT->GetRenderTargetItem().ShaderResourceTexture
-				);
-		}
-
-		SetShaderValue(RHICmdList, ShaderRHI, HistoryWeight, GAOHistoryWeight);
-		
-		const FIntPoint DownsampledBufferSize(GSceneRenderTargets.GetBufferSizeXY() / FIntPoint(GAODownsampleFactor, GAODownsampleFactor));
-		const FVector2D BaseLevelTexelSizeValue(1.0f / DownsampledBufferSize.X, 1.0f / DownsampledBufferSize.Y);
-		SetShaderValue(RHICmdList, ShaderRHI, BentNormalAOTexelSize, BaseLevelTexelSizeValue);
-	}
-
-	// FShader interface.
-	virtual bool Serialize(FArchive& Ar)
-	{
-		bool bShaderHasOutdatedParameters = FGlobalShader::Serialize(Ar);
-		Ar << BentNormalAOTexture;
-		Ar << BentNormalAOSampler;
-		Ar << IrradianceTexture;
-		Ar << IrradianceSampler;
-		Ar << HistoryWeight;
-		Ar << BentNormalAOTexelSize;
-		return bShaderHasOutdatedParameters;
-	}
-
-private:
-
-	FShaderResourceParameter BentNormalAOTexture;
-	FShaderResourceParameter BentNormalAOSampler;
-	FShaderResourceParameter IrradianceTexture;
-	FShaderResourceParameter IrradianceSampler;
-	FShaderParameter HistoryWeight;
-	FShaderParameter BentNormalAOTexelSize;
-};
-
-IMPLEMENT_SHADER_TYPE(template<>,TFilterHistoryPS<true>,TEXT("DistanceFieldLightingPost"),TEXT("FilterHistoryPS"),SF_Pixel);
-IMPLEMENT_SHADER_TYPE(template<>,TFilterHistoryPS<false>,TEXT("DistanceFieldLightingPost"),TEXT("FilterHistoryPS"),SF_Pixel);
-
-
-void AllocateOrReuseAORenderTarget(TRefCountPtr<IPooledRenderTarget>& Target, const TCHAR* Name, bool bWithAlpha)
-{
-	if (!Target)
-	{
-		FIntPoint BufferSize = GetBufferSizeForAO();
-
-		FPooledRenderTargetDesc Desc(FPooledRenderTargetDesc::Create2DDesc(BufferSize, bWithAlpha ? PF_FloatRGBA : PF_FloatRGB, TexCreate_None, TexCreate_RenderTargetable | TexCreate_UAV, false));
-		GRenderTargetPool.FindFreeElement(Desc, Target, Name);
-	}
-}
-
-void UpdateHistory(
-	FRHICommandList& RHICmdList,
-	const FViewInfo& View, 
-	const TCHAR* BentNormalHistoryRTName,
-	const TCHAR* IrradianceHistoryRTName,
-	FSceneRenderTargetItem& VelocityTexture,
-	/** Contains last frame's history, if non-NULL.  This will be updated with the new frame's history. */
-	TRefCountPtr<IPooledRenderTarget>* BentNormalHistoryState,
-	TRefCountPtr<IPooledRenderTarget>* IrradianceHistoryState,
-	/** Source */
-	TRefCountPtr<IPooledRenderTarget>& BentNormalSource, 
-	TRefCountPtr<IPooledRenderTarget>& IrradianceSource, 
-	/** Output of Temporal Reprojection for the next step in the pipeline. */
-	TRefCountPtr<IPooledRenderTarget>& BentNormalHistoryOutput,
-	TRefCountPtr<IPooledRenderTarget>& IrradianceHistoryOutput)
-{
-	if (BentNormalHistoryState)
-	{
-		const bool bUseDistanceFieldGI = IsDistanceFieldGIAllowed(View);
-
-		if (*BentNormalHistoryState 
-			&& !View.bCameraCut 
-			&& !View.bPrevTransformsReset 
-			&& (!bUseDistanceFieldGI || (IrradianceHistoryState && *IrradianceHistoryState)))
-		{
-			// Reuse a render target from the pool with a consistent name, for vis purposes
-			TRefCountPtr<IPooledRenderTarget> NewBentNormalHistory;
-			AllocateOrReuseAORenderTarget(NewBentNormalHistory, BentNormalHistoryRTName, true);
-
-			TRefCountPtr<IPooledRenderTarget> NewIrradianceHistory;
-
-			if (bUseDistanceFieldGI)
-			{
-				AllocateOrReuseAORenderTarget(NewIrradianceHistory, IrradianceHistoryRTName, false);
-			}
-
-			SCOPED_DRAW_EVENT(RHICmdList, UpdateHistory);
-
-			{
-				FTextureRHIParamRef RenderTargets[2] =
-				{
-					NewBentNormalHistory->GetRenderTargetItem().TargetableTexture,
-					bUseDistanceFieldGI ? NewIrradianceHistory->GetRenderTargetItem().TargetableTexture : NULL,
-				};
-
-				SetRenderTargets(RHICmdList, ARRAY_COUNT(RenderTargets) - (bUseDistanceFieldGI ? 0 : 1), RenderTargets, FTextureRHIParamRef(), 0, NULL);
-
-				RHICmdList.SetViewport(0, 0, 0.0f, View.ViewRect.Width() / GAODownsampleFactor, View.ViewRect.Height() / GAODownsampleFactor, 1.0f);
-				RHICmdList.SetRasterizerState(TStaticRasterizerState<FM_Solid, CM_None>::GetRHI());
-				RHICmdList.SetDepthStencilState(TStaticDepthStencilState<false, CF_Always>::GetRHI());
-				RHICmdList.SetBlendState(TStaticBlendState<>::GetRHI());
-
-				TShaderMapRef<FPostProcessVS> VertexShader(View.ShaderMap);
-
-				if (bUseDistanceFieldGI)
-				{
-					TShaderMapRef<TUpdateHistoryPS<true> > PixelShader(View.ShaderMap);
-
-					static FGlobalBoundShaderState BoundShaderState;
-					SetGlobalBoundShaderState(RHICmdList, View.GetFeatureLevel(), BoundShaderState, GFilterVertexDeclaration.VertexDeclarationRHI, *VertexShader, *PixelShader);
-					PixelShader->SetParameters(RHICmdList, View, (*BentNormalHistoryState)->GetRenderTargetItem(), IrradianceHistoryState, BentNormalSource->GetRenderTargetItem(), IrradianceSource, VelocityTexture);
-				}
-				else
-				{
-					TShaderMapRef<TUpdateHistoryPS<false> > PixelShader(View.ShaderMap);
-
-					static FGlobalBoundShaderState BoundShaderState;
-					SetGlobalBoundShaderState(RHICmdList, View.GetFeatureLevel(), BoundShaderState, GFilterVertexDeclaration.VertexDeclarationRHI, *VertexShader, *PixelShader);
-					PixelShader->SetParameters(RHICmdList, View, (*BentNormalHistoryState)->GetRenderTargetItem(), IrradianceHistoryState, BentNormalSource->GetRenderTargetItem(), IrradianceSource, VelocityTexture);
-				}
-
-				VertexShader->SetParameters(RHICmdList, View);
-
-				DrawRectangle( 
-					RHICmdList,
-					0, 0, 
-					View.ViewRect.Width() / GAODownsampleFactor, View.ViewRect.Height() / GAODownsampleFactor,
-					View.ViewRect.Min.X / GAODownsampleFactor, View.ViewRect.Min.Y / GAODownsampleFactor,
-					View.ViewRect.Width() / GAODownsampleFactor, View.ViewRect.Height() / GAODownsampleFactor,
-					FIntPoint(View.ViewRect.Width() / GAODownsampleFactor, View.ViewRect.Height() / GAODownsampleFactor),
-					GSceneRenderTargets.GetBufferSizeXY() / FIntPoint(GAODownsampleFactor, GAODownsampleFactor),
-					*VertexShader);
-
-				RHICmdList.CopyToResolveTarget(NewBentNormalHistory->GetRenderTargetItem().TargetableTexture, NewBentNormalHistory->GetRenderTargetItem().ShaderResourceTexture, false, FResolveParams());
-			
-				if (bUseDistanceFieldGI)
-				{
-					RHICmdList.CopyToResolveTarget(NewIrradianceHistory->GetRenderTargetItem().TargetableTexture, NewIrradianceHistory->GetRenderTargetItem().ShaderResourceTexture, false, FResolveParams());
-				}
-			}
-
-			if (GAOHistoryStabilityPass)
-			{
-				const FPooledRenderTargetDesc& HistoryDesc = (*BentNormalHistoryState)->GetDesc();
-
-				// Reallocate history if buffer sizes have changed
-				if (HistoryDesc.Extent != GSceneRenderTargets.GetBufferSizeXY() / FIntPoint(GAODownsampleFactor, GAODownsampleFactor))
-				{
-					GRenderTargetPool.FreeUnusedResource(*BentNormalHistoryState);
-					*BentNormalHistoryState = NULL;
-					// Update the view state's render target reference with the new history
-					AllocateOrReuseAORenderTarget(*BentNormalHistoryState, BentNormalHistoryRTName, true);
-
-					if (bUseDistanceFieldGI)
-					{
-						GRenderTargetPool.FreeUnusedResource(*IrradianceHistoryState);
-						*IrradianceHistoryState = NULL;
-						AllocateOrReuseAORenderTarget(*IrradianceHistoryState, IrradianceHistoryRTName, false);
-					}
-				}
-
-				{
-					FTextureRHIParamRef RenderTargets[2] =
-					{
-						(*BentNormalHistoryState)->GetRenderTargetItem().TargetableTexture,
-						bUseDistanceFieldGI ? (*IrradianceHistoryState)->GetRenderTargetItem().TargetableTexture : NULL,
-					};
-
-					SetRenderTargets(RHICmdList, ARRAY_COUNT(RenderTargets) - (bUseDistanceFieldGI ? 0 : 1), RenderTargets, FTextureRHIParamRef(), 0, NULL);
-
-					RHICmdList.SetViewport(0, 0, 0.0f, View.ViewRect.Width() / GAODownsampleFactor, View.ViewRect.Height() / GAODownsampleFactor, 1.0f);
-					RHICmdList.SetRasterizerState(TStaticRasterizerState<FM_Solid, CM_None>::GetRHI());
-					RHICmdList.SetDepthStencilState(TStaticDepthStencilState<false, CF_Always>::GetRHI());
-					RHICmdList.SetBlendState(TStaticBlendState<>::GetRHI());
-
-					TShaderMapRef<FPostProcessVS> VertexShader(View.ShaderMap);
-
-					if (bUseDistanceFieldGI)
-					{
-						TShaderMapRef<TFilterHistoryPS<true> > PixelShader(View.ShaderMap);
-
-						static FGlobalBoundShaderState BoundShaderState;
-						SetGlobalBoundShaderState(RHICmdList, View.GetFeatureLevel(), BoundShaderState, GFilterVertexDeclaration.VertexDeclarationRHI, *VertexShader, *PixelShader);
-						PixelShader->SetParameters(RHICmdList, View, NewBentNormalHistory->GetRenderTargetItem(), NewIrradianceHistory);
-					}
-					else
-					{
-						TShaderMapRef<TFilterHistoryPS<false> > PixelShader(View.ShaderMap);
-
-						static FGlobalBoundShaderState BoundShaderState;
-						SetGlobalBoundShaderState(RHICmdList, View.GetFeatureLevel(), BoundShaderState, GFilterVertexDeclaration.VertexDeclarationRHI, *VertexShader, *PixelShader);
-						PixelShader->SetParameters(RHICmdList, View, NewBentNormalHistory->GetRenderTargetItem(), NewIrradianceHistory);
-					}
-
-					VertexShader->SetParameters(RHICmdList, View);
-
-					DrawRectangle( 
-						RHICmdList,
-						0, 0, 
-						View.ViewRect.Width() / GAODownsampleFactor, View.ViewRect.Height() / GAODownsampleFactor,
-						0, 0,
-						View.ViewRect.Width() / GAODownsampleFactor, View.ViewRect.Height() / GAODownsampleFactor,
-						FIntPoint(View.ViewRect.Width() / GAODownsampleFactor, View.ViewRect.Height() / GAODownsampleFactor),
-						GSceneRenderTargets.GetBufferSizeXY() / FIntPoint(GAODownsampleFactor, GAODownsampleFactor),
-						*VertexShader);
-
-					RHICmdList.CopyToResolveTarget((*BentNormalHistoryState)->GetRenderTargetItem().TargetableTexture, (*BentNormalHistoryState)->GetRenderTargetItem().ShaderResourceTexture, false, FResolveParams());
-			
-					if (bUseDistanceFieldGI)
-					{
-						RHICmdList.CopyToResolveTarget((*IrradianceHistoryState)->GetRenderTargetItem().TargetableTexture, (*IrradianceHistoryState)->GetRenderTargetItem().ShaderResourceTexture, false, FResolveParams());
-					}
-				}
-
-				BentNormalHistoryOutput = *BentNormalHistoryState;
-				IrradianceHistoryOutput = *IrradianceHistoryState;
-			}
-			else
-			{
-				// Update the view state's render target reference with the new history
-				*BentNormalHistoryState = NewBentNormalHistory;
-				BentNormalHistoryOutput = NewBentNormalHistory;
-
-				*IrradianceHistoryState = NewIrradianceHistory;
-				IrradianceHistoryOutput = NewIrradianceHistory;
-			}
-		}
-		else
-		{
-			// Use the current frame's mask for next frame's history
-			*BentNormalHistoryState = BentNormalSource;
-			BentNormalHistoryOutput = BentNormalSource;
-			BentNormalSource = NULL;
-
-			*IrradianceHistoryState = IrradianceSource;
-			IrradianceHistoryOutput = IrradianceSource;
-			IrradianceSource = NULL;
-		}
-	}
-	else
-	{
-		// Temporal reprojection is disabled or there is no view state - pass through
-		BentNormalHistoryOutput = BentNormalSource;
-		IrradianceHistoryOutput = IrradianceSource;
-	}
-}
-
-void PostProcessBentNormalAO(
-	FRHICommandList& RHICmdList, 
-	const FDistanceFieldAOParameters& Parameters, 
-	const FViewInfo& View, 
-	FSceneRenderTargetItem& VelocityTexture,
-	FSceneRenderTargetItem& BentNormalInterpolation, 
-	IPooledRenderTarget* IrradianceInterpolation,
-	FSceneRenderTargetItem& DistanceFieldNormal,
-	TRefCountPtr<IPooledRenderTarget>& BentNormalOutput,
-	TRefCountPtr<IPooledRenderTarget>& IrradianceOutput)
-{
-	const bool bUseDistanceFieldGI = IsDistanceFieldGIAllowed(View);
-
-	TRefCountPtr<IPooledRenderTarget> DistanceFieldAOBentNormal;
-	TRefCountPtr<IPooledRenderTarget> DistanceFieldIrradiance;
-	AllocateOrReuseAORenderTarget(DistanceFieldAOBentNormal, TEXT("DistanceFieldBentNormalAO"), true);
-
-	if (bUseDistanceFieldGI)
-	{
-		AllocateOrReuseAORenderTarget(DistanceFieldIrradiance, TEXT("DistanceFieldIrradiance"), false);
-	}
-
-	TRefCountPtr<IPooledRenderTarget> DistanceFieldAOBentNormal2;
-	TRefCountPtr<IPooledRenderTarget> DistanceFieldIrradiance2;
-
-	if (GAOFillGaps)
-	{
-		AllocateOrReuseAORenderTarget(DistanceFieldAOBentNormal2, TEXT("DistanceFieldBentNormalAO2"), true);
-
-		if (bUseDistanceFieldGI)
-		{
-			AllocateOrReuseAORenderTarget(DistanceFieldIrradiance2, TEXT("DistanceFieldIrradiance2"), false);
-		}
-	}
-
-	{
-		SCOPED_DRAW_EVENT(RHICmdList, AOCombine);
-
-		TRefCountPtr<IPooledRenderTarget> EffectiveDiffuseIrradianceTarget;
-
-		if (bUseDistanceFieldGI)
-		{
-			EffectiveDiffuseIrradianceTarget = GAOFillGaps ? DistanceFieldIrradiance2 : DistanceFieldIrradiance;
-		}
-
-		FTextureRHIParamRef RenderTargets[2] =
-		{
-			GAOFillGaps ? DistanceFieldAOBentNormal2->GetRenderTargetItem().TargetableTexture : DistanceFieldAOBentNormal->GetRenderTargetItem().TargetableTexture,
-			EffectiveDiffuseIrradianceTarget ? EffectiveDiffuseIrradianceTarget->GetRenderTargetItem().TargetableTexture : NULL,
-		};
-
-		SetRenderTargets(RHICmdList, ARRAY_COUNT(RenderTargets) - (bUseDistanceFieldGI ? 0 : 1), RenderTargets, FTextureRHIParamRef(), 0, NULL);
-
-		{
-			RHICmdList.SetViewport(0, 0, 0.0f, View.ViewRect.Width() / GAODownsampleFactor, View.ViewRect.Height() / GAODownsampleFactor, 1.0f);
-			RHICmdList.SetRasterizerState(TStaticRasterizerState<FM_Solid, CM_None>::GetRHI());
-			RHICmdList.SetDepthStencilState(TStaticDepthStencilState<false, CF_Always>::GetRHI());
-			RHICmdList.SetBlendState(TStaticBlendState<>::GetRHI());
-
-			TShaderMapRef<FPostProcessVS> VertexShader(View.ShaderMap);
-
-			if (bUseDistanceFieldGI)
-			{
-				TShaderMapRef<TDistanceFieldAOCombinePS<true> > PixelShader(View.ShaderMap);
-				static FGlobalBoundShaderState BoundShaderState;
-				SetGlobalBoundShaderState(RHICmdList, View.GetFeatureLevel(), BoundShaderState, GFilterVertexDeclaration.VertexDeclarationRHI, *VertexShader, *PixelShader);
-				PixelShader->SetParameters(RHICmdList, View, BentNormalInterpolation, IrradianceInterpolation, DistanceFieldNormal, Parameters);
-			}
-			else
-			{
-				TShaderMapRef<TDistanceFieldAOCombinePS<false> > PixelShader(View.ShaderMap);
-				static FGlobalBoundShaderState BoundShaderState;
-				SetGlobalBoundShaderState(RHICmdList, View.GetFeatureLevel(), BoundShaderState, GFilterVertexDeclaration.VertexDeclarationRHI, *VertexShader, *PixelShader);
-				PixelShader->SetParameters(RHICmdList, View, BentNormalInterpolation, IrradianceInterpolation, DistanceFieldNormal, Parameters);
-			}
-
-			DrawRectangle( 
-				RHICmdList,
-				0, 0, 
-				View.ViewRect.Width() / GAODownsampleFactor, View.ViewRect.Height() / GAODownsampleFactor,
-				0, 0, 
-				View.ViewRect.Width() / GAODownsampleFactor, View.ViewRect.Height() / GAODownsampleFactor,
-				FIntPoint(View.ViewRect.Width() / GAODownsampleFactor, View.ViewRect.Height() / GAODownsampleFactor),
-				GetBufferSizeForAO(),
-				*VertexShader);
-		}
-
-		RHICmdList.CopyToResolveTarget(DistanceFieldAOBentNormal2->GetRenderTargetItem().TargetableTexture, DistanceFieldAOBentNormal2->GetRenderTargetItem().ShaderResourceTexture, false, FResolveParams());
-			
-		if (bUseDistanceFieldGI)
-		{
-			RHICmdList.CopyToResolveTarget(EffectiveDiffuseIrradianceTarget->GetRenderTargetItem().TargetableTexture, EffectiveDiffuseIrradianceTarget->GetRenderTargetItem().ShaderResourceTexture, false, FResolveParams());
-		}
-	}
-
-	if (GAOFillGaps)
-	{
-		SCOPED_DRAW_EVENT(RHICmdList, FillGaps);
-
-		FTextureRHIParamRef RenderTargets[2] =
-		{
-			DistanceFieldAOBentNormal->GetRenderTargetItem().TargetableTexture,
-			bUseDistanceFieldGI ? DistanceFieldIrradiance->GetRenderTargetItem().TargetableTexture : NULL,
-		};
-
-		SetRenderTargets(RHICmdList, ARRAY_COUNT(RenderTargets) - (bUseDistanceFieldGI ? 0 : 1), RenderTargets, FTextureRHIParamRef(), 0, NULL);
-
-		{
-			RHICmdList.SetViewport(0, 0, 0.0f, View.ViewRect.Width() / GAODownsampleFactor, View.ViewRect.Height() / GAODownsampleFactor, 1.0f);
-			RHICmdList.SetRasterizerState(TStaticRasterizerState<FM_Solid, CM_None>::GetRHI());
-			RHICmdList.SetDepthStencilState(TStaticDepthStencilState<false, CF_Always>::GetRHI());
-			RHICmdList.SetBlendState(TStaticBlendState<>::GetRHI());
-
-			TShaderMapRef<FPostProcessVS> VertexShader(View.ShaderMap);
-
-			if (GAOFillGapsHighQuality)
-			{
-				if (bUseDistanceFieldGI)
-				{
-					TShaderMapRef<TFillGapsPS<true, true> > PixelShader(View.ShaderMap);
-
-					static FGlobalBoundShaderState BoundShaderState;
-					SetGlobalBoundShaderState(RHICmdList, View.GetFeatureLevel(), BoundShaderState, GFilterVertexDeclaration.VertexDeclarationRHI, *VertexShader, *PixelShader);
-					PixelShader->SetParameters(RHICmdList, View, DistanceFieldAOBentNormal2->GetRenderTargetItem(), DistanceFieldIrradiance2);
-				}
-				else
-				{
-					TShaderMapRef<TFillGapsPS<false, true> > PixelShader(View.ShaderMap);
-
-					static FGlobalBoundShaderState BoundShaderState;
-					SetGlobalBoundShaderState(RHICmdList, View.GetFeatureLevel(), BoundShaderState, GFilterVertexDeclaration.VertexDeclarationRHI, *VertexShader, *PixelShader);
-					PixelShader->SetParameters(RHICmdList, View, DistanceFieldAOBentNormal2->GetRenderTargetItem(), DistanceFieldIrradiance2);
-				}
-			}
-			else
-			{
-				if (bUseDistanceFieldGI)
-				{
-					TShaderMapRef<TFillGapsPS<true, false> > PixelShader(View.ShaderMap);
-
-					static FGlobalBoundShaderState BoundShaderState;
-					SetGlobalBoundShaderState(RHICmdList, View.GetFeatureLevel(), BoundShaderState, GFilterVertexDeclaration.VertexDeclarationRHI, *VertexShader, *PixelShader);
-					PixelShader->SetParameters(RHICmdList, View, DistanceFieldAOBentNormal2->GetRenderTargetItem(), DistanceFieldIrradiance2);
-				}
-				else
-				{
-					TShaderMapRef<TFillGapsPS<false, false> > PixelShader(View.ShaderMap);
-
-					static FGlobalBoundShaderState BoundShaderState;
-					SetGlobalBoundShaderState(RHICmdList, View.GetFeatureLevel(), BoundShaderState, GFilterVertexDeclaration.VertexDeclarationRHI, *VertexShader, *PixelShader);
-					PixelShader->SetParameters(RHICmdList, View, DistanceFieldAOBentNormal2->GetRenderTargetItem(), DistanceFieldIrradiance2);
-				}
-			}
-
-			VertexShader->SetParameters(RHICmdList, View);
-
-			DrawRectangle( 
-				RHICmdList,
-				0, 0, 
-				View.ViewRect.Width() / GAODownsampleFactor, View.ViewRect.Height() / GAODownsampleFactor,
-				0, 0, 
-				View.ViewRect.Width() / GAODownsampleFactor, View.ViewRect.Height() / GAODownsampleFactor,
-				FIntPoint(View.ViewRect.Width() / GAODownsampleFactor, View.ViewRect.Height() / GAODownsampleFactor),
-				GSceneRenderTargets.GetBufferSizeXY() / FIntPoint(GAODownsampleFactor, GAODownsampleFactor),
-				*VertexShader);
-		}
-
-		RHICmdList.CopyToResolveTarget(DistanceFieldAOBentNormal->GetRenderTargetItem().TargetableTexture, DistanceFieldAOBentNormal->GetRenderTargetItem().ShaderResourceTexture, false, FResolveParams());
-			
-		if (bUseDistanceFieldGI)
-		{
-			RHICmdList.CopyToResolveTarget(DistanceFieldIrradiance->GetRenderTargetItem().TargetableTexture, DistanceFieldIrradiance->GetRenderTargetItem().ShaderResourceTexture, false, FResolveParams());
-		}
-	}
-
-	FSceneViewState* ViewState = (FSceneViewState*)View.State;
-	TRefCountPtr<IPooledRenderTarget>* BentNormalHistoryState = ViewState ? &ViewState->DistanceFieldAOHistoryRT : NULL;
-	TRefCountPtr<IPooledRenderTarget>* IrradianceHistoryState = ViewState ? &ViewState->DistanceFieldIrradianceHistoryRT : NULL;
-	BentNormalOutput = DistanceFieldAOBentNormal;
-	IrradianceOutput = DistanceFieldIrradiance;
-
-	if (GAOUseHistory)
-	{
-		UpdateHistory(
-			RHICmdList,
-			View, 
-			TEXT("DistanceFieldAOHistory"),
-			TEXT("DistanceFieldIrradianceHistory"),
-			VelocityTexture,
-			BentNormalHistoryState,
-			IrradianceHistoryState,
-			DistanceFieldAOBentNormal, 
-			DistanceFieldIrradiance,
-			BentNormalOutput,
-			IrradianceOutput);
-	}
-}
-
-enum EAOUpsampleType
-{
-	AOUpsample_OutputBentNormal,
-	AOUpsample_OutputAO,
-	AOUpsample_OutputBentNormalAndIrradiance,
-	AOUpsample_OutputIrradiance
-};
-
-template<EAOUpsampleType UpsampleType>
-class TDistanceFieldAOUpsamplePS : public FGlobalShader
-{
-	DECLARE_SHADER_TYPE(TDistanceFieldAOUpsamplePS, Global);
-public:
-
-	static bool ShouldCache(EShaderPlatform Platform)
-	{
-		return IsFeatureLevelSupported(Platform, ERHIFeatureLevel::SM5) && DoesPlatformSupportDistanceFieldAO(Platform);
-	}
-
-	static void ModifyCompilationEnvironment(EShaderPlatform Platform, FShaderCompilerEnvironment& OutEnvironment)
-	{
-		OutEnvironment.SetDefine(TEXT("DOWNSAMPLE_FACTOR"), GAODownsampleFactor);
-		OutEnvironment.SetDefine(TEXT("OUTPUT_BENT_NORMAL"), (UpsampleType == AOUpsample_OutputBentNormal || UpsampleType == AOUpsample_OutputBentNormalAndIrradiance) ? TEXT("1") : TEXT("0"));
-		OutEnvironment.SetDefine(TEXT("SUPPORT_IRRADIANCE"), (UpsampleType == AOUpsample_OutputIrradiance || UpsampleType == AOUpsample_OutputBentNormalAndIrradiance) ? TEXT("1") : TEXT("0"));
-		OutEnvironment.SetDefine(TEXT("OUTPUT_AO"), (UpsampleType == AOUpsample_OutputAO) ? TEXT("1") : TEXT("0"));
-	}
-
-	/** Default constructor. */
-	TDistanceFieldAOUpsamplePS() {}
-
-	/** Initialization constructor. */
-	TDistanceFieldAOUpsamplePS(const ShaderMetaType::CompiledShaderInitializerType& Initializer)
-		: FGlobalShader(Initializer)
-	{
-		DeferredParameters.Bind(Initializer.ParameterMap);
-		BentNormalAOTexture.Bind(Initializer.ParameterMap,TEXT("BentNormalAOTexture"));
-		BentNormalAOSampler.Bind(Initializer.ParameterMap,TEXT("BentNormalAOSampler"));
-		IrradianceTexture.Bind(Initializer.ParameterMap,TEXT("IrradianceTexture"));
-		IrradianceSampler.Bind(Initializer.ParameterMap,TEXT("IrradianceSampler"));
-	}
-
-	void SetParameters(FRHICommandList& RHICmdList, const FSceneView& View, TRefCountPtr<IPooledRenderTarget>& DistanceFieldAOBentNormal, IPooledRenderTarget* DistanceFieldIrradiance)
-	{
-		const FPixelShaderRHIParamRef ShaderRHI = GetPixelShader();
-
-		FGlobalShader::SetParameters(RHICmdList, ShaderRHI, View);
-		DeferredParameters.Set(RHICmdList, ShaderRHI, View);
-
-		SetTextureParameter(RHICmdList, ShaderRHI, BentNormalAOTexture, BentNormalAOSampler, TStaticSamplerState<SF_Bilinear>::GetRHI(), DistanceFieldAOBentNormal->GetRenderTargetItem().ShaderResourceTexture);
-
-		if (IrradianceTexture.IsBound())
-		{
-			SetTextureParameter(RHICmdList, ShaderRHI, IrradianceTexture, IrradianceSampler, TStaticSamplerState<SF_Bilinear>::GetRHI(), DistanceFieldIrradiance->GetRenderTargetItem().ShaderResourceTexture);
-		}
-	}
-	// FShader interface.
-	virtual bool Serialize(FArchive& Ar) override
-	{
-		bool bShaderHasOutdatedParameters = FGlobalShader::Serialize(Ar);
-		Ar << DeferredParameters;
-		Ar << BentNormalAOTexture;
-		Ar << BentNormalAOSampler;
-		Ar << IrradianceTexture;
-		Ar << IrradianceSampler;
-		return bShaderHasOutdatedParameters;
-	}
-
-private:
-
-	FDeferredPixelShaderParameters DeferredParameters;
-	FShaderResourceParameter BentNormalAOTexture;
-	FShaderResourceParameter BentNormalAOSampler;
-	FShaderResourceParameter IrradianceTexture;
-	FShaderResourceParameter IrradianceSampler;
-};
-
-#define IMPLEMENT_UPSAMPLE_PS_TYPE(UpsampleType) \
-	typedef TDistanceFieldAOUpsamplePS<UpsampleType> TDistanceFieldAOUpsamplePS##UpsampleType; \
-	IMPLEMENT_SHADER_TYPE(template<>,TDistanceFieldAOUpsamplePS##UpsampleType,TEXT("DistanceFieldLightingPost"),TEXT("AOUpsamplePS"),SF_Pixel);
-
-IMPLEMENT_UPSAMPLE_PS_TYPE(AOUpsample_OutputBentNormal)
-IMPLEMENT_UPSAMPLE_PS_TYPE(AOUpsample_OutputAO)
-IMPLEMENT_UPSAMPLE_PS_TYPE(AOUpsample_OutputBentNormalAndIrradiance)
-IMPLEMENT_UPSAMPLE_PS_TYPE(AOUpsample_OutputIrradiance)
-
-void UpsampleBentNormalAO(
-	FRHICommandList& RHICmdList, 
-	const TArray<FViewInfo>& Views, 
-	TRefCountPtr<IPooledRenderTarget>& DistanceFieldAOBentNormal, 
-	TRefCountPtr<IPooledRenderTarget>& DistanceFieldIrradiance,
-	bool bVisualizeAmbientOcclusion,
-	bool bVisualizeGlobalIllumination)
-{
-	for (int32 ViewIndex = 0; ViewIndex < Views.Num(); ViewIndex++)
-	{
-		const FViewInfo& View = Views[ViewIndex];
-		const bool bUseDistanceFieldGI = IsDistanceFieldGIAllowed(View);
-
-		SCOPED_DRAW_EVENT(RHICmdList, UpsampleAO);
-
-		RHICmdList.SetViewport(View.ViewRect.Min.X, View.ViewRect.Min.Y, 0.0f, View.ViewRect.Max.X, View.ViewRect.Max.Y, 1.0f);
-		RHICmdList.SetRasterizerState(TStaticRasterizerState<FM_Solid, CM_None>::GetRHI());
-		RHICmdList.SetDepthStencilState(TStaticDepthStencilState<false, CF_Always>::GetRHI());
-		RHICmdList.SetBlendState(TStaticBlendState<>::GetRHI());
-
-		TShaderMapRef<FPostProcessVS> VertexShader(View.ShaderMap);
-
-		if (bVisualizeAmbientOcclusion)
-		{
-			TShaderMapRef<TDistanceFieldAOUpsamplePS<AOUpsample_OutputAO> > PixelShader(View.ShaderMap);
-			static FGlobalBoundShaderState BoundShaderState;
-			SetGlobalBoundShaderState(RHICmdList, View.GetFeatureLevel(), BoundShaderState, GFilterVertexDeclaration.VertexDeclarationRHI, *VertexShader, *PixelShader);
-			PixelShader->SetParameters(RHICmdList, View, DistanceFieldAOBentNormal, DistanceFieldIrradiance);
-		}
-		else if (bVisualizeGlobalIllumination && bUseDistanceFieldGI)
-		{
-			TShaderMapRef<TDistanceFieldAOUpsamplePS<AOUpsample_OutputIrradiance> > PixelShader(View.ShaderMap);
-			static FGlobalBoundShaderState BoundShaderState;
-			SetGlobalBoundShaderState(RHICmdList, View.GetFeatureLevel(), BoundShaderState, GFilterVertexDeclaration.VertexDeclarationRHI, *VertexShader, *PixelShader);
-			PixelShader->SetParameters(RHICmdList, View, DistanceFieldAOBentNormal, DistanceFieldIrradiance);
-		}
-		else
-		{
-			if (bUseDistanceFieldGI)
-			{
-				TShaderMapRef<TDistanceFieldAOUpsamplePS<AOUpsample_OutputBentNormalAndIrradiance> > PixelShader(View.ShaderMap);
-				static FGlobalBoundShaderState BoundShaderState;
-				SetGlobalBoundShaderState(RHICmdList, View.GetFeatureLevel(), BoundShaderState, GFilterVertexDeclaration.VertexDeclarationRHI, *VertexShader, *PixelShader);
-				PixelShader->SetParameters(RHICmdList, View, DistanceFieldAOBentNormal, DistanceFieldIrradiance);
-			}
-			else
-			{
-				TShaderMapRef<TDistanceFieldAOUpsamplePS<AOUpsample_OutputBentNormal> > PixelShader(View.ShaderMap);
-				static FGlobalBoundShaderState BoundShaderState;
-				SetGlobalBoundShaderState(RHICmdList, View.GetFeatureLevel(), BoundShaderState, GFilterVertexDeclaration.VertexDeclarationRHI, *VertexShader, *PixelShader);
-				PixelShader->SetParameters(RHICmdList, View, DistanceFieldAOBentNormal, DistanceFieldIrradiance);
-			}
-		}
-
-		DrawRectangle( 
-			RHICmdList,
-			0, 0, 
-			View.ViewRect.Width(), View.ViewRect.Height(),
-			View.ViewRect.Min.X / GAODownsampleFactor, View.ViewRect.Min.Y / GAODownsampleFactor, 
-			View.ViewRect.Width() / GAODownsampleFactor, View.ViewRect.Height() / GAODownsampleFactor,
-			FIntPoint(View.ViewRect.Width(), View.ViewRect.Height()),
-			GetBufferSizeForAO(),
-			*VertexShader);
-	}
-}
-
-/** Generates a pseudo-random position inside the unit sphere, uniformly distributed over the volume of the sphere. */
-FVector GetUnitPosition2(FRandomStream& RandomStream)
-{
-	FVector Result;
-	// Use rejection sampling to generate a valid sample
-	do
-	{
-		Result.X = RandomStream.GetFraction() * 2 - 1;
-		Result.Y = RandomStream.GetFraction() * 2 - 1;
-		Result.Z = RandomStream.GetFraction() * 2 - 1;
-	} while( Result.SizeSquared() > 1.f );
-	return Result;
-}
-
-/** Generates a pseudo-random unit vector, uniformly distributed over all directions. */
-FVector GetUnitVector2(FRandomStream& RandomStream)
-{
-	return GetUnitPosition2(RandomStream).GetUnsafeNormal();
-}
-
-void GenerateBestSpacedVectors()
-{
-	static bool bGenerated = false;
-	bool bApplyRepulsion = false;
-
-	if (bApplyRepulsion && !bGenerated)
-	{
-		bGenerated = true;
-
-		FVector OriginalSpacedVectors9[ARRAY_COUNT(SpacedVectors9)];
-
-		for (int32 i = 0; i < ARRAY_COUNT(OriginalSpacedVectors9); i++)
-		{
-			OriginalSpacedVectors9[i] = SpacedVectors9[i];
-		}
-
-		float CosHalfAngle = 1 - 1.0f / (float)ARRAY_COUNT(OriginalSpacedVectors9);
-		// Used to prevent self-shadowing on a plane
-		float AngleBias = .03f;
-		float MinAngle = FMath::Acos(CosHalfAngle) + AngleBias;
-		float MinZ = FMath::Sin(MinAngle);
-
-		// Relaxation iterations by repulsion
-		for (int32 Iteration = 0; Iteration < 10000; Iteration++)
-		{
-			for (int32 i = 0; i < ARRAY_COUNT(OriginalSpacedVectors9); i++)
-			{
-				FVector Force(0.0f, 0.0f, 0.0f);
-
-				for (int32 j = 0; j < ARRAY_COUNT(OriginalSpacedVectors9); j++)
-				{
-					if (i != j)
-					{
-						FVector Distance = OriginalSpacedVectors9[i] - OriginalSpacedVectors9[j];
-						float Dot = OriginalSpacedVectors9[i] | OriginalSpacedVectors9[j];
-
-						if (Dot > 0)
-						{
-							// Repulsion force
-							Force += .001f * Distance.GetSafeNormal() * Dot * Dot * Dot * Dot;
-						}
-					}
-				}
-
-				FVector NewPosition = OriginalSpacedVectors9[i] + Force;
-				NewPosition.Z = FMath::Max(NewPosition.Z, MinZ);
-				NewPosition = NewPosition.GetSafeNormal();
-				OriginalSpacedVectors9[i] = NewPosition;
-			}
-		}
-
-		for (int32 i = 0; i < ARRAY_COUNT(OriginalSpacedVectors9); i++)
-		{
-			UE_LOG(LogDistanceField, Log, TEXT("FVector(%f, %f, %f),"), OriginalSpacedVectors9[i].X, OriginalSpacedVectors9[i].Y, OriginalSpacedVectors9[i].Z);
-		}
-
-		int32 temp = 0;
-	}
-
-	bool bBruteForceGenerateConeDirections = false;
-
-	if (bBruteForceGenerateConeDirections)
-	{
-		FVector BestSpacedVectors9[9];
-		float BestCoverage = 0;
-		// Each cone covers an area of ConeSolidAngle = HemisphereSolidAngle / NumCones
-		// HemisphereSolidAngle = 2 * PI
-		// ConeSolidAngle = 2 * PI * (1 - cos(ConeHalfAngle))
-		// cos(ConeHalfAngle) = 1 - 1 / NumCones
-		float CosHalfAngle = 1 - 1.0f / (float)ARRAY_COUNT(BestSpacedVectors9);
-		// Prevent self-intersection in sample set
-		float MinAngle = FMath::Acos(CosHalfAngle);
-		float MinZ = FMath::Sin(MinAngle);
-		FRandomStream RandomStream(123567);
-
-		// Super slow random brute force search
-		for (int i = 0; i < 1000000; i++)
-		{
-			FVector CandidateSpacedVectors[ARRAY_COUNT(BestSpacedVectors9)];
-
-			for (int j = 0; j < ARRAY_COUNT(CandidateSpacedVectors); j++)
-			{
-				FVector NewSample;
+			for (int j = 0; j < ARRAY_COUNT(CandidateSpacedVectors); j++)
+			{
+				FVector NewSample;
 
 				// Reject invalid directions until we get a valid one
 				do 
@@ -3671,130 +2525,372 @@ void RenderIrradianceCacheInterpolation(
 						TShaderMapRef<TIrradianceCacheSplatVS<false, false> > VertexShader(View.ShaderMap);
 						TShaderMapRef<TIrradianceCacheSplatPS<false, false> > PixelShader(View.ShaderMap);
 
-						static FGlobalBoundShaderState BoundShaderState;
-						SetGlobalBoundShaderState(RHICmdList, View.GetFeatureLevel(), BoundShaderState, GScreenVertexDeclaration.VertexDeclarationRHI, *VertexShader, *PixelShader);
+						static FGlobalBoundShaderState BoundShaderState;
+						SetGlobalBoundShaderState(RHICmdList, View.GetFeatureLevel(), BoundShaderState, GScreenVertexDeclaration.VertexDeclarationRHI, *VertexShader, *PixelShader);
+
+						VertexShader->SetParameters(RHICmdList, View, Parameters, SplatSourceDepthLevel, DestLevelDownsampleFactor, NormalizedOffsetToPixelCenter);
+						PixelShader->SetParameters(RHICmdList, View, DistanceFieldNormal, DestLevelDownsampleFactor, Parameters);
+					}
+
+					RHICmdList.DrawPrimitiveIndirect(PT_TriangleList, SurfaceCacheResources.Level[SplatSourceDepthLevel]->ScatterDrawParameters.Buffer, 0);
+				}
+			}
+		}
+
+		RHICmdList.CopyToResolveTarget(BentNormalInterpolationTarget->GetRenderTargetItem().TargetableTexture, BentNormalInterpolationTarget->GetRenderTargetItem().ShaderResourceTexture, false, FResolveParams());
+
+		if (bFinalInterpolation && bUseDistanceFieldGI && (GAOInterpolationDepthTesting || GAOInterpolationStencilTesting))
+		{
+			RHICmdList.CopyToResolveTarget(IrradianceInterpolationTarget->GetRenderTargetItem().TargetableTexture, IrradianceInterpolationTarget->GetRenderTargetItem().ShaderResourceTexture, false, FResolveParams());
+		}
+	}
+}
+
+void ListDistanceFieldLightingMemory(const FViewInfo& View)
+{
+	const FScene* Scene = (const FScene*)View.Family->Scene;
+	UE_LOG(LogTemp, Log, TEXT("Shared GPU memory (excluding render targets)"));
+
+	if (Scene->DistanceFieldSceneData.NumObjectsInBuffer > 0)
+	{
+		UE_LOG(LogTemp, Log, TEXT("   Scene Object data %.3fMb"), Scene->DistanceFieldSceneData.ObjectBuffers->GetSizeBytes() / 1024.0f / 1024.0f);
+	}
+	
+	UE_LOG(LogTemp, Log, TEXT("   %s"), *GDistanceFieldVolumeTextureAtlas.GetSizeString());
+	extern FString GetObjectBufferMemoryString();
+	UE_LOG(LogTemp, Log, TEXT("   %s"), *GetObjectBufferMemoryString());
+	UE_LOG(LogTemp, Log, TEXT(""));
+	UE_LOG(LogTemp, Log, TEXT("Distance Field AO"));
+
+	if (Scene->SurfaceCacheResources)
+	{
+		UE_LOG(LogTemp, Log, TEXT("   Surface cache %.3fMb"), Scene->SurfaceCacheResources->GetSizeBytes() / 1024.0f / 1024.0f);
+	}
+	
+	UE_LOG(LogTemp, Log, TEXT("   Temporary cache %.3fMb"), GTemporaryIrradianceCacheResources.GetSizeBytes() / 1024.0f / 1024.0f);
+	UE_LOG(LogTemp, Log, TEXT("   Culled objects %.3fMb"), GAOCulledObjectBuffers.Buffers.GetSizeBytes() / 1024.0f / 1024.0f);
+
+	FTileIntersectionResources* TileIntersectionResources = ((FSceneViewState*)View.State)->AOTileIntersectionResources;
+
+	if (TileIntersectionResources)
+	{
+		UE_LOG(LogTemp, Log, TEXT("   Tile Culled objects %.3fMb"), TileIntersectionResources->GetSizeBytes() / 1024.0f / 1024.0f);
+	}
+
+	FAOScreenGridResources* ScreenGridResources = ((FSceneViewState*)View.State)->AOScreenGridResources;
+
+	if (ScreenGridResources)
+	{
+		UE_LOG(LogTemp, Log, TEXT("   Screen grid temporaries %.3fMb"), ScreenGridResources->GetSizeBytesForAO() / 1024.0f / 1024.0f);
+	}
+	
+	extern void ListGlobalDistanceFieldMemory();
+	ListGlobalDistanceFieldMemory();
+
+	UE_LOG(LogTemp, Log, TEXT(""));
+	UE_LOG(LogTemp, Log, TEXT("Distance Field GI"));
+
+	if (Scene->DistanceFieldSceneData.SurfelBuffers)
+	{
+		UE_LOG(LogTemp, Log, TEXT("   Scene surfel data %.3fMb"), Scene->DistanceFieldSceneData.SurfelBuffers->GetSizeBytes() / 1024.0f / 1024.0f);
+	}
+	
+	if (Scene->DistanceFieldSceneData.InstancedSurfelBuffers)
+	{
+		UE_LOG(LogTemp, Log, TEXT("   Instanced scene surfel data %.3fMb"), Scene->DistanceFieldSceneData.InstancedSurfelBuffers->GetSizeBytes() / 1024.0f / 1024.0f);
+	}
+	
+	if (ScreenGridResources)
+	{
+		UE_LOG(LogTemp, Log, TEXT("   Screen grid temporaries %.3fMb"), ScreenGridResources->GetSizeBytesForGI() / 1024.0f / 1024.0f);
+	}
+
+	extern void ListDistanceFieldGIMemory(const FViewInfo& View);
+	ListDistanceFieldGIMemory(View);
+}
+
+bool SupportsDistanceFieldAO(ERHIFeatureLevel::Type FeatureLevel, EShaderPlatform ShaderPlatform)
+{
+	return GDistanceFieldAO 
+		&& FeatureLevel >= ERHIFeatureLevel::SM5
+		&& DoesPlatformSupportDistanceFieldAO(ShaderPlatform);
+}
+
+bool ShouldRenderDynamicSkyLight(const FScene* Scene, const FSceneViewFamily& ViewFamily)
+{
+	return Scene->SkyLight
+		&& Scene->SkyLight->ProcessedTexture
+		&& !Scene->SkyLight->bWantsStaticShadowing
+		&& !Scene->SkyLight->bHasStaticLighting
+		&& ViewFamily.EngineShowFlags.SkyLighting
+		&& Scene->GetFeatureLevel() >= ERHIFeatureLevel::SM4
+		&& !IsSimpleDynamicLightingEnabled() 
+		&& !ViewFamily.EngineShowFlags.VisualizeLightCulling;
+}
+
+bool FDeferredShadingSceneRenderer::ShouldPrepareForDistanceFieldAO() const
+{
+	return SupportsDistanceFieldAO(Scene->GetFeatureLevel(), Scene->GetShaderPlatform())
+		&& ((ShouldRenderDynamicSkyLight(Scene, ViewFamily) && Scene->SkyLight->bCastShadows && ViewFamily.EngineShowFlags.DistanceFieldAO)
+			|| ViewFamily.EngineShowFlags.VisualizeMeshDistanceFields
+			|| ViewFamily.EngineShowFlags.VisualizeDistanceFieldAO
+			|| ViewFamily.EngineShowFlags.VisualizeDistanceFieldGI);
+}
+
+void RenderDistanceFieldAOSurfaceCache(
+	FRHICommandListImmediate& RHICmdList, 
+	const FViewInfo& View,
+	FScene* Scene,
+	FIntPoint TileListGroupSize,
+	const FDistanceFieldAOParameters& Parameters, 
+	const TRefCountPtr<IPooledRenderTarget>& VelocityTexture,
+	const TRefCountPtr<IPooledRenderTarget>& DistanceFieldNormal, 
+	TRefCountPtr<IPooledRenderTarget>& OutDynamicBentNormalAO, 
+	TRefCountPtr<IPooledRenderTarget>& OutDynamicIrradiance)
+{
+	const bool bUseDistanceFieldGI = IsDistanceFieldGIAllowed(View);
+
+	// Create surface cache state that will persist across frames if needed
+	if (!Scene->SurfaceCacheResources)
+	{
+		Scene->SurfaceCacheResources = new FSurfaceCacheResources();
+	}
+
+	FSurfaceCacheResources& SurfaceCacheResources = *Scene->SurfaceCacheResources;
+
+	FIntPoint DownsampledViewRect = FIntPoint::DivideAndRoundDown(View.ViewRect.Size(), GAODownsampleFactor);
+	// Needed because samples are kept across frames
+	int32 SlackAmount = 4;
+	int32 MaxSurfaceCacheSamples = DownsampledViewRect.X * DownsampledViewRect.Y / (1 << (2 * GAOMinLevel * GAOPowerOfTwoBetweenLevels));
+	SurfaceCacheResources.AllocateFor(GAOMinLevel, GAOMaxLevel, SlackAmount * MaxSurfaceCacheSamples);
+
+	GTemporaryIrradianceCacheResources.AllocateFor(MaxSurfaceCacheSamples);
+
+	uint32 ClearValues[4] = {0};
+
+	if (!SurfaceCacheResources.bClearedResources 
+		|| !GAOReuseAcrossFrames 
+		// Drop records that will have uninitialized Irradiance if switching from AO only to AO + GI
+		|| (bUseDistanceFieldGI && !SurfaceCacheResources.bHasIrradiance))
+	{
+		// Reset the number of active cache records to 0
+		for (int32 DepthLevel = GAOMaxLevel; DepthLevel >= GAOMinLevel; DepthLevel--)
+		{
+			RHICmdList.ClearUAV(SurfaceCacheResources.Level[DepthLevel]->ScatterDrawParameters.UAV, ClearValues);
+		}
+
+		RHICmdList.ClearUAV(SurfaceCacheResources.TempResources->ScatterDrawParameters.UAV, ClearValues);
+
+		SurfaceCacheResources.bClearedResources = true;
+	}
+
+	SurfaceCacheResources.bHasIrradiance = bUseDistanceFieldGI;
+
+	if (GAOReuseAcrossFrames)
+	{
+		SCOPED_DRAW_EVENT(RHICmdList, TrimRecords);
+
+		FSurfaceCacheResources& SurfaceCacheResources = *Scene->SurfaceCacheResources;
+
+		// Copy and trim last frame's surface cache samples
+		for (int32 DepthLevel = GAOMaxLevel; DepthLevel >= GAOMinLevel; DepthLevel--)
+		{
+			if (GAOTrimOldRecordsFraction > 0)
+			{
+				{	
+					TShaderMapRef<FSetupCopyIndirectArgumentsCS> ComputeShader(View.ShaderMap);
+
+					RHICmdList.SetComputeShader(ComputeShader->GetComputeShader());
+					ComputeShader->SetParameters(RHICmdList, View, DepthLevel);
+					DispatchComputeShader(RHICmdList, *ComputeShader, 1, 1, 1);
+
+					ComputeShader->UnsetParameters(RHICmdList);
+				}
+
+				if (bUseDistanceFieldGI)
+				{
+					TShaderMapRef<TCopyIrradianceCacheSamplesCS<true> > ComputeShader(View.ShaderMap);
+
+					RHICmdList.SetComputeShader(ComputeShader->GetComputeShader());
+					ComputeShader->SetParameters(RHICmdList, View, DepthLevel);
+					DispatchIndirectComputeShader(RHICmdList, *ComputeShader, SurfaceCacheResources.DispatchParameters.Buffer, 0);
+
+					ComputeShader->UnsetParameters(RHICmdList);
+				}
+				else
+				{
+					TShaderMapRef<TCopyIrradianceCacheSamplesCS<false> > ComputeShader(View.ShaderMap);
+
+					RHICmdList.SetComputeShader(ComputeShader->GetComputeShader());
+					ComputeShader->SetParameters(RHICmdList, View, DepthLevel);
+					DispatchIndirectComputeShader(RHICmdList, *ComputeShader, SurfaceCacheResources.DispatchParameters.Buffer, 0);
+
+					ComputeShader->UnsetParameters(RHICmdList);
+				}
+
+				Swap(SurfaceCacheResources.Level[DepthLevel], SurfaceCacheResources.TempResources);
+			}
+		}
+	}
+
+	// Start from the highest depth, which will consider the fewest potential shading points
+	// Each level potentially prevents the next higher resolution level from doing expensive shading work
+	// This is the core of the surface cache algorithm
+	for (int32 DepthLevel = GAOMaxLevel; DepthLevel >= GAOMinLevel; DepthLevel--)
+	{
+		int32 DestLevelDownsampleFactor = GAODownsampleFactor * (1 << (DepthLevel * GAOPowerOfTwoBetweenLevels));
+
+		SCOPED_DRAW_EVENTF(RHICmdList, Level, TEXT("Level_%d"), DepthLevel);
+
+		TRefCountPtr<IPooledRenderTarget> DistanceFieldAOBentNormalSplat;
+
+		{
+			FIntPoint AOBufferSize = FIntPoint::DivideAndRoundUp(GSceneRenderTargets.GetBufferSizeXY(), DestLevelDownsampleFactor);
+			FPooledRenderTargetDesc Desc(FPooledRenderTargetDesc::Create2DDesc(AOBufferSize, PF_FloatRGBA, TexCreate_None, TexCreate_RenderTargetable | TexCreate_UAV, false));
+			GRenderTargetPool.FindFreeElement(Desc, DistanceFieldAOBentNormalSplat, TEXT("DistanceFieldAOBentNormalSplat"));
+		}
+
+		// Splat / interpolate the surface cache records onto the buffer sized for the current depth level
+		RenderIrradianceCacheInterpolation(RHICmdList, View, DistanceFieldAOBentNormalSplat, NULL, DistanceFieldNormal->GetRenderTargetItem(), DepthLevel, DestLevelDownsampleFactor, Parameters, false);
+
+		{
+			SCOPED_DRAW_EVENT(RHICmdList, PopulateAndConeTrace);
+			SetRenderTarget(RHICmdList, NULL, NULL);
+
+			// Save off the current record count before adding any more
+			{	
+				TShaderMapRef<FSaveStartIndexCS> ComputeShader(View.ShaderMap);
+
+				RHICmdList.SetComputeShader(ComputeShader->GetComputeShader());
+				ComputeShader->SetParameters(RHICmdList, View, DepthLevel);
+				DispatchComputeShader(RHICmdList, *ComputeShader, 1, 1, 1);
+
+				ComputeShader->UnsetParameters(RHICmdList);
+			}
+
+			// Create new records which haven't been shaded yet for shading points which don't have a valid interpolation from existing records
+			{
+				FIntPoint DownsampledViewSize = FIntPoint::DivideAndRoundUp(View.ViewRect.Size(), DestLevelDownsampleFactor);
+				uint32 GroupSizeX = FMath::DivideAndRoundUp(DownsampledViewSize.X, GDistanceFieldAOTileSizeX);
+				uint32 GroupSizeY = FMath::DivideAndRoundUp(DownsampledViewSize.Y, GDistanceFieldAOTileSizeY);
+
+				TShaderMapRef<FPopulateCacheCS> ComputeShader(View.ShaderMap);
+
+				RHICmdList.SetComputeShader(ComputeShader->GetComputeShader());
+				ComputeShader->SetParameters(RHICmdList, View, DistanceFieldAOBentNormalSplat->GetRenderTargetItem(), DistanceFieldNormal->GetRenderTargetItem(), DestLevelDownsampleFactor, DepthLevel, TileListGroupSize, Parameters);
+				DispatchComputeShader(RHICmdList, *ComputeShader, GroupSizeX, GroupSizeY, 1);
+
+				ComputeShader->UnsetParameters(RHICmdList);
+			}
+
+			{	
+				TShaderMapRef<TSetupFinalGatherIndirectArgumentsCS<true> > ComputeShader(View.ShaderMap);
+
+				RHICmdList.SetComputeShader(ComputeShader->GetComputeShader());
+				ComputeShader->SetParameters(RHICmdList, View, DepthLevel);
+				DispatchComputeShader(RHICmdList, *ComputeShader, 1, 1, 1);
+
+				ComputeShader->UnsetParameters(RHICmdList);
+			}
+
+			// Compute lighting for the new surface cache records by cone-stepping through the object distance fields
+			if (bUseDistanceFieldGI)
+			{
+				TShaderMapRef<TConeTraceSurfaceCacheOcclusionCS<true> > ComputeShader(View.ShaderMap);
+
+				RHICmdList.SetComputeShader(ComputeShader->GetComputeShader());
+				ComputeShader->SetParameters(RHICmdList, View, DistanceFieldAOBentNormalSplat->GetRenderTargetItem(), DistanceFieldNormal->GetRenderTargetItem(), DestLevelDownsampleFactor, DepthLevel, TileListGroupSize, Parameters);
+				DispatchIndirectComputeShader(RHICmdList, *ComputeShader, SurfaceCacheResources.DispatchParameters.Buffer, 0);
 
-						VertexShader->SetParameters(RHICmdList, View, Parameters, SplatSourceDepthLevel, DestLevelDownsampleFactor, NormalizedOffsetToPixelCenter);
-						PixelShader->SetParameters(RHICmdList, View, DistanceFieldNormal, DestLevelDownsampleFactor, Parameters);
-					}
+				ComputeShader->UnsetParameters(RHICmdList);
+			}
+			else
+			{
+				TShaderMapRef<TConeTraceSurfaceCacheOcclusionCS<false> > ComputeShader(View.ShaderMap);
 
-					RHICmdList.DrawPrimitiveIndirect(PT_TriangleList, SurfaceCacheResources.Level[SplatSourceDepthLevel]->ScatterDrawParameters.Buffer, 0);
-				}
+				RHICmdList.SetComputeShader(ComputeShader->GetComputeShader());
+				ComputeShader->SetParameters(RHICmdList, View, DistanceFieldAOBentNormalSplat->GetRenderTargetItem(), DistanceFieldNormal->GetRenderTargetItem(), DestLevelDownsampleFactor, DepthLevel, TileListGroupSize, Parameters);
+				DispatchIndirectComputeShader(RHICmdList, *ComputeShader, SurfaceCacheResources.DispatchParameters.Buffer, 0);
+
+				ComputeShader->UnsetParameters(RHICmdList);
 			}
 		}
 
-		RHICmdList.CopyToResolveTarget(BentNormalInterpolationTarget->GetRenderTargetItem().TargetableTexture, BentNormalInterpolationTarget->GetRenderTargetItem().ShaderResourceTexture, false, FResolveParams());
-
-		if (bFinalInterpolation && bUseDistanceFieldGI && (GAOInterpolationDepthTesting || GAOInterpolationStencilTesting))
+		if (bUseDistanceFieldGI)
 		{
-			RHICmdList.CopyToResolveTarget(IrradianceInterpolationTarget->GetRenderTargetItem().TargetableTexture, IrradianceInterpolationTarget->GetRenderTargetItem().ShaderResourceTexture, false, FResolveParams());
+			extern void ComputeIrradianceForSamples(
+				int32 DepthLevel,
+				FRHICommandListImmediate& RHICmdList,
+				const FViewInfo& View,
+				const FScene* Scene,
+				const FDistanceFieldAOParameters& Parameters,
+				FTemporaryIrradianceCacheResources* TemporaryIrradianceCacheResources);
+
+			ComputeIrradianceForSamples(DepthLevel, RHICmdList, View, Scene, Parameters, &GTemporaryIrradianceCacheResources);
 		}
-	}
-}
-
-void ListDistanceFieldLightingMemory(const FViewInfo& View)
-{
-	const FScene* Scene = (const FScene*)View.Family->Scene;
-	UE_LOG(LogTemp, Log, TEXT("Shared GPU memory (excluding render targets)"));
-
-	if (Scene->DistanceFieldSceneData.NumObjectsInBuffer > 0)
-	{
-		UE_LOG(LogTemp, Log, TEXT("   Scene Object data %.3fMb"), Scene->DistanceFieldSceneData.ObjectBuffers->GetSizeBytes() / 1024.0f / 1024.0f);
-	}
-	
-	UE_LOG(LogTemp, Log, TEXT("   %s"), *GDistanceFieldVolumeTextureAtlas.GetSizeString());
-	extern FString GetObjectBufferMemoryString();
-	UE_LOG(LogTemp, Log, TEXT("   %s"), *GetObjectBufferMemoryString());
-	UE_LOG(LogTemp, Log, TEXT(""));
-	UE_LOG(LogTemp, Log, TEXT("Distance Field AO"));
 
-	if (Scene->SurfaceCacheResources)
-	{
-		UE_LOG(LogTemp, Log, TEXT("   Surface cache %.3fMb"), Scene->SurfaceCacheResources->GetSizeBytes() / 1024.0f / 1024.0f);
-	}
-	
-	UE_LOG(LogTemp, Log, TEXT("   Temporary cache %.3fMb"), GTemporaryIrradianceCacheResources.GetSizeBytes() / 1024.0f / 1024.0f);
-	UE_LOG(LogTemp, Log, TEXT("   Culled objects %.3fMb"), GAOCulledObjectBuffers.Buffers.GetSizeBytes() / 1024.0f / 1024.0f);
+		// Compute heightfield occlusion after heightfield GI, otherwise it self-shadows incorrectly
+		View.HeightfieldLightingViewInfo.ComputeOcclusionForSamples(View, RHICmdList, GTemporaryIrradianceCacheResources, DepthLevel, Parameters);
 
-	FTileIntersectionResources* TileIntersectionResources = ((FSceneViewState*)View.State)->AOTileIntersectionResources;
+		{	
+			TShaderMapRef<TSetupFinalGatherIndirectArgumentsCS<false> > ComputeShader(View.ShaderMap);
 
-	if (TileIntersectionResources)
-	{
-		UE_LOG(LogTemp, Log, TEXT("   Tile Culled objects %.3fMb"), TileIntersectionResources->GetSizeBytes() / 1024.0f / 1024.0f);
-	}
-	
-	UE_LOG(LogTemp, Log, TEXT(""));
-	UE_LOG(LogTemp, Log, TEXT("Distance Field GI"));
+			RHICmdList.SetComputeShader(ComputeShader->GetComputeShader());
+			ComputeShader->SetParameters(RHICmdList, View, DepthLevel);
+			DispatchComputeShader(RHICmdList, *ComputeShader, 1, 1, 1);
 
-	if (Scene->DistanceFieldSceneData.SurfelBuffers)
-	{
-		UE_LOG(LogTemp, Log, TEXT("   Scene surfel data %.3fMb"), Scene->DistanceFieldSceneData.SurfelBuffers->GetSizeBytes() / 1024.0f / 1024.0f);
-	}
-	
-	if (Scene->DistanceFieldSceneData.InstancedSurfelBuffers)
-	{
-		UE_LOG(LogTemp, Log, TEXT("   Instanced scene surfel data %.3fMb"), Scene->DistanceFieldSceneData.InstancedSurfelBuffers->GetSizeBytes() / 1024.0f / 1024.0f);
-	}
-	
-	extern void ListDistanceFieldGIMemory(const FViewInfo& View);
-	ListDistanceFieldGIMemory(View);
-}
+			ComputeShader->UnsetParameters(RHICmdList);
+		}
 
-bool SupportsDistanceFieldAO(ERHIFeatureLevel::Type FeatureLevel, EShaderPlatform ShaderPlatform)
-{
-	return GDistanceFieldAO 
-		&& FeatureLevel >= ERHIFeatureLevel::SM5
-		&& DoesPlatformSupportDistanceFieldAO(ShaderPlatform);
-}
+		// Compute and store the final bent normal now that all occlusion sources have been computed (distance fields, heightfields)
+		{
+			TShaderMapRef<FCombineConesCS> ComputeShader(View.ShaderMap);
 
-bool ShouldRenderDynamicSkyLight(const FScene* Scene, const FSceneViewFamily& ViewFamily)
-{
-	return Scene->SkyLight
-		&& Scene->SkyLight->ProcessedTexture
-		&& !Scene->SkyLight->bWantsStaticShadowing
-		&& !Scene->SkyLight->bHasStaticLighting
-		&& ViewFamily.EngineShowFlags.SkyLighting
-		&& Scene->GetFeatureLevel() >= ERHIFeatureLevel::SM4
-		&& !IsSimpleDynamicLightingEnabled() 
-		&& !ViewFamily.EngineShowFlags.VisualizeLightCulling;
-}
+			RHICmdList.SetComputeShader(ComputeShader->GetComputeShader());
+			ComputeShader->SetParameters(RHICmdList, View, DepthLevel, Parameters, NULL, FMatrix::Identity, NULL, NULL);
+			DispatchIndirectComputeShader(RHICmdList, *ComputeShader, SurfaceCacheResources.DispatchParameters.Buffer, 0);
 
-bool FDeferredShadingSceneRenderer::ShouldPrepareForDistanceFieldAO() const
-{
-	return SupportsDistanceFieldAO(Scene->GetFeatureLevel(), Scene->GetShaderPlatform())
-		&& ((ShouldRenderDynamicSkyLight(Scene, ViewFamily) && Scene->SkyLight->bCastShadows && ViewFamily.EngineShowFlags.DistanceFieldAO)
-			|| ViewFamily.EngineShowFlags.VisualizeMeshDistanceFields
-			|| ViewFamily.EngineShowFlags.VisualizeDistanceFieldAO
-			|| ViewFamily.EngineShowFlags.VisualizeDistanceFieldGI);
-}
+			ComputeShader->UnsetParameters(RHICmdList);
+		}
+	}
 
-FRWBuffer* GDebugBuffer2 = NULL;
+	TRefCountPtr<IPooledRenderTarget> BentNormalAccumulation;
+	TRefCountPtr<IPooledRenderTarget> IrradianceAccumulation;
 
-IMPLEMENT_SHADER_TYPE(,FTrackGPUProgressCS,TEXT("DistanceFieldGlobalIllumination"),TEXT("TrackGPUProgressCS"),SF_Compute);
+	{
+		FIntPoint BufferSize = GetBufferSizeForAO();
+		FPooledRenderTargetDesc Desc(FPooledRenderTargetDesc::Create2DDesc(BufferSize, PF_FloatRGBA, TexCreate_None, TexCreate_RenderTargetable | TexCreate_UAV, false));
+		GRenderTargetPool.FindFreeElement(Desc, BentNormalAccumulation, TEXT("BentNormalAccumulation"));
 
-void TrackGPUProgress(FRHICommandListImmediate& RHICmdList, uint32 DebugId)
-{
-	/*
-	TShaderMapRef<FTrackGPUProgressCS> ComputeShader(GetGlobalShaderMap(ERHIFeatureLevel::SM5));
-	RHICmdList.SetComputeShader(ComputeShader->GetComputeShader());
-	ComputeShader->SetParameters(RHICmdList, DebugId);
-	DispatchComputeShader(RHICmdList, *ComputeShader, 1, 1, 1);
-	ComputeShader->UnsetParameters(RHICmdList);
+		if (bUseDistanceFieldGI)
+		{
+			GRenderTargetPool.FindFreeElement(Desc, IrradianceAccumulation, TEXT("IrradianceAccumulation"));
+		}
+	}
 
+	// Splat the surface cache records onto the opaque pixels, using less strict weighting so the lighting is smoothed in world space
 	{
-		FRenderQueryRHIRef Query = RHICmdList.CreateRenderQuery(RQT_AbsoluteTime);
-		uint64 Temp = 0;
-		RHICmdList.EndRenderQuery(Query);
-		RHICmdList.GetRenderQueryResult(Query, Temp, true);
+		SCOPED_DRAW_EVENT(RHICmdList, FinalIrradianceCacheSplat);
+		RenderIrradianceCacheInterpolation(RHICmdList, View, BentNormalAccumulation, IrradianceAccumulation, DistanceFieldNormal->GetRenderTargetItem(), 0, GAODownsampleFactor, Parameters, true);
+	}
+
+	GRenderTargetPool.VisualizeTexture.SetCheckPoint(RHICmdList, BentNormalAccumulation);
 
-		uint32* Ptr = (uint32*)RHICmdList.LockVertexBuffer(GDebugBuffer2->Buffer, 0, 4, RLM_ReadOnly);
+	// Post process the AO to cover over artifacts
 
-		UE_LOG(LogDistanceField, Warning, TEXT("DebugId %u,"), *Ptr);
-		RHICmdList.UnlockVertexBuffer(GDebugBuffer2->Buffer);
-	}*/
+	PostProcessBentNormalAOSurfaceCache(
+		RHICmdList, 
+		Parameters, 
+		View, 
+		VelocityTexture->GetRenderTargetItem(), 
+		BentNormalAccumulation->GetRenderTargetItem(), 
+		IrradianceAccumulation, 
+		DistanceFieldNormal->GetRenderTargetItem(), 
+		OutDynamicBentNormalAO, 
+		OutDynamicIrradiance);
 }
 
-bool FDeferredShadingSceneRenderer::RenderDistanceFieldAOSurfaceCache(
+bool FDeferredShadingSceneRenderer::RenderDistanceFieldLighting(
 	FRHICommandListImmediate& RHICmdList, 
 	const FDistanceFieldAOParameters& Parameters, 
 	const TRefCountPtr<IPooledRenderTarget>& VelocityTexture,
@@ -3812,66 +2908,23 @@ bool FDeferredShadingSceneRenderer::RenderDistanceFieldAOSurfaceCache(
 		&& View.State
 		&& View.IsPerspectiveProjection())
 	{
-		QUICK_SCOPE_CYCLE_COUNTER(STAT_RenderDistanceFieldAOSurfaceCache);
-		SCOPED_DRAW_EVENT(RHICmdList, DistanceFieldLighting);
+		QUICK_SCOPE_CYCLE_COUNTER(STAT_RenderDistanceFieldLighting);
 
 		if (GDistanceFieldVolumeTextureAtlas.VolumeTextureRHI && Scene->DistanceFieldSceneData.NumObjectsInBuffer)
 		{
 			check(!Scene->DistanceFieldSceneData.HasPendingOperations());
 			const bool bUseDistanceFieldGI = IsDistanceFieldGIAllowed(View);
 
-			if (!GDebugBuffer2)
-			{
-				GDebugBuffer2 = new FRWBuffer();
-				GDebugBuffer2->Initialize(4, 1, PF_R32_UINT);
-			}
+			SCOPED_DRAW_EVENT(RHICmdList, DistanceFieldLighting);
 
 			GenerateBestSpacedVectors();
 
-			// Create surface cache state that will persist across frames if needed
-			if (!Scene->SurfaceCacheResources)
-			{
-				Scene->SurfaceCacheResources = new FSurfaceCacheResources();
-			}
-
 			if (bListMemoryNextFrame)
 			{
 				bListMemoryNextFrame = false;
 				ListDistanceFieldLightingMemory(View);
 			}
 
-			FSurfaceCacheResources& SurfaceCacheResources = *Scene->SurfaceCacheResources;
-
-			FIntPoint DownsampledViewRect = FIntPoint::DivideAndRoundDown(View.ViewRect.Size(), GAODownsampleFactor);
-			// Needed because samples are kept across frames
-			int32 SlackAmount = 4;
-			int32 MaxSurfaceCacheSamples = DownsampledViewRect.X * DownsampledViewRect.Y / (1 << (2 * GAOMinLevel * GAOPowerOfTwoBetweenLevels));
-			SurfaceCacheResources.AllocateFor(GAOMinLevel, GAOMaxLevel, SlackAmount * MaxSurfaceCacheSamples);
-
-			GTemporaryIrradianceCacheResources.AllocateFor(MaxSurfaceCacheSamples);
-
-			uint32 ClearValues[4] = {0};
-
-			if (!SurfaceCacheResources.bClearedResources 
-				|| !GAOReuseAcrossFrames 
-				// Drop records that will have uninitialized Irradiance if switching from AO only to AO + GI
-				|| (bUseDistanceFieldGI && !SurfaceCacheResources.bHasIrradiance))
-			{
-				// Reset the number of active cache records to 0
-				for (int32 DepthLevel = GAOMaxLevel; DepthLevel >= GAOMinLevel; DepthLevel--)
-				{
-					RHICmdList.ClearUAV(SurfaceCacheResources.Level[DepthLevel]->ScatterDrawParameters.UAV, ClearValues);
-				}
-
-				RHICmdList.ClearUAV(SurfaceCacheResources.TempResources->ScatterDrawParameters.UAV, ClearValues);
-
-				SurfaceCacheResources.bClearedResources = true;
-			}
-
-			SurfaceCacheResources.bHasIrradiance = bUseDistanceFieldGI;
-
-			QUICK_SCOPE_CYCLE_COUNTER(STAT_AOIssueGPUWork);
-
 			{
 				SCOPED_DRAW_EVENT(RHICmdList, ObjectFrustumCulling);
 
@@ -3884,6 +2937,7 @@ bool FDeferredShadingSceneRenderer::RenderDistanceFieldAOSurfaceCache(
 				}
 
 				{
+					uint32 ClearValues[4] = { 0 };
 					RHICmdList.ClearUAV(GAOCulledObjectBuffers.Buffers.ObjectIndirectArguments.UAV, ClearValues);
 
 					TShaderMapRef<FCullObjectsForViewCS> ComputeShader(GetGlobalShaderMap(Scene->GetFeatureLevel()));
@@ -3903,7 +2957,6 @@ bool FDeferredShadingSceneRenderer::RenderDistanceFieldAOSurfaceCache(
 				GRenderTargetPool.FindFreeElement(Desc, DistanceFieldNormal, TEXT("DistanceFieldNormal"));
 			}
 
-			// Compute the distance field normal, this is used for surface caching instead of the GBuffer normal because it matches the occluding geometry
 			ComputeDistanceFieldNormal(RHICmdList, Views, DistanceFieldNormal->GetRenderTargetItem(), Parameters);
 
 			// Intersect objects with screen tiles, build lists
@@ -3911,51 +2964,6 @@ bool FDeferredShadingSceneRenderer::RenderDistanceFieldAOSurfaceCache(
 
 			GRenderTargetPool.VisualizeTexture.SetCheckPoint(RHICmdList, DistanceFieldNormal);
 
-			if (GAOReuseAcrossFrames)
-			{
-				SCOPED_DRAW_EVENT(RHICmdList, TrimRecords);
-
-				// Copy and trim last frame's surface cache samples
-				for (int32 DepthLevel = GAOMaxLevel; DepthLevel >= GAOMinLevel; DepthLevel--)
-				{
-					if (GAOTrimOldRecordsFraction > 0)
-					{
-						{	
-							TShaderMapRef<FSetupCopyIndirectArgumentsCS> ComputeShader(View.ShaderMap);
-
-							RHICmdList.SetComputeShader(ComputeShader->GetComputeShader());
-							ComputeShader->SetParameters(RHICmdList, View, DepthLevel);
-							DispatchComputeShader(RHICmdList, *ComputeShader, 1, 1, 1);
-
-							ComputeShader->UnsetParameters(RHICmdList);
-						}
-
-						if (bUseDistanceFieldGI)
-						{
-							TShaderMapRef<TCopyIrradianceCacheSamplesCS<true> > ComputeShader(View.ShaderMap);
-
-							RHICmdList.SetComputeShader(ComputeShader->GetComputeShader());
-							ComputeShader->SetParameters(RHICmdList, View, DepthLevel);
-							DispatchIndirectComputeShader(RHICmdList, *ComputeShader, SurfaceCacheResources.DispatchParameters.Buffer, 0);
-
-							ComputeShader->UnsetParameters(RHICmdList);
-						}
-						else
-						{
-							TShaderMapRef<TCopyIrradianceCacheSamplesCS<false> > ComputeShader(View.ShaderMap);
-
-							RHICmdList.SetComputeShader(ComputeShader->GetComputeShader());
-							ComputeShader->SetParameters(RHICmdList, View, DepthLevel);
-							DispatchIndirectComputeShader(RHICmdList, *ComputeShader, SurfaceCacheResources.DispatchParameters.Buffer, 0);
-
-							ComputeShader->UnsetParameters(RHICmdList);
-						}
-
-						Swap(SurfaceCacheResources.Level[DepthLevel], SurfaceCacheResources.TempResources);
-					}
-				}
-			}
-
 			if (bUseDistanceFieldGI)
 			{
 				extern void UpdateVPLs(
@@ -3967,172 +2975,40 @@ bool FDeferredShadingSceneRenderer::RenderDistanceFieldAOSurfaceCache(
 				UpdateVPLs(RHICmdList, View, Scene, Parameters);
 			}
 
-			// Start from the highest depth, which will consider the fewest potential shading points
-			// Each level potentially prevents the next higher resolution level from doing expensive shading work
-			// This is the core of the surface cache algorithm
-			for (int32 DepthLevel = GAOMaxLevel; DepthLevel >= GAOMinLevel; DepthLevel--)
-			{
-				int32 DestLevelDownsampleFactor = GAODownsampleFactor * (1 << (DepthLevel * GAOPowerOfTwoBetweenLevels));
-
-				SCOPED_DRAW_EVENTF(RHICmdList, Level, TEXT("Level_%d"), DepthLevel);
-
-				TRefCountPtr<IPooledRenderTarget> DistanceFieldAOBentNormalSplat;
-
-				{
-					FIntPoint AOBufferSize = FIntPoint::DivideAndRoundUp(GSceneRenderTargets.GetBufferSizeXY(), DestLevelDownsampleFactor);
-					FPooledRenderTargetDesc Desc(FPooledRenderTargetDesc::Create2DDesc(AOBufferSize, PF_FloatRGBA, TexCreate_None, TexCreate_RenderTargetable | TexCreate_UAV, false));
-					GRenderTargetPool.FindFreeElement(Desc, DistanceFieldAOBentNormalSplat, TEXT("DistanceFieldAOBentNormalSplat"));
-				}
-
-				// Splat / interpolate the surface cache records onto the buffer sized for the current depth level
-				RenderIrradianceCacheInterpolation(RHICmdList, View, DistanceFieldAOBentNormalSplat, NULL, DistanceFieldNormal->GetRenderTargetItem(), DepthLevel, DestLevelDownsampleFactor, Parameters, false);
-
-				{
-					SCOPED_DRAW_EVENT(RHICmdList, PopulateAndConeTrace);
-					SetRenderTarget(RHICmdList, NULL, NULL);
-
-					// Save off the current record count before adding any more
-					{	
-						TShaderMapRef<FSaveStartIndexCS> ComputeShader(View.ShaderMap);
-
-						RHICmdList.SetComputeShader(ComputeShader->GetComputeShader());
-						ComputeShader->SetParameters(RHICmdList, View, DepthLevel);
-						DispatchComputeShader(RHICmdList, *ComputeShader, 1, 1, 1);
-
-						ComputeShader->UnsetParameters(RHICmdList);
-					}
-
-					// Create new records which haven't been shaded yet for shading points which don't have a valid interpolation from existing records
-					for (int32 ViewIndex = 0; ViewIndex < Views.Num(); ViewIndex++)
-					{
-						const FViewInfo& ViewInfo = Views[ViewIndex];
-
-						FIntPoint DownsampledViewSize = FIntPoint::DivideAndRoundUp(ViewInfo.ViewRect.Size(), DestLevelDownsampleFactor);
-						uint32 GroupSizeX = FMath::DivideAndRoundUp(DownsampledViewSize.X, GDistanceFieldAOTileSizeX);
-						uint32 GroupSizeY = FMath::DivideAndRoundUp(DownsampledViewSize.Y, GDistanceFieldAOTileSizeY);
-
-						TShaderMapRef<FPopulateCacheCS> ComputeShader(ViewInfo.ShaderMap);
-
-						RHICmdList.SetComputeShader(ComputeShader->GetComputeShader());
-						ComputeShader->SetParameters(RHICmdList, ViewInfo, DistanceFieldAOBentNormalSplat->GetRenderTargetItem(), DistanceFieldNormal->GetRenderTargetItem(), DestLevelDownsampleFactor, DepthLevel, TileListGroupSize, Parameters);
-						DispatchComputeShader(RHICmdList, *ComputeShader, GroupSizeX, GroupSizeY, 1);
-
-						ComputeShader->UnsetParameters(RHICmdList);
-					}
-
-					{	
-						TShaderMapRef<TSetupFinalGatherIndirectArgumentsCS<true> > ComputeShader(View.ShaderMap);
-
-						RHICmdList.SetComputeShader(ComputeShader->GetComputeShader());
-						ComputeShader->SetParameters(RHICmdList, View, DepthLevel);
-						DispatchComputeShader(RHICmdList, *ComputeShader, 1, 1, 1);
-
-						ComputeShader->UnsetParameters(RHICmdList);
-					}
-
-					// Compute lighting for the new surface cache records by cone-stepping through the object distance fields
-					if (bUseDistanceFieldGI)
-					{
-						TShaderMapRef<TConeTraceOcclusionCS<true> > ComputeShader(View.ShaderMap);
-
-						RHICmdList.SetComputeShader(ComputeShader->GetComputeShader());
-						ComputeShader->SetParameters(RHICmdList, View, DistanceFieldAOBentNormalSplat->GetRenderTargetItem(), DistanceFieldNormal->GetRenderTargetItem(), DestLevelDownsampleFactor, DepthLevel, TileListGroupSize, Parameters);
-						DispatchIndirectComputeShader(RHICmdList, *ComputeShader, SurfaceCacheResources.DispatchParameters.Buffer, 0);
-
-						ComputeShader->UnsetParameters(RHICmdList);
-					}
-					else
-					{
-						TShaderMapRef<TConeTraceOcclusionCS<false> > ComputeShader(View.ShaderMap);
-
-						RHICmdList.SetComputeShader(ComputeShader->GetComputeShader());
-						ComputeShader->SetParameters(RHICmdList, View, DistanceFieldAOBentNormalSplat->GetRenderTargetItem(), DistanceFieldNormal->GetRenderTargetItem(), DestLevelDownsampleFactor, DepthLevel, TileListGroupSize, Parameters);
-						DispatchIndirectComputeShader(RHICmdList, *ComputeShader, SurfaceCacheResources.DispatchParameters.Buffer, 0);
-
-						ComputeShader->UnsetParameters(RHICmdList);
-					}
-				}
-
-				if (bUseDistanceFieldGI)
-				{
-					extern void ComputeIrradianceForSamples(
-						int32 DepthLevel,
-						FRHICommandListImmediate& RHICmdList,
-						const FViewInfo& View,
-						const FScene* Scene,
-						const FDistanceFieldAOParameters& Parameters,
-						FTemporaryIrradianceCacheResources* TemporaryIrradianceCacheResources);
-
-					ComputeIrradianceForSamples(DepthLevel, RHICmdList, View, Scene, Parameters, &GTemporaryIrradianceCacheResources);
-				}
-
-				// Compute heightfield occlusion after heightfield GI, otherwise it self-shadows incorrectly
-				View.HeightfieldLightingViewInfo.ComputeOcclusionForSamples(View, RHICmdList, GTemporaryIrradianceCacheResources, DepthLevel, Parameters);
-
-				{	
-					TShaderMapRef<TSetupFinalGatherIndirectArgumentsCS<false> > ComputeShader(View.ShaderMap);
-
-					RHICmdList.SetComputeShader(ComputeShader->GetComputeShader());
-					ComputeShader->SetParameters(RHICmdList, View, DepthLevel);
-					DispatchComputeShader(RHICmdList, *ComputeShader, 1, 1, 1);
-
-					ComputeShader->UnsetParameters(RHICmdList);
-				}
-
-				// Compute and store the final bent normal now that all occlusion sources have been computed (distance fields, heightfields)
-				{
-					TShaderMapRef<FCombineConesCS> ComputeShader(View.ShaderMap);
-
-					RHICmdList.SetComputeShader(ComputeShader->GetComputeShader());
-					ComputeShader->SetParameters(RHICmdList, View, DepthLevel, Parameters, NULL, FMatrix::Identity, NULL, NULL);
-					DispatchIndirectComputeShader(RHICmdList, *ComputeShader, SurfaceCacheResources.DispatchParameters.Buffer, 0);
-
-					ComputeShader->UnsetParameters(RHICmdList);
-				}
-			}
-
-			TRefCountPtr<IPooledRenderTarget> BentNormalAccumulation;
-			TRefCountPtr<IPooledRenderTarget> IrradianceAccumulation;
+			TRefCountPtr<IPooledRenderTarget> BentNormalOutput;
+			TRefCountPtr<IPooledRenderTarget> IrradianceOutput;
 
+			if (GAOUseSurfaceCache)
 			{
-				FIntPoint BufferSize = GetBufferSizeForAO();
-				FPooledRenderTargetDesc Desc(FPooledRenderTargetDesc::Create2DDesc(BufferSize, PF_FloatRGBA, TexCreate_None, TexCreate_RenderTargetable | TexCreate_UAV, false));
-				GRenderTargetPool.FindFreeElement(Desc, BentNormalAccumulation, TEXT("BentNormalAccumulation"));
-
-				if (bUseDistanceFieldGI)
-				{
-					GRenderTargetPool.FindFreeElement(Desc, IrradianceAccumulation, TEXT("IrradianceAccumulation"));
-				}
+				RenderDistanceFieldAOSurfaceCache(
+					RHICmdList, 
+					View,
+					Scene,
+					TileListGroupSize,
+					Parameters, 
+					VelocityTexture,
+					DistanceFieldNormal, 
+					BentNormalOutput, 
+					IrradianceOutput);
 			}
-
-			// Splat the surface cache records onto the opaque pixels, using less strict weighting so the lighting is smoothed in world space
+			else
 			{
-				SCOPED_DRAW_EVENT(RHICmdList, FinalIrradianceCacheSplat);
-				RenderIrradianceCacheInterpolation(RHICmdList, View, BentNormalAccumulation, IrradianceAccumulation, DistanceFieldNormal->GetRenderTargetItem(), 0, GAODownsampleFactor, Parameters, true);
+				RenderDistanceFieldAOScreenGrid(
+					RHICmdList, 
+					View,
+					TileListGroupSize,
+					Parameters, 
+					VelocityTexture,
+					DistanceFieldNormal, 
+					BentNormalOutput, 
+					IrradianceOutput);
 			}
 
-			GRenderTargetPool.VisualizeTexture.SetCheckPoint(RHICmdList, BentNormalAccumulation);
-
-			// Post process the AO to cover over artifacts
-			TRefCountPtr<IPooledRenderTarget> BentNormalOutput;
-			TRefCountPtr<IPooledRenderTarget> IrradianceOutput;
-
-			PostProcessBentNormalAO(
-				RHICmdList, 
-				Parameters, 
-				View, 
-				VelocityTexture->GetRenderTargetItem(), 
-				BentNormalAccumulation->GetRenderTargetItem(), 
-				IrradianceAccumulation, 
-				DistanceFieldNormal->GetRenderTargetItem(), 
-				BentNormalOutput, 
-				IrradianceOutput);
-
 			GRenderTargetPool.VisualizeTexture.SetCheckPoint(RHICmdList, BentNormalOutput);
 
 			if (bVisualizeAmbientOcclusion || bVisualizeGlobalIllumination)
 			{
-				GSceneRenderTargets.BeginRenderingSceneColor(RHICmdList);
+				GSceneRenderTargets.BeginRenderingSceneColor(RHICmdList, ESimpleRenderTargetMode::EExistingColorAndDepth, FExclusiveDepthStencil::DepthRead_StencilNop);
 			}
 			else
 			{
@@ -4299,7 +3175,7 @@ void FDeferredShadingSceneRenderer::RenderDynamicSkyLighting(FRHICommandListImme
 
 		if (ShouldRenderDistanceFieldAO())
 		{
-			bApplyShadowing = RenderDistanceFieldAOSurfaceCache(RHICmdList, Parameters, VelocityTexture, DynamicBentNormalAO, DynamicIrradiance, false, false);
+			bApplyShadowing = RenderDistanceFieldLighting(RHICmdList, Parameters, VelocityTexture, DynamicBentNormalAO, DynamicIrradiance, false, false);
 		}
 
 		GSceneRenderTargets.BeginRenderingSceneColor(RHICmdList);
@@ -4376,9 +3252,10 @@ void FDeferredShadingSceneRenderer::RenderDynamicSkyLighting(FRHICommandListImme
 	}
 }
 
-class FVisualizeMeshDistanceFieldCS : public FGlobalShader
+template<bool bUseGlobalDistanceField>
+class TVisualizeMeshDistanceFieldCS : public FGlobalShader
 {
-	DECLARE_SHADER_TYPE(FVisualizeMeshDistanceFieldCS, Global);
+	DECLARE_SHADER_TYPE(TVisualizeMeshDistanceFieldCS, Global);
 public:
 
 	static bool ShouldCache(EShaderPlatform Platform)
@@ -4391,13 +3268,14 @@ public:
 		OutEnvironment.SetDefine(TEXT("DOWNSAMPLE_FACTOR"), GAODownsampleFactor);
 		OutEnvironment.SetDefine(TEXT("THREADGROUP_SIZEX"), GDistanceFieldAOTileSizeX);
 		OutEnvironment.SetDefine(TEXT("THREADGROUP_SIZEY"), GDistanceFieldAOTileSizeY);
+		OutEnvironment.SetDefine(TEXT("USE_GLOBAL_DISTANCE_FIELD"), bUseGlobalDistanceField ? TEXT("1") : TEXT("0"));
 	}
 
 	/** Default constructor. */
-	FVisualizeMeshDistanceFieldCS() {}
+	TVisualizeMeshDistanceFieldCS() {}
 
 	/** Initialization constructor. */
-	FVisualizeMeshDistanceFieldCS(const ShaderMetaType::CompiledShaderInitializerType& Initializer)
+	TVisualizeMeshDistanceFieldCS(const ShaderMetaType::CompiledShaderInitializerType& Initializer)
 		: FGlobalShader(Initializer)
 	{
 		VisualizeMeshDistanceFields.Bind(Initializer.ParameterMap, TEXT("VisualizeMeshDistanceFields"));
@@ -4405,6 +3283,7 @@ public:
 		ObjectParameters.Bind(Initializer.ParameterMap);
 		DeferredParameters.Bind(Initializer.ParameterMap);
 		AOParameters.Bind(Initializer.ParameterMap);
+		GlobalDistanceFieldParameters.Bind(Initializer.ParameterMap);
 	}
 
 	void SetParameters(
@@ -4412,7 +3291,8 @@ public:
 		const FSceneView& View, 
 		FSceneRenderTargetItem& VisualizeMeshDistanceFieldsValue, 
 		FVector2D NumGroupsValue,
-		const FDistanceFieldAOParameters& Parameters)
+		const FDistanceFieldAOParameters& Parameters,
+		const FGlobalDistanceFieldInfo& GlobalDistanceFieldInfo)
 	{
 		const FComputeShaderRHIParamRef ShaderRHI = GetComputeShader();
 
@@ -4424,6 +3304,11 @@ public:
 		AOParameters.Set(RHICmdList, ShaderRHI, Parameters);
 		DeferredParameters.Set(RHICmdList, ShaderRHI, View);
 
+		if (bUseGlobalDistanceField)
+		{
+			GlobalDistanceFieldParameters.Set(RHICmdList, ShaderRHI, GlobalDistanceFieldInfo);
+		}
+
 		SetShaderValue(RHICmdList, ShaderRHI, NumGroups, NumGroupsValue);
 	}
 
@@ -4441,6 +3326,7 @@ public:
 		Ar << ObjectParameters;
 		Ar << DeferredParameters;
 		Ar << AOParameters;
+		Ar << GlobalDistanceFieldParameters;
 		return bShaderHasOutdatedParameters;
 	}
 
@@ -4451,10 +3337,11 @@ private:
 	FDistanceFieldCulledObjectBufferParameters ObjectParameters;
 	FDeferredPixelShaderParameters DeferredParameters;
 	FAOParameters AOParameters;
+	FGlobalDistanceFieldParameters GlobalDistanceFieldParameters;
 };
 
-IMPLEMENT_SHADER_TYPE(,FVisualizeMeshDistanceFieldCS,TEXT("DistanceFieldSurfaceCacheLighting"),TEXT("VisualizeMeshDistanceFieldCS"),SF_Compute);
-
+IMPLEMENT_SHADER_TYPE(template<>,TVisualizeMeshDistanceFieldCS<true>,TEXT("DistanceFieldSurfaceCacheLighting"),TEXT("VisualizeMeshDistanceFieldCS"),SF_Compute);
+IMPLEMENT_SHADER_TYPE(template<>,TVisualizeMeshDistanceFieldCS<false>,TEXT("DistanceFieldSurfaceCacheLighting"),TEXT("VisualizeMeshDistanceFieldCS"),SF_Compute);
 
 class FVisualizeDistanceFieldUpsamplePS : public FGlobalShader
 {
@@ -4531,6 +3418,8 @@ void FDeferredShadingSceneRenderer::RenderMeshDistanceFieldVisualization(FRHICom
 
 			QUICK_SCOPE_CYCLE_COUNTER(STAT_AOIssueGPUWork);
 
+			const bool bUseGlobalDistanceField = UseGlobalDistanceField(Parameters) && GAOVisualizeGlobalDistanceField != 0;
+
 			{
 				if (GAOCulledObjectBuffers.Buffers.MaxObjects < Scene->DistanceFieldSceneData.NumObjectsInBuffer)
 				{
@@ -4570,12 +3459,26 @@ void FDeferredShadingSceneRenderer::RenderMeshDistanceFieldVisualization(FRHICom
 					uint32 GroupSizeX = FMath::DivideAndRoundUp(ViewInfo.ViewRect.Size().X / GAODownsampleFactor, GDistanceFieldAOTileSizeX);
 					uint32 GroupSizeY = FMath::DivideAndRoundUp(ViewInfo.ViewRect.Size().Y / GAODownsampleFactor, GDistanceFieldAOTileSizeY);
 
+					SCOPED_DRAW_EVENT(RHICmdList, VisualizeMeshDistanceFieldCS);
+
+					if (bUseGlobalDistanceField)
+					{
+						check(View.GlobalDistanceFieldInfo.Clipmaps.Num() > 0);
+
+						TShaderMapRef<TVisualizeMeshDistanceFieldCS<true> > ComputeShader(ViewInfo.ShaderMap);
+
+						RHICmdList.SetComputeShader(ComputeShader->GetComputeShader());
+						ComputeShader->SetParameters(RHICmdList, ViewInfo, VisualizeResultRT->GetRenderTargetItem(), FVector2D(GroupSizeX, GroupSizeY), Parameters, View.GlobalDistanceFieldInfo);
+						DispatchComputeShader(RHICmdList, *ComputeShader, GroupSizeX, GroupSizeY, 1);
+
+						ComputeShader->UnsetParameters(RHICmdList);
+					}
+					else
 					{
-						SCOPED_DRAW_EVENT(RHICmdList, VisualizeMeshDistanceFieldCS);
-						TShaderMapRef<FVisualizeMeshDistanceFieldCS> ComputeShader(ViewInfo.ShaderMap);
+						TShaderMapRef<TVisualizeMeshDistanceFieldCS<false> > ComputeShader(ViewInfo.ShaderMap);
 
 						RHICmdList.SetComputeShader(ComputeShader->GetComputeShader());
-						ComputeShader->SetParameters(RHICmdList, ViewInfo, VisualizeResultRT->GetRenderTargetItem(), FVector2D(GroupSizeX, GroupSizeY), Parameters);
+						ComputeShader->SetParameters(RHICmdList, ViewInfo, VisualizeResultRT->GetRenderTargetItem(), FVector2D(GroupSizeX, GroupSizeY), Parameters, View.GlobalDistanceFieldInfo);
 						DispatchComputeShader(RHICmdList, *ComputeShader, GroupSizeX, GroupSizeY, 1);
 
 						ComputeShader->UnsetParameters(RHICmdList);
@@ -4619,4 +3522,4 @@ void FDeferredShadingSceneRenderer::RenderMeshDistanceFieldVisualization(FRHICom
 			}
 		}
 	}
-}
\ No newline at end of file
+}
