Distance Field AO is now evaluated on a screen grid and using a global distance field for far traces

diff --git a/Engine/Source/Runtime/Renderer/Private/DistanceFieldSurfaceCacheLighting.h b/Engine/Source/Runtime/Renderer/Private/DistanceFieldSurfaceCacheLighting.h
--- a/Engine/Source/Runtime/Renderer/Private/DistanceFieldSurfaceCacheLighting.h
+++ b/Engine/Source/Runtime/Renderer/Private/DistanceFieldSurfaceCacheLighting.h
@@ -11,18 +11,23 @@
 const static int32 GAOMaxSupportedLevel = 6;
 /** Number of cone traced directions. */
 const int32 NumConeSampleDirections = 9;
+extern const uint32 UpdateObjectsGroupSize;
+
+const float GDefaultDFAOMaxOcclusionDistance = 600.0f;
+
+inline bool DoesPlatformSupportDistanceFieldAO(EShaderPlatform Platform)
+{
+	return Platform == SP_PCD3D_SM5 || Platform == SP_PS4;
+}
 
 class FDistanceFieldAOParameters
 {
 public:
-	float OcclusionMaxDistance;
+	float GlobalMaxOcclusionDistance;
+	float ObjectMaxOcclusionDistance;
 	float Contrast;
 
-	FDistanceFieldAOParameters(float InOcclusionMaxDistance, float InContrast = 0)
-	{
-		Contrast = FMath::Clamp(InContrast, .01f, 2.0f);
-		OcclusionMaxDistance = FMath::Clamp(InOcclusionMaxDistance, 200.0f, 3000.0f);
-	}
+	FDistanceFieldAOParameters(float InOcclusionMaxDistance, float InContrast = 0);
 };
 
 /**  */
@@ -43,7 +48,10 @@ public:
 			PositionAndRadius.Initialize(sizeof(float) * 4, MaxIrradianceCacheSamples, PF_A32B32G32R32F, BUF_Static);
 			OccluderRadius.Initialize(sizeof(float), MaxIrradianceCacheSamples, PF_R32_FLOAT, BUF_Static);
 			Normal.Initialize(sizeof(FFloat16Color), MaxIrradianceCacheSamples, PF_FloatRGBA, BUF_Static);
-			BentNormal.Initialize(sizeof(FFloat16Color), MaxIrradianceCacheSamples, PF_FloatRGBA, BUF_Static);
+			
+			// Must match usf
+			const int32 NumGradients = 0;
+			BentNormal.Initialize(sizeof(FFloat16Color), MaxIrradianceCacheSamples * (1 + NumGradients), PF_FloatRGBA, BUF_Static);
 			Irradiance.Initialize(sizeof(FFloat16Color), MaxIrradianceCacheSamples, PF_FloatRGBA, BUF_Static);
 			ScatterDrawParameters.Initialize(sizeof(uint32), 4, PF_R32_UINT, BUF_Static | BUF_DrawIndirect);
 			SavedStartIndex.Initialize(sizeof(uint32), 1, PF_R32_UINT, BUF_Static);
@@ -210,10 +218,44 @@ private:
 	int32 MaxLevel;
 };
 
-inline bool DoesPlatformSupportDistanceFieldAO(EShaderPlatform Platform)
+class FIrradianceCacheParameters
 {
-	return Platform == SP_PCD3D_SM5 || Platform == SP_PS4;
-}
+public:
+	void Bind(const FShaderParameterMap& ParameterMap)
+	{
+		IrradianceCachePositionRadius.Bind(ParameterMap, TEXT("IrradianceCachePositionRadius"));
+		IrradianceCacheOccluderRadius.Bind(ParameterMap, TEXT("IrradianceCacheOccluderRadius"));
+		IrradianceCacheNormal.Bind(ParameterMap, TEXT("IrradianceCacheNormal"));
+		IrradianceCacheBentNormal.Bind(ParameterMap, TEXT("IrradianceCacheBentNormal"));
+		IrradianceCacheIrradiance.Bind(ParameterMap, TEXT("IrradianceCacheIrradiance"));
+		IrradianceCacheTileCoordinate.Bind(ParameterMap, TEXT("IrradianceCacheTileCoordinate"));
+	}
+
+	template<typename TParamRef>
+	void Set(FRHICommandList& RHICmdList, const TParamRef& ShaderRHI, const FRefinementLevelResources& LevelResources)
+	{
+		SetSRVParameter(RHICmdList, ShaderRHI, IrradianceCachePositionRadius, LevelResources.PositionAndRadius.SRV);
+		SetSRVParameter(RHICmdList, ShaderRHI, IrradianceCacheOccluderRadius, LevelResources.OccluderRadius.SRV);
+		SetSRVParameter(RHICmdList, ShaderRHI, IrradianceCacheNormal, LevelResources.Normal.SRV);
+		SetSRVParameter(RHICmdList, ShaderRHI, IrradianceCacheBentNormal, LevelResources.BentNormal.SRV);
+		SetSRVParameter(RHICmdList, ShaderRHI, IrradianceCacheIrradiance, LevelResources.Irradiance.SRV);
+		SetSRVParameter(RHICmdList, ShaderRHI, IrradianceCacheTileCoordinate, LevelResources.TileCoordinate.SRV);
+	}
+
+	friend FArchive& operator<<(FArchive& Ar, FIrradianceCacheParameters& P)
+	{
+		Ar << P.IrradianceCachePositionRadius << P.IrradianceCacheOccluderRadius << P.IrradianceCacheNormal << P.IrradianceCacheBentNormal << P.IrradianceCacheIrradiance << P.IrradianceCacheTileCoordinate;
+		return Ar;
+	}
+
+private:
+	FShaderResourceParameter IrradianceCachePositionRadius;
+	FShaderResourceParameter IrradianceCacheOccluderRadius;
+	FShaderResourceParameter IrradianceCacheNormal;
+	FShaderResourceParameter IrradianceCacheBentNormal;
+	FShaderResourceParameter IrradianceCacheIrradiance;
+	FShaderResourceParameter IrradianceCacheTileCoordinate;
+};
 
 template<bool bOneGroupPerRecord>
 class TSetupFinalGatherIndirectArgumentsCS : public FGlobalShader
@@ -314,6 +356,46 @@ public:
 	}
 };
 
+class FAOScreenGridResources : public FRenderResource
+{
+public:
+
+	FAOScreenGridResources() :
+		bAllocateResourceForGI(false)
+	{}
+
+	virtual void InitDynamicRHI() override;
+
+	virtual void ReleaseDynamicRHI() override
+	{
+		ScreenGridConeVisibility.Release();
+		ConeDepthVisibilityFunction.Release();
+		StepBentNormal.Release();
+		SurfelIrradiance.Release();
+		HeightfieldIrradiance.Release();
+	}
+
+	FIntPoint ScreenGridDimensions;
+
+	FRWBuffer ScreenGridConeVisibility;
+
+	bool bAllocateResourceForGI;
+	FRWBuffer ConeDepthVisibilityFunction;
+	FRWBuffer StepBentNormal;
+	FRWBuffer SurfelIrradiance;
+	FRWBuffer HeightfieldIrradiance;
+
+	size_t GetSizeBytesForAO() const
+	{
+		return ScreenGridConeVisibility.NumBytes;
+	}
+
+	size_t GetSizeBytesForGI() const
+	{
+		return ConeDepthVisibilityFunction.NumBytes + StepBentNormal.NumBytes + SurfelIrradiance.NumBytes + HeightfieldIrradiance.NumBytes;
+	}
+};
+
 extern void GetSpacedVectors(TArray<FVector, TInlineAllocator<9> >& OutVectors);
 
 BEGIN_UNIFORM_BUFFER_STRUCT(FAOSampleData2,)
@@ -325,28 +407,30 @@ class FAOParameters
 public:
 	void Bind(const FShaderParameterMap& ParameterMap)
 	{
-		AOMaxDistance.Bind(ParameterMap,TEXT("AOMaxDistance"));
+		AOObjectMaxDistance.Bind(ParameterMap,TEXT("AOObjectMaxDistance"));
 		AOStepScale.Bind(ParameterMap,TEXT("AOStepScale"));
 		AOStepExponentScale.Bind(ParameterMap,TEXT("AOStepExponentScale"));
 		AOMaxViewDistance.Bind(ParameterMap,TEXT("AOMaxViewDistance"));
+		AOGlobalMaxOcclusionDistance.Bind(ParameterMap,TEXT("AOGlobalMaxOcclusionDistance"));
 	}
 
 	friend FArchive& operator<<(FArchive& Ar,FAOParameters& Parameters)
 	{
-		Ar << Parameters.AOMaxDistance;
+		Ar << Parameters.AOObjectMaxDistance;
 		Ar << Parameters.AOStepScale;
 		Ar << Parameters.AOStepExponentScale;
 		Ar << Parameters.AOMaxViewDistance;
+		Ar << Parameters.AOGlobalMaxOcclusionDistance;
 		return Ar;
 	}
 
 	template<typename ShaderRHIParamRef>
 	void Set(FRHICommandList& RHICmdList, const ShaderRHIParamRef ShaderRHI, const FDistanceFieldAOParameters& Parameters)
 	{
-		SetShaderValue(RHICmdList, ShaderRHI, AOMaxDistance, Parameters.OcclusionMaxDistance);
+		SetShaderValue(RHICmdList, ShaderRHI, AOObjectMaxDistance, Parameters.ObjectMaxOcclusionDistance);
 
 		extern float GAOConeHalfAngle;
-		const float AOLargestSampleOffset = Parameters.OcclusionMaxDistance / (1 + FMath::Tan(GAOConeHalfAngle));
+		const float AOLargestSampleOffset = Parameters.ObjectMaxOcclusionDistance / (1 + FMath::Tan(GAOConeHalfAngle));
 
 		extern float GAOStepExponentScale;
 		extern uint32 GAONumConeSteps;
@@ -357,13 +441,17 @@ public:
 
 		extern float GetMaxAOViewDistance();
 		SetShaderValue(RHICmdList, ShaderRHI, AOMaxViewDistance, GetMaxAOViewDistance());
+
+		const float GlobalMaxOcclusionDistance = Parameters.GlobalMaxOcclusionDistance;
+		SetShaderValue(RHICmdList, ShaderRHI, AOGlobalMaxOcclusionDistance, GlobalMaxOcclusionDistance);
 	}
 
 private:
-	FShaderParameter AOMaxDistance;
+	FShaderParameter AOObjectMaxDistance;
 	FShaderParameter AOStepScale;
 	FShaderParameter AOStepExponentScale;
 	FShaderParameter AOMaxViewDistance;
+	FShaderParameter AOGlobalMaxOcclusionDistance;
 };
 
 inline float GetMaxAOViewDistance()
@@ -451,7 +539,7 @@ public:
 
 	size_t GetSizeBytes() const
 	{
-		return ConeVisibility.NumBytes + ConeData.NumBytes + StepBentNormal.NumBytes;
+		return ConeVisibility.NumBytes + ConeData.NumBytes + StepBentNormal.NumBytes + SurfelIrradiance.NumBytes + HeightfieldIrradiance.NumBytes;
 	}
 
 	FRWBuffer ConeVisibility;
@@ -461,54 +549,94 @@ public:
 	FRWBuffer HeightfieldIrradiance;
 };
 
-extern FRWBuffer* GDebugBuffer2;
-
-class FTrackGPUProgressCS : public FGlobalShader
+class FAOLevelParameters
 {
-	DECLARE_SHADER_TYPE(FTrackGPUProgressCS,Global)
 public:
-
-	static bool ShouldCache(EShaderPlatform Platform)
+	void Bind(const FShaderParameterMap& ParameterMap)
 	{
-		return IsFeatureLevelSupported(Platform, ERHIFeatureLevel::SM5) && DoesPlatformSupportDistanceFieldAO(Platform);
+		CurrentLevelDownsampleFactor.Bind(ParameterMap, TEXT("CurrentLevelDownsampleFactor"));
+		AOBufferSize.Bind(ParameterMap, TEXT("AOBufferSize"));
+		DownsampleFactorToBaseLevel.Bind(ParameterMap, TEXT("DownsampleFactorToBaseLevel"));
+		BaseLevelTexelSize.Bind(ParameterMap, TEXT("BaseLevelTexelSize"));
 	}
 
-	FTrackGPUProgressCS(const ShaderMetaType::CompiledShaderInitializerType& Initializer)
-		: FGlobalShader(Initializer)
+	template<typename TParamRef>
+	void Set(FRHICommandList& RHICmdList, const TParamRef& ShaderRHI, const FSceneView& View, int32 CurrentLevelDownsampleFactorValue)
 	{
-		DebugBuffer.Bind(Initializer.ParameterMap, TEXT("DebugBuffer"));
-		DebugId.Bind(Initializer.ParameterMap, TEXT("DebugId"));
+		SetShaderValue(RHICmdList, ShaderRHI, CurrentLevelDownsampleFactor, CurrentLevelDownsampleFactorValue);
+
+		// Round up, to match render target allocation
+		const FVector2D AOBufferSizeValue = FIntPoint::DivideAndRoundUp(GSceneRenderTargets.GetBufferSizeXY(), CurrentLevelDownsampleFactorValue);
+		SetShaderValue(RHICmdList, ShaderRHI, AOBufferSize, AOBufferSizeValue);
+
+		SetShaderValue(RHICmdList, ShaderRHI, DownsampleFactorToBaseLevel, CurrentLevelDownsampleFactorValue / GAODownsampleFactor);
+
+		const FIntPoint DownsampledBufferSize = GetBufferSizeForAO();
+		const FVector2D BaseLevelBufferSizeValue(1.0f / DownsampledBufferSize.X, 1.0f / DownsampledBufferSize.Y);
+		SetShaderValue(RHICmdList, ShaderRHI, BaseLevelTexelSize, BaseLevelBufferSizeValue);
 	}
 
-	FTrackGPUProgressCS()
+	friend FArchive& operator<<(FArchive& Ar,FAOLevelParameters& P)
 	{
+		Ar << P.CurrentLevelDownsampleFactor << P.AOBufferSize << P.DownsampleFactorToBaseLevel << P.BaseLevelTexelSize;
+		return Ar;
 	}
 
-	void SetParameters(FRHICommandList& RHICmdList, uint32 DebugIdValue)
-	{
-		FComputeShaderRHIParamRef ShaderRHI = GetComputeShader();
+private:
+	FShaderParameter CurrentLevelDownsampleFactor;
+	FShaderParameter AOBufferSize;
+	FShaderParameter DownsampleFactorToBaseLevel;
+	FShaderParameter BaseLevelTexelSize;
+};
 
-		DebugBuffer.SetBuffer(RHICmdList, ShaderRHI, *GDebugBuffer2);
-		SetShaderValue(RHICmdList, ShaderRHI, DebugId, DebugIdValue);
+class FScreenGridParameters
+{
+public:
+	void Bind(const FShaderParameterMap& ParameterMap)
+	{
+		BaseLevelTexelSize.Bind(ParameterMap, TEXT("BaseLevelTexelSize"));
+		JitterOffset.Bind(ParameterMap, TEXT("JitterOffset"));
+		ScreenGridConeVisibilitySize.Bind(ParameterMap, TEXT("ScreenGridConeVisibilitySize"));
+		DistanceFieldNormalTexture.Bind(ParameterMap, TEXT("DistanceFieldNormalTexture"));
+		DistanceFieldNormalSampler.Bind(ParameterMap, TEXT("DistanceFieldNormalSampler"));
 	}
 
-	void UnsetParameters(FRHICommandListImmediate& RHICmdList)
+	template<typename TParamRef>
+	void Set(FRHICommandList& RHICmdList, const TParamRef& ShaderRHI, const FViewInfo& View, FSceneRenderTargetItem& DistanceFieldNormal)
 	{
-		DebugBuffer.UnsetUAV(RHICmdList, GetComputeShader());
+		const FIntPoint DownsampledBufferSize = GetBufferSizeForAO();
+		const FVector2D BaseLevelTexelSizeValue(1.0f / DownsampledBufferSize.X, 1.0f / DownsampledBufferSize.Y);
+		SetShaderValue(RHICmdList, ShaderRHI, BaseLevelTexelSize, BaseLevelTexelSizeValue);
+
+		extern FVector2D GetJitterOffset(int32 SampleIndex);
+		SetShaderValue(RHICmdList, ShaderRHI, JitterOffset, GetJitterOffset(View.ViewState->GetDistanceFieldTemporalSampleIndex()));
+
+		FAOScreenGridResources* ScreenGridResources = View.ViewState->AOScreenGridResources;
+
+		SetShaderValue(RHICmdList, ShaderRHI, ScreenGridConeVisibilitySize, ScreenGridResources->ScreenGridDimensions);
+
+		SetTextureParameter(
+			RHICmdList,
+			ShaderRHI,
+			DistanceFieldNormalTexture,
+			DistanceFieldNormalSampler,
+			TStaticSamplerState<SF_Point,AM_Wrap,AM_Wrap,AM_Wrap>::GetRHI(),
+			DistanceFieldNormal.ShaderResourceTexture
+			);
 	}
 
-	virtual bool Serialize(FArchive& Ar) override
-	{		
-		bool bShaderHasOutdatedParameters = FGlobalShader::Serialize(Ar);
-		Ar << DebugBuffer;
-		Ar << DebugId;
-		return bShaderHasOutdatedParameters;
+	friend FArchive& operator<<(FArchive& Ar,FScreenGridParameters& P)
+	{
+		Ar << P.BaseLevelTexelSize << P.JitterOffset << P.ScreenGridConeVisibilitySize << P.DistanceFieldNormalTexture << P.DistanceFieldNormalSampler;
+		return Ar;
 	}
 
 private:
-
-	FRWShaderParameter DebugBuffer;
-	FShaderParameter DebugId;
+	FShaderParameter BaseLevelTexelSize;
+	FShaderParameter JitterOffset;
+	FShaderParameter ScreenGridConeVisibilitySize;
+	FShaderResourceParameter DistanceFieldNormalTexture;
+	FShaderResourceParameter DistanceFieldNormalSampler;
 };
 
 extern void TrackGPUProgress(FRHICommandListImmediate& RHICmdList, uint32 DebugId);
diff --git a/Engine/Source/Runtime/Renderer/Private/GlobalDistanceField.cpp b/Engine/Source/Runtime/Renderer/Private/GlobalDistanceField.cpp
new file mode 100644
--- /dev/null
+++ b/Engine/Source/Runtime/Renderer/Private/GlobalDistanceField.cpp
@@ -0,0 +1,805 @@
+// Copyright 1998-2015 Epic Games, Inc. All Rights Reserved.
+
+/*=============================================================================
+	GlobalDistanceField.cpp
+=============================================================================*/
+
+#include "RendererPrivate.h"
+#include "ScenePrivate.h"
+#include "UniformBuffer.h"
+#include "ShaderParameters.h"
+#include "PostProcessing.h"
+#include "SceneFilterRendering.h"
+#include "DistanceFieldLightingShared.h"
+#include "DistanceFieldSurfaceCacheLighting.h"
+#include "GlobalDistanceField.h"
+#include "RHICommandList.h"
+#include "SceneUtils.h"
+
+int32 GAOGlobalDistanceField = 1;
+FAutoConsoleVariableRef CVarAOGlobalDistanceField(
+	TEXT("r.AOGlobalDistanceField"),
+	GAOGlobalDistanceField,
+	TEXT("Whether to use a global distance field to optimize occlusion cone traces.\n")
+	TEXT("The global distance field is created by compositing object distance fields into clipmaps as the viewer moves through the level."),
+	ECVF_Cheat | ECVF_RenderThreadSafe
+	);
+
+int32 GAOUpdateGlobalDistanceField = 1;
+FAutoConsoleVariableRef CVarAOUpdateGlobalDistanceField(
+	TEXT("r.AOUpdateGlobalDistanceField"),
+	GAOUpdateGlobalDistanceField,
+	TEXT("Whether to update the global distance field, useful for debugging."),
+	ECVF_Cheat | ECVF_RenderThreadSafe
+	);
+
+int32 GAOGlobalDistanceFieldPartialUpdates = 1;
+FAutoConsoleVariableRef CVarAOGlobalDistanceFieldPartialUpdates(
+	TEXT("r.AOGlobalDistanceFieldPartialUpdates"),
+	GAOGlobalDistanceFieldPartialUpdates,
+	TEXT("Whether to allow partial updates of the global distance field.  When profiling it's useful to disable this and get the worst case composition time that happens on camera cuts."),
+	ECVF_Cheat | ECVF_RenderThreadSafe
+	);
+
+int32 GAOVisualizeGlobalDistanceField = 0;
+FAutoConsoleVariableRef CVarAOVisualizeGlobalDistanceField(
+	TEXT("r.AOVisualizeGlobalDistanceField"),
+	GAOVisualizeGlobalDistanceField,
+	TEXT("Whether to visualize the global distance field instead of the object DFs with the 'Mesh Distance Fields' Visualize mode"),
+	ECVF_Cheat | ECVF_RenderThreadSafe
+	);
+
+float GAOInnerGlobalDFClipmapDistance = 2500;
+FAutoConsoleVariableRef CVarAOInnerGlobalDFClipmapDistance(
+	TEXT("r.AOInnerGlobalDFClipmapDistance"),
+	GAOInnerGlobalDFClipmapDistance,
+	TEXT("World space distance from the camera that the first clipmap contains."),
+	ECVF_Cheat | ECVF_RenderThreadSafe
+	);
+
+float GAOGlobalDFClipmapDistanceExponent = 2;
+FAutoConsoleVariableRef CVarAOGlobalDFClipmapDistanceExponent(
+	TEXT("r.AOGlobalDFClipmapDistanceExponent"),
+	GAOGlobalDFClipmapDistanceExponent,
+	TEXT("Exponent used to derive each clipmap's size, together with r.AOInnerGlobalDFClipmapDistance."),
+	ECVF_Cheat | ECVF_RenderThreadSafe
+	);
+
+int32 GAOGlobalDFResolution = 128;
+FAutoConsoleVariableRef CVarAOGlobalDFResolution(
+	TEXT("r.AOGlobalDFResolution"),
+	GAOGlobalDFResolution,
+	TEXT("Resolution of the global distance field.  Higher values increase fidelity but also increase memory and composition cost."),
+	ECVF_Cheat | ECVF_RenderThreadSafe
+	);
+
+float GAOGlobalDFStartDistance = 200;
+FAutoConsoleVariableRef CVarAOGlobalDFStartDistance(
+	TEXT("r.AOGlobalDFStartDistance"),
+	GAOGlobalDFStartDistance,
+	TEXT("World space distance along a cone trace to switch to using the global distance field instead of the object distance fields.\n")
+	TEXT("This has to be large enough to hide the low res nature of the global distance field, but smaller values result in faster cone tracing."),
+	ECVF_Cheat | ECVF_RenderThreadSafe
+	);
+
+TGlobalResource<FDistanceFieldObjectBufferResource> GGlobalDistanceFieldCulledObjectBuffers;
+
+uint32 CullObjectsGroupSize = 64;
+
+class FCullObjectsForVolumeCS : public FGlobalShader
+{
+	DECLARE_SHADER_TYPE(FCullObjectsForVolumeCS,Global)
+public:
+
+	static bool ShouldCache(EShaderPlatform Platform)
+	{
+		return IsFeatureLevelSupported(Platform, ERHIFeatureLevel::SM5) && DoesPlatformSupportDistanceFieldAO(Platform);
+	}
+
+	static void ModifyCompilationEnvironment(EShaderPlatform Platform, FShaderCompilerEnvironment& OutEnvironment)
+	{
+		FGlobalShader::ModifyCompilationEnvironment(Platform,OutEnvironment);
+		OutEnvironment.SetDefine(TEXT("CULLOBJECTS_THREADGROUP_SIZE"), CullObjectsGroupSize);
+	}
+
+	FCullObjectsForVolumeCS(const ShaderMetaType::CompiledShaderInitializerType& Initializer)
+		: FGlobalShader(Initializer)
+	{
+		ObjectBufferParameters.Bind(Initializer.ParameterMap);
+		ObjectIndirectArguments.Bind(Initializer.ParameterMap, TEXT("ObjectIndirectArguments"));
+		CulledObjectBounds.Bind(Initializer.ParameterMap, TEXT("CulledObjectBounds"));
+		CulledObjectData.Bind(Initializer.ParameterMap, TEXT("CulledObjectData"));
+		CulledObjectBoxBounds.Bind(Initializer.ParameterMap, TEXT("CulledObjectBoxBounds"));
+		AOGlobalMaxSphereQueryRadius.Bind(Initializer.ParameterMap, TEXT("AOGlobalMaxSphereQueryRadius"));
+		VolumeBounds.Bind(Initializer.ParameterMap, TEXT("VolumeBounds"));
+	}
+
+	FCullObjectsForVolumeCS()
+	{
+	}
+
+	void SetParameters(FRHICommandList& RHICmdList, const FScene* Scene, const FSceneView& View, float MaxOcclusionDistance, const FVector4& VolumeBoundsValue)
+	{
+		FComputeShaderRHIParamRef ShaderRHI = GetComputeShader();
+		FGlobalShader::SetParameters(RHICmdList, ShaderRHI, View);
+		ObjectBufferParameters.Set(RHICmdList, ShaderRHI, *(Scene->DistanceFieldSceneData.ObjectBuffers), Scene->DistanceFieldSceneData.NumObjectsInBuffer);
+
+		ObjectIndirectArguments.SetBuffer(RHICmdList, ShaderRHI, GGlobalDistanceFieldCulledObjectBuffers.Buffers.ObjectIndirectArguments);
+		CulledObjectBounds.SetBuffer(RHICmdList, ShaderRHI, GGlobalDistanceFieldCulledObjectBuffers.Buffers.Bounds);
+		CulledObjectData.SetBuffer(RHICmdList, ShaderRHI, GGlobalDistanceFieldCulledObjectBuffers.Buffers.Data);
+		CulledObjectBoxBounds.SetBuffer(RHICmdList, ShaderRHI, GGlobalDistanceFieldCulledObjectBuffers.Buffers.BoxBounds);
+
+		extern float GAOConeHalfAngle;
+		const float GlobalMaxSphereQueryRadius = MaxOcclusionDistance / (1 + FMath::Tan(GAOConeHalfAngle));
+		SetShaderValue(RHICmdList, ShaderRHI, AOGlobalMaxSphereQueryRadius, GlobalMaxSphereQueryRadius);
+		SetShaderValue(RHICmdList, ShaderRHI, VolumeBounds, VolumeBoundsValue);
+	}
+
+	void UnsetParameters(FRHICommandList& RHICmdList)
+	{
+		ObjectBufferParameters.UnsetParameters(RHICmdList, GetComputeShader());
+		ObjectIndirectArguments.UnsetUAV(RHICmdList, GetComputeShader());
+		CulledObjectBounds.UnsetUAV(RHICmdList, GetComputeShader());
+		CulledObjectData.UnsetUAV(RHICmdList, GetComputeShader());
+		CulledObjectBoxBounds.UnsetUAV(RHICmdList, GetComputeShader());
+	}
+
+	virtual bool Serialize(FArchive& Ar)
+	{		
+		bool bShaderHasOutdatedParameters = FGlobalShader::Serialize(Ar);
+		Ar << ObjectBufferParameters;
+		Ar << ObjectIndirectArguments;
+		Ar << CulledObjectBounds;
+		Ar << CulledObjectData;
+		Ar << CulledObjectBoxBounds;
+		Ar << AOGlobalMaxSphereQueryRadius;
+		Ar << VolumeBounds;
+		return bShaderHasOutdatedParameters;
+	}
+
+private:
+
+	FDistanceFieldObjectBufferParameters ObjectBufferParameters;
+	FRWShaderParameter ObjectIndirectArguments;
+	FRWShaderParameter CulledObjectBounds;
+	FRWShaderParameter CulledObjectData;
+	FRWShaderParameter CulledObjectBoxBounds;
+	FShaderParameter AOGlobalMaxSphereQueryRadius;
+	FShaderParameter VolumeBounds;
+};
+
+IMPLEMENT_SHADER_TYPE(,FCullObjectsForVolumeCS,TEXT("GlobalDistanceField"),TEXT("CullObjectsForVolumeCS"),SF_Compute);
+
+const int32 GMaxGridCulledObjects = 2047;
+
+class FObjectGridBuffers : public FRenderResource
+{
+public:
+
+	int32 GridDimension;
+	FRWBuffer CulledObjectGrid;
+
+	FObjectGridBuffers()
+	{
+		GridDimension = 0;
+	}
+
+	virtual void InitDynamicRHI()  override
+	{
+		if (GridDimension > 0)
+		{
+			CulledObjectGrid.Initialize(sizeof(uint32), GridDimension * GridDimension * GridDimension * (GMaxGridCulledObjects + 1), PF_R32_UINT);
+		}
+	}
+
+	virtual void ReleaseDynamicRHI() override
+	{
+		CulledObjectGrid.Release();
+	}
+
+	size_t GetSizeBytes() const
+	{
+		return CulledObjectGrid.NumBytes;
+	}
+};
+
+TGlobalResource<FObjectGridBuffers> GObjectGridBuffers;
+
+const int32 GCullGridTileSize = 16;
+
+class FCullObjectsToGridCS : public FGlobalShader
+{
+	DECLARE_SHADER_TYPE(FCullObjectsToGridCS,Global)
+public:
+
+	static bool ShouldCache(EShaderPlatform Platform)
+	{
+		return IsFeatureLevelSupported(Platform, ERHIFeatureLevel::SM5) && DoesPlatformSupportDistanceFieldAO(Platform);
+	}
+
+	static void ModifyCompilationEnvironment(EShaderPlatform Platform, FShaderCompilerEnvironment& OutEnvironment)
+	{
+		FGlobalShader::ModifyCompilationEnvironment(Platform,OutEnvironment);
+		OutEnvironment.SetDefine(TEXT("CULL_GRID_TILE_SIZE"), GCullGridTileSize);
+		OutEnvironment.SetDefine(TEXT("MAX_GRID_CULLED_DF_OBJECTS"), GMaxGridCulledObjects);
+	}
+
+	FCullObjectsToGridCS(const ShaderMetaType::CompiledShaderInitializerType& Initializer)
+		: FGlobalShader(Initializer)
+	{
+		CulledObjectBufferParameters.Bind(Initializer.ParameterMap);
+		GlobalDistanceFieldParameters.Bind(Initializer.ParameterMap);
+		CulledObjectGrid.Bind(Initializer.ParameterMap, TEXT("CulledObjectGrid"));
+		CullGridDimension.Bind(Initializer.ParameterMap, TEXT("CullGridDimension"));
+		VolumeTexelSize.Bind(Initializer.ParameterMap, TEXT("VolumeTexelSize"));
+		UpdateRegionVolumeMin.Bind(Initializer.ParameterMap, TEXT("UpdateRegionVolumeMin"));
+		ClipmapIndex.Bind(Initializer.ParameterMap, TEXT("ClipmapIndex"));
+		AOGlobalMaxSphereQueryRadius.Bind(Initializer.ParameterMap, TEXT("AOGlobalMaxSphereQueryRadius"));
+	}
+
+	FCullObjectsToGridCS()
+	{
+	}
+
+	void SetParameters(
+		FRHICommandList& RHICmdList, 
+		const FScene* Scene, 
+		const FSceneView& View, 
+		float MaxOcclusionDistance, 
+		const FGlobalDistanceFieldInfo& GlobalDistanceFieldInfo,
+		int32 ClipmapIndexValue,
+		const FVolumeUpdateRegion& UpdateRegion)
+	{
+		FComputeShaderRHIParamRef ShaderRHI = GetComputeShader();
+		FGlobalShader::SetParameters(RHICmdList, ShaderRHI, View);
+		CulledObjectBufferParameters.Set(RHICmdList, ShaderRHI, GGlobalDistanceFieldCulledObjectBuffers.Buffers);
+		GlobalDistanceFieldParameters.Set(RHICmdList, ShaderRHI, GlobalDistanceFieldInfo);
+
+		CulledObjectGrid.SetBuffer(RHICmdList, ShaderRHI, GObjectGridBuffers.CulledObjectGrid);
+
+		const FIntVector GridDimensionValue(
+			FMath::DivideAndRoundUp<int32>(UpdateRegion.CellsSize.X, GCullGridTileSize),
+			FMath::DivideAndRoundUp<int32>(UpdateRegion.CellsSize.Y, GCullGridTileSize),
+			FMath::DivideAndRoundUp<int32>(UpdateRegion.CellsSize.Z, GCullGridTileSize));
+
+		SetShaderValue(RHICmdList, ShaderRHI, CullGridDimension, GridDimensionValue);
+		SetShaderValue(RHICmdList, ShaderRHI, VolumeTexelSize, FVector(1.0f / GAOGlobalDFResolution));
+		SetShaderValue(RHICmdList, ShaderRHI, UpdateRegionVolumeMin, UpdateRegion.Bounds.Min);
+		SetShaderValue(RHICmdList, ShaderRHI, ClipmapIndex, ClipmapIndexValue);
+
+		extern float GAOConeHalfAngle;
+		const float GlobalMaxSphereQueryRadius = MaxOcclusionDistance / (1 + FMath::Tan(GAOConeHalfAngle));
+		SetShaderValue(RHICmdList, ShaderRHI, AOGlobalMaxSphereQueryRadius, GlobalMaxSphereQueryRadius);
+	}
+
+	void UnsetParameters(FRHICommandList& RHICmdList)
+	{
+		CulledObjectGrid.UnsetUAV(RHICmdList, GetComputeShader());
+	}
+
+	virtual bool Serialize(FArchive& Ar)
+	{		
+		bool bShaderHasOutdatedParameters = FGlobalShader::Serialize(Ar);
+		Ar << CulledObjectBufferParameters;
+		Ar << GlobalDistanceFieldParameters;
+		Ar << CulledObjectGrid;
+		Ar << CullGridDimension;
+		Ar << VolumeTexelSize;
+		Ar << UpdateRegionVolumeMin;
+		Ar << ClipmapIndex;
+		Ar << AOGlobalMaxSphereQueryRadius;
+		return bShaderHasOutdatedParameters;
+	}
+
+private:
+
+	FDistanceFieldCulledObjectBufferParameters CulledObjectBufferParameters;
+	FGlobalDistanceFieldParameters GlobalDistanceFieldParameters;
+	FRWShaderParameter CulledObjectGrid;
+	FShaderParameter CullGridDimension;
+	FShaderParameter VolumeTexelSize;
+	FShaderParameter UpdateRegionVolumeMin;
+	FShaderParameter ClipmapIndex;
+	FShaderParameter AOGlobalMaxSphereQueryRadius;
+};
+
+IMPLEMENT_SHADER_TYPE(,FCullObjectsToGridCS,TEXT("GlobalDistanceField"),TEXT("CullObjectsToGridCS"),SF_Compute);
+
+const int32 CompositeTileSize = 4;
+
+class FCompositeObjectDistanceFieldsCS : public FGlobalShader
+{
+	DECLARE_SHADER_TYPE(FCompositeObjectDistanceFieldsCS,Global)
+public:
+
+	static bool ShouldCache(EShaderPlatform Platform)
+	{
+		return IsFeatureLevelSupported(Platform, ERHIFeatureLevel::SM5) && DoesPlatformSupportDistanceFieldAO(Platform);
+	}
+
+	static void ModifyCompilationEnvironment(EShaderPlatform Platform, FShaderCompilerEnvironment& OutEnvironment)
+	{
+		FGlobalShader::ModifyCompilationEnvironment(Platform,OutEnvironment);
+		OutEnvironment.SetDefine(TEXT("COMPOSITE_THREADGROUP_SIZE"), CompositeTileSize);
+		OutEnvironment.SetDefine(TEXT("CULL_GRID_TILE_SIZE"), GCullGridTileSize);
+		OutEnvironment.SetDefine(TEXT("MAX_GRID_CULLED_DF_OBJECTS"), GMaxGridCulledObjects);
+	}
+
+	FCompositeObjectDistanceFieldsCS(const ShaderMetaType::CompiledShaderInitializerType& Initializer)
+		: FGlobalShader(Initializer)
+	{
+		CulledObjectBufferParameters.Bind(Initializer.ParameterMap);
+		GlobalDistanceFieldParameters.Bind(Initializer.ParameterMap);
+		GlobalDistanceFieldTexture.Bind(Initializer.ParameterMap, TEXT("GlobalDistanceFieldTexture"));
+		CulledObjectGrid.Bind(Initializer.ParameterMap, TEXT("CulledObjectGrid"));
+		UpdateRegionSize.Bind(Initializer.ParameterMap, TEXT("UpdateRegionSize"));
+		CullGridDimension.Bind(Initializer.ParameterMap, TEXT("CullGridDimension"));
+		VolumeTexelSize.Bind(Initializer.ParameterMap, TEXT("VolumeTexelSize"));
+		UpdateRegionVolumeMin.Bind(Initializer.ParameterMap, TEXT("UpdateRegionVolumeMin"));
+		ClipmapIndex.Bind(Initializer.ParameterMap, TEXT("ClipmapIndex"));
+		AOGlobalMaxSphereQueryRadius.Bind(Initializer.ParameterMap, TEXT("AOGlobalMaxSphereQueryRadius"));
+	}
+
+	FCompositeObjectDistanceFieldsCS()
+	{
+	}
+
+	void SetParameters(
+		FRHICommandList& RHICmdList, 
+		const FScene* Scene, 
+		const FSceneView& View, 
+		float MaxOcclusionDistance, 
+		const FGlobalDistanceFieldInfo& GlobalDistanceFieldInfo,
+		int32 ClipmapIndexValue,
+		const FVolumeUpdateRegion& UpdateRegion)
+	{
+		FComputeShaderRHIParamRef ShaderRHI = GetComputeShader();
+		FGlobalShader::SetParameters(RHICmdList, ShaderRHI, View);
+		CulledObjectBufferParameters.Set(RHICmdList, ShaderRHI, GGlobalDistanceFieldCulledObjectBuffers.Buffers);
+		GlobalDistanceFieldParameters.Set(RHICmdList, ShaderRHI, GlobalDistanceFieldInfo);
+
+		GlobalDistanceFieldTexture.SetTexture(RHICmdList, ShaderRHI, GlobalDistanceFieldInfo.Clipmaps[ClipmapIndexValue].RenderTarget->GetRenderTargetItem().ShaderResourceTexture, GlobalDistanceFieldInfo.Clipmaps[ClipmapIndexValue].RenderTarget->GetRenderTargetItem().UAV);
+
+		SetSRVParameter(RHICmdList, ShaderRHI, CulledObjectGrid, GObjectGridBuffers.CulledObjectGrid.SRV);
+
+		const FIntVector GridDimensionValue(
+			FMath::DivideAndRoundUp<int32>(UpdateRegion.CellsSize.X, GCullGridTileSize),
+			FMath::DivideAndRoundUp<int32>(UpdateRegion.CellsSize.Y, GCullGridTileSize),
+			FMath::DivideAndRoundUp<int32>(UpdateRegion.CellsSize.Z, GCullGridTileSize));
+
+		SetShaderValue(RHICmdList, ShaderRHI, CullGridDimension, GridDimensionValue);
+
+		SetShaderValue(RHICmdList, ShaderRHI, UpdateRegionSize, UpdateRegion.CellsSize);
+		SetShaderValue(RHICmdList, ShaderRHI, VolumeTexelSize, FVector(1.0f / GAOGlobalDFResolution));
+		SetShaderValue(RHICmdList, ShaderRHI, UpdateRegionVolumeMin, UpdateRegion.Bounds.Min);
+		SetShaderValue(RHICmdList, ShaderRHI, ClipmapIndex, ClipmapIndexValue);
+
+		extern float GAOConeHalfAngle;
+		const float GlobalMaxSphereQueryRadius = MaxOcclusionDistance / (1 + FMath::Tan(GAOConeHalfAngle));
+		SetShaderValue(RHICmdList, ShaderRHI, AOGlobalMaxSphereQueryRadius, GlobalMaxSphereQueryRadius);
+	}
+
+	void UnsetParameters(FRHICommandList& RHICmdList)
+	{
+		GlobalDistanceFieldTexture.UnsetUAV(RHICmdList, GetComputeShader());
+	}
+
+	virtual bool Serialize(FArchive& Ar)
+	{		
+		bool bShaderHasOutdatedParameters = FGlobalShader::Serialize(Ar);
+		Ar << CulledObjectBufferParameters;
+		Ar << GlobalDistanceFieldParameters;
+		Ar << GlobalDistanceFieldTexture;
+		Ar << CulledObjectGrid;
+		Ar << UpdateRegionSize;
+		Ar << CullGridDimension;
+		Ar << VolumeTexelSize;
+		Ar << UpdateRegionVolumeMin;
+		Ar << ClipmapIndex;
+		Ar << AOGlobalMaxSphereQueryRadius;
+		return bShaderHasOutdatedParameters;
+	}
+
+private:
+
+	FDistanceFieldCulledObjectBufferParameters CulledObjectBufferParameters;
+	FGlobalDistanceFieldParameters GlobalDistanceFieldParameters;
+	FRWShaderParameter GlobalDistanceFieldTexture;
+	FShaderResourceParameter CulledObjectGrid;
+	FShaderParameter UpdateRegionSize;
+	FShaderParameter CullGridDimension;
+	FShaderParameter VolumeTexelSize;
+	FShaderParameter UpdateRegionVolumeMin;
+	FShaderParameter ClipmapIndex;
+	FShaderParameter AOGlobalMaxSphereQueryRadius;
+};
+
+IMPLEMENT_SHADER_TYPE(,FCompositeObjectDistanceFieldsCS,TEXT("GlobalDistanceField"),TEXT("CompositeObjectDistanceFieldsCS"),SF_Compute);
+
+/** Constructs and adds an update region based on camera movement for the given axis. */
+void AddUpdateRegionForAxis(FIntVector Movement, const FBox& ClipmapBounds, float CellSize, int32 ComponentIndex, TArray<FVolumeUpdateRegion, TInlineAllocator<3> >& UpdateRegions)
+{
+	FVolumeUpdateRegion UpdateRegion;
+	UpdateRegion.Bounds = ClipmapBounds;
+	UpdateRegion.CellsSize = FIntVector(GAOGlobalDFResolution);
+	UpdateRegion.CellsSize[ComponentIndex] = FMath::Min(FMath::Abs(Movement[ComponentIndex]), GAOGlobalDFResolution);
+
+	if (Movement[ComponentIndex] > 0)
+	{
+		// Positive axis movement, set the min of that axis to contain the newly exposed area
+		UpdateRegion.Bounds.Min[ComponentIndex] = FMath::Max(ClipmapBounds.Max[ComponentIndex] - Movement[ComponentIndex] * CellSize, ClipmapBounds.Min[ComponentIndex]);
+	}
+	else if (Movement[ComponentIndex] < 0)
+	{
+		// Negative axis movement, set the max of that axis to contain the newly exposed area
+		UpdateRegion.Bounds.Max[ComponentIndex] = FMath::Min(ClipmapBounds.Min[ComponentIndex] - Movement[ComponentIndex] * CellSize, ClipmapBounds.Max[ComponentIndex]);
+	}
+
+	if (UpdateRegion.CellsSize[ComponentIndex] > 0)
+	{
+		UpdateRegions.Add(UpdateRegion);
+	}
+}
+
+/** Constructs and adds an update region based on the given primitive bounds. */
+void AddUpdateRegionForPrimitive(const FVector4& Bounds, const FBox& ClipmapBounds, float CellSize, TArray<FVolumeUpdateRegion, TInlineAllocator<3> >& UpdateRegions)
+{
+	FBox BoundingBox((FVector)Bounds - Bounds.W, (FVector)Bounds + Bounds.W);
+
+	FVolumeUpdateRegion UpdateRegion;
+	UpdateRegion.Bounds.Init();
+	// Snap the min and clamp to clipmap bounds
+	UpdateRegion.Bounds.Min.X = FMath::Max(CellSize * FMath::FloorToFloat(BoundingBox.Min.X / CellSize), ClipmapBounds.Min.X);
+	UpdateRegion.Bounds.Min.Y = FMath::Max(CellSize * FMath::FloorToFloat(BoundingBox.Min.Y / CellSize), ClipmapBounds.Min.Y);
+	UpdateRegion.Bounds.Min.Z = FMath::Max(CellSize * FMath::FloorToFloat(BoundingBox.Min.Z / CellSize), ClipmapBounds.Min.Z);
+
+	// Derive the max from the snapped min and size, clamp to clipmap bounds
+	UpdateRegion.Bounds.Max = UpdateRegion.Bounds.Min + FVector(FMath::CeilToFloat(Bounds.W * 2 / CellSize)) * CellSize;
+	UpdateRegion.Bounds.Max.X = FMath::Min(UpdateRegion.Bounds.Max.X, ClipmapBounds.Max.X);
+	UpdateRegion.Bounds.Max.Y = FMath::Min(UpdateRegion.Bounds.Max.Y, ClipmapBounds.Max.Y);
+	UpdateRegion.Bounds.Max.Z = FMath::Min(UpdateRegion.Bounds.Max.Z, ClipmapBounds.Max.Z);
+
+	const FVector UpdateRegionSize = UpdateRegion.Bounds.GetSize();
+	UpdateRegion.CellsSize.X = FMath::TruncToInt(UpdateRegionSize.X / CellSize + .5f);
+	UpdateRegion.CellsSize.Y = FMath::TruncToInt(UpdateRegionSize.Y / CellSize + .5f);
+	UpdateRegion.CellsSize.Z = FMath::TruncToInt(UpdateRegionSize.Z / CellSize + .5f);
+
+	// Only add update regions with positive area
+	if (UpdateRegion.CellsSize.X > 0 && UpdateRegion.CellsSize.Y > 0 && UpdateRegion.CellsSize.Z > 0)
+	{
+		checkSlow(UpdateRegion.CellsSize.X <= GAOGlobalDFResolution && UpdateRegion.CellsSize.Y <= GAOGlobalDFResolution && UpdateRegion.CellsSize.Z <= GAOGlobalDFResolution);
+		UpdateRegions.Add(UpdateRegion);
+	}
+}
+
+void AllocateClipmapTexture(int32 ClipmapIndex, TRefCountPtr<IPooledRenderTarget>& Texture)
+{
+	const TCHAR* TextureName = TEXT("GlobalDistanceField0");
+
+	if (ClipmapIndex == 1)
+	{
+		TextureName = TEXT("GlobalDistanceField1");
+	}
+	else if (ClipmapIndex == 2)
+	{
+		TextureName = TEXT("GlobalDistanceField2");
+	}
+	else if (ClipmapIndex == 3)
+	{
+		TextureName = TEXT("GlobalDistanceField3");
+	}
+
+	GRenderTargetPool.FindFreeElement(
+		FPooledRenderTargetDesc(FPooledRenderTargetDesc::CreateVolumeDesc(
+			GAOGlobalDFResolution,
+			GAOGlobalDFResolution,
+			GAOGlobalDFResolution,
+			PF_R16F,
+			0,
+			TexCreate_ShaderResource | TexCreate_RenderTargetable | TexCreate_UAV,
+			false)),
+		Texture,
+		TextureName 
+	);
+}
+
+/** Staggers clipmap updates so there are only 2 per frame */
+bool ShouldUpdateClipmapThisFrame(int32 ClipmapIndex, int32 GlobalDistanceFieldUpdateIndex)
+{
+	if (ClipmapIndex == 0)
+	{
+		return true;
+	}
+	else if (ClipmapIndex == 1)
+	{
+		return GlobalDistanceFieldUpdateIndex % 2 == 0;
+	}
+	else if (ClipmapIndex == 2)
+	{
+		return GlobalDistanceFieldUpdateIndex % 4 == 1;
+	}
+	else
+	{
+		check(ClipmapIndex == 3);
+		return GlobalDistanceFieldUpdateIndex % 4 == 3;
+	}
+}
+
+float ComputeClipmapExtent(int32 ClipmapIndex)
+{
+	return GAOInnerGlobalDFClipmapDistance * FMath::Pow(GAOGlobalDFClipmapDistanceExponent, ClipmapIndex);
+}
+
+void ComputeUpdateRegionsAndUpdateViewState(const FViewInfo& View, FGlobalDistanceFieldInfo& GlobalDistanceFieldInfo, int32 NumClipmaps, const TArray<FVector4>& PrimitiveModifiedBounds)
+{
+	GlobalDistanceFieldInfo.Clipmaps.AddZeroed(NumClipmaps);
+
+	if (View.ViewState)
+	{
+		View.ViewState->GlobalDistanceFieldUpdateIndex++;
+
+		if (View.ViewState->GlobalDistanceFieldUpdateIndex > 4)
+		{
+			View.ViewState->GlobalDistanceFieldUpdateIndex = 0;
+		}
+
+		for (int32 ClipmapIndex = 0; ClipmapIndex < NumClipmaps; ClipmapIndex++)
+		{
+			FGlobalDistanceFieldClipmapState& ClipmapViewState = View.ViewState->GlobalDistanceFieldClipmapState[ClipmapIndex];
+			TRefCountPtr<IPooledRenderTarget>& RenderTarget = ClipmapViewState.VolumeTexture;
+
+			bool bReallocated = false;
+
+			if (!RenderTarget || RenderTarget->GetDesc().Extent.X != GAOGlobalDFResolution)
+			{
+				AllocateClipmapTexture(ClipmapIndex, RenderTarget);
+				bReallocated = true;
+			}
+
+			FGlobalDistanceFieldClipmap& Clipmap = GlobalDistanceFieldInfo.Clipmaps[ClipmapIndex];
+
+			const float Extent = ComputeClipmapExtent(ClipmapIndex);
+			const float CellSize = (Extent * 2) / GAOGlobalDFResolution;
+
+			if (ShouldUpdateClipmapThisFrame(ClipmapIndex, View.ViewState->GlobalDistanceFieldUpdateIndex)
+				|| !View.ViewState->bIntializedGlobalDistanceFieldOrigins
+				|| bReallocated)
+			{
+				const FVector NewCenter = View.ViewMatrices.ViewOrigin;
+
+				FIntVector GridCenter;
+				GridCenter.X = FMath::FloorToInt(NewCenter.X / CellSize);
+				GridCenter.Y = FMath::FloorToInt(NewCenter.Y / CellSize);
+				GridCenter.Z = FMath::FloorToInt(NewCenter.Z / CellSize);
+
+				const FVector SnappedCenter = FVector(GridCenter) * CellSize;
+				const FBox ClipmapBounds(SnappedCenter - Extent, SnappedCenter + Extent);
+
+				bool bUsePartialUpdates = GAOGlobalDistanceFieldPartialUpdates
+					&& View.ViewState->bIntializedGlobalDistanceFieldOrigins 
+					&& !bReallocated
+					// Only use partial updates with small numbers of primitive modifications
+					&& PrimitiveModifiedBounds.Num() < 100;
+
+				if (bUsePartialUpdates)
+				{
+					FIntVector Movement = GridCenter - ClipmapViewState.LastPartialUpdateOrigin;
+
+					// Add an update region for each potential axis of camera movement
+					AddUpdateRegionForAxis(Movement, ClipmapBounds, CellSize, 0, Clipmap.UpdateRegions);
+					AddUpdateRegionForAxis(Movement, ClipmapBounds, CellSize, 1, Clipmap.UpdateRegions);
+					AddUpdateRegionForAxis(Movement, ClipmapBounds, CellSize, 2, Clipmap.UpdateRegions);
+
+					for (int32 BoundsIndex = 0; BoundsIndex < PrimitiveModifiedBounds.Num(); BoundsIndex++)
+					{
+						// Add an update region for each primitive that has been modified
+						AddUpdateRegionForPrimitive(PrimitiveModifiedBounds[BoundsIndex], ClipmapBounds, CellSize, Clipmap.UpdateRegions);
+					}
+
+					int32 TotalTexelsBeingUpdated = 0;
+
+					// Trim fully contained update regions
+					for (int32 UpdateRegionIndex = 0; UpdateRegionIndex < Clipmap.UpdateRegions.Num(); UpdateRegionIndex++)
+					{
+						const FVolumeUpdateRegion& UpdateRegion = Clipmap.UpdateRegions[UpdateRegionIndex];
+						bool bCompletelyContained = false;
+
+						for (int32 OtherUpdateRegionIndex = 0; OtherUpdateRegionIndex < Clipmap.UpdateRegions.Num(); OtherUpdateRegionIndex++)
+						{
+							if (UpdateRegionIndex != OtherUpdateRegionIndex)
+							{
+								const FVolumeUpdateRegion& OtherUpdateRegion = Clipmap.UpdateRegions[OtherUpdateRegionIndex];
+
+								if (OtherUpdateRegion.Bounds.IsInsideOrOn(UpdateRegion.Bounds.Min)
+									&& OtherUpdateRegion.Bounds.IsInsideOrOn(UpdateRegion.Bounds.Max))
+								{
+									bCompletelyContained = true;
+									break;
+								}
+							}
+						}
+
+						if (bCompletelyContained)
+						{
+							Clipmap.UpdateRegions.RemoveAt(UpdateRegionIndex);
+							UpdateRegionIndex--;
+						}
+						else
+						{
+							TotalTexelsBeingUpdated += UpdateRegion.CellsSize.X * UpdateRegion.CellsSize.Y * UpdateRegion.CellsSize.Z;
+						}
+					}
+
+					if (TotalTexelsBeingUpdated > GAOGlobalDFResolution * GAOGlobalDFResolution * GAOGlobalDFResolution)
+					{
+						// Fall back to a full update if the partial updates were going to do more work
+						bUsePartialUpdates = false;
+						Clipmap.UpdateRegions.Reset();
+					}
+				}
+		
+				if (!bUsePartialUpdates)
+				{
+					FVolumeUpdateRegion UpdateRegion;
+					UpdateRegion.Bounds = ClipmapBounds;
+					UpdateRegion.CellsSize = FIntVector(GAOGlobalDFResolution);
+					Clipmap.UpdateRegions.Add(UpdateRegion);
+
+					// Store the location of the full update
+					ClipmapViewState.FullUpdateOrigin = GridCenter;
+					View.ViewState->bIntializedGlobalDistanceFieldOrigins = true;
+				}
+
+				ClipmapViewState.LastPartialUpdateOrigin = GridCenter;
+			}
+
+			const FVector Center = FVector(ClipmapViewState.LastPartialUpdateOrigin) * CellSize;
+
+			// Setup clipmap properties from view state exclusively, so we can skip updating on some frames
+			Clipmap.RenderTarget = ClipmapViewState.VolumeTexture;
+			Clipmap.Bounds = FBox(Center - Extent, Center + Extent);
+			// Scroll offset so the contents of the global distance field don't have to be moved as the camera moves around, only updated in slabs
+			Clipmap.ScrollOffset = FVector(ClipmapViewState.LastPartialUpdateOrigin - ClipmapViewState.FullUpdateOrigin) * CellSize;
+		}
+	}
+	else
+	{
+		for (int32 ClipmapIndex = 0; ClipmapIndex < NumClipmaps; ClipmapIndex++)
+		{
+			FGlobalDistanceFieldClipmap& Clipmap = GlobalDistanceFieldInfo.Clipmaps[ClipmapIndex];
+			AllocateClipmapTexture(ClipmapIndex, Clipmap.RenderTarget);
+			Clipmap.ScrollOffset = FVector::ZeroVector;
+
+			const float Extent = ComputeClipmapExtent(ClipmapIndex);
+			FVector Center = View.ViewMatrices.ViewOrigin;
+
+			const float CellSize = (Extent * 2) / GAOGlobalDFResolution;
+
+			FIntVector GridCenter;
+			GridCenter.X = FMath::FloorToInt(Center.X / CellSize);
+			GridCenter.Y = FMath::FloorToInt(Center.Y / CellSize);
+			GridCenter.Z = FMath::FloorToInt(Center.Z / CellSize);
+
+			Center = FVector(GridCenter) * CellSize;
+
+			FBox ClipmapBounds(Center - Extent, Center + Extent);
+			Clipmap.Bounds = ClipmapBounds;
+
+			FVolumeUpdateRegion UpdateRegion;
+			UpdateRegion.Bounds = ClipmapBounds;
+			UpdateRegion.CellsSize = FIntVector(GAOGlobalDFResolution);
+			Clipmap.UpdateRegions.Add(UpdateRegion);
+		}
+	}
+}
+
+/** 
+ * Updates the global distance field for a view.  
+ * Typically issues updates for just the newly exposed regions of the volume due to camera movement.
+ * In the worst case of a camera cut or large distance field scene changes, a full update of the global distance field will be done.
+ **/
+void UpdateGlobalDistanceFieldVolume(
+	FRHICommandListImmediate& RHICmdList, 
+	const FViewInfo& View, 
+	const FScene* Scene, 
+	float MaxOcclusionDistance, 
+	FGlobalDistanceFieldInfo& GlobalDistanceFieldInfo)
+{
+	if (Scene->DistanceFieldSceneData.NumObjectsInBuffer > 0)
+	{
+		ComputeUpdateRegionsAndUpdateViewState(View, GlobalDistanceFieldInfo, GMaxGlobalDistanceFieldClipmaps, Scene->DistanceFieldSceneData.PrimitiveModifiedBounds);
+
+		bool bHasUpdateRegions = false;
+
+		for (int32 ClipmapIndex = 0; ClipmapIndex < GlobalDistanceFieldInfo.Clipmaps.Num(); ClipmapIndex++)
+		{
+			bHasUpdateRegions = bHasUpdateRegions || GlobalDistanceFieldInfo.Clipmaps[ClipmapIndex].UpdateRegions.Num() > 0;
+		}
+
+		if (bHasUpdateRegions && GAOUpdateGlobalDistanceField)
+		{
+			SCOPED_DRAW_EVENT(RHICmdList, UpdateGlobalDistanceFieldVolume);
+
+			if (GGlobalDistanceFieldCulledObjectBuffers.Buffers.MaxObjects < Scene->DistanceFieldSceneData.NumObjectsInBuffer
+				|| GGlobalDistanceFieldCulledObjectBuffers.Buffers.MaxObjects > 3 * Scene->DistanceFieldSceneData.NumObjectsInBuffer)
+			{
+				GGlobalDistanceFieldCulledObjectBuffers.Buffers.MaxObjects = Scene->DistanceFieldSceneData.NumObjectsInBuffer * 5 / 4;
+				GGlobalDistanceFieldCulledObjectBuffers.Buffers.Release();
+				GGlobalDistanceFieldCulledObjectBuffers.Buffers.Initialize();
+			}
+
+			const uint32 MaxCullGridDimension = GAOGlobalDFResolution / GCullGridTileSize;
+
+			if (GObjectGridBuffers.GridDimension != MaxCullGridDimension)
+			{
+				GObjectGridBuffers.GridDimension = MaxCullGridDimension;
+				GObjectGridBuffers.UpdateRHI();
+			}
+
+			for (int32 ClipmapIndex = 0; ClipmapIndex < GlobalDistanceFieldInfo.Clipmaps.Num(); ClipmapIndex++)
+			{
+				SCOPED_DRAW_EVENT(RHICmdList, Clipmap);
+
+				FGlobalDistanceFieldClipmap& Clipmap = GlobalDistanceFieldInfo.Clipmaps[ClipmapIndex];
+
+				for (int32 UpdateRegionIndex = 0; UpdateRegionIndex < Clipmap.UpdateRegions.Num(); UpdateRegionIndex++)
+				{
+					const FVolumeUpdateRegion& UpdateRegion = Clipmap.UpdateRegions[UpdateRegionIndex];
+
+					{
+						SCOPED_DRAW_EVENT(RHICmdList, GridCull);
+
+						// Cull the global objects to the volume being updated
+						{
+							uint32 ClearValues[4] = { 0 };
+							RHICmdList.ClearUAV(GGlobalDistanceFieldCulledObjectBuffers.Buffers.ObjectIndirectArguments.UAV, ClearValues);
+
+							TShaderMapRef<FCullObjectsForVolumeCS> ComputeShader(View.ShaderMap);
+							RHICmdList.SetComputeShader(ComputeShader->GetComputeShader());
+							const FVector4 VolumeBounds(UpdateRegion.Bounds.GetCenter(), UpdateRegion.Bounds.GetExtent().Size());
+							ComputeShader->SetParameters(RHICmdList, Scene, View, MaxOcclusionDistance, VolumeBounds);
+
+							DispatchComputeShader(RHICmdList, *ComputeShader, FMath::DivideAndRoundUp<uint32>(Scene->DistanceFieldSceneData.NumObjectsInBuffer, CullObjectsGroupSize), 1, 1);
+							ComputeShader->UnsetParameters(RHICmdList);
+						}
+
+						// Further cull the objects into a low resolution grid
+						{
+							TShaderMapRef<FCullObjectsToGridCS> ComputeShader(View.ShaderMap);
+							RHICmdList.SetComputeShader(ComputeShader->GetComputeShader());
+							ComputeShader->SetParameters(RHICmdList, Scene, View, MaxOcclusionDistance, GlobalDistanceFieldInfo, ClipmapIndex, UpdateRegion);
+
+							const uint32 NumGroupsX = FMath::DivideAndRoundUp<int32>(UpdateRegion.CellsSize.X, GCullGridTileSize);
+							const uint32 NumGroupsY = FMath::DivideAndRoundUp<int32>(UpdateRegion.CellsSize.Y, GCullGridTileSize);
+							const uint32 NumGroupsZ = FMath::DivideAndRoundUp<int32>(UpdateRegion.CellsSize.Z, GCullGridTileSize); 
+
+							DispatchComputeShader(RHICmdList, *ComputeShader, NumGroupsX, NumGroupsY, NumGroupsZ);
+							ComputeShader->UnsetParameters(RHICmdList);
+						}
+					}
+
+					// Further cull the objects to the dispatch tile and composite the global distance field by computing the min distance from intersecting per-object distance fields
+					{
+						SCOPED_DRAW_EVENT(RHICmdList, TileCullAndComposite);
+						TShaderMapRef<FCompositeObjectDistanceFieldsCS> ComputeShader(View.ShaderMap);
+						RHICmdList.SetComputeShader(ComputeShader->GetComputeShader());
+						ComputeShader->SetParameters(RHICmdList, Scene, View, MaxOcclusionDistance, GlobalDistanceFieldInfo, ClipmapIndex, UpdateRegion);
+
+						//@todo - match typical update sizes.  Camera movement creates narrow slabs.
+						const uint32 NumGroupsX = FMath::DivideAndRoundUp<int32>(UpdateRegion.CellsSize.X, CompositeTileSize);
+						const uint32 NumGroupsY = FMath::DivideAndRoundUp<int32>(UpdateRegion.CellsSize.Y, CompositeTileSize);
+						const uint32 NumGroupsZ = FMath::DivideAndRoundUp<int32>(UpdateRegion.CellsSize.Z, CompositeTileSize);
+
+						DispatchComputeShader(RHICmdList, *ComputeShader, NumGroupsX, NumGroupsY, NumGroupsZ);
+						ComputeShader->UnsetParameters(RHICmdList);
+					}
+				}
+			}
+		}
+	}
+}
+
+void ListGlobalDistanceFieldMemory()
+{
+	UE_LOG(LogTemp, Log, TEXT("   Global DF culled objects %.3fMb"), (GGlobalDistanceFieldCulledObjectBuffers.Buffers.GetSizeBytes() + GObjectGridBuffers.GetSizeBytes()) / 1024.0f / 1024.0f);
+}
diff --git a/Engine/Source/Runtime/Renderer/Private/GlobalDistanceField.h b/Engine/Source/Runtime/Renderer/Private/GlobalDistanceField.h
new file mode 100644
--- /dev/null
+++ b/Engine/Source/Runtime/Renderer/Private/GlobalDistanceField.h
@@ -0,0 +1,109 @@
+// Copyright 1998-2015 Epic Games, Inc. All Rights Reserved.
+
+/*=============================================================================
+	GlobalDistanceField.h
+=============================================================================*/
+
+#pragma once
+
+class FGlobalDistanceFieldParameters
+{
+public:
+	void Bind(const FShaderParameterMap& ParameterMap)
+	{
+		GlobalDistanceFieldTexture.Bind(ParameterMap,TEXT("GlobalDistanceFieldTexture"));
+		GlobalDistanceFieldSampler.Bind(ParameterMap,TEXT("GlobalDistanceFieldSampler"));
+		GlobalVolumeCenterAndExtent.Bind(ParameterMap,TEXT("GlobalVolumeCenterAndExtent"));
+		GlobalVolumeDimension.Bind(ParameterMap,TEXT("GlobalVolumeDimension"));
+		GlobalVolumeWorldToUVAddAndMul.Bind(ParameterMap,TEXT("GlobalVolumeWorldToUVAddAndMul"));
+	}
+
+	friend FArchive& operator<<(FArchive& Ar,FGlobalDistanceFieldParameters& Parameters)
+	{
+		Ar << Parameters.GlobalDistanceFieldTexture;
+		Ar << Parameters.GlobalDistanceFieldSampler;
+		Ar << Parameters.GlobalVolumeCenterAndExtent;
+		Ar << Parameters.GlobalVolumeDimension;
+		Ar << Parameters.GlobalVolumeWorldToUVAddAndMul;
+		return Ar;
+	}
+
+	template<typename ShaderRHIParamRef>
+	void Set(FRHICommandList& RHICmdList, const ShaderRHIParamRef ShaderRHI, const FGlobalDistanceFieldInfo& Info)
+	{
+		for (int32 ClipmapIndex = 0; ClipmapIndex < (int32)GlobalDistanceFieldTexture.GetNumResources(); ClipmapIndex++)
+		{
+			FTextureRHIParamRef TextureValue = ClipmapIndex < Info.Clipmaps.Num() 
+				? Info.Clipmaps[ClipmapIndex].RenderTarget->GetRenderTargetItem().ShaderResourceTexture 
+				: NULL;
+
+			SetTextureParameter(
+				RHICmdList,
+				ShaderRHI,
+				GlobalDistanceFieldTexture,
+				GlobalDistanceFieldSampler,
+				TStaticSamplerState<SF_Bilinear,AM_Wrap,AM_Wrap,AM_Wrap>::GetRHI(),
+				TextureValue,
+				ClipmapIndex
+				);
+		}
+
+		FVector4 CenterAndExtent[GMaxGlobalDistanceFieldClipmaps] = {FVector4(0)};
+		FVector4 WorldToUVAddAndMul[GMaxGlobalDistanceFieldClipmaps] = {FVector4(0)};
+		const int32 NumToSet = GlobalVolumeCenterAndExtent.GetNumBytes() / sizeof(FVector4);
+
+		for (int32 ClipmapIndex = 0; ClipmapIndex < NumToSet; ClipmapIndex++)
+		{
+			if (ClipmapIndex < Info.Clipmaps.Num())
+			{
+				const FGlobalDistanceFieldClipmap& Clipmap = Info.Clipmaps[ClipmapIndex];
+				CenterAndExtent[ClipmapIndex] = FVector4(Clipmap.Bounds.GetCenter(), Clipmap.Bounds.GetExtent().X);
+
+				// GlobalUV = (WorldPosition - GlobalVolumeCenterAndExtent[ClipmapIndex].xyz + GlobalVolumeScollOffset[ClipmapIndex].xyz) / (GlobalVolumeCenterAndExtent[ClipmapIndex].w * 2) + .5f;
+				// WorldToUVMul = 1.0f / (GlobalVolumeCenterAndExtent[ClipmapIndex].w * 2)
+				// WorldToUVAdd = (GlobalVolumeScollOffset[ClipmapIndex].xyz - GlobalVolumeCenterAndExtent[ClipmapIndex].xyz) / (GlobalVolumeCenterAndExtent[ClipmapIndex].w * 2) + .5f
+				const FVector WorldToUVAdd = (Clipmap.ScrollOffset - Clipmap.Bounds.GetCenter()) / (Clipmap.Bounds.GetExtent().X * 2) + FVector(.5f);
+				WorldToUVAddAndMul[ClipmapIndex] = FVector4(WorldToUVAdd, 1.0f / (Clipmap.Bounds.GetExtent().X * 2));
+			}
+		}
+
+		SetShaderValueArray(RHICmdList, ShaderRHI, GlobalVolumeCenterAndExtent, CenterAndExtent, GMaxGlobalDistanceFieldClipmaps);
+		SetShaderValueArray(RHICmdList, ShaderRHI, GlobalVolumeWorldToUVAddAndMul, WorldToUVAddAndMul, GMaxGlobalDistanceFieldClipmaps);
+
+		extern int32 GAOGlobalDFResolution;
+		SetShaderValue(RHICmdList, ShaderRHI, GlobalVolumeDimension, (float)GAOGlobalDFResolution);
+	}
+
+private:
+	FShaderResourceParameter GlobalDistanceFieldTexture;
+	FShaderResourceParameter GlobalDistanceFieldSampler;
+	FShaderParameter GlobalVolumeCenterAndExtent;
+	FShaderParameter GlobalVolumeDimension;
+	FShaderParameter GlobalVolumeScollOffset;
+	FShaderParameter GlobalVolumeWorldToUVAddAndMul;
+};
+
+extern int32 GAOGlobalDistanceField;
+extern int32 GAOVisualizeGlobalDistanceField;
+
+inline bool UseGlobalDistanceField()
+{
+	return GAOGlobalDistanceField != 0;
+}
+
+inline bool UseGlobalDistanceField(const FDistanceFieldAOParameters& Parameters)
+{
+	return UseGlobalDistanceField() && Parameters.GlobalMaxOcclusionDistance > 0;
+}
+
+/** 
+ * Updates the global distance field for a view.  
+ * Typically issues updates for just the newly exposed regions of the volume due to camera movement.
+ * In the worst case of a camera cut or large distance field scene changes, a full update of the global distance field will be done.
+ **/
+extern void UpdateGlobalDistanceFieldVolume(
+	FRHICommandListImmediate& RHICmdList, 
+	const FViewInfo& View, 
+	const FScene* Scene, 
+	float MaxOcclusionDistance, 
+	FGlobalDistanceFieldInfo& Info);
\ No newline at end of file
diff --git a/Engine/Source/Runtime/Renderer/Private/HeightfieldLighting.cpp b/Engine/Source/Runtime/Renderer/Private/HeightfieldLighting.cpp
--- a/Engine/Source/Runtime/Renderer/Private/HeightfieldLighting.cpp
+++ b/Engine/Source/Runtime/Renderer/Private/HeightfieldLighting.cpp
@@ -429,8 +429,7 @@ void FHeightfieldLightingViewInfo::SetupVisibleHeightfields(const FViewInfo& Vie
 		&& NumPrimitives > 0
 		&& SupportsHeightfieldLighting(View.GetFeatureLevel(), View.GetShaderPlatform()))
 	{
-		extern float GAOMaxViewDistance;
-		const float MaxDistanceSquared = FMath::Square(GAOMaxViewDistance + GetGHeightfieldBounceDistance());
+		const float MaxDistanceSquared = FMath::Square(GetMaxAOViewDistance() + GetGHeightfieldBounceDistance());
 		float LocalToWorldScale = 1;
 
 		for (int32 HeightfieldPrimitiveIndex = 0; HeightfieldPrimitiveIndex < NumPrimitives; HeightfieldPrimitiveIndex++)
@@ -1225,7 +1224,6 @@ public:
 	static void ModifyCompilationEnvironment(EShaderPlatform Platform, FShaderCompilerEnvironment& OutEnvironment)
 	{
 		FGlobalShader::ModifyCompilationEnvironment(Platform,OutEnvironment);
-		OutEnvironment.SetDefine(TEXT("NUM_CONE_DIRECTIONS"), NumConeSampleDirections);
 
 		// To reduce shader compile time of compute shaders with shared memory, doesn't have an impact on generated code with current compiler (June 2010 DX SDK)
 		OutEnvironment.CompilerFlags.Add(CFLAG_StandardOptimization);
@@ -1243,9 +1241,6 @@ public:
 		RecordConeVisibility.Bind(Initializer.ParameterMap, TEXT("RecordConeVisibility"));
 		ScatterDrawParameters.Bind(Initializer.ParameterMap, TEXT("ScatterDrawParameters"));
 		SavedStartIndex.Bind(Initializer.ParameterMap, TEXT("SavedStartIndex"));
-		BentNormalNormalizeFactor.Bind(Initializer.ParameterMap, TEXT("BentNormalNormalizeFactor"));
-		DistanceFieldNormalTexture.Bind(Initializer.ParameterMap, TEXT("DistanceFieldNormalTexture"));
-		DistanceFieldNormalSampler.Bind(Initializer.ParameterMap, TEXT("DistanceFieldNormalSampler"));
 		RecordRadiusScale.Bind(Initializer.ParameterMap, TEXT("RecordRadiusScale"));
 		TanConeHalfAngle.Bind(Initializer.ParameterMap, TEXT("TanConeHalfAngle"));
 	}
@@ -1293,16 +1288,6 @@ public:
 
 		SetUniformBufferParameterImmediate(RHICmdList, ShaderRHI, GetUniformBufferParameter<FAOSampleData2>(), AOSampleData);
 
-		FVector UnoccludedVector(0);
-
-		for (int32 SampleIndex = 0; SampleIndex < NumConeSampleDirections; SampleIndex++)
-		{
-			UnoccludedVector += SampleDirections[SampleIndex];
-		}
-
-		float BentNormalNormalizeFactorValue = 1.0f / (UnoccludedVector / NumConeSampleDirections).Size();
-		SetShaderValue(RHICmdList, ShaderRHI, BentNormalNormalizeFactor, BentNormalNormalizeFactorValue);
-
 		extern float GAOConeHalfAngle;
 		SetShaderValue(RHICmdList, ShaderRHI, TanConeHalfAngle, FMath::Tan(GAOConeHalfAngle));
 		extern float GAORecordRadiusScale;
@@ -1327,9 +1312,6 @@ public:
 		Ar << RecordConeVisibility;
 		Ar << ScatterDrawParameters;
 		Ar << SavedStartIndex;
-		Ar << BentNormalNormalizeFactor;
-		Ar << DistanceFieldNormalTexture;
-		Ar << DistanceFieldNormalSampler;
 		Ar << RecordRadiusScale;
 		Ar << TanConeHalfAngle;
 		return bShaderHasOutdatedParameters;
@@ -1346,9 +1328,6 @@ private:
 	FRWShaderParameter RecordConeVisibility;
 	FShaderResourceParameter ScatterDrawParameters;
 	FShaderResourceParameter SavedStartIndex;
-	FShaderParameter BentNormalNormalizeFactor;
-	FShaderResourceParameter DistanceFieldNormalTexture;
-	FShaderResourceParameter DistanceFieldNormalSampler;
 	FShaderParameter RecordRadiusScale;
 	FShaderParameter TanConeHalfAngle;
 };
@@ -1421,7 +1400,6 @@ public:
 	static void ModifyCompilationEnvironment(EShaderPlatform Platform, FShaderCompilerEnvironment& OutEnvironment)
 	{
 		FGlobalShader::ModifyCompilationEnvironment(Platform,OutEnvironment);
-		OutEnvironment.SetDefine(TEXT("NUM_CONE_DIRECTIONS"), NumConeSampleDirections);
 
 		// To reduce shader compile time of compute shaders with shared memory, doesn't have an impact on generated code with current compiler (June 2010 DX SDK)
 		OutEnvironment.CompilerFlags.Add(CFLAG_StandardOptimization);
@@ -1438,10 +1416,8 @@ public:
 		IrradianceCachePositionRadius.Bind(Initializer.ParameterMap, TEXT("IrradianceCachePositionRadius"));
 		ScatterDrawParameters.Bind(Initializer.ParameterMap, TEXT("ScatterDrawParameters"));
 		SavedStartIndex.Bind(Initializer.ParameterMap, TEXT("SavedStartIndex"));
-		BentNormalNormalizeFactor.Bind(Initializer.ParameterMap, TEXT("BentNormalNormalizeFactor"));
-		DistanceFieldNormalTexture.Bind(Initializer.ParameterMap, TEXT("DistanceFieldNormalTexture"));
-		DistanceFieldNormalSampler.Bind(Initializer.ParameterMap, TEXT("DistanceFieldNormalSampler"));
 		RecordRadiusScale.Bind(Initializer.ParameterMap, TEXT("RecordRadiusScale"));
+		BentNormalNormalizeFactor.Bind(Initializer.ParameterMap, TEXT("BentNormalNormalizeFactor"));
 		TanConeHalfAngle.Bind(Initializer.ParameterMap, TEXT("TanConeHalfAngle"));
 		HeightfieldLighting.Bind(Initializer.ParameterMap, TEXT("HeightfieldLighting"));
 		HeightfieldLightingSampler.Bind(Initializer.ParameterMap, TEXT("HeightfieldLightingSampler"));
@@ -1533,10 +1509,8 @@ public:
 		Ar << IrradianceCachePositionRadius;
 		Ar << ScatterDrawParameters;
 		Ar << SavedStartIndex;
-		Ar << BentNormalNormalizeFactor;
-		Ar << DistanceFieldNormalTexture;
-		Ar << DistanceFieldNormalSampler;
 		Ar << RecordRadiusScale;
+		Ar << BentNormalNormalizeFactor;
 		Ar << TanConeHalfAngle;
 		Ar << HeightfieldLighting;
 		Ar << HeightfieldLightingSampler;
@@ -1556,10 +1530,8 @@ private:
 	FShaderResourceParameter IrradianceCachePositionRadius;
 	FShaderResourceParameter ScatterDrawParameters;
 	FShaderResourceParameter SavedStartIndex;
-	FShaderParameter BentNormalNormalizeFactor;
-	FShaderResourceParameter DistanceFieldNormalTexture;
-	FShaderResourceParameter DistanceFieldNormalSampler;
 	FShaderParameter RecordRadiusScale;
+	FShaderParameter BentNormalNormalizeFactor;
 	FShaderParameter TanConeHalfAngle;
 	FShaderResourceParameter HeightfieldLighting;
 	FShaderResourceParameter HeightfieldLightingSampler;
@@ -1629,3 +1601,341 @@ void FHeightfieldLightingViewInfo::ComputeIrradianceForSamples(
 		}
 	}
 }
+
+const int32 GHeightfieldOcclusionDispatchSize = 8;
+
+class FCalculateHeightfieldOcclusionScreenGridCS : public FGlobalShader
+{
+	DECLARE_SHADER_TYPE(FCalculateHeightfieldOcclusionScreenGridCS,Global)
+public:
+
+	static bool ShouldCache(EShaderPlatform Platform)
+	{
+		return IsFeatureLevelSupported(Platform, ERHIFeatureLevel::SM5) && DoesPlatformSupportDistanceFieldGI(Platform);
+	}
+
+	static void ModifyCompilationEnvironment(EShaderPlatform Platform, FShaderCompilerEnvironment& OutEnvironment)
+	{
+		FGlobalShader::ModifyCompilationEnvironment(Platform,OutEnvironment);
+
+		OutEnvironment.SetDefine(TEXT("HEIGHTFIELD_OCCLUSION_DISPATCH_SIZEX"), GHeightfieldOcclusionDispatchSize);
+		extern int32 GConeTraceDownsampleFactor;
+		OutEnvironment.SetDefine(TEXT("TRACE_DOWNSAMPLE_FACTOR"), GConeTraceDownsampleFactor);
+
+		// To reduce shader compile time of compute shaders with shared memory, doesn't have an impact on generated code with current compiler (June 2010 DX SDK)
+		OutEnvironment.CompilerFlags.Add(CFLAG_StandardOptimization);
+	}
+
+	FCalculateHeightfieldOcclusionScreenGridCS(const ShaderMetaType::CompiledShaderInitializerType& Initializer)
+		: FGlobalShader(Initializer)
+	{
+		AOParameters.Bind(Initializer.ParameterMap);
+		ScreenGridParameters.Bind(Initializer.ParameterMap);
+		HeightfieldDescriptionParameters.Bind(Initializer.ParameterMap);
+		HeightfieldTextureParameters.Bind(Initializer.ParameterMap);
+		ScreenGridConeVisibility.Bind(Initializer.ParameterMap, TEXT("ScreenGridConeVisibility"));
+		TanConeHalfAngle.Bind(Initializer.ParameterMap, TEXT("TanConeHalfAngle"));
+	}
+
+	FCalculateHeightfieldOcclusionScreenGridCS()
+	{
+	}
+
+	void SetParameters(
+		FRHICommandList& RHICmdList, 
+		const FViewInfo& View, 
+		UTexture2D* HeightfieldTextureValue, 
+		int32 NumHeightfieldsValue,
+		FSceneRenderTargetItem& DistanceFieldNormal, 
+		const FAOScreenGridResources& ScreenGridResources,
+		const FDistanceFieldAOParameters& Parameters)
+	{
+		FComputeShaderRHIParamRef ShaderRHI = GetComputeShader();
+		FGlobalShader::SetParameters(RHICmdList, ShaderRHI, View);
+
+		AOParameters.Set(RHICmdList, ShaderRHI, Parameters);
+		ScreenGridParameters.Set(RHICmdList, ShaderRHI, View, DistanceFieldNormal);
+		HeightfieldDescriptionParameters.Set(RHICmdList, ShaderRHI, NumHeightfieldsValue);
+		HeightfieldTextureParameters.Set(RHICmdList, ShaderRHI, HeightfieldTextureValue, NULL);
+
+		ScreenGridConeVisibility.SetBuffer(RHICmdList, ShaderRHI, ScreenGridResources.ScreenGridConeVisibility);
+
+		FAOSampleData2 AOSampleData;
+
+		TArray<FVector, TInlineAllocator<9> > SampleDirections;
+		GetSpacedVectors(SampleDirections);
+
+		for (int32 SampleIndex = 0; SampleIndex < NumConeSampleDirections; SampleIndex++)
+		{
+			AOSampleData.SampleDirections[SampleIndex] = FVector4(SampleDirections[SampleIndex]);
+		}
+
+		SetUniformBufferParameterImmediate(RHICmdList, ShaderRHI, GetUniformBufferParameter<FAOSampleData2>(), AOSampleData);
+
+		extern float GAOConeHalfAngle;
+		SetShaderValue(RHICmdList, ShaderRHI, TanConeHalfAngle, FMath::Tan(GAOConeHalfAngle));
+	}
+
+	void UnsetParameters(FRHICommandList& RHICmdList)
+	{
+		ScreenGridConeVisibility.UnsetUAV(RHICmdList, GetComputeShader());
+	}
+
+	virtual bool Serialize(FArchive& Ar) override
+	{		
+		bool bShaderHasOutdatedParameters = FGlobalShader::Serialize(Ar);
+		Ar << AOParameters;
+		Ar << ScreenGridParameters;
+		Ar << HeightfieldDescriptionParameters;
+		Ar << HeightfieldTextureParameters;
+		Ar << ScreenGridConeVisibility;
+		Ar << TanConeHalfAngle;
+		return bShaderHasOutdatedParameters;
+	}
+
+private:
+
+	FAOParameters AOParameters;
+	FScreenGridParameters ScreenGridParameters;
+	FHeightfieldDescriptionParameters HeightfieldDescriptionParameters;
+	FHeightfieldTextureParameters HeightfieldTextureParameters;
+	FRWShaderParameter ScreenGridConeVisibility;
+	FShaderParameter TanConeHalfAngle;
+};
+
+IMPLEMENT_SHADER_TYPE(,FCalculateHeightfieldOcclusionScreenGridCS,TEXT("HeightfieldLighting"),TEXT("CalculateHeightfieldOcclusionScreenGridCS"),SF_Compute);
+
+
+void FHeightfieldLightingViewInfo::ComputeOcclusionForScreenGrid(
+	const FViewInfo& View, 
+	FRHICommandListImmediate& RHICmdList, 
+	FSceneRenderTargetItem& DistanceFieldNormal, 
+	const FAOScreenGridResources& ScreenGridResources,
+	const FDistanceFieldAOParameters& Parameters) const
+{
+	const FScene* Scene = (const FScene*)View.Family->Scene;
+
+	if (Heightfield.ComponentDescriptions.Num() > 0 && GAOHeightfieldOcclusion)
+	{
+		SCOPED_DRAW_EVENT(RHICmdList, HeightfieldOcclusion);
+
+		FSceneViewState* ViewState = (FSceneViewState*)View.State;
+
+		{
+			for (TMap<FHeightfieldComponentTextures, TArray<FHeightfieldComponentDescription>>::TConstIterator It(Heightfield.ComponentDescriptions); It; ++It)
+			{
+				const TArray<FHeightfieldComponentDescription>& HeightfieldDescriptions = It.Value();
+
+				if (HeightfieldDescriptions.Num() > 0)
+				{
+					UploadHeightfieldDescriptions(HeightfieldDescriptions, FVector2D(1, 1), 1.0f / Heightfield.DownsampleFactor);
+
+					UTexture2D* HeightfieldTexture = It.Key().HeightAndNormal;
+
+					const uint32 GroupSizeX = FMath::DivideAndRoundUp(View.ViewRect.Size().X / GAODownsampleFactor, GHeightfieldOcclusionDispatchSize);
+					const uint32 GroupSizeY = FMath::DivideAndRoundUp(View.ViewRect.Size().Y / GAODownsampleFactor, GHeightfieldOcclusionDispatchSize);
+
+					TShaderMapRef<FCalculateHeightfieldOcclusionScreenGridCS> ComputeShader(View.ShaderMap);
+					RHICmdList.SetComputeShader(ComputeShader->GetComputeShader());
+					ComputeShader->SetParameters(RHICmdList, View, HeightfieldTexture, HeightfieldDescriptions.Num(), DistanceFieldNormal, ScreenGridResources, Parameters);
+					DispatchComputeShader(RHICmdList, *ComputeShader, GroupSizeX, GroupSizeY, 1);
+					ComputeShader->UnsetParameters(RHICmdList);
+				}
+			}
+		}
+	}
+}
+
+class FCalculateHeightfieldIrradianceScreenGridCS : public FGlobalShader
+{
+	DECLARE_SHADER_TYPE(FCalculateHeightfieldIrradianceScreenGridCS,Global)
+public:
+
+	static bool ShouldCache(EShaderPlatform Platform)
+	{
+		return IsFeatureLevelSupported(Platform, ERHIFeatureLevel::SM5) && DoesPlatformSupportDistanceFieldGI(Platform);
+	}
+
+	static void ModifyCompilationEnvironment(EShaderPlatform Platform, FShaderCompilerEnvironment& OutEnvironment)
+	{
+		FGlobalShader::ModifyCompilationEnvironment(Platform,OutEnvironment);
+
+		OutEnvironment.SetDefine(TEXT("HEIGHTFIELD_OCCLUSION_DISPATCH_SIZEX"), GHeightfieldOcclusionDispatchSize);
+		extern int32 GConeTraceDownsampleFactor;
+		OutEnvironment.SetDefine(TEXT("TRACE_DOWNSAMPLE_FACTOR"), GConeTraceDownsampleFactor);
+
+		// To reduce shader compile time of compute shaders with shared memory, doesn't have an impact on generated code with current compiler (June 2010 DX SDK)
+		OutEnvironment.CompilerFlags.Add(CFLAG_StandardOptimization);
+	}
+
+	FCalculateHeightfieldIrradianceScreenGridCS(const ShaderMetaType::CompiledShaderInitializerType& Initializer)
+		: FGlobalShader(Initializer)
+	{
+		AOParameters.Bind(Initializer.ParameterMap);
+		ScreenGridParameters.Bind(Initializer.ParameterMap);
+		HeightfieldDescriptionParameters.Bind(Initializer.ParameterMap);
+		HeightfieldIrradiance.Bind(Initializer.ParameterMap, TEXT("HeightfieldIrradiance"));
+		TanConeHalfAngle.Bind(Initializer.ParameterMap, TEXT("TanConeHalfAngle"));
+
+		GlobalHeightfieldParameters.Bind(Initializer.ParameterMap);
+		BentNormalNormalizeFactor.Bind(Initializer.ParameterMap, TEXT("BentNormalNormalizeFactor"));
+		HeightfieldLighting.Bind(Initializer.ParameterMap, TEXT("HeightfieldLighting"));
+		HeightfieldLightingSampler.Bind(Initializer.ParameterMap, TEXT("HeightfieldLightingSampler"));
+		InnerLightTransferDistance.Bind(Initializer.ParameterMap, TEXT("InnerLightTransferDistance"));
+		OuterLightTransferDistanceScale.Bind(Initializer.ParameterMap, TEXT("OuterLightTransferDistanceScale"));
+		RecordConeVisibility.Bind(Initializer.ParameterMap, TEXT("RecordConeVisibility"));
+	}
+
+	FCalculateHeightfieldIrradianceScreenGridCS()
+	{
+	}
+
+	void SetParameters(
+		FRHICommandList& RHICmdList, 
+		const FViewInfo& View, 
+		int32 NumHeightfieldsValue,
+		FSceneRenderTargetItem& DistanceFieldNormal, 
+		const FAOScreenGridResources& ScreenGridResources,
+		const FDistanceFieldAOParameters& Parameters,
+		const FHeightfieldLightingAtlas& Atlas)
+	{
+		FComputeShaderRHIParamRef ShaderRHI = GetComputeShader();
+		FGlobalShader::SetParameters(RHICmdList, ShaderRHI, View);
+
+		AOParameters.Set(RHICmdList, ShaderRHI, Parameters);
+		ScreenGridParameters.Set(RHICmdList, ShaderRHI, View, DistanceFieldNormal);
+		HeightfieldDescriptionParameters.Set(RHICmdList, ShaderRHI, NumHeightfieldsValue);
+
+		HeightfieldIrradiance.SetBuffer(RHICmdList, ShaderRHI, ScreenGridResources.HeightfieldIrradiance);
+
+		extern float GAOConeHalfAngle;
+		SetShaderValue(RHICmdList, ShaderRHI, TanConeHalfAngle, FMath::Tan(GAOConeHalfAngle));
+
+		GlobalHeightfieldParameters.Set(RHICmdList, ShaderRHI, Atlas);
+
+		{
+			FAOSampleData2 AOSampleData;
+
+			TArray<FVector, TInlineAllocator<9> > SampleDirections;
+			GetSpacedVectors(SampleDirections);
+
+			for (int32 SampleIndex = 0; SampleIndex < NumConeSampleDirections; SampleIndex++)
+			{
+				AOSampleData.SampleDirections[SampleIndex] = FVector4(SampleDirections[SampleIndex]);
+			}
+
+			SetUniformBufferParameterImmediate(RHICmdList, ShaderRHI, GetUniformBufferParameter<FAOSampleData2>(), AOSampleData);
+
+			FVector UnoccludedVector(0);
+
+			for (int32 SampleIndex = 0; SampleIndex < NumConeSampleDirections; SampleIndex++)
+			{
+				UnoccludedVector += SampleDirections[SampleIndex];
+			}
+
+			float BentNormalNormalizeFactorValue = 1.0f / (UnoccludedVector / NumConeSampleDirections).Size();
+			SetShaderValue(RHICmdList, ShaderRHI, BentNormalNormalizeFactor, BentNormalNormalizeFactorValue);
+		}
+
+		const FSceneViewState* ViewState = (const FSceneViewState*)View.State;
+		SetTextureParameter(RHICmdList, ShaderRHI, HeightfieldLighting, HeightfieldLightingSampler, TStaticSamplerState<SF_Bilinear>::GetRHI(), Atlas.Lighting->GetRenderTargetItem().ShaderResourceTexture);
+	
+		SetShaderValue(RHICmdList, ShaderRHI, InnerLightTransferDistance, GHeightfieldInnerBounceDistance);
+		SetShaderValue(RHICmdList, ShaderRHI, OuterLightTransferDistanceScale, GHeightfieldOuterBounceDistanceScale);
+
+		SetSRVParameter(RHICmdList, ShaderRHI, RecordConeVisibility, ScreenGridResources.ConeDepthVisibilityFunction.SRV);
+	}
+
+	void UnsetParameters(FRHICommandList& RHICmdList)
+	{
+		HeightfieldIrradiance.UnsetUAV(RHICmdList, GetComputeShader());
+	}
+
+	virtual bool Serialize(FArchive& Ar) override
+	{		
+		bool bShaderHasOutdatedParameters = FGlobalShader::Serialize(Ar);
+		Ar << AOParameters;
+		Ar << ScreenGridParameters;
+		Ar << HeightfieldDescriptionParameters;
+		Ar << HeightfieldIrradiance;
+		Ar << TanConeHalfAngle;
+		Ar << GlobalHeightfieldParameters;
+		Ar << BentNormalNormalizeFactor;
+		Ar << HeightfieldLighting;
+		Ar << HeightfieldLightingSampler;
+		Ar << InnerLightTransferDistance;
+		Ar << OuterLightTransferDistanceScale;
+		Ar << RecordConeVisibility;
+		return bShaderHasOutdatedParameters;
+	}
+
+private:
+
+	FAOParameters AOParameters;
+	FScreenGridParameters ScreenGridParameters;
+	FHeightfieldDescriptionParameters HeightfieldDescriptionParameters;
+	FRWShaderParameter HeightfieldIrradiance;
+	FShaderParameter TanConeHalfAngle;
+
+	FGlobalHeightfieldParameters GlobalHeightfieldParameters;
+	FShaderParameter BentNormalNormalizeFactor;
+	FShaderResourceParameter HeightfieldLighting;
+	FShaderResourceParameter HeightfieldLightingSampler;
+	FShaderParameter InnerLightTransferDistance;
+	FShaderParameter OuterLightTransferDistanceScale;
+	FShaderResourceParameter RecordConeVisibility;
+};
+
+IMPLEMENT_SHADER_TYPE(,FCalculateHeightfieldIrradianceScreenGridCS,TEXT("HeightfieldLighting"),TEXT("CalculateHeightfieldIrradianceScreenGridCS"),SF_Compute);
+
+void FHeightfieldLightingViewInfo::ComputeIrradianceForScreenGrid(
+	const FViewInfo& View, 
+	FRHICommandListImmediate& RHICmdList, 
+	FSceneRenderTargetItem& DistanceFieldNormal, 
+	const FAOScreenGridResources& ScreenGridResources,
+	const FDistanceFieldAOParameters& Parameters) const
+{
+	const FScene* Scene = (const FScene*)View.Family->Scene;
+
+	if (Heightfield.ComponentDescriptions.Num() > 0 
+		&& AllowHeightfieldGI(View)
+		&& SupportsHeightfieldLighting(View.GetFeatureLevel(), View.GetShaderPlatform()))
+	{
+		SCOPED_DRAW_EVENT(RHICmdList, HeightfieldIrradiance);
+
+		FSceneViewState* ViewState = (FSceneViewState*)View.State;
+
+		{
+			const FHeightfieldLightingAtlas& Atlas = *ViewState->HeightfieldLightingAtlas;
+
+			const FIntPoint LightingAtlasSize = Atlas.GetAtlasSize();
+			const FVector2D InvLightingAtlasSize(1.0f / LightingAtlasSize.X, 1.0f / LightingAtlasSize.Y);
+
+			TArray<FHeightfieldComponentDescription> CombinedHeightfieldDescriptions;
+
+			for (TMap<FHeightfieldComponentTextures, TArray<FHeightfieldComponentDescription>>::TConstIterator It(Heightfield.ComponentDescriptions); It; ++It)
+			{
+				const TArray<FHeightfieldComponentDescription>& HeightfieldDescriptions = It.Value();
+				CombinedHeightfieldDescriptions.Append(HeightfieldDescriptions);
+			}
+
+			if (CombinedHeightfieldDescriptions.Num() > 0)
+			{
+				UploadHeightfieldDescriptions(CombinedHeightfieldDescriptions, InvLightingAtlasSize, 1.0f / Heightfield.DownsampleFactor);
+
+				{
+					TShaderMapRef<FCalculateHeightfieldIrradianceScreenGridCS> ComputeShader(View.ShaderMap);
+					RHICmdList.SetComputeShader(ComputeShader->GetComputeShader());
+					ComputeShader->SetParameters(RHICmdList, View, CombinedHeightfieldDescriptions.Num(), DistanceFieldNormal, ScreenGridResources, Parameters, Atlas);
+
+					const uint32 GroupSizeX = View.ViewRect.Size().X / GAODownsampleFactor;
+					const uint32 GroupSizeY = View.ViewRect.Size().Y / GAODownsampleFactor;
+
+					DispatchComputeShader(RHICmdList, *ComputeShader, GroupSizeX, GroupSizeY, 1);
+
+					ComputeShader->UnsetParameters(RHICmdList);
+				}
+			}
+		}
+	}
+}
\ No newline at end of file
diff --git a/Engine/Source/Runtime/Renderer/Private/HeightfieldLighting.h b/Engine/Source/Runtime/Renderer/Private/HeightfieldLighting.h
--- a/Engine/Source/Runtime/Renderer/Private/HeightfieldLighting.h
+++ b/Engine/Source/Runtime/Renderer/Private/HeightfieldLighting.h
@@ -107,6 +107,20 @@ public:
 		int32 DepthLevel,
 		const class FDistanceFieldAOParameters& Parameters) const;
 
+	void ComputeOcclusionForScreenGrid(
+		const FViewInfo& View, 
+		FRHICommandListImmediate& RHICmdList, 
+		FSceneRenderTargetItem& DistanceFieldNormal,
+		const class FAOScreenGridResources& ScreenGridResources,
+		const class FDistanceFieldAOParameters& Parameters) const;
+
+	void ComputeIrradianceForScreenGrid(
+		const FViewInfo& View, 
+		FRHICommandListImmediate& RHICmdList, 
+		FSceneRenderTargetItem& DistanceFieldNormal, 
+		const FAOScreenGridResources& ScreenGridResources,
+		const FDistanceFieldAOParameters& Parameters) const;
+
 private:
 
 	FHeightfieldDescription Heightfield;
diff --git a/Engine/Source/Runtime/Renderer/Private/RendererScene.cpp b/Engine/Source/Runtime/Renderer/Private/RendererScene.cpp
--- a/Engine/Source/Runtime/Renderer/Private/RendererScene.cpp
+++ b/Engine/Source/Runtime/Renderer/Private/RendererScene.cpp
@@ -83,7 +83,9 @@ FSceneViewState::FSceneViewState()
 	MIDUsedCount = 0;
 	TemporalAASampleIndex = 0;
 	TemporalAASampleCount = 1;
+	DistanceFieldTemporalSampleIndex = 0;
 	AOTileIntersectionResources = NULL;
+	AOScreenGridResources = NULL;
 	bBokehDOFHistory = true;
 	bBokehDOFHistory2 = true;
 
@@ -95,6 +97,9 @@ FSceneViewState::FSceneViewState()
 		TranslucencyLightingCacheAllocations[CascadeIndex] = NULL;
 	}
 
+	bIntializedGlobalDistanceFieldOrigins = false;
+	GlobalDistanceFieldUpdateIndex = 0;
+
 #if BUFFERED_OCCLUSION_QUERIES
 	NumBufferedFrames = FOcclusionQueryHelpers::GetNumBufferedFrames();
 	ShadowOcclusionQueryMaps.Empty(NumBufferedFrames);
@@ -129,6 +134,8 @@ FSceneViewState::~FSceneViewState()
 	DestroyRenderResource(HeightfieldLightingAtlas);
 	DestroyRenderResource(AOTileIntersectionResources);
 	AOTileIntersectionResources = NULL;
+	DestroyRenderResource(AOScreenGridResources);
+	AOScreenGridResources = NULL;
 	DestroyLightPropagationVolume();
 }
 
diff --git a/Engine/Source/Runtime/Renderer/Private/ScenePrivate.h b/Engine/Source/Runtime/Renderer/Private/ScenePrivate.h
--- a/Engine/Source/Runtime/Renderer/Private/ScenePrivate.h
+++ b/Engine/Source/Runtime/Renderer/Private/ScenePrivate.h
@@ -340,6 +340,21 @@ public:
 	bool bValid;
 };
 
+class FGlobalDistanceFieldClipmapState
+{
+public:
+
+	FGlobalDistanceFieldClipmapState()
+	{
+		FullUpdateOrigin = FIntVector::ZeroValue;
+		LastPartialUpdateOrigin = FIntVector::ZeroValue;
+	}
+
+	FIntVector FullUpdateOrigin;
+	FIntVector LastPartialUpdateOrigin;
+	TRefCountPtr<IPooledRenderTarget> VolumeTexture;
+};
+
 /** Maps a single primitive to it's per-view fading state data */
 typedef TMap<FPrimitiveComponentId, FPrimitiveFadingState> FPrimitiveFadingStateMap;
 
@@ -382,6 +397,9 @@ private:
 /** Random table for occlusion **/
 extern FOcclusionRandomStream GOcclusionRandomStream;
 
+/** Must match global distance field shaders. */
+const int32 GMaxGlobalDistanceFieldClipmaps = 4;
+
 /**
  * The scene manager's private implementation of persistent view state.
  * This class is associated with a particular camera across multiple frames by the game thread.
@@ -498,6 +516,8 @@ private:
 	// >= 1, 1 means there is no TemporalAA
 	uint8 TemporalAASampleCount;
 
+	int32 DistanceFieldTemporalSampleIndex;
+
 	// can be 0
 	FLightPropagationVolume* LightPropagationVolume;
 
@@ -533,6 +553,12 @@ public:
 	/** Distance field AO tile intersection GPU resources.  Last frame's state is not used, but they must be sized exactly to the view so stored here. */
 	class FTileIntersectionResources* AOTileIntersectionResources;
 
+	class FAOScreenGridResources* AOScreenGridResources;
+
+	bool bIntializedGlobalDistanceFieldOrigins;
+	FGlobalDistanceFieldClipmapState GlobalDistanceFieldClipmapState[GMaxGlobalDistanceFieldClipmaps];
+	int32 GlobalDistanceFieldUpdateIndex;
+
 	// Is DOFHistoryRT set from Bokeh DOF?
 	bool bBokehDOFHistory;
 	bool bBokehDOFHistory2;
@@ -572,6 +598,24 @@ public:
 		}
 	}
 
+	void SetupDistanceFieldTemporalOffset(const FSceneViewFamily& Family)
+	{
+		if (!Family.bWorldIsPaused)
+		{
+			DistanceFieldTemporalSampleIndex++;
+		}
+
+		if(DistanceFieldTemporalSampleIndex >= 4)
+		{
+			DistanceFieldTemporalSampleIndex = 0;
+		}
+	}
+
+	int32 GetDistanceFieldTemporalSampleIndex() const
+	{
+		return DistanceFieldTemporalSampleIndex;
+	}
+
 	void FreeSeparateTranslucency()
 	{
 		SeparateTranslucencyRT.SafeRelease();
@@ -690,6 +734,11 @@ public:
 		MobileAaColor1.SafeRelease();
 		SelectionOutlineCacheKey.SafeRelease();
 		SelectionOutlineCacheValue.SafeRelease();
+
+		for (int32 CascadeIndex = 0; CascadeIndex < ARRAY_COUNT(GlobalDistanceFieldClipmapState); CascadeIndex++)
+		{
+			GlobalDistanceFieldClipmapState[CascadeIndex].VolumeTexture.SafeRelease();
+		}
 	}
 
 	// FSceneViewStateInterface
@@ -894,11 +943,13 @@ class FPrimitiveAndInstance
 {
 public:
 
-	FPrimitiveAndInstance(FPrimitiveSceneInfo* InPrimitive, int32 InInstanceIndex) :
+	FPrimitiveAndInstance(const FVector4& InBoundingSphere, FPrimitiveSceneInfo* InPrimitive, int32 InInstanceIndex) :
+		BoundingSphere(InBoundingSphere),
 		Primitive(InPrimitive),
 		InstanceIndex(InInstanceIndex)
 	{}
 
+	FVector4 BoundingSphere;
 	FPrimitiveSceneInfo* Primitive;
 	int32 InstanceIndex;
 };
@@ -1036,6 +1087,7 @@ public:
 	TArray<FPrimitiveSceneInfo*> PendingAddOperations;
 	TSet<FPrimitiveSceneInfo*> PendingUpdateOperations;
 	TArray<FPrimitiveRemoveInfo> PendingRemoveOperations;
+	TArray<FVector4> PrimitiveModifiedBounds;
 
 	/** Used to detect atlas reallocations, since objects store UVs into the atlas and need to be updated when it changes. */
 	int32 AtlasGeneration;
@@ -1212,6 +1264,8 @@ private:
 	/** Tracks used sections of the volume texture atlas. */
 	FTextureLayout3d BlockAllocator;
 
+	int32 NextPointId;
+
 	/** Tracks primitive allocations by component, so that they persist across re-registers. */
 	TMap<FPrimitiveComponentId, FIndirectLightingCacheAllocation*> PrimitiveAllocations;
 };
diff --git a/Engine/Source/Runtime/Renderer/Private/SceneRendering.cpp b/Engine/Source/Runtime/Renderer/Private/SceneRendering.cpp
--- a/Engine/Source/Runtime/Renderer/Private/SceneRendering.cpp
+++ b/Engine/Source/Runtime/Renderer/Private/SceneRendering.cpp
@@ -301,6 +301,8 @@ void FViewInfo::Init()
 	MaxShadowCascades = FMath::Clamp<int32>(CVarMaxShadowCascades.GetValueOnAnyThread(), 1, 10);
 
 	ShaderMap = GetGlobalShaderMap(FeatureLevel);
+
+	ViewState = (FSceneViewState*)State;
 }
 
 FViewInfo::~FViewInfo()
diff --git a/Engine/Source/Runtime/Renderer/Private/SceneRendering.h b/Engine/Source/Runtime/Renderer/Private/SceneRendering.h
--- a/Engine/Source/Runtime/Renderer/Private/SceneRendering.h
+++ b/Engine/Source/Runtime/Renderer/Private/SceneRendering.h
@@ -346,11 +346,50 @@ public:
 	}
 };
 
+class FVolumeUpdateRegion
+{
+public:
+	/** World space bounds. */
+	FBox Bounds;
+
+	/** Number of texels in each dimension to update. */
+	FIntVector CellsSize;
+};
+
+class FGlobalDistanceFieldClipmap
+{
+public:
+	/** World space bounds. */
+	FBox Bounds;
+
+	/** Offset applied to UVs so that only new or dirty areas of the volume texture have to be updated. */
+	FVector ScrollOffset;
+
+	/** Regions in the volume texture to update. */
+	TArray<FVolumeUpdateRegion, TInlineAllocator<3> > UpdateRegions;
+
+	/** Volume texture for this clipmap. */
+	TRefCountPtr<IPooledRenderTarget> RenderTarget;
+};
+
+class FGlobalDistanceFieldInfo
+{
+public:
+
+	TArray<FGlobalDistanceFieldClipmap> Clipmaps;
+};
+
 /** A FSceneView with additional state used by the scene renderer. */
 class FViewInfo : public FSceneView
 {
 public:
 
+	/** 
+	 * The view's state, or NULL if no state exists.
+	 * This should be used internally to the renderer module to avoid having to cast View.State to an FSceneViewState*
+	 */
+	FSceneViewState* ViewState;
+
 	/** A map from primitive ID to a boolean visibility value. */
 	FSceneBitArray PrimitiveVisibilityMap;
 
@@ -394,6 +433,9 @@ public:
 	/** The dynamic editor primitives visible in this view. */
 	TArray<const FPrimitiveSceneInfo*,SceneRenderingAllocator> VisibleEditorPrimitives;
 
+	/** View dependent global distance field clipmap info. */
+	FGlobalDistanceFieldInfo GlobalDistanceFieldInfo;
+
 	/** Set of translucent prims for this view */
 	FTranslucentPrimSet TranslucentPrimSet;
 
diff --git a/Engine/Source/Runtime/Renderer/Private/SceneVisibility.cpp b/Engine/Source/Runtime/Renderer/Private/SceneVisibility.cpp
--- a/Engine/Source/Runtime/Renderer/Private/SceneVisibility.cpp
+++ b/Engine/Source/Runtime/Renderer/Private/SceneVisibility.cpp
@@ -1440,7 +1440,7 @@ void FSceneRenderer::PreVisibilityFrameSetup(FRHICommandListImmediate& RHICmdLis
 	for(int32 ViewIndex = 0;ViewIndex < Views.Num();ViewIndex++)
 	{
 		FViewInfo& View = Views[ViewIndex];
-		FSceneViewState* ViewState = (FSceneViewState*) View.State;
+		FSceneViewState* ViewState = View.ViewState;
 		static bool bEnableTimeScale = true;
 
 		// Once per render increment the occlusion frame counter.
@@ -1468,6 +1468,11 @@ void FSceneRenderer::PreVisibilityFrameSetup(FRHICommandListImmediate& RHICmdLis
 		View.TemporalJitterPixelsX = 0.0f;
 		View.TemporalJitterPixelsY = 0.0f;
 
+		if (ViewState)
+		{
+			ViewState->SetupDistanceFieldTemporalOffset(ViewFamily);
+		}
+
 		if( View.FinalPostProcessSettings.AntiAliasingMethod == AAM_TemporalAA && ViewState )
 		{
 			// Subpixel jitter for temporal AA
