Fixed global distance field staggered ciipmap updates missing some primitive modifications (depending on the frame of the primitive modification)

diff --git a/Engine/Source/Runtime/Renderer/Private/GlobalDistanceField.cpp b/Engine/Source/Runtime/Renderer/Private/GlobalDistanceField.cpp
--- a/Engine/Source/Runtime/Renderer/Private/GlobalDistanceField.cpp
+++ b/Engine/Source/Runtime/Renderer/Private/GlobalDistanceField.cpp
@@ -561,6 +561,9 @@ void ComputeUpdateRegionsAndUpdateViewState(const FViewInfo& View, FGlobalDistan
 			const float Extent = ComputeClipmapExtent(ClipmapIndex);
 			const float CellSize = (Extent * 2) / GAOGlobalDFResolution;
 
+			// Accumulate primitive modifications in the viewstate in case we don't update the clipmap this frame
+			ClipmapViewState.PrimitiveModifiedBounds.Append(PrimitiveModifiedBounds);
+
 			if (ShouldUpdateClipmapThisFrame(ClipmapIndex, View.ViewState->GlobalDistanceFieldUpdateIndex)
 				|| !View.ViewState->bIntializedGlobalDistanceFieldOrigins
 				|| bReallocated)
@@ -579,7 +582,7 @@ void ComputeUpdateRegionsAndUpdateViewState(const FViewInfo& View, FGlobalDistan
 					&& View.ViewState->bIntializedGlobalDistanceFieldOrigins 
 					&& !bReallocated
 					// Only use partial updates with small numbers of primitive modifications
-					&& PrimitiveModifiedBounds.Num() < 100;
+					&& ClipmapViewState.PrimitiveModifiedBounds.Num() < 100;
 
 				if (bUsePartialUpdates)
 				{
@@ -590,10 +593,10 @@ void ComputeUpdateRegionsAndUpdateViewState(const FViewInfo& View, FGlobalDistan
 					AddUpdateRegionForAxis(Movement, ClipmapBounds, CellSize, 1, Clipmap.UpdateRegions);
 					AddUpdateRegionForAxis(Movement, ClipmapBounds, CellSize, 2, Clipmap.UpdateRegions);
 
-					for (int32 BoundsIndex = 0; BoundsIndex < PrimitiveModifiedBounds.Num(); BoundsIndex++)
+					for (int32 BoundsIndex = 0; BoundsIndex < ClipmapViewState.PrimitiveModifiedBounds.Num(); BoundsIndex++)
 					{
 						// Add an update region for each primitive that has been modified
-						AddUpdateRegionForPrimitive(PrimitiveModifiedBounds[BoundsIndex], ClipmapBounds, CellSize, Clipmap.UpdateRegions);
+						AddUpdateRegionForPrimitive(ClipmapViewState.PrimitiveModifiedBounds[BoundsIndex], ClipmapBounds, CellSize, Clipmap.UpdateRegions);
 					}
 
 					int32 TotalTexelsBeingUpdated = 0;
@@ -650,6 +653,7 @@ void ComputeUpdateRegionsAndUpdateViewState(const FViewInfo& View, FGlobalDistan
 					View.ViewState->bIntializedGlobalDistanceFieldOrigins = true;
 				}
 
+				ClipmapViewState.PrimitiveModifiedBounds.Reset();
 				ClipmapViewState.LastPartialUpdateOrigin = GridCenter;
 			}
 
diff --git a/Engine/Source/Runtime/Renderer/Private/ScenePrivate.h b/Engine/Source/Runtime/Renderer/Private/ScenePrivate.h
--- a/Engine/Source/Runtime/Renderer/Private/ScenePrivate.h
+++ b/Engine/Source/Runtime/Renderer/Private/ScenePrivate.h
@@ -352,6 +352,7 @@ public:
 
 	FIntVector FullUpdateOrigin;
 	FIntVector LastPartialUpdateOrigin;
+	TArray<FVector4> PrimitiveModifiedBounds;
 	TRefCountPtr<IPooledRenderTarget> VolumeTexture;
 };
 
