Linux compile errors

diff --git a/Engine/Source/Runtime/Renderer/Private/DistanceFieldSurfaceCacheLighting.cpp b/Engine/Source/Runtime/Renderer/Private/DistanceFieldSurfaceCacheLighting.cpp
--- a/Engine/Source/Runtime/Renderer/Private/DistanceFieldSurfaceCacheLighting.cpp
+++ b/Engine/Source/Runtime/Renderer/Private/DistanceFieldSurfaceCacheLighting.cpp
@@ -2690,8 +2690,6 @@ void RenderDistanceFieldAOSurfaceCache(
 	{
 		SCOPED_DRAW_EVENT(RHICmdList, TrimRecords);
 
-		FSurfaceCacheResources& SurfaceCacheResources = *Scene->SurfaceCacheResources;
-
 		// Copy and trim last frame's surface cache samples
 		for (int32 DepthLevel = GAOMaxLevel; DepthLevel >= GAOMinLevel; DepthLevel--)
 		{
diff --git a/Engine/Source/Runtime/Renderer/Private/GlobalDistanceField.h b/Engine/Source/Runtime/Renderer/Private/GlobalDistanceField.h
--- a/Engine/Source/Runtime/Renderer/Private/GlobalDistanceField.h
+++ b/Engine/Source/Runtime/Renderer/Private/GlobalDistanceField.h
@@ -48,8 +48,9 @@ public:
 				);
 		}
 
-		FVector4 CenterAndExtent[GMaxGlobalDistanceFieldClipmaps] = {FVector4(0)};
-		FVector4 WorldToUVAddAndMul[GMaxGlobalDistanceFieldClipmaps] = {FVector4(0)};
+		FVector4 CenterAndExtent[GMaxGlobalDistanceFieldClipmaps];
+		FVector4 WorldToUVAddAndMul[GMaxGlobalDistanceFieldClipmaps];
+
 		const int32 NumToSet = GlobalVolumeCenterAndExtent.GetNumBytes() / sizeof(FVector4);
 
 		for (int32 ClipmapIndex = 0; ClipmapIndex < NumToSet; ClipmapIndex++)
@@ -65,6 +66,11 @@ public:
 				const FVector WorldToUVAdd = (Clipmap.ScrollOffset - Clipmap.Bounds.GetCenter()) / (Clipmap.Bounds.GetExtent().X * 2) + FVector(.5f);
 				WorldToUVAddAndMul[ClipmapIndex] = FVector4(WorldToUVAdd, 1.0f / (Clipmap.Bounds.GetExtent().X * 2));
 			}
+			else
+			{
+				CenterAndExtent[ClipmapIndex] = FVector4(0);
+				WorldToUVAddAndMul[ClipmapIndex] = FVector4(0);
+			}
 		}
 
 		SetShaderValueArray(RHICmdList, ShaderRHI, GlobalVolumeCenterAndExtent, CenterAndExtent, GMaxGlobalDistanceFieldClipmaps);
diff --git a/Engine/Source/Runtime/Renderer/Private/SceneRendering.cpp b/Engine/Source/Runtime/Renderer/Private/SceneRendering.cpp
--- a/Engine/Source/Runtime/Renderer/Private/SceneRendering.cpp
+++ b/Engine/Source/Runtime/Renderer/Private/SceneRendering.cpp
@@ -640,7 +640,6 @@ TUniformBufferRef<FViewUniformShaderParameters> FViewInfo::CreateUniformBuffer(
 	if(State)
 	{
 		// safe to cast on the renderer side
-		FSceneViewState* ViewState = (FSceneViewState*)State;
 		ViewUniformShaderParameters.TemporalAAParams = FVector4(
 			ViewState->GetCurrentTemporalAASampleIndex(), 
 			ViewState->GetCurrentTemporalAASampleCount(),
@@ -946,7 +945,7 @@ void FViewInfo::InitRHIResources(const TArray<FProjectedShadowInfo*, SceneRender
 
 IPooledRenderTarget* FViewInfo::GetEyeAdaptation() const
 {
-	FSceneViewState* ViewState = (FSceneViewState*)State;
+	FSceneViewState* EffectiveViewState = ViewState;
 
 	// When rendering in stereo we want to use the same exposure for both eyes.
 	if (StereoPass == eSSP_RIGHT_EYE)
@@ -961,15 +960,15 @@ IPooledRenderTarget* FViewInfo::GetEyeAdaptation() const
 				const FSceneView* PrimaryView = Family->Views[ViewIndex];
 				if (PrimaryView->StereoPass == eSSP_LEFT_EYE)
 				{
-					ViewState = (FSceneViewState*)PrimaryView->State;
+					EffectiveViewState = (FSceneViewState*)PrimaryView->State;
 				}
 			}
 		}
 	}
 
-	if (ViewState)
+	if (EffectiveViewState)
 	{
-		TRefCountPtr<IPooledRenderTarget>& EyeAdaptRef = ViewState->GetEyeAdaptation();
+		TRefCountPtr<IPooledRenderTarget>& EyeAdaptRef = EffectiveViewState->GetEyeAdaptation();
 		if( IsValidRef(EyeAdaptRef) )
 		{
 			return EyeAdaptRef.GetReference();
