minor threshold tweak to get better blending of LPV settings (by [REDACTED_PERSON])

diff --git a/Engine/Source/Runtime/Renderer/Private/CompositionLighting/CompositionLighting.cpp b/Engine/Source/Runtime/Renderer/Private/CompositionLighting/CompositionLighting.cpp
--- a/Engine/Source/Runtime/Renderer/Private/CompositionLighting/CompositionLighting.cpp
+++ b/Engine/Source/Runtime/Renderer/Private/CompositionLighting/CompositionLighting.cpp
@@ -56,7 +56,7 @@ static bool IsLpvIndirectPassRequired(FPostprocessContext& Context)
 	{
 		FLightPropagationVolume* LightPropagationVolume = ViewState->GetLightPropagationVolume();
 
-		if(LightPropagationVolume && Context.View.FinalPostProcessSettings.LPVIntensity > 0.01f)
+		if(LightPropagationVolume && Context.View.FinalPostProcessSettings.LPVIntensity > 0.0f)
 		{
 			return true; 
 		}
