Fixed crash with DFAO in scene captures where the velocity pass is not supported

diff --git a/Engine/Source/Runtime/Renderer/Private/DistanceFieldLightingPost.cpp b/Engine/Source/Runtime/Renderer/Private/DistanceFieldLightingPost.cpp
--- a/Engine/Source/Runtime/Renderer/Private/DistanceFieldLightingPost.cpp
+++ b/Engine/Source/Runtime/Renderer/Private/DistanceFieldLightingPost.cpp
@@ -336,7 +336,7 @@ public:
 		TRefCountPtr<IPooledRenderTarget>* IrradianceHistoryRT, 
 		FSceneRenderTargetItem& DistanceFieldAOBentNormal, 
 		IPooledRenderTarget* DistanceFieldIrradiance,
-		FSceneRenderTargetItem& VelocityTextureValue)
+		IPooledRenderTarget* VelocityTextureValue)
 	{
 		const FPixelShaderRHIParamRef ShaderRHI = GetPixelShader();
 
@@ -397,7 +397,7 @@ public:
 			VelocityTexture,
 			VelocityTextureSampler,
 			TStaticSamplerState<SF_Bilinear,AM_Clamp,AM_Clamp,AM_Clamp>::GetRHI(),
-			VelocityTextureValue.ShaderResourceTexture
+			VelocityTextureValue ? VelocityTextureValue->GetRenderTargetItem().ShaderResourceTexture : GBlackTexture->TextureRHI
 			);
 	}
 	// FShader interface.
@@ -556,7 +556,7 @@ void UpdateHistory(
 	const FViewInfo& View, 
 	const TCHAR* BentNormalHistoryRTName,
 	const TCHAR* IrradianceHistoryRTName,
-	FSceneRenderTargetItem& VelocityTexture,
+	IPooledRenderTarget* VelocityTexture,
 	/** Contains last frame's history, if non-NULL.  This will be updated with the new frame's history. */
 	TRefCountPtr<IPooledRenderTarget>* BentNormalHistoryState,
 	TRefCountPtr<IPooledRenderTarget>* IrradianceHistoryState,
@@ -752,7 +752,7 @@ void PostProcessBentNormalAOSurfaceCache(
 	FRHICommandList& RHICmdList, 
 	const FDistanceFieldAOParameters& Parameters, 
 	const FViewInfo& View, 
-	FSceneRenderTargetItem& VelocityTexture,
+	IPooledRenderTarget* VelocityTexture,
 	FSceneRenderTargetItem& BentNormalInterpolation, 
 	IPooledRenderTarget* IrradianceInterpolation,
 	FSceneRenderTargetItem& DistanceFieldNormal,
diff --git a/Engine/Source/Runtime/Renderer/Private/DistanceFieldLightingPost.h b/Engine/Source/Runtime/Renderer/Private/DistanceFieldLightingPost.h
--- a/Engine/Source/Runtime/Renderer/Private/DistanceFieldLightingPost.h
+++ b/Engine/Source/Runtime/Renderer/Private/DistanceFieldLightingPost.h
@@ -13,7 +13,7 @@ extern void UpdateHistory(
 	const FViewInfo& View, 
 	const TCHAR* BentNormalHistoryRTName,
 	const TCHAR* IrradianceHistoryRTName,
-	FSceneRenderTargetItem& VelocityTexture,
+	IPooledRenderTarget* VelocityTexture,
 	/** Contains last frame's history, if non-NULL.  This will be updated with the new frame's history. */
 	TRefCountPtr<IPooledRenderTarget>* BentNormalHistoryState,
 	TRefCountPtr<IPooledRenderTarget>* IrradianceHistoryState,
@@ -28,7 +28,7 @@ extern void PostProcessBentNormalAOSurfaceCache(
 	FRHICommandList& RHICmdList, 
 	const FDistanceFieldAOParameters& Parameters, 
 	const FViewInfo& View, 
-	FSceneRenderTargetItem& VelocityTexture,
+	IPooledRenderTarget* VelocityTexture,
 	FSceneRenderTargetItem& BentNormalInterpolation, 
 	IPooledRenderTarget* IrradianceInterpolation,
 	FSceneRenderTargetItem& DistanceFieldNormal,
diff --git a/Engine/Source/Runtime/Renderer/Private/DistanceFieldScreenGridLighting.cpp b/Engine/Source/Runtime/Renderer/Private/DistanceFieldScreenGridLighting.cpp
--- a/Engine/Source/Runtime/Renderer/Private/DistanceFieldScreenGridLighting.cpp
+++ b/Engine/Source/Runtime/Renderer/Private/DistanceFieldScreenGridLighting.cpp
@@ -636,7 +636,7 @@ void PostProcessBentNormalAOScreenGrid(
 	FRHICommandList& RHICmdList, 
 	const FDistanceFieldAOParameters& Parameters, 
 	const FViewInfo& View, 
-	FSceneRenderTargetItem& VelocityTexture,
+	IPooledRenderTarget* VelocityTexture,
 	FSceneRenderTargetItem& BentNormalInterpolation, 
 	IPooledRenderTarget* IrradianceInterpolation,
 	FSceneRenderTargetItem& DistanceFieldNormal,
@@ -902,7 +902,7 @@ void FDeferredShadingSceneRenderer::RenderDistanceFieldAOScreenGrid(
 		RHICmdList, 
 		Parameters, 
 		View, 
-		VelocityTexture->GetRenderTargetItem(), 
+		VelocityTexture, 
 		DownsampledBentNormal->GetRenderTargetItem(), 
 		DownsampledIrradiance, 
 		DistanceFieldNormal->GetRenderTargetItem(), 
diff --git a/Engine/Source/Runtime/Renderer/Private/DistanceFieldSurfaceCacheLighting.cpp b/Engine/Source/Runtime/Renderer/Private/DistanceFieldSurfaceCacheLighting.cpp
--- a/Engine/Source/Runtime/Renderer/Private/DistanceFieldSurfaceCacheLighting.cpp
+++ b/Engine/Source/Runtime/Renderer/Private/DistanceFieldSurfaceCacheLighting.cpp
@@ -2880,7 +2880,7 @@ void RenderDistanceFieldAOSurfaceCache(
 		RHICmdList, 
 		Parameters, 
 		View, 
-		VelocityTexture->GetRenderTargetItem(), 
+		VelocityTexture, 
 		BentNormalAccumulation->GetRenderTargetItem(), 
 		IrradianceAccumulation, 
 		DistanceFieldNormal->GetRenderTargetItem(), 
diff --git a/Engine/Source/Runtime/Renderer/Private/VelocityRendering.cpp b/Engine/Source/Runtime/Renderer/Private/VelocityRendering.cpp
--- a/Engine/Source/Runtime/Renderer/Private/VelocityRendering.cpp
+++ b/Engine/Source/Runtime/Renderer/Private/VelocityRendering.cpp
@@ -718,7 +718,7 @@ bool FDeferredShadingSceneRenderer::ShouldRenderVelocities() const
 		bool bMotionBlur = IsMotionBlurEnabled(View);
 		bool bDistanceFieldAO = ShouldPrepareForDistanceFieldAO();
 
-		bNeedsVelocity |= bMotionBlur || bTemporalAA || bDistanceFieldAO;
+		bNeedsVelocity |= (bMotionBlur || bTemporalAA || bDistanceFieldAO) && !View.bIsSceneCapture;
 	}
 
 	return bNeedsVelocity;
