cleanup LightPropagationVolume code (from Lionhead Studios)

diff --git a/Engine/Shaders/ShadowDepthPixelShader.usf b/Engine/Shaders/ShadowDepthPixelShader.usf
--- a/Engine/Shaders/ShadowDepthPixelShader.usf
+++ b/Engine/Shaders/ShadowDepthPixelShader.usf
@@ -20,7 +20,6 @@
 	#include "LPVWriteVplCommon.usf"
 
 	#if RSM_TRANSMISSION
-		float TransmissionStrength;
 		float2 ReflectiveShadowMapTextureResolution;
 		row_major float4x4 ProjectionMatrix;
 	#endif
diff --git a/Engine/Source/Runtime/Renderer/Private/ShadowRendering.cpp b/Engine/Source/Runtime/Renderer/Private/ShadowRendering.cpp
--- a/Engine/Source/Runtime/Renderer/Private/ShadowRendering.cpp
+++ b/Engine/Source/Runtime/Renderer/Private/ShadowRendering.cpp
@@ -419,13 +419,12 @@ public:
 		: FMeshMaterialShader(Initializer)
 	{
 		ShadowParams.Bind(Initializer.ParameterMap,TEXT("ShadowParams"));
-		TransmissionStrength.Bind(Initializer.ParameterMap,TEXT("TransmissionStrength"));
 		ReflectiveShadowMapTextureResolution.Bind(Initializer.ParameterMap,TEXT("ReflectiveShadowMapTextureResolution"));
 		ProjectionMatrixParameter.Bind(Initializer.ParameterMap,TEXT("ProjectionMatrix"));
-		GvListBuffer.Bind(Initializer.ParameterMap,TEXT("GvListBuffer"));
-		GvListHeadBuffer.Bind(Initializer.ParameterMap,TEXT("GvListHeadBuffer"));
-		VplListBuffer.Bind(Initializer.ParameterMap,TEXT("VplListBuffer"));
-		VplListHeadBuffer.Bind(Initializer.ParameterMap,TEXT("VplListHeadBuffer"));
+		GvListBuffer.Bind(Initializer.ParameterMap,TEXT("RWGvListBuffer"));
+		GvListHeadBuffer.Bind(Initializer.ParameterMap,TEXT("RWGvListHeadBuffer"));
+		VplListBuffer.Bind(Initializer.ParameterMap,TEXT("RWVplListBuffer"));
+		VplListHeadBuffer.Bind(Initializer.ParameterMap,TEXT("RWVplListHeadBuffer"));
 	}
 
 	TShadowDepthBasePS() {}
@@ -447,7 +446,6 @@ public:
 		if(bRenderReflectiveShadowMap)
 		{
 			// LPV also propagates light transmission (for transmissive materials)
-			SetShaderValue(RHICmdList, ShaderRHI,TransmissionStrength,View.FinalPostProcessSettings.LPVTransmissionIntensity);
 			SetShaderValue(RHICmdList, ShaderRHI,ReflectiveShadowMapTextureResolution,FVector2D(GSceneRenderTargets.GetReflectiveShadowMapTextureResolution()));
 			SetShaderValue(
 				RHICmdList, 
@@ -480,7 +478,6 @@ public:
 
 		Ar << ShadowParams;
 
-		Ar << TransmissionStrength;
 		Ar << ReflectiveShadowMapTextureResolution;
 		Ar << ProjectionMatrixParameter;
 		Ar << GvListBuffer;
@@ -495,7 +492,6 @@ private:
 
 	FShaderParameter ShadowParams;
 
-	FShaderParameter TransmissionStrength;
 	FShaderParameter ReflectiveShadowMapTextureResolution;
 	FShaderParameter ProjectionMatrixParameter;
 
