fix cis

diff --git a/Engine/Source/Runtime/Renderer/Private/AtmosphereRendering.cpp b/Engine/Source/Runtime/Renderer/Private/AtmosphereRendering.cpp
--- a/Engine/Source/Runtime/Renderer/Private/AtmosphereRendering.cpp
+++ b/Engine/Source/Runtime/Renderer/Private/AtmosphereRendering.cpp
@@ -1601,11 +1601,12 @@ FAtmosphericFogSceneInfo::FAtmosphericFogSceneInfo(UAtmosphericFogComponent* InC
 	FMatrix WorldToLight = InComponent->ComponentToWorld.ToMatrixNoScale().InverseFast();
 	DefaultSunDirection = FVector(WorldToLight.M[0][0],WorldToLight.M[1][0],WorldToLight.M[2][0]);
 
+#if WITH_EDITORONLY_DATA
 	if (Component->PrecomputeCounter != UAtmosphericFogComponent::EValid)
 	{
 		bNeedRecompute = true;
 	}
-
+#endif
 	TransmittanceResource = Component->TransmittanceResource;
 	IrradianceResource = Component->IrradianceResource;
 	InscatterResource = Component->InscatterResource;
