Unity compile fixes

diff --git a/Engine/Source/Runtime/Renderer/Private/DistanceFieldLightingShared.h b/Engine/Source/Runtime/Renderer/Private/DistanceFieldLightingShared.h
--- a/Engine/Source/Runtime/Renderer/Private/DistanceFieldLightingShared.h
+++ b/Engine/Source/Runtime/Renderer/Private/DistanceFieldLightingShared.h
@@ -13,8 +13,6 @@ DECLARE_LOG_CATEGORY_EXTERN(LogDistanceField, Warning, All);
 /** Tile sized used for most AO compute shaders. */
 const int32 GDistanceFieldAOTileSizeX = 16;
 const int32 GDistanceFieldAOTileSizeY = 16;
-/** Base downsample factor that all distance field AO operations are done at. */
-const int32 GAODownsampleFactor = 2;
 /** Must match usf */
 static const int32 GMaxNumObjectsPerTile = 512;
 
@@ -34,8 +32,6 @@ inline bool SupportsDistanceFieldGI(ERHIFeatureLevel::Type FeatureLevel, EShader
 
 extern bool IsDistanceFieldGIAllowed(const FViewInfo& View);
 
-extern FIntPoint GetBufferSizeForAO();
-
 class FDistanceFieldObjectBuffers
 {
 public:
diff --git a/Engine/Source/Runtime/Renderer/Private/DistanceFieldSurfaceCacheLighting.h b/Engine/Source/Runtime/Renderer/Private/DistanceFieldSurfaceCacheLighting.h
--- a/Engine/Source/Runtime/Renderer/Private/DistanceFieldSurfaceCacheLighting.h
+++ b/Engine/Source/Runtime/Renderer/Private/DistanceFieldSurfaceCacheLighting.h
@@ -11,6 +11,10 @@
 const static int32 GAOMaxSupportedLevel = 6;
 /** Number of cone traced directions. */
 const int32 NumConeSampleDirections = 9;
+
+/** Base downsample factor that all distance field AO operations are done at. */
+const int32 GAODownsampleFactor = 2;
+
 extern const uint32 UpdateObjectsGroupSize;
 
 const float GDefaultDFAOMaxOcclusionDistance = 600.0f;
@@ -20,6 +24,8 @@ inline bool DoesPlatformSupportDistanceFieldAO(EShaderPlatform Platform)
 	return Platform == SP_PCD3D_SM5 || Platform == SP_PS4;
 }
 
+extern FIntPoint GetBufferSizeForAO();
+
 class FDistanceFieldAOParameters
 {
 public:
