UE-13256: Fixes SSAO flicker on small radius by building HZB only by downsampling per 2.

diff --git a/Engine/Config/BaseScalability.ini b/Engine/Config/BaseScalability.ini
--- a/Engine/Config/BaseScalability.ini
+++ b/Engine/Config/BaseScalability.ini
@@ -101,6 +101,7 @@ r.DistanceFieldAO=1
 
 [PostProcessQuality@0]
 r.MotionBlurQuality=0
+r.AmbientOcclusionMipLevelFactor=1.0
 r.AmbientOcclusionSampleSetQuality=0
 r.AmbientOcclusionLevels=0
 r.AmbientOcclusionRadiusScale=1.2
@@ -117,6 +118,7 @@ r.LightShaftQuality=0
 
 [PostProcessQuality@1]
 r.MotionBlurQuality=3
+r.AmbientOcclusionMipLevelFactor=1.0
 r.AmbientOcclusionSampleSetQuality=0
 r.AmbientOcclusionLevels=1
 r.AmbientOcclusionRadiusScale=1.5
@@ -133,6 +135,7 @@ r.LightShaftQuality=0
 
 [PostProcessQuality@2]
 r.MotionBlurQuality=3
+r.AmbientOcclusionMipLevelFactor=0.6
 r.AmbientOcclusionSampleSetQuality=-1
 r.AmbientOcclusionLevels=1
 r.AmbientOcclusionRadiusScale=1.5
@@ -149,6 +152,7 @@ r.LightShaftQuality=1
 
 [PostProcessQuality@3]
 r.MotionBlurQuality=4
+r.AmbientOcclusionMipLevelFactor=0.4
 r.AmbientOcclusionSampleSetQuality=-1
 r.AmbientOcclusionLevels=2
 r.AmbientOcclusionRadiusScale=1.0
diff --git a/Engine/Shaders/HZBOcclusion.usf b/Engine/Shaders/HZBOcclusion.usf
--- a/Engine/Shaders/HZBOcclusion.usf
+++ b/Engine/Shaders/HZBOcclusion.usf
@@ -3,40 +3,47 @@
 #include "Common.usf"
 
 float2 InvSize;
+float4 InputUvFactorAndOffset;
+float2 InputViewportMaxBound;
 
 Texture2D		Texture;
 SamplerState	TextureSampler;
 
-void HZBBuildPS(float2 InUV : TEXCOORD0, out float4 OutColor : SV_Target0)
+void HZBBuildPS(float4 RenderTargetUVAndScreenPos : TEXCOORD0, out float4 OutColor : SV_Target0)
 {
+	float4 Depth;
 	float2 UV[4];
 
-	UV[0] = InUV + float2(-0.25f, -0.25f) * InvSize;
-	UV[1] = InUV + float2( 0.25f, -0.25f) * InvSize;
-	UV[2] = InUV + float2(-0.25f,  0.25f) * InvSize;
-	UV[3] = InUV + float2( 0.25f,  0.25f) * InvSize;
-
 #if STAGE == 0
-	// Inverted Z buffer
-	float4 Depth;
+	const float2 ViewUV = RenderTargetUVAndScreenPos.zw * float2(0.5f, -0.5f) + 0.5f;
+	const float2 InUV = ViewUV * InputUvFactorAndOffset.xy + InputUvFactorAndOffset.zw;
+	
+	// min(..., InputViewportMaxBound) because we don't want to sample outside of the viewport
+	// when the view size has odd dimensions on X/Y axis.
+	UV[0] = min(InUV + float2(-0.25f, -0.25f) * InvSize, InputViewportMaxBound);
+	UV[1] = min(InUV + float2( 0.25f, -0.25f) * InvSize, InputViewportMaxBound);
+	UV[2] = min(InUV + float2(-0.25f,  0.25f) * InvSize, InputViewportMaxBound);
+	UV[3] = min(InUV + float2( 0.25f,  0.25f) * InvSize, InputViewportMaxBound);
+
 	Depth.x = SceneDepthTexture.SampleLevel( SceneDepthTextureSampler, UV[0], 0 ).r;
 	Depth.y = SceneDepthTexture.SampleLevel( SceneDepthTextureSampler, UV[1], 0 ).r;
 	Depth.z = SceneDepthTexture.SampleLevel( SceneDepthTextureSampler, UV[2], 0 ).r;
 	Depth.w = SceneDepthTexture.SampleLevel( SceneDepthTextureSampler, UV[3], 0 ).r;
+#else
+	const float2 InUV = RenderTargetUVAndScreenPos.xy;
 	
-	// Max to fill holes
-	//OutColor = max( max(Depth.x, Depth.y), max(Depth.z, Depth.w) );
+	UV[0] = InUV + float2(-0.25f, -0.25f) * InvSize;
+	UV[1] = InUV + float2( 0.25f, -0.25f) * InvSize;
+	UV[2] = InUV + float2(-0.25f,  0.25f) * InvSize;
+	UV[3] = InUV + float2( 0.25f,  0.25f) * InvSize;
 
-	// Min is conservative
-	OutColor = min( min(Depth.x, Depth.y), min(Depth.z, Depth.w) );
-#else
-	float4 Depth;	
 	Depth.x = Texture.SampleLevel( TextureSampler, UV[0], 0 ).r;
 	Depth.y = Texture.SampleLevel( TextureSampler, UV[1], 0 ).r;
 	Depth.z = Texture.SampleLevel( TextureSampler, UV[2], 0 ).r;
 	Depth.w = Texture.SampleLevel( TextureSampler, UV[3], 0 ).r;
-	OutColor = min( min(Depth.x, Depth.y), min(Depth.z, Depth.w) );
 #endif
+
+	OutColor = min( min(Depth.x, Depth.y), min(Depth.z, Depth.w) );
 }
 
 Texture2D		HZBTexture;
@@ -47,8 +54,8 @@ SamplerState	BoundsCenterSampler;
 Texture2D		BoundsExtentTexture;
 SamplerState	BoundsExtentSampler;
 
-// TODO pass as uniform
-static const float4 HZBSize = { 512, 256, 1.0 / 512.0, 1.0 / 256.0 };
+float3 HZBUvFactor;
+float4 HZBSize;
 
 void HZBTestPS(float2 InUV : TEXCOORD0, out float4 OutColor : SV_Target0)
 {
@@ -102,7 +109,7 @@ void HZBTestPS(float2 InUV : TEXCOORD0, out float4 OutColor : SV_Target0)
 	float4 Rect = saturate( float4( RectMin.xy, RectMax.xy ) * float2( 0.5, -0.5 ).xyxy + 0.5 ).xwzy;
 	float4 RectPixels = Rect * HZBSize.xyxy;
 	float2 RectSize = ( RectPixels.zw - RectPixels.xy ) * 0.5;	// 0.5 for 4x4
-	float Level = ceil( log2( max( RectSize.x, RectSize.y ) ) );
+	float Level = max(ceil( log2( max( RectSize.x, RectSize.y ) ) ), HZBUvFactor.z);
 
 	// Check if we can drop one level lower
 	float LevelLower = max( Level - 1, 0 );
@@ -227,8 +234,8 @@ void HZBTestPS(float2 InUV : TEXCOORD0, out float4 OutColor : SV_Target0)
 	OutColor = Visible ? 1 : 0;
 #else
 	// 4x4 samples
-	float2 Scale = ( Rect.zw - Rect.xy ) / 3;
-	float2 Bias = Rect.xy;
+	float2 Scale = HZBUvFactor.xy * ( Rect.zw - Rect.xy ) / 3;
+	float2 Bias = HZBUvFactor.xy * Rect.xy;
 
 	float4 MinDepth = 1;
 	UNROLL for( int i = 0; i < 4; i++ )
diff --git a/Engine/Shaders/PostProcessAmbientOcclusion.usf b/Engine/Shaders/PostProcessAmbientOcclusion.usf
--- a/Engine/Shaders/PostProcessAmbientOcclusion.usf
+++ b/Engine/Shaders/PostProcessAmbientOcclusion.usf
@@ -130,6 +130,7 @@ float4 NoiseScale;
 Texture2D RandomNormalTexture;
 SamplerState RandomNormalTextureSampler;
 
+float3 HzbUvAndStepMipLevelFactor;
 
 // could be moved to a more central spot
 // @param ScreenPos -1 .. 1
@@ -337,10 +338,10 @@ float2 Wedge(float Scale, float2 LocalRandom, float3 FovFix, float3 InvFovFix, f
 
 	float2 ScreenSpacePos = ReconstructSSPosFromCS(SamplePosition2 * FovFix);
 
-#if USE_HZB
+#if USE_estmap 
 	float2 UV = ScreenSpacePos * float2(0.5f, -0.5f) + 0.5f;
 
-	float Abs = ConvertFromDeviceZ(Texture2DSampleLevel( PostprocessInput3, PostprocessInput3Sampler, UV, MipLevel).r);
+	float Abs = ConvertFromDeviceZ(Texture2DSampleLevel( PostprocessInput3, PostprocessInput3Sampler, HzbUvAndStepMipLevelFactor.xy * UV, MipLevel).r);
 #else
 	float2 UV = ScreenPosToUV(ScreenSpacePos, PostprocessInput0Size);
 
@@ -385,7 +386,7 @@ float2 Wedge2(float Scale, float2 LocalRandom, float3 FovFix, float3 InvFovFix,
 	#if USE_HZB
 		float2 UV = ScreenSpacePos * float2(0.5f, -0.5f) + 0.5f;
 
-		AbsLeft = ConvertFromDeviceZ(Texture2DSampleLevel( PostprocessInput3, PostprocessInput3Sampler, UV, MipLevel).r);
+		AbsLeft = ConvertFromDeviceZ(Texture2DSampleLevel( PostprocessInput3, PostprocessInput3Sampler, HzbUvAndStepMipLevelFactor.xy * UV, MipLevel).r);
 	#else
 		float2 UV = ScreenPosToUV(ScreenSpacePos, PostprocessInput0Size);
 
@@ -418,7 +419,7 @@ float2 Wedge2(float Scale, float2 LocalRandom, float3 FovFix, float3 InvFovFix,
 	#if USE_HZB
 		float2 UV = ScreenSpacePos * float2(0.5f, -0.5f) + 0.5f;
 
-		AbsRight = ConvertFromDeviceZ(Texture2DSampleLevel( PostprocessInput3, PostprocessInput3Sampler, UV, MipLevel).r);
+		AbsLeft = ConvertFromDeviceZ(Texture2DSampleLevel( PostprocessInput3, PostprocessInput3Sampler, HzbUvAndStepMipLevelFactor.xy * UV, MipLevel).r);
 	#else
 		float2 UV = ScreenPosToUV(ScreenSpacePos, PostprocessInput0Size);
 
@@ -559,7 +560,7 @@ void MainPS(in float4 UVAndScreenPos : TEXCOORD0, out float4 OutColor : SV_Targe
 					// constant at run time
 					float Scale = (step + 1) / (float)SAMPLE_STEPS;
 					// constant at run time (higher is better for texture cache / performance, lower is better quality
-					float MipLevel = max(step, (uint)1);
+					float MipLevel = log2(1.0f + HzbUvAndStepMipLevelFactor.z * step);
 
 					float2 StepSample = Wedge2(Scale, LocalRandom, FovFix, InvFovFix, ViewSpacePosition, ScaledViewSpaceNormal, InvHaloSize, MipLevel);
 
@@ -581,7 +582,7 @@ void MainPS(in float4 UVAndScreenPos : TEXCOORD0, out float4 OutColor : SV_Targe
 						// constant at run time
 						float Scale = (step + 1) / (float)SAMPLE_STEPS * mirror;
 						// constant at run time (higher is better for texture cache / performance, lower is better quality
-						float MipLevel = max(step, (uint)1);
+						float MipLevel = log2(1.0f + HzbUvAndStepMipLevelFactor.z * step);
 
 						float2 StepSample = Wedge(Scale, LocalRandom, FovFix, InvFovFix, ViewSpacePosition, ScaledViewSpaceNormal, InvHaloSize, MipLevel);
 
diff --git a/Engine/Shaders/ScreenSpaceRayCast.usf b/Engine/Shaders/ScreenSpaceRayCast.usf
--- a/Engine/Shaders/ScreenSpaceRayCast.usf
+++ b/Engine/Shaders/ScreenSpaceRayCast.usf
@@ -11,10 +11,20 @@
 #define VECTORIZED_BRANCHLESS	0
 #define VECTORIZED_EARLY_OUT	1
 
+#if USE_HZB
+#define HZB_LEVEL_OFFSET 1.0
+
+float4 HZBUvFactorAndInvFactor;
+#else
+#define HZB_LEVEL_OFFSET 0.0
+
+#endif
+
 float4 SampleDepthTexture( Texture2D Texture, SamplerState Sampler, float Level, float4 SampleUV0, float4 SampleUV1 )
 {
 	float4 SampleDepth;
 #if USE_HZB
+	// SampleUV{0,1}.{xy,zw} should already be in the HZB UV frame using HZBUvFactorAndInvFactor.xy
 	SampleDepth.x = Texture2DSampleLevel( Texture, Sampler, SampleUV0.xy, Level ).r;
 	SampleDepth.y = Texture2DSampleLevel( Texture, Sampler, SampleUV0.zw, Level ).r;
 	SampleDepth.z = Texture2DSampleLevel( Texture, Sampler, SampleUV1.xy, Level ).r;
@@ -40,14 +50,25 @@ float4 RayCast( Texture2D Texture, SamplerState Sampler, float3 R, float Roughne
 
 	float3 RayStartScreen = RayStartClip.xyz / RayStartClip.w;
 	float3 RayEndScreen = RayEndClip.xyz / RayEndClip.w;
+	
+	// Make sure the RayStepScreen has a length of 2 in order to have the end point ouside of the viewport
+	float3 RayStepScreen = ( RayEndScreen - RayStartScreen ) / (0.5 * length( RayEndScreen.xy - RayStartScreen.xy ));
 
-	// Normalize 2D length
-	float3 RayStepScreen = ( RayEndScreen - RayStartScreen ) / length( RayEndScreen.xy - RayStartScreen.xy );
-	RayStepScreen *= 1.5;
+	{
+		// Computes the scale down factor for RayStepScreen required to fit on the X and Y axis in order to clip it in the viewport
+		float2 AbsRayStepScreen = abs(RayStepScreen);
+		float2 S = (AbsRayStepScreen - max(abs(RayStepScreen.xy + RayStartScreen.xy) - 1.0f, 0.0f)) / AbsRayStepScreen;
+
+		// Rescales RayStepScreen accordingly
+		RayStepScreen *= min(S.x, S.y);
+	}
 
 #if USE_HZB
 	float3 RayStartUVz = float3( RayStartScreen.xy * float2( 0.5, -0.5 ) + 0.5, RayStartScreen.z );
 	float3 RayStepUVz = float3( RayStepScreen.xy * float2( 0.5, -0.5 ), RayStepScreen.z );
+
+	RayStartUVz.xy *= HZBUvFactorAndInvFactor.xy;
+	RayStepUVz.xy *= HZBUvFactorAndInvFactor.xy;
 #else
 	float3 RayStartUVz = float3( RayStartScreen.xy * View.ScreenPositionScaleBias.xy + View.ScreenPositionScaleBias.wz, RayStartScreen.z );
 	float3 RayStepUVz = float3( RayStepScreen.xy * View.ScreenPositionScaleBias.xy, RayStepScreen.z );
@@ -74,7 +95,7 @@ float4 RayCast( Texture2D Texture, SamplerState Sampler, float3 R, float Roughne
 		float3 SampleUVz = RayStartUVz + RayStepUVz * SampleTime;
 		
 		// Use lower res for farther samples
-		float Level = Roughness * (i * 4.0 / NumSteps);
+		float Level = Roughness * (i * 4.0 / NumSteps) + HZB_LEVEL_OFFSET;
 		float SampleDepth = Texture.SampleLevel( Sampler, SampleUVz.xy, Level ).r;
 
 		float DepthDiff = SampleUVz.z - SampleDepth;
@@ -92,17 +113,14 @@ float4 RayCast( Texture2D Texture, SamplerState Sampler, float3 R, float Roughne
 	}
 
 	float3 HitUVz = RayStartUVz + RayStepUVz * MinHitTime;
-#if USE_HZB
-	HitUVz.xy = HitUVz.xy * float2( 2, -2 ) + float2( -1, 1 );
-	HitUVz.xy = HitUVz.xy * View.ScreenPositionScaleBias.xy + View.ScreenPositionScaleBias.wz;
-#endif
+
 	Result = float4( HitUVz, MinHitTime );
 
 #elif VECTORIZED_BRANCHLESS
 
 	float MinHitTime = 1;
 	float LastDiff = 0;
-	float Level = 0;
+	float Level = HZB_LEVEL_OFFSET;
 	
 	// Vectorized to group fetches
 	float4 SampleTime = ( StepOffset + float4( 1, 2, 3, 4 ) ) * Step;
@@ -155,16 +173,13 @@ float4 RayCast( Texture2D Texture, SamplerState Sampler, float3 R, float Roughne
 	}
 
 	float3 HitUVz = RayStartUVz + RayStepUVz * MinHitTime;
-#if USE_HZB
-	HitUVz.xy = HitUVz.xy = HitUVz.xy * float2( 2, -2 ) + float2( -1, 1 );
-	HitUVz.xy = HitUVz.xy * View.ScreenPositionScaleBias.xy + View.ScreenPositionScaleBias.wz;
-#endif
+
 	Result = float4( HitUVz, MinHitTime );
 
 #else // VECTORIZED_EARLY_OUT
 
 	float LastDiff = 0;
-	float Level = 0;
+	float Level = HZB_LEVEL_OFFSET;
 
 	float4 SampleTime = ( StepOffset + float4( 1, 2, 3, 4 ) ) * Step;
 
@@ -195,10 +210,7 @@ float4 RayCast( Texture2D Texture, SamplerState Sampler, float3 R, float Roughne
 			float MinHitTime = min( HitTime.x, HitTime.y );
 			
 			float3 HitUVz = RayStartUVz + RayStepUVz * MinHitTime;
-#if USE_HZB
-			HitUVz.xy = HitUVz.xy * float2( 2, -2 ) + float2( -1, 1 );
-			HitUVz.xy = HitUVz.xy * View.ScreenPositionScaleBias.xy + View.ScreenPositionScaleBias.wz;
-#endif
+
 			Result = float4( HitUVz, MinHitTime );
 			break;
 		}
@@ -209,6 +221,12 @@ float4 RayCast( Texture2D Texture, SamplerState Sampler, float3 R, float Roughne
 		SampleTime += 4.0 / (NumSteps + 1);
 	}
 
+#endif
+	
+#if USE_HZB
+	Result.xy *= HZBUvFactorAndInvFactor.zw;
+	Result.xy = Result.xy * float2( 2, -2 ) + float2( -1, 1 );
+	Result.xy = Result.xy * View.ScreenPositionScaleBias.xy + View.ScreenPositionScaleBias.wz;
 #endif
 
 	return Result;
diff --git a/Engine/Shaders/ShaderVersion.usf b/Engine/Shaders/ShaderVersion.usf
--- a/Engine/Shaders/ShaderVersion.usf
+++ b/Engine/Shaders/ShaderVersion.usf
@@ -2,4 +2,4 @@
 // Each time the console command is executed it generates a new GUID. As this file is included
 // in common.usf (which should be included in any shader) it allows to invalidate the shader DDC.
 // 
-// GUID = 7D700E574A85DE267C385C9E997BCC26
\ No newline at end of file
+// GUID = 8636091346FB77EB0D592BB4EAD2CD9D
\ No newline at end of file
diff --git a/Engine/Source/Runtime/Renderer/Private/CompositionLighting/PostProcessAmbientOcclusion.cpp b/Engine/Source/Runtime/Renderer/Private/CompositionLighting/PostProcessAmbientOcclusion.cpp
--- a/Engine/Source/Runtime/Renderer/Private/CompositionLighting/PostProcessAmbientOcclusion.cpp
+++ b/Engine/Source/Runtime/Renderer/Private/CompositionLighting/PostProcessAmbientOcclusion.cpp
@@ -20,6 +20,15 @@ static TAutoConsoleVariable<int32> CVarAmbientOcclusionSampleSetQuality(
 	TEXT("1: high sample count (defined in shader, 6 * 2 per pixel)"),
 	ECVF_Scalability | ECVF_RenderThreadSafe);
 
+static TAutoConsoleVariable<float> CVarAmbientOcclusionStepMipLevelFactor(
+	TEXT("r.AmbientOcclusionMipLevelFactor"),
+	0.5f,
+	TEXT("Controle mipmap level according to the SSAO step id\n")
+	TEXT("0: always look into the HZB mipmap level 0 (memory cache trashing)\n")
+	TEXT("0.5: sample count depends on post process settings (default)\n")
+	TEXT("1: Go into higher mipmap level (quality loss)"),
+	ECVF_Scalability | ECVF_RenderThreadSafe);
+
 IMPLEMENT_UNIFORM_BUFFER_STRUCT(FCameraMotionParameters,TEXT("CameraMotion"));
 
 /** Encapsulates the post processing ambient occlusion pixel shader. */
@@ -141,6 +150,7 @@ class FPostProcessAmbientOcclusionPS : public FGlobalShader
 	FPostProcessAmbientOcclusionPS() {}
 
 public:
+	FShaderParameter HzbUvAndStepMipLevelFactor;
 	FPostProcessPassParameters PostprocessParameter;
 	FDeferredPixelShaderParameters DeferredParameters;
 	FScreenSpaceAOandSSRShaderParameters ScreenSpaceAOandSSRShaderParams;
@@ -156,6 +166,7 @@ public:
 		ScreenSpaceAOandSSRShaderParams.Bind(Initializer.ParameterMap);
 		RandomNormalTexture.Bind(Initializer.ParameterMap, TEXT("RandomNormalTexture"));
 		RandomNormalTextureSampler.Bind(Initializer.ParameterMap, TEXT("RandomNormalTextureSampler"));
+		HzbUvAndStepMipLevelFactor.Bind(Initializer.ParameterMap, TEXT("HzbUvAndStepMipLevelFactor"));
 	}
 
 	void SetParameters(const FRenderingCompositePassContext& Context, FIntPoint InputTextureSize)
@@ -174,13 +185,22 @@ public:
 		SetTextureParameter(Context.RHICmdList, ShaderRHI, RandomNormalTexture, RandomNormalTextureSampler, TStaticSamplerState<SF_Point, AM_Wrap, AM_Wrap, AM_Wrap>::GetRHI(), SSAORandomization.ShaderResourceTexture);
 
 		ScreenSpaceAOandSSRShaderParams.Set(Context.RHICmdList, Context.View, ShaderRHI, InputTextureSize);
+		
+		const float HzbStepMipLevelFactorValue = FMath::Clamp(CVarAmbientOcclusionStepMipLevelFactor.GetValueOnRenderThread(), 0.0f, 1.0f);
+		const FVector HzbUvAndStepMipLevelFactorValue(
+			float(Context.View.ViewRect.Width()) / float(2 * Context.View.HZBMipmap0Size.X),
+			float(Context.View.ViewRect.Height()) / float(2 * Context.View.HZBMipmap0Size.Y),
+			HzbStepMipLevelFactorValue
+			);
+		
+		SetShaderValue(Context.RHICmdList, ShaderRHI, HzbUvAndStepMipLevelFactor, HzbStepMipLevelFactorValue );
 	}
 	
 	// FShader interface.
 	virtual bool Serialize(FArchive& Ar) override
 	{
 		bool bShaderHasOutdatedParameters = FGlobalShader::Serialize(Ar);
-		Ar << PostprocessParameter << DeferredParameters << ScreenSpaceAOandSSRShaderParams << RandomNormalTexture << RandomNormalTextureSampler;
+		Ar << HzbUvAndStepMipLevelFactor << PostprocessParameter << DeferredParameters << ScreenSpaceAOandSSRShaderParams << RandomNormalTexture << RandomNormalTextureSampler;
 		return bShaderHasOutdatedParameters;
 	}
 
diff --git a/Engine/Source/Runtime/Renderer/Private/PostProcess/ScreenSpaceReflections.cpp b/Engine/Source/Runtime/Renderer/Private/PostProcess/ScreenSpaceReflections.cpp
--- a/Engine/Source/Runtime/Renderer/Private/PostProcess/ScreenSpaceReflections.cpp
+++ b/Engine/Source/Runtime/Renderer/Private/PostProcess/ScreenSpaceReflections.cpp
@@ -89,6 +89,7 @@ public:
 	FPostProcessPassParameters PostprocessParameter;
 	FDeferredPixelShaderParameters DeferredParameters;
 	FShaderParameter SSRParams;
+	FShaderParameter HZBUvFactorAndInvFactor;
 
 	/** Initialization constructor. */
 	FPostProcessScreenSpaceReflectionsPS(const ShaderMetaType::CompiledShaderInitializerType& Initializer)
@@ -97,6 +98,7 @@ public:
 		PostprocessParameter.Bind(Initializer.ParameterMap);
 		DeferredParameters.Bind(Initializer.ParameterMap);
 		SSRParams.Bind(Initializer.ParameterMap, TEXT("SSRParams"));
+		HZBUvFactorAndInvFactor.Bind(Initializer.ParameterMap, TEXT("HZBUvFactorAndInvFactor"));
 	}
 
 	void SetParameters(const FRenderingCompositePassContext& Context)
@@ -127,13 +129,28 @@ public:
 
 			SetShaderValue(Context.RHICmdList, ShaderRHI, SSRParams, Value);
 		}
+
+		{
+			const FVector2D HZBUvFactor(
+				float(Context.View.ViewRect.Width()) / float(2 * Context.View.HZBMipmap0Size.X),
+				float(Context.View.ViewRect.Height()) / float(2 * Context.View.HZBMipmap0Size.Y)
+				);
+			const FVector4 HZBUvFactorAndInvFactorValue(
+				HZBUvFactor.X,
+				HZBUvFactor.Y,
+				1.0f / HZBUvFactor.X,
+				1.0f / HZBUvFactor.Y
+				);
+			
+			SetShaderValue(Context.RHICmdList, ShaderRHI, HZBUvFactorAndInvFactor, HZBUvFactorAndInvFactorValue);
+		}
 	}
 
 	// FShader interface.
 	virtual bool Serialize(FArchive& Ar) override
 	{
 		bool bShaderHasOutdatedParameters = FGlobalShader::Serialize(Ar);
-		Ar << PostprocessParameter << DeferredParameters << SSRParams;
+		Ar << PostprocessParameter << DeferredParameters << SSRParams << HZBUvFactorAndInvFactor;
 		return bShaderHasOutdatedParameters;
 	}
 };
diff --git a/Engine/Source/Runtime/Renderer/Private/SceneOcclusion.cpp b/Engine/Source/Runtime/Renderer/Private/SceneOcclusion.cpp
--- a/Engine/Source/Runtime/Renderer/Private/SceneOcclusion.cpp
+++ b/Engine/Source/Runtime/Renderer/Private/SceneOcclusion.cpp
@@ -625,7 +625,8 @@ class FHZBTestPS : public FGlobalShader
 	FHZBTestPS() {}
 
 public:
-	FShaderParameter				InvSizeParameter;
+	FShaderParameter				HZBUvFactor;
+	FShaderParameter				HZBSize;
 	FShaderResourceParameter		HZBTexture;
 	FShaderResourceParameter		HZBSampler;
 	FShaderResourceParameter		BoundsCenterTexture;
@@ -636,7 +637,8 @@ public:
 	FHZBTestPS(const ShaderMetaType::CompiledShaderInitializerType& Initializer)
 		: FGlobalShader(Initializer)
 	{
-		InvSizeParameter.Bind( Initializer.ParameterMap, TEXT("InvSize") );
+		HZBUvFactor.Bind( Initializer.ParameterMap, TEXT("HZBUvFactor") );
+		HZBSize.Bind( Initializer.ParameterMap, TEXT("HZBSize") );
 		HZBTexture.Bind( Initializer.ParameterMap, TEXT("HZBTexture") );
 		HZBSampler.Bind( Initializer.ParameterMap, TEXT("HZBSampler") );
 		BoundsCenterTexture.Bind( Initializer.ParameterMap, TEXT("BoundsCenterTexture") );
@@ -651,8 +653,26 @@ public:
 
 		FGlobalShader::SetParameters(RHICmdList, ShaderRHI, View );
 
-		const FVector2D InvSize( 1.0f / 512.0f, 1.0f / 256.0f );
-		SetShaderValue(RHICmdList, ShaderRHI, InvSizeParameter, InvSize );
+		/*
+		 * Defines the maximum number of mipmaps the HZB test is considering
+		 * to avoid memory cache trashing when rendering on high resolution.
+		 */
+		const float kHZBTestMaxMipmap = 9.0f;
+
+		const float HZBMipmapCounts = FMath::Log2(FMath::Max(View.HZBMipmap0Size.X, View.HZBMipmap0Size.Y));
+		const FVector HZBUvFactorValue(
+			float(View.ViewRect.Width()) / float(2 * View.HZBMipmap0Size.X),
+			float(View.ViewRect.Height()) / float(2 * View.HZBMipmap0Size.Y),
+			FMath::Max(HZBMipmapCounts - kHZBTestMaxMipmap, 0.0f)
+			);
+		const FVector4 HZBSizeValue(
+			View.HZBMipmap0Size.X,
+			View.HZBMipmap0Size.Y,
+			1.0f / float(View.HZBMipmap0Size.X),
+			1.0f / float(View.HZBMipmap0Size.Y)
+			);
+		SetShaderValue(RHICmdList, ShaderRHI, HZBUvFactor, HZBUvFactorValue);
+		SetShaderValue(RHICmdList, ShaderRHI, HZBSize, HZBSizeValue);
 
 		SetTextureParameter(RHICmdList, ShaderRHI, HZBTexture, HZBSampler, TStaticSamplerState<SF_Point,AM_Clamp,AM_Clamp,AM_Clamp>::GetRHI(), View.HZB->GetRenderTargetItem().ShaderResourceTexture );
 
@@ -663,7 +683,8 @@ public:
 	virtual bool Serialize(FArchive& Ar) override
 	{
 		bool bShaderHasOutdatedParameters = FGlobalShader::Serialize(Ar);
-		Ar << InvSizeParameter;
+		Ar << HZBUvFactor;
+		Ar << HZBSize;
 		Ar << HZBTexture;
 		Ar << HZBSampler;
 		Ar << BoundsCenterTexture;
@@ -876,6 +897,8 @@ class THZBBuildPS : public FGlobalShader
 
 public:
 	FShaderParameter				InvSizeParameter;
+	FShaderParameter				InputUvFactorAndOffsetParameter;
+	FShaderParameter				InputViewportMaxBoundParameter;
 	FSceneTextureShaderParameters	SceneTextureParameters;
 	FShaderResourceParameter		TextureParameter;
 	FShaderResourceParameter		TextureParameterSampler;
@@ -884,19 +907,34 @@ public:
 		: FGlobalShader(Initializer)
 	{
 		InvSizeParameter.Bind( Initializer.ParameterMap, TEXT("InvSize") );
+		InputUvFactorAndOffsetParameter.Bind( Initializer.ParameterMap, TEXT("InputUvFactorAndOffset") );
+		InputViewportMaxBoundParameter.Bind( Initializer.ParameterMap, TEXT("InputViewportMaxBound") );
 		SceneTextureParameters.Bind( Initializer.ParameterMap );
 		TextureParameter.Bind( Initializer.ParameterMap, TEXT("Texture") );
 		TextureParameterSampler.Bind( Initializer.ParameterMap, TEXT("TextureSampler") );
 	}
 
-	void SetParameters(FRHICommandList& RHICmdList, const FSceneView& View, const FIntPoint& Size )
+	void SetParameters(FRHICommandList& RHICmdList, const FViewInfo& View)
 	{
 		const FPixelShaderRHIParamRef ShaderRHI = GetPixelShader();
 
 		FGlobalShader::SetParameters(RHICmdList, ShaderRHI, View );
-
-		const FVector2D InvSize( 1.0f / Size.X, 1.0f / Size.Y );
+		
+		const FIntPoint GBufferSize = GSceneRenderTargets.GetBufferSizeXY();
+		const FVector2D InvSize( 1.0f / float(GBufferSize.X), 1.0f / float(GBufferSize.Y) );
+		const FVector4 InputUvFactorAndOffset (
+			float(2 * View.HZBMipmap0Size.X) / float(GBufferSize.X),
+			float(2 * View.HZBMipmap0Size.Y) / float(GBufferSize.Y),
+			float(View.ViewRect.Min.X) / float(GBufferSize.X),
+			float(View.ViewRect.Min.Y) / float(GBufferSize.Y)
+			);
+		const FVector2D InputViewportMaxBound (
+			float(View.ViewRect.Max.X) / float(GBufferSize.X) - 0.5f * InvSize.X,
+			float(View.ViewRect.Max.Y) / float(GBufferSize.Y) - 0.5f * InvSize.Y
+			);
 		SetShaderValue(RHICmdList, ShaderRHI, InvSizeParameter, InvSize );
+		SetShaderValue(RHICmdList, ShaderRHI, InputUvFactorAndOffsetParameter, InputUvFactorAndOffset );
+		SetShaderValue(RHICmdList, ShaderRHI, InputViewportMaxBoundParameter, InputViewportMaxBound );
 		
 		SceneTextureParameters.Set(RHICmdList, ShaderRHI, View );
 	}
@@ -920,6 +958,8 @@ public:
 	{
 		bool bShaderHasOutdatedParameters = FGlobalShader::Serialize(Ar);
 		Ar << InvSizeParameter;
+		Ar << InputUvFactorAndOffsetParameter;
+		Ar << InputViewportMaxBoundParameter;
 		Ar << SceneTextureParameters;
 		Ar << TextureParameter;
 		Ar << TextureParameterSampler;
@@ -934,10 +974,16 @@ void BuildHZB( FRHICommandListImmediate& RHICmdList, FViewInfo& View )
 {
 	QUICK_SCOPE_CYCLE_COUNTER(STAT_BuildHZB);
 	SCOPED_DRAW_EVENT(RHICmdList, BuildHZB);
+	
+	// View.ViewRect.{Width,Height}() are most likely to be < 2^24, so the float
+	// conversion won't loss any precision (assuming float have 23bits for mantissa)
+	const uint32 NumMipsX = FPlatformMath::CeilToInt(FMath::Log2(float(View.ViewRect.Width()))) - 1;
+	const uint32 NumMipsY = FPlatformMath::CeilToInt(FMath::Log2(float(View.ViewRect.Height()))) - 1;
+	const uint32 NumMips = FMath::Max(NumMipsX, NumMipsY);
 
 	// Must be power of 2
-	const FIntPoint HZBSize( 512, 256 );
-	const uint32 NumMips = 8;
+	const FIntPoint HZBSize( 1 << NumMipsX, 1 << NumMipsY );
+	View.HZBMipmap0Size = HZBSize;
 
 	FPooledRenderTargetDesc Desc(FPooledRenderTargetDesc::Create2DDesc(HZBSize, PF_R16F, TexCreate_None, TexCreate_RenderTargetable | TexCreate_ShaderResource | TexCreate_NoFastClear, false, NumMips));
 	Desc.Flags |= TexCreate_FastVRAM;
@@ -953,7 +999,7 @@ void BuildHZB( FRHICommandListImmediate& RHICmdList, FViewInfo& View )
 	{
 		SetRenderTarget(RHICmdList, HZBRenderTarget.TargetableTexture, 0, NULL);
 
-		TShaderMapRef< FScreenVS >		VertexShader(View.ShaderMap);
+		TShaderMapRef< FPostProcessVS >	VertexShader(View.ShaderMap);
 		TShaderMapRef< THZBBuildPS<0> >	PixelShader(View.ShaderMap);
 
 		static FGlobalBoundShaderState BoundShaderState;
@@ -961,7 +1007,7 @@ void BuildHZB( FRHICommandListImmediate& RHICmdList, FViewInfo& View )
 		SetGlobalBoundShaderState(RHICmdList, View.GetFeatureLevel(), BoundShaderState, GFilterVertexDeclaration.VertexDeclarationRHI, *VertexShader, *PixelShader);
 
 		// Imperfect sampling, doesn't matter too much
-		PixelShader->SetParameters( RHICmdList, View, HZBSize );
+		PixelShader->SetParameters( RHICmdList, View );
 
 		RHICmdList.SetViewport(0, 0, 0.0f, HZBSize.X, HZBSize.Y, 1.0f);
 
@@ -982,12 +1028,15 @@ void BuildHZB( FRHICommandListImmediate& RHICmdList, FViewInfo& View )
 	FIntPoint SrcSize = HZBSize;
 	FIntPoint DstSize = SrcSize / 2;
 	
-	// Mip 1-7
+	// Downsampling...
 	for( uint8 MipIndex = 1; MipIndex < NumMips; MipIndex++ )
 	{
+		DstSize.X = FMath::Max(DstSize.X, 1);
+		DstSize.Y = FMath::Max(DstSize.Y, 1);
+
 		SetRenderTarget(RHICmdList, HZBRenderTarget.TargetableTexture, MipIndex, NULL);
 
-		TShaderMapRef< FScreenVS >		VertexShader(View.ShaderMap);
+		TShaderMapRef< FPostProcessVS >	VertexShader(View.ShaderMap);
 		TShaderMapRef< THZBBuildPS<1> >	PixelShader(View.ShaderMap);
 
 		static FGlobalBoundShaderState BoundShaderState;
diff --git a/Engine/Source/Runtime/Renderer/Private/SceneRendering.h b/Engine/Source/Runtime/Renderer/Private/SceneRendering.h
--- a/Engine/Source/Runtime/Renderer/Private/SceneRendering.h
+++ b/Engine/Source/Runtime/Renderer/Private/SceneRendering.h
@@ -522,6 +522,10 @@ public:
 	// Hierarchical Z Buffer
 	TRefCountPtr<IPooledRenderTarget> HZB;
 
+	// Size of the HZB's mipmap 0
+	// NOTE: the mipmap 0 is downsampled version of the depth buffer
+	FIntPoint HZBMipmap0Size;
+
 	/** Used by occlusion for percent unoccluded calculations. */
 	float OneOverNumPossiblePixels;
 
