Fixes a missing include causing compilation failure on unity builds

diff --git a/Engine/Source/Runtime/Renderer/Private/SceneOcclusion.cpp b/Engine/Source/Runtime/Renderer/Private/SceneOcclusion.cpp
--- a/Engine/Source/Runtime/Renderer/Private/SceneOcclusion.cpp
+++ b/Engine/Source/Runtime/Renderer/Private/SceneOcclusion.cpp
@@ -11,6 +11,7 @@
 #include "ScreenRendering.h"
 #include "SceneFilterRendering.h"
 #include "SceneUtils.h"
+#include "PostProcessing.h"
 
 /*-----------------------------------------------------------------------------
 	Globals
