updated LightPropagationVolume (from Lionhead Studios)
* reduced order of SH (faster but a bit more blurry)
* DirectionalOcclusion (disabled by default, can be enabled in PostProcessSettings, affects SkyLight and ReflectionEnvironments but not AmbientCube)
* Material BlockGI feature
* Spotlight support
* Uses AsyncCompute on XboxOne as optimization (order is not optimized yet)

diff --git a/Engine/Shaders/DeferredShadingCommon.usf b/Engine/Shaders/DeferredShadingCommon.usf
--- a/Engine/Shaders/DeferredShadingCommon.usf
+++ b/Engine/Shaders/DeferredShadingCommon.usf
@@ -315,6 +315,8 @@ struct FScreenSpaceData
 	FGBufferData GBuffer;
 	// 0..1, only valid in some passes, 1 if off
 	float AmbientOcclusion;
+
+	float2 DirectionalOcclusion;
 };
 
 /** Populates OutGBufferA, B and C */
@@ -526,6 +528,9 @@ SamplerState ScreenSpaceAOTextureSampler;
 Texture2D CustomDepthTexture;
 SamplerState CustomDepthTextureSampler;
 
+Texture2D		DirectionalOcclusionTexture;
+SamplerState	DirectionalOcclusionSampler;
+
 #if FEATURE_LEVEL >= FEATURE_LEVEL_SM4
 	// In all but SM5 we need to explicitly declare how many samples are in a multisampled texture.
 	#if FEATURE_LEVEL >= FEATURE_LEVEL_SM5
@@ -596,6 +601,14 @@ SamplerState CustomDepthTextureSampler;
 			Out.AmbientOcclusion = ScreenSpaceAO.r;
 		}
 
+		{
+			uint width;
+			uint height;
+			uint levels;
+
+			DirectionalOcclusionTexture.GetDimensions(0, width, height, levels);
+			Out.DirectionalOcclusion = Texture2DSampleLevel( DirectionalOcclusionTexture, DirectionalOcclusionSampler, (PixelPos + 0.5f) / float2(width, height), 0 ).rg;
+		}
 		return Out;
 	}
 #endif
@@ -676,6 +689,7 @@ FScreenSpaceData GetScreenSpaceData(float2 UV, bool bGetNormalizedNormal = true)
 
 	Out.AmbientOcclusion = ScreenSpaceAO.r;
 
+	Out.DirectionalOcclusion = Texture2DSampleLevel( DirectionalOcclusionTexture, DirectionalOcclusionSampler, UV, 0 ).rg;
 	return Out;
 }
 
diff --git a/Engine/Shaders/LPVBuildGeometryVolume.usf b/Engine/Shaders/LPVBuildGeometryVolume.usf
--- a/Engine/Shaders/LPVBuildGeometryVolume.usf
+++ b/Engine/Shaders/LPVBuildGeometryVolume.usf
@@ -20,7 +20,8 @@
 
 //------------------------------------------------------------------------------
 
-#define LPV_MULTIPLE_BOUNCES_ENABLED LPV_MULTIPLE_BOUNCES
+#define LPV_MULTIPLE_BOUNCES_ENABLED	LPV_MULTIPLE_BOUNCES
+#define DEBUG_HIGHLIGHT_BORDERS			0
 
 //------------------------------------------------------------------------------
 
@@ -35,16 +36,12 @@ void CSBuildGeometryVolume(uint3 DTid : SV_DispatchThreadID)
 	int index = GetGridAddress( DTid );
 	if ( index == -1 ) return;
 
-	uint listIndex = gGvListHeadBuffer.Load( index*4 );
+	uint listIndex = gGvListHeadBuffer.Load( index*4 ) - 1; // 0-terminated lists to prevent infinite recursion
 
 	GeometryVolumeEntry cell;
-	[unroll]
-	for ( int i=0; i<9; i++ ) cell.SH[i] = 0;
 
-#if LPV_MULTIPLE_BOUNCES_ENABLED
-	// cell.color records the average color for the cell (it is not cumulative)
-	cell.color = 0.0f;
-#endif
+	ClearGvEntry( cell );
+	const float RsmAreaMultiplier = LpvWrite.RsmAreaIntensityMultiplier * 0.25;
 
 	uint count = 0;
 	while (listIndex != -1 )
@@ -58,12 +55,20 @@ void CSBuildGeometryVolume(uint3 DTid : SV_DispatchThreadID)
 		float weight = 0.15 / nDotL;
 
 #if LPV_MULTIPLE_BOUNCES_ENABLED
-		cell.color += vpl.flux;
+		cell.color += vpl.flux * RsmAreaMultiplier;
 #endif
-		AccumulateCoeffsScalar( weight, vpl.normal, cell.SH );
+
+#if DEBUG_HIGHLIGHT_BORDERS
+		if ( DTid.x==0 || DTid.x == 31 || DTid.y == 0 || DTid.y == 31 )
+		{
+			weight = 100;
+		}
+#endif
+		AccumulateCoeffsGeometryVolume( weight * RsmAreaMultiplier, vpl.normal, cell );
+
 		count++;
-		if ( count > 32 ) break; // Prevent huge spikes in this shader on PC
-		listIndex = listEntry.nextIndex;
+		if ( count > 128 ) break; // Prevent huge spikes in this shader on PC
+		listIndex = listEntry.nextIndex-1; // 0-terminated lists to prevent infinite recursion
 	}
 
 #if LPV_MULTIPLE_BOUNCES_ENABLED
diff --git a/Engine/Shaders/LPVClear.usf b/Engine/Shaders/LPVClear.usf
--- a/Engine/Shaders/LPVClear.usf
+++ b/Engine/Shaders/LPVClear.usf
@@ -20,6 +20,10 @@
 
 //------------------------------------------------------------------------------
 
+#define CLEAR_EDGES 0
+
+//------------------------------------------------------------------------------
+
 [numthreads(4,4,4)]
 void CSClear(uint3 DTid : SV_DispatchThreadID) 
 {
@@ -33,10 +37,6 @@ void CSClear(uint3 DTid : SV_DispatchThreadID)
 
 	float maxExtent = MaxGridExtent( (float3)idx );
 	float mul = PREV_FRAME_MULTIPLIER;
-	if ( maxExtent > 15.5 )
-	{
-		mul = 0.0f;
-	}
 
 	idx = clamp( idx, int3(0,0,0), int3(31,31,31) );
 	int sourceCell = GetGridAddress( idx );
@@ -45,7 +45,24 @@ void CSClear(uint3 DTid : SV_DispatchThreadID)
 
 	mul *= LpvWrite.ClearMultiplier;
 
-	//@TODO: add a separate shader for this
+	// corruption detection (workaround for brightlodge lighting corruption)
+	{
+		if ( isnan( cell.AO ) || isinf( cell.AO ) )
+		{
+			mul = 0.0f;
+		}
+		[unroll]
+		for ( int i=0; i<9; i++ )
+		{
+			float df = dot( cell.coeffs[i], float3(1.0f,1.0f,1.0f) );
+			if ( isnan( df ) || isinf( df ) )
+			{
+				mul = 0.0f;
+			}
+		}
+	}
+
+	////@TODO: add a separate shader for this
 	[branch]
 	if ( mul == 0.0f ) // This is needed in order to fix QNANs 
 	{
@@ -53,8 +70,15 @@ void CSClear(uint3 DTid : SV_DispatchThreadID)
 	}
 	else
 	{
-		MultiplyCell( cell, mul );
+		MultiplyCellCoeffs( cell, mul );
+	}
+#if CLEAR_EDGES
+	if ( maxExtent > 13.5 && false )
+	{
+		ClearCell( cell );
+		cell.AO = 0.5;
 	}
+#endif
 
 	WriteLpvCell( cell, i );
 
@@ -73,4 +97,3 @@ void CSClearGeometryVolume(uint3 DTid : SV_DispatchThreadID)
 	uint i = GetGridAddress( DTid );
 	ClearGvCell( i );
 }
-
diff --git a/Engine/Shaders/LPVClearLists.usf b/Engine/Shaders/LPVClearLists.usf
--- a/Engine/Shaders/LPVClearLists.usf
+++ b/Engine/Shaders/LPVClearLists.usf
@@ -29,8 +29,9 @@ void CSClearLists(uint3 DTid : SV_DispatchThreadID)
 {
 	uint i = GetGridAddress( DTid );
 
-	RWVplListHeadBuffer.Store( i*4, -1 );
-	RWGvListHeadBuffer.Store( i*4, -1 );
+	// 0 - terminated lists (index+1) - prevents infinite loop when reading off the end of the list
+	RWVplListHeadBuffer.Store( i*4, 0 ); 
+	RWGvListHeadBuffer.Store( i*4, 0 );
 }
 
 //------------------------------------------------------------------------------
diff --git a/Engine/Shaders/LPVCommon.usf b/Engine/Shaders/LPVCommon.usf
--- a/Engine/Shaders/LPVCommon.usf
+++ b/Engine/Shaders/LPVCommon.usf
@@ -9,7 +9,7 @@
 //
 //				Copyright (C) Microsoft. All rights reserved.
 //-----------------------------------------------------------------------------
- 
+
 //------------------------------------------------------------------------------
 
 #define REFINE_OVER_TIME	1
@@ -27,46 +27,10 @@
  
 #define TILED_ADDRESSING		0
 
-// HACK!
-#ifdef LPV_AMBIENT_CUBE
-	#undef	LPV_AMBIENT_CUBE 
-#endif
- 
-// TODO: pass these in from a common compile enviroment setup
-#define LPV_AMBIENT_CUBE		0				// This has to match the value in LightPropagationVolume.cpp
-#define LPV_AMBIENT_CUBE_WITH_SH_BUFFER_HACK 0	// This has to match the value in LightPropagationVolume.cpp
-#define LPV_VOLUME_TEXTURE		1				// This has to match the value in LightPropagationVolume.h
-
-//-------------------------------------------------------------------------------------------------
-
-#define AMBIENT_CUBE_NEGATIVE_X 0
-#define AMBIENT_CUBE_POSITIVE_X 1
-#define AMBIENT_CUBE_NEGATIVE_Y 2
-#define AMBIENT_CUBE_POSITIVE_Y 3
-#define AMBIENT_CUBE_NEGATIVE_Z 4
-#define AMBIENT_CUBE_POSITIVE_Z 5
-
 //-------------------------------------------------------------------------------------------------
-
-struct LPVCellPackedAmbientCube
-{
-#if LPV_AMBIENT_CUBE_WITH_SH_BUFFER_HACK
-	uint packedData[14];
-#else
-	uint packedData[6];
-#endif
-};
-
-struct LPVCellPackedSH
-{
-	uint packedCoeffs[14];
-};
-
-#if LPV_VOLUME_TEXTURE || LPV_GV_VOLUME_TEXTURE
 SamplerState								gLpv3DTextureSampler;
-#endif
 
-#if LPV_VOLUME_TEXTURE
+Texture3D<float>							gAOVolumeTexture; 
 Texture3D<float4>							gLpv3DTexture0;
 Texture3D<float4>							gLpv3DTexture1;
 Texture3D<float4>							gLpv3DTexture2;
@@ -74,16 +38,10 @@ Texture3D<float4>							gLpv3DTexture3;
 Texture3D<float4>							gLpv3DTexture4;
 Texture3D<float4>							gLpv3DTexture5;
 Texture3D<float4>							gLpv3DTexture6;
-#elif LPV_AMBIENT_CUBE
-StructuredBuffer<LPVCellPackedAmbientCube>	gLpvBuffer;
-#else
-StructuredBuffer<LPVCellPackedSH>			gLpvBuffer;
-#endif
 
 // Write shaders only
 #if LPV_WRITE_SHADER
 
-#if LPV_VOLUME_TEXTURE
 RWTexture3D<float4>							gLpv3DTextureRW0;
 RWTexture3D<float4>							gLpv3DTextureRW1;
 RWTexture3D<float4>							gLpv3DTextureRW2;
@@ -91,11 +49,8 @@ RWTexture3D<float4>							gLpv3DTextureRW3;
 RWTexture3D<float4>							gLpv3DTextureRW4;
 RWTexture3D<float4>							gLpv3DTextureRW5;
 RWTexture3D<float4>							gLpv3DTextureRW6;
-#elif LPV_AMBIENT_CUBE
-RWStructuredBuffer<LPVCellPackedAmbientCube> gLpvBufferRW;
-#else
-RWStructuredBuffer<LPVCellPackedSH>			gLpvBufferRW;
-#endif
+
+RWTexture3D<float>							gAOVolumeTextureRW;
 
 #endif // LPV_WRITE_SHADER
 
@@ -103,11 +58,8 @@ RWStructuredBuffer<LPVCellPackedSH>			gLpvBufferRW;
 
 struct LPVCell
 {
-#if LPV_AMBIENT_CUBE
-	float3 ambientCube[6];
-#else
 	float3 coeffs[9];
-#endif
+	float AO;
 };
 
 //-------------------------------------------------------------------------------------------------
@@ -116,60 +68,56 @@ struct LPVCell
 
 void ClearCell( in out LPVCell cell )
 { 
-#if LPV_AMBIENT_CUBE
-	[unroll]
-	for ( int i=0; i<6; i++ ) 
-		cell.ambientCube[i] = float3( 0.0f, 0.0f, 0.0f );
-#else
 	[unroll]
-	for ( int i=0; i<9; i++ ) 
+	for ( int i = 0; i < 9; i++ ) 
+	{
 		cell.coeffs[i] = float3( 0.0f, 0.0f, 0.0f );
-#endif
+	}
+
+	cell.AO = 0.0f;
 }
 
 //-------------------------------------------------------------------------------------------------
 
 void MultiplyCell( in out LPVCell cell, float multiplier )
 {
-#if LPV_AMBIENT_CUBE
 	[unroll]
-	for ( int i=0; i<6; i++ ) 
-		cell.ambientCube[i] *= multiplier;
-#else
+	for ( int i=0; i<9; i++ ) 
+	{
+		cell.coeffs[i] *= multiplier;
+	}
+	cell.AO *= multiplier;
+}
+
+void MultiplyCellCoeffs( in out LPVCell cell, float multiplier )
+{
 	[unroll]
 	for ( int i=0; i<9; i++ ) 
 		cell.coeffs[i] *= multiplier;
-#endif
 }
 
 //-------------------------------------------------------------------------------------------------
 
 void MultiplyCellV( in out LPVCell cellA, in out LPVCell cellB )
 {
-#if LPV_AMBIENT_CUBE
-	[unroll]
-	for ( int i=0; i<6; i++ ) 
-		cellA.ambientCube[i] *= cellB.ambientCube[i];
-#else
 	[unroll]
 	for ( int i=0; i<9; i++ ) 
+	{
 		cellA.coeffs[i] *= cellB.coeffs[i];
-#endif
+	}
+	cellA.AO *= cellB.AO;
 }
 
 //-------------------------------------------------------------------------------------------------
 
 void AddCell( in out LPVCell cellA, in LPVCell cellB )
 {
-#if LPV_AMBIENT_CUBE
-	[unroll]
-	for ( int i=0; i<6; i++ ) 
-		cellA.ambientCube[i] += cellB.ambientCube[i];
-#else
 	[unroll]
 	for ( int i=0; i<9; i++ ) 
+	{
 		cellA.coeffs[i] += cellB.coeffs[i];
-#endif
+	}
+	cellA.AO += cellB.AO;
 }
 
 //-------------------------------------------------------------------------------------------------
@@ -199,88 +147,6 @@ uint GetGridAddress( uint3 gridPos )
 
 //-------------------------------------------------------------------------------------------------
 
-void PackSH3( float3 coeffs[9], out LPVCellPackedSH cell )
-{
-	uint3 compressedCoeffs[9]; 
-	[unroll]
-	for ( int i=0; i<9; i++ )
-	{
-		compressedCoeffs[i] = f32tof16( coeffs[i] );
-	}
-	
-	cell.packedCoeffs[0] = compressedCoeffs[0].x | ( compressedCoeffs[0].y << 16 );
-	cell.packedCoeffs[1] = compressedCoeffs[0].z | ( compressedCoeffs[1].x << 16 );
-	cell.packedCoeffs[2] = compressedCoeffs[1].y | ( compressedCoeffs[1].z << 16 );
-
-	cell.packedCoeffs[3] = compressedCoeffs[2].x | ( compressedCoeffs[2].y << 16 );
-	cell.packedCoeffs[4] = compressedCoeffs[2].z | ( compressedCoeffs[3].x << 16 );
-	cell.packedCoeffs[5] = compressedCoeffs[3].y | ( compressedCoeffs[3].z << 16 );
-
-	cell.packedCoeffs[6] = compressedCoeffs[4].x | ( compressedCoeffs[4].y << 16 );
-	cell.packedCoeffs[7] = compressedCoeffs[4].z | ( compressedCoeffs[5].x << 16 );
-	cell.packedCoeffs[8] = compressedCoeffs[5].y | ( compressedCoeffs[5].z << 16 );
-
-	cell.packedCoeffs[9] = compressedCoeffs[6].x | ( compressedCoeffs[6].y << 16 );
-	cell.packedCoeffs[10] = compressedCoeffs[6].z | ( compressedCoeffs[7].x << 16 );
-	cell.packedCoeffs[11] = compressedCoeffs[7].y | ( compressedCoeffs[7].z << 16 );
-
-	cell.packedCoeffs[12] = compressedCoeffs[8].x | ( compressedCoeffs[8].y << 16 );
-	cell.packedCoeffs[13] = compressedCoeffs[8].z;
-}
-
-//-------------------------------------------------------------------------------------------------
-
-void UnpackSH3( LPVCellPackedSH cell, out float3 coeffs[9] )
-{
-	uint3 compressedCoeffs;
-	compressedCoeffs.x = cell.packedCoeffs[0];
-	compressedCoeffs.y = cell.packedCoeffs[0] >> 16;
-	compressedCoeffs.z = cell.packedCoeffs[1];
-	coeffs[0] = f16tof32( compressedCoeffs );
-
-	compressedCoeffs.x = cell.packedCoeffs[1] >> 16;
-	compressedCoeffs.y = cell.packedCoeffs[2];
-	compressedCoeffs.z = cell.packedCoeffs[2] >> 16;
-	coeffs[1] = f16tof32( compressedCoeffs );
-
-	compressedCoeffs.x = cell.packedCoeffs[3];
-	compressedCoeffs.y = cell.packedCoeffs[3] >> 16;
-	compressedCoeffs.z = cell.packedCoeffs[4];
-	coeffs[2] = f16tof32( compressedCoeffs );
-
-	compressedCoeffs.x = cell.packedCoeffs[4] >> 16;
-	compressedCoeffs.y = cell.packedCoeffs[5];
-	compressedCoeffs.z = cell.packedCoeffs[5] >> 16;
-	coeffs[3] = f16tof32( compressedCoeffs );
-
-	compressedCoeffs.x = cell.packedCoeffs[6];
-	compressedCoeffs.y = cell.packedCoeffs[6] >> 16;
-	compressedCoeffs.z = cell.packedCoeffs[7];
-	coeffs[4] = f16tof32( compressedCoeffs );
-
-	compressedCoeffs.x = cell.packedCoeffs[7] >> 16;
-	compressedCoeffs.y = cell.packedCoeffs[8];
-	compressedCoeffs.z = cell.packedCoeffs[8] >> 16;
-	coeffs[5] = f16tof32( compressedCoeffs );
-
-	compressedCoeffs.x = cell.packedCoeffs[9];
-	compressedCoeffs.y = cell.packedCoeffs[9] >> 16;
-	compressedCoeffs.z = cell.packedCoeffs[10];
-	coeffs[6] = f16tof32( compressedCoeffs );
-
-	compressedCoeffs.x = cell.packedCoeffs[10] >> 16;
-	compressedCoeffs.y = cell.packedCoeffs[11];
-	compressedCoeffs.z = cell.packedCoeffs[11] >> 16;
-	coeffs[7] = f16tof32( compressedCoeffs );
-
-	compressedCoeffs.x = cell.packedCoeffs[12];
-	compressedCoeffs.y = cell.packedCoeffs[12] >> 16;
-	compressedCoeffs.z = cell.packedCoeffs[13];
-	coeffs[8] = f16tof32( compressedCoeffs );
-}
-
-//-------------------------------------------------------------------------------------------------
-
 // Unpacks an HDR colour from RGY32 
 float3 UnpackRGY32( uint packedColour )
 {
@@ -308,54 +174,6 @@ uint PackRGY32( float3 colour )
 	packedValue |= f32tof16( luminance );
 	return packedValue;
 }
-//-------------------------------------------------------------------------------------------------
-
-void PackAmbientCube( in float3 unpackedCube[6], out LPVCellPackedAmbientCube packedCube )
-{
-#if LPV_AMBIENT_CUBE_WITH_SH_BUFFER_HACK
-	[unroll]
-	for ( int i=0; i<14; i++ ) packedCube.packedData[i] = 0;
-#endif
-
-	[unroll]
-	for ( int i=0; i<6; i++ )
-	{
-		packedCube.packedData[i] = PackRGY32( unpackedCube[i] );
-	} 
-
-}
-
-//-------------------------------------------------------------------------------------------------
-
-void UnpackAmbientCube( LPVCellPackedAmbientCube packedCube, out float3 unpackedCube[6] )
-{
-	[unroll]
-	for ( int i=0; i<6; i++ )
-	{
-		unpackedCube[i] = UnpackRGY32( packedCube.packedData[i] );
-	}
-}
-
-//-------------------------------------------------------------------------------------------------
-
-float3 AmbientCubeLookup( float3 ambientCube[6], float3 n ) 
-{
-	float weight[6];
-	weight[0] = saturate(-n.x );
-	weight[1] = saturate( n.x );
-	weight[2] = saturate(-n.y );
-	weight[3] = saturate( n.y );
-	weight[4] = saturate(-n.z );
-	weight[5] = saturate( n.z );
-
-	float3 colour = 0;
-	[unroll]
-	for ( int i=0; i<6; i++ )
-{
-		colour += ambientCube[i] * weight[i];
-	}
-	return colour;
-}
 
 //-------------------------------------------------------------------------------------------------
 
@@ -385,11 +203,7 @@ float3 SHLookup( float3 coeffs[9], float3 n )
 
 float3 LPVCellLookup( in LPVCell cell, in float3 normal ) 
 {
-#if LPV_AMBIENT_CUBE
-	return AmbientCubeLookup( cell.ambientCube, normal );
-#else
 	return SHLookup( cell.coeffs, normal );
-#endif
 }
 
 //-------------------------------------------------------------------------------------------------
@@ -428,21 +242,22 @@ float3 GridToWorld( float3 gridPos )
 
 float MaxGridExtent( float3 gridPos )
 {
-	float3 av = abs(gridPos-15.5)+0.5;
+	float3 av = abs(gridPos-16.0f);
 	return max(av.x,max(av.y,av.z));
 }
 
 //-------------------------------------------------------------------------------------------------
 
-int GetGridIndex( float3 worldPos )
+int GetGridIndex( float3 worldPos ) 
 {
 	float3 p = WorldToGrid( worldPos );
 
 	float maxExtent = MaxGridExtent( p );
-	if ( maxExtent > 15.5 )
+
+	if ( maxExtent > 16.0f )
 	{
 		return -1;
-	}
+	} 
 	else
 	{
 		int3 ui = int3( p );
@@ -452,10 +267,11 @@ int GetGridIndex( float3 worldPos )
 
 //-------------------------------------------------------------------------------------------------
 
-LPVCell ReadLpvCell( in uint cellIndex, const bool bUseReadWriteBuffer = false )
+LPVCell ReadLpvCell( in uint cellIndex )
 {
 	LPVCell cell;
-#if LPV_VOLUME_TEXTURE
+	cell.AO = 0.0f;
+
 	int4 texIndex = int4( IndexToGridPos( cellIndex ), 0 );
 	float4 tex0 = gLpv3DTexture0.Load( texIndex );
 	float4 tex1 = gLpv3DTexture1.Load( texIndex );
@@ -474,39 +290,12 @@ LPVCell ReadLpvCell( in uint cellIndex, const bool bUseReadWriteBuffer = false )
 	cell.coeffs[6] = float3( tex4.a, tex5.rg );
 	cell.coeffs[7] = float3( tex5.ba, tex6.r );
 	cell.coeffs[8] = tex6.gba;
-#elif LPV_AMBIENT_CUBE
-	LPVCellPackedAmbientCube PackedCube;
-  #if LPV_WRITE_SHADER
-	if ( bUseReadWriteBuffer )
-	{
-		PackedCube = gLpvBufferRW[cellIndex];
-	}
-	else
-  #endif
-	{
-		PackedCube = gLpvBuffer[cellIndex];
-	}
-	UnpackAmbientCube( PackedCube, cell.ambientCube );
-#else
-	LPVCellPackedSH packedCell;
-  #if LPV_WRITE_SHADER
-	if ( bUseReadWriteBuffer )
-	{
-		packedCell = gLpvBufferRW[cellIndex];
-	}
-	else
-  #endif
-	{
-		packedCell = gLpvBuffer[cellIndex];
-	}
-	UnpackSH3( packedCell, cell.coeffs );
-#endif
+	cell.AO = tex0.a;
 	return cell;
 }
 
 //-------------------------------------------------------------------------------------------------
 
-#if LPV_VOLUME_TEXTURE
 LPVCell ReadLpvCellVolumeTextureFiltered( float3 gridPos )
 {
 	float3 texPos = gridPos / 32.0f; // Half pixel offset?
@@ -529,8 +318,16 @@ LPVCell ReadLpvCellVolumeTextureFiltered( float3 gridPos )
 	cell.coeffs[7] = float3( tex5.ba, tex6.r );
 	cell.coeffs[8] = tex6.gba;
 
+	cell.AO = tex0.a;
 	return cell;
 }
-#endif
+
+//-------------------------------------------------------------------------------------------------
+
+float ReadLpvCellAO( in uint cellIndex )
+{
+	int4 texIndex = int4( IndexToGridPos( cellIndex ), 0 );
+	return gLpv3DTexture0.Load( texIndex ).a;
+}
 
 //-------------------------------------------------------------------------------------------------
diff --git a/Engine/Shaders/LPVDirectLightInject.usf b/Engine/Shaders/LPVDirectLightInject.usf
--- a/Engine/Shaders/LPVDirectLightInject.usf
+++ b/Engine/Shaders/LPVDirectLightInject.usf
@@ -26,19 +26,59 @@
 
 //-------------------------------------------------------------------------------------------------
 
+float ComputeAttenuation( float3 WorldPosition )
+{
+	float3 ToLight = LpvInject.LightPosition - WorldPosition;		
+	float DistanceSqr = dot( ToLight, ToLight );
+	float3 L = ToLight * rsqrt( DistanceSqr );
+	float Attenuation = 1.0f;
+
+	if (LpvInject.bLightInverseSquaredAttenuation > 0.0f )
+	{
+		// @TODO: From DeferredLightingCommon.usf GetDynamicLighting() - could refactor to avoid duplication
+		BRANCH
+		if( LpvInject.LightSourceLength > 0 )
+		{
+			// Line segment irradiance
+			float3 L01 = LpvInject.LightDirection * LpvInject.LightSourceLength;
+			float3 L0 = ToLight - 0.5 * L01;
+			float3 L1 = ToLight + 0.5 * L01;
+			Attenuation = rcp( ( length( L0 ) * length( L1 ) + dot( L0, L1 ) ) * 0.5 + 1 );
+		}
+		else
+		{
+			// Sphere irradiance (technically just 1/d^2 but this avoids inf)
+			Attenuation = 1 / ( DistanceSqr + 1 );
+		}
+		Attenuation *= Square( saturate( 1 - Square( DistanceSqr * Square(1.0f/LpvInject.LightRadius) ) ) );
+	}
+	else
+	{
+		Attenuation = RadialAttenuation( ToLight/LpvInject.LightRadius, LpvInject.LightFalloffExponent );
+	}
+
+#if SPOT_ATTENUATION
+	Attenuation *= SpotAttenuation(L, -LpvInject.LightDirection, LpvInject.LightSpotAngles);
+#endif 
+
+	return Attenuation;
+}
+
 [numthreads(4,4,4)]
-void CSPointLightInject_ListGenCS(uint3 DTid : SV_DispatchThreadID, uint3 GTid : SV_GroupThreadID ) 
+void CSLightInject_ListGenCS(uint3 DTid : SV_DispatchThreadID, uint3 GTid : SV_GroupThreadID ) 
 {
 	int index = GetGridAddress( DTid );
 
 	// TODO: optimise this. Don't read/write unless we need to. Requires read/write access to UAV, which requires
 	// aliasing for the textures... 
 	LPVCell cell = ReadLpvCell( index ); 
-	float3 gridPos = float3( DTid );
+
+	float3 GridPosition = float3( DTid );
 
 	float3 LightPos = LpvInject.LightPosition;
-	float3 worldPos = GridToWorld( gridPos+0.5  );
-	float3 LightVector = float3( LightPos-worldPos );
+
+	float3 WorldPosition = GridToWorld( GridPosition+0.5  );
+	float3 LightVector = float3( LightPos-WorldPosition );
 
 	float len = length( LightVector );
 	if ( len < LpvInject.LightRadius )
@@ -54,18 +94,14 @@ void CSPointLightInject_ListGenCS(uint3 DTid : SV_DispatchThreadID, uint3 GTid :
 		if ( LightVector.z < 0 ) gvOffset += zStride;
 
 		int gvIndex = index + gvOffset;
-#if FILTER
-		GeometryVolumeEntry gvCell = ReadFilteredGvCell( gvIndex );
-#else
 		GeometryVolumeEntry gvCell = ReadGvCell( gvIndex );
-#endif
 
 		if ( gvCell.SH[0] > 0.1f )
 		{
 			float3 LightDirection = LightVector/len;
 
 			// TODO: Not ideal for thin walls! Use SH directly rather than approximating the normal?
-			float3 normal = -SHGetDominantDirectionApproxScalar( gvCell.SH );
+			float3 normal = -GvGetDominantDirectionApprox( gvCell );
 
 			float nDotL = saturate( dot( normal, LightDirection ) );
 			
@@ -78,14 +114,15 @@ void CSPointLightInject_ListGenCS(uint3 DTid : SV_DispatchThreadID, uint3 GTid :
 				if ( distance > 100.0f )
 				{
 					float3 inc = LightVector / 16.0f;
-					float3 pos = worldPos + inc*2;
+
+					float3 pos = WorldPosition + inc*2;
 
 					for ( int i=0; i<15; i++ )
 					{
 						int traceGvIndex = GetGridIndex( pos ) + gvOffset;
 						GeometryVolumeEntry sCell = ReadGvCell( traceGvIndex );
 
-						float s = 1-SHLookupScalar( sCell.SH, -LightDirection );
+						float s = 1-SHLookupGeometryVolume( sCell, -LightDirection );
 
 						shadowValue = min( shadowValue, s );
 						pos+=inc;
@@ -95,12 +132,12 @@ void CSPointLightInject_ListGenCS(uint3 DTid : SV_DispatchThreadID, uint3 GTid :
 #endif
 
 				float3 flux = LpvInject.LightColor * gvCell.color * DIRECT_LIGHT_INJECTION_STRENGTH;
-				float Attenuation = nDotL*shadowValue;
-				float DistanceAttenuation = RadialAttenuation( LightVector/LpvInject.LightRadius, LpvInject.LightFalloffExponent );
-				Attenuation *= DistanceAttenuation;
+
+				float Attenuation = ComputeAttenuation( WorldPosition );
+				Attenuation *= nDotL*shadowValue;
 				flux *= Attenuation;
 
-				worldPos += normal * LpvWrite.VplInjectionBias;
+				WorldPosition += normal * LpvWrite.VplInjectionBias;
 
 				float solidAngle = 3.14;
 				AccumulateLighting( flux, -normal, solidAngle, cell );
diff --git a/Engine/Shaders/LPVDirectionalOcclusion.usf b/Engine/Shaders/LPVDirectionalOcclusion.usf
new file mode 100644
--- /dev/null
+++ b/Engine/Shaders/LPVDirectionalOcclusion.usf
@@ -0,0 +1,236 @@
+//-----------------------------------------------------------------------------
+// File:		LPVSkyInjection.usf
+//
+// Summary:		Compute shader which handles LPV propagation
+//
+// Created:		2013-03-01
+//
+// Author:		mailto:[REDACTED_EMAIL]
+//
+//				Copyright (C) Microsoft. All rights reserved. 
+//-----------------------------------------------------------------------------
+
+/*------------------------------------------------------------------------------
+	Compile time parameters:
+------------------------------------------------------------------------------*/
+ 
+#include "Common.usf"
+#include "LPVWriteCommon.usf"
+#include "LPVGeometryVolumeCommon.usf"
+#include "CubemapCommon.usf"
+
+//-------------------------------------------------------------------------------------------------
+
+#define JITTER						1
+#define AO_ITERATIVE_REFINEMENT     1
+#define OPTIMISED					1
+#define FILTERED_GV_LOOKUP			0 // Really expensive! Do not use!
+
+#if AO_ITERATIVE_REFINEMENT 
+#define REFINE_AO_MULTIPLIER 0.15
+#endif
+
+//-------------------------------------------------------------------------------------------------
+#define DIRECTIONS32	float3( -0.05512,  0.76586, -0.56600),	float3(  0.32073,  0.38860, -0.45865),	float3( -0.94457,  0.10946, -0.06761),	float3(  0.52959, -0.54968, -0.22405),\
+						float3(  0.35128,  0.08201,  0.46989),	float3(  0.23199, -0.17011,  0.70594),	float3(  0.60823,  0.34747,  0.27763),	float3( -0.08181,  0.04158, -0.12801),\
+						float3( -0.41883,  0.48855, -0.58237),	float3( -0.10297, -0.74525,  0.27079),	float3( -0.90964, -0.07262,  0.42218),	float3( -0.73489, -0.44662,  0.18148),\
+						float3(  0.44118,  0.15554, -0.23668),	float3(  0.15504, -0.44441,  0.12967),	float3(  0.22360, -0.43429, -0.62978),	float3( -0.27742, -0.22710,  0.77011),\
+						float3(  0.82907,  0.19693, -0.38211),	float3(  0.28126,  0.86776, -0.00408),	float3( -0.17777,  0.38409, -0.25268),	float3(  0.01183, -0.22878,  0.39425),\
+						float3(  0.44725, -0.00260, -0.63503),	float3( -0.42517, -0.58948, -0.66224),	float3( -0.49335,  0.05575, -0.22749),	float3( -0.27795,  0.03986, -0.68356),\
+						float3(  0.04470,  0.61118, -0.16199),	float3(  0.10392,  0.21927, -0.91510),	float3( -0.34484, -0.21977,  0.40968),	float3( -0.70361,  0.44257,  0.48594),\
+						float3(  0.53131, -0.16425,  0.77398),	float3( -0.07781,  0.06550,  0.63558),	float3(  0.18829, -0.62067,  0.46053),	float3( -0.59232,  0.16259, -0.71309)
+
+//-------------------------------------------------------------------------------------------------
+
+#define DIRECTIONS16	float3(  0.51472, -0.49163,  0.35433), float3( -0.35116, -0.01130, -0.39922), float3( -0.80703,  0.42269,  0.06061),    \
+						float3(  0.64946, -0.41234, -0.64451), float3(  0.66478, -0.05262,  0.61239), float3(  0.51724,  0.50787,  0.13460),    \
+						float3( -0.34963,  0.50349, -0.25580), float3( -0.06436, -0.00462, -0.99482), float3( -0.06583,  0.94559, -0.29587),    \
+						float3( -0.06452,  0.50999,  0.47619), float3(  0.00707, -0.05436,  0.40955), float3( -0.03803, -0.85644,  0.21446),    \
+						float3(  0.09232, -0.36510, -0.59058), float3(  0.61293, -0.00690, -0.08238), float3( -0.48265, -0.49618, -0.59132),    \
+						float3( -0.37232, -0.48738,  0.41218)
+
+
+//-------------------------------------------------------------------------------------------------
+//-------------------------------------------------------------------------------------------------
+
+#if FILTERED_GV_LOOKUP
+GeometryVolumeEntry ReadGvCellFiltered( float3 gridPos )
+{
+	GeometryVolumeEntry gv;
+
+	float3 texPos = saturate( gridPos / 32.0f ); // Half pixel offset?
+	float4 gvTex0 = gGv3DTexture0.SampleLevel( gLpv3DTextureSampler, texPos, 0 );
+	gv.SH[0] = gvTex0.x;
+#if ( GV_ORDER >= 1 ) 
+	float4 gvTex1 = gGv3DTexture1.SampleLevel( gLpv3DTextureSampler, texPos, 0 );
+	gv.SH[1] = gvTex0.y;
+	gv.SH[2] = gvTex0.z;
+	gv.SH[3] = gvTex0.w;
+#endif
+#if ( GV_ORDER >= 2 )
+	float4 gvTex2 = gGv3DTexture2.SampleLevel( gLpv3DTextureSampler, texPos, 0 );
+	gv.SH[4] = gvTex1.x;
+	gv.SH[5] = gvTex1.y;
+	gv.SH[6] = gvTex1.z;
+	gv.SH[7] = gvTex1.w;
+	gv.SH[8] = gvTex2.x;
+#endif
+
+ #if LPV_MULTIPLE_BOUNCES
+ #if (GV_ORDER == 0 )
+  gv.color = gvTex0.yzw;
+#elif (GV_ORDER == 1)
+  gv.color = gvTex1.xyz;
+#else
+	gv.color = gvTex2.yzw;
+ #endif
+ #endif
+	return gv;
+}
+
+#endif //FILTERED_GV_LOOKUP
+
+
+#if FILTERED_GV_LOOKUP
+GeometryVolumeEntry ReadGvCell_Pos( float3 gridPos ) 
+{
+	//gridPos+=float3(0.5,0.5,0.5)*1;
+	//gridPos = max( min( float3(32,32,32), gridPos ), float3(0,0,0) );
+	return ReadGvCellFiltered( gridPos );
+}
+#else
+GeometryVolumeEntry ReadGvCell_Pos( int3 pos ) 
+{
+	pos = max( min( int3(31,31,31), pos ), int3(0,0,0) );
+	return ReadGvCell( GetGridAddress( pos ) );
+}
+#endif
+
+//-------------------------------------------------------------------------------------------------
+
+float ComputeAmbientOcclusion( float3 gridPos )
+{
+#if OPTIMISED
+	const int NumSamples=16;
+	const float3 Directions[16] = 
+	{
+		DIRECTIONS16
+	};
+	const float rayLength = 16.0f;
+	const uint samplesPerRay = 8;
+	float sampleLength = rayLength / float(16); // Optimised = same initial step size as with 16 taps, but step size increases per iteration to compensate
+#else
+	const int NumSamples=32;
+	const float3 Directions[32] = 
+	{
+		DIRECTIONS32
+	};
+	const float rayLength = 16.0f;
+	const uint samplesPerRay = 16;
+	float sampleLength = rayLength / float(samplesPerRay);
+#endif
+	sampleLength *= LpvWrite.DirectionalOcclusionRadius / 16.0f;
+	float RayWeight = 1.0f/float(NumSamples);
+
+#if OPTIMISED
+	RayWeight *= 1.02;
+#endif 
+
+	float ao = 0.0f;
+
+#if JITTER
+	// Required, but not cache friendly (roughly 4X performance cost). Happily, we can get away with far fewer taps
+	float3 gp = abs( float3( gridPos - LpvWrite.mLpvGridOffset ) );
+	float3x3 rotation = { float3(1,0,0), float3(0,1,0), float3(0,0,1) };
+	int ix = int(gp.x+gp.y*5+gp.z*3+21) % NumSamples;
+	int iy = int(gp.x+gp.y+gp.z+456) % NumSamples;
+	rotation[0] = normalize(Directions[ix]);
+	rotation[1] = normalize(Directions[iy]);
+	rotation[2] = normalize( cross( rotation[0], rotation[1] ) );
+	rotation[1] = normalize( cross( rotation[0], rotation[2] ) );
+#endif
+
+	// NOTE: Balanced to give 0.5 AO on a plane, using default settings
+	LOOP
+	for ( int i=0; i<NumSamples; i++ )
+	{
+#if JITTER
+		float3 direction = mul( rotation, Directions[i] );
+#else 
+		float3 direction = Directions[i];
+#endif
+
+		float3 pos = gridPos + float3(0.5,0.5,0.5);
+		float3 inc = direction * sampleLength;
+		float rayOcclusion = 1.0f;
+
+		pos += inc*0.5f;
+		float weight = 2.0f;
+
+		for ( int j=0; j < samplesPerRay; j++ )
+		{ 
+			pos += inc;
+			GeometryVolumeEntry gvCell = ReadGvCell_Pos( pos );
+			{
+				// TODO: should this just be -direction? (need to rebalance)
+				float s = SHLookupGeometryVolume( gvCell, direction );
+				s += SHLookupGeometryVolume( gvCell, -direction );
+				rayOcclusion -= weight * s;
+			}
+#if OPTIMISED
+			inc *= 1.1f; // Sample distance increases with range (equivalent to adjusting weight in high quality version)			
+#else
+			weight *= 0.75f;
+#endif
+
+		}
+		rayOcclusion *= 1.25f;
+		ao += ( 1.0-saturate(rayOcclusion) );
+	}
+
+	float aoWeight = RayWeight;
+	aoWeight *= 1.2f;
+	ao *= aoWeight;
+
+	return ( 1.0f-ao ) * 1.81f;
+}
+//-------------------------------------------------------------------------------------------------
+
+[numthreads(4,4,4)]
+void CSDirectionalOcclusion(uint3 DTid : SV_DispatchThreadID, uint3 GTid : SV_GroupThreadID ) 
+{
+	int index = GetGridAddress( DTid );
+
+	float3 gridPos = float3( DTid );
+
+	// Read the current SH cell
+	LPVCell cell = ReadLpvCell( index );
+
+	float AO = ComputeAmbientOcclusion( gridPos );
+#if AO_ITERATIVE_REFINEMENT
+	cell.AO = lerp( cell.AO, AO, REFINE_AO_MULTIPLIER );
+#else
+	cell.AO = AO;
+#endif
+
+	// Write the accumulated cell
+	WriteLpvCell( cell, index );
+}
+
+//-------------------------------------------------------------------------------------------------
+
+// Compute shader to copy the AO volume from the full-float LPV to a lower-precision
+// volume texture
+[numthreads(4,4,4)]
+void CSCopyAOVolume(uint3 DTid : SV_DispatchThreadID, uint3 GTid : SV_GroupThreadID ) 
+{
+	int index = GetGridAddress( DTid );
+
+	float3 gridPos = float3( DTid );
+
+	// Read the current SH cell
+	float AO = ReadLpvCellAO( index );
+
+	gAOVolumeTextureRW[ DTid ] = AO;
+}
+
diff --git a/Engine/Shaders/LPVFinalPass.usf b/Engine/Shaders/LPVFinalPass.usf
--- a/Engine/Shaders/LPVFinalPass.usf
+++ b/Engine/Shaders/LPVFinalPass.usf
@@ -31,9 +31,6 @@
 
 // Indirect strength
 #define INDIRECT_LEVEL		( 0.00375 * REFINE_LIGHTING_MULTIPLIER )
-#define INDIRECT_SPECULAR_LEVEL	1.0
-#define INDIRECT_DIFFUSE_LEVEL	0.333
-
 #define INDIRECT_DIFFUSE	1
 #define INDIRECT_SPECULAR	1
 
@@ -43,19 +40,15 @@
 #define FILTER_TRILINEAR			2 
 #define FILTER_HARDWARE_TRILINEAR	3 // Use hardware volume texture filtering
 
+#ifndef LPV_USE_SEPARATE_AO_VOLUME
+#define LPV_USE_SEPARATE_AO_VOLUME	1 // Use the separate (low precision) AO texture for AO lookups
+#endif
+
+#define LPV_DEFAULT_AO_VALUE		0.4
+
 // Filtering settings
-#if LPV_VOLUME_TEXTURE 
 	#define DIFFUSE_FILTER		FILTER_HARDWARE_TRILINEAR 
 	#define SPECULAR_FILTER		FILTER_HARDWARE_TRILINEAR 
-#else
-	#ifndef DIFFUSE_FILTER
-	#define DIFFUSE_FILTER		FILTER_TRILINEAR
-	#endif
-
-	#ifndef SPECULAR_FILTER
-	#define SPECULAR_FILTER		FILTER_TRILINEAR
-	#endif
-#endif
 
 //-------------------------------------------------------------------------------------------------
 
@@ -113,12 +106,12 @@ void GetGridIndicesAndWeightsTetra( float3 gridPos, out int indices[4], out floa
 	if ( flip2.z < 0 ) cubePos.z = 1-cubePos.z;
 
 	[unroll]
-	for ( uint i=0; i<4; i++ )
+	for ( uint j=0; j<4; j++ )
 	{
-		float3 tip = tips[i];
+		float3 tip = tips[j];
 		float3 d = abs(cubePos-tip);
 		float mDist = d.x+d.y+d.z;
-		if ( mDist <= 1.0f ) tetra=i;
+		if ( mDist <= 1.0f ) tetra=j;
 	}
 
 	if ( tetra<4 )
@@ -402,7 +395,6 @@ LPVCell LPVGetCellTrilinear( float3 worldPos )
 
 //-------------------------------------------------------------------------------------------------
 
-#if LPV_VOLUME_TEXTURE
 LPVCell LPVGetCellHardwareTrilinear( float3 worldPos )
 {
 	float3 colour = 0;
@@ -411,7 +403,6 @@ LPVCell LPVGetCellHardwareTrilinear( float3 worldPos )
 
 	return ReadLpvCellVolumeTextureFiltered( gridPos );
 }
-#endif
 
 //-------------------------------------------------------------------------------------------------
 
@@ -428,8 +419,29 @@ LPVCell LPVGetCell( float3 worldPos )
 #endif
 }
 
-//-------------------------------------------------------------------------------------------------
+float LPVGetCellAO_Grid( float3 gridPos )
+{
+	float3 absGridPos = abs(gridPos-16.0f);
+	float m = max( absGridPos.x, max( absGridPos.y, absGridPos.z ));
+	if ( m<=16.0f )
+	{
+		float3 texPos = gridPos / 32.0f;
+#if LPV_USE_SEPARATE_AO_VOLUME
+		return gAOVolumeTexture.SampleLevel( gLpv3DTextureSampler, texPos, 0 ).r;
+#else
+		return gLpv3DTexture0.SampleLevel( gLpv3DTextureSampler, texPos, 0 ).a;
+#endif
+	}
+	return LPV_DEFAULT_AO_VALUE;
+}
 
+float LPVGetCellAO( float3 WorldPos )
+{
+	float3 gridPos = WorldToGrid( WorldPos );
+	return LPVGetCellAO_Grid( gridPos );
+}
+
+//-------------------------------------------------------------------------------------------------
 
 float3 LPVSpecular( float3 worldPos, float3 reflectionVec, float rOffsetLength )
 {
@@ -439,7 +451,7 @@ float3 LPVSpecular( float3 worldPos, float3 reflectionVec, float rOffsetLength )
 
 	LPVCell Cell = LPVGetCell( worldPos + reflectionVec * rOffsetLength * LpvRead.LpvScale );
 
-	float3 Colour = LPVCellLookup( Cell, reflectionVec ) * INDIRECT_LEVEL * INDIRECT_SPECULAR_LEVEL;
+	float3 Colour = LPVCellLookup( Cell, reflectionVec ) * INDIRECT_LEVEL * LpvRead.SpecularIntensity;
 
 #if SHOW_GRID
 	Colour = DebugGridColour( WorldToGrid( worldPos ), colour );
@@ -458,7 +470,7 @@ float3 LPVSpecular( float3 worldPos, float3 reflectionVec )
 
 //-------------------------------------------------------------------------------------------------
 
-float3 LPVLookup( float3 worldPos, float3 normal )
+float3 LPVDiffuse( float3 worldPos, float3 normal )
 {
 #if !INDIRECT_DIFFUSE
 	return 0;
@@ -466,7 +478,7 @@ float3 LPVLookup( float3 worldPos, float3 normal )
 
 	LPVCell Cell = LPVGetCell( worldPos );
 
-	float3 Colour = LPVCellLookup( Cell, normal ) * INDIRECT_LEVEL * INDIRECT_DIFFUSE_LEVEL;
+	float3 Colour = LPVCellLookup( Cell, normal ) * INDIRECT_LEVEL * LpvRead.DiffuseIntensity;
 
 	return Colour;
 }
@@ -477,8 +489,88 @@ void LPVGetDiffuseAndSpecular( float3 worldPos, float3 normal, float3 reflection
 {
 	// TODO: could implement optimised ref lookup by using SH from diffuse for reflection
 
-	diffuseOut = LPVLookup( worldPos, normal ) ;
+	diffuseOut = LPVDiffuse( worldPos, normal ) ;
 	specularOut = LPVSpecular( worldPos, reflectionVec ) ; 
 }
 
+
+//-------------------------------------------------------------------------------------------------
+
+// Directional occlusion lookup
+float2 LPVGetDirectionalOcclusion_DiffuseSpec( float3 WorldPosition, float3 R, float Roughness, float3 N )
+{
+	BRANCH
+	if ( LpvRead.DirectionalOcclusionIntensity < 0.0001f )
+	{
+		return float2(1,1);
+	}
+	float SpecularOcclusion = 0.0f;
+	float DiffuseOcclusion = 0.0f;
+	const int NumSpecularSteps = 6; 
+	const int NumDiffuseSteps  = 2;
+
+	float3 gridPos = WorldToGrid( WorldPosition );
+	float StepLength = 1.333 / float(NumDiffuseSteps);//0.301f;
+	float Weight = 1.0f;
+	float TotalWeight = 0.0f;
+
+#if DEBUG_USE_RAW_AO
+	float ao = saturate( LPVGetCellAO_Grid( gridPos ) );
+	if ( ao > 0.5 ) return 1;
+	return ao.xx;
+#endif
+
+	float DiffuseOcclusionRaw = LPVGetCellAO_Grid( gridPos + N * 0.25 );
+	DiffuseOcclusion = saturate( DiffuseOcclusionRaw * 2.0f / LpvRead.DiffuseOcclusionIntensity );
+	DiffuseOcclusion = pow( DiffuseOcclusion, LpvRead.DiffuseOcclusionExponent );
+	DiffuseOcclusion = saturate(DiffuseOcclusion);
+
+	gridPos = WorldToGrid( WorldPosition );
+	StepLength = 2.0f / (float)NumSpecularSteps;
+	SpecularOcclusion = 0.0f;
+	TotalWeight=0;
+	Weight=1.0f;
+
+	// Use the diffuse occlusion tap as the first specular tap (since it's nearly identical)
+	SpecularOcclusion = saturate( 1-DiffuseOcclusionRaw ) * Weight;
+	StepLength *= 1.3f;
+	gridPos+=R*StepLength;
+	TotalWeight+=Weight; 
+	Weight *= 0.9;
+	for  ( int i=1; i<NumSpecularSteps; i++ )
+	{
+		float AO = saturate( 1-LPVGetCellAO_Grid( gridPos ) );
+		SpecularOcclusion += AO * Weight;
+		StepLength *= 1.3f;
+		gridPos+=R*StepLength;
+		TotalWeight+=Weight; 
+		Weight *= 0.9;
+	}
+
+	// Increase intensity as power increases	
+	float LerpValue = saturate( LpvRead.SpecularOcclusionExponent * 0.1f ); // 0-1 range
+	float Strength = lerp( 0.9f, 1.3f, LerpValue ) * LpvRead.SpecularOcclusionIntensity;
+	SpecularOcclusion *= Strength;
+	SpecularOcclusion /= TotalWeight;
+	SpecularOcclusion = pow( SpecularOcclusion, LpvRead.SpecularOcclusionExponent );
+
+	SpecularOcclusion = 1-SpecularOcclusion;
+
+	float2 DirectionalOcclusion = lerp( float2( 1,1 ), saturate( float2( DiffuseOcclusion, SpecularOcclusion) ), LpvRead.DirectionalOcclusionIntensity );
+
+#if LPV_OCCLUSION_EDGE_FADE
+	float3 grid = WorldToGrid( WorldPosition );
+	float minGrid = min( grid.x, min(grid.y,grid.z ) );
+	float maxGrid = max( grid.x, max(grid.y,grid.z ) );
+
+	float inside = min(	saturate(minGrid*0.5),
+						saturate((32.0f-maxGrid)*0.5) );
+
+	DirectionalOcclusion = lerp( float2(0.5,0.5), DirectionalOcclusion, inside );
+#endif
+
+	return DirectionalOcclusion;
+}
+
+
 #endif
diff --git a/Engine/Shaders/LPVGeometryVolumeCommon.usf b/Engine/Shaders/LPVGeometryVolumeCommon.usf
--- a/Engine/Shaders/LPVGeometryVolumeCommon.usf
+++ b/Engine/Shaders/LPVGeometryVolumeCommon.usf
@@ -14,11 +14,29 @@
 #define LPV_MULTIPLE_BOUNCES 0
 #endif 
 
+#define GV_ORDER 2
+
+#ifndef LPV_GV_SH_ORDER
+#error LPV_GV_SH_ORDER is not defined!
+#else
+#if ( GV_ORDER > LPV_GV_SH_ORDER )
+#undef GV_ORDER
+#define GV_ORDER LPV_GV_SH_ORDER
+#endif
+#endif
+
 //------------------------------------------------------------------------------
 
 struct GeometryVolumeEntry
 {
+#if (GV_ORDER==2)
 	float SH[9];
+#elif (GV_ORDER==1)
+	float SH[4];
+#elif (GV_ORDER==0)
+	float SH[1];
+#endif
+
 #if LPV_MULTIPLE_BOUNCES
 	//@todo : store colour as SH? Is single component good enough?
 	float3 color;
@@ -27,17 +45,15 @@ struct GeometryVolumeEntry
 
 //------------------------------------------------------------------------------
 
-#if LPV_GV_VOLUME_TEXTURE
 RWTexture3D<float4>	gGv3DTextureRW0;
-RWTexture3D<float4>	gGv3DTextureRW1;
-RWTexture3D<float4>	gGv3DTextureRW2;
-
 Texture3D<float4>	gGv3DTexture0;
+#if ( GV_ORDER >= 1 )
+RWTexture3D<float4>	gGv3DTextureRW1;
 Texture3D<float4>	gGv3DTexture1;
+#endif
+#if (GV_ORDER >= 2 )
+RWTexture3D<float4>	gGv3DTextureRW2;
 Texture3D<float4>	gGv3DTexture2;
-#else
-RWBuffer<float3>	gGvBufferRW;
-Buffer<float3>		gGvBuffer;
 #endif
 
 //-------------------------------------------------------------------------------------------------
@@ -48,43 +64,48 @@ GeometryVolumeEntry ReadGvCell( uint index )
 {
 	GeometryVolumeEntry gv;
 
-#if LPV_GV_VOLUME_TEXTURE
 	int4 texIndex = int4( IndexToGridPos( index ), 0 );
 	float4 gvTex0 = gGv3DTexture0.Load( texIndex );
-	float4 gvTex1 = gGv3DTexture1.Load( texIndex );
-	float4 gvTex2 = gGv3DTexture2.Load( texIndex );
 	gv.SH[0] = gvTex0.x;
+#if ( GV_ORDER >= 1 )
+	float4 gvTex1 = gGv3DTexture1.Load( texIndex );
 	gv.SH[1] = gvTex0.y;
 	gv.SH[2] = gvTex0.z;
 	gv.SH[3] = gvTex0.w;
+#endif
+#if ( GV_ORDER >= 2 )
+	float4 gvTex2 = gGv3DTexture2.Load( texIndex );
 	gv.SH[4] = gvTex1.x;
 	gv.SH[5] = gvTex1.y;
 	gv.SH[6] = gvTex1.z;
 	gv.SH[7] = gvTex1.w;
 	gv.SH[8] = gvTex2.x;
+#endif
+
  #if LPV_MULTIPLE_BOUNCES
+ #if (GV_ORDER == 0 )
+  gv.color = gvTex0.yzw;
+#elif (GV_ORDER == 1)
+  gv.color = gvTex1.xyz;
+#else
 	gv.color = gvTex2.yzw;
  #endif
-#else // LPV_GV_VOLUME_TEXTURE
- #if LPV_MULTIPLE_BOUNCES
-	uint bufferIndex = index*4;
- #else
-	uint bufferIndex = index*3;
  #endif
-	[unroll]
-	for ( uint i=0; i<3; i++ )
-	{
-		gv.SH[i*3]   = gGvBuffer[ bufferIndex+i ].x * 8.0 - 4.0;
-		gv.SH[i*3+1] = gGvBuffer[ bufferIndex+i ].y * 8.0 - 4.0;
-		gv.SH[i*3+2] = gGvBuffer[ bufferIndex+i ].z * 8.0 - 4.0;
+	return gv;
 	}
- #if LPV_MULTIPLE_BOUNCES
-	gv.color = gGvBuffer[ bufferIndex+3 ].xyz;
- #endif
-#endif // LPV_GV_VOLUME_TEXTURE
 
+//-------------------------------------------------------------------------------------------------
 
-	return gv;
+float3 GvGetDominantDirectionApprox( in GeometryVolumeEntry cell )
+{
+#if (GV_ORDER == 0 )
+	return float3(0,0,0);
+#else
+	float3 dir = float3( -cell.SH[3], -cell.SH[1], -cell.SH[2] );
+	float3 len = length( dir );
+	float den = ( len > 0.01 ) ? len : 1.0f;
+	return dir/den;
+#endif
 }
 
 //-------------------------------------------------------------------------------------------------
@@ -92,8 +113,20 @@ GeometryVolumeEntry ReadGvCell( uint index )
 void WriteGvCell( uint index, GeometryVolumeEntry entry )
 {
 
-#if LPV_GV_VOLUME_TEXTURE
 	int3 texIndex = IndexToGridPos( index );
+
+#if ( GV_ORDER == 0 )
+ #if LPV_MULTIPLE_BOUNCES
+	gGv3DTextureRW0[ texIndex ] = float4( entry.SH[0], 0.0f, 0.0f, 0.0f );
+ #else
+	gGv3DTextureRW2[ texIndex ] = float4( entry.SH[0], entry.color.rgb );
+ #endif
+#elif ( GV_ORDER == 1 )
+	gGv3DTextureRW0[ texIndex ] = float4( entry.SH[0], entry.SH[1], entry.SH[2], entry.SH[3] );
+ #if LPV_MULTIPLE_BOUNCES
+	gGv3DTextureRW1[ texIndex ] = float4( entry.color.rgb, 0.0f );
+ #endif
+#else
 	gGv3DTextureRW0[ texIndex ] = float4( entry.SH[0], entry.SH[1], entry.SH[2], entry.SH[3] );
 	gGv3DTextureRW1[ texIndex ] = float4( entry.SH[4], entry.SH[5], entry.SH[6], entry.SH[7] );
  #if LPV_MULTIPLE_BOUNCES
@@ -101,54 +134,50 @@ void WriteGvCell( uint index, GeometryVolumeEntry entry )
  #else
 	gGv3DTextureRW2[ texIndex ] = float4( entry.SH[8], 0.0f, 0.0f, 0.0f );
  #endif
+#endif
+}
 
-#else // LPV_GV_VOLUME_TEXTURE
- #if LPV_MULTIPLE_BOUNCES
-	uint bufferIndex = index*4;
- #else
-	uint bufferIndex = index*3;
- #endif
+//-------------------------------------------------------------------------------------------------
 
-	[unroll]
-	for ( uint i=0; i<3; i++ )
+void ClearGvEntry( in out GeometryVolumeEntry entry )
 	{
-		gGvBufferRW[ bufferIndex+i ] = float3( entry.SH[i*3], entry.SH[i*3+1], entry.SH[i*3+2] ) / 8.0 + 0.5;
-	}
+	entry.SH[0] = 0.0f;
+ #if ( GV_ORDER >= 1 )
+	entry.SH[1] = 0.0f;
+	entry.SH[2] = 0.0f;
+	entry.SH[3] = 0.0f;
+ #endif
+ #if ( GV_ORDER >= 2 )
+	entry.SH[4] = 0.0f;
+	entry.SH[5] = 0.0f;
+	entry.SH[6] = 0.0f;
+	entry.SH[7] = 0.0f;
+	entry.SH[8] = 0.0f;
+ #endif
 
  #if LPV_MULTIPLE_BOUNCES
-	gGvBufferRW[ bufferIndex+3 ].xyz = entry.color;
+	entry.color = 0.0f;
  #endif
-#endif // LPV_GV_VOLUME_TEXTURE
 }
 
 //-------------------------------------------------------------------------------------------------
 
 void ClearGvCell( uint index )
 { 
-#if LPV_GV_VOLUME_TEXTURE
 	int3 texIndex = IndexToGridPos( index );
 	gGv3DTextureRW0[ texIndex ] = 0.0f;
+ #if ( GV_ORDER >= 1 )
 	gGv3DTextureRW1[ texIndex ] = 0.0f;
-	gGv3DTextureRW2[ texIndex ] = 0.0f;
-#else // LPV_GV_VOLUME_TEXTURE
- #if LPV_MULTIPLE_BOUNCES
-	uint bufferIndex = index*4;
- #else
-	uint bufferIndex = index*3;
  #endif
-	[unroll]
-	for ( uint i=0; i<3; i++ ) gGvBufferRW[bufferIndex+i] = float3(0.5f,0.5f,0.5f);
-
- #if LPV_MULTIPLE_BOUNCES
-	gGvBufferRW[ bufferIndex+3 ].xyz = float3(0.0f,0.0f,0.0f);
+ #if ( GV_ORDER >= 2 )
+	gGv3DTextureRW2[ texIndex ] = 0.0f;
  #endif
-#endif // LPV_GV_VOLUME_TEXTURE
 }
 
 //-------------------------------------------------------------------------------------------------
 
-float SHLookupScalar( 
-	float coeffs[9],
+float SHLookupGeometryVolume( 
+	GeometryVolumeEntry cell,
     float3 n ) 
 {
 	const float1 c1 = 0.429043 ;
@@ -158,20 +187,24 @@ float SHLookupScalar(
 	const float1 c5 = 0.247708 ;
 
 	float3 n2 = n*n;
+#if ( GV_ORDER == 0 )
+	float value = c4*cell.SH[0];
+#elif ( GV_ORDER == 1 )
+	float value = c4*cell.SH[0] + 2.0f*c2*(cell.SH[3]*n.x+cell.SH[1]*n.y+cell.SH[2]*n.z);
+#else
 	float xy = n.x*n.y ; 
 	float yz = n.y*n.z ; 
 	float xz = n.x*n.z ;
-
-	float value = c1*coeffs[8]*(n2.x-n2.y) + c3*coeffs[6]*n2.z + c4*coeffs[0] - c5*coeffs[6] 
-            + 2.0f*c1*(coeffs[4]*xy + coeffs[7]*xz + coeffs[5]*yz) 
-            + 2.0f*c2*(coeffs[3]*n.x+coeffs[1]*n.y+coeffs[2]*n.z) ;
-
+	float value = c1*cell.SH[8]*(n2.x-n2.y) + c3*cell.SH[6]*n2.z + c4*cell.SH[0] - c5*cell.SH[6] 
+            + 2.0f*c1*(cell.SH[4]*xy + cell.SH[7]*xz + cell.SH[5]*yz) 
+            + 2.0f*c2*(cell.SH[3]*n.x+cell.SH[1]*n.y+cell.SH[2]*n.z) ;
+#endif
 	return max( value, 0.0f );
 }
 
 //-------------------------------------------------------------------------------------------------
 
-void AccumulateCoeffsScalar(float value, float3 n, in out float coeffs[9] ) 
+void AccumulateCoeffsGeometryVolume(float value, float3 n, in out GeometryVolumeEntry cell ) 
 { 
 	const float c0 = 0.282095f;
 	const float c1 = 0.488603f;
@@ -179,15 +212,17 @@ void AccumulateCoeffsScalar(float value, float3 n, in out float coeffs[9] )
 	const float c3 = 0.315392f;
 	const float c4 = 0.546274f;
 
-	coeffs[0] += c0*value;
-
-	coeffs[1] += (c1*n.y)*value;   
-	coeffs[2] += (c1*n.z)*value;   
-	coeffs[3] += (c1*n.x)*value;   
-
-	coeffs[4] += (c2*n.x*n.y)*value; 
-	coeffs[5] += (c2*n.y*n.z)*value; 
-	coeffs[7] += (c2*n.x*n.z)*value; 
-	coeffs[6] += (c3*(3.0f*n.z*n.z-1.0f))*value; 
-	coeffs[8] += (c4*(n.x*n.x-n.y*n.y))*value;
+	cell.SH[0] += c0*value;
+#if ( GV_ORDER >= 1 )
+	cell.SH[1] += (c1*n.y)*value;   
+	cell.SH[2] += (c1*n.z)*value;   
+	cell.SH[3] += (c1*n.x)*value;   
+ #if ( GV_ORDER >= 2 )
+	cell.SH[4] += (c2*n.x*n.y)*value; 
+	cell.SH[5] += (c2*n.y*n.z)*value; 
+	cell.SH[7] += (c2*n.x*n.z)*value; 
+	cell.SH[6] += (c3*(3.0f*n.z*n.z-1.0f))*value; 
+	cell.SH[8] += (c4*(n.x*n.x-n.y*n.y))*value;
+ #endif
+#endif
 }
\ No newline at end of file
diff --git a/Engine/Shaders/LPVInject_AccumulateVplLists.usf b/Engine/Shaders/LPVInject_AccumulateVplLists.usf
--- a/Engine/Shaders/LPVInject_AccumulateVplLists.usf
+++ b/Engine/Shaders/LPVInject_AccumulateVplLists.usf
@@ -34,17 +34,13 @@ ByteAddressBuffer						gVplListHeadBuffer;
 void CSAccumulateVplLists(uint3 DTid : SV_DispatchThreadID) 
 {
 	int index = GetGridAddress( DTid );
-	uint listIndex = gVplListHeadBuffer.Load( index*4 );
+	uint listIndex = gVplListHeadBuffer.Load( index * 4 ) - 1;
 
 #if !LIGHT_INJECTION_ENABLED
 	return;
 #endif
 
-#if LPV_VOLUME_TEXTURE
 	LPVCell cell = ReadLpvCell( index );
-#else
-	LPVCell cell = ReadLpvCell( index, true );
-#endif
 
 #if LIGHT_INJECTION_DEBUG_TEST
 	float3 testPos = float3(989,-684,565);
@@ -52,7 +48,7 @@ void CSAccumulateVplLists(uint3 DTid : SV_DispatchThreadID)
 	if ( testIndex == index )
 	{
 		float3 col = float3(0.0f, 10000.0f, 0.0f);
-		float solidAngle = 3.14; // TODO: adjust for splat size
+		float solidAngle = 3.14; // TODO: adjust for splat size 
 		AccumulateLighting( col, float3(-1, 0, 0 ), solidAngle, cell );
 		AccumulateLighting( col, float3( 1, 0, 0 ), solidAngle, cell );
 		AccumulateLighting( col, float3( 0,-1, 0 ), solidAngle, cell );
@@ -69,9 +65,9 @@ void CSAccumulateVplLists(uint3 DTid : SV_DispatchThreadID)
 	{
 		VplListEntry listEntry = gVplListBuffer[ listIndex ];
 		Vpl vpl = UnpackVpl( listEntry );
-		float solidAngle = 3.14; // TODO: adjust for splat size
+		float solidAngle = 3.14; 
 		AccumulateLighting( vpl.flux, -vpl.normal, solidAngle, cell );
-		listIndex = listEntry.nextIndex;
+		listIndex = listEntry.nextIndex - 1; // 0-terminated lists to prevent infinite recursion
 		count++;
 	}
 	WriteLpvCell( cell, index );
diff --git a/Engine/Shaders/LPVInject_GenerateVplLists.usf b/Engine/Shaders/LPVInject_GenerateVplLists.usf
--- a/Engine/Shaders/LPVInject_GenerateVplLists.usf
+++ b/Engine/Shaders/LPVInject_GenerateVplLists.usf
@@ -47,7 +47,7 @@ void CSGenerateVplLists_LightDirectional(uint3 DTid : SV_DispatchThreadID)
 
 	flux *= biasFactor;
 
-	float2 texCoord = float2( DTid.xy ) / 255.0f;
+	float2 texCoord = float2( DTid.xy ) * LpvWrite.RsmPixelToTexcoordMultiplier;
 
 	float3 rsmPos = float3( texCoord, depth );
 	float3 worldPos = mul( float4(rsmPos,1), LpvWrite.mRsmToWorld ).xyz;
diff --git a/Engine/Shaders/LPVPacking.usf b/Engine/Shaders/LPVPacking.usf
new file mode 100644
--- /dev/null
+++ b/Engine/Shaders/LPVPacking.usf
@@ -0,0 +1,112 @@
+//-----------------------------------------------------------------------------
+// File:		LPVPacking.usf
+//
+// Summary:		Common functionality for LPV compute shaders
+//
+// Created:		2015-04-15
+//
+// Author:		Ben Woodhouse - mailto:[REDACTED_EMAIL]
+//
+//				Copyright (C) Microsoft. All rights reserved.
+//-----------------------------------------------------------------------------
+
+//------------------------------------------------------------------------------
+
+struct LPVCellPacked
+{
+	uint packedCoeffs[14];
+};
+
+//-------------------------------------------------------------------------------------------------
+
+//-------------------------------------------------------------------------------------------------
+
+LPVCellPacked PackLPVCell( LPVCell cellIn )
+{
+	LPVCellPacked cellOut;
+	uint3 compressedCoeffs[9]; 
+	[unroll]
+	for ( int i=0; i<9; i++ )
+	{
+		compressedCoeffs[i] = f32tof16( cellIn.coeffs[i] );
+	}
+	uint compressedAO = f32tof16( cellIn.AO );
+	
+	cellOut.packedCoeffs[0] = compressedCoeffs[0].x | ( compressedCoeffs[0].y << 16 );
+	cellOut.packedCoeffs[1] = compressedCoeffs[0].z | ( compressedCoeffs[1].x << 16 );
+	cellOut.packedCoeffs[2] = compressedCoeffs[1].y | ( compressedCoeffs[1].z << 16 );
+
+	cellOut.packedCoeffs[3] = compressedCoeffs[2].x | ( compressedCoeffs[2].y << 16 );
+	cellOut.packedCoeffs[4] = compressedCoeffs[2].z | ( compressedCoeffs[3].x << 16 );
+	cellOut.packedCoeffs[5] = compressedCoeffs[3].y | ( compressedCoeffs[3].z << 16 );
+
+	cellOut.packedCoeffs[6] = compressedCoeffs[4].x | ( compressedCoeffs[4].y << 16 );
+	cellOut.packedCoeffs[7] = compressedCoeffs[4].z | ( compressedCoeffs[5].x << 16 );
+	cellOut.packedCoeffs[8] = compressedCoeffs[5].y | ( compressedCoeffs[5].z << 16 );
+
+	cellOut.packedCoeffs[9] = compressedCoeffs[6].x | ( compressedCoeffs[6].y << 16 );
+	cellOut.packedCoeffs[10] = compressedCoeffs[6].z | ( compressedCoeffs[7].x << 16 );
+	cellOut.packedCoeffs[11] = compressedCoeffs[7].y | ( compressedCoeffs[7].z << 16 );
+
+	cellOut.packedCoeffs[12] = compressedCoeffs[8].x | ( compressedCoeffs[8].y << 16 );
+	cellOut.packedCoeffs[13] = compressedCoeffs[8].z | ( compressedAO << 16 );
+
+
+	return cellOut;
+}
+
+//-------------------------------------------------------------------------------------------------
+
+LPVCell UnpackLPVCell( LPVCellPacked cellPacked )
+{
+	LPVCell cellOut;
+	uint3 compressedCoeffs;
+	compressedCoeffs.x = cellPacked.packedCoeffs[0];
+	compressedCoeffs.y = cellPacked.packedCoeffs[0] >> 16;
+	compressedCoeffs.z = cellPacked.packedCoeffs[1];
+	cellOut.coeffs[0] = f16tof32( compressedCoeffs );
+
+	compressedCoeffs.x = cellPacked.packedCoeffs[1] >> 16;
+	compressedCoeffs.y = cellPacked.packedCoeffs[2];
+	compressedCoeffs.z = cellPacked.packedCoeffs[2] >> 16;
+	cellOut.coeffs[1] = f16tof32( compressedCoeffs );
+
+	compressedCoeffs.x = cellPacked.packedCoeffs[3];
+	compressedCoeffs.y = cellPacked.packedCoeffs[3] >> 16;
+	compressedCoeffs.z = cellPacked.packedCoeffs[4];
+	cellOut.coeffs[2] = f16tof32( compressedCoeffs );
+
+	compressedCoeffs.x = cellPacked.packedCoeffs[4] >> 16;
+	compressedCoeffs.y = cellPacked.packedCoeffs[5];
+	compressedCoeffs.z = cellPacked.packedCoeffs[5] >> 16;
+	cellOut.coeffs[3] = f16tof32( compressedCoeffs );
+
+	compressedCoeffs.x = cellPacked.packedCoeffs[6];
+	compressedCoeffs.y = cellPacked.packedCoeffs[6] >> 16;
+	compressedCoeffs.z = cellPacked.packedCoeffs[7];
+	cellOut.coeffs[4] = f16tof32( compressedCoeffs );
+
+	compressedCoeffs.x = cellPacked.packedCoeffs[7] >> 16;
+	compressedCoeffs.y = cellPacked.packedCoeffs[8];
+	compressedCoeffs.z = cellPacked.packedCoeffs[8] >> 16;
+	cellOut.coeffs[5] = f16tof32( compressedCoeffs );
+
+	compressedCoeffs.x = cellPacked.packedCoeffs[9];
+	compressedCoeffs.y = cellPacked.packedCoeffs[9] >> 16;
+	compressedCoeffs.z = cellPacked.packedCoeffs[10];
+	cellOut.coeffs[6] = f16tof32( compressedCoeffs );
+
+	compressedCoeffs.x = cellPacked.packedCoeffs[10] >> 16;
+	compressedCoeffs.y = cellPacked.packedCoeffs[11];
+	compressedCoeffs.z = cellPacked.packedCoeffs[11] >> 16;
+	cellOut.coeffs[7] = f16tof32( compressedCoeffs );
+
+	compressedCoeffs.x = cellPacked.packedCoeffs[12];
+	compressedCoeffs.y = cellPacked.packedCoeffs[12] >> 16;
+	compressedCoeffs.z = cellPacked.packedCoeffs[13];
+	cellOut.coeffs[8] = f16tof32( compressedCoeffs );
+
+	cellOut.AO = f16tof32( cellPacked.packedCoeffs[13] >> 16 );
+
+	return cellOut;
+}
diff --git a/Engine/Shaders/LPVPropagate.usf b/Engine/Shaders/LPVPropagate.usf
--- a/Engine/Shaders/LPVPropagate.usf
+++ b/Engine/Shaders/LPVPropagate.usf
@@ -10,24 +10,35 @@
 //				Copyright (C) Microsoft. All rights reserved. 
 //-----------------------------------------------------------------------------
 
-/*------------------------------------------------------------------------------
-	Compile time parameters:
+/*------------------------------------------------------------------------------  
+	Compile time parameters: 
 ------------------------------------------------------------------------------*/
  
-#include "Common.usf"
+#include "Common.usf"  
 #include "LPVWriteCommon.usf"
 #include "LPVGeometryVolumeCommon.usf"
+#include "LPVPacking.usf"
 
 //-------------------------------------------------------------------------------------------------
+ 
+#define PROPAGATION					1
+
+// Reduce register overhead by temporarily storing values in LDS
+#define LDS_OPTIMISATION			0
+#define LDS_PACKED 					0
+#define NORMALS_LOOKUP 				1
+
+#define LDS_READ_CACHE				1
 
-#define PROPOGATION					1
+// Should the GV block AO?
+#define LPV_SECONDARY_OCCLUSION_AO  0
 
 #if REFINE_OVER_TIME
-  #define PROPAGATE_WEIGHT 0.008 
+  #define PROPAGATE_WEIGHT 0.008  
 #else
   #define PROPAGATE_WEIGHT 0.08
 #endif 
-
+ 
 // LPV_MULTIPLE_BOUNCES_ENABLED only works if LPV_SECONDARY_OCCLUSION and LPV_MULTIPLE_BOUNCES are enabled
 #if ( !LPV_MULTIPLE_BOUNCES || !LPV_SECONDARY_OCCLUSION ) && LPV_MULTIPLE_BOUNCES_ENABLED
   #undef LPV_MULTIPLE_BOUNCES_ENABLED
@@ -40,46 +51,190 @@
 #define L2 0.7071067f
 #define L3 0.5773502f
 
-void GetGeometryVolumeCellWeights( int index, 
+//-------------------------------------------------------------------------------------------------
+
+#if LDS_READ_CACHE
+
+uint GetLDSCacheIndex( int3 ldsCacheOrigin, int3 gridAddress ) 
+{
+	int3 offset3 = gridAddress-ldsCacheOrigin;
+	return uint( offset3.x + offset3.y*6 + offset3.z*6*6 );
+}
+
+groupshared LPVCellPacked LDSCachedLPV[6*6*6]; //56 bytes per element
+
+#endif
+
+#if LDS_OPTIMISATION && LPV_SECONDARY_OCCLUSION
+
+//-------------------------------------------------------------------------------------------------
+
+// TODO: Actually share some data between threads!
+
+//-------------------------------------------------------------------------------------------------
+
+#if LPV_MULTIPLE_BOUNCES_ENABLED 
+
+groupshared uint LDSCachedGV[64][26];
+
+void ClearWeights( uint groupIndex )
+{
+}
+
+void StoreWeightAndColor( uint groupIndex, uint neighbourIndex, float weight, float3 color )
+{
+	uint4 intOut = uint4( saturate( float4( weight, color ) ) * 255.0f );
+	uint intPacked = ( (uint)intOut.x ) >>  0 |					
+					 ( (uint)intOut.y ) >>  8 |
+					 ( (uint)intOut.z ) >> 16 |
+  					 ( (uint)intOut.w ) >> 24;
+	LDSCachedGV[groupIndex][neighbourIndex] = intPacked;
+}
+
+float GetWeight( uint groupIndex, uint neighbourIndex )
+{
+	return float(LDSCachedGV[groupIndex][neighbourIndex] & 0x000000ff )/255.0f;
+}
+
+float3 GetColor( uint groupIndex, uint neighbourIndex )
+{
+	uint colourPacked = LDSCachedGV[groupIndex][neighbourIndex];
+	float3 col;
+	col.r = float( ( colourPacked & 0x0000ff00 ) >> 8  );
+	col.g = float( ( colourPacked & 0x00ff0000 ) >> 16 );
+	col.b = float( ( colourPacked & 0xff000000 ) >> 24 );
+	return col / 255.0f;
+}
+
+
+#else //!LPV_MULTIPLE_BOUNCES_ENABLED 
+
+#if LDS_PACKED
+groupshared uint LDSCachedGV[64][7];
+
+void ClearWeights( uint groupIndex )
+{
+	[unroll]
+	for ( int i=0; i<7; i++ )
+	{
+		LDSCachedGV[groupIndex][i] = 0;
+	}
+}
+
+void StoreWeight( uint groupIndex, uint neighbourIndex, float weight )
+{
+	uint uintIndex  = neighbourIndex /4;//>> 2;
+	uint bitOffset  = ( neighbourIndex % 4 ) * 8;//( neighbourIndex & 0x3 ) << 3;
+	uint unpacked = uint(saturate(weight)*255.0f);
+	uint packed = unpacked << bitOffset;
+	packed = packed | LDSCachedGV[groupIndex][uintIndex];
+	LDSCachedGV[groupIndex][uintIndex] = packed;
+}
+
+float GetWeight( uint groupIndex, uint neighbourIndex )
+{
+	uint uintIndex  = neighbourIndex /4;//>> 2;
+	uint bitOffset  = ( neighbourIndex % 4 ) * 8;//( neighbourIndex & 0x3 ) << 3;
+	uint packed = LDSCachedGV[groupIndex][uintIndex];
+	uint unpacked = ( packed >> bitOffset ) & 0xff;
+	return float(unpacked)/255.0f;
+}
+#else
+groupshared uint LDSCachedGV[64][26];
+
+void ClearWeights( uint groupIndex )
+{
+}
+void StoreWeight( uint groupIndex, uint neighbourIndex, float weight )
+{
+	LDSCachedGV[groupIndex][neighbourIndex] = uint( saturate( weight ) * 255.0f );
+}
+
+float GetWeight( uint groupIndex, uint neighbourIndex )
+{
+	return float(LDSCachedGV[groupIndex][neighbourIndex])/255.0f;
+}
+
+#endif // LDS_PACKED
+
+#endif // !LPV_MULTIPLE_BOUNCES_ENABLED 
+
+
+#endif
+
+//-------------------------------------------------------------------------------------------------
+
+int3 GetDirectionFromIndex( const uint index )
+{
+	uint idx = index<13 ? index : index+1;
+	return int3( idx/9,	(idx/3) % 3,  idx % 3 ) - int3(1,1,1);
+}
+
+float3 GetNormalisedDirectionFromIndex( const uint index )
+{
+	uint idx = index<13 ? index : index+1;
+	return normalize( float3( idx/9,	(idx/3) % 3,  idx % 3 )- float3(1,1,1) );
+}
+
+
+GeometryVolumeEntry ReadGvCell_Pos( int3 pos ) // TODO: move this to GvCommon
+{
+	pos = max( min( int3(31,31,31), pos ), int3(0,0,0) );
+	return ReadGvCell( GetGridAddress( pos ) );
+}
+
+#if LPV_SECONDARY_OCCLUSION
+void GetGeometryVolumeCellWeights( int3 gridPos, 
+#if LDS_OPTIMISATION
+								   const uint groupIndex,
+#else
 								   out float weights[26], 
 #if LPV_MULTIPLE_BOUNCES_ENABLED  
 								   out float3 colors[26], 
+#endif
 #endif
 								   const float3 directions[26] )
 {
 	// Read the surrounding 8 geometry volume cells
 	GeometryVolumeEntry gvCells[8];
-	const int xStride = 1;
-	const int yStride = 32;
-	const int zStride = 32*32;
-	gvCells[0] = ReadGvCell( index );
-	gvCells[1] = ReadGvCell( index + xStride );
-	gvCells[2] = ReadGvCell( index + yStride );
-	gvCells[3] = ReadGvCell( index + yStride + xStride );
-	gvCells[4] = ReadGvCell( index + zStride );
-	gvCells[5] = ReadGvCell( index + xStride + zStride );
-	gvCells[6] = ReadGvCell( index + yStride + zStride );
-	gvCells[7] = ReadGvCell( index + yStride + xStride + zStride );
-
-	[unroll]
-	for ( int i=0; i<26; i++ ) weights[i] = 1.0f;
 
-#if LPV_MULTIPLE_BOUNCES_ENABLED 
+#if LPV_MULTIPLE_BOUNCES_ENABLED && !LDS_OPTIMISATION
 	[unroll]
 	for ( int i=0; i<26; i++ ) colors[i] = 1.0f;
 #endif
+#if LDS_OPTIMISATION
+	ClearWeights( groupIndex );
+#endif
 
 	// 8 1-tap weights (corners)
 	{
 		const int destIndex[8] = { 0, 17, 6, 23, 2, 19, 8, 25 };
-		[unroll]
+
 		for ( int i=0; i<8; i++ )
 		{
+#if NORMALS_LOOKUP
 			const float3 direction = directions[ destIndex[i] ];
-			weights[ destIndex[i] ] = 1.0f-saturate( SHLookupScalar( gvCells[i].SH, direction ) * LpvWrite.SecondaryOcclusionStrength );
+#else
+			const float3 direction = GetNormalisedDirectionFromIndex( destIndex[i] );//directions[ destIndex[i] ];
+#endif
+
+			gvCells[i] = ReadGvCell_Pos( gridPos + int3( i&1,i&2,i&4 ) );
+			float weight = 1.0f-saturate( SHLookupGeometryVolume( gvCells[i], direction ) * LpvWrite.SecondaryOcclusionStrength );
+#if LDS_OPTIMISATION
+ #if !LPV_MULTIPLE_BOUNCES_ENABLED
+			StoreWeight( groupIndex, destIndex[i], weight );
+ #endif
+#else
+			weights[ destIndex[i] ] = weight;
+#endif
 
 #if LPV_MULTIPLE_BOUNCES_ENABLED 
-			colors[ destIndex[i] ] = gvCells[i].color;
+			float3 color = gvCells[i].color;
+#if LDS_OPTIMISATION
+			StoreWeightAndColor( groupIndex, destIndex[i], weight, color );
+#else
+			colors[ destIndex[i] ] = color;
+#endif
 #endif
 		}
 	}
@@ -87,61 +242,85 @@ void GetGeometryVolumeCellWeights( int index,
 	{
 		const int destIndex[12]      = {  1,	 3,	    5,	   7,	  9,	 11,	14,    16,    18,    20,    22,    24 };
 		const int sourceIndex[12][2] = { {0,4}, {0,2}, {4,6}, {2,6}, {0,1}, {4,5}, {2,3}, {6,7}, {1,5}, {1,3}, {5,7}, {3,7} };
-		[unroll]
+
 		for ( int i=0; i<12; i++ )
 		{
+#if NORMALS_LOOKUP
 			const float3 direction = directions[ destIndex[i] ];
-			float cellWeight = SHLookupScalar( gvCells[ sourceIndex[i][0] ].SH, direction ) * LpvWrite.SecondaryOcclusionStrength;
-			cellWeight		+= SHLookupScalar( gvCells[ sourceIndex[i][1] ].SH, direction ) * LpvWrite.SecondaryOcclusionStrength;
-			weights[ destIndex[i] ] = 1.0f-saturate( cellWeight * 0.5f );
+#else
+			const float3 direction = GetNormalisedDirectionFromIndex( destIndex[i] );
+#endif
+			float cellWeight = SHLookupGeometryVolume( gvCells[ sourceIndex[i][0] ], direction );
+			cellWeight		+= SHLookupGeometryVolume( gvCells[ sourceIndex[i][1] ], direction );
+			float weight = 1.0f-saturate( cellWeight * 0.5f * LpvWrite.SecondaryOcclusionStrength );
+#if LDS_OPTIMISATION
+ #if !LPV_MULTIPLE_BOUNCES_ENABLED
+			StoreWeight( groupIndex, destIndex[i], weight );
+ #endif
+#else
+			weights[ destIndex[i] ] = weight;
+#endif
 
 #if LPV_MULTIPLE_BOUNCES_ENABLED 
 			float3 cellColour = gvCells[ sourceIndex[i][0] ].color + gvCells[ sourceIndex[i][1] ].color;
-			colors[ destIndex[i] ] = cellColour * 0.5f;
+			float3 color = cellColour * 0.5f;
+#if LDS_OPTIMISATION
+			StoreWeightAndColor( groupIndex, destIndex[i], weight, color );
+#else
+			colors[ destIndex[i] ] = color;
 #endif
-		}
+#endif
+		} 
 	}
 	// 6 4-tap weights (faces)
 	{
 		const int destIndex[6]      = {  4,         10,        12,        13,        15,        21 };
 		const int sourceIndex[6][4] = { {0,2,4,6}, {0,1,4,5}, {0,1,2,3}, {4,5,6,7}, {2,3,6,7}, {1,3,5,7} };
-		[unroll]
+
 		for ( int i=0; i<6; i++ )
 		{
+#if NORMALS_LOOKUP
 			const float3 direction = directions[ destIndex[i] ];
-			float cellWeight = SHLookupScalar( gvCells[ sourceIndex[i][0] ].SH, direction ) * LpvWrite.SecondaryOcclusionStrength;
-			cellWeight		+= SHLookupScalar( gvCells[ sourceIndex[i][1] ].SH, direction ) * LpvWrite.SecondaryOcclusionStrength;
-			cellWeight		+= SHLookupScalar( gvCells[ sourceIndex[i][2] ].SH, direction ) * LpvWrite.SecondaryOcclusionStrength;
-			cellWeight		+= SHLookupScalar( gvCells[ sourceIndex[i][3] ].SH, direction ) * LpvWrite.SecondaryOcclusionStrength;
-			weights[ destIndex[i] ] = 1.0f-saturate( cellWeight * 0.25f );
+#else
+			const float3 direction = GetNormalisedDirectionFromIndex( destIndex[i] );
+#endif
+			float cellWeight = SHLookupGeometryVolume( gvCells[ sourceIndex[i][0] ], direction ) * LpvWrite.SecondaryOcclusionStrength;
+			cellWeight		+= SHLookupGeometryVolume( gvCells[ sourceIndex[i][1] ], direction ) * LpvWrite.SecondaryOcclusionStrength;
+			cellWeight		+= SHLookupGeometryVolume( gvCells[ sourceIndex[i][2] ], direction ) * LpvWrite.SecondaryOcclusionStrength;
+			cellWeight		+= SHLookupGeometryVolume( gvCells[ sourceIndex[i][3] ], direction ) * LpvWrite.SecondaryOcclusionStrength;
+			float weight = 1.0f-saturate( cellWeight * 0.25f );
+#if LDS_OPTIMISATION
+ #if !LPV_MULTIPLE_BOUNCES_ENABLED
+			StoreWeight( groupIndex, destIndex[i], weight );
+ #endif
+#else
+			weights[ destIndex[i] ] = weight; 
+#endif 
 #if LPV_MULTIPLE_BOUNCES_ENABLED
 			float3 cellColour = gvCells[ sourceIndex[i][0] ].color + gvCells[ sourceIndex[i][1] ].color +
 								gvCells[ sourceIndex[i][2] ].color + gvCells[ sourceIndex[i][3] ].color;
-			colors[ destIndex[i] ] = cellColour * 0.25f;
+			float3 color = cellColour * 0.25f;
+#if LDS_OPTIMISATION
+			StoreWeightAndColor( groupIndex, destIndex[i], weight, color );
+#else
+			colors[ destIndex[i] ] = color;
+#endif
 #endif
 		}
 	}
 }
+#endif
+
 
 //-------------------------------------------------------------------------------------------------
 
 [numthreads(4,4,4)]
-void CSPropogate(uint3 DTid : SV_DispatchThreadID, uint3 GTid : SV_GroupThreadID ) 
+void CSPropagate(uint3 DTid : SV_DispatchThreadID, uint3 GTid : SV_GroupThreadID, uint groupIndex : SV_GroupIndex  ) 
 {
 	int index = GetGridAddress( DTid );
 
-
 	float3 gridPos = float3( DTid );
 
-	// Read the current SH cell
-	LPVCell cell = ReadLpvCell( index );
-
-#if PROPOGATION
-
-#if LPV_MULTIPLE_BOUNCES_ENABLED
-	LPVCell originalCell = cell;
-#endif
-	 
 	const int n=26;
 	const int3 directions[n] = 
 		{	int3(-1,-1,-1 ), int3(-1,-1, 0 ), int3(-1,-1, 1 ), int3(-1, 0,-1 ), int3(-1, 0, 0 ), int3(-1, 0, 1 ), int3(-1, 1,-1 ), int3(-1, 1, 0 ), int3(-1, 1, 1 ),
@@ -153,37 +332,128 @@ void CSPropogate(uint3 DTid : SV_DispatchThreadID, uint3 GTid : SV_GroupThreadID
 			float3( 0 ,-L2,-L2 ), float3( 0 ,-L1, 0  ), float3( 0 ,-L2, L2 ), float3( 0 , 0 ,-L1 ),						  float3( 0 , 0 , L1 ), float3( 0 , L2,-L2 ), float3( 0 , L1, 0  ), float3( 0 , L2, L2 ),
 			float3( L3,-L3,-L3 ), float3( L2,-L2, 0  ), float3( L3,-L3, L3 ), float3( L2, 0 ,-L2 ), float3( L1, 0 , 0  ), float3( L2, 0 , L2 ), float3( L3, L3,-L3 ), float3( L2, L2, 0  ), float3( L3, L3, L3 ) };
 
+
 	// Get the weights from the geometry volume
 #if LPV_SECONDARY_OCCLUSION
+
+#if LDS_OPTIMISATION
+	GetGeometryVolumeCellWeights( DTid, groupIndex, nDirections );
+#else
 	float gvWeights[n];
 #if LPV_MULTIPLE_BOUNCES_ENABLED 
 	float3 gvColors[n];
-	GetGeometryVolumeCellWeights( index, gvWeights, gvColors, nDirections );
+	GetGeometryVolumeCellWeights( DTid, gvWeights, gvColors, nDirections );
+#else
+	GetGeometryVolumeCellWeights( DTid, gvWeights, nDirections );
+#endif
+#endif
+#endif
+
+	// FIXME: Use these weights for energy preserving propagation. For now we just use for AO
+	// These plus the centre add up to 1.0. 
+	// Read the current SH cell
+	LPVCell cell = ReadLpvCell( index );
+
+#if LPV_MULTIPLE_BOUNCES_ENABLED
+	LPVCell originalCell = cell;
+#endif
+
+#if LPV_PROPAGATE_AO 
+
+	const float gaussianWeights[26] = 
+	{
+		0.015625,	0.03125,		0.015625,
+		0.03125,		0.0625,		0.03125,
+		0.015625,	0.03125,		0.015625,
+		
+		0.03125,		0.0625,		0.03125,
+		0.0625,					0.0625,
+		0.03125,		0.0625,		0.03125,
+		
+		0.015625,	0.03125,		0.015625,
+		0.03125,		0.0625,		0.03125,
+		0.015625,	0.03125,		0.015625
+	}; 
+
+	const float centreGaussian = 0.125f;
+	cell.AO *= centreGaussian;
+#endif
+
+#if LDS_READ_CACHE
+#if XB1_PROFILE && 0
+	uint laneID = __XB_GetLaneID();
+	int thread0cellIndex = __XB_ReadLane(index,0); // Read the index at zone0
+	int3 gridPosThread0 = IndexToGridPos(thread0cellIndex);
 #else
-	GetGeometryVolumeCellWeights( index, gvWeights, nDirections );
+	uint laneID = groupIndex;
+	int3 ldsCacheOrigin = DTid & (~0x3);//(~0x3f);
+	ldsCacheOrigin -= int3(1,1,1);
 #endif
 
+	for (int i=0;i<4;i++)
+	{
+		// Get the cell index,from 0-215 (6x6x6 block). There will be 4 indices per thread (with some duplication)
+		uint cellCacheIndex = laneID*4+i;
+		if ( cellCacheIndex < (6*6*6) )
+		{
+			uint3 offset3;
+			offset3.x = ( cellCacheIndex % 6 );
+			offset3.y = ( cellCacheIndex / 6 ) %  6;
+			offset3.z = cellCacheIndex / ( 6 * 6 );
+			int3 cachedGridPos = ldsCacheOrigin + int3(offset3);
+			LPVCell CachedThreadCell = ReadLpvCell( GetGridAddress(cachedGridPos) );
+ 
+			LDSCachedLPV[cellCacheIndex] = PackLPVCell( CachedThreadCell );
+		}
+	}
+
+	GroupMemoryBarrierWithGroupSync();
 #endif
 
-#if !PS4_PROFILE
+
+#if XBOXONE_PROFILE
 	[loop]
 #endif
 	for ( int i=0; i<n; i++ )
 	{
-		int3 pos = ( DTid+directions[i] ) % 32;
+#if NORMALS_LOOKUP
+		int3 direction=directions[i];
+#else
+		int3 direction=GetDirectionFromIndex(i);
+#endif
+
+		int3 pos = DTid+direction;
+		pos = max(pos,float3(0,0,0) );
+		pos = min(pos,float3(31,31,31) );
+
 		int nIndex = GetGridAddress( pos );
 
+#if LDS_READ_CACHE
+		uint ldsIndex = GetLDSCacheIndex( ldsCacheOrigin, pos );
+		LPVCell nCell = UnpackLPVCell( LDSCachedLPV[ldsIndex] );
+
+#else
 		LPVCell nCell = ReadLpvCell( nIndex );
+#endif
 
+#if NORMALS_LOOKUP
 		float3 nDir = float3(directions[i]);
 		float lenSq = dot(nDir,nDir);
 		nDir = nDirections[i];
+#else
+		float3 nDir = float3(direction);
+		float lenSq = dot(nDir,nDir);
+		nDir = normalize(nDir);
+#endif
 
 		float cellWeight = 1.0f/(lenSq); // Falloff with 1/distance sq
 		float weight = cellWeight * PROPAGATE_WEIGHT;
 #if LPV_SECONDARY_OCCLUSION
-		//float gvWeight = pow( gvWeights[i], 4.0f );
+#if LDS_OPTIMISATION
+		float gvWeight = GetWeight( groupIndex, i );
+#else
 		float gvWeight = gvWeights[i];
+#endif
 		weight *= gvWeight;
 #endif
 		float3 col = LPVCellLookup( nCell, nDir ) * weight;
@@ -196,14 +466,29 @@ void CSPropogate(uint3 DTid : SV_DispatchThreadID, uint3 GTid : SV_GroupThreadID
 		float opacity = 1.0f - gvWeight;
 		// Reflect the light back, multiplying by the geometry's opacity and colour. 
 		// This is added to the light read back from the cell
-		float3 reflectColor = ( gvColors[i] * cellOutColor) * pow(opacity,3);
+#if LDS_OPTIMISATION
+		float3 gvColor = GetColor( groupIndex, i );
+#else
+		float3 gvColor = gvColors[i];
+#endif
+		float3 reflectColor = ( gvColor * cellOutColor) * pow(opacity,3);
 		col += reflectColor * LpvWrite.SecondaryBounceStrength;
 #endif
 
 		AccumulateLighting( col, nDir, 1.0f, cell );
+		// AO
+#if LPV_PROPAGATE_AO
+		{
+			float aoWeight = gaussianWeights[i];
+#if LPV_SECONDARY_OCCLUSION && LPV_SECONDARY_OCCLUSION_AO
+			aoWeight *= gvWeight;
+#endif
+			cell.AO += nCell.AO * aoWeight; 
+		}
+#endif // LPV_PROPAGATE_AO
 	}
-#endif  // Propagation
 
 	// Write the accumulated cell
 	WriteLpvCell( cell, index );
 }
+ 
\ No newline at end of file
diff --git a/Engine/Shaders/LPVVisualise.usf b/Engine/Shaders/LPVVisualise.usf
--- a/Engine/Shaders/LPVVisualise.usf
+++ b/Engine/Shaders/LPVVisualise.usf
@@ -21,6 +21,12 @@
 #define CUBESIZE					40
 #define VISUALIZE_GV				0
 #define VISUALIZE_LPV				1
+#define VISUALIZE_AO				0 
+#define SHOW_EDGES					1
+
+#define EDGE_THICKNESS				0.995
+
+#define Brightness 2
 
 //-------------------------------------------------------------------------------------------------
 
@@ -82,7 +88,7 @@ void GShader( point VIn input[1], inout TriangleStream<VOut> TriStream )
 	int indices[6] = { 0,1,2, 2,3,0 };
 
 	[unroll]
-	for( int i = 0; i < 6; i++ )
+	for ( int i=0; i<6; i++ )
 	{
 		int index = indices[i];
 
@@ -113,26 +119,35 @@ float4 PShader(VOut IN) : SV_Target
 	float3 colour = 0;
 #if VISUALIZE_LPV
 	LPVCell cell = ReadLpvCell( index );
-#if LPV_AMBIENT_CUBE
-	colour += ( cell.ambientCube[0]+cell.ambientCube[1]+cell.ambientCube[2]+cell.ambientCube[3]+cell.ambientCube[4]+cell.ambientCube[5] ) / 6.0f;
-#else
 	colour += cell.coeffs[0];
-#endif 
 	colour *= 0.000025f;
+#endif 
 
+#if VISUALIZE_AO
+	float AO = ReadLpvCellAO( index );
+	colour += saturate(1-AO*0.575);
 #endif
-
-	float brightness = 5.5;
+	float brightness = Brightness;
 
 #if VISUALIZE_GV
 	GeometryVolumeEntry gvCell = ReadGvCell( index );
-	colour.r += SHLookupScalar( gvCell.SH, float3( 1, 0, 0 ) );
-	colour.g += SHLookupScalar( gvCell.SH, float3( 0, 1, 0 ) );
-	colour.b += SHLookupScalar( gvCell.SH, float3( 0, 0, 1 ) );
+	colour.r += SHLookupGeometryVolume( gvCell, float3( 1, 0, 0 ) );
+	colour.g += SHLookupGeometryVolume( gvCell, float3( 0, 1, 0 ) );
+	colour.b += SHLookupGeometryVolume( gvCell, float3( 0, 0, 1 ) );
 	colour *= 0.1f;
 #endif
 	colour *= brightness;
 
+#if SHOW_EDGES
+	int count = 0;
+	float3 m3 = abs( ( IN.texCoord.xyz - 0.5 ) * 2.0f );
+	if ( m3.x > EDGE_THICKNESS ) count++;
+	if ( m3.y > EDGE_THICKNESS ) count++;
+	if ( m3.z > EDGE_THICKNESS ) count++;
+	if ( count >= 2 )
+	{
+		colour = float3(0,0,50);
+	}
+#endif 
 	return float4( colour / 3, 1 );
 }
-
diff --git a/Engine/Shaders/LPVWriteCommon.usf b/Engine/Shaders/LPVWriteCommon.usf
--- a/Engine/Shaders/LPVWriteCommon.usf
+++ b/Engine/Shaders/LPVWriteCommon.usf
@@ -102,22 +102,6 @@ VplListEntry PackVpl( Vpl vpl )
 void AccumulateLighting(float3 colour, float3 n, float solidAngle, in out LPVCell cell ) 
 { 
 	const float ONE_OVER_PI = 1.0f/3.1415927f;
-#if LPV_AMBIENT_CUBE
-	// Calculate weights
-	colour *= solidAngle;
-	float weight[6];
-	weight[AMBIENT_CUBE_NEGATIVE_X] = saturate(-n.x );
-	weight[AMBIENT_CUBE_POSITIVE_X] = saturate( n.x );
-	weight[AMBIENT_CUBE_NEGATIVE_Y] = saturate(-n.y );
-	weight[AMBIENT_CUBE_POSITIVE_Y] = saturate( n.y );
-	weight[AMBIENT_CUBE_NEGATIVE_Z] = saturate(-n.z );
-	weight[AMBIENT_CUBE_POSITIVE_Z] = saturate( n.z );
-	[unroll]
-	for ( int i=0; i<6; i++ )
-	{
-		cell.ambientCube[i] += colour * weight[i];
-	}
-#else
 	const float c0 = 0.282095f;
 	const float c1 = 0.488603f;
 	const float c2 = 1.092548f;
@@ -135,7 +119,6 @@ void AccumulateLighting(float3 colour, float3 n, float solidAngle, in out LPVCel
 	cell.coeffs[7] += colour*(c2*n.x*n.z)*solidAngle ; 
 	cell.coeffs[6] += colour*(c3*(3.0f*n.z*n.z-1.0f))*solidAngle ; 
 	cell.coeffs[8] += colour*(c4*(n.x*n.x-n.y*n.y))*solidAngle ;
-#endif
 }
 
 //-------------------------------------------------------------------------------------------------
@@ -143,24 +126,14 @@ void AccumulateLighting(float3 colour, float3 n, float solidAngle, in out LPVCel
 #if LPV_WRITE_SHADER
 void WriteLpvCell( in LPVCell cell, uint cellIndex )
 {
-#if LPV_VOLUME_TEXTURE
 	int3 texIndex = IndexToGridPos( cellIndex );
 
-	gLpv3DTextureRW0[ texIndex ] = float4( cell.coeffs[0],     0.0f ); // tex0.a is reserved for secondary occlusion;
+	gLpv3DTextureRW0[ texIndex ] = float4( cell.coeffs[0],     cell.AO ); // tex0.a is reserved for secondary occlusion;
 	gLpv3DTextureRW1[ texIndex ] = float4( cell.coeffs[1].rgb, cell.coeffs[2].r );
 	gLpv3DTextureRW2[ texIndex ] = float4( cell.coeffs[2].gb,  cell.coeffs[3].rg );
 	gLpv3DTextureRW3[ texIndex ] = float4( cell.coeffs[3].b,   cell.coeffs[4].rgb );
 	gLpv3DTextureRW4[ texIndex ] = float4( cell.coeffs[5].rgb, cell.coeffs[6].r );
 	gLpv3DTextureRW5[ texIndex ] = float4( cell.coeffs[6].gb,  cell.coeffs[7].rg );
 	gLpv3DTextureRW6[ texIndex ] = float4( cell.coeffs[7].b,   cell.coeffs[8].rgb );
-#elif LPV_AMBIENT_CUBE
-	LPVCellPackedAmbientCube packedCube;
-	PackAmbientCube( cell.ambientCube, packedCube );
-	gLpvBufferRW[cellIndex] = packedCube;
-#else
-	LPVCellPackedSH packedCell;
-	PackSH3( cell.coeffs, packedCell );
-	gLpvBufferRW[cellIndex] = packedCell;
-#endif 
 }
-#endif
+#endif
\ No newline at end of file
diff --git a/Engine/Shaders/LPVWriteVplCommon.usf b/Engine/Shaders/LPVWriteVplCommon.usf
--- a/Engine/Shaders/LPVWriteVplCommon.usf
+++ b/Engine/Shaders/LPVWriteVplCommon.usf
@@ -16,11 +16,22 @@
 
 //-------------------------------------------------------------------------------------------------
 
+// Bind UAVs to specific indices for materials with the following flags to ensure the UAVs set match
+// up with what the shader is expecting. Only do this for these materials as it causes other compute
+// shaders which include this file to fail to compile.
+#if MATERIAL_INJECT_EMISSIVE_INTO_LPV || MATERIAL_BLOCK_GI
+RWStructuredBuffer<VplListEntry>				RWGvListBuffer : register (u2);		// 0
+RWByteAddressBuffer								RWGvListHeadBuffer : register (u3);	// 1
+
+RWStructuredBuffer<VplListEntry>				RWVplListBuffer : register (u4);		// 2
+RWByteAddressBuffer								RWVplListHeadBuffer : register (u5);	// 3
+#else
 RWStructuredBuffer<VplListEntry>				RWGvListBuffer;		// 0
 RWByteAddressBuffer								RWGvListHeadBuffer;	// 1
 
 RWStructuredBuffer<VplListEntry>				RWVplListBuffer;		// 2
 RWByteAddressBuffer								RWVplListHeadBuffer;	// 3
+#endif
 
 //-------------------------------------------------------------------------------------------------
 // Functions
@@ -34,7 +45,8 @@ int AddToVplList( float3 worldPos, float3 flux, float3 direction, const bool emi
 		// Allocate a new list entry 
 		uint newHeadIndex = RWVplListBuffer.IncrementCounter();
 		uint oldHeadIndex;
-		RWVplListHeadBuffer.InterlockedExchange( gridIndex*4, newHeadIndex, oldHeadIndex );
+		// Note: we add one to the list index to make the list 0-terminated and prevent infinite recursion
+		RWVplListHeadBuffer.InterlockedExchange( gridIndex*4, newHeadIndex+1, oldHeadIndex );
 
 		if ( emissiveInjection )
 		{
@@ -69,7 +81,8 @@ int AddToGvList( float3 worldPos, float3 flux, float3 direction )
 		// Allocate a new list entry 
 		uint newHeadIndex = RWGvListBuffer.IncrementCounter();
 		uint oldHeadIndex;
-		RWGvListHeadBuffer.InterlockedExchange( gridIndex*4, newHeadIndex, oldHeadIndex );
+		// Note: we add one to the list index to make the list 0-terminated and prevent infinite recursion
+		RWGvListHeadBuffer.InterlockedExchange( gridIndex*4, newHeadIndex+1, oldHeadIndex );
 
 		// Fill in the new list node 
 		Vpl vpl;
diff --git a/Engine/Shaders/PostProcessLpvIndirect.usf b/Engine/Shaders/PostProcessLpvIndirect.usf
--- a/Engine/Shaders/PostProcessLpvIndirect.usf
+++ b/Engine/Shaders/PostProcessLpvIndirect.usf
@@ -9,6 +9,15 @@
 //
 //				Copyright (C) Microsoft. All rights reserved.
 //-----------------------------------------------------------------------------
+ 
+#define ENABLE_LPV_SPECULAR				1
+
+// Directional occlusion
+#define DEBUG_USE_RAW_AO				0
+#define LPV_OCCLUSION_EDGE_FADE			0
+#define ENABLE_DIRECTIONAL_OCCLUSION	1
+
+#define LuminanceVec float3(0.299f,0.587f,0.114f)
 
 #include "Common.usf"
 #include "PostProcessCommon.usf"
@@ -16,46 +25,57 @@
 #include "Random.usf"
 #include "LPVFinalPass.usf"
 
-#define MULTIPLE_SPEC_LOOKUPS 1
-
 float3 LpvSpecularHighQuality( float3 WorldPosition, float3 R, float Roughness )
 {
-	// Decreased roughness results in longer ray offsets and higher frequency reflections 
 	float Smooth = 1.0f-Roughness;
 	float offset = Smooth * Smooth * Smooth * Smooth *  8.0f;
-	float3 specularOut = LPVSpecular( WorldPosition, R, offset ); 
 
+	// Decreased roughness results in longer ray offsets and higher frequency reflections 
+	float3 specularOut = LPVSpecular( WorldPosition, R, offset ); 
 	// To prevent artifacts from very long ray lookups, we do multiple lookups and combine 
 	// with a weighting favouring brighter/further samples
-#if MULTIPLE_SPEC_LOOKUPS
 	[branch] 
-	if ( offset > 5.0f )
+	if ( offset > 5.0f ) // TODO: Use AO to block rays
 	{
 		float3 spec1 = LPVSpecular( WorldPosition, R, offset * 0.5 ); 
-		float l1 = dot( spec1, float3(0.299f,0.587f,0.114f) ) * 0.25;
-		float l2 = dot( specularOut, float3(0.299f,0.587f,0.114f) ) * 0.75;
-		specularOut = lerp( spec1, specularOut, l2/((l1+l2)+0.0001f) );
+		float l1 = dot( spec1,			LuminanceVec ) * 0.25;
+		float l2 = dot( specularOut,	LuminanceVec ) * 0.75;
+		float weight = l2/((l1+l2)+0.0001f);
+		specularOut = lerp( spec1, specularOut, weight );
 	}
 	[branch] 
 	if ( offset > 1.0f )
 	{
-		float3 spec1 = LPVLookup( WorldPosition, R ); 
-		float l1 = dot( spec1, float3(0.299f,0.587f,0.114f) ) * 0.25;
-		float l2 = dot( specularOut, float3(0.299f,0.587f,0.114f) ) * 0.75;
-		specularOut = lerp( spec1, specularOut, l2/((l1+l2)+0.0001f) );
+		float3 spec1 = ( LPVDiffuse( WorldPosition, R ) / LpvRead.DiffuseIntensity ) * LpvRead.SpecularIntensity; // FIX A BUG. This is horrible!
+		float l1 = dot( spec1,			LuminanceVec ) * 0.25;
+		float l2 = dot( specularOut,	LuminanceVec ) * 0.75;
+		float weight = l2/((l1+l2)+0.0001f);
+		specularOut = lerp( spec1, specularOut, weight );
 	}
-#endif
-
 	return specularOut;
 }
 
-void MainPS(in float4 UVAndScreenPos : TEXCOORD0, float4 SVPos: SV_POSITION, out float4 OutColor : SV_Target0)
+void MainPS(in float4 UVAndScreenPos : TEXCOORD0, float4 SVPos: SV_POSITION, out float4 OutColor : SV_Target0
+#if APPLY_SEPARATE_SPECULAR_RT
+			,out float4 OutSpecular : SV_Target1
+#endif
+			)
 {
 	float2 UV = UVAndScreenPos.xy; 
 
 	FScreenSpaceData ScreenSpaceData = GetScreenSpaceData(UV);  
 	FGBufferData GBuffer = ScreenSpaceData.GBuffer;
 	 
+	OutColor = 0;
+#if APPLY_SEPARATE_SPECULAR_RT
+	OutSpecular = 0;
+#endif
+
+	BRANCH if (GBuffer.ShadingModelID == SHADINGMODELID_UNLIT )
+	{
+		return;
+	}
+
 	// screen position in [-1, 1] screen space 
 	float2 ScreenSpacePos = UVAndScreenPos.zw;
 
@@ -66,19 +86,39 @@ void MainPS(in float4 UVAndScreenPos : TEXCOORD0, float4 SVPos: SV_POSITION, out
 	float3 WorldPosition = ScreenVectorU * SceneDepth + View.ViewOrigin.xyz;
 
 	float3 grid = WorldToGrid( WorldPosition );
-	float minGrid = min( grid.x, min(grid.y,grid.z ) );
-	float maxGrid = max( grid.x, max(grid.y,grid.z ) );
 
-	[branch]
-	if ( minGrid >= 0.0f && maxGrid <= 32.0f )
+#if 0 // Potential optimisation - might give worse results if we're outside the grid though
+	float maxgrid = max(grid.x, max(grid.y, grid.z) );
+	float mingrid = min(grid.x, min(grid.y, grid.z) );
+
+	if ( maxgrid <32 && mingrid > 0)
+#endif
 	{
 		float3 R = reflect( ScreenVector, GBuffer.WorldNormal );
 
-		float3 LpvDiffuse = LPVLookup( WorldPosition, GBuffer.WorldNormal );
+		float3 DiffuseColor  = GBuffer.DiffuseColor;
+		float3 DiffuseNormal = GBuffer.WorldNormal;
+		float3 SubsurfaceColor = ExtractSubsurfaceColor(ScreenSpaceData.GBuffer);
+		float3 Diffuse = LPVDiffuse( WorldPosition, DiffuseNormal ) * DiffuseColor;
+
+#if TRANSMISSIVE_LIGHTING_MODEL_LPV
+		BRANCH if( GBuffer.ShadingModelID != SHADINGMODELID_TRANSMISSIVE )
+		{
+			Diffuse += 0.00001f;
+		}
+		else
+		{
+			float3 SubsurfaceLighting = 0;
 
-		float3 Diffuse = GBuffer.DiffuseColor * LpvDiffuse;
+			float3 BackScatter = ( LPVDiffuse( WorldPosition, -GBuffer.WorldNormal ) ) * GBuffer.Opacity;
+			SubsurfaceLighting += BackScatter * SubsurfaceColor;
+
+			Diffuse += SubsurfaceLighting * 1.0f ;
+		}
+#endif
 
 		float3 LpvSpecular = 0;
+
 		[branch]
 		if ( dot( GBuffer.SpecularColor.xyz, GBuffer.SpecularColor.xyz ) >= 0.00001f )
 		{
@@ -87,18 +127,60 @@ void MainPS(in float4 UVAndScreenPos : TEXCOORD0, float4 SVPos: SV_POSITION, out
 	
 		// apply darkening from ambient occlusion (does not use PostprocessInput1 to set white texture if SSAO is off)
 		float AmbientOcclusion = GBuffer.GBufferAO * ScreenSpaceData.AmbientOcclusion;
+		OutColor = float4( Diffuse.xyz, 1 ) * AmbientOcclusion;
+
+#if !ENABLE_LPV_SPECULAR
+		LpvSpecular = 0;
+#endif 
+
+#if APPLY_SEPARATE_SPECULAR_RT
 
-		// diffuse
-		OutColor = float4( Diffuse.xyz, 0 );
+		// Ambient occlusion already applied with Reflection Environment
+		// Note: we don't need to apply the BRDF, since it gets applied during reflection apply
+		// Alpha = 1.0f;
+		OutSpecular = float4( LpvSpecular, 1 );
 
-		// specular
-		OutColor += float4( LpvSpecular * EnvBRDF( GBuffer.SpecularColor, GBuffer.Roughness, dot(GBuffer.WorldNormal, -ScreenVector) ), 0 );
+#else // APPLY_SEPARATE_SPECULAR_RT
 
-		// AO
-		OutColor.rgb *= AmbientOcclusion;
+		// Combine Specular (same behaviour as PostProcessAmbient)
+		OutColor.rgb += LpvSpecular;
+
+#endif // APPLY_SEPARATE_SPECULAR_RT
 	}
-	else
+}
+
+
+
+
+void DirectionalOcclusionPS(in float4 UVAndScreenPos : TEXCOORD0, float4 SVPos: SV_POSITION, out float4 OutColor : SV_Target0 )
+{
+	OutColor = float4(1,1,1,1);
+#if ENABLE_DIRECTIONAL_OCCLUSION 
+	float2 UV = UVAndScreenPos.xy; 
+
+	FScreenSpaceData ScreenSpaceData = GetScreenSpaceData(UV);  
+	FGBufferData GBuffer = ScreenSpaceData.GBuffer;
+
+	BRANCH if (GBuffer.ShadingModelID != SHADINGMODELID_UNLIT )
 	{
-		OutColor = 0;
+		float2 ScreenSpacePos = UVAndScreenPos.zw;
+
+		float3 ScreenVectorU = mul(float4(ScreenSpacePos, 1, 0), View.ScreenToWorld).xyz;
+		float3 ScreenVector = normalize(ScreenVectorU); 
+
+		float SceneDepth = CalcSceneDepth(UV);
+		float3 WorldPosition = ScreenVectorU * SceneDepth + View.ViewOrigin.xyz;
+
+		float3 grid = WorldToGrid( WorldPosition );
+
+		float3 R = reflect( ScreenVector, GBuffer.WorldNormal );
+
+		float3 DiffuseNormal = GBuffer.WorldNormal;
+
+		float3 SubsurfaceColor = ExtractSubsurfaceColor( ScreenSpaceData.GBuffer );
+
+		OutColor = LPVGetDirectionalOcclusion_DiffuseSpec( WorldPosition, R, GBuffer.Roughness, DiffuseNormal ).xyxy;
 	}
+
+	#endif  // ENABLE_DIRECTIONAL_OCCLUSION
 }
diff --git a/Engine/Shaders/ReflectionEnvironmentShaders.usf b/Engine/Shaders/ReflectionEnvironmentShaders.usf
--- a/Engine/Shaders/ReflectionEnvironmentShaders.usf
+++ b/Engine/Shaders/ReflectionEnvironmentShaders.usf
@@ -508,6 +508,8 @@ void ReflectionApplyPS(
 			// So, we combine the two for the best of both. This is done by removing the low spatial frequencies from the envmap and replacing them with the lightmap data.
 			// This is only done with luma so as to not get odd color shifting.
 			// Note: make sure this matches the lightmap mixing done for translucency (BasePassPixelShader.usf)
+
+			// NOTE: this breaks LPV specular if r.allowstaticlighting is enabled
 			SpecularLighting.rgb *= GBuffer.IndirectIrradiance;
 		}
 #endif
@@ -523,6 +525,9 @@ void ReflectionApplyPS(
 			bool bNormalize = SkyLightParameters.z < 1 && View.UseLightmaps > 0 && !ALLOW_STATIC_LIGHTING;
 			float3 SkyLighting = GetSkyLightReflection(ReflectionVector, GBuffer.Roughness, bNormalize);
 
+		    float DirectionalOcclusion = ScreenSpaceData.DirectionalOcclusion.g;
+		    SkyLighting *= DirectionalOcclusion;
+
 			FLATTEN
 			if (bNormalize)
 			{
diff --git a/Engine/Shaders/ShaderVersion.usf b/Engine/Shaders/ShaderVersion.usf
--- a/Engine/Shaders/ShaderVersion.usf
+++ b/Engine/Shaders/ShaderVersion.usf
@@ -2,4 +2,4 @@
 // Each time the console command is executed it generates a new GUID. As this file is included
 // in common.usf (which should be included in any shader) it allows to invalidate the shader DDC.
 // 
-// GUID = 8636091346FB77EB0D592BB4EAD2CD9D
\ No newline at end of file
+// GUID = 639690ED421CFC45604E3687590EE0F7
\ No newline at end of file
diff --git a/Engine/Shaders/ShadowDepthPixelShader.usf b/Engine/Shaders/ShadowDepthPixelShader.usf
--- a/Engine/Shaders/ShadowDepthPixelShader.usf
+++ b/Engine/Shaders/ShadowDepthPixelShader.usf
@@ -14,6 +14,13 @@
 
 #define SECONDARY_OCCLUSION			1
 
+
+
+
+// With this enabled, very large or one-sided objects (e.g. terrain) inject an additional layer of GV cells 
+// to increase AO contribution. The AO calculation works best with closed geometry and small-medium scale objects
+#define LPV_GV_BOOST				1
+
 //@todo-rco: Remove this when reenabling for OpenGL
 #if !COMPILER_GLSL
 #if REFLECTIVE_SHADOW_MAP
@@ -84,6 +91,13 @@ void Main(
 #if !COMPILER_GLSL
 #if MATERIAL_INJECT_EMISSIVE_INTO_LPV
 	float3 emissive = GetMaterialEmissive( MaterialParameters );
+
+	// Remove directionality from the light in the VPL.
+	normal = float3 (0.0f, 0.0f, 0.0f);
+#endif
+
+#if MATERIAL_INJECT_EMISSIVE_INTO_LPV || MATERIAL_BLOCK_GI
+	float opacity = GetMaterialOpacity( MaterialParameters );
 #endif
 
 #if SECONDARY_OCCLUSION || MATERIAL_INJECT_EMISSIVE_INTO_LPV
@@ -91,21 +105,43 @@ void Main(
 	float3 grid = WorldToGrid( MaterialParameters.WorldPosition.xyz );
 	float minGrid = min( grid.x, min(grid.y,grid.z ) );
 	float maxGrid = max( grid.x, max(grid.y,grid.z ) );
+
 	[branch]
-	if ( minGrid >= 0.0f && maxGrid <= 32.0f )
+	if ( minGrid > -0.5f && maxGrid <= 32.5f ) // Account for 0.5 cell offset
 	{
-#if SECONDARY_OCCLUSION && !( MATERIALBLENDING_TRANSLUCENT || MATERIALBLENDING_ALPHACOMPOSITE )
-		AddToGvList( MaterialParameters.WorldPosition - normal * LpvWrite.GeometryVolumeInjectionBias, flux, normal ); 
+#if SECONDARY_OCCLUSION && (MATERIALBLENDING_MASKED || MATERIALBLENDING_SOLID) && !MATERIAL_BLOCK_GI
+		float3 GvInjectPosition = MaterialParameters.WorldPosition - normal * LpvWrite.GeometryVolumeInjectionBias;
+#if LPV_GV_BOOST
+		if ( Primitive.LpvBiasMultiplier < 0.25 ) 
+		{
+			// Add an additional layer of GV cells facing the opposite direction to improve AO contribution
+			// TODO: add per-primitive AO-boost tickbox instead of hijacking LpvBiasMultiplier like this
+			// NOTE: The cost of this additional branch is minimal (on AMD at least)
+			float offset = ( LpvWrite.LpvScale * 3.0f );
+			AddToGvList( GvInjectPosition - normal * offset, flux, -normal );
+		}
 #endif
+		AddToGvList( GvInjectPosition, flux, normal ); 
 
-#if MATERIAL_INJECT_EMISSIVE_INTO_LPV
+#endif
+
+#if MATERIAL_INJECT_EMISSIVE_INTO_LPV && !MATERIAL_BLOCK_GI
 		AddToVplList( MaterialParameters.WorldPosition, emissive * LpvWrite.EmissiveInjectionMultiplier, normal, true );
 #endif
 	}
+	else
+	{
+		OutColor = 0;
+		OutReflectiveShadowMapDiffuse = 0;
+		return;
+	}
 #endif
 #endif
 
-#if MATERIALBLENDING_TRANSLUCENT || MATERIALBLENDING_ALPHACOMPOSITE
+#if MATERIAL_BLOCK_GI
+	OutColor = 1.0f;
+	OutReflectiveShadowMapDiffuse = opacity;
+#elif MATERIALBLENDING_TRANSLUCENT || MATERIALBLENDING_ALPHACOMPOSITE
 	OutColor = 0;
 	OutReflectiveShadowMapDiffuse = 0;
 	clip(-1);
diff --git a/Engine/Shaders/SkyLighting.usf b/Engine/Shaders/SkyLighting.usf
--- a/Engine/Shaders/SkyLighting.usf
+++ b/Engine/Shaders/SkyLighting.usf
@@ -104,7 +104,11 @@ void SkyLightDiffusePS(
 
 		// Compute the preconvolved incoming lighting with the bent normal direction
 		float3 DiffuseLookup = GetSkySHDiffuse(SkyLightingNormal) * View.SkyLightColor.rgb;
-		
+
+		// darken by DirectionalOcclusion feature from LightPropagationVolume (could be optimized if the feature is not used)
+		float DirectionalOcclusion = ScreenSpaceData.DirectionalOcclusion.r;
+		DiffuseLookup *= DirectionalOcclusion;
+
 		// Apply AO to the sky diffuse and account for darkening due to the geometry term
 		Lighting += (ScalarFactors * DotProductFactor) * DiffuseLookup * GBuffer.DiffuseColor + (1 - SkyVisibility) * OcclusionTintAndMinOcclusion.xyz;
 
@@ -125,4 +129,4 @@ void SkyLightDiffusePS(
 		LightAccumulator_Add(LightAccumulator, Lighting, 0, 1.0f);
 		OutColor = LightAccumulator_GetResult(LightAccumulator);
 	}
-}
\ No newline at end of file
+}
diff --git a/Engine/Source/Runtime/Engine/Classes/Engine/Scene.h b/Engine/Source/Runtime/Engine/Classes/Engine/Scene.h
--- a/Engine/Source/Runtime/Engine/Classes/Engine/Scene.h
+++ b/Engine/Source/Runtime/Engine/Classes/Engine/Scene.h
@@ -266,8 +266,23 @@ struct FPostProcessSettings
 	UPROPERTY(BlueprintReadWrite, Category=Overrides, meta=(PinHiddenByDefault))
 	uint32 bOverride_LPVIntensity:1;
 
-	UPROPERTY(BlueprintReadWrite, Category=Overrides, meta=(PinHiddenByDefault))
-	uint32 bOverride_LPVWarpIntensity:1;
+	UPROPERTY()
+	uint32 bOverride_LPVDirectionalOcclusionIntensity:1;
+
+	UPROPERTY()
+	uint32 bOverride_LPVDirectionalOcclusionRadius:1;
+
+	UPROPERTY()
+	uint32 bOverride_LPVDiffuseOcclusionExponent:1;
+
+	UPROPERTY()
+	uint32 bOverride_LPVSpecularOcclusionExponent:1;
+
+	UPROPERTY()
+	uint32 bOverride_LPVDiffuseOcclusionIntensity:1;
+
+	UPROPERTY()
+	uint32 bOverride_LPVSpecularOcclusionIntensity:1;
 
 	UPROPERTY(BlueprintReadWrite, Category=Overrides, meta=(PinHiddenByDefault))
 	uint32 bOverride_LPVSize:1;
@@ -287,9 +302,6 @@ struct FPostProcessSettings
 	UPROPERTY(BlueprintReadWrite, Category=Overrides, meta=(PinHiddenByDefault))
 	uint32 bOverride_LPVEmissiveInjectionIntensity:1;
 
-	UPROPERTY(BlueprintReadWrite, Category=Overrides, meta=(PinHiddenByDefault))
-	uint32 bOverride_LPVTransmissionIntensity:1;
-
 	UPROPERTY(BlueprintReadWrite, Category=Overrides, meta=(PinHiddenByDefault))
 	uint32 bOverride_IndirectLightingColor:1;
 
@@ -545,12 +557,8 @@ struct FPostProcessSettings
 	UPROPERTY(interp, BlueprintReadWrite, Category=LightPropagationVolume, meta=(editcondition = "bOverride_LPVIntensity", UIMin = "0", UIMax = "20", DisplayName = "Intensity") )
 	float LPVIntensity;
 
-	/** CURRENTLY DISABLED - The strength of the warp offset for reducing light bleeding. 0.0 is off, 1.0 is the "normal" value, but higher values can be used to boost the effect*/
-	UPROPERTY(interp, BlueprintReadWrite, Category=LightPropagationVolume, AdvancedDisplay, meta=(editcondition = "bOverride_LPVWarpIntensity", UIMin = "0", UIMax = "1", DisplayName = "(DISABLED) Grid Warp Intensity"))
-	float LPVWarpIntensity;
-
 	/** Bias applied to light injected into the LPV in cell units. Increase to reduce bleeding through thin walls*/
-	UPROPERTY(interp, BlueprintReadWrite, Category=LightPropagationVolume, meta=(editcondition = "bOverride_LPVVplInjectionBias", UIMin = "0", UIMax = "2", DisplayName = "Light Injection Bias") )
+	UPROPERTY(interp, BlueprintReadWrite, Category=LightPropagationVolume, AdvancedDisplay, meta=(editcondition = "bOverride_LPVVplInjectionBias", UIMin = "0", UIMax = "2", DisplayName = "Light Injection Bias") )
 	float LPVVplInjectionBias;
 
 	/** The size of the LPV volume, in Unreal units*/
@@ -562,20 +570,40 @@ struct FPostProcessSettings
 	float LPVSecondaryOcclusionIntensity;
 
 	/** Secondary bounce light strength (bounce light shadows). Set to 0 to disable*/
-	UPROPERTY(interp, BlueprintReadWrite, Category=LightPropagationVolume, meta=(editcondition = "bOverride_LPVSecondaryBounceIntensity", UIMin = "0", UIMax = "1", DisplayName = "Secondary Bounce Intensity") )
+	UPROPERTY(interp, BlueprintReadWrite, Category=LightPropagationVolume, AdvancedDisplay, meta=(editcondition = "bOverride_LPVSecondaryBounceIntensity", UIMin = "0", UIMax = "1", DisplayName = "Secondary Bounce Intensity") )
 	float LPVSecondaryBounceIntensity;
 
 	/** Bias applied to the geometry volume in cell units. Increase to reduce darkening due to secondary occlusion */
 	UPROPERTY(interp, BlueprintReadWrite, Category=LightPropagationVolume, AdvancedDisplay, meta=(editcondition = "bOverride_LPVGeometryVolumeBias", UIMin = "0", UIMax = "2", DisplayName = "Geometry Volume Bias"))
 	float LPVGeometryVolumeBias;
 
-	UPROPERTY(interp, BlueprintReadWrite, Category=LightPropagationVolume, meta=(editcondition = "bOverride_LPVEmissiveInjectionIntensity", UIMin = "0", UIMax = "20", DisplayName = "Emissive Injection Intensity") )
+	UPROPERTY(interp, BlueprintReadWrite, Category=LightPropagationVolume, AdvancedDisplay, meta=(editcondition = "bOverride_LPVEmissiveInjectionIntensity", UIMin = "0", UIMax = "20", DisplayName = "Emissive Injection Intensity") )
 	float LPVEmissiveInjectionIntensity;
 
-	/** How strong light transmission from the LPV should be. 0.0 is off, 1.0 is the "normal" value, but higher values can be used to boost the effect*/
-	UPROPERTY(interp, BlueprintReadWrite, Category=LightPropagationVolume, meta=(editcondition = "bOverride_LPVTransmissionIntensity", UIMin = "0", UIMax = "2", DisplayName = "Transmission Intensity") )
-	float LPVTransmissionIntensity;
+	/** Controls the amount of directional occlusion. Requires LPV. Values very close to 1.0 are recommended */
+	UPROPERTY(interp, BlueprintReadWrite, Category=LightPropagationVolume, meta=(editcondition = "bOverride_LPVDirectionalOcclusionIntensity", UIMin = "0", UIMax = "1", DisplayName = "Occlusion Intensity") )
+	float LPVDirectionalOcclusionIntensity;
 
+	/** Occlusion Radius - 16 is recommended for most scenes */
+	UPROPERTY(interp, BlueprintReadWrite, Category=LightPropagationVolume, AdvancedDisplay, meta=(editcondition = "bOverride_LPVDirectionalOcclusionRadius", UIMin = "1", UIMax = "16", DisplayName = "Occlusion Radius") )
+	float LPVDirectionalOcclusionRadius;
+
+	/** Diffuse occlusion exponent - increase for more contrast. 1 to 2 is recommended */
+	UPROPERTY(interp, BlueprintReadWrite, Category=LightPropagationVolume, meta=(editcondition = "bOverride_LPVDiffuseOcclusionExponent", UIMin = "0.5", UIMax = "5", DisplayName = "Diffuse occlusion exponent") )
+	float LPVDiffuseOcclusionExponent;
+
+	/** Specular occlusion exponent - increase for more contrast. 6 to 9 is recommended */
+	UPROPERTY(interp, BlueprintReadWrite, Category=LightPropagationVolume, meta=(editcondition = "bOverride_LPVSpecularOcclusionExponent", UIMin = "1", UIMax = "16", DisplayName = "Specular occlusion exponent") )
+	float LPVSpecularOcclusionExponent;
+
+	/** Diffuse occlusion intensity - higher values provide increased diffuse occlusion.*/
+	UPROPERTY(interp, BlueprintReadWrite, Category=LightPropagationVolume, AdvancedDisplay, meta=(editcondition = "bOverride_LPVDiffuseOcclusionIntensity", UIMin = "0", UIMax = "4", DisplayName = "Diffuse occlusion intensity") )
+	float LPVDiffuseOcclusionIntensity;
+
+	/** Specular occlusion intensity - higher values provide increased specular occlusion.*/
+	UPROPERTY(interp, BlueprintReadWrite, Category=LightPropagationVolume, AdvancedDisplay, meta=(editcondition = "bOverride_LPVSpecularOcclusionIntensity", UIMin = "0", UIMax = "4", DisplayName = "Specular occlusion intensity") )
+	float LPVSpecularOcclusionIntensity;
+	
 	/** AmbientCubemap tint color */
 	UPROPERTY(interp, BlueprintReadWrite, Category=AmbientCubemap, AdvancedDisplay, meta=(editcondition = "bOverride_AmbientCubemapTint", DisplayName = "Tint", HideAlphaChannel))
 	FLinearColor AmbientCubemapTint;
@@ -932,14 +960,12 @@ struct FPostProcessSettings
 		AmbientCubemapIntensity = 1.0f;
 		AmbientCubemapTint = FLinearColor(1, 1, 1);
 		LPVIntensity = 1.0f;
-		LPVWarpIntensity = 0.0f;
 		LPVSize = 5312.0f;
 		LPVSecondaryOcclusionIntensity = 0.0f;
 		LPVSecondaryBounceIntensity = 0.0f;
 		LPVVplInjectionBias = 0.64f;
 		LPVGeometryVolumeBias = 0.384f;
 		LPVEmissiveInjectionIntensity = 1.0f;
-		LPVTransmissionIntensity = 1.0f;
 		AutoExposureLowPercent = 80.0f;
 		AutoExposureHighPercent = 98.3f;
 		// next value might get overwritten by r.DefaultFeature.AutoExposure
@@ -949,6 +975,12 @@ struct FPostProcessSettings
 		AutoExposureBias = 0.0f;
 		AutoExposureSpeedUp = 3.0f;
 		AutoExposureSpeedDown = 1.0f;
+		LPVDirectionalOcclusionIntensity = 0.0f;
+		LPVDirectionalOcclusionRadius = 8.0f;
+		LPVDiffuseOcclusionExponent = 1.0f;
+		LPVSpecularOcclusionExponent = 7.0f;
+		LPVDiffuseOcclusionIntensity = 1.0f;
+		LPVSpecularOcclusionIntensity = 1.0f;
 		HistogramLogMin = -8.0f;
 		HistogramLogMax = 4.0f;
 		// next value might get overwritten by r.DefaultFeature.LensFlare
diff --git a/Engine/Source/Runtime/Engine/Classes/Materials/Material.h b/Engine/Source/Runtime/Engine/Classes/Materials/Material.h
--- a/Engine/Source/Runtime/Engine/Classes/Materials/Material.h
+++ b/Engine/Source/Runtime/Engine/Classes/Materials/Material.h
@@ -38,7 +38,7 @@ struct FMaterialsWithDirtyUsageFlags
 
 	/** Query the annotation to see if the specified flag has been changed */	
 	bool IsUsageFlagDirty(EMaterialUsage UsageFlag);
-	
+
 	/** Default state for annotations (no flags changed)*/
 	static const FMaterialsWithDirtyUsageFlags DefaultAnnotation;
 };
@@ -483,6 +483,12 @@ public:
 	UPROPERTY(EditAnywhere, BlueprintReadWrite, Category=Material, meta=(DisplayName = "Emissive (Dynamic Area Light)"), AdvancedDisplay)
 	uint32 bUseEmissiveForDynamicAreaLighting : 1;
 
+	/**
+	 * If enabled, the material's opacity defines how much GI is blocked when using the LightPropagationVolume feature
+	 */
+	UPROPERTY(EditAnywhere, BlueprintReadWrite, Category=Material, meta=(DisplayName = "Block Global Illumination"), AdvancedDisplay)
+	uint32 bBlockGI : 1;
+
 	/** 
 	 * This is a special usage flag that allows a material to be assignable to any primitive type.
 	 * This is useful for materials used by code to implement certain viewmodes, for example the default material or lighting only material.
diff --git a/Engine/Source/Runtime/Engine/Private/Materials/Material.cpp b/Engine/Source/Runtime/Engine/Private/Materials/Material.cpp
--- a/Engine/Source/Runtime/Engine/Private/Materials/Material.cpp
+++ b/Engine/Source/Runtime/Engine/Private/Materials/Material.cpp
@@ -553,6 +553,7 @@ UMaterial::UMaterial(const FObjectInitializer& ObjectInitializer)
 	BlendablePriority = 0;
 
 	bUseEmissiveForDynamicAreaLighting = false;
+	bBlockGI = false;
 	RefractionDepthBias = 0.0f;
 	MaterialDecalResponse = MDR_ColorNormalRoughness;
 
diff --git a/Engine/Source/Runtime/Engine/Private/Materials/MaterialShared.cpp b/Engine/Source/Runtime/Engine/Private/Materials/MaterialShared.cpp
--- a/Engine/Source/Runtime/Engine/Private/Materials/MaterialShared.cpp
+++ b/Engine/Source/Runtime/Engine/Private/Materials/MaterialShared.cpp
@@ -698,6 +698,7 @@ bool IsTranslucentBlendMode(EBlendMode BlendMode)
 int32 FMaterialResource::GetMaterialDomain() const { return Material->MaterialDomain; }
 bool FMaterialResource::IsTangentSpaceNormal() const { return Material->bTangentSpaceNormal || Material->Normal.Expression == NULL; }
 bool FMaterialResource::ShouldInjectEmissiveIntoLPV() const { return Material->bUseEmissiveForDynamicAreaLighting; }
+bool FMaterialResource::ShouldBlockGI() const { return Material->bBlockGI; }
 bool FMaterialResource::ShouldGenerateSphericalParticleNormals() const { return Material->bGenerateSphericalParticleNormals; }
 bool FMaterialResource::ShouldDisableDepthTest() const { return Material->bDisableDepthTest; }
 bool FMaterialResource::ShouldEnableResponsiveAA() const { return Material->bEnableResponsiveAA; }
@@ -1203,6 +1204,7 @@ void FMaterial::SetupMaterialEnvironment(
 	OutEnvironment.SetDefine(TEXT("MATERIAL_USE_LM_DIRECTIONALITY"), UseLmDirectionality() ? TEXT("1") : TEXT("0"));
 	OutEnvironment.SetDefine(TEXT("MATERIAL_INJECT_EMISSIVE_INTO_LPV"), ShouldInjectEmissiveIntoLPV() ? TEXT("1") : TEXT("0"));
 	OutEnvironment.SetDefine(TEXT("MATERIAL_SSR"), ShouldDoSSR() ? TEXT("1") : TEXT("0"));
+	OutEnvironment.SetDefine(TEXT("MATERIAL_BLOCK_GI"), ShouldBlockGI() ? TEXT("1") : TEXT("0"));
 
 	{
 		auto DecalBlendMode = (EDecalBlendMode)GetDecalBlendMode();
@@ -1756,8 +1758,12 @@ FString FMaterialResource::GetMaterialUsageDescription() const
 
 	// this changed from ",SpecialEngine, TwoSided" to ",SpecialEngine=1, TwoSided=1, TSNormal=0, ..." to be more readable
 	BaseDescription += FString::Printf(
-		TEXT("SpecialEngine=%d, TwoSided=%d, TSNormal=%d, InjectEmissiveIntoLPV=%d, Masked=%d, Distorted=%d, Usage={"),
-		(int32)IsSpecialEngineMaterial(), (int32)IsTwoSided(), (int32)IsTangentSpaceNormal(), (int32)ShouldInjectEmissiveIntoLPV(), (int32)IsMasked(), (int32)IsDistorted());
+		TEXT("SpecialEngine=%d, TwoSided=%d, TSNormal=%d, InjectEmissiveIntoLPV=%d, Masked=%d, Distorted=%d")
+		TEXT(", BlockGI=%d")
+		TEXT(", Usage={"),
+		(int32)IsSpecialEngineMaterial(), (int32)IsTwoSided(), (int32)IsTangentSpaceNormal(), (int32)ShouldInjectEmissiveIntoLPV(), (int32)IsMasked(), (int32)IsDistorted()
+		, (int32)ShouldBlockGI()
+		);
 
 	bool bFirst = true;
 	for (int32 MaterialUsageIndex = 0; MaterialUsageIndex < MATUSAGE_MAX; MaterialUsageIndex++)
diff --git a/Engine/Source/Runtime/Engine/Private/PhysicsEngine/PhysDrawing.cpp b/Engine/Source/Runtime/Engine/Private/PhysicsEngine/PhysDrawing.cpp
--- a/Engine/Source/Runtime/Engine/Private/PhysicsEngine/PhysDrawing.cpp
+++ b/Engine/Source/Runtime/Engine/Private/PhysicsEngine/PhysDrawing.cpp
@@ -626,8 +626,8 @@ void FKAggregateGeom::GetAggGeom(const FTransform& Transform, const FColor Color
 				FBoxSphereBounds WorldBounds, LocalBounds;
 				CalcBoxSphereBounds(WorldBounds, LocalToWorld);
 				CalcBoxSphereBounds(LocalBounds, FTransform::Identity);
-				BatchElement.PrimitiveUniformBuffer = CreatePrimitiveUniformBufferImmediate(LocalToWorld.ToMatrixWithScale(), WorldBounds, LocalBounds, true, bUseEditorDepthTest);
-				// previous l2w not used so treat as static
+			    BatchElement.PrimitiveUniformBuffer = CreatePrimitiveUniformBufferImmediate(LocalToWorld.ToMatrixWithScale(), WorldBounds, LocalBounds, 1.0f, true, bUseEditorDepthTest);
+			 	// previous l2w not used so treat as static
 				BatchElement.FirstIndex = 0;
 				BatchElement.NumPrimitives = RenderInfo->IndexBuffer->Indices.Num() / 3;
 				BatchElement.MinVertexIndex = 0;
diff --git a/Engine/Source/Runtime/Engine/Private/SceneView.cpp b/Engine/Source/Runtime/Engine/Private/SceneView.cpp
--- a/Engine/Source/Runtime/Engine/Private/SceneView.cpp
+++ b/Engine/Source/Runtime/Engine/Private/SceneView.cpp
@@ -802,8 +802,12 @@ void FSceneView::OverridePostProcessSettings(const FPostProcessSettings& Src, fl
 	LERP_PP(AmbientCubemapIntensity);
 	LERP_PP(AmbientCubemapTint);
 	LERP_PP(LPVIntensity);
-	LERP_PP(LPVWarpIntensity);
-	LERP_PP(LPVTransmissionIntensity);
+	LERP_PP(LPVDirectionalOcclusionIntensity);
+	LERP_PP(LPVDirectionalOcclusionRadius);
+	LERP_PP(LPVSpecularOcclusionExponent);
+	LERP_PP(LPVDiffuseOcclusionExponent);
+	LERP_PP(LPVDiffuseOcclusionIntensity);
+	LERP_PP(LPVSpecularOcclusionIntensity);
 	LERP_PP(AutoExposureLowPercent);
 	LERP_PP(AutoExposureHighPercent);
 	LERP_PP(AutoExposureMinBrightness);
@@ -1078,6 +1082,16 @@ void FSceneView::StartFinalPostprocessSettings(FVector InViewLocation)
 
 void FSceneView::EndFinalPostprocessSettings(const FSceneViewInitOptions& ViewInitOptions)
 {
+	if(FinalPostProcessSettings.LPVDirectionalOcclusionIntensity < 0.001f)
+	{
+		FinalPostProcessSettings.LPVDirectionalOcclusionIntensity = 0.0f;
+	}
+
+	if(FinalPostProcessSettings.LPVIntensity < 0.001f)
+	{
+		FinalPostProcessSettings.LPVIntensity = 0.0f;
+	}
+
 	{
 		static const auto CVarMobileMSAA = IConsoleManager::Get().FindTConsoleVariableDataInt(TEXT("r.MobileMSAA"));
 		if(CVarMobileMSAA ? CVarMobileMSAA->GetValueOnGameThread() > 1 : false)
diff --git a/Engine/Source/Runtime/Engine/Public/MaterialShared.h b/Engine/Source/Runtime/Engine/Public/MaterialShared.h
--- a/Engine/Source/Runtime/Engine/Public/MaterialShared.h
+++ b/Engine/Source/Runtime/Engine/Public/MaterialShared.h
@@ -902,6 +902,7 @@ public:
 	virtual bool IsDitheredLODTransition() const = 0;
 	virtual bool IsTangentSpaceNormal() const { return false; }
 	virtual bool ShouldInjectEmissiveIntoLPV() const { return false; }
+	virtual bool ShouldBlockGI() const { return false; }
 	virtual bool ShouldGenerateSphericalParticleNormals() const { return false; }
 	virtual	bool ShouldDisableDepthTest() const { return false; }
 	virtual	bool ShouldEnableResponsiveAA() const { return false; }
@@ -1458,6 +1459,7 @@ public:
 	ENGINE_API virtual bool IsDitheredLODTransition() const override;
 	ENGINE_API virtual bool IsTangentSpaceNormal() const override;
 	ENGINE_API virtual bool ShouldInjectEmissiveIntoLPV() const override;
+	ENGINE_API virtual bool ShouldBlockGI() const override;
 	ENGINE_API virtual bool ShouldGenerateSphericalParticleNormals() const override;
 	ENGINE_API virtual bool ShouldDisableDepthTest() const override;
 	ENGINE_API virtual bool ShouldEnableResponsiveAA() const override;
diff --git a/Engine/Source/Runtime/Engine/Public/ShowFlags.h b/Engine/Source/Runtime/Engine/Public/ShowFlags.h
--- a/Engine/Source/Runtime/Engine/Public/ShowFlags.h
+++ b/Engine/Source/Runtime/Engine/Public/ShowFlags.h
@@ -274,7 +274,6 @@ private:
 		SetVisualizeHDR(false);
 		SetOverrideDiffuseAndSpecular(false);
 		SetReflectionOverride(false);
-		SetLpvLightingOnly(false);
 		SetVisualizeBuffer(false);
 		SetVectorFields(false);
 		SetGBufferHints(false);
diff --git a/Engine/Source/Runtime/Engine/Public/ShowFlagsValues.inl b/Engine/Source/Runtime/Engine/Public/ShowFlagsValues.inl
--- a/Engine/Source/Runtime/Engine/Public/ShowFlagsValues.inl
+++ b/Engine/Source/Runtime/Engine/Public/ShowFlagsValues.inl
@@ -266,8 +266,6 @@ SHOWFLAG_ALWAYS_ACCESSIBLE(PreviewShadowsIndicator, SFG_Visualize, LOCTEXT("Prev
 SHOWFLAG_ALWAYS_ACCESSIBLE(PrecomputedVisibilityCells, SFG_Visualize, LOCTEXT("PrecomputedVisibilityCellsSF", "Precomputed Visibility Cells"))
 /** Visualize volume lighting samples used for GI on dynamic objects */
 SHOWFLAG_ALWAYS_ACCESSIBLE(VolumeLightingSamples, SFG_Visualize, LOCTEXT("VolumeLightingSamplesSF", "Volume Lighting Samples"))
-/** needed for VMI_LpvLightingViewMode, Whether to show only LPV lighting*/
-SHOWFLAG_ALWAYS_ACCESSIBLE(LpvLightingOnly, SFG_Hidden, LOCTEXT("VisualizeLPVSF_ViewMode", "Visualize LPV"))
 /** Render Paper2D sprites */
 SHOWFLAG_ALWAYS_ACCESSIBLE(Paper2DSprites, SFG_Advanced, LOCTEXT("Paper2DSpritesSF", "Paper 2D Sprites"))
 /** Visualization of distance field AO */
diff --git a/Engine/Source/Runtime/Renderer/Private/CompositionLighting/CompositionLighting.cpp b/Engine/Source/Runtime/Renderer/Private/CompositionLighting/CompositionLighting.cpp
--- a/Engine/Source/Runtime/Renderer/Private/CompositionLighting/CompositionLighting.cpp
+++ b/Engine/Source/Runtime/Renderer/Private/CompositionLighting/CompositionLighting.cpp
@@ -65,15 +65,6 @@ static bool IsLpvIndirectPassRequired(FPostprocessContext& Context)
 	return false;
 }
 
-static void AddPostProcessingLpvIndirect(FPostprocessContext& Context, FRenderingCompositePass* SSAO )
-{
-	FRenderingCompositePass* Pass = Context.Graph.RegisterPass(new FRCPassPostProcessLpvIndirect());
-	Pass->SetInput(ePId_Input0, Context.FinalOutput);
-	Pass->SetInput(ePId_Input1, SSAO );
-
-	Context.FinalOutput = FRenderingCompositeOutputRef(Pass);
-}
-
 static bool IsReflectionEnvironmentActive(FPostprocessContext& Context)
 {
 	FScene* Scene = (FScene*)Context.View.Family->Scene;
@@ -325,7 +316,35 @@ void FCompositionLighting::ProcessAfterBasePass(FRHICommandListImmediate& RHICmd
 }
 
 
-void FCompositionLighting::ProcessLighting(FRHICommandListImmediate& RHICmdList, FViewInfo& View)
+void FCompositionLighting::ProcessLpvIndirect(FRHICommandListImmediate& RHICmdList, FViewInfo& View)
+{
+	check(IsInRenderingThread());
+	
+	FMemMark Mark(FMemStack::Get());
+	FRenderingCompositePassContext CompositeContext(RHICmdList, View);
+	FPostprocessContext Context(CompositeContext.Graph, View);
+
+	if(IsLpvIndirectPassRequired(Context))
+	{
+		FRenderingCompositePass* SSAO = Context.Graph.RegisterPass(new FRCPassPostProcessInput(GSceneRenderTargets.ScreenSpaceAO));
+
+		FRenderingCompositePass* Pass = Context.Graph.RegisterPass(new FRCPassPostProcessLpvIndirect());
+		Pass->SetInput(ePId_Input0, Context.FinalOutput);
+		Pass->SetInput(ePId_Input1, SSAO );
+
+		Context.FinalOutput = FRenderingCompositeOutputRef(Pass);
+	}
+
+	// The graph setup should be finished before this line ----------------------------------------
+
+	SCOPED_DRAW_EVENT(RHICmdList, CompositionLpvIndirect);
+
+	// we don't replace the final element with the scenecolor because this is what those passes should do by themself
+
+	CompositeContext.Process(Context.FinalOutput.GetPass(), TEXT("CompositionLighting"));
+}
+
+void FCompositionLighting::ProcessAfterLighting(FRHICommandListImmediate& RHICmdList, FViewInfo& View)
 {
 	check(IsInRenderingThread());
 	
@@ -336,17 +355,9 @@ void FCompositionLighting::ProcessLighting(FRHICommandListImmediate& RHICmdList,
 	{
 		FMemMark Mark(FMemStack::Get());
 		FRenderingCompositePassContext CompositeContext(RHICmdList, View);
-
 		FPostprocessContext Context(CompositeContext.Graph, View);
-
 		FRenderingCompositeOutputRef AmbientOcclusion;
 
-		if(IsLpvIndirectPassRequired(Context))
-		{
-			FRenderingCompositePass* SSAO = Context.Graph.RegisterPass(new FRCPassPostProcessInput(GSceneRenderTargets.ScreenSpaceAO));
-			AddPostProcessingLpvIndirect( Context, SSAO );
-		}
-
 		// Screen Space Subsurface Scattering
 		{
 			float Radius = CVarSSSScale.GetValueOnRenderThread();
@@ -385,7 +396,7 @@ void FCompositionLighting::ProcessLighting(FRHICommandListImmediate& RHICmdList,
 
 		// The graph setup should be finished before this line ----------------------------------------
 
-		SCOPED_DRAW_EVENT(RHICmdList, CompositionLighting);
+		SCOPED_DRAW_EVENT(RHICmdList, CompositionAfterLighting);
 
 		// we don't replace the final element with the scenecolor because this is what those passes should do by themself
 
diff --git a/Engine/Source/Runtime/Renderer/Private/CompositionLighting/CompositionLighting.h b/Engine/Source/Runtime/Renderer/Private/CompositionLighting/CompositionLighting.h
--- a/Engine/Source/Runtime/Renderer/Private/CompositionLighting/CompositionLighting.h
+++ b/Engine/Source/Runtime/Renderer/Private/CompositionLighting/CompositionLighting.h
@@ -18,7 +18,9 @@ public:
 
 	void ProcessAfterBasePass(FRHICommandListImmediate& RHICmdList,  FViewInfo& View);
 
-	void ProcessLighting(FRHICommandListImmediate& RHICmdList, FViewInfo& View);
+	void ProcessLpvIndirect(FRHICommandListImmediate& RHICmdList, FViewInfo& View);
+
+	void ProcessAfterLighting(FRHICommandListImmediate& RHICmdList, FViewInfo& View);
 };
 
 /** The global used for deferred lighting. */
diff --git a/Engine/Source/Runtime/Renderer/Private/CompositionLighting/PostProcessLpvIndirect.cpp b/Engine/Source/Runtime/Renderer/Private/CompositionLighting/PostProcessLpvIndirect.cpp
--- a/Engine/Source/Runtime/Renderer/Private/CompositionLighting/PostProcessLpvIndirect.cpp
+++ b/Engine/Source/Runtime/Renderer/Private/CompositionLighting/PostProcessLpvIndirect.cpp
@@ -22,6 +22,15 @@
 IMPLEMENT_UNIFORM_BUFFER_STRUCT(FLpvReadUniformBufferParameters,TEXT("LpvRead"));
 typedef TUniformBufferRef<FLpvReadUniformBufferParameters> FLpvReadUniformBufferRef;
 
+
+TAutoConsoleVariable<int32> CVarLPVMixing(
+	TEXT("r.LPV.Mixing"),
+	1,
+	TEXT("Reflection environment mixes with indirect shading (Ambient + LPV).\n")
+	TEXT(" 0 is off, 1 is on (default)"),
+	ECVF_RenderThreadSafe | ECVF_Cheat);
+
+
 /** Encapsulates the post processing ambient pixel shader. */
 class FPostProcessLpvIndirectPS : public FGlobalShader
 {
@@ -35,14 +44,9 @@ class FPostProcessLpvIndirectPS : public FGlobalShader
 
 public:
 	FPostProcessPassParameters PostprocessParameter;
-
-#if LPV_VOLUME_TEXTURE
 	FShaderResourceParameter LpvBufferSRVParameters[7];
 	FShaderResourceParameter LpvVolumeTextureSampler;
-#else
-	FShaderResourceParameter LpvBufferSRV;
-#endif 
-
+	FShaderResourceParameter AOVolumeTextureSRVParameter;
 	FDeferredPixelShaderParameters DeferredParameters;
 	FShaderResourceParameter PreIntegratedGF;
 	FShaderResourceParameter PreIntegratedGFSampler;
@@ -53,27 +57,21 @@ public:
 	{
 		PostprocessParameter.Bind(Initializer.ParameterMap);
 		DeferredParameters.Bind(Initializer.ParameterMap);
-#if LPV_VOLUME_TEXTURE
 		for ( int i=0; i<7; i++ )
 		{
 			LpvBufferSRVParameters[i].Bind( Initializer.ParameterMap, LpvVolumeTextureSRVNames[i] );
 		}
 		LpvVolumeTextureSampler.Bind(Initializer.ParameterMap, TEXT("gLpv3DTextureSampler"));
-#else
-		LpvBufferSRV.Bind( Initializer.ParameterMap, TEXT("gLpvBuffer") );
-#endif
+
+		AOVolumeTextureSRVParameter.Bind( Initializer.ParameterMap, TEXT("gAOVolumeTexture") );
 
 		PreIntegratedGF.Bind(Initializer.ParameterMap, TEXT("PreIntegratedGF"));
 		PreIntegratedGFSampler.Bind(Initializer.ParameterMap, TEXT("PreIntegratedGFSampler"));
 	}
 
 	void SetParameters(	
-#if LPV_VOLUME_TEXTURE
 		FTextureRHIParamRef* LpvBufferSRVsIn, 
-#else
-		FShaderResourceViewRHIParamRef LpvBufferSRVIn, 
-#endif 
-
+		FTextureRHIParamRef AOVolumeTextureSRVIn, 
 		FLpvReadUniformBufferRef LpvUniformBuffer, 
 		const FRenderingCompositePassContext& Context )
 	{
@@ -81,21 +79,19 @@ public:
 
 		SetUniformBufferParameter(Context.RHICmdList, ShaderRHI, GetUniformBufferParameter<FLpvReadUniformBufferParameters>(), LpvUniformBuffer);
 
-#if LPV_VOLUME_TEXTURE
 		for ( int i=0; i<7; i++ )
 		{
 			if ( LpvBufferSRVParameters[i].IsBound() )
 			{
 				Context.RHICmdList.SetShaderTexture(ShaderRHI, LpvBufferSRVParameters[i].GetBaseIndex(), LpvBufferSRVsIn[i]);
-				SetTextureParameter(Context.RHICmdList, ShaderRHI, LpvBufferSRVParameters[i], LpvVolumeTextureSampler, TStaticSamplerState<SF_Bilinear, AM_Clamp, AM_Clamp, AM_Clamp>::GetRHI(), LpvBufferSRVsIn[i]);
+				SetTextureParameter(Context.RHICmdList, ShaderRHI, LpvBufferSRVParameters[i], LpvVolumeTextureSampler, TStaticSamplerState<SF_Bilinear, AM_Border, AM_Border, AM_Border>::GetRHI(), LpvBufferSRVsIn[i]);
 			}
 		}
-#else
-		if ( LpvBufferSRV.IsBound() )
+
+		if ( AOVolumeTextureSRVParameter.IsBound() )
 		{
-			Context.RHICmdList.SetShaderResourceViewParameter( ShaderRHI, LpvBufferSRV.GetBaseIndex(), LpvBufferSRVIn );
+			Context.RHICmdList.SetShaderTexture(ShaderRHI, AOVolumeTextureSRVParameter.GetBaseIndex(), AOVolumeTextureSRVIn );
 		}
-#endif
 		FGlobalShader::SetParameters(Context.RHICmdList, ShaderRHI, Context.View);
 		PostprocessParameter.SetPS(ShaderRHI, Context, TStaticSamplerState<SF_Bilinear, AM_Clamp, AM_Clamp, AM_Clamp>::GetRHI());
 		DeferredParameters.Set(Context.RHICmdList, ShaderRHI, Context.View);
@@ -108,30 +104,106 @@ public:
 		bool bShaderHasOutdatedParameters = FGlobalShader::Serialize(Ar);
 		Ar << PostprocessParameter;
 
-#if LPV_VOLUME_TEXTURE
 		for ( int i=0; i<7; i++ )
 		{
 			Ar << LpvBufferSRVParameters[i];
 		}
 		Ar << LpvVolumeTextureSampler;
-#else
-		Ar << LpvBufferSRV;
-#endif
 		Ar << DeferredParameters;
 		Ar << PreIntegratedGF;
 		Ar << PreIntegratedGFSampler;
-
+		Ar << AOVolumeTextureSRVParameter;
 		return bShaderHasOutdatedParameters;
 	}
 };
 
-IMPLEMENT_SHADER_TYPE(,FPostProcessLpvIndirectPS,TEXT("PostProcessLpvIndirect"),TEXT("MainPS"),SF_Pixel);
+template<bool bApplySeparateSpecularRT>
+class TPostProcessLpvIndirectPS : public FPostProcessLpvIndirectPS
+{
+	DECLARE_SHADER_TYPE(TPostProcessLpvIndirectPS, Global);
 
+	/** Default constructor. */
+	TPostProcessLpvIndirectPS() {}
 
-void FRCPassPostProcessLpvIndirect::Process(FRenderingCompositePassContext& Context)
+public:
+
+	TPostProcessLpvIndirectPS(const ShaderMetaType::CompiledShaderInitializerType& Initializer)
+		: FPostProcessLpvIndirectPS(Initializer)
+	{}
+
+	static void ModifyCompilationEnvironment(EShaderPlatform Platform, FShaderCompilerEnvironment& OutEnvironment)
+	{
+		FPostProcessLpvIndirectPS::ModifyCompilationEnvironment(Platform,OutEnvironment);
+		OutEnvironment.SetDefine(TEXT("APPLY_SEPARATE_SPECULAR_RT"), (uint32)(bApplySeparateSpecularRT ? 1 : 0));
+	}
+};
+
+IMPLEMENT_SHADER_TYPE(template<>,TPostProcessLpvIndirectPS<false>,TEXT("PostProcessLpvIndirect"),TEXT("MainPS"),SF_Pixel);
+IMPLEMENT_SHADER_TYPE(template<>,TPostProcessLpvIndirectPS<true>,TEXT("PostProcessLpvIndirect"),TEXT("MainPS"),SF_Pixel);
+
+class FPostProcessLpvDirectionalOcclusionPS : public FGlobalShader
 {
-	SCOPED_DRAW_EVENT(Context.RHICmdList, PostProcessLpvIndirect);
+	DECLARE_SHADER_TYPE(FPostProcessLpvDirectionalOcclusionPS, Global);
+
+	//@todo-rco: Remove this when reenabling for OpenGL
+	static bool ShouldCache( EShaderPlatform Platform )		{ return IsFeatureLevelSupported(Platform, ERHIFeatureLevel::SM5) && !IsOpenGLPlatform(Platform); }
+
+	/** Default constructor. */
+	FPostProcessLpvDirectionalOcclusionPS() {}
+
+public:
+	FPostProcessPassParameters PostprocessParameter;
+	FShaderResourceParameter LpvVolumeTextureSampler;
+	FShaderResourceParameter AOVolumeTextureSRVParameter;
+
+	FDeferredPixelShaderParameters DeferredParameters;
+
+	/** Initialization constructor. */
+	FPostProcessLpvDirectionalOcclusionPS(const ShaderMetaType::CompiledShaderInitializerType& Initializer)
+		: FGlobalShader(Initializer)
+	{
+		PostprocessParameter.Bind(Initializer.ParameterMap);
+		DeferredParameters.Bind(Initializer.ParameterMap);
+		LpvVolumeTextureSampler.Bind(Initializer.ParameterMap, TEXT("gLpv3DTextureSampler"));
+		AOVolumeTextureSRVParameter.Bind( Initializer.ParameterMap, TEXT("gAOVolumeTexture") );
+	}
+
+	void SetParameters(	
+		FTextureRHIParamRef AOVolumeTextureSRVIn, 
+		FLpvReadUniformBufferRef LpvUniformBuffer, 
+		const FRenderingCompositePassContext& Context )
+	{
+		const FPixelShaderRHIParamRef ShaderRHI = GetPixelShader();
+
+		SetUniformBufferParameter(Context.RHICmdList, ShaderRHI, GetUniformBufferParameter<FLpvReadUniformBufferParameters>(), LpvUniformBuffer);
+
+		if ( AOVolumeTextureSRVParameter.IsBound() )
+		{
+			Context.RHICmdList.SetShaderTexture(ShaderRHI, AOVolumeTextureSRVParameter.GetBaseIndex(), AOVolumeTextureSRVIn );
+		}
+		Context.RHICmdList.SetShaderSampler(ShaderRHI, LpvVolumeTextureSampler.GetBaseIndex(), TStaticSamplerState<SF_Bilinear, AM_Clamp, AM_Clamp, AM_Clamp>::GetRHI() );
+
+		FGlobalShader::SetParameters(Context.RHICmdList, ShaderRHI, Context.View);
+		PostprocessParameter.SetPS(ShaderRHI, Context, TStaticSamplerState<SF_Bilinear,AM_Clamp,AM_Clamp,AM_Clamp>::GetRHI());
+		DeferredParameters.Set(Context.RHICmdList, ShaderRHI, Context.View);
+	}
+
+	// FShader interface.
+	virtual bool Serialize(FArchive& Ar)
+	{
+		bool bShaderHasOutdatedParameters = FGlobalShader::Serialize(Ar);
+		Ar << LpvVolumeTextureSampler;
+		Ar << DeferredParameters;
+		Ar << AOVolumeTextureSRVParameter;
+		Ar << PostprocessParameter;
+		return bShaderHasOutdatedParameters;
+	}
+};
+
+IMPLEMENT_SHADER_TYPE(,FPostProcessLpvDirectionalOcclusionPS,TEXT("PostProcessLpvIndirect"),TEXT("DirectionalOcclusionPS"),SF_Pixel);
 
+void FRCPassPostProcessLpvIndirect::Process(FRenderingCompositePassContext& Context)
+{
 	{
 		FRenderingCompositeOutput* OutputOfMyInput = GetInput(ePId_Input0)->GetOutput();
 		PassOutputs[0].PooledRenderTarget = OutputOfMyInput->PooledRenderTarget;
@@ -144,6 +216,14 @@ void FRCPassPostProcessLpvIndirect::Process(FRenderingCompositePassContext& Cont
 
 	const FPostProcessSettings& PostprocessSettings = Context.View.FinalPostProcessSettings;
 	const FSceneView& View = Context.View;
+
+	FSceneViewState* ViewState = (FSceneViewState*)View.State;
+	if ( ViewState == nullptr || ViewState->GetLightPropagationVolume() == nullptr || PostprocessSettings.LPVIntensity == 0.0f )
+	{
+		return;
+	}
+	FLightPropagationVolume* Lpv = ViewState->GetLightPropagationVolume();
+
 	const FSceneViewFamily& ViewFamily = *(View.Family);
 
 	FIntRect SrcRect = View.ViewRect;
@@ -151,76 +231,144 @@ void FRCPassPostProcessLpvIndirect::Process(FRenderingCompositePassContext& Cont
 	FIntRect DestRect = View.ViewRect;
 	FIntPoint DestSize = DestRect.Size();
 
-	uint32 NumReflectionCaptures = ViewFamily.Scene->GetRenderScene()->ReflectionSceneData.RegisteredReflectionCaptures.Num();
+	const bool bMixing = CVarLPVMixing.GetValueOnRenderThread() != 0;
+	// Apply specular separately if we're mixing reflection environment with indirect lighting
+	const bool bApplySeparateSpecularRT = View.Family->EngineShowFlags.ReflectionEnvironment && bMixing;
 
 	const FSceneRenderTargetItem& DestColorRenderTarget = GSceneRenderTargets.GetSceneColor()->GetRenderTargetItem();
+	const FSceneRenderTargetItem& DestSpecularRenderTarget = GSceneRenderTargets.LightAccumulation->GetRenderTargetItem();
 
-	// Set the view family's render target/viewport.
-	FTextureRHIParamRef RenderTargets[] =
+	const FSceneRenderTargetItem& DestDirectionalOcclusionRenderTarget = GSceneRenderTargets.DirectionalOcclusion->GetRenderTargetItem();
+
+	// Make sure the LPV Update has completed
+	Lpv->InsertGPUWaitForAsyncUpdate(Context.RHICmdList);
+
+	if ( Context.View.FinalPostProcessSettings.LPVDirectionalOcclusionIntensity > 0.0001f )
 	{
-		DestColorRenderTarget.TargetableTexture,
-	};
+		DoDirectionalOcclusionPass(Context);
+	}
+
+	FTextureRHIParamRef RenderTargets[2];
+	RenderTargets[0] = DestColorRenderTarget.TargetableTexture;
+	RenderTargets[1] = DestSpecularRenderTarget.TargetableTexture;
 
 	// Set the view family's render target/viewport.
-	SetRenderTargets(Context.RHICmdList, 1, RenderTargets, GSceneRenderTargets.GetSceneDepthSurface(), ESimpleRenderTargetMode::EExistingColorAndDepth, FExclusiveDepthStencil::DepthRead_StencilWrite);
+	// If specular not applied: set only color target
+	uint32 NumRenderTargets = 1; 
+	if ( bApplySeparateSpecularRT ) 
+	{
+		NumRenderTargets = 2;
+	}
+
+	SetRenderTargets(Context.RHICmdList, NumRenderTargets, RenderTargets, 0, ESimpleRenderTargetMode::EExistingColorAndDepth, FExclusiveDepthStencil::DepthRead_StencilNop);
 
 	Context.SetViewportAndCallRHI(View.ViewRect);
 
 	// set the state
-	if ( ViewFamily.EngineShowFlags.LpvLightingOnly )
+	Context.RHICmdList.SetBlendState(TStaticBlendState<CW_RGB,BO_Add,BF_One,BF_One,BO_Add,BF_One,BF_One>::GetRHI());
+	Context.RHICmdList.SetRasterizerState(TStaticRasterizerState<>::GetRHI());
+	Context.RHICmdList.SetDepthStencilState(TStaticDepthStencilState<false,CF_Always>::GetRHI());
+
+	TShaderMapRef<FPostProcessVS> VertexShader(Context.GetShaderMap());
+
+	TShaderMapRef< TPostProcessLpvIndirectPS<false> >	  PixelShaderNoSpecular(Context.GetShaderMap());
+	TShaderMapRef< TPostProcessLpvIndirectPS<true> >	  PixelShaderWithSpecular(Context.GetShaderMap());
+
+	FPostProcessLpvIndirectPS* PixelShader = NULL;
+	int BoundShaderIndex = -1;
+	if ( bApplySeparateSpecularRT )
 	{
-		Context.RHICmdList.SetBlendState(TStaticBlendState<>::GetRHI());
+		PixelShader = (FPostProcessLpvIndirectPS*)*PixelShaderWithSpecular;
+		BoundShaderIndex = 0;
 	}
 	else
 	{
-		// additive blending
-		Context.RHICmdList.SetBlendState(TStaticBlendState<CW_RGB, BO_Add, BF_One, BF_One>::GetRHI());
+		PixelShader = (FPostProcessLpvIndirectPS*)*PixelShaderNoSpecular;
+		BoundShaderIndex = 1;
 	}
-	Context.RHICmdList.SetRasterizerState(TStaticRasterizerState<>::GetRHI());
-	Context.RHICmdList.SetDepthStencilState(TStaticDepthStencilState<false, CF_Always>::GetRHI());
+	static FGlobalBoundShaderState BoundShaderState[2];
 
-	TShaderMapRef<FPostProcessVS> VertexShader(Context.GetShaderMap());
+	// call it once after setting up the shader data to avoid the warnings in the function
+	SetGlobalBoundShaderState(Context.RHICmdList, Context.GetFeatureLevel(), BoundShaderState[BoundShaderIndex], GFilterVertexDeclaration.VertexDeclarationRHI, *VertexShader, PixelShader);
 
-	FSceneViewState* ViewState = (FSceneViewState*)View.State;
-	FLightPropagationVolume* Lpv = NULL;
-	bool bUseLpv = false;
-	if ( ViewState )
+	FLpvReadUniformBufferParameters	LpvReadUniformBufferParams;
+	FLpvReadUniformBufferRef LpvReadUniformBuffer;
+
+	LpvReadUniformBufferParams = Lpv->GetReadUniformBufferParams();
+	LpvReadUniformBuffer = FLpvReadUniformBufferRef::CreateUniformBufferImmediate( LpvReadUniformBufferParams, UniformBuffer_SingleDraw ); 
+
+	FTextureRHIParamRef LpvBufferSrvs[7];
+	for ( int i = 0; i < 7; i++ ) 
 	{
-		Lpv = ViewState->GetLightPropagationVolume();
+		LpvBufferSrvs[i] = Lpv->GetLpvBufferSrv(i);
+	}
+
+	PixelShader->SetParameters( LpvBufferSrvs, Lpv->GetAOVolumeTextureSRV(), LpvReadUniformBuffer, Context );
+
+	{
+		SCOPED_DRAW_EVENT(Context.RHICmdList, PostProcessLpvIndirect );
+
+	    // Draw a quad mapping scene color to the view's render target
+	    DrawRectangle( 
+		    Context.RHICmdList,
+		    0, 0,
+		    View.ViewRect.Width(), View.ViewRect.Height(),
+		    View.ViewRect.Min.X, View.ViewRect.Min.Y, 
+		    View.ViewRect.Width(), View.ViewRect.Height(),
+		    View.ViewRect.Size(),
+		    GSceneRenderTargets.GetBufferSizeXY(),
+		    *VertexShader);
+
+		Context.RHICmdList.CopyToResolveTarget(DestColorRenderTarget.TargetableTexture, DestColorRenderTarget.ShaderResourceTexture, false, FResolveParams());
+		if(bApplySeparateSpecularRT)
+		{
+			Context.RHICmdList.CopyToResolveTarget(DestSpecularRenderTarget.TargetableTexture, DestSpecularRenderTarget.ShaderResourceTexture, false, FResolveParams());
+		}
+	}
 
-		bUseLpv = Lpv && PostprocessSettings.LPVIntensity > 0.0f;
+	if ( Context.View.FinalPostProcessSettings.LPVDirectionalOcclusionIntensity > 0.0001f )
+	{
+		GRenderTargetPool.VisualizeTexture.SetCheckPoint(Context.RHICmdList, GSceneRenderTargets.DirectionalOcclusion);
 	}
+}
 
-	if ( !bUseLpv )
+void FRCPassPostProcessLpvIndirect::DoDirectionalOcclusionPass(FRenderingCompositePassContext& Context) const
+{
+	SCOPED_DRAW_EVENT(Context.RHICmdList, PostProcessLpvDirectionalOcclusion );
+	const FSceneRenderTargetItem& DestDirectionalOcclusionRenderTarget = GSceneRenderTargets.DirectionalOcclusion->GetRenderTargetItem();
+	FViewInfo& View = Context.View;
+	FSceneViewState* ViewState = (FSceneViewState*)View.State;
+	const FPostProcessSettings& PostprocessSettings = Context.View.FinalPostProcessSettings;
+	if ( ViewState == nullptr || ViewState->GetLightPropagationVolume() == nullptr || PostprocessSettings.LPVIntensity == 0.0f )
 	{
 		return;
 	}
+	FLightPropagationVolume* Lpv = ViewState->GetLightPropagationVolume();
 
-	TShaderMapRef<FPostProcessLpvIndirectPS> PixelShader(Context.GetShaderMap());
+	FTextureRHIParamRef RenderTarget = DestDirectionalOcclusionRenderTarget.TargetableTexture;
 
-	static FGlobalBoundShaderState BoundShaderState;
+	SetRenderTargets(Context.RHICmdList, 1, &RenderTarget, NULL, ESimpleRenderTargetMode::EExistingColorAndDepth, FExclusiveDepthStencil::DepthRead_StencilNop);
 	
-
-	// call it once after setting up the shader data to avoid the warnings in the function
-	SetGlobalBoundShaderState(Context.RHICmdList, Context.GetFeatureLevel(), BoundShaderState, GFilterVertexDeclaration.VertexDeclarationRHI, *VertexShader, *PixelShader);
+	Context.SetViewportAndCallRHI(View.ViewRect);
+	Context.RHICmdList.SetBlendState( TStaticBlendState<>::GetRHI() );
+	Context.RHICmdList.SetRasterizerState(TStaticRasterizerState<>::GetRHI());
+	Context.RHICmdList.SetDepthStencilState(TStaticDepthStencilState<false,CF_Always>::GetRHI());
+	TShaderMapRef<FPostProcessVS> VertexShader(View.ShaderMap);
 
 	FLpvReadUniformBufferParameters	LpvReadUniformBufferParams;
 	FLpvReadUniformBufferRef LpvReadUniformBuffer;
 
+	TShaderMapRef< FPostProcessLpvDirectionalOcclusionPS > PixelShader(View.ShaderMap);
+
+	static FGlobalBoundShaderState BoundShaderState;
+
+	// call it once after setting up the shader data to avoid the warnings in the function
+	SetGlobalBoundShaderState(Context.RHICmdList, Context.GetFeatureLevel(), BoundShaderState, GFilterVertexDeclaration.VertexDeclarationRHI, *VertexShader, *PixelShader);
+
 	LpvReadUniformBufferParams = Lpv->GetReadUniformBufferParams();
 	LpvReadUniformBuffer = FLpvReadUniformBufferRef::CreateUniformBufferImmediate( LpvReadUniformBufferParams, UniformBuffer_SingleDraw ); 
 
-#if LPV_VOLUME_TEXTURE
-	FTextureRHIParamRef LpvBufferSrvs[7];
-	for ( int i = 0; i < 7; i++ ) 
-	{
-		LpvBufferSrvs[i] = Lpv->GetLpvBufferSrv(i);
-	}
-	PixelShader->SetParameters(LpvBufferSrvs, LpvReadUniformBuffer, Context);
-#else
-	FShaderResourceViewRHIParamRef LpvBufferSrv = Lpv->GetLpvBufferSrv();
-	PixelShader->SetParameters(LpvBufferSrv, LpvReadUniformBuffer, Context);
-#endif
+	PixelShader->SetParameters( Lpv->GetAOVolumeTextureSRV(), LpvReadUniformBuffer, Context );
 
 	// Draw a quad mapping scene color to the view's render target
 	DrawRectangle( 
@@ -232,8 +380,6 @@ void FRCPassPostProcessLpvIndirect::Process(FRenderingCompositePassContext& Cont
 		View.ViewRect.Size(),
 		GSceneRenderTargets.GetBufferSizeXY(),
 		*VertexShader);
-
-	Context.RHICmdList.CopyToResolveTarget(DestColorRenderTarget.TargetableTexture, DestColorRenderTarget.ShaderResourceTexture, false, FResolveParams());
 }
 
 FPooledRenderTargetDesc FRCPassPostProcessLpvIndirect::ComputeOutputDesc(EPassOutputId InPassOutputId) const
@@ -244,4 +390,4 @@ FPooledRenderTargetDesc FRCPassPostProcessLpvIndirect::ComputeOutputDesc(EPassOu
 	Ret.DebugName = TEXT("LpvIndirect");
 
 	return Ret;
-}
+}
\ No newline at end of file
diff --git a/Engine/Source/Runtime/Renderer/Private/CompositionLighting/PostProcessLpvIndirect.h b/Engine/Source/Runtime/Renderer/Private/CompositionLighting/PostProcessLpvIndirect.h
--- a/Engine/Source/Runtime/Renderer/Private/CompositionLighting/PostProcessLpvIndirect.h
+++ b/Engine/Source/Runtime/Renderer/Private/CompositionLighting/PostProcessLpvIndirect.h
@@ -25,6 +25,9 @@ public:
 	virtual void Process(FRenderingCompositePassContext& Context) override;
 	virtual void Release() override { delete this; }
 	virtual FPooledRenderTargetDesc ComputeOutputDesc(EPassOutputId InPassOutputId) const override;
+
+private:
+	void DoDirectionalOcclusionPass(FRenderingCompositePassContext& Context) const;
 };
 
 
diff --git a/Engine/Source/Runtime/Renderer/Private/DeferredShadingRenderer.cpp b/Engine/Source/Runtime/Renderer/Private/DeferredShadingRenderer.cpp
--- a/Engine/Source/Runtime/Renderer/Private/DeferredShadingRenderer.cpp
+++ b/Engine/Source/Runtime/Renderer/Private/DeferredShadingRenderer.cpp
@@ -1012,27 +1012,29 @@ void FDeferredShadingSceneRenderer::Render(FRHICommandListImmediate& RHICmdList)
 		// Filter the translucency lighting volume now that it is complete
 		FilterTranslucentVolumeLighting(RHICmdList);
 
-		// Clear LPVs for all views
-		if ( FeatureLevel >= ERHIFeatureLevel::SM5 )
-		{
-			PropagateLPVs(RHICmdList);
+		// Pre-lighting composition lighting stage
+		// e.g. LPV indirect
+		for(int32 ViewIndex = 0; ViewIndex < Views.Num(); ++ViewIndex)
+		{	
+			SCOPED_CONDITIONAL_DRAW_EVENTF(RHICmdList, EventView,Views.Num() > 1, TEXT("View%d"), ViewIndex);
+			GCompositionLighting.ProcessLpvIndirect(RHICmdList, Views[ViewIndex]);
 		}
 
 		TRefCountPtr<IPooledRenderTarget> DynamicBentNormalAO;
 		RenderDynamicSkyLighting(RHICmdList, VelocityRT, DynamicBentNormalAO);
 
-		//SSR and SSS need the SceneColor finalized as an SRV.
+		// SSS need the SceneColor finalized as an SRV.
 		GSceneRenderTargets.FinishRenderingSceneColor(RHICmdList, true);
 
 		// Render reflections that only operate on opaque pixels
 		RenderDeferredReflections(RHICmdList, DynamicBentNormalAO);
 
 		// Post-lighting composition lighting stage
-		// e.g. ambient cubemaps, ambient occlusion, LPV indirect
+		// e.g. ScreenSpaceSubsurfaceScattering
 		for(int32 ViewIndex = 0; ViewIndex < Views.Num(); ++ViewIndex)
 		{	
 			SCOPED_CONDITIONAL_DRAW_EVENTF(RHICmdList, EventView,Views.Num() > 1, TEXT("View%d"), ViewIndex);
-			GCompositionLighting.ProcessLighting(RHICmdList, Views[ViewIndex]);
+			GCompositionLighting.ProcessAfterLighting(RHICmdList, Views[ViewIndex]);
 		}
 	}
 
@@ -1467,10 +1469,15 @@ bool FDeferredShadingSceneRenderer::RenderBasePass(FRHICommandListImmediate& RHI
 
 void FDeferredShadingSceneRenderer::ClearLPVs(FRHICommandListImmediate& RHICmdList)
 {
+	SCOPED_DRAW_EVENT(RHICmdList, ClearLPVs);
+	SCOPE_CYCLE_COUNTER(STAT_UpdateLPVs);
+
 	// clear light propagation volumes
 
 	for(int32 ViewIndex = 0; ViewIndex < Views.Num(); ViewIndex++)
 	{
+		SCOPED_CONDITIONAL_DRAW_EVENTF(RHICmdList, EventView, Views.Num() > 1, TEXT("View%d"), ViewIndex);
+
 		FViewInfo& View = Views[ViewIndex];
 
 		FSceneViewState* ViewState = (FSceneViewState*)Views[ViewIndex].State;
@@ -1489,13 +1496,18 @@ void FDeferredShadingSceneRenderer::ClearLPVs(FRHICommandListImmediate& RHICmdLi
 	}
 }
 
-void FDeferredShadingSceneRenderer::PropagateLPVs(FRHICommandListImmediate& RHICmdList)
+void FDeferredShadingSceneRenderer::UpdateLPVs(FRHICommandListImmediate& RHICmdList)
 {
+	SCOPED_DRAW_EVENT(RHICmdList, UpdateLPVs);
+	SCOPE_CYCLE_COUNTER(STAT_UpdateLPVs);
+
 	for(int32 ViewIndex = 0; ViewIndex < Views.Num(); ViewIndex++)
 	{
-		FViewInfo& View = Views[ViewIndex];
+		SCOPED_CONDITIONAL_DRAW_EVENTF(RHICmdList, EventView, Views.Num() > 1, TEXT("View%d"), ViewIndex);
 
+		FViewInfo& View = Views[ViewIndex];
 		FSceneViewState* ViewState = (FSceneViewState*)Views[ViewIndex].State;
+
 		if(ViewState)
 		{
 			FLightPropagationVolume* LightPropagationVolume = ViewState->GetLightPropagationVolume();
@@ -1504,8 +1516,8 @@ void FDeferredShadingSceneRenderer::PropagateLPVs(FRHICommandListImmediate& RHIC
 			{
 				SCOPED_DRAW_EVENT(RHICmdList, UpdateLPVs);
 				SCOPE_CYCLE_COUNTER(STAT_UpdateLPVs);
-				
-				LightPropagationVolume->Propagate(RHICmdList, View);
+
+				LightPropagationVolume->Update(RHICmdList, View);
 			}
 		}
 	}
diff --git a/Engine/Source/Runtime/Renderer/Private/DeferredShadingRenderer.h b/Engine/Source/Runtime/Renderer/Private/DeferredShadingRenderer.h
--- a/Engine/Source/Runtime/Renderer/Private/DeferredShadingRenderer.h
+++ b/Engine/Source/Runtime/Renderer/Private/DeferredShadingRenderer.h
@@ -48,7 +48,7 @@ public:
 	void ClearLPVs(FRHICommandListImmediate& RHICmdList);
 
 	/** Propagates LPVs for all views */
-	void PropagateLPVs(FRHICommandListImmediate& RHICmdList);
+	void UpdateLPVs(FRHICommandListImmediate& RHICmdList);
 
 	/**
 	 * Renders the dynamic scene's prepass for a particular view
diff --git a/Engine/Source/Runtime/Renderer/Private/LightPropagationVolume.cpp b/Engine/Source/Runtime/Renderer/Private/LightPropagationVolume.cpp
--- a/Engine/Source/Runtime/Renderer/Private/LightPropagationVolume.cpp
+++ b/Engine/Source/Runtime/Renderer/Private/LightPropagationVolume.cpp
@@ -5,7 +5,7 @@
 //
 // Created:		2013-03-01
 //
-// Author:		mailto:[REDACTED_EMAIL]
+// Author:		Ben Woodhouse - mailto:[REDACTED_EMAIL]
 //
 //				Copyright (C) Microsoft. All rights reserved.
 //-----------------------------------------------------------------------------
@@ -24,26 +24,18 @@ static TAutoConsoleVariable<int32> CVarLightPropagationVolume(
 	TEXT(" 1: on"),
 	ECVF_RenderThreadSafe | ECVF_ReadOnly);
 
+struct LPVBufferElementUncompressed
+{
+	FVector elements[9];
+	float AO;
+}; 
 // ----------------------------------------------------------------------------
 
 static const uint32 LPV_GRIDRES = 32;
 static float LPV_CENTRE_OFFSET = 10.0f;
 
-#define LPV_AMBIENT_CUBE 0
-
-static int GLpvEnabled = 0;
-
 // ----------------------------------------------------------------------------
 
-struct LPVBufferElement
-{
-#if LPV_AMBIENT_CUBE
-	uint32 values[6];   // 12 halfs (24 bytes)
-#else
-	uint32 values[14];   // 28 halfs (56 bytes)
-#endif
-};
-
 struct VplListEntry
 {
 	uint32	normalPacked;
@@ -59,6 +51,40 @@ struct FRsmInfo
 	float AreaBrightnessMultiplier; 
 };
 
+
+static TAutoConsoleVariable<float> CVarLPVSpecularIntensity(
+	TEXT("r.LPV.SpecularIntensity"),
+	0.333f,
+	TEXT("Multiplier for LPV Specular.") );
+
+static TAutoConsoleVariable<float> CVarLPVDiffuseIntensity(
+	TEXT("r.LPV.DiffuseIntensity"),
+	0.333f,
+	TEXT("Multiplier for LPV Diffuse.") );
+
+static TAutoConsoleVariable<float> CVarLPVIntensity(
+	TEXT("r.LPV.Intensity"),
+	1.0f,
+	TEXT("Multiplier for LPV intensity. 1.0 is the default.") );
+
+static TAutoConsoleVariable<int32> CVarNumPropagationSteps(
+	TEXT("r.LPV.NumPropagationSteps"),
+	3,
+	TEXT("Number of LPV propagation steps") );
+
+static TAutoConsoleVariable<int32> CVarLPVNumAOPropagationSteps(
+	TEXT("r.LPV.NumAOPropagationSteps"),
+	1,
+	TEXT("Number of LPV AO propagation steps\n")
+	TEXT("0: noisy (good for debugging)\n")
+	TEXT("1: normal (default)\n")
+	TEXT("2: blurry") );
+
+static TAutoConsoleVariable<float> CVarLPVEmissiveIntensityMultiplier(
+	TEXT("r.LPV.EmissiveMultiplier"),
+	1.0f,
+	TEXT("Emissive intensity multiplier") );
+
 // ----------------------------------------------------------------------------
 
 /**
@@ -69,6 +95,10 @@ BEGIN_UNIFORM_BUFFER_STRUCT( FLpvDirectLightInjectParameters, )
 	DECLARE_UNIFORM_BUFFER_STRUCT_MEMBER( FVector4, LightPosition )
 	DECLARE_UNIFORM_BUFFER_STRUCT_MEMBER( FVector4, LightColor )
 	DECLARE_UNIFORM_BUFFER_STRUCT_MEMBER( float, LightFalloffExponent )
+	DECLARE_UNIFORM_BUFFER_STRUCT_MEMBER( float, LightSourceLength )
+	DECLARE_UNIFORM_BUFFER_STRUCT_MEMBER( FVector4, LightDirection )
+	DECLARE_UNIFORM_BUFFER_STRUCT_MEMBER( FVector2D, LightSpotAngles )
+	DECLARE_UNIFORM_BUFFER_STRUCT_MEMBER( float, bLightInverseSquaredAttenuation )
 END_UNIFORM_BUFFER_STRUCT( FLpvDirectLightInjectParameters )
 
 IMPLEMENT_UNIFORM_BUFFER_STRUCT(FLpvDirectLightInjectParameters,TEXT("LpvInject"));
@@ -77,7 +107,6 @@ typedef TUniformBufferRef<FLpvDirectLightInjectParameters> FDirectLightInjectBuf
 
 IMPLEMENT_UNIFORM_BUFFER_STRUCT(FLpvWriteUniformBufferParameters,TEXT("LpvWrite"));
 
-
 // ----------------------------------------------------------------------------
 // Base LPV Write Compute shader
 // ----------------------------------------------------------------------------
@@ -91,90 +120,66 @@ public:
 	explicit FLpvWriteShaderCSBase( const ShaderMetaType::CompiledShaderInitializerType& Initializer )
 		: FGlobalShader(Initializer)
 	{
-#if LPV_VOLUME_TEXTURE
 		for ( int i = 0; i < 7; i++ )
 		{
 			LpvBufferSRVParameters[i].Bind(Initializer.ParameterMap, LpvVolumeTextureSRVNames[i] );
 			LpvBufferUAVs[i].Bind(Initializer.ParameterMap, LpvVolumeTextureUAVNames[i] );
 		}
-#else
-		LpvBufferSRV.Bind(Initializer.ParameterMap, TEXT("gLpvBuffer") );
-		LpvBufferUAV.Bind(Initializer.ParameterMap, TEXT("gLpvBufferRW") );
-#endif
 
-#if LPV_VOLUME_TEXTURE || LPV_GV_VOLUME_TEXTURE
 		LpvVolumeTextureSampler.Bind( Initializer.ParameterMap, TEXT("gLpv3DTextureSampler") );
-#endif
-
 		VplListHeadBufferSRV.Bind(Initializer.ParameterMap, TEXT("gVplListHeadBuffer") );
 		VplListHeadBufferUAV.Bind(Initializer.ParameterMap, TEXT("RWVplListHeadBuffer") );
 		VplListBufferSRV.Bind(Initializer.ParameterMap, TEXT("gVplListBuffer") );
 		VplListBufferUAV.Bind(Initializer.ParameterMap, TEXT("RWVplListBuffer") );
 
-#if LPV_GV_VOLUME_TEXTURE
-		for ( int i = 0; i < 3; i++ )
+		for ( int i = 0; i < NUM_GV_TEXTURES; i++ ) 
 		{
 			GvBufferSRVParameters[i].Bind(Initializer.ParameterMap, LpvGvVolumeTextureSRVNames[i] );
 			GvBufferUAVs[i].Bind(Initializer.ParameterMap, LpvGvVolumeTextureUAVNames[i] );
 		}
-#else
-		GvBufferSRV.Bind(Initializer.ParameterMap, TEXT("gGvBuffer") );
-		GvBufferUAV.Bind(Initializer.ParameterMap, TEXT("gGvBufferRW") );
-#endif
 		GvListBufferUAV.Bind(Initializer.ParameterMap, TEXT("RWGvListBuffer") );
 		GvListHeadBufferUAV.Bind(Initializer.ParameterMap, TEXT("RWGvListHeadBuffer") );
 
 		GvListBufferSRV.Bind(Initializer.ParameterMap, TEXT("gGvListBuffer") );
 		GvListHeadBufferSRV.Bind(Initializer.ParameterMap, TEXT("gGvListHeadBuffer") );
+
+		AOVolumeTextureUAV.Bind(Initializer.ParameterMap, TEXT("gAOVolumeTextureRW") );
+		AOVolumeTextureSRV.Bind(Initializer.ParameterMap, TEXT("gAOVolumeTexture") );
 	}
 
 	static void ModifyCompilationEnvironment( EShaderPlatform Platform, FShaderCompilerEnvironment& OutEnvironment )
 	{
 		FGlobalShader::ModifyCompilationEnvironment( Platform, OutEnvironment );
-		OutEnvironment.SetDefine( TEXT("LPV_AMBIENT_CUBE"),			(uint32)LPV_AMBIENT_CUBE );
 		OutEnvironment.SetDefine( TEXT("LPV_MULTIPLE_BOUNCES"), (uint32)LPV_MULTIPLE_BOUNCES );
-		OutEnvironment.SetDefine( TEXT("LPV_GV_VOLUME_TEXTURE"),	(uint32)LPV_GV_VOLUME_TEXTURE );
+		OutEnvironment.SetDefine( TEXT("LPV_GV_SH_ORDER"),			(uint32)LPV_GV_SH_ORDER );
 	}
 	// Serialization
 	virtual bool Serialize( FArchive& Ar ) override
 	{
 		bool bShaderHasOutdatedParameters = FGlobalShader::Serialize( Ar );
 
-#if LPV_VOLUME_TEXTURE
 		for(int i = 0; i < 7; i++)
 		{
 			Ar << LpvBufferSRVParameters[i];
 			Ar << LpvBufferUAVs[i];
 		}
-#else
-		Ar << LpvBufferSRV;
-		Ar << LpvBufferUAV;
-#endif
 
-#if LPV_VOLUME_TEXTURE || LPV_GV_VOLUME_TEXTURE
 		Ar << LpvVolumeTextureSampler;
-#endif
-
 		Ar << VplListHeadBufferSRV;
 		Ar << VplListHeadBufferUAV;
 		Ar << VplListBufferSRV;
 		Ar << VplListBufferUAV;
-#if LPV_GV_VOLUME_TEXTURE
-		for(int i = 0; i < 3; i++)
+		for ( int i = 0; i < NUM_GV_TEXTURES; i++ )
 		{
 			Ar << GvBufferSRVParameters[i];
 			Ar << GvBufferUAVs[i];
 		}
-#else
-		Ar << GvBufferSRV;
-		Ar << GvBufferUAV;
-#endif
 		Ar << GvListBufferUAV;
 		Ar << GvListHeadBufferUAV;
-
 		Ar << GvListBufferSRV;
 		Ar << GvListHeadBufferSRV;
-
+		Ar << AOVolumeTextureUAV;
+		Ar << AOVolumeTextureSRV;
 		return bShaderHasOutdatedParameters;
 	}
 
@@ -182,7 +187,6 @@ public:
 	{
 		FComputeShaderRHIParamRef ShaderRHI = GetComputeShader();
 		SetUniformBufferParameter(RHICmdList, ShaderRHI, GetUniformBufferParameter<FLpvWriteUniformBufferParameters>(), Params.UniformBuffer );
-#if LPV_VOLUME_TEXTURE
 		for(int i  =0; i < 7; i++)
 		{
 			if ( LpvBufferSRVParameters[i].IsBound() )
@@ -196,16 +200,6 @@ public:
 
 			SetTextureParameter(RHICmdList, ShaderRHI, LpvBufferSRVParameters[i], LpvVolumeTextureSampler, TStaticSamplerState<SF_Bilinear,AM_Clamp,AM_Clamp,AM_Clamp>::GetRHI(), Params.LpvBufferSRVs[i] );
 		}
-#else
-		if ( LpvBufferSRV.IsBound() )
-		{
-			RHICmdList.SetShaderResourceViewParameter( ShaderRHI, LpvBufferSRV.GetBaseIndex(), Params.LpvBufferSRV );
-		}
-		if ( LpvBufferUAV.IsBound() )
-		{
-			RHICmdList.SetUAVParameter( ShaderRHI, LpvBufferUAV.GetBaseIndex(), Params.LpvBufferUAV );
-		}
-#endif
 		if ( VplListHeadBufferSRV.IsBound() )
 		{
 			RHICmdList.SetShaderResourceViewParameter( ShaderRHI, VplListHeadBufferSRV.GetBaseIndex(), Params.VplListHeadBufferSRV );
@@ -222,8 +216,9 @@ public:
 		{
 			RHICmdList.SetUAVParameter( ShaderRHI, VplListBufferUAV.GetBaseIndex(), Params.VplListBufferUAV ); 
 		}
-#if LPV_GV_VOLUME_TEXTURE
-		for ( int i=0;i<3; i++ )
+
+		// GV Volume texture
+		for ( int i=0;i<NUM_GV_TEXTURES; i++ )
 		{
 			if ( GvBufferSRVParameters[i].IsBound() )
 			{
@@ -236,16 +231,7 @@ public:
 
 			SetTextureParameter(RHICmdList, ShaderRHI, GvBufferSRVParameters[i], LpvVolumeTextureSampler, TStaticSamplerState<SF_Bilinear,AM_Clamp,AM_Clamp,AM_Clamp>::GetRHI(), Params.GvBufferSRVs[i] );
 		}
-#else
-		if ( GvBufferSRV.IsBound() )
-		{
-			RHICmdList.SetShaderResourceViewParameter( ShaderRHI, GvBufferSRV.GetBaseIndex(), Params.GvBufferSRV );
-		}
-		if ( GvBufferUAV.IsBound() )
-		{
-			RHICmdList.SetUAVParameter( ShaderRHI, GvBufferUAV.GetBaseIndex(), Params.GvBufferUAV );
-		}
-#endif
+
 		if(GvListBufferUAV.IsBound())
 		{
 			RHICmdList.SetUAVParameter(ShaderRHI, GvListBufferUAV.GetBaseIndex(), Params.GvListBufferUAV); 
@@ -262,34 +248,31 @@ public:
 		{
 			RHICmdList.SetShaderResourceViewParameter(ShaderRHI, GvListHeadBufferSRV.GetBaseIndex(), Params.GvListHeadBufferSRV);
 		}
+		if ( AOVolumeTextureUAV.IsBound() )
+		{
+			RHICmdList.SetUAVParameter( ShaderRHI, AOVolumeTextureUAV.GetBaseIndex(), Params.AOVolumeTextureUAV ); 
+		}
+		if ( AOVolumeTextureSRV.IsBound() )
+		{
+			RHICmdList.SetShaderTexture(ShaderRHI, AOVolumeTextureSRV.GetBaseIndex(), Params.AOVolumeTextureSRV );
+		}
 	}
 
 	// Unbinds any buffers that have been bound.
 	void UnbindBuffers(FRHICommandList& RHICmdList)
 	{
 		FComputeShaderRHIParamRef ShaderRHI = GetComputeShader();
-#if LPV_VOLUME_TEXTURE
 		for ( int i = 0; i < 7; i++ )
 		{
 			if ( LpvBufferSRVParameters[i].IsBound() )
-		{
+		    {
 				RHICmdList.SetShaderTexture(ShaderRHI, LpvBufferSRVParameters[i].GetBaseIndex(), FTextureRHIParamRef());
-		}
+		    }
 			if ( LpvBufferUAVs[i].IsBound() )
-		{
-				RHICmdList.SetUAVParameter( ShaderRHI, LpvBufferUAVs[i].GetBaseIndex(), FUnorderedAccessViewRHIParamRef() );
-		}
-	}
-#else
-		if ( LpvBufferSRV.IsBound() )
-		{
-			RHICmdList.SetShaderResourceViewParameter( ShaderRHI, LpvBufferSRV.GetBaseIndex(), FShaderResourceViewRHIParamRef() );
-		}
-		if ( LpvBufferUAV.IsBound() )
-		{
-			RHICmdList.SetUAVParameter( ShaderRHI, LpvBufferUAV.GetBaseIndex(), FUnorderedAccessViewRHIParamRef() );
-		}
-#endif
+		    {
+			    RHICmdList.SetUAVParameter( ShaderRHI, LpvBufferUAVs[i].GetBaseIndex(), FUnorderedAccessViewRHIParamRef() );
+		    }
+	    }
 		if ( VplListHeadBufferSRV.IsBound() )
 		{
 			RHICmdList.SetShaderResourceViewParameter( ShaderRHI, VplListHeadBufferSRV.GetBaseIndex(), FShaderResourceViewRHIParamRef() );
@@ -306,28 +289,26 @@ public:
 		{
 			RHICmdList.SetUAVParameter( ShaderRHI, VplListBufferUAV.GetBaseIndex(), FUnorderedAccessViewRHIParamRef() );
 		}
-#if LPV_GV_VOLUME_TEXTURE
-		for ( int i = 0; i < 3; i++ )
+		for ( int i = 0; i < NUM_GV_TEXTURES; i++ )
 		{
 			if ( GvBufferSRVParameters[i].IsBound() )
-		{
-				RHICmdList.SetShaderTexture(ShaderRHI, GvBufferSRVParameters[i].GetBaseIndex(), FTextureRHIParamRef());
-		}
+			{
+					RHICmdList.SetShaderTexture(ShaderRHI, GvBufferSRVParameters[i].GetBaseIndex(), FTextureRHIParamRef());
+			}
 			if ( GvBufferUAVs[i].IsBound() )
-		{
-				RHICmdList.SetUAVParameter( ShaderRHI, GvBufferUAVs[i].GetBaseIndex(), FUnorderedAccessViewRHIParamRef() );
-		}
+			{
+					RHICmdList.SetUAVParameter( ShaderRHI, GvBufferUAVs[i].GetBaseIndex(), FUnorderedAccessViewRHIParamRef() );
+			}
 		}
-#else
-		if ( GvBufferSRV.IsBound() )
+
+		if ( AOVolumeTextureUAV.IsBound() )
 		{
-			RHICmdList.SetShaderResourceViewParameter( ShaderRHI, GvBufferSRV.GetBaseIndex(), FShaderResourceViewRHIParamRef() );
+			RHICmdList.SetUAVParameter( ShaderRHI, AOVolumeTextureUAV.GetBaseIndex(), FUnorderedAccessViewRHIParamRef() );
 		}
-		if ( GvBufferUAV.IsBound() )
+		if ( AOVolumeTextureSRV.IsBound() )
 		{
-			RHICmdList.SetUAVParameter( ShaderRHI, GvBufferUAV.GetBaseIndex(), FUnorderedAccessViewRHIParamRef() );
+			RHICmdList.SetShaderResourceViewParameter( ShaderRHI, AOVolumeTextureSRV.GetBaseIndex(), FShaderResourceViewRHIParamRef() );
 		}
-#endif
 		if(GvListBufferUAV.IsBound())
 		{
 			RHICmdList.SetUAVParameter(ShaderRHI, GvListBufferUAV.GetBaseIndex(), FUnorderedAccessViewRHIParamRef());
@@ -336,38 +317,33 @@ public:
 		{
 			RHICmdList.SetUAVParameter(ShaderRHI, GvListHeadBufferUAV.GetBaseIndex(), FUnorderedAccessViewRHIParamRef());
 		}
+		if ( GvListBufferSRV.IsBound() )
+		{
+			RHICmdList.SetShaderResourceViewParameter( ShaderRHI, GvListBufferSRV.GetBaseIndex(), FShaderResourceViewRHIParamRef() );
+		}
+		if ( GvListHeadBufferSRV.IsBound() )
+		{
+			RHICmdList.SetShaderResourceViewParameter( ShaderRHI, GvListHeadBufferSRV.GetBaseIndex(), FShaderResourceViewRHIParamRef() );
+		}
 	}
 
 protected:
-#if LPV_VOLUME_TEXTURE
 	FShaderResourceParameter LpvBufferSRVParameters[7];
 	FShaderResourceParameter LpvBufferUAVs[7];
-#else
-	FShaderResourceParameter LpvBufferSRV;
-	FShaderResourceParameter LpvBufferUAV;
-#endif
-
-#if LPV_VOLUME_TEXTURE || LPV_GV_VOLUME_TEXTURE
 	FShaderResourceParameter LpvVolumeTextureSampler;
-#endif
-
 	FShaderResourceParameter VplListHeadBufferSRV;
 	FShaderResourceParameter VplListHeadBufferUAV;
 	FShaderResourceParameter VplListBufferSRV;
 	FShaderResourceParameter VplListBufferUAV;
-#if LPV_GV_VOLUME_TEXTURE
-	FShaderResourceParameter GvBufferSRVParameters[3];
-	FShaderResourceParameter GvBufferUAVs[3];
-#else
-	FShaderResourceParameter GvBufferSRV;
-	FShaderResourceParameter GvBufferUAV;
-#endif
-
-	FShaderResourceParameter GvListBufferUAV;
-	FShaderResourceParameter GvListHeadBufferUAV;
 
+	FShaderResourceParameter GvBufferSRVParameters[NUM_GV_TEXTURES];
+	FShaderResourceParameter GvBufferUAVs[NUM_GV_TEXTURES];
 	FShaderResourceParameter GvListBufferSRV;
+	FShaderResourceParameter GvListBufferUAV;
 	FShaderResourceParameter GvListHeadBufferSRV;
+	FShaderResourceParameter GvListHeadBufferUAV;
+	FShaderResourceParameter AOVolumeTextureUAV;
+	FShaderResourceParameter AOVolumeTextureSRV;
 };
 
 
@@ -519,7 +495,6 @@ protected:
 };
 IMPLEMENT_SHADER_TYPE(,FLpvInject_GenerateVplListsCS,TEXT("LPVInject_GenerateVplLists"),TEXT("CSGenerateVplLists_LightDirectional"),SF_Compute);
 
-
 // ----------------------------------------------------------------------------
 // LPV accumulate VPL lists compute shader
 // ----------------------------------------------------------------------------
@@ -544,6 +519,86 @@ public:
 };
 IMPLEMENT_SHADER_TYPE(,FLpvInject_AccumulateVplListsCS,TEXT("LPVInject_AccumulateVplLists"),TEXT("CSAccumulateVplLists"),SF_Compute);
 
+// ----------------------------------------------------------------------------
+// LPV directional occlusion compute shader
+// ----------------------------------------------------------------------------
+class FLpvDirectionalOcclusionCS : public FLpvWriteShaderCSBase
+{
+	DECLARE_SHADER_TYPE(FLpvDirectionalOcclusionCS,Global);
+
+public:
+	//@todo-rco: Remove this when reenabling for OpenGL
+	static bool ShouldCache( EShaderPlatform Platform )		{ return IsFeatureLevelSupported(Platform, ERHIFeatureLevel::SM5) && !IsOpenGLPlatform(Platform); }
+
+	static void ModifyCompilationEnvironment( EShaderPlatform Platform, FShaderCompilerEnvironment& OutEnvironment )
+	{
+		FLpvWriteShaderCSBase::ModifyCompilationEnvironment( Platform, OutEnvironment );
+	}
+
+	FLpvDirectionalOcclusionCS()	{	}
+
+	explicit FLpvDirectionalOcclusionCS( const ShaderMetaType::CompiledShaderInitializerType& Initializer ) : FLpvWriteShaderCSBase(Initializer)		
+	{	
+	}
+
+	// FShader interface.
+	virtual bool Serialize(FArchive& Ar) override
+	{
+
+		bool bShaderHasOutdatedParameters = FLpvWriteShaderCSBase::Serialize( Ar );
+		return bShaderHasOutdatedParameters;
+	}
+
+	virtual void SetParameters(
+		FRHICommandList& RHICmdList, 
+		FLpvBaseWriteShaderParams& BaseParams )
+	{
+		FLpvWriteShaderCSBase::SetParameters( RHICmdList, BaseParams );
+		const FComputeShaderRHIParamRef ShaderRHI = GetComputeShader();
+	}
+};
+IMPLEMENT_SHADER_TYPE(,FLpvDirectionalOcclusionCS,TEXT("LpvDirectionalOcclusion"),TEXT("CSDirectionalOcclusion"),SF_Compute);
+
+// ----------------------------------------------------------------------------
+// LPV directional occlusion compute shader
+// ----------------------------------------------------------------------------
+class FLpvCopyAOVolumeCS : public FLpvWriteShaderCSBase
+{
+	DECLARE_SHADER_TYPE(FLpvCopyAOVolumeCS,Global);
+
+public:
+	//@todo-rco: Remove this when reenabling for OpenGL
+	static bool ShouldCache( EShaderPlatform Platform )		{ return IsFeatureLevelSupported(Platform, ERHIFeatureLevel::SM5) && !IsOpenGLPlatform(Platform); }
+
+	static void ModifyCompilationEnvironment( EShaderPlatform Platform, FShaderCompilerEnvironment& OutEnvironment )
+	{
+		FLpvWriteShaderCSBase::ModifyCompilationEnvironment( Platform, OutEnvironment );
+	}
+
+	FLpvCopyAOVolumeCS()	{	}
+
+	explicit FLpvCopyAOVolumeCS( const ShaderMetaType::CompiledShaderInitializerType& Initializer ) : FLpvWriteShaderCSBase(Initializer)		
+	{	
+	}
+
+	// FShader interface.
+	virtual bool Serialize(FArchive& Ar) override
+	{
+
+		bool bShaderHasOutdatedParameters = FLpvWriteShaderCSBase::Serialize( Ar );
+		return bShaderHasOutdatedParameters;
+	}
+
+	virtual void SetParameters(
+		FRHICommandList& RHICmdList, 
+		FLpvBaseWriteShaderParams& BaseParams )
+	{	
+		FLpvWriteShaderCSBase::SetParameters( RHICmdList, BaseParams );
+		const FComputeShaderRHIParamRef ShaderRHI = GetComputeShader();
+	}
+};
+IMPLEMENT_SHADER_TYPE(,FLpvCopyAOVolumeCS,TEXT("LpvDirectionalOcclusion"),TEXT("CSCopyAOVolume"),SF_Compute);
+
 
 // ----------------------------------------------------------------------------
 // Compute shader to build a geometry volume
@@ -569,8 +624,6 @@ public:
 };
 IMPLEMENT_SHADER_TYPE(,FLpvBuildGeometryVolumeCS,TEXT("LPVBuildGeometryVolume"),TEXT("CSBuildGeometryVolume"),SF_Compute);
 
-
-
 // ----------------------------------------------------------------------------
 // LPV propagate compute shader
 // ----------------------------------------------------------------------------
@@ -578,8 +631,7 @@ enum PropagateShaderFlags
 {
 	PROPAGATE_SECONDARY_OCCLUSION	= 0x01,
 	PROPAGATE_MULTIPLE_BOUNCES		= 0x02,
-
-	PROPAGATE_SECONDARY_OCCLUSION_AND_MULTIPLE_BOUNCES		= PROPAGATE_SECONDARY_OCCLUSION | PROPAGATE_MULTIPLE_BOUNCES,
+	PROPAGATE_AO					= 0x04,
 };
 
 template<uint32 ShaderFlags>
@@ -595,7 +647,9 @@ public:
 	{
 		OutEnvironment.SetDefine(TEXT("LPV_SECONDARY_OCCLUSION"), (uint32)(ShaderFlags & PROPAGATE_SECONDARY_OCCLUSION ? 1 : 0));
 		OutEnvironment.SetDefine(TEXT("LPV_MULTIPLE_BOUNCES_ENABLED"), (uint32)(ShaderFlags & PROPAGATE_MULTIPLE_BOUNCES ? 1 : 0));
-			
+		OutEnvironment.SetDefine(TEXT("LPV_PROPAGATE_AO"), (uint32)(ShaderFlags & PROPAGATE_AO ? 1 : 0));
+		OutEnvironment.CompilerFlags.Add(CFLAG_StandardOptimization);
+
 		FLpvWriteShaderCSBase::ModifyCompilationEnvironment( Platform, OutEnvironment );
 	}
 
@@ -606,29 +660,37 @@ public:
 	virtual bool Serialize( FArchive& Ar ) override			{ return FLpvWriteShaderCSBase::Serialize( Ar ); }
 };
 
-IMPLEMENT_SHADER_TYPE(template<>,TLpvPropagateCS<0>,													TEXT("LPVPropagate"),TEXT("CSPropogate"),SF_Compute);
-IMPLEMENT_SHADER_TYPE(template<>,TLpvPropagateCS<PROPAGATE_SECONDARY_OCCLUSION>,						TEXT("LPVPropagate"),TEXT("CSPropogate"),SF_Compute);
-IMPLEMENT_SHADER_TYPE(template<>,TLpvPropagateCS<PROPAGATE_MULTIPLE_BOUNCES>,							TEXT("LPVPropagate"),TEXT("CSPropogate"),SF_Compute);
-IMPLEMENT_SHADER_TYPE(template<>,TLpvPropagateCS<PROPAGATE_SECONDARY_OCCLUSION_AND_MULTIPLE_BOUNCES>,	TEXT("LPVPropagate"),TEXT("CSPropogate"),SF_Compute);
+IMPLEMENT_SHADER_TYPE(template<>,TLpvPropagateCS<0>,																		TEXT("LPVPropagate"),TEXT("CSPropagate"),SF_Compute);
+IMPLEMENT_SHADER_TYPE(template<>,TLpvPropagateCS<PROPAGATE_SECONDARY_OCCLUSION>,											TEXT("LPVPropagate"),TEXT("CSPropagate"),SF_Compute);
+IMPLEMENT_SHADER_TYPE(template<>,TLpvPropagateCS<PROPAGATE_MULTIPLE_BOUNCES>,												TEXT("LPVPropagate"),TEXT("CSPropagate"),SF_Compute);
+IMPLEMENT_SHADER_TYPE(template<>,TLpvPropagateCS<PROPAGATE_SECONDARY_OCCLUSION|PROPAGATE_MULTIPLE_BOUNCES>,					TEXT("LPVPropagate"),TEXT("CSPropagate"),SF_Compute);
+IMPLEMENT_SHADER_TYPE(template<>,TLpvPropagateCS<PROPAGATE_AO>,																TEXT("LPVPropagate"),TEXT("CSPropagate"),SF_Compute);
+IMPLEMENT_SHADER_TYPE(template<>,TLpvPropagateCS<PROPAGATE_AO|PROPAGATE_SECONDARY_OCCLUSION>,								TEXT("LPVPropagate"),TEXT("CSPropagate"),SF_Compute);
+IMPLEMENT_SHADER_TYPE(template<>,TLpvPropagateCS<PROPAGATE_AO|PROPAGATE_MULTIPLE_BOUNCES>,									TEXT("LPVPropagate"),TEXT("CSPropagate"),SF_Compute);
+IMPLEMENT_SHADER_TYPE(template<>,TLpvPropagateCS<PROPAGATE_AO|PROPAGATE_SECONDARY_OCCLUSION|PROPAGATE_MULTIPLE_BOUNCES>,	TEXT("LPVPropagate"),TEXT("CSPropagate"),SF_Compute);
 
-FLpvWriteShaderCSBase* GetPropagateShader(FViewInfo& View, bool bSecondaryOcclusion, bool bSecondaryBounces )
+FLpvWriteShaderCSBase* GetPropagateShader( FViewInfo& View, uint32 ShaderFlags )
 {
-	FLpvWriteShaderCSBase* ShaderOut = NULL;
-	if ( bSecondaryBounces )
-	{
-		if ( bSecondaryOcclusion ) 
-			ShaderOut = (FLpvWriteShaderCSBase*)*TShaderMapRef<TLpvPropagateCS<PROPAGATE_SECONDARY_OCCLUSION_AND_MULTIPLE_BOUNCES> >(View.ShaderMap);
-		else 
-			ShaderOut = (FLpvWriteShaderCSBase*)*TShaderMapRef<TLpvPropagateCS<PROPAGATE_MULTIPLE_BOUNCES> >(View.ShaderMap);
-	}
-	else
+	switch( ShaderFlags )
 	{
-		if ( bSecondaryOcclusion ) 
-			ShaderOut = (FLpvWriteShaderCSBase*)*TShaderMapRef<TLpvPropagateCS<PROPAGATE_SECONDARY_OCCLUSION> >(View.ShaderMap);
-		else 
-			ShaderOut = (FLpvWriteShaderCSBase*)*TShaderMapRef<TLpvPropagateCS<0> >(View.ShaderMap);
+	case 0:
+		return (FLpvWriteShaderCSBase*)*TShaderMapRef<TLpvPropagateCS<0> >( View.ShaderMap );
+	case PROPAGATE_SECONDARY_OCCLUSION:										
+		return (FLpvWriteShaderCSBase*)*TShaderMapRef<TLpvPropagateCS<PROPAGATE_SECONDARY_OCCLUSION> >( View.ShaderMap );
+	case PROPAGATE_MULTIPLE_BOUNCES:	
+		return (FLpvWriteShaderCSBase*)*TShaderMapRef<TLpvPropagateCS<PROPAGATE_MULTIPLE_BOUNCES> >( View.ShaderMap );
+	case PROPAGATE_SECONDARY_OCCLUSION|PROPAGATE_MULTIPLE_BOUNCES:	
+		return (FLpvWriteShaderCSBase*)*TShaderMapRef<TLpvPropagateCS<PROPAGATE_SECONDARY_OCCLUSION|PROPAGATE_MULTIPLE_BOUNCES> >( View.ShaderMap );
+	case PROPAGATE_AO:							
+		return (FLpvWriteShaderCSBase*)*TShaderMapRef<TLpvPropagateCS<PROPAGATE_AO> >( View.ShaderMap );
+	case PROPAGATE_AO|PROPAGATE_SECONDARY_OCCLUSION:							
+		return (FLpvWriteShaderCSBase*)*TShaderMapRef<TLpvPropagateCS<PROPAGATE_AO|PROPAGATE_SECONDARY_OCCLUSION> >( View.ShaderMap );
+	case PROPAGATE_AO|PROPAGATE_MULTIPLE_BOUNCES:
+		return (FLpvWriteShaderCSBase*)*TShaderMapRef<TLpvPropagateCS<PROPAGATE_AO|PROPAGATE_MULTIPLE_BOUNCES> >( View.ShaderMap );
+	case PROPAGATE_AO|PROPAGATE_SECONDARY_OCCLUSION|PROPAGATE_MULTIPLE_BOUNCES:
+		return (FLpvWriteShaderCSBase*)*TShaderMapRef<TLpvPropagateCS<PROPAGATE_AO|PROPAGATE_SECONDARY_OCCLUSION|PROPAGATE_MULTIPLE_BOUNCES> >( View.ShaderMap );
 	}
-	return ShaderOut;
+	return NULL;
 }
 
 // ----------------------------------------------------------------------------
@@ -636,7 +698,8 @@ FLpvWriteShaderCSBase* GetPropagateShader(FViewInfo& View, bool bSecondaryOcclus
 // ----------------------------------------------------------------------------
 enum EInjectFlags
 {
-	INJECT_SHADOWED = 0x01,
+	INJECT_SHADOW_CASTING	 = 0x01,
+	INJECT_SPOT_ATTENUATION  = 0x02,
 };
 
 class FLpvInjectShader_Base : public FLpvWriteShaderCSBase
@@ -662,10 +725,11 @@ public:
 // ----------------------------------------------------------------------------
 // Point light injection compute shader
 // ----------------------------------------------------------------------------
+
 template <uint32 InjectFlags>
-class TLpvInject_PointLightCS : public FLpvInjectShader_Base
+class TLpvInject_LightCS : public FLpvInjectShader_Base
 {
-	DECLARE_SHADER_TYPE(TLpvInject_PointLightCS,Global);
+	DECLARE_SHADER_TYPE(TLpvInject_LightCS,Global);
 
 public:
 	//@todo-rco: Remove this when reenabling for OpenGL
@@ -673,18 +737,22 @@ public:
 
 	static void ModifyCompilationEnvironment( EShaderPlatform Platform, FShaderCompilerEnvironment& OutEnvironment )
 	{
-		OutEnvironment.SetDefine(TEXT("SHADOW_CASTING"), (uint32)(InjectFlags & INJECT_SHADOWED ? 1 : 0));
+		OutEnvironment.SetDefine(TEXT("SHADOW_CASTING"),   (uint32)(InjectFlags & INJECT_SHADOW_CASTING ? 1 : 0));
+		OutEnvironment.SetDefine(TEXT("SPOT_ATTENUATION"), (uint32)(InjectFlags & INJECT_SPOT_ATTENUATION ? 1 : 0));
 		FLpvWriteShaderCSBase::ModifyCompilationEnvironment( Platform, OutEnvironment );
 	}
 
-	TLpvInject_PointLightCS()	{	}
+	TLpvInject_LightCS()	{	}
 
-	explicit TLpvInject_PointLightCS( const ShaderMetaType::CompiledShaderInitializerType& Initializer ) : FLpvInjectShader_Base(Initializer)		{	}
+	explicit TLpvInject_LightCS( const ShaderMetaType::CompiledShaderInitializerType& Initializer ) : FLpvInjectShader_Base(Initializer)		{	}
 
 	virtual bool Serialize( FArchive& Ar ) override			{ return FLpvInjectShader_Base::Serialize( Ar ); }
 };
-IMPLEMENT_SHADER_TYPE(template<>,TLpvInject_PointLightCS<0>,TEXT("LPVDirectLightInject"),TEXT("CSPointLightInject_ListGenCS"),SF_Compute);
-IMPLEMENT_SHADER_TYPE(template<>,TLpvInject_PointLightCS<INJECT_SHADOWED>,TEXT("LPVDirectLightInject"),TEXT("CSPointLightInject_ListGenCS"),SF_Compute);
+
+IMPLEMENT_SHADER_TYPE(template<>,TLpvInject_LightCS<0>,TEXT("LPVDirectLightInject"),TEXT("CSLightInject_ListGenCS"),SF_Compute);
+IMPLEMENT_SHADER_TYPE(template<>,TLpvInject_LightCS<1>,TEXT("LPVDirectLightInject"),TEXT("CSLightInject_ListGenCS"),SF_Compute);
+IMPLEMENT_SHADER_TYPE(template<>,TLpvInject_LightCS<2>,TEXT("LPVDirectLightInject"),TEXT("CSLightInject_ListGenCS"),SF_Compute);
+IMPLEMENT_SHADER_TYPE(template<>,TLpvInject_LightCS<3>,TEXT("LPVDirectLightInject"),TEXT("CSLightInject_ListGenCS"),SF_Compute);
 
 // ----------------------------------------------------------------------------
 // FLightPropagationVolume
@@ -699,42 +767,28 @@ FLightPropagationVolume::FLightPropagationVolume() :
 	, bEnabled( false )
 	, mWriteBufferIndex( 0 )
 	, GeometryVolumeGenerated( false )
+	, bGeometryVolumeNeeded( false )
+	, AsyncJobFenceID(-1)
 {
 	bNeedsBufferClear = true;
 
-#if !LPV_VOLUME_TEXTURE
-	for ( int i=0; i<2; i++ )
-	{
-		mLpvBuffers[i] = new FRWBufferStructured();
-		mLpvBuffers[i]->Initialize( sizeof(LPVBufferElement), LPV_GRIDRES*LPV_GRIDRES*LPV_GRIDRES, BUF_FastVRAM, false, false );
-	}
-#endif
 	// VPL List buffers
 	mVplListBuffer = new FRWBufferStructured();
-	mVplListBuffer->Initialize(sizeof(VplListEntry), 256 * 256 * LPV_GRIDRES, 0, true, false);
+	int32 RSMResolution = GSceneRenderTargets.GetReflectiveShadowMapResolution();
+	int32 GvListBufferSize = RSMResolution * RSMResolution * 16; // Allow 16 layers of depth per every pixel of the RSM (on average) 
+	int32 VplListBufferSize = RSMResolution * RSMResolution * 4; // Allow 4 layers of depth per pixel in the RSM (1 for the RSM injection + 3 for light injection)
+	mVplListBuffer->Initialize( sizeof( VplListEntry ), VplListBufferSize, 0, true, false );
 	mVplListHeadBuffer = new FRWBufferByteAddress();
-	mVplListHeadBuffer->Initialize(LPV_GRIDRES*LPV_GRIDRES*LPV_GRIDRES * 4, BUF_ByteAddressBuffer);
+	mVplListHeadBuffer->Initialize( LPV_GRIDRES*LPV_GRIDRES*LPV_GRIDRES*4, BUF_ByteAddressBuffer );
 
 	// Geometry volume buffers
-#if !LPV_GV_VOLUME_TEXTURE
-	GvBuffer = new FRWBuffer();
-
-#if LPV_MULTIPLE_BOUNCES
-	GvBuffer->Initialize( 4, LPV_GRIDRES*LPV_GRIDRES*LPV_GRIDRES*4, PF_A2B10G10R10, BUF_FastVRAM );
-#else
-	GvBuffer->Initialize( 4, LPV_GRIDRES*LPV_GRIDRES*LPV_GRIDRES*3, PF_A2B10G10R10, 0 );
-#endif
-
-#endif // !LPV_GV_VOLUME_TEXTURE
-
 	GvListBuffer = new FRWBufferStructured();
-	GvListBuffer->Initialize(sizeof(VplListEntry), 256 * 256 * LPV_GRIDRES, 0, true, false);
+	GvListBuffer->Initialize( sizeof( VplListEntry ), GvListBufferSize, 0, true, false );
 	GvListHeadBuffer = new FRWBufferByteAddress();
-	GvListHeadBuffer->Initialize(LPV_GRIDRES*LPV_GRIDRES*LPV_GRIDRES * 4, BUF_ByteAddressBuffer);
+	GvListHeadBuffer->Initialize( LPV_GRIDRES*LPV_GRIDRES*LPV_GRIDRES*4, BUF_ByteAddressBuffer );
 
 	LpvWriteUniformBufferParams = new FLpvWriteUniformBufferParameters;
 	FMemory::Memzero( LpvWriteUniformBufferParams, sizeof(FLpvWriteUniformBufferParameters) );
-
 	bInitialized = false;
 }
 
@@ -746,7 +800,6 @@ FLightPropagationVolume::~FLightPropagationVolume()
 {
 	LpvWriteUniformBuffer.ReleaseResource();
 
-#if LPV_VOLUME_TEXTURE
 	// Note: this is double-buffered!
 	for ( int i = 0; i < 2; i++ )
 	{
@@ -755,13 +808,6 @@ FLightPropagationVolume::~FLightPropagationVolume()
 			LpvVolumeTextures[i][j].SafeRelease();
 		}
 	}
-#else
-	for ( int i = 0; i < 2; i++ )
-	{
-		mLpvBuffers[i]->Release();
-		delete mLpvBuffers[i];
-	}
-#endif 
 
 	mVplListHeadBuffer->Release();
 	delete mVplListHeadBuffer;
@@ -769,15 +815,10 @@ FLightPropagationVolume::~FLightPropagationVolume()
 	mVplListBuffer->Release();
 	delete mVplListBuffer;
 
-#if LPV_GV_VOLUME_TEXTURE
-	for ( int i = 0; i < 3; i++ )
+	for ( int i = 0; i < NUM_GV_TEXTURES; i++ )
 	{
 		GvVolumeTextures[i].SafeRelease();
 	}
-#else
-	GvBuffer->Release();
-	delete GvBuffer;
-#endif
 
 	GvListHeadBuffer->Release();
 	delete GvListHeadBuffer;
@@ -792,10 +833,16 @@ FLightPropagationVolume::~FLightPropagationVolume()
 * Sets up the LPV at the beginning of the frame
 */
 void FLightPropagationVolume::InitSettings(FRHICommandList& RHICmdList, const FSceneView& View)
+{
+	int32 NumFastLpvTextures = 7;
+	int32 NumFastGvTextures = 2;
+	FIntPoint BufferSize = GSceneRenderTargets.GetBufferSizeXY();
+	if ( BufferSize.X >= 1600 && BufferSize.Y >= 900 )
 	{
+		NumFastLpvTextures = 5;
+		NumFastGvTextures = 1;
+	}
 	check(View.GetFeatureLevel() >= ERHIFeatureLevel::SM5);
-
-#if LPV_VOLUME_TEXTURE
 	if ( !bInitialized )
 	{
 		// @TODO: this should probably be derived from FRenderResource (with InitDynamicRHI etc)
@@ -822,20 +869,30 @@ void FLightPropagationVolume::InitSettings(FRHICommandList& RHICmdList, const FS
 			}
 		}
 
-#if LPV_GV_VOLUME_TEXTURE
 		{
 			const TCHAR* Names[] = { TEXT("LPV_GV0"), TEXT("LPV_GV1"), TEXT("LPV_GV2") };
 
-			for(int i = 0; i < 3; i++)
+			for ( int i = 0; i < NUM_GV_TEXTURES; i++ )
 			{
 				GRenderTargetPool.FindFreeElement(Desc, GvVolumeTextures[i], Names[i]);
 			}
 		}
-#endif
+
+		{
+			FPooledRenderTargetDesc AODesc(FPooledRenderTargetDesc::CreateVolumeDesc(
+				LPV_GRIDRES,
+				LPV_GRIDRES,
+				LPV_GRIDRES,
+				PF_G8,
+				TexCreate_HideInVisualizeTexture,
+				TexCreate_ShaderResource | TexCreate_UAV | TexCreate_FastVRAM,
+				false,
+				1));
+			GRenderTargetPool.FindFreeElement(AODesc, AOVolumeTexture, TEXT("LPVAOVolume"));
+		}
+
 		bInitialized = true;
 	}  
-#endif
-
 
 	// Copy the LPV postprocess settings
 	Strength	 = View.FinalPostProcessSettings.LPVIntensity;
@@ -845,6 +902,7 @@ void FLightPropagationVolume::InitSettings(FRHICommandList& RHICmdList, const FS
 	SecondaryOcclusionStrength = View.FinalPostProcessSettings.LPVSecondaryOcclusionIntensity;
 	SecondaryBounceStrength = View.FinalPostProcessSettings.LPVSecondaryBounceIntensity;
 
+	bGeometryVolumeNeeded = View.FinalPostProcessSettings.LPVSecondaryOcclusionIntensity > 0.001f || View.FinalPostProcessSettings.LPVDirectionalOcclusionIntensity > 0.001;
 	GeometryVolumeGenerated = false;
 
 	if ( !bEnabled )
@@ -865,11 +923,6 @@ void FLightPropagationVolume::InitSettings(FRHICommandList& RHICmdList, const FS
 	// Clear the UAVs if necessary
 	if ( bNeedsBufferClear )
 	{
-#if !LPV_VOLUME_TEXTURE
-		uint32 ClearValues[4] = {0};
-		RHICmdList.ClearUAV( mVplListHeadBuffer->UAV, ClearValues );
-#endif
-
 		ClearMultiplier = 0.0f;
 		bNeedsBufferClear = false;
 	}
@@ -902,15 +955,23 @@ void FLightPropagationVolume::InitSettings(FRHICommandList& RHICmdList, const FS
 		LpvWriteUniformBufferParams->GeometryVolumeInjectionBias	= View.FinalPostProcessSettings.LPVGeometryVolumeBias * LpvScale;
 		LpvWriteUniformBufferParams->VplInjectionBias				= View.FinalPostProcessSettings.LPVVplInjectionBias * LpvScale;
 		LpvWriteUniformBufferParams->PropagationIndex				= 0;
-		LpvWriteUniformBufferParams->EmissiveInjectionMultiplier	= View.FinalPostProcessSettings.LPVEmissiveInjectionIntensity;
-
+		LpvWriteUniformBufferParams->EmissiveInjectionMultiplier	= View.FinalPostProcessSettings.LPVEmissiveInjectionIntensity
+			* LpvWriteUniformBufferParams->RsmAreaIntensityMultiplier * CVarLPVEmissiveIntensityMultiplier.GetValueOnRenderThread() * 0.25f;
+		LpvWriteUniformBufferParams->DirectionalOcclusionIntensity	= View.FinalPostProcessSettings.LPVDirectionalOcclusionIntensity;
+		LpvWriteUniformBufferParams->DirectionalOcclusionRadius		= View.FinalPostProcessSettings.LPVDirectionalOcclusionRadius;
+		LpvWriteUniformBufferParams->RsmPixelToTexcoordMultiplier	= 1.0f / float(GSceneRenderTargets.GetReflectiveShadowMapResolution() - 1);
+
+		LpvReadUniformBufferParams.DirectionalOcclusionIntensity	= View.FinalPostProcessSettings.LPVDirectionalOcclusionIntensity;
+		LpvReadUniformBufferParams.DiffuseOcclusionExponent			= View.FinalPostProcessSettings.LPVDiffuseOcclusionExponent;
+		LpvReadUniformBufferParams.SpecularOcclusionExponent		= View.FinalPostProcessSettings.LPVSpecularOcclusionExponent;
+		LpvReadUniformBufferParams.DiffuseOcclusionIntensity		= View.FinalPostProcessSettings.LPVDiffuseOcclusionIntensity;
+		LpvReadUniformBufferParams.SpecularOcclusionIntensity		= View.FinalPostProcessSettings.LPVSpecularOcclusionIntensity;
 
 		LpvReadUniformBufferParams.mLpvGridOffset		= mGridOffset;
 		LpvReadUniformBufferParams.LpvScale				= LpvScale;
 		LpvReadUniformBufferParams.OneOverLpvScale		= OneOverLpvScale;
-
-		LpvReadUniformBufferParams.LpvScale				= LpvScale;
-		LpvReadUniformBufferParams.OneOverLpvScale		= OneOverLpvScale;
+		LpvReadUniformBufferParams.SpecularIntensity	= CVarLPVSpecularIntensity.GetValueOnRenderThread();
+		LpvReadUniformBufferParams.DiffuseIntensity		= CVarLPVDiffuseIntensity.GetValueOnRenderThread();
 
 		// Compute the bounding box
 		FVector Centre = ( FVector(mGridOffset.X, mGridOffset.Y, mGridOffset.Z ) + FVector(0.5f,0.5f,0.5f) - HalfGridRes ) * -LpvScale;
@@ -938,7 +999,8 @@ void FLightPropagationVolume::Clear(FRHICommandListImmediate& RHICmdList, FViewI
 	LpvWriteUniformBuffer.SetContents( *LpvWriteUniformBufferParams );
 
 	// TODO: these could be run in parallel... 
-	//RHIBatchComputeShaderBegin();
+	RHICmdList.AutomaticCacheFlushAfterComputeShader(false);
+
 	// Clear the list buffers
 	{
 		TShaderMapRef<FLpvClearListsCS> Shader(View.ShaderMap);
@@ -964,7 +1026,7 @@ void FLightPropagationVolume::Clear(FRHICommandListImmediate& RHICmdList, FViewI
 	}
 
 	// Clear the geometry volume if necessary
-	if ( SecondaryOcclusionStrength > 0.001f )
+	if ( bGeometryVolumeNeeded )
 	{
 		TShaderMapRef<FLpvClearGeometryVolumeCS> Shader(View.ShaderMap);
 		RHICmdList.SetComputeShader(Shader->GetComputeShader());
@@ -975,7 +1037,8 @@ void FLightPropagationVolume::Clear(FRHICommandListImmediate& RHICmdList, FViewI
 		DispatchComputeShader(RHICmdList, *Shader, LPV_GRIDRES/4, LPV_GRIDRES/4, LPV_GRIDRES/4 );
 		Shader->UnbindBuffers(RHICmdList);
 	}
-	//RHIBatchComputeShaderEnd();
+	RHICmdList.AutomaticCacheFlushAfterComputeShader(true);
+	RHICmdList.FlushComputeShaderCache();
 
 	RHICmdList.SetUAVParameter( FComputeShaderRHIRef(), 7, mVplListBuffer->UAV, 0 );
 	RHICmdList.SetUAVParameter( FComputeShaderRHIRef(), 7, GvListBuffer->UAV, 0 );
@@ -998,7 +1061,9 @@ void FLightPropagationVolume::GetShadowInfo( const FProjectedShadowInfo& Project
 	FVector4 WorldRight = RsmInfoOut.ShadowToWorld.TransformVector( ShadowRight );
 	float ShadowArea = WorldUp.Size3() * WorldRight.Size3();
 
-	static float IdealCubeSizeMultiplier = 0.5f; 
+	int32 RSMResolution = GSceneRenderTargets.GetReflectiveShadowMapResolution();
+	float ResolutionMultiplier = RSMResolution / 256.0f;
+	float IdealCubeSizeMultiplier = 0.5f * ResolutionMultiplier; 
 	float IdealRsmArea = ( CubeSize * IdealCubeSizeMultiplier * CubeSize * IdealCubeSizeMultiplier );
 	RsmInfoOut.AreaBrightnessMultiplier = ShadowArea/IdealRsmArea;
 }
@@ -1021,6 +1086,8 @@ void FLightPropagationVolume::SetVplInjectionConstants(
 		LpvStrength = Strength;
 	}
 	LpvStrength *= RsmInfo.AreaBrightnessMultiplier;
+	LpvStrength *= CVarLPVIntensity.GetValueOnRenderThread();
+	LpvWriteUniformBufferParams->RsmAreaIntensityMultiplier = RsmInfo.AreaBrightnessMultiplier;
 
 	LpvStrength *= LightProxy->GetIndirectLightingScale();
 	LpvWriteUniformBufferParams->mRsmToWorld = RsmInfo.ShadowToWorld;
@@ -1054,13 +1121,15 @@ void FLightPropagationVolume::InjectDirectionalLightRSM(
 		GetShaderParams( ShaderParams );
 		Shader->SetParameters(RHICmdList, ShaderParams, RsmDiffuseTex, RsmNormalTex, RsmDepthTex );
 
-		DispatchComputeShader(RHICmdList, *Shader, 256/8, 256/8, 1 ); 
+		int32 RSMResolution = GSceneRenderTargets.GetReflectiveShadowMapResolution();
+		// todo: what if not divisble by 8?
+		DispatchComputeShader(RHICmdList, *Shader, RSMResolution / 8, RSMResolution / 8, 1 ); 
 
 		Shader->UnbindBuffers(RHICmdList); 
 	}
 
 	// If this is the first directional light, build the geometry volume with it
-	if ( !GeometryVolumeGenerated && SecondaryOcclusionStrength > 0.001f )
+	if ( !GeometryVolumeGenerated && bGeometryVolumeNeeded )
 	{
 		SCOPED_DRAW_EVENT(RHICmdList, LpvBuildGeometryVolume);
 		GeometryVolumeGenerated = true;
@@ -1088,17 +1157,53 @@ void FLightPropagationVolume::InjectDirectionalLightRSM(
 	mInjectedLightCount++;
 }
 
+/**
+* Injects sky lighting into the LPV
+*/
+void FLightPropagationVolume::ComputeDirectionalOcclusion( FRHICommandListImmediate& RHICmdList, FViewInfo& View )
+{
+	//if ( View.FinalPostProcessSettings.ContributingCubemaps.Num() > 0 || GScene-> )
+	{
+		// Compute directional occlusion
+		{
+			SCOPED_DRAW_EVENT(RHICmdList, LpvDirectionalOcclusion);
+
+			mWriteBufferIndex = 1-mWriteBufferIndex; // Swap buffers with each iteration
+			TShaderMapRef<FLpvDirectionalOcclusionCS> Shader(View.ShaderMap);
+			RHICmdList.SetComputeShader(Shader->GetComputeShader());
+			FLpvBaseWriteShaderParams ShaderParams;
+			GetShaderParams( ShaderParams );
+			Shader->SetParameters( RHICmdList, ShaderParams );
+			LpvWriteUniformBuffer.SetContents( *LpvWriteUniformBufferParams );
+
+			DispatchComputeShader( RHICmdList, *Shader, LPV_GRIDRES/4, LPV_GRIDRES/4, LPV_GRIDRES/4 );
+			Shader->UnbindBuffers(RHICmdList); 
+		}
+	}
+	RHICmdList.FlushComputeShaderCache();
+}
+
 /**
 * Propagates light in the LPV 
 */
-void FLightPropagationVolume::Propagate(FRHICommandListImmediate& RHICmdList, FViewInfo& View)
+void FLightPropagationVolume::Update( FRHICommandListImmediate& RHICmdList, FViewInfo& View )
 {
 	if ( !bEnabled )
 	{
 		return;
 	}
 
-	// Inject the VPLs into the LPV
+	LpvWriteUniformBuffer.SetContents( *LpvWriteUniformBufferParams );
+
+	// Begin the async compute job
+	RHICmdList.BeginAsyncComputeJob_DrawThread(AsyncComputePriority_High);
+
+	check(View.GetFeatureLevel() == ERHIFeatureLevel::SM5);
+
+	bool bSecondaryOcclusion	= (SecondaryOcclusionStrength > 0.001f); 
+	bool bSecondaryBounces		= (SecondaryBounceStrength > 0.001f); 
+	bool bDirectionalOcclusion	= (LpvWriteUniformBufferParams->DirectionalOcclusionIntensity > 0.001f);
+
 	if ( mInjectedLightCount )
 	{
 		SCOPED_DRAW_EVENT(RHICmdList, LpvAccumulateVplLists);
@@ -1107,36 +1212,47 @@ void FLightPropagationVolume::Propagate(FRHICommandListImmediate& RHICmdList, FV
 		TShaderMapRef<FLpvInject_AccumulateVplListsCS> Shader(View.ShaderMap);
 		RHICmdList.SetComputeShader(Shader->GetComputeShader());
 
-		LpvWriteUniformBuffer.SetContents( *LpvWriteUniformBufferParams );
+		//LpvWriteUniformBuffer.SetContents( *LpvWriteUniformBufferParams );
 
 		FLpvBaseWriteShaderParams ShaderParams;
 		GetShaderParams( ShaderParams );
 		Shader->SetParameters(RHICmdList, ShaderParams );
 
 		DispatchComputeShader(RHICmdList, *Shader, LPV_GRIDRES/4, LPV_GRIDRES/4, LPV_GRIDRES/4 );
+		RHICmdList.FlushComputeShaderCache();
 
 		Shader->UnbindBuffers(RHICmdList);
 	}
 
 	// Propagate lighting, ping-ponging between the two buffers
+	if ( bDirectionalOcclusion ) 
 	{
-		static int PropagationStep = 0;
+		ComputeDirectionalOcclusion( RHICmdList, View );
+	}
 
+	// Propagate lighting, ping-ponging between the two buffers
+	{
 		SCOPED_DRAW_EVENT(RHICmdList, LpvPropagate);
-		static int numIterations = 3;
-		for ( int i=0; i<numIterations; i++ )
+
+		int32 LPVNumPropagationSteps = CVarNumPropagationSteps.GetValueOnRenderThread();
+
+		for ( int i = 0; i < LPVNumPropagationSteps; i++ )
 		{
-			// Disable secondary occlusion on the first iteration
-			bool bSecondaryOcclusion	= (SecondaryOcclusionStrength > 0.001f); 
-			bool bSecondaryBounces		= (SecondaryBounceStrength > 0.001f); 
 			mWriteBufferIndex = 1-mWriteBufferIndex; // Swap buffers with each iteration
-			FLpvWriteShaderCSBase* Shader = GetPropagateShader(View, bSecondaryOcclusion, bSecondaryBounces);
-			RHICmdList.SetComputeShader(Shader->GetComputeShader());
 
-			LpvWriteUniformBufferParams->PropagationIndex = PropagationStep;
-			PropagationStep++;
+			// Compute shader flags
+			uint32 ShaderFlags = 0;
+			if ( bSecondaryOcclusion ) 
+				ShaderFlags |= PROPAGATE_SECONDARY_OCCLUSION;
+			if ( bSecondaryBounces ) 
+				ShaderFlags |= PROPAGATE_MULTIPLE_BOUNCES;
+			if ( i < CVarLPVNumAOPropagationSteps.GetValueOnRenderThread() )
+				ShaderFlags |= PROPAGATE_AO;
 
-			LpvWriteUniformBuffer.SetContents( *LpvWriteUniformBufferParams );
+			FLpvWriteShaderCSBase* Shader = GetPropagateShader( View, ShaderFlags );
+			RHICmdList.SetComputeShader(Shader->GetComputeShader());
+
+			LpvWriteUniformBufferParams->PropagationIndex = i;
 
 			FLpvBaseWriteShaderParams ShaderParams;
 			GetShaderParams( ShaderParams );
@@ -1144,12 +1260,47 @@ void FLightPropagationVolume::Propagate(FRHICommandListImmediate& RHICmdList, FV
 
 			DispatchComputeShader(RHICmdList, Shader, LPV_GRIDRES/4, LPV_GRIDRES/4, LPV_GRIDRES/4 );
 
+			// Insert a flush for all iterations except the last - these dispatches can't overlap!
+			if ( i < LPVNumPropagationSteps - 1 )
+			{
+				RHICmdList.FlushComputeShaderCache();
+			}
+
 			Shader->UnbindBuffers(RHICmdList); 
 		}
 	}
 
 	// Swap buffers
 	mWriteBufferIndex = 1-mWriteBufferIndex; 
+
+	// Copy the AO volume from the LPV to a separate volume texture
+	{
+		SCOPED_DRAW_EVENT(RHICmdList, LpvCopyAOVolume);
+
+		TShaderMapRef<FLpvCopyAOVolumeCS> Shader(View.ShaderMap);
+		RHICmdList.SetComputeShader(Shader->GetComputeShader());
+		FLpvBaseWriteShaderParams ShaderParams;
+		GetShaderParams( ShaderParams );
+		Shader->SetParameters( RHICmdList, ShaderParams );
+		LpvWriteUniformBuffer.SetContents( *LpvWriteUniformBufferParams );
+		DispatchComputeShader( RHICmdList, *Shader, LPV_GRIDRES/4, LPV_GRIDRES/4, LPV_GRIDRES/4 );
+		Shader->UnbindBuffers(RHICmdList); 
+	}
+
+	// End the async compute job
+	AsyncJobFenceID = RHICmdList.EndAsyncComputeJob_DrawThread();
+}
+
+
+void FLightPropagationVolume::InsertGPUWaitForAsyncUpdate(FRHICommandListImmediate& RHICmdList)
+{
+#if USE_ASYNC_COMPUTE_CONTEXT
+	if ( AsyncJobFenceID != -1 )
+	{
+		RHICmdList.GraphicsWaitOnAsyncComputeJob(AsyncJobFenceID);
+		AsyncJobFenceID = -1;
+	}
+#endif
 }
 
 /**
@@ -1157,37 +1308,28 @@ void FLightPropagationVolume::Propagate(FRHICommandListImmediate& RHICmdList, FV
 */
 void FLightPropagationVolume::GetShaderParams( FLpvBaseWriteShaderParams& OutParams ) const
 {
-#if LPV_VOLUME_TEXTURE
 	for(int i = 0; i < 7; ++i)
 	{
 		OutParams.LpvBufferSRVs[i] = LpvVolumeTextures[ 1 - mWriteBufferIndex ][i]->GetRenderTargetItem().ShaderResourceTexture;
 		OutParams.LpvBufferUAVs[i] = LpvVolumeTextures[ mWriteBufferIndex ][i]->GetRenderTargetItem().UAV;
 	}
-#else
-	OutParams.LpvBufferSRV = mLpvBuffers[ 1-mWriteBufferIndex ]->SRV;	// Double-buffered
-	OutParams.LpvBufferUAV = mLpvBuffers[ mWriteBufferIndex ]->UAV;		// Double-buffered
-#endif
 
 	OutParams.VplListBufferSRV = mVplListBuffer->SRV;
 	OutParams.VplListBufferUAV = mVplListBuffer->UAV;
 	OutParams.VplListHeadBufferSRV = mVplListHeadBuffer->SRV;
 	OutParams.VplListHeadBufferUAV = mVplListHeadBuffer->UAV;
 
-#if LPV_GV_VOLUME_TEXTURE
-	for(int i = 0; i < 3; ++i)
+	for ( int i = 0; i < NUM_GV_TEXTURES; i++ )
 	{
 		OutParams.GvBufferSRVs[i] = GvVolumeTextures[i]->GetRenderTargetItem().ShaderResourceTexture;
 		OutParams.GvBufferUAVs[i] = GvVolumeTextures[i]->GetRenderTargetItem().UAV;
 	}
-#else
-	OutParams.GvBufferSRV = GvBuffer->SRV;
-	OutParams.GvBufferUAV = GvBuffer->UAV;
-#endif
 	OutParams.GvListBufferSRV = GvListBuffer->SRV;
 	OutParams.GvListBufferUAV = GvListBuffer->UAV;
 	OutParams.GvListHeadBufferSRV = GvListHeadBuffer->SRV;
 	OutParams.GvListHeadBufferUAV = GvListHeadBuffer->UAV;
-
+	OutParams.AOVolumeTextureUAV = AOVolumeTexture->GetRenderTargetItem().UAV;
+	OutParams.AOVolumeTextureSRV = AOVolumeTexture->GetRenderTargetItem().ShaderResourceTexture;
 	OutParams.UniformBuffer = LpvWriteUniformBuffer;
 }
 
@@ -1198,9 +1340,9 @@ void FLightPropagationVolume::InjectLightDirect(FRHICommandListImmediate& RHICmd
 		return;
 	}
 
-	if(Light.GetLightType() != LightType_Point)
+	// Only point and spot lights are supported (directional lights use the RSM instead)
+	if(Light.GetLightType() != LightType_Point && Light.GetLightType() != LightType_Spot )
 	{
-		//@TODO - Only point lights supported right now
 		return;
 	}
 
@@ -1212,29 +1354,58 @@ void FLightPropagationVolume::InjectLightDirect(FRHICommandListImmediate& RHICmd
 		SCOPED_DRAW_EVENT(RHICmdList, LpvDirectLightInjection);
 
 		FLpvDirectLightInjectParameters InjectUniformBufferParams;
+		// Get light params
+		FVector4	LightPositionAndInvRadius;
+		FVector4	LightColorAndFalloffExponent;
+		FVector		LightDirection;	
+		FVector2D	SpotAngles;
+		float		SourceRadius;
+		float		LightSourceLength;
+		float		LightMinRoughness;
+		Light.GetParameters( LightPositionAndInvRadius, LightColorAndFalloffExponent, LightDirection, SpotAngles, SourceRadius, LightSourceLength, LightMinRoughness );
+
 		InjectUniformBufferParams.LightColor = Light.GetColor() * Light.GetIndirectLightingScale();
 		InjectUniformBufferParams.LightPosition = Light.GetPosition();
 		InjectUniformBufferParams.LightRadius = Light.GetRadius();
-		InjectUniformBufferParams.LightFalloffExponent = 0.0f;
+		InjectUniformBufferParams.LightFalloffExponent = LightColorAndFalloffExponent.W;
+		InjectUniformBufferParams.LightDirection = LightDirection;
+		InjectUniformBufferParams.LightSpotAngles = SpotAngles;
+		InjectUniformBufferParams.LightSourceLength = LightSourceLength;
+		InjectUniformBufferParams.bLightInverseSquaredAttenuation = Light.IsInverseSquared() ? 1.0f : 0.0f;
 
 		FLpvInjectShader_Base* Shader = nullptr;
-		if(Light.GetLightType() == LightType_Point)
-		{
-			// Todo:  hard coded - should be use inverse squared ones only?
-			InjectUniformBufferParams.LightFalloffExponent = 8.0f;
 
-			if ( Light.CastsStaticShadow() || Light.CastsDynamicShadow() )
-			{
-				Shader = (FLpvInjectShader_Base*)*TShaderMapRef<TLpvInject_PointLightCS<INJECT_SHADOWED> >(View.ShaderMap);
-			}
-			else
-			{
-				Shader = (FLpvInjectShader_Base*)*TShaderMapRef<TLpvInject_PointLightCS<0> >(View.ShaderMap);
-			}
+		switch ( Light.GetLightType() )
+		{
+			case LightType_Point:
+				{
+					if ( Light.CastsStaticShadow() || Light.CastsDynamicShadow() )
+					{
+						Shader = (FLpvInjectShader_Base*)*TShaderMapRef<TLpvInject_LightCS<INJECT_SHADOW_CASTING> >(View.ShaderMap);
+					}
+					else
+					{
+						Shader = (FLpvInjectShader_Base*)*TShaderMapRef<TLpvInject_LightCS<0> >(View.ShaderMap);
+					}
+				}
+				break;
+			case LightType_Spot:
+				{
+					if ( Light.CastsStaticShadow() || Light.CastsDynamicShadow() )
+					{
+						Shader = (FLpvInjectShader_Base*)*TShaderMapRef<TLpvInject_LightCS<INJECT_SPOT_ATTENUATION | INJECT_SHADOW_CASTING> >(View.ShaderMap);
+					}
+					else
+					{
+						Shader = (FLpvInjectShader_Base*)*TShaderMapRef<TLpvInject_LightCS<INJECT_SPOT_ATTENUATION | 0> >(View.ShaderMap);
+					}
+				}
+				break;
 		}
 		RHICmdList.SetComputeShader(Shader->GetComputeShader());
 
-		FDirectLightInjectBufferRef InjectUniformBuffer = FDirectLightInjectBufferRef::CreateUniformBufferImmediate(InjectUniformBufferParams, UniformBuffer_SingleDraw);
+  	    FDirectLightInjectBufferRef InjectUniformBuffer = 
+			FDirectLightInjectBufferRef::CreateUniformBufferImmediate(InjectUniformBufferParams, UniformBuffer_SingleFrame );
 
 		mWriteBufferIndex = 1 - mWriteBufferIndex; // Swap buffers with each iteration
 
@@ -1247,7 +1418,6 @@ void FLightPropagationVolume::InjectLightDirect(FRHICommandListImmediate& RHICmd
 		DispatchComputeShader(RHICmdList, Shader, LPV_GRIDRES / 4, LPV_GRIDRES / 4, LPV_GRIDRES / 4 );
 		Shader->UnbindBuffers(RHICmdList); 
 	}
-
 }
 
 
diff --git a/Engine/Source/Runtime/Renderer/Private/LightPropagationVolume.h b/Engine/Source/Runtime/Renderer/Private/LightPropagationVolume.h
--- a/Engine/Source/Runtime/Renderer/Private/LightPropagationVolume.h
+++ b/Engine/Source/Runtime/Renderer/Private/LightPropagationVolume.h
@@ -5,18 +5,18 @@
 //
 // Created:		2013-03-01
 //
-// Author:		mailto:[REDACTED_EMAIL]
+// Author:		Ben Woodhouse - mailto:[REDACTED_EMAIL]
 //
 //				Copyright (C) Microsoft. All rights reserved.
 //-----------------------------------------------------------------------------
-
 #pragma once
 
 #include "UniformBuffer.h"
 
 #define LPV_MULTIPLE_BOUNCES  1
-#define LPV_VOLUME_TEXTURE    1
-#define LPV_GV_VOLUME_TEXTURE 1
+#define LPV_GV_SH_ORDER		  1
+
+static const int32 NUM_GV_TEXTURES = LPV_GV_SH_ORDER + 1;
 
 class FLpvWriteUniformBufferParameters;
 struct FLpvBaseWriteShaderParams;
@@ -24,8 +24,7 @@ typedef TUniformBufferRef<FLpvWriteUniformBufferParameters> FLpvWriteUniformBuff
 typedef TUniformBuffer<FLpvWriteUniformBufferParameters> FLpvWriteUniformBuffer;
 struct FRsmInfo;
 
-#if LPV_VOLUME_TEXTURE
-static const TCHAR* LpvVolumeTextureSRVNames[7] = { 
+static const TCHAR * LpvVolumeTextureSRVNames[7] = { 
 	TEXT("gLpv3DTexture0"),
 	TEXT("gLpv3DTexture1"),
 	TEXT("gLpv3DTexture2"),
@@ -34,7 +33,7 @@ static const TCHAR* LpvVolumeTextureSRVNames[7] = {
 	TEXT("gLpv3DTexture5"),
 	TEXT("gLpv3DTexture6") };
 
-static const TCHAR* LpvVolumeTextureUAVNames[7] = { 
+static const TCHAR * LpvVolumeTextureUAVNames[7] = { 
 	TEXT("gLpv3DTextureRW0"),
 	TEXT("gLpv3DTextureRW1"),
 	TEXT("gLpv3DTextureRW2"),
@@ -42,19 +41,26 @@ static const TCHAR* LpvVolumeTextureUAVNames[7] = {
 	TEXT("gLpv3DTextureRW4"),
 	TEXT("gLpv3DTextureRW5"),
 	TEXT("gLpv3DTextureRW6") };
-#endif
 
-#if LPV_GV_VOLUME_TEXTURE
-static const TCHAR* LpvGvVolumeTextureSRVNames[3] = { 
+static const TCHAR * LpvGvVolumeTextureSRVNames[NUM_GV_TEXTURES] = { 
 	TEXT("gGv3DTexture0"),
+#if ( LPV_GV_SH_ORDER >= 1 )
 	TEXT("gGv3DTexture1"),
-	TEXT("gGv3DTexture2") };
+#endif
+#if ( LPV_GV_SH_ORDER >= 2 )
+	TEXT("gGv3DTexture2") 
+#endif
+};
 
-static const TCHAR* LpvGvVolumeTextureUAVNames[3] = { 
+static const TCHAR * LpvGvVolumeTextureUAVNames[NUM_GV_TEXTURES] = { 
 	TEXT("gGv3DTextureRW0"),
+#if ( LPV_GV_SH_ORDER >= 1 )
 	TEXT("gGv3DTextureRW1"),
-	TEXT("gGv3DTextureRW2") };
 #endif
+#if ( LPV_GV_SH_ORDER >= 2 )
+	TEXT("gGv3DTextureRW2") 
+#endif
+};
 
 //
 // LPV Read constant buffer
@@ -63,6 +69,14 @@ BEGIN_UNIFORM_BUFFER_STRUCT( FLpvReadUniformBufferParameters, )
 	DECLARE_UNIFORM_BUFFER_STRUCT_MEMBER( FIntVector, mLpvGridOffset )
 	DECLARE_UNIFORM_BUFFER_STRUCT_MEMBER( float, LpvScale )
 	DECLARE_UNIFORM_BUFFER_STRUCT_MEMBER( float, OneOverLpvScale )
+	DECLARE_UNIFORM_BUFFER_STRUCT_MEMBER( float, SpecularIntensity )
+	DECLARE_UNIFORM_BUFFER_STRUCT_MEMBER( float, DiffuseIntensity )
+
+	DECLARE_UNIFORM_BUFFER_STRUCT_MEMBER( float, DirectionalOcclusionIntensity )
+	DECLARE_UNIFORM_BUFFER_STRUCT_MEMBER( float, DiffuseOcclusionExponent )
+	DECLARE_UNIFORM_BUFFER_STRUCT_MEMBER( float, SpecularOcclusionExponent )
+	DECLARE_UNIFORM_BUFFER_STRUCT_MEMBER( float, SpecularOcclusionIntensity )
+	DECLARE_UNIFORM_BUFFER_STRUCT_MEMBER( float, DiffuseOcclusionIntensity )
 END_UNIFORM_BUFFER_STRUCT( FLpvReadUniformBufferParameters )
 
 
@@ -79,7 +93,10 @@ BEGIN_UNIFORM_BUFFER_STRUCT( FLpvWriteUniformBufferParameters, )
 	DECLARE_UNIFORM_BUFFER_STRUCT_MEMBER( float, ClearMultiplier )
 	DECLARE_UNIFORM_BUFFER_STRUCT_MEMBER( float, LpvScale )
 	DECLARE_UNIFORM_BUFFER_STRUCT_MEMBER( float, OneOverLpvScale )
-	DECLARE_UNIFORM_BUFFER_STRUCT_MEMBER( float, WarpStrength )
+	DECLARE_UNIFORM_BUFFER_STRUCT_MEMBER( float, DirectionalOcclusionIntensity )
+	DECLARE_UNIFORM_BUFFER_STRUCT_MEMBER( float, DirectionalOcclusionRadius )
+	DECLARE_UNIFORM_BUFFER_STRUCT_MEMBER( float, RsmAreaIntensityMultiplier )
+	DECLARE_UNIFORM_BUFFER_STRUCT_MEMBER( float, RsmPixelToTexcoordMultiplier )
 	DECLARE_UNIFORM_BUFFER_STRUCT_MEMBER( float, SecondaryOcclusionStrength )
 	DECLARE_UNIFORM_BUFFER_STRUCT_MEMBER( float, SecondaryBounceStrength )
 	DECLARE_UNIFORM_BUFFER_STRUCT_MEMBER( float, VplInjectionBias )
@@ -94,30 +111,24 @@ END_UNIFORM_BUFFER_STRUCT( FLpvWriteUniformBufferParameters )
 struct FLpvBaseWriteShaderParams
 {
 	FLpvWriteUniformBufferRef		UniformBuffer;
-#if LPV_VOLUME_TEXTURE
 	FTextureRHIParamRef				LpvBufferSRVs[7];
 	FUnorderedAccessViewRHIParamRef LpvBufferUAVs[7];
-#else
-	FShaderResourceViewRHIParamRef	LpvBufferSRV;
-	FUnorderedAccessViewRHIParamRef LpvBufferUAV;
-#endif
 
 	FShaderResourceViewRHIParamRef	VplListHeadBufferSRV;
 	FUnorderedAccessViewRHIParamRef VplListHeadBufferUAV;
 	FShaderResourceViewRHIParamRef	VplListBufferSRV;
 	FUnorderedAccessViewRHIParamRef VplListBufferUAV;
 
-#if LPV_GV_VOLUME_TEXTURE
 	FTextureRHIParamRef				GvBufferSRVs[3];
 	FUnorderedAccessViewRHIParamRef GvBufferUAVs[3];
-#else
-	FShaderResourceViewRHIParamRef	GvBufferSRV;
-	FUnorderedAccessViewRHIParamRef GvBufferUAV;
-#endif
+
 	FShaderResourceViewRHIParamRef	GvListHeadBufferSRV;
 	FUnorderedAccessViewRHIParamRef GvListHeadBufferUAV;
 	FShaderResourceViewRHIParamRef	GvListBufferSRV;
 	FUnorderedAccessViewRHIParamRef GvListBufferUAV;
+
+	FUnorderedAccessViewRHIParamRef AOVolumeTextureUAV;
+	FTextureRHIParamRef				AOVolumeTextureSRV;
 };
 
 class FLightPropagationVolume // @TODO: this should probably be derived from FRenderResource (with InitDynamicRHI etc)
@@ -145,7 +156,7 @@ public:
 
 	void InjectLightDirect(FRHICommandListImmediate& RHICmdList, const FLightSceneProxy& Light, const FViewInfo& View);
 
-	void Propagate(FRHICommandListImmediate& RHICmdList, FViewInfo& View);
+	void Update(FRHICommandListImmediate& RHICmdList, FViewInfo& View);
 
 	void Visualise(FRHICommandList& RHICmdList, const FViewInfo& View) const;
 
@@ -156,11 +167,7 @@ public:
 
 	FLpvWriteUniformBufferRef GetWriteUniformBuffer() const				{ return (FLpvWriteUniformBufferRef)LpvWriteUniformBuffer; }
 
-#if LPV_VOLUME_TEXTURE
 	FTextureRHIParamRef GetLpvBufferSrv( int i )						{ return LpvVolumeTextures[ 1-mWriteBufferIndex ][i]->GetRenderTargetItem().ShaderResourceTexture; }
-#else
-	FShaderResourceViewRHIParamRef GetLpvBufferSrv()					{ return mLpvBuffers[ 1-mWriteBufferIndex ]->SRV; }
-#endif
 
 	FUnorderedAccessViewRHIParamRef GetVplListBufferUav()				{ return mVplListBuffer->UAV; }
 	FUnorderedAccessViewRHIParamRef GetVplListHeadBufferUav()			{ return mVplListHeadBuffer->UAV; }
@@ -171,17 +178,17 @@ public:
 
 	const FBox&	GetBoundingBox()										{ return BoundingBox; }
 
-//private:	// MM
+	void InsertGPUWaitForAsyncUpdate(FRHICommandListImmediate& RHICmdList);							
+
 	void GetShaderParams( FLpvBaseWriteShaderParams& OutParams) const;
 public:
 
 	void GetShadowInfo( const FProjectedShadowInfo& ProjectedShadowInfo, FRsmInfo& RsmInfoOut );
 
-#if LPV_VOLUME_TEXTURE
+	void ComputeDirectionalOcclusion(FRHICommandListImmediate& RHICmdList, FViewInfo& View);
+	FTextureRHIParamRef GetAOVolumeTextureSRV() { return AOVolumeTexture->GetRenderTargetItem().ShaderResourceTexture; }
+
 	TRefCountPtr<IPooledRenderTarget>	LpvVolumeTextures[2][7];		// double buffered
-#else
-	FRWBufferStructured*				mLpvBuffers[2];					// Double buffered for propagation
-#endif
 	FRWBufferByteAddress*				mVplListHeadBuffer;
 	FRWBufferStructured*				mVplListBuffer;
 
@@ -197,21 +204,18 @@ public:
 	FRWBufferByteAddress*				GvListHeadBuffer;
 	FRWBufferStructured*				GvListBuffer;
 
-#if LPV_VOLUME_TEXTURE || LPV_GV_VOLUME_TEXTURE
 	FShaderResourceParameter			LpvVolumeTextureSampler;
-#endif
 
-#if LPV_GV_VOLUME_TEXTURE
-	TRefCountPtr<IPooledRenderTarget>	GvVolumeTextures[3];		// 9 coeffs + RGB
-#else
-	FRWBuffer*							GvBuffer;			
-#endif
+	TRefCountPtr<IPooledRenderTarget>	GvVolumeTextures[NUM_GV_TEXTURES];		// SH coeffs + RGB
+	TRefCountPtr<IPooledRenderTarget>	AOVolumeTexture;
+
 	float								SecondaryOcclusionStrength;
 	float								SecondaryBounceStrength;
 
 	float								CubeSize;
 	float								Strength;
 	bool								bEnabled;
+	bool								bGeometryVolumeNeeded;
 
 	uint32								mWriteBufferIndex;
 	bool								bNeedsBufferClear;
@@ -223,6 +227,9 @@ public:
 
 	bool								bInitialized;
 
+	// only needed for Async Compute
+	uint32								AsyncJobFenceID;
+
 	friend class FLpvVisualisePS;
 };
 
diff --git a/Engine/Source/Runtime/Renderer/Private/LightPropagationVolumeVisualisation.cpp b/Engine/Source/Runtime/Renderer/Private/LightPropagationVolumeVisualisation.cpp
--- a/Engine/Source/Runtime/Renderer/Private/LightPropagationVolumeVisualisation.cpp
+++ b/Engine/Source/Runtime/Renderer/Private/LightPropagationVolumeVisualisation.cpp
@@ -36,7 +36,7 @@ public:
 	static void ModifyCompilationEnvironment( EShaderPlatform Platform, FShaderCompilerEnvironment& OutEnvironment )				
 	{ 
 		OutEnvironment.SetDefine( TEXT("LPV_MULTIPLE_BOUNCES"), (uint32)LPV_MULTIPLE_BOUNCES );
-		OutEnvironment.SetDefine( TEXT("LPV_GV_VOLUME_TEXTURE"),	(uint32)LPV_GV_VOLUME_TEXTURE );
+		OutEnvironment.SetDefine( TEXT("LPV_GV_SH_ORDER"),		(uint32)LPV_GV_SH_ORDER );
 
 		FGlobalShader::ModifyCompilationEnvironment( Platform, OutEnvironment ); 
 	}
@@ -98,27 +98,16 @@ public:
 	FLpvVisualisePS()	{	}
 	explicit FLpvVisualisePS( const ShaderMetaType::CompiledShaderInitializerType& Initializer ) : FLpvVisualiseBase(Initializer) 
 	{
-#if LPV_VOLUME_TEXTURE
-		for ( int i=0; i<7; i++ )
+		for ( int i = 0; i < 7; i++ )
 		{
 			LpvBufferSRVParameters[i].Bind( Initializer.ParameterMap, LpvVolumeTextureSRVNames[i] );
 		}
-#else
-		InLpvBuffer.Bind( Initializer.ParameterMap, TEXT("gLpvBuffer") );
-#endif
 
-#if LPV_VOLUME_TEXTURE || LPV_GV_VOLUME_TEXTURE
 		LpvVolumeTextureSampler.Bind(Initializer.ParameterMap, TEXT("gLpv3DTextureSampler"));
-#endif
-
-#if LPV_GV_VOLUME_TEXTURE
 		for ( int i=0; i<3; i++ )
 		{
 			GvBufferSRVParameters[i].Bind( Initializer.ParameterMap, LpvGvVolumeTextureSRVNames[i] );
 		}
-#else
-		InGvBuffer.Bind(			Initializer.ParameterMap, TEXT("gGvBuffer") );
-#endif
 	}
 
 	//@todo-rco: Remove this when reenabling for OpenGL
@@ -133,7 +122,6 @@ public:
 		FPixelShaderRHIParamRef ShaderRHI = GetPixelShader();
 		FGlobalShader::SetParameters(RHICmdList, ShaderRHI, View); 
 
-#if LPV_VOLUME_TEXTURE
 		for ( int i=0; i<7; i++ )
 		{
 			FTextureRHIParamRef LpvBufferSrv = LPV->LpvVolumeTextures[ 1-LPV->mWriteBufferIndex ][i]->GetRenderTargetItem().ShaderResourceTexture;
@@ -143,14 +131,7 @@ public:
 			}
 			SetTextureParameter(RHICmdList, ShaderRHI, LpvBufferSRVParameters[i], LpvVolumeTextureSampler, TStaticSamplerState<SF_Bilinear,AM_Clamp,AM_Clamp,AM_Clamp>::GetRHI(), LpvBufferSrv );
 		}
-#else
-		if ( InLpvBuffer.IsBound() ) 
-		{
-			RHICmdList.SetShaderResourceViewParameter( ShaderRHI, InLpvBuffer.GetBaseIndex(), LPV->mLpvBuffers[ LPV->mWriteBufferIndex ]->SRV ); 
-		}
-#endif
 
-#if LPV_GV_VOLUME_TEXTURE
 		for ( int i=0; i<3; i++ )
 		{
 			FTextureRHIParamRef GvBufferSrv = LPV->GvVolumeTextures[i]->GetRenderTargetItem().ShaderResourceTexture;
@@ -161,12 +142,6 @@ public:
 			SetTextureParameter(RHICmdList, ShaderRHI, GvBufferSRVParameters[i], LpvVolumeTextureSampler, TStaticSamplerState<SF_Bilinear,AM_Clamp,AM_Clamp,AM_Clamp>::GetRHI(), GvBufferSrv );
 		}
 
-#else
-		if ( InGvBuffer.IsBound() ) 
-		{
-			RHICmdList.SetShaderResourceViewParameter( ShaderRHI, InGvBuffer.GetBaseIndex(), LPV->GvBuffer->SRV ); 
-		}
-#endif
 	}
 
 
@@ -175,50 +150,29 @@ public:
 	{
 		bool bShaderHasOutdatedParameters = FLpvVisualiseBase::Serialize( Ar );
 
-#if LPV_VOLUME_TEXTURE
 		for ( int i=0; i<7; i++ )
 		{
 			Ar << LpvBufferSRVParameters[i];
 		}
-#else
-		Ar << InLpvBuffer;
-#endif
 
-#if LPV_VOLUME_TEXTURE | LPV_GV_VOLUME_TEXTURE
 		Ar << LpvVolumeTextureSampler;
-#endif
 
-#if LPV_GV_VOLUME_TEXTURE
 		for ( int i=0; i<3; i++ )
 		{
 			Ar << GvBufferSRVParameters[i];
 		}
-#else
-		Ar << InGvBuffer;
-#endif
 		return bShaderHasOutdatedParameters;
 	}
-#if LPV_VOLUME_TEXTURE
 	FShaderResourceParameter LpvBufferSRVParameters[7];
-#else
-	FShaderResourceParameter InLpvBuffer;
-#endif 
 
-#if LPV_VOLUME_TEXTURE | LPV_GV_VOLUME_TEXTURE
 	FShaderResourceParameter LpvVolumeTextureSampler;
-#endif
 
-#if LPV_GV_VOLUME_TEXTURE
 	FShaderResourceParameter GvBufferSRVParameters[3];
-#else
-	FShaderResourceParameter InGvBuffer;
-#endif
 
 	void UnbindBuffers(FRHICommandList& RHICmdList)
 	{
 		// TODO: Is this necessary here?
 		FPixelShaderRHIParamRef ShaderRHI = GetPixelShader();
-#if LPV_VOLUME_TEXTURE
 		for ( int i=0;i<7; i++ )
 		{
 			if ( LpvBufferSRVParameters[i].IsBound() )
@@ -226,11 +180,7 @@ public:
 				RHICmdList.SetShaderTexture(ShaderRHI, LpvBufferSRVParameters[i].GetBaseIndex(), FTextureRHIParamRef());
 			}
 		}
-#else
-		if ( InLpvBuffer.IsBound() ) RHICmdList.SetShaderResourceViewParameter( ShaderRHI, InLpvBuffer.GetBaseIndex(), FShaderResourceViewRHIParamRef() );
-#endif
 
-#if LPV_GV_VOLUME_TEXTURE
 		for ( int i=0;i<3; i++ )
 		{
 			if ( GvBufferSRVParameters[i].IsBound() )
@@ -238,9 +188,6 @@ public:
 				RHICmdList.SetShaderTexture(ShaderRHI, GvBufferSRVParameters[i].GetBaseIndex(), FTextureRHIParamRef());
 			}
 		}
-#else
-		if ( InGvBuffer.IsBound() ) RHICmdList.SetShaderResourceViewParameter( ShaderRHI, InGvBuffer.GetBaseIndex(), FShaderResourceViewRHIParamRef() );
-#endif
 	}
 
 };
diff --git a/Engine/Source/Runtime/Renderer/Private/LightRendering.cpp b/Engine/Source/Runtime/Renderer/Private/LightRendering.cpp
--- a/Engine/Source/Runtime/Renderer/Private/LightRendering.cpp
+++ b/Engine/Source/Runtime/Renderer/Private/LightRendering.cpp
@@ -473,144 +473,170 @@ void FDeferredShadingSceneRenderer::RenderLights(FRHICommandListImmediate& RHICm
 			}
 		}
 
-		bool bDirectLighting = ViewFamily.EngineShowFlags.DirectLighting;
+		EShaderPlatform ShaderPlatform = GShaderPlatformForFeatureLevel[FeatureLevel]; 
 
-		// Draw shadowed and light function lights
-		for (int32 LightIndex = AttenuationLightStart; LightIndex < SortedLights.Num(); LightIndex++)
+//INTEGRATETODO4_5: Equivalent Epic code changed from 'if ( IsFeatureLevelSupported(GRHIShaderPlatform, ERHIFeatureLevel::SM5) )' to 'if( FeatureLevel >= ERHIFeatureLevel::SM5 )'
+		if ( IsFeatureLevelSupported(ShaderPlatform, ERHIFeatureLevel::SM5) )
 		{
-			const FSortedLightSceneInfo& SortedLightInfo = SortedLights[LightIndex];
-			const FLightSceneInfoCompact& LightSceneInfoCompact = SortedLightInfo.SceneInfo;
-			const FLightSceneInfo& LightSceneInfo = *LightSceneInfoCompact.LightSceneInfo;
-			bool bDrawShadows = SortedLightInfo.SortKey.Fields.bShadowed;
-			bool bDrawLightFunction = SortedLightInfo.SortKey.Fields.bLightFunction;
-			bool bInjectedTranslucentVolume = false;
-			bool bUsedLightAttenuation = false;
-			FScopeCycleCounter Context(LightSceneInfo.Proxy->GetStatId());
-
-			FString LightNameWithLevel;
-			GetLightNameForDrawEvent(LightSceneInfo.Proxy, LightNameWithLevel);
-			SCOPED_DRAW_EVENTF(RHICmdList, EventLightPass, *LightNameWithLevel);
-
-			// Do not resolve to scene color texture, this is done lazily
-			GSceneRenderTargets.FinishRenderingSceneColor(RHICmdList, false);
-
-			if (bDrawShadows)
+			SCOPED_DRAW_EVENT(RHICmdList, IndirectLighting);
+			bool bRenderedRSM = false;
+			// Render Reflective shadow maps
+			// Draw shadowed and light function lights
+			for (int32 LightIndex = AttenuationLightStart; LightIndex < SortedLights.Num(); LightIndex++)
 			{
-				INC_DWORD_STAT(STAT_NumShadowedLights);
-
-				for (int32 ViewIndex = 0; ViewIndex < Views.Num(); ViewIndex++)
+				const FSortedLightSceneInfo& SortedLightInfo = SortedLights[LightIndex];
+				const FLightSceneInfoCompact& LightSceneInfoCompact = SortedLightInfo.SceneInfo;
+				const FLightSceneInfo& LightSceneInfo = *LightSceneInfoCompact.LightSceneInfo;
+				// Render any reflective shadow maps (if necessary)
+				if ( LightSceneInfo.Proxy && LightSceneInfo.Proxy->NeedsLPVInjection() )
 				{
-					const FViewInfo& View = Views[ViewIndex];
-					View.HeightfieldLightingViewInfo.ClearShadowing(View, RHICmdList, LightSceneInfo);
+					if ( LightSceneInfo.Proxy->HasReflectiveShadowMap() )
+					{
+						INC_DWORD_STAT(STAT_NumReflectiveShadowMapLights);
+						RenderReflectiveShadowMaps( RHICmdList, &LightSceneInfo );
+						bRenderedRSM = true;
+					}
 				}
-
-				// All shadows render with min blending
-				bool bClearToWhite = true;
-				GSceneRenderTargets.BeginRenderingLightAttenuation(RHICmdList, bClearToWhite);
-
-				bool bRenderedTranslucentObjectShadows = RenderTranslucentProjectedShadows(RHICmdList, &LightSceneInfo );
-				// Render non-modulated projected shadows to the attenuation buffer.
-				RenderProjectedShadows(RHICmdList, &LightSceneInfo, bRenderedTranslucentObjectShadows, bInjectedTranslucentVolume );
-				
-				bUsedLightAttenuation = true;
 			}
 
-			// Render any reflective shadow maps (if necessary)
-			if(ViewFamily.EngineShowFlags.GlobalIllumination && LightSceneInfo.Proxy->NeedsLPVInjection())
+			// LPV Direct Light Injection
+			if ( bRenderedRSM )
 			{
-				if ( LightSceneInfo.Proxy->HasReflectiveShadowMap() )
+				for (int32 LightIndex = 0; LightIndex < SortedLights.Num(); LightIndex++)
 				{
-					INC_DWORD_STAT(STAT_NumReflectiveShadowMapLights);
-					RenderReflectiveShadowMaps(RHICmdList, &LightSceneInfo );
+					const FSortedLightSceneInfo& SortedLightInfo = SortedLights[LightIndex];
+					const FLightSceneInfoCompact& LightSceneInfoCompact = SortedLightInfo.SceneInfo;
+					const FLightSceneInfo* const LightSceneInfo = LightSceneInfoCompact.LightSceneInfo;
+
+					// Render any reflective shadow maps (if necessary)
+					if ( LightSceneInfo && LightSceneInfo->Proxy && LightSceneInfo->Proxy->NeedsLPVInjection() )
+					{
+						if ( !LightSceneInfo->Proxy->HasReflectiveShadowMap() )
+						{
+							// Inject the light directly into all relevant LPVs
+							for (int32 ViewIndex = 0; ViewIndex < Views.Num(); ViewIndex++)
+							{
+								FViewInfo& View = Views[ViewIndex];
+
+								if (LightSceneInfo->ShouldRenderLight(View))
+								{
+									FSceneViewState* ViewState = (FSceneViewState*)View.State;
+									FLightPropagationVolume* Lpv = ViewState->GetLightPropagationVolume();
+									if ( Lpv && LightSceneInfo->Proxy )
+									{
+										Lpv->InjectLightDirect( RHICmdList, *LightSceneInfo->Proxy, View );
+									}
+								}
+							}					
+						}
+					}
 				}
 			}
 
-			for (int32 ViewIndex = 0; ViewIndex < Views.Num(); ViewIndex++)
-			{
-				const FViewInfo& View = Views[ViewIndex];
-				View.HeightfieldLightingViewInfo.ComputeLighting(View, RHICmdList, LightSceneInfo);
-			}
-			
-			// Render light function to the attenuation buffer.
-			if (bDirectLighting)
-			{
-				const bool bLightFunctionRendered = RenderLightFunction(RHICmdList, &LightSceneInfo, bDrawShadows);
-				bUsedLightAttenuation |= bLightFunctionRendered;
+			// Kickoff the LPV update (asynchronously if possible)
+			UpdateLPVs(RHICmdList);
+		}
 
-				if (ViewFamily.EngineShowFlags.PreviewShadowsIndicator
-					&& !LightSceneInfo.IsPrecomputedLightingValid() 
-					&& LightSceneInfo.Proxy->HasStaticShadowing())
+		
+		{
+			SCOPED_DRAW_EVENT(RHICmdList, ShadowedLights);
+
+			bool bDirectLighting = ViewFamily.EngineShowFlags.DirectLighting;
+
+			// Draw shadowed and light function lights
+			for (int32 LightIndex = AttenuationLightStart; LightIndex < SortedLights.Num(); LightIndex++)
+			{
+				const FSortedLightSceneInfo& SortedLightInfo = SortedLights[LightIndex];
+				const FLightSceneInfoCompact& LightSceneInfoCompact = SortedLightInfo.SceneInfo;
+				const FLightSceneInfo& LightSceneInfo = *LightSceneInfoCompact.LightSceneInfo;
+				bool bDrawShadows = SortedLightInfo.SortKey.Fields.bShadowed;
+				bool bDrawLightFunction = SortedLightInfo.SortKey.Fields.bLightFunction;
+				bool bInjectedTranslucentVolume = false;
+				bool bUsedLightAttenuation = false;
+				FScopeCycleCounter Context(LightSceneInfo.Proxy->GetStatId());
+
+				FString LightNameWithLevel;
+				GetLightNameForDrawEvent(LightSceneInfo.Proxy, LightNameWithLevel);
+				SCOPED_DRAW_EVENTF(RHICmdList, EventLightPass, *LightNameWithLevel);
+
+				// Do not resolve to scene color texture, this is done lazily
+				GSceneRenderTargets.FinishRenderingSceneColor(RHICmdList, false);
+
+				if (bDrawShadows)
 				{
-					RenderPreviewShadowsIndicator(RHICmdList, &LightSceneInfo, bUsedLightAttenuation);
+					INC_DWORD_STAT(STAT_NumShadowedLights);
+
+					for (int32 ViewIndex = 0; ViewIndex < Views.Num(); ViewIndex++)
+					{
+						const FViewInfo& View = Views[ViewIndex];
+						View.HeightfieldLightingViewInfo.ClearShadowing(View, RHICmdList, LightSceneInfo);
+					}
+
+					// All shadows render with min blending
+					bool bClearToWhite = true;
+					GSceneRenderTargets.BeginRenderingLightAttenuation(RHICmdList, bClearToWhite);
+
+					bool bRenderedTranslucentObjectShadows = RenderTranslucentProjectedShadows(RHICmdList, &LightSceneInfo );
+					// Render non-modulated projected shadows to the attenuation buffer.
+					RenderProjectedShadows(RHICmdList, &LightSceneInfo, bRenderedTranslucentObjectShadows, bInjectedTranslucentVolume );
+				
+					bUsedLightAttenuation = true;
 				}
 
-				if (!bDrawShadows)
+				for (int32 ViewIndex = 0; ViewIndex < Views.Num(); ViewIndex++)
 				{
-					INC_DWORD_STAT(STAT_NumLightFunctionOnlyLights);
+					const FViewInfo& View = Views[ViewIndex];
+					View.HeightfieldLightingViewInfo.ComputeLighting(View, RHICmdList, LightSceneInfo);
 				}
-			}
 			
-			if( bUsedLightAttenuation )
-			{
-				// Resolve light attenuation buffer
-				GSceneRenderTargets.FinishRenderingLightAttenuation(RHICmdList);
-			}
-			
-			if(bDirectLighting && !bInjectedTranslucentVolume)
-			{
-				SCOPED_DRAW_EVENT(RHICmdList, InjectTranslucentVolume);
-				// Accumulate this light's unshadowed contribution to the translucency lighting volume
-				InjectTranslucentVolumeLighting(RHICmdList, LightSceneInfo, NULL);
-			}
-
-			GSceneRenderTargets.SetLightAttenuationMode(bUsedLightAttenuation);
-			GSceneRenderTargets.BeginRenderingSceneColor(RHICmdList, ESimpleRenderTargetMode::EExistingColorAndDepth, FExclusiveDepthStencil::DepthRead_StencilWrite);
-
-			// Render the light to the scene color buffer, conditionally using the attenuation buffer or a 1x1 white texture as input 
-			if(bDirectLighting)
-			{
-				RenderLight(RHICmdList, &LightSceneInfo, false, true);
-			}
-		}
+				// Render light function to the attenuation buffer.
+				if (bDirectLighting)
+				{
+					const bool bLightFunctionRendered = RenderLightFunction(RHICmdList, &LightSceneInfo, bDrawShadows);
+					bUsedLightAttenuation |= bLightFunctionRendered;
 
-		// Do not resolve to scene color texture, this is done lazily
-		GSceneRenderTargets.FinishRenderingSceneColor(RHICmdList, false);
+					if (ViewFamily.EngineShowFlags.PreviewShadowsIndicator
+						&& !LightSceneInfo.IsPrecomputedLightingValid() 
+						&& LightSceneInfo.Proxy->HasStaticShadowing())
+					{
+						RenderPreviewShadowsIndicator(RHICmdList, &LightSceneInfo, bUsedLightAttenuation);
+					}
 
-		// Restore the default mode
-		GSceneRenderTargets.SetLightAttenuationMode(true);
+					if (!bDrawShadows)
+					{
+						INC_DWORD_STAT(STAT_NumLightFunctionOnlyLights);
+					}
+				}
+			
+				if( bUsedLightAttenuation )
+				{
+					// Resolve light attenuation buffer
+					GSceneRenderTargets.FinishRenderingLightAttenuation(RHICmdList);
+				}
+			
+				if(bDirectLighting && !bInjectedTranslucentVolume)
+				{
+					SCOPED_DRAW_EVENT(RHICmdList, InjectTranslucentVolume);
+					// Accumulate this light's unshadowed contribution to the translucency lighting volume
+					InjectTranslucentVolumeLighting(RHICmdList, LightSceneInfo, NULL);
+				}
 
-		// LPV Direct Light Injection
-		for (int32 LightIndex = 0; LightIndex < SortedLights.Num(); LightIndex++)
-		{
-			const FSortedLightSceneInfo& SortedLightInfo = SortedLights[LightIndex];
-			const FLightSceneInfoCompact& LightSceneInfoCompact = SortedLightInfo.SceneInfo;
-			const FLightSceneInfo* const LightSceneInfo = LightSceneInfoCompact.LightSceneInfo;
+				GSceneRenderTargets.SetLightAttenuationMode(bUsedLightAttenuation);
+				GSceneRenderTargets.BeginRenderingSceneColor(RHICmdList, ESimpleRenderTargetMode::EExistingColorAndDepth, FExclusiveDepthStencil::DepthRead_StencilWrite);
 
-			// Render any reflective shadow maps (if necessary)
-			if ( LightSceneInfo && LightSceneInfo->Proxy && LightSceneInfo->Proxy->NeedsLPVInjection() )
-			{
-				if ( !LightSceneInfo->Proxy->HasReflectiveShadowMap() )
+				// Render the light to the scene color buffer, conditionally using the attenuation buffer or a 1x1 white texture as input 
+				if(bDirectLighting)
 				{
-					// Inject the light directly into all relevant LPVs
-					for (int32 ViewIndex = 0; ViewIndex < Views.Num(); ViewIndex++)
-					{
-						FViewInfo& View = Views[ViewIndex];
-						FSceneViewState* ViewState = (FSceneViewState*)View.State;
-
-						if (ViewState)
-						{
-							if (LightSceneInfo->ShouldRenderLight(View))
-							{
-								FLightPropagationVolume* Lpv = ViewState->GetLightPropagationVolume();
-								if (Lpv && LightSceneInfo->Proxy)
-								{
-									Lpv->InjectLightDirect(RHICmdList, *LightSceneInfo->Proxy, View);
-								}
-							}
-						}
-					}					
+					RenderLight(RHICmdList, &LightSceneInfo, false, true);
 				}
 			}
+
+			// Do not resolve to scene color texture, this is done lazily
+			GSceneRenderTargets.FinishRenderingSceneColor(RHICmdList, false);
+
+			// Restore the default mode
+			GSceneRenderTargets.SetLightAttenuationMode(true);
+
 		}
 	}
 }
diff --git a/Engine/Source/Runtime/Renderer/Private/PostProcess/SceneRenderTargets.cpp b/Engine/Source/Runtime/Renderer/Private/PostProcess/SceneRenderTargets.cpp
--- a/Engine/Source/Runtime/Renderer/Private/PostProcess/SceneRenderTargets.cpp
+++ b/Engine/Source/Runtime/Renderer/Private/PostProcess/SceneRenderTargets.cpp
@@ -12,21 +12,23 @@
 #include "SceneUtils.h"
 #include "HdrCustomResolveShaders.h"
 
-
-// for LightPropagationVolume feature, could be exposed
-const int ReflectiveShadowMapResolution = 256;
-
 IMPLEMENT_UNIFORM_BUFFER_STRUCT(FGBufferResourceStruct,TEXT("GBuffers"));
 
 static TAutoConsoleVariable<int32> CVarBasePassOutputsVelocityDebug(
 	TEXT("r.BasePassOutputsVelocityDebug"),
 	0,
-	TEXT("Debug settings for Base Pass outputting velocity.\n") \
-	TEXT("0 - Regular rendering\n") \
-	TEXT("1 - Skip setting GBufferVelocity RT\n") \
+	TEXT("Debug settings for Base Pass outputting velocity.\n")
+	TEXT("0 - Regular rendering\n")
+	TEXT("1 - Skip setting GBufferVelocity RT\n")
 	TEXT("2 - Set Color Mask 0 for GBufferVelocity RT\n"),
 	ECVF_RenderThreadSafe);
 
+static TAutoConsoleVariable<int32> CVarRSMResolution(
+	TEXT("r.LPV.RSMResolution"),
+	360,
+	TEXT("Reflective Shadow Map resolution (used for LPV) - higher values result in less aliasing artifacts, at the cost of performance"),
+	ECVF_Scalability | ECVF_RenderThreadSafe);
+
 static int32 GBasePassOutputsVelocityDebug = 0;
 
 /*-----------------------------------------------------------------------------
@@ -235,6 +237,8 @@ void FSceneRenderTargets::Allocate(const FSceneViewFamily& ViewFamily)
 
 	int32 MaxShadowResolution = GetCachedScalabilityCVars().MaxShadowResolution;
 
+	int32 RSMResolution = FMath::Clamp(CVarRSMResolution.GetValueOnRenderThread(), 1, 2048);
+
 	if (ViewFamily.Scene->ShouldUseDeferredRenderer() == false)
 	{
 		// ensure there is always enough space for forward renderer's tiled shadow maps
@@ -269,6 +273,7 @@ void FSceneRenderTargets::Allocate(const FSceneViewFamily& ViewFamily)
 		(bAllowStaticLighting != bNewAllowStaticLighting) ||
 		(bUseDownsizedOcclusionQueries != bDownsampledOcclusionQueries) ||
 		(CurrentMaxShadowResolution != MaxShadowResolution) ||
+ 		(CurrentRSMResolution != RSMResolution) ||
 		(CurrentTranslucencyLightingVolumeDim != TranslucencyLightingVolumeDim) ||
 		(CurrentMobile32bpp != Mobile32bpp) ||
 		(CurrentMobileMSAA != MobileMSAA) ||
@@ -280,6 +285,7 @@ void FSceneRenderTargets::Allocate(const FSceneViewFamily& ViewFamily)
 		bAllowStaticLighting = bNewAllowStaticLighting;
 		bUseDownsizedOcclusionQueries = bDownsampledOcclusionQueries;
 		CurrentMaxShadowResolution = MaxShadowResolution;
+		CurrentRSMResolution = RSMResolution;
 		CurrentTranslucencyLightingVolumeDim = TranslucencyLightingVolumeDim;
 		CurrentMobile32bpp = Mobile32bpp;
 		CurrentMobileMSAA = MobileMSAA;
@@ -1508,6 +1514,18 @@ void FSceneRenderTargets::AllocateDeferredShadingPathRenderTargets()
 		}
 	}
 
+	// LPV : Dynamic directional occlusion for diffuse and specular
+	if(UseLightPropagationVolumeRT(CurrentFeatureLevel))
+	{
+		FPooledRenderTargetDesc Desc(FPooledRenderTargetDesc::Create2DDesc(BufferSize, PF_R8G8, TexCreate_None, TexCreate_RenderTargetable, false));
+		GRenderTargetPool.FindFreeElement(Desc, DirectionalOcclusion, TEXT("DirectionalOcclusion"));
+	}
+
+	{
+		FPooledRenderTargetDesc Desc(FPooledRenderTargetDesc::Create2DDesc(BufferSize, GetSceneColorFormat(), TexCreate_None, TexCreate_RenderTargetable | TexCreate_UAV, false));
+		GRenderTargetPool.FindFreeElement(Desc, LightAccumulation, TEXT("LightAccumulation"));
+	}
+
 	AllocateReflectionTargets();
 
 	if (CurrentFeatureLevel >= ERHIFeatureLevel::SM5)
@@ -1515,18 +1533,21 @@ void FSceneRenderTargets::AllocateDeferredShadingPathRenderTargets()
 		// Create the reflective shadow map textures for LightPropagationVolume feature
 		if(bCurrentLightPropagationVolume)
 		{
+			int32 Res = GetReflectiveShadowMapResolution();
+			FIntPoint Extent = FIntPoint(Res, Res);
+
 			{
-				FPooledRenderTargetDesc Desc(FPooledRenderTargetDesc::Create2DDesc(GetReflectiveShadowMapTextureResolution(), PF_R8G8B8A8, TexCreate_None, TexCreate_RenderTargetable, false));
+				FPooledRenderTargetDesc Desc(FPooledRenderTargetDesc::Create2DDesc(Extent, PF_R8G8B8A8, TexCreate_None, TexCreate_RenderTargetable, false));
 				GRenderTargetPool.FindFreeElement(Desc, ReflectiveShadowMapNormal, TEXT("RSMNormal"));
 			}
 
 			{
-				FPooledRenderTargetDesc Desc(FPooledRenderTargetDesc::Create2DDesc(GetReflectiveShadowMapTextureResolution(), PF_FloatR11G11B10, TexCreate_None, TexCreate_RenderTargetable, false));
+				FPooledRenderTargetDesc Desc(FPooledRenderTargetDesc::Create2DDesc(Extent, PF_FloatR11G11B10, TexCreate_None, TexCreate_RenderTargetable, false));
 				GRenderTargetPool.FindFreeElement(Desc, ReflectiveShadowMapDiffuse, TEXT("RSMDiffuse"));
 			}
 
 			{
-				FPooledRenderTargetDesc Desc(FPooledRenderTargetDesc::Create2DDesc(GetReflectiveShadowMapTextureResolution(), PF_DepthStencil, TexCreate_None, TexCreate_DepthStencilTargetable , false));
+				FPooledRenderTargetDesc Desc(FPooledRenderTargetDesc::Create2DDesc(Extent, PF_DepthStencil, TexCreate_None, TexCreate_DepthStencilTargetable , false));
 				GRenderTargetPool.FindFreeElement(Desc, ReflectiveShadowMapDepth, TEXT("RSMDepth"));
 			}
 		}
@@ -1614,6 +1635,8 @@ void FSceneRenderTargets::ReleaseAllTargets()
 	DBufferC.SafeRelease();
 	ScreenSpaceAO.SafeRelease();
 	LightAttenuation.SafeRelease();
+	LightAccumulation.SafeRelease();
+	DirectionalOcclusion.SafeRelease();
 	CustomDepth.SafeRelease();
 	ReflectiveShadowMapNormal.SafeRelease();
 	ReflectiveShadowMapDiffuse.SafeRelease();
@@ -1625,6 +1648,8 @@ void FSceneRenderTargets::ReleaseAllTargets()
 	}
 
 	ShadowDepthZ.SafeRelease();
+	OptionalShadowDepthColor.SafeRelease();
+
 	PreShadowCacheDepthZ.SafeRelease();
 	
 	for(int32 Index = 0; Index < NumCubeShadowDepthSurfaces; ++Index)
@@ -1674,9 +1699,11 @@ FIntPoint FSceneRenderTargets::GetShadowDepthTextureResolution() const
 	return ShadowBufferResolution;
 }
 
-FIntPoint FSceneRenderTargets::GetReflectiveShadowMapTextureResolution() const
+int32 FSceneRenderTargets::GetReflectiveShadowMapResolution() const
 {
-	return FIntPoint( ReflectiveShadowMapResolution, ReflectiveShadowMapResolution );
+	check(IsInRenderingThread());
+
+	return CurrentRSMResolution;
 }
 
 FIntPoint FSceneRenderTargets::GetPreShadowCacheTextureResolution() const
@@ -1834,6 +1861,8 @@ void FSceneTextureShaderParameters::Bind(const FShaderParameterMap& ParameterMap
 	SceneColorSurfaceParameter.Bind(ParameterMap,TEXT("SceneColorSurface"));
 	// only used if Material has an expression that requires SceneColorTextureMSAA
 	SceneDepthSurfaceParameter.Bind(ParameterMap,TEXT("SceneDepthSurface"));
+	DirectionalOcclusionSampler.Bind(ParameterMap, TEXT("DirectionalOcclusionSampler"));
+	DirectionalOcclusionTexture.Bind(ParameterMap, TEXT("DirectionalOcclusionTexture"));
 }
 
 template< typename ShaderRHIParamRef >
@@ -1965,6 +1994,40 @@ void FSceneTextureShaderParameters::Set(
 		SetTextureParameter(RHICmdList, ShaderRHI, SceneColorSurfaceParameter, GBlackTexture->TextureRHI);
 		SetTextureParameter(RHICmdList, ShaderRHI, SceneDepthSurfaceParameter, GBlackTexture->TextureRHI);
 	}
+
+	if( DirectionalOcclusionSampler.IsBound() )
+	{
+		bool bDirectionalOcclusion = false;
+		FSceneViewState* ViewState = (FSceneViewState*)View.State;
+		if ( ViewState != nullptr && ViewState->GetLightPropagationVolume() != nullptr && View.FinalPostProcessSettings.LPVIntensity > 0.0f )
+		{
+			if ( View.FinalPostProcessSettings.LPVDirectionalOcclusionIntensity > 0.0001f )
+			{
+				bDirectionalOcclusion = true;
+			}
+		}
+
+		FTextureRHIParamRef DirectionalOcclusion = nullptr;
+		if( bDirectionalOcclusion )
+		{
+			DirectionalOcclusion = GSceneRenderTargets.GetDirectionalOcclusionTexture();
+		}
+		else
+		{
+			DirectionalOcclusion = GWhiteTexture->TextureRHI;
+		}
+
+		FSamplerStateRHIRef Filter;
+		Filter = TStaticSamplerState<SF_Bilinear,AM_Clamp,AM_Clamp,AM_Clamp>::GetRHI();
+		SetTextureParameter(
+			RHICmdList,
+			ShaderRHI,
+			DirectionalOcclusionTexture,
+			DirectionalOcclusionSampler,
+			Filter,
+			DirectionalOcclusion
+			);
+	}
 }
 
 #define IMPLEMENT_SCENE_TEXTURE_PARAM_SET( ShaderRHIParamRef ) \
@@ -1994,6 +2057,8 @@ FArchive& operator<<(FArchive& Ar,FSceneTextureShaderParameters& Parameters)
 	Ar << Parameters.SceneDepthTextureParameterSampler;
 	Ar << Parameters.SceneDepthSurfaceParameter;
 	Ar << Parameters.SceneDepthTextureNonMS;
+	Ar << Parameters.DirectionalOcclusionSampler;
+	Ar << Parameters.DirectionalOcclusionTexture;
 	return Ar;
 }
 
@@ -2133,4 +2198,4 @@ FArchive& operator<<(FArchive& Ar,FDeferredPixelShaderParameters& Parameters)
 	Ar << Parameters.CustomDepthTextureSampler;
 
 	return Ar;
-}
\ No newline at end of file
+}
diff --git a/Engine/Source/Runtime/Renderer/Private/PostProcess/SceneRenderTargets.h b/Engine/Source/Runtime/Renderer/Private/PostProcess/SceneRenderTargets.h
--- a/Engine/Source/Runtime/Renderer/Private/PostProcess/SceneRenderTargets.h
+++ b/Engine/Source/Runtime/Renderer/Private/PostProcess/SceneRenderTargets.h
@@ -130,6 +130,7 @@ protected:
 		CurrentSceneColorFormat(0),
 		bAllowStaticLighting(true),
 		CurrentMaxShadowResolution(0),
+		CurrentRSMResolution(0),
 		CurrentTranslucencyLightingVolumeDim(64),
 		CurrentMobile32bpp(0),
 		bCurrentLightPropagationVolume(false),
@@ -355,6 +356,11 @@ public:
 		return (const FTexture2DRHIRef&)AuxiliarySceneDepthZ->GetRenderTargetItem().TargetableTexture; 
 	}
 
+	const FTexture2DRHIRef& GetDirectionalOcclusionTexture() const 
+	{	
+		return (const FTexture2DRHIRef&)DirectionalOcclusion->GetRenderTargetItem().TargetableTexture; 
+	}
+
 	IPooledRenderTarget* GetGBufferVelocityRT();
 
 	int32 GetGBufferEIndex() const
@@ -412,7 +418,7 @@ public:
 	int32 GetTranslucentShadowDownsampleFactor() const { return 2; }
 
 	/** Returns the size of the RSM buffer, taking into account platform limitations and game specific resolution limits. */
-	FIntPoint GetReflectiveShadowMapTextureResolution() const;
+	int32 GetReflectiveShadowMapResolution() const;
 
 	int32 GetNumGBufferTargets() const;
 
@@ -468,7 +474,11 @@ private: // Get...() methods instead of direct access
 	// also used as LDR scene color
 	TRefCountPtr<IPooledRenderTarget> LightAttenuation;
 public:
+	// Reflection Environment: Bringing back light accumulation buffer to apply indirect reflections
+	TRefCountPtr<IPooledRenderTarget> LightAccumulation;
 
+	// Reflection Environment: Bringing back light accumulation buffer to apply indirect reflections
+	TRefCountPtr<IPooledRenderTarget> DirectionalOcclusion;
 	//
 	TRefCountPtr<IPooledRenderTarget> SceneDepthZ;
 	// Mobile without frame buffer fetch (to get depth from alpha).
@@ -626,6 +636,8 @@ private:
 	bool bAllowStaticLighting;
 	/** To detect a change of the CVar r.Shadow.MaxResolution */
 	int32 CurrentMaxShadowResolution;
+	/** To detect a change of the CVar r.Shadow.RsmResolution*/
+	int32 CurrentRSMResolution;
 	/** To detect a change of the CVar r.TranslucencyLightingVolumeDim */
 	int32 CurrentTranslucencyLightingVolumeDim;
 	/** To detect a change of the CVar r.MobileHDR / r.MobileHDR32bpp */
diff --git a/Engine/Source/Runtime/Renderer/Private/ReflectionEnvironment.cpp b/Engine/Source/Runtime/Renderer/Private/ReflectionEnvironment.cpp
--- a/Engine/Source/Runtime/Renderer/Private/ReflectionEnvironment.cpp
+++ b/Engine/Source/Runtime/Renderer/Private/ReflectionEnvironment.cpp
@@ -24,6 +24,8 @@ const int32 GReflectionEnvironmentTileSizeX = 16;
 const int32 GReflectionEnvironmentTileSizeY = 16;
 extern ENGINE_API int32 GReflectionCaptureSize;
 
+extern TAutoConsoleVariable<int32> CVarLPVMixing;
+
 static TAutoConsoleVariable<int32> CVarDiffuseFromCaptures(
 	TEXT("r.DiffuseFromCaptures"),
 	0,
@@ -72,10 +74,13 @@ static int GetReflectionEnvironmentCVar()
 
 bool IsReflectionEnvironmentAvailable(ERHIFeatureLevel::Type InFeatureLevel)
 {
-	static const auto AllowStaticLightingVar = IConsoleManager::Get().FindTConsoleVariableDataInt(TEXT("r.AllowStaticLighting"));
-	const bool bAllowStaticLighting = (!AllowStaticLightingVar || AllowStaticLightingVar->GetValueOnAnyThread() != 0);
+	return (InFeatureLevel >= ERHIFeatureLevel::SM4) && (GetReflectionEnvironmentCVar() != 0);
+}
 
-	return (InFeatureLevel >= ERHIFeatureLevel::SM4) && (GetReflectionEnvironmentCVar() != 0) && bAllowStaticLighting;
+bool IsReflectionCaptureAvailable()
+{
+	static IConsoleVariable* AllowStaticLightingVar = IConsoleManager::Get().FindConsoleVariable(TEXT("r.AllowStaticLighting"));
+	return (!AllowStaticLightingVar || AllowStaticLightingVar->GetInt() != 0);
 }
 
 void FReflectionEnvironmentCubemapArray::InitDynamicRHI()
@@ -829,8 +834,9 @@ FReflectionEnvironmentTiledDeferredCS* SelectReflectionEnvironmentTiledDeferredC
 
 void FDeferredShadingSceneRenderer::RenderTiledDeferredImageBasedReflections(FRHICommandListImmediate& RHICmdList, const TRefCountPtr<IPooledRenderTarget>& DynamicBentNormalAO)
 {
-	const bool bUseLightmaps = CVarDiffuseFromCaptures.GetValueOnRenderThread() == 0;
-	const uint32 bHalfRes = CVarHalfResReflections.GetValueOnRenderThread() != 0;
+	static const auto AllowStaticLightingVar = IConsoleManager::Get().FindTConsoleVariableDataInt(TEXT("r.AllowStaticLighting"));
+	const bool bUseLightmaps = (AllowStaticLightingVar->GetValueOnRenderThread() == 1) && (CVarDiffuseFromCaptures.GetValueOnRenderThread() == 0);
+	const bool bHalfRes = CVarHalfResReflections.GetValueOnRenderThread() != 0;
 
 	TRefCountPtr<IPooledRenderTarget> NewSceneColor;
 	{
@@ -898,6 +904,11 @@ void FDeferredShadingSceneRenderer::RenderTiledDeferredImageBasedReflections(FRH
 
 void FDeferredShadingSceneRenderer::RenderStandardDeferredImageBasedReflections(FRHICommandListImmediate& RHICmdList, bool bReflectionEnv, const TRefCountPtr<IPooledRenderTarget>& DynamicBentNormalAO)
 {
+	if(!ViewFamily.EngineShowFlags.Lighting)
+	{
+		return;
+	}
+
 	const bool bSkyLight = Scene->SkyLight
 		&& Scene->SkyLight->ProcessedTexture
 		&& ViewFamily.EngineShowFlags.SkyLighting;
@@ -958,7 +969,19 @@ void FDeferredShadingSceneRenderer::RenderStandardDeferredImageBasedReflections(
 	{
 		FViewInfo& View = Views[ViewIndex];
 
-		bool bRequiresApply = bSkyLight;
+		bool bLPV = false;
+		FSceneViewState* ViewState = (FSceneViewState*)View.State;
+		if ( ViewState != nullptr && ViewState->GetLightPropagationVolume() != nullptr && View.FinalPostProcessSettings.LPVIntensity > 0.0f )
+		{
+			bLPV = true;
+		} 
+		
+		bool bAmbient = View.FinalPostProcessSettings.ContributingCubemaps.Num() > 0;
+		bool bMixing = CVarLPVMixing.GetValueOnRenderThread() != 0;
+		bool bEnvironmentMixing = bMixing && (bAmbient || bLPV);
+		bool bRequiresApply = bSkyLight
+			// If Reflection Environment is active and mixed with indirect lighting (Ambient + LPV), apply is required!
+			|| (View.Family->EngineShowFlags.ReflectionEnvironment && (bReflectionEnv || bEnvironmentMixing) );
 
 		const bool bSSR = DoScreenSpaceReflections(View);
 
@@ -970,7 +993,8 @@ void FDeferredShadingSceneRenderer::RenderStandardDeferredImageBasedReflections(
 			ScreenSpaceReflections(RHICmdList, View, SSROutput);
 		}
 
-		TRefCountPtr<IPooledRenderTarget> LightAccumulation;
+	    /* Light Accumulation moved to SceneRenderTargets */
+	    TRefCountPtr<IPooledRenderTarget> LightAccumulation = GSceneRenderTargets.LightAccumulation;
 
 		if (bReflectionEnv)
 		{
@@ -979,14 +1003,17 @@ void FDeferredShadingSceneRenderer::RenderStandardDeferredImageBasedReflections(
 			SCOPED_DRAW_EVENT(RHICmdList, StandardDeferredReflectionEnvironment);
 
 			{
-				FPooledRenderTargetDesc Desc = GSceneRenderTargets.GetSceneColor()->GetDesc();
-				// Make sure we get an alpha channel
-				Desc.Format = PF_FloatRGBA;
-				GRenderTargetPool.FindFreeElement(Desc, LightAccumulation, TEXT("LightAccumulation"));
+				// Clear to no reflection contribution, alpha of 1 indicates full background contribution
+				ESimpleRenderTargetMode SimpleRenderTargetMode = ESimpleRenderTargetMode::EExistingColorAndDepth;
+
+				// If Reflection Environment is mixed with indirect lighting (Ambient + LPV), skip clear!
+				if (!bMixing)
+				{
+					SimpleRenderTargetMode = ESimpleRenderTargetMode::EClearColorToBlackWithFullAlpha;
+				}
+				SetRenderTarget(RHICmdList, LightAccumulation->GetRenderTargetItem().TargetableTexture, NULL, SimpleRenderTargetMode);
 			}
 
-			// Clear to no reflection contribution, alpha of 1 indicates full background contribution
-			SetRenderTarget(RHICmdList, LightAccumulation->GetRenderTargetItem().TargetableTexture, NULL, ESimpleRenderTargetMode::EClearColorToBlackWithFullAlpha);
 
 			RHICmdList.SetViewport(View.ViewRect.Min.X, View.ViewRect.Min.Y, 0.0f, View.ViewRect.Max.X, View.ViewRect.Max.Y, 1.0f);
 
@@ -1063,6 +1090,10 @@ void FDeferredShadingSceneRenderer::RenderStandardDeferredImageBasedReflections(
 				// should never be used but during debugging it can happen
 				LightAccumulation = GSystemTextures.WhiteDummy;
 			}
+			
+			// Activate Reflection Environment if we choose to mix it with indirect lighting (Ambient + LPV)
+			// todo: refactor (we abuse another boolean to pass the data through)
+			bReflectionEnv = bReflectionEnv || bEnvironmentMixing;
 
 #define CASE(A,B,C) \
 			case ((A << 2) | (B << 1) | C) : \
diff --git a/Engine/Source/Runtime/Renderer/Private/SceneRendering.cpp b/Engine/Source/Runtime/Renderer/Private/SceneRendering.cpp
--- a/Engine/Source/Runtime/Renderer/Private/SceneRendering.cpp
+++ b/Engine/Source/Runtime/Renderer/Private/SceneRendering.cpp
@@ -631,8 +631,10 @@ TUniformBufferRef<FViewUniformShaderParameters> FViewInfo::CreateUniformBuffer(
 	ViewUniformShaderParameters.Random = FMath::Rand();
 	ViewUniformShaderParameters.FrameNumber = Family->FrameNumber;
 
+	// Lets not use Lightmaps if we don't allow static lighting, shall we?
+	static const auto AllowStaticLightingVar = IConsoleManager::Get().FindConsoleVariable(TEXT("r.AllowStaticLighting"));
 	static const auto CVar = IConsoleManager::Get().FindConsoleVariable(TEXT("r.DiffuseFromCaptures"));
-	const bool bUseLightmaps = CVar->GetInt() == 0;
+	const bool bUseLightmaps = (AllowStaticLightingVar->GetInt() == 1) && (CVar->GetInt() == 0);
 
 	ViewUniformShaderParameters.CameraCut = bCameraCut ? 1 : 0;
 	ViewUniformShaderParameters.UseLightmaps = bUseLightmaps ? 1 : 0;
diff --git a/Engine/Source/Runtime/Renderer/Private/ShadowRendering.cpp b/Engine/Source/Runtime/Renderer/Private/ShadowRendering.cpp
--- a/Engine/Source/Runtime/Renderer/Private/ShadowRendering.cpp
+++ b/Engine/Source/Runtime/Renderer/Private/ShadowRendering.cpp
@@ -161,7 +161,8 @@ public:
 			// Reflective shadow map shaders must be compiled for every material because they access the material normal
 			return !bUsePositionOnlyStream
 				// Don't render ShadowDepth for translucent unlit materials, unless we're injecting emissive
-				&& ((!IsTranslucentBlendMode(Material->GetBlendMode()) && Material->GetShadingModel() != MSM_Unlit) || Material->ShouldInjectEmissiveIntoLPV() )
+				&& ((!IsTranslucentBlendMode(Material->GetBlendMode()) && Material->GetShadingModel() != MSM_Unlit) || Material->ShouldInjectEmissiveIntoLPV() 
+					|| Material->ShouldBlockGI() )
 				&& IsFeatureLevelSupported(Platform, ERHIFeatureLevel::SM5);
 		}
 		else
@@ -446,7 +447,9 @@ public:
 		if(bRenderReflectiveShadowMap)
 		{
 			// LPV also propagates light transmission (for transmissive materials)
-			SetShaderValue(RHICmdList, ShaderRHI,ReflectiveShadowMapTextureResolution,FVector2D(GSceneRenderTargets.GetReflectiveShadowMapTextureResolution()));
+			SetShaderValue(RHICmdList, ShaderRHI,ReflectiveShadowMapTextureResolution,
+				FVector2D(GSceneRenderTargets.GetReflectiveShadowMapResolution(),
+					GSceneRenderTargets.GetReflectiveShadowMapResolution()));
 			SetShaderValue(
 				RHICmdList, 
 				ShaderRHI,
@@ -533,7 +536,7 @@ public:
 			// Reflective shadow map shaders must be compiled for every material because they access the material normal
 			return 
 				// Only compile one pass point light shaders for feature levels >= SM4
-				( (!IsTranslucentBlendMode(Material->GetBlendMode()) && Material->GetShadingModel() != MSM_Unlit) || Material->ShouldInjectEmissiveIntoLPV() )
+				( (!IsTranslucentBlendMode(Material->GetBlendMode()) && Material->GetShadingModel() != MSM_Unlit) || Material->ShouldInjectEmissiveIntoLPV() || Material->ShouldBlockGI() )
 				&& IsFeatureLevelSupported(Platform, ERHIFeatureLevel::SM5);
 		}
 		else
@@ -1002,7 +1005,7 @@ void FShadowDepthDrawingPolicyFactory::AddStaticMesh(FScene* Scene,FStaticMesh*
 		const bool bLightPropagationVolume = UseLightPropagationVolumeRT(FeatureLevel);
 		const bool bTwoSided  = Material->IsTwoSided() || StaticMesh->PrimitiveSceneInfo->Proxy->CastsShadowAsTwoSided();
 		const bool bLitOpaque = !IsTranslucentBlendMode(BlendMode) && ShadingModel != MSM_Unlit;
-		if(bLightPropagationVolume && (bLitOpaque || Material->ShouldInjectEmissiveIntoLPV()))
+		if (bLightPropagationVolume && ((!IsTranslucentBlendMode(BlendMode) && ShadingModel != MSM_Unlit) || Material->ShouldInjectEmissiveIntoLPV() || Material->ShouldBlockGI()))
 		{
 			// Add the static mesh to the shadow's subject draw list.
 			if ( StaticMesh->PrimitiveSceneInfo->Proxy->AffectsDynamicIndirectLighting() )
@@ -1447,7 +1450,7 @@ public:
 	}
 };
 
-void FProjectedShadowInfo::SetStateForDepth(FRHICommandList& RHICmdList)
+void FProjectedShadowInfo::SetStateForDepth(FRHICommandList& RHICmdList, EShadowDepthRenderMode RenderMode)
 {
 	check(bAllocated);
 
@@ -1490,7 +1493,20 @@ void FProjectedShadowInfo::SetStateForDepth(FRHICommandList& RHICmdList)
 		}
 	}
 
-	if (bReflectiveShadowmap && !CascadeSettings.bOnePassPointLightShadow)
+	// GIBlockingVolumes render mode only affects the reflective shadow map, using the opacity of the material to multiply against the existing color.
+	if (RenderMode == ShadowDepthRenderMode_GIBlockingVolumes)
+	{
+		RHICmdList.SetBlendState(TStaticBlendState<CW_NONE, BO_Add, BF_Zero, BF_One, BO_Add, BF_Zero, BF_One,
+			CW_RGBA, BO_Add, BF_Zero, BF_InverseSourceAlpha, BO_Add, BF_Zero, BF_One>::GetRHI());
+	}
+
+	// The EmissiveOnly render mode shouldn't write into the reflective shadow map, only into the LPV.
+	else if (RenderMode == ShadowDepthRenderMode_EmissiveOnly)
+	{
+		RHICmdList.SetBlendState(TStaticBlendState<CW_NONE, BO_Add, BF_Zero, BF_One, BO_Add, BF_Zero, BF_One, CW_NONE>::GetRHI());
+	}
+
+	else if (bReflectiveShadowmap && !CascadeSettings.bOnePassPointLightShadow)
 	{
 		// Enable color writes to the reflective shadow map targets with opaque blending
 		RHICmdList.SetBlendState(TStaticBlendStateWriteMask<CW_RGBA, CW_RGBA>::GetRHI());
@@ -1501,7 +1517,15 @@ void FProjectedShadowInfo::SetStateForDepth(FRHICommandList& RHICmdList)
 		RHICmdList.SetBlendState(TStaticBlendState<CW_NONE>::GetRHI());
 	}
 
-	RHICmdList.SetDepthStencilState(TStaticDepthStencilState<true, CF_LessEqual>::GetRHI());
+
+	if (RenderMode == ShadowDepthRenderMode_EmissiveOnly || RenderMode == ShadowDepthRenderMode_GIBlockingVolumes)
+	{
+		RHICmdList.SetDepthStencilState(TStaticDepthStencilState<false, CF_LessEqual>::GetRHI());
+	}
+	else
+	{
+		RHICmdList.SetDepthStencilState(TStaticDepthStencilState<true, CF_LessEqual>::GetRHI());
+	}
 }
 
 
@@ -1521,12 +1545,16 @@ class FShadowParallelCommandListSet : public FParallelCommandListSet
 {
 	FProjectedShadowInfo& ProjectedShadowInfo;
 	TFunctionRef<void (FRHICommandList& RHICmdList)> SetShadowRenderTargets;
+	EShadowDepthRenderMode RenderMode;
+
 public:
-	FShadowParallelCommandListSet(const FViewInfo& InView, FRHICommandList& InParentCmdList, bool* InOutDirty, bool bInParallelExecute, FProjectedShadowInfo& InProjectedShadowInfo, TFunctionRef<void (FRHICommandList& RHICmdList)> InSetShadowRenderTargets)
+	FShadowParallelCommandListSet(const FViewInfo& InView, FRHICommandList& InParentCmdList, bool* InOutDirty, bool bInParallelExecute, FProjectedShadowInfo& InProjectedShadowInfo, TFunctionRef<void (FRHICommandList& RHICmdList)> InSetShadowRenderTargets
+		, EShadowDepthRenderMode RenderModeIn )
 		: FParallelCommandListSet(InView, InParentCmdList, InOutDirty, bInParallelExecute)
 		, ProjectedShadowInfo(InProjectedShadowInfo)
 		, SetShadowRenderTargets(InSetShadowRenderTargets)
-	{
+		, RenderMode(RenderModeIn)
+		{
 		SetStateOnCommandList(ParentCmdList);
 	}
 
@@ -1538,11 +1566,11 @@ public:
 	virtual void SetStateOnCommandList(FRHICommandList& CmdList) override
 	{	
 		SetShadowRenderTargets(CmdList);
-		ProjectedShadowInfo.SetStateForDepth(CmdList);
+		ProjectedShadowInfo.SetStateForDepth(CmdList,RenderMode);
 	}
 };
 
-void FProjectedShadowInfo::RenderDepthInner(FRHICommandList& RHICmdList, FSceneRenderer* SceneRenderer, const FViewInfo* FoundView, TFunctionRef<void (FRHICommandList& RHICmdList)> SetShadowRenderTargets)
+void FProjectedShadowInfo::RenderDepthInner(FRHICommandList& RHICmdList, FSceneRenderer* SceneRenderer, const FViewInfo* FoundView, TFunctionRef<void (FRHICommandList& RHICmdList)> SetShadowRenderTargets, EShadowDepthRenderMode RenderMode)
 {
 	FShadowDepthDrawingPolicyContext PolicyContext(this);
 
@@ -1553,10 +1581,10 @@ void FProjectedShadowInfo::RenderDepthInner(FRHICommandList& RHICmdList, FSceneR
 		FScopedCommandListWaitForTasks Flusher;
 
 		{
-			FShadowParallelCommandListSet ParallelCommandListSet(*FoundView, RHICmdList, nullptr, CVarRHICmdShadowDeferredContexts.GetValueOnRenderThread() > 0, *this, SetShadowRenderTargets);
+			FShadowParallelCommandListSet ParallelCommandListSet(*FoundView, RHICmdList, nullptr, CVarRHICmdShadowDeferredContexts.GetValueOnRenderThread() > 0, *this, SetShadowRenderTargets, RenderMode);
 
 			// Draw the subject's static elements using static draw lists
-			if (IsWholeSceneDirectionalShadow())
+			if (IsWholeSceneDirectionalShadow() && RenderMode != ShadowDepthRenderMode_EmissiveOnly && RenderMode != ShadowDepthRenderMode_GIBlockingVolumes)
 			{
 				SCOPE_CYCLE_COUNTER(STAT_WholeSceneStaticDrawListShadowDepthsTime);
 
@@ -1594,10 +1622,10 @@ void FProjectedShadowInfo::RenderDepthInner(FRHICommandList& RHICmdList, FSceneR
 	else
 	{
 		// single threaded version
-		SetStateForDepth(RHICmdList);
+		SetStateForDepth(RHICmdList, RenderMode);
 
 		// Draw the subject's static elements using static draw lists
-		if (IsWholeSceneDirectionalShadow())
+		if (IsWholeSceneDirectionalShadow() && RenderMode != ShadowDepthRenderMode_EmissiveOnly && RenderMode != ShadowDepthRenderMode_GIBlockingVolumes)
 		{
 			SCOPE_CYCLE_COUNTER(STAT_WholeSceneStaticDrawListShadowDepthsTime);
 
@@ -1631,8 +1659,32 @@ void FProjectedShadowInfo::RenderDepthInner(FRHICommandList& RHICmdList, FSceneR
 	}
 }
 
-void FProjectedShadowInfo::RenderDepth(FRHICommandList& RHICmdList, FSceneRenderer* SceneRenderer, TFunctionRef<void (FRHICommandList& RHICmdList)> SetShadowRenderTargets)
+void FProjectedShadowInfo::RenderDepth(FRHICommandList& RHICmdList, FSceneRenderer* SceneRenderer, TFunctionRef<void (FRHICommandList& RHICmdList)> SetShadowRenderTargets, EShadowDepthRenderMode RenderMode)
 {
+	// Select the correct set of arrays for the current render mode
+	TArray<FShadowStaticMeshElement,SceneRenderingAllocator>* PtrCurrentMeshElements = nullptr;
+	PrimitiveArrayType* PtrCurrentPrimitives = nullptr;
+
+	switch(RenderMode)
+	{
+	case ShadowDepthRenderMode_Dynamic:
+		PtrCurrentMeshElements = &SubjectMeshElements;
+		PtrCurrentPrimitives = &SubjectPrimitives;
+		break;
+	case ShadowDepthRenderMode_EmissiveOnly:
+		PtrCurrentMeshElements = &EmissiveOnlyMeshElements;
+		PtrCurrentPrimitives = &EmissiveOnlyPrimitives;
+		break;
+	case ShadowDepthRenderMode_GIBlockingVolumes:
+		PtrCurrentMeshElements = &GIBlockingMeshElements;
+		PtrCurrentPrimitives = &GIBlockingPrimitives;
+		break;
+	default:
+		check(0);
+	}
+	TArray<FShadowStaticMeshElement,SceneRenderingAllocator>& CurrentMeshElements = *PtrCurrentMeshElements;
+	PrimitiveArrayType& CurrentPrimitives = *PtrCurrentPrimitives;
+
 #if WANTS_DRAW_MESH_EVENTS
 	FString EventName;
 	GetShadowTypeNameForDrawEvent(EventName);
@@ -1642,6 +1694,12 @@ void FProjectedShadowInfo::RenderDepth(FRHICommandList& RHICmdList, FSceneRender
 	CONDITIONAL_SCOPE_CYCLE_COUNTER(STAT_RenderWholeSceneShadowDepthsTime, bWholeSceneShadow);
 	CONDITIONAL_SCOPE_CYCLE_COUNTER(STAT_RenderPerObjectShadowDepthsTime, !bWholeSceneShadow);
 
+	// Exit early if there are no meshes or primitives to render in the emissive ony render mode.
+	if (RenderMode != ShadowDepthRenderMode_Dynamic && CurrentMeshElements.Num() == 0 && CurrentPrimitives.Num() == 0)
+	{
+		return;
+	}
+
 	// Choose an arbitrary view where this shadow's subject is relevant.
 	FViewInfo* FoundView = NULL;
 	for(int32 ViewIndex = 0;ViewIndex < SceneRenderer->Views.Num();ViewIndex++)
@@ -1705,7 +1763,7 @@ void FProjectedShadowInfo::RenderDepth(FRHICommandList& RHICmdList, FSceneRender
 	// Lighting only should only affect the material used with direct lighting, not the indirect lighting
 	FoundView->bForceShowMaterials = true;
 
-	RenderDepthInner(RHICmdList, SceneRenderer, FoundView, SetShadowRenderTargets);
+	RenderDepthInner(RHICmdList, SceneRenderer, FoundView, SetShadowRenderTargets, RenderMode);
 
 	FoundView->bForceShowMaterials = false;
 	FoundView->UniformBuffer = OriginalUniformBuffer;
@@ -3089,6 +3147,11 @@ bool FDeferredShadingSceneRenderer::RenderReflectiveShadowMaps(FRHICommandListIm
 				{
 					ProjectedShadowInfo->ClearDepth(RHICmdList, this, true);
 					ProjectedShadowInfo->RenderDepth(RHICmdList, this, SetShadowRenderTargets);
+
+					// Render emissive only meshes as they are held in a separate list.
+					ProjectedShadowInfo->RenderDepth(RHICmdList, this, SetShadowRenderTargets, ShadowDepthRenderMode_EmissiveOnly);
+					// Render gi blocking volume meshes.
+					ProjectedShadowInfo->RenderDepth(RHICmdList, this, SetShadowRenderTargets, ShadowDepthRenderMode_GIBlockingVolumes);
 				}
 
 				GSceneRenderTargets.FinishRenderingReflectiveShadowMap(RHICmdList);
diff --git a/Engine/Source/Runtime/Renderer/Private/ShadowRendering.h b/Engine/Source/Runtime/Renderer/Private/ShadowRendering.h
--- a/Engine/Source/Runtime/Renderer/Private/ShadowRendering.h
+++ b/Engine/Source/Runtime/Renderer/Private/ShadowRendering.h
@@ -487,6 +487,18 @@ public:
 	bool bIsTwoSided;
 };
 
+enum EShadowDepthRenderMode
+{
+	/** The render mode used by regular shadows */
+	ShadowDepthRenderMode_Dynamic,
+
+	/** The render mode used when injecting emissive-only objects into the RSM. */
+	ShadowDepthRenderMode_EmissiveOnly,
+
+	/** The render mode used when rendering volumes which block global illumination. */
+	ShadowDepthRenderMode_GIBlockingVolumes,
+};
+
 /**
  * Information about a projected shadow.
  */
@@ -630,10 +642,10 @@ public:
 	/**
 	 * Renders the shadow subject depth.
 	 */
-	void RenderDepth(FRHICommandList& RHICmdList, class FSceneRenderer* SceneRenderer, TFunctionRef<void (FRHICommandList& RHICmdList)> SetShadowRenderTargets);
+	void RenderDepth(FRHICommandList& RHICmdList, class FSceneRenderer* SceneRenderer, TFunctionRef<void (FRHICommandList& RHICmdList)> SetShadowRenderTargets, EShadowDepthRenderMode RenderMode = ShadowDepthRenderMode_Dynamic);
 
 	/** Set state for depth rendering */
-	void SetStateForDepth(FRHICommandList& RHICmdList);
+	void SetStateForDepth(FRHICommandList& RHICmdList, EShadowDepthRenderMode RenderMode );
 
 	void ClearDepth(FRHICommandList& RHICmdList, class FDeferredShadingSceneRenderer* SceneRenderer, bool bPerformClear);
 
@@ -743,6 +755,16 @@ private:
 	/** Subject primitives with translucent relevance. */
 	PrimitiveArrayType SubjectTranslucentPrimitives;
 
+	/** Translucent LPV injection: dynamic shadow casting elements */
+	PrimitiveArrayType EmissiveOnlyPrimitives;
+	/** Translucent LPV injection: Static shadow casting elements. */
+	TArray<FShadowStaticMeshElement,SceneRenderingAllocator> EmissiveOnlyMeshElements;
+
+	/** GI blocking volume: dynamic shadow casting elements */
+	PrimitiveArrayType GIBlockingPrimitives;
+	/** GI blocking volume: Static shadow casting elements. */
+	TArray<FShadowStaticMeshElement,SceneRenderingAllocator> GIBlockingMeshElements;
+
 	/** Static shadow casting elements. */
 	TArray<FShadowStaticMeshElement,SceneRenderingAllocator> SubjectMeshElements;
 
@@ -762,7 +784,7 @@ private:
 	/**
 	* Renders the shadow subject depth, to a particular hacked view
 	*/
-	void RenderDepthInner(FRHICommandList& RHICmdList, class FSceneRenderer* SceneRenderer, const FViewInfo* FoundView, TFunctionRef<void (FRHICommandList& RHICmdList)> SetShadowRenderTargets);
+	void RenderDepthInner(FRHICommandList& RHICmdList, class FSceneRenderer* SceneRenderer, const FViewInfo* FoundView, TFunctionRef<void (FRHICommandList& RHICmdList)> SetShadowRenderTargets, EShadowDepthRenderMode RenderMode );
 
 	/**
 	* Renders the dynamic shadow subject depth, to a particular hacked view
diff --git a/Engine/Source/Runtime/Renderer/Private/ShadowSetup.cpp b/Engine/Source/Runtime/Renderer/Private/ShadowSetup.cpp
--- a/Engine/Source/Runtime/Renderer/Private/ShadowSetup.cpp
+++ b/Engine/Source/Runtime/Renderer/Private/ShadowSetup.cpp
@@ -537,9 +537,9 @@ void FProjectedShadowInfo::SetupWholeSceneProjection(
 		static bool bQuantize = true;
 		if ( bQuantize )
 		{
-			const FIntPoint ShadowBufferResolution = GSceneRenderTargets.GetReflectiveShadowMapTextureResolution();
-			uint32 ShadowDepthBufferSizeX = ShadowBufferResolution.X;
-			uint32 ShadowDepthBufferSizeY = ShadowBufferResolution.Y;
+			const int32 ShadowBufferResolution = GSceneRenderTargets.GetReflectiveShadowMapResolution();
+			uint32 ShadowDepthBufferSizeX = ShadowBufferResolution;
+			uint32 ShadowDepthBufferSizeY = ShadowBufferResolution;
 			// Transform the shadow's position into shadowmap space
 			const FVector TransformedPosition = WorldToFace.TransformPosition(-PreShadowTranslation);
 
@@ -859,16 +859,57 @@ void FProjectedShadowInfo::AddSubjectPrimitive(FPrimitiveSceneInfo* PrimitiveSce
 		}
 
 		// Add translucent shadow casting primitives to SubjectTranslucentPrimitives
-		if (bTranslucentRelevance && bShadowRelevance && bTranslucentShadow)
+		if (bTranslucentRelevance && bShadowRelevance)
 		{
-			SubjectTranslucentPrimitives.Add(PrimitiveSceneInfo);
+			if (bTranslucentShadow)
+			{
+				SubjectTranslucentPrimitives.Add(PrimitiveSceneInfo);
+			}
+			else if (bReflectiveShadowmap)
+			{
+				if (PrimitiveSceneInfo->StaticMeshes.Num() > 0)
+				{
+					for (int32 ViewIndex = 0; ViewIndex < Views.Num(); ViewIndex++)
+					{
+						FViewInfo& CurrentView = *Views[ViewIndex];
+						// Add the primitive's static mesh elements to the draw lists.
+						for (int32 MeshIndex = 0; MeshIndex < PrimitiveSceneInfo->StaticMeshes.Num(); MeshIndex++)
+						{
+							FStaticMesh& StaticMesh = PrimitiveSceneInfo->StaticMeshes[MeshIndex];
+							const FMaterialRenderProxy* MaterialRenderProxy = StaticMesh.MaterialRenderProxy;
+							const FMaterial* Material = MaterialRenderProxy->GetMaterial(FeatureLevel);
+							const bool bTwoSided = Material->IsTwoSided() || PrimitiveSceneInfo->Proxy->CastsShadowAsTwoSided();
+    
+							if (Material->ShouldBlockGI())
+							{
+								CurrentView.StaticMeshShadowDepthMap[StaticMesh.Id] = true;
+								GIBlockingMeshElements.Add(FShadowStaticMeshElement(MaterialRenderProxy, Material, &StaticMesh,bTwoSided));
+							}
+							else if (Material->ShouldInjectEmissiveIntoLPV())
+							{
+								CurrentView.StaticMeshShadowDepthMap[StaticMesh.Id] = true;
+								EmissiveOnlyMeshElements.Add(FShadowStaticMeshElement(MaterialRenderProxy, Material, &StaticMesh,bTwoSided));
+							}
+						}
+					}
+				}
+    			else
+				{
+					EmissiveOnlyPrimitives.Add(PrimitiveSceneInfo);
+				}
+			}
 		}
 	}
 }
 
 bool FProjectedShadowInfo::HasSubjectPrims() const
 {
-	return SubjectPrimitives.Num() > 0 || SubjectMeshElements.Num() > 0;
+	return SubjectPrimitives.Num() > 0
+		|| SubjectMeshElements.Num() > 0
+		|| EmissiveOnlyPrimitives.Num() > 0 
+		|| EmissiveOnlyMeshElements.Num() > 0
+		|| GIBlockingMeshElements.Num() > 0
+		|| GIBlockingPrimitives.Num() > 0;
 }
 
 void FProjectedShadowInfo::AddReceiverPrimitive(FPrimitiveSceneInfo* PrimitiveSceneInfo)
@@ -1021,6 +1062,8 @@ void FProjectedShadowInfo::ClearTransientArrays()
 	SubjectPrimitives.Empty();
 	ReceiverPrimitives.Empty();
 	SubjectMeshElements.Empty();
+	EmissiveOnlyPrimitives.Empty();
+	EmissiveOnlyMeshElements.Empty();
 	DynamicSubjectMeshElements.Empty();
 	DynamicReceiverMeshElements.Empty();
 	DynamicSubjectTranslucentMeshElements.Empty();
@@ -2204,7 +2247,7 @@ void FSceneRenderer::AddViewDependentWholeSceneShadowsForView(
 			{
 				FLightPropagationVolume* LightPropagationVolume = ViewState->GetLightPropagationVolume();
 
-				if (LightPropagationVolume && View.FinalPostProcessSettings.LPVIntensity > 0)
+				if (LightPropagationVolume && LightPropagationVolume->bInitialized && View.FinalPostProcessSettings.LPVIntensity > 0)
 				{
 					// Generate the RSM shadow info
 					FWholeSceneProjectedShadowInitializer ProjectedShadowInitializer;
@@ -2215,7 +2258,7 @@ void FSceneRenderer::AddViewDependentWholeSceneShadowsForView(
 						// moved out from the FProjectedShadowInfo constructor
 						ProjectedShadowInitializer.CascadeSettings.ShadowSplitIndex = 0;
 
-						const FIntPoint ShadowBufferResolution = GSceneRenderTargets.GetReflectiveShadowMapTextureResolution();
+						const int32 ShadowBufferResolution = GSceneRenderTargets.GetReflectiveShadowMapResolution();
 
 						// Create the projected shadow info.
 						FProjectedShadowInfo* ProjectedShadowInfo = new(FMemStack::Get(), 1, 16) FProjectedShadowInfo;
@@ -2224,8 +2267,8 @@ void FSceneRenderer::AddViewDependentWholeSceneShadowsForView(
 							&LightSceneInfo,
 							&View,
 							ProjectedShadowInitializer,
-							ShadowBufferResolution.X,
-							ShadowBufferResolution.Y,
+							ShadowBufferResolution,
+							ShadowBufferResolution,
 							true);		// RSM
 
 						FVisibleLightInfo& LightViewInfo = VisibleLightInfos[LightSceneInfo.Id];
diff --git a/Engine/Source/Runtime/Renderer/Public/SceneRenderTargetParameters.h b/Engine/Source/Runtime/Renderer/Public/SceneRenderTargetParameters.h
--- a/Engine/Source/Runtime/Renderer/Public/SceneRenderTargetParameters.h
+++ b/Engine/Source/Runtime/Renderer/Public/SceneRenderTargetParameters.h
@@ -54,6 +54,8 @@ private:
 	FShaderResourceParameter SceneDepthSurfaceParameter;
 	/**  */
 	FShaderResourceParameter SceneDepthTextureNonMS;
+	FShaderResourceParameter DirectionalOcclusionSampler;
+	FShaderResourceParameter DirectionalOcclusionTexture;
 };
 
 /** Pixel shader parameters needed for deferred passes. */
