Fix build failure on Clang platforms due to shadowed variables.

diff --git a/Engine/Source/Runtime/Renderer/Private/DeferredShadingRenderer.cpp b/Engine/Source/Runtime/Renderer/Private/DeferredShadingRenderer.cpp
--- a/Engine/Source/Runtime/Renderer/Private/DeferredShadingRenderer.cpp
+++ b/Engine/Source/Runtime/Renderer/Private/DeferredShadingRenderer.cpp
@@ -1487,8 +1487,8 @@ void FDeferredShadingSceneRenderer::ClearLPVs(FRHICommandListImmediate& RHICmdLi
 
 			if(LightPropagationVolume)
 			{
-				SCOPED_DRAW_EVENT(RHICmdList, ClearLPVs);
-				SCOPE_CYCLE_COUNTER(STAT_UpdateLPVs);
+//				SCOPED_DRAW_EVENT(RHICmdList, ClearLPVs);
+//				SCOPE_CYCLE_COUNTER(STAT_UpdateLPVs);
 				LightPropagationVolume->InitSettings(RHICmdList, Views[ViewIndex]);
 				LightPropagationVolume->Clear(RHICmdList, View);
 			}
@@ -1514,8 +1514,8 @@ void FDeferredShadingSceneRenderer::UpdateLPVs(FRHICommandListImmediate& RHICmdL
 
 			if(LightPropagationVolume)
 			{
-				SCOPED_DRAW_EVENT(RHICmdList, UpdateLPVs);
-				SCOPE_CYCLE_COUNTER(STAT_UpdateLPVs);
+//				SCOPED_DRAW_EVENT(RHICmdList, UpdateLPVs);
+//				SCOPE_CYCLE_COUNTER(STAT_UpdateLPVs);
 
 				LightPropagationVolume->Update(RHICmdList, View);
 			}
