Indirect lighting cache point samples don't use the texture layout anymore, since they don't use the volume texture anyway

diff --git a/Engine/Source/Runtime/Renderer/Private/IndirectLightingCache.cpp b/Engine/Source/Runtime/Renderer/Private/IndirectLightingCache.cpp
--- a/Engine/Source/Runtime/Renderer/Private/IndirectLightingCache.cpp
+++ b/Engine/Source/Runtime/Renderer/Private/IndirectLightingCache.cpp
@@ -142,6 +142,7 @@ FIndirectLightingCache::FIndirectLightingCache(ERHIFeatureLevel::Type InFeatureL
 	, bUpdateAllCacheEntries(true)
 	, BlockAllocator(0, 0, 0, GLightingCacheDimension, GLightingCacheDimension, GLightingCacheDimension, false, false)
 {
+	NextPointId = GLightingCacheDimension + 1;
 	CacheSize = GLightingCacheDimension;
 }
 
@@ -193,13 +194,29 @@ const FIndirectLightingCacheBlock& FIndirectLightingCache::FindBlock(FIntVector
 
 void FIndirectLightingCache::DeallocateBlock(FIntVector Min, int32 Size)
 {
-	verify(BlockAllocator.RemoveElement(Min.X, Min.Y, Min.Z, Size, Size, Size));
+	if (Size > 1)
+	{
+		verify(BlockAllocator.RemoveElement(Min.X, Min.Y, Min.Z, Size, Size, Size));
+	}
+	
 	VolumeBlocks.Remove(Min);
 }
 
 bool FIndirectLightingCache::AllocateBlock(int32 Size, FIntVector& OutMin)
 {
-	return BlockAllocator.AddElement((uint32&)OutMin.X, (uint32&)OutMin.Y, (uint32&)OutMin.Z, Size, Size, Size);
+	if (Size == 1)
+	{
+		// Assign a min that won't overlap with any of the samples allocated from the volume texture, so we can be added to VolumeBlocks without collisions
+		// This min is not used for anything else for point samples
+		OutMin = FIntVector(NextPointId, 0, 0);
+		NextPointId++;
+		// Point samples don't go through the volume texture, allocation always succeeds
+		return true;
+	}
+	else
+	{
+		return BlockAllocator.AddElement((uint32&)OutMin.X, (uint32&)OutMin.Y, (uint32&)OutMin.Z, Size, Size, Size);
+	}
 }
 
 void FIndirectLightingCache::CalculateBlockPositionAndSize(const FBoxSphereBounds& Bounds, int32 TexelSize, FVector& OutMin, FVector& OutSize) const
@@ -305,6 +322,8 @@ FIndirectLightingCacheAllocation* FIndirectLightingCache::AllocatePrimitive(cons
 
 FIndirectLightingCacheAllocation* FIndirectLightingCache::CreateAllocation(int32 BlockSize, const FBoxSphereBounds& Bounds, bool bPointSample)
 {
+	check(BlockSize > 1 || bPointSample);
+
 	FIndirectLightingCacheAllocation* NewAllocation = new FIndirectLightingCacheAllocation();
 	FIndirectLightingCacheBlock NewBlock;
 
