Merge fix for 'black edge' artifact on PS4.

diff --git a/Engine/Source/Runtime/Engine/Private/Particles/ParticleGpuSimulation.cpp b/Engine/Source/Runtime/Engine/Private/Particles/ParticleGpuSimulation.cpp
--- a/Engine/Source/Runtime/Engine/Private/Particles/ParticleGpuSimulation.cpp
+++ b/Engine/Source/Runtime/Engine/Private/Particles/ParticleGpuSimulation.cpp
@@ -255,7 +255,7 @@ public:
 			PF_B8G8R8A8,
 			/*NumMips=*/ 1,
 			TexCreate_None,
-			TexCreate_RenderTargetable,
+			TexCreate_RenderTargetable | TexCreate_NoFastClear,
 			/*bForceSeparateTargetAndShaderResource=*/ false,
 			CreateInfo,
 			TextureTargetRHI,
diff --git a/Engine/Source/Runtime/Renderer/Private/ReflectionEnvironment.cpp b/Engine/Source/Runtime/Renderer/Private/ReflectionEnvironment.cpp
--- a/Engine/Source/Runtime/Renderer/Private/ReflectionEnvironment.cpp
+++ b/Engine/Source/Runtime/Renderer/Private/ReflectionEnvironment.cpp
@@ -844,6 +844,7 @@ void FDeferredShadingSceneRenderer::RenderTiledDeferredImageBasedReflections(FRH
 
 		FPooledRenderTargetDesc Desc = GSceneRenderTargets.GetSceneColor()->GetDesc();
 		Desc.TargetableFlags |= TexCreate_UAV;
+		Desc.TargetableFlags |= TexCreate_NoFastClear;
 
 		// we don't create a new name to make it easier to use "vis SceneColor" and get the last HDRSceneColor
 		GRenderTargetPool.FindFreeElement( Desc, NewSceneColor, TEXT("SceneColor") );
