Fixes a HZB unbound parameter for translussent object having screen space reflection.

diff --git a/Engine/Source/Runtime/Renderer/Private/BasePassRendering.cpp b/Engine/Source/Runtime/Renderer/Private/BasePassRendering.cpp
--- a/Engine/Source/Runtime/Renderer/Private/BasePassRendering.cpp
+++ b/Engine/Source/Runtime/Renderer/Private/BasePassRendering.cpp
@@ -74,9 +74,11 @@ void FSkyLightReflectionParameters::GetSkyParametersFromScene(const FScene* Scen
 
 void FTranslucentLightingParameters::Set(FRHICommandList& RHICmdList, FShader* Shader, const FViewInfo* View)
 {
+	auto PixelShader = Shader->GetPixelShader();
+
 	SetTextureParameter(
 		RHICmdList, 
-		Shader->GetPixelShader(), 
+		PixelShader, 
 		TranslucencyLightingVolumeAmbientInner, 
 		TranslucencyLightingVolumeAmbientInnerSampler, 
 		TStaticSamplerState<SF_Bilinear, AM_Clamp, AM_Clamp, AM_Clamp>::GetRHI(), 
@@ -84,7 +86,7 @@ void FTranslucentLightingParameters::Set(FRHICommandList& RHICmdList, FShader* S
 
 	SetTextureParameter(
 		RHICmdList, 
-		Shader->GetPixelShader(), 
+		PixelShader, 
 		TranslucencyLightingVolumeAmbientOuter, 
 		TranslucencyLightingVolumeAmbientOuterSampler, 
 		TStaticSamplerState<SF_Bilinear, AM_Clamp, AM_Clamp, AM_Clamp>::GetRHI(), 
@@ -92,7 +94,7 @@ void FTranslucentLightingParameters::Set(FRHICommandList& RHICmdList, FShader* S
 
 	SetTextureParameter(
 		RHICmdList, 
-		Shader->GetPixelShader(), 
+		PixelShader, 
 		TranslucencyLightingVolumeDirectionalInner, 
 		TranslucencyLightingVolumeDirectionalInnerSampler, 
 		TStaticSamplerState<SF_Bilinear, AM_Clamp, AM_Clamp, AM_Clamp>::GetRHI(), 
@@ -100,7 +102,7 @@ void FTranslucentLightingParameters::Set(FRHICommandList& RHICmdList, FShader* S
 
 	SetTextureParameter(
 		RHICmdList, 
-		Shader->GetPixelShader(), 
+		PixelShader, 
 		TranslucencyLightingVolumeDirectionalOuter, 
 		TranslucencyLightingVolumeDirectionalOuterSampler, 
 		TStaticSamplerState<SF_Bilinear, AM_Clamp, AM_Clamp, AM_Clamp>::GetRHI(), 
@@ -112,7 +114,7 @@ void FTranslucentLightingParameters::Set(FRHICommandList& RHICmdList, FShader* S
 	{
 		SetTextureParameter(
 			RHICmdList, 
-			Shader->GetPixelShader(), 
+			PixelShader, 
 			HZBTexture, 
 			HZBSampler, 
 			TStaticSamplerState<SF_Point, AM_Clamp, AM_Clamp, AM_Clamp>::GetRHI(), 
@@ -128,11 +130,24 @@ void FTranslucentLightingParameters::Set(FRHICommandList& RHICmdList, FShader* S
 
 		SetTextureParameter(
 			RHICmdList, 
-			Shader->GetPixelShader(), 
+			PixelShader, 
 			PrevSceneColor, 
 			PrevSceneColorSampler, 
 			TStaticSamplerState<SF_Bilinear, AM_Clamp, AM_Clamp, AM_Clamp>::GetRHI(), 
 			(*PrevSceneColorRT)->GetRenderTargetItem().ShaderResourceTexture );
+		
+		const FVector2D HZBUvFactor(
+			float(View->ViewRect.Width()) / float(2 * View->HZBMipmap0Size.X),
+			float(View->ViewRect.Height()) / float(2 * View->HZBMipmap0Size.Y)
+			);
+		const FVector4 HZBUvFactorAndInvFactorValue(
+			HZBUvFactor.X,
+			HZBUvFactor.Y,
+			1.0f / HZBUvFactor.X,
+			1.0f / HZBUvFactor.Y
+			);
+			
+		SetShaderValue(RHICmdList, PixelShader, HZBUvFactorAndInvFactor, HZBUvFactorAndInvFactorValue);
 	}
 }
 
diff --git a/Engine/Source/Runtime/Renderer/Private/BasePassRendering.h b/Engine/Source/Runtime/Renderer/Private/BasePassRendering.h
--- a/Engine/Source/Runtime/Renderer/Private/BasePassRendering.h
+++ b/Engine/Source/Runtime/Renderer/Private/BasePassRendering.h
@@ -265,6 +265,7 @@ public:
 		SkyLightReflectionParameters.Bind(ParameterMap);
 		HZBTexture.Bind(ParameterMap, TEXT("HZBTexture"));
 		HZBSampler.Bind(ParameterMap, TEXT("HZBSampler"));
+		HZBUvFactorAndInvFactor.Bind(ParameterMap, TEXT("HZBUvFactorAndInvFactor"));
 		PrevSceneColor.Bind(ParameterMap, TEXT("PrevSceneColor"));
 		PrevSceneColorSampler.Bind(ParameterMap, TEXT("PrevSceneColorSampler"));
 	}
@@ -290,6 +291,7 @@ public:
 		Ar << P.SkyLightReflectionParameters;
 		Ar << P.HZBTexture;
 		Ar << P.HZBSampler;
+		Ar << P.HZBUvFactorAndInvFactor;
 		Ar << P.PrevSceneColor;
 		Ar << P.PrevSceneColorSampler;
 		return Ar;
@@ -311,6 +313,7 @@ private:
 	FSkyLightReflectionParameters SkyLightReflectionParameters;
 	FShaderResourceParameter HZBTexture;
 	FShaderResourceParameter HZBSampler;
+	FShaderParameter HZBUvFactorAndInvFactor;
 	FShaderResourceParameter PrevSceneColor;
 	FShaderResourceParameter PrevSceneColorSampler;
 };
