Fix a Reusable MID leak in stereo rendering.

diff --git a/Engine/Source/Runtime/Renderer/Private/ScenePrivate.h b/Engine/Source/Runtime/Renderer/Private/ScenePrivate.h
--- a/Engine/Source/Runtime/Renderer/Private/ScenePrivate.h
+++ b/Engine/Source/Runtime/Renderer/Private/ScenePrivate.h
@@ -770,12 +770,10 @@ public:
 	virtual void OnStartPostProcessing(FSceneView& CurrentView) override
 	{
 		check(IsInGameThread());
-
-		if(CurrentView.Family->Views[0] == &CurrentView)
-		{
-			// Needs to be done once for all views, otherwise garbage collection goes wrong
-			MIDUsedCount = 0;
-		}
+		
+		// Needs to be done once for all viewstates.  If multiple FSceneViews are sharing the same ViewState, this will cause problems.
+		// Sharing should be illegal right now though.
+		MIDUsedCount = 0;
 	}
 
 	// Note: OnStartPostProcessing() needs to be called each frame for each view
