Merging UE4 - Add scissor to the set of state that a deferred context needs

diff --git a/Engine/Source/Runtime/Renderer/Private/ShadowRendering.cpp b/Engine/Source/Runtime/Renderer/Private/ShadowRendering.cpp
--- a/Engine/Source/Runtime/Renderer/Private/ShadowRendering.cpp
+++ b/Engine/Source/Runtime/Renderer/Private/ShadowRendering.cpp
@@ -1526,6 +1526,8 @@ void FProjectedShadowInfo::SetStateForDepth(FRHICommandList& RHICmdList, EShadow
 	{
 		RHICmdList.SetDepthStencilState(TStaticDepthStencilState<true, CF_LessEqual>::GetRHI());
 	}
+
+	RHICmdList.SetScissorRect(false, 0, 0, 0, 0);
 }
 
 
diff --git a/Engine/Source/Runtime/Renderer/Private/TranslucentRendering.cpp b/Engine/Source/Runtime/Renderer/Private/TranslucentRendering.cpp
--- a/Engine/Source/Runtime/Renderer/Private/TranslucentRendering.cpp
+++ b/Engine/Source/Runtime/Renderer/Private/TranslucentRendering.cpp
@@ -29,6 +29,7 @@ static void SetTranslucentRenderTargetAndState(FRHICommandList& RHICmdList, cons
 		// Enable depth test, disable depth writes.
 		RHICmdList.SetDepthStencilState(TStaticDepthStencilState<false, CF_DepthNearOrEqual>::GetRHI());
 	}
+	RHICmdList.SetScissorRect(false, 0, 0, 0, 0);
 }
 
 static void FinishTranslucentRenderTarget(FRHICommandListImmediate& RHICmdList, const FViewInfo& View, bool bSeparateTranslucencyPass)
diff --git a/Engine/Source/Runtime/Renderer/Private/VelocityRendering.cpp b/Engine/Source/Runtime/Renderer/Private/VelocityRendering.cpp
--- a/Engine/Source/Runtime/Renderer/Private/VelocityRendering.cpp
+++ b/Engine/Source/Runtime/Renderer/Private/VelocityRendering.cpp
@@ -620,6 +620,7 @@ static void SetVelocitiesState(FRHICommandList& RHICmdList, const FViewInfo& Vie
 	RHICmdList.SetBlendState(TStaticBlendState<CW_RGBA>::GetRHI());
 	RHICmdList.SetDepthStencilState(TStaticDepthStencilState<false,CF_DepthNearOrEqual>::GetRHI());
 	RHICmdList.SetRasterizerState(GetStaticRasterizerState<true>(FM_Solid, CM_CW));
+	RHICmdList.SetScissorRect(false, 0, 0, 0, 0);
 }
 
 class FVelocityPassParallelCommandListSet : public FParallelCommandListSet
