PS4 - Disable fast clears on depth cube maps - fixes broken pointlight shadows

diff --git a/Engine/Source/Runtime/Renderer/Private/PostProcess/SceneRenderTargets.cpp b/Engine/Source/Runtime/Renderer/Private/PostProcess/SceneRenderTargets.cpp
--- a/Engine/Source/Runtime/Renderer/Private/PostProcess/SceneRenderTargets.cpp
+++ b/Engine/Source/Runtime/Renderer/Private/PostProcess/SceneRenderTargets.cpp
@@ -1455,7 +1455,7 @@ void FSceneRenderTargets::AllocateDeferredShadingPathRenderTargets()
 		{
 			const int32 SurfaceResolution = GetCubeShadowDepthZResolution(SurfaceIndex);
 
-			FPooledRenderTargetDesc Desc(FPooledRenderTargetDesc::CreateCubemapDesc(SurfaceResolution, PF_ShadowDepth, TexCreate_None, TexCreate_DepthStencilTargetable, false));
+			FPooledRenderTargetDesc Desc(FPooledRenderTargetDesc::CreateCubemapDesc(SurfaceResolution, PF_ShadowDepth, TexCreate_None, TexCreate_DepthStencilTargetable | TexCreate_NoFastClear, false));
 			GRenderTargetPool.FindFreeElement(Desc, CubeShadowDepthZ[SurfaceIndex], TEXT("CubeShadowDepthZ[]"));
 		}
 	}
