Fixed translucent self shadowing

diff --git a/Engine/Source/Runtime/Renderer/Private/TranslucentLighting.cpp b/Engine/Source/Runtime/Renderer/Private/TranslucentLighting.cpp
--- a/Engine/Source/Runtime/Renderer/Private/TranslucentLighting.cpp
+++ b/Engine/Source/Runtime/Renderer/Private/TranslucentLighting.cpp
@@ -605,7 +605,7 @@ void FProjectedShadowInfo::RenderTranslucencyDepths(FRHICommandList& RHICmdList,
 
 		FResolveParams ResolveParams;
 		RHICmdList.CopyToResolveTarget(ShadowTransmission0, ShadowTransmission0, true, ResolveParams);
-		RHICmdList.CopyToResolveTarget(ShadowTransmission1, ShadowTransmission0, true, ResolveParams);
+		RHICmdList.CopyToResolveTarget(ShadowTransmission1, ShadowTransmission1, true, ResolveParams);
 	}	
 
 	// Restore overridden properties
