Implement the SSR stencil pre-pass behind the console variable r.SSR.Stencil

This feature buidls a stencil buffer according to the rougness in the GBuffer to make the difference between pixels that need a SSR computation and those who don't, instead of an early return in the SSR shader itself. Therefore SSR pixel shader invocations will only be for pixel that needs on hardware that have a high resolution stencil test before the pixel shader invocations packing into front lines, removing the efficiency drop caused by the thread that might early return in front lines that still have at least one thread that didn't early return.

diff --git a/Engine/Shaders/ScreenSpaceReflections.usf b/Engine/Shaders/ScreenSpaceReflections.usf
--- a/Engine/Shaders/ScreenSpaceReflections.usf
+++ b/Engine/Shaders/ScreenSpaceReflections.usf
@@ -38,17 +38,8 @@ uint ReverseUIntBits( uint bits )
 	return bits;
 }
 
-void ScreenSpaceReflectionsPS(in float4 UVAndScreenPos : TEXCOORD0, out float4 OutColor : SV_Target0)
+float GetRoughness(const in FGBufferData GBuffer)
 {
-	float2 UV = UVAndScreenPos.xy;
-	float2 ScreenPos = UVAndScreenPos.zw;
-	int2 PixelPos = int2(UVAndScreenPos.zw * ScreenPosToPixel.xy + ScreenPosToPixel.zw + 0.5f);
-
-	OutColor = 0;
-
-	FScreenSpaceData ScreenSpaceData = GetScreenSpaceData(UV);
-	FGBufferData GBuffer = ScreenSpaceData.GBuffer;
-	
 	float Roughness = GBuffer.Roughness;
 
 #if USE_CLEARCOAT
@@ -62,16 +53,65 @@ void ScreenSpaceReflectionsPS(in float4 UVAndScreenPos : TEXCOORD0, out float4 O
 	}
 #endif
 
+	return Roughness;
+}
+
+float GetRoughnessFade(in float Roughness)
+{
 	// mask SSR to reduce noise and for better performance, roughness of 0 should have SSR, at MaxRoughness we fade to 0
-	float RoughnessFade = saturate(Roughness * SSRParams.y + 2);
+	return min(Roughness * SSRParams.y + 2, 1.0);
+}
 
-#if 1
-	// Early out
-	BRANCH if( RoughnessFade == 0 || GBuffer.ShadingModelID == 0 )
+// Used only if r.SSR.Stencil = 1
+//
+// This shader generate the stencil buffer according to if a pixel needs to
+// compute SSR. If it doesn't, it initializes the pixel's color.
+//
+// This optimization use the stencil to test to prunes out the pixel
+// that dones't need to compute SSR, before pixel shader invocation are
+// packed into front lines.
+void ScreenSpaceReflectionsStencilPS(in float4 UVAndScreenPos : TEXCOORD0, out float4 OutColor : SV_Target0)
+{
+	const float2 UV = UVAndScreenPos.xy;
+	
+	const FScreenSpaceData ScreenSpaceData = GetScreenSpaceData(UV);
+	const FGBufferData GBuffer = ScreenSpaceData.GBuffer;
+	
+	const float Roughness = GetRoughness(GBuffer);
+	const float RoughnessFade = GetRoughnessFade(Roughness);
+	
+	if (RoughnessFade > 0.0 && GBuffer.ShadingModelID != 0)
+	{
+		// we are going to compute SSR for this pixel, so we discard this
+		// pixel shader invocation to not overwrite the stencil buffer and
+		// tehrefore execute ScreenSpaceReflectionsPS() for this pixel. 
+		discard;
+	}
+	
+	// we are not going to compute SSR for this pixel, so we clear the color
+	// since ScreenSpaceReflectionsPS() won't be executed in this pixel.
+	OutColor = 0;
+}
+
+void ScreenSpaceReflectionsPS(in float4 UVAndScreenPos : TEXCOORD0, out float4 OutColor : SV_Target0)
+{
+	float2 UV = UVAndScreenPos.xy;
+	float2 ScreenPos = UVAndScreenPos.zw;
+	int2 PixelPos = int2(UVAndScreenPos.zw * ScreenPosToPixel.xy + ScreenPosToPixel.zw + 0.5f);
+
+	OutColor = 0;
+
+	FScreenSpaceData ScreenSpaceData = GetScreenSpaceData(UV);
+	FGBufferData GBuffer = ScreenSpaceData.GBuffer;
+	
+	float Roughness = GetRoughness(GBuffer);
+	float RoughnessFade = GetRoughnessFade(Roughness);
+
+	// Early out. Useless if using the stencil prepass.
+	BRANCH if( RoughnessFade <= 0.0 || GBuffer.ShadingModelID == 0 )
 	{
 		return;
 	}
-#endif
 
 #if SSR_QUALITY == 0
 	// visualize SSR
diff --git a/Engine/Source/Runtime/Renderer/Private/PostProcess/SceneRenderTargets.cpp b/Engine/Source/Runtime/Renderer/Private/PostProcess/SceneRenderTargets.cpp
--- a/Engine/Source/Runtime/Renderer/Private/PostProcess/SceneRenderTargets.cpp
+++ b/Engine/Source/Runtime/Renderer/Private/PostProcess/SceneRenderTargets.cpp
@@ -1065,6 +1065,9 @@ void FSceneRenderTargets::BeginRenderingTranslucency(FRHICommandList& RHICmdList
 	// Use the scene color buffer.
 	GSceneRenderTargets.BeginRenderingSceneColor(RHICmdList, ESimpleRenderTargetMode::EExistingColorAndDepth, FExclusiveDepthStencil::DepthRead_StencilWrite);
 	
+	// Clear the stencil buffer for ResponsiveAA
+	RHICmdList.Clear(false, FLinearColor::White, false, (float)ERHIZBuffer::FarPlane, true, 0, FIntRect());
+		
 	// viewport to match view size
 	RHICmdList.SetViewport(View.ViewRect.Min.X, View.ViewRect.Min.Y, 0.0f, View.ViewRect.Max.X, View.ViewRect.Max.Y, 1.0f);
 }
@@ -1098,7 +1101,9 @@ bool FSceneRenderTargets::BeginRenderingSeparateTranslucency(FRHICommandList& RH
 		if (!bFirstTimeThisFrame)
 		{
 			FLinearColor ClearColor(0.0f, 0.0f, 0.0f, 1.0f);
-			RHICmdList.BindClearMRTValues(true, 1, &ClearColor, false, (float)ERHIZBuffer::FarPlane, false, 0);
+			
+			// Clear the stencil buffer for ResponsiveAA
+			RHICmdList.BindClearMRTValues(true, 1, &ClearColor, false, (float)ERHIZBuffer::FarPlane, true, 0);
 		}
 
 		RHICmdList.SetViewport(View.ViewRect.Min.X, View.ViewRect.Min.Y, 0.0f, View.ViewRect.Max.X, View.ViewRect.Max.Y, 1.0f);
diff --git a/Engine/Source/Runtime/Renderer/Private/PostProcess/ScreenSpaceReflections.cpp b/Engine/Source/Runtime/Renderer/Private/PostProcess/ScreenSpaceReflections.cpp
--- a/Engine/Source/Runtime/Renderer/Private/PostProcess/ScreenSpaceReflections.cpp
+++ b/Engine/Source/Runtime/Renderer/Private/PostProcess/ScreenSpaceReflections.cpp
@@ -33,6 +33,13 @@ static TAutoConsoleVariable<int32> CVarSSRTemporal(
 	TEXT(" 0 is off (for debugging), 1 is on (default)"),
 	ECVF_RenderThreadSafe);
 
+static TAutoConsoleVariable<int32> CVarSSRStencil(
+	TEXT("r.SSR.Stencil"),
+	0,
+	TEXT("Defines if we use the stencil prepass for the screen space reflection\n")
+	TEXT(" 0 is off (default), 1 is on"),
+	ECVF_RenderThreadSafe);
+
 bool DoScreenSpaceReflections(const FViewInfo& View)
 {
 	if(!View.Family->EngineShowFlags.ScreenSpaceReflections)
@@ -61,6 +68,88 @@ bool DoScreenSpaceReflections(const FViewInfo& View)
 	return true;
 }
 
+static float ComputeRoughnessMaskScale(const FRenderingCompositePassContext& Context, uint32 SSRQuality)
+{
+	float MaxRoughness = FMath::Clamp(Context.View.FinalPostProcessSettings.ScreenSpaceReflectionMaxRoughness, 0.01f, 1.0f);
+
+	// f(x) = x * Scale + Bias
+	// f(MaxRoughness) = 0
+	// f(MaxRoughness/2) = 1
+
+	float RoughnessMaskScale = -2.0f / MaxRoughness;
+	return RoughnessMaskScale * (SSRQuality < 3 ? 2.0f : 1.0f);
+}
+
+/**
+ * Encapsulates the post processing screen space reflections pixel shader stencil pass.
+ */
+class FPostProcessScreenSpaceReflectionsStencilPS : public FGlobalShader
+{
+	DECLARE_SHADER_TYPE(FPostProcessScreenSpaceReflectionsStencilPS, Global);
+
+	static bool ShouldCache(EShaderPlatform Platform)
+	{
+		return IsFeatureLevelSupported(Platform, ERHIFeatureLevel::SM4);
+	}
+
+	static void ModifyCompilationEnvironment(EShaderPlatform Platform, FShaderCompilerEnvironment& OutEnvironment)
+	{
+		FGlobalShader::ModifyCompilationEnvironment(Platform, OutEnvironment);
+		OutEnvironment.SetDefine( TEXT("PREV_FRAME_COLOR"), uint32(0) );
+		OutEnvironment.SetDefine( TEXT("SSR_QUALITY"), uint32(0) );
+	}
+
+	/** Default constructor. */
+	FPostProcessScreenSpaceReflectionsStencilPS() {}
+
+public:
+	FPostProcessPassParameters PostprocessParameter;
+	FDeferredPixelShaderParameters DeferredParameters;
+	FShaderParameter SSRParams;
+
+	/** Initialization constructor. */
+	FPostProcessScreenSpaceReflectionsStencilPS(const ShaderMetaType::CompiledShaderInitializerType& Initializer)
+		: FGlobalShader(Initializer)
+	{
+		PostprocessParameter.Bind(Initializer.ParameterMap);
+		DeferredParameters.Bind(Initializer.ParameterMap);
+		SSRParams.Bind(Initializer.ParameterMap, TEXT("SSRParams"));
+	}
+
+	void SetParameters(const FRenderingCompositePassContext& Context, uint32 SSRQuality, bool EnableDiscard)
+	{
+		const FFinalPostProcessSettings& Settings = Context.View.FinalPostProcessSettings;
+		const FPixelShaderRHIParamRef ShaderRHI = GetPixelShader();
+
+		FGlobalShader::SetParameters(Context.RHICmdList, ShaderRHI, Context.View);
+
+		PostprocessParameter.SetPS(ShaderRHI, Context, TStaticSamplerState<SF_Point, AM_Clamp, AM_Clamp, AM_Clamp>::GetRHI());
+		DeferredParameters.Set(Context.RHICmdList, ShaderRHI, Context.View);
+
+		{
+			float RoughnessMaskScale = ComputeRoughnessMaskScale(Context, SSRQuality);
+
+			FLinearColor Value(
+				FMath::Clamp(Context.View.FinalPostProcessSettings.ScreenSpaceReflectionIntensity * 0.01f, 0.0f, 1.0f), 
+				RoughnessMaskScale,
+				float(EnableDiscard), //TODO
+				0);
+
+			SetShaderValue(Context.RHICmdList, ShaderRHI, SSRParams, Value);
+		}
+	}
+
+	// FShader interface.
+	virtual bool Serialize(FArchive& Ar) override
+	{
+		bool bShaderHasOutdatedParameters = FGlobalShader::Serialize(Ar);
+		Ar << PostprocessParameter << DeferredParameters << SSRParams;
+		return bShaderHasOutdatedParameters;
+	}
+};
+
+IMPLEMENT_SHADER_TYPE(,FPostProcessScreenSpaceReflectionsStencilPS,TEXT("ScreenSpaceReflections"),TEXT("ScreenSpaceReflectionsStencilPS"),SF_Pixel);
+
 /**
  * Encapsulates the post processing screen space reflections pixel shader.
  * @param SSRQuality 0:Visualize Mask
@@ -112,14 +201,7 @@ public:
 		DeferredParameters.Set(Context.RHICmdList, ShaderRHI, Context.View);
 
 		{
-			float MaxRoughness = FMath::Clamp(Context.View.FinalPostProcessSettings.ScreenSpaceReflectionMaxRoughness, 0.01f, 1.0f);
-
-			// f(x) = x * Scale + Bias
-			// f(MaxRoughness) = 0
-			// f(MaxRoughness/2) = 1
-
-			float RoughnessMaskScale = -2.0f / MaxRoughness;
-			RoughnessMaskScale *= SSRQuality < 3 ? 2.0f : 1.0f;
+			float RoughnessMaskScale = ComputeRoughnessMaskScale(Context, SSRQuality);
 
 			FLinearColor Value(
 				FMath::Clamp(Context.View.FinalPostProcessSettings.ScreenSpaceReflectionIntensity * 0.01f, 0.0f, 1.0f), 
@@ -195,74 +277,133 @@ static int32 ComputeSSRQuality(float Quality)
 
 void FRCPassPostProcessScreenSpaceReflections::Process(FRenderingCompositePassContext& Context)
 {
-	SCOPED_DRAW_EVENT(Context.RHICmdList, ScreenSpaceReflections);
+	auto& RHICmdList = Context.RHICmdList;
 
 	const FSceneView& View = Context.View;
 	const auto FeatureLevel = Context.GetFeatureLevel();
-	const FSceneRenderTargetItem& DestRenderTarget = PassOutputs[0].RequestSurface(Context);
 
-	// Set the view family's render target/viewport.
-	SetRenderTarget(Context.RHICmdList, DestRenderTarget.TargetableTexture, FTextureRHIRef());
-	Context.RHICmdList.Clear(true, FLinearColor(0, 0, 0, 0), false, 1.0f, false, 0, FIntRect());
-	Context.SetViewportAndCallRHI(View.ViewRect);
-
-	// set the state
-	Context.RHICmdList.SetBlendState(TStaticBlendState<>::GetRHI());
-	Context.RHICmdList.SetRasterizerState(TStaticRasterizerState<>::GetRHI());
-	Context.RHICmdList.SetDepthStencilState(TStaticDepthStencilState<false, CF_Always>::GetRHI());
-
-	int SSRQuality = ComputeSSRQuality(View.FinalPostProcessSettings.ScreenSpaceReflectionQuality);
+	int32 SSRQuality = ComputeSSRQuality(View.FinalPostProcessSettings.ScreenSpaceReflectionQuality);
+	uint32 iPreFrame = bPrevFrame ? 1 : 0;
 
 	SSRQuality = FMath::Clamp(SSRQuality, 1, 4);
 	
-
-	uint32 iPreFrame = bPrevFrame ? 1 : 0;
-
 	if (View.Family->EngineShowFlags.VisualizeSSR)
 	{
 		iPreFrame = 0;
 		SSRQuality = 0;
 	}
+	
+	// Write 1 to highest bit of stencil to areas that should compute SSR
+	bool SSRStencilPrePass = CVarSSRStencil.GetValueOnRenderThread() != 0 && SSRQuality != 0;
+	
+	const FSceneRenderTargetItem& DestRenderTarget = PassOutputs[0].RequestSurface(Context);
 
-	TShaderMapRef< FPostProcessVS > VertexShader(Context.GetShaderMap());
-
-	#define CASE(A, B) \
-		case (A + 2 * (B + 3 * 0 )): \
-		{ \
-			TShaderMapRef< FPostProcessScreenSpaceReflectionsPS<A, B> > PixelShader(Context.GetShaderMap()); \
-			static FGlobalBoundShaderState BoundShaderState; \
-			SetGlobalBoundShaderState(Context.RHICmdList, FeatureLevel, BoundShaderState, GFilterVertexDeclaration.VertexDeclarationRHI, *VertexShader, *PixelShader); \
-			VertexShader->SetParameters(Context); \
-			PixelShader->SetParameters(Context); \
-		}; \
-		break
+	if (SSRStencilPrePass)
+	{ // ScreenSpaceReflectionsStencil draw event
+		SCOPED_DRAW_EVENT(RHICmdList, ScreenSpaceReflectionsStencil);
+
+		TShaderMapRef< FPostProcessVS > VertexShader(Context.GetShaderMap());
+		TShaderMapRef< FPostProcessScreenSpaceReflectionsStencilPS > PixelShader(Context.GetShaderMap());
+		
+		// bind the dest render target and the depth stencil render target
+		SetRenderTarget(RHICmdList, DestRenderTarget.TargetableTexture, GSceneRenderTargets.GetSceneDepthSurface(), ESimpleRenderTargetMode::EUninitializedColorAndDepth, FExclusiveDepthStencil::DepthRead_StencilWrite);
+		Context.SetViewportAndCallRHI(View.ViewRect);
+		
+		// Clear stencil to 0
+		RHICmdList.Clear(false, FLinearColor::White, false, (float)ERHIZBuffer::FarPlane, true, 0, FIntRect());
+		
+		// bind shader
+		static FGlobalBoundShaderState BoundShaderState;
+		SetGlobalBoundShaderState(Context.RHICmdList, FeatureLevel, BoundShaderState, GFilterVertexDeclaration.VertexDeclarationRHI, *VertexShader, *PixelShader);
+		VertexShader->SetParameters(Context);
+		PixelShader->SetParameters(Context, SSRQuality, true);
+		
+		// Clobers the stencil to pixel that should not compute SSR
+		RHICmdList.SetDepthStencilState(TStaticDepthStencilState<false, CF_Always, true, CF_Always, SO_Replace, SO_Replace, SO_Replace>::GetRHI(), 0x80);
+
+		// Set rasterizer state to solid
+		RHICmdList.SetRasterizerState(TStaticRasterizerState<>::GetRHI());
+
+		// disable blend mode
+		RHICmdList.SetBlendState(TStaticBlendState<>::GetRHI());
+		
+		// Draw a quad mapping scene color to the view's render target to set stencil to set the stencil mask where it needs to be
+		DrawRectangle( 
+			Context.RHICmdList,
+			0, 0,
+			View.ViewRect.Width(), View.ViewRect.Height(),
+			View.ViewRect.Min.X, View.ViewRect.Min.Y, 
+			View.ViewRect.Width(), View.ViewRect.Height(),
+			View.ViewRect.Size(),
+			GSceneRenderTargets.GetBufferSizeXY(),
+			*VertexShader,
+			EDRF_UseTriangleOptimization);
+	} // ScreenSpaceReflectionsStencil draw event
+
+	{ // ScreenSpaceReflections draw event
+		SCOPED_DRAW_EVENT(Context.RHICmdList, ScreenSpaceReflections);
+		
+		if (SSRStencilPrePass)
+		{
+			// set up the stencil test to match 0, meaning FPostProcessScreenSpaceReflectionsStencilPS has been discarded
+			RHICmdList.SetDepthStencilState(TStaticDepthStencilState<false, CF_Always, true, CF_Equal, SO_Keep, SO_Keep, SO_Keep>::GetRHI(), 0);
+		}
+		else
+		{
+			// bind only the dest render target
+			SetRenderTarget(RHICmdList, DestRenderTarget.TargetableTexture, FTextureRHIRef());
+			Context.SetViewportAndCallRHI(View.ViewRect);
 
-	switch (iPreFrame + 2 * (SSRQuality + 3 * 0))
-	{
-		CASE(0,0);
-		CASE(0,1);	CASE(1,1);
-		CASE(0,2);	CASE(1,2);
-		CASE(0,3);	CASE(1,3);
-		CASE(0,4);	CASE(1,4);
-		default:
-			check(!"Missing case in FRCPassPostProcessScreenSpaceReflections");
-	}
-	#undef CASE
-
-
-	// Draw a quad mapping scene color to the view's render target
-	DrawRectangle( 
-		Context.RHICmdList,
-		0, 0,
-		View.ViewRect.Width(), View.ViewRect.Height(),
-		View.ViewRect.Min.X, View.ViewRect.Min.Y, 
-		View.ViewRect.Width(), View.ViewRect.Height(),
-		View.ViewRect.Size(),
-		GSceneRenderTargets.GetBufferSizeXY(),
-		*VertexShader,
-		EDRF_UseTriangleOptimization);
-
-	Context.RHICmdList.CopyToResolveTarget(DestRenderTarget.TargetableTexture, DestRenderTarget.ShaderResourceTexture, false, FResolveParams());
+			RHICmdList.SetDepthStencilState(TStaticDepthStencilState<false, CF_Always>::GetRHI());
+		}
+		
+		// clear DestRenderTarget only outside of the view's rectangle
+		RHICmdList.Clear(true, FLinearColor::Black, false, (float)ERHIZBuffer::FarPlane, false, 0, View.ViewRect);
+		
+		// set the state
+		RHICmdList.SetBlendState(TStaticBlendState<>::GetRHI());
+		RHICmdList.SetRasterizerState(TStaticRasterizerState<>::GetRHI());
+
+		TShaderMapRef< FPostProcessVS > VertexShader(Context.GetShaderMap());
+
+		#define CASE(A, B) \
+			case (A + 2 * (B + 3 * 0 )): \
+			{ \
+				TShaderMapRef< FPostProcessScreenSpaceReflectionsPS<A, B> > PixelShader(Context.GetShaderMap()); \
+				static FGlobalBoundShaderState BoundShaderState; \
+				SetGlobalBoundShaderState(RHICmdList, FeatureLevel, BoundShaderState, GFilterVertexDeclaration.VertexDeclarationRHI, *VertexShader, *PixelShader); \
+				VertexShader->SetParameters(Context); \
+				PixelShader->SetParameters(Context); \
+			}; \
+			break
+
+		switch (iPreFrame + 2 * (SSRQuality + 3 * 0))
+		{
+			CASE(0,0);
+			CASE(0,1);	CASE(1,1);
+			CASE(0,2);	CASE(1,2);
+			CASE(0,3);	CASE(1,3);
+			CASE(0,4);	CASE(1,4);
+			default:
+				check(!"Missing case in FRCPassPostProcessScreenSpaceReflections");
+		}
+		#undef CASE
+
+
+		// Draw a quad mapping scene color to the view's render target
+		DrawRectangle( 
+			RHICmdList,
+			0, 0,
+			View.ViewRect.Width(), View.ViewRect.Height(),
+			View.ViewRect.Min.X, View.ViewRect.Min.Y, 
+			View.ViewRect.Width(), View.ViewRect.Height(),
+			View.ViewRect.Size(),
+			GSceneRenderTargets.GetBufferSizeXY(),
+			*VertexShader,
+			EDRF_UseTriangleOptimization);
+
+		RHICmdList.CopyToResolveTarget(DestRenderTarget.TargetableTexture, DestRenderTarget.ShaderResourceTexture, false, FResolveParams());
+	} // ScreenSpaceReflections
 }
 
 FPooledRenderTargetDesc FRCPassPostProcessScreenSpaceReflections::ComputeOutputDesc(EPassOutputId InPassOutputId) const
