Replaces cylindrical projection by Panini projection

diff --git a/Engine/Shaders/MaterialTemplate.usf b/Engine/Shaders/MaterialTemplate.usf
--- a/Engine/Shaders/MaterialTemplate.usf
+++ b/Engine/Shaders/MaterialTemplate.usf
@@ -22,6 +22,8 @@
 	#include "AtmosphereCommon.usf"
 #endif
 
+#include "PaniniProjection.usf"
+
 #ifndef USE_DITHERED_LOD_TRANSITION
 	#define USE_DITHERED_LOD_TRANSITION 0
 #endif
diff --git a/Engine/Shaders/PaniniProjection.usf b/Engine/Shaders/PaniniProjection.usf
new file mode 100644
--- /dev/null
+++ b/Engine/Shaders/PaniniProjection.usf
@@ -0,0 +1,90 @@
+// Copyright 1998-2015 Epic Games, Inc. All Rights Reserved.
+
+/*=============================================================================
+	PaniniProjection.usf: Panini implementation for the node material function
+	and upsample post processing pass
+=============================================================================*/
+  
+#pragma once
+
+/*
+* Reproject a ray direction into the panini view's plan.
+*	Source: http://tksharpless.net/vedutismo/Pannini/panini.pdf
+*
+* Panini frame:
+*                                        |
+*                                        |   /
+*                                      M | /
+*                                        +
+*                                      / |
+*                        ->          /   |
+*                        z         /     |
+*        - + - - - - - - < - - - + - - - + - - + - - - - - - -
+*          D                     O       |A
+*                                        |
+*                                        |
+*                                (panini view plan)
+*
+*		O: origine of the frame
+*		A: origin of the panini view plane (0.0, 0.0, -1.0)
+*		D: panini projection center (0.0, 0.0, d)
+*		OM: the direction of the ray.
+*
+* @param OM - The direction of the ray in the panini frame. The z component is implicitely -1.0f.
+* @param d - The panini projection center's distance form the panini frame's origin.
+*		0.0f: Panini projection doesn't change anything, it will return PaniniDirection.
+*		1.0f: Panini
+*		@assert d>=0
+* @param s - The panini projection hard vertical compression fade (lerp). Works with any number in the real domain.
+*		0.0f: No vertical compression
+*		1.0f: Hard vertical compression
+* @return - Returns the position (x, y) of the panini projection in the panini view plane at z=-1.0f in the panini frame
+*/
+float2 PaniniProjection(float2 OM, float d, float s)
+{
+	float PaniniDirectionXZInvLength = 1.0f / sqrt(1.0f + OM.x * OM.x);
+	float SinPhi = OM.x * PaniniDirectionXZInvLength;
+	float TanTheta = OM.y * PaniniDirectionXZInvLength;
+	float CosPhi = sqrt(1.0f - SinPhi * SinPhi);
+	float S = (d + 1.0f) / (d + CosPhi);
+
+	return S * float2(SinPhi, lerp(TanTheta, TanTheta / CosPhi, s));
+}
+
+float2 GetHalfFOVPerAxis()
+{
+	// TODO: should be computed on the CPU side
+	float3 VCenter = normalize(mul(float4(0.0, 0.0, 0.0, 1.0), View.ClipToView)).xyz;
+	float3 VUp = normalize(mul(float4(0.0, 1.0, 0.0, 1.0), View.ClipToView)).xyz;
+	float3 VRight = normalize(mul(float4(1.0, 0.0, 0.0, 1.0), View.ClipToView)).xyz;
+
+	return float2(acos(dot(VCenter, VRight)), acos(dot(VCenter, VUp)));
+}
+
+/*
+* Utility to do the projection 
+*
+*/
+float2 ScreenPosPaniniProjection(float2 ScreenPos, float d, float s)
+{
+	float2 HalfFOV = GetHalfFOVPerAxis();
+	float2 ScreenPosToPaniniFactor = float2(tan(HalfFOV.x), tan(HalfFOV.y));
+
+	// Compute the pixel ray direction into the Panini frame
+	// PaniniDirection.z is implicitly 1.0f
+	float2 PaniniDirection = ScreenPos * ScreenPosToPaniniFactor;
+	float2 PaniniPosition = PaniniProjection(PaniniDirection, d, s);
+
+	// Return the new position back in the screen space frame
+	return PaniniPosition / ScreenPosToPaniniFactor;
+}
+
+// Panini material function node returning a not translated world position according to the panini projection.
+float3 CalcPaniniWorldPosition(float3 InWorldPosition, float d, float s, float PostScale)
+{
+	float4 InClipPosition = mul(float4(InWorldPosition + View.PreViewTranslation, 1.0), View.TranslatedWorldToClip);
+	float2 InScreenPos = InClipPosition.xy / InClipPosition.w;
+	float2 OutScreenPos = ScreenPosPaniniProjection(InScreenPos, d, s) * PostScale;
+	float4 OutClipPosition = float4(OutScreenPos.xy * InClipPosition.w, InClipPosition.z, InClipPosition.w);
+	return mul(OutClipPosition, View.ClipToTranslatedWorld).xyz - View.PreViewTranslation;
+}
diff --git a/Engine/Shaders/PostProcessUpscale.usf b/Engine/Shaders/PostProcessUpscale.usf
--- a/Engine/Shaders/PostProcessUpscale.usf
+++ b/Engine/Shaders/PostProcessUpscale.usf
@@ -6,13 +6,14 @@
 
 #include "Common.usf"
 #include "PostProcessCommon.usf"
+#include "PaniniProjection.usf"
 
 // only for MainPS
 float UpscaleSoftness;
 
 // only for MainVS()
 // .x:HalfFOV .y:tan(HalfFov) .z:CylinderDistortion(0=none, 1:full) .w:Correction to scale Y the same as X is scaled in the center
-float4 DistortionParams;
+float3 PaniniParams;
 
 // to avoid compiler errors on the PS
 #ifndef TESS_RECT_X
@@ -36,20 +37,13 @@ float Gaussian(float Scale, float2 Offset)
 
 // @param In -1 .. 1 range
 // @return In -1 .. 1 range
-float2 CylindricalDistortPos(float2 In)
+float2 PaniniDistortPos(float2 In)
 {
-	float HalfFov = DistortionParams.x;
-	float TanHalfFov = DistortionParams.y;
-	float InvHalfFov = 1.0f / HalfFov;
-	float CylinderDistortion = DistortionParams.z;
-	float Correction = DistortionParams.w;
+	const float d = PaniniParams.x;
+	const float s = PaniniParams.y;
+	const float OutScreenPosScale = PaniniParams.z;
 
-	float X = In.x * TanHalfFov;
-	float AngleX = atan(X);
-
-	float2 Ret = float2(AngleX * InvHalfFov, In.y * cos(AngleX) * Correction);
-
-	return lerp(In, Ret, CylinderDistortion);
+	return ScreenPosPaniniProjection(In, d, s) * OutScreenPosScale;
 }
 
 
@@ -71,7 +65,7 @@ void MainVS(
 	float2 PlanarScreenPos = TexCoord * 2 - 1;
 
 	// distort pos
-	Position.xy = CylindricalDistortPos(PlanarScreenPos) * 0.5f + 0.5f;
+	Position.xy = PaniniDistortPos(PlanarScreenPos) * 0.5f + 0.5f;
 
 	DrawRectangle(Position, TexCoord, OutPosition, OutTexCoord.xy);
 	OutTexCoord.zw =  OutPosition.xy;
diff --git a/Engine/Source/Runtime/Engine/Public/SceneView.h b/Engine/Source/Runtime/Engine/Public/SceneView.h
--- a/Engine/Source/Runtime/Engine/Public/SceneView.h
+++ b/Engine/Source/Runtime/Engine/Public/SceneView.h
@@ -226,6 +226,21 @@ struct FViewMatrices
 	{
 		return GetInvProjMatrix() * GetInvViewMatrix();
 	}
+
+	FVector2D GetFieldOfViewPerAxis() const
+	{
+		const FMatrix ClipToView = GetInvProjMatrix();
+
+		FVector VCenter = FVector(ClipToView.TransformPosition(FVector(0.0, 0.0, 0.0)));
+		FVector VUp = FVector(ClipToView.TransformPosition(FVector(0.0, 1.0, 0.0)));
+		FVector VRight = FVector(ClipToView.TransformPosition(FVector(1.0, 0.0, 0.0)));
+
+		VCenter.Normalize();
+		VUp.Normalize();
+		VRight.Normalize();
+
+		return FVector2D(FMath::Acos(VCenter | VRight), FMath::Acos(VCenter | VUp));
+	}
 };
 
 #include "ShaderParameters.h"
diff --git a/Engine/Source/Runtime/Renderer/Private/ForwardShadingRenderer.cpp b/Engine/Source/Runtime/Renderer/Private/ForwardShadingRenderer.cpp
--- a/Engine/Source/Runtime/Renderer/Private/ForwardShadingRenderer.cpp
+++ b/Engine/Source/Runtime/Renderer/Private/ForwardShadingRenderer.cpp
@@ -234,7 +234,7 @@ void FForwardShadingSceneRenderer::BasicPostProcess(FRHICommandListImmediate& RH
 
 	if (bDoUpscale)
 	{	// simple bilinear upscaling for ES2.
-		FRenderingCompositePass* Node = Context.Graph.RegisterPass(new(FMemStack::Get()) FRCPassPostProcessUpscale(1, 0.0f));
+		FRenderingCompositePass* Node = Context.Graph.RegisterPass(new(FMemStack::Get()) FRCPassPostProcessUpscale(1));
 
 		Node->SetInput(ePId_Input0, FRenderingCompositeOutputRef(Context.FinalOutput));
 		Node->SetInput(ePId_Input1, FRenderingCompositeOutputRef(Context.FinalOutput));
diff --git a/Engine/Source/Runtime/Renderer/Private/PostProcess/PostProcessUpscale.cpp b/Engine/Source/Runtime/Renderer/Private/PostProcess/PostProcessUpscale.cpp
--- a/Engine/Source/Runtime/Renderer/Private/PostProcess/PostProcessUpscale.cpp
+++ b/Engine/Source/Runtime/Renderer/Private/PostProcess/PostProcessUpscale.cpp
@@ -20,6 +20,19 @@ static TAutoConsoleVariable<float> CVarUpscaleSoftness(
 	TEXT(" 0..1 (0.3 is good for ScreenPercentage 90"),
 	ECVF_Scalability | ECVF_RenderThreadSafe);
 
+const FRCPassPostProcessUpscale::PaniniParams FRCPassPostProcessUpscale::PaniniParams::Default;
+
+static FVector2D PaniniProjection(FVector2D OM, float d, float s)
+{
+	float PaniniDirectionXZInvLength = 1.0f / FMath::Sqrt(1.0f + OM.X * OM.X);
+	float SinPhi = OM.X * PaniniDirectionXZInvLength;
+	float TanTheta = OM.Y * PaniniDirectionXZInvLength;
+	float CosPhi = FMath::Sqrt(1.0f - SinPhi * SinPhi);
+	float S = (d + 1.0f) / (d + CosPhi);
+
+	return S * FVector2D(SinPhi, FMath::Lerp(TanTheta, TanTheta / CosPhi, s));
+}
+
 /** Encapsulates the upscale vertex shader. */
 class FPostProcessUpscaleVS : public FPostProcessVS
 {
@@ -29,13 +42,13 @@ class FPostProcessUpscaleVS : public FPostProcessVS
 	FPostProcessUpscaleVS() {}
 
 public:
-	FShaderParameter DistortionParams;
+	FShaderParameter PaniniParameters;
 
 	/** Initialization constructor. */
 	FPostProcessUpscaleVS(const ShaderMetaType::CompiledShaderInitializerType& Initializer)
 		: FPostProcessVS(Initializer)
 	{
-		DistortionParams.Bind(Initializer.ParameterMap,TEXT("DistortionParams"));
+		PaniniParameters.Bind(Initializer.ParameterMap,TEXT("PaniniParams"));
 	}
 
 	static bool ShouldCache(EShaderPlatform Platform)
@@ -50,32 +63,31 @@ public:
 		OutEnvironment.SetDefine(TEXT("TESS_RECT_Y"), FTesselatedScreenRectangleIndexBuffer::Height);
 	}
 
-	// @param InCylinderDistortion 0=none..1=full in percent, must be in that range
-	void SetParameters(const FRenderingCompositePassContext& Context, float InCylinderDistortion)
+	void SetParameters(const FRenderingCompositePassContext& Context, const FRCPassPostProcessUpscale::PaniniParams& InPaniniConfig)
 	{
 		const FVertexShaderRHIParamRef ShaderRHI = GetVertexShader();
 
 		FGlobalShader::SetParameters(Context.RHICmdList, ShaderRHI, Context.View);
 
 		{
-			float HalfFOV = FMath::Atan(1.0f / Context.View.ViewMatrices.ProjMatrix.M[0][0]);
-			float TanHalfFov = 1.0f / Context.View.ViewMatrices.ProjMatrix.M[0][0];
-			float InvHalfFov = 1.0f / HalfFOV;
+			const FVector2D FOVPerAxis = Context.View.ViewMatrices.GetFieldOfViewPerAxis();
+			const FVector2D ScreenPosToPaniniFactor = FVector2D(FMath::Tan(FOVPerAxis.X), FMath::Tan(FOVPerAxis.Y));
+			const FVector2D PaniniDirection = FVector2D(1.0f, 0.0f) * ScreenPosToPaniniFactor;
+			const FVector2D PaniniPosition = PaniniProjection(PaniniDirection, InPaniniConfig.D, InPaniniConfig.S);
 
-			// compute Correction to scale Y the same as X is scaled in the center
-			float SmallX = 0.01f;
-			float Correction = atan(SmallX * TanHalfFov) * InvHalfFov / SmallX;
+			const float WidthFit = ScreenPosToPaniniFactor.X / PaniniPosition.X;
+			const float OutScreenPosScale = FMath::Lerp(1.0f, WidthFit, InPaniniConfig.ScreenFit);
 
-			FVector4 Value(HalfFOV, TanHalfFov, InCylinderDistortion, Correction);
+			FVector Value(InPaniniConfig.D, InPaniniConfig.S, OutScreenPosScale);
 
-			SetShaderValue(Context.RHICmdList, ShaderRHI, DistortionParams, Value);
+			SetShaderValue(Context.RHICmdList, ShaderRHI, PaniniParameters, Value);
 		}
 	}
 
 	virtual bool Serialize(FArchive& Ar) override
 	{
 		bool bShaderHasOutdatedParameters = FPostProcessVS::Serialize(Ar);
-		Ar << DistortionParams;
+		Ar << PaniniParameters;
 		return bShaderHasOutdatedParameters;
 	}
 };
@@ -173,18 +185,20 @@ VARIATION1(3)
 
 #undef VARIATION1
 
-FRCPassPostProcessUpscale::FRCPassPostProcessUpscale(uint32 InUpscaleQuality, float InCylinderDistortion)
+FRCPassPostProcessUpscale::FRCPassPostProcessUpscale(uint32 InUpscaleQuality, const PaniniParams& InPaniniConfig)
 	: UpscaleQuality(InUpscaleQuality)
-	, CylinderDistortion(FMath::Clamp(InCylinderDistortion, 0.0f, 1.0f))
 {
+	PaniniConfig.D = FMath::Max(InPaniniConfig.D, 0.0f);
+	PaniniConfig.S = InPaniniConfig.S;
+	PaniniConfig.ScreenFit = FMath::Max(InPaniniConfig.ScreenFit, 0.0f);
 }
 
 template <uint32 Method, uint32 bTesselatedQuad>
-FShader* FRCPassPostProcessUpscale::SetShader(const FRenderingCompositePassContext& Context, float InCylinderDistortion)
+FShader* FRCPassPostProcessUpscale::SetShader(const FRenderingCompositePassContext& Context, const PaniniParams& PaniniConfig)
 {
 	if(bTesselatedQuad)
 	{
-		check(InCylinderDistortion > 0.0f);
+		check(PaniniConfig.D > 0.0f);
 
 		TShaderMapRef<FPostProcessUpscaleVS> VertexShader(Context.GetShaderMap());
 		TShaderMapRef<FPostProcessUpscalePS<Method> > PixelShader(Context.GetShaderMap());
@@ -194,12 +208,12 @@ FShader* FRCPassPostProcessUpscale::SetShader(const FRenderingCompositePassConte
 		SetGlobalBoundShaderState(Context.RHICmdList, Context.GetFeatureLevel(), BoundShaderState, GFilterVertexDeclaration.VertexDeclarationRHI, *VertexShader, *PixelShader);
 
 		PixelShader->SetPS(Context);
-		VertexShader->SetParameters(Context, InCylinderDistortion);
+		VertexShader->SetParameters(Context, PaniniConfig);
 		return *VertexShader;
 	}
 	else
 	{
-		check(InCylinderDistortion == 0.0f);
+		check(PaniniConfig.D == 0.0f);
 
 		TShaderMapRef<FPostProcessVS> VertexShader(Context.GetShaderMap());
 		TShaderMapRef<FPostProcessUpscalePS<Method> > PixelShader(Context.GetShaderMap());
@@ -238,7 +252,7 @@ void FRCPassPostProcessUpscale::Process(FRenderingCompositePassContext& Context)
 	SetRenderTarget(Context.RHICmdList, DestRenderTarget.TargetableTexture, FTextureRHIRef());
 	Context.SetViewportAndCallRHI(DestRect);
 
-	bool bTessellatedQuad = CylinderDistortion >= 0.01f;
+	bool bTessellatedQuad = PaniniConfig.D >= 0.01f;
 
 	// with distortion (bTessellatedQuad) we need to clear the background
 	FIntRect ExcludeRect = bTessellatedQuad ? FIntRect() : View.ViewRect;
@@ -256,10 +270,10 @@ void FRCPassPostProcessUpscale::Process(FRenderingCompositePassContext& Context)
 	{
 		switch (UpscaleQuality)
 		{
-			case 0:	VertexShader = SetShader<0, 1>(Context, CylinderDistortion); break;
-			case 1:	VertexShader = SetShader<1, 1>(Context, CylinderDistortion); break;
-			case 2:	VertexShader = SetShader<2, 1>(Context, CylinderDistortion); break;
-			case 3:	VertexShader = SetShader<3, 1>(Context, CylinderDistortion); break;
+			case 0:	VertexShader = SetShader<0, 1>(Context, PaniniConfig); break;
+			case 1:	VertexShader = SetShader<1, 1>(Context, PaniniConfig); break;
+			case 2:	VertexShader = SetShader<2, 1>(Context, PaniniConfig); break;
+			case 3:	VertexShader = SetShader<3, 1>(Context, PaniniConfig); break;
 			default:
 				checkNoEntry();
 				break;
@@ -269,10 +283,10 @@ void FRCPassPostProcessUpscale::Process(FRenderingCompositePassContext& Context)
 	{
 		switch (UpscaleQuality)
 		{
-			case 0:	VertexShader = SetShader<0, 0>(Context); break;
-			case 1:	VertexShader = SetShader<1, 0>(Context); break;
-			case 2:	VertexShader = SetShader<2, 0>(Context); break;
-			case 3:	VertexShader = SetShader<3, 0>(Context); break;
+			case 0:	VertexShader = SetShader<0, 0>(Context, PaniniParams::Default); break;
+			case 1:	VertexShader = SetShader<1, 0>(Context, PaniniParams::Default); break;
+			case 2:	VertexShader = SetShader<2, 0>(Context, PaniniParams::Default); break;
+			case 3:	VertexShader = SetShader<3, 0>(Context, PaniniParams::Default); break;
 			default:
 				checkNoEntry();
 				break;
diff --git a/Engine/Source/Runtime/Renderer/Private/PostProcess/PostProcessUpscale.h b/Engine/Source/Runtime/Renderer/Private/PostProcess/PostProcessUpscale.h
--- a/Engine/Source/Runtime/Renderer/Private/PostProcess/PostProcessUpscale.h
+++ b/Engine/Source/Runtime/Renderer/Private/PostProcess/PostProcessUpscale.h
@@ -14,14 +14,36 @@
 class FRCPassPostProcessUpscale : public TRenderingCompositePassBase<2, 1>
 {
 public:
+	// Panini configuration
+	//  NOTE: more details in Common.usf
+	struct PaniniParams
+	{
+		// 0=none..1=full, must be >= 0
+		float D;
+
+		// Panini hard vertical compression lerp (0=no vertical compresion, 1=hard compression)
+		float S;
+
+		// Panini screen fit factor (lerp between vertical and horizontal) 
+		float ScreenFit;
+
+		PaniniParams()
+			: D(0.0f)
+			, S(0.0f)
+			, ScreenFit(1.0f)
+		{}
+
+		static const PaniniParams Default;
+	};
+
 	// constructor
 	// @param InUpscaleQuality - value denoting Upscale method to use:
 	//				0: Nearest
 	//				1: Bilinear
 	//				2: 4 tap Bilinear (with radius adjustment)
 	//				3: Directional blur with unsharp mask upsample.
-	// @param InCylinderDistortion 0=none..1=full in percent, can be outside of that range
-	FRCPassPostProcessUpscale(uint32 InUpscaleQuality, float InCylinderDistortion);
+	// @param InPaniniConfig - the panini configuration parameter
+	FRCPassPostProcessUpscale(uint32 InUpscaleQuality, const PaniniParams& InPaniniConfig = PaniniParams::Default);
 
 	// interface FRenderingCompositePass ---------
 	virtual void Process(FRenderingCompositePassContext& Context) override;
@@ -29,11 +51,12 @@ public:
 	FPooledRenderTargetDesc ComputeOutputDesc(EPassOutputId InPassOutputId) const override;
 private:
 	// @param InCylinderDistortion 0=none..1=full in percent, must be in that range
-	template <uint32 Quality, uint32 bTesselatedQuad> static FShader* SetShader(const FRenderingCompositePassContext& Context, float InCylinderDistortion = 0.0f);
+	template <uint32 Quality, uint32 bTesselatedQuad> static FShader* SetShader(const FRenderingCompositePassContext& Context, const PaniniParams& PaniniConfig);
 
 	// 0: Nearest, 1: Bilinear, 2: 4 tap Bilinear (with radius adjustment), 3: Directional blur with unsharp mask upsample.
 	uint32 UpscaleQuality;
-	// 0=none..1=full, must be in that range
-	float CylinderDistortion;
+
+	// Panini projection's parameter
+	PaniniParams PaniniConfig;
 };
 
diff --git a/Engine/Source/Runtime/Renderer/Private/PostProcess/PostProcessing.cpp b/Engine/Source/Runtime/Renderer/Private/PostProcess/PostProcessing.cpp
--- a/Engine/Source/Runtime/Renderer/Private/PostProcess/PostProcessing.cpp
+++ b/Engine/Source/Runtime/Renderer/Private/PostProcess/PostProcessing.cpp
@@ -71,14 +71,30 @@ static TAutoConsoleVariable<int32> CVarRenderTargetSwitchWorkaround(
 	TEXT("We want this enabled (1) on all 32 bit iOS devices (implemented through DeviceProfiles)."),
 	ECVF_RenderThreadSafe);
 
-static TAutoConsoleVariable<float> CVarUpscaleCylinder(
-	TEXT("r.Upscale.Cylinder"),
+static TAutoConsoleVariable<float> CVarUpscalePaniniD(
+	TEXT("r.Upscale.Panini.D"),
 	0,
-	TEXT("Allows to apply a cylindrical distortion to the rendered image. Values between 0 and 1 allow to fade the effect (lerp).\n")
-	TEXT("There is a quality loss that can be compensated by adjusting r.ScreenPercentage (>100).\n")
-	TEXT("0: off(default)\n")
+	TEXT("Allow and configure to apply a panini distortion to the rendered image. Values between 0 and 1 allow to fade the effect (lerp).\n")
+	TEXT("Implementation from research paper \"Pannini: A New Projection for Rendering Wide Angle Perspective Images\"\n")
+	TEXT("0: off (default)\n")
 	TEXT(">0: enabled (requires an extra post processing pass if upsampling wasn't used - see r.ScreenPercentage)\n")
-	TEXT("1: full effect"),
+	TEXT("1: Panini cylindrical stereographic projection"),
+	ECVF_RenderThreadSafe);
+
+static TAutoConsoleVariable<float> CVarUpscalePaniniS(
+	TEXT("r.Upscale.Panini.S"),
+	0,
+	TEXT("Panini projection's hard vertical compression factor.\n")
+	TEXT("0: no vertical compression factor (default)\n")
+	TEXT("1: Hard vertical compression"),
+	ECVF_RenderThreadSafe);
+
+static TAutoConsoleVariable<float> CVarUpscalePaniniScreenFit(
+	TEXT("r.Upscale.Panini.ScreenFit"),
+	1.0f,
+	TEXT("Panini projection screen fit effect factor (lerp).\n")
+	TEXT("0: fit vertically\n")
+	TEXT("1: fit horizontally (default)"),
 	ECVF_RenderThreadSafe);
 
 static TAutoConsoleVariable<int32> CVarUpscaleQuality(
@@ -1558,20 +1574,22 @@ void FPostProcessing::Process(FRHICommandListImmediate& RHICmdList, FViewInfo& V
 		AddHighResScreenshotMask(Context, SeparateTranslucency);
 
 		// 0=none..1=full
-		float Cylinder = 0.0f;
+		FRCPassPostProcessUpscale::PaniniParams PaniniConfig;
 
 		if(View.IsPerspectiveProjection() && !GEngine->StereoRenderingDevice.IsValid())
 		{
-			Cylinder = FMath::Clamp(CVarUpscaleCylinder.GetValueOnRenderThread(), 0.0f, 1.0f);
+			PaniniConfig.D = FMath::Max(CVarUpscalePaniniD.GetValueOnRenderThread(), 0.0f);
+			PaniniConfig.S = CVarUpscalePaniniS.GetValueOnRenderThread();
+			PaniniConfig.ScreenFit = FMath::Max(CVarUpscalePaniniScreenFit.GetValueOnRenderThread(), 0.0f);
 		}
 
 		// Do not use upscale if SeparateRenderTarget is in use!
-		if ((Cylinder > 0.01f || View.UnscaledViewRect != View.ViewRect) && 
+		if ((PaniniConfig.D > 0.01f || View.UnscaledViewRect != View.ViewRect) && 
 			(bHMDWantsUpscale ||!View.Family->EngineShowFlags.StereoRendering || (!View.Family->EngineShowFlags.HMDDistortion && !View.Family->bUseSeparateRenderTarget)))
 		{
 			int32 UpscaleQuality = CVarUpscaleQuality.GetValueOnRenderThread();
 			UpscaleQuality = FMath::Clamp(UpscaleQuality, 0, 3);
-			FRenderingCompositePass* Node = Context.Graph.RegisterPass(new(FMemStack::Get()) FRCPassPostProcessUpscale(UpscaleQuality, Cylinder));
+			FRenderingCompositePass* Node = Context.Graph.RegisterPass(new(FMemStack::Get()) FRCPassPostProcessUpscale(UpscaleQuality, PaniniConfig));
 			Node->SetInput(ePId_Input0, FRenderingCompositeOutputRef(Context.FinalOutput)); // Bilinear sampling.
 			Node->SetInput(ePId_Input1, FRenderingCompositeOutputRef(Context.FinalOutput)); // Point sampling.
 			Context.FinalOutput = FRenderingCompositeOutputRef(Node);
