Fixed max velocity being scattered too far.

diff --git a/Engine/Shaders/PostProcessMotionBlur.usf b/Engine/Shaders/PostProcessMotionBlur.usf
--- a/Engine/Shaders/PostProcessMotionBlur.usf
+++ b/Engine/Shaders/PostProcessMotionBlur.usf
@@ -697,9 +697,11 @@ void VelocityScatterVS(
 
 	float2 Velocity = PostprocessInput0[ TilePos ].xy;
 	OutColor = Velocity.xyxy;
+
+	Velocity *= VelocityScale.xy;
 	
 	BRANCH
-	if( length( Velocity * 0.5 * float2(1920,1080) ) <= 0.5 )
+	if( length( Velocity * float2(1920,1080) ) <= 0.5 )
 	{
 		return;
 	}
diff --git a/Engine/Source/Runtime/Renderer/Private/PostProcess/PostProcessMotionBlur.cpp b/Engine/Source/Runtime/Renderer/Private/PostProcess/PostProcessMotionBlur.cpp
--- a/Engine/Source/Runtime/Renderer/Private/PostProcess/PostProcessMotionBlur.cpp
+++ b/Engine/Source/Runtime/Renderer/Private/PostProcess/PostProcessMotionBlur.cpp
@@ -953,6 +953,7 @@ class FPostProcessVelocityScatterVS : public FGlobalShader
 public:
 	FPostProcessPassParameters PostprocessParameter;
 	FShaderParameter TileCount;
+	FShaderParameter VelocityScale;
 
 	/** Initialization constructor. */
 	FPostProcessVelocityScatterVS(const ShaderMetaType::CompiledShaderInitializerType& Initializer)
@@ -960,13 +961,14 @@ public:
 	{
 		PostprocessParameter.Bind(Initializer.ParameterMap);
 		TileCount.Bind(Initializer.ParameterMap, TEXT("TileCount"));
+		VelocityScale.Bind( Initializer.ParameterMap, TEXT("VelocityScale") );
 	}
 
 	// FShader interface.
 	virtual bool Serialize(FArchive& Ar) override
 	{
 		bool bShaderHasOutdatedParameters = FGlobalShader::Serialize(Ar);
-		Ar << PostprocessParameter << TileCount;
+		Ar << PostprocessParameter << TileCount << VelocityScale;
 		return bShaderHasOutdatedParameters;
 	}
 
@@ -979,6 +981,14 @@ public:
 		PostprocessParameter.SetVS(ShaderRHI, Context, TStaticSamplerState<SF_Point,AM_Clamp,AM_Clamp,AM_Clamp>::GetRHI());
 
 		SetShaderValue(Context.RHICmdList, ShaderRHI, TileCount, TileCountValue);
+
+		{
+			const FSceneViewState* ViewState = (FSceneViewState*) Context.View.State;
+			const float MotionBlurTimeScale = ViewState ? ViewState->MotionBlurTimeScale : 1.0f;
+
+			const float ViewMotionBlurScale = 0.5f * MotionBlurTimeScale * Context.View.FinalPostProcessSettings.MotionBlurAmount;
+			SetShaderValue(Context.RHICmdList, ShaderRHI, VelocityScale, FVector4(ViewMotionBlurScale, ViewMotionBlurScale, 0, 0));
+		}
 	}
 
 	static const TCHAR* GetSourceFilename()
