Added support for stencil read/writes to CustomDepth pass (requires r.CustomDepth == 3).

diff --git a/Engine/Config/BaseEngine.ini b/Engine/Config/BaseEngine.ini
--- a/Engine/Config/BaseEngine.ini
+++ b/Engine/Config/BaseEngine.ini
@@ -2067,6 +2067,7 @@ EnableNiagara=false
 [Engine.BufferVisualizationMaterials]
 BaseColor=(Material="/Engine/BufferVisualization/BaseColor.BaseColor", Name=LOCTEXT("BaseColorMat", "Base Color"))
 CustomDepth=(Material="/Engine/BufferVisualization/CustomDepth.CustomDepth", Name=LOCTEXT("BaseCustomDepthMat", "Custom Depth"))
+CustomStencil=(Material="/Engine/BufferVisualization/CustomStencil.CustomStencil", Name=LOCTEXT("BaseCustomStencilMat", "Custom Stencil"))
 DecalMask=(Material="/Engine/BufferVisualization/DecalMask.DecalMask", Name=LOCTEXT("BaseDecalMaskMat", "Decal Mask"))
 FinalImage=(Material="/Engine/BufferVisualization/FinalImage.FinalImage", Name=LOCTEXT("BaseFinalImageMat", "Final Image"))
 ShadingModel=(Material="/Engine/BufferVisualization/LightingModel.LightingModel", Name=LOCTEXT("BaseShadingModelMat", "Shading Model"))
diff --git a/Engine/Documentation/Source/Shared/Types/UPrimitiveComponent/UPrimitiveComponent.INT.udn b/Engine/Documentation/Source/Shared/Types/UPrimitiveComponent/UPrimitiveComponent.INT.udn
--- a/Engine/Documentation/Source/Shared/Types/UPrimitiveComponent/UPrimitiveComponent.INT.udn
+++ b/Engine/Documentation/Source/Shared/Types/UPrimitiveComponent/UPrimitiveComponent.INT.udn
@@ -131,6 +131,10 @@ Engine/Physics/PhysicsBodies#collision
 [EXCERPT:bRenderCustomDepth]
 [/EXCERPT:bRenderCustomDepth]
 
+<!-- Optionally write this 0-255 value to the stencil buffer in CustomDepth pass (Requires project setting or r.CustomDepth == 3)  -->
+[EXCERPT:CustomDepthStencilValue]
+[/EXCERPT:CustomDepthStencilValue]
+
 
 <!-- If true, this component will be rendered in the main pass (z prepass, basepass, transparency) -->
 [EXCERPT:bRenderInMainPass]
diff --git a/Engine/Shaders/DeferredShadingCommon.usf b/Engine/Shaders/DeferredShadingCommon.usf
--- a/Engine/Shaders/DeferredShadingCommon.usf
+++ b/Engine/Shaders/DeferredShadingCommon.usf
@@ -282,6 +282,8 @@ struct FGBufferData
 	uint HasHeightfieldRepresentation; 
 	// in world units
 	float CustomDepth;
+	// Custom depth stencil value
+	uint CustomStencil;
 	// in unreal units (linear), can be used to reconstruct world position,
 	// only valid when decoding the GBuffer as the value gets reconstructed from the Z buffer
 	float Depth;
@@ -381,6 +383,7 @@ FGBufferData DecodeGBufferData(
 	float4 InGBufferE,
 	float4 InGBufferVelocity,
 	float CustomNativeDepth,
+	uint CustomStencil,
 	float SceneDepth,
 	bool bGetNormalizedNormal)
 {
@@ -415,6 +418,7 @@ FGBufferData DecodeGBufferData(
 #endif
 	GBuffer.PrecomputedShadowFactors = InGBufferE;
 	GBuffer.CustomDepth = ConvertFromDeviceZ(CustomNativeDepth);
+	GBuffer.CustomStencil = CustomStencil;
 	GBuffer.Depth = SceneDepth;
 
 	// derived from BaseColor, Metalness, Specular
@@ -528,6 +532,10 @@ SamplerState ScreenSpaceAOTextureSampler;
 Texture2D CustomDepthTexture;
 SamplerState CustomDepthTextureSampler;
 
+#if FEATURE_LEVEL >= FEATURE_LEVEL_SM4 
+Texture2D<uint2> CustomStencilTexture;
+#endif
+
 Texture2D		DirectionalOcclusionTexture;
 SamplerState	DirectionalOcclusionSampler;
 
@@ -553,6 +561,7 @@ SamplerState	DirectionalOcclusionSampler;
 		float4 GBufferC = GBuffers.GBufferCTextureNonMS.Load(int3(PixelPos, 0));
 		float4 GBufferD = GBuffers.GBufferDTextureNonMS.Load(int3(PixelPos, 0));
 		float CustomNativeDepth = CustomDepthTextureNonMS.Load(int3(PixelPos, 0)).r;
+		uint CustomStencil = CustomStencilTexture.Load(int3(PixelPos, 0)).g;
 
 		#if ALLOW_STATIC_LIGHTING
 			float4 GBufferE = GBuffers.GBufferETextureNonMS.Load(int3(PixelPos, 0));
@@ -568,7 +577,7 @@ SamplerState	DirectionalOcclusionSampler;
 
 		float SceneDepth = CalcSceneDepth(PixelPos);
 
-		return DecodeGBufferData(GBufferA, GBufferB, GBufferC, GBufferD, GBufferE, GBufferVelocity, CustomNativeDepth, SceneDepth, bGetNormalizedNormal);
+		return DecodeGBufferData(GBufferA, GBufferB, GBufferC, GBufferD, GBufferE, GBufferVelocity, CustomNativeDepth, CustomStencil, SceneDepth, bGetNormalizedNormal);
 	}
 
 	FDBufferData GetDBufferData(uint2 PixelPos)
@@ -637,6 +646,16 @@ FGBufferData GetGBufferData(float2 UV, bool bGetNormalizedNormal = true)
 	float4 GBufferD = Texture2DSampleLevel(GBuffers.GBufferDTexture, GBuffers.GBufferDTextureSampler, UV, 0);
 	float CustomNativeDepth = Texture2DSampleLevel(CustomDepthTexture, CustomDepthTextureSampler, UV, 0).r;
 
+#if FEATURE_LEVEL >= FEATURE_LEVEL_SM4
+	float2 SurfaceDimensions;
+	CustomStencilTexture.GetDimensions(SurfaceDimensions.x, SurfaceDimensions.y);
+
+	int2 IntUV = (int2)trunc(UV * SurfaceDimensions);
+	uint CustomStencil = CustomStencilTexture.Load(int3(IntUV, 0)).g;
+#else
+	uint CustomStencil = 0;
+#endif
+
 	#if ALLOW_STATIC_LIGHTING
 		float4 GBufferE = Texture2DSampleLevel(GBuffers.GBufferETexture, GBuffers.GBufferETextureSampler, UV, 0);
 	#else
@@ -652,7 +671,7 @@ FGBufferData GetGBufferData(float2 UV, bool bGetNormalizedNormal = true)
 
 	float SceneDepth = CalcSceneDepth(UV);
 	
-	return DecodeGBufferData(GBufferA, GBufferB, GBufferC, GBufferD, GBufferE, GBufferVelocity, CustomNativeDepth, SceneDepth, bGetNormalizedNormal);
+	return DecodeGBufferData(GBufferA, GBufferB, GBufferC, GBufferD, GBufferE, GBufferVelocity, CustomNativeDepth, CustomStencil, SceneDepth, bGetNormalizedNormal);
 }
 
 // Minimal path for just the lighting model, used to branch around unlit pixels (skybox)
@@ -702,6 +721,7 @@ FScreenSpaceData GetScreenSpaceData(float2 UV, bool bGetNormalizedNormal = true)
 		float4 GBufferC = GBuffers.GBufferCTextureMS.Load(IntUV, SampleIndex);
 		float4 GBufferD = GBuffers.GBufferDTextureMS.Load(IntUV, SampleIndex);
 		float CustomNativeDepth = CustomDepthTextureNonMS.Load(int3(IntUV, 0)).r;
+		uint CustomStencil = CustomStencilTexture.Load(int3(IntUV, 0)).g;
 
 		#if ALLOW_STATIC_LIGHTING
 			float4 GBufferE = GBuffers.GBufferETextureMS.Load(IntUV, SampleIndex);
@@ -718,7 +738,7 @@ FScreenSpaceData GetScreenSpaceData(float2 UV, bool bGetNormalizedNormal = true)
 		float4 GBufferVelocity = 0;
 	#endif
 
-		return DecodeGBufferData(GBufferA, GBufferB, GBufferC, GBufferD, GBufferE, GBufferVelocity, CustomNativeDepth, SceneDepth, bGetNormalizedNormal);
+		return DecodeGBufferData(GBufferA, GBufferB, GBufferC, GBufferD, GBufferE, GBufferVelocity, CustomNativeDepth, CustomStencil, SceneDepth, bGetNormalizedNormal);
 	}
 
 	FGBufferData GetGBufferDataMS(float2 UV, uint SampleIndex, bool bGetNormalizedNormal = true)
diff --git a/Engine/Shaders/MaterialTemplate.usf b/Engine/Shaders/MaterialTemplate.usf
--- a/Engine/Shaders/MaterialTemplate.usf
+++ b/Engine/Shaders/MaterialTemplate.usf
@@ -1081,6 +1081,7 @@ float4 SceneTextureLookup(float2 UV, int SceneTextureIndex, bool bFiltered)
 	else if (SceneTextureIndex == 21)	return ScreenSpaceData.GBuffer.DecalMask;
 	else if (SceneTextureIndex == 22)	return float4(GetShadingModelColor(ScreenSpaceData.GBuffer), 1);
 	else if (SceneTextureIndex == 23)	return ScreenSpaceData.AmbientOcclusion;
+	else if (SceneTextureIndex == 24)	return ScreenSpaceData.GBuffer.CustomStencil;
 #else
 	switch(SceneTextureIndex)
 	{
@@ -1138,6 +1139,8 @@ float4 SceneTextureLookup(float2 UV, int SceneTextureIndex, bool bFiltered)
 		case 22: return float4(GetShadingModelColor(ScreenSpaceData.GBuffer), 1);
 		// PPI_AmbientOcclusion
 		case 23: return ScreenSpaceData.AmbientOcclusion;
+		// PPI_CustomStencil
+		case 24: return ScreenSpaceData.GBuffer.CustomStencil;
 		default:
 			return float4(0, 0, 0, 0);
 	}
diff --git a/Engine/Source/Runtime/Engine/Classes/Components/PrimitiveComponent.h b/Engine/Source/Runtime/Engine/Classes/Components/PrimitiveComponent.h
--- a/Engine/Source/Runtime/Engine/Classes/Components/PrimitiveComponent.h
+++ b/Engine/Source/Runtime/Engine/Classes/Components/PrimitiveComponent.h
@@ -212,10 +212,6 @@ public:
 	UPROPERTY()
 	uint32 bHasMotionBlurVelocityMeshes:1;
 	
-	/** If true, this component will be rendered in the CustomDepth pass (usually used for outlines) */
-	UPROPERTY(EditAnywhere, AdvancedDisplay, BlueprintReadOnly, Category=Rendering)
-	uint32 bRenderCustomDepth:1;
-
 	/** If true, this component will be rendered in the main pass (z prepass, basepass, transparency) */
 	UPROPERTY(EditAnywhere, AdvancedDisplay, BlueprintReadOnly, Category = Rendering)
 	uint32 bRenderInMainPass:1;
@@ -379,6 +375,14 @@ public:
 	UPROPERTY()
 	uint32 bUseEditorCompositing:1;
 
+	/** If true, this component will be rendered in the CustomDepth pass (usually used for outlines) */
+	UPROPERTY(EditAnywhere, AdvancedDisplay, BlueprintReadOnly, Category=Rendering, meta=(DisplayName = "Render CustomDepth Pass"))
+	uint32 bRenderCustomDepth:1;
+
+	/** Optionally write this 0-255 value to the stencil buffer in CustomDepth pass (Requires project setting or r.CustomDepth == 3) */
+	UPROPERTY(EditAnywhere, AdvancedDisplay, BlueprintReadOnly, Category=Rendering,  meta=(UIMin = "0", UIMax = "255", editcondition = "bRenderCustomDepth", DisplayName = "CustomDepth Stencil Value"))
+	int32 CustomDepthStencilValue;
+
 	/**
 	 * Translucent objects with a lower sort priority draw behind objects with a higher priority.
 	 * Translucent objects with the same priority are rendered from back-to-front based on their bounds origin.
@@ -968,6 +972,10 @@ public:
 	UFUNCTION(BlueprintCallable, Category="Rendering")
 	void SetRenderCustomDepth(bool bValue);
 
+	/** Sets the CustomDepth stencil value (0 - 255) and marks the render state dirty. */
+	UFUNCTION(BlueprintCallable, Category = "Rendering", meta=(UIMin = "0", UIMax = "255"))
+	void SetCustomDepthStencilValue(int32 Value);
+
 	/** Sets bRenderInMainPass property and marks the render state dirty. */
 	UFUNCTION(BlueprintCallable, Category = "Rendering")
 	void SetRenderInMainPass(bool bValue);
diff --git a/Engine/Source/Runtime/Engine/Classes/Engine/RendererSettings.h b/Engine/Source/Runtime/Engine/Classes/Engine/RendererSettings.h
--- a/Engine/Source/Runtime/Engine/Classes/Engine/RendererSettings.h
+++ b/Engine/Source/Runtime/Engine/Classes/Engine/RendererSettings.h
@@ -43,13 +43,14 @@ namespace ECompositingSampleCount
  * Enumerates available options for custom depth.
  */
 UENUM()
-namespace ECustomDepth
+namespace ECustomDepthStencil
 {
 	enum Type
 	{
 		Disabled = 0,
-		Enabled = 1,
-		EnabledOnDemand = 2,
+		Enabled = 1 UMETA(ToolTip="Depth buffer created immediately. Stencil disabled."),
+		EnabledOnDemand = 2 UMETA(ToolTip="Depth buffer created on first use, can save memory but cause stalls. Stencil disabled."),
+		EnabledWithStencil = 3 UMETA(ToolTip="Depth buffer created immediately. Stencil available for read/write."),
 	};
 }
 
@@ -182,9 +183,9 @@ class ENGINE_API URendererSettings : public UDeveloperSettings
 	FVector TranslucentSortAxis;
 
 	UPROPERTY(config, EditAnywhere, Category=Postprocessing, meta=(
-		ConsoleVariable="r.CustomDepth",DisplayName="Custom Depth Pass",
+		ConsoleVariable="r.CustomDepth",DisplayName="Custom Depth-Stencil Pass",
 		ToolTip="Whether the custom depth pass for tagging primitives for postprocessing passes is enabled. Enabling it on demand can save memory but may cause a hitch the first time the feature is used."))
-	TEnumAsByte<ECustomDepth::Type> CustomDepth;
+	TEnumAsByte<ECustomDepthStencil::Type> CustomDepthStencil;
 
 	UPROPERTY(config, EditAnywhere, Category = DefaultPostprocessingSettings, meta = (
 		ConsoleVariable = "r.DefaultFeature.Bloom", DisplayName = "Bloom",
diff --git a/Engine/Source/Runtime/Engine/Classes/Materials/MaterialExpressionSceneTexture.h b/Engine/Source/Runtime/Engine/Classes/Materials/MaterialExpressionSceneTexture.h
--- a/Engine/Source/Runtime/Engine/Classes/Materials/MaterialExpressionSceneTexture.h
+++ b/Engine/Source/Runtime/Engine/Classes/Materials/MaterialExpressionSceneTexture.h
@@ -57,6 +57,8 @@ enum ESceneTextureId
 	PPI_ShadingModel UMETA(DisplayName="Shading Model"),
 	/** Ambient Occlusion, single channel */
 	PPI_AmbientOcclusion UMETA(DisplayName="Ambient Occlusion"),
+	/** Scene stencil, contains CustomStencil mesh property of the opaque objects rendered with CustomDepth */
+	PPI_CustomStencil UMETA(DisplayName="CustomStencil"),
 };
 
 UCLASS(collapsecategories, hidecategories=Object)
diff --git a/Engine/Source/Runtime/Engine/Private/Components/PrimitiveComponent.cpp b/Engine/Source/Runtime/Engine/Private/Components/PrimitiveComponent.cpp
--- a/Engine/Source/Runtime/Engine/Private/Components/PrimitiveComponent.cpp
+++ b/Engine/Source/Runtime/Engine/Private/Components/PrimitiveComponent.cpp
@@ -140,6 +140,7 @@ UPrimitiveComponent::UPrimitiveComponent(const FObjectInitializer& ObjectInitial
 	CanBeCharacterBase_DEPRECATED = ECB_Yes;
 	CanCharacterStepUpOn = ECB_Yes;
 	ComponentId.PrimIDValue = NextComponentId.Increment();
+	CustomDepthStencilValue = 0;
 
 	bUseEditorCompositing = false;
 
@@ -2690,6 +2691,18 @@ void UPrimitiveComponent::SetRenderCustomDepth(bool bValue)
 	}
 }
 
+void UPrimitiveComponent::SetCustomDepthStencilValue(int32 Value)
+{
+	// Clamping to currently usable stencil range (as specified in property UI and tooltips)
+	int32 ClampedValue = FMath::Clamp(Value, 0, 255);
+
+	if (CustomDepthStencilValue != ClampedValue)
+	{
+		CustomDepthStencilValue = ClampedValue;
+		MarkRenderStateDirty();
+	}
+}
+
 void UPrimitiveComponent::SetRenderInMainPass(bool bValue)
 {
 	if (bRenderInMainPass != bValue)
diff --git a/Engine/Source/Runtime/Engine/Private/PrimitiveSceneProxy.cpp b/Engine/Source/Runtime/Engine/Private/PrimitiveSceneProxy.cpp
--- a/Engine/Source/Runtime/Engine/Private/PrimitiveSceneProxy.cpp
+++ b/Engine/Source/Runtime/Engine/Private/PrimitiveSceneProxy.cpp
@@ -27,7 +27,6 @@ FPrimitiveSceneProxy::FPrimitiveSceneProxy(const UPrimitiveComponent* InComponen
 ,	StaticDepthPriorityGroup(InComponent->GetStaticDepthPriorityGroup())
 ,	ViewOwnerDepthPriorityGroup(InComponent->ViewOwnerDepthPriorityGroup)
 ,	bStaticLighting(InComponent->HasStaticLighting())
-,	bRenderCustomDepth(InComponent->bRenderCustomDepth)
 ,	bRenderInMainPass(InComponent->bRenderInMainPass)
 ,	bRequiresVisibleLevelToRender(false)
 ,	bIsComponentLevelVisible(false)
@@ -61,6 +60,8 @@ FPrimitiveSceneProxy::FPrimitiveSceneProxy(const UPrimitiveComponent* InComponen
 ,	bSelectable(InComponent->bSelectable)
 ,	bHasPerInstanceHitProxies(InComponent->bHasPerInstanceHitProxies)
 ,	bUseEditorCompositing(InComponent->bUseEditorCompositing)
+,	bRenderCustomDepth(InComponent->bRenderCustomDepth)
+,	CustomDepthStencilValue((uint8)InComponent->CustomDepthStencilValue) 
 ,	LpvBiasMultiplier(InComponent->LpvBiasMultiplier)
 ,	IndirectLightingCacheQuality(InComponent->IndirectLightingCacheQuality)
 ,	Scene(InComponent->GetScene())
diff --git a/Engine/Source/Runtime/Engine/Public/PrimitiveSceneProxy.h b/Engine/Source/Runtime/Engine/Public/PrimitiveSceneProxy.h
--- a/Engine/Source/Runtime/Engine/Public/PrimitiveSceneProxy.h
+++ b/Engine/Source/Runtime/Engine/Public/PrimitiveSceneProxy.h
@@ -379,6 +379,7 @@ public:
 	inline bool IsIndividuallySelected() const { return bIndividuallySelected; }
 	inline bool IsSelected() const { return IsParentSelected() || IsIndividuallySelected(); }
 	inline bool ShouldRenderCustomDepth() const { return bRenderCustomDepth; }
+	inline uint8 GetCustomDepthStencilValue() const { return CustomDepthStencilValue; }
 	inline bool ShouldRenderInMainPass() const { return bRenderInMainPass; }
 	inline bool IsCollisionEnabled() const { return bCollisionEnabled; }
 	inline bool IsHovered() const { return bHovered; }
@@ -529,9 +530,6 @@ private:
 	/** True if the primitive will cache static lighting. */
 	uint32 bStaticLighting : 1;
 
-	/** This primitive has bRenderCustomDepth enabled */
-	uint32 bRenderCustomDepth : 1;
-
 	/** If true this primitive Renders in the mainPass */
 	uint32 bRenderInMainPass : 1;
 
@@ -648,6 +646,12 @@ private:
 	/** Whether this primitive should be composited onto the scene after post processing (editor only) */
 	uint32 bUseEditorCompositing : 1;
 
+	/** This primitive has bRenderCustomDepth enabled */
+	uint32 bRenderCustomDepth : 1;
+
+	/** Optionally write this stencil value during the CustomDepth pass */
+	uint8 CustomDepthStencilValue;
+
 protected:
 	/** The bias applied to LPV injection */
 	float LpvBiasMultiplier;
diff --git a/Engine/Source/Runtime/Renderer/Private/CustomDepthRendering.cpp b/Engine/Source/Runtime/Renderer/Private/CustomDepthRendering.cpp
--- a/Engine/Source/Runtime/Renderer/Private/CustomDepthRendering.cpp
+++ b/Engine/Source/Runtime/Renderer/Private/CustomDepthRendering.cpp
@@ -17,7 +17,7 @@
 	FCustomDepthPrimSet
 -----------------------------------------------------------------------------*/
 
-bool FCustomDepthPrimSet::DrawPrims(FRHICommandListImmediate& RHICmdList, const FViewInfo& View)
+bool FCustomDepthPrimSet::DrawPrims(FRHICommandListImmediate& RHICmdList, const FViewInfo& View, bool bWriteCustomStencilValues)
 {
 	bool bDirty=false;
 
@@ -36,6 +36,12 @@ bool FCustomDepthPrimSet::DrawPrims(FRHICommandListImmediate& RHICmdList, const
 
 				FDepthDrawingPolicyFactory::ContextType Context(DDM_AllOccluders);
 
+				if (bWriteCustomStencilValues)
+				{
+					const uint32 CustomDepthStencilValue = PrimitiveSceneProxy->GetCustomDepthStencilValue();
+					RHICmdList.SetDepthStencilState(TStaticDepthStencilState<true, CF_DepthNearOrEqual, true, CF_Always, SO_Keep, SO_Keep, SO_Replace>::GetRHI(), CustomDepthStencilValue);
+				}
+
 				for (int32 MeshBatchIndex = 0; MeshBatchIndex < View.DynamicMeshElements.Num(); MeshBatchIndex++)
 				{
 					const FMeshBatchAndRelevance& MeshBatchAndRelevance = View.DynamicMeshElements[MeshBatchIndex];
diff --git a/Engine/Source/Runtime/Renderer/Private/CustomDepthRendering.h b/Engine/Source/Runtime/Renderer/Private/CustomDepthRendering.h
--- a/Engine/Source/Runtime/Renderer/Private/CustomDepthRendering.h
+++ b/Engine/Source/Runtime/Renderer/Private/CustomDepthRendering.h
@@ -18,7 +18,7 @@ public:
 	* @param ViewInfo - current view used to draw items
 	* @return true if anything was drawn
 	*/
-	bool DrawPrims(FRHICommandListImmediate& RHICmdList, const class FViewInfo& View);
+	bool DrawPrims(FRHICommandListImmediate& RHICmdList, const class FViewInfo& View, bool bWriteCustomStencilValues);
 
 	/**
 	* Add a new primitive to the list of prims
diff --git a/Engine/Source/Runtime/Renderer/Private/PostProcess/SceneRenderTargets.cpp b/Engine/Source/Runtime/Renderer/Private/PostProcess/SceneRenderTargets.cpp
--- a/Engine/Source/Runtime/Renderer/Private/PostProcess/SceneRenderTargets.cpp
+++ b/Engine/Source/Runtime/Renderer/Private/PostProcess/SceneRenderTargets.cpp
@@ -69,7 +69,8 @@ static TAutoConsoleVariable<int32> CVarCustomDepth(
 	1,
 	TEXT("0: feature is disabled\n")
 	TEXT("1: feature is enabled, texture is created on demand\n")
-	TEXT("2: feature is enabled, texture is not released until required (should be the project setting if the feature should not stall)"),
+	TEXT("2: feature is enabled, texture is not released until required (should be the project setting if the feature should not stall)\n")
+	TEXT("3: feature is enabled, stencil writes are enabled, texture is not released until required (should be the project setting if the feature should not stall)"),
 	ECVF_RenderThreadSafe
 	);
 
@@ -803,7 +804,12 @@ bool FSceneRenderTargets::BeginRenderingCustomDepth(FRHICommandListImmediate& RH
 	{
 		SCOPED_DRAW_EVENT(RHICmdList, BeginRenderingCustomDepth);
 		
-		SetRenderTarget(RHICmdList, FTextureRHIRef(), CustomDepthRenderTarget->GetRenderTargetItem().ShaderResourceTexture);
+		FRHIDepthRenderTargetView DepthView(CustomDepthRenderTarget->GetRenderTargetItem().ShaderResourceTexture);
+		FRHISetRenderTargetsInfo Info(0, nullptr, DepthView);
+		Info.bClearStencil = IsCustomDepthPassWritingStencil();
+		Info.StencilClearValue = 0;
+
+		RHICmdList.SetRenderTargetsAndClear(Info);
 
 		return true;
 	}
@@ -815,8 +821,7 @@ void FSceneRenderTargets::FinishRenderingCustomDepth(FRHICommandListImmediate& R
 {
 	SCOPED_DRAW_EVENT(RHICmdList, FinishRenderingCustomDepth);
 
-	auto& CurrentSceneColor = GetSceneColor();
-	RHICmdList.CopyToResolveTarget(CurrentSceneColor->GetRenderTargetItem().TargetableTexture, CurrentSceneColor->GetRenderTargetItem().ShaderResourceTexture, true, FResolveParams(ResolveRect));
+	RHICmdList.CopyToResolveTarget(CustomDepth->GetRenderTargetItem().TargetableTexture, CustomDepth->GetRenderTargetItem().ShaderResourceTexture, true, FResolveParams(ResolveRect));
 
 	bCustomDepthIsValid = true;
 }
@@ -1763,8 +1768,9 @@ IPooledRenderTarget* FSceneRenderTargets::RequestCustomDepth(bool bPrimitives)
 {
 	int Value = CVarCustomDepth.GetValueOnRenderThread();
 
-	if((Value == 1  && bPrimitives) || Value == 2)
+	if((Value == 1  && bPrimitives) || Value == 2 || IsCustomDepthPassWritingStencil())
 	{
+		// Todo: Could check if writes stencil here and create min viable target
 		FPooledRenderTargetDesc Desc(FPooledRenderTargetDesc::Create2DDesc(BufferSize, PF_DepthStencil, TexCreate_None, TexCreate_DepthStencilTargetable, false));
 		GRenderTargetPool.FindFreeElement(Desc, CustomDepth, TEXT("CustomDepth"));
 		return CustomDepth;
@@ -1773,6 +1779,11 @@ IPooledRenderTarget* FSceneRenderTargets::RequestCustomDepth(bool bPrimitives)
 	return 0;
 }
 
+bool FSceneRenderTargets::IsCustomDepthPassWritingStencil() const
+{
+	return (CVarCustomDepth.GetValueOnRenderThread() == 3);
+}
+
 /** Returns an index in the range [0, NumCubeShadowDepthSurfaces) given an input resolution. */
 int32 FSceneRenderTargets::GetCubeShadowDepthZIndex(int32 ShadowResolution) const
 {
@@ -2092,6 +2103,7 @@ void FDeferredPixelShaderParameters::Bind(const FShaderParameterMap& ParameterMa
 	ScreenSpaceAOTextureSampler.Bind(ParameterMap,TEXT("ScreenSpaceAOTextureSampler"));
 	CustomDepthTexture.Bind(ParameterMap,TEXT("CustomDepthTexture"));
 	CustomDepthTextureSampler.Bind(ParameterMap,TEXT("CustomDepthTextureSampler"));
+	CustomStencilTexture.Bind(ParameterMap,TEXT("CustomStencilTexture"));
 }
 
 bool IsDBufferEnabled();
@@ -2152,6 +2164,36 @@ void FDeferredPixelShaderParameters::Set(FRHICommandList& RHICmdList, const Shad
 
 			SetTextureParameter(RHICmdList, ShaderRHI, CustomDepthTexture, CustomDepthTextureSampler, TStaticSamplerState<>::GetRHI(), CustomDepth->GetRenderTargetItem().ShaderResourceTexture);
 			SetTextureParameter(RHICmdList, ShaderRHI, CustomDepthTextureNonMS, CustomDepth->GetRenderTargetItem().ShaderResourceTexture);
+
+			if (CustomStencilTexture.IsBound())
+			{
+				FSceneViewState* ViewState = (FSceneViewState*)View.State;	
+					
+				if (ViewState && GSceneRenderTargets.bCustomDepthIsValid)
+				{
+					FTexture2DRHIRef& CustomStencilTargetableTexture = (FTexture2DRHIRef&)CustomDepth->GetRenderTargetItem().TargetableTexture;
+					
+					// Release cached stencil SRV if outdated
+					if (ViewState->CustomStencilSRVCacheKey != CustomStencilTargetableTexture)
+					{
+						ViewState->CustomStencilSRVCacheKey.SafeRelease();
+						ViewState->CustomStencilSRVCacheValue.SafeRelease();
+					}
+
+					// Create and cache stencil SRV if needed
+					if (!ViewState->CustomStencilSRVCacheValue)
+					{
+						ViewState->CustomStencilSRVCacheKey = CustomStencilTargetableTexture;
+						ViewState->CustomStencilSRVCacheValue = RHICreateShaderResourceView(CustomStencilTargetableTexture, 0, 1, PF_X24_G8);
+					}
+
+					SetSRVParameter(RHICmdList, ShaderRHI, CustomStencilTexture, ViewState->CustomStencilSRVCacheValue);
+				}
+				else
+				{
+					SetTextureParameter(RHICmdList, ShaderRHI, CustomStencilTexture, GSystemTextures.BlackDummy->GetRenderTargetItem().ShaderResourceTexture);
+				}
+			}
 		}
 	}
 	else if (TextureMode == ESceneRenderTargetsMode::DontSet ||
@@ -2201,6 +2243,7 @@ FArchive& operator<<(FArchive& Ar,FDeferredPixelShaderParameters& Parameters)
 	Ar << Parameters.ScreenSpaceAOTextureSampler;
 	Ar << Parameters.CustomDepthTexture;
 	Ar << Parameters.CustomDepthTextureSampler;
+	Ar << Parameters.CustomStencilTexture;
 
 	return Ar;
 }
diff --git a/Engine/Source/Runtime/Renderer/Private/PostProcess/SceneRenderTargets.h b/Engine/Source/Runtime/Renderer/Private/PostProcess/SceneRenderTargets.h
--- a/Engine/Source/Runtime/Renderer/Private/PostProcess/SceneRenderTargets.h
+++ b/Engine/Source/Runtime/Renderer/Private/PostProcess/SceneRenderTargets.h
@@ -381,6 +381,8 @@ public:
 	// @return can be 0 if the feature is disabled
 	IPooledRenderTarget* RequestCustomDepth(bool bPrimitives);
 
+	bool IsCustomDepthPassWritingStencil() const;
+
 	// ---
 
 	/** */
diff --git a/Engine/Source/Runtime/Renderer/Private/ScenePrivate.h b/Engine/Source/Runtime/Renderer/Private/ScenePrivate.h
--- a/Engine/Source/Runtime/Renderer/Private/ScenePrivate.h
+++ b/Engine/Source/Runtime/Renderer/Private/ScenePrivate.h
@@ -545,9 +545,11 @@ public:
 	TRefCountPtr<IPooledRenderTarget> MobileAaColor0;
 	TRefCountPtr<IPooledRenderTarget> MobileAaColor1;
 
-	// cache for selection outline to a avoid reallocations of the SRV, Key is to detect if the object has changed
+	// cache for stencil reads to a avoid reallocations of the SRV, Key is to detect if the object has changed
 	FTextureRHIRef SelectionOutlineCacheKey;
 	TRefCountPtr<FRHIShaderResourceView> SelectionOutlineCacheValue;
+	FTextureRHIRef CustomStencilSRVCacheKey;
+	TRefCountPtr<FRHIShaderResourceView> CustomStencilSRVCacheValue;
 
 	/** Distance field AO tile intersection GPU resources.  Last frame's state is not used, but they must be sized exactly to the view so stored here. */
 	class FTileIntersectionResources* AOTileIntersectionResources;
@@ -733,6 +735,8 @@ public:
 		MobileAaColor1.SafeRelease();
 		SelectionOutlineCacheKey.SafeRelease();
 		SelectionOutlineCacheValue.SafeRelease();
+		CustomStencilSRVCacheKey.SafeRelease();
+		CustomStencilSRVCacheValue.SafeRelease();
 
 		for (int32 CascadeIndex = 0; CascadeIndex < ARRAY_COUNT(GlobalDistanceFieldClipmapState); CascadeIndex++)
 		{
diff --git a/Engine/Source/Runtime/Renderer/Private/SceneRendering.cpp b/Engine/Source/Runtime/Renderer/Private/SceneRendering.cpp
--- a/Engine/Source/Runtime/Renderer/Private/SceneRendering.cpp
+++ b/Engine/Source/Runtime/Renderer/Private/SceneRendering.cpp
@@ -1305,16 +1305,19 @@ void FSceneRenderer::RenderCustomDepthPass(FRHICommandListImmediate& RHICmdList)
 			FViewInfo& View = Views[ViewIndex];
 
 			RHICmdList.SetViewport(View.ViewRect.Min.X, View.ViewRect.Min.Y, 0.0f, View.ViewRect.Max.X, View.ViewRect.Max.Y, 1.0f);
-
-			RHICmdList.Clear(false, FLinearColor(0, 0, 0, 0), true, (float)ERHIZBuffer::FarPlane, false, 0, FIntRect());
 			
 			// seems this is set each draw call anyway
 			RHICmdList.SetRasterizerState(TStaticRasterizerState<>::GetRHI());
-
-			RHICmdList.SetDepthStencilState(TStaticDepthStencilState<true, CF_DepthNearOrEqual>::GetRHI());
 			RHICmdList.SetBlendState(TStaticBlendState<>::GetRHI());
+			
+			const bool bWriteCustomStencilValues = GSceneRenderTargets.IsCustomDepthPassWritingStencil();
+
+			if (!bWriteCustomStencilValues)
+			{
+				RHICmdList.SetDepthStencilState(TStaticDepthStencilState<true, CF_DepthNearOrEqual>::GetRHI());
+			}
 
-			View.CustomDepthSet.DrawPrims(RHICmdList, View);
+			View.CustomDepthSet.DrawPrims(RHICmdList, View, bWriteCustomStencilValues);
 		}
 
 		// resolve using the current ResolveParams 
diff --git a/Engine/Source/Runtime/Renderer/Public/SceneRenderTargetParameters.h b/Engine/Source/Runtime/Renderer/Public/SceneRenderTargetParameters.h
--- a/Engine/Source/Runtime/Renderer/Public/SceneRenderTargetParameters.h
+++ b/Engine/Source/Runtime/Renderer/Public/SceneRenderTargetParameters.h
@@ -96,4 +96,5 @@ private:
 	FShaderResourceParameter ScreenSpaceAOTextureSampler;
 	FShaderResourceParameter CustomDepthTexture;
 	FShaderResourceParameter CustomDepthTextureSampler;
+	FShaderResourceParameter CustomStencilTexture;
 };
\ No newline at end of file
