missing file for CL 2560964
fixed UE-15921 Skeletal Meshes with Multiple Translucent Batches will render Both

diff --git a/Engine/Source/Runtime/Renderer/Private/TranslucentRendering.cpp b/Engine/Source/Runtime/Renderer/Private/TranslucentRendering.cpp
--- a/Engine/Source/Runtime/Renderer/Private/TranslucentRendering.cpp
+++ b/Engine/Source/Runtime/Renderer/Private/TranslucentRendering.cpp
@@ -389,6 +389,11 @@ bool FTranslucencyDrawingPolicyFactory::DrawMesh(
 	// Only render translucent materials.
 	if(IsTranslucentBlendMode(BlendMode))
 	{
+		if (Material->IsSeparateTranslucencyEnabled() != DrawingContext.bSeparateTranslucencyPass)
+		{
+			return false;
+		}
+
 		if (Material->RequiresSceneColorCopy_RenderThread())
 		{
 			if (DrawingContext.bSceneColorCopyIsUpToDate == false)
