UE4 - Large upgrade to parallel rendering. No longer block per pass. Several RHI cmdlist passthrough functions added and implemented on the consoles

diff --git a/Engine/Source/Runtime/Renderer/Private/BasePassRendering.cpp b/Engine/Source/Runtime/Renderer/Private/BasePassRendering.cpp
--- a/Engine/Source/Runtime/Renderer/Private/BasePassRendering.cpp
+++ b/Engine/Source/Runtime/Renderer/Private/BasePassRendering.cpp
@@ -74,6 +74,7 @@ void FSkyLightReflectionParameters::GetSkyParametersFromScene(const FScene* Scen
 
 void FTranslucentLightingParameters::Set(FRHICommandList& RHICmdList, FShader* Shader, const FViewInfo* View)
 {
+	FSceneRenderTargets& SceneContext = FSceneRenderTargets::Get(RHICmdList);
 	auto PixelShader = Shader->GetPixelShader();
 
 	SetTextureParameter(
@@ -82,7 +83,7 @@ void FTranslucentLightingParameters::Set(FRHICommandList& RHICmdList, FShader* S
 		TranslucencyLightingVolumeAmbientInner, 
 		TranslucencyLightingVolumeAmbientInnerSampler, 
 		TStaticSamplerState<SF_Bilinear, AM_Clamp, AM_Clamp, AM_Clamp>::GetRHI(), 
-		GSceneRenderTargets.GetTranslucencyVolumeAmbient(TVC_Inner)->GetRenderTargetItem().ShaderResourceTexture);
+		SceneContext.GetTranslucencyVolumeAmbient(TVC_Inner)->GetRenderTargetItem().ShaderResourceTexture);
 
 	SetTextureParameter(
 		RHICmdList, 
@@ -90,7 +91,7 @@ void FTranslucentLightingParameters::Set(FRHICommandList& RHICmdList, FShader* S
 		TranslucencyLightingVolumeAmbientOuter, 
 		TranslucencyLightingVolumeAmbientOuterSampler, 
 		TStaticSamplerState<SF_Bilinear, AM_Clamp, AM_Clamp, AM_Clamp>::GetRHI(), 
-		GSceneRenderTargets.GetTranslucencyVolumeAmbient(TVC_Outer)->GetRenderTargetItem().ShaderResourceTexture);
+		SceneContext.GetTranslucencyVolumeAmbient(TVC_Outer)->GetRenderTargetItem().ShaderResourceTexture);
 
 	SetTextureParameter(
 		RHICmdList, 
@@ -98,7 +99,7 @@ void FTranslucentLightingParameters::Set(FRHICommandList& RHICmdList, FShader* S
 		TranslucencyLightingVolumeDirectionalInner, 
 		TranslucencyLightingVolumeDirectionalInnerSampler, 
 		TStaticSamplerState<SF_Bilinear, AM_Clamp, AM_Clamp, AM_Clamp>::GetRHI(), 
-		GSceneRenderTargets.GetTranslucencyVolumeDirectional(TVC_Inner)->GetRenderTargetItem().ShaderResourceTexture);
+		SceneContext.GetTranslucencyVolumeDirectional(TVC_Inner)->GetRenderTargetItem().ShaderResourceTexture);
 
 	SetTextureParameter(
 		RHICmdList, 
@@ -106,7 +107,7 @@ void FTranslucentLightingParameters::Set(FRHICommandList& RHICmdList, FShader* S
 		TranslucencyLightingVolumeDirectionalOuter, 
 		TranslucencyLightingVolumeDirectionalOuterSampler, 
 		TStaticSamplerState<SF_Bilinear, AM_Clamp, AM_Clamp, AM_Clamp>::GetRHI(), 
-		GSceneRenderTargets.GetTranslucencyVolumeDirectional(TVC_Outer)->GetRenderTargetItem().ShaderResourceTexture);
+		SceneContext.GetTranslucencyVolumeDirectional(TVC_Outer)->GetRenderTargetItem().ShaderResourceTexture);
 
 	SkyLightReflectionParameters.SetParameters(RHICmdList, Shader->GetPixelShader(), (const FScene*)(View->Family->Scene), true);
 
diff --git a/Engine/Source/Runtime/Renderer/Private/CompositionLighting/CompositionLighting.cpp b/Engine/Source/Runtime/Renderer/Private/CompositionLighting/CompositionLighting.cpp
--- a/Engine/Source/Runtime/Renderer/Private/CompositionLighting/CompositionLighting.cpp
+++ b/Engine/Source/Runtime/Renderer/Private/CompositionLighting/CompositionLighting.cpp
@@ -187,7 +187,9 @@ static FRenderingCompositeOutputRef AddPostProcessingAmbientOcclusion(FRHIComman
 		AmbientOcclusionPassMip1->SetInput(ePId_Input3, HZBInput);
 	}
 
-	FRenderingCompositePass* GBufferA = Context.Graph.RegisterPass(new(FMemStack::Get()) FRCPassPostProcessInput(GSceneRenderTargets.GBufferA));
+	FSceneRenderTargets& SceneContext = FSceneRenderTargets::Get(RHICmdList);
+
+	FRenderingCompositePass* GBufferA = Context.Graph.RegisterPass(new(FMemStack::Get()) FRCPassPostProcessInput(SceneContext.GBufferA));
 
 	// finally full resolution
 
@@ -202,7 +204,7 @@ static FRenderingCompositeOutputRef AddPostProcessingAmbientOcclusion(FRHIComman
 
 	Context.FinalOutput = FRenderingCompositeOutputRef(AmbientOcclusionPassMip0);
 
-	GSceneRenderTargets.bScreenSpaceAOIsValid = true;
+	SceneContext.bScreenSpaceAOIsValid = true;
 
 	if(IsBasePassAmbientOcclusionRequired(Context))
 	{
@@ -264,18 +266,19 @@ void FCompositionLighting::ProcessAfterBasePass(FRHICommandListImmediate& RHICmd
 {
 	check(IsInRenderingThread());
 	
+	FSceneRenderTargets& SceneContext = FSceneRenderTargets::Get(RHICmdList);
 	// might get renamed to refracted or ...WithAO
-	GSceneRenderTargets.GetSceneColor()->SetDebugName(TEXT("SceneColor"));
+	SceneContext.GetSceneColor()->SetDebugName(TEXT("SceneColor"));
 	// to be able to observe results with VisualizeTexture
 
-	GRenderTargetPool.VisualizeTexture.SetCheckPoint(RHICmdList, GSceneRenderTargets.GetSceneColor());
-	GRenderTargetPool.VisualizeTexture.SetCheckPoint(RHICmdList, GSceneRenderTargets.GBufferA);
-	GRenderTargetPool.VisualizeTexture.SetCheckPoint(RHICmdList, GSceneRenderTargets.GBufferB);
-	GRenderTargetPool.VisualizeTexture.SetCheckPoint(RHICmdList, GSceneRenderTargets.GBufferC);
-	GRenderTargetPool.VisualizeTexture.SetCheckPoint(RHICmdList, GSceneRenderTargets.GBufferD);
-	GRenderTargetPool.VisualizeTexture.SetCheckPoint(RHICmdList, GSceneRenderTargets.GBufferE);
-	GRenderTargetPool.VisualizeTexture.SetCheckPoint(RHICmdList, GSceneRenderTargets.GBufferVelocity);
-	GRenderTargetPool.VisualizeTexture.SetCheckPoint(RHICmdList, GSceneRenderTargets.ScreenSpaceAO);
+	GRenderTargetPool.VisualizeTexture.SetCheckPoint(RHICmdList, SceneContext.GetSceneColor());
+	GRenderTargetPool.VisualizeTexture.SetCheckPoint(RHICmdList, SceneContext.GBufferA);
+	GRenderTargetPool.VisualizeTexture.SetCheckPoint(RHICmdList, SceneContext.GBufferB);
+	GRenderTargetPool.VisualizeTexture.SetCheckPoint(RHICmdList, SceneContext.GBufferC);
+	GRenderTargetPool.VisualizeTexture.SetCheckPoint(RHICmdList, SceneContext.GBufferD);
+	GRenderTargetPool.VisualizeTexture.SetCheckPoint(RHICmdList, SceneContext.GBufferE);
+	GRenderTargetPool.VisualizeTexture.SetCheckPoint(RHICmdList, SceneContext.GBufferVelocity);
+	GRenderTargetPool.VisualizeTexture.SetCheckPoint(RHICmdList, SceneContext.ScreenSpaceAO);
 	
 	// so that the passes can register themselves to the graph
 	{
@@ -309,7 +312,7 @@ void FCompositionLighting::ProcessAfterBasePass(FRHICommandListImmediate& RHICmd
 
 		SCOPED_DRAW_EVENT(RHICmdList, LightCompositionTasks_PreLighting);
 
-		TRefCountPtr<IPooledRenderTarget>& SceneColor = GSceneRenderTargets.GetSceneColor();
+		TRefCountPtr<IPooledRenderTarget>& SceneColor = SceneContext.GetSceneColor();
 
 		Context.FinalOutput.GetOutput()->RenderTargetDesc = SceneColor->GetDesc();
 		Context.FinalOutput.GetOutput()->PooledRenderTarget = SceneColor;
@@ -329,7 +332,9 @@ void FCompositionLighting::ProcessLpvIndirect(FRHICommandListImmediate& RHICmdLi
 
 	if(IsLpvIndirectPassRequired(Context))
 	{
-		FRenderingCompositePass* SSAO = Context.Graph.RegisterPass(new FRCPassPostProcessInput(GSceneRenderTargets.ScreenSpaceAO));
+		FSceneRenderTargets& SceneContext = FSceneRenderTargets::Get(RHICmdList);
+
+		FRenderingCompositePass* SSAO = Context.Graph.RegisterPass(new FRCPassPostProcessInput(SceneContext.ScreenSpaceAO));
 
 		FRenderingCompositePass* Pass = Context.Graph.RegisterPass(new FRCPassPostProcessLpvIndirect());
 		Pass->SetInput(ePId_Input0, Context.FinalOutput);
@@ -350,10 +355,12 @@ void FCompositionLighting::ProcessLpvIndirect(FRHICommandListImmediate& RHICmdLi
 void FCompositionLighting::ProcessAfterLighting(FRHICommandListImmediate& RHICmdList, FViewInfo& View)
 {
 	check(IsInRenderingThread());
+	FSceneRenderTargets& SceneContext = FSceneRenderTargets::Get(RHICmdList);
+
 	
-	GRenderTargetPool.VisualizeTexture.SetCheckPoint(RHICmdList, GSceneRenderTargets.ReflectiveShadowMapDiffuse);
-	GRenderTargetPool.VisualizeTexture.SetCheckPoint(RHICmdList, GSceneRenderTargets.ReflectiveShadowMapNormal);
-	GRenderTargetPool.VisualizeTexture.SetCheckPoint(RHICmdList, GSceneRenderTargets.ReflectiveShadowMapDepth);
+	GRenderTargetPool.VisualizeTexture.SetCheckPoint(RHICmdList, SceneContext.ReflectiveShadowMapDiffuse);
+	GRenderTargetPool.VisualizeTexture.SetCheckPoint(RHICmdList, SceneContext.ReflectiveShadowMapNormal);
+	GRenderTargetPool.VisualizeTexture.SetCheckPoint(RHICmdList, SceneContext.ReflectiveShadowMapDepth);
 
 	{
 		FMemMark Mark(FMemStack::Get());
@@ -411,6 +418,6 @@ void FCompositionLighting::ProcessAfterLighting(FRHICommandListImmediate& RHICmd
 	{
 		// The RT should be released as early as possible to allow sharing of that memory for other purposes.
 		// This becomes even more important with some limited VRam (XBoxOne).
-		GSceneRenderTargets.SetLightAttenuation(0);
+		SceneContext.SetLightAttenuation(0);
 	}
 }
diff --git a/Engine/Source/Runtime/Renderer/Private/CompositionLighting/PostProcessAmbient.cpp b/Engine/Source/Runtime/Renderer/Private/CompositionLighting/PostProcessAmbient.cpp
--- a/Engine/Source/Runtime/Renderer/Private/CompositionLighting/PostProcessAmbient.cpp
+++ b/Engine/Source/Runtime/Renderer/Private/CompositionLighting/PostProcessAmbient.cpp
@@ -100,7 +100,7 @@ void FRCPassPostProcessAmbient::Render(FRenderingCompositePassContext& Context)
 			View.ViewRect.Min.X, View.ViewRect.Min.Y,
 			View.ViewRect.Width(), View.ViewRect.Height(),
 			View.ViewRect.Size(),
-			GSceneRenderTargets.GetBufferSizeXY(),
+			FSceneRenderTargets::Get(Context.RHICmdList).GetBufferSizeXY(),
 			*VertexShader,
 			EDRF_UseTriangleOptimization);
 	}
@@ -118,7 +118,7 @@ void FRCPassPostProcessAmbient::Process(FRenderingCompositePassContext& Context)
 	FIntRect DestRect = View.ViewRect;
 	FIntPoint DestSize = DestRect.Size();
 
-	const FSceneRenderTargetItem& DestRenderTarget = GSceneRenderTargets.GetSceneColor()->GetRenderTargetItem();
+	const FSceneRenderTargetItem& DestRenderTarget = FSceneRenderTargets::Get(Context.RHICmdList).GetSceneColor()->GetRenderTargetItem();
 
 	// Set the view family's render target/viewport.
 	SetRenderTarget(Context.RHICmdList, DestRenderTarget.TargetableTexture, FTextureRHIRef());
diff --git a/Engine/Source/Runtime/Renderer/Private/CompositionLighting/PostProcessAmbientOcclusion.cpp b/Engine/Source/Runtime/Renderer/Private/CompositionLighting/PostProcessAmbientOcclusion.cpp
--- a/Engine/Source/Runtime/Renderer/Private/CompositionLighting/PostProcessAmbientOcclusion.cpp
+++ b/Engine/Source/Runtime/Renderer/Private/CompositionLighting/PostProcessAmbientOcclusion.cpp
@@ -76,7 +76,7 @@ public:
 		DeferredParameters.Set(Context.RHICmdList, ShaderRHI, Context.View);
 
 		// e.g. 4 means the input texture is 4x smaller than the buffer size
-		uint32 ScaleToFullRes = GSceneRenderTargets.GetBufferSizeXY().X / Context.Pass->GetOutput(ePId_Output0)->RenderTargetDesc.Extent.X;
+		uint32 ScaleToFullRes = FSceneRenderTargets::Get(Context.RHICmdList).GetBufferSizeXY().X / Context.Pass->GetOutput(ePId_Output0)->RenderTargetDesc.Extent.X;
 
 		// /1000 to be able to define the value in that distance
 		FVector4 AmbientOcclusionSetupParamsValue = FVector4(ScaleToFullRes, Settings.AmbientOcclusionMipThreshold / ScaleToFullRes, 0, 0);
@@ -315,7 +315,7 @@ void FRCPassPostProcessAmbientOcclusionSetup::Process(FRenderingCompositePassCon
 	FIntPoint DestSize = PassOutputs[0].RenderTargetDesc.Extent;
 
 	// e.g. 4 means the input texture is 4x smaller than the buffer size
-	uint32 ScaleFactor = GSceneRenderTargets.GetBufferSizeXY().X / DestSize.X;
+	uint32 ScaleFactor = FSceneRenderTargets::Get(Context.RHICmdList).GetBufferSizeXY().X / DestSize.X;
 
 	FIntRect SrcRect = View.ViewRect;
 	FIntRect DestRect = SrcRect  / ScaleFactor;
@@ -349,7 +349,7 @@ void FRCPassPostProcessAmbientOcclusionSetup::Process(FRenderingCompositePassCon
 			SrcRect.Min.X, SrcRect.Min.Y, 
 		SrcRect.Width(), SrcRect.Height(),
 		DestRect.Size(),
-		GSceneRenderTargets.GetBufferSizeXY(),
+		FSceneRenderTargets::Get(Context.RHICmdList).GetBufferSizeXY(),
 		VertexShader,
 		EDRF_UseTriangleOptimization);
 
@@ -427,6 +427,7 @@ void FRCPassPostProcessAmbientOcclusion::Process(FRenderingCompositePassContext&
 	const FPooledRenderTargetDesc* InputDesc2 = GetInputDesc(ePId_Input2);
 
 	const FSceneRenderTargetItem* DestRenderTarget = 0;
+	FSceneRenderTargets& SceneContext = FSceneRenderTargets::Get(Context.RHICmdList);
 
 	if(bAOSetupAsInput)
 	{
@@ -434,7 +435,7 @@ void FRCPassPostProcessAmbientOcclusion::Process(FRenderingCompositePassContext&
 	}
 	else
 	{
-		DestRenderTarget = &GSceneRenderTargets.ScreenSpaceAO->GetRenderTargetItem();
+		DestRenderTarget = &SceneContext.ScreenSpaceAO->GetRenderTargetItem();
 	}
 
 	ensure(InputDesc0);
@@ -442,7 +443,7 @@ void FRCPassPostProcessAmbientOcclusion::Process(FRenderingCompositePassContext&
 	FIntPoint TexSize = InputDesc0->Extent;
 
 	// usually 1, 2, 4 or 8
-	uint32 ScaleToFullRes = GSceneRenderTargets.GetBufferSizeXY().X / TexSize.X;
+	uint32 ScaleToFullRes = SceneContext.GetBufferSizeXY().X / TexSize.X;
 
 	FIntRect ViewRect = FIntRect::DivideAndRoundUp(View.ViewRect, ScaleToFullRes);
 
@@ -542,8 +543,9 @@ void FRCPassPostProcessBasePassAO::Process(FRenderingCompositePassContext& Conte
 	SCOPED_DRAW_EVENT(Context.RHICmdList, ApplyAOToBasePassSceneColor);
 
 	const FSceneView& View = Context.View;
+	FSceneRenderTargets& SceneContext = FSceneRenderTargets::Get(Context.RHICmdList);
 
-	const FSceneRenderTargetItem& DestRenderTarget = GSceneRenderTargets.GetSceneColor()->GetRenderTargetItem();
+	const FSceneRenderTargetItem& DestRenderTarget = SceneContext.GetSceneColor()->GetRenderTargetItem();
 
 	// Set the view family's render target/viewport.
 	SetRenderTarget(Context.RHICmdList, DestRenderTarget.TargetableTexture,	FTextureRHIParamRef(), ESimpleRenderTargetMode::EExistingColorAndDepth);
@@ -563,7 +565,7 @@ void FRCPassPostProcessBasePassAO::Process(FRenderingCompositePassContext& Conte
 	SetGlobalBoundShaderState(Context.RHICmdList, Context.GetFeatureLevel(), BoundShaderState, GFilterVertexDeclaration.VertexDeclarationRHI, *VertexShader, *PixelShader);
 
 	VertexShader->SetParameters(Context);
-	PixelShader->SetParameters(Context, GSceneRenderTargets.GetBufferSizeXY());
+	PixelShader->SetParameters(Context, SceneContext.GetBufferSizeXY());
 
 	// Draw a quad mapping scene color to the view's render target
 	DrawRectangle( 
@@ -573,7 +575,7 @@ void FRCPassPostProcessBasePassAO::Process(FRenderingCompositePassContext& Conte
 		View.ViewRect.Min.X, View.ViewRect.Min.Y, 
 		View.ViewRect.Width(), View.ViewRect.Height(),
 		View.ViewRect.Size(),
-		GSceneRenderTargets.GetBufferSizeXY(),
+		SceneContext.GetBufferSizeXY(),
 		*VertexShader,
 		EDRF_UseTriangleOptimization);
 
diff --git a/Engine/Source/Runtime/Renderer/Private/CompositionLighting/PostProcessAmbientOcclusion.h b/Engine/Source/Runtime/Renderer/Private/CompositionLighting/PostProcessAmbientOcclusion.h
--- a/Engine/Source/Runtime/Renderer/Private/CompositionLighting/PostProcessAmbientOcclusion.h
+++ b/Engine/Source/Runtime/Renderer/Private/CompositionLighting/PostProcessAmbientOcclusion.h
@@ -88,7 +88,7 @@ public:
 		FVector2D ViewportUVToRandomUV(InputTextureSize.X / (float)SSAORandomizationSize.X, InputTextureSize.Y / (float)SSAORandomizationSize.Y);
 
 		// e.g. 4 means the input texture is 4x smaller than the buffer size
-		uint32 ScaleToFullRes = GSceneRenderTargets.GetBufferSizeXY().X / InputTextureSize.X;
+		uint32 ScaleToFullRes = FSceneRenderTargets::Get(RHICmdList).GetBufferSizeXY().X / InputTextureSize.X;
 
 		FIntRect ViewRect = FIntRect::DivideAndRoundUp(View.ViewRect, ScaleToFullRes);
 
diff --git a/Engine/Source/Runtime/Renderer/Private/CompositionLighting/PostProcessDeferredDecals.cpp b/Engine/Source/Runtime/Renderer/Private/CompositionLighting/PostProcessDeferredDecals.cpp
--- a/Engine/Source/Runtime/Renderer/Private/CompositionLighting/PostProcessDeferredDecals.cpp
+++ b/Engine/Source/Runtime/Renderer/Private/CompositionLighting/PostProcessDeferredDecals.cpp
@@ -375,9 +375,10 @@ FGlobalBoundShaderState StencilDecalMaskBoundShaderState;
 void StencilDecalMask(FRHICommandList& RHICmdList, const FViewInfo& View)
 {
 	SCOPED_DRAW_EVENT(RHICmdList, StencilDecalMask);
+	FSceneRenderTargets& SceneContext = FSceneRenderTargets::Get(RHICmdList);
 	RHICmdList.SetRasterizerState(TStaticRasterizerState<FM_Solid, CM_None>::GetRHI());
 	RHICmdList.SetBlendState(TStaticBlendState<CW_NONE>::GetRHI());
-	SetRenderTarget(RHICmdList, NULL, GSceneRenderTargets.GetSceneDepthSurface(), ESimpleRenderTargetMode::EUninitializedColorExistingDepth, FExclusiveDepthStencil::DepthRead_StencilWrite);
+	SetRenderTarget(RHICmdList, NULL, SceneContext.GetSceneDepthSurface(), ESimpleRenderTargetMode::EUninitializedColorExistingDepth, FExclusiveDepthStencil::DepthRead_StencilWrite);
 	RHICmdList.SetViewport(View.ViewRect.Min.X, View.ViewRect.Min.Y, 0.0f, View.ViewRect.Max.X, View.ViewRect.Max.Y, 1.0f);
 
 	// Write 1 to highest bit of stencil to areas that should not receive decals
@@ -399,7 +400,7 @@ void StencilDecalMask(FRHICommandList& RHICmdList, const FViewInfo& View)
 		View.ViewRect.Min.X, View.ViewRect.Min.Y, 
 		View.ViewRect.Width(), View.ViewRect.Height(),
 		FIntPoint(View.ViewRect.Width(), View.ViewRect.Height()),
-		GSceneRenderTargets.GetBufferSizeXY(),
+		SceneContext.GetBufferSizeXY(),
 		*ScreenVertexShader,
 		EDRF_UseTriangleOptimization);
 }
@@ -721,6 +722,7 @@ FRCPassPostProcessDeferredDecals::FRCPassPostProcessDeferredDecals(uint32 InRend
 void FRCPassPostProcessDeferredDecals::Process(FRenderingCompositePassContext& Context)
 {
 	FRHICommandListImmediate& RHICmdList = Context.RHICmdList;
+	FSceneRenderTargets& SceneContext = FSceneRenderTargets::Get(RHICmdList);
 
 	const bool bShaderComplexity = Context.View.Family->EngineShowFlags.ShaderComplexity;
 	const bool bDBuffer = IsDBufferEnabled();
@@ -728,6 +730,20 @@ void FRCPassPostProcessDeferredDecals::Process(FRenderingCompositePassContext& C
 
 	SCOPED_DRAW_EVENT(RHICmdList, PostProcessDeferredDecals);
 
+	enum EDecalResolveBufferIndex
+	{
+		SceneColorIndex,
+		GBufferAIndex,
+		GBufferBIndex,
+		GBufferCIndex,
+		DBufferAIndex,
+		DBufferBIndex,
+		DBufferCIndex,
+		ResolveBufferMax,
+	};
+
+	FTextureRHIParamRef TargetsToResolve[ResolveBufferMax] = { nullptr };
+
 	if(RenderStage == 0)
 	{
 		// before BasePass, only if DBuffer is enabled
@@ -735,27 +751,27 @@ void FRCPassPostProcessDeferredDecals::Process(FRenderingCompositePassContext& C
 		check(bDBuffer);
 
 		// DBuffer: Decal buffer
-		FPooledRenderTargetDesc Desc(FPooledRenderTargetDesc::Create2DDesc(GSceneRenderTargets.GBufferA->GetDesc().Extent, 
+		FPooledRenderTargetDesc Desc(FPooledRenderTargetDesc::Create2DDesc(SceneContext.GBufferA->GetDesc().Extent, 
 			PF_B8G8R8A8, 
 			TexCreate_None, 
 			TexCreate_ShaderResource | TexCreate_RenderTargetable,
 			false));
 
-		if(!GSceneRenderTargets.DBufferA)
+		if(!SceneContext.DBufferA)
 		{
-			GRenderTargetPool.FindFreeElement(Desc, GSceneRenderTargets.DBufferA, TEXT("DBufferA"));
+			GRenderTargetPool.FindFreeElement(Desc, SceneContext.DBufferA, TEXT("DBufferA"));
 		}
 
-		if(!GSceneRenderTargets.DBufferB)
+		if(!SceneContext.DBufferB)
 		{
-			GRenderTargetPool.FindFreeElement(Desc, GSceneRenderTargets.DBufferB, TEXT("DBufferB"));
+			GRenderTargetPool.FindFreeElement(Desc, SceneContext.DBufferB, TEXT("DBufferB"));
 		}
 
 		Desc.Format = PF_R8G8;
 
-		if(!GSceneRenderTargets.DBufferC)
+		if(!SceneContext.DBufferC)
 		{
-			GRenderTargetPool.FindFreeElement(Desc, GSceneRenderTargets.DBufferC, TEXT("DBufferC"));
+			GRenderTargetPool.FindFreeElement(Desc, SceneContext.DBufferC, TEXT("DBufferC"));
 		}
 
 		// we assume views are non overlapping, then we need to clear only once in the beginning, otherwise we would need to set scissor rects
@@ -766,15 +782,24 @@ void FRCPassPostProcessDeferredDecals::Process(FRenderingCompositePassContext& C
 		{
 			SCOPED_DRAW_EVENT(RHICmdList, DBufferClear);
 
-			// could be optimized
-			SetRenderTarget(RHICmdList, GSceneRenderTargets.DBufferA->GetRenderTargetItem().TargetableTexture, FTextureRHIParamRef());
-			RHICmdList.Clear(true, FLinearColor(0, 0, 0, 1), false, (float)ERHIZBuffer::FarPlane, false, 0, FIntRect());
-			SetRenderTarget(RHICmdList, GSceneRenderTargets.DBufferB->GetRenderTargetItem().TargetableTexture, FTextureRHIParamRef());
-			// todo: some hardware would like to have 0 or 1 for faster clear, we chose 128/255 to represent 0 (8 bit cannot represent 0.5f)
-			RHICmdList.Clear(true, FLinearColor(128.0f / 255.0f, 128.0f / 255.0f, 128.0f / 255.0f, 1), false, (float)ERHIZBuffer::FarPlane, false, 0, FIntRect());
-			SetRenderTarget(RHICmdList, GSceneRenderTargets.DBufferC->GetRenderTargetItem().TargetableTexture, FTextureRHIParamRef());
-			// R:roughness, G:roughness opacity
-			RHICmdList.Clear(true, FLinearColor(0, 1, 0, 1), false, (float)ERHIZBuffer::FarPlane, false, 0, FIntRect());
+			
+			FRHIRenderTargetView RenderTargets[3];
+			RenderTargets[0] = FRHIRenderTargetView(SceneContext.DBufferA->GetRenderTargetItem().TargetableTexture, 0, -1, ERenderTargetLoadAction::EClear, ERenderTargetStoreAction::EStore);
+			RenderTargets[1] = FRHIRenderTargetView(SceneContext.DBufferB->GetRenderTargetItem().TargetableTexture, 0, -1, ERenderTargetLoadAction::EClear, ERenderTargetStoreAction::EStore);
+			RenderTargets[2] = FRHIRenderTargetView(SceneContext.DBufferC->GetRenderTargetItem().TargetableTexture, 0, -1, ERenderTargetLoadAction::EClear, ERenderTargetStoreAction::EStore);
+
+			FRHIDepthRenderTargetView DepthView(SceneContext.GetSceneDepthSurface(), ERenderTargetLoadAction::ELoad, ERenderTargetStoreAction::ENoAction, ERenderTargetLoadAction::ELoad, ERenderTargetStoreAction::ENoAction, FExclusiveDepthStencil(FExclusiveDepthStencil::DepthRead_StencilWrite));
+
+			FRHISetRenderTargetsInfo Info(3, RenderTargets, DepthView);
+			Info.ClearColors[0] = FLinearColor(0, 0, 0, 1);
+			Info.ClearColors[1] = FLinearColor(128.0f / 255.0f, 128.0f / 255.0f, 128.0f / 255.0f, 1);
+			Info.ClearColors[2] = FLinearColor(0, 1, 0, 1);
+
+			RHICmdList.SetRenderTargetsAndClear(Info);
+
+			TargetsToResolve[DBufferAIndex] = SceneContext.DBufferA->GetRenderTargetItem().TargetableTexture;
+			TargetsToResolve[DBufferBIndex] = SceneContext.DBufferB->GetRenderTargetItem().TargetableTexture;
+			TargetsToResolve[DBufferCIndex] = SceneContext.DBufferC->GetRenderTargetItem().TargetableTexture;
 		}
 	}
 
@@ -783,12 +808,8 @@ void FRCPassPostProcessDeferredDecals::Process(FRenderingCompositePassContext& C
 	const FSceneViewFamily& ViewFamily = *(View.Family);
 	FScene& Scene = *(FScene*)ViewFamily.Scene;
 
-	if(!Scene.Decals.Num())
+	if (Scene.Decals.Num())
 	{
-		// to avoid the stats showing up
-		return;
-	}
-
 	TArray<FTransientDecalRenderData, SceneRenderingAllocator> SortedDecals;
 	SortedDecals.Empty(Scene.Decals.Num());
 
@@ -813,12 +834,12 @@ void FRCPassPostProcessDeferredDecals::Process(FRenderingCompositePassContext& C
 
 		// can be optimized as we test against a sphere around the box instead of the box itself
 		const float ConservativeRadius = FMath::Sqrt(
-			ComponentToWorldMatrix.GetScaledAxis( EAxis::X ).SizeSquared() * FMath::Square(GDefaultDecalSize.X) +
-			ComponentToWorldMatrix.GetScaledAxis( EAxis::Y ).SizeSquared() * FMath::Square(GDefaultDecalSize.Y) +
-			ComponentToWorldMatrix.GetScaledAxis( EAxis::Z ).SizeSquared() * FMath::Square(GDefaultDecalSize.Z));
+				ComponentToWorldMatrix.GetScaledAxis(EAxis::X).SizeSquared() * FMath::Square(GDefaultDecalSize.X) +
+				ComponentToWorldMatrix.GetScaledAxis(EAxis::Y).SizeSquared() * FMath::Square(GDefaultDecalSize.Y) +
+				ComponentToWorldMatrix.GetScaledAxis(EAxis::Z).SizeSquared() * FMath::Square(GDefaultDecalSize.Z));
 
 		// can be optimized as the test is too conservative (sphere instead of OBB)
-		if(ConservativeRadius < SMALL_NUMBER || !View.ViewFrustum.IntersectSphere(ComponentToWorldMatrix.GetOrigin(), ConservativeRadius))
+			if (ConservativeRadius < SMALL_NUMBER || !View.ViewFrustum.IntersectSphere(ComponentToWorldMatrix.GetOrigin(), ConservativeRadius))
 		{
 			bIsShown = false;
 		}
@@ -837,7 +858,7 @@ void FRCPassPostProcessDeferredDecals::Process(FRenderingCompositePassContext& C
 		}
 	}
 
-	if(SortedDecals.Num() > 0)
+		if (SortedDecals.Num() > 0)
 	{
 		FIntRect SrcRect = View.ViewRect;
 		FIntRect DestRect = View.ViewRect;
@@ -845,14 +866,14 @@ void FRCPassPostProcessDeferredDecals::Process(FRenderingCompositePassContext& C
 		bool bStencilDecals = true;
 
 #if DBUFFER_DONT_USE_STENCIL_YET
-		if(RenderStage == 0)
+			if (RenderStage == 0)
 		{
 			bStencilDecals = false;
 		}
 #endif
 
 		// Setup a stencil mask to prevent certain pixels from receiving deferred decals
-		if(bStencilDecals)
+			if (bStencilDecals)
 		{
 			StencilDecalMask(RHICmdList, View);
 		}
@@ -862,7 +883,7 @@ void FRCPassPostProcessDeferredDecals::Process(FRenderingCompositePassContext& C
 		// Also sort by component since Sort() is not stable
 		struct FCompareFTransientDecalRenderData
 		{
-			FORCEINLINE bool operator()( const FTransientDecalRenderData& A, const FTransientDecalRenderData& B ) const
+				FORCEINLINE bool operator()(const FTransientDecalRenderData& A, const FTransientDecalRenderData& B) const
 			{
 				if (B.DecalProxy->SortOrder != A.DecalProxy->SortOrder)
 				{ 
@@ -877,7 +898,7 @@ void FRCPassPostProcessDeferredDecals::Process(FRenderingCompositePassContext& C
 					return B.bHasNormal < A.bHasNormal;
 				}
 				// Batch decals with the same material together
-				if (B.MaterialProxy != A.MaterialProxy )
+					if (B.MaterialProxy != A.MaterialProxy)
 				{
 					return B.MaterialProxy < A.MaterialProxy;
 				}
@@ -886,7 +907,7 @@ void FRCPassPostProcessDeferredDecals::Process(FRenderingCompositePassContext& C
 		};
 
 		// Sort decals by blend mode to reduce render target switches
-		SortedDecals.Sort( FCompareFTransientDecalRenderData() );
+			SortedDecals.Sort(FCompareFTransientDecalRenderData());
 
 		// optimization to have less state changes
 		int32 LastDecalBlendMode = -1;
@@ -898,19 +919,7 @@ void FRCPassPostProcessDeferredDecals::Process(FRenderingCompositePassContext& C
 		SCOPED_DRAW_EVENT(RHICmdList, Decals);
 		INC_DWORD_STAT_BY(STAT_Decals, SortedDecals.Num());
 
-		enum EDecalResolveBufferIndex
-		{
-			SceneColorIndex,
-			GBufferAIndex,
-			GBufferBIndex,
-			GBufferCIndex,
-			DBufferAIndex,
-			DBufferBIndex,
-			DBufferCIndex,
-			ResolveBufferMax,
-		};
 	
-		FTextureRHIParamRef TargetsToResolve[ResolveBufferMax] = { nullptr };
 
 		for (int32 DecalIndex = 0, DecalCount = SortedDecals.Num(); DecalIndex < DecalCount; DecalIndex++)
 		{
@@ -927,16 +936,16 @@ void FRCPassPostProcessDeferredDecals::Process(FRenderingCompositePassContext& C
 
 			// can be optimized as we test against a sphere around the box instead of the box itself
 			const float ConservativeRadius = FMath::Sqrt(
-				ComponentToWorldMatrix.GetScaledAxis( EAxis::X ).SizeSquared() * FMath::Square(GDefaultDecalSize.X) +
-				ComponentToWorldMatrix.GetScaledAxis( EAxis::Y ).SizeSquared() * FMath::Square(GDefaultDecalSize.Y) +
-				ComponentToWorldMatrix.GetScaledAxis( EAxis::Z ).SizeSquared() * FMath::Square(GDefaultDecalSize.Z));
+					ComponentToWorldMatrix.GetScaledAxis(EAxis::X).SizeSquared() * FMath::Square(GDefaultDecalSize.X) +
+					ComponentToWorldMatrix.GetScaledAxis(EAxis::Y).SizeSquared() * FMath::Square(GDefaultDecalSize.Y) +
+					ComponentToWorldMatrix.GetScaledAxis(EAxis::Z).SizeSquared() * FMath::Square(GDefaultDecalSize.Z));
 
 			EDecalBlendMode DecalBlendMode = DecalData.DecalBlendMode;
 
 			bool bStencilThisDecal = bStencilDecals;
 			
 #if DBUFFER_DONT_USE_STENCIL_YET
-			if(ComputeRenderStage(DecalBlendMode) == 0)
+				if (ComputeRenderStage(DecalBlendMode) == 0)
 			{
 				bStencilThisDecal = false;
 			}
@@ -944,7 +953,7 @@ void FRCPassPostProcessDeferredDecals::Process(FRenderingCompositePassContext& C
 
 			ERenderTargetMode CurrentRenderTargetMode = ComputeRenderTargetMode(DecalBlendMode);
 
-			if(bShaderComplexity)
+				if (bShaderComplexity)
 			{
 				CurrentRenderTargetMode = RTM_SceneColor;
 				// we want additive blending for the ShaderComplexity mode
@@ -952,39 +961,39 @@ void FRCPassPostProcessDeferredDecals::Process(FRenderingCompositePassContext& C
 			}
 
 			// fewer rendertarget switches if possible
-			if(CurrentRenderTargetMode != LastRenderTargetMode)
+				if (CurrentRenderTargetMode != LastRenderTargetMode)
 			{
 				LastRenderTargetMode = CurrentRenderTargetMode;
 
-				switch(CurrentRenderTargetMode)
+					switch (CurrentRenderTargetMode)
 				{
 					case RTM_SceneColorAndGBuffer:
 						{							
-							TargetsToResolve[SceneColorIndex] = GSceneRenderTargets.GetSceneColor()->GetRenderTargetItem().TargetableTexture;
-							TargetsToResolve[GBufferAIndex] = GSceneRenderTargets.GBufferA->GetRenderTargetItem().TargetableTexture;
-							TargetsToResolve[GBufferBIndex] = GSceneRenderTargets.GBufferB->GetRenderTargetItem().TargetableTexture;
-							TargetsToResolve[GBufferCIndex] = GSceneRenderTargets.GBufferC->GetRenderTargetItem().TargetableTexture;
+							TargetsToResolve[SceneColorIndex] = SceneContext.GetSceneColor()->GetRenderTargetItem().TargetableTexture;
+							TargetsToResolve[GBufferAIndex] = SceneContext.GBufferA->GetRenderTargetItem().TargetableTexture;
+							TargetsToResolve[GBufferBIndex] = SceneContext.GBufferB->GetRenderTargetItem().TargetableTexture;
+							TargetsToResolve[GBufferCIndex] = SceneContext.GBufferC->GetRenderTargetItem().TargetableTexture;
 							
-							SetRenderTargets(RHICmdList, 4, TargetsToResolve, GSceneRenderTargets.GetSceneDepthSurface(), ESimpleRenderTargetMode::EExistingColorAndDepth, FExclusiveDepthStencil::DepthRead_StencilWrite);
+							SetRenderTargets(RHICmdList, 4, TargetsToResolve, SceneContext.GetSceneDepthSurface(), ESimpleRenderTargetMode::EExistingColorAndDepth, FExclusiveDepthStencil::DepthRead_StencilWrite);
 						}
 						break;
 
 					case RTM_GBufferNormal:
-						TargetsToResolve[GBufferAIndex] = GSceneRenderTargets.GBufferA->GetRenderTargetItem().TargetableTexture;
-						SetRenderTarget(RHICmdList, TargetsToResolve[GBufferAIndex], GSceneRenderTargets.GetSceneDepthSurface(), ESimpleRenderTargetMode::EExistingColorAndDepth, FExclusiveDepthStencil::DepthRead_StencilWrite);
+						TargetsToResolve[GBufferAIndex] = SceneContext.GBufferA->GetRenderTargetItem().TargetableTexture;
+						SetRenderTarget(RHICmdList, TargetsToResolve[GBufferAIndex], SceneContext.GetSceneDepthSurface(), ESimpleRenderTargetMode::EExistingColorAndDepth, FExclusiveDepthStencil::DepthRead_StencilWrite);
 						break;
 					
 					case RTM_SceneColor:
-						TargetsToResolve[SceneColorIndex] = GSceneRenderTargets.GetSceneColor()->GetRenderTargetItem().TargetableTexture;
-						SetRenderTarget(RHICmdList, TargetsToResolve[SceneColorIndex], GSceneRenderTargets.GetSceneDepthSurface(), ESimpleRenderTargetMode::EExistingColorAndDepth, FExclusiveDepthStencil::DepthRead_StencilWrite);
+						TargetsToResolve[SceneColorIndex] = SceneContext.GetSceneColor()->GetRenderTargetItem().TargetableTexture;
+						SetRenderTarget(RHICmdList, TargetsToResolve[SceneColorIndex], SceneContext.GetSceneDepthSurface(), ESimpleRenderTargetMode::EExistingColorAndDepth, FExclusiveDepthStencil::DepthRead_StencilWrite);
 						break;
 
 					case RTM_DBuffer:
 						{							
-							TargetsToResolve[DBufferAIndex] = GSceneRenderTargets.DBufferA->GetRenderTargetItem().TargetableTexture;
-							TargetsToResolve[DBufferBIndex] = GSceneRenderTargets.DBufferB->GetRenderTargetItem().TargetableTexture;
-							TargetsToResolve[DBufferCIndex] = GSceneRenderTargets.DBufferC->GetRenderTargetItem().TargetableTexture;
-							SetRenderTargets(RHICmdList, 3, &TargetsToResolve[DBufferAIndex], GSceneRenderTargets.GetSceneDepthSurface(), ESimpleRenderTargetMode::EExistingColorAndDepth, FExclusiveDepthStencil::DepthRead_StencilWrite);
+							TargetsToResolve[DBufferAIndex] = SceneContext.DBufferA->GetRenderTargetItem().TargetableTexture;
+							TargetsToResolve[DBufferBIndex] = SceneContext.DBufferB->GetRenderTargetItem().TargetableTexture;
+							TargetsToResolve[DBufferCIndex] = SceneContext.DBufferC->GetRenderTargetItem().TargetableTexture;
+							SetRenderTargets(RHICmdList, 3, &TargetsToResolve[DBufferAIndex], SceneContext.GetSceneDepthSurface(), ESimpleRenderTargetMode::EExistingColorAndDepth, FExclusiveDepthStencil::DepthRead_StencilWrite);
 						}
 						break;
 
@@ -1032,33 +1041,33 @@ void FRCPassPostProcessDeferredDecals::Process(FRenderingCompositePassContext& C
 
 				VertexShader->SetParameters(RHICmdList, View, FrustumComponentToClip);
 
-				const int32 IsInsideDecal = ((FVector)View.ViewMatrices.ViewOrigin - ComponentToWorldMatrix.GetOrigin()).SizeSquared() < FMath::Square(ConservativeRadius * 1.05f + View.NearClippingDistance * 2.0f) + ( bThisDecalUsesStencil ) ? 2 : 0;
-				if ( WasInsideDecal != IsInsideDecal )
+					const int32 IsInsideDecal = ((FVector)View.ViewMatrices.ViewOrigin - ComponentToWorldMatrix.GetOrigin()).SizeSquared() < FMath::Square(ConservativeRadius * 1.05f + View.NearClippingDistance * 2.0f) + (bThisDecalUsesStencil) ? 2 : 0;
+					if (WasInsideDecal != IsInsideDecal)
 				{
 					WasInsideDecal = IsInsideDecal;
-					if ( !(IsInsideDecal & 1) )
+						if (!(IsInsideDecal & 1))
 					{
 						// Render backfaces with depth tests disabled since the camera is inside (or close to inside) the light function geometry
 						RHICmdList.SetRasterizerState(View.bReverseCulling ? TStaticRasterizerState<FM_Solid, CM_CCW>::GetRHI() : TStaticRasterizerState<FM_Solid, CM_CW>::GetRHI());
-						if(bStencilDecals)
+							if (bStencilDecals)
 						{
 							// Enable stencil testing, only write to pixels with stencil of 0
-							if ( bThisDecalUsesStencil )
+								if (bThisDecalUsesStencil)
 							{
 								RHICmdList.SetDepthStencilState(TStaticDepthStencilState<
-									false,CF_Always,
-									true,CF_Equal,SO_Zero,SO_Zero,SO_Zero,
-									true,CF_Equal,SO_Zero,SO_Zero,SO_Zero,
+										false, CF_Always,
+										true, CF_Equal, SO_Zero, SO_Zero, SO_Zero,
+										true, CF_Equal, SO_Zero, SO_Zero, SO_Zero,
 									0xff, 0x7f
 								>::GetRHI(), 1);
 							}
 							else
 							{
 								RHICmdList.SetDepthStencilState(TStaticDepthStencilState<
-									false,CF_Always,
-									true,CF_Equal,SO_Keep,SO_Keep,SO_Keep,
-									false,CF_Always,SO_Keep,SO_Keep,SO_Keep,
-									0x80,0x00>::GetRHI(), 0);
+										false, CF_Always,
+										true, CF_Equal, SO_Keep, SO_Keep, SO_Keep,
+										false, CF_Always, SO_Keep, SO_Keep, SO_Keep,
+										0x80, 0x00>::GetRHI(), 0);
 							}
 						}
 						else
@@ -1069,26 +1078,26 @@ void FRCPassPostProcessDeferredDecals::Process(FRenderingCompositePassContext& C
 					else
 					{
 						// Render frontfaces with depth tests on to get the speedup from HiZ since the camera is outside the light function geometry
-						if(bStencilDecals)
+							if (bStencilDecals)
 						{
 							// Render frontfaces with depth tests on to get the speedup from HiZ since the camera is outside the light function geometry
 							// Enable stencil testing, only write to pixels with stencil of 0
-							if ( bThisDecalUsesStencil )
+								if (bThisDecalUsesStencil)
 							{
 								RHICmdList.SetDepthStencilState(TStaticDepthStencilState<
-									false,CF_DepthNearOrEqual,
-									true,CF_Equal,SO_Zero,SO_Zero,SO_Zero,
-									true,CF_Equal,SO_Zero,SO_Zero,SO_Zero,
+										false, CF_DepthNearOrEqual,
+										true, CF_Equal, SO_Zero, SO_Zero, SO_Zero,
+										true, CF_Equal, SO_Zero, SO_Zero, SO_Zero,
 									0xff, 0x7f
 								>::GetRHI(), 1);
 							}
 							else
 							{
 								RHICmdList.SetDepthStencilState(TStaticDepthStencilState<
-									false,CF_DepthNearOrEqual,
-									true,CF_Equal,SO_Keep,SO_Keep,SO_Keep,
-									false,CF_Always,SO_Keep,SO_Keep,SO_Keep,
-									0x80,0x00>::GetRHI(), 0);
+										false, CF_DepthNearOrEqual,
+										true, CF_Equal, SO_Keep, SO_Keep, SO_Keep,
+										false, CF_Always, SO_Keep, SO_Keep, SO_Keep,
+										0x80, 0x00>::GetRHI(), 0);
 							}
 							RHICmdList.SetRasterizerState(View.bReverseCulling ? TStaticRasterizerState<FM_Solid, CM_CW>::GetRHI() : TStaticRasterizerState<FM_Solid, CM_CCW>::GetRHI());
 						}
@@ -1107,6 +1116,16 @@ void FRCPassPostProcessDeferredDecals::Process(FRenderingCompositePassContext& C
 		// we don't modify stencil but if out input was having stencil for us (after base pass - we need to clear)
 		// Clear stencil to 0, which is the assumed default by other passes
 		RHICmdList.Clear(false, FLinearColor::White, false, (float)ERHIZBuffer::FarPlane, true, 0, FIntRect());
+		}
+
+		if (RenderStage == 0)
+		{
+			// before BasePass
+			GRenderTargetPool.VisualizeTexture.SetCheckPoint(RHICmdList, SceneContext.DBufferA);
+			GRenderTargetPool.VisualizeTexture.SetCheckPoint(RHICmdList, SceneContext.DBufferB);
+			GRenderTargetPool.VisualizeTexture.SetCheckPoint(RHICmdList, SceneContext.DBufferC);
+		}
+	}
 
 		// resolve the targets we wrote to.
 		FResolveParams ResolveParams;
@@ -1119,15 +1138,6 @@ void FRCPassPostProcessDeferredDecals::Process(FRenderingCompositePassContext& C
 		}		
 	}
 
-	if(RenderStage == 0)
-	{
-		// before BasePass
-		GRenderTargetPool.VisualizeTexture.SetCheckPoint(RHICmdList, GSceneRenderTargets.DBufferA);
-		GRenderTargetPool.VisualizeTexture.SetCheckPoint(RHICmdList, GSceneRenderTargets.DBufferB);
-		GRenderTargetPool.VisualizeTexture.SetCheckPoint(RHICmdList, GSceneRenderTargets.DBufferC);
-	}
-}
-
 FPooledRenderTargetDesc FRCPassPostProcessDeferredDecals::ComputeOutputDesc(EPassOutputId InPassOutputId) const
 {
 	// This pass creates it's own output so the compositing graph output isn't needed.
diff --git a/Engine/Source/Runtime/Renderer/Private/CompositionLighting/PostProcessLpvIndirect.cpp b/Engine/Source/Runtime/Renderer/Private/CompositionLighting/PostProcessLpvIndirect.cpp
--- a/Engine/Source/Runtime/Renderer/Private/CompositionLighting/PostProcessLpvIndirect.cpp
+++ b/Engine/Source/Runtime/Renderer/Private/CompositionLighting/PostProcessLpvIndirect.cpp
@@ -205,6 +205,8 @@ IMPLEMENT_SHADER_TYPE(,FPostProcessLpvDirectionalOcclusionPS,TEXT("PostProcessLp
 
 void FRCPassPostProcessLpvIndirect::Process(FRenderingCompositePassContext& Context)
 {
+	FSceneRenderTargets& SceneContext = FSceneRenderTargets::Get(Context.RHICmdList);
+
 	{
 		FRenderingCompositeOutput* OutputOfMyInput = GetInput(ePId_Input0)->GetOutput();
 		PassOutputs[0].PooledRenderTarget = OutputOfMyInput->PooledRenderTarget;
@@ -224,6 +226,7 @@ void FRCPassPostProcessLpvIndirect::Process(FRenderingCompositePassContext& Cont
 	{
 		return;
 	}
+
 	FLightPropagationVolume* Lpv = ViewState->GetLightPropagationVolume();
 
 	const FSceneViewFamily& ViewFamily = *(View.Family);
@@ -237,10 +240,10 @@ void FRCPassPostProcessLpvIndirect::Process(FRenderingCompositePassContext& Cont
 	// Apply specular separately if we're mixing reflection environment with indirect lighting
 	const bool bApplySeparateSpecularRT = View.Family->EngineShowFlags.ReflectionEnvironment && bMixing;
 
-	const FSceneRenderTargetItem& DestColorRenderTarget = GSceneRenderTargets.GetSceneColor()->GetRenderTargetItem();
-	const FSceneRenderTargetItem& DestSpecularRenderTarget = GSceneRenderTargets.LightAccumulation->GetRenderTargetItem();
+	const FSceneRenderTargetItem& DestColorRenderTarget = SceneContext.GetSceneColor()->GetRenderTargetItem();
+	const FSceneRenderTargetItem& DestSpecularRenderTarget = SceneContext.LightAccumulation->GetRenderTargetItem();
 
-	const FSceneRenderTargetItem& DestDirectionalOcclusionRenderTarget = GSceneRenderTargets.DirectionalOcclusion->GetRenderTargetItem();
+	const FSceneRenderTargetItem& DestDirectionalOcclusionRenderTarget = SceneContext.DirectionalOcclusion->GetRenderTargetItem();
 
 	// Make sure the LPV Update has completed
 	Lpv->InsertGPUWaitForAsyncUpdate(Context.RHICmdList);
@@ -318,7 +321,7 @@ void FRCPassPostProcessLpvIndirect::Process(FRenderingCompositePassContext& Cont
 		    View.ViewRect.Min.X, View.ViewRect.Min.Y, 
 		    View.ViewRect.Width(), View.ViewRect.Height(),
 		    View.ViewRect.Size(),
-		    GSceneRenderTargets.GetBufferSizeXY(),
+		    SceneContext.GetBufferSizeXY(),
 		    *VertexShader);
 
 		Context.RHICmdList.CopyToResolveTarget(DestColorRenderTarget.TargetableTexture, DestColorRenderTarget.ShaderResourceTexture, false, FResolveParams());
@@ -330,14 +333,16 @@ void FRCPassPostProcessLpvIndirect::Process(FRenderingCompositePassContext& Cont
 
 	if ( LPVSettings.LPVDirectionalOcclusionIntensity > 0.0001f )
 	{
-		GRenderTargetPool.VisualizeTexture.SetCheckPoint(Context.RHICmdList, GSceneRenderTargets.DirectionalOcclusion);
+		GRenderTargetPool.VisualizeTexture.SetCheckPoint(Context.RHICmdList, SceneContext.DirectionalOcclusion);
 	}
 }
 
 void FRCPassPostProcessLpvIndirect::DoDirectionalOcclusionPass(FRenderingCompositePassContext& Context) const
 {
+	FSceneRenderTargets& SceneContext = FSceneRenderTargets::Get(Context.RHICmdList);
+
 	SCOPED_DRAW_EVENT(Context.RHICmdList, PostProcessLpvDirectionalOcclusion );
-	const FSceneRenderTargetItem& DestDirectionalOcclusionRenderTarget = GSceneRenderTargets.DirectionalOcclusion->GetRenderTargetItem();
+	const FSceneRenderTargetItem& DestDirectionalOcclusionRenderTarget = SceneContext.DirectionalOcclusion->GetRenderTargetItem();
 	FViewInfo& View = Context.View;
 	FSceneViewState* ViewState = (FSceneViewState*)View.State;
 	const FFinalPostProcessSettings& PostprocessSettings = Context.View.FinalPostProcessSettings;
@@ -352,7 +357,7 @@ void FRCPassPostProcessLpvIndirect::DoDirectionalOcclusionPass(FRenderingComposi
 	FTextureRHIParamRef RenderTarget = DestDirectionalOcclusionRenderTarget.TargetableTexture;
 
 	SetRenderTargets(Context.RHICmdList, 1, &RenderTarget, NULL, ESimpleRenderTargetMode::EExistingColorAndDepth, FExclusiveDepthStencil::DepthRead_StencilNop);
-	
+
 	Context.SetViewportAndCallRHI(View.ViewRect);
 	Context.RHICmdList.SetBlendState( TStaticBlendState<>::GetRHI() );
 	Context.RHICmdList.SetRasterizerState(TStaticRasterizerState<>::GetRHI());
@@ -365,7 +370,7 @@ void FRCPassPostProcessLpvIndirect::DoDirectionalOcclusionPass(FRenderingComposi
 	TShaderMapRef< FPostProcessLpvDirectionalOcclusionPS > PixelShader(View.ShaderMap);
 
 	static FGlobalBoundShaderState BoundShaderState;
-
+	
 	// call it once after setting up the shader data to avoid the warnings in the function
 	SetGlobalBoundShaderState(Context.RHICmdList, Context.GetFeatureLevel(), BoundShaderState, GFilterVertexDeclaration.VertexDeclarationRHI, *VertexShader, *PixelShader);
 
@@ -382,7 +387,7 @@ void FRCPassPostProcessLpvIndirect::DoDirectionalOcclusionPass(FRenderingComposi
 		View.ViewRect.Min.X, View.ViewRect.Min.Y, 
 		View.ViewRect.Width(), View.ViewRect.Height(),
 		View.ViewRect.Size(),
-		GSceneRenderTargets.GetBufferSizeXY(),
+		SceneContext.GetBufferSizeXY(),
 		*VertexShader);
 }
 
diff --git a/Engine/Source/Runtime/Renderer/Private/CompositionLighting/PostProcessPassThrough.cpp b/Engine/Source/Runtime/Renderer/Private/CompositionLighting/PostProcessPassThrough.cpp
--- a/Engine/Source/Runtime/Renderer/Private/CompositionLighting/PostProcessPassThrough.cpp
+++ b/Engine/Source/Runtime/Renderer/Private/CompositionLighting/PostProcessPassThrough.cpp
@@ -88,10 +88,11 @@ void FRCPassPostProcessPassThrough::Process(FRenderingCompositePassContext& Cont
 
 	FIntPoint SrcSize = InputDesc->Extent;
 	FIntPoint DestSize = Dest ? Dest->GetDesc().Extent : PassOutputs[0].RenderTargetDesc.Extent;
+	FSceneRenderTargets& SceneContext = FSceneRenderTargets::Get(Context.RHICmdList);
 
 	// e.g. 4 means the input texture is 4x smaller than the buffer size
-	uint32 InputScaleFactor = GSceneRenderTargets.GetBufferSizeXY().X / SrcSize.X;
-	uint32 OutputScaleFactor = GSceneRenderTargets.GetBufferSizeXY().X / DestSize.X;
+	uint32 InputScaleFactor = SceneContext.GetBufferSizeXY().X / SrcSize.X;
+	uint32 OutputScaleFactor = SceneContext.GetBufferSizeXY().X / DestSize.X;
 
 	FIntRect SrcRect = View.ViewRect / InputScaleFactor;
 	FIntRect DestRect = View.ViewRect / OutputScaleFactor;
diff --git a/Engine/Source/Runtime/Renderer/Private/DeferredShadingRenderer.cpp b/Engine/Source/Runtime/Renderer/Private/DeferredShadingRenderer.cpp
--- a/Engine/Source/Runtime/Renderer/Private/DeferredShadingRenderer.cpp
+++ b/Engine/Source/Runtime/Renderer/Private/DeferredShadingRenderer.cpp
@@ -57,6 +57,11 @@ static TAutoConsoleVariable<int32> CVarVisualizeTexturePool(
 	ECVF_Cheat | ECVF_RenderThreadSafe);
 #endif
 
+static TAutoConsoleVariable<int32> CVarRHICmdFlushRenderThreadTasksBasePass(
+	TEXT("r.RHICmdFlushRenderThreadTasksBasePass"),
+	0,
+	TEXT("Wait for completion of parallel render thread tasks at the end of the base pass. A more granular version of r.RHICmdFlushRenderThreadTasks. If either r.RHICmdFlushRenderThreadTasks or r.RHICmdFlushRenderThreadTasksBasePass is > 0 we will flush."));
+
 /*-----------------------------------------------------------------------------
 	FDeferredShadingSceneRenderer
 -----------------------------------------------------------------------------*/
@@ -114,7 +119,7 @@ void FDeferredShadingSceneRenderer::ClearGBufferAtMaxZ(FRHICommandList& RHICmdLi
 	FLinearColor ClearColors[6] = 
 		{bClearBlack ? FLinearColor(0,0,0,0) : Views[0].BackgroundColor, FLinearColor(0.5f,0.5f,0.5f,0), FLinearColor(0,0,0,1), FLinearColor(0,0,0,0), FLinearColor(0,1,1,1), FLinearColor(1,1,1,1)};
 
-	uint32 NumActiveRenderTargets = GSceneRenderTargets.GetNumGBufferTargets();
+	uint32 NumActiveRenderTargets = FSceneRenderTargets::Get(RHICmdList).GetNumGBufferTargets();
 	
 	auto ShaderMap = GetGlobalShaderMap(FeatureLevel);
 
@@ -402,20 +407,17 @@ class FRenderBasePassDynamicDataThreadTask
 	FDeferredShadingSceneRenderer& ThisRenderer;
 	FRHICommandList& RHICmdList;
 	const FViewInfo& View;
-	bool& OutDirty;
 
 public:
 
 	FRenderBasePassDynamicDataThreadTask(
 		FDeferredShadingSceneRenderer& InThisRenderer,
 		FRHICommandList& InRHICmdList,
-		const FViewInfo& InView,
-		bool& InOutDirty
+		const FViewInfo& InView
 		)
 		: ThisRenderer(InThisRenderer)
 		, RHICmdList(InRHICmdList)
 		, View(InView)
-		, OutDirty(InOutDirty)
 	{
 	}
 
@@ -433,6 +435,7 @@ public:
 
 	void DoTask(ENamedThreads::Type CurrentThread, const FGraphEventRef& MyCompletionGraphEvent)
 	{
+		bool OutDirty = false;
 		ThisRenderer.RenderBasePassDynamicData(RHICmdList, View, OutDirty);
 		RHICmdList.HandleRTThreadTaskCompletion(MyCompletionGraphEvent);
 	}
@@ -441,8 +444,8 @@ public:
 void FDeferredShadingSceneRenderer::RenderBasePassDynamicDataParallel(FParallelCommandListSet& ParallelCommandListSet)
 {
 	FRHICommandList* CmdList = ParallelCommandListSet.NewParallelCommandList();
-	FGraphEventRef AnyThreadCompletionEvent = TGraphTask<FRenderBasePassDynamicDataThreadTask>::CreateTask(nullptr, ENamedThreads::RenderThread)
-		.ConstructAndDispatchWhenReady(*this, *CmdList, ParallelCommandListSet.View, ParallelCommandListSet.OutDirty);
+	FGraphEventRef AnyThreadCompletionEvent = TGraphTask<FRenderBasePassDynamicDataThreadTask>::CreateTask(ParallelCommandListSet.GetPrereqs(), ENamedThreads::RenderThread)
+		.ConstructAndDispatchWhenReady(*this, *CmdList, ParallelCommandListSet.View);
 
 	ParallelCommandListSet.AddParallelCommandList(CmdList, AnyThreadCompletionEvent);
 }
@@ -481,8 +484,8 @@ class FBasePassParallelCommandListSet : public FParallelCommandListSet
 public:
 	const FSceneViewFamily& ViewFamily;
 
-	FBasePassParallelCommandListSet(const FViewInfo& InView, FRHICommandList& InParentCmdList, bool* InOutDirty, bool bInParallelExecute, const FSceneViewFamily& InViewFamily)
-		: FParallelCommandListSet(InView, InParentCmdList, InOutDirty, bInParallelExecute)
+	FBasePassParallelCommandListSet(const FViewInfo& InView, FRHICommandListImmediate& InParentCmdList, bool bInParallelExecute, bool bInCreateSceneContext, const FSceneViewFamily& InViewFamily)
+		: FParallelCommandListSet(InView, InParentCmdList, bInParallelExecute, bInCreateSceneContext)
 		, ViewFamily(InViewFamily)
 	{
 		SetStateOnCommandList(ParentCmdList);
@@ -495,7 +498,7 @@ public:
 
 	virtual void SetStateOnCommandList(FRHICommandList& CmdList) override
 	{
-		GSceneRenderTargets.BeginRenderingGBuffer(CmdList, ERenderTargetLoadAction::ELoad, ERenderTargetLoadAction::ELoad);
+		FSceneRenderTargets::Get(CmdList).BeginRenderingGBuffer(CmdList, ERenderTargetLoadAction::ELoad, ERenderTargetLoadAction::ELoad);
 		SetupBasePassView(CmdList, View.ViewRect, !!ViewFamily.EngineShowFlags.ShaderComplexity);
 	}
 };
@@ -505,9 +508,12 @@ static TAutoConsoleVariable<int32> CVarRHICmdBasePassDeferredContexts(
 	1,
 	TEXT("True to use deferred contexts to parallelize base pass command list execution."));
 
-void FDeferredShadingSceneRenderer::RenderBasePassViewParallel(FViewInfo& View, FRHICommandList& ParentCmdList, bool& OutDirty)
+void FDeferredShadingSceneRenderer::RenderBasePassViewParallel(FViewInfo& View, FRHICommandListImmediate& ParentCmdList)
 {
-	FBasePassParallelCommandListSet ParallelSet(View, ParentCmdList, &OutDirty, CVarRHICmdBasePassDeferredContexts.GetValueOnRenderThread() > 0, ViewFamily);
+	FBasePassParallelCommandListSet ParallelSet(View, ParentCmdList, 
+		CVarRHICmdBasePassDeferredContexts.GetValueOnRenderThread() > 0, 
+		CVarRHICmdFlushRenderThreadTasksBasePass.GetValueOnRenderThread() == 0 && CVarRHICmdFlushRenderThreadTasks.GetValueOnRenderThread() == 0,
+		ViewFamily);
 
 	RenderBasePassStaticDataParallel(ParallelSet);
 	RenderBasePassDynamicDataParallel(ParallelSet);
@@ -550,7 +556,7 @@ void FDeferredShadingSceneRenderer::RenderVisualizeTexturePool(FRHICommandListIm
 
 	RHICmdList.UnlockTexture2D((FTexture2DRHIRef&)VisualizeTexturePool->GetRenderTargetItem().ShaderResourceTexture, 0, false);
 
-	FIntPoint RTExtent = GSceneRenderTargets.GetBufferSizeXY();
+	FIntPoint RTExtent = FSceneRenderTargets::Get(RHICmdList).GetBufferSizeXY();
 
 	FVector2D Tex00 = FVector2D(0, 0);
 	FVector2D Tex11 = FVector2D(1, 1);
@@ -580,7 +586,7 @@ void FDeferredShadingSceneRenderer::RenderFinish(FRHICommandListImmediate& RHICm
 	// Some RT should be released as early as possible to allow sharing of that memory for other purposes.
 	// SceneColor is be released in tone mapping, if not we want to get access to the HDR scene color after this pass so we keep it.
 	// This becomes even more important with some limited VRam (XBoxOne).
-	GSceneRenderTargets.SetLightAttenuation(0);
+	FSceneRenderTargets::Get(RHICmdList).SetLightAttenuation(0);
 }
 
 void BuildHZB( FRHICommandListImmediate& RHICmdList, FViewInfo& View );
@@ -623,7 +629,7 @@ static void SetAndClearViewGBuffer(FRHICommandListImmediate& RHICmdList, FViewIn
 	const FLinearColor ClearColor = (bClearBlack ? FLinearColor(0, 0, 0, 0) : View.BackgroundColor);
 
 	// clearing the GBuffer
-	GSceneRenderTargets.BeginRenderingGBuffer(RHICmdList, ERenderTargetLoadAction::EClear, DepthLoadAction, ClearColor);
+	FSceneRenderTargets::Get(RHICmdList).BeginRenderingGBuffer(RHICmdList, ERenderTargetLoadAction::EClear, DepthLoadAction, ClearColor);
 }
 
 static TAutoConsoleVariable<int32> CVarOcclusionQueryLocation(
@@ -676,18 +682,27 @@ void FDeferredShadingSceneRenderer::RenderOcclusion(FRHICommandListImmediate& RH
 		if (bRenderQueries && GRHIThread)
 		{
 			QUICK_SCOPE_CYCLE_COUNTER(STAT_OcclusionSubmittedFence_Dispatch);
-			OcclusionSubmittedFence = FRHICommandListExecutor::RHIThreadFence();
+			OcclusionSubmittedFence = RHICmdList.RHIThreadFence();
 		}
 	}
 }
 
+// The render thread is involved in sending stuff to the RHI, so we will periodically service that queue
+static void ServiceLocalQueue()
+{
+	QUICK_SCOPE_CYCLE_COUNTER(STAT_FDeferredShadingSceneRenderer_Render_ServiceLocalQueue);
+	FTaskGraphInterface::Get().ProcessThreadUntilIdle(ENamedThreads::RenderThread_Local);
+}
+
 void FDeferredShadingSceneRenderer::Render(FRHICommandListImmediate& RHICmdList)
 {
+	FSceneRenderTargets& SceneContext = FSceneRenderTargets::Get(RHICmdList);
 	bool bDBuffer = IsDBufferEnabled();	
 	if (GRHIThread)
 	{
 		QUICK_SCOPE_CYCLE_COUNTER(STAT_OcclusionSubmittedFence_Wait);
 		FRHICommandListExecutor::WaitOnRHIThreadFence(OcclusionSubmittedFence);
+		OcclusionSubmittedFence = nullptr;
 	}
 
 	if(!ViewFamily.EngineShowFlags.Rendering)
@@ -703,16 +718,18 @@ void FDeferredShadingSceneRenderer::Render(FRHICommandListImmediate& RHICmdList)
 		GSystemTextures.InitializeTextures(RHICmdList, FeatureLevel);
 
 		// Allocate the maximum scene render target space for the current view family.
-		GSceneRenderTargets.Allocate(ViewFamily);
+		SceneContext.Allocate(ViewFamily);
 	}
 	// Find the visible primitives.
 	InitViews(RHICmdList);
 
-	if (GRHIThread)
+	if (GRHICommandList.UseParallelAlgorithms())
 	{
-		// we will probably stall on occlusion queries, so might as well have the RHI thread and GPU work while we wait.
-		QUICK_SCOPE_CYCLE_COUNTER(STAT_PostInitViews_FlushDel);
-		FRHICommandListExecutor::GetImmediateCommandList().ImmediateFlush(EImmediateFlushType::FlushRHIThreadFlushResources);
+		// there are dynamic attempts to get this target during parallel rendering
+		for (int32 ViewIndex = 0; ViewIndex < Views.Num(); ViewIndex++)
+	{
+			Views[ViewIndex].GetEyeAdaptation();
+		}	
 	}
 
 	if (ShouldPrepareDistanceFields())
@@ -734,6 +751,13 @@ void FDeferredShadingSceneRenderer::Render(FRHICommandListImmediate& RHICmdList)
 		}	
 	}
 
+	if (GRHIThread)
+	{
+		// we will probably stall on occlusion queries, so might as well have the RHI thread and GPU work while we wait.
+		QUICK_SCOPE_CYCLE_COUNTER(STAT_PostInitViews_FlushDel);
+		FRHICommandListExecutor::GetImmediateCommandList().ImmediateFlush(EImmediateFlushType::FlushRHIThreadFlushResources);
+	}
+
 	const bool bIsWireframe = ViewFamily.EngineShowFlags.Wireframe;
 	static const auto ClearMethodCVar = IConsoleManager::Get().FindTConsoleVariableDataInt(TEXT("r.ClearSceneMethod"));
 	bool bRequiresRHIClear = true;
@@ -815,6 +839,7 @@ void FDeferredShadingSceneRenderer::Render(FRHICommandListImmediate& RHICmdList)
 	// Notify the FX system that the scene is about to be rendered.
 	if (Scene->FXSystem)
 	{
+		QUICK_SCOPE_CYCLE_COUNTER(STAT_FDeferredShadingSceneRenderer_FXSystem_PreRender);
 		Scene->FXSystem->PreRender(RHICmdList, &Views[0].GlobalDistanceFieldInfo.ParameterData);
 	}
 
@@ -831,19 +856,36 @@ void FDeferredShadingSceneRenderer::Render(FRHICommandListImmediate& RHICmdList)
 		bDepthWasCleared = true;
 	}
 	RHICmdList.SetCurrentStat(GET_STATID(STAT_CLM_AfterPrePass));
-
+	ServiceLocalQueue();
 	//occlusion can't run before basepass if there's no prepass to fill in some depth to occlude against.
 	bool bOcclusionBeforeBasePass = (CVarOcclusionQueryLocation.GetValueOnRenderThread() == 1) && bNeedsPrePass;	
 	bool bHZBBeforeBasePass = false;
 	RenderOcclusion(RHICmdList, bOcclusionBeforeBasePass, bHZBBeforeBasePass);
-	
+	ServiceLocalQueue();	
 	const bool bShouldRenderVelocities = ShouldRenderVelocities();
 	const bool bUseVelocityGBuffer = FVelocityRendering::OutputsToGBuffer();
 
+	if (bShouldRenderVelocities)
+	{
+		static bool bOnce = false;
+		if (!bOnce)
+		{
+			bOnce = true;
+			GPrevPerBoneMotionBlur.SetVelocityPassCallback(			
+				// this is a strange intermodule bridge so that the skeletal mesh vertex factory knows when to add bone data in a parallel thread
+				[](FRHICommandList& RHICmdList)->bool
+			{
+				return FSceneRenderTargets::Get(RHICmdList).IsVelocityPass();
+			});
+		}
+	}
+
+
+
 	{
 		QUICK_SCOPE_CYCLE_COUNTER(STAT_FDeferredShadingSceneRenderer_AllocGBufferTargets);
-		GSceneRenderTargets.PreallocGBufferTargets(bShouldRenderVelocities && bUseVelocityGBuffer);
-		GSceneRenderTargets.AllocGBufferTargets();
+		SceneContext.PreallocGBufferTargets(bShouldRenderVelocities && bUseVelocityGBuffer);
+		SceneContext.AllocGBufferTargets();
 	}
 	
 	// Clear LPVs for all views
@@ -851,18 +893,19 @@ void FDeferredShadingSceneRenderer::Render(FRHICommandListImmediate& RHICmdList)
 	{
 		QUICK_SCOPE_CYCLE_COUNTER(STAT_FDeferredShadingSceneRenderer_ClearLPVs);
 		ClearLPVs(RHICmdList);
+		ServiceLocalQueue();
 	}
 
 	// only temporarily available after early z pass and until base pass
-	check(!GSceneRenderTargets.DBufferA);
-	check(!GSceneRenderTargets.DBufferB);
-	check(!GSceneRenderTargets.DBufferC);
+	check(!SceneContext.DBufferA);
+	check(!SceneContext.DBufferB);
+	check(!SceneContext.DBufferC);
 
 	if(bDBuffer)
 	{
 		QUICK_SCOPE_CYCLE_COUNTER(STAT_FDeferredShadingSceneRenderer_DBuffer);
-		GSceneRenderTargets.ResolveSceneDepthTexture(RHICmdList);
-		GSceneRenderTargets.ResolveSceneDepthToAuxiliaryTexture(RHICmdList);
+		SceneContext.ResolveSceneDepthTexture(RHICmdList);
+		SceneContext.ResolveSceneDepthToAuxiliaryTexture(RHICmdList);
 
 		// e.g. DBuffer deferred decals
 		for(int32 ViewIndex = 0;ViewIndex < Views.Num();ViewIndex++)
@@ -871,6 +914,7 @@ void FDeferredShadingSceneRenderer::Render(FRHICommandListImmediate& RHICmdList)
 
 			GCompositionLighting.ProcessBeforeBasePass(RHICmdList, Views[ViewIndex]);
 		}
+		ServiceLocalQueue();
 	}
 
 	// Clear the G Buffer render targets
@@ -884,19 +928,20 @@ void FDeferredShadingSceneRenderer::Render(FRHICommandListImmediate& RHICmdList)
 		// depth was cleared now no matter what
 		bDepthWasCleared = true;
 		bIsGBufferCurrent = true;
+		ServiceLocalQueue();
 	}
 
 	if(bIsWireframe && FSceneRenderer::ShouldCompositeEditorPrimitives(Views[0]))
 	{
 		// In Editor we want wire frame view modes to be MSAA for better quality. Resolve will be done with EditorPrimitives
-		SetRenderTarget(RHICmdList, GSceneRenderTargets.GetEditorPrimitivesColor(), GSceneRenderTargets.GetEditorPrimitivesDepth());
+		SetRenderTarget(RHICmdList, SceneContext.GetEditorPrimitivesColor(), SceneContext.GetEditorPrimitivesDepth());
 		RHICmdList.Clear(true, FLinearColor(0, 0, 0, 0), true, (float)ERHIZBuffer::FarPlane, false, 0, FIntRect());
 	}
 	else if (!bIsGBufferCurrent)
 	{
 		// make sure the GBuffer is set, in case we didn't need to clear above
 		ERenderTargetLoadAction DepthLoadAction = bDepthWasCleared ? ERenderTargetLoadAction::ELoad : ERenderTargetLoadAction::EClear;
-		GSceneRenderTargets.BeginRenderingGBuffer(RHICmdList, ERenderTargetLoadAction::ENoAction, DepthLoadAction);
+		SceneContext.BeginRenderingGBuffer(RHICmdList, ERenderTargetLoadAction::ENoAction, DepthLoadAction);
 	}
 
 	GRenderTargetPool.AddPhaseEvent(TEXT("BasePass"));
@@ -904,22 +949,23 @@ void FDeferredShadingSceneRenderer::Render(FRHICommandListImmediate& RHICmdList)
 	RHICmdList.SetCurrentStat(GET_STATID(STAT_CLM_BasePass));
 	RenderBasePass(RHICmdList);
 	RHICmdList.SetCurrentStat(GET_STATID(STAT_CLM_AfterBasePass));
+	ServiceLocalQueue();
 
 	  if(ViewFamily.EngineShowFlags.VisualizeLightCulling)
 	  {
 		  // clear out emissive and baked lighting (not too efficient but simple and only needed for this debug view)
-		  GSceneRenderTargets.BeginRenderingSceneColor(RHICmdList);
+		  SceneContext.BeginRenderingSceneColor(RHICmdList);
 		  RHICmdList.Clear(true, FLinearColor(0, 0, 0, 0), false, (float)ERHIZBuffer::FarPlane, false, 0, FIntRect());
 	  }
 
-	  GSceneRenderTargets.DBufferA.SafeRelease();
-	  GSceneRenderTargets.DBufferB.SafeRelease();
-	  GSceneRenderTargets.DBufferC.SafeRelease();
+	  SceneContext.DBufferA.SafeRelease();
+	  SceneContext.DBufferB.SafeRelease();
+	  SceneContext.DBufferC.SafeRelease();
 
 	  // only temporarily available after early z pass and until base pass
-	  check(!GSceneRenderTargets.DBufferA);
-	  check(!GSceneRenderTargets.DBufferB);
-	  check(!GSceneRenderTargets.DBufferC);
+	  check(!SceneContext.DBufferA);
+	  check(!SceneContext.DBufferB);
+	  check(!SceneContext.DBufferC);
 
 	if (bRequiresFarZQuadClear)
 	{
@@ -927,6 +973,7 @@ void FDeferredShadingSceneRenderer::Render(FRHICommandListImmediate& RHICmdList)
 		// Clears view by drawing quad at maximum Z
 		// TODO: if all the platforms have fast color clears, we can replace this with an RHICmdList.Clear.
 		ClearGBufferAtMaxZ(RHICmdList);
+		ServiceLocalQueue();
 
 		bRequiresFarZQuadClear = false;
 	}
@@ -934,34 +981,38 @@ void FDeferredShadingSceneRenderer::Render(FRHICommandListImmediate& RHICmdList)
 	{
 		QUICK_SCOPE_CYCLE_COUNTER(STAT_FDeferredShadingSceneRenderer_Resolve_After_Basepass);
 
-		GSceneRenderTargets.ResolveSceneColor(RHICmdList, FResolveRect(0, 0, ViewFamily.FamilySizeX, ViewFamily.FamilySizeY));
-		GSceneRenderTargets.ResolveSceneDepthTexture(RHICmdList);
-		GSceneRenderTargets.ResolveSceneDepthToAuxiliaryTexture(RHICmdList);
+		SceneContext.ResolveSceneColor(RHICmdList, FResolveRect(0, 0, ViewFamily.FamilySizeX, ViewFamily.FamilySizeY));
+		SceneContext.ResolveSceneDepthTexture(RHICmdList);
+		SceneContext.ResolveSceneDepthToAuxiliaryTexture(RHICmdList);
 
-		GSceneRenderTargets.FinishRenderingGBuffer(RHICmdList);
+		SceneContext.FinishRenderingGBuffer(RHICmdList);
 
 		RenderCustomDepthPass(RHICmdList);
+		ServiceLocalQueue();
 	}
 	// Notify the FX system that opaque primitives have been rendered and we now have a valid depth buffer.
 	if (Scene->FXSystem && Views.IsValidIndex(0))
 	{
+		QUICK_SCOPE_CYCLE_COUNTER(STAT_FDeferredShadingSceneRenderer_FXSystem_PostRenderOpaque);
 		Scene->FXSystem->PostRenderOpaque(
 			RHICmdList,
 			Views.GetData(),
-			GSceneRenderTargets.GetSceneDepthTexture(),
-			GSceneRenderTargets.GetGBufferATexture()
+			SceneContext.GetSceneDepthTexture(),
+			SceneContext.GetGBufferATexture()
 			);
+		ServiceLocalQueue();
 	}
 
 	bool bOcclusionAfterBasePass = bIsOcclusionTesting && !bOcclusionBeforeBasePass;
 	bool bHZBAfterBasePass = true;
 	RenderOcclusion(RHICmdList, bOcclusionAfterBasePass, bHZBAfterBasePass);
+	ServiceLocalQueue();
 
 	TRefCountPtr<IPooledRenderTarget> VelocityRT;
 
 	if (bUseVelocityGBuffer)
 	{
-		VelocityRT = GSceneRenderTargets.GetGBufferVelocityRT();
+		VelocityRT = SceneContext.GetGBufferVelocityRT();
 	}
 	else if (bShouldRenderVelocities)
 	{
@@ -969,6 +1020,7 @@ void FDeferredShadingSceneRenderer::Render(FRHICommandListImmediate& RHICmdList)
 		RHICmdList.SetCurrentStat(GET_STATID(STAT_CLM_Velocity));
 		RenderVelocities(RHICmdList, VelocityRT);
 		RHICmdList.SetCurrentStat(GET_STATID(STAT_CLM_AfterVelocity));
+		ServiceLocalQueue();
 	}
 
 	// Pre-lighting composition lighting stage
@@ -985,6 +1037,7 @@ void FDeferredShadingSceneRenderer::Render(FRHICommandListImmediate& RHICmdList)
 			SCOPED_CONDITIONAL_DRAW_EVENTF(RHICmdList, EventView, Views.Num() > 1, TEXT("View%d"), ViewIndex);
 			GCompositionLighting.ProcessAfterBasePass(RHICmdList, Views[ViewIndex]);
 		}
+		ServiceLocalQueue();
 	}
 
 	// Render lighting.
@@ -1004,30 +1057,36 @@ void FDeferredShadingSceneRenderer::Render(FRHICommandListImmediate& RHICmdList)
 		RHICmdList.SetCurrentStat(GET_STATID(STAT_CLM_Lighting));
 		RenderLights(RHICmdList);
 		RHICmdList.SetCurrentStat(GET_STATID(STAT_CLM_AfterLighting));
+		ServiceLocalQueue();
 
 		GRenderTargetPool.AddPhaseEvent(TEXT("AfterRenderLights"));
 
 		InjectAmbientCubemapTranslucentVolumeLighting(RHICmdList);
+		ServiceLocalQueue();
 
 		// Filter the translucency lighting volume now that it is complete
 		FilterTranslucentVolumeLighting(RHICmdList);
+		ServiceLocalQueue();
 
 		// Pre-lighting composition lighting stage
 		// e.g. LPV indirect
 		for(int32 ViewIndex = 0; ViewIndex < Views.Num(); ++ViewIndex)
-		{	
+		{
 			SCOPED_CONDITIONAL_DRAW_EVENTF(RHICmdList, EventView,Views.Num() > 1, TEXT("View%d"), ViewIndex);
 			GCompositionLighting.ProcessLpvIndirect(RHICmdList, Views[ViewIndex]);
+			ServiceLocalQueue();
 		}
 
 		TRefCountPtr<IPooledRenderTarget> DynamicBentNormalAO;
 		RenderDynamicSkyLighting(RHICmdList, VelocityRT, DynamicBentNormalAO);
+		ServiceLocalQueue();
 
 		// SSS need the SceneColor finalized as an SRV.
-		GSceneRenderTargets.FinishRenderingSceneColor(RHICmdList, true);
+		SceneContext.FinishRenderingSceneColor(RHICmdList, true);
 
 		// Render reflections that only operate on opaque pixels
 		RenderDeferredReflections(RHICmdList, DynamicBentNormalAO);
+		ServiceLocalQueue();
 
 		// Post-lighting composition lighting stage
 		// e.g. ScreenSpaceSubsurfaceScattering
@@ -1036,12 +1095,14 @@ void FDeferredShadingSceneRenderer::Render(FRHICommandListImmediate& RHICmdList)
 			SCOPED_CONDITIONAL_DRAW_EVENTF(RHICmdList, EventView,Views.Num() > 1, TEXT("View%d"), ViewIndex);
 			GCompositionLighting.ProcessAfterLighting(RHICmdList, Views[ViewIndex]);
 		}
+		ServiceLocalQueue();
 	}
 
 	if( ViewFamily.EngineShowFlags.StationaryLightOverlap &&
 		FeatureLevel >= ERHIFeatureLevel::SM4)
 	{
 		RenderStationaryLightOverlap(RHICmdList);
+		ServiceLocalQueue();
 	}
 
 	FLightShaftsOutput LightShaftOutput;
@@ -1051,6 +1112,7 @@ void FDeferredShadingSceneRenderer::Render(FRHICommandListImmediate& RHICmdList)
 	{
 		QUICK_SCOPE_CYCLE_COUNTER(STAT_FDeferredShadingSceneRenderer_RenderLightShaftOcclusion);
 		LightShaftOutput = RenderLightShaftOcclusion(RHICmdList);
+		ServiceLocalQueue();
 	}
 
 	// Draw atmosphere
@@ -1076,6 +1138,7 @@ void FDeferredShadingSceneRenderer::Render(FRHICommandListImmediate& RHICmdList)
 			}
 #endif
 			RenderAtmosphere(RHICmdList, LightShaftOutput);
+			ServiceLocalQueue();
 		}
 	}
 
@@ -1086,6 +1149,7 @@ void FDeferredShadingSceneRenderer::Render(FRHICommandListImmediate& RHICmdList)
 	{
 		QUICK_SCOPE_CYCLE_COUNTER(STAT_FDeferredShadingSceneRenderer_RenderFog);
 		RenderFog(RHICmdList, LightShaftOutput);
+		ServiceLocalQueue();
 	}
 
 	// No longer needed, release
@@ -1101,6 +1165,7 @@ void FDeferredShadingSceneRenderer::Render(FRHICommandListImmediate& RHICmdList)
 		RHICmdList.SetCurrentStat(GET_STATID(STAT_CLM_Translucency));
 		RenderTranslucency(RHICmdList);
 		RHICmdList.SetCurrentStat(GET_STATID(STAT_CLM_AfterTranslucency));
+		ServiceLocalQueue();
 
 		if (ViewFamily.EngineShowFlags.Refraction)
 		{
@@ -1108,12 +1173,15 @@ void FDeferredShadingSceneRenderer::Render(FRHICommandListImmediate& RHICmdList)
 			// After non separate translucency as that is considered at scene depth anyway
 			// It allows skybox translucency (set to non separate translucency) to be refracted.
 			RenderDistortion(RHICmdList);
+			ServiceLocalQueue();
 		}
 	}
+
 	if (ViewFamily.EngineShowFlags.LightShafts)
 	{
 		QUICK_SCOPE_CYCLE_COUNTER(STAT_FDeferredShadingSceneRenderer_RenderLightShaftBloom);
 		RenderLightShaftBloom(RHICmdList);
+		ServiceLocalQueue();
 	}
 
 	if (ViewFamily.EngineShowFlags.VisualizeDistanceFieldAO || ViewFamily.EngineShowFlags.VisualizeDistanceFieldGI)
@@ -1122,15 +1190,19 @@ void FDeferredShadingSceneRenderer::Render(FRHICommandListImmediate& RHICmdList)
 		const float OcclusionMaxDistance = Scene->SkyLight && !Scene->SkyLight->bWantsStaticShadowing ? Scene->SkyLight->OcclusionMaxDistance : GDefaultDFAOMaxOcclusionDistance;
 		TRefCountPtr<IPooledRenderTarget> DummyOutput;
 		RenderDistanceFieldLighting(RHICmdList, FDistanceFieldAOParameters(OcclusionMaxDistance), VelocityRT, DummyOutput, DummyOutput, ViewFamily.EngineShowFlags.VisualizeDistanceFieldAO, ViewFamily.EngineShowFlags.VisualizeDistanceFieldGI);
+		ServiceLocalQueue();
 	}
 
 	if (ViewFamily.EngineShowFlags.VisualizeMeshDistanceFields)
 	{
 		RenderMeshDistanceFieldVisualization(RHICmdList, FDistanceFieldAOParameters(GDefaultDFAOMaxOcclusionDistance));
+		ServiceLocalQueue();
 	}
 
 	// Resolve the scene color for post processing.
-	GSceneRenderTargets.ResolveSceneColor(RHICmdList, FResolveRect(0, 0, ViewFamily.FamilySizeX, ViewFamily.FamilySizeY));
+	SceneContext.ResolveSceneColor(RHICmdList, FResolveRect(0, 0, ViewFamily.FamilySizeX, ViewFamily.FamilySizeY));
+
+	GPrevPerBoneMotionBlur.EndAppend(RHICmdList);
 
 	// Finish rendering for each view.
 	if(ViewFamily.bResolveScene)
@@ -1152,7 +1224,7 @@ void FDeferredShadingSceneRenderer::Render(FRHICommandListImmediate& RHICmdList)
 	else
 	{
 		// Release the original reference on the scene render targets
-		GSceneRenderTargets.AdjustGBufferRefCount(-1);
+		SceneContext.AdjustGBufferRefCount(-1);
 	}
 
 	//grab the new transform out of the proxies for next frame
@@ -1169,6 +1241,7 @@ void FDeferredShadingSceneRenderer::Render(FRHICommandListImmediate& RHICmdList)
 		RenderFinish(RHICmdList);
 		RHICmdList.SetCurrentStat(GET_STATID(STAT_CLM_AfterFrame));
 	}
+	ServiceLocalQueue();
 }
 
 bool FDeferredShadingSceneRenderer::RenderPrePassViewDynamic(FRHICommandList& RHICmdList, const FViewInfo& View)
@@ -1254,20 +1327,17 @@ class FRenderPrepassDynamicDataThreadTask
 	FDeferredShadingSceneRenderer& ThisRenderer;
 	FRHICommandList& RHICmdList;
 	const FViewInfo& View;
-	bool& OutDirty;
 
 public:
 
 	FRenderPrepassDynamicDataThreadTask(
 		FDeferredShadingSceneRenderer& InThisRenderer,
 		FRHICommandList& InRHICmdList,
-		const FViewInfo& InView,
-		bool& InOutDirty
+		const FViewInfo& InView
 		)
 		: ThisRenderer(InThisRenderer)
 		, RHICmdList(InRHICmdList)
 		, View(InView)
-		, OutDirty(InOutDirty)
 	{
 	}
 
@@ -1285,10 +1355,7 @@ public:
 
 	void DoTask(ENamedThreads::Type CurrentThread, const FGraphEventRef& MyCompletionGraphEvent)
 	{
-		if (ThisRenderer.RenderPrePassViewDynamic(RHICmdList, View))
-		{
-			OutDirty = true;
-		}
+		ThisRenderer.RenderPrePassViewDynamic(RHICmdList, View);
 		RHICmdList.HandleRTThreadTaskCompletion(MyCompletionGraphEvent);
 	}
 };
@@ -1296,8 +1363,8 @@ public:
 class FPrePassParallelCommandListSet : public FParallelCommandListSet
 {
 public:
-	FPrePassParallelCommandListSet(const FViewInfo& InView, FRHICommandList& InParentCmdList, bool* InOutDirty, bool bInParallelExecute)
-		: FParallelCommandListSet(InView, InParentCmdList, InOutDirty, bInParallelExecute)
+	FPrePassParallelCommandListSet(const FViewInfo& InView, FRHICommandListImmediate& InParentCmdList, bool bInParallelExecute, bool bInCreateSceneContext)
+		: FParallelCommandListSet(InView, InParentCmdList, bInParallelExecute, bInCreateSceneContext)
 	{
 		SetStateOnCommandList(ParentCmdList);
 	}
@@ -1309,7 +1376,7 @@ public:
 
 	virtual void SetStateOnCommandList(FRHICommandList& CmdList) override
 	{
-		GSceneRenderTargets.BeginRenderingPrePass(CmdList, false);
+		FSceneRenderTargets::Get(CmdList).BeginRenderingPrePass(CmdList, false);
 		SetupPrePassView(CmdList, View.ViewRect);
 	}
 };
@@ -1318,10 +1385,23 @@ static TAutoConsoleVariable<int32> CVarRHICmdPrePassDeferredContexts(
 	TEXT("r.RHICmdPrePassDeferredContexts"),
 	1,
 	TEXT("True to use deferred contexts to parallelize prepass command list execution."));
+static TAutoConsoleVariable<int32> CVarParallelPrePass(
+	TEXT("r.ParallelPrePass"),
+	1,
+	TEXT("Toggles parallel zprepass rendering. Parallel rendering must be enabled for this to have an effect."),
+	ECVF_RenderThreadSafe
+	);
+static TAutoConsoleVariable<int32> CVarRHICmdFlushRenderThreadTasksPrePass(
+	TEXT("r.RHICmdFlushRenderThreadTasksPrePass"),
+	0,
+	TEXT("Wait for completion of parallel render thread tasks at the end of the pre pass.  A more granular version of r.RHICmdFlushRenderThreadTasks. If either r.RHICmdFlushRenderThreadTasks or r.RHICmdFlushRenderThreadTasksPrePass is > 0 we will flush."));
+
 
-void FDeferredShadingSceneRenderer::RenderPrePassViewParallel(const FViewInfo& View, FRHICommandList& ParentCmdList, bool& OutDirty)
+void FDeferredShadingSceneRenderer::RenderPrePassViewParallel(const FViewInfo& View, FRHICommandListImmediate& ParentCmdList)
 {
-	FPrePassParallelCommandListSet ParallelCommandListSet(View, ParentCmdList, &OutDirty, CVarRHICmdPrePassDeferredContexts.GetValueOnRenderThread() > 0);
+	FPrePassParallelCommandListSet ParallelCommandListSet(View, ParentCmdList, 
+		CVarRHICmdPrePassDeferredContexts.GetValueOnRenderThread() > 0, 
+		CVarRHICmdFlushRenderThreadTasksPrePass.GetValueOnRenderThread() == 0  && CVarRHICmdFlushRenderThreadTasks.GetValueOnRenderThread() == 0);
 
 
 	// Draw the static occluder primitives using a depth drawing policy.
@@ -1350,20 +1430,13 @@ void FDeferredShadingSceneRenderer::RenderPrePassViewParallel(const FViewInfo& V
 	{
 		FRHICommandList* CmdList = ParallelCommandListSet.NewParallelCommandList();
 
-		FGraphEventRef AnyThreadCompletionEvent = TGraphTask<FRenderPrepassDynamicDataThreadTask>::CreateTask(nullptr, ENamedThreads::RenderThread)
-			.ConstructAndDispatchWhenReady(*this, *CmdList, View, OutDirty);
+		FGraphEventRef AnyThreadCompletionEvent = TGraphTask<FRenderPrepassDynamicDataThreadTask>::CreateTask(ParallelCommandListSet.GetPrereqs(), ENamedThreads::RenderThread)
+			.ConstructAndDispatchWhenReady(*this, *CmdList, View);
 
 		ParallelCommandListSet.AddParallelCommandList(CmdList, AnyThreadCompletionEvent);
 	}
 }
 
-static TAutoConsoleVariable<int32> CVarParallelPrePass(
-	TEXT("r.ParallelPrePass"),
-	1,
-	TEXT("Toggles parallel zprepass rendering. Parallel rendering must be enabled for this to have an effect."),
-	ECVF_RenderThreadSafe
-	);
-
 /** Renders the scene's prepass and occlusion queries */
 bool FDeferredShadingSceneRenderer::RenderPrePass(FRHICommandListImmediate& RHICmdList)
 {
@@ -1371,21 +1444,23 @@ bool FDeferredShadingSceneRenderer::RenderPrePass(FRHICommandListImmediate& RHIC
 	SCOPE_CYCLE_COUNTER(STAT_DepthDrawTime);
 
 	bool bDirty = false;
+	FSceneRenderTargets& SceneContext = FSceneRenderTargets::Get(RHICmdList);
 
-	GSceneRenderTargets.BeginRenderingPrePass(RHICmdList, true);
+	SceneContext.BeginRenderingPrePass(RHICmdList, true);
 
 	// Draw a depth pass to avoid overdraw in the other passes.
 	if(EarlyZPassMode != DDM_None)
 	{
 		if (GRHICommandList.UseParallelAlgorithms() && CVarParallelPrePass.GetValueOnRenderThread())
 		{
-			FScopedCommandListWaitForTasks Flusher(RHICmdList);
+			FScopedCommandListWaitForTasks Flusher(CVarRHICmdFlushRenderThreadTasksPrePass.GetValueOnRenderThread() > 0 || CVarRHICmdFlushRenderThreadTasks.GetValueOnRenderThread() > 0, RHICmdList);
 
 			for(int32 ViewIndex = 0;ViewIndex < Views.Num();ViewIndex++)
 			{
 				SCOPED_CONDITIONAL_DRAW_EVENTF(RHICmdList, EventView, Views.Num() > 1, TEXT("View%d"), ViewIndex);
 				const FViewInfo& View = Views[ViewIndex];
-				RenderPrePassViewParallel(View, RHICmdList, bDirty);
+				RenderPrePassViewParallel(View, RHICmdList);
+				bDirty = true; // assume dirty since we are not going to wait
 			}
 		}
 		else
@@ -1399,7 +1474,7 @@ bool FDeferredShadingSceneRenderer::RenderPrePass(FRHICommandListImmediate& RHIC
 		}
 	}
 
-	GSceneRenderTargets.FinishRenderingPrePass(RHICmdList);
+	SceneContext.FinishRenderingPrePass(RHICmdList);
 
 	return bDirty;
 }
@@ -1422,7 +1497,8 @@ bool FDeferredShadingSceneRenderer::RenderBasePass(FRHICommandListImmediate& RHI
 	if (FVelocityRendering::OutputsToGBuffer())
 	{
 		QUICK_SCOPE_CYCLE_COUNTER(STAT_FDeferredShadingSceneRenderer_RenderBasePass_GPrevPerBoneMotionBlur_LockData);
-		GPrevPerBoneMotionBlur.StartAppend(ViewFamily.bWorldIsPaused);
+		FSceneRenderTargets::Get(RHICmdList).SetVelocityPass(true);
+		GPrevPerBoneMotionBlur.StartAppend(RHICmdList, ViewFamily.bWorldIsPaused);
 	}
 
 	if(ViewFamily.EngineShowFlags.LightMapDensity && AllowDebugViewmodes())
@@ -1437,12 +1513,18 @@ bool FDeferredShadingSceneRenderer::RenderBasePass(FRHICommandListImmediate& RHI
 
 		if (GRHICommandList.UseParallelAlgorithms() && CVarParallelBasePass.GetValueOnRenderThread())
 		{
-			FScopedCommandListWaitForTasks Flusher(RHICmdList);
+			FScopedCommandListWaitForTasks Flusher(CVarRHICmdFlushRenderThreadTasksBasePass.GetValueOnRenderThread() > 0 || CVarRHICmdFlushRenderThreadTasks.GetValueOnRenderThread() > 0, RHICmdList);
 			for(int32 ViewIndex = 0;ViewIndex < Views.Num();ViewIndex++)
 			{
 				SCOPED_CONDITIONAL_DRAW_EVENTF(RHICmdList, EventView, Views.Num() > 1, TEXT("View%d"), ViewIndex);
 				FViewInfo& View = Views[ViewIndex];
-				RenderBasePassViewParallel(View, RHICmdList, bDirty);
+				RenderBasePassViewParallel(View, RHICmdList);
+				bDirty = true; // assume dirty since we are not going to wait
+				if (FVelocityRendering::OutputsToGBuffer())
+				{
+					GPrevPerBoneMotionBlur.EndAppendFence(RHICmdList);
+					FSceneRenderTargets::Get(RHICmdList).SetVelocityPass(false);
+				}
 			}
 		}
 		else
@@ -1458,12 +1540,6 @@ bool FDeferredShadingSceneRenderer::RenderBasePass(FRHICommandListImmediate& RHI
 
 	}
 
-	if (FVelocityRendering::OutputsToGBuffer())
-	{
-		QUICK_SCOPE_CYCLE_COUNTER(STAT_FDeferredShadingSceneRenderer_RenderBasePass_GPrevPerBoneMotionBlur_UnlockData);
-		GPrevPerBoneMotionBlur.EndAppend();
-	}
-
 	return bDirty;
 }
 
@@ -1516,7 +1592,7 @@ void FDeferredShadingSceneRenderer::UpdateLPVs(FRHICommandListImmediate& RHICmdL
 			{
 //				SCOPED_DRAW_EVENT(RHICmdList, UpdateLPVs);
 //				SCOPE_CYCLE_COUNTER(STAT_UpdateLPVs);
-
+				
 				LightPropagationVolume->Update(RHICmdList, View);
 			}
 		}
@@ -1547,15 +1623,16 @@ public:
 	void SetParameters(FRHICommandList& RHICmdList, const FSceneView& View)
 	{
 		FGlobalShader::SetParameters(RHICmdList, GetPixelShader(), View);
+		FSceneRenderTargets& SceneContext = FSceneRenderTargets::Get(RHICmdList);
 
 		// Used to remap view space Z (which is stored in scene color alpha) into post projection z and w so we can write z/w into the downsized depth buffer
 		const FVector2D ProjectionScaleBiasValue(View.ViewMatrices.ProjMatrix.M[2][2], View.ViewMatrices.ProjMatrix.M[3][2]);
 		SetShaderValue(RHICmdList, GetPixelShader(), ProjectionScaleBias, ProjectionScaleBiasValue);
 
-		FIntPoint BufferSize = GSceneRenderTargets.GetBufferSizeXY();
+		FIntPoint BufferSize = SceneContext.GetBufferSizeXY();
 
-		const uint32 DownsampledBufferSizeX = BufferSize.X / GSceneRenderTargets.GetSmallColorDepthDownsampleFactor();
-		const uint32 DownsampledBufferSizeY = BufferSize.Y / GSceneRenderTargets.GetSmallColorDepthDownsampleFactor();
+		const uint32 DownsampledBufferSizeX = BufferSize.X / SceneContext.GetSmallColorDepthDownsampleFactor();
+		const uint32 DownsampledBufferSizeY = BufferSize.Y / SceneContext.GetSmallColorDepthDownsampleFactor();
 
 		// Offsets of the four full resolution pixels corresponding with a low resolution pixel
 		const FVector4 Offsets01(0.0f, 0.0f, 1.0f / DownsampledBufferSizeX, 0.0f);
@@ -1588,9 +1665,10 @@ FGlobalBoundShaderState DownsampleDepthBoundShaderState;
 /** Updates the downsized depth buffer with the current full resolution depth buffer. */
 void FDeferredShadingSceneRenderer::UpdateDownsampledDepthSurface(FRHICommandList& RHICmdList)
 {
-	if (GSceneRenderTargets.UseDownsizedOcclusionQueries() && (FeatureLevel >= ERHIFeatureLevel::SM4))
+	FSceneRenderTargets& SceneContext = FSceneRenderTargets::Get(RHICmdList);
+	if (SceneContext.UseDownsizedOcclusionQueries() && (FeatureLevel >= ERHIFeatureLevel::SM4))
 	{
-		SetRenderTarget(RHICmdList, NULL, GSceneRenderTargets.GetSmallDepthSurface());
+		SetRenderTarget(RHICmdList, NULL, SceneContext.GetSmallDepthSurface());
 
 		SCOPED_DRAW_EVENT(RHICmdList, DownsampleDepth);
 
@@ -1611,10 +1689,10 @@ void FDeferredShadingSceneRenderer::UpdateDownsampledDepthSurface(FRHICommandLis
 
 			PixelShader->SetParameters(RHICmdList, View);
 
-			const uint32 DownsampledX = FMath::TruncToInt(View.ViewRect.Min.X / GSceneRenderTargets.GetSmallColorDepthDownsampleFactor());
-			const uint32 DownsampledY = FMath::TruncToInt(View.ViewRect.Min.Y / GSceneRenderTargets.GetSmallColorDepthDownsampleFactor());
-			const uint32 DownsampledSizeX = FMath::TruncToInt(View.ViewRect.Width() / GSceneRenderTargets.GetSmallColorDepthDownsampleFactor());
-			const uint32 DownsampledSizeY = FMath::TruncToInt(View.ViewRect.Height() / GSceneRenderTargets.GetSmallColorDepthDownsampleFactor());
+			const uint32 DownsampledX = FMath::TruncToInt(View.ViewRect.Min.X / SceneContext.GetSmallColorDepthDownsampleFactor());
+			const uint32 DownsampledY = FMath::TruncToInt(View.ViewRect.Min.Y / SceneContext.GetSmallColorDepthDownsampleFactor());
+			const uint32 DownsampledSizeX = FMath::TruncToInt(View.ViewRect.Width() / SceneContext.GetSmallColorDepthDownsampleFactor());
+			const uint32 DownsampledSizeY = FMath::TruncToInt(View.ViewRect.Height() / SceneContext.GetSmallColorDepthDownsampleFactor());
 
 			RHICmdList.SetViewport(DownsampledX, DownsampledY, 0.0f, DownsampledX + DownsampledSizeX, DownsampledY + DownsampledSizeY, 1.0f);
 
@@ -1625,7 +1703,7 @@ void FDeferredShadingSceneRenderer::UpdateDownsampledDepthSurface(FRHICommandLis
 				View.ViewRect.Min.X, View.ViewRect.Min.Y,
 				View.ViewRect.Width(), View.ViewRect.Height(),
 				FIntPoint(DownsampledSizeX, DownsampledSizeY),
-				GSceneRenderTargets.GetBufferSizeXY(),
+				SceneContext.GetBufferSizeXY(),
 				*ScreenVertexShader,
 				EDRF_UseTriangleOptimization);
 		}
diff --git a/Engine/Source/Runtime/Renderer/Private/DeferredShadingRenderer.h b/Engine/Source/Runtime/Renderer/Private/DeferredShadingRenderer.h
--- a/Engine/Source/Runtime/Renderer/Private/DeferredShadingRenderer.h
+++ b/Engine/Source/Runtime/Renderer/Private/DeferredShadingRenderer.h
@@ -65,7 +65,7 @@ public:
 	/**
 	 * Renders the scene's prepass for a particular view in parallel
 	 */
-	void RenderPrePassViewParallel(const FViewInfo& View, FRHICommandList& ParentCmdList, bool& OutDirty);
+	void RenderPrePassViewParallel(const FViewInfo& View, FRHICommandListImmediate& ParentCmdList);
 
 	/** Renders the basepass for the static data of a given View. */
 	bool RenderBasePassStaticData(FRHICommandList& RHICmdList, FViewInfo& View);
@@ -87,7 +87,7 @@ public:
 	void RenderBasePassDynamicDataParallel(FParallelCommandListSet& ParallelCommandListSet);
 
 	/** Renders the basepass for a given View, in parallel */
-	void RenderBasePassViewParallel(FViewInfo& View, FRHICommandList& ParentCmdList, bool& OutDirty);
+	void RenderBasePassViewParallel(FViewInfo& View, FRHICommandListImmediate& ParentCmdList);
 
 	/** Renders the basepass for a given View. */
 	bool RenderBasePassView(FRHICommandListImmediate& RHICmdList, FViewInfo& View);
@@ -137,7 +137,7 @@ private:
 	void CreateWholeSceneProjectedShadow(FLightSceneInfo* LightSceneInfo);
 
 	/** Updates the preshadow cache, allocating new preshadows that can fit and evicting old ones. */
-	void UpdatePreshadowCache();
+	void UpdatePreshadowCache(FSceneRenderTargets& SceneContext);
 
 	/** Finds a matching cached preshadow, if one exists. */
 	TRefCountPtr<FProjectedShadowInfo> GetCachedPreshadow(
diff --git a/Engine/Source/Runtime/Renderer/Private/DistanceFieldLightingPost.cpp b/Engine/Source/Runtime/Renderer/Private/DistanceFieldLightingPost.cpp
--- a/Engine/Source/Runtime/Renderer/Private/DistanceFieldLightingPost.cpp
+++ b/Engine/Source/Runtime/Renderer/Private/DistanceFieldLightingPost.cpp
@@ -247,8 +247,9 @@ public:
 				DistanceFieldIrradiance2->GetRenderTargetItem().ShaderResourceTexture
 				);
 		}
+		FSceneRenderTargets& SceneContext = FSceneRenderTargets::Get(RHICmdList);
 
-		const FIntPoint DownsampledBufferSize(GSceneRenderTargets.GetBufferSizeXY() / FIntPoint(GAODownsampleFactor, GAODownsampleFactor));
+		const FIntPoint DownsampledBufferSize(SceneContext.GetBufferSizeXY() / FIntPoint(GAODownsampleFactor, GAODownsampleFactor));
 		const FVector2D BaseLevelTexelSizeValue(1.0f / DownsampledBufferSize.X, 1.0f / DownsampledBufferSize.Y);
 		SetShaderValue(RHICmdList, ShaderRHI, BentNormalAOTexelSize, BaseLevelTexelSizeValue);
 
@@ -507,8 +508,9 @@ public:
 		}
 
 		SetShaderValue(RHICmdList, ShaderRHI, HistoryWeight, GAOHistoryWeight);
+		FSceneRenderTargets& SceneContext = FSceneRenderTargets::Get(RHICmdList);
 		
-		const FIntPoint DownsampledBufferSize(GSceneRenderTargets.GetBufferSizeXY() / FIntPoint(GAODownsampleFactor, GAODownsampleFactor));
+		const FIntPoint DownsampledBufferSize(SceneContext.GetBufferSizeXY() / FIntPoint(GAODownsampleFactor, GAODownsampleFactor));
 		const FVector2D BaseLevelTexelSizeValue(1.0f / DownsampledBufferSize.X, 1.0f / DownsampledBufferSize.Y);
 		SetShaderValue(RHICmdList, ShaderRHI, BentNormalAOTexelSize, BaseLevelTexelSizeValue);
 	}
@@ -567,6 +569,8 @@ void UpdateHistory(
 	TRefCountPtr<IPooledRenderTarget>& BentNormalHistoryOutput,
 	TRefCountPtr<IPooledRenderTarget>& IrradianceHistoryOutput)
 {
+	FSceneRenderTargets& SceneContext = FSceneRenderTargets::Get(RHICmdList);
+
 	if (BentNormalHistoryState)
 	{
 		const bool bUseDistanceFieldGI = IsDistanceFieldGIAllowed(View);
@@ -631,7 +635,7 @@ void UpdateHistory(
 					View.ViewRect.Min.X / GAODownsampleFactor, View.ViewRect.Min.Y / GAODownsampleFactor,
 					View.ViewRect.Width() / GAODownsampleFactor, View.ViewRect.Height() / GAODownsampleFactor,
 					FIntPoint(View.ViewRect.Width() / GAODownsampleFactor, View.ViewRect.Height() / GAODownsampleFactor),
-					GSceneRenderTargets.GetBufferSizeXY() / FIntPoint(GAODownsampleFactor, GAODownsampleFactor),
+					SceneContext.GetBufferSizeXY() / FIntPoint(GAODownsampleFactor, GAODownsampleFactor),
 					*VertexShader);
 
 				RHICmdList.CopyToResolveTarget(NewBentNormalHistory->GetRenderTargetItem().TargetableTexture, NewBentNormalHistory->GetRenderTargetItem().ShaderResourceTexture, false, FResolveParams());
@@ -647,7 +651,7 @@ void UpdateHistory(
 				const FPooledRenderTargetDesc& HistoryDesc = (*BentNormalHistoryState)->GetDesc();
 
 				// Reallocate history if buffer sizes have changed
-				if (HistoryDesc.Extent != GSceneRenderTargets.GetBufferSizeXY() / FIntPoint(GAODownsampleFactor, GAODownsampleFactor))
+				if (HistoryDesc.Extent != SceneContext.GetBufferSizeXY() / FIntPoint(GAODownsampleFactor, GAODownsampleFactor))
 				{
 					GRenderTargetPool.FreeUnusedResource(*BentNormalHistoryState);
 					*BentNormalHistoryState = NULL;
@@ -704,7 +708,7 @@ void UpdateHistory(
 						0, 0,
 						View.ViewRect.Width() / GAODownsampleFactor, View.ViewRect.Height() / GAODownsampleFactor,
 						FIntPoint(View.ViewRect.Width() / GAODownsampleFactor, View.ViewRect.Height() / GAODownsampleFactor),
-						GSceneRenderTargets.GetBufferSizeXY() / FIntPoint(GAODownsampleFactor, GAODownsampleFactor),
+						SceneContext.GetBufferSizeXY() / FIntPoint(GAODownsampleFactor, GAODownsampleFactor),
 						*VertexShader);
 
 					RHICmdList.CopyToResolveTarget((*BentNormalHistoryState)->GetRenderTargetItem().TargetableTexture, (*BentNormalHistoryState)->GetRenderTargetItem().ShaderResourceTexture, false, FResolveParams());
@@ -906,6 +910,8 @@ void PostProcessBentNormalAOSurfaceCache(
 
 			VertexShader->SetParameters(RHICmdList, View);
 
+			FSceneRenderTargets& SceneContext = FSceneRenderTargets::Get(RHICmdList);
+
 			DrawRectangle( 
 				RHICmdList,
 				0, 0, 
@@ -913,7 +919,7 @@ void PostProcessBentNormalAOSurfaceCache(
 				0, 0, 
 				View.ViewRect.Width() / GAODownsampleFactor, View.ViewRect.Height() / GAODownsampleFactor,
 				FIntPoint(View.ViewRect.Width() / GAODownsampleFactor, View.ViewRect.Height() / GAODownsampleFactor),
-				GSceneRenderTargets.GetBufferSizeXY() / FIntPoint(GAODownsampleFactor, GAODownsampleFactor),
+				SceneContext.GetBufferSizeXY() / FIntPoint(GAODownsampleFactor, GAODownsampleFactor),
 				*VertexShader);
 		}
 
diff --git a/Engine/Source/Runtime/Renderer/Private/DistanceFieldScreenGridLighting.cpp b/Engine/Source/Runtime/Renderer/Private/DistanceFieldScreenGridLighting.cpp
--- a/Engine/Source/Runtime/Renderer/Private/DistanceFieldScreenGridLighting.cpp
+++ b/Engine/Source/Runtime/Renderer/Private/DistanceFieldScreenGridLighting.cpp
@@ -635,6 +635,7 @@ void PostProcessBentNormalAOScreenGrid(
 	TRefCountPtr<IPooledRenderTarget>& BentNormalOutput,
 	TRefCountPtr<IPooledRenderTarget>& IrradianceOutput)
 {
+	FSceneRenderTargets& SceneContext = FSceneRenderTargets::Get(RHICmdList);
 	const bool bUseDistanceFieldGI = IsDistanceFieldGIAllowed(View);
 
 	TRefCountPtr<IPooledRenderTarget> DistanceFieldAOBentNormal;
@@ -691,7 +692,7 @@ void PostProcessBentNormalAOScreenGrid(
 				0, 0, 
 				View.ViewRect.Width() / GAODownsampleFactor, View.ViewRect.Height() / GAODownsampleFactor,
 				FIntPoint(View.ViewRect.Width() / GAODownsampleFactor, View.ViewRect.Height() / GAODownsampleFactor),
-				GSceneRenderTargets.GetBufferSizeXY() / FIntPoint(GAODownsampleFactor, GAODownsampleFactor),
+				SceneContext.GetBufferSizeXY() / FIntPoint(GAODownsampleFactor, GAODownsampleFactor),
 				*VertexShader);
 		}
 
diff --git a/Engine/Source/Runtime/Renderer/Private/DistanceFieldShadowing.cpp b/Engine/Source/Runtime/Renderer/Private/DistanceFieldShadowing.cpp
--- a/Engine/Source/Runtime/Renderer/Private/DistanceFieldShadowing.cpp
+++ b/Engine/Source/Runtime/Renderer/Private/DistanceFieldShadowing.cpp
@@ -63,7 +63,7 @@ int32 GetDFShadowDownsampleFactor()
 
 FIntPoint GetBufferSizeForDFShadows()
 {
-	return FIntPoint::DivideAndRoundDown(GSceneRenderTargets.GetBufferSizeXY(), GetDFShadowDownsampleFactor());
+	return FIntPoint::DivideAndRoundDown(FSceneRenderTargets::Get_FrameConstantsOnly().GetBufferSizeXY(), GetDFShadowDownsampleFactor());
 }
 
 TGlobalResource<FDistanceFieldObjectBufferResource> GShadowCulledObjectBuffers;
@@ -780,7 +780,7 @@ void FProjectedShadowInfo::RenderRayTracedDistanceFieldProjection(FRHICommandLis
 		{
 			check(!Scene->DistanceFieldSceneData.HasPendingOperations());
 
-			GSceneRenderTargets.FinishRenderingLightAttenuation(RHICmdList);
+			FSceneRenderTargets::Get(RHICmdList).FinishRenderingLightAttenuation(RHICmdList);
 			SetRenderTarget(RHICmdList, NULL, NULL);
 
 			int32 NumPlanes = 0;
@@ -873,7 +873,7 @@ void FProjectedShadowInfo::RenderRayTracedDistanceFieldProjection(FRHICommandLis
 			}
 
 			{
-				GSceneRenderTargets.BeginRenderingLightAttenuation(RHICmdList);
+				FSceneRenderTargets::Get(RHICmdList).BeginRenderingLightAttenuation(RHICmdList);
 
 				SCOPED_DRAW_EVENT(RHICmdList, Upsample);
 
diff --git a/Engine/Source/Runtime/Renderer/Private/DistanceFieldSurfaceCacheLighting.cpp b/Engine/Source/Runtime/Renderer/Private/DistanceFieldSurfaceCacheLighting.cpp
--- a/Engine/Source/Runtime/Renderer/Private/DistanceFieldSurfaceCacheLighting.cpp
+++ b/Engine/Source/Runtime/Renderer/Private/DistanceFieldSurfaceCacheLighting.cpp
@@ -217,7 +217,7 @@ FDistanceFieldAOParameters::FDistanceFieldAOParameters(float InOcclusionMaxDista
 
 FIntPoint GetBufferSizeForAO()
 {
-	return FIntPoint::DivideAndRoundDown(GSceneRenderTargets.GetBufferSizeXY(), GAODownsampleFactor);
+	return FIntPoint::DivideAndRoundDown(FSceneRenderTargets::Get_FrameConstantsOnly().GetBufferSizeXY(), GAODownsampleFactor);
 }
 
 // Sample set restricted to not self-intersect a surface based on cone angle .475882232
@@ -986,7 +986,7 @@ void ComputeDistanceFieldNormal(FRHICommandListImmediate& RHICmdList, const TArr
 				0, 0, 
 				View.ViewRect.Width(), View.ViewRect.Height(),
 				FIntPoint(View.ViewRect.Width() / GAODownsampleFactor, View.ViewRect.Height() / GAODownsampleFactor),
-				GSceneRenderTargets.GetBufferSizeXY(),
+				FSceneRenderTargets::Get(RHICmdList).GetBufferSizeXY(),
 				*VertexShader);
 		}
 	}
@@ -2448,7 +2448,7 @@ void RenderIrradianceCacheInterpolation(
 
 			RHICmdList.SetBlendState(TStaticBlendState<CW_RGBA, BO_Add, BF_One, BF_One, BO_Add, BF_One, BF_One, CW_RGBA, BO_Add, BF_One, BF_One, BO_Add, BF_One, BF_One>::GetRHI());
 
-			const FIntPoint BufferSize = GSceneRenderTargets.GetBufferSizeXY();
+			const FIntPoint BufferSize = FSceneRenderTargets::Get(RHICmdList).GetBufferSizeXY();
 			const float InvBufferSizeX = 1.0f / BufferSize.X;
 			const float InvBufferSizeY = 1.0f / BufferSize.Y;
 
@@ -2567,7 +2567,7 @@ void ListDistanceFieldLightingMemory(const FViewInfo& View)
 	{
 		UE_LOG(LogTemp, Log, TEXT("   Tile Culled objects %.3fMb"), TileIntersectionResources->GetSizeBytes() / 1024.0f / 1024.0f);
 	}
-
+	
 	FAOScreenGridResources* ScreenGridResources = ((FSceneViewState*)View.State)->AOScreenGridResources;
 
 	if (ScreenGridResources)
@@ -2648,6 +2648,7 @@ void RenderDistanceFieldAOSurfaceCache(
 	TRefCountPtr<IPooledRenderTarget>& OutDynamicBentNormalAO, 
 	TRefCountPtr<IPooledRenderTarget>& OutDynamicIrradiance)
 {
+	FSceneRenderTargets& SceneContext = FSceneRenderTargets::Get(RHICmdList);
 	const bool bUseDistanceFieldGI = IsDistanceFieldGIAllowed(View);
 
 	// Create surface cache state that will persist across frames if needed
@@ -2743,7 +2744,7 @@ void RenderDistanceFieldAOSurfaceCache(
 		TRefCountPtr<IPooledRenderTarget> DistanceFieldAOBentNormalSplat;
 
 		{
-			FIntPoint AOBufferSize = FIntPoint::DivideAndRoundUp(GSceneRenderTargets.GetBufferSizeXY(), DestLevelDownsampleFactor);
+			FIntPoint AOBufferSize = FIntPoint::DivideAndRoundUp(SceneContext.GetBufferSizeXY(), DestLevelDownsampleFactor);
 			FPooledRenderTargetDesc Desc(FPooledRenderTargetDesc::Create2DDesc(AOBufferSize, PF_FloatRGBA, TexCreate_None, TexCreate_RenderTargetable | TexCreate_UAV, false));
 			GRenderTargetPool.FindFreeElement(Desc, DistanceFieldAOBentNormalSplat, TEXT("DistanceFieldAOBentNormalSplat"));
 		}
@@ -2768,14 +2769,14 @@ void RenderDistanceFieldAOSurfaceCache(
 
 			// Create new records which haven't been shaded yet for shading points which don't have a valid interpolation from existing records
 			{
-				FIntPoint DownsampledViewSize = FIntPoint::DivideAndRoundUp(View.ViewRect.Size(), DestLevelDownsampleFactor);
-				uint32 GroupSizeX = FMath::DivideAndRoundUp(DownsampledViewSize.X, GDistanceFieldAOTileSizeX);
-				uint32 GroupSizeY = FMath::DivideAndRoundUp(DownsampledViewSize.Y, GDistanceFieldAOTileSizeY);
-
-				TShaderMapRef<FPopulateCacheCS> ComputeShader(View.ShaderMap);
-
-				RHICmdList.SetComputeShader(ComputeShader->GetComputeShader());
-				ComputeShader->SetParameters(RHICmdList, View, DistanceFieldAOBentNormalSplat->GetRenderTargetItem(), DistanceFieldNormal->GetRenderTargetItem(), DestLevelDownsampleFactor, DepthLevel, TileListGroupSize, Parameters);
+		        FIntPoint DownsampledViewSize = FIntPoint::DivideAndRoundUp(View.ViewRect.Size(), DestLevelDownsampleFactor);
+				        uint32 GroupSizeX = FMath::DivideAndRoundUp(DownsampledViewSize.X, GDistanceFieldAOTileSizeX);
+				        uint32 GroupSizeY = FMath::DivideAndRoundUp(DownsampledViewSize.Y, GDistanceFieldAOTileSizeY);
+        
+		        TShaderMapRef<FPopulateCacheCS> ComputeShader(View.ShaderMap);
+        
+				        RHICmdList.SetComputeShader(ComputeShader->GetComputeShader());
+		        ComputeShader->SetParameters(RHICmdList, View, DistanceFieldAOBentNormalSplat->GetRenderTargetItem(), DistanceFieldNormal->GetRenderTargetItem(), DestLevelDownsampleFactor, DepthLevel, TileListGroupSize, Parameters);
 				DispatchComputeShader(RHICmdList, *ComputeShader, GroupSizeX, GroupSizeY, 1);
 
 				ComputeShader->UnsetParameters(RHICmdList);
@@ -2874,18 +2875,18 @@ void RenderDistanceFieldAOSurfaceCache(
 
 	GRenderTargetPool.VisualizeTexture.SetCheckPoint(RHICmdList, BentNormalAccumulation);
 
-	// Post process the AO to cover over artifacts
+			// Post process the AO to cover over artifacts
 
 	PostProcessBentNormalAOSurfaceCache(
-		RHICmdList, 
-		Parameters, 
-		View, 
-		VelocityTexture, 
-		BentNormalAccumulation->GetRenderTargetItem(), 
-		IrradianceAccumulation, 
-		DistanceFieldNormal->GetRenderTargetItem(), 
-		OutDynamicBentNormalAO, 
-		OutDynamicIrradiance);
+				RHICmdList, 
+				Parameters, 
+				View, 
+	VelocityTexture, 
+			BentNormalAccumulation->GetRenderTargetItem(), 
+			IrradianceAccumulation, 
+			DistanceFieldNormal->GetRenderTargetItem(), 
+	OutDynamicBentNormalAO, 
+	OutDynamicIrradiance);
 }
 
 bool FDeferredShadingSceneRenderer::RenderDistanceFieldLighting(
@@ -2899,6 +2900,7 @@ bool FDeferredShadingSceneRenderer::RenderDistanceFieldLighting(
 {
 	//@todo - support multiple views
 	const FViewInfo& View = Views[0];
+	FSceneRenderTargets& SceneContext = FSceneRenderTargets::Get(RHICmdList);
 
 	if (SupportsDistanceFieldAO(View.GetFeatureLevel(), View.GetShaderPlatform())
 		&& Views.Num() == 1
@@ -2986,8 +2988,8 @@ bool FDeferredShadingSceneRenderer::RenderDistanceFieldLighting(
 					Parameters, 
 					VelocityTexture,
 					DistanceFieldNormal, 
-					BentNormalOutput, 
-					IrradianceOutput);
+				BentNormalOutput, 
+				IrradianceOutput);
 			}
 			else
 			{
@@ -3006,11 +3008,11 @@ bool FDeferredShadingSceneRenderer::RenderDistanceFieldLighting(
 
 			if (bVisualizeAmbientOcclusion || bVisualizeGlobalIllumination)
 			{
-				GSceneRenderTargets.BeginRenderingSceneColor(RHICmdList, ESimpleRenderTargetMode::EExistingColorAndDepth, FExclusiveDepthStencil::DepthRead_StencilNop);
+				SceneContext.BeginRenderingSceneColor(RHICmdList, ESimpleRenderTargetMode::EExistingColorAndDepth, FExclusiveDepthStencil::DepthRead_StencilNop);
 			}
 			else
 			{
-				FPooledRenderTargetDesc Desc = GSceneRenderTargets.GetSceneColor()->GetDesc();
+				FPooledRenderTargetDesc Desc = SceneContext.GetSceneColor()->GetDesc();
 				// Make sure we get a signed format
 				Desc.Format = PF_FloatRGBA;
 				GRenderTargetPool.FindFreeElement(Desc, OutDynamicBentNormalAO, TEXT("DynamicBentNormalAO"));
@@ -3175,8 +3177,9 @@ void FDeferredShadingSceneRenderer::RenderDynamicSkyLighting(FRHICommandListImme
 		{
 			bApplyShadowing = RenderDistanceFieldLighting(RHICmdList, Parameters, VelocityTexture, DynamicBentNormalAO, DynamicIrradiance, false, false);
 		}
+		FSceneRenderTargets& SceneContext = FSceneRenderTargets::Get(RHICmdList);
 
-		GSceneRenderTargets.BeginRenderingSceneColor(RHICmdList);
+		SceneContext.BeginRenderingSceneColor(RHICmdList);
 
 		for( int32 ViewIndex = 0; ViewIndex < Views.Num(); ViewIndex++ )
 		{
@@ -3244,7 +3247,7 @@ void FDeferredShadingSceneRenderer::RenderDynamicSkyLighting(FRHICommandListImme
 				View.ViewRect.Min.X, View.ViewRect.Min.Y, 
 				View.ViewRect.Width(), View.ViewRect.Height(),
 				FIntPoint(View.ViewRect.Width(), View.ViewRect.Height()),
-				GSceneRenderTargets.GetBufferSizeXY(),
+				SceneContext.GetBufferSizeXY(),
 				*VertexShader);
 		}
 	}
@@ -3485,7 +3488,7 @@ void FDeferredShadingSceneRenderer::RenderMeshDistanceFieldVisualization(FRHICom
 			}
 
 			{
-				GSceneRenderTargets.BeginRenderingSceneColor(RHICmdList);
+				FSceneRenderTargets::Get(RHICmdList).BeginRenderingSceneColor(RHICmdList);
 
 				for (int32 ViewIndex = 0; ViewIndex < Views.Num(); ViewIndex++)
 				{
diff --git a/Engine/Source/Runtime/Renderer/Private/DistanceFieldSurfaceCacheLighting.h b/Engine/Source/Runtime/Renderer/Private/DistanceFieldSurfaceCacheLighting.h
--- a/Engine/Source/Runtime/Renderer/Private/DistanceFieldSurfaceCacheLighting.h
+++ b/Engine/Source/Runtime/Renderer/Private/DistanceFieldSurfaceCacheLighting.h
@@ -570,7 +570,8 @@ public:
 		SetShaderValue(RHICmdList, ShaderRHI, CurrentLevelDownsampleFactor, CurrentLevelDownsampleFactorValue);
 
 		// Round up, to match render target allocation
-		const FVector2D AOBufferSizeValue = FIntPoint::DivideAndRoundUp(GSceneRenderTargets.GetBufferSizeXY(), CurrentLevelDownsampleFactorValue);
+		FSceneRenderTargets& SceneContext = FSceneRenderTargets::Get(RHICmdList);
+		const FVector2D AOBufferSizeValue = FIntPoint::DivideAndRoundUp(SceneContext.GetBufferSizeXY(), CurrentLevelDownsampleFactorValue);
 		SetShaderValue(RHICmdList, ShaderRHI, AOBufferSize, AOBufferSizeValue);
 
 		SetShaderValue(RHICmdList, ShaderRHI, DownsampleFactorToBaseLevel, CurrentLevelDownsampleFactorValue / GAODownsampleFactor);
