UE4 - Large upgrade to parallel rendering. No longer block per pass. Several RHI cmdlist passthrough functions added and implemented on the consoles

diff --git a/Engine/Source/Runtime/Renderer/Private/DistortionRendering.cpp b/Engine/Source/Runtime/Renderer/Private/DistortionRendering.cpp
--- a/Engine/Source/Runtime/Renderer/Private/DistortionRendering.cpp
+++ b/Engine/Source/Runtime/Renderer/Private/DistortionRendering.cpp
@@ -36,10 +36,11 @@ public:
 
 	void SetParameters(const FRenderingCompositePassContext& Context, IPooledRenderTarget& DistortionRT)
 	{
+		FSceneRenderTargets& SceneContext = FSceneRenderTargets::Get(Context.RHICmdList);
 		const FPixelShaderRHIParamRef ShaderRHI = GetPixelShader();
 
 		FTextureRHIParamRef DistortionTextureValue = DistortionRT.GetRenderTargetItem().ShaderResourceTexture;
-		FTextureRHIParamRef SceneColorTextureValue = GSceneRenderTargets.GetSceneColor()->GetRenderTargetItem().ShaderResourceTexture;
+		FTextureRHIParamRef SceneColorTextureValue = SceneContext.GetSceneColor()->GetRenderTargetItem().ShaderResourceTexture;
 
 		// Here we use SF_Point as in fullscreen the pixels are 1:1 mapped.
 		SetTextureParameter(
@@ -60,7 +61,7 @@ public:
 			SceneColorTextureValue
 			);
 
-		FIntPoint SceneBufferSize = GSceneRenderTargets.GetBufferSizeXY();
+		FIntPoint SceneBufferSize = SceneContext.GetBufferSizeXY();
 		FIntRect ViewportRect = Context.GetViewport();
 		FVector4 SceneColorRectValue = FVector4((float)ViewportRect.Min.X/SceneBufferSize.X,
 												(float)ViewportRect.Min.Y/SceneBufferSize.Y,
@@ -767,6 +768,7 @@ void FSceneRenderer::RenderDistortion(FRHICommandListImmediate& RHICmdList)
 	bool bDirty = false;
 
 	TRefCountPtr<IPooledRenderTarget> DistortionRT;
+	FSceneRenderTargets& SceneContext = FSceneRenderTargets::Get(RHICmdList);
 
 	// Render accumulated distortion offsets
 	if( bRender)
@@ -775,7 +777,7 @@ void FSceneRenderer::RenderDistortion(FRHICommandListImmediate& RHICmdList)
 
 		// Create a texture to store the resolved light attenuation values, and a render-targetable surface to hold the unresolved light attenuation values.
 		{
-			FPooledRenderTargetDesc Desc(FPooledRenderTargetDesc::Create2DDesc(GSceneRenderTargets.GetBufferSizeXY(), PF_B8G8R8A8, TexCreate_None, TexCreate_RenderTargetable, false));
+			FPooledRenderTargetDesc Desc(FPooledRenderTargetDesc::Create2DDesc(SceneContext.GetBufferSizeXY(), PF_B8G8R8A8, TexCreate_None, TexCreate_RenderTargetable, false));
 			Desc.Flags |= TexCreate_FastVRAM;
 			GRenderTargetPool.FindFreeElement(Desc, DistortionRT, TEXT("Distortion"));
 
@@ -789,7 +791,7 @@ void FSceneRenderer::RenderDistortion(FRHICommandListImmediate& RHICmdList)
 		// DistortionRT==0 should never happen but better we don't crash
 		if(DistortionRT)
 		{
-			SetRenderTarget(RHICmdList, DistortionRT->GetRenderTargetItem().TargetableTexture, GSceneRenderTargets.GetSceneDepthSurface(), ESimpleRenderTargetMode::EUninitializedColorExistingDepth, FExclusiveDepthStencil::DepthRead_StencilWrite);
+			SetRenderTarget(RHICmdList, DistortionRT->GetRenderTargetItem().TargetableTexture, SceneContext.GetSceneDepthSurface(), ESimpleRenderTargetMode::EUninitializedColorExistingDepth, FExclusiveDepthStencil::DepthRead_StencilWrite);
 
 			for(int32 ViewIndex = 0;ViewIndex < Views.Num();ViewIndex++)
 			{
@@ -829,12 +831,12 @@ void FSceneRenderer::RenderDistortion(FRHICommandListImmediate& RHICmdList)
 	{
 		SCOPED_DRAW_EVENT(RHICmdList, DistortionApply);
 
-		GSceneRenderTargets.ResolveSceneColor(RHICmdList, FResolveRect(0, 0, ViewFamily.FamilySizeX, ViewFamily.FamilySizeY));
+		SceneContext.ResolveSceneColor(RHICmdList, FResolveRect(0, 0, ViewFamily.FamilySizeX, ViewFamily.FamilySizeY));
 
 // OCULUS BEGIN: select ONE render target for all views (eyes)
 		TRefCountPtr<IPooledRenderTarget> NewSceneColor;
 		// we don't create a new name to make it easier to use "vis SceneColor" and get the last HDRSceneColor
-		GRenderTargetPool.FindFreeElement(GSceneRenderTargets.GetSceneColor()->GetDesc(), NewSceneColor, TEXT("SceneColor"));
+		GRenderTargetPool.FindFreeElement(SceneContext.GetSceneColor()->GetDesc(), NewSceneColor, TEXT("SceneColor"));
 		const FSceneRenderTargetItem& DestRenderTarget = NewSceneColor->GetRenderTargetItem();
 // OCULUS END
 
@@ -856,7 +858,7 @@ void FSceneRenderer::RenderDistortion(FRHICommandListImmediate& RHICmdList)
 // OCULUS BEGIN: Weird, different eyes may have different render targets!
 // 				TRefCountPtr<IPooledRenderTarget> NewSceneColor;
 // we don't create a new name to make it easier to use "vis SceneColor" and get the last HDRSceneColor
-// 				GRenderTargetPool.FindFreeElement(GSceneRenderTargets.SceneColor->GetDesc(), NewSceneColor, TEXT("SceneColor"));
+// 				GRenderTargetPool.FindFreeElement(SceneContext.SceneColor->GetDesc(), NewSceneColor, TEXT("SceneColor"));
 // 				const FSceneRenderTargetItem& DestRenderTarget = NewSceneColor->GetRenderTargetItem(); 
 // OCULUS END
 
@@ -886,15 +888,15 @@ void FSceneRenderer::RenderDistortion(FRHICommandListImmediate& RHICmdList)
 					View.ViewRect.Min.X, View.ViewRect.Min.Y, 
 					View.ViewRect.Width(), View.ViewRect.Height(),
 					View.ViewRect.Size(),
-					GSceneRenderTargets.GetBufferSizeXY(),
+					SceneContext.GetBufferSizeXY(),
 					*VertexShader,
 					EDRF_UseTriangleOptimization);
 
 				RHICmdList.CopyToResolveTarget(DestRenderTarget.TargetableTexture, DestRenderTarget.ShaderResourceTexture, false, FResolveParams());
 
 // OCULUS BEGIN
-// 				GSceneRenderTargets.SetSceneColor(NewSceneColor);
-// 				check(GSceneRenderTargets.GetSceneColor());
+// 				SceneContext.SetSceneColor(NewSceneColor);
+// 				check(SceneContext.GetSceneColor());
 // OCULUS END				
 
 			}
@@ -902,11 +904,11 @@ void FSceneRenderer::RenderDistortion(FRHICommandListImmediate& RHICmdList)
 		}
 
 		// OCULUS BEGIN
-		GSceneRenderTargets.SetSceneColor(NewSceneColor);
-		check(GSceneRenderTargets.GetSceneColor());
+		SceneContext.SetSceneColor(NewSceneColor);
+		check(SceneContext.GetSceneColor());
 // OCULUS END				
 		// Distortions RT is no longer needed, buffer can be reused by the pool, see BeginRenderingDistortionAccumulation() call above
-		GSceneRenderTargets.FinishRenderingSceneColor(RHICmdList, false);
+		SceneContext.FinishRenderingSceneColor(RHICmdList, false);
 	}
 }
 
diff --git a/Engine/Source/Runtime/Renderer/Private/DynamicPrimitiveDrawing.inl b/Engine/Source/Runtime/Renderer/Private/DynamicPrimitiveDrawing.inl
--- a/Engine/Source/Runtime/Renderer/Private/DynamicPrimitiveDrawing.inl
+++ b/Engine/Source/Runtime/Renderer/Private/DynamicPrimitiveDrawing.inl
@@ -131,7 +131,7 @@ void DrawViewElementsParallel(
 				{
 					FRHICommandList* CmdList = ParallelCommandListSet.NewParallelCommandList();
 
-					FGraphEventRef AnyThreadCompletionEvent = TGraphTask<FDrawViewElementsAnyThreadTask<DrawingPolicyFactoryType> >::CreateTask(nullptr, ENamedThreads::RenderThread)
+					FGraphEventRef AnyThreadCompletionEvent = TGraphTask<FDrawViewElementsAnyThreadTask<DrawingPolicyFactoryType> >::CreateTask(ParallelCommandListSet.GetPrereqs(), ENamedThreads::RenderThread)
 						.ConstructAndDispatchWhenReady(*CmdList, ParallelCommandListSet.View, DrawingContext, DPGIndex, bPreFog, Start, Last);
 
 					ParallelCommandListSet.AddParallelCommandList(CmdList, AnyThreadCompletionEvent, Last - Start + 1);
diff --git a/Engine/Source/Runtime/Renderer/Private/EditorCompositeParams.cpp b/Engine/Source/Runtime/Renderer/Private/EditorCompositeParams.cpp
--- a/Engine/Source/Runtime/Renderer/Private/EditorCompositeParams.cpp
+++ b/Engine/Source/Runtime/Renderer/Private/EditorCompositeParams.cpp
@@ -22,6 +22,7 @@ void FEditorCompositingParameters::SetParameters(
 	const FPixelShaderRHIParamRef ShaderRHI)
 {
 #if WITH_EDITOR
+		FSceneRenderTargets& SceneContext = FSceneRenderTargets::Get(RHICmdList);
 		if (GMaxRHIFeatureLevel >= ERHIFeatureLevel::SM4)
 		{
 			if (MaterialResource.IsUsedWithEditorCompositing())
@@ -39,7 +40,7 @@ void FEditorCompositingParameters::SetParameters(
 
 				if (FilteredSceneDepthTexture.IsBound())
 				{
-					const FTexture2DRHIRef* DepthTexture = GSceneRenderTargets.GetActualDepthTexture();
+					const FTexture2DRHIRef* DepthTexture = SceneContext.GetActualDepthTexture();
 					check(DepthTexture != NULL);
 					SetTextureParameter(
 						RHICmdList,
@@ -52,7 +53,7 @@ void FEditorCompositingParameters::SetParameters(
 				}
 
 				SetShaderValue(RHICmdList, ShaderRHI, EditorCompositeDepthTestParameter, bEnableEditorPrimitveDepthTest);
-				SetShaderValue(RHICmdList, ShaderRHI, MSAASampleCount, GSceneRenderTargets.EditorPrimitivesColor ? GSceneRenderTargets.EditorPrimitivesColor->GetDesc().NumSamples : 0);
+				SetShaderValue(RHICmdList, ShaderRHI, MSAASampleCount, SceneContext.EditorPrimitivesColor ? SceneContext.EditorPrimitivesColor->GetDesc().NumSamples : 0);
 			}
 		}
 #endif
diff --git a/Engine/Source/Runtime/Renderer/Private/FogRendering.cpp b/Engine/Source/Runtime/Renderer/Private/FogRendering.cpp
--- a/Engine/Source/Runtime/Renderer/Private/FogRendering.cpp
+++ b/Engine/Source/Runtime/Renderer/Private/FogRendering.cpp
@@ -307,8 +307,9 @@ bool FDeferredShadingSceneRenderer::RenderFog(FRHICommandListImmediate& RHICmdLi
 			0, 1, 2,
 			0, 2, 3
 		};
+		FSceneRenderTargets& SceneContext = FSceneRenderTargets::Get(RHICmdList);
 
-		GSceneRenderTargets.BeginRenderingSceneColor(RHICmdList, ESimpleRenderTargetMode::EExistingColorAndDepth, FExclusiveDepthStencil::DepthRead_StencilWrite);
+		SceneContext.BeginRenderingSceneColor(RHICmdList, ESimpleRenderTargetMode::EExistingColorAndDepth, FExclusiveDepthStencil::DepthRead_StencilWrite);
 		for(int32 ViewIndex = 0;ViewIndex < Views.Num();ViewIndex++)
 		{
 			const FViewInfo& View = Views[ViewIndex];
@@ -345,7 +346,7 @@ bool FDeferredShadingSceneRenderer::RenderFog(FRHICommandListImmediate& RHICmdLi
 		}
 
 		//no need to resolve since we used alpha blending
-		GSceneRenderTargets.FinishRenderingSceneColor(RHICmdList, false);
+		SceneContext.FinishRenderingSceneColor(RHICmdList, false);
 		return true;
 	}
 
diff --git a/Engine/Source/Runtime/Renderer/Private/ForwardShadingRenderer.cpp b/Engine/Source/Runtime/Renderer/Private/ForwardShadingRenderer.cpp
--- a/Engine/Source/Runtime/Renderer/Private/ForwardShadingRenderer.cpp
+++ b/Engine/Source/Runtime/Renderer/Private/ForwardShadingRenderer.cpp
@@ -88,9 +88,10 @@ void FForwardShadingSceneRenderer::Render(FRHICommandListImmediate& RHICmdList)
 
 	// Initialize global system textures (pass-through if already initialized).
 	GSystemTextures.InitializeTextures(RHICmdList, FeatureLevel);
+	FSceneRenderTargets& SceneContext = FSceneRenderTargets::Get(RHICmdList);
 
 	// Allocate the maximum scene render target space for the current view family.
-	GSceneRenderTargets.Allocate(ViewFamily);
+	SceneContext.Allocate(ViewFamily);
 
 	// Find the visible primitives.
 	InitViews(RHICmdList);
@@ -119,12 +120,12 @@ void FForwardShadingSceneRenderer::Render(FRHICommandListImmediate& RHICmdList)
 
 	if (bGammaSpace && !bRenderToScene)
 	{
-		SetRenderTarget(RHICmdList, ViewFamily.RenderTarget->GetRenderTargetTexture(), GSceneRenderTargets.GetSceneDepthTexture(), ESimpleRenderTargetMode::EClearToDefault);
+		SetRenderTarget(RHICmdList, ViewFamily.RenderTarget->GetRenderTargetTexture(), SceneContext.GetSceneDepthTexture(), ESimpleRenderTargetMode::EClearToDefault);
 	}
 	else
 	{
 		// Begin rendering to scene color
-		GSceneRenderTargets.BeginRenderingSceneColor(RHICmdList, ESimpleRenderTargetMode::EClearToDefault);
+		SceneContext.BeginRenderingSceneColor(RHICmdList, ESimpleRenderTargetMode::EClearToDefault);
 	}
 
 	if (GIsEditor)
@@ -135,7 +136,7 @@ void FForwardShadingSceneRenderer::Render(FRHICommandListImmediate& RHICmdList)
 	RenderForwardShadingBasePass(RHICmdList);
 
 	// Make a copy of the scene depth if the current hardware doesn't support reading and writing to the same depth buffer
-	GSceneRenderTargets.ResolveSceneDepthToAuxiliaryTexture(RHICmdList);
+	SceneContext.ResolveSceneDepthToAuxiliaryTexture(RHICmdList);
 
 	// Notify the FX system that opaque primitives have been rendered.
 	if (Scene->FXSystem)
@@ -174,7 +175,7 @@ void FForwardShadingSceneRenderer::Render(FRHICommandListImmediate& RHICmdList)
 		// Convert alpha from depth to circle of confusion with sunshaft intensity.
 		// This is done before resolve on hardware with framebuffer fetch.
 		// This will break when PrePostSourceViewportSize is not full size.
-		FIntPoint PrePostSourceViewportSize = GSceneRenderTargets.GetBufferSizeXY();
+		FIntPoint PrePostSourceViewportSize = SceneContext.GetBufferSizeXY();
 
 		FMemMark Mark(FMemStack::Get());
 		FRenderingCompositePassContext CompositeContext(RHICmdList, View);
@@ -186,7 +187,7 @@ void FForwardShadingSceneRenderer::Render(FRHICommandListImmediate& RHICmdList)
 	if (!bGammaSpace || bRenderToScene)
 	{
 		// Resolve the scene color for post processing.
-		GSceneRenderTargets.ResolveSceneColor(RHICmdList, FResolveRect(0, 0, ViewFamily.FamilySizeX, ViewFamily.FamilySizeY));
+		SceneContext.ResolveSceneColor(RHICmdList, FResolveRect(0, 0, ViewFamily.FamilySizeX, ViewFamily.FamilySizeY));
 
 		// Drop depth and stencil before post processing to avoid export.
 		RHICmdList.DiscardRenderTargets(true, true, 0);
diff --git a/Engine/Source/Runtime/Renderer/Private/ForwardTranslucentRendering.cpp b/Engine/Source/Runtime/Renderer/Private/ForwardTranslucentRendering.cpp
--- a/Engine/Source/Runtime/Renderer/Private/ForwardTranslucentRendering.cpp
+++ b/Engine/Source/Runtime/Renderer/Private/ForwardTranslucentRendering.cpp
@@ -48,16 +48,18 @@ FGlobalBoundShaderState ForwardCopySceneAlphaBoundShaderState;
 void FForwardShadingSceneRenderer::CopySceneAlpha(FRHICommandListImmediate& RHICmdList, const FViewInfo& View)
 {
 	SCOPED_DRAW_EVENTF(RHICmdList, EventCopy, TEXT("CopySceneAlpha"));
+	FSceneRenderTargets& SceneContext = FSceneRenderTargets::Get(RHICmdList);
+
 	RHICmdList.SetRasterizerState(TStaticRasterizerState<FM_Solid, CM_None>::GetRHI());
 	RHICmdList.SetDepthStencilState(TStaticDepthStencilState<false,CF_Always>::GetRHI());
 	RHICmdList.SetBlendState(TStaticBlendState<>::GetRHI());
 
-	GSceneRenderTargets.ResolveSceneColor(RHICmdList);
+	SceneContext.ResolveSceneColor(RHICmdList);
 
-	GSceneRenderTargets.BeginRenderingSceneAlphaCopy(RHICmdList);
+	SceneContext.BeginRenderingSceneAlphaCopy(RHICmdList);
 
-	int X = GSceneRenderTargets.GetBufferSizeXY().X;
-	int Y = GSceneRenderTargets.GetBufferSizeXY().Y;
+	int X = SceneContext.GetBufferSizeXY().X;
+	int Y = SceneContext.GetBufferSizeXY().Y;
 
 	RHICmdList.SetViewport(0, 0, 0.0f, X, Y, 1.0f);
 
@@ -74,11 +76,11 @@ void FForwardShadingSceneRenderer::CopySceneAlpha(FRHICommandListImmediate& RHIC
 		0, 0, 
 		X, Y,
 		FIntPoint(X, Y),
-		GSceneRenderTargets.GetBufferSizeXY(),
+		SceneContext.GetBufferSizeXY(),
 		*ScreenVertexShader,
 		EDRF_UseTriangleOptimization);
 
-	GSceneRenderTargets.FinishRenderingSceneAlphaCopy(RHICmdList);
+	SceneContext.FinishRenderingSceneAlphaCopy(RHICmdList);
 }
 
 
@@ -344,7 +346,7 @@ void FForwardShadingSceneRenderer::RenderTranslucency(FRHICommandListImmediate&
 
 			if (!bGammaSpace)
 			{
-				GSceneRenderTargets.BeginRenderingTranslucency(RHICmdList, View);
+				FSceneRenderTargets::Get(RHICmdList).BeginRenderingTranslucency(RHICmdList, View);
 			}
 			else
 			{
diff --git a/Engine/Source/Runtime/Renderer/Private/GammaCorrection.cpp b/Engine/Source/Runtime/Renderer/Private/GammaCorrection.cpp
--- a/Engine/Source/Runtime/Renderer/Private/GammaCorrection.cpp
+++ b/Engine/Source/Runtime/Renderer/Private/GammaCorrection.cpp
@@ -123,8 +123,9 @@ void FSceneRenderer::GammaCorrectToViewportRenderTarget(FRHICommandList& RHICmdL
 		);
 	SetShaderValue(RHICmdList, ShaderRHI,PixelShader->ColorScale,View->ColorScale);
 	SetShaderValue(RHICmdList, ShaderRHI,PixelShader->OverlayColor,View->OverlayColor);
+	FSceneRenderTargets& SceneContext = FSceneRenderTargets::Get(RHICmdList);
 
-	const FTextureRHIRef DesiredSceneColorTexture = GSceneRenderTargets.GetSceneColorTexture();
+	const FTextureRHIRef DesiredSceneColorTexture = SceneContext.GetSceneColorTexture();
 
 	SetTextureParameter(
 		RHICmdList, 
@@ -143,7 +144,7 @@ void FSceneRenderer::GammaCorrectToViewportRenderTarget(FRHICommandList& RHICmdL
 		View->ViewRect.Min.X,View->ViewRect.Min.Y,
 		View->ViewRect.Width(),View->ViewRect.Height(),
 		ViewFamily.RenderTarget->GetSizeXY(),
-		GSceneRenderTargets.GetBufferSizeXY(),
+		SceneContext.GetBufferSizeXY(),
 		*VertexShader,
 		EDRF_UseTriangleOptimization);
 }
diff --git a/Engine/Source/Runtime/Renderer/Private/HeightfieldLighting.cpp b/Engine/Source/Runtime/Renderer/Private/HeightfieldLighting.cpp
--- a/Engine/Source/Runtime/Renderer/Private/HeightfieldLighting.cpp
+++ b/Engine/Source/Runtime/Renderer/Private/HeightfieldLighting.cpp
@@ -779,7 +779,7 @@ public:
 		SetShaderValue(RHICmdList, ShaderRHI, ShadowDepthBias, ProjectedShadowInfo->GetShaderDepthBias());
 		SetShaderValue(RHICmdList, ShaderRHI, CascadeDepthMinMax, FVector2D(ProjectedShadowInfo->CascadeSettings.SplitNear, ProjectedShadowInfo->CascadeSettings.SplitFar));
 
-		FTexture2DRHIRef ShadowDepthTextureValue = GSceneRenderTargets.GetShadowDepthZTexture(ProjectedShadowInfo->bAllocatedInPreshadowCache);
+		FTexture2DRHIRef ShadowDepthTextureValue = FSceneRenderTargets::Get(RHICmdList).GetShadowDepthZTexture(ProjectedShadowInfo->bAllocatedInPreshadowCache);
 		FSamplerStateRHIParamRef DepthSamplerState = TStaticSamplerState<SF_Point,AM_Clamp,AM_Clamp,AM_Clamp>::GetRHI();
 
 		SetTextureParameter(RHICmdList, ShaderRHI, ShadowDepthTexture, ShadowDepthTextureSampler, DepthSamplerState, ShadowDepthTextureValue);	
diff --git a/Engine/Source/Runtime/Renderer/Private/LightFunctionRendering.cpp b/Engine/Source/Runtime/Renderer/Private/LightFunctionRendering.cpp
--- a/Engine/Source/Runtime/Renderer/Private/LightFunctionRendering.cpp
+++ b/Engine/Source/Runtime/Renderer/Private/LightFunctionRendering.cpp
@@ -240,7 +240,7 @@ bool FDeferredShadingSceneRenderer::RenderLightFunctionForMaterial(FRHICommandLi
 
 	if (MaterialProxy && MaterialProxy->GetMaterial(Scene->GetFeatureLevel())->IsLightFunction())
 	{
-		GSceneRenderTargets.BeginRenderingLightAttenuation(RHICmdList);
+		FSceneRenderTargets::Get(RHICmdList).BeginRenderingLightAttenuation(RHICmdList);
 
 		bRenderedLightFunction = true;
 
diff --git a/Engine/Source/Runtime/Renderer/Private/LightPropagationVolume.cpp b/Engine/Source/Runtime/Renderer/Private/LightPropagationVolume.cpp
--- a/Engine/Source/Runtime/Renderer/Private/LightPropagationVolume.cpp
+++ b/Engine/Source/Runtime/Renderer/Private/LightPropagationVolume.cpp
@@ -775,7 +775,7 @@ FLightPropagationVolume::FLightPropagationVolume() :
 
 	// VPL List buffers
 	mVplListBuffer = new FRWBufferStructured();
-	int32 RSMResolution = GSceneRenderTargets.GetReflectiveShadowMapResolution();
+	int32 RSMResolution = FSceneRenderTargets::Get_FrameConstantsOnly().GetReflectiveShadowMapResolution();
 	int32 GvListBufferSize = RSMResolution * RSMResolution * 16; // Allow 16 layers of depth per every pixel of the RSM (on average) 
 	int32 VplListBufferSize = RSMResolution * RSMResolution * 4; // Allow 4 layers of depth per pixel in the RSM (1 for the RSM injection + 3 for light injection)
 	mVplListBuffer->Initialize( sizeof( VplListEntry ), VplListBufferSize, 0, true, false );
@@ -837,7 +837,7 @@ void FLightPropagationVolume::InitSettings(FRHICommandList& RHICmdList, const FS
 {
 	int32 NumFastLpvTextures = 7;
 	int32 NumFastGvTextures = 2;
-	FIntPoint BufferSize = GSceneRenderTargets.GetBufferSizeXY();
+	FIntPoint BufferSize = FSceneRenderTargets::Get_FrameConstantsOnly().GetBufferSizeXY();
 	if ( BufferSize.X >= 1600 && BufferSize.Y >= 900 )
 	{
 		NumFastLpvTextures = 5;
@@ -962,7 +962,7 @@ void FLightPropagationVolume::InitSettings(FRHICommandList& RHICmdList, const FS
 			* LpvWriteUniformBufferParams->RsmAreaIntensityMultiplier * CVarLPVEmissiveIntensityMultiplier.GetValueOnRenderThread() * 0.25f;
 		LpvWriteUniformBufferParams->DirectionalOcclusionIntensity	= LPVSettings.LPVDirectionalOcclusionIntensity;
 		LpvWriteUniformBufferParams->DirectionalOcclusionRadius		= LPVSettings.LPVDirectionalOcclusionRadius;
-		LpvWriteUniformBufferParams->RsmPixelToTexcoordMultiplier	= 1.0f / float(GSceneRenderTargets.GetReflectiveShadowMapResolution() - 1);
+		LpvWriteUniformBufferParams->RsmPixelToTexcoordMultiplier	= 1.0f / float(FSceneRenderTargets::Get_FrameConstantsOnly().GetReflectiveShadowMapResolution() - 1);
 
 		LpvReadUniformBufferParams.DirectionalOcclusionIntensity	= LPVSettings.LPVDirectionalOcclusionIntensity;
 		LpvReadUniformBufferParams.DiffuseOcclusionExponent			= LPVSettings.LPVDiffuseOcclusionExponent;
@@ -1064,7 +1064,7 @@ void FLightPropagationVolume::GetShadowInfo( const FProjectedShadowInfo& Project
 	FVector4 WorldRight = RsmInfoOut.ShadowToWorld.TransformVector( ShadowRight );
 	float ShadowArea = WorldUp.Size3() * WorldRight.Size3();
 
-	int32 RSMResolution = GSceneRenderTargets.GetReflectiveShadowMapResolution();
+	int32 RSMResolution = FSceneRenderTargets::Get_FrameConstantsOnly().GetReflectiveShadowMapResolution();
 	float ResolutionMultiplier = RSMResolution / 256.0f;
 	float IdealCubeSizeMultiplier = 0.5f * ResolutionMultiplier; 
 	float IdealRsmArea = ( CubeSize * IdealCubeSizeMultiplier * CubeSize * IdealCubeSizeMultiplier );
@@ -1124,7 +1124,7 @@ void FLightPropagationVolume::InjectDirectionalLightRSM(
 		GetShaderParams( ShaderParams );
 		Shader->SetParameters(RHICmdList, ShaderParams, RsmDiffuseTex, RsmNormalTex, RsmDepthTex );
 
-		int32 RSMResolution = GSceneRenderTargets.GetReflectiveShadowMapResolution();
+		int32 RSMResolution = FSceneRenderTargets::Get_FrameConstantsOnly().GetReflectiveShadowMapResolution();
 		// todo: what if not divisble by 8?
 		DispatchComputeShader(RHICmdList, *Shader, RSMResolution / 8, RSMResolution / 8, 1 ); 
 
diff --git a/Engine/Source/Runtime/Renderer/Private/LightShaftRendering.cpp b/Engine/Source/Runtime/Renderer/Private/LightShaftRendering.cpp
--- a/Engine/Source/Runtime/Renderer/Private/LightShaftRendering.cpp
+++ b/Engine/Source/Runtime/Renderer/Private/LightShaftRendering.cpp
@@ -108,9 +108,10 @@ public:
 	template<typename ShaderRHIParamRef>
 	void SetParameters(FRHICommandList& RHICmdList, const ShaderRHIParamRef Shader, const FLightSceneInfo* LightSceneInfo, const FSceneView& View, TRefCountPtr<IPooledRenderTarget>& PassSource)
 	{
+		FSceneRenderTargets& SceneContext = FSceneRenderTargets::Get(RHICmdList);
 		const uint32 DownsampleFactor = GetLightShaftDownsampleFactor();
 		FIntPoint DownSampledViewSize(FMath::FloorToInt(View.ViewRect.Width() / DownsampleFactor), FMath::FloorToInt(View.ViewRect.Height() / DownsampleFactor));
-		const FIntPoint FilterBufferSize = GSceneRenderTargets.GetBufferSizeXY() / DownsampleFactor;
+		const FIntPoint FilterBufferSize = SceneContext.GetBufferSizeXY() / DownsampleFactor;
 
 		const FVector2D ViewRatioOfBuffer((float)DownSampledViewSize.X / FilterBufferSize.X, (float)DownSampledViewSize.Y / FilterBufferSize.Y);
 		const FVector4 AspectRatioAndInvAspectRatio(
@@ -125,7 +126,7 @@ public:
 		// Transform into texture coordinates
 		FVector4 ProjectedBlurOrigin = View.WorldToScreen(WorldSpaceBlurOrigin);
 
-		const FIntPoint BufferSize = GSceneRenderTargets.GetBufferSizeXY();
+		const FIntPoint BufferSize = SceneContext.GetBufferSizeXY();
 		const float InvBufferSizeX = 1.0f / BufferSize.X;
 		const float InvBufferSizeY = 1.0f / BufferSize.Y;
 
@@ -298,7 +299,7 @@ public:
 		FGlobalShader::SetParameters(RHICmdList, GetPixelShader(), View);
 		LightShaftParameters.SetParameters(RHICmdList, GetPixelShader(), LightSceneInfo, View, PassSource);
 
-		const FIntPoint BufferSize = GSceneRenderTargets.GetBufferSizeXY();
+		const FIntPoint BufferSize = FSceneRenderTargets::Get(RHICmdList).GetBufferSizeXY();
 		FVector2D SampleOffsets(1.0f / BufferSize.X, 1.0f / BufferSize.Y);
 		SetShaderValue(RHICmdList, GetPixelShader(),SampleOffsetsParameter,SampleOffsets);
 		SceneTextureParams.Set(RHICmdList, GetPixelShader(), View);
@@ -437,7 +438,7 @@ void AllocateOrReuseLightShaftRenderTarget(FRHICommandListImmediate& RHICmdList,
 	if (!Target)
 	{
 		EPixelFormat LightShaftFilterBufferFormat = PF_FloatRGB;
-		const FIntPoint BufferSize = GSceneRenderTargets.GetBufferSizeXY();
+		const FIntPoint BufferSize = FSceneRenderTargets::Get(RHICmdList).GetBufferSizeXY();
 		FIntPoint LightShaftSize(FMath::Max<uint32>(BufferSize.X / GetLightShaftDownsampleFactor(), 1), FMath::Max<uint32>(BufferSize.Y / GetLightShaftDownsampleFactor(), 1));
 		FPooledRenderTargetDesc Desc(FPooledRenderTargetDesc::Create2DDesc(LightShaftSize, LightShaftFilterBufferFormat, TexCreate_None, TexCreate_RenderTargetable, false));
 		GRenderTargetPool.FindFreeElement(Desc, Target, Name);
@@ -451,9 +452,9 @@ void AllocateOrReuseLightShaftRenderTarget(FRHICommandListImmediate& RHICmdList,
 template<bool bDownsampleOcclusion>
 void DownsamplePass(FRHICommandListImmediate& RHICmdList, const FViewInfo& View, const FLightSceneInfo* LightSceneInfo, TRefCountPtr<IPooledRenderTarget>& LightShaftsSource, TRefCountPtr<IPooledRenderTarget>& LightShaftsDest)
 {
-	const FIntPoint BufferSize = GSceneRenderTargets.GetBufferSizeXY();
+	const FIntPoint BufferSize = FSceneRenderTargets::Get(RHICmdList).GetBufferSizeXY();
 	const uint32 DownsampleFactor	= GetLightShaftDownsampleFactor();
-	const FIntPoint FilterBufferSize = GSceneRenderTargets.GetBufferSizeXY() / DownsampleFactor;
+	const FIntPoint FilterBufferSize = BufferSize / DownsampleFactor;
 	const FIntPoint DownSampledXY = View.ViewRect.Min / DownsampleFactor;
 	const uint32 DownsampledSizeX = View.ViewRect.Width() / DownsampleFactor;
 	const uint32 DownsampledSizeY = View.ViewRect.Height() / DownsampleFactor;
@@ -595,9 +596,9 @@ void ApplyRadialBlurPasses(
 {
 	TShaderMapRef<FScreenVS> ScreenVertexShader(View.ShaderMap);
 
-	const FIntPoint BufferSize = GSceneRenderTargets.GetBufferSizeXY();
+	const FIntPoint BufferSize = FSceneRenderTargets::Get(RHICmdList).GetBufferSizeXY();
 	const uint32 DownsampleFactor	= GetLightShaftDownsampleFactor();
-	const FIntPoint FilterBufferSize = GSceneRenderTargets.GetBufferSizeXY() / DownsampleFactor;
+	const FIntPoint FilterBufferSize = BufferSize / DownsampleFactor;
 	const FIntPoint DownSampledXY = View.ViewRect.Min / DownsampleFactor;
 	const uint32 DownsampledSizeX = View.ViewRect.Width() / DownsampleFactor;
 	const uint32 DownsampledSizeY = View.ViewRect.Height() / DownsampleFactor;
@@ -643,9 +644,9 @@ void FinishOcclusionTerm(FRHICommandList& RHICmdList, const FViewInfo& View, con
 {
 	TShaderMapRef<FScreenVS> ScreenVertexShader(View.ShaderMap);
 
-	const FIntPoint BufferSize = GSceneRenderTargets.GetBufferSizeXY();
+	const FIntPoint BufferSize = FSceneRenderTargets::Get(RHICmdList).GetBufferSizeXY();
 	const uint32 DownsampleFactor	= GetLightShaftDownsampleFactor();
-	const FIntPoint FilterBufferSize = GSceneRenderTargets.GetBufferSizeXY() / DownsampleFactor;
+	const FIntPoint FilterBufferSize = BufferSize / DownsampleFactor;
 	const FIntPoint DownSampledXY = View.ViewRect.Min / DownsampleFactor;
 	const uint32 DownsampledSizeX = View.ViewRect.Width() / DownsampleFactor;
 	const uint32 DownsampledSizeY = View.ViewRect.Height() / DownsampleFactor;
@@ -830,7 +831,8 @@ IMPLEMENT_SHADER_TYPE(,FApplyLightShaftsPixelShader,TEXT("LightShaftShader"),TEX
 
 void ApplyLightShaftBloom(FRHICommandListImmediate& RHICmdList, const FViewInfo& View, const FLightSceneInfo* const LightSceneInfo, TRefCountPtr<IPooledRenderTarget>& LightShaftsSource)
 {
-	GSceneRenderTargets.BeginRenderingSceneColor(RHICmdList, ESimpleRenderTargetMode::EUninitializedColorExistingDepth, FExclusiveDepthStencil::DepthRead_StencilWrite);
+	FSceneRenderTargets& SceneContext = FSceneRenderTargets::Get(RHICmdList);
+	SceneContext.BeginRenderingSceneColor(RHICmdList, ESimpleRenderTargetMode::EUninitializedColorExistingDepth, FExclusiveDepthStencil::DepthRead_StencilWrite);
 
 	RHICmdList.SetViewport(View.ViewRect.Min.X, View.ViewRect.Min.Y, 0.0f, View.ViewRect.Max.X, View.ViewRect.Max.Y, 1.0f);
 	RHICmdList.SetBlendState(TStaticBlendState<CW_RGB, BO_Add, BF_One, BF_One>::GetRHI());
@@ -845,9 +847,9 @@ void ApplyLightShaftBloom(FRHICommandListImmediate& RHICmdList, const FViewInfo&
 	/// ?
 	ApplyLightShaftsPixelShader->SetParameters(RHICmdList, View, LightShaftsSource);
 
-	const FIntPoint BufferSize = GSceneRenderTargets.GetBufferSizeXY();
+	const FIntPoint BufferSize = SceneContext.GetBufferSizeXY();
 	const uint32 DownsampleFactor	= GetLightShaftDownsampleFactor();
-	const FIntPoint FilterBufferSize = GSceneRenderTargets.GetBufferSizeXY() / DownsampleFactor;
+	const FIntPoint FilterBufferSize = SceneContext.GetBufferSizeXY() / DownsampleFactor;
 	const FIntPoint DownSampledXY = View.ViewRect.Min / DownsampleFactor;
 	const uint32 DownsampledSizeX = View.ViewRect.Width() / DownsampleFactor;
 	const uint32 DownsampledSizeY = View.ViewRect.Height() / DownsampleFactor;
@@ -862,7 +864,7 @@ void ApplyLightShaftBloom(FRHICommandListImmediate& RHICmdList, const FViewInfo&
 		*ScreenVertexShader,
 		EDRF_UseTriangleOptimization);
 
-	GSceneRenderTargets.FinishRenderingSceneColor(RHICmdList, false);
+	SceneContext.FinishRenderingSceneColor(RHICmdList, false);
 }
 
 void FSceneViewState::TrimHistoryRenderTargets(const FScene* Scene)
diff --git a/Engine/Source/Runtime/Renderer/Private/PostProcess/PostProcessBloomSetup.cpp b/Engine/Source/Runtime/Renderer/Private/PostProcess/PostProcessBloomSetup.cpp
--- a/Engine/Source/Runtime/Renderer/Private/PostProcess/PostProcessBloomSetup.cpp
+++ b/Engine/Source/Runtime/Renderer/Private/PostProcess/PostProcessBloomSetup.cpp
@@ -156,7 +156,7 @@ void FRCPassPostProcessBloomSetup::Process(FRenderingCompositePassContext& Conte
 	FIntPoint DestSize = PassOutputs[0].RenderTargetDesc.Extent;
 
 	// e.g. 4 means the input texture is 4x smaller than the buffer size
-	uint32 ScaleFactor = GSceneRenderTargets.GetBufferSizeXY().X / SrcSize.X;
+	uint32 ScaleFactor = FSceneRenderTargets::Get(Context.RHICmdList).GetBufferSizeXY().X / SrcSize.X;
 
 	FIntRect SrcRect = View.ViewRect / ScaleFactor;
 	FIntRect DestRect = SrcRect;
@@ -282,7 +282,7 @@ void FRCPassPostProcessVisualizeBloomSetup::Process(FRenderingCompositePassConte
 	FIntPoint DestSize = PassOutputs[0].RenderTargetDesc.Extent;
 
 	// e.g. 4 means the input texture is 4x smaller than the buffer size
-	uint32 ScaleFactor = GSceneRenderTargets.GetBufferSizeXY().X / SrcSize.X;
+	uint32 ScaleFactor = FSceneRenderTargets::Get(Context.RHICmdList).GetBufferSizeXY().X / SrcSize.X;
 
 	FIntRect SrcRect = View.ViewRect / ScaleFactor;
 	FIntRect DestRect = SrcRect;
@@ -420,7 +420,7 @@ void FRCPassPostProcessVisualizeBloomOverlay::Process(FRenderingCompositePassCon
 	FIntPoint DestSize = PassOutputs[0].RenderTargetDesc.Extent;
 
 	// e.g. 4 means the input texture is 4x smaller than the buffer size
-	uint32 ScaleFactor = GSceneRenderTargets.GetBufferSizeXY().X / SrcSize.X;
+	uint32 ScaleFactor = FSceneRenderTargets::Get(Context.RHICmdList).GetBufferSizeXY().X / SrcSize.X;
 
 	FIntRect SrcRect = View.ViewRect / ScaleFactor;
 	FIntRect DestRect = SrcRect;
diff --git a/Engine/Source/Runtime/Renderer/Private/PostProcess/PostProcessBokehDOF.cpp b/Engine/Source/Runtime/Renderer/Private/PostProcess/PostProcessBokehDOF.cpp
--- a/Engine/Source/Runtime/Renderer/Private/PostProcess/PostProcessBokehDOF.cpp
+++ b/Engine/Source/Runtime/Renderer/Private/PostProcess/PostProcessBokehDOF.cpp
@@ -180,7 +180,7 @@ void FRCPassPostProcessVisualizeDOF::Process(FRenderingCompositePassContext& Con
 	FIntPoint DestSize = PassOutputs[0].RenderTargetDesc.Extent;
 
 	// e.g. 4 means the input texture is 4x smaller than the buffer size
-	uint32 ScaleFactor = GSceneRenderTargets.GetBufferSizeXY().X / SrcSize.X;
+	uint32 ScaleFactor = FSceneRenderTargets::Get(Context.RHICmdList).GetBufferSizeXY().X / SrcSize.X;
 
 	FIntRect SrcRect = FIntRect::DivideAndRoundUp(View.ViewRect, ScaleFactor);
 	FIntRect DestRect = FIntRect::DivideAndRoundUp(SrcRect, 2);
@@ -402,7 +402,7 @@ void FRCPassPostProcessBokehDOFSetup::Process(FRenderingCompositePassContext& Co
 	FIntPoint DestSize = PassOutputs[0].RenderTargetDesc.Extent;
 
 	// e.g. 4 means the input texture is 4x smaller than the buffer size
-	uint32 ScaleFactor = GSceneRenderTargets.GetBufferSizeXY().X / SrcSize.X;
+	uint32 ScaleFactor = FSceneRenderTargets::Get(Context.RHICmdList).GetBufferSizeXY().X / SrcSize.X;
 
 	FIntRect SrcRect = FIntRect::DivideAndRoundUp(View.ViewRect, ScaleFactor);
 	FIntRect DestRect = FIntRect::DivideAndRoundUp(SrcRect, 2);
@@ -654,7 +654,7 @@ void FRCPassPostProcessBokehDOF::SetShaderTempl(const FRenderingCompositePassCon
 
 void FRCPassPostProcessBokehDOF::ComputeDepthOfFieldParams(const FRenderingCompositePassContext& Context, FVector4 Out[2])
 {
-	uint32 FullRes = GSceneRenderTargets.GetBufferSizeXY().Y;
+	uint32 FullRes = FSceneRenderTargets::Get(Context.RHICmdList).GetBufferSizeXY().Y;
 	uint32 HalfRes = FMath::DivideAndRoundUp(FullRes, (uint32)2);
 	uint32 BokehLayerSizeY = HalfRes * 2 + SafetyBorder;
 
@@ -709,7 +709,7 @@ void FRCPassPostProcessBokehDOF::Process(FRenderingCompositePassContext& Context
 	FIntPoint TexSize = InputDesc->Extent;
 
 	// usually 1, 2, 4 or 8
-	uint32 ScaleToFullRes = GSceneRenderTargets.GetBufferSizeXY().X / TexSize.X;
+	uint32 ScaleToFullRes = FSceneRenderTargets::Get(Context.RHICmdList).GetBufferSizeXY().X / TexSize.X;
 
 	// don't use DivideAndRoundUp as this could cause cause lookups into areas we don't have setup 
 	FIntRect LocalViewRect = View.ViewRect / ScaleToFullRes;
@@ -823,7 +823,7 @@ FPooledRenderTargetDesc FRCPassPostProcessBokehDOF::ComputeOutputDesc(EPassOutpu
 	// more precision for additive blending
 	Ret.Format = PF_FloatRGBA;
 
-	uint32 FullRes = GSceneRenderTargets.GetBufferSizeXY().Y;
+	uint32 FullRes = FSceneRenderTargets::Get_FrameConstantsOnly().GetBufferSizeXY().Y;
 	uint32 HalfRes = FMath::DivideAndRoundUp(FullRes, (uint32)2);
 
 	// we need space for the front part and the back part
diff --git a/Engine/Source/Runtime/Renderer/Private/PostProcess/PostProcessBokehDOFRecombine.cpp b/Engine/Source/Runtime/Renderer/Private/PostProcess/PostProcessBokehDOFRecombine.cpp
--- a/Engine/Source/Runtime/Renderer/Private/PostProcess/PostProcessBokehDOFRecombine.cpp
+++ b/Engine/Source/Runtime/Renderer/Private/PostProcess/PostProcessBokehDOFRecombine.cpp
@@ -140,7 +140,7 @@ void FRCPassPostProcessBokehDOFRecombine::Process(FRenderingCompositePassContext
 	FIntPoint TexSize = InputDesc1 ? InputDesc1->Extent : InputDesc0->Extent;
 
 	// usually 1, 2, 4 or 8
-	uint32 ScaleToFullRes = GSceneRenderTargets.GetBufferSizeXY().X / TexSize.X;
+	uint32 ScaleToFullRes = FSceneRenderTargets::Get(Context.RHICmdList).GetBufferSizeXY().X / TexSize.X;
 
 	FIntRect HalfResViewRect = FIntRect::DivideAndRoundUp(View.ViewRect, ScaleToFullRes);
 
diff --git a/Engine/Source/Runtime/Renderer/Private/PostProcess/PostProcessBusyWait.cpp b/Engine/Source/Runtime/Renderer/Private/PostProcess/PostProcessBusyWait.cpp
--- a/Engine/Source/Runtime/Renderer/Private/PostProcess/PostProcessBusyWait.cpp
+++ b/Engine/Source/Runtime/Renderer/Private/PostProcess/PostProcessBusyWait.cpp
@@ -84,15 +84,16 @@ IMPLEMENT_SHADER_TYPE(,FPostProcessBusyWaitPS,TEXT("PostProcessBusyWait"),TEXT("
 void FRCPassPostProcessBusyWait::Process(FRenderingCompositePassContext& Context)
 {
 	SCOPED_DRAW_EVENT(Context.RHICmdList, BusyWait);
+	FSceneRenderTargets& SceneContext = FSceneRenderTargets::Get(Context.RHICmdList);
 
 	const FSceneView& View = Context.View;
 	
 	FIntRect SrcRect = View.ViewRect;
 	FIntRect DestRect = View.UnscaledViewRect;
 	
-	GSceneRenderTargets.BeginRenderingLightAttenuation(Context.RHICmdList);
+	SceneContext.BeginRenderingLightAttenuation(Context.RHICmdList);
 	
-	const FSceneRenderTargetItem& DestRenderTarget = GSceneRenderTargets.GetLightAttenuation()->GetRenderTargetItem();
+	const FSceneRenderTargetItem& DestRenderTarget = SceneContext.GetLightAttenuation()->GetRenderTargetItem();
 
 	// Set the view family's render target/viewport.
 	SetRenderTarget(Context.RHICmdList, DestRenderTarget.TargetableTexture, FTextureRHIRef());
@@ -126,7 +127,7 @@ void FRCPassPostProcessBusyWait::Process(FRenderingCompositePassContext& Context
 
 	Context.RHICmdList.CopyToResolveTarget(DestRenderTarget.TargetableTexture, DestRenderTarget.ShaderResourceTexture, false, FResolveParams());
 
-	GSceneRenderTargets.SetLightAttenuation(0);
+	SceneContext.SetLightAttenuation(0);
 }
 
 FPooledRenderTargetDesc FRCPassPostProcessBusyWait::ComputeOutputDesc(EPassOutputId InPassOutputId) const
diff --git a/Engine/Source/Runtime/Renderer/Private/PostProcess/PostProcessCombineLUTs.cpp b/Engine/Source/Runtime/Renderer/Private/PostProcess/PostProcessCombineLUTs.cpp
--- a/Engine/Source/Runtime/Renderer/Private/PostProcess/PostProcessCombineLUTs.cpp
+++ b/Engine/Source/Runtime/Renderer/Private/PostProcess/PostProcessCombineLUTs.cpp
@@ -319,9 +319,9 @@ public:
 			float FilmSlopeS = FilmMidYS / (FilmMidXS);
 			float FilmHiYS = 1.0f - FilmHiY;
 			float FilmLoYS = FilmLoY;
-			float FilmToeS = FilmLoX;
+			float FilmToe = FilmLoX;
 			float FilmHiG = (-FilmHiYS + (FilmSlopeS*FilmHeal)) / (FilmSlopeS*FilmHeal);
-			float FilmLoG = (-FilmLoYS + (FilmSlopeS*FilmToeS)) / (FilmSlopeS*FilmToeS);
+			float FilmLoG = (-FilmLoYS + (FilmSlopeS*FilmToe)) / (FilmSlopeS*FilmToe);
 
 			// Constants.
 			float OutColorCurveCh1 = FilmHiYS/FilmHiG;
diff --git a/Engine/Source/Runtime/Renderer/Private/PostProcess/PostProcessCompositeEditorPrimitives.cpp b/Engine/Source/Runtime/Renderer/Private/PostProcess/PostProcessCompositeEditorPrimitives.cpp
--- a/Engine/Source/Runtime/Renderer/Private/PostProcess/PostProcessCompositeEditorPrimitives.cpp
+++ b/Engine/Source/Runtime/Renderer/Private/PostProcess/PostProcessCompositeEditorPrimitives.cpp
@@ -69,6 +69,7 @@ public:
 
 	void SetParameters(const FRenderingCompositePassContext& Context)
 	{
+		FSceneRenderTargets& SceneContext = FSceneRenderTargets::Get(Context.RHICmdList);
 		const FPixelShaderRHIParamRef ShaderRHI = GetPixelShader();
 
 		FGlobalShader::SetParameters(Context.RHICmdList, ShaderRHI, Context.View);
@@ -79,13 +80,13 @@ public:
 		PostProcessParameters.SetPS(ShaderRHI, Context, SamplerStateRHIRef);
 		if(MSAASampleCount > 1)
 		{
-			SetTextureParameter(Context.RHICmdList, ShaderRHI, EditorPrimitivesColor, GSceneRenderTargets.EditorPrimitivesColor->GetRenderTargetItem().TargetableTexture);
-			SetTextureParameter(Context.RHICmdList, ShaderRHI, EditorPrimitivesDepth, GSceneRenderTargets.EditorPrimitivesDepth->GetRenderTargetItem().TargetableTexture);
+			SetTextureParameter(Context.RHICmdList, ShaderRHI, EditorPrimitivesColor, SceneContext.EditorPrimitivesColor->GetRenderTargetItem().TargetableTexture);
+			SetTextureParameter(Context.RHICmdList, ShaderRHI, EditorPrimitivesDepth, SceneContext.EditorPrimitivesDepth->GetRenderTargetItem().TargetableTexture);
 		}
 		else
 		{
-			SetTextureParameter(Context.RHICmdList, ShaderRHI, EditorPrimitivesColor, EditorPrimitivesColorSampler, SamplerStateRHIRef, GSceneRenderTargets.EditorPrimitivesColor->GetRenderTargetItem().ShaderResourceTexture);
-			SetTextureParameter(Context.RHICmdList, ShaderRHI, EditorPrimitivesDepth, GSceneRenderTargets.EditorPrimitivesDepth->GetRenderTargetItem().ShaderResourceTexture);
+			SetTextureParameter(Context.RHICmdList, ShaderRHI, EditorPrimitivesColor, EditorPrimitivesColorSampler, SamplerStateRHIRef, SceneContext.EditorPrimitivesColor->GetRenderTargetItem().ShaderResourceTexture);
+			SetTextureParameter(Context.RHICmdList, ShaderRHI, EditorPrimitivesDepth, SceneContext.EditorPrimitivesDepth->GetRenderTargetItem().ShaderResourceTexture);
 		}
 
 		{
@@ -110,7 +111,7 @@ public:
 
 		if(FilteredSceneDepthTexture.IsBound())
 		{
-			const FTexture2DRHIRef* DepthTexture = GSceneRenderTargets.GetActualDepthTexture();
+			const FTexture2DRHIRef* DepthTexture = SceneContext.GetActualDepthTexture();
 			SetTextureParameter(
 				Context.RHICmdList,
 				ShaderRHI,
@@ -197,15 +198,16 @@ void FRCPassPostProcessCompositeEditorPrimitives::Process(FRenderingCompositePas
 	FIntRect SrcRect = View.ViewRect;
 	FIntRect DestRect = View.ViewRect;
 	FIntPoint SrcSize = InputDesc->Extent;
+	FSceneRenderTargets& SceneContext = FSceneRenderTargets::Get(Context.RHICmdList);
 
 	// If we render wirframe we already started rendering to the EditorPrimitives buffer, so we don't want to clear it.
-	bool bClearIsNeeded = !IsValidRef(GSceneRenderTargets.EditorPrimitivesColor);
+	bool bClearIsNeeded = !IsValidRef(SceneContext.EditorPrimitivesColor);
 
 	// Get or create the msaa depth and color buffer
-	FTexture2DRHIRef ColorTarget = GSceneRenderTargets.GetEditorPrimitivesColor();
-	FTexture2DRHIRef DepthTarget = GSceneRenderTargets.GetEditorPrimitivesDepth();
+	FTexture2DRHIRef ColorTarget = SceneContext.GetEditorPrimitivesColor();
+	FTexture2DRHIRef DepthTarget = SceneContext.GetEditorPrimitivesDepth();
 
-	const uint32 MSAASampleCount = GSceneRenderTargets.EditorPrimitivesColor->GetDesc().NumSamples;
+	const uint32 MSAASampleCount = SceneContext.EditorPrimitivesColor->GetDesc().NumSamples;
 
 	{
 		SetRenderTarget(Context.RHICmdList, ColorTarget, DepthTarget);
@@ -231,10 +233,10 @@ void FRCPassPostProcessCompositeEditorPrimitives::Process(FRenderingCompositePas
 			RenderPrimitivesToComposite<FBasePassForwardOpaqueDrawingPolicyFactory>(Context.RHICmdList, View);
 		}
 
-		GRenderTargetPool.VisualizeTexture.SetCheckPoint(Context.RHICmdList, GSceneRenderTargets.EditorPrimitivesColor);
+		GRenderTargetPool.VisualizeTexture.SetCheckPoint(Context.RHICmdList, SceneContext.EditorPrimitivesColor);
 	}
 
-	//	Context.RHICmdList.CopyToResolveTarget(GSceneRenderTargets.EditorPrimitivesColor->GetRenderTargetItem().TargetableTexture, GSceneRenderTargets.EditorPrimitivesColor->GetRenderTargetItem().ShaderResourceTexture, false, FResolveParams());
+	//	Context.RHICmdList.CopyToResolveTarget(SceneContext.EditorPrimitivesColor->GetRenderTargetItem().TargetableTexture, SceneContext.EditorPrimitivesColor->GetRenderTargetItem().ShaderResourceTexture, false, FResolveParams());
 
 	const FSceneRenderTargetItem& DestRenderTarget = PassOutputs[0].RequestSurface(Context);
 	const FTexture2DRHIRef& DestRenderTargetSurface = (const FTexture2DRHIRef&)DestRenderTarget.TargetableTexture;
@@ -288,7 +290,7 @@ void FRCPassPostProcessCompositeEditorPrimitives::Process(FRenderingCompositePas
 	Context.RHICmdList.CopyToResolveTarget(DestRenderTargetSurface, DestRenderTarget.ShaderResourceTexture, false, FResolveParams());
 
 	// Clean up targets
-	GSceneRenderTargets.CleanUpEditorPrimitiveTargets();
+	SceneContext.CleanUpEditorPrimitiveTargets();
 }
 
 FPooledRenderTargetDesc FRCPassPostProcessCompositeEditorPrimitives::ComputeOutputDesc(EPassOutputId InPassOutputId) const
@@ -304,6 +306,7 @@ FPooledRenderTargetDesc FRCPassPostProcessCompositeEditorPrimitives::ComputeOutp
 template<typename TBasePass>
 void FRCPassPostProcessCompositeEditorPrimitives::RenderPrimitivesToComposite(FRHICommandListImmediate& RHICmdList, const FViewInfo& View)
 {
+	FSceneRenderTargets& SceneContext = FSceneRenderTargets::Get(RHICmdList);
 
 	// Always depth test against other editor primitives
 	RHICmdList.SetDepthStencilState(TStaticDepthStencilState<true, CF_DepthNearOrEqual>::GetRHI());
@@ -315,7 +318,7 @@ void FRCPassPostProcessCompositeEditorPrimitives::RenderPrimitivesToComposite(FR
 	const auto FeatureLevel = View.GetFeatureLevel();
 	const auto ShaderPlatform = GShaderPlatformForFeatureLevel[FeatureLevel];
 	const bool bNeedToSwitchVerticalAxis = RHINeedsToSwitchVerticalAxis(ShaderPlatform);
-	FTexture2DRHIRef SceneDepth = GSceneRenderTargets.GetSceneDepthTexture();
+	FTexture2DRHIRef SceneDepth = SceneContext.GetSceneDepthTexture();
 
 	typename TBasePass::ContextType Context(true, ESceneRenderTargetsMode::SetTextures);
 
@@ -330,7 +333,7 @@ void FRCPassPostProcessCompositeEditorPrimitives::RenderPrimitivesToComposite(FR
 		}
 	}
 
-	View.EditorSimpleElementCollector.DrawBatchedElements(RHICmdList, View, GSceneRenderTargets.GetSceneDepthTexture(), EBlendModeFilter::OpaqueAndMasked);
+	View.EditorSimpleElementCollector.DrawBatchedElements(RHICmdList, View, SceneContext.GetSceneDepthTexture(), EBlendModeFilter::OpaqueAndMasked);
 	
 	// Draw the base pass for the view's batched mesh elements.
 	DrawViewElements<TBasePass>(RHICmdList, View, typename TBasePass::ContextType(bDepthTest, ESceneRenderTargetsMode::SetTextures), SDPG_World, false);
diff --git a/Engine/Source/Runtime/Renderer/Private/PostProcess/PostProcessDOF.cpp b/Engine/Source/Runtime/Renderer/Private/PostProcess/PostProcessDOF.cpp
--- a/Engine/Source/Runtime/Renderer/Private/PostProcess/PostProcessDOF.cpp
+++ b/Engine/Source/Runtime/Renderer/Private/PostProcess/PostProcessDOF.cpp
@@ -102,7 +102,7 @@ void FRCPassPostProcessDOFSetup::Process(FRenderingCompositePassContext& Context
 	FIntPoint DestSize = PassOutputs[0].RenderTargetDesc.Extent;
 
 	// e.g. 4 means the input texture is 4x smaller than the buffer size
-	uint32 ScaleFactor = GSceneRenderTargets.GetBufferSizeXY().X / SrcSize.X;
+	uint32 ScaleFactor = FSceneRenderTargets::Get(Context.RHICmdList).GetBufferSizeXY().X / SrcSize.X;
 
 	FIntRect SrcRect = View.ViewRect / ScaleFactor;
 	FIntRect DestRect = SrcRect / 2;
@@ -241,6 +241,7 @@ public:
 
 	void SetParameters(const FRenderingCompositePassContext& Context)
 	{
+		FSceneRenderTargets& SceneContext = FSceneRenderTargets::Get(Context.RHICmdList);
 		const FPixelShaderRHIParamRef ShaderRHI = GetPixelShader();
 
 		FGlobalShader::SetParameters(Context.RHICmdList, ShaderRHI, Context.View);
@@ -250,10 +251,10 @@ public:
 		
 		// Compute out of bounds UVs in the source texture.
 		FVector4 Bounds;
-		Bounds.X = (((float)((Context.View.ViewRect.Min.X + 1) & (~1))) + 3.0f) / ((float)(GSceneRenderTargets.GetBufferSizeXY().X));
-		Bounds.Y = (((float)((Context.View.ViewRect.Min.Y + 1) & (~1))) + 3.0f) / ((float)(GSceneRenderTargets.GetBufferSizeXY().Y));
-		Bounds.Z = (((float)(Context.View.ViewRect.Max.X & (~1))) - 3.0f) / ((float)(GSceneRenderTargets.GetBufferSizeXY().X));
-		Bounds.W = (((float)(Context.View.ViewRect.Max.Y & (~1))) - 3.0f) / ((float)(GSceneRenderTargets.GetBufferSizeXY().Y));
+		Bounds.X = (((float)((Context.View.ViewRect.Min.X + 1) & (~1))) + 3.0f) / ((float)(SceneContext.GetBufferSizeXY().X));
+		Bounds.Y = (((float)((Context.View.ViewRect.Min.Y + 1) & (~1))) + 3.0f) / ((float)(SceneContext.GetBufferSizeXY().Y));
+		Bounds.Z = (((float)(Context.View.ViewRect.Max.X & (~1))) - 3.0f) / ((float)(SceneContext.GetBufferSizeXY().X));
+		Bounds.W = (((float)(Context.View.ViewRect.Max.Y & (~1))) - 3.0f) / ((float)(SceneContext.GetBufferSizeXY().Y));
 
 		SetShaderValue(Context.RHICmdList, ShaderRHI, DepthOfFieldUVLimit, Bounds);
 	}
@@ -282,7 +283,7 @@ void FRCPassPostProcessDOFRecombine::Process(FRenderingCompositePassContext& Con
 	FIntPoint TexSize = InputDesc->Extent;
 
 	// usually 1, 2, 4 or 8
-	uint32 ScaleToFullRes = GSceneRenderTargets.GetBufferSizeXY().X / TexSize.X;
+	uint32 ScaleToFullRes = FSceneRenderTargets::Get(Context.RHICmdList).GetBufferSizeXY().X / TexSize.X;
 
 	FIntRect HalfResViewRect = View.ViewRect / ScaleToFullRes;
 
@@ -501,7 +502,7 @@ void FRCPassPostProcessCircleDOFSetup::Process(FRenderingCompositePassContext& C
 	FIntPoint DestSize = PassOutputs[0].RenderTargetDesc.Extent;
 
 	// e.g. 4 means the input texture is 4x smaller than the buffer size
-	uint32 ScaleFactor = GSceneRenderTargets.GetBufferSizeXY().X / SrcSize.X;
+	uint32 ScaleFactor = FSceneRenderTargets::Get(Context.RHICmdList).GetBufferSizeXY().X / SrcSize.X;
 
 	FIntRect SrcRect = View.ViewRect / ScaleFactor;
 	FIntRect DestRect = SrcRect / 2;
@@ -689,7 +690,7 @@ void FRCPassPostProcessCircleDOFDilate::Process(FRenderingCompositePassContext&
 	FIntPoint DestSize = PassOutputs[0].RenderTargetDesc.Extent;
 
 	// e.g. 4 means the input texture is 4x smaller than the buffer size
-	uint32 ScaleFactor = GSceneRenderTargets.GetBufferSizeXY().X / SrcSize.X;
+	uint32 ScaleFactor = FSceneRenderTargets::Get(Context.RHICmdList).GetBufferSizeXY().X / SrcSize.X;
 
 	FIntRect SrcRect = View.ViewRect / ScaleFactor;
 	FIntRect DestRect = SrcRect / 2;
@@ -909,7 +910,7 @@ void FRCPassPostProcessCircleDOF::Process(FRenderingCompositePassContext& Contex
 	FIntPoint DestSize = PassOutputs[0].RenderTargetDesc.Extent;
 
 	// e.g. 4 means the input texture is 4x smaller than the buffer size
-	uint32 ScaleFactor = GSceneRenderTargets.GetBufferSizeXY().X / SrcSize.X;
+	uint32 ScaleFactor = FSceneRenderTargets::Get(Context.RHICmdList).GetBufferSizeXY().X / SrcSize.X;
 
 	FIntRect SrcRect = View.ViewRect / ScaleFactor;
 	FIntRect DestRect = SrcRect;
@@ -1048,6 +1049,7 @@ public:
 
 	void SetParameters(const FRenderingCompositePassContext& Context)
 	{
+		FSceneRenderTargets& SceneContext = FSceneRenderTargets::Get(Context.RHICmdList);
 		const FPixelShaderRHIParamRef ShaderRHI = GetPixelShader();
 
 		FGlobalShader::SetParameters(Context.RHICmdList, ShaderRHI, Context.View);
@@ -1057,10 +1059,10 @@ public:
 
 		// Compute out of bounds UVs in the source texture.
 		FVector4 Bounds;
-		Bounds.X = (((float)((Context.View.ViewRect.Min.X + 1) & (~1))) + 3.0f) / ((float)(GSceneRenderTargets.GetBufferSizeXY().X));
-		Bounds.Y = (((float)((Context.View.ViewRect.Min.Y + 1) & (~1))) + 3.0f) / ((float)(GSceneRenderTargets.GetBufferSizeXY().Y));
-		Bounds.Z = (((float)(Context.View.ViewRect.Max.X & (~1))) - 3.0f) / ((float)(GSceneRenderTargets.GetBufferSizeXY().X));
-		Bounds.W = (((float)(Context.View.ViewRect.Max.Y & (~1))) - 3.0f) / ((float)(GSceneRenderTargets.GetBufferSizeXY().Y));
+		Bounds.X = (((float)((Context.View.ViewRect.Min.X + 1) & (~1))) + 3.0f) / ((float)(SceneContext.GetBufferSizeXY().X));
+		Bounds.Y = (((float)((Context.View.ViewRect.Min.Y + 1) & (~1))) + 3.0f) / ((float)(SceneContext.GetBufferSizeXY().Y));
+		Bounds.Z = (((float)(Context.View.ViewRect.Max.X & (~1))) - 3.0f) / ((float)(SceneContext.GetBufferSizeXY().X));
+		Bounds.W = (((float)(Context.View.ViewRect.Max.Y & (~1))) - 3.0f) / ((float)(SceneContext.GetBufferSizeXY().Y));
 
 		SetShaderValue(Context.RHICmdList, ShaderRHI, DepthOfFieldUVLimit, Bounds);
 
@@ -1095,7 +1097,7 @@ void FRCPassPostProcessCircleDOFRecombine::Process(FRenderingCompositePassContex
 	FIntPoint TexSize = InputDesc->Extent;
 
 	// usually 1, 2, 4 or 8
-	uint32 ScaleToFullRes = GSceneRenderTargets.GetBufferSizeXY().X / TexSize.X;
+	uint32 ScaleToFullRes = FSceneRenderTargets::Get(Context.RHICmdList).GetBufferSizeXY().X / TexSize.X;
 
 	FIntRect HalfResViewRect = View.ViewRect / ScaleToFullRes;
 
diff --git a/Engine/Source/Runtime/Renderer/Private/PostProcess/PostProcessDownsample.cpp b/Engine/Source/Runtime/Renderer/Private/PostProcess/PostProcessDownsample.cpp
--- a/Engine/Source/Runtime/Renderer/Private/PostProcess/PostProcessDownsample.cpp
+++ b/Engine/Source/Runtime/Renderer/Private/PostProcess/PostProcessDownsample.cpp
@@ -171,7 +171,7 @@ void FRCPassPostProcessDownsample::Process(FRenderingCompositePassContext& Conte
 	FIntPoint DestSize = PassOutputs[0].RenderTargetDesc.Extent;
 
 	// e.g. 4 means the input texture is 4x smaller than the buffer size
-	uint32 ScaleFactor = GSceneRenderTargets.GetBufferSizeXY().X / SrcSize.X;
+	uint32 ScaleFactor = FSceneRenderTargets::Get(Context.RHICmdList).GetBufferSizeXY().X / SrcSize.X;
 
 	const FSceneRenderTargetItem& DestRenderTarget = PassOutputs[0].RequestSurface(Context);
 
diff --git a/Engine/Source/Runtime/Renderer/Private/PostProcess/PostProcessGBufferHints.cpp b/Engine/Source/Runtime/Renderer/Private/PostProcess/PostProcessGBufferHints.cpp
--- a/Engine/Source/Runtime/Renderer/Private/PostProcess/PostProcessGBufferHints.cpp
+++ b/Engine/Source/Runtime/Renderer/Private/PostProcess/PostProcessGBufferHints.cpp
@@ -81,7 +81,7 @@ IMPLEMENT_SHADER_TYPE(,FPostProcessGBufferHintsPS,TEXT("PostProcessGBufferHints"
 FRCPassPostProcessGBufferHints::FRCPassPostProcessGBufferHints()
 {
 	// AdjustGBufferRefCount(-1) call is done when the pass gets executed
-	GSceneRenderTargets.AdjustGBufferRefCount(1);
+	FSceneRenderTargets::Get_Todo_PassContext().AdjustGBufferRefCount(1);
 }
 
 
@@ -183,7 +183,7 @@ void FRCPassPostProcessGBufferHints::Process(FRenderingCompositePassContext& Con
 	Context.RHICmdList.CopyToResolveTarget(DestRenderTarget.TargetableTexture, DestRenderTarget.ShaderResourceTexture, false, FResolveParams());
 
 	// AdjustGBufferRefCount(1) call is done in constructor
-	GSceneRenderTargets.AdjustGBufferRefCount(-1);
+	FSceneRenderTargets::Get(Context.RHICmdList).AdjustGBufferRefCount(-1);
 }
 
 FPooledRenderTargetDesc FRCPassPostProcessGBufferHints::ComputeOutputDesc(EPassOutputId InPassOutputId) const
diff --git a/Engine/Source/Runtime/Renderer/Private/PostProcess/PostProcessLensBlur.cpp b/Engine/Source/Runtime/Renderer/Private/PostProcess/PostProcessLensBlur.cpp
--- a/Engine/Source/Runtime/Renderer/Private/PostProcess/PostProcessLensBlur.cpp
+++ b/Engine/Source/Runtime/Renderer/Private/PostProcess/PostProcessLensBlur.cpp
@@ -177,7 +177,7 @@ void FRCPassPostProcessLensBlur::Process(FRenderingCompositePassContext& Context
 	FIntPoint TexSize = InputDesc->Extent;
 
 	// usually 1, 2, 4 or 8
-	uint32 ScaleToFullRes = GSceneRenderTargets.GetBufferSizeXY().X / TexSize.X;
+	uint32 ScaleToFullRes = FSceneRenderTargets::Get(Context.RHICmdList).GetBufferSizeXY().X / TexSize.X;
 
 	FIntRect ViewRect = FIntRect::DivideAndRoundUp(View.ViewRect, ScaleToFullRes);
 	FIntPoint ViewSize = ViewRect.Size();
diff --git a/Engine/Source/Runtime/Renderer/Private/PostProcess/PostProcessLensFlares.cpp b/Engine/Source/Runtime/Renderer/Private/PostProcess/PostProcessLensFlares.cpp
--- a/Engine/Source/Runtime/Renderer/Private/PostProcess/PostProcessLensFlares.cpp
+++ b/Engine/Source/Runtime/Renderer/Private/PostProcess/PostProcessLensFlares.cpp
@@ -130,8 +130,9 @@ void FRCPassPostProcessLensFlares::Process(FRenderingCompositePassContext& Conte
 	FIntPoint TexSize1 = InputDesc1->Extent;
 	FIntPoint TexSize2 = InputDesc2->Extent;
 
-	uint32 ScaleToFullRes1 = GSceneRenderTargets.GetBufferSizeXY().X / TexSize1.X;
-	uint32 ScaleToFullRes2 = GSceneRenderTargets.GetBufferSizeXY().X / TexSize2.X;
+	FSceneRenderTargets& SceneContext = FSceneRenderTargets::Get(Context.RHICmdList);
+	uint32 ScaleToFullRes1 = SceneContext.GetBufferSizeXY().X / TexSize1.X;
+	uint32 ScaleToFullRes2 = SceneContext.GetBufferSizeXY().X / TexSize2.X;
 
 	FIntRect ViewRect1 = FIntRect::DivideAndRoundUp(View.ViewRect, ScaleToFullRes1);
 	FIntRect ViewRect2 = FIntRect::DivideAndRoundUp(View.ViewRect, ScaleToFullRes2);
diff --git a/Engine/Source/Runtime/Renderer/Private/PostProcess/PostProcessMaterial.cpp b/Engine/Source/Runtime/Renderer/Private/PostProcess/PostProcessMaterial.cpp
--- a/Engine/Source/Runtime/Renderer/Private/PostProcess/PostProcessMaterial.cpp
+++ b/Engine/Source/Runtime/Renderer/Private/PostProcess/PostProcessMaterial.cpp
@@ -202,7 +202,7 @@ void FRCPassPostProcessMaterial::Process(FRenderingCompositePassContext& Context
 
 	if(Material->NeedsGBuffer())
 	{
-		GSceneRenderTargets.AdjustGBufferRefCount(-1);
+		FSceneRenderTargets::Get(Context.RHICmdList).AdjustGBufferRefCount(-1);
 	}
 }
 
diff --git a/Engine/Source/Runtime/Renderer/Private/PostProcess/PostProcessMotionBlur.cpp b/Engine/Source/Runtime/Renderer/Private/PostProcess/PostProcessMotionBlur.cpp
--- a/Engine/Source/Runtime/Renderer/Private/PostProcess/PostProcessMotionBlur.cpp
+++ b/Engine/Source/Runtime/Renderer/Private/PostProcess/PostProcessMotionBlur.cpp
@@ -161,7 +161,7 @@ void FRCPassPostProcessMotionBlurSetup::Process(FRenderingCompositePassContext&
 	FIntPoint DestSize = PassOutputs[0].RenderTargetDesc.Extent;
 
 	// e.g. 4 means the input texture is 4x smaller than the buffer size
-	uint32 ScaleFactor = GSceneRenderTargets.GetBufferSizeXY().X / SrcSize.X;
+	uint32 ScaleFactor = FSceneRenderTargets::Get(Context.RHICmdList).GetBufferSizeXY().X / SrcSize.X;
 
 	FIntRect SrcRect = View.ViewRect / ScaleFactor;
 
@@ -361,7 +361,7 @@ public:
 
 		// to mask out samples from outside of the view
 		{
-			FIntPoint BufferSize = GSceneRenderTargets.GetBufferSizeXY();
+			FIntPoint BufferSize = FSceneRenderTargets::Get(Context.RHICmdList).GetBufferSizeXY();
 			FVector2D InvBufferSize(1.0f / BufferSize.X, 1.0f / BufferSize.Y);
 
 			FIntRect ClipRect = Context.View.ViewRect;
@@ -476,7 +476,7 @@ void FRCPassPostProcessMotionBlur::Process(FRenderingCompositePassContext& Conte
 	FIntPoint DestSize = PassOutputs[0].RenderTargetDesc.Extent;
 
 	// e.g. 4 means the input texture is 4x smaller than the buffer size
-	uint32 ScaleFactor = GSceneRenderTargets.GetBufferSizeXY().X / SrcSize.X;
+	uint32 ScaleFactor = FSceneRenderTargets::Get(Context.RHICmdList).GetBufferSizeXY().X / SrcSize.X;
 
 	FIntRect SrcRect = FIntRect::DivideAndRoundUp(View.ViewRect, ScaleFactor);
 	FIntRect DestRect = SrcRect;
@@ -618,7 +618,7 @@ void FRCPassPostProcessMotionBlurRecombine::Process(FRenderingCompositePassConte
 	FIntPoint DestSize = PassOutputs[0].RenderTargetDesc.Extent;
 
 	// e.g. 4 means the input texture is 4x smaller than the buffer size
-	uint32 ScaleFactor = GSceneRenderTargets.GetBufferSizeXY().X / SrcSize.X;
+	uint32 ScaleFactor = FSceneRenderTargets::Get(Context.RHICmdList).GetBufferSizeXY().X / SrcSize.X;
 
 	FIntRect SrcRect = View.ViewRect / ScaleFactor;
 	FIntRect DestRect = SrcRect;
@@ -1054,7 +1054,7 @@ void FRCPassPostProcessVelocityScatter::Process(FRenderingCompositePassContext&
 	FIntPoint DestSize = PassOutputs[0].RenderTargetDesc.Extent;
 
 	// e.g. 4 means the input texture is 4x smaller than the buffer size
-	uint32 ScaleFactor = GSceneRenderTargets.GetBufferSizeXY().X / SrcSize.X;
+	uint32 ScaleFactor = FSceneRenderTargets::Get(Context.RHICmdList).GetBufferSizeXY().X / SrcSize.X;
 
 	FIntRect SrcRect = FIntRect::DivideAndRoundUp(View.ViewRect, ScaleFactor);
 	FIntRect DestRect = SrcRect;
@@ -1182,7 +1182,7 @@ void FRCPassPostProcessVelocityDilate::Process(FRenderingCompositePassContext& C
 	FIntPoint DestSize = PassOutputs[0].RenderTargetDesc.Extent;
 
 	// e.g. 4 means the input texture is 4x smaller than the buffer size
-	uint32 ScaleFactor = GSceneRenderTargets.GetBufferSizeXY().X / SrcSize.X;
+	uint32 ScaleFactor = FSceneRenderTargets::Get(Context.RHICmdList).GetBufferSizeXY().X / SrcSize.X;
 
 	FIntRect SrcRect = FIntRect::DivideAndRoundUp(View.ViewRect, ScaleFactor);
 	FIntRect DestRect = SrcRect;
@@ -1418,7 +1418,7 @@ void FRCPassPostProcessMotionBlurNew::Process(FRenderingCompositePassContext& Co
 	FIntPoint DestSize = PassOutputs[0].RenderTargetDesc.Extent;
 
 	// e.g. 4 means the input texture is 4x smaller than the buffer size
-	uint32 ScaleFactor = GSceneRenderTargets.GetBufferSizeXY().X / SrcSize.X;
+	uint32 ScaleFactor = FSceneRenderTargets::Get(Context.RHICmdList).GetBufferSizeXY().X / SrcSize.X;
 
 	FIntRect SrcRect = View.ViewRect / ScaleFactor;
 	FIntRect DestRect = SrcRect;
@@ -1509,7 +1509,7 @@ void FRCPassPostProcessVisualizeMotionBlur::Process(FRenderingCompositePassConte
 	FIntPoint DestSize = PassOutputs[0].RenderTargetDesc.Extent;
 
 	// e.g. 4 means the input texture is 4x smaller than the buffer size
-	uint32 ScaleFactor = GSceneRenderTargets.GetBufferSizeXY().X / SrcSize.X;
+	uint32 ScaleFactor = FSceneRenderTargets::Get(Context.RHICmdList).GetBufferSizeXY().X / SrcSize.X;
 
 	FIntRect SrcRect = FIntRect::DivideAndRoundUp(View.ViewRect, ScaleFactor);
 	FIntRect DestRect = SrcRect;
diff --git a/Engine/Source/Runtime/Renderer/Private/PostProcess/PostProcessNoiseBlur.cpp b/Engine/Source/Runtime/Renderer/Private/PostProcess/PostProcessNoiseBlur.cpp
--- a/Engine/Source/Runtime/Renderer/Private/PostProcess/PostProcessNoiseBlur.cpp
+++ b/Engine/Source/Runtime/Renderer/Private/PostProcess/PostProcessNoiseBlur.cpp
@@ -129,7 +129,7 @@ void FRCPassPostProcessNoiseBlur::Process(FRenderingCompositePassContext& Contex
 	FIntPoint DestSize = PassOutputs[0].RenderTargetDesc.Extent;
 
 	// e.g. 4 means the input texture is 4x smaller than the buffer size
-	uint32 ScaleFactor = GSceneRenderTargets.GetBufferSizeXY().X / SrcSize.X;
+	uint32 ScaleFactor = FSceneRenderTargets::Get(Context.RHICmdList).GetBufferSizeXY().X / SrcSize.X;
 
 	FIntRect SrcRect = View.ViewRect / ScaleFactor;
 	FIntRect DestRect = SrcRect;
diff --git a/Engine/Source/Runtime/Renderer/Private/PostProcess/PostProcessSceneColorFringe.cpp b/Engine/Source/Runtime/Renderer/Private/PostProcess/PostProcessSceneColorFringe.cpp
--- a/Engine/Source/Runtime/Renderer/Private/PostProcess/PostProcessSceneColorFringe.cpp
+++ b/Engine/Source/Runtime/Renderer/Private/PostProcess/PostProcessSceneColorFringe.cpp
@@ -186,7 +186,7 @@ void FRCPassPostProcessSceneColorFringe::Process(FRenderingCompositePassContext&
 		View.ViewRect.Min.X, View.ViewRect.Min.Y, 
 		View.ViewRect.Width(), View.ViewRect.Height(),
 		View.ViewRect.Size(),
-		GSceneRenderTargets.GetBufferSizeXY(),
+		FSceneRenderTargets::Get(Context.RHICmdList).GetBufferSizeXY(),
 		*VertexShader,
 		EDRF_UseTriangleOptimization);
 
diff --git a/Engine/Source/Runtime/Renderer/Private/PostProcess/PostProcessSelectionOutline.cpp b/Engine/Source/Runtime/Renderer/Private/PostProcess/PostProcessSelectionOutline.cpp
--- a/Engine/Source/Runtime/Renderer/Private/PostProcess/PostProcessSelectionOutline.cpp
+++ b/Engine/Source/Runtime/Renderer/Private/PostProcess/PostProcessSelectionOutline.cpp
@@ -133,7 +133,7 @@ FPooledRenderTargetDesc FRCPassPostProcessSelectionOutlineColor::ComputeOutputDe
 	Ret.Flags = TexCreate_None;
 	Ret.TargetableFlags = TexCreate_DepthStencilTargetable;
 	Ret.DebugName = TEXT("SelectionDepthStencil");
-	Ret.NumSamples = GSceneRenderTargets.GetEditorMSAACompositingSampleCount();
+	Ret.NumSamples = FSceneRenderTargets::Get_FrameConstantsOnly().GetEditorMSAACompositingSampleCount();
 
 	return Ret;
 }
@@ -348,7 +348,7 @@ void FRCPassPostProcessSelectionOutline::Process(FRenderingCompositePassContext&
 	Context.RHICmdList.SetRasterizerState(TStaticRasterizerState<>::GetRHI());
 	Context.RHICmdList.SetDepthStencilState(TStaticDepthStencilState<false, CF_Always>::GetRHI());
 
-	const uint32 MSAASampleCount = GSceneRenderTargets.GetEditorMSAACompositingSampleCount();
+	const uint32 MSAASampleCount = FSceneRenderTargets::Get(Context.RHICmdList).GetEditorMSAACompositingSampleCount();
 
 	if(MSAASampleCount == 1)
 	{
diff --git a/Engine/Source/Runtime/Renderer/Private/PostProcess/PostProcessSubsurface.cpp b/Engine/Source/Runtime/Renderer/Private/PostProcess/PostProcessSubsurface.cpp
--- a/Engine/Source/Runtime/Renderer/Private/PostProcess/PostProcessSubsurface.cpp
+++ b/Engine/Source/Runtime/Renderer/Private/PostProcess/PostProcessSubsurface.cpp
@@ -66,8 +66,9 @@ public:
 			// from Separabale.usf: float distanceToProjectionWindow = 1.0 / tan(0.5 * radians(SSSS_FOVY))
 			// can be extracted out of projection matrix
 
-			float ScaleCorrectionX = Context.View.ViewRect.Width() / (float)GSceneRenderTargets.GetBufferSizeXY().X;
-			float ScaleCorrectionY = Context.View.ViewRect.Height() / (float)GSceneRenderTargets.GetBufferSizeXY().Y;
+			FSceneRenderTargets& SceneContext = FSceneRenderTargets::Get(RHICmdList);
+			float ScaleCorrectionX = Context.View.ViewRect.Width() / (float)SceneContext.GetBufferSizeXY().X;
+			float ScaleCorrectionY = Context.View.ViewRect.Height() / (float)SceneContext.GetBufferSizeXY().Y;
 
 			 // Divide by 3 as the kernels range from -3 to 3.
 			const float KernelSize = 3.0f;
@@ -198,7 +199,7 @@ void SetSubsurfaceVisualizeShader(const FRenderingCompositePassContext& Context)
 FRCPassPostProcessSubsurfaceVisualize::FRCPassPostProcessSubsurfaceVisualize()
 {
 	// we need the GBuffer, we release it Process()
-	GSceneRenderTargets.AdjustGBufferRefCount(1);
+	FSceneRenderTargets::Get_Todo_PassContext().AdjustGBufferRefCount(1);
 }
 
 void FRCPassPostProcessSubsurfaceVisualize::Process(FRenderingCompositePassContext& Context)
@@ -219,8 +220,9 @@ void FRCPassPostProcessSubsurfaceVisualize::Process(FRenderingCompositePassConte
 	FIntPoint SrcSize = InputDesc->Extent;
 	FIntPoint DestSize = PassOutputs[0].RenderTargetDesc.Extent;
 
+	FSceneRenderTargets& SceneContext = FSceneRenderTargets::Get(Context.RHICmdList);
 	// e.g. 4 means the input texture is 4x smaller than the buffer size
-	uint32 ScaleFactor = GSceneRenderTargets.GetBufferSizeXY().X / SrcSize.X;
+	uint32 ScaleFactor = SceneContext.GetBufferSizeXY().X / SrcSize.X;
 
 	FIntRect SrcRect = View.ViewRect / ScaleFactor;
 	FIntRect DestRect = SrcRect;
@@ -302,12 +304,12 @@ void FRCPassPostProcessSubsurfaceVisualize::Process(FRenderingCompositePassConte
 	Context.RHICmdList.CopyToResolveTarget(DestRenderTarget.TargetableTexture, DestRenderTarget.ShaderResourceTexture, false, FResolveParams());
 
 	// we no longer need the GBuffer
-	GSceneRenderTargets.AdjustGBufferRefCount(-1);
+	SceneContext.AdjustGBufferRefCount(-1);
 }
 
 FPooledRenderTargetDesc FRCPassPostProcessSubsurfaceVisualize::ComputeOutputDesc(EPassOutputId InPassOutputId) const
 {
-	FPooledRenderTargetDesc Ret = GSceneRenderTargets.GetSceneColor()->GetDesc();
+	FPooledRenderTargetDesc Ret = FSceneRenderTargets::Get_Todo_PassContext().GetSceneColor()->GetDesc();
 
 	Ret.Reset();
 	Ret.DebugName = TEXT("SubsurfaceVisualize");
@@ -617,7 +619,7 @@ void FRCPassPostProcessSubsurfaceExtractSpecular::Process(FRenderingCompositePas
 
 FPooledRenderTargetDesc FRCPassPostProcessSubsurfaceExtractSpecular::ComputeOutputDesc(EPassOutputId InPassOutputId) const
 {
-	FPooledRenderTargetDesc Ret = GSceneRenderTargets.GetSceneColor()->GetDesc();
+	FPooledRenderTargetDesc Ret = FSceneRenderTargets::Get_Todo_PassContext().GetSceneColor()->GetDesc();
 
 	Ret.Reset();
 	Ret.DebugName = TEXT("SubsurfaceExtractSpecular");
@@ -700,7 +702,7 @@ void FRCPassPostProcessSubsurfaceSetup::Process(FRenderingCompositePassContext&
 
 FPooledRenderTargetDesc FRCPassPostProcessSubsurfaceSetup::ComputeOutputDesc(EPassOutputId InPassOutputId) const
 {
-	FPooledRenderTargetDesc Ret = GSceneRenderTargets.GetSceneColor()->GetDesc();
+	FPooledRenderTargetDesc Ret = FSceneRenderTargets::Get_Todo_PassContext().GetSceneColor()->GetDesc();
 
 	Ret.Reset();
 	Ret.DebugName = TEXT("SubsurfaceSetup");
@@ -1013,6 +1015,7 @@ void SetSubsurfaceRecombineShader(const FRenderingCompositePassContext& Context,
 
 void FRCPassPostProcessSubsurfaceRecombine::Process(FRenderingCompositePassContext& Context)
 {
+	FSceneRenderTargets& SceneContext = FSceneRenderTargets::Get(Context.RHICmdList);
 	const FPooledRenderTargetDesc* InputDesc = GetInputDesc(ePId_Input0);
 
 	check(InputDesc);
@@ -1021,7 +1024,7 @@ void FRCPassPostProcessSubsurfaceRecombine::Process(FRenderingCompositePassConte
 	const FSceneViewFamily& ViewFamily = *(View.Family);
 
 	FIntPoint SrcSize = InputDesc->Extent;
-	FIntPoint DestSize = GSceneRenderTargets.GetBufferSizeXY();
+	FIntPoint DestSize = SceneContext.GetBufferSizeXY();
 
 	check(DestSize.X);
 	check(DestSize.Y);
@@ -1031,7 +1034,7 @@ void FRCPassPostProcessSubsurfaceRecombine::Process(FRenderingCompositePassConte
 	FIntRect SrcRect = FIntRect(0, 0, InputDesc->Extent.X, InputDesc->Extent.Y);
 	FIntRect DestRect = View.ViewRect;
 
-	TRefCountPtr<IPooledRenderTarget>& SceneColor = GSceneRenderTargets.GetSceneColor();
+	TRefCountPtr<IPooledRenderTarget>& SceneColor = SceneContext.GetSceneColor();
 
 	const FSceneRenderTargetItem& DestRenderTarget = SceneColor->GetRenderTargetItem();
 
diff --git a/Engine/Source/Runtime/Renderer/Private/PostProcess/PostProcessTemporalAA.cpp b/Engine/Source/Runtime/Renderer/Private/PostProcess/PostProcessTemporalAA.cpp
--- a/Engine/Source/Runtime/Renderer/Private/PostProcess/PostProcessTemporalAA.cpp
+++ b/Engine/Source/Runtime/Renderer/Private/PostProcess/PostProcessTemporalAA.cpp
@@ -222,7 +222,7 @@ void FRCPassPostProcessSSRTemporalAA::Process(FRenderingCompositePassContext& Co
 	FIntPoint SrcSize = InputDesc->Extent;
 
 	// e.g. 4 means the input texture is 4x smaller than the buffer size
-	uint32 ScaleFactor = GSceneRenderTargets.GetBufferSizeXY().X / SrcSize.X;
+	uint32 ScaleFactor = FSceneRenderTargets::Get(Context.RHICmdList).GetBufferSizeXY().X / SrcSize.X;
 
 	FIntRect SrcRect = View.ViewRect / ScaleFactor;
 
@@ -298,7 +298,7 @@ void FRCPassPostProcessDOFTemporalAA::Process(FRenderingCompositePassContext& Co
 	FIntPoint DestSize = PassOutputs[0].RenderTargetDesc.Extent;
 
 	// e.g. 4 means the input texture is 4x smaller than the buffer size
-	uint32 ScaleFactor = GSceneRenderTargets.GetBufferSizeXY().X / SrcSize.X;
+	uint32 ScaleFactor = FSceneRenderTargets::Get(Context.RHICmdList).GetBufferSizeXY().X / SrcSize.X;
 
 	FIntRect SrcRect = View.ViewRect / ScaleFactor;
 	FIntRect DestRect = SrcRect;
@@ -379,7 +379,7 @@ void FRCPassPostProcessDOFTemporalAANear::Process(FRenderingCompositePassContext
 	FIntPoint DestSize = PassOutputs[0].RenderTargetDesc.Extent;
 
 	// e.g. 4 means the input texture is 4x smaller than the buffer size
-	uint32 ScaleFactor = GSceneRenderTargets.GetBufferSizeXY().X / SrcSize.X;
+	uint32 ScaleFactor = FSceneRenderTargets::Get(Context.RHICmdList).GetBufferSizeXY().X / SrcSize.X;
 
 	FIntRect SrcRect = View.ViewRect / ScaleFactor;
 	FIntRect DestRect = SrcRect;
@@ -459,7 +459,7 @@ void FRCPassPostProcessLightShaftTemporalAA::Process(FRenderingCompositePassCont
 	FIntPoint DestSize = PassOutputs[0].RenderTargetDesc.Extent;
 
 	// e.g. 4 means the input texture is 4x smaller than the buffer size
-	uint32 ScaleFactor = GSceneRenderTargets.GetBufferSizeXY().X / SrcSize.X;
+	uint32 ScaleFactor = FSceneRenderTargets::Get(Context.RHICmdList).GetBufferSizeXY().X / SrcSize.X;
 
 	FIntRect SrcRect = View.ViewRect / ScaleFactor;
 	FIntRect DestRect = SrcRect;
@@ -525,6 +525,7 @@ void FRCPassPostProcessTemporalAA::Process(FRenderingCompositePassContext& Conte
 		// input is not hooked up correctly
 		return;
 	}
+	FSceneRenderTargets& SceneContext = FSceneRenderTargets::Get(Context.RHICmdList);
 
 	FViewInfo& View = Context.View;
 	FSceneViewState* ViewState = Context.ViewState;
@@ -537,7 +538,7 @@ void FRCPassPostProcessTemporalAA::Process(FRenderingCompositePassContext& Conte
 	FIntPoint DestSize = PassOutputs[0].RenderTargetDesc.Extent;
 
 	// e.g. 4 means the input texture is 4x smaller than the buffer size
-	uint32 ScaleFactor = GSceneRenderTargets.GetBufferSizeXY().X / SrcSize.X;
+	uint32 ScaleFactor = SceneContext.GetBufferSizeXY().X / SrcSize.X;
 
 	FIntRect SrcRect = View.ViewRect / ScaleFactor;
 	FIntRect DestRect = SrcRect;
@@ -545,7 +546,7 @@ void FRCPassPostProcessTemporalAA::Process(FRenderingCompositePassContext& Conte
 	const FSceneRenderTargetItem& DestRenderTarget = PassOutputs[0].RequestSurface(Context);
 
 	//Context.SetRenderTarget(RHICmdList, DestRenderTarget.TargetableTexture, FTextureRHIRef());
-	SetRenderTarget(Context.RHICmdList, DestRenderTarget.TargetableTexture, GSceneRenderTargets.GetSceneDepthTexture(), ESimpleRenderTargetMode::EUninitializedColorExistingDepth, FExclusiveDepthStencil::DepthRead_StencilWrite);
+	SetRenderTarget(Context.RHICmdList, DestRenderTarget.TargetableTexture, SceneContext.GetSceneDepthTexture(), ESimpleRenderTargetMode::EUninitializedColorExistingDepth, FExclusiveDepthStencil::DepthRead_StencilWrite);
 
 	// is optimized away if possible (RT size=view size, )
 	Context.RHICmdList.Clear(true, FLinearColor::Black, false, (float)ERHIZBuffer::FarPlane, false, 0, SrcRect);
diff --git a/Engine/Source/Runtime/Renderer/Private/PostProcess/PostProcessTestImage.cpp b/Engine/Source/Runtime/Renderer/Private/PostProcess/PostProcessTestImage.cpp
--- a/Engine/Source/Runtime/Renderer/Private/PostProcess/PostProcessTestImage.cpp
+++ b/Engine/Source/Runtime/Renderer/Private/PostProcess/PostProcessTestImage.cpp
@@ -126,7 +126,7 @@ void FRCPassPostProcessTestImage::Process(FRenderingCompositePassContext& Contex
 		SrcRect.Min.X, SrcRect.Min.Y,
 		SrcRect.Width(), SrcRect.Height(),
 		DestRect.Size(),
-		GSceneRenderTargets.GetBufferSizeXY(),
+		FSceneRenderTargets::Get(Context.RHICmdList).GetBufferSizeXY(),
 		*VertexShader,
 		EDRF_UseTriangleOptimization);
 
@@ -203,7 +203,7 @@ void FRCPassPostProcessTestImage::Process(FRenderingCompositePassContext& Contex
 
 FPooledRenderTargetDesc FRCPassPostProcessTestImage::ComputeOutputDesc(EPassOutputId InPassOutputId) const
 {
-	FPooledRenderTargetDesc Ret(FPooledRenderTargetDesc::Create2DDesc(GSceneRenderTargets.GetBufferSizeXY(), PF_B8G8R8A8, TexCreate_None, TexCreate_RenderTargetable, false));
+	FPooledRenderTargetDesc Ret(FPooledRenderTargetDesc::Create2DDesc(FSceneRenderTargets::Get_FrameConstantsOnly().GetBufferSizeXY(), PF_B8G8R8A8, TexCreate_None, TexCreate_RenderTargetable, false));
 
 	Ret.DebugName = TEXT("TestImage");
 
diff --git a/Engine/Source/Runtime/Renderer/Private/PostProcess/PostProcessTonemap.cpp b/Engine/Source/Runtime/Renderer/Private/PostProcess/PostProcessTonemap.cpp
--- a/Engine/Source/Runtime/Renderer/Private/PostProcess/PostProcessTonemap.cpp
+++ b/Engine/Source/Runtime/Renderer/Private/PostProcess/PostProcessTonemap.cpp
@@ -1095,6 +1095,7 @@ void FRCPassPostProcessTonemap::Process(FRenderingCompositePassContext& Context)
 		default:
 			check(0);
 	}
+	FSceneRenderTargets& SceneContext = FSceneRenderTargets::Get(Context.RHICmdList);
 
 	// Draw a quad mapping scene color to the view's render target
 	TShaderMapRef<FPostProcessTonemapVS> VertexShader(Context.GetShaderMap());
@@ -1106,7 +1107,7 @@ void FRCPassPostProcessTonemap::Process(FRenderingCompositePassContext& Context)
 		View.ViewRect.Min.X, View.ViewRect.Min.Y, 
 		View.ViewRect.Width(), View.ViewRect.Height(),
 		View.ViewRect.Size(),
-		GSceneRenderTargets.GetBufferSizeXY(),
+		SceneContext.GetBufferSizeXY(),
 		*VertexShader,
 		EDRF_UseTriangleOptimization);
 
@@ -1117,7 +1118,7 @@ void FRCPassPostProcessTonemap::Process(FRenderingCompositePassContext& Context)
 	{
 		// The RT should be released as early as possible to allow sharing of that memory for other purposes.
 		// This becomes even more important with some limited VRam (XBoxOne).
-		GSceneRenderTargets.SetSceneColor(0);
+		SceneContext.SetSceneColor(0);
 	}
 }
 
diff --git a/Engine/Source/Runtime/Renderer/Private/PostProcess/PostProcessUpscale.cpp b/Engine/Source/Runtime/Renderer/Private/PostProcess/PostProcessUpscale.cpp
--- a/Engine/Source/Runtime/Renderer/Private/PostProcess/PostProcessUpscale.cpp
+++ b/Engine/Source/Runtime/Renderer/Private/PostProcess/PostProcessUpscale.cpp
@@ -255,7 +255,7 @@ void FRCPassPostProcessUpscale::Process(FRenderingCompositePassContext& Context)
 	bool bTessellatedQuad = PaniniConfig.D >= 0.01f;
 
 	// with distortion (bTessellatedQuad) we need to clear the background
-	FIntRect ExcludeRect = bTessellatedQuad ? FIntRect() : View.ViewRect;
+	FIntRect ExcludeRect = bTessellatedQuad ? FIntRect() : DestRect;
 
 	Context.RHICmdList.Clear(true, FLinearColor::Black, false, 1.0f, false, 0, ExcludeRect);
 
diff --git a/Engine/Source/Runtime/Renderer/Private/PostProcess/PostProcessWeightedSampleSum.cpp b/Engine/Source/Runtime/Renderer/Private/PostProcess/PostProcessWeightedSampleSum.cpp
--- a/Engine/Source/Runtime/Renderer/Private/PostProcess/PostProcessWeightedSampleSum.cpp
+++ b/Engine/Source/Runtime/Renderer/Private/PostProcess/PostProcessWeightedSampleSum.cpp
@@ -354,10 +354,11 @@ void FRCPassPostProcessWeightedSampleSum::Process(FRenderingCompositePassContext
 
 	FIntPoint SrcSize = InputDesc->Extent;
 	FIntPoint DestSize = PassOutputs[0].RenderTargetDesc.Extent;
+	FSceneRenderTargets& SceneContext = FSceneRenderTargets::Get(Context.RHICmdList);
 
 	// e.g. 4 means the input texture is 4x smaller than the buffer size
-	FIntPoint SrcScaleFactor = GSceneRenderTargets.GetBufferSizeXY() / SrcSize;
-	FIntPoint DstScaleFactor = GSceneRenderTargets.GetBufferSizeXY() / DestSize;
+	FIntPoint SrcScaleFactor = SceneContext.GetBufferSizeXY() / SrcSize;
+	FIntPoint DstScaleFactor = SceneContext.GetBufferSizeXY() / DestSize;
 
 	TRefCountPtr<IPooledRenderTarget> InputPooledElement = Input->RequestInput();
 
@@ -543,7 +544,7 @@ bool FRCPassPostProcessWeightedSampleSum::DoFastBlur() const
 		else
 		{
 			FIntPoint SrcSize = InputDesc->Extent;
-			FIntPoint BufferSize = GSceneRenderTargets.GetBufferSizeXY();
+			FIntPoint BufferSize = FSceneRenderTargets::Get_FrameConstantsOnly().GetBufferSizeXY();
 
 			float InputRatio = SrcSize.X / (float)SrcSize.Y;
 			float BufferRatio = BufferSize.X / (float)BufferSize.Y;
diff --git a/Engine/Source/Runtime/Renderer/Private/PostProcess/PostProcessing.cpp b/Engine/Source/Runtime/Renderer/Private/PostProcess/PostProcessing.cpp
--- a/Engine/Source/Runtime/Renderer/Private/PostProcess/PostProcessing.cpp
+++ b/Engine/Source/Runtime/Renderer/Private/PostProcess/PostProcessing.cpp
@@ -136,12 +136,13 @@ FPostprocessContext::FPostprocessContext(class FRenderingCompositionGraph& InGra
 	, SceneColor(0)
 	, SceneDepth(0)
 {
-	if(GSceneRenderTargets.IsSceneColorAllocated())
+	FSceneRenderTargets& SceneContext = FSceneRenderTargets::Get_Todo_PassContext();
+	if(SceneContext.IsSceneColorAllocated())
 	{
-		SceneColor = Graph.RegisterPass(new(FMemStack::Get()) FRCPassPostProcessInput(GSceneRenderTargets.GetSceneColor()));
+		SceneColor = Graph.RegisterPass(new(FMemStack::Get()) FRCPassPostProcessInput(SceneContext.GetSceneColor()));
 	}
 
-	SceneDepth = Graph.RegisterPass(new(FMemStack::Get()) FRCPassPostProcessInput(GSceneRenderTargets.SceneDepthZ));
+	SceneDepth = Graph.RegisterPass(new(FMemStack::Get()) FRCPassPostProcessInput(SceneContext.SceneDepthZ));
 
 	FinalOutput = FRenderingCompositeOutputRef(SceneColor);
 }
@@ -636,7 +637,7 @@ static void AddTemporalAA( FPostprocessContext& Context, FRenderingCompositeOutp
 	else
 	{
 		// No history so use current as history
-		HistoryInput = Context.Graph.RegisterPass( new(FMemStack::Get()) FRCPassPostProcessInput( GSceneRenderTargets.GetSceneColor() ) );
+		HistoryInput = Context.Graph.RegisterPass( new(FMemStack::Get()) FRCPassPostProcessInput( FSceneRenderTargets::Get_Todo_PassContext().GetSceneColor() ) );
 	}
 
 	FRenderingCompositePass* TemporalAAPass = Context.Graph.RegisterPass( new(FMemStack::Get()) FRCPassPostProcessTemporalAA );
@@ -702,7 +703,7 @@ static FRenderingCompositePass* AddSinglePostProcessMaterial(FPostprocessContext
 		if(Material->NeedsGBuffer())
 		{
 			// AdjustGBufferRefCount(-1) call is done when the pass gets executed
-			GSceneRenderTargets.AdjustGBufferRefCount(1);
+			FSceneRenderTargets::Get_Todo_PassContext().AdjustGBufferRefCount(1);
 		}
 
 		FRenderingCompositePass* Node = Context.Graph.RegisterPass(new(FMemStack::Get()) FRCPassPostProcessMaterial(MaterialInterface, Context.View.GetFeatureLevel()));
@@ -770,7 +771,7 @@ static void AddPostProcessMaterial(FPostprocessContext& Context, EBlendableLocat
 		if(Material->NeedsGBuffer())
 		{
 			// AdjustGBufferRefCount(-1) call is done when the pass gets executed
-			GSceneRenderTargets.AdjustGBufferRefCount(1);
+			FSceneRenderTargets::Get_Todo_PassContext().AdjustGBufferRefCount(1);
 		}
 
 		FRenderingCompositePass* Node = Context.Graph.RegisterPass(new(FMemStack::Get()) FRCPassPostProcessMaterial(MaterialInterface,FeatureLevel));
@@ -823,7 +824,7 @@ static void AddHighResScreenshotMask(FPostprocessContext& Context, FRenderingCom
 		if (RendererMaterial->NeedsGBuffer())
 		{
 			// AdjustGBufferRefCount(-1) call is done when the pass gets executed
-			GSceneRenderTargets.AdjustGBufferRefCount(1);
+			FSceneRenderTargets::Get_Todo_PassContext().AdjustGBufferRefCount(1);
 		}
 	}
 }
@@ -872,7 +873,7 @@ static void AddGBufferVisualizationOverview(FPostprocessContext& Context, FRende
 					if (Material->NeedsGBuffer())
 					{
 						// AdjustGBufferRefCount(-1) call is done when the pass gets executed
-						GSceneRenderTargets.AdjustGBufferRefCount(1);
+						FSceneRenderTargets::Get_Todo_PassContext().AdjustGBufferRefCount(1);
 					}
 
 					if (BaseFilename.Len())
@@ -1443,7 +1444,7 @@ void FPostProcessing::Process(FRHICommandListImmediate& RHICmdList, FViewInfo& V
 				VisualizeNode->SetInput(ePId_Input0, FRenderingCompositeOutputRef(Context.FinalOutput));
 
 				// PassThrough is needed to upscale the half res texture
-				FPooledRenderTargetDesc Desc = GSceneRenderTargets.GetSceneColor()->GetDesc();
+				FPooledRenderTargetDesc Desc = FSceneRenderTargets::Get(RHICmdList).GetSceneColor()->GetDesc();
 				Desc.Format = PF_B8G8R8A8;
 				FRenderingCompositePass* NullPass = Context.Graph.RegisterPass(new(FMemStack::Get()) FRCPassPostProcessPassThrough(Desc));
 				NullPass->SetInput(ePId_Input0, FRenderingCompositeOutputRef(VisualizeNode));
@@ -1603,7 +1604,7 @@ void FPostProcessing::Process(FRHICommandListImmediate& RHICmdList, FViewInfo& V
 		{
 			// Generally we no longer need the GBuffers, anyone that wants to keep the GBuffers for longer should have called AdjustGBufferRefCount(1) to keep it for longer
 			// and call AdjustGBufferRefCount(-1) once it's consumed. This needs to happen each frame. PostProcessMaterial do that automatically
-			GSceneRenderTargets.AdjustGBufferRefCount(-1);
+			FSceneRenderTargets::Get_Todo_PassContext().AdjustGBufferRefCount(-1);
 		}
 
 		// The graph setup should be finished before this line ----------------------------------------
@@ -1687,7 +1688,7 @@ void FPostProcessing::ProcessES2(FRHICommandListImmediate& RHICmdList, FViewInfo
 		FIntRect FinalOutputViewRect = View.ViewRect;
 		FIntPoint PrePostSourceViewportSize = View.ViewRect.Size();
 		// ES2 preview uses a subsection of the scene RT, bUsedFramebufferFetch == true deals with this case.  
-		bool bViewRectSource = bUsedFramebufferFetch || GSceneRenderTargets.GetBufferSizeXY() != PrePostSourceViewportSize;
+		bool bViewRectSource = bUsedFramebufferFetch || FSceneRenderTargets::Get(RHICmdList).GetBufferSizeXY() != PrePostSourceViewportSize;
 
 		// add the passes we want to add to the graph (commenting a line means the pass is not inserted into the graph) ---------
 		if( View.Family->EngineShowFlags.PostProcessing )
diff --git a/Engine/Source/Runtime/Renderer/Private/PostProcess/RenderTargetPool.cpp b/Engine/Source/Runtime/Renderer/Private/PostProcess/RenderTargetPool.cpp
--- a/Engine/Source/Runtime/Renderer/Private/PostProcess/RenderTargetPool.cpp
+++ b/Engine/Source/Runtime/Renderer/Private/PostProcess/RenderTargetPool.cpp
@@ -174,6 +174,8 @@ bool FRenderTargetPool::FindFreeElement(const FPooledRenderTargetDesc& Desc, TRe
 	{
 		FPooledRenderTarget* Current = (FPooledRenderTarget*)Out.GetReference();
 
+		check(!Current->IsSnapshot());
+
 		const bool bExactMatch = true;
 
 		if(Out->GetDesc().Compare(Desc, bExactMatch))
@@ -220,6 +222,7 @@ bool FRenderTargetPool::FindFreeElement(const FPooledRenderTargetDesc& Desc, TRe
 
 				if(Element && Element->IsFree() && Element->GetDesc().Compare(Desc, bExactMatch))
 				{
+					check(!Element->IsSnapshot());
 					Found = Element;
 					FoundIndex = i;
 					break;
@@ -409,6 +412,7 @@ bool FRenderTargetPool::FindFreeElement(const FPooledRenderTargetDesc& Desc, TRe
 #endif
 
 	check(Found->IsFree());
+	check(!Found->IsSnapshot());
 
 	Found->Desc.DebugName = InDebugName;
 	Found->UnusedForNFrames = 0;
@@ -433,11 +437,25 @@ void FRenderTargetPool::CreateUntrackedElement(const FPooledRenderTargetDesc& De
 	FPooledRenderTarget* Found = new FPooledRenderTarget(Desc);
 
 	Found->RenderTargetItem = Item;
+	check(!Found->IsSnapshot());
 
 	// assign to the reference counted variable
 	Out = Found;
 }
 
+IPooledRenderTarget* FRenderTargetPool::MakeSnapshot(const TRefCountPtr<IPooledRenderTarget>& In)
+{
+	check(IsInRenderingThread());
+	FPooledRenderTarget* NewSnapshot = nullptr;
+	if (In.GetReference())
+	{
+		NewSnapshot = new (FMemStack::Get()) FPooledRenderTarget(*static_cast<FPooledRenderTarget*>(In.GetReference()));
+		PooledRenderTargetSnapshots.Add(NewSnapshot);
+	}
+	return NewSnapshot;
+}
+
+
 void FRenderTargetPool::GetStats(uint32& OutWholeCount, uint32& OutWholePoolInKB, uint32& OutUsedInKB) const
 {
 	OutWholeCount = (uint32)PooledRenderTargets.Num();
@@ -450,6 +468,7 @@ void FRenderTargetPool::GetStats(uint32& OutWholeCount, uint32& OutWholePoolInKB
 
 		if(Element)
 		{
+			check(!Element->IsSnapshot());
 			uint32 SizeInKB = ComputeSizeInKB(*Element);
 
 			OutWholePoolInKB += SizeInKB;
@@ -769,7 +788,7 @@ void FRenderTargetPool::PresentContent(FRHICommandListImmediate& RHICmdList, con
 			SMemoryStats MemoryStats = ComputeView();
 
 			SetRenderTarget(RHICmdList, View.Family->RenderTarget->GetRenderTargetTexture(), FTextureRHIRef());
-			RHICmdList.SetViewport(0, 0, 0.0f, GSceneRenderTargets.GetBufferSizeXY().X, GSceneRenderTargets.GetBufferSizeXY().Y, 1.0f);
+			RHICmdList.SetViewport(0, 0, 0.0f, FSceneRenderTargets::Get(RHICmdList).GetBufferSizeXY().X, FSceneRenderTargets::Get(RHICmdList).GetBufferSizeXY().Y, 1.0f);
 
 			RHICmdList.SetBlendState(TStaticBlendState<>::GetRHI());
 			RHICmdList.SetRasterizerState(TStaticRasterizerState<>::GetRHI());
@@ -1030,6 +1049,7 @@ void FRenderTargetPool::TickPoolElements()
 
 		if(Element)
 		{
+			check(!Element->IsSnapshot());
 			Element->OnFrameStart();
 		}
 	}
@@ -1125,6 +1145,7 @@ int32 FRenderTargetPool::FindIndex(IPooledRenderTarget* In) const
 
 			if(Element == In)
 			{
+				check(!Element->IsSnapshot());
 				return i;
 			}
 		}
@@ -1146,6 +1167,7 @@ void FRenderTargetPool::FreeUnusedResource(TRefCountPtr<IPooledRenderTarget>& In
 
 		if(Element)
 		{
+			check(!Element->IsSnapshot());
 			AllocationLevelInKB -= ComputeSizeInKB(*Element);
 			// we assume because of reference counting the resource gets released when not needed any more
 			// we don't use Remove() to not shuffle around the elements for better transparency on RenderTargetPoolEvents
@@ -1168,6 +1190,7 @@ void FRenderTargetPool::FreeUnusedResources()
 
 		if(Element && Element->IsFree())
 		{
+			check(!Element->IsSnapshot());
 			AllocationLevelInKB -= ComputeSizeInKB(*Element);
 			// we assume because of reference counting the resource gets released when not needed any more
 			// we don't use Remove() to not shuffle around the elements for better transparency on RenderTargetPoolEvents
@@ -1188,6 +1211,7 @@ void FRenderTargetPool::DumpMemoryUsage(FOutputDevice& OutputDevice)
 
 		if(Element)
 		{
+			check(!Element->IsSnapshot());
 			OutputDevice.Logf(
 				TEXT("  %6.3fMB %4dx%4d%s%s %2dmip(s) %s (%s)"),
 				ComputeSizeInKB(*Element) / 1024.0f,
@@ -1210,23 +1234,33 @@ void FRenderTargetPool::DumpMemoryUsage(FOutputDevice& OutputDevice)
 
 uint32 FPooledRenderTarget::AddRef() const
 {
-	check(IsInRenderingThread());
-	return uint32(++NumRefs);
+	if (!bSnapshot)
+	{
+		check(IsInRenderingThread());
+		return uint32(++NumRefs);
+	}
+	check(NumRefs == 1);
+	return 1;
 }
 
 uint32 FPooledRenderTarget::Release() const
 {
-	check(IsInRenderingThread());
-	uint32 Refs = uint32(--NumRefs);
-	if(Refs == 0)
+	if (!bSnapshot)
 	{
-		// better we remove const from Release()
-		FSceneRenderTargetItem& NonConstItem = (FSceneRenderTargetItem&)RenderTargetItem;
+		check(IsInRenderingThread());
+		uint32 Refs = uint32(--NumRefs);
+		if(Refs == 0)
+		{
+			// better we remove const from Release()
+			FSceneRenderTargetItem& NonConstItem = (FSceneRenderTargetItem&)RenderTargetItem;
 
-		NonConstItem.SafeRelease();
-		delete this;
+			NonConstItem.SafeRelease();
+			delete this;
+		}
+		return Refs;
 	}
-	return Refs;
+	check(NumRefs == 1);
+	return 1;
 }
 
 uint32 FPooledRenderTarget::GetRefCount() const
@@ -1250,8 +1284,22 @@ void FRenderTargetPool::ReleaseDynamicRHI()
 {
 	check(IsInRenderingThread());
 	PooledRenderTargets.Empty();
+	if (PooledRenderTargetSnapshots.Num())
+	{
+		DestructSnapshots();
+	}
 }
 
+void FRenderTargetPool::DestructSnapshots()
+{
+	for (auto Snapshot : PooledRenderTargetSnapshots)
+	{
+		Snapshot->~FPooledRenderTarget();
+	}
+	PooledRenderTargetSnapshots.Reset();
+}
+
+
 // for debugging purpose
 FPooledRenderTarget* FRenderTargetPool::GetElementById(uint32 Id) const
 {
@@ -1293,7 +1341,7 @@ void FRenderTargetPool::CompactPool()
 
 bool FPooledRenderTarget::OnFrameStart()
 {
-	check(IsInRenderingThread());
+	check(IsInRenderingThread() && !bSnapshot);
 
 	// If there are any references to the pooled render target other than the pool itself, then it may not be freed.
 	if(!IsFree())
@@ -1317,32 +1365,33 @@ bool FPooledRenderTarget::OnFrameStart()
 uint32 FPooledRenderTarget::ComputeMemorySize() const
 {
 	uint32 Size = 0;
-
-	if(Desc.Is2DTexture())
+	if (!bSnapshot)
 	{
-		Size += RHIComputeMemorySize((const FTexture2DRHIRef&)RenderTargetItem.TargetableTexture);
-		if(RenderTargetItem.ShaderResourceTexture != RenderTargetItem.TargetableTexture)
+		if(Desc.Is2DTexture())
 		{
-			Size += RHIComputeMemorySize((const FTexture2DRHIRef&)RenderTargetItem.ShaderResourceTexture);
+			Size += RHIComputeMemorySize((const FTexture2DRHIRef&)RenderTargetItem.TargetableTexture);
+			if(RenderTargetItem.ShaderResourceTexture != RenderTargetItem.TargetableTexture)
+			{
+				Size += RHIComputeMemorySize((const FTexture2DRHIRef&)RenderTargetItem.ShaderResourceTexture);
+			}
 		}
-	}
-	else if(Desc.Is3DTexture())
-	{
-		Size += RHIComputeMemorySize((const FTexture3DRHIRef&)RenderTargetItem.TargetableTexture);
-		if(RenderTargetItem.ShaderResourceTexture != RenderTargetItem.TargetableTexture)
+		else if(Desc.Is3DTexture())
 		{
-			Size += RHIComputeMemorySize((const FTexture3DRHIRef&)RenderTargetItem.ShaderResourceTexture);
+			Size += RHIComputeMemorySize((const FTexture3DRHIRef&)RenderTargetItem.TargetableTexture);
+			if(RenderTargetItem.ShaderResourceTexture != RenderTargetItem.TargetableTexture)
+			{
+				Size += RHIComputeMemorySize((const FTexture3DRHIRef&)RenderTargetItem.ShaderResourceTexture);
+			}
 		}
-	}
-	else
-	{
-		Size += RHIComputeMemorySize((const FTextureCubeRHIRef&)RenderTargetItem.TargetableTexture);
-		if(RenderTargetItem.ShaderResourceTexture != RenderTargetItem.TargetableTexture)
+		else
 		{
-			Size += RHIComputeMemorySize((const FTextureCubeRHIRef&)RenderTargetItem.ShaderResourceTexture);
+			Size += RHIComputeMemorySize((const FTextureCubeRHIRef&)RenderTargetItem.TargetableTexture);
+			if(RenderTargetItem.ShaderResourceTexture != RenderTargetItem.TargetableTexture)
+			{
+				Size += RHIComputeMemorySize((const FTextureCubeRHIRef&)RenderTargetItem.ShaderResourceTexture);
+			}
 		}
 	}
-
 	return Size;
 }
 
@@ -1350,5 +1399,5 @@ bool FPooledRenderTarget::IsFree() const
 {
 	check(GetRefCount() >= 1);
 	// If the only reference to the pooled render target is from the pool, then it's unused.
-	return GetRefCount() == 1;
+	return !bSnapshot && GetRefCount() == 1;
 }
diff --git a/Engine/Source/Runtime/Renderer/Private/PostProcess/RenderTargetPool.h b/Engine/Source/Runtime/Renderer/Private/PostProcess/RenderTargetPool.h
--- a/Engine/Source/Runtime/Renderer/Private/PostProcess/RenderTargetPool.h
+++ b/Engine/Source/Runtime/Renderer/Private/PostProcess/RenderTargetPool.h
@@ -10,21 +10,40 @@
 /** The reference to a pooled render target, use like this: TRefCountPtr<IPooledRenderTarget> */
 struct FPooledRenderTarget : public IPooledRenderTarget
 {
-	/** constructor */
 	FPooledRenderTarget(const FPooledRenderTargetDesc& InDesc) 
 		: NumRefs(0)
+		, UnusedForNFrames(0)
 		, Desc(InDesc)
+		, bSnapshot(false)
+	{
+	}
+	/* Constructor that makes a snapshot */
+	FPooledRenderTarget(const FPooledRenderTarget& SnaphotSource)
+		: NumRefs(1)
 		, UnusedForNFrames(0)
+		, Desc(SnaphotSource.Desc)
+		, bSnapshot(true)
 	{
+		check(IsInRenderingThread());
+		RenderTargetItem = SnaphotSource.RenderTargetItem;
 	}
 
 	virtual ~FPooledRenderTarget()
 	{
-		check(!NumRefs);
+		check(!NumRefs || (bSnapshot && NumRefs == 1));
 		RenderTargetItem.SafeRelease();
 	}
 
-	uint32 GetUnusedForNFrames() const { return UnusedForNFrames; }
+	bool IsSnapshot() const 
+	{ 
+		return bSnapshot;
+	}
+
+	uint32 GetUnusedForNFrames() const 
+	{ 
+		check(!bSnapshot);
+		return UnusedForNFrames; 
+	}
 
 	// interface IPooledRenderTarget --------------
 
@@ -43,10 +62,12 @@ private:
 
 	/** For pool management (only if NumRef == 0 the element can be reused) */
 	mutable int32 NumRefs;
-	/** All necessary data to create the render target */
-	FPooledRenderTargetDesc Desc;
 	/** Allows to defer the release to save performance on some hardware (DirectX) */
 	uint32 UnusedForNFrames;
+	/** All necessary data to create the render target */
+	FPooledRenderTargetDesc Desc;
+	/** Snapshots are sortof fake pooled render targets, they don't own anything and can outlive the things that created them. These are for threaded rendering. */
+	bool bSnapshot;
 
 	/** @return true:release this one, false otherwise */
 	bool OnFrameStart();
@@ -177,6 +198,14 @@ public:
 
 	void CreateUntrackedElement(const FPooledRenderTargetDesc& Desc, TRefCountPtr<IPooledRenderTarget>& Out, const FSceneRenderTargetItem& Item);
 
+	/** Destruct all snapshots, this must be done after all outstanding async tasks are done. It is important because they hold ref counted texture pointers etc **/
+	IPooledRenderTarget* MakeSnapshot(const TRefCountPtr<IPooledRenderTarget>& In);
+
+	/** Destruct all snapshots, this must be done after all outstanding async tasks are done. It is important because they hold ref counted texture pointers etc **/
+	void DestructSnapshots();
+
+
+
 	/** Only to get statistics on usage and free elements. Normally only called in renderthread or if FlushRenderingCommands was called() */
 	void GetStats(uint32& OutWholeCount, uint32& OutWholePoolInKB, uint32& OutUsedInKB) const;
 	/**
@@ -232,6 +261,9 @@ private:
 	/** Elements can be 0, we compact the buffer later. */
 	TArray< TRefCountPtr<FPooledRenderTarget> > PooledRenderTargets;
 
+	/** These are snapshots, have odd life times, live in the scene allocator, and don't contibute to any accounting or other management. */
+	TArray<FPooledRenderTarget*> PooledRenderTargetSnapshots;
+
 	// redundant, can always be computed with GetStats(), to debug "out of memory" situations and used for r.RenderTargetPoolMin
 	uint32 AllocationLevelInKB;
 
diff --git a/Engine/Source/Runtime/Renderer/Private/PostProcess/RenderingCompositionGraph.cpp b/Engine/Source/Runtime/Renderer/Private/PostProcess/RenderingCompositionGraph.cpp
--- a/Engine/Source/Runtime/Renderer/Private/PostProcess/RenderingCompositionGraph.cpp
+++ b/Engine/Source/Runtime/Renderer/Private/PostProcess/RenderingCompositionGraph.cpp
@@ -734,7 +734,7 @@ void FPostProcessPassParameters::Set(
 
 	//Calculate a base scene texture min max which will be pulled in by a pixel for each PP input.
 	FIntRect ContextViewportRect = Context.IsViewportValid() ? Context.GetViewport() : FIntRect(0,0,0,0);
-	const FIntPoint SceneRTSize = GSceneRenderTargets.GetBufferSizeXY();
+	const FIntPoint SceneRTSize = FSceneRenderTargets::Get(RHICmdList).GetBufferSizeXY();
 	FVector4 BaseSceneTexMinMax(	((float)ContextViewportRect.Min.X/SceneRTSize.X), 
 									((float)ContextViewportRect.Min.Y/SceneRTSize.Y), 
 									((float)ContextViewportRect.Max.X/SceneRTSize.X), 
diff --git a/Engine/Source/Runtime/Renderer/Private/PostProcess/RenderingCompositionGraph.h b/Engine/Source/Runtime/Renderer/Private/PostProcess/RenderingCompositionGraph.h
--- a/Engine/Source/Runtime/Renderer/Private/PostProcess/RenderingCompositionGraph.h
+++ b/Engine/Source/Runtime/Renderer/Private/PostProcess/RenderingCompositionGraph.h
@@ -27,7 +27,7 @@ public:
 
 	/**
 	 * Returns the input pointer as output to allow this:
-	 * Example:  SceneColor = Graph.RegisterPass(new FRCPassPostProcessInput(GSceneRenderTargets.SceneColor));
+	 * Example:  SceneColor = Graph.RegisterPass(new FRCPassPostProcessInput(FSceneRenderTargets::Get(RHICmdList).SceneColor));
 	 * @param InPass - must not be 0
 	 */
 	template<class T>
diff --git a/Engine/Source/Runtime/Renderer/Private/PostProcess/SceneFilterRendering.h b/Engine/Source/Runtime/Renderer/Private/PostProcess/SceneFilterRendering.h
--- a/Engine/Source/Runtime/Renderer/Private/PostProcess/SceneFilterRendering.h
+++ b/Engine/Source/Runtime/Renderer/Private/PostProcess/SceneFilterRendering.h
@@ -122,7 +122,7 @@ public:
 			OverlayColor
 			);
 
-		FIntPoint BufferSize = GSceneRenderTargets.GetBufferSizeXY();
+		FIntPoint BufferSize = FSceneRenderTargets::Get(RHICmdList).GetBufferSizeXY();
 		float BufferSizeX = (float)BufferSize.X;
 		float BufferSizeY = (float)BufferSize.Y;
 		float InvBufferSizeX = 1.0f / BufferSizeX;
