UE4 - Large upgrade to parallel rendering. No longer block per pass. Several RHI cmdlist passthrough functions added and implemented on the consoles

diff --git a/Engine/Source/Runtime/Renderer/Private/PostProcess/SceneRenderTargets.cpp b/Engine/Source/Runtime/Renderer/Private/PostProcess/SceneRenderTargets.cpp
--- a/Engine/Source/Runtime/Renderer/Private/PostProcess/SceneRenderTargets.cpp
+++ b/Engine/Source/Runtime/Renderer/Private/PostProcess/SceneRenderTargets.cpp
@@ -96,7 +96,157 @@ static TAutoConsoleVariable<int32> CVarGBufferFormat(
 	ECVF_RenderThreadSafe);
 
 /** The global render targets used for scene rendering. */
-TGlobalResource<FSceneRenderTargets> GSceneRenderTargets;
+static TGlobalResource<FSceneRenderTargets> SceneRenderTargetsSingleton;
+
+FSceneRenderTargets& FSceneRenderTargets::Get(FRHICommandList& RHICmdList)
+{
+	FSceneRenderTargets* SceneContext = (FSceneRenderTargets*)RHICmdList.GetRenderThreadContext(FRHICommandListBase::ERenderThreadContext::SceneRenderTargets);
+	if (!SceneContext)
+	{
+		return SceneRenderTargetsSingleton;
+	}
+	check(!RHICmdList.IsImmediate());
+	return *SceneContext;
+}
+
+FSceneRenderTargets& FSceneRenderTargets::Get(FRHICommandListImmediate& RHICmdList)
+{
+	check(IsInRenderingThread() && !RHICmdList.GetRenderThreadContext(FRHICommandListBase::ERenderThreadContext::SceneRenderTargets)
+		&& !FTaskGraphInterface::Get().IsThreadProcessingTasks(ENamedThreads::RenderThread_Local)); // if we are processing tasks on the local queue, it is assumed this are in support of async tasks, which cannot use the current state of the render targets. This can be relaxed if needed.
+	return SceneRenderTargetsSingleton;
+}
+
+FSceneRenderTargets& FSceneRenderTargets::Get_Todo_PassContext()
+{
+	check(IsInRenderingThread()
+		&& !FTaskGraphInterface::Get().IsThreadProcessingTasks(ENamedThreads::RenderThread_Local)); // if we are processing tasks on the local queue, it is assumed this are in support of async tasks, which cannot use the current state of the render targets. This can be relaxed if needed.
+	return SceneRenderTargetsSingleton;
+}
+
+FSceneRenderTargets& FSceneRenderTargets::Get_FrameConstantsOnly()
+{
+	return SceneRenderTargetsSingleton;
+}
+
+ FSceneRenderTargets* FSceneRenderTargets::CreateSnapshot(const FViewInfo& InView)
+{
+	QUICK_SCOPE_CYCLE_COUNTER(STAT_FSceneRenderTargets_CreateSnapshot);
+	check(IsInRenderingThread() && FMemStack::Get().GetNumMarks() == 1); // we do not want this popped before the end of the scene and it better be the scene allocator
+	FSceneRenderTargets* NewSnapshot = new (FMemStack::Get()) FSceneRenderTargets(InView, *this);
+	check(NewSnapshot->bSnapshot);
+	Snapshots.Add(NewSnapshot);
+	return NewSnapshot;
+}
+
+void FSceneRenderTargets::SetSnapshotOnCmdList(FRHICommandList& TargetCmdList)
+{
+	check(bSnapshot);
+	TargetCmdList.SetRenderThreadContext(this, FRHICommandListBase::ERenderThreadContext::SceneRenderTargets);
+}
+
+void FSceneRenderTargets::DestroyAllSnapshots()
+{
+	if (Snapshots.Num())
+	{
+		QUICK_SCOPE_CYCLE_COUNTER(STAT_FSceneRenderTargets_DestroyAllSnapshots);
+		check(IsInRenderingThread());
+		for (auto Snapshot : Snapshots)
+		{
+			Snapshot->~FSceneRenderTargets();
+		}
+		Snapshots.Reset();
+		GRenderTargetPool.DestructSnapshots();
+	}
+}
+
+
+
+template <size_t N>
+static void SnapshotArray(TRefCountPtr<IPooledRenderTarget> (&Dest)[N], const TRefCountPtr<IPooledRenderTarget> (&Src)[N])
+{
+	for (int32 Index = 0; Index < N; Index++)
+	{
+		Dest[Index] = GRenderTargetPool.MakeSnapshot(Src[Index]);
+	}
+}
+
+FSceneRenderTargets::FSceneRenderTargets(const FViewInfo& View, const FSceneRenderTargets& SnapshotSource)
+	: LightAttenuation(GRenderTargetPool.MakeSnapshot(SnapshotSource.LightAttenuation))
+	, LightAccumulation(GRenderTargetPool.MakeSnapshot(SnapshotSource.LightAccumulation))
+	, DirectionalOcclusion(GRenderTargetPool.MakeSnapshot(SnapshotSource.DirectionalOcclusion))
+	, SceneDepthZ(GRenderTargetPool.MakeSnapshot(SnapshotSource.SceneDepthZ))
+	, SceneAlphaCopy(GRenderTargetPool.MakeSnapshot(SnapshotSource.SceneAlphaCopy))
+	, AuxiliarySceneDepthZ(GRenderTargetPool.MakeSnapshot(SnapshotSource.AuxiliarySceneDepthZ))
+	, SmallDepthZ(GRenderTargetPool.MakeSnapshot(SnapshotSource.SmallDepthZ))
+	, GBufferA(GRenderTargetPool.MakeSnapshot(SnapshotSource.GBufferA))
+	, GBufferB(GRenderTargetPool.MakeSnapshot(SnapshotSource.GBufferB))
+	, GBufferC(GRenderTargetPool.MakeSnapshot(SnapshotSource.GBufferC))
+	, GBufferD(GRenderTargetPool.MakeSnapshot(SnapshotSource.GBufferD))
+	, GBufferE(GRenderTargetPool.MakeSnapshot(SnapshotSource.GBufferE))
+	, GBufferVelocity(GRenderTargetPool.MakeSnapshot(SnapshotSource.GBufferVelocity))
+	, DBufferA(GRenderTargetPool.MakeSnapshot(SnapshotSource.DBufferA))
+	, DBufferB(GRenderTargetPool.MakeSnapshot(SnapshotSource.DBufferB))
+	, DBufferC(GRenderTargetPool.MakeSnapshot(SnapshotSource.DBufferC))
+	, ScreenSpaceAO(GRenderTargetPool.MakeSnapshot(SnapshotSource.ScreenSpaceAO))
+	, CustomDepth(GRenderTargetPool.MakeSnapshot(SnapshotSource.CustomDepth))
+	, CustomStencilSRV(SnapshotSource.CustomStencilSRV)
+	, ShadowDepthZ(GRenderTargetPool.MakeSnapshot(SnapshotSource.ShadowDepthZ))
+	, OptionalShadowDepthColor(GRenderTargetPool.MakeSnapshot(SnapshotSource.OptionalShadowDepthColor))
+	, PreShadowCacheDepthZ(GRenderTargetPool.MakeSnapshot(SnapshotSource.PreShadowCacheDepthZ))
+	, ReflectiveShadowMapNormal(GRenderTargetPool.MakeSnapshot(SnapshotSource.ReflectiveShadowMapNormal))
+	, ReflectiveShadowMapDiffuse(GRenderTargetPool.MakeSnapshot(SnapshotSource.ReflectiveShadowMapDiffuse))
+	, ReflectiveShadowMapDepth(GRenderTargetPool.MakeSnapshot(SnapshotSource.ReflectiveShadowMapDepth))
+	, SkySHIrradianceMap(GRenderTargetPool.MakeSnapshot(SnapshotSource.SkySHIrradianceMap))
+	, EditorPrimitivesColor(GRenderTargetPool.MakeSnapshot(SnapshotSource.EditorPrimitivesColor))
+	, EditorPrimitivesDepth(GRenderTargetPool.MakeSnapshot(SnapshotSource.EditorPrimitivesDepth))
+	, bScreenSpaceAOIsValid(SnapshotSource.bScreenSpaceAOIsValid)
+	, bCustomDepthIsValid(SnapshotSource.bCustomDepthIsValid)
+	, bPreshadowCacheNewlyAllocated(SnapshotSource.bPreshadowCacheNewlyAllocated)
+	, GBufferRefCount(SnapshotSource.GBufferRefCount)
+	, LargestDesiredSizeThisFrame(SnapshotSource.LargestDesiredSizeThisFrame)
+	, LargestDesiredSizeLastFrame(SnapshotSource.LargestDesiredSizeLastFrame)
+	, ThisFrameNumber(SnapshotSource.ThisFrameNumber)
+	, bVelocityPass(SnapshotSource.bVelocityPass)
+	, GBufferResourcesUniformBuffer(SnapshotSource.GBufferResourcesUniformBuffer)
+	, BufferSize(SnapshotSource.BufferSize)
+	, SmallColorDepthDownsampleFactor(SnapshotSource.SmallColorDepthDownsampleFactor)
+	, bLightAttenuationEnabled(SnapshotSource.bLightAttenuationEnabled)
+	, bUseDownsizedOcclusionQueries(SnapshotSource.bUseDownsizedOcclusionQueries)
+	, CurrentGBufferFormat(SnapshotSource.CurrentGBufferFormat)
+	, CurrentSceneColorFormat(SnapshotSource.CurrentSceneColorFormat)
+	, bAllowStaticLighting(SnapshotSource.bAllowStaticLighting)
+	, CurrentMaxShadowResolution(SnapshotSource.CurrentMaxShadowResolution)
+	, CurrentRSMResolution(SnapshotSource.CurrentRSMResolution)
+	, CurrentTranslucencyLightingVolumeDim(SnapshotSource.CurrentTranslucencyLightingVolumeDim)
+	, CurrentMobile32bpp(SnapshotSource.CurrentMobile32bpp)
+	, CurrentMobileMSAA(SnapshotSource.CurrentMobileMSAA)
+	, CurrentMinShadowResolution(SnapshotSource.CurrentMinShadowResolution)
+	, bCurrentLightPropagationVolume(SnapshotSource.bCurrentLightPropagationVolume)
+	, CurrentFeatureLevel(SnapshotSource.CurrentFeatureLevel)
+	, CurrentShadingPath(SnapshotSource.CurrentShadingPath)
+	, bAllocateVelocityGBuffer(SnapshotSource.bAllocateVelocityGBuffer)
+	, bGBuffersCleared(SnapshotSource.bGBuffersCleared)
+	, GBufferClearColor(SnapshotSource.GBufferClearColor)
+	, bSceneDepthCleared(SnapshotSource.bSceneDepthCleared)
+	, SceneDepthClearValue(SnapshotSource.SceneDepthClearValue)
+	, bSnapshot(true)
+{
+	SnapshotArray(SceneColor, SnapshotSource.SceneColor);
+	SnapshotArray(TranslucencyShadowTransmission, SnapshotSource.TranslucencyShadowTransmission);
+	SnapshotArray(CubeShadowDepthZ, SnapshotSource.CubeShadowDepthZ);
+	SnapshotArray(ReflectionColorScratchCubemap, SnapshotSource.ReflectionColorScratchCubemap);
+	SnapshotArray(DiffuseIrradianceScratchCubemap, SnapshotSource.DiffuseIrradianceScratchCubemap);
+	SnapshotArray(ReflectionBrightness, SnapshotSource.ReflectionBrightness);
+	SnapshotArray(TranslucencyLightingVolumeAmbient, SnapshotSource.TranslucencyLightingVolumeAmbient);
+	SnapshotArray(TranslucencyLightingVolumeDirectional, SnapshotSource.TranslucencyLightingVolumeDirectional);
+
+	FSceneViewState* ViewState = (FSceneViewState*)View.State;
+	if (ViewState)
+	{
+		SeparateTranslucencyRT = GRenderTargetPool.MakeSnapshot(ViewState->SeparateTranslucencyRT);
+	}
+}
+
 
 inline const TCHAR* GetReflectionBrightnessTargetName(uint32 Index)
 {
@@ -319,10 +469,10 @@ void FSceneRenderTargets::BeginRenderingSceneColor(FRHICommandList& RHICmdList,
 int32 FSceneRenderTargets::GetGBufferRenderTargets(ERenderTargetLoadAction ColorLoadAction, FRHIRenderTargetView OutRenderTargets[MaxSimultaneousRenderTargets], int32& OutVelocityRTIndex)
 {
 	OutRenderTargets[0] = FRHIRenderTargetView(GetSceneColorSurface(), 0, -1, ColorLoadAction, ERenderTargetStoreAction::EStore);
-	OutRenderTargets[1] = FRHIRenderTargetView(GSceneRenderTargets.GBufferA->GetRenderTargetItem().TargetableTexture, 0, -1, ColorLoadAction, ERenderTargetStoreAction::EStore);
-	OutRenderTargets[2] = FRHIRenderTargetView(GSceneRenderTargets.GBufferB->GetRenderTargetItem().TargetableTexture, 0, -1, ColorLoadAction, ERenderTargetStoreAction::EStore);
-	OutRenderTargets[3] = FRHIRenderTargetView(GSceneRenderTargets.GBufferC->GetRenderTargetItem().TargetableTexture, 0, -1, ColorLoadAction, ERenderTargetStoreAction::EStore);
-	OutRenderTargets[4] = FRHIRenderTargetView(GSceneRenderTargets.GBufferD->GetRenderTargetItem().TargetableTexture, 0, -1, ColorLoadAction, ERenderTargetStoreAction::EStore);
+	OutRenderTargets[1] = FRHIRenderTargetView(GBufferA->GetRenderTargetItem().TargetableTexture, 0, -1, ColorLoadAction, ERenderTargetStoreAction::EStore);
+	OutRenderTargets[2] = FRHIRenderTargetView(GBufferB->GetRenderTargetItem().TargetableTexture, 0, -1, ColorLoadAction, ERenderTargetStoreAction::EStore);
+	OutRenderTargets[3] = FRHIRenderTargetView(GBufferC->GetRenderTargetItem().TargetableTexture, 0, -1, ColorLoadAction, ERenderTargetStoreAction::EStore);
+	OutRenderTargets[4] = FRHIRenderTargetView(GBufferD->GetRenderTargetItem().TargetableTexture, 0, -1, ColorLoadAction, ERenderTargetStoreAction::EStore);
 	int32 MRTCount = 5;	// Derived from previously used RenderTargets[]; would have been nice to keep a pointer while filling it but it's more confusing to use
 
 	OutVelocityRTIndex = -1;
@@ -330,7 +480,7 @@ int32 FSceneRenderTargets::GetGBufferRenderTargets(ERenderTargetLoadAction Color
 	if (bAllowStaticLighting)
 	{
 		check(MRTCount == GetGBufferEIndex());
-		OutRenderTargets[MRTCount] = FRHIRenderTargetView(GSceneRenderTargets.GBufferE->GetRenderTargetItem().TargetableTexture, 0, -1, ColorLoadAction, ERenderTargetStoreAction::EStore);
+		OutRenderTargets[MRTCount] = FRHIRenderTargetView(GBufferE->GetRenderTargetItem().TargetableTexture, 0, -1, ColorLoadAction, ERenderTargetStoreAction::EStore);
 		++MRTCount;
 	}
 
@@ -339,7 +489,7 @@ int32 FSceneRenderTargets::GetGBufferRenderTargets(ERenderTargetLoadAction Color
 		OutVelocityRTIndex = MRTCount;
 		++MRTCount;
 		check(OutVelocityRTIndex == GetGBufferVelocityIndex());
-		OutRenderTargets[OutVelocityRTIndex] = FRHIRenderTargetView(GSceneRenderTargets.GBufferVelocity->GetRenderTargetItem().TargetableTexture, 0, -1, ColorLoadAction, ERenderTargetStoreAction::EStore);
+		OutRenderTargets[OutVelocityRTIndex] = FRHIRenderTargetView(GBufferVelocity->GetRenderTargetItem().TargetableTexture, 0, -1, ColorLoadAction, ERenderTargetStoreAction::EStore);
 	}
 
 	check(MRTCount <= MaxSimultaneousRenderTargets);
@@ -902,19 +1052,19 @@ void FSceneRenderTargets::ResolveGBufferSurfaces(FRHICommandList& RHICmdList, co
 	{
 		SCOPED_DRAW_EVENT(RHICmdList, ResolveGBufferSurfaces);
 
-		RHICmdList.CopyToResolveTarget(GSceneRenderTargets.GBufferA->GetRenderTargetItem().TargetableTexture, GSceneRenderTargets.GBufferA->GetRenderTargetItem().ShaderResourceTexture, true, FResolveParams(ResolveRect));
-		RHICmdList.CopyToResolveTarget(GSceneRenderTargets.GBufferB->GetRenderTargetItem().TargetableTexture, GSceneRenderTargets.GBufferB->GetRenderTargetItem().ShaderResourceTexture, true, FResolveParams(ResolveRect));
-		RHICmdList.CopyToResolveTarget(GSceneRenderTargets.GBufferC->GetRenderTargetItem().TargetableTexture, GSceneRenderTargets.GBufferC->GetRenderTargetItem().ShaderResourceTexture, true, FResolveParams(ResolveRect));
-		RHICmdList.CopyToResolveTarget(GSceneRenderTargets.GBufferD->GetRenderTargetItem().TargetableTexture, GSceneRenderTargets.GBufferD->GetRenderTargetItem().ShaderResourceTexture, true, FResolveParams(ResolveRect));
+		RHICmdList.CopyToResolveTarget(GBufferA->GetRenderTargetItem().TargetableTexture, GBufferA->GetRenderTargetItem().ShaderResourceTexture, true, FResolveParams(ResolveRect));
+		RHICmdList.CopyToResolveTarget(GBufferB->GetRenderTargetItem().TargetableTexture, GBufferB->GetRenderTargetItem().ShaderResourceTexture, true, FResolveParams(ResolveRect));
+		RHICmdList.CopyToResolveTarget(GBufferC->GetRenderTargetItem().TargetableTexture, GBufferC->GetRenderTargetItem().ShaderResourceTexture, true, FResolveParams(ResolveRect));
+		RHICmdList.CopyToResolveTarget(GBufferD->GetRenderTargetItem().TargetableTexture, GBufferD->GetRenderTargetItem().ShaderResourceTexture, true, FResolveParams(ResolveRect));
 		
 		if (bAllowStaticLighting)
 		{
-			RHICmdList.CopyToResolveTarget(GSceneRenderTargets.GBufferE->GetRenderTargetItem().TargetableTexture, GSceneRenderTargets.GBufferE->GetRenderTargetItem().ShaderResourceTexture, true, FResolveParams(ResolveRect));
+			RHICmdList.CopyToResolveTarget(GBufferE->GetRenderTargetItem().TargetableTexture, GBufferE->GetRenderTargetItem().ShaderResourceTexture, true, FResolveParams(ResolveRect));
 		}
 
 		if (bAllocateVelocityGBuffer)
 		{
-			RHICmdList.CopyToResolveTarget(GSceneRenderTargets.GBufferVelocity->GetRenderTargetItem().TargetableTexture, GSceneRenderTargets.GBufferVelocity->GetRenderTargetItem().ShaderResourceTexture, true, FResolveParams(ResolveRect));
+			RHICmdList.CopyToResolveTarget(GBufferVelocity->GetRenderTargetItem().TargetableTexture, GBufferVelocity->GetRenderTargetItem().ShaderResourceTexture, true, FResolveParams(ResolveRect));
 		}
 	}
 }
@@ -1032,7 +1182,7 @@ void FSceneRenderTargets::FinishRenderingCubeShadowDepth(FRHICommandList& RHICmd
 void FSceneRenderTargets::BeginRenderingSceneAlphaCopy(FRHICommandListImmediate& RHICmdList)
 {
 	SCOPED_DRAW_EVENT(RHICmdList, BeginRenderingSceneAlphaCopy);
-	GRenderTargetPool.VisualizeTexture.SetCheckPoint(RHICmdList, GSceneRenderTargets.SceneAlphaCopy);
+	GRenderTargetPool.VisualizeTexture.SetCheckPoint(RHICmdList, SceneAlphaCopy);
 	SetRenderTarget(RHICmdList, GetSceneAlphaCopySurface(), 0);
 }
 
@@ -1040,7 +1190,7 @@ void FSceneRenderTargets::FinishRenderingSceneAlphaCopy(FRHICommandListImmediate
 {
 	SCOPED_DRAW_EVENT(RHICmdList, FinishRenderingSceneAlphaCopy);
 	RHICmdList.CopyToResolveTarget(GetSceneAlphaCopySurface(), SceneAlphaCopy->GetRenderTargetItem().ShaderResourceTexture, false, FResolveParams(FResolveRect()));
-	GRenderTargetPool.VisualizeTexture.SetCheckPoint(RHICmdList, GSceneRenderTargets.SceneAlphaCopy);
+	GRenderTargetPool.VisualizeTexture.SetCheckPoint(RHICmdList, SceneAlphaCopy);
 }
 
 
@@ -1050,7 +1200,7 @@ void FSceneRenderTargets::BeginRenderingLightAttenuation(FRHICommandList& RHICmd
 
 	AllocLightAttenuation();
 
-	GRenderTargetPool.VisualizeTexture.SetCheckPoint(RHICmdList, GSceneRenderTargets.GetLightAttenuation());
+	GRenderTargetPool.VisualizeTexture.SetCheckPoint(RHICmdList, GetLightAttenuation());
 
 	// Set the light attenuation surface as the render target, and the scene depth buffer as the depth-stencil surface.
 	SetRenderTarget(RHICmdList, GetLightAttenuationSurface(), GetSceneDepthSurface(), bClearToWhite ? ESimpleRenderTargetMode::EClearColorToWhite : ESimpleRenderTargetMode::EExistingColorAndDepth, FExclusiveDepthStencil::DepthRead_StencilWrite);
@@ -1063,13 +1213,13 @@ void FSceneRenderTargets::FinishRenderingLightAttenuation(FRHICommandList& RHICm
 	// Resolve the light attenuation surface.
 	RHICmdList.CopyToResolveTarget(GetLightAttenuationSurface(), LightAttenuation->GetRenderTargetItem().ShaderResourceTexture, false, FResolveParams(FResolveRect()));
 	
-	GRenderTargetPool.VisualizeTexture.SetCheckPoint(RHICmdList, GSceneRenderTargets.GetLightAttenuation());
+	GRenderTargetPool.VisualizeTexture.SetCheckPoint(RHICmdList, GetLightAttenuation());
 }
 
 void FSceneRenderTargets::BeginRenderingTranslucency(FRHICommandList& RHICmdList, const FViewInfo& View)
 {
 	// Use the scene color buffer.
-	GSceneRenderTargets.BeginRenderingSceneColor(RHICmdList, ESimpleRenderTargetMode::EExistingColorAndDepth, FExclusiveDepthStencil::DepthRead_StencilWrite);
+	BeginRenderingSceneColor(RHICmdList, ESimpleRenderTargetMode::EExistingColorAndDepth, FExclusiveDepthStencil::DepthRead_StencilWrite);
 	
 	// Clear the stencil buffer for ResponsiveAA
 	RHICmdList.Clear(false, FLinearColor::White, false, (float)ERHIZBuffer::FarPlane, true, 0, FIntRect());
@@ -1080,7 +1230,7 @@ void FSceneRenderTargets::BeginRenderingTranslucency(FRHICommandList& RHICmdList
 
 void FSceneRenderTargets::FinishRenderingTranslucency(FRHICommandListImmediate& RHICmdList, const class FViewInfo& View)
 {
-	GSceneRenderTargets.FinishRenderingSceneColor(RHICmdList, true);
+	FinishRenderingSceneColor(RHICmdList, true);
 }
 
 
@@ -1090,18 +1240,22 @@ bool FSceneRenderTargets::BeginRenderingSeparateTranslucency(FRHICommandList& RH
 	{
 		SCOPED_DRAW_EVENT(RHICmdList, BeginSeparateTranslucency);
 
-		FSceneViewState* ViewState = (FSceneViewState*)View.State;
-
-		// the RT should only be available for a short period during rendering
-		if(bFirstTimeThisFrame)
+		TRefCountPtr<IPooledRenderTarget>* SeparateTranslucency;
+		if (bSnapshot)
+		{
+			check(SeparateTranslucencyRT.GetReference());
+			SeparateTranslucency = &SeparateTranslucencyRT;
+		}
+		else
 		{
-			check(!ViewState->SeparateTranslucencyRT);
+			FSceneViewState* ViewState = (FSceneViewState*)View.State;
+			SeparateTranslucency = &ViewState->GetSeparateTranslucency(GetBufferSizeXY());
 		}
 
-		TRefCountPtr<IPooledRenderTarget>& SeparateTranslucency = ViewState->GetSeparateTranslucency(View);
+
 
 		// clear the render target the first time, re-use afterwards
-		SetRenderTarget(RHICmdList, SeparateTranslucency->GetRenderTargetItem().TargetableTexture, GetSceneDepthSurface(), 
+		SetRenderTarget(RHICmdList, (*SeparateTranslucency)->GetRenderTargetItem().TargetableTexture, GetSceneDepthSurface(), 
 			bFirstTimeThisFrame ? ESimpleRenderTargetMode::EClearColorToBlackWithFullAlpha : ESimpleRenderTargetMode::EExistingColorAndDepth, FExclusiveDepthStencil::DepthRead_StencilWrite);
 
 		if (!bFirstTimeThisFrame)
@@ -1126,10 +1280,19 @@ void FSceneRenderTargets::FinishRenderingSeparateTranslucency(FRHICommandList& R
 	{
 		SCOPED_DRAW_EVENT(RHICmdList, FinishSeparateTranslucency);
 
-		FSceneViewState* ViewState = (FSceneViewState*)View.State;
-		TRefCountPtr<IPooledRenderTarget>& SeparateTranslucency = ViewState->GetSeparateTranslucency(View);
+		TRefCountPtr<IPooledRenderTarget>* SeparateTranslucency;
+		if (bSnapshot)
+		{
+			check(SeparateTranslucencyRT.GetReference());
+			SeparateTranslucency = &SeparateTranslucencyRT;
+		}
+		else
+		{
+			FSceneViewState* ViewState = (FSceneViewState*)View.State;
+			SeparateTranslucency = &ViewState->GetSeparateTranslucency(GetBufferSizeXY());
+		}
 
-		RHICmdList.CopyToResolveTarget(SeparateTranslucency->GetRenderTargetItem().TargetableTexture, SeparateTranslucency->GetRenderTargetItem().ShaderResourceTexture, true, FResolveParams());
+		RHICmdList.CopyToResolveTarget((*SeparateTranslucency)->GetRenderTargetItem().TargetableTexture, (*SeparateTranslucency)->GetRenderTargetItem().ShaderResourceTexture, true, FResolveParams());
 	}
 }
 
@@ -1267,7 +1430,7 @@ void FSceneRenderTargets::InitEditorPrimitivesDepth()
 void FSceneRenderTargets::QuantizeBufferSize(int32& InOutBufferSizeX, int32& InOutBufferSizeY)
 {
 	// ensure sizes are dividable by DividableBy to get post processing effects with lower resolution working well
-	const uint32 DividableBy = 8;
+	const uint32 DividableBy = 4;
 
 	const uint32 Mask = ~(DividableBy - 1);
 	InOutBufferSizeX = (InOutBufferSizeX + DividableBy - 1) & Mask;
@@ -1649,6 +1812,7 @@ void FSceneRenderTargets::ReleaseAllTargets()
 	LightAccumulation.SafeRelease();
 	DirectionalOcclusion.SafeRelease();
 	CustomDepth.SafeRelease();
+	CustomStencilSRV.SafeRelease();
 	ReflectiveShadowMapNormal.SafeRelease();
 	ReflectiveShadowMapDiffuse.SafeRelease();
 	ReflectiveShadowMapDepth.SafeRelease();
@@ -1774,6 +1938,7 @@ IPooledRenderTarget* FSceneRenderTargets::RequestCustomDepth(bool bPrimitives)
 		// Todo: Could check if writes stencil here and create min viable target
 		FPooledRenderTargetDesc Desc(FPooledRenderTargetDesc::Create2DDesc(BufferSize, PF_DepthStencil, TexCreate_None, TexCreate_DepthStencilTargetable, false));
 		GRenderTargetPool.FindFreeElement(Desc, CustomDepth, TEXT("CustomDepth"));
+		CustomStencilSRV = RHICreateShaderResourceView((FTexture2DRHIRef&)CustomDepth->GetRenderTargetItem().TargetableTexture, 0, 1, PF_X24_G8);
 		return CustomDepth;
 	}
 
@@ -1892,6 +2057,8 @@ void FSceneTextureShaderParameters::Set(
 {
 	if (TextureMode == ESceneRenderTargetsMode::SetTextures)
 	{
+		FSceneRenderTargets& SceneContext = FSceneRenderTargets::Get(RHICmdList);
+
 		// optimization possible: TShaderRHIParamRef is no param Ref
 		if (SceneColorTextureParameter.IsBound())
 		{
@@ -1922,11 +2089,11 @@ void FSceneTextureShaderParameters::Set(
 				SceneColorTextureParameter,
 				SceneColorTextureParameterSampler,
 				Filter,
-				GSceneRenderTargets.GetSceneColorTexture()
+				SceneContext.GetSceneColorTexture()
 				);
 		}
 
-		if (SceneAlphaCopyTextureParameter.IsBound() && GSceneRenderTargets.HasSceneAlphaCopyTexture())
+		if (SceneAlphaCopyTextureParameter.IsBound() && SceneContext.HasSceneAlphaCopyTexture())
 		{
 			FSamplerStateRHIRef Filter;
 			Filter = TStaticSamplerState<SF_Point,AM_Clamp,AM_Clamp,AM_Clamp>::GetRHI();
@@ -1936,13 +2103,13 @@ void FSceneTextureShaderParameters::Set(
 				SceneAlphaCopyTextureParameter,
 				SceneAlphaCopyTextureParameterSampler,
 				Filter,
-				GSceneRenderTargets.GetSceneAlphaCopyTexture()
+				SceneContext.GetSceneAlphaCopyTexture()
 				);
 		}
 
 		if(SceneDepthTextureParameter.IsBound() || SceneDepthTextureParameterSampler.IsBound())
 		{
-			const FTexture2DRHIRef* DepthTexture = GSceneRenderTargets.GetActualDepthTexture();
+			const FTexture2DRHIRef* DepthTexture = SceneContext.GetActualDepthTexture();
 			SetTextureParameter(
 				RHICmdList, 
 				ShaderRHI,
@@ -1957,7 +2124,7 @@ void FSceneTextureShaderParameters::Set(
 
 		if (FeatureLevel >= ERHIFeatureLevel::SM5)
 		{
-			SetTextureParameter(RHICmdList, ShaderRHI, SceneColorSurfaceParameter, GSceneRenderTargets.GetSceneColorSurface());
+			SetTextureParameter(RHICmdList, ShaderRHI, SceneColorSurfaceParameter, SceneContext.GetSceneColorSurface());
 		}
 		if (FeatureLevel >= ERHIFeatureLevel::SM4)
 		{
@@ -1965,22 +2132,22 @@ void FSceneTextureShaderParameters::Set(
 			{
 				if(SceneDepthSurfaceParameter.IsBound())
 				{
-					SetTextureParameter(RHICmdList, ShaderRHI, SceneDepthSurfaceParameter, GSceneRenderTargets.GetSceneDepthSurface());
+					SetTextureParameter(RHICmdList, ShaderRHI, SceneDepthSurfaceParameter, SceneContext.GetSceneDepthSurface());
 				}
 				if(SceneDepthTextureNonMS.IsBound())
 				{
-					SetTextureParameter(RHICmdList, ShaderRHI, SceneDepthTextureNonMS, GSceneRenderTargets.GetSceneDepthTexture());
+					SetTextureParameter(RHICmdList, ShaderRHI, SceneDepthTextureNonMS, SceneContext.GetSceneDepthTexture());
 				}
 			}
 			else
 			{
 				if(SceneDepthSurfaceParameter.IsBound())
 				{
-					SetTextureParameter(RHICmdList, ShaderRHI, SceneDepthSurfaceParameter, GSceneRenderTargets.GetAuxiliarySceneDepthSurface());
+					SetTextureParameter(RHICmdList, ShaderRHI, SceneDepthSurfaceParameter, SceneContext.GetAuxiliarySceneDepthSurface());
 				}
 				if(SceneDepthTextureNonMS.IsBound())
 				{
-					SetTextureParameter(RHICmdList, ShaderRHI, SceneDepthTextureNonMS, GSceneRenderTargets.GetAuxiliarySceneDepthSurface());
+					SetTextureParameter(RHICmdList, ShaderRHI, SceneDepthTextureNonMS, SceneContext.GetAuxiliarySceneDepthSurface());
 				}
 			}
 		}
@@ -2030,7 +2197,8 @@ void FSceneTextureShaderParameters::Set(
 		FTextureRHIParamRef DirectionalOcclusion = nullptr;
 		if( bDirectionalOcclusion )
 		{
-			DirectionalOcclusion = GSceneRenderTargets.GetDirectionalOcclusionTexture();
+			FSceneRenderTargets& SceneContext = FSceneRenderTargets::Get(RHICmdList);
+			DirectionalOcclusion = SceneContext.GetDirectionalOcclusionTexture();
 		}
 		else
 		{
@@ -2115,15 +2283,16 @@ bool IsDBufferEnabled();
 template< typename ShaderRHIParamRef >
 void FDeferredPixelShaderParameters::Set(FRHICommandList& RHICmdList, const ShaderRHIParamRef ShaderRHI, const FSceneView& View, ESceneRenderTargetsMode::Type TextureMode) const
 {
+	FSceneRenderTargets& SceneContext = FSceneRenderTargets::Get(RHICmdList);
 	// This is needed on PC ES2 for SceneAlphaCopy, probably should be refactored for performance.
 	SceneTextureParameters.Set(RHICmdList, ShaderRHI, View, TextureMode, SF_Point);
 
 	// if() is purely an optimization and could be removed
 	if(IsDBufferEnabled())
 	{
-		IPooledRenderTarget* DBufferA = GSceneRenderTargets.DBufferA ? GSceneRenderTargets.DBufferA : GSystemTextures.BlackDummy;
-		IPooledRenderTarget* DBufferB = GSceneRenderTargets.DBufferB ? GSceneRenderTargets.DBufferB : GSystemTextures.BlackDummy;
-		IPooledRenderTarget* DBufferC = GSceneRenderTargets.DBufferC ? GSceneRenderTargets.DBufferC : GSystemTextures.BlackDummy;
+		IPooledRenderTarget* DBufferA = SceneContext.DBufferA ? SceneContext.DBufferA : GSystemTextures.BlackDummy;
+		IPooledRenderTarget* DBufferB = SceneContext.DBufferB ? SceneContext.DBufferB : GSystemTextures.BlackDummy;
+		IPooledRenderTarget* DBufferC = SceneContext.DBufferC ? SceneContext.DBufferC : GSystemTextures.BlackDummy;
 
 		// todo: optimize out when not needed
 		SetTextureParameter(RHICmdList, ShaderRHI, DBufferATexture, DBufferATextureSampler, TStaticSamplerState<>::GetRHI(), DBufferA->GetRenderTargetItem().ShaderResourceTexture);
@@ -2142,14 +2311,14 @@ void FDeferredPixelShaderParameters::Set(FRHICommandList& RHICmdList, const Shad
 	if (TextureMode == ESceneRenderTargetsMode::SetTextures && FeatureLevel >= ERHIFeatureLevel::SM4)
 	{
 		// if there is no ambient occlusion it's better to have white there
-		IPooledRenderTarget* ScreenSpaceAO = GSceneRenderTargets.ScreenSpaceAO;
-		if(!GSceneRenderTargets.bScreenSpaceAOIsValid)
+		IPooledRenderTarget* ScreenSpaceAO = SceneContext.ScreenSpaceAO;
+		if(!SceneContext.bScreenSpaceAOIsValid)
 		{
 			ScreenSpaceAO = GSystemTextures.WhiteDummy;
 		}
 
 		// if there is no custom depth it's better to have the far distance there
-		IPooledRenderTarget* CustomDepth = GSceneRenderTargets.bCustomDepthIsValid ? GSceneRenderTargets.CustomDepth.GetReference() : 0;
+		IPooledRenderTarget* CustomDepth = SceneContext.bCustomDepthIsValid ? SceneContext.CustomDepth.GetReference() : 0;
 		if(!CustomDepth)
 		{
 			CustomDepth = GSystemTextures.BlackDummy;
@@ -2159,7 +2328,7 @@ void FDeferredPixelShaderParameters::Set(FRHICommandList& RHICmdList, const Shad
 		{
 			if (GBufferResources.IsBound())
 			{
-				SetUniformBufferParameter(RHICmdList, ShaderRHI, GBufferResources, GSceneRenderTargets.GetGBufferResourcesUniformBuffer());
+				SetUniformBufferParameter(RHICmdList, ShaderRHI, GBufferResources, SceneContext.GetGBufferResourcesUniformBuffer());
 			}
 
 			SetTextureParameter(RHICmdList, ShaderRHI, ScreenSpaceAOTexture, ScreenSpaceAOTextureSampler, TStaticSamplerState<>::GetRHI(), ScreenSpaceAO->GetRenderTargetItem().ShaderResourceTexture);
@@ -2171,27 +2340,9 @@ void FDeferredPixelShaderParameters::Set(FRHICommandList& RHICmdList, const Shad
 
 			if (CustomStencilTexture.IsBound())
 			{
-				FSceneViewState* ViewState = (FSceneViewState*)View.State;	
-					
-				if (ViewState && GSceneRenderTargets.bCustomDepthIsValid)
+				if (SceneContext.bCustomDepthIsValid && SceneContext.CustomStencilSRV.GetReference())
 				{
-					FTexture2DRHIRef& CustomStencilTargetableTexture = (FTexture2DRHIRef&)CustomDepth->GetRenderTargetItem().TargetableTexture;
-					
-					// Release cached stencil SRV if outdated
-					if (ViewState->CustomStencilSRVCacheKey != CustomStencilTargetableTexture)
-					{
-						ViewState->CustomStencilSRVCacheKey.SafeRelease();
-						ViewState->CustomStencilSRVCacheValue.SafeRelease();
-					}
-
-					// Create and cache stencil SRV if needed
-					if (!ViewState->CustomStencilSRVCacheValue)
-					{
-						ViewState->CustomStencilSRVCacheKey = CustomStencilTargetableTexture;
-						ViewState->CustomStencilSRVCacheValue = RHICreateShaderResourceView(CustomStencilTargetableTexture, 0, 1, PF_X24_G8);
-					}
-
-					SetSRVParameter(RHICmdList, ShaderRHI, CustomStencilTexture, ViewState->CustomStencilSRVCacheValue);
+					SetSRVParameter(RHICmdList, ShaderRHI, CustomStencilTexture, SceneContext.CustomStencilSRV);
 				}
 				else
 				{
diff --git a/Engine/Source/Runtime/Renderer/Private/PostProcess/SceneRenderTargets.h b/Engine/Source/Runtime/Renderer/Private/PostProcess/SceneRenderTargets.h
--- a/Engine/Source/Runtime/Renderer/Private/PostProcess/SceneRenderTargets.h
+++ b/Engine/Source/Runtime/Renderer/Private/PostProcess/SceneRenderTargets.h
@@ -112,6 +112,21 @@ public:
 	/** Destructor. */
 	virtual ~FSceneRenderTargets() {}
 
+	/** Singletons. At the moment parallel tasks get their snapshot from the rhicmdlist */
+	static FSceneRenderTargets& Get(FRHICommandList& RHICmdList);
+	static FSceneRenderTargets& Get(FRHICommandListImmediate& RHICmdList);
+	// this is a placeholder, the context should come from somewhere. This is very unsafe, please don't use it!
+	static FSceneRenderTargets& Get_Todo_PassContext();
+	// As above but relaxed checks and always gives the global FSceneRenderTargets. The intention here is that it is only used for constants that don't change during a frame. This is very unsafe, please don't use it!
+	static FSceneRenderTargets& Get_FrameConstantsOnly();
+
+	/** Create a snapshot on the scene allocator */
+	FSceneRenderTargets* CreateSnapshot(const FViewInfo& InView);
+	/** Set a snapshot on the TargetCmdList */
+	void SetSnapshotOnCmdList(FRHICommandList& TargetCmdList);	
+	/** Destruct all snapshots */
+	void DestroyAllSnapshots();
+
 protected:
 	/** Constructor */
 	FSceneRenderTargets(): 
@@ -122,6 +137,7 @@ protected:
 		LargestDesiredSizeThisFrame( 0, 0 ),
 		LargestDesiredSizeLastFrame( 0, 0 ),
 		ThisFrameNumber( 0 ),
+		bVelocityPass(false),
 		BufferSize(0, 0), 
 		SmallColorDepthDownsampleFactor(2),
 		bLightAttenuationEnabled(true),
@@ -136,9 +152,12 @@ protected:
 		bCurrentLightPropagationVolume(false),
 		CurrentFeatureLevel(ERHIFeatureLevel::Num),
 		CurrentShadingPath(EShadingPath::Num),
-		bAllocateVelocityGBuffer(false)
+		bAllocateVelocityGBuffer(false),
+		bSnapshot(false)
 		{
 		}
+	/** Constructor that creates snapshot */
+	FSceneRenderTargets(const FViewInfo& InView, const FSceneRenderTargets& SnapshotSource);
 public:
 
 	enum class EShadingPath
@@ -468,6 +487,18 @@ public:
 	 */
 	static void QuantizeBufferSize(int32& InOutBufferSizeX, int32& InOutBufferSizeY);
 
+	bool IsSeparateTranslucencyActive(const FViewInfo& View) const;
+
+	bool IsVelocityPass() const
+	{
+		return bVelocityPass;
+	}
+	void SetVelocityPass(bool bInVelocityPass)
+	{
+		bVelocityPass = bInVelocityPass;
+	}
+
+
 private: // Get...() methods instead of direct access
 
 	// 0 before BeginRenderingSceneColor and after tone mapping in deferred shading
@@ -508,6 +539,8 @@ public:
 	TRefCountPtr<IPooledRenderTarget> ScreenSpaceAO;
 	// used by the CustomDepth material feature, is allocated on demand or if r.CustomDepth is 2
 	TRefCountPtr<IPooledRenderTarget> CustomDepth;
+	// used by the CustomDepth material feature for stencil
+	TRefCountPtr<FRHIShaderResourceView> CustomStencilSRV;
 	// Render target for per-object shadow depths.
 	TRefCountPtr<IPooledRenderTarget> ShadowDepthZ;
 	// optional in case this RHI requires a color render target
@@ -550,7 +583,8 @@ public:
 	/** Depth for editor primitives */
 	TRefCountPtr<IPooledRenderTarget> EditorPrimitivesDepth;
 
-	bool IsSeparateTranslucencyActive(const FViewInfo& View) const;
+	/** ONLY for snapshots!!! this is a copy of the SeparateTranslucencyRT from the view state. */
+	TRefCountPtr<IPooledRenderTarget> SeparateTranslucencyRT;
 
 	// todo: free ScreenSpaceAO so pool can reuse
 	bool bScreenSpaceAOIsValid;
@@ -571,6 +605,9 @@ private:
 	/** to detect when LargestDesiredSizeThisFrame is outdated */
 	uint32 ThisFrameNumber;
 
+	bool bVelocityPass;
+	/** CAUTION: When adding new data, make sure you copy it in the snapshot constructor! **/
+
 	/**
 	 * Initializes the editor primitive color render target
 	 */
@@ -665,7 +702,12 @@ private:
 	/** Helpers to track scenedepth state on platforms that need to propagate clear information across parallel rendering boundaries. */
 	bool bSceneDepthCleared;
 	float SceneDepthClearValue;
+	/** true is this is a snapshot on the scene allocator */
+	bool bSnapshot;
+	/** All outstanding snapshots */
+	TArray<FSceneRenderTargets*> Snapshots;
+
+	/** CAUTION: When adding new data, make sure you copy it in the snapshot constructor! **/
+
 };
 
-/** The global render targets used for scene rendering. */
-extern RENDERER_API TGlobalResource<FSceneRenderTargets> GSceneRenderTargets;
diff --git a/Engine/Source/Runtime/Renderer/Private/PostProcess/ScreenSpaceReflections.cpp b/Engine/Source/Runtime/Renderer/Private/PostProcess/ScreenSpaceReflections.cpp
--- a/Engine/Source/Runtime/Renderer/Private/PostProcess/ScreenSpaceReflections.cpp
+++ b/Engine/Source/Runtime/Renderer/Private/PostProcess/ScreenSpaceReflections.cpp
@@ -278,6 +278,7 @@ static int32 ComputeSSRQuality(float Quality)
 void FRCPassPostProcessScreenSpaceReflections::Process(FRenderingCompositePassContext& Context)
 {
 	auto& RHICmdList = Context.RHICmdList;
+	FSceneRenderTargets& SceneContext = FSceneRenderTargets::Get(RHICmdList);
 
 	const FSceneView& View = Context.View;
 	const auto FeatureLevel = Context.GetFeatureLevel();
@@ -306,11 +307,11 @@ void FRCPassPostProcessScreenSpaceReflections::Process(FRenderingCompositePassCo
 		TShaderMapRef< FPostProcessScreenSpaceReflectionsStencilPS > PixelShader(Context.GetShaderMap());
 		
 		// bind the dest render target and the depth stencil render target
-		SetRenderTarget(RHICmdList, DestRenderTarget.TargetableTexture, GSceneRenderTargets.GetSceneDepthSurface(), ESimpleRenderTargetMode::EUninitializedColorAndDepth, FExclusiveDepthStencil::DepthRead_StencilWrite);
+		SetRenderTarget(RHICmdList, DestRenderTarget.TargetableTexture, SceneContext.GetSceneDepthSurface(), ESimpleRenderTargetMode::EUninitializedColorAndDepth, FExclusiveDepthStencil::DepthRead_StencilWrite);
 		Context.SetViewportAndCallRHI(View.ViewRect);
-		
+
 		// Clear stencil to 0
-		RHICmdList.Clear(false, FLinearColor::White, false, (float)ERHIZBuffer::FarPlane, true, 0, FIntRect());
+		RHICmdList.Clear(false, FLinearColor::White, false, (float)ERHIZBuffer::FarPlane, true, 0, View.ViewRect);
 		
 		// bind shader
 		static FGlobalBoundShaderState BoundShaderState;
@@ -326,7 +327,7 @@ void FRCPassPostProcessScreenSpaceReflections::Process(FRenderingCompositePassCo
 
 		// disable blend mode
 		RHICmdList.SetBlendState(TStaticBlendState<>::GetRHI());
-		
+	
 		// Draw a quad mapping scene color to the view's render target to set stencil to set the stencil mask where it needs to be
 		DrawRectangle( 
 			Context.RHICmdList,
@@ -335,28 +336,28 @@ void FRCPassPostProcessScreenSpaceReflections::Process(FRenderingCompositePassCo
 			View.ViewRect.Min.X, View.ViewRect.Min.Y, 
 			View.ViewRect.Width(), View.ViewRect.Height(),
 			View.ViewRect.Size(),
-			GSceneRenderTargets.GetBufferSizeXY(),
+			SceneContext.GetBufferSizeXY(),
 			*VertexShader,
 			EDRF_UseTriangleOptimization);
 	} // ScreenSpaceReflectionsStencil draw event
 
 	{ // ScreenSpaceReflections draw event
 		SCOPED_DRAW_EVENT(Context.RHICmdList, ScreenSpaceReflections);
-		
+
 		if (SSRStencilPrePass)
 		{
 			// set up the stencil test to match 0, meaning FPostProcessScreenSpaceReflectionsStencilPS has been discarded
 			RHICmdList.SetDepthStencilState(TStaticDepthStencilState<false, CF_Always, true, CF_Equal, SO_Keep, SO_Keep, SO_Keep>::GetRHI(), 0);
 		}
 		else
-		{
+	{
 			// bind only the dest render target
 			SetRenderTarget(RHICmdList, DestRenderTarget.TargetableTexture, FTextureRHIRef());
 			Context.SetViewportAndCallRHI(View.ViewRect);
 
 			RHICmdList.SetDepthStencilState(TStaticDepthStencilState<false, CF_Always>::GetRHI());
-		}
-		
+	}
+
 		// clear DestRenderTarget only outside of the view's rectangle
 		RHICmdList.Clear(true, FLinearColor::Black, false, (float)ERHIZBuffer::FarPlane, false, 0, View.ViewRect);
 		
@@ -364,43 +365,43 @@ void FRCPassPostProcessScreenSpaceReflections::Process(FRenderingCompositePassCo
 		RHICmdList.SetBlendState(TStaticBlendState<>::GetRHI());
 		RHICmdList.SetRasterizerState(TStaticRasterizerState<>::GetRHI());
 
-		TShaderMapRef< FPostProcessVS > VertexShader(Context.GetShaderMap());
+	TShaderMapRef< FPostProcessVS > VertexShader(Context.GetShaderMap());
 
-		#define CASE(A, B) \
-			case (A + 2 * (B + 3 * 0 )): \
-			{ \
-				TShaderMapRef< FPostProcessScreenSpaceReflectionsPS<A, B> > PixelShader(Context.GetShaderMap()); \
-				static FGlobalBoundShaderState BoundShaderState; \
+	#define CASE(A, B) \
+		case (A + 2 * (B + 3 * 0 )): \
+		{ \
+			TShaderMapRef< FPostProcessScreenSpaceReflectionsPS<A, B> > PixelShader(Context.GetShaderMap()); \
+			static FGlobalBoundShaderState BoundShaderState; \
 				SetGlobalBoundShaderState(RHICmdList, FeatureLevel, BoundShaderState, GFilterVertexDeclaration.VertexDeclarationRHI, *VertexShader, *PixelShader); \
-				VertexShader->SetParameters(Context); \
-				PixelShader->SetParameters(Context); \
-			}; \
-			break
+			VertexShader->SetParameters(Context); \
+			PixelShader->SetParameters(Context); \
+		}; \
+		break
 
-		switch (iPreFrame + 2 * (SSRQuality + 3 * 0))
-		{
-			CASE(0,0);
-			CASE(0,1);	CASE(1,1);
-			CASE(0,2);	CASE(1,2);
-			CASE(0,3);	CASE(1,3);
-			CASE(0,4);	CASE(1,4);
-			default:
-				check(!"Missing case in FRCPassPostProcessScreenSpaceReflections");
-		}
-		#undef CASE
+	switch (iPreFrame + 2 * (SSRQuality + 3 * 0))
+	{
+		CASE(0,0);
+		CASE(0,1);	CASE(1,1);
+		CASE(0,2);	CASE(1,2);
+		CASE(0,3);	CASE(1,3);
+		CASE(0,4);	CASE(1,4);
+		default:
+			check(!"Missing case in FRCPassPostProcessScreenSpaceReflections");
+	}
+	#undef CASE
 
 
-		// Draw a quad mapping scene color to the view's render target
-		DrawRectangle( 
+	// Draw a quad mapping scene color to the view's render target
+	DrawRectangle( 
 			RHICmdList,
-			0, 0,
-			View.ViewRect.Width(), View.ViewRect.Height(),
-			View.ViewRect.Min.X, View.ViewRect.Min.Y, 
-			View.ViewRect.Width(), View.ViewRect.Height(),
-			View.ViewRect.Size(),
-			GSceneRenderTargets.GetBufferSizeXY(),
-			*VertexShader,
-			EDRF_UseTriangleOptimization);
+		0, 0,
+		View.ViewRect.Width(), View.ViewRect.Height(),
+		View.ViewRect.Min.X, View.ViewRect.Min.Y, 
+		View.ViewRect.Width(), View.ViewRect.Height(),
+		View.ViewRect.Size(),
+		FSceneRenderTargets::Get(Context.RHICmdList).GetBufferSizeXY(),
+		*VertexShader,
+		EDRF_UseTriangleOptimization);
 
 		RHICmdList.CopyToResolveTarget(DestRenderTarget.TargetableTexture, DestRenderTarget.ShaderResourceTexture, false, FResolveParams());
 	} // ScreenSpaceReflections
@@ -408,7 +409,7 @@ void FRCPassPostProcessScreenSpaceReflections::Process(FRenderingCompositePassCo
 
 FPooledRenderTargetDesc FRCPassPostProcessScreenSpaceReflections::ComputeOutputDesc(EPassOutputId InPassOutputId) const
 {
-	FPooledRenderTargetDesc Ret(FPooledRenderTargetDesc::Create2DDesc(GSceneRenderTargets.GetBufferSizeXY(), PF_FloatRGBA, TexCreate_None, TexCreate_RenderTargetable, false));
+	FPooledRenderTargetDesc Ret(FPooledRenderTargetDesc::Create2DDesc(FSceneRenderTargets::Get_FrameConstantsOnly().GetBufferSizeXY(), PF_FloatRGBA, TexCreate_None, TexCreate_RenderTargetable, false));
 
 	Ret.DebugName = TEXT("ScreenSpaceReflections");
 
@@ -417,12 +418,12 @@ FPooledRenderTargetDesc FRCPassPostProcessScreenSpaceReflections::ComputeOutputD
 
 void ScreenSpaceReflections(FRHICommandListImmediate& RHICmdList, FViewInfo& View, TRefCountPtr<IPooledRenderTarget>& SSROutput)
 {
-	FRenderingCompositePassContext CompositeContext(RHICmdList, View);
+	FRenderingCompositePassContext CompositeContext(RHICmdList, View);	
 	FPostprocessContext Context( CompositeContext.Graph, View );
 
 	FSceneViewState* ViewState = (FSceneViewState*)Context.View.State;
 
-	FRenderingCompositePass* SceneColorInput = Context.Graph.RegisterPass( new FRCPassPostProcessInput( GSceneRenderTargets.GetSceneColor() ) );
+	FRenderingCompositePass* SceneColorInput = Context.Graph.RegisterPass( new FRCPassPostProcessInput( FSceneRenderTargets::Get(RHICmdList).GetSceneColor() ) );
 	FRenderingCompositePass* HZBInput = Context.Graph.RegisterPass( new FRCPassPostProcessInput( View.HZB ) );
 
 	bool bPrevFrame = 0;
diff --git a/Engine/Source/Runtime/Renderer/Private/PostProcess/VisualizeTexture.cpp b/Engine/Source/Runtime/Renderer/Private/PostProcess/VisualizeTexture.cpp
--- a/Engine/Source/Runtime/Renderer/Private/PostProcess/VisualizeTexture.cpp
+++ b/Engine/Source/Runtime/Renderer/Private/PostProcess/VisualizeTexture.cpp
@@ -202,19 +202,20 @@ template<uint32 TextureType> void VisualizeTextureForTextureType(FRHICommandList
 
 	SetGlobalBoundShaderState(RHICmdList, FeatureLevel, BoundShaderState, GFilterVertexDeclaration.VertexDeclarationRHI, *VertexShader, *PixelShader);
 	PixelShader->SetParameters(RHICmdList, Data);
+	FSceneRenderTargets& SceneContext = FSceneRenderTargets::Get(RHICmdList);
 
 	DrawRectangle(
 		RHICmdList,
 		// XY
 		0, 0,
 		// SizeXY
-		GSceneRenderTargets.GetBufferSizeXY().X, GSceneRenderTargets.GetBufferSizeXY().Y,
+		SceneContext.GetBufferSizeXY().X, SceneContext.GetBufferSizeXY().Y,
 		// UV
 		Data.Tex00.X, Data.Tex00.Y,
 		// SizeUV
 		Data.Tex11.X - Data.Tex00.X, Data.Tex11.Y - Data.Tex00.Y,
 		// TargetSize
-		GSceneRenderTargets.GetBufferSizeXY(),
+		SceneContext.GetBufferSizeXY(),
 		// TextureSize
 		FIntPoint(1, 1),
 		*VertexShader,
@@ -361,7 +362,7 @@ void FVisualizeTexture::GenerateContent(FRHICommandListImmediate& RHICmdList, co
 	RHICmdList.SetRasterizerState(TStaticRasterizerState<>::GetRHI());
 	RHICmdList.SetDepthStencilState(TStaticDepthStencilState<false, CF_Always>::GetRHI());
 
-	FIntPoint RTExtent = GSceneRenderTargets.GetBufferSizeXY();
+	FIntPoint RTExtent = FSceneRenderTargets::Get(RHICmdList).GetBufferSizeXY();
 
 	FVector2D Tex00 = FVector2D(0, 0);
 	FVector2D Tex11 = FVector2D(1, 1);
@@ -616,7 +617,7 @@ void FVisualizeTexture::PresentContent(FRHICommandListImmediate& RHICmdList, con
 		Canvas.DrawShadowedString( X + 10, Y += YStep, *Line, GetStatsFont(), FLinearColor(1, 1, 1));
 	}
 	{
-		FString Line = FString::Printf(TEXT("  BufferSize:(%d,%d)"), GSceneRenderTargets.GetBufferSizeXY().X, GSceneRenderTargets.GetBufferSizeXY().Y);
+		FString Line = FString::Printf(TEXT("  BufferSize:(%d,%d)"), FSceneRenderTargets::Get(RHICmdList).GetBufferSizeXY().X, FSceneRenderTargets::Get(RHICmdList).GetBufferSizeXY().Y);
 		Canvas.DrawShadowedString( X + 10, Y += YStep, *Line, GetStatsFont(), FLinearColor(1, 1, 1));
 	}
 
diff --git a/Engine/Source/Runtime/Renderer/Private/ReflectionEnvironment.cpp b/Engine/Source/Runtime/Renderer/Private/ReflectionEnvironment.cpp
--- a/Engine/Source/Runtime/Renderer/Private/ReflectionEnvironment.cpp
+++ b/Engine/Source/Runtime/Renderer/Private/ReflectionEnvironment.cpp
@@ -279,7 +279,7 @@ public:
 
 		SetTextureParameter(RHICmdList, ShaderRHI, ScreenSpaceReflections, SSRTexture );
 
-		SetTextureParameter(RHICmdList, ShaderRHI, InSceneColor, GSceneRenderTargets.GetSceneColor()->GetRenderTargetItem().ShaderResourceTexture );
+		SetTextureParameter(RHICmdList, ShaderRHI, InSceneColor, FSceneRenderTargets::Get(RHICmdList).GetSceneColor()->GetRenderTargetItem().ShaderResourceTexture );
 		OutSceneColor.SetTexture(RHICmdList, ShaderRHI, NULL, OutSceneColorUAV);
 
 		SetShaderValue(RHICmdList, ShaderRHI, ViewDimensionsParameter, View.ViewRect);
@@ -637,7 +637,8 @@ IMPLEMENT_SHADER_TYPE(template<>,TStandardDeferredReflectionPS<false>,TEXT("Refl
 
 void FDeferredShadingSceneRenderer::RenderReflectionCaptureSpecularBounceForAllViews(FRHICommandListImmediate& RHICmdList)
 {
-	GSceneRenderTargets.BeginRenderingSceneColor(RHICmdList, ESimpleRenderTargetMode::EUninitializedColorExistingDepth, FExclusiveDepthStencil::DepthRead_StencilWrite);
+	FSceneRenderTargets& SceneContext = FSceneRenderTargets::Get(RHICmdList);
+	SceneContext.BeginRenderingSceneColor(RHICmdList, ESimpleRenderTargetMode::EUninitializedColorExistingDepth, FExclusiveDepthStencil::DepthRead_StencilWrite);
 	RHICmdList.SetRasterizerState(TStaticRasterizerState< FM_Solid, CM_None >::GetRHI());
 	RHICmdList.SetDepthStencilState(TStaticDepthStencilState< false, CF_Always >::GetRHI());
 	RHICmdList.SetBlendState(TStaticBlendState< CW_RGB, BO_Add, BF_One, BF_One >::GetRHI());
@@ -666,12 +667,12 @@ void FDeferredShadingSceneRenderer::RenderReflectionCaptureSpecularBounceForAllV
 			0, 0,
 			View.ViewRect.Width(), View.ViewRect.Height(),
 			FIntPoint(View.ViewRect.Width(), View.ViewRect.Height()),
-			GSceneRenderTargets.GetBufferSizeXY(),
+			SceneContext.GetBufferSizeXY(),
 			*VertexShader,
 			EDRF_UseTriangleOptimization);
 	}
 
-	GSceneRenderTargets.FinishRenderingSceneColor(RHICmdList);
+	SceneContext.FinishRenderingSceneColor(RHICmdList);
 }
 
 bool FDeferredShadingSceneRenderer::ShouldDoReflectionEnvironment() const
@@ -834,15 +835,16 @@ FReflectionEnvironmentTiledDeferredCS* SelectReflectionEnvironmentTiledDeferredC
 
 void FDeferredShadingSceneRenderer::RenderTiledDeferredImageBasedReflections(FRHICommandListImmediate& RHICmdList, const TRefCountPtr<IPooledRenderTarget>& DynamicBentNormalAO)
 {
+	FSceneRenderTargets& SceneContext = FSceneRenderTargets::Get(RHICmdList);
 	static const auto AllowStaticLightingVar = IConsoleManager::Get().FindTConsoleVariableDataInt(TEXT("r.AllowStaticLighting"));
 	const bool bUseLightmaps = (AllowStaticLightingVar->GetValueOnRenderThread() == 1) && (CVarDiffuseFromCaptures.GetValueOnRenderThread() == 0);
 	const bool bHalfRes = CVarHalfResReflections.GetValueOnRenderThread() != 0;
 
 	TRefCountPtr<IPooledRenderTarget> NewSceneColor;
 	{
-		GSceneRenderTargets.ResolveSceneColor(RHICmdList, FResolveRect(0, 0, ViewFamily.FamilySizeX, ViewFamily.FamilySizeY));
+		SceneContext.ResolveSceneColor(RHICmdList, FResolveRect(0, 0, ViewFamily.FamilySizeX, ViewFamily.FamilySizeY));
 
-		FPooledRenderTargetDesc Desc = GSceneRenderTargets.GetSceneColor()->GetDesc();
+		FPooledRenderTargetDesc Desc = SceneContext.GetSceneColor()->GetDesc();
 		Desc.TargetableFlags |= TexCreate_UAV;
 		Desc.TargetableFlags |= TexCreate_NoFastClear;
 
@@ -899,8 +901,8 @@ void FDeferredShadingSceneRenderer::RenderTiledDeferredImageBasedReflections(FRH
 		}
 	}
 
-	GSceneRenderTargets.SetSceneColor(NewSceneColor);
-	check(GSceneRenderTargets.GetSceneColor());
+	SceneContext.SetSceneColor(NewSceneColor);
+	check(SceneContext.GetSceneColor());
 }
 
 void FDeferredShadingSceneRenderer::RenderStandardDeferredImageBasedReflections(FRHICommandListImmediate& RHICmdList, bool bReflectionEnv, const TRefCountPtr<IPooledRenderTarget>& DynamicBentNormalAO)
@@ -964,6 +966,7 @@ void FDeferredShadingSceneRenderer::RenderStandardDeferredImageBasedReflections(
 
 		SortData.Sort();
 	}
+	FSceneRenderTargets& SceneContext = FSceneRenderTargets::Get(RHICmdList);
 
 	// Use standard deferred shading to composite reflection capture contribution
 	for (int32 ViewIndex = 0, Num = Views.Num(); ViewIndex < Num; ViewIndex++)
@@ -998,7 +1001,7 @@ void FDeferredShadingSceneRenderer::RenderStandardDeferredImageBasedReflections(
 		}
 
 	    /* Light Accumulation moved to SceneRenderTargets */
-	    TRefCountPtr<IPooledRenderTarget> LightAccumulation = GSceneRenderTargets.LightAccumulation;
+	    TRefCountPtr<IPooledRenderTarget> LightAccumulation = SceneContext.LightAccumulation;
 
 		if (!LightAccumulation)
 		{
@@ -1076,7 +1079,7 @@ void FDeferredShadingSceneRenderer::RenderStandardDeferredImageBasedReflections(
 			// Apply reflections to screen
 			SCOPED_DRAW_EVENT(RHICmdList, ReflectionApply);
 
-			GSceneRenderTargets.BeginRenderingSceneColor(RHICmdList, ESimpleRenderTargetMode::EUninitializedColorExistingDepth, FExclusiveDepthStencil::DepthRead_StencilWrite);
+			SceneContext.BeginRenderingSceneColor(RHICmdList, ESimpleRenderTargetMode::EUninitializedColorExistingDepth, FExclusiveDepthStencil::DepthRead_StencilWrite);
 
 			RHICmdList.SetViewport(View.ViewRect.Min.X, View.ViewRect.Min.Y, 0.0f, View.ViewRect.Max.X, View.ViewRect.Max.Y, 1.0f);
 			RHICmdList.SetRasterizerState(TStaticRasterizerState<FM_Solid, CM_None>::GetRHI());
@@ -1129,10 +1132,10 @@ void FDeferredShadingSceneRenderer::RenderStandardDeferredImageBasedReflections(
 				View.ViewRect.Min.X, View.ViewRect.Min.Y,
 				View.ViewRect.Width(), View.ViewRect.Height(),
 				FIntPoint(View.ViewRect.Width(), View.ViewRect.Height()),
-				GSceneRenderTargets.GetBufferSizeXY(),
+				SceneContext.GetBufferSizeXY(),
 				*VertexShader);
 
-			GSceneRenderTargets.FinishRenderingSceneColor(RHICmdList);
+			SceneContext.FinishRenderingSceneColor(RHICmdList);
 		}
 	}
 }
diff --git a/Engine/Source/Runtime/Renderer/Private/ReflectionEnvironmentCapture.cpp b/Engine/Source/Runtime/Renderer/Private/ReflectionEnvironmentCapture.cpp
--- a/Engine/Source/Runtime/Renderer/Private/ReflectionEnvironmentCapture.cpp
+++ b/Engine/Source/Runtime/Renderer/Private/ReflectionEnvironmentCapture.cpp
@@ -50,7 +50,7 @@ FAutoConsoleCommandWithWorld CaptureConsoleCommand(
 
 
 /** Encapsulates render target picking logic for cubemap mip generation. */
-FSceneRenderTargetItem& GetEffectiveRenderTarget(bool bDownsamplePass, int32 TargetMipIndex)
+FSceneRenderTargetItem& GetEffectiveRenderTarget(FSceneRenderTargets& SceneContext, bool bDownsamplePass, int32 TargetMipIndex)
 {
 	int32 ScratchTextureIndex = TargetMipIndex % 2;
 
@@ -59,11 +59,11 @@ FSceneRenderTargetItem& GetEffectiveRenderTarget(bool bDownsamplePass, int32 Tar
 		ScratchTextureIndex = 1 - ScratchTextureIndex;
 	}
 
-	return GSceneRenderTargets.ReflectionColorScratchCubemap[ScratchTextureIndex]->GetRenderTargetItem();
+	return SceneContext.ReflectionColorScratchCubemap[ScratchTextureIndex]->GetRenderTargetItem();
 }
 
 /** Encapsulates source texture picking logic for cubemap mip generation. */
-FSceneRenderTargetItem& GetEffectiveSourceTexture(bool bDownsamplePass, int32 TargetMipIndex)
+FSceneRenderTargetItem& GetEffectiveSourceTexture(FSceneRenderTargets& SceneContext, bool bDownsamplePass, int32 TargetMipIndex)
 {
 	int32 ScratchTextureIndex = TargetMipIndex % 2;
 
@@ -72,13 +72,14 @@ FSceneRenderTargetItem& GetEffectiveSourceTexture(bool bDownsamplePass, int32 Ta
 		ScratchTextureIndex = 1 - ScratchTextureIndex;
 	}
 
-	return GSceneRenderTargets.ReflectionColorScratchCubemap[ScratchTextureIndex]->GetRenderTargetItem();
+	return SceneContext.ReflectionColorScratchCubemap[ScratchTextureIndex]->GetRenderTargetItem();
 }
 
 void FullyResolveReflectionScratchCubes(FRHICommandListImmediate& RHICmdList)
 {
-	FTextureRHIRef& Scratch0 = GSceneRenderTargets.ReflectionColorScratchCubemap[0]->GetRenderTargetItem().TargetableTexture;
-	FTextureRHIRef& Scratch1 = GSceneRenderTargets.ReflectionColorScratchCubemap[1]->GetRenderTargetItem().TargetableTexture;
+	FSceneRenderTargets& SceneContext = FSceneRenderTargets::Get(RHICmdList);
+	FTextureRHIRef& Scratch0 = SceneContext.ReflectionColorScratchCubemap[0]->GetRenderTargetItem().TargetableTexture;
+	FTextureRHIRef& Scratch1 = SceneContext.ReflectionColorScratchCubemap[1]->GetRenderTargetItem().TargetableTexture;
 	FResolveParams ResolveParams(FResolveRect(), CubeFace_PosX, -1, -1, -1);
 	RHICmdList.CopyToResolveTarget(Scratch0, Scratch0, true, ResolveParams);
 	RHICmdList.CopyToResolveTarget(Scratch1, Scratch1, true, ResolveParams);  
@@ -172,7 +173,7 @@ public:
 			AverageBrightnessTexture, 
 			AverageBrightnessSampler, 
 			TStaticSamplerState<SF_Bilinear, AM_Clamp, AM_Clamp, AM_Clamp>::GetRHI(), 
-			GSceneRenderTargets.GetReflectionBrightnessTarget()->GetRenderTargetItem().ShaderResourceTexture);
+			FSceneRenderTargets::Get(RHICmdList).GetReflectionBrightnessTarget()->GetRenderTargetItem().ShaderResourceTexture);
 	}
 
 	virtual bool Serialize(FArchive& Ar) override
@@ -247,7 +248,7 @@ public:
 		const int32 EffectiveTopMipSize = GReflectionCaptureSize;
 		const int32 NumMips = FMath::CeilLogTwo(EffectiveTopMipSize) + 1;
 		// Read from the smallest mip that was downsampled to
-		FSceneRenderTargetItem& Cubemap = GetEffectiveRenderTarget(true, NumMips - 1);
+		FSceneRenderTargetItem& Cubemap = GetEffectiveRenderTarget(FSceneRenderTargets::Get(RHICmdList), true, NumMips - 1);
 
 		if (Cubemap.IsValid())
 		{
@@ -284,7 +285,7 @@ IMPLEMENT_SHADER_TYPE(,FComputeBrightnessPS,TEXT("ReflectionEnvironmentShaders")
 /** Computes the average brightness of the given reflection capture and stores it in the scene. */
 void ComputeAverageBrightness(FRHICommandList& RHICmdList, ERHIFeatureLevel::Type FeatureLevel)
 {
-	FTextureRHIRef& BrightnessTarget = GSceneRenderTargets.GetReflectionBrightnessTarget()->GetRenderTargetItem().TargetableTexture;
+	FTextureRHIRef& BrightnessTarget = FSceneRenderTargets::Get(RHICmdList).GetReflectionBrightnessTarget()->GetRenderTargetItem().TargetableTexture;
 	SetRenderTarget(RHICmdList, BrightnessTarget, NULL);
 	RHICmdList.SetRasterizerState(TStaticRasterizerState<FM_Solid, CM_None>::GetRHI());
 	RHICmdList.SetDepthStencilState(TStaticDepthStencilState<false, CF_Always>::GetRHI());
@@ -330,6 +331,8 @@ void FilterReflectionEnvironment(FRHICommandListImmediate& RHICmdList, ERHIFeatu
 
 	auto ShaderMap = GetGlobalShaderMap(FeatureLevel);
 
+	FSceneRenderTargets& SceneContext = FSceneRenderTargets::Get(RHICmdList);
+
 	{	
 		SCOPED_DRAW_EVENT(RHICmdList, DownsampleCubeMips);
 		// Downsample all the mips, each one reads from the mip above it
@@ -339,8 +342,8 @@ void FilterReflectionEnvironment(FRHICommandListImmediate& RHICmdList, ERHIFeatu
 			const int32 SourceMipIndex = FMath::Max(MipIndex - 1, 0);
 			const int32 MipSize = 1 << (NumMips - MipIndex - 1);
 
-			FSceneRenderTargetItem& EffectiveRT = GetEffectiveRenderTarget(true, MipIndex);
-			FSceneRenderTargetItem& EffectiveSource = GetEffectiveSourceTexture(true, MipIndex);
+			FSceneRenderTargetItem& EffectiveRT = GetEffectiveRenderTarget(SceneContext, true, MipIndex);
+			FSceneRenderTargetItem& EffectiveSource = GetEffectiveSourceTexture(SceneContext, true, MipIndex);
 			check(EffectiveRT.TargetableTexture != EffectiveSource.ShaderResourceTexture);
 			
 			for (int32 CubeFace = 0; CubeFace < CubeFace_MAX; CubeFace++)
@@ -396,8 +399,8 @@ void FilterReflectionEnvironment(FRHICommandListImmediate& RHICmdList, ERHIFeatu
 		for (int32 MipIndex = 0; MipIndex < NumMips; MipIndex++)
 		{
 			SCOPED_DRAW_EVENT(RHICmdList, FilterCubeMip);
-			FSceneRenderTargetItem& EffectiveRT = GetEffectiveRenderTarget(false, MipIndex);
-			FSceneRenderTargetItem& EffectiveSource = GetEffectiveSourceTexture(false, MipIndex);
+			FSceneRenderTargetItem& EffectiveRT = GetEffectiveRenderTarget(SceneContext, false, MipIndex);
+			FSceneRenderTargetItem& EffectiveSource = GetEffectiveSourceTexture(SceneContext, false, MipIndex);
 			check(EffectiveRT.TargetableTexture != EffectiveSource.ShaderResourceTexture);
 			const int32 MipSize = 1 << (NumMips - MipIndex - 1);
 
@@ -519,7 +522,7 @@ public:
 			InTexture, 
 			InTextureSampler, 
 			TStaticSamplerState<SF_Bilinear, AM_Clamp, AM_Clamp, AM_Clamp>::GetRHI(), 
-			GSceneRenderTargets.GetSceneColor()->GetRenderTargetItem().ShaderResourceTexture);		
+			FSceneRenderTargets::Get(RHICmdList).GetSceneColor()->GetRenderTargetItem().ShaderResourceTexture);		
 
 		FVector SkyLightParametersValue = FVector::ZeroVector;
 		FScene* Scene = (FScene*)View.Family->Scene;
@@ -663,12 +666,13 @@ int32 FindOrAllocateCubemapIndex(FScene* Scene, const UReflectionCaptureComponen
 
 void ClearScratchCubemaps(FRHICommandList& RHICmdList)
 {
-	GSceneRenderTargets.AllocateReflectionTargets();
+	FSceneRenderTargets& SceneContext = FSceneRenderTargets::Get(RHICmdList);
+	SceneContext.AllocateReflectionTargets();
 	// Clear scratch render targets to a consistent but noticeable value
 	// This makes debugging capture issues much easier, otherwise the random contents from previous captures is shown
 
 	const int32 NumMips = FMath::CeilLogTwo(GReflectionCaptureSize) + 1;
-	FSceneRenderTargetItem& RT0 = GSceneRenderTargets.ReflectionColorScratchCubemap[0]->GetRenderTargetItem();
+	FSceneRenderTargetItem& RT0 = SceneContext.ReflectionColorScratchCubemap[0]->GetRenderTargetItem();
 
 	for (int32 MipIndex = 0; MipIndex < NumMips; MipIndex++)
 	{
@@ -679,7 +683,7 @@ void ClearScratchCubemaps(FRHICommandList& RHICmdList)
 		}
 	}
 
-	FSceneRenderTargetItem& RT1 = GSceneRenderTargets.ReflectionColorScratchCubemap[1]->GetRenderTargetItem();
+	FSceneRenderTargetItem& RT1 = SceneContext.ReflectionColorScratchCubemap[1]->GetRenderTargetItem();
 
 	for (int32 MipIndex = 0; MipIndex < NumMips; MipIndex++)
 	{
@@ -719,9 +723,10 @@ void CaptureSceneToScratchCubemap(FRHICommandListImmediate& RHICmdList, FSceneRe
 #endif
 
 		auto ShaderMap = GetGlobalShaderMap(FeatureLevel);
+		FSceneRenderTargets& SceneContext = FSceneRenderTargets::Get(RHICmdList);
 
 		const int32 EffectiveSize = GReflectionCaptureSize;
-		FSceneRenderTargetItem& EffectiveColorRT =  GSceneRenderTargets.ReflectionColorScratchCubemap[0]->GetRenderTargetItem();
+		FSceneRenderTargetItem& EffectiveColorRT =  SceneContext.ReflectionColorScratchCubemap[0]->GetRenderTargetItem();
 
 		{
 			SCOPED_DRAW_EVENT(RHICmdList, CubeMapCopyScene);
@@ -751,14 +756,14 @@ void CaptureSceneToScratchCubemap(FRHICommandListImmediate& RHICmdList, FSceneRe
 				ViewRect.Min.X, ViewRect.Min.Y, 
 				ViewRect.Width() * GSupersampleCaptureFactor, ViewRect.Height() * GSupersampleCaptureFactor,
 				FIntPoint(ViewRect.Width(), ViewRect.Height()),
-				GSceneRenderTargets.GetBufferSizeXY(),
+				SceneContext.GetBufferSizeXY(),
 				*VertexShader);
 
 			RHICmdList.CopyToResolveTarget(EffectiveColorRT.TargetableTexture, EffectiveColorRT.ShaderResourceTexture, true, FResolveParams(FResolveRect(), CubeFace));			
 		}
 	}
 
-	delete SceneRenderer;
+	FSceneRenderer::WaitForTasksClearSnapshotsAndDeleteSceneRenderer(RHICmdList, SceneRenderer);
 }
 
 void CopyCubemapToScratchCubemap(FRHICommandList& RHICmdList, ERHIFeatureLevel::Type FeatureLevel, UTextureCube* SourceCubemap, bool bIsSkyLight, bool bLowerHemisphereIsBlack)
@@ -766,7 +771,7 @@ void CopyCubemapToScratchCubemap(FRHICommandList& RHICmdList, ERHIFeatureLevel::
 	check(SourceCubemap);
 	
 	const int32 EffectiveSize = GReflectionCaptureSize;
-	FSceneRenderTargetItem& EffectiveColorRT =  GSceneRenderTargets.ReflectionColorScratchCubemap[0]->GetRenderTargetItem();
+	FSceneRenderTargetItem& EffectiveColorRT =  FSceneRenderTargets::Get(RHICmdList).ReflectionColorScratchCubemap[0]->GetRenderTargetItem();
 
 	for (uint32 CubeFace = 0; CubeFace < CubeFace_MAX; CubeFace++)
 	{
@@ -1240,12 +1245,13 @@ void CopyToSceneArray(FRHICommandListImmediate& RHICmdList, FScene* Scene, FRefl
 	const int32 NumMips = FMath::CeilLogTwo(EffectiveTopMipSize) + 1;
 
 	const int32 CaptureIndex = FindOrAllocateCubemapIndex(Scene, ReflectionProxy->Component);
+	FSceneRenderTargets& SceneContext = FSceneRenderTargets::Get(RHICmdList);
 
 	// GPU copy back to the scene's texture array, which is not a render target
 	for (int32 MipIndex = 0; MipIndex < NumMips; MipIndex++)
 	{
 		// The source for this copy is the dest from the filtering pass
-		FSceneRenderTargetItem& EffectiveSource = GetEffectiveRenderTarget(false, MipIndex);
+		FSceneRenderTargetItem& EffectiveSource = GetEffectiveRenderTarget(SceneContext, false, MipIndex);
 		FSceneRenderTargetItem& EffectiveDest = Scene->ReflectionSceneData.CubemapArray.GetRenderTarget();
 
 		for (int32 CubeFace = 0; CubeFace < CubeFace_MAX; CubeFace++)
@@ -1260,12 +1266,13 @@ void CopyToComponentTexture(FRHICommandList& RHICmdList, FScene* Scene, FReflect
 	check(ReflectionProxy->SM4FullHDRCubemap);
 	const int32 EffectiveTopMipSize = GReflectionCaptureSize;
 	const int32 NumMips = FMath::CeilLogTwo(EffectiveTopMipSize) + 1;
+	FSceneRenderTargets& SceneContext = FSceneRenderTargets::Get(RHICmdList);
 
 	// GPU copy back to the component's cubemap texture, which is not a render target
 	for (int32 MipIndex = 0; MipIndex < NumMips; MipIndex++)
 	{
 		// The source for this copy is the dest from the filtering pass
-		FSceneRenderTargetItem& EffectiveSource = GetEffectiveRenderTarget(false, MipIndex);
+		FSceneRenderTargetItem& EffectiveSource = GetEffectiveRenderTarget(SceneContext, false, MipIndex);
 
 		for (int32 CubeFace = 0; CubeFace < CubeFace_MAX; CubeFace++)
 		{
@@ -1349,12 +1356,13 @@ void CopyToSkyTexture(FRHICommandList& RHICmdList, FScene* Scene, FTexture* Proc
 {
 	const int32 EffectiveTopMipSize = ProcessedTexture->GetSizeX();
 	const int32 NumMips = FMath::CeilLogTwo(EffectiveTopMipSize) + 1;
+	FSceneRenderTargets& SceneContext = FSceneRenderTargets::Get(RHICmdList);
 
 	// GPU copy back to the skylight's texture, which is not a render target
 	for (int32 MipIndex = 0; MipIndex < NumMips; MipIndex++)
 	{
 		// The source for this copy is the dest from the filtering pass
-		FSceneRenderTargetItem& EffectiveSource = GetEffectiveRenderTarget(false, MipIndex);
+		FSceneRenderTargetItem& EffectiveSource = GetEffectiveRenderTarget(SceneContext, false, MipIndex);
 
 		for (int32 CubeFace = 0; CubeFace < CubeFace_MAX; CubeFace++)
 		{
diff --git a/Engine/Source/Runtime/Renderer/Private/ReflectionEnvironmentDiffuseIrradiance.cpp b/Engine/Source/Runtime/Renderer/Private/ReflectionEnvironmentDiffuseIrradiance.cpp
--- a/Engine/Source/Runtime/Renderer/Private/ReflectionEnvironmentDiffuseIrradiance.cpp
+++ b/Engine/Source/Runtime/Renderer/Private/ReflectionEnvironmentDiffuseIrradiance.cpp
@@ -16,18 +16,18 @@
 
 extern int32 GDiffuseIrradianceCubemapSize;
 
-FSceneRenderTargetItem& GetEffectiveDiffuseIrradianceRenderTarget(int32 TargetMipIndex)
+FSceneRenderTargetItem& GetEffectiveDiffuseIrradianceRenderTarget(FSceneRenderTargets& SceneContext, int32 TargetMipIndex)
 {
 	const int32 ScratchTextureIndex = TargetMipIndex % 2;
 
-	return GSceneRenderTargets.DiffuseIrradianceScratchCubemap[ScratchTextureIndex]->GetRenderTargetItem();
+	return SceneContext.DiffuseIrradianceScratchCubemap[ScratchTextureIndex]->GetRenderTargetItem();
 }
 
-FSceneRenderTargetItem& GetEffectiveDiffuseIrradianceSourceTexture(int32 TargetMipIndex)
+FSceneRenderTargetItem& GetEffectiveDiffuseIrradianceSourceTexture(FSceneRenderTargets& SceneContext, int32 TargetMipIndex)
 {
 	const int32 ScratchTextureIndex = 1 - TargetMipIndex % 2;
 
-	return GSceneRenderTargets.DiffuseIrradianceScratchCubemap[ScratchTextureIndex]->GetRenderTargetItem();
+	return SceneContext.DiffuseIrradianceScratchCubemap[ScratchTextureIndex]->GetRenderTargetItem();
 }
 
 /** Pixel shader used for copying to diffuse irradiance texture. */
@@ -238,6 +238,8 @@ FGlobalBoundShaderState AccumulateCubeFacesBoundShaderState;
 void ComputeDiffuseIrradiance(FRHICommandListImmediate& RHICmdList, ERHIFeatureLevel::Type FeatureLevel, FTextureRHIRef LightingSource, int32 LightingSourceMipIndex, FSHVectorRGB3* OutIrradianceEnvironmentMap)
 {
 	auto ShaderMap = GetGlobalShaderMap(FeatureLevel);
+	FSceneRenderTargets& SceneContext = FSceneRenderTargets::Get(RHICmdList);
+
 
 	for (int32 CoefficientIndex = 0; CoefficientIndex < FSHVector3::MaxSHBasis; CoefficientIndex++)
 	{
@@ -245,7 +247,7 @@ void ComputeDiffuseIrradiance(FRHICommandListImmediate& RHICmdList, ERHIFeatureL
 		{
 			const int32 MipIndex = 0;
 			const int32 MipSize = GDiffuseIrradianceCubemapSize;
-			FSceneRenderTargetItem& EffectiveRT = GetEffectiveDiffuseIrradianceRenderTarget(MipIndex);			
+			FSceneRenderTargetItem& EffectiveRT = GetEffectiveDiffuseIrradianceRenderTarget(SceneContext, MipIndex);			
 			
 			for (int32 CubeFace = 0; CubeFace < CubeFace_MAX; CubeFace++)
 			{
@@ -287,8 +289,8 @@ void ComputeDiffuseIrradiance(FRHICommandListImmediate& RHICmdList, ERHIFeatureL
 				const int32 SourceMipIndex = FMath::Max(MipIndex - 1, 0);
 				const int32 MipSize = 1 << (NumMips - MipIndex - 1);
 
-				FSceneRenderTargetItem& EffectiveRT = GetEffectiveDiffuseIrradianceRenderTarget(MipIndex);
-				FSceneRenderTargetItem& EffectiveSource = GetEffectiveDiffuseIrradianceSourceTexture(MipIndex);
+				FSceneRenderTargetItem& EffectiveRT = GetEffectiveDiffuseIrradianceRenderTarget(SceneContext, MipIndex);
+				FSceneRenderTargetItem& EffectiveSource = GetEffectiveDiffuseIrradianceSourceTexture(SceneContext, MipIndex);
 				check(EffectiveRT.TargetableTexture != EffectiveSource.ShaderResourceTexture);
 				
 				for (int32 CubeFace = 0; CubeFace < CubeFace_MAX; CubeFace++)
@@ -324,8 +326,8 @@ void ComputeDiffuseIrradiance(FRHICommandListImmediate& RHICmdList, ERHIFeatureL
 		}
 
 		{
-			// Gather the cubemap face results and normalize, copy this coefficient to GSceneRenderTargets.SkySHIrradianceMap
-			FSceneRenderTargetItem& EffectiveRT = GSceneRenderTargets.SkySHIrradianceMap->GetRenderTargetItem();
+			// Gather the cubemap face results and normalize, copy this coefficient to FSceneRenderTargets::Get(RHICmdList).SkySHIrradianceMap
+			FSceneRenderTargetItem& EffectiveRT = FSceneRenderTargets::Get(RHICmdList).SkySHIrradianceMap->GetRenderTargetItem();
 
 			//load/store actions so we don't lose results as we render one pixel at a time on tile renderers.
 			FRHIRenderTargetView RTV(EffectiveRT.TargetableTexture, 0, -1, ERenderTargetLoadAction::ELoad, ERenderTargetStoreAction::EStore);
@@ -343,7 +345,7 @@ void ComputeDiffuseIrradiance(FRHICommandListImmediate& RHICmdList, ERHIFeatureL
 
 			const int32 SourceMipIndex = NumMips - 1;
 			const int32 MipSize = 1;
-			FSceneRenderTargetItem& EffectiveSource = GetEffectiveDiffuseIrradianceRenderTarget(SourceMipIndex);
+			FSceneRenderTargetItem& EffectiveSource = GetEffectiveDiffuseIrradianceRenderTarget(SceneContext, SourceMipIndex);
 			PixelShader->SetParameters(RHICmdList, SourceMipIndex, EffectiveSource.ShaderResourceTexture);
 
 			DrawRectangle( 
@@ -362,7 +364,7 @@ void ComputeDiffuseIrradiance(FRHICommandListImmediate& RHICmdList, ERHIFeatureL
 
 	{
 		// Read back the completed SH environment map
-		FSceneRenderTargetItem& EffectiveRT = GSceneRenderTargets.SkySHIrradianceMap->GetRenderTargetItem();
+		FSceneRenderTargetItem& EffectiveRT = FSceneRenderTargets::Get(RHICmdList).SkySHIrradianceMap->GetRenderTargetItem();
 		check(EffectiveRT.ShaderResourceTexture->GetFormat() == PF_FloatRGBA);
 
 		TArray<FFloat16Color> SurfaceData;
diff --git a/Engine/Source/Runtime/Renderer/Private/Renderer.cpp b/Engine/Source/Runtime/Renderer/Private/Renderer.cpp
--- a/Engine/Source/Runtime/Renderer/Private/Renderer.cpp
+++ b/Engine/Source/Runtime/Renderer/Private/Renderer.cpp
@@ -23,13 +23,13 @@ IMPLEMENT_MODULE(FRendererModule, Renderer);
 void FRendererModule::ReallocateSceneRenderTargets()
 {
 	FLightPrimitiveInteraction::InitializeMemoryPool();
-	GSceneRenderTargets.UpdateRHI();
+	FSceneRenderTargets::Get_Todo_PassContext().UpdateRHI();
 }
 
 void FRendererModule::SceneRenderTargetsSetBufferSize(uint32 SizeX, uint32 SizeY)
 {
-	GSceneRenderTargets.SetBufferSize(SizeX, SizeY);
-	GSceneRenderTargets.UpdateRHI();
+	FSceneRenderTargets::Get_Todo_PassContext().SetBufferSize(SizeX, SizeY);
+	FSceneRenderTargets::Get_Todo_PassContext().UpdateRHI();
 }
 
 void FRendererModule::DrawTileMesh(FRHICommandListImmediate& RHICmdList, const FSceneView& SceneView, const FMeshBatch& Mesh, bool bIsHitTesting, const FHitProxyId& HitProxyId)
diff --git a/Engine/Source/Runtime/Renderer/Private/SceneCaptureRendering.cpp b/Engine/Source/Runtime/Renderer/Private/SceneCaptureRendering.cpp
--- a/Engine/Source/Runtime/Renderer/Private/SceneCaptureRendering.cpp
+++ b/Engine/Source/Runtime/Renderer/Private/SceneCaptureRendering.cpp
@@ -55,7 +55,7 @@ static void CopyCaptureToTarget(FRHICommandListImmediate& RHICmdList, const FRen
 			ViewRect.Min.X, ViewRect.Min.Y,
 			ViewRect.Width(), ViewRect.Height(),
 			TargetSize,
-			GSceneRenderTargets.GetBufferSizeXY(),
+			FSceneRenderTargets::Get(RHICmdList).GetBufferSizeXY(),
 			*VertexShader,
 			EDRF_UseTriangleOptimization);
 	}
@@ -148,13 +148,12 @@ static void UpdateSceneCaptureContent_RenderThread(FRHICommandListImmediate& RHI
 			// Copy the captured scene into the destination texture (only required on HDR or deferred as that implies post-processing)
 			SCOPED_DRAW_EVENT(RHICmdList, CaptureSceneColor);
 			FIntPoint TargetSize(UnconstrainedViewRect.Width(), UnconstrainedViewRect.Height());
-			CopyCaptureToTarget(RHICmdList, Target, TargetSize, View, ViewRect, GSceneRenderTargets.GetSceneColorTexture(), false);
+			CopyCaptureToTarget(RHICmdList, Target, TargetSize, View, ViewRect, FSceneRenderTargets::Get(RHICmdList).GetSceneColorTexture(), false);
 		}
 
 		RHICmdList.CopyToResolveTarget(TextureRenderTarget->GetRenderTargetTexture(), TextureRenderTarget->TextureRHI, false, ResolveParams);
 	}
-
-	delete SceneRenderer;
+	FSceneRenderer::WaitForTasksClearSnapshotsAndDeleteSceneRenderer(RHICmdList, SceneRenderer);
 }
 
 
diff --git a/Engine/Source/Runtime/Renderer/Private/SceneHitProxyRendering.cpp b/Engine/Source/Runtime/Renderer/Private/SceneHitProxyRendering.cpp
--- a/Engine/Source/Runtime/Renderer/Private/SceneHitProxyRendering.cpp
+++ b/Engine/Source/Runtime/Renderer/Private/SceneHitProxyRendering.cpp
@@ -332,14 +332,15 @@ TRefCountPtr<IPooledRenderTarget> InitHitProxyRender(FRHICommandListImmediate& R
 	// Initialize global system textures (pass-through if already initialized).
 	GSystemTextures.InitializeTextures(RHICmdList, FeatureLevel);
 
+	FSceneRenderTargets& SceneContext = FSceneRenderTargets::Get(RHICmdList);
 	// Allocate the maximum scene render target space for the current view family.
-	GSceneRenderTargets.Allocate(ViewFamily);
+	SceneContext.Allocate(ViewFamily);
 
 	TRefCountPtr<IPooledRenderTarget> HitProxyRT;
 
 	// Create a texture to store the resolved light attenuation values, and a render-targetable surface to hold the unresolved light attenuation values.
 	{
-		FPooledRenderTargetDesc Desc(FPooledRenderTargetDesc::Create2DDesc(GSceneRenderTargets.GetBufferSizeXY(), PF_B8G8R8A8, TexCreate_None, TexCreate_RenderTargetable, false));
+		FPooledRenderTargetDesc Desc(FPooledRenderTargetDesc::Create2DDesc(SceneContext.GetBufferSizeXY(), PF_B8G8R8A8, TexCreate_None, TexCreate_RenderTargetable, false));
 		Desc.Flags |= TexCreate_FastVRAM;
 		GRenderTargetPool.FindFreeElement(Desc, HitProxyRT, TEXT("HitProxy"));
 	}
@@ -350,7 +351,7 @@ TRefCountPtr<IPooledRenderTarget> InitHitProxyRender(FRHICommandListImmediate& R
 		return HitProxyRT;
 	}
 
-	SetRenderTarget(RHICmdList, HitProxyRT->GetRenderTargetItem().TargetableTexture, GSceneRenderTargets.GetSceneDepthSurface(), ESimpleRenderTargetMode::EExistingColorAndDepth, FExclusiveDepthStencil::DepthWrite_StencilWrite);
+	SetRenderTarget(RHICmdList, HitProxyRT->GetRenderTargetItem().TargetableTexture, SceneContext.GetSceneDepthSurface(), ESimpleRenderTargetMode::EExistingColorAndDepth, FExclusiveDepthStencil::DepthWrite_StencilWrite);
 
 	// Clear color for each view.
 	auto& Views = SceneRenderer->Views;
@@ -379,6 +380,7 @@ void RenderHitProxies(FRHICommandListImmediate& RHICmdList, const FSceneRenderer
 	RHICmdList.SetBlendState(TStaticBlendState<>::GetRHI());
 
 	const bool bNeedToSwitchVerticalAxis = RHINeedsToSwitchVerticalAxis(GShaderPlatformForFeatureLevel[SceneRenderer->FeatureLevel]);
+	FSceneRenderTargets& SceneContext = FSceneRenderTargets::Get(RHICmdList);
 
 	for (int32 ViewIndex = 0; ViewIndex < Views.Num(); ViewIndex++)
 	{
@@ -480,7 +482,7 @@ void RenderHitProxies(FRHICommandListImmediate& RHICmdList, const FSceneRenderer
 	{
 		const FViewInfo& View = Views[ViewIndex];
 
-		FIntPoint BufferSize = GSceneRenderTargets.GetBufferSizeXY();
+		FIntPoint BufferSize = SceneContext.GetBufferSizeXY();
 		float InvBufferSizeX = 1.0f / BufferSize.X;
 		float InvBufferSizeY = 1.0f / BufferSize.Y;
 
diff --git a/Engine/Source/Runtime/Renderer/Private/SceneOcclusion.cpp b/Engine/Source/Runtime/Renderer/Private/SceneOcclusion.cpp
--- a/Engine/Source/Runtime/Renderer/Private/SceneOcclusion.cpp
+++ b/Engine/Source/Runtime/Renderer/Private/SceneOcclusion.cpp
@@ -375,7 +375,7 @@ FRenderQueryRHIParamRef FOcclusionQueryBatcher::BatchPrimitive(const FVector& Bo
 		check(OcclusionQueryPool);
 		CurrentBatchOcclusionQuery = new(BatchOcclusionQueries) FOcclusionBatch;
 		CurrentBatchOcclusionQuery->Query = OcclusionQueryPool->AllocateQuery();
-		CurrentBatchOcclusionQuery->VertexAllocation = FGlobalDynamicVertexBuffer::Get().Allocate(MaxBatchedPrimitives * 8 * sizeof(FVector), true);
+		CurrentBatchOcclusionQuery->VertexAllocation = FGlobalDynamicVertexBuffer::Get().Allocate(MaxBatchedPrimitives * 8 * sizeof(FVector));
 		check(CurrentBatchOcclusionQuery->VertexAllocation.IsValid());
 		NumBatchedPrimitives = 0;
 	}
@@ -920,8 +920,9 @@ public:
 		const FPixelShaderRHIParamRef ShaderRHI = GetPixelShader();
 
 		FGlobalShader::SetParameters(RHICmdList, ShaderRHI, View );
+		FSceneRenderTargets& SceneContext = FSceneRenderTargets::Get(RHICmdList);
 		
-		const FIntPoint GBufferSize = GSceneRenderTargets.GetBufferSizeXY();
+		const FIntPoint GBufferSize = SceneContext.GetBufferSizeXY();
 		const FVector2D InvSize( 1.0f / float(GBufferSize.X), 1.0f / float(GBufferSize.Y) );
 		const FVector4 InputUvFactorAndOffset (
 			float(2 * View.HZBMipmap0Size.X) / float(GBufferSize.X),
@@ -1019,7 +1020,7 @@ void BuildHZB( FRHICommandListImmediate& RHICmdList, FViewInfo& View )
 			View.ViewRect.Min.X, View.ViewRect.Min.Y,
 			View.ViewRect.Width(), View.ViewRect.Height(),
 			HZBSize,
-			GSceneRenderTargets.GetBufferSizeXY(),
+			FSceneRenderTargets::Get(RHICmdList).GetBufferSizeXY(),
 			*VertexShader,
 			EDRF_UseTriangleOptimization);
 
@@ -1071,15 +1072,16 @@ void BuildHZB( FRHICommandListImmediate& RHICmdList, FViewInfo& View )
 void FDeferredShadingSceneRenderer::BeginOcclusionTests(FRHICommandListImmediate& RHICmdList, bool bRenderQueries, bool bRenderHZB)
 {
 	SCOPE_CYCLE_COUNTER(STAT_BeginOcclusionTestsTime);
-	const bool bUseDownsampledDepth = IsValidRef(GSceneRenderTargets.GetSmallDepthSurface()) && GSceneRenderTargets.UseDownsizedOcclusionQueries();	
+	FSceneRenderTargets& SceneContext = FSceneRenderTargets::Get(RHICmdList);
+	const bool bUseDownsampledDepth = IsValidRef(SceneContext.GetSmallDepthSurface()) && SceneContext.UseDownsizedOcclusionQueries();	
 
 	if (bUseDownsampledDepth)
 	{
-		SetRenderTarget(RHICmdList, NULL, GSceneRenderTargets.GetSmallDepthSurface(), ESimpleRenderTargetMode::EExistingColorAndDepth, FExclusiveDepthStencil::DepthRead_StencilWrite);
+		SetRenderTarget(RHICmdList, NULL, SceneContext.GetSmallDepthSurface(), ESimpleRenderTargetMode::EExistingColorAndDepth, FExclusiveDepthStencil::DepthRead_StencilWrite);
 	}
 	else
 	{
-		SetRenderTarget(RHICmdList, NULL, GSceneRenderTargets.GetSceneDepthSurface(), ESimpleRenderTargetMode::EExistingColorAndDepth, FExclusiveDepthStencil::DepthRead_StencilWrite);
+		SetRenderTarget(RHICmdList, NULL, SceneContext.GetSceneDepthSurface(), ESimpleRenderTargetMode::EExistingColorAndDepth, FExclusiveDepthStencil::DepthRead_StencilWrite);
 	}
 
 	if (bRenderQueries)
@@ -1094,10 +1096,10 @@ void FDeferredShadingSceneRenderer::BeginOcclusionTests(FRHICommandListImmediate
 
 			if (bUseDownsampledDepth)
 			{
-				const uint32 DownsampledX = FMath::TruncToInt(View.ViewRect.Min.X / GSceneRenderTargets.GetSmallColorDepthDownsampleFactor());
-				const uint32 DownsampledY = FMath::TruncToInt(View.ViewRect.Min.Y / GSceneRenderTargets.GetSmallColorDepthDownsampleFactor());
-				const uint32 DownsampledSizeX = FMath::TruncToInt(View.ViewRect.Width() / GSceneRenderTargets.GetSmallColorDepthDownsampleFactor());
-				const uint32 DownsampledSizeY = FMath::TruncToInt(View.ViewRect.Height() / GSceneRenderTargets.GetSmallColorDepthDownsampleFactor());
+				const uint32 DownsampledX = FMath::TruncToInt(View.ViewRect.Min.X / SceneContext.GetSmallColorDepthDownsampleFactor());
+				const uint32 DownsampledY = FMath::TruncToInt(View.ViewRect.Min.Y / SceneContext.GetSmallColorDepthDownsampleFactor());
+				const uint32 DownsampledSizeX = FMath::TruncToInt(View.ViewRect.Width() / SceneContext.GetSmallColorDepthDownsampleFactor());
+				const uint32 DownsampledSizeY = FMath::TruncToInt(View.ViewRect.Height() / SceneContext.GetSmallColorDepthDownsampleFactor());
 
 				// Setup the viewport for rendering to the downsampled depth buffer
 				RHICmdList.SetViewport(DownsampledX, DownsampledY, 0.0f, DownsampledX + DownsampledSizeX, DownsampledY + DownsampledSizeY, 1.0f);
@@ -1242,6 +1244,6 @@ void FDeferredShadingSceneRenderer::BeginOcclusionTests(FRHICommandListImmediate
 	if (bUseDownsampledDepth)
 	{
 		// Restore default render target
-		GSceneRenderTargets.BeginRenderingSceneColor(RHICmdList, ESimpleRenderTargetMode::EUninitializedColorExistingDepth, FExclusiveDepthStencil::DepthRead_StencilWrite);
+		SceneContext.BeginRenderingSceneColor(RHICmdList, ESimpleRenderTargetMode::EUninitializedColorExistingDepth, FExclusiveDepthStencil::DepthRead_StencilWrite);
 	}
 }
diff --git a/Engine/Source/Runtime/Renderer/Private/ScenePrivate.h b/Engine/Source/Runtime/Renderer/Private/ScenePrivate.h
--- a/Engine/Source/Runtime/Renderer/Private/ScenePrivate.h
+++ b/Engine/Source/Runtime/Renderer/Private/ScenePrivate.h
@@ -548,8 +548,6 @@ public:
 	// cache for stencil reads to a avoid reallocations of the SRV, Key is to detect if the object has changed
 	FTextureRHIRef SelectionOutlineCacheKey;
 	TRefCountPtr<FRHIShaderResourceView> SelectionOutlineCacheValue;
-	FTextureRHIRef CustomStencilSRVCacheKey;
-	TRefCountPtr<FRHIShaderResourceView> CustomStencilSRVCacheValue;
 
 	/** Distance field AO tile intersection GPU resources.  Last frame's state is not used, but they must be sized exactly to the view so stored here. */
 	class FTileIntersectionResources* AOTileIntersectionResources;
@@ -683,19 +681,23 @@ public:
 
 	TRefCountPtr<IPooledRenderTarget>& GetEyeAdaptation()
 	{
-		// Create the texture needed for EyeAdaptation
-		FPooledRenderTargetDesc Desc(FPooledRenderTargetDesc::Create2DDesc(FIntPoint(1, 1), PF_R32_FLOAT, TexCreate_None, TexCreate_RenderTargetable, false));
-		GRenderTargetPool.FindFreeElement(Desc, EyeAdaptationRT, TEXT("EyeAdaptation"));
-
+		if (!EyeAdaptationRT)
+		{
+			// Create the texture needed for EyeAdaptation
+			FPooledRenderTargetDesc Desc(FPooledRenderTargetDesc::Create2DDesc(FIntPoint(1, 1), PF_R32_FLOAT, TexCreate_None, TexCreate_RenderTargetable, false));
+			GRenderTargetPool.FindFreeElement(Desc, EyeAdaptationRT, TEXT("EyeAdaptation"));
+		}
 		return EyeAdaptationRT;
 	}
 
-	TRefCountPtr<IPooledRenderTarget>& GetSeparateTranslucency(const FViewInfo& View)
+	TRefCountPtr<IPooledRenderTarget>& GetSeparateTranslucency(FIntPoint Size)
 	{
-		// Create the SeparateTranslucency render target (alpha is needed to lerping)
-		FPooledRenderTargetDesc Desc(FPooledRenderTargetDesc::Create2DDesc(GSceneRenderTargets.GetBufferSizeXY(), PF_FloatRGBA, TexCreate_None, TexCreate_RenderTargetable, false));
-		GRenderTargetPool.FindFreeElement(Desc, SeparateTranslucencyRT, TEXT("SeparateTranslucency"));
-
+		if (!SeparateTranslucencyRT || SeparateTranslucencyRT->GetDesc().Extent != Size)
+		{
+			// Create the SeparateTranslucency render target (alpha is needed to lerping)
+			FPooledRenderTargetDesc Desc(FPooledRenderTargetDesc::Create2DDesc(Size, PF_FloatRGBA, TexCreate_None, TexCreate_RenderTargetable, false));
+			GRenderTargetPool.FindFreeElement(Desc, SeparateTranslucencyRT, TEXT("SeparateTranslucency"));
+		}
 		return SeparateTranslucencyRT;
 	}
 
@@ -735,8 +737,6 @@ public:
 		MobileAaColor1.SafeRelease();
 		SelectionOutlineCacheKey.SafeRelease();
 		SelectionOutlineCacheValue.SafeRelease();
-		CustomStencilSRVCacheKey.SafeRelease();
-		CustomStencilSRVCacheValue.SafeRelease();
 
 		for (int32 CascadeIndex = 0; CascadeIndex < ARRAY_COUNT(GlobalDistanceFieldClipmapState); CascadeIndex++)
 		{
diff --git a/Engine/Source/Runtime/Renderer/Private/SceneRendering.cpp b/Engine/Source/Runtime/Renderer/Private/SceneRendering.cpp
--- a/Engine/Source/Runtime/Renderer/Private/SceneRendering.cpp
+++ b/Engine/Source/Runtime/Renderer/Private/SceneRendering.cpp
@@ -117,11 +117,6 @@ static TAutoConsoleVariable<float> CVarTessellationAdaptivePixelsPerTriangle(
 -----------------------------------------------------------------------------*/
 
 
-FRHICommandList* FParallelCommandListSet::AllocCommandList()
-{
-	return new FRHICommandList;
-}
-
 static TAutoConsoleVariable<int32> CVarRHICmdSpewParallelListBalance(
 	TEXT("r.RHICmdSpewParallelListBalance"),
 	0,
@@ -142,12 +137,12 @@ static TAutoConsoleVariable<int32> CVarRHICmdMinDrawsPerParallelCmdList(
 	32,
 	TEXT("The minimum number of draws per cmdlist. If the total number of draws is less than this, then no parallel work will be done at all. This can't always be honored or done correctly. More effective with RHICmdBalanceParallelLists."));
 
-FParallelCommandListSet::FParallelCommandListSet(const FViewInfo& InView, FRHICommandList& InParentCmdList, bool* InOutDirty, bool bInParallelExecute)
+FParallelCommandListSet::FParallelCommandListSet(const FViewInfo& InView, FRHICommandListImmediate& InParentCmdList, bool bInParallelExecute, bool bInCreateSceneContext)
 	: View(InView)
 	, ParentCmdList(InParentCmdList)
-	, OutDirtyIfIgnored(false)
-	, OutDirty(InOutDirty ? *InOutDirty : OutDirtyIfIgnored)
+	, Snapshot(nullptr)
 	, bParallelExecute(GRHISupportsParallelRHIExecute && bInParallelExecute)
+	, bCreateSceneContext(bInCreateSceneContext)
 {
 	Width = CVarRHICmdWidth.GetValueOnRenderThread();
 	MinDrawsPerCommandList = CVarRHICmdMinDrawsPerParallelCmdList.GetValueOnRenderThread();
@@ -158,8 +153,14 @@ FParallelCommandListSet::FParallelCommandListSet(const FViewInfo& InView, FRHICo
 	NumDrawsIfKnown.Reserve(Width * 8);
 }
 
+FRHICommandList* FParallelCommandListSet::AllocCommandList()
+{
+	return new FRHICommandList;
+}
+
 void FParallelCommandListSet::Dispatch()
 {
+	check(IsInRenderingThread() && FMemStack::Get().GetNumMarks() == 1); // we do not want this popped before the end of the scene and it better be the scene allocator
 	check(CommandLists.Num() == Events.Num());
 	if (bSpewBalance)
 	{
@@ -213,11 +214,13 @@ void FParallelCommandListSet::Dispatch()
 		}
 	}
 	CommandLists.Reset();
+	Snapshot = nullptr;
 	Events.Reset();
 }
 
 FParallelCommandListSet::~FParallelCommandListSet()
 {
+	check(IsInRenderingThread() && FMemStack::Get().GetNumMarks() == 1); // we do not want this popped before the end of the scene and it better be the scene allocator
 	checkf(CommandLists.Num() == 0, TEXT("Derived class of FParallelCommandListSet did not call Dispatch in virtual destructor"));
 }
 
@@ -225,11 +228,23 @@ FRHICommandList* FParallelCommandListSet::NewParallelCommandList()
 {
 	FRHICommandList* Result = AllocCommandList();
 	SetStateOnCommandList(*Result); 
+	if (bCreateSceneContext)
+	{
+		FSceneRenderTargets& SceneContext = FSceneRenderTargets::Get(ParentCmdList);
+		check(&SceneContext == &FSceneRenderTargets::Get_FrameConstantsOnly()); // the immediate should not have an overridden context
+		if (!Snapshot)
+		{
+			Snapshot = SceneContext.CreateSnapshot(View);
+		}
+		Snapshot->SetSnapshotOnCmdList(*Result);
+		check(&SceneContext != &FSceneRenderTargets::Get(*Result)); // the new commandlist should have a snapshot
+	}
 	return Result;
 }
 
 void FParallelCommandListSet::AddParallelCommandList(FRHICommandList* CmdList, FGraphEventRef& CompletionEvent, int32 InNumDrawsIfKnown)
 {
+	check(IsInRenderingThread() && FMemStack::Get().GetNumMarks() == 1); // we do not want this popped before the end of the scene and it better be the scene allocator
 	check(CommandLists.Num() == Events.Num());
 	CommandLists.Add(CmdList);
 	Events.Add(CompletionEvent);
@@ -303,6 +318,7 @@ void FViewInfo::Init()
 	ShaderMap = GetGlobalShaderMap(FeatureLevel);
 
 	ViewState = (FSceneViewState*)State;
+	bIsSnapshot = false;
 }
 
 FViewInfo::~FViewInfo()
@@ -382,17 +398,20 @@ void FViewInfo::SetupSkyIrradianceEnvironmentMapConstants(FVector4* OutSkyIrradi
 
 /** Creates the view's uniform buffer given a set of view transforms. */
 TUniformBufferRef<FViewUniformShaderParameters> FViewInfo::CreateUniformBuffer(
+	FRHICommandList& RHICmdList,
 	const TArray<FProjectedShadowInfo*, SceneRenderingAllocator>* DirectionalLightShadowInfo,	
 	const FMatrix& EffectiveTranslatedViewMatrix, 
 	const FMatrix& EffectiveViewToTranslatedWorld, 
 	FBox* OutTranslucentCascadeBoundsArray, 
 	int32 NumTranslucentCascades) const
 {
+	FSceneRenderTargets& SceneContext = FSceneRenderTargets::Get(RHICmdList);
+
 	check(Family);
 	check(!DirectionalLightShadowInfo || DirectionalLightShadowInfo->Num() > 0);
 
 	// Calculate the vector used by shaders to convert clip space coordinates to texture space.
-	const FIntPoint BufferSize = GSceneRenderTargets.GetBufferSizeXY();
+	const FIntPoint BufferSize = SceneContext.GetBufferSizeXY();
 	const float InvBufferSizeX = 1.0f / BufferSize.X;
 	const float InvBufferSizeY = 1.0f / BufferSize.Y;
 	const FVector4 ScreenPositionScaleBias(
@@ -521,7 +540,7 @@ TUniformBufferRef<FViewUniformShaderParameters> FViewInfo::CreateUniformBuffer(
 					FIntPoint ShadowBufferResolution = ShadowInfo.GetShadowBufferResolution();
 					FVector4 ShadowBufferSizeValue((float)ShadowBufferResolution.X, (float)ShadowBufferResolution.Y, 1.0f / (float)ShadowBufferResolution.X, 1.0f / (float)ShadowBufferResolution.Y);
 
-					ViewUniformShaderParameters.DirectionalLightShadowTexture = GSceneRenderTargets.GetShadowDepthZTexture();				
+					ViewUniformShaderParameters.DirectionalLightShadowTexture = SceneContext.GetShadowDepthZTexture();				
 					ViewUniformShaderParameters.DirectionalLightShadowTransition = 1.0f / ShadowInfo.ComputeTransitionSize();
 					ViewUniformShaderParameters.DirectionalLightShadowSize = ShadowBufferSizeValue;
 				}
@@ -919,7 +938,9 @@ void FViewInfo::InitRHIResources(const TArray<FProjectedShadowInfo*, SceneRender
 	/** The view transform, starting from world-space points translated by -ViewOrigin. */
 	FMatrix TranslatedViewMatrix = FTranslationMatrix(-ViewMatrices.PreViewTranslation) * ViewMatrices.ViewMatrix;
 
+	check(IsInRenderingThread());
 	UniformBuffer = CreateUniformBuffer(
+		FRHICommandListExecutor::GetImmediateCommandList(),
 		DirectionalLightShadowInfo,
 		TranslatedViewMatrix,
 		InvViewMatrix * FTranslationMatrix(ViewMatrices.PreViewTranslation),
@@ -945,6 +966,34 @@ void FViewInfo::InitRHIResources(const TArray<FProjectedShadowInfo*, SceneRender
 	}
 }
 
+static TArray<FViewInfo*> ViewInfoSnapshots;
+
+FViewInfo* FViewInfo::CreateSnapshot() const
+{
+	QUICK_SCOPE_CYCLE_COUNTER(STAT_FViewInfo_CreateSnapshot);
+
+	check(IsInRenderingThread() && FMemStack::Get().GetNumMarks() == 1); // we do not want this popped before the end of the scene and it better be the scene allocator
+	FViewInfo* Result = (FViewInfo*)FMemStack::Get().Alloc(sizeof(FViewInfo),ALIGNOF(FViewInfo));
+	FMemory::Memcpy(*Result, *this);
+	TUniformBufferRef<FViewUniformShaderParameters> NullUniformBuffer;
+	FMemory::Memcpy(Result->UniformBuffer, NullUniformBuffer); // we want this to start null without a reference count, since we clear a ref later
+	Result->bIsSnapshot = true;
+	ViewInfoSnapshots.Add(Result);
+	return Result;
+}
+
+void FViewInfo::DestroyAllSnapshots()
+{
+	QUICK_SCOPE_CYCLE_COUNTER(STAT_FViewInfo_DestroyAllSnapshots);
+
+	check(IsInRenderingThread());
+	for (FViewInfo* Snapshot : ViewInfoSnapshots)
+	{
+		Snapshot->UniformBuffer.SafeRelease();
+	}
+	ViewInfoSnapshots.Reset();
+}
+
 IPooledRenderTarget* FViewInfo::GetEyeAdaptation() const
 {
 	FSceneViewState* EffectiveViewState = ViewState;
@@ -1294,7 +1343,8 @@ void FSceneRenderer::RenderCustomDepthPass(FRHICommandListImmediate& RHICmdList)
 	}
 
 	// Render CustomDepth
-	if (GSceneRenderTargets.BeginRenderingCustomDepth(RHICmdList, bPrimitives))
+	FSceneRenderTargets& SceneContext = FSceneRenderTargets::Get(RHICmdList);
+	if (SceneContext.BeginRenderingCustomDepth(RHICmdList, bPrimitives))
 	{
 		SCOPED_DRAW_EVENT(RHICmdList, CustomDepth);
 
@@ -1310,7 +1360,7 @@ void FSceneRenderer::RenderCustomDepthPass(FRHICommandListImmediate& RHICmdList)
 			RHICmdList.SetRasterizerState(TStaticRasterizerState<>::GetRHI());
 			RHICmdList.SetBlendState(TStaticBlendState<>::GetRHI());
 			
-			const bool bWriteCustomStencilValues = GSceneRenderTargets.IsCustomDepthPassWritingStencil();
+			const bool bWriteCustomStencilValues = SceneContext.IsCustomDepthPassWritingStencil();
 
 			if (!bWriteCustomStencilValues)
 			{
@@ -1321,16 +1371,18 @@ void FSceneRenderer::RenderCustomDepthPass(FRHICommandListImmediate& RHICmdList)
 		}
 
 		// resolve using the current ResolveParams 
-		GSceneRenderTargets.FinishRenderingCustomDepth(RHICmdList);
+		SceneContext.FinishRenderingCustomDepth(RHICmdList);
 	}
 }
 
 void FSceneRenderer::OnStartFrame()
 {
+	FSceneRenderTargets& SceneContext = FSceneRenderTargets::Get_Todo_PassContext();
+
 	GRenderTargetPool.VisualizeTexture.OnStartFrame(Views[0]);
 	CompositionGraph_OnStartFrame();
-	GSceneRenderTargets.bScreenSpaceAOIsValid = false;
-	GSceneRenderTargets.bCustomDepthIsValid = false;
+	SceneContext.bScreenSpaceAOIsValid = false;
+	SceneContext.bCustomDepthIsValid = false;
 
 	for(int32 ViewIndex = 0; ViewIndex < Views.Num(); ++ViewIndex)
 	{	
@@ -1372,6 +1424,27 @@ bool FSceneRenderer::ShouldCompositeEditorPrimitives(const FViewInfo& View)
 	return false;
 }
 
+void FSceneRenderer::WaitForTasksClearSnapshotsAndDeleteSceneRenderer(FRHICommandListImmediate& RHICmdList, FSceneRenderer* SceneRenderer)
+{
+	// we are about to destroy things that are being used for async tasks, so we wait here for them.
+	{
+		QUICK_SCOPE_CYCLE_COUNTER(STAT_DeleteSceneRenderer_WaitForTasks);
+		RHICmdList.ImmediateFlush(EImmediateFlushType::WaitForOutstandingTasksOnly);
+	}
+	FViewInfo::DestroyAllSnapshots(); // this destroys viewinfo snapshots
+	FSceneRenderTargets::Get(RHICmdList).DestroyAllSnapshots(); // this will destroy the render target snapshots
+	{
+		QUICK_SCOPE_CYCLE_COUNTER(STAT_DeleteSceneRenderer_Dispatch);
+		RHICmdList.ImmediateFlush(EImmediateFlushType::WaitForDispatchToRHIThread); // we want to make sure this all gets to the rhi thread this frame and doesn't hang around
+	}
+	// Delete the scene renderer.
+	{
+		QUICK_SCOPE_CYCLE_COUNTER(STAT_DeleteSceneRenderer);
+		delete SceneRenderer;
+	}
+}
+
+
 /*-----------------------------------------------------------------------------
 	FRendererModule::BeginRenderingViewFamily
 -----------------------------------------------------------------------------*/
@@ -1440,15 +1513,8 @@ static void RenderViewFamily_RenderThread(FRHICommandListImmediate& RHICmdList,
 
 		GRenderTargetPool.SetEventRecordingActive(false);
 
-		// Delete the scene renderer.
-		{
-			QUICK_SCOPE_CYCLE_COUNTER(STAT_DeleteSceneRenderer);
-			delete SceneRenderer;
-		}
-		{
-			QUICK_SCOPE_CYCLE_COUNTER(STAT_DeleteSceneRenderer_Dispatch);
-			FRHICommandListExecutor::GetImmediateCommandList().ImmediateFlush(EImmediateFlushType::WaitForDispatchToRHIThread); // we want to make sure this all gets to the rhi thread this frame and doesn't hang around
-		}
+		FSceneRenderer::WaitForTasksClearSnapshotsAndDeleteSceneRenderer(RHICmdList, SceneRenderer);
+
 	}
 
 #if STATS
diff --git a/Engine/Source/Runtime/Renderer/Private/SceneRendering.h b/Engine/Source/Runtime/Renderer/Private/SceneRendering.h
--- a/Engine/Source/Runtime/Renderer/Private/SceneRendering.h
+++ b/Engine/Source/Runtime/Renderer/Private/SceneRendering.h
@@ -312,15 +312,13 @@ class FParallelCommandListSet
 {
 public:
 	const FViewInfo& View;
-	FRHICommandList& ParentCmdList;
+	FRHICommandListImmediate& ParentCmdList;
+	FSceneRenderTargets* Snapshot;
 	int32 Width;
 	int32 MinDrawsPerCommandList;
 	bool bBalanceCommands;
 	bool bSpewBalance;
-private:
-	bool OutDirtyIfIgnored;
 public:
-	bool& OutDirty;
 	TArray<FRHICommandList*,SceneRenderingAllocator> CommandLists;
 	TArray<FGraphEventRef,SceneRenderingAllocator> Events;
 	// number of draws in this commandlist if known, -1 if not known. Overestimates are better than nothing.
@@ -331,14 +329,19 @@ protected:
 	void Dispatch();
 	FRHICommandList* AllocCommandList();
 	bool bParallelExecute;
+	bool bCreateSceneContext;
 public:
-	FParallelCommandListSet(const FViewInfo& InView, FRHICommandList& InParentCmdList, bool* InOutDirty, bool bInParallelExecute);
+	FParallelCommandListSet(const FViewInfo& InView, FRHICommandListImmediate& InParentCmdList, bool bInParallelExecute, bool bInCreateSceneContext);
 	virtual ~FParallelCommandListSet();
 	int32 NumParallelCommandLists() const
 	{
 		return CommandLists.Num();
 	}
 	FRHICommandList* NewParallelCommandList();
+	FORCEINLINE FGraphEventArray* GetPrereqs()
+	{
+		return nullptr;
+	}
 	void AddParallelCommandList(FRHICommandList* CmdList, FGraphEventRef& CompletionEvent, int32 InNumDrawsIfKnown = -1);	
 
 	virtual void SetStateOnCommandList(FRHICommandList& CmdList)
@@ -545,6 +548,8 @@ public:
 
 	TShaderMap<FGlobalShaderType>* ShaderMap;
 
+	bool bIsSnapshot;
+
 	/** Custom visibility query for view */
 	ICustomVisibilityQuery* CustomVisibilityQuery;
 
@@ -566,6 +571,7 @@ public:
 
 	/** Creates the view's uniform buffer given a set of view transforms. */
 	TUniformBufferRef<FViewUniformShaderParameters> CreateUniformBuffer(
+		FRHICommandList& RHICmdList,
 		const TArray<FProjectedShadowInfo*, SceneRenderingAllocator>* DirectionalLightShadowInfo,
 		const FMatrix& EffectiveTranslatedViewMatrix, 
 		const FMatrix& EffectiveViewToTranslatedWorld, 
@@ -601,6 +607,12 @@ public:
 		return DitherValue;
 	}
 
+	/** Create a snapshot of this view info on the scene allocator. */
+	FViewInfo* CreateSnapshot() const;
+
+	/** Destroy all snapshots before we wipe the scene allocator. */
+	static void DestroyAllSnapshots();
+
 private:
 
 	/** Initialization that is common to the constructors. */
@@ -703,6 +715,10 @@ public:
 	*/
 	static bool ShouldCompositeEditorPrimitives(const FViewInfo& View);
 
+	/** the last thing we do with a scene renderer, lots of cleanup related to the threading **/
+	static void WaitForTasksClearSnapshotsAndDeleteSceneRenderer(FRHICommandListImmediate& RHICmdList, FSceneRenderer* SceneRenderer);
+
+
 protected:
 
 	// Shared functionality between all scene renderers
