UE4 - Large upgrade to parallel rendering. No longer block per pass. Several RHI cmdlist passthrough functions added and implemented on the consoles

diff --git a/Engine/Source/Runtime/Renderer/Private/SceneVisibility.cpp b/Engine/Source/Runtime/Renderer/Private/SceneVisibility.cpp
--- a/Engine/Source/Runtime/Renderer/Private/SceneVisibility.cpp
+++ b/Engine/Source/Runtime/Renderer/Private/SceneVisibility.cpp
@@ -1052,63 +1052,6 @@ struct FRelevancePacket
 	}
 };
 
-class FRelevancePacketAnyThreadTask
-{
-	FRelevancePacket& Packet;
-	ENamedThreads::Type ThreadToUse;
-public:
-
-	FRelevancePacketAnyThreadTask(FRelevancePacket& InPacket, ENamedThreads::Type InThreadToUse)
-		: Packet(InPacket)
-		, ThreadToUse(InThreadToUse)
-	{
-	}
-
-	FORCEINLINE TStatId GetStatId() const
-	{
-		RETURN_QUICK_DECLARE_CYCLE_STAT(FRelevancePacketAnyThreadTask, STATGROUP_TaskGraphTasks);
-	}
-
-	ENamedThreads::Type GetDesiredThread()
-	{
-		return ThreadToUse;
-	}
-
-	static ESubsequentsMode::Type GetSubsequentsMode() { return ESubsequentsMode::TrackSubsequents; }
-
-	void DoTask(ENamedThreads::Type CurrentThread, const FGraphEventRef& MyCompletionGraphEvent)
-	{
-		Packet.AnyThreadTask();
-	}
-};
-
-class FRelevancePacketRenderThreadTask
-{
-	FRelevancePacket& Packet;
-public:
-
-	FRelevancePacketRenderThreadTask(FRelevancePacket& InPacket)
-		: Packet(InPacket)
-	{
-	}
-
-	FORCEINLINE TStatId GetStatId() const
-	{
-		RETURN_QUICK_DECLARE_CYCLE_STAT(FRelevancePacketRenderThreadTask, STATGROUP_TaskGraphTasks);
-	}
-
-	ENamedThreads::Type GetDesiredThread()
-	{
-		return ENamedThreads::RenderThread_Local;
-	}
-
-	static ESubsequentsMode::Type GetSubsequentsMode() { return ESubsequentsMode::TrackSubsequents; }
-
-	void DoTask(ENamedThreads::Type CurrentThread, const FGraphEventRef& MyCompletionGraphEvent)
-	{
-		Packet.RenderThreadFinalize();
-	}
-};
 
 static TAutoConsoleVariable<int32> CVarParallelInitViews(
 	TEXT("r.ParallelInitViews"),
@@ -1117,21 +1060,10 @@ static TAutoConsoleVariable<int32> CVarParallelInitViews(
 #else
 	1,  
 #endif
-	TEXT("Toggles parallel init views."),
+	TEXT("Toggles parallel init views. 0 = off; 1 = on"),
 	ECVF_RenderThreadSafe
 	);
 
-
-/**
- * Computes view relevance for visible primitives in the view and adds them to
- * appropriate per-view rendering lists.
- * @param Scene - The scene being rendered.
- * @param View - The view for which to compute relevance.
- * @param ViewBit - Bit mask: 1 << ViewIndex where Views(ViewIndex) == View.
- * @param OutRelevantStaticPrimitives - Upon return contains a list of relevant
- *                                      static primitives.
- *                                callback for this view will have ViewBit set.
- */
 static void ComputeAndMarkRelevanceForViewParallel(
 	FRHICommandListImmediate& RHICmdList,
 	const FScene* Scene,
@@ -1150,21 +1082,16 @@ static void ComputeAndMarkRelevanceForViewParallel(
 	uint8* RESTRICT MarkMasks = (uint8*)FMemStack::Get().Alloc(NumMesh + 31 , 8); // some padding to simplify the high speed transpose
 	FMemory::Memzero(MarkMasks, NumMesh + 31);
 
-	// I am going to do the render thread tasks in order, maybe that isn't necessary
-	FGraphEventRef LastRenderThread;
+	int32 EstimateOfNumPackets = NumMesh / (FRelevancePrimSet<int32>::MaxPrims * 4);
+
+	TArray<FRelevancePacket*,SceneRenderingAllocator> Packets;
+
+	Packets.Reserve(EstimateOfNumPackets);
+
 	{
 		FSceneSetBitIterator BitIt(View.PrimitiveVisibilityMap);
 		if (BitIt)
 		{
-			int32 AnyThreadTasksPerRenderThreadTasks = 1;
-
-			if (FApp::ShouldUseThreadingForPerformance() && CVarParallelInitViews.GetValueOnRenderThread() > 0)
-			{
-				AnyThreadTasksPerRenderThreadTasks = FTaskGraphInterface::Get().GetNumWorkerThreads() + 1; // the idea is to put the render thread to work
-			}
-			// else AnyThreadTasksPerRenderThreadTasks == 1 means we never use a task thread
-			int32 WorkingAnyThreadTasksPerRenderThreadTasks = AnyThreadTasksPerRenderThreadTasks;
-			
 
 			FRelevancePacket* Packet = new(FMemStack::Get()) FRelevancePacket(
 				RHICmdList,
@@ -1175,6 +1102,7 @@ static void ComputeAndMarkRelevanceForViewParallel(
 				OutHasDynamicMeshElementsMasks,
 				OutHasDynamicEditorMeshElementsMasks,
 				MarkMasks);
+			Packets.Add(Packet);
 
 			while (1)
 			{
@@ -1182,21 +1110,6 @@ static void ComputeAndMarkRelevanceForViewParallel(
 				++BitIt;
 				if (Packet->Input.IsFull() || !BitIt)
 				{
-					// submit task
-					ENamedThreads::Type ThreadToUse = ENamedThreads::AnyThread;
-					if (!--WorkingAnyThreadTasksPerRenderThreadTasks)
-					{
-						WorkingAnyThreadTasksPerRenderThreadTasks = AnyThreadTasksPerRenderThreadTasks;
-						ThreadToUse = ENamedThreads::RenderThread_Local;
-					}
-					FGraphEventArray RenderPrereqs;
-					if (LastRenderThread.GetReference())
-					{
-						RenderPrereqs.Add(LastRenderThread); // this puts the render thread ones in order
-					}
-					FGraphEventRef AnyThread = TGraphTask<FRelevancePacketAnyThreadTask>::CreateTask(nullptr, ENamedThreads::RenderThread).ConstructAndDispatchWhenReady(*Packet, ThreadToUse);
-					RenderPrereqs.Add(AnyThread);
-					LastRenderThread = TGraphTask<FRelevancePacketRenderThreadTask>::CreateTask(&RenderPrereqs, ENamedThreads::RenderThread).ConstructAndDispatchWhenReady(*Packet);
 					if (!BitIt)
 					{
 						break;
@@ -1212,15 +1125,28 @@ static void ComputeAndMarkRelevanceForViewParallel(
 							OutHasDynamicMeshElementsMasks,
 							OutHasDynamicEditorMeshElementsMasks,
 							MarkMasks);
+						Packets.Add(Packet);
 					}
 				}
 			}
 		}
 	}
-	if (LastRenderThread.GetReference())
 	{
-		QUICK_SCOPE_CYCLE_COUNTER(STAT_ComputeAndMarkRelevanceForViewParallel_Wait);
-		FTaskGraphInterface::Get().WaitUntilTaskCompletes(LastRenderThread, ENamedThreads::RenderThread_Local);
+		QUICK_SCOPE_CYCLE_COUNTER(STAT_ComputeAndMarkRelevanceForViewParallel_ParallelFor);
+		ParallelFor(Packets.Num(), 
+			[&Packets](int32 Index)
+			{
+				Packets[Index]->AnyThreadTask();
+			},
+			!(FApp::ShouldUseThreadingForPerformance() && CVarParallelInitViews.GetValueOnRenderThread() > 0)
+		);
+	}
+	{
+		QUICK_SCOPE_CYCLE_COUNTER(STAT_ComputeAndMarkRelevanceForViewParallel_RenderThreadFinalize);
+		for (auto Packet : Packets)
+		{
+			Packet->RenderThreadFinalize();
+		}
 	}
 	QUICK_SCOPE_CYCLE_COUNTER(STAT_ComputeAndMarkRelevanceForViewParallel_TransposeMeshBits);
 	check(View.StaticMeshVelocityMap.Num() == NumMesh && 
@@ -1289,6 +1215,7 @@ static void ComputeAndMarkRelevanceForViewParallel(
 	}
 }
 
+
 void FSceneRenderer::GatherDynamicMeshElements(
 	TArray<FViewInfo>& InViews, 
 	const FScene* InScene, 
@@ -1463,10 +1390,11 @@ void FSceneRenderer::PreVisibilityFrameSetup(FRHICommandListImmediate& RHICmdLis
 			View.bDisableQuerySubmissions = true;
 			View.bIgnoreExistingQueries = true;
 		}
+		FSceneRenderTargets& SceneContext = FSceneRenderTargets::Get(RHICmdList);
 
 		// set up the screen area for occlusion
-		float NumPossiblePixels = GSceneRenderTargets.UseDownsizedOcclusionQueries() && IsValidRef(GSceneRenderTargets.GetSmallDepthSurface()) ? 
-			(float)View.ViewRect.Width() / GSceneRenderTargets.GetSmallColorDepthDownsampleFactor() * (float)View.ViewRect.Height() / GSceneRenderTargets.GetSmallColorDepthDownsampleFactor() :
+		float NumPossiblePixels = SceneContext.UseDownsizedOcclusionQueries() && IsValidRef(SceneContext.GetSmallDepthSurface()) ? 
+			(float)View.ViewRect.Width() / SceneContext.GetSmallColorDepthDownsampleFactor() * (float)View.ViewRect.Height() / SceneContext.GetSmallColorDepthDownsampleFactor() :
 			View.ViewRect.Width() * View.ViewRect.Height();
 		View.OneOverNumPossiblePixels = NumPossiblePixels > 0.0 ? 1.0f / NumPossiblePixels : 0.0f;
 
diff --git a/Engine/Source/Runtime/Renderer/Private/ShaderBaseClasses.cpp b/Engine/Source/Runtime/Renderer/Private/ShaderBaseClasses.cpp
--- a/Engine/Source/Runtime/Renderer/Private/ShaderBaseClasses.cpp
+++ b/Engine/Source/Runtime/Renderer/Private/ShaderBaseClasses.cpp
@@ -264,7 +264,7 @@ void FMaterialShader::SetParameters(
 				LightAttenuation,
 				LightAttenuationSampler,
 				TStaticSamplerState<SF_Bilinear,AM_Clamp,AM_Clamp,AM_Clamp>::GetRHI(),
-				GSceneRenderTargets.GetLightAttenuationTexture());
+				FSceneRenderTargets::Get(RHICmdList).GetLightAttenuationTexture());
 		}
 	}
 
diff --git a/Engine/Source/Runtime/Renderer/Private/ShadowRendering.cpp b/Engine/Source/Runtime/Renderer/Private/ShadowRendering.cpp
--- a/Engine/Source/Runtime/Renderer/Private/ShadowRendering.cpp
+++ b/Engine/Source/Runtime/Renderer/Private/ShadowRendering.cpp
@@ -168,7 +168,7 @@
 		else
 		{
 			return (Material->IsSpecialEngineMaterial()
-					// Masked and WPO materials need their shaders but cannot be used with a position only stream.
+						// Masked and WPO materials need their shaders but cannot be used with a position only stream.
 					|| ((!Material->WritesEveryPixel() || Material->MaterialMayModifyMeshPosition()) && !bUsePositionOnlyStream))
 					// Only compile one pass point light shaders for feature levels >= SM4
 					&& (ShaderMode != VertexShadowDepth_OnePassPointLight || IsFeatureLevelSupported(Platform, ERHIFeatureLevel::SM4))
@@ -446,10 +446,12 @@
 
 		if(bRenderReflectiveShadowMap)
 		{
+			FSceneRenderTargets& SceneContext = FSceneRenderTargets::Get(RHICmdList);
+
 			// LPV also propagates light transmission (for transmissive materials)
 			SetShaderValue(RHICmdList, ShaderRHI,ReflectiveShadowMapTextureResolution,
-				FVector2D(GSceneRenderTargets.GetReflectiveShadowMapResolution(),
-					GSceneRenderTargets.GetReflectiveShadowMapResolution()));
+				FVector2D(SceneContext.GetReflectiveShadowMapResolution(),
+					SceneContext.GetReflectiveShadowMapResolution()));
 			SetShaderValue(
 				RHICmdList, 
 				ShaderRHI,
@@ -1359,7 +1361,7 @@
 
 class FDrawShadowMeshElementsThreadTask
 {
-	FProjectedShadowInfo& ThisRenderer;
+	FProjectedShadowInfo& ThisShadow;
 	FRHICommandList& RHICmdList;
 	const FViewInfo& View;
 	bool bReflective;
@@ -1367,12 +1369,12 @@
 public:
 
 	FDrawShadowMeshElementsThreadTask(
-		FProjectedShadowInfo& InThisRenderer,
+		FProjectedShadowInfo& InThisShadow,
 		FRHICommandList& InRHICmdList,
 		const FViewInfo& InView,
 		bool InbReflective
 		)
-		: ThisRenderer(InThisRenderer)
+		: ThisShadow(InThisShadow)
 		, RHICmdList(InRHICmdList)
 		, View(InView)
 		, bReflective(InbReflective)
@@ -1398,12 +1400,12 @@
 		if (bReflective)
 		{
 			// reflective shadow map
-			DrawShadowMeshElements<true>(RHICmdList, View, ThisRenderer);
+			DrawShadowMeshElements<true>(RHICmdList, View, ThisShadow);
 		}
 		else
 		{
 			// normal shadow map
-			DrawShadowMeshElements<false>(RHICmdList, View, ThisRenderer);
+			DrawShadowMeshElements<false>(RHICmdList, View, ThisShadow);
 		}
 		RHICmdList.HandleRTThreadTaskCompletion(MyCompletionGraphEvent);
 	}
@@ -1411,7 +1413,7 @@
 
 class FRenderDepthDynamicThreadTask
 {
-	FProjectedShadowInfo& ThisRenderer;
+	FProjectedShadowInfo& ThisShadow;
 	FRHICommandList& RHICmdList;
 	const FViewInfo& View;
 	FSceneRenderer* SceneRenderer;
@@ -1419,12 +1421,12 @@
 public:
 
 	FRenderDepthDynamicThreadTask(
-		FProjectedShadowInfo& InThisRenderer,
+		FProjectedShadowInfo& InThisShadow,
 		FRHICommandList& InRHICmdList,
 		const FViewInfo& InView,
 		FSceneRenderer* InSceneRenderer
 		)
-		: ThisRenderer(InThisRenderer)
+		: ThisShadow(InThisShadow)
 		, RHICmdList(InRHICmdList)
 		, View(InView)
 		, SceneRenderer(InSceneRenderer)
@@ -1445,7 +1447,7 @@
 
 	void DoTask(ENamedThreads::Type CurrentThread, const FGraphEventRef& MyCompletionGraphEvent)
 	{
-		ThisRenderer.RenderDepthDynamic(RHICmdList, SceneRenderer, &View);
+		ThisShadow.RenderDepthDynamic(RHICmdList, SceneRenderer, &View);
 		RHICmdList.HandleRTThreadTaskCompletion(MyCompletionGraphEvent);
 	}
 };
@@ -1524,7 +1526,7 @@
 	}
 	else
 	{
-		RHICmdList.SetDepthStencilState(TStaticDepthStencilState<true, CF_LessEqual>::GetRHI());
+	RHICmdList.SetDepthStencilState(TStaticDepthStencilState<true, CF_LessEqual>::GetRHI());
 	}
 }
 
@@ -1535,12 +1537,24 @@
 	TEXT("Toggles parallel shadow rendering. Parallel rendering must be enabled for this to have an effect."),
 	ECVF_RenderThreadSafe
 	);
+static TAutoConsoleVariable<int32> CVarParallelShadowsNonWholeScene(
+	TEXT("r.ParallelShadowsNonWholeScene"),
+	0,
+	TEXT("Toggles parallel shadow rendering for non whole-scene shadows. r.ParallelShadows must be enabled for this to have an effect."),
+	ECVF_RenderThreadSafe
+	);
 
+
 static TAutoConsoleVariable<int32> CVarRHICmdShadowDeferredContexts(
 	TEXT("r.RHICmdShadowDeferredContexts"),
-	0,
+	1,
 	TEXT("True to use deferred contexts to parallelize shadow command list execution."));
 
+static TAutoConsoleVariable<int32> CVarRHICmdFlushRenderThreadTasksShadowPass(
+	TEXT("r.RHICmdFlushRenderThreadTasksShadowPass"),
+	0,
+	TEXT("Wait for completion of parallel render thread tasks at the end of each shadow pass.  A more granular version of r.RHICmdFlushRenderThreadTasks. If either r.RHICmdFlushRenderThreadTasks or r.RHICmdFlushRenderThreadTasksShadowPass is > 0 we will flush."));
+
 class FShadowParallelCommandListSet : public FParallelCommandListSet
 {
 	FProjectedShadowInfo& ProjectedShadowInfo;
@@ -1548,13 +1562,12 @@
 	EShadowDepthRenderMode RenderMode;
 
 public:
-	FShadowParallelCommandListSet(const FViewInfo& InView, FRHICommandList& InParentCmdList, bool* InOutDirty, bool bInParallelExecute, FProjectedShadowInfo& InProjectedShadowInfo, TFunctionRef<void (FRHICommandList& RHICmdList)> InSetShadowRenderTargets
-		, EShadowDepthRenderMode RenderModeIn )
-		: FParallelCommandListSet(InView, InParentCmdList, InOutDirty, bInParallelExecute)
+	FShadowParallelCommandListSet(const FViewInfo& InView, FRHICommandListImmediate& InParentCmdList, bool bInParallelExecute, bool bInCreateSceneContext, FProjectedShadowInfo& InProjectedShadowInfo, TFunctionRef<void (FRHICommandList& RHICmdList)> InSetShadowRenderTargets, EShadowDepthRenderMode RenderModeIn )
+		: FParallelCommandListSet(InView, InParentCmdList, bInParallelExecute, bInCreateSceneContext)
 		, ProjectedShadowInfo(InProjectedShadowInfo)
 		, SetShadowRenderTargets(InSetShadowRenderTargets)
 		, RenderMode(RenderModeIn)
-		{
+	{
 		SetStateOnCommandList(ParentCmdList);
 	}
 
@@ -1572,30 +1585,48 @@
 
 void FProjectedShadowInfo::RenderDepthInner(FRHICommandList& RHICmdList, FSceneRenderer* SceneRenderer, const FViewInfo* FoundView, TFunctionRef<void (FRHICommandList& RHICmdList)> SetShadowRenderTargets, EShadowDepthRenderMode RenderMode)
 {
-	FShadowDepthDrawingPolicyContext PolicyContext(this);
+	FShadowDepthDrawingPolicyContext StackPolicyContext(this);
+	FShadowDepthDrawingPolicyContext* PolicyContext(&StackPolicyContext);
 
-	if (RHICmdList.IsImmediate() &&  // translucent shadows are draw on the render thread, using a recursive cmdlist
-		GRHICommandList.UseParallelAlgorithms() && CVarParallelShadows.GetValueOnRenderThread())
+	bool bIsWholeSceneDirectionalShadow = IsWholeSceneDirectionalShadow();
+
+
+	if (RHICmdList.IsImmediate() &&  // translucent shadows are draw on the render thread, using a recursive cmdlist (which is not immediate)
+		GRHICommandList.UseParallelAlgorithms() && CVarParallelShadows.GetValueOnRenderThread() &&
+		(bIsWholeSceneDirectionalShadow || CVarParallelShadowsNonWholeScene.GetValueOnRenderThread() > 0)
+		)
 	{
+		check(IsInRenderingThread());
+
 		// parallel version
-		FScopedCommandListWaitForTasks Flusher;
+		bool bFlush = CVarRHICmdFlushRenderThreadTasksShadowPass.GetValueOnRenderThread() > 0 || CVarRHICmdFlushRenderThreadTasks.GetValueOnRenderThread() > 0; 
+		FScopedCommandListWaitForTasks Flusher(bFlush); 
+		if (!bFlush)
+		{
+			/** CAUTION, this is assumed to be a POD type. We allocate the on the scene allocator and NEVER CALL A DESTRUCTOR.
+				If you want to add non-pod data, not a huge problem, we just need to track and destruct them at the end of the scene.
+			**/
+			check(IsInRenderingThread() && FMemStack::Get().GetNumMarks() == 1); // we do not want this popped before the end of the scene and it better be the scene allocator
+			PolicyContext = new (FMemStack::Get()) FShadowDepthDrawingPolicyContext(this);
+		}
+		{
+			check(RHICmdList.IsImmediate());
+			FRHICommandListImmediate& Immed = static_cast<FRHICommandListImmediate&>(RHICmdList);
+			FShadowParallelCommandListSet ParallelCommandListSet(*FoundView, Immed, CVarRHICmdShadowDeferredContexts.GetValueOnRenderThread() > 0, !bFlush, *this, SetShadowRenderTargets, RenderMode);
 
-		{
-			FShadowParallelCommandListSet ParallelCommandListSet(*FoundView, RHICmdList, nullptr, CVarRHICmdShadowDeferredContexts.GetValueOnRenderThread() > 0, *this, SetShadowRenderTargets, RenderMode);
-
 			// Draw the subject's static elements using static draw lists
-			if (IsWholeSceneDirectionalShadow() && RenderMode != ShadowDepthRenderMode_EmissiveOnly && RenderMode != ShadowDepthRenderMode_GIBlockingVolumes)
+			if (bIsWholeSceneDirectionalShadow && RenderMode != ShadowDepthRenderMode_EmissiveOnly && RenderMode != ShadowDepthRenderMode_GIBlockingVolumes)
 			{
 				SCOPE_CYCLE_COUNTER(STAT_WholeSceneStaticDrawListShadowDepthsTime);
 
 				if (bReflectiveShadowmap)
 				{
-					SceneRenderer->Scene->WholeSceneReflectiveShadowMapDrawList.DrawVisibleParallel(PolicyContext, StaticMeshWholeSceneShadowDepthMap, StaticMeshWholeSceneShadowBatchVisibility, ParallelCommandListSet);
+					SceneRenderer->Scene->WholeSceneReflectiveShadowMapDrawList.DrawVisibleParallel(*PolicyContext, StaticMeshWholeSceneShadowDepthMap, StaticMeshWholeSceneShadowBatchVisibility, ParallelCommandListSet);
 				}
 				else
 				{
 					// Use the scene's shadow depth draw list with this shadow's visibility map
-					SceneRenderer->Scene->WholeSceneShadowDepthDrawList.DrawVisibleParallel(PolicyContext, StaticMeshWholeSceneShadowDepthMap, StaticMeshWholeSceneShadowBatchVisibility, ParallelCommandListSet);
+					SceneRenderer->Scene->WholeSceneShadowDepthDrawList.DrawVisibleParallel(*PolicyContext, StaticMeshWholeSceneShadowDepthMap, StaticMeshWholeSceneShadowBatchVisibility, ParallelCommandListSet);
 				}
 			}
 			// Draw the subject's static elements using manual state filtering
@@ -1603,7 +1634,7 @@
 			{
 				FRHICommandList* CmdList = ParallelCommandListSet.NewParallelCommandList();
 
-				FGraphEventRef AnyThreadCompletionEvent = TGraphTask<FDrawShadowMeshElementsThreadTask>::CreateTask(nullptr, ENamedThreads::RenderThread)
+				FGraphEventRef AnyThreadCompletionEvent = TGraphTask<FDrawShadowMeshElementsThreadTask>::CreateTask(ParallelCommandListSet.GetPrereqs(), ENamedThreads::RenderThread)
 					.ConstructAndDispatchWhenReady(*this, *CmdList, *FoundView, bReflectiveShadowmap && !CascadeSettings.bOnePassPointLightShadow);
 
 				ParallelCommandListSet.AddParallelCommandList(CmdList, AnyThreadCompletionEvent, SubjectMeshElements.Num());
@@ -1612,7 +1643,7 @@
 			{
 				FRHICommandList* CmdList = ParallelCommandListSet.NewParallelCommandList();
 
-				FGraphEventRef AnyThreadCompletionEvent = TGraphTask<FRenderDepthDynamicThreadTask>::CreateTask(nullptr, ENamedThreads::RenderThread)
+				FGraphEventRef AnyThreadCompletionEvent = TGraphTask<FRenderDepthDynamicThreadTask>::CreateTask(ParallelCommandListSet.GetPrereqs(), ENamedThreads::RenderThread)
 					.ConstructAndDispatchWhenReady(*this, *CmdList, *FoundView, SceneRenderer);
 
 				ParallelCommandListSet.AddParallelCommandList(CmdList, AnyThreadCompletionEvent, DynamicSubjectMeshElements.Num());
@@ -1625,18 +1656,18 @@
 		SetStateForDepth(RHICmdList, RenderMode);
 
 		// Draw the subject's static elements using static draw lists
-		if (IsWholeSceneDirectionalShadow() && RenderMode != ShadowDepthRenderMode_EmissiveOnly && RenderMode != ShadowDepthRenderMode_GIBlockingVolumes)
+		if (bIsWholeSceneDirectionalShadow && RenderMode != ShadowDepthRenderMode_EmissiveOnly && RenderMode != ShadowDepthRenderMode_GIBlockingVolumes)
 		{
 			SCOPE_CYCLE_COUNTER(STAT_WholeSceneStaticDrawListShadowDepthsTime);
 
 			if (bReflectiveShadowmap)
 			{
-				SceneRenderer->Scene->WholeSceneReflectiveShadowMapDrawList.DrawVisible(RHICmdList, *FoundView, PolicyContext, StaticMeshWholeSceneShadowDepthMap, StaticMeshWholeSceneShadowBatchVisibility);
+				SceneRenderer->Scene->WholeSceneReflectiveShadowMapDrawList.DrawVisible(RHICmdList, *FoundView, *PolicyContext, StaticMeshWholeSceneShadowDepthMap, StaticMeshWholeSceneShadowBatchVisibility);
 			}
 			else
 			{
 				// Use the scene's shadow depth draw list with this shadow's visibility map
-				SceneRenderer->Scene->WholeSceneShadowDepthDrawList.DrawVisible(RHICmdList, *FoundView, PolicyContext, StaticMeshWholeSceneShadowDepthMap, StaticMeshWholeSceneShadowBatchVisibility);
+				SceneRenderer->Scene->WholeSceneShadowDepthDrawList.DrawVisible(RHICmdList, *FoundView, *PolicyContext, StaticMeshWholeSceneShadowDepthMap, StaticMeshWholeSceneShadowBatchVisibility);
 			}
 		}
 		// Draw the subject's static elements using manual state filtering
@@ -1659,6 +1690,72 @@
 	}
 }
 
+
+void FProjectedShadowInfo::ModifyViewForShadow(FRHICommandList& RHICmdList, FViewInfo* FoundView)
+{
+	FIntRect OriginalViewRect = FoundView->ViewRect;
+	FoundView->ViewRect.Min.X = 0;
+	FoundView->ViewRect.Min.Y = 0;
+	FoundView->ViewRect.Max.X = ResolutionX;
+	FoundView->ViewRect.Max.Y =  ResolutionY;
+	FoundView->ViewMatrices.ProjMatrix.M[2][0] = 0.0f;
+	FoundView->ViewMatrices.ProjMatrix.M[2][1] = 0.0f;
+
+	{
+		// Compute the view projection matrix and its inverse.
+		FoundView->ViewProjectionMatrix = FoundView->ViewMatrices.ViewMatrix * FoundView->ViewMatrices.ProjMatrix;
+		FoundView->InvViewProjectionMatrix = FoundView->ViewMatrices.GetInvProjMatrix() * FoundView->InvViewMatrix;
+
+		/** The view transform, starting from world-space points translated by -ViewOrigin. */
+		FMatrix TranslatedViewMatrix = FTranslationMatrix(-FoundView->ViewMatrices.PreViewTranslation) * FoundView->ViewMatrices.ViewMatrix;
+
+		// Compute a transform from view origin centered world-space to clip space.
+		FoundView->ViewMatrices.TranslatedViewProjectionMatrix = TranslatedViewMatrix * FoundView->ViewMatrices.ProjMatrix;
+		FoundView->ViewMatrices.InvTranslatedViewProjectionMatrix = FoundView->ViewMatrices.TranslatedViewProjectionMatrix.Inverse();
+	}
+
+
+	// Override the view matrix so that billboarding primitives will be aligned to the light
+	//@todo - creating a new uniform buffer is expensive, only do this when the vertex factory needs an accurate view matrix (particle sprites)
+	FoundView->ViewMatrices.ViewMatrix = ShadowViewMatrix;
+	FBox VolumeBounds[TVC_MAX];
+	FoundView->UniformBuffer = FoundView->CreateUniformBuffer(
+		RHICmdList,
+		nullptr,
+		ShadowViewMatrix, 
+		ShadowViewMatrix.Inverse(),
+		VolumeBounds,
+		TVC_MAX);
+
+	// we are going to set this back now because we only want the correct view rect for the uniform buffer. For LOD calculations, we want the rendering viewrect and proj matrix.
+	FoundView->ViewRect = OriginalViewRect;
+}
+
+FViewInfo* FProjectedShadowInfo::FindViewForShadow(FSceneRenderer* SceneRenderer)
+{
+	// Choose an arbitrary view where this shadow's subject is relevant.
+	FViewInfo* FoundView = NULL;
+	for(int32 ViewIndex = 0;ViewIndex < SceneRenderer->Views.Num();ViewIndex++)
+	{
+		FViewInfo* CheckView = &SceneRenderer->Views[ViewIndex];
+		const FVisibleLightViewInfo& VisibleLightViewInfo = CheckView->VisibleLightInfos[LightSceneInfo->Id];
+		FPrimitiveViewRelevance ViewRel = VisibleLightViewInfo.ProjectedShadowViewRelevanceMap[ShadowId];
+		if (ViewRel.bShadowRelevance)
+		{
+			FoundView = CheckView;
+			break;
+		}
+	}
+	check(FoundView);
+	return FoundView;
+}
+
+static TAutoConsoleVariable<int32> CVarRHICmdFlushRenderThreadTasksForShadowViewInfos(
+	TEXT("r.RHICmdFlushRenderThreadTasksForShadowViewInfos"),
+	0,
+	TEXT("Wait for completion of parallel render thread tasks before modifying views. Otherwise, memcpy a clone and use that. A more granular version of r.RHICmdFlushRenderThreadTasks. If either r.RHICmdFlushRenderThreadTasks or r.RHICmdFlushRenderThreadTasksForShadowViewInfos is > 0 we will flush."));
+
+
 void FProjectedShadowInfo::RenderDepth(FRHICommandList& RHICmdList, FSceneRenderer* SceneRenderer, TFunctionRef<void (FRHICommandList& RHICmdList)> SetShadowRenderTargets, EShadowDepthRenderMode RenderMode)
 {
 	// Select the correct set of arrays for the current render mode
@@ -1700,90 +1797,63 @@
 		return;
 	}
 
+	bool bMakeViewSnapshot = CVarRHICmdFlushRenderThreadTasksForShadowViewInfos.GetValueOnRenderThread() == 0  && CVarRHICmdFlushRenderThreadTasks.GetValueOnRenderThread() == 0;
+
 	// Choose an arbitrary view where this shadow's subject is relevant.
-	FViewInfo* FoundView = NULL;
-	for(int32 ViewIndex = 0;ViewIndex < SceneRenderer->Views.Num();ViewIndex++)
+	FViewInfo* FoundView = FindViewForShadow(SceneRenderer);
+
+	check(FoundView && IsInRenderingThread()); // we should not hack the view or create the uniform buffer in parallel
+
+	if (bMakeViewSnapshot)
 	{
-		FViewInfo* CheckView = &SceneRenderer->Views[ViewIndex];
-		const FVisibleLightViewInfo& VisibleLightViewInfo = CheckView->VisibleLightInfos[LightSceneInfo->Id];
-		FPrimitiveViewRelevance ViewRel = VisibleLightViewInfo.ProjectedShadowViewRelevanceMap[ShadowId];
-		if (ViewRel.bShadowRelevance)
-		{
-			FoundView = CheckView;
-			break;
-		}
-	}
-	check(FoundView 
-		&& IsInRenderingThread()); // we should not hack the view in parallel
-
-	// Backup properties of the view that we will override
-	TUniformBufferRef<FViewUniformShaderParameters> OriginalUniformBuffer = FoundView->UniformBuffer;
-	FMatrix OriginalViewMatrix = FoundView->ViewMatrices.ViewMatrix;
-	FIntRect OriginalViewRect = FoundView->ViewRect;
-	FoundView->ViewRect.Min.X = 0;
-	FoundView->ViewRect.Min.Y = 0;
-	FoundView->ViewRect.Max.X = ResolutionX;
-	FoundView->ViewRect.Max.Y =  ResolutionY;
-
-	float JitterX = FoundView->ViewMatrices.ProjMatrix.M[2][0];
-	float JitterY = FoundView->ViewMatrices.ProjMatrix.M[2][1];
-
-	FoundView->ViewMatrices.ProjMatrix.M[2][0] = 0.0f;
-	FoundView->ViewMatrices.ProjMatrix.M[2][1] = 0.0f;
-
-	{
-		// Compute the view projection matrix and its inverse.
-		FoundView->ViewProjectionMatrix = FoundView->ViewMatrices.ViewMatrix * FoundView->ViewMatrices.ProjMatrix;
-		FoundView->InvViewProjectionMatrix = FoundView->ViewMatrices.GetInvProjMatrix() * FoundView->InvViewMatrix;
-
-		/** The view transform, starting from world-space points translated by -ViewOrigin. */
-		FMatrix TranslatedViewMatrix = FTranslationMatrix(-FoundView->ViewMatrices.PreViewTranslation) * FoundView->ViewMatrices.ViewMatrix;
-
-		// Compute a transform from view origin centered world-space to clip space.
-		FoundView->ViewMatrices.TranslatedViewProjectionMatrix = TranslatedViewMatrix * FoundView->ViewMatrices.ProjMatrix;
-		FoundView->ViewMatrices.InvTranslatedViewProjectionMatrix = FoundView->ViewMatrices.TranslatedViewProjectionMatrix.Inverse();
+		FoundView = FoundView->CreateSnapshot();
 	}
 
+	TUniformBufferRef<FViewUniformShaderParameters> OriginalUniformBuffer;
+	FMatrix OriginalViewMatrix;
+	float JitterX = 0.0f;
+	float JitterY = 0.0f;
 
-	// Override the view matrix so that billboarding primitives will be aligned to the light
-	//@todo - creating a new uniform buffer is expensive, only do this when the vertex factory needs an accurate view matrix (particle sprites)
-	FoundView->ViewMatrices.ViewMatrix = ShadowViewMatrix;
-	FBox VolumeBounds[TVC_MAX];
-	FoundView->UniformBuffer = FoundView->CreateUniformBuffer(
-		nullptr,
-		ShadowViewMatrix, 
-		ShadowViewMatrix.Inverse(),
-		VolumeBounds,
-		TVC_MAX);
+	if (!bMakeViewSnapshot)
+	{
+		// Backup properties of the view that we will override
+		OriginalUniformBuffer = FoundView->UniformBuffer;
+		OriginalViewMatrix = FoundView->ViewMatrices.ViewMatrix;
 
-	// we are going to set this back now because we only want the correct view rect for the uniform buffer. For LOD calculations, we want the rendering viewrect and proj matrix.
-	FoundView->ViewRect = OriginalViewRect;
+		JitterX = FoundView->ViewMatrices.ProjMatrix.M[2][0];
+		JitterY = FoundView->ViewMatrices.ProjMatrix.M[2][1];
+		// previous jobs rely on previously unhacked views
+		FRHICommandListExecutor::GetImmediateCommandList().ImmediateFlush(EImmediateFlushType::WaitForOutstandingTasksOnly);
+	}
 
-	// Prevent materials from getting overridden during shadow casting in viewmodes like lighting only
-	// Lighting only should only affect the material used with direct lighting, not the indirect lighting
-	FoundView->bForceShowMaterials = true;
+	ModifyViewForShadow(RHICmdList, FoundView);
 
 	RenderDepthInner(RHICmdList, SceneRenderer, FoundView, SetShadowRenderTargets, RenderMode);
 
-	FoundView->bForceShowMaterials = false;
-	FoundView->UniformBuffer = OriginalUniformBuffer;
-	FoundView->ViewMatrices.ViewMatrix = OriginalViewMatrix;
+	if (!bMakeViewSnapshot)
+	{
+		// previous jobs rely on previously hacked views...maybe we didn't actually do shadows in parallel, but if not, there are no outstanding jobs anyway
+		FRHICommandListExecutor::GetImmediateCommandList().ImmediateFlush(EImmediateFlushType::WaitForOutstandingTasksOnly);
 
-	FoundView->ViewMatrices.ProjMatrix.M[2][0] = JitterX;
-	FoundView->ViewMatrices.ProjMatrix.M[2][1] = JitterY;
-
-	{
-		// Compute the view projection matrix and its inverse.
-		FoundView->ViewProjectionMatrix = FoundView->ViewMatrices.ViewMatrix * FoundView->ViewMatrices.ProjMatrix;
-		FoundView->InvViewProjectionMatrix = FoundView->ViewMatrices.GetInvProjMatrix() * FoundView->InvViewMatrix;
-
-		/** The view transform, starting from world-space points translated by -ViewOrigin. */
-		FMatrix TranslatedViewMatrix = FTranslationMatrix(-FoundView->ViewMatrices.PreViewTranslation) * FoundView->ViewMatrices.ViewMatrix;
-
-		// Compute a transform from view origin centered world-space to clip space.
-		FoundView->ViewMatrices.TranslatedViewProjectionMatrix = TranslatedViewMatrix * FoundView->ViewMatrices.ProjMatrix;
-		FoundView->ViewMatrices.InvTranslatedViewProjectionMatrix = FoundView->ViewMatrices.TranslatedViewProjectionMatrix.Inverse();
-	}
+	    FoundView->UniformBuffer = OriginalUniformBuffer;
+	    FoundView->ViewMatrices.ViewMatrix = OriginalViewMatrix;
+    
+	    FoundView->ViewMatrices.ProjMatrix.M[2][0] = JitterX;
+	    FoundView->ViewMatrices.ProjMatrix.M[2][1] = JitterY;
+    
+	    {
+		    // Compute the view projection matrix and its inverse.
+		    FoundView->ViewProjectionMatrix = FoundView->ViewMatrices.ViewMatrix * FoundView->ViewMatrices.ProjMatrix;
+		    FoundView->InvViewProjectionMatrix = FoundView->ViewMatrices.GetInvProjMatrix() * FoundView->InvViewMatrix;
+    
+		    /** The view transform, starting from world-space points translated by -ViewOrigin. */
+		    FMatrix TranslatedViewMatrix = FTranslationMatrix(-FoundView->ViewMatrices.PreViewTranslation) * FoundView->ViewMatrices.ViewMatrix;
+    
+		    // Compute a transform from view origin centered world-space to clip space.
+		    FoundView->ViewMatrices.TranslatedViewProjectionMatrix = TranslatedViewMatrix * FoundView->ViewMatrices.ProjMatrix;
+		    FoundView->ViewMatrices.InvTranslatedViewProjectionMatrix = FoundView->ViewMatrices.TranslatedViewProjectionMatrix.Inverse();
+	    }
+    }
 }
 
 void StencilingGeometry::DrawSphere(FRHICommandList& RHICmdList)
@@ -2491,12 +2561,14 @@
 /** Returns the resolution of the shadow buffer used for this shadow, based on the shadow's type. */
 FIntPoint FProjectedShadowInfo::GetShadowBufferResolution() const
 {
+	FSceneRenderTargets& SceneContext_ConstantsOnly = FSceneRenderTargets::Get_FrameConstantsOnly(); // this is kind of scary, but those RT's stay allocated during a frame so should be fine.
+
 	if (bTranslucentShadow)
 	{
-		return GSceneRenderTargets.GetTranslucentShadowDepthTextureResolution();
+		return SceneContext_ConstantsOnly.GetTranslucentShadowDepthTextureResolution();
 	}
 
-	const FTexture2DRHIRef& ShadowTexture = GSceneRenderTargets.GetShadowDepthZTexture(bAllocatedInPreshadowCache);
+	const FTexture2DRHIRef& ShadowTexture = SceneContext_ConstantsOnly.GetShadowDepthZTexture(bAllocatedInPreshadowCache);
 
 	//prefer to return the actual size of the allocated texture if possible.  It may be larger than the size of a single shadowmap due to atlasing (see forward renderer CSM handling in InitDynamicShadows).
 	if (ShadowTexture)
@@ -2505,7 +2577,7 @@
 	}
 	else
 	{
-		return bAllocatedInPreshadowCache ? GSceneRenderTargets.GetPreShadowCacheTextureResolution() : GSceneRenderTargets.GetShadowDepthTextureResolution();
+		return bAllocatedInPreshadowCache ? SceneContext_ConstantsOnly.GetPreShadowCacheTextureResolution() : SceneContext_ConstantsOnly.GetShadowDepthTextureResolution();
 	}
 }
 
@@ -2714,6 +2786,7 @@
 
 		if (bShadowIsVisible && ProjectedShadowInfo->CascadeSettings.bOnePassPointLightShadow)
 		{
+			FSceneRenderTargets& SceneContext = FSceneRenderTargets::Get(RHICmdList);
 			INC_DWORD_STAT(STAT_WholeSceneShadows);
 
 			if (!ProjectedShadowInfo->CascadeSettings.bRayTracedDistanceField)
@@ -2721,9 +2794,9 @@
 				SCOPED_DRAW_EVENT(RHICmdList, ShadowDepthsFromOpaquePointLight);
 
 				bool bPerformClear = true;
-				auto SetShadowRenderTargets = [this, &bPerformClear, ProjectedShadowInfo](FRHICommandList& InRHICmdList)
+				auto SetShadowRenderTargets = [this, &bPerformClear, ProjectedShadowInfo, &SceneContext](FRHICommandList& InRHICmdList)
 				{
-					GSceneRenderTargets.BeginRenderingCubeShadowDepth(InRHICmdList, ProjectedShadowInfo->ResolutionX);
+					SceneContext.BeginRenderingCubeShadowDepth(InRHICmdList, ProjectedShadowInfo->ResolutionX);
 					ProjectedShadowInfo->ClearDepth(InRHICmdList, this, bPerformClear);
 				};
 
@@ -2732,11 +2805,11 @@
 				bPerformClear = false;
 
 				ProjectedShadowInfo->RenderDepth(RHICmdList, this, SetShadowRenderTargets);
-				GSceneRenderTargets.FinishRenderingCubeShadowDepth(RHICmdList, ProjectedShadowInfo->ResolutionX);
+				SceneContext.FinishRenderingCubeShadowDepth(RHICmdList, ProjectedShadowInfo->ResolutionX);
 			}
 
 			{
-				GSceneRenderTargets.BeginRenderingLightAttenuation(RHICmdList);
+				SceneContext.BeginRenderingLightAttenuation(RHICmdList);
 
 				SCOPED_DRAW_EVENT(RHICmdList, ShadowProjectionOnOpaque);
 
@@ -2784,8 +2857,9 @@
 	const TArray<FProjectedShadowInfo*,SceneRenderingAllocator>& Shadows
 	)
 {
+	FSceneRenderTargets& SceneContext = FSceneRenderTargets::Get(RHICmdList);
 	// Normal shadows render to light attenuation
-	GSceneRenderTargets.BeginRenderingLightAttenuation(RHICmdList);
+	SceneContext.BeginRenderingLightAttenuation(RHICmdList);
 
 	for (int32 ViewIndex = 0; ViewIndex < Views.Num(); ViewIndex++)
 	{
@@ -2810,7 +2884,7 @@
 				{
 					ProjectedShadowInfo->RenderProjection(RHICmdList, ViewIndex, &View);
 
-					GRenderTargetPool.VisualizeTexture.SetCheckPoint(RHICmdList, GSceneRenderTargets.GetLightAttenuation());
+					GRenderTargetPool.VisualizeTexture.SetCheckPoint(RHICmdList, SceneContext.GetLightAttenuation());
 				}
 			}
 		}
@@ -2925,6 +2999,7 @@
 
 	int32 NumShadowsRendered = 0;
 
+	FSceneRenderTargets& SceneContext = FSceneRenderTargets::Get(RHICmdList);
 	while (NumShadowsRendered < Shadows.Num())
 	{
 		for (int32 ShadowIndex = 0; ShadowIndex < Shadows.Num(); ShadowIndex++)
@@ -2937,7 +3012,7 @@
 		int32 NumAllocatedShadows = 0;
 
 		// Reset the translucent shadow allocations
-		TranslucentSelfShadowLayout = FTextureLayout(1, 1, GSceneRenderTargets.GetTranslucentShadowDepthTextureResolution().X, GSceneRenderTargets.GetTranslucentShadowDepthTextureResolution().Y, false, false);
+		TranslucentSelfShadowLayout = FTextureLayout(1, 1, SceneContext.GetTranslucentShadowDepthTextureResolution().X, SceneContext.GetTranslucentShadowDepthTextureResolution().Y, false, false);
 
 		for (int32 ShadowIndex = 0; ShadowIndex < Shadows.Num(); ShadowIndex++)
 		{
@@ -3087,13 +3162,14 @@
 
 	int32 NumShadowsRendered	= 0;
 
+	FSceneRenderTargets& SceneContext = FSceneRenderTargets::Get(RHICmdList);
 	while (NumShadowsRendered < Shadows.Num())
 	{
 		int32 NumAllocatedShadows = 0;
 
 		// Allocate shadow texture space to the shadows.
 		//@TODO: Is this code all necessary? 
-		const FIntPoint ShadowBufferResolution = GSceneRenderTargets.GetShadowDepthTextureResolution();
+		const FIntPoint ShadowBufferResolution = SceneContext.GetShadowDepthTextureResolution();
 		FTextureLayout ShadowLayout(1, 1, ShadowBufferResolution.X, ShadowBufferResolution.Y, false, false);
 
 		for (int32 ShadowIndex = 0; ShadowIndex < Shadows.Num(); ShadowIndex++)
@@ -3134,9 +3210,9 @@
 				FLightPropagationVolume* LightPropagationVolume = ViewState->GetLightPropagationVolume();
 
 				check(LightPropagationVolume);
-				auto SetShadowRenderTargets = [this, LightPropagationVolume, ProjectedShadowInfo](FRHICommandList& InRHICmdList)
+				auto SetShadowRenderTargets = [this, LightPropagationVolume, ProjectedShadowInfo, &SceneContext](FRHICommandList& InRHICmdList)
 				{
-					GSceneRenderTargets.BeginRenderingReflectiveShadowMap(InRHICmdList, LightPropagationVolume);
+					SceneContext.BeginRenderingReflectiveShadowMap(InRHICmdList, LightPropagationVolume);
 					ProjectedShadowInfo->ClearDepth(InRHICmdList, this, false);
 				};
 				SetShadowRenderTargets(RHICmdList);  // run it now, maybe run it later for parallel command lists				
@@ -3154,7 +3230,7 @@
 					ProjectedShadowInfo->RenderDepth(RHICmdList, this, SetShadowRenderTargets, ShadowDepthRenderMode_GIBlockingVolumes);
 				}
 
-				GSceneRenderTargets.FinishRenderingReflectiveShadowMap(RHICmdList);
+				SceneContext.FinishRenderingReflectiveShadowMap(RHICmdList);
 			}
 		}
 
@@ -3176,9 +3252,9 @@
 						LightPropagationVolume->InjectDirectionalLightRSM( 
 							RHICmdList, 
 							*ProjectedShadowInfo->DependentView,
-							GSceneRenderTargets.GetReflectiveShadowMapDiffuseTexture(), 
-							GSceneRenderTargets.GetReflectiveShadowMapNormalTexture(),
-							GSceneRenderTargets.GetReflectiveShadowMapDepthTexture(),
+							SceneContext.GetReflectiveShadowMapDiffuseTexture(), 
+							SceneContext.GetReflectiveShadowMapNormalTexture(),
+							SceneContext.GetReflectiveShadowMapDepthTexture(),
 							*ProjectedShadowInfo, 
 							LightSceneInfo->Proxy->GetColor() );
 					}
@@ -3279,12 +3355,13 @@
 
 	int32 NumShadowsRendered = 0;
 
+	FSceneRenderTargets& SceneContext = FSceneRenderTargets::Get(RHICmdList);
 	while (NumShadowsRendered < Shadows.Num())
 	{
 		int32 NumAllocatedShadows = 0;
 
 		// Allocate shadow texture space to the shadows.
-		const FIntPoint ShadowBufferResolution = GSceneRenderTargets.GetShadowDepthTextureResolution();
+		const FIntPoint ShadowBufferResolution = SceneContext.GetShadowDepthTextureResolution();
 		FTextureLayout ShadowLayout(1, 1, ShadowBufferResolution.X, ShadowBufferResolution.Y, false, false);
 
 		for (int32 ShadowIndex = 0; ShadowIndex < Shadows.Num(); ShadowIndex++)
@@ -3316,9 +3393,9 @@
 			SCOPED_DRAW_EVENT(RHICmdList, ShadowDepthsFromOpaqueProjected);
 
 			bool bPerformClear = true;
-			auto SetShadowRenderTargets = [&bPerformClear](FRHICommandList& InRHICmdList)
+			auto SetShadowRenderTargets = [&bPerformClear, &SceneContext](FRHICommandList& InRHICmdList)
 			{
-				GSceneRenderTargets.BeginRenderingShadowDepth(InRHICmdList, bPerformClear);
+				SceneContext.BeginRenderingShadowDepth(InRHICmdList, bPerformClear);
 			};
 
 			SetShadowRenderTargets(RHICmdList);  // run it now, maybe run it later for parallel command lists
@@ -3334,7 +3411,7 @@
 				}
 			}
 
-			GSceneRenderTargets.FinishRenderingShadowDepth(RHICmdList);
+			SceneContext.FinishRenderingShadowDepth(RHICmdList);
 		}
 
 		// Render the shadow projections.
@@ -3477,13 +3554,14 @@
 				// Only render depths for shadows which haven't already cached their depths
 				if (!ProjectedShadowInfo->bDepthsCached)
 				{
+					FSceneRenderTargets& SceneContext = FSceneRenderTargets::Get(RHICmdList);
 					check(ProjectedShadowInfo->bAllocated);
 					
 					//
 					bool bPerformClear = true;;
-					auto SetShadowRenderTargets = [this, &bPerformClear, ProjectedShadowInfo](FRHICommandList& InRHICmdList)
+					auto SetShadowRenderTargets = [this, &bPerformClear, ProjectedShadowInfo, &SceneContext](FRHICommandList& InRHICmdList)
 					{
-						SetRenderTarget(InRHICmdList, FTextureRHIRef(), GSceneRenderTargets.PreShadowCacheDepthZ->GetRenderTargetItem().TargetableTexture);
+						SetRenderTarget(InRHICmdList, FTextureRHIRef(), SceneContext.PreShadowCacheDepthZ->GetRenderTargetItem().TargetableTexture);
 						ProjectedShadowInfo->ClearDepth(InRHICmdList, this, bPerformClear);
 					};					
 					SetShadowRenderTargets(RHICmdList); // run it now, maybe run it later for parallel command lists
@@ -3499,8 +3577,8 @@
 						ProjectedShadowInfo->Y + ProjectedShadowInfo->ResolutionY + SHADOW_BORDER * 2));
 
 					RHICmdList.CopyToResolveTarget(
-						GSceneRenderTargets.PreShadowCacheDepthZ->GetRenderTargetItem().TargetableTexture, 
-						GSceneRenderTargets.PreShadowCacheDepthZ->GetRenderTargetItem().ShaderResourceTexture,
+						SceneContext.PreShadowCacheDepthZ->GetRenderTargetItem().TargetableTexture, 
+						SceneContext.PreShadowCacheDepthZ->GetRenderTargetItem().ShaderResourceTexture,
 						false, 
 						ResolveParams);
 				}
@@ -3590,13 +3668,14 @@
 	Shadows.Sort( FCompareFProjectedShadowInfoBySplitIndex() );	
 
 	{
+		FSceneRenderTargets& SceneContext = FSceneRenderTargets::Get(RHICmdList);
 		// Render the shadow depths.
 		SCOPED_DRAW_EVENT(RHICmdList, ShadowDepthsFromOpaqueForward);
 
 		bool bFirst = true;
-		auto SetShadowRenderTargets = [&bFirst](FRHICommandList& InRHICmdList)
+		auto SetShadowRenderTargets = [&bFirst, &SceneContext](FRHICommandList& InRHICmdList)
 		{
-			GSceneRenderTargets.BeginRenderingShadowDepth(InRHICmdList, bFirst);
+			SceneContext.BeginRenderingShadowDepth(InRHICmdList, bFirst);
 		};
 
 		SetShadowRenderTargets(RHICmdList);  // run it now, maybe run it later for parallel command lists
@@ -3616,7 +3695,7 @@
 			}
 		}
 
-		GSceneRenderTargets.FinishRenderingShadowDepth(RHICmdList);
+		SceneContext.FinishRenderingShadowDepth(RHICmdList);
 	}
 
 	return bAttenuationBufferDirty;
diff --git a/Engine/Source/Runtime/Renderer/Private/ShadowRendering.h b/Engine/Source/Runtime/Renderer/Private/ShadowRendering.h
--- a/Engine/Source/Runtime/Renderer/Private/ShadowRendering.h
+++ b/Engine/Source/Runtime/Renderer/Private/ShadowRendering.h
@@ -288,6 +288,9 @@ class FShadowStaticMeshElement;
  */
 struct FShadowDepthDrawingPolicyContext : FMeshDrawingPolicy::ContextDataType
 {
+	/** CAUTION, this is assumed to be a POD type. We allocate the on the scene allocator and NEVER CALL A DESTRUCTOR.
+		If you want to add non-pod data, not a huge problem, we just need to track and destruct them at the end of the scene.
+	**/
 	/** The projected shadow info for which we are rendering shadow depths. */
 	const FProjectedShadowInfo* ShadowInfo;
 
@@ -786,6 +789,16 @@ private:
 	*/
 	void RenderDepthInner(FRHICommandList& RHICmdList, class FSceneRenderer* SceneRenderer, const FViewInfo* FoundView, TFunctionRef<void (FRHICommandList& RHICmdList)> SetShadowRenderTargets, EShadowDepthRenderMode RenderMode );
 
+	/**
+	* Modifies the passed in view for this shadow
+	*/
+	void ModifyViewForShadow(FRHICommandList& RHICmdList, FViewInfo* FoundView);
+
+	/**
+	* Finds a relevant view for a shadow
+	*/
+	FViewInfo* FindViewForShadow(FSceneRenderer* SceneRenderer);
+
 	/**
 	* Renders the dynamic shadow subject depth, to a particular hacked view
 	*/
@@ -1085,7 +1098,7 @@ public:
 				FVector4(ShadowBufferSizeValue.X, ShadowBufferSizeValue.Y, 1.0f / ShadowBufferSizeValue.X, 1.0f / ShadowBufferSizeValue.Y));
 		}
 
-		FTexture2DRHIRef ShadowDepthTextureValue = GSceneRenderTargets.GetShadowDepthZTexture(ShadowInfo->bAllocatedInPreshadowCache);
+		FTexture2DRHIRef ShadowDepthTextureValue = FSceneRenderTargets::Get(RHICmdList).GetShadowDepthZTexture(ShadowInfo->bAllocatedInPreshadowCache);
 		FSamplerStateRHIParamRef DepthSamplerState = TStaticSamplerState<SF_Point,AM_Clamp,AM_Clamp,AM_Clamp>::GetRHI();
 
 		SetTextureParameter(RHICmdList, ShaderRHI, ShadowDepthTexture, ShadowDepthTextureSampler, DepthSamplerState, ShadowDepthTextureValue);		
@@ -1237,13 +1250,14 @@ public:
 
 	void Set(FRHICommandList& RHICmdList, FShader* Shader) const
 	{
+		FSceneRenderTargets& SceneContext = FSceneRenderTargets::Get(RHICmdList);
 		SetTextureParameter(
 			RHICmdList, 
 			Shader->GetPixelShader(),
 			TranslucencyShadowTransmission0,
 			TranslucencyShadowTransmission0Sampler,
 			TStaticSamplerState<SF_Bilinear,AM_Clamp,AM_Clamp,AM_Clamp>::GetRHI(),
-			GSceneRenderTargets.TranslucencyShadowTransmission[0]->GetRenderTargetItem().ShaderResourceTexture
+			SceneContext.TranslucencyShadowTransmission[0]->GetRenderTargetItem().ShaderResourceTexture
 			);
 
 		SetTextureParameter(
@@ -1252,7 +1266,7 @@ public:
 			TranslucencyShadowTransmission1,
 			TranslucencyShadowTransmission1Sampler,
 			TStaticSamplerState<SF_Bilinear,AM_Clamp,AM_Clamp,AM_Clamp>::GetRHI(),
-			GSceneRenderTargets.TranslucencyShadowTransmission[1]->GetRenderTargetItem().ShaderResourceTexture
+			SceneContext.TranslucencyShadowTransmission[1]->GetRenderTargetItem().ShaderResourceTexture
 			);
 	}
 
@@ -1347,7 +1361,7 @@ public:
 			RHICmdList, 
 			ShaderRHI, 
 			ShadowDepthTexture, 
-			GSceneRenderTargets.GetCubeShadowDepthZTexture(ShadowInfo->ResolutionX)
+			FSceneRenderTargets::Get(RHICmdList).GetCubeShadowDepthZTexture(ShadowInfo->ResolutionX)
 			);
 
 		if (ShadowDepthCubeComparisonSampler.IsBound())
diff --git a/Engine/Source/Runtime/Renderer/Private/ShadowSetup.cpp b/Engine/Source/Runtime/Renderer/Private/ShadowSetup.cpp
--- a/Engine/Source/Runtime/Renderer/Private/ShadowSetup.cpp
+++ b/Engine/Source/Runtime/Renderer/Private/ShadowSetup.cpp
@@ -538,7 +538,7 @@ void FProjectedShadowInfo::SetupWholeSceneProjection(
 		static bool bQuantize = true;
 		if ( bQuantize )
 		{
-			const int32 ShadowBufferResolution = GSceneRenderTargets.GetReflectiveShadowMapResolution();
+			const int32 ShadowBufferResolution = FSceneRenderTargets::Get_FrameConstantsOnly().GetReflectiveShadowMapResolution();
 			uint32 ShadowDepthBufferSizeX = ShadowBufferResolution;
 			uint32 ShadowDepthBufferSizeY = ShadowBufferResolution;
 			// Transform the shadow's position into shadowmap space
@@ -573,7 +573,7 @@ void FProjectedShadowInfo::SetupWholeSceneProjection(
 			MaxSubjectZ = FMath::Max(MaxSubjectZ, DepthRangeClamp);
 			MinSubjectZ = FMath::Min(MinSubjectZ, -DepthRangeClamp);
 
-			const FIntPoint ShadowBufferResolution = GSceneRenderTargets.GetShadowDepthTextureResolution();
+			const FIntPoint ShadowBufferResolution = FSceneRenderTargets::Get_FrameConstantsOnly().GetShadowDepthTextureResolution();
 			const uint32 ShadowDepthBufferSizeX = ShadowBufferResolution.X - SHADOW_BORDER * 2;
 			const uint32 ShadowDepthBufferSizeY = ShadowBufferResolution.Y - SHADOW_BORDER * 2;
 			// Transform the shadow's position into shadowmap space
@@ -863,9 +863,9 @@ void FProjectedShadowInfo::AddSubjectPrimitive(FPrimitiveSceneInfo* PrimitiveSce
 		if (bTranslucentRelevance && bShadowRelevance)
 		{
 			if (bTranslucentShadow)
-			{
-				SubjectTranslucentPrimitives.Add(PrimitiveSceneInfo);
-			}
+		{
+			SubjectTranslucentPrimitives.Add(PrimitiveSceneInfo);
+		}
 			else if (bReflectiveShadowmap)
 			{
 				if (PrimitiveSceneInfo->StaticMeshes.Num() > 0)
@@ -949,52 +949,48 @@ void FProjectedShadowInfo::GatherDynamicMeshElements(FSceneRenderer& Renderer, F
 
 		// Shadow must be visible in at least one view
 		if (bVisibleInAnyView)
-		{
-			check(FoundView && IsInRenderingThread());
-
-			// Backup properties of the view that we will override
-			FMatrix OriginalViewMatrix = FoundView->ViewMatrices.ViewMatrix;
-
-			// Override the view matrix so that billboarding primitives will be aligned to the light
-			FoundView->ViewMatrices.ViewMatrix = ShadowViewMatrix;
-
-			// Prevent materials from getting overridden during shadow casting in viewmodes like lighting only
-			// Lighting only should only affect the material used with direct lighting, not the indirect lighting
-			FoundView->bForceShowMaterials = true;
-
-			ReusedViewsArray[0] = FoundView;
-
-			check(!FoundView->ViewMatrices.GetDynamicMeshElementsShadowCullFrustum);
-
-			int32 Disable = 0; //CVarDisableCullShadows.GetValueOnRenderThread();
-			FConvexVolume NoCull;
-
-			if (IsWholeSceneDirectionalShadow())
-			{
-				FoundView->ViewMatrices.PreShadowTranslation = FVector(0,0,0);
-				FoundView->ViewMatrices.GetDynamicMeshElementsShadowCullFrustum = (Disable & 1) ? &NoCull : &CascadeSettings.ShadowBoundsAccurate;
-				GatherDynamicMeshElementsArray(FoundView, Renderer, SubjectPrimitives, DynamicSubjectMeshElements, ReusedViewsArray);
-				FoundView->ViewMatrices.PreShadowTranslation = PreShadowTranslation;
-			}
-			else
-			{
-				FoundView->ViewMatrices.PreShadowTranslation = PreShadowTranslation;
-				FoundView->ViewMatrices.GetDynamicMeshElementsShadowCullFrustum = (Disable & 1) ? &NoCull : &CasterFrustum;
-				GatherDynamicMeshElementsArray(FoundView, Renderer, SubjectPrimitives, DynamicSubjectMeshElements, ReusedViewsArray);
-			}
+		{		
+			check(FoundView && IsInRenderingThread() 
+				&& !FRHICommandListImmediate::AnyRenderThreadTasksOutstanding()); // this would be bad because they probably rely on the view we are hacking
 
-			FoundView->ViewMatrices.GetDynamicMeshElementsShadowCullFrustum = (Disable & 2) ? &NoCull : &ReceiverFrustum;
-			GatherDynamicMeshElementsArray(FoundView, Renderer, ReceiverPrimitives, DynamicReceiverMeshElements, ReusedViewsArray);
-
-			FoundView->ViewMatrices.GetDynamicMeshElementsShadowCullFrustum = (Disable & 4) ? &NoCull : &CasterFrustum;
-			GatherDynamicMeshElementsArray(FoundView, Renderer, SubjectTranslucentPrimitives, DynamicSubjectTranslucentMeshElements, ReusedViewsArray);
-
-			FoundView->ViewMatrices.GetDynamicMeshElementsShadowCullFrustum = nullptr;
-
-			FoundView->bForceShowMaterials = false;
-			FoundView->ViewMatrices.ViewMatrix = OriginalViewMatrix;
-		}
-	}
+		    // Backup properties of the view that we will override
+		    FMatrix OriginalViewMatrix = FoundView->ViewMatrices.ViewMatrix;
+    
+		    // Override the view matrix so that billboarding primitives will be aligned to the light
+		    FoundView->ViewMatrices.ViewMatrix = ShadowViewMatrix;
+    
+		    ReusedViewsArray[0] = FoundView;
+    
+		    check(!FoundView->ViewMatrices.GetDynamicMeshElementsShadowCullFrustum);
+    
+		    int32 Disable = 0; //CVarDisableCullShadows.GetValueOnRenderThread();
+		    FConvexVolume NoCull;
+    
+		    if (IsWholeSceneDirectionalShadow())
+		    {
+			    FoundView->ViewMatrices.PreShadowTranslation = FVector(0,0,0);
+			    FoundView->ViewMatrices.GetDynamicMeshElementsShadowCullFrustum = (Disable & 1) ? &NoCull : &CascadeSettings.ShadowBoundsAccurate;
+			    GatherDynamicMeshElementsArray(FoundView, Renderer, SubjectPrimitives, DynamicSubjectMeshElements, ReusedViewsArray);
+			    FoundView->ViewMatrices.PreShadowTranslation = PreShadowTranslation;
+		    }
+		    else
+		    {
+			    FoundView->ViewMatrices.PreShadowTranslation = PreShadowTranslation;
+			    FoundView->ViewMatrices.GetDynamicMeshElementsShadowCullFrustum = (Disable & 1) ? &NoCull : &CasterFrustum;
+			    GatherDynamicMeshElementsArray(FoundView, Renderer, SubjectPrimitives, DynamicSubjectMeshElements, ReusedViewsArray);
+		    }
+    
+		    FoundView->ViewMatrices.GetDynamicMeshElementsShadowCullFrustum = (Disable & 2) ? &NoCull : &ReceiverFrustum;
+		    GatherDynamicMeshElementsArray(FoundView, Renderer, ReceiverPrimitives, DynamicReceiverMeshElements, ReusedViewsArray);
+    
+		    FoundView->ViewMatrices.GetDynamicMeshElementsShadowCullFrustum = (Disable & 4) ? &NoCull : &CasterFrustum;
+		    GatherDynamicMeshElementsArray(FoundView, Renderer, SubjectTranslucentPrimitives, DynamicSubjectTranslucentMeshElements, ReusedViewsArray);
+    
+		    FoundView->ViewMatrices.GetDynamicMeshElementsShadowCullFrustum = nullptr;
+    
+		    FoundView->ViewMatrices.ViewMatrix = OriginalViewMatrix;
+	    }
+    }
 }
 
 void FProjectedShadowInfo::GatherDynamicMeshElementsArray(
@@ -1120,7 +1116,7 @@ struct FComparePreshadows
 };
 
 /** Removes stale shadows and attempts to add new preshadows to the cache. */
-void FDeferredShadingSceneRenderer::UpdatePreshadowCache()
+void FDeferredShadingSceneRenderer::UpdatePreshadowCache(FSceneRenderTargets& SceneContext)
 {
 	if (ShouldUseCachePreshadows() && !Views[0].bIsSceneCapture)
 	{
@@ -1128,7 +1124,7 @@ void FDeferredShadingSceneRenderer::UpdatePreshadowCache()
 		if (Scene->PreshadowCacheLayout.GetSizeX() == 0)
 		{
 			// Initialize the texture layout if necessary
-			const FIntPoint PreshadowCacheBufferSize = GSceneRenderTargets.GetPreShadowCacheTextureResolution();
+			const FIntPoint PreshadowCacheBufferSize = SceneContext.GetPreShadowCacheTextureResolution();
 			Scene->PreshadowCacheLayout = FTextureLayout(1, 1, PreshadowCacheBufferSize.X, PreshadowCacheBufferSize.Y, false, false);
 		}
 
@@ -1153,13 +1149,13 @@ void FDeferredShadingSceneRenderer::UpdatePreshadowCache()
 					CachedShadow->ResolutionY + SHADOW_BORDER * 2));
 				Scene->CachedPreshadows.RemoveAt(CachedShadowIndex);
 			}
-			else if (GSceneRenderTargets.bPreshadowCacheNewlyAllocated)
+			else if (SceneContext.bPreshadowCacheNewlyAllocated)
 			{
 				CachedShadow->bDepthsCached = false;
 			}
 		}
 
-		GSceneRenderTargets.bPreshadowCacheNewlyAllocated = false;
+		SceneContext.bPreshadowCacheNewlyAllocated = false;
 
 		TArray<TRefCountPtr<FProjectedShadowInfo>, SceneRenderingAllocator> UncachedPreShadows;
 
@@ -1340,11 +1336,12 @@ void FDeferredShadingSceneRenderer::CreatePerObjectProjectedShadow(
 		const FPrimitiveSceneInfo* ShadowChild = ShadowGroupPrimitives[ChildIndex];
 		OriginalBounds = OriginalBounds + ShadowChild->Proxy->GetBounds();
 	}
+	FSceneRenderTargets& SceneContext = FSceneRenderTargets::Get(RHICmdList);
 	
 	// Shadowing constants.
 	const uint32 MinShadowResolution = CVarMinShadowResolution.GetValueOnRenderThread();
 	const uint32 MaxShadowResolutionSetting = GetCachedScalabilityCVars().MaxShadowResolution;
-	const FIntPoint ShadowBufferResolution = GSceneRenderTargets.GetShadowDepthTextureResolution();
+	const FIntPoint ShadowBufferResolution = SceneContext.GetShadowDepthTextureResolution();
 	const uint32 MaxShadowResolution = FMath::Min<int32>(MaxShadowResolutionSetting, ShadowBufferResolution.X) - SHADOW_BORDER * 2;
 	const uint32 MaxShadowResolutionY = FMath::Min<int32>(MaxShadowResolutionSetting, ShadowBufferResolution.Y) - SHADOW_BORDER * 2;
 	const int32 ShadowFadeResolution = CVarShadowFadeResolution.GetValueOnRenderThread();
@@ -1485,8 +1482,8 @@ void FDeferredShadingSceneRenderer::CreatePerObjectProjectedShadow(
 					ShadowInitializer,
 					false,					// no preshadow
 					// Size was computed for the full res opaque shadow, convert to downsampled translucent shadow size with proper clamping
-					FMath::Clamp<int32>(SizeX / GSceneRenderTargets.GetTranslucentShadowDownsampleFactor(), 1, GSceneRenderTargets.GetTranslucentShadowDepthTextureResolution().X - SHADOW_BORDER * 2),
-					FMath::Clamp<int32>(MaxShadowResolutionY / GSceneRenderTargets.GetTranslucentShadowDownsampleFactor(), 1, GSceneRenderTargets.GetTranslucentShadowDepthTextureResolution().Y - SHADOW_BORDER * 2),
+					FMath::Clamp<int32>(SizeX / SceneContext.GetTranslucentShadowDownsampleFactor(), 1, SceneContext.GetTranslucentShadowDepthTextureResolution().X - SHADOW_BORDER * 2),
+					FMath::Clamp<int32>(MaxShadowResolutionY / SceneContext.GetTranslucentShadowDownsampleFactor(), 1, SceneContext.GetTranslucentShadowDepthTextureResolution().Y - SHADOW_BORDER * 2),
 					MaxScreenPercent,
 					true))					// translucent shadow
 				{
@@ -1521,7 +1518,7 @@ void FDeferredShadingSceneRenderer::CreatePerObjectProjectedShadow(
 			// Round down to the nearest power of two so that resolution changes are always doubling or halving the resolution, which increases filtering stability.
 			int32 PreshadowSizeX = 1 << (FMath::CeilLogTwo(FMath::TruncToInt(MaxDesiredResolution * CVarPreShadowResolutionFactor.GetValueOnRenderThread())) - 1);
 
-			const FIntPoint PreshadowCacheResolution = GSceneRenderTargets.GetPreShadowCacheTextureResolution();
+			const FIntPoint PreshadowCacheResolution = SceneContext.GetPreShadowCacheTextureResolution();
 			checkSlow(PreshadowSizeX <= PreshadowCacheResolution.X);
 			bool bIsOutsideWholeSceneShadow = true;
 
@@ -1602,13 +1599,15 @@ void FDeferredShadingSceneRenderer::CreateWholeSceneProjectedShadow(FLightSceneI
 	TArray<FWholeSceneProjectedShadowInitializer, TInlineAllocator<6> > ProjectedShadowInitializers;
 	if (LightSceneInfo->Proxy->GetWholeSceneProjectedShadowInitializer(ViewFamily, ProjectedShadowInitializers))
 	{
+		FSceneRenderTargets& SceneContext_ConstantsOnly = FSceneRenderTargets::Get_FrameConstantsOnly();
+
 		checkSlow(ProjectedShadowInitializers.Num() > 0);
 
 		// Shadow resolution constants.
 		const uint32 EffectiveDoubleShadowBorder = ProjectedShadowInitializers[0].CascadeSettings.bOnePassPointLightShadow ? 0 : SHADOW_BORDER * 2;
 		const int32 MinShadowResolution = CVarMinShadowResolution.GetValueOnRenderThread();
 		const int32 MaxShadowResolutionSetting = GetCachedScalabilityCVars().MaxShadowResolution;
-		const FIntPoint ShadowBufferResolution = GSceneRenderTargets.GetShadowDepthTextureResolution();
+		const FIntPoint ShadowBufferResolution = SceneContext_ConstantsOnly.GetShadowDepthTextureResolution();
 		const uint32 MaxShadowResolution = FMath::Min(MaxShadowResolutionSetting, ShadowBufferResolution.X) - EffectiveDoubleShadowBorder;
 		const uint32 MaxShadowResolutionY = FMath::Min(MaxShadowResolutionSetting, ShadowBufferResolution.Y) - EffectiveDoubleShadowBorder;
 		const int32 ShadowFadeResolution = CVarShadowFadeResolution.GetValueOnRenderThread();
@@ -1664,7 +1663,7 @@ void FDeferredShadingSceneRenderer::CreateWholeSceneProjectedShadow(FLightSceneI
 				if (ProjectedShadowInitializer.CascadeSettings.bOnePassPointLightShadow)
 				{
 					// Round to a resolution that is supported for one pass point light shadows
-					SizeX = SizeY = GSceneRenderTargets.GetCubeShadowDepthZResolution(GSceneRenderTargets.GetCubeShadowDepthZIndex(MaxDesiredResolution));
+					SizeX = SizeY = SceneContext_ConstantsOnly.GetCubeShadowDepthZResolution(SceneContext_ConstantsOnly.GetCubeShadowDepthZIndex(MaxDesiredResolution));
 				}
 
 				// Create the projected shadow info.
@@ -2199,6 +2198,9 @@ void FSceneRenderer::AddViewDependentWholeSceneShadowsForView(
 
 			checkSlow(INDEX_NONE == -1);
 
+			FSceneRenderTargets& SceneContext_ConstantsOnly = FSceneRenderTargets::Get_FrameConstantsOnly();
+
+
 			// todo: this code can be simplified by computing all the distances in one place - avoiding some redundant work and complexity
 			for (int32 Index = 0; Index < ProjectionCount; Index++)
 			{
@@ -2214,7 +2216,7 @@ void FSceneRenderer::AddViewDependentWholeSceneShadowsForView(
 
 				if (LightSceneInfo.Proxy->GetViewDependentWholeSceneProjectedShadowInitializer(View, LocalIndex, LightSceneInfo.IsPrecomputedLightingValid(), ProjectedShadowInitializer))
 				{
-					const FIntPoint ShadowBufferResolution = GSceneRenderTargets.GetShadowDepthTextureResolution();
+					const FIntPoint ShadowBufferResolution = SceneContext_ConstantsOnly.GetShadowDepthTextureResolution();
 					// Create the projected shadow info.
 					FProjectedShadowInfo* ProjectedShadowInfo = new(FMemStack::Get(), 1, 16) FProjectedShadowInfo;
 
@@ -2261,7 +2263,7 @@ void FSceneRenderer::AddViewDependentWholeSceneShadowsForView(
 						// moved out from the FProjectedShadowInfo constructor
 						ProjectedShadowInitializer.CascadeSettings.ShadowSplitIndex = 0;
 
-						const int32 ShadowBufferResolution = GSceneRenderTargets.GetReflectiveShadowMapResolution();
+						const int32 ShadowBufferResolution = SceneContext_ConstantsOnly.GetReflectiveShadowMapResolution();
 
 						// Create the projected shadow info.
 						FProjectedShadowInfo* ProjectedShadowInfo = new(FMemStack::Get(), 1, 16) FProjectedShadowInfo;
@@ -2332,8 +2334,9 @@ void FForwardShadingSceneRenderer::InitDynamicShadows(FRHICommandListImmediate&
 				const int32 NumShadows = FMath::Min(VisibleLightInfo.AllProjectedShadows.Num(), MAX_FORWARD_SHADOWCASCADES);
 				if (NumShadows > 0)
 				{
+					FSceneRenderTargets& SceneContext = FSceneRenderTargets::Get(RHICmdList);
 					//create the shadow depth texture and/or surface
-					const FIntPoint ShadowBufferResolution = GSceneRenderTargets.GetShadowDepthTextureResolution();
+					const FIntPoint ShadowBufferResolution = SceneContext.GetShadowDepthTextureResolution();
 					const int32 MaxWide = GMaxShadowDepthBufferSizeX / ShadowBufferResolution.X;
 					const int32 MaxHigh = GMaxShadowDepthBufferSizeY / ShadowBufferResolution.Y;
 
@@ -2342,7 +2345,7 @@ void FForwardShadingSceneRenderer::InitDynamicShadows(FRHICommandListImmediate&
 					const int32 NumHigh = FMath::Min(((NumShadows - 1) / MaxWide) + 1, MaxHigh);
 
 					const FIntPoint AtlasShadowBufferResolution(ShadowBufferResolution.X * NumWide, ShadowBufferResolution.Y * NumHigh);
-					GSceneRenderTargets.AllocateForwardShadingShadowDepthTarget(AtlasShadowBufferResolution);
+					SceneContext.AllocateForwardShadingShadowDepthTarget(AtlasShadowBufferResolution);
 
 
 					// Allocate atlas shadow texture space to the shadows.				
@@ -2476,7 +2479,7 @@ void FDeferredShadingSceneRenderer::InitDynamicShadows(FRHICommandListImmediate&
 	}
 
 	// Clear old preshadows and attempt to add new ones to the cache
-	UpdatePreshadowCache();
+	UpdatePreshadowCache(FSceneRenderTargets::Get(RHICmdList));
 
 	// Gathers the list of primitives used to draw various shadow types
 	GatherShadowPrimitives(PreShadows, ViewDependentWholeSceneShadowsThatNeedCulling, bStaticSceneOnly);
diff --git a/Engine/Source/Runtime/Renderer/Private/StaticMeshDrawList.inl b/Engine/Source/Runtime/Renderer/Private/StaticMeshDrawList.inl
--- a/Engine/Source/Runtime/Renderer/Private/StaticMeshDrawList.inl
+++ b/Engine/Source/Runtime/Renderer/Private/StaticMeshDrawList.inl
@@ -310,7 +310,6 @@ class FDrawVisibleAnyThreadTask
 	const int32 FirstPolicy;
 	const int32 LastPolicy;
 
-	bool& OutDirty;
 	const TArray<uint16, SceneRenderingAllocator>* PerDrawingPolicyCounts;
 
 public:
@@ -324,7 +323,6 @@ public:
 		const TArray<uint64, SceneRenderingAllocator>& InBatchVisibilityArray,
 		int32 InFirstPolicy,
 		int32 InLastPolicy,
-		bool& InOutDirty,
 		const TArray<uint16, SceneRenderingAllocator>* InPerDrawingPolicyCounts
 		)
 		: Caller(InCaller)
@@ -335,7 +333,6 @@ public:
 		, BatchVisibilityArray(InBatchVisibilityArray)
 		, FirstPolicy(InFirstPolicy)
 		, LastPolicy(InLastPolicy)
-		, OutDirty(InOutDirty)
 		, PerDrawingPolicyCounts(InPerDrawingPolicyCounts)
 	{
 	}
@@ -371,20 +368,14 @@ public:
 					{
 						BatchEnd++;
 					}
-					if (this->Caller.DrawVisibleInner(this->RHICmdList, this->View, this->PolicyContext, this->StaticMeshVisibilityMap, this->BatchVisibilityArray, Start, BatchEnd))
-					{
-						this->OutDirty = true;
-					}
+					this->Caller.DrawVisibleInner(this->RHICmdList, this->View, this->PolicyContext, this->StaticMeshVisibilityMap, this->BatchVisibilityArray, Start, BatchEnd);
 					Start = BatchEnd + 1;
 				}
 			}
 		}
 		else
 		{
-			if (this->Caller.DrawVisibleInner(this->RHICmdList, this->View, this->PolicyContext, this->StaticMeshVisibilityMap, this->BatchVisibilityArray, this->FirstPolicy, this->LastPolicy))
-			{
-				this->OutDirty = true;
-			}
+			this->Caller.DrawVisibleInner(this->RHICmdList, this->View, this->PolicyContext, this->StaticMeshVisibilityMap, this->BatchVisibilityArray, this->FirstPolicy, this->LastPolicy);
 		}
 		this->RHICmdList.HandleRTThreadTaskCompletion(MyCompletionGraphEvent);
 	}
@@ -408,7 +399,7 @@ void TStaticMeshDrawList<DrawingPolicyType>::DrawVisibleParallel(
 			TArray<uint16, SceneRenderingAllocator>& PerDrawingPolicyCounts = *new (FMemStack::Get()) TArray<uint16, SceneRenderingAllocator>();
 			PerDrawingPolicyCounts.AddZeroed(NumPolicies);
 
-			ParallelFor(NumPolicies, 
+			ParallelForWithPreWork(NumPolicies, 
 				[this, &PerDrawingPolicyCounts, &StaticMeshVisibilityMap, &BatchVisibilityArray](int32 Index)
 				{
 					int32 Count = 0;
@@ -431,6 +422,11 @@ void TStaticMeshDrawList<DrawingPolicyType>::DrawVisibleParallel(
 						// this is unlikely to overflow, but I don't think it would matter much if it did
 						PerDrawingPolicyCounts[Index] = (uint16)FMath::Min<int32>(Count, MAX_uint16);
 					}
+				},
+				[]()
+				{
+					QUICK_SCOPE_CYCLE_COUNTER(STAT_FDeferredShadingSceneRenderer_DrawVisibleParallel_ServiceLocalQueue);
+					FTaskGraphInterface::Get().ProcessThreadUntilIdle(ENamedThreads::RenderThread_Local);
 				}
 			);
 
@@ -487,8 +483,8 @@ void TStaticMeshDrawList<DrawingPolicyType>::DrawVisibleParallel(
 							{
 								// this is the last batch and it is small, let merge it
 								UE_CLOG(ParallelCommandListSet.bSpewBalance, LogTemp, Display, TEXT("    Index %d  BatchCount %d    (last merge)"), ParallelCommandListSet.NumParallelCommandLists(), PreviousBatchDraws + BatchCount);
-								FGraphEventRef AnyThreadCompletionEvent = TGraphTask<FDrawVisibleAnyThreadTask<DrawingPolicyType> >::CreateTask(nullptr, ENamedThreads::RenderThread)
-									.ConstructAndDispatchWhenReady(*this, *CmdList, ParallelCommandListSet.View, PolicyContext, StaticMeshVisibilityMap, BatchVisibilityArray, PreviousBatchStart, LastNonZeroPolicy, ParallelCommandListSet.OutDirty, &PerDrawingPolicyCounts);
+								FGraphEventRef AnyThreadCompletionEvent = TGraphTask<FDrawVisibleAnyThreadTask<DrawingPolicyType> >::CreateTask(ParallelCommandListSet.GetPrereqs(), ENamedThreads::RenderThread)
+									.ConstructAndDispatchWhenReady(*this, *CmdList, ParallelCommandListSet.View, PolicyContext, StaticMeshVisibilityMap, BatchVisibilityArray, PreviousBatchStart, LastNonZeroPolicy, &PerDrawingPolicyCounts);
 								ParallelCommandListSet.AddParallelCommandList(CmdList, AnyThreadCompletionEvent, PreviousBatchDraws + BatchCount);
 								PreviousBatchStart = -1;
 								PreviousBatchEnd = -2;
@@ -498,8 +494,8 @@ void TStaticMeshDrawList<DrawingPolicyType>::DrawVisibleParallel(
 							{
 								// this is a decent sized batch, emit the previous batch and save this one for possible merging
 								UE_CLOG(ParallelCommandListSet.bSpewBalance, LogTemp, Display, TEXT("    Index %d  BatchCount %d    "), ParallelCommandListSet.NumParallelCommandLists(), PreviousBatchDraws);
-								FGraphEventRef AnyThreadCompletionEvent = TGraphTask<FDrawVisibleAnyThreadTask<DrawingPolicyType> >::CreateTask(nullptr, ENamedThreads::RenderThread)
-									.ConstructAndDispatchWhenReady(*this, *CmdList, ParallelCommandListSet.View, PolicyContext, StaticMeshVisibilityMap, BatchVisibilityArray, PreviousBatchStart, PreviousBatchEnd, ParallelCommandListSet.OutDirty, &PerDrawingPolicyCounts);
+								FGraphEventRef AnyThreadCompletionEvent = TGraphTask<FDrawVisibleAnyThreadTask<DrawingPolicyType> >::CreateTask(ParallelCommandListSet.GetPrereqs(), ENamedThreads::RenderThread)
+									.ConstructAndDispatchWhenReady(*this, *CmdList, ParallelCommandListSet.View, PolicyContext, StaticMeshVisibilityMap, BatchVisibilityArray, PreviousBatchStart, PreviousBatchEnd, &PerDrawingPolicyCounts);
 								ParallelCommandListSet.AddParallelCommandList(CmdList, AnyThreadCompletionEvent, PreviousBatchDraws);
 								PreviousBatchStart = Start;
 								PreviousBatchEnd = LastNonZeroPolicy;
@@ -522,8 +518,8 @@ void TStaticMeshDrawList<DrawingPolicyType>::DrawVisibleParallel(
 			{
 				UE_CLOG(ParallelCommandListSet.bSpewBalance, LogTemp, Display, TEXT("    Index %d  BatchCount %d    (last)"), ParallelCommandListSet.NumParallelCommandLists(), PreviousBatchDraws);
 				FRHICommandList* CmdList = ParallelCommandListSet.NewParallelCommandList();
-				FGraphEventRef AnyThreadCompletionEvent = TGraphTask<FDrawVisibleAnyThreadTask<DrawingPolicyType> >::CreateTask(nullptr, ENamedThreads::RenderThread)
-					.ConstructAndDispatchWhenReady(*this, *CmdList, ParallelCommandListSet.View, PolicyContext, StaticMeshVisibilityMap, BatchVisibilityArray, PreviousBatchStart, PreviousBatchEnd, ParallelCommandListSet.OutDirty, &PerDrawingPolicyCounts);
+				FGraphEventRef AnyThreadCompletionEvent = TGraphTask<FDrawVisibleAnyThreadTask<DrawingPolicyType> >::CreateTask(ParallelCommandListSet.GetPrereqs(), ENamedThreads::RenderThread)
+					.ConstructAndDispatchWhenReady(*this, *CmdList, ParallelCommandListSet.View, PolicyContext, StaticMeshVisibilityMap, BatchVisibilityArray, PreviousBatchStart, PreviousBatchEnd, &PerDrawingPolicyCounts);
 				ParallelCommandListSet.AddParallelCommandList(CmdList, AnyThreadCompletionEvent, PreviousBatchDraws);
 			}
 		}
@@ -540,8 +536,8 @@ void TStaticMeshDrawList<DrawingPolicyType>::DrawVisibleParallel(
 				{
 					FRHICommandList* CmdList = ParallelCommandListSet.NewParallelCommandList();
 
-					FGraphEventRef AnyThreadCompletionEvent = TGraphTask<FDrawVisibleAnyThreadTask<DrawingPolicyType> >::CreateTask(nullptr, ENamedThreads::RenderThread)
-						.ConstructAndDispatchWhenReady(*this, *CmdList, ParallelCommandListSet.View, PolicyContext, StaticMeshVisibilityMap, BatchVisibilityArray, Start, Last, ParallelCommandListSet.OutDirty, nullptr);
+					FGraphEventRef AnyThreadCompletionEvent = TGraphTask<FDrawVisibleAnyThreadTask<DrawingPolicyType> >::CreateTask(ParallelCommandListSet.GetPrereqs(), ENamedThreads::RenderThread)
+						.ConstructAndDispatchWhenReady(*this, *CmdList, ParallelCommandListSet.View, PolicyContext, StaticMeshVisibilityMap, BatchVisibilityArray, Start, Last, nullptr);
 
 					ParallelCommandListSet.AddParallelCommandList(CmdList, AnyThreadCompletionEvent);
 				}
diff --git a/Engine/Source/Runtime/Renderer/Private/TiledDeferredLightRendering.cpp b/Engine/Source/Runtime/Renderer/Private/TiledDeferredLightRendering.cpp
--- a/Engine/Source/Runtime/Renderer/Private/TiledDeferredLightRendering.cpp
+++ b/Engine/Source/Runtime/Renderer/Private/TiledDeferredLightRendering.cpp
@@ -325,6 +325,7 @@ void FDeferredShadingSceneRenderer::RenderTiledDeferredLighting(FRHICommandListI
 
 	if (NumLightsToRender > 0)
 	{
+		FSceneRenderTargets& SceneContext = FSceneRenderTargets::Get(RHICmdList);
 		INC_DWORD_STAT_BY(STAT_NumLightsUsingTiledDeferred, NumLightsToRender);
 		INC_DWORD_STAT_BY(STAT_NumLightsUsingSimpleTiledDeferred, SimpleLights.InstanceData.Num());
 		SCOPE_CYCLE_COUNTER(STAT_DirectLightRenderingTime);
@@ -342,9 +343,9 @@ void FDeferredShadingSceneRenderer::RenderTiledDeferredLighting(FRHICommandListI
 			// One some hardware we can read and write from the same UAV with a 32 bit format. We don't do that yet.
 			TRefCountPtr<IPooledRenderTarget> OutTexture;
 			{
-				GSceneRenderTargets.ResolveSceneColor(RHICmdList, FResolveRect(0, 0, ViewFamily.FamilySizeX, ViewFamily.FamilySizeY));
+				SceneContext.ResolveSceneColor(RHICmdList, FResolveRect(0, 0, ViewFamily.FamilySizeX, ViewFamily.FamilySizeY));
 
-				FPooledRenderTargetDesc Desc = GSceneRenderTargets.GetSceneColor()->GetDesc();
+				FPooledRenderTargetDesc Desc = SceneContext.GetSceneColor()->GetDesc();
 				Desc.TargetableFlags |= TexCreate_UAV;
 
 				GRenderTargetPool.FindFreeElement( Desc, OutTexture, TEXT("SceneColorTiled") );
@@ -353,7 +354,7 @@ void FDeferredShadingSceneRenderer::RenderTiledDeferredLighting(FRHICommandListI
 			{
 				SCOPED_DRAW_EVENT(RHICmdList, TiledDeferredLighting);
 
-				IPooledRenderTarget& InTexture = *GSceneRenderTargets.GetSceneColor();
+				IPooledRenderTarget& InTexture = *SceneContext.GetSceneColor();
 
 				for (int32 ViewIndex = 0; ViewIndex < Views.Num(); ViewIndex++)
 				{
@@ -371,7 +372,7 @@ void FDeferredShadingSceneRenderer::RenderTiledDeferredLighting(FRHICommandListI
 			}
 
 			// swap with the former SceneColor
-			GSceneRenderTargets.SetSceneColor(OutTexture);
+			SceneContext.SetSceneColor(OutTexture);
 		}
 	}
 }
diff --git a/Engine/Source/Runtime/Renderer/Private/TranslucentLighting.cpp b/Engine/Source/Runtime/Renderer/Private/TranslucentLighting.cpp
--- a/Engine/Source/Runtime/Renderer/Private/TranslucentLighting.cpp
+++ b/Engine/Source/Runtime/Renderer/Private/TranslucentLighting.cpp
@@ -485,10 +485,12 @@ public:
 /** Renders shadow maps for translucent primitives. */
 void FProjectedShadowInfo::RenderTranslucencyDepths(FRHICommandList& RHICmdList, FDeferredShadingSceneRenderer* SceneRenderer)
 {
-	
+	check(IsInRenderingThread());
 	checkSlow(!bWholeSceneShadow);
 	SCOPE_CYCLE_COUNTER(STAT_RenderWholeSceneShadowDepthsTime);
 
+	bool bMakeViewSnapshot = FRHICommandListImmediate::AnyRenderThreadTasksOutstanding(); // if there are task in flight, we need to clone
+
 	// Choose an arbitrary view where this shadow's subject is relevant.
 	FViewInfo* FoundView = NULL;
 	for(int32 ViewIndex = 0;ViewIndex < SceneRenderer->Views.Num();ViewIndex++)
@@ -498,29 +500,41 @@ void FProjectedShadowInfo::RenderTranslucencyDepths(FRHICommandList& RHICmdList,
 		FPrimitiveViewRelevance ViewRel = VisibleLightViewInfo.ProjectedShadowViewRelevanceMap[ShadowId];
 		if (ViewRel.bShadowRelevance)
 		{
-			FoundView = CheckView;
+			if (bMakeViewSnapshot)
+			{
+				FoundView = CheckView->CreateSnapshot();
+			}
+			else
+			{
+				FoundView = CheckView;
+			}
 			break;
 		}
 	}
-	check(FoundView && IsInRenderingThread());
-	
-	// Backup properties of the view that we will override
-	TUniformBufferRef<FViewUniformShaderParameters> OriginalUniformBuffer = FoundView->UniformBuffer;
-	FMatrix OriginalViewMatrix = FoundView->ViewMatrices.ViewMatrix;
+	check(FoundView); 
+
+	TUniformBufferRef<FViewUniformShaderParameters> OriginalUniformBuffer;
+	FMatrix OriginalViewMatrix;
+
+	if (!bMakeViewSnapshot)
+	{
+		check(!FRHICommandListImmediate::AnyRenderThreadTasksOutstanding()); // we should add tasks that use the hacked view if it wasn't cloned.
+		// Backup properties of the view that we will override
+		OriginalUniformBuffer = FoundView->UniformBuffer;
+		OriginalViewMatrix = FoundView->ViewMatrices.ViewMatrix;
+	}
+
 
 	// Override the view matrix so that billboarding primitives will be aligned to the light
 	FoundView->ViewMatrices.ViewMatrix = ShadowViewMatrix;
 	FBox VolumeBounds[TVC_MAX];
 	FoundView->UniformBuffer = FoundView->CreateUniformBuffer(
+		RHICmdList,
 		nullptr,
 		ShadowViewMatrix, 
 		ShadowViewMatrix.Inverse(),
 		VolumeBounds,
 		TVC_MAX);
-		
-	// Prevent materials from getting overridden during shadow casting and in viewmodes like lighting only
-	// Lighting only should only affect the material used with direct lighting, not the indirect lighting
-	FoundView->bForceShowMaterials = true;
 
 	{
 #if WANTS_DRAW_MESH_EVENTS
@@ -529,8 +543,9 @@ void FProjectedShadowInfo::RenderTranslucencyDepths(FRHICommandList& RHICmdList,
 		SCOPED_DRAW_EVENTF(RHICmdList, EventShadowDepthActor, *EventName);
 #endif
 
-		FTextureRHIParamRef ShadowTransmission0 = GSceneRenderTargets.TranslucencyShadowTransmission[0]->GetRenderTargetItem().TargetableTexture;
-		FTextureRHIParamRef ShadowTransmission1 = GSceneRenderTargets.TranslucencyShadowTransmission[1]->GetRenderTargetItem().TargetableTexture;
+		FSceneRenderTargets& SceneContext = FSceneRenderTargets::Get(RHICmdList);
+		FTextureRHIParamRef ShadowTransmission0 = SceneContext.TranslucencyShadowTransmission[0]->GetRenderTargetItem().TargetableTexture;
+		FTextureRHIParamRef ShadowTransmission1 = SceneContext.TranslucencyShadowTransmission[1]->GetRenderTargetItem().TargetableTexture;
 		FTextureRHIParamRef RenderTargets[2] =
 		{
 			ShadowTransmission0,
@@ -609,9 +624,12 @@ void FProjectedShadowInfo::RenderTranslucencyDepths(FRHICommandList& RHICmdList,
 	}	
 
 	// Restore overridden properties
-	FoundView->bForceShowMaterials = false;
-	FoundView->UniformBuffer = OriginalUniformBuffer;
-	FoundView->ViewMatrices.ViewMatrix = OriginalViewMatrix;
+	if (!bMakeViewSnapshot)
+	{
+		check(!FRHICommandListImmediate::AnyRenderThreadTasksOutstanding()); // we should add tasks that use the hacked view if it wasn't cloned.
+		FoundView->UniformBuffer = OriginalUniformBuffer;
+		FoundView->ViewMatrices.ViewMatrix = OriginalViewMatrix;
+	}
 }
 
 IMPLEMENT_SHADER_TYPE(,FWriteToSliceGS,TEXT("TranslucentLightingShaders"),TEXT("WriteToSliceMainGS"),SF_Geometry);
@@ -641,6 +659,7 @@ public:
 
 	void SetParameters(FRHICommandList& RHICmdList, const FViewInfo& View, int32 VolumeCascadeIndex)
 	{
+		FSceneRenderTargets& SceneContext = FSceneRenderTargets::Get(RHICmdList);
 		const FPixelShaderRHIParamRef ShaderRHI = GetPixelShader();
 
 		FGlobalShader::SetParameters(RHICmdList, ShaderRHI, View);
@@ -652,7 +671,7 @@ public:
 			TranslucencyLightingVolumeAmbient, 
 			TranslucencyLightingVolumeAmbientSampler, 
 			TStaticSamplerState<SF_Bilinear, AM_Clamp, AM_Clamp, AM_Clamp>::GetRHI(), 
-			GSceneRenderTargets.TranslucencyLightingVolumeAmbient[VolumeCascadeIndex]->GetRenderTargetItem().ShaderResourceTexture);
+			SceneContext.TranslucencyLightingVolumeAmbient[VolumeCascadeIndex]->GetRenderTargetItem().ShaderResourceTexture);
 
 		SetTextureParameter(
 			RHICmdList,
@@ -660,7 +679,7 @@ public:
 			TranslucencyLightingVolumeDirectional, 
 			TranslucencyLightingVolumeDirectionalSampler, 
 			TStaticSamplerState<SF_Bilinear, AM_Clamp, AM_Clamp, AM_Clamp>::GetRHI(), 
-			GSceneRenderTargets.TranslucencyLightingVolumeDirectional[VolumeCascadeIndex]->GetRenderTargetItem().ShaderResourceTexture);
+			SceneContext.TranslucencyLightingVolumeDirectional[VolumeCascadeIndex]->GetRenderTargetItem().ShaderResourceTexture);
 	}
 
 	virtual bool Serialize(FArchive& Ar) override
@@ -907,7 +926,7 @@ public:
 				ShadowDepthTexture,
 				ShadowDepthTextureSampler,
 				TStaticSamplerState<SF_Point,AM_Clamp,AM_Clamp,AM_Clamp>::GetRHI(),
-				GSceneRenderTargets.GetShadowDepthZTexture()
+				FSceneRenderTargets::Get(RHICmdList).GetShadowDepthZTexture()
 				);
 		}
 
@@ -1108,15 +1127,16 @@ void FDeferredShadingSceneRenderer::ClearTranslucentVolumeLighting(FRHICommandLi
 	if (GUseTranslucentLightingVolumes && GSupportsVolumeTextureRendering)
 	{
 		SCOPED_DRAW_EVENT(RHICmdList, ClearTranslucentVolumeLighting);
+		FSceneRenderTargets& SceneContext = FSceneRenderTargets::Get(RHICmdList);
 
 		// Clear all volume textures in the same draw with MRT, which is faster than individually
 
 		static_assert(TVC_MAX == 2, "Only expecting two translucency lighting cascades.");
 		FTextureRHIParamRef RenderTargets[4];
-		RenderTargets[0] = GSceneRenderTargets.TranslucencyLightingVolumeAmbient[0]->GetRenderTargetItem().TargetableTexture;
-		RenderTargets[1] = GSceneRenderTargets.TranslucencyLightingVolumeDirectional[0]->GetRenderTargetItem().TargetableTexture;
-		RenderTargets[2] = GSceneRenderTargets.TranslucencyLightingVolumeAmbient[1]->GetRenderTargetItem().TargetableTexture;
-		RenderTargets[3] = GSceneRenderTargets.TranslucencyLightingVolumeDirectional[1]->GetRenderTargetItem().TargetableTexture;
+		RenderTargets[0] = SceneContext.TranslucencyLightingVolumeAmbient[0]->GetRenderTargetItem().TargetableTexture;
+		RenderTargets[1] = SceneContext.TranslucencyLightingVolumeDirectional[0]->GetRenderTargetItem().TargetableTexture;
+		RenderTargets[2] = SceneContext.TranslucencyLightingVolumeAmbient[1]->GetRenderTargetItem().TargetableTexture;
+		RenderTargets[3] = SceneContext.TranslucencyLightingVolumeDirectional[1]->GetRenderTargetItem().TargetableTexture;
 
 		FLinearColor ClearColors[4];
 		ClearColors[0] = FLinearColor(0, 0, 0, 0);
@@ -1189,10 +1209,11 @@ void FDeferredShadingSceneRenderer::InjectAmbientCubemapTranslucentVolumeLightin
 
 		auto ShaderMap = GetGlobalShaderMap(FeatureLevel);
 		
+		FSceneRenderTargets& SceneContext = FSceneRenderTargets::Get(RHICmdList);
 		for (int32 VolumeCascadeIndex = 0; VolumeCascadeIndex < TVC_MAX; VolumeCascadeIndex++)
 		{
 			// we don't update the directional volume (could be a HQ option)
-			SetRenderTarget(RHICmdList, GSceneRenderTargets.TranslucencyLightingVolumeAmbient[VolumeCascadeIndex]->GetRenderTargetItem().TargetableTexture, FTextureRHIRef());
+			SetRenderTarget(RHICmdList, SceneContext.TranslucencyLightingVolumeAmbient[VolumeCascadeIndex]->GetRenderTargetItem().TargetableTexture, FTextureRHIRef());
 
 			TShaderMapRef<FWriteToSliceVS> VertexShader(ShaderMap);
 			TShaderMapRef<FWriteToSliceGS> GeometryShader(ShaderMap);
@@ -1215,8 +1236,8 @@ void FDeferredShadingSceneRenderer::InjectAmbientCubemapTranslucentVolumeLightin
 				RasterizeToVolumeTexture(RHICmdList, VolumeBounds);
 			}
 			
-			RHICmdList.CopyToResolveTarget(GSceneRenderTargets.TranslucencyLightingVolumeAmbient[VolumeCascadeIndex]->GetRenderTargetItem().TargetableTexture,
-				GSceneRenderTargets.TranslucencyLightingVolumeAmbient[VolumeCascadeIndex]->GetRenderTargetItem().ShaderResourceTexture, true, FResolveParams());
+			RHICmdList.CopyToResolveTarget(SceneContext.TranslucencyLightingVolumeAmbient[VolumeCascadeIndex]->GetRenderTargetItem().TargetableTexture,
+				SceneContext.TranslucencyLightingVolumeAmbient[VolumeCascadeIndex]->GetRenderTargetItem().ShaderResourceTexture, true, FResolveParams());
 		}
 	}
 }
@@ -1225,12 +1246,13 @@ void FDeferredShadingSceneRenderer::ClearTranslucentVolumePerObjectShadowing(FRH
 {
 	if (GUseTranslucentLightingVolumes && GSupportsVolumeTextureRendering)
 	{
+		FSceneRenderTargets& SceneContext = FSceneRenderTargets::Get(RHICmdList);
 		SCOPED_DRAW_EVENT(RHICmdList, ClearTranslucentVolumePerLightShadowing);
 
 		static_assert(TVC_MAX == 2, "Only expecting two translucency lighting cascades.");
 		FTextureRHIParamRef RenderTargets[2];
-		RenderTargets[0] = GSceneRenderTargets.GetTranslucencyVolumeAmbient(TVC_Inner)->GetRenderTargetItem().TargetableTexture;
-		RenderTargets[1] = GSceneRenderTargets.GetTranslucencyVolumeDirectional(TVC_Inner)->GetRenderTargetItem().TargetableTexture;
+		RenderTargets[0] = SceneContext.GetTranslucencyVolumeAmbient(TVC_Inner)->GetRenderTargetItem().TargetableTexture;
+		RenderTargets[1] = SceneContext.GetTranslucencyVolumeDirectional(TVC_Inner)->GetRenderTargetItem().TargetableTexture;
 
 		FLinearColor ClearColors[2];
 		ClearColors[0] = FLinearColor(1, 1, 1, 1);
@@ -1290,6 +1312,7 @@ void FDeferredShadingSceneRenderer::AccumulateTranslucentVolumeObjectShadowing(F
 
 		auto ShaderMap = GetGlobalShaderMap(FeatureLevel);
 
+		FSceneRenderTargets& SceneContext = FSceneRenderTargets::Get(RHICmdList);
 		// Inject into each volume cascade
 		for (uint32 VolumeCascadeIndex = 0; VolumeCascadeIndex < TVC_MAX; VolumeCascadeIndex++)
 		{
@@ -1304,11 +1327,11 @@ void FDeferredShadingSceneRenderer::AccumulateTranslucentVolumeObjectShadowing(F
 
 				if (VolumeCascadeIndex == 0)
 				{
-					RenderTarget = GSceneRenderTargets.GetTranslucencyVolumeAmbient(TVC_Inner)->GetRenderTargetItem().TargetableTexture;
+					RenderTarget = SceneContext.GetTranslucencyVolumeAmbient(TVC_Inner)->GetRenderTargetItem().TargetableTexture;
 				}
 				else
 				{
-					RenderTarget = GSceneRenderTargets.GetTranslucencyVolumeDirectional(TVC_Inner)->GetRenderTargetItem().TargetableTexture;
+					RenderTarget = SceneContext.GetTranslucencyVolumeDirectional(TVC_Inner)->GetRenderTargetItem().TargetableTexture;
 				}
 
 				SetRenderTarget(RHICmdList, RenderTarget, FTextureRHIRef());
@@ -1324,8 +1347,8 @@ void FDeferredShadingSceneRenderer::AccumulateTranslucentVolumeObjectShadowing(F
 				
 				RasterizeToVolumeTexture(RHICmdList, VolumeBounds);
 
-				RHICmdList.CopyToResolveTarget(GSceneRenderTargets.GetTranslucencyVolumeAmbient((ETranslucencyVolumeCascade)VolumeCascadeIndex)->GetRenderTargetItem().TargetableTexture,
-					GSceneRenderTargets.GetTranslucencyVolumeAmbient((ETranslucencyVolumeCascade)VolumeCascadeIndex)->GetRenderTargetItem().ShaderResourceTexture, true, FResolveParams());
+				RHICmdList.CopyToResolveTarget(SceneContext.GetTranslucencyVolumeAmbient((ETranslucencyVolumeCascade)VolumeCascadeIndex)->GetRenderTargetItem().TargetableTexture,
+					SceneContext.GetTranslucencyVolumeAmbient((ETranslucencyVolumeCascade)VolumeCascadeIndex)->GetRenderTargetItem().ShaderResourceTexture, true, FResolveParams());
 			}
 		}
 	}
@@ -1498,14 +1521,15 @@ static void AddLightForInjection(
 /** Injects all the lights in LightInjectionData into the translucent lighting volume textures. */
 static void InjectTranslucentLightArray(FRHICommandListImmediate& RHICmdList, const FViewInfo& View, const TArray<FTranslucentLightInjectionData, SceneRenderingAllocator>& LightInjectionData)
 {
+	FSceneRenderTargets& SceneContext = FSceneRenderTargets::Get(RHICmdList);
 	INC_DWORD_STAT_BY(STAT_NumLightsInjectedIntoTranslucency, LightInjectionData.Num());
 
 	// Inject into each volume cascade
 	// Operate on one cascade at a time to reduce render target switches
 	for (uint32 VolumeCascadeIndex = 0; VolumeCascadeIndex < TVC_MAX; VolumeCascadeIndex++)
 	{
-		const IPooledRenderTarget* RT0 = GSceneRenderTargets.TranslucencyLightingVolumeAmbient[VolumeCascadeIndex];
-		const IPooledRenderTarget* RT1 = GSceneRenderTargets.TranslucencyLightingVolumeDirectional[VolumeCascadeIndex];
+		const IPooledRenderTarget* RT0 = SceneContext.TranslucencyLightingVolumeAmbient[VolumeCascadeIndex];
+		const IPooledRenderTarget* RT1 = SceneContext.TranslucencyLightingVolumeDirectional[VolumeCascadeIndex];
 
 		GRenderTargetPool.VisualizeTexture.SetCheckPoint(RHICmdList, RT0);
 		GRenderTargetPool.VisualizeTexture.SetCheckPoint(RHICmdList, RT1);
@@ -1709,13 +1733,14 @@ void FDeferredShadingSceneRenderer::InjectSimpleTranslucentVolumeLightingArray(F
 		const int32 ViewIndex = 0;
 
 		INC_DWORD_STAT_BY(STAT_NumLightsInjectedIntoTranslucency, NumLightsToInject);
+		FSceneRenderTargets& SceneContext = FSceneRenderTargets::Get(RHICmdList);
 
 		// Inject into each volume cascade
 		// Operate on one cascade at a time to reduce render target switches
 		for (int32 VolumeCascadeIndex = 0; VolumeCascadeIndex < TVC_MAX; VolumeCascadeIndex++)
 		{
-			const IPooledRenderTarget* RT0 = GSceneRenderTargets.TranslucencyLightingVolumeAmbient[VolumeCascadeIndex];
-			const IPooledRenderTarget* RT1 = GSceneRenderTargets.TranslucencyLightingVolumeDirectional[VolumeCascadeIndex];
+			const IPooledRenderTarget* RT0 = SceneContext.TranslucencyLightingVolumeAmbient[VolumeCascadeIndex];
+			const IPooledRenderTarget* RT1 = SceneContext.TranslucencyLightingVolumeDirectional[VolumeCascadeIndex];
 
 			GRenderTargetPool.VisualizeTexture.SetCheckPoint(RHICmdList, RT0);
 			GRenderTargetPool.VisualizeTexture.SetCheckPoint(RHICmdList, RT1);
@@ -1773,11 +1798,12 @@ void FDeferredShadingSceneRenderer::FilterTranslucentVolumeLighting(FRHICommandL
 {
 	if (GUseTranslucentLightingVolumes && GSupportsVolumeTextureRendering)
 	{
+		FSceneRenderTargets& SceneContext = FSceneRenderTargets::Get(RHICmdList);
 		// textures have to be finalized before reading.
 		for (int32 VolumeCascadeIndex = 0; VolumeCascadeIndex < TVC_MAX; VolumeCascadeIndex++)
 		{
-			const IPooledRenderTarget* RT0 = GSceneRenderTargets.TranslucencyLightingVolumeAmbient[VolumeCascadeIndex];
-			const IPooledRenderTarget* RT1 = GSceneRenderTargets.TranslucencyLightingVolumeDirectional[VolumeCascadeIndex];
+			const IPooledRenderTarget* RT0 = SceneContext.TranslucencyLightingVolumeAmbient[VolumeCascadeIndex];
+			const IPooledRenderTarget* RT1 = SceneContext.TranslucencyLightingVolumeDirectional[VolumeCascadeIndex];
 			FTextureRHIRef TargetTexture0 = RT0->GetRenderTargetItem().TargetableTexture;
 			FTextureRHIRef TargetTexture1 = RT1->GetRenderTargetItem().TargetableTexture;
 			RHICmdList.CopyToResolveTarget(TargetTexture0, TargetTexture0, true, FResolveParams());
@@ -1795,8 +1821,8 @@ void FDeferredShadingSceneRenderer::FilterTranslucentVolumeLighting(FRHICommandL
 			// Filter each cascade
 			for (int32 VolumeCascadeIndex = 0; VolumeCascadeIndex < TVC_MAX; VolumeCascadeIndex++)
 			{
-				const IPooledRenderTarget* RT0 = GSceneRenderTargets.GetTranslucencyVolumeAmbient((ETranslucencyVolumeCascade)VolumeCascadeIndex);
-				const IPooledRenderTarget* RT1 = GSceneRenderTargets.GetTranslucencyVolumeDirectional((ETranslucencyVolumeCascade)VolumeCascadeIndex);
+				const IPooledRenderTarget* RT0 = SceneContext.GetTranslucencyVolumeAmbient((ETranslucencyVolumeCascade)VolumeCascadeIndex);
+				const IPooledRenderTarget* RT1 = SceneContext.GetTranslucencyVolumeDirectional((ETranslucencyVolumeCascade)VolumeCascadeIndex);
 
 				GRenderTargetPool.VisualizeTexture.SetCheckPoint(RHICmdList, RT0);
 				GRenderTargetPool.VisualizeTexture.SetCheckPoint(RHICmdList, RT1);
