UE4 - Large upgrade to parallel rendering. No longer block per pass. Several RHI cmdlist passthrough functions added and implemented on the consoles

diff --git a/Engine/Source/Runtime/Renderer/Private/TranslucentRendering.cpp b/Engine/Source/Runtime/Renderer/Private/TranslucentRendering.cpp
--- a/Engine/Source/Runtime/Renderer/Private/TranslucentRendering.cpp
+++ b/Engine/Source/Runtime/Renderer/Private/TranslucentRendering.cpp
@@ -13,15 +13,16 @@
 
 static void SetTranslucentRenderTargetAndState(FRHICommandList& RHICmdList, const FViewInfo& View, bool bSeparateTranslucencyPass, bool bFirstTimeThisFrame = false)
 {
+	FSceneRenderTargets& SceneContext = FSceneRenderTargets::Get(RHICmdList);
 	bool bSetupTranslucentState = true;
 
-	if (bSeparateTranslucencyPass && GSceneRenderTargets.IsSeparateTranslucencyActive(View))
+	if (bSeparateTranslucencyPass && SceneContext.IsSeparateTranslucencyActive(View))
 	{
-		bSetupTranslucentState = GSceneRenderTargets.BeginRenderingSeparateTranslucency(RHICmdList, View, bFirstTimeThisFrame);
+		bSetupTranslucentState = SceneContext.BeginRenderingSeparateTranslucency(RHICmdList, View, bFirstTimeThisFrame);
 	}
 	else
 	{
-		GSceneRenderTargets.BeginRenderingTranslucency(RHICmdList, View);
+		SceneContext.BeginRenderingTranslucency(RHICmdList, View);
 	}
 
 	if (bSetupTranslucentState)
@@ -33,13 +34,14 @@ static void SetTranslucentRenderTargetAndState(FRHICommandList& RHICmdList, cons
 
 static void FinishTranslucentRenderTarget(FRHICommandListImmediate& RHICmdList, const FViewInfo& View, bool bSeparateTranslucencyPass)
 {
-	if (bSeparateTranslucencyPass && GSceneRenderTargets.IsSeparateTranslucencyActive(View))
+	FSceneRenderTargets& SceneContext = FSceneRenderTargets::Get(RHICmdList);
+	if (bSeparateTranslucencyPass && SceneContext.IsSeparateTranslucencyActive(View))
 	{
-		GSceneRenderTargets.FinishRenderingSeparateTranslucency(RHICmdList, View);
+		SceneContext.FinishRenderingSeparateTranslucency(RHICmdList, View);
 	}
 	else
 	{
-		GSceneRenderTargets.FinishRenderingTranslucency(RHICmdList, View);
+		SceneContext.FinishRenderingTranslucency(RHICmdList, View);
 	}
 }
 
@@ -181,15 +183,16 @@ FGlobalBoundShaderState CopySceneColorBoundShaderState;
 
 void FTranslucencyDrawingPolicyFactory::CopySceneColor(FRHICommandList& RHICmdList, const FViewInfo& View, const FPrimitiveSceneProxy* PrimitiveSceneProxy)
 {
+	FSceneRenderTargets& SceneContext = FSceneRenderTargets::Get(RHICmdList);
 
 	SCOPED_DRAW_EVENTF(RHICmdList, EventCopy, TEXT("CopySceneColor for %s %s"), *PrimitiveSceneProxy->GetOwnerName().ToString(), *PrimitiveSceneProxy->GetResourceName().ToString());
 	RHICmdList.SetRasterizerState(TStaticRasterizerState<FM_Solid, CM_None>::GetRHI());
 	RHICmdList.SetDepthStencilState(TStaticDepthStencilState<false, CF_Always>::GetRHI());
 	RHICmdList.SetBlendState(TStaticBlendState<>::GetRHI());
 
-	GSceneRenderTargets.ResolveSceneColor(RHICmdList);
+	SceneContext.ResolveSceneColor(RHICmdList);
 
-	GSceneRenderTargets.BeginRenderingLightAttenuation(RHICmdList);
+	SceneContext.BeginRenderingLightAttenuation(RHICmdList);
 	RHICmdList.SetViewport(View.ViewRect.Min.X, View.ViewRect.Min.Y, 0.0f, View.ViewRect.Max.X, View.ViewRect.Max.Y, 1.0f);
 
 	TShaderMapRef<FScreenVS> ScreenVertexShader(View.ShaderMap);
@@ -206,11 +209,11 @@ void FTranslucencyDrawingPolicyFactory::CopySceneColor(FRHICommandList& RHICmdLi
 		View.ViewRect.Min.X, View.ViewRect.Min.Y, 
 		View.ViewRect.Width(), View.ViewRect.Height(),
 		FIntPoint(View.ViewRect.Width(), View.ViewRect.Height()),
-		GSceneRenderTargets.GetBufferSizeXY(),
+		SceneContext.GetBufferSizeXY(),
 		*ScreenVertexShader,
 		EDRF_UseTriangleOptimization);
 
-	GSceneRenderTargets.FinishRenderingLightAttenuation(RHICmdList);
+	SceneContext.FinishRenderingLightAttenuation(RHICmdList);
 }
 
 /** The parameters used to draw a translucent mesh. */
@@ -401,6 +404,7 @@ bool FTranslucencyDrawingPolicyFactory::DrawMesh(
 				if (!RHICmdList.Bypass() && !IsInRenderingThread())
 				{
 					FRHICommandList* CmdList = new FRHICommandList;
+					CmdList->CopyRenderThreadContexts(RHICmdList);
 					FGraphEventRef RenderThreadCompletionEvent = TGraphTask<FCopySceneColorAndRestoreRenderThreadTask>::CreateTask().ConstructAndDispatchWhenReady(*CmdList, View, PrimitiveSceneProxy);
 					RHICmdList.QueueRenderThreadCommandListSubmit(RenderThreadCompletionEvent, CmdList);
 				}
@@ -626,6 +630,7 @@ void FTranslucentPrimSet::DrawPrimitivesParallel(
 			check(!IsInRenderingThread());
 			// can't do this in parallel, defer
 			FRHICommandList* CmdList = new FRHICommandList;
+			CmdList->CopyRenderThreadContexts(RHICmdList);
 			FGraphEventRef RenderThreadCompletionEvent = TGraphTask<FVolumetricTranslucentShadowRenderThreadTask>::CreateTask().ConstructAndDispatchWhenReady(*CmdList, *this, View, Renderer, bSeparateTranslucencyPass, PrimIdx);
 			RHICmdList.QueueRenderThreadCommandListSubmit(RenderThreadCompletionEvent, CmdList);
 		}
@@ -887,8 +892,8 @@ class FTranslucencyPassParallelCommandListSet : public FParallelCommandListSet
 	bool bSeparateTranslucency;
 	bool bFirstTimeThisFrame;
 public:
-	FTranslucencyPassParallelCommandListSet(const FViewInfo& InView, FRHICommandList& InParentCmdList, bool* InOutDirty, bool bInParallelExecute, bool bInSeparateTranslucency)
-		: FParallelCommandListSet(InView, InParentCmdList, InOutDirty, bInParallelExecute)
+	FTranslucencyPassParallelCommandListSet(const FViewInfo& InView, FRHICommandListImmediate& InParentCmdList, bool bInParallelExecute, bool bInCreateSceneContext, bool bInSeparateTranslucency)
+		: FParallelCommandListSet(InView, InParentCmdList, bInParallelExecute, bInCreateSceneContext)
 		, bSeparateTranslucency(bInSeparateTranslucency)
 		, bFirstTimeThisFrame(bInSeparateTranslucency)
 	{
@@ -912,24 +917,42 @@ static TAutoConsoleVariable<int32> CVarRHICmdTranslucencyPassDeferredContexts(
 	1,
 	TEXT("True to use deferred contexts to parallelize base pass command list execution."));
 
+static TAutoConsoleVariable<int32> CVarRHICmdFlushRenderThreadTasksTranslucentPass(
+	TEXT("r.RHICmdFlushRenderThreadTasksTranslucentPass"),
+	0,
+	TEXT("Wait for completion of parallel render thread tasks at the end of the translucent pass.  A more granular version of r.RHICmdFlushRenderThreadTasks. If either r.RHICmdFlushRenderThreadTasks or r.RHICmdFlushRenderThreadTasksTranslucentPass is > 0 we will flush."));
+
+// this is a static because we let the async tasks neyond the function
+static FTranslucencyDrawingPolicyFactory::ContextType GParallelTranslucencyContext;
+
 void FDeferredShadingSceneRenderer::RenderTranslucencyParallel(FRHICommandListImmediate& RHICmdList)
 {
-	GSceneRenderTargets.AllocLightAttenuation(); // materials will attempt to get this texture before the deferred command to set it up executes
+	FSceneRenderTargets& SceneContext = FSceneRenderTargets::Get(RHICmdList);
+	SceneContext.AllocLightAttenuation(); // materials will attempt to get this texture before the deferred command to set it up executes
 	check(IsInRenderingThread());
 
 
-	FScopedCommandListWaitForTasks Flusher(RHICmdList);
-	FTranslucencyDrawingPolicyFactory::ContextType ThisContext;
+	GParallelTranslucencyContext.TranslucentSelfShadow = nullptr;
+	GParallelTranslucencyContext.bSeparateTranslucencyPass = false;
+	GParallelTranslucencyContext.bSceneColorCopyIsUpToDate = false;
+	FScopedCommandListWaitForTasks Flusher(CVarRHICmdFlushRenderThreadTasksTranslucentPass.GetValueOnRenderThread() > 0 || CVarRHICmdFlushRenderThreadTasks.GetValueOnRenderThread() > 0, RHICmdList);
 
 	for (int32 ViewIndex = 0; ViewIndex < Views.Num(); ViewIndex++)
 	{
 		SCOPED_CONDITIONAL_DRAW_EVENTF(RHICmdList, EventView, Views.Num() > 1, TEXT("View%d"), ViewIndex);
 
 		const FViewInfo& View = Views[ViewIndex];
-
 		{
-			const bool bSeparateTranslucency = false;
-			FTranslucencyPassParallelCommandListSet ParallelCommandListSet(View, RHICmdList, nullptr, CVarRHICmdTranslucencyPassDeferredContexts.GetValueOnRenderThread() > 0, bSeparateTranslucency);
+			if (SceneContext.IsSeparateTranslucencyActive(View))
+			{
+				FSceneViewState* ViewState = (FSceneViewState*)View.State;
+				// we need to allocate this now so it ends up in the snapshot
+				ViewState->GetSeparateTranslucency(SceneContext.GetBufferSizeXY());
+			}
+			FTranslucencyPassParallelCommandListSet ParallelCommandListSet(View, RHICmdList, 
+				CVarRHICmdTranslucencyPassDeferredContexts.GetValueOnRenderThread() > 0, 
+				CVarRHICmdFlushRenderThreadTasksTranslucentPass.GetValueOnRenderThread() == 0  && CVarRHICmdFlushRenderThreadTasks.GetValueOnRenderThread() == 0, 
+				false);
 
 			{
 				int32 NumPrims = View.TranslucentPrimSet.NumPrims() - View.TranslucentPrimSet.NumSeparateTranslucencyPrims();
@@ -951,7 +974,7 @@ void FDeferredShadingSceneRenderer::RenderTranslucencyParallel(FRHICommandListIm
 						{
 							FRHICommandList* CmdList = ParallelCommandListSet.NewParallelCommandList();
 
-							FGraphEventRef AnyThreadCompletionEvent = TGraphTask<FDrawSortedTransAnyThreadTask>::CreateTask(nullptr, ENamedThreads::RenderThread)
+							FGraphEventRef AnyThreadCompletionEvent = TGraphTask<FDrawSortedTransAnyThreadTask>::CreateTask(ParallelCommandListSet.GetPrereqs(), ENamedThreads::RenderThread)
 								.ConstructAndDispatchWhenReady(*this, *CmdList, View, false, Start, Last);
 
 							ParallelCommandListSet.AddParallelCommandList(CmdList, AnyThreadCompletionEvent);
@@ -961,12 +984,12 @@ void FDeferredShadingSceneRenderer::RenderTranslucencyParallel(FRHICommandListIm
 				}
 			}
 			// Draw the view's mesh elements with the translucent drawing policy.
-			DrawViewElementsParallel<FTranslucencyDrawingPolicyFactory>(ThisContext, SDPG_World, false, ParallelCommandListSet);
+			DrawViewElementsParallel<FTranslucencyDrawingPolicyFactory>(GParallelTranslucencyContext, SDPG_World, false, ParallelCommandListSet);
 			// Draw the view's mesh elements with the translucent drawing policy.
-			DrawViewElementsParallel<FTranslucencyDrawingPolicyFactory>(ThisContext, SDPG_Foreground, false, ParallelCommandListSet);
+			DrawViewElementsParallel<FTranslucencyDrawingPolicyFactory>(GParallelTranslucencyContext, SDPG_Foreground, false, ParallelCommandListSet);
 
-			FinishTranslucentRenderTarget(RHICmdList, View, bSeparateTranslucency);
 		}
+		FinishTranslucentRenderTarget(RHICmdList, View, false);
 
 #if 0 // unsupported visualization in the parallel case
 		const FSceneViewState* ViewState = (const FSceneViewState*)View.State;
@@ -983,7 +1006,11 @@ void FDeferredShadingSceneRenderer::RenderTranslucencyParallel(FRHICommandListIm
 		{
 			{
 				// always call BeginRenderingSeparateTranslucency() even if there are no primitives to we keep the RT allocated
-				FTranslucencyPassParallelCommandListSet ParallelCommandListSet(View, RHICmdList, nullptr, CVarRHICmdTranslucencyPassDeferredContexts.GetValueOnRenderThread() > 0, true);
+				FTranslucencyPassParallelCommandListSet ParallelCommandListSet(View, 
+					RHICmdList, 
+					CVarRHICmdTranslucencyPassDeferredContexts.GetValueOnRenderThread() > 0, 
+					CVarRHICmdFlushRenderThreadTasksTranslucentPass.GetValueOnRenderThread() == 0 && CVarRHICmdFlushRenderThreadTasks.GetValueOnRenderThread() == 0, 
+					true);
 
 				// Draw only translucent prims that are in the SeparateTranslucency pass
 				if (View.TranslucentPrimSet.NumSeparateTranslucencyPrims() > 0)
@@ -1006,7 +1033,7 @@ void FDeferredShadingSceneRenderer::RenderTranslucencyParallel(FRHICommandListIm
 							{
 								FRHICommandList* CmdList = ParallelCommandListSet.NewParallelCommandList();
 
-								FGraphEventRef AnyThreadCompletionEvent = TGraphTask<FDrawSortedTransAnyThreadTask>::CreateTask(nullptr, ENamedThreads::RenderThread)
+								FGraphEventRef AnyThreadCompletionEvent = TGraphTask<FDrawSortedTransAnyThreadTask>::CreateTask(ParallelCommandListSet.GetPrereqs(), ENamedThreads::RenderThread)
 									.ConstructAndDispatchWhenReady(*this, *CmdList, View, true, Start, Last);
 
 								ParallelCommandListSet.AddParallelCommandList(CmdList, AnyThreadCompletionEvent);
@@ -1016,7 +1043,7 @@ void FDeferredShadingSceneRenderer::RenderTranslucencyParallel(FRHICommandListIm
 					}
 				}
 			}
-			GSceneRenderTargets.FinishRenderingSeparateTranslucency(RHICmdList, View);
+			SceneContext.FinishRenderingSeparateTranslucency(RHICmdList, View);
 		}
 	}
 }
@@ -1071,10 +1098,11 @@ void FDeferredShadingSceneRenderer::RenderTranslucency(FRHICommandListImmediate&
 			FinishTranslucentRenderTarget(RHICmdList, View, false);
 			
 			{
+				FSceneRenderTargets& SceneContext = FSceneRenderTargets::Get(RHICmdList);
 				bool bRenderSeparateTranslucency = View.TranslucentPrimSet.NumSeparateTranslucencyPrims() > 0;
 
 				// always call BeginRenderingSeparateTranslucency() even if there are no primitives to we keep the RT allocated
-				bool bSetupTranslucency = GSceneRenderTargets.BeginRenderingSeparateTranslucency(RHICmdList, View, true);
+				bool bSetupTranslucency = SceneContext.BeginRenderingSeparateTranslucency(RHICmdList, View, true);
 
 				// Draw only translucent prims that are in the SeparateTranslucency pass
 				if (bRenderSeparateTranslucency)
@@ -1087,7 +1115,7 @@ void FDeferredShadingSceneRenderer::RenderTranslucency(FRHICommandListImmediate&
 					View.TranslucentPrimSet.DrawPrimitives(RHICmdList, View, *this, true);
 				}
 
-				GSceneRenderTargets.FinishRenderingSeparateTranslucency(RHICmdList, View);
+				SceneContext.FinishRenderingSeparateTranslucency(RHICmdList, View);
 			}
 		}
 	}
diff --git a/Engine/Source/Runtime/Renderer/Private/VelocityRendering.cpp b/Engine/Source/Runtime/Renderer/Private/VelocityRendering.cpp
--- a/Engine/Source/Runtime/Renderer/Private/VelocityRendering.cpp
+++ b/Engine/Source/Runtime/Renderer/Private/VelocityRendering.cpp
@@ -581,7 +581,7 @@ public:
 static void BeginVelocityRendering(FRHICommandList& RHICmdList, TRefCountPtr<IPooledRenderTarget>& VelocityRT, bool bPerformClear)
 {
 	FTextureRHIRef VelocityTexture = VelocityRT->GetRenderTargetItem().TargetableTexture;
-	FTexture2DRHIRef DepthTexture = GSceneRenderTargets.GetSceneDepthTexture();
+	FTexture2DRHIRef DepthTexture = FSceneRenderTargets::Get(RHICmdList).GetSceneDepthTexture();
 	FLinearColor VelocityClearColor = FLinearColor::Black;
 	if (bPerformClear)
 	{
@@ -607,7 +607,7 @@ static void BeginVelocityRendering(FRHICommandList& RHICmdList, TRefCountPtr<IPo
 
 static void SetVelocitiesState(FRHICommandList& RHICmdList, const FViewInfo& View, TRefCountPtr<IPooledRenderTarget>& VelocityRT)
 {
-	const FIntPoint BufferSize = GSceneRenderTargets.GetBufferSizeXY();
+	const FIntPoint BufferSize = FSceneRenderTargets::Get(RHICmdList).GetBufferSizeXY();
 	const FIntPoint VelocityBufferSize = BufferSize;		// full resolution so we can reuse the existing full res z buffer
 
 	const uint32 MinX = View.ViewRect.Min.X * VelocityBufferSize.X / BufferSize.X;
@@ -626,8 +626,8 @@ class FVelocityPassParallelCommandListSet : public FParallelCommandListSet
 {
 	TRefCountPtr<IPooledRenderTarget>& VelocityRT;
 public:
-	FVelocityPassParallelCommandListSet(const FViewInfo& InView, FRHICommandList& InParentCmdList, bool* InOutDirty, bool bInParallelExecute, TRefCountPtr<IPooledRenderTarget>& InVelocityRT)
-		: FParallelCommandListSet(InView, InParentCmdList, InOutDirty, bInParallelExecute)
+	FVelocityPassParallelCommandListSet(const FViewInfo& InView, FRHICommandListImmediate& InParentCmdList, bool bInParallelExecute, bool bInCreateSceneContext, TRefCountPtr<IPooledRenderTarget>& InVelocityRT)
+		: FParallelCommandListSet(InView, InParentCmdList, bInParallelExecute, bInCreateSceneContext)
 		, VelocityRT(InVelocityRT)
 	{
 		SetStateOnCommandList(ParentCmdList);
@@ -645,16 +645,25 @@ public:
 	}
 };
 
+static TAutoConsoleVariable<int32> CVarRHICmdFlushRenderThreadTasksVelocityPass(
+	TEXT("r.RHICmdFlushRenderThreadTasksVelocityPass"),
+	0,
+	TEXT("Wait for completion of parallel render thread tasks at the end of the velocity pass.  A more granular version of r.RHICmdFlushRenderThreadTasks. If either r.RHICmdFlushRenderThreadTasks or r.RHICmdFlushRenderThreadTasksVelocityPass is > 0 we will flush."));
+
 void FDeferredShadingSceneRenderer::RenderVelocitiesInnerParallel(FRHICommandListImmediate& RHICmdList, TRefCountPtr<IPooledRenderTarget>& VelocityRT)
 {
 	// parallel version
-	FScopedCommandListWaitForTasks Flusher(RHICmdList);
+	FScopedCommandListWaitForTasks Flusher(CVarRHICmdFlushRenderThreadTasksVelocityPass.GetValueOnRenderThread() > 0 || CVarRHICmdFlushRenderThreadTasks.GetValueOnRenderThread() > 0, RHICmdList);
 
 	for(int32 ViewIndex = 0; ViewIndex < Views.Num(); ViewIndex++)
 	{
 		const FViewInfo& View = Views[ViewIndex];
 
-		FVelocityPassParallelCommandListSet ParallelCommandListSet(View, RHICmdList, nullptr, CVarRHICmdVelocityPassDeferredContexts.GetValueOnRenderThread() > 0, VelocityRT);
+		FVelocityPassParallelCommandListSet ParallelCommandListSet(View, 
+			RHICmdList, 
+			CVarRHICmdVelocityPassDeferredContexts.GetValueOnRenderThread() > 0, 
+			CVarRHICmdFlushRenderThreadTasksVelocityPass.GetValueOnRenderThread() == 0 && CVarRHICmdFlushRenderThreadTasks.GetValueOnRenderThread() == 0, 
+			VelocityRT);
 
 		Scene->VelocityDrawList.DrawVisibleParallel(View.StaticMeshVelocityMap, View.StaticMeshBatchVisibility, ParallelCommandListSet);
 
@@ -677,7 +686,7 @@ void FDeferredShadingSceneRenderer::RenderVelocitiesInnerParallel(FRHICommandLis
 
 				FRHICommandList* CmdList = ParallelCommandListSet.NewParallelCommandList();
 
-				FGraphEventRef AnyThreadCompletionEvent = TGraphTask<FRenderVelocityDynamicThreadTask>::CreateTask(nullptr, ENamedThreads::RenderThread)
+				FGraphEventRef AnyThreadCompletionEvent = TGraphTask<FRenderVelocityDynamicThreadTask>::CreateTask(ParallelCommandListSet.GetPrereqs(), ENamedThreads::RenderThread)
 					.ConstructAndDispatchWhenReady(*this, *CmdList, View, Start, Last);
 
 				ParallelCommandListSet.AddParallelCommandList(CmdList, AnyThreadCompletionEvent);
@@ -752,13 +761,15 @@ void FDeferredShadingSceneRenderer::RenderVelocities(FRHICommandListImmediate& R
 	}
 
 	{
-		GPrevPerBoneMotionBlur.StartAppend(ViewFamily.bWorldIsPaused);
+		FSceneRenderTargets::Get(RHICmdList).SetVelocityPass(true);
+		GPrevPerBoneMotionBlur.StartAppend(RHICmdList, ViewFamily.bWorldIsPaused);
 
 		BeginVelocityRendering(RHICmdList, VelocityRT, true);
 
 		if (IsParallelVelocity())
 		{
 			RenderVelocitiesInnerParallel(RHICmdList, VelocityRT);
+			GPrevPerBoneMotionBlur.EndAppendFence(RHICmdList);
 		}
 		else
 		{
@@ -766,8 +777,7 @@ void FDeferredShadingSceneRenderer::RenderVelocities(FRHICommandListImmediate& R
 		}
 
 		RHICmdList.CopyToResolveTarget(VelocityRT->GetRenderTargetItem().TargetableTexture, VelocityRT->GetRenderTargetItem().ShaderResourceTexture, false, FResolveParams());
-
-		GPrevPerBoneMotionBlur.EndAppend();
+		FSceneRenderTargets::Get(RHICmdList).SetVelocityPass(false);
 	}
 
 	// restore any color write state changes
@@ -780,7 +790,7 @@ void FDeferredShadingSceneRenderer::RenderVelocities(FRHICommandListImmediate& R
 
 FPooledRenderTargetDesc FVelocityRendering::GetRenderTargetDesc()
 {
-	const FIntPoint BufferSize = GSceneRenderTargets.GetBufferSizeXY();
+	const FIntPoint BufferSize = FSceneRenderTargets::Get_FrameConstantsOnly().GetBufferSizeXY();
 	const FIntPoint VelocityBufferSize = BufferSize;		// full resolution so we can reuse the existing full res z buffer
 	return FPooledRenderTargetDesc(FPooledRenderTargetDesc::Create2DDesc(VelocityBufferSize, PF_G16R16, TexCreate_None, TexCreate_RenderTargetable, false));
 }
