Fixed non-unity build.

diff --git a/Engine/Source/Runtime/Renderer/Private/ReflectionEnvironment.cpp b/Engine/Source/Runtime/Renderer/Private/ReflectionEnvironment.cpp
--- a/Engine/Source/Runtime/Renderer/Private/ReflectionEnvironment.cpp
+++ b/Engine/Source/Runtime/Renderer/Private/ReflectionEnvironment.cpp
@@ -18,6 +18,7 @@
 #include "ShaderParameterUtils.h"
 #include "LightRendering.h"
 #include "SceneUtils.h"
+#include "LightPropagationVolumeBlendable.h"
 
 /** Tile size for the reflection environment compute shader, tweaked for 680 GTX. */
 const int32 GReflectionEnvironmentTileSizeX = 16;
