UE4 - fix for shadow culling on foliage

diff --git a/Engine/Source/Runtime/Renderer/Private/ShadowSetup.cpp b/Engine/Source/Runtime/Renderer/Private/ShadowSetup.cpp
--- a/Engine/Source/Runtime/Renderer/Private/ShadowSetup.cpp
+++ b/Engine/Source/Runtime/Renderer/Private/ShadowSetup.cpp
@@ -1714,6 +1714,9 @@ void FDeferredShadingSceneRenderer::CreateWholeSceneProjectedShadow(FLightSceneI
 					ProjectedShadowInfo->OnePassShadowFrustums.Empty(6);
 					ProjectedShadowInfo->OnePassShadowFrustums.AddZeroed(6);
 					const FMatrix ScaleMatrix = FScaleMatrix(FVector(1, -1, 1));
+
+					// fill in the caster frustum with the far plane from every face
+					ProjectedShadowInfo->CasterFrustum.Planes.Empty();
 					for (int32 FaceIndex = 0; FaceIndex < 6; FaceIndex++)
 					{
 						// Create a view projection matrix for each cube face
@@ -1721,7 +1724,15 @@ void FDeferredShadingSceneRenderer::CreateWholeSceneProjectedShadow(FLightSceneI
 						ProjectedShadowInfo->OnePassShadowViewProjectionMatrices.Add(ShadowViewProjectionMatrix);
 						// Create a convex volume out of the frustum so it can be used for object culling
 						GetViewFrustumBounds(ProjectedShadowInfo->OnePassShadowFrustums[FaceIndex], ShadowViewProjectionMatrix, false);
+
+						// We are assuming here that the last plane is the far plane
+						// we need to incorporate PreShadowTranslation (so it can be disincorporated later)
+						FPlane Src = ProjectedShadowInfo->OnePassShadowFrustums[FaceIndex].Planes.Last();
+						// add world space preview translation
+						Src.W += (FVector(Src) | ProjectedShadowInfo->PreShadowTranslation);
+						ProjectedShadowInfo->CasterFrustum.Planes.Add(Src);
 					}
+					ProjectedShadowInfo->CasterFrustum.Init();
 				}
 
 				// Ray traced shadows use the GPU managed distance field object buffers, no CPU culling should be used
