Fix: Materials that use Pixel Depth Offset weren't using the pixel shader during the prepass so the depths didn't match up during the base pass.

diff --git a/Engine/Source/Runtime/Renderer/Private/DepthRendering.cpp b/Engine/Source/Runtime/Renderer/Private/DepthRendering.cpp
--- a/Engine/Source/Runtime/Renderer/Private/DepthRendering.cpp
+++ b/Engine/Source/Runtime/Renderer/Private/DepthRendering.cpp
@@ -106,7 +106,7 @@ public:
 	static bool ShouldCache(EShaderPlatform Platform,const FMaterial* Material,const FVertexFactoryType* VertexFactoryType)
 	{
 		// Compile for materials that are masked.
-		return !Material->WritesEveryPixel() && IsFeatureLevelSupported(Platform, ERHIFeatureLevel::SM4);
+		return (!Material->WritesEveryPixel() || Material->HasPixelDepthOffsetConnected()) && IsFeatureLevelSupported(Platform, ERHIFeatureLevel::SM4);
 	}
 
 	FDepthOnlyPS(const ShaderMetaType::CompiledShaderInitializerType& Initializer):
@@ -137,9 +137,10 @@ FDepthDrawingPolicy::FDepthDrawingPolicy(
 	):
 	FMeshDrawingPolicy(InVertexFactory,InMaterialRenderProxy,InMaterialResource,false,/*bInTwoSidedOverride=*/ bIsTwoSided)
 {
-	// The primitive needs to be rendered with the material's pixel and vertex shaders if it is masked
+	// The primitive needs to be rendered with the material's pixel and vertex shaders if it is masked or if it is offsetting
+	// the pixel depth
 	bNeedsPixelShader = false;
-	if (!InMaterialResource.WritesEveryPixel())
+	if (!InMaterialResource.WritesEveryPixel() || InMaterialResource.HasPixelDepthOffsetConnected())
 	{
 		bNeedsPixelShader = true;
 		PixelShader = InMaterialResource.GetShader<FDepthOnlyPS>(InVertexFactory->GetType());
