Fixes a crash with parrallel rendering in the GPU benchmark, caused by a GetRenderQueryResult() call without flushing the command list before, leading to some EndRenderQuery() call not being executed.

diff --git a/Engine/Source/Runtime/Renderer/Private/GPUBenchmark.cpp b/Engine/Source/Runtime/Renderer/Private/GPUBenchmark.cpp
--- a/Engine/Source/Runtime/Renderer/Private/GPUBenchmark.cpp
+++ b/Engine/Source/Runtime/Renderer/Private/GPUBenchmark.cpp
@@ -474,6 +474,8 @@ void RendererGPUBenchmark(FRHICommandListImmediate& RHICmdList, FSynthBenchmarkR
 
 		{
 			uint64 OldAbsTime = 0;
+			// flushes the RHI thread to make sure all RHICmdList.EndRenderQuery() commands got executed.
+			RHICmdList.ImmediateFlush(EImmediateFlushType::FlushRHIThread);
 			RHICmdList.GetRenderQueryResult(TimerQueries[0], OldAbsTime, true);
 			GTimerQueryPool.ReleaseQuery(RHICmdList, TimerQueries[0]);
 
