Fixing rendering crash when launching Kite Demo

diff --git a/Engine/Source/Runtime/Renderer/Private/DeferredShadingRenderer.cpp b/Engine/Source/Runtime/Renderer/Private/DeferredShadingRenderer.cpp
--- a/Engine/Source/Runtime/Renderer/Private/DeferredShadingRenderer.cpp
+++ b/Engine/Source/Runtime/Renderer/Private/DeferredShadingRenderer.cpp
@@ -1514,22 +1514,21 @@ bool FDeferredShadingSceneRenderer::RenderBasePass(FRHICommandListImmediate& RHI
 		if (GRHICommandList.UseParallelAlgorithms() && CVarParallelBasePass.GetValueOnRenderThread())
 		{
 			FScopedCommandListWaitForTasks Flusher(CVarRHICmdFlushRenderThreadTasksBasePass.GetValueOnRenderThread() > 0 || CVarRHICmdFlushRenderThreadTasks.GetValueOnRenderThread() > 0, RHICmdList);
-			for(int32 ViewIndex = 0;ViewIndex < Views.Num();ViewIndex++)
+			for (int32 ViewIndex = 0; ViewIndex < Views.Num(); ViewIndex++)
 			{
 				SCOPED_CONDITIONAL_DRAW_EVENTF(RHICmdList, EventView, Views.Num() > 1, TEXT("View%d"), ViewIndex);
 				FViewInfo& View = Views[ViewIndex];
 				RenderBasePassViewParallel(View, RHICmdList);
-				bDirty = true; // assume dirty since we are not going to wait
-				if (FVelocityRendering::OutputsToGBuffer())
-				{
-					GPrevPerBoneMotionBlur.EndAppendFence(RHICmdList);
-					FSceneRenderTargets::Get(RHICmdList).SetVelocityPass(false);
-				}
+			}
+			bDirty = true; // assume dirty since we are not going to wait
+			if (FVelocityRendering::OutputsToGBuffer())
+			{
+				GPrevPerBoneMotionBlur.EndAppendFence(RHICmdList);
 			}
 		}
 		else
 		{
-			for(int32 ViewIndex = 0;ViewIndex < Views.Num();ViewIndex++)
+			for (int32 ViewIndex = 0; ViewIndex < Views.Num(); ViewIndex++)
 			{
 				SCOPED_CONDITIONAL_DRAW_EVENTF(RHICmdList, EventView, Views.Num() > 1, TEXT("View%d"), ViewIndex);
 				FViewInfo& View = Views[ViewIndex];
@@ -1537,7 +1536,10 @@ bool FDeferredShadingSceneRenderer::RenderBasePass(FRHICommandListImmediate& RHI
 				bDirty |= RenderBasePassView(RHICmdList, View);
 			}
 		}
-
+		if (FVelocityRendering::OutputsToGBuffer())
+		{
+			FSceneRenderTargets::Get(RHICmdList).SetVelocityPass(false);
+		}
 	}
 
 	return bDirty;
