Increases Gaussian filter max number of samples from 32 to 128 by using a custom number of sample iterations shader behind the console variable r.Filter.LoopMode.

diff --git a/Engine/Shaders/FilterPixelShader.usf b/Engine/Shaders/FilterPixelShader.usf
--- a/Engine/Shaders/FilterPixelShader.usf
+++ b/Engine/Shaders/FilterPixelShader.usf
@@ -20,9 +20,6 @@ SamplerState AdditiveTextureSampler;
 Texture2D AdditiveTexture;
 #endif
 
-//
-half4 SampleWeights[NUM_SAMPLES];
-
 // UV.y and UV.wz are the input UV coordinates
 half ComputeMask(float2 UV)
 {
@@ -50,7 +47,53 @@ void Combine(inout float4 Dest, float4 Src, float4 Weight)
 #endif
 }
 
-#if (ES2_PROFILE || ES3_1_PROFILE)
+#if NUM_SAMPLES==0
+// dynamic number of sample
+
+float4 SampleOffsets[(MAX_NUM_SAMPLES + 1) / 2];
+half4 SampleWeights[MAX_NUM_SAMPLES];
+int SampleCount;
+
+void Main(
+	float2 InUV : TEXCOORD0,
+	out half4 OutColor : SV_Target0
+	)
+{
+	half4 Dest = 0;
+	
+	int SampleIndex;
+	LOOP for (SampleIndex = 0; SampleIndex < SampleCount - 1; SampleIndex += 2)
+	{
+		{
+			const float2 UV = InUV + SampleOffsets[SampleIndex / 2].xy;
+			const half Mask = ComputeMask(UV);
+			Combine(Dest, Texture2DSample(FilterTexture, FilterTextureSampler, UV), Mask * SampleWeights[SampleIndex + 0]);
+		}
+		
+		{
+			const float2 UV = InUV + SampleOffsets[SampleIndex / 2].wz;
+			const half Mask = ComputeMask(UV);
+			Combine(Dest, Texture2DSample(FilterTexture, FilterTextureSampler, UV), Mask * SampleWeights[SampleIndex + 1]);
+		}
+	}
+	BRANCH if (SampleIndex < SampleCount)
+	{
+		const float2 UV = InUV + SampleOffsets[SampleIndex / 2].xy;
+		const half Mask = ComputeMask(UV);
+		Combine(Dest, Texture2DSample(FilterTexture, FilterTextureSampler, UV), Mask * SampleWeights[SampleIndex + 0]);
+	}
+
+#if COMBINE_METHOD == 1
+	Dest += Texture2DSample(AdditiveTexture, AdditiveTextureSampler, InUV);
+#endif
+
+	// RETURN_COLOR not needed unless writing to SceneColor;
+	OutColor = Dest;
+}
+
+#elif (ES2_PROFILE || ES3_1_PROFILE)
+
+half4 SampleWeights[NUM_SAMPLES];
 
 void Main(
 	float2 InUV : TEXCOORD0,
@@ -73,6 +116,8 @@ void Main(
 
 #else
 
+half4 SampleWeights[NUM_SAMPLES];
+
 void Main(
 	float2 InUV : TEXCOORD0,
 	float4 InOffsetUVs[(NUM_SAMPLES + 1) / 2] : TEXCOORD1,
diff --git a/Engine/Source/Runtime/Renderer/Private/PostProcess/PostProcessWeightedSampleSum.cpp b/Engine/Source/Runtime/Renderer/Private/PostProcess/PostProcessWeightedSampleSum.cpp
--- a/Engine/Source/Runtime/Renderer/Private/PostProcess/PostProcessWeightedSampleSum.cpp
+++ b/Engine/Source/Runtime/Renderer/Private/PostProcess/PostProcessWeightedSampleSum.cpp
@@ -10,11 +10,31 @@
 #include "PostProcessWeightedSampleSum.h"
 #include "SceneUtils.h"
 
+// maximum number of sample using the shader that has the dynamic loop
+#define MAX_FILTER_SAMPLES	128
+
+// maximum number of sample available using unrolled loop shaders
+#define MAX_FILTER_COMPILE_TIME_SAMPLES 32
+
+#define MAX_PACKED_SAMPLES_OFFSET ((MAX_FILTER_SAMPLES + 1) / 2)
+
+static TAutoConsoleVariable<int32> CVarLoopMode(
+	TEXT("r.Filter.LoopMode"),
+	0,
+	TEXT("Controles when to use either dynamic or unrolled loops to iterates over the Gaussian filtering.\n")
+	TEXT("This passes is used for Gaussian Blur, Bloom and Depth of Field. The dynamic loop allows\n")
+	TEXT("up to 128 samples versus the 32 samples of unrolled loops, but add an additonal cost for\n")
+	TEXT("the loop's stop test at every iterations.\n")
+	TEXT(" 0: Unrolled loop only (default; limited to 32 samples).\n")
+	TEXT(" 1: Fall back to dynamic loop if needs more than 32 samples.\n")
+	TEXT(" 2: Dynamic loop only."),
+	ECVF_Scalability | ECVF_RenderThreadSafe);
+
 /**
  * A pixel shader which filters a texture.
  * @param CombineMethod 0:weighted filtering, 1: weighted filtering + additional texture, 2: max magnitude
  */
-template<uint32 NumSamples, uint32 CombineMethodInt>
+template<uint32 CompileTimeNumSamples, uint32 CombineMethodInt>
 class TFilterPS : public FGlobalShader
 {
 	DECLARE_SHADER_TYPE(TFilterPS,Global);
@@ -28,19 +48,26 @@ public:
 		}
 		else if( IsFeatureLevelSupported(Platform, ERHIFeatureLevel::SM4) )
 		{
-			return NumSamples <= 16;
+			return CompileTimeNumSamples <= 16;
 		}
 		else
 		{
-			return NumSamples <= 7;
+			return CompileTimeNumSamples <= 7;
 		}
 	}
 
 	static void ModifyCompilationEnvironment(EShaderPlatform Platform, FShaderCompilerEnvironment& OutEnvironment)
 	{
 		FGlobalShader::ModifyCompilationEnvironment(Platform, OutEnvironment);
-		OutEnvironment.SetDefine(TEXT("NUM_SAMPLES"), NumSamples);
+		OutEnvironment.SetDefine(TEXT("NUM_SAMPLES"), CompileTimeNumSamples);
 		OutEnvironment.SetDefine(TEXT("COMBINE_METHOD"), CombineMethodInt);
+
+		if (CompileTimeNumSamples == 0)
+		{
+			// CompileTimeNumSamples == 0 implies the dynamic loop, but we still needs to pass the number
+			// of maximum number of samples for the uniform arraies.
+			OutEnvironment.SetDefine(TEXT("MAX_NUM_SAMPLES"), MAX_FILTER_SAMPLES);
+		}
 	}
 
 	/** Default constructor. */
@@ -55,24 +82,54 @@ public:
 		AdditiveTexture.Bind(Initializer.ParameterMap,TEXT("AdditiveTexture"));
 		AdditiveTextureSampler.Bind(Initializer.ParameterMap,TEXT("AdditiveTextureSampler"));
 		SampleWeights.Bind(Initializer.ParameterMap,TEXT("SampleWeights"));
+		
+		if (CompileTimeNumSamples == 0)
+		{
+			// dynamic loop do UV offset in the pixel shader, and requires the number of samples.
+			SampleOffsets.Bind(Initializer.ParameterMap,TEXT("SampleOffsets"));
+			SampleCount.Bind(Initializer.ParameterMap,TEXT("SampleCount"));
+		}
 	}
 
 	/** Serializer */
 	virtual bool Serialize(FArchive& Ar) override
 	{
 		bool bShaderHasOutdatedParameters = FGlobalShader::Serialize(Ar);
-		Ar << FilterTexture << FilterTextureSampler << AdditiveTexture << AdditiveTextureSampler << SampleWeights;
+		Ar << FilterTexture << FilterTextureSampler << AdditiveTexture << AdditiveTextureSampler << SampleWeights << SampleOffsets << SampleCount;
 		return bShaderHasOutdatedParameters;
 	}
 
 	/** Sets shader parameter values */
-	void SetParameters(FRHICommandList& RHICmdList, FSamplerStateRHIParamRef SamplerStateRHI, FTextureRHIParamRef FilterTextureRHI, FTextureRHIParamRef AdditiveTextureRHI, const FLinearColor* SampleWeightValues)
+	void SetParameters(
+		FRHICommandList& RHICmdList, FSamplerStateRHIParamRef SamplerStateRHI, FTextureRHIParamRef FilterTextureRHI, FTextureRHIParamRef AdditiveTextureRHI, 
+		const FLinearColor* SampleWeightValues, const FVector2D* SampleOffsetValues, uint32 NumSamples )
 	{
+		check(CompileTimeNumSamples == 0 && NumSamples > 0 && NumSamples <= MAX_FILTER_SAMPLES || CompileTimeNumSamples == NumSamples);
 		const FPixelShaderRHIParamRef ShaderRHI = GetPixelShader();
 
 		SetTextureParameter(RHICmdList, ShaderRHI, FilterTexture, FilterTextureSampler, SamplerStateRHI, FilterTextureRHI);
 		SetTextureParameter(RHICmdList, ShaderRHI, AdditiveTexture, AdditiveTextureSampler, SamplerStateRHI, AdditiveTextureRHI);
 		SetShaderValueArray(RHICmdList, ShaderRHI, SampleWeights, SampleWeightValues, NumSamples);
+
+		if (CompileTimeNumSamples == 0)
+		{
+			// we needs additional setups for the dynamic loop
+			FVector4 PackedSampleOffsetsValues[MAX_PACKED_SAMPLES_OFFSET];
+
+			for(uint32 SampleIndex = 0;SampleIndex < NumSamples;SampleIndex += 2)
+			{
+				PackedSampleOffsetsValues[SampleIndex / 2].X = SampleOffsetValues[SampleIndex + 0].X;
+				PackedSampleOffsetsValues[SampleIndex / 2].Y = SampleOffsetValues[SampleIndex + 0].Y;
+				if(SampleIndex + 1 < NumSamples)
+				{
+					PackedSampleOffsetsValues[SampleIndex / 2].W = SampleOffsetValues[SampleIndex + 1].X;
+					PackedSampleOffsetsValues[SampleIndex / 2].Z = SampleOffsetValues[SampleIndex + 1].Y;
+				}
+			}
+
+			SetShaderValueArray(RHICmdList, ShaderRHI, SampleOffsets, PackedSampleOffsetsValues, MAX_PACKED_SAMPLES_OFFSET);
+			SetShaderValue(RHICmdList, ShaderRHI, SampleCount, NumSamples);
+		}
 	}
 
 	static const TCHAR* GetSourceFilename()
@@ -91,6 +148,10 @@ protected:
 	FShaderResourceParameter AdditiveTexture;
 	FShaderResourceParameter AdditiveTextureSampler;
 	FShaderParameter SampleWeights;
+
+	// parameters only for CompileTimeNumSamples == 0
+	FShaderParameter SampleOffsets;
+	FShaderParameter SampleCount;
 };
 
 
@@ -99,6 +160,7 @@ protected:
 #define VARIATION1(A)		VARIATION2(A,0)			VARIATION2(A,1)			VARIATION2(A,2)
 #define VARIATION2(A, B) typedef TFilterPS<A, B> TFilterPS##A##B; \
 	IMPLEMENT_SHADER_TYPE2(TFilterPS##A##B, SF_Pixel);
+	VARIATION1(0) // number of samples known at runtime
 	VARIATION1(1) VARIATION1(2) VARIATION1(3) VARIATION1(4) VARIATION1(5) VARIATION1(6) VARIATION1(7) VARIATION1(8)
 	VARIATION1(9) VARIATION1(10) VARIATION1(11) VARIATION1(12) VARIATION1(13) VARIATION1(14) VARIATION1(15) VARIATION1(16)
 	VARIATION1(17) VARIATION1(18) VARIATION1(19) VARIATION1(20) VARIATION1(21) VARIATION1(22) VARIATION1(23) VARIATION1(24)
@@ -172,44 +234,85 @@ void SetFilterShaders(
 	)
 {
 	check(CombineMethodInt <= 2);
+	check(NumSamples <= MAX_FILTER_SAMPLES && NumSamples > 0);
+
 	auto ShaderMap = GetGlobalShaderMap(FeatureLevel);
+	const auto DynamicNumSample = CVarLoopMode.GetValueOnRenderThread();
+	
+	if (NumSamples > MAX_FILTER_COMPILE_TIME_SAMPLES && DynamicNumSample != 0 || DynamicNumSample == 2)
+	{
+		// there is to many samples, so we use the dynamic sample count shader
+		
+		TShaderMapRef<FPostProcessVS> VertexShader(ShaderMap);
+		*OutVertexShader = *VertexShader;
+		if(CombineMethodInt == 0)
+		{
+			TShaderMapRef<TFilterPS<0, 0> > PixelShader(ShaderMap);
+			{
+				static FGlobalBoundShaderState BoundShaderState;
+				SetGlobalBoundShaderState(RHICmdList, FeatureLevel, BoundShaderState, GFilterVertexDeclaration.VertexDeclarationRHI, *VertexShader, *PixelShader);
+			}
+			PixelShader->SetParameters(RHICmdList, SamplerStateRHI, FilterTextureRHI, AdditiveTextureRHI, SampleWeights, SampleOffsets, NumSamples);
+		}
+		else if(CombineMethodInt == 1)
+		{
+			TShaderMapRef<TFilterPS<0, 1> > PixelShader(ShaderMap);
+			{
+				static FGlobalBoundShaderState BoundShaderState;
+				SetGlobalBoundShaderState(RHICmdList, FeatureLevel, BoundShaderState, GFilterVertexDeclaration.VertexDeclarationRHI, *VertexShader, *PixelShader);
+			}
+			PixelShader->SetParameters(RHICmdList, SamplerStateRHI, FilterTextureRHI, AdditiveTextureRHI, SampleWeights, SampleOffsets, NumSamples);
+		}
+		else\
+		{
+			TShaderMapRef<TFilterPS<0, 2> > PixelShader(ShaderMap);
+			{
+				static FGlobalBoundShaderState BoundShaderState;
+				SetGlobalBoundShaderState(RHICmdList, FeatureLevel, BoundShaderState, GFilterVertexDeclaration.VertexDeclarationRHI, *VertexShader, *PixelShader);
+			}
+			PixelShader->SetParameters(RHICmdList, SamplerStateRHI, FilterTextureRHI, AdditiveTextureRHI, SampleWeights, SampleOffsets, NumSamples);
+		}
+		return;
+	}
 
 	// A macro to handle setting the filter shader for a specific number of samples.
-#define SET_FILTER_SHADER_TYPE(NumSamples) \
-	case NumSamples: \
+#define SET_FILTER_SHADER_TYPE(CompileTimeNumSamples) \
+	case CompileTimeNumSamples: \
 	{ \
-		TShaderMapRef<TFilterVS<NumSamples> > VertexShader(ShaderMap); \
+		TShaderMapRef<TFilterVS<CompileTimeNumSamples> > VertexShader(ShaderMap); \
 		*OutVertexShader = *VertexShader; \
 		if(CombineMethodInt == 0) \
 		{ \
-			TShaderMapRef<TFilterPS<NumSamples, 0> > PixelShader(ShaderMap); \
+			TShaderMapRef<TFilterPS<CompileTimeNumSamples, 0> > PixelShader(ShaderMap); \
 			{ \
 				static FGlobalBoundShaderState BoundShaderState; \
 				SetGlobalBoundShaderState(RHICmdList, FeatureLevel, BoundShaderState, GFilterVertexDeclaration.VertexDeclarationRHI, *VertexShader, *PixelShader); \
 			} \
-			PixelShader->SetParameters(RHICmdList, SamplerStateRHI, FilterTextureRHI, AdditiveTextureRHI, SampleWeights); \
+			PixelShader->SetParameters(RHICmdList, SamplerStateRHI, FilterTextureRHI, AdditiveTextureRHI, SampleWeights, SampleOffsets, NumSamples); \
 		} \
 		else if(CombineMethodInt == 1) \
 		{ \
-			TShaderMapRef<TFilterPS<NumSamples, 1> > PixelShader(ShaderMap); \
+			TShaderMapRef<TFilterPS<CompileTimeNumSamples, 1> > PixelShader(ShaderMap); \
 			{ \
 				static FGlobalBoundShaderState BoundShaderState; \
 				SetGlobalBoundShaderState(RHICmdList, FeatureLevel, BoundShaderState, GFilterVertexDeclaration.VertexDeclarationRHI, *VertexShader, *PixelShader); \
 			} \
-			PixelShader->SetParameters(RHICmdList, SamplerStateRHI, FilterTextureRHI, AdditiveTextureRHI, SampleWeights); \
+			PixelShader->SetParameters(RHICmdList, SamplerStateRHI, FilterTextureRHI, AdditiveTextureRHI, SampleWeights, SampleOffsets, NumSamples); \
 		} \
 		else\
 		{ \
-			TShaderMapRef<TFilterPS<NumSamples, 2> > PixelShader(ShaderMap); \
+			TShaderMapRef<TFilterPS<CompileTimeNumSamples, 2> > PixelShader(ShaderMap); \
 			{ \
 				static FGlobalBoundShaderState BoundShaderState; \
 				SetGlobalBoundShaderState(RHICmdList, FeatureLevel, BoundShaderState, GFilterVertexDeclaration.VertexDeclarationRHI, *VertexShader, *PixelShader); \
 			} \
-			PixelShader->SetParameters(RHICmdList, SamplerStateRHI, FilterTextureRHI, AdditiveTextureRHI, SampleWeights); \
+			PixelShader->SetParameters(RHICmdList, SamplerStateRHI, FilterTextureRHI, AdditiveTextureRHI, SampleWeights, SampleOffsets, NumSamples); \
 		} \
 		VertexShader->SetParameters(RHICmdList, SampleOffsets); \
 		break; \
 	};
+	
+	check(NumSamples <= MAX_FILTER_COMPILE_TIME_SAMPLES);
 
 	// Set the appropriate filter shader for the given number of samples.
 	switch(NumSamples)
@@ -593,16 +696,18 @@ void FRCPassPostProcessWeightedSampleSum::DrawQuad(FRHICommandListImmediate& RHI
 
 uint32 FRCPassPostProcessWeightedSampleSum::GetMaxNumSamples(ERHIFeatureLevel::Type InFeatureLevel)
 {
-	uint32 MaxNumSamples;
-	if (InFeatureLevel >= ERHIFeatureLevel::SM5)
+	if (CVarLoopMode.GetValueOnRenderThread() != 0)
 	{
-		MaxNumSamples = MAX_FILTER_SAMPLES;
+		return MAX_FILTER_SAMPLES;
 	}
-	else if (InFeatureLevel >= ERHIFeatureLevel::SM4)
+	
+	uint32 MaxNumSamples = MAX_FILTER_COMPILE_TIME_SAMPLES;
+
+	if (InFeatureLevel == ERHIFeatureLevel::SM4)
 	{
 		MaxNumSamples = 16;
 	}
-	else
+	else if (InFeatureLevel < ERHIFeatureLevel::SM4)
 	{
 		MaxNumSamples = 7;
 	}
diff --git a/Engine/Source/Runtime/Renderer/Private/PostProcess/SceneFilterRendering.h b/Engine/Source/Runtime/Renderer/Private/PostProcess/SceneFilterRendering.h
--- a/Engine/Source/Runtime/Renderer/Private/PostProcess/SceneFilterRendering.h
+++ b/Engine/Source/Runtime/Renderer/Private/PostProcess/SceneFilterRendering.h
@@ -7,8 +7,6 @@
 #ifndef _INC_SCENEFILTERRENDERING
 #define _INC_SCENEFILTERRENDERING
 
-#define MAX_FILTER_SAMPLES	32
-
 #include "SceneRenderTargets.h"
 
 /**
