Fixes a compilation failure in PostProcessWeightedSampleSum.cpp caused by missing parenthesis in conditions containing && and ||

diff --git a/Engine/Source/Runtime/Renderer/Private/PostProcess/PostProcessWeightedSampleSum.cpp b/Engine/Source/Runtime/Renderer/Private/PostProcess/PostProcessWeightedSampleSum.cpp
--- a/Engine/Source/Runtime/Renderer/Private/PostProcess/PostProcessWeightedSampleSum.cpp
+++ b/Engine/Source/Runtime/Renderer/Private/PostProcess/PostProcessWeightedSampleSum.cpp
@@ -239,7 +239,7 @@ void SetFilterShaders(
 	auto ShaderMap = GetGlobalShaderMap(FeatureLevel);
 	const auto DynamicNumSample = CVarLoopMode.GetValueOnRenderThread();
 	
-	if (NumSamples > MAX_FILTER_COMPILE_TIME_SAMPLES && DynamicNumSample != 0 || DynamicNumSample == 2)
+	if ((NumSamples > MAX_FILTER_COMPILE_TIME_SAMPLES && DynamicNumSample != 0) || (DynamicNumSample == 2))
 	{
 		// there is to many samples, so we use the dynamic sample count shader
 		
