Implements r.Filter.WrongVarienceWorkAround to let r.ScreenPercentage work arround a gaussian filter's invalid variance computation with the screen resolution.

diff --git a/Engine/Source/Runtime/Renderer/Private/PostProcess/PostProcessWeightedSampleSum.cpp b/Engine/Source/Runtime/Renderer/Private/PostProcess/PostProcessWeightedSampleSum.cpp
--- a/Engine/Source/Runtime/Renderer/Private/PostProcess/PostProcessWeightedSampleSum.cpp
+++ b/Engine/Source/Runtime/Renderer/Private/PostProcess/PostProcessWeightedSampleSum.cpp
@@ -30,6 +30,14 @@ static TAutoConsoleVariable<int32> CVarLoopMode(
 	TEXT(" 2: Dynamic loop only."),
 	ECVF_Scalability | ECVF_RenderThreadSafe);
 
+static TAutoConsoleVariable<int32> CVarWrongVarienceWorkAround(
+	TEXT("r.Filter.WrongVarienceWorkAround"),
+	0,
+	TEXT("The usage of the variance is currently wrong. That cause the bloom kernel to not scale linearly with the resolution.\n")
+	TEXT("When enabled, this hack fix it with screen percentage.\n")
+	TEXT(" 0 is off (default), 1 is on"),
+	ECVF_RenderThreadSafe);
+
 /**
  * A pixel shader which filters a texture.
  * @param CombineMethod 0:weighted filtering, 1: weighted filtering + additional texture, 2: max magnitude
@@ -390,7 +398,16 @@ static float NormalDistributionUnscaled(float X,float Mean,float Variance, EFilt
 
 static uint32 Compute1DGaussianFilterKernel(ERHIFeatureLevel::Type InFeatureLevel, float KernelRadius, FVector2D OutOffsetAndWeight[MAX_FILTER_SAMPLES], uint32 MaxFilterSamples, EFilterShape FilterShape, float CrossCenterWeight)
 {
-	float ClampedKernelRadius = FRCPassPostProcessWeightedSampleSum::GetClampedKernelRadius( InFeatureLevel, KernelRadius );
+	float s = 1;
+	if (CVarWrongVarienceWorkAround.GetValueOnRenderThread())
+	{
+		// dirty workaround that needs to be enabled to not change the content by default.
+		static const auto CVar = IConsoleManager::Get().FindConsoleVariable(TEXT("r.ScreenPercentage"));
+
+		s = 0.01 * CVar->GetFloat();
+	}
+
+	float ClampedKernelRadius = FRCPassPostProcessWeightedSampleSum::GetClampedKernelRadius( InFeatureLevel, KernelRadius * s );
 	int32 IntegerKernelRadius = FRCPassPostProcessWeightedSampleSum::GetIntegerKernelRadius( InFeatureLevel, KernelRadius );
 
 	// smallest IntegerKernelRadius will be 1
