fixed SceneColorFringe showflag

diff --git a/Engine/Source/Runtime/Engine/Private/SceneView.cpp b/Engine/Source/Runtime/Engine/Private/SceneView.cpp
--- a/Engine/Source/Runtime/Engine/Private/SceneView.cpp
+++ b/Engine/Source/Runtime/Engine/Private/SceneView.cpp
@@ -1273,6 +1273,11 @@ void FSceneView::EndFinalPostprocessSettings(const FSceneViewInitOptions& ViewIn
 		{
 			FinalPostProcessSettings.SceneFringeIntensity = FMath::Min(FinalPostProcessSettings.SceneFringeIntensity, Value);
 		}
+
+		if(!Family->EngineShowFlags.SceneColorFringe || !Family->EngineShowFlags.CameraImperfections)
+		{
+			FinalPostProcessSettings.SceneFringeIntensity = 0;
+		}
 	}
 
 	if (!Family->EngineShowFlags.Lighting || !Family->EngineShowFlags.GlobalIllumination)
diff --git a/Engine/Source/Runtime/Renderer/Private/PostProcess/PostProcessTonemap.cpp b/Engine/Source/Runtime/Renderer/Private/PostProcess/PostProcessTonemap.cpp
--- a/Engine/Source/Runtime/Renderer/Private/PostProcess/PostProcessTonemap.cpp
+++ b/Engine/Source/Runtime/Renderer/Private/PostProcess/PostProcessTonemap.cpp
@@ -526,10 +526,7 @@ static uint32 TonemapperGenerateBitmaskPC(const FViewInfo* RESTRICT View, bool b
 	static const auto CVar = IConsoleManager::Get().FindTConsoleVariableDataInt(TEXT("r.SceneColorFringeQuality")); 
 
 	int32 FringeQuality = CVar->GetValueOnRenderThread();
-	if(View->Family->EngineShowFlags.SceneColorFringe
-		// Removed this from the camera imperfections toggle because this no longer takes an extra pass.
-		// && Context->View.Family->EngineShowFlags.CameraImperfections
-		&& View->FinalPostProcessSettings.SceneFringeIntensity > 0.01f
+	if( View->FinalPostProcessSettings.SceneFringeIntensity > 0.01f
 		&& FringeQuality > 0)
 	{
 		Bitmask |= TonemapperColorFringe;
