fixed old bug where blom and gaussian DOF size doesn't scale linearly, but costing the linear performance increase.
Depending in which screen resolution size the bloom was tweaked this can change the content a bit.
Old method (for a while we keep the cvar): "r.Filter.NewMethod 0"
Added new cvar r.Filter.SizeScale allowing to scale sample count to tweak performance/quality

diff --git a/Engine/Source/Runtime/Renderer/Private/PostProcess/PostProcessWeightedSampleSum.cpp b/Engine/Source/Runtime/Renderer/Private/PostProcess/PostProcessWeightedSampleSum.cpp
--- a/Engine/Source/Runtime/Renderer/Private/PostProcess/PostProcessWeightedSampleSum.cpp
+++ b/Engine/Source/Runtime/Renderer/Private/PostProcess/PostProcessWeightedSampleSum.cpp
@@ -30,12 +30,23 @@ static TAutoConsoleVariable<int32> CVarLoopMode(
 	TEXT(" 2: Dynamic loop only."),
 	ECVF_Scalability | ECVF_RenderThreadSafe);
 
-static TAutoConsoleVariable<int32> CVarWrongVarienceWorkAround(
-	TEXT("r.Filter.WrongVarienceWorkAround"),
-	0,
-	TEXT("The usage of the variance is currently wrong. That cause the bloom kernel to not scale linearly with the resolution.\n")
-	TEXT("When enabled, this hack fix it with screen percentage.\n")
-	TEXT(" 0 is off (default), 1 is on"),
+static TAutoConsoleVariable<float> CVarFilterSizeScale(
+	TEXT("r.Filter.SizeScale"),
+	1.0f,
+	TEXT("Allows to scale down or up the sample count used for bloom and Gaussian depth of field (scale is clampled to give reasonable results).\n")
+	TEXT("Values down to 0.6 are hard to notice\n")
+	TEXT(" 1 full quality (default)\n")
+	TEXT(" >1 more samples (slower)\n")
+	TEXT(" <1 less samples (faster, artifacts with HDR content or boxy results with GaussianDOF)"),
+	ECVF_Scalability | ECVF_RenderThreadSafe);
+
+// will be removed soon
+static TAutoConsoleVariable<int32> FilterNewMethod(
+	TEXT("r.Filter.NewMethod"),
+	1,
+	TEXT("Affects bloom and Gaussian depth of field.\n")
+	TEXT(" 0: old method (doesn't scale linearly with size)\n")
+	TEXT(" 1: new method, might need asset tweak (default)"),
 	ECVF_RenderThreadSafe);
 
 /**
@@ -365,28 +376,56 @@ void SetFilterShaders(
 }
 
 /**
- * Evaluates a normal distribution PDF at given X.
+ * Evaluates a normal distribution PDF (around 0) at given X.
  * This function misses the math for scaling the result (faster, not needed if the resulting values are renormalized).
  * @param X - The X to evaluate the PDF at.
- * @param Mean - The normal distribution's mean.
- * @param Variance - The normal distribution's variance.
+ * @param Scale - The normal distribution's variance.
  * @return The value of the normal distribution at X. (unscaled)
  */
-static float NormalDistributionUnscaled(float X,float Mean,float Variance, EFilterShape FilterShape, float CrossCenterWeight)
+static float NormalDistributionUnscaled(float X,float Scale, EFilterShape FilterShape, float CrossCenterWeight)
 {
-	float dx = FMath::Abs(X - Mean);
+	float Ret;
 
-	float Ret = FMath::Exp(-FMath::Square(dx) / (2.0f * Variance));
-
-	// tweak the gaussian shape e.g. "r.Bloom.Cross 3.5"
-	if(CrossCenterWeight > 1.0f)
+	if(FilterNewMethod.GetValueOnRenderThread())
 	{
-		Ret = FMath::Max(0.0f, 1.0f - dx / Variance);
-		Ret = FMath::Pow(Ret, CrossCenterWeight);
+		float dxUnScaled = FMath::Abs(X);
+		float dxScaled = dxUnScaled * Scale;
+
+		// Constant is tweaked give a similar look to UE4 before we fix the scale bug (Some content tweaking might be needed).
+		// The value defines how much of the Gaussian clipped by the sample window. 
+		// r.Filter.SizeScale allows to tweak that for performance/quality.
+		Ret = FMath::Exp(-16.7f * FMath::Square(dxScaled));
+
+		// tweak the gaussian shape e.g. "r.Bloom.Cross 3.5"
+		if (CrossCenterWeight > 1.0f)
+		{
+			Ret = FMath::Max(0.0f, 1.0f - dxUnScaled);
+			Ret = FMath::Pow(Ret, CrossCenterWeight);
+		}
+		else
+		{
+			Ret = FMath::Lerp(Ret, FMath::Max(0.0f, 1.0f - dxUnScaled), CrossCenterWeight);
+		}
 	}
 	else
 	{
-		Ret = FMath::Lerp(Ret, FMath::Max(0.0f, 1.0f - dx / Variance), CrossCenterWeight);
+		// will be removed soon
+		float OldVariance = 1.0f / Scale;
+
+		float dx = FMath::Abs(X);
+
+		Ret = FMath::Exp(-FMath::Square(dx) / (2.0f * OldVariance));
+
+		// tweak the gaussian shape e.g. "r.Bloom.Cross 3.5"
+		if(CrossCenterWeight > 1.0f)
+		{
+			Ret = FMath::Max(0.0f, 1.0f - dx / OldVariance);
+			Ret = FMath::Pow(Ret, CrossCenterWeight);
+		}
+		else
+		{
+			Ret = FMath::Lerp(Ret, FMath::Max(0.0f, 1.0f - dx / OldVariance), CrossCenterWeight);
+		}
 	}
 
 	return Ret;
@@ -398,17 +437,12 @@ static float NormalDistributionUnscaled(float X,float Mean,float Variance, EFilt
 
 static uint32 Compute1DGaussianFilterKernel(ERHIFeatureLevel::Type InFeatureLevel, float KernelRadius, FVector2D OutOffsetAndWeight[MAX_FILTER_SAMPLES], uint32 MaxFilterSamples, EFilterShape FilterShape, float CrossCenterWeight)
 {
-	float s = 1;
-	if (CVarWrongVarienceWorkAround.GetValueOnRenderThread())
-	{
-		// dirty workaround that needs to be enabled to not change the content by default.
-		static const auto CVar = IConsoleManager::Get().FindConsoleVariable(TEXT("r.ScreenPercentage"));
+	float FilterSizeScale = FMath::Clamp(CVarFilterSizeScale.GetValueOnRenderThread(), 0.1f, 10.0f);
 
-		s = 0.01 * CVar->GetFloat();
-	}
+	float ClampedKernelRadius = FRCPassPostProcessWeightedSampleSum::GetClampedKernelRadius(InFeatureLevel, KernelRadius);
+	int32 IntegerKernelRadius = FRCPassPostProcessWeightedSampleSum::GetIntegerKernelRadius(InFeatureLevel, KernelRadius * FilterSizeScale);
 
-	float ClampedKernelRadius = FRCPassPostProcessWeightedSampleSum::GetClampedKernelRadius( InFeatureLevel, KernelRadius * s );
-	int32 IntegerKernelRadius = FRCPassPostProcessWeightedSampleSum::GetIntegerKernelRadius( InFeatureLevel, KernelRadius );
+	float Scale = 1.0f / ClampedKernelRadius;
 
 	// smallest IntegerKernelRadius will be 1
 
@@ -416,7 +450,7 @@ static uint32 Compute1DGaussianFilterKernel(ERHIFeatureLevel::Type InFeatureLeve
 	float WeightSum = 0.0f;
 	for(int32 SampleIndex = -IntegerKernelRadius; SampleIndex <= IntegerKernelRadius; SampleIndex += 2)
 	{
-		float Weight0 = NormalDistributionUnscaled(SampleIndex, 0, ClampedKernelRadius, FilterShape, CrossCenterWeight);
+		float Weight0 = NormalDistributionUnscaled(SampleIndex, Scale, FilterShape, CrossCenterWeight);
 		float Weight1 = 0.0f;
 
 		// Because we use bilinear filtering we only require half the sample count.
@@ -426,7 +460,7 @@ static uint32 Compute1DGaussianFilterKernel(ERHIFeatureLevel::Type InFeatureLeve
 		//    c * .. but another texel to the right would accidentially leak into this computation.
 		if(SampleIndex != IntegerKernelRadius)
 		{
-			Weight1 = NormalDistributionUnscaled(SampleIndex + 1, 0, ClampedKernelRadius, FilterShape, CrossCenterWeight);
+			Weight1 = NormalDistributionUnscaled(SampleIndex + 1, Scale, FilterShape, CrossCenterWeight);
 		}
 
 		float TotalWeight = Weight0 + Weight1;
