Fixes an assertion failure in FogRendering.cpp when the camera is far away from the exponential fog actuator.

diff --git a/Engine/Source/Runtime/Renderer/Private/FogRendering.cpp b/Engine/Source/Runtime/Renderer/Private/FogRendering.cpp
--- a/Engine/Source/Runtime/Renderer/Private/FogRendering.cpp
+++ b/Engine/Source/Runtime/Renderer/Private/FogRendering.cpp
@@ -241,7 +241,12 @@ void FSceneRenderer::InitFogConstants()
 			{
 				const FExponentialHeightFogSceneInfo& FogInfo = Scene->ExponentialFogs[0];
 				const float CosTerminatorAngle = FMath::Clamp(FMath::Cos(FogInfo.LightTerminatorAngle * PI / 180.0f), -1.0f + DELTA, 1.0f - DELTA);
-				const float CollapsedFogParameter = FogInfo.FogDensity * FMath::Pow(2.0f, -FogInfo.FogHeightFalloff * (View.ViewMatrices.ViewOrigin.Z - FogInfo.FogHeight));
+				const float CollapsedFogParameterPower = FMath::Clamp(
+						-FogInfo.FogHeightFalloff * (View.ViewMatrices.ViewOrigin.Z - FogInfo.FogHeight),
+						-126.f + 1.f, // min and max exponent values for IEEE floating points (http://en.wikipedia.org/wiki/IEEE_floating_point)
+						+127.f - 1.f
+						);
+				const float CollapsedFogParameter = FogInfo.FogDensity * FMath::Pow(2.0f, CollapsedFogParameterPower);
 				View.ExponentialFogParameters = FVector4(CollapsedFogParameter, FogInfo.FogHeightFalloff, CosTerminatorAngle, FogInfo.StartDistance);
 				View.ExponentialFogColor = FVector(FogInfo.FogColor.R, FogInfo.FogColor.G, FogInfo.FogColor.B);
 				View.FogMaxOpacity = FogInfo.FogMaxOpacity;
