Added experimental volumetric decals (can be useful but for many applications it might not be fast enough)
 defined by a distance function
 Material needs to output distance to an object through OpacityMask, given the LightVector
position.

 * can be quite slow (needs to evaluate the distance function for many steps) when large on screen
 * normal is reconstructed from the distance function (6 extra evaluations)
 * intersects with the scene like any other z buffer object
 * Distortions are possible but require more steps ()
 * can be animated but motionblur and TemporalAA will be problematic
 * shadow receiving, lighting, sorting all works like any opaque
 * can output any GBuffer attribute and use any existing shading model

not done yet:
 * shadow casting
 * quality settings (step count, hit threshold)
 * material normal/bump adjustment
 * optimizations

maybe later:
* transluceny volumetric e.g. clouds
* create material functions for the usual primitives and operations
  http://iquilezles.org/www/articles/distfunctions/distfunctions.htm

diff --git a/Engine/Shaders/DeferredDecal.usf b/Engine/Shaders/DeferredDecal.usf
--- a/Engine/Shaders/DeferredDecal.usf
+++ b/Engine/Shaders/DeferredDecal.usf
@@ -4,16 +4,18 @@
 	DeferredDecal.usf: Pixel shader for computing a deferred decal.
 =============================================================================*/
 
-// DECAL_BLEND_MODE is set by C++ from EDecalBlendMode e.g. DBM_Translucent, DBM_TranslucentColorNormalRoughness, ...
+// DECAL_BLEND_MODE is set by C++ from EDecalBlendMode e.g. DECALBLENDMODEID_VOLUMETRIC, DECALBLENDMODEID_NORMAL, ...
 // DECAL_RENDERTARGET_COUNT is set by C++ e.g. 1: single RT, 2: two MRT, 3: three MRT
-// DECAL_RENDERSTAGE 0:before BasePass, 1: after base pass, before lighting
+// DECAL_RENDERSTAGE is set by C++ from EDecalRenderStage e.g. 0:before base pass, 1:after base pass, 2: before lighting
 
 #include "Common.usf"
 
 // from screen space position to decal space (for position)
 float4x4 ScreenToDecal;
-// from decal space to world space (for normals)
+// to transform position from decal space to world space, for normals use transpose of inverse
 float4x4 DecalToWorld;
+// to transform position from world space to decal space, for normals use transpose of inverse
+float4x4 WorldToDecal;
 
 #if DECAL_PROJECTION
 	#include "Material.usf"
@@ -50,13 +52,87 @@ void MainVS(
 	OutPosition = OutScreenPosition = mul(InPosition, FrustumComponentToClip);
 }
 
+float DistanceFunction(float4 ScreenPosition, float3 Pos)
+{
+	float3 SwizzlePos = Pos.zyx;
+
+#if DECAL_PROJECTION && DECAL_BLEND_MODE == DECALBLENDMODEID_VOLUMETRIC
+	
+	FMaterialPixelParameters MaterialParameters = MakeInitializedMaterialPixelParameters();
+#if NUM_MATERIAL_TEXCOORDS
+	for(int CoordinateIndex = 0;CoordinateIndex < NUM_MATERIAL_TEXCOORDS;CoordinateIndex++)
+	{
+		MaterialParameters.TexCoords[CoordinateIndex] = SwizzlePos.xy * 0.5f + 0.5f;
+	}
+#endif
+
+	MaterialParameters.VertexColor = 1;
+	MaterialParameters.ScreenPosition = ScreenPosition;
+	MaterialParameters.LightVector = SwizzlePos * 0.5f + 0.5f;		// to be compatible with decals but *0.5+0.5 seems wrong
+	MaterialParameters.WorldPosition = 
+	MaterialParameters.WorldPosition_CamRelative = 
+	MaterialParameters.WorldPosition_NoOffsets = 
+	MaterialParameters.WorldPosition_NoOffsets_CamRelative = mul(DecalToWorld, float4(Pos, 1)).xyz;
+	MaterialParameters.CameraVector = normalize(View.ViewOrigin.xyz - MaterialParameters.WorldPosition);
+
+	// the material needs to output the distance from "Opacity" depending on LightVector (in 0..1 range)
+	return GetMaterialMaskInputRaw(MaterialParameters);
+
+#endif
+
+	return 0;
+}
+
+// @return t=0..1, -1: no hit
+float Raycast(float4 ScreenPosition, float3 StartPos, float3 EndPos)
+{
+	const uint steps = 300; 
+
+	float t = 0;
+
+	float RayLength = length(EndPos - StartPos);
+	float3 Direction = normalize(EndPos - StartPos);
+
+	LOOP for(uint i = 0; i < steps; ++i)
+	{
+		float Distance = DistanceFunction(ScreenPosition, StartPos + t * Direction);
+
+		if(Distance <= i * 0.001f)
+		{
+			return t / RayLength;
+		}
+
+		t += Distance;
+	}
+
+	return -1;
+}
+
+// like explained here http://iquilezles.org/www/articles/derivative/derivative.htm
+float3 GradientNotNormalized(float4 ScreenPosition, float3 OSPos)
+{
+	float eps = 0.1f;
+
+	float2 e = float2(eps, 0);
+
+    return float3(DistanceFunction(ScreenPosition, OSPos + e.xyy) - DistanceFunction(ScreenPosition, OSPos - e.xyy),
+                  DistanceFunction(ScreenPosition, OSPos + e.yxy) - DistanceFunction(ScreenPosition, OSPos - e.yxy),
+                  DistanceFunction(ScreenPosition, OSPos + e.yyx) - DistanceFunction(ScreenPosition, OSPos - e.yyx));
+}
 
 
-#if DECAL_PROJECTION
 
 // decal pixel shader
+#if DECAL_PROJECTION
+
+#if DECAL_BLEND_MODE == DECALBLENDMODEID_VOLUMETRIC
+// should work, but doesn't 	EARLYDEPTHSTENCIL
+#endif
+
 void MainPS(
-	in float4 ScreenPosition : TEXCOORD0
+	in float4 ScreenPosition : TEXCOORD0			// is ClipPosition
+	
+	, in INPUT_POSITION_QUALIFIERS float4 SvPosition : SV_Position
 
 #if DECAL_RENDERTARGET_COUNT > 0
 	, out float4 OutTarget0 : SV_Target0
@@ -73,30 +149,81 @@ void MainPS(
 #if DECAL_RENDERTARGET_COUNT > 3
 	, out float4 OutTarget3 : SV_Target3
 #endif
+
+	OPTIONAL_OutDepth
+
 	)
 {
-	half SceneW = PreviousDepth(ScreenPosition);
-	float4 DecalVector = mul(float4(ScreenPosition.xy / ScreenPosition.w * SceneW, SceneW, 1), ScreenToDecal);
-	DecalVector.xyz /= DecalVector.w;
+	bool bVolumetric = DECAL_BLEND_MODE == DECALBLENDMODEID_VOLUMETRIC;
 
-	// clip content outside the decal
-	// not needed if we stencil out the decal but we do that only on large (screen space) ones
-	clip(DecalVector.xyz + 1.0f);
-	clip(1.0f - DecalVector.xyz);
+	float3 OSPosition;
+	float3 WSPosition;
 
-	float4 HomogeneousWorldPosition = mul(float4(ScreenPosition.xy / ScreenPosition.w * SceneW, SceneW, 1), View.ScreenToWorld);
-	float3 WorldPosition = HomogeneousWorldPosition.xyz / HomogeneousWorldPosition.w;
-	float3 CameraVector = normalize(View.ViewOrigin.xyz - WorldPosition);
+	if(bVolumetric)
+	{
+		// LineBoxIntersect want's to clip a long ray so we provide one, this could be done differently as it limites the draw distance and reduces precision 
+		float SceneW = 100000000;
 
-	// can be optimized
-	DecalVector = DecalVector * 0.5f + 0.5f;
+		float4 ProjectedScreenPosition = float4(ScreenPosition.xy / ScreenPosition.w * SceneW, SceneW, 1);		// is ProjectedClipPosition
+
+		ProjectedScreenPosition = float4(ScreenPosition.xy / ScreenPosition.w * SceneW, SceneW, 1);
+
+		float4 WorldPosHom = mul(ProjectedScreenPosition, View.ScreenToWorld);
+
+		float3 WSStartRay = View.ViewOrigin.xyz;
+		float3 WSEndRay = WorldPosHom.xyz;
+
+		// in object/decal space
+		float4 OSStartRayHom = mul(float4(WSStartRay,1), WorldToDecal);
+		float4 OSEndRayHom = mul(float4(WSEndRay, 1), WorldToDecal);
+
+		float3 OSStartRay = OSStartRayHom.xyz / OSStartRayHom.w;
+		float3 OSEndRay = OSEndRayHom.xyz / OSEndRayHom.w;
+
+		float2 OSBoxMinMax = LineBoxIntersect(OSStartRay, OSEndRay, -1, 1);
+
+		float3 OSBoxStartRay = OSStartRay + OSBoxMinMax.x * (OSEndRay - OSStartRay);
+		float3 OSBoxEndRay = OSStartRay + OSBoxMinMax.y * (OSEndRay - OSStartRay);
+
+		float HitT = Raycast(ScreenPosition, OSBoxStartRay, OSBoxEndRay);
+
+		clip(HitT);
+
+		OSPosition = OSBoxStartRay + HitT * (OSBoxEndRay - OSBoxStartRay);
+		WSPosition = mul(float4(OSPosition, 1), DecalToWorld).xyz;
+
+		#if OUTPUT_PIXEL_DEPTH_OFFSET
+			float4 CSHitPos = mul(float4(WSPosition, 1), View.WorldToClip);
+			OutDepth = CSHitPos.z / CSHitPos.w;
+		#endif
+	}
+	else
+	{
+		half SceneW = PreviousDepth(ScreenPosition);
+		float4 ProjectedScreenPosition = float4(ScreenPosition.xy / ScreenPosition.w * SceneW, SceneW, 1);	
+
+		float4 DecalVectorHom = mul(ProjectedScreenPosition, ScreenToDecal);
+		OSPosition = DecalVectorHom.xyz / DecalVectorHom.w;
+
+		// clip content outside the decal
+		// not needed if we stencil out the decal but we do that only on large (screen space) ones
+		clip(OSPosition.xyz + 1.0f);
+		clip(1.0f - OSPosition.xyz);
+
+		float4 WSPositionHom = mul(ProjectedScreenPosition, View.ScreenToWorld);
+		WSPosition = WSPositionHom.xyz / WSPositionHom.w;
+	}
+
+	float3 CameraVector = normalize(View.ViewOrigin.xyz - WSPosition);
 
 	// can be optimized
+	float3 DecalVector = OSPosition * 0.5f + 0.5f;
+
 	// Swizzle so that DecalVector.xy are perpendicular to the projection direction and DecalVector.z is distance along the projection direction
-	DecalVector.xyz = DecalVector.zyx;
+	float3 SwizzlePos = DecalVector.zyx;
 
 	// By default, map textures using the vectors perpendicular to the projection direction
-	float2 DecalUVs = DecalVector.xy;
+	float2 DecalUVs = SwizzlePos.xy;
 
 	FMaterialPixelParameters MaterialParameters = MakeInitializedMaterialPixelParameters();
 #if NUM_MATERIAL_TEXCOORDS
@@ -108,21 +235,37 @@ void MainPS(
 	MaterialParameters.VertexColor = 1;
 	MaterialParameters.CameraVector = CameraVector;
 	MaterialParameters.ScreenPosition = ScreenPosition;
-	MaterialParameters.LightVector = DecalVector.xyz;
+	MaterialParameters.LightVector = SwizzlePos;
 	MaterialParameters.WorldPosition = 
 		MaterialParameters.WorldPosition_CamRelative = 
 		MaterialParameters.WorldPosition_NoOffsets = 
-		MaterialParameters.WorldPosition_NoOffsets_CamRelative = WorldPosition;
+		MaterialParameters.WorldPosition_NoOffsets_CamRelative = WSPosition;
 	
 	{
-		float3 Normal = GetMaterialNormal(MaterialParameters);
-	
-		// tangent/decal space to world space
+		if(bVolumetric)
+		{
+			// not normalized
+			float3 OSNormal = GradientNotNormalized(ScreenPosition, OSPosition);
+			
+			// not normalized
+			float3 WSNormal = mul(float4(OSNormal, 0), transpose(WorldToDecal)).xyz;
+
+			MaterialParameters.WorldNormal = normalize(WSNormal);
+		}
+		else
+		{
+			// not normalized
+			float3 MaterialNormal = GetMaterialNormal(MaterialParameters);
+
+			// tangent/decal space to world space
 #if MATERIAL_TANGENTSPACENORMAL
-		MaterialParameters.WorldNormal = normalize(mul(float4(-Normal.z, Normal.y, Normal.x, 0), DecalToWorld).xyz);
+			// normals are transformed by the tranposed of the inverse
+			// DecalToWorld becomes transpose(WorldToDecal)
+			MaterialParameters.WorldNormal = normalize(mul(float4(-MaterialNormal.z, MaterialNormal.y, MaterialNormal.x, 0), transpose(WorldToDecal)).xyz);
 #else
-		MaterialParameters.WorldNormal = normalize(Normal);
+			MaterialParameters.WorldNormal = normalize(MaterialNormal);
 #endif
+		}
 	}
 
 	MaterialParameters.ReflectionVector = -MaterialParameters.CameraVector + MaterialParameters.WorldNormal * dot(MaterialParameters.WorldNormal, MaterialParameters.CameraVector) * 2.0;
@@ -134,7 +277,7 @@ void MainPS(
 
 	float3 Color = 1;
 
-	#if DECAL_BLEND_MODE == 2
+	#if DECAL_BLEND_MODE == DECALBLENDMODEID_NORMAL
 		// -1..1 range to 0..1
 		Color = MaterialParameters.WorldNormal * 0.5f + 0.5f;
 	#else
@@ -151,7 +294,7 @@ void MainPS(
 #endif
 
 	// Z fade
-	Opacity *= saturate( 4 - 4 * abs( DecalVector.z * 2 - 1 ) );
+	Opacity *= saturate( 4 - 4 * abs( SwizzlePos.z * 2 - 1 ) );
 
 	// RETURN_COLOR not needed unless writing to SceneColor
 	OutTarget0 = float4(Color, Opacity);
@@ -187,7 +330,18 @@ void MainPS(
 	}
 	#elif DECAL_RENDERSTAGE == 1
 	{
-		// after base pass, before lighting (GBuffer)
+		// after base pass (GBuffer)
+
+		// static lighting isn't updated by decals so we don't need to update that render target
+		float4 OutTarget4 = 0;
+		float4 OutTarget5 = 0;
+		float4 OutTarget6 = 0;
+
+		EncodeGBuffer(Data, OutTarget1, OutTarget2, OutTarget3, OutTarget4, OutTarget5, OutTarget6);
+	}
+	#elif DECAL_RENDERSTAGE == 2
+	{
+		// before lighting (GBuffer)
 
 		// static lighting isn't updated by decals so we don't need to update that render target
 		float4 OutTarget4 = 0;
@@ -196,8 +350,7 @@ void MainPS(
 
 		EncodeGBuffer(Data, OutTarget1, OutTarget2, OutTarget3, OutTarget4, OutTarget5, OutTarget6);
 
-		#if DECAL_BLEND_MODE == 1
-			// Stain
+		#if DECAL_BLEND_MODE == DECALBLENDMODEID_STAIN
 			OutTarget3.rgb *= Opacity;
 		#endif
 
diff --git a/Engine/Shaders/MaterialTemplate.usf b/Engine/Shaders/MaterialTemplate.usf
--- a/Engine/Shaders/MaterialTemplate.usf
+++ b/Engine/Shaders/MaterialTemplate.usf
@@ -1310,8 +1310,7 @@ half GetMaterialOpacityRaw(FMaterialPixelParameters Parameters)
 %s;
 }
 
-
-#if MATERIALBLENDING_MASKED
+#if MATERIALBLENDING_MASKED || (DECAL_BLEND_MODE == DECALBLENDMODEID_VOLUMETRIC)
 // Returns the material mask value generated from the material expressions.
 // Use GetMaterialMask() to get the value altered depending on the material blend mode.
 half GetMaterialMaskInputRaw(FMaterialPixelParameters Parameters)
diff --git a/Engine/Source/Runtime/Engine/Classes/Materials/Material.h b/Engine/Source/Runtime/Engine/Classes/Materials/Material.h
--- a/Engine/Source/Runtime/Engine/Classes/Materials/Material.h
+++ b/Engine/Source/Runtime/Engine/Classes/Materials/Material.h
@@ -72,6 +72,9 @@ enum EDecalBlendMode
 	/** Non metal, put into DBuffer to work for baked lighting as well. */
 	DBM_DBuffer_Roughness UMETA(DisplayName="DBuffer Translucent Roughness"),
 
+	/** Output signed distance in Opacity depending on LightVector. Note: Can be costly, no shadow casting but receiving, no per pixel normal yet, no quality settings yet */
+	DBM_Volumetric_DistanceFunction UMETA(DisplayName="Volumetric Distance Function (experimental)"),
+
 	DBM_MAX,
 };
 
diff --git a/Engine/Source/Runtime/Engine/Private/Materials/HLSLMaterialTranslator.h b/Engine/Source/Runtime/Engine/Private/Materials/HLSLMaterialTranslator.h
--- a/Engine/Source/Runtime/Engine/Private/Materials/HLSLMaterialTranslator.h
+++ b/Engine/Source/Runtime/Engine/Private/Materials/HLSLMaterialTranslator.h
@@ -350,7 +350,9 @@ public:
 				}
 			}
 
-			bUsesPixelDepthOffset = IsMaterialPropertyUsed(MP_PixelDepthOffset, Chunk[MP_PixelDepthOffset], FLinearColor(0, 0, 0, 0), 1);
+			bUsesPixelDepthOffset = IsMaterialPropertyUsed(MP_PixelDepthOffset, Chunk[MP_PixelDepthOffset], FLinearColor(0, 0, 0, 0), 1)
+				|| Material->GetDecalBlendMode() == DBM_Volumetric_DistanceFunction;
+
 			const bool bUsesWorldPositionOffset = IsMaterialPropertyUsed(MP_WorldPositionOffset, Chunk[MP_WorldPositionOffset], FLinearColor(0, 0, 0, 0), 3);
 			MaterialCompilationOutput.bModifiesMeshPosition = bUsesPixelDepthOffset || bUsesWorldPositionOffset;
 			
diff --git a/Engine/Source/Runtime/Engine/Private/Materials/Material.cpp b/Engine/Source/Runtime/Engine/Private/Materials/Material.cpp
--- a/Engine/Source/Runtime/Engine/Private/Materials/Material.cpp
+++ b/Engine/Source/Runtime/Engine/Private/Materials/Material.cpp
@@ -2494,7 +2494,7 @@ bool UMaterial::CanEditChange(const UProperty* InProperty) const
 	
 		if (PropertyName == GET_MEMBER_NAME_STRING_CHECKED(UMaterial, ShadingModel))
 		{
-			return MaterialDomain == MD_Surface;
+			return MaterialDomain == MD_Surface || (MaterialDomain == MD_DeferredDecal && DecalBlendMode == DBM_Volumetric_DistanceFunction);
 		}
 
 		if (PropertyName == GET_MEMBER_NAME_STRING_CHECKED(UMaterial, DecalBlendMode))
@@ -3888,6 +3888,15 @@ bool UMaterial::IsPropertyActive(EMaterialProperty InProperty) const
 				return InProperty == MP_Roughness
 					|| InProperty == MP_Opacity;
 
+			case DBM_Volumetric_DistanceFunction:
+				return InProperty == MP_EmissiveColor
+					|| InProperty == MP_Normal
+					|| InProperty == MP_Metallic
+					|| InProperty == MP_Specular
+					|| InProperty == MP_BaseColor
+					|| InProperty == MP_Roughness
+					|| InProperty == MP_OpacityMask;
+
 			default:
 				// if you create a new mode it needs to expose the right pins
 				return false;
diff --git a/Engine/Source/Runtime/Engine/Private/Materials/MaterialShared.cpp b/Engine/Source/Runtime/Engine/Private/Materials/MaterialShared.cpp
--- a/Engine/Source/Runtime/Engine/Private/Materials/MaterialShared.cpp
+++ b/Engine/Source/Runtime/Engine/Private/Materials/MaterialShared.cpp
@@ -526,7 +526,8 @@ bool FMaterial::MaterialMayModifyMeshPosition() const
 {
 	// Conservative estimate when called before material translation has occurred. 
 	// This function is only intended for use in deciding whether or not shader permutations are required.
-	return HasVertexPositionOffsetConnected() || HasPixelDepthOffsetConnected() || HasMaterialAttributesConnected() || GetTessellationMode() != MTM_NoTessellation;
+	return HasVertexPositionOffsetConnected() || HasPixelDepthOffsetConnected() || HasMaterialAttributesConnected() || GetTessellationMode() != MTM_NoTessellation
+		|| (GetMaterialDomain() == MD_DeferredDecal && GetDecalBlendMode() == DBM_Volumetric_DistanceFunction);
 }
 
 FMaterialShaderMap* FMaterial::GetRenderingThreadShaderMap() const 
diff --git a/Engine/Source/Runtime/Renderer/Private/CompositionLighting/CompositionLighting.cpp b/Engine/Source/Runtime/Renderer/Private/CompositionLighting/CompositionLighting.cpp
--- a/Engine/Source/Runtime/Renderer/Private/CompositionLighting/CompositionLighting.cpp
+++ b/Engine/Source/Runtime/Renderer/Private/CompositionLighting/CompositionLighting.cpp
@@ -217,22 +217,6 @@ static FRenderingCompositeOutputRef AddPostProcessingAmbientOcclusion(FRHIComman
 	return FRenderingCompositeOutputRef(AmbientOcclusionPassMip0);
 }
 
-static void AddDeferredDecalsBeforeBasePass(FPostprocessContext& Context)
-{
-	FRenderingCompositePass* Pass = Context.Graph.RegisterPass(new(FMemStack::Get()) FRCPassPostProcessDeferredDecals(0));
-	Pass->SetInput(ePId_Input0, Context.FinalOutput);
-
-	Context.FinalOutput = FRenderingCompositeOutputRef(Pass);
-}
-
-static void AddDeferredDecalsBeforeLighting(FPostprocessContext& Context)
-{
-	FRenderingCompositePass* Pass = Context.Graph.RegisterPass(new(FMemStack::Get()) FRCPassPostProcessDeferredDecals(1));
-	Pass->SetInput(ePId_Input0, Context.FinalOutput);
-
-	Context.FinalOutput = FRenderingCompositeOutputRef(Pass);
-}
-
 void FCompositionLighting::ProcessBeforeBasePass(FRHICommandListImmediate& RHICmdList, FViewInfo& View)
 {
 	check(IsInRenderingThread());
@@ -251,7 +235,10 @@ void FCompositionLighting::ProcessBeforeBasePass(FRHICommandListImmediate& RHICm
 			!Context.View.Family->EngineShowFlags.ShaderComplexity &&
 			IsDBufferEnabled()) 
 		{
-			AddDeferredDecalsBeforeBasePass(Context);
+			FRenderingCompositePass* Pass = Context.Graph.RegisterPass(new(FMemStack::Get()) FRCPassPostProcessDeferredDecals(DRS_BeforeBasePass));
+			Pass->SetInput(ePId_Input0, Context.FinalOutput);
+
+			Context.FinalOutput = FRenderingCompositeOutputRef(Pass);
 		}
 
 		// The graph setup should be finished before this line ----------------------------------------
@@ -287,13 +274,24 @@ void FCompositionLighting::ProcessAfterBasePass(FRHICommandListImmediate& RHICmd
 
 		FPostprocessContext Context(CompositeContext.Graph, View);
 
-		// Add the passes we want to add to the graph (commenting a line means the pass is not inserted into the graph) ----------
+		// Add the passes we want to add to the graph ----------
+		
+		if( Context.View.Family->EngineShowFlags.Decals)
+		{
+			FRenderingCompositePass* Pass = Context.Graph.RegisterPass(new(FMemStack::Get()) FRCPassPostProcessDeferredDecals(DRS_AfterBasePass));
+			Pass->SetInput(ePId_Input0, Context.FinalOutput);
+
+			Context.FinalOutput = FRenderingCompositeOutputRef(Pass);
+		}
 
 		// decals are before AmbientOcclusion so the decal can output a normal that AO is affected by
 		if( Context.View.Family->EngineShowFlags.Decals &&
 			!Context.View.Family->EngineShowFlags.VisualizeLightCulling)		// decal are distracting when looking at LightCulling
 		{
-			AddDeferredDecalsBeforeLighting(Context);
+			FRenderingCompositePass* Pass = Context.Graph.RegisterPass(new(FMemStack::Get()) FRCPassPostProcessDeferredDecals(DRS_BeforeLighting));
+			Pass->SetInput(ePId_Input0, Context.FinalOutput);
+
+			Context.FinalOutput = FRenderingCompositeOutputRef(Pass);
 		}
 
 		FRenderingCompositeOutputRef AmbientOcclusion;
diff --git a/Engine/Source/Runtime/Renderer/Private/CompositionLighting/PostProcessDeferredDecals.cpp b/Engine/Source/Runtime/Renderer/Private/CompositionLighting/PostProcessDeferredDecals.cpp
--- a/Engine/Source/Runtime/Renderer/Private/CompositionLighting/PostProcessDeferredDecals.cpp
+++ b/Engine/Source/Runtime/Renderer/Private/CompositionLighting/PostProcessDeferredDecals.cpp
@@ -28,20 +28,58 @@ enum ERenderTargetMode
 {
 	RTM_Unknown = -1,
 	RTM_SceneColorAndGBuffer,
+	RTM_SceneColorAndGBufferDepthWrite,
 	RTM_DBuffer,
 	RTM_GBufferNormal,
 	RTM_SceneColor,
 };
 
+enum EDecalDepthInputState
+{
+	DDS_Undefined,
+	DDS_Always,
+	DDS_DepthTest,
+	DDS_DepthAlways_StencilEqual1Write,
+	DDS_DepthAlways_StencilEqual0,
+	DDS_DepthTest_StencilEqual1Write,
+	DDS_DepthTest_StencilEqual0,
+};
+
+struct FDecalDepthState
+{
+	EDecalDepthInputState DepthTest;
+	bool bDepthOutput;
+
+	FDecalDepthState()
+		: DepthTest(DDS_Undefined)
+		, bDepthOutput(false)
+	{
+	}
+
+	bool operator !=(const FDecalDepthState &rhs) const
+	{
+		return DepthTest != rhs.DepthTest || bDepthOutput != rhs.bDepthOutput;
+	}
+};
+
+enum EDecalRasterizerState
+{
+	DRS_Undefined,
+	DRS_CCW,
+	DRS_CW,
+};
+
+
 // @return DECAL_RENDERTARGET_COUNT for the shader
 uint32 ComputeRenderTargetCount(ERenderTargetMode RenderTargetMode)
 {
 	switch(RenderTargetMode)
 	{
-		case RTM_SceneColorAndGBuffer:	return 4;
-		case RTM_DBuffer:				return 3;
-		case RTM_GBufferNormal:			return 1;
-		case RTM_SceneColor:			return 1;
+		case RTM_SceneColorAndGBuffer:				return 4;
+		case RTM_SceneColorAndGBufferDepthWrite:	return 4;
+		case RTM_DBuffer:							return 3;
+		case RTM_GBufferNormal:						return 1;
+		case RTM_SceneColor:						return 1;
 	}
 
 	return 0;
@@ -87,6 +125,9 @@ ERenderTargetMode ComputeRenderTargetMode(EDecalBlendMode DecalBlendMode)
 		case DBM_DBuffer_Roughness:
 			// can be optimized using less MRT when possible
 			return RTM_DBuffer;
+
+		case DBM_Volumetric_DistanceFunction:
+			return RTM_SceneColorAndGBufferDepthWrite;
 	}
 
 	// add the missing decal blend mode to the switch
@@ -94,31 +135,34 @@ ERenderTargetMode ComputeRenderTargetMode(EDecalBlendMode DecalBlendMode)
 	return RTM_Unknown;
 }
 
-// @return 0:before BasePass, 1:after base pass, before lighting, (later we could add "after lighting" and multiply)
-uint32 ComputeRenderStage(EDecalBlendMode DecalBlendMode)
+// @return see EDecalRenderStage
+EDecalRenderStage ComputeRenderStage(EDecalBlendMode DecalBlendMode)
 {
 	switch(DecalBlendMode)
 	{
-	case DBM_DBuffer_ColorNormalRoughness:
-	case DBM_DBuffer_Color:
-	case DBM_DBuffer_ColorNormal:
-	case DBM_DBuffer_ColorRoughness:
-	case DBM_DBuffer_Normal:
-	case DBM_DBuffer_NormalRoughness:
-	case DBM_DBuffer_Roughness:
-		return 0;
-
-	case DBM_Translucent:
-	case DBM_Stain:
-	case DBM_Normal:
-	case DBM_Emissive:
-		return 1;
-
-	default:
-		check(0);
-	}
+		case DBM_DBuffer_ColorNormalRoughness:
+		case DBM_DBuffer_Color:
+		case DBM_DBuffer_ColorNormal:
+		case DBM_DBuffer_ColorRoughness:
+		case DBM_DBuffer_Normal:
+		case DBM_DBuffer_NormalRoughness:
+		case DBM_DBuffer_Roughness:
+			return DRS_BeforeBasePass;
+
+		case DBM_Translucent:
+		case DBM_Stain:
+		case DBM_Normal:
+		case DBM_Emissive:
+			return DRS_BeforeLighting;
+		
+		case DBM_Volumetric_DistanceFunction:
+			return DRS_AfterBasePass;
 
-	return -1;
+		default:
+			check(0);
+	}
+	
+	return DRS_BeforeBasePass;
 }
 
 struct FTransientDecalRenderData
@@ -140,9 +184,9 @@ struct FTransientDecalRenderData
 	}
 };
 // @param RenderState 0:before BasePass, 1:before lighting, (later we could add "after lighting" and multiply)
-void SetDecalBlendState(FRHICommandList& RHICmdList, const ERHIFeatureLevel::Type SMFeatureLevel, uint32 RenderStage, EDecalBlendMode DecalBlendMode, bool bHasNormal)
+void SetDecalBlendState(FRHICommandList& RHICmdList, const ERHIFeatureLevel::Type SMFeatureLevel, EDecalRenderStage InDecalRenderStage, EDecalBlendMode DecalBlendMode, bool bHasNormal)
 {
-	if(RenderStage == 0)
+	if(InDecalRenderStage == DRS_BeforeBasePass)
 	{
 		// before base pass (for DBuffer decals)
 
@@ -229,6 +273,12 @@ void SetDecalBlendState(FRHICommandList& RHICmdList, const ERHIFeatureLevel::Typ
 
 		return;
 	}
+	else if(InDecalRenderStage == DRS_AfterBasePass)
+	{
+		ensure(DecalBlendMode == DBM_Volumetric_DistanceFunction);
+
+		RHICmdList.SetBlendState( TStaticBlendState<>::GetRHI() );
+	}	
 	else
 	{
 		// before lighting (for non DBuffer decals)
@@ -479,10 +529,16 @@ public:
 		ERenderTargetMode RenderTargetMode = ComputeRenderTargetMode(DecalBlendMode);
 		uint32 RenderTargetCount = ComputeRenderTargetCount(RenderTargetMode);
 
+		// avoid using the index directly, better use DECALBLENDMODEID_VOLUMETRIC, DECALBLENDMODEID_STAIN, ...
 		OutEnvironment.SetDefine(TEXT("DECAL_BLEND_MODE"), (int32)DecalBlendMode);
 		OutEnvironment.SetDefine(TEXT("DECAL_PROJECTION"), 1);
 		OutEnvironment.SetDefine(TEXT("DECAL_RENDERTARGET_COUNT"), RenderTargetCount);
-		OutEnvironment.SetDefine(TEXT("DECAL_RENDERSTAGE"), ComputeRenderStage(DecalBlendMode));
+		OutEnvironment.SetDefine(TEXT("DECAL_RENDERSTAGE"), (int32)ComputeRenderStage(DecalBlendMode));
+
+		// to compare against DECAL_BLEND_MODE, we can expose more if needed
+		OutEnvironment.SetDefine(TEXT("DECALBLENDMODEID_VOLUMETRIC"), (int32)DBM_Volumetric_DistanceFunction);
+		OutEnvironment.SetDefine(TEXT("DECALBLENDMODEID_STAIN"), (int32)DBM_Stain);
+		OutEnvironment.SetDefine(TEXT("DECALBLENDMODEID_NORMAL"), (int32)DBM_Normal);
 	}
 
 	FDeferredDecalPS() {}
@@ -491,6 +547,7 @@ public:
 	{
 		ScreenToDecal.Bind(Initializer.ParameterMap,TEXT("ScreenToDecal"));
 		DecalToWorld.Bind(Initializer.ParameterMap,TEXT("DecalToWorld"));
+		WorldToDecal.Bind(Initializer.ParameterMap,TEXT("WorldToDecal"));
 	}
 
 	void SetParameters(FRHICommandList& RHICmdList, const FSceneView& View, const FMaterialRenderProxy* MaterialProxy, const FDeferredDecalProxy& DecalProxy)
@@ -520,28 +577,31 @@ public:
 			SetShaderValue(RHICmdList, ShaderRHI, ScreenToDecal, ScreenToDecalValue);
 		}
 
-		// Set the transform from light space to world space (only for normals)
+		// Set the transform from light space to world space
 		if(DecalToWorld.IsBound())
 		{
-			const FMatrix DecalToWorldValue = ComponentTrans.ToMatrixNoScale();
+			const FMatrix DecalToWorldValue = ComponentTrans.ToMatrixWithScale();
 			
 			// 1,1,1 requires no scale
 			//			DecalToWorldValue = DecalToWorldValue.GetScaled(GDefaultDecalSize);
 
 			SetShaderValue(RHICmdList, ShaderRHI, DecalToWorld, DecalToWorldValue);
 		}
+
+		SetShaderValue(RHICmdList, ShaderRHI, WorldToDecal, WorldToComponent);
 	}
 
 	virtual bool Serialize(FArchive& Ar) override
 	{
 		bool bShaderHasOutdatedParameters = FMaterialShader::Serialize(Ar);
-		Ar << ScreenToDecal << DecalToWorld;
+		Ar << ScreenToDecal << DecalToWorld << WorldToDecal;
 		return bShaderHasOutdatedParameters;
 	}
 
 private:
 	FShaderParameter ScreenToDecal;
 	FShaderParameter DecalToWorld;
+	FShaderParameter WorldToDecal;
 };
 
 IMPLEMENT_MATERIAL_SHADER_TYPE(,FDeferredDecalPS,TEXT("DeferredDecal"),TEXT("MainPS"),SF_Pixel);
@@ -714,9 +774,156 @@ bool IsDBufferEnabled()
 }
 
 
-FRCPassPostProcessDeferredDecals::FRCPassPostProcessDeferredDecals(uint32 InRenderStage)
-	: RenderStage(InRenderStage)
+static EDecalRasterizerState ComputeDecalRasterizerState(bool bInsideDecal, const FViewInfo& View)
 {
+	bool bClockwise = bInsideDecal;
+
+	if(View.bReverseCulling)
+	{
+		bClockwise = !bClockwise;
+	}
+
+	return bClockwise ? DRS_CW : DRS_CCW;
+}
+
+FRCPassPostProcessDeferredDecals::FRCPassPostProcessDeferredDecals(EDecalRenderStage InDecalRenderStage)
+	: DecalRenderStage(InDecalRenderStage)
+{
+}
+
+static FDecalDepthState ComputeDecalDepthState(EDecalBlendMode DecalBlendMode, bool bInsideDecal, bool bStencilDecalsInThisStage, bool bThisDecalUsesStencil)
+{
+	FDecalDepthState Ret;
+
+	Ret.bDepthOutput = DecalBlendMode == DBM_Volumetric_DistanceFunction;
+				
+	if(Ret.bDepthOutput)
+	{
+		// can be made one enum
+		Ret.DepthTest = DDS_DepthTest;
+		return Ret;
+	}
+
+	if(bInsideDecal)
+	{
+		// Render backfaces with depth tests disabled since the camera is inside (or close to inside)
+		if(bStencilDecalsInThisStage)
+		{
+			// Enable stencil testing, only write to pixels with stencil of 0
+			if(bThisDecalUsesStencil)
+			{
+				Ret.DepthTest = DDS_DepthAlways_StencilEqual1Write;
+			}
+			else
+			{
+				Ret.DepthTest = DDS_DepthAlways_StencilEqual0;
+			}
+		}
+		else
+		{
+			Ret.DepthTest = DDS_Always;
+		}
+	}
+	else
+	{
+		// Render frontfaces with depth tests on to get the speedup from HiZ since the camera is outside
+		if(bStencilDecalsInThisStage)
+		{
+			// Render frontfaces with depth tests on to get the speedup from HiZ since the camera is outside
+			// Enable stencil testing, only write to pixels with stencil of 0
+			if(bThisDecalUsesStencil)
+			{
+				Ret.DepthTest = DDS_DepthTest_StencilEqual1Write;
+			}
+			else
+			{
+				Ret.DepthTest = DDS_DepthTest_StencilEqual0;
+			}
+		}
+		else
+		{
+			Ret.DepthTest = DDS_DepthTest;
+		}
+	}
+
+	return Ret;
+}
+
+static void SetDecalDepthState(FDecalDepthState DecalDepthState, FRHICommandListImmediate& RHICmdList)
+{
+	switch(DecalDepthState.DepthTest)
+	{
+		case DDS_DepthAlways_StencilEqual1Write:
+			check(!DecalDepthState.bDepthOutput);			// todo
+			RHICmdList.SetDepthStencilState(TStaticDepthStencilState<
+				false,CF_Always,
+				true,CF_Equal,SO_Zero,SO_Zero,SO_Zero,
+				true,CF_Equal,SO_Zero,SO_Zero,SO_Zero,
+				0xff, 0x7f>::GetRHI(), 1);
+			break;
+
+		case DDS_DepthAlways_StencilEqual0:
+			check(!DecalDepthState.bDepthOutput);			// todo
+			RHICmdList.SetDepthStencilState(TStaticDepthStencilState<
+				false,CF_Always,
+				true,CF_Equal,SO_Keep,SO_Keep,SO_Keep,
+				false,CF_Always,SO_Keep,SO_Keep,SO_Keep,
+				0x80,0x00>::GetRHI(), 0);
+			break;
+
+		case DDS_Always:
+			if(DecalDepthState.bDepthOutput)
+			{
+				RHICmdList.SetDepthStencilState(TStaticDepthStencilState<true, CF_Always, true>::GetRHI(), 0);
+			}
+			else
+			{
+				RHICmdList.SetDepthStencilState(TStaticDepthStencilState<false, CF_Always, true>::GetRHI(), 0);
+			}
+			break;
+
+		case DDS_DepthTest_StencilEqual1Write:
+			check(!DecalDepthState.bDepthOutput);			// todo
+			RHICmdList.SetDepthStencilState(TStaticDepthStencilState<
+				false,CF_DepthNearOrEqual,
+				true,CF_Equal,SO_Zero,SO_Zero,SO_Zero,
+				true,CF_Equal,SO_Zero,SO_Zero,SO_Zero,
+				0xff, 0x7f>::GetRHI(), 1);
+			break;
+
+		case DDS_DepthTest_StencilEqual0:
+			check(!DecalDepthState.bDepthOutput);			// todo
+			RHICmdList.SetDepthStencilState(TStaticDepthStencilState<
+				false,CF_DepthNearOrEqual,
+				true,CF_Equal,SO_Keep,SO_Keep,SO_Keep,
+				false,CF_Always,SO_Keep,SO_Keep,SO_Keep,
+				0x80,0x00>::GetRHI(), 0);
+			break;
+
+		case DDS_DepthTest:
+			if(DecalDepthState.bDepthOutput)
+			{
+				RHICmdList.SetDepthStencilState(TStaticDepthStencilState<true, CF_DepthNearOrEqual>::GetRHI(), 0);
+			}
+			else
+			{
+				RHICmdList.SetDepthStencilState(TStaticDepthStencilState<false, CF_DepthNearOrEqual>::GetRHI(), 0);
+			}
+			break;
+
+		default:
+			check(0);
+	}
+}
+
+static void SetDecalRasterizerState(EDecalRasterizerState DecalRasterizerState, FRHICommandList& RHICmdList)
+{
+	switch (DecalRasterizerState)
+	{
+		case DRS_CW: RHICmdList.SetRasterizerState(TStaticRasterizerState<FM_Solid, CM_CW>::GetRHI()); break;
+		case DRS_CCW: RHICmdList.SetRasterizerState(TStaticRasterizerState<FM_Solid, CM_CCW>::GetRHI()); break;
+		default: check(0);
+	}
 }
 
 void FRCPassPostProcessDeferredDecals::Process(FRenderingCompositePassContext& Context)
@@ -744,7 +951,7 @@ void FRCPassPostProcessDeferredDecals::Process(FRenderingCompositePassContext& C
 
 	FTextureRHIParamRef TargetsToResolve[ResolveBufferMax] = { nullptr };
 
-	if(RenderStage == 0)
+	if(DecalRenderStage == DRS_BeforeBasePass)
 	{
 		// before BasePass, only if DBuffer is enabled
 
@@ -808,8 +1015,12 @@ void FRCPassPostProcessDeferredDecals::Process(FRenderingCompositePassContext& C
 	const FSceneViewFamily& ViewFamily = *(View.Family);
 	FScene& Scene = *(FScene*)ViewFamily.Scene;
 
-	if (Scene.Decals.Num())
+	if(!Scene.Decals.Num())
 	{
+		// to avoid the stats showing up
+		return;
+	}
+
 	TArray<FTransientDecalRenderData, SceneRenderingAllocator> SortedDecals;
 	SortedDecals.Empty(Scene.Decals.Num());
 
@@ -848,32 +1059,32 @@ void FRCPassPostProcessDeferredDecals::Process(FRenderingCompositePassContext& C
 		{
 			FTransientDecalRenderData Data(Scene, DecalProxy);
 
-			uint32 DecalRenderStage = ComputeRenderStage(Data.DecalBlendMode);
+			EDecalRenderStage LocalDecalRenderStage = ComputeRenderStage(Data.DecalBlendMode);
 
 			// we could do this test earlier to avoid the decal intersection but getting DecalBlendMode also costs
-			if (Context.View.Family->EngineShowFlags.ShaderComplexity || RenderStage == DecalRenderStage)
+			if (Context.View.Family->EngineShowFlags.ShaderComplexity || DecalRenderStage == LocalDecalRenderStage)
 			{
 				SortedDecals.Add(Data);
 			}
 		}
 	}
 
-		if (SortedDecals.Num() > 0)
+	if (SortedDecals.Num() > 0)
 	{
 		FIntRect SrcRect = View.ViewRect;
 		FIntRect DestRect = View.ViewRect;
-
-		bool bStencilDecals = true;
+		
+		bool bStencilDecalsInThisStage = true;
 
 #if DBUFFER_DONT_USE_STENCIL_YET
-			if (RenderStage == 0)
+		if(DecalRenderStage != DRS_BeforeLighting)
 		{
-			bStencilDecals = false;
+			bStencilDecalsInThisStage = false;
 		}
 #endif
 
 		// Setup a stencil mask to prevent certain pixels from receiving deferred decals
-			if (bStencilDecals)
+		if(bStencilDecalsInThisStage)
 		{
 			StencilDecalMask(RHICmdList, View);
 		}
@@ -907,13 +1118,14 @@ void FRCPassPostProcessDeferredDecals::Process(FRenderingCompositePassContext& C
 		};
 
 		// Sort decals by blend mode to reduce render target switches
-			SortedDecals.Sort(FCompareFTransientDecalRenderData());
+		SortedDecals.Sort(FCompareFTransientDecalRenderData());
 
 		// optimization to have less state changes
+		EDecalRasterizerState LastDecalRasterizerState = DRS_Undefined;
+		FDecalDepthState LastDecalDepthState;
 		int32 LastDecalBlendMode = -1;
 		int32 LastDecalHasNormal = -1; // Decal state can change based on its normal property.(SM5)
 		ERenderTargetMode LastRenderTargetMode = RTM_Unknown;
-		int32 WasInsideDecal = -1;
 		const ERHIFeatureLevel::Type SMFeatureLevel = Context.GetFeatureLevel();
 
 		SCOPED_DRAW_EVENT(RHICmdList, Decals);
@@ -942,10 +1154,10 @@ void FRCPassPostProcessDeferredDecals::Process(FRenderingCompositePassContext& C
 
 			EDecalBlendMode DecalBlendMode = DecalData.DecalBlendMode;
 
-			bool bStencilThisDecal = bStencilDecals;
+			bool bStencilThisDecal = bStencilDecalsInThisStage;
 			
 #if DBUFFER_DONT_USE_STENCIL_YET
-				if (ComputeRenderStage(DecalBlendMode) == 0)
+			if(ComputeRenderStage(DecalBlendMode) != DRS_BeforeLighting)
 			{
 				bStencilThisDecal = false;
 			}
@@ -953,7 +1165,7 @@ void FRCPassPostProcessDeferredDecals::Process(FRenderingCompositePassContext& C
 
 			ERenderTargetMode CurrentRenderTargetMode = ComputeRenderTargetMode(DecalBlendMode);
 
-				if (bShaderComplexity)
+			if (bShaderComplexity)
 			{
 				CurrentRenderTargetMode = RTM_SceneColor;
 				// we want additive blending for the ShaderComplexity mode
@@ -961,11 +1173,11 @@ void FRCPassPostProcessDeferredDecals::Process(FRenderingCompositePassContext& C
 			}
 
 			// fewer rendertarget switches if possible
-				if (CurrentRenderTargetMode != LastRenderTargetMode)
+			if (CurrentRenderTargetMode != LastRenderTargetMode)
 			{
 				LastRenderTargetMode = CurrentRenderTargetMode;
 
-					switch (CurrentRenderTargetMode)
+				switch (CurrentRenderTargetMode)
 				{
 					case RTM_SceneColorAndGBuffer:
 						{							
@@ -978,6 +1190,17 @@ void FRCPassPostProcessDeferredDecals::Process(FRenderingCompositePassContext& C
 						}
 						break;
 
+					case RTM_SceneColorAndGBufferDepthWrite:
+						{							
+							TargetsToResolve[SceneColorIndex] = SceneContext.GetSceneColor()->GetRenderTargetItem().TargetableTexture;
+							TargetsToResolve[GBufferAIndex] = SceneContext.GBufferA->GetRenderTargetItem().TargetableTexture;
+							TargetsToResolve[GBufferBIndex] = SceneContext.GBufferB->GetRenderTargetItem().TargetableTexture;
+							TargetsToResolve[GBufferCIndex] = SceneContext.GBufferC->GetRenderTargetItem().TargetableTexture;
+							
+							SetRenderTargets(RHICmdList, 4, TargetsToResolve, SceneContext.GetSceneDepthSurface(), ESimpleRenderTargetMode::EExistingColorAndDepth, FExclusiveDepthStencil::DepthWrite_StencilWrite);
+						}
+						break;
+
 					case RTM_GBufferNormal:
 						TargetsToResolve[GBufferAIndex] = SceneContext.GBufferA->GetRenderTargetItem().TargetableTexture;
 						SetRenderTarget(RHICmdList, TargetsToResolve[GBufferAIndex], SceneContext.GetSceneDepthSurface(), ESimpleRenderTargetMode::EExistingColorAndDepth, FExclusiveDepthStencil::DepthRead_StencilWrite);
@@ -1016,7 +1239,8 @@ void FRCPassPostProcessDeferredDecals::Process(FRenderingCompositePassContext& C
 					// note this is after a SetStreamSource (in if CurrentRenderTargetMode != LastRenderTargetMode) call as it needs to get the VB input
 					bThisDecalUsesStencil = RenderPreStencil(Context, MaterialShaderMap, ComponentToWorldMatrix, FrustumComponentToClip);
 
-					WasInsideDecal = -1;
+					LastDecalRasterizerState = DRS_Undefined;
+					LastDecalDepthState = FDecalDepthState();
 					LastDecalBlendMode = -1;
 				}
 			}
@@ -1032,93 +1256,49 @@ void FRCPassPostProcessDeferredDecals::Process(FRenderingCompositePassContext& C
 				LastDecalBlendMode = DecalBlendMode;
 				LastDecalHasNormal = (int32)DecalData.bHasNormal;
 
-				SetDecalBlendState(RHICmdList, SMFeatureLevel, RenderStage, (EDecalBlendMode)LastDecalBlendMode, DecalData.bHasNormal);
+				SetDecalBlendState(RHICmdList, SMFeatureLevel, DecalRenderStage, (EDecalBlendMode)LastDecalBlendMode, DecalData.bHasNormal);
 			}
 
-			{
-				TShaderMapRef<FDeferredDecalVS> VertexShader(Context.GetShaderMap());
-				SetShader(Context, bShaderComplexity, DecalData, *VertexShader);
+				// todo
+//			const int32 IsInsideDecal = ((FVector)View.ViewMatrices.ViewOrigin - ComponentToWorldMatrix.GetOrigin()).SizeSquared() < FMath::Square(ConservativeRadius * 1.05f + View.NearClippingDistance * 2.0f) + ( bThisDecalUsesStencil ) ? 2 : 0;
+			const bool bInsideDecal = ((FVector)View.ViewMatrices.ViewOrigin - ComponentToWorldMatrix.GetOrigin()).SizeSquared() < FMath::Square(ConservativeRadius * 1.05f + View.NearClippingDistance * 2.0f);
+//			const bool bInsideDecal =  !(IsInsideDecal & 1);
 
-				VertexShader->SetParameters(RHICmdList, View, FrustumComponentToClip);
+			// update rasterizer state if needed
+			{
+				EDecalRasterizerState DecalRasterizerState = ComputeDecalRasterizerState(bInsideDecal, View);
 
-					const int32 IsInsideDecal = ((FVector)View.ViewMatrices.ViewOrigin - ComponentToWorldMatrix.GetOrigin()).SizeSquared() < FMath::Square(ConservativeRadius * 1.05f + View.NearClippingDistance * 2.0f) + (bThisDecalUsesStencil) ? 2 : 0;
-					if (WasInsideDecal != IsInsideDecal)
+				if(LastDecalRasterizerState != DecalRasterizerState)
 				{
-					WasInsideDecal = IsInsideDecal;
-						if (!(IsInsideDecal & 1))
-					{
-						// Render backfaces with depth tests disabled since the camera is inside (or close to inside) the light function geometry
-						RHICmdList.SetRasterizerState(View.bReverseCulling ? TStaticRasterizerState<FM_Solid, CM_CCW>::GetRHI() : TStaticRasterizerState<FM_Solid, CM_CW>::GetRHI());
-							if (bStencilDecals)
-						{
-							// Enable stencil testing, only write to pixels with stencil of 0
-								if (bThisDecalUsesStencil)
-							{
-								RHICmdList.SetDepthStencilState(TStaticDepthStencilState<
-										false, CF_Always,
-										true, CF_Equal, SO_Zero, SO_Zero, SO_Zero,
-										true, CF_Equal, SO_Zero, SO_Zero, SO_Zero,
-									0xff, 0x7f
-								>::GetRHI(), 1);
-							}
-							else
-							{
-								RHICmdList.SetDepthStencilState(TStaticDepthStencilState<
-										false, CF_Always,
-										true, CF_Equal, SO_Keep, SO_Keep, SO_Keep,
-										false, CF_Always, SO_Keep, SO_Keep, SO_Keep,
-										0x80, 0x00>::GetRHI(), 0);
-							}
-						}
-						else
-						{
-							RHICmdList.SetDepthStencilState(TStaticDepthStencilState<false, CF_Always, true>::GetRHI(), 0);
-						}
-					}
-					else
-					{
-						// Render frontfaces with depth tests on to get the speedup from HiZ since the camera is outside the light function geometry
-							if (bStencilDecals)
-						{
-							// Render frontfaces with depth tests on to get the speedup from HiZ since the camera is outside the light function geometry
-							// Enable stencil testing, only write to pixels with stencil of 0
-								if (bThisDecalUsesStencil)
-							{
-								RHICmdList.SetDepthStencilState(TStaticDepthStencilState<
-										false, CF_DepthNearOrEqual,
-										true, CF_Equal, SO_Zero, SO_Zero, SO_Zero,
-										true, CF_Equal, SO_Zero, SO_Zero, SO_Zero,
-									0xff, 0x7f
-								>::GetRHI(), 1);
-							}
-							else
-							{
-								RHICmdList.SetDepthStencilState(TStaticDepthStencilState<
-										false, CF_DepthNearOrEqual,
-										true, CF_Equal, SO_Keep, SO_Keep, SO_Keep,
-										false, CF_Always, SO_Keep, SO_Keep, SO_Keep,
-										0x80, 0x00>::GetRHI(), 0);
-							}
-							RHICmdList.SetRasterizerState(View.bReverseCulling ? TStaticRasterizerState<FM_Solid, CM_CW>::GetRHI() : TStaticRasterizerState<FM_Solid, CM_CCW>::GetRHI());
-						}
-						else
-						{
-							RHICmdList.SetDepthStencilState(TStaticDepthStencilState<false, CF_DepthNearOrEqual>::GetRHI(), 0);
-						}
-						RHICmdList.SetRasterizerState(View.bReverseCulling ? TStaticRasterizerState<FM_Solid, CM_CW>::GetRHI() : TStaticRasterizerState<FM_Solid, CM_CCW>::GetRHI());
-					}
+					LastDecalRasterizerState = DecalRasterizerState;
+					SetDecalRasterizerState(DecalRasterizerState, RHICmdList);
 				}
+			}
 
-				RHICmdList.DrawIndexedPrimitive(GUnitCubeIndexBuffer.IndexBufferRHI, PT_TriangleList, 0, 0, 8, 0, GUnitCubeIndexBuffer.GetIndexCount() / 3, 1);
+			// update DepthStencil state if needed
+			{
+				FDecalDepthState DecalDepthState = ComputeDecalDepthState(DecalBlendMode, bInsideDecal, bStencilDecalsInThisStage, bThisDecalUsesStencil);
+
+				if(LastDecalDepthState != DecalDepthState)
+				{
+					LastDecalDepthState = DecalDepthState;
+					SetDecalDepthState(DecalDepthState, RHICmdList);
+				}
 			}
+
+			TShaderMapRef<FDeferredDecalVS> VertexShader(Context.GetShaderMap());
+			SetShader(Context, bShaderComplexity, DecalData, *VertexShader);
+
+			VertexShader->SetParameters(RHICmdList, View, FrustumComponentToClip);
+
+			RHICmdList.DrawIndexedPrimitive(GUnitCubeIndexBuffer.IndexBufferRHI, PT_TriangleList, 0, 0, 8, 0, GUnitCubeIndexBuffer.GetIndexCount() / 3, 1);
 		}
 
 		// we don't modify stencil but if out input was having stencil for us (after base pass - we need to clear)
 		// Clear stencil to 0, which is the assumed default by other passes
 		RHICmdList.Clear(false, FLinearColor::White, false, (float)ERHIZBuffer::FarPlane, true, 0, FIntRect());
-		}
-
-		if (RenderStage == 0)
+		
+		if(DecalRenderStage == DRS_BeforeBasePass)
 		{
 			// before BasePass
 			GRenderTargetPool.VisualizeTexture.SetCheckPoint(RHICmdList, SceneContext.DBufferA);
@@ -1127,16 +1307,16 @@ void FRCPassPostProcessDeferredDecals::Process(FRenderingCompositePassContext& C
 		}
 	}
 
-		// resolve the targets we wrote to.
-		FResolveParams ResolveParams;
-		for (int32 i = 0; i < ResolveBufferMax; ++i)
+	// resolve the targets we wrote to.
+	FResolveParams ResolveParams;
+	for (int32 i = 0; i < ResolveBufferMax; ++i)
+	{
+		if (TargetsToResolve[i])
 		{
-			if (TargetsToResolve[i])
-			{
-				RHICmdList.CopyToResolveTarget(TargetsToResolve[i], TargetsToResolve[i], true, ResolveParams);
-			}
-		}		
-	}
+			RHICmdList.CopyToResolveTarget(TargetsToResolve[i], TargetsToResolve[i], true, ResolveParams);
+		}
+	}		
+}
 
 FPooledRenderTargetDesc FRCPassPostProcessDeferredDecals::ComputeOutputDesc(EPassOutputId InPassOutputId) const
 {
diff --git a/Engine/Source/Runtime/Renderer/Private/CompositionLighting/PostProcessDeferredDecals.h b/Engine/Source/Runtime/Renderer/Private/CompositionLighting/PostProcessDeferredDecals.h
--- a/Engine/Source/Runtime/Renderer/Private/CompositionLighting/PostProcessDeferredDecals.h
+++ b/Engine/Source/Runtime/Renderer/Private/CompositionLighting/PostProcessDeferredDecals.h
@@ -8,13 +8,27 @@
 
 #include "RenderingCompositionGraph.h"
 
+
+// Actual values are used in the shader so do not change
+enum EDecalRenderStage
+{
+	// for DBuffer decals (get proper baked lighting)
+	DRS_BeforeBasePass = 0,
+	// for volumetrics to update the depth buffer
+	DRS_AfterBasePass = 1,
+	// for normal decals not modifying the depth buffer
+	DRS_BeforeLighting = 2,
+
+	// later we could add "after lighting" and multiply
+};
+
 // ePId_Input0: SceneColor (not needed for DBuffer decals)
 // derives from TRenderingCompositePassBase<InputCount, OutputCount> 
 class FRCPassPostProcessDeferredDecals : public TRenderingCompositePassBase<1, 1>
 {
 public:
-	// @param InRenderStage 0:before BasePass, 1:before lighting, 2:AfterLighting 
-	FRCPassPostProcessDeferredDecals(uint32 InRenderStage);
+	// One instance for each render stage
+	FRCPassPostProcessDeferredDecals(EDecalRenderStage InDecalRenderStage);
 
 	// interface FRenderingCompositePass ---------
 	virtual void Process(FRenderingCompositePassContext& Context) override;
@@ -22,8 +36,8 @@ public:
 	virtual FPooledRenderTargetDesc ComputeOutputDesc(EPassOutputId InPassOutputId) const override;
 
 private:
-	// 0:before BasePass, 1:before lighting, (later we could add "after lighting" and multiply)
-	uint32 RenderStage;
+	// see EDecalRenderStage
+	EDecalRenderStage DecalRenderStage;
 };
 
 bool IsDBufferEnabled();
