Added support for dynamic point lights on Mobile (UE-11983)

diff --git a/Engine/Shaders/BasePassForForwardShadingCommon.usf b/Engine/Shaders/BasePassForForwardShadingCommon.usf
--- a/Engine/Shaders/BasePassForForwardShadingCommon.usf
+++ b/Engine/Shaders/BasePassForForwardShadingCommon.usf
@@ -12,7 +12,7 @@
 #define USE_VERTEX_FOG					(!OUTPUT_GAMMA_SPACE)
 
 #define PACK_INTERPOLANTS (USE_VERTEX_FOG && NUM_VF_PACKED_INTERPOLANTS > 0 && (ES2_PROFILE || ES3_1_PROFILE))
-#define LANDSCAPE_BUG_WORKAROUND (IS_FORWARD_BASEPASS_VERTEX_SHADER && PACK_INTERPOLANTS)
+#define LANDSCAPE_BUG_WORKAROUND (IOS && IS_FORWARD_BASEPASS_VERTEX_SHADER && PACK_INTERPOLANTS)
 
 #if METAL_PROFILE 
 // @todo-rco: @todo metal: FIXME!
diff --git a/Engine/Shaders/BasePassForForwardShadingPixelShader.usf b/Engine/Shaders/BasePassForForwardShadingPixelShader.usf
--- a/Engine/Shaders/BasePassForForwardShadingPixelShader.usf
+++ b/Engine/Shaders/BasePassForForwardShadingPixelShader.usf
@@ -13,6 +13,15 @@
 #include "BRDF.usf"
 #include "SHCommon.usf"
 #include "ShadowFilteringCommon.usf"
+#include "DynamicLightingCommon.usf"
+
+#if MAX_DYNAMIC_POINT_LIGHTS > 0
+#if VARIABLE_NUM_DYNAMIC_POINT_LIGHTS
+int NumDynamicPointLights;
+#endif
+float4 LightPositionAndInvRadius[MAX_DYNAMIC_POINT_LIGHTS];
+float4 LightColorAndFalloffExponent[MAX_DYNAMIC_POINT_LIGHTS];
+#endif
 
 /** Prenormalized capture of the scene that's closest to the object being rendered. */
 #if !MATERIAL_FULLY_ROUGH
@@ -201,7 +210,7 @@ void Main(
 			Shadow = DirectionalLightShadowing;
 		#endif
 			
-#if MOVABLE_DIRECTIONAL_LIGHT_CSM
+		#if MOVABLE_DIRECTIONAL_LIGHT_CSM
 		// Cascaded Shadow Map
 		{
 			FPCFSamplerSettings Settings;
@@ -229,14 +238,14 @@ void Main(
 			float LightSpacePixelDepthForOpaque = min(ShadowPosition.z, 0.99999f);
 			Settings.SceneDepth = LightSpacePixelDepthForOpaque;
 
-	#if 0
+		#if 0
 			half ShadowMap = Manual1x1PCF(ShadowPosition.xy, Settings);
-	#else
+		#else
 			half ShadowMap = Manual2x2PCF(ShadowPosition.xy, Settings);
-	#endif
+		#endif
 			Shadow = ShadowMap;
 		}
-#endif
+		#endif /* MOVABLE_DIRECTIONAL_LIGHT_CSM */
 
 		half NoL = max(0, dot(MaterialParameters.WorldNormal, View.DirectionalLightDirection));
 		half RoL = max(0, dot(MaterialParameters.ReflectionVector, View.DirectionalLightDirection));
@@ -253,8 +262,32 @@ void Main(
 			// TODO: Also need to do specular.
 			Color += GetSkySHDiffuseSimple(MaterialParameters.WorldNormal) * View.SkyLightColor.rgb * DiffuseColor;    
 		#endif
-	#endif
-			 
+
+		#if MAX_DYNAMIC_POINT_LIGHTS > 0
+			#if VARIABLE_NUM_DYNAMIC_POINT_LIGHTS
+				for (int i = 0; i < NumDynamicPointLights; i++)
+			#else
+				for (int i = 0; i < NUM_DYNAMIC_POINT_LIGHTS; i++)
+			#endif
+			{
+				half3 ToLight = LightPositionAndInvRadius[i].xyz - MaterialParameters.WorldPosition;
+				half DistanceSqr = dot(ToLight, ToLight);
+				half3 L = ToLight * rsqrt(DistanceSqr);
+
+				half PointNoL = max(0, dot(MaterialParameters.WorldNormal, L));
+				half PointRoL = max(0, dot(MaterialParameters.ReflectionVector, L));
+
+				half DistanceAttenuation = RadialAttenuation(ToLight * LightPositionAndInvRadius[i].w, LightColorAndFalloffExponent[i].w);		
+				#if !MATERIAL_FULLY_ROUGH
+					Color += (DistanceAttenuation * PointNoL) * LightColorAndFalloffExponent[i].rgb * (DiffuseColor + SpecularColor * PhongApprox(Roughness, PointRoL));
+				#else
+					Color += (DistanceAttenuation * PointNoL) * LightColorAndFalloffExponent[i].rgb * DiffuseColor;
+				#endif
+			}
+		#endif
+
+	#endif /* !MATERIAL_SHADINGMODEL_UNLIT */
+		 
 	half Opacity = GetMaterialOpacity(MaterialParameters);
 	half3 Emissive = GetMaterialEmissive(MaterialParameters);
 
diff --git a/Engine/Source/Runtime/Core/Private/HAL/ConsoleManager.cpp b/Engine/Source/Runtime/Core/Private/HAL/ConsoleManager.cpp
--- a/Engine/Source/Runtime/Core/Private/HAL/ConsoleManager.cpp
+++ b/Engine/Source/Runtime/Core/Private/HAL/ConsoleManager.cpp
@@ -1601,6 +1601,19 @@ static TAutoConsoleVariable<int32> CVarMobileHDR(
 	TEXT("1: Mobile renders in HDR linear space. (default)"),
 	ECVF_RenderThreadSafe | ECVF_ReadOnly);
 
+static TAutoConsoleVariable<int32> CVarMobileNumDynamicPointLights(
+	TEXT("r.MobileNumDynamicPointLights"),
+	4,
+	TEXT("The number of dynamic point lights to support on mobile devices. Setting this to 0 for games which do not require dynamic point lights will reduce the number of shaders generated."), 
+	ECVF_RenderThreadSafe | ECVF_ReadOnly);
+
+static TAutoConsoleVariable<int32> CVarMobileDynamicPointLightsUseStaticBranch(
+	TEXT("r.MobileDynamicPointLightsUseStaticBranch"),
+	1,
+	TEXT("0: Generate unique forward rendering base pass shaders for 0, 1, ... N mobile dynamic point lights. (faster but generates many more shaders)\n")
+	TEXT("1: Use a shared shader with static branching for rendering 1 or more dynamic point lights (slightly slower but reduces shaders generated, recommended for most games)."),
+	ECVF_RenderThreadSafe | ECVF_ReadOnly);
+
 static TAutoConsoleVariable<int32> CVarMobileHDR32bpp(
 	TEXT("r.MobileHDR32bpp"),
 	0,
diff --git a/Engine/Source/Runtime/Engine/Classes/Engine/RendererSettings.h b/Engine/Source/Runtime/Engine/Classes/Engine/RendererSettings.h
--- a/Engine/Source/Runtime/Engine/Classes/Engine/RendererSettings.h
+++ b/Engine/Source/Runtime/Engine/Classes/Engine/RendererSettings.h
@@ -98,6 +98,18 @@ class ENGINE_API URendererSettings : public UDeveloperSettings
 		ToolTip="If true, mobile renders in full HDR. Disable this setting for games that do not require lighting features for better performance on slow devices."))
 	uint32 bMobileHDR:1;
 
+	UPROPERTY(config, EditAnywhere, Category = Mobile, meta = (
+		ConsoleVariable = "r.MobileNumDynamicPointLights", DisplayName = "Max Dynamic Point Lights", ClampMax = 4, 
+		ToolTip = "The number of dynamic point lights to support on mobile devices. Setting this to 0 for games which do not require dynamic point lights will reduce the number of shaders generated. Changing this setting requires restarting the editor.",
+		ConfigRestartRequired = true))
+	uint32 MobileNumDynamicPointLights;
+
+	UPROPERTY(config, EditAnywhere, Category = Mobile, meta = (
+		ConsoleVariable = "r.MobileDynamicPointLightsUseStaticBranch", DisplayName = "Use Shared Dynamic Point Light Shaders",
+		ToolTip = "If this setting is enabled, the same shader will be used for any number of dynamic point lights (up to the maximum specified above) hitting a surface. This is slightly slower but reduces the number of shaders generated. Changing this setting requires restarting the editor.",
+		ConfigRestartRequired = true))
+	uint32 bMobileDynamicPointLightsUseStaticBranch : 1;
+
 	UPROPERTY(config, EditAnywhere, Category=Culling, meta=(
 		ConsoleVariable="r.AllowOcclusionQueries",DisplayName="Occlusion Culling",
 		ToolTip="Allows occluded meshes to be culled and no rendered."))
diff --git a/Engine/Source/Runtime/Engine/Private/Materials/MaterialShared.cpp b/Engine/Source/Runtime/Engine/Private/Materials/MaterialShared.cpp
--- a/Engine/Source/Runtime/Engine/Private/Materials/MaterialShared.cpp
+++ b/Engine/Source/Runtime/Engine/Private/Materials/MaterialShared.cpp
@@ -992,7 +992,7 @@ void FMaterialResource::GetRepresentativeShaderTypesAndDescriptions(TArray<FStri
 		if (GetShadingModel() == MSM_Unlit)
 		{
 			//unlit materials are never lightmapped
-			new (ShaderTypeNames)FString(FString::Printf(TEXT("TBasePassForForwardShadingPSFNoLightMapPolicy%s"), ShaderSuffix));
+			new (ShaderTypeNames)FString(FString::Printf(TEXT("TBasePassForForwardShadingPSFNoLightMapPolicy0%s"), ShaderSuffix));
 			new (ShaderTypeDescriptions)FString(FString::Printf(TEXT("Mobile base pass shader without light map%s"), DescSuffix));
 		}
 		else
@@ -1000,16 +1000,16 @@ void FMaterialResource::GetRepresentativeShaderTypesAndDescriptions(TArray<FStri
 			if (IsUsedWithStaticLighting())
 			{
 				//lit materials are usually lightmapped
-				new (ShaderTypeNames)FString(FString::Printf(TEXT("TBasePassForForwardShadingPSTLightMapPolicyLQ%s"), ShaderSuffix));
+				new (ShaderTypeNames)FString(FString::Printf(TEXT("TBasePassForForwardShadingPSTLightMapPolicy0LQ%s"), ShaderSuffix));
 				new (ShaderTypeDescriptions)FString(FString::Printf(TEXT("Mobile base pass shader with static lighting%s"), DescSuffix));
 
 				// + distance field shadows
-				new (ShaderTypeNames)FString(FString::Printf(TEXT("TBasePassForForwardShadingPSTDistanceFieldShadowsAndLightMapPolicyLQ%s"), ShaderSuffix));
+				new (ShaderTypeNames)FString(FString::Printf(TEXT("TBasePassForForwardShadingPSTDistanceFieldShadowsAndLightMapPolicy0LQ%s"), ShaderSuffix));
 				new (ShaderTypeDescriptions)FString(FString::Printf(TEXT("Mobile base pass shader with distance field shadows%s"), DescSuffix));
 			}
 
 			//also show a dynamically lit shader
-			new (ShaderTypeNames)FString(FString::Printf(TEXT("TBasePassForForwardShadingPSFSimpleDirectionalLightAndSHIndirectPolicy%s"), ShaderSuffix));
+			new (ShaderTypeNames)FString(FString::Printf(TEXT("TBasePassForForwardShadingPSFSimpleDirectionalLightAndSHIndirectPolicy0%s"), ShaderSuffix));
 			new (ShaderTypeDescriptions)FString(FString::Printf(TEXT("Mobile base pass shader with only dynamic lighting%s"), DescSuffix));
 		}
 
diff --git a/Engine/Source/Runtime/Engine/Private/PrimitiveSceneProxy.cpp b/Engine/Source/Runtime/Engine/Private/PrimitiveSceneProxy.cpp
--- a/Engine/Source/Runtime/Engine/Private/PrimitiveSceneProxy.cpp
+++ b/Engine/Source/Runtime/Engine/Private/PrimitiveSceneProxy.cpp
@@ -32,6 +32,7 @@ FPrimitiveSceneProxy::FPrimitiveSceneProxy(const UPrimitiveComponent* InComponen
 ,	bIsComponentLevelVisible(false)
 ,	bCollisionEnabled(InComponent->IsCollisionEnabled())
 ,	bTreatAsBackgroundForOcclusion(InComponent->bTreatAsBackgroundForOcclusion)
+,	bDisableStaticPath(false)
 ,	bNeedsUnbuiltPreviewLighting(InComponent->HasStaticLighting() && !InComponent->bHasCachedStaticLighting)
 ,	bHasValidSettingsForStaticLighting(InComponent->HasValidSettingsForStaticLighting())
 ,	bWillEverBeLit(true)
diff --git a/Engine/Source/Runtime/Engine/Private/StaticMeshRender.cpp b/Engine/Source/Runtime/Engine/Private/StaticMeshRender.cpp
--- a/Engine/Source/Runtime/Engine/Private/StaticMeshRender.cpp
+++ b/Engine/Source/Runtime/Engine/Private/StaticMeshRender.cpp
@@ -605,12 +605,13 @@ void FStaticMeshSceneProxy::GetDynamicMeshElements(const TArray<const FSceneView
 	const bool bDrawComplexWireframeCollision = (EngineShowFlags.Collision && IsCollisionEnabled() && CollisionTraceFlag == ECollisionTraceFlag::CTF_UseComplexAsSimple);
 
 	const bool bDrawMesh = (bInCollisionView) ? (bDrawComplexCollision) : 
-		(	IsRichView(ViewFamily) || HasViewDependentDPG()
-			|| EngineShowFlags.Collision
-			|| EngineShowFlags.Bounds
-			|| bProxyIsSelected 
-			|| IsHovered()
-			|| bIsLightmapSettingError ) ;
+	(	IsRichView(ViewFamily) || HasViewDependentDPG()
+		|| EngineShowFlags.Collision
+		|| EngineShowFlags.Bounds
+		|| bProxyIsSelected 
+		|| IsHovered()
+		|| bIsLightmapSettingError 
+		|| !IsStaticPathAvailable() );
 
 	// Draw polygon mesh if we are either not in a collision view, or are drawing it as collision.
 	if (EngineShowFlags.StaticMeshes && bDrawMesh)
@@ -920,7 +921,8 @@ FPrimitiveViewRelevance FStaticMeshSceneProxy::GetViewRelevance(const FSceneView
 #endif
 		// Force down dynamic rendering path if invalid lightmap settings, so we can apply an error material in DrawRichMesh
 		(HasStaticLighting() && !HasValidSettingsForStaticLighting()) ||
-		HasViewDependentDPG() 
+		HasViewDependentDPG() ||
+		 !IsStaticPathAvailable()
 #if WITH_EDITOR
 		//only check these in the editor
 		|| IsSelected() || IsHovered()
diff --git a/Engine/Source/Runtime/Engine/Public/PrimitiveSceneProxy.h b/Engine/Source/Runtime/Engine/Public/PrimitiveSceneProxy.h
--- a/Engine/Source/Runtime/Engine/Public/PrimitiveSceneProxy.h
+++ b/Engine/Source/Runtime/Engine/Public/PrimitiveSceneProxy.h
@@ -419,6 +419,7 @@ public:
 	inline bool TreatAsBackgroundForOcclusion() const { return bTreatAsBackgroundForOcclusion; }
 	inline bool NeedsLevelAddedToWorldNotification() const { return bNeedsLevelAddedToWorldNotification; }
 	inline bool IsComponentLevelVisible() const { return bIsComponentLevelVisible; }
+	inline bool IsStaticPathAvailable() const { return !bDisableStaticPath; }
 
 #if WITH_EDITOR
 	inline int32 GetNumUncachedStaticLightingInteractions() { return NumUncachedStaticLightingInteractions; }
@@ -545,6 +546,9 @@ private:
 	/** Whether the primitive should be treated as part of the background for occlusion purposes. */
 	uint32 bTreatAsBackgroundForOcclusion : 1;
 
+	friend class FLightPrimitiveInteraction;
+	/** Whether the renderer needs us to temporarily use only the dynamic drawing path */
+	uint32 bDisableStaticPath : 1;
 protected:
 
 	/** Whether the primitive should be statically lit but has unbuilt lighting, and a preview should be used. */
@@ -731,7 +735,6 @@ private:
 	*	How many invalid lights for this primitive, just refer for scene outliner
 	*/
 	int32 NumUncachedStaticLightingInteractions;
-	friend class FLightPrimitiveInteraction;
 
 	/** this is used if world setting has EnableHierarchical LOD true */
 	int32 HierarchicalLODOverride;
diff --git a/Engine/Source/Runtime/Landscape/Private/LandscapeRender.cpp b/Engine/Source/Runtime/Landscape/Private/LandscapeRender.cpp
--- a/Engine/Source/Runtime/Landscape/Private/LandscapeRender.cpp
+++ b/Engine/Source/Runtime/Landscape/Private/LandscapeRender.cpp
@@ -870,14 +870,14 @@ FPrimitiveViewRelevance FLandscapeComponentSceneProxy::GetViewRelevance(const FS
 		View->Family->EngineShowFlags.Wireframe ||
 #if WITH_EDITOR
 		(IsSelected() && !GLandscapeEditModeActive) ||
-		GLandscapeViewMode != ELandscapeViewMode::Normal
+		GLandscapeViewMode != ELandscapeViewMode::Normal ||
 #else
-		IsSelected()
+		IsSelected() ||
 #endif
-		)
+		 !IsStaticPathAvailable())
 	{
 		Result.bDynamicRelevance = true;
-}
+	}
 	else
 	{
 		Result.bStaticRelevance = true;
@@ -1559,11 +1559,11 @@ void FLandscapeComponentSceneProxy::GetDynamicMeshElements(const TArray<const FS
 					GLandscapeDebugOptions.bDisableStatic ||
 					bIsWireframe ||
 #if WITH_EDITOR
-					(IsSelected() && !GLandscapeEditModeActive)
+					(IsSelected() && !GLandscapeEditModeActive) ||
 #else
-					IsSelected()
+					IsSelected() ||
 #endif
-					)
+					!IsStaticPathAvailable())
 				{
 					Mesh.MaterialRenderProxy = MaterialInterface->GetRenderProxy(false);
 
diff --git a/Engine/Source/Runtime/Renderer/Private/ForwardBasePassRendering.cpp b/Engine/Source/Runtime/Renderer/Private/ForwardBasePassRendering.cpp
--- a/Engine/Source/Runtime/Renderer/Private/ForwardBasePassRendering.cpp
+++ b/Engine/Source/Runtime/Renderer/Private/ForwardBasePassRendering.cpp
@@ -8,36 +8,108 @@
 #include "ScenePrivate.h"
 #include "SceneUtils.h"
 
-#define IMPLEMENT_FORWARD_SHADING_BASEPASS_LIGHTMAPPED_SHADER_TYPE(LightMapPolicyType,LightMapPolicyName) \
+#define IMPLEMENT_FORWARD_SHADING_BASEPASS_LIGHTMAPPED_VERTEX_SHADER_TYPE(LightMapPolicyType,LightMapPolicyName) \
 	typedef TBasePassForForwardShadingVS< LightMapPolicyType, LDR_GAMMA_32 > TBasePassForForwardShadingVS##LightMapPolicyName##LDRGamma32; \
 	typedef TBasePassForForwardShadingVS< LightMapPolicyType, HDR_LINEAR_64 > TBasePassForForwardShadingVS##LightMapPolicyName##HDRLinear64; \
-	IMPLEMENT_MATERIAL_SHADER_TYPE(template<>,TBasePassForForwardShadingVS##LightMapPolicyName##LDRGamma32,TEXT("BasePassForForwardShadingVertexShader"),TEXT("Main"),SF_Vertex); \
-	IMPLEMENT_MATERIAL_SHADER_TYPE(template<>,TBasePassForForwardShadingVS##LightMapPolicyName##HDRLinear64,TEXT("BasePassForForwardShadingVertexShader"),TEXT("Main"),SF_Vertex); \
-	typedef TBasePassForForwardShadingPS< LightMapPolicyType, LDR_GAMMA_32, false > TBasePassForForwardShadingPS##LightMapPolicyName##LDRGamma32; \
-	typedef TBasePassForForwardShadingPS< LightMapPolicyType, HDR_LINEAR_32, false > TBasePassForForwardShadingPS##LightMapPolicyName##HDRLinear32; \
-	typedef TBasePassForForwardShadingPS< LightMapPolicyType, HDR_LINEAR_64, false > TBasePassForForwardShadingPS##LightMapPolicyName##HDRLinear64; \
-	typedef TBasePassForForwardShadingPS< LightMapPolicyType, LDR_GAMMA_32, true > TBasePassForForwardShadingPS##LightMapPolicyName##LDRGamma32##Skylight; \
-	typedef TBasePassForForwardShadingPS< LightMapPolicyType, HDR_LINEAR_32, true > TBasePassForForwardShadingPS##LightMapPolicyName##HDRLinear32##Skylight; \
-	typedef TBasePassForForwardShadingPS< LightMapPolicyType, HDR_LINEAR_64, true > TBasePassForForwardShadingPS##LightMapPolicyName##HDRLinear64##Skylight; \
-	IMPLEMENT_MATERIAL_SHADER_TYPE(template<>,TBasePassForForwardShadingPS##LightMapPolicyName##LDRGamma32,TEXT("BasePassForForwardShadingPixelShader"),TEXT("Main"),SF_Pixel); \
-	IMPLEMENT_MATERIAL_SHADER_TYPE(template<>,TBasePassForForwardShadingPS##LightMapPolicyName##HDRLinear32,TEXT("BasePassForForwardShadingPixelShader"),TEXT("Main"),SF_Pixel); \
-	IMPLEMENT_MATERIAL_SHADER_TYPE(template<>,TBasePassForForwardShadingPS##LightMapPolicyName##HDRLinear64,TEXT("BasePassForForwardShadingPixelShader"),TEXT("Main"),SF_Pixel); \
-	IMPLEMENT_MATERIAL_SHADER_TYPE(template<>, TBasePassForForwardShadingPS##LightMapPolicyName##LDRGamma32##Skylight, TEXT("BasePassForForwardShadingPixelShader"), TEXT("Main"), SF_Pixel); \
-	IMPLEMENT_MATERIAL_SHADER_TYPE(template<>, TBasePassForForwardShadingPS##LightMapPolicyName##HDRLinear32##Skylight, TEXT("BasePassForForwardShadingPixelShader"), TEXT("Main"), SF_Pixel); \
-	IMPLEMENT_MATERIAL_SHADER_TYPE(template<>, TBasePassForForwardShadingPS##LightMapPolicyName##HDRLinear64##Skylight, TEXT("BasePassForForwardShadingPixelShader"), TEXT("Main"), SF_Pixel);
-
-// Implement shader types per lightmap policy
-IMPLEMENT_FORWARD_SHADING_BASEPASS_LIGHTMAPPED_SHADER_TYPE( FNoLightMapPolicy, FNoLightMapPolicy ); 
-IMPLEMENT_FORWARD_SHADING_BASEPASS_LIGHTMAPPED_SHADER_TYPE( TLightMapPolicy<LQ_LIGHTMAP>, TLightMapPolicyLQ );
-IMPLEMENT_FORWARD_SHADING_BASEPASS_LIGHTMAPPED_SHADER_TYPE( TDistanceFieldShadowsAndLightMapPolicy<LQ_LIGHTMAP>, TDistanceFieldShadowsAndLightMapPolicyLQ );
-IMPLEMENT_FORWARD_SHADING_BASEPASS_LIGHTMAPPED_SHADER_TYPE( FSimpleDirectionalLightAndSHIndirectPolicy, FSimpleDirectionalLightAndSHIndirectPolicy );
+	IMPLEMENT_MATERIAL_SHADER_TYPE(template<>, TBasePassForForwardShadingVS##LightMapPolicyName##LDRGamma32, TEXT("BasePassForForwardShadingVertexShader"), TEXT("Main"), SF_Vertex); \
+	IMPLEMENT_MATERIAL_SHADER_TYPE(template<>, TBasePassForForwardShadingVS##LightMapPolicyName##HDRLinear64, TEXT("BasePassForForwardShadingVertexShader"), TEXT("Main"), SF_Vertex);
+
+#define IMPLEMENT_FORWARD_SHADING_BASEPASS_LIGHTMAPPED_PIXEL_SHADER_TYPE(LightMapPolicyType,LightMapPolicyName,NumDynamicPointLights) \
+	typedef TBasePassForForwardShadingPS< LightMapPolicyType, LDR_GAMMA_32, false, NumDynamicPointLights > TBasePassForForwardShadingPS##LightMapPolicyName##NumDynamicPointLights##LDRGamma32; \
+	typedef TBasePassForForwardShadingPS< LightMapPolicyType, HDR_LINEAR_32, false, NumDynamicPointLights > TBasePassForForwardShadingPS##LightMapPolicyName##NumDynamicPointLights##HDRLinear32; \
+	typedef TBasePassForForwardShadingPS< LightMapPolicyType, HDR_LINEAR_64, false, NumDynamicPointLights > TBasePassForForwardShadingPS##LightMapPolicyName##NumDynamicPointLights##HDRLinear64; \
+	typedef TBasePassForForwardShadingPS< LightMapPolicyType, LDR_GAMMA_32, true, NumDynamicPointLights > TBasePassForForwardShadingPS##LightMapPolicyName##NumDynamicPointLights##LDRGamma32##Skylight; \
+	typedef TBasePassForForwardShadingPS< LightMapPolicyType, HDR_LINEAR_32, true, NumDynamicPointLights > TBasePassForForwardShadingPS##LightMapPolicyName##NumDynamicPointLights##HDRLinear32##Skylight; \
+	typedef TBasePassForForwardShadingPS< LightMapPolicyType, HDR_LINEAR_64, true, NumDynamicPointLights > TBasePassForForwardShadingPS##LightMapPolicyName##NumDynamicPointLights##HDRLinear64##Skylight; \
+	IMPLEMENT_MATERIAL_SHADER_TYPE(template<>, TBasePassForForwardShadingPS##LightMapPolicyName##NumDynamicPointLights##LDRGamma32, TEXT("BasePassForForwardShadingPixelShader"), TEXT("Main"), SF_Pixel); \
+	IMPLEMENT_MATERIAL_SHADER_TYPE(template<>, TBasePassForForwardShadingPS##LightMapPolicyName##NumDynamicPointLights##HDRLinear32, TEXT("BasePassForForwardShadingPixelShader"), TEXT("Main"), SF_Pixel); \
+	IMPLEMENT_MATERIAL_SHADER_TYPE(template<>, TBasePassForForwardShadingPS##LightMapPolicyName##NumDynamicPointLights##HDRLinear64, TEXT("BasePassForForwardShadingPixelShader"), TEXT("Main"), SF_Pixel); \
+	IMPLEMENT_MATERIAL_SHADER_TYPE(template<>, TBasePassForForwardShadingPS##LightMapPolicyName##NumDynamicPointLights##LDRGamma32##Skylight, TEXT("BasePassForForwardShadingPixelShader"), TEXT("Main"), SF_Pixel); \
+	IMPLEMENT_MATERIAL_SHADER_TYPE(template<>, TBasePassForForwardShadingPS##LightMapPolicyName##NumDynamicPointLights##HDRLinear32##Skylight, TEXT("BasePassForForwardShadingPixelShader"), TEXT("Main"), SF_Pixel); \
+	IMPLEMENT_MATERIAL_SHADER_TYPE(template<>, TBasePassForForwardShadingPS##LightMapPolicyName##NumDynamicPointLights##HDRLinear64##Skylight, TEXT("BasePassForForwardShadingPixelShader"), TEXT("Main"), SF_Pixel);
+
+#define IMPLEMENT_FORWARD_SHADING_BASEPASS_LIGHTMAPPED_SHADER_TYPE(LightMapPolicyType,LightMapPolicyName) \
+	IMPLEMENT_FORWARD_SHADING_BASEPASS_LIGHTMAPPED_VERTEX_SHADER_TYPE(LightMapPolicyType, LightMapPolicyName) \
+	IMPLEMENT_FORWARD_SHADING_BASEPASS_LIGHTMAPPED_PIXEL_SHADER_TYPE(LightMapPolicyType, LightMapPolicyName,0)
+
+static_assert(MAX_BASEPASS_DYNAMIC_POINT_LIGHTS == 4, "If you change MAX_BASEPASS_DYNAMIC_POINT_LIGHTS, you need to make shader types below");
+
+// Implement shader types per lightmap policy and for the number of point lights to support.
+// INT32_MAX indicates the shader should use branching to support a variable number of point lights.
+IMPLEMENT_FORWARD_SHADING_BASEPASS_LIGHTMAPPED_SHADER_TYPE(FNoLightMapPolicy, FNoLightMapPolicy);
+IMPLEMENT_FORWARD_SHADING_BASEPASS_LIGHTMAPPED_VERTEX_SHADER_TYPE(TLightMapPolicy<LQ_LIGHTMAP>, TLightMapPolicyLQ);
+IMPLEMENT_FORWARD_SHADING_BASEPASS_LIGHTMAPPED_PIXEL_SHADER_TYPE(TLightMapPolicy<LQ_LIGHTMAP>, TLightMapPolicyLQ, 0);
+IMPLEMENT_FORWARD_SHADING_BASEPASS_LIGHTMAPPED_PIXEL_SHADER_TYPE(TLightMapPolicy<LQ_LIGHTMAP>, TLightMapPolicyLQ, 1);
+IMPLEMENT_FORWARD_SHADING_BASEPASS_LIGHTMAPPED_PIXEL_SHADER_TYPE(TLightMapPolicy<LQ_LIGHTMAP>, TLightMapPolicyLQ, 2);
+IMPLEMENT_FORWARD_SHADING_BASEPASS_LIGHTMAPPED_PIXEL_SHADER_TYPE(TLightMapPolicy<LQ_LIGHTMAP>, TLightMapPolicyLQ, 3);
+IMPLEMENT_FORWARD_SHADING_BASEPASS_LIGHTMAPPED_PIXEL_SHADER_TYPE(TLightMapPolicy<LQ_LIGHTMAP>, TLightMapPolicyLQ, 4);
+IMPLEMENT_FORWARD_SHADING_BASEPASS_LIGHTMAPPED_PIXEL_SHADER_TYPE(TLightMapPolicy<LQ_LIGHTMAP>, TLightMapPolicyLQ, INT32_MAX);
+IMPLEMENT_FORWARD_SHADING_BASEPASS_LIGHTMAPPED_VERTEX_SHADER_TYPE(TDistanceFieldShadowsAndLightMapPolicy<LQ_LIGHTMAP>, TDistanceFieldShadowsAndLightMapPolicyLQ);
+IMPLEMENT_FORWARD_SHADING_BASEPASS_LIGHTMAPPED_PIXEL_SHADER_TYPE(TDistanceFieldShadowsAndLightMapPolicy<LQ_LIGHTMAP>, TDistanceFieldShadowsAndLightMapPolicyLQ, 0);
+IMPLEMENT_FORWARD_SHADING_BASEPASS_LIGHTMAPPED_PIXEL_SHADER_TYPE(TDistanceFieldShadowsAndLightMapPolicy<LQ_LIGHTMAP>, TDistanceFieldShadowsAndLightMapPolicyLQ, 1);
+IMPLEMENT_FORWARD_SHADING_BASEPASS_LIGHTMAPPED_PIXEL_SHADER_TYPE(TDistanceFieldShadowsAndLightMapPolicy<LQ_LIGHTMAP>, TDistanceFieldShadowsAndLightMapPolicyLQ, 2);
+IMPLEMENT_FORWARD_SHADING_BASEPASS_LIGHTMAPPED_PIXEL_SHADER_TYPE(TDistanceFieldShadowsAndLightMapPolicy<LQ_LIGHTMAP>, TDistanceFieldShadowsAndLightMapPolicyLQ, 3);
+IMPLEMENT_FORWARD_SHADING_BASEPASS_LIGHTMAPPED_PIXEL_SHADER_TYPE(TDistanceFieldShadowsAndLightMapPolicy<LQ_LIGHTMAP>, TDistanceFieldShadowsAndLightMapPolicyLQ, 4);
+IMPLEMENT_FORWARD_SHADING_BASEPASS_LIGHTMAPPED_PIXEL_SHADER_TYPE(TDistanceFieldShadowsAndLightMapPolicy<LQ_LIGHTMAP>, TDistanceFieldShadowsAndLightMapPolicyLQ, INT32_MAX);
+IMPLEMENT_FORWARD_SHADING_BASEPASS_LIGHTMAPPED_VERTEX_SHADER_TYPE(FSimpleDirectionalLightAndSHIndirectPolicy, FSimpleDirectionalLightAndSHIndirectPolicy);
+IMPLEMENT_FORWARD_SHADING_BASEPASS_LIGHTMAPPED_PIXEL_SHADER_TYPE(FSimpleDirectionalLightAndSHIndirectPolicy, FSimpleDirectionalLightAndSHIndirectPolicy, 0);
+IMPLEMENT_FORWARD_SHADING_BASEPASS_LIGHTMAPPED_PIXEL_SHADER_TYPE(FSimpleDirectionalLightAndSHIndirectPolicy, FSimpleDirectionalLightAndSHIndirectPolicy, 1);
+IMPLEMENT_FORWARD_SHADING_BASEPASS_LIGHTMAPPED_PIXEL_SHADER_TYPE(FSimpleDirectionalLightAndSHIndirectPolicy, FSimpleDirectionalLightAndSHIndirectPolicy, 2);
+IMPLEMENT_FORWARD_SHADING_BASEPASS_LIGHTMAPPED_PIXEL_SHADER_TYPE(FSimpleDirectionalLightAndSHIndirectPolicy, FSimpleDirectionalLightAndSHIndirectPolicy, 3);
+IMPLEMENT_FORWARD_SHADING_BASEPASS_LIGHTMAPPED_PIXEL_SHADER_TYPE(FSimpleDirectionalLightAndSHIndirectPolicy, FSimpleDirectionalLightAndSHIndirectPolicy, 4);
+IMPLEMENT_FORWARD_SHADING_BASEPASS_LIGHTMAPPED_PIXEL_SHADER_TYPE(FSimpleDirectionalLightAndSHIndirectPolicy, FSimpleDirectionalLightAndSHIndirectPolicy, INT32_MAX);
 IMPLEMENT_FORWARD_SHADING_BASEPASS_LIGHTMAPPED_SHADER_TYPE(FSimpleDirectionalLightAndSHDirectionalIndirectPolicy, FSimpleDirectionalLightAndSHDirectionalIndirectPolicy);
 IMPLEMENT_FORWARD_SHADING_BASEPASS_LIGHTMAPPED_SHADER_TYPE(FSimpleDirectionalLightAndSHDirectionalCSMIndirectPolicy, FSimpleDirectionalLightAndSHDirectionalCSMIndirectPolicy);
-IMPLEMENT_FORWARD_SHADING_BASEPASS_LIGHTMAPPED_SHADER_TYPE( FMovableDirectionalLightLightingPolicy, FMovableDirectionalLightLightingPolicy );
-IMPLEMENT_FORWARD_SHADING_BASEPASS_LIGHTMAPPED_SHADER_TYPE( FMovableDirectionalLightCSMLightingPolicy, FMovableDirectionalLightCSMLightingPolicy );
+IMPLEMENT_FORWARD_SHADING_BASEPASS_LIGHTMAPPED_SHADER_TYPE(FMovableDirectionalLightLightingPolicy, FMovableDirectionalLightLightingPolicy);
+IMPLEMENT_FORWARD_SHADING_BASEPASS_LIGHTMAPPED_SHADER_TYPE(FMovableDirectionalLightCSMLightingPolicy, FMovableDirectionalLightCSMLightingPolicy);
 IMPLEMENT_FORWARD_SHADING_BASEPASS_LIGHTMAPPED_SHADER_TYPE(FMovableDirectionalLightWithLightmapLightingPolicy, FMovableDirectionalLightWithLightmapLightingPolicy);
 IMPLEMENT_FORWARD_SHADING_BASEPASS_LIGHTMAPPED_SHADER_TYPE(FMovableDirectionalLightCSMWithLightmapLightingPolicy, FMovableDirectionalLightCSMWithLightmapLightingPolicy);
 
+FBasePassFowardDynamicPointLightInfo::FBasePassFowardDynamicPointLightInfo(const FPrimitiveSceneProxy* InSceneProxy)
+: NumDynamicPointLights(0)
+{
+	if (InSceneProxy != nullptr)
+	{
+		for (FLightPrimitiveInteraction* LPI = InSceneProxy->GetPrimitiveSceneInfo()->LightList; LPI; LPI = LPI->GetNextLight())
+		{
+			FLightSceneProxy* LightProxy = LPI->GetLight()->Proxy;
+			if (LightProxy->GetLightType() == LightType_Point)
+			{
+				FVector NormalizedLightDirection;
+				FVector2D SpotAngles;
+				float SourceRadius;
+				float SourceLength;
+				float MinRoughness;
+
+				// Get the light parameters
+				LightProxy->GetParameters(
+					LightPositionAndInvRadius[NumDynamicPointLights],
+					LightColorAndFalloffExponent[NumDynamicPointLights],
+					NormalizedLightDirection,
+					SpotAngles,
+					SourceRadius,
+					SourceLength,
+					MinRoughness);
+
+				if (LightProxy->IsInverseSquared())
+				{
+					// Correction for lumen units
+					LightColorAndFalloffExponent[NumDynamicPointLights].X *= 16.0f;
+					LightColorAndFalloffExponent[NumDynamicPointLights].Y *= 16.0f;
+					LightColorAndFalloffExponent[NumDynamicPointLights].Z *= 16.0f;
+					LightColorAndFalloffExponent[NumDynamicPointLights].W = 0;
+				}
+
+				NumDynamicPointLights++;
+				if (NumDynamicPointLights == MAX_BASEPASS_DYNAMIC_POINT_LIGHTS)
+				{
+					break;
+				}
+			}
+		}
+	}
+}
+
 /** The action used to draw a base pass static mesh element. */
 class FDrawBasePassForwardShadingStaticMeshAction
 {
@@ -63,7 +135,7 @@ public:
 	}
 
 	/** Draws the translucent mesh with a specific light-map type, and fog volume type */
-	template<typename LightMapPolicyType>
+	template<typename LightMapPolicyType, int32 NumDynamicPointLights>
 	void Process(
 		FRHICommandList& RHICmdList, 
 		const FProcessBasePassMeshParameters& Parameters,
@@ -81,15 +153,15 @@ public:
 		if ( Scene )
 		{
 			// Find the appropriate draw list for the static mesh based on the light-map policy type.
-			TStaticMeshDrawList<TBasePassForForwardShadingDrawingPolicy<LightMapPolicyType> >& DrawList =
+			TStaticMeshDrawList<TBasePassForForwardShadingDrawingPolicy<LightMapPolicyType, 0> >& DrawList =
 				Scene->GetForwardShadingBasePassDrawList<LightMapPolicyType>(DrawType);
 
 			ERHIFeatureLevel::Type FeatureLevel = Scene->GetFeatureLevel();
 			// Add the static mesh to the draw list.
 			DrawList.AddMesh(
 				StaticMesh,
-				typename TBasePassForForwardShadingDrawingPolicy<LightMapPolicyType>::ElementDataType(LightMapElementData),
-				TBasePassForForwardShadingDrawingPolicy<LightMapPolicyType>(
+				typename TBasePassForForwardShadingDrawingPolicy<LightMapPolicyType, 0>::ElementDataType(LightMapElementData),
+				TBasePassForForwardShadingDrawingPolicy<LightMapPolicyType, 0>(
 				StaticMesh->VertexFactory,
 				StaticMesh->MaterialRenderProxy,
 				*Parameters.Material,
@@ -123,7 +195,7 @@ void FBasePassForwardOpaqueDrawingPolicyFactory::AddStaticMesh(FRHICommandList&
 		// following check moved from ProcessBasePassMeshForForwardShading to avoid passing feature level.
 		check(!AllowHighQualityLightmaps(Scene->GetFeatureLevel()));
 
-		ProcessBasePassMeshForForwardShading(
+		ProcessBasePassMeshForForwardShading<FDrawBasePassForwardShadingStaticMeshAction, 0>(
 			RHICmdList,
 			FProcessBasePassMeshParameters(
 				*StaticMesh,
@@ -174,7 +246,7 @@ public:
 	{}
 
 	/** Draws the translucent mesh with a specific light-map type, and shader complexity predicate. */
-	template<typename LightMapPolicyType>
+	template<typename LightMapPolicyType, int32 NumDynamicPointLights>
 	void Process(
 		FRHICommandList& RHICmdList, 
 		const FProcessBasePassMeshParameters& Parameters,
@@ -193,7 +265,7 @@ public:
 		const bool bIsLitMaterial = Parameters.ShadingModel != MSM_Unlit;
 		const FScene* Scene = Parameters.PrimitiveSceneProxy ? Parameters.PrimitiveSceneProxy->GetPrimitiveSceneInfo()->Scene : NULL;
 
-		TBasePassForForwardShadingDrawingPolicy<LightMapPolicyType> DrawingPolicy(
+		TBasePassForForwardShadingDrawingPolicy<LightMapPolicyType, NumDynamicPointLights> DrawingPolicy(
 			Parameters.Mesh.VertexFactory,
 			Parameters.Mesh.MaterialRenderProxy,
 			*Parameters.Material,
@@ -206,7 +278,7 @@ public:
 			Parameters.bEditorCompositeDepthTest
 			);
 		RHICmdList.BuildAndSetLocalBoundShaderState(DrawingPolicy.GetBoundShaderStateInput(View.GetFeatureLevel()));
-		DrawingPolicy.SetSharedState(RHICmdList, &View, typename TBasePassForForwardShadingDrawingPolicy<LightMapPolicyType>::ContextDataType());
+		DrawingPolicy.SetSharedState(RHICmdList, &View, typename TBasePassForForwardShadingDrawingPolicy<LightMapPolicyType, NumDynamicPointLights>::ContextDataType());
 
 		for( int32 BatchElementIndex=0;BatchElementIndex<Parameters.Mesh.Elements.Num();BatchElementIndex++ )
 		{
@@ -218,8 +290,8 @@ public:
 				BatchElementIndex,
 				bBackFace,
 				DitheredLODTransitionValue,
-				typename TBasePassForForwardShadingDrawingPolicy<LightMapPolicyType>::ElementDataType(LightMapElementData),
-				typename TBasePassForForwardShadingDrawingPolicy<LightMapPolicyType>::ContextDataType()
+				typename TBasePassForForwardShadingDrawingPolicy<LightMapPolicyType, NumDynamicPointLights>::ElementDataType(LightMapElementData),
+				typename TBasePassForForwardShadingDrawingPolicy<LightMapPolicyType, NumDynamicPointLights>::ContextDataType()
 				);
 			DrawingPolicy.DrawMesh(RHICmdList, Parameters.Mesh, BatchElementIndex);
 		}
@@ -234,6 +306,38 @@ public:
 	}
 };
 
+template<int32 NumDynamicPointLights>
+void FBasePassForwardOpaqueDrawingPolicyFactory::DrawDynamicMeshTempl(
+	FRHICommandList& RHICmdList,
+	const FViewInfo& View,
+	ContextType DrawingContext,
+	const FMeshBatch& Mesh,
+	const FMaterial* Material,
+	bool bBackFace,
+	const FPrimitiveSceneProxy* PrimitiveSceneProxy,
+	FHitProxyId HitProxyId
+	)
+{
+	ProcessBasePassMeshForForwardShading<FDrawBasePassForwardShadingDynamicMeshAction, NumDynamicPointLights>(
+		RHICmdList, 
+		FProcessBasePassMeshParameters(
+			Mesh, 
+			Material, 
+			PrimitiveSceneProxy, 
+			true, 
+			DrawingContext.bEditorCompositeDepthTest, 
+			DrawingContext.TextureMode, 
+			View.GetFeatureLevel()
+		), 
+		FDrawBasePassForwardShadingDynamicMeshAction(
+			View, 
+			bBackFace, 
+			Mesh.DitheredLODTransitionAlpha, 
+			HitProxyId														
+		)																
+	);
+}
+
 bool FBasePassForwardOpaqueDrawingPolicyFactory::DrawDynamicMesh(
 	FRHICommandList& RHICmdList, 
 	const FViewInfo& View,
@@ -254,24 +358,42 @@ bool FBasePassForwardOpaqueDrawingPolicyFactory::DrawDynamicMesh(
 	// Only draw opaque materials.
 	if(!IsTranslucentBlendMode(BlendMode))
 	{
-		ProcessBasePassMeshForForwardShading(
-			RHICmdList,
-			FProcessBasePassMeshParameters(
-				Mesh,
-				Material,
-				PrimitiveSceneProxy,
-				true,
-				DrawingContext.bEditorCompositeDepthTest,
-				DrawingContext.TextureMode,
-				FeatureLevel
-				),
-			FDrawBasePassForwardShadingDynamicMeshAction(
-				View,
-				bBackFace,
-				Mesh.DitheredLODTransitionAlpha,
-				HitProxyId
-				)
-			);
+		int32 NumDynamicPointLights = PrimitiveSceneProxy ? PrimitiveSceneProxy->GetPrimitiveSceneInfo()->NumES2DynamicPointLights : 0;
+		const bool bIsUnlit = Material->GetShadingModel() == MSM_Unlit;
+
+		if (NumDynamicPointLights == 0 || bIsUnlit)
+		{
+			DrawDynamicMeshTempl<0>(RHICmdList, View, DrawingContext, Mesh, Material, bBackFace, PrimitiveSceneProxy, HitProxyId);
+		}
+		else
+		{
+			static auto* MobileDynamicPointLightsUseStaticBranchCVar = IConsoleManager::Get().FindTConsoleVariableDataInt(TEXT("r.MobileDynamicPointLightsUseStaticBranch"));
+			if (MobileDynamicPointLightsUseStaticBranchCVar->GetValueOnRenderThread() == 1)
+			{
+				DrawDynamicMeshTempl<INT32_MAX>(RHICmdList, View, DrawingContext, Mesh, Material, bBackFace, PrimitiveSceneProxy, HitProxyId);
+			}
+			else
+			{
+				static_assert(MAX_BASEPASS_DYNAMIC_POINT_LIGHTS == 4, "If you change MAX_BASEPASS_DYNAMIC_POINT_LIGHTS, you need to change the switch statement below");
+
+				switch (NumDynamicPointLights)
+				{
+				case 1:
+					DrawDynamicMeshTempl<1>(RHICmdList, View, DrawingContext, Mesh, Material, bBackFace, PrimitiveSceneProxy, HitProxyId);
+					break;
+				case 2:
+					DrawDynamicMeshTempl<2>(RHICmdList, View, DrawingContext, Mesh, Material, bBackFace, PrimitiveSceneProxy, HitProxyId);
+					break;
+				case 3:
+					DrawDynamicMeshTempl<3>(RHICmdList, View, DrawingContext, Mesh, Material, bBackFace, PrimitiveSceneProxy, HitProxyId);
+					break;
+				default:
+					DrawDynamicMeshTempl<4>(RHICmdList, View, DrawingContext, Mesh, Material, bBackFace, PrimitiveSceneProxy, HitProxyId);
+					break;
+				}
+			}
+		}
+
 		return true;
 	}
 	else
diff --git a/Engine/Source/Runtime/Renderer/Private/ForwardBasePassRendering.h b/Engine/Source/Runtime/Renderer/Private/ForwardBasePassRendering.h
--- a/Engine/Source/Runtime/Renderer/Private/ForwardBasePassRendering.h
+++ b/Engine/Source/Runtime/Renderer/Private/ForwardBasePassRendering.h
@@ -10,6 +10,8 @@
 #include "ShaderBaseClasses.h"
 #include "EditorCompositeParams.h"
 
+class FBasePassFowardDynamicPointLightInfo;
+
 enum EOutputFormat
 {
 	LDR_GAMMA_32,
@@ -17,6 +19,19 @@ enum EOutputFormat
 	HDR_LINEAR_64,
 };
 
+#define MAX_BASEPASS_DYNAMIC_POINT_LIGHTS 4
+
+/* Info for dynamic point lights rendered in base pass */
+class FBasePassFowardDynamicPointLightInfo
+{
+public:
+	FBasePassFowardDynamicPointLightInfo(const FPrimitiveSceneProxy* InSceneProxy);
+
+	int32 NumDynamicPointLights;
+	FVector4 LightPositionAndInvRadius[MAX_BASEPASS_DYNAMIC_POINT_LIGHTS];
+	FVector4 LightColorAndFalloffExponent[MAX_BASEPASS_DYNAMIC_POINT_LIGHTS];
+};
+
 static bool ShouldCacheShaderByOutputFormat(EOutputFormat OutputFormat)
 {
 	bool bSupportsMobileHDR = IsMobileHDR();
@@ -119,7 +134,7 @@ public:
 /**
  * The base type for pixel shaders that render the emissive color, and light-mapped/ambient lighting of a mesh.
  */
-template<typename LightMapPolicyType>
+template<typename LightMapPolicyType, int32 NumDynamicPointLights>
 class TBasePassForForwardShadingPSBaseType : public FMeshMaterialShader, public LightMapPolicyType::PixelParametersType
 {
 	DECLARE_SHADER_TYPE(TBasePassForForwardShadingPSBaseType,MeshMaterial);
@@ -128,13 +143,34 @@ public:
 
 	static bool ShouldCache(EShaderPlatform Platform,const FMaterial* Material,const FVertexFactoryType* VertexFactoryType)
 	{
-		return IsMobilePlatform(Platform) && LightMapPolicyType::ShouldCache(Platform,Material,VertexFactoryType);
+		// We compile the point light shader combinations based on the project settings
+		static auto* MobileDynamicPointLightsUseStaticBranchCVar = IConsoleManager::Get().FindTConsoleVariableDataInt(TEXT("r.MobileDynamicPointLightsUseStaticBranch"));
+		static auto* MobileNumDynamicPointLightsCVar = IConsoleManager::Get().FindTConsoleVariableDataInt(TEXT("r.MobileNumDynamicPointLights"));
+		const bool bMobileDynamicPointLightsUseStaticBranch = (MobileDynamicPointLightsUseStaticBranchCVar->GetValueOnAnyThread() == 1);
+		const int32 MobileNumDynamicPointLights = MobileNumDynamicPointLightsCVar->GetValueOnAnyThread();
+		const bool bIsUnlit = Material->GetShadingModel() == MSM_Unlit;
+
+		return IsMobilePlatform(Platform) && LightMapPolicyType::ShouldCache(Platform, Material, VertexFactoryType) &&
+			(NumDynamicPointLights == 0 ||
+			(!bIsUnlit && NumDynamicPointLights == INT32_MAX && bMobileDynamicPointLightsUseStaticBranch) ||					// single shader for variable number of point lights
+			(!bIsUnlit && NumDynamicPointLights <= MobileNumDynamicPointLights && !bMobileDynamicPointLightsUseStaticBranch));	// unique 1...N point light shaders
 	}
 
 	static void ModifyCompilationEnvironment(EShaderPlatform Platform, const FMaterial* Material, FShaderCompilerEnvironment& OutEnvironment)
 	{
 		FMeshMaterialShader::ModifyCompilationEnvironment(Platform, Material, OutEnvironment);
 		LightMapPolicyType::ModifyCompilationEnvironment(Platform, Material, OutEnvironment);
+		if (NumDynamicPointLights == INT32_MAX)
+		{
+			OutEnvironment.SetDefine(TEXT("MAX_DYNAMIC_POINT_LIGHTS"), (uint32)MAX_BASEPASS_DYNAMIC_POINT_LIGHTS);
+			OutEnvironment.SetDefine(TEXT("VARIABLE_NUM_DYNAMIC_POINT_LIGHTS"), (uint32)1);
+		}
+		else
+		{
+			OutEnvironment.SetDefine(TEXT("MAX_DYNAMIC_POINT_LIGHTS"), (uint32)NumDynamicPointLights);
+			OutEnvironment.SetDefine(TEXT("VARIABLE_NUM_DYNAMIC_POINT_LIGHTS"), (uint32)0);
+			OutEnvironment.SetDefine(TEXT("NUM_DYNAMIC_POINT_LIGHTS"), (uint32)NumDynamicPointLights);
+		}
 	}
 
 	/** Initialization constructor. */
@@ -145,6 +181,12 @@ public:
 		ReflectionCubemap.Bind(Initializer.ParameterMap, TEXT("ReflectionCubemap"));
 		ReflectionSampler.Bind(Initializer.ParameterMap, TEXT("ReflectionCubemapSampler"));
 		EditorCompositeParams.Bind(Initializer.ParameterMap);
+		LightPositionAndInvRadiusParameter.Bind(Initializer.ParameterMap, TEXT("LightPositionAndInvRadius"));
+		LightColorAndFalloffExponentParameter.Bind(Initializer.ParameterMap, TEXT("LightColorAndFalloffExponent"));
+		if (NumDynamicPointLights == INT32_MAX)
+		{
+			NumDynamicPointLightsParameter.Bind(Initializer.ParameterMap, TEXT("NumDynamicPointLights"));
+		}
 	}
 	TBasePassForForwardShadingPSBaseType() {}
 
@@ -174,6 +216,20 @@ public:
 			SetTextureParameter(RHICmdList, PixelShader, ReflectionCubemap, ReflectionSampler, ReflectionTexture);
 		}
 
+		if (NumDynamicPointLights > 0)
+		{
+			FBasePassFowardDynamicPointLightInfo LightInfo(Proxy);
+
+			if (NumDynamicPointLights == INT32_MAX)
+			{
+				SetShaderValue(RHICmdList, PixelShader, NumDynamicPointLightsParameter, LightInfo.NumDynamicPointLights);
+			}
+
+			// Set dynamic point lights
+			SetShaderValueArray(RHICmdList, PixelShader, LightPositionAndInvRadiusParameter, LightInfo.LightPositionAndInvRadius, LightInfo.NumDynamicPointLights);
+			SetShaderValueArray(RHICmdList, PixelShader, LightColorAndFalloffExponentParameter, LightInfo.LightColorAndFalloffExponent, LightInfo.NumDynamicPointLights);
+		}
+
 		FMeshMaterialShader::SetMesh(RHICmdList, PixelShader,VertexFactory,View,Proxy,BatchElement,DitheredLODTransitionValue);		
 	}
 
@@ -184,6 +240,12 @@ public:
 		Ar << ReflectionCubemap;
 		Ar << ReflectionSampler;
 		Ar << EditorCompositeParams;
+		Ar << LightPositionAndInvRadiusParameter;
+		Ar << LightColorAndFalloffExponentParameter;
+		if (NumDynamicPointLights == INT32_MAX)
+		{
+			Ar << NumDynamicPointLightsParameter;
+		}
 		return bShaderHasOutdatedParameters;
 	}
 
@@ -191,11 +253,14 @@ private:
 	FShaderResourceParameter ReflectionCubemap;
 	FShaderResourceParameter ReflectionSampler;
 	FEditorCompositingParameters EditorCompositeParams;
+	FShaderParameter LightPositionAndInvRadiusParameter;
+	FShaderParameter LightColorAndFalloffExponentParameter;
+	FShaderParameter NumDynamicPointLightsParameter;
 };
 
 
-template< typename LightMapPolicyType, EOutputFormat OutputFormat, bool bEnableSkyLight>
-class TBasePassForForwardShadingPS : public TBasePassForForwardShadingPSBaseType<LightMapPolicyType>
+template< typename LightMapPolicyType, EOutputFormat OutputFormat, bool bEnableSkyLight, int32 NumDynamicPointLights>
+class TBasePassForForwardShadingPS : public TBasePassForForwardShadingPSBaseType<LightMapPolicyType, NumDynamicPointLights>
 {
 	DECLARE_SHADER_TYPE(TBasePassForForwardShadingPS,MeshMaterial);
 public:
@@ -205,20 +270,20 @@ public:
 		// Only compile skylight version for lit materials on ES2 (Metal) or higher
 		const bool bShouldCacheBySkylight = !bEnableSkyLight || (Material->GetShadingModel() != MSM_Unlit);
 
-		return TBasePassForForwardShadingPSBaseType<LightMapPolicyType>::ShouldCache(Platform, Material, VertexFactoryType) && ShouldCacheShaderByOutputFormat(OutputFormat) && bShouldCacheBySkylight;
+		return TBasePassForForwardShadingPSBaseType<LightMapPolicyType, NumDynamicPointLights>::ShouldCache(Platform, Material, VertexFactoryType) && ShouldCacheShaderByOutputFormat(OutputFormat) && bShouldCacheBySkylight;
 	}
 	
 	static void ModifyCompilationEnvironment(EShaderPlatform Platform, const FMaterial* Material, FShaderCompilerEnvironment& OutEnvironment)
 	{		
-		TBasePassForForwardShadingPSBaseType<LightMapPolicyType>::ModifyCompilationEnvironment(Platform, Material, OutEnvironment);
+		TBasePassForForwardShadingPSBaseType<LightMapPolicyType, NumDynamicPointLights>::ModifyCompilationEnvironment(Platform, Material, OutEnvironment);
 		OutEnvironment.SetDefine(TEXT("ENABLE_SKY_LIGHT"), (uint32)(bEnableSkyLight ? 1 : 0));
-		OutEnvironment.SetDefine( TEXT("USE_HDR_MOSAIC"), OutputFormat == HDR_LINEAR_32 ? 1u : 0u );
-		OutEnvironment.SetDefine( TEXT("OUTPUT_GAMMA_SPACE"), OutputFormat == LDR_GAMMA_32 ? 1u : 0u );
+		OutEnvironment.SetDefine(TEXT("USE_HDR_MOSAIC"), OutputFormat == HDR_LINEAR_32 ? 1u : 0u );
+		OutEnvironment.SetDefine(TEXT("OUTPUT_GAMMA_SPACE"), OutputFormat == LDR_GAMMA_32 ? 1u : 0u );
 	}
 	
 	/** Initialization constructor. */
 	TBasePassForForwardShadingPS(const ShaderMetaType::CompiledShaderInitializerType& Initializer):
-		TBasePassForForwardShadingPSBaseType<LightMapPolicyType>(Initializer)
+		TBasePassForForwardShadingPSBaseType<LightMapPolicyType, NumDynamicPointLights>(Initializer)
 	{}
 
 	/** Default constructor. */
@@ -229,7 +294,7 @@ public:
 /**
  * Draws the emissive color and the light-map of a mesh.
  */
-template<typename LightMapPolicyType>
+template<typename LightMapPolicyType, int32 NumDynamicPointLights>
 class TBasePassForForwardShadingDrawingPolicy : public FMeshDrawingPolicy
 {
 public:
@@ -277,11 +342,11 @@ public:
 
 			if (bInEnableSkyLight)
 			{
-				PixelShader = InMaterialResource.GetShader< TBasePassForForwardShadingPS<LightMapPolicyType, HDR_LINEAR_32, true> >(InVertexFactory->GetType());
+				PixelShader = InMaterialResource.GetShader< TBasePassForForwardShadingPS<LightMapPolicyType, HDR_LINEAR_32, true, NumDynamicPointLights> >(InVertexFactory->GetType());
 			}
 			else
 			{
-				PixelShader = InMaterialResource.GetShader< TBasePassForForwardShadingPS<LightMapPolicyType, HDR_LINEAR_32, false> >(InVertexFactory->GetType());
+				PixelShader = InMaterialResource.GetShader< TBasePassForForwardShadingPS<LightMapPolicyType, HDR_LINEAR_32, false, NumDynamicPointLights> >(InVertexFactory->GetType());
 			}
 		}
 		else if (IsMobileHDR())
@@ -290,11 +355,11 @@ public:
 
 			if (bInEnableSkyLight)
 			{
-				PixelShader = InMaterialResource.GetShader< TBasePassForForwardShadingPS<LightMapPolicyType, HDR_LINEAR_64, true> >(InVertexFactory->GetType());
+				PixelShader = InMaterialResource.GetShader< TBasePassForForwardShadingPS<LightMapPolicyType, HDR_LINEAR_64, true, NumDynamicPointLights> >(InVertexFactory->GetType());
 			}
 			else
 			{
-				PixelShader = InMaterialResource.GetShader< TBasePassForForwardShadingPS<LightMapPolicyType, HDR_LINEAR_64, false> >(InVertexFactory->GetType());
+				PixelShader = InMaterialResource.GetShader< TBasePassForForwardShadingPS<LightMapPolicyType, HDR_LINEAR_64, false, NumDynamicPointLights> >(InVertexFactory->GetType());
 			}
 			
 		}
@@ -304,11 +369,11 @@ public:
 
 			if (bInEnableSkyLight)
 			{
-				PixelShader = InMaterialResource.GetShader< TBasePassForForwardShadingPS<LightMapPolicyType, LDR_GAMMA_32, true> >(InVertexFactory->GetType());
+				PixelShader = InMaterialResource.GetShader< TBasePassForForwardShadingPS<LightMapPolicyType, LDR_GAMMA_32, true, NumDynamicPointLights> >(InVertexFactory->GetType());
 			}
 			else
 			{
-				PixelShader = InMaterialResource.GetShader< TBasePassForForwardShadingPS<LightMapPolicyType, LDR_GAMMA_32, false> >(InVertexFactory->GetType());
+				PixelShader = InMaterialResource.GetShader< TBasePassForForwardShadingPS<LightMapPolicyType, LDR_GAMMA_32, false, NumDynamicPointLights> >(InVertexFactory->GetType());
 			}			
 		}
 
@@ -477,7 +542,7 @@ public:
 
 protected:
 	TBasePassForForwardShadingVSBaseType<LightMapPolicyType>* VertexShader;
-	TBasePassForForwardShadingPSBaseType<LightMapPolicyType>* PixelShader;
+	TBasePassForForwardShadingPSBaseType<LightMapPolicyType, NumDynamicPointLights>* PixelShader;
 
 	LightMapPolicyType LightMapPolicy;
 	EBlendMode BlendMode;
@@ -518,10 +583,24 @@ public:
 		const FPrimitiveSceneProxy* PrimitiveSceneProxy,
 		FHitProxyId HitProxyId
 		);
+
+private:
+	// templated version of DrawDynamicMesh on number of point lights
+	template<int32 NumDynamicPointLights>
+	static void DrawDynamicMeshTempl(
+		FRHICommandList& RHICmdList,
+		const FViewInfo& View,
+		ContextType DrawingContext,
+		const FMeshBatch& Mesh,
+		const FMaterial* Material,
+		bool bBackFace,
+		const FPrimitiveSceneProxy* PrimitiveSceneProxy,
+		FHitProxyId HitProxyId
+		);
 };
 
 /** Processes a base pass mesh using an unknown light map policy, and unknown fog density policy. */
-template<typename ProcessActionType>
+template<typename ProcessActionType, int32 NumDynamicPointLights>
 void ProcessBasePassMeshForForwardShading(
 	FRHICommandList& RHICmdList,
 	const FProcessBasePassMeshParameters& Parameters,
@@ -547,11 +626,11 @@ void ProcessBasePassMeshForForwardShading(
 				// final determination of whether CSMs are rendered can be view dependent, thus we always need to clear the CSMs even if we're not going to render to them based on the condition below.
 				if (SimpleDirectionalLight && SimpleDirectionalLight->ShouldRenderViewIndependentWholeSceneShadows())
 				{
-					Action.template Process<FMovableDirectionalLightCSMWithLightmapLightingPolicy>(RHICmdList, Parameters, FMovableDirectionalLightCSMWithLightmapLightingPolicy(), LightMapInteraction);
+					Action.template Process<FMovableDirectionalLightCSMWithLightmapLightingPolicy, 0>(RHICmdList, Parameters, FMovableDirectionalLightCSMWithLightmapLightingPolicy(), LightMapInteraction);
 				}
 				else
 				{
-					Action.template Process<FMovableDirectionalLightWithLightmapLightingPolicy>(RHICmdList, Parameters, FMovableDirectionalLightCSMWithLightmapLightingPolicy(), LightMapInteraction);
+					Action.template Process<FMovableDirectionalLightWithLightmapLightingPolicy, 0>(RHICmdList, Parameters, FMovableDirectionalLightCSMWithLightmapLightingPolicy(), LightMapInteraction);
 				}
 			}
 			else
@@ -562,7 +641,7 @@ void ProcessBasePassMeshForForwardShading(
 
 				if (ShadowMapInteraction.GetType() == SMIT_Texture)
 				{
-					Action.template Process< TDistanceFieldShadowsAndLightMapPolicy<LQ_LIGHTMAP> >(
+					Action.template Process< TDistanceFieldShadowsAndLightMapPolicy<LQ_LIGHTMAP>, NumDynamicPointLights >(
 						RHICmdList,
 						Parameters,
 						TDistanceFieldShadowsAndLightMapPolicy<LQ_LIGHTMAP>(),
@@ -570,7 +649,7 @@ void ProcessBasePassMeshForForwardShading(
 				}
 				else
 				{
-					Action.template Process< TLightMapPolicy<LQ_LIGHTMAP> >(RHICmdList, Parameters, TLightMapPolicy<LQ_LIGHTMAP>(), LightMapInteraction);
+					Action.template Process< TLightMapPolicy<LQ_LIGHTMAP>, NumDynamicPointLights >(RHICmdList, Parameters, TLightMapPolicy<LQ_LIGHTMAP>(), LightMapInteraction);
 				}
 			}
 
@@ -586,16 +665,16 @@ void ProcessBasePassMeshForForwardShading(
 			{
 				if (SimpleDirectionalLight && SimpleDirectionalLight->ShouldRenderViewIndependentWholeSceneShadows())
 				{
-					Action.template Process<FSimpleDirectionalLightAndSHDirectionalCSMIndirectPolicy>(RHICmdList, Parameters, FSimpleDirectionalLightAndSHDirectionalCSMIndirectPolicy(), FSimpleDirectionalLightAndSHDirectionalCSMIndirectPolicy::ElementDataType(Action.ShouldPackAmbientSH()));
+					Action.template Process<FSimpleDirectionalLightAndSHDirectionalCSMIndirectPolicy, 0>(RHICmdList, Parameters, FSimpleDirectionalLightAndSHDirectionalCSMIndirectPolicy(), FSimpleDirectionalLightAndSHDirectionalCSMIndirectPolicy::ElementDataType(Action.ShouldPackAmbientSH()));
 				}
 				else
 				{
-					Action.template Process<FSimpleDirectionalLightAndSHDirectionalIndirectPolicy>(RHICmdList, Parameters, FSimpleDirectionalLightAndSHDirectionalIndirectPolicy(), FSimpleDirectionalLightAndSHDirectionalIndirectPolicy::ElementDataType(Action.ShouldPackAmbientSH()));
+					Action.template Process<FSimpleDirectionalLightAndSHDirectionalIndirectPolicy, 0>(RHICmdList, Parameters, FSimpleDirectionalLightAndSHDirectionalIndirectPolicy(), FSimpleDirectionalLightAndSHDirectionalIndirectPolicy::ElementDataType(Action.ShouldPackAmbientSH()));
 				}
 			}
 			else
 			{
-				Action.template Process<FSimpleDirectionalLightAndSHIndirectPolicy>(RHICmdList, Parameters, FSimpleDirectionalLightAndSHIndirectPolicy(), FSimpleDirectionalLightAndSHIndirectPolicy::ElementDataType(Action.ShouldPackAmbientSH()));
+				Action.template Process<FSimpleDirectionalLightAndSHIndirectPolicy, NumDynamicPointLights>(RHICmdList, Parameters, FSimpleDirectionalLightAndSHIndirectPolicy(), FSimpleDirectionalLightAndSHIndirectPolicy::ElementDataType(Action.ShouldPackAmbientSH()));
 			}
 
 			// Exit to avoid NoLightmapPolicy
@@ -606,11 +685,11 @@ void ProcessBasePassMeshForForwardShading(
 			// final determination of whether CSMs are rendered can be view dependent, thus we always need to clear the CSMs even if we're not going to render to them based on the condition below.
 			if (SimpleDirectionalLight && SimpleDirectionalLight->ShouldRenderViewIndependentWholeSceneShadows())
 			{
-				Action.template Process<FMovableDirectionalLightCSMLightingPolicy>(RHICmdList, Parameters, FMovableDirectionalLightCSMLightingPolicy(), FMovableDirectionalLightCSMLightingPolicy::ElementDataType());
+				Action.template Process<FMovableDirectionalLightCSMLightingPolicy, 0>(RHICmdList, Parameters, FMovableDirectionalLightCSMLightingPolicy(), FMovableDirectionalLightCSMLightingPolicy::ElementDataType());
 			}
 			else
 			{
-				Action.template Process<FMovableDirectionalLightLightingPolicy>(RHICmdList, Parameters, FMovableDirectionalLightLightingPolicy(), FMovableDirectionalLightLightingPolicy::ElementDataType());
+				Action.template Process<FMovableDirectionalLightLightingPolicy, 0>(RHICmdList, Parameters, FMovableDirectionalLightLightingPolicy(), FMovableDirectionalLightLightingPolicy::ElementDataType());
 			}
 
 			// Exit to avoid NoLightmapPolicy
@@ -619,5 +698,5 @@ void ProcessBasePassMeshForForwardShading(
 	}
 
 	// Default to NoLightmapPolicy
-	Action.template Process<FNoLightMapPolicy>(RHICmdList, Parameters, FNoLightMapPolicy(), FNoLightMapPolicy::ElementDataType());
+	Action.template Process<FNoLightMapPolicy, 0>(RHICmdList, Parameters, FNoLightMapPolicy(), FNoLightMapPolicy::ElementDataType());
 }
diff --git a/Engine/Source/Runtime/Renderer/Private/ForwardTranslucentRendering.cpp b/Engine/Source/Runtime/Renderer/Private/ForwardTranslucentRendering.cpp
--- a/Engine/Source/Runtime/Renderer/Private/ForwardTranslucentRendering.cpp
+++ b/Engine/Source/Runtime/Renderer/Private/ForwardTranslucentRendering.cpp
@@ -124,7 +124,7 @@ public:
 	}
 
 	/** Draws the translucent mesh with a specific light-map type, and fog volume type */
-	template<typename LightMapPolicyType>
+	template<typename LightMapPolicyType, int32 NumDynamicPointLights>
 	void Process(
 		FRHICommandList& RHICmdList, 
 		const FProcessBasePassMeshParameters& Parameters,
@@ -135,7 +135,7 @@ public:
 		const bool bIsLitMaterial = Parameters.ShadingModel != MSM_Unlit;
 		const FScene* Scene = Parameters.PrimitiveSceneProxy ? Parameters.PrimitiveSceneProxy->GetPrimitiveSceneInfo()->Scene : NULL;
 
-		TBasePassForForwardShadingDrawingPolicy<LightMapPolicyType> DrawingPolicy(
+		TBasePassForForwardShadingDrawingPolicy<LightMapPolicyType, NumDynamicPointLights> DrawingPolicy(
 			Parameters.Mesh.VertexFactory,
 			Parameters.Mesh.MaterialRenderProxy,
 			*Parameters.Material,
@@ -148,7 +148,7 @@ public:
 			);
 
 		RHICmdList.BuildAndSetLocalBoundShaderState(DrawingPolicy.GetBoundShaderStateInput(View.GetFeatureLevel()));
-		DrawingPolicy.SetSharedState(RHICmdList, &View, typename TBasePassForForwardShadingDrawingPolicy<LightMapPolicyType>::ContextDataType());
+		DrawingPolicy.SetSharedState(RHICmdList, &View, typename TBasePassForForwardShadingDrawingPolicy<LightMapPolicyType, NumDynamicPointLights>::ContextDataType());
 
 		for (int32 BatchElementIndex = 0; BatchElementIndex<Parameters.Mesh.Elements.Num(); BatchElementIndex++)
 		{
@@ -160,8 +160,8 @@ public:
 				BatchElementIndex,
 				bBackFace,
 				DitheredLODTransitionValue,
-				typename TBasePassForForwardShadingDrawingPolicy<LightMapPolicyType>::ElementDataType(LightMapElementData),
-				typename TBasePassForForwardShadingDrawingPolicy<LightMapPolicyType>::ContextDataType()
+				typename TBasePassForForwardShadingDrawingPolicy<LightMapPolicyType, NumDynamicPointLights>::ElementDataType(LightMapElementData),
+				typename TBasePassForForwardShadingDrawingPolicy<LightMapPolicyType, NumDynamicPointLights>::ContextDataType()
 				);
 			DrawingPolicy.DrawMesh(RHICmdList, Parameters.Mesh, BatchElementIndex);
 		}
@@ -200,7 +200,7 @@ bool FTranslucencyForwardShadingDrawingPolicyFactory::DrawDynamicMesh(
 			RHICmdList.SetDepthStencilState(TStaticDepthStencilState<false, CF_Always>::GetRHI());
 		}
 
-		ProcessBasePassMeshForForwardShading(
+		ProcessBasePassMeshForForwardShading<FDrawTranslucentMeshForwardShadingAction, 0>(
 			RHICmdList,
 			FProcessBasePassMeshParameters(
 				Mesh,
diff --git a/Engine/Source/Runtime/Renderer/Private/LightSceneInfo.cpp b/Engine/Source/Runtime/Renderer/Private/LightSceneInfo.cpp
--- a/Engine/Source/Runtime/Renderer/Private/LightSceneInfo.cpp
+++ b/Engine/Source/Runtime/Renderer/Private/LightSceneInfo.cpp
@@ -83,7 +83,9 @@ void FLightSceneInfo::AddToScene()
 	if (Proxy->CastsDynamicShadow() 
 		|| Proxy->CastsStaticShadow() 
 		// Lights that should be baked need to check for interactions to track unbuilt state correctly
-		|| Proxy->HasStaticLighting())
+		|| Proxy->HasStaticLighting()
+		// ES2 path supports dynamic point lights in the base pass using forward rendering, so we need to know the primitives
+		|| (Scene->GetFeatureLevel() < ERHIFeatureLevel::SM4 && Proxy->GetLightType() == LightType_Point))
 	{
 		// Add the light to the scene's light octree.
 		Scene->LightOctree.AddElement(LightSceneInfoCompact);
diff --git a/Engine/Source/Runtime/Renderer/Private/PrimitiveSceneInfo.cpp b/Engine/Source/Runtime/Renderer/Private/PrimitiveSceneInfo.cpp
--- a/Engine/Source/Runtime/Renderer/Private/PrimitiveSceneInfo.cpp
+++ b/Engine/Source/Runtime/Renderer/Private/PrimitiveSceneInfo.cpp
@@ -96,6 +96,7 @@ FPrimitiveSceneInfo::FPrimitiveSceneInfo(UPrimitiveComponent* InComponent,FScene
 	LastRenderTime(-FLT_MAX),
 	LastVisibilityChangeTime(0.0f),
 	Scene(InScene),
+	NumES2DynamicPointLights(0),
 	PackedIndex(INDEX_NONE),
 	ComponentForDebuggingOnly(InComponent),
 	bNeedsStaticMeshUpdate(false)
diff --git a/Engine/Source/Runtime/Renderer/Private/RendererScene.cpp b/Engine/Source/Runtime/Renderer/Private/RendererScene.cpp
--- a/Engine/Source/Runtime/Renderer/Private/RendererScene.cpp
+++ b/Engine/Source/Runtime/Renderer/Private/RendererScene.cpp
@@ -2512,62 +2512,62 @@ TStaticMeshDrawList<TBasePassDrawingPolicy<FSimpleDynamicLightingPolicy> >& FSce
 
 /** Maps the no light-map case to the appropriate base pass draw list. */
 template<>
-TStaticMeshDrawList<TBasePassForForwardShadingDrawingPolicy<FNoLightMapPolicy> >& FScene::GetForwardShadingBasePassDrawList<FNoLightMapPolicy>(EBasePassDrawListType DrawType)
+TStaticMeshDrawList<TBasePassForForwardShadingDrawingPolicy<FNoLightMapPolicy,0> >& FScene::GetForwardShadingBasePassDrawList<FNoLightMapPolicy>(EBasePassDrawListType DrawType)
 {
 	return BasePassForForwardShadingNoLightMapDrawList[DrawType];
 }
 
 /** Maps the simple light-map texture case to the appropriate base pass draw list. */
 template<>
-TStaticMeshDrawList< TBasePassForForwardShadingDrawingPolicy< TLightMapPolicy<LQ_LIGHTMAP> > >& FScene::GetForwardShadingBasePassDrawList< TLightMapPolicy<LQ_LIGHTMAP> >(EBasePassDrawListType DrawType)
+TStaticMeshDrawList< TBasePassForForwardShadingDrawingPolicy< TLightMapPolicy<LQ_LIGHTMAP>,0 > >& FScene::GetForwardShadingBasePassDrawList< TLightMapPolicy<LQ_LIGHTMAP> >(EBasePassDrawListType DrawType)
 {
 	return BasePassForForwardShadingLowQualityLightMapDrawList[DrawType];
 }
 
 template<>
-TStaticMeshDrawList< TBasePassForForwardShadingDrawingPolicy< TDistanceFieldShadowsAndLightMapPolicy<LQ_LIGHTMAP> > >& FScene::GetForwardShadingBasePassDrawList< TDistanceFieldShadowsAndLightMapPolicy<LQ_LIGHTMAP> >(EBasePassDrawListType DrawType)
+TStaticMeshDrawList< TBasePassForForwardShadingDrawingPolicy< TDistanceFieldShadowsAndLightMapPolicy<LQ_LIGHTMAP>,0 > >& FScene::GetForwardShadingBasePassDrawList< TDistanceFieldShadowsAndLightMapPolicy<LQ_LIGHTMAP> >(EBasePassDrawListType DrawType)
 {
 	return BasePassForForwardShadingDistanceFieldShadowMapLightMapDrawList[DrawType];
 }
 
 template<>
-TStaticMeshDrawList<TBasePassForForwardShadingDrawingPolicy<FSimpleDirectionalLightAndSHIndirectPolicy> >& FScene::GetForwardShadingBasePassDrawList<FSimpleDirectionalLightAndSHIndirectPolicy>(EBasePassDrawListType DrawType)
+TStaticMeshDrawList<TBasePassForForwardShadingDrawingPolicy<FSimpleDirectionalLightAndSHIndirectPolicy, 0> >& FScene::GetForwardShadingBasePassDrawList<FSimpleDirectionalLightAndSHIndirectPolicy>(EBasePassDrawListType DrawType)
 {
 	return BasePassForForwardShadingDirectionalLightAndSHIndirectDrawList[DrawType];
 }
 
 template<>
-TStaticMeshDrawList<TBasePassForForwardShadingDrawingPolicy<FSimpleDirectionalLightAndSHDirectionalIndirectPolicy> >& FScene::GetForwardShadingBasePassDrawList<FSimpleDirectionalLightAndSHDirectionalIndirectPolicy>(EBasePassDrawListType DrawType)
+TStaticMeshDrawList<TBasePassForForwardShadingDrawingPolicy<FSimpleDirectionalLightAndSHDirectionalIndirectPolicy, 0> >& FScene::GetForwardShadingBasePassDrawList<FSimpleDirectionalLightAndSHDirectionalIndirectPolicy>(EBasePassDrawListType DrawType)
 {
 	return BasePassForForwardShadingDirectionalLightAndSHDirectionalIndirectDrawList[DrawType];
 }
 
 template<>
-TStaticMeshDrawList<TBasePassForForwardShadingDrawingPolicy<FSimpleDirectionalLightAndSHDirectionalCSMIndirectPolicy> >& FScene::GetForwardShadingBasePassDrawList<FSimpleDirectionalLightAndSHDirectionalCSMIndirectPolicy>(EBasePassDrawListType DrawType)
+TStaticMeshDrawList<TBasePassForForwardShadingDrawingPolicy<FSimpleDirectionalLightAndSHDirectionalCSMIndirectPolicy, 0> >& FScene::GetForwardShadingBasePassDrawList<FSimpleDirectionalLightAndSHDirectionalCSMIndirectPolicy>(EBasePassDrawListType DrawType)
 {
 	return BasePassForForwardShadingDirectionalLightAndSHDirectionalCSMIndirectDrawList[DrawType];
 }
 
 template<>
-TStaticMeshDrawList<TBasePassForForwardShadingDrawingPolicy<FMovableDirectionalLightLightingPolicy> >& FScene::GetForwardShadingBasePassDrawList<FMovableDirectionalLightLightingPolicy>(EBasePassDrawListType DrawType)
+TStaticMeshDrawList<TBasePassForForwardShadingDrawingPolicy<FMovableDirectionalLightLightingPolicy, 0> >& FScene::GetForwardShadingBasePassDrawList<FMovableDirectionalLightLightingPolicy>(EBasePassDrawListType DrawType)
 {
 	return BasePassForForwardShadingMovableDirectionalLightDrawList[DrawType];
 }
 
 template<>
-TStaticMeshDrawList<TBasePassForForwardShadingDrawingPolicy<FMovableDirectionalLightCSMLightingPolicy> >& FScene::GetForwardShadingBasePassDrawList<FMovableDirectionalLightCSMLightingPolicy>(EBasePassDrawListType DrawType)
+TStaticMeshDrawList<TBasePassForForwardShadingDrawingPolicy<FMovableDirectionalLightCSMLightingPolicy, 0> >& FScene::GetForwardShadingBasePassDrawList<FMovableDirectionalLightCSMLightingPolicy>(EBasePassDrawListType DrawType)
 {
 	return BasePassForForwardShadingMovableDirectionalLightCSMDrawList[DrawType];
 }
 
 template<>
-TStaticMeshDrawList<TBasePassForForwardShadingDrawingPolicy<FMovableDirectionalLightWithLightmapLightingPolicy> >& FScene::GetForwardShadingBasePassDrawList<FMovableDirectionalLightWithLightmapLightingPolicy>(EBasePassDrawListType DrawType)
+TStaticMeshDrawList<TBasePassForForwardShadingDrawingPolicy<FMovableDirectionalLightWithLightmapLightingPolicy, 0> >& FScene::GetForwardShadingBasePassDrawList<FMovableDirectionalLightWithLightmapLightingPolicy>(EBasePassDrawListType DrawType)
 {
 	return BasePassForForwardShadingMovableDirectionalLightLightmapDrawList[DrawType];
 }
 
 template<>
-TStaticMeshDrawList<TBasePassForForwardShadingDrawingPolicy<FMovableDirectionalLightCSMWithLightmapLightingPolicy> >& FScene::GetForwardShadingBasePassDrawList<FMovableDirectionalLightCSMWithLightmapLightingPolicy>(EBasePassDrawListType DrawType)
+TStaticMeshDrawList<TBasePassForForwardShadingDrawingPolicy<FMovableDirectionalLightCSMWithLightmapLightingPolicy, 0> >& FScene::GetForwardShadingBasePassDrawList<FMovableDirectionalLightCSMWithLightmapLightingPolicy>(EBasePassDrawListType DrawType)
 {
 	return BasePassForForwardShadingMovableDirectionalLightCSMLightmapDrawList[DrawType];
 }
diff --git a/Engine/Source/Runtime/Renderer/Private/SceneCore.cpp b/Engine/Source/Runtime/Renderer/Private/SceneCore.cpp
--- a/Engine/Source/Runtime/Renderer/Private/SceneCore.cpp
+++ b/Engine/Source/Runtime/Renderer/Private/SceneCore.cpp
@@ -136,7 +136,7 @@ FLightPrimitiveInteraction::FLightPrimitiveInteraction(
 	bool bInIsShadowMapped,
 	bool bInHasTranslucentObjectShadow,
 	bool bInHasInsetObjectShadow
-	):
+	) :
 	LightSceneInfo(InLightSceneInfo),
 	PrimitiveSceneInfo(InPrimitiveSceneInfo),
 	LightId(InLightSceneInfo->Id),
@@ -146,7 +146,8 @@ FLightPrimitiveInteraction::FLightPrimitiveInteraction(
 	bUncachedStaticLighting(false),
 	bHasTranslucentObjectShadow(bInHasTranslucentObjectShadow),
 	bHasInsetObjectShadow(bInHasInsetObjectShadow),
-	bSelfShadowOnly(false)
+	bSelfShadowOnly(false),
+	bES2DynamicPointLight(false)
 {
 	// Determine whether this light-primitive interaction produces a shadow.
 	if(PrimitiveSceneInfo->Proxy->HasStaticLighting())
@@ -195,6 +196,15 @@ FLightPrimitiveInteraction::FLightPrimitiveInteraction(
 	{
 		// Add the interaction to the light's interaction list.
 		PrevPrimitiveLink = &LightSceneInfo->DynamicPrimitiveList;
+
+		// ES2 dynamic point lights
+		if (PrimitiveSceneInfo->Scene->GetFeatureLevel() < ERHIFeatureLevel::SM4 && LightSceneInfo->Proxy->GetLightType() == LightType_Point)
+		{
+			bES2DynamicPointLight = true;
+			PrimitiveSceneInfo->NumES2DynamicPointLights++;
+			// The forward renderer renders dynamic point lights as part of the base pass using the dynamic path only.
+			PrimitiveSceneInfo->Proxy->bDisableStaticPath = true;
+		}
 	}
 
 	NextPrimitive = *PrevPrimitiveLink;
@@ -230,6 +240,16 @@ FLightPrimitiveInteraction::~FLightPrimitiveInteraction()
 	}
 #endif
 
+	// Track ES2 dynamic point light count
+	if (bES2DynamicPointLight)
+	{
+		PrimitiveSceneInfo->NumES2DynamicPointLights--;
+		if (PrimitiveSceneInfo->NumES2DynamicPointLights == 0)
+		{
+			PrimitiveSceneInfo->Proxy->bDisableStaticPath = false;
+		}
+	}
+
 	// Remove the interaction from the light's interaction list.
 	if(NextPrimitive)
 	{
diff --git a/Engine/Source/Runtime/Renderer/Private/SceneCore.h b/Engine/Source/Runtime/Renderer/Private/SceneCore.h
--- a/Engine/Source/Runtime/Renderer/Private/SceneCore.h
+++ b/Engine/Source/Runtime/Renderer/Private/SceneCore.h
@@ -37,6 +37,7 @@ public:
 	bool HasTranslucentObjectShadow() const { return bHasTranslucentObjectShadow; }
 	bool HasInsetObjectShadow() const { return bHasInsetObjectShadow; }
 	bool CastsSelfShadowOnly() const { return bSelfShadowOnly; }
+	bool IsES2DynamicPointLight() const { return bES2DynamicPointLight; }
 	FLightSceneInfo* GetLight() const { return LightSceneInfo; }
 	int32 GetLightId() const { return LightId; }
 	FPrimitiveSceneInfo* GetPrimitiveSceneInfo() const { return PrimitiveSceneInfo; }
@@ -99,6 +100,9 @@ private:
 	/** True if the primitive only shadows itself. */
 	uint32 bSelfShadowOnly : 1;
 
+	/** True this is an ES2 dynamic point light interaction. */
+	uint32 bES2DynamicPointLight : 1;	
+
 	/** Initialization constructor. */
 	FLightPrimitiveInteraction(FLightSceneInfo* InLightSceneInfo,FPrimitiveSceneInfo* InPrimitiveSceneInfo,
 		bool bIsDynamic,bool bInLightMapped,bool bInIsShadowMapped, bool bInHasTranslucentObjectShadow, bool bInHasInsetObjectShadow);
diff --git a/Engine/Source/Runtime/Renderer/Private/ScenePrivate.h b/Engine/Source/Runtime/Renderer/Private/ScenePrivate.h
--- a/Engine/Source/Runtime/Renderer/Private/ScenePrivate.h
+++ b/Engine/Source/Runtime/Renderer/Private/ScenePrivate.h
@@ -1421,20 +1421,20 @@ public:
 	TStaticMeshDrawList<TBasePassDrawingPolicy<LightMapPolicyType> >& GetBasePassDrawList(EBasePassDrawListType DrawType);
 
 	/** Forward shading base pass draw lists */
-	TStaticMeshDrawList<TBasePassForForwardShadingDrawingPolicy<FNoLightMapPolicy> > BasePassForForwardShadingNoLightMapDrawList[EBasePass_MAX];
-	TStaticMeshDrawList< TBasePassForForwardShadingDrawingPolicy< TLightMapPolicy<LQ_LIGHTMAP> > >							BasePassForForwardShadingLowQualityLightMapDrawList[EBasePass_MAX];
-	TStaticMeshDrawList< TBasePassForForwardShadingDrawingPolicy< TDistanceFieldShadowsAndLightMapPolicy<LQ_LIGHTMAP> > >	BasePassForForwardShadingDistanceFieldShadowMapLightMapDrawList[EBasePass_MAX];
-	TStaticMeshDrawList<TBasePassForForwardShadingDrawingPolicy< FSimpleDirectionalLightAndSHIndirectPolicy > >				BasePassForForwardShadingDirectionalLightAndSHIndirectDrawList[EBasePass_MAX];
-	TStaticMeshDrawList<TBasePassForForwardShadingDrawingPolicy< FSimpleDirectionalLightAndSHDirectionalIndirectPolicy > >		BasePassForForwardShadingDirectionalLightAndSHDirectionalIndirectDrawList[EBasePass_MAX];
-	TStaticMeshDrawList<TBasePassForForwardShadingDrawingPolicy< FSimpleDirectionalLightAndSHDirectionalCSMIndirectPolicy > >	BasePassForForwardShadingDirectionalLightAndSHDirectionalCSMIndirectDrawList[EBasePass_MAX];
-	TStaticMeshDrawList<TBasePassForForwardShadingDrawingPolicy< FMovableDirectionalLightLightingPolicy > >					BasePassForForwardShadingMovableDirectionalLightDrawList[EBasePass_MAX];
-	TStaticMeshDrawList<TBasePassForForwardShadingDrawingPolicy< FMovableDirectionalLightCSMLightingPolicy > >				BasePassForForwardShadingMovableDirectionalLightCSMDrawList[EBasePass_MAX];
-	TStaticMeshDrawList<TBasePassForForwardShadingDrawingPolicy< FMovableDirectionalLightWithLightmapLightingPolicy> >		BasePassForForwardShadingMovableDirectionalLightLightmapDrawList[EBasePass_MAX];
-	TStaticMeshDrawList<TBasePassForForwardShadingDrawingPolicy< FMovableDirectionalLightCSMWithLightmapLightingPolicy> >	BasePassForForwardShadingMovableDirectionalLightCSMLightmapDrawList[EBasePass_MAX];
+	TStaticMeshDrawList<TBasePassForForwardShadingDrawingPolicy<FNoLightMapPolicy,0> > BasePassForForwardShadingNoLightMapDrawList[EBasePass_MAX];
+	TStaticMeshDrawList<TBasePassForForwardShadingDrawingPolicy< TLightMapPolicy<LQ_LIGHTMAP>,0 > >								BasePassForForwardShadingLowQualityLightMapDrawList[EBasePass_MAX];
+	TStaticMeshDrawList<TBasePassForForwardShadingDrawingPolicy< TDistanceFieldShadowsAndLightMapPolicy<LQ_LIGHTMAP>,0 > >		BasePassForForwardShadingDistanceFieldShadowMapLightMapDrawList[EBasePass_MAX];
+	TStaticMeshDrawList<TBasePassForForwardShadingDrawingPolicy< FSimpleDirectionalLightAndSHIndirectPolicy,0 > >				BasePassForForwardShadingDirectionalLightAndSHIndirectDrawList[EBasePass_MAX];
+	TStaticMeshDrawList<TBasePassForForwardShadingDrawingPolicy< FSimpleDirectionalLightAndSHDirectionalIndirectPolicy,0 > >	BasePassForForwardShadingDirectionalLightAndSHDirectionalIndirectDrawList[EBasePass_MAX];
+	TStaticMeshDrawList<TBasePassForForwardShadingDrawingPolicy< FSimpleDirectionalLightAndSHDirectionalCSMIndirectPolicy,0 > >	BasePassForForwardShadingDirectionalLightAndSHDirectionalCSMIndirectDrawList[EBasePass_MAX];
+	TStaticMeshDrawList<TBasePassForForwardShadingDrawingPolicy< FMovableDirectionalLightLightingPolicy,0 > >					BasePassForForwardShadingMovableDirectionalLightDrawList[EBasePass_MAX];
+	TStaticMeshDrawList<TBasePassForForwardShadingDrawingPolicy< FMovableDirectionalLightCSMLightingPolicy,0 > >				BasePassForForwardShadingMovableDirectionalLightCSMDrawList[EBasePass_MAX];
+	TStaticMeshDrawList<TBasePassForForwardShadingDrawingPolicy< FMovableDirectionalLightWithLightmapLightingPolicy,0> >		BasePassForForwardShadingMovableDirectionalLightLightmapDrawList[EBasePass_MAX];
+	TStaticMeshDrawList<TBasePassForForwardShadingDrawingPolicy< FMovableDirectionalLightCSMWithLightmapLightingPolicy,0> >		BasePassForForwardShadingMovableDirectionalLightCSMLightmapDrawList[EBasePass_MAX];
 
 	/** Maps a light-map type to the appropriate base pass draw list. */
 	template<typename LightMapPolicyType>
-	TStaticMeshDrawList<TBasePassForForwardShadingDrawingPolicy<LightMapPolicyType> >& GetForwardShadingBasePassDrawList(EBasePassDrawListType DrawType);
+	TStaticMeshDrawList<TBasePassForForwardShadingDrawingPolicy<LightMapPolicyType,0> >& GetForwardShadingBasePassDrawList(EBasePassDrawListType DrawType);
 
 	/**
 	 * The following arrays are densely packed primitive data needed by various
diff --git a/Engine/Source/Runtime/Renderer/Public/PrimitiveSceneInfo.h b/Engine/Source/Runtime/Renderer/Public/PrimitiveSceneInfo.h
--- a/Engine/Source/Runtime/Renderer/Public/PrimitiveSceneInfo.h
+++ b/Engine/Source/Runtime/Renderer/Public/PrimitiveSceneInfo.h
@@ -230,6 +230,9 @@ public:
 	/** The scene the primitive is in. */
 	FScene* Scene;
 
+	/** The number of dynamic point lights for ES2 */
+	int32 NumES2DynamicPointLights;
+
 	/** Initialization constructor. */
 	FPrimitiveSceneInfo(UPrimitiveComponent* InPrimitive,FScene* InScene);
 
