Distortion pass uses stencil mask, optimizes lower screen coverages.

diff --git a/Engine/Shaders/DistortApplyScreenPS.usf b/Engine/Shaders/DistortApplyScreenPS.usf
--- a/Engine/Shaders/DistortApplyScreenPS.usf
+++ b/Engine/Shaders/DistortApplyScreenPS.usf
@@ -58,4 +58,12 @@ void Main(
 	// sample screen using offset coords
 	half4 DistColor = Texture2DSample(SceneColorTexture, SceneColorTextureSampler, NewTexCoord);
 	OutColor = DistColor;
+}
+
+void Merge(
+	in float4 TexCoord: TEXCOORD0,
+	out float4 OutColor : SV_Target0
+	)
+{
+	OutColor = Texture2DSample(SceneColorTexture, SceneColorTextureSampler, TexCoord.xy);
 }
\ No newline at end of file
diff --git a/Engine/Source/Runtime/Renderer/Private/DistortionRendering.cpp b/Engine/Source/Runtime/Renderer/Private/DistortionRendering.cpp
--- a/Engine/Source/Runtime/Renderer/Private/DistortionRendering.cpp
+++ b/Engine/Source/Runtime/Renderer/Private/DistortionRendering.cpp
@@ -89,6 +89,54 @@ private:
 /** distortion apply screen pixel shader implementation */
 IMPLEMENT_SHADER_TYPE(,FDistortionApplyScreenPS,TEXT("DistortApplyScreenPS"),TEXT("Main"),SF_Pixel);
 
+/**
+* A pixel shader that applies the distorted image to the scene
+*/
+class FDistortionMergePS : public FGlobalShader
+{
+	DECLARE_SHADER_TYPE(FDistortionMergePS,Global);
+public:
+
+	static bool ShouldCache(EShaderPlatform Platform) { return true; }
+
+	FDistortionMergePS(const ShaderMetaType::CompiledShaderInitializerType& Initializer)
+		: FGlobalShader(Initializer)
+	{
+		SceneColorTexture.Bind(Initializer.ParameterMap,TEXT("SceneColorTexture"));
+		SceneColorTextureSampler.Bind(Initializer.ParameterMap,TEXT("SceneColorTextureSampler"));
+	}
+	FDistortionMergePS() {}
+
+	void SetParameters(const FRenderingCompositePassContext& Context, const FTextureRHIParamRef& PassTexture)
+	{
+		FSceneRenderTargets& SceneContext = FSceneRenderTargets::Get(Context.RHICmdList);
+		const FPixelShaderRHIParamRef ShaderRHI = GetPixelShader();
+
+		SetTextureParameter(
+			Context.RHICmdList,
+			ShaderRHI,
+			SceneColorTexture,
+			SceneColorTextureSampler,
+			TStaticSamplerState<SF_Point,AM_Clamp,AM_Clamp,AM_Clamp>::GetRHI(),
+			PassTexture
+			);
+	}
+
+	virtual bool Serialize(FArchive& Ar) override
+	{
+		bool bShaderHasOutdatedParameters = FGlobalShader::Serialize(Ar);
+		Ar << SceneColorTexture << SceneColorTextureSampler;
+		return bShaderHasOutdatedParameters;
+	}
+
+private:
+	FShaderResourceParameter SceneColorTexture;
+	FShaderResourceParameter SceneColorTextureSampler;
+};
+
+/** distortion merge to scene pixel shader implementation */
+IMPLEMENT_SHADER_TYPE(,FDistortionMergePS,TEXT("DistortApplyScreenPS"),TEXT("Merge"),SF_Pixel);
+
 /**
 * Policy for drawing distortion mesh accumulated offsets
 */
@@ -769,6 +817,11 @@ void FSceneRenderer::RenderDistortion(FRHICommandListImmediate& RHICmdList)
 
 	TRefCountPtr<IPooledRenderTarget> DistortionRT;
 	FSceneRenderTargets& SceneContext = FSceneRenderTargets::Get(RHICmdList);
+	
+	// Use stencil mask to optimize cases with lower screen coverage.
+	// Note: This adds an extra pass which is actually slower as distortion tends towards full-screen.
+	//       It could be worth testing object screen bounds then reverting to a target flip and single pass.
+	const uint8 StencilMaskBit = 0x80;
 
 	// Render accumulated distortion offsets
 	if( bRender)
@@ -804,8 +857,8 @@ void FSceneRenderer::RenderDistortion(FRHICommandListImmediate& RHICmdList)
 				// clear offsets to 0
 				RHICmdList.Clear(true, FLinearColor(0, 0, 0, 0), false, (float)ERHIZBuffer::FarPlane, false, 0, FIntRect());
 
-				// enable depth test but disable depth writes
-				RHICmdList.SetDepthStencilState(TStaticDepthStencilState<false, CF_DepthNearOrEqual>::GetRHI());
+				// test against depth and write stencil mask
+				RHICmdList.SetDepthStencilState(TStaticDepthStencilState<false, CF_DepthNearOrEqual, true, CF_Always, SO_Keep, SO_Keep, SO_Replace, false, CF_Always, SO_Keep, SO_Keep, SO_Keep, StencilMaskBit, StencilMaskBit>::GetRHI(), StencilMaskBit);
 
 				// additive blending of offsets (or complexity if the shader complexity viewmode is enabled)
 				RHICmdList.SetBlendState(TStaticBlendState<CW_RGBA, BO_Add, BF_One, BF_One, BO_Add, BF_One, BF_One>::GetRHI());
@@ -816,9 +869,6 @@ void FSceneRenderer::RenderDistortion(FRHICommandListImmediate& RHICmdList)
 
 			if (bDirty)
 			{
-				// restore default stencil state
-				RHICmdList.SetDepthStencilState(TStaticDepthStencilState<false, CF_DepthNearOrEqual>::GetRHI());
-
 				// resolve using the current ResolveParams 
 				RHICmdList.CopyToResolveTarget(DistortionRT->GetRenderTargetItem().TargetableTexture, DistortionRT->GetRenderTargetItem().ShaderResourceTexture, false, FResolveParams());
 				// to be able to observe results with VisualizeTexture
@@ -840,73 +890,92 @@ void FSceneRenderer::RenderDistortion(FRHICommandListImmediate& RHICmdList)
 		const FSceneRenderTargetItem& DestRenderTarget = NewSceneColor->GetRenderTargetItem();
 // OCULUS END
 
-		// Apply distortion as a full-screen pass		
+		// Apply distortion and store off-screen
+		SetRenderTarget(RHICmdList, DestRenderTarget.TargetableTexture, SceneContext.GetSceneDepthSurface());
+
 		for(int32 ViewIndex = 0, Num = Views.Num(); ViewIndex < Num; ++ViewIndex)
 		{
 			SCOPED_CONDITIONAL_DRAW_EVENTF(RHICmdList, EventView, Views.Num() > 1, TEXT("View%d"), ViewIndex);
 
 			FViewInfo& View = Views[ViewIndex];
 
-//			const FSceneView& View = Context.View;
+			// useful when we move this into the compositing graph
+			FRenderingCompositePassContext Context(RHICmdList, View);
 
-			// set the state
+			// Set the view family's render target/viewport.
+			Context.SetViewportAndCallRHI(View.ViewRect);
+
+			// test against stencil mask
 			RHICmdList.SetBlendState(TStaticBlendState<>::GetRHI());
 			RHICmdList.SetRasterizerState(TStaticRasterizerState<>::GetRHI());
-			RHICmdList.SetDepthStencilState(TStaticDepthStencilState<false, CF_Always>::GetRHI());
+			RHICmdList.SetDepthStencilState(TStaticDepthStencilState<false, CF_Always, true, CF_Equal, SO_Keep, SO_Keep, SO_Keep, false, CF_Always, SO_Keep, SO_Keep, SO_Keep, StencilMaskBit, StencilMaskBit>::GetRHI(), StencilMaskBit);
 
-			{
-// OCULUS BEGIN: Weird, different eyes may have different render targets!
-// 				TRefCountPtr<IPooledRenderTarget> NewSceneColor;
-// we don't create a new name to make it easier to use "vis SceneColor" and get the last HDRSceneColor
-// 				GRenderTargetPool.FindFreeElement(SceneContext.SceneColor->GetDesc(), NewSceneColor, TEXT("SceneColor"));
-// 				const FSceneRenderTargetItem& DestRenderTarget = NewSceneColor->GetRenderTargetItem(); 
-// OCULUS END
+			TShaderMapRef<FPostProcessVS> VertexShader(View.ShaderMap);
+			TShaderMapRef<FDistortionApplyScreenPS> PixelShader(View.ShaderMap);
 
-				SetRenderTarget(RHICmdList, DestRenderTarget.TargetableTexture, FTextureRHIParamRef());
+			static FGlobalBoundShaderState BoundShaderState;
+				
+			SetGlobalBoundShaderState(RHICmdList, View.GetFeatureLevel(), BoundShaderState, GFilterVertexDeclaration.VertexDeclarationRHI, *VertexShader, *PixelShader);
+
+			VertexShader->SetParameters(Context);
+			PixelShader->SetParameters(Context, *DistortionRT);
+
+			// Draw a quad mapping scene color to the view's render target
+			DrawRectangle(
+				RHICmdList,
+				0, 0,
+				View.ViewRect.Width(), View.ViewRect.Height(),
+				View.ViewRect.Min.X, View.ViewRect.Min.Y, 
+				View.ViewRect.Width(), View.ViewRect.Height(),
+				View.ViewRect.Size(),
+				SceneContext.GetBufferSizeXY(),
+				*VertexShader,
+				EDRF_UseTriangleOptimization);
+		}
 
-				// useful when we move this into the compositing graph
-				FRenderingCompositePassContext Context(RHICmdList, View);
+		// Resolve and merge distortion to main scene
+		RHICmdList.CopyToResolveTarget(DestRenderTarget.TargetableTexture, DestRenderTarget.ShaderResourceTexture, false, FResolveParams());
+		SetRenderTarget(RHICmdList, SceneContext.GetSceneColor()->GetRenderTargetItem().TargetableTexture, SceneContext.GetSceneDepthSurface());
 
-				// Set the view family's render target/viewport.
-				Context.SetViewportAndCallRHI(View.ViewRect);
+		for(int32 ViewIndex = 0, Num = Views.Num(); ViewIndex < Num; ++ViewIndex)
+		{
+			SCOPED_CONDITIONAL_DRAW_EVENTF(RHICmdList, EventView, Views.Num() > 1, TEXT("View%d"), ViewIndex);
 
-				TShaderMapRef<FPostProcessVS> VertexShader(View.ShaderMap);
-				TShaderMapRef<FDistortionApplyScreenPS> PixelShader(View.ShaderMap);
+			FViewInfo& View = Views[ViewIndex];
 
-				static FGlobalBoundShaderState BoundShaderState;
-				
-				SetGlobalBoundShaderState(RHICmdList, View.GetFeatureLevel(), BoundShaderState, GFilterVertexDeclaration.VertexDeclarationRHI, *VertexShader, *PixelShader);
-
-				VertexShader->SetParameters(Context);
-				PixelShader->SetParameters(Context, *DistortionRT);
-
-				// Draw a quad mapping scene color to the view's render target
-				DrawRectangle(
-					RHICmdList,
-					0, 0,
-					View.ViewRect.Width(), View.ViewRect.Height(),
-					View.ViewRect.Min.X, View.ViewRect.Min.Y, 
-					View.ViewRect.Width(), View.ViewRect.Height(),
-					View.ViewRect.Size(),
-					SceneContext.GetBufferSizeXY(),
-					*VertexShader,
-					EDRF_UseTriangleOptimization);
-
-				RHICmdList.CopyToResolveTarget(DestRenderTarget.TargetableTexture, DestRenderTarget.ShaderResourceTexture, false, FResolveParams());
-
-// OCULUS BEGIN
-// 				SceneContext.SetSceneColor(NewSceneColor);
-// 				check(SceneContext.GetSceneColor());
-// OCULUS END				
+			// useful when we move this into the compositing graph
+			FRenderingCompositePassContext Context(RHICmdList, View);
 
-			}
-		
+			// Set the view family's render target/viewport.
+			Context.SetViewportAndCallRHI(View.ViewRect);				
+
+			// test against stencil mask and clear it
+			RHICmdList.SetDepthStencilState(TStaticDepthStencilState<false, CF_Always, true, CF_Equal, SO_Keep, SO_Keep, SO_Zero, false, CF_Always, SO_Keep, SO_Keep, SO_Keep, StencilMaskBit, StencilMaskBit>::GetRHI(), StencilMaskBit);
+					
+			static FGlobalBoundShaderState BoundShaderState;
+				
+			TShaderMapRef<FPostProcessVS> VertexShader(View.ShaderMap);
+			TShaderMapRef<FDistortionMergePS> PixelShader(View.ShaderMap);
+			SetGlobalBoundShaderState(RHICmdList, View.GetFeatureLevel(), BoundShaderState, GFilterVertexDeclaration.VertexDeclarationRHI, *VertexShader, *PixelShader);
+
+			VertexShader->SetParameters(Context);
+			PixelShader->SetParameters(Context, DestRenderTarget.ShaderResourceTexture);
+
+			DrawRectangle(
+				RHICmdList,
+				0, 0,
+				View.ViewRect.Width(), View.ViewRect.Height(),
+				View.ViewRect.Min.X, View.ViewRect.Min.Y, 
+				View.ViewRect.Width(), View.ViewRect.Height(),
+				View.ViewRect.Size(),
+				SceneContext.GetBufferSizeXY(),
+				*VertexShader,
+				EDRF_UseTriangleOptimization);
 		}
 
-		// OCULUS BEGIN
-		SceneContext.SetSceneColor(NewSceneColor);
-		check(SceneContext.GetSceneColor());
-// OCULUS END				
+		// restore state
+		RHICmdList.SetDepthStencilState(TStaticDepthStencilState<false, CF_DepthNearOrEqual>::GetRHI());		
+			
 		// Distortions RT is no longer needed, buffer can be reused by the pool, see BeginRenderingDistortionAccumulation() call above
 		SceneContext.FinishRenderingSceneColor(RHICmdList, false);
 	}
