Only consider movable lights for mobile dynamic light rendering
Allow r.MobileNumDynamicPointLights to be modified at runtime

diff --git a/Engine/Source/Runtime/Core/Private/HAL/ConsoleManager.cpp b/Engine/Source/Runtime/Core/Private/HAL/ConsoleManager.cpp
--- a/Engine/Source/Runtime/Core/Private/HAL/ConsoleManager.cpp
+++ b/Engine/Source/Runtime/Core/Private/HAL/ConsoleManager.cpp
@@ -1605,7 +1605,7 @@ static TAutoConsoleVariable<int32> CVarMobileNumDynamicPointLights(
 	TEXT("r.MobileNumDynamicPointLights"),
 	4,
 	TEXT("The number of dynamic point lights to support on mobile devices. Setting this to 0 for games which do not require dynamic point lights will reduce the number of shaders generated."), 
-	ECVF_RenderThreadSafe | ECVF_ReadOnly);
+	ECVF_RenderThreadSafe);
 
 static TAutoConsoleVariable<int32> CVarMobileDynamicPointLightsUseStaticBranch(
 	TEXT("r.MobileDynamicPointLightsUseStaticBranch"),
diff --git a/Engine/Source/Runtime/Engine/Classes/Components/LightComponentBase.h b/Engine/Source/Runtime/Engine/Classes/Components/LightComponentBase.h
--- a/Engine/Source/Runtime/Engine/Classes/Components/LightComponentBase.h
+++ b/Engine/Source/Runtime/Engine/Classes/Components/LightComponentBase.h
@@ -143,6 +143,12 @@ class ENGINE_API ULightComponentBase : public USceneComponent
 	 */
 	virtual void UpdateLightGUIDs();
 
+	/** Returns true if the light's Mobility is set to Movable */
+	bool IsMovable() const
+	{
+		return (Mobility == EComponentMobility::Movable);
+	}
+
 	/**
 	 * Return True if a light's parameters as well as its position is static during gameplay, and can thus use static lighting.
 	 * A light with HasStaticLighting() == true will always have HasStaticShadowing() == true as well.
diff --git a/Engine/Source/Runtime/Engine/Private/Components/LightComponent.cpp b/Engine/Source/Runtime/Engine/Private/Components/LightComponent.cpp
--- a/Engine/Source/Runtime/Engine/Private/Components/LightComponent.cpp
+++ b/Engine/Source/Runtime/Engine/Private/Components/LightComponent.cpp
@@ -213,6 +213,7 @@ FLightSceneProxy::FLightSceneProxy(const ULightComponent* InLightComponent)
 	, ShadowMapChannel(InLightComponent->ShadowMapChannel)
 	, PreviewShadowMapChannel(InLightComponent->PreviewShadowMapChannel)
 	, IESTexture(0)
+	, bMovable(InLightComponent->IsMovable())
 	, bStaticLighting(InLightComponent->HasStaticLighting())
 	, bStaticShadowing(InLightComponent->HasStaticShadowing())
 	, bCastDynamicShadow(InLightComponent->CastShadows && InLightComponent->CastDynamicShadows)
diff --git a/Engine/Source/Runtime/Engine/Public/SceneManagement.h b/Engine/Source/Runtime/Engine/Public/SceneManagement.h
--- a/Engine/Source/Runtime/Engine/Public/SceneManagement.h
+++ b/Engine/Source/Runtime/Engine/Public/SceneManagement.h
@@ -784,6 +784,7 @@ public:
 	inline UTextureLightProfile* GetIESTexture() const { return IESTexture; }
 	inline FTexture* GetIESTextureResource() const { return IESTexture ? IESTexture->Resource : 0; }
 	inline const FMaterialRenderProxy* GetLightFunctionMaterial() const { return LightFunctionMaterial; }
+	inline bool IsMovable() const { return bMovable; }
 	inline bool HasStaticLighting() const { return bStaticLighting; }
 	inline bool HasStaticShadowing() const { return bStaticShadowing; }
 	inline bool CastsDynamicShadow() const { return bCastDynamicShadow; }
@@ -875,6 +876,9 @@ protected:
 	 */
 	UTextureLightProfile* IESTexture;
 
+	/* True if the light's Mobility is set to Movable. */
+	const uint32 bMovable : 1;
+
 	/**
 	 * Return True if a light's parameters as well as its position is static during gameplay, and can thus use static lighting.
 	 * A light with HasStaticLighting() == true will always have HasStaticShadowing() == true as well.
diff --git a/Engine/Source/Runtime/Renderer/Private/ForwardBasePassRendering.cpp b/Engine/Source/Runtime/Renderer/Private/ForwardBasePassRendering.cpp
--- a/Engine/Source/Runtime/Renderer/Private/ForwardBasePassRendering.cpp
+++ b/Engine/Source/Runtime/Renderer/Private/ForwardBasePassRendering.cpp
@@ -68,12 +68,15 @@ IMPLEMENT_FORWARD_SHADING_BASEPASS_LIGHTMAPPED_SHADER_TYPE(FMovableDirectionalLi
 FBasePassFowardDynamicPointLightInfo::FBasePassFowardDynamicPointLightInfo(const FPrimitiveSceneProxy* InSceneProxy)
 : NumDynamicPointLights(0)
 {
+	static auto* MobileNumDynamicPointLightsCVar = IConsoleManager::Get().FindTConsoleVariableDataInt(TEXT("r.MobileNumDynamicPointLights"));
+	const int32 MobileNumDynamicPointLights = MobileNumDynamicPointLightsCVar->GetValueOnRenderThread();
+
 	if (InSceneProxy != nullptr)
 	{
-		for (FLightPrimitiveInteraction* LPI = InSceneProxy->GetPrimitiveSceneInfo()->LightList; LPI; LPI = LPI->GetNextLight())
+		for (FLightPrimitiveInteraction* LPI = InSceneProxy->GetPrimitiveSceneInfo()->LightList; LPI && NumDynamicPointLights < MobileNumDynamicPointLights; LPI = LPI->GetNextLight())
 		{
 			FLightSceneProxy* LightProxy = LPI->GetLight()->Proxy;
-			if (LightProxy->GetLightType() == LightType_Point)
+			if (LightProxy->GetLightType() == LightType_Point && LightProxy->IsMovable())
 			{
 				FVector NormalizedLightDirection;
 				FVector2D SpotAngles;
@@ -101,10 +104,6 @@ FBasePassFowardDynamicPointLightInfo::FBasePassFowardDynamicPointLightInfo(const
 				}
 
 				NumDynamicPointLights++;
-				if (NumDynamicPointLights == MAX_BASEPASS_DYNAMIC_POINT_LIGHTS)
-				{
-					break;
-				}
 			}
 		}
 	}
@@ -358,7 +357,10 @@ bool FBasePassForwardOpaqueDrawingPolicyFactory::DrawDynamicMesh(
 	// Only draw opaque materials.
 	if(!IsTranslucentBlendMode(BlendMode))
 	{
-		int32 NumDynamicPointLights = PrimitiveSceneProxy ? PrimitiveSceneProxy->GetPrimitiveSceneInfo()->NumES2DynamicPointLights : 0;
+		static auto* MobileNumDynamicPointLightsCVar = IConsoleManager::Get().FindTConsoleVariableDataInt(TEXT("r.MobileNumDynamicPointLights"));
+		const int32 MobileNumDynamicPointLights = MobileNumDynamicPointLightsCVar->GetValueOnRenderThread();
+
+		int32 NumDynamicPointLights = PrimitiveSceneProxy ? FMath::Min<int32>(PrimitiveSceneProxy->GetPrimitiveSceneInfo()->NumES2DynamicPointLights, MobileNumDynamicPointLights) : 0;
 		const bool bIsUnlit = Material->GetShadingModel() == MSM_Unlit;
 
 		if (NumDynamicPointLights == 0 || bIsUnlit)
diff --git a/Engine/Source/Runtime/Renderer/Private/LightSceneInfo.cpp b/Engine/Source/Runtime/Renderer/Private/LightSceneInfo.cpp
--- a/Engine/Source/Runtime/Renderer/Private/LightSceneInfo.cpp
+++ b/Engine/Source/Runtime/Renderer/Private/LightSceneInfo.cpp
@@ -85,7 +85,7 @@ void FLightSceneInfo::AddToScene()
 		// Lights that should be baked need to check for interactions to track unbuilt state correctly
 		|| Proxy->HasStaticLighting()
 		// ES2 path supports dynamic point lights in the base pass using forward rendering, so we need to know the primitives
-		|| (Scene->GetFeatureLevel() < ERHIFeatureLevel::SM4 && Proxy->GetLightType() == LightType_Point))
+		|| (Scene->GetFeatureLevel() < ERHIFeatureLevel::SM4 && Proxy->GetLightType() == LightType_Point && Proxy->IsMovable()))
 	{
 		// Add the light to the scene's light octree.
 		Scene->LightOctree.AddElement(LightSceneInfoCompact);
diff --git a/Engine/Source/Runtime/Renderer/Private/SceneCore.cpp b/Engine/Source/Runtime/Renderer/Private/SceneCore.cpp
--- a/Engine/Source/Runtime/Renderer/Private/SceneCore.cpp
+++ b/Engine/Source/Runtime/Renderer/Private/SceneCore.cpp
@@ -198,7 +198,7 @@ FLightPrimitiveInteraction::FLightPrimitiveInteraction(
 		PrevPrimitiveLink = &LightSceneInfo->DynamicPrimitiveList;
 
 		// ES2 dynamic point lights
-		if (PrimitiveSceneInfo->Scene->GetFeatureLevel() < ERHIFeatureLevel::SM4 && LightSceneInfo->Proxy->GetLightType() == LightType_Point)
+		if (PrimitiveSceneInfo->Scene->GetFeatureLevel() < ERHIFeatureLevel::SM4 && LightSceneInfo->Proxy->GetLightType() == LightType_Point && LightSceneInfo->Proxy->IsMovable())
 		{
 			bES2DynamicPointLight = true;
 			PrimitiveSceneInfo->NumES2DynamicPointLights++;
