Prevent mosaic rendering during scene captures.

UE-11518, UE-15212

diff --git a/Engine/Source/Runtime/Renderer/Private/ForwardBasePassRendering.cpp b/Engine/Source/Runtime/Renderer/Private/ForwardBasePassRendering.cpp
--- a/Engine/Source/Runtime/Renderer/Private/ForwardBasePassRendering.cpp
+++ b/Engine/Source/Runtime/Renderer/Private/ForwardBasePassRendering.cpp
@@ -263,6 +263,7 @@ public:
 
 		const bool bIsLitMaterial = Parameters.ShadingModel != MSM_Unlit;
 		const FScene* Scene = Parameters.PrimitiveSceneProxy ? Parameters.PrimitiveSceneProxy->GetPrimitiveSceneInfo()->Scene : NULL;
+		const bool bUseMosaic = IsMobileHDR32bpp() && !View.bIsSceneCapture;
 
 		TBasePassForForwardShadingDrawingPolicy<LightMapPolicyType, NumDynamicPointLights> DrawingPolicy(
 			Parameters.Mesh.VertexFactory,
@@ -274,6 +275,7 @@ public:
 			Parameters.ShadingModel != MSM_Unlit && Scene && Scene->ShouldRenderSkylight(),
 			View.Family->EngineShowFlags.ShaderComplexity,
 			View.GetFeatureLevel(),
+			bUseMosaic,
 			Parameters.bEditorCompositeDepthTest
 			);
 		RHICmdList.BuildAndSetLocalBoundShaderState(DrawingPolicy.GetBoundShaderStateInput(View.GetFeatureLevel()));
diff --git a/Engine/Source/Runtime/Renderer/Private/ForwardBasePassRendering.h b/Engine/Source/Runtime/Renderer/Private/ForwardBasePassRendering.h
--- a/Engine/Source/Runtime/Renderer/Private/ForwardBasePassRendering.h
+++ b/Engine/Source/Runtime/Renderer/Private/ForwardBasePassRendering.h
@@ -328,6 +328,7 @@ public:
 		bool bInEnableSkyLight,
 		bool bOverrideWithShaderComplexity,
 		ERHIFeatureLevel::Type FeatureLevel,
+		bool bUseMosaic,
 		bool bInEnableEditorPrimitiveDepthTest = false
 		):
 		FMeshDrawingPolicy(InVertexFactory,InMaterialRenderProxy,InMaterialResource,bOverrideWithShaderComplexity),
@@ -336,7 +337,7 @@ public:
 		SceneTextureMode(InSceneTextureMode),
 		bEnableEditorPrimitiveDepthTest(bInEnableEditorPrimitiveDepthTest)
 	{
-		if (IsMobileHDR32bpp())
+		if (bUseMosaic)
 		{
 			VertexShader = InMaterialResource.GetShader<TBasePassForForwardShadingVS<LightMapPolicyType, HDR_LINEAR_64> >(InVertexFactory->GetType());
 
diff --git a/Engine/Source/Runtime/Renderer/Private/ForwardTranslucentRendering.cpp b/Engine/Source/Runtime/Renderer/Private/ForwardTranslucentRendering.cpp
--- a/Engine/Source/Runtime/Renderer/Private/ForwardTranslucentRendering.cpp
+++ b/Engine/Source/Runtime/Renderer/Private/ForwardTranslucentRendering.cpp
@@ -134,6 +134,7 @@ public:
 	{
 		const bool bIsLitMaterial = Parameters.ShadingModel != MSM_Unlit;
 		const FScene* Scene = Parameters.PrimitiveSceneProxy ? Parameters.PrimitiveSceneProxy->GetPrimitiveSceneInfo()->Scene : NULL;
+		const bool bUseMosaic = IsMobileHDR32bpp() && !View.bIsSceneCapture;
 
 		TBasePassForForwardShadingDrawingPolicy<LightMapPolicyType, NumDynamicPointLights> DrawingPolicy(
 			Parameters.Mesh.VertexFactory,
@@ -144,7 +145,8 @@ public:
 			Parameters.TextureMode,
 			Parameters.ShadingModel != MSM_Unlit && Scene && Scene->ShouldRenderSkylight(),
 			View.Family->EngineShowFlags.ShaderComplexity,
-			View.GetFeatureLevel()
+			View.GetFeatureLevel(),
+			bUseMosaic
 			);
 
 		RHICmdList.BuildAndSetLocalBoundShaderState(DrawingPolicy.GetBoundShaderStateInput(View.GetFeatureLevel()));
diff --git a/Engine/Source/Runtime/Renderer/Private/PostProcess/PostProcessTonemap.cpp b/Engine/Source/Runtime/Renderer/Private/PostProcess/PostProcessTonemap.cpp
--- a/Engine/Source/Runtime/Renderer/Private/PostProcess/PostProcessTonemap.cpp
+++ b/Engine/Source/Runtime/Renderer/Private/PostProcess/PostProcessTonemap.cpp
@@ -542,7 +542,7 @@ static uint32 TonemapperGenerateBitmaskMobile(const FViewInfo* RESTRICT View, bo
 
 	uint32 Bitmask = TonemapperGenerateBitmask(View, bGammaOnly, true);
 
-	bool bUseMosaic = IsMobileHDR32bpp();
+	bool bUseMosaic = IsMobileHDR32bpp() && !View->bIsSceneCapture;
 
 	// Must early exit if gamma only.
 	if(Bitmask == TonemapperGammaOnly)
diff --git a/Engine/Source/Runtime/Renderer/Private/SceneRendering.cpp b/Engine/Source/Runtime/Renderer/Private/SceneRendering.cpp
--- a/Engine/Source/Runtime/Renderer/Private/SceneRendering.cpp
+++ b/Engine/Source/Runtime/Renderer/Private/SceneRendering.cpp
@@ -742,7 +742,8 @@ TUniformBufferRef<FViewUniformShaderParameters> FViewInfo::CreateUniformBuffer(
 
 	{
 		// Enables toggle of HDR Mosaic mode without recompile of all PC shaders during ES2 emulation.
-		ViewUniformShaderParameters.HdrMosaic = IsMobileHDR32bpp() ? 1.0f : 0.0f;
+		const bool bUseMosaic = IsMobileHDR32bpp() && !bIsSceneCapture;
+		ViewUniformShaderParameters.HdrMosaic = bUseMosaic ? 1.0f : 0.0f;
 	}
 	
 	FVector2D OneScenePixelUVSize = FVector2D(1.0f / BufferSize.X, 1.0f / BufferSize.Y);
