Merging decal fading

//depot/UE4-UT/Engine/...

to //depot/UE4/Engine/...

diff --git a/Engine/Shaders/DeferredDecal.usf b/Engine/Shaders/DeferredDecal.usf
--- a/Engine/Shaders/DeferredDecal.usf
+++ b/Engine/Shaders/DeferredDecal.usf
@@ -25,6 +25,9 @@ float4x4 WorldToDecal;
 
 // from component to clip space (for decal frustum)
 float4x4 FrustumComponentToClip;
+// fade alpha from screen size fading
+float FadeAlpha;
+
 // only for StencilDecalMaskMain, to get ShaderComplexity for Decals (doesn't look at the DecalMask yet)
 float MaskComparison;
 
@@ -296,6 +299,8 @@ void MainPS(
 	// Z fade
 	Opacity *= saturate( 4 - 4 * abs( SwizzlePos.z * 2 - 1 ) );
 
+	Opacity *= FadeAlpha;
+
 	// RETURN_COLOR not needed unless writing to SceneColor
 	OutTarget0 = float4(Color, Opacity);
 
diff --git a/Engine/Source/Runtime/Engine/Classes/Components/DecalComponent.h b/Engine/Source/Runtime/Engine/Classes/Components/DecalComponent.h
--- a/Engine/Source/Runtime/Engine/Classes/Components/DecalComponent.h
+++ b/Engine/Source/Runtime/Engine/Classes/Components/DecalComponent.h
@@ -29,6 +29,9 @@ class ENGINE_API UDecalComponent : public USceneComponent
 	UPROPERTY(EditAnywhere, BlueprintReadOnly, Category=Decal)
 	int32 SortOrder;
 
+	UPROPERTY(EditAnywhere, BlueprintReadOnly, Category = Decal)
+	float FadeScreenSize;
+
 	/** Sets the sort order for the decal component. Higher values draw later (on top). This will force the decal to reattach */
 	UFUNCTION(BlueprintCallable, Category = "Rendering|Components|Decal")
 	void SetSortOrder(int32 Value);
diff --git a/Engine/Source/Runtime/Engine/Private/Components/DecalComponent.cpp b/Engine/Source/Runtime/Engine/Private/Components/DecalComponent.cpp
--- a/Engine/Source/Runtime/Engine/Private/Components/DecalComponent.cpp
+++ b/Engine/Source/Runtime/Engine/Private/Components/DecalComponent.cpp
@@ -36,7 +36,7 @@ void FDeferredDecalProxy::SetTransform(const FTransform& InComponentToWorld)
 }
 
 UDecalComponent::UDecalComponent(const FObjectInitializer& ObjectInitializer)
-	: Super(ObjectInitializer)
+: Super(ObjectInitializer), FadeScreenSize(0.01)
 {
 }
 
diff --git a/Engine/Source/Runtime/Renderer/Private/CompositionLighting/PostProcessDeferredDecals.cpp b/Engine/Source/Runtime/Renderer/Private/CompositionLighting/PostProcessDeferredDecals.cpp
--- a/Engine/Source/Runtime/Renderer/Private/CompositionLighting/PostProcessDeferredDecals.cpp
+++ b/Engine/Source/Runtime/Renderer/Private/CompositionLighting/PostProcessDeferredDecals.cpp
@@ -24,6 +24,13 @@ static TAutoConsoleVariable<float> CVarStencilSizeThreshold(
 	TEXT("0..1: optimization is enabled, value defines the minimum size (screen space) to trigger the optimization (default 0.1)")
 	);
 
+static TAutoConsoleVariable<float> CVarDecalFadeScreenSizeMultiplier(
+	TEXT("r.Decal.FadeScreenSizeMult"),
+	1.0f,
+	TEXT("Control the per decal fade screen size. Multiplies with the per-decal screen size fade threshold.")
+	TEXT("  Smaller means decals fade less aggressively.")
+	);
+
 enum ERenderTargetMode
 {
 	RTM_Unknown = -1,
@@ -172,9 +179,10 @@ struct FTransientDecalRenderData
 	const FMaterial* MaterialResource;
 	const FDeferredDecalProxy* DecalProxy;
 	bool bHasNormal;
+	float FadeAlpha;
 
 	FTransientDecalRenderData(const FScene& InScene, FDeferredDecalProxy* InDecalProxy)
-		: DecalProxy(InDecalProxy)
+		: DecalProxy(InDecalProxy), FadeAlpha(1.0f)
 	{
 		MaterialProxy = InDecalProxy->DecalMaterial->GetRenderProxy(InDecalProxy->bOwnerSelected);
 		MaterialResource = MaterialProxy->GetMaterial(InScene.GetFeatureLevel());
@@ -547,10 +555,11 @@ public:
 	{
 		ScreenToDecal.Bind(Initializer.ParameterMap,TEXT("ScreenToDecal"));
 		DecalToWorld.Bind(Initializer.ParameterMap,TEXT("DecalToWorld"));
+		FadeAlpha.Bind(Initializer.ParameterMap, TEXT("FadeAlpha"));
 		WorldToDecal.Bind(Initializer.ParameterMap,TEXT("WorldToDecal"));
 	}
 
-	void SetParameters(FRHICommandList& RHICmdList, const FSceneView& View, const FMaterialRenderProxy* MaterialProxy, const FDeferredDecalProxy& DecalProxy)
+	void SetParameters(FRHICommandList& RHICmdList, const FSceneView& View, const FMaterialRenderProxy* MaterialProxy, const FDeferredDecalProxy& DecalProxy, const float FadeAlphaValue=1.0f)
 	{
 		const FPixelShaderRHIParamRef ShaderRHI = GetPixelShader();
 
@@ -588,19 +597,21 @@ public:
 			SetShaderValue(RHICmdList, ShaderRHI, DecalToWorld, DecalToWorldValue);
 		}
 
+		SetShaderValue(RHICmdList, ShaderRHI, FadeAlpha, FadeAlphaValue);
 		SetShaderValue(RHICmdList, ShaderRHI, WorldToDecal, WorldToComponent);
 	}
 
 	virtual bool Serialize(FArchive& Ar) override
 	{
 		bool bShaderHasOutdatedParameters = FMaterialShader::Serialize(Ar);
-		Ar << ScreenToDecal << DecalToWorld << WorldToDecal;
+		Ar << ScreenToDecal << DecalToWorld << WorldToDecal << FadeAlpha;
 		return bShaderHasOutdatedParameters;
 	}
 
 private:
 	FShaderParameter ScreenToDecal;
 	FShaderParameter DecalToWorld;
+	FShaderParameter FadeAlpha;
 	FShaderParameter WorldToDecal;
 };
 
@@ -710,8 +721,8 @@ void SetShader(const FRenderingCompositePassContext& Context, bool bShaderComple
 		// first Bind, then SetParameters()
 		Context.RHICmdList.SetLocalBoundShaderState(Context.RHICmdList.BuildLocalBoundShaderState(GetVertexDeclarationFVector3(), VertexShader->GetVertexShader(), FHullShaderRHIRef(), FDomainShaderRHIRef(), PixelShader->GetPixelShader(), FGeometryShaderRHIRef()));
 
-		PixelShader->SetParameters(Context.RHICmdList, View, DecalData.MaterialProxy, *DecalData.DecalProxy);
- 	}
+		PixelShader->SetParameters(Context.RHICmdList, View, DecalData.MaterialProxy, *DecalData.DecalProxy, DecalData.FadeAlpha);
+	}
 }
 
 bool RenderPreStencil(FRenderingCompositePassContext& Context, const FMaterialShaderMap* MaterialShaderMap, const FMatrix& ComponentToWorldMatrix, const FMatrix& FrustumComponentToClip)
@@ -1050,7 +1061,7 @@ void FRCPassPostProcessDeferredDecals::Process(FRenderingCompositePassContext& C
 				ComponentToWorldMatrix.GetScaledAxis(EAxis::Z).SizeSquared() * FMath::Square(GDefaultDecalSize.Z));
 
 		// can be optimized as the test is too conservative (sphere instead of OBB)
-			if (ConservativeRadius < SMALL_NUMBER || !View.ViewFrustum.IntersectSphere(ComponentToWorldMatrix.GetOrigin(), ConservativeRadius))
+		if(ConservativeRadius < SMALL_NUMBER || !View.ViewFrustum.IntersectSphere(ComponentToWorldMatrix.GetOrigin(), ConservativeRadius))
 		{
 			bIsShown = false;
 		}
@@ -1058,11 +1069,29 @@ void FRCPassPostProcessDeferredDecals::Process(FRenderingCompositePassContext& C
 		if (bIsShown)
 		{
 			FTransientDecalRenderData Data(Scene, DecalProxy);
+			if (Data.DecalProxy->Component->FadeScreenSize != 0.0f)
+			{
+				float FadeMultiplier = CVarDecalFadeScreenSizeMultiplier.GetValueOnRenderThread();
+
+				float Distance = (View.ViewMatrices.ViewOrigin - ComponentToWorldMatrix.GetOrigin()).Size();
+				float Radius = ComponentToWorldMatrix.GetMaximumAxisScale();
+				float CurrentScreenSize = ((Radius / Distance) * FadeMultiplier);
+
+				// fading coefficient needs to increase with increasing field of view and decrease with increasing resolution
+				// FadeCoeffScale is an empirically determined constant to bring us back roughly to fraction of screen size for FadeScreenSize
+				const float FadeCoeffScale = 600.0f;
+				float FOVFactor = ((2.0f/View.ViewMatrices.ProjMatrix.M[0][0]) / View.ViewRect.Width()) * FadeCoeffScale;
+				float FadeCoeff = Data.DecalProxy->Component->FadeScreenSize * FOVFactor;
+				float FadeRange = FadeCoeff * 0.5f;
+
+				float Alpha = (CurrentScreenSize - FadeCoeff) / FadeRange;
+				Data.FadeAlpha = FMath::Min(Alpha, 1.0f);
+			}
 
 			EDecalRenderStage LocalDecalRenderStage = ComputeRenderStage(Data.DecalBlendMode);
 
 			// we could do this test earlier to avoid the decal intersection but getting DecalBlendMode also costs
-			if (Context.View.Family->EngineShowFlags.ShaderComplexity || DecalRenderStage == LocalDecalRenderStage)
+			if (Context.View.Family->EngineShowFlags.ShaderComplexity || DecalRenderStage == LocalDecalRenderStage && Data.FadeAlpha>0.0f)
 			{
 				SortedDecals.Add(Data);
 			}
