fixed possible crash reported by [REDACTED_PERSON]

diff --git a/Engine/Source/Runtime/Renderer/Private/LightPropagationVolumeVisualisation.cpp b/Engine/Source/Runtime/Renderer/Private/LightPropagationVolumeVisualisation.cpp
--- a/Engine/Source/Runtime/Renderer/Private/LightPropagationVolumeVisualisation.cpp
+++ b/Engine/Source/Runtime/Renderer/Private/LightPropagationVolumeVisualisation.cpp
@@ -104,7 +104,7 @@ public:
 		}
 
 		LpvVolumeTextureSampler.Bind(Initializer.ParameterMap, TEXT("gLpv3DTextureSampler"));
-		for ( int i=0; i<3; i++ )
+		for ( int i = 0; i < NUM_GV_TEXTURES; i++ ) 
 		{
 			GvBufferSRVParameters[i].Bind( Initializer.ParameterMap, LpvGvVolumeTextureSRVNames[i] );
 		}
@@ -121,8 +121,8 @@ public:
 	{
 		FPixelShaderRHIParamRef ShaderRHI = GetPixelShader();
 		FGlobalShader::SetParameters(RHICmdList, ShaderRHI, View); 
-
-		for ( int i=0; i<7; i++ )
+		
+		for ( int i = 0; i < 7; i++ )
 		{
 			FTextureRHIParamRef LpvBufferSrv = LPV->LpvVolumeTextures[ 1-LPV->mWriteBufferIndex ][i]->GetRenderTargetItem().ShaderResourceTexture;
 			if ( LpvBufferSRVParameters[i].IsBound() )
@@ -131,8 +131,8 @@ public:
 			}
 			SetTextureParameter(RHICmdList, ShaderRHI, LpvBufferSRVParameters[i], LpvVolumeTextureSampler, TStaticSamplerState<SF_Bilinear,AM_Clamp,AM_Clamp,AM_Clamp>::GetRHI(), LpvBufferSrv );
 		}
-
-		for ( int i=0; i<3; i++ )
+		
+		for ( int i = 0; i < NUM_GV_TEXTURES; i++ ) 
 		{
 			FTextureRHIParamRef GvBufferSrv = LPV->GvVolumeTextures[i]->GetRenderTargetItem().ShaderResourceTexture;
 			if ( GvBufferSRVParameters[i].IsBound() )
@@ -149,15 +149,15 @@ public:
 	virtual bool Serialize( FArchive& Ar ) override
 	{
 		bool bShaderHasOutdatedParameters = FLpvVisualiseBase::Serialize( Ar );
-
-		for ( int i=0; i<7; i++ )
+		
+		for ( int i = 0; i < 7; i++ )
 		{
 			Ar << LpvBufferSRVParameters[i];
 		}
 
 		Ar << LpvVolumeTextureSampler;
 
-		for ( int i=0; i<3; i++ )
+		for ( int i = 0; i < NUM_GV_TEXTURES; i++ ) 
 		{
 			Ar << GvBufferSRVParameters[i];
 		}
@@ -173,15 +173,15 @@ public:
 	{
 		// TODO: Is this necessary here?
 		FPixelShaderRHIParamRef ShaderRHI = GetPixelShader();
-		for ( int i=0;i<7; i++ )
+		for ( int i = 0; i < 7; i++ )
 		{
 			if ( LpvBufferSRVParameters[i].IsBound() )
 			{
 				RHICmdList.SetShaderTexture(ShaderRHI, LpvBufferSRVParameters[i].GetBaseIndex(), FTextureRHIParamRef());
 			}
 		}
-
-		for ( int i=0;i<3; i++ )
+		
+		for ( int i = 0; i < NUM_GV_TEXTURES; i++ ) 
 		{
 			if ( GvBufferSRVParameters[i].IsBound() )
 			{
