Compile error fix.

diff --git a/Engine/Source/Runtime/Renderer/Private/CompositionLighting/PostProcessDeferredDecals.cpp b/Engine/Source/Runtime/Renderer/Private/CompositionLighting/PostProcessDeferredDecals.cpp
--- a/Engine/Source/Runtime/Renderer/Private/CompositionLighting/PostProcessDeferredDecals.cpp
+++ b/Engine/Source/Runtime/Renderer/Private/CompositionLighting/PostProcessDeferredDecals.cpp
@@ -1091,7 +1091,7 @@ void FRCPassPostProcessDeferredDecals::Process(FRenderingCompositePassContext& C
 			EDecalRenderStage LocalDecalRenderStage = ComputeRenderStage(Data.DecalBlendMode);
 
 			// we could do this test earlier to avoid the decal intersection but getting DecalBlendMode also costs
-			if (Context.View.Family->EngineShowFlags.ShaderComplexity || DecalRenderStage == LocalDecalRenderStage && Data.FadeAlpha>0.0f)
+			if (Context.View.Family->EngineShowFlags.ShaderComplexity || (DecalRenderStage == LocalDecalRenderStage && Data.FadeAlpha>0.0f) )
 			{
 				SortedDecals.Add(Data);
 			}
