Faster Indirect Lighting Cache point sample transition, and unbuilt meshes transition much faster to improve lighting preview during meshing

diff --git a/Engine/Source/Runtime/Renderer/Private/IndirectLightingCache.cpp b/Engine/Source/Runtime/Renderer/Private/IndirectLightingCache.cpp
--- a/Engine/Source/Runtime/Renderer/Private/IndirectLightingCache.cpp
+++ b/Engine/Source/Runtime/Renderer/Private/IndirectLightingCache.cpp
@@ -56,7 +56,7 @@ static FAutoConsoleVariableRef CVarCacheUpdateEveryFrame(
 	ECVF_RenderThreadSafe
 	);
 
-float GSingleSampleTransitionSpeed = 400;
+float GSingleSampleTransitionSpeed = 800;
 static FAutoConsoleVariableRef CVarSingleSampleTransitionSpeed(
 	TEXT("r.Cache.SampleTransitionSpeed"),
 	GSingleSampleTransitionSpeed,
@@ -317,10 +317,10 @@ FIndirectLightingCacheAllocation* FIndirectLightingCache::AllocatePrimitive(cons
 {
 	const bool bPointSample = PrimitiveSceneInfo->Proxy->GetIndirectLightingCacheQuality() == ILCQ_Point || bUnbuiltPreview;
 	const int32 BlockSize = bPointSample ? 1 : GLightingCacheMovableObjectAllocationSize;
-	return PrimitiveAllocations.Add(PrimitiveSceneInfo->PrimitiveComponentId, CreateAllocation(BlockSize, PrimitiveSceneInfo->Proxy->GetBounds(), bPointSample));
+	return PrimitiveAllocations.Add(PrimitiveSceneInfo->PrimitiveComponentId, CreateAllocation(BlockSize, PrimitiveSceneInfo->Proxy->GetBounds(), bPointSample, bUnbuiltPreview));
 }
 
-FIndirectLightingCacheAllocation* FIndirectLightingCache::CreateAllocation(int32 BlockSize, const FBoxSphereBounds& Bounds, bool bPointSample)
+FIndirectLightingCacheAllocation* FIndirectLightingCache::CreateAllocation(int32 BlockSize, const FBoxSphereBounds& Bounds, bool bPointSample, bool bUnbuiltPreview)
 {
 	check(BlockSize > 1 || bPointSample);
 
@@ -340,7 +340,7 @@ FIndirectLightingCacheAllocation* FIndirectLightingCache::CreateAllocation(int32
 		CalculateBlockScaleAndAdd(NewBlock.MinTexel, NewBlock.TexelSize, NewBlock.Min, NewBlock.Size, Scale, Add, MinUV, MaxUV);
 
 		VolumeBlocks.Add(NewBlock.MinTexel, NewBlock);
-		NewAllocation->SetParameters(NewBlock.MinTexel, NewBlock.TexelSize, Scale, Add, MinUV, MaxUV, bPointSample);
+		NewAllocation->SetParameters(NewBlock.MinTexel, NewBlock.TexelSize, Scale, Add, MinUV, MaxUV, bPointSample, bUnbuiltPreview);
 	}
 
 	return NewAllocation;
@@ -441,6 +441,7 @@ void FIndirectLightingCache::UpdateCacheAllocation(
 	const FBoxSphereBounds& Bounds, 
 	int32 BlockSize,
 	bool bPointSample,
+	bool bUnbuiltPreview,
 	FIndirectLightingCacheAllocation*& Allocation, 
 	TMap<FIntVector, FBlockUpdateInfo>& BlocksToUpdate,
 	TArray<FIndirectLightingCacheAllocation*>& TransitionsOverTimeToUpdate)
@@ -467,7 +468,7 @@ void FIndirectLightingCache::UpdateCacheAllocation(
 			FVector MaxUV;
 			CalculateBlockScaleAndAdd(Allocation->MinTexel, Allocation->AllocationTexelSize, NewMin, NewSize, NewScale, NewAdd, MinUV, MaxUV);
 
-			Allocation->SetParameters(Allocation->MinTexel, Allocation->AllocationTexelSize, NewScale, NewAdd, MinUV, MaxUV, bPointSample);
+			Allocation->SetParameters(Allocation->MinTexel, Allocation->AllocationTexelSize, NewScale, NewAdd, MinUV, MaxUV, bPointSample, bUnbuiltPreview);
 			BlocksToUpdate.Add(Block.MinTexel, FBlockUpdateInfo(Block, Allocation));
 		}
 
@@ -479,7 +480,7 @@ void FIndirectLightingCache::UpdateCacheAllocation(
 	else
 	{
 		delete Allocation;
-		Allocation = CreateAllocation(BlockSize, Bounds, bPointSample);
+		Allocation = CreateAllocation(BlockSize, Bounds, bPointSample, bUnbuiltPreview);
 
 		if (Allocation->IsValid())
 		{
@@ -530,7 +531,7 @@ void FIndirectLightingCache::UpdateCachePrimitive(
 			const int32 BlockSize = bPointSample ? 1 : GLightingCacheMovableObjectAllocationSize;
 
 			// Light with the cumulative bounds of the entire attachment group
-			UpdateCacheAllocation(PrimitiveSceneInfo->GetAttachmentGroupBounds(), BlockSize, bPointSample, PrimitiveAllocation, BlocksToUpdate, TransitionsOverTimeToUpdate);
+			UpdateCacheAllocation(PrimitiveSceneInfo->GetAttachmentGroupBounds(), BlockSize, bPointSample, bUnbuiltPreview, PrimitiveAllocation, BlocksToUpdate, TransitionsOverTimeToUpdate);
 
 			// Cache the primitive allocation pointer on the FPrimitiveSceneInfo for base pass rendering
 			PrimitiveSceneInfo->IndirectLightingCacheAllocation = PrimitiveAllocation;
@@ -573,7 +574,8 @@ void FIndirectLightingCache::UpdateTransitionsOverTime(const TArray<FIndirectLig
 
 		if (TransitionDistance > DELTA)
 		{
-			// Compute a frame rate independent transition by maintaining a constant world space speed between the current sample position and the target position
+			// Transition faster for unbuilt meshes which is important for meshing visualization
+			const float EffectiveTransitionSpeed = GSingleSampleTransitionSpeed * (Allocation->bUnbuiltPreview ? 4 : 1);
 			const float LerpFactor = FMath::Clamp(GSingleSampleTransitionSpeed * DeltaWorldTime / TransitionDistance, 0.0f, 1.0f);
 			Allocation->SingleSamplePosition = FMath::Lerp(Allocation->SingleSamplePosition, Allocation->TargetPosition, LerpFactor);
 
diff --git a/Engine/Source/Runtime/Renderer/Private/ScenePrivate.h b/Engine/Source/Runtime/Renderer/Private/ScenePrivate.h
--- a/Engine/Source/Runtime/Renderer/Private/ScenePrivate.h
+++ b/Engine/Source/Runtime/Renderer/Private/ScenePrivate.h
@@ -1125,6 +1125,7 @@ private:
 		const FBoxSphereBounds& Bounds, 
 		int32 BlockSize,
 		bool bPointSample,
+		bool bUnbuiltPreview,
 		FIndirectLightingCacheAllocation*& Allocation, 
 		TMap<FIntVector, FBlockUpdateInfo>& BlocksToUpdate,
 		TArray<FIndirectLightingCacheAllocation*>& TransitionsOverTimeToUpdate);
@@ -1148,7 +1149,7 @@ private:
 	void UpdateTransitionsOverTime(const TArray<FIndirectLightingCacheAllocation*>& TransitionsOverTimeToUpdate, float DeltaWorldTime) const;
 
 	/** Creates an allocation to be used outside the indirect lighting cache and a block to be used internally. */
-	FIndirectLightingCacheAllocation* CreateAllocation(int32 BlockSize, const FBoxSphereBounds& Bounds, bool bPointSample);
+	FIndirectLightingCacheAllocation* CreateAllocation(int32 BlockSize, const FBoxSphereBounds& Bounds, bool bPointSample, bool bUnbuiltPreview);
 
 	/** Block accessors. */
 	FIndirectLightingCacheBlock& FindBlock(FIntVector TexelMin);
diff --git a/Engine/Source/Runtime/Renderer/Public/PrimitiveSceneInfo.h b/Engine/Source/Runtime/Renderer/Public/PrimitiveSceneInfo.h
--- a/Engine/Source/Runtime/Renderer/Public/PrimitiveSceneInfo.h
+++ b/Engine/Source/Runtime/Renderer/Public/PrimitiveSceneInfo.h
@@ -29,7 +29,8 @@ public:
 		CurrentSkyBentNormal(FVector4(0, 0, 1, 1)),
 		bHasEverUpdatedSingleSample(false),
 		bPointSample(true),
-		bIsDirty(false)
+		bIsDirty(false),
+		bUnbuiltPreview(false)
 	{
 		for (int32 VectorIndex = 0; VectorIndex < ARRAY_COUNT(TargetSamplePacked); VectorIndex++)
 		{
@@ -89,6 +90,8 @@ public:
 	/** Whether the primitive allocation is dirty and should be updated regardless of having moved. */
 	bool bIsDirty;
 
+	bool bUnbuiltPreview;
+
 	void SetDirty() 
 	{ 
 		bIsDirty = true; 
@@ -99,7 +102,7 @@ public:
 		return MinTexel.X >= 0 && MinTexel.Y >= 0 && MinTexel.Z >= 0 && AllocationTexelSize > 0;
 	}
 
-	void SetParameters(FIntVector InMinTexel, int32 InAllocationTexelSize, FVector InScale, FVector InAdd, FVector InMinUV, FVector InMaxUV, bool bInPointSample)
+	void SetParameters(FIntVector InMinTexel, int32 InAllocationTexelSize, FVector InScale, FVector InAdd, FVector InMinUV, FVector InMaxUV, bool bInPointSample, bool bInUnbuiltPreview)
 	{
 		Add = InAdd;
 		Scale = InScale;
@@ -109,6 +112,7 @@ public:
 		AllocationTexelSize = InAllocationTexelSize;
 		bIsDirty = false;
 		bPointSample = bInPointSample;
+		bUnbuiltPreview = bInUnbuiltPreview;
 	}
 };
 
