Rendering: Improved the error message logging in FScene::Release() when a component still has a reference to the dying scene to list the associated asset for any kind of component, not just UStaticMeshComponent

diff --git a/Engine/Source/Runtime/Renderer/Private/RendererScene.cpp b/Engine/Source/Runtime/Renderer/Private/RendererScene.cpp
--- a/Engine/Source/Runtime/Renderer/Private/RendererScene.cpp
+++ b/Engine/Source/Runtime/Renderer/Private/RendererScene.cpp
@@ -1915,13 +1915,13 @@ void FScene::Release()
 		for (auto* ActorComponent : TObjectRange<UActorComponent>())
 		{
 			if ( !ensureMsg(!ActorComponent->IsRegistered() || ActorComponent->GetScene() != this, 
-					*FString::Printf(TEXT("Component Name: %s World Name: %s Component Mesh: %s"), 
+					*FString::Printf(TEXT("Component Name: %s World Name: %s Component Asset: %s"), 
 										*ActorComponent->GetFullName(), 
 										*GetWorld()->GetFullName(), 
-										Cast<UStaticMeshComponent>(ActorComponent) ? *CastChecked<UStaticMeshComponent>(ActorComponent)->StaticMesh->GetFullName() : TEXT("Not a static mesh"))) )
+										*ActorComponent->AdditionalStatObject()->GetPathName())) )
 			{
 				bTriggeredOnce = true;
-				break;	
+				break;
 			}
 		}
 	}
