added printout to make LPV debugging easier, also good to demonstrate Blendables (needed for docs)

diff --git a/Engine/Source/Runtime/Renderer/Private/CompositionLighting/PostProcessLpvIndirect.cpp b/Engine/Source/Runtime/Renderer/Private/CompositionLighting/PostProcessLpvIndirect.cpp
--- a/Engine/Source/Runtime/Renderer/Private/CompositionLighting/PostProcessLpvIndirect.cpp
+++ b/Engine/Source/Runtime/Renderer/Private/CompositionLighting/PostProcessLpvIndirect.cpp
@@ -399,4 +399,94 @@ FPooledRenderTargetDesc FRCPassPostProcessLpvIndirect::ComputeOutputDesc(EPassOu
 	Ret.DebugName = TEXT("LpvIndirect");
 
 	return Ret;
-}
\ No newline at end of file
+}
+
+void FRCPassPostProcessVisualizeLPV::Process(FRenderingCompositePassContext& Context)
+{
+	SCOPED_DRAW_EVENT(Context.RHICmdList, VisualizeLPV);
+
+	const FSceneView& View = Context.View;
+	const FSceneViewFamily& ViewFamily = *(View.Family);
+	
+//	const FSceneRenderTargetItem& DestRenderTarget = PassOutputs[0].RequestSurface(Context);
+	const TRefCountPtr<IPooledRenderTarget> RenderTarget = GetInput(ePId_Input0)->GetOutput()->PooledRenderTarget;
+	const FSceneRenderTargetItem& DestRenderTarget = RenderTarget->GetRenderTargetItem();
+
+	// Set the view family's render target/viewport.
+	SetRenderTarget(Context.RHICmdList, DestRenderTarget.TargetableTexture, FTextureRHIRef());
+
+	{
+		// this is a helper class for FCanvas to be able to get screen size
+		class FRenderTargetTemp : public FRenderTarget
+		{
+		public:
+			const FSceneView& View;
+			const FTexture2DRHIRef Texture;
+
+			FRenderTargetTemp(const FSceneView& InView, const FTexture2DRHIRef InTexture)
+				: View(InView), Texture(InTexture)
+			{
+			}
+			virtual FIntPoint GetSizeXY() const
+			{
+				return View.ViewRect.Size();
+			};
+			virtual const FTexture2DRHIRef& GetRenderTargetTexture() const
+			{
+				return Texture;
+			}
+		} TempRenderTarget(View, (const FTexture2DRHIRef&)DestRenderTarget.TargetableTexture);
+
+		FCanvas Canvas(&TempRenderTarget, NULL, ViewFamily.CurrentRealTime, ViewFamily.CurrentWorldTime, ViewFamily.DeltaWorldTime, View.GetFeatureLevel());
+
+		float X = 30;
+		float Y = 28;
+		const float YStep = 14;
+		const float ColumnWidth = 250;
+
+
+		const FLightPropagationVolumeSettings& Dest = View.FinalPostProcessSettings.BlendableManager.GetSingleFinalDataConst<FLightPropagationVolumeSettings>();
+
+#define ENTRY(name)\
+		Canvas.DrawShadowedString( X, Y += YStep, TEXT(#name) TEXT(":"), GetStatsFont(), FLinearColor(1, 1, 1));\
+		Canvas.DrawShadowedString( X + ColumnWidth, Y, *FString::Printf(TEXT("%g"), Dest.name), GetStatsFont(), FLinearColor(1, 1, 1));
+
+		ENTRY(LPVIntensity)
+		ENTRY(LPVVplInjectionBias)
+		ENTRY(LPVSize)
+		ENTRY(LPVSecondaryOcclusionIntensity)
+		ENTRY(LPVSecondaryBounceIntensity)
+		ENTRY(LPVGeometryVolumeBias)
+		ENTRY(LPVEmissiveInjectionIntensity)
+		ENTRY(LPVDirectionalOcclusionIntensity)
+		ENTRY(LPVDirectionalOcclusionRadius)
+		ENTRY(LPVDiffuseOcclusionExponent)
+		ENTRY(LPVSpecularOcclusionExponent)
+		ENTRY(LPVDiffuseOcclusionIntensity)
+		ENTRY(LPVSpecularOcclusionIntensity)
+#undef ENTRY
+
+		Canvas.Flush_RenderThread(Context.RHICmdList);
+	}
+
+	Context.RHICmdList.CopyToResolveTarget(DestRenderTarget.TargetableTexture, DestRenderTarget.ShaderResourceTexture, false, FResolveParams());
+	
+	// to satify following passws
+	FRenderingCompositeOutput* Output = GetOutput(ePId_Output0);
+	
+	Output->PooledRenderTarget = RenderTarget;
+}
+
+FPooledRenderTargetDesc FRCPassPostProcessVisualizeLPV::ComputeOutputDesc(EPassOutputId InPassOutputId) const
+{
+	FPooledRenderTargetDesc Ret = PassInputs[0].GetOutput()->RenderTargetDesc;
+
+	Ret.Reset();
+
+	// we assume this pass is additively blended with the scene color so this data is not needed
+//	FPooledRenderTargetDesc Ret;
+
+	Ret.DebugName = TEXT("VisualizeLPV");
+
+	return Ret;
+}
diff --git a/Engine/Source/Runtime/Renderer/Private/CompositionLighting/PostProcessLpvIndirect.h b/Engine/Source/Runtime/Renderer/Private/CompositionLighting/PostProcessLpvIndirect.h
--- a/Engine/Source/Runtime/Renderer/Private/CompositionLighting/PostProcessLpvIndirect.h
+++ b/Engine/Source/Runtime/Renderer/Private/CompositionLighting/PostProcessLpvIndirect.h
@@ -30,4 +30,19 @@ private:
 	void DoDirectionalOcclusionPass(FRenderingCompositePassContext& Context) const;
 };
 
+// ePId_Input0: SceneColor
+// derives from TRenderingCompositePassBase<InputCount, OutputCount> 
+class FRCPassPostProcessVisualizeLPV: public TRenderingCompositePassBase<1, 1>
+{
+public:
+	// interface FRenderingCompositePass ---------
+	virtual const TCHAR* GetDebugName() { return TEXT("FRCPassPostProcessVisualizeLPV"); }
+	virtual void Process(FRenderingCompositePassContext& Context) override;
+	virtual void Release() override { delete this; }
+	virtual FPooledRenderTargetDesc ComputeOutputDesc(EPassOutputId InPassOutputId) const override;
+	virtual bool FrameBufferBlendingWithInput0() const { return true; }
+};
+
+
+
 
diff --git a/Engine/Source/Runtime/Renderer/Private/PostProcess/PostProcessing.cpp b/Engine/Source/Runtime/Renderer/Private/PostProcess/PostProcessing.cpp
--- a/Engine/Source/Runtime/Renderer/Private/PostProcess/PostProcessing.cpp
+++ b/Engine/Source/Runtime/Renderer/Private/PostProcess/PostProcessing.cpp
@@ -1475,6 +1475,12 @@ void FPostProcessing::Process(FRHICommandListImmediate& RHICmdList, FViewInfo& V
 			Context.FinalOutput = FRenderingCompositeOutputRef(Node);
 		}
 
+		if(View.Family->EngineShowFlags.VisualizeLPV && !View.Family->EngineShowFlags.VisualizeHDR)
+		{
+			FRenderingCompositePass* Node = Context.Graph.RegisterPass(new(FMemStack::Get()) FRCPassPostProcessVisualizeLPV());
+			Node->SetInput(ePId_Input0, Context.FinalOutput);
+			Context.FinalOutput = FRenderingCompositeOutputRef(Node);
+		}
 
 		// Show the selection outline if it is in the editor and we arent in wireframe 
 		// If the engine is in demo mode and game view is on we also do not show the selection outline
