Fix Linux build.

diff --git a/Engine/Source/Runtime/Renderer/Private/PostProcess/PostProcessing.cpp b/Engine/Source/Runtime/Renderer/Private/PostProcess/PostProcessing.cpp
--- a/Engine/Source/Runtime/Renderer/Private/PostProcess/PostProcessing.cpp
+++ b/Engine/Source/Runtime/Renderer/Private/PostProcess/PostProcessing.cpp
@@ -45,6 +45,7 @@
 #include "PostProcessMorpheus.h"
 #include "IHeadMountedDisplay.h"
 #include "BufferVisualizationData.h"
+#include "PostProcessLpvIndirect.h"
 
 
 /** The global center for all post processing activities. */
