Back out changelist 2588598 : "Prevent mosaic rendering during scene captures. UE-11518, UE-15212"

Need a better plan for dealing with drawlists.

diff --git a/Engine/Source/Runtime/Renderer/Private/ForwardBasePassRendering.cpp b/Engine/Source/Runtime/Renderer/Private/ForwardBasePassRendering.cpp
--- a/Engine/Source/Runtime/Renderer/Private/ForwardBasePassRendering.cpp
+++ b/Engine/Source/Runtime/Renderer/Private/ForwardBasePassRendering.cpp
@@ -263,7 +263,6 @@ public:
 
 		const bool bIsLitMaterial = Parameters.ShadingModel != MSM_Unlit;
 		const FScene* Scene = Parameters.PrimitiveSceneProxy ? Parameters.PrimitiveSceneProxy->GetPrimitiveSceneInfo()->Scene : NULL;
-		const bool bUseMosaic = IsMobileHDR32bpp() && !View.bIsSceneCapture;
 
 		TBasePassForForwardShadingDrawingPolicy<LightMapPolicyType, NumDynamicPointLights> DrawingPolicy(
 			Parameters.Mesh.VertexFactory,
@@ -275,7 +274,6 @@ public:
 			Parameters.ShadingModel != MSM_Unlit && Scene && Scene->ShouldRenderSkylight(),
 			View.Family->EngineShowFlags.ShaderComplexity,
 			View.GetFeatureLevel(),
-			bUseMosaic,
 			Parameters.bEditorCompositeDepthTest
 			);
 		RHICmdList.BuildAndSetLocalBoundShaderState(DrawingPolicy.GetBoundShaderStateInput(View.GetFeatureLevel()));
diff --git a/Engine/Source/Runtime/Renderer/Private/ForwardBasePassRendering.h b/Engine/Source/Runtime/Renderer/Private/ForwardBasePassRendering.h
--- a/Engine/Source/Runtime/Renderer/Private/ForwardBasePassRendering.h
+++ b/Engine/Source/Runtime/Renderer/Private/ForwardBasePassRendering.h
@@ -328,7 +328,6 @@ public:
 		bool bInEnableSkyLight,
 		bool bOverrideWithShaderComplexity,
 		ERHIFeatureLevel::Type FeatureLevel,
-		bool bUseMosaic,
 		bool bInEnableEditorPrimitiveDepthTest = false
 		):
 		FMeshDrawingPolicy(InVertexFactory,InMaterialRenderProxy,InMaterialResource,bOverrideWithShaderComplexity),
@@ -337,7 +336,7 @@ public:
 		SceneTextureMode(InSceneTextureMode),
 		bEnableEditorPrimitiveDepthTest(bInEnableEditorPrimitiveDepthTest)
 	{
-		if (bUseMosaic)
+		if (IsMobileHDR32bpp())
 		{
 			VertexShader = InMaterialResource.GetShader<TBasePassForForwardShadingVS<LightMapPolicyType, HDR_LINEAR_64> >(InVertexFactory->GetType());
 
diff --git a/Engine/Source/Runtime/Renderer/Private/ForwardTranslucentRendering.cpp b/Engine/Source/Runtime/Renderer/Private/ForwardTranslucentRendering.cpp
--- a/Engine/Source/Runtime/Renderer/Private/ForwardTranslucentRendering.cpp
+++ b/Engine/Source/Runtime/Renderer/Private/ForwardTranslucentRendering.cpp
@@ -134,7 +134,6 @@ public:
 	{
 		const bool bIsLitMaterial = Parameters.ShadingModel != MSM_Unlit;
 		const FScene* Scene = Parameters.PrimitiveSceneProxy ? Parameters.PrimitiveSceneProxy->GetPrimitiveSceneInfo()->Scene : NULL;
-		const bool bUseMosaic = IsMobileHDR32bpp() && !View.bIsSceneCapture;
 
 		TBasePassForForwardShadingDrawingPolicy<LightMapPolicyType, NumDynamicPointLights> DrawingPolicy(
 			Parameters.Mesh.VertexFactory,
@@ -145,8 +144,7 @@ public:
 			Parameters.TextureMode,
 			Parameters.ShadingModel != MSM_Unlit && Scene && Scene->ShouldRenderSkylight(),
 			View.Family->EngineShowFlags.ShaderComplexity,
-			View.GetFeatureLevel(),
-			bUseMosaic
+			View.GetFeatureLevel()
 			);
 
 		RHICmdList.BuildAndSetLocalBoundShaderState(DrawingPolicy.GetBoundShaderStateInput(View.GetFeatureLevel()));
diff --git a/Engine/Source/Runtime/Renderer/Private/PostProcess/PostProcessTonemap.cpp b/Engine/Source/Runtime/Renderer/Private/PostProcess/PostProcessTonemap.cpp
--- a/Engine/Source/Runtime/Renderer/Private/PostProcess/PostProcessTonemap.cpp
+++ b/Engine/Source/Runtime/Renderer/Private/PostProcess/PostProcessTonemap.cpp
@@ -542,7 +542,7 @@ static uint32 TonemapperGenerateBitmaskMobile(const FViewInfo* RESTRICT View, bo
 
 	uint32 Bitmask = TonemapperGenerateBitmask(View, bGammaOnly, true);
 
-	bool bUseMosaic = IsMobileHDR32bpp() && !View->bIsSceneCapture;
+	bool bUseMosaic = IsMobileHDR32bpp();
 
 	// Must early exit if gamma only.
 	if(Bitmask == TonemapperGammaOnly)
diff --git a/Engine/Source/Runtime/Renderer/Private/SceneRendering.cpp b/Engine/Source/Runtime/Renderer/Private/SceneRendering.cpp
--- a/Engine/Source/Runtime/Renderer/Private/SceneRendering.cpp
+++ b/Engine/Source/Runtime/Renderer/Private/SceneRendering.cpp
@@ -742,8 +742,7 @@ TUniformBufferRef<FViewUniformShaderParameters> FViewInfo::CreateUniformBuffer(
 
 	{
 		// Enables toggle of HDR Mosaic mode without recompile of all PC shaders during ES2 emulation.
-		const bool bUseMosaic = IsMobileHDR32bpp() && !bIsSceneCapture;
-		ViewUniformShaderParameters.HdrMosaic = bUseMosaic ? 1.0f : 0.0f;
+		ViewUniformShaderParameters.HdrMosaic = IsMobileHDR32bpp() ? 1.0f : 0.0f;
 	}
 	
 	FVector2D OneScenePixelUVSize = FVector2D(1.0f / BufferSize.X, 1.0f / BufferSize.Y);
