#vr Integrating forward renderer stereo fixes

diff --git a/Engine/Source/Runtime/OpenGLDrv/Private/Android/AndroidOpenGL.cpp b/Engine/Source/Runtime/OpenGLDrv/Private/Android/AndroidOpenGL.cpp
--- a/Engine/Source/Runtime/OpenGLDrv/Private/Android/AndroidOpenGL.cpp
+++ b/Engine/Source/Runtime/OpenGLDrv/Private/Android/AndroidOpenGL.cpp
@@ -119,17 +119,15 @@ void* PlatformGetWindow(FPlatformOpenGLContext* Context, void** AddParam)
 
 bool PlatformBlitToViewport( FPlatformOpenGLDevice* Device, const FOpenGLViewport& Viewport, uint32 BackbufferSizeX, uint32 BackbufferSizeY, bool bPresent,bool bLockToVsync, int32 SyncInterval )
 {
-	if (Viewport.GetCustomPresent())
+	if (bPresent && Viewport.GetCustomPresent())
 	{
-		Viewport.GetCustomPresent()->Present(SyncInterval);
+		bPresent = Viewport.GetCustomPresent()->Present(SyncInterval);
 	}
-	else
+	if (bPresent)
 	{
 		AndroidEGL::GetInstance()->SwapBuffers();
 	}
-//	return true;
-	//Do not want WaitForFrameEventCompletion
-	return false;
+	return bPresent;
 }
 
 void PlatformRenderingContextSetup(FPlatformOpenGLDevice* Device)
diff --git a/Engine/Source/Runtime/Renderer/Private/ForwardShadingRenderer.cpp b/Engine/Source/Runtime/Renderer/Private/ForwardShadingRenderer.cpp
--- a/Engine/Source/Runtime/Renderer/Private/ForwardShadingRenderer.cpp
+++ b/Engine/Source/Runtime/Renderer/Private/ForwardShadingRenderer.cpp
@@ -115,7 +115,7 @@ void FForwardShadingSceneRenderer::Render(FRHICommandListImmediate& RHICmdList)
 	FViewInfo& View = Views[0];
 
 	const bool bGammaSpace = !IsMobileHDR();
-	const bool bRequiresUpscale = ((uint32)ViewFamily.RenderTarget->GetSizeXY().X > ViewFamily.FamilySizeX || (uint32)ViewFamily.RenderTarget->GetSizeXY().Y > ViewFamily.FamilySizeY);
+	const bool bRequiresUpscale = !ViewFamily.bUseSeparateRenderTarget && ((uint32)ViewFamily.RenderTarget->GetSizeXY().X > ViewFamily.FamilySizeX || (uint32)ViewFamily.RenderTarget->GetSizeXY().Y > ViewFamily.FamilySizeY);
 	const bool bRenderToScene = bRequiresUpscale || FSceneRenderer::ShouldCompositeEditorPrimitives(View);
 
 	if (bGammaSpace && !bRenderToScene)
@@ -198,17 +198,6 @@ void FForwardShadingSceneRenderer::Render(FRHICommandListImmediate& RHICmdList)
 		// Finish rendering for each view, or the full stereo buffer if enabled
 		if (ViewFamily.bResolveScene)
 		{
-			if (ViewFamily.EngineShowFlags.StereoRendering)
-			{
-				check(Views.Num() > 1);
-
-				//@todo ES2 stereo post: until we get proper stereo postprocessing for ES2, process the stereo buffer as one view
-				FIntPoint OriginalMax0 = Views[0].ViewRect.Max;
-				Views[0].ViewRect.Max = Views[1].ViewRect.Max;
-				GPostProcessing.ProcessES2(RHICmdList, Views[0], bOnChipSunMask);
-				Views[0].ViewRect.Max = OriginalMax0;
-			}
-			else
 			{
 				SCOPED_DRAW_EVENT(RHICmdList, PostProcessing);
 				SCOPE_CYCLE_COUNTER(STAT_FinishRenderViewTargetTime);
@@ -222,7 +211,10 @@ void FForwardShadingSceneRenderer::Render(FRHICommandListImmediate& RHICmdList)
 	}
 	else if (bRenderToScene)
 	{
-		BasicPostProcess(RHICmdList, View, bRequiresUpscale, FSceneRenderer::ShouldCompositeEditorPrimitives(View));
+		for (int32 ViewIndex = 0; ViewIndex < Views.Num(); ViewIndex++)
+		{
+			BasicPostProcess(RHICmdList, Views[ViewIndex], bRequiresUpscale, FSceneRenderer::ShouldCompositeEditorPrimitives(Views[ViewIndex]));
+		}
 	}
 	RenderFinish(RHICmdList);
 }
diff --git a/Engine/Source/Runtime/Renderer/Private/PostProcess/PostProcessTonemap.cpp b/Engine/Source/Runtime/Renderer/Private/PostProcess/PostProcessTonemap.cpp
--- a/Engine/Source/Runtime/Renderer/Private/PostProcess/PostProcessTonemap.cpp
+++ b/Engine/Source/Runtime/Renderer/Private/PostProcess/PostProcessTonemap.cpp
@@ -1425,7 +1425,8 @@ void FRCPassPostProcessTonemapES2::Process(FRenderingCompositePassContext& Conte
 	const FPooledRenderTargetDesc& OutputDesc = PassOutputs[0].RenderTargetDesc;
 
 	FIntRect SrcRect = ViewRect;
-	FIntRect DestRect = View.UnscaledViewRect; // Simple upscaling, ES2 post process does not currently have a specific upscaling pass.
+	// no upscale if separate ren target is used.
+	FIntRect DestRect = (ViewFamily.bUseSeparateRenderTarget) ? ViewRect : View.UnscaledViewRect; // Simple upscaling, ES2 post process does not currently have a specific upscaling pass.
 	FIntPoint SrcSize = InputDesc->Extent;
 	FIntPoint DstSize = OutputDesc.Extent;
 
@@ -1433,7 +1434,10 @@ void FRCPassPostProcessTonemapES2::Process(FRenderingCompositePassContext& Conte
 	SetRenderTarget(Context.RHICmdList, DestRenderTarget.TargetableTexture, FTextureRHIParamRef());
 
 	// Full clear to avoid restore
-	Context.RHICmdList.Clear(true, FLinearColor::Black, false, 1.0f, false, 0, FIntRect());
+	if (View.StereoPass == eSSP_FULL || View.StereoPass == eSSP_LEFT_EYE)
+	{
+		Context.RHICmdList.Clear(true, FLinearColor::Black, false, 1.0f, false, 0, FIntRect());
+	}
 
 	Context.SetViewportAndCallRHI(DestRect);
 
diff --git a/Engine/Source/Runtime/Renderer/Private/PostProcess/PostProcessUpscale.cpp b/Engine/Source/Runtime/Renderer/Private/PostProcess/PostProcessUpscale.cpp
--- a/Engine/Source/Runtime/Renderer/Private/PostProcess/PostProcessUpscale.cpp
+++ b/Engine/Source/Runtime/Renderer/Private/PostProcess/PostProcessUpscale.cpp
@@ -243,21 +243,29 @@ void FRCPassPostProcessUpscale::Process(FRenderingCompositePassContext& Context)
 	const FSceneViewFamily& ViewFamily = *(View.Family);
 	
 	FIntRect SrcRect = View.ViewRect;
-	FIntRect DestRect = View.UnscaledViewRect;
+	// no upscale if separate ren target is used.
+	FIntRect DestRect = (ViewFamily.bUseSeparateRenderTarget) ? View.ViewRect : View.UnscaledViewRect; // Simple upscaling, ES2 post process does not currently have a specific upscaling pass.
 	FIntPoint SrcSize = InputDesc->Extent;
 
 	const FSceneRenderTargetItem& DestRenderTarget = PassOutputs[0].RequestSurface(Context);
+	if (!DestRenderTarget.TargetableTexture)
+	{
+		return;
+	}
 
 	// Set the view family's render target/viewport.
 	SetRenderTarget(Context.RHICmdList, DestRenderTarget.TargetableTexture, FTextureRHIRef());
-	Context.SetViewportAndCallRHI(DestRect);
 
 	bool bTessellatedQuad = PaniniConfig.D >= 0.01f;
 
 	// with distortion (bTessellatedQuad) we need to clear the background
 	FIntRect ExcludeRect = bTessellatedQuad ? FIntRect() : DestRect;
 
-	Context.RHICmdList.Clear(true, FLinearColor::Black, false, 1.0f, false, 0, ExcludeRect);
+	if (View.StereoPass == eSSP_FULL || View.StereoPass == eSSP_LEFT_EYE)
+	{
+		Context.RHICmdList.Clear(true, FLinearColor::Black, false, 1.0f, false, 0, ExcludeRect);
+	}
+	Context.SetViewportAndCallRHI(DestRect);
 
 	// set the state
 	Context.RHICmdList.SetBlendState(TStaticBlendState<>::GetRHI());
diff --git a/Engine/Source/Runtime/Renderer/Private/PostProcess/PostProcessing.cpp b/Engine/Source/Runtime/Renderer/Private/PostProcess/PostProcessing.cpp
--- a/Engine/Source/Runtime/Renderer/Private/PostProcess/PostProcessing.cpp
+++ b/Engine/Source/Runtime/Renderer/Private/PostProcess/PostProcessing.cpp
@@ -2020,7 +2020,16 @@ void FPostProcessing::ProcessES2(FRHICommandListImmediate& RHICmdList, FViewInfo
 
 			FPooledRenderTargetDesc Desc;
 
-			Desc.Extent = View.Family->RenderTarget->GetSizeXY();
+			if (View.Family->RenderTarget->GetRenderTargetTexture())
+			{
+				Desc.Extent.X = View.Family->RenderTarget->GetRenderTargetTexture()->GetSizeX();
+				Desc.Extent.Y = View.Family->RenderTarget->GetRenderTargetTexture()->GetSizeY();
+			}
+			else
+			{
+				Desc.Extent = View.Family->RenderTarget->GetSizeXY();
+			}
+
 			// todo: this should come from View.Family->RenderTarget
 			Desc.Format = PF_B8G8R8A8;
 			Desc.NumMips = 1;
