fixed UE-17333 Switching to SM4 with LPVs active crashes Editor

diff --git a/Engine/Source/Runtime/Renderer/Private/CompositionLighting/CompositionLighting.cpp b/Engine/Source/Runtime/Renderer/Private/CompositionLighting/CompositionLighting.cpp
--- a/Engine/Source/Runtime/Renderer/Private/CompositionLighting/CompositionLighting.cpp
+++ b/Engine/Source/Runtime/Renderer/Private/CompositionLighting/CompositionLighting.cpp
@@ -55,13 +55,16 @@ static bool IsLpvIndirectPassRequired(FPostprocessContext& Context)
 
 	if(ViewState)
 	{
-		FLightPropagationVolume* LightPropagationVolume = ViewState->GetLightPropagationVolume();
+		FLightPropagationVolume* LightPropagationVolume = ViewState->GetLightPropagationVolume(Context.View.GetFeatureLevel());
 
-		const FLightPropagationVolumeSettings& LPVSettings = Context.View.FinalPostProcessSettings.BlendableManager.GetSingleFinalDataConst<FLightPropagationVolumeSettings>();
-
-		if(LightPropagationVolume && LPVSettings.LPVIntensity > 0.0f)
+		if(LightPropagationVolume)
 		{
-			return true; 
+			const FLightPropagationVolumeSettings& LPVSettings = Context.View.FinalPostProcessSettings.BlendableManager.GetSingleFinalDataConst<FLightPropagationVolumeSettings>();
+
+			if(LPVSettings.LPVIntensity > 0.0f)
+			{
+				return true;
+			}
 		}
 	}
 
diff --git a/Engine/Source/Runtime/Renderer/Private/CompositionLighting/PostProcessLpvIndirect.cpp b/Engine/Source/Runtime/Renderer/Private/CompositionLighting/PostProcessLpvIndirect.cpp
--- a/Engine/Source/Runtime/Renderer/Private/CompositionLighting/PostProcessLpvIndirect.cpp
+++ b/Engine/Source/Runtime/Renderer/Private/CompositionLighting/PostProcessLpvIndirect.cpp
@@ -219,15 +219,22 @@ void FRCPassPostProcessLpvIndirect::Process(FRenderingCompositePassContext& Cont
 
 	const FFinalPostProcessSettings& PostprocessSettings = Context.View.FinalPostProcessSettings;
 	const FSceneView& View = Context.View;
-	const FLightPropagationVolumeSettings& LPVSettings = PostprocessSettings.BlendableManager.GetSingleFinalDataConst<FLightPropagationVolumeSettings>();
 
 	FSceneViewState* ViewState = (FSceneViewState*)View.State;
-	if ( ViewState == nullptr || ViewState->GetLightPropagationVolume() == nullptr || LPVSettings.LPVIntensity == 0.0f )
+
+	if(!ViewState)
 	{
 		return;
 	}
 
-	FLightPropagationVolume* Lpv = ViewState->GetLightPropagationVolume();
+	FLightPropagationVolume* Lpv = ViewState->GetLightPropagationVolume(Context.GetFeatureLevel());
+
+	const FLightPropagationVolumeSettings& LPVSettings = PostprocessSettings.BlendableManager.GetSingleFinalDataConst<FLightPropagationVolumeSettings>();
+
+	if(!Lpv || LPVSettings.LPVIntensity == 0.0f)
+	{
+		return;
+	}
 
 	const FSceneViewFamily& ViewFamily = *(View.Family);
 
@@ -345,14 +352,21 @@ void FRCPassPostProcessLpvIndirect::DoDirectionalOcclusionPass(FRenderingComposi
 	const FSceneRenderTargetItem& DestDirectionalOcclusionRenderTarget = SceneContext.DirectionalOcclusion->GetRenderTargetItem();
 	FViewInfo& View = Context.View;
 	FSceneViewState* ViewState = (FSceneViewState*)View.State;
+
+	if(!ViewState)
+	{
+		return;
+	}
+
 	const FFinalPostProcessSettings& PostprocessSettings = Context.View.FinalPostProcessSettings;
 	const FLightPropagationVolumeSettings& LPVSettings = PostprocessSettings.BlendableManager.GetSingleFinalDataConst<FLightPropagationVolumeSettings>();
+	
+	FLightPropagationVolume* Lpv = ViewState->GetLightPropagationVolume(Context.GetFeatureLevel());
 
-	if ( ViewState == nullptr || ViewState->GetLightPropagationVolume() == nullptr || LPVSettings.LPVIntensity == 0.0f )
+	if(!Lpv || LPVSettings.LPVIntensity == 0.0f)
 	{
 		return;
 	}
-	FLightPropagationVolume* Lpv = ViewState->GetLightPropagationVolume();
 
 	FTextureRHIParamRef RenderTarget = DestDirectionalOcclusionRenderTarget.TargetableTexture;
 
diff --git a/Engine/Source/Runtime/Renderer/Private/DeferredShadingRenderer.cpp b/Engine/Source/Runtime/Renderer/Private/DeferredShadingRenderer.cpp
--- a/Engine/Source/Runtime/Renderer/Private/DeferredShadingRenderer.cpp
+++ b/Engine/Source/Runtime/Renderer/Private/DeferredShadingRenderer.cpp
@@ -1577,7 +1577,7 @@ void FDeferredShadingSceneRenderer::ClearLPVs(FRHICommandListImmediate& RHICmdLi
 		FSceneViewState* ViewState = (FSceneViewState*)Views[ViewIndex].State;
 		if(ViewState)
 		{
-			FLightPropagationVolume* LightPropagationVolume = ViewState->GetLightPropagationVolume();
+			FLightPropagationVolume* LightPropagationVolume = ViewState->GetLightPropagationVolume(View.GetFeatureLevel());
 
 			if(LightPropagationVolume)
 			{
@@ -1604,7 +1604,7 @@ void FDeferredShadingSceneRenderer::UpdateLPVs(FRHICommandListImmediate& RHICmdL
 
 		if(ViewState)
 		{
-			FLightPropagationVolume* LightPropagationVolume = ViewState->GetLightPropagationVolume();
+			FLightPropagationVolume* LightPropagationVolume = ViewState->GetLightPropagationVolume(View.GetFeatureLevel());
 
 			if(LightPropagationVolume)
 			{
diff --git a/Engine/Source/Runtime/Renderer/Private/LightPropagationVolume.cpp b/Engine/Source/Runtime/Renderer/Private/LightPropagationVolume.cpp
--- a/Engine/Source/Runtime/Renderer/Private/LightPropagationVolume.cpp
+++ b/Engine/Source/Runtime/Renderer/Private/LightPropagationVolume.cpp
@@ -1460,6 +1460,17 @@ void FSceneViewState::CreateLightPropagationVolumeIfNeeded(ERHIFeatureLevel::Typ
 	}
 }
 
+FLightPropagationVolume* FSceneViewState::GetLightPropagationVolume(ERHIFeatureLevel::Type InFeatureLevel) const
+{
+	if (InFeatureLevel < ERHIFeatureLevel::SM5)
+	{
+		// to prevent crash when starting in SM5 and the using the editor preview SM4
+		return 0;
+	}
+
+	return LightPropagationVolume;
+}
+
 void FSceneViewState::DestroyLightPropagationVolume()
 {
 	if ( LightPropagationVolume ) 
diff --git a/Engine/Source/Runtime/Renderer/Private/LightRendering.cpp b/Engine/Source/Runtime/Renderer/Private/LightRendering.cpp
--- a/Engine/Source/Runtime/Renderer/Private/LightRendering.cpp
+++ b/Engine/Source/Runtime/Renderer/Private/LightRendering.cpp
@@ -523,7 +523,7 @@ void FDeferredShadingSceneRenderer::RenderLights(FRHICommandListImmediate& RHICm
 								if (LightSceneInfo->ShouldRenderLight(View))
 								{
 									FSceneViewState* ViewState = (FSceneViewState*)View.State;
-									FLightPropagationVolume* Lpv = ViewState->GetLightPropagationVolume();
+									FLightPropagationVolume* Lpv = ViewState->GetLightPropagationVolume(View.GetFeatureLevel());
 									if ( Lpv && LightSceneInfo->Proxy )
 									{
 										Lpv->InjectLightDirect( RHICmdList, *LightSceneInfo->Proxy, View );
diff --git a/Engine/Source/Runtime/Renderer/Private/PostProcess/SceneRenderTargets.cpp b/Engine/Source/Runtime/Renderer/Private/PostProcess/SceneRenderTargets.cpp
--- a/Engine/Source/Runtime/Renderer/Private/PostProcess/SceneRenderTargets.cpp
+++ b/Engine/Source/Runtime/Renderer/Private/PostProcess/SceneRenderTargets.cpp
@@ -2184,14 +2184,11 @@ void FSceneTextureShaderParameters::Set(
 		bool bDirectionalOcclusion = false;
 		FSceneViewState* ViewState = (FSceneViewState*)View.State;
 
-		const FLightPropagationVolumeSettings& LPVSettings = View.FinalPostProcessSettings.BlendableManager.GetSingleFinalDataConst<FLightPropagationVolumeSettings>();
-
-		if ( ViewState != nullptr && ViewState->GetLightPropagationVolume() != nullptr && LPVSettings.LPVIntensity > 0.0f )
+		if(ViewState && ViewState->GetLightPropagationVolume(View.GetFeatureLevel()))
 		{
-			if ( LPVSettings.LPVDirectionalOcclusionIntensity > 0.0001f )
-			{
-				bDirectionalOcclusion = true;
-			}
+			const FLightPropagationVolumeSettings& LPVSettings = View.FinalPostProcessSettings.BlendableManager.GetSingleFinalDataConst<FLightPropagationVolumeSettings>();
+
+			bDirectionalOcclusion = LPVSettings.LPVIntensity > 0.0f && LPVSettings.LPVDirectionalOcclusionIntensity > 0.0001f;
 		}
 
 		FTextureRHIParamRef DirectionalOcclusion = nullptr;
diff --git a/Engine/Source/Runtime/Renderer/Private/ReflectionEnvironment.cpp b/Engine/Source/Runtime/Renderer/Private/ReflectionEnvironment.cpp
--- a/Engine/Source/Runtime/Renderer/Private/ReflectionEnvironment.cpp
+++ b/Engine/Source/Runtime/Renderer/Private/ReflectionEnvironment.cpp
@@ -973,15 +973,15 @@ void FDeferredShadingSceneRenderer::RenderStandardDeferredImageBasedReflections(
 	for (int32 ViewIndex = 0, Num = Views.Num(); ViewIndex < Num; ViewIndex++)
 	{
 		FViewInfo& View = Views[ViewIndex];
+		FSceneViewState* ViewState = (FSceneViewState*)View.State;
 
 		bool bLPV = false;
-		FSceneViewState* ViewState = (FSceneViewState*)View.State;
 		
-		const FLightPropagationVolumeSettings& LPVSettings = View.FinalPostProcessSettings.BlendableManager.GetSingleFinalDataConst<FLightPropagationVolumeSettings>();
-
-		if ( ViewState != nullptr && ViewState->GetLightPropagationVolume() != nullptr && LPVSettings.LPVIntensity > 0.0f )
+		if(ViewState && ViewState->GetLightPropagationVolume(View.GetFeatureLevel()))
 		{
-			bLPV = true;
+			const FLightPropagationVolumeSettings& LPVSettings = View.FinalPostProcessSettings.BlendableManager.GetSingleFinalDataConst<FLightPropagationVolumeSettings>();
+
+			bLPV = LPVSettings.LPVIntensity > 0.0f;
 		} 
 		
 		bool bAmbient = View.FinalPostProcessSettings.ContributingCubemaps.Num() > 0;
diff --git a/Engine/Source/Runtime/Renderer/Private/ScenePrivate.h b/Engine/Source/Runtime/Renderer/Private/ScenePrivate.h
--- a/Engine/Source/Runtime/Renderer/Private/ScenePrivate.h
+++ b/Engine/Source/Runtime/Renderer/Private/ScenePrivate.h
@@ -584,7 +584,7 @@ public:
 	void CreateLightPropagationVolumeIfNeeded(ERHIFeatureLevel::Type InFeatureLevel);
 
 	// @return can return 0 (if globally disabled)
-	FLightPropagationVolume* GetLightPropagationVolume() const { return LightPropagationVolume; }
+	FLightPropagationVolume* GetLightPropagationVolume(ERHIFeatureLevel::Type InFeatureLevel) const;
 
 	/** Default constructor. */
 	FSceneViewState();
diff --git a/Engine/Source/Runtime/Renderer/Private/ShadowRendering.cpp b/Engine/Source/Runtime/Renderer/Private/ShadowRendering.cpp
--- a/Engine/Source/Runtime/Renderer/Private/ShadowRendering.cpp
+++ b/Engine/Source/Runtime/Renderer/Private/ShadowRendering.cpp
@@ -462,7 +462,7 @@ public:
 			const FSceneViewState* ViewState = (const FSceneViewState*)View.State;
 			if(ViewState)
 			{
-				const FLightPropagationVolume* Lpv = ViewState->GetLightPropagationVolume();
+				const FLightPropagationVolume* Lpv = ViewState->GetLightPropagationVolume(View.GetFeatureLevel());
 
 				if(Lpv)
 				{
@@ -3207,7 +3207,7 @@ bool FDeferredShadingSceneRenderer::RenderReflectiveShadowMaps(FRHICommandListIm
 				FProjectedShadowInfo* ProjectedShadowInfo = Shadows[ShadowIndex];
 				FSceneViewState* ViewState = (FSceneViewState*)ProjectedShadowInfo->DependentView->State;
 
-				FLightPropagationVolume* LightPropagationVolume = ViewState->GetLightPropagationVolume();
+				FLightPropagationVolume* LightPropagationVolume = ViewState->GetLightPropagationVolume(FeatureLevel);
 
 				check(LightPropagationVolume);
 				auto SetShadowRenderTargets = [this, LightPropagationVolume, ProjectedShadowInfo, &SceneContext](FRHICommandList& InRHICmdList)
@@ -3243,7 +3243,7 @@ bool FDeferredShadingSceneRenderer::RenderReflectiveShadowMaps(FRHICommandListIm
 			{
 				FSceneViewState* ViewState = (FSceneViewState*)ProjectedShadowInfo->DependentView->State;
 
-				FLightPropagationVolume* LightPropagationVolume = ViewState ? ViewState->GetLightPropagationVolume() : NULL;
+				FLightPropagationVolume* LightPropagationVolume = ViewState ? ViewState->GetLightPropagationVolume(FeatureLevel) : NULL;
 
 				if ( LightPropagationVolume )
 				{
diff --git a/Engine/Source/Runtime/Renderer/Private/ShadowSetup.cpp b/Engine/Source/Runtime/Renderer/Private/ShadowSetup.cpp
--- a/Engine/Source/Runtime/Renderer/Private/ShadowSetup.cpp
+++ b/Engine/Source/Runtime/Renderer/Private/ShadowSetup.cpp
@@ -2265,7 +2265,7 @@ void FSceneRenderer::AddViewDependentWholeSceneShadowsForView(
 			FSceneViewState* ViewState = (FSceneViewState*)View.State;
 			if (ViewState)
 			{
-				FLightPropagationVolume* LightPropagationVolume = ViewState->GetLightPropagationVolume();
+				FLightPropagationVolume* LightPropagationVolume = ViewState->GetLightPropagationVolume(View.GetFeatureLevel());
 				
 				FLightPropagationVolumeSettings& LPVSettings = View.FinalPostProcessSettings.BlendableManager.GetSingleFinalData<FLightPropagationVolumeSettings>();
 
diff --git a/Engine/Source/Runtime/Renderer/Private/TranslucentRendering.cpp b/Engine/Source/Runtime/Renderer/Private/TranslucentRendering.cpp
--- a/Engine/Source/Runtime/Renderer/Private/TranslucentRendering.cpp
+++ b/Engine/Source/Runtime/Renderer/Private/TranslucentRendering.cpp
@@ -1087,7 +1087,7 @@ void FDeferredShadingSceneRenderer::RenderTranslucency(FRHICommandListImmediate&
 
 			if(ViewState && View.Family->EngineShowFlags.VisualizeLPV)
 			{
-				FLightPropagationVolume* LightPropagationVolume = ViewState->GetLightPropagationVolume();
+				FLightPropagationVolume* LightPropagationVolume = ViewState->GetLightPropagationVolume(View.GetFeatureLevel());
 
 				if (LightPropagationVolume)
 				{
