UE4 - Fix sm4 reflections

diff --git a/Engine/Source/Runtime/Renderer/Private/PostProcess/SceneRenderTargets.cpp b/Engine/Source/Runtime/Renderer/Private/PostProcess/SceneRenderTargets.cpp
--- a/Engine/Source/Runtime/Renderer/Private/PostProcess/SceneRenderTargets.cpp
+++ b/Engine/Source/Runtime/Renderer/Private/PostProcess/SceneRenderTargets.cpp
@@ -1696,7 +1696,11 @@ void FSceneRenderTargets::AllocateDeferredShadingPathRenderTargets()
 	}
 
 	{
-		FPooledRenderTargetDesc Desc(FPooledRenderTargetDesc::Create2DDesc(BufferSize, GetSceneColorFormat(), TexCreate_None, TexCreate_RenderTargetable | TexCreate_UAV, false));
+		FPooledRenderTargetDesc Desc(FPooledRenderTargetDesc::Create2DDesc(BufferSize, GetSceneColorFormat(), TexCreate_None, TexCreate_RenderTargetable, false));
+		if (CurrentFeatureLevel >= ERHIFeatureLevel::SM5)
+		{
+			Desc.TargetableFlags |= TexCreate_UAV;
+		}
 		GRenderTargetPool.FindFreeElement(Desc, LightAccumulation, TEXT("LightAccumulation"));
 	}
 
