nicer VisualizeLPV

diff --git a/Engine/Source/Runtime/Renderer/Private/CompositionLighting/PostProcessLpvIndirect.cpp b/Engine/Source/Runtime/Renderer/Private/CompositionLighting/PostProcessLpvIndirect.cpp
--- a/Engine/Source/Runtime/Renderer/Private/CompositionLighting/PostProcessLpvIndirect.cpp
+++ b/Engine/Source/Runtime/Renderer/Private/CompositionLighting/PostProcessLpvIndirect.cpp
@@ -458,6 +458,9 @@ void FRCPassPostProcessVisualizeLPV::Process(FRenderingCompositePassContext& Con
 		const float YStep = 14;
 		const float ColumnWidth = 250;
 
+		Canvas.DrawShadowedString( X, Y += YStep, TEXT("VisualizeLightPropagationVolume"), GetStatsFont(), FLinearColor(0.2f, 0.2f, 1));
+
+		Y += YStep;
 
 		const FLightPropagationVolumeSettings& Dest = View.FinalPostProcessSettings.BlendableManager.GetSingleFinalDataConst<FLightPropagationVolumeSettings>();
 
