UE4 - Fix reflections on SM4 broken from LPV changes (UE-17270)

diff --git a/Engine/Source/Runtime/Renderer/Private/ReflectionEnvironment.cpp b/Engine/Source/Runtime/Renderer/Private/ReflectionEnvironment.cpp
--- a/Engine/Source/Runtime/Renderer/Private/ReflectionEnvironment.cpp
+++ b/Engine/Source/Runtime/Renderer/Private/ReflectionEnvironment.cpp
@@ -985,7 +985,7 @@ void FDeferredShadingSceneRenderer::RenderStandardDeferredImageBasedReflections(
 		} 
 		
 		bool bAmbient = View.FinalPostProcessSettings.ContributingCubemaps.Num() > 0;
-		bool bMixing = CVarLPVMixing.GetValueOnRenderThread() != 0;
+		bool bMixing = bLPV && (CVarLPVMixing.GetValueOnRenderThread() != 0);
 		bool bEnvironmentMixing = bMixing && (bAmbient || bLPV);
 		bool bRequiresApply = bSkyLight
 			// If Reflection Environment is active and mixed with indirect lighting (Ambient + LPV), apply is required!
