Motion blur optimizations. Still in progress.
Fixed temporal aa with pause. Requires r.TemporalAAPauseCorrect 1.

diff --git a/Engine/Shaders/PostProcessCommon.usf b/Engine/Shaders/PostProcessCommon.usf
--- a/Engine/Shaders/PostProcessCommon.usf
+++ b/Engine/Shaders/PostProcessCommon.usf
@@ -51,7 +51,7 @@ float4 PostprocessInput6MinMax;
 // viewport width, height, 1/width, 1/height (scaled to the current rendertarget resolution), depends on SetViewportAndCallRHI() call
 float4 ViewportSize;
 // in pixels (scaled to the current rendertarget resolution), float4(minx,miny,maxx,maxy), depends on SetViewportAndCallRHI() call
-float4 ViewportRect;
+uint4 ViewportRect;
 
 SamplerState BilinearTextureSampler0;
 SamplerState BilinearTextureSampler1;
diff --git a/Engine/Shaders/PostProcessMotionBlur.usf b/Engine/Shaders/PostProcessMotionBlur.usf
--- a/Engine/Shaders/PostProcessMotionBlur.usf
+++ b/Engine/Shaders/PostProcessMotionBlur.usf
@@ -7,6 +7,7 @@
 #include "Common.usf"
 #include "PostProcessCommon.usf"
 #include "DeferredShadingCommon.usf"		// FGBufferData
+#include "FastMath.usf"
 
 // Spaces:
 //   screen space: -1..1 -1..1
@@ -660,12 +661,13 @@ void MainRecombinePS(float4 UVAndScreenPos : TEXCOORD0, out float4 OutColor : SV
 
 
 
-uint2 TileCount;
+
+uint	bDrawMax;
 
 void VelocityScatterVS(
 	uint VId : SV_VertexID,
 	uint IId : SV_InstanceID,
-	out nointerpolation float2 OutColor : TEXCOORD0,
+	out nointerpolation float4 OutColor : TEXCOORD0,
 	out float4 OutPosition : SV_POSITION
 	)
 {
@@ -691,40 +693,70 @@ void VelocityScatterVS(
 	float2 CornerOffset = float2(VId % 2, VId / 2) * 2 - 1;
 #endif
 
-	uint2 TilePos = uint2( IId % TileCount.x, IId / TileCount.x );
+	uint2 TilePos = uint2( IId % ViewportRect.z, IId / ViewportRect.z );
 
-	float2 ScreenPos = ( TilePos - ScreenPosToPixel.zw ) / ScreenPosToPixel.xy;
+	BRANCH
+	if( TilePos.y >= ViewportRect.w )
+	{
+		OutColor = 0;
+		return;
+	}
 
-	float2 Velocity = PostprocessInput0[ TilePos ].xy;
-	OutColor = Velocity.xyxy;
+	float4 MinMaxVelocity = PostprocessInput0[ TilePos ];
+	OutColor = MinMaxVelocity;
 
-	Velocity *= VelocityScale.xy;
+	//float2 VelocityPixels = MinMaxVelocity.zw * ScreenPosToPixel.x;
+	float2 VelocityPixels = MinMaxVelocity.zw * 0.5 * 1920.0 / 16.0;
+	
+	// TODO move to velocity flatten shader
+	//VelocityPixels *= VelocityScale.xy;
+	VelocityPixels *= 0.5;
 	
+#if 1
 	BRANCH
-	if( length( Velocity * float2(1920,1080) ) <= 0.5 )
+	if( dot( VelocityPixels, VelocityPixels ) * 16*16 <= 0.25 )
 	{
+		OutPosition.xy = ( TilePos + CornerOffset * 0.5 - ScreenPosToPixel.zw ) / ScreenPosToPixel.xy;
+		OutPosition.z = 0.0002;		// zero clips
 		return;
 	}
+#endif
 
-	float  VelocityLengthSqr = dot( Velocity, Velocity );
-	float  VelocityLengthInv = rsqrt( VelocityLengthSqr );
-	float  VelocityLength = VelocityLengthSqr * VelocityLengthInv;
-	float2 VelocityDir = Velocity * VelocityLengthInv;
+	float  VelocityLengthPixelsSqr = dot( VelocityPixels, VelocityPixels );
+	float  VelocityLengthPixelsInv = rsqrtFast( VelocityLengthPixelsSqr );
+	float  VelocityLengthPixels = VelocityLengthPixelsSqr * VelocityLengthPixelsInv;
+	float2 VelocityDir = VelocityPixels * VelocityLengthPixelsInv;
+
+	// Project pixel corner on to dir. This is the oriented extent of a pixel.
+	// 1/2 pixel because shape is swept tile
+	// +1/2 pixel for conservative rasterization
+	// 99% to give epsilon before neighbor is filled. Otherwise all neighbors lie on edges of quad when no velocity in their direction.
+	float Extent = abs( VelocityDir.x ) + abs( VelocityDir.y );
+	CornerOffset *= float2( VelocityLengthPixels, 0 ) + Extent.xx * 0.99;
 
-	// +1 pixel radius because shape is swept tile, approx as capsule
-	// +1 pixel radius for conservative rasterization
 	const float TwoPixelRadius = sqrt( 2.0 );
-	CornerOffset *= float2( VelocityLength, 0 ) + TwoPixelRadius.xx / ScreenPosToPixel.xy;
+	//CornerOffset *= float2( VelocityLengthPixels, 0 ) + TwoPixelRadius.xx;
 
 	// Orient along velocity direction
 	float2 AxisX = VelocityDir;
 	float2 AxisY = float2( -VelocityDir.y, VelocityDir.x );
 	CornerOffset = AxisX * CornerOffset.x + AxisY * CornerOffset.y;
 		
-	OutPosition.xy = ScreenPos + CornerOffset;
+	OutPosition.xy = ( TilePos + CornerOffset - ScreenPosToPixel.zw ) / ScreenPosToPixel.xy;
+
+	if( bDrawMax )
+	{
+		// Depth ordered by velocity length
+		OutPosition.z = saturate( VelocityLengthPixels / ScreenPosToPixel.x * 0.5 );
+	}
+	else
+	{
+		//float2 MinVelocityPixels = MinMaxVelocity.xy * ScreenPosToPixel.xy * VelocityScale.xy;
+		float2 MinVelocityPixels = MinMaxVelocity.xy * 0.5 * 1920.0 / 16.0 * 0.5;
+		OutPosition.z = saturate( length( MinVelocityPixels ) / ScreenPosToPixel.x * 0.5 );
+	}
 
-	// Depth ordered by velocity length
-	OutPosition.z = saturate( VelocityLength * 0.5 );
+	OutPosition.z = clamp( OutPosition.z, 0.0002, 0.999 );
 }
 
 void VelocityScatterPS(
@@ -738,46 +770,6 @@ void VelocityScatterPS(
 
 
 
-float2 MaxLength( float2 v0, float2 v1 )
-{
-	float LengthSqr0 = dot( v0, v0 );
-	float LengthSqr1 = dot( v1, v1 );
-	return LengthSqr0 > LengthSqr1 ? v0 : v1;
-}
-
-float2 MaxLength( float2 v0, float2 v1, float2 v2 )
-{
-	float LengthSqr0 = dot( v0, v0 );
-	float LengthSqr1 = dot( v1, v1 );
-	float LengthSqr2 = dot( v2, v2 );
-	return LengthSqr0 > LengthSqr1 ? ( LengthSqr0 > LengthSqr2 ? v0 : v2 ) : ( LengthSqr1 > LengthSqr2 ? v1 : v2 );
-}
-
-void VelocityDilatePS(float4 UVAndScreenPos : TEXCOORD0, out float4 OutColor : SV_Target0)
-{
-	float2 UV = UVAndScreenPos.xy;
-
-	float2 Neighbor0 = PostprocessInput0.SampleLevel( PostprocessInput0Sampler, UV, 0, int2(-1, -1) ).xy;
-	float2 Neighbor1 = PostprocessInput0.SampleLevel( PostprocessInput0Sampler, UV, 0, int2( 0, -1) ).xy;
-	float2 Neighbor2 = PostprocessInput0.SampleLevel( PostprocessInput0Sampler, UV, 0, int2( 1, -1) ).xy;
-	float2 Neighbor3 = PostprocessInput0.SampleLevel( PostprocessInput0Sampler, UV, 0, int2(-1,  0) ).xy;
-	float2 Neighbor4 = PostprocessInput0.SampleLevel( PostprocessInput0Sampler, UV, 0				).xy;
-	float2 Neighbor5 = PostprocessInput0.SampleLevel( PostprocessInput0Sampler, UV, 0, int2( 1,  0) ).xy;
-	float2 Neighbor6 = PostprocessInput0.SampleLevel( PostprocessInput0Sampler, UV, 0, int2(-1,  1) ).xy;
-	float2 Neighbor7 = PostprocessInput0.SampleLevel( PostprocessInput0Sampler, UV, 0, int2( 0,  1) ).xy;
-	float2 Neighbor8 = PostprocessInput0.SampleLevel( PostprocessInput0Sampler, UV, 0, int2( 1,  1) ).xy;
-
-	float2 Max012 = MaxLength( Neighbor0, Neighbor1, Neighbor2 );
-	float2 Max345 = MaxLength( Neighbor3, Neighbor4, Neighbor5 );
-	float2 Max678 = MaxLength( Neighbor6, Neighbor7, Neighbor8 );
-
-	OutColor.xy = MaxLength( Max012, Max345, Max678 );
-	OutColor.zw = 0;
-}
-
-
-
-
 
 float2 DepthCmp( float CenterDepth, float SampleDepth, float DepthScale )
 {
@@ -808,8 +800,52 @@ float RandFast( uint2 PixelPos, float Magic = 3571.0 )
 	return Random.x;
 }
 
+// [0, 1]
+float InterleavedGradientNoise( uint2 PixelPos )
+{
+	return frac( 52.9829189 * frac( dot( PixelPos, float2( 0.06711056, 0.00583715 ) ) ) );
+}
+
+float2 DecodeVelocity( float2 Velocity )
+{
+#if 1
+	// 11:11:10  (VelocityLength, VelocityAngle, Depth)
+	float VelocityLength = Velocity.x;
+	float VelocityAngle = Velocity.y * (2 * PI) - PI;
+	sincos( VelocityAngle, Velocity.x, Velocity.y );
+	Velocity *= VelocityLength;
+#else
+	// 11:11:10 (Velocity.xy, Depth)
+	// Stored signed value in 6e5 unsigned float
+	// Extract sign from top bit of exponent
+	Velocity *= Velocity >= 2.0 ? (-1.0 / 32768.0) : 1;
+#endif
+	return Velocity;
+}
+
+static const float LimitVelocity = 100;
+
+float GetVelocityLengthPixels( float2 EncodedVelocity )
+{
+#if 1
+	// 11:11:10  (VelocityLength, VelocityAngle, Depth)
+	float VelocityLength = EncodedVelocity.x;
+	VelocityLength *= ScreenPosToPixel.x;
+	return min( VelocityLength, LimitVelocity );
+#else
+	float2 Velocity = DecodeVelocity( EncodedVelocity ) * ScreenPosToPixel.xy;	// in pixels
+	float  VelocityLengthSqr = dot( Velocity, Velocity );
+	float  VelocityLength = sqrtFast( VelocityLengthSqr );
+	return min( VelocityLength, LimitVelocity );
+#endif
+}
+
 
-void MainNewPS( float4 UVAndScreenPos : TEXCOORD0, out float4 OutColor : SV_Target0 )
+void MainNewPS(
+	in float4 UVAndScreenPos : TEXCOORD0,
+	in float4 SvPosition : SV_Position,
+	out float4 OutColor : SV_Target0
+	)
 {
 	// Screen Quad UV 0..1
 	float2 UV = UVAndScreenPos.xy;
@@ -818,46 +854,58 @@ void MainNewPS( float4 UVAndScreenPos : TEXCOORD0, out float4 OutColor : SV_Targ
 
 	OutColor = 0;
 
-	float VelocityScale = MotionBlurParameters.x;
+	float VelocityScale = 0.5;//MotionBlurParameters.x;
 
 	const uint StepCount = GetStepCountFromQuality() / 2;
 
-#if 1
-	uint2 PixelPos = uint2(UVAndScreenPos.zw * ScreenPosToPixel.xy + ScreenPosToPixel.zw + 0.5f);
+	//uint2 PixelPos = uint2(UVAndScreenPos.zw * ScreenPosToPixel.xy + ScreenPosToPixel.zw + 0.5f);
+	uint2 PixelPos = SvPosition.xy;
+#if 0
 	float2 PosMod = float2( PixelPos & 1 );
-	float Dither = ( PosMod.x * 2 - 1 ) * ( PosMod.y * 2 - 1 );
+	float Dither = ( PosMod.x * 0.5 - 0.25 ) * ( PosMod.y * 2 - 1 );
 	float Random = RandFast( PixelPos );
 	float Random2 = RandFast( PixelPos, 5521 );
+#else
+	float Random = InterleavedGradientNoise( PixelPos );
+	//float Random2 = InterleavedGradientNoise( PixelPos.yx );
+
+	float2 PosMod = float2( PixelPos & 1 );
+	float Random2 = ( PosMod.x * 0.5 - 0.25 ) * ( PosMod.y * 2 - 1 );
 #endif
-	
+
 	float2 TileJitter = ( float2( Random, Random2 ) - 0.5 ) * 0.5;
-	
-	float2 MaxVelocity = PostprocessInput3.SampleLevel( PostprocessInput3Sampler, UV + TileJitter * PostprocessInput3Size.zw, 0 ).xy;
-	//float2 MaxVelocity = PostprocessInput3[ PixelPos / 16 ].xy;
+	float2 TileUV = min( ( PixelPos - ViewportRect.xy + 0.5 ) / 16.0 + TileJitter, ViewportRect.zw ) * PostprocessInput3Size.zw;
 
-	float2 MaxVelocityScreen = MaxVelocity.xy * VelocityScale;
-	float2 MaxVelocityPixels = MaxVelocityScreen * ScreenPosToPixel.xy;
+	float4 MinMaxVelocity = PostprocessInput3.SampleLevel( PostprocessInput3Sampler, TileUV, 0 );
+	//float4 MinMaxVelocity = PostprocessInput3[ ( PixelPos - ViewportRect.xy ) / 16 ];
 
-	float3 CenterColor = PostprocessInput0.SampleLevel( PostprocessInput0Sampler, UV, 0 ).rgb;
-	float  CenterDepth = PostprocessInput1.SampleLevel( PostprocessInput1Sampler, UV, 0 ).r;
+	float4 MinMaxVelocityScreen = MinMaxVelocity * VelocityScale;
+	float2 MinVelocityPixels = MinMaxVelocityScreen.xy * ScreenPosToPixel.x;
+	float2 MaxVelocityPixels = MinMaxVelocityScreen.zw * ScreenPosToPixel.x;
 
-	CenterDepth = ConvertFromDeviceZ( CenterDepth );
+	float MinVelocityLengthSqrPixels = dot( MinVelocityPixels, MinVelocityPixels );
+	float MaxVelocityLengthSqrPixels = dot( MaxVelocityPixels, MaxVelocityPixels );
+	
+	// TODO expose cvars
+	bool bSkipPath = MaxVelocityLengthSqrPixels < 0.25;
+	bool bFastPath = MinVelocityLengthSqrPixels > 0.4 * MaxVelocityLengthSqrPixels;
+
+	float3 CenterColor = PostprocessInput0.SampleLevel( PostprocessInput0Sampler, UV, 0 ).rgb;
 	
 	BRANCH
-	if( length( MaxVelocityPixels ) <= 0.5 )
+	if( bSkipPath )
 	{
 		OutColor.rgb = CenterColor;
 		return;
 	}
-	
-	// in pixels
-	float2 CenterVelocity = PostprocessInput2.SampleLevel( PostprocessInput2Sampler, UV, 0 ).xy;
-	CenterVelocity = DecodeVelocityFromTexture( CenterVelocity ) * VelocityScale * ScreenPosToPixel.xy;
-	float CenterVelocityLength = length( CenterVelocity );
 
-	
+#if 1
+	float4 SearchVectorPixels = float4( MaxVelocityPixels, -MaxVelocityPixels );
+	float2 SearchLengthPixels = length( MaxVelocityPixels ).xx;
+	float4 SearchVector = SearchVectorPixels * PostprocessInput0Size.zwzw;
+#else
 	// Clip MaxVelocity to screen rect
-	float2 InvVelocityScreen = rcp( MaxVelocityScreen );
+	float2 InvVelocityScreen = rcp( MinMaxVelocityScreen.zw );
 	float2 MinIntersect = -InvVelocityScreen - ScreenSpacePos * InvVelocityScreen;
 	float2 MaxIntersect =  InvVelocityScreen - ScreenSpacePos * InvVelocityScreen;
 	float4 FarIntersect = float4( max( MinIntersect, MaxIntersect ), max( -MinIntersect, -MaxIntersect ) );
@@ -866,79 +914,85 @@ void MainNewPS( float4 UVAndScreenPos : TEXCOORD0, out float4 OutColor : SV_Targ
 	// +/-
 	float4 SearchVectorPixels = MaxVelocityPixels.xyxy * float4( Intersect.xx, -Intersect.yy );
 	float2 SearchLengthPixels = { length( SearchVectorPixels.xy ), length( SearchVectorPixels.zw ) };
+	float4 SearchVector = SearchVectorPixels * PostprocessInput0Size.zwzw;
+#endif
 
 	// converts pixel length to sample steps
-	float2 PixelToSampleScale = ( StepCount - 0.5 ) / SearchLengthPixels;
-
-	float4 ColorAccum = 0;
+	//float2 PixelToSampleScale = StepCount / SearchLengthPixels;
+	float2 PixelToSampleScale = StepCount / length( MaxVelocityPixels );
 
-	UNROLL
-	for( uint i = 0; i < StepCount; i++ )
+	BRANCH
+	if( bFastPath )
 	{
-		float2 SampleUV[2];
-		float3 SampleColor[2];
-		float  SampleDepth[2];
-		float  SampleVelocityLength[2];
-		float  Weight[2];
+		float4 ColorAccum = 0;
 
-		float OffsetLength = (float)i + 0.5 + (Random - 0.5);
-		float OffsetFraction = OffsetLength / ( StepCount - 0.5 );
+		UNROLL
+		for( uint i = 0; i < StepCount; i++ )
+		{
+			float2 OffsetLength = (float)i + 0.5 + float2( Random - 0.5, 0.5 - Random );
+			float2 OffsetFraction = OffsetLength / StepCount;
 
-		//float2 TestVector = (i & 1) > 0 ? CenterVelocity : SearchVectorPixels;
-		//float2 PixelToSampleScale = ( StepCount - 0.5 ) / length( TestVector );
+			float2 SampleUV[2];
+			SampleUV[0] = UV + OffsetFraction.x * SearchVector.xy;
+			SampleUV[1] = UV + OffsetFraction.y * SearchVector.zw;
 
-		// can be optimized
-		SampleUV[0] = UV + OffsetFraction * SearchVectorPixels.xy * PostprocessInput0Size.zw;
-		SampleUV[1] = UV + OffsetFraction * SearchVectorPixels.zw * PostprocessInput0Size.zw;
+			ColorAccum.rgb += PostprocessInput0.SampleLevel( PostprocessInput0Sampler, SampleUV[0], 0 ).rgb;
+			ColorAccum.rgb += PostprocessInput0.SampleLevel( PostprocessInput0Sampler, SampleUV[1], 0 ).rgb;
+		}
+	
+		ColorAccum *= 0.5 / StepCount;
+
+		OutColor.rgb = ColorAccum.rgb;
+	}
+	else
+	{
+		float3 CenterVelocityDepth = PostprocessInput2.SampleLevel( PostprocessInput2Sampler, UV, 0 ).xyz;
+		float  CenterDepth = CenterVelocityDepth.z;
+		float  CenterVelocityLength = GetVelocityLengthPixels( CenterVelocityDepth.xy ) * VelocityScale;
+
+		float4 ColorAccum = 0;
 
 		UNROLL
-		for( uint j = 0; j < 2; j++ )
+		for( uint i = 0; i < StepCount; i++ )
 		{
-			SampleColor[j] = PostprocessInput0.SampleLevel( PostprocessInput0Sampler, SampleUV[j], 0 ).rgb;
-			SampleDepth[j] = PostprocessInput1.SampleLevel( PostprocessInput1Sampler, SampleUV[j], 0 ).r;
+			float2 SampleUV[2];
+			float3 SampleColor[2];
+			float  SampleDepth[2];
+			float  SampleVelocityLength[2];
+			float  Weight[2];
 
-			SampleDepth[j] = ConvertFromDeviceZ( SampleDepth[j] );
+			float2 OffsetLength = (float)i + 0.5 + float2( Random - 0.5, 0.5 - Random );
+			float2 OffsetFraction = OffsetLength / StepCount;
 
-			// in pixels
-			float2 SampleVelocity = PostprocessInput2.SampleLevel( PostprocessInput2Sampler, SampleUV[j], 0 ).xy;
-			SampleVelocity = DecodeVelocityFromTexture( SampleVelocity ) * VelocityScale * ScreenPosToPixel.xy;
-			SampleVelocityLength[j] = length( SampleVelocity );
+			SampleUV[0] = UV + OffsetFraction.x * SearchVector.xy;
+			SampleUV[1] = UV + OffsetFraction.y * SearchVector.zw;
 
-			Weight[j] = SampleWeight( CenterDepth, SampleDepth[j], OffsetLength, CenterVelocityLength, SampleVelocityLength[j], PixelToSampleScale[j], SOFT_Z_EXTENT );
+			UNROLL
+			for( uint j = 0; j < 2; j++ )
+			{
+				float3 SampleVelocityDepth = PostprocessInput2.SampleLevel( PostprocessInput2Sampler, SampleUV[j], 0 ).xyz;
+				SampleColor[j] = PostprocessInput0.SampleLevel( PostprocessInput0Sampler, SampleUV[j], 0 ).rgb;
+				SampleDepth[j] = SampleVelocityDepth.z;
 
-#if 0
-			float TestLengthPixels = OffsetFraction * SearchLengthPixels[j];
+				// in pixels
+				SampleVelocityLength[j] = GetVelocityLengthPixels( SampleVelocityDepth.xy ) * VelocityScale;
 
-			float2 DepthWeights = DepthCmp( CenterDepth, SampleDepth[j], SOFT_Z_EXTENT );
-			
-			//Weight[j]  = DepthWeights.x * saturate( 1 - TestLengthPixels / CenterVelocityLength );
-			//Weight[j] += DepthWeights.y * saturate( 1 - TestLengthPixels / SampleVelocityLength[j] );
-			//Weight[j] += ( 1 - smoothstep(0.95f * CenterVelocityLength, 1.05f * CenterVelocityLength, TestLengthPixels ) ) *
-			//			( 1 - smoothstep(0.95f * SampleVelocityLength[j], 1.05f * SampleVelocityLength[j], TestLengthPixels ) ) * 2;
-#endif
-		}
+				Weight[j] = SampleWeight( CenterDepth, SampleDepth[j], OffsetLength, CenterVelocityLength, SampleVelocityLength[j], PixelToSampleScale[j], SOFT_Z_EXTENT );
+			}
 		
-		bool2 Mirror = bool2( SampleDepth[0] > SampleDepth[1], SampleVelocityLength[1] > SampleVelocityLength[0] );
-		Weight[0] = all( Mirror ) ? Weight[1] : Weight[0];
-		Weight[1] = any( Mirror ) ? Weight[1] : Weight[0];
-
-		ColorAccum += Weight[0] * float4( SampleColor[0].rgb, 1 );
-		ColorAccum += Weight[1] * float4( SampleColor[1].rgb, 1 );
-	}
+			bool2 Mirror = bool2( SampleDepth[0] > SampleDepth[1], SampleVelocityLength[1] > SampleVelocityLength[0] );
+			Weight[0] = all( Mirror ) ? Weight[1] : Weight[0];
+			Weight[1] = any( Mirror ) ? Weight[1] : Weight[0];
 
-	//ColorAccum += float4( CenterColor, 1 ) / CenterVelocityLength;
-	
-	ColorAccum *= 0.5 / StepCount;
-	OutColor.rgb = ColorAccum.rgb + ( 1 - ColorAccum.a ) * CenterColor;
-	
-#if 0
-	// Fudge factor to increase alpha
-	ColorAccum.rgb /= ColorAccum.a;
-	ColorAccum.a = saturate( ColorAccum.a * 1.2 );
-	OutColor.rgb = lerp( CenterColor, ColorAccum.rgb, ColorAccum.a );
-#endif
+			ColorAccum += Weight[0] * float4( SampleColor[0], 1 );
+			ColorAccum += Weight[1] * float4( SampleColor[1], 1 );
+		}
 	
-	OutColor.a = 1;
+		ColorAccum *= 0.5 / StepCount;
+
+		OutColor.rgb = ColorAccum.rgb + ( 1 - ColorAccum.a ) * CenterColor;
+	}
 
-	//OutColor.rgb = ColorAccum.a * 0.2;
+	//OutColor.rgb *= bFastPath ? float3(1,0,0) : float3(0,1,0);
+	//OutColor.rgb = lerp( OutColor.rgb, float3( abs( MinMaxVelocity.zw ), 0 ) * 0.2, 0.9 );
 }
\ No newline at end of file
diff --git a/Engine/Shaders/PostProcessVelocityFlatten.usf b/Engine/Shaders/PostProcessVelocityFlatten.usf
--- a/Engine/Shaders/PostProcessVelocityFlatten.usf
+++ b/Engine/Shaders/PostProcessVelocityFlatten.usf
@@ -2,19 +2,45 @@
 
 #include "Common.usf"
 #include "PostProcessCommon.usf"
+#include "FastMath.usf"
 
-#define THREADGROUP_TOTALSIZE (THREADGROUP_SIZEX * THREADGROUP_SIZEY)
+#define THREADGROUP_SIZEX		16
+#define THREADGROUP_SIZEY		16
+#define THREADGROUP_TOTALSIZE	(THREADGROUP_SIZEX * THREADGROUP_SIZEY)
 
-groupshared float2 Shared[ THREADGROUP_TOTALSIZE ];
+groupshared float4 Shared[ THREADGROUP_TOTALSIZE ];
 
 RWTexture2D<float4>	OutVelocityFlat;
 RWTexture2D<float4>	OutMaxTileVelocity;
 
-float2 MaxLength( float2 v0, float2 v1 )
+
+float4 MinMaxLength( float4 v0, float4 v1 )
+{
+	float2 Min = dot( v0.xy, v0.xy ) < dot( v1.xy, v1.xy ) ? v0.xy : v1.xy;
+	float2 Max = dot( v0.zw, v0.zw ) > dot( v1.zw, v1.zw ) ? v0.zw : v1.zw;
+	return float4( Min, Max );
+}
+
+float4 MinMaxLengthPolar( float4 v0, float4 v1 )
+{
+	float2 Min = v0.x < v1.x ? v0.xy : v1.xy;
+	float2 Max = v0.z > v1.z ? v0.zw : v1.zw;
+	return float4( Min, Max );
+}
+
+float2 CartesianToPolar( float2 Velocity )
 {
-	float LengthSqr0 = dot( v0, v0 );
-	float LengthSqr1 = dot( v1, v1 );
-	return LengthSqr0 > LengthSqr1 ? v0 : v1;
+	float Length = length( Velocity );
+	float Angle = atan2Fast( Velocity.y, Velocity.x );
+	return float2( Length, Angle );
+}
+
+float2 PolarToCartesian( float2 Velocity )
+{
+	float Length = Velocity.x;
+	float Angle = Velocity.y;
+	sincos( Angle, Velocity.y, Velocity.x );
+	return Velocity * Length;
 }
 
 [numthreads(THREADGROUP_SIZEX, THREADGROUP_SIZEY, 1)]
@@ -24,100 +50,182 @@ void VelocityFlattenMain(
 	uint3 GroupThreadId : SV_GroupThreadID,
 	uint GroupIndex : SV_GroupIndex ) 
 {
-	uint ThreadIndex = GroupThreadId.y * THREADGROUP_SIZEX + GroupThreadId.x;
-	uint2 PixelPos = DispatchThreadId.xy * 2 + VelocityFlatten.ViewDimensions.xy;
-	float2 ScreenUV = ( (float2)( DispatchThreadId.xy * 2 ) + 0.5 ) / ( VelocityFlatten.ViewDimensions.zw - VelocityFlatten.ViewDimensions.xy );
-	float2 ScreenPosition = float2(2.0f, -2.0f) * ScreenUV + float2(-1.0f, 1.0f);
+	if(0)
+	{
+		// Swizzle threads
+		uint TileIndex = (DispatchThreadId.x & 3) + (DispatchThreadId.y & 3) * 4;
+		DispatchThreadId.x = (DispatchThreadId.x & ~3) | ((0xebfa5014 >> (TileIndex*2) ) & 3);
+		DispatchThreadId.y = (DispatchThreadId.y & ~3) | ((0xfa4141af >> (TileIndex*2) ) & 3);
+	}
+
+	uint2 PixelPos = DispatchThreadId.xy + ViewportRect.xy;
 
-	// FIXME for sub viewrect
-	float2 UV = ScreenUV;
+	float2 Velocity = PostprocessInput0[ PixelPos ].xy;
+	float Depth = PostprocessInput1[ PixelPos ].x;
 
-	uint2 Offsets[] =
+	if( Velocity.x > 0.0 )
 	{
-		{ 0, 0 },
-		{ 0, 1 },
-		{ 1, 0 },
-		{ 1, 1 },
-	};
+		Velocity = DecodeVelocityFromTexture( Velocity );
+	}
+	else
+	{
+		float2 ScreenUV = ( (float2)DispatchThreadId.xy + 0.5 ) / ( ViewportRect.zw - ViewportRect.xy );
+		float3 PosN = float3( ScreenUV, Depth );
+		float ScaleM = rcp( dot( PosN, CameraMotion.Value[0].xyz ) + CameraMotion.Value[0].w );
+		Velocity.x = -2.0 * ((PosN.x * ((CameraMotion.Value[1].x * PosN.y) + (CameraMotion.Value[1].y * PosN.z) + CameraMotion.Value[1].z)) + (CameraMotion.Value[1].w * PosN.y) + (CameraMotion.Value[2].x * PosN.x * PosN.x) + (CameraMotion.Value[2].y * PosN.z) + CameraMotion.Value[2].z) * ScaleM;
+		Velocity.y =  2.0 * ((PosN.y * ((CameraMotion.Value[3].x * PosN.x) + (CameraMotion.Value[3].y * PosN.z) + CameraMotion.Value[3].z)) + (CameraMotion.Value[3].w * PosN.x) + (CameraMotion.Value[4].x * PosN.y * PosN.y) + (CameraMotion.Value[4].y * PosN.z) + CameraMotion.Value[4].z) * ScaleM;
+
+		//Velocity[i].x += 0.3;
+		//Velocity[i].y += 0.1;
+	}
+
+	const float AspectRatio = ScreenPosToPixel.y / ScreenPosToPixel.x;
+	Velocity.y *= AspectRatio;
+
+	float2 VelocityPolar = CartesianToPolar( Velocity );
 	
-	float4 DepthSamples;
-	float2 VelocitySamples[4];
-#if 0
-	// using Gather: xyzw in counter clockwise order starting with the sample to the lower left of the queried location
-	float4 X = PostprocessInput0.GatherRed( PostprocessInput0Sampler, UV ).wzxy;
-	float4 Y = PostprocessInput0.GatherGreen( PostprocessInput0Sampler, UV ).wzxy;
-
-	VelocitySamples[0] = float2(X.x, Y.x);
-	VelocitySamples[1] = float2(X.y, Y.y);
-	VelocitySamples[2] = float2(X.z, Y.z);
-	VelocitySamples[3] = float2(X.w, Y.w);
+	bool bInsideViewport = all( PixelPos.xy < ViewportRect.zw );
 	
-	DepthSamples = PostprocessInput1.GatherRed( PostprocessInput1Sampler, UV, 0 ).wzxy;
+#if 1
+	// 10:10:10:2  (VelocityLength, VelocityAngle, Depth)
+	float2 EncodedVelocity;
+	EncodedVelocity.x = VelocityPolar.x;
+	EncodedVelocity.y = VelocityPolar.y * (0.5 / PI) + 0.5;
 #else
-	VelocitySamples[0] = PostprocessInput0[ PixelPos + Offsets[0] ].xy;
-	VelocitySamples[1] = PostprocessInput0[ PixelPos + Offsets[1] ].xy;
-	VelocitySamples[2] = PostprocessInput0[ PixelPos + Offsets[2] ].xy;
-	VelocitySamples[3] = PostprocessInput0[ PixelPos + Offsets[3] ].xy;
-	
-	DepthSamples[0] = PostprocessInput1[ PixelPos + Offsets[0] ].x;
-	DepthSamples[1] = PostprocessInput1[ PixelPos + Offsets[1] ].x;
-	DepthSamples[2] = PostprocessInput1[ PixelPos + Offsets[2] ].x;
-	DepthSamples[3] = PostprocessInput1[ PixelPos + Offsets[3] ].x;
+	// 11:11:10 (Velocity.xy, Depth)
+	// Only handles values in range (-2,2)
+	Velocity = clamp( Velocity, -1.9999, 1.9999 );
+	// Kill denormals
+	Velocity = abs( Velocity ) < exp2(-14.0) ? 0 : Velocity;
+	// Store sign in top bit of 6e5 exponent
+	float2 EncodedVelocity = Velocity * ( Velocity < 0 ? -32768.0 : 1 );
 #endif
-
-	float2 Velocity[4];
-	UNROLL
-	for( uint i = 0; i < 4; i++ )
+	
+	BRANCH
+	if( bInsideViewport )
 	{
-		if( VelocitySamples[i].x > 0.0 )
-		{
-			Velocity[i] = DecodeVelocityFromTexture( VelocitySamples[i] );
-		}
-		else
-		{
-			float3 PosN = float3( ScreenUV, DepthSamples[i] );
-			float ScaleM = rcp( dot( PosN, CameraMotion.Value[0].xyz ) + CameraMotion.Value[0].w );
-			Velocity[i].x = -2.0 * ((PosN.x * ((CameraMotion.Value[1].x * PosN.y) + (CameraMotion.Value[1].y * PosN.z) + CameraMotion.Value[1].z)) + (CameraMotion.Value[1].w * PosN.y) + (CameraMotion.Value[2].x * PosN.x * PosN.x) + (CameraMotion.Value[2].y * PosN.z) + CameraMotion.Value[2].z) * ScaleM;
-			Velocity[i].y =  2.0 * ((PosN.y * ((CameraMotion.Value[3].x * PosN.x) + (CameraMotion.Value[3].y * PosN.z) + CameraMotion.Value[3].z)) + (CameraMotion.Value[3].w * PosN.x) + (CameraMotion.Value[4].x * PosN.y * PosN.y) + (CameraMotion.Value[4].y * PosN.z) + CameraMotion.Value[4].z) * ScaleM;
-		}
-
-		BRANCH
-		if( all( PixelPos.xy + Offsets[i] < VelocityFlatten.ViewDimensions.zw ) )
-		{
-			OutVelocityFlat[ PixelPos + Offsets[i] ] = EncodeVelocityToTexture( Velocity[i] ).xyxy;
-			//OutPackedVelocityDepth[ PixelPos + Offsets[i] ] = float2( Velocity[i], DepthSamples[i] ).xyxy;
-		}
+		OutVelocityFlat[ PixelPos ] = float3( EncodedVelocity, ConvertFromDeviceZ( Depth ) ).xyzz;
 	}
 
-	Shared[ GroupIndex ] = MaxLength(
-		MaxLength( Velocity[0], Velocity[1] ),
-		MaxLength( Velocity[2], Velocity[3] )
-	);
+#if 1
+	// Limit velocity
+	VelocityPolar.x = min( VelocityPolar.x, 100 / ScreenPosToPixel.x );	// TODO input limit
+#endif
+
+	float4 VelocityMinMax = VelocityPolar.xyxy;
+	VelocityMinMax.x = bInsideViewport ? VelocityMinMax.x : 2;
+	VelocityMinMax.z = bInsideViewport ? VelocityMinMax.z : 0;
+
+	Shared[ GroupIndex ] = VelocityMinMax;
 
 	GroupMemoryBarrierWithGroupSync();
 
-#if 0
-	if( GroupIndex < 512 )	Shared[ GroupIndex ] = MaxLength( Shared[ GroupIndex ], Shared[ GroupIndex + 512 ] );
+#if THREADGROUP_TOTALSIZE > 512
+	if( GroupIndex < 512 )	Shared[ GroupIndex ] = MinMaxLengthPolar( Shared[ GroupIndex ], Shared[ GroupIndex + 512 ] );
 	GroupMemoryBarrierWithGroupSync();
-	if( GroupIndex < 256 )	Shared[ GroupIndex ] = MaxLength( Shared[ GroupIndex ], Shared[ GroupIndex + 256 ] );
+#endif
+#if THREADGROUP_TOTALSIZE > 256
+	if( GroupIndex < 256 )	Shared[ GroupIndex ] = MinMaxLengthPolar( Shared[ GroupIndex ], Shared[ GroupIndex + 256 ] );
 	GroupMemoryBarrierWithGroupSync();
-	if( GroupIndex < 128 )	Shared[ GroupIndex ] = MaxLength( Shared[ GroupIndex ], Shared[ GroupIndex + 128 ] );
+#endif
+#if THREADGROUP_TOTALSIZE > 128
+	if( GroupIndex < 128 )	Shared[ GroupIndex ] = MinMaxLengthPolar( Shared[ GroupIndex ], Shared[ GroupIndex + 128 ] );
 	GroupMemoryBarrierWithGroupSync();
-	if( GroupIndex <  64 )	Shared[ GroupIndex ] = MaxLength( Shared[ GroupIndex ], Shared[ GroupIndex +  64 ] );
+#endif
+#if THREADGROUP_TOTALSIZE > 64
+	if( GroupIndex <  64 )	Shared[ GroupIndex ] = MinMaxLengthPolar( Shared[ GroupIndex ], Shared[ GroupIndex +  64 ] );
 	GroupMemoryBarrierWithGroupSync();
 #endif
 
 	// Safe for vector sizes 32 or larger, AMD and NV
 	// TODO Intel variable size vector
-	if( GroupIndex < 32 )	Shared[ GroupIndex ] = MaxLength( Shared[ GroupIndex ], Shared[ GroupIndex + 32 ] );
-	if( GroupIndex < 16 )	Shared[ GroupIndex ] = MaxLength( Shared[ GroupIndex ], Shared[ GroupIndex + 16 ] );
-	if( GroupIndex <  8 )	Shared[ GroupIndex ] = MaxLength( Shared[ GroupIndex ], Shared[ GroupIndex +  8 ] );
-	if( GroupIndex <  4 )	Shared[ GroupIndex ] = MaxLength( Shared[ GroupIndex ], Shared[ GroupIndex +  4 ] );
-	if( GroupIndex <  2 )	Shared[ GroupIndex ] = MaxLength( Shared[ GroupIndex ], Shared[ GroupIndex +  2 ] );
-	if( GroupIndex <  1 )	Shared[ GroupIndex ] = MaxLength( Shared[ GroupIndex ], Shared[ GroupIndex +  1 ] );
+	if( GroupIndex < 32 )	Shared[ GroupIndex ] = MinMaxLengthPolar( Shared[ GroupIndex ], Shared[ GroupIndex + 32 ] );
+	if( GroupIndex < 16 )	Shared[ GroupIndex ] = MinMaxLengthPolar( Shared[ GroupIndex ], Shared[ GroupIndex + 16 ] );
+	if( GroupIndex <  8 )	Shared[ GroupIndex ] = MinMaxLengthPolar( Shared[ GroupIndex ], Shared[ GroupIndex +  8 ] );
+	if( GroupIndex <  4 )	Shared[ GroupIndex ] = MinMaxLengthPolar( Shared[ GroupIndex ], Shared[ GroupIndex +  4 ] );
+	if( GroupIndex <  2 )	Shared[ GroupIndex ] = MinMaxLengthPolar( Shared[ GroupIndex ], Shared[ GroupIndex +  2 ] );
+	if( GroupIndex <  1 )	Shared[ GroupIndex ] = MinMaxLengthPolar( Shared[ GroupIndex ], Shared[ GroupIndex +  1 ] );
 
 	if( GroupIndex == 0 )
 	{
-		OutMaxTileVelocity[ GroupId.xy ] = Shared[0].xyxy;
+		OutMaxTileVelocity[ GroupId.xy ] = float4( PolarToCartesian( Shared[0].xy ), PolarToCartesian( Shared[0].zw ) );
+	}
+}
+
+
+
+RWTexture2D<float4>	OutScatteredMaxVelocity;
+
+[numthreads(16, 16, 1)]
+void VelocityGatherCS(
+	uint3 GroupId : SV_GroupID,
+	uint3 DispatchThreadId : SV_DispatchThreadID,
+	uint3 GroupThreadId : SV_GroupThreadID,
+	uint GroupIndex : SV_GroupIndex ) 
+{
+	uint2 PixelPos = DispatchThreadId.xy;
+	float2 UV = ( (float2)PixelPos.xy + 0.5 ) * PostprocessInput0Size.zw;
+
+	float4 MinMaxVelocity = PostprocessInput0[ PixelPos ];
+
+#if 1
+	// Scatter as gather
+	for( int x = -3; x <= 3; x++ )
+	{
+		for( int y = -3; y <= 3; y++ )
+		{
+			if( x == 0 && y == 0 )
+				continue;
+
+			int2 Offset = int2(x,y);
+			int2 SampleIndex = PixelPos + Offset;
+			
+			bool2 bInsideViewport = 0 <= SampleIndex && SampleIndex < ViewportRect.zw;
+			if( !all( bInsideViewport ) )
+				continue;
+
+			float4 ScatterMinMax = PostprocessInput0[ SampleIndex ];
+			//float4 ScatterMinMax = PostprocessInput0.SampleLevel( PostprocessInput0Sampler, UV, 0, Offset );
+			float2 MaxVelocity = ScatterMinMax.zw;
+
+			// TODO move to velocity flatten shader
+			//float2 VelocityPixels = MaxVelocity * ScreenPosToPixel.x;// * VelocityScale.xy;
+			float2 VelocityPixels = MaxVelocity * 0.5 * 1920.0 / 16.0;
+			VelocityPixels *= 0.5;
+
+			float  VelocityLengthPixelsSqr = dot( VelocityPixels, VelocityPixels );
+			float  VelocityLengthPixelsInv = rsqrtFast( VelocityLengthPixelsSqr + 1e-8 );
+			float  VelocityLengthPixels = VelocityLengthPixelsSqr * VelocityLengthPixelsInv;
+			float2 VelocityDir = VelocityPixels * VelocityLengthPixelsInv;
+
+			// Project pixel corner on to dir. This is the oriented extent of a pixel.
+			// 1/2 pixel because shape is swept tile
+			// +1/2 pixel for conservative rasterization
+			// 99% to give epsilon before neighbor is filled. Otherwise all neighbors lie on edges of quad when no velocity in their direction.
+			float PixelExtent = abs( VelocityDir.x ) + abs( VelocityDir.y );
+			float2 QuadExtent = float2( VelocityLengthPixels, 0 ) + PixelExtent.xx * 0.99;
+
+			// Orient quad along velocity direction
+			float2 AxisX = VelocityDir;
+			float2 AxisY = float2( -VelocityDir.y, VelocityDir.x );
+
+			// Project this pixel center onto scatter quad
+			float2 PixelCenterOnQuad;
+			PixelCenterOnQuad.x = dot( AxisX, Offset );
+			PixelCenterOnQuad.y = dot( AxisY, Offset );
+
+			bool2 bInsideQuad = abs( PixelCenterOnQuad ) < QuadExtent;
+			if( all( bInsideQuad ) ) 
+			{
+				MinMaxVelocity = MinMaxLength( MinMaxVelocity, ScatterMinMax );
+			}
+		}
+	}
+#endif
+
+	BRANCH
+	if( all( PixelPos.xy < ViewportRect.zw ) )
+	{
+		OutScatteredMaxVelocity[ PixelPos ] = MinMaxVelocity;
 	}
 }
\ No newline at end of file
diff --git a/Engine/Source/Runtime/Renderer/Private/PostProcess/PostProcessMotionBlur.cpp b/Engine/Source/Runtime/Renderer/Private/PostProcess/PostProcessMotionBlur.cpp
--- a/Engine/Source/Runtime/Renderer/Private/PostProcess/PostProcessMotionBlur.cpp
+++ b/Engine/Source/Runtime/Renderer/Private/PostProcess/PostProcessMotionBlur.cpp
@@ -678,12 +678,13 @@ FPooledRenderTargetDesc FRCPassPostProcessMotionBlurRecombine::ComputeOutputDesc
 	return Ret;
 }
 
-/** The uniform shader parameters used in FPostProcessVelocityFlattenCS. */
-BEGIN_UNIFORM_BUFFER_STRUCT(FVelocityFlattenParameters,)
-	DECLARE_UNIFORM_BUFFER_STRUCT_MEMBER(FIntRect,ViewDimensions)	// View rect min in xy, max in zw.
-END_UNIFORM_BUFFER_STRUCT(FVelocityFlattenParameters)
 
-IMPLEMENT_UNIFORM_BUFFER_STRUCT(FVelocityFlattenParameters,TEXT("VelocityFlatten"));
+FIntPoint GetNumTiles16x16( FIntPoint PixelExtent )
+{
+	uint32 TilesX = (PixelExtent.X + 15) / 16;
+	uint32 TilesY = (PixelExtent.Y + 15) / 16;
+	return FIntPoint( TilesX, TilesY );
+}
 
 class FPostProcessVelocityFlattenCS : public FGlobalShader
 {
@@ -697,16 +698,12 @@ class FPostProcessVelocityFlattenCS : public FGlobalShader
 	static void ModifyCompilationEnvironment(EShaderPlatform Platform, FShaderCompilerEnvironment& OutEnvironment)
 	{
 		FGlobalShader::ModifyCompilationEnvironment(Platform, OutEnvironment);
-		OutEnvironment.SetDefine(TEXT("THREADGROUP_SIZEX"), FRCPassPostProcessVelocityFlatten::ThreadGroupSizeX);
-		OutEnvironment.SetDefine(TEXT("THREADGROUP_SIZEY"), FRCPassPostProcessVelocityFlatten::ThreadGroupSizeY);
-		OutEnvironment.CompilerFlags.Add( CFLAG_StandardOptimization );
 	}
 
 	FPostProcessVelocityFlattenCS() {}
 
 public:
 	FShaderParameter		OutVelocityFlat;		// UAV
-	FShaderParameter		OutPackedVelocityDepth;	// unused
 	FShaderParameter		OutMaxTileVelocity;		// UAV
 
 	FPostProcessVelocityFlattenCS(const ShaderMetaType::CompiledShaderInitializerType& Initializer)
@@ -714,7 +711,6 @@ public:
 	{
 		PostprocessParameter.Bind(Initializer.ParameterMap);
 		OutVelocityFlat.Bind(Initializer.ParameterMap, TEXT("OutVelocityFlat"));
-		OutPackedVelocityDepth.Bind(Initializer.ParameterMap, TEXT("OutPackedVelocityDepth"));
 		OutMaxTileVelocity.Bind(Initializer.ParameterMap, TEXT("OutMaxTileVelocity"));
 	}
 
@@ -726,15 +722,6 @@ public:
 		PostprocessParameter.SetCS(ShaderRHI, Context, TStaticSamplerState<SF_Point,AM_Clamp,AM_Clamp,AM_Clamp>::GetRHI());
 
 		SetUniformBufferParameter(RHICmdList, ShaderRHI, GetUniformBufferParameter<FCameraMotionParameters>(), CreateCameraMotionParametersUniformBuffer(Context.View));
-		
-		{
-			FVelocityFlattenParameters Data;
-
-			Data.ViewDimensions = View.ViewRect;
-
-			TUniformBufferRef<FVelocityFlattenParameters> UniformBuffer = TUniformBufferRef<FVelocityFlattenParameters>::CreateUniformBufferImmediate(Data, UniformBuffer_SingleFrame);
-			SetUniformBufferParameter(RHICmdList, ShaderRHI, GetUniformBufferParameter<FVelocityFlattenParameters>(), UniformBuffer);
-		}
 	}
 	
 	virtual bool Serialize(FArchive& Ar) override
@@ -742,7 +729,6 @@ public:
 		bool bShaderHasOutdatedParameters = FGlobalShader::Serialize(Ar);
 		Ar << PostprocessParameter;
 		Ar << OutVelocityFlat;
-		Ar << OutPackedVelocityDepth;
 		Ar << OutMaxTileVelocity;
 		return bShaderHasOutdatedParameters;
 	}
@@ -754,9 +740,13 @@ private:
 IMPLEMENT_SHADER_TYPE(,FPostProcessVelocityFlattenCS,TEXT("PostProcessVelocityFlatten"),TEXT("VelocityFlattenMain"),SF_Compute);
 
 
+FRCPassPostProcessVelocityFlatten::FRCPassPostProcessVelocityFlatten()
+	: AsyncJobFenceID(-1)
+{}
+
 void FRCPassPostProcessVelocityFlatten::Process(FRenderingCompositePassContext& Context)
 {
-	SCOPED_DRAW_EVENT(Context.RHICmdList, PostProcessVelocityFlatten);
+	SCOPED_DRAW_EVENT(Context.RHICmdList, VelocityFlatten);
 	const FPooledRenderTargetDesc* InputDesc = GetInputDesc(ePId_Input0);
 
 	if (AsyncJobFenceID != -1)
@@ -779,7 +769,6 @@ void FRCPassPostProcessVelocityFlatten::Process(FRenderingCompositePassContext&
 	
 	const FSceneRenderTargetItem& DestRenderTarget0 = PassOutputs[0].RequestSurface(Context);
 	const FSceneRenderTargetItem& DestRenderTarget1 = PassOutputs[1].RequestSurface(Context);
-	//const FSceneRenderTargetItem& DestRenderTarget2 = PassOutputs[2].RequestSurface(Context);
 
 	TShaderMapRef< FPostProcessVelocityFlattenCS > ComputeShader( Context.GetShaderMap() );
 
@@ -793,16 +782,16 @@ void FRCPassPostProcessVelocityFlatten::Process(FRenderingCompositePassContext&
 		Context.RHICmdList.BeginAsyncComputeJob_DrawThread(AsyncComputePriority_Default);
 	}
 
+	Context.SetViewportAndCallRHI( View.ViewRect );
 	Context.RHICmdList.SetComputeShader(ComputeShader->GetComputeShader());
 
 	// set destination
 	Context.RHICmdList.SetUAVParameter(ComputeShader->GetComputeShader(), ComputeShader->OutVelocityFlat.GetBaseIndex(), DestRenderTarget0.UAV);
-	//Context.RHICmdList.SetUAVParameter(ComputeShader->GetComputeShader(), ComputeShader->OutPackedVelocityDepth.GetBaseIndex(), DestRenderTarget1.UAV);
 	Context.RHICmdList.SetUAVParameter(ComputeShader->GetComputeShader(), ComputeShader->OutMaxTileVelocity.GetBaseIndex(), DestRenderTarget1.UAV);
 
 	ComputeShader->SetCS(Context.RHICmdList, Context, View);
 
-	FIntPoint ThreadGroupCountValue = ComputeThreadGroupCount(View.ViewRect.Size());
+	FIntPoint ThreadGroupCountValue = GetNumTiles16x16(View.ViewRect.Size());
 	DispatchComputeShader(Context.RHICmdList, *ComputeShader, ThreadGroupCountValue.X, ThreadGroupCountValue.Y, 1);
 
 #if USE_ASYNC_COMPUTE_CONTEXT
@@ -818,12 +807,10 @@ void FRCPassPostProcessVelocityFlatten::Process(FRenderingCompositePassContext&
 
 	// un-set destination
 	Context.RHICmdList.SetUAVParameter(ComputeShader->GetComputeShader(), ComputeShader->OutVelocityFlat.GetBaseIndex(), NULL);
-	//Context.RHICmdList.SetUAVParameter(ComputeShader->GetComputeShader(), ComputeShader->OutPackedVelocityDepth.GetBaseIndex(), NULL);
 	Context.RHICmdList.SetUAVParameter(ComputeShader->GetComputeShader(), ComputeShader->OutMaxTileVelocity.GetBaseIndex(), NULL);
 
 	Context.RHICmdList.CopyToResolveTarget(DestRenderTarget0.TargetableTexture, DestRenderTarget0.ShaderResourceTexture, false, FResolveParams());
 	Context.RHICmdList.CopyToResolveTarget(DestRenderTarget1.TargetableTexture, DestRenderTarget1.ShaderResourceTexture, false, FResolveParams());
-	//Context.RHICmdList.CopyToResolveTarget(DestRenderTarget1.TargetableTexture, DestRenderTarget2.ShaderResourceTexture, false, FResolveParams());
 
 	// we want this pass to be executed another time
 	// so we do this: CompositeContext.Process(VelocityFlattenPass, TEXT("VelocityFlattenPass"));
@@ -842,14 +829,6 @@ void FRCPassPostProcessVelocityFlatten::Process(FRenderingCompositePassContext&
 	}
 }
 
-FIntPoint FRCPassPostProcessVelocityFlatten::ComputeThreadGroupCount(FIntPoint PixelExtent)
-{
-	uint32 ThreadGroupCountX = (PixelExtent.X + TileSizeX - 1) / TileSizeX;
-	uint32 ThreadGroupCountY = (PixelExtent.Y + TileSizeY - 1) / TileSizeY;
-
-	return FIntPoint(ThreadGroupCountX, ThreadGroupCountY);
-}
-
 FPooledRenderTargetDesc FRCPassPostProcessVelocityFlatten::ComputeOutputDesc(EPassOutputId InPassOutputId) const
 {
 	if( InPassOutputId == ePId_Output0 )
@@ -858,25 +837,13 @@ FPooledRenderTargetDesc FRCPassPostProcessVelocityFlatten::ComputeOutputDesc(EPa
 		FPooledRenderTargetDesc Ret = PassInputs[0].GetOutput()->RenderTargetDesc;
 		Ret.Reset();
 
+		Ret.Format = PF_FloatR11G11B10;
 		Ret.TargetableFlags |= TexCreate_UAV;
 		Ret.TargetableFlags |= TexCreate_RenderTargetable;
 		Ret.DebugName = TEXT("VelocityFlat");
 
 		return Ret;
 	}
-	/*else if( InPassOutputId == ePId_Output1 )
-	{
-		// Packed VelocityLength, Depth
-		FPooledRenderTargetDesc Ret = PassInputs[0].GetOutput()->RenderTargetDesc;
-		Ret.Reset();
-
-		Ret.Format = PF_R8G8;
-		Ret.TargetableFlags |= TexCreate_UAV;
-		Ret.TargetableFlags |= TexCreate_RenderTargetable;
-		Ret.DebugName = TEXT("PackedVelocityDepth");
-
-		return Ret;
-	}*/
 	else
 	{
 		// Max tile velocity
@@ -884,12 +851,9 @@ FPooledRenderTargetDesc FRCPassPostProcessVelocityFlatten::ComputeOutputDesc(EPa
 		UnmodifiedRet.Reset();
 
 		FIntPoint PixelExtent = UnmodifiedRet.Extent;
+		FIntPoint TileCount = GetNumTiles16x16(PixelExtent);
 
-		FIntPoint ThreadGroupCount = ComputeThreadGroupCount(PixelExtent);
-		FIntPoint NewSize = ThreadGroupCount;
-
-		// format can be optimized later
-		FPooledRenderTargetDesc Ret(FPooledRenderTargetDesc::Create2DDesc(NewSize, PF_G16R16F, TexCreate_None, TexCreate_RenderTargetable | TexCreate_UAV, false));
+		FPooledRenderTargetDesc Ret(FPooledRenderTargetDesc::Create2DDesc(TileCount, PF_FloatRGBA, TexCreate_None, TexCreate_RenderTargetable | TexCreate_UAV, false));
 
 		Ret.DebugName = TEXT("MaxVelocity");
 
@@ -952,7 +916,7 @@ class FPostProcessVelocityScatterVS : public FGlobalShader
 	
 public:
 	FPostProcessPassParameters PostprocessParameter;
-	FShaderParameter TileCount;
+	FShaderParameter DrawMax;
 	FShaderParameter VelocityScale;
 
 	/** Initialization constructor. */
@@ -960,7 +924,7 @@ public:
 		: FGlobalShader(Initializer)
 	{
 		PostprocessParameter.Bind(Initializer.ParameterMap);
-		TileCount.Bind(Initializer.ParameterMap, TEXT("TileCount"));
+		DrawMax.Bind(Initializer.ParameterMap, TEXT("bDrawMax"));
 		VelocityScale.Bind( Initializer.ParameterMap, TEXT("VelocityScale") );
 	}
 
@@ -968,19 +932,19 @@ public:
 	virtual bool Serialize(FArchive& Ar) override
 	{
 		bool bShaderHasOutdatedParameters = FGlobalShader::Serialize(Ar);
-		Ar << PostprocessParameter << TileCount << VelocityScale;
+		Ar << PostprocessParameter << DrawMax << VelocityScale;
 		return bShaderHasOutdatedParameters;
 	}
 
 	/** to have a similar interface as all other shaders */
-	void SetParameters(const FRenderingCompositePassContext& Context, FIntPoint TileCountValue)
+	void SetParameters(const FRenderingCompositePassContext& Context, int32 bDrawMax)
 	{
 		const FVertexShaderRHIParamRef ShaderRHI = GetVertexShader();
 
 		FGlobalShader::SetParameters(Context.RHICmdList, ShaderRHI, Context.View);
 		PostprocessParameter.SetVS(ShaderRHI, Context, TStaticSamplerState<SF_Point,AM_Clamp,AM_Clamp,AM_Clamp>::GetRHI());
 
-		SetShaderValue(Context.RHICmdList, ShaderRHI, TileCount, TileCountValue);
+		SetShaderValue(Context.RHICmdList, ShaderRHI, DrawMax, bDrawMax);
 
 		{
 			const FSceneViewState* ViewState = (FSceneViewState*) Context.View.State;
@@ -1031,14 +995,10 @@ public:
 IMPLEMENT_SHADER_TYPE(,FPostProcessVelocityScatterVS,TEXT("PostProcessMotionBlur"),TEXT("VelocityScatterVS"),SF_Vertex);
 IMPLEMENT_SHADER_TYPE(,FPostProcessVelocityScatterPS,TEXT("PostProcessMotionBlur"),TEXT("VelocityScatterPS"),SF_Pixel);
 
-FRCPassPostProcessVelocityFlatten::FRCPassPostProcessVelocityFlatten()
-	: AsyncJobFenceID(-1)
-{
-}
 
 void FRCPassPostProcessVelocityScatter::Process(FRenderingCompositePassContext& Context)
 {
-	SCOPED_DRAW_EVENT(Context.RHICmdList, PassPostProcessVelocityScatter);
+	SCOPED_DRAW_EVENT(Context.RHICmdList, VelocityScatter);
 
 	const FPooledRenderTargetDesc* InputDesc = GetInputDesc(ePId_Input0);
 	
@@ -1053,49 +1013,59 @@ void FRCPassPostProcessVelocityScatter::Process(FRenderingCompositePassContext&
 	FIntPoint SrcSize = InputDesc->Extent;
 	FIntPoint DestSize = PassOutputs[0].RenderTargetDesc.Extent;
 
-	// e.g. 4 means the input texture is 4x smaller than the buffer size
-	uint32 ScaleFactor = FSceneRenderTargets::Get(Context.RHICmdList).GetBufferSizeXY().X / SrcSize.X;
-
-	FIntRect SrcRect = FIntRect::DivideAndRoundUp(View.ViewRect, ScaleFactor);
-	FIntRect DestRect = SrcRect;
+	FIntPoint TileCount = GetNumTiles16x16( View.ViewRect.Size() );
 
 	const FSceneRenderTargetItem& DestRenderTarget = PassOutputs[0].RequestSurface(Context);
 
 	TRefCountPtr<IPooledRenderTarget> DepthTarget;
-	
-	FPooledRenderTargetDesc Desc( FPooledRenderTargetDesc::Create2DDesc( DestSize, PF_ShadowDepth, TexCreate_None, TexCreate_DepthStencilTargetable, false ) );
+	FPooledRenderTargetDesc Desc( FPooledRenderTargetDesc::Create2DDesc( DestSize, PF_ShadowDepth, TexCreate_None, TexCreate_DepthStencilTargetable | TexCreate_NoFastClear, false ) );
 	GRenderTargetPool.FindFreeElement( Desc, DepthTarget, TEXT("VelocityScatterDepth") );
 
 	// Set the view family's render target/viewport.
 	SetRenderTarget( Context.RHICmdList, DestRenderTarget.TargetableTexture, DepthTarget->GetRenderTargetItem().TargetableTexture );
+	
+	Context.SetViewportAndCallRHI( 0, 0, 0.0f, TileCount.X, TileCount.Y, 1.0f );
+	
+	// Min,Max
+	for( int i = 0; i < 2; i++ )
+	{
+		if( i == 0 )
+		{
+			// clear depth
+			// Max >= Min so no need to clear on second pass
+			Context.RHICmdList.Clear( false, FLinearColor::Black, true, 1.0f, false, 0, FIntRect() );
+		}
 
-	Context.RHICmdList.Clear( true, FLinearColor::Black, true, 0.0f, false, 0, FIntRect() );
-
-	Context.SetViewportAndCallRHI(SrcRect);
-
-	// set the state
-	Context.RHICmdList.SetBlendState(TStaticBlendState<>::GetRHI());
-	Context.RHICmdList.SetRasterizerState(TStaticRasterizerState<>::GetRHI());
-	//@todo-briank: Should this be CF_DepthNear?
-	static_assert((int32)ERHIZBuffer::IsInverted != 0, "Should this be CF_DepthNear?");
-	Context.RHICmdList.SetDepthStencilState(TStaticDepthStencilState<true, CF_Greater>::GetRHI());
-
-	TShaderMapRef< FPostProcessVelocityScatterVS > VertexShader(Context.GetShaderMap());
-	TShaderMapRef< FPostProcessVelocityScatterPS > PixelShader(Context.GetShaderMap());
+		if( i == 0 )
+		{
+			// min
+			Context.RHICmdList.SetBlendState( TStaticBlendStateWriteMask< CW_RGBA >::GetRHI() );
+			Context.RHICmdList.SetRasterizerState( TStaticRasterizerState<>::GetRHI() );
+			Context.RHICmdList.SetDepthStencilState( TStaticDepthStencilState< true, CF_Less >::GetRHI() );
+		}
+		else
+		{
+			// max
+			Context.RHICmdList.SetBlendState( TStaticBlendStateWriteMask< CW_BA >::GetRHI() );
+			Context.RHICmdList.SetRasterizerState( TStaticRasterizerState<>::GetRHI() );
+			Context.RHICmdList.SetDepthStencilState( TStaticDepthStencilState< true, CF_Greater >::GetRHI() );
+		}
 
-	static FGlobalBoundShaderState BoundShaderState;
-	SetGlobalBoundShaderState(Context.RHICmdList, Context.GetFeatureLevel(), BoundShaderState, GFilterVertexDeclaration.VertexDeclarationRHI, *VertexShader, *PixelShader);
+		TShaderMapRef< FPostProcessVelocityScatterVS > VertexShader(Context.GetShaderMap());
+		TShaderMapRef< FPostProcessVelocityScatterPS > PixelShader(Context.GetShaderMap());
 
-	FIntPoint TileCount = SrcRect.Size();
+		static FGlobalBoundShaderState BoundShaderState;
+		SetGlobalBoundShaderState(Context.RHICmdList, Context.GetFeatureLevel(), BoundShaderState, GFilterVertexDeclaration.VertexDeclarationRHI, *VertexShader, *PixelShader);
 
-	VertexShader->SetParameters( Context, TileCount );
-	PixelShader->SetParameters( Context );
+		VertexShader->SetParameters( Context, i );
+		PixelShader->SetParameters( Context );
 
-	// needs to be the same on shader side (faster on NVIDIA and AMD)
-	int32 QuadsPerInstance = 8;
+		// needs to be the same on shader side (faster on NVIDIA and AMD)
+		int32 QuadsPerInstance = 8;
 
-	Context.RHICmdList.SetStreamSource(0, NULL, 0, 0);
-	Context.RHICmdList.DrawIndexedPrimitive(GScatterQuadIndexBuffer.IndexBufferRHI, PT_TriangleList, 0, 0, 32, 0, 2 * QuadsPerInstance, FMath::DivideAndRoundUp(TileCount.X * TileCount.Y, QuadsPerInstance));
+		Context.RHICmdList.SetStreamSource(0, NULL, 0, 0);
+		Context.RHICmdList.DrawIndexedPrimitive(GScatterQuadIndexBuffer.IndexBufferRHI, PT_TriangleList, 0, 0, 32, 0, 2 * QuadsPerInstance, FMath::DivideAndRoundUp(TileCount.X * TileCount.Y, QuadsPerInstance));
+	}
 
 	Context.RHICmdList.CopyToResolveTarget(DestRenderTarget.TargetableTexture, DestRenderTarget.ShaderResourceTexture, false, FResolveParams());
 }
@@ -1111,15 +1081,9 @@ FPooledRenderTargetDesc FRCPassPostProcessVelocityScatter::ComputeOutputDesc(EPa
 }
 
 
-
-
-
-
-
-
-class FPostProcessVelocityDilatePS : public FGlobalShader
+class FPostProcessVelocityGatherCS : public FGlobalShader
 {
-	DECLARE_SHADER_TYPE(FPostProcessVelocityDilatePS, Global);
+	DECLARE_SHADER_TYPE(FPostProcessVelocityGatherCS, Global);
 
 	static bool ShouldCache(EShaderPlatform Platform)
 	{
@@ -1132,16 +1096,17 @@ class FPostProcessVelocityDilatePS : public FGlobalShader
 	}
 
 	/** Default constructor. */
-	FPostProcessVelocityDilatePS() {}
+	FPostProcessVelocityGatherCS() {}
 
 public:
-	FPostProcessPassParameters PostprocessParameter;
+	FShaderParameter		OutScatteredMaxVelocity;
 
 	/** Initialization constructor. */
-	FPostProcessVelocityDilatePS(const ShaderMetaType::CompiledShaderInitializerType& Initializer)
+	FPostProcessVelocityGatherCS(const ShaderMetaType::CompiledShaderInitializerType& Initializer)
 		: FGlobalShader(Initializer)
 	{
 		PostprocessParameter.Bind(Initializer.ParameterMap);
+		OutScatteredMaxVelocity.Bind(Initializer.ParameterMap, TEXT("OutScatteredMaxVelocity"));
 	}
 
 	// FShader interface.
@@ -1149,22 +1114,26 @@ public:
 	{
 		bool bShaderHasOutdatedParameters = FGlobalShader::Serialize(Ar);
 		Ar << PostprocessParameter;
+		Ar << OutScatteredMaxVelocity;
 		return bShaderHasOutdatedParameters;
 	}
 
 	void SetParameters(const FRenderingCompositePassContext& Context)
 	{
-		const FPixelShaderRHIParamRef ShaderRHI = GetPixelShader();
+		const FComputeShaderRHIParamRef ShaderRHI = GetComputeShader();
 
 		FGlobalShader::SetParameters(Context.RHICmdList, ShaderRHI, Context.View);
 
-		PostprocessParameter.SetPS(ShaderRHI, Context, TStaticSamplerState<SF_Point,AM_Clamp,AM_Clamp,AM_Clamp>::GetRHI());
+		PostprocessParameter.SetCS(ShaderRHI, Context, TStaticSamplerState<SF_Point,AM_Clamp,AM_Clamp,AM_Clamp>::GetRHI());
 	}
+	
+private:
+	FPostProcessPassParameters PostprocessParameter;
 };
 
-IMPLEMENT_SHADER_TYPE(,FPostProcessVelocityDilatePS,TEXT("PostProcessMotionBlur"),TEXT("VelocityDilatePS"),SF_Pixel);
+IMPLEMENT_SHADER_TYPE(,FPostProcessVelocityGatherCS,TEXT("PostProcessVelocityFlatten"),TEXT("VelocityGatherCS"),SF_Compute);
 
-void FRCPassPostProcessVelocityDilate::Process(FRenderingCompositePassContext& Context)
+void FRCPassPostProcessVelocityGather::Process(FRenderingCompositePassContext& Context)
 {
 	SCOPED_DRAW_EVENT(Context.RHICmdList, VelocityDilate);
 
@@ -1178,60 +1147,37 @@ void FRCPassPostProcessVelocityDilate::Process(FRenderingCompositePassContext& C
 
 	const FSceneView& View = Context.View;
 
-	FIntPoint SrcSize = InputDesc->Extent;
-	FIntPoint DestSize = PassOutputs[0].RenderTargetDesc.Extent;
-
-	// e.g. 4 means the input texture is 4x smaller than the buffer size
-	uint32 ScaleFactor = FSceneRenderTargets::Get(Context.RHICmdList).GetBufferSizeXY().X / SrcSize.X;
-
-	FIntRect SrcRect = FIntRect::DivideAndRoundUp(View.ViewRect, ScaleFactor);
-	FIntRect DestRect = SrcRect;
-
+	FIntPoint TileCount = GetNumTiles16x16( View.ViewRect.Size() );
+	
 	const FSceneRenderTargetItem& DestRenderTarget = PassOutputs[0].RequestSurface(Context);
 
-	// Set the view family's render target/viewport.
-	SetRenderTarget(Context.RHICmdList, DestRenderTarget.TargetableTexture, FTextureRHIRef());
-
-	// is optimized away if possible (RT size=view size, )
-	//Context.RHICmdList.Clear(true, FLinearColor::Black, false, 1.0f, false, 0, SrcRect);
-
-	Context.SetViewportAndCallRHI(SrcRect);
-
-	// set the state
-	Context.RHICmdList.SetBlendState(TStaticBlendState<>::GetRHI());
-	Context.RHICmdList.SetRasterizerState(TStaticRasterizerState<>::GetRHI());
-	Context.RHICmdList.SetDepthStencilState(TStaticDepthStencilState<false, CF_Always>::GetRHI());
+	SetRenderTarget(Context.RHICmdList, FTextureRHIRef(), FTextureRHIRef());
+	Context.SetViewportAndCallRHI( 0, 0, 0.0f, TileCount.X, TileCount.Y, 1.0f );
+	
+	TShaderMapRef< FPostProcessVelocityGatherCS > ComputeShader( Context.GetShaderMap() );
+	Context.RHICmdList.SetComputeShader(ComputeShader->GetComputeShader());
 
-	TShaderMapRef< FPostProcessVS>					VertexShader( Context.GetShaderMap() );
-	TShaderMapRef< FPostProcessVelocityDilatePS >	PixelShader( Context.GetShaderMap() );
+	// set destination
+	Context.RHICmdList.SetUAVParameter( ComputeShader->GetComputeShader(), ComputeShader->OutScatteredMaxVelocity.GetBaseIndex(), DestRenderTarget.UAV );
 
-	static FGlobalBoundShaderState BoundShaderState;
-	SetGlobalBoundShaderState(Context.RHICmdList, Context.GetFeatureLevel(), BoundShaderState, GFilterVertexDeclaration.VertexDeclarationRHI, *VertexShader, *PixelShader);
+	ComputeShader->SetParameters( Context );
 
-	VertexShader->SetParameters(Context);
-	PixelShader->SetParameters(Context);
+	DispatchComputeShader(Context.RHICmdList, *ComputeShader, TileCount.X, TileCount.Y, 1);
 
-	// Draw a quad mapping scene color to the view's render target
-	DrawRectangle(
-		Context.RHICmdList,
-		0, 0,
-		SrcRect.Width(), SrcRect.Height(),
-		SrcRect.Min.X, SrcRect.Min.Y, 
-		SrcRect.Width(), SrcRect.Height(),
-		SrcRect.Size(),
-		SrcSize,
-		*VertexShader,
-		EDRF_UseTriangleOptimization);
+	// un-set destination
+	Context.RHICmdList.SetUAVParameter( ComputeShader->GetComputeShader(), ComputeShader->OutScatteredMaxVelocity.GetBaseIndex(), NULL );
 
 	Context.RHICmdList.CopyToResolveTarget(DestRenderTarget.TargetableTexture, DestRenderTarget.ShaderResourceTexture, false, FResolveParams());
 }
 
-FPooledRenderTargetDesc FRCPassPostProcessVelocityDilate::ComputeOutputDesc(EPassOutputId InPassOutputId) const
+FPooledRenderTargetDesc FRCPassPostProcessVelocityGather::ComputeOutputDesc(EPassOutputId InPassOutputId) const
 {
 	FPooledRenderTargetDesc Ret = PassInputs[0].GetOutput()->RenderTargetDesc;
 	Ret.Reset();
 
-	Ret.DebugName = TEXT("DilatedMaxVelocity");
+	Ret.TargetableFlags |= TexCreate_UAV;
+	Ret.TargetableFlags |= TexCreate_RenderTargetable;
+	Ret.DebugName = TEXT("ScatteredMaxVelocity");
 
 	return Ret;
 }
@@ -1427,12 +1373,11 @@ void FRCPassPostProcessMotionBlurNew::Process(FRenderingCompositePassContext& Co
 
 	// Set the view family's render target/viewport.
 	SetRenderTarget(Context.RHICmdList, DestRenderTarget.TargetableTexture, FTextureRHIRef());
-
+	Context.SetViewportAndCallRHI(SrcRect);
+	
 	// is optimized away if possible (RT size=view size, )
 	//Context.RHICmdList.Clear(true, FLinearColor::Black, false, 1.0f, false, 0, SrcRect);
 
-	Context.SetViewportAndCallRHI(SrcRect);
-
 	// set the state
 	Context.RHICmdList.SetBlendState(TStaticBlendState<>::GetRHI());
 	Context.RHICmdList.SetRasterizerState(TStaticRasterizerState<>::GetRHI());
diff --git a/Engine/Source/Runtime/Renderer/Private/PostProcess/PostProcessMotionBlur.h b/Engine/Source/Runtime/Renderer/Private/PostProcess/PostProcessMotionBlur.h
--- a/Engine/Source/Runtime/Renderer/Private/PostProcess/PostProcessMotionBlur.h
+++ b/Engine/Source/Runtime/Renderer/Private/PostProcess/PostProcessMotionBlur.h
@@ -90,12 +90,6 @@ public:
 	// -1 if not set yet (AsyncCompute wasn't started, we need another Process to finish it)
 	uint32 AsyncJobFenceID;
 
-	static const uint32 TileSizeX = 16;
-	static const uint32 TileSizeY = 16;
-	
-	static const uint32 ThreadGroupSizeX = TileSizeX / 2;
-	static const uint32 ThreadGroupSizeY = TileSizeY / 2;
-
 	static FIntPoint ComputeThreadGroupCount(FIntPoint PixelExtent);
 };
 
@@ -111,7 +105,7 @@ public:
 };
 
 // derives from TRenderingCompositePassBase<InputCount, OutputCount>
-class FRCPassPostProcessVelocityDilate : public TRenderingCompositePassBase<1, 1>
+class FRCPassPostProcessVelocityGather : public TRenderingCompositePassBase<1, 1>
 {
 public:
 	// interface FRenderingCompositePass ---------
diff --git a/Engine/Source/Runtime/Renderer/Private/PostProcess/PostProcessTemporalAA.cpp b/Engine/Source/Runtime/Renderer/Private/PostProcess/PostProcessTemporalAA.cpp
--- a/Engine/Source/Runtime/Renderer/Private/PostProcess/PostProcessTemporalAA.cpp
+++ b/Engine/Source/Runtime/Renderer/Private/PostProcess/PostProcessTemporalAA.cpp
@@ -39,6 +39,12 @@ static TAutoConsoleVariable<float> CVarTemporalAASharpness(
 	TEXT("Sharpness of temporal AA (0.0 = smoother, 1.0 = sharper)."),
 	ECVF_Scalability | ECVF_RenderThreadSafe);
 
+static TAutoConsoleVariable<int32> CVarTemporalAAPauseCorrect(
+	TEXT("r.TemporalAAPauseCorrect"),
+	0,
+	TEXT("Correct temporal AA in pause."),
+	ECVF_RenderThreadSafe);
+
 /** Encapsulates a TemporalAA pixel shader. */
 template< uint32 Type, uint32 Responsive >
 class FPostProcessTemporalAAPS : public FGlobalShader
@@ -697,12 +703,15 @@ void FRCPassPostProcessTemporalAA::Process(FRenderingCompositePassContext& Conte
 	}
 
 	Context.RHICmdList.CopyToResolveTarget(DestRenderTarget.TargetableTexture, DestRenderTarget.ShaderResourceTexture, false, FResolveParams());
-
-	if( !View.Family->bWorldIsPaused )
+	
+	if( !CVarTemporalAAPauseCorrect.GetValueOnRenderThread() )
 	{
-		ViewState->TemporalAAHistoryRT = PassOutputs[0].PooledRenderTarget;
-		check( ViewState->TemporalAAHistoryRT );
+		// Release to potentially reuse memory.
+		// Otherwise keep around so that pause looks correct.
+		ViewState->TemporalAAHistoryRT.SafeRelease();
 	}
+	ViewState->PendingTemporalAAHistoryRT = PassOutputs[0].PooledRenderTarget;
+	check( ViewState->PendingTemporalAAHistoryRT );
 
 	// TODO draw separate translucency after jitter has been removed
 
diff --git a/Engine/Source/Runtime/Renderer/Private/PostProcess/PostProcessing.cpp b/Engine/Source/Runtime/Renderer/Private/PostProcess/PostProcessing.cpp
--- a/Engine/Source/Runtime/Renderer/Private/PostProcess/PostProcessing.cpp
+++ b/Engine/Source/Runtime/Renderer/Private/PostProcess/PostProcessing.cpp
@@ -963,6 +963,13 @@ static TAutoConsoleVariable<int32> CVarMotionBlurDilate(
 	ECVF_RenderThreadSafe
 );
 
+static TAutoConsoleVariable<int32> CVarMotionBlurSeparable(
+	TEXT("r.MotionBlurSeparable"),
+	0,
+	TEXT(""),
+	ECVF_RenderThreadSafe
+);
+
 void FPostProcessing::Process(FRHICommandListImmediate& RHICmdList, FViewInfo& View, TRefCountPtr<IPooledRenderTarget>& VelocityRT)
 {
 	QUICK_SCOPE_CYCLE_COUNTER( STAT_PostProcessing_Process );
@@ -1170,9 +1177,9 @@ void FPostProcessing::Process(FRHICommandListImmediate& RHICmdList, FViewInfo& V
 						MaxTileVelocity	= FRenderingCompositeOutputRef( VelocityScatterPass );
 					}
 
-					if( CVarMotionBlurDilate.GetValueOnRenderThread() )
+					if( !CVarMotionBlurScatter.GetValueOnRenderThread() )
 					{
-						FRenderingCompositePass* VelocityDilatePass = Context.Graph.RegisterPass( new(FMemStack::Get()) FRCPassPostProcessVelocityDilate() );
+						FRenderingCompositePass* VelocityDilatePass = Context.Graph.RegisterPass( new(FMemStack::Get()) FRCPassPostProcessVelocityGather() );
 						VelocityDilatePass->SetInput( ePId_Input0, MaxTileVelocity );
 
 						MaxTileVelocity	= FRenderingCompositeOutputRef( VelocityDilatePass );
@@ -1187,6 +1194,17 @@ void FPostProcessing::Process(FRHICommandListImmediate& RHICmdList, FViewInfo& V
 					
 						Context.FinalOutput = FRenderingCompositeOutputRef( MotionBlurPass );
 					}
+
+					if( CVarMotionBlurSeparable.GetValueOnRenderThread() )
+					{
+						FRenderingCompositePass* MotionBlurPass = Context.Graph.RegisterPass( new(FMemStack::Get()) FRCPassPostProcessMotionBlurNew( GetMotionBlurQualityFromCVar() ) );
+						MotionBlurPass->SetInput( ePId_Input0, Context.FinalOutput );
+						MotionBlurPass->SetInput( ePId_Input1, SceneDepth );
+						MotionBlurPass->SetInput( ePId_Input2, VelocityInput );
+						MotionBlurPass->SetInput( ePId_Input3, MaxTileVelocity );
+					
+						Context.FinalOutput = FRenderingCompositeOutputRef( MotionBlurPass );
+					}
 				}
 				else
 				{
@@ -1644,7 +1662,7 @@ void FPostProcessing::Process(FRHICommandListImmediate& RHICmdList, FViewInfo& V
 			OverrideRenderTarget(Context.FinalOutput, Temp, Desc);
 
 			// We want this executed early as it uses AsyncCompute (where available e.g. XboxOne).
-			CompositeContext.Process(VelocityFlattenPass, TEXT("VelocityFlattenPass"));
+			//CompositeContext.Process(VelocityFlattenPass, TEXT("VelocityFlattenPass"));
 			// Now we process the remaining part of the graph.
 			CompositeContext.Process(Context.FinalOutput.GetPass(), TEXT("PostProcessing"));
 		}
diff --git a/Engine/Source/Runtime/Renderer/Private/PostProcess/RenderingCompositionGraph.cpp b/Engine/Source/Runtime/Renderer/Private/PostProcess/RenderingCompositionGraph.cpp
--- a/Engine/Source/Runtime/Renderer/Private/PostProcess/RenderingCompositionGraph.cpp
+++ b/Engine/Source/Runtime/Renderer/Private/PostProcess/RenderingCompositionGraph.cpp
@@ -717,9 +717,7 @@ void FPostProcessPassParameters::Set(
 		}
 
 		{
-			FVector4 Value(LocalViewport.Min.X, LocalViewport.Min.Y, LocalViewport.Max.X, LocalViewport.Max.Y);
-
-			SetShaderValue(RHICmdList, ShaderRHI, ViewportRect, Value);
+			SetShaderValue(RHICmdList, ShaderRHI, ViewportRect, Context.GetViewport());
 		}
 
 		{
diff --git a/Engine/Source/Runtime/Renderer/Private/ScenePrivate.h b/Engine/Source/Runtime/Renderer/Private/ScenePrivate.h
--- a/Engine/Source/Runtime/Renderer/Private/ScenePrivate.h
+++ b/Engine/Source/Runtime/Renderer/Private/ScenePrivate.h
@@ -486,6 +486,7 @@ public:
 	TRefCountPtr<IPooledRenderTarget> SeparateTranslucencyRT;
 	// Temporal AA result of last frame
 	TRefCountPtr<IPooledRenderTarget> TemporalAAHistoryRT;
+	TRefCountPtr<IPooledRenderTarget> PendingTemporalAAHistoryRT;
 	// Temporal AA result for DOF of last frame
 	TRefCountPtr<IPooledRenderTarget> DOFHistoryRT;
 	TRefCountPtr<IPooledRenderTarget> DOFHistoryRT2;
diff --git a/Engine/Source/Runtime/Renderer/Private/SceneVisibility.cpp b/Engine/Source/Runtime/Renderer/Private/SceneVisibility.cpp
--- a/Engine/Source/Runtime/Renderer/Private/SceneVisibility.cpp
+++ b/Engine/Source/Runtime/Renderer/Private/SceneVisibility.cpp
@@ -1584,6 +1584,7 @@ void FSceneRenderer::PreVisibilityFrameSetup(FRHICommandListImmediate& RHICmdLis
 					if(!ViewFamily.bWorldIsPaused)
 					{
 						ViewState->PrevViewMatrices = ViewState->PendingPrevViewMatrices;
+						ViewState->TemporalAAHistoryRT = ViewState->PendingTemporalAAHistoryRT;
 
 						// pending is needed as we are in init view and still need to render.
 						ViewState->PendingPrevViewMatrices = View.ViewMatrices;
