Fix for SSR + no temporal AA bug.

diff --git a/Engine/Source/Runtime/Renderer/Private/SceneVisibility.cpp b/Engine/Source/Runtime/Renderer/Private/SceneVisibility.cpp
--- a/Engine/Source/Runtime/Renderer/Private/SceneVisibility.cpp
+++ b/Engine/Source/Runtime/Renderer/Private/SceneVisibility.cpp
@@ -1585,6 +1585,7 @@ void FSceneRenderer::PreVisibilityFrameSetup(FRHICommandListImmediate& RHICmdLis
 					{
 						ViewState->PrevViewMatrices = ViewState->PendingPrevViewMatrices;
 						ViewState->TemporalAAHistoryRT = ViewState->PendingTemporalAAHistoryRT;
+						ViewState->PendingTemporalAAHistoryRT.SafeRelease();
 
 						// pending is needed as we are in init view and still need to render.
 						ViewState->PendingPrevViewMatrices = View.ViewMatrices;
