Mobile dynamic point lights with movable directional lights

diff --git a/Engine/Source/Runtime/Renderer/Private/ForwardBasePassRendering.cpp b/Engine/Source/Runtime/Renderer/Private/ForwardBasePassRendering.cpp
--- a/Engine/Source/Runtime/Renderer/Private/ForwardBasePassRendering.cpp
+++ b/Engine/Source/Runtime/Renderer/Private/ForwardBasePassRendering.cpp
@@ -28,36 +28,23 @@
 	IMPLEMENT_MATERIAL_SHADER_TYPE(template<>, TBasePassForForwardShadingPS##LightMapPolicyName##NumDynamicPointLights##HDRLinear32##Skylight, TEXT("BasePassForForwardShadingPixelShader"), TEXT("Main"), SF_Pixel); \
 	IMPLEMENT_MATERIAL_SHADER_TYPE(template<>, TBasePassForForwardShadingPS##LightMapPolicyName##NumDynamicPointLights##HDRLinear64##Skylight, TEXT("BasePassForForwardShadingPixelShader"), TEXT("Main"), SF_Pixel);
 
+static_assert(MAX_BASEPASS_DYNAMIC_POINT_LIGHTS == 4, "If you change MAX_BASEPASS_DYNAMIC_POINT_LIGHTS, you need to add shader types below");
+
+// Permutations for the number of point lights to support. INT32_MAX indicates the shader should use branching to support a variable number of point lights.
 #define IMPLEMENT_FORWARD_SHADING_BASEPASS_LIGHTMAPPED_SHADER_TYPE(LightMapPolicyType,LightMapPolicyName) \
 	IMPLEMENT_FORWARD_SHADING_BASEPASS_LIGHTMAPPED_VERTEX_SHADER_TYPE(LightMapPolicyType, LightMapPolicyName) \
-	IMPLEMENT_FORWARD_SHADING_BASEPASS_LIGHTMAPPED_PIXEL_SHADER_TYPE(LightMapPolicyType, LightMapPolicyName,0)
-
-static_assert(MAX_BASEPASS_DYNAMIC_POINT_LIGHTS == 4, "If you change MAX_BASEPASS_DYNAMIC_POINT_LIGHTS, you need to make shader types below");
-
-// Implement shader types per lightmap policy and for the number of point lights to support.
-// INT32_MAX indicates the shader should use branching to support a variable number of point lights.
+	IMPLEMENT_FORWARD_SHADING_BASEPASS_LIGHTMAPPED_PIXEL_SHADER_TYPE(LightMapPolicyType, LightMapPolicyName, 0) \
+	IMPLEMENT_FORWARD_SHADING_BASEPASS_LIGHTMAPPED_PIXEL_SHADER_TYPE(LightMapPolicyType, LightMapPolicyName, 1) \
+	IMPLEMENT_FORWARD_SHADING_BASEPASS_LIGHTMAPPED_PIXEL_SHADER_TYPE(LightMapPolicyType, LightMapPolicyName, 2) \
+	IMPLEMENT_FORWARD_SHADING_BASEPASS_LIGHTMAPPED_PIXEL_SHADER_TYPE(LightMapPolicyType, LightMapPolicyName, 3) \
+	IMPLEMENT_FORWARD_SHADING_BASEPASS_LIGHTMAPPED_PIXEL_SHADER_TYPE(LightMapPolicyType, LightMapPolicyName, 4) \
+	IMPLEMENT_FORWARD_SHADING_BASEPASS_LIGHTMAPPED_PIXEL_SHADER_TYPE(LightMapPolicyType, LightMapPolicyName, INT32_MAX)
+
+// Implement shader types per lightmap policy 
 IMPLEMENT_FORWARD_SHADING_BASEPASS_LIGHTMAPPED_SHADER_TYPE(FNoLightMapPolicy, FNoLightMapPolicy);
-IMPLEMENT_FORWARD_SHADING_BASEPASS_LIGHTMAPPED_VERTEX_SHADER_TYPE(TLightMapPolicy<LQ_LIGHTMAP>, TLightMapPolicyLQ);
-IMPLEMENT_FORWARD_SHADING_BASEPASS_LIGHTMAPPED_PIXEL_SHADER_TYPE(TLightMapPolicy<LQ_LIGHTMAP>, TLightMapPolicyLQ, 0);
-IMPLEMENT_FORWARD_SHADING_BASEPASS_LIGHTMAPPED_PIXEL_SHADER_TYPE(TLightMapPolicy<LQ_LIGHTMAP>, TLightMapPolicyLQ, 1);
-IMPLEMENT_FORWARD_SHADING_BASEPASS_LIGHTMAPPED_PIXEL_SHADER_TYPE(TLightMapPolicy<LQ_LIGHTMAP>, TLightMapPolicyLQ, 2);
-IMPLEMENT_FORWARD_SHADING_BASEPASS_LIGHTMAPPED_PIXEL_SHADER_TYPE(TLightMapPolicy<LQ_LIGHTMAP>, TLightMapPolicyLQ, 3);
-IMPLEMENT_FORWARD_SHADING_BASEPASS_LIGHTMAPPED_PIXEL_SHADER_TYPE(TLightMapPolicy<LQ_LIGHTMAP>, TLightMapPolicyLQ, 4);
-IMPLEMENT_FORWARD_SHADING_BASEPASS_LIGHTMAPPED_PIXEL_SHADER_TYPE(TLightMapPolicy<LQ_LIGHTMAP>, TLightMapPolicyLQ, INT32_MAX);
-IMPLEMENT_FORWARD_SHADING_BASEPASS_LIGHTMAPPED_VERTEX_SHADER_TYPE(TDistanceFieldShadowsAndLightMapPolicy<LQ_LIGHTMAP>, TDistanceFieldShadowsAndLightMapPolicyLQ);
-IMPLEMENT_FORWARD_SHADING_BASEPASS_LIGHTMAPPED_PIXEL_SHADER_TYPE(TDistanceFieldShadowsAndLightMapPolicy<LQ_LIGHTMAP>, TDistanceFieldShadowsAndLightMapPolicyLQ, 0);
-IMPLEMENT_FORWARD_SHADING_BASEPASS_LIGHTMAPPED_PIXEL_SHADER_TYPE(TDistanceFieldShadowsAndLightMapPolicy<LQ_LIGHTMAP>, TDistanceFieldShadowsAndLightMapPolicyLQ, 1);
-IMPLEMENT_FORWARD_SHADING_BASEPASS_LIGHTMAPPED_PIXEL_SHADER_TYPE(TDistanceFieldShadowsAndLightMapPolicy<LQ_LIGHTMAP>, TDistanceFieldShadowsAndLightMapPolicyLQ, 2);
-IMPLEMENT_FORWARD_SHADING_BASEPASS_LIGHTMAPPED_PIXEL_SHADER_TYPE(TDistanceFieldShadowsAndLightMapPolicy<LQ_LIGHTMAP>, TDistanceFieldShadowsAndLightMapPolicyLQ, 3);
-IMPLEMENT_FORWARD_SHADING_BASEPASS_LIGHTMAPPED_PIXEL_SHADER_TYPE(TDistanceFieldShadowsAndLightMapPolicy<LQ_LIGHTMAP>, TDistanceFieldShadowsAndLightMapPolicyLQ, 4);
-IMPLEMENT_FORWARD_SHADING_BASEPASS_LIGHTMAPPED_PIXEL_SHADER_TYPE(TDistanceFieldShadowsAndLightMapPolicy<LQ_LIGHTMAP>, TDistanceFieldShadowsAndLightMapPolicyLQ, INT32_MAX);
-IMPLEMENT_FORWARD_SHADING_BASEPASS_LIGHTMAPPED_VERTEX_SHADER_TYPE(FSimpleDirectionalLightAndSHIndirectPolicy, FSimpleDirectionalLightAndSHIndirectPolicy);
-IMPLEMENT_FORWARD_SHADING_BASEPASS_LIGHTMAPPED_PIXEL_SHADER_TYPE(FSimpleDirectionalLightAndSHIndirectPolicy, FSimpleDirectionalLightAndSHIndirectPolicy, 0);
-IMPLEMENT_FORWARD_SHADING_BASEPASS_LIGHTMAPPED_PIXEL_SHADER_TYPE(FSimpleDirectionalLightAndSHIndirectPolicy, FSimpleDirectionalLightAndSHIndirectPolicy, 1);
-IMPLEMENT_FORWARD_SHADING_BASEPASS_LIGHTMAPPED_PIXEL_SHADER_TYPE(FSimpleDirectionalLightAndSHIndirectPolicy, FSimpleDirectionalLightAndSHIndirectPolicy, 2);
-IMPLEMENT_FORWARD_SHADING_BASEPASS_LIGHTMAPPED_PIXEL_SHADER_TYPE(FSimpleDirectionalLightAndSHIndirectPolicy, FSimpleDirectionalLightAndSHIndirectPolicy, 3);
-IMPLEMENT_FORWARD_SHADING_BASEPASS_LIGHTMAPPED_PIXEL_SHADER_TYPE(FSimpleDirectionalLightAndSHIndirectPolicy, FSimpleDirectionalLightAndSHIndirectPolicy, 4);
-IMPLEMENT_FORWARD_SHADING_BASEPASS_LIGHTMAPPED_PIXEL_SHADER_TYPE(FSimpleDirectionalLightAndSHIndirectPolicy, FSimpleDirectionalLightAndSHIndirectPolicy, INT32_MAX);
+IMPLEMENT_FORWARD_SHADING_BASEPASS_LIGHTMAPPED_SHADER_TYPE(TLightMapPolicy<LQ_LIGHTMAP>, TLightMapPolicyLQ);
+IMPLEMENT_FORWARD_SHADING_BASEPASS_LIGHTMAPPED_SHADER_TYPE(TDistanceFieldShadowsAndLightMapPolicy<LQ_LIGHTMAP>, TDistanceFieldShadowsAndLightMapPolicyLQ);
+IMPLEMENT_FORWARD_SHADING_BASEPASS_LIGHTMAPPED_SHADER_TYPE(FSimpleDirectionalLightAndSHIndirectPolicy, FSimpleDirectionalLightAndSHIndirectPolicy);
 IMPLEMENT_FORWARD_SHADING_BASEPASS_LIGHTMAPPED_SHADER_TYPE(FSimpleDirectionalLightAndSHDirectionalIndirectPolicy, FSimpleDirectionalLightAndSHDirectionalIndirectPolicy);
 IMPLEMENT_FORWARD_SHADING_BASEPASS_LIGHTMAPPED_SHADER_TYPE(FSimpleDirectionalLightAndSHDirectionalCSMIndirectPolicy, FSimpleDirectionalLightAndSHDirectionalCSMIndirectPolicy);
 IMPLEMENT_FORWARD_SHADING_BASEPASS_LIGHTMAPPED_SHADER_TYPE(FMovableDirectionalLightLightingPolicy, FMovableDirectionalLightLightingPolicy);
diff --git a/Engine/Source/Runtime/Renderer/Private/ForwardBasePassRendering.h b/Engine/Source/Runtime/Renderer/Private/ForwardBasePassRendering.h
--- a/Engine/Source/Runtime/Renderer/Private/ForwardBasePassRendering.h
+++ b/Engine/Source/Runtime/Renderer/Private/ForwardBasePassRendering.h
@@ -152,8 +152,8 @@ public:
 
 		return IsMobilePlatform(Platform) && LightMapPolicyType::ShouldCache(Platform, Material, VertexFactoryType) &&
 			(NumDynamicPointLights == 0 ||
-			(!bIsUnlit && NumDynamicPointLights == INT32_MAX && bMobileDynamicPointLightsUseStaticBranch) ||					// single shader for variable number of point lights
-			(!bIsUnlit && NumDynamicPointLights <= MobileNumDynamicPointLights && !bMobileDynamicPointLightsUseStaticBranch));	// unique 1...N point light shaders
+			(!bIsUnlit && NumDynamicPointLights == INT32_MAX && bMobileDynamicPointLightsUseStaticBranch && MobileNumDynamicPointLights > 0) ||	// single shader for variable number of point lights
+			(!bIsUnlit && NumDynamicPointLights <= MobileNumDynamicPointLights && !bMobileDynamicPointLightsUseStaticBranch));					// unique 1...N point light shaders
 	}
 
 	static void ModifyCompilationEnvironment(EShaderPlatform Platform, const FMaterial* Material, FShaderCompilerEnvironment& OutEnvironment)
@@ -626,11 +626,11 @@ void ProcessBasePassMeshForForwardShading(
 				// final determination of whether CSMs are rendered can be view dependent, thus we always need to clear the CSMs even if we're not going to render to them based on the condition below.
 				if (SimpleDirectionalLight && SimpleDirectionalLight->ShouldRenderViewIndependentWholeSceneShadows())
 				{
-					Action.template Process<FMovableDirectionalLightCSMWithLightmapLightingPolicy, 0>(RHICmdList, Parameters, FMovableDirectionalLightCSMWithLightmapLightingPolicy(), LightMapInteraction);
+					Action.template Process<FMovableDirectionalLightCSMWithLightmapLightingPolicy, NumDynamicPointLights>(RHICmdList, Parameters, FMovableDirectionalLightCSMWithLightmapLightingPolicy(), LightMapInteraction);
 				}
 				else
 				{
-					Action.template Process<FMovableDirectionalLightWithLightmapLightingPolicy, 0>(RHICmdList, Parameters, FMovableDirectionalLightCSMWithLightmapLightingPolicy(), LightMapInteraction);
+					Action.template Process<FMovableDirectionalLightWithLightmapLightingPolicy, NumDynamicPointLights==0 ? 0 : INT32_MAX>(RHICmdList, Parameters, FMovableDirectionalLightCSMWithLightmapLightingPolicy(), LightMapInteraction);
 				}
 			}
 			else
@@ -665,11 +665,11 @@ void ProcessBasePassMeshForForwardShading(
 			{
 				if (SimpleDirectionalLight && SimpleDirectionalLight->ShouldRenderViewIndependentWholeSceneShadows())
 				{
-					Action.template Process<FSimpleDirectionalLightAndSHDirectionalCSMIndirectPolicy, 0>(RHICmdList, Parameters, FSimpleDirectionalLightAndSHDirectionalCSMIndirectPolicy(), FSimpleDirectionalLightAndSHDirectionalCSMIndirectPolicy::ElementDataType(Action.ShouldPackAmbientSH()));
+					Action.template Process<FSimpleDirectionalLightAndSHDirectionalCSMIndirectPolicy, NumDynamicPointLights>(RHICmdList, Parameters, FSimpleDirectionalLightAndSHDirectionalCSMIndirectPolicy(), FSimpleDirectionalLightAndSHDirectionalCSMIndirectPolicy::ElementDataType(Action.ShouldPackAmbientSH()));
 				}
 				else
 				{
-					Action.template Process<FSimpleDirectionalLightAndSHDirectionalIndirectPolicy, 0>(RHICmdList, Parameters, FSimpleDirectionalLightAndSHDirectionalIndirectPolicy(), FSimpleDirectionalLightAndSHDirectionalIndirectPolicy::ElementDataType(Action.ShouldPackAmbientSH()));
+					Action.template Process<FSimpleDirectionalLightAndSHDirectionalIndirectPolicy, NumDynamicPointLights>(RHICmdList, Parameters, FSimpleDirectionalLightAndSHDirectionalIndirectPolicy(), FSimpleDirectionalLightAndSHDirectionalIndirectPolicy::ElementDataType(Action.ShouldPackAmbientSH()));
 				}
 			}
 			else
@@ -685,11 +685,11 @@ void ProcessBasePassMeshForForwardShading(
 			// final determination of whether CSMs are rendered can be view dependent, thus we always need to clear the CSMs even if we're not going to render to them based on the condition below.
 			if (SimpleDirectionalLight && SimpleDirectionalLight->ShouldRenderViewIndependentWholeSceneShadows())
 			{
-				Action.template Process<FMovableDirectionalLightCSMLightingPolicy, 0>(RHICmdList, Parameters, FMovableDirectionalLightCSMLightingPolicy(), FMovableDirectionalLightCSMLightingPolicy::ElementDataType());
+				Action.template Process<FMovableDirectionalLightCSMLightingPolicy, NumDynamicPointLights>(RHICmdList, Parameters, FMovableDirectionalLightCSMLightingPolicy(), FMovableDirectionalLightCSMLightingPolicy::ElementDataType());
 			}
 			else
 			{
-				Action.template Process<FMovableDirectionalLightLightingPolicy, 0>(RHICmdList, Parameters, FMovableDirectionalLightLightingPolicy(), FMovableDirectionalLightLightingPolicy::ElementDataType());
+				Action.template Process<FMovableDirectionalLightLightingPolicy, NumDynamicPointLights>(RHICmdList, Parameters, FMovableDirectionalLightLightingPolicy(), FMovableDirectionalLightLightingPolicy::ElementDataType());
 			}
 
 			// Exit to avoid NoLightmapPolicy
