Removed test code accidentally submitted in CL 2596623

diff --git a/Engine/Source/Runtime/Renderer/Private/ForwardBasePassRendering.h b/Engine/Source/Runtime/Renderer/Private/ForwardBasePassRendering.h
--- a/Engine/Source/Runtime/Renderer/Private/ForwardBasePassRendering.h
+++ b/Engine/Source/Runtime/Renderer/Private/ForwardBasePassRendering.h
@@ -630,7 +630,7 @@ void ProcessBasePassMeshForForwardShading(
 				}
 				else
 				{
-					Action.template Process<FMovableDirectionalLightWithLightmapLightingPolicy, NumDynamicPointLights==0 ? 0 : INT32_MAX>(RHICmdList, Parameters, FMovableDirectionalLightCSMWithLightmapLightingPolicy(), LightMapInteraction);
+					Action.template Process<FMovableDirectionalLightWithLightmapLightingPolicy, NumDynamicPointLights>(RHICmdList, Parameters, FMovableDirectionalLightCSMWithLightmapLightingPolicy(), LightMapInteraction);
 				}
 			}
 			else
@@ -641,7 +641,7 @@ void ProcessBasePassMeshForForwardShading(
 
 				if (ShadowMapInteraction.GetType() == SMIT_Texture)
 				{
-					Action.template Process< TDistanceFieldShadowsAndLightMapPolicy<LQ_LIGHTMAP>, NumDynamicPointLights >(
+					Action.template Process< TDistanceFieldShadowsAndLightMapPolicy<LQ_LIGHTMAP>, NumDynamicPointLights>(
 						RHICmdList,
 						Parameters,
 						TDistanceFieldShadowsAndLightMapPolicy<LQ_LIGHTMAP>(),
@@ -649,7 +649,7 @@ void ProcessBasePassMeshForForwardShading(
 				}
 				else
 				{
-					Action.template Process< TLightMapPolicy<LQ_LIGHTMAP>, NumDynamicPointLights >(RHICmdList, Parameters, TLightMapPolicy<LQ_LIGHTMAP>(), LightMapInteraction);
+					Action.template Process< TLightMapPolicy<LQ_LIGHTMAP>, NumDynamicPointLights>(RHICmdList, Parameters, TLightMapPolicy<LQ_LIGHTMAP>(), LightMapInteraction);
 				}
 			}
 
