Changed temporal aa pause correctness to default.

diff --git a/Engine/Source/Runtime/Renderer/Private/PostProcess/PostProcessTemporalAA.cpp b/Engine/Source/Runtime/Renderer/Private/PostProcess/PostProcessTemporalAA.cpp
--- a/Engine/Source/Runtime/Renderer/Private/PostProcess/PostProcessTemporalAA.cpp
+++ b/Engine/Source/Runtime/Renderer/Private/PostProcess/PostProcessTemporalAA.cpp
@@ -41,8 +41,8 @@ static TAutoConsoleVariable<float> CVarTemporalAASharpness(
 
 static TAutoConsoleVariable<int32> CVarTemporalAAPauseCorrect(
 	TEXT("r.TemporalAAPauseCorrect"),
-	0,
-	TEXT("Correct temporal AA in pause."),
+	1,
+	TEXT("Correct temporal AA in pause. This holds onto render targets longer preventing reuse and consumes more memory."),
 	ECVF_RenderThreadSafe);
 
 /** Encapsulates a TemporalAA pixel shader. */
