Fixed crash on exit of PIE due to FSceneViewState::PendingTemporalAAHistoryRT not getting released on RT properly

diff --git a/Engine/Source/Runtime/Renderer/Private/ScenePrivate.h b/Engine/Source/Runtime/Renderer/Private/ScenePrivate.h
--- a/Engine/Source/Runtime/Renderer/Private/ScenePrivate.h
+++ b/Engine/Source/Runtime/Renderer/Private/ScenePrivate.h
@@ -679,6 +679,7 @@ public:
 		EyeAdaptationRT.SafeRelease();
 		SeparateTranslucencyRT.SafeRelease();
 		TemporalAAHistoryRT.SafeRelease();
+		PendingTemporalAAHistoryRT.SafeRelease();
 		DOFHistoryRT.SafeRelease();
 		DOFHistoryRT2.SafeRelease();
 		SSRHistoryRT.SafeRelease();
