fixed shadercomplexity with decals

diff --git a/Engine/Source/Runtime/Renderer/Private/CompositionLighting/CompositionLighting.cpp b/Engine/Source/Runtime/Renderer/Private/CompositionLighting/CompositionLighting.cpp
--- a/Engine/Source/Runtime/Renderer/Private/CompositionLighting/CompositionLighting.cpp
+++ b/Engine/Source/Runtime/Renderer/Private/CompositionLighting/CompositionLighting.cpp
@@ -279,8 +279,9 @@ void FCompositionLighting::ProcessAfterBasePass(FRHICommandListImmediate& RHICmd
 
 		// Add the passes we want to add to the graph ----------
 		
-		if( Context.View.Family->EngineShowFlags.Decals)
+		if(Context.View.Family->EngineShowFlags.Decals && !Context.View.Family->EngineShowFlags.ShaderComplexity)
 		{
+			// DRS_AfterBasePass is for Volumetric decals which don't support ShaderComplexity yet
 			FRenderingCompositePass* Pass = Context.Graph.RegisterPass(new(FMemStack::Get()) FRCPassPostProcessDeferredDecals(DRS_AfterBasePass));
 			Pass->SetInput(ePId_Input0, Context.FinalOutput);
 
