fixed UE-16338 Crash when setting console command for shadow max resolution to 1

diff --git a/Engine/Source/Runtime/Renderer/Private/PostProcess/SceneRenderTargets.cpp b/Engine/Source/Runtime/Renderer/Private/PostProcess/SceneRenderTargets.cpp
--- a/Engine/Source/Runtime/Renderer/Private/PostProcess/SceneRenderTargets.cpp
+++ b/Engine/Source/Runtime/Renderer/Private/PostProcess/SceneRenderTargets.cpp
@@ -1892,14 +1892,26 @@ FIntPoint FSceneRenderTargets::GetPreShadowCacheTextureResolution() const
 	const int32 ExpandFactor = 2;
 
 	static auto CVarPreShadowResolutionFactor = IConsoleManager::Get().FindTConsoleVariableDataFloat(TEXT("r.Shadow.PreShadowResolutionFactor"));
-	return FIntPoint(FMath::TruncToInt(ShadowDepthResolution.X * CVarPreShadowResolutionFactor->GetValueOnRenderThread()), FMath::TruncToInt(ShadowDepthResolution.Y * CVarPreShadowResolutionFactor->GetValueOnRenderThread())) * ExpandFactor;
+
+	float Factor = CVarPreShadowResolutionFactor->GetValueOnRenderThread();
+
+	FIntPoint Ret;
+
+	Ret.X = FMath::Clamp(FMath::TruncToInt(ShadowDepthResolution.X * Factor) * ExpandFactor, 1, GMaxShadowDepthBufferSizeX);
+	Ret.Y = FMath::Clamp(FMath::TruncToInt(ShadowDepthResolution.Y * Factor) * ExpandFactor, 1, GMaxShadowDepthBufferSizeY);
+
+	return Ret;
 }
 
 FIntPoint FSceneRenderTargets::GetTranslucentShadowDepthTextureResolution() const
 {
 	FIntPoint ShadowDepthResolution = GetShadowDepthTextureResolution();
-	ShadowDepthResolution.X = FMath::Max<int32>(ShadowDepthResolution.X / GetTranslucentShadowDownsampleFactor(), 1);
-	ShadowDepthResolution.Y = FMath::Max<int32>(ShadowDepthResolution.Y / GetTranslucentShadowDownsampleFactor(), 1);
+
+	int32 Factor = GetTranslucentShadowDownsampleFactor();
+
+	ShadowDepthResolution.X = FMath::Clamp(ShadowDepthResolution.X / Factor, 1, GMaxShadowDepthBufferSizeX);
+	ShadowDepthResolution.Y = FMath::Clamp(ShadowDepthResolution.Y / Factor, 1, GMaxShadowDepthBufferSizeY);
+
 	return ShadowDepthResolution;
 }
 
diff --git a/Engine/Source/Runtime/Renderer/Private/PostProcess/SceneRenderTargets.h b/Engine/Source/Runtime/Renderer/Private/PostProcess/SceneRenderTargets.h
--- a/Engine/Source/Runtime/Renderer/Private/PostProcess/SceneRenderTargets.h
+++ b/Engine/Source/Runtime/Renderer/Private/PostProcess/SceneRenderTargets.h
@@ -434,6 +434,7 @@ public:
 	int32 GetCubeShadowDepthZResolution(int32 ShadowIndex) const;
 	/** Returns the size of the shadow depth buffer, taking into account platform limitations and game specific resolution limits. */
 	FIntPoint GetShadowDepthTextureResolution() const;
+	// @return >= 1x1 <= GMaxShadowDepthBufferSizeX x GMaxShadowDepthBufferSizeY
 	FIntPoint GetPreShadowCacheTextureResolution() const;
 	FIntPoint GetTranslucentShadowDepthTextureResolution() const;
 	int32 GetTranslucentShadowDownsampleFactor() const { return 2; }
