Minor comment changes to mobile decals

diff --git a/Engine/Source/Developer/ShaderFormatOpenGL/Private/GlslBackend.cpp b/Engine/Source/Developer/ShaderFormatOpenGL/Private/GlslBackend.cpp
--- a/Engine/Source/Developer/ShaderFormatOpenGL/Private/GlslBackend.cpp
+++ b/Engine/Source/Developer/ShaderFormatOpenGL/Private/GlslBackend.cpp
@@ -442,6 +442,7 @@ static bool UsesFramebufferFetchES2(exec_list* Instructions)
 	return Visitor.bFound;
 }
 
+// Returns true if the DepthbufferFetchES2() 'intrinsic' is used
 static bool UsesDepthbufferFetchES2(exec_list* Instructions)
 {
 	struct SFindepthbufferFetchES2Intrinsic : public ir_hierarchical_visitor
diff --git a/Engine/Source/Runtime/RenderCore/Public/RenderCore.h b/Engine/Source/Runtime/RenderCore/Public/RenderCore.h
--- a/Engine/Source/Runtime/RenderCore/Public/RenderCore.h
+++ b/Engine/Source/Runtime/RenderCore/Public/RenderCore.h
@@ -37,7 +37,7 @@ DECLARE_DWORD_COUNTER_STAT_EXTERN(TEXT("Static list draw calls"),STAT_StaticDraw
 DECLARE_CYCLE_STAT_EXTERN(TEXT("Present time"),STAT_PresentTime,STATGROUP_SceneRendering, RENDERCORE_API);
 DECLARE_DWORD_ACCUMULATOR_STAT_EXTERN(TEXT("Decals in scene"),STAT_SceneDecals,STATGROUP_SceneRendering, RENDERCORE_API);
 DECLARE_DWORD_COUNTER_STAT_EXTERN(TEXT("Decals in view"),STAT_Decals,STATGROUP_SceneRendering, RENDERCORE_API);
-DECLARE_CYCLE_STAT_EXTERN(TEXT("Decals drawing"),STAT_DecalsDrawTime,STATGROUP_SceneRendering, RENDERCORE_API);
+DECLARE_CYCLE_STAT_EXTERN(TEXT("Decal drawing"),STAT_DecalsDrawTime,STATGROUP_SceneRendering, RENDERCORE_API);
 DECLARE_CYCLE_STAT_EXTERN(TEXT("Cache Uniform Expressions"),STAT_CacheUniformExpressions,STATGROUP_SceneRendering, RENDERCORE_API);
 DECLARE_DWORD_COUNTER_STAT_EXTERN(TEXT("Lights using light shafts"), STAT_LightShaftsLights, STATGROUP_SceneRendering, RENDERCORE_API);
 
diff --git a/Engine/Source/Runtime/Renderer/Private/ForwardBasePassRendering.cpp b/Engine/Source/Runtime/Renderer/Private/ForwardBasePassRendering.cpp
--- a/Engine/Source/Runtime/Renderer/Private/ForwardBasePassRendering.cpp
+++ b/Engine/Source/Runtime/Renderer/Private/ForwardBasePassRendering.cpp
@@ -520,8 +520,8 @@ void FForwardShadingSceneRenderer::RenderForwardShadingBasePass(FRHICommandListI
 			
 			int32 MeshBatchIndex = 0;
 			int32 MeshBatchNum = View.DynamicMeshElements.Num();
-			// In case scene has decals, subdivide mesh elements list to two groups, 
-			// first group that receives decals and second don't
+			// If the scene has decals we need to partition the draw list into two groups: 
+			// those that receive decals and those that don't, and render the latter with stencil writes enabled
 			if (Scene->Decals.Num() > 0)
 			{
 				MeshBatchNum = Algo::Partition(View.DynamicMeshElements.GetData(), MeshBatchNum, [](const FMeshBatchAndRelevance& El) { 
@@ -540,9 +540,9 @@ void FForwardShadingSceneRenderer::RenderForwardShadingBasePass(FRHICommandListI
 				}
 			}
 
-			// Render second group with enabled stencil
 			if (MeshBatchNum < View.DynamicMeshElements.Num())
 			{
+				// Primitives without decals are rendered with stencil enabled
 				RHICmdList.SetDepthStencilState(TStaticDepthStencilState<true, CF_DepthNearOrEqual, true, CF_Always, SO_Keep, SO_Keep, SO_Replace>::GetRHI(), 0x80);
 				MeshBatchNum = View.DynamicMeshElements.Num();
 
