non-unity build

diff --git a/Engine/Source/Runtime/Renderer/Private/DecalRenderingForwardShading.cpp b/Engine/Source/Runtime/Renderer/Private/DecalRenderingForwardShading.cpp
--- a/Engine/Source/Runtime/Renderer/Private/DecalRenderingForwardShading.cpp
+++ b/Engine/Source/Runtime/Renderer/Private/DecalRenderingForwardShading.cpp
@@ -6,6 +6,7 @@
 
 #include "RendererPrivate.h"
 #include "ScenePrivate.h"
+#include "SceneUtils.h"
 #include "DecalRenderingShared.h"
 
 void FForwardShadingSceneRenderer::RenderDecals(FRHICommandListImmediate& RHICmdList)
@@ -38,7 +39,7 @@ void FForwardShadingSceneRenderer::RenderDecals(FRHICommandListImmediate& RHICmd
 
 		if (SortedDecals.Num())
 		{
-			SCOPED_DRAW_EVENT(RHICmdList, Decals);
+			SCOPED_DRAW_EVENT(RHICmdList, DeferredDecals);
 			INC_DWORD_STAT_BY(STAT_Decals, SortedDecals.Num());
 		
 			RHICmdList.SetViewport(View.ViewRect.Min.X, View.ViewRect.Min.Y, 0, View.ViewRect.Max.X, View.ViewRect.Max.Y, 1);
