Forward shaded modulated shadows.
Added cast modulated shadows flag and modulated color to ULightComponentBase.
moved some common shadow projection code from FDeferredShadingSceneRenderer to FSceneRenderer .
Added MODULATED_SHADOWS to projection pixel shader.
Enabled Depth shaders for ES2.
RenderProjections is used during forward shading.

diff --git a/Engine/Shaders/ShadowProjectionPixelShader.usf b/Engine/Shaders/ShadowProjectionPixelShader.usf
--- a/Engine/Shaders/ShadowProjectionPixelShader.usf
+++ b/Engine/Shaders/ShadowProjectionPixelShader.usf
@@ -37,6 +37,8 @@ float2 ProjectionDepthBiasParameters;
 	float InvFadePlaneLength;
 #endif
 
+float4 ModulatedShadowColor;
+
 /**
 * Entry point for uniform manual PCF that supports lights using normal shadows.
 */
@@ -154,6 +156,11 @@ void Main(
 		// When the fade plane is in use for CSMs, we output the fade value in the alpha channel for blending.
 		OutColor.a = BlendFactor;
 	#endif
+
+#if MODULATED_SHADOWS
+	OutColor.rgb = lerp(ModulatedShadowColor.rgb, float3(1, 1, 1), Shadow);
+	OutColor.a = 0;
+#endif
 }
 
 #if FEATURE_LEVEL >= FEATURE_LEVEL_SM4
diff --git a/Engine/Source/Runtime/Engine/Classes/Components/LightComponentBase.h b/Engine/Source/Runtime/Engine/Classes/Components/LightComponentBase.h
--- a/Engine/Source/Runtime/Engine/Classes/Components/LightComponentBase.h
+++ b/Engine/Source/Runtime/Engine/Classes/Components/LightComponentBase.h
@@ -60,6 +60,15 @@ class ENGINE_API ULightComponentBase : public USceneComponent
 	UPROPERTY(EditAnywhere, BlueprintReadOnly, Category=Light, AdvancedDisplay)
 	uint32 CastDynamicShadows:1;
 
+	/**
+	* Whether the light should cast modulated shadows from dynamic objects (mobile only).  Also requires Cast Shadows to be set to True.
+	**/
+	UPROPERTY(EditAnywhere, BlueprintReadOnly, Category = Light, AdvancedDisplay)
+	uint32 bCastModulatedShadows : 1;
+
+	UPROPERTY(BlueprintReadOnly, interp, Category = Light, meta = (HideAlphaChannel), AdvancedDisplay)
+	FColor ModulatedShadowColor;
+
 	/** Whether the light affects translucency or not.  Disabling this can save GPU time when there are many small lights. */
 	UPROPERTY(EditAnywhere, BlueprintReadOnly, Category=Light, AdvancedDisplay)
 	uint32 bAffectTranslucentLighting:1;
diff --git a/Engine/Source/Runtime/Engine/Private/Components/LightComponent.cpp b/Engine/Source/Runtime/Engine/Private/Components/LightComponent.cpp
--- a/Engine/Source/Runtime/Engine/Private/Components/LightComponent.cpp
+++ b/Engine/Source/Runtime/Engine/Private/Components/LightComponent.cpp
@@ -220,6 +220,7 @@ FLightSceneProxy::FLightSceneProxy(const ULightComponent* InLightComponent)
 	, bCastStaticShadow(InLightComponent->CastShadows && InLightComponent->CastStaticShadows)
 	, bCastTranslucentShadows(InLightComponent->CastTranslucentShadows)
 	, bCastShadowsFromCinematicObjectsOnly(InLightComponent->bCastShadowsFromCinematicObjectsOnly)
+	, bCastModulatedShadows(InLightComponent->bCastModulatedShadows)
 	, bAffectTranslucentLighting(InLightComponent->bAffectTranslucentLighting)
 	, bUsedAsAtmosphereSunLight(InLightComponent->IsUsedAsAtmosphereSunLight())
 	, bAffectDynamicIndirectLighting(InLightComponent->bAffectDynamicIndirectLighting)
@@ -267,6 +268,8 @@ FLightSceneProxy::FLightSceneProxy(const ULightComponent* InLightComponent)
 	LightFunctionDisabledBrightness = LightComponent->DisabledBrightness;
 
 	StaticShadowDepthMap = &LightComponent->StaticShadowDepthMap;
+
+	ModulatedShadowColor = FLinearColor(LightComponent->ModulatedShadowColor);
 }
 
 bool FLightSceneProxy::ShouldCreatePerObjectShadowsForDynamicObjects() const
@@ -515,6 +518,11 @@ bool ULightComponent::CanEditChange(const UProperty* InProperty) const
 		{
 			return bUseTemperature;
 		}
+
+		if (PropertyName == GET_MEMBER_NAME_STRING_CHECKED(ULightComponent, ModulatedShadowColor))
+		{
+			return bCastModulatedShadows;
+		}
 	}
 
 	return Super::CanEditChange(InProperty);
diff --git a/Engine/Source/Runtime/Engine/Public/SceneManagement.h b/Engine/Source/Runtime/Engine/Public/SceneManagement.h
--- a/Engine/Source/Runtime/Engine/Public/SceneManagement.h
+++ b/Engine/Source/Runtime/Engine/Public/SceneManagement.h
@@ -800,6 +800,8 @@ public:
 	inline bool CastsStaticShadow() const { return bCastStaticShadow; }
 	inline bool CastsTranslucentShadows() const { return bCastTranslucentShadows; }
 	inline bool CastsShadowsFromCinematicObjectsOnly() const { return bCastShadowsFromCinematicObjectsOnly; }
+	inline bool CastsModulatedShadows() const { return bCastModulatedShadows; }
+	inline const FLinearColor& GetModulatedShadowColor() const { return ModulatedShadowColor; }
 	inline bool AffectsTranslucentLighting() const { return bAffectTranslucentLighting; }
 	inline bool UseRayTracedDistanceFieldShadows() const { return bUseRayTracedDistanceFieldShadows; }
 	inline float GetRayStartOffsetDepthScale() const { return RayStartOffsetDepthScale; }
@@ -912,6 +914,9 @@ protected:
 
 	const uint32 bCastShadowsFromCinematicObjectsOnly : 1;
 
+	/** Whether the light will cast modulated shadows when using the forward renderer (mobile). */
+	const uint32 bCastModulatedShadows : 1;
+
 	/** Whether the light affects translucency or not.  Disabling this can save GPU time when there are many small lights. */
 	const uint32 bAffectTranslucentLighting : 1;
 
@@ -945,6 +950,8 @@ protected:
 	/** Only for whole scene directional lights, 0: no FarShadowCascades, otherwise the count of cascades between WholeSceneDynamicShadowRadius and FarShadowDistance that are covered by distant shadow cascades. */
 	uint32 FarShadowCascadeCount;
 
+	/** Modulated shadow color. */
+	FLinearColor ModulatedShadowColor;
 	/**
 	 * Updates the light proxy's cached transforms.
 	 * @param InLightToWorld - The new light-to-world transform.
diff --git a/Engine/Source/Runtime/Renderer/Private/DeferredShadingRenderer.h b/Engine/Source/Runtime/Renderer/Private/DeferredShadingRenderer.h
--- a/Engine/Source/Runtime/Renderer/Private/DeferredShadingRenderer.h
+++ b/Engine/Source/Runtime/Renderer/Private/DeferredShadingRenderer.h
@@ -140,22 +140,6 @@ private:
 	/** Updates the preshadow cache, allocating new preshadows that can fit and evicting old ones. */
 	void UpdatePreshadowCache(FSceneRenderTargets& SceneContext);
 
-	/** Finds a matching cached preshadow, if one exists. */
-	TRefCountPtr<FProjectedShadowInfo> GetCachedPreshadow(
-		const FLightPrimitiveInteraction* InParentInteraction,
-		const FProjectedShadowInitializer& Initializer,
-		const FBoxSphereBounds& Bounds,
-		uint32 InResolutionX);
-
-	/** Creates shadows for the given interaction. */
-	void SetupInteractionShadows(
-		FRHICommandListImmediate& RHICmdList,
-		FLightPrimitiveInteraction* Interaction,
-		FVisibleLightInfo& VisibleLightInfo,
-		bool bStaticSceneOnly,
-		const TArray<FProjectedShadowInfo*, SceneRenderingAllocator>& ViewDependentWholeSceneShadows,
-		TArray<FProjectedShadowInfo*, SceneRenderingAllocator>& PreShadows);
-
 	/** Finds the visible dynamic shadows for each view. */
 	void InitDynamicShadows(FRHICommandListImmediate& RHICmdList);
 
@@ -272,13 +256,6 @@ private:
 	/** Renders one pass point light shadows. */
 	bool RenderOnePassPointLightShadows(FRHICommandListImmediate& RHICmdList, const FLightSceneInfo* LightSceneInfo, bool bRenderedTranslucentObjectShadows, bool& bInjectedTranslucentVolume);
 
-	/** Renders the projections of the given Shadows to the appropriate color render target. */
-	void RenderProjections(
-		FRHICommandListImmediate& RHICmdList,
-		const FLightSceneInfo* LightSceneInfo, 
-		const TArray<FProjectedShadowInfo*,SceneRenderingAllocator>& Shadows
-		);
-
 	/** Renders the shadowmaps of translucent shadows and their projections onto opaque surfaces. */
 	bool RenderTranslucentProjectedShadows(FRHICommandListImmediate& RHICmdList, const FLightSceneInfo* LightSceneInfo);
 
diff --git a/Engine/Source/Runtime/Renderer/Private/DepthRendering.cpp b/Engine/Source/Runtime/Renderer/Private/DepthRendering.cpp
--- a/Engine/Source/Runtime/Renderer/Private/DepthRendering.cpp
+++ b/Engine/Source/Runtime/Renderer/Private/DepthRendering.cpp
@@ -28,13 +28,11 @@ public:
 		// Only the local vertex factory supports the position-only stream
 		if (bUsePositionOnlyStream)
 		{
-			return VertexFactoryType->SupportsPositionOnly() && Material->IsSpecialEngineMaterial() 
-				&& IsFeatureLevelSupported(Platform, ERHIFeatureLevel::SM4);
+			return VertexFactoryType->SupportsPositionOnly() && Material->IsSpecialEngineMaterial();
 		}
 
 		// Only compile for the default material and masked materials
-		return (Material->IsSpecialEngineMaterial() || !Material->WritesEveryPixel() || Material->MaterialMayModifyMeshPosition())
-			 && IsFeatureLevelSupported(Platform, ERHIFeatureLevel::SM4);
+		return (Material->IsSpecialEngineMaterial() || !Material->WritesEveryPixel() || Material->MaterialMayModifyMeshPosition());
 	}
 
 	void SetParameters(FRHICommandList& RHICmdList, const FMaterialRenderProxy* MaterialRenderProxy,const FMaterial& MaterialResource,const FSceneView& View)
@@ -106,7 +104,7 @@ public:
 	static bool ShouldCache(EShaderPlatform Platform,const FMaterial* Material,const FVertexFactoryType* VertexFactoryType)
 	{
 		// Compile for materials that are masked.
-		return (!Material->WritesEveryPixel() || Material->HasPixelDepthOffsetConnected()) && IsFeatureLevelSupported(Platform, ERHIFeatureLevel::SM4);
+		return (!Material->WritesEveryPixel() || Material->HasPixelDepthOffsetConnected());
 	}
 
 	FDepthOnlyPS(const ShaderMetaType::CompiledShaderInitializerType& Initializer):
diff --git a/Engine/Source/Runtime/Renderer/Private/ForwardShadingRenderer.cpp b/Engine/Source/Runtime/Renderer/Private/ForwardShadingRenderer.cpp
--- a/Engine/Source/Runtime/Renderer/Private/ForwardShadingRenderer.cpp
+++ b/Engine/Source/Runtime/Renderer/Private/ForwardShadingRenderer.cpp
@@ -96,8 +96,6 @@ void FForwardShadingSceneRenderer::Render(FRHICommandListImmediate& RHICmdList)
 	// Find the visible primitives.
 	InitViews(RHICmdList);
 	
-	RenderShadowDepthMaps(RHICmdList);
-
 	// Notify the FX system that the scene is about to be rendered.
 	if (Scene->FXSystem)
 	{
@@ -106,6 +104,8 @@ void FForwardShadingSceneRenderer::Render(FRHICommandListImmediate& RHICmdList)
 
 	GRenderTargetPool.VisualizeTexture.OnStartFrame(Views[0]);
 
+	RenderShadowDepthMaps(RHICmdList);
+
 	// Dynamic vertex and index buffers need to be committed before rendering.
 	FGlobalDynamicVertexBuffer::Get().Commit();
 	FGlobalDynamicIndexBuffer::Get().Commit();
@@ -166,6 +166,8 @@ void FForwardShadingSceneRenderer::Render(FRHICommandListImmediate& RHICmdList)
 		RenderTranslucency(RHICmdList);
 	}
 
+	RenderModulatedShadowProjections(RHICmdList);
+
 	static const auto CVarMobileMSAA = IConsoleManager::Get().FindTConsoleVariableDataInt(TEXT("r.MobileMSAA"));
 	bool bOnChipSunMask =
 		GSupportsRenderTargetFormat_PF_FloatRGBA &&
diff --git a/Engine/Source/Runtime/Renderer/Private/SceneRendering.h b/Engine/Source/Runtime/Renderer/Private/SceneRendering.h
--- a/Engine/Source/Runtime/Renderer/Private/SceneRendering.h
+++ b/Engine/Source/Runtime/Renderer/Private/SceneRendering.h
@@ -723,6 +723,39 @@ protected:
 
 	// Shared functionality between all scene renderers
 
+	/** Renders the projections of the given Shadows to the appropriate color render target. */
+	void RenderProjections(
+		FRHICommandListImmediate& RHICmdList,
+		const FLightSceneInfo* LightSceneInfo,
+		const TArray<FProjectedShadowInfo*, SceneRenderingAllocator>& Shadows,
+		bool bForwardShading
+		);
+
+	/** Finds a matching cached preshadow, if one exists. */
+	TRefCountPtr<FProjectedShadowInfo> GetCachedPreshadow(
+		const FLightPrimitiveInteraction* InParentInteraction,
+		const FProjectedShadowInitializer& Initializer,
+		const FBoxSphereBounds& Bounds,
+		uint32 InResolutionX);
+
+	/** Creates a per object projected shadow for the given interaction. */
+	void CreatePerObjectProjectedShadow(
+		FRHICommandListImmediate& RHICmdList,
+		FLightPrimitiveInteraction* Interaction,
+		bool bCreateTranslucentObjectShadow,
+		bool bCreateInsetObjectShadow,
+		const TArray<FProjectedShadowInfo*, SceneRenderingAllocator>& ViewDependentWholeSceneShadows,
+		TArray<FProjectedShadowInfo*, SceneRenderingAllocator>& OutPreShadows);
+
+	/** Creates shadows for the given interaction. */
+	void SetupInteractionShadows(
+		FRHICommandListImmediate& RHICmdList,
+		FLightPrimitiveInteraction* Interaction,
+		FVisibleLightInfo& VisibleLightInfo,
+		bool bStaticSceneOnly,
+		const TArray<FProjectedShadowInfo*, SceneRenderingAllocator>& ViewDependentWholeSceneShadows,
+		TArray<FProjectedShadowInfo*, SceneRenderingAllocator>& PreShadows);
+
 	/** Generates FProjectedShadowInfos for all wholesceneshadows on the given light.*/
 	void AddViewDependentWholeSceneShadowsForView(
 		TArray<FProjectedShadowInfo*, SceneRenderingAllocator>& ShadowInfos, 
@@ -837,6 +870,9 @@ protected:
 	/** Renders the opaque base pass for forward shading. */
 	void RenderForwardShadingBasePass(FRHICommandListImmediate& RHICmdList);
 
+	/** Projects any shadow depth images into the scene. (modulated shadows) */
+	void RenderModulatedShadowProjections(FRHICommandListImmediate& RHICmdList);
+
 	/** Makes a copy of scene alpha so PC can emulate ES2 framebuffer fetch. */
 	void CopySceneAlpha(FRHICommandListImmediate& RHICmdList, const FViewInfo& View);
 
@@ -859,4 +895,7 @@ protected:
 	  * @return true if anything got rendered
 	  */
 	bool RenderShadowDepthMap(FRHICommandListImmediate& RHICmdList, const FLightSceneInfo* LightSceneInfo);
+
+	/** Helper function to determine if a shadow should be cast **/
+	bool IsProjectedShadowPotentiallyVisible(const FProjectedShadowInfo* ProjectedShadowInfo) const;
 };
diff --git a/Engine/Source/Runtime/Renderer/Private/ShadowRendering.cpp b/Engine/Source/Runtime/Renderer/Private/ShadowRendering.cpp
--- a/Engine/Source/Runtime/Renderer/Private/ShadowRendering.cpp
+++ b/Engine/Source/Runtime/Renderer/Private/ShadowRendering.cpp
@@ -52,6 +52,12 @@ static TAutoConsoleVariable<float> CVarSpotLightShadowDepthBias(
 	TEXT("Depth bias that is applied in the depth pass for per object projected shadows from spot lights"),
 	ECVF_RenderThreadSafe);
 
+static TAutoConsoleVariable<int32> CVarEnableModulatedSelfShadow(
+	TEXT("r.Shadow.EnableModulatedSelfShadow"),
+	0,
+	TEXT("Allows modulated shadows to affect the shadow caster. (mobile only)"),
+	ECVF_RenderThreadSafe);
+
 // 0:off, 1:low, 2:med, 3:high, 4:very high, 5:max
 uint32 GetShadowQuality()
 {
@@ -609,7 +615,7 @@ TGlobalResource<StencilingGeometry::FStencilConeIndexBuffer> StencilingGeometry:
 
 bool FShadowProjectionVS::ShouldCache(EShaderPlatform Platform)
 {
-	return IsFeatureLevelSupported(Platform, ERHIFeatureLevel::SM4);
+	return true;
 }
 
 void FShadowProjectionVS::SetParameters(FRHICommandList& RHICmdList, const FSceneView& View, const FProjectedShadowInfo* ShadowInfo)
@@ -661,6 +667,13 @@ IMPLEMENT_SHADOW_PROJECTION_PIXEL_SHADER(4,true);
 IMPLEMENT_SHADOW_PROJECTION_PIXEL_SHADER(5,true);
 #endif
 
+// Implement a pixel shader for rendering modulated shadow projections.
+IMPLEMENT_SHADER_TYPE(template<>, TModulatedShadowProjection<1>, TEXT("ShadowProjectionPixelShader"), TEXT("Main"), SF_Pixel);
+IMPLEMENT_SHADER_TYPE(template<>, TModulatedShadowProjection<2>, TEXT("ShadowProjectionPixelShader"), TEXT("Main"), SF_Pixel);
+IMPLEMENT_SHADER_TYPE(template<>, TModulatedShadowProjection<3>, TEXT("ShadowProjectionPixelShader"), TEXT("Main"), SF_Pixel);
+IMPLEMENT_SHADER_TYPE(template<>, TModulatedShadowProjection<4>, TEXT("ShadowProjectionPixelShader"), TEXT("Main"), SF_Pixel);
+IMPLEMENT_SHADER_TYPE(template<>, TModulatedShadowProjection<5>, TEXT("ShadowProjectionPixelShader"), TEXT("Main"), SF_Pixel);
+
 // with different quality levels
 IMPLEMENT_SHADER_TYPE(template<>,TShadowProjectionFromTranslucencyPS<1>,TEXT("ShadowProjectionPixelShader"),TEXT("Main"),SF_Pixel);
 IMPLEMENT_SHADER_TYPE(template<>,TShadowProjectionFromTranslucencyPS<2>,TEXT("ShadowProjectionPixelShader"),TEXT("Main"),SF_Pixel);
@@ -1877,7 +1890,7 @@ void StencilingGeometry::DrawCone(FRHICommandList& RHICmdList)
 static FGlobalBoundShaderState MaskBoundShaderState[2];
 
 template <uint32 Quality>
-static void SetShadowProjectionShaderTemplNew(FRHICommandList& RHICmdList, int32 ViewIndex, const FViewInfo& View, const FProjectedShadowInfo* ShadowInfo)
+static void SetShadowProjectionShaderTemplNew(FRHICommandList& RHICmdList, int32 ViewIndex, const FViewInfo& View, const FProjectedShadowInfo* ShadowInfo, bool bForwardShading)
 {
 	if (ShadowInfo->bTranslucentShadow)
 	{
@@ -1934,7 +1947,9 @@ static void SetShadowProjectionShaderTemplNew(FRHICommandList& RHICmdList, int32
 
 		// Get the Shadow Projection Pixel Shader
 		// This shader is the ordinary projection shader used by point/spot lights.		
-		FShadowProjectionPixelShaderInterface* ShadowProjPS = View.ShaderMap->GetShader<TShadowProjectionPS<Quality, false> >();
+		FShadowProjectionPixelShaderInterface* ShadowProjPS = bForwardShading ?
+					View.ShaderMap->GetShader<TModulatedShadowProjection<Quality> >() :
+					View.ShaderMap->GetShader<TShadowProjectionPS<Quality, false> >();
 
 		static FGlobalBoundShaderState BoundShaderState;
 		
@@ -1945,7 +1960,7 @@ static void SetShadowProjectionShaderTemplNew(FRHICommandList& RHICmdList, int32
 	}
 }
 
-void FProjectedShadowInfo::RenderProjection(FRHICommandListImmediate& RHICmdList, int32 ViewIndex, const FViewInfo* View) const
+void FProjectedShadowInfo::RenderProjection(FRHICommandListImmediate& RHICmdList, int32 ViewIndex, const FViewInfo* View, bool bForwardShading) const
 {
 #if WANTS_DRAW_MESH_EVENTS
 	FString EventName;
@@ -2228,6 +2243,25 @@ void FProjectedShadowInfo::RenderProjection(FRHICommandListImmediate& RHICmdList
 
 		// Draw the frustum using the stencil buffer to mask just the pixels which are inside the shadow frustum.
 		DrawIndexedPrimitiveUP(RHICmdList, PT_TriangleList, 0, 8, 12, GCubeIndices, sizeof(uint16), FrustumVertices, sizeof(FVector4));
+
+		// if rendering modulated shadows mask out subject mesh elements to prevent self shadowing.
+		if (bForwardShading && !CVarEnableModulatedSelfShadow.GetValueOnRenderThread())
+		{
+			RHICmdList.SetDepthStencilState(TStaticDepthStencilState<
+				false, CF_DepthNearOrEqual,
+				true, CF_Always, SO_Keep, SO_Keep, SO_Replace,
+				true, CF_Always, SO_Keep, SO_Keep, SO_Replace,
+				0xff, 0xff
+			>::GetRHI(), 0);
+
+			FDepthDrawingPolicyFactory::ContextType Context(DDM_AllOccluders);
+			for (int32 MeshBatchIndex = 0; MeshBatchIndex < DynamicSubjectMeshElements.Num(); MeshBatchIndex++)
+			{
+				const FMeshBatchAndRelevance& MeshBatchAndRelevance = DynamicSubjectMeshElements[MeshBatchIndex];
+				const FMeshBatch& MeshBatch = *MeshBatchAndRelevance.Mesh;
+				FDepthDrawingPolicyFactory::DrawDynamicMesh(RHICmdList, *View, Context, MeshBatch, false, true, MeshBatchAndRelevance.PrimitiveSceneProxy, MeshBatch.BatchHitProxyId);
+			}
+		}
 	}
 
 	// solid rasterization w/ back-face culling.
@@ -2280,12 +2314,20 @@ void FProjectedShadowInfo::RenderProjection(FRHICommandListImmediate& RHICmdList
 	}
 	else
 	{
-		// use B and A in Light Attenuation (see above)
+		if (bForwardShading)
+		{
+			// Color modulate shadows, ignore alpha.
+			RHICmdList.SetBlendState(TStaticBlendState<CW_RGB, BO_Add, BF_Zero, BF_SourceColor, BO_Add, BF_Zero, BF_One>::GetRHI());
+		}
+		else
+		{
+			// use B and A in Light Attenuation (see above)
 
-		// (CO_Min is needed to combine multiple shadow passes)
-		RHICmdList.SetBlendState(TStaticBlendState<CW_BA, BO_Min, BF_One, BF_One, BO_Min, BF_One, BF_One>::GetRHI());
+			// (CO_Min is needed to combine multiple shadow passes)
+			RHICmdList.SetBlendState(TStaticBlendState<CW_BA, BO_Min, BF_One, BF_One, BO_Min, BF_One, BF_One>::GetRHI());
+		}
 	}
-	
+
 	{
 		uint32 LocalQuality = GetShadowQuality();
 
@@ -2320,11 +2362,11 @@ void FProjectedShadowInfo::RenderProjection(FRHICommandListImmediate& RHICmdList
 
 		switch(LocalQuality)
 		{
-			case 1: SetShadowProjectionShaderTemplNew<1>(RHICmdList, ViewIndex, *View, this); break;
-			case 2: SetShadowProjectionShaderTemplNew<2>(RHICmdList, ViewIndex, *View, this); break;
-			case 3: SetShadowProjectionShaderTemplNew<3>(RHICmdList, ViewIndex, *View, this); break;
-			case 4: SetShadowProjectionShaderTemplNew<4>(RHICmdList, ViewIndex, *View, this); break;
-			case 5: SetShadowProjectionShaderTemplNew<5>(RHICmdList, ViewIndex, *View, this); break;
+			case 1: SetShadowProjectionShaderTemplNew<1>(RHICmdList, ViewIndex, *View, this, bForwardShading); break;
+			case 2: SetShadowProjectionShaderTemplNew<2>(RHICmdList, ViewIndex, *View, this, bForwardShading); break;
+			case 3: SetShadowProjectionShaderTemplNew<3>(RHICmdList, ViewIndex, *View, this, bForwardShading); break;
+			case 4: SetShadowProjectionShaderTemplNew<4>(RHICmdList, ViewIndex, *View, this, bForwardShading); break;
+			case 5: SetShadowProjectionShaderTemplNew<5>(RHICmdList, ViewIndex, *View, this, bForwardShading); break;
 			default:
 				check(0);
 		}
@@ -2851,15 +2893,24 @@ bool FDeferredShadingSceneRenderer::RenderOnePassPointLightShadows(FRHICommandLi
 }
 
 /** Renders the projections of the given Shadows to the appropriate color render target. */
-void FDeferredShadingSceneRenderer::RenderProjections(
+void FSceneRenderer::RenderProjections(
 	FRHICommandListImmediate& RHICmdList,
 	const FLightSceneInfo* LightSceneInfo, 
-	const TArray<FProjectedShadowInfo*,SceneRenderingAllocator>& Shadows
+	const TArray<FProjectedShadowInfo*,SceneRenderingAllocator>& Shadows,
+	bool bForwardShading
 	)
 {
 	FSceneRenderTargets& SceneContext = FSceneRenderTargets::Get(RHICmdList);
-	// Normal shadows render to light attenuation
-	SceneContext.BeginRenderingLightAttenuation(RHICmdList);
+	
+	if (bForwardShading)
+	{
+		SceneContext.BeginRenderingSceneColor(RHICmdList);
+	}
+	else
+	{
+		// Normal deferred shadows render to light attenuation
+		SceneContext.BeginRenderingLightAttenuation(RHICmdList);
+	}
 
 	for (int32 ViewIndex = 0; ViewIndex < Views.Num(); ViewIndex++)
 	{
@@ -2882,9 +2933,11 @@ void FDeferredShadingSceneRenderer::RenderProjections(
 				// Only project the shadow if it's large enough in this particular view (split screen, etc... may have shadows that are large in one view but irrelevantly small in others)
 				if (ProjectedShadowInfo->FadeAlphas[ViewIndex] > 1.0f / 256.0f)
 				{
-					ProjectedShadowInfo->RenderProjection(RHICmdList, ViewIndex, &View);
-
-					GRenderTargetPool.VisualizeTexture.SetCheckPoint(RHICmdList, SceneContext.GetLightAttenuation());
+					ProjectedShadowInfo->RenderProjection(RHICmdList, ViewIndex, &View, bForwardShading);
+					if (!bForwardShading)
+					{
+						GRenderTargetPool.VisualizeTexture.SetCheckPoint(RHICmdList, SceneContext.GetLightAttenuation());
+					}
 				}
 			}
 		}
@@ -3060,7 +3113,7 @@ bool FDeferredShadingSceneRenderer::RenderTranslucentProjectedShadows(FRHIComman
 		{
 			{
 				SCOPED_DRAW_EVENT(RHICmdList, ShadowProjectionOnOpaque);
-				RenderProjections(RHICmdList, LightSceneInfo, Shadows);
+				RenderProjections(RHICmdList, LightSceneInfo, Shadows, false);
 			}
 
 			for (int32 ShadowIndex = 0; ShadowIndex < Shadows.Num(); ShadowIndex++)
@@ -3418,7 +3471,7 @@ bool FDeferredShadingSceneRenderer::RenderProjectedShadows(FRHICommandListImmedi
 		{
 			{
 				SCOPED_DRAW_EVENT(RHICmdList, ShadowProjectionOnOpaque);
-				RenderProjections(RHICmdList, LightSceneInfo, Shadows);
+				RenderProjections(RHICmdList, LightSceneInfo, Shadows, false);
 			}
 			
 			for (int32 ShadowIndex = 0; ShadowIndex < Shadows.Num(); ShadowIndex++)
@@ -3588,7 +3641,7 @@ bool FDeferredShadingSceneRenderer::RenderCachedPreshadows(FRHICommandListImmedi
 		// Render the shadow projections.
 		{
 			SCOPED_DRAW_EVENTF(RHICmdList, EventShadowProj, TEXT("Cached Preshadow Projections"));
-			RenderProjections(RHICmdList, LightSceneInfo, OpaqueCachedPreshadows);
+			RenderProjections(RHICmdList, LightSceneInfo, OpaqueCachedPreshadows, false);
 
 			// Mark the attenuation buffer as dirty.
 			bAttenuationBufferDirty = true;
@@ -3632,33 +3685,16 @@ bool FForwardShadingSceneRenderer::RenderShadowDepthMap(FRHICommandListImmediate
 				&& VisibleLightViewInfo.ProjectedShadowVisibilityMap[ShadowIndex];
 		}
 
-		// Skip shadows which will be handled in RenderOnePassPointLightShadows or RenderReflectiveShadowmaps
-		if (ProjectedShadowInfo->CascadeSettings.bOnePassPointLightShadow || ProjectedShadowInfo->bReflectiveShadowmap)
-		{
-			bShadowIsVisible = false;
-		}
-
-		if (ProjectedShadowInfo->bPreShadow && !LightSceneInfo->Proxy->HasStaticShadowing())
-		{
-			bShadowIsVisible = false;
-		}
-
-		if (bShadowIsVisible
-			&& (!ProjectedShadowInfo->bPreShadow || ProjectedShadowInfo->HasSubjectPrims())
-			&& !ProjectedShadowInfo->bAllocatedInPreshadowCache)
+		if (IsProjectedShadowPotentiallyVisible(ProjectedShadowInfo) && (ProjectedShadowInfo->bWholeSceneShadow || ProjectedShadowInfo->bPreShadow))
 		{
 			// Add the shadow to the list of visible shadows cast by this light.
 			if (ProjectedShadowInfo->bWholeSceneShadow)
 			{
 				INC_DWORD_STAT(STAT_WholeSceneShadows);
 			}
-			else if (ProjectedShadowInfo->bPreShadow)
-			{
-				INC_DWORD_STAT(STAT_PreShadows);
-			}
 			else
 			{
-				INC_DWORD_STAT(STAT_PerObjectShadows);
+				INC_DWORD_STAT(STAT_PreShadows);
 			}
 			Shadows.Add(ProjectedShadowInfo);
 		}
@@ -3713,7 +3749,7 @@ void FForwardShadingSceneRenderer::RenderShadowDepthMaps(FRHICommandListImmediat
 
 	// render shadowmaps for relevant lights.
 	if (Scene->SimpleDirectionalLight)
-	{		
+	{
 		const FLightSceneInfo& LightSceneInfo = *Scene->SimpleDirectionalLight;
 
 		if (LightSceneInfo.ShouldRenderViewIndependentWholeSceneShadows()
@@ -3730,21 +3766,203 @@ void FForwardShadingSceneRenderer::RenderShadowDepthMaps(FRHICommandListImmediat
 					bRender = true;
 					break;
 				}
-			}			
+			}
 
 			if (bRender)
-			{				
+			{
 				FScopeCycleCounter Context(LightSceneInfo.Proxy->GetStatId());
 
 				FString LightNameWithLevel;
 				GetLightNameForDrawEvent(LightSceneInfo.Proxy, LightNameWithLevel);
 				SCOPED_DRAW_EVENTF(RHICmdList, EventLightPass, *LightNameWithLevel);
-				
+
 				INC_DWORD_STAT(STAT_NumShadowedLights);
 
 				//render shadowmap depth				
-				RenderShadowDepthMap(RHICmdList, &LightSceneInfo);				
-			}					
-		}		
-	}	
+				RenderShadowDepthMap(RHICmdList, &LightSceneInfo);
+			}
+		}
+	}
+}
+
+void FForwardShadingSceneRenderer::RenderModulatedShadowProjections(FRHICommandListImmediate& RHICmdList)
+{
+	SCOPED_DRAW_EVENT(RHICmdList, Lights);
+
+	bool bDynamicShadows = ViewFamily.EngineShowFlags.DynamicShadows && GetShadowQuality() > 0;
+
+	if (IsSimpleDynamicLightingEnabled() || !bDynamicShadows)
+	{
+		return;
+	}
+	FSceneRenderTargets& SceneContext = FSceneRenderTargets::Get(RHICmdList);
+
+	// render shadowmaps for relevant lights.
+	for (TSparseArray<FLightSceneInfoCompact>::TConstIterator LightIt(Scene->Lights); LightIt; ++LightIt)
+	{
+		const FLightSceneInfoCompact& LightSceneInfoCompact = *LightIt;
+		FLightSceneInfo* LightSceneInfo = LightSceneInfoCompact.LightSceneInfo;
+
+		SCOPE_CYCLE_COUNTER(STAT_ProjectedShadowDrawTime);
+
+		if (!LightSceneInfo->Proxy->CastsModulatedShadows())
+		{
+			continue;
+		}
+
+		// Find the projected shadows cast by this light.
+		FVisibleLightInfo& VisibleLightInfo = VisibleLightInfos[LightSceneInfo->Id];
+		TArray<FProjectedShadowInfo*, SceneRenderingAllocator> Shadows;
+		for (int32 ShadowIndex = 0; ShadowIndex < VisibleLightInfo.AllProjectedShadows.Num(); ShadowIndex++)
+		{
+			FProjectedShadowInfo* ProjectedShadowInfo = VisibleLightInfo.AllProjectedShadows[ShadowIndex];
+
+			// Check that the shadow is visible in at least one view before rendering it.
+			bool bShadowIsVisible = false;
+			for (int32 ViewIndex = 0; ViewIndex < Views.Num(); ViewIndex++)
+			{
+				const FViewInfo& View = Views[ViewIndex];
+				if (ProjectedShadowInfo->DependentView && ProjectedShadowInfo->DependentView != &View)
+				{
+					continue;
+				}
+				const FVisibleLightViewInfo& VisibleLightViewInfo = View.VisibleLightInfos[LightSceneInfo->Id];
+				const FPrimitiveViewRelevance ViewRelevance = VisibleLightViewInfo.ProjectedShadowViewRelevanceMap[ShadowIndex];
+				bShadowIsVisible |= !ProjectedShadowInfo->bTranslucentShadow && ViewRelevance.bOpaqueRelevance
+					&& VisibleLightViewInfo.ProjectedShadowVisibilityMap[ShadowIndex];
+			}
+
+			bShadowIsVisible &= IsProjectedShadowPotentiallyVisible(ProjectedShadowInfo);
+
+			// skip shadows that are used with forward basepass (whole scene shadows)
+			if (ProjectedShadowInfo->bWholeSceneShadow)
+			{
+				bShadowIsVisible = false;
+			}
+
+			if (bShadowIsVisible)
+			{
+				if (ProjectedShadowInfo->bPreShadow)
+				{
+					INC_DWORD_STAT(STAT_PreShadows);
+				}
+				else
+				{
+					INC_DWORD_STAT(STAT_PerObjectShadows);
+				}
+				Shadows.Add(ProjectedShadowInfo);
+			}
+		}
+
+		// Sort the projected shadows by resolution.
+		Shadows.Sort(FCompareFProjectedShadowInfoBySplitIndex());
+
+		int32 NumShadowsRendered = 0;
+
+		while (NumShadowsRendered < Shadows.Num())
+		{
+			int32 NumAllocatedShadows = 0;
+
+			// Allocate shadow texture space to the shadows.
+			const FIntPoint ShadowBufferResolution = SceneContext.GetShadowDepthTextureResolution();
+			FTextureLayout ShadowLayout(1, 1, ShadowBufferResolution.X, ShadowBufferResolution.Y, false, false);
+
+			for (int32 ShadowIndex = 0; ShadowIndex < Shadows.Num(); ShadowIndex++)
+			{
+				FProjectedShadowInfo* ProjectedShadowInfo = Shadows[ShadowIndex];
+				if (!ProjectedShadowInfo->bRendered)
+				{
+					if (ShadowLayout.AddElement(
+						ProjectedShadowInfo->X,
+						ProjectedShadowInfo->Y,
+						ProjectedShadowInfo->ResolutionX + SHADOW_BORDER * 2,
+						ProjectedShadowInfo->ResolutionY + SHADOW_BORDER * 2)
+						)
+					{
+						ProjectedShadowInfo->bAllocated = true;
+						NumAllocatedShadows++;
+					}
+				}
+			}
+
+			// Abort if we encounter a shadow that doesn't fit in the render target.
+			if (!NumAllocatedShadows)
+			{
+				break;
+			}
+
+			{
+				// Render the shadow depths.
+				SCOPED_DRAW_EVENT(RHICmdList, ShadowDepthsFromOpaqueProjected);
+
+				bool bPerformClear = true;
+				auto SetShadowRenderTargets = [&bPerformClear, &SceneContext](FRHICommandList& InRHICmdList)
+				{
+					SceneContext.BeginRenderingShadowDepth(InRHICmdList, bPerformClear);
+				};
+
+				SetShadowRenderTargets(RHICmdList);  // run it now, maybe run it later for parallel command lists
+				bPerformClear = false;
+
+				// render depth for each shadow
+				for (int32 ShadowIndex = 0; ShadowIndex < Shadows.Num(); ShadowIndex++)
+				{
+					FProjectedShadowInfo* ProjectedShadowInfo = Shadows[ShadowIndex];
+					if (ProjectedShadowInfo->bAllocated && !ProjectedShadowInfo->bTranslucentShadow && !ProjectedShadowInfo->CascadeSettings.bRayTracedDistanceField)
+					{
+						ProjectedShadowInfo->RenderDepth(RHICmdList, this, SetShadowRenderTargets);
+					}
+				}
+
+				SceneContext.FinishRenderingShadowDepth(RHICmdList);
+			}
+
+			// Render the shadow projections.
+			{
+				SCOPED_DRAW_EVENT(RHICmdList, ShadowProjectionOnOpaque);
+				RenderProjections(RHICmdList, LightSceneInfo, Shadows, true);
+			}
+
+			// Mark and count the rendered shadows.
+			for (int32 ShadowIndex = 0; ShadowIndex < Shadows.Num(); ShadowIndex++)
+			{
+				FProjectedShadowInfo* ProjectedShadowInfo = Shadows[ShadowIndex];
+				if (ProjectedShadowInfo->bAllocated)
+				{
+					ProjectedShadowInfo->bAllocated = false;
+					ProjectedShadowInfo->bRendered = true;
+					NumShadowsRendered++;
+				}
+			}
+		}
+	}
+}
+
+bool FForwardShadingSceneRenderer::IsProjectedShadowPotentiallyVisible(const FProjectedShadowInfo* ProjectedShadowInfo) const
+{
+	// Check that the shadow is visible in at least one view before rendering it.
+	const FLightSceneInfo& LightSceneInfo = ProjectedShadowInfo->GetLightSceneInfo();
+
+	// Skip shadows which will be handled in RenderOnePassPointLightShadows or RenderReflectiveShadowmaps
+	if (ProjectedShadowInfo->CascadeSettings.bOnePassPointLightShadow || ProjectedShadowInfo->bReflectiveShadowmap)
+	{
+		return false;
+	}
+
+	if (ProjectedShadowInfo->bPreShadow && !LightSceneInfo.Proxy->HasStaticShadowing())
+	{
+		return false;
+	}
+
+	if (ProjectedShadowInfo->bWholeSceneShadow && LightSceneInfo.Proxy->HasStaticShadowing())
+	{
+		return false;
+	}
+
+	if ((ProjectedShadowInfo->bPreShadow && !ProjectedShadowInfo->HasSubjectPrims()) || ProjectedShadowInfo->bAllocatedInPreshadowCache)
+	{
+		return false;
+	}
+
+	return true;
 }
diff --git a/Engine/Source/Runtime/Renderer/Private/ShadowRendering.h b/Engine/Source/Runtime/Renderer/Private/ShadowRendering.h
--- a/Engine/Source/Runtime/Renderer/Private/ShadowRendering.h
+++ b/Engine/Source/Runtime/Renderer/Private/ShadowRendering.h
@@ -658,7 +658,7 @@ public:
 	/**
 	 * Projects the shadow onto the scene for a particular view.
 	 */
-	void RenderProjection(FRHICommandListImmediate& RHICmdList, int32 ViewIndex, const class FViewInfo* View) const;
+	void RenderProjection(FRHICommandListImmediate& RHICmdList, int32 ViewIndex, const class FViewInfo* View, bool bForwardShading) const;
 
 	/** Renders ray traced distance field shadows. */
 	void RenderRayTracedDistanceFieldProjection(FRHICommandListImmediate& RHICmdList, const class FViewInfo& View) const;
@@ -1235,6 +1235,58 @@ protected:
 	FShaderParameter ShadowSharpen;
 };
 
+/** Pixel shader to project modulated shadows onto the scene. */
+template<uint32 Quality>
+class TModulatedShadowProjection : public TShadowProjectionPS<Quality>
+{
+	DECLARE_SHADER_TYPE(TModulatedShadowProjection, Global);
+public:
+
+	static void ModifyCompilationEnvironment(EShaderPlatform Platform, FShaderCompilerEnvironment& OutEnvironment)
+	{
+		TShadowProjectionPS<Quality>::ModifyCompilationEnvironment(Platform, OutEnvironment);
+		OutEnvironment.SetDefine(TEXT("MODULATED_SHADOWS"), 1);
+	}
+
+	static bool ShouldCache(EShaderPlatform Platform)
+	{
+		return IsMobilePlatform(Platform);
+	}
+
+	TModulatedShadowProjection() {}
+
+	TModulatedShadowProjection(const ShaderMetaType::CompiledShaderInitializerType& Initializer) :
+		TShadowProjectionPS<Quality>(Initializer)
+	{
+		ModulatedShadowColorParameter.Bind(Initializer.ParameterMap, TEXT("ModulatedShadowColor"));
+	}
+
+	virtual void SetParameters(
+		FRHICommandList& RHICmdList,
+		int32 ViewIndex,
+		const FSceneView& View,
+		const FProjectedShadowInfo* ShadowInfo) override
+	{
+		TShadowProjectionPS<Quality>::SetParameters(RHICmdList, ViewIndex, View, ShadowInfo);
+		const FPixelShaderRHIParamRef ShaderRHI = this->GetPixelShader();
+		SetShaderValue(RHICmdList, ShaderRHI, ModulatedShadowColorParameter, ShadowInfo->GetLightSceneInfo().Proxy->GetModulatedShadowColor());
+	}
+
+	/**
+	* Serialize the parameters for this shader
+	* @param Ar - archive to serialize to
+	*/
+	virtual bool Serialize(FArchive& Ar) override
+	{
+		bool bShaderHasOutdatedParameters = TShadowProjectionPS<Quality>::Serialize(Ar);
+		Ar << ModulatedShadowColorParameter;
+		return bShaderHasOutdatedParameters;
+	}
+
+protected:
+	FShaderParameter ModulatedShadowColorParameter;
+};
+
 /** Translucency shadow projection parameters used by multiple shaders. */
 class FTranslucencyShadowProjectionShaderParameters
 {
diff --git a/Engine/Source/Runtime/Renderer/Private/ShadowSetup.cpp b/Engine/Source/Runtime/Renderer/Private/ShadowSetup.cpp
--- a/Engine/Source/Runtime/Renderer/Private/ShadowSetup.cpp
+++ b/Engine/Source/Runtime/Renderer/Private/ShadowSetup.cpp
@@ -1073,7 +1073,7 @@ void FProjectedShadowInfo::ClearTransientArrays()
 }
 
 /** Returns a cached preshadow matching the input criteria if one exists. */
-TRefCountPtr<FProjectedShadowInfo> FDeferredShadingSceneRenderer::GetCachedPreshadow(
+TRefCountPtr<FProjectedShadowInfo> FSceneRenderer::GetCachedPreshadow(
 	const FLightPrimitiveInteraction* InParentInteraction, 
 	const FProjectedShadowInitializer& Initializer,
 	const FBoxSphereBounds& Bounds,
@@ -1218,7 +1218,7 @@ bool FSceneRenderer::ShouldCreateObjectShadowForStationaryLight(const FLightScen
 	return bCreateObjectShadowForStationaryLight;
 }
 
-void FDeferredShadingSceneRenderer::SetupInteractionShadows(
+void FSceneRenderer::SetupInteractionShadows(
 	FRHICommandListImmediate& RHICmdList,
 	FLightPrimitiveInteraction* Interaction, 
 	FVisibleLightInfo& VisibleLightInfo, 
@@ -1260,7 +1260,7 @@ void FDeferredShadingSceneRenderer::SetupInteractionShadows(
 	}
 }
 
-void FDeferredShadingSceneRenderer::CreatePerObjectProjectedShadow(
+void FSceneRenderer::CreatePerObjectProjectedShadow(
 	FRHICommandListImmediate& RHICmdList,
 	FLightPrimitiveInteraction* Interaction, 
 	bool bCreateTranslucentObjectShadow, 
@@ -2320,6 +2320,15 @@ void FForwardShadingSceneRenderer::InitDynamicShadows(FRHICommandListImmediate&
 	TArray<FProjectedShadowInfo*, SceneRenderingAllocator> ViewDependentWholeSceneShadows;
 	TArray<FProjectedShadowInfo*, SceneRenderingAllocator> ViewDependentWholeSceneShadowsThatNeedCulling;
 	TArray<FProjectedShadowInfo*, SceneRenderingAllocator> PreShadows;
+
+	bool bStaticSceneOnly = false;
+
+	for (int32 ViewIndex = 0; ViewIndex < Views.Num(); ViewIndex++)
+	{
+		FViewInfo& View = Views[ViewIndex];
+		bStaticSceneOnly = bStaticSceneOnly || View.bStaticSceneOnly;
+	}
+
 	{
 		SCOPE_CYCLE_COUNTER(STAT_InitDynamicShadowsTime);
 
@@ -2329,63 +2338,74 @@ void FForwardShadingSceneRenderer::InitDynamicShadows(FRHICommandListImmediate&
 			FLightSceneInfo* LightSceneInfo = LightSceneInfoCompact.LightSceneInfo;
 			FVisibleLightInfo& VisibleLightInfo = VisibleLightInfos[LightSceneInfo->Id];
 			
-			// Only consider lights that may have shadows.
-			if (LightSceneInfo->ShouldRenderViewIndependentWholeSceneShadows())
+			// see if the light is visible in any view
+			bool bIsVisibleInAnyView = false;
+
+			for (int32 ViewIndex = 0; ViewIndex < Views.Num(); ViewIndex++)
 			{
-				// see if the light is visible in any view
-				bool bIsVisibleInAnyView = false;
+				// View frustums are only checked when lights have visible primitives or have modulated shadows,
+				// so we don't need to check for that again here
+				bIsVisibleInAnyView = LightSceneInfo->ShouldRenderLight(Views[ViewIndex]);
 
-				for (int32 ViewIndex = 0; ViewIndex < Views.Num(); ViewIndex++)
+				if (bIsVisibleInAnyView)
 				{
-					// View frustums are only checked when lights have visible primitives or have modulated shadows,
-					// so we don't need to check for that again here
-					bIsVisibleInAnyView = LightSceneInfo->ShouldRenderLight(Views[ViewIndex]);
-
-					if (bIsVisibleInAnyView)
-					{
-						break;
-					}
+					break;
 				}
+			}
 
-				if (bIsVisibleInAnyView)
-				{									
-					AddViewDependentWholeSceneShadowsForView(ViewDependentWholeSceneShadows, ViewDependentWholeSceneShadowsThatNeedCulling, VisibleLightInfo, *LightSceneInfo);					
+			int32 NumWholeSceneShadows = 0;
+			if (bIsVisibleInAnyView)
+			{
+				if (LightSceneInfo->ShouldRenderViewIndependentWholeSceneShadows() && Scene->SimpleDirectionalLight == LightSceneInfo)
+				{
+					AddViewDependentWholeSceneShadowsForView(ViewDependentWholeSceneShadows, ViewDependentWholeSceneShadowsThatNeedCulling, VisibleLightInfo, *LightSceneInfo);
 				}
+				NumWholeSceneShadows = VisibleLightInfo.AllProjectedShadows.Num();
 
-				const int32 NumShadows = FMath::Min(VisibleLightInfo.AllProjectedShadows.Num(), MAX_FORWARD_SHADOWCASCADES);
-				if (NumShadows > 0)
+				// If we're casting modulated shadows then look for dynamic shadow interactions.
+				if (LightSceneInfo->Proxy->CastsModulatedShadows())
 				{
-					FSceneRenderTargets& SceneContext = FSceneRenderTargets::Get(RHICmdList);
-					//create the shadow depth texture and/or surface
-					const FIntPoint ShadowBufferResolution = SceneContext.GetShadowDepthTextureResolution();
-					const int32 MaxWide = GMaxShadowDepthBufferSizeX / ShadowBufferResolution.X;
-					const int32 MaxHigh = GMaxShadowDepthBufferSizeY / ShadowBufferResolution.Y;
+					// Look for individual primitives with a dynamic shadow.
+					for (FLightPrimitiveInteraction* Interaction = LightSceneInfo->DynamicPrimitiveList;
+						Interaction;
+						Interaction = Interaction->GetNextPrimitive()
+						)
+					{
+						SetupInteractionShadows(RHICmdList, Interaction, VisibleLightInfo, bStaticSceneOnly, ViewDependentWholeSceneShadows, PreShadows);
+					}
+				}
+			}
 
+			NumWholeSceneShadows = FMath::Min(NumWholeSceneShadows, MAX_FORWARD_SHADOWCASCADES);
 
-					const int32 NumWide = FMath::Min(NumShadows, MaxWide);
-					const int32 NumHigh = FMath::Min(((NumShadows - 1) / MaxWide) + 1, MaxHigh);
+			if (NumWholeSceneShadows > 0)
+			{
+				FSceneRenderTargets& SceneContext = FSceneRenderTargets::Get(RHICmdList);
+				//create the shadow depth texture and/or surface
+				const FIntPoint ShadowBufferResolution = SceneContext.GetShadowDepthTextureResolution();
+				const int32 MaxWide = GMaxShadowDepthBufferSizeX / ShadowBufferResolution.X;
+				const int32 MaxHigh = GMaxShadowDepthBufferSizeY / ShadowBufferResolution.Y;
 
-					const FIntPoint AtlasShadowBufferResolution(ShadowBufferResolution.X * NumWide, ShadowBufferResolution.Y * NumHigh);
-					SceneContext.AllocateForwardShadingShadowDepthTarget(AtlasShadowBufferResolution);
+				const int32 NumWide = FMath::Min(NumWholeSceneShadows, MaxWide);
+				const int32 NumHigh = FMath::Min(((NumWholeSceneShadows - 1) / MaxWide) + 1, MaxHigh);
 
+				const FIntPoint AtlasShadowBufferResolution(ShadowBufferResolution.X * NumWide, ShadowBufferResolution.Y * NumHigh);
+				SceneContext.AllocateForwardShadingShadowDepthTarget(AtlasShadowBufferResolution);
 
-					// Allocate atlas shadow texture space to the shadows.				
-					FTextureLayout ShadowLayout(1, 1, AtlasShadowBufferResolution.X, AtlasShadowBufferResolution.Y, false, false);
-					for (int32 ShadowIndex = 0; ShadowIndex < VisibleLightInfo.AllProjectedShadows.Num(); ShadowIndex++)
+				// Allocate atlas shadow texture space to the shadows.
+				FTextureLayout ShadowLayout(1, 1, AtlasShadowBufferResolution.X, AtlasShadowBufferResolution.Y, false, false);
+
+				for (int32 ShadowIndex = 0; ShadowIndex < NumWholeSceneShadows; ShadowIndex++)
+				{
+					FProjectedShadowInfo* ProjectedShadowInfo = VisibleLightInfo.AllProjectedShadows[ShadowIndex];
+					if (ShadowLayout.AddElement(
+						ProjectedShadowInfo->X,
+						ProjectedShadowInfo->Y,
+						ProjectedShadowInfo->ResolutionX + SHADOW_BORDER * 2,
+						ProjectedShadowInfo->ResolutionY + SHADOW_BORDER * 2)
+						)
 					{
-						FProjectedShadowInfo* ProjectedShadowInfo = VisibleLightInfo.AllProjectedShadows[ShadowIndex];
-						if (!ProjectedShadowInfo->bRendered)
-						{
-							if (ShadowLayout.AddElement(
-								ProjectedShadowInfo->X,
-								ProjectedShadowInfo->Y,
-								ProjectedShadowInfo->ResolutionX + SHADOW_BORDER * 2,
-								ProjectedShadowInfo->ResolutionY + SHADOW_BORDER * 2)
-								)
-							{
-								ProjectedShadowInfo->bAllocated = true;
-							}
-						}
+						ProjectedShadowInfo->bAllocated = true;
 					}
 				}
 			}
