Fix modulated shadow rendering depth/stencil setup. Fixes rendering on DX11 ES2 preview.

diff --git a/Engine/Source/Runtime/Renderer/Private/ShadowRendering.cpp b/Engine/Source/Runtime/Renderer/Private/ShadowRendering.cpp
--- a/Engine/Source/Runtime/Renderer/Private/ShadowRendering.cpp
+++ b/Engine/Source/Runtime/Renderer/Private/ShadowRendering.cpp
@@ -2904,7 +2904,7 @@ void FSceneRenderer::RenderProjections(
 	
 	if (bForwardShading)
 	{
-		SceneContext.BeginRenderingSceneColor(RHICmdList);
+		SceneContext.BeginRenderingSceneColor(RHICmdList, ESimpleRenderTargetMode::EUninitializedColorExistingDepth, FExclusiveDepthStencil::DepthRead_StencilWrite);
 	}
 	else
 	{
