Guarantee all endofframe updates to recreate render proxies are processed before rendering.

diff --git a/Engine/Source/Editor/Kismet/Private/BlueprintEditor.cpp b/Engine/Source/Editor/Kismet/Private/BlueprintEditor.cpp
--- a/Engine/Source/Editor/Kismet/Private/BlueprintEditor.cpp
+++ b/Engine/Source/Editor/Kismet/Private/BlueprintEditor.cpp
@@ -7366,8 +7366,7 @@ void FBlueprintEditor::UpdatePreviewActor(UBlueprint* InBlueprint, bool bInForce
 			if ( PreviewActor->bHidden )
 			{
 				PreviewActor->bHidden = false;
-				PreviewActor->MarkComponentsRenderStateDirty();
-				PreviewScene.GetWorld()->SendAllEndOfFrameUpdates();
+				PreviewActor->MarkComponentsRenderStateDirty();				
 			}
 
 			// Prevent any audio from playing as a result of spawning
diff --git a/Engine/Source/Editor/Kismet/Private/SCSEditorViewportClient.cpp b/Engine/Source/Editor/Kismet/Private/SCSEditorViewportClient.cpp
--- a/Engine/Source/Editor/Kismet/Private/SCSEditorViewportClient.cpp
+++ b/Engine/Source/Editor/Kismet/Private/SCSEditorViewportClient.cpp
@@ -772,7 +772,6 @@ void FSCSEditorViewportClient::ToggleShowFloor()
 	
 	EditorFloorComp->SetVisibility(bShowFloor);
 	EditorFloorComp->SetCollisionEnabled(bShowFloor? ECollisionEnabled::QueryAndPhysics : ECollisionEnabled::NoCollision);
-	GetWorld()->SendAllEndOfFrameUpdates();
 
 	Settings->bSCSEditorShowFloor = bShowFloor;
 	Settings->PostEditChange();
diff --git a/Engine/Source/Editor/UnrealEd/Private/EditorBuildUtils.cpp b/Engine/Source/Editor/UnrealEd/Private/EditorBuildUtils.cpp
--- a/Engine/Source/Editor/UnrealEd/Private/EditorBuildUtils.cpp
+++ b/Engine/Source/Editor/UnrealEd/Private/EditorBuildUtils.cpp
@@ -387,14 +387,7 @@ bool FEditorBuildUtils::EditorBuild( UWorld* InWorld, EBuildOptions::Type Id, co
 	if (GUseThreadedRendering)
 	{
 		StartRenderingThread();
-	}
-
-	if ( bDoBuild && InWorld->Scene )
-	{
-		// Invalidating lighting marked various components as needing a re-register
-		// Propagate the re-registers before rendering the scene to get the latest state
-		InWorld->SendAllEndOfFrameUpdates();
-	}
+	}	
 
 	if ( bDoBuild )
 	{
diff --git a/Engine/Source/Editor/UnrealEd/Private/EditorEngine.cpp b/Engine/Source/Editor/UnrealEd/Private/EditorEngine.cpp
--- a/Engine/Source/Editor/UnrealEd/Private/EditorEngine.cpp
+++ b/Engine/Source/Editor/UnrealEd/Private/EditorEngine.cpp
@@ -1436,8 +1436,7 @@ void UEditorEngine::Tick( float DeltaSeconds, bool bIdleMode )
 
 
 	// Commit changes to the BSP model.
-	EditorContext.World()->CommitModelSurfaces();
-	EditorContext.World()->SendAllEndOfFrameUpdates();
+	EditorContext.World()->CommitModelSurfaces();	
 
 	bool bUpdateLinkedOrthoViewports = false;
 	/////////////////////////////
diff --git a/Engine/Source/Editor/UnrealEd/Private/ThumbnailHelpers.cpp b/Engine/Source/Editor/UnrealEd/Private/ThumbnailHelpers.cpp
--- a/Engine/Source/Editor/UnrealEd/Private/ThumbnailHelpers.cpp
+++ b/Engine/Source/Editor/UnrealEd/Private/ThumbnailHelpers.cpp
@@ -982,9 +982,7 @@ void FClassActorThumbnailScene::SetObject(UObject* Obj)
 		}
 
 		VisualizableComponents.Empty();
-	}
-
-	GetWorld()->SendAllEndOfFrameUpdates();
+	}	
 }
 
 bool FClassActorThumbnailScene::IsValidComponentForVisualization(UActorComponent* Component) const
diff --git a/Engine/Source/Runtime/Engine/Private/Components/StaticMeshComponent.cpp b/Engine/Source/Runtime/Engine/Private/Components/StaticMeshComponent.cpp
--- a/Engine/Source/Runtime/Engine/Private/Components/StaticMeshComponent.cpp
+++ b/Engine/Source/Runtime/Engine/Private/Components/StaticMeshComponent.cpp
@@ -1644,10 +1644,8 @@ void UStaticMeshComponent::ApplyComponentInstanceData(FStaticMeshComponentInstan
 		}
 	}
 
-	{
-		FComponentReregisterContext ReregisterStaticMesh(this);
-		StaticMeshInstanceData->ApplyVertexColorData(this);
-	}
+	StaticMeshInstanceData->ApplyVertexColorData(this);
+	MarkRenderStateDirty();
 }
 
 #include "AI/Navigation/RecastHelpers.h"
diff --git a/Engine/Source/Runtime/Engine/Private/LevelTick.cpp b/Engine/Source/Runtime/Engine/Private/LevelTick.cpp
--- a/Engine/Source/Runtime/Engine/Private/LevelTick.cpp
+++ b/Engine/Source/Runtime/Engine/Private/LevelTick.cpp
@@ -1301,8 +1301,7 @@ void UWorld::Tick( ELevelTick TickType, float DeltaSeconds )
 		SCOPE_CYCLE_COUNTER(STAT_NetBroadcastTickTime);
 		BroadcastTickFlush(RealDeltaSeconds); // note: undilated time is being used here
 	}
-
-	SendAllEndOfFrameUpdates();
+	
      // PostTick all net drivers
 	{
 		SCOPE_CYCLE_COUNTER(STAT_NetBroadcastPostTickTime);
diff --git a/Engine/Source/Runtime/Renderer/Private/SceneRendering.cpp b/Engine/Source/Runtime/Renderer/Private/SceneRendering.cpp
--- a/Engine/Source/Runtime/Renderer/Private/SceneRendering.cpp
+++ b/Engine/Source/Runtime/Renderer/Private/SceneRendering.cpp
@@ -1550,6 +1550,20 @@ void FRendererModule::CreateAndInitSingleView(FRHICommandListImmediate& RHICmdLi
 
 void FRendererModule::BeginRenderingViewFamily(FCanvas* Canvas,FSceneViewFamily* ViewFamily)
 {
+	UWorld* World = nullptr; 
+	check(ViewFamily->Scene);
+
+	FScene* const Scene = ViewFamily->Scene->GetRenderScene();
+	if (Scene)
+	{
+		World = Scene->GetWorld();
+		if (World)
+		{
+			//guarantee that all render proxies are up to date before kicking off a BeginRenderViewFamily.
+			World->SendAllEndOfFrameUpdates();
+		}
+	}
+
 	// Flush the canvas first.
 	Canvas->Flush_GameThread();
 
@@ -1564,12 +1578,9 @@ void FRendererModule::BeginRenderingViewFamily(FCanvas* Canvas,FSceneViewFamily*
 	{
 		ViewFamily->ViewExtensions[ViewExt]->BeginRenderViewFamily(*ViewFamily);
 	}
-
-	check(ViewFamily->Scene);
-	FScene* const Scene = ViewFamily->Scene->GetRenderScene();
+	
 	if (Scene)
-	{
-		UWorld* const World = Scene->GetWorld();
+	{		
 		// Set the world's "needs full lighting rebuild" flag if the scene has any uncached static lighting interactions.
 		if(World)
 		{
