UE4 - PR #820: [Engine] Optimize FDeferredShadingSceneRenderer::InitViews() by making front-to-back sorting of base pass draw lists asynchronous. (Contributed by [REDACTED_PERSON]) #git author [REDACTED_PERSON] [REDACTED_EMAIL]

diff --git a/Engine/Source/Runtime/Renderer/Private/DeferredShadingRenderer.cpp b/Engine/Source/Runtime/Renderer/Private/DeferredShadingRenderer.cpp
--- a/Engine/Source/Runtime/Renderer/Private/DeferredShadingRenderer.cpp
+++ b/Engine/Source/Runtime/Renderer/Private/DeferredShadingRenderer.cpp
@@ -271,6 +271,66 @@ void FDeferredShadingSceneRenderer::RenderBasePassStaticDataDefaultParallel(FPar
 	}
 }
 
+template<typename StaticMeshDrawList>
+class FSortFrontToBackTask
+{
+private:
+	StaticMeshDrawList * const StaticMeshDrawListToSort;
+	const FVector ViewPosition;
+
+public:
+	FSortFrontToBackTask(StaticMeshDrawList * const InStaticMeshDrawListToSort, const FVector InViewPosition)
+		: StaticMeshDrawListToSort(InStaticMeshDrawListToSort)
+		, ViewPosition(InViewPosition)
+	{
+
+	}
+
+	FORCEINLINE TStatId GetStatId() const
+	{
+		RETURN_QUICK_DECLARE_CYCLE_STAT(FSortFrontToBackTask, STATGROUP_TaskGraphTasks);
+	}
+
+	ENamedThreads::Type GetDesiredThread()
+	{
+		return ENamedThreads::AnyThread;
+	}
+
+	static ESubsequentsMode::Type GetSubsequentsMode() { return ESubsequentsMode::TrackSubsequents; }
+
+	void DoTask(ENamedThreads::Type CurrentThread, const FGraphEventRef& MyCompletionGraphEvent)
+	{
+		StaticMeshDrawListToSort->SortFrontToBack(ViewPosition);
+	}
+};
+
+void FDeferredShadingSceneRenderer::AsyncSortBasePassStaticData(const FVector InViewPosition, FGraphEventArray &OutSortEvents)
+{
+	// If we're not using a depth only pass, sort the static draw list buckets roughly front to back, to maximize HiZ culling
+	// Note that this is only a very rough sort, since it does not interfere with state sorting, and each list is sorted separately
+	if (EarlyZPassMode != DDM_None)
+		return;
+	QUICK_SCOPE_CYCLE_COUNTER(STAT_FDeferredShadingSceneRenderer_AsyncSortBasePassStaticData);
+
+	for (int32 DrawType = 0; DrawType < FScene::EBasePass_MAX; ++DrawType)
+	{
+		OutSortEvents.Add(TGraphTask<FSortFrontToBackTask<TStaticMeshDrawList<TBasePassDrawingPolicy<FNoLightMapPolicy> > > >::CreateTask(
+			nullptr, ENamedThreads::RenderThread).ConstructAndDispatchWhenReady(&(Scene->BasePassNoLightMapDrawList[DrawType]), InViewPosition));
+		OutSortEvents.Add(TGraphTask<FSortFrontToBackTask<TStaticMeshDrawList<TBasePassDrawingPolicy<FSimpleDynamicLightingPolicy> > > >::CreateTask(
+			nullptr, ENamedThreads::RenderThread).ConstructAndDispatchWhenReady(&(Scene->BasePassSimpleDynamicLightingDrawList[DrawType]), InViewPosition));
+		OutSortEvents.Add(TGraphTask<FSortFrontToBackTask<TStaticMeshDrawList<TBasePassDrawingPolicy<FCachedVolumeIndirectLightingPolicy> > > >::CreateTask(
+			nullptr, ENamedThreads::RenderThread).ConstructAndDispatchWhenReady(&(Scene->BasePassCachedVolumeIndirectLightingDrawList[DrawType]), InViewPosition));
+		OutSortEvents.Add(TGraphTask<FSortFrontToBackTask<TStaticMeshDrawList<TBasePassDrawingPolicy<FCachedPointIndirectLightingPolicy> > > >::CreateTask(
+			nullptr, ENamedThreads::RenderThread).ConstructAndDispatchWhenReady(&(Scene->BasePassCachedPointIndirectLightingDrawList[DrawType]), InViewPosition));
+		OutSortEvents.Add(TGraphTask<FSortFrontToBackTask<TStaticMeshDrawList<TBasePassDrawingPolicy<TLightMapPolicy<HQ_LIGHTMAP> > > > >::CreateTask(
+			nullptr, ENamedThreads::RenderThread).ConstructAndDispatchWhenReady(&(Scene->BasePassHighQualityLightMapDrawList[DrawType]), InViewPosition));
+		OutSortEvents.Add(TGraphTask<FSortFrontToBackTask<TStaticMeshDrawList<TBasePassDrawingPolicy<TDistanceFieldShadowsAndLightMapPolicy<HQ_LIGHTMAP> > > > >::CreateTask(
+			nullptr, ENamedThreads::RenderThread).ConstructAndDispatchWhenReady(&(Scene->BasePassDistanceFieldShadowMapLightMapDrawList[DrawType]), InViewPosition));
+		OutSortEvents.Add(TGraphTask<FSortFrontToBackTask<TStaticMeshDrawList<TBasePassDrawingPolicy<TLightMapPolicy<LQ_LIGHTMAP> > > > >::CreateTask(
+			nullptr, ENamedThreads::RenderThread).ConstructAndDispatchWhenReady(&(Scene->BasePassLowQualityLightMapDrawList[DrawType]), InViewPosition));
+	}
+}
+
 void FDeferredShadingSceneRenderer::SortBasePassStaticData(FVector ViewPosition)
 {
 	// If we're not using a depth only pass, sort the static draw list buckets roughly front to back, to maximize HiZ culling
diff --git a/Engine/Source/Runtime/Renderer/Private/DeferredShadingRenderer.h b/Engine/Source/Runtime/Renderer/Private/DeferredShadingRenderer.h
--- a/Engine/Source/Runtime/Renderer/Private/DeferredShadingRenderer.h
+++ b/Engine/Source/Runtime/Renderer/Private/DeferredShadingRenderer.h
@@ -77,6 +77,9 @@ public:
 	void RenderBasePassStaticDataMaskedParallel(FParallelCommandListSet& ParallelCommandListSet);
 	void RenderBasePassStaticDataDefaultParallel(FParallelCommandListSet& ParallelCommandListSet);
 
+	/** Asynchronously sorts base pass draw lists front to back for improved GPU culling. */
+	void AsyncSortBasePassStaticData(const FVector ViewPosition, FGraphEventArray &SortEvents);
+
 	/** Sorts base pass draw lists front to back for improved GPU culling. */
 	void SortBasePassStaticData(FVector ViewPosition);
 
diff --git a/Engine/Source/Runtime/Renderer/Private/SceneVisibility.cpp b/Engine/Source/Runtime/Renderer/Private/SceneVisibility.cpp
--- a/Engine/Source/Runtime/Renderer/Private/SceneVisibility.cpp
+++ b/Engine/Source/Runtime/Renderer/Private/SceneVisibility.cpp
@@ -2081,7 +2081,15 @@ void FDeferredShadingSceneRenderer::InitViews(FRHICommandListImmediate& RHICmdLi
 		AverageViewPosition += View.ViewMatrices.ViewOrigin / Views.Num();
 	}
 
-	SortBasePassStaticData(AverageViewPosition);
+	FGraphEventArray SortEvents;
+	if (FApp::ShouldUseThreadingForPerformance() && CVarParallelInitViews.GetValueOnRenderThread() > 0)
+	{
+		AsyncSortBasePassStaticData(AverageViewPosition, SortEvents);
+	}
+	else
+	{
+		SortBasePassStaticData(AverageViewPosition);
+	}
 
 	bool bDynamicShadows = ViewFamily.EngineShowFlags.DynamicShadows && GetShadowQuality() > 0;
 
@@ -2096,7 +2104,13 @@ void FDeferredShadingSceneRenderer::InitViews(FRHICommandListImmediate& RHICmdLi
 	{
 		// Initialize the view's RHI resources.
 		Views[ViewIndex].InitRHIResources(nullptr);
-	}	
+	}
+
+	if (SortEvents.Num())
+	{
+		QUICK_SCOPE_CYCLE_COUNTER(STAT_FDeferredShadingSceneRenderer_AsyncSortBasePassStaticData_Wait);
+		FTaskGraphInterface::Get().WaitUntilTasksComplete(SortEvents, ENamedThreads::RenderThread);
+	}
 
 	OnStartFrame();
 }
diff --git a/Engine/Source/Runtime/Renderer/Private/StaticMeshDrawList.h b/Engine/Source/Runtime/Renderer/Private/StaticMeshDrawList.h
--- a/Engine/Source/Runtime/Renderer/Private/StaticMeshDrawList.h
+++ b/Engine/Source/Runtime/Renderer/Private/StaticMeshDrawList.h
@@ -353,8 +353,9 @@ public:
 	// FRenderResource interface.
 	virtual void ReleaseRHI();
 
+	typedef TSet<FDrawingPolicyLink, FDrawingPolicyKeyFuncs> TDrawingPolicySet;
 	/** Sorts OrderedDrawingPolicies front to back.  Relies on static variables SortDrawingPolicySet and SortViewPosition being set. */
-	static int32 Compare(FSetElementId A, FSetElementId B);
+	static int32 Compare(FSetElementId A, FSetElementId B, const TDrawingPolicySet * const InSortDrawingPolicySet, const FVector InSortViewPosition);
 
 	/** Computes statistics for this draw list. */
 	FDrawListStats GetStats() const;
@@ -366,23 +367,27 @@ private:
 	typedef TSet<FDrawingPolicyLink,FDrawingPolicyKeyFuncs> TDrawingPolicySet;
 	/** All drawing policy element sets in the draw list, hashed by drawing policy. */
 	TDrawingPolicySet DrawingPolicySet;
-
-	/** 
-	 * Static variables for getting data into the Compare function.
-	 * Ideally Sort would accept a non-static member function which would avoid having to go through globals.
-	 */ 
-	static TDrawingPolicySet* SortDrawingPolicySet;
-	static FVector SortViewPosition;
 };
 
-/** Helper stuct for sorting */
+/** Helper struct for sorting */
 template<typename DrawingPolicyType>
 struct TCompareStaticMeshDrawList
 {
-	FORCEINLINE bool operator()( const FSetElementId& A, const FSetElementId& B ) const
+private:
+	const typename TStaticMeshDrawList<DrawingPolicyType>::TDrawingPolicySet * const SortDrawingPolicySet;
+	const FVector SortViewPosition;
+
+public:
+	TCompareStaticMeshDrawList(const typename TStaticMeshDrawList<DrawingPolicyType>::TDrawingPolicySet * const InSortDrawingPolicySet, const FVector InSortViewPosition)
+		: SortDrawingPolicySet(InSortDrawingPolicySet)
+		, SortViewPosition(InSortViewPosition)
+	{
+	}
+
+	FORCEINLINE bool operator()(const FSetElementId& A, const FSetElementId& B) const
 	{
 		// Use static Compare from TStaticMeshDrawList
-		return TStaticMeshDrawList<DrawingPolicyType>::Compare( A, B ) < 0;
+		return TStaticMeshDrawList<DrawingPolicyType>::Compare(A, B, SortDrawingPolicySet, SortViewPosition) < 0;
 	}
 };
 
diff --git a/Engine/Source/Runtime/Renderer/Private/StaticMeshDrawList.inl b/Engine/Source/Runtime/Renderer/Private/StaticMeshDrawList.inl
--- a/Engine/Source/Runtime/Renderer/Private/StaticMeshDrawList.inl
+++ b/Engine/Source/Runtime/Renderer/Private/StaticMeshDrawList.inl
@@ -447,7 +447,7 @@ void TStaticMeshDrawList<DrawingPolicyType>::DrawVisibleParallel(
 			int32 DrawsPerCmdList = FMath::DivideAndRoundUp(Total, EffectiveThreads);
 			int32 DrawsPerCmdListMergeLimit = FMath::DivideAndRoundUp(DrawsPerCmdList, 3); // if the last list would be small, we just merge it into the previous one
 
-			int32 Start = 0;
+	int32 Start = 0;
 
 			int32 PreviousBatchStart = -1;
 			int32 PreviousBatchEnd = -2;
@@ -524,26 +524,26 @@ void TStaticMeshDrawList<DrawingPolicyType>::DrawVisibleParallel(
 			}
 		}
 		else
-		{
-			int32 NumPer = OrderedDrawingPolicies.Num() / EffectiveThreads;
-			int32 Extra = OrderedDrawingPolicies.Num() - NumPer * EffectiveThreads;
+	{
+		int32 NumPer = OrderedDrawingPolicies.Num() / EffectiveThreads;
+		int32 Extra = OrderedDrawingPolicies.Num() - NumPer * EffectiveThreads;
 			int32 Start = 0;
-			for (int32 ThreadIndex = 0; ThreadIndex < EffectiveThreads; ThreadIndex++)
-			{
-				int32 Last = Start + (NumPer - 1) + (ThreadIndex < Extra);
-				check(Last >= Start);
+		for (int32 ThreadIndex = 0; ThreadIndex < EffectiveThreads; ThreadIndex++)
+		{
+			int32 Last = Start + (NumPer - 1) + (ThreadIndex < Extra);
+			check(Last >= Start);
 
-				{
-					FRHICommandList* CmdList = ParallelCommandListSet.NewParallelCommandList();
+			{
+				FRHICommandList* CmdList = ParallelCommandListSet.NewParallelCommandList();
 
 					FGraphEventRef AnyThreadCompletionEvent = TGraphTask<FDrawVisibleAnyThreadTask<DrawingPolicyType> >::CreateTask(ParallelCommandListSet.GetPrereqs(), ENamedThreads::RenderThread)
 						.ConstructAndDispatchWhenReady(*this, *CmdList, ParallelCommandListSet.View, PolicyContext, StaticMeshVisibilityMap, BatchVisibilityArray, Start, Last, nullptr);
 
-					ParallelCommandListSet.AddParallelCommandList(CmdList, AnyThreadCompletionEvent);
-				}
-
-				Start = Last + 1;
+				ParallelCommandListSet.AddParallelCommandList(CmdList, AnyThreadCompletionEvent);
 			}
+
+			Start = Last + 1;
+		}
 			check(Start == OrderedDrawingPolicies.Num());
 		}
 	}
@@ -615,16 +615,10 @@ int32 TStaticMeshDrawList<DrawingPolicyType>::DrawVisibleFrontToBack(
 }
 
 template<typename DrawingPolicyType>
-FVector TStaticMeshDrawList<DrawingPolicyType>::SortViewPosition = FVector(0);
-
-template<typename DrawingPolicyType>
-typename TStaticMeshDrawList<DrawingPolicyType>::TDrawingPolicySet* TStaticMeshDrawList<DrawingPolicyType>::SortDrawingPolicySet;
-
-template<typename DrawingPolicyType>
-int32 TStaticMeshDrawList<DrawingPolicyType>::Compare(FSetElementId A, FSetElementId B)
+int32 TStaticMeshDrawList<DrawingPolicyType>::Compare(FSetElementId A, FSetElementId B, const TDrawingPolicySet * const InSortDrawingPolicySet, const FVector InSortViewPosition)
 {
-	const FSphere& BoundsA = (*SortDrawingPolicySet)[A].CachedBoundingSphere;
-	const FSphere& BoundsB = (*SortDrawingPolicySet)[B].CachedBoundingSphere;
+	const FSphere& BoundsA = (*InSortDrawingPolicySet)[A].CachedBoundingSphere;
+	const FSphere& BoundsB = (*InSortDrawingPolicySet)[B].CachedBoundingSphere;
 
 	// Assume state buckets with large bounds are background geometry
 	if (BoundsA.W >= HALF_WORLD_MAX / 2 && BoundsB.W < HALF_WORLD_MAX / 2)
@@ -637,8 +631,8 @@ int32 TStaticMeshDrawList<DrawingPolicyType>::Compare(FSetElementId A, FSetEleme
 	}
 	else
 	{
-		const float DistanceASquared = (BoundsA.Center - SortViewPosition).SizeSquared();
-		const float DistanceBSquared = (BoundsB.Center - SortViewPosition).SizeSquared();
+		const float DistanceASquared = (BoundsA.Center - InSortViewPosition).SizeSquared();
+		const float DistanceBSquared = (BoundsB.Center - InSortViewPosition).SizeSquared();
 		// Sort front to back
 		return DistanceASquared > DistanceBSquared ? 1 : -1;
 	}
@@ -653,22 +647,22 @@ void TStaticMeshDrawList<DrawingPolicyType>::SortFrontToBack(FVector ViewPositio
 		FBoxSphereBounds AccumulatedBounds(ForceInit);
 
 		FDrawingPolicyLink& DrawingPolicyLink = *DrawingPolicyIt;
-		if (DrawingPolicyLink.Elements.Num())
+
+		const int32 NumElements = DrawingPolicyLink.Elements.Num();
+		if (NumElements)
 		{
 			AccumulatedBounds = DrawingPolicyLink.Elements[0].Bounds;
-			for (int32 ElementIndex = 1; ElementIndex < DrawingPolicyLink.Elements.Num(); ElementIndex++)
+
+		    for (int32 ElementIndex = 1; ElementIndex < NumElements; ElementIndex++)
 		    {
-		        FElement& Element = DrawingPolicyLink.Elements[ElementIndex];
-		        AccumulatedBounds = AccumulatedBounds + Element.Bounds;
+			    AccumulatedBounds = AccumulatedBounds + DrawingPolicyLink.Elements[ElementIndex].Bounds;
 		    }
 		}
+
 		DrawingPolicyIt->CachedBoundingSphere = AccumulatedBounds.GetSphere();
 	}
 
-	SortViewPosition = ViewPosition;
-	SortDrawingPolicySet = &DrawingPolicySet;
-
-	OrderedDrawingPolicies.Sort( TCompareStaticMeshDrawList<DrawingPolicyType>() );
+	OrderedDrawingPolicies.Sort(TCompareStaticMeshDrawList<DrawingPolicyType>(&DrawingPolicySet, ViewPosition));
 }
 
 template<typename DrawingPolicyType>
