UE4 - fix CIS

diff --git a/Engine/Source/Runtime/Renderer/Private/StaticMeshDrawList.h b/Engine/Source/Runtime/Renderer/Private/StaticMeshDrawList.h
--- a/Engine/Source/Runtime/Renderer/Private/StaticMeshDrawList.h
+++ b/Engine/Source/Runtime/Renderer/Private/StaticMeshDrawList.h
@@ -364,7 +364,6 @@ private:
 	/** All drawing policies in the draw list, in rendering order. */
 	TArray<FSetElementId> OrderedDrawingPolicies;
 	
-	typedef TSet<FDrawingPolicyLink,FDrawingPolicyKeyFuncs> TDrawingPolicySet;
 	/** All drawing policy element sets in the draw list, hashed by drawing policy. */
 	TDrawingPolicySet DrawingPolicySet;
 };
