Combined redundant operation in GatherShadowsForPrimitiveInner

diff --git a/Engine/Source/Runtime/Renderer/Private/ShadowSetup.cpp b/Engine/Source/Runtime/Renderer/Private/ShadowSetup.cpp
--- a/Engine/Source/Runtime/Renderer/Private/ShadowSetup.cpp
+++ b/Engine/Source/Runtime/Renderer/Private/ShadowSetup.cpp
@@ -2031,19 +2031,20 @@ inline void FSceneRenderer::GatherShadowsForPrimitiveInner(
 				const float ProjectedDistanceFromShadowOriginAlongLightDir = PrimitiveToShadowCenter | LightDirection;
 				// Calculate the primitive's squared distance to the cylinder's axis
 				const float PrimitiveDistanceFromCylinderAxisSq = (-LightDirection * ProjectedDistanceFromShadowOriginAlongLightDir + PrimitiveToShadowCenter).SizeSquared();
+				const float CombinedRadiusSq = FMath::Square(ProjectedShadowInfo->ShadowBounds.W + PrimitiveBounds.SphereRadius);
 
 				// Include all primitives for movable lights, but only statically shadowed primitives from a light with static shadowing,
 				// Since lights with static shadowing still create per-object shadows for primitives without static shadowing.
 				if( (!LightProxy->HasStaticLighting() || !ProjectedShadowInfo->GetLightSceneInfo().IsPrecomputedLightingValid())
 					// Check if this primitive is in the shadow's cylinder
-					&& PrimitiveDistanceFromCylinderAxisSq < FMath::Square(ProjectedShadowInfo->ShadowBounds.W + PrimitiveBounds.SphereRadius)
+					&& PrimitiveDistanceFromCylinderAxisSq < CombinedRadiusSq
 					// Check if the primitive is closer than the cylinder cap toward the light
 					// next line is commented as it breaks large world shadows, if this was meant to be an optimization we should think about a better solution
 					//// && ProjectedDistanceFromShadowOriginAlongLightDir - PrimitiveBounds.SphereRadius < -ProjectedShadowInfo->MinPreSubjectZ
 					// If the primitive is further along the cone axis than the shadow bounds origin, 
 					// Check if the primitive is inside the spherical cap of the cascade's bounds
 					&& !(ProjectedDistanceFromShadowOriginAlongLightDir < 0 
-						&& PrimitiveToShadowCenter.SizeSquared() > FMath::Square(ProjectedShadowInfo->ShadowBounds.W + PrimitiveBounds.SphereRadius)))
+						&& PrimitiveToShadowCenter.SizeSquared() > CombinedRadiusSq))
 				{
 					const bool bInFrustum = ProjectedShadowInfo->CascadeSettings.ShadowBoundsAccurate.IntersectBox( PrimitiveBounds.Origin, PrimitiveBounds.BoxExtent );
 
