Fewer state switches for forward shading decals

diff --git a/Engine/Source/Runtime/Renderer/Private/DecalRenderingForwardShading.cpp b/Engine/Source/Runtime/Renderer/Private/DecalRenderingForwardShading.cpp
--- a/Engine/Source/Runtime/Renderer/Private/DecalRenderingForwardShading.cpp
+++ b/Engine/Source/Runtime/Renderer/Private/DecalRenderingForwardShading.cpp
@@ -44,7 +44,10 @@ void FForwardShadingSceneRenderer::RenderDecals(FRHICommandListImmediate& RHICmd
 		
 			RHICmdList.SetViewport(View.ViewRect.Min.X, View.ViewRect.Min.Y, 0, View.ViewRect.Max.X, View.ViewRect.Max.Y, 1);
 			RHICmdList.SetStreamSource(0, GetUnitCubeVertexBuffer(), sizeof(FVector4), 0);
-	
+
+			TOptional<EDecalBlendMode> LastDecalBlendMode;
+			TOptional<bool> LastDecalDepthState;
+				
 			for (int32 DecalIndex = 0, DecalCount = SortedDecals.Num(); DecalIndex < DecalCount; DecalIndex++)
 			{
 				const FTransientDecalRenderData& DecalData = SortedDecals[DecalIndex];
@@ -54,33 +57,42 @@ void FForwardShadingSceneRenderer::RenderDecals(FRHICommandListImmediate& RHICmd
 						
 				const float ConservativeRadius = DecalData.ConservativeRadius;
 				const bool bInsideDecal = ((FVector)View.ViewMatrices.ViewOrigin - ComponentToWorldMatrix.GetOrigin()).SizeSquared() < FMath::Square(ConservativeRadius * 1.05f + View.NearClippingDistance * 2.0f);
-				if (bInsideDecal)
+
+				if (!LastDecalDepthState.IsSet() || LastDecalDepthState.GetValue() != bInsideDecal)
 				{
-					RHICmdList.SetRasterizerState(View.bReverseCulling ? TStaticRasterizerState<FM_Solid, CM_CCW>::GetRHI() : TStaticRasterizerState<FM_Solid, CM_CW>::GetRHI());
-					RHICmdList.SetDepthStencilState(TStaticDepthStencilState<false,CF_Always, true, CF_Equal, SO_Keep, SO_Keep, SO_Keep>::GetRHI(), 0);
+					LastDecalDepthState = bInsideDecal;
+					if (bInsideDecal)
+					{
+						RHICmdList.SetRasterizerState(View.bReverseCulling ? TStaticRasterizerState<FM_Solid, CM_CCW>::GetRHI() : TStaticRasterizerState<FM_Solid, CM_CW>::GetRHI());
+						RHICmdList.SetDepthStencilState(TStaticDepthStencilState<false,CF_Always, true, CF_Equal, SO_Keep, SO_Keep, SO_Keep>::GetRHI(), 0);
+					}
+					else
+					{
+						RHICmdList.SetRasterizerState(View.bReverseCulling ? TStaticRasterizerState<FM_Solid, CM_CW>::GetRHI() : TStaticRasterizerState<FM_Solid, CM_CCW>::GetRHI());
+						RHICmdList.SetDepthStencilState(TStaticDepthStencilState<false,CF_DepthNearOrEqual, true, CF_Equal, SO_Keep, SO_Keep, SO_Keep>::GetRHI(), 0);
+					}
 				}
-				else
+
+				if (!LastDecalBlendMode.IsSet() || LastDecalBlendMode.GetValue() != DecalData.DecalBlendMode)
 				{
-					RHICmdList.SetRasterizerState(View.bReverseCulling ? TStaticRasterizerState<FM_Solid, CM_CW>::GetRHI() : TStaticRasterizerState<FM_Solid, CM_CCW>::GetRHI());
-					RHICmdList.SetDepthStencilState(TStaticDepthStencilState<false,CF_DepthNearOrEqual, true, CF_Equal, SO_Keep, SO_Keep, SO_Keep>::GetRHI(), 0);
+					LastDecalBlendMode = DecalData.DecalBlendMode;
+					switch(DecalData.DecalBlendMode)
+					{
+					case DBM_Translucent:
+						RHICmdList.SetBlendState(TStaticBlendState<CW_RGB, BO_Add, BF_SourceAlpha, BF_InverseSourceAlpha>::GetRHI());
+						break;
+					case DBM_Stain:
+						// Modulate
+						RHICmdList.SetBlendState(TStaticBlendState<CW_RGB, BO_Add, BF_DestColor, BF_InverseSourceAlpha>::GetRHI());
+						break;
+					case DBM_Emissive:
+						// Additive
+						RHICmdList.SetBlendState(TStaticBlendState<CW_RGB, BO_Add, BF_SourceAlpha, BF_One>::GetRHI());
+						break;
+					default:
+						check(0);
+					};
 				}
-			
-				switch(DecalData.DecalBlendMode)
-				{
-				case DBM_Translucent:
-					RHICmdList.SetBlendState(TStaticBlendState<CW_RGB, BO_Add, BF_SourceAlpha, BF_InverseSourceAlpha>::GetRHI());
-					break;
-				case DBM_Stain:
-					// Modulate
-					RHICmdList.SetBlendState(TStaticBlendState<CW_RGB, BO_Add, BF_DestColor, BF_InverseSourceAlpha>::GetRHI());
-					break;
-				case DBM_Emissive:
-					// Additive
-					RHICmdList.SetBlendState(TStaticBlendState<CW_RGB, BO_Add, BF_SourceAlpha, BF_One>::GetRHI());
-					break;
-				default:
-					check(0);
-				};
 
 				// Set shader params
 				FDecalRendering::SetShader(RHICmdList, View, bShaderComplexity, DecalData, FrustumComponentToClip);
