UE4 - fix rare bug with translucent lighting. FViewInfo snapshots no longer use the scene allocator.

diff --git a/Engine/Source/Runtime/Renderer/Private/SceneRendering.cpp b/Engine/Source/Runtime/Renderer/Private/SceneRendering.cpp
--- a/Engine/Source/Runtime/Renderer/Private/SceneRendering.cpp
+++ b/Engine/Source/Runtime/Renderer/Private/SceneRendering.cpp
@@ -974,14 +974,25 @@ void FViewInfo::InitRHIResources(const TArray<FProjectedShadowInfo*, SceneRender
 	}
 }
 
+// These are not real view infos, just dumb memory blocks
 static TArray<FViewInfo*> ViewInfoSnapshots;
+// these are never freed, even at program shutdown
+static TArray<FViewInfo*> FreeViewInfoSnapshots;
 
 FViewInfo* FViewInfo::CreateSnapshot() const
 {
 	QUICK_SCOPE_CYCLE_COUNTER(STAT_FViewInfo_CreateSnapshot);
 
-	check(IsInRenderingThread() && FMemStack::Get().GetNumMarks() == 1); // we do not want this popped before the end of the scene and it better be the scene allocator
-	FViewInfo* Result = (FViewInfo*)FMemStack::Get().Alloc(sizeof(FViewInfo),ALIGNOF(FViewInfo));
+	check(IsInRenderingThread()); // we do not want this popped before the end of the scene and it better be the scene allocator
+	FViewInfo* Result;
+	if (FreeViewInfoSnapshots.Num())
+	{
+		Result = FreeViewInfoSnapshots.Pop(false);
+	}
+	else
+	{
+		Result = (FViewInfo*)FMemory::Malloc(sizeof(FViewInfo), ALIGNOF(FViewInfo));
+	}
 	FMemory::Memcpy(*Result, *this);
 	TUniformBufferRef<FViewUniformShaderParameters> NullUniformBuffer;
 	FMemory::Memcpy(Result->UniformBuffer, NullUniformBuffer); // we want this to start null without a reference count, since we clear a ref later
@@ -995,9 +1006,20 @@ void FViewInfo::DestroyAllSnapshots()
 	QUICK_SCOPE_CYCLE_COUNTER(STAT_FViewInfo_DestroyAllSnapshots);
 
 	check(IsInRenderingThread());
+	// we will only keep double the number actually used, plus a few
+	int32 NumToRemove = FreeViewInfoSnapshots.Num() - (ViewInfoSnapshots.Num() + 2);
+	if (NumToRemove > 0)
+	{
+		for (int32 Index = 0; Index < NumToRemove; Index++)
+		{
+			FMemory::Free(FreeViewInfoSnapshots[Index]);
+		}
+		FreeViewInfoSnapshots.RemoveAt(0, NumToRemove, false);
+	}
 	for (FViewInfo* Snapshot : ViewInfoSnapshots)
 	{
 		Snapshot->UniformBuffer.SafeRelease();
+		FreeViewInfoSnapshots.Add(Snapshot);
 	}
 	ViewInfoSnapshots.Reset();
 }
