fixed showflag for DBuffer decals

diff --git a/Engine/Source/Runtime/Renderer/Private/CompositionLighting/CompositionLighting.cpp b/Engine/Source/Runtime/Renderer/Private/CompositionLighting/CompositionLighting.cpp
--- a/Engine/Source/Runtime/Renderer/Private/CompositionLighting/CompositionLighting.cpp
+++ b/Engine/Source/Runtime/Renderer/Private/CompositionLighting/CompositionLighting.cpp
@@ -234,8 +234,7 @@ void FCompositionLighting::ProcessBeforeBasePass(FRHICommandListImmediate& RHICm
 		// Add the passes we want to add to the graph (commenting a line means the pass is not inserted into the graph) ----------
 
 		// decals are before AmbientOcclusion so the decal can output a normal that AO is affected by
-		if (Context.View.Family->EngineShowFlags.Decals &&
-			!Context.View.Family->EngineShowFlags.ShaderComplexity &&
+		if (!Context.View.Family->EngineShowFlags.ShaderComplexity &&
 			IsDBufferEnabled()) 
 		{
 			FRenderingCompositePass* Pass = Context.Graph.RegisterPass(new(FMemStack::Get()) FRCPassPostProcessDeferredDecals(DRS_BeforeBasePass));
diff --git a/Engine/Source/Runtime/Renderer/Private/CompositionLighting/PostProcessDeferredDecals.cpp b/Engine/Source/Runtime/Renderer/Private/CompositionLighting/PostProcessDeferredDecals.cpp
--- a/Engine/Source/Runtime/Renderer/Private/CompositionLighting/PostProcessDeferredDecals.cpp
+++ b/Engine/Source/Runtime/Renderer/Private/CompositionLighting/PostProcessDeferredDecals.cpp
@@ -626,7 +626,7 @@ void FRCPassPostProcessDeferredDecals::Process(FRenderingCompositePassContext& C
 	const FSceneViewFamily& ViewFamily = *(View.Family);
 	FScene& Scene = *(FScene*)ViewFamily.Scene;
 
-	if(!Scene.Decals.Num())
+	if(!Scene.Decals.Num() || !ViewFamily.EngineShowFlags.Decals)
 	{
 		// to avoid the stats showing up
 		return;
