Further fixes for pause with temporal aa

diff --git a/Engine/Source/Runtime/Renderer/Private/PostProcess/PostProcessTemporalAA.cpp b/Engine/Source/Runtime/Renderer/Private/PostProcess/PostProcessTemporalAA.cpp
--- a/Engine/Source/Runtime/Renderer/Private/PostProcess/PostProcessTemporalAA.cpp
+++ b/Engine/Source/Runtime/Renderer/Private/PostProcess/PostProcessTemporalAA.cpp
@@ -703,15 +703,15 @@ void FRCPassPostProcessTemporalAA::Process(FRenderingCompositePassContext& Conte
 	}
 
 	Context.RHICmdList.CopyToResolveTarget(DestRenderTarget.TargetableTexture, DestRenderTarget.ShaderResourceTexture, false, FResolveParams());
-	
-	if( !CVarTemporalAAPauseCorrect.GetValueOnRenderThread() )
+
+	if( CVarTemporalAAPauseCorrect.GetValueOnRenderThread() )
+	{
+		ViewState->PendingTemporalAAHistoryRT = PassOutputs[0].PooledRenderTarget;
+	}
+	else
 	{
-		// Release to potentially reuse memory.
-		// Otherwise keep around so that pause looks correct.
-		ViewState->TemporalAAHistoryRT.SafeRelease();
+		ViewState->TemporalAAHistoryRT = PassOutputs[0].PooledRenderTarget;
 	}
-	ViewState->PendingTemporalAAHistoryRT = PassOutputs[0].PooledRenderTarget;
-	check( ViewState->PendingTemporalAAHistoryRT );
 
 	// TODO draw separate translucency after jitter has been removed
 
@@ -734,7 +734,7 @@ FPooledRenderTargetDesc FRCPassPostProcessTemporalAA::ComputeOutputDesc(EPassOut
 	FPooledRenderTargetDesc Ret = PassInputs[0].GetOutput()->RenderTargetDesc;
 
 	Ret.Reset();
-	Ret.Format = PF_FloatRGBA;
+	//Ret.Format = PF_FloatRGBA;
 	Ret.DebugName = TEXT("TemporalAA");
 
 	return Ret;
diff --git a/Engine/Source/Runtime/Renderer/Private/SceneVisibility.cpp b/Engine/Source/Runtime/Renderer/Private/SceneVisibility.cpp
--- a/Engine/Source/Runtime/Renderer/Private/SceneVisibility.cpp
+++ b/Engine/Source/Runtime/Renderer/Private/SceneVisibility.cpp
@@ -1584,8 +1584,11 @@ void FSceneRenderer::PreVisibilityFrameSetup(FRHICommandListImmediate& RHICmdLis
 					if(!ViewFamily.bWorldIsPaused)
 					{
 						ViewState->PrevViewMatrices = ViewState->PendingPrevViewMatrices;
-						ViewState->TemporalAAHistoryRT = ViewState->PendingTemporalAAHistoryRT;
-						ViewState->PendingTemporalAAHistoryRT.SafeRelease();
+						if( ViewState->PendingTemporalAAHistoryRT.GetRefCount() )
+						{
+							ViewState->TemporalAAHistoryRT = ViewState->PendingTemporalAAHistoryRT;
+							ViewState->PendingTemporalAAHistoryRT.SafeRelease();
+						}
 
 						// pending is needed as we are in init view and still need to render.
 						ViewState->PendingPrevViewMatrices = View.ViewMatrices;
