Fix ensure every frame when using ResponsiveAA on a material used in a Cavnas Tile render.

diff --git a/Engine/Source/Runtime/Renderer/Private/Renderer.cpp b/Engine/Source/Runtime/Renderer/Private/Renderer.cpp
--- a/Engine/Source/Runtime/Renderer/Private/Renderer.cpp
+++ b/Engine/Source/Runtime/Renderer/Private/Renderer.cpp
@@ -53,7 +53,7 @@ void FRendererModule::DrawTileMesh(FRHICommandListImmediate& RHICmdList, const F
 		{
 			if (FeatureLevel >= ERHIFeatureLevel::SM4)
 			{
-				FTranslucencyDrawingPolicyFactory::DrawDynamicMesh(RHICmdList, View, FTranslucencyDrawingPolicyFactory::ContextType(), Mesh, false, false, NULL, HitProxyId);
+				FTranslucencyDrawingPolicyFactory::DrawDynamicMesh(RHICmdList, View, FTranslucencyDrawingPolicyFactory::ContextType(false, false, true), Mesh, false, false, NULL, HitProxyId);
 			}
 			else
 			{
diff --git a/Engine/Source/Runtime/Renderer/Private/TranslucentRendering.cpp b/Engine/Source/Runtime/Renderer/Private/TranslucentRendering.cpp
--- a/Engine/Source/Runtime/Renderer/Private/TranslucentRendering.cpp
+++ b/Engine/Source/Runtime/Renderer/Private/TranslucentRendering.cpp
@@ -425,7 +425,8 @@ bool FTranslucencyDrawingPolicyFactory::DrawMesh(
 		// editor compositing not supported on translucent materials currently
 		const bool bEditorCompositeDepthTest = false;
 
-		if( bEnableResponsiveAA )
+		// if this draw is coming postAA then there is probably no depth buffer (it's canvas) and bEnableResponsiveAA wont' do anything anyway.
+		if (bEnableResponsiveAA && !DrawingContext.bPostAA)
 		{
 			if( bDisableDepthTest )
 			{
diff --git a/Engine/Source/Runtime/Renderer/Private/TranslucentRendering.h b/Engine/Source/Runtime/Renderer/Private/TranslucentRendering.h
--- a/Engine/Source/Runtime/Renderer/Private/TranslucentRendering.h
+++ b/Engine/Source/Runtime/Renderer/Private/TranslucentRendering.h
@@ -19,11 +19,13 @@ public:
 		const FProjectedShadowInfo* TranslucentSelfShadow;
 		bool bSeparateTranslucencyPass;
 		bool bSceneColorCopyIsUpToDate;
+		bool bPostAA;
 
-		ContextType(const FProjectedShadowInfo* InTranslucentSelfShadow = NULL, bool bSeparateTranslucencyPassIn = false)
+		ContextType(const FProjectedShadowInfo* InTranslucentSelfShadow = NULL, bool bSeparateTranslucencyPassIn = false, bool bPostAAIn = false)
 			: TranslucentSelfShadow(InTranslucentSelfShadow),
 			bSeparateTranslucencyPass(bSeparateTranslucencyPassIn),
-			bSceneColorCopyIsUpToDate(false)
+			bSceneColorCopyIsUpToDate(false),
+			bPostAA(bPostAAIn)
 		{}
 	};
 
