define occlusion fence array with enum.

diff --git a/Engine/Source/Runtime/Renderer/Private/DeferredShadingRenderer.cpp b/Engine/Source/Runtime/Renderer/Private/DeferredShadingRenderer.cpp
--- a/Engine/Source/Runtime/Renderer/Private/DeferredShadingRenderer.cpp
+++ b/Engine/Source/Runtime/Renderer/Private/DeferredShadingRenderer.cpp
@@ -674,7 +674,7 @@ DECLARE_CYCLE_STAT(TEXT("AfterVelocity"), STAT_CLM_AfterVelocity, STATGROUP_Comm
 DECLARE_CYCLE_STAT(TEXT("RenderFinish"), STAT_CLM_RenderFinish, STATGROUP_CommandListMarkers);
 DECLARE_CYCLE_STAT(TEXT("AfterFrame"), STAT_CLM_AfterFrame, STATGROUP_CommandListMarkers);
 
-FGraphEventRef FDeferredShadingSceneRenderer::OcclusionSubmittedFence[2];
+FGraphEventRef FDeferredShadingSceneRenderer::OcclusionSubmittedFence[FOcclusionQueryHelpers::MaxBufferedOcclusionFrames];
 
 /**
  * Returns true if the depth Prepass needs to run
