Fix CIS.

diff --git a/Engine/Source/Runtime/Renderer/Private/Renderer.cpp b/Engine/Source/Runtime/Renderer/Private/Renderer.cpp
--- a/Engine/Source/Runtime/Renderer/Private/Renderer.cpp
+++ b/Engine/Source/Runtime/Renderer/Private/Renderer.cpp
@@ -53,7 +53,7 @@ void FRendererModule::DrawTileMesh(FRHICommandListImmediate& RHICmdList, const F
 		{
 			if (FeatureLevel >= ERHIFeatureLevel::SM4)
 			{
-				FTranslucencyDrawingPolicyFactory::DrawDynamicMesh(RHICmdList, View, FTranslucencyDrawingPolicyFactory::ContextType(false, false, true), Mesh, false, false, NULL, HitProxyId);
+				FTranslucencyDrawingPolicyFactory::DrawDynamicMesh(RHICmdList, View, FTranslucencyDrawingPolicyFactory::ContextType(nullptr, false, true), Mesh, false, false, NULL, HitProxyId);
 			}
 			else
 			{
