removed unused code - using it could be wrong anyway

diff --git a/Engine/Source/Runtime/Renderer/Private/PostProcess/PostProcessTemporalAA.cpp b/Engine/Source/Runtime/Renderer/Private/PostProcess/PostProcessTemporalAA.cpp
--- a/Engine/Source/Runtime/Renderer/Private/PostProcess/PostProcessTemporalAA.cpp
+++ b/Engine/Source/Runtime/Renderer/Private/PostProcess/PostProcessTemporalAA.cpp
@@ -116,9 +116,6 @@ public:
 
 		{
 			const FPooledRenderTargetDesc* InputDesc = Context.Pass->GetInputDesc(ePId_Input0);
-			const FIntPoint ViewportSize = Context.GetViewport().Size();
-			const float Width = ViewportSize.X;
-			const float Height = ViewportSize.Y;
 
 			static const float SampleOffsets[9][2] =
 			{
