added comment

diff --git a/Engine/Source/Runtime/Renderer/Private/PostProcess/RenderingCompositionGraph.h b/Engine/Source/Runtime/Renderer/Private/PostProcess/RenderingCompositionGraph.h
--- a/Engine/Source/Runtime/Renderer/Private/PostProcess/RenderingCompositionGraph.h
+++ b/Engine/Source/Runtime/Renderer/Private/PostProcess/RenderingCompositionGraph.h
@@ -114,7 +114,7 @@ struct FRenderingCompositePassContext
 		check(!IsViewportValid());
 	}
 	
-	//
+	// Return the hardware viewport rectangle, not necessarily the current view rectangle (e.g. a post process can Set it to be larger than that)
 	FIntRect GetViewport() const
 	{
 		// need to call SetViewportAndCallRHI() before
