fixed UE-17405 Circle Depth of Field is no Resolution independant

diff --git a/Engine/Shaders/PostProcessDOF.usf b/Engine/Shaders/PostProcessDOF.usf
--- a/Engine/Shaders/PostProcessDOF.usf
+++ b/Engine/Shaders/PostProcessDOF.usf
@@ -198,8 +198,8 @@ void MainRecombinePS(
 // PROTOTYPE CIRCLE DOF : WORK IN PROGRESS
 //
 
-// {radius, depth blur amp, depth blur radius}
-float3 CircleDofParams;
+// {radius, depth blur amp, depth blur radius, ViewPortScale}
+float4 CircleDofParams;
 
 // CIRCLE DOF: Compute circle of confusion size in pixels.
 float DepthToCoc(float Depth) 
@@ -211,6 +211,7 @@ float DepthToCoc(float Depth)
 	float ReturnCoc = max(abs(CocRadius), DepthBlurRadius);
 	if(CocRadius < 0.0) 
 	{
+		// near CoC is using negative values
 		ReturnCoc = -ReturnCoc;
 	}
 	return ReturnCoc;
@@ -347,7 +348,7 @@ void CirclePS(float4 UVAndScreenPos : TEXCOORD0, out float4 OutColor0 : SV_Targe
 
 	// Fixed maximum search size (in terms of Circle of Confusion radius).
 	// Higher than 8 is too noizy for 4 samples.
-	float Coc = 8.0;
+	float Coc = 8.0 * CircleDofParams.w;
 
 	// Get base semi-random direction and dither along radius.
 	// Reused throughout the rest of the algorithm.
@@ -405,7 +406,7 @@ void CirclePS(float4 UVAndScreenPos : TEXCOORD0, out float4 OutColor0 : SV_Targe
 
 	// Used for sample pattern.
 	Coc = max(abs(FarCoc),abs(NearCoc));
-
+	
 	// Bring out to the smaller radius of sample sets.
 	// This has the highest chance of seeing a smaller overlapping CoC.
 	R1 = (RadiusBase+9.0/11.5) * Coc;
diff --git a/Engine/Source/Runtime/Renderer/Private/PostProcess/PostProcessDOF.cpp b/Engine/Source/Runtime/Renderer/Private/PostProcess/PostProcessDOF.cpp
--- a/Engine/Source/Runtime/Renderer/Private/PostProcess/PostProcessDOF.cpp
+++ b/Engine/Source/Runtime/Renderer/Private/PostProcess/PostProcessDOF.cpp
@@ -356,7 +356,7 @@ FPooledRenderTargetDesc FRCPassPostProcessDOFRecombine::ComputeOutputDesc(EPassO
 
 // Convert f-stop and focal distance into projected size in half resolution pixels.
 // Setup depth based blur.
-static FVector CircleDofCoc(const FRenderingCompositePassContext& Context)
+static FVector4 CircleDofCoc(const FRenderingCompositePassContext& Context)
 {
 	// Convert FOV to focal length,
 	// 
@@ -386,7 +386,7 @@ static FVector CircleDofCoc(const FRenderingCompositePassContext& Context)
 
 	// Scale so that APS-C 24.576 mm = full frame.
 	// Convert mm to pixels.
-	float Width = (float)Context.GetViewport().Size().X;
+	float Width = (float)Context.View.ViewRect.Width();
 	Radius = Radius * Width * (1.0f/24.576f);
 
 	// Convert diameter to radius at half resolution (algorithm radius is at half resolution).
@@ -403,10 +403,11 @@ static FVector CircleDofCoc(const FRenderingCompositePassContext& Context)
 
 	// The DepthOfFieldDepthBlurAmount = km at which depth blur is 50%.
 	// Need to convert to cm here.
-	return FVector(
+	return FVector4(
 		Radius, 
 		1.0f/(Context.View.FinalPostProcessSettings.DepthOfFieldDepthBlurAmount * 100000.0),
-		Context.View.FinalPostProcessSettings.DepthOfFieldDepthBlurRadius * Width / 1920.0f);
+		Context.View.FinalPostProcessSettings.DepthOfFieldDepthBlurRadius * Width / 1920.0f,
+		Width / 1920.0f);
 }
 
 /** Encapsulates the Circle DOF setup pixel shader. */
