UE4 - fix bad recycling on RHICommandLists, was an exit crash and wasn't recycling

diff --git a/Engine/Source/Runtime/RHI/Private/RHICommandList.cpp b/Engine/Source/Runtime/RHI/Private/RHICommandList.cpp
--- a/Engine/Source/Runtime/RHI/Private/RHICommandList.cpp
+++ b/Engine/Source/Runtime/RHI/Private/RHICommandList.cpp
@@ -1401,9 +1401,6 @@ void* FRHICommandList::operator new(size_t Size)
 	//return FMemory::Malloc(Size);
 }
 
-/**
- * Custom delete
- */
 void FRHICommandList::operator delete(void *RawMemory)
 {
 	check(RawMemory != (void*) &GRHICommandList.GetImmediateCommandList());
@@ -1411,6 +1408,19 @@ void FRHICommandList::operator delete(void *RawMemory)
 	//FMemory::Free(RawMemory);
 }	
 
+void* FRHICommandListBase::operator new(size_t Size)
+{
+	check(0); // you shouldn't be creating these
+	return nullptr;
+}
+
+void FRHICommandListBase::operator delete(void *RawMemory)
+{
+	check(RawMemory != (void*) &GRHICommandList.GetImmediateCommandList());
+	RHICommandListAllocator.Free(RawMemory);
+	//FMemory::Free(RawMemory);
+}	
+
 ///////// Pass through functions that allow RHIs to optimize certain calls.
 
 FVertexBufferRHIRef FDynamicRHI::CreateAndLockVertexBuffer_RenderThread(class FRHICommandListImmediate& RHICmdList, uint32 Size, uint32 InUsage, FRHIResourceCreateInfo& CreateInfo, void*& OutDataBuffer)
diff --git a/Engine/Source/Runtime/RHI/Public/RHICommandList.h b/Engine/Source/Runtime/RHI/Public/RHICommandList.h
--- a/Engine/Source/Runtime/RHI/Public/RHICommandList.h
+++ b/Engine/Source/Runtime/RHI/Public/RHICommandList.h
@@ -108,6 +108,10 @@ public:
 	FRHICommandListBase();
 	~FRHICommandListBase();
 
+	/** Custom new/delete with recycling */
+	void* operator new(size_t Size);
+	void operator delete(void *RawMemory);
+
 	inline void Flush();
 	inline bool IsImmediate();
 	const int32 GetUsedMemory() const;
diff --git a/Engine/Source/Runtime/Renderer/Private/SceneRendering.cpp b/Engine/Source/Runtime/Renderer/Private/SceneRendering.cpp
--- a/Engine/Source/Runtime/Renderer/Private/SceneRendering.cpp
+++ b/Engine/Source/Runtime/Renderer/Private/SceneRendering.cpp
@@ -141,6 +141,7 @@ FParallelCommandListSet::FParallelCommandListSet(const FViewInfo& InView, FRHICo
 	: View(InView)
 	, ParentCmdList(InParentCmdList)
 	, Snapshot(nullptr)
+	, NumAlloc(0)
 	, bParallelExecute(GRHISupportsParallelRHIExecute && bInParallelExecute)
 	, bCreateSceneContext(bInCreateSceneContext)
 {
@@ -155,6 +156,7 @@ FParallelCommandListSet::FParallelCommandListSet(const FViewInfo& InView, FRHICo
 
 FRHICommandList* FParallelCommandListSet::AllocCommandList()
 {
+	NumAlloc++;
 	return new FRHICommandList;
 }
 
@@ -162,6 +164,7 @@ void FParallelCommandListSet::Dispatch()
 {
 	check(IsInRenderingThread() && FMemStack::Get().GetNumMarks() == 1); // we do not want this popped before the end of the scene and it better be the scene allocator
 	check(CommandLists.Num() == Events.Num());
+	check(CommandLists.Num() == NumAlloc);
 	if (bSpewBalance)
 	{
 		// finish them all
@@ -202,6 +205,7 @@ void FParallelCommandListSet::Dispatch()
 	{
 		UE_CLOG(bSpewBalance, LogTemp, Display, TEXT("%d cmdlists for parallel translate"), CommandLists.Num());
 		check(GRHISupportsParallelRHIExecute);
+		NumAlloc -= CommandLists.Num();
 		ParentCmdList.QueueParallelAsyncCommandListSubmit(&Events[0], &CommandLists[0], &NumDrawsIfKnown[0], CommandLists.Num(), (MinDrawsPerCommandList * 4) / 3, bSpewBalance);
 		SetStateOnCommandList(ParentCmdList);
 	}
@@ -211,6 +215,7 @@ void FParallelCommandListSet::Dispatch()
 		for (int32 Index = 0; Index < CommandLists.Num(); Index++)
 		{
 			ParentCmdList.QueueAsyncCommandListSubmit(Events[Index], CommandLists[Index]);
+			NumAlloc--;
 		}
 	}
 	CommandLists.Reset();
@@ -222,6 +227,7 @@ FParallelCommandListSet::~FParallelCommandListSet()
 {
 	check(IsInRenderingThread() && FMemStack::Get().GetNumMarks() == 1); // we do not want this popped before the end of the scene and it better be the scene allocator
 	checkf(CommandLists.Num() == 0, TEXT("Derived class of FParallelCommandListSet did not call Dispatch in virtual destructor"));
+	checkf(NumAlloc == 0, TEXT("Derived class of FParallelCommandListSet did not call Dispatch in virtual destructor"));
 }
 
 FRHICommandList* FParallelCommandListSet::NewParallelCommandList()
diff --git a/Engine/Source/Runtime/Renderer/Private/SceneRendering.h b/Engine/Source/Runtime/Renderer/Private/SceneRendering.h
--- a/Engine/Source/Runtime/Renderer/Private/SceneRendering.h
+++ b/Engine/Source/Runtime/Renderer/Private/SceneRendering.h
@@ -315,6 +315,7 @@ public:
 	FRHICommandListImmediate& ParentCmdList;
 	FSceneRenderTargets* Snapshot;
 	int32 Width;
+	int32 NumAlloc;
 	int32 MinDrawsPerCommandList;
 	bool bBalanceCommands;
 	bool bSpewBalance;
