Build Hierarchical Color Buffer and use it for the SSR cone prototype

diff --git a/Engine/Shaders/BasePassPixelShader.usf b/Engine/Shaders/BasePassPixelShader.usf
--- a/Engine/Shaders/BasePassPixelShader.usf
+++ b/Engine/Shaders/BasePassPixelShader.usf
@@ -132,8 +132,15 @@ half3 GetImageBasedReflectionLighting(FMaterialPixelParameters MaterialParameter
 		float StepOffset = Offset / 15.0;
 		StepOffset -= 0.5;
 		
-		//float4 HitUVzTime = RayCastCheap( HZBTexture, HZBSampler, MaterialParameters.ReflectionVector, MaterialParameters.ScreenPosition.w, MaterialParameters.WorldPosition_CamRelative, StepOffset );
-		float4 HitUVzTime = RayCast( HZBTexture, HZBSampler, MaterialParameters.ReflectionVector, 0, MaterialParameters.ScreenPosition.w, MaterialParameters.WorldPosition_CamRelative, 12, StepOffset );
+		float4 HitUVzTime;
+		float HCBLevel;
+
+		RayCast(
+			HZBTexture, HZBSampler, float2(1, 1),
+			MaterialParameters.WorldPosition_CamRelative, MaterialParameters.ReflectionVector, 0, 0, MaterialParameters.ScreenPosition.w,
+			12, StepOffset,
+			HitUVzTime, HCBLevel
+		);
 
 		// if there was a hit
 		BRANCH if( HitUVzTime.w < 1 )
diff --git a/Engine/Shaders/MonteCarlo.usf b/Engine/Shaders/MonteCarlo.usf
--- a/Engine/Shaders/MonteCarlo.usf
+++ b/Engine/Shaders/MonteCarlo.usf
@@ -1,5 +1,7 @@
 // Copyright 1998-2015 Epic Games, Inc. All Rights Reserved.
 
+#pragma once
+
 /*=============================================================================
 	MonteCarlo.usf: Monte Carlo integration of distributions
 =============================================================================*/
diff --git a/Engine/Shaders/PostProcessHierarchical.usf b/Engine/Shaders/PostProcessHierarchical.usf
new file mode 100644
--- /dev/null
+++ b/Engine/Shaders/PostProcessHierarchical.usf
@@ -0,0 +1,52 @@
+// Copyright 1998-2015 Epic Games, Inc. All Rights Reserved.
+
+/*=============================================================================
+	PostProcessHierarchical.usf: To generate hierarchical buffers
+=============================================================================*/
+
+#include "Common.usf"
+#include "PostProcessCommon.usf"
+
+float2 InvSize;
+float4 InputUvFactorAndOffset;
+float4 InputUvBundaries;
+
+Texture2D		ColorMip;
+SamplerState	ColorMipSampler;
+
+// kernel to build hierarchical scene color buffer
+void BuildHCB(float4 RenderTargetUVAndScreenPos : TEXCOORD0, out float4 OutColor : SV_Target0)
+{
+#if STAGE == 0
+	const float2 ViewUV = RenderTargetUVAndScreenPos.zw * float2(0.5f, -0.5f) + 0.5f;
+	const float2 InUV = ViewUV * InputUvFactorAndOffset.xy + InputUvFactorAndOffset.zw;
+	const float2 InUVMinBound = InputUvBundaries.xy;
+	const float2 InUVMaxBound = InputUvBundaries.zw;
+	const Texture2D Texture = PostprocessInput0;
+	const SamplerState TextureSampler = PostprocessInput0Sampler;
+
+#else
+	const float2 InUV = RenderTargetUVAndScreenPos.xy;
+	const float2 InUVMinBound = InputUvBundaries.xy;
+	const float2 InUVMaxBound = InputUvBundaries.zw;
+	const Texture2D Texture = ColorMip;
+	const SamplerState TextureSampler = ColorMipSampler;
+	
+#endif
+	
+	const float Radius = 2.0 * InvSize;
+	// sigma = 0.5
+	const float C = 0.468592;
+	const float A = 0.107973;
+	const float B = 0.024879;
+
+	OutColor =  Texture.SampleLevel(TextureSampler, clamp(InUV, InUVMinBound, InUVMaxBound), 0) * C;
+	OutColor += Texture.SampleLevel(TextureSampler, clamp(InUV + float2( 1,  1) * Radius, InUVMinBound, InUVMaxBound), 0) * B;
+	OutColor += Texture.SampleLevel(TextureSampler, clamp(InUV + float2( 0,  1) * Radius, InUVMinBound, InUVMaxBound), 0) * A;
+	OutColor += Texture.SampleLevel(TextureSampler, clamp(InUV + float2(-1,  1) * Radius, InUVMinBound, InUVMaxBound), 0) * B;
+	OutColor += Texture.SampleLevel(TextureSampler, clamp(InUV + float2(-1,  0) * Radius, InUVMinBound, InUVMaxBound), 0) * A;
+	OutColor += Texture.SampleLevel(TextureSampler, clamp(InUV + float2(-1, -1) * Radius, InUVMinBound, InUVMaxBound), 0) * B;
+	OutColor += Texture.SampleLevel(TextureSampler, clamp(InUV + float2( 0, -1) * Radius, InUVMinBound, InUVMaxBound), 0) * A;
+	OutColor += Texture.SampleLevel(TextureSampler, clamp(InUV + float2( 1, -1) * Radius, InUVMinBound, InUVMaxBound), 0) * B;
+	OutColor += Texture.SampleLevel(TextureSampler, clamp(InUV + float2( 1,  0) * Radius, InUVMinBound, InUVMaxBound), 0) * A;
+}
diff --git a/Engine/Shaders/ScreenSpaceRayCast.usf b/Engine/Shaders/ScreenSpaceRayCast.usf
--- a/Engine/Shaders/ScreenSpaceRayCast.usf
+++ b/Engine/Shaders/ScreenSpaceRayCast.usf
@@ -38,41 +38,58 @@ float4 SampleDepthTexture( Texture2D Texture, SamplerState Sampler, float Level,
 	return SampleDepth;
 }
 
-float4 RayCast( Texture2D Texture, SamplerState Sampler, float3 R, float Roughness, float SceneDepth, float3 PositionTranslatedWorld, int NumSteps, float StepOffset )
+void RayCast(
+	Texture2D Texture, SamplerState Sampler, float2 TextureSize,
+	float3 RayOriginTranslatedWorld, float3 RayDirection,
+	float Roughness, float ConeAngleWorld, float SceneDepth,
+	int NumSteps, float StepOffset,
+	out float4 OutHitUVzTime,
+	out float OutLevelHCB
+)
 {
 	// TODO provide RayStartUVz
+	const float4 RayStartV = mul( float4( RayOriginTranslatedWorld, 1 ), View.TranslatedWorldToView );
+	const float4 RayDirV = mul(float4(RayDirection * SceneDepth, 0),  View.TranslatedWorldToView);
+	const float4 RayEndV = RayStartV + RayDirV;
 
-	// NOTE could clip ray against frustum planes
+	const float RayEndRadiusV = length(RayDirV.xyz) * tan(ConeAngleWorld * 0.4);
+	const float2 RayEndBorderV = RayEndV.xy + RayEndRadiusV / sqrt(2);
 
-	// TODO use screen position and skip matrix mul
-	float4 RayStartClip	= mul( float4( PositionTranslatedWorld, 1 ), View.TranslatedWorldToClip );
-	float4 RayEndClip	= mul( float4( PositionTranslatedWorld + R * SceneDepth, 1 ), View.TranslatedWorldToClip );
+	const float4 RayStartClip	= mul( RayStartV, View.ViewToClip );
+	const float4 RayEndClip	= mul( RayEndV, View.ViewToClip );
+	const float2 RayEndBorderClip = mul(float4(RayEndBorderV, RayEndV.zw), View.ViewToClip).xy;
 
-	float3 RayStartScreen = RayStartClip.xyz / RayStartClip.w;
-	float3 RayEndScreen = RayEndClip.xyz / RayEndClip.w;
+	const float3 RayStartScreen = RayStartClip.xyz / RayStartClip.w;
+	const float3 RayEndScreen = RayEndClip.xyz / RayEndClip.w;
+	const float2 RayEndBorderScreen = RayEndBorderClip.xy / RayEndClip.w;
 	
-	// Make sure the RayStepScreen has a length of 2 in order to have the end point ouside of the viewport
-	float3 RayStepScreen = ( RayEndScreen - RayStartScreen ) / (0.5 * length( RayEndScreen.xy - RayStartScreen.xy ));
+	float3 RayStepScreen = (RayEndScreen - RayStartScreen);
+	float2 RayStepRadiusScreen = (RayEndBorderScreen - RayEndScreen.xy);
 
 	{
 		// Computes the scale down factor for RayStepScreen required to fit on the X and Y axis in order to clip it in the viewport
-		float2 AbsRayStepScreen = abs(RayStepScreen.xy);
-		float2 S = (AbsRayStepScreen - max(abs(RayStepScreen.xy + RayStartScreen.xy) - 1.0f, 0.0f)) / AbsRayStepScreen;
+		const float RayStepScreenInvFactor = 0.5 * length( RayStepScreen.xy );
+		const float2 AbsRayStepScreen = abs(RayStepScreen.xy);
+		const float2 S = (AbsRayStepScreen - max(abs(RayStepScreen.xy + RayStartScreen.xy * RayStepScreenInvFactor) - RayStepScreenInvFactor, 0.0f)) / AbsRayStepScreen;
 
 		// Rescales RayStepScreen accordingly
-		RayStepScreen *= min(S.x, S.y);
+		const float RayStepFactor = min(S.x, S.y) / RayStepScreenInvFactor;
+
+		RayStepScreen *= RayStepFactor;
+		RayStepRadiusScreen *= RayStepFactor;
 	}
 
+	const float2 HzbSize = float2(1024, 512);
 #if USE_HZB
-	float3 RayStartUVz = float3( RayStartScreen.xy * float2( 0.5, -0.5 ) + 0.5, RayStartScreen.z );
-	float3 RayStepUVz = float3( RayStepScreen.xy * float2( 0.5, -0.5 ), RayStepScreen.z );
-
-	RayStartUVz.xy *= HZBUvFactorAndInvFactor.xy;
-	RayStepUVz.xy *= HZBUvFactorAndInvFactor.xy;
+	const float3 RayStartUVz = float3( (RayStartScreen.xy * float2( 0.5, -0.5 ) + 0.5) * HZBUvFactorAndInvFactor.xy, RayStartScreen.z );
+	const float3 RayStepUVz = float3( RayStepScreen.xy * float2( 0.5, -0.5 ) * HZBUvFactorAndInvFactor.xy, RayStepScreen.z );
+	const float2 RayStepRadiusUV = RayStepRadiusScreen.xy * float2( 0.5, -0.5 ) * HZBUvFactorAndInvFactor.xy;
 #else
-	float3 RayStartUVz = float3( RayStartScreen.xy * View.ScreenPositionScaleBias.xy + View.ScreenPositionScaleBias.wz, RayStartScreen.z );
-	float3 RayStepUVz = float3( RayStepScreen.xy * View.ScreenPositionScaleBias.xy, RayStepScreen.z );
+	const float3 RayStartUVz = float3( (RayStartScreen.xy * float2( 0.5, -0.5 ) + 0.5), RayStartScreen.z );
+	const float3 RayStepUVz = float3( RayStepScreen.xy * float2( 0.5, -0.5 ), RayStepScreen.z );
+	const float2 RayStepRadiusUV = RayStepRadiusScreen.xy * float2( 0.5, -0.5 );
 #endif
+	const float RayStepRadiusFactor = length(RayStepRadiusUV * TextureSize);
 	
 	const float Step = 1.0 / (NumSteps + 1);
 
@@ -229,7 +246,8 @@ float4 RayCast( Texture2D Texture, SamplerState Sampler, float3 R, float Roughne
 	Result.xy = Result.xy * View.ScreenPositionScaleBias.xy + View.ScreenPositionScaleBias.wz;
 #endif
 
-	return Result;
+	OutHitUVzTime = Result;
+	OutLevelHCB = log2(OutHitUVzTime.w * RayStepRadiusFactor) + 1;
 }
 
 float4 RayCastCheap( Texture2D Texture, SamplerState Sampler, float3 R, float SceneDepth, float3 PositionTranslatedWorld, float StepOffset )
@@ -367,4 +385,44 @@ float4 SampleScreenColor( Texture2D Texture, SamplerState Sampler, float3 HitUVz
 	OutColor.rgb = -min(-OutColor.rgb, 0.0);
 
 	return OutColor;
-}
\ No newline at end of file
+}
+
+float4 SampleHCBLevel( Texture2D Texture, SamplerState Sampler, float3 HitUVz, float Level )
+{
+	float4 OutColor;
+
+#if PREV_FRAME_COLOR
+	// Find previous screen position for hit since color buffer is from last frame
+	// TODO combine to single matrix
+	float4 HitClip = float4( (HitUVz.xy - View.ScreenPositionScaleBias.wz) / View.ScreenPositionScaleBias.xy, HitUVz.z, 1 );
+	float4 HitTranslatedWorld = mul( HitClip, View.ClipToTranslatedWorld );
+	HitTranslatedWorld /= HitTranslatedWorld.w;
+		
+	float3 PrevTranslatedWorld = HitTranslatedWorld.xyz + (View.PrevPreViewTranslation - View.PreViewTranslation);
+	float4 PrevClip = mul( float4( PrevTranslatedWorld, 1 ), View.PrevTranslatedWorldToClip );
+	float2 PrevScreen = PrevClip.xy / PrevClip.w;
+	float2 PrevUV = PrevScreen.xy * View.ScreenPositionScaleBias.xy + View.ScreenPositionScaleBias.wz;
+
+	OutColor.rgb = Texture.SampleLevel( Sampler, PrevUV * HZBUvFactorAndInvFactor.xy, Level ).rgb;
+	OutColor.a = 1;
+
+	// Off screen masking
+	float2 Vignette = saturate( abs( PrevScreen ) * 5 - 4 );
+#else
+	OutColor.rgb = Texture.SampleLevel( Sampler, HitUVz.xy * HZBUvFactorAndInvFactor.xy, Level ).rgb;
+	OutColor.a = 1;
+
+	// Off screen masking
+	float2 HitScreenPos = ( HitUVz.xy - View.ScreenPositionScaleBias.wz ) / View.ScreenPositionScaleBias.xy;
+	float2 Vignette = saturate( abs( HitScreenPos ) * 5 - 4 );
+#endif
+		
+	//PrevScreen sometimes has NaNs or Infs.  DX11 is protected because saturate turns NaNs -> 0.
+	//Do a SafeSaturate so other platforms get the same protection.
+	OutColor *= SafeSaturate( 1.0 - dot( Vignette, Vignette ) );
+
+	// Transform NaNs to black, transform negative colors to black.
+	OutColor.rgb = -min(-OutColor.rgb, 0.0);
+
+	return OutColor;
+}
diff --git a/Engine/Shaders/ScreenSpaceReflections.usf b/Engine/Shaders/ScreenSpaceReflections.usf
--- a/Engine/Shaders/ScreenSpaceReflections.usf
+++ b/Engine/Shaders/ScreenSpaceReflections.usf
@@ -122,13 +122,10 @@ void ScreenSpaceReflectionsPS(in float4 UVAndScreenPos : TEXCOORD0, out float4 O
 	return;
 #endif
 
-	float3 N = GBuffer.WorldNormal;
-	float SceneDepth = CalcSceneDepth(UV);
-
-	float3 PositionTranslatedWorld = mul( float4( ScreenPos * SceneDepth, SceneDepth, 1 ), View.ScreenToTranslatedWorld ).xyz;
-
-	float3 V = View.TranslatedViewOrigin.xyz - PositionTranslatedWorld;
-	V = normalize( V );
+	const float3 N = GBuffer.WorldNormal;
+	const float SceneDepth = CalcSceneDepth(UV);
+	const float3 PositionTranslatedWorld = mul( float4( ScreenPos * SceneDepth, SceneDepth, 1 ), View.ScreenToTranslatedWorld ).xyz;
+	const float3 V = normalize(View.TranslatedViewOrigin.xyz - PositionTranslatedWorld);
 
 	uint FrameRandom;
 	{
@@ -148,6 +145,40 @@ void ScreenSpaceReflectionsPS(in float4 UVAndScreenPos : TEXCOORD0, out float4 O
 		}
 	}
 
+#if SSR_QUALITY == SSR_CONE_QUALITY
+	const int NumSteps = 16;
+	
+	{
+		const float2 PixelPos = UVAndScreenPos.zw * ScreenPosToPixel.xy + ScreenPosToPixel.zw;
+		const float2 E = float2(PseudoRandom(PixelPos + int2(FrameRandom, 0)), PseudoRandom(PixelPos + int2(0, FrameRandom)));
+		const float StepOffset = PseudoRandom(PixelPos + FrameRandom);
+
+		const float4 HNAndPDF = ImportanceSampleBlinn( E, Roughness );
+		const float3 H = TangentToWorld( HNAndPDF.xyz, N );
+		const float3 R = 2 * dot( V, H ) * H - V;
+
+		const float ReflectPDF = HNAndPDF.w / (4 * saturate(dot(V, H)));
+		const float TotalSamples = View.TemporalAAParams.g;
+		const float SolidAngleSample = 1.0 / ( TotalSamples * ReflectPDF );
+		const float ConeAngleWorld = acos( 1 - SolidAngleSample / (2*PI) );
+
+		float4 HitUVzTime;
+		float HCBLevel;
+		
+		RayCast(
+			PostprocessInput1, PostprocessInput1Sampler, PostprocessInput1Size,
+			PositionTranslatedWorld, R, Roughness, ConeAngleWorld, SceneDepth,
+			NumSteps, StepOffset,
+			HitUVzTime, HCBLevel
+		);
+
+		// if there was a hit
+		BRANCH if( HitUVzTime.w < 1 )
+		{
+			OutColor = SampleHCBLevel( PostprocessInput2, PostprocessInput2Sampler, HitUVzTime.xyz, HCBLevel );
+		}
+	}
+#else //SSR_QUALITY != SSR_CONE_QUALITY
 #if SSR_QUALITY == 1
 	const int NumSteps = 8;
 	const int NumRays = 1;
@@ -182,16 +213,21 @@ void ScreenSpaceReflectionsPS(in float4 UVAndScreenPos : TEXCOORD0, out float4 O
 			float3 H = TangentToWorld( ImportanceSampleBlinn( E, Roughness ).xyz, N );
 			float3 R = 2 * dot( V, H ) * H - V;
 
-			float4 HitUVzTime = RayCast( PostprocessInput1, PostprocessInput1Sampler, R, Roughness, SceneDepth, PositionTranslatedWorld, NumSteps, StepOffset );
+			float4 HitUVzTime;
+			float HCBLevel;
+
+			RayCast(
+				PostprocessInput1, PostprocessInput1Sampler, PostprocessInput1Size,
+				PositionTranslatedWorld, R, Roughness, 0.001, SceneDepth,
+				NumSteps, StepOffset,
+				HitUVzTime, HCBLevel
+			);
 
 			// if there was a hit
 			BRANCH if( HitUVzTime.w < 1 )
 			{
 				float4 SampleColor = SampleScreenColor( PostprocessInput0, PostprocessInput0Sampler, HitUVzTime.xyz );
 				
-				// Fade end of trace
-				SampleColor *= saturate( 4 - 4 * HitUVzTime.w );
-				
 				SampleColor.rgb /= 1 + Luminance(SampleColor.rgb);
 				OutColor += SampleColor;
 			}
@@ -207,17 +243,24 @@ void ScreenSpaceReflectionsPS(in float4 UVAndScreenPos : TEXCOORD0, out float4 O
 		StepOffset -= 0.5;
 		
 		float3 R = reflect( -V, N );
+		
+		float4 HitUVzTime;
+		float HCBLevel;
+
+		RayCast(
+			PostprocessInput1, PostprocessInput1Sampler, PostprocessInput1Size,
+			PositionTranslatedWorld, R, Roughness, 0.001, SceneDepth,
+			NumSteps, StepOffset,
+			HitUVzTime, HCBLevel
+		);
 
-		float4 HitUVzTime = RayCast( PostprocessInput1, PostprocessInput1Sampler, R, Roughness, SceneDepth, PositionTranslatedWorld, NumSteps, StepOffset );
 		// if there was a hit
 		BRANCH if( HitUVzTime.w < 1 )
 		{
 			OutColor = SampleScreenColor( PostprocessInput0, PostprocessInput0Sampler, HitUVzTime.xyz );
-
-			// Fade end of trace
-			OutColor *= saturate( 4 - 4 * HitUVzTime.w );
 		}
 	}
+#endif //SSR_QUALITY  != SSR_CONE_QUALITY
 	
 	OutColor *= RoughnessFade;
 	OutColor *= SSRParams.r;
diff --git a/Engine/Source/Runtime/Renderer/Private/PostProcess/PostProcessHierarchical.cpp b/Engine/Source/Runtime/Renderer/Private/PostProcess/PostProcessHierarchical.cpp
new file mode 100644
--- /dev/null
+++ b/Engine/Source/Runtime/Renderer/Private/PostProcess/PostProcessHierarchical.cpp
@@ -0,0 +1,266 @@
+// Copyright 1998-2015 Epic Games, Inc. All Rights Reserved.
+
+/*=============================================================================
+	PostProcessHierarchical.cpp: Post processing Screen Space Reflections implementation.
+=============================================================================*/
+
+#include "RendererPrivate.h"
+#include "ScenePrivate.h"
+#include "SceneFilterRendering.h"
+#include "PostProcessing.h"
+#include "PostProcessHierarchical.h"
+#include "SceneUtils.h"
+
+template< uint32 Stage >
+class TPostProcessBuildHCBPS : public FGlobalShader
+{
+	DECLARE_SHADER_TYPE(TPostProcessBuildHCBPS, Global);
+
+	static bool ShouldCache(EShaderPlatform Platform)
+	{
+		return IsFeatureLevelSupported(Platform, ERHIFeatureLevel::SM4);
+	}
+
+	static void ModifyCompilationEnvironment(EShaderPlatform Platform, FShaderCompilerEnvironment& OutEnvironment)
+	{
+		FGlobalShader::ModifyCompilationEnvironment(Platform, OutEnvironment);
+		OutEnvironment.SetDefine( TEXT("STAGE"), Stage );
+		OutEnvironment.SetRenderTargetOutputFormat(0, PF_FloatRGBA);
+	}
+
+	TPostProcessBuildHCBPS() {}
+
+public:
+	FShaderParameter				InvSize;
+	FShaderParameter				InputUvFactorAndOffset;
+	FShaderParameter				InputUvBundaries;
+	FPostProcessPassParameters		PostprocessParameter;
+	FShaderResourceParameter		ColorMip;
+	FShaderResourceParameter		ColorMipSampler;
+
+	TPostProcessBuildHCBPS(const ShaderMetaType::CompiledShaderInitializerType& Initializer)
+		: FGlobalShader(Initializer)
+	{
+		InvSize.Bind( Initializer.ParameterMap, TEXT("InvSize") );
+		InputUvFactorAndOffset.Bind( Initializer.ParameterMap, TEXT("InputUvFactorAndOffset") );
+		InputUvBundaries.Bind( Initializer.ParameterMap, TEXT("InputUvBundaries") );
+		PostprocessParameter.Bind( Initializer.ParameterMap );
+		ColorMip.Bind( Initializer.ParameterMap, TEXT("ColorMip") );
+		ColorMipSampler.Bind( Initializer.ParameterMap, TEXT("ColorMipSampler") );
+	}
+
+	void SetParameters(const FRenderingCompositePassContext& Context)
+	{
+		const FPixelShaderRHIParamRef ShaderRHI = GetPixelShader();
+		const FSceneRenderTargets& SceneContext = FSceneRenderTargets::Get(Context.RHICmdList);
+
+		FGlobalShader::SetParameters(Context.RHICmdList, ShaderRHI, Context.View );
+		
+		const FIntPoint GBufferSize = SceneContext.GetBufferSizeXY();
+		const FVector InvSizeValue( 1.0f / float(GBufferSize.X), 1.0f / float(GBufferSize.Y), 0.0 );
+		const FVector4 InputUvFactorAndOffsetValue (
+			float(2 * Context.View.HZBMipmap0Size.X) / float(GBufferSize.X),
+			float(2 * Context.View.HZBMipmap0Size.Y) / float(GBufferSize.Y),
+			float(Context.View.ViewRect.Min.X) / float(GBufferSize.X),
+			float(Context.View.ViewRect.Min.Y) / float(GBufferSize.Y)
+			);
+		const FVector4 InputUvBundariesValue (
+			float(Context.View.ViewRect.Min.X) / float(GBufferSize.X) + 0.5f * InvSizeValue.X,
+			float(Context.View.ViewRect.Min.Y) / float(GBufferSize.Y) + 0.5f * InvSizeValue.Y,
+			float(Context.View.ViewRect.Max.X) / float(GBufferSize.X) - 0.5f * InvSizeValue.X,
+			float(Context.View.ViewRect.Max.Y) / float(GBufferSize.Y) - 0.5f * InvSizeValue.Y
+			);
+		SetShaderValue(Context.RHICmdList, ShaderRHI, InvSize, InvSizeValue );
+		SetShaderValue(Context.RHICmdList, ShaderRHI, InputUvFactorAndOffset, InputUvFactorAndOffsetValue );
+		SetShaderValue(Context.RHICmdList, ShaderRHI, InputUvBundaries, InputUvBundariesValue );
+		
+		PostprocessParameter.SetPS(ShaderRHI, Context, TStaticSamplerState<SF_Bilinear, AM_Clamp, AM_Clamp, AM_Clamp>::GetRHI());
+	}
+
+	void SetParameters(FRHICommandList& RHICmdList, const FViewInfo& View, const FIntPoint& Size, FShaderResourceViewRHIParamRef ShaderResourceView )
+	{
+		const FPixelShaderRHIParamRef ShaderRHI = GetPixelShader();
+
+		FGlobalShader::SetParameters(RHICmdList, ShaderRHI, View );
+		
+		const FVector2D InvSizeValue( 1.0f / Size.X, 1.0f / Size.Y );
+		const FVector4 InputUvBundariesValue (
+			0.0f,
+			0.0f,
+			float(View.ViewRect.Size().X) / float(2 * View.HZBMipmap0Size.X) - 0.5f * InvSizeValue.X,
+			float(View.ViewRect.Size().Y) / float(2 * View.HZBMipmap0Size.Y) - 0.5f * InvSizeValue.Y
+			);
+		SetShaderValue(RHICmdList, ShaderRHI, InvSize, InvSizeValue );
+		SetShaderValue(RHICmdList, ShaderRHI, InputUvBundaries, InputUvBundariesValue );
+
+		SetSRVParameter(RHICmdList, ShaderRHI, ColorMip, ShaderResourceView );
+		SetSamplerParameter(RHICmdList, ShaderRHI, ColorMipSampler, TStaticSamplerState<SF_Bilinear,AM_Clamp,AM_Clamp,AM_Clamp>::GetRHI() );
+	}
+
+	virtual bool Serialize(FArchive& Ar) override
+	{
+		bool bShaderHasOutdatedParameters = FGlobalShader::Serialize(Ar);
+		Ar << InvSize;
+		Ar << InputUvFactorAndOffset;
+		Ar << InputUvBundaries;
+		Ar << PostprocessParameter;
+		Ar << ColorMip;
+		Ar << ColorMipSampler;
+		return bShaderHasOutdatedParameters;
+	}
+};
+
+IMPLEMENT_SHADER_TYPE(template<>,TPostProcessBuildHCBPS<0>,TEXT("PostProcessHierarchical"),TEXT("BuildHCB"),SF_Pixel);
+IMPLEMENT_SHADER_TYPE(template<>,TPostProcessBuildHCBPS<1>,TEXT("PostProcessHierarchical"),TEXT("BuildHCB"),SF_Pixel);
+
+static void HierarchycalSizeAndMips(const FIntPoint& BufferSize, uint32 Mip0Downsample, uint32& OutNumMips, FIntPoint& OutHierachicalSize)
+{
+	const uint32 NumMipsX = FPlatformMath::CeilLogTwo(BufferSize.X) - Mip0Downsample;
+	const uint32 NumMipsY = FPlatformMath::CeilLogTwo(BufferSize.Y) - Mip0Downsample;
+
+	OutNumMips = FMath::Max(NumMipsX, NumMipsY);
+	OutHierachicalSize = FIntPoint(1 << NumMipsX, 1 << NumMipsY);
+}
+
+void FRCPassPostProcessBuildHCB::Process(FRenderingCompositePassContext& Context)
+{
+	FIntPoint HCBSize;
+	uint32 HCBMipCount;
+	HierarchycalSizeAndMips(FSceneRenderTargets::Get_FrameConstantsOnly().GetBufferSizeXY(), 1, HCBMipCount, HCBSize);
+	
+	auto& RHICmdList = Context.RHICmdList;
+	FSceneRenderTargets& SceneContext = FSceneRenderTargets::Get(RHICmdList);
+
+	const FViewInfo& View = Context.View;
+	const auto FeatureLevel = Context.GetFeatureLevel();
+
+	SCOPED_DRAW_EVENT(RHICmdList, BuildHCB);
+		
+	FSceneRenderTargetItem HCBRenderTarget = PassOutputs[0].RequestSurface(Context);
+	const FVector2D HCBUsedPercentage (
+		float(View.ViewRect.Size().X) / float(2 * HCBSize.X),
+		float(View.ViewRect.Size().Y) / float(2 * HCBSize.Y)
+		);
+
+	RHICmdList.SetDepthStencilState(TStaticDepthStencilState<false, CF_Always>::GetRHI());
+	RHICmdList.SetRasterizerState(TStaticRasterizerState<>::GetRHI());
+	RHICmdList.SetBlendState(TStaticBlendState<>::GetRHI());
+
+	{ // mip 0
+		TShaderMapRef< FPostProcessVS > VertexShader(Context.GetShaderMap());
+		TShaderMapRef< TPostProcessBuildHCBPS<0> > PixelShader(Context.GetShaderMap());
+		
+		static FGlobalBoundShaderState BoundShaderState;
+		SetGlobalBoundShaderState(Context.RHICmdList, FeatureLevel, BoundShaderState, GFilterVertexDeclaration.VertexDeclarationRHI, *VertexShader, *PixelShader);
+		VertexShader->SetParameters(Context);
+		PixelShader->SetParameters(Context);
+			
+		SetRenderTarget(RHICmdList, HCBRenderTarget.TargetableTexture, 0, NULL);
+			
+#if 0
+		const FIntPoint DstDrawSize (
+			FMath::Min(FMath::CeilToInt(HCBUsedPercentage.X * float(HCBSize.X)) + 1, HCBSize.X),
+			FMath::Min(FMath::CeilToInt(HCBUsedPercentage.Y * float(HCBSize.Y)) + 1, HCBSize.Y)
+			);
+
+		RHICmdList.SetViewport(0, 0, 0.0f, DstDrawSize.X, DstDrawSize.Y, 1.0f);
+			
+		DrawRectangle(
+			RHICmdList,
+			0, 0,
+			HCBSize.X, HCBSize.Y,
+			View.ViewRect.Min.X, View.ViewRect.Min.Y,
+			View.ViewRect.Width(), View.ViewRect.Height(),
+			DstDrawSize,
+			SceneContext.GetBufferSizeXY(),
+			*VertexShader,
+			EDRF_UseTriangleOptimization);
+#else
+		RHICmdList.SetViewport(0, 0, 0.0f, HCBSize.X, HCBSize.Y, 1.0f);
+			
+		DrawRectangle(
+			RHICmdList,
+			0, 0,
+			HCBSize.X, HCBSize.Y,
+			View.ViewRect.Min.X, View.ViewRect.Min.Y,
+			View.ViewRect.Width(), View.ViewRect.Height(),
+			HCBSize,
+			SceneContext.GetBufferSizeXY(),
+			*VertexShader,
+			EDRF_UseTriangleOptimization);
+#endif
+
+		RHICmdList.CopyToResolveTarget(HCBRenderTarget.TargetableTexture, HCBRenderTarget.ShaderResourceTexture, false, FResolveParams(FResolveRect(), CubeFace_PosX, 0));
+	}
+		
+	FIntPoint SrcSize = HCBSize;
+	FIntPoint DstSize = SrcSize / 2;
+		
+	// Downsampling...
+	for( uint8 MipIndex = 1; MipIndex < HCBMipCount; MipIndex++ )
+	{
+		DstSize.X = FMath::Max(DstSize.X, 1);
+		DstSize.Y = FMath::Max(DstSize.Y, 1);
+
+		SetRenderTarget(RHICmdList, HCBRenderTarget.TargetableTexture, MipIndex, NULL);
+
+		TShaderMapRef< FPostProcessVS >	VertexShader(Context.GetShaderMap());
+		TShaderMapRef< TPostProcessBuildHCBPS<1> >	PixelShader(Context.GetShaderMap());
+
+		static FGlobalBoundShaderState BoundShaderState;
+		
+		SetGlobalBoundShaderState(RHICmdList, View.GetFeatureLevel(), BoundShaderState, GFilterVertexDeclaration.VertexDeclarationRHI, *VertexShader, *PixelShader);
+
+		PixelShader->SetParameters(RHICmdList, View, SrcSize, HCBRenderTarget.MipSRVs[ MipIndex - 1 ] );
+			
+#if 0
+		const FIntPoint DstDrawSize (
+			FMath::Min(FMath::CeilToInt(HCBUsedPercentage.X * float(DstSize.X)) + 1, DstSize.X),
+			FMath::Min(FMath::CeilToInt(HCBUsedPercentage.Y * float(DstSize.Y)) + 1, DstSize.Y)
+			);
+
+		RHICmdList.SetViewport(0, 0, 0.0f, DstDrawSize.X, DstDrawSize.Y, 1.0f);
+
+		DrawRectangle(
+			RHICmdList,
+			0, 0,
+			DstSize.X, DstSize.Y,
+			0, 0,
+			SrcSize.X, SrcSize.Y,
+			DstDrawSize,
+			SrcSize,
+			*VertexShader,
+			EDRF_UseTriangleOptimization);
+#else
+		RHICmdList.SetViewport(0, 0, 0.0f, DstSize.X, DstSize.Y, 1.0f);
+
+		DrawRectangle(
+			RHICmdList,
+			0, 0,
+			DstSize.X, DstSize.Y,
+			0, 0,
+			SrcSize.X, SrcSize.Y,
+			DstSize,
+			SrcSize,
+			*VertexShader,
+			EDRF_UseTriangleOptimization);
+#endif
+
+		RHICmdList.CopyToResolveTarget(HCBRenderTarget.TargetableTexture, HCBRenderTarget.ShaderResourceTexture, false, FResolveParams(FResolveRect(), CubeFace_PosX, MipIndex));
+
+		SrcSize /= 2;
+		DstSize /= 2;
+	}
+}
+
+FPooledRenderTargetDesc FRCPassPostProcessBuildHCB::ComputeOutputDesc(EPassOutputId InPassOutputId) const
+{
+	FIntPoint HCBSize;
+	uint32 HCBMipCount;
+	HierarchycalSizeAndMips(FSceneRenderTargets::Get_FrameConstantsOnly().GetBufferSizeXY(), 1, HCBMipCount, HCBSize);
+	FPooledRenderTargetDesc Ret(FPooledRenderTargetDesc::Create2DDesc(HCBSize, PF_FloatRGBA, TexCreate_None, TexCreate_RenderTargetable | TexCreate_ShaderResource | TexCreate_NoFastClear, false, HCBMipCount));
+	
+	Ret.DebugName = TEXT("HCB");
+
+	return Ret;
+}
diff --git a/Engine/Source/Runtime/Renderer/Private/PostProcess/PostProcessHierarchical.h b/Engine/Source/Runtime/Renderer/Private/PostProcess/PostProcessHierarchical.h
new file mode 100644
--- /dev/null
+++ b/Engine/Source/Runtime/Renderer/Private/PostProcess/PostProcessHierarchical.h
@@ -0,0 +1,20 @@
+// Copyright 1998-2015 Epic Games, Inc. All Rights Reserved.
+
+/*=============================================================================
+	PostProcessHierarchical.h: Build hierarchical buffers
+=============================================================================*/
+
+#pragma once
+
+#include "RenderingCompositionGraph.h"
+
+// ePId_Input0: scene color
+// derives from TRenderingCompositePassBase<InputCount, OutputCount> 
+class FRCPassPostProcessBuildHCB : public TRenderingCompositePassBase<1, 1>
+{
+public:
+	// interface FRenderingCompositePass ---------
+	virtual void Process(FRenderingCompositePassContext& Context) override;
+	virtual void Release() override { delete this; }
+	virtual FPooledRenderTargetDesc ComputeOutputDesc(EPassOutputId InPassOutputId) const override;
+};
diff --git a/Engine/Source/Runtime/Renderer/Private/PostProcess/ScreenSpaceReflections.cpp b/Engine/Source/Runtime/Renderer/Private/PostProcess/ScreenSpaceReflections.cpp
--- a/Engine/Source/Runtime/Renderer/Private/PostProcess/ScreenSpaceReflections.cpp
+++ b/Engine/Source/Runtime/Renderer/Private/PostProcess/ScreenSpaceReflections.cpp
@@ -11,6 +11,7 @@
 #include "ScreenSpaceReflections.h"
 #include "PostProcessTemporalAA.h"
 #include "PostProcessAmbientOcclusion.h"
+#include "PostProcessHierarchical.h"
 #include "SceneUtils.h"
 
 static TAutoConsoleVariable<int32> CVarSSRQuality(
@@ -40,6 +41,13 @@ static TAutoConsoleVariable<int32> CVarSSRStencil(
 	TEXT(" 0 is off (default), 1 is on"),
 	ECVF_RenderThreadSafe);
 
+static TAutoConsoleVariable<int32> CVarSSRCone(
+	TEXT("r.SSR.Cone"),
+	0,
+	TEXT("Defines if we use cone traced screen space reflection\n")
+	TEXT(" 0 is off (default), 1 is on"),
+	ECVF_RenderThreadSafe);
+
 bool DoScreenSpaceReflections(const FViewInfo& View)
 {
 	if(!View.Family->EngineShowFlags.ScreenSpaceReflections)
@@ -149,6 +157,7 @@ public:
 };
 
 IMPLEMENT_SHADER_TYPE(,FPostProcessScreenSpaceReflectionsStencilPS,TEXT("ScreenSpaceReflections"),TEXT("ScreenSpaceReflectionsStencilPS"),SF_Pixel);
+static const uint32 SSRConeQuality = 5;
 
 /**
  * Encapsulates the post processing screen space reflections pixel shader.
@@ -169,6 +178,7 @@ class FPostProcessScreenSpaceReflectionsPS : public FGlobalShader
 		FGlobalShader::ModifyCompilationEnvironment(Platform, OutEnvironment);
 		OutEnvironment.SetDefine( TEXT("PREV_FRAME_COLOR"), PrevFrame );
 		OutEnvironment.SetDefine( TEXT("SSR_QUALITY"), SSRQuality );
+		OutEnvironment.SetDefine( TEXT("SSR_CONE_QUALITY"), SSRConeQuality );
 	}
 
 	/** Default constructor. */
@@ -251,6 +261,8 @@ IMPLEMENT_REFLECTION_PIXELSHADER_TYPE(0,3);
 IMPLEMENT_REFLECTION_PIXELSHADER_TYPE(1,3);
 IMPLEMENT_REFLECTION_PIXELSHADER_TYPE(0,4);
 IMPLEMENT_REFLECTION_PIXELSHADER_TYPE(1,4);
+IMPLEMENT_REFLECTION_PIXELSHADER_TYPE(0,5); // SSRConeQuality
+IMPLEMENT_REFLECTION_PIXELSHADER_TYPE(1,5); // SSRConeQuality
 
 // --------------------------------------------------------
 
@@ -288,17 +300,22 @@ void FRCPassPostProcessScreenSpaceReflections::Process(FRenderingCompositePassCo
 
 	SSRQuality = FMath::Clamp(SSRQuality, 1, 4);
 	
-	if (View.Family->EngineShowFlags.VisualizeSSR)
+	const bool VisualizeSSR = View.Family->EngineShowFlags.VisualizeSSR;
+	const bool SSRStencilPrePass = CVarSSRStencil.GetValueOnRenderThread() != 0 && !VisualizeSSR;
+	const bool SSRConeTracing = GetInput(ePId_Input2)->GetOutput() != nullptr;
+	
+	if (VisualizeSSR)
 	{
 		iPreFrame = 0;
 		SSRQuality = 0;
 	}
-	
-	// Write 1 to highest bit of stencil to areas that should compute SSR
-	bool SSRStencilPrePass = CVarSSRStencil.GetValueOnRenderThread() != 0 && SSRQuality != 0;
+	else if (SSRConeTracing)
+	{
+		SSRQuality = SSRConeQuality;
+	}
 	
 	const FSceneRenderTargetItem& DestRenderTarget = PassOutputs[0].RequestSurface(Context);
-
+	
 	if (SSRStencilPrePass)
 	{ // ScreenSpaceReflectionsStencil draw event
 		SCOPED_DRAW_EVENT(RHICmdList, ScreenSpaceReflectionsStencil);
@@ -350,13 +367,13 @@ void FRCPassPostProcessScreenSpaceReflections::Process(FRenderingCompositePassCo
 			RHICmdList.SetDepthStencilState(TStaticDepthStencilState<false, CF_Always, true, CF_Equal, SO_Keep, SO_Keep, SO_Keep>::GetRHI(), 0);
 		}
 		else
-	{
+		{
 			// bind only the dest render target
 			SetRenderTarget(RHICmdList, DestRenderTarget.TargetableTexture, FTextureRHIRef());
 			Context.SetViewportAndCallRHI(View.ViewRect);
 
 			RHICmdList.SetDepthStencilState(TStaticDepthStencilState<false, CF_Always>::GetRHI());
-	}
+		}
 
 		// clear DestRenderTarget only outside of the view's rectangle
 		RHICmdList.Clear(true, FLinearColor::Black, false, (float)ERHIZBuffer::FarPlane, false, 0, View.ViewRect);
@@ -365,43 +382,44 @@ void FRCPassPostProcessScreenSpaceReflections::Process(FRenderingCompositePassCo
 		RHICmdList.SetBlendState(TStaticBlendState<>::GetRHI());
 		RHICmdList.SetRasterizerState(TStaticRasterizerState<>::GetRHI());
 
-	TShaderMapRef< FPostProcessVS > VertexShader(Context.GetShaderMap());
+		TShaderMapRef< FPostProcessVS > VertexShader(Context.GetShaderMap());
 
-	#define CASE(A, B) \
-		case (A + 2 * (B + 3 * 0 )): \
-		{ \
-			TShaderMapRef< FPostProcessScreenSpaceReflectionsPS<A, B> > PixelShader(Context.GetShaderMap()); \
-			static FGlobalBoundShaderState BoundShaderState; \
+		#define CASE(A, B) \
+			case (A + 2 * (B + 3 * 0 )): \
+			{ \
+				TShaderMapRef< FPostProcessScreenSpaceReflectionsPS<A, B> > PixelShader(Context.GetShaderMap()); \
+				static FGlobalBoundShaderState BoundShaderState; \
 				SetGlobalBoundShaderState(RHICmdList, FeatureLevel, BoundShaderState, GFilterVertexDeclaration.VertexDeclarationRHI, *VertexShader, *PixelShader); \
-			VertexShader->SetParameters(Context); \
-			PixelShader->SetParameters(Context); \
-		}; \
-		break
+				VertexShader->SetParameters(Context); \
+				PixelShader->SetParameters(Context); \
+			}; \
+			break
 
-	switch (iPreFrame + 2 * (SSRQuality + 3 * 0))
-	{
-		CASE(0,0);
-		CASE(0,1);	CASE(1,1);
-		CASE(0,2);	CASE(1,2);
-		CASE(0,3);	CASE(1,3);
-		CASE(0,4);	CASE(1,4);
-		default:
-			check(!"Missing case in FRCPassPostProcessScreenSpaceReflections");
-	}
-	#undef CASE
+		switch (iPreFrame + 2 * (SSRQuality + 3 * 0))
+		{
+			CASE(0,0);
+			CASE(0,1);	CASE(1,1);
+			CASE(0,2);	CASE(1,2);
+			CASE(0,3);	CASE(1,3);
+			CASE(0,4);	CASE(1,4);
+			CASE(0,5);	CASE(1,5); //SSRConeQuality
+			default:
+				check(!"Missing case in FRCPassPostProcessScreenSpaceReflections");
+		}
+		#undef CASE
 
 
-	// Draw a quad mapping scene color to the view's render target
-	DrawRectangle( 
+		// Draw a quad mapping scene color to the view's render target
+		DrawRectangle( 
 			RHICmdList,
-		0, 0,
-		View.ViewRect.Width(), View.ViewRect.Height(),
-		View.ViewRect.Min.X, View.ViewRect.Min.Y, 
-		View.ViewRect.Width(), View.ViewRect.Height(),
-		View.ViewRect.Size(),
-		FSceneRenderTargets::Get(Context.RHICmdList).GetBufferSizeXY(),
-		*VertexShader,
-		EDRF_UseTriangleOptimization);
+			0, 0,
+			View.ViewRect.Width(), View.ViewRect.Height(),
+			View.ViewRect.Min.X, View.ViewRect.Min.Y, 
+			View.ViewRect.Width(), View.ViewRect.Height(),
+			View.ViewRect.Size(),
+			FSceneRenderTargets::Get(Context.RHICmdList).GetBufferSizeXY(),
+			*VertexShader,
+			EDRF_UseTriangleOptimization);
 
 		RHICmdList.CopyToResolveTarget(DestRenderTarget.TargetableTexture, DestRenderTarget.ShaderResourceTexture, false, FResolveParams());
 	} // ScreenSpaceReflections
@@ -425,6 +443,7 @@ void ScreenSpaceReflections(FRHICommandListImmediate& RHICmdList, FViewInfo& Vie
 
 	FRenderingCompositePass* SceneColorInput = Context.Graph.RegisterPass( new FRCPassPostProcessInput( FSceneRenderTargets::Get(RHICmdList).GetSceneColor() ) );
 	FRenderingCompositePass* HZBInput = Context.Graph.RegisterPass( new FRCPassPostProcessInput( View.HZB ) );
+	FRenderingCompositePass* HCBInput = nullptr;
 
 	bool bPrevFrame = 0;
 	if( ViewState && ViewState->TemporalAAHistoryRT && !Context.View.bCameraCut )
@@ -432,11 +451,18 @@ void ScreenSpaceReflections(FRHICommandListImmediate& RHICmdList, FViewInfo& Vie
 		SceneColorInput = Context.Graph.RegisterPass( new FRCPassPostProcessInput( ViewState->TemporalAAHistoryRT ) );
 		bPrevFrame = 1;
 	}
+	
+	if (CVarSSRCone.GetValueOnRenderThread())
+	{
+		HCBInput = Context.Graph.RegisterPass( new FRCPassPostProcessBuildHCB() );
+		HCBInput->SetInput( ePId_Input0, SceneColorInput );
+	}
 
 	{
 		FRenderingCompositePass* TracePass = Context.Graph.RegisterPass( new FRCPassPostProcessScreenSpaceReflections( bPrevFrame ) );
 		TracePass->SetInput( ePId_Input0, SceneColorInput );
 		TracePass->SetInput( ePId_Input1, HZBInput );
+		TracePass->SetInput( ePId_Input2, HCBInput );
 
 		Context.FinalOutput = FRenderingCompositeOutputRef( TracePass );
 	}
diff --git a/Engine/Source/Runtime/Renderer/Private/PostProcess/ScreenSpaceReflections.h b/Engine/Source/Runtime/Renderer/Private/PostProcess/ScreenSpaceReflections.h
--- a/Engine/Source/Runtime/Renderer/Private/PostProcess/ScreenSpaceReflections.h
+++ b/Engine/Source/Runtime/Renderer/Private/PostProcess/ScreenSpaceReflections.h
@@ -18,10 +18,11 @@ public:
 	virtual FPooledRenderTargetDesc ComputeOutputDesc(EPassOutputId InPassOutputId) const override;
 };
 
-// ePId_Input0: half res scene color
-// ePId_Input1: half res scene depth
+// ePId_Input0: scene color
+// ePId_Input1: scene depth
+// ePId_Input2: hierarchical scene color (optional)
 // derives from TRenderingCompositePassBase<InputCount, OutputCount> 
-class FRCPassPostProcessScreenSpaceReflections : public TRenderingCompositePassBase<2, 1>
+class FRCPassPostProcessScreenSpaceReflections : public TRenderingCompositePassBase<3, 1>
 {
 public:
 	FRCPassPostProcessScreenSpaceReflections( bool InPrevFrame )
