fixed UE-17239 Stationary light shadow are occluded for additional players in split-screen mode

diff --git a/Engine/Source/Runtime/Renderer/Private/ShadowSetup.cpp b/Engine/Source/Runtime/Renderer/Private/ShadowSetup.cpp
--- a/Engine/Source/Runtime/Renderer/Private/ShadowSetup.cpp
+++ b/Engine/Source/Runtime/Renderer/Private/ShadowSetup.cpp
@@ -936,18 +936,19 @@ void FProjectedShadowInfo::GatherDynamicMeshElements(FSceneRenderer& Renderer, F
 		FViewInfo* FoundView = NULL;
 		bool bVisibleInAnyView = false;
 
-		for (int32 ViewIndex = 0;ViewIndex < Renderer.Views.Num();ViewIndex++)
+		for (int32 ViewIndex = 0, Count = Renderer.Views.Num(); ViewIndex < Count; ViewIndex++)
 		{
-			FViewInfo* CheckView = &Renderer.Views[ViewIndex];
+			const FViewInfo* CheckView = &Renderer.Views[ViewIndex];
 			const FVisibleLightViewInfo& VisibleLightViewInfo = CheckView->VisibleLightInfos[LightSceneInfo->Id];
 			FPrimitiveViewRelevance ViewRel = VisibleLightViewInfo.ProjectedShadowViewRelevanceMap[ShadowId];
 
+			// we need to iterate through all
 			bVisibleInAnyView |= VisibleLightViewInfo.ProjectedShadowVisibilityMap[ShadowId];
 
 			if (ViewRel.bShadowRelevance)
 			{
 				FoundView = CheckView;
-				break;
+				// we could ealy out for this but it would break bVisibleInAnyView
 			}
 		}
 
