fixed compile error

diff --git a/Engine/Source/Runtime/Renderer/Private/ShadowSetup.cpp b/Engine/Source/Runtime/Renderer/Private/ShadowSetup.cpp
--- a/Engine/Source/Runtime/Renderer/Private/ShadowSetup.cpp
+++ b/Engine/Source/Runtime/Renderer/Private/ShadowSetup.cpp
@@ -938,7 +938,7 @@ void FProjectedShadowInfo::GatherDynamicMeshElements(FSceneRenderer& Renderer, F
 
 		for (int32 ViewIndex = 0, Count = Renderer.Views.Num(); ViewIndex < Count; ViewIndex++)
 		{
-			const FViewInfo* CheckView = &Renderer.Views[ViewIndex];
+			FViewInfo* CheckView = &Renderer.Views[ViewIndex];
 			const FVisibleLightViewInfo& VisibleLightViewInfo = CheckView->VisibleLightInfos[LightSceneInfo->Id];
 			FPrimitiveViewRelevance ViewRel = VisibleLightViewInfo.ProjectedShadowViewRelevanceMap[ShadowId];
 
