* fixed UE-18117 LightGrid not fully disabled by CVAR

diff --git a/Engine/Source/Runtime/Renderer/Private/BasePassRendering.h b/Engine/Source/Runtime/Renderer/Private/BasePassRendering.h
--- a/Engine/Source/Runtime/Renderer/Private/BasePassRendering.h
+++ b/Engine/Source/Runtime/Renderer/Private/BasePassRendering.h
@@ -400,14 +400,19 @@ public:
 				{
 					check(GetUniformBufferParameter<FForwardLightData>().IsInitialized());
 
-					if(GetUniformBufferParameter<FForwardLightData>().IsBound())
+					static const auto CVar = IConsoleManager::Get().FindTConsoleVariableDataInt(TEXT("r.ForwardLighting"));
+					bool bForwardLighting = CVar->GetValueOnRenderThread() != 0;
+
+					if(bForwardLighting)
 					{
 						SetUniformBufferParameter(RHICmdList, ShaderRHI,GetUniformBufferParameter<FForwardLightData>(),View->ForwardLightData);
+						SetSRVParameter(RHICmdList, ShaderRHI, LightGrid, GLightGridVertexBuffer.VertexBufferSRV);
 					}
-
-					if(LightGrid.IsBound())
+					else
 					{
-						RHICmdList.SetShaderResourceViewParameter(ShaderRHI, LightGrid.GetBaseIndex(), GLightGridVertexBuffer.VertexBufferSRV);
+						// feature is disabled
+						SetUniformBufferParameter(RHICmdList, ShaderRHI,GetUniformBufferParameter<FForwardLightData>(), 0);
+						SetSRVParameter(RHICmdList, ShaderRHI, LightGrid, 0);
 					}
 				}
 			}
