Force to update cached single sample data when new light volume is added/removed from the scene UE-17784

diff --git a/Engine/Source/Runtime/Renderer/Private/IndirectLightingCache.cpp b/Engine/Source/Runtime/Renderer/Private/IndirectLightingCache.cpp
--- a/Engine/Source/Runtime/Renderer/Private/IndirectLightingCache.cpp
+++ b/Engine/Source/Runtime/Renderer/Private/IndirectLightingCache.cpp
@@ -603,6 +603,7 @@ void FIndirectLightingCache::SetLightingCacheDirty()
 	for (TMap<FPrimitiveComponentId, FIndirectLightingCacheAllocation*>::TIterator It(PrimitiveAllocations); It; ++It)
 	{
 		It.Value()->bIsDirty = true;
+		It.Value()->bHasEverUpdatedSingleSample = false;
 	}
 	
 	// next rendering we update all entries no matter if they are visible to avoid further hitches
