changed all direct access to PassOutputs to GetInput() to allow internal restructuring

diff --git a/Engine/Source/Runtime/Renderer/Private/CompositionLighting/PostProcessAmbientOcclusion.cpp b/Engine/Source/Runtime/Renderer/Private/CompositionLighting/PostProcessAmbientOcclusion.cpp
--- a/Engine/Source/Runtime/Renderer/Private/CompositionLighting/PostProcessAmbientOcclusion.cpp
+++ b/Engine/Source/Runtime/Renderer/Private/CompositionLighting/PostProcessAmbientOcclusion.cpp
@@ -362,11 +362,11 @@ FPooledRenderTargetDesc FRCPassPostProcessAmbientOcclusionSetup::ComputeOutputDe
 	
 	if(IsInitialPass())
 	{
-		Ret = PassInputs[0].GetOutput()->RenderTargetDesc;
+		Ret = GetInput(ePId_Input0)->GetOutput()->RenderTargetDesc;
 	}
 	else
 	{
-		Ret = PassInputs[1].GetOutput()->RenderTargetDesc;
+		Ret = GetInput(ePId_Input1)->GetOutput()->RenderTargetDesc;
 	}
 
 	Ret.Reset();
@@ -524,7 +524,7 @@ FPooledRenderTargetDesc FRCPassPostProcessAmbientOcclusion::ComputeOutputDesc(EP
 		return Ret;
 	}
 
-	FPooledRenderTargetDesc Ret = PassInputs[0].GetOutput()->RenderTargetDesc;
+	FPooledRenderTargetDesc Ret = GetInput(ePId_Input0)->GetOutput()->RenderTargetDesc;
 
 	Ret.Reset();
 	// R:AmbientOcclusion, GBA:used for normal
diff --git a/Engine/Source/Runtime/Renderer/Private/CompositionLighting/PostProcessLpvIndirect.cpp b/Engine/Source/Runtime/Renderer/Private/CompositionLighting/PostProcessLpvIndirect.cpp
--- a/Engine/Source/Runtime/Renderer/Private/CompositionLighting/PostProcessLpvIndirect.cpp
+++ b/Engine/Source/Runtime/Renderer/Private/CompositionLighting/PostProcessLpvIndirect.cpp
@@ -496,7 +496,7 @@ void FRCPassPostProcessVisualizeLPV::Process(FRenderingCompositePassContext& Con
 
 FPooledRenderTargetDesc FRCPassPostProcessVisualizeLPV::ComputeOutputDesc(EPassOutputId InPassOutputId) const
 {
-	FPooledRenderTargetDesc Ret = PassInputs[0].GetOutput()->RenderTargetDesc;
+	FPooledRenderTargetDesc Ret = GetInput(ePId_Input0)->GetOutput()->RenderTargetDesc;
 
 	Ret.Reset();
 
diff --git a/Engine/Source/Runtime/Renderer/Private/CompositionLighting/PostProcessPassThrough.cpp b/Engine/Source/Runtime/Renderer/Private/CompositionLighting/PostProcessPassThrough.cpp
--- a/Engine/Source/Runtime/Renderer/Private/CompositionLighting/PostProcessPassThrough.cpp
+++ b/Engine/Source/Runtime/Renderer/Private/CompositionLighting/PostProcessPassThrough.cpp
@@ -151,7 +151,7 @@ FPooledRenderTargetDesc FRCPassPostProcessPassThrough::ComputeOutputDesc(EPassOu
 	// we assume this pass is additively blended with the scene color so an intermediate is not always needed
 	if(!Dest)
 	{
-		Ret = PassInputs[0].GetOutput()->RenderTargetDesc;
+		Ret = GetInput(ePId_Input0)->GetOutput()->RenderTargetDesc;
 
 		if(NewDesc.IsValid())
 		{
diff --git a/Engine/Source/Runtime/Renderer/Private/PostProcess/PostProcessAA.cpp b/Engine/Source/Runtime/Renderer/Private/PostProcess/PostProcessAA.cpp
--- a/Engine/Source/Runtime/Renderer/Private/PostProcess/PostProcessAA.cpp
+++ b/Engine/Source/Runtime/Renderer/Private/PostProcess/PostProcessAA.cpp
@@ -290,7 +290,7 @@ void FRCPassPostProcessAA::Process(FRenderingCompositePassContext& Context)
 
 FPooledRenderTargetDesc FRCPassPostProcessAA::ComputeOutputDesc(EPassOutputId InPassOutputId) const
 {
-	FPooledRenderTargetDesc Ret = PassInputs[0].GetOutput()->RenderTargetDesc;
+	FPooledRenderTargetDesc Ret = GetInput(ePId_Input0)->GetOutput()->RenderTargetDesc;
 
 	Ret.Reset();
 	Ret.DebugName = TEXT("PostProcessAA");
diff --git a/Engine/Source/Runtime/Renderer/Private/PostProcess/PostProcessBloomSetup.cpp b/Engine/Source/Runtime/Renderer/Private/PostProcess/PostProcessBloomSetup.cpp
--- a/Engine/Source/Runtime/Renderer/Private/PostProcess/PostProcessBloomSetup.cpp
+++ b/Engine/Source/Runtime/Renderer/Private/PostProcess/PostProcessBloomSetup.cpp
@@ -203,7 +203,7 @@ void FRCPassPostProcessBloomSetup::Process(FRenderingCompositePassContext& Conte
 
 FPooledRenderTargetDesc FRCPassPostProcessBloomSetup::ComputeOutputDesc(EPassOutputId InPassOutputId) const
 {
-	FPooledRenderTargetDesc Ret = PassInputs[0].GetOutput()->RenderTargetDesc;
+	FPooledRenderTargetDesc Ret = GetInput(ePId_Input0)->GetOutput()->RenderTargetDesc;
 
 	Ret.Reset();
 	Ret.DebugName = TEXT("BloomSetup");
@@ -329,7 +329,7 @@ void FRCPassPostProcessVisualizeBloomSetup::Process(FRenderingCompositePassConte
 
 FPooledRenderTargetDesc FRCPassPostProcessVisualizeBloomSetup::ComputeOutputDesc(EPassOutputId InPassOutputId) const
 {
-	FPooledRenderTargetDesc Ret = PassInputs[0].GetOutput()->RenderTargetDesc;
+	FPooledRenderTargetDesc Ret = GetInput(ePId_Input0)->GetOutput()->RenderTargetDesc;
 
 	Ret.Reset();
 	Ret.TargetableFlags &= ~TexCreate_UAV;
@@ -467,7 +467,7 @@ void FRCPassPostProcessVisualizeBloomOverlay::Process(FRenderingCompositePassCon
 
 FPooledRenderTargetDesc FRCPassPostProcessVisualizeBloomOverlay::ComputeOutputDesc(EPassOutputId InPassOutputId) const
 {
-	FPooledRenderTargetDesc Ret = PassInputs[0].GetOutput()->RenderTargetDesc;
+	FPooledRenderTargetDesc Ret = GetInput(ePId_Input0)->GetOutput()->RenderTargetDesc;
 
 	Ret.Reset();
 	Ret.TargetableFlags &= ~TexCreate_UAV;
diff --git a/Engine/Source/Runtime/Renderer/Private/PostProcess/PostProcessBokehDOF.cpp b/Engine/Source/Runtime/Renderer/Private/PostProcess/PostProcessBokehDOF.cpp
--- a/Engine/Source/Runtime/Renderer/Private/PostProcess/PostProcessBokehDOF.cpp
+++ b/Engine/Source/Runtime/Renderer/Private/PostProcess/PostProcessBokehDOF.cpp
@@ -297,7 +297,7 @@ void FRCPassPostProcessVisualizeDOF::Process(FRenderingCompositePassContext& Con
 
 FPooledRenderTargetDesc FRCPassPostProcessVisualizeDOF::ComputeOutputDesc(EPassOutputId InPassOutputId) const
 {
-	FPooledRenderTargetDesc Ret = PassInputs[0].GetOutput()->RenderTargetDesc;
+	FPooledRenderTargetDesc Ret = GetInput(ePId_Input0)->GetOutput()->RenderTargetDesc;
 
 	Ret.Reset();
 	Ret.Extent /= 2;
@@ -455,7 +455,7 @@ void FRCPassPostProcessBokehDOFSetup::Process(FRenderingCompositePassContext& Co
 
 FPooledRenderTargetDesc FRCPassPostProcessBokehDOFSetup::ComputeOutputDesc(EPassOutputId InPassOutputId) const
 {
-	FPooledRenderTargetDesc Ret = PassInputs[0].GetOutput()->RenderTargetDesc;
+	FPooledRenderTargetDesc Ret = GetInput(ePId_Input0)->GetOutput()->RenderTargetDesc;
 
 	Ret.Reset();
 	Ret.Extent /= 2;
@@ -818,7 +818,7 @@ void FRCPassPostProcessBokehDOF::Process(FRenderingCompositePassContext& Context
 
 FPooledRenderTargetDesc FRCPassPostProcessBokehDOF::ComputeOutputDesc(EPassOutputId InPassOutputId) const
 {
-	FPooledRenderTargetDesc Ret = PassInputs[0].GetOutput()->RenderTargetDesc;
+	FPooledRenderTargetDesc Ret = GetInput(ePId_Input0)->GetOutput()->RenderTargetDesc;
 
 	Ret.Reset();
 	// more precision for additive blending
diff --git a/Engine/Source/Runtime/Renderer/Private/PostProcess/PostProcessBokehDOFRecombine.cpp b/Engine/Source/Runtime/Renderer/Private/PostProcess/PostProcessBokehDOFRecombine.cpp
--- a/Engine/Source/Runtime/Renderer/Private/PostProcess/PostProcessBokehDOFRecombine.cpp
+++ b/Engine/Source/Runtime/Renderer/Private/PostProcess/PostProcessBokehDOFRecombine.cpp
@@ -114,7 +114,9 @@ void FRCPassPostProcessBokehDOFRecombine::Process(FRenderingCompositePassContext
 {
 	uint32 Method = 2;
 
-	if(GetInput(ePId_Input1)->GetPass())
+	FRenderingCompositeOutputRef* Input1 = GetInput(ePId_Input1);
+
+	if(Input1 && Input1->GetPass())
 	{
 		if(GetInput(ePId_Input2)->GetPass())
 		{
@@ -187,7 +189,7 @@ void FRCPassPostProcessBokehDOFRecombine::Process(FRenderingCompositePassContext
 
 FPooledRenderTargetDesc FRCPassPostProcessBokehDOFRecombine::ComputeOutputDesc(EPassOutputId InPassOutputId) const
 {
-	FPooledRenderTargetDesc Ret = PassInputs[0].GetOutput()->RenderTargetDesc;
+	FPooledRenderTargetDesc Ret = GetInput(ePId_Input0)->GetOutput()->RenderTargetDesc;
 
 	Ret.Reset();
 	Ret.DebugName = TEXT("BokehDOFRecombine");
diff --git a/Engine/Source/Runtime/Renderer/Private/PostProcess/PostProcessCompositeEditorPrimitives.cpp b/Engine/Source/Runtime/Renderer/Private/PostProcess/PostProcessCompositeEditorPrimitives.cpp
--- a/Engine/Source/Runtime/Renderer/Private/PostProcess/PostProcessCompositeEditorPrimitives.cpp
+++ b/Engine/Source/Runtime/Renderer/Private/PostProcess/PostProcessCompositeEditorPrimitives.cpp
@@ -295,7 +295,7 @@ void FRCPassPostProcessCompositeEditorPrimitives::Process(FRenderingCompositePas
 
 FPooledRenderTargetDesc FRCPassPostProcessCompositeEditorPrimitives::ComputeOutputDesc(EPassOutputId InPassOutputId) const
 {
-	FPooledRenderTargetDesc Ret = PassInputs[0].GetOutput()->RenderTargetDesc;
+	FPooledRenderTargetDesc Ret = GetInput(ePId_Input0)->GetOutput()->RenderTargetDesc;
 
 	Ret.Reset();
 	Ret.DebugName = TEXT("EditorPrimitives");
diff --git a/Engine/Source/Runtime/Renderer/Private/PostProcess/PostProcessDOF.cpp b/Engine/Source/Runtime/Renderer/Private/PostProcess/PostProcessDOF.cpp
--- a/Engine/Source/Runtime/Renderer/Private/PostProcess/PostProcessDOF.cpp
+++ b/Engine/Source/Runtime/Renderer/Private/PostProcess/PostProcessDOF.cpp
@@ -175,7 +175,7 @@ void FRCPassPostProcessDOFSetup::Process(FRenderingCompositePassContext& Context
 
 FPooledRenderTargetDesc FRCPassPostProcessDOFSetup::ComputeOutputDesc(EPassOutputId InPassOutputId) const
 {
-	FPooledRenderTargetDesc Ret = PassInputs[0].GetOutput()->RenderTargetDesc;
+	FPooledRenderTargetDesc Ret = GetInput(ePId_Input0)->GetOutput()->RenderTargetDesc;
 
 //	Ret.Extent = FIntPoint::DivideAndRoundUp(ret.Extent, 2);
 	Ret.Extent /= 2;
@@ -340,7 +340,7 @@ void FRCPassPostProcessDOFRecombine::Process(FRenderingCompositePassContext& Con
 
 FPooledRenderTargetDesc FRCPassPostProcessDOFRecombine::ComputeOutputDesc(EPassOutputId InPassOutputId) const
 {
-	FPooledRenderTargetDesc Ret = PassInputs[0].GetOutput()->RenderTargetDesc;
+	FPooledRenderTargetDesc Ret = GetInput(ePId_Input0)->GetOutput()->RenderTargetDesc;
 
 	Ret.Reset();
 	Ret.DebugName = TEXT("DOFRecombine");
@@ -576,7 +576,7 @@ void FRCPassPostProcessCircleDOFSetup::Process(FRenderingCompositePassContext& C
 
 FPooledRenderTargetDesc FRCPassPostProcessCircleDOFSetup::ComputeOutputDesc(EPassOutputId InPassOutputId) const
 {
-	FPooledRenderTargetDesc Ret = PassInputs[0].GetOutput()->RenderTargetDesc;
+	FPooledRenderTargetDesc Ret = GetInput(ePId_Input0)->GetOutput()->RenderTargetDesc;
 
 	//	Ret.Extent = FIntPoint::DivideAndRoundUp(ret.Extent, 2);
 	Ret.Extent /= 2;
@@ -764,7 +764,7 @@ void FRCPassPostProcessCircleDOFDilate::Process(FRenderingCompositePassContext&
 
 FPooledRenderTargetDesc FRCPassPostProcessCircleDOFDilate::ComputeOutputDesc(EPassOutputId InPassOutputId) const
 {
-	FPooledRenderTargetDesc Ret = PassInputs[0].GetOutput()->RenderTargetDesc;
+	FPooledRenderTargetDesc Ret = GetInput(ePId_Input0)->GetOutput()->RenderTargetDesc;
 
 	//	Ret.Extent = FIntPoint::DivideAndRoundUp(ret.Extent, 2);
 	Ret.Extent /= 2;
@@ -984,7 +984,7 @@ void FRCPassPostProcessCircleDOF::Process(FRenderingCompositePassContext& Contex
 
 FPooledRenderTargetDesc FRCPassPostProcessCircleDOF::ComputeOutputDesc(EPassOutputId InPassOutputId) const
 {
-	FPooledRenderTargetDesc Ret = PassInputs[0].GetOutput()->RenderTargetDesc;
+	FPooledRenderTargetDesc Ret = GetInput(ePId_Input0)->GetOutput()->RenderTargetDesc;
 
 	Ret.Extent.X = FMath::Max(1, Ret.Extent.X);
 	Ret.Extent.Y = FMath::Max(1, Ret.Extent.Y);
@@ -1155,7 +1155,7 @@ void FRCPassPostProcessCircleDOFRecombine::Process(FRenderingCompositePassContex
 
 FPooledRenderTargetDesc FRCPassPostProcessCircleDOFRecombine::ComputeOutputDesc(EPassOutputId InPassOutputId) const
 {
-	FPooledRenderTargetDesc Ret = PassInputs[0].GetOutput()->RenderTargetDesc;
+	FPooledRenderTargetDesc Ret = GetInput(ePId_Input0)->GetOutput()->RenderTargetDesc;
 
 	Ret.Reset();
 	Ret.DebugName = TEXT("CircleDOFRecombine");
diff --git a/Engine/Source/Runtime/Renderer/Private/PostProcess/PostProcessDownsample.cpp b/Engine/Source/Runtime/Renderer/Private/PostProcess/PostProcessDownsample.cpp
--- a/Engine/Source/Runtime/Renderer/Private/PostProcess/PostProcessDownsample.cpp
+++ b/Engine/Source/Runtime/Renderer/Private/PostProcess/PostProcessDownsample.cpp
@@ -255,7 +255,7 @@ bool FRCPassPostProcessDownsample::IsDepthInputAvailable() const
 
 FPooledRenderTargetDesc FRCPassPostProcessDownsample::ComputeOutputDesc(EPassOutputId InPassOutputId) const
 {
-	FPooledRenderTargetDesc Ret = PassInputs[0].GetOutput()->RenderTargetDesc;
+	FPooledRenderTargetDesc Ret = GetInput(ePId_Input0)->GetOutput()->RenderTargetDesc;
 
 	Ret.Reset();
 
diff --git a/Engine/Source/Runtime/Renderer/Private/PostProcess/PostProcessGBufferHints.cpp b/Engine/Source/Runtime/Renderer/Private/PostProcess/PostProcessGBufferHints.cpp
--- a/Engine/Source/Runtime/Renderer/Private/PostProcess/PostProcessGBufferHints.cpp
+++ b/Engine/Source/Runtime/Renderer/Private/PostProcess/PostProcessGBufferHints.cpp
@@ -188,7 +188,7 @@ void FRCPassPostProcessGBufferHints::Process(FRenderingCompositePassContext& Con
 
 FPooledRenderTargetDesc FRCPassPostProcessGBufferHints::ComputeOutputDesc(EPassOutputId InPassOutputId) const
 {
-	FPooledRenderTargetDesc Ret = PassInputs[0].GetOutput()->RenderTargetDesc;
+	FPooledRenderTargetDesc Ret = GetInput(ePId_Input0)->GetOutput()->RenderTargetDesc;
 
 	Ret.Reset();
 	Ret.DebugName = TEXT("GBufferHints");
diff --git a/Engine/Source/Runtime/Renderer/Private/PostProcess/PostProcessHistogram.cpp b/Engine/Source/Runtime/Renderer/Private/PostProcess/PostProcessHistogram.cpp
--- a/Engine/Source/Runtime/Renderer/Private/PostProcess/PostProcessHistogram.cpp
+++ b/Engine/Source/Runtime/Renderer/Private/PostProcess/PostProcessHistogram.cpp
@@ -151,7 +151,7 @@ FIntPoint FRCPassPostProcessHistogram::ComputeThreadGroupCount(FIntPoint PixelEx
 
 FPooledRenderTargetDesc FRCPassPostProcessHistogram::ComputeOutputDesc(EPassOutputId InPassOutputId) const
 {
-	FPooledRenderTargetDesc UnmodifiedRet = PassInputs[0].GetOutput()->RenderTargetDesc;
+	FPooledRenderTargetDesc UnmodifiedRet = GetInput(ePId_Input0)->GetOutput()->RenderTargetDesc;
 
 	UnmodifiedRet.Reset();
 	FIntPoint PixelExtent = UnmodifiedRet.Extent;
diff --git a/Engine/Source/Runtime/Renderer/Private/PostProcess/PostProcessHistogramReduce.cpp b/Engine/Source/Runtime/Renderer/Private/PostProcess/PostProcessHistogramReduce.cpp
--- a/Engine/Source/Runtime/Renderer/Private/PostProcess/PostProcessHistogramReduce.cpp
+++ b/Engine/Source/Runtime/Renderer/Private/PostProcess/PostProcessHistogramReduce.cpp
@@ -152,7 +152,7 @@ uint32 FRCPassPostProcessHistogramReduce::ComputeLoopSize(FIntPoint PixelExtent)
 
 FPooledRenderTargetDesc FRCPassPostProcessHistogramReduce::ComputeOutputDesc(EPassOutputId InPassOutputId) const
 {
-	FPooledRenderTargetDesc UnmodifiedRet = PassInputs[0].GetOutput()->RenderTargetDesc;
+	FPooledRenderTargetDesc UnmodifiedRet = GetInput(ePId_Input0)->GetOutput()->RenderTargetDesc;
 	
 	UnmodifiedRet.Reset();
 
diff --git a/Engine/Source/Runtime/Renderer/Private/PostProcess/PostProcessLensBlur.cpp b/Engine/Source/Runtime/Renderer/Private/PostProcess/PostProcessLensBlur.cpp
--- a/Engine/Source/Runtime/Renderer/Private/PostProcess/PostProcessLensBlur.cpp
+++ b/Engine/Source/Runtime/Renderer/Private/PostProcess/PostProcessLensBlur.cpp
@@ -223,7 +223,7 @@ void FRCPassPostProcessLensBlur::Process(FRenderingCompositePassContext& Context
 
 FPooledRenderTargetDesc FRCPassPostProcessLensBlur::ComputeOutputDesc(EPassOutputId InPassOutputId) const
 {
-	FPooledRenderTargetDesc Ret = PassInputs[0].GetOutput()->RenderTargetDesc;
+	FPooledRenderTargetDesc Ret = GetInput(ePId_Input0)->GetOutput()->RenderTargetDesc;
 
 	Ret.Reset();
 
diff --git a/Engine/Source/Runtime/Renderer/Private/PostProcess/PostProcessLensFlares.cpp b/Engine/Source/Runtime/Renderer/Private/PostProcess/PostProcessLensFlares.cpp
--- a/Engine/Source/Runtime/Renderer/Private/PostProcess/PostProcessLensFlares.cpp
+++ b/Engine/Source/Runtime/Renderer/Private/PostProcess/PostProcessLensFlares.cpp
@@ -242,7 +242,7 @@ void FRCPassPostProcessLensFlares::Process(FRenderingCompositePassContext& Conte
 
 FPooledRenderTargetDesc FRCPassPostProcessLensFlares::ComputeOutputDesc(EPassOutputId InPassOutputId) const
 {
-	FPooledRenderTargetDesc Ret = PassInputs[0].GetOutput()->RenderTargetDesc;
+	FPooledRenderTargetDesc Ret = GetInput(ePId_Input0)->GetOutput()->RenderTargetDesc;
 
 	Ret.Reset();
 	Ret.DebugName = TEXT("LensFlares");
diff --git a/Engine/Source/Runtime/Renderer/Private/PostProcess/PostProcessMaterial.cpp b/Engine/Source/Runtime/Renderer/Private/PostProcess/PostProcessMaterial.cpp
--- a/Engine/Source/Runtime/Renderer/Private/PostProcess/PostProcessMaterial.cpp
+++ b/Engine/Source/Runtime/Renderer/Private/PostProcess/PostProcessMaterial.cpp
@@ -208,7 +208,7 @@ void FRCPassPostProcessMaterial::Process(FRenderingCompositePassContext& Context
 
 FPooledRenderTargetDesc FRCPassPostProcessMaterial::ComputeOutputDesc(EPassOutputId InPassOutputId) const
 {
-	FPooledRenderTargetDesc Ret = PassInputs[0].GetOutput()->RenderTargetDesc;
+	FPooledRenderTargetDesc Ret = GetInput(ePId_Input0)->GetOutput()->RenderTargetDesc;
 
 	if (OutputFormat != PF_Unknown)
 	{
diff --git a/Engine/Source/Runtime/Renderer/Private/PostProcess/PostProcessMobile.cpp b/Engine/Source/Runtime/Renderer/Private/PostProcess/PostProcessMobile.cpp
--- a/Engine/Source/Runtime/Renderer/Private/PostProcess/PostProcessMobile.cpp
+++ b/Engine/Source/Runtime/Renderer/Private/PostProcess/PostProcessMobile.cpp
@@ -2807,7 +2807,7 @@ FPooledRenderTargetDesc FRCPassPostProcessAaES2::ComputeOutputDesc(EPassOutputId
 	Ret.Format = PF_B8G8R8A8;
 	Ret.NumSamples = 1;
 	Ret.DebugName = TEXT("Aa");
-	Ret.Extent = PassInputs[0].GetOutput()->RenderTargetDesc.Extent;
+	Ret.Extent = GetInput(ePId_Input0)->GetOutput()->RenderTargetDesc.Extent;
 	return Ret;
 }
 
diff --git a/Engine/Source/Runtime/Renderer/Private/PostProcess/PostProcessMorpheus.cpp b/Engine/Source/Runtime/Renderer/Private/PostProcess/PostProcessMorpheus.cpp
--- a/Engine/Source/Runtime/Renderer/Private/PostProcess/PostProcessMorpheus.cpp
+++ b/Engine/Source/Runtime/Renderer/Private/PostProcess/PostProcessMorpheus.cpp
@@ -251,7 +251,7 @@ void FRCPassPostProcessMorpheus::Process(FRenderingCompositePassContext& Context
 
 FPooledRenderTargetDesc FRCPassPostProcessMorpheus::ComputeOutputDesc(EPassOutputId InPassOutputId) const
 {
-	FPooledRenderTargetDesc Ret = PassInputs[0].GetOutput()->RenderTargetDesc;
+	FPooledRenderTargetDesc Ret = GetInput(ePId_Input0)->GetOutput()->RenderTargetDesc;
 
 	Ret.NumSamples = 1;	// no MSAA
 	Ret.Reset();
diff --git a/Engine/Source/Runtime/Renderer/Private/PostProcess/PostProcessMotionBlur.cpp b/Engine/Source/Runtime/Renderer/Private/PostProcess/PostProcessMotionBlur.cpp
--- a/Engine/Source/Runtime/Renderer/Private/PostProcess/PostProcessMotionBlur.cpp
+++ b/Engine/Source/Runtime/Renderer/Private/PostProcess/PostProcessMotionBlur.cpp
@@ -225,7 +225,7 @@ FPooledRenderTargetDesc FRCPassPostProcessMotionBlurSetup::ComputeOutputDesc(EPa
 	{
 		// downsampled velocity
 
-		FPooledRenderTargetDesc Ret = PassInputs[0].GetOutput()->RenderTargetDesc;
+		FPooledRenderTargetDesc Ret = GetInput(ePId_Input0)->GetOutput()->RenderTargetDesc;
 
 		Ret.Reset();
 		Ret.Extent /= 2;
@@ -248,7 +248,7 @@ FPooledRenderTargetDesc FRCPassPostProcessMotionBlurSetup::ComputeOutputDesc(EPa
 		check(InPassOutputId == ePId_Output1);
 
 		// scene color with depth in alpha
-		FPooledRenderTargetDesc Ret = PassInputs[1].GetOutput()->RenderTargetDesc;
+		FPooledRenderTargetDesc Ret = GetInput(ePId_Input1)->GetOutput()->RenderTargetDesc;
 
 		Ret.Reset();
 		Ret.Extent /= 2;
@@ -533,7 +533,7 @@ void FRCPassPostProcessMotionBlur::Process(FRenderingCompositePassContext& Conte
 
 FPooledRenderTargetDesc FRCPassPostProcessMotionBlur::ComputeOutputDesc(EPassOutputId InPassOutputId) const
 {
-	FPooledRenderTargetDesc Ret = PassInputs[0].GetOutput()->RenderTargetDesc;
+	FPooledRenderTargetDesc Ret = GetInput(ePId_Input0)->GetOutput()->RenderTargetDesc;
 
 	Ret.Reset();
 	Ret.DebugName = TEXT("MotionBlur");
@@ -668,7 +668,7 @@ void FRCPassPostProcessMotionBlurRecombine::Process(FRenderingCompositePassConte
 
 FPooledRenderTargetDesc FRCPassPostProcessMotionBlurRecombine::ComputeOutputDesc(EPassOutputId InPassOutputId) const
 {
-	FPooledRenderTargetDesc Ret = PassInputs[0].GetOutput()->RenderTargetDesc;
+	FPooledRenderTargetDesc Ret = GetInput(ePId_Input0)->GetOutput()->RenderTargetDesc;
 
 	Ret.Reset();
 	// we don't need the alpha channel and 32bit is faster and costs less memory
@@ -834,7 +834,7 @@ FPooledRenderTargetDesc FRCPassPostProcessVelocityFlatten::ComputeOutputDesc(EPa
 	if( InPassOutputId == ePId_Output0 )
 	{
 		// Flattened velocity
-		FPooledRenderTargetDesc Ret = PassInputs[0].GetOutput()->RenderTargetDesc;
+		FPooledRenderTargetDesc Ret = GetInput(ePId_Input0)->GetOutput()->RenderTargetDesc;
 		Ret.Reset();
 
 		Ret.Format = PF_FloatR11G11B10;
@@ -847,7 +847,7 @@ FPooledRenderTargetDesc FRCPassPostProcessVelocityFlatten::ComputeOutputDesc(EPa
 	else
 	{
 		// Max tile velocity
-		FPooledRenderTargetDesc UnmodifiedRet = PassInputs[0].GetOutput()->RenderTargetDesc;
+		FPooledRenderTargetDesc UnmodifiedRet = GetInput(ePId_Input0)->GetOutput()->RenderTargetDesc;
 		UnmodifiedRet.Reset();
 
 		FIntPoint PixelExtent = UnmodifiedRet.Extent;
@@ -1073,7 +1073,7 @@ void FRCPassPostProcessVelocityScatter::Process(FRenderingCompositePassContext&
 
 FPooledRenderTargetDesc FRCPassPostProcessVelocityScatter::ComputeOutputDesc(EPassOutputId InPassOutputId) const
 {
-	FPooledRenderTargetDesc Ret = PassInputs[0].GetOutput()->RenderTargetDesc;
+	FPooledRenderTargetDesc Ret = GetInput(ePId_Input0)->GetOutput()->RenderTargetDesc;
 	Ret.Reset();
 
 	Ret.DebugName = TEXT("ScatteredMaxVelocity");
@@ -1173,7 +1173,7 @@ void FRCPassPostProcessVelocityGather::Process(FRenderingCompositePassContext& C
 
 FPooledRenderTargetDesc FRCPassPostProcessVelocityGather::ComputeOutputDesc(EPassOutputId InPassOutputId) const
 {
-	FPooledRenderTargetDesc Ret = PassInputs[0].GetOutput()->RenderTargetDesc;
+	FPooledRenderTargetDesc Ret = GetInput(ePId_Input0)->GetOutput()->RenderTargetDesc;
 	Ret.Reset();
 
 	Ret.TargetableFlags |= TexCreate_UAV;
@@ -1421,7 +1421,7 @@ void FRCPassPostProcessMotionBlurNew::Process(FRenderingCompositePassContext& Co
 
 FPooledRenderTargetDesc FRCPassPostProcessMotionBlurNew::ComputeOutputDesc(EPassOutputId InPassOutputId) const
 {
-	FPooledRenderTargetDesc Ret = PassInputs[0].GetOutput()->RenderTargetDesc;
+	FPooledRenderTargetDesc Ret = GetInput(ePId_Input0)->GetOutput()->RenderTargetDesc;
 
 	Ret.Reset();
 	// we don't need the alpha channel and 32bit is faster and costs less memory
@@ -1580,7 +1580,7 @@ void FRCPassPostProcessVisualizeMotionBlur::Process(FRenderingCompositePassConte
 
 FPooledRenderTargetDesc FRCPassPostProcessVisualizeMotionBlur::ComputeOutputDesc(EPassOutputId InPassOutputId) const
 {
-	FPooledRenderTargetDesc Ret = PassInputs[0].GetOutput()->RenderTargetDesc;
+	FPooledRenderTargetDesc Ret = GetInput(ePId_Input0)->GetOutput()->RenderTargetDesc;
 
 	Ret.Reset();
 	Ret.DebugName = TEXT("MotionBlur");
diff --git a/Engine/Source/Runtime/Renderer/Private/PostProcess/PostProcessNoiseBlur.cpp b/Engine/Source/Runtime/Renderer/Private/PostProcess/PostProcessNoiseBlur.cpp
--- a/Engine/Source/Runtime/Renderer/Private/PostProcess/PostProcessNoiseBlur.cpp
+++ b/Engine/Source/Runtime/Renderer/Private/PostProcess/PostProcessNoiseBlur.cpp
@@ -181,7 +181,7 @@ void FRCPassPostProcessNoiseBlur::Process(FRenderingCompositePassContext& Contex
 
 FPooledRenderTargetDesc FRCPassPostProcessNoiseBlur::ComputeOutputDesc(EPassOutputId InPassOutputId) const
 {
-	FPooledRenderTargetDesc Ret = PassInputs[0].GetOutput()->RenderTargetDesc;
+	FPooledRenderTargetDesc Ret = GetInput(ePId_Input0)->GetOutput()->RenderTargetDesc;
 
 	Ret.Reset();
 
diff --git a/Engine/Source/Runtime/Renderer/Private/PostProcess/PostProcessOutput.cpp b/Engine/Source/Runtime/Renderer/Private/PostProcess/PostProcessOutput.cpp
--- a/Engine/Source/Runtime/Renderer/Private/PostProcess/PostProcessOutput.cpp
+++ b/Engine/Source/Runtime/Renderer/Private/PostProcess/PostProcessOutput.cpp
@@ -28,7 +28,7 @@ void FRCPassPostProcessOutput::Process(FRenderingCompositePassContext& Context)
 
 FPooledRenderTargetDesc FRCPassPostProcessOutput::ComputeOutputDesc(EPassOutputId InPassOutputId) const
 {
-	FPooledRenderTargetDesc Ret = PassInputs[0].GetOutput()->RenderTargetDesc;
+	FPooledRenderTargetDesc Ret = GetInput(ePId_Input0)->GetOutput()->RenderTargetDesc;
 
 	Ret.DebugName = TEXT("PostProcessOutput");
 
diff --git a/Engine/Source/Runtime/Renderer/Private/PostProcess/PostProcessSceneColorFringe.cpp b/Engine/Source/Runtime/Renderer/Private/PostProcess/PostProcessSceneColorFringe.cpp
--- a/Engine/Source/Runtime/Renderer/Private/PostProcess/PostProcessSceneColorFringe.cpp
+++ b/Engine/Source/Runtime/Renderer/Private/PostProcess/PostProcessSceneColorFringe.cpp
@@ -196,7 +196,7 @@ void FRCPassPostProcessSceneColorFringe::Process(FRenderingCompositePassContext&
 
 FPooledRenderTargetDesc FRCPassPostProcessSceneColorFringe::ComputeOutputDesc(EPassOutputId InPassOutputId) const
 {
-	FPooledRenderTargetDesc Ret = PassInputs[0].GetOutput()->RenderTargetDesc;
+	FPooledRenderTargetDesc Ret = GetInput(ePId_Input0)->GetOutput()->RenderTargetDesc;
 
 	Ret.Reset();
 	Ret.DebugName = TEXT("PostProcessAA");
diff --git a/Engine/Source/Runtime/Renderer/Private/PostProcess/PostProcessSelectionOutline.cpp b/Engine/Source/Runtime/Renderer/Private/PostProcess/PostProcessSelectionOutline.cpp
--- a/Engine/Source/Runtime/Renderer/Private/PostProcess/PostProcessSelectionOutline.cpp
+++ b/Engine/Source/Runtime/Renderer/Private/PostProcess/PostProcessSelectionOutline.cpp
@@ -125,7 +125,7 @@ void FRCPassPostProcessSelectionOutlineColor::Process(FRenderingCompositePassCon
 
 FPooledRenderTargetDesc FRCPassPostProcessSelectionOutlineColor::ComputeOutputDesc(EPassOutputId InPassOutputId) const
 {
-	FPooledRenderTargetDesc Ret = PassInputs[0].GetOutput()->RenderTargetDesc;
+	FPooledRenderTargetDesc Ret = GetInput(ePId_Input0)->GetOutput()->RenderTargetDesc;
 
 	Ret.Reset();
 
@@ -390,7 +390,7 @@ void FRCPassPostProcessSelectionOutline::Process(FRenderingCompositePassContext&
 
 FPooledRenderTargetDesc FRCPassPostProcessSelectionOutline::ComputeOutputDesc(EPassOutputId InPassOutputId) const
 {
-	FPooledRenderTargetDesc Ret = PassInputs[0].GetOutput()->RenderTargetDesc;
+	FPooledRenderTargetDesc Ret = GetInput(ePId_Input0)->GetOutput()->RenderTargetDesc;
 
 	Ret.DebugName = TEXT("SelectionComposited");
 
diff --git a/Engine/Source/Runtime/Renderer/Private/PostProcess/PostProcessSubsurface.cpp b/Engine/Source/Runtime/Renderer/Private/PostProcess/PostProcessSubsurface.cpp
--- a/Engine/Source/Runtime/Renderer/Private/PostProcess/PostProcessSubsurface.cpp
+++ b/Engine/Source/Runtime/Renderer/Private/PostProcess/PostProcessSubsurface.cpp
@@ -916,7 +916,7 @@ void FRCPassPostProcessSubsurface::Process(FRenderingCompositePassContext& Conte
 
 FPooledRenderTargetDesc FRCPassPostProcessSubsurface::ComputeOutputDesc(EPassOutputId InPassOutputId) const
 {
-	FPooledRenderTargetDesc Ret = PassInputs[0].GetOutput()->RenderTargetDesc;
+	FPooledRenderTargetDesc Ret = GetInput(ePId_Input0)->GetOutput()->RenderTargetDesc;
 
 	Ret.Reset();
 	Ret.DebugName = (Direction == 0) ? TEXT("SubsurfaceX") : TEXT("SubsurfaceY");
diff --git a/Engine/Source/Runtime/Renderer/Private/PostProcess/PostProcessTemporalAA.cpp b/Engine/Source/Runtime/Renderer/Private/PostProcess/PostProcessTemporalAA.cpp
--- a/Engine/Source/Runtime/Renderer/Private/PostProcess/PostProcessTemporalAA.cpp
+++ b/Engine/Source/Runtime/Renderer/Private/PostProcess/PostProcessTemporalAA.cpp
@@ -271,7 +271,7 @@ void FRCPassPostProcessSSRTemporalAA::Process(FRenderingCompositePassContext& Co
 
 FPooledRenderTargetDesc FRCPassPostProcessSSRTemporalAA::ComputeOutputDesc(EPassOutputId InPassOutputId) const
 {
-	FPooledRenderTargetDesc Ret = PassInputs[0].GetOutput()->RenderTargetDesc;
+	FPooledRenderTargetDesc Ret = GetInput(ePId_Input0)->GetOutput()->RenderTargetDesc;
 
 	Ret.DebugName = TEXT("SSRTemporalAA");
 
@@ -351,7 +351,7 @@ void FRCPassPostProcessDOFTemporalAA::Process(FRenderingCompositePassContext& Co
 
 FPooledRenderTargetDesc FRCPassPostProcessDOFTemporalAA::ComputeOutputDesc(EPassOutputId InPassOutputId) const
 {
-	FPooledRenderTargetDesc Ret = PassInputs[0].GetOutput()->RenderTargetDesc;
+	FPooledRenderTargetDesc Ret = GetInput(ePId_Input0)->GetOutput()->RenderTargetDesc;
 
 	Ret.DebugName = TEXT("BokehDOFTemporalAA");
 
@@ -432,7 +432,7 @@ void FRCPassPostProcessDOFTemporalAANear::Process(FRenderingCompositePassContext
 
 FPooledRenderTargetDesc FRCPassPostProcessDOFTemporalAANear::ComputeOutputDesc(EPassOutputId InPassOutputId) const
 {
-	FPooledRenderTargetDesc Ret = PassInputs[0].GetOutput()->RenderTargetDesc;
+	FPooledRenderTargetDesc Ret = GetInput(ePId_Input0)->GetOutput()->RenderTargetDesc;
 
 	Ret.DebugName = TEXT("BokehDOFTemporalAANear");
 
@@ -509,7 +509,7 @@ void FRCPassPostProcessLightShaftTemporalAA::Process(FRenderingCompositePassCont
 
 FPooledRenderTargetDesc FRCPassPostProcessLightShaftTemporalAA::ComputeOutputDesc(EPassOutputId InPassOutputId) const
 {
-	FPooledRenderTargetDesc Ret = PassInputs[0].GetOutput()->RenderTargetDesc;
+	FPooledRenderTargetDesc Ret = GetInput(ePId_Input0)->GetOutput()->RenderTargetDesc;
 
 	Ret.DebugName = TEXT("LightShaftTemporalAA");
 
@@ -728,7 +728,7 @@ void FRCPassPostProcessTemporalAA::Process(FRenderingCompositePassContext& Conte
 
 FPooledRenderTargetDesc FRCPassPostProcessTemporalAA::ComputeOutputDesc(EPassOutputId InPassOutputId) const
 {
-	FPooledRenderTargetDesc Ret = PassInputs[0].GetOutput()->RenderTargetDesc;
+	FPooledRenderTargetDesc Ret = GetInput(ePId_Input0)->GetOutput()->RenderTargetDesc;
 
 	Ret.Reset();
 	//Ret.Format = PF_FloatRGBA;
diff --git a/Engine/Source/Runtime/Renderer/Private/PostProcess/PostProcessTonemap.cpp b/Engine/Source/Runtime/Renderer/Private/PostProcess/PostProcessTonemap.cpp
--- a/Engine/Source/Runtime/Renderer/Private/PostProcess/PostProcessTonemap.cpp
+++ b/Engine/Source/Runtime/Renderer/Private/PostProcess/PostProcessTonemap.cpp
@@ -1121,7 +1121,7 @@ void FRCPassPostProcessTonemap::Process(FRenderingCompositePassContext& Context)
 
 FPooledRenderTargetDesc FRCPassPostProcessTonemap::ComputeOutputDesc(EPassOutputId InPassOutputId) const
 {
-	FPooledRenderTargetDesc Ret = PassInputs[0].GetOutput()->RenderTargetDesc;
+	FPooledRenderTargetDesc Ret = GetInput(ePId_Input0)->GetOutput()->RenderTargetDesc;
 
 	Ret.Reset();
 	// RGB is the color in LDR, A is the luminance for PostprocessAA
@@ -1520,7 +1520,7 @@ void FRCPassPostProcessTonemapES2::Process(FRenderingCompositePassContext& Conte
 
 FPooledRenderTargetDesc FRCPassPostProcessTonemapES2::ComputeOutputDesc(EPassOutputId InPassOutputId) const
 {
-	FPooledRenderTargetDesc Ret = PassInputs[0].GetOutput()->RenderTargetDesc;
+	FPooledRenderTargetDesc Ret = GetInput(ePId_Input0)->GetOutput()->RenderTargetDesc;
 
 	Ret.Reset();
 	Ret.Format = PF_B8G8R8A8;
diff --git a/Engine/Source/Runtime/Renderer/Private/PostProcess/PostProcessUpscale.cpp b/Engine/Source/Runtime/Renderer/Private/PostProcess/PostProcessUpscale.cpp
--- a/Engine/Source/Runtime/Renderer/Private/PostProcess/PostProcessUpscale.cpp
+++ b/Engine/Source/Runtime/Renderer/Private/PostProcess/PostProcessUpscale.cpp
@@ -318,7 +318,7 @@ void FRCPassPostProcessUpscale::Process(FRenderingCompositePassContext& Context)
 
 FPooledRenderTargetDesc FRCPassPostProcessUpscale::ComputeOutputDesc(EPassOutputId InPassOutputId) const
 {
-	FPooledRenderTargetDesc Ret = PassInputs[0].GetOutput()->RenderTargetDesc;
+	FPooledRenderTargetDesc Ret = GetInput(ePId_Input0)->GetOutput()->RenderTargetDesc;
 
 	Ret.Reset();
 	Ret.DebugName = TEXT("Upscale");
diff --git a/Engine/Source/Runtime/Renderer/Private/PostProcess/PostProcessVisualizeBuffer.cpp b/Engine/Source/Runtime/Renderer/Private/PostProcess/PostProcessVisualizeBuffer.cpp
--- a/Engine/Source/Runtime/Renderer/Private/PostProcess/PostProcessVisualizeBuffer.cpp
+++ b/Engine/Source/Runtime/Renderer/Private/PostProcess/PostProcessVisualizeBuffer.cpp
@@ -267,7 +267,7 @@ void FRCPassPostProcessVisualizeBuffer::Process(FRenderingCompositePassContext&
 
 FPooledRenderTargetDesc FRCPassPostProcessVisualizeBuffer::ComputeOutputDesc(EPassOutputId InPassOutputId) const
 {
-	FPooledRenderTargetDesc Ret = PassInputs[0].GetOutput()->RenderTargetDesc;
+	FPooledRenderTargetDesc Ret = GetInput(ePId_Input0)->GetOutput()->RenderTargetDesc;
 
 	Ret.Reset();
 	Ret.DebugName = TEXT("VisualizeBuffer");
diff --git a/Engine/Source/Runtime/Renderer/Private/PostProcess/PostProcessVisualizeHDR.cpp b/Engine/Source/Runtime/Renderer/Private/PostProcess/PostProcessVisualizeHDR.cpp
--- a/Engine/Source/Runtime/Renderer/Private/PostProcess/PostProcessVisualizeHDR.cpp
+++ b/Engine/Source/Runtime/Renderer/Private/PostProcess/PostProcessVisualizeHDR.cpp
@@ -282,7 +282,7 @@ void FRCPassPostProcessVisualizeHDR::Process(FRenderingCompositePassContext& Con
 
 FPooledRenderTargetDesc FRCPassPostProcessVisualizeHDR::ComputeOutputDesc(EPassOutputId InPassOutputId) const
 {
-	FPooledRenderTargetDesc Ret = PassInputs[0].GetOutput()->RenderTargetDesc;
+	FPooledRenderTargetDesc Ret = GetInput(ePId_Input0)->GetOutput()->RenderTargetDesc;
 
 	Ret.Reset();
 	Ret.DebugName = TEXT("VisualizeHDR");
diff --git a/Engine/Source/Runtime/Renderer/Private/PostProcess/PostProcessWeightedSampleSum.cpp b/Engine/Source/Runtime/Renderer/Private/PostProcess/PostProcessWeightedSampleSum.cpp
--- a/Engine/Source/Runtime/Renderer/Private/PostProcess/PostProcessWeightedSampleSum.cpp
+++ b/Engine/Source/Runtime/Renderer/Private/PostProcess/PostProcessWeightedSampleSum.cpp
@@ -494,7 +494,7 @@ void FRCPassPostProcessWeightedSampleSum::Process(FRenderingCompositePassContext
 {
 	const FSceneView& View = Context.View;
 
-	FRenderingCompositeOutput *Input = PassInputs[0].GetOutput();
+	FRenderingCompositeOutput *Input = GetInput(ePId_Input0)->GetOutput();
 
 	// input is not hooked up correctly
 	check(Input);
@@ -557,7 +557,9 @@ void FRCPassPostProcessWeightedSampleSum::Process(FRenderingCompositePassContext
 
 	const FTextureRHIRef& FilterTexture = InputPooledElement->GetRenderTargetItem().ShaderResourceTexture;
 
-	FRenderingCompositeOutput *Input1 = PassInputs[1].GetOutput();
+	FRenderingCompositeOutputRef* NodeInput1 = GetInput(ePId_Input1);
+
+	FRenderingCompositeOutput *Input1 = NodeInput1 ? NodeInput1->GetOutput() : 0;
 
 	uint32 CombineMethodInt = 0;
 
@@ -631,7 +633,7 @@ void FRCPassPostProcessWeightedSampleSum::Process(FRenderingCompositePassContext
 
 FPooledRenderTargetDesc FRCPassPostProcessWeightedSampleSum::ComputeOutputDesc(EPassOutputId InPassOutputId) const
 {
-	FPooledRenderTargetDesc Ret = PassInputs[0].GetOutput()->RenderTargetDesc;
+	FPooledRenderTargetDesc Ret = GetInput(ePId_Input0)->GetOutput()->RenderTargetDesc;
 
 	if(DoFastBlur())
 	{
diff --git a/Engine/Source/Runtime/Renderer/Private/PostProcess/RenderingCompositionGraph.h b/Engine/Source/Runtime/Renderer/Private/PostProcess/RenderingCompositionGraph.h
--- a/Engine/Source/Runtime/Renderer/Private/PostProcess/RenderingCompositionGraph.h
+++ b/Engine/Source/Runtime/Renderer/Private/PostProcess/RenderingCompositionGraph.h
@@ -471,9 +471,10 @@ struct TRenderingCompositePassBase :public FRenderingCompositePass
 		PassOutputColorArrays[OutputId] = OutputBuffer;
 	}
 
-protected:
-	/** hack to allow 0 inputs */
+private:
+	/** use GetInput() instead of accessing PassInputs directly */
 	FRenderingCompositeOutputRef PassInputs[InputCount == 0 ? 1 : InputCount];
+protected:
 	/** */
 	FRenderingCompositeOutput PassOutputs[OutputCount];
 	/** Filenames that the outputs can be written to after being processed */
diff --git a/Engine/Source/Runtime/Renderer/Private/ShaderComplexityRendering.cpp b/Engine/Source/Runtime/Renderer/Private/ShaderComplexityRendering.cpp
--- a/Engine/Source/Runtime/Renderer/Private/ShaderComplexityRendering.cpp
+++ b/Engine/Source/Runtime/Renderer/Private/ShaderComplexityRendering.cpp
@@ -205,7 +205,7 @@ void FRCPassPostProcessVisualizeComplexity::Process(FRenderingCompositePassConte
 
 FPooledRenderTargetDesc FRCPassPostProcessVisualizeComplexity::ComputeOutputDesc(EPassOutputId InPassOutputId) const
 {
-	FPooledRenderTargetDesc Ret = PassInputs[0].GetOutput()->RenderTargetDesc;
+	FPooledRenderTargetDesc Ret = GetInput(ePId_Input0)->GetOutput()->RenderTargetDesc;
 
 	Ret.Reset();
 	Ret.DebugName = TEXT("VisualizeComplexity");
