Fixes light shafts when rendering with r.ScreenPercentage != 100

diff --git a/Engine/Source/Runtime/Renderer/Private/LightShaftRendering.cpp b/Engine/Source/Runtime/Renderer/Private/LightShaftRendering.cpp
--- a/Engine/Source/Runtime/Renderer/Private/LightShaftRendering.cpp
+++ b/Engine/Source/Runtime/Renderer/Private/LightShaftRendering.cpp
@@ -131,8 +131,16 @@ public:
 		const float InvBufferSizeY = 1.0f / BufferSize.Y;
 
 		FVector2D ScreenSpaceBlurOrigin;
-
-		verify(View.ScreenToPixel(ProjectedBlurOrigin, ScreenSpaceBlurOrigin)); 
+		{
+			verify(ProjectedBlurOrigin.W > 0.0f); 
+			float InvW = 1.0f / ProjectedBlurOrigin.W;
+			float Y = (GProjectionSignY > 0.0f) ? ProjectedBlurOrigin.Y : 1.0f - ProjectedBlurOrigin.Y;
+			ScreenSpaceBlurOrigin = FVector2D(
+				View.ViewRect.Min.X + (0.5f + ProjectedBlurOrigin.X * 0.5f * InvW) * View.ViewRect.Width(),
+				View.ViewRect.Min.Y + (0.5f - Y * 0.5f * InvW) * View.ViewRect.Height()
+				);
+		}
+		
 		ScreenSpaceBlurOrigin.X *= InvBufferSizeX;
 		ScreenSpaceBlurOrigin.Y *= InvBufferSizeY;
 		FVector2D TextureSpaceBlurOrigin(ScreenSpaceBlurOrigin * FVector2D(AspectRatioAndInvAspectRatio.Z, AspectRatioAndInvAspectRatio.W));
