fixed case that cannot happen yet

diff --git a/Engine/Source/Runtime/Renderer/Private/PostProcess/ScreenSpaceReflections.cpp b/Engine/Source/Runtime/Renderer/Private/PostProcess/ScreenSpaceReflections.cpp
--- a/Engine/Source/Runtime/Renderer/Private/PostProcess/ScreenSpaceReflections.cpp
+++ b/Engine/Source/Runtime/Renderer/Private/PostProcess/ScreenSpaceReflections.cpp
@@ -302,7 +302,10 @@ void FRCPassPostProcessScreenSpaceReflections::Process(FRenderingCompositePassCo
 	
 	const bool VisualizeSSR = View.Family->EngineShowFlags.VisualizeSSR;
 	const bool SSRStencilPrePass = CVarSSRStencil.GetValueOnRenderThread() != 0 && !VisualizeSSR;
-	const bool SSRConeTracing = GetInput(ePId_Input2)->GetOutput() != nullptr;
+
+	FRenderingCompositeOutputRef* Input2 = GetInput(ePId_Input2);
+
+	const bool SSRConeTracing = Input2 && Input2->GetOutput();
 	
 	if (VisualizeSSR)
 	{
