added r.CompositionGraphOrder, default is not changed
help text:
 Defines in which order the nodes in the CompositionGraph are executed (affects postprocess and some lighting).
 Option 1 provides more control, which can be useful for preserving ESRAM, avoid GPU sync, cluster up compute shaders for performance and control AsyncCompute.
  0: tree order starting with the root, first all inputs then dependencies (classic UE4, unconnected nodes are not getting executed)
  1: RegisterPass() call order, unless the dependencies (input and additional) require a different order (might become new default as it provides more control, executes all registered nodes

diff --git a/Engine/Source/Runtime/Renderer/Private/CompositionLighting/CompositionLighting.cpp b/Engine/Source/Runtime/Renderer/Private/CompositionLighting/CompositionLighting.cpp
--- a/Engine/Source/Runtime/Renderer/Private/CompositionLighting/CompositionLighting.cpp
+++ b/Engine/Source/Runtime/Renderer/Private/CompositionLighting/CompositionLighting.cpp
@@ -398,10 +398,6 @@ void FCompositionLighting::ProcessAfterLighting(FRHICommandListImmediate& RHICmd
 				FRenderingCompositePass* RecombinePass = Context.Graph.RegisterPass(new(FMemStack::Get()) FRCPassPostProcessSubsurfaceRecombine());
 				RecombinePass->SetInput(ePId_Input0, Pass1);
 				RecombinePass->SetInput(ePId_Input1, PassExtractSpecular);
-				
-				SCOPED_DRAW_EVENT(RHICmdList, CompositionLighting_SSSSS);
-
-				CompositeContext.Process(RecombinePass, TEXT("CompositionLighting_SSSSS"));
 			}
 		}
 
diff --git a/Engine/Source/Runtime/Renderer/Private/PostProcess/PostProcessing.cpp b/Engine/Source/Runtime/Renderer/Private/PostProcess/PostProcessing.cpp
--- a/Engine/Source/Runtime/Renderer/Private/PostProcess/PostProcessing.cpp
+++ b/Engine/Source/Runtime/Renderer/Private/PostProcess/PostProcessing.cpp
@@ -1661,9 +1661,7 @@ void FPostProcessing::Process(FRHICommandListImmediate& RHICmdList, FViewInfo& V
 
 			OverrideRenderTarget(Context.FinalOutput, Temp, Desc);
 
-			// We want this executed early as it uses AsyncCompute (where available e.g. XboxOne).
-			//CompositeContext.Process(VelocityFlattenPass, TEXT("VelocityFlattenPass"));
-			// Now we process the remaining part of the graph.
+			// execute the graph/DAG
 			CompositeContext.Process(Context.FinalOutput.GetPass(), TEXT("PostProcessing"));
 		}
 	}
diff --git a/Engine/Source/Runtime/Renderer/Private/PostProcess/RenderingCompositionGraph.cpp b/Engine/Source/Runtime/Renderer/Private/PostProcess/RenderingCompositionGraph.cpp
--- a/Engine/Source/Runtime/Renderer/Private/PostProcess/RenderingCompositionGraph.cpp
+++ b/Engine/Source/Runtime/Renderer/Private/PostProcess/RenderingCompositionGraph.cpp
@@ -8,6 +8,26 @@
 #include "RenderingCompositionGraph.h"
 #include "HighResScreenshot.h"
 
+void ExecuteCompositionGraphDebug();
+
+static TAutoConsoleVariable<int32> CVarCompositionGraphOrder(
+	TEXT("r.CompositionGraphOrder"),
+	0,
+	TEXT("Defines in which order the nodes in the CompositionGraph are executed (affects postprocess and some lighting).\n")
+	TEXT("Option 1 provides more control, which can be useful for preserving ESRAM, avoid GPU sync, cluster up compute shaders for performance and control AsyncCompute.\n")
+	TEXT(" 0: tree order starting with the root, first all inputs then dependencies (classic UE4, unconnected nodes are not getting executed)\n")
+	TEXT(" 1: RegisterPass() call order, unless the dependencies (input and additional) require a different order (might become new default as it provides more control, executes all registered nodes)"),
+	ECVF_RenderThreadSafe);
+
+#if !UE_BUILD_SHIPPING
+FAutoConsoleCommand CmdCompositionGraphDebug(
+	TEXT("r.CompositionGraphDebug"),
+	TEXT("Execute this command to get a single frame dump of the composition graph of one frame (post processing and lighting)."),
+	FConsoleCommandDelegate::CreateStatic(ExecuteCompositionGraphDebug)
+	);
+#endif
+
+
 // render thread, 0:off, >0 next n frames should be debugged
 uint32 GDebugCompositionGraphFrames = 0;
 
@@ -94,15 +114,6 @@ void CompositionGraph_OnStartFrame()
 #endif
 }
 
-
-#if !UE_BUILD_SHIPPING
-FAutoConsoleCommand CmdCompositionGraphDebug(
-	TEXT("r.CompositionGraphDebug"),
-	TEXT("Execute to get a single frame dump of the composition graph of one frame (post processing and lighting)."),
-	FConsoleCommandDelegate::CreateStatic(ExecuteCompositionGraphDebug)
-	);
-#endif
-
 FRenderingCompositePassContext::FRenderingCompositePassContext(FRHICommandListImmediate& InRHICmdList, FViewInfo& InView/*, const FSceneRenderTargetItem& InRenderTargetItem*/)
 	: View(InView)
 	, ViewState((FSceneViewState*)InView.State)
@@ -111,6 +122,7 @@ FRenderingCompositePassContext::FRenderingCompositePassContext(FRHICommandListIm
 	, ViewPortRect(0, 0, 0 ,0)
 	, FeatureLevel(View.GetFeatureLevel())
 	, ShaderMap(InView.ShaderMap)
+	, bWasProcessed(false)
 {
 	check(!IsViewportValid());
 }
@@ -122,6 +134,11 @@ FRenderingCompositePassContext::~FRenderingCompositePassContext()
 
 void FRenderingCompositePassContext::Process(FRenderingCompositePass* Root, const TCHAR *GraphDebugName)
 {
+	// call this method only once afetr the graph is finished
+	check(!bWasProcessed);
+
+	bWasProcessed = true;
+
 	if(Root)
 	{
 		if(ShouldDebugCompositionGraph())
@@ -140,10 +157,25 @@ void FRenderingCompositePassContext::Process(FRenderingCompositePass* Root, cons
 			GGMLFileWriter.WriteLine("\tdirected\t1");
 		}
 
-		FRenderingCompositeOutputRef Output(Root);
+		bool bNewOrder = CVarCompositionGraphOrder.GetValueOnRenderThread() != 0;
 
-		Graph.RecursivelyGatherDependencies(Output);
-		Graph.RecursivelyProcess(Root, *this);
+		if(bNewOrder)
+		{
+			for (FRenderingCompositePass* Node : Graph.Nodes)
+			{
+				Graph.RecursivelyGatherDependencies(Node);
+			}
+
+			for (FRenderingCompositePass* Node : Graph.Nodes)
+			{
+				Graph.RecursivelyProcess(Node, *this);
+			}
+		}
+		else
+		{
+			Graph.RecursivelyGatherDependencies(Root);
+			Graph.RecursivelyProcess(Root, *this);
+		}
 
 		if(ShouldDebugCompositionGraph())
 		{
@@ -185,9 +217,9 @@ void FRenderingCompositionGraph::Free()
 	Nodes.Empty();
 }
 
-void FRenderingCompositionGraph::RecursivelyGatherDependencies(const FRenderingCompositeOutputRef& InOutputRef)
+void FRenderingCompositionGraph::RecursivelyGatherDependencies(FRenderingCompositePass *Pass)
 {
-	FRenderingCompositePass *Pass = InOutputRef.GetPass();
+	checkSlow(Pass);
 
 	if(Pass->bComputeOutputDescWasCalled)
 	{
@@ -208,10 +240,10 @@ void FRenderingCompositionGraph::RecursivelyGatherDependencies(const FRenderingC
 			InputOutput->AddDependency();
 		}
 		
-		if(OutputRefIt->GetPass())
+		if(FRenderingCompositePass* Pass = OutputRefIt->GetPass())
 		{
 			// recursively process all inputs of this Pass
-			RecursivelyGatherDependencies(*OutputRefIt);
+			RecursivelyGatherDependencies(Pass);
 		}
 	}
 
diff --git a/Engine/Source/Runtime/Renderer/Private/PostProcess/RenderingCompositionGraph.h b/Engine/Source/Runtime/Renderer/Private/PostProcess/RenderingCompositionGraph.h
--- a/Engine/Source/Runtime/Renderer/Private/PostProcess/RenderingCompositionGraph.h
+++ b/Engine/Source/Runtime/Renderer/Private/PostProcess/RenderingCompositionGraph.h
@@ -45,14 +45,17 @@ private:
 	/** */
 	TArray<FRenderingCompositePass*> Nodes;
 
-	/** */
+	/** release all nodes */
 	void Free();
 
 	/** */
 	void ProcessGatherDependency(const FRenderingCompositeOutputRef* OutputRefIt);
 
-	/** should only be called by GatherDependencies(), can also be implemented without recursion */
-	static void RecursivelyGatherDependencies(const FRenderingCompositeOutputRef& InOutputRef);
+	/**
+	 * could also be implemented without recursion
+	 * @param Pass must not be 0
+	 */
+	static void RecursivelyGatherDependencies(FRenderingCompositePass *Pass);
 
 	/** can also be implemented without recursion */
 	void RecursivelyProcess(const FRenderingCompositeOutputRef& InOutputRef, FRenderingCompositePassContext& Context) const;
@@ -76,19 +79,16 @@ private:
 
 struct FRenderingCompositePassContext
 {
+	// constructor
 	FRenderingCompositePassContext(FRHICommandListImmediate& RHICmdList, FViewInfo& InView);
 
+	// destructor
 	~FRenderingCompositePassContext();
 
+	// call this only once after all nodes have been registered and connected (SetInput() or SetDependency())
 	// @param GraphDebugName must not be 0
 	void Process(FRenderingCompositePass* Root, const TCHAR *GraphDebugName);
 
-	//
-	FViewInfo& View;
-	FSceneViewState* ViewState;
-	// is updated before each Pass->Process() call
-	FRenderingCompositePass* Pass;
-
 	// call this method instead of RHISetViewport() so we can cache the values and use them to map beteen ScreenPos and pixels
 	void SetViewportAndCallRHI(FIntRect InViewPortRect, float InMinZ = 0.0f, float InMaxZ = 1.0f)
 	{
@@ -123,6 +123,7 @@ struct FRenderingCompositePassContext
 		return ViewPortRect;
 	}
 
+	//
 	bool IsViewportValid() const
 	{
 		return ViewPortRect.Min != ViewPortRect.Max;
@@ -132,16 +133,27 @@ struct FRenderingCompositePassContext
 	EShaderPlatform GetShaderPlatform() const { return GShaderPlatformForFeatureLevel[FeatureLevel]; }
 	TShaderMap<FGlobalShaderType>* GetShaderMap() const { check(ShaderMap); return ShaderMap; }
 
+	//
+	FViewInfo& View;
+	//
+	FSceneViewState* ViewState;
+	// is updated before each Pass->Process() call
+	FRenderingCompositePass* Pass;
+	//
 	FRenderingCompositionGraph Graph;
-
+	//
 	FRHICommandListImmediate& RHICmdList;
 
 private:
+
 	// cached state to map between ScreenPos and pixels
 	FIntRect ViewPortRect;
-
+	//
 	ERHIFeatureLevel::Type FeatureLevel;
+	//
 	TShaderMap<FGlobalShaderType>* ShaderMap;
+	// to ensure we only process the graph once
+	bool bWasProcessed;
 };
 
 // ---------------------------------------------------------------------------
