enabled new Composition Graph order (node creation order), should not make a difference, in rare cases it could execute more nodes (if they were not connected up before which shouldn't be the case)

diff --git a/Engine/Source/Runtime/Renderer/Private/PostProcess/RenderingCompositionGraph.cpp b/Engine/Source/Runtime/Renderer/Private/PostProcess/RenderingCompositionGraph.cpp
--- a/Engine/Source/Runtime/Renderer/Private/PostProcess/RenderingCompositionGraph.cpp
+++ b/Engine/Source/Runtime/Renderer/Private/PostProcess/RenderingCompositionGraph.cpp
@@ -12,7 +12,7 @@ void ExecuteCompositionGraphDebug();
 
 static TAutoConsoleVariable<int32> CVarCompositionGraphOrder(
 	TEXT("r.CompositionGraphOrder"),
-	0,
+	1,
 	TEXT("Defines in which order the nodes in the CompositionGraph are executed (affects postprocess and some lighting).\n")
 	TEXT("Option 1 provides more control, which can be useful for preserving ESRAM, avoid GPU sync, cluster up compute shaders for performance and control AsyncCompute.\n")
 	TEXT(" 0: tree order starting with the root, first all inputs then dependencies (classic UE4, unconnected nodes are not getting executed)\n")
diff --git a/Engine/Source/Runtime/Renderer/Private/PostProcess/RenderingCompositionGraph.h b/Engine/Source/Runtime/Renderer/Private/PostProcess/RenderingCompositionGraph.h
--- a/Engine/Source/Runtime/Renderer/Private/PostProcess/RenderingCompositionGraph.h
+++ b/Engine/Source/Runtime/Renderer/Private/PostProcess/RenderingCompositionGraph.h
@@ -52,12 +52,12 @@ private:
 	void ProcessGatherDependency(const FRenderingCompositeOutputRef* OutputRefIt);
 
 	/**
-	 * could also be implemented without recursion
+	 * Is called by FRenderingCompositePassContext::Process(), could be implemented without recursion
 	 * @param Pass must not be 0
 	 */
 	static void RecursivelyGatherDependencies(FRenderingCompositePass *Pass);
 
-	/** can also be implemented without recursion */
+	/** could be implemented without recursion */
 	void RecursivelyProcess(const FRenderingCompositeOutputRef& InOutputRef, FRenderingCompositePassContext& Context) const;
 
 	/** Write the contents of the specified output to a file */
