refactor, made bool -> enum (ETranslucencyPassType) to make code easier to follow

diff --git a/Engine/Source/Runtime/Renderer/Private/DeferredShadingRenderer.h b/Engine/Source/Runtime/Renderer/Private/DeferredShadingRenderer.h
--- a/Engine/Source/Runtime/Renderer/Private/DeferredShadingRenderer.h
+++ b/Engine/Source/Runtime/Renderer/Private/DeferredShadingRenderer.h
@@ -235,7 +235,7 @@ private:
 	void RenderLightShaftBloom(FRHICommandListImmediate& RHICmdList);
 
 	/** Reuses an existing translucent shadow map if possible or re-renders one if necessary. */
-	const FProjectedShadowInfo* PrepareTranslucentShadowMap(FRHICommandList& RHICmdList, const FViewInfo& View, FPrimitiveSceneInfo* PrimitiveSceneInfo, bool bSeparateTranslucencyPass);
+	const FProjectedShadowInfo* PrepareTranslucentShadowMap(FRHICommandList& RHICmdList, const FViewInfo& View, FPrimitiveSceneInfo* PrimitiveSceneInfo, ETranslucencyPassType TranslucenyPassType);
 
 	bool ShouldRenderVelocities() const;
 
diff --git a/Engine/Source/Runtime/Renderer/Private/Renderer.cpp b/Engine/Source/Runtime/Renderer/Private/Renderer.cpp
--- a/Engine/Source/Runtime/Renderer/Private/Renderer.cpp
+++ b/Engine/Source/Runtime/Renderer/Private/Renderer.cpp
@@ -53,7 +53,7 @@ void FRendererModule::DrawTileMesh(FRHICommandListImmediate& RHICmdList, const F
 		{
 			if (FeatureLevel >= ERHIFeatureLevel::SM4)
 			{
-				FTranslucencyDrawingPolicyFactory::DrawDynamicMesh(RHICmdList, View, FTranslucencyDrawingPolicyFactory::ContextType(nullptr, false, true), Mesh, false, false, NULL, HitProxyId);
+				FTranslucencyDrawingPolicyFactory::DrawDynamicMesh(RHICmdList, View, FTranslucencyDrawingPolicyFactory::ContextType(nullptr, TPT_NonSeparateTransluceny, true), Mesh, false, false, NULL, HitProxyId);
 			}
 			else
 			{
diff --git a/Engine/Source/Runtime/Renderer/Private/SceneRendering.h b/Engine/Source/Runtime/Renderer/Private/SceneRendering.h
--- a/Engine/Source/Runtime/Renderer/Private/SceneRendering.h
+++ b/Engine/Source/Runtime/Renderer/Private/SceneRendering.h
@@ -59,6 +59,12 @@ public:
 	TArray<FProjectedShadowInfo*,SceneRenderingAllocator> OccludedPerObjectShadows;
 };
 
+// enum instead of bool to get better visibility when we pass around multiple bools
+enum ETranslucencyPassType
+{
+	TPT_NonSeparateTransluceny,
+	TPT_SeparateTransluceny
+};
 
 /** 
 * Set of sorted translucent scene prims  
@@ -71,28 +77,28 @@ public:
 	* Iterate over the sorted list of prims and draw them
 	* @param View - current view used to draw items
 	* @param PhaseSortedPrimitives - array with the primitives we want to draw
-	* @param bSeparateTranslucencyPass
+	* @param TranslucenyPassType
 	*/
-	void DrawPrimitives(FRHICommandListImmediate& RHICmdList, const class FViewInfo& View, class FDeferredShadingSceneRenderer& Renderer, bool bSeparateTranslucencyPass) const;
+	void DrawPrimitives(FRHICommandListImmediate& RHICmdList, const class FViewInfo& View, class FDeferredShadingSceneRenderer& Renderer, ETranslucencyPassType TranslucenyPassType) const;
 
 	/**
 	* Iterate over the sorted list of prims and draw them
 	* @param View - current view used to draw items
 	* @param PhaseSortedPrimitives - array with the primitives we want to draw
-	* @param bSeparateTranslucencyPass
+	* @param TranslucenyPassType
 	* @param FirstIndex, range of elements to render
 	* @param LastIndex, range of elements to render
 	*/
-	void DrawPrimitivesParallel(FRHICommandList& RHICmdList, const class FViewInfo& View, class FDeferredShadingSceneRenderer& Renderer, bool bSeparateTranslucencyPass, int32 FirstIndex, int32 LastIndex) const;
+	void DrawPrimitivesParallel(FRHICommandList& RHICmdList, const class FViewInfo& View, class FDeferredShadingSceneRenderer& Renderer, ETranslucencyPassType TranslucenyPassType, int32 FirstIndex, int32 LastIndex) const;
 
 	/**
 	* Draw a single primitive...this is used when we are rendering in parallel and we need to handlke a translucent shadow
 	* @param View - current view used to draw items
 	* @param PhaseSortedPrimitives - array with the primitives we want to draw
-	* @param bSeparateTranslucencyPass
+	* @param TranslucenyPassType
 	* @param Index, element to render
 	*/
-	void DrawAPrimitive(FRHICommandList& RHICmdList, const class FViewInfo& View, class FDeferredShadingSceneRenderer& Renderer, bool bSeparateTranslucencyPass, int32 Index) const;
+	void DrawAPrimitive(FRHICommandList& RHICmdList, const class FViewInfo& View, class FDeferredShadingSceneRenderer& Renderer, ETranslucencyPassType TranslucenyPassType, int32 Index) const;
 
 	/** Draw all the primitives in this set for the forward shading pipeline. */
 	void DrawPrimitivesForForwardShading(FRHICommandListImmediate& RHICmdList, const class FViewInfo& View, class FSceneRenderer& Renderer) const;
@@ -203,7 +209,7 @@ private:
 	TArray<FSortedPrim,SceneRenderingAllocator> SortedSeparateTranslucencyPrims;
 
 	/** Renders a single primitive for the deferred shading pipeline. */
-	void RenderPrimitive(FRHICommandList& RHICmdList, const FViewInfo& View, FPrimitiveSceneInfo* PrimitiveSceneInfo, const FPrimitiveViewRelevance& ViewRelevance, const FProjectedShadowInfo* TranslucentSelfShadow, bool bSeparateTranslucencyPass) const;
+	void RenderPrimitive(FRHICommandList& RHICmdList, const FViewInfo& View, FPrimitiveSceneInfo* PrimitiveSceneInfo, const FPrimitiveViewRelevance& ViewRelevance, const FProjectedShadowInfo* TranslucentSelfShadow, ETranslucencyPassType TranslucenyPassType) const;
 };
 
 template <> struct TIsPODType<FTranslucentPrimSet::FDepthSortedPrim> { enum { Value = true }; };
diff --git a/Engine/Source/Runtime/Renderer/Private/TranslucentRendering.cpp b/Engine/Source/Runtime/Renderer/Private/TranslucentRendering.cpp
--- a/Engine/Source/Runtime/Renderer/Private/TranslucentRendering.cpp
+++ b/Engine/Source/Runtime/Renderer/Private/TranslucentRendering.cpp
@@ -11,12 +11,12 @@
 #include "LightPropagationVolume.h"
 #include "SceneUtils.h"
 
-static void SetTranslucentRenderTargetAndState(FRHICommandList& RHICmdList, const FViewInfo& View, bool bSeparateTranslucencyPass, bool bFirstTimeThisFrame = false)
+static void SetTranslucentRenderTargetAndState(FRHICommandList& RHICmdList, const FViewInfo& View, ETranslucencyPassType TranslucenyPassType, bool bFirstTimeThisFrame = false)
 {
 	FSceneRenderTargets& SceneContext = FSceneRenderTargets::Get(RHICmdList);
 	bool bSetupTranslucentState = true;
 
-	if (bSeparateTranslucencyPass && SceneContext.IsSeparateTranslucencyActive(View))
+	if ((TranslucenyPassType == TPT_SeparateTransluceny) && SceneContext.IsSeparateTranslucencyActive(View))
 	{
 		bSetupTranslucentState = SceneContext.BeginRenderingSeparateTranslucency(RHICmdList, View, bFirstTimeThisFrame);
 	}
@@ -32,10 +32,10 @@ static void SetTranslucentRenderTargetAndState(FRHICommandList& RHICmdList, cons
 	}
 }
 
-static void FinishTranslucentRenderTarget(FRHICommandListImmediate& RHICmdList, const FViewInfo& View, bool bSeparateTranslucencyPass)
+static void FinishTranslucentRenderTarget(FRHICommandListImmediate& RHICmdList, const FViewInfo& View, ETranslucencyPassType TranslucenyPassType)
 {
 	FSceneRenderTargets& SceneContext = FSceneRenderTargets::Get(RHICmdList);
-	if (bSeparateTranslucencyPass && SceneContext.IsSeparateTranslucencyActive(View))
+	if ((TranslucenyPassType == TPT_SeparateTransluceny) && SceneContext.IsSeparateTranslucencyActive(View))
 	{
 		SceneContext.FinishRenderingSeparateTranslucency(RHICmdList, View);
 	}
@@ -45,7 +45,7 @@ static void FinishTranslucentRenderTarget(FRHICommandListImmediate& RHICmdList,
 	}
 }
 
-const FProjectedShadowInfo* FDeferredShadingSceneRenderer::PrepareTranslucentShadowMap(FRHICommandList& RHICmdList, const FViewInfo& View, FPrimitiveSceneInfo* PrimitiveSceneInfo, bool bSeparateTranslucencyPass)
+const FProjectedShadowInfo* FDeferredShadingSceneRenderer::PrepareTranslucentShadowMap(FRHICommandList& RHICmdList, const FViewInfo& View, FPrimitiveSceneInfo* PrimitiveSceneInfo, ETranslucencyPassType TranslucenyPassType)
 {
 	const FVisibleLightInfo* VisibleLightInfo = NULL;
 	FProjectedShadowInfo* TranslucentSelfShadow = NULL;
@@ -68,7 +68,7 @@ const FProjectedShadowInfo* FDeferredShadingSceneRenderer::PrepareTranslucentSha
 				VisibleLightInfo = &VisibleLightInfos[LightSceneInfo->Id];
 				FProjectedShadowInfo* ObjectShadow = NULL;
 
-				for (int32 ShadowIndex = 0; ShadowIndex < VisibleLightInfo->AllProjectedShadows.Num(); ShadowIndex++)
+				for (int32 ShadowIndex = 0, Count = VisibleLightInfo->AllProjectedShadows.Num(); ShadowIndex < Count; ShadowIndex++)
 				{
 					FProjectedShadowInfo* CurrentShadowInfo = VisibleLightInfo->AllProjectedShadows[ShadowIndex];
 
@@ -138,7 +138,7 @@ const FProjectedShadowInfo* FDeferredShadingSceneRenderer::PrepareTranslucentSha
 				TranslucentSelfShadow->RenderTranslucencyDepths(RHICmdList, this);
 
 				// Restore state
-				SetTranslucentRenderTargetAndState(RHICmdList, View, bSeparateTranslucencyPass);
+				SetTranslucentRenderTargetAndState(RHICmdList, View, TranslucenyPassType);
 			}
 		}
 	}
@@ -328,7 +328,7 @@ static void CopySceneColorAndRestore(FRHICommandList& RHICmdList, const FViewInf
 	check(IsInRenderingThread());
 	FTranslucencyDrawingPolicyFactory::CopySceneColor(RHICmdList, View, PrimitiveSceneProxy);
 	// Restore state
-	SetTranslucentRenderTargetAndState(RHICmdList, View, false);
+	SetTranslucentRenderTargetAndState(RHICmdList, View, TPT_NonSeparateTransluceny);
 }
 
 class FCopySceneColorAndRestoreRenderThreadTask
@@ -392,7 +392,7 @@ bool FTranslucencyDrawingPolicyFactory::DrawMesh(
 	// Only render translucent materials.
 	if(IsTranslucentBlendMode(BlendMode))
 	{
-		if (Material->IsSeparateTranslucencyEnabled() != DrawingContext.bSeparateTranslucencyPass)
+		if (Material->IsSeparateTranslucencyEnabled() != (DrawingContext.TranslucenyPassType == TPT_SeparateTransluceny))
 		{
 			return false;
 		}
@@ -405,7 +405,7 @@ bool FTranslucencyDrawingPolicyFactory::DrawMesh(
 				{
 					FRHICommandList* CmdList = new FRHICommandList;
 					CmdList->CopyRenderThreadContexts(RHICmdList);
-					FGraphEventRef RenderThreadCompletionEvent = TGraphTask<FCopySceneColorAndRestoreRenderThreadTask>::CreateTask().ConstructAndDispatchWhenReady(*CmdList, View, PrimitiveSceneProxy);
+					FGraphEventRef RenderThreadCompletionEvent = 					TGraphTask<FCopySceneColorAndRestoreRenderThreadTask>::CreateTask().ConstructAndDispatchWhenReady(*CmdList, View, PrimitiveSceneProxy);
 					RHICmdList.QueueRenderThreadCommandListSubmit(RenderThreadCompletionEvent, CmdList);
 				}
 				else
@@ -416,7 +416,7 @@ bool FTranslucencyDrawingPolicyFactory::DrawMesh(
 				// separate translucency is not updating scene color so we don't need to copy it multiple times.
 				// optimization:
 				// we should consider the same for non separate translucency (could cause artifacts but will be much faster)
-				DrawingContext.bSceneColorCopyIsUpToDate = DrawingContext.bSeparateTranslucencyPass;
+				DrawingContext.bSceneColorCopyIsUpToDate = (DrawingContext.TranslucenyPassType == TPT_SeparateTransluceny);
 			}
 		}
 
@@ -496,7 +496,7 @@ bool FTranslucencyDrawingPolicyFactory::DrawDynamicMesh(
 		View,
 		DrawingContext,
 		Mesh,
-		Mesh.Elements.Num()==1 ? 1 : (1<<Mesh.Elements.Num())-1,	// 1 bit set for each mesh element
+		Mesh.Elements.Num() == 1 ? 1 : (1 << Mesh.Elements.Num()) - 1,	// 1 bit set for each mesh element
 		bBackFace,
 		Mesh.DitheredLODTransitionAlpha,
 		bPreFog,
@@ -542,7 +542,7 @@ void FTranslucentPrimSet::DrawAPrimitive(
 	FRHICommandList& RHICmdList,
 	const FViewInfo& View,
 	FDeferredShadingSceneRenderer& Renderer,
-	bool bSeparateTranslucencyPass,
+	ETranslucencyPassType TranslucenyPassType,
 	int32 PrimIdx
 	) const
 {
@@ -550,7 +550,8 @@ void FTranslucentPrimSet::DrawAPrimitive(
 
 	// copy to reference for easier access
 	const TArray<FSortedPrim, SceneRenderingAllocator>& PhaseSortedPrimitives =
-		bSeparateTranslucencyPass ? SortedSeparateTranslucencyPrims : SortedPrims;
+		(TranslucenyPassType == TPT_SeparateTransluceny) ? SortedSeparateTranslucencyPrims : SortedPrims;
+
 	check(PrimIdx < PhaseSortedPrimitives.Num());
 
 	FPrimitiveSceneInfo* PrimitiveSceneInfo = PhaseSortedPrimitives[PrimIdx].PrimitiveSceneInfo;
@@ -559,9 +560,9 @@ void FTranslucentPrimSet::DrawAPrimitive(
 
 	checkSlow(ViewRelevance.HasTranslucency());
 
-	const FProjectedShadowInfo* TranslucentSelfShadow = Renderer.PrepareTranslucentShadowMap(RHICmdList, View, PrimitiveSceneInfo, bSeparateTranslucencyPass);
+	const FProjectedShadowInfo* TranslucentSelfShadow = Renderer.PrepareTranslucentShadowMap(RHICmdList, View, PrimitiveSceneInfo, TranslucenyPassType);
 
-	RenderPrimitive(RHICmdList, View, PrimitiveSceneInfo, ViewRelevance, TranslucentSelfShadow, bSeparateTranslucencyPass);
+	RenderPrimitive(RHICmdList, View, PrimitiveSceneInfo, ViewRelevance, TranslucentSelfShadow, TranslucenyPassType);
 }
 
 class FVolumetricTranslucentShadowRenderThreadTask
@@ -570,17 +571,17 @@ class FVolumetricTranslucentShadowRenderThreadTask
 	const FTranslucentPrimSet &PrimSet;
 	const FViewInfo& View;
 	FDeferredShadingSceneRenderer& Renderer;
-	bool bSeparateTranslucencyPass;
+	ETranslucencyPassType TranslucenyPassType;
 	int32 Index;
 
 public:
 
-	FORCEINLINE_DEBUGGABLE FVolumetricTranslucentShadowRenderThreadTask(FRHICommandList& InRHICmdList, const FTranslucentPrimSet& InPrimSet, const FViewInfo& InView, FDeferredShadingSceneRenderer& InRenderer, bool InbSeparateTranslucencyPass, int32 InIndex)
+	FORCEINLINE_DEBUGGABLE FVolumetricTranslucentShadowRenderThreadTask(FRHICommandList& InRHICmdList, const FTranslucentPrimSet& InPrimSet, const FViewInfo& InView, FDeferredShadingSceneRenderer& InRenderer, ETranslucencyPassType InTranslucenyPassType, int32 InIndex)
 		: RHICmdList(InRHICmdList)
 		, PrimSet(InPrimSet)
 		, View(InView)
 		, Renderer(InRenderer)
-		, bSeparateTranslucencyPass(InbSeparateTranslucencyPass)
+		, TranslucenyPassType(InTranslucenyPassType)
 		, Index(InIndex)
 	{
 	}
@@ -599,7 +600,7 @@ public:
 
 	void DoTask(ENamedThreads::Type CurrentThread, const FGraphEventRef& MyCompletionGraphEvent)
 	{
-		PrimSet.DrawAPrimitive(RHICmdList, View, Renderer, bSeparateTranslucencyPass, Index);
+		PrimSet.DrawAPrimitive(RHICmdList, View, Renderer, TranslucenyPassType, Index);
 	}
 };
 
@@ -607,7 +608,7 @@ void FTranslucentPrimSet::DrawPrimitivesParallel(
 	FRHICommandList& RHICmdList,
 	const FViewInfo& View,
 	FDeferredShadingSceneRenderer& Renderer,
-	bool bSeparateTranslucencyPass,
+	ETranslucencyPassType TranslucenyPassType,
 	int32 FirstIndex, int32 LastIndex
 	) const
 {
@@ -615,8 +616,10 @@ void FTranslucentPrimSet::DrawPrimitivesParallel(
 
 	// copy to reference for easier access
 	const TArray<FSortedPrim, SceneRenderingAllocator>& PhaseSortedPrimitives =
-		bSeparateTranslucencyPass ? SortedSeparateTranslucencyPrims : SortedPrims;
+		(TranslucenyPassType == TPT_SeparateTransluceny) ? SortedSeparateTranslucencyPrims : SortedPrims;
+
 	check(LastIndex < PhaseSortedPrimitives.Num());
+	
 	// Draw sorted scene prims
 	for (int32 PrimIdx = FirstIndex; PrimIdx <= LastIndex; PrimIdx++)
 	{
@@ -632,12 +635,12 @@ void FTranslucentPrimSet::DrawPrimitivesParallel(
 			// can't do this in parallel, defer
 			FRHICommandList* CmdList = new FRHICommandList;
 			CmdList->CopyRenderThreadContexts(RHICmdList);
-			FGraphEventRef RenderThreadCompletionEvent = TGraphTask<FVolumetricTranslucentShadowRenderThreadTask>::CreateTask().ConstructAndDispatchWhenReady(*CmdList, *this, View, Renderer, bSeparateTranslucencyPass, PrimIdx);
+			FGraphEventRef RenderThreadCompletionEvent = TGraphTask<FVolumetricTranslucentShadowRenderThreadTask>::CreateTask().ConstructAndDispatchWhenReady(*CmdList, *this, View, Renderer, TranslucenyPassType, PrimIdx);
 			RHICmdList.QueueRenderThreadCommandListSubmit(RenderThreadCompletionEvent, CmdList);
 		}
 		else
 		{
-			RenderPrimitive(RHICmdList, View, PrimitiveSceneInfo, ViewRelevance, nullptr, bSeparateTranslucencyPass);
+			RenderPrimitive(RHICmdList, View, PrimitiveSceneInfo, ViewRelevance, nullptr, TranslucenyPassType);
 		}
 	}
 }
@@ -646,30 +649,27 @@ void FTranslucentPrimSet::DrawPrimitives(
 	FRHICommandListImmediate& RHICmdList,
 	const FViewInfo& View,
 	FDeferredShadingSceneRenderer& Renderer,
-	bool bSeparateTranslucencyPass
+	ETranslucencyPassType TranslucenyPassType
 	) const
 {
 	const TArray<FSortedPrim,SceneRenderingAllocator>* PhaseSortedPrimitivesPtr = NULL;
 
 	// copy to reference for easier access
 	const TArray<FSortedPrim,SceneRenderingAllocator>& PhaseSortedPrimitives =
-		bSeparateTranslucencyPass ? SortedSeparateTranslucencyPrims : SortedPrims;
+		(TranslucenyPassType == TPT_SeparateTransluceny) ? SortedSeparateTranslucencyPrims : SortedPrims;
 
-	if( PhaseSortedPrimitives.Num() )
+	// Draw sorted scene prims
+	for( int32 PrimIdx = 0, Count = PhaseSortedPrimitives.Num(); PrimIdx < Count; PrimIdx++ )
 	{
-		// Draw sorted scene prims
-		for( int32 PrimIdx=0; PrimIdx < PhaseSortedPrimitives.Num(); PrimIdx++ )
-		{
-			FPrimitiveSceneInfo* PrimitiveSceneInfo = PhaseSortedPrimitives[PrimIdx].PrimitiveSceneInfo;
-			int32 PrimitiveId = PrimitiveSceneInfo->GetIndex();
-			const FPrimitiveViewRelevance& ViewRelevance = View.PrimitiveViewRelevanceMap[PrimitiveId];
+		FPrimitiveSceneInfo* PrimitiveSceneInfo = PhaseSortedPrimitives[PrimIdx].PrimitiveSceneInfo;
+		int32 PrimitiveId = PrimitiveSceneInfo->GetIndex();
+		const FPrimitiveViewRelevance& ViewRelevance = View.PrimitiveViewRelevanceMap[PrimitiveId];
 
-			checkSlow(ViewRelevance.HasTranslucency());
+		checkSlow(ViewRelevance.HasTranslucency());
 			
-			const FProjectedShadowInfo* TranslucentSelfShadow = Renderer.PrepareTranslucentShadowMap(RHICmdList, View, PrimitiveSceneInfo, bSeparateTranslucencyPass);
+		const FProjectedShadowInfo* TranslucentSelfShadow = Renderer.PrepareTranslucentShadowMap(RHICmdList, View, PrimitiveSceneInfo, TranslucenyPassType);
 
-			RenderPrimitive(RHICmdList, View, PrimitiveSceneInfo, ViewRelevance, TranslucentSelfShadow, bSeparateTranslucencyPass);
-		}
+		RenderPrimitive(RHICmdList, View, PrimitiveSceneInfo, ViewRelevance, TranslucentSelfShadow, TranslucenyPassType);
 	}
 
 	View.SimpleElementCollector.DrawBatchedElements(RHICmdList, View, FTexture2DRHIRef(), EBlendModeFilter::Translucent);
@@ -681,17 +681,17 @@ void FTranslucentPrimSet::RenderPrimitive(
 	FPrimitiveSceneInfo* PrimitiveSceneInfo, 
 	const FPrimitiveViewRelevance& ViewRelevance, 
 	const FProjectedShadowInfo* TranslucentSelfShadow,
-	bool bSeparateTranslucencyPass) const
+	ETranslucencyPassType TranslucenyPassType) const
 {
 	checkSlow(ViewRelevance.HasTranslucency());
 	auto FeatureLevel = View.GetFeatureLevel();
 
 	if(ViewRelevance.bDrawRelevance)
 	{
-		FTranslucencyDrawingPolicyFactory::ContextType Context(TranslucentSelfShadow, bSeparateTranslucencyPass);
+		FTranslucencyDrawingPolicyFactory::ContextType Context(TranslucentSelfShadow, TranslucenyPassType);
 
 		//@todo parallelrendering - come up with a better way to filter these by primitive
-		for (int32 MeshBatchIndex = 0; MeshBatchIndex < View.DynamicMeshElements.Num(); MeshBatchIndex++)
+		for (int32 MeshBatchIndex = 0, Count = View.DynamicMeshElements.Num(); MeshBatchIndex < Count; MeshBatchIndex++)
 		{
 			const FMeshBatchAndRelevance& MeshBatchAndRelevance = View.DynamicMeshElements[MeshBatchIndex];
 
@@ -706,18 +706,18 @@ void FTranslucentPrimSet::RenderPrimitive(
 		if( ViewRelevance.bStaticRelevance )
 		{
 			// Render static meshes from static scene prim
-			for( int32 StaticMeshIdx=0; StaticMeshIdx < PrimitiveSceneInfo->StaticMeshes.Num(); StaticMeshIdx++ )
+			for( int32 StaticMeshIdx = 0, Count = PrimitiveSceneInfo->StaticMeshes.Num(); StaticMeshIdx < Count; StaticMeshIdx++ )
 			{
 				FStaticMesh& StaticMesh = PrimitiveSceneInfo->StaticMeshes[StaticMeshIdx];
 				if (View.StaticMeshVisibilityMap[StaticMesh.Id]
 					// Only render static mesh elements using translucent materials
 					&& StaticMesh.IsTranslucent(FeatureLevel)
-					&& (StaticMesh.MaterialRenderProxy->GetMaterial(FeatureLevel)->IsSeparateTranslucencyEnabled() == bSeparateTranslucencyPass))
+					&& (StaticMesh.MaterialRenderProxy->GetMaterial(FeatureLevel)->IsSeparateTranslucencyEnabled() == (TranslucenyPassType == TPT_SeparateTransluceny)))
 				{
 					FTranslucencyDrawingPolicyFactory::DrawStaticMesh(
 						RHICmdList,
 						View,
-						FTranslucencyDrawingPolicyFactory::ContextType( TranslucentSelfShadow, bSeparateTranslucencyPass),
+						FTranslucencyDrawingPolicyFactory::ContextType( TranslucentSelfShadow, TranslucenyPassType),
 						StaticMesh,
 						StaticMesh.Elements.Num() == 1 ? 1 : View.StaticMeshBatchVisibility[StaticMesh.Id],
 						false,
@@ -845,7 +845,7 @@ class FDrawSortedTransAnyThreadTask
 	FDeferredShadingSceneRenderer& Renderer;
 	FRHICommandList& RHICmdList;
 	const FViewInfo& View;
-	bool bSeparateTranslucencyPass;
+	ETranslucencyPassType TranslucenyPassType;
 
 	const int32 FirstIndex;
 	const int32 LastIndex;
@@ -856,14 +856,14 @@ public:
 		FDeferredShadingSceneRenderer& InRenderer,
 		FRHICommandList& InRHICmdList,
 		const FViewInfo& InView,
-		bool InbSeparateTranslucencyPass,
+		ETranslucencyPassType InTranslucenyPassType,
 		int32 InFirstIndex,
 		int32 InLastIndex
 		)
 		: Renderer(InRenderer)
 		, RHICmdList(InRHICmdList)
 		, View(InView)
-		, bSeparateTranslucencyPass(InbSeparateTranslucencyPass)
+		, TranslucenyPassType(InTranslucenyPassType)
 		, FirstIndex(InFirstIndex)
 		, LastIndex(InLastIndex)
 	{
@@ -883,20 +883,20 @@ public:
 
 	void DoTask(ENamedThreads::Type CurrentThread, const FGraphEventRef& MyCompletionGraphEvent)
 	{
-		View.TranslucentPrimSet.DrawPrimitivesParallel(RHICmdList, View, Renderer, bSeparateTranslucencyPass, FirstIndex, LastIndex);
+		View.TranslucentPrimSet.DrawPrimitivesParallel(RHICmdList, View, Renderer, TranslucenyPassType, FirstIndex, LastIndex);
 		RHICmdList.HandleRTThreadTaskCompletion(MyCompletionGraphEvent);
 	}
 };
 
 class FTranslucencyPassParallelCommandListSet : public FParallelCommandListSet
 {
-	bool bSeparateTranslucency;
+	ETranslucencyPassType TranslucenyPassType;
 	bool bFirstTimeThisFrame;
 public:
-	FTranslucencyPassParallelCommandListSet(const FViewInfo& InView, FRHICommandListImmediate& InParentCmdList, bool bInParallelExecute, bool bInCreateSceneContext, bool bInSeparateTranslucency)
+	FTranslucencyPassParallelCommandListSet(const FViewInfo& InView, FRHICommandListImmediate& InParentCmdList, bool bInParallelExecute, bool bInCreateSceneContext, ETranslucencyPassType InTranslucenyPassType)
 		: FParallelCommandListSet(InView, InParentCmdList, bInParallelExecute, bInCreateSceneContext)
-		, bSeparateTranslucency(bInSeparateTranslucency)
-		, bFirstTimeThisFrame(bInSeparateTranslucency)
+		, TranslucenyPassType(InTranslucenyPassType)
+		, bFirstTimeThisFrame(InTranslucenyPassType == TPT_SeparateTransluceny)
 	{
 		SetStateOnCommandList(ParentCmdList);
 	}
@@ -908,7 +908,7 @@ public:
 
 	virtual void SetStateOnCommandList(FRHICommandList& CmdList) override
 	{
-		SetTranslucentRenderTargetAndState(CmdList, View, bSeparateTranslucency, bFirstTimeThisFrame);
+		SetTranslucentRenderTargetAndState(CmdList, View, TranslucenyPassType, bFirstTimeThisFrame);
 		bFirstTimeThisFrame = false;
 	}
 };
@@ -934,7 +934,7 @@ void FDeferredShadingSceneRenderer::RenderTranslucencyParallel(FRHICommandListIm
 
 
 	GParallelTranslucencyContext.TranslucentSelfShadow = nullptr;
-	GParallelTranslucencyContext.bSeparateTranslucencyPass = false;
+	GParallelTranslucencyContext.TranslucenyPassType = TPT_NonSeparateTransluceny;
 	GParallelTranslucencyContext.bSceneColorCopyIsUpToDate = false;
 	FScopedCommandListWaitForTasks Flusher(CVarRHICmdFlushRenderThreadTasksTranslucentPass.GetValueOnRenderThread() > 0 || CVarRHICmdFlushRenderThreadTasks.GetValueOnRenderThread() > 0, RHICmdList);
 
@@ -953,7 +953,7 @@ void FDeferredShadingSceneRenderer::RenderTranslucencyParallel(FRHICommandListIm
 			FTranslucencyPassParallelCommandListSet ParallelCommandListSet(View, RHICmdList, 
 				CVarRHICmdTranslucencyPassDeferredContexts.GetValueOnRenderThread() > 0, 
 				CVarRHICmdFlushRenderThreadTasksTranslucentPass.GetValueOnRenderThread() == 0  && CVarRHICmdFlushRenderThreadTasks.GetValueOnRenderThread() == 0, 
-				false);
+				TPT_NonSeparateTransluceny);
 
 			{
 				int32 NumPrims = View.TranslucentPrimSet.NumPrims() - View.TranslucentPrimSet.NumSeparateTranslucencyPrims();
@@ -976,7 +976,7 @@ void FDeferredShadingSceneRenderer::RenderTranslucencyParallel(FRHICommandListIm
 							FRHICommandList* CmdList = ParallelCommandListSet.NewParallelCommandList();
 
 							FGraphEventRef AnyThreadCompletionEvent = TGraphTask<FDrawSortedTransAnyThreadTask>::CreateTask(ParallelCommandListSet.GetPrereqs(), ENamedThreads::RenderThread)
-								.ConstructAndDispatchWhenReady(*this, *CmdList, View, false, Start, Last);
+								.ConstructAndDispatchWhenReady(*this, *CmdList, View, TPT_NonSeparateTransluceny, Start, Last);
 
 							ParallelCommandListSet.AddParallelCommandList(CmdList, AnyThreadCompletionEvent);
 						}
@@ -990,7 +990,7 @@ void FDeferredShadingSceneRenderer::RenderTranslucencyParallel(FRHICommandListIm
 			DrawViewElementsParallel<FTranslucencyDrawingPolicyFactory>(GParallelTranslucencyContext, SDPG_Foreground, false, ParallelCommandListSet);
 
 		}
-		FinishTranslucentRenderTarget(RHICmdList, View, false);
+		FinishTranslucentRenderTarget(RHICmdList, View, TPT_NonSeparateTransluceny);
 
 #if 0 // unsupported visualization in the parallel case
 		const FSceneViewState* ViewState = (const FSceneViewState*)View.State;
@@ -1011,7 +1011,7 @@ void FDeferredShadingSceneRenderer::RenderTranslucencyParallel(FRHICommandListIm
 					RHICmdList, 
 					CVarRHICmdTranslucencyPassDeferredContexts.GetValueOnRenderThread() > 0, 
 					CVarRHICmdFlushRenderThreadTasksTranslucentPass.GetValueOnRenderThread() == 0 && CVarRHICmdFlushRenderThreadTasks.GetValueOnRenderThread() == 0, 
-					true);
+					TPT_SeparateTransluceny);
 
 				// Draw only translucent prims that are in the SeparateTranslucency pass
 				if (View.TranslucentPrimSet.NumSeparateTranslucencyPrims() > 0)
@@ -1035,7 +1035,7 @@ void FDeferredShadingSceneRenderer::RenderTranslucencyParallel(FRHICommandListIm
 								FRHICommandList* CmdList = ParallelCommandListSet.NewParallelCommandList();
 
 								FGraphEventRef AnyThreadCompletionEvent = TGraphTask<FDrawSortedTransAnyThreadTask>::CreateTask(ParallelCommandListSet.GetPrereqs(), ENamedThreads::RenderThread)
-									.ConstructAndDispatchWhenReady(*this, *CmdList, View, true, Start, Last);
+									.ConstructAndDispatchWhenReady(*this, *CmdList, View, TPT_SeparateTransluceny, Start, Last);
 
 								ParallelCommandListSet.AddParallelCommandList(CmdList, AnyThreadCompletionEvent);
 							}
@@ -1075,10 +1075,10 @@ void FDeferredShadingSceneRenderer::RenderTranslucency(FRHICommandListImmediate&
 
 			const FViewInfo& View = Views[ViewIndex];
 
-			SetTranslucentRenderTargetAndState(RHICmdList, View, false);
+			SetTranslucentRenderTargetAndState(RHICmdList, View, TPT_NonSeparateTransluceny);
 
 			// Draw only translucent prims that don't read from scene color
-			View.TranslucentPrimSet.DrawPrimitives(RHICmdList, View, *this, false);
+			View.TranslucentPrimSet.DrawPrimitives(RHICmdList, View, *this, TPT_NonSeparateTransluceny);
 			// Draw the view's mesh elements with the translucent drawing policy.
 			DrawViewElements<FTranslucencyDrawingPolicyFactory>(RHICmdList, View, FTranslucencyDrawingPolicyFactory::ContextType(), SDPG_World, false);
 			// Draw the view's mesh elements with the translucent drawing policy.
@@ -1096,7 +1096,7 @@ void FDeferredShadingSceneRenderer::RenderTranslucency(FRHICommandListImmediate&
 				}
 			}
 
-			FinishTranslucentRenderTarget(RHICmdList, View, false);
+			FinishTranslucentRenderTarget(RHICmdList, View, TPT_NonSeparateTransluceny);
 			
 			{
 				FSceneRenderTargets& SceneContext = FSceneRenderTargets::Get(RHICmdList);
@@ -1113,7 +1113,7 @@ void FDeferredShadingSceneRenderer::RenderTranslucency(FRHICommandListImmediate&
 						RHICmdList.SetDepthStencilState(TStaticDepthStencilState<false, CF_DepthNearOrEqual>::GetRHI());
 					}
 
-					View.TranslucentPrimSet.DrawPrimitives(RHICmdList, View, *this, true);
+					View.TranslucentPrimSet.DrawPrimitives(RHICmdList, View, *this, TPT_SeparateTransluceny);
 				}
 
 				SceneContext.FinishRenderingSeparateTranslucency(RHICmdList, View);
diff --git a/Engine/Source/Runtime/Renderer/Private/TranslucentRendering.h b/Engine/Source/Runtime/Renderer/Private/TranslucentRendering.h
--- a/Engine/Source/Runtime/Renderer/Private/TranslucentRendering.h
+++ b/Engine/Source/Runtime/Renderer/Private/TranslucentRendering.h
@@ -17,15 +17,15 @@ public:
 	struct ContextType 
 	{
 		const FProjectedShadowInfo* TranslucentSelfShadow;
-		bool bSeparateTranslucencyPass;
+		ETranslucencyPassType TranslucenyPassType;
 		bool bSceneColorCopyIsUpToDate;
 		bool bPostAA;
 
-		ContextType(const FProjectedShadowInfo* InTranslucentSelfShadow = NULL, bool bSeparateTranslucencyPassIn = false, bool bPostAAIn = false)
-			: TranslucentSelfShadow(InTranslucentSelfShadow),
-			bSeparateTranslucencyPass(bSeparateTranslucencyPassIn),
-			bSceneColorCopyIsUpToDate(false),
-			bPostAA(bPostAAIn)
+		ContextType(const FProjectedShadowInfo* InTranslucentSelfShadow = NULL, ETranslucencyPassType InTranslucenyPassType = TPT_NonSeparateTransluceny, bool bPostAAIn = false)
+			: TranslucentSelfShadow(InTranslucentSelfShadow)
+			, TranslucenyPassType(InTranslucenyPassType)
+			, bSceneColorCopyIsUpToDate(false)
+			, bPostAA(bPostAAIn)
 		{}
 	};
 
