Only include shadow casting children in grouped shadow projection bounds. UE-15112

diff --git a/Engine/Source/Runtime/Renderer/Private/ShadowSetup.cpp b/Engine/Source/Runtime/Renderer/Private/ShadowSetup.cpp
--- a/Engine/Source/Runtime/Renderer/Private/ShadowSetup.cpp
+++ b/Engine/Source/Runtime/Renderer/Private/ShadowSetup.cpp
@@ -1342,7 +1342,10 @@ void FSceneRenderer::CreatePerObjectProjectedShadow(
 	for (int32 ChildIndex = 1; ChildIndex < ShadowGroupPrimitives.Num(); ChildIndex++)
 	{
 		const FPrimitiveSceneInfo* ShadowChild = ShadowGroupPrimitives[ChildIndex];
-		OriginalBounds = OriginalBounds + ShadowChild->Proxy->GetBounds();
+		if (ShadowChild->Proxy->CastsDynamicShadow())
+		{
+			OriginalBounds = OriginalBounds + ShadowChild->Proxy->GetBounds();
+		}
 	}
 	FSceneRenderTargets& SceneContext = FSceneRenderTargets::Get(RHICmdList);
 	
