removed unused code

diff --git a/Engine/Source/Runtime/Renderer/Private/ForwardBasePassRendering.cpp b/Engine/Source/Runtime/Renderer/Private/ForwardBasePassRendering.cpp
--- a/Engine/Source/Runtime/Renderer/Private/ForwardBasePassRendering.cpp
+++ b/Engine/Source/Runtime/Renderer/Private/ForwardBasePassRendering.cpp
@@ -338,7 +338,6 @@ bool FBasePassForwardOpaqueDrawingPolicyFactory::DrawDynamicMesh(
 {
 	// Determine the mesh's material and blend mode.
 	const auto FeatureLevel = View.GetFeatureLevel();
-	const auto ShaderPlatform = View.GetShaderPlatform();
 	const FMaterial* Material = Mesh.MaterialRenderProxy->GetMaterial(FeatureLevel);
 	const EBlendMode BlendMode = Material->GetBlendMode();
 
