fixed missing rendering of some translucent geometry, editor/debug material had SeparateTransluceny on, which often it shouldn't be

diff --git a/Engine/Source/Runtime/Renderer/Private/DeferredShadingRenderer.h b/Engine/Source/Runtime/Renderer/Private/DeferredShadingRenderer.h
--- a/Engine/Source/Runtime/Renderer/Private/DeferredShadingRenderer.h
+++ b/Engine/Source/Runtime/Renderer/Private/DeferredShadingRenderer.h
@@ -346,5 +346,7 @@ private:
 
 	void UpdateGlobalDistanceFieldObjectBuffers(FRHICommandListImmediate& RHICmdList);
 
+	void DrawAllTranslucencyPasses(FRHICommandListImmediate& RHICmdList, const FViewInfo& View, ETranslucencyPassType TranslucenyPassType);
+
 	friend class FTranslucentPrimSet;
 };
diff --git a/Engine/Source/Runtime/Renderer/Private/TranslucentRendering.cpp b/Engine/Source/Runtime/Renderer/Private/TranslucentRendering.cpp
--- a/Engine/Source/Runtime/Renderer/Private/TranslucentRendering.cpp
+++ b/Engine/Source/Runtime/Renderer/Private/TranslucentRendering.cpp
@@ -387,90 +387,86 @@ bool FTranslucencyDrawingPolicyFactory::DrawMesh(
 	// Determine the mesh's material and blend mode.
 	const FMaterial* Material = Mesh.MaterialRenderProxy->GetMaterial(FeatureLevel);
 	const EBlendMode BlendMode = Material->GetBlendMode();
-	const EMaterialShadingModel ShadingModel = Material->GetShadingModel();
 
-	// Only render translucent materials.
+	// Only render translucent materials
 	if(IsTranslucentBlendMode(BlendMode))
 	{
-		if (Material->IsSeparateTranslucencyEnabled() != (DrawingContext.TranslucenyPassType == TPT_SeparateTransluceny))
+		// if we are in relevant pass
+		if(Material->IsSeparateTranslucencyEnabled() == (DrawingContext.TranslucenyPassType == TPT_SeparateTransluceny))
 		{
-			return false;
-		}
+			if (Material->RequiresSceneColorCopy_RenderThread())
+			{
+				if (DrawingContext.bSceneColorCopyIsUpToDate == false)
+				{
+					if (!RHICmdList.Bypass() && !IsInRenderingThread())
+					{
+						FRHICommandList* CmdList = new FRHICommandList;
+						CmdList->CopyRenderThreadContexts(RHICmdList);
+						FGraphEventRef RenderThreadCompletionEvent = TGraphTask<FCopySceneColorAndRestoreRenderThreadTask>::CreateTask().ConstructAndDispatchWhenReady(*CmdList, View, PrimitiveSceneProxy);
+						RHICmdList.QueueRenderThreadCommandListSubmit(RenderThreadCompletionEvent, CmdList);
+					}
+					else
+					{
+						// otherwise, just do it now. We don't want to defer in this case because that can interfere with render target visualization (a debugging tool).
+						CopySceneColorAndRestore(RHICmdList, View, PrimitiveSceneProxy);
+					}
+					// todo: this optimization is currently broken
+					DrawingContext.bSceneColorCopyIsUpToDate = (DrawingContext.TranslucenyPassType == TPT_SeparateTransluceny);
+				}
+			}
 
-		if (Material->RequiresSceneColorCopy_RenderThread())
-		{
-			if (DrawingContext.bSceneColorCopyIsUpToDate == false)
+			const bool bDisableDepthTest = Material->ShouldDisableDepthTest();
+			const bool bEnableResponsiveAA = Material->ShouldEnableResponsiveAA();
+			// editor compositing not supported on translucent materials currently
+			const bool bEditorCompositeDepthTest = false;
+
+			// if this draw is coming postAA then there is probably no depth buffer (it's canvas) and bEnableResponsiveAA wont' do anything anyway.
+			if (bEnableResponsiveAA && !DrawingContext.bPostAA)
 			{
-				if (!RHICmdList.Bypass() && !IsInRenderingThread())
+				if( bDisableDepthTest )
 				{
-					FRHICommandList* CmdList = new FRHICommandList;
-					CmdList->CopyRenderThreadContexts(RHICmdList);
-					FGraphEventRef RenderThreadCompletionEvent = 					TGraphTask<FCopySceneColorAndRestoreRenderThreadTask>::CreateTask().ConstructAndDispatchWhenReady(*CmdList, View, PrimitiveSceneProxy);
-					RHICmdList.QueueRenderThreadCommandListSubmit(RenderThreadCompletionEvent, CmdList);
+					RHICmdList.SetDepthStencilState(TStaticDepthStencilState<false,CF_Always,true,CF_Always,SO_Keep,SO_Keep,SO_Replace>::GetRHI(), 1);
 				}
 				else
 				{
-					// otherwise, just do it now. We don't want to defer in this case because that can interfere with render target visualization (a debugging tool).
-					CopySceneColorAndRestore(RHICmdList, View, PrimitiveSceneProxy);
+					RHICmdList.SetDepthStencilState(TStaticDepthStencilState<false,CF_DepthNearOrEqual,true,CF_Always,SO_Keep,SO_Keep,SO_Replace>::GetRHI(), 1);
 				}
-				// separate translucency is not updating scene color so we don't need to copy it multiple times.
-				// optimization:
-				// we should consider the same for non separate translucency (could cause artifacts but will be much faster)
-				DrawingContext.bSceneColorCopyIsUpToDate = (DrawingContext.TranslucenyPassType == TPT_SeparateTransluceny);
 			}
-		}
-
-		const bool bDisableDepthTest = Material->ShouldDisableDepthTest();
-		const bool bEnableResponsiveAA = Material->ShouldEnableResponsiveAA();
-		// editor compositing not supported on translucent materials currently
-		const bool bEditorCompositeDepthTest = false;
-
-		// if this draw is coming postAA then there is probably no depth buffer (it's canvas) and bEnableResponsiveAA wont' do anything anyway.
-		if (bEnableResponsiveAA && !DrawingContext.bPostAA)
-		{
-			if( bDisableDepthTest )
+			else if( bDisableDepthTest )
 			{
-				RHICmdList.SetDepthStencilState(TStaticDepthStencilState<false,CF_Always,true,CF_Always,SO_Keep,SO_Keep,SO_Replace>::GetRHI(), 1);
+				RHICmdList.SetDepthStencilState(TStaticDepthStencilState<false,CF_Always>::GetRHI());
 			}
-			else
+
+			ProcessBasePassMesh(
+				RHICmdList, 
+				FProcessBasePassMeshParameters(
+					Mesh,
+					BatchElementMask,
+					Material,
+					PrimitiveSceneProxy, 
+					!bPreFog,
+					bEditorCompositeDepthTest,
+					ESceneRenderTargetsMode::SetTextures,
+					FeatureLevel
+				),
+				FDrawTranslucentMeshAction(
+					View,
+					bBackFace,
+					DitheredLODTransitionValue,
+					HitProxyId,
+					DrawingContext.TranslucentSelfShadow,
+					PrimitiveSceneProxy && PrimitiveSceneProxy->CastsVolumetricTranslucentShadow()
+				)
+			);
+
+			if (bDisableDepthTest || bEnableResponsiveAA)
 			{
-				RHICmdList.SetDepthStencilState(TStaticDepthStencilState<false,CF_DepthNearOrEqual,true,CF_Always,SO_Keep,SO_Keep,SO_Replace>::GetRHI(), 1);
+				// Restore default depth state
+				RHICmdList.SetDepthStencilState(TStaticDepthStencilState<false,CF_DepthNearOrEqual>::GetRHI());
 			}
-		}
-		else if( bDisableDepthTest )
-		{
-			RHICmdList.SetDepthStencilState(TStaticDepthStencilState<false,CF_Always>::GetRHI());
-		}
 
-		ProcessBasePassMesh(
-			RHICmdList, 
-			FProcessBasePassMeshParameters(
-				Mesh,
-				BatchElementMask,
-				Material,
-				PrimitiveSceneProxy, 
-				!bPreFog,
-				bEditorCompositeDepthTest,
-				ESceneRenderTargetsMode::SetTextures,
-				FeatureLevel
-			),
-			FDrawTranslucentMeshAction(
-				View,
-				bBackFace,
-				DitheredLODTransitionValue,
-				HitProxyId,
-				DrawingContext.TranslucentSelfShadow,
-				PrimitiveSceneProxy && PrimitiveSceneProxy->CastsVolumetricTranslucentShadow()
-			)
-		);
-
-		if (bDisableDepthTest || bEnableResponsiveAA)
-		{
-			// Restore default depth state
-			RHICmdList.SetDepthStencilState(TStaticDepthStencilState<false,CF_DepthNearOrEqual>::GetRHI());
+			bDirty = true;
 		}
-
-		bDirty = true;
 	}
 	return bDirty;
 }
@@ -1056,6 +1052,18 @@ static TAutoConsoleVariable<int32> CVarParallelTranslucency(
 	ECVF_RenderThreadSafe
 	);
 
+void FDeferredShadingSceneRenderer::DrawAllTranslucencyPasses(FRHICommandListImmediate& RHICmdList, const FViewInfo& View, ETranslucencyPassType TranslucenyPassType)
+{
+	// Draw translucent prims
+	View.TranslucentPrimSet.DrawPrimitives(RHICmdList, View, *this, TranslucenyPassType);
+
+	FTranslucencyDrawingPolicyFactory::ContextType Context(0, TranslucenyPassType);
+
+	// editor and debug rendering
+	DrawViewElements<FTranslucencyDrawingPolicyFactory>(RHICmdList, View, Context, SDPG_World, false);
+	DrawViewElements<FTranslucencyDrawingPolicyFactory>(RHICmdList, View, Context, SDPG_Foreground, false);
+}
+
 void FDeferredShadingSceneRenderer::RenderTranslucency(FRHICommandListImmediate& RHICmdList)
 {
 	if (ShouldRenderTranslucency())
@@ -1075,36 +1083,39 @@ void FDeferredShadingSceneRenderer::RenderTranslucency(FRHICommandListImmediate&
 
 			const FViewInfo& View = Views[ViewIndex];
 
-			SetTranslucentRenderTargetAndState(RHICmdList, View, TPT_NonSeparateTransluceny);
-
-			// Draw only translucent prims that don't read from scene color
-			View.TranslucentPrimSet.DrawPrimitives(RHICmdList, View, *this, TPT_NonSeparateTransluceny);
-			// Draw the view's mesh elements with the translucent drawing policy.
-			DrawViewElements<FTranslucencyDrawingPolicyFactory>(RHICmdList, View, FTranslucencyDrawingPolicyFactory::ContextType(), SDPG_World, false);
-			// Draw the view's mesh elements with the translucent drawing policy.
-			DrawViewElements<FTranslucencyDrawingPolicyFactory>(RHICmdList, View, FTranslucencyDrawingPolicyFactory::ContextType(), SDPG_Foreground, false);
+			// non separate translucency
+			{
+				SetTranslucentRenderTargetAndState(RHICmdList, View, TPT_NonSeparateTransluceny);
 
-			const FSceneViewState* ViewState = (const FSceneViewState*)View.State;
+				DrawAllTranslucencyPasses(RHICmdList, View, TPT_NonSeparateTransluceny);
 
-			if(ViewState && View.Family->EngineShowFlags.VisualizeLPV)
-			{
-				FLightPropagationVolume* LightPropagationVolume = ViewState->GetLightPropagationVolume(View.GetFeatureLevel());
+				const FSceneViewState* ViewState = (const FSceneViewState*)View.State;
 
-				if (LightPropagationVolume)
+				if (ViewState && View.Family->EngineShowFlags.VisualizeLPV)
 				{
-					LightPropagationVolume->Visualise(RHICmdList, View);
+					FLightPropagationVolume* LightPropagationVolume = ViewState->GetLightPropagationVolume(View.GetFeatureLevel());
+
+					if (LightPropagationVolume)
+					{
+						LightPropagationVolume->Visualise(RHICmdList, View);
+					}
 				}
-			}
 
-			FinishTranslucentRenderTarget(RHICmdList, View, TPT_NonSeparateTransluceny);
+				FinishTranslucentRenderTarget(RHICmdList, View, TPT_NonSeparateTransluceny);
+			}
 			
+			// separate translucency
 			{
 				FSceneRenderTargets& SceneContext = FSceneRenderTargets::Get(RHICmdList);
-				bool bRenderSeparateTranslucency = View.TranslucentPrimSet.NumSeparateTranslucencyPrims() > 0;
-
+				
 				// always call BeginRenderingSeparateTranslucency() even if there are no primitives to we keep the RT allocated
 				bool bSetupTranslucency = SceneContext.BeginRenderingSeparateTranslucency(RHICmdList, View, true);
 
+				const TIndirectArray<FMeshBatch>& WorldList = View.ViewMeshElements;
+				const TIndirectArray<FMeshBatch>& ForegroundList = View.TopViewMeshElements;
+
+				bool bRenderSeparateTranslucency = View.TranslucentPrimSet.NumSeparateTranslucencyPrims() > 0 || WorldList.Num() || ForegroundList.Num();
+
 				// Draw only translucent prims that are in the SeparateTranslucency pass
 				if (bRenderSeparateTranslucency)
 				{
@@ -1112,8 +1123,8 @@ void FDeferredShadingSceneRenderer::RenderTranslucency(FRHICommandListImmediate&
 					{
 						RHICmdList.SetDepthStencilState(TStaticDepthStencilState<false, CF_DepthNearOrEqual>::GetRHI());
 					}
-
-					View.TranslucentPrimSet.DrawPrimitives(RHICmdList, View, *this, TPT_SeparateTransluceny);
+					
+					DrawAllTranslucencyPasses(RHICmdList, View, TPT_SeparateTransluceny);
 				}
 
 				SceneContext.FinishRenderingSeparateTranslucency(RHICmdList, View);
