removed unused code

diff --git a/Engine/Source/Runtime/Renderer/Private/TranslucentRendering.cpp b/Engine/Source/Runtime/Renderer/Private/TranslucentRendering.cpp
--- a/Engine/Source/Runtime/Renderer/Private/TranslucentRendering.cpp
+++ b/Engine/Source/Runtime/Renderer/Private/TranslucentRendering.cpp
@@ -542,9 +542,6 @@ void FTranslucentPrimSet::DrawAPrimitive(
 	int32 PrimIdx
 	) const
 {
-	const TArray<FSortedPrim, SceneRenderingAllocator>* PhaseSortedPrimitivesPtr = NULL;
-
-	// copy to reference for easier access
 	const TArray<FSortedPrim, SceneRenderingAllocator>& PhaseSortedPrimitives =
 		(TranslucenyPassType == TPT_SeparateTransluceny) ? SortedSeparateTranslucencyPrims : SortedPrims;
 
@@ -608,9 +605,6 @@ void FTranslucentPrimSet::DrawPrimitivesParallel(
 	int32 FirstIndex, int32 LastIndex
 	) const
 {
-	const TArray<FSortedPrim, SceneRenderingAllocator>* PhaseSortedPrimitivesPtr = NULL;
-
-	// copy to reference for easier access
 	const TArray<FSortedPrim, SceneRenderingAllocator>& PhaseSortedPrimitives =
 		(TranslucenyPassType == TPT_SeparateTransluceny) ? SortedSeparateTranslucencyPrims : SortedPrims;
 
@@ -648,9 +642,6 @@ void FTranslucentPrimSet::DrawPrimitives(
 	ETranslucencyPassType TranslucenyPassType
 	) const
 {
-	const TArray<FSortedPrim,SceneRenderingAllocator>* PhaseSortedPrimitivesPtr = NULL;
-
-	// copy to reference for easier access
 	const TArray<FSortedPrim,SceneRenderingAllocator>& PhaseSortedPrimitives =
 		(TranslucenyPassType == TPT_SeparateTransluceny) ? SortedSeparateTranslucencyPrims : SortedPrims;
 
