refactor CircleDOF code, improved VisualizeDOF

diff --git a/Engine/Shaders/PostProcessCircleDOF.usf b/Engine/Shaders/PostProcessCircleDOF.usf
new file mode 100644
--- /dev/null
+++ b/Engine/Shaders/PostProcessCircleDOF.usf
@@ -0,0 +1,554 @@
+// Copyright 1998-2015 Epic Games, Inc. All Rights Reserved.
+
+/*=============================================================================
+	PostProcessCircleDOF.usf: PostProcessing Circle Depth of Field
+=============================================================================*/
+
+#include "Common.usf"
+#include "PostProcessCommon.usf"
+#include "DeferredShadingCommon.usf"		// FGBufferData
+#include "DepthOfFieldCommon.usf"
+
+
+// {radius, depth blur amp, depth blur radius, ViewPortScale}
+float4 CircleDofParams;
+
+// CIRCLE DOF: Compute circle of confusion size in pixels.
+float DepthToCoc(float Depth) 
+{
+	float Focus = View.DepthOfFieldFocalDistance;
+	float Radius = CircleDofParams.x;
+	float CocRadius = ((Depth - Focus) / Depth) * Radius;
+	float DepthBlurRadius = (1.0 - exp2(-Depth * CircleDofParams.y)) * CircleDofParams.z;
+	float ReturnCoc = max(abs(CocRadius), DepthBlurRadius);
+	if(CocRadius < 0.0) 
+	{
+		// near CoC is using negative values
+		ReturnCoc = -ReturnCoc;
+	}
+	return ReturnCoc;
+}
+
+// pixel shader entry point
+void CircleSetupPS(float4 UVAndScreenPos : TEXCOORD0, out float4 OutColor0 : SV_Target0)
+{
+	float2 UV = UVAndScreenPos.xy;
+
+	float4 DepthQuad = GatherSceneDepth(UV, PostprocessInput1Size.zw);
+
+	UV = UVAndScreenPos.xy - 0.5*PostprocessInput0Size.zw;
+
+	float4 CW = PostprocessInput0.SampleLevel(PostprocessInput0Sampler, UV.xy, 0);
+	float4 CZ = PostprocessInput0.SampleLevel(PostprocessInput0Sampler, UV.xy, 0, int2(1,0));
+	float4 CX = PostprocessInput0.SampleLevel(PostprocessInput0Sampler, UV.xy, 0, int2(0,1));
+	float4 CY = PostprocessInput0.SampleLevel(PostprocessInput0Sampler, UV.xy, 0, int2(1,1));
+
+	float4 CocQuad = float4(DepthToCoc(DepthQuad.x), DepthToCoc(DepthQuad.y), DepthToCoc(DepthQuad.z), DepthToCoc(DepthQuad.w));
+
+	// Doing a max depth reduction (erode the foreground). Less correct, but less artifacts.
+	// Perhaps need to re-open this in the future.
+
+	#if 0
+		// Stuff max radius in alpha.
+		OutColor0.a = max(max(CocQuad.x,CocQuad.y),max(CocQuad.z,CocQuad.w));
+	#else
+		// This in theory is better but causes bleeding artifacts with temporal AA..
+		// This is important otherwise near thin objects disappear (leaves clamping artifacts in recombined pass).
+		OutColor0.a = CocQuad.x;
+		if(abs(OutColor0.a) > CocQuad.y) OutColor0.a = CocQuad.y;
+		if(abs(OutColor0.a) > CocQuad.z) OutColor0.a = CocQuad.z;
+		if(abs(OutColor0.a) > CocQuad.w) OutColor0.a = CocQuad.w;
+	#endif
+
+	// Remove samples which are outside the size.
+	// TODO: Tune the ScaleFactor.
+	float ScaleFactor = 64.0;
+	float4 W = float4(
+		1.0-saturate(abs(OutColor0.a - CocQuad.x) * ScaleFactor),
+		1.0-saturate(abs(OutColor0.a - CocQuad.y) * ScaleFactor),
+		1.0-saturate(abs(OutColor0.a - CocQuad.z) * ScaleFactor),
+		1.0-saturate(abs(OutColor0.a - CocQuad.w) * ScaleFactor));
+
+	OutColor0.rgb = (1.0/(W.x+W.y+W.z+W.w)) * (CX.rgb*W.x + CY.rgb*W.y + CZ.rgb*W.z + CW.rgb*W.w);
+}
+
+
+
+// {0 to 1} output.
+float NoizNorm(float2 N, float X)
+{
+	N+=X;
+	return frac(sin(dot(N.xy,float2(12.9898, 78.233)))*43758.5453);
+}
+
+// {-1 to 1} output.
+float NoizSnorm(float2 N, float X)
+{
+	return NoizNorm(N,X)*2.0-1.0;
+}
+
+float2 RotVec(float Radius, float Radians)
+{
+	return Radius * float2(cos(Radians), sin(Radians));
+}
+
+float2 RandomOffset;
+
+
+
+
+float Min4(float4 A)
+{
+	return min(min(A.x,A.y),min(A.z,A.w));
+}
+
+float Min16(float4 A, float4 B, float4 C, float4 D)
+{
+	return min(min(Min4(A),Min4(B)),min(Min4(C),Min4(D)));
+}
+
+// This does a 2x2:1 reduction with a 4x4:1 dilation.
+void CircleDilatePS(float4 UVAndScreenPos : TEXCOORD0, out float OutColor : SV_Target0)
+{
+	// Sampling pattern (each gather4)
+    //   d g
+	//   j M  (M={0,0} point)
+
+	#if COMPILER_GLSL || COMPILER_GLSL_ES2
+		float2 UV = UVAndScreenPos.xy + 0.5*PostprocessInput0Size.zw;
+		// This leverages nearest sampling (bilinear won't work).
+		// Probably not the best way to do this.
+		float4 Sd, Sg, Sj, Sm;
+		Sd.x = PostprocessInput0.SampleLevel(PostprocessInput0Sampler, UV.xy, 0, int2(-2,-2)).a;
+		Sd.y = PostprocessInput0.SampleLevel(PostprocessInput0Sampler, UV.xy, 0, int2(-1,-2)).a;
+		Sd.z = PostprocessInput0.SampleLevel(PostprocessInput0Sampler, UV.xy, 0, int2(-2,-1)).a;
+		Sd.w = PostprocessInput0.SampleLevel(PostprocessInput0Sampler, UV.xy, 0, int2(-1,-1)).a;
+		Sg.x = PostprocessInput0.SampleLevel(PostprocessInput0Sampler, UV.xy, 0, int2(0,-2)).a;
+		Sg.y = PostprocessInput0.SampleLevel(PostprocessInput0Sampler, UV.xy, 0, int2(1,-2)).a;
+		Sg.z = PostprocessInput0.SampleLevel(PostprocessInput0Sampler, UV.xy, 0, int2(0,-1)).a;
+		Sg.w = PostprocessInput0.SampleLevel(PostprocessInput0Sampler, UV.xy, 0, int2(1,-1)).a;
+		Sj.x = PostprocessInput0.SampleLevel(PostprocessInput0Sampler, UV.xy, 0, int2(-2,0)).a;
+		Sj.y = PostprocessInput0.SampleLevel(PostprocessInput0Sampler, UV.xy, 0, int2(-1,0)).a;
+		Sj.z = PostprocessInput0.SampleLevel(PostprocessInput0Sampler, UV.xy, 0, int2(-2,1)).a;
+		Sj.w = PostprocessInput0.SampleLevel(PostprocessInput0Sampler, UV.xy, 0, int2(-1,1)).a;
+		Sm.x = PostprocessInput0.SampleLevel(PostprocessInput0Sampler, UV.xy, 0, int2(0,0)).a;
+		Sm.y = PostprocessInput0.SampleLevel(PostprocessInput0Sampler, UV.xy, 0, int2(1,0)).a;
+		Sm.z = PostprocessInput0.SampleLevel(PostprocessInput0Sampler, UV.xy, 0, int2(0,1)).a;
+		Sm.w = PostprocessInput0.SampleLevel(PostprocessInput0Sampler, UV.xy, 0, int2(1,1)).a;
+	#else
+		float2 UV = UVAndScreenPos.xy + 1.0*PostprocessInput0Size.zw;
+		float4 Sd = PostprocessInput0.GatherAlpha(PostprocessInput0Sampler, UV, int2(-2,-2));
+		float4 Sg = PostprocessInput0.GatherAlpha(PostprocessInput0Sampler, UV, int2(0,-2));
+		float4 Sj = PostprocessInput0.GatherAlpha(PostprocessInput0Sampler, UV, int2(-2,0));
+		float4 Sm = PostprocessInput0.GatherAlpha(PostprocessInput0Sampler, UV, int2(0,0));
+	#endif // COMPILER_GLSL
+	
+	// Make sure near is only near blur.
+	OutColor = min(0.0, Min16(Sd, Sg, Sj, Sm));
+}
+
+
+// pixel shader entry point
+void CirclePS(float4 UVAndScreenPos : TEXCOORD0, out float4 OutColor0 : SV_Target0)
+{
+	float2 UV = UVAndScreenPos.xy;
+
+	//
+	// Pass 0
+	// Dilate near minimum CoC (near CoC is negative values).
+	//
+
+	// Fixed maximum search size (in terms of Circle of Confusion radius).
+	// Higher than 8 is too noizy for 4 samples.
+	float Coc = 8.0 * CircleDofParams.w;
+
+	// Get base semi-random direction and dither along radius.
+	// Reused throughout the rest of the algorithm.
+	float TwoPi = 2.0 * 3.14159;
+	float RadianBase = NoizSnorm(UVAndScreenPos.xy, 0.010 * RandomOffset.x) * TwoPi;
+	float RadiusBase = NoizNorm(UVAndScreenPos.xy, 0.013 * RandomOffset.x);
+
+	// Radius 
+	float RadiusBase2 = RadiusBase * (1.0/4.0);
+	float R1 = sqrt(RadiusBase2 + 3.0/4.0) * Coc;
+	float R2 = sqrt(RadiusBase2 + 2.0/4.0) * Coc;
+	float R3 = sqrt(RadiusBase2 + 1.0/4.0) * Coc;
+	float R4 = sqrt(RadiusBase2 + 0.0/4.0) * Coc;
+
+	float2 UV1 = RotVec(R1, RadianBase + TwoPi * 0.0/4.0);
+	float2 UV2 = RotVec(R2, RadianBase + TwoPi * 2.0/4.0);
+	float2 UV3 = RotVec(R3, RadianBase + TwoPi * 1.0/4.0);
+	float2 UV4 = RotVec(R4, RadianBase + TwoPi * 3.0/4.0);
+
+	UV1 = UVAndScreenPos.xy + UV1 * PostprocessInput0Size.zw;
+	UV2 = UVAndScreenPos.xy + UV2 * PostprocessInput0Size.zw;
+	UV3 = UVAndScreenPos.xy + UV3 * PostprocessInput0Size.zw;
+	UV4 = UVAndScreenPos.xy + UV4 * PostprocessInput0Size.zw;
+
+	float D1 = PostprocessInput1.SampleLevel(PostprocessInput1Sampler, UV1, 0).x;
+	float D2 = PostprocessInput1.SampleLevel(PostprocessInput1Sampler, UV2, 0).x;
+	float D3 = PostprocessInput1.SampleLevel(PostprocessInput1Sampler, UV3, 0).x;
+	float D4 = PostprocessInput1.SampleLevel(PostprocessInput1Sampler, UV4, 0).x;
+
+	float NearCoc = 65536.0;
+	float Feather = 2.0;
+	if(abs(D1)+Feather > R1) NearCoc = min(NearCoc, D1);
+	if(abs(D2)+Feather > R2) NearCoc = min(NearCoc, D2);
+	if(abs(D3)+Feather > R3) NearCoc = min(NearCoc, D3);
+	if(abs(D4)+Feather > R4) NearCoc = min(NearCoc, D4);
+
+
+	//
+	// Pass 1
+	//
+
+	// Going to grab sets of 4 samples per pass.
+	// Each set of 4 samples can be a smaller circle of confusion
+	// (aka can be in-front of the larger background).
+
+	// Setup for 12 samples (3 passes of 4 samples).
+	RadiusBase *= (1.0/11.5);
+
+	// Grab circle of confusion for the pixel and pixel color.
+	OutColor0 = Texture2DSampleLevel(PostprocessInput0, PostprocessInput0Sampler, UV, 0);
+	float FarCoc = OutColor0.a;
+
+	// Fix in case no near exists.
+	NearCoc = min(NearCoc, FarCoc);
+
+	// Used for sample pattern.
+	Coc = max(abs(FarCoc),abs(NearCoc));
+	
+	// Bring out to the smaller radius of sample sets.
+	// This has the highest chance of seeing a smaller overlapping CoC.
+	R1 = (RadiusBase+9.0/11.5) * Coc;
+	R2 = (RadiusBase+3.0/11.5) * Coc;
+	R3 = (RadiusBase+6.0/11.5) * Coc;
+	R4 = (RadiusBase+0.0/11.5) * Coc;
+
+	// Ensure at least getting different sample than center pixel.
+	float R1a = max(1.0,R1);
+	float R2a = max(1.0,R2);
+	float R3a = max(1.0,R3);
+	float R4a = max(1.0,R4);
+
+	UV1 = RotVec(R1a, RadianBase + TwoPi * 0.0/12.0);
+	UV2 = RotVec(R2a, RadianBase + TwoPi * 3.0/12.0);
+	UV3 = RotVec(R3a, RadianBase + TwoPi * 6.0/12.0);
+	UV4 = RotVec(R4a, RadianBase + TwoPi * 9.0/12.0);
+
+	UV1 = UVAndScreenPos.xy + UV1 * PostprocessInput0Size.zw;
+	UV2 = UVAndScreenPos.xy + UV2 * PostprocessInput0Size.zw;
+	UV3 = UVAndScreenPos.xy + UV3 * PostprocessInput0Size.zw;
+	UV4 = UVAndScreenPos.xy + UV4 * PostprocessInput0Size.zw;
+
+	float4 C1 = PostprocessInput0.SampleLevel(PostprocessInput0Sampler, UV1, 0);
+	float4 C2 = PostprocessInput0.SampleLevel(PostprocessInput0Sampler, UV2, 0);
+	float4 C3 = PostprocessInput0.SampleLevel(PostprocessInput0Sampler, UV3, 0);
+	float4 C4 = PostprocessInput0.SampleLevel(PostprocessInput0Sampler, UV4, 0);
+
+	// Base weight works around the max(1.0,radius) constraint.
+	// Base weight also shapes to weight higher on the outside radius.
+	float W0 = 1.0 - saturate(Coc);
+	float W1 = R1;
+	float W2 = R2;
+	float W3 = R3;
+	float W4 = R4;
+
+	// Intersection weight: 0=sample does not intersect pixel, to 1=sample intersects.
+	// TODO: Tune feather factors.
+	float IFeather0 = 1.0/4.0;
+	float I1 = saturate((abs(C1.a) - R1) * IFeather0);
+	float I2 = saturate((abs(C2.a) - R2) * IFeather0);
+	float I3 = saturate((abs(C3.a) - R3) * IFeather0);
+	float I4 = saturate((abs(C4.a) - R4) * IFeather0);
+
+	// Check if have a more near intersecting Coc for next pass.
+	float FarCoc2 = FarCoc;
+	if(I1*W1 > 0.0) FarCoc2 = min(FarCoc2, C1.a);
+	if(I2*W2 > 0.0) FarCoc2 = min(FarCoc2, C2.a);
+	if(I3*W3 > 0.0) FarCoc2 = min(FarCoc2, C3.a);
+	if(I4*W4 > 0.0) FarCoc2 = min(FarCoc2, C4.a);
+
+	// Fully ignore intersection weight when in nearfield blur
+	// and sample average CoC is 50% between near and far CoC neighborhood.
+	float AvgCoc = (FarCoc + C1.a + C2.a + C3.a + C4.a) * (1.0/5.0);
+	// Get dilated far.
+	FarCoc = max(FarCoc, max(max(C1.a, C2.a),max(C3.a, C4.a)));
+	// Controls the transition between states.
+	float IFeather1 = 1.0;
+	float IFeather2 = 2.0;
+	float Ignore = saturate(-NearCoc * IFeather1) * saturate(((AvgCoc - FarCoc) / (NearCoc - FarCoc)) * IFeather2);
+
+	W1 *= lerp(I1, 1.0, Ignore);
+	W2 *= lerp(I2, 1.0, Ignore);
+	W3 *= lerp(I3, 1.0, Ignore);
+	W4 *= lerp(I4, 1.0, Ignore);
+
+	// Make sure at least something is not zero.
+	W0 += 1.0/65536.0;
+
+	// Start weighted accumulation.
+	OutColor0.rgb = OutColor0.rgb * W0 + C1.rgb * W1 + C2.rgb * W2 + C3.rgb * W3 + C4.rgb * W4;
+	float Weight = W0+W1+W2+W3+W4;
+
+	// Set current result as possible background.
+	float3 Background = OutColor0.rgb * (1.0/Weight);
+
+
+	//
+	// Pass 2
+	//
+
+	// Drop weight of existing pass if Coc changes too much.
+	float Coc2 = max(abs(FarCoc2),abs(NearCoc));
+	float Drop = (1.0/65536.0) + 1.0 - saturate(abs(Coc - Coc2));
+	OutColor0.rgb *= Drop;
+	Weight *= Drop;
+
+	R1 = (RadiusBase+10.0/11.5) * Coc2;
+	R2 = (RadiusBase+ 4.0/11.5) * Coc2;
+	R3 = (RadiusBase+ 7.0/11.5) * Coc2;
+	R4 = (RadiusBase+ 1.0/11.5) * Coc2;
+
+	R1a = max(1.0,R1);
+	R2a = max(1.0,R2);
+	R3a = max(1.0,R3);
+	R4a = max(1.0,R4);
+
+	UV1 = RotVec(R1a, RadianBase + TwoPi *  8.0/12.0);
+	UV2 = RotVec(R2a, RadianBase + TwoPi * 11.0/12.0);
+	UV3 = RotVec(R3a, RadianBase + TwoPi *  2.0/12.0);
+	UV4 = RotVec(R4a, RadianBase + TwoPi *  5.0/12.0);
+
+	UV1 = UVAndScreenPos.xy + UV1 * PostprocessInput0Size.zw;
+	UV2 = UVAndScreenPos.xy + UV2 * PostprocessInput0Size.zw;
+	UV3 = UVAndScreenPos.xy + UV3 * PostprocessInput0Size.zw;
+	UV4 = UVAndScreenPos.xy + UV4 * PostprocessInput0Size.zw;
+
+	C1 = PostprocessInput0.SampleLevel(PostprocessInput0Sampler, UV1, 0);
+	C2 = PostprocessInput0.SampleLevel(PostprocessInput0Sampler, UV2, 0);
+	C3 = PostprocessInput0.SampleLevel(PostprocessInput0Sampler, UV3, 0);
+	C4 = PostprocessInput0.SampleLevel(PostprocessInput0Sampler, UV4, 0);
+
+	// Lerp to background if outside possibly smaller CoC.
+	C1.rgb = lerp(C1.rgb, Background.rgb, saturate(abs(C1.a) - Coc2));
+	C2.rgb = lerp(C2.rgb, Background.rgb, saturate(abs(C2.a) - Coc2));
+	C3.rgb = lerp(C3.rgb, Background.rgb, saturate(abs(C3.a) - Coc2));
+	C4.rgb = lerp(C4.rgb, Background.rgb, saturate(abs(C4.a) - Coc2));
+
+	W1 = R1;
+	W2 = R2;
+	W3 = R3;
+	W4 = R4;
+
+	// Intersection weight: 0=sample does not intersect pixel, to 1=sample intersects.
+	I1 = saturate((abs(C1.a) - R1) * IFeather0);
+	I2 = saturate((abs(C2.a) - R2) * IFeather0);
+	I3 = saturate((abs(C3.a) - R3) * IFeather0);
+	I4 = saturate((abs(C4.a) - R4) * IFeather0);
+
+	// Check if have a more near intersecting Coc for next pass.
+	float FarCoc3 = FarCoc2;
+	if(I1*W1 > 0.0) FarCoc3 = min(FarCoc3, C1.a);
+	if(I2*W2 > 0.0) FarCoc3 = min(FarCoc3, C2.a);
+	if(I3*W3 > 0.0) FarCoc3 = min(FarCoc3, C3.a);
+	if(I4*W4 > 0.0) FarCoc3 = min(FarCoc3, C4.a);
+
+	W1 *= lerp(I1, 1.0, Ignore);
+	W2 *= lerp(I2, 1.0, Ignore);
+	W3 *= lerp(I3, 1.0, Ignore);
+	W4 *= lerp(I4, 1.0, Ignore);
+
+	OutColor0.rgb += C1.rgb * W1 + C2.rgb * W2 + C3.rgb * W3 + C4.rgb * W4;
+	Weight += W1+W2+W3+W4;
+
+
+	//
+	// Pass 3
+	//
+
+	// Drop weight of existing pass if Coc changes too much.
+	float Coc3 = max(abs(FarCoc3),abs(NearCoc));
+	Drop = (1.0/65536.0) + 1.0 - saturate(abs(Coc2 - Coc3));
+	OutColor0.rgb *= Drop;
+	Weight *= Drop;
+
+	// Send near most CoC back to recombine pass.
+	OutColor0.a = min(FarCoc3, NearCoc);
+
+	R1 = (RadiusBase+11.0/11.5) * Coc3;
+	R2 = (RadiusBase+ 5.0/11.5) * Coc3;
+	R3 = (RadiusBase+ 8.0/11.5) * Coc3;
+	R4 = (RadiusBase+ 2.0/11.5) * Coc3;
+
+	R1a = max(1.0,R1);
+	R2a = max(1.0,R2);
+	R3a = max(1.0,R3);
+	R4a = max(1.0,R4);
+
+	UV1 = RotVec(R1a, RadianBase + TwoPi *  4.0/12.0);
+	UV2 = RotVec(R2a, RadianBase + TwoPi *  7.0/12.0);
+	UV3 = RotVec(R3a, RadianBase + TwoPi * 10.0/12.0);
+	UV4 = RotVec(R4a, RadianBase + TwoPi *  1.0/12.0);
+
+	UV1 = UVAndScreenPos.xy + UV1 * PostprocessInput0Size.zw;
+	UV2 = UVAndScreenPos.xy + UV2 * PostprocessInput0Size.zw;
+	UV3 = UVAndScreenPos.xy + UV3 * PostprocessInput0Size.zw;
+	UV4 = UVAndScreenPos.xy + UV4 * PostprocessInput0Size.zw;
+
+	C1 = PostprocessInput0.SampleLevel(PostprocessInput0Sampler, UV1, 0);
+	C2 = PostprocessInput0.SampleLevel(PostprocessInput0Sampler, UV2, 0);
+	C3 = PostprocessInput0.SampleLevel(PostprocessInput0Sampler, UV3, 0);
+	C4 = PostprocessInput0.SampleLevel(PostprocessInput0Sampler, UV4, 0);
+
+	C1.rgb = lerp(C1.rgb, Background.rgb, saturate(abs(C1.a) - Coc3));
+	C2.rgb = lerp(C2.rgb, Background.rgb, saturate(abs(C2.a) - Coc3));
+	C3.rgb = lerp(C3.rgb, Background.rgb, saturate(abs(C3.a) - Coc3));
+	C4.rgb = lerp(C4.rgb, Background.rgb, saturate(abs(C4.a) - Coc3));
+
+	W1 = R1;
+	W2 = R2;
+	W3 = R3;
+	W4 = R4;
+
+	I1 = saturate((abs(C1.a) - R1) * IFeather0);
+	I2 = saturate((abs(C2.a) - R2) * IFeather0);
+	I3 = saturate((abs(C3.a) - R3) * IFeather0);
+	I4 = saturate((abs(C4.a) - R4) * IFeather0);
+
+	W1 *= lerp(I1, 1.0, Ignore);
+	W2 *= lerp(I2, 1.0, Ignore);
+	W3 *= lerp(I3, 1.0, Ignore);
+	W4 *= lerp(I4, 1.0, Ignore);
+
+	OutColor0.rgb += C1.rgb * W1 + C2.rgb * W2 + C3.rgb * W3 + C4.rgb * W4;
+	Weight += W1+W2+W3+W4;
+
+	OutColor0.rgb *= (1.0/Weight);
+
+}
+
+
+// pixel shader to combine the full res scene and the blurred images behind and in front of the the focal plane
+void MainCircleRecombinePS(in float4 UVAndScreenPos : TEXCOORD0, out float4 OutColor : SV_Target0)
+{
+	// Circle of confusion size for the pixel.
+	float PixDepth = CalcSceneDepth(UVAndScreenPos.xy);
+	float PixCoc = DepthToCoc(PixDepth);
+
+	// Grab nearest Coc.
+	PixCoc = min(PixCoc, PostprocessInput1.SampleLevel(PostprocessInput1Sampler, UVAndScreenPos.xy, 0).a);
+
+	// Transform into sample pattern.
+	PixCoc = abs(PixCoc) * 2.0; // 2x because full instead of half resolution.
+
+	#if 1
+		// Fetch 2 samples mirrored around the pixel 
+		// which is stochastically distributed to fill out the circle of confusion.
+		// TODO: Fix the "random values".
+		float2 UV = UVAndScreenPos.xy * PostprocessInput0Size.xy;
+		float RadianBase = NoizNorm(UVAndScreenPos.xy, 0.010 * RandomOffset.x) * 3.14159 * 2.0;
+		float RadiusJitter = NoizNorm(UVAndScreenPos.xy, 0.013 * RandomOffset.x);
+		float ICoc = PixCoc*sqrt(RadiusJitter);
+		float2 VP = RotVec(float2(ICoc, 0.0), RadianBase) * PostprocessInput0Size.zw;
+
+		// These two samples will still have jitter induced artifacts (very limited utility).
+		// These two samples will also have bleeding artifacts.
+		float4 CA = PostprocessInput0.SampleLevel(PostprocessInput0Sampler, UVAndScreenPos.xy + VP, 0);
+		float4 CB = PostprocessInput0.SampleLevel(PostprocessInput0Sampler, UVAndScreenPos.xy - VP, 0);
+
+		// Weight the two samples to avoid forground into background bleed.
+		float IFeather0 = 1.0/4.0;
+		ICoc *= 0.5; // Coc is half res units.
+		float I1 = saturate((abs(CA.a) - ICoc) * IFeather0);
+		float I2 = saturate((abs(CB.a) - ICoc) * IFeather0);
+		I1 += 1.0/65536.0;
+		I2 += 1.0/65536.0;
+	
+		OutColor = (CA * I1 + CB * I2) * (1.0/(I1+I2));
+	#else
+	    // Possibly higher quality option in the future.
+
+		// Fetch 4 samples in filled disc pattern
+		// which is stochastically distributed to fill out the circle of confusion.
+		float2 UV = UVAndScreenPos.xy * PostprocessInput0Size.xy;
+		float RadianBase = NoizNorm(UVAndScreenPos.xy, 0.010 * RandomOffset.x) * 3.14159 * 2.0;
+		float RadiusBase = NoizNorm(UVAndScreenPos.xy, 0.013 * RandomOffset.x);
+
+		float RadiusBase2 = RadiusBase * (1.0/4.0);
+		float R1 = sqrt(RadiusBase2 + 3.0/4.0) * PixCoc;
+		float R2 = sqrt(RadiusBase2 + 2.0/4.0) * PixCoc;
+		float R3 = sqrt(RadiusBase2 + 1.0/4.0) * PixCoc;
+		float R4 = sqrt(RadiusBase2 + 0.0/4.0) * PixCoc;
+
+		float TwoPi = 3.14159 * 2.0;
+		float2 UV1 = RotVec(R1, RadianBase + TwoPi * 0.0/4.0);
+		float2 UV2 = RotVec(R2, RadianBase + TwoPi * 2.0/4.0);
+		float2 UV3 = RotVec(R3, RadianBase + TwoPi * 1.0/4.0);
+		float2 UV4 = RotVec(R4, RadianBase + TwoPi * 3.0/4.0);
+
+		UV1 = UVAndScreenPos.xy + UV1 * PostprocessInput0Size.zw;
+		UV2 = UVAndScreenPos.xy + UV2 * PostprocessInput0Size.zw;
+		UV3 = UVAndScreenPos.xy + UV3 * PostprocessInput0Size.zw;
+		UV4 = UVAndScreenPos.xy + UV4 * PostprocessInput0Size.zw;
+
+		float4 CA = PostprocessInput0.SampleLevel(PostprocessInput0Sampler, UV1, 0);
+		float4 CB = PostprocessInput0.SampleLevel(PostprocessInput0Sampler, UV2, 0);
+		float4 CC = PostprocessInput0.SampleLevel(PostprocessInput0Sampler, UV3, 0);
+		float4 CD = PostprocessInput0.SampleLevel(PostprocessInput0Sampler, UV4, 0);
+
+		// Weight the two samples to avoid forground into background bleed.
+		float IFeather0 = 1.0/4.0;
+		float ICoc = PixCoc * 0.5; // Coc is half res units.
+		float I1 = saturate((abs(CA.a) - ICoc) * IFeather0);
+		float I2 = saturate((abs(CB.a) - ICoc) * IFeather0);
+		float I3 = saturate((abs(CC.a) - ICoc) * IFeather0);
+		float I4 = saturate((abs(CD.a) - ICoc) * IFeather0);
+
+		// Make sure something is non-zero.
+		I1 += 1.0/65536.0;
+		I2 += 1.0/65536.0;
+		I3 += 1.0/65536.0;
+		I4 += 1.0/65536.0;
+
+		OutColor = (CA * I1 + CB * I2 + CC * I3 + CD * I4) * (1.0/(I1+I2+I3+I4));
+	#endif
+
+	// Grab the half resolution neighborhood to remove the artifacts from the full resolution output.
+	// Nearest location.
+	#if 1
+		// This has higher in-focus contrast, but possibly lower noise reduction later.
+		float2 HUVBase = UVAndScreenPos.xy * PostprocessInput1Size.xy - 0.5;
+		float2 HUVFrac = frac(HUVBase);
+		float2 HUV = (trunc(HUVBase) + 0.5) * PostprocessInput1Size.zw;
+	#else
+		// This makes the mostly in-focus transition bad (too blurry).
+		float2 HUV = UVAndScreenPos.xy - 0.5 * PostprocessInput1Size.zw;
+	#endif
+
+	// Load four nearest samples.
+	float4 H0 = PostprocessInput1.SampleLevel(PostprocessInput1Sampler, HUV, 0);
+	float4 H1 = PostprocessInput1.SampleLevel(PostprocessInput1Sampler, HUV, 0, int2(1,0));
+	float4 H2 = PostprocessInput1.SampleLevel(PostprocessInput1Sampler, HUV, 0, int2(0,1));
+	float4 H3 = PostprocessInput1.SampleLevel(PostprocessInput1Sampler, HUV, 0, int2(1,1));
+
+	// TODO: This would work a lot better in YUV style colorspace?
+	// Limit the full resolution to remove jitter artifacts.
+	float4 HMax = max(max(H0,H1),max(H2,H3));
+	float4 HMin = min(min(H0,H1),min(H2,H3));
+
+	#if 1
+		// Increase constrast of limit a little to workaround to strong denoise at near-in-focus.
+		float4 HD = HMin / 8.0;
+		float Small = 1.0 - saturate(PixCoc*PixCoc*(1.0/64.0));
+		HMax += HD * Small;
+		HMin -= HD * Small;
+	#endif
+
+	// Blend in the limited version quickly to remove HDR jitter artifacts and noise.
+	float4 OutLimited = min(max(OutColor,HMin),HMax);
+	OutColor = lerp(OutColor, OutLimited, saturate(PixCoc*PixCoc*4.0));
+}
+
+
diff --git a/Engine/Source/Runtime/Engine/Classes/Engine/Scene.h b/Engine/Source/Runtime/Engine/Classes/Engine/Scene.h
--- a/Engine/Source/Runtime/Engine/Classes/Engine/Scene.h
+++ b/Engine/Source/Runtime/Engine/Classes/Engine/Scene.h
@@ -13,7 +13,7 @@ UENUM()
 enum EDepthOfFieldMethod
 {
 	DOFM_BokehDOF UMETA(DisplayName="BokehDOF"),
-	DOFM_Gaussian UMETA(DisplayName="Gaussian"),
+	DOFM_Gaussian UMETA(DisplayName="GaussianDOF"),
 	DOFM_CircleDOF UMETA(DisplayName="CircleDOF"),
 	DOFM_MAX,
 };
@@ -837,8 +837,8 @@ struct FPostProcessSettings
 	UPROPERTY(interp, BlueprintReadWrite, Category=DepthOfField, meta=(ClampMin = "0.0", ClampMax = "4.0", editcondition = "bOverride_DepthOfFieldDepthBlurRadius", DisplayName = "Depth Blur Radius"))
 	float DepthOfFieldDepthBlurRadius;
 	
-	/** CircleDOF only: F-stop number (1/fstop) */
-	UPROPERTY(interp, BlueprintReadWrite, Category=DepthOfField, meta=(ClampMin = "1.0", ClampMax = "32.0", editcondition = "bOverride_DepthOfFieldFstop", DisplayName = "1/Fstop"))
+	/** CircleDOF only: F-stop number (1/fstop) e.g. 1.2 for a large aperture, 2.8 is smaller and results in less Depth of Field effect, does not affect noise/grain but in real world it would */
+	UPROPERTY(interp, BlueprintReadWrite, Category=DepthOfField, meta=(ClampMin = "1.0", ClampMax = "32.0", editcondition = "bOverride_DepthOfFieldFstop", DisplayName = "Aperture"))
 	float DepthOfFieldFstop;
 
 	/** Distance in which the Depth of Field effect should be sharp, in unreal units (cm) */
diff --git a/Engine/Source/Runtime/Renderer/Private/PostProcess/PostProcessBokehDOF.cpp b/Engine/Source/Runtime/Renderer/Private/PostProcess/PostProcessBokehDOF.cpp
--- a/Engine/Source/Runtime/Renderer/Private/PostProcess/PostProcessBokehDOF.cpp
+++ b/Engine/Source/Runtime/Renderer/Private/PostProcess/PostProcessBokehDOF.cpp
@@ -10,6 +10,7 @@
 #include "PostProcessPassThrough.h"
 #include "PostProcessing.h"
 #include "PostProcessBokehDOF.h"
+#include "PostProcessCircleDOF.h"
 #include "SceneUtils.h"
 
 
@@ -261,33 +262,92 @@ void FRCPassPostProcessVisualizeDOF::Process(FRenderingCompositePassContext& Con
 
 		FString Line;
 
-		Line = FString::Printf(TEXT("Near:%d Far:%d"), DepthOfFieldStats.bNear ? 1 : 0, DepthOfFieldStats.bFar ? 1 : 0);
-		Canvas.DrawShadowedString(X, Y += YStep, *Line, GetStatsFont(), FLinearColor(1, 1, 1));
+		Line = FString::Printf(TEXT("Visualize Depth of Field"));
+		Canvas.DrawShadowedString(20, Y += YStep, *Line, GetStatsFont(), FLinearColor(1, 1, 0));
 		Y += YStep;
-		Line = FString::Printf(TEXT("FocalDistance: %.2f"), View.FinalPostProcessSettings.DepthOfFieldFocalDistance);
-		Canvas.DrawShadowedString(X, Y += YStep, *Line, GetStatsFont(), FLinearColor(1, 1, 1));
-		Line = FString::Printf(TEXT("FocalRegion: %.2f"), View.FinalPostProcessSettings.DepthOfFieldFocalRegion);
-		Canvas.DrawShadowedString(X, Y += YStep, *Line, GetStatsFont(), FLinearColor(1, 1, 1));
-		Line = FString::Printf(TEXT("NearTransitionRegion: %.2f"), View.FinalPostProcessSettings.DepthOfFieldNearTransitionRegion);
-		Canvas.DrawShadowedString(X, Y += YStep, *Line, GetStatsFont(), FLinearColor(1, 1, 1));
-		Line = FString::Printf(TEXT("FarTransitionRegion: %.2f"), View.FinalPostProcessSettings.DepthOfFieldFarTransitionRegion);
-		Canvas.DrawShadowedString(X, Y += YStep, *Line, GetStatsFont(), FLinearColor(1, 1, 1));
-		Line = FString::Printf(TEXT("Scale: %.2f"), View.FinalPostProcessSettings.DepthOfFieldScale);
-		Canvas.DrawShadowedString(X, Y += YStep, *Line, GetStatsFont(), FLinearColor(1, 1, 1));
-		Line = FString::Printf(TEXT("NearBlurSize: %.2f"), View.FinalPostProcessSettings.DepthOfFieldNearBlurSize);
-		Canvas.DrawShadowedString(X, Y += YStep, *Line, GetStatsFont(), FLinearColor(1, 1, 1));
-		Line = FString::Printf(TEXT("FarBlurSize: %.2f"), View.FinalPostProcessSettings.DepthOfFieldFarBlurSize);
-		Canvas.DrawShadowedString(X, Y += YStep, *Line, GetStatsFont(), FLinearColor(1, 1, 1));
-		Line = FString::Printf(TEXT("Method: %d"), (uint32)View.FinalPostProcessSettings.DepthOfFieldMethod.GetValue());
-		Canvas.DrawShadowedString(X, Y += YStep, *Line, GetStatsFont(), FLinearColor(1, 1, 1));
-		Line = FString::Printf(TEXT("Occlusion: %.2f"), View.FinalPostProcessSettings.DepthOfFieldOcclusion);
-		Canvas.DrawShadowedString(X, Y += YStep, *Line, GetStatsFont(), FLinearColor(1, 1, 1));
-		Line = FString::Printf(TEXT("ColorThreshold: %.2f"), View.FinalPostProcessSettings.DepthOfFieldColorThreshold);
-		Canvas.DrawShadowedString(X, Y += YStep, *Line, GetStatsFont(), FLinearColor(1, 1, 1));
-		Line = FString::Printf(TEXT("SizeThreshold: %.2f"), View.FinalPostProcessSettings.DepthOfFieldSizeThreshold);
-		Canvas.DrawShadowedString(X, Y += YStep, *Line, GetStatsFont(), FLinearColor(1, 1, 1));
-		Line = FString::Printf(TEXT("SkyFocusDistance: %.2f"), View.FinalPostProcessSettings.DepthOfFieldSkyFocusDistance);
-		Canvas.DrawShadowedString(X, Y += YStep, *Line, GetStatsFont(), FLinearColor(1, 1, 1));
+
+		EDepthOfFieldMethod MethodId = View.FinalPostProcessSettings.DepthOfFieldMethod;
+
+		if(MethodId == DOFM_BokehDOF)
+		{
+			Line = FString::Printf(TEXT("Method: BokehDOF (blue is far, green is near, black is in focus)"));
+			Canvas.DrawShadowedString(X, Y += YStep, *Line, GetStatsFont(), FLinearColor(1, 1, 1));
+			Y += YStep;
+			Line = FString::Printf(TEXT("FocalDistance: %.2f"), View.FinalPostProcessSettings.DepthOfFieldFocalDistance);
+			Canvas.DrawShadowedString(X, Y += YStep, *Line, GetStatsFont(), FLinearColor(1, 1, 1));
+			Line = FString::Printf(TEXT("FocalRegion: %.2f"), View.FinalPostProcessSettings.DepthOfFieldFocalRegion);
+			Canvas.DrawShadowedString(X, Y += YStep, *Line, GetStatsFont(), FLinearColor(1, 1, 1));
+			Y += YStep;
+			Line = FString::Printf(TEXT("Scale: %.2f"), View.FinalPostProcessSettings.DepthOfFieldScale);
+			Canvas.DrawShadowedString(X, Y += YStep, *Line, GetStatsFont(), FLinearColor(1, 1, 1));
+			Line = FString::Printf(TEXT("MaxBokehSize: %.2f"), View.FinalPostProcessSettings.DepthOfFieldMaxBokehSize);
+			Canvas.DrawShadowedString(X, Y += YStep, *Line, GetStatsFont(), FLinearColor(1, 1, 1));
+			Line = FString::Printf(TEXT("NearTransitionRegion: %.2f"), View.FinalPostProcessSettings.DepthOfFieldNearTransitionRegion);
+			Canvas.DrawShadowedString(X, Y += YStep, *Line, GetStatsFont(), FLinearColor(1, 1, 1));
+			Line = FString::Printf(TEXT("FarTransitionRegion: %.2f"), View.FinalPostProcessSettings.DepthOfFieldFarTransitionRegion);
+			Canvas.DrawShadowedString(X, Y += YStep, *Line, GetStatsFont(), FLinearColor(1, 1, 1));
+			Line = FString::Printf(TEXT("ColorThreshold: %.2f"), View.FinalPostProcessSettings.DepthOfFieldColorThreshold);
+			Canvas.DrawShadowedString(X, Y += YStep, *Line, GetStatsFont(), FLinearColor(1, 1, 1));
+			Line = FString::Printf(TEXT("SizeThreshold: %.2f"), View.FinalPostProcessSettings.DepthOfFieldSizeThreshold);
+			Canvas.DrawShadowedString(X, Y += YStep, *Line, GetStatsFont(), FLinearColor(1, 1, 1));
+			Line = FString::Printf(TEXT("Occlusion: %.2f"), View.FinalPostProcessSettings.DepthOfFieldOcclusion);
+			Canvas.DrawShadowedString(X, Y += YStep, *Line, GetStatsFont(), FLinearColor(1, 1, 1));
+		}
+		else if(MethodId == DOFM_Gaussian)
+		{
+			Line = FString::Printf(TEXT("Method: GaussianDOF (blue is far, green is near, grey is disabled, black is in focus)"));
+			Canvas.DrawShadowedString(X, Y += YStep, *Line, GetStatsFont(), FLinearColor(1, 1, 1));
+			Y += YStep;
+			Line = FString::Printf(TEXT("FocalDistance: %.2f"), View.FinalPostProcessSettings.DepthOfFieldFocalDistance);
+			Canvas.DrawShadowedString(X, Y += YStep, *Line, GetStatsFont(), FLinearColor(1, 1, 1));
+			Line = FString::Printf(TEXT("FocalRegion: %.2f"), View.FinalPostProcessSettings.DepthOfFieldFocalRegion);
+			Canvas.DrawShadowedString(X, Y += YStep, *Line, GetStatsFont(), FLinearColor(1, 1, 1));
+			Y += YStep;
+			Line = FString::Printf(TEXT("NearTransitionRegion: %.2f"), View.FinalPostProcessSettings.DepthOfFieldNearTransitionRegion);
+			Canvas.DrawShadowedString(X, Y += YStep, *Line, GetStatsFont(), FLinearColor(1, 1, 1));
+			Line = FString::Printf(TEXT("FarTransitionRegion: %.2f"), View.FinalPostProcessSettings.DepthOfFieldFarTransitionRegion);
+			Canvas.DrawShadowedString(X, Y += YStep, *Line, GetStatsFont(), FLinearColor(1, 1, 1));
+			Line = FString::Printf(TEXT("NearBlurSize: %.2f"), View.FinalPostProcessSettings.DepthOfFieldNearBlurSize);
+			Canvas.DrawShadowedString(X, Y += YStep, *Line, GetStatsFont(), FLinearColor(1, 1, 1));
+			Line = FString::Printf(TEXT("FarBlurSize: %.2f"), View.FinalPostProcessSettings.DepthOfFieldFarBlurSize);
+			Canvas.DrawShadowedString(X, Y += YStep, *Line, GetStatsFont(), FLinearColor(1, 1, 1));
+			Line = FString::Printf(TEXT("Occlusion: %.2f"), View.FinalPostProcessSettings.DepthOfFieldOcclusion);
+			Canvas.DrawShadowedString(X, Y += YStep, *Line, GetStatsFont(), FLinearColor(1, 1, 1));
+			Line = FString::Printf(TEXT("SkyFocusDistance: %.2f"), View.FinalPostProcessSettings.DepthOfFieldSkyFocusDistance);
+			Canvas.DrawShadowedString(X, Y += YStep, *Line, GetStatsFont(), FLinearColor(1, 1, 1));
+			Y += YStep;
+			Line = FString::Printf(TEXT("Near:%d Far:%d"), DepthOfFieldStats.bNear ? 1 : 0, DepthOfFieldStats.bFar ? 1 : 0);
+			Canvas.DrawShadowedString(X, Y += YStep, *Line, GetStatsFont(), FLinearColor(1, 1, 1));
+		}
+		else if(MethodId == DOFM_CircleDOF)
+		{
+			Line = FString::Printf(TEXT("Method: CircleDOF"));
+			Canvas.DrawShadowedString(X, Y += YStep, *Line, GetStatsFont(), FLinearColor(1, 1, 1));
+			Y += YStep;
+			Line = FString::Printf(TEXT("FocalDistance: %.2f"), View.FinalPostProcessSettings.DepthOfFieldFocalDistance);
+			Canvas.DrawShadowedString(X, Y += YStep, *Line, GetStatsFont(), FLinearColor(1, 1, 1));
+			Line = FString::Printf(TEXT("Aperture: %.2f"), View.FinalPostProcessSettings.DepthOfFieldFstop);
+			Canvas.DrawShadowedString(X, Y += YStep, *Line, GetStatsFont(), FLinearColor(1, 1, 1));
+			Y += YStep;
+			Line = FString::Printf(TEXT("DepthBlur (not related to Depth of Field, due to light traveling long distances in atmosphere)"));
+			Canvas.DrawShadowedString(X, Y += YStep, *Line, GetStatsFont(), FLinearColor(1, 1, 1));
+			Line = FString::Printf(TEXT("  km for 50%: %.2f"), View.FinalPostProcessSettings.DepthOfFieldDepthBlurAmount);
+			Canvas.DrawShadowedString(X, Y += YStep, *Line, GetStatsFont(), FLinearColor(1, 1, 1));
+			Line = FString::Printf(TEXT("  Radius (pixels in 1920x): %.2f"), View.FinalPostProcessSettings.DepthOfFieldDepthBlurRadius);
+			Canvas.DrawShadowedString(X, Y += YStep, *Line, GetStatsFont(), FLinearColor(1, 1, 1));
+			Y += YStep;
+
+			FVector2D Fov = View.ViewMatrices.GetHalfFieldOfViewPerAxis(); 
+			
+			float FocalLength = ComputeFocalLengthFromFov(View);
+
+			Line = FString::Printf(TEXT("Field Of View in deg. (computed): %.1f x %.1f"), FMath::RadiansToDegrees(Fov.X) * 2.0f, FMath::RadiansToDegrees(Fov.Y) * 2.0f);
+			Canvas.DrawShadowedString(X, Y += YStep, *Line, GetStatsFont(), FLinearColor(0.5, 0.5, 1));
+			Line = FString::Printf(TEXT("Focal Length (computed): %.1f"), FocalLength);
+			Canvas.DrawShadowedString(X, Y += YStep, *Line, GetStatsFont(), FLinearColor(0.5, 0.5, 10));
+			Line = FString::Printf(TEXT("Sensor: APS-C 24.576 mm sensor, crop-factor 1.61x"));
+			Canvas.DrawShadowedString(X, Y += YStep, *Line, GetStatsFont(), FLinearColor(0.5, 0.5, 1));
+		}
 
 		Canvas.Flush_RenderThread(Context.RHICmdList);
 	}
diff --git a/Engine/Source/Runtime/Renderer/Private/PostProcess/PostProcessBokehDOF.h b/Engine/Source/Runtime/Renderer/Private/PostProcess/PostProcessBokehDOF.h
--- a/Engine/Source/Runtime/Renderer/Private/PostProcess/PostProcessBokehDOF.h
+++ b/Engine/Source/Runtime/Renderer/Private/PostProcess/PostProcessBokehDOF.h
@@ -11,8 +11,8 @@
 struct FDepthOfFieldStats
 {
 	FDepthOfFieldStats()
-		: bNear(false)
-		, bFar(false)
+		: bNear(true)
+		, bFar(true)
 	{
 	}
 
diff --git a/Engine/Source/Runtime/Renderer/Private/PostProcess/PostProcessCircleDOF.cpp b/Engine/Source/Runtime/Renderer/Private/PostProcess/PostProcessCircleDOF.cpp
new file mode 100644
--- /dev/null
+++ b/Engine/Source/Runtime/Renderer/Private/PostProcess/PostProcessCircleDOF.cpp
@@ -0,0 +1,830 @@
+// Copyright 1998-2015 Epic Games, Inc. All Rights Reserved.
+
+/*=============================================================================
+	PostProcessDOF.cpp: Post process Depth of Field implementation.
+=============================================================================*/
+
+#include "RendererPrivate.h"
+#include "ScenePrivate.h"
+#include "SceneFilterRendering.h"
+#include "PostProcessBokehDOF.h"
+#include "PostProcessCircleDOF.h"
+#include "PostProcessing.h"
+#include "SceneUtils.h"
+
+
+float ComputeFocalLengthFromFov(const FSceneView& View)
+{
+	// Convert FOV to focal length,
+	// 
+	// fov = 2 * atan(d/(2*f))
+	// where,
+	//   d = sensor dimension (APS-C 24.576 mm)
+	//   f = focal length
+	// 
+	// f = 0.5 * d * (1/tan(fov/2))
+	float HalfFOV = FMath::Atan(1.0f / View.ViewMatrices.ProjMatrix.M[0][0]);
+	float FocalLength = 0.5f * 24.576f * (1.0f/FMath::Tan(HalfFOV));
+
+	return FocalLength;
+}
+
+// Convert f-stop and focal distance into projected size in half resolution pixels.
+// Setup depth based blur.
+static FVector4 CircleDofCoc(const FRenderingCompositePassContext& Context)
+{
+	float FocalLength = ComputeFocalLengthFromFov(Context.View);
+	 
+	// Convert focal distance in world position to mm.
+	// Conversion is 1 world position = 1 cm.
+	float Distance = Context.View.FinalPostProcessSettings.DepthOfFieldFocalDistance;
+	Distance *= 10.0f;
+
+	// Convert f-stop, focal length, and focal distance to
+	// projected circle of confusion size at infinity in mm.
+	//
+	// coc = f*f / (n * (d - f))
+	// where,
+	//   f = focal length
+	//   d = focal distance
+	//   n = fstop (where n is the "n" in "f/n")
+	float Radius = FocalLength * FocalLength / (Context.View.FinalPostProcessSettings.DepthOfFieldFstop * (Distance - FocalLength));
+
+	// Scale so that APS-C 24.576 mm = full frame.
+	// Convert mm to pixels.
+	float Width = (float)Context.View.ViewRect.Width();
+	Radius = Radius * Width * (1.0f/24.576f);
+
+	// Convert diameter to radius at half resolution (algorithm radius is at half resolution).
+	Radius *= 0.25f;
+
+	// Comment out for now, allowing settings which the algorithm cannot cleanly do.
+	#if 0
+		// Limit to algorithm max size.
+		if(Radius > 6.0f) 
+		{
+			Radius = 6.0f; 
+		}
+	#endif
+
+	// The DepthOfFieldDepthBlurAmount = km at which depth blur is 50%.
+	// Need to convert to cm here.
+	return FVector4(
+		Radius, 
+		1.0f/(Context.View.FinalPostProcessSettings.DepthOfFieldDepthBlurAmount * 100000.0),
+		Context.View.FinalPostProcessSettings.DepthOfFieldDepthBlurRadius * Width / 1920.0f,
+		Width / 1920.0f);
+}
+
+/** Encapsulates the Circle DOF setup pixel shader. */
+template <uint32 NearBlurEnable>
+class FPostProcessCircleDOFSetupPS : public FGlobalShader
+{
+	DECLARE_SHADER_TYPE(FPostProcessCircleDOFSetupPS, Global);
+
+	static bool ShouldCache(EShaderPlatform Platform)
+	{
+		return IsFeatureLevelSupported(Platform, ERHIFeatureLevel::SM5);
+	}
+
+	static void ModifyCompilationEnvironment(EShaderPlatform Platform, FShaderCompilerEnvironment& OutEnvironment)
+	{
+		FGlobalShader::ModifyCompilationEnvironment(Platform,OutEnvironment);
+		OutEnvironment.SetDefine(TEXT("ENABLE_NEAR_BLUR"), NearBlurEnable);
+	}
+
+	/** Default constructor. */
+	FPostProcessCircleDOFSetupPS() {}
+
+public:
+	FPostProcessPassParameters PostprocessParameter;
+	FDeferredPixelShaderParameters DeferredParameters;
+	FShaderParameter DepthOfFieldParams;
+	FShaderParameter CircleDofParams;
+
+	/** Initialization constructor. */
+	FPostProcessCircleDOFSetupPS(const ShaderMetaType::CompiledShaderInitializerType& Initializer)
+		: FGlobalShader(Initializer)
+	{
+		PostprocessParameter.Bind(Initializer.ParameterMap);
+		DeferredParameters.Bind(Initializer.ParameterMap);
+		DepthOfFieldParams.Bind(Initializer.ParameterMap,TEXT("DepthOfFieldParams"));
+		CircleDofParams.Bind(Initializer.ParameterMap,TEXT("CircleDofParams"));
+	}
+
+	// FShader interface.
+	virtual bool Serialize(FArchive& Ar) override
+	{
+		bool bShaderHasOutdatedParameters = FGlobalShader::Serialize(Ar);
+		Ar << PostprocessParameter << DeferredParameters << DepthOfFieldParams << CircleDofParams;
+		return bShaderHasOutdatedParameters;
+	}
+
+	void SetParameters(const FRenderingCompositePassContext& Context)
+	{
+		const FPixelShaderRHIParamRef ShaderRHI = GetPixelShader();
+
+		FGlobalShader::SetParameters(Context.RHICmdList, ShaderRHI, Context.View);
+
+		PostprocessParameter.SetPS(ShaderRHI, Context, TStaticSamplerState<SF_Point,AM_Border,AM_Border,AM_Clamp>::GetRHI());
+
+		DeferredParameters.Set(Context.RHICmdList, ShaderRHI, Context.View);
+
+		{
+			FVector4 DepthOfFieldParamValues[2];
+
+			FRCPassPostProcessBokehDOF::ComputeDepthOfFieldParams(Context, DepthOfFieldParamValues);
+
+			SetShaderValueArray(Context.RHICmdList, ShaderRHI, DepthOfFieldParams, DepthOfFieldParamValues, 2);
+		}
+
+		SetShaderValue(Context.RHICmdList, ShaderRHI, CircleDofParams, CircleDofCoc(Context));
+	}
+};
+
+IMPLEMENT_SHADER_TYPE(template<>,FPostProcessCircleDOFSetupPS<0>,TEXT("PostProcessCircleDOF"),TEXT("CircleSetupPS"),SF_Pixel);
+IMPLEMENT_SHADER_TYPE(template<>,FPostProcessCircleDOFSetupPS<1>,TEXT("PostProcessCircleDOF"),TEXT("CircleSetupPS"),SF_Pixel);
+
+void FRCPassPostProcessCircleDOFSetup::Process(FRenderingCompositePassContext& Context)
+{
+	SCOPED_DRAW_EVENT(Context.RHICmdList, CircleDOFSetup);
+
+	const FPooledRenderTargetDesc* InputDesc = GetInputDesc(ePId_Input0);
+
+	if(!InputDesc)
+	{
+		// input is not hooked up correctly
+		return;
+	}
+
+	uint32 NumRenderTargets = bNearBlurEnabled ? 2 : 1;
+
+	const FSceneView& View = Context.View;
+	const FSceneViewFamily& ViewFamily = *(View.Family);
+
+	const auto FeatureLevel = Context.GetFeatureLevel();
+	auto ShaderMap = Context.GetShaderMap();
+
+	FIntPoint SrcSize = InputDesc->Extent;
+	FIntPoint DestSize = PassOutputs[0].RenderTargetDesc.Extent;
+
+	// e.g. 4 means the input texture is 4x smaller than the buffer size
+	uint32 ScaleFactor = FSceneRenderTargets::Get(Context.RHICmdList).GetBufferSizeXY().X / SrcSize.X;
+
+	FIntRect SrcRect = View.ViewRect / ScaleFactor;
+	FIntRect DestRect = SrcRect / 2;
+
+	const FSceneRenderTargetItem& DestRenderTarget0 = PassOutputs[0].RequestSurface(Context);
+	const FSceneRenderTargetItem& DestRenderTarget1 = bNearBlurEnabled ? PassOutputs[1].RequestSurface(Context) : FSceneRenderTargetItem();
+
+	// Set the view family's render target/viewport.
+	FTextureRHIParamRef RenderTargets[2] =
+	{
+		DestRenderTarget0.TargetableTexture,
+		DestRenderTarget1.TargetableTexture
+	};
+	SetRenderTargets(Context.RHICmdList, NumRenderTargets, RenderTargets, FTextureRHIParamRef(), 0, NULL);
+
+	FLinearColor ClearColors[2] = 
+	{
+		FLinearColor(0, 0, 0, 0),
+		FLinearColor(0, 0, 0, 0)
+	};
+	// is optimized away if possible (RT size=view size, )
+	Context.RHICmdList.ClearMRT(true, NumRenderTargets, ClearColors, false, 1.0f, false, 0, DestRect);
+
+	Context.SetViewportAndCallRHI(0, 0, 0.0f, DestSize.X, DestSize.Y, 1.0f );
+
+	// set the state
+	Context.RHICmdList.SetBlendState(TStaticBlendState<>::GetRHI());
+	Context.RHICmdList.SetRasterizerState(TStaticRasterizerState<>::GetRHI());
+	Context.RHICmdList.SetDepthStencilState(TStaticDepthStencilState<false, CF_Always>::GetRHI());
+
+	TShaderMapRef<FPostProcessVS> VertexShader(ShaderMap);
+
+	if (bNearBlurEnabled)
+	{
+		static FGlobalBoundShaderState BoundShaderState;
+
+
+		TShaderMapRef< FPostProcessCircleDOFSetupPS<1> > PixelShader(ShaderMap);
+		SetGlobalBoundShaderState(Context.RHICmdList, FeatureLevel, BoundShaderState, GFilterVertexDeclaration.VertexDeclarationRHI, *VertexShader, *PixelShader);
+
+		PixelShader->SetParameters(Context);
+	}
+	else
+	{
+		static FGlobalBoundShaderState BoundShaderState;
+
+		TShaderMapRef< FPostProcessCircleDOFSetupPS<0> > PixelShader(ShaderMap);
+		SetGlobalBoundShaderState(Context.RHICmdList, FeatureLevel, BoundShaderState, GFilterVertexDeclaration.VertexDeclarationRHI, *VertexShader, *PixelShader);
+
+		PixelShader->SetParameters(Context);
+	}
+
+	VertexShader->SetParameters(Context);
+
+	// Draw a quad mapping scene color to the view's render target
+	DrawRectangle(
+		Context.RHICmdList,
+		DestRect.Min.X, DestRect.Min.Y,
+		DestRect.Width() + 1, DestRect.Height() + 1,
+		SrcRect.Min.X, SrcRect.Min.Y,
+		SrcRect.Width() + 1, SrcRect.Height() + 1,
+		DestSize,
+		SrcSize,
+		*VertexShader,
+		EDRF_UseTriangleOptimization);
+
+	Context.RHICmdList.CopyToResolveTarget(DestRenderTarget0.TargetableTexture, DestRenderTarget0.ShaderResourceTexture, false, FResolveParams());
+	Context.RHICmdList.CopyToResolveTarget(DestRenderTarget1.TargetableTexture, DestRenderTarget1.ShaderResourceTexture, false, FResolveParams());
+}
+
+FPooledRenderTargetDesc FRCPassPostProcessCircleDOFSetup::ComputeOutputDesc(EPassOutputId InPassOutputId) const
+{
+	FPooledRenderTargetDesc Ret = GetInput(ePId_Input0)->GetOutput()->RenderTargetDesc;
+
+	//	Ret.Extent = FIntPoint::DivideAndRoundUp(ret.Extent, 2);
+	Ret.Extent /= 2;
+	Ret.Extent.X = FMath::Max(1, Ret.Extent.X);
+	Ret.Extent.Y = FMath::Max(1, Ret.Extent.Y);
+
+	Ret.Reset();
+	Ret.TargetableFlags &= ~(uint32)TexCreate_UAV;
+	Ret.TargetableFlags |= TexCreate_RenderTargetable;
+
+	Ret.DebugName = (InPassOutputId == ePId_Output0) ? TEXT("CircleDOFSetup0") : TEXT("CircleDOFSetup1");
+
+	// more precision for additive blending and we need the alpha channel
+	Ret.Format = PF_FloatRGBA;
+
+	return Ret;
+}
+
+
+
+
+/** Encapsulates the Circle DOF Dilate pixel shader. */
+template <uint32 NearBlurEnable>
+class FPostProcessCircleDOFDilatePS : public FGlobalShader
+{
+	DECLARE_SHADER_TYPE(FPostProcessCircleDOFDilatePS, Global);
+
+	static bool ShouldCache(EShaderPlatform Platform)
+	{
+		return IsFeatureLevelSupported(Platform, ERHIFeatureLevel::SM5);
+	}
+
+	static void ModifyCompilationEnvironment(EShaderPlatform Platform, FShaderCompilerEnvironment& OutEnvironment)
+	{
+		FGlobalShader::ModifyCompilationEnvironment(Platform,OutEnvironment);
+		OutEnvironment.SetDefine(TEXT("ENABLE_NEAR_BLUR"), NearBlurEnable);
+	}
+
+	/** Default constructor. */
+	FPostProcessCircleDOFDilatePS() {}
+
+public:
+	FPostProcessPassParameters PostprocessParameter;
+	FDeferredPixelShaderParameters DeferredParameters;
+	FShaderParameter DepthOfFieldParams;
+	FShaderParameter CircleDofParams;
+
+	/** Initialization constructor. */
+	FPostProcessCircleDOFDilatePS(const ShaderMetaType::CompiledShaderInitializerType& Initializer)
+		: FGlobalShader(Initializer)
+	{
+		PostprocessParameter.Bind(Initializer.ParameterMap);
+		DeferredParameters.Bind(Initializer.ParameterMap);
+		DepthOfFieldParams.Bind(Initializer.ParameterMap,TEXT("DepthOfFieldParams"));
+		CircleDofParams.Bind(Initializer.ParameterMap,TEXT("CircleDofParams"));
+	}
+
+	// FShader interface.
+	virtual bool Serialize(FArchive& Ar) override
+	{
+		bool bShaderHasOutdatedParameters = FGlobalShader::Serialize(Ar);
+		Ar << PostprocessParameter << DeferredParameters << DepthOfFieldParams << CircleDofParams;
+		return bShaderHasOutdatedParameters;
+	}
+
+	void SetParameters(const FRenderingCompositePassContext& Context)
+	{
+		const FPixelShaderRHIParamRef ShaderRHI = GetPixelShader();
+
+		FGlobalShader::SetParameters(Context.RHICmdList, ShaderRHI, Context.View);
+
+		PostprocessParameter.SetPS(ShaderRHI, Context, TStaticSamplerState<SF_Point,AM_Border,AM_Border,AM_Clamp>::GetRHI());
+
+		DeferredParameters.Set(Context.RHICmdList, ShaderRHI, Context.View);
+
+		{
+			FVector4 DepthOfFieldParamValues[2];
+
+			FRCPassPostProcessBokehDOF::ComputeDepthOfFieldParams(Context, DepthOfFieldParamValues);
+
+			SetShaderValueArray(Context.RHICmdList, ShaderRHI, DepthOfFieldParams, DepthOfFieldParamValues, 2);
+		}
+
+		SetShaderValue(Context.RHICmdList, ShaderRHI, CircleDofParams, CircleDofCoc(Context));
+	}
+};
+
+IMPLEMENT_SHADER_TYPE(template<>,FPostProcessCircleDOFDilatePS<0>,TEXT("PostProcessCircleDOF"),TEXT("CircleDilatePS"),SF_Pixel);
+IMPLEMENT_SHADER_TYPE(template<>,FPostProcessCircleDOFDilatePS<1>,TEXT("PostProcessCircleDOF"),TEXT("CircleDilatePS"),SF_Pixel);
+
+void FRCPassPostProcessCircleDOFDilate::Process(FRenderingCompositePassContext& Context)
+{
+	SCOPED_DRAW_EVENT(Context.RHICmdList, CircleDOFNear);
+
+	const FPooledRenderTargetDesc* InputDesc = GetInputDesc(ePId_Input0);
+
+	if(!InputDesc)
+	{
+		// input is not hooked up correctly
+		return;
+	}
+
+	uint32 NumRenderTargets = 1;
+
+	const FSceneView& View = Context.View;
+	const FSceneViewFamily& ViewFamily = *(View.Family);
+
+	const auto FeatureLevel = Context.GetFeatureLevel();
+	auto ShaderMap = Context.GetShaderMap();
+
+	FIntPoint SrcSize = InputDesc->Extent;
+	FIntPoint DestSize = PassOutputs[0].RenderTargetDesc.Extent;
+
+	// e.g. 4 means the input texture is 4x smaller than the buffer size
+	uint32 ScaleFactor = FSceneRenderTargets::Get(Context.RHICmdList).GetBufferSizeXY().X / SrcSize.X;
+
+	FIntRect SrcRect = View.ViewRect / ScaleFactor;
+	FIntRect DestRect = SrcRect / 2;
+
+	const FSceneRenderTargetItem& DestRenderTarget0 = PassOutputs[0].RequestSurface(Context);
+	const FSceneRenderTargetItem& DestRenderTarget1 = FSceneRenderTargetItem();
+
+	// Set the view family's render target/viewport.
+	FTextureRHIParamRef RenderTargets[2] =
+	{
+		DestRenderTarget0.TargetableTexture,
+		DestRenderTarget1.TargetableTexture
+	};
+	SetRenderTargets(Context.RHICmdList, NumRenderTargets, RenderTargets, FTextureRHIParamRef(), 0, NULL);
+
+	FLinearColor ClearColors[2] = 
+	{
+		FLinearColor(0, 0, 0, 0),
+		FLinearColor(0, 0, 0, 0)
+	};
+	// is optimized away if possible (RT size=view size, )
+	Context.RHICmdList.ClearMRT(true, NumRenderTargets, ClearColors, false, 1.0f, false, 0, DestRect);
+
+	Context.SetViewportAndCallRHI(0, 0, 0.0f, DestSize.X, DestSize.Y, 1.0f );
+
+	// set the state
+	Context.RHICmdList.SetBlendState(TStaticBlendState<>::GetRHI());
+	Context.RHICmdList.SetRasterizerState(TStaticRasterizerState<>::GetRHI());
+	Context.RHICmdList.SetDepthStencilState(TStaticDepthStencilState<false, CF_Always>::GetRHI());
+
+	TShaderMapRef<FPostProcessVS> VertexShader(ShaderMap);
+
+	if (false)
+	{
+		static FGlobalBoundShaderState BoundShaderState;
+
+
+		TShaderMapRef< FPostProcessCircleDOFDilatePS<1> > PixelShader(ShaderMap);
+		SetGlobalBoundShaderState(Context.RHICmdList, FeatureLevel, BoundShaderState, GFilterVertexDeclaration.VertexDeclarationRHI, *VertexShader, *PixelShader);
+
+		PixelShader->SetParameters(Context);
+	}
+	else
+	{
+		static FGlobalBoundShaderState BoundShaderState;
+
+		TShaderMapRef< FPostProcessCircleDOFDilatePS<0> > PixelShader(ShaderMap);
+		SetGlobalBoundShaderState(Context.RHICmdList, FeatureLevel, BoundShaderState, GFilterVertexDeclaration.VertexDeclarationRHI, *VertexShader, *PixelShader);
+
+		PixelShader->SetParameters(Context);
+	}
+
+	VertexShader->SetParameters(Context);
+
+	// Draw a quad mapping scene color to the view's render target
+	DrawRectangle(
+		Context.RHICmdList,
+		DestRect.Min.X, DestRect.Min.Y,
+		DestRect.Width() + 1, DestRect.Height() + 1,
+		SrcRect.Min.X, SrcRect.Min.Y,
+		SrcRect.Width() + 1, SrcRect.Height() + 1,
+		DestSize,
+		SrcSize,
+		*VertexShader,
+		EDRF_UseTriangleOptimization);
+
+	Context.RHICmdList.CopyToResolveTarget(DestRenderTarget0.TargetableTexture, DestRenderTarget0.ShaderResourceTexture, false, FResolveParams());
+	Context.RHICmdList.CopyToResolveTarget(DestRenderTarget1.TargetableTexture, DestRenderTarget1.ShaderResourceTexture, false, FResolveParams());
+}
+
+FPooledRenderTargetDesc FRCPassPostProcessCircleDOFDilate::ComputeOutputDesc(EPassOutputId InPassOutputId) const
+{
+	FPooledRenderTargetDesc Ret = GetInput(ePId_Input0)->GetOutput()->RenderTargetDesc;
+
+	//	Ret.Extent = FIntPoint::DivideAndRoundUp(ret.Extent, 2);
+	Ret.Extent /= 2;
+	Ret.Extent.X = FMath::Max(1, Ret.Extent.X);
+	Ret.Extent.Y = FMath::Max(1, Ret.Extent.Y);
+
+	Ret.Reset();
+	Ret.TargetableFlags &= ~(uint32)TexCreate_UAV;
+	Ret.TargetableFlags |= TexCreate_RenderTargetable;
+
+	Ret.DebugName = (InPassOutputId == ePId_Output0) ? TEXT("CircleDOFDilate0") : TEXT("CircleDOFDilate1");
+
+	// more precision for additive blending and we need the alpha channel
+	Ret.Format = PF_FloatRGBA;
+
+	return Ret;
+}
+
+
+
+
+
+
+
+
+
+/** Encapsulates the Circle DOF pixel shader. */
+
+static float TemporalHalton2( int32 Index, int32 Base )
+{
+	float Result = 0.0f;
+	float InvBase = 1.0f / Base;
+	float Fraction = InvBase;
+	while( Index > 0 )
+	{
+		Result += ( Index % Base ) * Fraction;
+		Index /= Base;
+		Fraction *= InvBase;
+	}
+	return Result;
+}
+
+static void TemporalRandom2(FVector2D* RESTRICT const Constant, uint32 FrameNumber)
+{
+	Constant->X = TemporalHalton2(FrameNumber & 1023, 2);
+	Constant->Y = TemporalHalton2(FrameNumber & 1023, 3);
+}
+
+template <uint32 NearBlurEnable>
+class FPostProcessCircleDOFPS : public FGlobalShader
+{
+	DECLARE_SHADER_TYPE(FPostProcessCircleDOFPS, Global);
+
+	static bool ShouldCache(EShaderPlatform Platform)
+	{
+		return IsFeatureLevelSupported(Platform, ERHIFeatureLevel::SM5);
+	}
+
+	static void ModifyCompilationEnvironment(EShaderPlatform Platform, FShaderCompilerEnvironment& OutEnvironment)
+	{
+		FGlobalShader::ModifyCompilationEnvironment(Platform,OutEnvironment);
+		OutEnvironment.SetDefine(TEXT("ENABLE_NEAR_BLUR"), NearBlurEnable);
+	}
+
+	/** Default constructor. */
+	FPostProcessCircleDOFPS() {}
+
+public:
+	FPostProcessPassParameters PostprocessParameter;
+	FDeferredPixelShaderParameters DeferredParameters;
+	FShaderParameter DepthOfFieldParams;
+	FShaderParameter RandomOffset;
+	FShaderParameter CircleDofParams;
+
+	/** Initialization constructor. */
+	FPostProcessCircleDOFPS(const ShaderMetaType::CompiledShaderInitializerType& Initializer)
+		: FGlobalShader(Initializer)
+	{
+		PostprocessParameter.Bind(Initializer.ParameterMap);
+		DeferredParameters.Bind(Initializer.ParameterMap);
+		DepthOfFieldParams.Bind(Initializer.ParameterMap,TEXT("DepthOfFieldParams"));
+		RandomOffset.Bind(Initializer.ParameterMap, TEXT("RandomOffset"));
+		CircleDofParams.Bind(Initializer.ParameterMap,TEXT("CircleDofParams"));
+	}
+
+	// FShader interface.
+	virtual bool Serialize(FArchive& Ar) override
+	{
+		bool bShaderHasOutdatedParameters = FGlobalShader::Serialize(Ar);
+		Ar << PostprocessParameter << DeferredParameters << DepthOfFieldParams << RandomOffset << CircleDofParams;
+		return bShaderHasOutdatedParameters;
+	}
+
+	void SetParameters(const FRenderingCompositePassContext& Context)
+	{
+		const FPixelShaderRHIParamRef ShaderRHI = GetPixelShader();
+
+		FGlobalShader::SetParameters(Context.RHICmdList, ShaderRHI, Context.View);
+
+		PostprocessParameter.SetPS(ShaderRHI, Context, TStaticSamplerState<SF_Point,AM_Border,AM_Border,AM_Clamp>::GetRHI());
+
+		DeferredParameters.Set(Context.RHICmdList, ShaderRHI, Context.View);
+
+		{
+			FVector4 DepthOfFieldParamValues[2];
+
+			FRCPassPostProcessBokehDOF::ComputeDepthOfFieldParams(Context, DepthOfFieldParamValues);
+
+			SetShaderValueArray(Context.RHICmdList, ShaderRHI, DepthOfFieldParams, DepthOfFieldParamValues, 2);
+		}
+
+		FVector2D RandomOffsetValue;
+		TemporalRandom2(&RandomOffsetValue, Context.View.Family->FrameNumber);
+		SetShaderValue(Context.RHICmdList, ShaderRHI, RandomOffset, RandomOffsetValue);
+
+		SetShaderValue(Context.RHICmdList, ShaderRHI, CircleDofParams, CircleDofCoc(Context));
+	}
+};
+
+IMPLEMENT_SHADER_TYPE(template<>,FPostProcessCircleDOFPS<0>,TEXT("PostProcessCircleDOF"),TEXT("CirclePS"),SF_Pixel);
+IMPLEMENT_SHADER_TYPE(template<>,FPostProcessCircleDOFPS<1>,TEXT("PostProcessCircleDOF"),TEXT("CirclePS"),SF_Pixel);
+
+void FRCPassPostProcessCircleDOF::Process(FRenderingCompositePassContext& Context)
+{
+	SCOPED_DRAW_EVENT(Context.RHICmdList, CircleDOFApply);
+
+	const FPooledRenderTargetDesc* InputDesc = GetInputDesc(ePId_Input0);
+
+	if(!InputDesc)
+	{
+		// input is not hooked up correctly
+		return;
+	}
+
+	uint32 NumRenderTargets = bNearBlurEnabled ? 2 : 1;
+
+	const FSceneView& View = Context.View;
+	const FSceneViewFamily& ViewFamily = *(View.Family);
+
+	const auto FeatureLevel = Context.GetFeatureLevel();
+	auto ShaderMap = Context.GetShaderMap();
+
+	FIntPoint SrcSize = InputDesc->Extent;
+	FIntPoint DestSize = PassOutputs[0].RenderTargetDesc.Extent;
+
+	// e.g. 4 means the input texture is 4x smaller than the buffer size
+	uint32 ScaleFactor = FSceneRenderTargets::Get(Context.RHICmdList).GetBufferSizeXY().X / SrcSize.X;
+
+	FIntRect SrcRect = View.ViewRect / ScaleFactor;
+	FIntRect DestRect = SrcRect;
+
+	const FSceneRenderTargetItem& DestRenderTarget0 = PassOutputs[0].RequestSurface(Context);
+	const FSceneRenderTargetItem& DestRenderTarget1 = bNearBlurEnabled ? PassOutputs[1].RequestSurface(Context) : FSceneRenderTargetItem();
+
+	// Set the view family's render target/viewport.
+	FTextureRHIParamRef RenderTargets[2] =
+	{
+		DestRenderTarget0.TargetableTexture,
+		DestRenderTarget1.TargetableTexture
+	};
+	SetRenderTargets(Context.RHICmdList, NumRenderTargets, RenderTargets, FTextureRHIParamRef(), 0, NULL);
+
+	FLinearColor ClearColors[2] = 
+	{
+		FLinearColor(0, 0, 0, 0),
+		FLinearColor(0, 0, 0, 0)
+	};
+	// is optimized away if possible (RT size=view size, )
+	Context.RHICmdList.ClearMRT(true, NumRenderTargets, ClearColors, false, 1.0f, false, 0, DestRect);
+
+	Context.SetViewportAndCallRHI(0, 0, 0.0f, DestSize.X, DestSize.Y, 1.0f );
+
+	// set the state
+	Context.RHICmdList.SetBlendState(TStaticBlendState<>::GetRHI());
+	Context.RHICmdList.SetRasterizerState(TStaticRasterizerState<>::GetRHI());
+	Context.RHICmdList.SetDepthStencilState(TStaticDepthStencilState<false, CF_Always>::GetRHI());
+
+	TShaderMapRef<FPostProcessVS> VertexShader(ShaderMap);
+
+	if (bNearBlurEnabled)
+	{
+		static FGlobalBoundShaderState BoundShaderState;
+
+
+		TShaderMapRef< FPostProcessCircleDOFPS<1> > PixelShader(ShaderMap);
+		SetGlobalBoundShaderState(Context.RHICmdList, FeatureLevel, BoundShaderState, GFilterVertexDeclaration.VertexDeclarationRHI, *VertexShader, *PixelShader);
+
+		PixelShader->SetParameters(Context);
+	}
+	else
+	{
+		static FGlobalBoundShaderState BoundShaderState;
+
+		TShaderMapRef< FPostProcessCircleDOFPS<0> > PixelShader(ShaderMap);
+		SetGlobalBoundShaderState(Context.RHICmdList, FeatureLevel, BoundShaderState, GFilterVertexDeclaration.VertexDeclarationRHI, *VertexShader, *PixelShader);
+
+		PixelShader->SetParameters(Context);
+	}
+
+	VertexShader->SetParameters(Context);
+
+	// Draw a quad mapping scene color to the view's render target
+	DrawRectangle(
+		Context.RHICmdList,
+		DestRect.Min.X, DestRect.Min.Y,
+		DestRect.Width() + 1, DestRect.Height() + 1,
+		SrcRect.Min.X, SrcRect.Min.Y,
+		SrcRect.Width() + 1, SrcRect.Height() + 1,
+		DestSize,
+		SrcSize,
+		*VertexShader,
+		EDRF_UseTriangleOptimization);
+
+	Context.RHICmdList.CopyToResolveTarget(DestRenderTarget0.TargetableTexture, DestRenderTarget0.ShaderResourceTexture, false, FResolveParams());
+	Context.RHICmdList.CopyToResolveTarget(DestRenderTarget1.TargetableTexture, DestRenderTarget1.ShaderResourceTexture, false, FResolveParams());
+}
+
+FPooledRenderTargetDesc FRCPassPostProcessCircleDOF::ComputeOutputDesc(EPassOutputId InPassOutputId) const
+{
+	FPooledRenderTargetDesc Ret = GetInput(ePId_Input0)->GetOutput()->RenderTargetDesc;
+
+	Ret.Extent.X = FMath::Max(1, Ret.Extent.X);
+	Ret.Extent.Y = FMath::Max(1, Ret.Extent.Y);
+
+	Ret.Reset();
+	Ret.TargetableFlags &= ~(uint32)TexCreate_UAV;
+	Ret.TargetableFlags |= TexCreate_RenderTargetable;
+
+	Ret.DebugName = (InPassOutputId == ePId_Output0) ? TEXT("CircleDOF0") : TEXT("CircleDOF1");
+
+	// more precision for additive blending and we need the alpha channel
+	Ret.Format = PF_FloatRGBA;
+
+	return Ret;
+}
+
+
+/** Encapsulates  the Circle DOF recombine pixel shader. */
+template <uint32 NearBlurEnable>
+class FPostProcessCircleDOFRecombinePS : public FGlobalShader
+{
+	DECLARE_SHADER_TYPE(FPostProcessCircleDOFRecombinePS, Global);
+
+	static bool ShouldCache(EShaderPlatform Platform)
+	{
+		return IsFeatureLevelSupported(Platform, ERHIFeatureLevel::SM5);
+	}
+
+	static void ModifyCompilationEnvironment(EShaderPlatform Platform, FShaderCompilerEnvironment& OutEnvironment)
+	{
+		FGlobalShader::ModifyCompilationEnvironment(Platform,OutEnvironment);
+		OutEnvironment.SetDefine(TEXT("ENABLE_NEAR_BLUR"), NearBlurEnable);
+	}
+
+	/** Default constructor. */
+	FPostProcessCircleDOFRecombinePS() {}
+
+public:
+	FPostProcessPassParameters PostprocessParameter;
+	FDeferredPixelShaderParameters DeferredParameters;
+	FShaderParameter DepthOfFieldUVLimit;
+	FShaderParameter RandomOffset;
+	FShaderParameter CircleDofParams;
+
+	/** Initialization constructor. */
+	FPostProcessCircleDOFRecombinePS(const ShaderMetaType::CompiledShaderInitializerType& Initializer)
+		: FGlobalShader(Initializer)
+	{
+		PostprocessParameter.Bind(Initializer.ParameterMap);
+		DeferredParameters.Bind(Initializer.ParameterMap);
+		DepthOfFieldUVLimit.Bind(Initializer.ParameterMap,TEXT("DepthOfFieldUVLimit"));
+		RandomOffset.Bind(Initializer.ParameterMap, TEXT("RandomOffset"));
+		CircleDofParams.Bind(Initializer.ParameterMap,TEXT("CircleDofParams"));
+	}
+
+	// FShader interface.
+	virtual bool Serialize(FArchive& Ar) override
+	{
+		bool bShaderHasOutdatedParameters = FGlobalShader::Serialize(Ar);
+		Ar << PostprocessParameter << DeferredParameters << DepthOfFieldUVLimit << RandomOffset << CircleDofParams;
+		return bShaderHasOutdatedParameters;
+	}
+
+	void SetParameters(const FRenderingCompositePassContext& Context)
+	{
+		FSceneRenderTargets& SceneContext = FSceneRenderTargets::Get(Context.RHICmdList);
+		const FPixelShaderRHIParamRef ShaderRHI = GetPixelShader();
+
+		FGlobalShader::SetParameters(Context.RHICmdList, ShaderRHI, Context.View);
+
+		DeferredParameters.Set(Context.RHICmdList, ShaderRHI, Context.View);
+		PostprocessParameter.SetPS(ShaderRHI, Context, TStaticSamplerState<SF_Bilinear,AM_Clamp,AM_Clamp,AM_Clamp>::GetRHI());
+
+		// Compute out of bounds UVs in the source texture.
+		FVector4 Bounds;
+		Bounds.X = (((float)((Context.View.ViewRect.Min.X + 1) & (~1))) + 3.0f) / ((float)(SceneContext.GetBufferSizeXY().X));
+		Bounds.Y = (((float)((Context.View.ViewRect.Min.Y + 1) & (~1))) + 3.0f) / ((float)(SceneContext.GetBufferSizeXY().Y));
+		Bounds.Z = (((float)(Context.View.ViewRect.Max.X & (~1))) - 3.0f) / ((float)(SceneContext.GetBufferSizeXY().X));
+		Bounds.W = (((float)(Context.View.ViewRect.Max.Y & (~1))) - 3.0f) / ((float)(SceneContext.GetBufferSizeXY().Y));
+
+		SetShaderValue(Context.RHICmdList, ShaderRHI, DepthOfFieldUVLimit, Bounds);
+
+		FVector2D RandomOffsetValue;
+		TemporalRandom2(&RandomOffsetValue, Context.View.Family->FrameNumber);
+		SetShaderValue(Context.RHICmdList, ShaderRHI, RandomOffset, RandomOffsetValue);
+
+		SetShaderValue(Context.RHICmdList, ShaderRHI, CircleDofParams, CircleDofCoc(Context));
+	}
+};
+
+IMPLEMENT_SHADER_TYPE(template<>,FPostProcessCircleDOFRecombinePS<0>,TEXT("PostProcessCircleDOF"),TEXT("MainCircleRecombinePS"),SF_Pixel);
+IMPLEMENT_SHADER_TYPE(template<>,FPostProcessCircleDOFRecombinePS<1>,TEXT("PostProcessCircleDOF"),TEXT("MainCircleRecombinePS"),SF_Pixel);
+
+void FRCPassPostProcessCircleDOFRecombine::Process(FRenderingCompositePassContext& Context)
+{
+	SCOPED_DRAW_EVENT(Context.RHICmdList, CircleDOFRecombine);
+
+	const FPooledRenderTargetDesc* InputDesc = GetInputDesc(ePId_Input0);
+
+	if(!InputDesc)
+	{
+		// input is not hooked up correctly
+		return;
+	}
+
+	const FSceneView& View = Context.View;
+
+	const auto FeatureLevel = Context.GetFeatureLevel();
+	auto ShaderMap = Context.GetShaderMap();
+
+	FIntPoint TexSize = InputDesc->Extent;
+
+	// usually 1, 2, 4 or 8
+	uint32 ScaleToFullRes = FSceneRenderTargets::Get(Context.RHICmdList).GetBufferSizeXY().X / TexSize.X;
+
+	FIntRect HalfResViewRect = View.ViewRect / ScaleToFullRes;
+
+	const FSceneRenderTargetItem& DestRenderTarget = PassOutputs[0].RequestSurface(Context);
+
+	// Set the view family's render target/viewport.
+	SetRenderTarget(Context.RHICmdList, DestRenderTarget.TargetableTexture, FTextureRHIRef());
+
+	// is optimized away if possible (RT size=view size, )
+	Context.RHICmdList.Clear(true, FLinearColor::Black, false, 1.0f, false, 0, View.ViewRect);
+
+	Context.SetViewportAndCallRHI(View.ViewRect);
+
+	// set the state
+	Context.RHICmdList.SetBlendState(TStaticBlendState<>::GetRHI());
+	Context.RHICmdList.SetRasterizerState(TStaticRasterizerState<>::GetRHI());
+	Context.RHICmdList.SetDepthStencilState(TStaticDepthStencilState<false, CF_Always>::GetRHI());
+
+	TShaderMapRef<FPostProcessVS> VertexShader(ShaderMap);
+
+	if (bNearBlurEnabled)
+	{
+		static FGlobalBoundShaderState BoundShaderState;
+
+		TShaderMapRef< FPostProcessCircleDOFRecombinePS<1> > PixelShader(ShaderMap);
+		SetGlobalBoundShaderState(Context.RHICmdList, FeatureLevel, BoundShaderState, GFilterVertexDeclaration.VertexDeclarationRHI, *VertexShader, *PixelShader);
+		PixelShader->SetParameters(Context);
+	}
+	else
+	{
+		static FGlobalBoundShaderState BoundShaderState;
+
+		TShaderMapRef< FPostProcessCircleDOFRecombinePS<0> > PixelShader(ShaderMap);
+		SetGlobalBoundShaderState(Context.RHICmdList, FeatureLevel, BoundShaderState, GFilterVertexDeclaration.VertexDeclarationRHI, *VertexShader, *PixelShader);
+		PixelShader->SetParameters(Context);
+	}
+
+	VertexShader->SetParameters(Context);
+
+	// Draw a quad mapping scene color to the view's render target
+	DrawRectangle(
+		Context.RHICmdList,
+		0, 0,
+		View.ViewRect.Width(), View.ViewRect.Height(),
+		View.ViewRect.Min.X, View.ViewRect.Min.Y,
+		View.ViewRect.Width(), View.ViewRect.Height(),
+		View.ViewRect.Size(),
+		TexSize,
+		*VertexShader,
+		EDRF_UseTriangleOptimization);
+
+	Context.RHICmdList.CopyToResolveTarget(DestRenderTarget.TargetableTexture, DestRenderTarget.ShaderResourceTexture, false, FResolveParams());
+}
+
+FPooledRenderTargetDesc FRCPassPostProcessCircleDOFRecombine::ComputeOutputDesc(EPassOutputId InPassOutputId) const
+{
+	FPooledRenderTargetDesc Ret = GetInput(ePId_Input0)->GetOutput()->RenderTargetDesc;
+
+	Ret.Reset();
+	Ret.DebugName = TEXT("CircleDOFRecombine");
+
+	return Ret;
+}
diff --git a/Engine/Source/Runtime/Renderer/Private/PostProcess/PostProcessCircleDOF.h b/Engine/Source/Runtime/Renderer/Private/PostProcess/PostProcessCircleDOF.h
new file mode 100644
--- /dev/null
+++ b/Engine/Source/Runtime/Renderer/Private/PostProcess/PostProcessCircleDOF.h
@@ -0,0 +1,94 @@
+// Copyright 1998-2015 Epic Games, Inc. All Rights Reserved.
+
+/*=============================================================================
+	PostProcessDOF.h: Post process Depth of Field implementation.
+=============================================================================*/
+
+#pragma once
+
+#include "RenderingCompositionGraph.h"
+
+
+// down sample and setup DOF input
+// ePId_Input0: SceneColor
+// ePId_Input1: SceneDepth
+// derives from TRenderingCompositePassBase<InputCount, OutputCount> 
+class FRCPassPostProcessCircleDOFSetup : public TRenderingCompositePassBase<2, 2>
+{
+public:
+
+	FRCPassPostProcessCircleDOFSetup(bool bInNearBlurEnabled) : bNearBlurEnabled ( bInNearBlurEnabled ) {}
+
+	// interface FRenderingCompositePass ---------
+
+	virtual void Process(FRenderingCompositePassContext& Context) override;
+	virtual void Release() override { delete this; }
+	virtual FPooledRenderTargetDesc ComputeOutputDesc(EPassOutputId InPassOutputId) const override;
+
+private:
+
+	bool bNearBlurEnabled;
+};
+
+
+// ePId_Input0: SceneColor
+class FRCPassPostProcessCircleDOFDilate : public TRenderingCompositePassBase<1, 1>
+{
+public:
+
+	FRCPassPostProcessCircleDOFDilate(void) {}
+
+	// interface FRenderingCompositePass ---------
+
+	virtual void Process(FRenderingCompositePassContext& Context) override;
+	virtual void Release() override { delete this; }
+	virtual FPooledRenderTargetDesc ComputeOutputDesc(EPassOutputId InPassOutputId) const override;
+};
+
+
+
+
+// ePId_Input0: setup results for far
+// ePId_Input1: setup results for near
+// derives from TRenderingCompositePassBase<InputCount, OutputCount> 
+class FRCPassPostProcessCircleDOF : public TRenderingCompositePassBase<2, 2>
+{
+public:
+
+	FRCPassPostProcessCircleDOF(bool bInNearBlurEnabled) : bNearBlurEnabled ( bInNearBlurEnabled ) {}
+
+	// interface FRenderingCompositePass ---------
+
+	virtual void Process(FRenderingCompositePassContext& Context) override;
+	virtual void Release() override { delete this; }
+	virtual FPooledRenderTargetDesc ComputeOutputDesc(EPassOutputId InPassOutputId) const override;
+
+private:
+
+	bool bNearBlurEnabled;
+};
+
+// derives from TRenderingCompositePassBase<InputCount, OutputCount> 
+// ePId_Input0: Full res scene color
+// ePId_Input1: output 0 from the DOFSetup (possibly further blurred)
+// ePId_Input2: output 1 from the DOFSetup (possibly further blurred)
+class FRCPassPostProcessCircleDOFRecombine : public TRenderingCompositePassBase<3, 1>
+{
+public:
+
+	FRCPassPostProcessCircleDOFRecombine(bool bInNearBlurEnabled) : bNearBlurEnabled ( bInNearBlurEnabled ) {}
+
+	// interface FRenderingCompositePass ---------
+	virtual void Process(FRenderingCompositePassContext& Context) override;
+	virtual void Release() override { delete this; }
+	virtual FPooledRenderTargetDesc ComputeOutputDesc(EPassOutputId InPassOutputId) const override;
+
+private:
+
+	bool bNearBlurEnabled;
+};
+
+// to verify this can be used: http://www.radical.org/aov
+// @return in mm, assuming this sensor: APS-C 24.576 mm
+float ComputeFocalLengthFromFov(const FSceneView& View);
+
diff --git a/Engine/Source/Runtime/Renderer/Private/PostProcess/PostProcessDOF.cpp b/Engine/Source/Runtime/Renderer/Private/PostProcess/PostProcessDOF.cpp
--- a/Engine/Source/Runtime/Renderer/Private/PostProcess/PostProcessDOF.cpp
+++ b/Engine/Source/Runtime/Renderer/Private/PostProcess/PostProcessDOF.cpp
@@ -346,819 +346,4 @@ FPooledRenderTargetDesc FRCPassPostProcessDOFRecombine::ComputeOutputDesc(EPassO
 	Ret.DebugName = TEXT("DOFRecombine");
 
 	return Ret;
-}
-
-
-
-
-
-
-
-// Convert f-stop and focal distance into projected size in half resolution pixels.
-// Setup depth based blur.
-static FVector4 CircleDofCoc(const FRenderingCompositePassContext& Context)
-{
-	// Convert FOV to focal length,
-	// 
-	// fov = 2 * atan(d/(2*f))
-	// where,
-	//   d = sensor dimension (APS-C 24.576 mm)
-	//   f = focal length
-	// 
-	// f = 0.5 * d * (1/tan(fov/2))
-	float HalfFOV = FMath::Atan(1.0f / Context.View.ViewMatrices.ProjMatrix.M[0][0]);
-	float FocalLength = 0.5f * 24.576f * (1.0f/FMath::Tan(HalfFOV));
-	 
-	// Convert focal distance in world position to mm.
-	// Conversion is 1 world position = 1 cm.
-	float Distance = Context.View.FinalPostProcessSettings.DepthOfFieldFocalDistance;
-	Distance *= 10.0f;
-
-	// Convert f-stop, focal length, and focal distance to
-	// projected circle of confusion size at infinity in mm.
-	//
-	// coc = f*f / (n * (d - f))
-	// where,
-	//   f = focal length
-	//   d = focal distance
-	//   n = fstop (where n is the "n" in "f/n")
-	float Radius = FocalLength * FocalLength / (Context.View.FinalPostProcessSettings.DepthOfFieldFstop * (Distance - FocalLength));
-
-	// Scale so that APS-C 24.576 mm = full frame.
-	// Convert mm to pixels.
-	float Width = (float)Context.View.ViewRect.Width();
-	Radius = Radius * Width * (1.0f/24.576f);
-
-	// Convert diameter to radius at half resolution (algorithm radius is at half resolution).
-	Radius *= 0.25f;
-
-	// Comment out for now, allowing settings which the algorithm cannot cleanly do.
-	#if 0
-		// Limit to algorithm max size.
-		if(Radius > 6.0f) 
-		{
-			Radius = 6.0f; 
-		}
-	#endif
-
-	// The DepthOfFieldDepthBlurAmount = km at which depth blur is 50%.
-	// Need to convert to cm here.
-	return FVector4(
-		Radius, 
-		1.0f/(Context.View.FinalPostProcessSettings.DepthOfFieldDepthBlurAmount * 100000.0),
-		Context.View.FinalPostProcessSettings.DepthOfFieldDepthBlurRadius * Width / 1920.0f,
-		Width / 1920.0f);
-}
-
-/** Encapsulates the Circle DOF setup pixel shader. */
-template <uint32 NearBlurEnable>
-class FPostProcessCircleDOFSetupPS : public FGlobalShader
-{
-	DECLARE_SHADER_TYPE(FPostProcessCircleDOFSetupPS, Global);
-
-	static bool ShouldCache(EShaderPlatform Platform)
-	{
-		return IsFeatureLevelSupported(Platform, ERHIFeatureLevel::SM5);
-	}
-
-	static void ModifyCompilationEnvironment(EShaderPlatform Platform, FShaderCompilerEnvironment& OutEnvironment)
-	{
-		FGlobalShader::ModifyCompilationEnvironment(Platform,OutEnvironment);
-		OutEnvironment.SetDefine(TEXT("ENABLE_NEAR_BLUR"), NearBlurEnable);
-	}
-
-	/** Default constructor. */
-	FPostProcessCircleDOFSetupPS() {}
-
-public:
-	FPostProcessPassParameters PostprocessParameter;
-	FDeferredPixelShaderParameters DeferredParameters;
-	FShaderParameter DepthOfFieldParams;
-	FShaderParameter CircleDofParams;
-
-	/** Initialization constructor. */
-	FPostProcessCircleDOFSetupPS(const ShaderMetaType::CompiledShaderInitializerType& Initializer)
-		: FGlobalShader(Initializer)
-	{
-		PostprocessParameter.Bind(Initializer.ParameterMap);
-		DeferredParameters.Bind(Initializer.ParameterMap);
-		DepthOfFieldParams.Bind(Initializer.ParameterMap,TEXT("DepthOfFieldParams"));
-		CircleDofParams.Bind(Initializer.ParameterMap,TEXT("CircleDofParams"));
-	}
-
-	// FShader interface.
-	virtual bool Serialize(FArchive& Ar) override
-	{
-		bool bShaderHasOutdatedParameters = FGlobalShader::Serialize(Ar);
-		Ar << PostprocessParameter << DeferredParameters << DepthOfFieldParams << CircleDofParams;
-		return bShaderHasOutdatedParameters;
-	}
-
-	void SetParameters(const FRenderingCompositePassContext& Context)
-	{
-		const FPixelShaderRHIParamRef ShaderRHI = GetPixelShader();
-
-		FGlobalShader::SetParameters(Context.RHICmdList, ShaderRHI, Context.View);
-
-		PostprocessParameter.SetPS(ShaderRHI, Context, TStaticSamplerState<SF_Point,AM_Border,AM_Border,AM_Clamp>::GetRHI());
-
-		DeferredParameters.Set(Context.RHICmdList, ShaderRHI, Context.View);
-
-		{
-			FVector4 DepthOfFieldParamValues[2];
-
-			FRCPassPostProcessBokehDOF::ComputeDepthOfFieldParams(Context, DepthOfFieldParamValues);
-
-			SetShaderValueArray(Context.RHICmdList, ShaderRHI, DepthOfFieldParams, DepthOfFieldParamValues, 2);
-		}
-
-		SetShaderValue(Context.RHICmdList, ShaderRHI, CircleDofParams, CircleDofCoc(Context));
-	}
-};
-
-IMPLEMENT_SHADER_TYPE(template<>,FPostProcessCircleDOFSetupPS<0>,TEXT("PostProcessDOF"),TEXT("CircleSetupPS"),SF_Pixel);
-IMPLEMENT_SHADER_TYPE(template<>,FPostProcessCircleDOFSetupPS<1>,TEXT("PostProcessDOF"),TEXT("CircleSetupPS"),SF_Pixel);
-
-void FRCPassPostProcessCircleDOFSetup::Process(FRenderingCompositePassContext& Context)
-{
-	SCOPED_DRAW_EVENT(Context.RHICmdList, CircleDOFSetup);
-
-	const FPooledRenderTargetDesc* InputDesc = GetInputDesc(ePId_Input0);
-
-	if(!InputDesc)
-	{
-		// input is not hooked up correctly
-		return;
-	}
-
-	uint32 NumRenderTargets = bNearBlurEnabled ? 2 : 1;
-
-	const FSceneView& View = Context.View;
-	const FSceneViewFamily& ViewFamily = *(View.Family);
-
-	const auto FeatureLevel = Context.GetFeatureLevel();
-	auto ShaderMap = Context.GetShaderMap();
-
-	FIntPoint SrcSize = InputDesc->Extent;
-	FIntPoint DestSize = PassOutputs[0].RenderTargetDesc.Extent;
-
-	// e.g. 4 means the input texture is 4x smaller than the buffer size
-	uint32 ScaleFactor = FSceneRenderTargets::Get(Context.RHICmdList).GetBufferSizeXY().X / SrcSize.X;
-
-	FIntRect SrcRect = View.ViewRect / ScaleFactor;
-	FIntRect DestRect = SrcRect / 2;
-
-	const FSceneRenderTargetItem& DestRenderTarget0 = PassOutputs[0].RequestSurface(Context);
-	const FSceneRenderTargetItem& DestRenderTarget1 = bNearBlurEnabled ? PassOutputs[1].RequestSurface(Context) : FSceneRenderTargetItem();
-
-	// Set the view family's render target/viewport.
-	FTextureRHIParamRef RenderTargets[2] =
-	{
-		DestRenderTarget0.TargetableTexture,
-		DestRenderTarget1.TargetableTexture
-	};
-	SetRenderTargets(Context.RHICmdList, NumRenderTargets, RenderTargets, FTextureRHIParamRef(), 0, NULL);
-
-	FLinearColor ClearColors[2] = 
-	{
-		FLinearColor(0, 0, 0, 0),
-		FLinearColor(0, 0, 0, 0)
-	};
-	// is optimized away if possible (RT size=view size, )
-	Context.RHICmdList.ClearMRT(true, NumRenderTargets, ClearColors, false, 1.0f, false, 0, DestRect);
-
-	Context.SetViewportAndCallRHI(0, 0, 0.0f, DestSize.X, DestSize.Y, 1.0f );
-
-	// set the state
-	Context.RHICmdList.SetBlendState(TStaticBlendState<>::GetRHI());
-	Context.RHICmdList.SetRasterizerState(TStaticRasterizerState<>::GetRHI());
-	Context.RHICmdList.SetDepthStencilState(TStaticDepthStencilState<false, CF_Always>::GetRHI());
-
-	TShaderMapRef<FPostProcessVS> VertexShader(ShaderMap);
-
-	if (bNearBlurEnabled)
-	{
-		static FGlobalBoundShaderState BoundShaderState;
-
-
-		TShaderMapRef< FPostProcessCircleDOFSetupPS<1> > PixelShader(ShaderMap);
-		SetGlobalBoundShaderState(Context.RHICmdList, FeatureLevel, BoundShaderState, GFilterVertexDeclaration.VertexDeclarationRHI, *VertexShader, *PixelShader);
-
-		PixelShader->SetParameters(Context);
-	}
-	else
-	{
-		static FGlobalBoundShaderState BoundShaderState;
-
-		TShaderMapRef< FPostProcessCircleDOFSetupPS<0> > PixelShader(ShaderMap);
-		SetGlobalBoundShaderState(Context.RHICmdList, FeatureLevel, BoundShaderState, GFilterVertexDeclaration.VertexDeclarationRHI, *VertexShader, *PixelShader);
-
-		PixelShader->SetParameters(Context);
-	}
-
-	VertexShader->SetParameters(Context);
-
-	// Draw a quad mapping scene color to the view's render target
-	DrawRectangle(
-		Context.RHICmdList,
-		DestRect.Min.X, DestRect.Min.Y,
-		DestRect.Width() + 1, DestRect.Height() + 1,
-		SrcRect.Min.X, SrcRect.Min.Y,
-		SrcRect.Width() + 1, SrcRect.Height() + 1,
-		DestSize,
-		SrcSize,
-		*VertexShader,
-		EDRF_UseTriangleOptimization);
-
-	Context.RHICmdList.CopyToResolveTarget(DestRenderTarget0.TargetableTexture, DestRenderTarget0.ShaderResourceTexture, false, FResolveParams());
-	Context.RHICmdList.CopyToResolveTarget(DestRenderTarget1.TargetableTexture, DestRenderTarget1.ShaderResourceTexture, false, FResolveParams());
-}
-
-FPooledRenderTargetDesc FRCPassPostProcessCircleDOFSetup::ComputeOutputDesc(EPassOutputId InPassOutputId) const
-{
-	FPooledRenderTargetDesc Ret = GetInput(ePId_Input0)->GetOutput()->RenderTargetDesc;
-
-	//	Ret.Extent = FIntPoint::DivideAndRoundUp(ret.Extent, 2);
-	Ret.Extent /= 2;
-	Ret.Extent.X = FMath::Max(1, Ret.Extent.X);
-	Ret.Extent.Y = FMath::Max(1, Ret.Extent.Y);
-
-	Ret.Reset();
-	Ret.TargetableFlags &= ~(uint32)TexCreate_UAV;
-	Ret.TargetableFlags |= TexCreate_RenderTargetable;
-
-	Ret.DebugName = (InPassOutputId == ePId_Output0) ? TEXT("CircleDOFSetup0") : TEXT("CircleDOFSetup1");
-
-	// more precision for additive blending and we need the alpha channel
-	Ret.Format = PF_FloatRGBA;
-
-	return Ret;
-}
-
-
-
-
-/** Encapsulates the Circle DOF Dilate pixel shader. */
-template <uint32 NearBlurEnable>
-class FPostProcessCircleDOFDilatePS : public FGlobalShader
-{
-	DECLARE_SHADER_TYPE(FPostProcessCircleDOFDilatePS, Global);
-
-	static bool ShouldCache(EShaderPlatform Platform)
-	{
-		return IsFeatureLevelSupported(Platform, ERHIFeatureLevel::SM5);
-	}
-
-	static void ModifyCompilationEnvironment(EShaderPlatform Platform, FShaderCompilerEnvironment& OutEnvironment)
-	{
-		FGlobalShader::ModifyCompilationEnvironment(Platform,OutEnvironment);
-		OutEnvironment.SetDefine(TEXT("ENABLE_NEAR_BLUR"), NearBlurEnable);
-	}
-
-	/** Default constructor. */
-	FPostProcessCircleDOFDilatePS() {}
-
-public:
-	FPostProcessPassParameters PostprocessParameter;
-	FDeferredPixelShaderParameters DeferredParameters;
-	FShaderParameter DepthOfFieldParams;
-	FShaderParameter CircleDofParams;
-
-	/** Initialization constructor. */
-	FPostProcessCircleDOFDilatePS(const ShaderMetaType::CompiledShaderInitializerType& Initializer)
-		: FGlobalShader(Initializer)
-	{
-		PostprocessParameter.Bind(Initializer.ParameterMap);
-		DeferredParameters.Bind(Initializer.ParameterMap);
-		DepthOfFieldParams.Bind(Initializer.ParameterMap,TEXT("DepthOfFieldParams"));
-		CircleDofParams.Bind(Initializer.ParameterMap,TEXT("CircleDofParams"));
-	}
-
-	// FShader interface.
-	virtual bool Serialize(FArchive& Ar) override
-	{
-		bool bShaderHasOutdatedParameters = FGlobalShader::Serialize(Ar);
-		Ar << PostprocessParameter << DeferredParameters << DepthOfFieldParams << CircleDofParams;
-		return bShaderHasOutdatedParameters;
-	}
-
-	void SetParameters(const FRenderingCompositePassContext& Context)
-	{
-		const FPixelShaderRHIParamRef ShaderRHI = GetPixelShader();
-
-		FGlobalShader::SetParameters(Context.RHICmdList, ShaderRHI, Context.View);
-
-		PostprocessParameter.SetPS(ShaderRHI, Context, TStaticSamplerState<SF_Point,AM_Border,AM_Border,AM_Clamp>::GetRHI());
-
-		DeferredParameters.Set(Context.RHICmdList, ShaderRHI, Context.View);
-
-		{
-			FVector4 DepthOfFieldParamValues[2];
-
-			FRCPassPostProcessBokehDOF::ComputeDepthOfFieldParams(Context, DepthOfFieldParamValues);
-
-			SetShaderValueArray(Context.RHICmdList, ShaderRHI, DepthOfFieldParams, DepthOfFieldParamValues, 2);
-		}
-
-		SetShaderValue(Context.RHICmdList, ShaderRHI, CircleDofParams, CircleDofCoc(Context));
-	}
-};
-
-IMPLEMENT_SHADER_TYPE(template<>,FPostProcessCircleDOFDilatePS<0>,TEXT("PostProcessDOF"),TEXT("CircleDilatePS"),SF_Pixel);
-IMPLEMENT_SHADER_TYPE(template<>,FPostProcessCircleDOFDilatePS<1>,TEXT("PostProcessDOF"),TEXT("CircleDilatePS"),SF_Pixel);
-
-void FRCPassPostProcessCircleDOFDilate::Process(FRenderingCompositePassContext& Context)
-{
-	SCOPED_DRAW_EVENT(Context.RHICmdList, CircleDOFNear);
-
-	const FPooledRenderTargetDesc* InputDesc = GetInputDesc(ePId_Input0);
-
-	if(!InputDesc)
-	{
-		// input is not hooked up correctly
-		return;
-	}
-
-	uint32 NumRenderTargets = 1;
-
-	const FSceneView& View = Context.View;
-	const FSceneViewFamily& ViewFamily = *(View.Family);
-
-	const auto FeatureLevel = Context.GetFeatureLevel();
-	auto ShaderMap = Context.GetShaderMap();
-
-	FIntPoint SrcSize = InputDesc->Extent;
-	FIntPoint DestSize = PassOutputs[0].RenderTargetDesc.Extent;
-
-	// e.g. 4 means the input texture is 4x smaller than the buffer size
-	uint32 ScaleFactor = FSceneRenderTargets::Get(Context.RHICmdList).GetBufferSizeXY().X / SrcSize.X;
-
-	FIntRect SrcRect = View.ViewRect / ScaleFactor;
-	FIntRect DestRect = SrcRect / 2;
-
-	const FSceneRenderTargetItem& DestRenderTarget0 = PassOutputs[0].RequestSurface(Context);
-	const FSceneRenderTargetItem& DestRenderTarget1 = FSceneRenderTargetItem();
-
-	// Set the view family's render target/viewport.
-	FTextureRHIParamRef RenderTargets[2] =
-	{
-		DestRenderTarget0.TargetableTexture,
-		DestRenderTarget1.TargetableTexture
-	};
-	SetRenderTargets(Context.RHICmdList, NumRenderTargets, RenderTargets, FTextureRHIParamRef(), 0, NULL);
-
-	FLinearColor ClearColors[2] = 
-	{
-		FLinearColor(0, 0, 0, 0),
-		FLinearColor(0, 0, 0, 0)
-	};
-	// is optimized away if possible (RT size=view size, )
-	Context.RHICmdList.ClearMRT(true, NumRenderTargets, ClearColors, false, 1.0f, false, 0, DestRect);
-
-	Context.SetViewportAndCallRHI(0, 0, 0.0f, DestSize.X, DestSize.Y, 1.0f );
-
-	// set the state
-	Context.RHICmdList.SetBlendState(TStaticBlendState<>::GetRHI());
-	Context.RHICmdList.SetRasterizerState(TStaticRasterizerState<>::GetRHI());
-	Context.RHICmdList.SetDepthStencilState(TStaticDepthStencilState<false, CF_Always>::GetRHI());
-
-	TShaderMapRef<FPostProcessVS> VertexShader(ShaderMap);
-
-	if (false)
-	{
-		static FGlobalBoundShaderState BoundShaderState;
-
-
-		TShaderMapRef< FPostProcessCircleDOFDilatePS<1> > PixelShader(ShaderMap);
-		SetGlobalBoundShaderState(Context.RHICmdList, FeatureLevel, BoundShaderState, GFilterVertexDeclaration.VertexDeclarationRHI, *VertexShader, *PixelShader);
-
-		PixelShader->SetParameters(Context);
-	}
-	else
-	{
-		static FGlobalBoundShaderState BoundShaderState;
-
-		TShaderMapRef< FPostProcessCircleDOFDilatePS<0> > PixelShader(ShaderMap);
-		SetGlobalBoundShaderState(Context.RHICmdList, FeatureLevel, BoundShaderState, GFilterVertexDeclaration.VertexDeclarationRHI, *VertexShader, *PixelShader);
-
-		PixelShader->SetParameters(Context);
-	}
-
-	VertexShader->SetParameters(Context);
-
-	// Draw a quad mapping scene color to the view's render target
-	DrawRectangle(
-		Context.RHICmdList,
-		DestRect.Min.X, DestRect.Min.Y,
-		DestRect.Width() + 1, DestRect.Height() + 1,
-		SrcRect.Min.X, SrcRect.Min.Y,
-		SrcRect.Width() + 1, SrcRect.Height() + 1,
-		DestSize,
-		SrcSize,
-		*VertexShader,
-		EDRF_UseTriangleOptimization);
-
-	Context.RHICmdList.CopyToResolveTarget(DestRenderTarget0.TargetableTexture, DestRenderTarget0.ShaderResourceTexture, false, FResolveParams());
-	Context.RHICmdList.CopyToResolveTarget(DestRenderTarget1.TargetableTexture, DestRenderTarget1.ShaderResourceTexture, false, FResolveParams());
-}
-
-FPooledRenderTargetDesc FRCPassPostProcessCircleDOFDilate::ComputeOutputDesc(EPassOutputId InPassOutputId) const
-{
-	FPooledRenderTargetDesc Ret = GetInput(ePId_Input0)->GetOutput()->RenderTargetDesc;
-
-	//	Ret.Extent = FIntPoint::DivideAndRoundUp(ret.Extent, 2);
-	Ret.Extent /= 2;
-	Ret.Extent.X = FMath::Max(1, Ret.Extent.X);
-	Ret.Extent.Y = FMath::Max(1, Ret.Extent.Y);
-
-	Ret.Reset();
-	Ret.TargetableFlags &= ~(uint32)TexCreate_UAV;
-	Ret.TargetableFlags |= TexCreate_RenderTargetable;
-
-	Ret.DebugName = (InPassOutputId == ePId_Output0) ? TEXT("CircleDOFDilate0") : TEXT("CircleDOFDilate1");
-
-	// more precision for additive blending and we need the alpha channel
-	Ret.Format = PF_FloatRGBA;
-
-	return Ret;
-}
-
-
-
-
-
-
-
-
-
-/** Encapsulates the Circle DOF pixel shader. */
-
-static float TemporalHalton2( int32 Index, int32 Base )
-{
-	float Result = 0.0f;
-	float InvBase = 1.0f / Base;
-	float Fraction = InvBase;
-	while( Index > 0 )
-	{
-		Result += ( Index % Base ) * Fraction;
-		Index /= Base;
-		Fraction *= InvBase;
-	}
-	return Result;
-}
-
-static void TemporalRandom2(FVector2D* RESTRICT const Constant, uint32 FrameNumber)
-{
-	Constant->X = TemporalHalton2(FrameNumber & 1023, 2);
-	Constant->Y = TemporalHalton2(FrameNumber & 1023, 3);
-}
-
-template <uint32 NearBlurEnable>
-class FPostProcessCircleDOFPS : public FGlobalShader
-{
-	DECLARE_SHADER_TYPE(FPostProcessCircleDOFPS, Global);
-
-	static bool ShouldCache(EShaderPlatform Platform)
-	{
-		return IsFeatureLevelSupported(Platform, ERHIFeatureLevel::SM5);
-	}
-
-	static void ModifyCompilationEnvironment(EShaderPlatform Platform, FShaderCompilerEnvironment& OutEnvironment)
-	{
-		FGlobalShader::ModifyCompilationEnvironment(Platform,OutEnvironment);
-		OutEnvironment.SetDefine(TEXT("ENABLE_NEAR_BLUR"), NearBlurEnable);
-	}
-
-	/** Default constructor. */
-	FPostProcessCircleDOFPS() {}
-
-public:
-	FPostProcessPassParameters PostprocessParameter;
-	FDeferredPixelShaderParameters DeferredParameters;
-	FShaderParameter DepthOfFieldParams;
-	FShaderParameter RandomOffset;
-	FShaderParameter CircleDofParams;
-
-	/** Initialization constructor. */
-	FPostProcessCircleDOFPS(const ShaderMetaType::CompiledShaderInitializerType& Initializer)
-		: FGlobalShader(Initializer)
-	{
-		PostprocessParameter.Bind(Initializer.ParameterMap);
-		DeferredParameters.Bind(Initializer.ParameterMap);
-		DepthOfFieldParams.Bind(Initializer.ParameterMap,TEXT("DepthOfFieldParams"));
-		RandomOffset.Bind(Initializer.ParameterMap, TEXT("RandomOffset"));
-		CircleDofParams.Bind(Initializer.ParameterMap,TEXT("CircleDofParams"));
-	}
-
-	// FShader interface.
-	virtual bool Serialize(FArchive& Ar) override
-	{
-		bool bShaderHasOutdatedParameters = FGlobalShader::Serialize(Ar);
-		Ar << PostprocessParameter << DeferredParameters << DepthOfFieldParams << RandomOffset << CircleDofParams;
-		return bShaderHasOutdatedParameters;
-	}
-
-	void SetParameters(const FRenderingCompositePassContext& Context)
-	{
-		const FPixelShaderRHIParamRef ShaderRHI = GetPixelShader();
-
-		FGlobalShader::SetParameters(Context.RHICmdList, ShaderRHI, Context.View);
-
-		PostprocessParameter.SetPS(ShaderRHI, Context, TStaticSamplerState<SF_Point,AM_Border,AM_Border,AM_Clamp>::GetRHI());
-
-		DeferredParameters.Set(Context.RHICmdList, ShaderRHI, Context.View);
-
-		{
-			FVector4 DepthOfFieldParamValues[2];
-
-			FRCPassPostProcessBokehDOF::ComputeDepthOfFieldParams(Context, DepthOfFieldParamValues);
-
-			SetShaderValueArray(Context.RHICmdList, ShaderRHI, DepthOfFieldParams, DepthOfFieldParamValues, 2);
-		}
-
-		FVector2D RandomOffsetValue;
-		TemporalRandom2(&RandomOffsetValue, Context.View.Family->FrameNumber);
-		SetShaderValue(Context.RHICmdList, ShaderRHI, RandomOffset, RandomOffsetValue);
-
-		SetShaderValue(Context.RHICmdList, ShaderRHI, CircleDofParams, CircleDofCoc(Context));
-	}
-};
-
-IMPLEMENT_SHADER_TYPE(template<>,FPostProcessCircleDOFPS<0>,TEXT("PostProcessDOF"),TEXT("CirclePS"),SF_Pixel);
-IMPLEMENT_SHADER_TYPE(template<>,FPostProcessCircleDOFPS<1>,TEXT("PostProcessDOF"),TEXT("CirclePS"),SF_Pixel);
-
-void FRCPassPostProcessCircleDOF::Process(FRenderingCompositePassContext& Context)
-{
-	SCOPED_DRAW_EVENT(Context.RHICmdList, CircleDOFApply);
-
-	const FPooledRenderTargetDesc* InputDesc = GetInputDesc(ePId_Input0);
-
-	if(!InputDesc)
-	{
-		// input is not hooked up correctly
-		return;
-	}
-
-	uint32 NumRenderTargets = bNearBlurEnabled ? 2 : 1;
-
-	const FSceneView& View = Context.View;
-	const FSceneViewFamily& ViewFamily = *(View.Family);
-
-	const auto FeatureLevel = Context.GetFeatureLevel();
-	auto ShaderMap = Context.GetShaderMap();
-
-	FIntPoint SrcSize = InputDesc->Extent;
-	FIntPoint DestSize = PassOutputs[0].RenderTargetDesc.Extent;
-
-	// e.g. 4 means the input texture is 4x smaller than the buffer size
-	uint32 ScaleFactor = FSceneRenderTargets::Get(Context.RHICmdList).GetBufferSizeXY().X / SrcSize.X;
-
-	FIntRect SrcRect = View.ViewRect / ScaleFactor;
-	FIntRect DestRect = SrcRect;
-
-	const FSceneRenderTargetItem& DestRenderTarget0 = PassOutputs[0].RequestSurface(Context);
-	const FSceneRenderTargetItem& DestRenderTarget1 = bNearBlurEnabled ? PassOutputs[1].RequestSurface(Context) : FSceneRenderTargetItem();
-
-	// Set the view family's render target/viewport.
-	FTextureRHIParamRef RenderTargets[2] =
-	{
-		DestRenderTarget0.TargetableTexture,
-		DestRenderTarget1.TargetableTexture
-	};
-	SetRenderTargets(Context.RHICmdList, NumRenderTargets, RenderTargets, FTextureRHIParamRef(), 0, NULL);
-
-	FLinearColor ClearColors[2] = 
-	{
-		FLinearColor(0, 0, 0, 0),
-		FLinearColor(0, 0, 0, 0)
-	};
-	// is optimized away if possible (RT size=view size, )
-	Context.RHICmdList.ClearMRT(true, NumRenderTargets, ClearColors, false, 1.0f, false, 0, DestRect);
-
-	Context.SetViewportAndCallRHI(0, 0, 0.0f, DestSize.X, DestSize.Y, 1.0f );
-
-	// set the state
-	Context.RHICmdList.SetBlendState(TStaticBlendState<>::GetRHI());
-	Context.RHICmdList.SetRasterizerState(TStaticRasterizerState<>::GetRHI());
-	Context.RHICmdList.SetDepthStencilState(TStaticDepthStencilState<false, CF_Always>::GetRHI());
-
-	TShaderMapRef<FPostProcessVS> VertexShader(ShaderMap);
-
-	if (bNearBlurEnabled)
-	{
-		static FGlobalBoundShaderState BoundShaderState;
-
-
-		TShaderMapRef< FPostProcessCircleDOFPS<1> > PixelShader(ShaderMap);
-		SetGlobalBoundShaderState(Context.RHICmdList, FeatureLevel, BoundShaderState, GFilterVertexDeclaration.VertexDeclarationRHI, *VertexShader, *PixelShader);
-
-		PixelShader->SetParameters(Context);
-	}
-	else
-	{
-		static FGlobalBoundShaderState BoundShaderState;
-
-		TShaderMapRef< FPostProcessCircleDOFPS<0> > PixelShader(ShaderMap);
-		SetGlobalBoundShaderState(Context.RHICmdList, FeatureLevel, BoundShaderState, GFilterVertexDeclaration.VertexDeclarationRHI, *VertexShader, *PixelShader);
-
-		PixelShader->SetParameters(Context);
-	}
-
-	VertexShader->SetParameters(Context);
-
-	// Draw a quad mapping scene color to the view's render target
-	DrawRectangle(
-		Context.RHICmdList,
-		DestRect.Min.X, DestRect.Min.Y,
-		DestRect.Width() + 1, DestRect.Height() + 1,
-		SrcRect.Min.X, SrcRect.Min.Y,
-		SrcRect.Width() + 1, SrcRect.Height() + 1,
-		DestSize,
-		SrcSize,
-		*VertexShader,
-		EDRF_UseTriangleOptimization);
-
-	Context.RHICmdList.CopyToResolveTarget(DestRenderTarget0.TargetableTexture, DestRenderTarget0.ShaderResourceTexture, false, FResolveParams());
-	Context.RHICmdList.CopyToResolveTarget(DestRenderTarget1.TargetableTexture, DestRenderTarget1.ShaderResourceTexture, false, FResolveParams());
-}
-
-FPooledRenderTargetDesc FRCPassPostProcessCircleDOF::ComputeOutputDesc(EPassOutputId InPassOutputId) const
-{
-	FPooledRenderTargetDesc Ret = GetInput(ePId_Input0)->GetOutput()->RenderTargetDesc;
-
-	Ret.Extent.X = FMath::Max(1, Ret.Extent.X);
-	Ret.Extent.Y = FMath::Max(1, Ret.Extent.Y);
-
-	Ret.Reset();
-	Ret.TargetableFlags &= ~(uint32)TexCreate_UAV;
-	Ret.TargetableFlags |= TexCreate_RenderTargetable;
-
-	Ret.DebugName = (InPassOutputId == ePId_Output0) ? TEXT("CircleDOF0") : TEXT("CircleDOF1");
-
-	// more precision for additive blending and we need the alpha channel
-	Ret.Format = PF_FloatRGBA;
-
-	return Ret;
-}
-
-
-/** Encapsulates  the Circle DOF recombine pixel shader. */
-template <uint32 NearBlurEnable>
-class FPostProcessCircleDOFRecombinePS : public FGlobalShader
-{
-	DECLARE_SHADER_TYPE(FPostProcessCircleDOFRecombinePS, Global);
-
-	static bool ShouldCache(EShaderPlatform Platform)
-	{
-		return IsFeatureLevelSupported(Platform, ERHIFeatureLevel::SM5);
-	}
-
-	static void ModifyCompilationEnvironment(EShaderPlatform Platform, FShaderCompilerEnvironment& OutEnvironment)
-	{
-		FGlobalShader::ModifyCompilationEnvironment(Platform,OutEnvironment);
-		OutEnvironment.SetDefine(TEXT("ENABLE_NEAR_BLUR"), NearBlurEnable);
-	}
-
-	/** Default constructor. */
-	FPostProcessCircleDOFRecombinePS() {}
-
-public:
-	FPostProcessPassParameters PostprocessParameter;
-	FDeferredPixelShaderParameters DeferredParameters;
-	FShaderParameter DepthOfFieldUVLimit;
-	FShaderParameter RandomOffset;
-	FShaderParameter CircleDofParams;
-
-	/** Initialization constructor. */
-	FPostProcessCircleDOFRecombinePS(const ShaderMetaType::CompiledShaderInitializerType& Initializer)
-		: FGlobalShader(Initializer)
-	{
-		PostprocessParameter.Bind(Initializer.ParameterMap);
-		DeferredParameters.Bind(Initializer.ParameterMap);
-		DepthOfFieldUVLimit.Bind(Initializer.ParameterMap,TEXT("DepthOfFieldUVLimit"));
-		RandomOffset.Bind(Initializer.ParameterMap, TEXT("RandomOffset"));
-		CircleDofParams.Bind(Initializer.ParameterMap,TEXT("CircleDofParams"));
-	}
-
-	// FShader interface.
-	virtual bool Serialize(FArchive& Ar) override
-	{
-		bool bShaderHasOutdatedParameters = FGlobalShader::Serialize(Ar);
-		Ar << PostprocessParameter << DeferredParameters << DepthOfFieldUVLimit << RandomOffset << CircleDofParams;
-		return bShaderHasOutdatedParameters;
-	}
-
-	void SetParameters(const FRenderingCompositePassContext& Context)
-	{
-		FSceneRenderTargets& SceneContext = FSceneRenderTargets::Get(Context.RHICmdList);
-		const FPixelShaderRHIParamRef ShaderRHI = GetPixelShader();
-
-		FGlobalShader::SetParameters(Context.RHICmdList, ShaderRHI, Context.View);
-
-		DeferredParameters.Set(Context.RHICmdList, ShaderRHI, Context.View);
-		PostprocessParameter.SetPS(ShaderRHI, Context, TStaticSamplerState<SF_Bilinear,AM_Clamp,AM_Clamp,AM_Clamp>::GetRHI());
-
-		// Compute out of bounds UVs in the source texture.
-		FVector4 Bounds;
-		Bounds.X = (((float)((Context.View.ViewRect.Min.X + 1) & (~1))) + 3.0f) / ((float)(SceneContext.GetBufferSizeXY().X));
-		Bounds.Y = (((float)((Context.View.ViewRect.Min.Y + 1) & (~1))) + 3.0f) / ((float)(SceneContext.GetBufferSizeXY().Y));
-		Bounds.Z = (((float)(Context.View.ViewRect.Max.X & (~1))) - 3.0f) / ((float)(SceneContext.GetBufferSizeXY().X));
-		Bounds.W = (((float)(Context.View.ViewRect.Max.Y & (~1))) - 3.0f) / ((float)(SceneContext.GetBufferSizeXY().Y));
-
-		SetShaderValue(Context.RHICmdList, ShaderRHI, DepthOfFieldUVLimit, Bounds);
-
-		FVector2D RandomOffsetValue;
-		TemporalRandom2(&RandomOffsetValue, Context.View.Family->FrameNumber);
-		SetShaderValue(Context.RHICmdList, ShaderRHI, RandomOffset, RandomOffsetValue);
-
-		SetShaderValue(Context.RHICmdList, ShaderRHI, CircleDofParams, CircleDofCoc(Context));
-	}
-};
-
-IMPLEMENT_SHADER_TYPE(template<>,FPostProcessCircleDOFRecombinePS<0>,TEXT("PostProcessDOF"),TEXT("MainCircleRecombinePS"),SF_Pixel);
-IMPLEMENT_SHADER_TYPE(template<>,FPostProcessCircleDOFRecombinePS<1>,TEXT("PostProcessDOF"),TEXT("MainCircleRecombinePS"),SF_Pixel);
-
-void FRCPassPostProcessCircleDOFRecombine::Process(FRenderingCompositePassContext& Context)
-{
-	SCOPED_DRAW_EVENT(Context.RHICmdList, CircleDOFRecombine);
-
-	const FPooledRenderTargetDesc* InputDesc = GetInputDesc(ePId_Input0);
-
-	if(!InputDesc)
-	{
-		// input is not hooked up correctly
-		return;
-	}
-
-	const FSceneView& View = Context.View;
-
-	const auto FeatureLevel = Context.GetFeatureLevel();
-	auto ShaderMap = Context.GetShaderMap();
-
-	FIntPoint TexSize = InputDesc->Extent;
-
-	// usually 1, 2, 4 or 8
-	uint32 ScaleToFullRes = FSceneRenderTargets::Get(Context.RHICmdList).GetBufferSizeXY().X / TexSize.X;
-
-	FIntRect HalfResViewRect = View.ViewRect / ScaleToFullRes;
-
-	const FSceneRenderTargetItem& DestRenderTarget = PassOutputs[0].RequestSurface(Context);
-
-	// Set the view family's render target/viewport.
-	SetRenderTarget(Context.RHICmdList, DestRenderTarget.TargetableTexture, FTextureRHIRef());
-
-	// is optimized away if possible (RT size=view size, )
-	Context.RHICmdList.Clear(true, FLinearColor::Black, false, 1.0f, false, 0, View.ViewRect);
-
-	Context.SetViewportAndCallRHI(View.ViewRect);
-
-	// set the state
-	Context.RHICmdList.SetBlendState(TStaticBlendState<>::GetRHI());
-	Context.RHICmdList.SetRasterizerState(TStaticRasterizerState<>::GetRHI());
-	Context.RHICmdList.SetDepthStencilState(TStaticDepthStencilState<false, CF_Always>::GetRHI());
-
-	TShaderMapRef<FPostProcessVS> VertexShader(ShaderMap);
-
-	if (bNearBlurEnabled)
-	{
-		static FGlobalBoundShaderState BoundShaderState;
-
-		TShaderMapRef< FPostProcessCircleDOFRecombinePS<1> > PixelShader(ShaderMap);
-		SetGlobalBoundShaderState(Context.RHICmdList, FeatureLevel, BoundShaderState, GFilterVertexDeclaration.VertexDeclarationRHI, *VertexShader, *PixelShader);
-		PixelShader->SetParameters(Context);
-	}
-	else
-	{
-		static FGlobalBoundShaderState BoundShaderState;
-
-		TShaderMapRef< FPostProcessCircleDOFRecombinePS<0> > PixelShader(ShaderMap);
-		SetGlobalBoundShaderState(Context.RHICmdList, FeatureLevel, BoundShaderState, GFilterVertexDeclaration.VertexDeclarationRHI, *VertexShader, *PixelShader);
-		PixelShader->SetParameters(Context);
-	}
-
-	VertexShader->SetParameters(Context);
-
-	// Draw a quad mapping scene color to the view's render target
-	DrawRectangle(
-		Context.RHICmdList,
-		0, 0,
-		View.ViewRect.Width(), View.ViewRect.Height(),
-		View.ViewRect.Min.X, View.ViewRect.Min.Y,
-		View.ViewRect.Width(), View.ViewRect.Height(),
-		View.ViewRect.Size(),
-		TexSize,
-		*VertexShader,
-		EDRF_UseTriangleOptimization);
-
-	Context.RHICmdList.CopyToResolveTarget(DestRenderTarget.TargetableTexture, DestRenderTarget.ShaderResourceTexture, false, FResolveParams());
-}
-
-FPooledRenderTargetDesc FRCPassPostProcessCircleDOFRecombine::ComputeOutputDesc(EPassOutputId InPassOutputId) const
-{
-	FPooledRenderTargetDesc Ret = GetInput(ePId_Input0)->GetOutput()->RenderTargetDesc;
-
-	Ret.Reset();
-	Ret.DebugName = TEXT("CircleDOFRecombine");
-
-	return Ret;
-}
+}
\ No newline at end of file
diff --git a/Engine/Source/Runtime/Renderer/Private/PostProcess/PostProcessDOF.h b/Engine/Source/Runtime/Renderer/Private/PostProcess/PostProcessDOF.h
--- a/Engine/Source/Runtime/Renderer/Private/PostProcess/PostProcessDOF.h
+++ b/Engine/Source/Runtime/Renderer/Private/PostProcess/PostProcessDOF.h
@@ -48,83 +48,3 @@ private:
 
 	bool bNearBlurEnabled;
 };
-
-
-// down sample and setup DOF input
-// ePId_Input0: SceneColor
-// ePId_Input1: SceneDepth
-// derives from TRenderingCompositePassBase<InputCount, OutputCount> 
-class FRCPassPostProcessCircleDOFSetup : public TRenderingCompositePassBase<2, 2>
-{
-public:
-
-	FRCPassPostProcessCircleDOFSetup(bool bInNearBlurEnabled) : bNearBlurEnabled ( bInNearBlurEnabled ) {}
-
-	// interface FRenderingCompositePass ---------
-
-	virtual void Process(FRenderingCompositePassContext& Context) override;
-	virtual void Release() override { delete this; }
-	virtual FPooledRenderTargetDesc ComputeOutputDesc(EPassOutputId InPassOutputId) const override;
-
-private:
-
-	bool bNearBlurEnabled;
-};
-
-
-// ePId_Input0: SceneColor
-class FRCPassPostProcessCircleDOFDilate : public TRenderingCompositePassBase<1, 1>
-{
-public:
-
-	FRCPassPostProcessCircleDOFDilate(void) {}
-
-	// interface FRenderingCompositePass ---------
-
-	virtual void Process(FRenderingCompositePassContext& Context) override;
-	virtual void Release() override { delete this; }
-	virtual FPooledRenderTargetDesc ComputeOutputDesc(EPassOutputId InPassOutputId) const override;
-};
-
-
-
-
-// ePId_Input0: setup results for far
-// ePId_Input1: setup results for near
-// derives from TRenderingCompositePassBase<InputCount, OutputCount> 
-class FRCPassPostProcessCircleDOF : public TRenderingCompositePassBase<2, 2>
-{
-public:
-
-	FRCPassPostProcessCircleDOF(bool bInNearBlurEnabled) : bNearBlurEnabled ( bInNearBlurEnabled ) {}
-
-	// interface FRenderingCompositePass ---------
-
-	virtual void Process(FRenderingCompositePassContext& Context) override;
-	virtual void Release() override { delete this; }
-	virtual FPooledRenderTargetDesc ComputeOutputDesc(EPassOutputId InPassOutputId) const override;
-
-private:
-
-	bool bNearBlurEnabled;
-};
-
-// derives from TRenderingCompositePassBase<InputCount, OutputCount> 
-// ePId_Input0: Full res scene color
-// ePId_Input1: output 0 from the DOFSetup (possibly further blurred)
-// ePId_Input2: output 1 from the DOFSetup (possibly further blurred)
-class FRCPassPostProcessCircleDOFRecombine : public TRenderingCompositePassBase<3, 1>
-{
-public:
-
-	FRCPassPostProcessCircleDOFRecombine(bool bInNearBlurEnabled) : bNearBlurEnabled ( bInNearBlurEnabled ) {}
-
-	// interface FRenderingCompositePass ---------
-	virtual void Process(FRenderingCompositePassContext& Context) override;
-	virtual void Release() override { delete this; }
-	virtual FPooledRenderTargetDesc ComputeOutputDesc(EPassOutputId InPassOutputId) const override;
-
-private:
-
-	bool bNearBlurEnabled;
-};
diff --git a/Engine/Source/Runtime/Renderer/Private/PostProcess/PostProcessing.cpp b/Engine/Source/Runtime/Renderer/Private/PostProcess/PostProcessing.cpp
--- a/Engine/Source/Runtime/Renderer/Private/PostProcess/PostProcessing.cpp
+++ b/Engine/Source/Runtime/Renderer/Private/PostProcess/PostProcessing.cpp
@@ -32,6 +32,7 @@
 #include "PostProcessTemporalAA.h"
 #include "PostProcessMotionBlur.h"
 #include "PostProcessDOF.h"
+#include "PostProcessCircleDOF.h"
 #include "PostProcessUpscale.h"
 #include "PostProcessHMD.h"
 #include "PostProcessVisualizeComplexity.h"
