VisualizeDOF is now full res

diff --git a/Engine/Shaders/PostProcessBokehDOF.usf b/Engine/Shaders/PostProcessBokehDOF.usf
--- a/Engine/Shaders/PostProcessBokehDOF.usf
+++ b/Engine/Shaders/PostProcessBokehDOF.usf
@@ -393,6 +393,11 @@ void VisualizeDOFPS(in float4 UVAndScreenPos : TEXCOORD0, out float4 OutColor :
 	CommonDOFSetup(UVAndScreenPos.xy, bFrontLayer, Mask);
 
 	OutColor = (bFrontLayer ? VisualizeColors[0] : VisualizeColors[1]) * dot(Mask, 1);
+	
+	half3 SceneColor = Texture2DSample(PostprocessInput0, PostprocessInput0Sampler, UVAndScreenPos.xy).rgb;
+
+	// blend in a bit of the scene color to make navigation easier
+	OutColor.rgb = lerp(saturate(OutColor.rgb), Luminance(SceneColor), 0.1f);
 }
 
 
diff --git a/Engine/Source/Runtime/Engine/Private/SceneView.cpp b/Engine/Source/Runtime/Engine/Private/SceneView.cpp
--- a/Engine/Source/Runtime/Engine/Private/SceneView.cpp
+++ b/Engine/Source/Runtime/Engine/Private/SceneView.cpp
@@ -1372,7 +1372,8 @@ void FSceneView::EndFinalPostprocessSettings(const FSceneViewInitOptions& ViewIn
 			// Disable antialiasing in GammaLDR mode to avoid jittering.
 			|| (SceneViewFeatureLevel == ERHIFeatureLevel::ES2 && MobileHDRCvar->GetValueOnGameThread() == 0)
 			|| (SceneViewFeatureLevel <= ERHIFeatureLevel::ES3_1 && (MSAAValue > 1))
-			|| Family->EngineShowFlags.VisualizeBloom)
+			|| Family->EngineShowFlags.VisualizeBloom
+			|| Family->EngineShowFlags.VisualizeDOF)
 		{
 			FinalPostProcessSettings.AntiAliasingMethod = AAM_None;
 		}
diff --git a/Engine/Source/Runtime/Renderer/Private/PostProcess/PostProcessBokehDOF.cpp b/Engine/Source/Runtime/Renderer/Private/PostProcess/PostProcessBokehDOF.cpp
--- a/Engine/Source/Runtime/Renderer/Private/PostProcess/PostProcessBokehDOF.cpp
+++ b/Engine/Source/Runtime/Renderer/Private/PostProcess/PostProcessBokehDOF.cpp
@@ -185,7 +185,7 @@ void FRCPassPostProcessVisualizeDOF::Process(FRenderingCompositePassContext& Con
 	uint32 ScaleFactor = FSceneRenderTargets::Get(Context.RHICmdList).GetBufferSizeXY().X / SrcSize.X;
 
 	FIntRect SrcRect = FIntRect::DivideAndRoundUp(View.ViewRect, ScaleFactor);
-	FIntRect DestRect = FIntRect::DivideAndRoundUp(SrcRect, 2);
+	FIntRect DestRect = SrcRect;
 
 	const FSceneRenderTargetItem& DestRenderTarget = PassOutputs[0].RequestSurface(Context);
 
@@ -256,7 +256,7 @@ void FRCPassPostProcessVisualizeDOF::Process(FRenderingCompositePassContext& Con
 		FCanvas Canvas(&TempRenderTarget, NULL, ViewFamily.CurrentRealTime, ViewFamily.CurrentWorldTime, ViewFamily.DeltaWorldTime, Context.GetFeatureLevel());
 
 		float X = 30;
-		float Y = 8;
+		float Y = 18;
 		const float YStep = 14;
 		const float ColumnWidth = 250;
 
@@ -360,9 +360,6 @@ FPooledRenderTargetDesc FRCPassPostProcessVisualizeDOF::ComputeOutputDesc(EPassO
 	FPooledRenderTargetDesc Ret = GetInput(ePId_Input0)->GetOutput()->RenderTargetDesc;
 
 	Ret.Reset();
-	Ret.Extent /= 2;
-	Ret.Extent.X = FMath::Max(1, Ret.Extent.X);
-	Ret.Extent.Y = FMath::Max(1, Ret.Extent.Y);
 	Ret.Format = PF_B8G8R8A8;
 	Ret.DebugName = TEXT("VisualizeDOF");
 
diff --git a/Engine/Source/Runtime/Renderer/Private/PostProcess/PostProcessing.cpp b/Engine/Source/Runtime/Renderer/Private/PostProcess/PostProcessing.cpp
--- a/Engine/Source/Runtime/Renderer/Private/PostProcess/PostProcessing.cpp
+++ b/Engine/Source/Runtime/Renderer/Private/PostProcess/PostProcessing.cpp
@@ -1463,12 +1463,7 @@ void FPostProcessing::Process(FRHICommandListImmediate& RHICmdList, FViewInfo& V
 				FRenderingCompositePass* VisualizeNode = Context.Graph.RegisterPass(new(FMemStack::Get()) FRCPassPostProcessVisualizeDOF(DepthOfFieldStat));
 				VisualizeNode->SetInput(ePId_Input0, FRenderingCompositeOutputRef(Context.FinalOutput));
 
-				// PassThrough is needed to upscale the half res texture
-				FPooledRenderTargetDesc Desc = FSceneRenderTargets::Get(RHICmdList).GetSceneColor()->GetDesc();
-				Desc.Format = PF_B8G8R8A8;
-				FRenderingCompositePass* NullPass = Context.Graph.RegisterPass(new(FMemStack::Get()) FRCPassPostProcessPassThrough(Desc));
-				NullPass->SetInput(ePId_Input0, FRenderingCompositeOutputRef(VisualizeNode));
-				Context.FinalOutput = FRenderingCompositeOutputRef(NullPass);
+				Context.FinalOutput = FRenderingCompositeOutputRef(VisualizeNode);
 			}
 		}
 		else
