fixed DOF material function for CircleDOF

diff --git a/Engine/Shaders/PostProcessCircleDOF.usf b/Engine/Shaders/PostProcessCircleDOF.usf
--- a/Engine/Shaders/PostProcessCircleDOF.usf
+++ b/Engine/Shaders/PostProcessCircleDOF.usf
@@ -14,12 +14,14 @@
 float4 CircleDofParams;
 
 // CIRCLE DOF: Compute circle of confusion size in pixels.
-float DepthToCoc(float Depth) 
+float DepthToCoc(float SceneDepth) 
 {
+	// We intentionally don't do the DepthOfFieldFocalRegion as it breaks realism.
+
 	float Focus = View.DepthOfFieldFocalDistance;
 	float Radius = CircleDofParams.x;
-	float CocRadius = ((Depth - Focus) / Depth) * Radius;
-	float DepthBlurRadius = (1.0 - exp2(-Depth * CircleDofParams.y)) * CircleDofParams.z;
+	float CocRadius = ((SceneDepth - Focus) / SceneDepth) * Radius;
+	float DepthBlurRadius = (1.0 - exp2(-SceneDepth * CircleDofParams.y)) * CircleDofParams.z;
 	float ReturnCoc = max(abs(CocRadius), DepthBlurRadius);
 	if(CocRadius < 0.0) 
 	{
diff --git a/Engine/Source/Runtime/Engine/Classes/Engine/Scene.h b/Engine/Source/Runtime/Engine/Classes/Engine/Scene.h
--- a/Engine/Source/Runtime/Engine/Classes/Engine/Scene.h
+++ b/Engine/Source/Runtime/Engine/Classes/Engine/Scene.h
@@ -837,8 +837,8 @@ struct FPostProcessSettings
 	UPROPERTY(interp, BlueprintReadWrite, Category=DepthOfField, meta=(ClampMin = "0.0", ClampMax = "4.0", editcondition = "bOverride_DepthOfFieldDepthBlurRadius", DisplayName = "Depth Blur Radius"))
 	float DepthOfFieldDepthBlurRadius;
 	
-	/** CircleDOF only: F-stop number (1/fstop) e.g. 1.2 for a large aperture, 2.8 is smaller and results in less Depth of Field effect, does not affect noise/grain but in real world it would */
-	UPROPERTY(interp, BlueprintReadWrite, Category=DepthOfField, meta=(ClampMin = "1.0", ClampMax = "32.0", editcondition = "bOverride_DepthOfFieldFstop", DisplayName = "Aperture"))
+	/** CircleDOF only: Defines the opening of the camera lens, aperture is 1/fstop, typical lens go down to 1.2 (large opening), larger numbers reduce the DOF effect */
+	UPROPERTY(interp, BlueprintReadWrite, Category=DepthOfField, meta=(ClampMin = "1.0", ClampMax = "32.0", editcondition = "bOverride_DepthOfFieldFstop", DisplayName = "Aperture F-stop"))
 	float DepthOfFieldFstop;
 
 	/** Distance in which the Depth of Field effect should be sharp, in unreal units (cm) */
diff --git a/Engine/Source/Runtime/Engine/Private/SceneView.cpp b/Engine/Source/Runtime/Engine/Private/SceneView.cpp
--- a/Engine/Source/Runtime/Engine/Private/SceneView.cpp
+++ b/Engine/Source/Runtime/Engine/Private/SceneView.cpp
@@ -1256,6 +1256,13 @@ void FSceneView::EndFinalPostprocessSettings(const FSceneViewInitOptions& ViewIn
 	}
 #endif
 
+	if(FinalPostProcessSettings.DepthOfFieldMethod == DOFM_CircleDOF)
+	{
+		// We intentionally don't do the DepthOfFieldFocalRegion as it breaks realism.
+		// Doing this fixes DOF material expression.
+		FinalPostProcessSettings.DepthOfFieldFocalRegion = 0;
+	}
+
 	{
 		static const auto CVar = IConsoleManager::Get().FindTConsoleVariableDataFloat(TEXT("r.ScreenPercentage"));
 
diff --git a/Engine/Source/Runtime/Renderer/Private/PostProcess/PostProcessBokehDOF.cpp b/Engine/Source/Runtime/Renderer/Private/PostProcess/PostProcessBokehDOF.cpp
--- a/Engine/Source/Runtime/Renderer/Private/PostProcess/PostProcessBokehDOF.cpp
+++ b/Engine/Source/Runtime/Renderer/Private/PostProcess/PostProcessBokehDOF.cpp
@@ -275,7 +275,7 @@ void FRCPassPostProcessVisualizeDOF::Process(FRenderingCompositePassContext& Con
 			Y += YStep;
 			Line = FString::Printf(TEXT("FocalDistance: %.2f"), View.FinalPostProcessSettings.DepthOfFieldFocalDistance);
 			Canvas.DrawShadowedString(X, Y += YStep, *Line, GetStatsFont(), FLinearColor(1, 1, 1));
-			Line = FString::Printf(TEXT("FocalRegion: %.2f"), View.FinalPostProcessSettings.DepthOfFieldFocalRegion);
+			Line = FString::Printf(TEXT("FocalRegion (Artificial, avoid): %.2f"), View.FinalPostProcessSettings.DepthOfFieldFocalRegion);
 			Canvas.DrawShadowedString(X, Y += YStep, *Line, GetStatsFont(), FLinearColor(1, 1, 1));
 			Y += YStep;
 			Line = FString::Printf(TEXT("Scale: %.2f"), View.FinalPostProcessSettings.DepthOfFieldScale);
@@ -300,7 +300,7 @@ void FRCPassPostProcessVisualizeDOF::Process(FRenderingCompositePassContext& Con
 			Y += YStep;
 			Line = FString::Printf(TEXT("FocalDistance: %.2f"), View.FinalPostProcessSettings.DepthOfFieldFocalDistance);
 			Canvas.DrawShadowedString(X, Y += YStep, *Line, GetStatsFont(), FLinearColor(1, 1, 1));
-			Line = FString::Printf(TEXT("FocalRegion: %.2f"), View.FinalPostProcessSettings.DepthOfFieldFocalRegion);
+			Line = FString::Printf(TEXT("FocalRegion (Artificial, avoid): %.2f"), View.FinalPostProcessSettings.DepthOfFieldFocalRegion);
 			Canvas.DrawShadowedString(X, Y += YStep, *Line, GetStatsFont(), FLinearColor(1, 1, 1));
 			Y += YStep;
 			Line = FString::Printf(TEXT("NearTransitionRegion: %.2f"), View.FinalPostProcessSettings.DepthOfFieldNearTransitionRegion);
@@ -326,7 +326,9 @@ void FRCPassPostProcessVisualizeDOF::Process(FRenderingCompositePassContext& Con
 			Y += YStep;
 			Line = FString::Printf(TEXT("FocalDistance: %.2f"), View.FinalPostProcessSettings.DepthOfFieldFocalDistance);
 			Canvas.DrawShadowedString(X, Y += YStep, *Line, GetStatsFont(), FLinearColor(1, 1, 1));
-			Line = FString::Printf(TEXT("Aperture: %.2f"), View.FinalPostProcessSettings.DepthOfFieldFstop);
+			Line = FString::Printf(TEXT("Aperture F-stop: %.2f"), View.FinalPostProcessSettings.DepthOfFieldFstop);
+			Canvas.DrawShadowedString(X, Y += YStep, *Line, GetStatsFont(), FLinearColor(1, 1, 1));
+			Line = FString::Printf(TEXT("Aperture: 1 / %.2f"), View.FinalPostProcessSettings.DepthOfFieldFstop);
 			Canvas.DrawShadowedString(X, Y += YStep, *Line, GetStatsFont(), FLinearColor(1, 1, 1));
 			Y += YStep;
 			Line = FString::Printf(TEXT("DepthBlur (not related to Depth of Field, due to light traveling long distances in atmosphere)"));
diff --git a/Engine/Source/Runtime/Renderer/Private/PostProcess/PostProcessCircleDOF.cpp b/Engine/Source/Runtime/Renderer/Private/PostProcess/PostProcessCircleDOF.cpp
--- a/Engine/Source/Runtime/Renderer/Private/PostProcess/PostProcessCircleDOF.cpp
+++ b/Engine/Source/Runtime/Renderer/Private/PostProcess/PostProcessCircleDOF.cpp
@@ -33,12 +33,10 @@ float ComputeFocalLengthFromFov(const FSceneView& View)
 // Setup depth based blur.
 static FVector4 CircleDofCoc(const FRenderingCompositePassContext& Context)
 {
-	float FocalLength = ComputeFocalLengthFromFov(Context.View);
+	float FocalLengthInMM = ComputeFocalLengthFromFov(Context.View);
 	 
-	// Convert focal distance in world position to mm.
-	// Conversion is 1 world position = 1 cm.
-	float Distance = Context.View.FinalPostProcessSettings.DepthOfFieldFocalDistance;
-	Distance *= 10.0f;
+	// Convert focal distance in world position to mm (from cm to mm)
+	float FocalDistanceInMM = Context.View.FinalPostProcessSettings.DepthOfFieldFocalDistance * 10.0f;
 
 	// Convert f-stop, focal length, and focal distance to
 	// projected circle of confusion size at infinity in mm.
@@ -48,7 +46,7 @@ static FVector4 CircleDofCoc(const FRenderingCompositePassContext& Context)
 	//   f = focal length
 	//   d = focal distance
 	//   n = fstop (where n is the "n" in "f/n")
-	float Radius = FocalLength * FocalLength / (Context.View.FinalPostProcessSettings.DepthOfFieldFstop * (Distance - FocalLength));
+	float Radius = FMath::Square(FocalLengthInMM) / (Context.View.FinalPostProcessSettings.DepthOfFieldFstop * (FocalDistanceInMM - FocalLengthInMM));
 
 	// Scale so that APS-C 24.576 mm = full frame.
 	// Convert mm to pixels.
