Bind clear color to rendertargets at creation time.

diff --git a/Engine/Plugins/Runtime/OculusRift/Source/OculusRift/Private/OculusRiftRenderD3D.cpp b/Engine/Plugins/Runtime/OculusRift/Source/OculusRift/Private/OculusRiftRenderD3D.cpp
--- a/Engine/Plugins/Runtime/OculusRift/Source/OculusRift/Private/OculusRiftRenderD3D.cpp
+++ b/Engine/Plugins/Runtime/OculusRift/Source/OculusRift/Private/OculusRiftRenderD3D.cpp
@@ -59,7 +59,8 @@ public:
 		InFormat,
 		bInCubemap,
 		InFlags,
-		bInPooled
+		bInPooled,
+		FClearValueBinding::None
 		)
 	{
 		TextureSet = nullptr;
@@ -315,7 +316,8 @@ static FD3D11Texture2D* D3D11CreateTexture2DAlias(
 		InFormat,
 		/*bInCubemap=*/ false,
 		InFlags,
-		/*bPooledTexture=*/ false
+		/*bPooledTexture=*/ false,
+		FClearValueBinding::None
 		);
 
 	return NewTexture;
diff --git a/Engine/Plugins/Runtime/OculusRift/Source/OculusRift/Private/OculusRiftRenderGL.cpp b/Engine/Plugins/Runtime/OculusRift/Source/OculusRift/Private/OculusRiftRenderGL.cpp
--- a/Engine/Plugins/Runtime/OculusRift/Source/OculusRift/Private/OculusRiftRenderGL.cpp
+++ b/Engine/Plugins/Runtime/OculusRift/Source/OculusRift/Private/OculusRiftRenderGL.cpp
@@ -55,7 +55,8 @@ public:
 		bInCubemap,
 		bInAllocatedStorage,
 		InFlags,
-		InTextureRange
+		InTextureRange,
+		FClearValueBinding::None
 		)
 	{
 		TextureSet = nullptr;
@@ -215,7 +216,7 @@ static FOpenGLTexture2D* OpenGLCreateTexture2DAlias(
 		InGLRHI,
 		InResource,
 		Target,
-		Attachment, InSizeX, InSizeY, 0, InNumMips, InNumSamples, 1, InFormat, false, bAllocatedStorage, InFlags, TextureRange);
+		Attachment, InSizeX, InSizeY, 0, InNumMips, InNumSamples, 1, InFormat, false, bAllocatedStorage, InFlags, TextureRange, FClearValueBinding::None);
 
 	OpenGLTextureAllocated(NewTexture, InFlags);
 	return NewTexture;
diff --git a/Engine/Shaders/OneColorShader.usf b/Engine/Shaders/OneColorShader.usf
--- a/Engine/Shaders/OneColorShader.usf
+++ b/Engine/Shaders/OneColorShader.usf
@@ -38,25 +38,25 @@ void MainVertexShader(
  * Pixel shader
  =============================================================================*/
 
-float4 DrawColorMRT[8];
-
 #ifndef NUM_OUTPUTS
 #define NUM_OUTPUTS 1
 #endif
 
+//These shaders used for RHIClear.  Don't use any global uniform variables here as it greatly complicates
+//the state restoration for an inline shader clear.
 void MainPixelShaderMRT( 
 	out float4 OutColor[NUM_OUTPUTS] : SV_Target0)
 {
 	for (int i = 0; i < NUM_OUTPUTS; i++)
 	{
-		OutColor[i] = DrawColorMRT[i];
+		OutColor[i] = ClearShaderUB.DrawColorMRT[i];
 	}
 }
 
 void MainPixelShader( 
 	out float4 OutColor : SV_Target0)
 {
-	OutColor = DrawColorMRT[0];
+	OutColor = ClearShaderUB.DrawColorMRT[0];
 }
 
 /** 
diff --git a/Engine/Source/Runtime/EmptyRHI/Private/EmptyCommands.cpp b/Engine/Source/Runtime/EmptyRHI/Private/EmptyCommands.cpp
--- a/Engine/Source/Runtime/EmptyRHI/Private/EmptyCommands.cpp
+++ b/Engine/Source/Runtime/EmptyRHI/Private/EmptyCommands.cpp
@@ -396,7 +396,7 @@ void FEmptyDynamicRHI::RHIClearMRT(bool bClearColor,int32 NumClearColors,const F
 
 }
 
-void FEmptyDynamicRHI::RHIBindClearMRTValues(bool bClearColor, int32 NumClearColors, const FLinearColor* ColorArray, bool bClearDepth, float Depth, bool bClearStencil, uint32 Stencil)
+void FEmptyDynamicRHI::RHIBindClearMRTValues(bool bClearColor, bool bClearDepth, bool bClearStencil)
 {
 }
 
diff --git a/Engine/Source/Runtime/EmptyRHI/Public/EmptyRHI.h b/Engine/Source/Runtime/EmptyRHI/Public/EmptyRHI.h
--- a/Engine/Source/Runtime/EmptyRHI/Public/EmptyRHI.h
+++ b/Engine/Source/Runtime/EmptyRHI/Public/EmptyRHI.h
@@ -193,8 +193,7 @@ public:
 	virtual void RHISetDepthStencilState(FDepthStencilStateRHIParamRef NewState, uint32 StencilRef) final override;
 	virtual void RHISetBlendState(FBlendStateRHIParamRef NewState, const FLinearColor& BlendFactor) final override;
 	virtual void RHISetRenderTargets(uint32 NumSimultaneousRenderTargets, const FRHIRenderTargetView* NewRenderTargets, const FRHIDepthRenderTargetView* NewDepthStencilTarget, uint32 NumUAVs, const FUnorderedAccessViewRHIParamRef* UAVs) final override;
-	virtual void RHISetRenderTargetsAndClear(const FRHISetRenderTargetsInfo& RenderTargetsInfo) final override;
-	virtual void RHIBindClearMRTValues(bool bClearColor, int32 NumClearColors, const FLinearColor* ColorArray, bool bClearDepth, float Depth, bool bClearStencil, uint32 Stencil) final override;
+	virtual void RHISetRenderTargetsAndClear(const FRHISetRenderTargetsInfo& RenderTargetsInfo) final override;	
 	virtual void RHIDrawPrimitive(uint32 PrimitiveType, uint32 BaseVertexIndex, uint32 NumPrimitives, uint32 NumInstances) final override;
 	virtual void RHIDrawPrimitiveIndirect(uint32 PrimitiveType, FVertexBufferRHIParamRef ArgumentBuffer, uint32 ArgumentOffset) final override;
 	virtual void RHIDrawIndexedIndirect(FIndexBufferRHIParamRef IndexBufferRHI, uint32 PrimitiveType, FStructuredBufferRHIParamRef ArgumentsBufferRHI, int32 DrawArgumentsIndex, uint32 NumInstances) final override;
diff --git a/Engine/Source/Runtime/Engine/Private/Particles/ParticleCurveTexture.cpp b/Engine/Source/Runtime/Engine/Private/Particles/ParticleCurveTexture.cpp
--- a/Engine/Source/Runtime/Engine/Private/Particles/ParticleCurveTexture.cpp
+++ b/Engine/Source/Runtime/Engine/Private/Particles/ParticleCurveTexture.cpp
@@ -501,7 +501,7 @@ FParticleCurveTexture::FParticleCurveTexture()
 void FParticleCurveTexture::InitRHI()
 {
 	// 8-bit per channel RGBA texture for curves.
-	FRHIResourceCreateInfo CreateInfo;
+	FRHIResourceCreateInfo CreateInfo(FClearValueBinding::None);
 	RHICreateTargetableShaderResource2D(
 		GParticleCurveTextureSizeX,
 		GParticleCurveTextureSizeY,
diff --git a/Engine/Source/Runtime/Engine/Private/Particles/ParticleGpuSimulation.cpp b/Engine/Source/Runtime/Engine/Private/Particles/ParticleGpuSimulation.cpp
--- a/Engine/Source/Runtime/Engine/Private/Particles/ParticleGpuSimulation.cpp
+++ b/Engine/Source/Runtime/Engine/Private/Particles/ParticleGpuSimulation.cpp
@@ -187,7 +187,7 @@ public:
 		check( !IsValidRef( PositionTextureTargetRHI ) );
 		check( !IsValidRef( PositionTextureRHI ) );
 
-		FRHIResourceCreateInfo CreateInfo;
+		FRHIResourceCreateInfo CreateInfo(FClearValueBinding::Transparent);
 		RHICreateTargetableShaderResource2D(
 			SizeX,
 			SizeY,
@@ -253,7 +253,7 @@ public:
 		const int32 SizeX = GParticleSimulationTextureSizeX;
 		const int32 SizeY = GParticleSimulationTextureSizeY;
 
-		FRHIResourceCreateInfo CreateInfo;
+		FRHIResourceCreateInfo CreateInfo(FClearValueBinding::None);
 		RHICreateTargetableShaderResource2D(
 			SizeX,
 			SizeY,
@@ -4109,16 +4109,16 @@ void FFXSystem::SimulateGPUParticles(
 	// On some platforms, the textures are filled with garbage after creation, so we need to clear them to black the first time we use them
 	if ( !CurrentStateTextures.bTexturesCleared )
 	{
-		SetRenderTarget(RHICmdList, CurrentStateTextures.PositionTextureTargetRHI, FTextureRHIRef(), ESimpleRenderTargetMode::EClearToDefault);
-		SetRenderTarget(RHICmdList, CurrentStateTextures.VelocityTextureTargetRHI, FTextureRHIRef(), ESimpleRenderTargetMode::EClearToDefault);
+		SetRenderTarget(RHICmdList, CurrentStateTextures.PositionTextureTargetRHI, FTextureRHIRef(), ESimpleRenderTargetMode::EClearColorAndDepth);
+		SetRenderTarget(RHICmdList, CurrentStateTextures.VelocityTextureTargetRHI, FTextureRHIRef(), ESimpleRenderTargetMode::EClearColorAndDepth);
 
 		CurrentStateTextures.bTexturesCleared = true;
 	}
 
 	if ( !PrevStateTextures.bTexturesCleared )
 	{
-		SetRenderTarget(RHICmdList, PrevStateTextures.PositionTextureTargetRHI, FTextureRHIRef(), ESimpleRenderTargetMode::EClearToDefault);
-		SetRenderTarget(RHICmdList, PrevStateTextures.VelocityTextureTargetRHI, FTextureRHIRef(), ESimpleRenderTargetMode::EClearToDefault);
+		SetRenderTarget(RHICmdList, PrevStateTextures.PositionTextureTargetRHI, FTextureRHIRef(), ESimpleRenderTargetMode::EClearColorAndDepth);
+		SetRenderTarget(RHICmdList, PrevStateTextures.VelocityTextureTargetRHI, FTextureRHIRef(), ESimpleRenderTargetMode::EClearColorAndDepth);
 		RHICmdList.CopyToResolveTarget(PrevStateTextures.PositionTextureTargetRHI, PrevStateTextures.PositionTextureTargetRHI, true, FResolveParams());
 		RHICmdList.CopyToResolveTarget(PrevStateTextures.VelocityTextureTargetRHI, PrevStateTextures.VelocityTextureTargetRHI, true, FResolveParams());
 		
diff --git a/Engine/Source/Runtime/Engine/Private/Rendering/SubsurfaceProfile.cpp b/Engine/Source/Runtime/Engine/Private/Rendering/SubsurfaceProfile.cpp
--- a/Engine/Source/Runtime/Engine/Private/Rendering/SubsurfaceProfile.cpp
+++ b/Engine/Source/Runtime/Engine/Private/Rendering/SubsurfaceProfile.cpp
@@ -140,7 +140,7 @@ void FSubsurfaceProfileTexture::CreateTexture(FRHICommandListImmediate& RHICmdLi
 	const uint32 Width = 32;
 
 	// at minimum 64 lines (less reallocations)
-	FPooledRenderTargetDesc Desc(FPooledRenderTargetDesc::Create2DDesc(FIntPoint(Width, FMath::Max(Height, (uint32)64)), PF_B8G8R8A8, TexCreate_FastVRAM, TexCreate_None, false));
+	FPooledRenderTargetDesc Desc(FPooledRenderTargetDesc::Create2DDesc(FIntPoint(Width, FMath::Max(Height, (uint32)64)), PF_B8G8R8A8, FClearValueBinding::None, TexCreate_FastVRAM, TexCreate_None, false));
 
 	if (b16Bit)
 	{
diff --git a/Engine/Source/Runtime/Engine/Private/TextureRenderTarget2D.cpp b/Engine/Source/Runtime/Engine/Private/TextureRenderTarget2D.cpp
--- a/Engine/Source/Runtime/Engine/Private/TextureRenderTarget2D.cpp
+++ b/Engine/Source/Runtime/Engine/Private/TextureRenderTarget2D.cpp
@@ -315,11 +315,12 @@ void FTextureRenderTarget2DResource::InitDynamicRHI()
 
 		// Create the RHI texture. Only one mip is used and the texture is targetable for resolve.
 		uint32 TexCreateFlags = bSRGB ? TexCreate_SRGB : 0;
-		FRHIResourceCreateInfo CreateInfo;
+		FRHIResourceCreateInfo CreateInfo = FRHIResourceCreateInfo(FClearValueBinding(ClearColor));
+
 		RHICreateTargetableShaderResource2D(
 			Owner->SizeX, 
 			Owner->SizeY, 
-			Format, 
+			Format,
 			1,
 			TexCreateFlags,
 			TexCreate_RenderTargetable,
@@ -375,9 +376,9 @@ void FTextureRenderTarget2DResource::UpdateDeferredResource( FRHICommandListImme
  	// clear the target surface to green
 	if (bClearRenderTarget)
 	{
-		SetRenderTarget(RHICmdList, RenderTargetTextureRHI, FTextureRHIRef());
 		RHICmdList.SetViewport(0, 0, 0.0f, TargetSizeX, TargetSizeY, 1.0f);
-		RHICmdList.Clear(true, ClearColor, false, 0.f, false, 0, FIntRect());
+		ensure(RenderTargetTextureRHI->GetClearColor() == ClearColor);
+		SetRenderTarget(RHICmdList, RenderTargetTextureRHI, FTextureRHIRef(), ESimpleRenderTargetMode::EClearColorExistingDepth);
 	}
  
  	// copy surface to the texture for use
diff --git a/Engine/Source/Runtime/GameLiveStreaming/GameLiveStreaming.cpp b/Engine/Source/Runtime/GameLiveStreaming/GameLiveStreaming.cpp
--- a/Engine/Source/Runtime/GameLiveStreaming/GameLiveStreaming.cpp
+++ b/Engine/Source/Runtime/GameLiveStreaming/GameLiveStreaming.cpp
@@ -313,7 +313,7 @@ void FGameLiveStreaming::StartCopyingNextGameVideoFrame( const FViewportRHIRef&
 		ReadSurfaceCommand,
 		FCopyVideoFrame,Context,CopyVideoFrame,
 	{
-		FPooledRenderTargetDesc OutputDesc(FPooledRenderTargetDesc::Create2DDesc(Context.ResizeTo, PF_B8G8R8A8, TexCreate_None, TexCreate_RenderTargetable, false));
+		FPooledRenderTargetDesc OutputDesc(FPooledRenderTargetDesc::Create2DDesc(Context.ResizeTo, PF_B8G8R8A8, FClearValueBinding::None, TexCreate_None, TexCreate_RenderTargetable, false));
 			
 		const auto FeatureLevel = GMaxRHIFeatureLevel;
 
diff --git a/Engine/Source/Runtime/IOS/MetalRHI/Private/MetalCommands.cpp b/Engine/Source/Runtime/IOS/MetalRHI/Private/MetalCommands.cpp
--- a/Engine/Source/Runtime/IOS/MetalRHI/Private/MetalCommands.cpp
+++ b/Engine/Source/Runtime/IOS/MetalRHI/Private/MetalCommands.cpp
@@ -784,16 +784,8 @@ void FMetalDynamicRHI::RHIClearMRT(bool bClearColor,int32 NumClearColors,const F
 
 		{
 			FRHICommandList_RecursiveHazardous RHICmdList(this);
-			SetGlobalBoundShaderState(RHICmdList, GMaxRHIFeatureLevel, GClearMRTBoundShaderState[0], GVector4VertexDeclaration.VertexDeclarationRHI, *VertexShader, PixelShader);
-			FLinearColor ShaderClearColors[MaxSimultaneousRenderTargets];
-			FMemory::Memzero(ShaderClearColors);
-
-			for (int32 i = 0; i < NumClearColors; i++)
-			{
-				ShaderClearColors[i] = ClearColorArray[i];
-			}
-
-			SetShaderValueArray(RHICmdList, PixelShader->GetPixelShader(), PixelShader->ColorParameter, ShaderClearColors, NumClearColors);
+			SetGlobalBoundShaderState(RHICmdList, GMaxRHIFeatureLevel, GClearMRTBoundShaderState[0], GVector4VertexDeclaration.VertexDeclarationRHI, *VertexShader, PixelShader);			
+			PixelShader->SetColors(RHICmdList, ClearColorArray, NumClearColors);
 
 			{
 				// Draw a fullscreen quad
@@ -849,11 +841,6 @@ void FMetalDynamicRHI::RHIClearMRT(bool bClearColor,int32 NumClearColors,const F
 	}
 }											
 
-void FMetalDynamicRHI::RHIBindClearMRTValues(bool bClearColor, int32 NumClearColors, const FLinearColor* ClearColorArray, bool bClearDepth, float Depth, bool bClearStencil, uint32 Stencil)
-{
-	// Not necessary
-}
-
 void FMetalDynamicRHI::RHISuspendRendering()
 {
 	// Not supported
diff --git a/Engine/Source/Runtime/IOS/MetalRHI/Private/MetalManager.cpp b/Engine/Source/Runtime/IOS/MetalRHI/Private/MetalManager.cpp
--- a/Engine/Source/Runtime/IOS/MetalRHI/Private/MetalManager.cpp
+++ b/Engine/Source/Runtime/IOS/MetalRHI/Private/MetalManager.cpp
@@ -687,11 +687,15 @@ void FMetalManager::SetRenderTargetsInfo(const FRHISetRenderTargetsInfo& RenderT
 				ColorAttachment.storeAction = GetMetalRTStoreAction(RenderTargetView.StoreAction);
 				Pipeline.SampleCount = 1;
 			}
+			ERenderTargetLoadAction ColorLoadAction = RenderTargetView.LoadAction;
 			ColorAttachment.level = RenderTargetView.MipIndex;
 			ColorAttachment.slice = ArraySliceIndex;
-			ColorAttachment.loadAction = GetMetalRTLoadAction(RenderTargetView.LoadAction);
-			const FLinearColor& ClearColor = RenderTargetsInfo.ClearColors[RenderTargetIndex];
-			ColorAttachment.clearColor = MTLClearColorMake(ClearColor.R, ClearColor.G, ClearColor.B, ClearColor.A);
+			ColorAttachment.loadAction = GetMetalRTLoadAction(ColorLoadAction);
+			if (ColorLoadAction == ERenderTargetLoadAction::EClear)
+			{
+				const FLinearColor& ClearColor = RenderTargetsInfo.ColorRenderTarget[RenderTargetIndex].Texture->GetClearColor();
+				ColorAttachment.clearColor = MTLClearColorMake(ClearColor.R, ClearColor.G, ClearColor.B, ClearColor.A);
+			}
 
 			// assign the attachment to the slot
 			[RenderPass.colorAttachments setObject:ColorAttachment atIndexedSubscript:RenderTargetIndex];
@@ -724,11 +728,16 @@ void FMetalManager::SetRenderTargetsInfo(const FRHISetRenderTargetsInfo& RenderT
 		{
 			MTLRenderPassDepthAttachmentDescriptor* DepthAttachment = [[MTLRenderPassDepthAttachmentDescriptor alloc] init];
 
+			
 			// set up the depth attachment
+			ERenderTargetLoadAction DepthLoadAction = RenderTargetsInfo.DepthStencilRenderTarget.DepthLoadAction;
 			DepthAttachment.texture = Surface.Texture;
-			DepthAttachment.loadAction = GetMetalRTLoadAction(RenderTargetsInfo.DepthStencilRenderTarget.DepthLoadAction);
+			DepthAttachment.loadAction = GetMetalRTLoadAction(DepthLoadAction);
 			DepthAttachment.storeAction = GetMetalRTStoreAction(RenderTargetsInfo.DepthStencilRenderTarget.DepthStoreAction);
-			DepthAttachment.clearDepth = RenderTargetsInfo.DepthClearValue;
+			if (DepthLoadAction == ERenderTargetLoadAction::EClear)
+			{
+				DepthAttachment.clearDepth = RenderTargetsInfo.DepthStencilRenderTarget.Texture->GetDepthClearValue();
+			}
 
 			Pipeline.DepthTargetFormat = DepthAttachment.texture.pixelFormat;
 			if (Pipeline.SampleCount == 0)
@@ -746,10 +755,14 @@ void FMetalManager::SetRenderTargetsInfo(const FRHISetRenderTargetsInfo& RenderT
 			MTLRenderPassStencilAttachmentDescriptor* StencilAttachment = [[MTLRenderPassStencilAttachmentDescriptor alloc] init];
 
 			// set up the stencil attachment
+			ERenderTargetLoadAction StencilLoadAction = RenderTargetsInfo.DepthStencilRenderTarget.StencilLoadAction;
 			StencilAttachment.texture = Surface.StencilTexture;
 			StencilAttachment.loadAction = GetMetalRTLoadAction(RenderTargetsInfo.DepthStencilRenderTarget.StencilLoadAction);
 			StencilAttachment.storeAction = GetMetalRTStoreAction(RenderTargetsInfo.DepthStencilRenderTarget.GetStencilStoreAction());
-			StencilAttachment.clearStencil = RenderTargetsInfo.StencilClearValue;
+			if (StencilLoadAction == ERenderTargetLoadAction::EClear)
+			{
+				StencilAttachment.clearStencil = RenderTargetsInfo.DepthStencilRenderTarget.Texture->GetStencilClearValue();
+			}
 
 			Pipeline.StencilTargetFormat = StencilAttachment.texture.pixelFormat;
 			if (Pipeline.SampleCount == 0)
diff --git a/Engine/Source/Runtime/IOS/MetalRHI/Private/MetalTexture.cpp b/Engine/Source/Runtime/IOS/MetalRHI/Private/MetalTexture.cpp
--- a/Engine/Source/Runtime/IOS/MetalRHI/Private/MetalTexture.cpp
+++ b/Engine/Source/Runtime/IOS/MetalRHI/Private/MetalTexture.cpp
@@ -423,7 +423,7 @@ uint32 FMetalDynamicRHI::RHIComputeMemorySize(FTextureRHIParamRef TextureRHI)
 
 FTexture2DRHIRef FMetalDynamicRHI::RHICreateTexture2D(uint32 SizeX, uint32 SizeY, uint8 Format, uint32 NumMips, uint32 NumSamples, uint32 Flags, FRHIResourceCreateInfo& CreateInfo)
 {
-	return new FMetalTexture2D((EPixelFormat)Format, SizeX, SizeY, NumMips, NumSamples, Flags, CreateInfo.BulkData);
+	return new FMetalTexture2D((EPixelFormat)Format, SizeX, SizeY, NumMips, NumSamples, Flags, CreateInfo.BulkData, CreateInfo.ClearValueBinding);
 }
 
 FTexture2DRHIRef FMetalDynamicRHI::RHIAsyncCreateTexture2D(uint32 SizeX,uint32 SizeY,uint8 Format,uint32 NumMips,uint32 Flags,void** InitialMipData,uint32 NumInitialMips)
@@ -439,12 +439,12 @@ void FMetalDynamicRHI::RHICopySharedMips(FTexture2DRHIParamRef DestTexture2D,FTe
 
 FTexture2DArrayRHIRef FMetalDynamicRHI::RHICreateTexture2DArray(uint32 SizeX, uint32 SizeY, uint32 SizeZ, uint8 Format, uint32 NumMips, uint32 Flags, FRHIResourceCreateInfo& CreateInfo)
 {
-	return new FMetalTexture2DArray((EPixelFormat)Format, SizeX, SizeY, SizeZ, NumMips, Flags, CreateInfo.BulkData);
+	return new FMetalTexture2DArray((EPixelFormat)Format, SizeX, SizeY, SizeZ, NumMips, Flags, CreateInfo.BulkData, CreateInfo.ClearValueBinding);
 }
 
 FTexture3DRHIRef FMetalDynamicRHI::RHICreateTexture3D(uint32 SizeX, uint32 SizeY, uint32 SizeZ, uint8 Format, uint32 NumMips, uint32 Flags, FRHIResourceCreateInfo& CreateInfo)
 {
-	return new FMetalTexture3D((EPixelFormat)Format, SizeX, SizeY, SizeZ, NumMips, Flags, CreateInfo.BulkData);
+	return new FMetalTexture3D((EPixelFormat)Format, SizeX, SizeY, SizeZ, NumMips, Flags, CreateInfo.BulkData, CreateInfo.ClearValueBinding);
 }
 
 void FMetalDynamicRHI::RHIGetResourceInfo(FTextureRHIParamRef Ref, FRHIResourceInfo& OutInfo)
@@ -460,7 +460,7 @@ FTexture2DRHIRef FMetalDynamicRHI::RHIAsyncReallocateTexture2D(FTexture2DRHIPara
 {
 	FMetalTexture2D* OldTexture = ResourceCast(OldTextureRHI);
 
-	FMetalTexture2D* NewTexture = new FMetalTexture2D(OldTexture->GetFormat(), NewSizeX, NewSizeY, NewMipCount, OldTexture->GetNumSamples(), OldTexture->GetFlags(), NULL);
+	FMetalTexture2D* NewTexture = new FMetalTexture2D(OldTexture->GetFormat(), NewSizeX, NewSizeY, NewMipCount, OldTexture->GetNumSamples(), OldTexture->GetFlags(), NULL, OldTextureRHI->GetClearBinding());
 
 	// @todo: gather these all up over a frame
 	id<MTLCommandBuffer> CommandBuffer = FMetalManager::Get()->CreateTempCommandBuffer(false/*bRetainReferences*/);
@@ -561,12 +561,12 @@ void FMetalDynamicRHI::RHIUpdateTexture3D(FTexture3DRHIParamRef TextureRHI,uint3
 -----------------------------------------------------------------------------*/
 FTextureCubeRHIRef FMetalDynamicRHI::RHICreateTextureCube(uint32 Size, uint8 Format, uint32 NumMips, uint32 Flags, FRHIResourceCreateInfo& CreateInfo)
 {
-	return new FMetalTextureCube((EPixelFormat)Format, Size, false, 1, NumMips, Flags, CreateInfo.BulkData);
+	return new FMetalTextureCube((EPixelFormat)Format, Size, false, 1, NumMips, Flags, CreateInfo.BulkData, CreateInfo.ClearValueBinding);
 }
 
 FTextureCubeRHIRef FMetalDynamicRHI::RHICreateTextureCubeArray(uint32 Size, uint32 ArraySize, uint8 Format, uint32 NumMips, uint32 Flags, FRHIResourceCreateInfo& CreateInfo)
 {
-	return new FMetalTextureCube((EPixelFormat)Format, Size, true, ArraySize, NumMips, Flags, CreateInfo.BulkData);
+	return new FMetalTextureCube((EPixelFormat)Format, Size, true, ArraySize, NumMips, Flags, CreateInfo.BulkData, CreateInfo.ClearValueBinding);
 }
 
 void* FMetalDynamicRHI::RHILockTextureCubeFace(FTextureCubeRHIParamRef TextureCubeRHI,uint32 FaceIndex,uint32 ArrayIndex,uint32 MipIndex,EResourceLockMode LockMode,uint32& DestStride,bool bLockWithinMiptail)
diff --git a/Engine/Source/Runtime/IOS/MetalRHI/Public/MetalRHI.h b/Engine/Source/Runtime/IOS/MetalRHI/Public/MetalRHI.h
--- a/Engine/Source/Runtime/IOS/MetalRHI/Public/MetalRHI.h
+++ b/Engine/Source/Runtime/IOS/MetalRHI/Public/MetalRHI.h
@@ -197,7 +197,6 @@ public:
 	virtual void RHISetBlendState(FBlendStateRHIParamRef NewState, const FLinearColor& BlendFactor) final override;
 	virtual void RHISetRenderTargets(uint32 NumSimultaneousRenderTargets, const FRHIRenderTargetView* NewRenderTargets, const FRHIDepthRenderTargetView* NewDepthStencilTarget, uint32 NumUAVs, const FUnorderedAccessViewRHIParamRef* UAVs) final override;
 	virtual void RHISetRenderTargetsAndClear(const FRHISetRenderTargetsInfo& RenderTargetsInfo) final override;
-	virtual void RHIBindClearMRTValues(bool bClearColor, int32 NumClearColors, const FLinearColor* ColorArray, bool bClearDepth, float Depth, bool bClearStencil, uint32 Stencil) final override;
 	virtual void RHIDrawPrimitive(uint32 PrimitiveType, uint32 BaseVertexIndex, uint32 NumPrimitives, uint32 NumInstances) final override;
 	virtual void RHIDrawPrimitiveIndirect(uint32 PrimitiveType, FVertexBufferRHIParamRef ArgumentBuffer, uint32 ArgumentOffset) final override;
 	virtual void RHIDrawIndexedIndirect(FIndexBufferRHIParamRef IndexBufferRHI, uint32 PrimitiveType, FStructuredBufferRHIParamRef ArgumentsBufferRHI, int32 DrawArgumentsIndex, uint32 NumInstances) final override;
diff --git a/Engine/Source/Runtime/IOS/MetalRHI/Public/MetalResources.h b/Engine/Source/Runtime/IOS/MetalRHI/Public/MetalResources.h
--- a/Engine/Source/Runtime/IOS/MetalRHI/Public/MetalResources.h
+++ b/Engine/Source/Runtime/IOS/MetalRHI/Public/MetalResources.h
@@ -210,8 +210,8 @@ public:
 	FMetalSurface Surface;
 
 	// Constructor, just calls base and Surface constructor
-	FMetalTexture2D(EPixelFormat Format, uint32 SizeX, uint32 SizeY, uint32 NumMips, uint32 NumSamples, uint32 Flags, FResourceBulkDataInterface* BulkData)
-		: FRHITexture2D(SizeX, SizeY, NumMips, NumSamples, Format, Flags)
+	FMetalTexture2D(EPixelFormat Format, uint32 SizeX, uint32 SizeY, uint32 NumMips, uint32 NumSamples, uint32 Flags, FResourceBulkDataInterface* BulkData, const FClearValueBinding& InClearValue)
+		: FRHITexture2D(SizeX, SizeY, NumMips, NumSamples, Format, Flags, InClearValue)
 		, Surface(RRT_Texture2D, Format, SizeX, SizeY, 1, NumSamples, /*bArray=*/ false, 1, NumMips, Flags, BulkData)
 	{
 	}
@@ -226,8 +226,8 @@ public:
 	FMetalSurface Surface;
 
 	// Constructor, just calls base and Surface constructor
-	FMetalTexture2DArray(EPixelFormat Format, uint32 SizeX, uint32 SizeY, uint32 ArraySize, uint32 NumMips, uint32 Flags, FResourceBulkDataInterface* BulkData)
-		: FRHITexture2DArray(SizeX, SizeY, ArraySize, NumMips, Format, Flags)
+	FMetalTexture2DArray(EPixelFormat Format, uint32 SizeX, uint32 SizeY, uint32 ArraySize, uint32 NumMips, uint32 Flags, FResourceBulkDataInterface* BulkData, const FClearValueBinding& InClearValue)
+		: FRHITexture2DArray(SizeX, SizeY, ArraySize, NumMips, Format, Flags, InClearValue)
 		, Surface(RRT_Texture2DArray, Format, SizeX, SizeY, 1, /*NumSamples=*/1, /*bArray=*/ true, ArraySize, NumMips, Flags, BulkData)
 	{
 	}
@@ -240,8 +240,8 @@ public:
 	FMetalSurface Surface;
 
 	// Constructor, just calls base and Surface constructor
-	FMetalTexture3D(EPixelFormat Format, uint32 SizeX, uint32 SizeY, uint32 SizeZ, uint32 NumMips, uint32 Flags, FResourceBulkDataInterface* BulkData)
-		: FRHITexture3D(SizeX, SizeY, SizeZ, NumMips, Format, Flags)
+	FMetalTexture3D(EPixelFormat Format, uint32 SizeX, uint32 SizeY, uint32 SizeZ, uint32 NumMips, uint32 Flags, FResourceBulkDataInterface* BulkData, const FClearValueBinding& InClearValue)
+		: FRHITexture3D(SizeX, SizeY, SizeZ, NumMips, Format, Flags, InClearValue)
 		, Surface(RRT_Texture3D, Format, SizeX, SizeY, SizeZ, /*NumSamples=*/1, /*bArray=*/ false, 1, NumMips, Flags, BulkData)
 	{
 	}
@@ -254,8 +254,8 @@ public:
 	FMetalSurface Surface;
 
 	// Constructor, just calls base and Surface constructor
-	FMetalTextureCube(EPixelFormat Format, uint32 Size, bool bArray, uint32 ArraySize, uint32 NumMips, uint32 Flags, FResourceBulkDataInterface* BulkData)
-		: FRHITextureCube(Size, NumMips, Format, Flags)
+	FMetalTextureCube(EPixelFormat Format, uint32 Size, bool bArray, uint32 ArraySize, uint32 NumMips, uint32 Flags, FResourceBulkDataInterface* BulkData, const FClearValueBinding& InClearValue)
+		: FRHITextureCube(Size, NumMips, Format, Flags, InClearValue)
 		, Surface(RRT_TextureCube, Format, Size, Size, 6, /*NumSamples=*/1, bArray, ArraySize, NumMips, Flags, BulkData)
 	{
 	}
diff --git a/Engine/Source/Runtime/Landscape/Private/LandscapeGrass.cpp b/Engine/Source/Runtime/Landscape/Private/LandscapeGrass.cpp
--- a/Engine/Source/Runtime/Landscape/Private/LandscapeGrass.cpp
+++ b/Engine/Source/Runtime/Landscape/Private/LandscapeGrass.cpp
@@ -419,7 +419,7 @@ public:
 
 		RenderTargetTexture = NewObject<UTextureRenderTarget2D>();
 		check(RenderTargetTexture);
-		RenderTargetTexture->ClearColor = FLinearColor::Transparent;
+		RenderTargetTexture->ClearColor = FLinearColor::White;
 		RenderTargetTexture->TargetGamma = 1.f;
 		RenderTargetTexture->InitCustomFormat(TargetSize.X, TargetSize.Y, PF_B8G8R8A8, false);
 		RenderTargetResource = RenderTargetTexture->GameThread_GetRenderTargetResource()->GetTextureRenderTarget2DResource();
diff --git a/Engine/Source/Runtime/NullDrv/Public/NullRHI.h b/Engine/Source/Runtime/NullDrv/Public/NullRHI.h
--- a/Engine/Source/Runtime/NullDrv/Public/NullRHI.h
+++ b/Engine/Source/Runtime/NullDrv/Public/NullRHI.h
@@ -261,7 +261,7 @@ public:
 
 	virtual FTexture2DRHIRef RHICreateTexture2D(uint32 SizeX, uint32 SizeY, uint8 Format, uint32 NumMips, uint32 NumSamples, uint32 Flags, FRHIResourceCreateInfo& CreateInfo) final override
 	{ 
-		return new FRHITexture2D(SizeX,SizeY,NumMips,NumSamples,(EPixelFormat)Format,Flags); 
+		return new FRHITexture2D(SizeX,SizeY,NumMips,NumSamples,(EPixelFormat)Format,Flags, CreateInfo.ClearValueBinding); 
 	}
 
 	virtual FTexture2DRHIRef RHIAsyncCreateTexture2D(uint32 SizeX, uint32 SizeY, uint8 Format, uint32 NumMips, uint32 Flags, void** InitialMipData, uint32 NumInitialMips) final override
@@ -274,12 +274,12 @@ public:
 	}
 	virtual FTexture2DArrayRHIRef RHICreateTexture2DArray(uint32 SizeX, uint32 SizeY, uint32 SizeZ, uint8 Format, uint32 NumMips, uint32 Flags, FRHIResourceCreateInfo& CreateInfo) final override
 	{ 
-		return new FRHITexture2DArray(SizeX,SizeY,SizeZ,NumMips,(EPixelFormat)Format,Flags); 
+		return new FRHITexture2DArray(SizeX,SizeY,SizeZ,NumMips,(EPixelFormat)Format,Flags, CreateInfo.ClearValueBinding); 
 	}
 
 	virtual FTexture3DRHIRef RHICreateTexture3D(uint32 SizeX, uint32 SizeY, uint32 SizeZ, uint8 Format, uint32 NumMips, uint32 Flags, FRHIResourceCreateInfo& CreateInfo) final override
 	{ 
-		return new FRHITexture3D(SizeX,SizeY,SizeZ,NumMips,(EPixelFormat)Format,Flags); 
+		return new FRHITexture3D(SizeX, SizeY, SizeZ, NumMips, (EPixelFormat)Format, Flags, CreateInfo.ClearValueBinding);
 	}
 	virtual void RHIGetResourceInfo(FTextureRHIParamRef Ref, FRHIResourceInfo& OutInfo) final override
 	{
@@ -320,7 +320,7 @@ public:
 	}
 	virtual FTexture2DRHIRef RHIAsyncReallocateTexture2D(FTexture2DRHIParamRef Texture2D, int32 NewMipCount, int32 NewSizeX, int32 NewSizeY, FThreadSafeCounter* RequestStatus) final override
 	{ 
-		return new FRHITexture2D(NewSizeX,NewSizeY,NewMipCount,1,Texture2D->GetFormat(),Texture2D->GetFlags()); 
+		return new FRHITexture2D(NewSizeX,NewSizeY,NewMipCount,1,Texture2D->GetFormat(),Texture2D->GetFlags(), Texture2D->GetClearBinding());
 	}
 	virtual ETextureReallocationStatus RHIFinalizeAsyncReallocateTexture2D(FTexture2DRHIParamRef Texture2D, bool bBlockUntilCompleted) final override
 	{ 
@@ -358,11 +358,11 @@ public:
 	}
 	virtual FTextureCubeRHIRef RHICreateTextureCube(uint32 Size, uint8 Format, uint32 NumMips, uint32 Flags, FRHIResourceCreateInfo& CreateInfo) final override
 	{ 
-		return new FRHITextureCube(Size,NumMips,(EPixelFormat)Format,Flags); 
+		return new FRHITextureCube(Size, NumMips, (EPixelFormat)Format, Flags, CreateInfo.ClearValueBinding);
 	}
 	virtual FTextureCubeRHIRef RHICreateTextureCubeArray(uint32 Size, uint32 ArraySize, uint8 Format, uint32 NumMips, uint32 Flags, FRHIResourceCreateInfo& CreateInfo) final override
 	{ 
-		return new FRHITextureCube(Size,NumMips,(EPixelFormat)Format,Flags); 
+		return new FRHITextureCube(Size, NumMips, (EPixelFormat)Format, Flags, CreateInfo.ClearValueBinding);
 	}
 	virtual void* RHILockTextureCubeFace(FTextureCubeRHIParamRef Texture, uint32 FaceIndex, uint32 ArrayIndex, uint32 MipIndex, EResourceLockMode LockMode, uint32& DestStride, bool bLockWithinMiptail) final override
 	{ 
@@ -457,7 +457,7 @@ public:
 
 	virtual FTexture2DRHIRef RHIGetViewportBackBuffer(FViewportRHIParamRef Viewport) final override
 	{ 
-		return new FRHITexture2D(1,1,1,1,PF_B8G8R8A8,TexCreate_RenderTargetable); 
+		return new FRHITexture2D(1,1,1,1,PF_B8G8R8A8,TexCreate_RenderTargetable, FClearValueBinding()); 
 	}
 
 	virtual void RHIBeginFrame() final override
@@ -736,15 +736,6 @@ public:
 
 	}
 
-
-
-
-	virtual void RHIBindClearMRTValues(bool bClearColor, int32 NumClearColors, const FLinearColor* ColorArray, bool bClearDepth, float Depth, bool bClearStencil, uint32 Stencil) final override
-	{
-
-	}
-
-
 	virtual void RHIDrawPrimitive(uint32 PrimitiveType, uint32 BaseVertexIndex, uint32 NumPrimitives, uint32 NumInstances) final override
 	{
 
diff --git a/Engine/Source/Runtime/OpenGLDrv/Private/IOS/IOSOpenGL.cpp b/Engine/Source/Runtime/OpenGLDrv/Private/IOS/IOSOpenGL.cpp
--- a/Engine/Source/Runtime/OpenGLDrv/Private/IOS/IOSOpenGL.cpp
+++ b/Engine/Source/Runtime/OpenGLDrv/Private/IOS/IOSOpenGL.cpp
@@ -230,7 +230,7 @@ FRHITexture* PlatformCreateBuiltinBackBuffer(FOpenGLDynamicRHI* OpenGLRHI, uint3
 	[GLView UpdateRenderWidth:SizeX andHeight:SizeY];
 
 	uint32 Flags = TexCreate_RenderTargetable | TexCreate_Presentable;
-	FOpenGLTexture2D* Texture2D = new FOpenGLTexture2D(OpenGLRHI, GLView.OnScreenColorRenderBuffer, GL_RENDERBUFFER, GL_COLOR_ATTACHMENT0, SizeX, SizeY, 0, 1, 1, 1, PF_B8G8R8A8, false, false, Flags, nullptr);
+	FOpenGLTexture2D* Texture2D = new FOpenGLTexture2D(OpenGLRHI, GLView.OnScreenColorRenderBuffer, GL_RENDERBUFFER, GL_COLOR_ATTACHMENT0, SizeX, SizeY, 0, 1, 1, 1, PF_B8G8R8A8, false, false, Flags, nullptr, FClearValueBinding::Transparent);
 	OpenGLTextureAllocated(Texture2D, Flags);
 
 	return Texture2D;
diff --git a/Engine/Source/Runtime/OpenGLDrv/Private/OpenGLCommands.cpp b/Engine/Source/Runtime/OpenGLDrv/Private/OpenGLCommands.cpp
--- a/Engine/Source/Runtime/OpenGLDrv/Private/OpenGLCommands.cpp
+++ b/Engine/Source/Runtime/OpenGLDrv/Private/OpenGLCommands.cpp
@@ -1828,7 +1828,27 @@ void FOpenGLDynamicRHI::RHISetRenderTargetsAndClear(const FRHISetRenderTargetsIn
 		nullptr);
 	if (RenderTargetsInfo.bClearColor || RenderTargetsInfo.bClearStencil || RenderTargetsInfo.bClearDepth)
 	{
-		this->RHIClearMRT(RenderTargetsInfo.bClearColor, RenderTargetsInfo.NumColorRenderTargets, RenderTargetsInfo.ClearColors, RenderTargetsInfo.bClearDepth, RenderTargetsInfo.DepthClearValue, RenderTargetsInfo.bClearStencil, RenderTargetsInfo.StencilClearValue, FIntRect());
+		FLinearColor ClearColors[MaxSimultaneousRenderTargets];
+		float DepthClear = 0.0;
+		uint32 StencilClear = 0;
+
+		if (RenderTargetsInfo.bClearColor)
+		{
+			for (int32 i = 0; i < RenderTargetsInfo.NumColorRenderTargets; ++i)
+			{
+				const FClearValueBinding& ClearValue = RenderTargetsInfo.ColorRenderTarget[i].Texture->GetClearBinding();
+				checkf(ClearValue.ColorBinding == EClearBinding::EColorBound, TEXT("Texture: %s does not have a color bound for fast clears"), *RenderTargetsInfo.ColorRenderTarget[i].Texture->GetName().GetPlainNameString());
+				ClearColors[i] = ClearValue.GetClearColor();
+			}
+		}
+		if (RenderTargetsInfo.bClearDepth || RenderTargetsInfo.bClearStencil)
+		{
+			const FClearValueBinding& ClearValue = RenderTargetsInfo.DepthStencilRenderTarget.Texture->GetClearBinding();
+			checkf(ClearValue.ColorBinding == EClearBinding::EDepthStencilBound, TEXT("Texture: %s does not have a DS value bound for fast clears"), *RenderTargetsInfo.DepthStencilRenderTarget.Texture->GetName().GetPlainNameString());
+			ClearValue.GetDepthStencil(DepthClear, StencilClear);
+		}
+
+		this->RHIClearMRT(RenderTargetsInfo.bClearColor, RenderTargetsInfo.NumColorRenderTargets, ClearColors, RenderTargetsInfo.bClearDepth, DepthClear, RenderTargetsInfo.bClearStencil, StencilClear, FIntRect());
 	}
 }
 
@@ -3376,11 +3396,6 @@ void FOpenGLDynamicRHI::RHIClearMRT(bool bClearColor,int32 NumClearColors,const
 	}
 }
 
-void FOpenGLDynamicRHI::RHIBindClearMRTValues(bool bClearColor, int32 NumClearColors, const FLinearColor* ClearColorArray, bool bClearDepth, float Depth, bool bClearStencil, uint32 Stencil)
-{
-	// Not necessary for opengl.
-}
-
 // Functions to yield and regain rendering control from OpenGL
 
 void FOpenGLDynamicRHI::RHISuspendRendering()
diff --git a/Engine/Source/Runtime/OpenGLDrv/Private/OpenGLTexture.cpp b/Engine/Source/Runtime/OpenGLDrv/Private/OpenGLTexture.cpp
--- a/Engine/Source/Runtime/OpenGLDrv/Private/OpenGLTexture.cpp
+++ b/Engine/Source/Runtime/OpenGLDrv/Private/OpenGLTexture.cpp
@@ -230,7 +230,7 @@ bool FOpenGLDynamicRHI::RHIGetTextureMemoryVisualizeData( FColor* /*TextureData*
 
 
 
-FRHITexture* FOpenGLDynamicRHI::CreateOpenGLTexture(uint32 SizeX,uint32 SizeY,bool bCubeTexture, bool bArrayTexture, uint8 Format,uint32 NumMips,uint32 NumSamples, uint32 ArraySize, uint32 Flags, FResourceBulkDataInterface* BulkData)
+FRHITexture* FOpenGLDynamicRHI::CreateOpenGLTexture(uint32 SizeX, uint32 SizeY, bool bCubeTexture, bool bArrayTexture, uint8 Format, uint32 NumMips, uint32 NumSamples, uint32 ArraySize, uint32 Flags, const FClearValueBinding& InClearValue, FResourceBulkDataInterface* BulkData)
 {
 	VERIFY_GL_SCOPE();
 
@@ -558,12 +558,12 @@ FRHITexture* FOpenGLDynamicRHI::CreateOpenGLTexture(uint32 SizeX,uint32 SizeY,bo
 	FRHITexture* Texture;
 	if (bCubeTexture)
 	{
-		FOpenGLTextureCube* TextureCube = new FOpenGLTextureCube(this,TextureID,Target,Attachment,SizeX,SizeY,0,NumMips,1, ArraySize, (EPixelFormat)Format,true,bAllocatedStorage,Flags,TextureRange);
+		FOpenGLTextureCube* TextureCube = new FOpenGLTextureCube(this, TextureID, Target, Attachment, SizeX, SizeY, 0, NumMips, 1, ArraySize, (EPixelFormat)Format, true, bAllocatedStorage, Flags, TextureRange, InClearValue);
 		Texture = TextureCube;
 	}
 	else
 	{
-		FOpenGLTexture2D* Texture2D = new FOpenGLTexture2D(this,TextureID,Target,Attachment,SizeX,SizeY,0,NumMips,NumSamples, 1, (EPixelFormat)Format,false,bAllocatedStorage,Flags,TextureRange);
+		FOpenGLTexture2D* Texture2D = new FOpenGLTexture2D(this,TextureID,Target,Attachment,SizeX,SizeY,0,NumMips,NumSamples, 1, (EPixelFormat)Format,false,bAllocatedStorage,Flags,TextureRange, InClearValue);
 		Texture = Texture2D;
 	}
 	OpenGLTextureAllocated( Texture, Flags );
@@ -1412,10 +1412,10 @@ void TOpenGLTexture<RHIResourceType>::CloneViaPBO( TOpenGLTexture<RHIResourceTyp
 */
 FTexture2DRHIRef FOpenGLDynamicRHI::RHICreateTexture2D(uint32 SizeX,uint32 SizeY,uint8 Format,uint32 NumMips,uint32 NumSamples,uint32 Flags,FRHIResourceCreateInfo& Info)
 {
-	return (FRHITexture2D*)CreateOpenGLTexture(SizeX,SizeY,false,false,Format,NumMips,NumSamples,1,Flags,Info.BulkData);
+	return (FRHITexture2D*)CreateOpenGLTexture(SizeX,SizeY,false,false,Format,NumMips,NumSamples,1,Flags,Info.ClearValueBinding,Info.BulkData);
 }
 
-FTexture2DRHIRef FOpenGLDynamicRHI::RHIAsyncCreateTexture2D(uint32 SizeX,uint32 SizeY,uint8 Format,uint32 NumMips,uint32 Flags,void** InitialMipData,uint32 NumInitialMips)
+FTexture2DRHIRef FOpenGLDynamicRHI::RHIAsyncCreateTexture2D(uint32 SizeX, uint32 SizeY, uint8 Format, uint32 NumMips, uint32 Flags, void** InitialMipData, uint32 NumInitialMips)
 {
 	check(0);
 	return FTexture2DRHIRef();
@@ -1525,7 +1525,7 @@ FTexture2DArrayRHIRef FOpenGLDynamicRHI::RHICreateTexture2DArray(uint32 SizeX,ui
 			: GL_COLOR_ATTACHMENT0);
 	}
 
-	FOpenGLTexture2DArray* Texture = new FOpenGLTexture2DArray(this,TextureID,Target,Attachment,SizeX,SizeY,SizeZ,NumMips,1, SizeZ, (EPixelFormat)Format,false,true,Flags,nullptr);
+	FOpenGLTexture2DArray* Texture = new FOpenGLTexture2DArray(this,TextureID,Target,Attachment,SizeX,SizeY,SizeZ,NumMips,1, SizeZ, (EPixelFormat)Format,false,true,Flags,nullptr,Info.ClearValueBinding);
 	OpenGLTextureAllocated( Texture, Flags );
 
 	// No need to restore texture stage; leave it like this,
@@ -1633,7 +1633,7 @@ FTexture3DRHIRef FOpenGLDynamicRHI::RHICreateTexture3D(uint32 SizeX,uint32 SizeY
 			: GL_COLOR_ATTACHMENT0);
 	}
 
-	FOpenGLTexture3D* Texture = new FOpenGLTexture3D(this,TextureID,Target,Attachment,SizeX,SizeY,SizeZ,NumMips,1,1, (EPixelFormat)Format,false,true,Flags,nullptr);
+	FOpenGLTexture3D* Texture = new FOpenGLTexture3D(this,TextureID,Target,Attachment,SizeX,SizeY,SizeZ,NumMips,1,1, (EPixelFormat)Format,false,true,Flags,nullptr,CreateInfo.ClearValueBinding);
 	OpenGLTextureAllocated( Texture, Flags );
 
 	// No need to restore texture stage; leave it like this,
@@ -1886,7 +1886,7 @@ FTexture2DRHIRef FOpenGLDynamicRHI::RHIAsyncReallocateTexture2D(FTexture2DRHIPar
 	FOpenGLTexture2D* Texture2D = ResourceCast(Texture2DRHI);
 
 	// Allocate a new texture.
-	FOpenGLTexture2D* NewTexture2D = (FOpenGLTexture2D*)CreateOpenGLTexture(NewSizeX,NewSizeY,false,false, Texture2D->GetFormat(),NewMipCount,1,1, Texture2D->GetFlags());
+	FOpenGLTexture2D* NewTexture2D = (FOpenGLTexture2D*)CreateOpenGLTexture(NewSizeX,NewSizeY,false,false, Texture2D->GetFormat(),NewMipCount,1,1, Texture2D->GetFlags(), Texture2DRHI->GetClearBinding());
 	
 	const uint32 BlockSizeX = GPixelFormats[Texture2D->GetFormat()].BlockSizeX;
 	const uint32 BlockSizeY = GPixelFormats[Texture2D->GetFormat()].BlockSizeY;
@@ -2115,7 +2115,7 @@ FTextureCubeRHIRef FOpenGLDynamicRHI::RHICreateTextureCube( uint32 Size, uint8 F
 	// not yet supported
 	check(!CreateInfo.BulkData);
 
-	return (FRHITextureCube*)CreateOpenGLTexture(Size,Size,true, false, Format, NumMips, 1, 1, Flags);
+	return (FRHITextureCube*)CreateOpenGLTexture(Size,Size,true, false, Format, NumMips, 1, 1, Flags, CreateInfo.ClearValueBinding);
 }
 
 FTextureCubeRHIRef FOpenGLDynamicRHI::RHICreateTextureCubeArray( uint32 Size, uint32 ArraySize, uint8 Format, uint32 NumMips, uint32 Flags, FRHIResourceCreateInfo& CreateInfo )
@@ -2123,7 +2123,7 @@ FTextureCubeRHIRef FOpenGLDynamicRHI::RHICreateTextureCubeArray( uint32 Size, ui
 	// not yet supported
 	check(!CreateInfo.BulkData);
 
-	return (FRHITextureCube*)CreateOpenGLTexture(Size,Size,true, true, Format, NumMips, 1, 6 * ArraySize, Flags);
+	return (FRHITextureCube*)CreateOpenGLTexture(Size, Size, true, true, Format, NumMips, 1, 6 * ArraySize, Flags, CreateInfo.ClearValueBinding);
 }
 
 void* FOpenGLDynamicRHI::RHILockTextureCubeFace(FTextureCubeRHIParamRef TextureCubeRHI,uint32 FaceIndex,uint32 ArrayIndex,uint32 MipIndex,EResourceLockMode LockMode,uint32& DestStride,bool bLockWithinMiptail)
diff --git a/Engine/Source/Runtime/OpenGLDrv/Private/OpenGLViewport.cpp b/Engine/Source/Runtime/OpenGLDrv/Private/OpenGLViewport.cpp
--- a/Engine/Source/Runtime/OpenGLDrv/Private/OpenGLViewport.cpp
+++ b/Engine/Source/Runtime/OpenGLDrv/Private/OpenGLViewport.cpp
@@ -283,7 +283,7 @@ void FOpenGLViewport::Resize(uint32 InSizeX,uint32 InSizeY,bool bInIsFullscreen)
 	BackBuffer = (FOpenGLTexture2D*)PlatformCreateBuiltinBackBuffer(OpenGLRHI, InSizeX, InSizeY);
 	if (!BackBuffer)
 	{
-		BackBuffer = (FOpenGLTexture2D*)OpenGLRHI->CreateOpenGLTexture(InSizeX, InSizeY, false, false, PixelFormat, 1, 1, 1, TexCreate_RenderTargetable);
+		BackBuffer = (FOpenGLTexture2D*)OpenGLRHI->CreateOpenGLTexture(InSizeX, InSizeY, false, false, PixelFormat, 1, 1, 1, TexCreate_RenderTargetable, FClearValueBinding::None);
 	}
 
 	PlatformResizeGLContext(OpenGLRHI->PlatformDevice, OpenGLContext, InSizeX, InSizeY, bInIsFullscreen, bIsFullscreen, BackBuffer->Target, BackBuffer->Resource);
diff --git a/Engine/Source/Runtime/OpenGLDrv/Public/OpenGLDrv.h b/Engine/Source/Runtime/OpenGLDrv/Public/OpenGLDrv.h
--- a/Engine/Source/Runtime/OpenGLDrv/Public/OpenGLDrv.h
+++ b/Engine/Source/Runtime/OpenGLDrv/Public/OpenGLDrv.h
@@ -466,7 +466,6 @@ public:
 	virtual void RHISetBlendState(FBlendStateRHIParamRef NewState, const FLinearColor& BlendFactor) final override;
 	virtual void RHISetRenderTargets(uint32 NumSimultaneousRenderTargets, const FRHIRenderTargetView* NewRenderTargets, const FRHIDepthRenderTargetView* NewDepthStencilTarget, uint32 NumUAVs, const FUnorderedAccessViewRHIParamRef* UAVs) final override;
 	virtual void RHISetRenderTargetsAndClear(const FRHISetRenderTargetsInfo& RenderTargetsInfo) final override;
-	virtual void RHIBindClearMRTValues(bool bClearColor, int32 NumClearColors, const FLinearColor* ColorArray, bool bClearDepth, float Depth, bool bClearStencil, uint32 Stencil) final override;
 	virtual void RHIDrawPrimitive(uint32 PrimitiveType, uint32 BaseVertexIndex, uint32 NumPrimitives, uint32 NumInstances) final override;
 	virtual void RHIDrawPrimitiveIndirect(uint32 PrimitiveType, FVertexBufferRHIParamRef ArgumentBuffer, uint32 ArgumentOffset) final override;
 	virtual void RHIDrawIndexedIndirect(FIndexBufferRHIParamRef IndexBufferRHI, uint32 PrimitiveType, FStructuredBufferRHIParamRef ArgumentsBufferRHI, int32 DrawArgumentsIndex, uint32 NumInstances) final override;
@@ -558,7 +557,7 @@ public:
 
 	FOpenGLSamplerState* GetPointSamplerState() const { return (FOpenGLSamplerState*)PointSamplerState.GetReference(); }
 
-	FRHITexture* CreateOpenGLTexture(uint32 SizeX, uint32 SizeY, bool CubeTexture, bool ArrayTexture, uint8 Format, uint32 NumMips, uint32 NumSamples, uint32 ArraySize, uint32 Flags, FResourceBulkDataInterface* BulkData = NULL);
+	FRHITexture* CreateOpenGLTexture(uint32 SizeX, uint32 SizeY, bool CubeTexture, bool ArrayTexture, uint8 Format, uint32 NumMips, uint32 NumSamples, uint32 ArraySize, uint32 Flags, const FClearValueBinding& InClearValue, FResourceBulkDataInterface* BulkData = NULL);
 
 private:
 
diff --git a/Engine/Source/Runtime/OpenGLDrv/Public/OpenGLResources.h b/Engine/Source/Runtime/OpenGLDrv/Public/OpenGLResources.h
--- a/Engine/Source/Runtime/OpenGLDrv/Public/OpenGLResources.h
+++ b/Engine/Source/Runtime/OpenGLDrv/Public/OpenGLResources.h
@@ -844,9 +844,10 @@ public:
 		bool bInCubemap,
 		bool bInAllocatedStorage,
 		uint32 InFlags,
-		uint8* InTextureRange
+		uint8* InTextureRange,
+		const FClearValueBinding& InClearValue
 		)
-	: BaseType(InSizeX,InSizeY,InSizeZ,InNumMips,InNumSamples, InArraySize, InFormat,InFlags)
+	: BaseType(InSizeX,InSizeY,InSizeZ,InNumMips,InNumSamples, InArraySize, InFormat,InFlags, InClearValue)
 	, FOpenGLTextureBase(
 		InOpenGLRHI,
 		InResource,
@@ -1045,8 +1046,8 @@ private:
 class OPENGLDRV_API FOpenGLBaseTexture2D : public FRHITexture2D
 {
 public:
-	FOpenGLBaseTexture2D(uint32 InSizeX,uint32 InSizeY,uint32 InSizeZ,uint32 InNumMips,uint32 InNumSamples, uint32 InArraySize, EPixelFormat InFormat,uint32 InFlags)
-	: FRHITexture2D(InSizeX,InSizeY,InNumMips,InNumSamples,InFormat,InFlags)
+	FOpenGLBaseTexture2D(uint32 InSizeX, uint32 InSizeY, uint32 InSizeZ, uint32 InNumMips, uint32 InNumSamples, uint32 InArraySize, EPixelFormat InFormat, uint32 InFlags, const FClearValueBinding& InClearValue)
+	: FRHITexture2D(InSizeX,InSizeY,InNumMips,InNumSamples,InFormat,InFlags, InClearValue)
 	, SampleCount(InNumSamples)
 	{}
 	uint32 GetSizeZ() const { return 0; }
@@ -1058,8 +1059,8 @@ private:
 class FOpenGLBaseTexture2DArray : public FRHITexture2DArray
 {
 public:
-	FOpenGLBaseTexture2DArray(uint32 InSizeX,uint32 InSizeY,uint32 InSizeZ,uint32 InNumMips,uint32 InNumSamples,uint32 InArraySize, EPixelFormat InFormat,uint32 InFlags)
-	: FRHITexture2DArray(InSizeX,InSizeY,InSizeZ,InNumMips,InFormat,InFlags)
+	FOpenGLBaseTexture2DArray(uint32 InSizeX, uint32 InSizeY, uint32 InSizeZ, uint32 InNumMips, uint32 InNumSamples, uint32 InArraySize, EPixelFormat InFormat, uint32 InFlags, const FClearValueBinding& InClearValue)
+	: FRHITexture2DArray(InSizeX,InSizeY,InSizeZ,InNumMips,InFormat,InFlags, InClearValue)
 	{
 		check(InNumSamples == 1);	// OpenGL supports multisampled texture arrays, but they're currently not implemented in OpenGLDrv.
 	}
@@ -1068,8 +1069,8 @@ public:
 class FOpenGLBaseTextureCube : public FRHITextureCube
 {
 public:
-	FOpenGLBaseTextureCube(uint32 InSizeX,uint32 InSizeY,uint32 InSizeZ,uint32 InNumMips,uint32 InNumSamples, uint32 InArraySize, EPixelFormat InFormat,uint32 InFlags)
-	: FRHITextureCube(InSizeX,InNumMips,InFormat,InFlags)
+	FOpenGLBaseTextureCube(uint32 InSizeX, uint32 InSizeY, uint32 InSizeZ, uint32 InNumMips, uint32 InNumSamples, uint32 InArraySize, EPixelFormat InFormat, uint32 InFlags, const FClearValueBinding& InClearValue)
+	: FRHITextureCube(InSizeX,InNumMips,InFormat,InFlags,InClearValue)
 	, ArraySize(InArraySize)
 	{
 		check(InNumSamples == 1);	// OpenGL doesn't currently support multisampled cube textures
@@ -1087,8 +1088,8 @@ private:
 class FOpenGLBaseTexture3D : public FRHITexture3D
 {
 public:
-	FOpenGLBaseTexture3D(uint32 InSizeX,uint32 InSizeY,uint32 InSizeZ,uint32 InNumMips,uint32 InNumSamples,uint32 InArraySize, EPixelFormat InFormat,uint32 InFlags)
-	: FRHITexture3D(InSizeX,InSizeY,InSizeZ,InNumMips,InFormat,InFlags)
+	FOpenGLBaseTexture3D(uint32 InSizeX, uint32 InSizeY, uint32 InSizeZ, uint32 InNumMips, uint32 InNumSamples, uint32 InArraySize, EPixelFormat InFormat, uint32 InFlags, const FClearValueBinding& InClearValue)
+	: FRHITexture3D(InSizeX,InSizeY,InSizeZ,InNumMips,InFormat,InFlags,InClearValue)
 	{
 		check(InNumSamples == 1);	// Can't have multisampled texture 3D. Not supported anywhere.
 	}
diff --git a/Engine/Source/Runtime/RHI/Private/RHI.cpp b/Engine/Source/Runtime/RHI/Private/RHI.cpp
--- a/Engine/Source/Runtime/RHI/Private/RHI.cpp
+++ b/Engine/Source/Runtime/RHI/Private/RHI.cpp
@@ -60,6 +60,18 @@ static FAutoConsoleCommandWithOutputDevice GDumpRHIMemoryCmd(
 	);
 #endif
 
+//DO NOT USE THE STATIC FLINEARCOLORS TO INITIALIZE THIS STUFF.  
+//Static init order is undefined and you will likely end up with bad values on some platforms.
+const FClearValueBinding FClearValueBinding::None(EClearBinding::ENoneBound);
+const FClearValueBinding FClearValueBinding::Black(FLinearColor(0.0f, 0.0f, 0.0f, 1.0f));
+const FClearValueBinding FClearValueBinding::White(FLinearColor(1.0f, 1.0f, 1.0f, 1.0f));
+const FClearValueBinding FClearValueBinding::Transparent(FLinearColor(0.0f, 0.0f, 0.0f, 0.0f));
+const FClearValueBinding FClearValueBinding::DepthOne(1.0f, 0);
+const FClearValueBinding FClearValueBinding::DepthZero(0.0f, 0);
+const FClearValueBinding FClearValueBinding::DepthNear((float)ERHIZBuffer::NearPlane, 0);
+const FClearValueBinding FClearValueBinding::DepthFar((float)ERHIZBuffer::FarPlane, 0);
+
+
 TLockFreePointerList<FRHIResource> FRHIResource::PendingDeletes;
 FRHIResource* FRHIResource::CurrentlyDeleting = nullptr;
 
diff --git a/Engine/Source/Runtime/RHI/Private/RHICommandList.cpp b/Engine/Source/Runtime/RHI/Private/RHICommandList.cpp
--- a/Engine/Source/Runtime/RHI/Private/RHICommandList.cpp
+++ b/Engine/Source/Runtime/RHI/Private/RHICommandList.cpp
@@ -460,8 +460,9 @@ void FRHICommandListExecutor::LatchBypass()
 	}
 	else
 	{
-		GRHICommandList.GetImmediateCommandList().ImmediateFlush(EImmediateFlushType::DispatchToRHIThread);
+		GRHICommandList.GetImmediateCommandList().ImmediateFlush(EImmediateFlushType::DispatchToRHIThread);		
 
+#if 0
 		static bool bOnce = false;
 		if (!bOnce)
 		{
@@ -472,9 +473,11 @@ void FRHICommandListExecutor::LatchBypass()
 				BypassVar->Set(0, ECVF_SetByCommandline);
 			}
 		}
+#endif
 
 		check((GRHICommandList.OutstandingCmdListCount.GetValue() == 1 && !GRHICommandList.GetImmediateCommandList().HasCommands()));
 		bool NewBypass = (CVarRHICmdBypass.GetValueOnAnyThread() >= 1);
+		NewBypass = true;
 		if (NewBypass && !bLatchedBypass)
 		{
 			FRHIResource::FlushPendingDeletes();
diff --git a/Engine/Source/Runtime/RHI/Public/DynamicRHI.h b/Engine/Source/Runtime/RHI/Public/DynamicRHI.h
--- a/Engine/Source/Runtime/RHI/Public/DynamicRHI.h
+++ b/Engine/Source/Runtime/RHI/Public/DynamicRHI.h
@@ -233,7 +233,7 @@ public:
 	/**
 	 * Thread-safe function that can be used to create a texture outside of the
 	 * rendering thread. This function can ONLY be called if GRHISupportsAsyncTextureCreation
-	 * is true.
+	 * is true.  Cannot create rendertargets with this method.
 	 * @param SizeX - width of the texture to create
 	 * @param SizeY - height of the texture to create
 	 * @param Format - EPixelFormat texture format
@@ -876,10 +876,10 @@ public:
 
 	virtual void RHISetRenderTargetsAndClear(const FRHISetRenderTargetsInfo& RenderTargetsInfo) = 0;
 
-	// Bind a clear values to currently bound rendertargets.  This is used by platforms which
-	// need the color when transitioning a target that supports hardware clears from a rendertarget to a shader resource.
-	// The explicit bind is needed to support parallel rendering.
-	virtual void RHIBindClearMRTValues(bool bClearColor, int32 NumClearColors, const FLinearColor* ColorArray, bool bClearDepth, float Depth, bool bClearStencil, uint32 Stencil) = 0;
+	// Bind the clear state of the currently set rendertargets.  This is used by platforms which
+	// need the state of the target when finalizing a hardware clear or a resource transition to SRV
+	// The explicit bind is needed to support parallel rendering (propagate state between contexts).
+	virtual void RHIBindClearMRTValues(bool bClearColor, bool bClearDepth, bool bClearStencil){}
 
 	virtual void RHIDrawPrimitive(uint32 PrimitiveType, uint32 BaseVertexIndex, uint32 NumPrimitives, uint32 NumInstances) = 0;
 
diff --git a/Engine/Source/Runtime/RHI/Public/RHI.h b/Engine/Source/Runtime/RHI/Public/RHI.h
--- a/Engine/Source/Runtime/RHI/Public/RHI.h
+++ b/Engine/Source/Runtime/RHI/Public/RHI.h
@@ -742,27 +742,144 @@ struct FRHIResourceInfo
 	FVRamAllocation VRamAllocation;
 };
 
+enum class EClearBinding
+{
+	ENoneBound, //no clear color associated with this target.  Target will not do hardware clears on most platforms
+	EColorBound, //target has a clear color bound.  Clears will use the bound color, and do hardware clears.
+	EDepthStencilBound, //target has a depthstencil value bound.  Clears will use the bound values and do hardware clears.
+};
+
+struct FClearValueBinding
+{
+	struct DSVAlue
+	{
+		float Depth;
+		uint32 Stencil;
+	};
+
+	FClearValueBinding()
+		: ColorBinding(EClearBinding::EColorBound)
+	{
+		Value.Color[0] = 0.0f;
+		Value.Color[1] = 0.0f;
+		Value.Color[2] = 0.0f;
+		Value.Color[3] = 0.0f;
+	}
+
+	FClearValueBinding(EClearBinding NoBinding)
+		: ColorBinding(NoBinding)
+	{
+		check(ColorBinding == EClearBinding::ENoneBound);
+	}
+
+	explicit FClearValueBinding(const FLinearColor& InClearColor)
+		: ColorBinding(EClearBinding::EColorBound)
+	{
+		Value.Color[0] = InClearColor.R;
+		Value.Color[1] = InClearColor.G;
+		Value.Color[2] = InClearColor.B;
+		Value.Color[3] = InClearColor.A;
+	}
+
+	explicit FClearValueBinding(float DepthClearValue, uint32 StencilClearValue = 0)
+		: ColorBinding(EClearBinding::EDepthStencilBound)
+	{
+		Value.DSValue.Depth = DepthClearValue;
+		Value.DSValue.Stencil = StencilClearValue;
+	}
+
+	FLinearColor GetClearColor() const
+	{
+		ensure(ColorBinding == EClearBinding::EColorBound);
+		return FLinearColor(Value.Color[0], Value.Color[1], Value.Color[2], Value.Color[3]);
+	}
+
+	void GetDepthStencil(float& OutDepth, uint32& OutStencil) const
+	{
+		ensure(ColorBinding == EClearBinding::EDepthStencilBound);
+		OutDepth = Value.DSValue.Depth;
+		OutStencil = Value.DSValue.Stencil;
+	}
+
+	bool operator==(const FClearValueBinding& Other) const
+	{
+		if (ColorBinding == Other.ColorBinding)
+		{
+			if (ColorBinding == EClearBinding::EColorBound)
+			{
+				return
+					Value.Color[0] == Other.Value.Color[0] &&
+					Value.Color[1] == Other.Value.Color[1] &&
+					Value.Color[2] == Other.Value.Color[2] &&
+					Value.Color[3] == Other.Value.Color[3];
+
+			}
+			if (ColorBinding == EClearBinding::EDepthStencilBound)
+			{
+				return
+					Value.DSValue.Depth == Value.DSValue.Depth &&
+					Value.DSValue.Stencil == Value.DSValue.Stencil;
+			}
+			return true;
+		}
+		return false;
+	}
+
+	EClearBinding ColorBinding;
+
+	union ClearValueType
+	{
+		float Color[4];
+		DSVAlue DSValue;
+	} Value;
+
+	// common clear values
+	static RHI_API const FClearValueBinding None;
+	static RHI_API const FClearValueBinding Black;
+	static RHI_API const FClearValueBinding White;
+	static RHI_API const FClearValueBinding Transparent;
+	static RHI_API const FClearValueBinding DepthOne;
+	static RHI_API const FClearValueBinding DepthZero;
+	static RHI_API const FClearValueBinding DepthNear;
+	static RHI_API const FClearValueBinding DepthFar;	
+};
+
 struct FRHIResourceCreateInfo
 {
 	FRHIResourceCreateInfo()
-		: BulkData(0)
-		, ResourceArray(0)
+		: BulkData(nullptr)
+		, ResourceArray(nullptr)
+		, ClearValueBinding(FLinearColor::Transparent)
 	{}
 
 	// for CreateTexture calls
 	FRHIResourceCreateInfo(FResourceBulkDataInterface* InBulkData)
 		: BulkData(InBulkData)
+		, ResourceArray(nullptr)
+		, ClearValueBinding(FLinearColor::Transparent)
 	{}
 
 	// for CreateVertexBuffer/CreateStructuredBuffer calls
 	FRHIResourceCreateInfo(FResourceArrayInterface* InResourceArray)
-		: ResourceArray(InResourceArray)
+		: BulkData(nullptr)
+		, ResourceArray(InResourceArray)
+		, ClearValueBinding(FLinearColor::Transparent)
 	{}
 
+	FRHIResourceCreateInfo(const FClearValueBinding& InClearValueBinding)
+		: BulkData(nullptr)
+		, ResourceArray(nullptr)
+		, ClearValueBinding(InClearValueBinding)
+	{
+	}
+
 	// for CreateTexture calls
 	FResourceBulkDataInterface* BulkData;
 	// for CreateVertexBuffer/CreateStructuredBuffer calls
 	FResourceArrayInterface* ResourceArray;
+
+	// for binding clear colors to rendertargets.
+	FClearValueBinding ClearValueBinding;
 };
 
 // Forward-declaration.
diff --git a/Engine/Source/Runtime/RHI/Public/RHICommandList.h b/Engine/Source/Runtime/RHI/Public/RHICommandList.h
--- a/Engine/Source/Runtime/RHI/Public/RHICommandList.h
+++ b/Engine/Source/Runtime/RHI/Public/RHICommandList.h
@@ -573,36 +573,20 @@ struct FRHICommandSetRenderTargetsAndClear : public FRHICommand<FRHICommandSetRe
 };
 
 struct FRHICommandBindClearMRTValues : public FRHICommand<FRHICommandBindClearMRTValues>
-{
-	FLinearColor ColorArray[MaxSimultaneousRenderTargets];	
-	float Depth;
-	uint32 Stencil;
-	int32 NumClearColors;
+{	
 	bool bClearColor;
 	bool bClearDepth;
 	bool bClearStencil;
 
 	FORCEINLINE_DEBUGGABLE FRHICommandBindClearMRTValues(
-		bool InbClearColor,
-		int32 InNumClearColors,
-		const FLinearColor* InColorArray,
-		bool InbClearDepth,
-		float InDepth,
-		bool InbClearStencil,
-		uint32 InStencil		
+		bool InbClearColor,		
+		bool InbClearDepth,		
+		bool InbClearStencil
 		) 
-		: Depth(InDepth)
-		, Stencil(InStencil)
-		, NumClearColors(InNumClearColors)
-		, bClearColor(InbClearColor)
+		: bClearColor(InbClearColor)
 		, bClearDepth(InbClearDepth)
 		, bClearStencil(InbClearStencil)
-	{
-		check(InNumClearColors < MaxSimultaneousRenderTargets);
-		for (int32 Index = 0; Index < InNumClearColors; Index++)
-		{
-			ColorArray[Index] = InColorArray[Index];
-		}
+	{		
 	}	
 
 	RHI_API void Execute(FRHICommandListBase& CmdList);
@@ -1552,14 +1536,14 @@ public:
 		new (AllocCommand<FRHICommandSetRenderTargetsAndClear>()) FRHICommandSetRenderTargetsAndClear(RenderTargetsInfo);
 	}	
 
-	FORCEINLINE_DEBUGGABLE void BindClearMRTValues(bool bClearColor, int32 NumClearColors, const FLinearColor* ColorArray, bool bClearDepth, float Depth, bool bClearStencil, uint32 Stencil)
+	FORCEINLINE_DEBUGGABLE void BindClearMRTValues(bool bClearColor, bool bClearDepth, bool bClearStencil)
 	{
 		if (Bypass())
 		{
-			CMD_CONTEXT(BindClearMRTValues)(bClearColor, NumClearColors, ColorArray, bClearDepth, Depth, bClearStencil, Stencil);
+			CMD_CONTEXT(BindClearMRTValues)(bClearColor, bClearDepth, bClearStencil);
 			return;
 		}
-		new (AllocCommand<FRHICommandBindClearMRTValues>()) FRHICommandBindClearMRTValues(bClearColor, NumClearColors, ColorArray, bClearDepth, Depth, bClearStencil, Stencil);
+		new (AllocCommand<FRHICommandBindClearMRTValues>()) FRHICommandBindClearMRTValues(bClearColor, bClearDepth, bClearStencil);
 	}	
 
 	FORCEINLINE_DEBUGGABLE void BeginDrawPrimitiveUP(uint32 PrimitiveType, uint32 NumPrimitives, uint32 NumVertices, uint32 VertexDataStride, void*& OutVertexData)
@@ -2645,7 +2629,7 @@ class RHI_API FRHICommandListExecutor
 public:
 	enum
 	{
-		DefaultBypass = PLATFORM_RHITHREAD_DEFAULT_BYPASS
+		DefaultBypass = 1 //PLATFORM_RHITHREAD_DEFAULT_BYPASS
 	};
 	FRHICommandListExecutor()
 		: bLatchedBypass(!!DefaultBypass)
diff --git a/Engine/Source/Runtime/RHI/Public/RHICommandListCommandExecutes.inl b/Engine/Source/Runtime/RHI/Public/RHICommandListCommandExecutes.inl
--- a/Engine/Source/Runtime/RHI/Public/RHICommandListCommandExecutes.inl
+++ b/Engine/Source/Runtime/RHI/Public/RHICommandListCommandExecutes.inl
@@ -165,12 +165,8 @@ void FRHICommandBindClearMRTValues::Execute(FRHICommandListBase& CmdList)
 	RHISTAT(BindClearMRTValues);
 	INTERNAL_DECORATOR(BindClearMRTValues)(
 		bClearColor,
-		NumClearColors,
-		ColorArray,
 		bClearDepth,
-		Depth,
-		bClearStencil,
-		Stencil	
+		bClearStencil		
 		);
 }
 
diff --git a/Engine/Source/Runtime/RHI/Public/RHIDefinitions.h b/Engine/Source/Runtime/RHI/Public/RHIDefinitions.h
--- a/Engine/Source/Runtime/RHI/Public/RHIDefinitions.h
+++ b/Engine/Source/Runtime/RHI/Public/RHIDefinitions.h
@@ -589,11 +589,8 @@ enum class ESimpleRenderTargetMode
 	EUninitializedColorAndDepth,					// Color = ????, Depth = ????
 	EUninitializedColorExistingDepth,				// Color = ????, Depth = Existing
 	EUninitializedColorClearDepth,					// Color = ????, Depth = Default
-	EClearToDefault,								// Default Color = (0,0,0,0), Default Depth = 0.0f
-	EClearColorToBlack,								// Color = (0,0,0,0), Depth = Existing
-	EClearColorToBlackWithFullAlpha,				// Color = (0,0,0,1), Depth = Existing
-	EClearColorToWhite,								// Color = (1,1,1,1), Depth = Existing
-	EClearDepthToOne,								// Color = Existing, Depth = 1.0
+	EClearColorExistingDepth,						// Clear Color = whatever was bound to the rendertarget at creation time. Depth = Existing
+	EClearColorAndDepth,							// Clear color and depth to bound clear values.
 	EExistingContents_NoDepthStore,					// Load existing contents, but don't store depth out.  depth can be written.
 
 	// If you add an item here, make sure to add it to DecodeRenderTargetMode() as well!
diff --git a/Engine/Source/Runtime/RHI/Public/RHIResources.h b/Engine/Source/Runtime/RHI/Public/RHIResources.h
--- a/Engine/Source/Runtime/RHI/Public/RHIResources.h
+++ b/Engine/Source/Runtime/RHI/Public/RHIResources.h
@@ -295,12 +295,13 @@ class RHI_API FRHITexture : public FRHIResource
 public:
 	
 	/** Initialization constructor. */
-	FRHITexture(uint32 InNumMips,uint32 InNumSamples,EPixelFormat InFormat,uint32 InFlags,FLastRenderTimeContainer* InLastRenderTime)
-	: NumMips(InNumMips)
+	FRHITexture(uint32 InNumMips, uint32 InNumSamples, EPixelFormat InFormat, uint32 InFlags, FLastRenderTimeContainer* InLastRenderTime, const FClearValueBinding& InClearValue)
+	: ClearValue(InClearValue)
+	, NumMips(InNumMips)
 	, NumSamples(InNumSamples)
 	, Format(InFormat)
 	, Flags(InFlags)
-	, LastRenderTime(InLastRenderTime ? *InLastRenderTime : DefaultLastRenderTime)
+	, LastRenderTime(InLastRenderTime ? *InLastRenderTime : DefaultLastRenderTime)	
 	{}
 
 	// Dynamic cast methods.
@@ -377,13 +378,50 @@ public:
 		return TextureName;
 	}
 
+	bool HasClearValue() const
+	{
+		return ClearValue.ColorBinding != EClearBinding::ENoneBound;
+	}
+
+	FLinearColor GetClearColor() const
+	{
+		return ClearValue.GetClearColor();
+	}
+
+	void GetDepthStencilClearValue(float& OutDepth, uint32& OutStencil) const
+	{
+		return ClearValue.GetDepthStencil(OutDepth, OutStencil);
+	}
+
+	float GetDepthClearValue() const
+	{
+		float Depth;
+		uint32 Stencil;
+		ClearValue.GetDepthStencil(Depth, Stencil);
+		return Depth;
+	}
+
+	uint32 GetStencilClearValue() const
+	{
+		float Depth;
+		uint32 Stencil;
+		ClearValue.GetDepthStencil(Depth, Stencil);
+		return Stencil;
+	}
+
+	const FClearValueBinding GetClearBinding() const
+	{
+		return ClearValue;
+	}
+
 private:
+	FClearValueBinding ClearValue;
 	uint32 NumMips;
 	uint32 NumSamples;
 	EPixelFormat Format;
 	uint32 Flags;
 	FLastRenderTimeContainer& LastRenderTime;
-	FLastRenderTimeContainer DefaultLastRenderTime;
+	FLastRenderTimeContainer DefaultLastRenderTime;	
 	FName TextureName;
 };
 
@@ -392,8 +430,8 @@ class RHI_API FRHITexture2D : public FRHITexture
 public:
 	
 	/** Initialization constructor. */
-	FRHITexture2D(uint32 InSizeX,uint32 InSizeY,uint32 InNumMips,uint32 InNumSamples,EPixelFormat InFormat,uint32 InFlags)
-	: FRHITexture(InNumMips,InNumSamples,InFormat,InFlags,NULL)
+	FRHITexture2D(uint32 InSizeX,uint32 InSizeY,uint32 InNumMips,uint32 InNumSamples,EPixelFormat InFormat,uint32 InFlags, const FClearValueBinding& InClearValue)
+	: FRHITexture(InNumMips, InNumSamples, InFormat, InFlags, NULL, InClearValue)
 	, SizeX(InSizeX)
 	, SizeY(InSizeY)
 	{}
@@ -418,8 +456,8 @@ class RHI_API FRHITexture2DArray : public FRHITexture
 public:
 	
 	/** Initialization constructor. */
-	FRHITexture2DArray(uint32 InSizeX,uint32 InSizeY,uint32 InSizeZ,uint32 InNumMips,EPixelFormat InFormat,uint32 InFlags)
-	: FRHITexture(InNumMips,1,InFormat,InFlags,NULL)
+	FRHITexture2DArray(uint32 InSizeX,uint32 InSizeY,uint32 InSizeZ,uint32 InNumMips,EPixelFormat InFormat,uint32 InFlags, const FClearValueBinding& InClearValue)
+	: FRHITexture(InNumMips,1,InFormat,InFlags,NULL, InClearValue)
 	, SizeX(InSizeX)
 	, SizeY(InSizeY)
 	, SizeZ(InSizeZ)
@@ -449,8 +487,8 @@ class RHI_API FRHITexture3D : public FRHITexture
 public:
 	
 	/** Initialization constructor. */
-	FRHITexture3D(uint32 InSizeX,uint32 InSizeY,uint32 InSizeZ,uint32 InNumMips,EPixelFormat InFormat,uint32 InFlags)
-	: FRHITexture(InNumMips,1,InFormat,InFlags,NULL)
+	FRHITexture3D(uint32 InSizeX,uint32 InSizeY,uint32 InSizeZ,uint32 InNumMips,EPixelFormat InFormat,uint32 InFlags, const FClearValueBinding& InClearValue)
+	: FRHITexture(InNumMips,1,InFormat,InFlags,NULL, InClearValue)
 	, SizeX(InSizeX)
 	, SizeY(InSizeY)
 	, SizeZ(InSizeZ)
@@ -480,8 +518,8 @@ class RHI_API FRHITextureCube : public FRHITexture
 public:
 	
 	/** Initialization constructor. */
-	FRHITextureCube(uint32 InSize,uint32 InNumMips,EPixelFormat InFormat,uint32 InFlags)
-	: FRHITexture(InNumMips,1,InFormat,InFlags,NULL)
+	FRHITextureCube(uint32 InSize,uint32 InNumMips,EPixelFormat InFormat,uint32 InFlags, const FClearValueBinding& InClearValue)
+	: FRHITexture(InNumMips,1,InFormat,InFlags,NULL, InClearValue)
 	, Size(InSize)
 	{}
 	
@@ -500,7 +538,7 @@ class RHI_API FRHITextureReference : public FRHITexture
 {
 public:
 	explicit FRHITextureReference(FLastRenderTimeContainer* InLastRenderTime)
-		: FRHITexture(0,0,PF_Unknown,0,InLastRenderTime)
+		: FRHITexture(0,0,PF_Unknown,0,InLastRenderTime, FClearValueBinding())
 	{}
 
 	virtual FRHITextureReference* GetTextureReference() override { return this; }
@@ -980,23 +1018,18 @@ class FRHISetRenderTargetsInfo
 {
 public:
 	// Color Render Targets Info
-	FRHIRenderTargetView ColorRenderTarget[MaxSimultaneousRenderTargets];
-	FLinearColor ClearColors[MaxSimultaneousRenderTargets];
+	FRHIRenderTargetView ColorRenderTarget[MaxSimultaneousRenderTargets];	
 	int32 NumColorRenderTargets;
 	bool bClearColor;
 
 	// Depth/Stencil Render Target Info
-	FRHIDepthRenderTargetView DepthStencilRenderTarget;
-	float DepthClearValue;
-	uint32 StencilClearValue;
+	FRHIDepthRenderTargetView DepthStencilRenderTarget;	
 	bool bClearDepth;
 	bool bClearStencil;
 
 	FRHISetRenderTargetsInfo() :
 		NumColorRenderTargets(0),
-		bClearColor(false),
-		DepthClearValue((float)ERHIZBuffer::FarPlane),
-		StencilClearValue(0),
+		bClearColor(false),		
 		bClearDepth(false),
 		bClearStencil(false)
 	{}
@@ -1004,21 +1037,18 @@ public:
 	FRHISetRenderTargetsInfo(int32 InNumColorRenderTargets, const FRHIRenderTargetView* InColorRenderTargets, const FRHIDepthRenderTargetView& InDepthStencilRenderTarget) :
 		NumColorRenderTargets(InNumColorRenderTargets),
 		bClearColor(InNumColorRenderTargets > 0 && InColorRenderTargets[0].LoadAction == ERenderTargetLoadAction::EClear),
-		DepthStencilRenderTarget(InDepthStencilRenderTarget),
-		DepthClearValue((float)ERHIZBuffer::FarPlane),
-		StencilClearValue(0),
+		DepthStencilRenderTarget(InDepthStencilRenderTarget),		
 		bClearDepth(InDepthStencilRenderTarget.Texture && InDepthStencilRenderTarget.DepthLoadAction == ERenderTargetLoadAction::EClear),
 		bClearStencil(InDepthStencilRenderTarget.Texture && InDepthStencilRenderTarget.StencilLoadAction == ERenderTargetLoadAction::EClear)
 	{
 		check(InNumColorRenderTargets <= 0 || InColorRenderTargets);
 		for (int32 Index = 0; Index < InNumColorRenderTargets; ++Index)
 		{
-			ColorRenderTarget[Index] = InColorRenderTargets[Index];
-			ClearColors[Index] = FLinearColor(0, 0, 0, 1);
+			ColorRenderTarget[Index] = InColorRenderTargets[Index];			
 		}
 	}
 	// @todo metal mrt: This can go away after all the cleanup is done
-	void SetClearDepthStencil(bool bInClearDepth, float InDepthClear, bool bInClearStencil = false, uint32 InClearStencil = 0)
+	void SetClearDepthStencil(bool bInClearDepth, bool bInClearStencil = false)
 	{
 		if (bInClearDepth)
 		{
@@ -1028,10 +1058,8 @@ public:
 		{
 			DepthStencilRenderTarget.StencilLoadAction = ERenderTargetLoadAction::EClear;
 		}
-		bClearDepth = bInClearDepth;
-		DepthClearValue = InDepthClear;
-		bClearStencil = bInClearStencil;
-		StencilClearValue = InClearStencil;
+		bClearDepth = bInClearDepth;		
+		bClearStencil = bInClearStencil;		
 	}
 };
 
diff --git a/Engine/Source/Runtime/RHI/Public/RHIUtilities.h b/Engine/Source/Runtime/RHI/Public/RHIUtilities.h
--- a/Engine/Source/Runtime/RHI/Public/RHIUtilities.h
+++ b/Engine/Source/Runtime/RHI/Public/RHIUtilities.h
@@ -123,13 +123,11 @@ struct FRWBufferByteAddress
  * Convert the ESimpleRenderTargetMode into usable values 
  * @todo: Can we easily put this into a .cpp somewhere?
  */
-inline void DecodeRenderTargetMode(ESimpleRenderTargetMode Mode, ERenderTargetLoadAction& ColorLoadAction, ERenderTargetStoreAction& ColorStoreAction, ERenderTargetLoadAction& DepthLoadAction, ERenderTargetStoreAction& DepthStoreAction, FLinearColor& ClearColor, float& ClearDepth)
+inline void DecodeRenderTargetMode(ESimpleRenderTargetMode Mode, ERenderTargetLoadAction& ColorLoadAction, ERenderTargetStoreAction& ColorStoreAction, ERenderTargetLoadAction& DepthLoadAction, ERenderTargetStoreAction& DepthStoreAction)
 {
 	// set defaults
 	ColorStoreAction = ERenderTargetStoreAction::EStore;
-	DepthStoreAction = ERenderTargetStoreAction::EStore;
-	ClearColor = FLinearColor(0, 0, 0, 0);
-	ClearDepth = 0.0f;
+	DepthStoreAction = ERenderTargetStoreAction::EStore;	
 
 	switch (Mode)
 	{
@@ -149,29 +147,13 @@ inline void DecodeRenderTargetMode(ESimpleRenderTargetMode Mode, ERenderTargetLo
 		ColorLoadAction = ERenderTargetLoadAction::ENoAction;
 		DepthLoadAction = ERenderTargetLoadAction::EClear;
 		break;
-	case ESimpleRenderTargetMode::EClearToDefault:
-		ColorLoadAction = ERenderTargetLoadAction::EClear;
-		DepthLoadAction = ERenderTargetLoadAction::EClear;
-		break;
-	case ESimpleRenderTargetMode::EClearColorToBlack:
-		ColorLoadAction = ERenderTargetLoadAction::EClear;
-		DepthLoadAction = ERenderTargetLoadAction::ELoad;
-		ClearColor = FLinearColor::Black;
-		break;
-	case ESimpleRenderTargetMode::EClearColorToBlackWithFullAlpha:
+	case ESimpleRenderTargetMode::EClearColorExistingDepth:
 		ColorLoadAction = ERenderTargetLoadAction::EClear;
 		DepthLoadAction = ERenderTargetLoadAction::ELoad;
-		ClearColor = FLinearColor(0, 0, 0, 1);
 		break;
-	case ESimpleRenderTargetMode::EClearColorToWhite:
+	case ESimpleRenderTargetMode::EClearColorAndDepth:
 		ColorLoadAction = ERenderTargetLoadAction::EClear;
-		DepthLoadAction = ERenderTargetLoadAction::ELoad;
-		ClearColor = FLinearColor::White;
-		break;
-	case ESimpleRenderTargetMode::EClearDepthToOne:
-		ColorLoadAction = ERenderTargetLoadAction::ELoad;
 		DepthLoadAction = ERenderTargetLoadAction::EClear;
-		ClearDepth = 1.0f;
 		break;
 	case ESimpleRenderTargetMode::EExistingContents_NoDepthStore:
 		ColorLoadAction = ERenderTargetLoadAction::ELoad;
@@ -195,16 +177,12 @@ inline void SetRenderTarget(FRHICommandList& RHICmdList, FTextureRHIParamRef New
 inline void SetRenderTarget(FRHICommandList& RHICmdList, FTextureRHIParamRef NewRenderTarget, FTextureRHIParamRef NewDepthStencilTarget, ESimpleRenderTargetMode Mode, FExclusiveDepthStencil DepthStencilAccess = FExclusiveDepthStencil::DepthWrite_StencilWrite)
 {
 	ERenderTargetLoadAction ColorLoadAction, DepthLoadAction;
-	ERenderTargetStoreAction ColorStoreAction, DepthStoreAction;
-	FLinearColor ClearColor;
-	float ClearDepth;
-	DecodeRenderTargetMode(Mode, ColorLoadAction, ColorStoreAction, DepthLoadAction, DepthStoreAction, ClearColor, ClearDepth);
+	ERenderTargetStoreAction ColorStoreAction, DepthStoreAction;	
+	DecodeRenderTargetMode(Mode, ColorLoadAction, ColorStoreAction, DepthLoadAction, DepthStoreAction);
 
 	// now make the FRHISetRenderTargetsInfo that encapsulates all of the info
 	FRHIRenderTargetView ColorView(NewRenderTarget, 0, -1, ColorLoadAction, ColorStoreAction);
 	FRHISetRenderTargetsInfo Info(1, &ColorView, FRHIDepthRenderTargetView(NewDepthStencilTarget, DepthLoadAction, DepthStoreAction, DepthStencilAccess));
-	Info.ClearColors[0] = ClearColor;
-	Info.DepthClearValue = ClearDepth;
 	RHICmdList.SetRenderTargetsAndClear(Info);
 }
 
@@ -256,10 +234,8 @@ inline void SetRenderTargets(
 	)
 {
 	ERenderTargetLoadAction ColorLoadAction, DepthLoadAction;
-	ERenderTargetStoreAction ColorStoreAction, DepthStoreAction;
-	FLinearColor ClearColor;
-	float ClearDepth;
-	DecodeRenderTargetMode(Mode, ColorLoadAction, ColorStoreAction, DepthLoadAction, DepthStoreAction, ClearColor, ClearDepth);
+	ERenderTargetStoreAction ColorStoreAction, DepthStoreAction;	
+	DecodeRenderTargetMode(Mode, ColorLoadAction, ColorStoreAction, DepthLoadAction, DepthStoreAction);
 
 	FRHIRenderTargetView RTVs[MaxSimultaneousRenderTargets];
 
@@ -283,7 +259,7 @@ inline void SetRenderTargets(
 inline void RHICreateTargetableShaderResource2D(
 	uint32 SizeX,
 	uint32 SizeY,
-	uint8 Format,
+	uint8 Format,	
 	uint32 NumMips,
 	uint32 Flags,
 	uint32 TargetableTextureFlags,
diff --git a/Engine/Source/Runtime/Renderer/Private/AtmosphereRendering.cpp b/Engine/Source/Runtime/Renderer/Private/AtmosphereRendering.cpp
--- a/Engine/Source/Runtime/Renderer/Private/AtmosphereRendering.cpp
+++ b/Engine/Source/Runtime/Renderer/Private/AtmosphereRendering.cpp
@@ -1158,9 +1158,8 @@ void FAtmosphericFogSceneInfo::RenderAtmosphereShaders(FRHICommandList& RHICmdLi
 	case AP_ClearIrradiance:
 		{
 			const FSceneRenderTargetItem& DestRenderTarget = AtmosphereTextures->AtmosphereIrradiance->GetRenderTargetItem();
-			SetRenderTarget(RHICmdList, DestRenderTarget.TargetableTexture, FTextureRHIRef());
-
-			RHICmdList.Clear(true, FLinearColor::Black, false, 0.0f, false, 0, FIntRect());
+			ensure(DestRenderTarget.TargetableTexture->GetClearColor() == FLinearColor::Black);
+			SetRenderTarget(RHICmdList, DestRenderTarget.TargetableTexture, FTextureRHIRef(), ESimpleRenderTargetMode::EClearColorExistingDepth);
 			RHICmdList.CopyToResolveTarget(DestRenderTarget.TargetableTexture, DestRenderTarget.ShaderResourceTexture, true, FResolveParams());
 		}
 		break;
diff --git a/Engine/Source/Runtime/Renderer/Private/AtmosphereTextures.cpp b/Engine/Source/Runtime/Renderer/Private/AtmosphereTextures.cpp
--- a/Engine/Source/Runtime/Renderer/Private/AtmosphereTextures.cpp
+++ b/Engine/Source/Runtime/Renderer/Private/AtmosphereTextures.cpp
@@ -16,7 +16,7 @@ void FAtmosphereTextures::InitDynamicRHI()
 		// todo: Expose
 		// Transmittance
 		FIntPoint GTransmittanceTexSize(PrecomputeParams->TransmittanceTexWidth, PrecomputeParams->TransmittanceTexHeight);
-		FPooledRenderTargetDesc TransmittanceDesc(FPooledRenderTargetDesc::Create2DDesc(GTransmittanceTexSize, PF_FloatRGBA, TexCreate_None, TexCreate_RenderTargetable, false));
+		FPooledRenderTargetDesc TransmittanceDesc(FPooledRenderTargetDesc::Create2DDesc(GTransmittanceTexSize, PF_FloatRGBA, FClearValueBinding::None, TexCreate_None, TexCreate_RenderTargetable, false));
 		GRenderTargetPool.FindFreeElement(TransmittanceDesc, AtmosphereTransmittance, TEXT("AtmosphereTransmittance"));
 
 		FRHICommandListImmediate& RHICmdList = FRHICommandListExecutor::GetImmediateCommandList();
@@ -27,7 +27,7 @@ void FAtmosphereTextures::InitDynamicRHI()
 
 		// Irradiance
 		FIntPoint GIrradianceTexSize(PrecomputeParams->IrradianceTexWidth, PrecomputeParams->IrradianceTexHeight);
-		FPooledRenderTargetDesc IrradianceDesc(FPooledRenderTargetDesc::Create2DDesc(GIrradianceTexSize, PF_FloatRGBA, TexCreate_None, TexCreate_RenderTargetable, false));
+		FPooledRenderTargetDesc IrradianceDesc(FPooledRenderTargetDesc::Create2DDesc(GIrradianceTexSize, PF_FloatRGBA, FClearValueBinding::Black, TexCreate_None, TexCreate_RenderTargetable, false));
 		GRenderTargetPool.FindFreeElement(IrradianceDesc, AtmosphereIrradiance, TEXT("AtmosphereIrradiance"));
 
 		SetRenderTarget(RHICmdList, AtmosphereIrradiance->GetRenderTargetItem().TargetableTexture, FTextureRHIRef());
@@ -39,7 +39,7 @@ void FAtmosphereTextures::InitDynamicRHI()
 
 		// 3D Texture
 		// Inscatter
-		FPooledRenderTargetDesc InscatterDesc(FPooledRenderTargetDesc::CreateVolumeDesc(PrecomputeParams->InscatterMuSNum * PrecomputeParams->InscatterNuNum, PrecomputeParams->InscatterMuNum, PrecomputeParams->InscatterAltitudeSampleNum, PF_FloatRGBA, TexCreate_None, TexCreate_ShaderResource | TexCreate_RenderTargetable, false));
+		FPooledRenderTargetDesc InscatterDesc(FPooledRenderTargetDesc::CreateVolumeDesc(PrecomputeParams->InscatterMuSNum * PrecomputeParams->InscatterNuNum, PrecomputeParams->InscatterMuNum, PrecomputeParams->InscatterAltitudeSampleNum, PF_FloatRGBA, FClearValueBinding::None, TexCreate_None, TexCreate_ShaderResource | TexCreate_RenderTargetable, false));
 		GRenderTargetPool.FindFreeElement(InscatterDesc, AtmosphereInscatter, TEXT("AtmosphereInscatter"));
 
 		// DeltaSR
diff --git a/Engine/Source/Runtime/Renderer/Private/ClearQuad.cpp b/Engine/Source/Runtime/Renderer/Private/ClearQuad.cpp
--- a/Engine/Source/Runtime/Renderer/Private/ClearQuad.cpp
+++ b/Engine/Source/Runtime/Renderer/Private/ClearQuad.cpp
@@ -98,16 +98,8 @@ void DrawClearQuadMRT(FRHICommandList& RHICmdList, ERHIFeatureLevel::Type Featur
 		PixelShader = *MRTPixelShader;
 	}
 
-	SetGlobalBoundShaderState(RHICmdList, FeatureLevel, GClearMRTBoundShaderState[FMath::Max(NumClearColors - 1, 0)], GetVertexDeclarationFVector4(), *VertexShader, PixelShader);
-	FLinearColor ShaderClearColors[MaxSimultaneousRenderTargets];
-	FMemory::Memzero(ShaderClearColors);
-
-	for (int32 i = 0; i < NumClearColors; i++)
-	{
-		ShaderClearColors[i] = ClearColorArray[i];
-	}
-
-	SetShaderValueArray(RHICmdList, PixelShader->GetPixelShader(),PixelShader->ColorParameter,ShaderClearColors,NumClearColors);
+	SetGlobalBoundShaderState(RHICmdList, FeatureLevel, GClearMRTBoundShaderState[FMath::Max(NumClearColors - 1, 0)], GetVertexDeclarationFVector4(), *VertexShader, PixelShader);	
+	PixelShader->SetColors(RHICmdList, ClearColorArray, NumClearColors);
 		
 	{
 		// Draw a fullscreen quad
diff --git a/Engine/Source/Runtime/Renderer/Private/CompositionLighting/PostProcessAmbientOcclusion.cpp b/Engine/Source/Runtime/Renderer/Private/CompositionLighting/PostProcessAmbientOcclusion.cpp
--- a/Engine/Source/Runtime/Renderer/Private/CompositionLighting/PostProcessAmbientOcclusion.cpp
+++ b/Engine/Source/Runtime/Renderer/Private/CompositionLighting/PostProcessAmbientOcclusion.cpp
@@ -371,6 +371,7 @@ FPooledRenderTargetDesc FRCPassPostProcessAmbientOcclusionSetup::ComputeOutputDe
 
 	Ret.Reset();
 	Ret.Format = PF_FloatRGBA;
+	Ret.ClearValue = FClearValueBinding::None;
 	Ret.TargetableFlags &= ~TexCreate_DepthStencilTargetable;
 	Ret.TargetableFlags |= TexCreate_RenderTargetable;
 	Ret.Extent = FIntPoint::DivideAndRoundUp(Ret.Extent, 2);
diff --git a/Engine/Source/Runtime/Renderer/Private/CompositionLighting/PostProcessDeferredDecals.cpp b/Engine/Source/Runtime/Renderer/Private/CompositionLighting/PostProcessDeferredDecals.cpp
--- a/Engine/Source/Runtime/Renderer/Private/CompositionLighting/PostProcessDeferredDecals.cpp
+++ b/Engine/Source/Runtime/Renderer/Private/CompositionLighting/PostProcessDeferredDecals.cpp
@@ -570,18 +570,21 @@ void FRCPassPostProcessDeferredDecals::Process(FRenderingCompositePassContext& C
 
 		// DBuffer: Decal buffer
 		FPooledRenderTargetDesc Desc(FPooledRenderTargetDesc::Create2DDesc(SceneContext.GBufferA->GetDesc().Extent, 
-			PF_B8G8R8A8, 
+			PF_B8G8R8A8,
+			FClearValueBinding::None,
 			TexCreate_None, 
 			TexCreate_ShaderResource | TexCreate_RenderTargetable,
 			false));
 
 		if(!SceneContext.DBufferA)
 		{
+			Desc.ClearValue = FClearValueBinding::Black;
 			GRenderTargetPool.FindFreeElement(Desc, SceneContext.DBufferA, TEXT("DBufferA"));
 		}
 
 		if(!SceneContext.DBufferB)
 		{
+			Desc.ClearValue = FClearValueBinding(FLinearColor(128.0f / 255.0f, 128.0f / 255.0f, 128.0f / 255.0f, 1));
 			GRenderTargetPool.FindFreeElement(Desc, SceneContext.DBufferB, TEXT("DBufferB"));
 		}
 
@@ -589,6 +592,7 @@ void FRCPassPostProcessDeferredDecals::Process(FRenderingCompositePassContext& C
 
 		if(!SceneContext.DBufferC)
 		{
+			Desc.ClearValue = FClearValueBinding(FLinearColor(0, 1, 0, 1));
 			GRenderTargetPool.FindFreeElement(Desc, SceneContext.DBufferC, TEXT("DBufferC"));
 		}
 
@@ -609,10 +613,6 @@ void FRCPassPostProcessDeferredDecals::Process(FRenderingCompositePassContext& C
 			FRHIDepthRenderTargetView DepthView(SceneContext.GetSceneDepthSurface(), ERenderTargetLoadAction::ELoad, ERenderTargetStoreAction::ENoAction, ERenderTargetLoadAction::ELoad, ERenderTargetStoreAction::ENoAction, FExclusiveDepthStencil(FExclusiveDepthStencil::DepthRead_StencilWrite));
 
 			FRHISetRenderTargetsInfo Info(3, RenderTargets, DepthView);
-			Info.ClearColors[0] = FLinearColor(0, 0, 0, 1);
-			Info.ClearColors[1] = FLinearColor(128.0f / 255.0f, 128.0f / 255.0f, 128.0f / 255.0f, 1);
-			Info.ClearColors[2] = FLinearColor(0, 1, 0, 1);
-
 			RHICmdList.SetRenderTargetsAndClear(Info);
 
 			TargetsToResolve[DBufferAIndex] = SceneContext.DBufferA->GetRenderTargetItem().TargetableTexture;
diff --git a/Engine/Source/Runtime/Renderer/Private/DeferredShadingRenderer.cpp b/Engine/Source/Runtime/Renderer/Private/DeferredShadingRenderer.cpp
--- a/Engine/Source/Runtime/Renderer/Private/DeferredShadingRenderer.cpp
+++ b/Engine/Source/Runtime/Renderer/Private/DeferredShadingRenderer.cpp
@@ -116,8 +116,8 @@ void FDeferredShadingSceneRenderer::ClearGBufferAtMaxZ(FRHICommandList& RHICmdLi
 	// Clear the G Buffer render targets
 	const bool bClearBlack = Views[0].Family->EngineShowFlags.ShaderComplexity || Views[0].Family->EngineShowFlags.StationaryLightOverlap;
 	// Same clear color from RHIClearMRT
-	FLinearColor ClearColors[6] = 
-		{bClearBlack ? FLinearColor(0,0,0,0) : Views[0].BackgroundColor, FLinearColor(0.5f,0.5f,0.5f,0), FLinearColor(0,0,0,1), FLinearColor(0,0,0,0), FLinearColor(0,1,1,1), FLinearColor(1,1,1,1)};
+	FLinearColor ClearColors[MaxSimultaneousRenderTargets] = 
+		{bClearBlack ? FLinearColor(0,0,0,0) : Views[0].BackgroundColor, FLinearColor(0.5f,0.5f,0.5f,0), FLinearColor(0,0,0,1), FLinearColor(0,0,0,0), FLinearColor(0,1,1,1), FLinearColor(1,1,1,1), FLinearColor::Transparent, FLinearColor::Transparent};
 
 	uint32 NumActiveRenderTargets = FSceneRenderTargets::Get(RHICmdList).GetNumGBufferTargets();
 	
@@ -168,7 +168,7 @@ void FDeferredShadingSceneRenderer::ClearGBufferAtMaxZ(FRHICommandList& RHICmdLi
 		RHICmdList.SetViewport(View.ViewRect.Min.X, View.ViewRect.Min.Y, 0, View.ViewRect.Max.X, View.ViewRect.Max.Y, 1);
 
 		// Setup PS
-		SetShaderValueArray(RHICmdList, PixelShader->GetPixelShader(),PixelShader->ColorParameter, ClearColors, NumActiveRenderTargets);
+		PixelShader->SetColors(RHICmdList, ClearColors, NumActiveRenderTargets);
 
 		// Render quad
 		static const FVector4 ClearQuadVertices[4] = 
@@ -602,7 +602,7 @@ void FDeferredShadingSceneRenderer::RenderVisualizeTexturePool(FRHICommandListIm
 		TexturePoolVisualizerSizeY = 140,
 	};
 
-	FPooledRenderTargetDesc Desc(FPooledRenderTargetDesc::Create2DDesc(FIntPoint(TexturePoolVisualizerSizeX, TexturePoolVisualizerSizeY), PF_B8G8R8A8, TexCreate_None, TexCreate_None, false));
+	FPooledRenderTargetDesc Desc(FPooledRenderTargetDesc::Create2DDesc(FIntPoint(TexturePoolVisualizerSizeX, TexturePoolVisualizerSizeY), PF_B8G8R8A8, FClearValueBinding::None, TexCreate_None, TexCreate_None, false));
 	GRenderTargetPool.FindFreeElement(Desc, VisualizeTexturePool, TEXT("VisualizeTexturePool"));
 	
 	uint32 Pitch;
diff --git a/Engine/Source/Runtime/Renderer/Private/DistanceFieldLightingPost.cpp b/Engine/Source/Runtime/Renderer/Private/DistanceFieldLightingPost.cpp
--- a/Engine/Source/Runtime/Renderer/Private/DistanceFieldLightingPost.cpp
+++ b/Engine/Source/Runtime/Renderer/Private/DistanceFieldLightingPost.cpp
@@ -548,7 +548,7 @@ void AllocateOrReuseAORenderTarget(TRefCountPtr<IPooledRenderTarget>& Target, co
 	{
 		FIntPoint BufferSize = GetBufferSizeForAO();
 
-		FPooledRenderTargetDesc Desc(FPooledRenderTargetDesc::Create2DDesc(BufferSize, bWithAlpha ? PF_FloatRGBA : PF_FloatRGB, TexCreate_None, TexCreate_RenderTargetable | TexCreate_UAV, false));
+		FPooledRenderTargetDesc Desc(FPooledRenderTargetDesc::Create2DDesc(BufferSize, bWithAlpha ? PF_FloatRGBA : PF_FloatRGB, FClearValueBinding::None, TexCreate_None, TexCreate_RenderTargetable | TexCreate_UAV, false));
 		GRenderTargetPool.FindFreeElement(Desc, Target, Name);
 	}
 }
diff --git a/Engine/Source/Runtime/Renderer/Private/DistanceFieldScreenGridLighting.cpp b/Engine/Source/Runtime/Renderer/Private/DistanceFieldScreenGridLighting.cpp
--- a/Engine/Source/Runtime/Renderer/Private/DistanceFieldScreenGridLighting.cpp
+++ b/Engine/Source/Runtime/Renderer/Private/DistanceFieldScreenGridLighting.cpp
@@ -850,7 +850,7 @@ void FDeferredShadingSceneRenderer::RenderDistanceFieldAOScreenGrid(
 	if (bUseDistanceFieldGI)
 	{
 		{
-			FPooledRenderTargetDesc Desc(FPooledRenderTargetDesc::Create2DDesc(ConeTraceBufferSize, PF_FloatRGBA, TexCreate_None, TexCreate_RenderTargetable | TexCreate_UAV, false));
+			FPooledRenderTargetDesc Desc(FPooledRenderTargetDesc::Create2DDesc(ConeTraceBufferSize, PF_FloatRGBA, FClearValueBinding::None, TexCreate_None, TexCreate_RenderTargetable | TexCreate_UAV, false));
 			GRenderTargetPool.FindFreeElement(Desc, DownsampledIrradiance, TEXT("DownsampledIrradiance"));
 		}
 
@@ -872,7 +872,7 @@ void FDeferredShadingSceneRenderer::RenderDistanceFieldAOScreenGrid(
 	TRefCountPtr<IPooledRenderTarget> DownsampledBentNormal;
 
 	{
-		FPooledRenderTargetDesc Desc(FPooledRenderTargetDesc::Create2DDesc(ConeTraceBufferSize, PF_FloatRGBA, TexCreate_None, TexCreate_RenderTargetable | TexCreate_UAV, false));
+		FPooledRenderTargetDesc Desc(FPooledRenderTargetDesc::Create2DDesc(ConeTraceBufferSize, PF_FloatRGBA, FClearValueBinding::None, TexCreate_None, TexCreate_RenderTargetable | TexCreate_UAV, false));
 		GRenderTargetPool.FindFreeElement(Desc, DownsampledBentNormal, TEXT("DownsampledBentNormal"));
 	}
 
diff --git a/Engine/Source/Runtime/Renderer/Private/DistanceFieldShadowing.cpp b/Engine/Source/Runtime/Renderer/Private/DistanceFieldShadowing.cpp
--- a/Engine/Source/Runtime/Renderer/Private/DistanceFieldShadowing.cpp
+++ b/Engine/Source/Runtime/Renderer/Private/DistanceFieldShadowing.cpp
@@ -825,7 +825,7 @@ void FProjectedShadowInfo::RenderRayTracedDistanceFieldProjection(FRHICommandLis
 
 			{
 				const FIntPoint BufferSize = GetBufferSizeForDFShadows();
-				FPooledRenderTargetDesc Desc(FPooledRenderTargetDesc::Create2DDesc(BufferSize, PF_G16R16F, TexCreate_None, TexCreate_RenderTargetable | TexCreate_UAV, false));
+				FPooledRenderTargetDesc Desc(FPooledRenderTargetDesc::Create2DDesc(BufferSize, PF_G16R16F, FClearValueBinding::None, TexCreate_None, TexCreate_RenderTargetable | TexCreate_UAV, false));
 				GRenderTargetPool.FindFreeElement(Desc, RayTracedShadowsRT, TEXT("RayTracedShadows"));
 			}
 
diff --git a/Engine/Source/Runtime/Renderer/Private/DistanceFieldSurfaceCacheLighting.cpp b/Engine/Source/Runtime/Renderer/Private/DistanceFieldSurfaceCacheLighting.cpp
--- a/Engine/Source/Runtime/Renderer/Private/DistanceFieldSurfaceCacheLighting.cpp
+++ b/Engine/Source/Runtime/Renderer/Private/DistanceFieldSurfaceCacheLighting.cpp
@@ -2302,8 +2302,7 @@ void SetupDepthStencil(
 
 		static FGlobalBoundShaderState BoundShaderState;
 		SetGlobalBoundShaderState(RHICmdList, View.GetFeatureLevel(), BoundShaderState, GetVertexDeclarationFVector4(), *VertexShader, *PixelShader);
-
-		SetShaderValue(RHICmdList, PixelShader->GetPixelShader(), PixelShader->ColorParameter, FLinearColor::Black);
+		PixelShader->SetColors(RHICmdList, &FLinearColor::Black, 1);
 
 		FVector ViewSpaceMaxDistance(0.0f, 0.0f, GetMaxAOViewDistance());
 		FVector4 ClipSpaceMaxDistance = View.ViewMatrices.ProjMatrix.TransformPosition(ViewSpaceMaxDistance);
@@ -2344,7 +2343,7 @@ void RenderIrradianceCacheInterpolation(
 
 		if (bFinalInterpolation && (GAOInterpolationDepthTesting || GAOInterpolationStencilTesting))
 		{
-			FPooledRenderTargetDesc Desc(FPooledRenderTargetDesc::Create2DDesc(BentNormalInterpolationTarget->GetDesc().Extent, PF_DepthStencil, TexCreate_None, TexCreate_DepthStencilTargetable, false));
+			FPooledRenderTargetDesc Desc(FPooledRenderTargetDesc::Create2DDesc(BentNormalInterpolationTarget->GetDesc().Extent, PF_DepthStencil, FClearValueBinding::DepthZero, TexCreate_None, TexCreate_DepthStencilTargetable, false));
 			GRenderTargetPool.FindFreeElement(Desc, SplatDepthStencilBuffer, TEXT("DistanceFieldAOSplatDepthBuffer"));
 
 			SetupDepthStencil(RHICmdList, View, SplatDepthStencilBuffer->GetRenderTargetItem(), DistanceFieldNormal, DepthLevel, DestLevelDownsampleFactor);
@@ -2359,7 +2358,7 @@ void RenderIrradianceCacheInterpolation(
 		FTextureRHIParamRef DepthBuffer = SplatDepthStencilBuffer ? SplatDepthStencilBuffer->GetRenderTargetItem().TargetableTexture : NULL;
 
 		// PS4 fast clear which is only triggered through SetRenderTargetsAndClear adds 3ms of GPU idle time currently
-		static bool bUseSetAndClear = false;
+		static bool bUseSetAndClear = true;
 
 		if (bUseSetAndClear)
 		{
@@ -2372,13 +2371,16 @@ void RenderIrradianceCacheInterpolation(
 			FRHIDepthRenderTargetView DepthView(DepthBuffer);
 			FRHISetRenderTargetsInfo Info(MRTCount, RenderTargets, DepthView);
 
+			check(RenderTargets[0].Texture->GetClearColor() == FLinearColor::Transparent);
+			check(RenderTargets[1].Texture->GetClearColor() == FLinearColor::Transparent);
 			Info.bClearColor = true;
-			Info.ClearColors[0] = FLinearColor(0, 0, 0, 0);
-			Info.ClearColors[1] = FLinearColor(0, 0, 0, 0);
+						
 
 			// set the render target
 			RHICmdList.SetRenderTargetsAndClear(Info);
 		}
+		//we need to prefer fast clears.  If PS4 is seeing 3ms of latency we should investigate soon.
+#if 0
 		else
 		{
 			FTextureRHIParamRef RenderTargets[2] =
@@ -2392,6 +2394,7 @@ void RenderIrradianceCacheInterpolation(
 			FLinearColor ClearColors[2] = {FLinearColor(0, 0, 0, 0), FLinearColor(0, 0, 0, 0)};
 			RHICmdList.ClearMRT(true, ARRAY_COUNT(ClearColors) - (bBindIrradiance ? 0 : 1), ClearColors, false, 0, false, 0, FIntRect());
 		}
+#endif
 
 		static int32 NumInterpolationSections = 8;
 		if (GCircleVertexBuffer.NumSections != NumInterpolationSections)
@@ -2745,7 +2748,7 @@ void RenderDistanceFieldAOSurfaceCache(
 
 		{
 			FIntPoint AOBufferSize = FIntPoint::DivideAndRoundUp(SceneContext.GetBufferSizeXY(), DestLevelDownsampleFactor);
-			FPooledRenderTargetDesc Desc(FPooledRenderTargetDesc::Create2DDesc(AOBufferSize, PF_FloatRGBA, TexCreate_None, TexCreate_RenderTargetable | TexCreate_UAV, false));
+			FPooledRenderTargetDesc Desc(FPooledRenderTargetDesc::Create2DDesc(AOBufferSize, PF_FloatRGBA, FClearValueBinding::Transparent, TexCreate_None, TexCreate_RenderTargetable | TexCreate_UAV, false));
 			GRenderTargetPool.FindFreeElement(Desc, DistanceFieldAOBentNormalSplat, TEXT("DistanceFieldAOBentNormalSplat"));
 		}
 
@@ -2858,7 +2861,7 @@ void RenderDistanceFieldAOSurfaceCache(
 
 	{
 		FIntPoint BufferSize = GetBufferSizeForAO();
-		FPooledRenderTargetDesc Desc(FPooledRenderTargetDesc::Create2DDesc(BufferSize, PF_FloatRGBA, TexCreate_None, TexCreate_RenderTargetable | TexCreate_UAV, false));
+		FPooledRenderTargetDesc Desc(FPooledRenderTargetDesc::Create2DDesc(BufferSize, PF_FloatRGBA, FClearValueBinding::Transparent, TexCreate_None, TexCreate_RenderTargetable | TexCreate_UAV, false));
 		GRenderTargetPool.FindFreeElement(Desc, BentNormalAccumulation, TEXT("BentNormalAccumulation"));
 
 		if (bUseDistanceFieldGI)
@@ -2953,7 +2956,7 @@ bool FDeferredShadingSceneRenderer::RenderDistanceFieldLighting(
 
 			{
 				const FIntPoint BufferSize = GetBufferSizeForAO();
-				FPooledRenderTargetDesc Desc(FPooledRenderTargetDesc::Create2DDesc(BufferSize, PF_FloatRGBA, TexCreate_None, TexCreate_RenderTargetable | TexCreate_UAV, false));
+				FPooledRenderTargetDesc Desc(FPooledRenderTargetDesc::Create2DDesc(BufferSize, PF_FloatRGBA, FClearValueBinding::Transparent, TexCreate_None, TexCreate_RenderTargetable | TexCreate_UAV, false));
 				GRenderTargetPool.FindFreeElement(Desc, DistanceFieldNormal, TEXT("DistanceFieldNormal"));
 			}
 
@@ -3439,7 +3442,7 @@ void FDeferredShadingSceneRenderer::RenderMeshDistanceFieldVisualization(FRHICom
 
 			{
 				const FIntPoint BufferSize = GetBufferSizeForAO();
-				FPooledRenderTargetDesc Desc(FPooledRenderTargetDesc::Create2DDesc(BufferSize, PF_FloatRGBA, TexCreate_None, TexCreate_RenderTargetable | TexCreate_UAV, false));
+				FPooledRenderTargetDesc Desc(FPooledRenderTargetDesc::Create2DDesc(BufferSize, PF_FloatRGBA, FClearValueBinding::None, TexCreate_None, TexCreate_RenderTargetable | TexCreate_UAV, false));
 				GRenderTargetPool.FindFreeElement(Desc, VisualizeResultRT, TEXT("VisualizeDistanceField"));
 			}
 
diff --git a/Engine/Source/Runtime/Renderer/Private/DistortionRendering.cpp b/Engine/Source/Runtime/Renderer/Private/DistortionRendering.cpp
--- a/Engine/Source/Runtime/Renderer/Private/DistortionRendering.cpp
+++ b/Engine/Source/Runtime/Renderer/Private/DistortionRendering.cpp
@@ -830,7 +830,7 @@ void FSceneRenderer::RenderDistortion(FRHICommandListImmediate& RHICmdList)
 
 		// Create a texture to store the resolved light attenuation values, and a render-targetable surface to hold the unresolved light attenuation values.
 		{
-			FPooledRenderTargetDesc Desc(FPooledRenderTargetDesc::Create2DDesc(SceneContext.GetBufferSizeXY(), PF_B8G8R8A8, TexCreate_None, TexCreate_RenderTargetable, false));
+			FPooledRenderTargetDesc Desc(FPooledRenderTargetDesc::Create2DDesc(SceneContext.GetBufferSizeXY(), PF_B8G8R8A8, FClearValueBinding::None, TexCreate_None, TexCreate_RenderTargetable, false));
 			Desc.Flags |= TexCreate_FastVRAM;
 			GRenderTargetPool.FindFreeElement(Desc, DistortionRT, TEXT("Distortion"));
 
diff --git a/Engine/Source/Runtime/Renderer/Private/ForwardShadingRenderer.cpp b/Engine/Source/Runtime/Renderer/Private/ForwardShadingRenderer.cpp
--- a/Engine/Source/Runtime/Renderer/Private/ForwardShadingRenderer.cpp
+++ b/Engine/Source/Runtime/Renderer/Private/ForwardShadingRenderer.cpp
@@ -120,12 +120,12 @@ void FForwardShadingSceneRenderer::Render(FRHICommandListImmediate& RHICmdList)
 
 	if (bGammaSpace && !bRenderToScene)
 	{
-		SetRenderTarget(RHICmdList, ViewFamily.RenderTarget->GetRenderTargetTexture(), SceneContext.GetSceneDepthTexture(), ESimpleRenderTargetMode::EClearToDefault);
+		SetRenderTarget(RHICmdList, ViewFamily.RenderTarget->GetRenderTargetTexture(), SceneContext.GetSceneDepthTexture(), ESimpleRenderTargetMode::EClearColorAndDepth);
 	}
 	else
 	{
 		// Begin rendering to scene color
-		SceneContext.BeginRenderingSceneColor(RHICmdList, ESimpleRenderTargetMode::EClearToDefault);
+		SceneContext.BeginRenderingSceneColor(RHICmdList, ESimpleRenderTargetMode::EClearColorAndDepth);
 	}
 
 	if (GIsEditor)
diff --git a/Engine/Source/Runtime/Renderer/Private/GPUBenchmark.cpp b/Engine/Source/Runtime/Renderer/Private/GPUBenchmark.cpp
--- a/Engine/Source/Runtime/Renderer/Private/GPUBenchmark.cpp
+++ b/Engine/Source/Runtime/Renderer/Private/GPUBenchmark.cpp
@@ -313,7 +313,7 @@ void RendererGPUBenchmark(FRHICommandListImmediate& RHICmdList, FSynthBenchmarkR
 	// two RT to ping pong so we force the GPU to flush it's pipeline
 	TRefCountPtr<IPooledRenderTarget> RTItems[3];
 	{
-		FPooledRenderTargetDesc Desc(FPooledRenderTargetDesc::Create2DDesc(FIntPoint(GBenchmarkResolution, GBenchmarkResolution), PF_B8G8R8A8, TexCreate_None, TexCreate_RenderTargetable | TexCreate_ShaderResource, false));
+		FPooledRenderTargetDesc Desc(FPooledRenderTargetDesc::Create2DDesc(FIntPoint(GBenchmarkResolution, GBenchmarkResolution), PF_B8G8R8A8, FClearValueBinding::None, TexCreate_None, TexCreate_RenderTargetable | TexCreate_ShaderResource, false));
 		GRenderTargetPool.FindFreeElement(Desc, RTItems[0], TEXT("Benchmark0"));
 		GRenderTargetPool.FindFreeElement(Desc, RTItems[1], TEXT("Benchmark1"));
 
diff --git a/Engine/Source/Runtime/Renderer/Private/GlobalDistanceField.cpp b/Engine/Source/Runtime/Renderer/Private/GlobalDistanceField.cpp
--- a/Engine/Source/Runtime/Renderer/Private/GlobalDistanceField.cpp
+++ b/Engine/Source/Runtime/Renderer/Private/GlobalDistanceField.cpp
@@ -534,6 +534,7 @@ void AllocateClipmapTexture(int32 ClipmapIndex, TRefCountPtr<IPooledRenderTarget
 			GAOGlobalDFResolution,
 			GAOGlobalDFResolution,
 			PF_R16F,
+			FClearValueBinding::None,
 			0,
 			TexCreate_ShaderResource | TexCreate_RenderTargetable | TexCreate_UAV,
 			false)),
diff --git a/Engine/Source/Runtime/Renderer/Private/HeightfieldLighting.cpp b/Engine/Source/Runtime/Renderer/Private/HeightfieldLighting.cpp
--- a/Engine/Source/Runtime/Renderer/Private/HeightfieldLighting.cpp
+++ b/Engine/Source/Runtime/Renderer/Private/HeightfieldLighting.cpp
@@ -69,7 +69,8 @@ void FHeightfieldLightingAtlas::InitDynamicRHI()
 	{
 		FPooledRenderTargetDesc Desc(FPooledRenderTargetDesc::Create2DDesc(
 			AtlasSize,
-			PF_G16, 
+			PF_G16,
+			FClearValueBinding::Transparent,
 			TexCreate_None,
 			TexCreate_RenderTargetable,
 			false));
@@ -78,7 +79,8 @@ void FHeightfieldLightingAtlas::InitDynamicRHI()
 
 		FPooledRenderTargetDesc Desc2(FPooledRenderTargetDesc::Create2DDesc(
 			AtlasSize,
-			PF_R8G8, 
+			PF_R8G8,
+			FClearValueBinding::Transparent,
 			TexCreate_None,
 			TexCreate_RenderTargetable,
 			false));
@@ -88,6 +90,7 @@ void FHeightfieldLightingAtlas::InitDynamicRHI()
 		FPooledRenderTargetDesc Desc3(FPooledRenderTargetDesc::Create2DDesc(
 			AtlasSize,
 			PF_R8G8B8A8, 
+			FClearValueBinding::Transparent,
 			TexCreate_None,
 			TexCreate_RenderTargetable,
 			false));
@@ -96,7 +99,8 @@ void FHeightfieldLightingAtlas::InitDynamicRHI()
 
 		FPooledRenderTargetDesc Desc4(FPooledRenderTargetDesc::Create2DDesc(
 			AtlasSize,
-			PF_G8, 
+			PF_G8,			
+			FClearValueBinding::White,
 			TexCreate_None,
 			TexCreate_RenderTargetable,
 			false));
@@ -105,7 +109,8 @@ void FHeightfieldLightingAtlas::InitDynamicRHI()
 
 		FPooledRenderTargetDesc Desc5(FPooledRenderTargetDesc::Create2DDesc(
 			AtlasSize,
-			PF_FloatR11G11B10, 
+			PF_FloatR11G11B10,
+			FClearValueBinding::Black,
 			TexCreate_None,
 			TexCreate_RenderTargetable,
 			false));
@@ -507,12 +512,9 @@ void FHeightfieldLightingViewInfo::SetupVisibleHeightfields(const FViewInfo& Vie
 						ExistingAtlas->DiffuseColor->GetRenderTargetItem().TargetableTexture
 					};
 
-					SetRenderTargets(RHICmdList, ARRAY_COUNT(RenderTargets), RenderTargets, FTextureRHIParamRef(), 0, NULL);
-
-					FLinearColor ClearColors[3] = {FLinearColor(0, 0, 0, 0), FLinearColor(0, 0, 0, 0), FLinearColor(0, 0, 0, 0)};
-					RHICmdList.ClearMRT(true, ARRAY_COUNT(ClearColors), ClearColors, false, 0, false, 0, FIntRect());
-
 					RHICmdList.SetViewport(0, 0, 0.0f, LightingAtlasSize.X, LightingAtlasSize.Y, 1.0f);
+					SetRenderTargets(RHICmdList, ARRAY_COUNT(RenderTargets), RenderTargets, FTextureRHIParamRef(), ESimpleRenderTargetMode::EClearColorExistingDepth, FExclusiveDepthStencil::DepthRead_StencilRead);
+					
 					RHICmdList.SetRasterizerState(TStaticRasterizerState<FM_Solid, CM_None>::GetRHI());
 					RHICmdList.SetDepthStencilState(TStaticDepthStencilState<false, CF_Always>::GetRHI());
 					RHICmdList.SetBlendState(TStaticBlendState<>::GetRHI());
@@ -824,8 +826,7 @@ void FHeightfieldLightingViewInfo::ClearShadowing(const FViewInfo& View, FRHICom
 	{
 		FSceneViewState* ViewState = (FSceneViewState*)View.State;
 		const FHeightfieldLightingAtlas& Atlas = *ViewState->HeightfieldLightingAtlas;
-		SetRenderTarget(RHICmdList, Atlas.DirectionalLightShadowing->GetRenderTargetItem().TargetableTexture, NULL);
-		RHICmdList.Clear(true, FLinearColor(1, 1, 1), false, 0, false, 0, FIntRect());
+		SetRenderTarget(RHICmdList, Atlas.DirectionalLightShadowing->GetRenderTargetItem().TargetableTexture, NULL, ESimpleRenderTargetMode::EClearColorExistingDepth);		
 	}
 }
 
@@ -1159,10 +1160,7 @@ void FHeightfieldLightingViewInfo::ComputeLighting(const FViewInfo& View, FRHICo
 		const FIntPoint LightingAtlasSize = Atlas.GetAtlasSize();
 		const FVector2D InvLightingAtlasSize(1.0f / LightingAtlasSize.X, 1.0f / LightingAtlasSize.Y);
 
-		SetRenderTarget(RHICmdList, Atlas.Lighting->GetRenderTargetItem().TargetableTexture, NULL);
-
-		// We will read outside the valid area later during light transfer
-		RHICmdList.Clear(true, FLinearColor(0, 0, 0), false, 0, false, 0, FIntRect());
+		SetRenderTarget(RHICmdList, Atlas.Lighting->GetRenderTargetItem().TargetableTexture, NULL, ESimpleRenderTargetMode::EClearColorExistingDepth);		
 
 		const bool bApplyLightFunction = (View.Family->EngineShowFlags.LightFunctions &&
 			LightSceneInfo.Proxy->GetLightFunctionMaterial() && 
diff --git a/Engine/Source/Runtime/Renderer/Private/IndirectLightingCache.cpp b/Engine/Source/Runtime/Renderer/Private/IndirectLightingCache.cpp
--- a/Engine/Source/Runtime/Renderer/Private/IndirectLightingCache.cpp
+++ b/Engine/Source/Runtime/Renderer/Private/IndirectLightingCache.cpp
@@ -156,7 +156,8 @@ void FIndirectLightingCache::InitDynamicRHI()
 			CacheSize,
 			CacheSize,
 			CacheSize,
-			PF_FloatRGBA, 
+			PF_FloatRGBA,
+			FClearValueBinding::None,
 			Flags,
 			TexCreate_None,
 			false, 
diff --git a/Engine/Source/Runtime/Renderer/Private/LightPropagationVolume.cpp b/Engine/Source/Runtime/Renderer/Private/LightPropagationVolume.cpp
--- a/Engine/Source/Runtime/Renderer/Private/LightPropagationVolume.cpp
+++ b/Engine/Source/Runtime/Renderer/Private/LightPropagationVolume.cpp
@@ -852,6 +852,7 @@ void FLightPropagationVolume::InitSettings(FRHICommandList& RHICmdList, const FS
 			LPV_GRIDRES,
 			LPV_GRIDRES,
 			PF_FloatRGBA,
+			FClearValueBinding::None,
 			TexCreate_HideInVisualizeTexture,
 			TexCreate_ShaderResource | TexCreate_UAV | TexCreate_FastVRAM,
 			false,
@@ -885,6 +886,7 @@ void FLightPropagationVolume::InitSettings(FRHICommandList& RHICmdList, const FS
 				LPV_GRIDRES,
 				LPV_GRIDRES,
 				PF_G8,
+				FClearValueBinding::None,
 				TexCreate_HideInVisualizeTexture,
 				TexCreate_ShaderResource | TexCreate_UAV | TexCreate_FastVRAM,
 				false,
diff --git a/Engine/Source/Runtime/Renderer/Private/LightShaftRendering.cpp b/Engine/Source/Runtime/Renderer/Private/LightShaftRendering.cpp
--- a/Engine/Source/Runtime/Renderer/Private/LightShaftRendering.cpp
+++ b/Engine/Source/Runtime/Renderer/Private/LightShaftRendering.cpp
@@ -448,7 +448,7 @@ void AllocateOrReuseLightShaftRenderTarget(FRHICommandListImmediate& RHICmdList,
 		EPixelFormat LightShaftFilterBufferFormat = PF_FloatRGB;
 		const FIntPoint BufferSize = FSceneRenderTargets::Get(RHICmdList).GetBufferSizeXY();
 		FIntPoint LightShaftSize(FMath::Max<uint32>(BufferSize.X / GetLightShaftDownsampleFactor(), 1), FMath::Max<uint32>(BufferSize.Y / GetLightShaftDownsampleFactor(), 1));
-		FPooledRenderTargetDesc Desc(FPooledRenderTargetDesc::Create2DDesc(LightShaftSize, LightShaftFilterBufferFormat, TexCreate_None, TexCreate_RenderTargetable, false));
+		FPooledRenderTargetDesc Desc(FPooledRenderTargetDesc::Create2DDesc(LightShaftSize, LightShaftFilterBufferFormat, FClearValueBinding::None, TexCreate_None, TexCreate_RenderTargetable, false));
 		GRenderTargetPool.FindFreeElement(Desc, Target, Name);
 
 		SetRenderTarget(RHICmdList, Target->GetRenderTargetItem().TargetableTexture, FTextureRHIRef());
diff --git a/Engine/Source/Runtime/Renderer/Private/PostProcess/PostProcessCombineLUTs.cpp b/Engine/Source/Runtime/Renderer/Private/PostProcess/PostProcessCombineLUTs.cpp
--- a/Engine/Source/Runtime/Renderer/Private/PostProcess/PostProcessCombineLUTs.cpp
+++ b/Engine/Source/Runtime/Renderer/Private/PostProcess/PostProcessCombineLUTs.cpp
@@ -706,7 +706,7 @@ void FRCPassPostProcessCombineLUTs::Process(FRenderingCompositePassContext& Cont
 
 FPooledRenderTargetDesc FRCPassPostProcessCombineLUTs::ComputeOutputDesc(EPassOutputId InPassOutputId) const
 {
-	FPooledRenderTargetDesc Ret = FPooledRenderTargetDesc::Create2DDesc(FIntPoint(GLUTSize * GLUTSize, GLUTSize), PF_A2B10G10R10, TexCreate_None, TexCreate_RenderTargetable | TexCreate_ShaderResource, false);
+	FPooledRenderTargetDesc Ret = FPooledRenderTargetDesc::Create2DDesc(FIntPoint(GLUTSize * GLUTSize, GLUTSize), PF_A2B10G10R10, FClearValueBinding::None, TexCreate_None, TexCreate_RenderTargetable | TexCreate_ShaderResource, false);
 
 	if(UseVolumeTextureLUT(ShaderPlatform))
 	{
diff --git a/Engine/Source/Runtime/Renderer/Private/PostProcess/PostProcessCompositeEditorPrimitives.cpp b/Engine/Source/Runtime/Renderer/Private/PostProcess/PostProcessCompositeEditorPrimitives.cpp
--- a/Engine/Source/Runtime/Renderer/Private/PostProcess/PostProcessCompositeEditorPrimitives.cpp
+++ b/Engine/Source/Runtime/Renderer/Private/PostProcess/PostProcessCompositeEditorPrimitives.cpp
@@ -210,14 +210,19 @@ void FRCPassPostProcessCompositeEditorPrimitives::Process(FRenderingCompositePas
 	const uint32 MSAASampleCount = SceneContext.EditorPrimitivesColor->GetDesc().NumSamples;
 
 	{
-		SetRenderTarget(Context.RHICmdList, ColorTarget, DepthTarget);
 		Context.SetViewportAndCallRHI(DestRect);
-
-		if(bClearIsNeeded)
+		
+		if (bClearIsNeeded)
 		{
+			check(DepthTarget->GetDepthClearValue() == (float)ERHIZBuffer::FarPlane);
+			check(ColorTarget->GetClearColor() == FLinearColor(0, 0, 0, 0));			
+
 			SCOPED_DRAW_EVENT(Context.RHICmdList, ClearViewEditorPrimitives);
-			// Clear color and depth
-			Context.RHICmdList.Clear(true, FLinearColor(0, 0, 0, 0), true, (float)ERHIZBuffer::FarPlane, false, 0, FIntRect());
+			SetRenderTarget(Context.RHICmdList, ColorTarget, DepthTarget, ESimpleRenderTargetMode::EClearColorAndDepth);
+		}
+		else
+		{
+			SetRenderTarget(Context.RHICmdList, ColorTarget, DepthTarget, ESimpleRenderTargetMode::EExistingColorAndDepth);
 		}
 
 		SCOPED_DRAW_EVENT(Context.RHICmdList, RenderEditorPrimitives);
diff --git a/Engine/Source/Runtime/Renderer/Private/PostProcess/PostProcessHierarchical.cpp b/Engine/Source/Runtime/Renderer/Private/PostProcess/PostProcessHierarchical.cpp
--- a/Engine/Source/Runtime/Renderer/Private/PostProcess/PostProcessHierarchical.cpp
+++ b/Engine/Source/Runtime/Renderer/Private/PostProcess/PostProcessHierarchical.cpp
@@ -258,7 +258,7 @@ FPooledRenderTargetDesc FRCPassPostProcessBuildHCB::ComputeOutputDesc(EPassOutpu
 	FIntPoint HCBSize;
 	uint32 HCBMipCount;
 	HierarchycalSizeAndMips(FSceneRenderTargets::Get_FrameConstantsOnly().GetBufferSizeXY(), 1, HCBMipCount, HCBSize);
-	FPooledRenderTargetDesc Ret(FPooledRenderTargetDesc::Create2DDesc(HCBSize, PF_FloatRGBA, TexCreate_None, TexCreate_RenderTargetable | TexCreate_ShaderResource | TexCreate_NoFastClear, false, HCBMipCount));
+	FPooledRenderTargetDesc Ret(FPooledRenderTargetDesc::Create2DDesc(HCBSize, PF_FloatRGBA, FClearValueBinding::None, TexCreate_None, TexCreate_RenderTargetable | TexCreate_ShaderResource | TexCreate_NoFastClear, false, HCBMipCount));
 	
 	Ret.DebugName = TEXT("HCB");
 
diff --git a/Engine/Source/Runtime/Renderer/Private/PostProcess/PostProcessHistogram.cpp b/Engine/Source/Runtime/Renderer/Private/PostProcess/PostProcessHistogram.cpp
--- a/Engine/Source/Runtime/Renderer/Private/PostProcess/PostProcessHistogram.cpp
+++ b/Engine/Source/Runtime/Renderer/Private/PostProcess/PostProcessHistogram.cpp
@@ -162,7 +162,7 @@ FPooledRenderTargetDesc FRCPassPostProcessHistogram::ComputeOutputDesc(EPassOutp
 	FIntPoint NewSize = FIntPoint(HistogramTexelCount, ThreadGroupCount.X * ThreadGroupCount.Y);
 
 	// format can be optimized later
-	FPooledRenderTargetDesc Ret(FPooledRenderTargetDesc::Create2DDesc(NewSize, PF_FloatRGBA, TexCreate_None, TexCreate_RenderTargetable | TexCreate_UAV, false));
+	FPooledRenderTargetDesc Ret(FPooledRenderTargetDesc::Create2DDesc(NewSize, PF_FloatRGBA, FClearValueBinding::None, TexCreate_None, TexCreate_RenderTargetable | TexCreate_UAV, false));
 
 	Ret.DebugName = TEXT("Histogram");
 
diff --git a/Engine/Source/Runtime/Renderer/Private/PostProcess/PostProcessHistogramReduce.cpp b/Engine/Source/Runtime/Renderer/Private/PostProcess/PostProcessHistogramReduce.cpp
--- a/Engine/Source/Runtime/Renderer/Private/PostProcess/PostProcessHistogramReduce.cpp
+++ b/Engine/Source/Runtime/Renderer/Private/PostProcess/PostProcessHistogramReduce.cpp
@@ -162,7 +162,7 @@ FPooledRenderTargetDesc FRCPassPostProcessHistogramReduce::ComputeOutputDesc(EPa
 	FIntPoint NewSize = FIntPoint(FRCPassPostProcessHistogram::HistogramTexelCount, 2);
 
 	// for quality float4 to get best quality for smooth eye adaptation transitions
-	FPooledRenderTargetDesc Ret(FPooledRenderTargetDesc::Create2DDesc(NewSize, PF_A32B32G32R32F, TexCreate_None, TexCreate_RenderTargetable, false));
+	FPooledRenderTargetDesc Ret(FPooledRenderTargetDesc::Create2DDesc(NewSize, PF_A32B32G32R32F, FClearValueBinding::None, TexCreate_None, TexCreate_RenderTargetable, false));
 	
 	Ret.DebugName = TEXT("HistogramReduce");
 	
diff --git a/Engine/Source/Runtime/Renderer/Private/PostProcess/PostProcessLensBlur.cpp b/Engine/Source/Runtime/Renderer/Private/PostProcess/PostProcessLensBlur.cpp
--- a/Engine/Source/Runtime/Renderer/Private/PostProcess/PostProcessLensBlur.cpp
+++ b/Engine/Source/Runtime/Renderer/Private/PostProcess/PostProcessLensBlur.cpp
@@ -185,7 +185,7 @@ void FRCPassPostProcessLensBlur::Process(FRenderingCompositePassContext& Context
 	const FSceneRenderTargetItem& DestRenderTarget = PassOutputs[0].RequestSurface(Context);
 
 	// Set the view family's render target/viewport.
-	SetRenderTarget(Context.RHICmdList, DestRenderTarget.TargetableTexture, FTextureRHIRef(), ESimpleRenderTargetMode::EClearColorToBlack);
+	SetRenderTarget(Context.RHICmdList, DestRenderTarget.TargetableTexture, FTextureRHIRef(), ESimpleRenderTargetMode::EClearColorExistingDepth);
 
 	Context.SetViewportAndCallRHI(ViewRect);
 
@@ -226,6 +226,7 @@ FPooledRenderTargetDesc FRCPassPostProcessLensBlur::ComputeOutputDesc(EPassOutpu
 	FPooledRenderTargetDesc Ret = GetInput(ePId_Input0)->GetOutput()->RenderTargetDesc;
 
 	Ret.Reset();
+	Ret.ClearValue = FClearValueBinding::Black;
 
 	// more precision for additive blending
 	Ret.Format = PF_FloatRGBA;
diff --git a/Engine/Source/Runtime/Renderer/Private/PostProcess/PostProcessMotionBlur.cpp b/Engine/Source/Runtime/Renderer/Private/PostProcess/PostProcessMotionBlur.cpp
--- a/Engine/Source/Runtime/Renderer/Private/PostProcess/PostProcessMotionBlur.cpp
+++ b/Engine/Source/Runtime/Renderer/Private/PostProcess/PostProcessMotionBlur.cpp
@@ -836,7 +836,7 @@ FPooledRenderTargetDesc FRCPassPostProcessVelocityFlatten::ComputeOutputDesc(EPa
 		// Flattened velocity
 		FPooledRenderTargetDesc Ret = GetInput(ePId_Input0)->GetOutput()->RenderTargetDesc;
 		Ret.Reset();
-
+		Ret.ClearValue = FClearValueBinding::None;
 		Ret.Format = PF_FloatR11G11B10;
 		Ret.TargetableFlags |= TexCreate_UAV;
 		Ret.TargetableFlags |= TexCreate_RenderTargetable;
@@ -853,7 +853,7 @@ FPooledRenderTargetDesc FRCPassPostProcessVelocityFlatten::ComputeOutputDesc(EPa
 		FIntPoint PixelExtent = UnmodifiedRet.Extent;
 		FIntPoint TileCount = GetNumTiles16x16(PixelExtent);
 
-		FPooledRenderTargetDesc Ret(FPooledRenderTargetDesc::Create2DDesc(TileCount, PF_FloatRGBA, TexCreate_None, TexCreate_RenderTargetable | TexCreate_UAV, false));
+		FPooledRenderTargetDesc Ret(FPooledRenderTargetDesc::Create2DDesc(TileCount, PF_FloatRGBA, FClearValueBinding::None, TexCreate_None, TexCreate_RenderTargetable | TexCreate_UAV, false));
 
 		Ret.DebugName = TEXT("MaxVelocity");
 
@@ -1019,24 +1019,22 @@ void FRCPassPostProcessVelocityScatter::Process(FRenderingCompositePassContext&
 	const FSceneRenderTargetItem& DestRenderTarget = PassOutputs[0].RequestSurface(Context);
 
 	TRefCountPtr<IPooledRenderTarget> DepthTarget;
-	FPooledRenderTargetDesc Desc( FPooledRenderTargetDesc::Create2DDesc( DestSize, PF_ShadowDepth, TexCreate_None, TexCreate_DepthStencilTargetable | TexCreate_NoFastClear, false ) );
+	FPooledRenderTargetDesc Desc( FPooledRenderTargetDesc::Create2DDesc( DestSize, PF_ShadowDepth, FClearValueBinding::DepthOne, TexCreate_None, TexCreate_DepthStencilTargetable, false ) );
 	GRenderTargetPool.FindFreeElement( Desc, DepthTarget, TEXT("VelocityScatterDepth") );
 
 	// Set the view family's render target/viewport.
-	SetRenderTarget( Context.RHICmdList, DestRenderTarget.TargetableTexture, DepthTarget->GetRenderTargetItem().TargetableTexture );
-	
-	Context.SetViewportAndCallRHI( 0, 0, 0.0f, TileCount.X, TileCount.Y, 1.0f );
+	FRHIRenderTargetView ColorView(DestRenderTarget.TargetableTexture, 0, -1, ERenderTargetLoadAction::ELoad, ERenderTargetStoreAction::EStore);
+	FRHIDepthRenderTargetView DepthView(DepthTarget->GetRenderTargetItem().TargetableTexture, ERenderTargetLoadAction::EClear, ERenderTargetStoreAction::EStore, ERenderTargetLoadAction::ELoad, ERenderTargetStoreAction::EStore);
+	FRHISetRenderTargetsInfo RTInfo(1, &ColorView, DepthView);
+
+	// clear depth
+	// Max >= Min so no need to clear on second pass
+	Context.SetViewportAndCallRHI(0, 0, 0.0f, TileCount.X, TileCount.Y, 1.0f);
+	Context.RHICmdList.SetRenderTargetsAndClear(RTInfo);			
 	
 	// Min,Max
 	for( int i = 0; i < 2; i++ )
 	{
-		if( i == 0 )
-		{
-			// clear depth
-			// Max >= Min so no need to clear on second pass
-			Context.RHICmdList.Clear( false, FLinearColor::Black, true, 1.0f, false, 0, FIntRect() );
-		}
-
 		if( i == 0 )
 		{
 			// min
diff --git a/Engine/Source/Runtime/Renderer/Private/PostProcess/PostProcessTestImage.cpp b/Engine/Source/Runtime/Renderer/Private/PostProcess/PostProcessTestImage.cpp
--- a/Engine/Source/Runtime/Renderer/Private/PostProcess/PostProcessTestImage.cpp
+++ b/Engine/Source/Runtime/Renderer/Private/PostProcess/PostProcessTestImage.cpp
@@ -203,7 +203,7 @@ void FRCPassPostProcessTestImage::Process(FRenderingCompositePassContext& Contex
 
 FPooledRenderTargetDesc FRCPassPostProcessTestImage::ComputeOutputDesc(EPassOutputId InPassOutputId) const
 {
-	FPooledRenderTargetDesc Ret(FPooledRenderTargetDesc::Create2DDesc(FSceneRenderTargets::Get_FrameConstantsOnly().GetBufferSizeXY(), PF_B8G8R8A8, TexCreate_None, TexCreate_RenderTargetable, false));
+	FPooledRenderTargetDesc Ret(FPooledRenderTargetDesc::Create2DDesc(FSceneRenderTargets::Get_FrameConstantsOnly().GetBufferSizeXY(), PF_B8G8R8A8, FClearValueBinding::None, TexCreate_None, TexCreate_RenderTargetable, false));
 
 	Ret.DebugName = TEXT("TestImage");
 
diff --git a/Engine/Source/Runtime/Renderer/Private/PostProcess/RenderTargetPool.cpp b/Engine/Source/Runtime/Renderer/Private/PostProcess/RenderTargetPool.cpp
--- a/Engine/Source/Runtime/Renderer/Private/PostProcess/RenderTargetPool.cpp
+++ b/Engine/Source/Runtime/Renderer/Private/PostProcess/RenderTargetPool.cpp
@@ -243,7 +243,7 @@ bool FRenderTargetPool::FindFreeElement(const FPooledRenderTargetDesc& Desc, TRe
 		// TexCreate_UAV should be used on Desc.TargetableFlags
 		check(!(Desc.Flags & TexCreate_UAV));
 
-		FRHIResourceCreateInfo CreateInfo;
+		FRHIResourceCreateInfo CreateInfo(Desc.ClearValue);
 
 		if(Desc.TargetableFlags & (TexCreate_RenderTargetable | TexCreate_DepthStencilTargetable | TexCreate_UAV))
 		{
diff --git a/Engine/Source/Runtime/Renderer/Private/PostProcess/SceneRenderTargets.cpp b/Engine/Source/Runtime/Renderer/Private/PostProcess/SceneRenderTargets.cpp
--- a/Engine/Source/Runtime/Renderer/Private/PostProcess/SceneRenderTargets.cpp
+++ b/Engine/Source/Runtime/Renderer/Private/PostProcess/SceneRenderTargets.cpp
@@ -225,10 +225,8 @@ FSceneRenderTargets::FSceneRenderTargets(const FViewInfo& View, const FSceneRend
 	, CurrentFeatureLevel(SnapshotSource.CurrentFeatureLevel)
 	, CurrentShadingPath(SnapshotSource.CurrentShadingPath)
 	, bAllocateVelocityGBuffer(SnapshotSource.bAllocateVelocityGBuffer)
-	, bGBuffersCleared(SnapshotSource.bGBuffersCleared)
-	, GBufferClearColor(SnapshotSource.GBufferClearColor)
-	, bSceneDepthCleared(SnapshotSource.bSceneDepthCleared)
-	, SceneDepthClearValue(SnapshotSource.SceneDepthClearValue)
+	, bGBuffersFastCleared(SnapshotSource.bGBuffersFastCleared)	
+	, bSceneDepthCleared(SnapshotSource.bSceneDepthCleared)	
 	, bSnapshot(true)
 {
 	SnapshotArray(SceneColor, SnapshotSource.SceneColor);
@@ -513,45 +511,61 @@ void FSceneRenderTargets::BeginRenderingGBuffer(FRHICommandList& RHICmdList, ERe
 	// Set the scene color surface as the render target, and the scene depth surface as the depth-stencil target.
 	if (CurrentFeatureLevel >= ERHIFeatureLevel::SM4)
 	{
-		int32 VelocityRTIndex;
-		FRHIRenderTargetView RenderTargets[MaxSimultaneousRenderTargets];
-		int32 MRTCount = GetGBufferRenderTargets(ColorLoadAction, RenderTargets, VelocityRTIndex);
-
-		const float DepthClearValue = (float)ERHIZBuffer::FarPlane;
 		bool bClearColor = ColorLoadAction == ERenderTargetLoadAction::EClear;
 		bool bClearDepth = DepthLoadAction == ERenderTargetLoadAction::EClear;
-		FRHIDepthRenderTargetView DepthView(GetSceneDepthSurface(), DepthLoadAction, ERenderTargetStoreAction::EStore);
-		FRHISetRenderTargetsInfo Info(MRTCount, RenderTargets, DepthView);
+
+		//if the desired clear color doesn't match the bound hwclear value, or there isn't one at all (editor code)
+		//then we need to fall back to a shader clear.
+		const FTextureRHIRef& SceneColorTex = GetSceneColorSurface();
+		bool bShaderClear = false;
 		if (bClearColor)
-		{
-			bGBuffersCleared = true;
-			GBufferClearColor = ClearColor;
+		{			
+			if (!SceneColorTex->HasClearValue() || (ClearColor != SceneColorTex->GetClearColor()))
+			{
+				ColorLoadAction = ERenderTargetLoadAction::ENoAction;				
+				bShaderClear = true;
+			}
+			else
+			{
+				bGBuffersFastCleared = true;
+			}
 		}
 
-		Info.ClearColors[0] = GBufferClearColor;
-		Info.ClearColors[1] = Info.ClearColors[2] = Info.ClearColors[3] = FLinearColor::Transparent;
-		Info.ClearColors[4] = FLinearColor(0, 1, 1, 1);
-		Info.ClearColors[5] = FLinearColor(1, 1, 1, 1);
+		int32 VelocityRTIndex;
+		FRHIRenderTargetView RenderTargets[MaxSimultaneousRenderTargets];
+		int32 MRTCount = GetGBufferRenderTargets(ColorLoadAction, RenderTargets, VelocityRTIndex);
 
-		if (VelocityRTIndex != -1)
-		{
-			Info.ClearColors[VelocityRTIndex] = FLinearColor(0, 0, 0, 0);
-		}
+		//make sure our conditions for shader clear fallback are valid.
+		check(RenderTargets[0].Texture == SceneColorTex);
+		
+		
+
+		FRHIDepthRenderTargetView DepthView(GetSceneDepthSurface(), DepthLoadAction, ERenderTargetStoreAction::EStore);
+		FRHISetRenderTargetsInfo Info(MRTCount, RenderTargets, DepthView);		
 
 		if (bClearDepth)
 		{
-			bSceneDepthCleared = true;
-			SceneDepthClearValue = DepthClearValue;
-		}
-		Info.DepthClearValue = DepthClearValue;
+			bSceneDepthCleared = true;			
+		}		
 
 		// set the render target
 		RHICmdList.SetRenderTargetsAndClear(Info);
+		if (bShaderClear)
+		{
+			FLinearColor ClearColors[MaxSimultaneousRenderTargets];
+			ClearColors[0] = ClearColor;
+			for (int32 i = 1; i < MRTCount; ++i)
+			{
+				ClearColors[i] = RenderTargets[i].Texture->GetClearColor();
+			}
+			//depth/stencil should have been handled by the fast clear.  only color for RT0 can get changed.
+			RHICmdList.ClearMRT(true, MRTCount, ClearColors, false, 0.0f, false, 0, FIntRect());
+		}
 
 		//bind any clear data that won't be bound automatically by the preceding SetRenderTargetsAndClear
-		bool bBindClearColor = !bClearColor && bGBuffersCleared;
+		bool bBindClearColor = !bClearColor && bGBuffersFastCleared;
 		bool bBindClearDepth = !bClearDepth && bSceneDepthCleared;
-		RHICmdList.BindClearMRTValues(bBindClearColor, MRTCount, Info.ClearColors, bBindClearDepth, SceneDepthClearValue, bBindClearDepth, Info.StencilClearValue);
+		RHICmdList.BindClearMRTValues(bBindClearColor, bBindClearDepth, bBindClearDepth);
 	}
 }
 
@@ -619,7 +633,7 @@ void FSceneRenderTargets::AllocSceneColor()
 
 	// Create the scene color.
 	{
-		FPooledRenderTargetDesc Desc(FPooledRenderTargetDesc::Create2DDesc(BufferSize, SceneColorBufferFormat, TexCreate_None, TexCreate_RenderTargetable, false));
+		FPooledRenderTargetDesc Desc(FPooledRenderTargetDesc::Create2DDesc(BufferSize, SceneColorBufferFormat, FClearValueBinding::Transparent, TexCreate_None, TexCreate_RenderTargetable, false));
 
 		Desc.Flags |= TexCreate_FastVRAM;
 
@@ -658,7 +672,7 @@ void FSceneRenderTargets::AllocLightAttenuation()
 
 	// Create a texture to store the resolved light attenuation values, and a render-targetable surface to hold the unresolved light attenuation values.
 	{
-		FPooledRenderTargetDesc Desc(FPooledRenderTargetDesc::Create2DDesc(BufferSize, PF_B8G8R8A8, TexCreate_None, TexCreate_RenderTargetable, false));
+		FPooledRenderTargetDesc Desc(FPooledRenderTargetDesc::Create2DDesc(BufferSize, PF_B8G8R8A8, FClearValueBinding::White, TexCreate_None, TexCreate_RenderTargetable, false));
 		Desc.Flags |= TexCreate_FastVRAM;
 		GRenderTargetPool.FindFreeElement(Desc, LightAttenuation, TEXT("LightAttenuation"));
 
@@ -726,7 +740,7 @@ void FSceneRenderTargets::AllocGBufferTargets()
 			NormalGBufferFormat = PF_B8G8R8A8;
 		}
 
-		FPooledRenderTargetDesc Desc(FPooledRenderTargetDesc::Create2DDesc(BufferSize, NormalGBufferFormat, TexCreate_None, TexCreate_RenderTargetable, false));
+		FPooledRenderTargetDesc Desc(FPooledRenderTargetDesc::Create2DDesc(BufferSize, NormalGBufferFormat, FClearValueBinding::Transparent, TexCreate_None, TexCreate_RenderTargetable, false));
 		GRenderTargetPool.FindFreeElement(Desc, GBufferA, TEXT("GBufferA"));
 	}
 
@@ -734,7 +748,7 @@ void FSceneRenderTargets::AllocGBufferTargets()
 	{
 		const EPixelFormat SpecularGBufferFormat = bHighPrecisionGBuffers ? PF_FloatRGBA : PF_B8G8R8A8;
 
-		FPooledRenderTargetDesc Desc(FPooledRenderTargetDesc::Create2DDesc(BufferSize, SpecularGBufferFormat, TexCreate_None, TexCreate_RenderTargetable, false));
+		FPooledRenderTargetDesc Desc(FPooledRenderTargetDesc::Create2DDesc(BufferSize, SpecularGBufferFormat, FClearValueBinding::Transparent, TexCreate_None, TexCreate_RenderTargetable, false));
 		GRenderTargetPool.FindFreeElement(Desc, GBufferB, TEXT("GBufferB"));
 	}
 
@@ -748,19 +762,19 @@ void FSceneRenderTargets::AllocGBufferTargets()
 		DiffuseGBufferFlags = CVar && CVar->GetValueOnRenderThread() ? TexCreate_SRGB : TexCreate_None;
 #endif
 
-		FPooledRenderTargetDesc Desc(FPooledRenderTargetDesc::Create2DDesc(BufferSize, DiffuseGBufferFormat, DiffuseGBufferFlags, TexCreate_RenderTargetable, false));
+		FPooledRenderTargetDesc Desc(FPooledRenderTargetDesc::Create2DDesc(BufferSize, DiffuseGBufferFormat, FClearValueBinding::Transparent, DiffuseGBufferFlags, TexCreate_RenderTargetable, false));
 		GRenderTargetPool.FindFreeElement(Desc, GBufferC, TEXT("GBufferC"));
 	}
 
 	// Create the mask g-buffer (e.g. SSAO, subsurface scattering, wet surface mask, skylight mask, ...).
 	{
-		FPooledRenderTargetDesc Desc(FPooledRenderTargetDesc::Create2DDesc(BufferSize, PF_B8G8R8A8, TexCreate_None, TexCreate_RenderTargetable, false));
+		FPooledRenderTargetDesc Desc(FPooledRenderTargetDesc::Create2DDesc(BufferSize, PF_B8G8R8A8, FClearValueBinding(FLinearColor(0, 1, 1, 1)), TexCreate_None, TexCreate_RenderTargetable, false));
 		GRenderTargetPool.FindFreeElement(Desc, GBufferD, TEXT("GBufferD"));
 	}
 
 	if (bAllowStaticLighting)
 	{
-		FPooledRenderTargetDesc Desc(FPooledRenderTargetDesc::Create2DDesc(BufferSize, PF_B8G8R8A8, TexCreate_None, TexCreate_RenderTargetable, false));
+		FPooledRenderTargetDesc Desc(FPooledRenderTargetDesc::Create2DDesc(BufferSize, PF_B8G8R8A8, FClearValueBinding(FLinearColor(1, 1, 1, 1)), TexCreate_None, TexCreate_RenderTargetable, false));
 		GRenderTargetPool.FindFreeElement(Desc, GBufferE, TEXT("GBufferE"));
 	}
 
@@ -958,7 +972,7 @@ bool FSceneRenderTargets::BeginRenderingCustomDepth(FRHICommandListImmediate& RH
 		FRHIDepthRenderTargetView DepthView(CustomDepthRenderTarget->GetRenderTargetItem().ShaderResourceTexture);
 		FRHISetRenderTargetsInfo Info(0, nullptr, DepthView);
 		Info.bClearStencil = IsCustomDepthPassWritingStencil();
-		Info.StencilClearValue = 0;
+		check(DepthView.Texture->GetStencilClearValue() == 0);
 
 		RHICmdList.SetRenderTargetsAndClear(Info);
 
@@ -1075,9 +1089,7 @@ void FSceneRenderTargets::BeginRenderingPrePass(FRHICommandList& RHICmdList, boo
 
 	FTextureRHIRef ColorTarget;
 	FTexture2DRHIRef DepthTarget = GetSceneDepthSurface();
-
-	const float DepthClearValue = (float)ERHIZBuffer::FarPlane;
-	const uint32 StencilClearValue = 0;
+	
 	if (bPerformClear)
 	{				
 		FRHIRenderTargetView ColorView(ColorTarget, 0, -1, ERenderTargetLoadAction::ENoAction, ERenderTargetStoreAction::ENoAction);
@@ -1086,20 +1098,16 @@ void FSceneRenderTargets::BeginRenderingPrePass(FRHICommandList& RHICmdList, boo
 		// Clear the depth buffer.
 		// Note, this is a reversed Z depth surface, so 0.0f is the far plane.
 		FRHISetRenderTargetsInfo Info(1, &ColorView, DepthView);
-		Info.DepthClearValue = DepthClearValue;
-		Info.StencilClearValue = StencilClearValue;
-
+	
 		RHICmdList.SetRenderTargetsAndClear(Info);
-
-		bSceneDepthCleared = true;
-		SceneDepthClearValue = DepthClearValue;
+		bSceneDepthCleared = true;	
 	}
 	else
 	{
 		// Set the scene depth surface and a DUMMY buffer as color buffer
 		// (as long as it's the same dimension as the depth buffer),	
 		SetRenderTarget(RHICmdList, FTextureRHIRef(), GetSceneDepthSurface());
-		RHICmdList.BindClearMRTValues(false, 0, nullptr, true, DepthClearValue, true, StencilClearValue);
+		RHICmdList.BindClearMRTValues(false, true, true);
 	}
 }
 
@@ -1120,7 +1128,7 @@ void FSceneRenderTargets::BeginRenderingShadowDepth(FRHICommandList& RHICmdList,
 		ERenderTargetLoadAction::ELoad, 
 		ERenderTargetStoreAction::EStore));
 
-	Info.DepthClearValue = 1.0f;
+	check(Info.DepthStencilRenderTarget.Texture->GetDepthClearValue() == 1.0f);	
 	Info.ColorRenderTarget[0].StoreAction = ERenderTargetStoreAction::ENoAction;
 
 	if (!GSupportsDepthRenderTargetWithoutColorRenderTarget)
@@ -1203,7 +1211,7 @@ void FSceneRenderTargets::BeginRenderingLightAttenuation(FRHICommandList& RHICmd
 	GRenderTargetPool.VisualizeTexture.SetCheckPoint(RHICmdList, GetLightAttenuation());
 
 	// Set the light attenuation surface as the render target, and the scene depth buffer as the depth-stencil surface.
-	SetRenderTarget(RHICmdList, GetLightAttenuationSurface(), GetSceneDepthSurface(), bClearToWhite ? ESimpleRenderTargetMode::EClearColorToWhite : ESimpleRenderTargetMode::EExistingColorAndDepth, FExclusiveDepthStencil::DepthRead_StencilWrite);
+	SetRenderTarget(RHICmdList, GetLightAttenuationSurface(), GetSceneDepthSurface(), bClearToWhite ? ESimpleRenderTargetMode::EClearColorExistingDepth : ESimpleRenderTargetMode::EExistingColorAndDepth, FExclusiveDepthStencil::DepthRead_StencilWrite);
 }
 
 void FSceneRenderTargets::FinishRenderingLightAttenuation(FRHICommandList& RHICmdList)
@@ -1252,18 +1260,17 @@ bool FSceneRenderTargets::BeginRenderingSeparateTranslucency(FRHICommandList& RH
 			SeparateTranslucency = &ViewState->GetSeparateTranslucency(GetBufferSizeXY());
 		}
 
-
+		check((*SeparateTranslucency)->GetRenderTargetItem().TargetableTexture->GetClearColor() == FLinearColor::Black);
 
 		// clear the render target the first time, re-use afterwards
 		SetRenderTarget(RHICmdList, (*SeparateTranslucency)->GetRenderTargetItem().TargetableTexture, GetSceneDepthSurface(), 
-			bFirstTimeThisFrame ? ESimpleRenderTargetMode::EClearColorToBlackWithFullAlpha : ESimpleRenderTargetMode::EExistingColorAndDepth, FExclusiveDepthStencil::DepthRead_StencilWrite);
+			bFirstTimeThisFrame ? ESimpleRenderTargetMode::EClearColorExistingDepth : ESimpleRenderTargetMode::EExistingColorAndDepth, FExclusiveDepthStencil::DepthRead_StencilWrite);
 
+		
 		if (!bFirstTimeThisFrame)
 		{
-			FLinearColor ClearColor(0.0f, 0.0f, 0.0f, 1.0f);
-			
 			// Clear the stencil buffer for ResponsiveAA
-			RHICmdList.BindClearMRTValues(true, 1, &ClearColor, false, (float)ERHIZBuffer::FarPlane, true, 0);
+			RHICmdList.BindClearMRTValues(true, false, true);
 		}
 
 		RHICmdList.SetViewport(View.ViewRect.Min.X, View.ViewRect.Min.Y, 0.0f, View.ViewRect.Max.X, View.ViewRect.Max.Y, 1.0f);
@@ -1404,7 +1411,8 @@ bool FSceneRenderTargets::IsSeparateTranslucencyActive(const FViewInfo& View) co
 void FSceneRenderTargets::InitEditorPrimitivesColor()
 {
 	FPooledRenderTargetDesc Desc(FPooledRenderTargetDesc::Create2DDesc(BufferSize, 
-		PF_B8G8R8A8, 
+		PF_B8G8R8A8,
+		FClearValueBinding::Transparent,
 		TexCreate_None, 
 		TexCreate_ShaderResource | TexCreate_RenderTargetable,
 		false));
@@ -1417,7 +1425,8 @@ void FSceneRenderTargets::InitEditorPrimitivesColor()
 void FSceneRenderTargets::InitEditorPrimitivesDepth()
 {
 	FPooledRenderTargetDesc Desc(FPooledRenderTargetDesc::Create2DDesc(BufferSize, 
-		PF_DepthStencil, 
+		PF_DepthStencil,
+		FClearValueBinding::DepthFar,
 		TexCreate_None, 
 		TexCreate_ShaderResource | TexCreate_DepthStencilTargetable,
 		false));
@@ -1459,9 +1468,9 @@ void FSceneRenderTargets::AllocateForwardShadingPathRenderTargets()
 #if PLATFORM_HTML5 || PLATFORM_ANDROID
 		// creating a PF_R16F (a true one-channel renderable fp texture) is only supported on GL if EXT_texture_rg is available.  It's present
 		// on iOS, but not in WebGL or Android.
-		FPooledRenderTargetDesc Desc(FPooledRenderTargetDesc::Create2DDesc(BufferSize, PF_FloatRGBA, TexCreate_None, TexCreate_RenderTargetable, false));
+		FPooledRenderTargetDesc Desc(FPooledRenderTargetDesc::Create2DDesc(BufferSize, PF_FloatRGBA, FClearValueBinding::None, TexCreate_None, TexCreate_RenderTargetable, false));
 #else
-		FPooledRenderTargetDesc Desc(FPooledRenderTargetDesc::Create2DDesc(BufferSize, PF_R16F, TexCreate_None, TexCreate_RenderTargetable, false));
+		FPooledRenderTargetDesc Desc(FPooledRenderTargetDesc::Create2DDesc(BufferSize, PF_R16F, FClearValueBinding::None, TexCreate_None, TexCreate_RenderTargetable, false));
 #endif
 		GRenderTargetPool.FindFreeElement(Desc, SceneAlphaCopy, TEXT("SceneAlphaCopy"));
 	}
@@ -1474,13 +1483,13 @@ void FSceneRenderTargets::AllocateForwardShadingPathRenderTargets()
 void FSceneRenderTargets::AllocateForwardShadingShadowDepthTarget(const FIntPoint& ShadowBufferResolution)
 {
 	{
-		FPooledRenderTargetDesc Desc(FPooledRenderTargetDesc::Create2DDesc(ShadowBufferResolution, PF_ShadowDepth, TexCreate_None, TexCreate_DepthStencilTargetable, false));
+		FPooledRenderTargetDesc Desc(FPooledRenderTargetDesc::Create2DDesc(ShadowBufferResolution, PF_ShadowDepth, FClearValueBinding::DepthOne, TexCreate_None, TexCreate_DepthStencilTargetable, false));
 		GRenderTargetPool.FindFreeElement(Desc, ShadowDepthZ, TEXT("ShadowDepthZ"));
 	}
 
 	if (!GSupportsDepthRenderTargetWithoutColorRenderTarget)
 	{
-		FPooledRenderTargetDesc Desc(FPooledRenderTargetDesc::Create2DDesc(ShadowBufferResolution, PF_B8G8R8A8, TexCreate_None, TexCreate_RenderTargetable, false));
+		FPooledRenderTargetDesc Desc(FPooledRenderTargetDesc::Create2DDesc(ShadowBufferResolution, PF_B8G8R8A8, FClearValueBinding::None, TexCreate_None, TexCreate_RenderTargetable, false));
 		GRenderTargetPool.FindFreeElement(Desc, OptionalShadowDepthColor, TEXT("OptionalShadowDepthColor"));
 	}
 }
@@ -1535,7 +1544,7 @@ void FSceneRenderTargets::AllocateReflectionTargets()
 
 			{
 				// Create scratch cubemaps for filtering passes
-				FPooledRenderTargetDesc Desc2(FPooledRenderTargetDesc::CreateCubemapDesc(GReflectionCaptureSize, PF_FloatRGBA, CubeTexFlags, TexCreate_RenderTargetable, false, 1, NumReflectionCaptureMips));
+				FPooledRenderTargetDesc Desc2(FPooledRenderTargetDesc::CreateCubemapDesc(GReflectionCaptureSize, PF_FloatRGBA, FClearValueBinding::None, CubeTexFlags, TexCreate_RenderTargetable, false, 1, NumReflectionCaptureMips));
 				GRenderTargetPool.FindFreeElement(Desc2, ReflectionColorScratchCubemap[0], TEXT("ReflectionColorScratchCubemap0"));
 				GRenderTargetPool.FindFreeElement(Desc2, ReflectionColorScratchCubemap[1], TEXT("ReflectionColorScratchCubemap1"));
 			}
@@ -1544,13 +1553,13 @@ void FSceneRenderTargets::AllocateReflectionTargets()
 			const int32 NumDiffuseIrradianceMips = FMath::CeilLogTwo(GDiffuseIrradianceCubemapSize) + 1;
 
 			{
-				FPooledRenderTargetDesc Desc2(FPooledRenderTargetDesc::CreateCubemapDesc(GDiffuseIrradianceCubemapSize, PF_FloatRGBA, CubeTexFlags, TexCreate_RenderTargetable, false, 1, NumDiffuseIrradianceMips));
+				FPooledRenderTargetDesc Desc2(FPooledRenderTargetDesc::CreateCubemapDesc(GDiffuseIrradianceCubemapSize, PF_FloatRGBA, FClearValueBinding::None, CubeTexFlags, TexCreate_RenderTargetable, false, 1, NumDiffuseIrradianceMips));
 				GRenderTargetPool.FindFreeElement(Desc2, DiffuseIrradianceScratchCubemap[0], TEXT("DiffuseIrradianceScratchCubemap0"));
 				GRenderTargetPool.FindFreeElement(Desc2, DiffuseIrradianceScratchCubemap[1], TEXT("DiffuseIrradianceScratchCubemap1"));
 			}
 
 			{
-				FPooledRenderTargetDesc Desc(FPooledRenderTargetDesc::Create2DDesc(FIntPoint(FSHVector3::MaxSHBasis, 1), PF_FloatRGBA, TexCreate_None, TexCreate_RenderTargetable, false));
+				FPooledRenderTargetDesc Desc(FPooledRenderTargetDesc::Create2DDesc(FIntPoint(FSHVector3::MaxSHBasis, 1), PF_FloatRGBA, FClearValueBinding::None, TexCreate_None, TexCreate_RenderTargetable, false));
 				GRenderTargetPool.FindFreeElement(Desc, SkySHIrradianceMap, TEXT("SkySHIrradianceMap"));
 			}
 		}
@@ -1567,7 +1576,7 @@ void FSceneRenderTargets::AllocateReflectionTargets()
 			}
 			int32 ReflectionBrightnessIndex = bSupportsR32Float ? 0 : 1;
 			EPixelFormat BrightnessFormat = bSupportsR32Float ? PF_R32_FLOAT : PF_FloatRGBA;
-			FPooledRenderTargetDesc Desc3(FPooledRenderTargetDesc::Create2DDesc(FIntPoint(1, 1), BrightnessFormat, TexCreate_None, TexCreate_RenderTargetable, false));
+			FPooledRenderTargetDesc Desc3(FPooledRenderTargetDesc::Create2DDesc(FIntPoint(1, 1), BrightnessFormat, FClearValueBinding::None, TexCreate_None, TexCreate_RenderTargetable, false));
 			GRenderTargetPool.FindFreeElement(Desc3, ReflectionBrightnessTarget, GetReflectionBrightnessTargetName(ReflectionBrightnessIndex));
 		}
 	}
@@ -1578,7 +1587,7 @@ void FSceneRenderTargets::AllocateCommonDepthTargets()
 	if (!SceneDepthZ)
 	{
 		// Create a texture to store the resolved scene depth, and a render-targetable surface to hold the unresolved scene depth.
-		FPooledRenderTargetDesc Desc(FPooledRenderTargetDesc::Create2DDesc(BufferSize, PF_DepthStencil, TexCreate_None, TexCreate_DepthStencilTargetable, false));
+		FPooledRenderTargetDesc Desc(FPooledRenderTargetDesc::Create2DDesc(BufferSize, PF_DepthStencil, FClearValueBinding::DepthFar, TexCreate_None, TexCreate_DepthStencilTargetable, false));
 		Desc.NumSamples = GetNumSceneColorMSAASamples(CurrentFeatureLevel);
 		Desc.Flags |= TexCreate_FastVRAM;
 		GRenderTargetPool.FindFreeElement(Desc, SceneDepthZ, TEXT("SceneDepthZ"));
@@ -1587,7 +1596,7 @@ void FSceneRenderTargets::AllocateCommonDepthTargets()
 	// When targeting DX Feature Level 10, create an auxiliary texture to store the resolved scene depth, and a render-targetable surface to hold the unresolved scene depth.
 	if (!AuxiliarySceneDepthZ && !GSupportsDepthFetchDuringDepthTest)
 	{
-		FPooledRenderTargetDesc Desc(FPooledRenderTargetDesc::Create2DDesc(BufferSize, PF_DepthStencil, TexCreate_None, TexCreate_DepthStencilTargetable, false));
+		FPooledRenderTargetDesc Desc(FPooledRenderTargetDesc::Create2DDesc(BufferSize, PF_DepthStencil, FClearValueBinding::DepthFar, TexCreate_None, TexCreate_DepthStencilTargetable, false));
 		GRenderTargetPool.FindFreeElement(Desc, AuxiliarySceneDepthZ, TEXT("AuxiliarySceneDepthZ"));
 	}
 }
@@ -1600,7 +1609,7 @@ void FSceneRenderTargets::AllocateDeferredShadingPathRenderTargets()
 	{
 		FIntPoint SmallDepthZSize(FMath::Max<uint32>(BufferSize.X / SmallColorDepthDownsampleFactor, 1), FMath::Max<uint32>(BufferSize.Y / SmallColorDepthDownsampleFactor, 1));
 
-		FPooledRenderTargetDesc Desc(FPooledRenderTargetDesc::Create2DDesc(SmallDepthZSize, PF_DepthStencil, TexCreate_None, TexCreate_DepthStencilTargetable, true));
+		FPooledRenderTargetDesc Desc(FPooledRenderTargetDesc::Create2DDesc(SmallDepthZSize, PF_DepthStencil, FClearValueBinding::None, TexCreate_None, TexCreate_DepthStencilTargetable, true));
 		GRenderTargetPool.FindFreeElement(Desc, SmallDepthZ, TEXT("SmallDepthZ"));
 	}
 
@@ -1616,7 +1625,7 @@ void FSceneRenderTargets::AllocateDeferredShadingPathRenderTargets()
 			if (!TranslucencyShadowTransmission[SurfaceIndex])
 			{
 				// Using PF_FloatRGBA because Fourier coefficients used by Fourier opacity maps have a large range and can be negative
-				FPooledRenderTargetDesc Desc(FPooledRenderTargetDesc::Create2DDesc(TranslucentShadowBufferResolution, PF_FloatRGBA, TexCreate_None, TexCreate_RenderTargetable, false));
+				FPooledRenderTargetDesc Desc(FPooledRenderTargetDesc::Create2DDesc(TranslucentShadowBufferResolution, PF_FloatRGBA, FClearValueBinding::None, TexCreate_None, TexCreate_RenderTargetable, false));
 				GRenderTargetPool.FindFreeElement(Desc, TranslucencyShadowTransmission[SurfaceIndex], GetTranslucencyShadowTransmissionName(SurfaceIndex));
 			}
 		}
@@ -1629,19 +1638,19 @@ void FSceneRenderTargets::AllocateDeferredShadingPathRenderTargets()
 		{
 			const int32 SurfaceResolution = GetCubeShadowDepthZResolution(SurfaceIndex);
 
-			FPooledRenderTargetDesc Desc(FPooledRenderTargetDesc::CreateCubemapDesc(SurfaceResolution, PF_ShadowDepth, TexCreate_None, TexCreate_DepthStencilTargetable | TexCreate_NoFastClear, false));
+			FPooledRenderTargetDesc Desc(FPooledRenderTargetDesc::CreateCubemapDesc(SurfaceResolution, PF_ShadowDepth, FClearValueBinding::None, TexCreate_None, TexCreate_DepthStencilTargetable | TexCreate_NoFastClear, false));
 			GRenderTargetPool.FindFreeElement(Desc, CubeShadowDepthZ[SurfaceIndex], TEXT("CubeShadowDepthZ[]"));
 		}
 	}
 
 	//create the shadow depth texture and/or surface
 	{
-		FPooledRenderTargetDesc Desc(FPooledRenderTargetDesc::Create2DDesc(ShadowBufferResolution, PF_ShadowDepth, TexCreate_None, TexCreate_DepthStencilTargetable, false));
+		FPooledRenderTargetDesc Desc(FPooledRenderTargetDesc::Create2DDesc(ShadowBufferResolution, PF_ShadowDepth, FClearValueBinding::DepthOne, TexCreate_None, TexCreate_DepthStencilTargetable, false));
 		GRenderTargetPool.FindFreeElement(Desc, ShadowDepthZ, TEXT("ShadowDepthZ"));
 	}
 		
 	{
-		FPooledRenderTargetDesc Desc(FPooledRenderTargetDesc::Create2DDesc(GetPreShadowCacheTextureResolution(), PF_ShadowDepth, TexCreate_None, TexCreate_DepthStencilTargetable, false));
+		FPooledRenderTargetDesc Desc(FPooledRenderTargetDesc::Create2DDesc(GetPreShadowCacheTextureResolution(), PF_ShadowDepth, FClearValueBinding::None, TexCreate_None, TexCreate_DepthStencilTargetable, false));
 		GRenderTargetPool.FindFreeElement(Desc, PreShadowCacheDepthZ, TEXT("PreShadowCacheDepthZ"));
 		// Mark the preshadow cache as newly allocated, so the cache will know to update
 		bPreshadowCacheNewlyAllocated = true;
@@ -1652,7 +1661,7 @@ void FSceneRenderTargets::AllocateDeferredShadingPathRenderTargets()
 	{
 		// Create the screen space ambient occlusion buffer
 		{
-			FPooledRenderTargetDesc Desc(FPooledRenderTargetDesc::Create2DDesc(BufferSize, PF_G8, TexCreate_None, TexCreate_RenderTargetable, false));
+			FPooledRenderTargetDesc Desc(FPooledRenderTargetDesc::Create2DDesc(BufferSize, PF_G8, FClearValueBinding::None, TexCreate_None, TexCreate_RenderTargetable, false));
 			GRenderTargetPool.FindFreeElement(Desc, ScreenSpaceAO, TEXT("ScreenSpaceAO"));
 		}
 
@@ -1665,6 +1674,7 @@ void FSceneRenderTargets::AllocateDeferredShadingPathRenderTargets()
 						GTranslucencyLightingVolumeDim,
 						GTranslucencyLightingVolumeDim,
 						PF_FloatRGBA,
+						FClearValueBinding::None,
 						0,
 						TexCreate_ShaderResource | TexCreate_RenderTargetable,
 						false)),
@@ -1678,6 +1688,7 @@ void FSceneRenderTargets::AllocateDeferredShadingPathRenderTargets()
 						GTranslucencyLightingVolumeDim,
 						GTranslucencyLightingVolumeDim,
 						PF_FloatRGBA,
+						FClearValueBinding::None,
 						0,
 						TexCreate_ShaderResource | TexCreate_RenderTargetable,
 						false)),
@@ -1691,12 +1702,12 @@ void FSceneRenderTargets::AllocateDeferredShadingPathRenderTargets()
 	// LPV : Dynamic directional occlusion for diffuse and specular
 	if(UseLightPropagationVolumeRT(CurrentFeatureLevel))
 	{
-		FPooledRenderTargetDesc Desc(FPooledRenderTargetDesc::Create2DDesc(BufferSize, PF_R8G8, TexCreate_None, TexCreate_RenderTargetable, false));
+		FPooledRenderTargetDesc Desc(FPooledRenderTargetDesc::Create2DDesc(BufferSize, PF_R8G8, FClearValueBinding::None, TexCreate_None, TexCreate_RenderTargetable, false));
 		GRenderTargetPool.FindFreeElement(Desc, DirectionalOcclusion, TEXT("DirectionalOcclusion"));
 	}
 
 	{
-		FPooledRenderTargetDesc Desc(FPooledRenderTargetDesc::Create2DDesc(BufferSize, GetSceneColorFormat(), TexCreate_None, TexCreate_RenderTargetable, false));
+		FPooledRenderTargetDesc Desc(FPooledRenderTargetDesc::Create2DDesc(BufferSize, GetSceneColorFormat(), FClearValueBinding::Black, TexCreate_None, TexCreate_RenderTargetable, false));
 		if (CurrentFeatureLevel >= ERHIFeatureLevel::SM5)
 		{
 			Desc.TargetableFlags |= TexCreate_UAV;
@@ -1715,17 +1726,17 @@ void FSceneRenderTargets::AllocateDeferredShadingPathRenderTargets()
 			FIntPoint Extent = FIntPoint(Res, Res);
 
 			{
-				FPooledRenderTargetDesc Desc(FPooledRenderTargetDesc::Create2DDesc(Extent, PF_R8G8B8A8, TexCreate_None, TexCreate_RenderTargetable, false));
+				FPooledRenderTargetDesc Desc(FPooledRenderTargetDesc::Create2DDesc(Extent, PF_R8G8B8A8, FClearValueBinding::None, TexCreate_None, TexCreate_RenderTargetable, false));
 				GRenderTargetPool.FindFreeElement(Desc, ReflectiveShadowMapNormal, TEXT("RSMNormal"));
 			}
 
 			{
-				FPooledRenderTargetDesc Desc(FPooledRenderTargetDesc::Create2DDesc(Extent, PF_FloatR11G11B10, TexCreate_None, TexCreate_RenderTargetable, false));
+				FPooledRenderTargetDesc Desc(FPooledRenderTargetDesc::Create2DDesc(Extent, PF_FloatR11G11B10, FClearValueBinding::None, TexCreate_None, TexCreate_RenderTargetable, false));
 				GRenderTargetPool.FindFreeElement(Desc, ReflectiveShadowMapDiffuse, TEXT("RSMDiffuse"));
 			}
 
 			{
-				FPooledRenderTargetDesc Desc(FPooledRenderTargetDesc::Create2DDesc(Extent, PF_DepthStencil, TexCreate_None, TexCreate_DepthStencilTargetable , false));
+				FPooledRenderTargetDesc Desc(FPooledRenderTargetDesc::Create2DDesc(Extent, PF_DepthStencil, FClearValueBinding::None, TexCreate_None, TexCreate_DepthStencilTargetable, false));
 				GRenderTargetPool.FindFreeElement(Desc, ReflectiveShadowMapDepth, TEXT("RSMDepth"));
 			}
 		}
@@ -1952,7 +1963,7 @@ IPooledRenderTarget* FSceneRenderTargets::RequestCustomDepth(bool bPrimitives)
 	if((Value == 1  && bPrimitives) || Value == 2 || IsCustomDepthPassWritingStencil())
 	{
 		// Todo: Could check if writes stencil here and create min viable target
-		FPooledRenderTargetDesc Desc(FPooledRenderTargetDesc::Create2DDesc(BufferSize, PF_DepthStencil, TexCreate_None, TexCreate_DepthStencilTargetable, false));
+		FPooledRenderTargetDesc Desc(FPooledRenderTargetDesc::Create2DDesc(BufferSize, PF_DepthStencil, FClearValueBinding::DepthFar, TexCreate_None, TexCreate_DepthStencilTargetable, false));
 		GRenderTargetPool.FindFreeElement(Desc, CustomDepth, TEXT("CustomDepth"));
 		CustomStencilSRV = RHICreateShaderResourceView((FTexture2DRHIRef&)CustomDepth->GetRenderTargetItem().TargetableTexture, 0, 1, PF_X24_G8);
 		return CustomDepth;
diff --git a/Engine/Source/Runtime/Renderer/Private/PostProcess/SceneRenderTargets.h b/Engine/Source/Runtime/Renderer/Private/PostProcess/SceneRenderTargets.h
--- a/Engine/Source/Runtime/Renderer/Private/PostProcess/SceneRenderTargets.h
+++ b/Engine/Source/Runtime/Renderer/Private/PostProcess/SceneRenderTargets.h
@@ -697,12 +697,11 @@ private:
 	bool bAllocateVelocityGBuffer;
 
 	/** Helpers to track gbuffer state on platforms that need to propagate clear information across parallel rendering boundaries. */
-	bool bGBuffersCleared;
-	FLinearColor GBufferClearColor;
+	bool bGBuffersFastCleared;	
 
 	/** Helpers to track scenedepth state on platforms that need to propagate clear information across parallel rendering boundaries. */
 	bool bSceneDepthCleared;
-	float SceneDepthClearValue;
+	
 	/** true is this is a snapshot on the scene allocator */
 	bool bSnapshot;
 	/** All outstanding snapshots */
diff --git a/Engine/Source/Runtime/Renderer/Private/PostProcess/ScreenSpaceReflections.cpp b/Engine/Source/Runtime/Renderer/Private/PostProcess/ScreenSpaceReflections.cpp
--- a/Engine/Source/Runtime/Renderer/Private/PostProcess/ScreenSpaceReflections.cpp
+++ b/Engine/Source/Runtime/Renderer/Private/PostProcess/ScreenSpaceReflections.cpp
@@ -430,7 +430,7 @@ void FRCPassPostProcessScreenSpaceReflections::Process(FRenderingCompositePassCo
 
 FPooledRenderTargetDesc FRCPassPostProcessScreenSpaceReflections::ComputeOutputDesc(EPassOutputId InPassOutputId) const
 {
-	FPooledRenderTargetDesc Ret(FPooledRenderTargetDesc::Create2DDesc(FSceneRenderTargets::Get_FrameConstantsOnly().GetBufferSizeXY(), PF_FloatRGBA, TexCreate_None, TexCreate_RenderTargetable, false));
+	FPooledRenderTargetDesc Ret(FPooledRenderTargetDesc::Create2DDesc(FSceneRenderTargets::Get_FrameConstantsOnly().GetBufferSizeXY(), PF_FloatRGBA, FClearValueBinding::None, TexCreate_None, TexCreate_RenderTargetable, false));
 
 	Ret.DebugName = TEXT("ScreenSpaceReflections");
 
diff --git a/Engine/Source/Runtime/Renderer/Private/PostProcess/VisualizeTexture.cpp b/Engine/Source/Runtime/Renderer/Private/PostProcess/VisualizeTexture.cpp
--- a/Engine/Source/Runtime/Renderer/Private/PostProcess/VisualizeTexture.cpp
+++ b/Engine/Source/Runtime/Renderer/Private/PostProcess/VisualizeTexture.cpp
@@ -345,7 +345,7 @@ void FVisualizeTexture::GenerateContent(FRHICommandListImmediate& RHICmdList, co
 	Size.X = FMath::Max(Size.X, 1);
 	Size.Y = FMath::Max(Size.Y, 1);
 
-	FPooledRenderTargetDesc OutputDesc(FPooledRenderTargetDesc::Create2DDesc(Size, PF_B8G8R8A8, TexCreate_None, TexCreate_RenderTargetable | TexCreate_ShaderResource, false));
+	FPooledRenderTargetDesc OutputDesc(FPooledRenderTargetDesc::Create2DDesc(Size, PF_B8G8R8A8, FClearValueBinding::None, TexCreate_None, TexCreate_RenderTargetable | TexCreate_ShaderResource, false));
 	
 	GRenderTargetPool.FindFreeElement(OutputDesc, VisualizeTextureContent, TEXT("VisualizeTexture"));
 
diff --git a/Engine/Source/Runtime/Renderer/Private/ReflectionEnvironment.cpp b/Engine/Source/Runtime/Renderer/Private/ReflectionEnvironment.cpp
--- a/Engine/Source/Runtime/Renderer/Private/ReflectionEnvironment.cpp
+++ b/Engine/Source/Runtime/Renderer/Private/ReflectionEnvironment.cpp
@@ -96,6 +96,7 @@ void FReflectionEnvironmentCubemapArray::InitDynamicRHI()
 		GReflectionCaptureSize,
 		//@todo - get rid of the alpha channel (currently stores brightness which is a constant), could use PF_FloatRGB for half memory, would need to implement RHIReadSurface support
 		PF_FloatRGBA, 
+		FClearValueBinding::None,
 		TexCreate_None,
 		TexCreate_None,
 		false, 
@@ -848,6 +849,7 @@ void FDeferredShadingSceneRenderer::RenderTiledDeferredImageBasedReflections(FRH
 		FPooledRenderTargetDesc Desc = SceneContext.GetSceneColor()->GetDesc();
 		Desc.TargetableFlags |= TexCreate_UAV;
 		Desc.TargetableFlags |= TexCreate_NoFastClear;
+		Desc.ClearValue = FClearValueBinding::None;
 
 		// we don't create a new name to make it easier to use "vis SceneColor" and get the last HDRSceneColor
 		GRenderTargetPool.FindFreeElement( Desc, NewSceneColor, TEXT("SceneColor") );
@@ -1023,7 +1025,7 @@ void FDeferredShadingSceneRenderer::RenderStandardDeferredImageBasedReflections(
 				// If Reflection Environment is mixed with indirect lighting (Ambient + LPV), skip clear!
 				if (!bMixing)
 				{
-					SimpleRenderTargetMode = ESimpleRenderTargetMode::EClearColorToBlackWithFullAlpha;
+					SimpleRenderTargetMode = ESimpleRenderTargetMode::EClearColorExistingDepth;
 				}
 				SetRenderTarget(RHICmdList, LightAccumulation->GetRenderTargetItem().TargetableTexture, NULL, SimpleRenderTargetMode);
 			}
diff --git a/Engine/Source/Runtime/Renderer/Private/SceneCaptureRendering.cpp b/Engine/Source/Runtime/Renderer/Private/SceneCaptureRendering.cpp
--- a/Engine/Source/Runtime/Renderer/Private/SceneCaptureRendering.cpp
+++ b/Engine/Source/Runtime/Renderer/Private/SceneCaptureRendering.cpp
@@ -91,6 +91,7 @@ static void UpdateSceneCaptureContent_RenderThread(FRHICommandListImmediate& RHI
 			auto& RenderTarget = Target->GetRenderTargetTexture();
 			FPooledRenderTargetDesc Desc(FPooledRenderTargetDesc::Create2DDesc(Target->GetSizeXY(), 
 				RenderTarget.GetReference()->GetFormat(), 
+				FClearValueBinding::None,
 				TexCreate_None, 
 				TexCreate_RenderTargetable,
 				false));
diff --git a/Engine/Source/Runtime/Renderer/Private/SceneHitProxyRendering.cpp b/Engine/Source/Runtime/Renderer/Private/SceneHitProxyRendering.cpp
--- a/Engine/Source/Runtime/Renderer/Private/SceneHitProxyRendering.cpp
+++ b/Engine/Source/Runtime/Renderer/Private/SceneHitProxyRendering.cpp
@@ -340,7 +340,7 @@ TRefCountPtr<IPooledRenderTarget> InitHitProxyRender(FRHICommandListImmediate& R
 
 	// Create a texture to store the resolved light attenuation values, and a render-targetable surface to hold the unresolved light attenuation values.
 	{
-		FPooledRenderTargetDesc Desc(FPooledRenderTargetDesc::Create2DDesc(SceneContext.GetBufferSizeXY(), PF_B8G8R8A8, TexCreate_None, TexCreate_RenderTargetable, false));
+		FPooledRenderTargetDesc Desc(FPooledRenderTargetDesc::Create2DDesc(SceneContext.GetBufferSizeXY(), PF_B8G8R8A8, FClearValueBinding::None, TexCreate_None, TexCreate_RenderTargetable, false));
 		Desc.Flags |= TexCreate_FastVRAM;
 		GRenderTargetPool.FindFreeElement(Desc, HitProxyRT, TEXT("HitProxy"));
 	}
diff --git a/Engine/Source/Runtime/Renderer/Private/SceneOcclusion.cpp b/Engine/Source/Runtime/Renderer/Private/SceneOcclusion.cpp
--- a/Engine/Source/Runtime/Renderer/Private/SceneOcclusion.cpp
+++ b/Engine/Source/Runtime/Renderer/Private/SceneOcclusion.cpp
@@ -521,9 +521,9 @@ void FHZBOcclusionTester::InitDynamicRHI()
 	if (GetFeatureLevel() >= ERHIFeatureLevel::SM4)
 	{
 #if PLATFORM_MAC // Workaround radr://16096028 Texture Readback via glReadPixels + PBOs stalls on Nvidia GPUs
-		FPooledRenderTargetDesc Desc( FPooledRenderTargetDesc::Create2DDesc( FIntPoint( SizeX, SizeY ), PF_R8G8B8A8, TexCreate_CPUReadback | TexCreate_HideInVisualizeTexture, TexCreate_None, false ) );
+		FPooledRenderTargetDesc Desc( FPooledRenderTargetDesc::Create2DDesc( FIntPoint( SizeX, SizeY ), PF_R8G8B8A8, FClearValueBinding::None, TexCreate_CPUReadback | TexCreate_HideInVisualizeTexture, TexCreate_None, false ) );
 #else
-		FPooledRenderTargetDesc Desc( FPooledRenderTargetDesc::Create2DDesc( FIntPoint( SizeX, SizeY ), PF_B8G8R8A8, TexCreate_CPUReadback | TexCreate_HideInVisualizeTexture, TexCreate_None, false ) );
+		FPooledRenderTargetDesc Desc( FPooledRenderTargetDesc::Create2DDesc( FIntPoint( SizeX, SizeY ), PF_B8G8R8A8, FClearValueBinding::None, TexCreate_CPUReadback | TexCreate_HideInVisualizeTexture, TexCreate_None, false ) );
 #endif
 		GRenderTargetPool.FindFreeElement( Desc, ResultsTextureCPU, TEXT("HZBResultsCPU") );
 	}
@@ -720,7 +720,7 @@ void FHZBOcclusionTester::Submit(FRHICommandListImmediate& RHICmdList, const FVi
 	TRefCountPtr< IPooledRenderTarget >	BoundsExtentTexture;
 	{
 		uint32 Flags = TexCreate_ShaderResource | TexCreate_Dynamic;
-		FPooledRenderTargetDesc Desc( FPooledRenderTargetDesc::Create2DDesc( FIntPoint( SizeX, SizeY ), PF_A32B32G32R32F, Flags, TexCreate_None, false ) );
+		FPooledRenderTargetDesc Desc( FPooledRenderTargetDesc::Create2DDesc( FIntPoint( SizeX, SizeY ), PF_A32B32G32R32F, FClearValueBinding::None, Flags, TexCreate_None, false ) );
 
 		GRenderTargetPool.FindFreeElement( Desc, BoundsCenterTexture, TEXT("HZBBoundsCenter") );
 		GRenderTargetPool.FindFreeElement( Desc, BoundsExtentTexture, TEXT("HZBBoundsExtent") );
@@ -729,9 +729,9 @@ void FHZBOcclusionTester::Submit(FRHICommandListImmediate& RHICmdList, const FVi
 	TRefCountPtr< IPooledRenderTarget >	ResultsTextureGPU;
 	{
 #if PLATFORM_MAC // Workaround radr://16096028 Texture Readback via glReadPixels + PBOs stalls on Nvidia GPUs
-		FPooledRenderTargetDesc Desc( FPooledRenderTargetDesc::Create2DDesc( FIntPoint( SizeX, SizeY ), PF_R8G8B8A8, TexCreate_None, TexCreate_RenderTargetable, false ) );
+		FPooledRenderTargetDesc Desc( FPooledRenderTargetDesc::Create2DDesc( FIntPoint( SizeX, SizeY ), PF_R8G8B8A8, FClearValueBinding::None, TexCreate_None, TexCreate_RenderTargetable, false ) );
 #else
-		FPooledRenderTargetDesc Desc( FPooledRenderTargetDesc::Create2DDesc( FIntPoint( SizeX, SizeY ), PF_B8G8R8A8, TexCreate_None, TexCreate_RenderTargetable, false ) );
+		FPooledRenderTargetDesc Desc( FPooledRenderTargetDesc::Create2DDesc( FIntPoint( SizeX, SizeY ), PF_B8G8R8A8, FClearValueBinding::None, TexCreate_None, TexCreate_RenderTargetable, false ) );
 #endif
 		GRenderTargetPool.FindFreeElement( Desc, ResultsTextureGPU, TEXT("HZBResultsGPU") );
 	}
@@ -995,7 +995,7 @@ void BuildHZB( FRHICommandListImmediate& RHICmdList, FViewInfo& View )
 	const FIntPoint HZBSize( 1 << NumMipsX, 1 << NumMipsY );
 	View.HZBMipmap0Size = HZBSize;
 
-	FPooledRenderTargetDesc Desc(FPooledRenderTargetDesc::Create2DDesc(HZBSize, PF_R16F, TexCreate_None, TexCreate_RenderTargetable | TexCreate_ShaderResource | TexCreate_NoFastClear, false, NumMips));
+	FPooledRenderTargetDesc Desc(FPooledRenderTargetDesc::Create2DDesc(HZBSize, PF_R16F, FClearValueBinding::None, TexCreate_None, TexCreate_RenderTargetable | TexCreate_ShaderResource | TexCreate_NoFastClear, false, NumMips));
 	Desc.Flags |= TexCreate_FastVRAM;
 	GRenderTargetPool.FindFreeElement( Desc, View.HZB, TEXT("HZB") );
 	
diff --git a/Engine/Source/Runtime/Renderer/Private/ScenePrivate.h b/Engine/Source/Runtime/Renderer/Private/ScenePrivate.h
--- a/Engine/Source/Runtime/Renderer/Private/ScenePrivate.h
+++ b/Engine/Source/Runtime/Renderer/Private/ScenePrivate.h
@@ -643,7 +643,7 @@ public:
 		if (!EyeAdaptationRT)
 		{
 			// Create the texture needed for EyeAdaptation
-			FPooledRenderTargetDesc Desc(FPooledRenderTargetDesc::Create2DDesc(FIntPoint(1, 1), PF_R32_FLOAT, TexCreate_None, TexCreate_RenderTargetable, false));
+			FPooledRenderTargetDesc Desc(FPooledRenderTargetDesc::Create2DDesc(FIntPoint(1, 1), PF_R32_FLOAT, FClearValueBinding::None, TexCreate_None, TexCreate_RenderTargetable, false));
 			GRenderTargetPool.FindFreeElement(Desc, EyeAdaptationRT, TEXT("EyeAdaptation"));
 		}
 		return EyeAdaptationRT;
@@ -654,7 +654,7 @@ public:
 		if (!SeparateTranslucencyRT || SeparateTranslucencyRT->GetDesc().Extent != Size)
 		{
 			// Create the SeparateTranslucency render target (alpha is needed to lerping)
-			FPooledRenderTargetDesc Desc(FPooledRenderTargetDesc::Create2DDesc(Size, PF_FloatRGBA, TexCreate_None, TexCreate_RenderTargetable, false));
+			FPooledRenderTargetDesc Desc(FPooledRenderTargetDesc::Create2DDesc(Size, PF_FloatRGBA, FClearValueBinding::Black, TexCreate_None, TexCreate_RenderTargetable, false));
 			GRenderTargetPool.FindFreeElement(Desc, SeparateTranslucencyRT, TEXT("SeparateTranslucency"));
 		}
 		return SeparateTranslucencyRT;
diff --git a/Engine/Source/Runtime/Renderer/Private/SceneRendering.cpp b/Engine/Source/Runtime/Renderer/Private/SceneRendering.cpp
--- a/Engine/Source/Runtime/Renderer/Private/SceneRendering.cpp
+++ b/Engine/Source/Runtime/Renderer/Private/SceneRendering.cpp
@@ -1571,7 +1571,7 @@ void FRendererModule::CreateAndInitSingleView(FRHICommandListImmediate& RHICmdLi
 	// Create and add the new view
 	FViewInfo* NewView = new FViewInfo(*ViewInitOptions);
 	ViewFamily->Views.Add(NewView);
-	SetRenderTarget(RHICmdList, ViewFamily->RenderTarget->GetRenderTargetTexture(), nullptr, ESimpleRenderTargetMode::EClearColorToWhite);
+	SetRenderTarget(RHICmdList, ViewFamily->RenderTarget->GetRenderTargetTexture(), nullptr, ESimpleRenderTargetMode::EClearColorExistingDepth);
 	FViewInfo* View = (FViewInfo*)ViewFamily->Views[0];
 	View->InitRHIResources(nullptr);
 }
diff --git a/Engine/Source/Runtime/Renderer/Private/ShadowRendering.cpp b/Engine/Source/Runtime/Renderer/Private/ShadowRendering.cpp
--- a/Engine/Source/Runtime/Renderer/Private/ShadowRendering.cpp
+++ b/Engine/Source/Runtime/Renderer/Private/ShadowRendering.cpp
@@ -1264,7 +1264,7 @@ void FProjectedShadowInfo::ClearDepth(FRHICommandList& RHICmdList, FDeferredShad
 	}
 	else
 	{
-		RHICmdList.BindClearMRTValues(bClearColor, NumClearColors, Colors, bClearDepth, Depth, bClearStencil, Stencil);
+		RHICmdList.BindClearMRTValues(bClearColor, bClearDepth, bClearStencil);
 	}
 }
 
diff --git a/Engine/Source/Runtime/Renderer/Private/SystemTextures.cpp b/Engine/Source/Runtime/Renderer/Private/SystemTextures.cpp
--- a/Engine/Source/Runtime/Renderer/Private/SystemTextures.cpp
+++ b/Engine/Source/Runtime/Renderer/Private/SystemTextures.cpp
@@ -32,7 +32,7 @@ void FSystemTextures::InitializeTextures(FRHICommandListImmediate& RHICmdList, E
 	{
 		// Create a dummy white texture.
 		{
-			FPooledRenderTargetDesc Desc(FPooledRenderTargetDesc::Create2DDesc(FIntPoint(1, 1), PF_B8G8R8A8, TexCreate_HideInVisualizeTexture, TexCreate_RenderTargetable | TexCreate_NoFastClear, false));
+			FPooledRenderTargetDesc Desc(FPooledRenderTargetDesc::Create2DDesc(FIntPoint(1, 1), PF_B8G8R8A8, FClearValueBinding::None, TexCreate_HideInVisualizeTexture, TexCreate_RenderTargetable | TexCreate_NoFastClear, false));
 			GRenderTargetPool.FindFreeElement(Desc, WhiteDummy, TEXT("WhiteDummy"));
 
 			SetRenderTarget(RHICmdList, WhiteDummy->GetRenderTargetItem().TargetableTexture, FTextureRHIRef());
@@ -42,7 +42,7 @@ void FSystemTextures::InitializeTextures(FRHICommandListImmediate& RHICmdList, E
 
 		// Create a dummy black texture.
 		{
-			FPooledRenderTargetDesc Desc(FPooledRenderTargetDesc::Create2DDesc(FIntPoint(1, 1), PF_B8G8R8A8, TexCreate_HideInVisualizeTexture, TexCreate_RenderTargetable | TexCreate_NoFastClear, false));
+			FPooledRenderTargetDesc Desc(FPooledRenderTargetDesc::Create2DDesc(FIntPoint(1, 1), PF_B8G8R8A8, FClearValueBinding::None, TexCreate_HideInVisualizeTexture, TexCreate_RenderTargetable | TexCreate_NoFastClear, false));
 			GRenderTargetPool.FindFreeElement(Desc, BlackDummy, TEXT("BlackDummy"));
 
 			SetRenderTarget(RHICmdList, BlackDummy->GetRenderTargetItem().TargetableTexture, FTextureRHIRef());
@@ -52,7 +52,7 @@ void FSystemTextures::InitializeTextures(FRHICommandListImmediate& RHICmdList, E
 
 		// Create the PerlinNoiseGradient texture
 		{
-			FPooledRenderTargetDesc Desc(FPooledRenderTargetDesc::Create2DDesc(FIntPoint(128, 128), PF_B8G8R8A8, TexCreate_HideInVisualizeTexture, TexCreate_None | TexCreate_NoFastClear, false));
+			FPooledRenderTargetDesc Desc(FPooledRenderTargetDesc::Create2DDesc(FIntPoint(128, 128), PF_B8G8R8A8, FClearValueBinding::None, TexCreate_HideInVisualizeTexture, TexCreate_None | TexCreate_NoFastClear, false));
 			GRenderTargetPool.FindFreeElement(Desc, PerlinNoiseGradient, TEXT("PerlinNoiseGradient"));
 			// Write the contents of the texture.
 			uint32 DestStride;
@@ -83,7 +83,7 @@ void FSystemTextures::InitializeTextures(FRHICommandListImmediate& RHICmdList, E
 
 		if (!GSupportsShaderFramebufferFetch && GPixelFormats[PF_FloatRGBA].Supported)
 		{
-			FPooledRenderTargetDesc Desc(FPooledRenderTargetDesc::Create2DDesc(FIntPoint(1, 1), PF_FloatRGBA, TexCreate_HideInVisualizeTexture, TexCreate_RenderTargetable | TexCreate_NoFastClear, false));
+			FPooledRenderTargetDesc Desc(FPooledRenderTargetDesc::Create2DDesc(FIntPoint(1, 1), PF_FloatRGBA, FClearValueBinding::None, TexCreate_HideInVisualizeTexture, TexCreate_RenderTargetable | TexCreate_NoFastClear, false));
 			GRenderTargetPool.FindFreeElement(Desc, MaxFP16Depth, TEXT("MaxFP16Depth"));
 
 			SetRenderTarget(RHICmdList, MaxFP16Depth->GetRenderTargetItem().TargetableTexture, FTextureRHIRef());
@@ -99,7 +99,7 @@ void FSystemTextures::InitializeTextures(FRHICommandListImmediate& RHICmdList, E
 
 		const uint32 Square = Extent * Extent;
 
-		FPooledRenderTargetDesc Desc(FPooledRenderTargetDesc::CreateVolumeDesc(Extent, Extent, Extent, PF_B8G8R8A8, TexCreate_ShaderResource | TexCreate_HideInVisualizeTexture | TexCreate_NoTiling, TexCreate_None, false));
+		FPooledRenderTargetDesc Desc(FPooledRenderTargetDesc::CreateVolumeDesc(Extent, Extent, Extent, PF_B8G8R8A8, FClearValueBinding::None, TexCreate_ShaderResource | TexCreate_HideInVisualizeTexture | TexCreate_NoTiling, TexCreate_None, false));
 		GRenderTargetPool.FindFreeElement(Desc, PerlinNoise3D, TEXT("PerlinNoise3D"));
 		// Write the contents of the texture.
 		TArray<uint32> DestBuffer;
@@ -232,7 +232,7 @@ void FSystemTextures::InitializeTextures(FRHICommandListImmediate& RHICmdList, E
 		}
 
 		{
-			FPooledRenderTargetDesc Desc(FPooledRenderTargetDesc::Create2DDesc(FIntPoint(64, 64), PF_B8G8R8A8, TexCreate_HideInVisualizeTexture, TexCreate_None | TexCreate_NoFastClear, false));
+			FPooledRenderTargetDesc Desc(FPooledRenderTargetDesc::Create2DDesc(FIntPoint(64, 64), PF_B8G8R8A8, FClearValueBinding::None, TexCreate_HideInVisualizeTexture, TexCreate_None | TexCreate_NoFastClear, false));
 			GRenderTargetPool.FindFreeElement(Desc, SSAORandomization, TEXT("SSAORandomization"));
 			// Write the contents of the texture.
 			uint32 DestStride;
@@ -263,7 +263,7 @@ void FSystemTextures::InitializeTextures(FRHICommandListImmediate& RHICmdList, E
 				Format = PF_G16R16;
 			}
 
-			FPooledRenderTargetDesc Desc(FPooledRenderTargetDesc::Create2DDesc(FIntPoint(128, 32), Format, TexCreate_FastVRAM, TexCreate_None, false));
+			FPooledRenderTargetDesc Desc(FPooledRenderTargetDesc::Create2DDesc(FIntPoint(128, 32), Format, FClearValueBinding::None, TexCreate_FastVRAM, TexCreate_None, false));
 
 			if (bReference)
 			{
diff --git a/Engine/Source/Runtime/Renderer/Private/TranslucentLighting.cpp b/Engine/Source/Runtime/Renderer/Private/TranslucentLighting.cpp
--- a/Engine/Source/Runtime/Renderer/Private/TranslucentLighting.cpp
+++ b/Engine/Source/Runtime/Renderer/Private/TranslucentLighting.cpp
@@ -1106,17 +1106,8 @@ void ClearVolumeTextures(FRHICommandList& RHICmdList, ERHIFeatureLevel::Type Fea
 		SetGlobalBoundShaderState(RHICmdList, FeatureLevel, VolumeClearBoundShaderState, GScreenVertexDeclaration.VertexDeclarationRHI, *VertexShader, *PixelShader, *GeometryShader);
 
 		VertexShader->SetParameters(RHICmdList, VolumeBounds, GTranslucencyLightingVolumeDim);
-		GeometryShader->SetParameters(RHICmdList, VolumeBounds);
-
-		FLinearColor ShaderClearColors[MaxSimultaneousRenderTargets];
-		FMemory::Memzero(ShaderClearColors);
-
-		for (int32 i = 0; i < NumRenderTargets; i++)
-		{
-			ShaderClearColors[i] = ClearColors[i];
-		}
-
-		SetShaderValueArray(RHICmdList, PixelShader->GetPixelShader(), PixelShader->ColorParameter, ShaderClearColors, NumRenderTargets);
+		GeometryShader->SetParameters(RHICmdList, VolumeBounds);		
+		PixelShader->SetColors(RHICmdList, ClearColors, NumRenderTargets);
 
 		RasterizeToVolumeTexture(RHICmdList, VolumeBounds);
 	}
diff --git a/Engine/Source/Runtime/Renderer/Private/VelocityRendering.cpp b/Engine/Source/Runtime/Renderer/Private/VelocityRendering.cpp
--- a/Engine/Source/Runtime/Renderer/Private/VelocityRendering.cpp
+++ b/Engine/Source/Runtime/Renderer/Private/VelocityRendering.cpp
@@ -581,16 +581,14 @@ public:
 static void BeginVelocityRendering(FRHICommandList& RHICmdList, TRefCountPtr<IPooledRenderTarget>& VelocityRT, bool bPerformClear)
 {
 	FTextureRHIRef VelocityTexture = VelocityRT->GetRenderTargetItem().TargetableTexture;
-	FTexture2DRHIRef DepthTexture = FSceneRenderTargets::Get(RHICmdList).GetSceneDepthTexture();
-	FLinearColor VelocityClearColor = FLinearColor::Black;
+	FTexture2DRHIRef DepthTexture = FSceneRenderTargets::Get(RHICmdList).GetSceneDepthTexture();	
 	if (bPerformClear)
 	{
 		// now make the FRHISetRenderTargetsInfo that encapsulates all of the info
 		FRHIRenderTargetView ColorView(VelocityTexture, 0, -1, ERenderTargetLoadAction::EClear, ERenderTargetStoreAction::EStore);
 		FRHIDepthRenderTargetView DepthView(DepthTexture, ERenderTargetLoadAction::ELoad, ERenderTargetStoreAction::ENoAction, FExclusiveDepthStencil::DepthRead_StencilWrite);
 
-		FRHISetRenderTargetsInfo Info(1, &ColorView, DepthView);
-		Info.ClearColors[0] = VelocityClearColor;
+		FRHISetRenderTargetsInfo Info(1, &ColorView, DepthView);		
 
 		// Clear the velocity buffer (0.0f means "use static background velocity").
 		RHICmdList.SetRenderTargetsAndClear(Info);		
@@ -601,7 +599,7 @@ static void BeginVelocityRendering(FRHICommandList& RHICmdList, TRefCountPtr<IPo
 
 		// some platforms need the clear color when rendertargets transition to SRVs.  We propagate here to allow parallel rendering to always
 		// have the proper mapping when the RT is transitioned.
-		RHICmdList.BindClearMRTValues(true, 1, &VelocityClearColor, false, (float)ERHIZBuffer::FarPlane, false, 0);
+		RHICmdList.BindClearMRTValues(true, false, false);
 	}
 }
 
@@ -792,7 +790,7 @@ FPooledRenderTargetDesc FVelocityRendering::GetRenderTargetDesc()
 {
 	const FIntPoint BufferSize = FSceneRenderTargets::Get_FrameConstantsOnly().GetBufferSizeXY();
 	const FIntPoint VelocityBufferSize = BufferSize;		// full resolution so we can reuse the existing full res z buffer
-	return FPooledRenderTargetDesc(FPooledRenderTargetDesc::Create2DDesc(VelocityBufferSize, PF_G16R16, TexCreate_None, TexCreate_RenderTargetable, false));
+	return FPooledRenderTargetDesc(FPooledRenderTargetDesc::Create2DDesc(VelocityBufferSize, PF_G16R16, FClearValueBinding::Transparent, TexCreate_None, TexCreate_RenderTargetable, false));
 }
 
 bool FVelocityRendering::OutputsToGBuffer()
diff --git a/Engine/Source/Runtime/Renderer/Public/RendererInterface.h b/Engine/Source/Runtime/Renderer/Public/RendererInterface.h
--- a/Engine/Source/Runtime/Renderer/Public/RendererInterface.h
+++ b/Engine/Source/Runtime/Renderer/Public/RendererInterface.h
@@ -34,7 +34,8 @@ public:
 
 	/** Default constructor, use one of the factory functions below to make a valid description */
 	FPooledRenderTargetDesc()
-		: Extent(0, 0)
+		: ClearValue(FClearValueBinding())
+		, Extent(0, 0)
 		, Depth(0)
 		, ArraySize(1)
 		, bIsArray(false)
@@ -56,6 +57,7 @@ public:
 	static FPooledRenderTargetDesc Create2DDesc(
 		FIntPoint InExtent,
 		EPixelFormat InFormat,
+		const FClearValueBinding& InClearValue,
 		uint32 InFlags,
 		uint32 InTargetableFlags,
 		bool bInForceSeparateTargetAndShaderResource,
@@ -65,6 +67,7 @@ public:
 		check(InExtent.Y);
 
 		FPooledRenderTargetDesc NewDesc;
+		NewDesc.ClearValue = InClearValue;
 		NewDesc.Extent = InExtent;
 		NewDesc.Depth = 0;
 		NewDesc.ArraySize = 1;
@@ -89,6 +92,7 @@ public:
 		uint32 InSizeY,
 		uint32 InSizeZ,
 		EPixelFormat InFormat,
+		const FClearValueBinding& InClearValue,
 		uint32 InFlags,
 		uint32 InTargetableFlags,
 		bool bInForceSeparateTargetAndShaderResource,
@@ -98,6 +102,7 @@ public:
 		check(InSizeY);
 
 		FPooledRenderTargetDesc NewDesc;
+		NewDesc.ClearValue = InClearValue;
 		NewDesc.Extent = FIntPoint(InSizeX,InSizeY);
 		NewDesc.Depth = InSizeZ;
 		NewDesc.ArraySize = 1;
@@ -120,6 +125,7 @@ public:
 	static FPooledRenderTargetDesc CreateCubemapDesc(
 		uint32 InExtent,
 		EPixelFormat InFormat,
+		const FClearValueBinding& InClearValue,
 		uint32 InFlags,
 		uint32 InTargetableFlags,
 		bool bInForceSeparateTargetAndShaderResource,
@@ -129,6 +135,7 @@ public:
 		check(InExtent);
 
 		FPooledRenderTargetDesc NewDesc;
+		NewDesc.ClearValue = InClearValue;
 		NewDesc.Extent = FIntPoint(InExtent, 0);
 		NewDesc.Depth = 0;
 		NewDesc.ArraySize = InArraySize;
@@ -167,7 +174,8 @@ public:
 			&& Format == rhs.Format
 			&& LhsFlags == RhsFlags
 			&& TargetableFlags == rhs.TargetableFlags
-			&& bForceSeparateTargetAndShaderResource == rhs.bForceSeparateTargetAndShaderResource;
+			&& bForceSeparateTargetAndShaderResource == rhs.bForceSeparateTargetAndShaderResource
+			&& ClearValue == rhs.ClearValue;
 	}
 
 	bool IsCubemap() const
@@ -282,6 +290,9 @@ public:
 		TargetableFlags &= (~TexCreate_UAV);
 	}
 
+	/** Value allowed for fast clears for this target. */
+	FClearValueBinding ClearValue;
+
 	/** In pixels, (0,0) if not set, (x,0) for cube maps, todo: make 3d int vector for volume textures */
 	FIntPoint Extent;
 	/** 0, unless it's texture array or volume texture */
diff --git a/Engine/Source/Runtime/SlateRHIRenderer/Private/Slate3DRenderer.cpp b/Engine/Source/Runtime/SlateRHIRenderer/Private/Slate3DRenderer.cpp
--- a/Engine/Source/Runtime/SlateRHIRenderer/Private/Slate3DRenderer.cpp
+++ b/Engine/Source/Runtime/SlateRHIRenderer/Private/Slate3DRenderer.cpp
@@ -96,7 +96,7 @@ void FSlate3DRenderer::DrawWindowToTarget_RenderThread( FRHICommandListImmediate
 	FRHIRenderTargetView ColorRTV(RenderTargetResource->GetTextureRHI());
 	FRHISetRenderTargetsInfo Info(1, &ColorRTV, FTextureRHIParamRef());
 	Info.bClearColor = true;
-	Info.ClearColors[0] = RenderTarget->ClearColor;
+	ensure(ColorRTV.Texture->GetClearColor() == RenderTarget->ClearColor);
 
 	RHICmdList.SetRenderTargetsAndClear(Info);
 
diff --git a/Engine/Source/Runtime/UtilityShaders/Private/OneColorShader.cpp b/Engine/Source/Runtime/UtilityShaders/Private/OneColorShader.cpp
--- a/Engine/Source/Runtime/UtilityShaders/Private/OneColorShader.cpp
+++ b/Engine/Source/Runtime/UtilityShaders/Private/OneColorShader.cpp
@@ -4,6 +4,39 @@
 #include "UtilityShadersPrivatePCH.h"
 #include "OneColorShader.h"
 
+BEGIN_UNIFORM_BUFFER_STRUCT(FClearShaderUB, )
+DECLARE_UNIFORM_BUFFER_STRUCT_MEMBER_ARRAY(FVector4, DrawColorMRT, [8] )
+END_UNIFORM_BUFFER_STRUCT(FClearShaderUB)
+
+IMPLEMENT_UNIFORM_BUFFER_STRUCT(FClearShaderUB, TEXT("ClearShaderUB"));
+
+void FOneColorPS::SetColors(FRHICommandList& RHICmdList, const FLinearColor* Colors, int32 NumColors)
+{
+	check(NumColors <= MaxSimultaneousRenderTargets);
+
+	const FShaderUniformBufferParameter& ClearUBParam = GetUniformBufferParameter<FClearShaderUB>();
+	if (ClearUBParam.IsInitialized())
+	{
+		if (ClearUBParam.IsBound())
+		{
+			FClearShaderUB ClearData;
+			FMemory::Memzero(ClearData.DrawColorMRT);			
+			for (int32 i = 0; i < NumColors; ++i)
+			{
+				ClearData.DrawColorMRT[i].X = Colors[i].R;
+				ClearData.DrawColorMRT[i].Y = Colors[i].G;
+				ClearData.DrawColorMRT[i].Z = Colors[i].B;
+				ClearData.DrawColorMRT[i].W = Colors[i].A;
+			}
+
+			TUniformBufferRef<FClearShaderUB> Data = TUniformBufferRef<FClearShaderUB>::CreateUniformBufferImmediate(ClearData, UniformBuffer_SingleFrame);			
+			SetUniformBufferParameter(RHICmdList, GetPixelShader(), GetUniformBufferParameter<FClearShaderUB>(), Data);
+		}
+	}
+
+	
+}
+
 IMPLEMENT_SHADER_TYPE(template<> UTILITYSHADERS_API, TOneColorVS<true>,TEXT("OneColorShader"),TEXT("MainVertexShader"),SF_Vertex);
 IMPLEMENT_SHADER_TYPE(template<> UTILITYSHADERS_API, TOneColorVS<false>,TEXT("OneColorShader"),TEXT("MainVertexShader"),SF_Vertex);
 IMPLEMENT_SHADER_TYPE(,FOneColorPS,TEXT("OneColorShader"),TEXT("MainPixelShader"),SF_Pixel);
diff --git a/Engine/Source/Runtime/UtilityShaders/Public/OneColorShader.h b/Engine/Source/Runtime/UtilityShaders/Public/OneColorShader.h
--- a/Engine/Source/Runtime/UtilityShaders/Public/OneColorShader.h
+++ b/Engine/Source/Runtime/UtilityShaders/Public/OneColorShader.h
@@ -42,18 +42,21 @@ class FOneColorPS : public FGlobalShader
 {
 	DECLARE_EXPORTED_SHADER_TYPE(FOneColorPS,Global,UTILITYSHADERS_API);
 public:
+	
 	FOneColorPS( )	{ }
 	FOneColorPS(const ShaderMetaType::CompiledShaderInitializerType& Initializer)
 	:	FGlobalShader( Initializer )
 	{
-		ColorParameter.Bind( Initializer.ParameterMap, TEXT("DrawColorMRT"), SPF_Mandatory);
+		//ColorParameter.Bind( Initializer.ParameterMap, TEXT("DrawColorMRT"), SPF_Mandatory);
 	}
 
+	virtual void SetColors(FRHICommandList& RHICmdList, const FLinearColor* Colors, int32 NumColors);	
+
 	// FShader interface.
 	virtual bool Serialize(FArchive& Ar) override
 	{
 		bool bShaderHasOutdatedParameters = FGlobalShader::Serialize(Ar);
-		Ar << ColorParameter;
+		//Ar << ColorParameter;
 		return bShaderHasOutdatedParameters;
 	}
 	static bool ShouldCache(EShaderPlatform Platform)
@@ -62,7 +65,7 @@ public:
 	}
 
 	/** The parameter to use for setting the draw Color. */
-	FShaderParameter ColorParameter;
+	//FShaderParameter ColorParameter;
 };
 
 /**
diff --git a/Engine/Source/Runtime/Windows/D3D11RHI/Private/D3D11Commands.cpp b/Engine/Source/Runtime/Windows/D3D11RHI/Private/D3D11Commands.cpp
--- a/Engine/Source/Runtime/Windows/D3D11RHI/Private/D3D11Commands.cpp
+++ b/Engine/Source/Runtime/Windows/D3D11RHI/Private/D3D11Commands.cpp
@@ -966,7 +966,27 @@ void FD3D11DynamicRHI::RHISetRenderTargetsAndClear(const FRHISetRenderTargetsInf
 		nullptr);
 	if (RenderTargetsInfo.bClearColor || RenderTargetsInfo.bClearStencil || RenderTargetsInfo.bClearDepth)
 	{
-		this->RHIClearMRT(RenderTargetsInfo.bClearColor, RenderTargetsInfo.NumColorRenderTargets, RenderTargetsInfo.ClearColors, RenderTargetsInfo.bClearDepth, RenderTargetsInfo.DepthClearValue, RenderTargetsInfo.bClearStencil, RenderTargetsInfo.StencilClearValue, FIntRect());
+		FLinearColor ClearColors[MaxSimultaneousRenderTargets];
+		float DepthClear = 0.0;
+		uint32 StencilClear = 0;
+
+		if (RenderTargetsInfo.bClearColor)
+		{
+			for (int32 i = 0; i < RenderTargetsInfo.NumColorRenderTargets; ++i)
+			{
+				const FClearValueBinding& ClearValue = RenderTargetsInfo.ColorRenderTarget[i].Texture->GetClearBinding();
+				checkf(ClearValue.ColorBinding == EClearBinding::EColorBound, TEXT("Texture: %s does not have a color bound for fast clears"), *RenderTargetsInfo.ColorRenderTarget[i].Texture->GetName().GetPlainNameString());
+				ClearColors[i] = ClearValue.GetClearColor();
+			}
+		}
+		if (RenderTargetsInfo.bClearDepth || RenderTargetsInfo.bClearStencil)
+		{
+			const FClearValueBinding& ClearValue = RenderTargetsInfo.DepthStencilRenderTarget.Texture->GetClearBinding();
+			checkf(ClearValue.ColorBinding == EClearBinding::EDepthStencilBound, TEXT("Texture: %s does not have a DS value bound for fast clears"), *RenderTargetsInfo.DepthStencilRenderTarget.Texture->GetName().GetPlainNameString());
+			ClearValue.GetDepthStencil(DepthClear, StencilClear);
+		}
+
+		this->RHIClearMRTImpl(RenderTargetsInfo.bClearColor, RenderTargetsInfo.NumColorRenderTargets, ClearColors, RenderTargetsInfo.bClearDepth, DepthClear, RenderTargetsInfo.bClearStencil, StencilClear, FIntRect(), false);
 	}
 }
 
@@ -1470,11 +1490,22 @@ void FD3D11DynamicRHI::RHIEndDrawIndexedPrimitiveUP()
 // Raster operations.
 void FD3D11DynamicRHI::RHIClear(bool bClearColor,const FLinearColor& Color,bool bClearDepth,float Depth,bool bClearStencil,uint32 Stencil, FIntRect ExcludeRect)
 {
-	FD3D11DynamicRHI::RHIClearMRT(bClearColor, 1, &Color, bClearDepth, Depth, bClearStencil, Stencil, ExcludeRect);
+	FD3D11DynamicRHI::RHIClearMRTImpl(bClearColor, 1, &Color, bClearDepth, Depth, bClearStencil, Stencil, ExcludeRect, true);
 }
 
-void FD3D11DynamicRHI::RHIClearMRT(bool bClearColor,int32 NumClearColors,const FLinearColor* ClearColorArray,bool bClearDepth,float Depth,bool bClearStencil,uint32 Stencil, FIntRect ExcludeRect)
+void FD3D11DynamicRHI::RHIClearMRT(bool bClearColor, int32 NumClearColors, const FLinearColor* ClearColorArray, bool bClearDepth, float Depth, bool bClearStencil, uint32 Stencil, FIntRect ExcludeRect)
 {
+	RHIClearMRTImpl(bClearColor, NumClearColors, ClearColorArray, bClearDepth, Depth, bClearStencil, Stencil, ExcludeRect, true);
+}
+
+bool gNoShaderClears = true;
+void FD3D11DynamicRHI::RHIClearMRTImpl(bool bClearColor, int32 NumClearColors, const FLinearColor* ClearColorArray, bool bClearDepth, float Depth, bool bClearStencil, uint32 Stencil, FIntRect ExcludeRect, bool bForceShaderClear)
+{	
+	if (gNoShaderClears)
+	{
+		bForceShaderClear = false;
+	}
+
 	// Helper struct to record and restore device states RHIClearMRT modifies.
 	class FDeviceStateHelper
 	{
@@ -1484,17 +1515,24 @@ void FD3D11DynamicRHI::RHIClearMRT(bool bClearColor,int32 NumClearColors,const F
 		TRefCountPtr<FD3D11DeviceContext> Direct3DDeviceIMContext;
 
 		enum { ResourceCount = D3D11_COMMONSHADER_INPUT_RESOURCE_SLOT_COUNT };
+		enum { ConstantBufferCount = D3D11_COMMONSHADER_CONSTANT_BUFFER_API_SLOT_COUNT };
+
 		//////////////////////////////////////////////////////////////////////////
 		// Relevant recorded states:
 		ID3D11ShaderResourceView* VertResources[ResourceCount];
+		ID3D11Buffer* VertexConstantBuffers[ConstantBufferCount];
+		ID3D11Buffer* PixelConstantBuffers[ConstantBufferCount];
 		ID3D11VertexShader* VSOld;
 		ID3D11PixelShader* PSOld;
 		ID3D11DepthStencilState* OldDepthStencilState;
 		ID3D11RasterizerState* OldRasterizerState;
 		ID3D11BlendState* OldBlendState;
+		ID3D11InputLayout* OldInputLayout;
 		uint32 StencilRef;
 		float BlendFactor[4];
 		uint32 SampleMask;
+		FBoundShaderStateRHIParamRef LastBoundShaderStateRHI;
+
 		//////////////////////////////////////////////////////////////////////////
 		void ReleaseResources()
 		{
@@ -1506,23 +1544,33 @@ void FD3D11DynamicRHI::RHIClearMRT(bool bClearColor,int32 NumClearColors,const F
 			{
 				SAFE_RELEASE(*Resources);
 			}
+			for (int32 i = 0; i < ConstantBufferCount; ++i)
+			{
+				SAFE_RELEASE(VertexConstantBuffers[i]);
+				SAFE_RELEASE(PixelConstantBuffers[i]);
+			}
 
 			SAFE_RELEASE(OldDepthStencilState);
 			SAFE_RELEASE(OldBlendState);
 			SAFE_RELEASE(OldRasterizerState);
+			SAFE_RELEASE(OldInputLayout);
 		}
 	public:
 		/** The global D3D device's immediate context */
 		FDeviceStateHelper(TRefCountPtr<FD3D11DeviceContext> InDirect3DDeviceIMContext) : Direct3DDeviceIMContext(InDirect3DDeviceIMContext) {}
 
-		void CaptureDeviceState(FD3D11StateCache& StateCacheRef)
-		{
+		void CaptureDeviceState(FD3D11StateCache& StateCacheRef, TGlobalResource< TBoundShaderStateHistory<10000> >& BSSHistory)
+		{			
 			StateCacheRef.GetVertexShader(&VSOld);
 			StateCacheRef.GetPixelShader(&PSOld);
 			StateCacheRef.GetShaderResourceViews<SF_Vertex>(0, ResourceCount, &VertResources[0]);
+			StateCacheRef.GetConstantBuffers<SF_Pixel>(0, ConstantBufferCount, &(PixelConstantBuffers[0]));
+			StateCacheRef.GetConstantBuffers<SF_Vertex>(0, ConstantBufferCount, &(VertexConstantBuffers[0]));
 			StateCacheRef.GetDepthStencilState(&OldDepthStencilState, &StencilRef);
 			StateCacheRef.GetBlendState(&OldBlendState, BlendFactor, &SampleMask);
 			StateCacheRef.GetRasterizerState(&OldRasterizerState);
+			StateCacheRef.GetInputLayout(&OldInputLayout);
+			LastBoundShaderStateRHI = BSSHistory.GetLast();
 		}
 
 		void ClearCurrentVertexResources(FD3D11StateCache& StateCacheRef)
@@ -1534,18 +1582,28 @@ void FD3D11DynamicRHI::RHIClearMRT(bool bClearColor,int32 NumClearColors,const F
 			}
 		}
 
-		void RestoreDeviceState(FD3D11StateCache& StateCacheRef) 
+		void RestoreDeviceState(FD3D11StateCache& StateCacheRef, TGlobalResource< TBoundShaderStateHistory<10000> >& BSSHistory)
 		{
 
 			// Restore the old shaders
 			StateCacheRef.SetVertexShader(VSOld);
 			StateCacheRef.SetPixelShader(PSOld);
-			for (int ResourceLoop = 0 ; ResourceLoop < ResourceCount; ResourceLoop++)
-				StateCacheRef.SetShaderResourceView<SF_Vertex>(VertResources[ResourceLoop],ResourceLoop);
+			for (int ResourceLoop = 0; ResourceLoop < ResourceCount; ResourceLoop++)
+			{
+				StateCacheRef.SetShaderResourceView<SF_Vertex>(VertResources[ResourceLoop], ResourceLoop);
+			}
+			for (int BufferIndex = 0; BufferIndex < ConstantBufferCount; ++BufferIndex)
+			{
+				StateCacheRef.SetConstantBuffer<SF_Pixel>(PixelConstantBuffers[BufferIndex], BufferIndex);
+				StateCacheRef.SetConstantBuffer<SF_Vertex>(VertexConstantBuffers[BufferIndex], BufferIndex);
+			}
+
 			StateCacheRef.SetDepthStencilState(OldDepthStencilState, StencilRef);
 			StateCacheRef.SetBlendState(OldBlendState, BlendFactor, SampleMask);
 			StateCacheRef.SetRasterizerState(OldRasterizerState);
+			StateCacheRef.SetInputLayout(OldInputLayout);
 
+			BSSHistory.Add(LastBoundShaderStateRHI);
 			ReleaseResources();
 		}
 	};
@@ -1596,7 +1654,7 @@ void FD3D11DynamicRHI::RHIClearMRT(bool bClearColor,int32 NumClearColors,const F
 	ID3D11DepthStencilView* DepthStencilView = BoundRenderTargets.GetDepthStencilView();
 
 	// Determine if we're trying to clear a subrect of the screen
-	bool UseDrawClear = false;
+	bool UseDrawClear = bForceShaderClear;
 	uint32 NumViews = 1;
 	D3D11_VIEWPORT Viewport;
 	StateCache.GetViewports(&NumViews,&Viewport);
@@ -1734,7 +1792,7 @@ void FD3D11DynamicRHI::RHIClearMRT(bool bClearColor,int32 NumClearColors,const F
 
 		// Store the current device state
 		FDeviceStateHelper OriginalResourceState(Direct3DDeviceIMContext);
-		OriginalResourceState.CaptureDeviceState(StateCache);
+		OriginalResourceState.CaptureDeviceState(StateCache, BoundShaderStateHistory);
 
 		// Set the cached state objects
 		StateCache.SetBlendState(BlendState->Resource, BF, 0xffffffff);
@@ -1794,15 +1852,7 @@ void FD3D11DynamicRHI::RHIClearMRT(bool bClearColor,int32 NumClearColors,const F
 		{
 			FRHICommandList_RecursiveHazardous RHICmdList(this);
 			SetGlobalBoundShaderState(RHICmdList, GMaxRHIFeatureLevel, GD3D11ClearMRTBoundShaderState[FMath::Max(BoundRenderTargets.GetNumActiveTargets() - 1, 0)], GD3D11Vector4VertexDeclaration.VertexDeclarationRHI, *VertexShader, PixelShader);
-			FLinearColor ShaderClearColors[MaxSimultaneousRenderTargets];
-			FMemory::Memzero(ShaderClearColors);
-
-			for (int32 i = 0; i < NumClearColors; i++)
-			{
-				ShaderClearColors[i] = ClearColorArray[i];
-			}
-
-			SetShaderValueArray(RHICmdList, PixelShader->GetPixelShader(), PixelShader->ColorParameter, ShaderClearColors, NumClearColors);
+			PixelShader->SetColors(RHICmdList, ClearColorArray, NumClearColors);
 
 			{
 				// Draw a fullscreen quad
@@ -1854,14 +1904,14 @@ void FD3D11DynamicRHI::RHIClearMRT(bool bClearColor,int32 NumClearColors,const F
 		}
 
 		// Restore the original device state
-		OriginalResourceState.RestoreDeviceState(StateCache); 
+		OriginalResourceState.RestoreDeviceState(StateCache, BoundShaderStateHistory);
 	}
 	else
 	{
 		if (bClearColor && BoundRenderTargets.GetNumActiveTargets() > 0)
 		{
 			for (int32 TargetIndex = 0; TargetIndex < BoundRenderTargets.GetNumActiveTargets(); TargetIndex++)
-			{
+			{				
 				Direct3DDeviceIMContext->ClearRenderTargetView(BoundRenderTargets.GetRenderTargetView(TargetIndex),(float*)&ClearColorArray[TargetIndex]);
 			}
 		}
@@ -1884,7 +1934,7 @@ void FD3D11DynamicRHI::RHIClearMRT(bool bClearColor,int32 NumClearColors,const F
 	GPUProfilingData.RegisterGPUWork(0);
 }
 
-void FD3D11DynamicRHI::RHIBindClearMRTValues(bool bClearColor, int32 NumClearColors, const FLinearColor* ClearColorArray, bool bClearDepth, float Depth, bool bClearStencil, uint32 Stencil)
+void FD3D11DynamicRHI::RHIBindClearMRTValues(bool bClearColor, bool bClearDepth, bool bClearStencil)
 {
 	// Not necessary for d3d.
 }
diff --git a/Engine/Source/Runtime/Windows/D3D11RHI/Private/D3D11RHIPrivate.h b/Engine/Source/Runtime/Windows/D3D11RHI/Private/D3D11RHIPrivate.h
--- a/Engine/Source/Runtime/Windows/D3D11RHI/Private/D3D11RHIPrivate.h
+++ b/Engine/Source/Runtime/Windows/D3D11RHI/Private/D3D11RHIPrivate.h
@@ -461,7 +461,7 @@ public:
 	virtual void RHISetBlendState(FBlendStateRHIParamRef NewState, const FLinearColor& BlendFactor) final override;
 	virtual void RHISetRenderTargets(uint32 NumSimultaneousRenderTargets, const FRHIRenderTargetView* NewRenderTargets, const FRHIDepthRenderTargetView* NewDepthStencilTarget, uint32 NumUAVs, const FUnorderedAccessViewRHIParamRef* UAVs) final override;
 	virtual void RHISetRenderTargetsAndClear(const FRHISetRenderTargetsInfo& RenderTargetsInfo) final override;
-	virtual void RHIBindClearMRTValues(bool bClearColor, int32 NumClearColors, const FLinearColor* ColorArray, bool bClearDepth, float Depth, bool bClearStencil, uint32 Stencil) final override;
+	virtual void RHIBindClearMRTValues(bool bClearColor, bool bClearDepth, bool bClearStencil) final override;
 	virtual void RHIDrawPrimitive(uint32 PrimitiveType, uint32 BaseVertexIndex, uint32 NumPrimitives, uint32 NumInstances) final override;
 	virtual void RHIDrawPrimitiveIndirect(uint32 PrimitiveType, FVertexBufferRHIParamRef ArgumentBuffer, uint32 ArgumentOffset) final override;
 	virtual void RHIDrawIndexedIndirect(FIndexBufferRHIParamRef IndexBufferRHI, uint32 PrimitiveType, FStructuredBufferRHIParamRef ArgumentsBufferRHI, int32 DrawArgumentsIndex, uint32 NumInstances) final override;
@@ -495,6 +495,8 @@ public:
 		return DXGIFactory1;
 	}
 private:
+	virtual void RHIClearMRTImpl(bool bClearColor, int32 NumClearColors, const FLinearColor* ColorArray, bool bClearDepth, float Depth, bool bClearStencil, uint32 Stencil, FIntRect ExcludeRect, bool bForceShaderClear);
+
 	template <EShaderFrequency ShaderFrequency>
 	void ClearShaderResourceViews(FD3D11BaseShaderResource* Resource);
 
diff --git a/Engine/Source/Runtime/Windows/D3D11RHI/Private/D3D11StateCachePrivate.h b/Engine/Source/Runtime/Windows/D3D11RHI/Private/D3D11StateCachePrivate.h
--- a/Engine/Source/Runtime/Windows/D3D11RHI/Private/D3D11StateCachePrivate.h
+++ b/Engine/Source/Runtime/Windows/D3D11RHI/Private/D3D11StateCachePrivate.h
@@ -468,7 +468,7 @@ public:
 	}
 
 template <EShaderFrequency ShaderFrequency>
-	D3D11_STATE_CACHE_INLINE void GetConstantBuffers(uint32 StartSlotIndex, uint32 NumBuffers, ID3D11Buffer* ConstantBuffers)
+	D3D11_STATE_CACHE_INLINE void GetConstantBuffers(uint32 StartSlotIndex, uint32 NumBuffers, ID3D11Buffer** ConstantBuffers)
 	{
 #if D3D11_ALLOW_STATE_CACHE
 		check(StartSlotIndex + NumBuffers <= D3D11_COMMONSHADER_CONSTANT_BUFFER_API_SLOT_COUNT);
diff --git a/Engine/Source/Runtime/Windows/D3D11RHI/Private/D3D11Texture.cpp b/Engine/Source/Runtime/Windows/D3D11RHI/Private/D3D11Texture.cpp
--- a/Engine/Source/Runtime/Windows/D3D11RHI/Private/D3D11Texture.cpp
+++ b/Engine/Source/Runtime/Windows/D3D11RHI/Private/D3D11Texture.cpp
@@ -836,7 +836,8 @@ TD3D11Texture2D<BaseResourceType>* FD3D11DynamicRHI::CreateD3D11Texture2D(uint32
 		(EPixelFormat)Format,
 		bCubeTexture,
 		Flags,
-		bPooledTexture
+		bPooledTexture,
+		CreateInfo.ClearValueBinding
 #if PLATFORM_SUPPORTS_VIRTUAL_TEXTURES
 		, RawTextureMemory
 #endif
@@ -962,7 +963,7 @@ FD3D11Texture3D* FD3D11DynamicRHI::CreateD3D11Texture3D(uint32 SizeX,uint32 Size
 
 	TArray<TRefCountPtr<ID3D11RenderTargetView> > RenderTargetViews;
 	RenderTargetViews.Add(RenderTargetView);
-	FD3D11Texture3D* Texture3D = new FD3D11Texture3D(this,TextureResource,ShaderResourceView,RenderTargetViews,SizeX,SizeY,SizeZ,NumMips,(EPixelFormat)Format,Flags);
+	FD3D11Texture3D* Texture3D = new FD3D11Texture3D(this,TextureResource,ShaderResourceView,RenderTargetViews,SizeX,SizeY,SizeZ,NumMips,(EPixelFormat)Format,Flags, CreateInfo.ClearValueBinding);
 
 	Texture3D->ResourceInfo.VRamAllocation = VRamAllocation;
 
@@ -1079,7 +1080,8 @@ FTexture2DRHIRef FD3D11DynamicRHI::RHIAsyncCreateTexture2D(uint32 SizeX,uint32 S
 		(EPixelFormat)Format,
 		/*bInCubemap=*/ false,
 		Flags,
-		/*bPooledTexture=*/ false
+		/*bPooledTexture=*/ false,
+		FClearValueBinding()
 		);
 
 	D3D11TextureAllocated(*NewTexture);
diff --git a/Engine/Source/Runtime/Windows/D3D11RHI/Private/D3D11Viewport.cpp b/Engine/Source/Runtime/Windows/D3D11RHI/Private/D3D11Viewport.cpp
--- a/Engine/Source/Runtime/Windows/D3D11RHI/Private/D3D11Viewport.cpp
+++ b/Engine/Source/Runtime/Windows/D3D11RHI/Private/D3D11Viewport.cpp
@@ -141,7 +141,8 @@ FD3D11Texture2D* GetSwapChainSurface(FD3D11DynamicRHI* D3DRHI, EPixelFormat Pixe
 		PixelFormat,
 		false,
 		false,
-		false
+		false,
+		FClearValueBinding()
 		);
 
 	D3D11TextureAllocated2D(*NewTexture);
diff --git a/Engine/Source/Runtime/Windows/D3D11RHI/Public/D3D11Resources.h b/Engine/Source/Runtime/Windows/D3D11RHI/Public/D3D11Resources.h
--- a/Engine/Source/Runtime/Windows/D3D11RHI/Public/D3D11Resources.h
+++ b/Engine/Source/Runtime/Windows/D3D11RHI/Public/D3D11Resources.h
@@ -315,7 +315,8 @@ public:
 		EPixelFormat InFormat,
 		bool bInCubemap,
 		uint32 InFlags,
-		bool bInPooled
+		bool bInPooled,
+		const FClearValueBinding& InClearValue
 #if PLATFORM_SUPPORTS_VIRTUAL_TEXTURES
 		, void* InRawTextureMemory = nullptr
 #endif
@@ -327,7 +328,8 @@ public:
 		InNumMips,
 		InNumSamples,
 		InFormat,
-		InFlags
+		InFlags,
+		InClearValue
 		)
 	, FD3D11TextureBase(
 		InD3DRHI,
@@ -419,9 +421,10 @@ public:
 		uint32 InSizeZ,
 		uint32 InNumMips,
 		EPixelFormat InFormat,
-		uint32 InFlags
+		uint32 InFlags,
+		const FClearValueBinding& InClearValue
 		)
-	: FRHITexture3D(InSizeX,InSizeY,InSizeZ,InNumMips,InFormat,InFlags)
+	: FRHITexture3D(InSizeX,InSizeY,InSizeZ,InNumMips,InFormat,InFlags,InClearValue)
 	, FD3D11TextureBase(
 		InD3DRHI,
 		InResource,
@@ -457,8 +460,8 @@ public:
 class FD3D11BaseTexture2D : public FRHITexture2D
 {
 public:
-	FD3D11BaseTexture2D(uint32 InSizeX,uint32 InSizeY,uint32 InSizeZ,uint32 InNumMips,uint32 InNumSamples,EPixelFormat InFormat,uint32 InFlags)
-	: FRHITexture2D(InSizeX,InSizeY,InNumMips,InNumSamples,InFormat,InFlags)
+	FD3D11BaseTexture2D(uint32 InSizeX, uint32 InSizeY, uint32 InSizeZ, uint32 InNumMips, uint32 InNumSamples, EPixelFormat InFormat, uint32 InFlags, const FClearValueBinding& InClearValue)
+	: FRHITexture2D(InSizeX,InSizeY,InNumMips,InNumSamples,InFormat,InFlags, InClearValue)
 	{}
 	uint32 GetSizeZ() const { return 0; }
 };
@@ -466,16 +469,16 @@ public:
 class FD3D11BaseTexture2DArray : public FRHITexture2DArray
 {
 public:
-	FD3D11BaseTexture2DArray(uint32 InSizeX,uint32 InSizeY,uint32 InSizeZ,uint32 InNumMips,uint32 InNumSamples,EPixelFormat InFormat,uint32 InFlags)
-	: FRHITexture2DArray(InSizeX,InSizeY,InSizeZ,InNumMips,InFormat,InFlags)
+	FD3D11BaseTexture2DArray(uint32 InSizeX, uint32 InSizeY, uint32 InSizeZ, uint32 InNumMips, uint32 InNumSamples, EPixelFormat InFormat, uint32 InFlags, const FClearValueBinding& InClearValue)
+	: FRHITexture2DArray(InSizeX,InSizeY,InSizeZ,InNumMips,InFormat,InFlags,InClearValue)
 	{ check(InNumSamples == 1); }
 };
 
 class FD3D11BaseTextureCube : public FRHITextureCube
 {
 public:
-	FD3D11BaseTextureCube(uint32 InSizeX,uint32 InSizeY,uint32 InSizeZ,uint32 InNumMips,uint32 InNumSamples,EPixelFormat InFormat,uint32 InFlags)
-	: FRHITextureCube(InSizeX,InNumMips,InFormat,InFlags)
+	FD3D11BaseTextureCube(uint32 InSizeX, uint32 InSizeY, uint32 InSizeZ, uint32 InNumMips, uint32 InNumSamples, EPixelFormat InFormat, uint32 InFlags, const FClearValueBinding& InClearValue)
+	: FRHITextureCube(InSizeX,InNumMips,InFormat,InFlags,InClearValue)
 	{ check(InNumSamples == 1); }
 	uint32 GetSizeX() const { return GetSize(); }
 	uint32 GetSizeY() const { return GetSize(); } //-V524
