Improved VisualizeDOF for CircleDOF, polish

diff --git a/Engine/Shaders/CircleDOFCommon.usf b/Engine/Shaders/CircleDOFCommon.usf
new file mode 100644
--- /dev/null
+++ b/Engine/Shaders/CircleDOFCommon.usf
@@ -0,0 +1,27 @@
+// Copyright 1998-2015 Epic Games, Inc. All Rights Reserved.
+
+/*=============================================================================
+	CircleDOFCommon.usf: Circle depth of field common
+=============================================================================*/
+
+
+// {radius, depth blur amp, depth blur radius, ViewPortScale}
+float4 CircleDofParams;
+
+// CircleDOF: Compute circle of confusion size in pixels.
+float DepthToCoc(float SceneDepth) 
+{
+	// We intentionally don't do the DepthOfFieldFocalRegion as it breaks realism.
+
+	float Focus = View.DepthOfFieldFocalDistance;
+	float Radius = CircleDofParams.x;
+	float CocRadius = ((SceneDepth - Focus) / SceneDepth) * Radius;
+	float DepthBlurRadius = (1.0 - exp2(-SceneDepth * CircleDofParams.y)) * CircleDofParams.z;
+	float ReturnCoc = max(abs(CocRadius), DepthBlurRadius);
+	if(CocRadius < 0.0) 
+	{
+		// near CoC is using negative values
+		ReturnCoc = -ReturnCoc;
+	}
+	return CocRadius;
+}
diff --git a/Engine/Shaders/PostProcessBokehDOF.usf b/Engine/Shaders/PostProcessBokehDOF.usf
--- a/Engine/Shaders/PostProcessBokehDOF.usf
+++ b/Engine/Shaders/PostProcessBokehDOF.usf
@@ -16,13 +16,11 @@ float4 DepthOfFieldThresholds;
 uint4 TileCountAndSize;
 // .xy:size in pixels, .zw:LeftTop of the original viewport in rendertaregt scaled coordiantes
 float4 KernelSize;
+
 // small Bokeh like texture
 Texture2D LensTexture;
 SamplerState LensTextureSampler;
 
-// for VisualizeDOF [0]:near, [1]:far
-float4 VisualizeColors[2];
-
 
 // @return x:layer in front of focal plane, y: layer behind focal plane  1-x-y:layer in focus
 float2 ComputeLayerContributions(float Depth)
@@ -383,24 +381,6 @@ void MainSetupPS(in float4 UVAndScreenPos : TEXCOORD0, out float4 OutColor : SV_
 	OutColor.rgb = min(float3(256 * 256, 256 * 256, 256 * 256), OutColor.rgb);
 }
 
-
-// render in fullres to visualize the half res DOF input from the setup pass
-void VisualizeDOFPS(in float4 UVAndScreenPos : TEXCOORD0, out float4 OutColor : SV_Target0)
-{
-	bool bFrontLayer;
-	float4 Mask;
-	
-	CommonDOFSetup(UVAndScreenPos.xy, bFrontLayer, Mask);
-
-	OutColor = (bFrontLayer ? VisualizeColors[0] : VisualizeColors[1]) * dot(Mask, 1);
-	
-	half3 SceneColor = Texture2DSample(PostprocessInput0, PostprocessInput0Sampler, UVAndScreenPos.xy).rgb;
-
-	// blend in a bit of the scene color to make navigation easier
-	OutColor.rgb = lerp(saturate(OutColor.rgb), Luminance(SceneColor), 0.1f);
-}
-
-
 // pixel shader to accumulate the Bokeh shaped elements
 void MainPS(
 	float2 TexCoord : TEXCOORD0,
diff --git a/Engine/Shaders/PostProcessCircleDOF.usf b/Engine/Shaders/PostProcessCircleDOF.usf
--- a/Engine/Shaders/PostProcessCircleDOF.usf
+++ b/Engine/Shaders/PostProcessCircleDOF.usf
@@ -8,28 +8,7 @@
 #include "PostProcessCommon.usf"
 #include "DeferredShadingCommon.usf"		// FGBufferData
 #include "DepthOfFieldCommon.usf"
-
-
-// {radius, depth blur amp, depth blur radius, ViewPortScale}
-float4 CircleDofParams;
-
-// CIRCLE DOF: Compute circle of confusion size in pixels.
-float DepthToCoc(float SceneDepth) 
-{
-	// We intentionally don't do the DepthOfFieldFocalRegion as it breaks realism.
-
-	float Focus = View.DepthOfFieldFocalDistance;
-	float Radius = CircleDofParams.x;
-	float CocRadius = ((SceneDepth - Focus) / SceneDepth) * Radius;
-	float DepthBlurRadius = (1.0 - exp2(-SceneDepth * CircleDofParams.y)) * CircleDofParams.z;
-	float ReturnCoc = max(abs(CocRadius), DepthBlurRadius);
-	if(CocRadius < 0.0) 
-	{
-		// near CoC is using negative values
-		ReturnCoc = -ReturnCoc;
-	}
-	return ReturnCoc;
-}
+#include "CircleDOFCommon.usf"
 
 // pixel shader entry point
 void CircleSetupPS(float4 UVAndScreenPos : TEXCOORD0, out float4 OutColor0 : SV_Target0)
diff --git a/Engine/Shaders/PostProcessVisualizeDOF.usf b/Engine/Shaders/PostProcessVisualizeDOF.usf
new file mode 100644
--- /dev/null
+++ b/Engine/Shaders/PostProcessVisualizeDOF.usf
@@ -0,0 +1,93 @@
+// Copyright 1998-2015 Epic Games, Inc. All Rights Reserved.
+
+/*=============================================================================
+	PostProcessBokehDOF.usf: PostProcessing Lens Flares.
+=============================================================================*/
+
+#include "Common.usf"
+#include "PostProcessCommon.usf"
+#include "DepthOfFieldCommon.usf"
+#include "MiniFontCommon.usf"		// for PrintFloat()
+#include "CircleDOFCommon.usf"
+
+// for the BokehDOF vertex shader, from postprocessing settings
+// .x: color threshold, .y:size threshold, zw: unused
+float4 DepthOfFieldThresholds;
+
+// .xy:tilecount, .zw:tilesize
+uint4 TileCountAndSize;
+// .xy:size in pixels, .zw:LeftTop of the original viewport in rendertaregt scaled coordiantes
+float4 KernelSize;
+
+// for VisualizeDOF only
+int2 CursorPos;
+
+// for VisualizeDOF [0]:near, [1]:far
+float4 VisualizeColors[2];
+
+// render in fullres to visualize the half res DOF input from the setup pass
+void VisualizeDOFPS(in float4 UVAndScreenPos : TEXCOORD0, out float4 OutColor : SV_Target0)
+{
+	OutColor = 0;
+
+	float2 UV = UVAndScreenPos.xy; 
+	int2 PixelPos = (int2)(UVAndScreenPos.zw * ScreenPosToPixel.xy + ScreenPosToPixel.zw + 0.5f);
+	float2 ViewLocalUV = float2(UVAndScreenPos.z * 0.5f + 0.5f, 0.5f - 0.5f * UVAndScreenPos.w);
+	
+	float SceneDepth = CalcSceneDepth(UV);
+
+	// 0..1
+	half Near = CalcUnfocusedPercentCustomBound(SceneDepth, 1, 0);
+	// 0..1
+	half Far = CalcUnfocusedPercentCustomBound(SceneDepth, 0, 1);
+
+	OutColor.rgb = lerp(OutColor.rgb, VisualizeColors[0].rgb, Near);
+	OutColor.rgb = lerp(OutColor.rgb, VisualizeColors[1].rgb, Far);
+	
+	half3 SceneColor = Texture2DSample(PostprocessInput0, PostprocessInput0Sampler, UV).rgb;
+
+	// blend in a bit of the scene color to make navigation easier
+	OutColor.rgb = lerp(saturate(OutColor.rgb), Luminance(SceneColor), 0.1f);
+
+	// draw a crosshair
+	{
+		float CrossHairMask = PixelPos.x == CursorPos.x || PixelPos.y == CursorPos.y;
+		float2 DistAbs = abs(PixelPos - CursorPos);
+		float Dist = max(DistAbs.x, DistAbs.y);
+		float DistMask = Dist >= 12;
+
+		OutColor.xyz = lerp(OutColor.xyz, float3(1, 1, 1), CrossHairMask * DistMask * 0.1f);
+	}
+
+	// value under crosshair
+	{
+		float SceneDepthAtCursor = CalcSceneDepth(CursorPos * PostprocessInput0Size.zw);
+
+		// DepthToCoc is half res do we multiply by 2, sign is used to identigy near/far wehich we don't need here
+		float CoCInPixel = abs(DepthToCoc(SceneDepthAtCursor)) * 2.0f;
+
+		// draw a circle
+		{
+			float Dist = length(PixelPos - CursorPos);
+			
+			float Mask = saturate(1 - abs(CoCInPixel - Dist));
+
+			OutColor.xyz = lerp(OutColor.xyz, float3(1, 1, 1), Mask);
+		}
+		
+		int2 PrintCursor = CursorPos + int2(28, 7 + 0 * 8);
+
+		{
+			int2 LeftTop = PrintCursor - int2(11, 0);
+			PrintCharacter(PixelPos, OutColor.xyz, float3(1, 1, 1), LeftTop, 'D' - 'A' + 10);
+		}
+		PrintFloatNoFraction(PixelPos, OutColor.xyz, float3(1, 1, 1)*0.6, PrintCursor, SceneDepthAtCursor);
+		PrintCursor.y += 16;
+		{
+			int2 LeftTop = PrintCursor - int2(11, 0);
+			PrintCharacter(PixelPos, OutColor.xyz, float3(1, 1, 1), LeftTop, 'R' - 'A' + 10);
+		}
+		PrintFloat(PixelPos, OutColor.xyz, float3(1, 1, 1)*0.6, PrintCursor, CoCInPixel);
+	}
+}
+
diff --git a/Engine/Source/Runtime/Engine/Classes/Engine/Scene.h b/Engine/Source/Runtime/Engine/Classes/Engine/Scene.h
--- a/Engine/Source/Runtime/Engine/Classes/Engine/Scene.h
+++ b/Engine/Source/Runtime/Engine/Classes/Engine/Scene.h
@@ -837,12 +837,12 @@ struct FPostProcessSettings
 	UPROPERTY(interp, BlueprintReadWrite, Category=DepthOfField, meta=(ClampMin = "0.0", ClampMax = "4.0", editcondition = "bOverride_DepthOfFieldDepthBlurRadius", DisplayName = "Depth Blur Radius"))
 	float DepthOfFieldDepthBlurRadius;
 	
-	/** CircleDOF only: Defines the opening of the camera lens, aperture is 1/fstop, typical lens go down to 1.2 (large opening), larger numbers reduce the DOF effect */
+	/** CircleDOF only: Defines the opening of the camera lens, Aperture is 1/fstop, typical lens go down to f/1.2 (large opening), larger numbers reduce the DOF effect */
 	UPROPERTY(interp, BlueprintReadWrite, Category=DepthOfField, meta=(ClampMin = "1.0", ClampMax = "32.0", editcondition = "bOverride_DepthOfFieldFstop", DisplayName = "Aperture F-stop"))
 	float DepthOfFieldFstop;
 
 	/** Distance in which the Depth of Field effect should be sharp, in unreal units (cm) */
-	UPROPERTY(interp, BlueprintReadWrite, Category=DepthOfField, meta=(UIMin = "0.0", UIMax = "10000.0", editcondition = "bOverride_DepthOfFieldFocalDistance", DisplayName = "Focal Distance"))
+	UPROPERTY(interp, BlueprintReadWrite, Category=DepthOfField, meta=(ClampMin = "0.0", UIMin = "1.0", UIMax = "10000.0", editcondition = "bOverride_DepthOfFieldFocalDistance", DisplayName = "Focal Distance"))
 	float DepthOfFieldFocalDistance;
 
 	/** Artificial region where all content is in focus, starting after DepthOfFieldFocalDistance, in unreal units  (cm) */
diff --git a/Engine/Source/Runtime/Renderer/Private/PostProcess/PostProcessBokehDOF.cpp b/Engine/Source/Runtime/Renderer/Private/PostProcess/PostProcessBokehDOF.cpp
--- a/Engine/Source/Runtime/Renderer/Private/PostProcess/PostProcessBokehDOF.cpp
+++ b/Engine/Source/Runtime/Renderer/Private/PostProcess/PostProcessBokehDOF.cpp
@@ -94,7 +94,10 @@ public:
 	FDeferredPixelShaderParameters DeferredParameters;
 	FShaderParameter DepthOfFieldParams;
 	FShaderParameter VisualizeColors;
-
+	FShaderParameter CursorPos;
+	FShaderResourceParameter MiniFontTexture;
+	FShaderParameter CircleDofParams;
+	
 	/** Initialization constructor. */
 	PostProcessVisualizeDOFPS(const ShaderMetaType::CompiledShaderInitializerType& Initializer)
 		: FGlobalShader(Initializer)
@@ -102,14 +105,17 @@ public:
 		PostprocessParameter.Bind(Initializer.ParameterMap);
 		DeferredParameters.Bind(Initializer.ParameterMap);
 		DepthOfFieldParams.Bind(Initializer.ParameterMap, TEXT("DepthOfFieldParams"));
+		MiniFontTexture.Bind(Initializer.ParameterMap, TEXT("MiniFontTexture"));
 		VisualizeColors.Bind(Initializer.ParameterMap, TEXT("VisualizeColors"));
+		CursorPos.Bind(Initializer.ParameterMap, TEXT("CursorPos"));
+		CircleDofParams.Bind(Initializer.ParameterMap,TEXT("CircleDofParams"));
 	}
 
 	// FShader interface.
 	virtual bool Serialize(FArchive& Ar) override
 	{
 		bool bShaderHasOutdatedParameters = FGlobalShader::Serialize(Ar);
-		Ar << PostprocessParameter << DeferredParameters << DepthOfFieldParams << VisualizeColors;
+		Ar << PostprocessParameter << DeferredParameters << MiniFontTexture << DepthOfFieldParams << VisualizeColors << CursorPos << CircleDofParams;
 		return bShaderHasOutdatedParameters;
 	}
 
@@ -123,6 +129,8 @@ public:
 
 		PostprocessParameter.SetPS(ShaderRHI, Context, TStaticSamplerState<SF_Point, AM_Clamp, AM_Clamp, AM_Clamp>::GetRHI());
 
+		SetTextureParameter(Context.RHICmdList, ShaderRHI, MiniFontTexture, GEngine->MiniFontTexture ? GEngine->MiniFontTexture->Resource->TextureRHI : GSystemTextures.WhiteDummy->GetRenderTargetItem().TargetableTexture);
+
 		{
 			FVector4 DepthOfFieldParamValues[2];
 
@@ -134,6 +142,18 @@ public:
 			SetShaderValueArray(Context.RHICmdList, ShaderRHI, DepthOfFieldParams, DepthOfFieldParamValues, 2);
 		}
 
+		{
+			// a negative values disables the cross hair feature
+			FIntPoint CursorPosValue(-100,-100);
+			
+			if(Context.View.FinalPostProcessSettings.DepthOfFieldMethod == DOFM_CircleDOF)
+			{
+				CursorPosValue = Context.View.CursorPos;
+			}
+
+			SetShaderValue(Context.RHICmdList, ShaderRHI, CursorPos, CursorPosValue);
+		}
+
 		{
 			FLinearColor Colors[2] = { FLinearColor(0.1f, 0.1f, 0.1f, 0), FLinearColor(0.1f, 0.1f, 0.1f, 0) };
 
@@ -148,11 +168,13 @@ public:
 
 			SetShaderValueArray(Context.RHICmdList, ShaderRHI, VisualizeColors, Colors, 2);
 		}
+
+		SetShaderValue(Context.RHICmdList, ShaderRHI, CircleDofParams, CircleDofCoc(Context.View));
 	}
 
 	static const TCHAR* GetSourceFilename()
 	{
-		return TEXT("PostProcessBokehDOF");
+		return TEXT("PostProcessVisualizeDOF");
 	}
 
 	static const TCHAR* GetFunctionName()
@@ -321,14 +343,14 @@ void FRCPassPostProcessVisualizeDOF::Process(FRenderingCompositePassContext& Con
 		}
 		else if(MethodId == DOFM_CircleDOF)
 		{
-			Line = FString::Printf(TEXT("Method: CircleDOF"));
+			Line = FString::Printf(TEXT("Method: CircleDOF (blue is far, green is near, black is in focus, cross hair shows Depth and CoC radius in pixel)"));
 			Canvas.DrawShadowedString(X, Y += YStep, *Line, GetStatsFont(), FLinearColor(1, 1, 1));
 			Y += YStep;
 			Line = FString::Printf(TEXT("FocalDistance: %.2f"), View.FinalPostProcessSettings.DepthOfFieldFocalDistance);
 			Canvas.DrawShadowedString(X, Y += YStep, *Line, GetStatsFont(), FLinearColor(1, 1, 1));
 			Line = FString::Printf(TEXT("Aperture F-stop: %.2f"), View.FinalPostProcessSettings.DepthOfFieldFstop);
 			Canvas.DrawShadowedString(X, Y += YStep, *Line, GetStatsFont(), FLinearColor(1, 1, 1));
-			Line = FString::Printf(TEXT("Aperture: 1 / %.2f"), View.FinalPostProcessSettings.DepthOfFieldFstop);
+			Line = FString::Printf(TEXT("Aperture: f/%.2f"), View.FinalPostProcessSettings.DepthOfFieldFstop);
 			Canvas.DrawShadowedString(X, Y += YStep, *Line, GetStatsFont(), FLinearColor(1, 1, 1));
 			Y += YStep;
 			Line = FString::Printf(TEXT("DepthBlur (not related to Depth of Field, due to light traveling long distances in atmosphere)"));
diff --git a/Engine/Source/Runtime/Renderer/Private/PostProcess/PostProcessCircleDOF.cpp b/Engine/Source/Runtime/Renderer/Private/PostProcess/PostProcessCircleDOF.cpp
--- a/Engine/Source/Runtime/Renderer/Private/PostProcess/PostProcessCircleDOF.cpp
+++ b/Engine/Source/Runtime/Renderer/Private/PostProcess/PostProcessCircleDOF.cpp
@@ -31,12 +31,12 @@ float ComputeFocalLengthFromFov(const FSceneView& View)
 
 // Convert f-stop and focal distance into projected size in half resolution pixels.
 // Setup depth based blur.
-static FVector4 CircleDofCoc(const FRenderingCompositePassContext& Context)
+FVector4 CircleDofCoc(const FSceneView& View)
 {
-	float FocalLengthInMM = ComputeFocalLengthFromFov(Context.View);
+	float FocalLengthInMM = ComputeFocalLengthFromFov(View);
 	 
 	// Convert focal distance in world position to mm (from cm to mm)
-	float FocalDistanceInMM = Context.View.FinalPostProcessSettings.DepthOfFieldFocalDistance * 10.0f;
+	float FocalDistanceInMM = View.FinalPostProcessSettings.DepthOfFieldFocalDistance * 10.0f;
 
 	// Convert f-stop, focal length, and focal distance to
 	// projected circle of confusion size at infinity in mm.
@@ -46,11 +46,11 @@ static FVector4 CircleDofCoc(const FRenderingCompositePassContext& Context)
 	//   f = focal length
 	//   d = focal distance
 	//   n = fstop (where n is the "n" in "f/n")
-	float Radius = FMath::Square(FocalLengthInMM) / (Context.View.FinalPostProcessSettings.DepthOfFieldFstop * (FocalDistanceInMM - FocalLengthInMM));
+	float Radius = FMath::Square(FocalLengthInMM) / (View.FinalPostProcessSettings.DepthOfFieldFstop * (FocalDistanceInMM - FocalLengthInMM));
 
 	// Scale so that APS-C 24.576 mm = full frame.
 	// Convert mm to pixels.
-	float Width = (float)Context.View.ViewRect.Width();
+	float Width = (float)View.ViewRect.Width();
 	Radius = Radius * Width * (1.0f/24.576f);
 
 	// Convert diameter to radius at half resolution (algorithm radius is at half resolution).
@@ -69,8 +69,8 @@ static FVector4 CircleDofCoc(const FRenderingCompositePassContext& Context)
 	// Need to convert to cm here.
 	return FVector4(
 		Radius, 
-		1.0f/(Context.View.FinalPostProcessSettings.DepthOfFieldDepthBlurAmount * 100000.0),
-		Context.View.FinalPostProcessSettings.DepthOfFieldDepthBlurRadius * Width / 1920.0f,
+		1.0f/(View.FinalPostProcessSettings.DepthOfFieldDepthBlurAmount * 100000.0),
+		View.FinalPostProcessSettings.DepthOfFieldDepthBlurRadius * Width / 1920.0f,
 		Width / 1920.0f);
 }
 
@@ -136,7 +136,7 @@ public:
 			SetShaderValueArray(Context.RHICmdList, ShaderRHI, DepthOfFieldParams, DepthOfFieldParamValues, 2);
 		}
 
-		SetShaderValue(Context.RHICmdList, ShaderRHI, CircleDofParams, CircleDofCoc(Context));
+		SetShaderValue(Context.RHICmdList, ShaderRHI, CircleDofParams, CircleDofCoc(Context.View));
 	}
 };
 
@@ -324,7 +324,7 @@ public:
 			SetShaderValueArray(Context.RHICmdList, ShaderRHI, DepthOfFieldParams, DepthOfFieldParamValues, 2);
 		}
 
-		SetShaderValue(Context.RHICmdList, ShaderRHI, CircleDofParams, CircleDofCoc(Context));
+		SetShaderValue(Context.RHICmdList, ShaderRHI, CircleDofParams, CircleDofCoc(Context.View));
 	}
 };
 
@@ -544,7 +544,7 @@ public:
 		TemporalRandom2(&RandomOffsetValue, Context.View.Family->FrameNumber);
 		SetShaderValue(Context.RHICmdList, ShaderRHI, RandomOffset, RandomOffsetValue);
 
-		SetShaderValue(Context.RHICmdList, ShaderRHI, CircleDofParams, CircleDofCoc(Context));
+		SetShaderValue(Context.RHICmdList, ShaderRHI, CircleDofParams, CircleDofCoc(Context.View));
 	}
 };
 
@@ -735,7 +735,7 @@ public:
 		TemporalRandom2(&RandomOffsetValue, Context.View.Family->FrameNumber);
 		SetShaderValue(Context.RHICmdList, ShaderRHI, RandomOffset, RandomOffsetValue);
 
-		SetShaderValue(Context.RHICmdList, ShaderRHI, CircleDofParams, CircleDofCoc(Context));
+		SetShaderValue(Context.RHICmdList, ShaderRHI, CircleDofParams, CircleDofCoc(Context. View));
 	}
 };
 
diff --git a/Engine/Source/Runtime/Renderer/Private/PostProcess/PostProcessCircleDOF.h b/Engine/Source/Runtime/Renderer/Private/PostProcess/PostProcessCircleDOF.h
--- a/Engine/Source/Runtime/Renderer/Private/PostProcess/PostProcessCircleDOF.h
+++ b/Engine/Source/Runtime/Renderer/Private/PostProcess/PostProcessCircleDOF.h
@@ -92,3 +92,4 @@ private:
 // @return in mm, assuming this sensor: APS-C 24.576 mm
 float ComputeFocalLengthFromFov(const FSceneView& View);
 
+FVector4 CircleDofCoc(const FSceneView& View);
