Remove slow clear path as PS4 issue has been fixed.

diff --git a/Engine/Source/Runtime/Renderer/Private/DistanceFieldSurfaceCacheLighting.cpp b/Engine/Source/Runtime/Renderer/Private/DistanceFieldSurfaceCacheLighting.cpp
--- a/Engine/Source/Runtime/Renderer/Private/DistanceFieldSurfaceCacheLighting.cpp
+++ b/Engine/Source/Runtime/Renderer/Private/DistanceFieldSurfaceCacheLighting.cpp
@@ -2357,10 +2357,6 @@ void RenderIrradianceCacheInterpolation(
 
 		FTextureRHIParamRef DepthBuffer = SplatDepthStencilBuffer ? SplatDepthStencilBuffer->GetRenderTargetItem().TargetableTexture : NULL;
 
-		// PS4 fast clear which is only triggered through SetRenderTargetsAndClear adds 3ms of GPU idle time currently
-		static bool bUseSetAndClear = true;
-
-		if (bUseSetAndClear)
 		{
 			FRHIRenderTargetView RenderTargets[MaxSimultaneousRenderTargets];
 			RenderTargets[0] = FRHIRenderTargetView(BentNormalInterpolationTarget->GetRenderTargetItem().TargetableTexture);
@@ -2379,22 +2375,6 @@ void RenderIrradianceCacheInterpolation(
 			// set the render target
 			RHICmdList.SetRenderTargetsAndClear(Info);
 		}
-		//we need to prefer fast clears.  If PS4 is seeing 3ms of latency we should investigate soon.
-#if 0
-		else
-		{
-			FTextureRHIParamRef RenderTargets[2] =
-			{
-				BentNormalInterpolationTarget->GetRenderTargetItem().TargetableTexture,
-				bBindIrradiance ? IrradianceInterpolationTarget->GetRenderTargetItem().TargetableTexture : NULL
-			};
-
-			SetRenderTargets(RHICmdList, ARRAY_COUNT(RenderTargets) - (bBindIrradiance ? 0 : 1), RenderTargets, DepthBuffer, 0, NULL);
-
-			FLinearColor ClearColors[2] = {FLinearColor(0, 0, 0, 0), FLinearColor(0, 0, 0, 0)};
-			RHICmdList.ClearMRT(true, ARRAY_COUNT(ClearColors) - (bBindIrradiance ? 0 : 1), ClearColors, false, 0, false, 0, FIntRect());
-		}
-#endif
 
 		static int32 NumInterpolationSections = 8;
 		if (GCircleVertexBuffer.NumSections != NumInterpolationSections)
