UE4 - fix for UE-18300, wait for parallel sorts before shadow setup

diff --git a/Engine/Source/Runtime/Renderer/Private/SceneVisibility.cpp b/Engine/Source/Runtime/Renderer/Private/SceneVisibility.cpp
--- a/Engine/Source/Runtime/Renderer/Private/SceneVisibility.cpp
+++ b/Engine/Source/Runtime/Renderer/Private/SceneVisibility.cpp
@@ -2094,6 +2094,13 @@ void FDeferredShadingSceneRenderer::InitViews(FRHICommandListImmediate& RHICmdLi
 		SortBasePassStaticData(AverageViewPosition);
 	}
 
+	// this cannot be moved later because of static mesh updates for stuff that is only visible in shadows
+	if (SortEvents.Num())
+	{
+		QUICK_SCOPE_CYCLE_COUNTER(STAT_FDeferredShadingSceneRenderer_AsyncSortBasePassStaticData_Wait);
+		FTaskGraphInterface::Get().WaitUntilTasksComplete(SortEvents, ENamedThreads::RenderThread);
+	}
+
 	bool bDynamicShadows = ViewFamily.EngineShowFlags.DynamicShadows && GetShadowQuality() > 0;
 
 	if (bDynamicShadows && !IsSimpleDynamicLightingEnabled())
@@ -2109,12 +2116,6 @@ void FDeferredShadingSceneRenderer::InitViews(FRHICommandListImmediate& RHICmdLi
 		Views[ViewIndex].InitRHIResources(nullptr);
 	}
 
-	if (SortEvents.Num())
-	{
-		QUICK_SCOPE_CYCLE_COUNTER(STAT_FDeferredShadingSceneRenderer_AsyncSortBasePassStaticData_Wait);
-		FTaskGraphInterface::Get().WaitUntilTasksComplete(SortEvents, ENamedThreads::RenderThread);
-	}
-
 	OnStartFrame();
 }
 
