Fix editor sprite 'desync' bug.

diff --git a/Engine/Source/Runtime/Renderer/Private/PostProcess/PostProcessCompositeEditorPrimitives.cpp b/Engine/Source/Runtime/Renderer/Private/PostProcess/PostProcessCompositeEditorPrimitives.cpp
--- a/Engine/Source/Runtime/Renderer/Private/PostProcess/PostProcessCompositeEditorPrimitives.cpp
+++ b/Engine/Source/Runtime/Renderer/Private/PostProcess/PostProcessCompositeEditorPrimitives.cpp
@@ -210,19 +210,14 @@ void FRCPassPostProcessCompositeEditorPrimitives::Process(FRenderingCompositePas
 	const uint32 MSAASampleCount = SceneContext.EditorPrimitivesColor->GetDesc().NumSamples;
 
 	{
+		SetRenderTarget(Context.RHICmdList, ColorTarget, DepthTarget);
 		Context.SetViewportAndCallRHI(DestRect);
-		
+
 		if (bClearIsNeeded)
 		{
-			check(DepthTarget->GetDepthClearValue() == (float)ERHIZBuffer::FarPlane);
-			check(ColorTarget->GetClearColor() == FLinearColor(0, 0, 0, 0));			
-
 			SCOPED_DRAW_EVENT(Context.RHICmdList, ClearViewEditorPrimitives);
-			SetRenderTarget(Context.RHICmdList, ColorTarget, DepthTarget, ESimpleRenderTargetMode::EClearColorAndDepth);
-		}
-		else
-		{
-			SetRenderTarget(Context.RHICmdList, ColorTarget, DepthTarget, ESimpleRenderTargetMode::EExistingColorAndDepth);
+			// Clear color and depth
+			Context.RHICmdList.Clear(true, FLinearColor(0, 0, 0, 0), true, (float)ERHIZBuffer::FarPlane, false, 0, FIntRect());
 		}
 
 		SCOPED_DRAW_EVENT(Context.RHICmdList, RenderEditorPrimitives);
diff --git a/Engine/Source/Runtime/Windows/D3D11RHI/Private/D3D11Commands.cpp b/Engine/Source/Runtime/Windows/D3D11RHI/Private/D3D11Commands.cpp
--- a/Engine/Source/Runtime/Windows/D3D11RHI/Private/D3D11Commands.cpp
+++ b/Engine/Source/Runtime/Windows/D3D11RHI/Private/D3D11Commands.cpp
@@ -1498,14 +1498,8 @@ void FD3D11DynamicRHI::RHIClearMRT(bool bClearColor, int32 NumClearColors, const
 	RHIClearMRTImpl(bClearColor, NumClearColors, ClearColorArray, bClearDepth, Depth, bClearStencil, Stencil, ExcludeRect, true);
 }
 
-bool gNoShaderClears = true;
 void FD3D11DynamicRHI::RHIClearMRTImpl(bool bClearColor, int32 NumClearColors, const FLinearColor* ClearColorArray, bool bClearDepth, float Depth, bool bClearStencil, uint32 Stencil, FIntRect ExcludeRect, bool bForceShaderClear)
 {	
-	if (gNoShaderClears)
-	{
-		bForceShaderClear = false;
-	}
-
 	// Helper struct to record and restore device states RHIClearMRT modifies.
 	class FDeviceStateHelper
 	{
@@ -1554,6 +1548,7 @@ void FD3D11DynamicRHI::RHIClearMRTImpl(bool bClearColor, int32 NumClearColors, c
 			SAFE_RELEASE(OldBlendState);
 			SAFE_RELEASE(OldRasterizerState);
 			SAFE_RELEASE(OldInputLayout);
+			LastBoundShaderStateRHI = nullptr;
 		}
 	public:
 		/** The global D3D device's immediate context */
