UE4 - parallel frustum culling, also a cvar to disable sphere culling, which is usually a waste of time.

diff --git a/Engine/Source/Runtime/Renderer/Private/SceneVisibility.cpp b/Engine/Source/Runtime/Renderer/Private/SceneVisibility.cpp
--- a/Engine/Source/Runtime/Renderer/Private/SceneVisibility.cpp
+++ b/Engine/Source/Runtime/Renderer/Private/SceneVisibility.cpp
@@ -110,6 +110,18 @@ static FAutoConsoleVariableRef CVarLODFadeTime( TEXT("r.LODFadeTime"), GFadeTime
 static float GDistanceFadeMaxTravel = 1000.0f;
 static FAutoConsoleVariableRef CVarDistanceFadeMaxTravel( TEXT("r.DistanceFadeMaxTravel"), GDistanceFadeMaxTravel, TEXT("Max distance that the player can travel during the fade time."), ECVF_RenderThreadSafe );
 
+
+static TAutoConsoleVariable<int32> CVarParallelInitViews(
+	TEXT("r.ParallelInitViews"),
+#if WITH_EDITOR
+	0,  
+#else
+	1,  
+#endif
+	TEXT("Toggles parallel init views. 0 = off; 1 = on"),
+	ECVF_RenderThreadSafe
+	);
+
 /*------------------------------------------------------------------------------
 	Visibility determination.
 ------------------------------------------------------------------------------*/
@@ -229,67 +241,110 @@ bool FViewInfo::IsDistanceCulled( float DistanceSquared, float MinDrawDistance,
 	return ( bDistanceCulled && !bStillFading );
 }
 
-/**
- * Frustum cull primitives in the scene against the view.
- */
-template<bool UseCustomCulling>
+static int32 FrustumCullNumWordsPerTask = 128;
+static FAutoConsoleVariableRef CVarFrustumCullNumWordsPerTask(
+	TEXT("r.FrustumCullNumWordsPerTask"),
+	FrustumCullNumWordsPerTask,
+	TEXT("Performance tweak. Controls the granularity for the ParallelFor for frustum culling."),
+	ECVF_Default
+	);
+
+template<bool UseCustomCulling, bool bAlsoUseSphereTest>
 static int32 FrustumCull(const FScene* Scene, FViewInfo& View)
 {
 	SCOPE_CYCLE_COUNTER(STAT_FrustumCull);
 
-	int32 NumCulledPrimitives = 0;
-	FVector ViewOriginForDistanceCulling = View.ViewMatrices.ViewOrigin;
+	FThreadSafeCounter NumCulledPrimitives;
 	float MaxDrawDistanceScale = GetCachedScalabilityCVars().ViewDistanceScale;
-	float FadeRadius = GDisableLODFade ? 0.0f : GDistanceFadeMaxTravel;
-	uint8 CustomVisibilityFlags = EOcclusionFlags::CanBeOccluded | EOcclusionFlags::HasPrecomputedVisibility;
 
-	for (FSceneBitArray::FIterator BitIt(View.PrimitiveVisibilityMap); BitIt; ++BitIt)
-	{
-		const FPrimitiveBounds& Bounds = Scene->PrimitiveBounds[BitIt.GetIndex()];
-		float DistanceSquared = (Bounds.Origin - ViewOriginForDistanceCulling).SizeSquared();
-		float MaxDrawDistance = Bounds.MaxDrawDistance * MaxDrawDistanceScale;
-		int32 VisibilityId = INDEX_NONE;
+	//Primitives per ParallelFor task
+	//Using async FrustumCull. Thanks Yager! See https://udn.unrealengine.com/questions/252385/performance-of-frustumcull.html
+	//Performance varies on total primitive count and tasks scheduled. Check the mentioned link above for some measurements.
+	//There have been some changes as compared to the code measured in the link
 
-		if (UseCustomCulling &&
-			((Scene->PrimitiveOcclusionFlags[BitIt.GetIndex()] & CustomVisibilityFlags) == CustomVisibilityFlags))
-		{
-			VisibilityId = Scene->PrimitiveVisibilityIds[BitIt.GetIndex()].ByteIndex;
-		}
+	const int32 BitArrayNum = View.PrimitiveVisibilityMap.Num();
+	const int32 BitArrayWords = FMath::DivideAndRoundUp(View.PrimitiveVisibilityMap.Num(), (int32)NumBitsPerDWORD);
+	const int32 NumTasks = FMath::DivideAndRoundUp(BitArrayWords, FrustumCullNumWordsPerTask);
 
-		// If cull distance is disabled, always show (except foliage)
-		if (View.Family->EngineShowFlags.DistanceCulledPrimitives
-			&& !Scene->Primitives[BitIt.GetIndex()]->Proxy->IsDetailMesh())
+	ParallelFor(NumTasks, 
+		[&NumCulledPrimitives, Scene, &View, MaxDrawDistanceScale](int32 TaskIndex)
 		{
-			MaxDrawDistance = FLT_MAX;
-		}
+			QUICK_SCOPE_CYCLE_COUNTER(STAT_FrustumCull_Loop);
+			const int32 BitArrayNumInner = View.PrimitiveVisibilityMap.Num();
+			FVector ViewOriginForDistanceCulling = View.ViewMatrices.ViewOrigin;
+			float FadeRadius = GDisableLODFade ? 0.0f : GDistanceFadeMaxTravel;
+			uint8 CustomVisibilityFlags = EOcclusionFlags::CanBeOccluded | EOcclusionFlags::HasPrecomputedVisibility;
 
-		// The primitive is always culled if it exceeds the max fade distance or lay outside the view frustum.
-		if (DistanceSquared > FMath::Square(MaxDrawDistance + FadeRadius) ||
-			DistanceSquared < Bounds.MinDrawDistanceSq ||
-			(UseCustomCulling && !View.CustomVisibilityQuery->IsVisible(VisibilityId, FBoxSphereBounds(Bounds.Origin, Bounds.BoxExtent, Bounds.SphereRadius))) ||
-			View.ViewFrustum.IntersectSphere(Bounds.Origin, Bounds.SphereRadius) == false ||
-			View.ViewFrustum.IntersectBox(Bounds.Origin, Bounds.BoxExtent) == false)
-		{
-			STAT(NumCulledPrimitives++);
-			continue;
-		}
+			const int32 TaskWordOffset = TaskIndex * FrustumCullNumWordsPerTask;
 
-		if (DistanceSquared > FMath::Square(MaxDrawDistance))
-		{
-			View.PotentiallyFadingPrimitiveMap.AccessCorrespondingBit(BitIt) = true;
-		}
-		else
-		{
-			// The primitive is visible!
-			BitIt.GetValue() = true;
-			if (DistanceSquared > FMath::Square(MaxDrawDistance - FadeRadius))
+			for (int32 WordIndex = TaskWordOffset; WordIndex < TaskWordOffset + FrustumCullNumWordsPerTask && WordIndex * NumBitsPerDWORD < BitArrayNumInner; WordIndex++)
 			{
-				View.PotentiallyFadingPrimitiveMap.AccessCorrespondingBit(BitIt) = true;
+				uint32 Mask = 0x1;
+				uint32 VisBits = 0;
+				uint32 FadingBits = 0;
+				for(int32 BitSubIndex = 0; BitSubIndex < NumBitsPerDWORD && WordIndex * NumBitsPerDWORD + BitSubIndex < BitArrayNumInner; BitSubIndex++, Mask <<= 1)
+				{
+					int32 Index = WordIndex * NumBitsPerDWORD + BitSubIndex;
+					const FPrimitiveBounds& Bounds = Scene->PrimitiveBounds[Index];
+					float DistanceSquared = (Bounds.Origin - ViewOriginForDistanceCulling).SizeSquared();
+					float MaxDrawDistance = Bounds.MaxDrawDistance * MaxDrawDistanceScale;
+					int32 VisibilityId = INDEX_NONE;
+
+					if (UseCustomCulling &&
+						((Scene->PrimitiveOcclusionFlags[Index] & CustomVisibilityFlags) == CustomVisibilityFlags))
+					{
+						VisibilityId = Scene->PrimitiveVisibilityIds[Index].ByteIndex;
+					}
+
+					// If cull distance is disabled, always show (except foliage)
+					if (View.Family->EngineShowFlags.DistanceCulledPrimitives
+						&& !Scene->Primitives[Index]->Proxy->IsDetailMesh())
+					{
+						MaxDrawDistance = FLT_MAX;
+					}
+
+					// The primitive is always culled if it exceeds the max fade distance or lay outside the view frustum.
+					if (DistanceSquared > FMath::Square(MaxDrawDistance + FadeRadius) ||
+						DistanceSquared < Bounds.MinDrawDistanceSq || 
+						(UseCustomCulling && !View.CustomVisibilityQuery->IsVisible(VisibilityId, FBoxSphereBounds(Bounds.Origin, Bounds.BoxExtent, Bounds.SphereRadius))) ||
+						(bAlsoUseSphereTest && View.ViewFrustum.IntersectSphere(Bounds.Origin, Bounds.SphereRadius) == false) ||
+						View.ViewFrustum.IntersectBox(Bounds.Origin, Bounds.BoxExtent) == false)
+					{
+						STAT(NumCulledPrimitives.Increment());
+					}
+					else
+					{
+						if (DistanceSquared > FMath::Square(MaxDrawDistance))
+						{
+							FadingBits |= Mask;
+						}
+						else
+						{
+							// The primitive is visible!
+							VisBits |= Mask;
+							if (DistanceSquared > FMath::Square(MaxDrawDistance - FadeRadius))
+							{
+								FadingBits |= Mask;
+							}
+						}
+					}
+				}
+				if (FadingBits)
+				{
+					check(!View.PotentiallyFadingPrimitiveMap.GetData()[WordIndex]); // this should start at zero
+					View.PotentiallyFadingPrimitiveMap.GetData()[WordIndex] = FadingBits;
+				}
+				if (VisBits)
+				{
+					check(!View.PrimitiveVisibilityMap.GetData()[WordIndex]); // this should start at zero
+					View.PrimitiveVisibilityMap.GetData()[WordIndex] = VisBits;
+				}
 			}
-		}
-	}
+		},
+		!FApp::ShouldUseThreadingForPerformance() || (UseCustomCulling && !View.CustomVisibilityQuery->IsThreadsafe()) || CVarParallelInitViews.GetValueOnRenderThread() == 0
+	);
 
-	return NumCulledPrimitives;
+	return NumCulledPrimitives.GetValue();
 }
 
 /**
@@ -1054,18 +1109,6 @@ struct FRelevancePacket
 	}
 };
 
-
-static TAutoConsoleVariable<int32> CVarParallelInitViews(
-	TEXT("r.ParallelInitViews"),
-#if WITH_EDITOR
-	0,  
-#else
-	1,  
-#endif
-	TEXT("Toggles parallel init views. 0 = off; 1 = on"),
-	ECVF_RenderThreadSafe
-	);
-
 static void ComputeAndMarkRelevanceForViewParallel(
 	FRHICommandListImmediate& RHICmdList,
 	const FScene* Scene,
@@ -1618,6 +1661,13 @@ void FSceneRenderer::PreVisibilityFrameSetup(FRHICommandListImmediate& RHICmdLis
 	}
 }
 
+static TAutoConsoleVariable<int32> CVarAlsoUseSphereForFrustumCull(
+	TEXT("r.AlsoUseSphereForFrustumCull"),
+	0,  
+	TEXT("Performance tweak. If > 0, then use a sphere cull before and in addition to a box for frustum culling."),
+	ECVF_RenderThreadSafe
+	);
+
 void FSceneRenderer::ComputeViewVisibility(FRHICommandListImmediate& RHICmdList)
 {
 	SCOPE_CYCLE_COUNTER(STAT_ViewVisibilityTime);
@@ -1752,11 +1802,25 @@ void FSceneRenderer::ComputeViewVisibility(FRHICommandListImmediate& RHICmdList)
 			int32 NumCulledPrimitivesForView;
 			if (View.CustomVisibilityQuery && View.CustomVisibilityQuery->Prepare())
 			{
-				NumCulledPrimitivesForView = FrustumCull<true>(Scene, View);
+				if (CVarAlsoUseSphereForFrustumCull.GetValueOnRenderThread())
+				{
+					NumCulledPrimitivesForView = FrustumCull<true, true>(Scene, View);
+				}
+				else
+				{
+					NumCulledPrimitivesForView = FrustumCull<true, true>(Scene, View);
+				}
 			}
 			else
 			{
-				NumCulledPrimitivesForView = FrustumCull<false>(Scene, View);
+				if (CVarAlsoUseSphereForFrustumCull.GetValueOnRenderThread())
+				{
+					NumCulledPrimitivesForView = FrustumCull<false, true>(Scene, View);
+				}
+				else
+				{
+					NumCulledPrimitivesForView = FrustumCull<false, false>(Scene, View);
+				}
 			}
 			STAT(NumCulledPrimitives += NumCulledPrimitivesForView);
 			UpdatePrimitiveFading(Scene, View);			
@@ -1830,7 +1894,7 @@ void FSceneRenderer::ComputeViewVisibility(FRHICommandListImmediate& RHICmdList)
 
 		{
 			SCOPE_CYCLE_COUNTER(STAT_ViewRelevance);
-			ComputeAndMarkRelevanceForViewParallel(RHICmdList, Scene, View, ViewBit, HasDynamicMeshElementsMasks, HasDynamicEditorMeshElementsMasks);
+				ComputeAndMarkRelevanceForViewParallel(RHICmdList, Scene, View, ViewBit, HasDynamicMeshElementsMasks, HasDynamicEditorMeshElementsMasks);
 
 			if (bLODSceneTreeActive)
 			{
@@ -2091,7 +2155,7 @@ void FDeferredShadingSceneRenderer::InitViews(FRHICommandListImmediate& RHICmdLi
 	}
 	else
 	{
-		SortBasePassStaticData(AverageViewPosition);
+	SortBasePassStaticData(AverageViewPosition);
 	}
 
 	// this cannot be moved later because of static mesh updates for stuff that is only visible in shadows
diff --git a/Engine/Source/Runtime/Renderer/Public/RendererInterface.h b/Engine/Source/Runtime/Renderer/Public/RendererInterface.h
--- a/Engine/Source/Runtime/Renderer/Public/RendererInterface.h
+++ b/Engine/Source/Runtime/Renderer/Public/RendererInterface.h
@@ -469,6 +469,12 @@ public:
 
 	/** test primitive visiblity */
 	virtual bool IsVisible(int32 VisibilityId, const FBoxSphereBounds& Bounds) = 0;
+
+	/** return true if we can call IsVisible from a ParallelFor */
+	virtual bool IsThreadsafe()
+	{
+		return false;
+	}
 };
 
 class ICustomCulling
