minor CPU optimization of the decal code

diff --git a/Engine/Source/Runtime/Renderer/Private/DecalRenderingShared.cpp b/Engine/Source/Runtime/Renderer/Private/DecalRenderingShared.cpp
--- a/Engine/Source/Runtime/Renderer/Private/DecalRenderingShared.cpp
+++ b/Engine/Source/Runtime/Renderer/Private/DecalRenderingShared.cpp
@@ -15,8 +15,6 @@ static TAutoConsoleVariable<float> CVarDecalFadeScreenSizeMultiplier(
 	TEXT("  Smaller means decals fade less aggressively.")
 	);
 
-static const FVector GDefaultDecalSize(1.0f, 1.0f, 1.0f);
-
 static bool IsBlendModeSupported(EShaderPlatform Platform, EDecalBlendMode DecalBlendMode)
 {
 	if (IsMobilePlatform(Platform))
@@ -163,9 +161,6 @@ public:
 
 		FTransform ComponentTrans = DecalProxy.ComponentTrans;
 
-		// 1,1,1 requires no scale
-		//			ComponentTrans = ComponentTrans.GetScaled(GDefaultDecalSize);
-
 		FMatrix WorldToComponent = ComponentTrans.ToMatrixWithScale().InverseFast();
 
 		// Set the transform from screen space to light space.
@@ -187,9 +182,6 @@ public:
 		{
 			const FMatrix DecalToWorldValue = ComponentTrans.ToMatrixWithScale();
 			
-			// 1,1,1 requires no scale
-			//			DecalToWorldValue = DecalToWorldValue.GetScaled(GDefaultDecalSize);
-
 			SetShaderValue(RHICmdList, ShaderRHI, DecalToWorld, DecalToWorldValue);
 		}
 
@@ -241,9 +233,9 @@ void FDecalRendering::BuildVisibleDecalList(const FScene& Scene, const FViewInfo
 
 		// can be optimized as we test against a sphere around the box instead of the box itself
 		const float ConservativeRadius = FMath::Sqrt(
-				ComponentToWorldMatrix.GetScaledAxis(EAxis::X).SizeSquared() * FMath::Square(GDefaultDecalSize.X) +
-				ComponentToWorldMatrix.GetScaledAxis(EAxis::Y).SizeSquared() * FMath::Square(GDefaultDecalSize.Y) +
-				ComponentToWorldMatrix.GetScaledAxis(EAxis::Z).SizeSquared() * FMath::Square(GDefaultDecalSize.Z));
+				ComponentToWorldMatrix.GetScaledAxis(EAxis::X).SizeSquared() +
+				ComponentToWorldMatrix.GetScaledAxis(EAxis::Y).SizeSquared() +
+				ComponentToWorldMatrix.GetScaledAxis(EAxis::Z).SizeSquared());
 
 		// can be optimized as the test is too conservative (sphere instead of OBB)
 		if(ConservativeRadius < SMALL_NUMBER || !View.ViewFrustum.IntersectSphere(ComponentToWorldMatrix.GetOrigin(), ConservativeRadius))
@@ -323,9 +315,8 @@ void FDecalRendering::BuildVisibleDecalList(const FScene& Scene, const FViewInfo
 
 FMatrix FDecalRendering::ComputeComponentToClipMatrix(const FViewInfo& View, const FMatrix& DecalComponentToWorld)
 {
-	const FScaleMatrix DecalScaleTransform(GDefaultDecalSize);
-	const FTranslationMatrix PreViewTranslation(View.ViewMatrices.PreViewTranslation);
-	return DecalScaleTransform * DecalComponentToWorld * PreViewTranslation * View.ViewMatrices.TranslatedViewProjectionMatrix;
+	FMatrix ComponentToWorldMatrixTrans = DecalComponentToWorld.ConcatTranslation(View.ViewMatrices.PreViewTranslation);
+	return ComponentToWorldMatrixTrans * View.ViewMatrices.TranslatedViewProjectionMatrix;
 }
 
 FDecalRendering::ERenderTargetMode FDecalRendering::ComputeRenderTargetMode(EShaderPlatform Platform, EDecalBlendMode DecalBlendMode)
