UE4 - fix typo on sphere cull github pull request 1380

diff --git a/Engine/Source/Runtime/Renderer/Private/SceneVisibility.cpp b/Engine/Source/Runtime/Renderer/Private/SceneVisibility.cpp
--- a/Engine/Source/Runtime/Renderer/Private/SceneVisibility.cpp
+++ b/Engine/Source/Runtime/Renderer/Private/SceneVisibility.cpp
@@ -1808,7 +1808,7 @@ void FSceneRenderer::ComputeViewVisibility(FRHICommandListImmediate& RHICmdList)
 				}
 				else
 				{
-					NumCulledPrimitivesForView = FrustumCull<true, true>(Scene, View);
+					NumCulledPrimitivesForView = FrustumCull<true, false>(Scene, View);
 				}
 			}
 			else
