Objects with 0 scale are excluded from the distance field scene, fixes bugs with several lighting methods with 0 scale objects

diff --git a/Engine/Source/Runtime/Renderer/Private/DistanceFieldObjectManagement.cpp b/Engine/Source/Runtime/Renderer/Private/DistanceFieldObjectManagement.cpp
--- a/Engine/Source/Runtime/Renderer/Private/DistanceFieldObjectManagement.cpp
+++ b/Engine/Source/Runtime/Renderer/Private/DistanceFieldObjectManagement.cpp
@@ -543,154 +543,163 @@ void UpdateGlobalDistanceFieldObjectRemoves(FRHICommandListImmediate& RHICmdList
 			DistanceFieldSceneData.SurfelAllocations.RemovePrimitive(Primitive);
 			DistanceFieldSceneData.InstancedSurfelAllocations.RemovePrimitive(Primitive);
 			const TArray<int32, TInlineAllocator<1>>& DistanceFieldInstanceIndices = DistanceFieldSceneData.PendingRemoveOperations[RemoveIndex].DistanceFieldInstanceIndices;
-			PendingRemoveOperations.Append(DistanceFieldInstanceIndices);
 
 			for (int32 RemoveInstanceIndex = 0; RemoveInstanceIndex < DistanceFieldInstanceIndices.Num(); RemoveInstanceIndex++)
 			{
 				const int32 InstanceIndex = DistanceFieldInstanceIndices[RemoveInstanceIndex];
-				DistanceFieldSceneData.PrimitiveModifiedBounds.Add(DistanceFieldSceneData.PrimitiveInstanceMapping[InstanceIndex].BoundingSphere);
+
+				// InstanceIndex will be -1 with zero scale meshes
+				if (InstanceIndex >= 0)
+				{
+					DistanceFieldSceneData.PrimitiveModifiedBounds.Add(DistanceFieldSceneData.PrimitiveInstanceMapping[InstanceIndex].BoundingSphere);
+					PendingRemoveOperations.Add(InstanceIndex);
+				}
 			}
 		}
 
-		check(DistanceFieldSceneData.NumObjectsInBuffer >= PendingRemoveOperations.Num());
+		DistanceFieldSceneData.PendingRemoveOperations.Reset();
+
+		if (PendingRemoveOperations.Num() > 0)
+		{
+			check(DistanceFieldSceneData.NumObjectsInBuffer >= PendingRemoveOperations.Num());
 
-		// Sort from smallest to largest
-		PendingRemoveOperations.Sort();
+			// Sort from smallest to largest
+			PendingRemoveOperations.Sort();
 
-		// We have multiple remove requests enqueued in PendingRemoveOperations, can only use the RemoveAtSwap version when there won't be collisions
-		const bool bUseRemoveAtSwap = PendingRemoveOperations.Last() < DistanceFieldSceneData.NumObjectsInBuffer - PendingRemoveOperations.Num();
+			// We have multiple remove requests enqueued in PendingRemoveOperations, can only use the RemoveAtSwap version when there won't be collisions
+			const bool bUseRemoveAtSwap = PendingRemoveOperations.Last() < DistanceFieldSceneData.NumObjectsInBuffer - PendingRemoveOperations.Num();
 
-		if (bUseRemoveAtSwap)
-		{
-			// Remove everything in parallel in the same buffer with a RemoveAtSwap algorithm
-			for (int32 RemovePrimitiveIndex = 0; RemovePrimitiveIndex < PendingRemoveOperations.Num(); RemovePrimitiveIndex++)
+			if (bUseRemoveAtSwap)
 			{
-				DistanceFieldSceneData.NumObjectsInBuffer--;
-				const int32 RemoveIndex = PendingRemoveOperations[RemovePrimitiveIndex];
-				const int32 MoveFromIndex = DistanceFieldSceneData.NumObjectsInBuffer;
+				// Remove everything in parallel in the same buffer with a RemoveAtSwap algorithm
+				for (int32 RemovePrimitiveIndex = 0; RemovePrimitiveIndex < PendingRemoveOperations.Num(); RemovePrimitiveIndex++)
+				{
+					DistanceFieldSceneData.NumObjectsInBuffer--;
+					const int32 RemoveIndex = PendingRemoveOperations[RemovePrimitiveIndex];
+					const int32 MoveFromIndex = DistanceFieldSceneData.NumObjectsInBuffer;
 
-				check(RemoveIndex != MoveFromIndex);
-				// Queue a compute shader move
-				RemoveObjectIndices.Add(FIntRect(RemoveIndex, MoveFromIndex, 0, 0));
+					check(RemoveIndex != MoveFromIndex);
+					// Queue a compute shader move
+					RemoveObjectIndices.Add(FIntRect(RemoveIndex, MoveFromIndex, 0, 0));
 
-				// Fixup indices of the primitive that is being moved
-				FPrimitiveAndInstance& PrimitiveAndInstanceBeingMoved = DistanceFieldSceneData.PrimitiveInstanceMapping[MoveFromIndex];
-				check(PrimitiveAndInstanceBeingMoved.Primitive && PrimitiveAndInstanceBeingMoved.Primitive->DistanceFieldInstanceIndices.Num() > 0);
-				PrimitiveAndInstanceBeingMoved.Primitive->DistanceFieldInstanceIndices[PrimitiveAndInstanceBeingMoved.InstanceIndex] = RemoveIndex;
+					// Fixup indices of the primitive that is being moved
+					FPrimitiveAndInstance& PrimitiveAndInstanceBeingMoved = DistanceFieldSceneData.PrimitiveInstanceMapping[MoveFromIndex];
+					check(PrimitiveAndInstanceBeingMoved.Primitive && PrimitiveAndInstanceBeingMoved.Primitive->DistanceFieldInstanceIndices.Num() > 0);
+					PrimitiveAndInstanceBeingMoved.Primitive->DistanceFieldInstanceIndices[PrimitiveAndInstanceBeingMoved.InstanceIndex] = RemoveIndex;
 
-				DistanceFieldSceneData.PrimitiveInstanceMapping.RemoveAtSwap(RemoveIndex);
+					DistanceFieldSceneData.PrimitiveInstanceMapping.RemoveAtSwap(RemoveIndex);
+				}
 			}
-		}
-		else
-		{
-			const double StartTime = FPlatformTime::Seconds();
+			else
+			{
+				const double StartTime = FPlatformTime::Seconds();
 
-			// Have to copy the object data to allow parallel removing
-			TemporaryCopySourceBuffers = DistanceFieldSceneData.ObjectBuffers;
-			DistanceFieldSceneData.ObjectBuffers = new FDistanceFieldObjectBuffers();
-			DistanceFieldSceneData.ObjectBuffers->MaxObjects = TemporaryCopySourceBuffers->MaxObjects;
-			DistanceFieldSceneData.ObjectBuffers->Initialize();
+				// Have to copy the object data to allow parallel removing
+				TemporaryCopySourceBuffers = DistanceFieldSceneData.ObjectBuffers;
+				DistanceFieldSceneData.ObjectBuffers = new FDistanceFieldObjectBuffers();
+				DistanceFieldSceneData.ObjectBuffers->MaxObjects = TemporaryCopySourceBuffers->MaxObjects;
+				DistanceFieldSceneData.ObjectBuffers->Initialize();
 
-			TArray<FPrimitiveAndInstance> OriginalPrimitiveInstanceMapping = DistanceFieldSceneData.PrimitiveInstanceMapping;
-			DistanceFieldSceneData.PrimitiveInstanceMapping.Reset();
+				TArray<FPrimitiveAndInstance> OriginalPrimitiveInstanceMapping = DistanceFieldSceneData.PrimitiveInstanceMapping;
+				DistanceFieldSceneData.PrimitiveInstanceMapping.Reset();
 
-			const int32 NumDestObjects = DistanceFieldSceneData.NumObjectsInBuffer - PendingRemoveOperations.Num();
-			int32 SourceIndex = 0;
-			int32 NextPendingRemoveIndex = 0;
+				const int32 NumDestObjects = DistanceFieldSceneData.NumObjectsInBuffer - PendingRemoveOperations.Num();
+				int32 SourceIndex = 0;
+				int32 NextPendingRemoveIndex = 0;
 
-			for (int32 DestinationIndex = 0; DestinationIndex < NumDestObjects; DestinationIndex++)
-			{
-				while (NextPendingRemoveIndex < PendingRemoveOperations.Num()
-					&& PendingRemoveOperations[NextPendingRemoveIndex] == SourceIndex)
+				for (int32 DestinationIndex = 0; DestinationIndex < NumDestObjects; DestinationIndex++)
 				{
-					NextPendingRemoveIndex++;
-					SourceIndex++;
-				}
+					while (NextPendingRemoveIndex < PendingRemoveOperations.Num()
+						&& PendingRemoveOperations[NextPendingRemoveIndex] == SourceIndex)
+					{
+						NextPendingRemoveIndex++;
+						SourceIndex++;
+					}
 
-				// Queue a compute shader move
-				RemoveObjectIndices.Add(FIntRect(DestinationIndex, SourceIndex, 0, 0));
+					// Queue a compute shader move
+					RemoveObjectIndices.Add(FIntRect(DestinationIndex, SourceIndex, 0, 0));
 
-				// Fixup indices of the primitive that is being moved
-				FPrimitiveAndInstance& PrimitiveAndInstanceBeingMoved = OriginalPrimitiveInstanceMapping[SourceIndex];
-				check(PrimitiveAndInstanceBeingMoved.Primitive && PrimitiveAndInstanceBeingMoved.Primitive->DistanceFieldInstanceIndices.Num() > 0);
-				PrimitiveAndInstanceBeingMoved.Primitive->DistanceFieldInstanceIndices[PrimitiveAndInstanceBeingMoved.InstanceIndex] = DestinationIndex;
+					// Fixup indices of the primitive that is being moved
+					FPrimitiveAndInstance& PrimitiveAndInstanceBeingMoved = OriginalPrimitiveInstanceMapping[SourceIndex];
+					check(PrimitiveAndInstanceBeingMoved.Primitive && PrimitiveAndInstanceBeingMoved.Primitive->DistanceFieldInstanceIndices.Num() > 0);
+					PrimitiveAndInstanceBeingMoved.Primitive->DistanceFieldInstanceIndices[PrimitiveAndInstanceBeingMoved.InstanceIndex] = DestinationIndex;
 
-				check(DistanceFieldSceneData.PrimitiveInstanceMapping.Num() == DestinationIndex);
-				DistanceFieldSceneData.PrimitiveInstanceMapping.Add(PrimitiveAndInstanceBeingMoved);
+					check(DistanceFieldSceneData.PrimitiveInstanceMapping.Num() == DestinationIndex);
+					DistanceFieldSceneData.PrimitiveInstanceMapping.Add(PrimitiveAndInstanceBeingMoved);
 
-				SourceIndex++;
-			}
+					SourceIndex++;
+				}
 
-			DistanceFieldSceneData.NumObjectsInBuffer = NumDestObjects;
+				DistanceFieldSceneData.NumObjectsInBuffer = NumDestObjects;
 
-			if (GAOLogObjectBufferReallocation)
-			{
-				const float ElapsedTime = (float)(FPlatformTime::Seconds() - StartTime);
-				UE_LOG(LogDistanceField,Warning,TEXT("Global object buffer realloc %.3fs"), ElapsedTime);
-			}
+				if (GAOLogObjectBufferReallocation)
+				{
+					const float ElapsedTime = (float)(FPlatformTime::Seconds() - StartTime);
+					UE_LOG(LogDistanceField,Warning,TEXT("Global object buffer realloc %.3fs"), ElapsedTime);
+				}
 
-			/*
-			// Have to remove one at a time while any entries to remove are at the end of the buffer
-			DistanceFieldSceneData.NumObjectsInBuffer--;
-			const int32 RemoveIndex = DistanceFieldSceneData.PendingRemoveOperations[ParallelConflictIndex];
-			const int32 MoveFromIndex = DistanceFieldSceneData.NumObjectsInBuffer;
+				/*
+				// Have to remove one at a time while any entries to remove are at the end of the buffer
+				DistanceFieldSceneData.NumObjectsInBuffer--;
+				const int32 RemoveIndex = DistanceFieldSceneData.PendingRemoveOperations[ParallelConflictIndex];
+				const int32 MoveFromIndex = DistanceFieldSceneData.NumObjectsInBuffer;
 
-			if (RemoveIndex != MoveFromIndex)
-			{
-				// Queue a compute shader move
-				RemoveObjectIndices.Add(FIntRect(RemoveIndex, MoveFromIndex, 0, 0));
+				if (RemoveIndex != MoveFromIndex)
+				{
+					// Queue a compute shader move
+					RemoveObjectIndices.Add(FIntRect(RemoveIndex, MoveFromIndex, 0, 0));
 
-				// Fixup indices of the primitive that is being moved
-				FPrimitiveAndInstance& PrimitiveAndInstanceBeingMoved = DistanceFieldSceneData.PrimitiveInstanceMapping[MoveFromIndex];
-				check(PrimitiveAndInstanceBeingMoved.Primitive && PrimitiveAndInstanceBeingMoved.Primitive->DistanceFieldInstanceIndices.Num() > 0);
-				PrimitiveAndInstanceBeingMoved.Primitive->DistanceFieldInstanceIndices[PrimitiveAndInstanceBeingMoved.InstanceIndex] = RemoveIndex;
+					// Fixup indices of the primitive that is being moved
+					FPrimitiveAndInstance& PrimitiveAndInstanceBeingMoved = DistanceFieldSceneData.PrimitiveInstanceMapping[MoveFromIndex];
+					check(PrimitiveAndInstanceBeingMoved.Primitive && PrimitiveAndInstanceBeingMoved.Primitive->DistanceFieldInstanceIndices.Num() > 0);
+					PrimitiveAndInstanceBeingMoved.Primitive->DistanceFieldInstanceIndices[PrimitiveAndInstanceBeingMoved.InstanceIndex] = RemoveIndex;
+				}
+
+				DistanceFieldSceneData.PrimitiveInstanceMapping.RemoveAtSwap(RemoveIndex);
+				DistanceFieldSceneData.PendingRemoveOperations.RemoveAtSwap(ParallelConflictIndex);
+				*/
 			}
 
-			DistanceFieldSceneData.PrimitiveInstanceMapping.RemoveAtSwap(RemoveIndex);
-			DistanceFieldSceneData.PendingRemoveOperations.RemoveAtSwap(ParallelConflictIndex);
-			*/
-		}
+			PendingRemoveOperations.Reset();
 
-		PendingRemoveOperations.Reset();
-		DistanceFieldSceneData.PendingRemoveOperations.Reset();
-
-		if (RemoveObjectIndices.Num() > 0)
-		{
-			if (RemoveObjectIndices.Num() > GDistanceFieldRemoveIndices.RemoveIndices.MaxElements)
+			if (RemoveObjectIndices.Num() > 0)
 			{
-				GDistanceFieldRemoveIndices.RemoveIndices.MaxElements = RemoveObjectIndices.Num() * 5 / 4;
-				GDistanceFieldRemoveIndices.RemoveIndices.Release();
-				GDistanceFieldRemoveIndices.RemoveIndices.Initialize();
-			}
+				if (RemoveObjectIndices.Num() > GDistanceFieldRemoveIndices.RemoveIndices.MaxElements)
+				{
+					GDistanceFieldRemoveIndices.RemoveIndices.MaxElements = RemoveObjectIndices.Num() * 5 / 4;
+					GDistanceFieldRemoveIndices.RemoveIndices.Release();
+					GDistanceFieldRemoveIndices.RemoveIndices.Initialize();
+				}
 
-			void* LockedBuffer = RHILockVertexBuffer(GDistanceFieldRemoveIndices.RemoveIndices.Buffer, 0, GDistanceFieldRemoveIndices.RemoveIndices.Buffer->GetSize(), RLM_WriteOnly);
-			const uint32 MemcpySize = RemoveObjectIndices.GetTypeSize() * RemoveObjectIndices.Num();
-			check(GDistanceFieldRemoveIndices.RemoveIndices.Buffer->GetSize() >= MemcpySize);
-			FPlatformMemory::Memcpy(LockedBuffer, RemoveObjectIndices.GetData(), MemcpySize);
-			RHIUnlockVertexBuffer(GDistanceFieldRemoveIndices.RemoveIndices.Buffer);
+				void* LockedBuffer = RHILockVertexBuffer(GDistanceFieldRemoveIndices.RemoveIndices.Buffer, 0, GDistanceFieldRemoveIndices.RemoveIndices.Buffer->GetSize(), RLM_WriteOnly);
+				const uint32 MemcpySize = RemoveObjectIndices.GetTypeSize() * RemoveObjectIndices.Num();
+				check(GDistanceFieldRemoveIndices.RemoveIndices.Buffer->GetSize() >= MemcpySize);
+				FPlatformMemory::Memcpy(LockedBuffer, RemoveObjectIndices.GetData(), MemcpySize);
+				RHIUnlockVertexBuffer(GDistanceFieldRemoveIndices.RemoveIndices.Buffer);
 
-			if (bUseRemoveAtSwap)
-			{
-				check(!TemporaryCopySourceBuffers);
-				TShaderMapRef<TRemoveObjectsFromBufferCS<true> > ComputeShader(GetGlobalShaderMap(Scene->GetFeatureLevel()));
-				RHICmdList.SetComputeShader(ComputeShader->GetComputeShader());
-				ComputeShader->SetParameters(RHICmdList, Scene, RemoveObjectIndices.Num(), GDistanceFieldRemoveIndices.RemoveIndices.BufferSRV, NULL, NULL);
+				if (bUseRemoveAtSwap)
+				{
+					check(!TemporaryCopySourceBuffers);
+					TShaderMapRef<TRemoveObjectsFromBufferCS<true> > ComputeShader(GetGlobalShaderMap(Scene->GetFeatureLevel()));
+					RHICmdList.SetComputeShader(ComputeShader->GetComputeShader());
+					ComputeShader->SetParameters(RHICmdList, Scene, RemoveObjectIndices.Num(), GDistanceFieldRemoveIndices.RemoveIndices.BufferSRV, NULL, NULL);
 
-				DispatchComputeShader(RHICmdList, *ComputeShader, FMath::DivideAndRoundUp<uint32>(RemoveObjectIndices.Num(), UpdateObjectsGroupSize), 1, 1);
-				ComputeShader->UnsetParameters(RHICmdList);
-			}
-			else
-			{
-				check(TemporaryCopySourceBuffers);
-				TShaderMapRef<TRemoveObjectsFromBufferCS<false> > ComputeShader(GetGlobalShaderMap(Scene->GetFeatureLevel()));
-				RHICmdList.SetComputeShader(ComputeShader->GetComputeShader());
-				ComputeShader->SetParameters(RHICmdList, Scene, RemoveObjectIndices.Num(), GDistanceFieldRemoveIndices.RemoveIndices.BufferSRV, TemporaryCopySourceBuffers->Bounds.SRV, TemporaryCopySourceBuffers->Data.SRV);
+					DispatchComputeShader(RHICmdList, *ComputeShader, FMath::DivideAndRoundUp<uint32>(RemoveObjectIndices.Num(), UpdateObjectsGroupSize), 1, 1);
+					ComputeShader->UnsetParameters(RHICmdList);
+				}
+				else
+				{
+					check(TemporaryCopySourceBuffers);
+					TShaderMapRef<TRemoveObjectsFromBufferCS<false> > ComputeShader(GetGlobalShaderMap(Scene->GetFeatureLevel()));
+					RHICmdList.SetComputeShader(ComputeShader->GetComputeShader());
+					ComputeShader->SetParameters(RHICmdList, Scene, RemoveObjectIndices.Num(), GDistanceFieldRemoveIndices.RemoveIndices.BufferSRV, TemporaryCopySourceBuffers->Bounds.SRV, TemporaryCopySourceBuffers->Data.SRV);
 
-				DispatchComputeShader(RHICmdList, *ComputeShader, FMath::DivideAndRoundUp<uint32>(RemoveObjectIndices.Num(), UpdateObjectsGroupSize), 1, 1);
-				ComputeShader->UnsetParameters(RHICmdList);
+					DispatchComputeShader(RHICmdList, *ComputeShader, FMath::DivideAndRoundUp<uint32>(RemoveObjectIndices.Num(), UpdateObjectsGroupSize), 1, 1);
+					ComputeShader->UnsetParameters(RHICmdList);
 
-				TemporaryCopySourceBuffers->Release();
-				delete TemporaryCopySourceBuffers;
+					TemporaryCopySourceBuffers->Release();
+					delete TemporaryCopySourceBuffers;
+				}
 			}
 		}
 	}
@@ -760,96 +769,116 @@ void ProcessPrimitiveUpdate(
 
 			if (bIsAddOperation)
 			{
-				DistanceFieldSceneData.NumObjectsInBuffer += ObjectLocalToWorldTransforms.Num();
 				PrimitiveSceneInfo->DistanceFieldInstanceIndices.Empty(ObjectLocalToWorldTransforms.Num());
+				PrimitiveSceneInfo->DistanceFieldInstanceIndices.AddZeroed(ObjectLocalToWorldTransforms.Num());
 			}
 
 			for (int32 TransformIndex = 0; TransformIndex < ObjectLocalToWorldTransforms.Num(); TransformIndex++)
 			{
-				const uint32 UploadIndex = bIsAddOperation ? OriginalNumObjects + UploadObjectIndices.Num() : PrimitiveSceneInfo->DistanceFieldInstanceIndices[TransformIndex];
-
-				UploadObjectIndices.Add(UploadIndex);
-
 				FMatrix LocalToWorld = ObjectLocalToWorldTransforms[TransformIndex];
+				const float MaxScale = LocalToWorld.GetMaximumAxisScale();
 
-				if (bMeshWasPlane)
+				// Skip degenerate primitives
+				if (MaxScale > 0)
 				{
-					FVector LocalScales = LocalToWorld.GetScaleVector();
-					FVector AbsLocalScales(FMath::Abs(LocalScales.X), FMath::Abs(LocalScales.Y), FMath::Abs(LocalScales.Z));
-					float MidScale = FMath::Min(AbsLocalScales.X, AbsLocalScales.Y);
-					float ScaleAdjust = FMath::Sign(LocalScales.Z) * MidScale / AbsLocalScales.Z;
-					// The mesh was determined to be a plane flat in Z during the build process, so we can change the Z scale
-					// Helps in cases with modular ground pieces with scales of (10, 10, 1) and some triangles just above Z=0
-					LocalToWorld.SetAxis(2, LocalToWorld.GetScaledAxis(EAxis::Z) * ScaleAdjust);
-				}
+					uint32 UploadIndex;
 
-				const FMatrix VolumeToWorld = FScaleMatrix(LocalVolumeBounds.GetExtent()) 
-					* FTranslationMatrix(LocalVolumeBounds.GetCenter())
-					* LocalToWorld;
+					if (bIsAddOperation)
+					{
+						UploadIndex = OriginalNumObjects + UploadObjectIndices.Num();
+						DistanceFieldSceneData.NumObjectsInBuffer++;
+					}
+					else
+					{
+						UploadIndex = PrimitiveSceneInfo->DistanceFieldInstanceIndices[TransformIndex];
+					}
 
-				const FVector4 ObjectBoundingSphere(VolumeToWorld.GetOrigin(), VolumeToWorld.GetScaleVector().Size());
+					UploadObjectIndices.Add(UploadIndex);
 
-				UploadObjectData.Add(ObjectBoundingSphere);
+					if (bMeshWasPlane)
+					{
+						FVector LocalScales = LocalToWorld.GetScaleVector();
+						FVector AbsLocalScales(FMath::Abs(LocalScales.X), FMath::Abs(LocalScales.Y), FMath::Abs(LocalScales.Z));
+						float MidScale = FMath::Min(AbsLocalScales.X, AbsLocalScales.Y);
+						float ScaleAdjust = FMath::Sign(LocalScales.Z) * MidScale / AbsLocalScales.Z;
+						// The mesh was determined to be a plane flat in Z during the build process, so we can change the Z scale
+						// Helps in cases with modular ground pieces with scales of (10, 10, 1) and some triangles just above Z=0
+						LocalToWorld.SetAxis(2, LocalToWorld.GetScaledAxis(EAxis::Z) * ScaleAdjust);
+					}
 
-				const float MaxExtent = LocalVolumeBounds.GetExtent().GetMax();
+					const FMatrix VolumeToWorld = FScaleMatrix(LocalVolumeBounds.GetExtent()) 
+						* FTranslationMatrix(LocalVolumeBounds.GetCenter())
+						* LocalToWorld;
 
-				const FMatrix UniformScaleVolumeToWorld = FScaleMatrix(MaxExtent) 
-					* FTranslationMatrix(LocalVolumeBounds.GetCenter())
-					* LocalToWorld;
+					const FVector4 ObjectBoundingSphere(VolumeToWorld.GetOrigin(), VolumeToWorld.GetScaleVector().Size());
 
-				const FVector InvBlockSize(1.0f / BlockSize.X, 1.0f / BlockSize.Y, 1.0f / BlockSize.Z);
+					UploadObjectData.Add(ObjectBoundingSphere);
 
-				//float3 VolumeUV = (VolumePosition / LocalPositionExtent * .5f * UVScale + .5f * UVScale + UVAdd;
-				const FVector LocalPositionExtent = LocalVolumeBounds.GetExtent() / FVector(MaxExtent);
-				const FVector UVScale = FVector(BlockSize) * InvTextureDim;
-				const float VolumeScale = UniformScaleVolumeToWorld.GetMaximumAxisScale();
+					const float MaxExtent = LocalVolumeBounds.GetExtent().GetMax();
 
-				const FMatrix WorldToVolume = UniformScaleVolumeToWorld.Inverse();
-				// WorldToVolume
-				UploadObjectData.Add(*(FVector4*)&WorldToVolume.M[0]);
-				UploadObjectData.Add(*(FVector4*)&WorldToVolume.M[1]);
-				UploadObjectData.Add(*(FVector4*)&WorldToVolume.M[2]);
-				UploadObjectData.Add(*(FVector4*)&WorldToVolume.M[3]);
+					const FMatrix UniformScaleVolumeToWorld = FScaleMatrix(MaxExtent) 
+						* FTranslationMatrix(LocalVolumeBounds.GetCenter())
+						* LocalToWorld;
 
-				// Clamp to texel center by subtracting a half texel in the [-1,1] position space
-				// LocalPositionExtent
-				UploadObjectData.Add(FVector4(LocalPositionExtent - InvBlockSize, LocalVolumeBounds.Min.X));
+					const FVector InvBlockSize(1.0f / BlockSize.X, 1.0f / BlockSize.Y, 1.0f / BlockSize.Z);
 
-				// UVScale, VolumeScale and sign gives bGeneratedAsTwoSided
-				const float WSign = bBuiltAsIfTwoSided ? -1 : 1;
-				UploadObjectData.Add(FVector4(FVector(BlockSize) * InvTextureDim * .5f / LocalPositionExtent, WSign * VolumeScale));
+					//float3 VolumeUV = (VolumePosition / LocalPositionExtent * .5f * UVScale + .5f * UVScale + UVAdd;
+					const FVector LocalPositionExtent = LocalVolumeBounds.GetExtent() / FVector(MaxExtent);
+					const FVector UVScale = FVector(BlockSize) * InvTextureDim;
+					const float VolumeScale = UniformScaleVolumeToWorld.GetMaximumAxisScale();
 
-				// UVAdd
-				UploadObjectData.Add(FVector4(FVector(BlockMin) * InvTextureDim + .5f * UVScale, LocalVolumeBounds.Min.Y));
+					const FMatrix WorldToVolume = UniformScaleVolumeToWorld.Inverse();
+					// WorldToVolume
+					UploadObjectData.Add(*(FVector4*)&WorldToVolume.M[0]);
+					UploadObjectData.Add(*(FVector4*)&WorldToVolume.M[1]);
+					UploadObjectData.Add(*(FVector4*)&WorldToVolume.M[2]);
+					UploadObjectData.Add(*(FVector4*)&WorldToVolume.M[3]);
 
-				// Box bounds
-				UploadObjectData.Add(FVector4(LocalVolumeBounds.Max, LocalVolumeBounds.Min.Z));
+					// Clamp to texel center by subtracting a half texel in the [-1,1] position space
+					// LocalPositionExtent
+					UploadObjectData.Add(FVector4(LocalPositionExtent - InvBlockSize, LocalVolumeBounds.Min.X));
 
-				UploadObjectData.Add(*(FVector4*)&UniformScaleVolumeToWorld.M[0]);
-				UploadObjectData.Add(*(FVector4*)&UniformScaleVolumeToWorld.M[1]);
-				UploadObjectData.Add(*(FVector4*)&UniformScaleVolumeToWorld.M[2]);
+					// UVScale, VolumeScale and sign gives bGeneratedAsTwoSided
+					const float WSign = bBuiltAsIfTwoSided ? -1 : 1;
+					UploadObjectData.Add(FVector4(FVector(BlockSize) * InvTextureDim * .5f / LocalPositionExtent, WSign * VolumeScale));
 
-				UploadObjectData.Add(*(FVector4*)&LocalToWorld.M[0]);
-				UploadObjectData.Add(*(FVector4*)&LocalToWorld.M[1]);
-				UploadObjectData.Add(*(FVector4*)&LocalToWorld.M[2]);
-				UploadObjectData.Add(*(FVector4*)&LocalToWorld.M[3]);
+					// UVAdd
+					UploadObjectData.Add(FVector4(FVector(BlockMin) * InvTextureDim + .5f * UVScale, LocalVolumeBounds.Min.Y));
 
-				UploadObjectData.Add(FVector4(Allocation.Offset, Allocation.NumLOD0, Allocation.NumSurfels, InstancedAllocation.Offset + InstancedAllocation.NumSurfels * TransformIndex));
+					// Box bounds
+					UploadObjectData.Add(FVector4(LocalVolumeBounds.Max, LocalVolumeBounds.Min.Z));
 
-				checkSlow(UploadObjectData.Num() % UploadObjectDataStride == 0);
+					UploadObjectData.Add(*(FVector4*)&UniformScaleVolumeToWorld.M[0]);
+					UploadObjectData.Add(*(FVector4*)&UniformScaleVolumeToWorld.M[1]);
+					UploadObjectData.Add(*(FVector4*)&UniformScaleVolumeToWorld.M[2]);
 
-				if (bIsAddOperation)
-				{
-					const int32 AddIndex = UploadIndex;
-					DistanceFieldSceneData.PrimitiveInstanceMapping.Add(FPrimitiveAndInstance(ObjectBoundingSphere, PrimitiveSceneInfo, TransformIndex));
-					PrimitiveSceneInfo->DistanceFieldInstanceIndices.Add(AddIndex);
+					UploadObjectData.Add(*(FVector4*)&LocalToWorld.M[0]);
+					UploadObjectData.Add(*(FVector4*)&LocalToWorld.M[1]);
+					UploadObjectData.Add(*(FVector4*)&LocalToWorld.M[2]);
+					UploadObjectData.Add(*(FVector4*)&LocalToWorld.M[3]);
+
+					UploadObjectData.Add(FVector4(Allocation.Offset, Allocation.NumLOD0, Allocation.NumSurfels, InstancedAllocation.Offset + InstancedAllocation.NumSurfels * TransformIndex));
+
+					checkSlow(UploadObjectData.Num() % UploadObjectDataStride == 0);
+
+					if (bIsAddOperation)
+					{
+						const int32 AddIndex = UploadIndex;
+						DistanceFieldSceneData.PrimitiveInstanceMapping.Add(FPrimitiveAndInstance(ObjectBoundingSphere, PrimitiveSceneInfo, TransformIndex));
+						PrimitiveSceneInfo->DistanceFieldInstanceIndices[TransformIndex] = AddIndex;
+					}
+					else 
+					{
+						DistanceFieldSceneData.PrimitiveInstanceMapping[PrimitiveSceneInfo->DistanceFieldInstanceIndices[TransformIndex]].BoundingSphere = ObjectBoundingSphere;
+					}
+
+					DistanceFieldSceneData.PrimitiveModifiedBounds.Add(ObjectBoundingSphere);
 				}
-				else 
+				else if (bIsAddOperation)
 				{
-					DistanceFieldSceneData.PrimitiveInstanceMapping[PrimitiveSceneInfo->DistanceFieldInstanceIndices[TransformIndex]].BoundingSphere = ObjectBoundingSphere;
+					// Set to -1 for zero scale meshes
+					PrimitiveSceneInfo->DistanceFieldInstanceIndices[TransformIndex] = -1;
 				}
-
-				DistanceFieldSceneData.PrimitiveModifiedBounds.Add(ObjectBoundingSphere);
 			}
 		}
 		else
