fixed UE-16892 [CrashReport] Crash on Use of Circle Depth of Field in Shader Model 4

diff --git a/Engine/Shaders/PostProcessCircleDOF.usf b/Engine/Shaders/PostProcessCircleDOF.usf
--- a/Engine/Shaders/PostProcessCircleDOF.usf
+++ b/Engine/Shaders/PostProcessCircleDOF.usf
@@ -95,7 +95,7 @@ void CircleDilatePS(float4 UVAndScreenPos : TEXCOORD0, out float OutColor : SV_T
     //   d g
 	//   j M  (M={0,0} point)
 
-	#if COMPILER_GLSL || COMPILER_GLSL_ES2
+	#if COMPILER_GLSL || COMPILER_GLSL_ES2 || FEATURE_LEVEL < FEATURE_LEVEL_SM5
 		float2 UV = UVAndScreenPos.xy + 0.5*PostprocessInput0Size.zw;
 		// This leverages nearest sampling (bilinear won't work).
 		// Probably not the best way to do this.
@@ -122,7 +122,7 @@ void CircleDilatePS(float4 UVAndScreenPos : TEXCOORD0, out float OutColor : SV_T
 		float4 Sg = PostprocessInput0.GatherAlpha(PostprocessInput0Sampler, UV, int2(0,-2));
 		float4 Sj = PostprocessInput0.GatherAlpha(PostprocessInput0Sampler, UV, int2(-2,0));
 		float4 Sm = PostprocessInput0.GatherAlpha(PostprocessInput0Sampler, UV, int2(0,0));
-	#endif // COMPILER_GLSL
+	#endif
 	
 	// Make sure near is only near blur.
 	OutColor = min(0.0, Min16(Sd, Sg, Sj, Sm));
@@ -432,7 +432,7 @@ void MainCircleRecombinePS(in float4 UVAndScreenPos : TEXCOORD0, out float4 OutC
 		float RadianBase = NoizNorm(UVAndScreenPos.xy, 0.010 * RandomOffset.x) * 3.14159 * 2.0;
 		float RadiusJitter = NoizNorm(UVAndScreenPos.xy, 0.013 * RandomOffset.x);
 		float ICoc = PixCoc*sqrt(RadiusJitter);
-		float2 VP = RotVec(float2(ICoc, 0.0), RadianBase) * PostprocessInput0Size.zw;
+		float2 VP = RotVec(ICoc, RadianBase) * PostprocessInput0Size.zw;
 
 		// These two samples will still have jitter induced artifacts (very limited utility).
 		// These two samples will also have bleeding artifacts.
diff --git a/Engine/Source/Runtime/Renderer/Private/PostProcess/PostProcessCircleDOF.cpp b/Engine/Source/Runtime/Renderer/Private/PostProcess/PostProcessCircleDOF.cpp
--- a/Engine/Source/Runtime/Renderer/Private/PostProcess/PostProcessCircleDOF.cpp
+++ b/Engine/Source/Runtime/Renderer/Private/PostProcess/PostProcessCircleDOF.cpp
@@ -82,7 +82,7 @@ class FPostProcessCircleDOFSetupPS : public FGlobalShader
 
 	static bool ShouldCache(EShaderPlatform Platform)
 	{
-		return IsFeatureLevelSupported(Platform, ERHIFeatureLevel::SM5);
+		return IsFeatureLevelSupported(Platform, ERHIFeatureLevel::SM4);
 	}
 
 	static void ModifyCompilationEnvironment(EShaderPlatform Platform, FShaderCompilerEnvironment& OutEnvironment)
@@ -270,7 +270,7 @@ class FPostProcessCircleDOFDilatePS : public FGlobalShader
 
 	static bool ShouldCache(EShaderPlatform Platform)
 	{
-		return IsFeatureLevelSupported(Platform, ERHIFeatureLevel::SM5);
+		return IsFeatureLevelSupported(Platform, ERHIFeatureLevel::SM4);
 	}
 
 	static void ModifyCompilationEnvironment(EShaderPlatform Platform, FShaderCompilerEnvironment& OutEnvironment)
@@ -484,7 +484,7 @@ class FPostProcessCircleDOFPS : public FGlobalShader
 
 	static bool ShouldCache(EShaderPlatform Platform)
 	{
-		return IsFeatureLevelSupported(Platform, ERHIFeatureLevel::SM5);
+		return IsFeatureLevelSupported(Platform, ERHIFeatureLevel::SM4);
 	}
 
 	static void ModifyCompilationEnvironment(EShaderPlatform Platform, FShaderCompilerEnvironment& OutEnvironment)
@@ -674,7 +674,7 @@ class FPostProcessCircleDOFRecombinePS : public FGlobalShader
 
 	static bool ShouldCache(EShaderPlatform Platform)
 	{
-		return IsFeatureLevelSupported(Platform, ERHIFeatureLevel::SM5);
+		return IsFeatureLevelSupported(Platform, ERHIFeatureLevel::SM4);
 	}
 
 	static void ModifyCompilationEnvironment(EShaderPlatform Platform, FShaderCompilerEnvironment& OutEnvironment)
