Updated comment

diff --git a/Engine/Source/Runtime/Renderer/Private/ReflectionEnvironment.cpp b/Engine/Source/Runtime/Renderer/Private/ReflectionEnvironment.cpp
--- a/Engine/Source/Runtime/Renderer/Private/ReflectionEnvironment.cpp
+++ b/Engine/Source/Runtime/Renderer/Private/ReflectionEnvironment.cpp
@@ -94,7 +94,7 @@ void FReflectionEnvironmentCubemapArray::InitDynamicRHI()
 
 		FPooledRenderTargetDesc Desc(FPooledRenderTargetDesc::CreateCubemapDesc(
 		GReflectionCaptureSize,
-		//@todo - get rid of the alpha channel (currently stores brightness which is a constant), could use PF_FloatRGB for half memory, would need to implement RHIReadSurface support
+		// Alpha stores sky mask
 		PF_FloatRGBA, 
 		FClearValueBinding::None,
 		TexCreate_None,
