More graceful handling of overflowing object count for distance field shadowing and DFAO

diff --git a/Engine/Shaders/DistanceFieldShadowing.usf b/Engine/Shaders/DistanceFieldShadowing.usf
--- a/Engine/Shaders/DistanceFieldShadowing.usf
+++ b/Engine/Shaders/DistanceFieldShadowing.usf
@@ -468,7 +468,7 @@ void CullObjectsToTileWithGather(
 	GroupMemoryBarrierWithGroupSync();
 
 	GroupIndex = SceneDepth > MaxTileZ2 ? 1 : 0;
-	NumIntersectingObjects = GroupIndex == 0 ? TileNumObjects0 : TileNumObjects1;
+	NumIntersectingObjects = min(GroupIndex == 0 ? TileNumObjects0 : TileNumObjects1, MAX_INTERSECTING_OBJECTS);
 }
 
 uint DownsampleFactor;
diff --git a/Engine/Shaders/GlobalDistanceField.usf b/Engine/Shaders/GlobalDistanceField.usf
--- a/Engine/Shaders/GlobalDistanceField.usf
+++ b/Engine/Shaders/GlobalDistanceField.usf
@@ -252,7 +252,7 @@ void CompositeObjectDistanceFieldsCS(
 
 	GroupMemoryBarrierWithGroupSync();
 
-	NumCulledObjects = NumTileCulledObjects;
+	NumCulledObjects = min(NumTileCulledObjects, MAX_CULLED_DF_OBJECTS);
 
 #endif
 
diff --git a/Engine/Source/Runtime/Renderer/Private/DistanceFieldLightingShared.h b/Engine/Source/Runtime/Renderer/Private/DistanceFieldLightingShared.h
--- a/Engine/Source/Runtime/Renderer/Private/DistanceFieldLightingShared.h
+++ b/Engine/Source/Runtime/Renderer/Private/DistanceFieldLightingShared.h
@@ -255,8 +255,6 @@ public:
 	{
 		if (MaxObjects > 0)
 		{
-			const uint32 BufferFlags = BUF_ShaderResource;
-
 			ObjectIndirectArguments.Initialize(sizeof(uint32), 5, PF_R32_UINT, BUF_Static | BUF_DrawIndirect);
 			ObjectIndirectDispatch.Initialize(sizeof(uint32), 3, PF_R32_UINT, BUF_Static | BUF_DrawIndirect);
 			Bounds.Initialize(sizeof(FVector4), MaxObjects, PF_A32B32G32R32F);
