Extension interface for PostOpaque deferred lighting effects, and post transparency overlays.
#PR 1055

diff --git a/Engine/Source/Runtime/Renderer/Private/DeferredShadingRenderer.cpp b/Engine/Source/Runtime/Renderer/Private/DeferredShadingRenderer.cpp
--- a/Engine/Source/Runtime/Renderer/Private/DeferredShadingRenderer.cpp
+++ b/Engine/Source/Runtime/Renderer/Private/DeferredShadingRenderer.cpp
@@ -23,6 +23,7 @@
 #include "PostProcess/PostProcessing.h"
 #include "DistanceFieldAtlas.h"
 #include "../../Engine/Private/SkeletalRenderGPUSkin.h"		// GPrevPerBoneMotionBlur
+#include "EngineModule.h"
 
 TAutoConsoleVariable<int32> CVarEarlyZPass(
 	TEXT("r.EarlyZPass"),
@@ -768,6 +769,7 @@ void FDeferredShadingSceneRenderer::Render(FRHICommandListImmediate& RHICmdList)
 {
 	FSceneRenderTargets& SceneContext = FSceneRenderTargets::Get(RHICmdList);
 	bool bDBuffer = IsDBufferEnabled();	
+
 	if (GRHIThread)
 	{
 		QUICK_SCOPE_CYCLE_COUNTER(STAT_OcclusionSubmittedFence_Wait);
@@ -776,7 +778,7 @@ void FDeferredShadingSceneRenderer::Render(FRHICommandListImmediate& RHICmdList)
 		OcclusionSubmittedFence[BlockFrame] = nullptr;
 	}
 
-	if(!ViewFamily.EngineShowFlags.Rendering)
+	if (!ViewFamily.EngineShowFlags.Rendering)
 	{
 		return;
 	}
@@ -791,6 +793,7 @@ void FDeferredShadingSceneRenderer::Render(FRHICommandListImmediate& RHICmdList)
 		// Allocate the maximum scene render target space for the current view family.
 		SceneContext.Allocate(ViewFamily);
 	}
+
 	// Find the visible primitives.
 	InitViews(RHICmdList);
 
@@ -798,7 +801,7 @@ void FDeferredShadingSceneRenderer::Render(FRHICommandListImmediate& RHICmdList)
 	{
 		// there are dynamic attempts to get this target during parallel rendering
 		for (int32 ViewIndex = 0; ViewIndex < Views.Num(); ViewIndex++)
-	{
+		{
 			Views[ViewIndex].GetEyeAdaptation();
 		}	
 	}
@@ -836,7 +839,7 @@ void FDeferredShadingSceneRenderer::Render(FRHICommandListImmediate& RHICmdList)
 
 	static const auto GBufferCVar = IConsoleManager::Get().FindTConsoleVariableDataInt(TEXT("r.GBuffer"));
 	bool bGBuffer = GBufferCVar ? (GBufferCVar->GetValueOnRenderThread() != 0) : true;
-	if(ViewFamily.EngineShowFlags.ForceGBuffer)
+	if (ViewFamily.EngineShowFlags.ForceGBuffer)
 	{
 		bGBuffer = true;
 	}
@@ -951,8 +954,6 @@ void FDeferredShadingSceneRenderer::Render(FRHICommandListImmediate& RHICmdList)
 		}
 	}
 
-
-
 	{
 		QUICK_SCOPE_CYCLE_COUNTER(STAT_FDeferredShadingSceneRenderer_AllocGBufferTargets);
 		SceneContext.PreallocGBufferTargets(bShouldRenderVelocities && bUseVelocityGBuffer);
@@ -960,7 +961,7 @@ void FDeferredShadingSceneRenderer::Render(FRHICommandListImmediate& RHICmdList)
 	}
 	
 	// Clear LPVs for all views
-	if ( FeatureLevel >= ERHIFeatureLevel::SM5 )
+	if (FeatureLevel >= ERHIFeatureLevel::SM5)
 	{
 		QUICK_SCOPE_CYCLE_COUNTER(STAT_FDeferredShadingSceneRenderer_ClearLPVs);
 		ClearLPVs(RHICmdList);
@@ -972,7 +973,7 @@ void FDeferredShadingSceneRenderer::Render(FRHICommandListImmediate& RHICmdList)
 	check(!SceneContext.DBufferB);
 	check(!SceneContext.DBufferC);
 
-	if(bDBuffer)
+	if (bDBuffer)
 	{
 		QUICK_SCOPE_CYCLE_COUNTER(STAT_FDeferredShadingSceneRenderer_DBuffer);
 		SceneContext.ResolveSceneDepthTexture(RHICmdList);
@@ -1002,7 +1003,7 @@ void FDeferredShadingSceneRenderer::Render(FRHICommandListImmediate& RHICmdList)
 		ServiceLocalQueue();
 	}
 
-	if(bIsWireframe && FSceneRenderer::ShouldCompositeEditorPrimitives(Views[0]))
+	if (bIsWireframe && FSceneRenderer::ShouldCompositeEditorPrimitives(Views[0]))
 	{
 		// In Editor we want wire frame view modes to be MSAA for better quality. Resolve will be done with EditorPrimitives
 		SetRenderTarget(RHICmdList, SceneContext.GetEditorPrimitivesColor(), SceneContext.GetEditorPrimitivesDepth());
@@ -1022,21 +1023,21 @@ void FDeferredShadingSceneRenderer::Render(FRHICommandListImmediate& RHICmdList)
 	RHICmdList.SetCurrentStat(GET_STATID(STAT_CLM_AfterBasePass));
 	ServiceLocalQueue();
 
-	  if(ViewFamily.EngineShowFlags.VisualizeLightCulling)
-	  {
-		  // clear out emissive and baked lighting (not too efficient but simple and only needed for this debug view)
-		  SceneContext.BeginRenderingSceneColor(RHICmdList);
-		  RHICmdList.Clear(true, FLinearColor(0, 0, 0, 0), false, (float)ERHIZBuffer::FarPlane, false, 0, FIntRect());
-	  }
+	if (ViewFamily.EngineShowFlags.VisualizeLightCulling)
+	{
+		// clear out emissive and baked lighting (not too efficient but simple and only needed for this debug view)
+		SceneContext.BeginRenderingSceneColor(RHICmdList);
+		RHICmdList.Clear(true, FLinearColor(0, 0, 0, 0), false, (float)ERHIZBuffer::FarPlane, false, 0, FIntRect());
+	}
 
-	  SceneContext.DBufferA.SafeRelease();
-	  SceneContext.DBufferB.SafeRelease();
-	  SceneContext.DBufferC.SafeRelease();
+	SceneContext.DBufferA.SafeRelease();
+	SceneContext.DBufferB.SafeRelease();
+	SceneContext.DBufferC.SafeRelease();
 
-	  // only temporarily available after early z pass and until base pass
-	  check(!SceneContext.DBufferA);
-	  check(!SceneContext.DBufferB);
-	  check(!SceneContext.DBufferC);
+	// only temporarily available after early z pass and until base pass
+	check(!SceneContext.DBufferA);
+	check(!SceneContext.DBufferB);
+	check(!SceneContext.DBufferC);
 
 	if (bRequiresFarZQuadClear)
 	{
@@ -1141,7 +1142,7 @@ void FDeferredShadingSceneRenderer::Render(FRHICommandListImmediate& RHICmdList)
 
 		// Pre-lighting composition lighting stage
 		// e.g. LPV indirect
-		for(int32 ViewIndex = 0; ViewIndex < Views.Num(); ++ViewIndex)
+		for (int32 ViewIndex = 0; ViewIndex < Views.Num(); ++ViewIndex)
 		{
 			SCOPED_CONDITIONAL_DRAW_EVENTF(RHICmdList, EventView,Views.Num() > 1, TEXT("View%d"), ViewIndex);
 			GCompositionLighting.ProcessLpvIndirect(RHICmdList, Views[ViewIndex]);
@@ -1169,7 +1170,7 @@ void FDeferredShadingSceneRenderer::Render(FRHICommandListImmediate& RHICmdList)
 		ServiceLocalQueue();
 	}
 
-	if( ViewFamily.EngineShowFlags.StationaryLightOverlap &&
+	if (ViewFamily.EngineShowFlags.StationaryLightOverlap &&
 		FeatureLevel >= ERHIFeatureLevel::SM4)
 	{
 		RenderStationaryLightOverlap(RHICmdList);
@@ -1187,7 +1188,7 @@ void FDeferredShadingSceneRenderer::Render(FRHICommandListImmediate& RHICmdList)
 	}
 
 	// Draw atmosphere
-	if(ShouldRenderAtmosphere(ViewFamily))
+	if (ShouldRenderAtmosphere(ViewFamily))
 	{
 		QUICK_SCOPE_CYCLE_COUNTER(STAT_FDeferredShadingSceneRenderer_RenderAtmosphere);
 		if (Scene->AtmosphericFog)
@@ -1216,20 +1217,33 @@ void FDeferredShadingSceneRenderer::Render(FRHICommandListImmediate& RHICmdList)
 	GRenderTargetPool.AddPhaseEvent(TEXT("Fog"));
 
 	// Draw fog.
-	if(ShouldRenderFog(ViewFamily))
+	if (ShouldRenderFog(ViewFamily))
 	{
 		QUICK_SCOPE_CYCLE_COUNTER(STAT_FDeferredShadingSceneRenderer_RenderFog);
 		RenderFog(RHICmdList, LightShaftOutput);
 		ServiceLocalQueue();
 	}
 
+	if (GetRendererModule().HasPostOpaqueExtentions())
+	{
+		SceneContext.BeginRenderingSceneColor(RHICmdList);
+		for (int32 ViewIndex = 0; ViewIndex < Views.Num(); ++ViewIndex)
+		{
+			const FViewInfo& View = Views[ViewIndex];
+			RHICmdList.SetViewport(View.ViewRect.Min.X, View.ViewRect.Min.Y, 0.0f, View.ViewRect.Max.X, View.ViewRect.Max.Y, 1.0f);
+			GetRendererModule().RenderPostOpaqueExtensions(View, RHICmdList, SceneContext);
+		}
+
+		SceneContext.FinishRenderingSceneColor(RHICmdList, false);
+	}
+
 	// No longer needed, release
 	LightShaftOutput.LightShaftOcclusion = NULL;
 
 	GRenderTargetPool.AddPhaseEvent(TEXT("Translucency"));
 
 	// Draw translucency.
-	if(ViewFamily.EngineShowFlags.Translucency)
+	if (ViewFamily.EngineShowFlags.Translucency)
 	{
 		SCOPE_CYCLE_COUNTER(STAT_TranslucencyDrawTime);
 
@@ -1257,6 +1271,13 @@ void FDeferredShadingSceneRenderer::Render(FRHICommandListImmediate& RHICmdList)
 		ServiceLocalQueue();
 	}
 
+	for (int32 ViewIndex = 0; ViewIndex < Views.Num(); ++ViewIndex)
+	{
+		const FViewInfo& View = Views[ViewIndex];
+		RHICmdList.SetViewport(View.ViewRect.Min.X, View.ViewRect.Min.Y, 0.0f, View.ViewRect.Max.X, View.ViewRect.Max.Y, 1.0f);
+		GetRendererModule().RenderOverlayExtensions(View, RHICmdList, SceneContext);
+	}
+
 	if (ViewFamily.EngineShowFlags.VisualizeDistanceFieldAO || ViewFamily.EngineShowFlags.VisualizeDistanceFieldGI)
 	{
 		// Use the skylight's max distance if there is one, to be consistent with DFAO shadowing on the skylight
@@ -1279,7 +1300,7 @@ void FDeferredShadingSceneRenderer::Render(FRHICommandListImmediate& RHICmdList)
 	GPrevPerBoneMotionBlur.EndAppend(RHICmdList);
 
 	// Finish rendering for each view.
-	if(ViewFamily.bResolveScene)
+	if (ViewFamily.bResolveScene)
 	{
 		SCOPED_DRAW_EVENT(RHICmdList, PostProcessing);
 		SCOPE_CYCLE_COUNTER(STAT_FinishRenderViewTargetTime);
diff --git a/Engine/Source/Runtime/Renderer/Private/RendererPrivate.h b/Engine/Source/Runtime/Renderer/Private/RendererPrivate.h
--- a/Engine/Source/Runtime/Renderer/Private/RendererPrivate.h
+++ b/Engine/Source/Runtime/Renderer/Private/RendererPrivate.h
@@ -64,9 +64,21 @@ public:
 	virtual void RegisterCustomCullingImpl(ICustomCulling* impl) override;
 	virtual void UnregisterCustomCullingImpl(ICustomCulling* impl) override;
 
+	virtual void RegisterPostOpaqueRenderDelegate(const FPostOpaqueRenderDelegate& PostOpaqueRenderDelegate) override;
+	virtual void RegisterOverlayRenderDelegate(const FPostOpaqueRenderDelegate& OverlayRenderDelegate) override;
+	virtual void RenderPostOpaqueExtensions(const FSceneView& View, FRHICommandListImmediate& RHICmdList, FSceneRenderTargets& SceneContext) override;
+	virtual void RenderOverlayExtensions(const FSceneView& View, FRHICommandListImmediate& RHICmdList, FSceneRenderTargets& SceneContext) override;
+
+	virtual bool HasPostOpaqueExtentions() const override
+	{
+		return PostOpaqueRenderDelegate.IsBound();
+	}
+
 private:
 	TSet<FSceneInterface*> AllocatedScenes;
 	ICustomCulling* CustomCullingImpl;
+	FPostOpaqueRenderDelegate PostOpaqueRenderDelegate;
+	FPostOpaqueRenderDelegate OverlayRenderDelegate;
 };
 
 extern ICustomCulling* GCustomCullingImpl;
diff --git a/Engine/Source/Runtime/Renderer/Private/SceneRendering.cpp b/Engine/Source/Runtime/Renderer/Private/SceneRendering.cpp
--- a/Engine/Source/Runtime/Renderer/Private/SceneRendering.cpp
+++ b/Engine/Source/Runtime/Renderer/Private/SceneRendering.cpp
@@ -1666,6 +1666,47 @@ TGlobalResource<FFilterVertexDeclaration>& FRendererModule::GetFilterVertexDecla
 	return GFilterVertexDeclaration;
 }
 
+void FRendererModule::RegisterPostOpaqueRenderDelegate(const FPostOpaqueRenderDelegate& PostOpaqueRenderDelegate)
+{
+	this->PostOpaqueRenderDelegate = PostOpaqueRenderDelegate;
+}
+
+void FRendererModule::RegisterOverlayRenderDelegate(const FPostOpaqueRenderDelegate& OverlayRenderDelegate)
+{
+	this->OverlayRenderDelegate = OverlayRenderDelegate;
+}
+
+void FRendererModule::RenderPostOpaqueExtensions(const FSceneView& View, FRHICommandListImmediate& RHICmdList, FSceneRenderTargets& SceneContext)
+{
+	check(IsInRenderingThread());
+	FPostOpaqueRenderParameters RenderParameters;
+	RenderParameters.ViewMatrix = View.ViewMatrices.ViewMatrix;
+	RenderParameters.ProjMatrix = View.ViewMatrices.ProjMatrix;
+	RenderParameters.DepthTexture = SceneContext.GetSceneDepthSurface()->GetTexture2D();
+	RenderParameters.SmallDepthTexture = SceneContext.GetSmallDepthSurface()->GetTexture2D();
+
+	RenderParameters.ViewportRect = View.ViewRect;
+	RenderParameters.RHICmdList = &RHICmdList;
+
+	RenderParameters.Uid = (void*)(&View);
+	PostOpaqueRenderDelegate.ExecuteIfBound(RenderParameters);
+}
+
+void FRendererModule::RenderOverlayExtensions(const FSceneView& View, FRHICommandListImmediate& RHICmdList, FSceneRenderTargets& SceneContext)
+{
+	check(IsInRenderingThread());
+	FPostOpaqueRenderParameters RenderParameters;
+	RenderParameters.ViewMatrix = View.ViewMatrices.ViewMatrix;
+	RenderParameters.ProjMatrix = View.ViewMatrices.ProjMatrix;
+	RenderParameters.DepthTexture = SceneContext.GetSceneDepthSurface()->GetTexture2D();
+	RenderParameters.SmallDepthTexture = SceneContext.GetSmallDepthSurface()->GetTexture2D();
+
+	RenderParameters.ViewportRect = View.ViewRect;
+	RenderParameters.RHICmdList = &RHICmdList;
+
+	RenderParameters.Uid=(void*)(&View);
+	OverlayRenderDelegate.ExecuteIfBound(RenderParameters);
+}
 
 bool IsMobileHDR()
 {
diff --git a/Engine/Source/Runtime/Renderer/Public/RendererInterface.h b/Engine/Source/Runtime/Renderer/Public/RendererInterface.h
--- a/Engine/Source/Runtime/Renderer/Public/RendererInterface.h
+++ b/Engine/Source/Runtime/Renderer/Public/RendererInterface.h
@@ -461,6 +461,19 @@ enum EDrawRectangleFlags
 	EDRF_UseTesselatedIndexBuffer
 };
 
+class FPostOpaqueRenderParameters
+{
+	public:
+		FIntRect ViewportRect;
+		FMatrix ViewMatrix;
+		FMatrix ProjMatrix;
+		FRHITexture2D* DepthTexture;
+		FRHITexture2D* SmallDepthTexture;
+		FRHICommandListImmediate* RHICmdList;
+		void* Uid; // A unique identifier for the view.
+};
+DECLARE_DELEGATE_OneParam(FPostOpaqueRenderDelegate, class FPostOpaqueRenderParameters&);
+
 class ICustomVisibilityQuery: public IRefCountedObject
 {
 public:
@@ -583,6 +596,12 @@ public:
 	/** Register/unregister a custom occlusion culling implementation */
 	virtual void RegisterCustomCullingImpl(ICustomCulling* impl) = 0;
 	virtual void UnregisterCustomCullingImpl(ICustomCulling* impl) = 0;
+
+	virtual void RegisterPostOpaqueRenderDelegate(const FPostOpaqueRenderDelegate& PostOpaqueRenderDelegate) = 0;
+	virtual void RegisterOverlayRenderDelegate(const FPostOpaqueRenderDelegate& OverlayRenderDelegate) = 0;
+	virtual void RenderPostOpaqueExtensions(const FSceneView& View, FRHICommandListImmediate& RHICmdList, class FSceneRenderTargets& SceneContext) = 0;
+	virtual void RenderOverlayExtensions(const FSceneView& View, FRHICommandListImmediate& RHICmdList, FSceneRenderTargets& SceneContext) = 0;
+	virtual bool HasPostOpaqueExtentions() const = 0;
 };
 
 
