fixed UE-18633 Strange shader issue in color grading (Mac)

diff --git a/Engine/Shaders/PostProcessCombineLUTs.usf b/Engine/Shaders/PostProcessCombineLUTs.usf
--- a/Engine/Shaders/PostProcessCombineLUTs.usf
+++ b/Engine/Shaders/PostProcessCombineLUTs.usf
@@ -169,7 +169,8 @@ void MainPS(float4 InUV : TEXCOORD0, out float4 OutColor : SV_Target0)
 	{ 
 		float2 UV = InUV.xy;
 
-		UV -= float2(0.5f / (LUTSize * LUTSize), 0.5f / LUTSize);
+		// 0.49999f instead of 0.5f to avoid getting into negative values
+		UV -= float2(0.49999f / (LUTSize * LUTSize), 0.49999f / LUTSize);
 
 		float Scale = LUTSize / (LUTSize - 1);
 
diff --git a/Engine/Source/Runtime/Renderer/Private/PostProcess/PostProcessCombineLUTs.cpp b/Engine/Source/Runtime/Renderer/Private/PostProcess/PostProcessCombineLUTs.cpp
--- a/Engine/Source/Runtime/Renderer/Private/PostProcess/PostProcessCombineLUTs.cpp
+++ b/Engine/Source/Runtime/Renderer/Private/PostProcess/PostProcessCombineLUTs.cpp
@@ -188,7 +188,10 @@ public:
 			// we don't need to set the neutral one
 			if(i != 0)
 			{
-				SetTextureParameter(RHICmdList, ShaderRHI, TextureParameter[i], TextureParameterSampler[i], Texture[i]);
+				// don't use texture asset sampler as it might have anisotropic filtering enabled
+				FSamplerStateRHIParamRef Sampler = TStaticSamplerState<SF_Bilinear, AM_Clamp, AM_Clamp, AM_Clamp, 0, 1>::GetRHI();
+
+				SetTextureParameter(RHICmdList, ShaderRHI, TextureParameter[i], TextureParameterSampler[i], Sampler, Texture[i]->TextureRHI);
 			}
 
 			SetShaderValue(RHICmdList, ShaderRHI, WeightsParameter, Weights[i], i);
diff --git a/Engine/Source/Runtime/Renderer/Private/PostProcess/PostProcessTonemap.cpp b/Engine/Source/Runtime/Renderer/Private/PostProcess/PostProcessTonemap.cpp
--- a/Engine/Source/Runtime/Renderer/Private/PostProcess/PostProcessTonemap.cpp
+++ b/Engine/Source/Runtime/Renderer/Private/PostProcess/PostProcessTonemap.cpp
@@ -886,7 +886,18 @@ public:
 		
 		FGlobalShader::SetParameters(Context.RHICmdList, ShaderRHI, Context.View);
 
-		PostprocessParameter.SetPS(ShaderRHI, Context, TStaticSamplerState<SF_Bilinear,AM_Clamp,AM_Clamp,AM_Clamp>::GetRHI());
+		{
+			// filtering can cost performance so we use point where possible, we don't want anisotropic sampling
+			FSamplerStateRHIParamRef Filters[] =
+			{
+				TStaticSamplerState<SF_Bilinear, AM_Clamp, AM_Clamp, AM_Clamp, 0, 1>::GetRHI(),		// todo: could be SF_Point if fringe is disabled
+				TStaticSamplerState<SF_Bilinear, AM_Clamp, AM_Clamp, AM_Clamp, 0, 1>::GetRHI(),
+				TStaticSamplerState<SF_Point, AM_Clamp, AM_Clamp, AM_Clamp, 0, 1>::GetRHI(),
+				TStaticSamplerState<SF_Bilinear, AM_Clamp, AM_Clamp, AM_Clamp, 0, 1>::GetRHI(),
+			};
+
+			PostprocessParameter.SetPS(ShaderRHI, Context, 0, false, Filters);
+		}
 			
 		SetShaderValue(Context.RHICmdList, ShaderRHI, OverlayColor, Context.View.OverlayColor);
 
