Corrected render target setup for forward shading shadow projections.
Fixes broken modulated shadows with ios/metal.

diff --git a/Engine/Source/Runtime/Renderer/Private/ShadowRendering.cpp b/Engine/Source/Runtime/Renderer/Private/ShadowRendering.cpp
--- a/Engine/Source/Runtime/Renderer/Private/ShadowRendering.cpp
+++ b/Engine/Source/Runtime/Renderer/Private/ShadowRendering.cpp
@@ -2904,7 +2904,7 @@ void FSceneRenderer::RenderProjections(
 	
 	if (bForwardShading)
 	{
-		SceneContext.BeginRenderingSceneColor(RHICmdList, ESimpleRenderTargetMode::EUninitializedColorExistingDepth, FExclusiveDepthStencil::DepthRead_StencilWrite);
+		SceneContext.BeginRenderingSceneColor(RHICmdList, ESimpleRenderTargetMode::EExistingColorAndDepth, FExclusiveDepthStencil::DepthRead_StencilWrite);
 	}
 	else
 	{
