#ue4 - Fixed & optimized LPV for stereoscopic rendering.

-The right eye view now piggybacks off of the left eyes LPV, as opposed to creating and rendering one of it's own.

diff --git a/Engine/Source/Runtime/Engine/Public/SceneManagement.h b/Engine/Source/Runtime/Engine/Public/SceneManagement.h
--- a/Engine/Source/Runtime/Engine/Public/SceneManagement.h
+++ b/Engine/Source/Runtime/Engine/Public/SceneManagement.h
@@ -20,6 +20,7 @@
 
 // Forward declarations.
 class FLightSceneInfo;
+class FSceneViewState;
 class ULightComponent;
 class UDecalComponent;
 class HHitProxy;
@@ -125,12 +126,15 @@ public:
 		return NumChildren > 0;
 	}
 	
+	/** @return	the derived view state object */
+	virtual FSceneViewState* GetConcreteViewState () = 0;
+
 	virtual void AddReferencedObjects(FReferenceCollector& Collector) = 0;
 
 	virtual SIZE_T GetSizeBytes() const { return 0; }
 
 	/** called in InitViews() */
-	virtual void OnStartFrame(FSceneView& CurrentView) = 0;
+	virtual void OnStartFrame(FSceneView& View, FSceneViewFamily& ViewFamily) = 0;
 
 	/** Resets pool for GetReusableMID() */
 	virtual void OnStartPostProcessing(FSceneView& CurrentView) = 0;
diff --git a/Engine/Source/Runtime/Renderer/Private/CompositionLighting/CompositionLighting.cpp b/Engine/Source/Runtime/Renderer/Private/CompositionLighting/CompositionLighting.cpp
--- a/Engine/Source/Runtime/Renderer/Private/CompositionLighting/CompositionLighting.cpp
+++ b/Engine/Source/Runtime/Renderer/Private/CompositionLighting/CompositionLighting.cpp
@@ -55,7 +55,10 @@ static bool IsLpvIndirectPassRequired(FPostprocessContext& Context)
 
 	if(ViewState)
 	{
-		FLightPropagationVolume* LightPropagationVolume = ViewState->GetLightPropagationVolume(Context.View.GetFeatureLevel());
+		// This check should be inclusive to stereo views
+		const bool bIncludeStereoViews = true;
+
+		FLightPropagationVolume* LightPropagationVolume = ViewState->GetLightPropagationVolume(Context.View.GetFeatureLevel(), bIncludeStereoViews);
 
 		if(LightPropagationVolume)
 		{
diff --git a/Engine/Source/Runtime/Renderer/Private/CompositionLighting/PostProcessLpvIndirect.cpp b/Engine/Source/Runtime/Renderer/Private/CompositionLighting/PostProcessLpvIndirect.cpp
--- a/Engine/Source/Runtime/Renderer/Private/CompositionLighting/PostProcessLpvIndirect.cpp
+++ b/Engine/Source/Runtime/Renderer/Private/CompositionLighting/PostProcessLpvIndirect.cpp
@@ -227,7 +227,9 @@ void FRCPassPostProcessLpvIndirect::Process(FRenderingCompositePassContext& Cont
 		return;
 	}
 
-	FLightPropagationVolume* Lpv = ViewState->GetLightPropagationVolume(Context.GetFeatureLevel());
+	// This check should be inclusive to stereo views
+	const bool bIncludeStereoViews = true;
+	FLightPropagationVolume* Lpv = ViewState->GetLightPropagationVolume(Context.GetFeatureLevel(), bIncludeStereoViews);
 
 	const FLightPropagationVolumeSettings& LPVSettings = PostprocessSettings.BlendableManager.GetSingleFinalDataConst<FLightPropagationVolumeSettings>();
 
diff --git a/Engine/Source/Runtime/Renderer/Private/LightPropagationVolume.cpp b/Engine/Source/Runtime/Renderer/Private/LightPropagationVolume.cpp
--- a/Engine/Source/Runtime/Renderer/Private/LightPropagationVolume.cpp
+++ b/Engine/Source/Runtime/Renderer/Private/LightPropagationVolume.cpp
@@ -768,11 +768,10 @@ FLightPropagationVolume::FLightPropagationVolume() :
 	, bEnabled( false )
 	, bGeometryVolumeNeeded( false )
 	, mWriteBufferIndex( 0 )
+	, bNeedsBufferClear( true )
 	, GeometryVolumeGenerated( false )
 	, AsyncJobFenceID(-1)
 {
-	bNeedsBufferClear = true;
-
 	// VPL List buffers
 	mVplListBuffer = new FRWBufferStructured();
 	int32 RSMResolution = FSceneRenderTargets::Get_FrameConstantsOnly().GetReflectiveShadowMapResolution();
@@ -915,21 +914,16 @@ void FLightPropagationVolume::InitSettings(FRHICommandList& RHICmdList, const FS
 		return;
 	}
 
-	float ClearMultiplier = 1.0f;
-
-	// Clear the LPV for the first 2 frames (so both buffers are cleared initially)
-	static int FrameCount = 0;
-	if ( FrameCount < 2 )
-	{
-		bNeedsBufferClear = true;
-	}
-	FrameCount++;
-
 	// Clear the UAVs if necessary
+	float ClearMultiplier = 1.0f;
 	if ( bNeedsBufferClear )
 	{
 		ClearMultiplier = 0.0f;
-		bNeedsBufferClear = false;
+		// Since this is double buffered, the clear flag should remain set for the first 2 frames so that all buffers get cleared.
+		if ( mWriteBufferIndex > 0 )
+		{
+			bNeedsBufferClear = false;
+		}
 	}
 
 	mInjectedLightCount = 0;
@@ -1443,7 +1437,7 @@ bool UseLightPropagationVolumeRT(ERHIFeatureLevel::Type InFeatureLevel)
 // ----------------------------------------------------------------------------
 // FSceneViewState
 // ----------------------------------------------------------------------------
-void FSceneViewState::CreateLightPropagationVolumeIfNeeded(ERHIFeatureLevel::Type InFeatureLevel)
+void FSceneViewState::SetupLightPropagationVolume(FSceneView& View, FSceneViewFamily& ViewFamily)
 {
 	check(IsInRenderingThread());
 
@@ -1453,16 +1447,35 @@ void FSceneViewState::CreateLightPropagationVolumeIfNeeded(ERHIFeatureLevel::Typ
 		return;
 	}
 
-	bool bLPV = UseLightPropagationVolumeRT(InFeatureLevel);
+	const ERHIFeatureLevel::Type FeatureLevel = View.GetFeatureLevel();
 
-	//@todo-rco: Remove this when reenabling for OpenGL
-	if (bLPV && !IsOpenGLPlatform(GShaderPlatformForFeatureLevel[InFeatureLevel]))
+	if (View.StereoPass == eSSP_RIGHT_EYE)
+	{
+		// The right eye will reference the left eye's LPV with the assumption that the left eye uses the primary view (index 0)
+		const FSceneView* PrimaryView = ViewFamily.Views[0];
+		if (PrimaryView->StereoPass == eSSP_LEFT_EYE && PrimaryView->State)
+		{
+			FSceneViewState* PrimaryViewState = PrimaryView->State->GetConcreteViewState();
+			if (PrimaryViewState)
+			{
+				LightPropagationVolume = PrimaryViewState->GetLightPropagationVolume(FeatureLevel);
+				if (LightPropagationVolume.IsValid())
+				{
+					bIsStereoView = true;
+				}
+			}
+		}
+	}
+	else
 	{
-		LightPropagationVolume = new FLightPropagationVolume();
+		if (UseLightPropagationVolumeRT(FeatureLevel) && !IsOpenGLPlatform(GShaderPlatformForFeatureLevel[FeatureLevel]))
+		{
+			LightPropagationVolume = new FLightPropagationVolume();
+		}
 	}
 }
 
-FLightPropagationVolume* FSceneViewState::GetLightPropagationVolume(ERHIFeatureLevel::Type InFeatureLevel) const
+FLightPropagationVolume* FSceneViewState::GetLightPropagationVolume(ERHIFeatureLevel::Type InFeatureLevel, bool bIncludeStereo) const
 {
 	if (InFeatureLevel < ERHIFeatureLevel::SM5)
 	{
@@ -1470,6 +1483,12 @@ FLightPropagationVolume* FSceneViewState::GetLightPropagationVolume(ERHIFeatureL
 		return 0;
 	}
 
+	if (bIsStereoView && !bIncludeStereo)
+	{
+		// return 0 on stereo views when they aren't explicitly included
+		return 0;
+	}
+
 	return LightPropagationVolume;
 }
 
@@ -1481,9 +1500,9 @@ void FSceneViewState::DestroyLightPropagationVolume()
 			DeleteLPV,
 			FLightPropagationVolume*, LightPropagationVolume, LightPropagationVolume,
 		{
-			delete LightPropagationVolume;
+			LightPropagationVolume = nullptr;
 		}
 		);
+		bIsStereoView = false;
 	}
-	LightPropagationVolume = NULL;
 }
diff --git a/Engine/Source/Runtime/Renderer/Private/LightPropagationVolume.h b/Engine/Source/Runtime/Renderer/Private/LightPropagationVolume.h
--- a/Engine/Source/Runtime/Renderer/Private/LightPropagationVolume.h
+++ b/Engine/Source/Runtime/Renderer/Private/LightPropagationVolume.h
@@ -131,7 +131,7 @@ struct FLpvBaseWriteShaderParams
 	FTextureRHIParamRef				AOVolumeTextureSRV;
 };
 
-class FLightPropagationVolume // @TODO: this should probably be derived from FRenderResource (with InitDynamicRHI etc)
+class FLightPropagationVolume : public FRefCountedObject
 {
 public:
 	FLightPropagationVolume();
diff --git a/Engine/Source/Runtime/Renderer/Private/RendererScene.cpp b/Engine/Source/Runtime/Renderer/Private/RendererScene.cpp
--- a/Engine/Source/Runtime/Renderer/Private/RendererScene.cpp
+++ b/Engine/Source/Runtime/Renderer/Private/RendererScene.cpp
@@ -90,6 +90,9 @@ FSceneViewState::FSceneViewState()
 	bBokehDOFHistory2 = true;
 
 	LightPropagationVolume = NULL; 
+
+	bIsStereoView = false;
+
 	HeightfieldLightingAtlas = NULL;
 
 	for (int32 CascadeIndex = 0; CascadeIndex < ARRAY_COUNT(TranslucencyLightingCacheAllocations); CascadeIndex++)
diff --git a/Engine/Source/Runtime/Renderer/Private/ScenePrivate.h b/Engine/Source/Runtime/Renderer/Private/ScenePrivate.h
--- a/Engine/Source/Runtime/Renderer/Private/ScenePrivate.h
+++ b/Engine/Source/Runtime/Renderer/Private/ScenePrivate.h
@@ -32,7 +32,6 @@ typedef TConstDualSetBitIterator<SceneRenderingBitArrayAllocator,SceneRenderingB
 
 // Forward declarations.
 class FScene;
-class FLightPropagationVolume;
 
 /** True if HDR is enabled for the mobile renderer. */
 bool IsMobileHDR();
@@ -108,6 +107,7 @@ public:
 #include "ClearQuad.h"
 #include "AtmosphereRendering.h"
 #include "GlobalDistanceFieldParameters.h"
+#include "LightPropagationVolume.h"
 
 /** Factor by which to grow occlusion tests **/
 #define OCCLUSION_SLOP (1.0f)
@@ -478,8 +478,11 @@ private:
 
 	int32 DistanceFieldTemporalSampleIndex;
 
-	// can be 0
-	FLightPropagationVolume* LightPropagationVolume;
+	// light propagation volume used in this view
+	TRefCountPtr<FLightPropagationVolume> LightPropagationVolume;
+
+	// whether this view is a stereo counterpart to a primary view
+	bool bIsStereoView;
 
 public:
 
@@ -585,10 +588,13 @@ public:
 	}
 
 	// call only if not yet created
-	void CreateLightPropagationVolumeIfNeeded(ERHIFeatureLevel::Type InFeatureLevel);
+	void SetupLightPropagationVolume(FSceneView& View, FSceneViewFamily& ViewFamily);
 
-	// @return can return 0 (if globally disabled)
-	FLightPropagationVolume* GetLightPropagationVolume(ERHIFeatureLevel::Type InFeatureLevel) const;
+	/**
+	 * @return can return 0
+	 * @param bIncludeStereo - specifies whether the getter should include stereo views in its returned value
+	 */
+	FLightPropagationVolume* GetLightPropagationVolume(ERHIFeatureLevel::Type InFeatureLevel, bool bIncludeStereo = false) const;
 
 	/** Default constructor. */
 	FSceneViewState();
@@ -706,6 +712,11 @@ public:
 	// FSceneViewStateInterface
 	RENDERER_API virtual void Destroy() override;
 	
+	virtual FSceneViewState* GetConcreteViewState() override
+	{
+		return this;
+	}
+
 	virtual void AddReferencedObjects(FReferenceCollector& Collector) override
 	{
 		uint32 Count = MIDPool.Num();
@@ -719,13 +730,13 @@ public:
 	}
 
 	/** called in InitViews() */
-	virtual void OnStartFrame(FSceneView& CurrentView) override
+	virtual void OnStartFrame(FSceneView& View, FSceneViewFamily& ViewFamily) override
 	{
 		check(IsInRenderingThread());
 
-		if(!(CurrentView.FinalPostProcessSettings.IndirectLightingColor * CurrentView.FinalPostProcessSettings.IndirectLightingIntensity).IsAlmostBlack())
+		if(!(View.FinalPostProcessSettings.IndirectLightingColor * View.FinalPostProcessSettings.IndirectLightingIntensity).IsAlmostBlack())
 		{
-			CreateLightPropagationVolumeIfNeeded(CurrentView.GetFeatureLevel());
+			SetupLightPropagationVolume(View, ViewFamily);
 		}
 	}
 
diff --git a/Engine/Source/Runtime/Renderer/Private/SceneRendering.cpp b/Engine/Source/Runtime/Renderer/Private/SceneRendering.cpp
--- a/Engine/Source/Runtime/Renderer/Private/SceneRendering.cpp
+++ b/Engine/Source/Runtime/Renderer/Private/SceneRendering.cpp
@@ -1437,7 +1437,7 @@ void FSceneRenderer::OnStartFrame()
 
 		if(State)
 		{
-			State->OnStartFrame(View);
+			State->OnStartFrame(View, ViewFamily);
 		}
 	}
 }
