Replaces r.AmbientOcclusionSampleSetQuality by r.AmbientOcclusionMaxQuality exposing a lower quality level.

diff --git a/Engine/Config/BaseScalability.ini b/Engine/Config/BaseScalability.ini
--- a/Engine/Config/BaseScalability.ini
+++ b/Engine/Config/BaseScalability.ini
@@ -102,7 +102,7 @@ r.DistanceFieldAO=1
 [PostProcessQuality@0]
 r.MotionBlurQuality=0
 r.AmbientOcclusionMipLevelFactor=1.0
-r.AmbientOcclusionSampleSetQuality=0
+r.AmbientOcclusionMaxQuality=0
 r.AmbientOcclusionLevels=0
 r.AmbientOcclusionRadiusScale=1.2
 r.DepthOfFieldQuality=0
@@ -120,7 +120,7 @@ r.Filter.SizeScale=0.6
 [PostProcessQuality@1]
 r.MotionBlurQuality=3
 r.AmbientOcclusionMipLevelFactor=1.0
-r.AmbientOcclusionSampleSetQuality=0
+r.AmbientOcclusionMaxQuality=60
 r.AmbientOcclusionLevels=1
 r.AmbientOcclusionRadiusScale=1.5
 r.DepthOfFieldQuality=1
@@ -138,7 +138,7 @@ r.Filter.SizeScale=0.7
 [PostProcessQuality@2]
 r.MotionBlurQuality=3
 r.AmbientOcclusionMipLevelFactor=0.6
-r.AmbientOcclusionSampleSetQuality=-1
+r.AmbientOcclusionMaxQuality=100
 r.AmbientOcclusionLevels=1
 r.AmbientOcclusionRadiusScale=1.5
 r.DepthOfFieldQuality=2
@@ -156,7 +156,7 @@ r.Filter.SizeScale=0.8
 [PostProcessQuality@3]
 r.MotionBlurQuality=4
 r.AmbientOcclusionMipLevelFactor=0.4
-r.AmbientOcclusionSampleSetQuality=-1
+r.AmbientOcclusionMaxQuality=100
 r.AmbientOcclusionLevels=2
 r.AmbientOcclusionRadiusScale=1.0
 r.DepthOfFieldQuality=2
diff --git a/Engine/Shaders/PostProcessAmbientOcclusion.usf b/Engine/Shaders/PostProcessAmbientOcclusion.usf
--- a/Engine/Shaders/PostProcessAmbientOcclusion.usf
+++ b/Engine/Shaders/PostProcessAmbientOcclusion.usf
@@ -10,7 +10,7 @@
 
 // set by C++:
 //
-// 0:low / 1:high
+// 0:low / 1: medium / 2:high
 // SAMPLESET_QUALITY
 //
 // 0:no / 1:yes
@@ -26,7 +26,7 @@
 #define OPTIMIZATION_O1 1
 
 // 1:lowest quality, 2:medium , 3:high, more doesn't give too much (maybe HZB mip computations should be adjusted)
-#define SAMPLE_STEPS 3
+//#define SAMPLE_STEPS 3
 
 // 0:does not use normals (slightly faster, lower quality, WIP = not finished), 1: uses normals (slower but higher quality)
 #define SSAO_NORMAL_AFFECTED 1
@@ -40,9 +40,15 @@
 #if SAMPLESET_QUALITY == 0
 	// low
 	#define USE_SAMPLESET 1
-#else
+	#define SAMPLE_STEPS 1
+#elif SAMPLESET_QUALITY == 1
+	// medium
+	#define USE_SAMPLESET 1
+	#define SAMPLE_STEPS 3
+#else // SAMPLESET_QUALITY == 2
 	// high
 	#define USE_SAMPLESET 3
+	#define SAMPLE_STEPS 3
 #endif
 
 // 0:4 samples, 1:9 samples (only really noticable with dither usage ??)
diff --git a/Engine/Source/Runtime/Renderer/Private/CompositionLighting/PostProcessAmbientOcclusion.cpp b/Engine/Source/Runtime/Renderer/Private/CompositionLighting/PostProcessAmbientOcclusion.cpp
--- a/Engine/Source/Runtime/Renderer/Private/CompositionLighting/PostProcessAmbientOcclusion.cpp
+++ b/Engine/Source/Runtime/Renderer/Private/CompositionLighting/PostProcessAmbientOcclusion.cpp
@@ -11,22 +11,21 @@
 #include "PostProcessAmbientOcclusion.h"
 #include "SceneUtils.h"
 
-static TAutoConsoleVariable<int32> CVarAmbientOcclusionSampleSetQuality(
-	TEXT("r.AmbientOcclusionSampleSetQuality"),
-	-1,
-	TEXT("Defines how many samples we use for ScreenSpace Ambient Occlusion\n")
-	TEXT("-1: sample count depends on post process settings (default)\n")
-	TEXT("0: low sample count (defined in shader, 3 * 2 per pixel)\n")
-	TEXT("1: high sample count (defined in shader, 6 * 2 per pixel)"),
+static TAutoConsoleVariable<int32> CVarAmbientOcclusionMaxQuality(
+	TEXT("r.AmbientOcclusionMaxQuality"),
+	100,
+	TEXT("Defines the max clamping value from the post process volume's quality level for ScreenSpace Ambient Occlusion\n")
+	TEXT(" 100: don't override quality level from the post process volume (default)\n")
+	TEXT(" <100: clamp down quality level from the post process volume to the maximum set by this cvar"),
 	ECVF_Scalability | ECVF_RenderThreadSafe);
 
 static TAutoConsoleVariable<float> CVarAmbientOcclusionStepMipLevelFactor(
 	TEXT("r.AmbientOcclusionMipLevelFactor"),
 	0.5f,
 	TEXT("Controle mipmap level according to the SSAO step id\n")
-	TEXT("0: always look into the HZB mipmap level 0 (memory cache trashing)\n")
-	TEXT("0.5: sample count depends on post process settings (default)\n")
-	TEXT("1: Go into higher mipmap level (quality loss)"),
+	TEXT(" 0: always look into the HZB mipmap level 0 (memory cache trashing)\n")
+	TEXT(" 0.5: sample count depends on post process settings (default)\n")
+	TEXT(" 1: Go into higher mipmap level (quality loss)"),
 	ECVF_Scalability | ECVF_RenderThreadSafe);
 
 IMPLEMENT_UNIFORM_BUFFER_STRUCT(FCameraMotionParameters,TEXT("CameraMotion"));
@@ -217,13 +216,16 @@ public:
 
 
 // #define avoids a lot of code duplication
+#define VARIATION0(C)	    VARIATION1(0, C) VARIATION1(1, C)
 #define VARIATION1(A, C)	VARIATION2(A, 0, C) VARIATION2(A, 1, C)
 #define VARIATION2(A, B, C) typedef FPostProcessAmbientOcclusionPS<A, B, C> FPostProcessAmbientOcclusionPS##A##B##C; \
 	IMPLEMENT_SHADER_TYPE2(FPostProcessAmbientOcclusionPS##A##B##C, SF_Pixel);
 
-	VARIATION1(0,0) VARIATION1(1,0)
-	VARIATION1(0,1) VARIATION1(1,1)
-
+	VARIATION0(0)
+	VARIATION0(1)
+	VARIATION0(2)
+	
+#undef VARIATION0
 #undef VARIATION1
 #undef VARIATION2
 
@@ -457,43 +459,36 @@ void FRCPassPostProcessAmbientOcclusion::Process(FRenderingCompositePassContext&
 	Context.RHICmdList.SetRasterizerState(TStaticRasterizerState<>::GetRHI());
 	Context.RHICmdList.SetDepthStencilState(TStaticDepthStencilState<false, CF_Always>::GetRHI());
 
-	int32 SampleSetQuality = CVarAmbientOcclusionSampleSetQuality.GetValueOnRenderThread();
-
-	if(SampleSetQuality < 0)
-	{
-		SampleSetQuality = 0;
-
-		// SampleSetQuality depends on post process settings
-		if(Context.View.FinalPostProcessSettings.AmbientOcclusionQuality > 60.0f)
-		{
-			SampleSetQuality = 1;
-		}
-	}
-
+	const float QualityPercent = FMath::Min(float(CVarAmbientOcclusionMaxQuality.GetValueOnRenderThread()), Context.View.FinalPostProcessSettings.AmbientOcclusionQuality);
+	const int32 QualitySet = 
+		(QualityPercent > 60.0f) +
+		(QualityPercent > 25.0f);
 	bool bDoUpsample = (InputDesc2 != 0);
 
-	if(bAOSetupAsInput)
-	{
-		if(bDoUpsample)
-		{
-			if(SampleSetQuality) SetShaderTempl<1, 1, 1>(Context); else SetShaderTempl<1, 1, 0>(Context);
-		}
-		else
-		{
-			if(SampleSetQuality) SetShaderTempl<1, 0, 1>(Context); else SetShaderTempl<1, 0, 0>(Context);
-		}
-	}
-	else
+#define SET_SHADER_CASE(Quality)                                \
+	case Quality:			                                    \
+	if(bAOSetupAsInput)                                         \
+    {                                                           \
+		if(bDoUpsample) SetShaderTempl<1, 1, Quality>(Context); \
+		else SetShaderTempl<1, 0, Quality>(Context);            \
+	}                                                           \
+	else                                                        \
+    {                                                           \
+		if(bDoUpsample) SetShaderTempl<0, 1, Quality>(Context); \
+		else SetShaderTempl<0, 0, Quality>(Context);            \
+	}                                                           \
+	break
+
+	switch(QualitySet)
 	{
-		if(bDoUpsample)
-		{
-			if(SampleSetQuality) SetShaderTempl<0, 1, 1>(Context); else SetShaderTempl<0, 1, 0>(Context);
-		}
-		else
-		{
-			if(SampleSetQuality) SetShaderTempl<0, 0, 1>(Context); else SetShaderTempl<0, 0, 0>(Context);
-		}
-	}
+		SET_SHADER_CASE(0);
+		SET_SHADER_CASE(1);
+		SET_SHADER_CASE(2);
+		default:
+			break;
+	};
+
+#undef SET_SHADER_CASE
 
 	TShaderMapRef<FPostProcessVS> VertexShader(Context.GetShaderMap());
 
