Fixing CIS build warning: Missed FRendererModule::RegisterOverlayRenderDelegate() input, which shadowed FRendererModule member name.

diff --git a/Engine/Source/Runtime/Renderer/Private/SceneRendering.cpp b/Engine/Source/Runtime/Renderer/Private/SceneRendering.cpp
--- a/Engine/Source/Runtime/Renderer/Private/SceneRendering.cpp
+++ b/Engine/Source/Runtime/Renderer/Private/SceneRendering.cpp
@@ -1675,9 +1675,9 @@ void FRendererModule::RegisterPostOpaqueRenderDelegate(const FPostOpaqueRenderDe
 	this->PostOpaqueRenderDelegate = InPostOpaqueRenderDelegate;
 }
 
-void FRendererModule::RegisterOverlayRenderDelegate(const FPostOpaqueRenderDelegate& OverlayRenderDelegate)
+void FRendererModule::RegisterOverlayRenderDelegate(const FPostOpaqueRenderDelegate& InOverlayRenderDelegate)
 {
-	this->OverlayRenderDelegate = OverlayRenderDelegate;
+	this->OverlayRenderDelegate = InOverlayRenderDelegate;
 }
 
 void FRendererModule::RenderPostOpaqueExtensions(const FSceneView& View, FRHICommandListImmediate& RHICmdList, FSceneRenderTargets& SceneContext)
