Crash fix for LPV release across threads.

diff --git a/Engine/Source/Runtime/Renderer/Private/LightPropagationVolume.cpp b/Engine/Source/Runtime/Renderer/Private/LightPropagationVolume.cpp
--- a/Engine/Source/Runtime/Renderer/Private/LightPropagationVolume.cpp
+++ b/Engine/Source/Runtime/Renderer/Private/LightPropagationVolume.cpp
@@ -1496,11 +1496,15 @@ void FSceneViewState::DestroyLightPropagationVolume()
 {
 	if ( LightPropagationVolume ) 
 	{
+		LightPropagationVolume->AddRef();
+		FLightPropagationVolume* LPV = LightPropagationVolume;
+		LightPropagationVolume = nullptr;
+
 		ENQUEUE_UNIQUE_RENDER_COMMAND_ONEPARAMETER(
 			DeleteLPV,
-			FLightPropagationVolume*, LightPropagationVolume, LightPropagationVolume,
+			FLightPropagationVolume*, LPV, LPV,
 		{
-			LightPropagationVolume = nullptr;
+			LPV->Release();
 		}
 		);
 		bIsStereoView = false;
