refactor, allows optimization later on

diff --git a/Engine/Source/Runtime/Renderer/Private/DecalRenderingShared.cpp b/Engine/Source/Runtime/Renderer/Private/DecalRenderingShared.cpp
--- a/Engine/Source/Runtime/Renderer/Private/DecalRenderingShared.cpp
+++ b/Engine/Source/Runtime/Renderer/Private/DecalRenderingShared.cpp
@@ -161,7 +161,7 @@ public:
 
 		FTransform ComponentTrans = DecalProxy.ComponentTrans;
 
-		FMatrix WorldToComponent = ComponentTrans.ToMatrixWithScale().InverseFast();
+		FMatrix WorldToComponent = ComponentTrans.ToInverseMatrixWithScale();
 
 		// Set the transform from screen space to light space.
 		if(ScreenToDecal.IsBound())
