nicer cvar help

diff --git a/Engine/Source/Runtime/Renderer/Private/PostProcess/PostProcessSubsurface.cpp b/Engine/Source/Runtime/Renderer/Private/PostProcess/PostProcessSubsurface.cpp
--- a/Engine/Source/Runtime/Renderer/Private/PostProcess/PostProcessSubsurface.cpp
+++ b/Engine/Source/Runtime/Renderer/Private/PostProcess/PostProcessSubsurface.cpp
@@ -17,16 +17,18 @@ ENGINE_API const IPooledRenderTarget* GetSubsufaceProfileTexture_RT(FRHICommandL
 static TAutoConsoleVariable<int32> CVarSSSFilter(
 	TEXT("r.SSS.Filter"),
 	1,
-	TEXT("0: point filter (useful for testing, could be cleaner)\n")
-	TEXT("1: bilinear filter"),
+	TEXT("Defines the filter method for Screenspace Subsurface Scattering feature.\n")
+	TEXT(" 0: point filter (useful for testing, could be cleaner)\n")
+	TEXT(" 1: bilinear filter"),
 	ECVF_RenderThreadSafe  | ECVF_Scalability);
 
 static TAutoConsoleVariable<int32> CVarSSSSampleSet(
 	TEXT("r.SSS.SampleSet"),
 	2,
-	TEXT("0: lowest quality\n")
-	TEXT("1: medium quality\n")
-	TEXT("2: high quality (default)"),
+	TEXT("Defines how many samples we use for Screenspace Subsurface Scattering feature.\n")
+	TEXT(" 0: lowest quality (6*2+1)\n")
+	TEXT(" 1: medium quality (9*2+1)\n")
+	TEXT(" 2: high quality (13*2+1) (default)"),
 	ECVF_RenderThreadSafe  | ECVF_Scalability);
 
 static TAutoConsoleVariable<int> CVarSubsurfaceQuality(
