Fix first-frame corrupt 'white' rendering

diff --git a/Engine/Source/Runtime/Renderer/Private/PostProcess/PostProcessBloomSetup.cpp b/Engine/Source/Runtime/Renderer/Private/PostProcess/PostProcessBloomSetup.cpp
--- a/Engine/Source/Runtime/Renderer/Private/PostProcess/PostProcessBloomSetup.cpp
+++ b/Engine/Source/Runtime/Renderer/Private/PostProcess/PostProcessBloomSetup.cpp
@@ -113,10 +113,9 @@ public:
 
 		if(EyeAdaptation.IsBound())
 		{
-			IPooledRenderTarget* EyeAdaptationRT = Context.View.GetEyeAdaptation();
-
-			if(EyeAdaptationRT)
+			if (Context.View.HasValidEyeAdaptation())
 			{
+				IPooledRenderTarget* EyeAdaptationRT = Context.View.GetEyeAdaptation();
 				SetTextureParameter(Context.RHICmdList, ShaderRHI, EyeAdaptation, EyeAdaptationRT->GetRenderTargetItem().TargetableTexture);
 			}
 			else
diff --git a/Engine/Source/Runtime/Renderer/Private/PostProcess/PostProcessEyeAdaptation.cpp b/Engine/Source/Runtime/Renderer/Private/PostProcess/PostProcessEyeAdaptation.cpp
--- a/Engine/Source/Runtime/Renderer/Private/PostProcess/PostProcessEyeAdaptation.cpp
+++ b/Engine/Source/Runtime/Renderer/Private/PostProcess/PostProcessEyeAdaptation.cpp
@@ -114,6 +114,8 @@ void FRCPassPostProcessEyeAdaptation::Process(FRenderingCompositePassContext& Co
 		EDRF_UseTriangleOptimization);
 
 	Context.RHICmdList.CopyToResolveTarget(EyeAdaptation->GetRenderTargetItem().TargetableTexture, EyeAdaptation->GetRenderTargetItem().ShaderResourceTexture, false, FResolveParams());
+
+	Context.View.SetValidEyeAdaptation();
 }
 
 void FRCPassPostProcessEyeAdaptation::ComputeEyeAdaptationParamsValue(const FViewInfo& View, FVector4 Out[3])
diff --git a/Engine/Source/Runtime/Renderer/Private/PostProcess/PostProcessHistogramReduce.cpp b/Engine/Source/Runtime/Renderer/Private/PostProcess/PostProcessHistogramReduce.cpp
--- a/Engine/Source/Runtime/Renderer/Private/PostProcess/PostProcessHistogramReduce.cpp
+++ b/Engine/Source/Runtime/Renderer/Private/PostProcess/PostProcessHistogramReduce.cpp
@@ -58,10 +58,9 @@ public:
 
 		if(EyeAdaptationTexture.IsBound())
 		{
-			IPooledRenderTarget* EyeAdaptationRT = Context.View.GetEyeAdaptation();
-
-			if(EyeAdaptationRT)
+			if (Context.View.HasValidEyeAdaptation())
 			{
+				IPooledRenderTarget* EyeAdaptationRT = Context.View.GetEyeAdaptation();
 				SetTextureParameter(Context.RHICmdList, ShaderRHI, EyeAdaptationTexture, EyeAdaptationRT->GetRenderTargetItem().TargetableTexture);
 			}
 			else
diff --git a/Engine/Source/Runtime/Renderer/Private/PostProcess/PostProcessTonemap.h b/Engine/Source/Runtime/Renderer/Private/PostProcess/PostProcessTonemap.h
--- a/Engine/Source/Runtime/Renderer/Private/PostProcess/PostProcessTonemap.h
+++ b/Engine/Source/Runtime/Renderer/Private/PostProcess/PostProcessTonemap.h
@@ -127,10 +127,9 @@ public:
 
 		if(EyeAdaptation.IsBound())
 		{
-			IPooledRenderTarget* EyeAdaptationRT = Context.View.GetEyeAdaptation();
-
-			if(EyeAdaptationRT)
+			if (Context.View.HasValidEyeAdaptation())
 			{
+				IPooledRenderTarget* EyeAdaptationRT = Context.View.GetEyeAdaptation();
 				SetTextureParameter(Context.RHICmdList, ShaderRHI, EyeAdaptation, EyeAdaptationRT->GetRenderTargetItem().TargetableTexture);
 			}
 			else
diff --git a/Engine/Source/Runtime/Renderer/Private/RendererScene.cpp b/Engine/Source/Runtime/Renderer/Private/RendererScene.cpp
--- a/Engine/Source/Runtime/Renderer/Private/RendererScene.cpp
+++ b/Engine/Source/Runtime/Renderer/Private/RendererScene.cpp
@@ -105,6 +105,8 @@ FSceneViewState::FSceneViewState()
 
 	ShadowOcclusionQueryMaps.Empty(FOcclusionQueryHelpers::MaxBufferedOcclusionFrames);
 	ShadowOcclusionQueryMaps.AddZeroed(FOcclusionQueryHelpers::MaxBufferedOcclusionFrames);	
+
+	bValidEyeAdaptation = false;
 }
 
 void DestroyRenderResource(FRenderResource* RenderResource)
diff --git a/Engine/Source/Runtime/Renderer/Private/ScenePrivate.h b/Engine/Source/Runtime/Renderer/Private/ScenePrivate.h
--- a/Engine/Source/Runtime/Renderer/Private/ScenePrivate.h
+++ b/Engine/Source/Runtime/Renderer/Private/ScenePrivate.h
@@ -467,6 +467,9 @@ private:
 	// to implement eye adaptation changes over time
 	TRefCountPtr<IPooledRenderTarget> EyeAdaptationRT;
 
+	// eye adaptation is only valid after it has been computed, not on allocation of the RT
+	bool bValidEyeAdaptation;
+
 	// used by the Postprocess Material Blending system to avoid recreation and garbage collection of MIDs
 	TArray<UMaterialInstanceDynamic*> MIDPool;
 	uint32 MIDUsedCount;
@@ -658,6 +661,16 @@ public:
 		return EyeAdaptationRT;
 	}
 
+	bool HasValidEyeAdaptation() const
+	{
+		return bValidEyeAdaptation;
+	}
+
+	void SetValidEyeAdaptation()
+	{
+		bValidEyeAdaptation = true;
+	}
+
 	TRefCountPtr<IPooledRenderTarget>& GetSeparateTranslucency(FIntPoint Size)
 	{
 		if (!SeparateTranslucencyRT || SeparateTranslucencyRT->GetDesc().Extent != Size)
diff --git a/Engine/Source/Runtime/Renderer/Private/SceneRendering.cpp b/Engine/Source/Runtime/Renderer/Private/SceneRendering.cpp
--- a/Engine/Source/Runtime/Renderer/Private/SceneRendering.cpp
+++ b/Engine/Source/Runtime/Renderer/Private/SceneRendering.cpp
@@ -1040,7 +1040,7 @@ void FViewInfo::DestroyAllSnapshots()
 	ViewInfoSnapshots.Reset();
 }
 
-IPooledRenderTarget* FViewInfo::GetEyeAdaptation() const
+FSceneViewState* FViewInfo::GetEffectiveViewState() const
 {
 	FSceneViewState* EffectiveViewState = ViewState;
 
@@ -1063,6 +1063,13 @@ IPooledRenderTarget* FViewInfo::GetEyeAdaptation() const
 		}
 	}
 
+	return EffectiveViewState;
+}
+
+IPooledRenderTarget* FViewInfo::GetEyeAdaptation() const
+{
+	FSceneViewState* EffectiveViewState = GetEffectiveViewState();
+
 	if (EffectiveViewState)
 	{
 		TRefCountPtr<IPooledRenderTarget>& EyeAdaptRef = EffectiveViewState->GetEyeAdaptation();
@@ -1073,6 +1080,28 @@ IPooledRenderTarget* FViewInfo::GetEyeAdaptation() const
 	}
 	return NULL;
 }
+
+bool FViewInfo::HasValidEyeAdaptation() const
+{
+	FSceneViewState* EffectiveViewState = GetEffectiveViewState();	
+
+	if (EffectiveViewState)
+	{
+		return EffectiveViewState->HasValidEyeAdaptation();
+	}
+	return false;
+}
+
+void FViewInfo::SetValidEyeAdaptation()
+{
+	FSceneViewState* EffectiveViewState = GetEffectiveViewState();	
+
+	if (EffectiveViewState)
+	{
+		EffectiveViewState->SetValidEyeAdaptation();
+	}
+}
+
 /*-----------------------------------------------------------------------------
 	FSceneRenderer
 -----------------------------------------------------------------------------*/
diff --git a/Engine/Source/Runtime/Renderer/Private/SceneRendering.h b/Engine/Source/Runtime/Renderer/Private/SceneRendering.h
--- a/Engine/Source/Runtime/Renderer/Private/SceneRendering.h
+++ b/Engine/Source/Runtime/Renderer/Private/SceneRendering.h
@@ -594,6 +594,12 @@ public:
 	/** Gets the eye adaptation render target for this view. */
 	IPooledRenderTarget* GetEyeAdaptation() const;
 
+	/** Tells if the eyeadaptation texture exists without attempting to allocate it. */
+	bool HasValidEyeAdaptation() const;
+
+	/** Informs sceneinfo that eyedaptation has queued commands to compute it at least once */
+	void SetValidEyeAdaptation();
+
 	/** Create acceleration data structure and information to do forward lighting with dynamic branching. */
 	void CreateLightGrid();
 
@@ -622,6 +628,8 @@ public:
 
 private:
 
+	FSceneViewState* GetEffectiveViewState() const;
+
 	/** Initialization that is common to the constructors. */
 	void Init();
 
