Fixed crash with viewport of size 1

diff --git a/Engine/Source/Runtime/Renderer/Private/SceneOcclusion.cpp b/Engine/Source/Runtime/Renderer/Private/SceneOcclusion.cpp
--- a/Engine/Source/Runtime/Renderer/Private/SceneOcclusion.cpp
+++ b/Engine/Source/Runtime/Renderer/Private/SceneOcclusion.cpp
@@ -987,8 +987,8 @@ void BuildHZB( FRHICommandListImmediate& RHICmdList, FViewInfo& View )
 	
 	// View.ViewRect.{Width,Height}() are most likely to be < 2^24, so the float
 	// conversion won't loss any precision (assuming float have 23bits for mantissa)
-	const uint32 NumMipsX = FPlatformMath::CeilToInt(FMath::Log2(float(View.ViewRect.Width()))) - 1;
-	const uint32 NumMipsY = FPlatformMath::CeilToInt(FMath::Log2(float(View.ViewRect.Height()))) - 1;
+	const int32 NumMipsX = FMath::Max(FPlatformMath::CeilToInt(FMath::Log2(float(View.ViewRect.Width()))) - 1, 1);
+	const int32 NumMipsY = FMath::Max(FPlatformMath::CeilToInt(FMath::Log2(float(View.ViewRect.Height()))) - 1, 1);
 	const uint32 NumMips = FMath::Max(NumMipsX, NumMipsY);
 
 	// Must be power of 2
