Remove the alternate index buffer layout for OS X that works around radr://15076670 "Incorrect gl_VertexID in GLSL for glDrawElementsInstanced without vertex streams on NvidiaΓÇ¥ committed in CL #1897281. Instead when using OpenGL with an Nvidia GPU on OS X, bind the index buffer as vertex attribute 0  even though it is unused as this sidesteps the problem.

diff --git a/Engine/Shaders/PostProcessBokehDOF.usf b/Engine/Shaders/PostProcessBokehDOF.usf
--- a/Engine/Shaders/PostProcessBokehDOF.usf
+++ b/Engine/Shaders/PostProcessBokehDOF.usf
@@ -55,15 +55,6 @@ void MainVS(
 
 	// needs to be the same on C++ side (faster on NVIDIA and AMD)
 	uint QuadsPerInstance = 8;
-#if DOF_INDEX_STYLE == 1
-	// remap the indices to get vertexid to VId and quadid into IId
-	IId = IId * QuadsPerInstance + (VId / 6);
-	VId = VId % 6;
-
-	// triangle A: 0:left top, 1:right top, 2: left bottom
-	// triangle B: 3:right bottom, 4:left bottom, 5: right top
-	float2 LocalPos = float2(VId % 2, VId > 1 && VId < 5);
-#else // DOF_INDEX_STYLE == 1
 	// remap the indices to get vertexid to VId and quadid into IId
 	IId = IId * QuadsPerInstance + (VId / 4);
 	VId = VId % 4;
@@ -71,7 +62,6 @@ void MainVS(
 	// triangle A: 0:left top, 1:right top, 2: left bottom
 	// triangle B: 3:right bottom, 4:left bottom, 5: right top
 	float2 LocalPos = float2(VId % 2, VId / 2);
-#endif // DOF_INDEX_STYLE == 1
 	float2 TilePos = float2(IId % TileCount.x, IId / TileCount.x);
 	
 	OutPosition =  float4(0, 0, 0, 1);
@@ -157,17 +147,6 @@ void MainVS(
 
 	// needs to be the same on C++ side (faster on NVIDIA and AMD)
 	uint QuadsPerInstance = 8;
-#if DOF_INDEX_STYLE == 1
-	// remap the indices to get vertexid to VId and quadid into IId
-	uint QuadId = VId / 6;
-
-	// triangle A: 0:left top, 1:right top, 2: left bottom
-	// triangle B: 3:right bottom, 4:left bottom, 5: right top
-	uint VertexId = VId % 6;
-	uint HalfTileCountX = TileCount.x / 2;
-	uint QuarterTileCountX = TileCount.x / 4;
-	float2 LocalPos = float2(VertexId % 2, VertexId > 1 && VertexId < 5);
-#else // DOF_INDEX_STYLE == 1
 	// remap the indices to get vertexid to VId and quadid into IId
 	uint QuadId = VId / 4;
 
@@ -177,7 +156,6 @@ void MainVS(
 	uint HalfTileCountX = TileCount.x / 2;
 	uint QuarterTileCountX = TileCount.x / 4;
 	float2 LocalPos = float2(VertexId % 2, VertexId / 2);
-#endif // DOF_INDEX_STYLE == 1
 	uint QuadIdIndex = QuadId%4;
 	float2 TilePos = float2(IId % QuarterTileCountX, IId / QuarterTileCountX) * uint2(4,2) + uint2(QuadId/4,0)*2;
 
diff --git a/Engine/Source/Runtime/OpenGLDrv/Private/OpenGLCommands.cpp b/Engine/Source/Runtime/OpenGLDrv/Private/OpenGLCommands.cpp
--- a/Engine/Source/Runtime/OpenGLDrv/Private/OpenGLCommands.cpp
+++ b/Engine/Source/Runtime/OpenGLDrv/Private/OpenGLCommands.cpp
@@ -2746,6 +2746,30 @@ void FOpenGLDynamicRHI::RHIDrawIndexedPrimitive(FIndexBufferRHIParamRef IndexBuf
 #if DEBUG_GL_SHADERS
 	VerifyProgramPipeline();
 #endif
+	
+	// @todo Workaround for radr://15076670 "Incorrect gl_VertexID in GLSL for glDrawElementsInstanced without vertex streams on Nvidia” Alternative fix that avoids exposing the messy details to the Renderer, keeping it here in the RHI.
+	// This workaround has performance and correctness implications - it is only needed on Mac + OpenGL + Nvidia and will
+	// break AMD drivers entirely as it is technically an abuse of the OpenGL specification. Consequently it is deliberately
+	// compiled out for other platforms. Apple have closed the bug claiming the NV behaviour is permitted by the GL spec.
+#if PLATFORM_MAC
+	bool bAttributeLessDraw = (PendingState.BoundShaderState->VertexShader->Bindings.InOutMask == 0 && ContextState.ElementArrayBufferBound && IsRHIDeviceNVIDIA());
+	if(bAttributeLessDraw)
+	{
+		CachedBindArrayBuffer(ContextState, IndexBuffer->Resource);
+		glEnableVertexAttribArray(0);
+		glVertexAttribPointer(0, 1, IndexType, false, 0, INDEX_TO_VOID(StartIndex));
+		ContextState.VertexAttrs[0].Pointer = INDEX_TO_VOID(StartIndex);
+		ContextState.VertexAttrs[0].Stride = 0;
+		ContextState.VertexAttrs[0].Buffer = IndexBuffer->Resource;
+		ContextState.VertexAttrs[0].Size = 1;
+		ContextState.VertexAttrs[0].Divisor = 0;
+		ContextState.VertexAttrs[0].Type = IndexType;
+		ContextState.VertexAttrs[0].StreamOffset = 0;
+		ContextState.VertexAttrs[0].StreamIndex = 0;
+		ContextState.VertexAttrs[0].bNormalized = false;
+		ContextState.VertexAttrs[0].bEnabled = true;
+	}
+#endif
 
 	GPUProfilingData.RegisterGPUWork(NumPrimitives * NumInstances, NumElements * NumInstances);
 	if (NumInstances > 1)
@@ -2771,6 +2795,16 @@ void FOpenGLDynamicRHI::RHIDrawIndexedPrimitive(FIndexBufferRHIParamRef IndexBuf
 		}
 		REPORT_GL_DRAW_RANGE_ELEMENTS_EVENT_FOR_FRAME_DUMP(DrawMode, MinIndex, MinIndex + NumVertices, NumElements, IndexType, (void *)StartIndex);
 	}
+	
+	// @todo Workaround for radr://15076670 "Incorrect gl_VertexID in GLSL for glDrawElementsInstanced without vertex streams on Nvidia”
+#if PLATFORM_MAC
+	if(bAttributeLessDraw)
+	{
+		glDisableVertexAttribArray(0);
+		ContextState.VertexAttrs[0].bEnabled = false;
+		ContextState.VertexStreams[0].VertexBuffer = nullptr;
+	}
+#endif
 
 	FShaderCache::LogDraw(IndexBuffer->GetStride());
 }
diff --git a/Engine/Source/Runtime/Renderer/Private/PostProcess/PostProcessBokehDOF.cpp b/Engine/Source/Runtime/Renderer/Private/PostProcess/PostProcessBokehDOF.cpp
--- a/Engine/Source/Runtime/Renderer/Private/PostProcess/PostProcessBokehDOF.cpp
+++ b/Engine/Source/Runtime/Renderer/Private/PostProcess/PostProcessBokehDOF.cpp
@@ -15,20 +15,10 @@
 
 
 
-/**
- * Indexing style for DOF
- */
-enum EBokehIndexStyle
-{
-	BIS_Fast = 0, /* Default fast, packed indexing mode */
-	BIS_Slow = 1 /* Slower, unwound indexing mode, used to avoid driver bugs on OSX/NV */
-};
-
 /**
  * Index buffer for drawing an individual sprite.
  */
-template<EBokehIndexStyle DOFIndexStyle>
-class TBokehIndexBuffer : public FIndexBuffer
+class FBokehIndexBuffer : public FIndexBuffer
 {
 public:
 	virtual void InitRHI() override
@@ -39,37 +29,21 @@ public:
 		void* Buffer = nullptr;
 		IndexBufferRHI = RHICreateAndLockIndexBuffer(Stride, Size, BUF_Static, CreateInfo, Buffer);
 		uint16* Indices = (uint16*)Buffer;
-		if(DOFIndexStyle == BIS_Fast)
+		for (uint32 SpriteIndex = 0; SpriteIndex < 8; ++SpriteIndex)
 		{
-			for (uint32 SpriteIndex = 0; SpriteIndex < 8; ++SpriteIndex)
-			{
-				Indices[SpriteIndex*6 + 0] = SpriteIndex*4 + 0;
-				Indices[SpriteIndex*6 + 1] = SpriteIndex*4 + 3;
-				Indices[SpriteIndex*6 + 2] = SpriteIndex*4 + 2;
-				Indices[SpriteIndex*6 + 3] = SpriteIndex*4 + 0;
-				Indices[SpriteIndex*6 + 4] = SpriteIndex*4 + 1;
-				Indices[SpriteIndex*6 + 5] = SpriteIndex*4 + 3;
-			}
-		}
-		else
-		{
-			for (uint32 SpriteIndex = 0; SpriteIndex < 8; ++SpriteIndex)
-			{
-				Indices[SpriteIndex*6 + 0] = SpriteIndex*6 + 0;
-				Indices[SpriteIndex*6 + 1] = SpriteIndex*6 + 1;
-				Indices[SpriteIndex*6 + 2] = SpriteIndex*6 + 2;
-				Indices[SpriteIndex*6 + 3] = SpriteIndex*6 + 3;
-				Indices[SpriteIndex*6 + 4] = SpriteIndex*6 + 4;
-				Indices[SpriteIndex*6 + 5] = SpriteIndex*6 + 5;
-			}
+			Indices[SpriteIndex*6 + 0] = SpriteIndex*4 + 0;
+			Indices[SpriteIndex*6 + 1] = SpriteIndex*4 + 3;
+			Indices[SpriteIndex*6 + 2] = SpriteIndex*4 + 2;
+			Indices[SpriteIndex*6 + 3] = SpriteIndex*4 + 0;
+			Indices[SpriteIndex*6 + 4] = SpriteIndex*4 + 1;
+			Indices[SpriteIndex*6 + 5] = SpriteIndex*4 + 3;
 		}
 		RHIUnlockIndexBuffer( IndexBufferRHI );
 	}
 };
 
 /** Global Bokeh index buffer. */
-TGlobalResource< TBokehIndexBuffer<BIS_Fast> > GBokehIndexBuffer;
-TGlobalResource< TBokehIndexBuffer<BIS_Slow> > GBokehSlowIndexBuffer;
+TGlobalResource< FBokehIndexBuffer > GBokehIndexBuffer;
 
 /** Encapsulates the post processing depth of field setup pixel shader. */
 class PostProcessVisualizeDOFPS : public FGlobalShader
@@ -549,7 +523,7 @@ FPooledRenderTargetDesc FRCPassPostProcessBokehDOFSetup::ComputeOutputDesc(EPass
 }
 
 /** Encapsulates the post processing vertex shader. */
-template <uint32 DOFMethod, uint32 DOFIndexStyle>
+template <uint32 DOFMethod>
 class FPostProcessBokehDOFVS : public FGlobalShader
 {
 	DECLARE_SHADER_TYPE(FPostProcessBokehDOFVS,Global);
@@ -563,7 +537,6 @@ class FPostProcessBokehDOFVS : public FGlobalShader
 	{
 		FGlobalShader::ModifyCompilationEnvironment(Platform,OutEnvironment);
 		OutEnvironment.SetDefine(TEXT("DOF_METHOD"), DOFMethod);
-		OutEnvironment.SetDefine(TEXT("DOF_INDEX_STYLE"), DOFIndexStyle);
 	}
 
 	/** Default constructor. */
@@ -712,17 +685,16 @@ public:
 IMPLEMENT_SHADER_TYPE(,FPostProcessBokehDOFPS,TEXT("PostProcessBokehDOF"),TEXT("MainPS"),SF_Pixel);
 
 // #define avoids a lot of code duplication
-#define VARIATION1(A, B) typedef FPostProcessBokehDOFVS<A, B> FPostProcessBokehDOFVS##A##B; \
-	IMPLEMENT_SHADER_TYPE2(FPostProcessBokehDOFVS##A##B, SF_Vertex);
+#define VARIATION1(A) typedef FPostProcessBokehDOFVS<A> FPostProcessBokehDOFVS##A; \
+	IMPLEMENT_SHADER_TYPE2(FPostProcessBokehDOFVS##A, SF_Vertex);
 
-VARIATION1(0,0)			VARIATION1(1,0)			VARIATION1(2,0)
-VARIATION1(0,1)			VARIATION1(1,1)			VARIATION1(2,1)
+VARIATION1(0)			VARIATION1(1)			VARIATION1(2)
 #undef VARIATION1
 
-template <uint32 DOFMethod, uint32 DOFIndexStyle>
+template <uint32 DOFMethod>
 void FRCPassPostProcessBokehDOF::SetShaderTempl(const FRenderingCompositePassContext& Context, FIntPoint LeftTop, FIntPoint TileCount, uint32 TileSize, float PixelKernelSize)
 {
-	TShaderMapRef<FPostProcessBokehDOFVS<DOFMethod, DOFIndexStyle> > VertexShader(Context.GetShaderMap());
+	TShaderMapRef<FPostProcessBokehDOFVS<DOFMethod> > VertexShader(Context.GetShaderMap());
 	TShaderMapRef<FPostProcessBokehDOFPS> PixelShader(Context.GetShaderMap());
 
 	static FGlobalBoundShaderState BoundShaderState;
@@ -761,19 +733,6 @@ void FRCPassPostProcessBokehDOF::ComputeDepthOfFieldParams(const FRenderingCompo
 }
 
 
-static TAutoConsoleVariable<int32> CVarBokehDOFIndexStyle(
-	TEXT("r.BokehDOFIndexStyle"),
-#if PLATFORM_MAC // Avoid a driver bug on OSX/NV cards that causes driver to generate an unwound index buffer
-	1,
-#else
-	0,
-#endif
-	TEXT("Controls whether to use a packed or unwound index buffer for Bokeh DOF.\n")
-	TEXT("0: Use packed index buffer (faster) (default)\n")
-	TEXT("1: Use unwound index buffer (slower)"),
-	ECVF_ReadOnly | ECVF_RenderThreadSafe);
-
-
 void FRCPassPostProcessBokehDOF::Process(FRenderingCompositePassContext& Context)
 {
 	SCOPED_DRAW_EVENT(Context.RHICmdList, PassPostProcessBokehDOF);
@@ -838,61 +797,31 @@ void FRCPassPostProcessBokehDOF::Process(FRenderingCompositePassContext& Context
 	float PixelKernelSize = Context.View.FinalPostProcessSettings.DepthOfFieldMaxBokehSize / 100.0f * LocalViewSize.X;
 
 	FIntPoint LeftTop = LocalViewRect.Min;
-
-	static EBokehIndexStyle IndexStyle = (EBokehIndexStyle)CVarBokehDOFIndexStyle.GetValueOnRenderThread();
 	
 	if(bHighQuality)
 	{
 		if(View.Family->EngineShowFlags.VisualizeAdaptiveDOF)
 		{
 			// high quality, visualize in red and green where we spend more performance
-			if(IndexStyle == BIS_Fast)
-			{
-				SetShaderTempl<2, 0>(Context, LeftTop, TileCount, TileSize, PixelKernelSize);
-			}
-			else
-			{
-				SetShaderTempl<2, 1>(Context, LeftTop, TileCount, TileSize, PixelKernelSize);
-			}
+			SetShaderTempl<2>(Context, LeftTop, TileCount, TileSize, PixelKernelSize);
 		}
 		else
 		{
 			// high quality
-			if(IndexStyle == BIS_Fast)
-			{
-				SetShaderTempl<1, 0>(Context, LeftTop, TileCount, TileSize, PixelKernelSize);
-			}
-			else
-			{
-				SetShaderTempl<1, 1>(Context, LeftTop, TileCount, TileSize, PixelKernelSize);
-			}
+			SetShaderTempl<1>(Context, LeftTop, TileCount, TileSize, PixelKernelSize);
 		}
 	}
 	else
 	{
 		// low quality
-		if(IndexStyle == BIS_Fast)
-		{
-			SetShaderTempl<0, 0>(Context, LeftTop, TileCount, TileSize, PixelKernelSize);
-		}
-		else
-		{
-			SetShaderTempl<0, 1>(Context, LeftTop, TileCount, TileSize, PixelKernelSize);
-		}
+		SetShaderTempl<0>(Context, LeftTop, TileCount, TileSize, PixelKernelSize);
 	}
 
 	// needs to be the same on shader side (faster on NVIDIA and AMD)
 	int32 QuadsPerInstance = 8;
 
 	Context.RHICmdList.SetStreamSource(0, NULL, 0, 0);
-	if(IndexStyle == BIS_Fast)
-	{
-		Context.RHICmdList.DrawIndexedPrimitive(GBokehIndexBuffer.IndexBufferRHI, PT_TriangleList, 0, 0, 32, 0, 2 * QuadsPerInstance, FMath::DivideAndRoundUp(TileCount.X * TileCount.Y, QuadsPerInstance));
-	}
-	else
-	{
-		Context.RHICmdList.DrawIndexedPrimitive(GBokehSlowIndexBuffer.IndexBufferRHI, PT_TriangleList, 0, 0, 32, 0, 2 * QuadsPerInstance, FMath::DivideAndRoundUp(TileCount.X * TileCount.Y, QuadsPerInstance));
-	}
+	Context.RHICmdList.DrawIndexedPrimitive(GBokehIndexBuffer.IndexBufferRHI, PT_TriangleList, 0, 0, 32, 0, 2 * QuadsPerInstance, FMath::DivideAndRoundUp(TileCount.X * TileCount.Y, QuadsPerInstance));
 
 	Context.RHICmdList.CopyToResolveTarget(DestRenderTarget.TargetableTexture, DestRenderTarget.ShaderResourceTexture, false, FResolveParams());
 }
diff --git a/Engine/Source/Runtime/Renderer/Private/PostProcess/PostProcessBokehDOF.h b/Engine/Source/Runtime/Renderer/Private/PostProcess/PostProcessBokehDOF.h
--- a/Engine/Source/Runtime/Renderer/Private/PostProcess/PostProcessBokehDOF.h
+++ b/Engine/Source/Runtime/Renderer/Private/PostProcess/PostProcessBokehDOF.h
@@ -68,7 +68,7 @@ public:
 
 private:
 
-	template <uint32 HighQuality, uint32 IndexStyle>
+	template <uint32 HighQuality>
 	void SetShaderTempl(const FRenderingCompositePassContext& Context, FIntPoint LeftTop, FIntPoint TileCount, uint32 TileSize, float PixelKernelSize);
 
 	// border between front and back layer as we don't use viewports (only possible with GS)
