Enable subsurface scattering on OpenGL as it works.

diff --git a/Engine/Source/Runtime/Renderer/Private/CompositionLighting/CompositionLighting.cpp b/Engine/Source/Runtime/Renderer/Private/CompositionLighting/CompositionLighting.cpp
--- a/Engine/Source/Runtime/Renderer/Private/CompositionLighting/CompositionLighting.cpp
+++ b/Engine/Source/Runtime/Renderer/Private/CompositionLighting/CompositionLighting.cpp
@@ -379,9 +379,7 @@ void FCompositionLighting::ProcessAfterLighting(FRHICommandListImmediate& RHICmd
 			bool bSimpleDynamicLighting = IsSimpleDynamicLightingEnabled();
 
 			bool bScreenSpaceSubsurfacePassNeeded = (View.ShadingModelMaskInView & (1 << MSM_SubsurfaceProfile)) != 0;
-			if (bScreenSpaceSubsurfacePassNeeded && Radius > 0 && !bSimpleDynamicLighting && View.Family->EngineShowFlags.SubsurfaceScattering &&
-				//@todo-rco: Remove this when we fix the cross-compiler
-				!IsOpenGLPlatform(View.GetShaderPlatform()))
+			if (bScreenSpaceSubsurfacePassNeeded && Radius > 0 && !bSimpleDynamicLighting && View.Family->EngineShowFlags.SubsurfaceScattering)
 			{
 				// can be optimized out if we don't do split screen/stereo rendering (should be done after we some post process refactoring)
 				FRenderingCompositePass* PassExtractSpecular = Context.Graph.RegisterPass(new(FMemStack::Get()) FRCPassPostProcessSubsurfaceExtractSpecular());
