Fixes an assertion failure when launching the engine with r.BasePassOutputsVelocity = 1, but rendering a view with no temporal AA, motion blur and distance field ambient occlusion, blocking the KiteDemo's editor from launching.

diff --git a/Engine/Source/Runtime/Renderer/Private/DeferredShadingRenderer.cpp b/Engine/Source/Runtime/Renderer/Private/DeferredShadingRenderer.cpp
--- a/Engine/Source/Runtime/Renderer/Private/DeferredShadingRenderer.cpp
+++ b/Engine/Source/Runtime/Renderer/Private/DeferredShadingRenderer.cpp
@@ -939,7 +939,6 @@ void FDeferredShadingSceneRenderer::Render(FRHICommandListImmediate& RHICmdList)
 	const bool bShouldRenderVelocities = ShouldRenderVelocities();
 	const bool bUseVelocityGBuffer = FVelocityRendering::OutputsToGBuffer();
 
-	if (bShouldRenderVelocities)
 	{
 		static bool bOnce = false;
 		if (!bOnce)
