Fix UE-18916.

Ensure shadow depth target is initialised when only modulated shadows are in use.

diff --git a/Engine/Source/Runtime/Renderer/Private/ShadowSetup.cpp b/Engine/Source/Runtime/Renderer/Private/ShadowSetup.cpp
--- a/Engine/Source/Runtime/Renderer/Private/ShadowSetup.cpp
+++ b/Engine/Source/Runtime/Renderer/Private/ShadowSetup.cpp
@@ -2414,6 +2414,13 @@ void FForwardShadingSceneRenderer::InitDynamicShadows(FRHICommandListImmediate&
 					}
 				}
 			}
+			else if (VisibleLightInfo.AllProjectedShadows.Num() > 0)
+			{
+				// non whole scene projected shadows are in use. Ensure the shadow depth target is available for modulated shadow use later.
+				FSceneRenderTargets& SceneContext = FSceneRenderTargets::Get(RHICmdList);
+				const FIntPoint ShadowBufferResolution = SceneContext.GetShadowDepthTextureResolution();
+				SceneContext.AllocateForwardShadingShadowDepthTarget(ShadowBufferResolution);
+			}
 		}
 
 		// Calculate visibility of the projected shadows.
