fixed UE-18852 ElementalDemo crashes due to nil constant buffer

diff --git a/Engine/Source/Runtime/Engine/Private/SceneView.cpp b/Engine/Source/Runtime/Engine/Private/SceneView.cpp
--- a/Engine/Source/Runtime/Engine/Private/SceneView.cpp
+++ b/Engine/Source/Runtime/Engine/Private/SceneView.cpp
@@ -158,7 +158,7 @@ static TAutoConsoleVariable<float> CVarSceneColorFringeMax(
 EVertexColorViewMode::Type GVertexColorViewMode = EVertexColorViewMode::Color;
 
 /** Global primitive uniform buffer resource containing identity transformations. */
-TGlobalResource<FIdentityPrimitiveUniformBuffer> GIdentityPrimitiveUniformBuffer;
+ENGINE_API TGlobalResource<FIdentityPrimitiveUniformBuffer> GIdentityPrimitiveUniformBuffer;
 
 FSceneViewStateReference::~FSceneViewStateReference()
 {
diff --git a/Engine/Source/Runtime/Engine/Public/PrimitiveUniformShaderParameters.h b/Engine/Source/Runtime/Engine/Public/PrimitiveUniformShaderParameters.h
--- a/Engine/Source/Runtime/Engine/Public/PrimitiveUniformShaderParameters.h
+++ b/Engine/Source/Runtime/Engine/Public/PrimitiveUniformShaderParameters.h
@@ -122,4 +122,4 @@ public:
 };
 
 /** Global primitive uniform buffer resource containing identity transformations. */
-extern TGlobalResource<FIdentityPrimitiveUniformBuffer> GIdentityPrimitiveUniformBuffer;
\ No newline at end of file
+extern ENGINE_API TGlobalResource<FIdentityPrimitiveUniformBuffer> GIdentityPrimitiveUniformBuffer;
\ No newline at end of file
diff --git a/Engine/Source/Runtime/Renderer/Private/DecalRenderingShared.cpp b/Engine/Source/Runtime/Renderer/Private/DecalRenderingShared.cpp
--- a/Engine/Source/Runtime/Renderer/Private/DecalRenderingShared.cpp
+++ b/Engine/Source/Runtime/Renderer/Private/DecalRenderingShared.cpp
@@ -416,6 +416,20 @@ void FDecalRendering::SetShader(FRHICommandList& RHICmdList, const FViewInfo& Vi
 	auto PixelShader = MaterialShaderMap->GetShader<FDeferredDecalPS>();
 	TShaderMapRef<FDeferredDecalVS> VertexShader(View.ShaderMap);
 
+	// we don't have the Primitive uniform buffer setup for decals (later we want to batch)
+	{
+		auto& PrimitiveVS = VertexShader->GetUniformBufferParameter<FPrimitiveUniformShaderParameters>();
+		auto& PrimitivePS = PixelShader->GetUniformBufferParameter<FPrimitiveUniformShaderParameters>();
+
+		// uncomment to track down usage of the Primitive uniform buffer
+		//	check(!PrimitiveVS.IsBound());
+		//	check(!PrimitivePS.IsBound());
+
+		// to prevent potential shader error (UE-18852 ElementalDemo crashes due to nil constant buffer)
+		SetUniformBufferParameter(RHICmdList, VertexShader->GetVertexShader(), PrimitiveVS, GIdentityPrimitiveUniformBuffer);
+		SetUniformBufferParameter(RHICmdList, PixelShader->GetPixelShader(), PrimitivePS, GIdentityPrimitiveUniformBuffer);
+	}
+
 	if(bShaderComplexity)
 	{
 		TShaderMapRef<FShaderComplexityAccumulatePS> VisualizePixelShader(View.ShaderMap);
