Merges 2614242 from UT: Makes the distinction between the effective view space matrices and the camera view space matrices, to fix  transform node from and to thew view space in the shadow pass.

diff --git a/Engine/Source/Runtime/Engine/Private/Materials/HLSLMaterialTranslator.h b/Engine/Source/Runtime/Engine/Private/Materials/HLSLMaterialTranslator.h
--- a/Engine/Source/Runtime/Engine/Private/Materials/HLSLMaterialTranslator.h
+++ b/Engine/Source/Runtime/Engine/Private/Materials/HLSLMaterialTranslator.h
@@ -3289,7 +3289,7 @@ protected:
 				}
 				else if (DestCoordBasis == MCB_View)
 				{
-					CodeStr = TEXT("mul(<A>, <MATRIX>(View.<PREV>TranslatedWorldToView))");
+					CodeStr = TEXT("mul(<A>, <MATRIX>(View.<PREV>TranslatedWorldToCameraView))");
 				}
 				// else use MCB_World as intermediary basis
 				break;
@@ -3334,7 +3334,7 @@ protected:
 			{
 				if (DestCoordBasis == MCB_TranslatedWorld)
 				{
-					CodeStr = TEXT("mul(<A>, <MATRIX>(View.<PREV>ViewToTranslatedWorld))");
+					CodeStr = TEXT("mul(<A>, <MATRIX>(View.<PREV>CameraViewToTranslatedWorld))");
 				}
 				// else use MCB_TranslatedWorld as intermediary basis
 				IntermediaryBasis = MCB_TranslatedWorld;
diff --git a/Engine/Source/Runtime/Engine/Public/SceneView.h b/Engine/Source/Runtime/Engine/Public/SceneView.h
--- a/Engine/Source/Runtime/Engine/Public/SceneView.h
+++ b/Engine/Source/Runtime/Engine/Public/SceneView.h
@@ -354,6 +354,8 @@ BEGIN_UNIFORM_BUFFER_STRUCT_WITH_CONSTRUCTOR(FViewUniformShaderParameters,ENGINE
 	DECLARE_UNIFORM_BUFFER_STRUCT_MEMBER(FMatrix,WorldToClip)
 	DECLARE_UNIFORM_BUFFER_STRUCT_MEMBER(FMatrix,TranslatedWorldToView)
 	DECLARE_UNIFORM_BUFFER_STRUCT_MEMBER(FMatrix,ViewToTranslatedWorld)
+	DECLARE_UNIFORM_BUFFER_STRUCT_MEMBER(FMatrix,TranslatedWorldToCameraView)
+	DECLARE_UNIFORM_BUFFER_STRUCT_MEMBER(FMatrix,CameraViewToTranslatedWorld)
 	DECLARE_UNIFORM_BUFFER_STRUCT_MEMBER(FMatrix,ViewToClip)
 	DECLARE_UNIFORM_BUFFER_STRUCT_MEMBER(FMatrix,ClipToView)
 	DECLARE_UNIFORM_BUFFER_STRUCT_MEMBER(FMatrix,ClipToTranslatedWorld)
@@ -421,6 +423,8 @@ BEGIN_UNIFORM_BUFFER_STRUCT_WITH_CONSTRUCTOR(FViewUniformShaderParameters,ENGINE
 	DECLARE_UNIFORM_BUFFER_STRUCT_MEMBER(FMatrix,PrevTranslatedWorldToClip)
 	DECLARE_UNIFORM_BUFFER_STRUCT_MEMBER(FMatrix,PrevTranslatedWorldToView)
 	DECLARE_UNIFORM_BUFFER_STRUCT_MEMBER(FMatrix,PrevViewToTranslatedWorld)
+	DECLARE_UNIFORM_BUFFER_STRUCT_MEMBER(FMatrix,PrevTranslatedWorldToCameraView)
+	DECLARE_UNIFORM_BUFFER_STRUCT_MEMBER(FMatrix,PrevCameraViewToTranslatedWorld)
 	DECLARE_UNIFORM_BUFFER_STRUCT_MEMBER(FVector,PrevViewOrigin)
 	DECLARE_UNIFORM_BUFFER_STRUCT_MEMBER(FVector,PrevPreViewTranslation)
 	DECLARE_UNIFORM_BUFFER_STRUCT_MEMBER(FMatrix,PrevInvViewProj)
diff --git a/Engine/Source/Runtime/Renderer/Private/SceneRendering.cpp b/Engine/Source/Runtime/Renderer/Private/SceneRendering.cpp
--- a/Engine/Source/Runtime/Renderer/Private/SceneRendering.cpp
+++ b/Engine/Source/Runtime/Renderer/Private/SceneRendering.cpp
@@ -464,6 +464,8 @@ TUniformBufferRef<FViewUniformShaderParameters> FViewInfo::CreateUniformBuffer(
 	ViewUniformShaderParameters.WorldToClip = ViewProjectionMatrix;
 	ViewUniformShaderParameters.TranslatedWorldToView = EffectiveTranslatedViewMatrix;
 	ViewUniformShaderParameters.ViewToTranslatedWorld = EffectiveViewToTranslatedWorld;
+	ViewUniformShaderParameters.TranslatedWorldToCameraView = ViewMatrices.TranslatedViewMatrix;
+	ViewUniformShaderParameters.CameraViewToTranslatedWorld = ViewUniformShaderParameters.TranslatedWorldToCameraView.Inverse();
 	ViewUniformShaderParameters.ViewToClip = ViewMatrices.ProjMatrix;
 	ViewUniformShaderParameters.ClipToView = ViewMatrices.GetInvProjMatrix();
 	ViewUniformShaderParameters.ClipToTranslatedWorld = ViewMatrices.InvTranslatedViewProjectionMatrix;
@@ -497,6 +499,8 @@ TUniformBufferRef<FViewUniformShaderParameters> FViewInfo::CreateUniformBuffer(
 	ViewUniformShaderParameters.PrevTranslatedWorldToClip = PrevViewMatrices.TranslatedViewProjectionMatrix;
 	ViewUniformShaderParameters.PrevTranslatedWorldToView = PrevViewMatrices.TranslatedViewMatrix;
 	ViewUniformShaderParameters.PrevViewToTranslatedWorld = ViewUniformShaderParameters.PrevTranslatedWorldToView.Inverse();
+	ViewUniformShaderParameters.PrevTranslatedWorldToCameraView = PrevViewMatrices.TranslatedViewMatrix;
+	ViewUniformShaderParameters.PrevCameraViewToTranslatedWorld = ViewUniformShaderParameters.PrevTranslatedWorldToCameraView.Inverse();
 	ViewUniformShaderParameters.PrevViewOrigin = PrevViewMatrices.ViewOrigin;
 	ViewUniformShaderParameters.PrevPreViewTranslation = PrevViewMatrices.PreViewTranslation;
 	// can be optimized
