If IndirectLightingCache quality is set to OFF, then never choose an ILC lighting policy.

diff --git a/Engine/Source/Runtime/Renderer/Private/BasePassRendering.h b/Engine/Source/Runtime/Renderer/Private/BasePassRendering.h
--- a/Engine/Source/Runtime/Renderer/Private/BasePassRendering.h
+++ b/Engine/Source/Runtime/Renderer/Private/BasePassRendering.h
@@ -947,13 +947,15 @@ void ProcessBasePassMesh(
 							{
 								const FIndirectLightingCacheAllocation* IndirectLightingCacheAllocation = Parameters.PrimitiveSceneProxy->GetPrimitiveSceneInfo()->IndirectLightingCacheAllocation;
 								const bool bPrimitiveIsMovable = Parameters.PrimitiveSceneProxy->IsMovable();
+								const bool bPrimitiveUsesILC = Parameters.PrimitiveSceneProxy->GetIndirectLightingCacheQuality() != ILCQ_Off;								
 
 								// Use the indirect lighting cache shaders if the object has a cache allocation
 								// This happens for objects with unbuilt lighting
-								if ((IndirectLightingCacheAllocation && IndirectLightingCacheAllocation->IsValid())
+								if (bPrimitiveUsesILC &&
+									((IndirectLightingCacheAllocation && IndirectLightingCacheAllocation->IsValid())
 									// Use the indirect lighting cache shaders if the object is movable, it may not have a cache allocation yet because that is done in InitViews
 									// And movable objects are sometimes rendered in the static draw lists
-									|| bPrimitiveIsMovable)
+									|| bPrimitiveIsMovable))
 								{
 									if (CanIndirectLightingCacheUseVolumeTexture(Parameters.FeatureLevel) 
 										// Translucency forces point sample for pixel performance
