added r.SSS.HalfRes to get ScreenSpaceSubSurfaceScattering in higher quality (not yet optimized) as reference

diff --git a/Engine/Shaders/PostProcessSubsurface.usf b/Engine/Shaders/PostProcessSubsurface.usf
--- a/Engine/Shaders/PostProcessSubsurface.usf
+++ b/Engine/Shaders/PostProcessSubsurface.usf
@@ -38,6 +38,9 @@ float GetSubsurfaceStrength(float2 UV)
 	return Mask * ScreenSpaceData.GBuffer.Opacity;
 }
 
+// 0:full res (slower but reference) / 1:faster half res with some quality loss
+// #define HALF_RES 1
+
 // 0:speculars leak into SSS / 1:requires alpha channel for scenecolor
 // #define SSSS_SPECULAR_CORRECTION 1
 
@@ -301,6 +304,7 @@ void SetupPS(in float4 UVAndScreenPos : TEXCOORD0, out float4 OutColor : SV_Targ
 {
 	float2 UV = UVAndScreenPos.xy;
 
+#if HALF_RES
 	// order aligned with Gather() hardware implementation
 	// RGB: color*A, A:weight 0 if no subsurface scattering
 	float4 A = SetupSubsurfaceForOnePixel(UV + float2(-0.5,  0.5f) * PostprocessInput0Size.zw);
@@ -308,11 +312,11 @@ void SetupPS(in float4 UVAndScreenPos : TEXCOORD0, out float4 OutColor : SV_Targ
 	float4 C = SetupSubsurfaceForOnePixel(UV + float2( 0.5, -0.5f) * PostprocessInput0Size.zw);
 	float4 D = SetupSubsurfaceForOnePixel(UV + float2(-0.5, -0.5f) * PostprocessInput0Size.zw);
 
-	OutColor = (A + B) + (C + D);
+	float4 Sum = (A + B) + (C + D);
 
-	float Div = 1.0f / max(OutColor.a, 0.00001f);
+	float Div = 1.0f / max(Sum.a, 0.00001f);
 
-	OutColor.rgba *= Div;
+	OutColor.rgb = Sum.rgb * Div;
 	
 	float4 FourDepth = GatherSceneDepth(UV, PostprocessInput0Size.zw);
 
@@ -320,6 +324,10 @@ void SetupPS(in float4 UVAndScreenPos : TEXCOORD0, out float4 OutColor : SV_Targ
 	float SingleDepth = dot(FourDepth, float4(A.a, B.a, C.a, D.a)) * Div;
 
 	OutColor.a = SingleDepth;
+#else // HALF_RES
+	OutColor.rgb = SetupSubsurfaceForOnePixel(UV).rgb;
+	OutColor.a = CalcSceneDepth(UV);
+#endif // HALF_RES
 
 #if 0
 	float SSSScaleX = SSSParams.x;
@@ -382,7 +390,7 @@ void MainPS(float4 UVAndScreenPos : TEXCOORD0, out float4 OutColor : SV_Target0)
 #endif
 
 	// can be optimized
-	float2 GBufferUV = (ViewportUV * PostprocessInput0Size.xy * 2 + View.ViewRectMin.xy) * View.BufferSizeAndInvSize.zw;
+	float2 GBufferUV = (ViewportUV * View.ViewSizeAndInvSize.xy + View.ViewRectMin.xy) * View.BufferSizeAndInvSize.zw;
 
 	bool bIncludeCenter = true;
 
@@ -430,6 +438,7 @@ void SubsurfaceRecombinePS(float4 UVAndScreenPos : TEXCOORD0, out float4 OutColo
 	float BlendFactor = (ScreenSpaceData.GBuffer.ShadingModelID == SHADINGMODELID_SUBSURFACE_PROFILE);
 	float LerpFactor = 1;
 
+#if HALF_RES
 #if 1
 	// fade out subsurface scattering if radius is too small to be more crips (not blend with half resolution)
 	// minor quality improvement (faces are more detailed in distance)
@@ -438,6 +447,7 @@ void SubsurfaceRecombinePS(float4 UVAndScreenPos : TEXCOORD0, out float4 OutColo
 	// hack to debug if the fade is happening at the same time the VisualizeSSS shows it
 //	SSSColor = 0;
 #endif
+#endif // HALF_RES
 
 	// only affect pixels that have SSS
 	LerpFactor *= BlendFactor;
diff --git a/Engine/Source/Runtime/Renderer/Private/CompositionLighting/CompositionLighting.cpp b/Engine/Source/Runtime/Renderer/Private/CompositionLighting/CompositionLighting.cpp
--- a/Engine/Source/Runtime/Renderer/Private/CompositionLighting/CompositionLighting.cpp
+++ b/Engine/Source/Runtime/Renderer/Private/CompositionLighting/CompositionLighting.cpp
@@ -40,6 +40,13 @@ static TAutoConsoleVariable<float> CVarSSSScale(
 	TEXT(">1: scale scatter radius up (for testing)"),
 	ECVF_Scalability | ECVF_RenderThreadSafe);
 
+static TAutoConsoleVariable<int32> CVarSSSHalfRes(
+	TEXT("r.SSS.HalfRes"),
+	1,
+	TEXT(" 0: full quality (not optimized, as reference)\n")
+	TEXT(" 1: parts of the algorithm runs in half resolution which is lower quality but faster (default)"),
+	ECVF_RenderThreadSafe  | ECVF_Scalability);
+
 static bool IsAmbientCubemapPassRequired(FPostprocessContext& Context)
 {
 	FScene* Scene = (FScene*)Context.View.Family->Scene;
@@ -381,22 +388,24 @@ void FCompositionLighting::ProcessAfterLighting(FRHICommandListImmediate& RHICmd
 			bool bScreenSpaceSubsurfacePassNeeded = (View.ShadingModelMaskInView & (1 << MSM_SubsurfaceProfile)) != 0;
 			if (bScreenSpaceSubsurfacePassNeeded && Radius > 0 && !bSimpleDynamicLighting && View.Family->EngineShowFlags.SubsurfaceScattering)
 			{
+				bool bHalfRes = CVarSSSHalfRes.GetValueOnRenderThread() != 0;
+
 				// can be optimized out if we don't do split screen/stereo rendering (should be done after we some post process refactoring)
 				FRenderingCompositePass* PassExtractSpecular = Context.Graph.RegisterPass(new(FMemStack::Get()) FRCPassPostProcessSubsurfaceExtractSpecular());
 				PassExtractSpecular->SetInput(ePId_Input0, Context.FinalOutput);
 
-				FRenderingCompositePass* PassSetup = Context.Graph.RegisterPass(new(FMemStack::Get()) FRCPassPostProcessSubsurfaceSetup(View));
+				FRenderingCompositePass* PassSetup = Context.Graph.RegisterPass(new(FMemStack::Get()) FRCPassPostProcessSubsurfaceSetup(View, bHalfRes));
 				PassSetup->SetInput(ePId_Input0, Context.FinalOutput);
 
-				FRenderingCompositePass* Pass0 = Context.Graph.RegisterPass(new(FMemStack::Get()) FRCPassPostProcessSubsurface(0));
+				FRenderingCompositePass* Pass0 = Context.Graph.RegisterPass(new(FMemStack::Get()) FRCPassPostProcessSubsurface(0, bHalfRes));
 				Pass0->SetInput(ePId_Input0, PassSetup);
 
-				FRenderingCompositePass* Pass1 = Context.Graph.RegisterPass(new(FMemStack::Get()) FRCPassPostProcessSubsurface(1));
+				FRenderingCompositePass* Pass1 = Context.Graph.RegisterPass(new(FMemStack::Get()) FRCPassPostProcessSubsurface(1, bHalfRes));
 				Pass1->SetInput(ePId_Input0, Pass0);
 				Pass1->SetInput(ePId_Input1, PassSetup);
 
 				// full res composite pass, no blurring (Radius=0)
-				FRenderingCompositePass* RecombinePass = Context.Graph.RegisterPass(new(FMemStack::Get()) FRCPassPostProcessSubsurfaceRecombine());
+				FRenderingCompositePass* RecombinePass = Context.Graph.RegisterPass(new(FMemStack::Get()) FRCPassPostProcessSubsurfaceRecombine(bHalfRes));
 				RecombinePass->SetInput(ePId_Input0, Pass1);
 				RecombinePass->SetInput(ePId_Input1, PassExtractSpecular);
 			}
diff --git a/Engine/Source/Runtime/Renderer/Private/PostProcess/PostProcessSubsurface.cpp b/Engine/Source/Runtime/Renderer/Private/PostProcess/PostProcessSubsurface.cpp
--- a/Engine/Source/Runtime/Renderer/Private/PostProcess/PostProcessSubsurface.cpp
+++ b/Engine/Source/Runtime/Renderer/Private/PostProcess/PostProcessSubsurface.cpp
@@ -327,8 +327,9 @@ FPooledRenderTargetDesc FRCPassPostProcessSubsurfaceVisualize::ComputeOutputDesc
 /**
  * Encapsulates the post processing subsurface scattering pixel shader.
  * @param SetupMode 0:without specular correction, 1:with specular correction
+ * @param HalfRes 0:to full res, 1:to half res
  */
-template <uint32 SetupMode>
+template <uint32 SetupMode, uint32 HalfRes>
 class FPostProcessSubsurfaceSetupPS : public FGlobalShader
 {
 	DECLARE_SHADER_TYPE(FPostProcessSubsurfaceSetupPS , Global);
@@ -342,6 +343,7 @@ class FPostProcessSubsurfaceSetupPS : public FGlobalShader
 	{
 		FGlobalShader::ModifyCompilationEnvironment(Platform, OutEnvironment);
 		OutEnvironment.SetDefine(TEXT("SETUP_MODE"), SetupMode);
+		OutEnvironment.SetDefine(TEXT("HALF_RES"), HalfRes);
 	}
 
 	/** Default constructor. */
@@ -393,19 +395,19 @@ public:
 };
 
 // #define avoids a lot of code duplication
-#define VARIATION1(A) typedef FPostProcessSubsurfaceSetupPS<A> FPostProcessSubsurfaceSetupPS##A; \
-	IMPLEMENT_SHADER_TYPE2(FPostProcessSubsurfaceSetupPS##A, SF_Pixel);
-
+#define VARIATION1(A)		VARIATION2(A,0)			VARIATION2(A,1)
+#define VARIATION2(A, B) typedef FPostProcessSubsurfaceSetupPS<A, B> FPostProcessSubsurfaceSetupPS##A##B; \
+	IMPLEMENT_SHADER_TYPE2(FPostProcessSubsurfaceSetupPS##A##B, SF_Pixel);
 	VARIATION1(0) VARIATION1(1)
-
 #undef VARIATION1
+#undef VARIATION2
 
 
-template <uint32 SetupMode>
+template <uint32 SetupMode, uint32 HalfRes>
 void SetSubsurfaceSetupShader(const FRenderingCompositePassContext& Context)
 {
 	TShaderMapRef<FPostProcessVS> VertexShader(Context.GetShaderMap());
-	TShaderMapRef<FPostProcessSubsurfaceSetupPS<SetupMode> > PixelShader(Context.GetShaderMap());
+	TShaderMapRef<FPostProcessSubsurfaceSetupPS<SetupMode, HalfRes> > PixelShader(Context.GetShaderMap());
 
 	static FGlobalBoundShaderState BoundShaderState;
 
@@ -634,8 +636,9 @@ FPooledRenderTargetDesc FRCPassPostProcessSubsurfaceExtractSpecular::ComputeOutp
 
 // --------------------------------------
 
-FRCPassPostProcessSubsurfaceSetup::FRCPassPostProcessSubsurfaceSetup(FViewInfo& View)
+FRCPassPostProcessSubsurfaceSetup::FRCPassPostProcessSubsurfaceSetup(FViewInfo& View, bool bInHalfRes)
 	: ViewRect(View.ViewRect)
+	, bHalfRes(bInHalfRes)
 {
 }
 
@@ -658,8 +661,13 @@ void FRCPassPostProcessSubsurfaceSetup::Process(FRenderingCompositePassContext&
 	FIntPoint DestSize = PassOutputs[0].RenderTargetDesc.Extent;
 
 	FIntRect DestRect = FIntRect(0, 0, DestSize.X, DestSize.Y);
-	// upscale rectangle to not slightly scale (might miss a pixel)
-	FIntRect SrcRect = DestRect * 2 + View.ViewRect.Min;
+	FIntRect SrcRect = View.ViewRect;
+	
+	if(bHalfRes)
+	{
+		// upscale rectangle to not slightly scale (might miss a pixel)
+		SrcRect = DestRect * 2 + View.ViewRect.Min;
+	}
 
 	const FSceneRenderTargetItem& DestRenderTarget = PassOutputs[0].RequestSurface(Context);
 
@@ -678,11 +686,11 @@ void FRCPassPostProcessSubsurfaceSetup::Process(FRenderingCompositePassContext&
 
 	if(bSpecularCorrection)
 	{
-		SetSubsurfaceSetupShader<1>(Context);
+		if(bHalfRes) SetSubsurfaceSetupShader<1, 1>(Context); else SetSubsurfaceSetupShader<1, 0>(Context);
 	}
 	else
 	{
-		SetSubsurfaceSetupShader<0>(Context);
+		if(bHalfRes) SetSubsurfaceSetupShader<0, 1>(Context); else SetSubsurfaceSetupShader<0, 0>(Context);
 	}
 
 	// Draw a quad mapping scene color to the view's render target
@@ -713,10 +721,12 @@ FPooledRenderTargetDesc FRCPassPostProcessSubsurfaceSetup::ComputeOutputDesc(EPa
 
 	Ret.Extent = ViewRect.Size();
 
-	// half res
-	Ret.Extent  = FIntPoint::DivideAndRoundUp(Ret.Extent, 2);
-	Ret.Extent.X = FMath::Max(1, Ret.Extent.X);
-	Ret.Extent.Y = FMath::Max(1, Ret.Extent.Y);
+	if(bHalfRes)
+	{
+		Ret.Extent = FIntPoint::DivideAndRoundUp(Ret.Extent, 2);
+		Ret.Extent.X = FMath::Max(1, Ret.Extent.X);
+		Ret.Extent.Y = FMath::Max(1, Ret.Extent.Y);
+	}
 	
 	return Ret;
 }
@@ -805,8 +815,9 @@ public:
 #undef VARIATION2
 
 
-FRCPassPostProcessSubsurface::FRCPassPostProcessSubsurface(uint32 InDirection)
+FRCPassPostProcessSubsurface::FRCPassPostProcessSubsurface(uint32 InDirection, bool bInHalfRes)
 	: Direction(InDirection) 
+	, bHalfRes(bInHalfRes)
 {
 	check(InDirection < 2);
 }
@@ -931,9 +942,9 @@ FPooledRenderTargetDesc FRCPassPostProcessSubsurface::ComputeOutputDesc(EPassOut
 
 
 
-/** Encapsulates the post processing subsurafce recombine pixel shader. */
-// @param SpecularCorrection 0: reconstruct specular
-template <uint32 RecombineMethod>
+/** Encapsulates the post processing subsurface recombine pixel shader. */
+// @param RecombineMethod 1: with specular correction, 0: without specular correction
+template <uint32 RecombineMethod, uint32 HalfRes>
 class TPostProcessSubsurfaceRecombinePS : public FGlobalShader
 {
 	DECLARE_SHADER_TYPE(TPostProcessSubsurfaceRecombinePS, Global);
@@ -947,6 +958,7 @@ class TPostProcessSubsurfaceRecombinePS : public FGlobalShader
 	{
 		FGlobalShader::ModifyCompilationEnvironment(Platform,OutEnvironment);
 		OutEnvironment.SetDefine(TEXT("SSS_RECOMBINE_METHOD"), RecombineMethod);
+		OutEnvironment.SetDefine(TEXT("HALF_RES"), HalfRes);
 	}
 
 	/** Default constructor. */
@@ -996,16 +1008,17 @@ public:
 };
 
 // #define avoids a lot of code duplication
-#define VARIATION1(A) typedef TPostProcessSubsurfaceRecombinePS<A> TPostProcessSubsurfaceRecombinePS##A; \
-	IMPLEMENT_SHADER_TYPE2(TPostProcessSubsurfaceRecombinePS##A, SF_Pixel);
+#define VARIATION1(A)		VARIATION2(A,0)			VARIATION2(A,1)
+#define VARIATION2(A, B) typedef TPostProcessSubsurfaceRecombinePS<A, B> TPostProcessSubsurfaceRecombinePS##A##B; \
+	IMPLEMENT_SHADER_TYPE2(TPostProcessSubsurfaceRecombinePS##A##B, SF_Pixel);
 	VARIATION1(0) VARIATION1(1)
 #undef VARIATION1
+#undef VARIATION2
 
-
-template <uint32 RecombineMethod>
+template <uint32 RecombineMethod, uint32 HalfRes>
 void SetSubsurfaceRecombineShader(const FRenderingCompositePassContext& Context, TShaderMapRef<FPostProcessVS> &VertexShader)
 {
-	TShaderMapRef<TPostProcessSubsurfaceRecombinePS<RecombineMethod> > PixelShader(Context.GetShaderMap());
+	TShaderMapRef<TPostProcessSubsurfaceRecombinePS<RecombineMethod, HalfRes> > PixelShader(Context.GetShaderMap());
 
 	static FGlobalBoundShaderState BoundShaderState;
 
@@ -1015,6 +1028,11 @@ void SetSubsurfaceRecombineShader(const FRenderingCompositePassContext& Context,
 	VertexShader->SetParameters(Context);
 }
 
+FRCPassPostProcessSubsurfaceRecombine::FRCPassPostProcessSubsurfaceRecombine(bool bInHalfRes)
+	: bHalfRes(bInHalfRes)
+{
+}
+
 void FRCPassPostProcessSubsurfaceRecombine::Process(FRenderingCompositePassContext& Context)
 {
 	FSceneRenderTargets& SceneContext = FSceneRenderTargets::Get(Context.RHICmdList);
@@ -1057,11 +1075,11 @@ void FRCPassPostProcessSubsurfaceRecombine::Process(FRenderingCompositePassConte
 
 	if(bDoSpecularCorrection)
 	{
-		SetSubsurfaceRecombineShader<1>(Context, VertexShader);
+		if(bHalfRes) SetSubsurfaceRecombineShader<1, 1>(Context, VertexShader); else SetSubsurfaceRecombineShader<1, 0>(Context, VertexShader);
 	}
 	else
 	{
-		SetSubsurfaceRecombineShader<0>(Context, VertexShader);
+		if(bHalfRes) SetSubsurfaceRecombineShader<0, 1>(Context, VertexShader); else SetSubsurfaceRecombineShader<0, 0>(Context, VertexShader);
 	}
 
 	DrawRectangle(
diff --git a/Engine/Source/Runtime/Renderer/Private/PostProcess/PostProcessSubsurface.h b/Engine/Source/Runtime/Renderer/Private/PostProcess/PostProcessSubsurface.h
--- a/Engine/Source/Runtime/Renderer/Private/PostProcess/PostProcessSubsurface.h
+++ b/Engine/Source/Runtime/Renderer/Private/PostProcess/PostProcessSubsurface.h
@@ -43,7 +43,7 @@ class FRCPassPostProcessSubsurfaceSetup : public TRenderingCompositePassBase<1,
 {
 public:
 	// constructor
-	FRCPassPostProcessSubsurfaceSetup(FViewInfo& View);
+	FRCPassPostProcessSubsurfaceSetup(FViewInfo& View, bool bInHalfRes);
 
 	// interface FRenderingCompositePass ---------
 	virtual void Process(FRenderingCompositePassContext& Context) override;
@@ -51,6 +51,7 @@ public:
 	virtual void Release() override { delete this; }
 
 	FIntRect ViewRect;
+	bool bHalfRes;
 };
 
 
@@ -63,7 +64,7 @@ class FRCPassPostProcessSubsurface : public TRenderingCompositePassBase<2, 1>
 public:
 	// constructor
 	// @param InDirection 0:horizontal/1:vertical
-	FRCPassPostProcessSubsurface(uint32 InDirection);
+	FRCPassPostProcessSubsurface(uint32 InDirection, bool bInHalfRes);
 
 	// interface FRenderingCompositePass ---------
 
@@ -74,6 +75,7 @@ public:
 private:
 	// 0:horizontal/1:vertical
 	uint32 Direction;
+	bool bHalfRes;
 };
 
 
@@ -84,9 +86,14 @@ private:
 class FRCPassPostProcessSubsurfaceRecombine : public TRenderingCompositePassBase<2, 1>
 {
 public:
+	// constructor
+	FRCPassPostProcessSubsurfaceRecombine(bool bInHalfRes);
+
 	// interface FRenderingCompositePass ---------
 
 	virtual void Process(FRenderingCompositePassContext& Context) override;
 	virtual void Release() override { delete this; }
 	virtual FPooledRenderTargetDesc ComputeOutputDesc(EPassOutputId InPassOutputId) const override;
+
+	bool bHalfRes;
 };
