Draw static-relevant meshes during self-shadow-only projection rendering. UE-17370

diff --git a/Engine/Source/Runtime/Renderer/Private/ShadowRendering.cpp b/Engine/Source/Runtime/Renderer/Private/ShadowRendering.cpp
--- a/Engine/Source/Runtime/Renderer/Private/ShadowRendering.cpp
+++ b/Engine/Source/Runtime/Renderer/Private/ShadowRendering.cpp
@@ -2128,6 +2128,31 @@ void FProjectedShadowInfo::RenderProjection(FRHICommandListImmediate& RHICmdList
 				}
 			}
 		}
+
+		if (bSelfShadowOnly)
+		{
+			for (int32 ElementIndex = 0; ElementIndex < SubjectMeshElements.Num(); ++ElementIndex)
+			{
+				const FShadowStaticMeshElement& ShadowMesh = SubjectMeshElements[ElementIndex];
+				const FStaticMesh& StaticMesh = *ShadowMesh.Mesh;
+ 
+				if (View->StaticMeshShadowDepthMap[StaticMesh.Id])
+				{
+					const float DitherValue = View->GetDitheredLODTransitionValue(StaticMesh);
+ 
+					FDepthDrawingPolicyFactory::DrawStaticMesh(
+						RHICmdList, 
+						*View,
+						FDepthDrawingPolicyFactory::ContextType(DDM_AllOccluders),
+						StaticMesh,
+						StaticMesh.Elements.Num() == 1 ? 1 : View->StaticMeshBatchVisibility[StaticMesh.Id],
+						true,
+						DitherValue,
+						StaticMesh.PrimitiveSceneInfo->Proxy,
+						StaticMesh.BatchHitProxyId);
+				}
+			}
+		}
 	}
 	else if (IsWholeSceneDirectionalShadow())
 	{
