Only draw self-shadow mesh if not pre-shadow. UE-17370

diff --git a/Engine/Source/Runtime/Renderer/Private/ShadowRendering.cpp b/Engine/Source/Runtime/Renderer/Private/ShadowRendering.cpp
--- a/Engine/Source/Runtime/Renderer/Private/ShadowRendering.cpp
+++ b/Engine/Source/Runtime/Renderer/Private/ShadowRendering.cpp
@@ -2129,17 +2129,15 @@ void FProjectedShadowInfo::RenderProjection(FRHICommandListImmediate& RHICmdList
 			}
 		}
 
-		if (bSelfShadowOnly)
+		if (bSelfShadowOnly && !bPreShadow)
 		{
 			for (int32 ElementIndex = 0; ElementIndex < SubjectMeshElements.Num(); ++ElementIndex)
 			{
-				const FShadowStaticMeshElement& ShadowMesh = SubjectMeshElements[ElementIndex];
-				const FStaticMesh& StaticMesh = *ShadowMesh.Mesh;
+				const FStaticMesh& StaticMesh = *SubjectMeshElements[ElementIndex].Mesh;
  
 				if (View->StaticMeshShadowDepthMap[StaticMesh.Id])
 				{
 					const float DitherValue = View->GetDitheredLODTransitionValue(StaticMesh);
- 
 					FDepthDrawingPolicyFactory::DrawStaticMesh(
 						RHICmdList, 
 						*View,
@@ -2149,7 +2147,8 @@ void FProjectedShadowInfo::RenderProjection(FRHICommandListImmediate& RHICmdList
 						true,
 						DitherValue,
 						StaticMesh.PrimitiveSceneInfo->Proxy,
-						StaticMesh.BatchHitProxyId);
+						StaticMesh.BatchHitProxyId
+						);
 				}
 			}
 		}
