HMD hidden area mesh implementation.

Adds an interface for rendering an HMD's hidden area mesh as a depth prepass.
Either runs in the standard depth prepass if it's being used, or runs as its own prepass.
Runs only if running in stereo, the device has a hidden area mesh defined and vr.HiddenAreaMask is set to 1.

diff --git a/Engine/Plugins/Runtime/Steam/SteamVR/Source/SteamVR/Private/SteamVRHMD.cpp b/Engine/Plugins/Runtime/Steam/SteamVR/Source/SteamVR/Private/SteamVRHMD.cpp
--- a/Engine/Plugins/Runtime/Steam/SteamVR/Source/SteamVR/Private/SteamVRHMD.cpp
+++ b/Engine/Plugins/Runtime/Steam/SteamVR/Source/SteamVR/Private/SteamVRHMD.cpp
@@ -952,14 +952,18 @@ void FSteamVRHMD::Startup()
 		}
 
 		// Initialize our controller to device index
-		for (int32 UnrealControllerIndex = 0; UnrealControllerIndex < MAX_STEAMVR_CONTROLLER_PAIRS; ++UnrealControllerIndex )
+		for (int32 UnrealControllerIndex = 0; UnrealControllerIndex < MAX_STEAMVR_CONTROLLER_PAIRS; ++UnrealControllerIndex)
 		{
-			for( int32 HandIndex = 0; HandIndex < 2; ++HandIndex )
+			for (int32 HandIndex = 0; HandIndex < 2; ++HandIndex)
 			{
-				UnrealControllerIdAndHandToDeviceIdMap[ UnrealControllerIndex ][ HandIndex ] = INDEX_NONE;
+				UnrealControllerIdAndHandToDeviceIdMap[UnrealControllerIndex][HandIndex] = INDEX_NONE;
 			}
 		}
 
+		// Setup hidden area meshes
+		HiddenAreaMeshes[0].Build(VRSystem->GetHiddenAreaMesh(vr::Hmd_Eye::Eye_Left));
+		HiddenAreaMeshes[1].Build(VRSystem->GetHiddenAreaMesh(vr::Hmd_Eye::Eye_Right));
+
 #if PLATFORM_WINDOWS
 		if (IsPCPlatform(GMaxRHIShaderPlatform) && !IsOpenGLPlatform(GMaxRHIShaderPlatform))
 		{
@@ -1079,4 +1083,62 @@ void FSteamVRHMD::UnloadOpenVRModule()
 	}
 }
 
+FSteamVRHMD::FHiddenAreaMesh::FHiddenAreaMesh() :
+	pVertices(nullptr),
+	pIndices(nullptr),
+	NumVertices(0),
+	NumIndices(0),
+	NumTriangles(0)
+{}
+
+FSteamVRHMD::FHiddenAreaMesh::~FHiddenAreaMesh()
+{
+	if (pVertices)
+	{
+		delete[] pVertices;
+	}
+
+	if (pIndices)
+	{
+		delete[] pIndices;
+	}
+}
+
+void FSteamVRHMD::FHiddenAreaMesh::Build(const vr::HiddenAreaMesh_t& Mesh)
+{
+	check(pVertices == nullptr);
+
+	NumTriangles = Mesh.unTriangleCount;
+	NumVertices  = NumTriangles * 3;
+	NumIndices   = NumVertices;
+
+	// Did we get any data?
+	if (NumVertices == 0)
+	{
+		return;
+	}
+
+	pVertices = new FVector4[NumVertices];
+	pIndices  = new uint16[NumIndices];
+
+	uint32 DataIndex = 0;
+	for (uint32 TriangleIter = 0; TriangleIter < NumTriangles; ++TriangleIter)
+	{
+		for (uint32 VertexIter = 0; VertexIter < 3; ++VertexIter)
+		{
+			const vr::HmdVector2_t& Vertex = Mesh.pVertexData[DataIndex];
+
+			// Remap from [0 1] to [-1 1] to match our NDC space
+			pVertices[DataIndex].X = (Vertex.v[0] * 2.0f) - 1.0f;
+			pVertices[DataIndex].Y = (Vertex.v[1] * 2.0f) - 1.0f;
+			pVertices[DataIndex].Z = 1.0f; // Setting to 1 for reversed depth near plane
+			pVertices[DataIndex].W = 1.0f; 
+
+			pIndices[DataIndex] = DataIndex;
+
+			++DataIndex;
+		}
+	}
+}
+
 #endif //STEAMVR_SUPPORTED_PLATFORMS
diff --git a/Engine/Plugins/Runtime/Steam/SteamVR/Source/SteamVR/Private/SteamVRHMD.h b/Engine/Plugins/Runtime/Steam/SteamVR/Source/SteamVR/Private/SteamVRHMD.h
--- a/Engine/Plugins/Runtime/Steam/SteamVR/Source/SteamVR/Private/SteamVRHMD.h
+++ b/Engine/Plugins/Runtime/Steam/SteamVR/Source/SteamVR/Private/SteamVRHMD.h
@@ -84,6 +84,9 @@ public:
 	virtual void SetBaseOrientation(const FQuat& BaseOrient) override;
 	virtual FQuat GetBaseOrientation() const override;
 
+	virtual bool HasHiddenAreaMask() const override { return HiddenAreaMeshes[0].IsValid() && HiddenAreaMeshes[1].IsValid(); }
+	virtual void DrawHiddenAreaMaskView_RenderThread(FRHICommandList& RHICmdList, const FViewInfo& View) const override;
+
 	virtual void DrawDistortionMesh_RenderThread(struct FRenderingCompositePassContext& Context, const FIntPoint& TextureSize) override;
 
 	virtual void UpdateScreenSettings(const FViewport* InViewport) override {}
@@ -98,7 +101,6 @@ public:
 	virtual void InitCanvasFromView(FSceneView* InView, UCanvas* Canvas) override;
 	virtual void RenderTexture_RenderThread(FRHICommandListImmediate& RHICmdList, FTexture2DRHIParamRef BackBuffer, FTexture2DRHIParamRef SrcTexture) const override;
 	virtual void GetEyeRenderParams_RenderThread(const FRenderingCompositePassContext& Context, FVector2D& EyeToSrcUVScaleValue, FVector2D& EyeToSrcUVOffsetValue) const override;
-
 	virtual void CalculateRenderTargetSize(const class FViewport& Viewport, uint32& InOutSizeX, uint32& InOutSizeY) override;
 	virtual bool NeedReAllocateViewportRenderTarget(const FViewport& Viewport) override;
 	virtual bool ShouldUseSeparateRenderTarget() const override
@@ -106,6 +108,7 @@ public:
 		check(IsInGameThread());
 		return IsStereoEnabled();
 	}
+	virtual void UpdateViewport(bool bUseSeparateRenderTarget, const FViewport& Viewport, SViewport*) override;
 
 	/** ISceneViewExtension interface */
 	virtual void SetupViewFamily(FSceneViewFamily& InViewFamily) override;
@@ -114,8 +117,6 @@ public:
 	virtual void PreRenderView_RenderThread(FRHICommandListImmediate& RHICmdList, FSceneView& InView) override;
 	virtual void PreRenderViewFamily_RenderThread(FRHICommandListImmediate& RHICmdList, FSceneViewFamily& InViewFamily) override;
 
-	virtual void UpdateViewport(bool bUseSeparateRenderTarget, const FViewport& Viewport, SViewport*) override;
-
 	class BridgeBaseImpl : public FRHICustomPresent
 	{
 	public:
@@ -296,6 +297,27 @@ private:
 		return FVector(InVector.v[0], InVector.v[1], InVector.v[2]);
 	}
 
+	struct FHiddenAreaMesh
+	{
+		FHiddenAreaMesh();
+		~FHiddenAreaMesh();
+
+		bool IsValid() const
+		{
+			return NumTriangles > 0;
+		}
+
+		void Build(const vr::HiddenAreaMesh_t& Mesh);
+
+		FVector4* pVertices;
+		uint16*   pIndices;
+		unsigned  NumVertices;
+		unsigned  NumIndices;
+		unsigned  NumTriangles;
+	};
+
+	FHiddenAreaMesh HiddenAreaMeshes[2];
+
 	/** Chaperone Support */
 	struct FChaperoneBounds
 	{
diff --git a/Engine/Plugins/Runtime/Steam/SteamVR/Source/SteamVR/Private/SteamVRRender.cpp b/Engine/Plugins/Runtime/Steam/SteamVR/Source/SteamVR/Private/SteamVRRender.cpp
--- a/Engine/Plugins/Runtime/Steam/SteamVR/Source/SteamVR/Private/SteamVRRender.cpp
+++ b/Engine/Plugins/Runtime/Steam/SteamVR/Source/SteamVR/Private/SteamVRRender.cpp
@@ -75,6 +75,28 @@ void FSteamVRHMD::RenderTexture_RenderThread(FRHICommandListImmediate& RHICmdLis
 	}
 }
 
+void FSteamVRHMD::DrawHiddenAreaMaskView_RenderThread(FRHICommandList& RHICmdList, const FViewInfo& View) const
+{
+	check(IsInRenderingThread());
+	check(View.StereoPass != eSSP_FULL);
+
+	const uint32 HiddenMeshIndex = (View.StereoPass == eSSP_LEFT_EYE) ? 0 : 1;
+	const FHiddenAreaMesh& Mesh = HiddenAreaMeshes[HiddenMeshIndex];
+	check(Mesh.IsValid());
+
+	DrawIndexedPrimitiveUP(
+		RHICmdList,
+		PT_TriangleList,
+		0,
+		Mesh.NumVertices,
+		Mesh.NumTriangles,
+		Mesh.pIndices,
+		sizeof(Mesh.pIndices[0]),
+		Mesh.pVertices,
+		sizeof(Mesh.pVertices[0])
+	);
+}
+
 #if PLATFORM_WINDOWS
 
 FSteamVRHMD::D3D11Bridge::D3D11Bridge(FSteamVRHMD* plugin):
@@ -93,6 +115,8 @@ void FSteamVRHMD::D3D11Bridge::BeginRendering()
 		Plugin->VRCompositor->SetGraphicsDevice(vr::Compositor_DeviceType_D3D11, RHIGetNativeDevice());
 		Inited = true;
 	}
+
+
 }
 
 void FSteamVRHMD::D3D11Bridge::FinishRendering()
diff --git a/Engine/Source/Runtime/HeadMountedDisplay/Private/HeadMountedDisplayModule.cpp b/Engine/Source/Runtime/HeadMountedDisplay/Private/HeadMountedDisplayModule.cpp
--- a/Engine/Source/Runtime/HeadMountedDisplay/Private/HeadMountedDisplayModule.cpp
+++ b/Engine/Source/Runtime/HeadMountedDisplay/Private/HeadMountedDisplayModule.cpp
@@ -2,6 +2,15 @@
 
 #include "HeadMountedDisplayPrivate.h"
 
+/**
+* HMD device console vars
+*/
+static TAutoConsoleVariable<int32> CVarHiddenAreaMask(
+	TEXT("vr.HiddenAreaMask"),
+	1,
+	TEXT("0 to disable hidden area mask, 1 to enable."),
+	ECVF_Scalability | ECVF_RenderThreadSafe);
+
 class FHeadMountedDisplayModule : public IHeadMountedDisplayModule
 {
 	virtual TSharedPtr< class IHeadMountedDisplay, ESPMode::ThreadSafe > CreateHeadMountedDisplay()
diff --git a/Engine/Source/Runtime/HeadMountedDisplay/Public/IHeadMountedDisplay.h b/Engine/Source/Runtime/HeadMountedDisplay/Public/IHeadMountedDisplay.h
--- a/Engine/Source/Runtime/HeadMountedDisplay/Public/IHeadMountedDisplay.h
+++ b/Engine/Source/Runtime/HeadMountedDisplay/Public/IHeadMountedDisplay.h
@@ -248,6 +248,16 @@ public:
 	 */
 	virtual FQuat GetBaseOrientation() const { return FQuat::Identity; }
 
+	/**
+	* @return true if a hidden area mask is available for the device.
+	*/
+	virtual bool HasHiddenAreaMask() const { return false; }
+
+	/**
+	* Optional method to draw a view's hidden area mask where supported.
+	* This can be used to avoid rendering pixels which are not included as input into the final distortion pass.
+	*/
+	virtual void DrawHiddenAreaMaskView_RenderThread(class FRHICommandList& RHICmdList, const FViewInfo& View) const {};
 
 	virtual void DrawDistortionMesh_RenderThread(struct FRenderingCompositePassContext& Context, const FIntPoint& TextureSize) {}
 
diff --git a/Engine/Source/Runtime/Renderer/Private/DeferredShadingRenderer.cpp b/Engine/Source/Runtime/Renderer/Private/DeferredShadingRenderer.cpp
--- a/Engine/Source/Runtime/Renderer/Private/DeferredShadingRenderer.cpp
+++ b/Engine/Source/Runtime/Renderer/Private/DeferredShadingRenderer.cpp
@@ -24,6 +24,7 @@
 #include "DistanceFieldAtlas.h"
 #include "../../Engine/Private/SkeletalRenderGPUSkin.h"		// GPrevPerBoneMotionBlur
 #include "EngineModule.h"
+#include "IHeadMountedDisplay.h"
 
 TAutoConsoleVariable<int32> CVarEarlyZPass(
 	TEXT("r.EarlyZPass"),
@@ -686,6 +687,19 @@ static FORCEINLINE bool NeedsPrePass(const FDeferredShadingSceneRenderer* Render
 		(Renderer->EarlyZPassMode != DDM_None || GEarlyZPassMovable != 0);
 }
 
+/**
+ * Returns true if there's a hidden area mask available
+ */
+static FORCEINLINE bool HasHiddenAreaMask()
+{
+	static const IConsoleVariable* const CVar = IConsoleManager::Get().FindConsoleVariable(TEXT("vr.HiddenAreaMask"));
+
+	return (CVar->GetInt() == 1 &&
+		GEngine &&
+		GEngine->HMDDevice.IsValid() &&
+		GEngine->HMDDevice->HasHiddenAreaMask());
+}
+
 static void SetAndClearViewGBuffer(FRHICommandListImmediate& RHICmdList, FViewInfo& View, bool bClearDepth)
 {
 	// if we didn't to the prepass above, then we will need to clear now, otherwise, it's already been cleared and rendered to
@@ -929,6 +943,11 @@ void FDeferredShadingSceneRenderer::Render(FRHICommandListImmediate& RHICmdList)
 		// at this point, the depth was cleared
 		bDepthWasCleared = true;
 	}
+	else if (HasHiddenAreaMask())
+	{
+		bDepthWasCleared = RenderPrePassHMD(RHICmdList);
+	}
+
 	RHICmdList.SetCurrentStat(GET_STATID(STAT_CLM_AfterPrePass));
 	ServiceLocalQueue();
 	//occlusion can't run before basepass if there's no prepass to fill in some depth to occlude against.
@@ -1381,18 +1400,34 @@ static void SetupPrePassView(FRHICommandList& RHICmdList, const FIntRect& ViewRe
 {
 	// Disable color writes, enable depth tests and writes.
 	RHICmdList.SetBlendState(TStaticBlendState<CW_NONE>::GetRHI());
-	RHICmdList.SetDepthStencilState(TStaticDepthStencilState<true,CF_DepthNearOrEqual>::GetRHI());
-	RHICmdList.SetViewport(ViewRect.Min.X, ViewRect.Min.Y, 0, ViewRect.Max.X, ViewRect.Max.Y, 1);
+	RHICmdList.SetDepthStencilState(TStaticDepthStencilState<true, CF_DepthNearOrEqual>::GetRHI());
+	RHICmdList.SetViewport(ViewRect.Min.X, ViewRect.Min.Y, 0.0f, ViewRect.Max.X, ViewRect.Max.Y, 1.0f);
 	RHICmdList.SetRasterizerState(TStaticRasterizerState<FM_Solid, CM_None>::GetRHI());
 	RHICmdList.SetScissorRect(false, 0, 0, 0, 0);
 }
 
+static void RenderHiddenAreaMaskView(FRHICommandList& RHICmdList, const FViewInfo& View)
+{
+	const auto FeatureLevel = GMaxRHIFeatureLevel;
+	const auto ShaderMap = GetGlobalShaderMap(FeatureLevel);
+	TShaderMapRef<TOneColorVS<true> > VertexShader(ShaderMap);
+	static FGlobalBoundShaderState BoundShaderState;
+	SetGlobalBoundShaderState(RHICmdList, FeatureLevel, BoundShaderState, GetVertexDeclarationFVector4(), *VertexShader, nullptr);
+	GEngine->HMDDevice->DrawHiddenAreaMaskView_RenderThread(RHICmdList, View);
+}
+
 bool FDeferredShadingSceneRenderer::RenderPrePassView(FRHICommandList& RHICmdList, const FViewInfo& View)
 {
 	bool bDirty = false;
 
 	SetupPrePassView(RHICmdList, View.ViewRect);
 
+	if (HasHiddenAreaMask() && View.StereoPass != eSSP_FULL)
+	{
+		RenderHiddenAreaMaskView(RHICmdList, View);
+		bDirty = true;
+	}
+	
 	// Draw the static occluder primitives using a depth drawing policy.
 	{
 		// Draw opaque occluders which support a separate position-only
@@ -1474,6 +1509,11 @@ public:
 	{
 		FSceneRenderTargets::Get(CmdList).BeginRenderingPrePass(CmdList, false);
 		SetupPrePassView(CmdList, View.ViewRect);
+
+		if (HasHiddenAreaMask() && View.StereoPass != eSSP_FULL)
+		{
+			RenderHiddenAreaMaskView(CmdList, View);
+		}
 	}
 };
 
@@ -1575,6 +1615,36 @@ bool FDeferredShadingSceneRenderer::RenderPrePass(FRHICommandListImmediate& RHIC
 	return bDirty;
 }
 
+bool FDeferredShadingSceneRenderer::RenderPrePassHMD(FRHICommandListImmediate& RHICmdList)
+{
+
+	// Early out before we change any state if there's not a mask to render
+	if (!HasHiddenAreaMask())
+	{
+		return false;
+	}
+
+	SCOPED_DRAW_EVENT(RHICmdList, PrePassHMD);
+	SCOPE_CYCLE_COUNTER(STAT_DepthDrawTime);
+
+	FSceneRenderTargets& SceneContext = FSceneRenderTargets::Get(RHICmdList);
+	SceneContext.BeginRenderingPrePass(RHICmdList, true);
+
+	for (int32 ViewIndex = 0; ViewIndex < Views.Num(); ++ViewIndex)
+	{
+		const FViewInfo& View = Views[ViewIndex];
+		if (View.StereoPass != eSSP_FULL)
+		{
+			SetupPrePassView(RHICmdList, View.ViewRect);
+			RenderHiddenAreaMaskView(RHICmdList, View);
+		}
+	}
+
+	SceneContext.FinishRenderingPrePass(RHICmdList);
+
+	return true;
+}
+
 static TAutoConsoleVariable<int32> CVarParallelBasePass(
 	TEXT("r.ParallelBasePass"),
 	1,
diff --git a/Engine/Source/Runtime/Renderer/Private/DeferredShadingRenderer.h b/Engine/Source/Runtime/Renderer/Private/DeferredShadingRenderer.h
--- a/Engine/Source/Runtime/Renderer/Private/DeferredShadingRenderer.h
+++ b/Engine/Source/Runtime/Renderer/Private/DeferredShadingRenderer.h
@@ -163,6 +163,12 @@ private:
 	 */
 	bool RenderPrePass(FRHICommandListImmediate& RHICmdList);
 
+	/**
+	 * Renders the active HMD's hidden area mask as a depth prepass, if available.
+	 * @return true if depth is cleared
+	 */
+	bool RenderPrePassHMD(FRHICommandListImmediate& RHICmdList);
+
 	/** Issues occlusion queries. */
 	void BeginOcclusionTests(FRHICommandListImmediate& RHICmdList, bool bRenderQueries, bool bRenderHZB);
 
