Fix PS4 rendering when there are no dbuffer decals

diff --git a/Engine/Source/Runtime/Renderer/Private/CompositionLighting/PostProcessDeferredDecals.cpp b/Engine/Source/Runtime/Renderer/Private/CompositionLighting/PostProcessDeferredDecals.cpp
--- a/Engine/Source/Runtime/Renderer/Private/CompositionLighting/PostProcessDeferredDecals.cpp
+++ b/Engine/Source/Runtime/Renderer/Private/CompositionLighting/PostProcessDeferredDecals.cpp
@@ -626,204 +626,203 @@ void FRCPassPostProcessDeferredDecals::Process(FRenderingCompositePassContext& C
 	const FSceneViewFamily& ViewFamily = *(View.Family);
 	FScene& Scene = *(FScene*)ViewFamily.Scene;
 
-	if(!Scene.Decals.Num() || !ViewFamily.EngineShowFlags.Decals)
+	//don't early return.  Resolves must be run for fast clears to work.
+	bool bRenderDecal = Scene.Decals.Num() && ViewFamily.EngineShowFlags.Decals;
+	if (bRenderDecal)
 	{
-		// to avoid the stats showing up
-		return;
-	}
-
-	// Build a list of decals that need to be rendered for this view
-	FTransientDecalRenderDataList SortedDecals;
-	FDecalRendering::BuildVisibleDecalList(Scene, View, DecalRenderStage, SortedDecals);
-	
-	if (SortedDecals.Num() > 0)
-	{
-		FIntRect SrcRect = View.ViewRect;
-		FIntRect DestRect = View.ViewRect;
-		
-		bool bStencilDecalsInThisStage = true;
+		// Build a list of decals that need to be rendered for this view
+		FTransientDecalRenderDataList SortedDecals;
+		FDecalRendering::BuildVisibleDecalList(Scene, View, DecalRenderStage, SortedDecals);
 
-#if DBUFFER_DONT_USE_STENCIL_YET
-		if(DecalRenderStage != DRS_BeforeLighting)
+		if (SortedDecals.Num() > 0)
 		{
-			bStencilDecalsInThisStage = false;
-		}
-#endif
+			FIntRect SrcRect = View.ViewRect;
+			FIntRect DestRect = View.ViewRect;
 
-		// Setup a stencil mask to prevent certain pixels from receiving deferred decals
-		if(bStencilDecalsInThisStage)
-		{
-			StencilDecalMask(RHICmdList, View);
-		}
-		
-		// optimization to have less state changes
-		EDecalRasterizerState LastDecalRasterizerState = DRS_Undefined;
-		FDecalDepthState LastDecalDepthState;
-		int32 LastDecalBlendMode = -1;
-		int32 LastDecalHasNormal = -1; // Decal state can change based on its normal property.(SM5)
-		FDecalRendering::ERenderTargetMode LastRenderTargetMode = FDecalRendering::RTM_Unknown;
-		const ERHIFeatureLevel::Type SMFeatureLevel = Context.GetFeatureLevel();
-
-		SCOPED_DRAW_EVENT(RHICmdList, Decals);
-		INC_DWORD_STAT_BY(STAT_Decals, SortedDecals.Num());
-		
-		for (int32 DecalIndex = 0, DecalCount = SortedDecals.Num(); DecalIndex < DecalCount; DecalIndex++)
-		{
-			const FTransientDecalRenderData& DecalData = SortedDecals[DecalIndex];
-			const FDeferredDecalProxy& DecalProxy = *DecalData.DecalProxy;
-			const FMatrix ComponentToWorldMatrix = DecalProxy.ComponentTrans.ToMatrixWithScale();
-			const FMatrix FrustumComponentToClip = FDecalRendering::ComputeComponentToClipMatrix(View, ComponentToWorldMatrix);
+			bool bStencilDecalsInThisStage = true;
 
-			EDecalBlendMode DecalBlendMode = DecalData.DecalBlendMode;
-			bool bStencilThisDecal = bStencilDecalsInThisStage;
-			
 #if DBUFFER_DONT_USE_STENCIL_YET
-			if(FDecalRendering::ComputeRenderStage(View.GetShaderPlatform(), DecalBlendMode) != DRS_BeforeLighting)
+			if (DecalRenderStage != DRS_BeforeLighting)
 			{
-				bStencilThisDecal = false;
+				bStencilDecalsInThisStage = false;
 			}
-#endif				
-
-			FDecalRendering::ERenderTargetMode CurrentRenderTargetMode = FDecalRendering::ComputeRenderTargetMode(View.GetShaderPlatform(), DecalBlendMode);
+#endif
 
-			if (bShaderComplexity)
+			// Setup a stencil mask to prevent certain pixels from receiving deferred decals
+			if (bStencilDecalsInThisStage)
 			{
-				CurrentRenderTargetMode = FDecalRendering::RTM_SceneColor;
-				// we want additive blending for the ShaderComplexity mode
-				DecalBlendMode = DBM_Emissive;
+				StencilDecalMask(RHICmdList, View);
 			}
 
-			// fewer rendertarget switches if possible
-			if (CurrentRenderTargetMode != LastRenderTargetMode)
+			// optimization to have less state changes
+			EDecalRasterizerState LastDecalRasterizerState = DRS_Undefined;
+			FDecalDepthState LastDecalDepthState;
+			int32 LastDecalBlendMode = -1;
+			int32 LastDecalHasNormal = -1; // Decal state can change based on its normal property.(SM5)
+			FDecalRendering::ERenderTargetMode LastRenderTargetMode = FDecalRendering::RTM_Unknown;
+			const ERHIFeatureLevel::Type SMFeatureLevel = Context.GetFeatureLevel();
+
+			SCOPED_DRAW_EVENT(RHICmdList, Decals);
+			INC_DWORD_STAT_BY(STAT_Decals, SortedDecals.Num());
+
+			for (int32 DecalIndex = 0, DecalCount = SortedDecals.Num(); DecalIndex < DecalCount; DecalIndex++)
 			{
-				LastRenderTargetMode = CurrentRenderTargetMode;
+				const FTransientDecalRenderData& DecalData = SortedDecals[DecalIndex];
+				const FDeferredDecalProxy& DecalProxy = *DecalData.DecalProxy;
+				const FMatrix ComponentToWorldMatrix = DecalProxy.ComponentTrans.ToMatrixWithScale();
+				const FMatrix FrustumComponentToClip = FDecalRendering::ComputeComponentToClipMatrix(View, ComponentToWorldMatrix);
+
+				EDecalBlendMode DecalBlendMode = DecalData.DecalBlendMode;
+				bool bStencilThisDecal = bStencilDecalsInThisStage;
+
+#if DBUFFER_DONT_USE_STENCIL_YET
+				if (FDecalRendering::ComputeRenderStage(View.GetShaderPlatform(), DecalBlendMode) != DRS_BeforeLighting)
+				{
+					bStencilThisDecal = false;
+				}
+#endif				
+
+				FDecalRendering::ERenderTargetMode CurrentRenderTargetMode = FDecalRendering::ComputeRenderTargetMode(View.GetShaderPlatform(), DecalBlendMode);
+
+				if (bShaderComplexity)
+				{
+					CurrentRenderTargetMode = FDecalRendering::RTM_SceneColor;
+					// we want additive blending for the ShaderComplexity mode
+					DecalBlendMode = DBM_Emissive;
+				}
 
-				switch (CurrentRenderTargetMode)
+				// fewer rendertarget switches if possible
+				if (CurrentRenderTargetMode != LastRenderTargetMode)
 				{
+					LastRenderTargetMode = CurrentRenderTargetMode;
+
+					switch (CurrentRenderTargetMode)
+					{
 					case FDecalRendering::RTM_SceneColorAndGBuffer:
-						{							
-							TargetsToResolve[SceneColorIndex] = SceneContext.GetSceneColor()->GetRenderTargetItem().TargetableTexture;
-							TargetsToResolve[GBufferAIndex] = SceneContext.GBufferA->GetRenderTargetItem().TargetableTexture;
-							TargetsToResolve[GBufferBIndex] = SceneContext.GBufferB->GetRenderTargetItem().TargetableTexture;
-							TargetsToResolve[GBufferCIndex] = SceneContext.GBufferC->GetRenderTargetItem().TargetableTexture;
-							
-							SetRenderTargets(RHICmdList, 4, TargetsToResolve, SceneContext.GetSceneDepthSurface(), ESimpleRenderTargetMode::EExistingColorAndDepth, FExclusiveDepthStencil::DepthRead_StencilWrite);
-						}
+					{
+																	  TargetsToResolve[SceneColorIndex] = SceneContext.GetSceneColor()->GetRenderTargetItem().TargetableTexture;
+																	  TargetsToResolve[GBufferAIndex] = SceneContext.GBufferA->GetRenderTargetItem().TargetableTexture;
+																	  TargetsToResolve[GBufferBIndex] = SceneContext.GBufferB->GetRenderTargetItem().TargetableTexture;
+																	  TargetsToResolve[GBufferCIndex] = SceneContext.GBufferC->GetRenderTargetItem().TargetableTexture;
+
+																	  SetRenderTargets(RHICmdList, 4, TargetsToResolve, SceneContext.GetSceneDepthSurface(), ESimpleRenderTargetMode::EExistingColorAndDepth, FExclusiveDepthStencil::DepthRead_StencilWrite);
+					}
 						break;
 
 					case FDecalRendering::RTM_SceneColorAndGBufferDepthWrite:
-						{							
-							TargetsToResolve[SceneColorIndex] = SceneContext.GetSceneColor()->GetRenderTargetItem().TargetableTexture;
-							TargetsToResolve[GBufferAIndex] = SceneContext.GBufferA->GetRenderTargetItem().TargetableTexture;
-							TargetsToResolve[GBufferBIndex] = SceneContext.GBufferB->GetRenderTargetItem().TargetableTexture;
-							TargetsToResolve[GBufferCIndex] = SceneContext.GBufferC->GetRenderTargetItem().TargetableTexture;
-							
-							SetRenderTargets(RHICmdList, 4, TargetsToResolve, SceneContext.GetSceneDepthSurface(), ESimpleRenderTargetMode::EExistingColorAndDepth, FExclusiveDepthStencil::DepthWrite_StencilWrite);
-						}
+					{
+																				TargetsToResolve[SceneColorIndex] = SceneContext.GetSceneColor()->GetRenderTargetItem().TargetableTexture;
+																				TargetsToResolve[GBufferAIndex] = SceneContext.GBufferA->GetRenderTargetItem().TargetableTexture;
+																				TargetsToResolve[GBufferBIndex] = SceneContext.GBufferB->GetRenderTargetItem().TargetableTexture;
+																				TargetsToResolve[GBufferCIndex] = SceneContext.GBufferC->GetRenderTargetItem().TargetableTexture;
+
+																				SetRenderTargets(RHICmdList, 4, TargetsToResolve, SceneContext.GetSceneDepthSurface(), ESimpleRenderTargetMode::EExistingColorAndDepth, FExclusiveDepthStencil::DepthWrite_StencilWrite);
+					}
 						break;
 
 					case FDecalRendering::RTM_GBufferNormal:
 						TargetsToResolve[GBufferAIndex] = SceneContext.GBufferA->GetRenderTargetItem().TargetableTexture;
 						SetRenderTarget(RHICmdList, TargetsToResolve[GBufferAIndex], SceneContext.GetSceneDepthSurface(), ESimpleRenderTargetMode::EExistingColorAndDepth, FExclusiveDepthStencil::DepthRead_StencilWrite);
 						break;
-					
+
 					case FDecalRendering::RTM_SceneColor:
 						TargetsToResolve[SceneColorIndex] = SceneContext.GetSceneColor()->GetRenderTargetItem().TargetableTexture;
 						SetRenderTarget(RHICmdList, TargetsToResolve[SceneColorIndex], SceneContext.GetSceneDepthSurface(), ESimpleRenderTargetMode::EExistingColorAndDepth, FExclusiveDepthStencil::DepthRead_StencilWrite);
 						break;
 
 					case FDecalRendering::RTM_DBuffer:
-						{							
-							TargetsToResolve[DBufferAIndex] = SceneContext.DBufferA->GetRenderTargetItem().TargetableTexture;
-							TargetsToResolve[DBufferBIndex] = SceneContext.DBufferB->GetRenderTargetItem().TargetableTexture;
-							TargetsToResolve[DBufferCIndex] = SceneContext.DBufferC->GetRenderTargetItem().TargetableTexture;
-							SetRenderTargets(RHICmdList, 3, &TargetsToResolve[DBufferAIndex], SceneContext.GetSceneDepthSurface(), ESimpleRenderTargetMode::EExistingColorAndDepth, FExclusiveDepthStencil::DepthRead_StencilWrite);
-						}
+					{
+														 TargetsToResolve[DBufferAIndex] = SceneContext.DBufferA->GetRenderTargetItem().TargetableTexture;
+														 TargetsToResolve[DBufferBIndex] = SceneContext.DBufferB->GetRenderTargetItem().TargetableTexture;
+														 TargetsToResolve[DBufferCIndex] = SceneContext.DBufferC->GetRenderTargetItem().TargetableTexture;
+														 SetRenderTargets(RHICmdList, 3, &TargetsToResolve[DBufferAIndex], SceneContext.GetSceneDepthSurface(), ESimpleRenderTargetMode::EExistingColorAndDepth, FExclusiveDepthStencil::DepthRead_StencilWrite);
+					}
 						break;
 
 					default:
-						check(0);	
+						check(0);
 						break;
-				}
-				Context.SetViewportAndCallRHI(DestRect);
+					}
+					Context.SetViewportAndCallRHI(DestRect);
 
-				// we need to reset the stream source after any call to SetRenderTarget (at least for Metal, which doesn't queue up VB assignments)
-				RHICmdList.SetStreamSource(0, GetUnitCubeVertexBuffer(), sizeof(FVector4), 0);
-			}
+					// we need to reset the stream source after any call to SetRenderTarget (at least for Metal, which doesn't queue up VB assignments)
+					RHICmdList.SetStreamSource(0, GetUnitCubeVertexBuffer(), sizeof(FVector4), 0);
+				}
 
-			bool bThisDecalUsesStencil = false;
+				bool bThisDecalUsesStencil = false;
 
-			if (bStencilThisDecal)
-			{
-				if (bStencilSizeThreshold)
+				if (bStencilThisDecal)
 				{
-					// note this is after a SetStreamSource (in if CurrentRenderTargetMode != LastRenderTargetMode) call as it needs to get the VB input
-					bThisDecalUsesStencil = RenderPreStencil(Context, ComponentToWorldMatrix, FrustumComponentToClip);
-
-					LastDecalRasterizerState = DRS_Undefined;
-					LastDecalDepthState = FDecalDepthState();
-					LastDecalBlendMode = -1;
+					if (bStencilSizeThreshold)
+					{
+						// note this is after a SetStreamSource (in if CurrentRenderTargetMode != LastRenderTargetMode) call as it needs to get the VB input
+						bThisDecalUsesStencil = RenderPreStencil(Context, ComponentToWorldMatrix, FrustumComponentToClip);
+
+						LastDecalRasterizerState = DRS_Undefined;
+						LastDecalDepthState = FDecalDepthState();
+						LastDecalBlendMode = -1;
+					}
 				}
-			}
 
-			const bool bBlendStateChange = DecalBlendMode != LastDecalBlendMode;// Has decal mode changed.
-			const bool bDecalNormalChanged = GSupportsSeparateRenderTargetBlendState && // has normal changed for SM5 stain/translucent decals?
-							(DecalBlendMode == DBM_Translucent || DecalBlendMode == DBM_Stain) &&
-							(int32)DecalData.bHasNormal != LastDecalHasNormal;
+				const bool bBlendStateChange = DecalBlendMode != LastDecalBlendMode;// Has decal mode changed.
+				const bool bDecalNormalChanged = GSupportsSeparateRenderTargetBlendState && // has normal changed for SM5 stain/translucent decals?
+					(DecalBlendMode == DBM_Translucent || DecalBlendMode == DBM_Stain) &&
+					(int32)DecalData.bHasNormal != LastDecalHasNormal;
 
-			// fewer blend state changes if possible
-			if (bBlendStateChange || bDecalNormalChanged)
-			{
-				LastDecalBlendMode = DecalBlendMode;
-				LastDecalHasNormal = (int32)DecalData.bHasNormal;
+				// fewer blend state changes if possible
+				if (bBlendStateChange || bDecalNormalChanged)
+				{
+					LastDecalBlendMode = DecalBlendMode;
+					LastDecalHasNormal = (int32)DecalData.bHasNormal;
 
-				SetDecalBlendState(RHICmdList, SMFeatureLevel, DecalRenderStage, (EDecalBlendMode)LastDecalBlendMode, DecalData.bHasNormal);
-			}
+					SetDecalBlendState(RHICmdList, SMFeatureLevel, DecalRenderStage, (EDecalBlendMode)LastDecalBlendMode, DecalData.bHasNormal);
+				}
 
-			
-			// todo
-			const float ConservativeRadius = DecalData.ConservativeRadius;
-//			const int32 IsInsideDecal = ((FVector)View.ViewMatrices.ViewOrigin - ComponentToWorldMatrix.GetOrigin()).SizeSquared() < FMath::Square(ConservativeRadius * 1.05f + View.NearClippingDistance * 2.0f) + ( bThisDecalUsesStencil ) ? 2 : 0;
-			const bool bInsideDecal = ((FVector)View.ViewMatrices.ViewOrigin - ComponentToWorldMatrix.GetOrigin()).SizeSquared() < FMath::Square(ConservativeRadius * 1.05f + View.NearClippingDistance * 2.0f);
-//			const bool bInsideDecal =  !(IsInsideDecal & 1);
 
-			// update rasterizer state if needed
-			{
-				EDecalRasterizerState DecalRasterizerState = ComputeDecalRasterizerState(bInsideDecal, View);
+				// todo
+				const float ConservativeRadius = DecalData.ConservativeRadius;
+				//			const int32 IsInsideDecal = ((FVector)View.ViewMatrices.ViewOrigin - ComponentToWorldMatrix.GetOrigin()).SizeSquared() < FMath::Square(ConservativeRadius * 1.05f + View.NearClippingDistance * 2.0f) + ( bThisDecalUsesStencil ) ? 2 : 0;
+				const bool bInsideDecal = ((FVector)View.ViewMatrices.ViewOrigin - ComponentToWorldMatrix.GetOrigin()).SizeSquared() < FMath::Square(ConservativeRadius * 1.05f + View.NearClippingDistance * 2.0f);
+				//			const bool bInsideDecal =  !(IsInsideDecal & 1);
 
-				if(LastDecalRasterizerState != DecalRasterizerState)
+				// update rasterizer state if needed
 				{
-					LastDecalRasterizerState = DecalRasterizerState;
-					SetDecalRasterizerState(DecalRasterizerState, RHICmdList);
-				}
-			}
+					EDecalRasterizerState DecalRasterizerState = ComputeDecalRasterizerState(bInsideDecal, View);
 
-			// update DepthStencil state if needed
-			{
-				FDecalDepthState DecalDepthState = ComputeDecalDepthState(DecalBlendMode, bInsideDecal, bStencilDecalsInThisStage, bThisDecalUsesStencil);
+					if (LastDecalRasterizerState != DecalRasterizerState)
+					{
+						LastDecalRasterizerState = DecalRasterizerState;
+						SetDecalRasterizerState(DecalRasterizerState, RHICmdList);
+					}
+				}
 
-				if(LastDecalDepthState != DecalDepthState)
+				// update DepthStencil state if needed
 				{
-					LastDecalDepthState = DecalDepthState;
-					SetDecalDepthState(DecalDepthState, RHICmdList);
+					FDecalDepthState DecalDepthState = ComputeDecalDepthState(DecalBlendMode, bInsideDecal, bStencilDecalsInThisStage, bThisDecalUsesStencil);
+
+					if (LastDecalDepthState != DecalDepthState)
+					{
+						LastDecalDepthState = DecalDepthState;
+						SetDecalDepthState(DecalDepthState, RHICmdList);
+					}
 				}
-			}
 
-			FDecalRendering::SetShader(RHICmdList, View, bShaderComplexity, DecalData, FrustumComponentToClip);
+				FDecalRendering::SetShader(RHICmdList, View, bShaderComplexity, DecalData, FrustumComponentToClip);
 
-			RHICmdList.DrawIndexedPrimitive(GetUnitCubeIndexBuffer(), PT_TriangleList, 0, 0, 8, 0, ARRAY_COUNT(GCubeIndices) / 3, 1);
-		}
+				RHICmdList.DrawIndexedPrimitive(GetUnitCubeIndexBuffer(), PT_TriangleList, 0, 0, 8, 0, ARRAY_COUNT(GCubeIndices) / 3, 1);
+			}
 
-		// we don't modify stencil but if out input was having stencil for us (after base pass - we need to clear)
-		// Clear stencil to 0, which is the assumed default by other passes
-		RHICmdList.Clear(false, FLinearColor::White, false, (float)ERHIZBuffer::FarPlane, true, 0, FIntRect());
-		
-		if(DecalRenderStage == DRS_BeforeBasePass)
-		{
-			// before BasePass
-			GRenderTargetPool.VisualizeTexture.SetCheckPoint(RHICmdList, SceneContext.DBufferA);
-			GRenderTargetPool.VisualizeTexture.SetCheckPoint(RHICmdList, SceneContext.DBufferB);
-			GRenderTargetPool.VisualizeTexture.SetCheckPoint(RHICmdList, SceneContext.DBufferC);
+			// we don't modify stencil but if out input was having stencil for us (after base pass - we need to clear)
+			// Clear stencil to 0, which is the assumed default by other passes
+			RHICmdList.Clear(false, FLinearColor::White, false, (float)ERHIZBuffer::FarPlane, true, 0, FIntRect());
+
+			if (DecalRenderStage == DRS_BeforeBasePass)
+			{
+				// before BasePass
+				GRenderTargetPool.VisualizeTexture.SetCheckPoint(RHICmdList, SceneContext.DBufferA);
+				GRenderTargetPool.VisualizeTexture.SetCheckPoint(RHICmdList, SceneContext.DBufferB);
+				GRenderTargetPool.VisualizeTexture.SetCheckPoint(RHICmdList, SceneContext.DBufferC);
+			}
 		}
 	}
 
