Properly mark Distortion rendering as not requiring depth writes.

diff --git a/Engine/Source/Runtime/Renderer/Private/DistortionRendering.cpp b/Engine/Source/Runtime/Renderer/Private/DistortionRendering.cpp
--- a/Engine/Source/Runtime/Renderer/Private/DistortionRendering.cpp
+++ b/Engine/Source/Runtime/Renderer/Private/DistortionRendering.cpp
@@ -891,7 +891,7 @@ void FSceneRenderer::RenderDistortion(FRHICommandListImmediate& RHICmdList)
 // OCULUS END
 
 		// Apply distortion and store off-screen
-		SetRenderTarget(RHICmdList, DestRenderTarget.TargetableTexture, SceneContext.GetSceneDepthSurface());
+		SetRenderTarget(RHICmdList, DestRenderTarget.TargetableTexture, SceneContext.GetSceneDepthSurface(), ESimpleRenderTargetMode::EExistingColorAndDepth, FExclusiveDepthStencil::DepthRead_StencilWrite);
 
 		for(int32 ViewIndex = 0, Num = Views.Num(); ViewIndex < Num; ++ViewIndex)
 		{
@@ -935,7 +935,7 @@ void FSceneRenderer::RenderDistortion(FRHICommandListImmediate& RHICmdList)
 
 		// Resolve and merge distortion to main scene
 		RHICmdList.CopyToResolveTarget(DestRenderTarget.TargetableTexture, DestRenderTarget.ShaderResourceTexture, false, FResolveParams());
-		SetRenderTarget(RHICmdList, SceneContext.GetSceneColor()->GetRenderTargetItem().TargetableTexture, SceneContext.GetSceneDepthSurface());
+		SetRenderTarget(RHICmdList, SceneContext.GetSceneColor()->GetRenderTargetItem().TargetableTexture, SceneContext.GetSceneDepthSurface(), ESimpleRenderTargetMode::EExistingColorAndDepth, FExclusiveDepthStencil::DepthRead_StencilWrite);
 
 		for(int32 ViewIndex = 0, Num = Views.Num(); ViewIndex < Num; ++ViewIndex)
 		{
