Removed some unreachable RT resizing code.

diff --git a/Engine/Source/Runtime/Renderer/Private/PostProcess/SceneRenderTargets.cpp b/Engine/Source/Runtime/Renderer/Private/PostProcess/SceneRenderTargets.cpp
--- a/Engine/Source/Runtime/Renderer/Private/PostProcess/SceneRenderTargets.cpp
+++ b/Engine/Source/Runtime/Renderer/Private/PostProcess/SceneRenderTargets.cpp
@@ -270,8 +270,7 @@ inline const TCHAR* GetSceneColorTargetName(FSceneRenderTargets::EShadingPath Sh
 
 FIntPoint FSceneRenderTargets::ComputeDesiredSize(const FSceneViewFamily& ViewFamily)
 {
-	// Don't expose Clamped to the cvar since you need to at least grow to the initial state.
-	enum ESizingMethods { RequestedSize, ScreenRes, Grow, VisibleSizingMethodsCount, Clamped };
+	enum ESizingMethods { RequestedSize, ScreenRes, Grow, VisibleSizingMethodsCount };
 	ESizingMethods SceneTargetsSizingMethod = Grow;
 
 	bool bIsSceneCapture = false;
@@ -318,14 +317,6 @@ FIntPoint FSceneRenderTargets::ComputeDesiredSize(const FSceneViewFamily& ViewFa
 					FMath::Max((uint32)GetBufferSizeXY().Y, ViewFamily.FamilySizeY));
 			break;
 
-		case Clamped:
-			if (((uint32)BufferSize.X < ViewFamily.FamilySizeX) || ((uint32)BufferSize.Y < ViewFamily.FamilySizeY))
-			{
-				UE_LOG(LogRenderer, Warning, TEXT("Capture target size: %ux%u clamped to %ux%u."), ViewFamily.FamilySizeX, ViewFamily.FamilySizeY, BufferSize.X, BufferSize.Y);
-			}
-			DesiredBufferSize = FIntPoint(GetBufferSizeXY().X, GetBufferSizeXY().Y);
-			break;
-
 		default:
 			checkNoEntry();
 	}
