Fixed view rotation matrix having translation.
Fixed temporal aa sample position not being set.

diff --git a/Engine/Source/Editor/UnrealEd/Private/EditorViewportClient.cpp b/Engine/Source/Editor/UnrealEd/Private/EditorViewportClient.cpp
--- a/Engine/Source/Editor/UnrealEd/Private/EditorViewportClient.cpp
+++ b/Engine/Source/Editor/UnrealEd/Private/EditorViewportClient.cpp
@@ -656,7 +656,7 @@ FSceneView* FEditorViewportClient::CalcSceneView(FSceneViewFamily* ViewFamily)
 		{
 			if (bUsingOrbitCamera)
 			{
-				ViewInitOptions.ViewRotationMatrix = FTranslationMatrix(ViewLocation) * ViewTransform.ComputeOrbitMatrix();
+				ViewInitOptions.ViewRotationMatrix = ViewTransform.ComputeOrbitMatrix().RemoveTranslation();
 			}
 			else
 			{
@@ -807,6 +807,9 @@ FSceneView* FEditorViewportClient::CalcSceneView(FSceneViewFamily* ViewFamily)
 				check(false);
 			}
 
+			ViewInitOptions.ViewOrigin += ViewInitOptions.ViewRotationMatrix.InverseTransformPosition( FVector::ZeroVector );
+			ViewInitOptions.ViewRotationMatrix = ViewInitOptions.ViewRotationMatrix.RemoveTranslation();
+
 			ViewInitOptions.ProjectionMatrix = FReversedZOrthoMatrix(
 				OrthoWidth,
 				OrthoHeight,
diff --git a/Engine/Source/Editor/UnrealEd/Private/ThumbnailHelpers.cpp b/Engine/Source/Editor/UnrealEd/Private/ThumbnailHelpers.cpp
--- a/Engine/Source/Editor/UnrealEd/Private/ThumbnailHelpers.cpp
+++ b/Engine/Source/Editor/UnrealEd/Private/ThumbnailHelpers.cpp
@@ -103,6 +103,9 @@ void FThumbnailPreviewScene::GetView(FSceneViewFamily* ViewFamily, int32 X, int3
 			FPlane(0,	1,	0,	0),
 			FPlane(0,	0,	0,	1));
 
+		Origin -= ViewRotationMatrix.InverseTransformPosition( FVector::ZeroVector );
+		ViewRotationMatrix = ViewRotationMatrix.RemoveTranslation();
+
 		FSceneViewInitOptions ViewInitOptions;
 		ViewInitOptions.ViewFamily = ViewFamily;
 		ViewInitOptions.SetViewRectangle(ViewRect);
diff --git a/Engine/Source/Runtime/Engine/Private/SceneView.cpp b/Engine/Source/Runtime/Engine/Private/SceneView.cpp
--- a/Engine/Source/Runtime/Engine/Private/SceneView.cpp
+++ b/Engine/Source/Runtime/Engine/Private/SceneView.cpp
@@ -302,6 +302,7 @@ FSceneView::FSceneView(const FSceneViewInitOptions& InitOptions)
 	check(UnscaledViewRect.Min.Y >= 0);
 	check(UnscaledViewRect.Width() > 0);
 	check(UnscaledViewRect.Height() > 0);
+	check(InitOptions.ViewRotationMatrix.GetOrigin().IsNearlyZero());
 
 	ViewMatrices.ViewMatrix = FTranslationMatrix(-InitOptions.ViewOrigin) * InitOptions.ViewRotationMatrix;
 
@@ -342,7 +343,7 @@ FSceneView::FSceneView(const FSceneViewInitOptions& InitOptions)
 
 	/** The view transform, starting from world-space points translated by -ViewOrigin. */
 	FMatrix TranslatedViewMatrix = InitOptions.ViewRotationMatrix;
-	FMatrix InvTranslatedViewMatrix = TranslatedViewMatrix.Inverse();
+	FMatrix InvTranslatedViewMatrix = TranslatedViewMatrix.GetTransposed();
 
 	// Translate world-space so its origin is at ViewOrigin for improved precision.
 	// Note that this isn't exactly right for orthogonal projections (See the above special case), but we still use ViewOrigin
@@ -382,6 +383,7 @@ FSceneView::FSceneView(const FSceneViewInitOptions& InitOptions)
 	{
 		TranslatedViewMatrix = FTranslationMatrix(-ViewMatrices.PreViewTranslation)
 			* FTranslationMatrix(-InitOptions.ViewOrigin) * InitOptions.ViewRotationMatrix;
+		InvTranslatedViewMatrix = TranslatedViewMatrix.Inverse();
 	}
 	
 	// Compute a transform from view origin centered world-space to clip space.
diff --git a/Engine/Source/Runtime/Renderer/Private/PostProcess/PostProcessTemporalAA.cpp b/Engine/Source/Runtime/Renderer/Private/PostProcess/PostProcessTemporalAA.cpp
--- a/Engine/Source/Runtime/Renderer/Private/PostProcess/PostProcessTemporalAA.cpp
+++ b/Engine/Source/Runtime/Renderer/Private/PostProcess/PostProcessTemporalAA.cpp
@@ -45,6 +45,15 @@ static TAutoConsoleVariable<int32> CVarTemporalAAPauseCorrect(
 	TEXT("Correct temporal AA in pause. This holds onto render targets longer preventing reuse and consumes more memory."),
 	ECVF_RenderThreadSafe);
 
+static float CatmullRom( float x )
+{
+	float ax = FMath::Abs(x);
+	if( ax > 1.0f )
+		return ( ( -0.5f * ax + 2.5f ) * ax - 4.0f ) *ax + 2.0f;
+	else
+		return ( 1.5f * ax - 2.5f ) * ax*ax + 1.0f;
+}
+
 /** Encapsulates a TemporalAA pixel shader. */
 template< uint32 Type, uint32 Responsive >
 class FPostProcessTemporalAAPS : public FGlobalShader
@@ -131,10 +140,6 @@ public:
 			};
 		
 			float Sharpness = CVarTemporalAASharpness.GetValueOnRenderThread();
-			// With the new temporal AA, need to hardcode a value here.
-			// Going to remove this once the new temporal AA is validated.
-			Sharpness = -0.25f;
-			Sharpness = 1.0f + Sharpness * 0.5f;
 
 			float Weights[9];
 			float WeightsLow[9];
@@ -144,21 +149,30 @@ public:
 			float TotalWeightPlus = 0.0f;
 			for( int32 i = 0; i < 9; i++ )
 			{
-				// Exponential fit to Blackman-Harris 3.3
 				float PixelOffsetX = SampleOffsets[i][0] - JitterX;
 				float PixelOffsetY = SampleOffsets[i][1] - JitterY;
-				PixelOffsetX *= Sharpness;
-				PixelOffsetY *= Sharpness;
-				Weights[i] = FMath::Exp( -2.29f * ( PixelOffsetX * PixelOffsetX + PixelOffsetY * PixelOffsetY ) );
-				TotalWeight += Weights[i];
+
+				if( Sharpness > 1.0f )
+				{
+					Weights[i] = CatmullRom( PixelOffsetX ) * CatmullRom( PixelOffsetY );
+					TotalWeight += Weights[i];
+				}
+				else
+				{
+					// Exponential fit to Blackman-Harris 3.3
+					PixelOffsetX *= 1.0f + Sharpness * 0.5f;
+					PixelOffsetY *= 1.0f + Sharpness * 0.5f;
+					Weights[i] = FMath::Exp( -2.29f * ( PixelOffsetX * PixelOffsetX + PixelOffsetY * PixelOffsetY ) );
+					TotalWeight += Weights[i];
+				}
 
 				// Lowpass.
 				PixelOffsetX = SampleOffsets[i][0] - JitterX;
 				PixelOffsetY = SampleOffsets[i][1] - JitterY;
 				PixelOffsetX *= 0.25f;
 				PixelOffsetY *= 0.25f;
-				PixelOffsetX *= Sharpness;
-				PixelOffsetY *= Sharpness;
+				PixelOffsetX *= 1.0f + Sharpness * 0.5f;
+				PixelOffsetY *= 1.0f + Sharpness * 0.5f;
 				WeightsLow[i] = FMath::Exp( -2.29f * ( PixelOffsetX * PixelOffsetX + PixelOffsetY * PixelOffsetY ) );
 				TotalWeightLow += WeightsLow[i];
 			}
@@ -721,7 +735,7 @@ void FRCPassPostProcessTemporalAA::Process(FRenderingCompositePassContext& Conte
 
 	// Compute a transform from view origin centered world-space to clip space.
 	View.ViewMatrices.TranslatedViewProjectionMatrix = View.ViewMatrices.TranslatedViewMatrix * View.ViewMatrices.ProjMatrix;
-	View.ViewMatrices.InvTranslatedViewProjectionMatrix = View.ViewMatrices.TranslatedViewProjectionMatrix.Inverse();
+	View.ViewMatrices.InvTranslatedViewProjectionMatrix = View.ViewMatrices.GetInvProjMatrix() * View.ViewMatrices.TranslatedViewMatrix.Inverse();
 
 	View.InitRHIResources(nullptr);
 }
diff --git a/Engine/Source/Runtime/Renderer/Private/SceneVisibility.cpp b/Engine/Source/Runtime/Renderer/Private/SceneVisibility.cpp
--- a/Engine/Source/Runtime/Renderer/Private/SceneVisibility.cpp
+++ b/Engine/Source/Runtime/Renderer/Private/SceneVisibility.cpp
@@ -1550,6 +1550,9 @@ void FSceneRenderer::PreVisibilityFrameSetup(FRHICommandListImmediate& RHICmdLis
 					SampleY = Halton( Index, 3 ) - 0.5f;
 				}
 
+				View.TemporalJitterPixelsX = SampleX;
+				View.TemporalJitterPixelsY = SampleY;
+
 				View.ViewMatrices.TemporalAAProjJitter.X = SampleX * 2.0f / View.ViewRect.Width();
 				View.ViewMatrices.TemporalAAProjJitter.Y = SampleY * 2.0f / View.ViewRect.Height();
 				View.ViewMatrices.ProjMatrix.M[2][0] += View.ViewMatrices.TemporalAAProjJitter.X;
@@ -1559,12 +1562,18 @@ void FSceneRenderer::PreVisibilityFrameSetup(FRHICommandListImmediate& RHICmdLis
 				View.ViewProjectionMatrix = View.ViewMatrices.ViewMatrix * View.ViewMatrices.ProjMatrix;
 				View.InvViewProjectionMatrix = View.ViewMatrices.GetInvProjMatrix() * View.InvViewMatrix;
 
-				/** The view transform, starting from world-space points translated by -ViewOrigin. */
-				FMatrix TranslatedViewMatrix = FTranslationMatrix(-View.ViewMatrices.PreViewTranslation) * View.ViewMatrices.ViewMatrix;
-
 				// Compute a transform from view origin centered world-space to clip space.
-				View.ViewMatrices.TranslatedViewProjectionMatrix = TranslatedViewMatrix * View.ViewMatrices.ProjMatrix;
-				View.ViewMatrices.InvTranslatedViewProjectionMatrix = View.ViewMatrices.TranslatedViewProjectionMatrix.Inverse();
+				if( View.ViewMatrices.PreViewTranslation.IsNearlyZero() )
+				{
+					View.ViewMatrices.TranslatedViewProjectionMatrix = View.ViewProjectionMatrix;
+					View.ViewMatrices.InvTranslatedViewProjectionMatrix = View.InvViewProjectionMatrix;
+				}
+				else
+				{
+					ensure( View.ViewMatrices.TranslatedViewMatrix.GetOrigin().IsNearlyZero() );
+					View.ViewMatrices.TranslatedViewProjectionMatrix = View.ViewMatrices.TranslatedViewMatrix * View.ViewMatrices.ProjMatrix;
+					View.ViewMatrices.InvTranslatedViewProjectionMatrix = View.ViewMatrices.GetInvProjMatrix() * View.ViewMatrices.TranslatedViewMatrix.GetTransposed();
+				}
 			}
 		}
 		else if(ViewState)
