Decals and modulated shadows cause rendering artifacts on iPhone 4 UE-18913

diff --git a/Engine/Source/Developer/ShaderFormatOpenGL/Private/GlslBackend.cpp b/Engine/Source/Developer/ShaderFormatOpenGL/Private/GlslBackend.cpp
--- a/Engine/Source/Developer/ShaderFormatOpenGL/Private/GlslBackend.cpp
+++ b/Engine/Source/Developer/ShaderFormatOpenGL/Private/GlslBackend.cpp
@@ -2953,7 +2953,7 @@ public:
 		if (bUsesDepthbufferFetchES2)
 		{
 			ralloc_asprintf_append(buffer, "\n#ifdef GL_ARM_shader_framebuffer_fetch_depth_stencil\n");
-			ralloc_asprintf_append(buffer, "float DepthbufferFetchES2(float OptionalDepth, float C1, float C2) { float w = 1.0f/(gl_LastFragDepthARM*C1-C2); gl_FragDepth = w; return clamp(w, 0.0f, 65000.0f); }\n");
+			ralloc_asprintf_append(buffer, "float DepthbufferFetchES2(float OptionalDepth, float C1, float C2) { float w = 1.0f/(gl_LastFragDepthARM*C1-C2); return clamp(w, 0.0f, 65000.0f); }\n");
 			ralloc_asprintf_append(buffer, "#else\n");
 			ralloc_asprintf_append(buffer, "float DepthbufferFetchES2(float OptionalDepth, float C1, float C2) { return OptionalDepth; }\n");
 			ralloc_asprintf_append(buffer, "#endif\n\n");
diff --git a/Engine/Source/Runtime/OpenGLDrv/Private/OpenGLDevice.cpp b/Engine/Source/Runtime/OpenGLDrv/Private/OpenGLDevice.cpp
--- a/Engine/Source/Runtime/OpenGLDrv/Private/OpenGLDevice.cpp
+++ b/Engine/Source/Runtime/OpenGLDrv/Private/OpenGLDevice.cpp
@@ -693,6 +693,7 @@ static void InitRHICapabilitiesForGL()
 	GSupportsRenderTargetFormat_PF_FloatRGBA = FOpenGL::SupportsColorBufferHalfFloat();
 
 	GSupportsShaderFramebufferFetch = FOpenGL::SupportsShaderFramebufferFetch();
+	GSupportsShaderDepthStencilFetch = FOpenGL::SupportsShaderDepthStencilFetch();
 	GMaxShadowDepthBufferSizeX = FMath::Min<int32>(Value_GL_MAX_RENDERBUFFER_SIZE, 4096); // Limit to the D3D11 max.
 	GMaxShadowDepthBufferSizeY = FMath::Min<int32>(Value_GL_MAX_RENDERBUFFER_SIZE, 4096);
 	GHardwareHiddenSurfaceRemoval = FOpenGL::HasHardwareHiddenSurfaceRemoval();
diff --git a/Engine/Source/Runtime/OpenGLDrv/Private/OpenGLES2.cpp b/Engine/Source/Runtime/OpenGLDrv/Private/OpenGLES2.cpp
--- a/Engine/Source/Runtime/OpenGLDrv/Private/OpenGLES2.cpp
+++ b/Engine/Source/Runtime/OpenGLDrv/Private/OpenGLES2.cpp
@@ -58,6 +58,9 @@ bool FOpenGLES2::bSupportsColorBufferHalfFloat = false;
 /** GL_EXT_shader_framebuffer_fetch */
 bool FOpenGLES2::bSupportsShaderFramebufferFetch = false;
 
+/** GL_ARM_shader_framebuffer_fetch_depth_stencil */
+bool FOpenGLES2::bSupportsShaderDepthStencilFetch = false;
+
 /** GL_EXT_multisampled_render_to_texture */
 bool FOpenGLES2::bSupportsMultisampledRenderToTexture = false;
 
@@ -190,6 +193,7 @@ void FOpenGLES2::ProcessExtensions( const FString& ExtensionsString )
 	bSupportsSGRB = ExtensionsString.Contains(TEXT("GL_EXT_sRGB"));
 	bSupportsColorBufferHalfFloat = ExtensionsString.Contains(TEXT("GL_EXT_color_buffer_half_float"));
 	bSupportsShaderFramebufferFetch = ExtensionsString.Contains(TEXT("GL_EXT_shader_framebuffer_fetch")) || ExtensionsString.Contains(TEXT("GL_NV_shader_framebuffer_fetch")) || ExtensionsString.Contains(TEXT("GL_ARM_shader_framebuffer_fetch"));
+	bSupportsShaderDepthStencilFetch = ExtensionsString.Contains(TEXT("GL_ARM_shader_framebuffer_fetch_depth_stencil"));
 	bSupportsMultisampledRenderToTexture = ExtensionsString.Contains(TEXT("GL_EXT_multisampled_render_to_texture"));
 	// @todo ios7: SRGB support does not work with our texture format setup (ES2 docs indicate that internalFormat and format must match, but they don't at all with sRGB enabled)
 	//             One possible solution us to use GLFormat.InternalFormat[bSRGB] instead of GLFormat.Format
diff --git a/Engine/Source/Runtime/OpenGLDrv/Private/OpenGLES31.cpp b/Engine/Source/Runtime/OpenGLDrv/Private/OpenGLES31.cpp
--- a/Engine/Source/Runtime/OpenGLDrv/Private/OpenGLES31.cpp
+++ b/Engine/Source/Runtime/OpenGLDrv/Private/OpenGLES31.cpp
@@ -70,6 +70,9 @@ bool FOpenGLES31::bSupportsNvImageFormats = false;
 /** GL_EXT_shader_framebuffer_fetch */
 bool FOpenGLES31::bSupportsShaderFramebufferFetch = false;
 
+/** GL_ARM_shader_framebuffer_fetch_depth_stencil */
+bool FOpenGLES31::bSupportsShaderDepthStencilFetch = false;
+
 /** GL_EXT_sRGB */
 bool FOpenGLES31::bSupportsSGRB = false;
 
@@ -274,6 +277,7 @@ void FOpenGLES31::ProcessExtensions( const FString& ExtensionsString )
 	bSupportsColorBufferHalfFloat = ExtensionsString.Contains(TEXT("GL_EXT_color_buffer_half_float"));
 	bSupportsNvImageFormats = ExtensionsString.Contains(TEXT("GL_NV_image_formats"));
 	bSupportsShaderFramebufferFetch = ExtensionsString.Contains(TEXT("GL_EXT_shader_framebuffer_fetch")) || ExtensionsString.Contains(TEXT("GL_NV_shader_framebuffer_fetch")) || ExtensionsString.Contains(TEXT("GL_ARM_shader_framebuffer_fetch"));
+	bSupportsShaderDepthStencilFetch = ExtensionsString.Contains(TEXT("GL_ARM_shader_framebuffer_fetch_depth_stencil"));
 	// @todo es3: SRGB support does not work with our texture format setup (ES2 docs indicate that internalFormat and format must match, but they don't at all with sRGB enabled)
 	//             One possible solution us to use GLFormat.InternalFormat[bSRGB] instead of GLFormat.Format
 	bSupportsSGRB = false;//ExtensionsString.Contains(TEXT("GL_EXT_sRGB"));
diff --git a/Engine/Source/Runtime/OpenGLDrv/Public/OpenGL.h b/Engine/Source/Runtime/OpenGLDrv/Public/OpenGL.h
--- a/Engine/Source/Runtime/OpenGLDrv/Public/OpenGL.h
+++ b/Engine/Source/Runtime/OpenGLDrv/Public/OpenGL.h
@@ -111,6 +111,7 @@ public:
 	static FORCEINLINE bool SupportsR11G11B10F()						{ return true; }
 	static FORCEINLINE bool SupportsVolumeTextureRendering()			{ return false; }
 	static FORCEINLINE bool SupportsShaderFramebufferFetch()			{ return false; }
+	static FORCEINLINE bool SupportsShaderDepthStencilFetch()			{ return false; }
 	static FORCEINLINE bool SupportsVertexArrayBGRA()					{ return true; }
 	static FORCEINLINE bool SupportsBGRA8888()							{ return true; }
 	static FORCEINLINE bool SupportsSRGB()								{ return true; }
diff --git a/Engine/Source/Runtime/OpenGLDrv/Public/OpenGLES2.h b/Engine/Source/Runtime/OpenGLDrv/Public/OpenGLES2.h
--- a/Engine/Source/Runtime/OpenGLDrv/Public/OpenGLES2.h
+++ b/Engine/Source/Runtime/OpenGLDrv/Public/OpenGLES2.h
@@ -108,6 +108,7 @@ struct FOpenGLES2 : public FOpenGLBase
 	static FORCEINLINE bool	SupportsRG16UI()							{ return false; }
 	static FORCEINLINE bool SupportsR11G11B10F()						{ return false; }
 	static FORCEINLINE bool SupportsShaderFramebufferFetch()			{ return bSupportsShaderFramebufferFetch; }
+	static FORCEINLINE bool SupportsShaderDepthStencilFetch()			{ return bSupportsShaderDepthStencilFetch; }
 	static FORCEINLINE bool SupportsMultisampledRenderToTexture()		{ return bSupportsMultisampledRenderToTexture; }
 	static FORCEINLINE bool SupportsVertexArrayBGRA()					{ return false; }
 	static FORCEINLINE bool SupportsBGRA8888()							{ return bSupportsBGRA8888; }
@@ -439,6 +440,9 @@ protected:
 	/** GL_EXT_shader_framebuffer_fetch */
 	static bool bSupportsShaderFramebufferFetch;
 
+	/** GL_ARM_shader_framebuffer_fetch_depth_stencil */
+	static bool bSupportsShaderDepthStencilFetch;
+
 	/** GL_EXT_MULTISAMPLED_RENDER_TO_TEXTURE */
 	static bool bSupportsMultisampledRenderToTexture;
 
diff --git a/Engine/Source/Runtime/OpenGLDrv/Public/OpenGLES31.h b/Engine/Source/Runtime/OpenGLDrv/Public/OpenGLES31.h
--- a/Engine/Source/Runtime/OpenGLDrv/Public/OpenGLES31.h
+++ b/Engine/Source/Runtime/OpenGLDrv/Public/OpenGLES31.h
@@ -65,6 +65,7 @@ struct FOpenGLES31 : public FOpenGLBase
 	static FORCEINLINE bool	SupportsRG16UI()							{ return bSupportsNvImageFormats && !bES2Fallback; }
 	static FORCEINLINE bool SupportsR11G11B10F()						{ return bSupportsNvImageFormats && !bES2Fallback; }
 	static FORCEINLINE bool SupportsShaderFramebufferFetch()			{ return bSupportsShaderFramebufferFetch; }
+	static FORCEINLINE bool SupportsShaderDepthStencilFetch()			{ return bSupportsShaderDepthStencilFetch; }
 	static FORCEINLINE bool SupportsVertexArrayBGRA()					{ return false; }
 	static FORCEINLINE bool SupportsBGRA8888()							{ return bSupportsBGRA8888; }
 	static FORCEINLINE bool SupportsSRGB()								{ return bSupportsSGRB || !bES2Fallback; }
@@ -894,6 +895,9 @@ protected:
 
 	/** GL_EXT_shader_framebuffer_fetch */
 	static bool bSupportsShaderFramebufferFetch;
+	
+	/** GL_ARM_shader_framebuffer_fetch_depth_stencil */
+	static bool bSupportsShaderDepthStencilFetch;
 
 	/** GL_FRAGMENT_SHADER, GL_LOW_FLOAT */
 	static int ShaderLowPrecision;
diff --git a/Engine/Source/Runtime/RHI/Private/RHI.cpp b/Engine/Source/Runtime/RHI/Private/RHI.cpp
--- a/Engine/Source/Runtime/RHI/Private/RHI.cpp
+++ b/Engine/Source/Runtime/RHI/Private/RHI.cpp
@@ -183,6 +183,7 @@ bool GSupportsRenderDepthTargetableShaderResources = true;
 bool GSupportsRenderTargetFormat_PF_G8 = true;
 bool GSupportsRenderTargetFormat_PF_FloatRGBA = true;
 bool GSupportsShaderFramebufferFetch = false;
+bool GSupportsShaderDepthStencilFetch = false;
 bool GHardwareHiddenSurfaceRemoval = false;
 bool GRHISupportsAsyncTextureCreation = false;
 bool GSupportsQuads = false;
diff --git a/Engine/Source/Runtime/RHI/Public/RHI.h b/Engine/Source/Runtime/RHI/Public/RHI.h
--- a/Engine/Source/Runtime/RHI/Public/RHI.h
+++ b/Engine/Source/Runtime/RHI/Public/RHI.h
@@ -92,6 +92,9 @@ extern RHI_API bool GSupportsRenderTargetFormat_PF_FloatRGBA;
 /** true if mobile framebuffer fetch is supported */
 extern RHI_API bool GSupportsShaderFramebufferFetch;
 
+/** true if mobile depth & stencil fetch is supported */
+extern RHI_API bool GSupportsShaderDepthStencilFetch;
+
 /** true if the GPU supports hidden surface removal in hardware. */
 extern RHI_API bool GHardwareHiddenSurfaceRemoval;
 
diff --git a/Engine/Source/Runtime/Renderer/Private/DecalRenderingForwardShading.cpp b/Engine/Source/Runtime/Renderer/Private/DecalRenderingForwardShading.cpp
--- a/Engine/Source/Runtime/Renderer/Private/DecalRenderingForwardShading.cpp
+++ b/Engine/Source/Runtime/Renderer/Private/DecalRenderingForwardShading.cpp
@@ -16,18 +16,17 @@ void FForwardShadingSceneRenderer::RenderDecals(FRHICommandListImmediate& RHICmd
 		return;
 	}
 
-	SCOPE_CYCLE_COUNTER(STAT_DecalsDrawTime);
-
-	// Setup read access to depth/stencil buffer on PC platform (mobile preview)
-	if (IsPCPlatform(ViewFamily.GetShaderPlatform()))
+	// Check if it possible to render decals on this device
+	if (!DeviceSupportsShaderDepthFetch())
 	{
-		FSceneRenderTargets& SceneContext = FSceneRenderTargets::Get(RHICmdList);
-		FRHIDepthRenderTargetView	DepthView(SceneContext.GetSceneDepthSurface(), ERenderTargetLoadAction::ELoad, ERenderTargetStoreAction::ENoAction, ERenderTargetLoadAction::ELoad, ERenderTargetStoreAction::ENoAction, FExclusiveDepthStencil(FExclusiveDepthStencil::DepthRead_StencilRead));
-		FRHIRenderTargetView		ColorView(SceneContext.GetSceneColorSurface(), 0, -1, ERenderTargetLoadAction::ENoAction, ERenderTargetStoreAction::EStore);
-		FRHISetRenderTargetsInfo	Info(1, &ColorView, DepthView);
-		RHICmdList.SetRenderTargetsAndClear(Info);
+		return;
 	}
 
+	SCOPE_CYCLE_COUNTER(STAT_DecalsDrawTime);
+
+	FSceneRenderTargets& SceneContext = FSceneRenderTargets::Get(RHICmdList);
+	SceneContext.BeginRenderingSceneColor(RHICmdList, ESimpleRenderTargetMode::EExistingColorAndDepth, FExclusiveDepthStencil::DepthRead_StencilRead);
+	
 	for (int32 ViewIndex = 0; ViewIndex < Views.Num(); ViewIndex++)
 	{
 		const FViewInfo& View = Views[ViewIndex];
diff --git a/Engine/Source/Runtime/Renderer/Private/PostProcess/SceneRenderTargets.cpp b/Engine/Source/Runtime/Renderer/Private/PostProcess/SceneRenderTargets.cpp
--- a/Engine/Source/Runtime/Renderer/Private/PostProcess/SceneRenderTargets.cpp
+++ b/Engine/Source/Runtime/Renderer/Private/PostProcess/SceneRenderTargets.cpp
@@ -1446,19 +1446,17 @@ void FSceneRenderTargets::AllocateForwardShadingPathRenderTargets()
 
 	EPixelFormat Format = GetSceneColor()->GetDesc().Format;
 
+#if PLATFORM_HTML5
 	// For 64-bit ES2 without framebuffer fetch, create extra render target for copy of alpha channel.
 	if((Format == PF_FloatRGBA) && (GSupportsShaderFramebufferFetch == false)) 
 	{
-#if PLATFORM_HTML5 || PLATFORM_ANDROID
 		// creating a PF_R16F (a true one-channel renderable fp texture) is only supported on GL if EXT_texture_rg is available.  It's present
 		// on iOS, but not in WebGL or Android.
 		FPooledRenderTargetDesc Desc(FPooledRenderTargetDesc::Create2DDesc(BufferSize, PF_FloatRGBA, FClearValueBinding::None, TexCreate_None, TexCreate_RenderTargetable, false));
-#else
-		FPooledRenderTargetDesc Desc(FPooledRenderTargetDesc::Create2DDesc(BufferSize, PF_R16F, FClearValueBinding::None, TexCreate_None, TexCreate_RenderTargetable, false));
-#endif
 		GRenderTargetPool.FindFreeElement(Desc, SceneAlphaCopy, TEXT("SceneAlphaCopy"));
 	}
 	else
+#endif
 	{
 		SceneAlphaCopy = GSystemTextures.MaxFP16Depth;
 	}
diff --git a/Engine/Source/Runtime/Renderer/Private/ScenePrivate.h b/Engine/Source/Runtime/Renderer/Private/ScenePrivate.h
--- a/Engine/Source/Runtime/Renderer/Private/ScenePrivate.h
+++ b/Engine/Source/Runtime/Renderer/Private/ScenePrivate.h
@@ -42,6 +42,9 @@ bool IsMobileHDR32bpp();
 /** True if the mobile renderer is emulating HDR with mosaic. */
 bool IsMobileHDRMosaic();
 
+/** True if device has some way to fetch depth in a fragment shader. */
+bool DeviceSupportsShaderDepthFetch();
+
 class FOcclusionQueryHelpers
 {
 public:
diff --git a/Engine/Source/Runtime/Renderer/Private/SceneRendering.cpp b/Engine/Source/Runtime/Renderer/Private/SceneRendering.cpp
--- a/Engine/Source/Runtime/Renderer/Private/SceneRendering.cpp
+++ b/Engine/Source/Runtime/Renderer/Private/SceneRendering.cpp
@@ -1772,4 +1772,13 @@ bool IsMobileHDRMosaic()
 		default:
 			return !(GSupportsHDR32bppEncodeModeIntrinsic && GSupportsShaderFramebufferFetch);
 	}
+}
+
+bool DeviceSupportsShaderDepthFetch()
+{
+	// device should support extension to directly fetch depth in a shader
+	// or support and enable half-float color buffer (base pass stores depth in alpha channel)
+	return (!IsMobilePlatform(GMaxRHIShaderPlatform) 
+			|| GSupportsShaderDepthStencilFetch 
+			|| (GSupportsRenderTargetFormat_PF_FloatRGBA && (IsMobileHDR() && !IsMobileHDR32bpp())));
 }
\ No newline at end of file
diff --git a/Engine/Source/Runtime/Renderer/Private/ShadowRendering.cpp b/Engine/Source/Runtime/Renderer/Private/ShadowRendering.cpp
--- a/Engine/Source/Runtime/Renderer/Private/ShadowRendering.cpp
+++ b/Engine/Source/Runtime/Renderer/Private/ShadowRendering.cpp
@@ -3815,7 +3815,7 @@ void FForwardShadingSceneRenderer::RenderModulatedShadowProjections(FRHICommandL
 
 	bool bDynamicShadows = ViewFamily.EngineShowFlags.DynamicShadows && GetShadowQuality() > 0;
 
-	if (IsSimpleDynamicLightingEnabled() || !bDynamicShadows)
+	if (IsSimpleDynamicLightingEnabled() || !bDynamicShadows || !DeviceSupportsShaderDepthFetch())
 	{
 		return;
 	}
diff --git a/Engine/Source/Runtime/Renderer/Private/ShadowSetup.cpp b/Engine/Source/Runtime/Renderer/Private/ShadowSetup.cpp
--- a/Engine/Source/Runtime/Renderer/Private/ShadowSetup.cpp
+++ b/Engine/Source/Runtime/Renderer/Private/ShadowSetup.cpp
@@ -2334,6 +2334,9 @@ void FForwardShadingSceneRenderer::InitDynamicShadows(FRHICommandListImmediate&
 		bStaticSceneOnly = bStaticSceneOnly || View.bStaticSceneOnly;
 	}
 
+	// Modulated shadows require depth fetch
+	const bool bSupportsDepthFetch = DeviceSupportsShaderDepthFetch();
+
 	{
 		SCOPE_CYCLE_COUNTER(STAT_InitDynamicShadowsTime);
 
@@ -2368,7 +2371,7 @@ void FForwardShadingSceneRenderer::InitDynamicShadows(FRHICommandListImmediate&
 				NumWholeSceneShadows = VisibleLightInfo.AllProjectedShadows.Num();
 
 				// If we're casting modulated shadows then look for dynamic shadow interactions.
-				if (LightSceneInfo->Proxy->CastsModulatedShadows())
+				if (bSupportsDepthFetch && LightSceneInfo->Proxy->CastsModulatedShadows())
 				{
 					// Look for individual primitives with a dynamic shadow.
 					for (FLightPrimitiveInteraction* Interaction = LightSceneInfo->DynamicPrimitiveList;
