Disable decals and modulated shadows on all devices when MobileHDR is off

diff --git a/Engine/Source/Runtime/Renderer/Private/SceneRendering.cpp b/Engine/Source/Runtime/Renderer/Private/SceneRendering.cpp
--- a/Engine/Source/Runtime/Renderer/Private/SceneRendering.cpp
+++ b/Engine/Source/Runtime/Renderer/Private/SceneRendering.cpp
@@ -1776,9 +1776,13 @@ bool IsMobileHDRMosaic()
 
 bool DeviceSupportsShaderDepthFetch()
 {
-	// device should support extension to directly fetch depth in a shader
-	// or support and enable half-float color buffer (base pass stores depth in alpha channel)
-	return (!IsMobilePlatform(GMaxRHIShaderPlatform) 
-			|| GSupportsShaderDepthStencilFetch 
-			|| (GSupportsRenderTargetFormat_PF_FloatRGBA && (IsMobileHDR() && !IsMobileHDR32bpp())));
+	if (IsMobilePlatform(GMaxRHIShaderPlatform))
+	{
+		// HDR should be enabled
+		// device should support extension to directly fetch depth in a shader
+		// or half-float color buffer (base pass stores depth in alpha channel)
+		return IsMobileHDR() && (GSupportsShaderDepthStencilFetch || (GSupportsRenderTargetFormat_PF_FloatRGBA && !IsMobileHDR32bpp()));
+	}
+
+	return true;
 }
\ No newline at end of file
