fixed ScreenSpaceSubusrfaceScattering actifacts
 - wasn't scaling right with ViewportSize, e.g. HighResScreenshot
 - Missing a horizontal line in the 2d kernel
 - Improved quality of the full res portion
 - refactored to avoid one full res pass
 - simplified Stereo rendering support
 - In single viewport case it now goes down to 32bit texture format afterwards (input can be 32/RGB or 64/RGBA as before, without alpha it can look bad)

diff --git a/Engine/Shaders/PostProcessSubsurface.usf b/Engine/Shaders/PostProcessSubsurface.usf
--- a/Engine/Shaders/PostProcessSubsurface.usf
+++ b/Engine/Shaders/PostProcessSubsurface.usf
@@ -13,7 +13,7 @@
 
 /** one profile per line, sample set left to right: high/med/low */
 Texture2D SSProfilesTexture;
-// x:Radius*DistanceToProjectionWindow/KernelSize*0.5, y:DistanceToProjectionWindow, zw:xy scale correction for BufferSize!=ViewportSize
+// x:Radius*DistanceToProjectionWindow/KernelSize*0.5, y:DistanceToProjectionWindow, zw: unused
 float4 SSSParams;
 
 // 0: (testing, could be faster), 1: higher quality
@@ -65,15 +65,17 @@ float GetSubsurfaceStrength(float2 UV)
 	#define SSSS_SPECULAR_CORRECTION 0
 #endif
 
+#define	SSSS_N_KERNELSUBSURFACECOLOROFFSET 0
+
 #if SSS_SAMPLESET == 0
 	#define	SSSS_N_KERNELWEIGHTCOUNT 6
-	#define	SSSS_N_KERNELWEIGHTOFFSET (13 + 9)
+	#define	SSSS_N_KERNELWEIGHTOFFSET (1 + 13 + 9)
 #elif SSS_SAMPLESET == 1
 	#define	SSSS_N_KERNELWEIGHTCOUNT 9
-	#define	SSSS_N_KERNELWEIGHTOFFSET 13
+	#define	SSSS_N_KERNELWEIGHTOFFSET (1 + 13)
 #else // SSS_SAMPLESET == 2
 	#define	SSSS_N_KERNELWEIGHTCOUNT 13
-	#define	SSSS_N_KERNELWEIGHTOFFSET 0
+	#define	SSSS_N_KERNELWEIGHTOFFSET (1)
 #endif
 
 // 0: faster
@@ -90,7 +92,7 @@ float4 Quantize8Bit(float4 In)
 // @return .rgb is the weight for color channel, .a is the sample location
 float4 GetKernel(uint SampleIndex, uint SubsurfaceProfileInt)
 {
-	const float4 TableMax = float4(1, 1, 1, 3.0f);
+	const float4 TableMax = float4(1, 1, 1, SUBSURFACE_KERNEL_SIZE);
 
 	// profiled on NV670 fullscreen, one pass (total: 2 pass), samples: 13+1+13, large object filling the screen
 
@@ -104,7 +106,6 @@ float4 GetKernel(uint SampleIndex, uint SubsurfaceProfileInt)
 	// debug if the kernal was a box filter
 //	Value.rgb = 1.0f / (SSSS_N_KERNELWEIGHTCOUNT * 2 - 1);
 
-
 	return Value;
 }
 
@@ -116,6 +117,14 @@ float GetProfileRadiusScale(FGBufferData GBufferData)
 	return GetKernel(SSSS_N_KERNELWEIGHTOFFSET + SSSS_N_KERNELWEIGHTCOUNT - 1, SubsurfaceProfileInt).a;
 }
 
+float3 GetProfileSubsurfaceColor(FGBufferData GBufferData)
+{
+	// 0..255, which SubSurface profile to pick
+	uint SubsurfaceProfileInt = ExtractSubsurfaceProfileInt(GBufferData);
+
+	return GetKernel(SSSS_N_KERNELSUBSURFACECOLOROFFSET, SubsurfaceProfileInt).rgb;
+}
+
 // from [REDACTED_PROFILE_URL]/separable-sss/tree/master/Demo
 // [REDACTED_PERSON] [REDACTED_PROFILE_URL]
 // [REDACTED_PROFILE_URL]
@@ -146,9 +155,7 @@ float ComputeFullResLerp(FScreenSpaceData ScreenSpaceData, float2 UVSceneColor,
 
 	float scale = SSSScaleX / CalcSceneDepth(UVSceneColor);
 
-	// the viewport is only a fraction of the buffersize, here we compensate for that
-	// *1024 as we store the sample distances / 1000 to allows for scaling
-	float HorizontalScaler = View.ViewSizeAndInvSize.x * View.BufferSizeAndInvSize.z * 1024.0f;
+	float HorizontalScaler = SUBSURFACE_RADIUS_SCALE;
 
 	// Calculate the final step to fetch the surrounding pixels:
 	float finalStep = scale * HorizontalScaler;
@@ -288,19 +295,6 @@ void ExtractSpecularPS(float4 UVAndScreenPos : TEXCOORD0, out float4 OutColor :
 
 	// alpha channel can be anything - we consumed the value and likely output to a format that doesn't need the value
 	OutColor = float4(lerp(SceneColor4.rgb, ExtractedSpecular, BlendFactor), 0);
-
-#if FULLRES_WHERE_POSSIBLE
-	// 0..255, which SubSurface profile to pick
-	uint SubsurfaceProfileInt = ExtractSubsurfaceProfileInt(ScreenSpaceData.GBuffer);
-
-	// reference
-//	float3 DiffuseLighting = LightAccumulator_ReconstructDiffuseLighting(SceneColor4);
-	// cheaper:
-	float3 DiffuseLighting = SceneColor4 - ExtractedSpecular;
-
-	// add center sample from full resolution input for better quality
-	OutColor.rgb += DiffuseLighting * GetKernel(SSSS_N_KERNELWEIGHTOFFSET, SubsurfaceProfileInt).rgb;
-#endif
 }
 
 void SetupPS(in float4 UVAndScreenPos : TEXCOORD0, out float4 OutColor : SV_Target0)
@@ -338,9 +332,7 @@ void SetupPS(in float4 UVAndScreenPos : TEXCOORD0, out float4 OutColor : SV_Targ
 	float scale = SSSScaleX / SingleDepth;
 //	float scale = SSSScaleX / min(min(FourDepth.x, FourDepth.y), min(FourDepth.z, FourDepth.w));
 
-	// the viewport is only a fraction of the buffersize, here we compensate for that
-	// *1024 as we store the sample distances / 1024 to allows for scaling
-	float HorizontalScaler = View.ViewSizeAndInvSize.x * View.BufferSizeAndInvSize.z * 1024.0f;
+	float HorizontalScaler = SUBSURFACE_RADIUS_SCALE;
 
 	// Calculate the final step to fetch the surrounding pixels:
 	float finalStep = scale * HorizontalScaler;
@@ -359,21 +351,8 @@ void SetupPS(in float4 UVAndScreenPos : TEXCOORD0, out float4 OutColor : SV_Targ
 		OutColor = 0;
 	}
 #endif
-
-	// debug the kernel size and shape
-	/*
-	if(length((UV - 0.5f) * float2(2,1)) > 0.005f)
-	{
-		OutColor.rgb = 0;
-	}
-	else
-	{
-		OutColor.rgb *= 50;
-	}
-	*/
 }
 
-
 // input0 is created by the SetupPS shader
 void MainPS(float4 UVAndScreenPos : TEXCOORD0, out float4 OutColor : SV_Target0)
 {
@@ -382,26 +361,19 @@ void MainPS(float4 UVAndScreenPos : TEXCOORD0, out float4 OutColor : SV_Target0)
 	// call into "SeparableSSS.usf"
 
 	// the viewport is only a fraction of the buffersize, here we compensate for that
-	// *1024 as we store the sample distances / 1024 to allows for scaling
 
 #if SSS_DIRECTION == 0
 	// horizontal
-	float2 Direction = float2(1, 0) * View.ViewSizeAndInvSize.x * View.BufferSizeAndInvSize.z * 1024.0f;
+	float2 ViewportDirectionUV = float2(1, 0) * SUBSURFACE_RADIUS_SCALE;
 #else
 	// vertical
-	float2 Direction = float2(0, 1) * View.BufferSizeAndInvSize.w * View.ViewSizeAndInvSize.x * 1024.0f;
+	float2 ViewportDirectionUV = float2(0, 1) * SUBSURFACE_RADIUS_SCALE * (View.ViewSizeAndInvSize.x * View.ViewSizeAndInvSize.w);
 #endif
 
 	// can be optimized
 	float2 GBufferUV = (ViewportUV * View.ViewSizeAndInvSize.xy + View.ViewRectMin.xy) * View.BufferSizeAndInvSize.zw;
 
-	bool bIncludeCenter = true;
-
-#if FULLRES_WHERE_POSSIBLE && SSS_DIRECTION == 1
-	bIncludeCenter = false;
-#endif
-
-	OutColor = SSSSBlurPS(GBufferUV, ViewportUV, Direction, false, bIncludeCenter);
+	OutColor = SSSSBlurPS(GBufferUV, ViewportUV, ViewportDirectionUV, false);
 
 #if SSS_DIRECTION == 1
 	// second pass prepares the setup from the recombine pass which doesn't need depth but wants to reconstruct the color
@@ -423,15 +395,12 @@ void SubsurfaceRecombinePS(float4 UVAndScreenPos : TEXCOORD0, out float4 OutColo
 	float2 PixelPos = UVAndScreenPos.zw * ScreenPosToPixel.xy + ScreenPosToPixel.zw;
 	float2 UVSceneColor = (PixelPos + 0.5f) * View.BufferSizeAndInvSize.zw;
 
-	// specular extracted
-	float4 SceneColor4 = Texture2DSample(PostprocessInput1, PostprocessInput1Sampler, UVSceneColor);
-
 	FScreenSpaceData ScreenSpaceData = GetScreenSpaceData(UVSceneColor);
 
 	// bilinear filtering
 	float3 SSSColor;
 	{
-		float4 SSSColorWithAlpha = Texture2DSample(PostprocessInput0, PostprocessInput0Sampler, ViewportUV);
+		float4 SSSColorWithAlpha = Texture2DSample(PostprocessInput1, PostprocessInput1Sampler, ViewportUV);
 
 		// renormalize to dilate RGB to fix half res upsampling artifacts
 		SSSColor = SSSColorWithAlpha.rgb / max(SSSColorWithAlpha.a, 0.00001f);
@@ -452,10 +421,24 @@ void SubsurfaceRecombinePS(float4 UVAndScreenPos : TEXCOORD0, out float4 OutColo
 #endif
 #endif // HALF_RES
 
+	float4 SceneColor4 = Texture2DSample(PostprocessInput0, PostprocessInput0Sampler, UVSceneColor);
+
+	float3 ExtractedSpecular = 0;
+	{
+#if SSS_RECOMBINE_METHOD
+		// we took the specular highlights out, now we add them back in
+		ExtractedSpecular += LightAccumulator_ReconstructNonDiffuseLighting(SceneColor4);
+#endif
+
+		ExtractedSpecular = lerp(SceneColor4.rgb, ExtractedSpecular, BlendFactor);
+	}
+
+	float3 SubsurfaceColor = GetProfileSubsurfaceColor(ScreenSpaceData.GBuffer);
+
 	// only affect pixels that have SSS
 	LerpFactor *= BlendFactor;
 
 	// add back in non subsurface pixels and specular, alpha: 0=don't blend in, 1:fully blend in
-	OutColor = float4(SSSColor * BlendFactor + SceneColor4.rgb, LerpFactor);
+	OutColor = float4(lerp(SceneColor4.rgb, SSSColor * BlendFactor + ExtractedSpecular, SubsurfaceColor * LerpFactor), 0); 
 }
 
diff --git a/Engine/Shaders/SeparableSSS.usf b/Engine/Shaders/SeparableSSS.usf
--- a/Engine/Shaders/SeparableSSS.usf
+++ b/Engine/Shaders/SeparableSSS.usf
@@ -186,7 +186,7 @@ float3 SSSSTransmittance(float translucency, float sssWidth, float3 worldPositio
 // @param texcoord The usual quad texture coordinates.
 // @param dir Direction of the blur: First pass:   float2(1.0, 0.0), Second pass:  float2(0.0, 1.0)
 // @param initStencil indicates whether the stencil buffer should be initialized. Should be set to 'true' for the first pass if not previously initialized, to enable optimization of the second pass.
-float4 SSSSBlurPS(float2 GBufferUV, float2 ViewportUV, float2 dir, bool initStencil, bool bIncludeCenter)
+float4 SSSSBlurPS(float2 GBufferUV, float2 ViewportUV, float2 dir, bool initStencil)
 {
     // Fetch color of current pixel:
     float4 colorM = SSSSSampleSceneColorPoint(ViewportUV);
@@ -204,6 +204,7 @@ float4 SSSSBlurPS(float2 GBufferUV, float2 ViewportUV, float2 dir, bool initSten
 		return 0.0f;
 	}
 
+	// 0..1
 	float SSSStrength = GetSubsurfaceStrength(GBufferUV);
 
 	// Initialize the stencil buffer in case it was not already available:
@@ -231,19 +232,15 @@ float4 SSSSBlurPS(float2 GBufferUV, float2 ViewportUV, float2 dir, bool initSten
 	// >0 to avoid division by 0, not 100% correct to not visible
 	float3 colorInvDiv = 0.00001f;
 
-	// center sample is added in Recombine pass
+	// center sample
 	colorInvDiv += GetKernel(SSSS_N_KERNELWEIGHTOFFSET, SubsurfaceProfileInt).rgb;
-
-	if(bIncludeCenter)
-	{
-		colorAccum = colorM.rgb * GetKernel(SSSS_N_KERNELWEIGHTOFFSET, SubsurfaceProfileInt).rgb;
-	}
+	colorAccum = colorM.rgb * GetKernel(SSSS_N_KERNELWEIGHTOFFSET, SubsurfaceProfileInt).rgb;
 
 	// Accumulate the other samples:
     SSSS_UNROLL
 	for (int i = 1; i < SSSS_N_KERNELWEIGHTCOUNT; i++) 
 	{
-		// Kernel.a = 0..3
+		// Kernel.a = 0..SUBSURFACE_KERNEL_SIZE (radius)
 		half4 Kernel = GetKernel(SSSS_N_KERNELWEIGHTOFFSET + i, SubsurfaceProfileInt);
 
 		float4 LocalAccum = 0;
diff --git a/Engine/Source/Runtime/Engine/Classes/Engine/SubsurfaceProfile.h b/Engine/Source/Runtime/Engine/Classes/Engine/SubsurfaceProfile.h
--- a/Engine/Source/Runtime/Engine/Classes/Engine/SubsurfaceProfile.h
+++ b/Engine/Source/Runtime/Engine/Classes/Engine/SubsurfaceProfile.h
@@ -156,5 +156,12 @@ private:
 
 };
 
+// If you change this you need to recompile the SSS shaders.
+// Required if we use a texture format with limited size but want to express a larger radius
+#define SUBSURFACE_RADIUS_SCALE 1024.0f
+
+ // The kernels range from -3 to 3
+#define SUBSURFACE_KERNEL_SIZE 3.0f
+
 // lives on the render thread
 extern ENGINE_API TGlobalResource<FSubsurfaceProfileTexture> GSubsurfaceProfileTextureObject;
diff --git a/Engine/Source/Runtime/Engine/Private/Rendering/SeparableSSS.cpp b/Engine/Source/Runtime/Engine/Private/Rendering/SeparableSSS.cpp
--- a/Engine/Source/Runtime/Engine/Private/Rendering/SeparableSSS.cpp
+++ b/Engine/Source/Runtime/Engine/Private/Rendering/SeparableSSS.cpp
@@ -149,6 +149,8 @@ void ComputeMirroredSSSKernel(FLinearColor* TargetBuffer, uint32 TargetBufferSiz
 			}
 		}
 
+		/* we do that in the shader for better quality with half res 
+		
 		// Tweak them using the desired strength. The first one is:
 		//     lerp(1.0, kernel[0].rgb, strength)
 		kernel[0].R = FMath::Lerp(1.0f, kernel[0].R, SubsurfaceColor.R);
@@ -160,7 +162,7 @@ void ComputeMirroredSSSKernel(FLinearColor* TargetBuffer, uint32 TargetBufferSiz
 			kernel[i].R *= SubsurfaceColor.R;
 			kernel[i].G *= SubsurfaceColor.G;
 			kernel[i].B *= SubsurfaceColor.B;
-		}
+		}*/
 	}
 
 	// generate output (remove negative samples)
diff --git a/Engine/Source/Runtime/Engine/Private/Rendering/SubsurfaceProfile.cpp b/Engine/Source/Runtime/Engine/Private/Rendering/SubsurfaceProfile.cpp
--- a/Engine/Source/Runtime/Engine/Private/Rendering/SubsurfaceProfile.cpp
+++ b/Engine/Source/Runtime/Engine/Private/Rendering/SubsurfaceProfile.cpp
@@ -158,7 +158,8 @@ void FSubsurfaceProfileTexture::CreateTexture(FRHICommandListImmediate& RHICmdLi
 	const uint32 KernelSize1 = 9; 
 	const uint32 KernelSize2 = 6;
 
-	const uint32 KernelTotalSize = KernelSize0 + KernelSize1 + KernelSize2;
+	// index 0 is used for the SubsurfaceColor
+	const uint32 KernelTotalSize = 1 + KernelSize0 + KernelSize1 + KernelSize2;
 	check(KernelTotalSize < Width);
 
 	FLinearColor kernel[Width];
@@ -175,12 +176,17 @@ void FSubsurfaceProfileTexture::CreateTexture(FRHICommandListImmediate& RHICmdLi
 		// 0.0001f turned out to be too small to fix the issue (for a small KernelSize)
 		const float Bias = 0.009f;
 
-		Data.SubsurfaceColor = Data.SubsurfaceColor.GetClamped(Bias);
+		Data.SubsurfaceColor = Data.SubsurfaceColor.GetClamped();
 		Data.FalloffColor = Data.FalloffColor.GetClamped(Bias);
 
-		ComputeMirroredSSSKernel(&kernel[0], KernelSize0, Data.SubsurfaceColor, Data.FalloffColor);
-		ComputeMirroredSSSKernel(&kernel[KernelSize0], KernelSize1, Data.SubsurfaceColor, Data.FalloffColor);
-		ComputeMirroredSSSKernel(&kernel[KernelSize0 + KernelSize1], KernelSize2, Data.SubsurfaceColor, Data.FalloffColor);
+		// to allow blending of the Subsurface with fullres in the shader
+		kernel[0] = Data.SubsurfaceColor;
+		// unused
+		kernel[0].A = 0;
+
+		ComputeMirroredSSSKernel(&kernel[1], KernelSize0, Data.SubsurfaceColor, Data.FalloffColor);
+		ComputeMirroredSSSKernel(&kernel[1 + KernelSize0], KernelSize1, Data.SubsurfaceColor, Data.FalloffColor);
+		ComputeMirroredSSSKernel(&kernel[1 + KernelSize0 + KernelSize1], KernelSize2, Data.SubsurfaceColor, Data.FalloffColor);
 
 		// could be lower than 1 (but higher than 0) to range compress for better quality (for 8 bit)
 		const float TableMaxRGB = 1.0f;
@@ -192,7 +198,7 @@ void FSubsurfaceProfileTexture::CreateTexture(FRHICommandListImmediate& RHICmdLi
 			FVector4 C = kernel[Pos] * FLinearColor(1.0f / TableMaxRGB, 1.0f / TableMaxRGB, 1.0f / TableMaxRGB, 1.0f / TableMaxA);
 
 			// requires 16bit (could be made with 8 bit e.g. using sample0.w as 8bit scale applied to all samples (more multiplications in the shader))
-			C.W *= Data.ScatterRadius / 1024.0f;
+			C.W *= Data.ScatterRadius / SUBSURFACE_RADIUS_SCALE;
 
 			if (b16Bit)
 			{
diff --git a/Engine/Source/Runtime/Renderer/Private/CompositionLighting/CompositionLighting.cpp b/Engine/Source/Runtime/Renderer/Private/CompositionLighting/CompositionLighting.cpp
--- a/Engine/Source/Runtime/Renderer/Private/CompositionLighting/CompositionLighting.cpp
+++ b/Engine/Source/Runtime/Renderer/Private/CompositionLighting/CompositionLighting.cpp
@@ -389,25 +389,23 @@ void FCompositionLighting::ProcessAfterLighting(FRHICommandListImmediate& RHICmd
 			if (bScreenSpaceSubsurfacePassNeeded && Radius > 0 && !bSimpleDynamicLighting && View.Family->EngineShowFlags.SubsurfaceScattering)
 			{
 				bool bHalfRes = CVarSSSHalfRes.GetValueOnRenderThread() != 0;
-
-				// can be optimized out if we don't do split screen/stereo rendering (should be done after we some post process refactoring)
-				FRenderingCompositePass* PassExtractSpecular = Context.Graph.RegisterPass(new(FMemStack::Get()) FRCPassPostProcessSubsurfaceExtractSpecular());
-				PassExtractSpecular->SetInput(ePId_Input0, Context.FinalOutput);
+				bool bSingleViewportMode = View.Family->Views.Num() == 1;
 
 				FRenderingCompositePass* PassSetup = Context.Graph.RegisterPass(new(FMemStack::Get()) FRCPassPostProcessSubsurfaceSetup(View, bHalfRes));
 				PassSetup->SetInput(ePId_Input0, Context.FinalOutput);
 
-				FRenderingCompositePass* Pass0 = Context.Graph.RegisterPass(new(FMemStack::Get()) FRCPassPostProcessSubsurface(0, bHalfRes));
-				Pass0->SetInput(ePId_Input0, PassSetup);
+				FRenderingCompositePass* PassX = Context.Graph.RegisterPass(new(FMemStack::Get()) FRCPassPostProcessSubsurface(0, bHalfRes));
+				PassX->SetInput(ePId_Input0, PassSetup);
 
-				FRenderingCompositePass* Pass1 = Context.Graph.RegisterPass(new(FMemStack::Get()) FRCPassPostProcessSubsurface(1, bHalfRes));
-				Pass1->SetInput(ePId_Input0, Pass0);
-				Pass1->SetInput(ePId_Input1, PassSetup);
+				FRenderingCompositePass* PassY = Context.Graph.RegisterPass(new(FMemStack::Get()) FRCPassPostProcessSubsurface(1, bHalfRes));
+				PassY->SetInput(ePId_Input0, PassX);
+				PassY->SetInput(ePId_Input1, PassSetup);
 
-				// full res composite pass, no blurring (Radius=0)
-				FRenderingCompositePass* RecombinePass = Context.Graph.RegisterPass(new(FMemStack::Get()) FRCPassPostProcessSubsurfaceRecombine(bHalfRes));
-				RecombinePass->SetInput(ePId_Input0, Pass1);
-				RecombinePass->SetInput(ePId_Input1, PassExtractSpecular);
+				// full res composite pass, no blurring (Radius=0), replaces SceneColor
+				FRenderingCompositePass* RecombinePass = Context.Graph.RegisterPass(new(FMemStack::Get()) FRCPassPostProcessSubsurfaceRecombine(bHalfRes, bSingleViewportMode));
+				RecombinePass->SetInput(ePId_Input0, Context.FinalOutput);
+				RecombinePass->SetInput(ePId_Input1, PassY);
+				RecombinePass->SetInput(ePId_Input2, PassSetup);
 			}
 		}
 
diff --git a/Engine/Source/Runtime/Renderer/Private/CompositionLighting/PostProcessPassThrough.cpp b/Engine/Source/Runtime/Renderer/Private/CompositionLighting/PostProcessPassThrough.cpp
--- a/Engine/Source/Runtime/Renderer/Private/CompositionLighting/PostProcessPassThrough.cpp
+++ b/Engine/Source/Runtime/Renderer/Private/CompositionLighting/PostProcessPassThrough.cpp
@@ -143,7 +143,6 @@ void FRCPassPostProcessPassThrough::Process(FRenderingCompositePassContext& Cont
 	Context.RHICmdList.CopyToResolveTarget(DestRenderTarget.TargetableTexture, DestRenderTarget.ShaderResourceTexture, false, FResolveParams());
 }
 
-
 FPooledRenderTargetDesc FRCPassPostProcessPassThrough::ComputeOutputDesc(EPassOutputId InPassOutputId) const
 {
 	FPooledRenderTargetDesc Ret;
@@ -164,3 +163,50 @@ FPooledRenderTargetDesc FRCPassPostProcessPassThrough::ComputeOutputDesc(EPassOu
 
 	return Ret;
 }
+
+void CopyOverOtherViewportsIfNeeded(FRenderingCompositePassContext& Context, const FSceneView& ExcludeView)
+{
+	const FSceneView& View = Context.View;
+	const FSceneViewFamily* ViewFamily = View.Family;
+
+	// only for multiple views we need to do this
+	if(ViewFamily->Views.Num())
+	{
+		SCOPED_DRAW_EVENT(Context.RHICmdList, CopyOverOtherViewportsIfNeeded);
+
+		TShaderMapRef<FPostProcessVS> VertexShader(Context.GetShaderMap());
+		TShaderMapRef<FPostProcessPassThroughPS> PixelShader(Context.GetShaderMap());
+
+		static FGlobalBoundShaderState BoundShaderState;
+
+		SetGlobalBoundShaderState(Context.RHICmdList, Context.GetFeatureLevel(), BoundShaderState, GFilterVertexDeclaration.VertexDeclarationRHI, *VertexShader, *PixelShader);
+
+		VertexShader->SetParameters(Context);
+		PixelShader->SetParameters(Context);
+
+		FSceneRenderTargets& SceneContext = FSceneRenderTargets::Get(Context.RHICmdList);
+
+		FIntPoint Size = SceneContext.GetBufferSizeXY();
+	
+		for(uint32 ViewId = 0, ViewCount = ViewFamily->Views.Num(); ViewId < ViewCount; ++ViewId)
+		{
+			const FSceneView& LocalView = *ViewFamily->Views[ViewId];
+			
+			if(&LocalView != &ExcludeView)
+			{
+				FIntRect Rect = LocalView.ViewRect;
+
+				DrawRectangle(
+					Context.RHICmdList,
+					Rect.Min.X, Rect.Min.Y,
+					Rect.Width(), Rect.Height(),
+					Rect.Min.X, Rect.Min.Y,
+					Rect.Width(), Rect.Height(),
+					Size,
+					Size,
+					*VertexShader,
+					EDRF_UseTriangleOptimization);
+			}
+		}
+	}
+}
diff --git a/Engine/Source/Runtime/Renderer/Private/CompositionLighting/PostProcessPassThrough.h b/Engine/Source/Runtime/Renderer/Private/CompositionLighting/PostProcessPassThrough.h
--- a/Engine/Source/Runtime/Renderer/Private/CompositionLighting/PostProcessPassThrough.h
+++ b/Engine/Source/Runtime/Renderer/Private/CompositionLighting/PostProcessPassThrough.h
@@ -31,3 +31,12 @@ private:
 	bool bAdditiveBlend;
 	FPooledRenderTargetDesc NewDesc;
 };
+
+// useful to replace the SceneColor and leave the other views e.g. when doing stereo rendering, assumed there is no scaling on the size (input and output is BufferSize)
+// call before SetViewport()
+// assumes this to be called before
+//  SetRenderTarget()
+//  Context.RHICmdList.SetBlendState(TStaticBlendState<>::GetRHI());
+//  Context.RHICmdList.SetRasterizerState(TStaticRasterizerState<>::GetRHI());
+// 	Context.RHICmdList.SetDepthStencilState(TStaticDepthStencilState<false, CF_Always>::GetRHI());
+void CopyOverOtherViewportsIfNeeded(FRenderingCompositePassContext& Context, const FSceneView& ExcludeView);
\ No newline at end of file
diff --git a/Engine/Source/Runtime/Renderer/Private/DeferredShadingRenderer.cpp b/Engine/Source/Runtime/Renderer/Private/DeferredShadingRenderer.cpp
--- a/Engine/Source/Runtime/Renderer/Private/DeferredShadingRenderer.cpp
+++ b/Engine/Source/Runtime/Renderer/Private/DeferredShadingRenderer.cpp
@@ -782,6 +782,10 @@ static void ServiceLocalQueue()
 void FDeferredShadingSceneRenderer::Render(FRHICommandListImmediate& RHICmdList)
 {
 	FSceneRenderTargets& SceneContext = FSceneRenderTargets::Get(RHICmdList);
+
+	// this way we make sure the SceneColor format is the correct one and not the one from the end of frame before
+	SceneContext.ReleaseSceneColor();
+
 	bool bDBuffer = IsDBufferEnabled();	
 
 	if (GRHIThread)
diff --git a/Engine/Source/Runtime/Renderer/Private/PostProcess/PostProcessSubsurface.cpp b/Engine/Source/Runtime/Renderer/Private/PostProcess/PostProcessSubsurface.cpp
--- a/Engine/Source/Runtime/Renderer/Private/PostProcess/PostProcessSubsurface.cpp
+++ b/Engine/Source/Runtime/Renderer/Private/PostProcess/PostProcessSubsurface.cpp
@@ -69,11 +69,6 @@ public:
 			// can be extracted out of projection matrix
 
 			FSceneRenderTargets& SceneContext = FSceneRenderTargets::Get(RHICmdList);
-			float ScaleCorrectionX = Context.View.ViewRect.Width() / (float)SceneContext.GetBufferSizeXY().X;
-			float ScaleCorrectionY = Context.View.ViewRect.Height() / (float)SceneContext.GetBufferSizeXY().Y;
-
-			 // Divide by 3 as the kernels range from -3 to 3.
-			const float KernelSize = 3.0f;
 
 			// Calculate the sssWidth scale (1.0 for a unit plane sitting on the projection window):
 			float DistanceToProjectionWindow = Context.View.ViewMatrices.ProjMatrix.M[0][0]; 
@@ -81,9 +76,9 @@ public:
 			float SSSScaleZ = DistanceToProjectionWindow * GetSubsurfaceRadiusScale();
 
 			// * 0.5f: hacked in 0.5 - -1..1 to 0..1 but why this isn't in demo code?
-			float SSSScaleX = SSSScaleZ / KernelSize * 0.5f;
+			float SSSScaleX = SSSScaleZ / SUBSURFACE_KERNEL_SIZE * 0.5f;
 
-			FVector4 ColorScale(SSSScaleX, SSSScaleZ, ScaleCorrectionX, ScaleCorrectionY);
+			FVector4 ColorScale(SSSScaleX, SSSScaleZ, 0, 0);
 			SetShaderValue(Context.RHICmdList, ShaderRHI, SSSParams, ColorScale);
 		}
 
@@ -131,6 +126,8 @@ class FPostProcessSubsurfaceVisualizePS : public FGlobalShader
 	static void ModifyCompilationEnvironment(EShaderPlatform Platform, FShaderCompilerEnvironment& OutEnvironment)
 	{
 		FGlobalShader::ModifyCompilationEnvironment(Platform, OutEnvironment);
+		OutEnvironment.SetDefine(TEXT("SUBSURFACE_RADIUS_SCALE"), SUBSURFACE_RADIUS_SCALE);
+		OutEnvironment.SetDefine(TEXT("SUBSURFACE_KERNEL_SIZE"), SUBSURFACE_KERNEL_SIZE);
 	}
 
 	/** Default constructor. */
@@ -344,6 +341,8 @@ class FPostProcessSubsurfaceSetupPS : public FGlobalShader
 		FGlobalShader::ModifyCompilationEnvironment(Platform, OutEnvironment);
 		OutEnvironment.SetDefine(TEXT("SETUP_MODE"), SetupMode);
 		OutEnvironment.SetDefine(TEXT("HALF_RES"), HalfRes);
+		OutEnvironment.SetDefine(TEXT("SUBSURFACE_RADIUS_SCALE"), SUBSURFACE_RADIUS_SCALE);
+		OutEnvironment.SetDefine(TEXT("SUBSURFACE_KERNEL_SIZE"), SUBSURFACE_KERNEL_SIZE);
 	}
 
 	/** Default constructor. */
@@ -455,6 +454,8 @@ class FPostProcessSubsurfaceExtractSpecularPS : public FGlobalShader
 		FGlobalShader::ModifyCompilationEnvironment(Platform, OutEnvironment);
 		OutEnvironment.SetDefine(TEXT("SSS_RECOMBINE_METHOD"), RecombineMethod);
 		OutEnvironment.SetDefine(TEXT("SSS_SAMPLESET"), SampleSet);
+		OutEnvironment.SetDefine(TEXT("SUBSURFACE_RADIUS_SCALE"), SUBSURFACE_RADIUS_SCALE);
+		OutEnvironment.SetDefine(TEXT("SUBSURFACE_KERNEL_SIZE"), SUBSURFACE_KERNEL_SIZE);
 	}
 
 	/** Default constructor. */
@@ -555,85 +556,6 @@ void SetSubsurfaceExtractSpecular(const FRenderingCompositePassContext& Context,
 	VertexShader->SetParameters(Context);
 }
 
-void FRCPassPostProcessSubsurfaceExtractSpecular::Process(FRenderingCompositePassContext& Context)
-{
-	const FPooledRenderTargetDesc* InputDesc = GetInputDesc(ePId_Input0);
-
-	check(InputDesc);
-
-	const FSceneView& View = Context.View;
-	const FSceneViewFamily& ViewFamily = *(View.Family);
-
-	FIntPoint SrcSize = InputDesc->Extent;
-	FIntPoint DestSize = PassOutputs[0].RenderTargetDesc.Extent;
-
-	check(DestSize.X);
-	check(DestSize.Y);
-	check(SrcSize.X);
-	check(SrcSize.Y);
-
-	//	FIntRect SrcRect = View.ViewRect / ScaleFactor;
-	FIntRect SrcRect = View.ViewRect;
-	FIntRect DestRect = View.ViewRect;
-
-	TRefCountPtr<IPooledRenderTarget> NewSceneColor;
-
-	const FSceneRenderTargetItem& DestRenderTarget = PassOutputs[0].RequestSurface(Context);
-
-	// Set the view family's render target/viewport.
-	SetRenderTarget(Context.RHICmdList, DestRenderTarget.TargetableTexture, FTextureRHIRef());
-
-	Context.SetViewportAndCallRHI(0, 0, 0.0f, DestSize.X, DestSize.Y, 1.0f );
-
-	Context.RHICmdList.SetBlendState(TStaticBlendState<>::GetRHI());
-	Context.RHICmdList.SetRasterizerState(TStaticRasterizerState<>::GetRHI());
-	Context.RHICmdList.SetDepthStencilState(TStaticDepthStencilState<false, CF_Always>::GetRHI());
-
-	TShaderMapRef<FPostProcessVS> VertexShader(Context.GetShaderMap());
-
-	bool bDoSpecularCorrection = DoSpecularCorrection();
-
-	SCOPED_DRAW_EVENTF(Context.RHICmdList, SubsurfacePass, TEXT("SubsurfacePassExtractSpecular#%d"), (int32)bDoSpecularCorrection);
-
-	uint32 SampleSet = FMath::Clamp(CVarSSSSampleSet.GetValueOnRenderThread(), 0, 2);
-
-	if(bDoSpecularCorrection)
-	{
-		SetSubsurfaceExtractSpecular<1>(Context, SampleSet);
-	}
-	else
-	{
-		SetSubsurfaceExtractSpecular<0>(Context, SampleSet);
-	}
-
-	DrawRectangle(
-		Context.RHICmdList,
-		DestRect.Min.X, DestRect.Min.Y,
-		DestRect.Width(), DestRect.Height(),
-		SrcRect.Min.X, SrcRect.Min.Y,
-		SrcRect.Width(), SrcRect.Height(),
-		DestSize,
-		SrcSize,
-		*VertexShader,
-		EDRF_UseTriangleOptimization);
-
-	Context.RHICmdList.CopyToResolveTarget(DestRenderTarget.TargetableTexture, DestRenderTarget.ShaderResourceTexture, false, FResolveParams());
-}
-
-
-FPooledRenderTargetDesc FRCPassPostProcessSubsurfaceExtractSpecular::ComputeOutputDesc(EPassOutputId InPassOutputId) const
-{
-	FPooledRenderTargetDesc Ret = FSceneRenderTargets::Get_Todo_PassContext().GetSceneColor()->GetDesc();
-
-	Ret.Reset();
-	Ret.DebugName = TEXT("SubsurfaceExtractSpecular");
-	// alpha is not used
-	Ret.Format = PF_FloatRGB;
-
-	return Ret;
-}
-
-
 // --------------------------------------
 
 FRCPassPostProcessSubsurfaceSetup::FRCPassPostProcessSubsurfaceSetup(FViewInfo& View, bool bInHalfRes)
@@ -750,6 +672,8 @@ class TPostProcessSubsurfacePS : public FGlobalShader
 		FGlobalShader::ModifyCompilationEnvironment(Platform,OutEnvironment);
 		OutEnvironment.SetDefine(TEXT("SSS_DIRECTION"), Direction);
 		OutEnvironment.SetDefine(TEXT("SSS_SAMPLESET"), SampleSet);
+		OutEnvironment.SetDefine(TEXT("SUBSURFACE_RADIUS_SCALE"), SUBSURFACE_RADIUS_SCALE);
+		OutEnvironment.SetDefine(TEXT("SUBSURFACE_KERNEL_SIZE"), SUBSURFACE_KERNEL_SIZE);
 	}
 
 	/** Default constructor. */
@@ -959,6 +883,8 @@ class TPostProcessSubsurfaceRecombinePS : public FGlobalShader
 		FGlobalShader::ModifyCompilationEnvironment(Platform,OutEnvironment);
 		OutEnvironment.SetDefine(TEXT("SSS_RECOMBINE_METHOD"), RecombineMethod);
 		OutEnvironment.SetDefine(TEXT("HALF_RES"), HalfRes);
+		OutEnvironment.SetDefine(TEXT("SUBSURFACE_RADIUS_SCALE"), SUBSURFACE_RADIUS_SCALE);
+		OutEnvironment.SetDefine(TEXT("SUBSURFACE_KERNEL_SIZE"), SUBSURFACE_KERNEL_SIZE);
 	}
 
 	/** Default constructor. */
@@ -1028,8 +954,9 @@ void SetSubsurfaceRecombineShader(const FRenderingCompositePassContext& Context,
 	VertexShader->SetParameters(Context);
 }
 
-FRCPassPostProcessSubsurfaceRecombine::FRCPassPostProcessSubsurfaceRecombine(bool bInHalfRes)
+FRCPassPostProcessSubsurfaceRecombine::FRCPassPostProcessSubsurfaceRecombine(bool bInHalfRes, bool bInSingleViewportMode)
 	: bHalfRes(bInHalfRes)
+	, bSingleViewportMode(bInSingleViewportMode)
 {
 }
 
@@ -1056,17 +983,19 @@ void FRCPassPostProcessSubsurfaceRecombine::Process(FRenderingCompositePassConte
 
 	TRefCountPtr<IPooledRenderTarget>& SceneColor = SceneContext.GetSceneColor();
 
-	const FSceneRenderTargetItem& DestRenderTarget = SceneColor->GetRenderTargetItem();
+	const FSceneRenderTargetItem& DestRenderTarget = PassOutputs[0].RequestSurface(Context);
 
 	// Set the view family's render target/viewport.
 	SetRenderTarget(Context.RHICmdList, DestRenderTarget.TargetableTexture, FTextureRHIRef());
 
-	Context.SetViewportAndCallRHI(0, 0, 0.0f, DestSize.X, DestSize.Y, 1.0f );
-
-	Context.RHICmdList.SetBlendState(TStaticBlendState<CW_RGBA,BO_Add,BF_SourceAlpha,BF_InverseSourceAlpha,BO_Add,BF_SourceAlpha,BF_InverseSourceAlpha>::GetRHI());
+	Context.RHICmdList.SetBlendState(TStaticBlendState<>::GetRHI());
 	Context.RHICmdList.SetRasterizerState(TStaticRasterizerState<>::GetRHI());
 	Context.RHICmdList.SetDepthStencilState(TStaticDepthStencilState<false, CF_Always>::GetRHI());
 
+	CopyOverOtherViewportsIfNeeded(Context, View);
+
+	Context.SetViewportAndCallRHI(0, 0, 0.0f, DestSize.X, DestSize.Y, 1.0f );
+
 	TShaderMapRef<FPostProcessVS> VertexShader(Context.GetShaderMap());
 
 	bool bDoSpecularCorrection = DoSpecularCorrection();
@@ -1094,14 +1023,26 @@ void FRCPassPostProcessSubsurfaceRecombine::Process(FRenderingCompositePassConte
 		EDRF_UseTriangleOptimization);
 
 	Context.RHICmdList.CopyToResolveTarget(DestRenderTarget.TargetableTexture, DestRenderTarget.ShaderResourceTexture, false, FResolveParams());
+
+	// replace the current SceneColor with this one
+	SceneContext.SetSceneColor(PassOutputs[0].PooledRenderTarget);
+	PassOutputs[0].PooledRenderTarget.SafeRelease();
 }
 
 
 FPooledRenderTargetDesc FRCPassPostProcessSubsurfaceRecombine::ComputeOutputDesc(EPassOutputId InPassOutputId) const
 {
-	FPooledRenderTargetDesc Ret;
+	FPooledRenderTargetDesc Ret = GetInput(ePId_Input0)->GetOutput()->RenderTargetDesc;
+
+	Ret.Reset();
+	Ret.DebugName = TEXT("SceneColorSubsurface");
+
+	if(bSingleViewportMode)
+	{
+		// we don't need an alpha channel any more as it was used for ScreenSpaceSubsurfaceScattering only
+		Ret.Format = (Ret.Format == PF_FloatRGBA) ? PF_FloatRGB : PF_FloatRGBA;
+	}
 
-	Ret.DebugName = TEXT("SceneColor");
-	// we directly render to the HDR scene color
+	// we replace the HDR SceneColor with this one
 	return Ret;
 }
diff --git a/Engine/Source/Runtime/Renderer/Private/PostProcess/PostProcessSubsurface.h b/Engine/Source/Runtime/Renderer/Private/PostProcess/PostProcessSubsurface.h
--- a/Engine/Source/Runtime/Renderer/Private/PostProcess/PostProcessSubsurface.h
+++ b/Engine/Source/Runtime/Renderer/Private/PostProcess/PostProcessSubsurface.h
@@ -23,18 +23,6 @@ public:
 	virtual void Release() override { delete this; }
 };
 
-// ePId_Input0: SceneColor
-// derives from TRenderingCompositePassBase<InputCount, OutputCount> 
-// uses some GBuffer attributes
-class FRCPassPostProcessSubsurfaceExtractSpecular : public TRenderingCompositePassBase<1, 1>
-{
-public:
-	// interface FRenderingCompositePass ---------
-	virtual void Process(FRenderingCompositePassContext& Context) override;
-	virtual FPooledRenderTargetDesc ComputeOutputDesc(EPassOutputId InPassOutputId) const override;
-	virtual void Release() override { delete this; }
-};
-
 // ePId_Input0: SceneColor
 // derives from TRenderingCompositePassBase<InputCount, OutputCount> 
 // uses some GBuffer attributes
@@ -80,14 +68,15 @@ private:
 
 
 // derives from TRenderingCompositePassBase<InputCount, OutputCount>
-// ePId_Input0: output from FRCPassPostProcessSubsurface
-// ePId_Input1: SceneColor before Screen Space Subsurface input
+// ePId_Input0: SceneColor before Screen Space Subsurface input
+// ePId_Input1: output from FRCPassPostProcessSubsurface
+// ePId_Input2: SubsurfaceSetup, can be half res
 // modifies SceneColor, uses some GBuffer attributes
-class FRCPassPostProcessSubsurfaceRecombine : public TRenderingCompositePassBase<2, 1>
+class FRCPassPostProcessSubsurfaceRecombine : public TRenderingCompositePassBase<3, 1>
 {
 public:
 	// constructor
-	FRCPassPostProcessSubsurfaceRecombine(bool bInHalfRes);
+	FRCPassPostProcessSubsurfaceRecombine(bool bInHalfRes, bool bInSingleViewportMode);
 
 	// interface FRenderingCompositePass ---------
 
@@ -96,4 +85,5 @@ public:
 	virtual FPooledRenderTargetDesc ComputeOutputDesc(EPassOutputId InPassOutputId) const override;
 
 	bool bHalfRes;
+	bool bSingleViewportMode;
 };
diff --git a/Engine/Source/Runtime/Renderer/Private/PostProcess/SceneRenderTargets.cpp b/Engine/Source/Runtime/Renderer/Private/PostProcess/SceneRenderTargets.cpp
--- a/Engine/Source/Runtime/Renderer/Private/PostProcess/SceneRenderTargets.cpp
+++ b/Engine/Source/Runtime/Renderer/Private/PostProcess/SceneRenderTargets.cpp
@@ -1788,14 +1788,19 @@ void FSceneRenderTargets::AllocateRenderTargets()
 	}
 }
 
-void FSceneRenderTargets::ReleaseAllTargets()
+void FSceneRenderTargets::ReleaseSceneColor()
 {
-	ReleaseGBufferTargets();
-
 	for (auto i = 0; i < (int32)EShadingPath::Num; ++i)
 	{
 		SceneColor[i].SafeRelease();
 	}
+}
+
+void FSceneRenderTargets::ReleaseAllTargets()
+{
+	ReleaseGBufferTargets();
+
+	ReleaseSceneColor();
 
 	SceneAlphaCopy.SafeRelease();
 	SceneDepthZ.SafeRelease();
diff --git a/Engine/Source/Runtime/Renderer/Private/PostProcess/SceneRenderTargets.h b/Engine/Source/Runtime/Renderer/Private/PostProcess/SceneRenderTargets.h
--- a/Engine/Source/Runtime/Renderer/Private/PostProcess/SceneRenderTargets.h
+++ b/Engine/Source/Runtime/Renderer/Private/PostProcess/SceneRenderTargets.h
@@ -498,7 +498,10 @@ public:
 	{
 		bVelocityPass = bInVelocityPass;
 	}
-
+	
+	// Can be called when the Scene Color content is no longer needed. As we create SceneColor on demand we can make sure it is created with the right format.
+	// (as a call to SetSceneColor() can override it with a different format)
+	void ReleaseSceneColor();
 
 private: // Get...() methods instead of direct access
 
