Fixed crash with DFAO surface cache

diff --git a/Engine/Source/Runtime/Renderer/Private/DistanceFieldSurfaceCacheLighting.cpp b/Engine/Source/Runtime/Renderer/Private/DistanceFieldSurfaceCacheLighting.cpp
--- a/Engine/Source/Runtime/Renderer/Private/DistanceFieldSurfaceCacheLighting.cpp
+++ b/Engine/Source/Runtime/Renderer/Private/DistanceFieldSurfaceCacheLighting.cpp
@@ -2368,7 +2368,7 @@ void RenderIrradianceCacheInterpolation(
 			FRHISetRenderTargetsInfo Info(MRTCount, RenderTargets, DepthView);
 
 			check(RenderTargets[0].Texture->GetClearColor() == FLinearColor::Transparent);
-			check(RenderTargets[1].Texture->GetClearColor() == FLinearColor::Transparent);
+			check(!RenderTargets[1].Texture || RenderTargets[1].Texture->GetClearColor() == FLinearColor::Transparent);
 			Info.bClearColor = true;
 						
 
