[BuildWatcher] Fixing CIS on the Mac

diff --git a/Engine/Source/Runtime/Renderer/Private/PostProcess/PostProcessSubsurface.cpp b/Engine/Source/Runtime/Renderer/Private/PostProcess/PostProcessSubsurface.cpp
--- a/Engine/Source/Runtime/Renderer/Private/PostProcess/PostProcessSubsurface.cpp
+++ b/Engine/Source/Runtime/Renderer/Private/PostProcess/PostProcessSubsurface.cpp
@@ -10,6 +10,7 @@
 #include "PostProcessSubsurface.h"
 #include "PostProcessing.h"
 #include "SceneUtils.h"
+#include "CompositionLighting/PostProcessPassThrough.h"
 
 ENGINE_API const IPooledRenderTarget* GetSubsufaceProfileTexture_RT(FRHICommandListImmediate& RHICmdList);
 
