Global distance field is now updated properly upon primitive movement - previously wasn't taking into account the max occlusion influence distance

diff --git a/Engine/Source/Runtime/Renderer/Private/GlobalDistanceField.cpp b/Engine/Source/Runtime/Renderer/Private/GlobalDistanceField.cpp
--- a/Engine/Source/Runtime/Renderer/Private/GlobalDistanceField.cpp
+++ b/Engine/Source/Runtime/Renderer/Private/GlobalDistanceField.cpp
@@ -481,9 +481,10 @@ void AddUpdateRegionForAxis(FIntVector Movement, const FBox& ClipmapBounds, floa
 }
 
 /** Constructs and adds an update region based on the given primitive bounds. */
-void AddUpdateRegionForPrimitive(const FVector4& Bounds, const FBox& ClipmapBounds, float CellSize, TArray<FVolumeUpdateRegion, TInlineAllocator<3> >& UpdateRegions)
+void AddUpdateRegionForPrimitive(const FVector4& Bounds, float MaxSphereQueryRadius, const FBox& ClipmapBounds, float CellSize, TArray<FVolumeUpdateRegion, TInlineAllocator<3> >& UpdateRegions)
 {
-	FBox BoundingBox((FVector)Bounds - Bounds.W, (FVector)Bounds + Bounds.W);
+	// Object influence bounds
+	FBox BoundingBox((FVector)Bounds - Bounds.W - MaxSphereQueryRadius, (FVector)Bounds + Bounds.W + MaxSphereQueryRadius);
 
 	FVolumeUpdateRegion UpdateRegion;
 	UpdateRegion.Bounds.Init();
@@ -493,7 +494,7 @@ void AddUpdateRegionForPrimitive(const FVector4& Bounds, const FBox& ClipmapBoun
 	UpdateRegion.Bounds.Min.Z = FMath::Max(CellSize * FMath::FloorToFloat(BoundingBox.Min.Z / CellSize), ClipmapBounds.Min.Z);
 
 	// Derive the max from the snapped min and size, clamp to clipmap bounds
-	UpdateRegion.Bounds.Max = UpdateRegion.Bounds.Min + FVector(FMath::CeilToFloat(Bounds.W * 2 / CellSize)) * CellSize;
+	UpdateRegion.Bounds.Max = UpdateRegion.Bounds.Min + FVector(FMath::CeilToFloat((Bounds.W + MaxSphereQueryRadius) * 2 / CellSize)) * CellSize;
 	UpdateRegion.Bounds.Max.X = FMath::Min(UpdateRegion.Bounds.Max.X, ClipmapBounds.Max.X);
 	UpdateRegion.Bounds.Max.Y = FMath::Min(UpdateRegion.Bounds.Max.Y, ClipmapBounds.Max.Y);
 	UpdateRegion.Bounds.Max.Z = FMath::Min(UpdateRegion.Bounds.Max.Z, ClipmapBounds.Max.Z);
@@ -633,10 +634,13 @@ void ComputeUpdateRegionsAndUpdateViewState(const FViewInfo& View, FGlobalDistan
 					AddUpdateRegionForAxis(Movement, ClipmapBounds, CellSize, 1, Clipmap.UpdateRegions);
 					AddUpdateRegionForAxis(Movement, ClipmapBounds, CellSize, 2, Clipmap.UpdateRegions);
 
+					extern float GAOConeHalfAngle;
+					const float GlobalMaxSphereQueryRadius = MaxOcclusionDistance / (1 + FMath::Tan(GAOConeHalfAngle));
+
 					for (int32 BoundsIndex = 0; BoundsIndex < ClipmapViewState.PrimitiveModifiedBounds.Num(); BoundsIndex++)
 					{
 						// Add an update region for each primitive that has been modified
-						AddUpdateRegionForPrimitive(ClipmapViewState.PrimitiveModifiedBounds[BoundsIndex], ClipmapBounds, CellSize, Clipmap.UpdateRegions);
+						AddUpdateRegionForPrimitive(ClipmapViewState.PrimitiveModifiedBounds[BoundsIndex], GlobalMaxSphereQueryRadius, ClipmapBounds, CellSize, Clipmap.UpdateRegions);
 					}
 
 					int32 TotalTexelsBeingUpdated = 0;
