Removed unnecessary scissor set. UE-3623

diff --git a/Engine/Source/Runtime/Renderer/Private/PostProcess/SceneRenderTargets.cpp b/Engine/Source/Runtime/Renderer/Private/PostProcess/SceneRenderTargets.cpp
--- a/Engine/Source/Runtime/Renderer/Private/PostProcess/SceneRenderTargets.cpp
+++ b/Engine/Source/Runtime/Renderer/Private/PostProcess/SceneRenderTargets.cpp
@@ -1771,12 +1771,6 @@ void FSceneRenderTargets::AllocateRenderTargets()
 {
 	if (BufferSize.X > 0 && BufferSize.Y > 0 && !AreShadingPathRenderTargetsAllocated(CurrentShadingPath))
 	{
-		// start with a defined state for the scissor rect (D3D11 was returning (0,0,0,0) which caused a clear to not execute correctly)
-		// todo: move this to an earlier place (for dx9 is has to be after device creation which is after window creation)
-		// todo: potentially redundant now? Seems like a strange thing to have here anyway???
-		FRHICommandListImmediate& RHICmdList = FRHICommandListExecutor::GetImmediateCommandList();
-		RHICmdList.SetScissorRect(false, 0, 0, 0, 0);
-
 		if ((EShadingPath)CurrentShadingPath == EShadingPath::Forward)
 		{
 			AllocateForwardShadingPathRenderTargets();
diff --git a/Engine/Source/Runtime/Renderer/Private/SystemTextures.cpp b/Engine/Source/Runtime/Renderer/Private/SystemTextures.cpp
--- a/Engine/Source/Runtime/Renderer/Private/SystemTextures.cpp
+++ b/Engine/Source/Runtime/Renderer/Private/SystemTextures.cpp
@@ -22,10 +22,6 @@ void FSystemTextures::InitializeTextures(FRHICommandListImmediate& RHICmdList, E
 		return;
 	}
 
-	// start with a defined state for the scissor rect (D3D11 was returning (0,0,0,0) which caused a clear to not execute correctly)
-	// todo: move this to an earlier place (for dx9 is has to be after device creation which is after window creation)
-	RHICmdList.SetScissorRect(false, 0, 0, 0, 0);
-
 	// First initialize textures that are common to all feature levels. This is always done the first time we come into this function, as doesn't care about the
 	// requested feature level
 	if (!bTexturesInitialized)
