removed usnused code in SubSurfaceScattering

diff --git a/Engine/Source/Runtime/Renderer/Private/PostProcess/PostProcessSubsurface.cpp b/Engine/Source/Runtime/Renderer/Private/PostProcess/PostProcessSubsurface.cpp
--- a/Engine/Source/Runtime/Renderer/Private/PostProcess/PostProcessSubsurface.cpp
+++ b/Engine/Source/Runtime/Renderer/Private/PostProcess/PostProcessSubsurface.cpp
@@ -435,130 +435,6 @@ static bool DoSpecularCorrection()
 
 // --------------------------------------
 
-
-/**
- * Encapsulates the post processing subsurface scattering pixel shader.
- * @param SampleSet 0:low, 1:med, 2:high
- */
-template <uint32 RecombineMethod, uint32 SampleSet>
-class FPostProcessSubsurfaceExtractSpecularPS : public FGlobalShader
-{
-	DECLARE_SHADER_TYPE(FPostProcessSubsurfaceExtractSpecularPS, Global);
-
-	static bool ShouldCache(EShaderPlatform Platform)
-	{
-		return IsFeatureLevelSupported(Platform, ERHIFeatureLevel::SM4);
-	}
-
-	static void ModifyCompilationEnvironment(EShaderPlatform Platform, FShaderCompilerEnvironment& OutEnvironment)
-	{
-		FGlobalShader::ModifyCompilationEnvironment(Platform, OutEnvironment);
-		OutEnvironment.SetDefine(TEXT("SSS_RECOMBINE_METHOD"), RecombineMethod);
-		OutEnvironment.SetDefine(TEXT("SSS_SAMPLESET"), SampleSet);
-		OutEnvironment.SetDefine(TEXT("SUBSURFACE_RADIUS_SCALE"), SUBSURFACE_RADIUS_SCALE);
-		OutEnvironment.SetDefine(TEXT("SUBSURFACE_KERNEL_SIZE"), SUBSURFACE_KERNEL_SIZE);
-	}
-
-	/** Default constructor. */
-	FPostProcessSubsurfaceExtractSpecularPS () {}
-
-public:
-	FPostProcessPassParameters PostprocessParameter;
-	FDeferredPixelShaderParameters DeferredParameters;
-	FSubsurfaceParameters SubsurfaceParameters;
-
-	/** Initialization constructor. */
-	FPostProcessSubsurfaceExtractSpecularPS(const ShaderMetaType::CompiledShaderInitializerType& Initializer)
-		: FGlobalShader(Initializer)
-	{
-		PostprocessParameter.Bind(Initializer.ParameterMap);
-		DeferredParameters.Bind(Initializer.ParameterMap);
-		SubsurfaceParameters.Bind(Initializer.ParameterMap);
-	}
-
-	void SetParameters(const FRenderingCompositePassContext& Context)
-	{
-		const FFinalPostProcessSettings& Settings = Context.View.FinalPostProcessSettings;
-		const FPixelShaderRHIParamRef ShaderRHI = GetPixelShader();
-
-		FGlobalShader::SetParameters(Context.RHICmdList, ShaderRHI, Context.View);
-		PostprocessParameter.SetPS(ShaderRHI, Context, TStaticSamplerState<SF_Point,AM_Clamp,AM_Clamp,AM_Clamp>::GetRHI());
-		DeferredParameters.Set(Context.RHICmdList, ShaderRHI, Context.View);
-		SubsurfaceParameters.SetParameters(Context.RHICmdList, ShaderRHI, Context);
-	}
-
-	// FShader interface.
-	virtual bool Serialize(FArchive& Ar) override
-	{
-		bool bShaderHasOutdatedParameters = FGlobalShader::Serialize(Ar);
-		Ar << PostprocessParameter << DeferredParameters << SubsurfaceParameters;
-		return bShaderHasOutdatedParameters;
-	}
-
-	static const TCHAR* GetSourceFilename()
-	{
-		return TEXT("PostProcessSubsurface");
-	}
-
-	static const TCHAR* GetFunctionName()
-	{
-		return TEXT("ExtractSpecularPS");
-	}
-};
-
-// #define avoids a lot of code duplication
-#define VARIATION1(A)		VARIATION2(A,0)			VARIATION2(A,1)			VARIATION2(A,2)
-#define VARIATION2(A, B) typedef FPostProcessSubsurfaceExtractSpecularPS<A, B> FPostProcessSubsurfaceExtractSpecularPS##A##B; \
-	IMPLEMENT_SHADER_TYPE2(FPostProcessSubsurfaceExtractSpecularPS##A##B, SF_Pixel);
-	VARIATION1(0) VARIATION1(1)
-#undef VARIATION1
-#undef VARIATION2
-
-
-// @param SampleSet 0:low, 1:med, 2:high
-template <uint32 RecombineMethod>
-void SetSubsurfaceExtractSpecular(const FRenderingCompositePassContext& Context, uint32 SampleSet)
-{
-	TShaderMapRef<FPostProcessVS> VertexShader(Context.GetShaderMap());
-
-	switch(SampleSet)
-	{
-		case 0:
-			{
-				TShaderMapRef<FPostProcessSubsurfaceExtractSpecularPS<RecombineMethod, 0> > PixelShader(Context.GetShaderMap());
-				static FGlobalBoundShaderState BoundShaderState;
-
-				SetGlobalBoundShaderState(Context.RHICmdList, Context.GetFeatureLevel(), BoundShaderState, GFilterVertexDeclaration.VertexDeclarationRHI, *VertexShader, *PixelShader);
-				PixelShader->SetParameters(Context);
-				break;
-			}
-		case 1:
-			{
-				TShaderMapRef<FPostProcessSubsurfaceExtractSpecularPS<RecombineMethod, 1> > PixelShader(Context.GetShaderMap());
-				static FGlobalBoundShaderState BoundShaderState;
-
-				SetGlobalBoundShaderState(Context.RHICmdList, Context.GetFeatureLevel(), BoundShaderState, GFilterVertexDeclaration.VertexDeclarationRHI, *VertexShader, *PixelShader);
-				PixelShader->SetParameters(Context);
-				break;
-			}
-		case 2:
-			{
-				TShaderMapRef<FPostProcessSubsurfaceExtractSpecularPS<RecombineMethod, 2> > PixelShader(Context.GetShaderMap());
-				static FGlobalBoundShaderState BoundShaderState;
-
-				SetGlobalBoundShaderState(Context.RHICmdList, Context.GetFeatureLevel(), BoundShaderState, GFilterVertexDeclaration.VertexDeclarationRHI, *VertexShader, *PixelShader);
-				PixelShader->SetParameters(Context);
-				break;
-			} 
-		default:
-			check(0);
-	}
-
-	VertexShader->SetParameters(Context);
-}
-
-// --------------------------------------
-
 FRCPassPostProcessSubsurfaceSetup::FRCPassPostProcessSubsurfaceSetup(FViewInfo& View, bool bInHalfRes)
 	: ViewRect(View.ViewRect)
 	, bHalfRes(bInHalfRes)
