Fixed threading issue with modifying FViewInfo

diff --git a/Engine/Source/Runtime/Renderer/Private/PostProcess/PostProcessTemporalAA.cpp b/Engine/Source/Runtime/Renderer/Private/PostProcess/PostProcessTemporalAA.cpp
--- a/Engine/Source/Runtime/Renderer/Private/PostProcess/PostProcessTemporalAA.cpp
+++ b/Engine/Source/Runtime/Renderer/Private/PostProcess/PostProcessTemporalAA.cpp
@@ -724,7 +724,8 @@ void FRCPassPostProcessTemporalAA::Process(FRenderingCompositePassContext& Conte
 		ViewState->TemporalAAHistoryRT = PassOutputs[0].PooledRenderTarget;
 	}
 
-	// TODO draw separate translucency after jitter has been removed
+#if WITH_EDITOR
+	FRHICommandListExecutor::GetImmediateCommandList().ImmediateFlush(EImmediateFlushType::WaitForOutstandingTasksOnly);
 
 	// Remove jitter
 	View.ViewMatrices.RemoveTemporalJitter();
@@ -738,6 +739,7 @@ void FRCPassPostProcessTemporalAA::Process(FRenderingCompositePassContext& Conte
 	View.ViewMatrices.InvTranslatedViewProjectionMatrix = View.ViewMatrices.GetInvProjMatrix() * View.ViewMatrices.TranslatedViewMatrix.Inverse();
 
 	View.InitRHIResources(nullptr);
+#endif
 }
 
 FPooledRenderTargetDesc FRCPassPostProcessTemporalAA::ComputeOutputDesc(EPassOutputId InPassOutputId) const
