Improved SubsurfaceScatteringProfile quality, base color is now applied after scattering which gives crisper textures and more control.

This is more correct for scanned textures as they include scattering already. Affects the look of emissive lighting (rarely needed and can be fixed by content). All diffuse lighting for SubsurfaceScatteringProfile is now done without the BaseColor (minor performance cost) and reading the Lighting before the SubsurfaceScattering pass would be different.

Removed the need for a SceneColor alpha channel as lighting is now stored in a checkerboard pattern for diffuse and specular separate. Specular and diffuse lighting on Subsurface materials can appear a bit less crips.

This means the use of 32bit SceneColor (for better performance) is now possible (we might change the defaults soon).

diff --git a/Engine/Shaders/DeferredShadingCommon.usf b/Engine/Shaders/DeferredShadingCommon.usf
--- a/Engine/Shaders/DeferredShadingCommon.usf
+++ b/Engine/Shaders/DeferredShadingCommon.usf
@@ -240,8 +240,10 @@ struct FGBufferData
 	float3 DiffuseColor;
 	// 0..1 (derived from BaseColor, Metalness, Specular)
 	float3 SpecularColor;
-	// 0..1
+	// 0..1, white for SHADINGMODELID_SUBSURFACE_PROFILE (apply BaseColor after scattering is more correct and less blurry)
 	float3 BaseColor;
+	// 0..1, only needed by SHADINGMODELID_SUBSURFACE_PROFILE which applied BaseColor later
+	float3 StoredBaseColor;
 	// 0..1
 	float Metallic;
 	// 0..1
@@ -361,7 +363,27 @@ void EncodeGBuffer(
 #endif
 }
 
+// High frequency Checkerboard pattern
+// @param UVSceneColor at pixel center
+// @return true/false, todo: profile if float 0/1 would be better (need to make sure it's 100% the same)
+bool CheckerFromSceneColorUV(float2 UVSceneColor)
+{
+	// relative to left top of the rendertarget (not viewport)
+	uint2 PixelPos = uint2(UVSceneColor * View.BufferSizeAndInvSize.xy);
+
+	return (PixelPos.x + PixelPos.y) % 2;
+}
+
+// High frequency Checkerboard pattern
+// @param PixelPos relative to left top of the rendertarget (not viewport)
+// @return true/false, todo: profile if float 0/1 would be better (need to make sure it's 100% the same)
+bool CheckerFromPixelPos(uint2 PixelPos)
+{
+	return (PixelPos.x + PixelPos.y) % 2;
+}
+
 /** Populates FGBufferData */
+// @param bChecker High frequency Checkerboard pattern computed with one of the CheckerFrom.. functions, todo: profile if float 0/1 would be better (need to make sure it's 100% the same)
 FGBufferData DecodeGBufferData(
 	float4 InGBufferA,
 	float4 InGBufferB,
@@ -372,7 +394,8 @@ FGBufferData DecodeGBufferData(
 	float CustomNativeDepth,
 	uint CustomStencil,
 	float SceneDepth,
-	bool bGetNormalizedNormal)
+	bool bGetNormalizedNormal,
+	bool bChecker)
 {
 	FGBufferData GBuffer;
 
@@ -389,7 +412,9 @@ FGBufferData DecodeGBufferData(
 	GBuffer.DecalMask = 0;
 	GBuffer.ShadingModelID = DecodeShadingModelId(Decode71(InGBufferB.a, GBuffer.DecalMask));
 	
-	GBuffer.BaseColor = DecodeBaseColor(InGBufferC.rgb);
+	GBuffer.StoredBaseColor = DecodeBaseColor(InGBufferC.rgb);
+	GBuffer.BaseColor = GBuffer.StoredBaseColor;
+
 	GBuffer.Metallic = InGBufferB.r;
 	GBuffer.Specular = InGBufferB.g;
 	GBuffer.Opacity = InGBufferD.r;
@@ -408,6 +433,13 @@ FGBufferData DecodeGBufferData(
 	GBuffer.CustomStencil = CustomStencil;
 	GBuffer.Depth = SceneDepth;
 
+	if(GBuffer.ShadingModelID == SHADINGMODELID_SUBSURFACE_PROFILE)
+	{
+		// we apply the BaseColor later, this is why we need StoredBaseColor
+		GBuffer.BaseColor = bChecker;
+		GBuffer.Specular = !bChecker;
+	}
+
 	// derived from BaseColor, Metalness, Specular
 	{
 		GBuffer.DiffuseColor = GBuffer.BaseColor - GBuffer.BaseColor * GBuffer.Metallic;
@@ -540,7 +572,7 @@ SamplerState	DirectionalOcclusionSampler;
 	Texture2D<float4> ScreenSpaceAOTextureNonMS;
 	Texture2D<float> CustomDepthTextureNonMS;
 
-	// @param PixelPos - integer pixel pos (from left top)
+	// @param PixelPos relative to left top of the rendertarget (not viewport)
 	FGBufferData GetGBufferDataUint(uint2 PixelPos, bool bGetNormalizedNormal = true)
 	{
 		float4 GBufferA = GBuffers.GBufferATextureNonMS.Load(int3(PixelPos, 0));
@@ -564,9 +596,10 @@ SamplerState	DirectionalOcclusionSampler;
 
 		float SceneDepth = CalcSceneDepth(PixelPos);
 
-		return DecodeGBufferData(GBufferA, GBufferB, GBufferC, GBufferD, GBufferE, GBufferVelocity, CustomNativeDepth, CustomStencil, SceneDepth, bGetNormalizedNormal);
+		return DecodeGBufferData(GBufferA, GBufferB, GBufferC, GBufferD, GBufferE, GBufferVelocity, CustomNativeDepth, CustomStencil, SceneDepth, bGetNormalizedNormal, CheckerFromPixelPos(PixelPos));
 	}
 
+	// @param PixelPos relative to left top of the rendertarget (not viewport)
 	FDBufferData GetDBufferData(uint2 PixelPos)
 	{
 		float4 DBufferA = DBufferATextureNonMS.Load(int3(PixelPos, 0));
@@ -576,7 +609,7 @@ SamplerState	DirectionalOcclusionSampler;
 		return DecodeDBufferData(DBufferA, DBufferB, DBufferC);
 	}
 
-	// @param PixelPos - integer pixel pos (from left top)
+	// @param PixelPos relative to left top of the rendertarget (not viewport)
 	FScreenSpaceData GetScreenSpaceDataUint(uint2 PixelPos, bool bGetNormalizedNormal = true)
 	{
 		FScreenSpaceData Out;
@@ -658,7 +691,7 @@ FGBufferData GetGBufferData(float2 UV, bool bGetNormalizedNormal = true)
 
 	float SceneDepth = CalcSceneDepth(UV);
 	
-	return DecodeGBufferData(GBufferA, GBufferB, GBufferC, GBufferD, GBufferE, GBufferVelocity, CustomNativeDepth, CustomStencil, SceneDepth, bGetNormalizedNormal);
+	return DecodeGBufferData(GBufferA, GBufferB, GBufferC, GBufferD, GBufferE, GBufferVelocity, CustomNativeDepth, CustomStencil, SceneDepth, bGetNormalizedNormal, CheckerFromSceneColorUV(UV));
 }
 
 // Minimal path for just the lighting model, used to branch around unlit pixels (skybox)
@@ -701,31 +734,32 @@ FScreenSpaceData GetScreenSpaceData(float2 UV, bool bGetNormalizedNormal = true)
 
 
 #if FEATURE_LEVEL >= FEATURE_LEVEL_SM5
-	FGBufferData GetGBufferDataMS(int2 IntUV, uint SampleIndex, bool bGetNormalizedNormal = true)
+	// @param PixelPos relative to left top of the rendertarget (not viewport)
+	FGBufferData GetGBufferDataMS(int2 PixelPos, uint SampleIndex, bool bGetNormalizedNormal = true)
 	{
-		float4 GBufferA = GBuffers.GBufferATextureMS.Load(IntUV, SampleIndex);
-		float4 GBufferB = GBuffers.GBufferBTextureMS.Load(IntUV, SampleIndex);
-		float4 GBufferC = GBuffers.GBufferCTextureMS.Load(IntUV, SampleIndex);
-		float4 GBufferD = GBuffers.GBufferDTextureMS.Load(IntUV, SampleIndex);
-		float CustomNativeDepth = CustomDepthTextureNonMS.Load(int3(IntUV, 0)).r;
-		uint CustomStencil = CustomStencilTexture.Load(int3(IntUV, 0)).g;
+		float4 GBufferA = GBuffers.GBufferATextureMS.Load(PixelPos, SampleIndex);
+		float4 GBufferB = GBuffers.GBufferBTextureMS.Load(PixelPos, SampleIndex);
+		float4 GBufferC = GBuffers.GBufferCTextureMS.Load(PixelPos, SampleIndex);
+		float4 GBufferD = GBuffers.GBufferDTextureMS.Load(PixelPos, SampleIndex);
+		float CustomNativeDepth = CustomDepthTextureNonMS.Load(int3(PixelPos, 0)).r;
+		uint CustomStencil = CustomStencilTexture.Load(int3(PixelPos, 0)).g;
 
 		#if ALLOW_STATIC_LIGHTING
-			float4 GBufferE = GBuffers.GBufferETextureMS.Load(IntUV, SampleIndex);
+			float4 GBufferE = GBuffers.GBufferETextureMS.Load(PixelPos, SampleIndex);
 		#else
 			float4 GBufferE = 1;
 		#endif
 
-		float DeviceZ = SceneDepthSurface.Load(IntUV, SampleIndex);
+		float DeviceZ = SceneDepthSurface.Load(PixelPos, SampleIndex);
 		float SceneDepth = ConvertFromDeviceZ(DeviceZ);
 
 	#if WRITES_VELOCITY_TO_GBUFFER
-		float4 GBufferVelocity = GBuffers.GBufferVelocityTextureMS.Load(IntUV, SampleIndex);
+		float4 GBufferVelocity = GBuffers.GBufferVelocityTextureMS.Load(PixelPos, SampleIndex);
 	#else
 		float4 GBufferVelocity = 0;
 	#endif
 
-		return DecodeGBufferData(GBufferA, GBufferB, GBufferC, GBufferD, GBufferE, GBufferVelocity, CustomNativeDepth, CustomStencil, SceneDepth, bGetNormalizedNormal);
+		return DecodeGBufferData(GBufferA, GBufferB, GBufferC, GBufferD, GBufferE, GBufferVelocity, CustomNativeDepth, CustomStencil, SceneDepth, bGetNormalizedNormal, CheckerFromPixelPos(PixelPos));
 	}
 
 	FGBufferData GetGBufferDataMS(float2 UV, uint SampleIndex, bool bGetNormalizedNormal = true)
diff --git a/Engine/Shaders/LightAccumulator.usf b/Engine/Shaders/LightAccumulator.usf
--- a/Engine/Shaders/LightAccumulator.usf
+++ b/Engine/Shaders/LightAccumulator.usf
@@ -17,7 +17,7 @@
 // 0 : fastest (speculars leak in SSSSS)
 // 1 : luminance specular works well (needs 64bit SceneColor, stores luminance of specular in alpha, can be optimized further in LightAccumulator_Add() ) 
 // 2 : colored specular works well (best, requires another RT for light accumulation)
-#define SUBSURFACE_CHANNEL_MODE 1
+#define SUBSURFACE_CHANNEL_MODE 0
 
 struct FLightAccumulator
 {
diff --git a/Engine/Shaders/PostProcessSubsurface.usf b/Engine/Shaders/PostProcessSubsurface.usf
--- a/Engine/Shaders/PostProcessSubsurface.usf
+++ b/Engine/Shaders/PostProcessSubsurface.usf
@@ -41,30 +41,9 @@ float GetSubsurfaceStrength(float2 UV)
 // 0:full res (slower but reference) / 1:faster half res with some quality loss
 // #define HALF_RES 1
 
-// 0:speculars leak into SSS / 1:requires alpha channel for scenecolor
-// #define SSSS_SPECULAR_CORRECTION 1
-
-// set by C++ for SubsurfaceRecombinePS and ExtractSpecularPS
-// #define SSS_RECOMBINE_METHOD 0:without reconstruct specular/1:with reconstruct specular
-
 // set by C++ for SubsurfacePS
 // #define SSS_DIRECTION 0/1
 
-// define by C++ for the setup pass
-#ifndef SETUP_MODE
-	#define SETUP_MODE 0
-#endif
-
-
-#if SETUP_MODE == 0
-	#define SSSS_SPECULAR_CORRECTION 0
-#elif SETUP_MODE == 1
-	#define SSSS_SPECULAR_CORRECTION 1
-#else
-	// doesn't matter
-	#define SSSS_SPECULAR_CORRECTION 0
-#endif
-
 #define	SSSS_N_KERNELSUBSURFACECOLOROFFSET 0
 
 #if SSS_SAMPLESET == 0
@@ -132,16 +111,6 @@ float3 GetProfileSubsurfaceColor(FGBufferData GBufferData)
 
 // ------------------------------------------
 
-
-
-#if SSSS_SPECULAR_CORRECTION == 1
-	// requires alpha channel for scenecolor
-#elif SSSS_SPECULAR_CORRECTION == 0
-	// speculars leak into SSS
-#else
-	error
-#endif
-
 bool InUnitBox(float2 UV)
 {
 	return UV.x >= 0 && UV.y >= 0 && UV.y < 1 && UV.y < 1;
@@ -250,53 +219,66 @@ void VisualizePS(in float4 UVAndScreenPos : TEXCOORD0, out float4 OutColor : SV_
 
 }
 
-// @return .RGB Color that should be scattared, .A:1 for subsurface scattering material, 0 for not
-float4 SetupSubsurfaceForOnePixel(float2 UV)
+
+struct SDiffuseAndSpecular
 {
-	float4 Ret = 0;
+	float3 Diffuse;
+	float3 Specular;
+};
 
-	FScreenSpaceData ScreenSpaceData = GetScreenSpaceData(UV);
+// @param UVSceneColor for the full res rendertarget (BufferSize) e.g. SceneColor or GBuffers
+SDiffuseAndSpecular ReconstructDiffuse(float2 UVSceneColor)
+{
+	SDiffuseAndSpecular Ret;
 
-	FLATTEN if (ScreenSpaceData.GBuffer.ShadingModelID == SHADINGMODELID_SUBSURFACE_PROFILE)
-	{
-		float4 SceneColor4 = Texture2DSample(PostprocessInput0, PostprocessInput0Sampler, UV);
+	bool bChecker = CheckerFromSceneColorUV(UVSceneColor);
 
-		// it's a valid sample
-		Ret.a = 1;
+	// todo: We could alternate the diagonal with TemporalAA or even only only 1 sample for low spec or 4 for high spec
 
-#if SSSS_SPECULAR_CORRECTION == 1
-		// we take out the specular highlights
-		Ret.rgb = LightAccumulator_ReconstructDiffuseLighting(SceneColor4);
+	float3 Quant0 = Texture2DSample(PostprocessInput0, PostprocessInput0Sampler, UVSceneColor).rgb;
+
+	// todo: expose as scalability setting
+	// 0:fast but can appear pattern can appear, especially without TemporalAA / 1:high quality
+	#define QUALITY_LEVEL 1
+
+#if QUALITY_LEVEL == 0
+	float3 Quant1 = Texture2DSample(PostprocessInput0, PostprocessInput0Sampler, UVSceneColor + float2(1, 0) * View.BufferSizeAndInvSize.zw).rgb;
 #else
-		Ret.rgb = SceneColor4.rgb;
+	float3 Quant1 = 0.25f * (
+		Texture2DSample(PostprocessInput0, PostprocessInput0Sampler, UVSceneColor + float2( 1,  0) * View.BufferSizeAndInvSize.zw).rgb +
+		Texture2DSample(PostprocessInput0, PostprocessInput0Sampler, UVSceneColor + float2( 0,  1) * View.BufferSizeAndInvSize.zw).rgb +
+		Texture2DSample(PostprocessInput0, PostprocessInput0Sampler, UVSceneColor + float2(-1,  0) * View.BufferSizeAndInvSize.zw).rgb +
+		Texture2DSample(PostprocessInput0, PostprocessInput0Sampler, UVSceneColor + float2( 0, -1) * View.BufferSizeAndInvSize.zw).rgb);
 #endif
-	}
+
+	Ret.Diffuse = lerp(Quant1, Quant0, bChecker);
+	Ret.Specular = lerp(Quant0, Quant1, bChecker);
 
 	return Ret;
 }
 
-void ExtractSpecularPS(float4 UVAndScreenPos : TEXCOORD0, out float4 OutColor : SV_Target0)
+// @param UVSceneColor for the full res rendertarget (BufferSize) e.g. SceneColor or GBuffers
+// @return .RGB Color that should be scattared, .A:1 for subsurface scattering material, 0 for not
+float4 SetupSubsurfaceForOnePixel(float2 UVSceneColor)
 {
-	float2 GBufferUV = UVAndScreenPos.xy;
+	float4 Ret = 0;
 
-	float4 SceneColor4 = Texture2DSample(PostprocessInput0, PostprocessInput0Sampler, GBufferUV);
+	FScreenSpaceData ScreenSpaceData = GetScreenSpaceData(UVSceneColor);
 
-	FScreenSpaceData ScreenSpaceData = GetScreenSpaceData(GBufferUV);
+	FLATTEN if (ScreenSpaceData.GBuffer.ShadingModelID == SHADINGMODELID_SUBSURFACE_PROFILE)
+	{
+		SDiffuseAndSpecular DiffuseAndSpecular = ReconstructDiffuse(UVSceneColor);
 
-	// 0: dont' fade in subsurface half res,  1:fully fade in half res
-	float BlendFactor = (ScreenSpaceData.GBuffer.ShadingModelID == SHADINGMODELID_SUBSURFACE_PROFILE);
+		Ret.rgb = DiffuseAndSpecular.Diffuse;
 
-	float3 ExtractedSpecular = 0;
-	 
-#if SSS_RECOMBINE_METHOD
-	// we took the specular highlights out, now we add them back in
-	ExtractedSpecular += LightAccumulator_ReconstructNonDiffuseLighting(SceneColor4);
-#endif
+		// it's a valid sample
+		Ret.a = 1;
+	}
 
-	// alpha channel can be anything - we consumed the value and likely output to a format that doesn't need the value
-	OutColor = float4(lerp(SceneColor4.rgb, ExtractedSpecular, BlendFactor), 0);
+	return Ret;
 }
 
+
 void SetupPS(in float4 UVAndScreenPos : TEXCOORD0, out float4 OutColor : SV_Target0)
 {
 	float2 UV = UVAndScreenPos.xy;
@@ -408,6 +390,7 @@ void SubsurfaceRecombinePS(float4 UVAndScreenPos : TEXCOORD0, out float4 OutColo
 
 	// 0: no SSS,  1: full SSS
 	float BlendFactor = (ScreenSpaceData.GBuffer.ShadingModelID == SHADINGMODELID_SUBSURFACE_PROFILE);
+
 	float LerpFactor = 1;
 
 #if HALF_RES
@@ -421,24 +404,32 @@ void SubsurfaceRecombinePS(float4 UVAndScreenPos : TEXCOORD0, out float4 OutColo
 #endif
 #endif // HALF_RES
 
+#if !RECOMBINE_SUBSURFACESCATTER
+	// Scalability requests no Scatter, but we still need to reconstruct a color
+	LerpFactor = 0;
+#endif
+
 	float4 SceneColor4 = Texture2DSample(PostprocessInput0, PostprocessInput0Sampler, UVSceneColor);
 
-	float3 ExtractedSpecular = 0;
-	{
-#if SSS_RECOMBINE_METHOD
-		// we took the specular highlights out, now we add them back in
-		ExtractedSpecular += LightAccumulator_ReconstructNonDiffuseLighting(SceneColor4);
-#endif
 
-		ExtractedSpecular = lerp(SceneColor4.rgb, ExtractedSpecular, BlendFactor);
-	}
+	SDiffuseAndSpecular DiffuseAndSpecular = ReconstructDiffuse(UVSceneColor);
+	
+	float3 ExtractedNonSubsurface = lerp(SceneColor4.rgb, DiffuseAndSpecular.Specular, BlendFactor);
 
+	// asset specific color
 	float3 SubsurfaceColor = GetProfileSubsurfaceColor(ScreenSpaceData.GBuffer);
+	float3 FadedSubsurfaceColor = GetProfileSubsurfaceColor(ScreenSpaceData.GBuffer) * LerpFactor;
+
+	// hack to debug if the fade is happening at the same time the VisualizeSSS shows it
+//	SSSColor = float3(0,1,0);
+
+	// combine potentially half res with full res
+	float3 SubsurfaceLighting = lerp(DiffuseAndSpecular.Diffuse, SSSColor, FadedSubsurfaceColor);
+
+	// we multiply the base color later in to get more crips human skin textures (scanned data always has Subsurface included)
+	float3 StoredBaseColor = ScreenSpaceData.GBuffer.StoredBaseColor;
 
-	// only affect pixels that have SSS
-	LerpFactor *= BlendFactor;
+	OutColor = float4(SubsurfaceLighting * StoredBaseColor + ExtractedNonSubsurface, 0); 
 
-	// add back in non subsurface pixels and specular, alpha: 0=don't blend in, 1:fully blend in
-	OutColor = float4(lerp(SceneColor4.rgb, SSSColor * BlendFactor + ExtractedSpecular, SubsurfaceColor * LerpFactor), 0); 
 }
 
diff --git a/Engine/Source/Runtime/Renderer/Private/CompositionLighting/CompositionLighting.cpp b/Engine/Source/Runtime/Renderer/Private/CompositionLighting/CompositionLighting.cpp
--- a/Engine/Source/Runtime/Renderer/Private/CompositionLighting/CompositionLighting.cpp
+++ b/Engine/Source/Runtime/Renderer/Private/CompositionLighting/CompositionLighting.cpp
@@ -31,7 +31,7 @@ FCompositionLighting GCompositionLighting;
 static TAutoConsoleVariable<float> CVarSSSScale(
 	TEXT("r.SSS.Scale"),
 	1.0f,
-	TEXT("Experimental screen space subsurface scattering pass")
+	TEXT("Affects the Screen space subsurface scattering pass")
 	TEXT("(use shadingmodel SubsurfaceProfile, get near to the object as the default)\n")
 	TEXT("is human skin which only scatters about 1.2cm)\n")
 	TEXT(" 0: off (if there is no object on the screen using this pass it should automatically disable the post process pass)\n")
@@ -386,26 +386,35 @@ void FCompositionLighting::ProcessAfterLighting(FRHICommandListImmediate& RHICmd
 			bool bSimpleDynamicLighting = IsSimpleDynamicLightingEnabled();
 
 			bool bScreenSpaceSubsurfacePassNeeded = (View.ShadingModelMaskInView & (1 << MSM_SubsurfaceProfile)) != 0;
-			if (bScreenSpaceSubsurfacePassNeeded && Radius > 0 && !bSimpleDynamicLighting && View.Family->EngineShowFlags.SubsurfaceScattering)
+			if (bScreenSpaceSubsurfacePassNeeded && !bSimpleDynamicLighting && View.Family->EngineShowFlags.SubsurfaceScattering)
 			{
 				bool bHalfRes = CVarSSSHalfRes.GetValueOnRenderThread() != 0;
 				bool bSingleViewportMode = View.Family->Views.Num() == 1;
 
-				FRenderingCompositePass* PassSetup = Context.Graph.RegisterPass(new(FMemStack::Get()) FRCPassPostProcessSubsurfaceSetup(View, bHalfRes));
-				PassSetup->SetInput(ePId_Input0, Context.FinalOutput);
-
-				FRenderingCompositePass* PassX = Context.Graph.RegisterPass(new(FMemStack::Get()) FRCPassPostProcessSubsurface(0, bHalfRes));
-				PassX->SetInput(ePId_Input0, PassSetup);
-
-				FRenderingCompositePass* PassY = Context.Graph.RegisterPass(new(FMemStack::Get()) FRCPassPostProcessSubsurface(1, bHalfRes));
-				PassY->SetInput(ePId_Input0, PassX);
-				PassY->SetInput(ePId_Input1, PassSetup);
-
-				// full res composite pass, no blurring (Radius=0), replaces SceneColor
-				FRenderingCompositePass* RecombinePass = Context.Graph.RegisterPass(new(FMemStack::Get()) FRCPassPostProcessSubsurfaceRecombine(bHalfRes, bSingleViewportMode));
-				RecombinePass->SetInput(ePId_Input0, Context.FinalOutput);
-				RecombinePass->SetInput(ePId_Input1, PassY);
-				RecombinePass->SetInput(ePId_Input2, PassSetup);
+				if(Radius > 0)
+				{
+					FRenderingCompositePass* PassSetup = Context.Graph.RegisterPass(new(FMemStack::Get()) FRCPassPostProcessSubsurfaceSetup(View, bHalfRes));
+					PassSetup->SetInput(ePId_Input0, Context.FinalOutput);
+
+					FRenderingCompositePass* PassX = Context.Graph.RegisterPass(new(FMemStack::Get()) FRCPassPostProcessSubsurface(0, bHalfRes));
+					PassX->SetInput(ePId_Input0, PassSetup);
+
+					FRenderingCompositePass* PassY = Context.Graph.RegisterPass(new(FMemStack::Get()) FRCPassPostProcessSubsurface(1, bHalfRes));
+					PassY->SetInput(ePId_Input0, PassX);
+					PassY->SetInput(ePId_Input1, PassSetup);
+
+					// full res composite pass, no blurring (Radius=0), replaces SceneColor
+					FRenderingCompositePass* RecombinePass = Context.Graph.RegisterPass(new(FMemStack::Get()) FRCPassPostProcessSubsurfaceRecombine(bHalfRes, bSingleViewportMode));
+					RecombinePass->SetInput(ePId_Input0, Context.FinalOutput);
+					RecombinePass->SetInput(ePId_Input1, PassY);
+					RecombinePass->SetInput(ePId_Input2, PassSetup);
+				}
+				else
+				{
+					// needed for Scalability
+					FRenderingCompositePass* RecombinePass = Context.Graph.RegisterPass(new(FMemStack::Get()) FRCPassPostProcessSubsurfaceRecombine(bHalfRes, bSingleViewportMode));
+					RecombinePass->SetInput(ePId_Input0, Context.FinalOutput);
+				}
 			}
 		}
 
diff --git a/Engine/Source/Runtime/Renderer/Private/PostProcess/PostProcessSubsurface.cpp b/Engine/Source/Runtime/Renderer/Private/PostProcess/PostProcessSubsurface.cpp
--- a/Engine/Source/Runtime/Renderer/Private/PostProcess/PostProcessSubsurface.cpp
+++ b/Engine/Source/Runtime/Renderer/Private/PostProcess/PostProcessSubsurface.cpp
@@ -32,14 +32,6 @@ static TAutoConsoleVariable<int32> CVarSSSSampleSet(
 	TEXT(" 2: high quality (13*2+1) (default)"),
 	ECVF_RenderThreadSafe  | ECVF_Scalability);
 
-static TAutoConsoleVariable<int> CVarSubsurfaceQuality(
-	TEXT("r.SubsurfaceQuality"),
-	1,
-	TEXT("Define the quality of the Screenspace subsurface scattering postprocess.\n")
-	TEXT(" 0: low quality for speculars on subsurface materials\n")
-	TEXT(" 1: higher quality as specular is separated before screenspace blurring (Only used if SceneColor has an alpha channel)"),
-	ECVF_Scalability | ECVF_RenderThreadSafe);
-
 // -------------------------------------------------------------
 
 float GetSubsurfaceRadiusScale()
@@ -324,10 +316,9 @@ FPooledRenderTargetDesc FRCPassPostProcessSubsurfaceVisualize::ComputeOutputDesc
 
 /**
  * Encapsulates the post processing subsurface scattering pixel shader.
- * @param SetupMode 0:without specular correction, 1:with specular correction
  * @param HalfRes 0:to full res, 1:to half res
  */
-template <uint32 SetupMode, uint32 HalfRes>
+template <uint32 HalfRes>
 class FPostProcessSubsurfaceSetupPS : public FGlobalShader
 {
 	DECLARE_SHADER_TYPE(FPostProcessSubsurfaceSetupPS , Global);
@@ -340,7 +331,6 @@ class FPostProcessSubsurfaceSetupPS : public FGlobalShader
 	static void ModifyCompilationEnvironment(EShaderPlatform Platform, FShaderCompilerEnvironment& OutEnvironment)
 	{
 		FGlobalShader::ModifyCompilationEnvironment(Platform, OutEnvironment);
-		OutEnvironment.SetDefine(TEXT("SETUP_MODE"), SetupMode);
 		OutEnvironment.SetDefine(TEXT("HALF_RES"), HalfRes);
 		OutEnvironment.SetDefine(TEXT("SUBSURFACE_RADIUS_SCALE"), SUBSURFACE_RADIUS_SCALE);
 		OutEnvironment.SetDefine(TEXT("SUBSURFACE_KERNEL_SIZE"), SUBSURFACE_KERNEL_SIZE);
@@ -395,19 +385,17 @@ public:
 };
 
 // #define avoids a lot of code duplication
-#define VARIATION1(A)		VARIATION2(A,0)			VARIATION2(A,1)
-#define VARIATION2(A, B) typedef FPostProcessSubsurfaceSetupPS<A, B> FPostProcessSubsurfaceSetupPS##A##B; \
-	IMPLEMENT_SHADER_TYPE2(FPostProcessSubsurfaceSetupPS##A##B, SF_Pixel);
+#define VARIATION1(A) typedef FPostProcessSubsurfaceSetupPS<A> FPostProcessSubsurfaceSetupPS##A; \
+	IMPLEMENT_SHADER_TYPE2(FPostProcessSubsurfaceSetupPS##A, SF_Pixel);
 	VARIATION1(0) VARIATION1(1)
 #undef VARIATION1
-#undef VARIATION2
 
 
-template <uint32 SetupMode, uint32 HalfRes>
+template <uint32 HalfRes>
 void SetSubsurfaceSetupShader(const FRenderingCompositePassContext& Context)
 {
 	TShaderMapRef<FPostProcessVS> VertexShader(Context.GetShaderMap());
-	TShaderMapRef<FPostProcessSubsurfaceSetupPS<SetupMode, HalfRes> > PixelShader(Context.GetShaderMap());
+	TShaderMapRef<FPostProcessSubsurfaceSetupPS<HalfRes> > PixelShader(Context.GetShaderMap());
 
 	static FGlobalBoundShaderState BoundShaderState;
 
@@ -417,22 +405,6 @@ void SetSubsurfaceSetupShader(const FRenderingCompositePassContext& Context)
 	VertexShader->SetParameters(Context);
 }
 
-
-static bool DoSpecularCorrection()
-{
-	bool CVarState = CVarSubsurfaceQuality.GetValueOnRenderThread() > 0;
-
-	int SceneColorFormat;
-	{
-		static IConsoleVariable* CVar = IConsoleManager::Get().FindConsoleVariable(TEXT("r.SceneColorFormat"));
-
-		SceneColorFormat = CVar->GetInt();
-	}
-
-	// we need an alpha channel for this feature
-	return CVarState && (SceneColorFormat >= 4);
-}
-
 // --------------------------------------
 
 FRCPassPostProcessSubsurfaceSetup::FRCPassPostProcessSubsurfaceSetup(FViewInfo& View, bool bInHalfRes)
@@ -480,16 +452,13 @@ void FRCPassPostProcessSubsurfaceSetup::Process(FRenderingCompositePassContext&
 	Context.RHICmdList.SetRasterizerState(TStaticRasterizerState<>::GetRHI());
 	Context.RHICmdList.SetDepthStencilState(TStaticDepthStencilState<false, CF_Always>::GetRHI());
 
-	// reconstruct specular and add it in final pass
-	bool bSpecularCorrection = DoSpecularCorrection();
-
-	if(bSpecularCorrection)
+	if(bHalfRes)
 	{
-		if(bHalfRes) SetSubsurfaceSetupShader<1, 1>(Context); else SetSubsurfaceSetupShader<1, 0>(Context);
+		SetSubsurfaceSetupShader<1>(Context);
 	}
 	else
 	{
-		if(bHalfRes) SetSubsurfaceSetupShader<0, 1>(Context); else SetSubsurfaceSetupShader<0, 0>(Context);
+		SetSubsurfaceSetupShader<0>(Context);
 	}
 
 	// Draw a quad mapping scene color to the view's render target
@@ -744,8 +713,8 @@ FPooledRenderTargetDesc FRCPassPostProcessSubsurface::ComputeOutputDesc(EPassOut
 
 
 /** Encapsulates the post processing subsurface recombine pixel shader. */
-// @param RecombineMethod 1: with specular correction, 0: without specular correction
-template <uint32 RecombineMethod, uint32 HalfRes>
+// @param RecombineMode 0:fullres, 1: halfres, 2:no scattering, just reconstruct the lighting (needed for scalability)
+template <uint32 RecombineMode>
 class TPostProcessSubsurfaceRecombinePS : public FGlobalShader
 {
 	DECLARE_SHADER_TYPE(TPostProcessSubsurfaceRecombinePS, Global);
@@ -758,8 +727,8 @@ class TPostProcessSubsurfaceRecombinePS : public FGlobalShader
 	static void ModifyCompilationEnvironment(EShaderPlatform Platform, FShaderCompilerEnvironment& OutEnvironment)
 	{
 		FGlobalShader::ModifyCompilationEnvironment(Platform,OutEnvironment);
-		OutEnvironment.SetDefine(TEXT("SSS_RECOMBINE_METHOD"), RecombineMethod);
-		OutEnvironment.SetDefine(TEXT("HALF_RES"), HalfRes);
+		OutEnvironment.SetDefine(TEXT("HALF_RES"), (uint32)(RecombineMode == 1));
+		OutEnvironment.SetDefine(TEXT("RECOMBINE_SUBSURFACESCATTER"), (uint32)(RecombineMode != 2));
 		OutEnvironment.SetDefine(TEXT("SUBSURFACE_RADIUS_SCALE"), SUBSURFACE_RADIUS_SCALE);
 		OutEnvironment.SetDefine(TEXT("SUBSURFACE_KERNEL_SIZE"), SUBSURFACE_KERNEL_SIZE);
 	}
@@ -811,17 +780,18 @@ public:
 };
 
 // #define avoids a lot of code duplication
-#define VARIATION1(A)		VARIATION2(A,0)			VARIATION2(A,1)
-#define VARIATION2(A, B) typedef TPostProcessSubsurfaceRecombinePS<A, B> TPostProcessSubsurfaceRecombinePS##A##B; \
-	IMPLEMENT_SHADER_TYPE2(TPostProcessSubsurfaceRecombinePS##A##B, SF_Pixel);
-	VARIATION1(0) VARIATION1(1)
+#define VARIATION1(A) typedef TPostProcessSubsurfaceRecombinePS<A> TPostProcessSubsurfaceRecombinePS##A; \
+	IMPLEMENT_SHADER_TYPE2(TPostProcessSubsurfaceRecombinePS##A, SF_Pixel);
+	VARIATION1(0) VARIATION1(1) VARIATION1(2)
 #undef VARIATION1
-#undef VARIATION2
 
-template <uint32 RecombineMethod, uint32 HalfRes>
+// @param RecombineMode 0:fullres, 1: halfres, 2:no scattering, just reconstruct the lighting (needed for scalability)
+template <uint32 RecombineMode>
 void SetSubsurfaceRecombineShader(const FRenderingCompositePassContext& Context, TShaderMapRef<FPostProcessVS> &VertexShader)
 {
-	TShaderMapRef<TPostProcessSubsurfaceRecombinePS<RecombineMethod, HalfRes> > PixelShader(Context.GetShaderMap());
+	SCOPED_DRAW_EVENTF(Context.RHICmdList, SubsurfacePassRecombine, TEXT("SubsurfacePassRecombine%d"), RecombineMode);
+
+	TShaderMapRef<TPostProcessSubsurfaceRecombinePS<RecombineMode> > PixelShader(Context.GetShaderMap());
 
 	static FGlobalBoundShaderState BoundShaderState;
 
@@ -875,17 +845,21 @@ void FRCPassPostProcessSubsurfaceRecombine::Process(FRenderingCompositePassConte
 
 	TShaderMapRef<FPostProcessVS> VertexShader(Context.GetShaderMap());
 
-	bool bDoSpecularCorrection = DoSpecularCorrection();
-
-	SCOPED_DRAW_EVENTF(Context.RHICmdList, SubsurfacePass, TEXT("SubsurfacePassRecombine#%d"), (int32)bDoSpecularCorrection);
-
-	if(bDoSpecularCorrection)
+	if(GetInput(ePId_Input1))
 	{
-		if(bHalfRes) SetSubsurfaceRecombineShader<1, 1>(Context, VertexShader); else SetSubsurfaceRecombineShader<1, 0>(Context, VertexShader);
+		if(bHalfRes)
+		{
+			SetSubsurfaceRecombineShader<1>(Context, VertexShader);
+		}
+		else
+		{
+			SetSubsurfaceRecombineShader<0>(Context, VertexShader);
+		}
 	}
 	else
 	{
-		if(bHalfRes) SetSubsurfaceRecombineShader<0, 1>(Context, VertexShader); else SetSubsurfaceRecombineShader<0, 0>(Context, VertexShader);
+		// needed for Scalability
+		SetSubsurfaceRecombineShader<2>(Context, VertexShader);
 	}
 
 	DrawRectangle(
diff --git a/Engine/Source/Runtime/Renderer/Private/PostProcess/PostProcessSubsurface.h b/Engine/Source/Runtime/Renderer/Private/PostProcess/PostProcessSubsurface.h
--- a/Engine/Source/Runtime/Renderer/Private/PostProcess/PostProcessSubsurface.h
+++ b/Engine/Source/Runtime/Renderer/Private/PostProcess/PostProcessSubsurface.h
@@ -69,8 +69,8 @@ private:
 
 // derives from TRenderingCompositePassBase<InputCount, OutputCount>
 // ePId_Input0: SceneColor before Screen Space Subsurface input
-// ePId_Input1: output from FRCPassPostProcessSubsurface
-// ePId_Input2: SubsurfaceSetup, can be half res
+// ePId_Input1: optional output from FRCPassPostProcessSubsurface (if not present we do cheap reconstruction for Scalability)
+// ePId_Input2: optional SubsurfaceSetup, can be half res
 // modifies SceneColor, uses some GBuffer attributes
 class FRCPassPostProcessSubsurfaceRecombine : public TRenderingCompositePassBase<3, 1>
 {
