Fixed crash loading Orion in Main

Guarded against missing CVar.

diff --git a/Engine/Source/Runtime/Renderer/Private/DeferredShadingRenderer.cpp b/Engine/Source/Runtime/Renderer/Private/DeferredShadingRenderer.cpp
--- a/Engine/Source/Runtime/Renderer/Private/DeferredShadingRenderer.cpp
+++ b/Engine/Source/Runtime/Renderer/Private/DeferredShadingRenderer.cpp
@@ -694,7 +694,8 @@ static FORCEINLINE bool HasHiddenAreaMask()
 {
 	static const IConsoleVariable* const CVar = IConsoleManager::Get().FindConsoleVariable(TEXT("vr.HiddenAreaMask"));
 
-	return (CVar->GetInt() == 1 &&
+	return (CVar != nullptr &&
+		CVar->GetInt() == 1 &&
 		GEngine &&
 		GEngine->HMDDevice.IsValid() &&
 		GEngine->HMDDevice->HasHiddenAreaMask());
