Fixed bug where it was looking for tonemap shader features that don't exist in the PC path.

Changed the compiled shader options slightly to reflect common use better.

diff --git a/Engine/Source/Runtime/Renderer/Private/PostProcess/PostProcessTonemap.cpp b/Engine/Source/Runtime/Renderer/Private/PostProcess/PostProcessTonemap.cpp
--- a/Engine/Source/Runtime/Renderer/Private/PostProcess/PostProcessTonemap.cpp
+++ b/Engine/Source/Runtime/Renderer/Private/PostProcess/PostProcessTonemap.cpp
@@ -79,6 +79,7 @@ static uint32 TonemapperConfBitmaskPC[7] = {
 	0,
 
 	TonemapperBloom + 
+	TonemapperVignette +
 	TonemapperGrainQuantization +
 	0,
 
@@ -97,6 +98,7 @@ static uint32 TonemapperConfBitmaskPC[7] = {
 	0,
 
 	TonemapperBloom + 
+	TonemapperVignette +
 	0,
 
 	//
@@ -515,6 +517,11 @@ static uint32 TonemapperGenerateBitmaskPC(const FViewInfo* RESTRICT View, bool b
 {
 	uint32 Bitmask = TonemapperGenerateBitmask(View, bGammaOnly, false);
 
+	// PC doesn't support these
+	Bitmask &= ~TonemapperContrast;
+	Bitmask &= ~TonemapperColorMatrix;
+	Bitmask &= ~TonemapperShadowTint;
+
 	// Must early exit if gamma only.
 	if(Bitmask == TonemapperGammaOnly) 
 	{
