added new optional GaussianDOF masking with a circular mask, useful for SniperScope, no performance impact if not used
  UE-19432 DOF vignette blur

diff --git a/Engine/Shaders/DepthOfFieldCommon.usf b/Engine/Shaders/DepthOfFieldCommon.usf
--- a/Engine/Shaders/DepthOfFieldCommon.usf
+++ b/Engine/Shaders/DepthOfFieldCommon.usf
@@ -4,7 +4,7 @@
 	DepthOfFieldCommon.usf: depth of field common
 =============================================================================*/
 
-// [0] .x:SkyFocusDistance or very value if 0, .yz:unused, w:OcclusionTweak1
+// [0] .x:SkyFocusDistance or very value if 0, .y:DepthOfFieldVignetteMul, .z:DepthOfFieldVignetteAdd, w:OcclusionTweak1
 // [1] .x:MaxBokehSizeInPixel, y:vertical offset in UV, z:UsedY/TextureY, w:vertical offset to put two views in one texture with safety border
 float4 DepthOfFieldParams[2];
 
diff --git a/Engine/Source/Runtime/Engine/Classes/Engine/Scene.h b/Engine/Source/Runtime/Engine/Classes/Engine/Scene.h
--- a/Engine/Source/Runtime/Engine/Classes/Engine/Scene.h
+++ b/Engine/Source/Runtime/Engine/Classes/Engine/Scene.h
@@ -406,6 +406,9 @@ struct FPostProcessSettings
 	UPROPERTY(BlueprintReadWrite, Category=Overrides, meta=(PinHiddenByDefault))
 	uint32 bOverride_DepthOfFieldSkyFocusDistance:1;
 
+	UPROPERTY(BlueprintReadWrite, Category=Overrides, meta=(PinHiddenByDefault))
+	uint32 bOverride_DepthOfFieldVignetteSize:1;
+
 	UPROPERTY(BlueprintReadWrite, Category=Overrides, meta=(PinHiddenByDefault))
 	uint32 bOverride_MotionBlurAmount:1;
 
@@ -892,7 +895,11 @@ struct FPostProcessSettings
 	/** Artificial distance to allow the skybox to be in focus (e.g. 200000), <=0 to switch the feature off, only for GaussianDOF, can cost performance */
 	UPROPERTY(interp, BlueprintReadWrite, Category=DepthOfField, AdvancedDisplay, meta=(ClampMin = "0.0", ClampMax = "200000.0", editcondition = "bOverride_DepthOfFieldSkyFocusDistance", DisplayName = "Sky Distance"))
 	float DepthOfFieldSkyFocusDistance;
-	
+
+	/** Artificial circular mask to (near) blur content outside the radius, only for GaussianDOF, diameter in percent of screen width, costs performance if the mask is used, keep Feather can Radius on default to keep it off */
+	UPROPERTY(interp, BlueprintReadWrite, Category=DepthOfField, AdvancedDisplay, meta=(UIMin = "0.0", UIMax = "100.0", editcondition = "bOverride_DepthOfFieldVignetteSize", DisplayName = "Vignette Size"))
+	float DepthOfFieldVignetteSize;
+
 	/** Strength of motion blur, 0:off, should be renamed to intensity */
 	UPROPERTY(interp, BlueprintReadWrite, Category=MotionBlur, meta=(ClampMin = "0.0", ClampMax = "1.0", editcondition = "bOverride_MotionBlurAmount", DisplayName = "Amount"))
 	float MotionBlurAmount;
@@ -1124,6 +1131,8 @@ struct FPostProcessSettings
 		DepthOfFieldColorThreshold = 1.0f;
 		DepthOfFieldSizeThreshold = 0.08f;
 		DepthOfFieldSkyFocusDistance = 0.0f;
+		// 200 should be enough even for extreme aspect ratios to give the default no effect
+		DepthOfFieldVignetteSize = 200.0f;
 		LensFlareTints[0] = FLinearColor(1.0f, 0.8f, 0.4f, 0.6f);
 		LensFlareTints[1] = FLinearColor(1.0f, 1.0f, 0.6f, 0.53f);
 		LensFlareTints[2] = FLinearColor(0.8f, 0.8f, 1.0f, 0.46f);
diff --git a/Engine/Source/Runtime/Engine/Private/PostProcessVolume.cpp b/Engine/Source/Runtime/Engine/Private/PostProcessVolume.cpp
--- a/Engine/Source/Runtime/Engine/Private/PostProcessVolume.cpp
+++ b/Engine/Source/Runtime/Engine/Private/PostProcessVolume.cpp
@@ -74,7 +74,8 @@ bool APostProcessVolume::CanEditChange(const UProperty* InProperty) const
 
 			if (PropertyName == GET_MEMBER_NAME_STRING_CHECKED(FPostProcessSettings, DepthOfFieldNearBlurSize) ||
 				PropertyName == GET_MEMBER_NAME_STRING_CHECKED(FPostProcessSettings, DepthOfFieldFarBlurSize) ||
-				PropertyName == GET_MEMBER_NAME_STRING_CHECKED(FPostProcessSettings, DepthOfFieldSkyFocusDistance))
+				PropertyName == GET_MEMBER_NAME_STRING_CHECKED(FPostProcessSettings, DepthOfFieldSkyFocusDistance) ||
+				PropertyName == GET_MEMBER_NAME_STRING_CHECKED(FPostProcessSettings, DepthOfFieldVignetteSize))
 			{
 				return Settings.DepthOfFieldMethod == EDepthOfFieldMethod::DOFM_Gaussian;
 			}
diff --git a/Engine/Source/Runtime/Engine/Private/SceneView.cpp b/Engine/Source/Runtime/Engine/Private/SceneView.cpp
--- a/Engine/Source/Runtime/Engine/Private/SceneView.cpp
+++ b/Engine/Source/Runtime/Engine/Private/SceneView.cpp
@@ -926,6 +926,7 @@ void FSceneView::OverridePostProcessSettings(const FPostProcessSettings& Src, fl
 		LERP_PP(DepthOfFieldColorThreshold);
 		LERP_PP(DepthOfFieldSizeThreshold);
 		LERP_PP(DepthOfFieldSkyFocusDistance);
+		LERP_PP(DepthOfFieldVignetteSize);
 		LERP_PP(MotionBlurAmount);
 		LERP_PP(MotionBlurMax);
 		LERP_PP(MotionBlurPerObjectSize);
diff --git a/Engine/Source/Runtime/Renderer/Private/PostProcess/PostProcessBokehDOF.cpp b/Engine/Source/Runtime/Renderer/Private/PostProcess/PostProcessBokehDOF.cpp
--- a/Engine/Source/Runtime/Renderer/Private/PostProcess/PostProcessBokehDOF.cpp
+++ b/Engine/Source/Runtime/Renderer/Private/PostProcess/PostProcessBokehDOF.cpp
@@ -311,6 +311,8 @@ void FRCPassPostProcessVisualizeDOF::Process(FRenderingCompositePassContext& Con
 			Canvas.DrawShadowedString(X, Y += YStep, *Line, GetStatsFont(), FLinearColor(1, 1, 1));
 			Line = FString::Printf(TEXT("SkyFocusDistance: %.2f"), View.FinalPostProcessSettings.DepthOfFieldSkyFocusDistance);
 			Canvas.DrawShadowedString(X, Y += YStep, *Line, GetStatsFont(), FLinearColor(1, 1, 1));
+			Line = FString::Printf(TEXT("VignetteRadius: %.2f"), View.FinalPostProcessSettings.DepthOfFieldVignetteSize);
+			Canvas.DrawShadowedString(X, Y += YStep, *Line, GetStatsFont(), FLinearColor(1, 1, 1));
 			Y += YStep;
 			Line = FString::Printf(TEXT("Near:%d Far:%d"), DepthOfFieldStats.bNear ? 1 : 0, DepthOfFieldStats.bFar ? 1 : 0);
 			Canvas.DrawShadowedString(X, Y += YStep, *Line, GetStatsFont(), FLinearColor(1, 1, 1));
@@ -713,11 +715,19 @@ void FRCPassPostProcessBokehDOF::ComputeDepthOfFieldParams(const FRenderingCompo
 	uint32 BokehLayerSizeY = HalfRes * 2 + SafetyBorder;
 
 	float SkyFocusDistance = Context.View.FinalPostProcessSettings.DepthOfFieldSkyFocusDistance;
+	
+	// *2 to go to account for Radius/Diameter, 100 for percent
+	float DepthOfFieldVignetteSize = FMath::Max(0.0f, Context.View.FinalPostProcessSettings.DepthOfFieldVignetteSize / 100.0f * 2);
+	// doesn't make much sense to expose this property as the effect is very non linear and it would cost some performance to fix that
+	float DepthOfFieldVignetteFeather = 10.0f / 100.0f;
+
+	float DepthOfFieldVignetteMul = 1.0f / DepthOfFieldVignetteFeather;
+	float DepthOfFieldVignetteAdd = (0.5f - DepthOfFieldVignetteSize) * DepthOfFieldVignetteMul;
 
 	Out[0] = FVector4(
 		(SkyFocusDistance > 0) ? SkyFocusDistance : 100000000.0f,			// very large if <0 to not mask out skybox, can be optimized to disable feature completely
-		0,
-		0,
+		DepthOfFieldVignetteMul,
+		DepthOfFieldVignetteAdd,
 		Context.View.FinalPostProcessSettings.DepthOfFieldOcclusion);
 
 	FIntPoint ViewSize = Context.View.ViewRect.Size();
diff --git a/Engine/Source/Runtime/Renderer/Private/PostProcess/PostProcessDOF.cpp b/Engine/Source/Runtime/Renderer/Private/PostProcess/PostProcessDOF.cpp
--- a/Engine/Source/Runtime/Renderer/Private/PostProcess/PostProcessDOF.cpp
+++ b/Engine/Source/Runtime/Renderer/Private/PostProcess/PostProcessDOF.cpp
@@ -14,7 +14,8 @@
 
 
 /** Encapsulates the DOF setup pixel shader. */
-template <uint32 NearBlurEnable>
+// @param 0:no NearBlur, 1:with NearBlur, 2:With NearBlur and Vignette
+template <uint32 SetupMode>
 class FPostProcessDOFSetupPS : public FGlobalShader
 {
 	DECLARE_SHADER_TYPE(FPostProcessDOFSetupPS, Global);
@@ -27,7 +28,8 @@ class FPostProcessDOFSetupPS : public FGlobalShader
 	static void ModifyCompilationEnvironment(EShaderPlatform Platform, FShaderCompilerEnvironment& OutEnvironment)
 	{
 		FGlobalShader::ModifyCompilationEnvironment(Platform,OutEnvironment);
-		OutEnvironment.SetDefine(TEXT("ENABLE_NEAR_BLUR"), NearBlurEnable);
+		OutEnvironment.SetDefine(TEXT("ENABLE_NEAR_BLUR"), (uint32)(SetupMode >= 1));
+		OutEnvironment.SetDefine(TEXT("ENABLE_DOF_VIGNETTE"), (uint32)(SetupMode == 2));
 	}
 
 	/** Default constructor. */
@@ -77,6 +79,7 @@ public:
 
 IMPLEMENT_SHADER_TYPE(template<>,FPostProcessDOFSetupPS<0>,TEXT("PostProcessDOF"),TEXT("SetupPS"),SF_Pixel);
 IMPLEMENT_SHADER_TYPE(template<>,FPostProcessDOFSetupPS<1>,TEXT("PostProcessDOF"),TEXT("SetupPS"),SF_Pixel);
+IMPLEMENT_SHADER_TYPE(template<>,FPostProcessDOFSetupPS<2>,TEXT("PostProcessDOF"),TEXT("SetupPS"),SF_Pixel);
 
 void FRCPassPostProcessDOFSetup::Process(FRenderingCompositePassContext& Context)
 {
@@ -126,7 +129,7 @@ void FRCPassPostProcessDOFSetup::Process(FRenderingCompositePassContext& Context
 	// is optimized away if possible (RT size=view size, )
 	Context.RHICmdList.ClearMRT(true, NumRenderTargets, ClearColors, false, 1.0f, false, 0, DestRect);
 
-	Context.SetViewportAndCallRHI(0, 0, 0.0f, DestSize.X, DestSize.Y, 1.0f );
+	Context.SetViewportAndCallRHI(DestRect.Min.X, DestRect.Min.Y, 0.0f, DestRect.Max.X + 1, DestRect.Max.Y + 1, 1.0f );
 
 	// set the state
 	Context.RHICmdList.SetBlendState(TStaticBlendState<>::GetRHI());
@@ -134,16 +137,32 @@ void FRCPassPostProcessDOFSetup::Process(FRenderingCompositePassContext& Context
 	Context.RHICmdList.SetDepthStencilState(TStaticDepthStencilState<false, CF_Always>::GetRHI());
 	
 	TShaderMapRef<FPostProcessVS> VertexShader(ShaderMap);
-
+	
 	if (bNearBlurEnabled)
 	{
-		static FGlobalBoundShaderState BoundShaderState;
-		
+		const float DOFVignetteSize = FMath::Max(0.0f, View.FinalPostProcessSettings.DepthOfFieldVignetteSize);
 
-		TShaderMapRef< FPostProcessDOFSetupPS<1> > PixelShader(ShaderMap);
-		SetGlobalBoundShaderState(Context.RHICmdList, FeatureLevel, BoundShaderState, GFilterVertexDeclaration.VertexDeclarationRHI, *VertexShader, *PixelShader);
+		// todo: test is conservative, with bad content we would waste a bit of performance
+		const bool bDOFVignette = DOFVignetteSize < 200.0f;
+
+		if(bDOFVignette)
+		{
+			static FGlobalBoundShaderState BoundShaderState;
+
+			TShaderMapRef< FPostProcessDOFSetupPS<2> > PixelShader(ShaderMap);
+			SetGlobalBoundShaderState(Context.RHICmdList, FeatureLevel, BoundShaderState, GFilterVertexDeclaration.VertexDeclarationRHI, *VertexShader, *PixelShader);
 		
-		PixelShader->SetParameters(Context);
+			PixelShader->SetParameters(Context);
+		}
+		else
+		{
+			static FGlobalBoundShaderState BoundShaderState;
+
+			TShaderMapRef< FPostProcessDOFSetupPS<1> > PixelShader(ShaderMap);
+			SetGlobalBoundShaderState(Context.RHICmdList, FeatureLevel, BoundShaderState, GFilterVertexDeclaration.VertexDeclarationRHI, *VertexShader, *PixelShader);
+
+			PixelShader->SetParameters(Context);
+		}
 	}
 	else
 	{
@@ -160,11 +179,11 @@ void FRCPassPostProcessDOFSetup::Process(FRenderingCompositePassContext& Context
 	// Draw a quad mapping scene color to the view's render target
 	DrawRectangle(
 		Context.RHICmdList,
-		DestRect.Min.X, DestRect.Min.Y,
+		0, 0,
 		DestRect.Width() + 1, DestRect.Height() + 1,
 		SrcRect.Min.X, SrcRect.Min.Y,
 		SrcRect.Width() + 1, SrcRect.Height() + 1,
-		DestSize,
+		DestRect.Size() + FIntPoint(1, 1),
 		SrcSize,
 		*VertexShader,
 		EDRF_UseTriangleOptimization);
