added comment

diff --git a/Engine/Source/Runtime/Renderer/Private/PostProcess/SceneRenderTargets.h b/Engine/Source/Runtime/Renderer/Private/PostProcess/SceneRenderTargets.h
--- a/Engine/Source/Runtime/Renderer/Private/PostProcess/SceneRenderTargets.h
+++ b/Engine/Source/Runtime/Renderer/Private/PostProcess/SceneRenderTargets.h
@@ -420,6 +420,7 @@ public:
 	FUniformBufferRHIParamRef GetGBufferResourcesUniformBuffer() const 
 	{ 
 		// if this triggers you need to make sure the GBuffer is not getting released before (using AdjustGBufferRefCount(1) and AdjustGBufferRefCount(-1))
+		// Maybe You use a SceneTexture material expression that should set MaterialCompilationOutput.bNeedsGBuffer
 		check(IsValidRef(GBufferResourcesUniformBuffer));
 
 		return GBufferResourcesUniformBuffer; 
