Enable 32bpp HDR encoding for decals.

diff --git a/Engine/Shaders/DeferredDecal.usf b/Engine/Shaders/DeferredDecal.usf
--- a/Engine/Shaders/DeferredDecal.usf
+++ b/Engine/Shaders/DeferredDecal.usf
@@ -44,6 +44,25 @@ void StencilDecalMaskMain(
 	OutColor = 0;
 }
 
+half3 FrameBufferDecalBlendOp(half4 Source)
+{
+#if (COMPILER_GLSL_ES2)
+	half4 Dest = FramebufferFetchES2();
+	Dest.rgb = Decode32BPPHDR(Dest);
+#else
+	half4 Dest = half4(0, 0, 0, 0);
+#endif
+
+#if DECAL_BLEND_MODE == DECALBLENDMODEID_STAIN
+	return (Source.rgb*Dest.rgb) + (Dest.rgb*(1.0 - Source.a));
+#elif DECAL_BLEND_MODE == DECALBLENDMODEID_TRANSLUCENT
+	return (Source.rgb*Source.a) + (Dest.rgb*(1.0 - Source.a));
+#elif DECAL_BLEND_MODE == DECALBLENDMODEID_EMISSIVE
+	return (Source.rgb*Source.a) + Dest.rgb;
+#else
+	return Source.rgb;
+#endif
+}
 
 // decal vertex shader
 void MainVS(
@@ -308,6 +327,18 @@ void MainPS(
 	// RETURN_COLOR not needed unless writing to SceneColor
 	OutTarget0 = float4(Color, Opacity);
 
+#if (ES2_PROFILE || ES3_1_PROFILE)
+#if (COMPILER_GLSL_ES2) || (MOBILE_EMULATION)
+	// Do decal blending if encoding requires it.
+	if (GetHDR32bppEncodeMode() == HDR_ENCODE_RGBA)
+	{
+		OutTarget0.rgb = FrameBufferDecalBlendOp(OutTarget0);
+	}
+ 	// do 32bpp hdr encoding for mobile if required.
+	OutTarget0 = Encode32BPPHDR(OutTarget0, SvPosition.xy);
+#endif
+#endif
+
 #if DECAL_RENDERTARGET_COUNT > 1
 	// some MRT rendering
 	
diff --git a/Engine/Source/Runtime/Renderer/Private/DecalRenderingForwardShading.cpp b/Engine/Source/Runtime/Renderer/Private/DecalRenderingForwardShading.cpp
--- a/Engine/Source/Runtime/Renderer/Private/DecalRenderingForwardShading.cpp
+++ b/Engine/Source/Runtime/Renderer/Private/DecalRenderingForwardShading.cpp
@@ -46,7 +46,11 @@ void FForwardShadingSceneRenderer::RenderDecals(FRHICommandListImmediate& RHICmd
 
 			TOptional<EDecalBlendMode> LastDecalBlendMode;
 			TOptional<bool> LastDecalDepthState;
-				
+			bool bEncodedHDR = IsMobileHDR32bpp() && !IsMobileHDRMosaic();
+			if (bEncodedHDR)
+			{
+				RHICmdList.SetBlendState(TStaticBlendState<>::GetRHI());
+			}
 			for (int32 DecalIndex = 0, DecalCount = SortedDecals.Num(); DecalIndex < DecalCount; DecalIndex++)
 			{
 				const FTransientDecalRenderData& DecalData = SortedDecals[DecalIndex];
@@ -72,7 +76,7 @@ void FForwardShadingSceneRenderer::RenderDecals(FRHICommandListImmediate& RHICmd
 					}
 				}
 
-				if (!LastDecalBlendMode.IsSet() || LastDecalBlendMode.GetValue() != DecalData.DecalBlendMode)
+				if (!bEncodedHDR && (!LastDecalBlendMode.IsSet() || LastDecalBlendMode.GetValue() != DecalData.DecalBlendMode))
 				{
 					LastDecalBlendMode = DecalData.DecalBlendMode;
 					switch(DecalData.DecalBlendMode)
diff --git a/Engine/Source/Runtime/Renderer/Private/DecalRenderingShared.cpp b/Engine/Source/Runtime/Renderer/Private/DecalRenderingShared.cpp
--- a/Engine/Source/Runtime/Renderer/Private/DecalRenderingShared.cpp
+++ b/Engine/Source/Runtime/Renderer/Private/DecalRenderingShared.cpp
@@ -132,15 +132,17 @@ public:
 		uint32 RenderTargetCount = FDecalRendering::ComputeRenderTargetCount(Platform, RenderTargetMode);
 
 		// avoid using the index directly, better use DECALBLENDMODEID_VOLUMETRIC, DECALBLENDMODEID_STAIN, ...
-		OutEnvironment.SetDefine(TEXT("DECAL_BLEND_MODE"), (int32)DecalBlendMode);
-		OutEnvironment.SetDefine(TEXT("DECAL_PROJECTION"), 1);
+		OutEnvironment.SetDefine(TEXT("DECAL_BLEND_MODE"), (uint32)DecalBlendMode);
+		OutEnvironment.SetDefine(TEXT("DECAL_PROJECTION"), 1u);
 		OutEnvironment.SetDefine(TEXT("DECAL_RENDERTARGET_COUNT"), RenderTargetCount);
-		OutEnvironment.SetDefine(TEXT("DECAL_RENDERSTAGE"), (int32)FDecalRendering::ComputeRenderStage(Platform, DecalBlendMode));
+		OutEnvironment.SetDefine(TEXT("DECAL_RENDERSTAGE"), (uint32)FDecalRendering::ComputeRenderStage(Platform, DecalBlendMode));
 
 		// to compare against DECAL_BLEND_MODE, we can expose more if needed
-		OutEnvironment.SetDefine(TEXT("DECALBLENDMODEID_VOLUMETRIC"), (int32)DBM_Volumetric_DistanceFunction);
-		OutEnvironment.SetDefine(TEXT("DECALBLENDMODEID_STAIN"), (int32)DBM_Stain);
-		OutEnvironment.SetDefine(TEXT("DECALBLENDMODEID_NORMAL"), (int32)DBM_Normal);
+		OutEnvironment.SetDefine(TEXT("DECALBLENDMODEID_VOLUMETRIC"), (uint32)DBM_Volumetric_DistanceFunction);
+		OutEnvironment.SetDefine(TEXT("DECALBLENDMODEID_STAIN"), (uint32)DBM_Stain);
+		OutEnvironment.SetDefine(TEXT("DECALBLENDMODEID_NORMAL"), (uint32)DBM_Normal);
+		OutEnvironment.SetDefine(TEXT("DECALBLENDMODEID_EMISSIVE"), (uint32)DBM_Emissive);
+		OutEnvironment.SetDefine(TEXT("DECALBLENDMODEID_TRANSLUCENT"), (uint32)DBM_Translucent);
 	}
 
 	FDeferredDecalPS() {}
