fixed FastBlur radius scaling, causes wrong radius when using r.FastBlurThreshold or changing radius on bloom or DepthOfField

diff --git a/Engine/Source/Runtime/Renderer/Private/PostProcess/PostProcessWeightedSampleSum.cpp b/Engine/Source/Runtime/Renderer/Private/PostProcess/PostProcessWeightedSampleSum.cpp
--- a/Engine/Source/Runtime/Renderer/Private/PostProcess/PostProcessWeightedSampleSum.cpp
+++ b/Engine/Source/Runtime/Renderer/Private/PostProcess/PostProcessWeightedSampleSum.cpp
@@ -520,10 +520,17 @@ void FRCPassPostProcessWeightedSampleSum::Process(FRenderingCompositePassContext
 
 	const FSceneRenderTargetItem& DestRenderTarget = PassOutputs[0].RequestSurface(Context);
 
+	bool bDoFastBlur = DoFastBlur();
+
 	FVector2D InvSrcSize(1.0f / SrcSize.X, 1.0f / SrcSize.Y);
 	// we scale by width because FOV is defined horizontally
 	float SrcSizeForThisAxis = View.ViewRect.Width() / (float)SrcScaleFactor.X;
 
+	if(bDoFastBlur && FilterShape == EFS_Vert)
+	{
+		SrcSizeForThisAxis *= 2.0f;
+	}
+
 	// in texel (input resolution), /2 as we use the diameter, 100 as we use percent
 	float EffectiveBlurRadius = SizeScale * SrcSizeForThisAxis  / 2 / 100.0f;
 
@@ -581,8 +588,6 @@ void FRCPassPostProcessWeightedSampleSum::Process(FRenderingCompositePassContext
 		CombineMethodInt = 1;
 	}
 
-	bool bDoFastBlur = DoFastBlur();
-
 	if (FilterShape == EFS_Horiz)
 	{
 		float YOffset = bDoFastBlur ? (InvSrcSize.Y * 0.5f) : 0.0f;
