fixed discoloration if showflag for SubsurfaceScattering is off

diff --git a/Engine/Source/Runtime/Renderer/Private/CompositionLighting/CompositionLighting.cpp b/Engine/Source/Runtime/Renderer/Private/CompositionLighting/CompositionLighting.cpp
--- a/Engine/Source/Runtime/Renderer/Private/CompositionLighting/CompositionLighting.cpp
+++ b/Engine/Source/Runtime/Renderer/Private/CompositionLighting/CompositionLighting.cpp
@@ -386,12 +386,12 @@ void FCompositionLighting::ProcessAfterLighting(FRHICommandListImmediate& RHICmd
 			bool bSimpleDynamicLighting = IsSimpleDynamicLightingEnabled();
 
 			bool bScreenSpaceSubsurfacePassNeeded = (View.ShadingModelMaskInView & (1 << MSM_SubsurfaceProfile)) != 0;
-			if (bScreenSpaceSubsurfacePassNeeded && !bSimpleDynamicLighting && View.Family->EngineShowFlags.SubsurfaceScattering)
+			if (bScreenSpaceSubsurfacePassNeeded && !bSimpleDynamicLighting)
 			{
 				bool bHalfRes = CVarSSSHalfRes.GetValueOnRenderThread() != 0;
 				bool bSingleViewportMode = View.Family->Views.Num() == 1;
 
-				if(Radius > 0)
+				if(Radius > 0 && View.Family->EngineShowFlags.SubsurfaceScattering)
 				{
 					FRenderingCompositePass* PassSetup = Context.Graph.RegisterPass(new(FMemStack::Get()) FRCPassPostProcessSubsurfaceSetup(View, bHalfRes));
 					PassSetup->SetInput(ePId_Input0, Context.FinalOutput);
