Modulated shadow depths are rendered before basepass to avoid a resolve.
fixed metal's handling of scissor RHISetScissorRect( false ,,,, )

diff --git a/Engine/Source/Runtime/IOS/MetalRHI/Private/MetalCommands.cpp b/Engine/Source/Runtime/IOS/MetalRHI/Private/MetalCommands.cpp
--- a/Engine/Source/Runtime/IOS/MetalRHI/Private/MetalCommands.cpp
+++ b/Engine/Source/Runtime/IOS/MetalRHI/Private/MetalCommands.cpp
@@ -115,15 +115,27 @@ void FMetalDynamicRHI::RHISetMultipleViewports(uint32 Count, const FViewportBoun
 void FMetalDynamicRHI::RHISetScissorRect(bool bEnable,uint32 MinX,uint32 MinY,uint32 MaxX,uint32 MaxY)
 {
 	MTLScissorRect Scissor;
-	Scissor.x = MinX;
-	Scissor.y = MinY;
-	Scissor.width = MaxX - MinX;
-	Scissor.height = MaxY - MinY;
 
-	// metal doesn't support 0 sized scissor rect
-	if (Scissor.width == 0 || Scissor.height == 0)
+	if (!bEnable)
 	{
-		return;
+		FIntPoint TargetDims = FMetalManager::Get()->GetBoundRenderTargetDimensions();
+		Scissor.x = 0;
+		Scissor.y = 0;
+		Scissor.width = (uint32)TargetDims.X;
+		Scissor.height = (uint32)TargetDims.Y;
+	}
+	else
+	{
+		Scissor.x = MinX;
+		Scissor.y = MinY;
+		Scissor.width = MaxX - MinX;
+		Scissor.height = MaxY - MinY;
+
+		// metal doesn't support 0 sized scissor rect
+		if (Scissor.width == 0 || Scissor.height == 0)
+		{
+			return;
+		}
 	}
 	[FMetalManager::GetContext() setScissorRect:Scissor];
 }
diff --git a/Engine/Source/Runtime/Renderer/Private/ForwardShadingRenderer.cpp b/Engine/Source/Runtime/Renderer/Private/ForwardShadingRenderer.cpp
--- a/Engine/Source/Runtime/Renderer/Private/ForwardShadingRenderer.cpp
+++ b/Engine/Source/Runtime/Renderer/Private/ForwardShadingRenderer.cpp
@@ -21,6 +21,8 @@ uint32 GetShadowQuality();
 FForwardShadingSceneRenderer::FForwardShadingSceneRenderer(const FSceneViewFamily* InViewFamily,FHitProxyConsumer* HitProxyConsumer)
 	:	FSceneRenderer(InViewFamily, HitProxyConsumer)
 {
+	bModulatedShadowsInUse = false;
+	bCSMShadowsInUse = false;
 }
 
 /**
diff --git a/Engine/Source/Runtime/Renderer/Private/SceneRendering.h b/Engine/Source/Runtime/Renderer/Private/SceneRendering.h
--- a/Engine/Source/Runtime/Renderer/Private/SceneRendering.h
+++ b/Engine/Source/Runtime/Renderer/Private/SceneRendering.h
@@ -885,9 +885,15 @@ protected:
 	/** Renders the opaque base pass for forward shading. */
 	void RenderForwardShadingBasePass(FRHICommandListImmediate& RHICmdList);
 
-	/** Projects any shadow depth images into the scene. (modulated shadows) */
+	/** Render modulated shadow projections in to the scene, loops over any unrendered shadows until all are processed.*/
 	void RenderModulatedShadowProjections(FRHICommandListImmediate& RHICmdList);
 
+	/** Render shadow depths to the depth texture.*/
+	void RenderModulatedShadowDepthMaps(FRHICommandListImmediate& RHICmdList);
+	
+	/** Projects any existing depth images into the scene.*/
+	void RenderAllocatedModulatedShadowProjections(FRHICommandListImmediate& RHICmdList);
+	
 	/** Makes a copy of scene alpha so PC can emulate ES2 framebuffer fetch. */
 	void CopySceneAlpha(FRHICommandListImmediate& RHICmdList, const FViewInfo& View);
 
@@ -911,6 +917,7 @@ protected:
 	  */
 	bool RenderShadowDepthMap(FRHICommandListImmediate& RHICmdList, const FLightSceneInfo* LightSceneInfo);
 
-	/** Helper function to determine if a shadow should be cast **/
-	bool IsProjectedShadowPotentiallyVisible(const FProjectedShadowInfo* ProjectedShadowInfo) const;
+private:
+	bool bModulatedShadowsInUse;
+	bool bCSMShadowsInUse;
 };
diff --git a/Engine/Source/Runtime/Renderer/Private/ShadowRendering.cpp b/Engine/Source/Runtime/Renderer/Private/ShadowRendering.cpp
--- a/Engine/Source/Runtime/Renderer/Private/ShadowRendering.cpp
+++ b/Engine/Source/Runtime/Renderer/Private/ShadowRendering.cpp
@@ -3675,6 +3675,36 @@ bool FDeferredShadingSceneRenderer::RenderCachedPreshadows(FRHICommandListImmedi
 	return bAttenuationBufferDirty;
 }
 
+/** Helper function to determine if a shadow should be cast **/
+static bool IsProjectedForwardShadowPotentiallyVisible(const FProjectedShadowInfo* ProjectedShadowInfo)
+{
+	// Check that the shadow is visible in at least one view before rendering it.
+	const FLightSceneInfo& LightSceneInfo = ProjectedShadowInfo->GetLightSceneInfo();
+
+	// Skip shadows which will be handled in RenderOnePassPointLightShadows or RenderReflectiveShadowmaps
+	if (ProjectedShadowInfo->CascadeSettings.bOnePassPointLightShadow || ProjectedShadowInfo->bReflectiveShadowmap)
+	{
+		return false;
+	}
+
+	if (ProjectedShadowInfo->bPreShadow && !LightSceneInfo.Proxy->HasStaticShadowing())
+	{
+		return false;
+	}
+
+	if (ProjectedShadowInfo->bWholeSceneShadow && LightSceneInfo.Proxy->HasStaticShadowing())
+	{
+		return false;
+	}
+
+	if ((ProjectedShadowInfo->bPreShadow && !ProjectedShadowInfo->HasSubjectPrims()) || ProjectedShadowInfo->bAllocatedInPreshadowCache)
+	{
+		return false;
+	}
+
+	return true;
+}
+
 /**
  * Used by RenderLights to render shadows to the attenuation buffer.
  *
@@ -3709,7 +3739,7 @@ bool FForwardShadingSceneRenderer::RenderShadowDepthMap(FRHICommandListImmediate
 				&& VisibleLightViewInfo.ProjectedShadowVisibilityMap[ShadowIndex];
 		}
 
-		if (IsProjectedShadowPotentiallyVisible(ProjectedShadowInfo) && (ProjectedShadowInfo->bWholeSceneShadow || ProjectedShadowInfo->bPreShadow))
+		if (IsProjectedForwardShadowPotentiallyVisible(ProjectedShadowInfo) && (ProjectedShadowInfo->bWholeSceneShadow || ProjectedShadowInfo->bPreShadow))
 		{
 			// Add the shadow to the list of visible shadows cast by this light.
 			if (ProjectedShadowInfo->bWholeSceneShadow)
@@ -3772,6 +3802,7 @@ void FForwardShadingSceneRenderer::RenderShadowDepthMaps(FRHICommandListImmediat
 	}	
 
 	// render shadowmaps for relevant lights.
+	bool bCSMRendered = false;
 	if (Scene->SimpleDirectionalLight)
 	{
 		const FLightSceneInfo& LightSceneInfo = *Scene->SimpleDirectionalLight;
@@ -3803,190 +3834,204 @@ void FForwardShadingSceneRenderer::RenderShadowDepthMaps(FRHICommandListImmediat
 				INC_DWORD_STAT(STAT_NumShadowedLights);
 
 				//render shadowmap depth				
-				RenderShadowDepthMap(RHICmdList, &LightSceneInfo);
+				bCSMRendered = RenderShadowDepthMap(RHICmdList, &LightSceneInfo);
 			}
 		}
 	}
+	
+	// CSM was not used, render modulated shadow depths.
+	// TODO: reuse any available space left from CSM depth renders for mod shadows.
+	if(!bCSMRendered)
+	{
+		RenderModulatedShadowDepthMaps(RHICmdList);
+	}
 }
 
-void FForwardShadingSceneRenderer::RenderModulatedShadowProjections(FRHICommandListImmediate& RHICmdList)
+static void GatherModulatedShadowProjections(FLightSceneInfo* LightSceneInfo, TArray<FVisibleLightInfo, SceneRenderingAllocator> VisibleLightInfos, TArray<FViewInfo> &Views, TArray<FProjectedShadowInfo*, SceneRenderingAllocator> &OutShadows)
 {
-	SCOPED_DRAW_EVENT(RHICmdList, Lights);
-
-	bool bDynamicShadows = ViewFamily.EngineShowFlags.DynamicShadows && GetShadowQuality() > 0;
-
-	if (IsSimpleDynamicLightingEnabled() || !bDynamicShadows || !DeviceSupportsShaderDepthFetch())
+	if (!LightSceneInfo->Proxy->CastsModulatedShadows())
 	{
 		return;
 	}
-	FSceneRenderTargets& SceneContext = FSceneRenderTargets::Get(RHICmdList);
 
-	// render shadowmaps for relevant lights.
-	for (TSparseArray<FLightSceneInfoCompact>::TConstIterator LightIt(Scene->Lights); LightIt; ++LightIt)
+	// Find the projected shadows cast by this light.
+	FVisibleLightInfo& VisibleLightInfo = VisibleLightInfos[LightSceneInfo->Id];
+	for (int32 ShadowIndex = 0; ShadowIndex < VisibleLightInfo.AllProjectedShadows.Num(); ShadowIndex++)
 	{
-		const FLightSceneInfoCompact& LightSceneInfoCompact = *LightIt;
-		FLightSceneInfo* LightSceneInfo = LightSceneInfoCompact.LightSceneInfo;
-
-		SCOPE_CYCLE_COUNTER(STAT_ProjectedShadowDrawTime);
-
-		if (!LightSceneInfo->Proxy->CastsModulatedShadows())
-		{
-			continue;
-		}
+		FProjectedShadowInfo* ProjectedShadowInfo = VisibleLightInfo.AllProjectedShadows[ShadowIndex];
 
-		// Find the projected shadows cast by this light.
-		FVisibleLightInfo& VisibleLightInfo = VisibleLightInfos[LightSceneInfo->Id];
-		TArray<FProjectedShadowInfo*, SceneRenderingAllocator> Shadows;
-		for (int32 ShadowIndex = 0; ShadowIndex < VisibleLightInfo.AllProjectedShadows.Num(); ShadowIndex++)
+		// Check that the shadow is visible in at least one view.
+		bool bShadowIsVisible = false;
+		for (int32 ViewIndex = 0; ViewIndex < Views.Num(); ViewIndex++)
 		{
-			FProjectedShadowInfo* ProjectedShadowInfo = VisibleLightInfo.AllProjectedShadows[ShadowIndex];
-
-			// Check that the shadow is visible in at least one view before rendering it.
-			bool bShadowIsVisible = false;
-			for (int32 ViewIndex = 0; ViewIndex < Views.Num(); ViewIndex++)
-			{
-				const FViewInfo& View = Views[ViewIndex];
-				if (ProjectedShadowInfo->DependentView && ProjectedShadowInfo->DependentView != &View)
-				{
-					continue;
-				}
-				const FVisibleLightViewInfo& VisibleLightViewInfo = View.VisibleLightInfos[LightSceneInfo->Id];
-				const FPrimitiveViewRelevance ViewRelevance = VisibleLightViewInfo.ProjectedShadowViewRelevanceMap[ShadowIndex];
-				bShadowIsVisible |= !ProjectedShadowInfo->bTranslucentShadow && ViewRelevance.bOpaqueRelevance
-					&& VisibleLightViewInfo.ProjectedShadowVisibilityMap[ShadowIndex];
-			}
-
-			bShadowIsVisible &= IsProjectedShadowPotentiallyVisible(ProjectedShadowInfo);
-
-			// skip shadows that are used with forward basepass (whole scene shadows)
-			if (ProjectedShadowInfo->bWholeSceneShadow)
-			{
-				bShadowIsVisible = false;
-			}
-
-			if (bShadowIsVisible)
+			const FViewInfo& View = Views[ViewIndex];
+			if (ProjectedShadowInfo->DependentView && ProjectedShadowInfo->DependentView != &View)
 			{
-				if (ProjectedShadowInfo->bPreShadow)
-				{
-					INC_DWORD_STAT(STAT_PreShadows);
-				}
-				else
-				{
-					INC_DWORD_STAT(STAT_PerObjectShadows);
-				}
-				Shadows.Add(ProjectedShadowInfo);
+				continue;
 			}
+			const FVisibleLightViewInfo& VisibleLightViewInfo = View.VisibleLightInfos[LightSceneInfo->Id];
+			const FPrimitiveViewRelevance ViewRelevance = VisibleLightViewInfo.ProjectedShadowViewRelevanceMap[ShadowIndex];
+			bShadowIsVisible |= !ProjectedShadowInfo->bTranslucentShadow && ViewRelevance.bOpaqueRelevance
+				&& VisibleLightViewInfo.ProjectedShadowVisibilityMap[ShadowIndex];
 		}
 
-		// Sort the projected shadows by resolution.
-		Shadows.Sort(FCompareFProjectedShadowInfoBySplitIndex());
+		bShadowIsVisible &= IsProjectedForwardShadowPotentiallyVisible(ProjectedShadowInfo);
 
-		int32 NumShadowsRendered = 0;
+		// skip shadows that are used with forward basepass (whole scene shadows)
+		if (ProjectedShadowInfo->bWholeSceneShadow)
+		{
+			bShadowIsVisible = false;
+		}
 
-		while (NumShadowsRendered < Shadows.Num())
+		if (bShadowIsVisible)
 		{
-			int32 NumAllocatedShadows = 0;
+			OutShadows.Add(ProjectedShadowInfo);
+		}
+	}
+}
 
-			// Allocate shadow texture space to the shadows.
-			const FIntPoint ShadowBufferResolution = SceneContext.GetShadowDepthTextureResolution();
-			FTextureLayout ShadowLayout(1, 1, ShadowBufferResolution.X, ShadowBufferResolution.Y, false, false);
 
-			for (int32 ShadowIndex = 0; ShadowIndex < Shadows.Num(); ShadowIndex++)
-			{
-				FProjectedShadowInfo* ProjectedShadowInfo = Shadows[ShadowIndex];
-				if (!ProjectedShadowInfo->bRendered)
-				{
-					if (ShadowLayout.AddElement(
-						ProjectedShadowInfo->X,
-						ProjectedShadowInfo->Y,
-						ProjectedShadowInfo->ResolutionX + SHADOW_BORDER * 2,
-						ProjectedShadowInfo->ResolutionY + SHADOW_BORDER * 2)
-						)
-					{
-						ProjectedShadowInfo->bAllocated = true;
-						NumAllocatedShadows++;
-					}
-				}
-			}
+void FForwardShadingSceneRenderer::RenderModulatedShadowDepthMaps(FRHICommandListImmediate& RHICmdList)
+{
+	SCOPED_DRAW_EVENT(RHICmdList, Lights);
+	
+	bool bDynamicShadows = ViewFamily.EngineShowFlags.DynamicShadows && GetShadowQuality() > 0;
+	
+	if (IsSimpleDynamicLightingEnabled() || !bDynamicShadows || !DeviceSupportsShaderDepthFetch())
+	{
+		bModulatedShadowsInUse = false;
+	}
+	
+	if(!bModulatedShadowsInUse)
+	{
+		return;
+	}
+	
+	bool bShadowsDepthsAllocated = false;
+	FSceneRenderTargets& SceneContext = FSceneRenderTargets::Get(RHICmdList);
+	TArray<FProjectedShadowInfo*, SceneRenderingAllocator> Shadows;
+	const FIntPoint ShadowBufferResolution = SceneContext.GetShadowDepthTextureResolution();
+	SceneContext.AllocateForwardShadingShadowDepthTarget(ShadowBufferResolution);
+	FTextureLayout ShadowLayout(1, 1, ShadowBufferResolution.X, ShadowBufferResolution.Y, false, false);
+	
+	// render shadowmaps for relevant lights.
+	for (TSparseArray<FLightSceneInfoCompact>::TConstIterator LightIt(Scene->Lights); LightIt; ++LightIt)
+	{
+		const FLightSceneInfoCompact& LightSceneInfoCompact = *LightIt;
+		FLightSceneInfo* LightSceneInfo = LightSceneInfoCompact.LightSceneInfo;
+		
+		SCOPE_CYCLE_COUNTER(STAT_ProjectedShadowDrawTime);
+		
+		GatherModulatedShadowProjections(LightSceneInfo, VisibleLightInfos, Views, Shadows);
 
-			// Abort if we encounter a shadow that doesn't fit in the render target.
-			if (!NumAllocatedShadows)
-			{
-				break;
-			}
+		// Sort the projected shadows by resolution.
+		Shadows.Sort(FCompareFProjectedShadowInfoBySplitIndex());
 
+		// Allocate shadow texture space to the shadows.
+		for (int32 ShadowIndex = 0; ShadowIndex < Shadows.Num(); ShadowIndex++)
+		{
+			FProjectedShadowInfo* ProjectedShadowInfo = Shadows[ShadowIndex];
+			if (!ProjectedShadowInfo->bRendered)
 			{
-				// Render the shadow depths.
-				SCOPED_DRAW_EVENT(RHICmdList, ShadowDepthsFromOpaqueProjected);
-
-				bool bPerformClear = true;
-				auto SetShadowRenderTargets = [&bPerformClear, &SceneContext](FRHICommandList& InRHICmdList)
-				{
-					SceneContext.BeginRenderingShadowDepth(InRHICmdList, bPerformClear);
-				};
-
-				SetShadowRenderTargets(RHICmdList);  // run it now, maybe run it later for parallel command lists
-				bPerformClear = false;
-
-				// render depth for each shadow
-				for (int32 ShadowIndex = 0; ShadowIndex < Shadows.Num(); ShadowIndex++)
+				if (ShadowLayout.AddElement(
+											ProjectedShadowInfo->X,
+											ProjectedShadowInfo->Y,
+											ProjectedShadowInfo->ResolutionX + SHADOW_BORDER * 2,
+											ProjectedShadowInfo->ResolutionY + SHADOW_BORDER * 2)
+					)
 				{
-					FProjectedShadowInfo* ProjectedShadowInfo = Shadows[ShadowIndex];
-					if (ProjectedShadowInfo->bAllocated && !ProjectedShadowInfo->bTranslucentShadow && !ProjectedShadowInfo->CascadeSettings.bRayTracedDistanceField)
-					{
-						ProjectedShadowInfo->RenderDepth(RHICmdList, this, SetShadowRenderTargets);
-					}
+					ProjectedShadowInfo->bAllocated = true;
+					bShadowsDepthsAllocated = true;
 				}
-
-				SceneContext.FinishRenderingShadowDepth(RHICmdList);
-			}
-
-			// Render the shadow projections.
-			{
-				SCOPED_DRAW_EVENT(RHICmdList, ShadowProjectionOnOpaque);
-				RenderProjections(RHICmdList, LightSceneInfo, Shadows, true);
 			}
+		}
+	}
 
-			// Mark and count the rendered shadows.
-			for (int32 ShadowIndex = 0; ShadowIndex < Shadows.Num(); ShadowIndex++)
+	if(bShadowsDepthsAllocated)
+	{
+		// Render the shadow depths.
+		SCOPED_DRAW_EVENT(RHICmdList, ShadowDepthsFromOpaqueProjected);
+	
+		bool bPerformClear = true;
+		auto SetShadowRenderTargets = [&bPerformClear, &SceneContext](FRHICommandList& InRHICmdList)
+		{
+			SceneContext.BeginRenderingShadowDepth(InRHICmdList, bPerformClear);
+		};
+	
+		SetShadowRenderTargets(RHICmdList);  // run it now, maybe run it later for parallel command lists
+		bPerformClear = false;
+	
+		// render depth for each shadow
+		for (int32 ShadowIndex = 0; ShadowIndex < Shadows.Num(); ShadowIndex++)
+		{
+			FProjectedShadowInfo* ProjectedShadowInfo = Shadows[ShadowIndex];
+			if (ProjectedShadowInfo->bAllocated && !ProjectedShadowInfo->bTranslucentShadow && !ProjectedShadowInfo->CascadeSettings.bRayTracedDistanceField)
 			{
-				FProjectedShadowInfo* ProjectedShadowInfo = Shadows[ShadowIndex];
-				if (ProjectedShadowInfo->bAllocated)
-				{
-					ProjectedShadowInfo->bAllocated = false;
-					ProjectedShadowInfo->bRendered = true;
-					NumShadowsRendered++;
-				}
+				ProjectedShadowInfo->RenderDepth(RHICmdList, this, SetShadowRenderTargets);
 			}
 		}
+	
+		SceneContext.FinishRenderingShadowDepth(RHICmdList);
 	}
+	
+	bModulatedShadowsInUse &= bShadowsDepthsAllocated;
 }
 
-bool FForwardShadingSceneRenderer::IsProjectedShadowPotentiallyVisible(const FProjectedShadowInfo* ProjectedShadowInfo) const
+void FForwardShadingSceneRenderer::RenderModulatedShadowProjections(FRHICommandListImmediate& RHICmdList)
 {
-	// Check that the shadow is visible in at least one view before rendering it.
-	const FLightSceneInfo& LightSceneInfo = ProjectedShadowInfo->GetLightSceneInfo();
-
-	// Skip shadows which will be handled in RenderOnePassPointLightShadows or RenderReflectiveShadowmaps
-	if (ProjectedShadowInfo->CascadeSettings.bOnePassPointLightShadow || ProjectedShadowInfo->bReflectiveShadowmap)
+	if (bModulatedShadowsInUse)
 	{
-		return false;
+		if(bCSMShadowsInUse)
+		{
+			// If CSM depths were rendered then no modulated shadow depths are ready.
+			RenderModulatedShadowDepthMaps(RHICmdList);
+		}
+	
+		while( bModulatedShadowsInUse )
+		{
+			RenderAllocatedModulatedShadowProjections(RHICmdList);
+			RenderModulatedShadowDepthMaps(RHICmdList);
+		}
 	}
+}
 
-	if (ProjectedShadowInfo->bPreShadow && !LightSceneInfo.Proxy->HasStaticShadowing())
-	{
-		return false;
-	}
+void FForwardShadingSceneRenderer::RenderAllocatedModulatedShadowProjections(FRHICommandListImmediate& RHICmdList)
+{
+	SCOPED_DRAW_EVENT(RHICmdList, Lights);
 
-	if (ProjectedShadowInfo->bWholeSceneShadow && LightSceneInfo.Proxy->HasStaticShadowing())
+	bool bDynamicShadows = ViewFamily.EngineShowFlags.DynamicShadows && GetShadowQuality() > 0;
+
+	if (IsSimpleDynamicLightingEnabled() || !bDynamicShadows || !DeviceSupportsShaderDepthFetch())
 	{
-		return false;
+		return;
 	}
+	FSceneRenderTargets& SceneContext = FSceneRenderTargets::Get(RHICmdList);
 
-	if ((ProjectedShadowInfo->bPreShadow && !ProjectedShadowInfo->HasSubjectPrims()) || ProjectedShadowInfo->bAllocatedInPreshadowCache)
+	// render shadowmaps for relevant lights.
+	for (TSparseArray<FLightSceneInfoCompact>::TConstIterator LightIt(Scene->Lights); LightIt; ++LightIt)
 	{
-		return false;
-	}
+		const FLightSceneInfoCompact& LightSceneInfoCompact = *LightIt;
+		FLightSceneInfo* LightSceneInfo = LightSceneInfoCompact.LightSceneInfo;
+		TArray<FProjectedShadowInfo*, SceneRenderingAllocator> Shadows;
+		SCOPE_CYCLE_COUNTER(STAT_ProjectedShadowDrawTime);
 
-	return true;
+		GatherModulatedShadowProjections(LightSceneInfo, VisibleLightInfos, Views, Shadows);
+
+		// Render the shadow projections.
+		{
+			SCOPED_DRAW_EVENT(RHICmdList, ShadowProjectionOnOpaque);
+			RenderProjections(RHICmdList, LightSceneInfo, Shadows, true);
+		}
+
+		// Mark the shadows as rendered.
+		for (int32 ShadowIndex = 0; ShadowIndex < Shadows.Num(); ShadowIndex++)
+		{
+			FProjectedShadowInfo* ProjectedShadowInfo = Shadows[ShadowIndex];
+			if (ProjectedShadowInfo->bAllocated)
+			{
+				ProjectedShadowInfo->bAllocated = false;
+				ProjectedShadowInfo->bRendered = true;
+			}
+		}
+	}
 }
diff --git a/Engine/Source/Runtime/Renderer/Private/ShadowSetup.cpp b/Engine/Source/Runtime/Renderer/Private/ShadowSetup.cpp
--- a/Engine/Source/Runtime/Renderer/Private/ShadowSetup.cpp
+++ b/Engine/Source/Runtime/Renderer/Private/ShadowSetup.cpp
@@ -2334,6 +2334,9 @@ void FForwardShadingSceneRenderer::InitDynamicShadows(FRHICommandListImmediate&
 		bStaticSceneOnly = bStaticSceneOnly || View.bStaticSceneOnly;
 	}
 
+	// Setup modulated shadows if CSM is not in use.
+	bool bCSMAllocated = false;
+	bool bPerObjectShadowsInUse = false;
 	// Modulated shadows require depth fetch
 	const bool bSupportsDepthFetch = DeviceSupportsShaderDepthFetch();
 
@@ -2364,7 +2367,10 @@ void FForwardShadingSceneRenderer::InitDynamicShadows(FRHICommandListImmediate&
 			int32 NumWholeSceneShadows = 0;
 			if (bIsVisibleInAnyView)
 			{
-				if (LightSceneInfo->ShouldRenderViewIndependentWholeSceneShadows() && Scene->SimpleDirectionalLight == LightSceneInfo)
+				if (LightSceneInfo->ShouldRenderViewIndependentWholeSceneShadows() && Scene->SimpleDirectionalLight == LightSceneInfo
+					// Only consider movable shadowcasting lights
+					&& !LightSceneInfo->Proxy->HasStaticShadowing()
+					)
 				{
 					AddViewDependentWholeSceneShadowsForView(ViewDependentWholeSceneShadows, ViewDependentWholeSceneShadowsThatNeedCulling, VisibleLightInfo, *LightSceneInfo);
 				}
@@ -2382,12 +2388,14 @@ void FForwardShadingSceneRenderer::InitDynamicShadows(FRHICommandListImmediate&
 						SetupInteractionShadows(RHICmdList, Interaction, VisibleLightInfo, bStaticSceneOnly, ViewDependentWholeSceneShadows, PreShadows);
 					}
 				}
+				bPerObjectShadowsInUse |= (NumWholeSceneShadows != VisibleLightInfo.AllProjectedShadows.Num());
 			}
 
 			NumWholeSceneShadows = FMath::Min(NumWholeSceneShadows, MAX_FORWARD_SHADOWCASCADES);
 
 			if (NumWholeSceneShadows > 0)
 			{
+				bCSMAllocated = true;
 				FSceneRenderTargets& SceneContext = FSceneRenderTargets::Get(RHICmdList);
 				//create the shadow depth texture and/or surface
 				const FIntPoint ShadowBufferResolution = SceneContext.GetShadowDepthTextureResolution();
@@ -2417,9 +2425,9 @@ void FForwardShadingSceneRenderer::InitDynamicShadows(FRHICommandListImmediate&
 					}
 				}
 			}
-			else if (VisibleLightInfo.AllProjectedShadows.Num() > 0)
+			else if(bPerObjectShadowsInUse)
 			{
-				// non whole scene projected shadows are in use. Ensure the shadow depth target is available for modulated shadow use later.
+				// Per obj projected shadows are in use. Ensure the shadow depth target is available for modulated shadow use later.
 				FSceneRenderTargets& SceneContext = FSceneRenderTargets::Get(RHICmdList);
 				const FIntPoint ShadowBufferResolution = SceneContext.GetShadowDepthTextureResolution();
 				SceneContext.AllocateForwardShadingShadowDepthTarget(ShadowBufferResolution);
@@ -2435,6 +2443,9 @@ void FForwardShadingSceneRenderer::InitDynamicShadows(FRHICommandListImmediate&
 
 	// Generate mesh element arrays from shadow primitive arrays
 	GatherShadowDynamicMeshElements();
+	
+	bCSMShadowsInUse = bCSMAllocated;
+	bModulatedShadowsInUse = bPerObjectShadowsInUse;
 }
 
 void FDeferredShadingSceneRenderer::InitDynamicShadows(FRHICommandListImmediate& RHICmdList)
