Decals flicker while moving camera on Nexus 5

diff --git a/Engine/Config/BaseDeviceProfiles.ini b/Engine/Config/BaseDeviceProfiles.ini
--- a/Engine/Config/BaseDeviceProfiles.ini
+++ b/Engine/Config/BaseDeviceProfiles.ini
@@ -74,6 +74,7 @@ BaseProfileName=
 +CVars=r.DepthOfFieldQuality=0
 +CVars=r.LightShaftQuality=0
 +CVars=r.RefractionQuality=0
++CVars=r.ShadowQuality=2
 // HZB will be slower with tiled, and not needed
 +CVars=r.HZBOcclusion=0
 // Non need for depth prepass
@@ -174,11 +175,13 @@ BaseProfileName=
 +CVars=r.DepthOfFieldQuality=0
 +CVars=r.LightShaftQuality=0
 +CVars=r.RefractionQuality=0
++CVars=r.ShadowQuality=2
 
 [Android_Low DeviceProfile]
 DeviceType=Android
 BaseProfileName=Android
 +CVars=r.MobileContentScaleFactor=0.5
++CVars=r.ShadowQuality=1
 
 [Android_Mid DeviceProfile]
 DeviceType=Android
diff --git a/Engine/Shaders/Common.usf b/Engine/Shaders/Common.usf
--- a/Engine/Shaders/Common.usf
+++ b/Engine/Shaders/Common.usf
@@ -830,7 +830,7 @@ float CalcSceneDepth( float2 ScreenUV )
 		#elif WEBGL
 			return Texture2DSampleLevel(SceneAlphaCopyTexture, SceneAlphaCopyTextureSampler, ScreenUV, 0).r;
 		#else 
-			float SceneW = Texture2DSampleLevel(SceneColorTexture, SceneColorTextureSampler, ScreenUV, 0).a;		
+			float SceneW = ConvertFromDeviceZ(Texture2DSampleLevel(SceneDepthTexture, SceneDepthTextureSampler, ScreenUV, 0).r);		
 			return DepthbufferFetchES2(SceneW, View.InvDeviceZToWorldZTransform[2], View.InvDeviceZToWorldZTransform[3]);
 		#endif 
 	#elif METAL_PROFILE
diff --git a/Engine/Source/Runtime/Renderer/Private/DecalRenderingForwardShading.cpp b/Engine/Source/Runtime/Renderer/Private/DecalRenderingForwardShading.cpp
--- a/Engine/Source/Runtime/Renderer/Private/DecalRenderingForwardShading.cpp
+++ b/Engine/Source/Runtime/Renderer/Private/DecalRenderingForwardShading.cpp
@@ -16,12 +16,6 @@ void FForwardShadingSceneRenderer::RenderDecals(FRHICommandListImmediate& RHICmd
 		return;
 	}
 
-	// Check if it possible to render decals on this device
-	if (!DeviceSupportsShaderDepthFetch())
-	{
-		return;
-	}
-
 	SCOPE_CYCLE_COUNTER(STAT_DecalsDrawTime);
 
 	FSceneRenderTargets& SceneContext = FSceneRenderTargets::Get(RHICmdList);
diff --git a/Engine/Source/Runtime/Renderer/Private/ForwardShadingRenderer.cpp b/Engine/Source/Runtime/Renderer/Private/ForwardShadingRenderer.cpp
--- a/Engine/Source/Runtime/Renderer/Private/ForwardShadingRenderer.cpp
+++ b/Engine/Source/Runtime/Renderer/Private/ForwardShadingRenderer.cpp
@@ -138,8 +138,8 @@ void FForwardShadingSceneRenderer::Render(FRHICommandListImmediate& RHICmdList)
 	RenderForwardShadingBasePass(RHICmdList);
 
 	// Make a copy of the scene depth if the current hardware doesn't support reading and writing to the same depth buffer
-	SceneContext.ResolveSceneDepthToAuxiliaryTexture(RHICmdList);
-
+	ConditionalResolveSceneDepth(RHICmdList);
+	
 	if (ViewFamily.EngineShowFlags.Decals)
 	{
 		RenderDecals(RHICmdList);
@@ -273,3 +273,34 @@ void FForwardShadingSceneRenderer::BasicPostProcess(FRHICommandListImmediate& RH
 
 	CompositeContext.Process(Context.FinalOutput.GetPass(), TEXT("ES2BasicPostProcess"));
 }
+
+void FForwardShadingSceneRenderer::ConditionalResolveSceneDepth(FRHICommandListImmediate& RHICmdList)
+{
+	FSceneRenderTargets& SceneContext = FSceneRenderTargets::Get(RHICmdList);
+	
+	SceneContext.ResolveSceneDepthToAuxiliaryTexture(RHICmdList);
+
+	auto ShaderPlatform = ViewFamily.GetShaderPlatform();
+
+	if (IsMobilePlatform(ShaderPlatform) 
+		&& !IsPCPlatform(ShaderPlatform)) // exclude mobile emulation on PC
+	{
+		bool bSceneDepthInAlpha = (SceneContext.GetSceneColor()->GetDesc().Format == PF_FloatRGBA);
+		bool bOnChipDepthFetch = (GSupportsShaderDepthStencilFetch || (bSceneDepthInAlpha && GSupportsShaderFramebufferFetch));
+		
+		if (!bOnChipDepthFetch)
+		{
+			// Only these features require depth texture
+			bool bDecals = ViewFamily.EngineShowFlags.Decals && Scene->Decals.Num();
+			bool bModulatedShadows = ViewFamily.EngineShowFlags.DynamicShadows && GetShadowQuality() > 0 && bModulatedShadowsInUse;
+
+			if (bDecals || bModulatedShadows)
+			{
+				// Switch depth target to force hardware flush current depth to texture
+				FTextureRHIRef DummyDepthTarget = GSystemTextures.DepthDummy->GetRenderTargetItem().TargetableTexture;
+				SetRenderTarget(RHICmdList, SceneContext.GetSceneColorSurface(), DummyDepthTarget, ESimpleRenderTargetMode::EUninitializedColorClearDepth, FExclusiveDepthStencil::DepthWrite_StencilWrite);
+				RHICmdList.DiscardRenderTargets(true, true, 0);
+			}
+		}
+	}
+}
\ No newline at end of file
diff --git a/Engine/Source/Runtime/Renderer/Private/PostProcess/SceneRenderTargets.cpp b/Engine/Source/Runtime/Renderer/Private/PostProcess/SceneRenderTargets.cpp
--- a/Engine/Source/Runtime/Renderer/Private/PostProcess/SceneRenderTargets.cpp
+++ b/Engine/Source/Runtime/Renderer/Private/PostProcess/SceneRenderTargets.cpp
@@ -1927,6 +1927,41 @@ const FTextureRHIRef& FSceneRenderTargets::GetSceneColorTexture() const
 	return (const FTextureRHIRef&)GetSceneColor()->GetRenderTargetItem().ShaderResourceTexture; 
 }
 
+const FTexture2DRHIRef* FSceneRenderTargets::GetActualDepthTexture() const
+{
+	const FTexture2DRHIRef* DepthTexture = NULL;
+	if((CurrentFeatureLevel >= ERHIFeatureLevel::SM4) || IsPCPlatform(GShaderPlatformForFeatureLevel[CurrentFeatureLevel]))
+	{
+		if(GSupportsDepthFetchDuringDepthTest)
+		{
+			DepthTexture = &GetSceneDepthTexture();
+		}
+		else
+		{
+			DepthTexture = &GetAuxiliarySceneDepthSurface();
+		}
+	}
+	else if (IsMobilePlatform(GShaderPlatformForFeatureLevel[CurrentFeatureLevel]))
+	{
+		bool bSceneDepthInAlpha = (GetSceneColor()->GetDesc().Format == PF_FloatRGBA);
+		bool bOnChipDepthFetch = (GSupportsShaderDepthStencilFetch || (bSceneDepthInAlpha && GSupportsShaderFramebufferFetch));
+		
+		if (bOnChipDepthFetch)
+		{
+			DepthTexture = (const FTexture2DRHIRef*)(&GSystemTextures.DepthDummy->GetRenderTargetItem().ShaderResourceTexture);
+		}
+		else
+		{
+			DepthTexture = &GetSceneDepthTexture();
+		}
+	}
+
+	check(DepthTexture != NULL);
+
+	return DepthTexture;
+}
+
+
 IPooledRenderTarget* FSceneRenderTargets::GetGBufferVelocityRT()
 {
 	if (!bAllocateVelocityGBuffer)
diff --git a/Engine/Source/Runtime/Renderer/Private/PostProcess/SceneRenderTargets.h b/Engine/Source/Runtime/Renderer/Private/PostProcess/SceneRenderTargets.h
--- a/Engine/Source/Runtime/Renderer/Private/PostProcess/SceneRenderTargets.h
+++ b/Engine/Source/Runtime/Renderer/Private/PostProcess/SceneRenderTargets.h
@@ -298,24 +298,7 @@ public:
 			return (const FTexture2DRHIRef&)ShadowDepthZ->GetRenderTargetItem().ShaderResourceTexture; 
 		}
 	}
-	const FTexture2DRHIRef* GetActualDepthTexture() const
-	{
-		const FTexture2DRHIRef* DepthTexture = NULL;
-		if((CurrentFeatureLevel >= ERHIFeatureLevel::SM4) || IsPCPlatform(GShaderPlatformForFeatureLevel[CurrentFeatureLevel]))
-		{
-			if(GSupportsDepthFetchDuringDepthTest)
-			{
-				DepthTexture = &GetSceneDepthTexture();
-			}
-			else
-			{
-				DepthTexture = &GetAuxiliarySceneDepthSurface();
-			}
-		}
-		check(DepthTexture != NULL);
-
-		return DepthTexture;
-	}
+	const FTexture2DRHIRef* GetActualDepthTexture() const;
 	const FTexture2DRHIRef& GetReflectiveShadowMapDepthTexture() const { return (const FTexture2DRHIRef&)ReflectiveShadowMapDepth->GetRenderTargetItem().ShaderResourceTexture; }
 	const FTexture2DRHIRef& GetReflectiveShadowMapNormalTexture() const { return (const FTexture2DRHIRef&)ReflectiveShadowMapNormal->GetRenderTargetItem().ShaderResourceTexture; }
 	const FTexture2DRHIRef& GetReflectiveShadowMapDiffuseTexture() const { return (const FTexture2DRHIRef&)ReflectiveShadowMapDiffuse->GetRenderTargetItem().ShaderResourceTexture; }
diff --git a/Engine/Source/Runtime/Renderer/Private/ScenePrivate.h b/Engine/Source/Runtime/Renderer/Private/ScenePrivate.h
--- a/Engine/Source/Runtime/Renderer/Private/ScenePrivate.h
+++ b/Engine/Source/Runtime/Renderer/Private/ScenePrivate.h
@@ -42,9 +42,6 @@ bool IsMobileHDR32bpp();
 /** True if the mobile renderer is emulating HDR with mosaic. */
 bool IsMobileHDRMosaic();
 
-/** True if device has some way to fetch depth in a fragment shader. */
-bool DeviceSupportsShaderDepthFetch();
-
 class FOcclusionQueryHelpers
 {
 public:
diff --git a/Engine/Source/Runtime/Renderer/Private/SceneRendering.cpp b/Engine/Source/Runtime/Renderer/Private/SceneRendering.cpp
--- a/Engine/Source/Runtime/Renderer/Private/SceneRendering.cpp
+++ b/Engine/Source/Runtime/Renderer/Private/SceneRendering.cpp
@@ -1773,16 +1773,3 @@ bool IsMobileHDRMosaic()
 			return !(GSupportsHDR32bppEncodeModeIntrinsic && GSupportsShaderFramebufferFetch);
 	}
 }
-
-bool DeviceSupportsShaderDepthFetch()
-{
-	if (IsMobilePlatform(GMaxRHIShaderPlatform))
-	{
-		// HDR should be enabled
-		// device should support extension to directly fetch depth in a shader
-		// or half-float color buffer (base pass stores depth in alpha channel)
-		return IsMobileHDR() && (GSupportsShaderDepthStencilFetch || (GSupportsRenderTargetFormat_PF_FloatRGBA && !IsMobileHDR32bpp()));
-	}
-
-	return true;
-}
\ No newline at end of file
diff --git a/Engine/Source/Runtime/Renderer/Private/SceneRendering.h b/Engine/Source/Runtime/Renderer/Private/SceneRendering.h
--- a/Engine/Source/Runtime/Renderer/Private/SceneRendering.h
+++ b/Engine/Source/Runtime/Renderer/Private/SceneRendering.h
@@ -897,6 +897,9 @@ protected:
 	/** Makes a copy of scene alpha so PC can emulate ES2 framebuffer fetch. */
 	void CopySceneAlpha(FRHICommandListImmediate& RHICmdList, const FViewInfo& View);
 
+	/** Resolves scene depth in case hardware does not support reading depth in the shader */
+	void ConditionalResolveSceneDepth(FRHICommandListImmediate& RHICmdList);
+
 	/** Renders decals. */
 	void RenderDecals(FRHICommandListImmediate& RHICmdList);
 
diff --git a/Engine/Source/Runtime/Renderer/Private/ShadowRendering.cpp b/Engine/Source/Runtime/Renderer/Private/ShadowRendering.cpp
--- a/Engine/Source/Runtime/Renderer/Private/ShadowRendering.cpp
+++ b/Engine/Source/Runtime/Renderer/Private/ShadowRendering.cpp
@@ -3897,7 +3897,7 @@ void FForwardShadingSceneRenderer::RenderModulatedShadowDepthMaps(FRHICommandLis
 	
 	bool bDynamicShadows = ViewFamily.EngineShowFlags.DynamicShadows && GetShadowQuality() > 0;
 	
-	if (IsSimpleDynamicLightingEnabled() || !bDynamicShadows || !DeviceSupportsShaderDepthFetch())
+	if (IsSimpleDynamicLightingEnabled() || !bDynamicShadows)
 	{
 		bModulatedShadowsInUse = false;
 	}
@@ -4001,7 +4001,7 @@ void FForwardShadingSceneRenderer::RenderAllocatedModulatedShadowProjections(FRH
 
 	bool bDynamicShadows = ViewFamily.EngineShowFlags.DynamicShadows && GetShadowQuality() > 0;
 
-	if (IsSimpleDynamicLightingEnabled() || !bDynamicShadows || !DeviceSupportsShaderDepthFetch())
+	if (IsSimpleDynamicLightingEnabled() || !bDynamicShadows)
 	{
 		return;
 	}
diff --git a/Engine/Source/Runtime/Renderer/Private/ShadowSetup.cpp b/Engine/Source/Runtime/Renderer/Private/ShadowSetup.cpp
--- a/Engine/Source/Runtime/Renderer/Private/ShadowSetup.cpp
+++ b/Engine/Source/Runtime/Renderer/Private/ShadowSetup.cpp
@@ -2337,8 +2337,6 @@ void FForwardShadingSceneRenderer::InitDynamicShadows(FRHICommandListImmediate&
 	// Setup modulated shadows if CSM is not in use.
 	bool bCSMAllocated = false;
 	bool bPerObjectShadowsInUse = false;
-	// Modulated shadows require depth fetch
-	const bool bSupportsDepthFetch = DeviceSupportsShaderDepthFetch();
 
 	{
 		SCOPE_CYCLE_COUNTER(STAT_InitDynamicShadowsTime);
@@ -2377,7 +2375,7 @@ void FForwardShadingSceneRenderer::InitDynamicShadows(FRHICommandListImmediate&
 				NumWholeSceneShadows = VisibleLightInfo.AllProjectedShadows.Num();
 
 				// If we're casting modulated shadows then look for dynamic shadow interactions.
-				if (bSupportsDepthFetch && LightSceneInfo->Proxy->CastsModulatedShadows())
+				if (LightSceneInfo->Proxy->CastsModulatedShadows())
 				{
 					// Look for individual primitives with a dynamic shadow.
 					for (FLightPrimitiveInteraction* Interaction = LightSceneInfo->DynamicPrimitiveList;
diff --git a/Engine/Source/Runtime/Renderer/Private/SystemTextures.cpp b/Engine/Source/Runtime/Renderer/Private/SystemTextures.cpp
--- a/Engine/Source/Runtime/Renderer/Private/SystemTextures.cpp
+++ b/Engine/Source/Runtime/Renderer/Private/SystemTextures.cpp
@@ -86,6 +86,16 @@ void FSystemTextures::InitializeTextures(FRHICommandListImmediate& RHICmdList, E
 			RHICmdList.Clear(true, FLinearColor(65000.0f, 65000.0f, 65000.0f, 65000.0f), false, 0, false, 0, FIntRect());
 			RHICmdList.CopyToResolveTarget(MaxFP16Depth->GetRenderTargetItem().TargetableTexture, MaxFP16Depth->GetRenderTargetItem().ShaderResourceTexture, true, FResolveParams());
 		}
+
+		// Create dummy 1x1 depth texture 
+		{
+			FPooledRenderTargetDesc Desc(FPooledRenderTargetDesc::Create2DDesc(FIntPoint(1, 1), PF_DepthStencil, FClearValueBinding::DepthFar, TexCreate_None, TexCreate_DepthStencilTargetable, false));
+			GRenderTargetPool.FindFreeElement(Desc, DepthDummy, TEXT("DepthDummy"));
+
+			SetRenderTarget(RHICmdList, FTextureRHIRef(), DepthDummy->GetRenderTargetItem().TargetableTexture);
+			RHICmdList.Clear(false, FLinearColor::White, true, (float)ERHIZBuffer::FarPlane, true, 0, FIntRect());
+			RHICmdList.CopyToResolveTarget(DepthDummy->GetRenderTargetItem().TargetableTexture, DepthDummy->GetRenderTargetItem().ShaderResourceTexture, true, FResolveParams());
+		}
 	}
 
 	// Create the PerlinNoise3D texture (similar to http://prettyprocs.wordpress.com/2012/10/20/fast-perlin-noise/)
@@ -391,6 +401,7 @@ void FSystemTextures::ReleaseDynamicRHI()
 	SSAORandomization.SafeRelease();
 	PreintegratedGF.SafeRelease();
 	MaxFP16Depth.SafeRelease();
+	DepthDummy.SafeRelease();
 
 	GRenderTargetPool.FreeUnusedResources();
 
diff --git a/Engine/Source/Runtime/Renderer/Private/SystemTextures.h b/Engine/Source/Runtime/Renderer/Private/SystemTextures.h
--- a/Engine/Source/Runtime/Renderer/Private/SystemTextures.h
+++ b/Engine/Source/Runtime/Renderer/Private/SystemTextures.h
@@ -48,6 +48,8 @@ public:
 	TRefCountPtr<IPooledRenderTarget> PreintegratedGF;
 	/** Texture that holds a single value containing the maximum depth that can be stored as FP16. */
 	TRefCountPtr<IPooledRenderTarget> MaxFP16Depth;
+	/** Depth texture that holds a single depth value */
+	TRefCountPtr<IPooledRenderTarget> DepthDummy;
 
 	/** Maximum feature level that the textures have been initialized up to */
 	ERHIFeatureLevel::Type FeatureLevelInitializedTo;
