Removed HalfResReflections CVAR and remaining code as it is no longer used

diff --git a/Engine/Source/Runtime/Renderer/Private/ReflectionEnvironment.cpp b/Engine/Source/Runtime/Renderer/Private/ReflectionEnvironment.cpp
--- a/Engine/Source/Runtime/Renderer/Private/ReflectionEnvironment.cpp
+++ b/Engine/Source/Runtime/Renderer/Private/ReflectionEnvironment.cpp
@@ -42,12 +42,6 @@ static TAutoConsoleVariable<int32> CVarReflectionEnvironment(
 	TEXT("Whether to render the reflection environment feature, which implements local reflections through Reflection Capture actors."),
 	ECVF_Cheat | ECVF_RenderThreadSafe);
 #endif 
-static TAutoConsoleVariable<int32> CVarHalfResReflections(
-	TEXT("r.HalfResReflections"),
-	0,
-	TEXT("Compute ReflectionEnvironment samples at half resolution.\n")
-	TEXT(" 0 is off (default), 1 is on"),
-	ECVF_RenderThreadSafe);
 
 static TAutoConsoleVariable<int32> CVarDoTiledReflections(
 	TEXT("r.DoTiledReflections"),
@@ -840,7 +834,6 @@ void FDeferredShadingSceneRenderer::RenderTiledDeferredImageBasedReflections(FRH
 	FSceneRenderTargets& SceneContext = FSceneRenderTargets::Get(RHICmdList);
 	static const auto AllowStaticLightingVar = IConsoleManager::Get().FindTConsoleVariableDataInt(TEXT("r.AllowStaticLighting"));
 	const bool bUseLightmaps = (AllowStaticLightingVar->GetValueOnRenderThread() == 1) && (CVarDiffuseFromCaptures.GetValueOnRenderThread() == 0);
-	const bool bHalfRes = CVarHalfResReflections.GetValueOnRenderThread() != 0;
 
 	TRefCountPtr<IPooledRenderTarget> NewSceneColor;
 	{
@@ -876,12 +869,6 @@ void FDeferredShadingSceneRenderer::RenderTiledDeferredImageBasedReflections(FRH
 			SetRenderTarget(RHICmdList, NULL, NULL);
 
 			FReflectionEnvironmentTiledDeferredCS* ComputeShader = NULL;
-			uint32 AdjustedReflectionTileSizeX = GReflectionEnvironmentTileSizeX;
-			if (bHalfRes)
-			{
-				AdjustedReflectionTileSizeX *= 2;
-			}
-			
 			TArray<FReflectionCaptureSortData> SortData;
 			int32 NumBoxCaptures = 0;
 			int32 NumSphereCaptures = 0;
@@ -896,7 +883,7 @@ void FDeferredShadingSceneRenderer::RenderTiledDeferredImageBasedReflections(FRH
 
 			ComputeShader->SetParameters(RHICmdList, View, SSROutput->GetRenderTargetItem().ShaderResourceTexture, SortData, NewSceneColor->GetRenderTargetItem().UAV, DynamicBentNormalAO);
 
-			uint32 GroupSizeX = (View.ViewRect.Size().X + AdjustedReflectionTileSizeX - 1) / AdjustedReflectionTileSizeX;
+			uint32 GroupSizeX = (View.ViewRect.Size().X + GReflectionEnvironmentTileSizeX - 1) / GReflectionEnvironmentTileSizeX;
 			uint32 GroupSizeY = (View.ViewRect.Size().Y + GReflectionEnvironmentTileSizeY - 1) / GReflectionEnvironmentTileSizeY;
 			DispatchComputeShader(RHICmdList, ComputeShader, GroupSizeX, GroupSizeY, 1);
 
