Exclude depth texture resolve for HTML5

diff --git a/Engine/Source/Runtime/Renderer/Private/ForwardShadingRenderer.cpp b/Engine/Source/Runtime/Renderer/Private/ForwardShadingRenderer.cpp
--- a/Engine/Source/Runtime/Renderer/Private/ForwardShadingRenderer.cpp
+++ b/Engine/Source/Runtime/Renderer/Private/ForwardShadingRenderer.cpp
@@ -280,6 +280,7 @@ void FForwardShadingSceneRenderer::ConditionalResolveSceneDepth(FRHICommandListI
 	
 	SceneContext.ResolveSceneDepthToAuxiliaryTexture(RHICmdList);
 
+#if !PLATFORM_HTML5
 	auto ShaderPlatform = ViewFamily.GetShaderPlatform();
 
 	if (IsMobilePlatform(ShaderPlatform) 
@@ -303,4 +304,5 @@ void FForwardShadingSceneRenderer::ConditionalResolveSceneDepth(FRHICommandListI
 			}
 		}
 	}
+#endif //!PLATFORM_HTML5
 }
\ No newline at end of file
diff --git a/Engine/Source/Runtime/Renderer/Private/SystemTextures.cpp b/Engine/Source/Runtime/Renderer/Private/SystemTextures.cpp
--- a/Engine/Source/Runtime/Renderer/Private/SystemTextures.cpp
+++ b/Engine/Source/Runtime/Renderer/Private/SystemTextures.cpp
@@ -87,7 +87,8 @@ void FSystemTextures::InitializeTextures(FRHICommandListImmediate& RHICmdList, E
 			RHICmdList.CopyToResolveTarget(MaxFP16Depth->GetRenderTargetItem().TargetableTexture, MaxFP16Depth->GetRenderTargetItem().ShaderResourceTexture, true, FResolveParams());
 		}
 
-		// Create dummy 1x1 depth texture 
+		// Create dummy 1x1 depth texture
+		if (IsMobilePlatform(GShaderPlatformForFeatureLevel[InFeatureLevel])) // currently used only for mobile
 		{
 			FPooledRenderTargetDesc Desc(FPooledRenderTargetDesc::Create2DDesc(FIntPoint(1, 1), PF_DepthStencil, FClearValueBinding::DepthFar, TexCreate_None, TexCreate_DepthStencilTargetable, false));
 			GRenderTargetPool.FindFreeElement(Desc, DepthDummy, TEXT("DepthDummy"));
