Resolve the Light Accumulation target before using it in the non-tiled reflections so fast clears work.

diff --git a/Engine/Source/Runtime/Renderer/Private/ReflectionEnvironment.cpp b/Engine/Source/Runtime/Renderer/Private/ReflectionEnvironment.cpp
--- a/Engine/Source/Runtime/Renderer/Private/ReflectionEnvironment.cpp
+++ b/Engine/Source/Runtime/Renderer/Private/ReflectionEnvironment.cpp
@@ -1060,6 +1060,7 @@ void FDeferredShadingSceneRenderer::RenderStandardDeferredImageBasedReflections(
 					StencilingGeometry::DrawSphere(RHICmdList);
 				}
 			}
+			RHICmdList.CopyToResolveTarget(LightAccumulation->GetRenderTargetItem().TargetableTexture, LightAccumulation->GetRenderTargetItem().ShaderResourceTexture, false, FResolveParams());
 
 			GRenderTargetPool.VisualizeTexture.SetCheckPoint(RHICmdList, LightAccumulation);
 		}
