Fix for Translucent Lighting Crash - UE-18950
GBufferResources were being used in TranslucentShadowDepthPS but not being set. They shouldn't be set as GBuffers aren't valid for this draw.
A similar issue occured a while ago in which the SceneTextureParameters were being used here and caused artefacts in the translucent shadows based on the SceneColor.
I avoided this by setting them to a dummy black texture. This approach won't work for the GBufferResources so I've removed that fix and gone about things differently.
The shader sets a define to have SceneTextureLookup() just return black.
This means the SceneTextures, nor the GBufferResources need be set at all for this draw.

diff --git a/Engine/Shaders/MaterialTemplate.usf b/Engine/Shaders/MaterialTemplate.usf
--- a/Engine/Shaders/MaterialTemplate.usf
+++ b/Engine/Shaders/MaterialTemplate.usf
@@ -1056,6 +1056,10 @@ float4 GetPostProcessInputMinMax(uint Index)
 /** Applies an offset to the scene texture lookup and decodes the HDR linear space color. */
 float4 SceneTextureLookup(float2 UV, int SceneTextureIndex, bool bFiltered)
 {
+#if SCENE_TEXTURES_DISABLED
+	return float4(0.0f, 0.0f, 0.0f, 0.0f);
+#endif
+
 	FScreenSpaceData ScreenSpaceData = GetScreenSpaceData(UV, false);
 
 #if PS4_PROFILE
diff --git a/Engine/Shaders/TranslucentShadowDepthShaders.usf b/Engine/Shaders/TranslucentShadowDepthShaders.usf
--- a/Engine/Shaders/TranslucentShadowDepthShaders.usf
+++ b/Engine/Shaders/TranslucentShadowDepthShaders.usf
@@ -2,6 +2,8 @@
 
 #define SPHERICAL_OPACITY_FOR_SHADOW_DEPTHS 1
 
+#define SCENE_TEXTURES_DISABLED 1
+
 #include "Common.usf"
 #include "Material.usf"
 #include "VertexFactory.usf"
diff --git a/Engine/Source/Runtime/Renderer/Private/PostProcess/SceneRenderTargets.cpp b/Engine/Source/Runtime/Renderer/Private/PostProcess/SceneRenderTargets.cpp
--- a/Engine/Source/Runtime/Renderer/Private/PostProcess/SceneRenderTargets.cpp
+++ b/Engine/Source/Runtime/Renderer/Private/PostProcess/SceneRenderTargets.cpp
@@ -2213,14 +2213,6 @@ void FSceneTextureShaderParameters::Set(
 			&& !SceneColorSurfaceParameter.IsBound()
 			&& !SceneDepthSurfaceParameter.IsBound());
 	}
-	else if( TextureMode == ESceneRenderTargetsMode::NonSceneAlignedPass )
-	{
-		FSamplerStateRHIParamRef DefaultSampler = TStaticSamplerState<SF_Point,AM_Clamp,AM_Clamp,AM_Clamp>::GetRHI();
-		SetTextureParameter(RHICmdList, ShaderRHI, SceneColorTextureParameter, SceneColorTextureParameterSampler, DefaultSampler, GBlackTexture->TextureRHI);
-		SetTextureParameter(RHICmdList, ShaderRHI, SceneDepthTextureParameter, SceneDepthTextureParameterSampler, DefaultSampler, GBlackTexture->TextureRHI);
-		SetTextureParameter(RHICmdList, ShaderRHI, SceneColorSurfaceParameter, GBlackTexture->TextureRHI);
-		SetTextureParameter(RHICmdList, ShaderRHI, SceneDepthSurfaceParameter, GBlackTexture->TextureRHI);
-	}
 
 	if( DirectionalOcclusionSampler.IsBound() )
 	{
diff --git a/Engine/Source/Runtime/Renderer/Private/TranslucentLighting.cpp b/Engine/Source/Runtime/Renderer/Private/TranslucentLighting.cpp
--- a/Engine/Source/Runtime/Renderer/Private/TranslucentLighting.cpp
+++ b/Engine/Source/Runtime/Renderer/Private/TranslucentLighting.cpp
@@ -258,7 +258,7 @@ public:
 		//@todo - scene depth can be bound by the material for use in depth fades
 		// This is incorrect when rendering a shadowmap as it's not from the camera's POV
 		// Set the scene depth texture to something safe when rendering shadow depths
-		FMeshMaterialShader::SetParameters(RHICmdList, ShaderRHI, MaterialRenderProxy, *MaterialRenderProxy->GetMaterial(FeatureLevel), View, ESceneRenderTargetsMode::NonSceneAlignedPass);
+		FMeshMaterialShader::SetParameters(RHICmdList, ShaderRHI, MaterialRenderProxy, *MaterialRenderProxy->GetMaterial(FeatureLevel), View, ESceneRenderTargetsMode::DontSet);
 
 		SetShaderValue(RHICmdList, ShaderRHI, TranslInvMaxSubjectDepth, ShadowInfo->InvMaxSubjectDepth);
 
diff --git a/Engine/Source/Runtime/Renderer/Public/SceneRenderTargetParameters.h b/Engine/Source/Runtime/Renderer/Public/SceneRenderTargetParameters.h
--- a/Engine/Source/Runtime/Renderer/Public/SceneRenderTargetParameters.h
+++ b/Engine/Source/Runtime/Renderer/Public/SceneRenderTargetParameters.h
@@ -13,7 +13,6 @@ namespace ESceneRenderTargetsMode
 		SetTextures,
 		DontSet,
 		DontSetIgnoreBoundByEditorCompositing,
-		NonSceneAlignedPass,
 	};
 }
 
