fixed UE-18724 Reflection environment CVar should exist for scalability purposes

from CL2459480: ReflectionEnvironment is no longer a cheat CVAR and is available in shipping builds. If you set the cvar to 2 (debug visualization) then it will be assumed to be the same as 1 in shipping builds.

diff --git a/Engine/Source/Runtime/Renderer/Private/ReflectionEnvironment.cpp b/Engine/Source/Runtime/Renderer/Private/ReflectionEnvironment.cpp
--- a/Engine/Source/Runtime/Renderer/Private/ReflectionEnvironment.cpp
+++ b/Engine/Source/Runtime/Renderer/Private/ReflectionEnvironment.cpp
@@ -34,37 +34,43 @@ static TAutoConsoleVariable<int32> CVarDiffuseFromCaptures(
 	TEXT(" 0 is off (default), 1 is on"),
 	ECVF_RenderThreadSafe);
 
-#if !(UE_BUILD_SHIPPING || UE_BUILD_TEST)
 static TAutoConsoleVariable<int32> CVarReflectionEnvironment(
 	TEXT("r.ReflectionEnvironment"),
 	1,
-	TEXT("0:off, 1:on and blend with scene, 2:on and overwrite scene.\n")
-	TEXT("Whether to render the reflection environment feature, which implements local reflections through Reflection Capture actors."),
-	ECVF_Cheat | ECVF_RenderThreadSafe);
-#endif 
+	TEXT("Whether to render the reflection environment feature, which implements local reflections through Reflection Capture actors.\n"),
+	TEXT(" 0: off\n")
+	TEXT(" 1: on and blend with scene (default)")
+	TEXT(" 2: on and overwrite scene (only in non-shipping builds)"),
+	ECVF_RenderThreadSafe | ECVF_Scalability);
 
 static TAutoConsoleVariable<int32> CVarDoTiledReflections(
 	TEXT("r.DoTiledReflections"),
 	1,
 	TEXT("Compute Reflection Environment with Tiled compute shader..\n")
-	TEXT(" 1 is on (default), 1 is on"),
+	TEXT(" 0: off\n"),
+	TEXT(" 1: on (default)"),
 	ECVF_RenderThreadSafe);
 
 static TAutoConsoleVariable<float> CVarSkySpecularOcclusionStrength(
 	TEXT("r.SkySpecularOcclusionStrength"),
 	1,
-	TEXT("Strength of skylight specular occlusion from DFAO"),
+	TEXT("Strength of skylight specular occlusion from DFAO (default is 1.0)"),
 	ECVF_RenderThreadSafe);
 
 // to avoid having direct access from many places
 static int GetReflectionEnvironmentCVar()
 {
-#if !(UE_BUILD_SHIPPING || UE_BUILD_TEST)
-	return CVarReflectionEnvironment.GetValueOnAnyThread();
+	int32 RetVal = CVarReflectionEnvironment.GetValueOnAnyThread();
+
+#if (UE_BUILD_SHIPPING || UE_BUILD_TEST)
+	// Disabling the debug part of this CVar when in shipping
+	if (RetVal == 2)
+	{
+		RetVal = 1;
+	}
 #endif
 
-	// on, default mode
-	return 1;
+	return RetVal;
 }
 
 bool IsReflectionEnvironmentAvailable(ERHIFeatureLevel::Type InFeatureLevel)
