fixed typo

diff --git a/Engine/Source/Runtime/Renderer/Private/ReflectionEnvironment.cpp b/Engine/Source/Runtime/Renderer/Private/ReflectionEnvironment.cpp
--- a/Engine/Source/Runtime/Renderer/Private/ReflectionEnvironment.cpp
+++ b/Engine/Source/Runtime/Renderer/Private/ReflectionEnvironment.cpp
@@ -47,7 +47,7 @@ static TAutoConsoleVariable<int32> CVarDoTiledReflections(
 	TEXT("r.DoTiledReflections"),
 	1,
 	TEXT("Compute Reflection Environment with Tiled compute shader..\n")
-	TEXT(" 0: off\n"),
+	TEXT(" 0: off\n")
 	TEXT(" 1: on (default)"),
 	ECVF_RenderThreadSafe);
 
