fixed typo

diff --git a/Engine/Source/Runtime/Renderer/Private/ReflectionEnvironment.cpp b/Engine/Source/Runtime/Renderer/Private/ReflectionEnvironment.cpp
--- a/Engine/Source/Runtime/Renderer/Private/ReflectionEnvironment.cpp
+++ b/Engine/Source/Runtime/Renderer/Private/ReflectionEnvironment.cpp
@@ -37,7 +37,7 @@ static TAutoConsoleVariable<int32> CVarDiffuseFromCaptures(
 static TAutoConsoleVariable<int32> CVarReflectionEnvironment(
 	TEXT("r.ReflectionEnvironment"),
 	1,
-	TEXT("Whether to render the reflection environment feature, which implements local reflections through Reflection Capture actors.\n"),
+	TEXT("Whether to render the reflection environment feature, which implements local reflections through Reflection Capture actors.\n")
 	TEXT(" 0: off\n")
 	TEXT(" 1: on and blend with scene (default)")
 	TEXT(" 2: on and overwrite scene (only in non-shipping builds)"),
