optimized Bloom, threshold pass is only needed if > -1

diff --git a/Engine/Source/Runtime/Engine/Classes/Engine/Scene.h b/Engine/Source/Runtime/Engine/Classes/Engine/Scene.h
--- a/Engine/Source/Runtime/Engine/Classes/Engine/Scene.h
+++ b/Engine/Source/Runtime/Engine/Classes/Engine/Scene.h
@@ -510,9 +510,9 @@ struct FPostProcessSettings
 
 	/**
 	 * minimum brightness the bloom starts having effect
-	 * -1:all pixels affect bloom equally (dream effect), 0:all pixels affect bloom brights more, 1(default), >1 brighter
+	 * -1:all pixels affect bloom equally (physically correct, faster as a threshold pass is omitted), 0:all pixels affect bloom brights more, 1(default), >1 brighter
 	 */
-	UPROPERTY(interp, BlueprintReadWrite, Category=Bloom, AdvancedDisplay, meta=(ClampMin = "-1.0", UIMin = "0.0", UIMax = "8.0", editcondition = "bOverride_BloomThreshold", DisplayName = "Threshold"))
+	UPROPERTY(interp, BlueprintReadWrite, Category=Bloom, AdvancedDisplay, meta=(ClampMin = "-1.0", UIMax = "8.0", editcondition = "bOverride_BloomThreshold", DisplayName = "Threshold"))
 	float BloomThreshold;
 
 	/**
diff --git a/Engine/Source/Runtime/Renderer/Private/PostProcess/PostProcessing.cpp b/Engine/Source/Runtime/Renderer/Private/PostProcess/PostProcessing.cpp
--- a/Engine/Source/Runtime/Renderer/Private/PostProcess/PostProcessing.cpp
+++ b/Engine/Source/Runtime/Renderer/Private/PostProcess/PostProcessing.cpp
@@ -151,18 +151,26 @@ FPostprocessContext::FPostprocessContext(class FRenderingCompositionGraph& InGra
 
 static FRenderingCompositeOutputRef RenderHalfResBloomThreshold(FPostprocessContext& Context, FRenderingCompositeOutputRef SceneColorHalfRes, FRenderingCompositeOutputRef EyeAdaptation)
 {
-// todo: optimize later, the missing node causes some wrong behavior
-//	if(Context.View.FinalPostProcessSettings.BloomIntensity <= 0.0f)
-//	{
-//		// this pass is not required
-//		return FRenderingCompositeOutputRef();
-//	}
-	// bloom threshold
-	FRenderingCompositePass* PostProcessBloomSetup = Context.Graph.RegisterPass(new(FMemStack::Get()) FRCPassPostProcessBloomSetup());
-	PostProcessBloomSetup->SetInput(ePId_Input0, SceneColorHalfRes);
-	PostProcessBloomSetup->SetInput(ePId_Input1, EyeAdaptation);
-
-	return FRenderingCompositeOutputRef(PostProcessBloomSetup);
+	if(Context.View.FinalPostProcessSettings.BloomThreshold <= -1)
+	{
+		// no need for threshold
+		return SceneColorHalfRes;
+	}
+	else
+	{
+		// todo: optimize later, the missing node causes some wrong behavior
+		//	if(Context.View.FinalPostProcessSettings.BloomIntensity <= 0.0f)
+		//	{
+		//		// this pass is not required
+		//		return FRenderingCompositeOutputRef();
+		//	}
+		// bloom threshold
+		FRenderingCompositePass* PostProcessBloomSetup = Context.Graph.RegisterPass(new(FMemStack::Get()) FRCPassPostProcessBloomSetup());
+		PostProcessBloomSetup->SetInput(ePId_Input0, SceneColorHalfRes);
+		PostProcessBloomSetup->SetInput(ePId_Input1, EyeAdaptation);
+
+		return FRenderingCompositeOutputRef(PostProcessBloomSetup);
+	}
 }
 
 
