Remove errant Mac-hack that is no longer required.

diff --git a/Engine/Shaders/PostProcessMotionBlur.usf b/Engine/Shaders/PostProcessMotionBlur.usf
--- a/Engine/Shaders/PostProcessMotionBlur.usf
+++ b/Engine/Shaders/PostProcessMotionBlur.usf
@@ -675,15 +675,7 @@ void VelocityScatterVS(
 
 	// needs to be the same on C++ side (faster on NVIDIA and AMD)
 	uint QuadsPerInstance = 8;
-#if MAC
-	// remap the indices to get vertexid to VId and quadid into IId
-	IId = IId * QuadsPerInstance + (VId / 6);
-	VId = VId % 6;
 
-	// triangle A: 0:left top, 1:right top, 2: left bottom
-	// triangle B: 3:right bottom, 4:left bottom, 5: right top
-	float2 CornerOffset = float2(VId % 2, VId > 1 && VId < 5) * 2 - 1;
-#else
 	// remap the indices to get vertexid to VId and quadid into IId
 	IId = IId * QuadsPerInstance + (VId / 4);
 	VId = VId % 4;
@@ -691,7 +683,6 @@ void VelocityScatterVS(
 	// triangle A: 0:left top, 1:right top, 2: left bottom
 	// triangle B: 3:right bottom, 4:left bottom, 5: right top
 	float2 CornerOffset = float2(VId % 2, VId / 2) * 2 - 1;
-#endif
 
 	uint2 TilePos = uint2( IId % ViewportRect.z, IId / ViewportRect.z );
 
diff --git a/Engine/Source/Runtime/Renderer/Private/PostProcess/PostProcessMotionBlur.cpp b/Engine/Source/Runtime/Renderer/Private/PostProcess/PostProcessMotionBlur.cpp
--- a/Engine/Source/Runtime/Renderer/Private/PostProcess/PostProcessMotionBlur.cpp
+++ b/Engine/Source/Runtime/Renderer/Private/PostProcess/PostProcessMotionBlur.cpp
@@ -883,21 +883,12 @@ public:
 		uint16* Indices = (uint16*)Buffer;
 		for (uint32 SpriteIndex = 0; SpriteIndex < 8; ++SpriteIndex)
 		{
-#if PLATFORM_MAC // Avoid a driver bug on OSX/NV cards that causes driver to generate an unwound index buffer
-			Indices[SpriteIndex*6 + 0] = SpriteIndex*6 + 0;
-			Indices[SpriteIndex*6 + 1] = SpriteIndex*6 + 1;
-			Indices[SpriteIndex*6 + 2] = SpriteIndex*6 + 2;
-			Indices[SpriteIndex*6 + 3] = SpriteIndex*6 + 3;
-			Indices[SpriteIndex*6 + 4] = SpriteIndex*6 + 4;
-			Indices[SpriteIndex*6 + 5] = SpriteIndex*6 + 5;
-#else
 			Indices[SpriteIndex*6 + 0] = SpriteIndex*4 + 0;
 			Indices[SpriteIndex*6 + 1] = SpriteIndex*4 + 3;
 			Indices[SpriteIndex*6 + 2] = SpriteIndex*4 + 2;
 			Indices[SpriteIndex*6 + 3] = SpriteIndex*4 + 0;
 			Indices[SpriteIndex*6 + 4] = SpriteIndex*4 + 1;
 			Indices[SpriteIndex*6 + 5] = SpriteIndex*4 + 3;
-#endif
 		}
 		RHIUnlockIndexBuffer( IndexBufferRHI );
 	}
