fixed UERNDR-75 The near depth of field blur should occur after separate translucency is merged in to the final frame buffer to avoid errors

We now occlude Separate Transluceny by new Gaussian DOF.

Details: optimized one full res pass if GaussianDOF with SeparateTranslucency is used. NearDOF is now occlusing/hiding SeparateTransluceny which allows DOFVignette to work better and it seems more right in general. If Far and Near DOF is used with Gaussian DOF it now requires another half res setup pass and full res composite pass.
Added enum to allow for 3 custom colors if a texture is not bound.

diff --git a/Engine/Shaders/PostProcessDOF.usf b/Engine/Shaders/PostProcessDOF.usf
--- a/Engine/Shaders/PostProcessDOF.usf
+++ b/Engine/Shaders/PostProcessDOF.usf
@@ -46,7 +46,7 @@ float2 ComputeDOFFocalMask(float SceneDepth, float SkyWithoutHorizonMask)
 void SetupPS(
 	float4 UVAndScreenPos : TEXCOORD0
 	, out float4 OutColor0 : SV_Target0
-#if ENABLE_NEAR_BLUR 
+#if MRT_COUNT > 0
 	, out float4 OutColor1 : SV_Target1
 #endif
 	)
@@ -59,15 +59,8 @@ void SetupPS(
 
 	float4 DepthQuad = GatherSceneDepth(UV, PostprocessInput1Size.zw);
 
-#if ENABLE_NEAR_BLUR == 0
-	// We aren't writing out to the second render target, so we'll just make a dummy value here which 
-	// doesn't end up going anywhere. Then, the source code can stay neat and tidy, while the compiler can still 
-	// strip it out
-	float4 OutColor1;
-#endif
-
-	OutColor0 = 0;
-	OutColor1 = 0;
+	float4 FarColor = 0;
+	float4 NearColor = 0;
 
 	float2 Mask;
 	float4 Sample;
@@ -86,7 +79,7 @@ void SetupPS(
 	// 0:see though..1:Near blurred
 	float VignetteMask = 0;
 
-#if ENABLE_DOF_VIGNETTE
+#if DOF_VIGNETTE
 	{
 		float DepthOfFieldVignetteMul = DepthOfFieldParams[0].y;
 		float DepthOfFieldVignetteAdd = DepthOfFieldParams[0].z;
@@ -103,27 +96,38 @@ void SetupPS(
 
 	Mask = ComputeDOFFocalMask(DepthQuad.x, SkyWithoutHorizonMask);
 	Sample = float4(Texture2DSampleLevel(PostprocessInput0, PostprocessInput0Sampler, UV + Offset * float2(-1, 1), 0).rgb, 1);
-	OutColor0 += Sample * Mask.x;
-	OutColor1 += Sample * lerp(Mask.y, 1, VignetteMask);
+	FarColor += Sample * Mask.x;
+	NearColor += Sample * lerp(Mask.y, 1, VignetteMask);
 
 	Mask = ComputeDOFFocalMask(DepthQuad.y, SkyWithoutHorizonMask);
 	Sample = float4(Texture2DSample(PostprocessInput0, PostprocessInput0Sampler, UV + Offset * float2(1, 1)).rgb, 1);
-	OutColor0 += Sample * Mask.x;
-	OutColor1 += Sample * lerp(Mask.y, 1, VignetteMask);
+	FarColor += Sample * Mask.x;
+	NearColor += Sample * lerp(Mask.y, 1, VignetteMask);
 
 	Mask = ComputeDOFFocalMask(DepthQuad.z, SkyWithoutHorizonMask);
 	Sample = float4(Texture2DSample(PostprocessInput0, PostprocessInput0Sampler, UV + Offset * float2(1, -1)).rgb, 1);
-	OutColor0 += Sample * Mask.x;
-	OutColor1 += Sample * lerp(Mask.y, 1, VignetteMask);
+	FarColor += Sample * Mask.x;
+	NearColor += Sample * lerp(Mask.y, 1, VignetteMask);
 
 	Mask = ComputeDOFFocalMask(DepthQuad.w, SkyWithoutHorizonMask);
 	Sample = float4(Texture2DSample(PostprocessInput0, PostprocessInput0Sampler, UV + Offset * float2(-1, -1)).rgb, 1);
-	OutColor0 += Sample * Mask.x;
-	OutColor1 += Sample * lerp(Mask.y, 1, VignetteMask);
+	FarColor += Sample * Mask.x;
+	NearColor += Sample * lerp(Mask.y, 1, VignetteMask);
 
 	// we average 4 samples
-	OutColor0 /= 4;
-	OutColor1 /= 4;
+	FarColor /= 4;
+	NearColor /= 4;
+
+#if MRT_COUNT > 1
+	OutColor0 = FarColor;
+	OutColor1 = NearColor;
+#else
+	#if NEAR_BLUR
+		OutColor0 = NearColor;
+	#else
+		OutColor0 = FarColor;
+	#endif
+#endif
 }
 
 
@@ -151,18 +155,22 @@ void MainRecombinePS(
 
 
 	float4 SceneColorAndDepth = float4(Texture2DSample(PostprocessInput0, PostprocessInput0Sampler, FullResUV).rgb, CalcSceneDepth(FullResUV));
-
+	
 	float3 UnfocusedSceneColor = SceneColorAndDepth.rgb;
 
-	// behind focal plane
-	float4 DOFAccumLayer1 = Texture2DSample(PostprocessInput1, PostprocessInput1Sampler, ViewportUV);
-#if ENABLE_NEAR_BLUR
-	float4 DOFAccumLayer3 = Texture2DSample(PostprocessInput2, PostprocessInput2Sampler, ViewportUV);
-#else
 	// I'm presuming all that matters here is the W==0 bit to mask out this value
 	// TODO: Should check that compiler is doing a good job of removing the usages of this
 	// from the rest of the code. It has no reason not to be able to do so...
+
+	float4 DOFAccumLayer1 = float4(0,0,0,0); 
 	float4 DOFAccumLayer3 = float4(0,0,0,0); 
+
+#if FAR_BLUR
+	DOFAccumLayer1 = Texture2DSample(PostprocessInput1, PostprocessInput1Sampler, ViewportUV);
+#endif
+
+#if NEAR_BLUR
+	DOFAccumLayer3 = Texture2DSample(PostprocessInput2, PostprocessInput2Sampler, ViewportUV);
 #endif
 
 	float Layer1Mask = DOFAccumLayer1.a;
@@ -197,6 +205,14 @@ void MainRecombinePS(
 //	LayerMerger = lerp(LayerMerger, SceneColorAndDepth.rgb, Layer2Mask * (1 - PerPixelNearMask));
 	LayerMerger = lerp(LayerMerger, SceneColorAndDepth.rgb, Layer2Mask);
 
+	// apply SeparateTransluceny (todo: shader permutation to save a bit of performance if the feature is disabled)
+	{
+		float4 SeparateTranslucency = Texture2DSample(PostprocessInput3, PostprocessInput3Sampler, FullResUV);
+		       
+		// add RGB, darken by A (this allows to represent translucent and additive blending)
+		LayerMerger.rgb = LayerMerger.rgb * SeparateTranslucency.a + SeparateTranslucency.rgb;
+	}
+
 	float3 FrontLayer = (UnfocusedSceneColor * Div0Bias + DOFAccumLayer3.rgb) / (DOFAccumLayer3.a + Div0Bias);
 
 	// Layer 3: on top of that blend the front half res layer	
diff --git a/Engine/Source/Runtime/Renderer/Private/PostProcess/PostProcessDOF.cpp b/Engine/Source/Runtime/Renderer/Private/PostProcess/PostProcessDOF.cpp
--- a/Engine/Source/Runtime/Renderer/Private/PostProcess/PostProcessDOF.cpp
+++ b/Engine/Source/Runtime/Renderer/Private/PostProcess/PostProcessDOF.cpp
@@ -14,8 +14,9 @@
 
 
 /** Encapsulates the DOF setup pixel shader. */
-// @param 0:no NearBlur, 1:with NearBlur, 2:With NearBlur and Vignette
-template <uint32 SetupMode>
+// @param FarBlur 0:off, 1:on
+// @param NearBlur 0:off, 1:on, 2:on with Vignette
+template <uint32 FarBlur, uint32 NearBlur>
 class FPostProcessDOFSetupPS : public FGlobalShader
 {
 	DECLARE_SHADER_TYPE(FPostProcessDOFSetupPS, Global);
@@ -28,8 +29,9 @@ class FPostProcessDOFSetupPS : public FGlobalShader
 	static void ModifyCompilationEnvironment(EShaderPlatform Platform, FShaderCompilerEnvironment& OutEnvironment)
 	{
 		FGlobalShader::ModifyCompilationEnvironment(Platform,OutEnvironment);
-		OutEnvironment.SetDefine(TEXT("ENABLE_NEAR_BLUR"), (uint32)(SetupMode >= 1));
-		OutEnvironment.SetDefine(TEXT("ENABLE_DOF_VIGNETTE"), (uint32)(SetupMode == 2));
+		OutEnvironment.SetDefine(TEXT("NEAR_BLUR"), (uint32)(NearBlur >= 1));
+		OutEnvironment.SetDefine(TEXT("DOF_VIGNETTE"), (uint32)(NearBlur == 2));
+		OutEnvironment.SetDefine(TEXT("MRT_COUNT"), FarBlur + (NearBlur > 0));
 	}
 
 	/** Default constructor. */
@@ -75,11 +77,45 @@ public:
 			SetShaderValueArray(Context.RHICmdList, ShaderRHI, DepthOfFieldParams, DepthOfFieldParamValues, 2);
 		}
 	}
+
+	static const TCHAR* GetSourceFilename()
+	{
+		return TEXT("PostProcessDOF");
+	}
+
+	static const TCHAR* GetFunctionName()
+	{
+		return TEXT("SetupPS");
+	}
 };
 
-IMPLEMENT_SHADER_TYPE(template<>,FPostProcessDOFSetupPS<0>,TEXT("PostProcessDOF"),TEXT("SetupPS"),SF_Pixel);
-IMPLEMENT_SHADER_TYPE(template<>,FPostProcessDOFSetupPS<1>,TEXT("PostProcessDOF"),TEXT("SetupPS"),SF_Pixel);
-IMPLEMENT_SHADER_TYPE(template<>,FPostProcessDOFSetupPS<2>,TEXT("PostProcessDOF"),TEXT("SetupPS"),SF_Pixel);
+// #define avoids a lot of code duplication
+#define VARIATION1(A, B) typedef FPostProcessDOFSetupPS<A, B> FPostProcessDOFSetupPS##A##B; \
+	IMPLEMENT_SHADER_TYPE2(FPostProcessDOFSetupPS##A##B, SF_Pixel);
+
+	VARIATION1(0, 1) VARIATION1(0, 2)
+	VARIATION1(1, 0) VARIATION1(1, 1) VARIATION1(1, 2)
+	
+#undef VARIATION1
+
+// @param FarBlur 0:off, 1:on
+// @param NearBlur 0:off, 1:on, 2:on with Vignette
+template <uint32 FarBlur, uint32 NearBlur>
+static FShader* SetDOFShaderTempl(const FRenderingCompositePassContext& Context)
+{
+	TShaderMapRef<FPostProcessVS> VertexShader(Context.GetShaderMap());
+	TShaderMapRef<FPostProcessDOFSetupPS<FarBlur, NearBlur> > PixelShader(Context.GetShaderMap());
+
+	static FGlobalBoundShaderState BoundShaderState;
+
+	SetGlobalBoundShaderState(Context.RHICmdList, Context.GetFeatureLevel(), BoundShaderState, GFilterVertexDeclaration.VertexDeclarationRHI, *VertexShader, *PixelShader);
+
+	VertexShader->SetParameters(Context);
+	PixelShader->SetParameters(Context);
+
+	return *VertexShader;
+}
+
 
 void FRCPassPostProcessDOFSetup::Process(FRenderingCompositePassContext& Context)
 {
@@ -93,7 +129,7 @@ void FRCPassPostProcessDOFSetup::Process(FRenderingCompositePassContext& Context
 		return;
 	}
 
-	uint32 NumRenderTargets = bNearBlurEnabled ? 2 : 1;
+	uint32 NumRenderTargets = (bNearBlur && bFarBlur) ? 2 : 1;
 
 	const FSceneView& View = Context.View;
 	const FSceneViewFamily& ViewFamily = *(View.Family);
@@ -111,7 +147,7 @@ void FRCPassPostProcessDOFSetup::Process(FRenderingCompositePassContext& Context
 	FIntRect DestRect = SrcRect / 2;
 
 	const FSceneRenderTargetItem& DestRenderTarget0 = PassOutputs[0].RequestSurface(Context);
-	const FSceneRenderTargetItem& DestRenderTarget1 = bNearBlurEnabled ? PassOutputs[1].RequestSurface(Context) : FSceneRenderTargetItem();
+	const FSceneRenderTargetItem& DestRenderTarget1 = (NumRenderTargets == 2) ? PassOutputs[1].RequestSurface(Context) : FSceneRenderTargetItem();
 
 	// Set the view family's render target/viewport.
 	FTextureRHIParamRef RenderTargets[2] =
@@ -136,46 +172,43 @@ void FRCPassPostProcessDOFSetup::Process(FRenderingCompositePassContext& Context
 	Context.RHICmdList.SetRasterizerState(TStaticRasterizerState<>::GetRHI());
 	Context.RHICmdList.SetDepthStencilState(TStaticDepthStencilState<false, CF_Always>::GetRHI());
 	
-	TShaderMapRef<FPostProcessVS> VertexShader(ShaderMap);
+	const float DOFVignetteSize = FMath::Max(0.0f, View.FinalPostProcessSettings.DepthOfFieldVignetteSize);
+
+	// todo: test is conservative, with bad content we would waste a bit of performance
+	const bool bDOFVignette = bNearBlur && DOFVignetteSize < 200.0f;
 	
-	if (bNearBlurEnabled)
+	// 0:off, 1:on, 2:on with Vignette
+	uint32 NearBlur = 0;
 	{
-		const float DOFVignetteSize = FMath::Max(0.0f, View.FinalPostProcessSettings.DepthOfFieldVignetteSize);
-
-		// todo: test is conservative, with bad content we would waste a bit of performance
-		const bool bDOFVignette = DOFVignetteSize < 200.0f;
-
-		if(bDOFVignette)
+		if(bNearBlur)
 		{
-			static FGlobalBoundShaderState BoundShaderState;
-
-			TShaderMapRef< FPostProcessDOFSetupPS<2> > PixelShader(ShaderMap);
-			SetGlobalBoundShaderState(Context.RHICmdList, FeatureLevel, BoundShaderState, GFilterVertexDeclaration.VertexDeclarationRHI, *VertexShader, *PixelShader);
-		
-			PixelShader->SetParameters(Context);
+			NearBlur = (DOFVignetteSize < 200.0f) ? 2 : 1;
 		}
-		else
-		{
-			static FGlobalBoundShaderState BoundShaderState;
-
-			TShaderMapRef< FPostProcessDOFSetupPS<1> > PixelShader(ShaderMap);
-			SetGlobalBoundShaderState(Context.RHICmdList, FeatureLevel, BoundShaderState, GFilterVertexDeclaration.VertexDeclarationRHI, *VertexShader, *PixelShader);
+	}
+	
+	FShader* VertexShader = 0;
 
-			PixelShader->SetParameters(Context);
+	if(bFarBlur)
+	{
+		switch(NearBlur)
+		{
+			case 0: VertexShader = SetDOFShaderTempl<1, 0>(Context);break;
+			case 1: VertexShader = SetDOFShaderTempl<1, 1>(Context);break;
+			case 2: VertexShader = SetDOFShaderTempl<1, 2>(Context);break;
+			default: check(!"Invalid NearBlur input");
 		}
 	}
 	else
 	{
-		static FGlobalBoundShaderState BoundShaderState;
-		
-		TShaderMapRef< FPostProcessDOFSetupPS<0> > PixelShader(ShaderMap);
-		SetGlobalBoundShaderState(Context.RHICmdList, FeatureLevel, BoundShaderState, GFilterVertexDeclaration.VertexDeclarationRHI, *VertexShader, *PixelShader);
-
-		PixelShader->SetParameters(Context);
+		switch(NearBlur)
+		{
+			case 0: check(!"Internal error: In this case we should not run the pass");break;
+			case 1: VertexShader = SetDOFShaderTempl<0, 1>(Context);break;
+			case 2: VertexShader = SetDOFShaderTempl<0, 2>(Context);break;
+			default: check(!"Invalid NearBlur input");
+		}
 	}
 
-	VertexShader->SetParameters(Context);
-
 	DrawPostProcessPass(
 		Context.RHICmdList,
 		0, 0,
@@ -184,7 +217,7 @@ void FRCPassPostProcessDOFSetup::Process(FRenderingCompositePassContext& Context
 		SrcRect.Width() + 1, SrcRect.Height() + 1,
 		DestRect.Size() + FIntPoint(1, 1),
 		SrcSize,
-		*VertexShader,
+		VertexShader,
 		View.StereoPass,
 		Context.HasHmdMesh(),
 		EDRF_UseTriangleOptimization);
@@ -217,8 +250,10 @@ FPooledRenderTargetDesc FRCPassPostProcessDOFSetup::ComputeOutputDesc(EPassOutpu
 
 
 
-/** Encapsulates  the DOF recombine pixel shader. */
-template <uint32 NearBlurEnable>
+/** Encapsulates the DOF setup pixel shader. */
+// @param FarBlur 0:off, 1:on
+// @param NearBlur 0:off, 1:on
+template <uint32 FarBlur, uint32 NearBlur>
 class FPostProcessDOFRecombinePS : public FGlobalShader
 {
 	DECLARE_SHADER_TYPE(FPostProcessDOFRecombinePS, Global);
@@ -231,7 +266,8 @@ class FPostProcessDOFRecombinePS : public FGlobalShader
 	static void ModifyCompilationEnvironment(EShaderPlatform Platform, FShaderCompilerEnvironment& OutEnvironment)
 	{
 		FGlobalShader::ModifyCompilationEnvironment(Platform,OutEnvironment);
-		OutEnvironment.SetDefine(TEXT("ENABLE_NEAR_BLUR"), NearBlurEnable);
+		OutEnvironment.SetDefine(TEXT("FAR_BLUR"), FarBlur);
+		OutEnvironment.SetDefine(TEXT("NEAR_BLUR"), NearBlur);
 	}
 
 	/** Default constructor. */
@@ -267,7 +303,7 @@ public:
 		FGlobalShader::SetParameters(Context.RHICmdList, ShaderRHI, Context.View);
 
 		DeferredParameters.Set(Context.RHICmdList, ShaderRHI, Context.View);
-		PostprocessParameter.SetPS(ShaderRHI, Context, TStaticSamplerState<SF_Bilinear,AM_Clamp,AM_Clamp,AM_Clamp>::GetRHI());
+		PostprocessParameter.SetPS(ShaderRHI, Context, TStaticSamplerState<SF_Bilinear,AM_Clamp,AM_Clamp,AM_Clamp>::GetRHI(), eFC_0001);
 		
 		// Compute out of bounds UVs in the source texture.
 		FVector4 Bounds;
@@ -278,23 +314,58 @@ public:
 
 		SetShaderValue(Context.RHICmdList, ShaderRHI, DepthOfFieldUVLimit, Bounds);
 	}
+	
+	static const TCHAR* GetSourceFilename()
+	{
+		return TEXT("PostProcessDOF");
+	}
+
+	static const TCHAR* GetFunctionName()
+	{
+		return TEXT("MainRecombinePS");
+	}
 };
 
-IMPLEMENT_SHADER_TYPE(template<>,FPostProcessDOFRecombinePS<0>,TEXT("PostProcessDOF"),TEXT("MainRecombinePS"),SF_Pixel);
-IMPLEMENT_SHADER_TYPE(template<>,FPostProcessDOFRecombinePS<1>,TEXT("PostProcessDOF"),TEXT("MainRecombinePS"),SF_Pixel);
+// #define avoids a lot of code duplication
+#define VARIATION1(A, B) typedef FPostProcessDOFRecombinePS<A, B> FPostProcessDOFRecombinePS##A##B; \
+	IMPLEMENT_SHADER_TYPE2(FPostProcessDOFRecombinePS##A##B, SF_Pixel);
+
+	VARIATION1(0, 1) VARIATION1(1, 0) VARIATION1(1, 1)
+	
+#undef VARIATION1
+
+	// @param FarBlur 0:off, 1:on
+// @param NearBlur 0:off, 1:on, 2:on with Vignette
+template <uint32 FarBlur, uint32 NearBlur>
+static FShader* SetDOFRecombineShaderTempl(const FRenderingCompositePassContext& Context)
+{
+	TShaderMapRef<FPostProcessVS> VertexShader(Context.GetShaderMap());
+	TShaderMapRef<FPostProcessDOFRecombinePS<FarBlur, NearBlur> > PixelShader(Context.GetShaderMap());
+
+	static FGlobalBoundShaderState BoundShaderState;
+
+	SetGlobalBoundShaderState(Context.RHICmdList, Context.GetFeatureLevel(), BoundShaderState, GFilterVertexDeclaration.VertexDeclarationRHI, *VertexShader, *PixelShader);
+
+	VertexShader->SetParameters(Context);
+	PixelShader->SetParameters(Context);
+
+	return *VertexShader;
+}
 
 void FRCPassPostProcessDOFRecombine::Process(FRenderingCompositePassContext& Context)
 {
 	SCOPED_DRAW_EVENT(Context.RHICmdList, DOFRecombine);
 
+	// Get the Far or Near RTDesc
 	const FPooledRenderTargetDesc* InputDesc = GetInputDesc(ePId_Input1);
-
+	
 	if(!InputDesc)
 	{
-		// input is not hooked up correctly
-		return;
+		InputDesc = GetInputDesc(ePId_Input2);
 	}
 
+	check(InputDesc);
+
 	const FSceneView& View = Context.View;
 
 	const auto FeatureLevel = Context.GetFeatureLevel();
@@ -322,27 +393,27 @@ void FRCPassPostProcessDOFRecombine::Process(FRenderingCompositePassContext& Con
 	Context.RHICmdList.SetRasterizerState(TStaticRasterizerState<>::GetRHI());
 	Context.RHICmdList.SetDepthStencilState(TStaticDepthStencilState<false, CF_Always>::GetRHI());
 
-	TShaderMapRef<FPostProcessVS> VertexShader(ShaderMap);
+	bool bFarBlur = GetInputDesc(ePId_Input1) != 0;
+	bool bNearBlur = GetInputDesc(ePId_Input2) != 0;
+
+	FShader* VertexShader = 0;
 
-	if (bNearBlurEnabled)
+	if(bFarBlur)
 	{
-		static FGlobalBoundShaderState BoundShaderState;
-		
-		TShaderMapRef< FPostProcessDOFRecombinePS<1> > PixelShader(ShaderMap);
-		SetGlobalBoundShaderState(Context.RHICmdList, FeatureLevel, BoundShaderState, GFilterVertexDeclaration.VertexDeclarationRHI, *VertexShader, *PixelShader);
-		PixelShader->SetParameters(Context);
+		if(bNearBlur)
+		{
+			VertexShader = SetDOFRecombineShaderTempl<1, 1>(Context);
+		}
+		else
+		{
+			VertexShader = SetDOFRecombineShaderTempl<1, 0>(Context);
+		}
 	}
 	else
 	{
-		static FGlobalBoundShaderState BoundShaderState;
-		
-		TShaderMapRef< FPostProcessDOFRecombinePS<0> > PixelShader(ShaderMap);
-		SetGlobalBoundShaderState(Context.RHICmdList, FeatureLevel, BoundShaderState, GFilterVertexDeclaration.VertexDeclarationRHI, *VertexShader, *PixelShader);
-		PixelShader->SetParameters(Context);
+		VertexShader = SetDOFRecombineShaderTempl<0, 1>(Context);
 	}
 
-	VertexShader->SetParameters(Context);
-
 	DrawPostProcessPass(
 		Context.RHICmdList,
 		0, 0,
@@ -351,11 +422,12 @@ void FRCPassPostProcessDOFRecombine::Process(FRenderingCompositePassContext& Con
 		HalfResViewRect.Width(), HalfResViewRect.Height(),
 		View.ViewRect.Size(),
 		TexSize,
-		*VertexShader,
+		VertexShader,
 		View.StereoPass,
 		Context.HasHmdMesh(),
 		EDRF_UseTriangleOptimization);
 
+
 	Context.RHICmdList.CopyToResolveTarget(DestRenderTarget.TargetableTexture, DestRenderTarget.ShaderResourceTexture, false, FResolveParams());
 }
 
diff --git a/Engine/Source/Runtime/Renderer/Private/PostProcess/PostProcessDOF.h b/Engine/Source/Runtime/Renderer/Private/PostProcess/PostProcessDOF.h
--- a/Engine/Source/Runtime/Renderer/Private/PostProcess/PostProcessDOF.h
+++ b/Engine/Source/Runtime/Renderer/Private/PostProcess/PostProcessDOF.h
@@ -15,8 +15,10 @@
 class FRCPassPostProcessDOFSetup : public TRenderingCompositePassBase<2, 2>
 {
 public:
-
-	FRCPassPostProcessDOFSetup(bool bInNearBlurEnabled) : bNearBlurEnabled ( bInNearBlurEnabled ) {}
+	FRCPassPostProcessDOFSetup(bool bInFarBlur, bool bInNearBlur) 
+		: bFarBlur(bInFarBlur)
+		, bNearBlur(bInNearBlur)
+	{}
 
 	// interface FRenderingCompositePass ---------
 
@@ -25,26 +27,21 @@ public:
 	virtual FPooledRenderTargetDesc ComputeOutputDesc(EPassOutputId InPassOutputId) const override;
 
 private:
-
-	bool bNearBlurEnabled;
+	
+	bool bFarBlur;
+	bool bNearBlur;
 };
 
 // derives from TRenderingCompositePassBase<InputCount, OutputCount> 
 // ePId_Input0: Full res scene color
-// ePId_Input1: output 0 from the DOFSetup (possibly further blurred)
-// ePId_Input2: output 1 from the DOFSetup (possibly further blurred)
-class FRCPassPostProcessDOFRecombine : public TRenderingCompositePassBase<3, 1>
+// ePId_Input1: FarBlur from the DOFSetup (possibly further blurred)
+// ePId_Input2: NearBlur from the DOFSetup (possibly further blurred)
+// ePId_Input3: optional SeparateTransluceny
+class FRCPassPostProcessDOFRecombine : public TRenderingCompositePassBase<4, 1>
 {
 public:
-
-	FRCPassPostProcessDOFRecombine(bool bInNearBlurEnabled) : bNearBlurEnabled ( bInNearBlurEnabled ) {}
-
 	// interface FRenderingCompositePass ---------
 	virtual void Process(FRenderingCompositePassContext& Context) override;
 	virtual void Release() override { delete this; }
 	virtual FPooledRenderTargetDesc ComputeOutputDesc(EPassOutputId InPassOutputId) const override;
-
-private:
-
-	bool bNearBlurEnabled;
 };
diff --git a/Engine/Source/Runtime/Renderer/Private/PostProcess/PostProcessMotionBlur.cpp b/Engine/Source/Runtime/Renderer/Private/PostProcess/PostProcessMotionBlur.cpp
--- a/Engine/Source/Runtime/Renderer/Private/PostProcess/PostProcessMotionBlur.cpp
+++ b/Engine/Source/Runtime/Renderer/Private/PostProcess/PostProcessMotionBlur.cpp
@@ -1247,7 +1247,7 @@ public:
 					TStaticSamplerState<SF_Point,AM_Clamp,AM_Clamp,AM_Clamp>::GetRHI(),
 				};
 
-				PostprocessParameter.SetPS( ShaderRHI, Context, 0, false, Filters );
+				PostprocessParameter.SetPS( ShaderRHI, Context, 0, eFC_0000, Filters );
 			}
 			else if( CVarMotionBlurSmoothMax.GetValueOnRenderThread() )
 			{
@@ -1259,7 +1259,7 @@ public:
 					TStaticSamplerState<SF_Bilinear,AM_Clamp,AM_Clamp,AM_Clamp>::GetRHI(),
 				};
 
-				PostprocessParameter.SetPS( ShaderRHI, Context, 0, false, Filters );
+				PostprocessParameter.SetPS( ShaderRHI, Context, 0, eFC_0000, Filters );
 			}
 			else
 			{
diff --git a/Engine/Source/Runtime/Renderer/Private/PostProcess/PostProcessTemporalAA.cpp b/Engine/Source/Runtime/Renderer/Private/PostProcess/PostProcessTemporalAA.cpp
--- a/Engine/Source/Runtime/Renderer/Private/PostProcess/PostProcessTemporalAA.cpp
+++ b/Engine/Source/Runtime/Renderer/Private/PostProcess/PostProcessTemporalAA.cpp
@@ -114,7 +114,7 @@ public:
 		FilterTable[2] = FilterTable[0];
 		FilterTable[3] = FilterTable[0];
 
-		PostprocessParameter.SetPS(ShaderRHI, Context, 0, false, FilterTable);
+		PostprocessParameter.SetPS(ShaderRHI, Context, 0, eFC_0000, FilterTable);
 
 		DeferredParameters.Set(Context.RHICmdList, ShaderRHI, Context.View);
 
diff --git a/Engine/Source/Runtime/Renderer/Private/PostProcess/PostProcessTonemap.cpp b/Engine/Source/Runtime/Renderer/Private/PostProcess/PostProcessTonemap.cpp
--- a/Engine/Source/Runtime/Renderer/Private/PostProcess/PostProcessTonemap.cpp
+++ b/Engine/Source/Runtime/Renderer/Private/PostProcess/PostProcessTonemap.cpp
@@ -917,7 +917,7 @@ public:
 				TStaticSamplerState<SF_Bilinear, AM_Clamp, AM_Clamp, AM_Clamp, 0, 1>::GetRHI(),
 			};
 
-			PostprocessParameter.SetPS(ShaderRHI, Context, 0, false, Filters);
+			PostprocessParameter.SetPS(ShaderRHI, Context, 0, eFC_0000, Filters);
 		}
 			
 		SetShaderValue(Context.RHICmdList, ShaderRHI, OverlayColor, Context.View.OverlayColor);
diff --git a/Engine/Source/Runtime/Renderer/Private/PostProcess/PostProcessUpscale.cpp b/Engine/Source/Runtime/Renderer/Private/PostProcess/PostProcessUpscale.cpp
--- a/Engine/Source/Runtime/Renderer/Private/PostProcess/PostProcessUpscale.cpp
+++ b/Engine/Source/Runtime/Renderer/Private/PostProcess/PostProcessUpscale.cpp
@@ -144,7 +144,7 @@ public:
 		FilterTable[0] = TStaticSamplerState<SF_Bilinear,AM_Clamp,AM_Clamp,AM_Clamp>::GetRHI();
 		FilterTable[1] = TStaticSamplerState<SF_Point,AM_Clamp,AM_Clamp,AM_Clamp>::GetRHI();
 			
-		PostprocessParameter.SetPS(ShaderRHI, Context, 0, false, FilterTable);
+		PostprocessParameter.SetPS(ShaderRHI, Context, 0, eFC_0000, FilterTable);
 		DeferredParameters.Set(Context.RHICmdList, ShaderRHI, Context.View);
 
 		{
diff --git a/Engine/Source/Runtime/Renderer/Private/PostProcess/PostProcessing.cpp b/Engine/Source/Runtime/Renderer/Private/PostProcess/PostProcessing.cpp
--- a/Engine/Source/Runtime/Renderer/Private/PostProcess/PostProcessing.cpp
+++ b/Engine/Source/Runtime/Renderer/Private/PostProcess/PostProcessing.cpp
@@ -309,7 +309,7 @@ static void AddPostProcessDepthOfFieldBokeh(FPostprocessContext& Context, FRende
 	if( Context.View.FinalPostProcessSettings.AntiAliasingMethod == AAM_TemporalAA && ViewState )
 	{
 		FRenderingCompositePass* HistoryInput;
-		if( ViewState->DOFHistoryRT && ViewState->bBokehDOFHistory && !Context.View.bCameraCut )
+		if( ViewState->DOFHistoryRT && ViewState->bDOFHistory && !Context.View.bCameraCut )
 		{
 			HistoryInput = Context.Graph.RegisterPass( new(FMemStack::Get()) FRCPassPostProcessInput( ViewState->DOFHistoryRT ) );
 		}
@@ -326,7 +326,7 @@ static void AddPostProcessDepthOfFieldBokeh(FPostprocessContext& Context, FRende
 		NodeTemporalAA->SetInput( ePId_Input3, VelocityInput );
 
 		DOFInputPass = NodeTemporalAA;
-		ViewState->bBokehDOFHistory = true;
+		ViewState->bDOFHistory = true;
 	}
 
 	FRenderingCompositePass* NodeBlurred = Context.Graph.RegisterPass(new(FMemStack::Get()) FRCPassPostProcessBokehDOF());
@@ -342,8 +342,10 @@ static void AddPostProcessDepthOfFieldBokeh(FPostprocessContext& Context, FRende
 	Context.FinalOutput = FRenderingCompositeOutputRef(NodeRecombined);
 }
 
-static void AddPostProcessDepthOfFieldGaussian(FPostprocessContext& Context, FDepthOfFieldStats& Out, FRenderingCompositeOutputRef& VelocityInput)
+static bool AddPostProcessDepthOfFieldGaussian(FPostprocessContext& Context, FDepthOfFieldStats& Out, FRenderingCompositeOutputRef& VelocityInput, FRenderingCompositeOutputRef& SeparateTranslucency)
 {
+	bool bSepTransWasApplied = false;
+
 	float FarSize = Context.View.FinalPostProcessSettings.DepthOfFieldFarBlurSize;
 	float NearSize = Context.View.FinalPostProcessSettings.DepthOfFieldNearBlurSize;
 	
@@ -355,99 +357,97 @@ static void AddPostProcessDepthOfFieldGaussian(FPostprocessContext& Context, FDe
 	Out.bFar = FarSize >= 0.01f;
 	Out.bNear = NearSize >= GetCachedScalabilityCVars().GaussianDOFNearThreshold;
 
-	if(!Out.bFar && !Out.bNear)
-	{
-		return;
-	}
-
 	if(Context.View.Family->EngineShowFlags.VisualizeDOF)
 	{
 		// no need for this pass
-		return;
+		Out.bFar = false;
+		Out.bNear = false;
 	}
 
-	FRenderingCompositePass* DOFSetup = Context.Graph.RegisterPass(new(FMemStack::Get()) FRCPassPostProcessDOFSetup(Out.bNear));
-	DOFSetup->SetInput(ePId_Input0, FRenderingCompositeOutputRef(Context.FinalOutput));
-	// We need the depth to create the near and far mask
-	DOFSetup->SetInput(ePId_Input1, FRenderingCompositeOutputRef(Context.SceneDepth));
-
 	FSceneViewState* ViewState = (FSceneViewState*)Context.View.State;
-	
-	FRenderingCompositePass* DOFInputPass = DOFSetup;
-	if( Context.View.FinalPostProcessSettings.AntiAliasingMethod == AAM_TemporalAA && ViewState )
+
+	if(Out.bFar)
 	{
-		FRenderingCompositePass* HistoryInput;
-		if( ViewState->DOFHistoryRT && !ViewState->bBokehDOFHistory && !Context.View.bCameraCut )
-		{
-			HistoryInput = Context.Graph.RegisterPass( new(FMemStack::Get()) FRCPassPostProcessInput( ViewState->DOFHistoryRT ) );
-		}
-		else
+		FRenderingCompositePass* DOFSetupFar = Context.Graph.RegisterPass(new(FMemStack::Get()) FRCPassPostProcessDOFSetup(Out.bFar, false));
+		DOFSetupFar->SetInput(ePId_Input0, FRenderingCompositeOutputRef(Context.FinalOutput));
+		DOFSetupFar->SetInput(ePId_Input1, FRenderingCompositeOutputRef(Context.SceneDepth));
+
+		FRenderingCompositePass* DOFInputPassFar = DOFSetupFar;
+		if(Context.View.FinalPostProcessSettings.AntiAliasingMethod == AAM_TemporalAA && ViewState)
 		{
-			// No history so use current as history
-			HistoryInput = DOFSetup;
+			// If no history use current as history
+			FRenderingCompositePass* HistoryInput = DOFSetupFar;
+
+			if(ViewState->DOFHistoryRT && !ViewState->bDOFHistory && !Context.View.bCameraCut)
+			{
+				HistoryInput = Context.Graph.RegisterPass( new(FMemStack::Get()) FRCPassPostProcessInput(ViewState->DOFHistoryRT));
+			}
+
+			FRenderingCompositePass* NodeTemporalAA = Context.Graph.RegisterPass( new(FMemStack::Get()) FRCPassPostProcessDOFTemporalAA );
+			NodeTemporalAA->SetInput( ePId_Input0, DOFSetupFar );
+			NodeTemporalAA->SetInput( ePId_Input1, FRenderingCompositeOutputRef( HistoryInput ) );
+			NodeTemporalAA->SetInput( ePId_Input2, FRenderingCompositeOutputRef( HistoryInput ) );
+			NodeTemporalAA->SetInput( ePId_Input3, VelocityInput );
+
+			DOFInputPassFar = NodeTemporalAA;
+			ViewState->bDOFHistory = false;
 		}
+	
+		FRenderingCompositeOutputRef Far = RenderGaussianBlur(Context, TEXT("FarDOFBlurX"), TEXT("FarDOFBlurY"), FRenderingCompositeOutputRef(DOFInputPassFar, ePId_Output0), FarSize);
 
-		FRenderingCompositePass* NodeTemporalAA = Context.Graph.RegisterPass( new(FMemStack::Get()) FRCPassPostProcessDOFTemporalAA );
-		NodeTemporalAA->SetInput( ePId_Input0, DOFSetup );
-		NodeTemporalAA->SetInput( ePId_Input1, FRenderingCompositeOutputRef( HistoryInput ) );
-		NodeTemporalAA->SetInput( ePId_Input2, FRenderingCompositeOutputRef( HistoryInput ) );
-		NodeTemporalAA->SetInput( ePId_Input3, VelocityInput );
+		FRenderingCompositePass* NodeAllButNear = Context.Graph.RegisterPass(new(FMemStack::Get()) FRCPassPostProcessDOFRecombine());
+		NodeAllButNear->SetInput(ePId_Input0, Context.FinalOutput);
+		NodeAllButNear->SetInput(ePId_Input1, Far);
+		NodeAllButNear->SetInput(ePId_Input3, SeparateTranslucency);
+		bSepTransWasApplied = true;
 
-		DOFInputPass = NodeTemporalAA;
-		ViewState->bBokehDOFHistory = false;
+		Context.FinalOutput = FRenderingCompositeOutputRef(NodeAllButNear);
 	}
 
-	FRenderingCompositePass* DOFInputPass2 = DOFSetup;
-	EPassOutputId DOFInputPassId = ePId_Output1;
-	if( Context.View.FinalPostProcessSettings.AntiAliasingMethod == AAM_TemporalAA && ViewState )
+	// near
+	if(Out.bNear)
 	{
-		FRenderingCompositePass* HistoryInput;
-		EPassOutputId DOFInputPassId2 = ePId_Output1;
-		if( ViewState->DOFHistoryRT2 && !ViewState->bBokehDOFHistory2 && !Context.View.bCameraCut )
-		{
-			HistoryInput = Context.Graph.RegisterPass( new(FMemStack::Get()) FRCPassPostProcessInput( ViewState->DOFHistoryRT2 ) );
-			DOFInputPassId2 = ePId_Output0;
-		}
-		else
+		FRenderingCompositePass* DOFSetupNear = Context.Graph.RegisterPass(new(FMemStack::Get()) FRCPassPostProcessDOFSetup(false, Out.bNear));
+		DOFSetupNear->SetInput(ePId_Input0, Context.FinalOutput);
+		DOFSetupNear->SetInput(ePId_Input1, Context.SceneDepth);
+		
+		FRenderingCompositePass* DOFInputPassNear = DOFSetupNear;
+		if(Context.View.FinalPostProcessSettings.AntiAliasingMethod == AAM_TemporalAA && ViewState)
 		{
-			// No history so use current as history
-			HistoryInput = DOFSetup;
-			DOFInputPassId2 = ePId_Output1;
-		}
+			// If no history use current as history
+			FRenderingCompositePass* HistoryInput = DOFSetupNear;
 
-		FRenderingCompositePass* NodeTemporalAA = Context.Graph.RegisterPass( new(FMemStack::Get()) FRCPassPostProcessDOFTemporalAANear );
-		NodeTemporalAA->SetInput( ePId_Input0, FRenderingCompositeOutputRef( DOFSetup, ePId_Output1 ) );
-		NodeTemporalAA->SetInput( ePId_Input1, FRenderingCompositeOutputRef( HistoryInput, DOFInputPassId2 ) );
-		NodeTemporalAA->SetInput( ePId_Input2, FRenderingCompositeOutputRef( HistoryInput, DOFInputPassId2 ) );
-		NodeTemporalAA->SetInput( ePId_Input3, VelocityInput );
+			if(ViewState->DOFHistoryRT2 && !ViewState->bDOFHistory2 && !Context.View.bCameraCut)
+			{
+				HistoryInput = Context.Graph.RegisterPass(new(FMemStack::Get()) FRCPassPostProcessInput(ViewState->DOFHistoryRT2));
+			}
 
-		DOFInputPass2 = NodeTemporalAA;
-		ViewState->bBokehDOFHistory2 = false;
+			FRenderingCompositePass* NodeTemporalAA = Context.Graph.RegisterPass(new(FMemStack::Get()) FRCPassPostProcessDOFTemporalAANear);
+			NodeTemporalAA->SetInput( ePId_Input0, DOFSetupNear );
+			NodeTemporalAA->SetInput( ePId_Input1, HistoryInput);
+			NodeTemporalAA->SetInput( ePId_Input2, HistoryInput);
+			NodeTemporalAA->SetInput( ePId_Input3, VelocityInput );
 
-		DOFInputPassId = ePId_Output0;
-	}
-	
-	FRenderingCompositeOutputRef Far(DOFInputPass, ePId_Output0);
-	FRenderingCompositeOutputRef Near; // Don't need to bind a dummy here as we use a different permutation which doesn't read from this input when near DOF is disabled
+			DOFInputPassNear = NodeTemporalAA;
+			ViewState->bDOFHistory2 = false;
+		}
 
-	// far
-	if (Out.bFar)
-	{
-		Far = RenderGaussianBlur(Context, TEXT("FarDOFBlurX"), TEXT("FarDOFBlurY"), FRenderingCompositeOutputRef(DOFInputPass, ePId_Output0), FarSize);
-	}
+		FRenderingCompositeOutputRef Near = RenderGaussianBlur(Context, TEXT("NearDOFBlurX"), TEXT("NearDOFBlurY"), FRenderingCompositeOutputRef(DOFInputPassNear, ePId_Output0), NearSize);
 
-	// near
-	if (Out.bNear)
-	{
-		Near = RenderGaussianBlur(Context, TEXT("NearDOFBlurX"), TEXT("NearDOFBlurY"), FRenderingCompositeOutputRef(DOFInputPass2, DOFInputPassId), NearSize);
-	}
+		FRenderingCompositePass* NodeRecombined = Context.Graph.RegisterPass(new(FMemStack::Get()) FRCPassPostProcessDOFRecombine());
+		NodeRecombined->SetInput(ePId_Input0, Context.FinalOutput);
+		NodeRecombined->SetInput(ePId_Input2, Near);
 
-	FRenderingCompositePass* NodeRecombined = Context.Graph.RegisterPass(new(FMemStack::Get()) FRCPassPostProcessDOFRecombine(Out.bNear));
-	NodeRecombined->SetInput(ePId_Input0, Context.FinalOutput);
-	NodeRecombined->SetInput(ePId_Input1, Far);
-	NodeRecombined->SetInput(ePId_Input2, Near);
+		if(!bSepTransWasApplied)
+		{
+			NodeRecombined->SetInput(ePId_Input3, SeparateTranslucency);
+			bSepTransWasApplied = true;
+		}
 
-	Context.FinalOutput = FRenderingCompositeOutputRef(NodeRecombined);
+		Context.FinalOutput = FRenderingCompositeOutputRef(NodeRecombined);
+	}
+
+	return bSepTransWasApplied;
 }
 
 static void AddPostProcessDepthOfFieldCircle(FPostprocessContext& Context, FDepthOfFieldStats& Out, FRenderingCompositeOutputRef& VelocityInput)
@@ -468,7 +468,7 @@ static void AddPostProcessDepthOfFieldCircle(FPostprocessContext& Context, FDept
 	if( Context.View.FinalPostProcessSettings.AntiAliasingMethod == AAM_TemporalAA && ViewState )
 	{
 		FRenderingCompositePass* HistoryInput;
-		if( ViewState->DOFHistoryRT && !ViewState->bBokehDOFHistory && !Context.View.bCameraCut )
+		if( ViewState->DOFHistoryRT && !ViewState->bDOFHistory && !Context.View.bCameraCut )
 		{
 			HistoryInput = Context.Graph.RegisterPass( new(FMemStack::Get()) FRCPassPostProcessInput( ViewState->DOFHistoryRT ) );
 		}
@@ -485,7 +485,7 @@ static void AddPostProcessDepthOfFieldCircle(FPostprocessContext& Context, FDept
 		NodeTemporalAA->SetInput( ePId_Input3, VelocityInput );
 
 		DOFInputPass = NodeTemporalAA;
-		ViewState->bBokehDOFHistory = false;
+		ViewState->bDOFHistory = false;
 	}
 
 	FRenderingCompositeOutputRef Far;
@@ -1075,6 +1075,8 @@ void FPostProcessing::Process(FRHICommandListImmediate& RHICmdList, FViewInfo& V
 
 			FDepthOfFieldStats DepthOfFieldStat;
 
+			bool bSepTransWasApplied = false;
+
 			if(bDepthOfField && View.FinalPostProcessSettings.DepthOfFieldMethod != DOFM_BokehDOF)
 			{
 				bool bCircleDOF = View.FinalPostProcessSettings.DepthOfFieldMethod == DOFM_CircleDOF;
@@ -1082,14 +1084,16 @@ void FPostProcessing::Process(FRHICommandListImmediate& RHICmdList, FViewInfo& V
 				{
 					if(VelocityInput.IsValid())
 					{
-						AddPostProcessDepthOfFieldGaussian(Context, DepthOfFieldStat, VelocityInput);
+						bSepTransWasApplied = AddPostProcessDepthOfFieldGaussian(Context, DepthOfFieldStat, VelocityInput, SeparateTranslucency);
 					}
 					else
 					{
+						// todo: black/white default is a compositing graph feature, no need to hook up a node
+
 						// black is how we clear the velocity buffer so this means no velocity
 						FRenderingCompositePass* NoVelocity = Context.Graph.RegisterPass(new(FMemStack::Get()) FRCPassPostProcessInput(GSystemTextures.BlackDummy));
 						FRenderingCompositeOutputRef NoVelocityRef(NoVelocity);
-						AddPostProcessDepthOfFieldGaussian(Context, DepthOfFieldStat, NoVelocityRef);
+						bSepTransWasApplied = AddPostProcessDepthOfFieldGaussian(Context, DepthOfFieldStat, NoVelocityRef, SeparateTranslucency);
 					}
 				}
 				else
@@ -1100,6 +1104,8 @@ void FPostProcessing::Process(FRHICommandListImmediate& RHICmdList, FViewInfo& V
 					}
 					else
 					{
+						// todo: black/white default is a compositing graph feature, no need to hook up a node
+
 						// black is how we clear the velocity buffer so this means no velocity
 						FRenderingCompositePass* NoVelocity = Context.Graph.RegisterPass(new(FMemStack::Get()) FRCPassPostProcessInput(GSystemTextures.BlackDummy));
 						FRenderingCompositeOutputRef NoVelocityRef(NoVelocity);
@@ -1121,23 +1127,24 @@ void FPostProcessing::Process(FRHICommandListImmediate& RHICmdList, FViewInfo& V
 				}
 				else
 				{
+					// todo: black/white default is a compositing graph feature, no need to hook up a node
+
 					// black is how we clear the velocity buffer so this means no velocity
 					FRenderingCompositePass* NoVelocity = Context.Graph.RegisterPass(new(FMemStack::Get()) FRCPassPostProcessInput(GSystemTextures.BlackDummy));
 					FRenderingCompositeOutputRef NoVelocityRef(NoVelocity);
 					AddPostProcessDepthOfFieldBokeh(Context, SeparateTranslucency, NoVelocityRef);
 				}
+				bSepTransWasApplied = true;
 			}
-			else 
+
+			if(SeparateTranslucency.IsValid() && !bSepTransWasApplied)
 			{
-				if(SeparateTranslucency.IsValid())
-				{
-					// separate translucency is done here or in AddPostProcessDepthOfFieldBokeh()
-					FRenderingCompositePass* NodeRecombined = Context.Graph.RegisterPass(new(FMemStack::Get()) FRCPassPostProcessBokehDOFRecombine());
-					NodeRecombined->SetInput(ePId_Input0, Context.FinalOutput);
-					NodeRecombined->SetInput(ePId_Input2, SeparateTranslucency);
+				// separate translucency is done here or in AddPostProcessDepthOfFieldBokeh()
+				FRenderingCompositePass* NodeRecombined = Context.Graph.RegisterPass(new(FMemStack::Get()) FRCPassPostProcessBokehDOFRecombine());
+				NodeRecombined->SetInput(ePId_Input0, Context.FinalOutput);
+				NodeRecombined->SetInput(ePId_Input2, SeparateTranslucency);
 
-					Context.FinalOutput = FRenderingCompositeOutputRef(NodeRecombined);
-				}
+				Context.FinalOutput = FRenderingCompositeOutputRef(NodeRecombined);
 			}
 
 			AddPostProcessMaterial(Context, BL_BeforeTonemapping, SeparateTranslucency);
diff --git a/Engine/Source/Runtime/Renderer/Private/PostProcess/RenderingCompositionGraph.cpp b/Engine/Source/Runtime/Renderer/Private/PostProcess/RenderingCompositionGraph.cpp
--- a/Engine/Source/Runtime/Renderer/Private/PostProcess/RenderingCompositionGraph.cpp
+++ b/Engine/Source/Runtime/Renderer/Private/PostProcess/RenderingCompositionGraph.cpp
@@ -691,19 +691,19 @@ void FPostProcessPassParameters::Bind(const FShaderParameterMap& ParameterMap)
 	}
 }
 
-void FPostProcessPassParameters::SetPS(const FPixelShaderRHIParamRef& ShaderRHI, const FRenderingCompositePassContext& Context, FSamplerStateRHIParamRef Filter, bool bWhiteIfNoTexture, FSamplerStateRHIParamRef* FilterOverrideArray)
+void FPostProcessPassParameters::SetPS(const FPixelShaderRHIParamRef& ShaderRHI, const FRenderingCompositePassContext& Context, FSamplerStateRHIParamRef Filter, EFallbackColor FallbackColor, FSamplerStateRHIParamRef* FilterOverrideArray)
 {
-	Set(ShaderRHI, Context, Filter, bWhiteIfNoTexture, FilterOverrideArray);
+	Set(ShaderRHI, Context, Filter, FallbackColor, FilterOverrideArray);
 }
 
-void FPostProcessPassParameters::SetCS(const FComputeShaderRHIParamRef& ShaderRHI, const FRenderingCompositePassContext& Context, FSamplerStateRHIParamRef Filter, bool bWhiteIfNoTexture, FSamplerStateRHIParamRef* FilterOverrideArray)
+void FPostProcessPassParameters::SetCS(const FComputeShaderRHIParamRef& ShaderRHI, const FRenderingCompositePassContext& Context, FSamplerStateRHIParamRef Filter, EFallbackColor FallbackColor, FSamplerStateRHIParamRef* FilterOverrideArray)
 {
-	Set(ShaderRHI, Context, Filter, bWhiteIfNoTexture, FilterOverrideArray);
+	Set(ShaderRHI, Context, Filter, FallbackColor, FilterOverrideArray);
 }
 
-void FPostProcessPassParameters::SetVS(const FVertexShaderRHIParamRef& ShaderRHI, const FRenderingCompositePassContext& Context, FSamplerStateRHIParamRef Filter, bool bWhiteIfNoTexture, FSamplerStateRHIParamRef* FilterOverrideArray)
+void FPostProcessPassParameters::SetVS(const FVertexShaderRHIParamRef& ShaderRHI, const FRenderingCompositePassContext& Context, FSamplerStateRHIParamRef Filter, EFallbackColor FallbackColor, FSamplerStateRHIParamRef* FilterOverrideArray)
 {
-	Set(ShaderRHI, Context, Filter, bWhiteIfNoTexture, FilterOverrideArray);
+	Set(ShaderRHI, Context, Filter, FallbackColor, FilterOverrideArray);
 }
 
 template< typename ShaderRHIParamRef >
@@ -711,7 +711,7 @@ void FPostProcessPassParameters::Set(
 	const ShaderRHIParamRef& ShaderRHI,
 	const FRenderingCompositePassContext& Context,
 	FSamplerStateRHIParamRef Filter,
-	bool bWhiteIfNoTexture,
+	EFallbackColor FallbackColor,
 	FSamplerStateRHIParamRef* FilterOverrideArray)
 {
 	// assuming all outputs have the same size
@@ -780,6 +780,17 @@ void FPostProcessPassParameters::Set(
 									((float)ContextViewportRect.Max.X/SceneRTSize.X), 
 									((float)ContextViewportRect.Max.Y/SceneRTSize.Y) );
 
+	IPooledRenderTarget* FallbackTexture = 0;
+	
+	switch(FallbackColor)
+	{
+		case eFC_0000: FallbackTexture = GSystemTextures.BlackDummy; break;
+		case eFC_0001: FallbackTexture = GSystemTextures.BlackAlphaOneDummy; break;
+		case eFC_1111: FallbackTexture = GSystemTextures.WhiteDummy; break;
+		default:
+			ensure(!"Unhandled enum in EFallbackColor");
+	}
+
 	// ePId_Input0, ePId_Input1, ...
 	for(uint32 Id = 0; Id < (uint32)ePId_Input_MAX; ++Id)
 	{
@@ -831,11 +842,9 @@ void FPostProcessPassParameters::Set(
 		}
 		else
 		{
-			IPooledRenderTarget* Texture = bWhiteIfNoTexture ? GSystemTextures.WhiteDummy : GSystemTextures.BlackDummy;
-
 			// if the input is not there but the shader request it we give it at least some data to avoid d3ddebug errors and shader permutations
 			// to make features optional we use default black for additive passes without shader permutations
-			SetTextureParameter(RHICmdList, ShaderRHI, PostprocessInputParameter[Id], PostprocessInputParameterSampler[Id], LocalFilter, Texture->GetRenderTargetItem().TargetableTexture);
+			SetTextureParameter(RHICmdList, ShaderRHI, PostprocessInputParameter[Id], PostprocessInputParameterSampler[Id], LocalFilter, FallbackTexture->GetRenderTargetItem().TargetableTexture);
 
 			FVector4 Dummy(1, 1, 1, 1);
 			SetShaderValue(RHICmdList, ShaderRHI, PostprocessInputSizeParameter[Id], Dummy);
@@ -851,7 +860,7 @@ void FPostProcessPassParameters::Set(
 		const ShaderRHIParamRef& ShaderRHI,				\
 		const FRenderingCompositePassContext& Context,	\
 		FSamplerStateRHIParamRef Filter,				\
-		bool bWhiteIfNoTexture,							\
+		EFallbackColor FallbackColor,					\
 		FSamplerStateRHIParamRef* FilterOverrideArray	\
 	);
 
diff --git a/Engine/Source/Runtime/Renderer/Private/RendererScene.cpp b/Engine/Source/Runtime/Renderer/Private/RendererScene.cpp
--- a/Engine/Source/Runtime/Renderer/Private/RendererScene.cpp
+++ b/Engine/Source/Runtime/Renderer/Private/RendererScene.cpp
@@ -86,8 +86,8 @@ FSceneViewState::FSceneViewState()
 	DistanceFieldTemporalSampleIndex = 0;
 	AOTileIntersectionResources = NULL;
 	AOScreenGridResources = NULL;
-	bBokehDOFHistory = true;
-	bBokehDOFHistory2 = true;
+	bDOFHistory = true;
+	bDOFHistory2 = true;
 
 	LightPropagationVolume = NULL; 
 
diff --git a/Engine/Source/Runtime/Renderer/Private/ScenePrivate.h b/Engine/Source/Runtime/Renderer/Private/ScenePrivate.h
--- a/Engine/Source/Runtime/Renderer/Private/ScenePrivate.h
+++ b/Engine/Source/Runtime/Renderer/Private/ScenePrivate.h
@@ -526,9 +526,10 @@ public:
 	FGlobalDistanceFieldClipmapState GlobalDistanceFieldClipmapState[GMaxGlobalDistanceFieldClipmaps];
 	int32 GlobalDistanceFieldUpdateIndex;
 
-	// Is DOFHistoryRT set from Bokeh DOF?
-	bool bBokehDOFHistory;
-	bool bBokehDOFHistory2;
+	// Is DOFHistoryRT set from DepthOfField?
+	bool bDOFHistory;
+	// Is DOFHistoryRT2 set from DepthOfField?
+	bool bDOFHistory2;
 
 	FTemporalLODState TemporalLODState;
 
diff --git a/Engine/Source/Runtime/Renderer/Private/SystemTextures.cpp b/Engine/Source/Runtime/Renderer/Private/SystemTextures.cpp
--- a/Engine/Source/Runtime/Renderer/Private/SystemTextures.cpp
+++ b/Engine/Source/Runtime/Renderer/Private/SystemTextures.cpp
@@ -26,7 +26,7 @@ void FSystemTextures::InitializeTextures(FRHICommandListImmediate& RHICmdList, E
 	// requested feature level
 	if (!bTexturesInitialized)
 	{
-		// Create a dummy white texture.
+		// Create a WhiteDummy texture
 		{
 			FPooledRenderTargetDesc Desc(FPooledRenderTargetDesc::Create2DDesc(FIntPoint(1, 1), PF_B8G8R8A8, FClearValueBinding::None, TexCreate_HideInVisualizeTexture, TexCreate_RenderTargetable | TexCreate_NoFastClear, false));
 			GRenderTargetPool.FindFreeElement(Desc, WhiteDummy, TEXT("WhiteDummy"));
@@ -36,7 +36,7 @@ void FSystemTextures::InitializeTextures(FRHICommandListImmediate& RHICmdList, E
 			RHICmdList.CopyToResolveTarget(WhiteDummy->GetRenderTargetItem().TargetableTexture, WhiteDummy->GetRenderTargetItem().ShaderResourceTexture, true, FResolveParams());
 		}
 
-		// Create a dummy black texture.
+		// Create a BlackDummy texture
 		{
 			FPooledRenderTargetDesc Desc(FPooledRenderTargetDesc::Create2DDesc(FIntPoint(1, 1), PF_B8G8R8A8, FClearValueBinding::None, TexCreate_HideInVisualizeTexture, TexCreate_RenderTargetable | TexCreate_NoFastClear, false));
 			GRenderTargetPool.FindFreeElement(Desc, BlackDummy, TEXT("BlackDummy"));
@@ -45,6 +45,16 @@ void FSystemTextures::InitializeTextures(FRHICommandListImmediate& RHICmdList, E
 			RHICmdList.Clear(true, FLinearColor(0, 0, 0, 0), false, 0, false, 0, FIntRect());
 			RHICmdList.CopyToResolveTarget(BlackDummy->GetRenderTargetItem().TargetableTexture, BlackDummy->GetRenderTargetItem().ShaderResourceTexture, true, FResolveParams());
 		}
+		
+		// Create a BlackAlphaOneDummy texture
+		{
+			FPooledRenderTargetDesc Desc(FPooledRenderTargetDesc::Create2DDesc(FIntPoint(1, 1), PF_B8G8R8A8, FClearValueBinding::None, TexCreate_HideInVisualizeTexture, TexCreate_RenderTargetable | TexCreate_NoFastClear, false));
+			GRenderTargetPool.FindFreeElement(Desc, BlackAlphaOneDummy, TEXT("BlackAlphaOneDummy"));
+
+			SetRenderTarget(RHICmdList, BlackAlphaOneDummy->GetRenderTargetItem().TargetableTexture, FTextureRHIRef());
+			RHICmdList.Clear(true, FLinearColor(0, 0, 0, 1), false, 0, false, 0, FIntRect());
+			RHICmdList.CopyToResolveTarget(BlackAlphaOneDummy->GetRenderTargetItem().TargetableTexture, BlackAlphaOneDummy->GetRenderTargetItem().ShaderResourceTexture, true, FResolveParams());
+		}
 
 		// Create the PerlinNoiseGradient texture
 		{
diff --git a/Engine/Source/Runtime/Renderer/Private/SystemTextures.h b/Engine/Source/Runtime/Renderer/Private/SystemTextures.h
--- a/Engine/Source/Runtime/Renderer/Private/SystemTextures.h
+++ b/Engine/Source/Runtime/Renderer/Private/SystemTextures.h
@@ -34,10 +34,12 @@ public:
 
 	// -----------
 
-	// used in case a light is not shadow casting
+	// float4(1,1,1,1) can be used in case a light is not shadow casting
 	TRefCountPtr<IPooledRenderTarget> WhiteDummy;
-	// use in additive postprocessing to avoid a shader combination
+	// float4(0,0,0,0) can be used in additive postprocessing to avoid a shader combination
 	TRefCountPtr<IPooledRenderTarget> BlackDummy;
+	// float4(0,0,0,1)
+	TRefCountPtr<IPooledRenderTarget> BlackAlphaOneDummy;
 	// used by the material expression Noise
 	TRefCountPtr<IPooledRenderTarget> PerlinNoiseGradient;
 	// used by the material expression Noise (faster version, should replace old version), todo: move out of SceneRenderTargets
diff --git a/Engine/Source/Runtime/Renderer/Public/PostProcessParameters.h b/Engine/Source/Runtime/Renderer/Public/PostProcessParameters.h
--- a/Engine/Source/Runtime/Renderer/Public/PostProcessParameters.h
+++ b/Engine/Source/Runtime/Renderer/Public/PostProcessParameters.h
@@ -53,6 +53,16 @@ enum EPassOutputId
 	ePId_Output7
 };
 
+enum EFallbackColor
+{
+	// float4(0,0,0,0) BlackDummy
+	eFC_0000,
+	// float4(1,1,1,1) WhiteDummy 
+	eFC_1111,
+	// float4(0,0,0,1) BlackAlphaOneDummy
+	eFC_0001,
+};
+
 // currently hard coded to 4 input textures
 // convenient but not the most optimized solution
 struct FPostProcessPassParameters 
@@ -61,13 +71,13 @@ struct FPostProcessPassParameters
 	void Bind(const FShaderParameterMap& ParameterMap);
 
 	/** Set the pixel shader parameter values. */
-	void SetPS(const FPixelShaderRHIParamRef& ShaderRHI, const FRenderingCompositePassContext& Context, FSamplerStateRHIParamRef Filter = TStaticSamplerState<>::GetRHI(), bool bWhiteIfNoTexture = false, FSamplerStateRHIParamRef* FilterOverrideArray = 0);
+	void SetPS(const FPixelShaderRHIParamRef& ShaderRHI, const FRenderingCompositePassContext& Context, FSamplerStateRHIParamRef Filter = TStaticSamplerState<>::GetRHI(), EFallbackColor FallbackColor = eFC_0000, FSamplerStateRHIParamRef* FilterOverrideArray = 0);
 
 	/** Set the compute shader parameter values. */
-	void SetCS(const FComputeShaderRHIParamRef& ShaderRHI, const FRenderingCompositePassContext& Context, FSamplerStateRHIParamRef Filter = TStaticSamplerState<>::GetRHI(), bool bWhiteIfNoTexture = false, FSamplerStateRHIParamRef* FilterOverrideArray = 0);
+	void SetCS(const FComputeShaderRHIParamRef& ShaderRHI, const FRenderingCompositePassContext& Context, FSamplerStateRHIParamRef Filter = TStaticSamplerState<>::GetRHI(), EFallbackColor FallbackColor = eFC_0000, FSamplerStateRHIParamRef* FilterOverrideArray = 0);
 
 	/** Set the vertex shader parameter values. */
-	void SetVS(const FVertexShaderRHIParamRef& ShaderRHI, const FRenderingCompositePassContext& Context, FSamplerStateRHIParamRef Filter = TStaticSamplerState<>::GetRHI(), bool bWhiteIfNoTexture = false, FSamplerStateRHIParamRef* FilterOverrideArray = 0);
+	void SetVS(const FVertexShaderRHIParamRef& ShaderRHI, const FRenderingCompositePassContext& Context, FSamplerStateRHIParamRef Filter = TStaticSamplerState<>::GetRHI(), EFallbackColor FallbackColor = eFC_0000, FSamplerStateRHIParamRef* FilterOverrideArray = 0);
 
 	/** Serializer. */
 	friend FArchive& operator<<(FArchive& Ar,FPostProcessPassParameters& P);
@@ -92,7 +102,7 @@ public:
 		const ShaderRHIParamRef& ShaderRHI,
 		const FRenderingCompositePassContext& Context,
 		FSamplerStateRHIParamRef Filter,
-		bool bWhiteIfNoTexture = false,
+		EFallbackColor FallbackColor,
 		FSamplerStateRHIParamRef* FilterOverrideArray = 0
 	);
 };
\ No newline at end of file
