UE4 - probably fix for UE-20006 crash report crash

diff --git a/Engine/Source/Runtime/Renderer/Private/SceneOcclusion.cpp b/Engine/Source/Runtime/Renderer/Private/SceneOcclusion.cpp
--- a/Engine/Source/Runtime/Renderer/Private/SceneOcclusion.cpp
+++ b/Engine/Source/Runtime/Renderer/Private/SceneOcclusion.cpp
@@ -1074,7 +1074,7 @@ void FDeferredShadingSceneRenderer::BeginOcclusionTests(FRHICommandListImmediate
 	SCOPE_CYCLE_COUNTER(STAT_BeginOcclusionTestsTime);
 	int32 NumBufferedFrames = FOcclusionQueryHelpers::GetNumBufferedFrames();
 	FSceneRenderTargets& SceneContext = FSceneRenderTargets::Get(RHICmdList);
-	const bool bUseDownsampledDepth = IsValidRef(SceneContext.GetSmallDepthSurface()) && SceneContext.UseDownsizedOcclusionQueries();	
+	const bool bUseDownsampledDepth = SceneContext.UseDownsizedOcclusionQueries() && IsValidRef(SceneContext.SmallDepthZ) && IsValidRef(SceneContext.GetSmallDepthSurface());	
 
 	if (bUseDownsampledDepth)
 	{
