Added pre-tonemap 'HDR Color' to visualization modes and screenshot buffer dump. UE-19262

diff --git a/Engine/Config/BaseEngine.ini b/Engine/Config/BaseEngine.ini
--- a/Engine/Config/BaseEngine.ini
+++ b/Engine/Config/BaseEngine.ini
@@ -2217,6 +2217,7 @@ WorldNormal=(Material="/Engine/BufferVisualization/WorldNormal.WorldNormal", Nam
 AmbientOcclusion=(Material="/Engine/BufferVisualization/AmbientOcclusion.AmbientOcclusion", Name=LOCTEXT("BaseAmbientOcclusionMat", "Ambient Occlusion"))
 CustomDepthWorldUnits=(Material="/Engine/BufferVisualization/CustomDepthWorldUnits.CustomDepthWorldUnits", Name=LOCTEXT("BaseCustomDepthWorldUnitsMat", "Custom Depth World Units"))
 SceneDepthWorldUnits=(Material="/Engine/BufferVisualization/SceneDepthWorldUnits.SceneDepthWorldUnits", Name=LOCTEXT("BaseSceneDepthWorldUnitsMat", "Scene Depth World Units"))
+HDRColor=(Material="/Engine/BufferVisualization/HDRColor.HDRColor", Name=LOCTEXT("HDRColor", "Pre Tonemap HDR Color"))
 
 [DeviceProfileManager]
 DeviceProfileSelectionModule="ExampleDeviceProfileSelector"
diff --git a/Engine/Source/Runtime/Engine/Private/BufferVisualizationData.cpp b/Engine/Source/Runtime/Engine/Private/BufferVisualizationData.cpp
--- a/Engine/Source/Runtime/Engine/Private/BufferVisualizationData.cpp
+++ b/Engine/Source/Runtime/Engine/Private/BufferVisualizationData.cpp
@@ -95,7 +95,7 @@ void FBufferVisualizationData::ConfigureConsoleCommand()
 	ConsoleDocumentationOverviewTargets += AvailableVisualizationMaterials;
 
 	IConsoleManager::Get().RegisterConsoleVariable(TEXT("r.BufferVisualizationOverviewTargets"),
-		TEXT("BaseColor,Specular,SubsurfaceColor,WorldNormal,SeparateTranslucencyRGB,,,Opacity,SeparateTranslucencyA,,,DecalMask,SceneDepth,Roughness,Metallic,ShadingModel,,SceneDepthWorldUnits,SceneColor"), 
+		TEXT("BaseColor,Specular,SubsurfaceColor,WorldNormal,SeparateTranslucencyRGB,,,Opacity,SeparateTranslucencyA,,,DecalMask,SceneDepth,Roughness,Metallic,ShadingModel,,SceneDepthWorldUnits,SceneColor,HDRColor"), 
 		*ConsoleDocumentationOverviewTargets,
 		ECVF_Default);
 }
diff --git a/Engine/Source/Runtime/Engine/Private/SceneView.cpp b/Engine/Source/Runtime/Engine/Private/SceneView.cpp
--- a/Engine/Source/Runtime/Engine/Private/SceneView.cpp
+++ b/Engine/Source/Runtime/Engine/Private/SceneView.cpp
@@ -1551,9 +1551,6 @@ void FSceneView::ConfigureBufferVisualizationSettings()
 			BufferVisualizationData.SetCurrentOverviewMaterialNames(SelectedMaterialNames);
 			BufferVisualizationData.GetOverviewMaterials().Empty();
 
-			// Note - This will re-parse the list of names from the console variable every frame. It could be cached and only updated when
-			// the variable value changes if this turns out to be a performance issue.
-		
 			// Extract each material name from the comma separated string
 			while (SelectedMaterialNames.Len())
 			{
diff --git a/Engine/Source/Runtime/Renderer/Private/PostProcess/PostProcessing.cpp b/Engine/Source/Runtime/Renderer/Private/PostProcess/PostProcessing.cpp
--- a/Engine/Source/Runtime/Renderer/Private/PostProcess/PostProcessing.cpp
+++ b/Engine/Source/Runtime/Renderer/Private/PostProcess/PostProcessing.cpp
@@ -724,7 +724,7 @@ static FRenderingCompositePass* AddSinglePostProcessMaterial(FPostprocessContext
 	return 0;
 }
 
-static void AddPostProcessMaterial(FPostprocessContext& Context, EBlendableLocation InLocation, FRenderingCompositeOutputRef SeparateTranslucency)
+static void AddPostProcessMaterial(FPostprocessContext& Context, EBlendableLocation InLocation, FRenderingCompositeOutputRef SeparateTranslucency, FRenderingCompositeOutputRef HDRColor = FRenderingCompositeOutputRef())
 {
 	if(!Context.View.Family->EngineShowFlags.PostProcessing || !Context.View.Family->EngineShowFlags.PostProcessMaterial)
 	{
@@ -736,6 +736,7 @@ static void AddPostProcessMaterial(FPostprocessContext& Context, EBlendableLocat
 	FBlendableEntry* Iterator = 0;
 	FPostProcessMaterialNode PPNodes[MAX_PPMATERIALNODES];
 	uint32 PPNodeCount = 0;
+	bool bVisualizingBuffer = false;
 
 	if(Context.View.Family->EngineShowFlags.VisualizeBuffer)
 	{	
@@ -746,6 +747,7 @@ static void AddPostProcessMaterial(FPostprocessContext& Context, EBlendableLocat
 		{
 			PPNodes[0] = FPostProcessMaterialNode(Material, InLocation, Material->BlendablePriority);
 			++PPNodeCount;
+			bVisualizingBuffer = true;
 		}
 	}
 	for(;PPNodeCount < MAX_PPMATERIALNODES; ++PPNodeCount)
@@ -791,6 +793,13 @@ static void AddPostProcessMaterial(FPostprocessContext& Context, EBlendableLocat
 		// the separate translucency buffers through ePId_Input1. 
 		// TODO: Check if material actually uses this texture and only bind if needed.
 		Node->SetInput(ePId_Input1, SeparateTranslucency);
+
+		// This input is only needed for visualization and frame dumping
+		if (bVisualizingBuffer)
+		{
+			Node->SetInput(ePId_Input2, HDRColor);
+		}
+
 		Context.FinalOutput = FRenderingCompositeOutputRef(Node);
 	}
 }
@@ -839,7 +848,7 @@ static void AddHighResScreenshotMask(FPostprocessContext& Context, FRenderingCom
 	}
 }
 
-static void AddGBufferVisualizationOverview(FPostprocessContext& Context, FRenderingCompositeOutputRef& SeparateTranslucencyInput)
+static void AddGBufferVisualizationOverview(FPostprocessContext& Context, FRenderingCompositeOutputRef& SeparateTranslucencyInput, FRenderingCompositeOutputRef& HDRColorInput)
 {
 	static const auto CVarDumpFrames = IConsoleManager::Get().FindTConsoleVariableDataInt(TEXT("r.BufferVisualizationDumpFrames"));
 	static const auto CVarDumpFramesAsHDR = IConsoleManager::Get().FindTConsoleVariableDataInt(TEXT("r.BufferVisualizationDumpFramesAsHDR"));
@@ -877,6 +886,7 @@ static void AddGBufferVisualizationOverview(FPostprocessContext& Context, FRende
 					FRenderingCompositePass* MaterialPass = Context.Graph.RegisterPass(new(FMemStack::Get()) FRCPassPostProcessMaterial(*It, Context.View.GetFeatureLevel(), OutputFormat));
 					MaterialPass->SetInput(ePId_Input0, FRenderingCompositeOutputRef(IncomingStage));
 					MaterialPass->SetInput(ePId_Input1, FRenderingCompositeOutputRef(SeparateTranslucencyInput));
+					MaterialPass->SetInput(ePId_Input2, FRenderingCompositeOutputRef(HDRColorInput));
 
 					auto Proxy = MaterialInterface->GetRenderProxy(false);
 					const FMaterial* Material = Proxy->GetMaterial(Context.View.GetFeatureLevel());
@@ -1543,7 +1553,7 @@ void FPostProcessing::Process(FRHICommandListImmediate& RHICmdList, FViewInfo& V
 			Context.FinalOutput = FRenderingCompositeOutputRef(Node);
 		}
 		
-		AddPostProcessMaterial(Context, BL_AfterTonemapping, SeparateTranslucency);
+		AddPostProcessMaterial(Context, BL_AfterTonemapping, SeparateTranslucency, HDRColor);
 
 		if(bVisualizeBloom)
 		{
@@ -1558,7 +1568,7 @@ void FPostProcessing::Process(FRHICommandListImmediate& RHICmdList, FViewInfo& V
 			Context.FinalOutput = FRenderingCompositeOutputRef(PassVisualize);
 		}
 
-		AddGBufferVisualizationOverview(Context, SeparateTranslucency);
+		AddGBufferVisualizationOverview(Context, SeparateTranslucency, HDRColor);
 
 		bool bStereoRenderingAndHMD = View.Family->EngineShowFlags.StereoRendering && View.Family->EngineShowFlags.HMDDistortion;
 		bool bHMDWantsUpscale = false;
