UE4 - Vertex Density viewmode

diff --git a/Engine/Shaders/VertexDensityShader.usf b/Engine/Shaders/VertexDensityShader.usf
new file mode 100644
--- /dev/null
+++ b/Engine/Shaders/VertexDensityShader.usf
@@ -0,0 +1,84 @@
+// Copyright 1998-2015 Epic Games, Inc. All Rights Reserved.
+
+/*=============================================================================
+	LightMapDensityShader.hlsl: Shader for rendering lightmap density as a color
+=============================================================================*/
+
+#include "Common.usf"
+#include "Material.usf"
+#include "VertexFactory.usf"
+
+struct FVertexDensityVSToPS
+{
+	FVertexFactoryInterpolantsVSToPS FactoryInterpolants;
+	float4 WorldPosition	: TEXCOORD6;
+	float4 Position			: SV_POSITION;
+};
+
+void MainVS(FVertexFactoryInput Input, out FVertexDensityVSToPS Output)
+{
+	FVertexFactoryIntermediates VFIntermediates = GetVertexFactoryIntermediates(Input);
+	float4 WorldPosition = VertexFactoryGetWorldPosition(Input, VFIntermediates);
+	float3x3 TangentToLocal = VertexFactoryGetTangentToLocal(Input, VFIntermediates);
+	FMaterialVertexParameters VertexParameters = GetMaterialVertexParameters(Input, VFIntermediates, WorldPosition.xyz, TangentToLocal);
+	Output.FactoryInterpolants = VertexFactoryGetInterpolantsVSToPS(Input, VFIntermediates, VertexParameters);
+
+	Output.WorldPosition = WorldPosition;
+	ISOLATE
+	{
+		float4 RasterizedWorldPosition = VertexFactoryGetRasterizedWorldPosition(Input, VFIntermediates, Output.WorldPosition);
+		Output.Position = mul(RasterizedWorldPosition, View.TranslatedWorldToClip);
+	}
+	Output.FactoryInterpolants = VertexFactoryGetInterpolantsVSToPS(Input, VFIntermediates, VertexParameters);
+}
+
+float2 TriangleSize;
+
+[maxvertexcount(12)]
+void MainGS(triangle FVertexDensityVSToPS Input[3], inout TriangleStream<FVertexDensityVSToPS> TriStream )
+{
+	// Size in pixels
+	float2 SizeInClip = TriangleSize * View.ViewSizeAndInvSize.zw;
+
+	int i;
+
+	// Regular triangle
+	for (i = 0; i < 3; ++i)
+	{
+		FVertexDensityVSToPS T = Input[i];
+		T.WorldPosition.w = 1;
+		TriStream.Append(T);
+	}
+	TriStream.RestartStrip();
+
+	// Smaller triangle
+	for (i = 0; i < 3; ++i)
+	{
+		FVertexDensityVSToPS T = Input[i];
+		T.WorldPosition.w = 0;
+		TriStream.Append(T);
+		T.Position.x += SizeInClip.x * T.Position.w;
+		TriStream.Append(T);
+		T.Position.y += SizeInClip.y * T.Position.w;
+		TriStream.Append(T);
+
+		TriStream.RestartStrip();
+	}
+}
+
+float4 MainPS(FVertexDensityVSToPS Input) : SV_Target0
+{
+	if (Input.WorldPosition.w < 1e-6)
+	{
+		return float4(normalize(Input.WorldPosition).xyz * 0.5 + float3(0, 0.25, 0), 0.5);
+	}
+	else
+	{
+		FMaterialPixelParameters MaterialParameters = GetMaterialPixelParameters(Input.FactoryInterpolants, Input.WorldPosition);
+		CalcMaterialParameters(MaterialParameters, Input.Position, false, float4(0, 0, .00001f, 1));
+		GetMaterialCoverageAndClipping(MaterialParameters);
+		float3 Color = MaterialParameters.AOMaterialMask + MaterialParameters.VertexColor.xyz;
+		Color = normalize(Color) * 0.5 + float3(0, 0, 0.25);
+		return float4(Color, 0.25);
+	}
+}
diff --git a/Engine/Source/Editor/UnrealEd/Public/EditorViewportCommands.h b/Engine/Source/Editor/UnrealEd/Public/EditorViewportCommands.h
--- a/Engine/Source/Editor/UnrealEd/Public/EditorViewportCommands.h
+++ b/Engine/Source/Editor/UnrealEd/Public/EditorViewportCommands.h
@@ -71,6 +71,9 @@ public:
 	/** Changes the viewport to lightmap density mode */
 	TSharedPtr< FUICommandInfo > LightmapDensityMode;
 
+	/** Changes the viewport to vertex densities mode */
+	TSharedPtr< FUICommandInfo > VertexDensitiesMode;
+
 	/** Changes the viewport to visualize the buffer content */
 	TSharedPtr< FUICommandInfo > VisualizeBufferMode;
 
diff --git a/Engine/Source/Runtime/Engine/Classes/Engine/EngineBaseTypes.h b/Engine/Source/Runtime/Engine/Classes/Engine/EngineBaseTypes.h
--- a/Engine/Source/Runtime/Engine/Classes/Engine/EngineBaseTypes.h
+++ b/Engine/Source/Runtime/Engine/Classes/Engine/EngineBaseTypes.h
@@ -732,10 +732,10 @@ enum EViewModeIndex
 
 	/** Colored according to stationary light overlap. */
 	VMI_StationaryLightOverlap = 14,
-//	VMI_VertexDensity = 15,
 
 	VMI_CollisionPawn = 15, 
 	VMI_CollisionVisibility = 16, 
+	VMI_VertexDensities = 17,
 	VMI_Max UMETA(Hidden),
 
 	VMI_Unknown = 255 UMETA(Hidden),
diff --git a/Engine/Source/Runtime/Engine/Private/PrimitiveDrawingUtils.cpp b/Engine/Source/Runtime/Engine/Private/PrimitiveDrawingUtils.cpp
--- a/Engine/Source/Runtime/Engine/Private/PrimitiveDrawingUtils.cpp
+++ b/Engine/Source/Runtime/Engine/Private/PrimitiveDrawingUtils.cpp
@@ -1265,6 +1265,7 @@ bool IsRichView(const FSceneViewFamily& ViewFamily)
 		ViewFamily.EngineShowFlags.StationaryLightOverlap ||
 		ViewFamily.EngineShowFlags.BSPSplit ||
 		ViewFamily.EngineShowFlags.LightMapDensity ||
+		ViewFamily.EngineShowFlags.VertexDensities ||
 		ViewFamily.EngineShowFlags.PropertyColoration ||
 		ViewFamily.EngineShowFlags.MeshEdges ||
 		ViewFamily.EngineShowFlags.LightInfluences ||
diff --git a/Engine/Source/Runtime/Engine/Private/ShowFlags.cpp b/Engine/Source/Runtime/Engine/Private/ShowFlags.cpp
--- a/Engine/Source/Runtime/Engine/Private/ShowFlags.cpp
+++ b/Engine/Source/Runtime/Engine/Private/ShowFlags.cpp
@@ -276,6 +276,9 @@ void ApplyViewMode(EViewModeIndex ViewModeIndex, bool bPerspective, FEngineShowF
 		case VMI_CollisionVisibility:
 			bPostProcessing = false;
 			break;
+		case VMI_VertexDensities:
+			bPostProcessing = false;
+			break;
 	}
 
 	if(!bPerspective)
@@ -301,6 +304,7 @@ void ApplyViewMode(EViewModeIndex ViewModeIndex, bool bPerspective, FEngineShowF
 	EngineShowFlags.SetWireframe(ViewModeIndex == VMI_Wireframe || ViewModeIndex == VMI_BrushWireframe);
 	EngineShowFlags.SetCollisionPawn(ViewModeIndex == VMI_CollisionPawn);
 	EngineShowFlags.SetCollisionVisibility(ViewModeIndex == VMI_CollisionVisibility);
+	EngineShowFlags.SetVertexDensities(ViewModeIndex == VMI_VertexDensities);
 }
 
 void EngineShowFlagOverride(EShowFlagInitMode ShowFlagInitMode, EViewModeIndex ViewModeIndex, FEngineShowFlags& EngineShowFlags, FName CurrentBufferVisualizationMode, bool bIsSplitScreen)
@@ -375,6 +379,7 @@ void EngineShowFlagOverride(EShowFlagInitMode ShowFlagInitMode, EViewModeIndex V
 			ViewModeIndex == VMI_Wireframe ||
 			ViewModeIndex == VMI_Unlit ||
 			ViewModeIndex == VMI_LightmapDensity ||
+			ViewModeIndex == VMI_VertexDensities ||
 			ViewModeIndex == VMI_LitLightmapDensity)
 		{
 			EngineShowFlags.LightFunctions = 0;
@@ -385,6 +390,7 @@ void EngineShowFlagOverride(EShowFlagInitMode ShowFlagInitMode, EViewModeIndex V
 			ViewModeIndex == VMI_Unlit ||
 			ViewModeIndex == VMI_ShaderComplexity ||
 			ViewModeIndex == VMI_LightmapDensity ||
+			ViewModeIndex == VMI_VertexDensities ||
 			ViewModeIndex == VMI_LitLightmapDensity)
 		{
 			EngineShowFlags.DynamicShadows = 0;
@@ -400,6 +406,7 @@ void EngineShowFlagOverride(EShowFlagInitMode ShowFlagInitMode, EViewModeIndex V
 			ViewModeIndex == VMI_Unlit ||
 			ViewModeIndex == VMI_StationaryLightOverlap ||
 			ViewModeIndex == VMI_ShaderComplexity ||
+			ViewModeIndex == VMI_VertexDensities ||
 			ViewModeIndex == VMI_LightmapDensity)
 		{
 			EngineShowFlags.Lighting = 0;
@@ -558,6 +565,10 @@ EViewModeIndex FindViewMode(const FEngineShowFlags& EngineShowFlags)
 	{
 		return VMI_CollisionVisibility;
 	}
+	else if (EngineShowFlags.VertexDensities)
+	{
+		return VMI_VertexDensities;
+	}
 
 	return EngineShowFlags.Lighting ? VMI_Lit : VMI_Unlit;
 }
@@ -582,6 +593,7 @@ const TCHAR* GetViewModeName(EViewModeIndex ViewModeIndex)
 		case VMI_VisualizeBuffer:			return TEXT("VisualizeBuffer");
 		case VMI_CollisionPawn:				return TEXT("CollisionPawn");
 		case VMI_CollisionVisibility:		return TEXT("CollisionVis");
+		case VMI_VertexDensities:			return TEXT("VertexDensity");
 	}
 	return TEXT("");
 }
diff --git a/Engine/Source/Runtime/Engine/Public/ShowFlags.h b/Engine/Source/Runtime/Engine/Public/ShowFlags.h
--- a/Engine/Source/Runtime/Engine/Public/ShowFlags.h
+++ b/Engine/Source/Runtime/Engine/Public/ShowFlags.h
@@ -291,6 +291,7 @@ private:
 		SetShaderComplexity(false);
 		SetStationaryLightOverlap(false);
 		SetLightMapDensity(false);
+		SetVertexDensities(false);
 		SetVisualizeLPV(false);
 		SetStreamingBounds(false);
 		SetConstraints(false);
diff --git a/Engine/Source/Runtime/Engine/Public/ShowFlagsValues.inl b/Engine/Source/Runtime/Engine/Public/ShowFlagsValues.inl
--- a/Engine/Source/Runtime/Engine/Public/ShowFlagsValues.inl
+++ b/Engine/Source/Runtime/Engine/Public/ShowFlagsValues.inl
@@ -288,6 +288,8 @@ SHOWFLAG_ALWAYS_ACCESSIBLE(ForceGBuffer, SFG_Hidden, LOCTEXT("ForceGBuffer", "Fo
 SHOWFLAG_ALWAYS_ACCESSIBLE(VisualizeSenses, SFG_Advanced, LOCTEXT("VisualizeSenses", "Senses"))
 /** Visualize the bloom, for developer (by default off): */
 SHOWFLAG_FIXED_IN_SHIPPING(VisualizeBloom, 0, SFG_Visualize, LOCTEXT("VisualizeBloom", "Bloom"))
+/** Visualize Vertex Density */
+SHOWFLAG_ALWAYS_ACCESSIBLE(VertexDensities, SFG_Hidden, LOCTEXT("VisualizeVertexDensities", "Visualize Vertex Densities"))
 
 #undef SHOWFLAG_ALWAYS_ACCESSIBLE
 #undef SHOWFLAG_FIXED_IN_SHIPPING
diff --git a/Engine/Source/Runtime/Renderer/Private/DeferredShadingRenderer.cpp b/Engine/Source/Runtime/Renderer/Private/DeferredShadingRenderer.cpp
--- a/Engine/Source/Runtime/Renderer/Private/DeferredShadingRenderer.cpp
+++ b/Engine/Source/Runtime/Renderer/Private/DeferredShadingRenderer.cpp
@@ -1058,20 +1058,20 @@ void FDeferredShadingSceneRenderer::Render(FRHICommandListImmediate& RHICmdList)
 	ServiceLocalQueue();
 
 	if (ViewFamily.EngineShowFlags.VisualizeLightCulling)
-	  {
-		  // clear out emissive and baked lighting (not too efficient but simple and only needed for this debug view)
-		  SceneContext.BeginRenderingSceneColor(RHICmdList);
-		  RHICmdList.Clear(true, FLinearColor(0, 0, 0, 0), false, (float)ERHIZBuffer::FarPlane, false, 0, FIntRect());
-	  }
+	{
+		// clear out emissive and baked lighting (not too efficient but simple and only needed for this debug view)
+		SceneContext.BeginRenderingSceneColor(RHICmdList);
+		RHICmdList.Clear(true, FLinearColor(0, 0, 0, 0), false, (float)ERHIZBuffer::FarPlane, false, 0, FIntRect());
+	}
 
-	  SceneContext.DBufferA.SafeRelease();
-	  SceneContext.DBufferB.SafeRelease();
-	  SceneContext.DBufferC.SafeRelease();
+	SceneContext.DBufferA.SafeRelease();
+	SceneContext.DBufferB.SafeRelease();
+	SceneContext.DBufferC.SafeRelease();
 
-	  // only temporarily available after early z pass and until base pass
-	  check(!SceneContext.DBufferA);
-	  check(!SceneContext.DBufferB);
-	  check(!SceneContext.DBufferC);
+	// only temporarily available after early z pass and until base pass
+	check(!SceneContext.DBufferA);
+	check(!SceneContext.DBufferB);
+	check(!SceneContext.DBufferC);
 
 	if (bRequiresFarZQuadClear)
 	{
@@ -1683,11 +1683,16 @@ bool FDeferredShadingSceneRenderer::RenderBasePass(FRHICommandListImmediate& RHI
 		GPrevPerBoneMotionBlur.StartAppend(RHICmdList, ViewFamily.bWorldIsPaused);
 	}
 
-	if(ViewFamily.EngineShowFlags.LightMapDensity && AllowDebugViewmodes())
+	if (ViewFamily.EngineShowFlags.LightMapDensity && AllowDebugViewmodes())
 	{
 		// Override the base pass with the lightmap density pass if the viewmode is enabled.
 		bDirty = RenderLightMapDensities(RHICmdList);
 	}
+	else if (ViewFamily.EngineShowFlags.VertexDensities && AllowDebugViewmodes())
+	{
+		// Override the base pass with the lightmap density pass if the viewmode is enabled.
+		bDirty = RenderVertexDensities(RHICmdList);
+	}
 	else
 	{
 		SCOPED_DRAW_EVENT(RHICmdList, BasePass);
diff --git a/Engine/Source/Runtime/Renderer/Private/DeferredShadingRenderer.h b/Engine/Source/Runtime/Renderer/Private/DeferredShadingRenderer.h
--- a/Engine/Source/Runtime/Renderer/Private/DeferredShadingRenderer.h
+++ b/Engine/Source/Runtime/Renderer/Private/DeferredShadingRenderer.h
@@ -259,6 +259,9 @@ private:
 	/** Renders world-space lightmap density instead of the normal color. */
 	bool RenderLightMapDensities(FRHICommandListImmediate& RHICmdList);
 
+	/** Renders the visualize vertex densities mode. */
+	bool RenderVertexDensities(FRHICommandListImmediate& RHICmdList);
+
 	/** Updates the downsized depth buffer with the current full resolution depth buffer. */
 	void UpdateDownsampledDepthSurface(FRHICommandList& RHICmdList);
 
diff --git a/Engine/Source/Runtime/Renderer/Private/FogRendering.cpp b/Engine/Source/Runtime/Renderer/Private/FogRendering.cpp
--- a/Engine/Source/Runtime/Renderer/Private/FogRendering.cpp
+++ b/Engine/Source/Runtime/Renderer/Private/FogRendering.cpp
@@ -367,5 +367,6 @@ bool ShouldRenderFog(const FSceneViewFamily& Family)
 		&& EngineShowFlags.Materials 
 		&& !EngineShowFlags.ShaderComplexity
 		&& !EngineShowFlags.StationaryLightOverlap 
+		&& !EngineShowFlags.VertexDensities
 		&& !EngineShowFlags.LightMapDensity;
 }
diff --git a/Engine/Source/Runtime/Renderer/Private/VertexDensityRendering.cpp b/Engine/Source/Runtime/Renderer/Private/VertexDensityRendering.cpp
new file mode 100644
--- /dev/null
+++ b/Engine/Source/Runtime/Renderer/Private/VertexDensityRendering.cpp
@@ -0,0 +1,379 @@
+// Copyright 1998-2015 Epic Games, Inc. All Rights Reserved.
+
+/*=============================================================================
+	LightMapDensityRendering.cpp: Implementation for rendering lightmap density.
+=============================================================================*/
+
+#include "RendererPrivate.h"
+#include "ScenePrivate.h"
+#include "SceneUtils.h"
+
+static int32 GVertexDensityTriangleSize = 16;
+static FAutoConsoleVariableRef CVarVertexDensitySize(
+	TEXT("r.VertexDensity.Size"),
+	GVertexDensityTriangleSize,
+	TEXT("Size of the density indicator.")
+	);
+
+class FVertexDensityVS : public FMeshMaterialShader
+{
+	DECLARE_SHADER_TYPE(FVertexDensityVS, MeshMaterial);
+
+public:
+	static bool ShouldCache(EShaderPlatform Platform, const FMaterial* Material, const FVertexFactoryType* VertexFactoryType)
+	{
+		return (Material->IsSpecialEngineMaterial() || Material->IsMasked() || Material->MaterialMayModifyMeshPosition())
+			&& IsFeatureLevelSupported(Platform, ERHIFeatureLevel::SM4);
+	}
+
+	static void ModifyCompilationEnvironment(EShaderPlatform Platform, const FMaterial* Material, FShaderCompilerEnvironment& OutEnvironment)
+	{
+		FMeshMaterialShader::ModifyCompilationEnvironment(Platform, Material, OutEnvironment);
+	}
+
+	FVertexDensityVS(const FMeshMaterialShaderType::CompiledShaderInitializerType& Initializer) :
+		FMeshMaterialShader(Initializer)
+	{
+		//LightMapPolicyType::VertexParametersType::Bind(Initializer.ParameterMap);
+	}
+	FVertexDensityVS() {}
+
+	virtual bool Serialize(FArchive& Ar) override
+	{
+		bool bShaderHasOutdatedParameters = FMeshMaterialShader::Serialize(Ar);
+		//LightMapPolicyType::VertexParametersType::Serialize(Ar);
+		return bShaderHasOutdatedParameters;
+	}
+
+	void SetParameters(
+		FRHICommandList& RHICmdList,
+		const FMaterialRenderProxy* MaterialRenderProxy,
+		const FSceneView& View
+		)
+	{
+		FMeshMaterialShader::SetParameters(RHICmdList, GetVertexShader(), MaterialRenderProxy, *MaterialRenderProxy->GetMaterial(View.GetFeatureLevel()), View, ESceneRenderTargetsMode::SetTextures);
+	}
+
+	void SetMesh(FRHICommandList& RHICmdList, const FVertexFactory* VertexFactory, const FSceneView& View, const FPrimitiveSceneProxy* Proxy, const FMeshBatchElement& BatchElement, float DitheredLODTransitionValue)
+	{
+		FMeshMaterialShader::SetMesh(RHICmdList, GetVertexShader(), VertexFactory, View, Proxy, BatchElement, DitheredLODTransitionValue);
+	}
+};
+
+class FVertexDensityGS : public FMeshMaterialShader
+{
+	DECLARE_SHADER_TYPE(FVertexDensityGS, MeshMaterial);
+
+public:
+	static bool ShouldCache(EShaderPlatform Platform, const FMaterial* Material, const FVertexFactoryType* VertexFactoryType)
+	{
+		return (Material->IsSpecialEngineMaterial() || Material->IsMasked() || Material->MaterialMayModifyMeshPosition())
+			&& IsFeatureLevelSupported(Platform, ERHIFeatureLevel::SM4);
+	}
+
+	static void ModifyCompilationEnvironment(EShaderPlatform Platform, const FMaterial* Material, FShaderCompilerEnvironment& OutEnvironment)
+	{
+		FMeshMaterialShader::ModifyCompilationEnvironment(Platform, Material, OutEnvironment);
+	}
+
+	FVertexDensityGS(const FMeshMaterialShaderType::CompiledShaderInitializerType& Initializer) :
+		FMeshMaterialShader(Initializer)
+	{
+		TriangleSize.Bind(Initializer.ParameterMap, TEXT("TriangleSize"), SPF_Mandatory);
+	}
+	FVertexDensityGS() {}
+
+	virtual bool Serialize(FArchive& Ar) override
+	{
+		bool bShaderHasOutdatedParameters = FMeshMaterialShader::Serialize(Ar);
+		Ar << TriangleSize;
+		return bShaderHasOutdatedParameters;
+	}
+
+	void SetParameters(
+		FRHICommandList& RHICmdList,
+		const FMaterialRenderProxy* MaterialRenderProxy,
+		const FSceneView& View
+		)
+	{
+		FMeshMaterialShader::SetParameters(RHICmdList, GetGeometryShader(), MaterialRenderProxy, *MaterialRenderProxy->GetMaterial(View.GetFeatureLevel()), View, ESceneRenderTargetsMode::SetTextures);
+		FVector4 TriangleSizeValue(GVertexDensityTriangleSize, GVertexDensityTriangleSize, 0, 0);
+		SetShaderValue(RHICmdList, GetGeometryShader(), TriangleSize, TriangleSizeValue);
+	}
+
+	void SetMesh(FRHICommandList& RHICmdList, const FVertexFactory* VertexFactory, const FSceneView& View, const FPrimitiveSceneProxy* Proxy, const FMeshBatchElement& BatchElement, float DitheredLODTransitionValue)
+	{
+		FMeshMaterialShader::SetMesh(RHICmdList, GetGeometryShader(), VertexFactory, View, Proxy, BatchElement, DitheredLODTransitionValue);
+	}
+
+	FShaderParameter TriangleSize;
+};
+
+class FVertexDensityPS : public FMeshMaterialShader
+{
+	DECLARE_SHADER_TYPE(FVertexDensityPS, MeshMaterial);
+
+public:
+	static bool ShouldCache(EShaderPlatform Platform, const FMaterial* Material, const FVertexFactoryType* VertexFactoryType)
+	{
+		return (Material->IsSpecialEngineMaterial() || Material->IsMasked() || Material->MaterialMayModifyMeshPosition())
+			&& IsFeatureLevelSupported(Platform, ERHIFeatureLevel::SM4);
+	}
+
+	static void ModifyCompilationEnvironment(EShaderPlatform Platform, const FMaterial* Material, FShaderCompilerEnvironment& OutEnvironment)
+	{
+		FMeshMaterialShader::ModifyCompilationEnvironment(Platform, Material, OutEnvironment);
+	}
+
+	FVertexDensityPS(const FMeshMaterialShaderType::CompiledShaderInitializerType& Initializer) :
+		FMeshMaterialShader(Initializer)
+	{
+		//LightMapPolicyType::VertexParametersType::Bind(Initializer.ParameterMap);
+	}
+	FVertexDensityPS() {}
+
+	virtual bool Serialize(FArchive& Ar) override
+	{
+		bool bShaderHasOutdatedParameters = FMeshMaterialShader::Serialize(Ar);
+		//LightMapPolicyType::VertexParametersType::Serialize(Ar);
+		return bShaderHasOutdatedParameters;
+	}
+
+	void SetParameters(
+		FRHICommandList& RHICmdList,
+		const FMaterialRenderProxy* MaterialRenderProxy,
+		const FSceneView& View
+		)
+	{
+		FMeshMaterialShader::SetParameters(RHICmdList, GetPixelShader(), MaterialRenderProxy, *MaterialRenderProxy->GetMaterial(View.GetFeatureLevel()), View, ESceneRenderTargetsMode::SetTextures);
+	}
+
+	void SetMesh(FRHICommandList& RHICmdList, const FVertexFactory* VertexFactory, const FSceneView& View, const FPrimitiveSceneProxy* Proxy, const FMeshBatchElement& BatchElement, float DitheredLODTransitionValue)
+	{
+		FMeshMaterialShader::SetMesh(RHICmdList, GetPixelShader(), VertexFactory, View, Proxy, BatchElement, DitheredLODTransitionValue);
+	}
+};
+
+class FVertexDensityDrawingPolicy : public FMeshDrawingPolicy
+{
+public:
+	/** The data the drawing policy uses for each mesh element. */
+	class ElementDataType
+	{
+	public:
+	};
+
+	FVertexDensityDrawingPolicy(
+		const FViewInfo& View,
+		const FVertexFactory* InVertexFactory,
+		const FMaterialRenderProxy* InMaterialRenderProxy,
+		EBlendMode InBlendMode
+		) :
+		FMeshDrawingPolicy(InVertexFactory, InMaterialRenderProxy, *InMaterialRenderProxy->GetMaterial(View.GetFeatureLevel())),
+		BlendMode(InBlendMode)
+	{
+		HullShader = nullptr;
+		DomainShader = nullptr;
+
+		const EMaterialTessellationMode MaterialTessellationMode = MaterialResource->GetTessellationMode();
+		if (RHISupportsTessellation(View.GetShaderPlatform())
+			&& InVertexFactory->GetType()->SupportsTessellationShaders()
+			&& MaterialTessellationMode != MTM_NoTessellation)
+		{
+			//HullShader = MaterialResource->GetShader<TLightMapDensityHS<LightMapPolicyType> >(VertexFactory->GetType());
+			//DomainShader = MaterialResource->GetShader<TLightMapDensityDS<LightMapPolicyType> >(VertexFactory->GetType());
+		}
+
+		VertexShader = MaterialResource->GetShader<FVertexDensityVS>(InVertexFactory->GetType());
+		GeometryShader = MaterialResource->GetShader<FVertexDensityGS>(InVertexFactory->GetType());
+		PixelShader = MaterialResource->GetShader<FVertexDensityPS>(InVertexFactory->GetType());
+	}
+
+	// FMeshDrawingPolicy interface.
+	bool Matches(const FVertexDensityDrawingPolicy& Other) const
+	{
+		return FMeshDrawingPolicy::Matches(Other) &&
+			VertexShader == Other.VertexShader &&
+			GeometryShader == Other.GeometryShader &&
+			PixelShader == Other.PixelShader &&
+			HullShader == Other.HullShader &&
+			DomainShader == Other.DomainShader;
+	}
+
+	void SetSharedState(FRHICommandList& RHICmdList, const FSceneView* View, const ContextDataType) const
+	{
+		if (VertexFactory)
+		{
+			VertexFactory->Set(RHICmdList);
+		}
+
+		// Set the base pass shader parameters for the material.
+		VertexShader->SetParameters(RHICmdList, MaterialRenderProxy,*View);
+		GeometryShader->SetParameters(RHICmdList, MaterialRenderProxy, *View);
+		PixelShader->SetParameters(RHICmdList, MaterialRenderProxy,*View);
+
+		//if(HullShader && DomainShader)
+		//{
+		//	HullShader->SetParameters(RHICmdList, MaterialRenderProxy,*View);
+		//	DomainShader->SetParameters(RHICmdList, MaterialRenderProxy,*View);
+		//}
+
+		RHICmdList.SetBlendState(TStaticBlendState<>::GetRHI());
+	}
+
+	void SetMeshRenderState(
+		FRHICommandList& RHICmdList, 
+		const FSceneView& View,
+		const FPrimitiveSceneProxy* PrimitiveSceneProxy,
+		const FMeshBatch& Mesh,
+		int32 BatchElementIndex,
+		bool bBackFace,
+		float DitheredLODTransitionValue,
+		const ElementDataType& ElementData,
+		const ContextDataType PolicyContext
+		) const
+	{
+		const FMeshBatchElement& BatchElement = Mesh.Elements[BatchElementIndex];
+
+		VertexShader->SetMesh(RHICmdList, VertexFactory,View,PrimitiveSceneProxy,BatchElement,DitheredLODTransitionValue);
+
+		//if(HullShader && DomainShader)
+		//{
+		//	HullShader->SetMesh(RHICmdList, VertexFactory,View,PrimitiveSceneProxy,BatchElement,DitheredLODTransitionValue);
+		//	DomainShader->SetMesh(RHICmdList, VertexFactory,View,PrimitiveSceneProxy,BatchElement,DitheredLODTransitionValue);
+		//}
+
+		GeometryShader->SetMesh(RHICmdList, VertexFactory, View, PrimitiveSceneProxy, BatchElement, DitheredLODTransitionValue);
+
+		const auto FeatureLevel = View.GetFeatureLevel();
+		PixelShader->SetMesh(RHICmdList, VertexFactory, View, PrimitiveSceneProxy, BatchElement, DitheredLODTransitionValue);
+		//PixelShader->SetMesh(RHICmdList, VertexFactory,PrimitiveSceneProxy, BatchElement, View, bBackFace, BuiltLightingAndSelectedFlags, LMResolutionScale, bTextureMapped);	
+		FMeshDrawingPolicy::SetMeshRenderState(RHICmdList, View,PrimitiveSceneProxy,Mesh,BatchElementIndex,bBackFace, DitheredLODTransitionValue,FMeshDrawingPolicy::ElementDataType(),PolicyContext);
+	}
+
+	/** 
+	 * Create bound shader state using the vertex decl from the mesh draw policy
+	 * as well as the shaders needed to draw the mesh
+	 * @return new bound shader state object
+	 */
+	FBoundShaderStateInput GetBoundShaderStateInput(ERHIFeatureLevel::Type InFeatureLevel)
+	{
+		return FBoundShaderStateInput(
+			FMeshDrawingPolicy::GetVertexDeclaration(), 
+			VertexShader->GetVertexShader(),
+			FHullShaderRHIRef(),//GETSAFERHISHADER_HULL(HullShader), 
+			FDomainShaderRHIRef(),//GETSAFERHISHADER_DOMAIN(DomainShader),
+			PixelShader->GetPixelShader(),
+			GeometryShader->GetGeometryShader());
+	}
+
+	friend int32 CompareDrawingPolicy(const FVertexDensityDrawingPolicy& A,const FVertexDensityDrawingPolicy& B)
+	{
+		COMPAREDRAWINGPOLICYMEMBERS(VertexShader);
+		COMPAREDRAWINGPOLICYMEMBERS(GeometryShader);
+		COMPAREDRAWINGPOLICYMEMBERS(PixelShader);
+		COMPAREDRAWINGPOLICYMEMBERS(HullShader);
+		COMPAREDRAWINGPOLICYMEMBERS(DomainShader);
+		COMPAREDRAWINGPOLICYMEMBERS(VertexFactory);
+		COMPAREDRAWINGPOLICYMEMBERS(MaterialRenderProxy);
+		return 0;//CompareDrawingPolicy(A.LightMapPolicy,B.LightMapPolicy);
+	}
+
+private:
+	FVertexDensityVS* VertexShader;
+	FVertexDensityGS* GeometryShader;
+	FVertexDensityPS* PixelShader;
+	struct FVertexDensityHS* HullShader;
+	struct FVertexDensityDS* DomainShader;
+
+	EBlendMode BlendMode;
+};
+
+class FVertexDensityDrawingPolicyFactory
+{
+public:
+	enum { bAllowSimpleElements = false };
+	struct ContextType {};
+
+	static bool DrawDynamicMesh(
+		FRHICommandList& RHICmdList, 
+		const FViewInfo& View,
+		ContextType DrawingContext,
+		const FMeshBatch& Mesh,
+		bool bBackFace,
+		bool bPreFog,
+		const FPrimitiveSceneProxy* PrimitiveSceneProxy,
+		FHitProxyId HitProxyId
+		)
+	{
+		const auto FeatureLevel = View.GetFeatureLevel();
+		const FMaterialRenderProxy* MaterialRenderProxy = Mesh.MaterialRenderProxy;
+		const FMaterial* Material = MaterialRenderProxy->GetMaterial(FeatureLevel);
+		const EBlendMode BlendMode = Material->GetBlendMode();
+
+		const bool bMaterialMasked = Material->IsMasked();
+		const bool bMaterialModifiesMesh = Material->MaterialModifiesMeshPosition_RenderThread();
+		if (!bMaterialMasked && !bMaterialModifiesMesh)
+		{
+			// Override with the default material for opaque materials that are not two sided
+			MaterialRenderProxy = GEngine->LevelColorationLitMaterial->GetRenderProxy(false);
+		}
+
+		FVertexDensityDrawingPolicy DrawingPolicy(View, Mesh.VertexFactory, MaterialRenderProxy, BlendMode);
+		RHICmdList.BuildAndSetLocalBoundShaderState(DrawingPolicy.GetBoundShaderStateInput(FeatureLevel));
+		DrawingPolicy.SetSharedState(RHICmdList, &View, FVertexDensityDrawingPolicy::ContextDataType());
+		for (int32 BatchElementIndex = 0; BatchElementIndex < Mesh.Elements.Num(); BatchElementIndex++)
+		{
+			DrawingPolicy.SetMeshRenderState(RHICmdList, View, PrimitiveSceneProxy, Mesh, BatchElementIndex, bBackFace, Mesh.DitheredLODTransitionAlpha,
+				FVertexDensityDrawingPolicy::ElementDataType(),
+				FVertexDensityDrawingPolicy::ContextDataType()
+				);
+			DrawingPolicy.DrawMesh(RHICmdList, Mesh, BatchElementIndex);
+		}
+		return true;
+	}
+};
+
+#if !UE_BUILD_DOCS
+IMPLEMENT_MATERIAL_SHADER_TYPE(, FVertexDensityVS, TEXT("VertexDensityShader"), TEXT("MainVS"), SF_Vertex);
+IMPLEMENT_MATERIAL_SHADER_TYPE(, FVertexDensityGS, TEXT("VertexDensityShader"), TEXT("MainGS"), SF_Geometry);
+IMPLEMENT_MATERIAL_SHADER_TYPE(, FVertexDensityPS, TEXT("VertexDensityShader"), TEXT("MainPS"), SF_Pixel);
+#endif
+
+bool FDeferredShadingSceneRenderer::RenderVertexDensities(FRHICommandListImmediate& RHICmdList)
+{
+	if (Scene->GetFeatureLevel() >= ERHIFeatureLevel::SM4)
+	{
+		SCOPED_DRAW_EVENT(RHICmdList, VertexDensity);
+
+		for (int32 ViewIndex = 0;ViewIndex < Views.Num();ViewIndex++)
+		{
+			SCOPED_CONDITIONAL_DRAW_EVENTF(RHICmdList, EventView, Views.Num() > 1, TEXT("View%d"), ViewIndex);
+
+			FViewInfo& View = Views[ViewIndex];
+
+			// Opaque blending, depth tests and writes.
+			RHICmdList.SetBlendState(TStaticBlendState<CW_RGBA, BO_Add, BF_One, BF_One, BO_Add, BF_One, BF_One>::GetRHI());
+			RHICmdList.SetDepthStencilState(TStaticDepthStencilState<true, CF_DepthNearOrEqual>::GetRHI());
+			RHICmdList.SetViewport(View.ViewRect.Min.X, View.ViewRect.Min.Y, 0, View.ViewRect.Max.X, View.ViewRect.Max.Y, 1);
+
+			{
+				SCOPED_DRAW_EVENT(RHICmdList, Dynamic);
+				FVertexDensityDrawingPolicyFactory::ContextType Context;
+
+				for (int32 MeshBatchIndex = 0; MeshBatchIndex < View.DynamicMeshElements.Num(); MeshBatchIndex++)
+				{
+					const FMeshBatchAndRelevance& MeshBatchAndRelevance = View.DynamicMeshElements[MeshBatchIndex];
+
+					if (MeshBatchAndRelevance.bHasOpaqueOrMaskedMaterial || ViewFamily.EngineShowFlags.Wireframe)
+					{
+						const FMeshBatch& MeshBatch = *MeshBatchAndRelevance.Mesh;
+						FVertexDensityDrawingPolicyFactory::DrawDynamicMesh(RHICmdList, View, Context, MeshBatch, false, true, MeshBatchAndRelevance.PrimitiveSceneProxy, MeshBatch.BatchHitProxyId);
+					}
+				}
+			}
+		}
+	}
+
+	return false;
+}
