Merge parallel occlusion culling

diff --git a/Engine/Source/Runtime/RHI/Private/RHI.cpp b/Engine/Source/Runtime/RHI/Private/RHI.cpp
--- a/Engine/Source/Runtime/RHI/Private/RHI.cpp
+++ b/Engine/Source/Runtime/RHI/Private/RHI.cpp
@@ -209,6 +209,7 @@ bool GRHIRequiresEarlyBackBufferRenderTarget = true;
 bool GRHISupportsRHIThread = false;
 bool GRHISupportsParallelRHIExecute = false;
 bool GSupportsHDR32bppEncodeModeIntrinsic = false;
+bool GSupportsParallelOcclusionQueries = false;
 
 
 /** Whether we are profiling GPU hitches. */
diff --git a/Engine/Source/Runtime/RHI/Public/RHI.h b/Engine/Source/Runtime/RHI/Public/RHI.h
--- a/Engine/Source/Runtime/RHI/Public/RHI.h
+++ b/Engine/Source/Runtime/RHI/Public/RHI.h
@@ -119,6 +119,9 @@ extern RHI_API bool GSupportsDepthBoundsTest;
 /** True if the RHI supports 'GetHDR32bppEncodeModeES2' shader intrinsic. */
 extern RHI_API bool GSupportsHDR32bppEncodeModeIntrinsic;
 
+/** True if the RHI supports getting the result of occlusion queries when on a thread other than the renderthread */
+extern RHI_API bool GSupportsParallelOcclusionQueries;
+
 /** The minimum Z value in clip space for the RHI. */
 extern RHI_API float GMinClipZ;
 
diff --git a/Engine/Source/Runtime/Renderer/Private/GPUBenchmark.cpp b/Engine/Source/Runtime/Renderer/Private/GPUBenchmark.cpp
--- a/Engine/Source/Runtime/Renderer/Private/GPUBenchmark.cpp
+++ b/Engine/Source/Runtime/Renderer/Private/GPUBenchmark.cpp
@@ -477,7 +477,7 @@ void RendererGPUBenchmark(FRHICommandListImmediate& RHICmdList, FSynthBenchmarkR
 			// flushes the RHI thread to make sure all RHICmdList.EndRenderQuery() commands got executed.
 			RHICmdList.ImmediateFlush(EImmediateFlushType::FlushRHIThread);
 			RHICmdList.GetRenderQueryResult(TimerQueries[0], OldAbsTime, true);
-			GTimerQueryPool.ReleaseQuery(RHICmdList, TimerQueries[0]);
+			GTimerQueryPool.ReleaseQuery(TimerQueries[0]);
 
 			for(uint32 Iteration = 0; Iteration < IterationCount; ++Iteration)
 			{
@@ -489,7 +489,7 @@ void RendererGPUBenchmark(FRHICommandListImmediate& RHICmdList, FSynthBenchmarkR
 
 					uint64 AbsTime;
 					RHICmdList.GetRenderQueryResult(TimerQueries[QueryIndex], AbsTime, true);
-					GTimerQueryPool.ReleaseQuery(RHICmdList, TimerQueries[QueryIndex]);
+					GTimerQueryPool.ReleaseQuery(TimerQueries[QueryIndex]);
 
 					uint64 RelTime = AbsTime - OldAbsTime; 
 
diff --git a/Engine/Source/Runtime/Renderer/Private/SceneOcclusion.cpp b/Engine/Source/Runtime/Renderer/Private/SceneOcclusion.cpp
--- a/Engine/Source/Runtime/Renderer/Private/SceneOcclusion.cpp
+++ b/Engine/Source/Runtime/Renderer/Private/SceneOcclusion.cpp
@@ -84,7 +84,7 @@ FRenderQueryRHIRef FRenderQueryPool::AllocateQuery()
 	return Queries.Pop(/*bAllowShrinking=*/ false);
 }
 
-void FRenderQueryPool::ReleaseQuery(FRHICommandListImmediate& RHICmdList, FRenderQueryRHIRef &Query)
+void FRenderQueryPool::ReleaseQuery(FRenderQueryRHIRef &Query)
 {
 	if ( IsValidRef(Query) )
 	{
@@ -1138,7 +1138,7 @@ void FDeferredShadingSceneRenderer::BeginOcclusionTests(FRHICommandListImmediate
 				{
 					//FRenderQueryRHIParamRef Query = QueryIt.Value();
 					//check( Query.GetRefCount() == 1 );
-					ViewState->OcclusionQueryPool.ReleaseQuery(RHICmdList, QueryIt.Value());
+					ViewState->OcclusionQueryPool.ReleaseQuery(QueryIt.Value());
 				}
 				ShadowOcclusionQueryMap.Reset();
 
diff --git a/Engine/Source/Runtime/Renderer/Private/ScenePrivate.h b/Engine/Source/Runtime/Renderer/Private/ScenePrivate.h
--- a/Engine/Source/Runtime/Renderer/Private/ScenePrivate.h
+++ b/Engine/Source/Runtime/Renderer/Private/ScenePrivate.h
@@ -160,6 +160,20 @@ public:
 		PendingOcclusionQuery.AddZeroed(FOcclusionQueryHelpers::MaxBufferedOcclusionFrames);
 	}
 
+	FORCEINLINE FPrimitiveOcclusionHistory()
+		: HZBTestIndex(0)
+		, HZBTestFrameNumber(~0u)
+		, LastVisibleTime(0.0f)
+		, LastConsideredTime(0.0f)
+		, LastPixelsPercentage(0.0f)
+		, bGroupedQuery(false)
+		, CustomIndex(0)
+	{
+		PendingOcclusionQuery.Empty(FOcclusionQueryHelpers::MaxBufferedOcclusionFrames);
+		PendingOcclusionQuery.AddZeroed(FOcclusionQueryHelpers::MaxBufferedOcclusionFrames);
+	}
+
+
 	/** Destructor. Note that the query should have been released already. */
 	~FPrimitiveOcclusionHistory()
 	{
@@ -171,7 +185,7 @@ public:
 	{
 		for (int32 QueryIndex = 0; QueryIndex < NumBufferedFrames; QueryIndex++)
 		{
-			Pool.ReleaseQuery(RHICmdList, PendingOcclusionQuery[QueryIndex]);
+			Pool.ReleaseQuery(PendingOcclusionQuery[QueryIndex]);
 		}
 	}
 
@@ -245,7 +259,7 @@ public:
 	FRenderQueryRHIRef AllocateQuery();
 
 	/** De-reference an render query, returning it to the pool instead of deleting it when the refcount reaches 0. */
-	void ReleaseQuery(FRHICommandListImmediate& RHICmdList, FRenderQueryRHIRef &Query);
+	void ReleaseQuery(FRenderQueryRHIRef &Query);
 
 private:
 	/** Container for available render queries. */
diff --git a/Engine/Source/Runtime/Renderer/Private/SceneVisibility.cpp b/Engine/Source/Runtime/Renderer/Private/SceneVisibility.cpp
--- a/Engine/Source/Runtime/Renderer/Private/SceneVisibility.cpp
+++ b/Engine/Source/Runtime/Renderer/Private/SceneVisibility.cpp
@@ -100,6 +100,14 @@ static TAutoConsoleVariable<float> CVarStaticMeshLODDistanceScale(
 	TEXT("Scale factor for the distance used in computing discrete LOD for static meshes. (0.25-1)"),
 	ECVF_Scalability | ECVF_RenderThreadSafe);
 
+static int32 GOcclusionCullParallelPrimFetch = 1;
+static FAutoConsoleVariableRef CVarOcclusionCullParallelPrimFetch(
+	TEXT("r.OcclusionCullParallelPrimFetch"),
+	GOcclusionCullParallelPrimFetch,
+	TEXT("Enables Parallel Occlusion Cull primitive fetch."),
+	ECVF_RenderThreadSafe
+	);
+
 /** Distance fade cvars */
 static int32 GDisableLODFade = false;
 static FAutoConsoleVariableRef CVarDisableLODFade( TEXT("r.DisableLODFade"), GDisableLODFade, TEXT("Disable fading for distance culling"), ECVF_RenderThreadSafe );
@@ -395,346 +403,820 @@ static void UpdatePrimitiveFading(const FScene* Scene, FViewInfo& View)
 	}
 }
 
-/**
- * Cull occluded primitives in the view.
- */
-static int32 OcclusionCull(FRHICommandListImmediate& RHICmdList, const FScene* Scene, FViewInfo& View)
+struct FOcclusionBounds
 {
-	SCOPE_CYCLE_COUNTER(STAT_OcclusionCull);
-	int32 NumBufferedFrames = FOcclusionQueryHelpers::GetNumBufferedFrames();
+	FOcclusionBounds(FPrimitiveOcclusionHistory* InPrimitiveOcclusionHistory, const FVector& InBoundsOrigin, const FVector& InBoundsExtent, bool bInGroupedQuery)
+		: PrimitiveOcclusionHistory(InPrimitiveOcclusionHistory)
+		, BoundsOrigin(InBoundsOrigin)
+		, BoundsExtent(InBoundsExtent)
+		, bGroupedQuery(bInGroupedQuery)
+	{}
+	FPrimitiveOcclusionHistory* PrimitiveOcclusionHistory;
+	FVector BoundsOrigin;
+	FVector BoundsExtent;	
+	bool bGroupedQuery;
+};
 
-	// INITVIEWS_TODO: This could be more efficient if broken up in to separate concerns:
-	// - What is occluded?
-	// - For which primitives should we render occlusion queries?
-	// - Generate occlusion query geometry.
+struct FHZBBound
+{
+	FHZBBound(FPrimitiveOcclusionHistory* InTargetHistory, const FVector& InBoundsOrigin, const FVector& InBoundsExtent)
+	: TargetHistory(InTargetHistory)
+	, BoundsOrigin(InBoundsOrigin)
+	, BoundsExtent(InBoundsExtent)
+	{}
+
+	FPrimitiveOcclusionHistory* TargetHistory;
+	FVector BoundsOrigin;
+	FVector BoundsExtent;
+};
+
+#define BALANCE_LOAD 1
+#define QUERY_SANITY_CHECK 0
+
+struct FVisForPrimParams
+{
+	FVisForPrimParams(){}
+
+	FVisForPrimParams(const FScene* InScene, 
+						FViewInfo* InView, 
+						FViewElementPDI* InOcclusionPDI, 						
+						const int32 InStartIndex, 
+						const int32 InNumToProcess, 
+						const bool bInSubmitQueries, 
+						const bool bInHZBOcclusion,						
+						TArray<FPrimitiveOcclusionHistory>* OutOcclusionHistory,
+						TArray<FPrimitiveOcclusionHistory*>* OutQueriesToRelease,
+						TArray<FHZBBound>* OutHZBBounds,
+						TArray<FOcclusionBounds>* OutQueriesToRun)
+		: Scene(InScene)
+		, View(InView)
+		, OcclusionPDI(InOcclusionPDI)		
+		, StartIndex(InStartIndex)
+		, NumToProcess(InNumToProcess)
+		, bSubmitQueries(bInSubmitQueries)
+		, bHZBOcclusion(bInHZBOcclusion)		
+		, InsertPrimitiveOcclusionHistory(OutOcclusionHistory)
+		, QueriesToRelease(OutQueriesToRelease)
+		, HZBBoundsToAdd(OutHZBBounds)
+		, QueriesToAdd(OutQueriesToRun)			
+	{
+
+	}
+
+	void Init(	const FScene* InScene,
+				FViewInfo* InView,
+				FViewElementPDI* InOcclusionPDI,
+				const int32 InStartIndex,
+				const int32 InNumToProcess,
+				const bool bInSubmitQueries,
+				const bool bInHZBOcclusion,				
+				TArray<FPrimitiveOcclusionHistory>* OutOcclusionHistory,
+				TArray<FPrimitiveOcclusionHistory*>* OutQueriesToRelease,
+				TArray<FHZBBound>* OutHZBBounds,
+				TArray<FOcclusionBounds>* OutQueriesToRun)
+			
+	{
+		Scene = InScene;
+		View = InView;
+		OcclusionPDI = InOcclusionPDI;
+		StartIndex = InStartIndex;
+		NumToProcess = InNumToProcess;
+		bSubmitQueries = bInSubmitQueries;
+		bHZBOcclusion = bInHZBOcclusion;
+		InsertPrimitiveOcclusionHistory = OutOcclusionHistory;
+		QueriesToRelease = OutQueriesToRelease;
+		HZBBoundsToAdd = OutHZBBounds;
+		QueriesToAdd = OutQueriesToRun;
+	}
+
+	const FScene* Scene;
+	FViewInfo* View;
+	FViewElementPDI* OcclusionPDI;
+	int32 StartIndex;
+	int32 NumToProcess;
+	bool bSubmitQueries;
+	bool bHZBOcclusion;	
+
+	//occlusion history to insert into.  In parallel these will be all merged back into the view's history on the main thread.
+	//use TChunkedArray so pointers to the new FPrimitiveOcclusionHistory's won't change if the array grows.	
+	TArray<FPrimitiveOcclusionHistory>*		InsertPrimitiveOcclusionHistory;
+	TArray<FPrimitiveOcclusionHistory*>*	QueriesToRelease;
+	TArray<FHZBBound>*						HZBBoundsToAdd;
+	TArray<FOcclusionBounds>*				QueriesToAdd;
+	int32									NumOccludedPrims;
+};
 
+//This function is shared between the single and multi-threaded versions.  Modifications to any primitives indexed by BitIt should be ok
+//since only one of the task threads will ever reference it.  However, any modifications to shared state like the ViewState must be buffered
+//to be recombined later.
+template<bool bSingleThreaded>
+static void FetchVisibilityForPrimitives_Range(FVisForPrimParams& Params)
+{	
 	int32 NumOccludedPrimitives = 0;
+	
+	const FScene* Scene				= Params.Scene;
+	FViewInfo& View					= *Params.View;
+	FViewElementPDI* OcclusionPDI	= Params.OcclusionPDI;
+	const int32 StartIndex			= Params.StartIndex;
+	const int32 NumToProcess		= Params.NumToProcess;
+	const bool bSubmitQueries		= Params.bSubmitQueries;
+	const bool bHZBOcclusion		= Params.bHZBOcclusion;
+
 	FSceneViewState* ViewState = (FSceneViewState*)View.State;
+	const int32 NumBufferedFrames = FOcclusionQueryHelpers::GetNumBufferedFrames();
+	const bool bClearQueries = !View.Family->EngineShowFlags.HitProxies;
+	const float CurrentRealTime = View.Family->CurrentRealTime;
+	uint32 OcclusionFrameCounter = ViewState->OcclusionFrameCounter;
+	FRenderQueryPool& OcclusionQueryPool = ViewState->OcclusionQueryPool;
+	FHZBOcclusionTester& HZBOcclusionTests = ViewState->HZBOcclusionTests;
 	
-	// Disable HZB on OpenGL platforms to avoid rendering artefacts
-	// It can be forced on by setting HZBOcclusion to 2
-	bool bHZBOcclusion = (!IsOpenGLPlatform(GShaderPlatformForFeatureLevel[Scene->GetFeatureLevel()]) && GHZBOcclusion) || (GHZBOcclusion == 2);
 
-	// Use precomputed visibility data if it is available.
-	if (View.PrecomputedVisibilityData)
-	{
-		QUICK_SCOPE_CYCLE_COUNTER(STAT_LookupPrecomputedVisibility);
+	TSet<FPrimitiveOcclusionHistory, FPrimitiveOcclusionHistoryKeyFuncs>& ViewPrimitiveOcclusionHistory = ViewState->PrimitiveOcclusionHistorySet;
+	TArray<FPrimitiveOcclusionHistory>* InsertPrimitiveOcclusionHistory = Params.InsertPrimitiveOcclusionHistory;
+	TArray<FPrimitiveOcclusionHistory*>* QueriesToRelease = Params.QueriesToRelease;
+	TArray<FHZBBound>* HZBBoundsToAdd = Params.HZBBoundsToAdd;
+	TArray<FOcclusionBounds>* QueriesToAdd = Params.QueriesToAdd;	
 
-		uint8 PrecomputedVisibilityFlags = EOcclusionFlags::CanBeOccluded | EOcclusionFlags::HasPrecomputedVisibility;
-		for (FSceneSetBitIterator BitIt(View.PrimitiveVisibilityMap); BitIt; ++BitIt)
+	const int32 ReserveAmount = NumToProcess;
+	if (!bSingleThreaded)
+	{		
+		check(InsertPrimitiveOcclusionHistory);
+		check(QueriesToRelease);
+		check(HZBBoundsToAdd);
+		check(QueriesToAdd);
+
+		//avoid doing reallocs as much as possible.  Unlikely to make an entry per processed element.		
+		InsertPrimitiveOcclusionHistory->Reserve(ReserveAmount);
+		QueriesToRelease->Reserve(ReserveAmount);
+		HZBBoundsToAdd->Reserve(ReserveAmount);
+		QueriesToAdd->Reserve(ReserveAmount);
+	}
+
+	int32 NumProcessed = 0;
+	int32 NumTotalPrims = View.PrimitiveVisibilityMap.Num();
+	int32 NumTotalDefUnoccluded = View.PrimitiveDefinitelyUnoccludedMap.Num();
+
+	//if we are load balanced then we iterate only the set bits, and the ranges have been pre-selected to evenly distribute set bits among the tasks with no overlaps.
+	//if not, then the entire array is evenly divided by range.
+#if BALANCE_LOAD
+	for (FSceneSetBitIterator BitIt(View.PrimitiveVisibilityMap, StartIndex); BitIt && (NumProcessed < NumToProcess); ++BitIt, ++NumProcessed)
+#else
+	for (TBitArray<SceneRenderingBitArrayAllocator>::FIterator BitIt(View.PrimitiveVisibilityMap, StartIndex); BitIt && (NumProcessed < NumToProcess); ++BitIt, ++NumProcessed)
+#endif
+	{		
+		uint8 OcclusionFlags = Scene->PrimitiveOcclusionFlags[BitIt.GetIndex()];
+		bool bCanBeOccluded = (OcclusionFlags & EOcclusionFlags::CanBeOccluded) != 0;
+
+#if !BALANCE_LOAD		
+		if (!View.PrimitiveVisibilityMap.AccessCorrespondingBit(BitIt))
 		{
-			if ((Scene->PrimitiveOcclusionFlags[BitIt.GetIndex()] & PrecomputedVisibilityFlags) == PrecomputedVisibilityFlags)
+			continue;
+		}
+#endif
+
+		//we can't allow the prim history insertion array to realloc or it will invalidate pointers in the other output arrays.
+		const bool bCanAllocPrimHistory = bSingleThreaded || InsertPrimitiveOcclusionHistory->Num() < InsertPrimitiveOcclusionHistory->Max();		
+
+		if (GIsEditor)
+		{
+			FPrimitiveSceneInfo* PrimitiveSceneInfo = Scene->Primitives[BitIt.GetIndex()];
+
+			if (PrimitiveSceneInfo->Proxy->IsSelected())
 			{
-				FPrimitiveVisibilityId VisibilityId = Scene->PrimitiveVisibilityIds[BitIt.GetIndex()];
-				if ((View.PrecomputedVisibilityData[VisibilityId.ByteIndex] & VisibilityId.BitMask) == 0)
-				{
-					View.PrimitiveVisibilityMap.AccessCorrespondingBit(BitIt) = false;
-					INC_DWORD_STAT_BY(STAT_StaticallyOccludedPrimitives,1);
-					STAT(NumOccludedPrimitives++);
-				}
+				// to render occluded outline for selected objects
+				bCanBeOccluded = false;
 			}
 		}
-	}
+		int32 NumSubQueries = 1;
+		bool bSubQueries = false;
+		const TArray<FBoxSphereBounds>* SubBounds = nullptr;
+		TArray<bool, SceneRenderingAllocator> SubIsOccluded;
 
-	float CurrentRealTime = View.Family->CurrentRealTime;
-	if (ViewState)
-	{
-		if (Scene->GetFeatureLevel() >= ERHIFeatureLevel::SM4)
+		if ((OcclusionFlags & EOcclusionFlags::HasSubprimitiveQueries) && GAllowSubPrimitiveQueries)
 		{
-			bool bClearQueries = !View.Family->EngineShowFlags.HitProxies;
-			bool bSubmitQueries = !View.bDisableQuerySubmissions;
-#if !(UE_BUILD_SHIPPING || UE_BUILD_TEST)
-			bSubmitQueries = bSubmitQueries && !ViewState->HasViewParent() && !ViewState->bIsFrozen;
-#endif
-
-			if( bHZBOcclusion )
+			FPrimitiveSceneProxy* Proxy = Scene->Primitives[BitIt.GetIndex()]->Proxy;
+			SubBounds = Proxy->GetOcclusionQueries(&View);
+			NumSubQueries = SubBounds->Num();
+			bSubQueries = true;
+			if (!NumSubQueries)
 			{
-				QUICK_SCOPE_CYCLE_COUNTER(STAT_MapHZBResults);
-				check(!ViewState->HZBOcclusionTests.IsValidFrame(ViewState->OcclusionFrameCounter));
-				ViewState->HZBOcclusionTests.MapResults(RHICmdList);
+				View.PrimitiveVisibilityMap.AccessCorrespondingBit(BitIt) = false;
+				continue;
 			}
+			SubIsOccluded.Reserve(NumSubQueries);
+		}
 
-			FViewElementPDI OcclusionPDI(&View, NULL);
+		bool bAllSubOcclusionStateIsDefinite = true;
+		bool bAllSubOccluded = true;
+		FPrimitiveComponentId PrimitiveId = Scene->PrimitiveComponentIds[BitIt.GetIndex()];
 
-			QUICK_SCOPE_CYCLE_COUNTER(STAT_FetchVisibilityForPrimitives);
+		for (int32 SubQuery = 0; SubQuery < NumSubQueries; SubQuery++)
+		{
+			FPrimitiveOcclusionHistory* PrimitiveOcclusionHistory = ViewPrimitiveOcclusionHistory.Find(FPrimitiveOcclusionHistoryKey(PrimitiveId, SubQuery));
+			
+			bool bIsOccluded = false;
+			bool bOcclusionStateIsDefinite = false;			
 
-			for (FSceneSetBitIterator BitIt(View.PrimitiveVisibilityMap); BitIt; ++BitIt)
+			if (!PrimitiveOcclusionHistory)
 			{
-				uint8 OcclusionFlags = Scene->PrimitiveOcclusionFlags[BitIt.GetIndex()];
-				bool bCanBeOccluded = (OcclusionFlags & EOcclusionFlags::CanBeOccluded) != 0;
-
-				if(GIsEditor)
-				{
-					FPrimitiveSceneInfo* PrimitiveSceneInfo = Scene->Primitives[BitIt.GetIndex()];
-
-					if(PrimitiveSceneInfo->Proxy->IsSelected())
-					{
-						// to render occluded outline for selected objects
-						bCanBeOccluded = false;
-					}
+				// If the primitive doesn't have an occlusion history yet, create it.
+				if (bSingleThreaded)
+				{					
+					// In singlethreaded mode we can safely modify the view's history directly.
+					PrimitiveOcclusionHistory = &ViewPrimitiveOcclusionHistory[
+						ViewPrimitiveOcclusionHistory.Add(FPrimitiveOcclusionHistory(PrimitiveId, SubQuery))
+					];
 				}
-				int32 NumSubQueries = 1;
-				bool bSubQueries = false;
-				const TArray<FBoxSphereBounds>* SubBounds = nullptr;
-				TArray<bool, SceneRenderingAllocator> SubIsOccluded;
+				else if (bCanAllocPrimHistory)
+				{
+					// In multithreaded mode we have to buffer the new histories and add them to the view during a post-combine
+					PrimitiveOcclusionHistory = &(*InsertPrimitiveOcclusionHistory)[
+						InsertPrimitiveOcclusionHistory->Add(FPrimitiveOcclusionHistory(PrimitiveId, SubQuery))
+					];
+				}				
+				
+				// If the primitive hasn't been visible recently enough to have a history, treat it as unoccluded this frame so it will be rendered as an occluder and its true occlusion state can be determined.
+				// already set bIsOccluded = false;
 
-				if ((OcclusionFlags & EOcclusionFlags::HasSubprimitiveQueries) && GAllowSubPrimitiveQueries)
+				// Flag the primitive's occlusion state as indefinite, which will force it to be queried this frame.
+				// The exception is if the primitive isn't occludable, in which case we know that it's definitely unoccluded.
+				bOcclusionStateIsDefinite = bCanBeOccluded ? false : true;
+			}
+			else
+			{
+				if (View.bIgnoreExistingQueries)
 				{
-					FPrimitiveSceneProxy* Proxy = Scene->Primitives[BitIt.GetIndex()]->Proxy;
-					SubBounds = Proxy->GetOcclusionQueries(&View);
-					NumSubQueries = SubBounds->Num();
-					bSubQueries = true;
-					if (!NumSubQueries)
-					{
-						View.PrimitiveVisibilityMap.AccessCorrespondingBit(BitIt) = false;
-						continue;
-					}
-					SubIsOccluded.Reserve(NumSubQueries);
+					// If the view is ignoring occlusion queries, the primitive is definitely unoccluded.
+					// already set bIsOccluded = false;
+					bOcclusionStateIsDefinite = View.bDisableQuerySubmissions;
 				}
-
-				bool bAllSubOcclusionStateIsDefinite = true;
-				bool bAllSubOccluded = true;
-				FPrimitiveComponentId PrimitiveId = Scene->PrimitiveComponentIds[BitIt.GetIndex()];
-				for (int32 SubQuery = 0; SubQuery < NumSubQueries; SubQuery++)
+				else if (bCanBeOccluded)
 				{
-
-					FPrimitiveOcclusionHistory* PrimitiveOcclusionHistory = ViewState->PrimitiveOcclusionHistorySet.Find(FPrimitiveOcclusionHistoryKey(PrimitiveId, SubQuery));
-					bool bIsOccluded = false;
-					bool bOcclusionStateIsDefinite = false;
-					if (!PrimitiveOcclusionHistory)
+					if (bHZBOcclusion)
 					{
-						// If the primitive doesn't have an occlusion history yet, create it.
-						PrimitiveOcclusionHistory = &ViewState->PrimitiveOcclusionHistorySet[
-							ViewState->PrimitiveOcclusionHistorySet.Add(FPrimitiveOcclusionHistory(PrimitiveId, SubQuery))
-							];
-
-						// If the primitive hasn't been visible recently enough to have a history, treat it as unoccluded this frame so it will be rendered as an occluder and its true occlusion state can be determined.
-						// already set bIsOccluded = false;
-
-						// Flag the primitive's occlusion state as indefinite, which will force it to be queried this frame.
-						// The exception is if the primitive isn't occludable, in which case we know that it's definitely unoccluded.
-						bOcclusionStateIsDefinite = bCanBeOccluded ? false : true;
+						if (HZBOcclusionTests.IsValidFrame(PrimitiveOcclusionHistory->HZBTestFrameNumber))
+						{
+							bIsOccluded = !HZBOcclusionTests.IsVisible(PrimitiveOcclusionHistory->HZBTestIndex);
+							bOcclusionStateIsDefinite = true;
+						}
 					}
 					else
 					{
-						if (View.bIgnoreExistingQueries)
-						{
-							// If the view is ignoring occlusion queries, the primitive is definitely unoccluded.
-							// already set bIsOccluded = false;
-							bOcclusionStateIsDefinite = View.bDisableQuerySubmissions;
-						}
-						else if (bCanBeOccluded)
+						// Read the occlusion query results.
+						uint64 NumSamples = 0;
+						FRenderQueryRHIRef& PastQuery = PrimitiveOcclusionHistory->GetPastQuery(OcclusionFrameCounter, NumBufferedFrames);
+						if (IsValidRef(PastQuery))
 						{
-							if( bHZBOcclusion )
+							int32 RefCount = PastQuery.GetRefCount();
+							// NOTE: RHIGetOcclusionQueryResult should never fail when using a blocking call, rendering artifacts may show up.
+							//if (RHICmdList.GetRenderQueryResult(PastQuery, NumSamples, true))
+							if (GDynamicRHI->RHIGetRenderQueryResult(PastQuery.GetReference(), NumSamples, true))
 							{
-								if( ViewState->HZBOcclusionTests.IsValidFrame(PrimitiveOcclusionHistory->HZBTestFrameNumber) )
-								{
-									bIsOccluded = !ViewState->HZBOcclusionTests.IsVisible( PrimitiveOcclusionHistory->HZBTestIndex );
-									bOcclusionStateIsDefinite = true;
-								}
-							}
-							else
-							{
-								// Read the occlusion query results.
-								uint64 NumSamples = 0;
-								FRenderQueryRHIRef& PastQuery = PrimitiveOcclusionHistory->GetPastQuery(ViewState->OcclusionFrameCounter, NumBufferedFrames);
-								if (IsValidRef(PastQuery))
+								// we render occlusion without MSAA
+								uint32 NumPixels = (uint32)NumSamples;
+								NumPixels = 1000;
+
+								// The primitive is occluded if none of its bounding box's pixels were visible in the previous frame's occlusion query.
+								bIsOccluded = (NumPixels == 0);
+
+								
+								if (!bIsOccluded)
 								{
-									// NOTE: RHIGetOcclusionQueryResult should never fail when using a blocking call, rendering artifacts may show up.
-									if (RHICmdList.GetRenderQueryResult(PastQuery, NumSamples, true))
-									{
-										// we render occlusion without MSAA
-										uint32 NumPixels = (uint32)NumSamples;
-
-										// The primitive is occluded if none of its bounding box's pixels were visible in the previous frame's occlusion query.
-										bIsOccluded = (NumPixels == 0);
-										if (!bIsOccluded)
-										{
-											checkSlow(View.OneOverNumPossiblePixels > 0.0f);
-											PrimitiveOcclusionHistory->LastPixelsPercentage = NumPixels * View.OneOverNumPossiblePixels;
-										}
-										else
-										{
-											PrimitiveOcclusionHistory->LastPixelsPercentage = 0.0f;
-										}
-
-
-										// Flag the primitive's occlusion state as definite if it wasn't grouped.
-										bOcclusionStateIsDefinite = !PrimitiveOcclusionHistory->bGroupedQuery;
-									}
-									else
-									{
-										// If the occlusion query failed, treat the primitive as visible.  
-										// already set bIsOccluded = false;
-									}
+									checkSlow(View.OneOverNumPossiblePixels > 0.0f);
+									PrimitiveOcclusionHistory->LastPixelsPercentage = NumPixels * View.OneOverNumPossiblePixels;
 								}
 								else
 								{
-									// If there's no occlusion query for the primitive, set it's visibility state to whether it has been unoccluded recently.
-									bIsOccluded = (PrimitiveOcclusionHistory->LastVisibleTime + GEngine->PrimitiveProbablyVisibleTime < CurrentRealTime);
-									if (bIsOccluded)
-									{
-										PrimitiveOcclusionHistory->LastPixelsPercentage = 0.0f;
-									}
-									else
-									{
-										PrimitiveOcclusionHistory->LastPixelsPercentage = GEngine->MaxOcclusionPixelsFraction;
-									}
-
-									// the state was definite last frame, otherwise we would have ran a query
-									bOcclusionStateIsDefinite = true;
-								}
-							}
+									PrimitiveOcclusionHistory->LastPixelsPercentage = 0.0f;
+								}								
+
 
-							if( GVisualizeOccludedPrimitives && bIsOccluded )
+								// Flag the primitive's occlusion state as definite if it wasn't grouped.
+								bOcclusionStateIsDefinite = !PrimitiveOcclusionHistory->bGroupedQuery;
+							}
+							else
 							{
-								const FBoxSphereBounds& Bounds = bSubQueries ? (*SubBounds)[SubQuery] : Scene->PrimitiveOcclusionBounds[ BitIt.GetIndex() ];
-								DrawWireBox( &OcclusionPDI, Bounds.GetBox(), FColor(50, 255, 50), SDPG_Foreground );
+								// If the occlusion query failed, treat the primitive as visible.  
+								// already set bIsOccluded = false;
 							}
+
+							checkf(RefCount == PastQuery.GetRefCount(), TEXT("Ref count on prim: %i, old: %i, new: %i"), PrimitiveOcclusionHistory->PrimitiveId.PrimIDValue, RefCount, PastQuery.GetRefCount());
 						}
 						else
 						{
-							// Primitives that aren't occludable are considered definitely unoccluded.
-							// already set bIsOccluded = false;
-							bOcclusionStateIsDefinite = true;
-						}
+							// If there's no occlusion query for the primitive, set it's visibility state to whether it has been unoccluded recently.
+							bIsOccluded = (PrimitiveOcclusionHistory->LastVisibleTime + GEngine->PrimitiveProbablyVisibleTime < CurrentRealTime);
 
-						if (bClearQueries)
-						{
-							ViewState->OcclusionQueryPool.ReleaseQuery(RHICmdList, PrimitiveOcclusionHistory->GetPastQuery(ViewState->OcclusionFrameCounter, NumBufferedFrames));
+							
+							
+							if (bIsOccluded)
+							{
+								PrimitiveOcclusionHistory->LastPixelsPercentage = 0.0f;
+							}
+							else
+							{
+								PrimitiveOcclusionHistory->LastPixelsPercentage = GEngine->MaxOcclusionPixelsFraction;
+							}							
+
+							// the state was definite last frame, otherwise we would have ran a query
+							bOcclusionStateIsDefinite = true;							
 						}
 					}
 
-					// Set the primitive's considered time to keep its occlusion history from being trimmed.
-					PrimitiveOcclusionHistory->LastConsideredTime = CurrentRealTime;
+					if (GVisualizeOccludedPrimitives && OcclusionPDI && bIsOccluded)
+					{
+						const FBoxSphereBounds& Bounds = bSubQueries ? (*SubBounds)[SubQuery] : Scene->PrimitiveOcclusionBounds[BitIt.GetIndex()];
+						DrawWireBox(OcclusionPDI, Bounds.GetBox(), FColor(50, 255, 50), SDPG_Foreground);
+					}
+				}
+				else
+				{					
+					// Primitives that aren't occludable are considered definitely unoccluded.
+					// already set bIsOccluded = false;
+					bOcclusionStateIsDefinite = true;
+				}
 
-					if (bSubmitQueries && bCanBeOccluded)
+				if (bClearQueries)
+				{					
+					if (bSingleThreaded)
+					{						
+						OcclusionQueryPool.ReleaseQuery(PrimitiveOcclusionHistory->GetPastQuery(OcclusionFrameCounter, NumBufferedFrames));
+					}
+					else
 					{
-						bool bAllowBoundsTest;
-						const FBoxSphereBounds& OcclusionBounds = bSubQueries ? (*SubBounds)[SubQuery] : Scene->PrimitiveOcclusionBounds[ BitIt.GetIndex() ];
-						if (View.bHasNearClippingPlane)
+						FRenderQueryRHIRef &Query = PrimitiveOcclusionHistory->GetPastQuery(OcclusionFrameCounter, NumBufferedFrames);
+						if (IsValidRef(Query))
 						{
-							bAllowBoundsTest = View.NearClippingPlane.PlaneDot(OcclusionBounds.Origin) < 
-								-(FVector::BoxPushOut(View.NearClippingPlane,OcclusionBounds.BoxExtent));
-						}
-						else if (!View.IsPerspectiveProjection())
-						{
-							// Transform parallel near plane
-							static_assert((int32)ERHIZBuffer::IsInverted != 0, "Check equation for culling!");
-							bAllowBoundsTest = View.WorldToScreen(OcclusionBounds.Origin).Z - View.ViewMatrices.ProjMatrix.M[2][2] * OcclusionBounds.SphereRadius < 1;
-						}
-						else
-						{
-							bAllowBoundsTest = OcclusionBounds.SphereRadius < HALF_WORLD_MAX;
+							check(Query.GetRefCount() > 0);
+							QueriesToRelease->Add(PrimitiveOcclusionHistory);							
 						}
+					}
+				}
+			}
+
+			if (PrimitiveOcclusionHistory)
+			{
+
+				// Set the primitive's considered time to keep its occlusion history from being trimmed.
+				PrimitiveOcclusionHistory->LastConsideredTime = CurrentRealTime;
+
+				if (bSubmitQueries && bCanBeOccluded)
+				{
+					bool bAllowBoundsTest;
+					const FBoxSphereBounds& OcclusionBounds = bSubQueries ? (*SubBounds)[SubQuery] : Scene->PrimitiveOcclusionBounds[BitIt.GetIndex()];
+					if (View.bHasNearClippingPlane)
+					{
+						bAllowBoundsTest = View.NearClippingPlane.PlaneDot(OcclusionBounds.Origin) <
+							-(FVector::BoxPushOut(View.NearClippingPlane, OcclusionBounds.BoxExtent));
 
-						if (bAllowBoundsTest)
+					}
+					else if (!View.IsPerspectiveProjection())
+					{
+						// Transform parallel near plane
+						static_assert((int32)ERHIZBuffer::IsInverted != 0, "Check equation for culling!");
+						bAllowBoundsTest = View.WorldToScreen(OcclusionBounds.Origin).Z - View.ViewMatrices.ProjMatrix.M[2][2] * OcclusionBounds.SphereRadius < 1;
+					}
+					else
+					{
+						bAllowBoundsTest = OcclusionBounds.SphereRadius < HALF_WORLD_MAX;
+					}
+
+					if (bAllowBoundsTest)
+					{
+						if (bHZBOcclusion)
 						{
-							if( bHZBOcclusion )
-							{
-								// Always run
-								PrimitiveOcclusionHistory->HZBTestIndex = ViewState->HZBOcclusionTests.AddBounds( OcclusionBounds.Origin, OcclusionBounds.BoxExtent );
-								PrimitiveOcclusionHistory->HZBTestFrameNumber = ViewState->OcclusionFrameCounter;
+							// Always run
+							if (bSingleThreaded)
+							{								
+								PrimitiveOcclusionHistory->HZBTestIndex = HZBOcclusionTests.AddBounds(OcclusionBounds.Origin, OcclusionBounds.BoxExtent);
 							}
 							else
 							{
-								// decide if a query should be run this frame
-								bool bRunQuery,bGroupedQuery;
+								HZBBoundsToAdd->Emplace(PrimitiveOcclusionHistory, OcclusionBounds.Origin, OcclusionBounds.BoxExtent);
+							}
+							PrimitiveOcclusionHistory->HZBTestFrameNumber = OcclusionFrameCounter;
+						}
+						else
+						{
+							// decide if a query should be run this frame
+							bool bRunQuery, bGroupedQuery;
 
-								if (!bSubQueries && // sub queries are never grouped, we assume the custom code knows what it is doing and will group internally if it wants
-									(OcclusionFlags & EOcclusionFlags::AllowApproximateOcclusion))
+							if (!bSubQueries && // sub queries are never grouped, we assume the custom code knows what it is doing and will group internally if it wants
+								(OcclusionFlags & EOcclusionFlags::AllowApproximateOcclusion))
+							{
+								if (bIsOccluded)
 								{
-									if (bIsOccluded)
-									{
-										// Primitives that were occluded the previous frame use grouped queries.
-										bGroupedQuery = true;
-										bRunQuery = true;
-									}
-									else if (bOcclusionStateIsDefinite)
-									{
-										// If the primitive's is definitely unoccluded, only requery it occasionally.
-										float FractionMultiplier = FMath::Max(PrimitiveOcclusionHistory->LastPixelsPercentage/GEngine->MaxOcclusionPixelsFraction, 1.0f);
-										bRunQuery = (FractionMultiplier * GOcclusionRandomStream.GetFraction() < GEngine->MaxOcclusionPixelsFraction);
-										bGroupedQuery = false;
-									}
-									else
-									{
-										bGroupedQuery = false;
-										bRunQuery = true;
-									}
+									// Primitives that were occluded the previous frame use grouped queries.
+									bGroupedQuery = true;
+									bRunQuery = true;
+								}
+								else if (bOcclusionStateIsDefinite)
+								{
+									// If the primitive's is definitely unoccluded, only requery it occasionally.
+									float FractionMultiplier = FMath::Max(PrimitiveOcclusionHistory->LastPixelsPercentage / GEngine->MaxOcclusionPixelsFraction, 1.0f);
+									bRunQuery = (FractionMultiplier * GOcclusionRandomStream.GetFraction()) < GEngine->MaxOcclusionPixelsFraction;
+									bGroupedQuery = false;
 								}
 								else
 								{
-									// Primitives that need precise occlusion results use individual queries.
 									bGroupedQuery = false;
 									bRunQuery = true;
 								}
+							}
+							else
+							{
+								// Primitives that need precise occlusion results use individual queries.
+								bGroupedQuery = false;
+								bRunQuery = true;
+							}
 
-								if (bRunQuery)
+							if (bRunQuery)
+							{
+								const FVector BoundOrigin = OcclusionBounds.Origin + View.ViewMatrices.PreViewTranslation;
+								const FVector BoundExtent = OcclusionBounds.BoxExtent;
+
+								if (bSingleThreaded)
 								{
-									PrimitiveOcclusionHistory->SetCurrentQuery(ViewState->OcclusionFrameCounter, 
-										bGroupedQuery ? 
-										View.GroupedOcclusionQueries.BatchPrimitive(OcclusionBounds.Origin + View.ViewMatrices.PreViewTranslation,OcclusionBounds.BoxExtent) :
-										View.IndividualOcclusionQueries.BatchPrimitive(OcclusionBounds.Origin + View.ViewMatrices.PreViewTranslation,OcclusionBounds.BoxExtent),
+									PrimitiveOcclusionHistory->SetCurrentQuery(OcclusionFrameCounter,
+										bGroupedQuery ?
+										View.GroupedOcclusionQueries.BatchPrimitive(BoundOrigin, BoundExtent) :
+										View.IndividualOcclusionQueries.BatchPrimitive(BoundOrigin, BoundExtent),
 										NumBufferedFrames
-										);	
+										);
+								}
+								else
+								{
+									QueriesToAdd->Emplace(PrimitiveOcclusionHistory, BoundOrigin, BoundExtent, bGroupedQuery);
 								}
-								PrimitiveOcclusionHistory->bGroupedQuery = bGroupedQuery;
 							}
-						}
-						else
-						{
-							// If the primitive's bounding box intersects the near clipping plane, treat it as definitely unoccluded.
-							bIsOccluded = false;
-							bOcclusionStateIsDefinite = true;
-						}
-					}
 
-					if (bSubQueries)
-					{
-						SubIsOccluded.Add(bIsOccluded);
-						if (!bIsOccluded)
-						{
-							bAllSubOccluded = false;
-							if (bOcclusionStateIsDefinite)
-							{
-								PrimitiveOcclusionHistory->LastVisibleTime = CurrentRealTime;
-							}
-						}
-						if (bIsOccluded || !bOcclusionStateIsDefinite)
-						{
-							bAllSubOcclusionStateIsDefinite = false;
+							PrimitiveOcclusionHistory->bGroupedQuery = bGroupedQuery;
 						}
 					}
 					else
 					{
-						if (bIsOccluded)
-						{
-							View.PrimitiveVisibilityMap.AccessCorrespondingBit(BitIt) = false;
-							STAT(NumOccludedPrimitives++);
-						}
-						else if (bOcclusionStateIsDefinite)
+						// If the primitive's bounding box intersects the near clipping plane, treat it as definitely unoccluded.
+						bIsOccluded = false;
+						bOcclusionStateIsDefinite = true;
+					}
+				}
+			}
+
+
+			if (bSubQueries)
+			{
+				SubIsOccluded.Add(bIsOccluded);
+				if (!bIsOccluded)
+				{
+					bAllSubOccluded = false;
+					if (bOcclusionStateIsDefinite)
+					{
+						if (PrimitiveOcclusionHistory)
 						{
 							PrimitiveOcclusionHistory->LastVisibleTime = CurrentRealTime;
-							View.PrimitiveDefinitelyUnoccludedMap.AccessCorrespondingBit(BitIt) = true;
 						}
 					}
 				}
-				if (bSubQueries)
+				if (bIsOccluded || !bOcclusionStateIsDefinite)
+				{
+					bAllSubOcclusionStateIsDefinite = false;
+				}
+			}
+			else
+			{
+					
+				if (bIsOccluded)
+				{
+					View.PrimitiveVisibilityMap.AccessCorrespondingBit(BitIt) = false;
+					STAT(NumOccludedPrimitives++);
+				}
+				else if (bOcclusionStateIsDefinite)
 				{
-					FPrimitiveSceneProxy* Proxy = Scene->Primitives[BitIt.GetIndex()]->Proxy;
-					Proxy->AcceptOcclusionResults(&View, &SubIsOccluded[0], SubIsOccluded.Num());
-					if (bAllSubOccluded)
+					if (PrimitiveOcclusionHistory)
 					{
-						View.PrimitiveVisibilityMap.AccessCorrespondingBit(BitIt) = false;
-						STAT(NumOccludedPrimitives++);
+						PrimitiveOcclusionHistory->LastVisibleTime = CurrentRealTime;
 					}
-					else if (bAllSubOcclusionStateIsDefinite)
+					
+					View.PrimitiveDefinitelyUnoccludedMap.AccessCorrespondingBit(BitIt) = true;
+				}					
+			}			
+		}
+
+		if (bSubQueries)
+		{
+			FPrimitiveSceneProxy* Proxy = Scene->Primitives[BitIt.GetIndex()]->Proxy;
+			Proxy->AcceptOcclusionResults(&View, &SubIsOccluded[0], SubIsOccluded.Num());
+			if (bAllSubOccluded)
+			{
+				View.PrimitiveVisibilityMap.AccessCorrespondingBit(BitIt) = false;
+				STAT(NumOccludedPrimitives++);
+			}
+			else if (bAllSubOcclusionStateIsDefinite)
+			{
+				View.PrimitiveDefinitelyUnoccludedMap.AccessCorrespondingBit(BitIt) = true;
+			}
+		}
+	}
+
+	check(NumTotalDefUnoccluded == View.PrimitiveDefinitelyUnoccludedMap.Num());
+	check(NumTotalPrims == View.PrimitiveVisibilityMap.Num());
+	check(!InsertPrimitiveOcclusionHistory || InsertPrimitiveOcclusionHistory->Num() <= ReserveAmount);
+	Params.NumOccludedPrims = NumOccludedPrimitives;	
+}
+
+class FetchVisibilityForPrimitivesTask
+{
+	FVisForPrimParams& Params;
+
+public:
+
+	FetchVisibilityForPrimitivesTask(FVisForPrimParams& InParams)
+		: Params(InParams)
+	{
+	}
+
+	FORCEINLINE TStatId GetStatId() const
+	{
+		RETURN_QUICK_DECLARE_CYCLE_STAT(FetchVisibilityForPrimitivesTask, STATGROUP_TaskGraphTasks);
+	}
+
+	ENamedThreads::Type GetDesiredThread()
+	{
+		return ENamedThreads::AnyThread;
+	}
+
+	static ESubsequentsMode::Type GetSubsequentsMode() { return ESubsequentsMode::TrackSubsequents; }
+
+	void DoTask(ENamedThreads::Type CurrentThread, const FGraphEventRef& MyCompletionGraphEvent)
+	{
+		FetchVisibilityForPrimitives_Range<false>(Params);
+	}
+};
+
+static int32 FetchVisibilityForPrimitives(const FScene* Scene, FViewInfo& View, const bool bSubmitQueries, const bool bHZBOcclusion)
+{
+	QUICK_SCOPE_CYCLE_COUNTER(STAT_FetchVisibilityForPrimitives);
+	FSceneViewState* ViewState = (FSceneViewState*)View.State;
+
+	if (GOcclusionCullParallelPrimFetch && GSupportsParallelOcclusionQueries)
+	{		
+		static const int32 MaxNumCullTasks = 4;
+		static const int32 ActualNumCullTasks = 4;
+		static const int32 NumOutputArrays = MaxNumCullTasks;
+		
+		FGraphEventRef TaskRefArray[NumOutputArrays];
+
+		//params for each task
+		FVisForPrimParams Params[NumOutputArrays];
+
+		//output arrays for each task
+		TArray<FPrimitiveOcclusionHistory> OutputOcclusionHistory[NumOutputArrays];
+		TArray<FPrimitiveOcclusionHistory*> OutQueriesToRelease[NumOutputArrays];
+		TArray<FHZBBound> OutHZBBounds[NumOutputArrays];
+		TArray<FOcclusionBounds> OutQueriesToRun[NumOutputArrays];						
+
+		//optionally balance the tasks by how the visible primitives are distributed in the array rather than just breaking up the array by range.
+		//should make the tasks more equal length.
+#if BALANCE_LOAD
+		int32 StartIndices[NumOutputArrays] = { 0 };
+		int32 ProcessRange[NumOutputArrays] = { 0 };
+		{
+			QUICK_SCOPE_CYCLE_COUNTER(STAT_FetchVisibilityForPrimitivesPreProcess);
+			int32 NumBitsSet = 0;
+			for (FSceneSetBitIterator BitIt(View.PrimitiveVisibilityMap); BitIt; ++BitIt, ++NumBitsSet)
+			{
+			}
+			
+			int32 BitsPerTask = NumBitsSet / ActualNumCullTasks;
+			int32 NumBitsForRange = 0;
+			int32 CurrentStartIndex = 0;
+			int32 RangeToSet = 0;
+
+			//accumulate set bits for each task until we reach the target, then set the start/end and move on.
+			for (FSceneSetBitIterator BitIt(View.PrimitiveVisibilityMap); BitIt && RangeToSet < (ActualNumCullTasks - 1); ++BitIt)
+			{
+				++NumBitsForRange;
+				if (NumBitsForRange == BitsPerTask)
+				{
+					StartIndices[RangeToSet] = CurrentStartIndex;
+					ProcessRange[RangeToSet] = NumBitsForRange;
+
+					++RangeToSet;
+					NumBitsForRange = 0;
+					CurrentStartIndex = BitIt.GetIndex() + 1;
+				}
+			}
+
+			//final range is the rest of the set bits, no matter how many there are.
+			StartIndices[ActualNumCullTasks - 1] = CurrentStartIndex;
+			ProcessRange[ActualNumCullTasks - 1] = NumBitsSet - (BitsPerTask * 3);
+		}
+#endif
+
+		const int32 NumPrims = View.PrimitiveVisibilityMap.Num();
+		const int32 NumPerTask = NumPrims / ActualNumCullTasks;
+		int32 StartIndex = 0;
+		int32 NumToProcess = NumPerTask;
+
+		FGraphEventArray TaskWaitArray;
+		int32 NumTasks = 0;
+		for (int32 i = 0; i < ActualNumCullTasks && (StartIndex < NumPrims); ++i, ++NumTasks)
+		{
+			NumToProcess = (i == (ActualNumCullTasks - 1)) ? (NumPrims - StartIndex) : NumPerTask;
+
+			Params[i].Init(
+				Scene,
+				&View,
+				nullptr,
+#if BALANCE_LOAD
+				StartIndices[i],
+				ProcessRange[i],
+#else
+				StartIndex,
+				NumToProcess,
+#endif
+				bSubmitQueries,
+				bHZBOcclusion,				
+				&OutputOcclusionHistory[i],
+				&OutQueriesToRelease[i],
+				&OutHZBBounds[i],
+				&OutQueriesToRun[i]				
+				);
+
+			TaskRefArray[i] = TGraphTask<FetchVisibilityForPrimitivesTask>::CreateTask(nullptr, ENamedThreads::AnyThread).ConstructAndDispatchWhenReady(Params[i]);			
+			TaskWaitArray.Add(TaskRefArray[i]);
+
+			StartIndex += NumToProcess;
+		}
+
+		const int32 NumBufferedFrames = FOcclusionQueryHelpers::GetNumBufferedFrames();
+		uint32 OcclusionFrameCounter = ViewState->OcclusionFrameCounter;
+		TSet<FPrimitiveOcclusionHistory, FPrimitiveOcclusionHistoryKeyFuncs>& ViewPrimitiveOcclusionHistory = ViewState->PrimitiveOcclusionHistorySet;
+		FRenderQueryPool& OcclusionQueryPool = ViewState->OcclusionQueryPool;
+		FHZBOcclusionTester& HZBOcclusionTests = ViewState->HZBOcclusionTests;		
+
+		int32 NumOccludedPrims = 0;
+		{
+			QUICK_SCOPE_CYCLE_COUNTER(STAT_FetchVisibilityForPrimitivesCombine);
+
+			//wait for them all so we don't start modifying the prim histories while the gather is running
+			FTaskGraphInterface::Get().WaitUntilTasksComplete(TaskWaitArray, ENamedThreads::RenderThread_Local);			
+
+#if QUERY_SANITY_CHECK
+			{
+				QUICK_SCOPE_CYCLE_COUNTER(STAT_FetchVisibilityForPrimitivesSanity);
+				TSet<int32> ReleaseQuerySet;
+				TSet<int32> RunQuerySet;
+				TSet<int32> MasterPrimsProcessed;
+				for (int32 i = 0; i < NumTasks; ++i)
+				{
+					bool bAlreadyIn = false;
+					for (auto ReleaseQueryIter = OutQueriesToRelease[i].CreateIterator(); ReleaseQueryIter; ++ReleaseQueryIter)
 					{
-						View.PrimitiveDefinitelyUnoccludedMap.AccessCorrespondingBit(BitIt) = true;
+						FPrimitiveOcclusionHistory* History = *ReleaseQueryIter;
+						ReleaseQuerySet.Add(History->PrimitiveId.PrimIDValue, &bAlreadyIn);
+						checkf(!bAlreadyIn, TEXT("Prim: %i double released query."), History->PrimitiveId.PrimIDValue);
 					}
+
+					for (auto RunQueriesIter = OutQueriesToRun[i].CreateIterator(); RunQueriesIter; ++RunQueriesIter)
+					{
+						FPrimitiveOcclusionHistory* History = RunQueriesIter->PrimitiveOcclusionHistory;
+						RunQuerySet.Add(History->PrimitiveId.PrimIDValue, &bAlreadyIn);
+						checkf(!bAlreadyIn, TEXT("Prim: %i double run query."), History->PrimitiveId.PrimIDValue);
+					}					
+				}
+			}
+#endif
+
+			//Add/Release query ops use stored PrimitiveHistory pointers. We must do ALL of these from all tasks before adding any new PrimitiveHistories to the view.
+			//Adding new histories to the view could cause the array to resize which would invalidate all the stored output pointers for the other operations.
+			for (int32 i = 0; i < NumTasks; ++i)
+			{
+				//HZB output
+				for (auto HZBBoundIter = OutHZBBounds[i].CreateIterator(); HZBBoundIter; ++HZBBoundIter)
+				{
+					HZBBoundIter->TargetHistory->HZBTestIndex = HZBOcclusionTests.AddBounds(HZBBoundIter->BoundsOrigin, HZBBoundIter->BoundsExtent);
+				}
+
+				//Manual query release handling
+				for (auto ReleaseQueryIter = OutQueriesToRelease[i].CreateIterator(); ReleaseQueryIter; ++ReleaseQueryIter)
+				{
+					FPrimitiveOcclusionHistory* History = *ReleaseQueryIter;
+					OcclusionQueryPool.ReleaseQuery(History->GetPastQuery(OcclusionFrameCounter, NumBufferedFrames));
+				}
+				
+				//New query batching
+				for (auto RunQueriesIter = OutQueriesToRun[i].CreateIterator(); RunQueriesIter; ++RunQueriesIter)
+				{
+					RunQueriesIter->PrimitiveOcclusionHistory->SetCurrentQuery(OcclusionFrameCounter,
+						RunQueriesIter->bGroupedQuery ?
+						View.GroupedOcclusionQueries.BatchPrimitive(RunQueriesIter->BoundsOrigin, RunQueriesIter->BoundsExtent) :
+						View.IndividualOcclusionQueries.BatchPrimitive(RunQueriesIter->BoundsOrigin, RunQueriesIter->BoundsExtent),
+						NumBufferedFrames
+						);
+				}
+			}
+
+			//now add new primitivie histories to the view. may resize the view's array.
+			for (int32 i = 0; i < NumTasks; ++i)
+			{								
+				const TArray<FPrimitiveOcclusionHistory>& NewHistoryArray = OutputOcclusionHistory[i];				
+				for (int32 HistoryIndex = 0; HistoryIndex < NewHistoryArray.Num(); ++HistoryIndex)
+				{
+					const FPrimitiveOcclusionHistory& CopySourceHistory = NewHistoryArray[HistoryIndex];
+					ViewPrimitiveOcclusionHistory.Add(CopySourceHistory);
+				}
+
+				//accumulate occluded prims across tasks
+				NumOccludedPrims += Params[i].NumOccludedPrims;
+			}
+		}
+		
+		return NumOccludedPrims;
+	}
+	else
+	{
+		FViewElementPDI OcclusionPDI(&View, NULL);
+		int32 StartIndex = 0;
+		int32 NumToProcess = View.PrimitiveVisibilityMap.Num();				
+		FVisForPrimParams Params(
+			Scene,
+			&View,
+			&OcclusionPDI,
+			StartIndex,
+			NumToProcess,
+			bSubmitQueries,
+			bHZBOcclusion,			
+			nullptr,
+			nullptr,
+			nullptr,
+			nullptr			
+			);
+
+		FetchVisibilityForPrimitives_Range<true>(Params);
+		return Params.NumOccludedPrims;
+	}
+}
+
+/**
+ * Cull occluded primitives in the view.
+ */
+static int32 OcclusionCull(FRHICommandListImmediate& RHICmdList, const FScene* Scene, FViewInfo& View)
+{
+	SCOPE_CYCLE_COUNTER(STAT_OcclusionCull);	
+
+	// INITVIEWS_TODO: This could be more efficient if broken up in to separate concerns:
+	// - What is occluded?
+	// - For which primitives should we render occlusion queries?
+	// - Generate occlusion query geometry.
+
+	int32 NumOccludedPrimitives = 0;
+	FSceneViewState* ViewState = (FSceneViewState*)View.State;
+	
+	// Disable HZB on OpenGL platforms to avoid rendering artefacts
+	// It can be forced on by setting HZBOcclusion to 2
+	bool bHZBOcclusion = (!IsOpenGLPlatform(GShaderPlatformForFeatureLevel[Scene->GetFeatureLevel()]) && GHZBOcclusion) || (GHZBOcclusion == 2);
+
+	// Use precomputed visibility data if it is available.
+	if (View.PrecomputedVisibilityData)
+	{
+		QUICK_SCOPE_CYCLE_COUNTER(STAT_LookupPrecomputedVisibility);
+
+		uint8 PrecomputedVisibilityFlags = EOcclusionFlags::CanBeOccluded | EOcclusionFlags::HasPrecomputedVisibility;
+		for (FSceneSetBitIterator BitIt(View.PrimitiveVisibilityMap); BitIt; ++BitIt)
+		{
+			if ((Scene->PrimitiveOcclusionFlags[BitIt.GetIndex()] & PrecomputedVisibilityFlags) == PrecomputedVisibilityFlags)
+			{
+				FPrimitiveVisibilityId VisibilityId = Scene->PrimitiveVisibilityIds[BitIt.GetIndex()];
+				if ((View.PrecomputedVisibilityData[VisibilityId.ByteIndex] & VisibilityId.BitMask) == 0)
+				{
+					View.PrimitiveVisibilityMap.AccessCorrespondingBit(BitIt) = false;
+					INC_DWORD_STAT_BY(STAT_StaticallyOccludedPrimitives,1);
+					STAT(NumOccludedPrimitives++);
 				}
 			}
+		}
+	}
+
+	float CurrentRealTime = View.Family->CurrentRealTime;
+	if (ViewState)
+	{
+		if (Scene->GetFeatureLevel() >= ERHIFeatureLevel::SM4)
+		{
+			
+			bool bSubmitQueries = !View.bDisableQuerySubmissions;
+#if !(UE_BUILD_SHIPPING || UE_BUILD_TEST)
+			bSubmitQueries = bSubmitQueries && !ViewState->HasViewParent() && !ViewState->bIsFrozen;
+#endif
+
+			if( bHZBOcclusion )
+			{
+				QUICK_SCOPE_CYCLE_COUNTER(STAT_MapHZBResults);
+				check(!ViewState->HZBOcclusionTests.IsValidFrame(ViewState->OcclusionFrameCounter));
+				ViewState->HZBOcclusionTests.MapResults(RHICmdList);
+			}
+			
+			NumOccludedPrimitives += FetchVisibilityForPrimitives(Scene, View, bSubmitQueries, bHZBOcclusion);			
 
 			if( bHZBOcclusion )
 			{
