Exposes the view space and the camera space in the transform material expressions

diff --git a/Engine/Shaders/AtmosphericFogShader.usf b/Engine/Shaders/AtmosphericFogShader.usf
--- a/Engine/Shaders/AtmosphericFogShader.usf
+++ b/Engine/Shaders/AtmosphericFogShader.usf
@@ -57,7 +57,7 @@ void AtmosphericPixelMain(
 	float4 ScreenVector : TEXCOORD1,
 	out float4 OutColor : SV_Target0)
 {
-	OutColor = GetAtmosphericFog(View.ViewOrigin.xyz, ScreenVector.xyz, CalcSceneDepth(TexCoord), float3(0, 0, 0));
+	OutColor = GetAtmosphericFog(View.CameraWorldOrigin, ScreenVector.xyz, CalcSceneDepth(TexCoord), float3(0, 0, 0));
 #if !ATMOSPHERIC_NO_LIGHT_SHAFT
 	float LightShaftMask = Texture2DSample(OcclusionTexture, OcclusionTextureSampler, TexCoord).x;
 	OutColor.rgb = OutColor.rgb * LightShaftMask;
diff --git a/Engine/Shaders/BasePassForForwardShadingVertexShader.usf b/Engine/Shaders/BasePassForForwardShadingVertexShader.usf
--- a/Engine/Shaders/BasePassForForwardShadingVertexShader.usf
+++ b/Engine/Shaders/BasePassForForwardShadingVertexShader.usf
@@ -60,9 +60,9 @@ void Main(
 
 #if USE_VERTEX_FOG
 #if PACK_INTERPOLANTS
-	PackedInterps[0] = CalculateVertexHeightFog(WorldPosition.xyz, View.TranslatedViewOrigin);
+	PackedInterps[0] = CalculateVertexHeightFog(WorldPosition.xyz, float4(View.CameraTranslatedWorldOrigin, 1));
 #else
-	Output.BasePassInterpolants.VertexFog = CalculateVertexHeightFog(WorldPosition.xyz, View.TranslatedViewOrigin);
+	Output.BasePassInterpolants.VertexFog = CalculateVertexHeightFog(WorldPosition.xyz, float4(View.CameraTranslatedWorldOrigin, 1));
 #endif // PACK_INTERPOLANTS
 #endif
 
diff --git a/Engine/Shaders/BasePassVertexShader.usf b/Engine/Shaders/BasePassVertexShader.usf
--- a/Engine/Shaders/BasePassVertexShader.usf
+++ b/Engine/Shaders/BasePassVertexShader.usf
@@ -59,9 +59,9 @@ void Main(
 // Calculate the fog needed for translucency
 #if NEEDS_BASEPASS_FOGGING
 	#if BASEPASS_ATMOSPHERIC_FOG
-	Output.BasePassInterpolants.VertexFog = CalculateVertexAtmosphericFog(WorldPosition.xyz, View.TranslatedViewOrigin.xyz);
+	Output.BasePassInterpolants.VertexFog = CalculateVertexAtmosphericFog(WorldPosition.xyz, View.CameraTranslatedWorldOrigin);
 	#else
-	Output.BasePassInterpolants.VertexFog = CalculateVertexHeightFog(WorldPosition.xyz, View.TranslatedViewOrigin);
+	Output.BasePassInterpolants.VertexFog = CalculateVertexHeightFog(WorldPosition.xyz, float4(View.CameraTranslatedWorldOrigin, 1));
 	#endif
 #endif
 
diff --git a/Engine/Shaders/DeferredDecal.usf b/Engine/Shaders/DeferredDecal.usf
--- a/Engine/Shaders/DeferredDecal.usf
+++ b/Engine/Shaders/DeferredDecal.usf
@@ -95,7 +95,7 @@ float DistanceFunction(float4 ScreenPosition, float3 Pos)
 	MaterialParameters.WorldPosition_CamRelative = 
 	MaterialParameters.WorldPosition_NoOffsets = 
 	MaterialParameters.WorldPosition_NoOffsets_CamRelative = mul(DecalToWorld, float4(Pos, 1)).xyz;
-	MaterialParameters.CameraVector = normalize(View.ViewOrigin.xyz - MaterialParameters.WorldPosition);
+	MaterialParameters.CameraVector = normalize(View.CameraWorldOrigin - MaterialParameters.WorldPosition);
 
 	// the material needs to output the distance from "Opacity" depending on LightVector (in 0..1 range)
 	return GetMaterialMaskInputRaw(MaterialParameters);
@@ -192,7 +192,7 @@ void MainPS(
 
 		float4 WorldPosHom = mul(ProjectedScreenPosition, View.ScreenToWorld);
 
-		float3 WSStartRay = View.ViewOrigin.xyz;
+		float3 WSStartRay = View.CameraWorldOrigin;
 		float3 WSEndRay = WorldPosHom.xyz;
 
 		// in object/decal space
@@ -236,7 +236,7 @@ void MainPS(
 		WSPosition = WSPositionHom.xyz / WSPositionHom.w;
 	}
 
-	float3 CameraVector = normalize(View.ViewOrigin.xyz - WSPosition);
+	float3 CameraVector = normalize(View.CameraWorldOrigin - WSPosition);
 
 	// can be optimized
 	float3 DecalVector = OSPosition * 0.5f + 0.5f;
diff --git a/Engine/Shaders/DeferredLightPixelShaders.usf b/Engine/Shaders/DeferredLightPixelShaders.usf
--- a/Engine/Shaders/DeferredLightPixelShaders.usf
+++ b/Engine/Shaders/DeferredLightPixelShaders.usf
@@ -54,7 +54,7 @@ void DirectionalPixelMain(
 	BRANCH if( ScreenSpaceData.GBuffer.ShadingModelID > 0 )
 	{
 		float SceneDepth = CalcSceneDepth(InUV);
-		float3 WorldPosition = ScreenVector * SceneDepth + View.ViewOrigin.xyz;
+		float3 WorldPosition = ScreenVector * SceneDepth + View.CameraWorldOrigin;
 
 		FDeferredLightData LightData = SetupLightDataForStandardDeferred();
 		OutColor = GetDynamicLighting(WorldPosition, CameraVector, InUV, ScreenSpaceData, ScreenSpaceData.GBuffer.ShadingModelID, LightData, GetPerPixelLightAttenuation(InUV), uint2(0,0));
@@ -81,7 +81,7 @@ void RadialPixelMain(
 		float SceneDepth = CalcSceneDepth(ScreenUV);
 		float4 HomogeneousWorldPosition = mul(float4(InScreenPosition.xy / InScreenPosition.w * SceneDepth, SceneDepth, 1), View.ScreenToWorld);
 		float3 WorldPosition = HomogeneousWorldPosition.xyz / HomogeneousWorldPosition.w;
-		float3 CameraVector = normalize(WorldPosition - View.ViewOrigin.xyz);
+		float3 CameraVector = normalize(WorldPosition - View.CameraWorldOrigin);
 
 		FDeferredLightData LightData = SetupLightDataForStandardDeferred();
 
diff --git a/Engine/Shaders/DeferredLightingCommon.usf b/Engine/Shaders/DeferredLightingCommon.usf
--- a/Engine/Shaders/DeferredLightingCommon.usf
+++ b/Engine/Shaders/DeferredLightingCommon.usf
@@ -241,7 +241,7 @@ void GetShadowTerms(FScreenSpaceData ScreenSpaceData, FDeferredLightData LightDa
 		// Also fix up the fade between dynamic and static shadows
 		// to work with plane splits rather than spheres.
 
-		float DynamicShadowFraction = DistanceFromCameraFade(ScreenSpaceData, LightData, WorldPosition, View.ViewOrigin.xyz);
+		float DynamicShadowFraction = DistanceFromCameraFade(ScreenSpaceData, LightData, WorldPosition, View.CameraWorldOrigin);
 		// For a directional light, fade between static shadowing and the whole scene dynamic shadowing based on distance + per object shadows
 		OpaqueShadowTerm = lerp(LightAttenuation.x, StaticShadowing, DynamicShadowFraction);
 		// Fade between SSS dynamic shadowing and static shadowing based on distance
diff --git a/Engine/Shaders/DistanceFieldGlobalIllumination.usf b/Engine/Shaders/DistanceFieldGlobalIllumination.usf
--- a/Engine/Shaders/DistanceFieldGlobalIllumination.usf
+++ b/Engine/Shaders/DistanceFieldGlobalIllumination.usf
@@ -271,7 +271,7 @@ void CullVPLsForViewCS(
 	{
 		uint SourceIndex = VPLIndex * VPL_DATA_STRIDE;
 		float4 BoundingSphere = VPLData[SourceIndex + 0];
-		float DistanceToViewSq = dot(View.ViewOrigin.xyz - BoundingSphere.xyz, View.ViewOrigin.xyz - BoundingSphere.xyz);
+		float DistanceToViewSq = dot(View.CameraWorldOrigin - BoundingSphere.xyz, View.CameraWorldOrigin - BoundingSphere.xyz);
 
 		if (DistanceToViewSq < Square(AOMaxViewDistance + BoundingSphere.w)
 			&& ViewFrustumIntersectSphere(BoundingSphere.xyz, BoundingSphere.w + AOObjectMaxDistance))
@@ -511,7 +511,7 @@ float3 ComputeVPLFlux(uint SurfelIndex, float4x4 InstanceToWorld)
 
 uint GetViewBasedSurfelLOD(float3 ObjectPosition)
 {
-	return length(ObjectPosition - View.ViewOrigin.xyz) > .4f * AOMaxViewDistance ? 1 : 0;
+	return length(ObjectPosition - View.CameraWorldOrigin) > .4f * AOMaxViewDistance ? 1 : 0;
 	//return 0;
 }
 
diff --git a/Engine/Shaders/DistanceFieldSurfaceCacheLighting.usf b/Engine/Shaders/DistanceFieldSurfaceCacheLighting.usf
--- a/Engine/Shaders/DistanceFieldSurfaceCacheLighting.usf
+++ b/Engine/Shaders/DistanceFieldSurfaceCacheLighting.usf
@@ -201,8 +201,8 @@ void IrradianceCacheSplatVS(
 	float OccluderRadius = IrradianceCacheOccluderRadius[JobIndex];
 	RecordRadius = min(RecordRadius, OccluderRadius);
 
-	float ViewToCenterLength = length(RecordPosition - View.ViewOrigin.xyz);
-	float3 NormalizedViewToCenter = (RecordPosition - View.ViewOrigin.xyz) / ViewToCenterLength;
+	float ViewToCenterLength = length(RecordPosition - View.CameraWorldOrigin);
+	float3 NormalizedViewToCenter = (RecordPosition - View.CameraWorldOrigin) / ViewToCenterLength;
 
 	// Only allow full interpolation expanding for small samples, it won't be noticeable for large samples but adds a lot of splatting cost
 	float EffectiveInterpolationRadiusScale = lerp(InterpolationRadiusScale, 1, max(saturate(RecordRadius / 80), .2f));
@@ -254,7 +254,7 @@ void IrradianceCacheSplatVS(
 #if FINAL_INTERPOLATION_PASS
 	Output.BentNormal = IrradianceCacheBentNormal[JobIndex * BENT_NORMAL_STRIDE].xyz;
 	#if SUPPORT_IRRADIANCE
-		float ViewDistance = length(View.ViewOrigin.xyz - RecordPosition);
+		float ViewDistance = length(View.CameraWorldOrigin - RecordPosition);
 		float FadeAlpha = saturate(.001f * (.95f * AOMaxViewDistance - ViewDistance));
 		Output.Irradiance = IrradianceCacheIrradiance[JobIndex].xyz * FadeAlpha;
 	#endif
@@ -289,7 +289,7 @@ void IrradianceCacheSplatPS(
 
 #if USE_SCREENVECTOR_WORLD_POSITION
 
-	float3 OpaqueWorldPosition = View.ViewOrigin.xyz + Input.ScreenVector * SceneDepth;
+	float3 OpaqueWorldPosition = View.CameraWorldOrigin + Input.ScreenVector * SceneDepth;
 
 #else
 
@@ -515,8 +515,8 @@ void VisualizeMeshDistanceFieldCS(
 	float3 OpaqueWorldPosition = HomogeneousWorldPosition.xyz / HomogeneousWorldPosition.w;
 
 	float TraceDistance = 40000;
-	float3 WorldRayStart = View.ViewOrigin.xyz;
-	float3 WorldRayEnd = WorldRayStart + normalize(OpaqueWorldPosition - View.ViewOrigin.xyz) * TraceDistance;
+	float3 WorldRayStart = View.CameraWorldOrigin;
+	float3 WorldRayEnd = WorldRayStart + normalize(OpaqueWorldPosition - View.CameraWorldOrigin) * TraceDistance;
 	float3 WorldRayDirection = WorldRayEnd - WorldRayStart;
 	float3 UnitWorldRayDirection = normalize(WorldRayDirection);
 
diff --git a/Engine/Shaders/DistanceFieldSurfaceCacheLightingCompute.usf b/Engine/Shaders/DistanceFieldSurfaceCacheLightingCompute.usf
--- a/Engine/Shaders/DistanceFieldSurfaceCacheLightingCompute.usf
+++ b/Engine/Shaders/DistanceFieldSurfaceCacheLightingCompute.usf
@@ -196,7 +196,7 @@ void CullObjectsForViewCS(
 		uint SourceIndex = ObjectIndex;
 
 		float4 ObjectBoundingSphere = float4(ObjectBounds[4 * SourceIndex + 0], ObjectBounds[4 * SourceIndex + 1], ObjectBounds[4 * SourceIndex + 2], ObjectBounds[4 * SourceIndex + 3]);
-		float DistanceToViewSq = dot(View.ViewOrigin.xyz - ObjectBoundingSphere.xyz, View.ViewOrigin.xyz - ObjectBoundingSphere.xyz);
+		float DistanceToViewSq = dot(View.CameraWorldOrigin - ObjectBoundingSphere.xyz, View.CameraWorldOrigin - ObjectBoundingSphere.xyz);
 
 		if (DistanceToViewSq < Square(AOMaxViewDistance + ObjectBoundingSphere.w)
 			&& ViewFrustumIntersectSphere(ObjectBoundingSphere.xyz, ObjectBoundingSphere.w + AOObjectMaxDistance))
diff --git a/Engine/Shaders/ForEachLight.usf b/Engine/Shaders/ForEachLight.usf
--- a/Engine/Shaders/ForEachLight.usf
+++ b/Engine/Shaders/ForEachLight.usf
@@ -17,7 +17,7 @@ float3 SimplePointLightDiffuse( FScreenSpaceData ScreenSpaceData )
 
 float3 LightingForOneLight(FMaterialPixelParameters Parameters, uint LightIndex, FScreenSpaceData ScreenSpaceData)
 {
-	float3 CameraVector = normalize(Parameters.WorldPosition - View.ViewOrigin.xyz);
+	float3 CameraVector = normalize(Parameters.WorldPosition - View.CameraWorldOrigin);
 	
 	FDeferredLightData LightData = (FDeferredLightData)0;
 	LightData.LightPositionAndInvRadius = ForwardLightData.LightPositionAndInvRadius[LightIndex];
diff --git a/Engine/Shaders/GpuSkinVertexFactory.usf b/Engine/Shaders/GpuSkinVertexFactory.usf
--- a/Engine/Shaders/GpuSkinVertexFactory.usf
+++ b/Engine/Shaders/GpuSkinVertexFactory.usf
@@ -640,7 +640,7 @@ FMaterialTessellationParameters GetMaterialTessellationParameters(FVertexFactory
 	Result.TangentToWorldPreScale = 1;
 
 
-	Result.WorldPosition = CameraLocalWorldPosition + View.ViewOrigin.xyz;
+	Result.WorldPosition = CameraLocalWorldPosition + View.CameraWorldOrigin;
 
 	return Result;
 }
diff --git a/Engine/Shaders/HZBOcclusion.usf b/Engine/Shaders/HZBOcclusion.usf
--- a/Engine/Shaders/HZBOcclusion.usf
+++ b/Engine/Shaders/HZBOcclusion.usf
@@ -121,7 +121,7 @@ void HZBTestPS(float2 InUV : TEXCOORD0, out float4 OutColor : SV_Target0)
 	}
 
 #if 0
-	float3 ToView = View.TranslatedViewOrigin.xyz - BoundsCenter.xyz;
+	float3 ToView = View.CameraTranslatedWorldOrigin - BoundsCenter.xyz;
 	float3 Sign = ToView > 0 ? float3( 1, 1, 1 ) : float3( -1, -1, -1 );
 
 #if 0
@@ -153,9 +153,9 @@ void HZBTestPS(float2 InUV : TEXCOORD0, out float4 OutColor : SV_Target0)
 	);
 	
 	// Backface cull. Can't avoid evaluating the plane so replace it with the near plane.
-	FrontPlaneX = dot( FrontPlaneX, View.TranslatedViewOrigin ) > 0.0 ? FrontPlaneX : NearPlane;
-	FrontPlaneY = dot( FrontPlaneY, View.TranslatedViewOrigin ) > 0.0 ? FrontPlaneY : NearPlane;
-	FrontPlaneZ = dot( FrontPlaneZ, View.TranslatedViewOrigin ) > 0.0 ? FrontPlaneZ : NearPlane;
+	FrontPlaneX = dot( FrontPlaneX, float4(View.CameraTranslatedWorldOrigin, 1) ) > 0.0 ? FrontPlaneX : NearPlane;
+	FrontPlaneY = dot( FrontPlaneY, float4(View.CameraTranslatedWorldOrigin, 1) ) > 0.0 ? FrontPlaneY : NearPlane;
+	FrontPlaneZ = dot( FrontPlaneZ, float4(View.CameraTranslatedWorldOrigin, 1) ) > 0.0 ? FrontPlaneZ : NearPlane;
 	// Back planes don't need culling.
 
 	// Transform to post projective space. Inverse transpose for plane projection.
diff --git a/Engine/Shaders/HeightFogPixelShader.usf b/Engine/Shaders/HeightFogPixelShader.usf
--- a/Engine/Shaders/HeightFogPixelShader.usf
+++ b/Engine/Shaders/HeightFogPixelShader.usf
@@ -20,9 +20,9 @@ void ExponentialPixelMain(
 	float SceneDepth = CalcSceneDepth(TexCoord);
 
 	float3 CameraToReceiver = ScreenVector.xyz * SceneDepth;
-	float3 ReceiverPosition = CameraToReceiver + View.ViewOrigin.xyz; 
+	float3 ReceiverPosition = CameraToReceiver + View.CameraWorldOrigin; 
 
 	float LightShaftMask = Texture2DSample(OcclusionTexture, OcclusionSampler, TexCoord).x;
-	float4 ExponentialFog = GetExponentialHeightFog(ReceiverPosition, View.ViewOrigin.xyz);
+	float4 ExponentialFog = GetExponentialHeightFog(ReceiverPosition, View.CameraWorldOrigin);
 	OutColor = RETURN_COLOR(float4(ExponentialFog.rgb * LightShaftMask, ExponentialFog.w));
 }
diff --git a/Engine/Shaders/LandscapeVertexFactory.usf b/Engine/Shaders/LandscapeVertexFactory.usf
--- a/Engine/Shaders/LandscapeVertexFactory.usf
+++ b/Engine/Shaders/LandscapeVertexFactory.usf
@@ -277,7 +277,7 @@ FMaterialTessellationParameters GetMaterialTessellationParameters(FVertexFactory
 
 	Result.VertexColor = 1;
 	Result.TangentToWorldPreScale = 1;
-	Result.WorldPosition = CameraLocalWorldPosition + View.ViewOrigin.xyz;
+	Result.WorldPosition = CameraLocalWorldPosition + View.CameraWorldOrigin;
 
 	return Result;
 }
diff --git a/Engine/Shaders/LightFunctionPixelShader.usf b/Engine/Shaders/LightFunctionPixelShader.usf
--- a/Engine/Shaders/LightFunctionPixelShader.usf
+++ b/Engine/Shaders/LightFunctionPixelShader.usf
@@ -26,7 +26,7 @@ void Main(
 	float4 HomogeneousWorldPosition = mul(ProjectedScreenPosition, View.ScreenToWorld);
 	float3 WorldPosition = HomogeneousWorldPosition.xyz / HomogeneousWorldPosition.w;
 	// Calculate radial view distance for stable fading
-	float ViewDistance = length(View.ViewOrigin.xyz - WorldPosition);
+	float ViewDistance = length(View.CameraWorldOrigin - WorldPosition);
 
 	float3 Color = GetLightFunctionColor(LightVector.xyz, WorldPosition);
 
diff --git a/Engine/Shaders/LightShaftShader.usf b/Engine/Shaders/LightShaftShader.usf
--- a/Engine/Shaders/LightShaftShader.usf
+++ b/Engine/Shaders/LightShaftShader.usf
@@ -126,7 +126,7 @@ void DownsampleLightShaftsPixelMain(
 
 	#if POINT_LIGHT_SHAFTS
 		float3 CameraVector = ScreenVector * SceneDepth;
-		float3 WorldLightVector = CameraVector + View.ViewOrigin.xyz - WorldSpaceBlurOriginAndRadius.xyz;
+		float3 WorldLightVector = CameraVector + View.CameraWorldOrigin - WorldSpaceBlurOriginAndRadius.xyz;
 
 		// 1 at the radius of the light, 0 at the center of the light
 		float Distance = length(WorldLightVector);
diff --git a/Engine/Shaders/LocalVertexFactory.usf b/Engine/Shaders/LocalVertexFactory.usf
--- a/Engine/Shaders/LocalVertexFactory.usf
+++ b/Engine/Shaders/LocalVertexFactory.usf
@@ -670,7 +670,7 @@ float4 VertexFactoryGetPreviousWorldPosition(FVertexFactoryInput Input, FVertexF
 
 		Result.TangentToWorldPreScale = 1;
 
-		Result.WorldPosition = CameraLocalWorldPosition + View.ViewOrigin.xyz;
+		Result.WorldPosition = CameraLocalWorldPosition + View.CameraWorldOrigin;
 
 		return Result;
 	}
diff --git a/Engine/Shaders/MaterialTemplate.usf b/Engine/Shaders/MaterialTemplate.usf
--- a/Engine/Shaders/MaterialTemplate.usf
+++ b/Engine/Shaders/MaterialTemplate.usf
@@ -639,9 +639,9 @@ float4 MaterialExpressionAtmosphericFog(FMaterialPixelParameters Parameters, flo
 {
 #if MATERIAL_ATMOSPHERIC_FOG
 	// WorldPosition default value is Parameters.WorldPosition if not overridden by the user
-	float3 ViewVector = WorldPosition - View.ViewOrigin.xyz;
+	float3 ViewVector = WorldPosition - View.CameraWorldOrigin;
 	float SceneDepth = length(ViewVector);
-	return GetAtmosphericFog(View.ViewOrigin.xyz, ViewVector, SceneDepth, float3(0.f, 0.f, 0.f));
+	return GetAtmosphericFog(View.CameraWorldOrigin, ViewVector, SceneDepth, float3(0.f, 0.f, 0.f));
 #else
 	return float4(0.f, 0.f, 0.f, 0.f);
 #endif
@@ -864,7 +864,7 @@ float3 GetSpeedTreeVertexOffsetInner(FMaterialVertexParameters Parameters, int G
 		const float ScreenMultiple = 0.5 * max(View.ViewSizeAndInvSize.x * View.ViewToClip[0][0],
 												View.ViewSizeAndInvSize.y * View.ViewToClip[1][1]);
 		const float ScreenRadius = ScreenMultiple * Primitive.ObjectWorldPositionAndRadius.w / 
-									max(1.0, dot(TreePos - View.ViewOrigin.xyz, View.ViewForward.xyz));
+									max(1.0, dot(TreePos - View.CameraWorldOrigin, View.ViewForward.xyz));
 		const float ScreenArea = PI * ScreenRadius * ScreenRadius * View.ViewSizeAndInvSize.z * View.ViewSizeAndInvSize.w;
 		LodInterp = saturate((ScreenArea - SpeedTreeLODInfo.x) / SpeedTreeLODInfo.z);
 	}
diff --git a/Engine/Shaders/MeshParticleVertexFactory.usf b/Engine/Shaders/MeshParticleVertexFactory.usf
--- a/Engine/Shaders/MeshParticleVertexFactory.usf
+++ b/Engine/Shaders/MeshParticleVertexFactory.usf
@@ -489,7 +489,7 @@ float4 VertexFactoryGetPreviousWorldPosition(FVertexFactoryInput Input, FVertexF
 		half4 TangentToWorld2 = Interpolants.InterpolantsVSToPS.TangentToWorld2;
 		Result.TangentToWorld = AssembleTangentToWorld( TangentToWorld0, TangentToWorld2 );
 		Result.TangentToWorldPreScale = 1;
-		Result.WorldPosition = CameraLocalWorldPosition + View.ViewOrigin.xyz;
+		Result.WorldPosition = CameraLocalWorldPosition + View.CameraWorldOrigin;
 
 		return Result;
 	}
diff --git a/Engine/Shaders/ParticleBeamTrailVertexFactory.usf b/Engine/Shaders/ParticleBeamTrailVertexFactory.usf
--- a/Engine/Shaders/ParticleBeamTrailVertexFactory.usf
+++ b/Engine/Shaders/ParticleBeamTrailVertexFactory.usf
@@ -125,7 +125,7 @@ void GetTangents(FVertexFactoryInput Input,out float4 Right,out float4 Up)
 	float4	Position			= TransformLocalToTranslatedWorld(Input.Position.xyz),
 			OldPosition			= TransformLocalToTranslatedWorld(Input.OldPosition.xyz);
 
-	float3	CameraDirection		= SafeNormalize(View.ViewOrigin.xyz - Position.xyz),
+	float3	CameraDirection		= SafeNormalize(View.CameraWorldOrigin - Position.xyz),
 			ParticleDirection	= SafeNormalize(Position.xyz - OldPosition.xyz);
 
 	float4	Right_Square		= BeamTrailVF.CameraRight,
diff --git a/Engine/Shaders/ParticleGPUSpriteVertexFactory.usf b/Engine/Shaders/ParticleGPUSpriteVertexFactory.usf
--- a/Engine/Shaders/ParticleGPUSpriteVertexFactory.usf
+++ b/Engine/Shaders/ParticleGPUSpriteVertexFactory.usf
@@ -152,7 +152,7 @@ void GetSpriteTangents(
 	float3 CameraUp = lerp(-View.ViewUp, EmitterDynamicUniforms.AxisLockUp.xyz, EmitterDynamicUniforms.AxisLockUp.w);
 
 	// Determine the vector from the particle to the camera and the particle's movement direction.
-	float3 CameraDirection	= SafeNormalize(View.TranslatedViewOrigin.xyz - ParticleWorldPosition.xyz);
+	float3 CameraDirection	= SafeNormalize(View.CameraTranslatedWorldOrigin - ParticleWorldPosition.xyz);
 	float3 RightVector = CameraRight.xyz;
 	float3 UpVector = CameraUp.xyz;
 
@@ -272,7 +272,7 @@ FVertexFactoryIntermediates GetVertexFactoryIntermediates( FVertexFactoryInput I
 
 	// Optional camera motion blur.
 #if USES_PARTICLE_MOTION_BLUR
-	float3 MblurVector = -View.ViewOriginDelta;
+	float3 MblurVector = -View.CameraWorldOriginDelta;
 	float3 MblurDirection = normalize(MblurVector + 0.0001f);
 	//float MblurFactor = saturate(dot(MblurDirection,normalize(VertexWorldOffset)) * 100000.0f);
 	float MblurFactor = clamp(dot(MblurDirection,normalize(VertexWorldOffset)) * 100000.0f,-1.0f,1.0f);
@@ -368,7 +368,7 @@ FMaterialPixelParameters GetMaterialPixelParameters( FVertexFactoryInterpolantsV
 #endif
 
 #if USE_SPHERICAL_NORMALS
-	Result.TangentToWorld = GetSphericalParticleNormal(PixelPosition.xyz + View.ViewOrigin.xyz, Interpolants.ParticlePositionAndSize.xyz, Interpolants.ParticlePositionAndSize.w);
+	Result.TangentToWorld = GetSphericalParticleNormal(PixelPosition.xyz + View.CameraWorldOrigin, Interpolants.ParticlePositionAndSize.xyz, Interpolants.ParticlePositionAndSize.w);
 #else
 	half3 TangentToWorld0 = Interpolants.TangentToWorld0.xyz;
 	half4 TangentToWorld2 = Interpolants.TangentToWorld2;
diff --git a/Engine/Shaders/ParticleSpriteVertexFactory.usf b/Engine/Shaders/ParticleSpriteVertexFactory.usf
--- a/Engine/Shaders/ParticleSpriteVertexFactory.usf
+++ b/Engine/Shaders/ParticleSpriteVertexFactory.usf
@@ -170,7 +170,7 @@ FMaterialPixelParameters GetMaterialPixelParameters(FVertexFactoryInterpolantsVS
 	float SubImageLerp = Interpolants.TangentToWorld0AndInterp_Sizer.w;
 
 	#if GENERATE_SPHERICAL_PARTICLE_NORMALS && USE_PARTICLE_POSITION
-		Result.TangentToWorld = GetSphericalParticleNormal(PixelPosition.xyz + View.ViewOrigin.xyz, Interpolants.ParticlePositionAndSize.xyz, Interpolants.ParticlePositionAndSize.w);
+		Result.TangentToWorld = GetSphericalParticleNormal(PixelPosition.xyz + View.CameraWorldOrigin, Interpolants.ParticlePositionAndSize.xyz, Interpolants.ParticlePositionAndSize.w);
 	#else
 		Result.TangentToWorld = AssembleTangentToWorld(TangentToWorld0.xyz, TangentToWorld2);
 	#endif
@@ -295,7 +295,7 @@ void GetTangents(FVertexFactoryInput Input, float3 WorldPosition, float3 OldWorl
 	float3 CameraUp = lerp(-View.ViewUp, SpriteVF.AxisLockUp.xyz, SpriteVF.AxisLockUp.w);
 
 	// Determine the vector from the particle to the camera and the particle's movement direction.
-	float3 CameraDirection		= SafeNormalize(View.ViewOrigin.xyz - WorldPosition.xyz);
+	float3 CameraDirection		= SafeNormalize(View.CameraWorldOrigin - WorldPosition.xyz);
 	float3 ParticleDirection	= SafeNormalize(WorldPosition.xyz - OldWorldPosition.xyz);
 
 	// Tangent vectors for PSA_Velocity.
diff --git a/Engine/Shaders/PostProcessLpvIndirect.usf b/Engine/Shaders/PostProcessLpvIndirect.usf
--- a/Engine/Shaders/PostProcessLpvIndirect.usf
+++ b/Engine/Shaders/PostProcessLpvIndirect.usf
@@ -83,7 +83,7 @@ void MainPS(in float4 UVAndScreenPos : TEXCOORD0, float4 SVPos: SV_POSITION, out
 	float3 ScreenVector = normalize(ScreenVectorU); 
 
 	float SceneDepth = CalcSceneDepth(UV);
-	float3 WorldPosition = ScreenVectorU * SceneDepth + View.ViewOrigin.xyz;
+	float3 WorldPosition = ScreenVectorU * SceneDepth + View.CameraWorldOrigin;
 
 	float3 grid = WorldToGrid( WorldPosition );
 
@@ -169,7 +169,7 @@ void DirectionalOcclusionPS(in float4 UVAndScreenPos : TEXCOORD0, float4 SVPos:
 		float3 ScreenVector = normalize(ScreenVectorU); 
 
 		float SceneDepth = CalcSceneDepth(UV);
-		float3 WorldPosition = ScreenVectorU * SceneDepth + View.ViewOrigin.xyz;
+		float3 WorldPosition = ScreenVectorU * SceneDepth + View.CameraWorldOrigin;
 
 		float3 grid = WorldToGrid( WorldPosition );
 
diff --git a/Engine/Shaders/PostProcessMotionBlur.usf b/Engine/Shaders/PostProcessMotionBlur.usf
--- a/Engine/Shaders/PostProcessMotionBlur.usf
+++ b/Engine/Shaders/PostProcessMotionBlur.usf
@@ -587,7 +587,7 @@ void MainPS(float4 UVAndScreenPos : TEXCOORD0, out float4 OutColor : SV_Target0)
 #if MOTION_BLUR_QUALITY == 0
 	// visualize motion blur
 	{
-		float3 AbsWorldPos = View.ViewOrigin.xyz + OffsetWorldPos;
+		float3 AbsWorldPos = View.CameraWorldOrigin + OffsetWorldPos;
 		float3 WorldCheckerboard = Compute3DCheckerBoard(AbsWorldPos * 0.02f) * 0.1f + Compute3DCheckerBoard(AbsWorldPos * 0.002f) * 0.3f + Compute3DCheckerBoard(AbsWorldPos * 0.0002f) * 0.6f;
 
 		OutColor = float4(lerp(WorldCheckerboard, OutColor.rgb, 0.7f), 1);
diff --git a/Engine/Shaders/ReflectionEnvironmentComputeShaders.usf b/Engine/Shaders/ReflectionEnvironmentComputeShaders.usf
--- a/Engine/Shaders/ReflectionEnvironmentComputeShaders.usf
+++ b/Engine/Shaders/ReflectionEnvironmentComputeShaders.usf
@@ -607,7 +607,7 @@ void ReflectionEnvironmentTiledDeferredMain(
 	
 	float4 HomogeneousWorldPosition = mul(float4(ScreenPosition * SceneDepth, SceneDepth, 1), View.ScreenToWorld);
 	float3 WorldPosition = HomogeneousWorldPosition.xyz / HomogeneousWorldPosition.w;
-	float3 CameraToPixel = normalize(WorldPosition - View.ViewOrigin.xyz);
+	float3 CameraToPixel = normalize(WorldPosition - View.CameraWorldOrigin);
 	float3 ReflectionVector = reflect(CameraToPixel, GBuffer.WorldNormal);
 	float IndirectIrradiance = GBuffer.IndirectIrradiance;
 	
diff --git a/Engine/Shaders/ReflectionEnvironmentShaders.usf b/Engine/Shaders/ReflectionEnvironmentShaders.usf
--- a/Engine/Shaders/ReflectionEnvironmentShaders.usf
+++ b/Engine/Shaders/ReflectionEnvironmentShaders.usf
@@ -51,7 +51,7 @@ void CopySceneColorToCubeFaceColorPS(
 
 	// Convert INF's to valid values
 	SceneColor = min(SceneColor, 65503);
-	float3 WorldPosition = Input.ScreenVector * SceneDepth + View.ViewOrigin.xyz;
+	float3 WorldPosition = Input.ScreenVector * SceneDepth + View.CameraWorldOrigin;
 
 	float Alpha = 1;
 
@@ -64,7 +64,7 @@ void CopySceneColorToCubeFaceColorPS(
 	}
 	else if (SkyLightCaptureParameters.x == 1)
 	{
-		float RadialDistance = length(WorldPosition - View.ViewOrigin.xyz);
+		float RadialDistance = length(WorldPosition - View.CameraWorldOrigin);
 		float MaxDistance = SkyLightCaptureParameters.y;
 
 		// Setup alpha to fade out smoothly past the max distance
@@ -348,7 +348,7 @@ void StandardDeferredReflectionPS(
 		float SceneDepth = CalcSceneDepth(ScreenUV);
 		float4 HomogeneousWorldPosition = mul(float4(InScreenPosition.xy / InScreenPosition.w * SceneDepth, SceneDepth, 1), View.ScreenToWorld);
 		float3 WorldPosition = HomogeneousWorldPosition.xyz / HomogeneousWorldPosition.w;
-		float3 CameraToPixel = normalize(WorldPosition - View.ViewOrigin.xyz);
+		float3 CameraToPixel = normalize(WorldPosition - View.CameraWorldOrigin);
 		float3 ReflectionVector = reflect(CameraToPixel, InGBufferData.WorldNormal);
 
 		float3 ProjectedCaptureVector = ReflectionVector;
diff --git a/Engine/Shaders/ScreenSpaceReflections.usf b/Engine/Shaders/ScreenSpaceReflections.usf
--- a/Engine/Shaders/ScreenSpaceReflections.usf
+++ b/Engine/Shaders/ScreenSpaceReflections.usf
@@ -125,7 +125,7 @@ void ScreenSpaceReflectionsPS(in float4 UVAndScreenPos : TEXCOORD0, out float4 O
 	const float3 N = GBuffer.WorldNormal;
 	const float SceneDepth = CalcSceneDepth(UV);
 	const float3 PositionTranslatedWorld = mul( float4( ScreenPos * SceneDepth, SceneDepth, 1 ), View.ScreenToTranslatedWorld ).xyz;
-	const float3 V = normalize(View.TranslatedViewOrigin.xyz - PositionTranslatedWorld);
+	const float3 V = normalize(View.CameraTranslatedWorldOrigin - PositionTranslatedWorld);
 
 	uint FrameRandom;
 	{
diff --git a/Engine/Shaders/Tessellation.usf b/Engine/Shaders/Tessellation.usf
--- a/Engine/Shaders/Tessellation.usf
+++ b/Engine/Shaders/Tessellation.usf
@@ -101,9 +101,9 @@ float4 CalculateCompositeTessellationFactors(float3 Control0, float3 Control1, f
 	float3 Edge1 = ( Control1 - Control2 );
 	float3 Edge2 = ( Control2 - Control0 );
 
-	float3 ToMidpoint0 = 0.5 * ( Control0 + Control1 ) - View.TranslatedViewOrigin.xyz;
-	float3 ToMidpoint1 = 0.5 * ( Control1 + Control2 ) - View.TranslatedViewOrigin.xyz;
-	float3 ToMidpoint2 = 0.5 * ( Control2 + Control0 ) - View.TranslatedViewOrigin.xyz;
+	float3 ToMidpoint0 = 0.5 * ( Control0 + Control1 ) - View.CameraTranslatedWorldOrigin;
+	float3 ToMidpoint1 = 0.5 * ( Control1 + Control2 ) - View.CameraTranslatedWorldOrigin;
+	float3 ToMidpoint2 = 0.5 * ( Control2 + Control0 ) - View.CameraTranslatedWorldOrigin;
 
 	// Use spherical projection instead of planer
 	float4 CompositeFactors = float4(
diff --git a/Engine/Shaders/TiledDeferredLightShaders.usf b/Engine/Shaders/TiledDeferredLightShaders.usf
--- a/Engine/Shaders/TiledDeferredLightShaders.usf
+++ b/Engine/Shaders/TiledDeferredLightShaders.usf
@@ -239,7 +239,7 @@ void TiledDeferredLightingMain(
 	// Lookup GBuffer properties once per pixel 
 	FScreenSpaceData ScreenSpaceData = GetScreenSpaceDataUint(PixelPos);
 	FGBufferData InGBufferData = ScreenSpaceData.GBuffer; 
-	float3 CameraVector = normalize(WorldPosition - View.ViewOrigin.xyz);
+	float3 CameraVector = normalize(WorldPosition - View.CameraWorldOrigin);
 		
 	BRANCH
 	if (InGBufferData.ShadingModelID != SHADINGMODELID_UNLIT)
diff --git a/Engine/Shaders/TranslucentLightInjectionShaders.usf b/Engine/Shaders/TranslucentLightInjectionShaders.usf
--- a/Engine/Shaders/TranslucentLightInjectionShaders.usf
+++ b/Engine/Shaders/TranslucentLightInjectionShaders.usf
@@ -204,7 +204,7 @@ float ComputeVolumeShadowing(float3 WorldPositionForLighting)
 
 		// fade shadows in the distance
 		#if !RADIAL_ATTENUATION 
-			float Depth = dot(WorldPositionForLighting - View.ViewOrigin.xyz, View.ViewForward);
+			float Depth = dot(WorldPositionForLighting - View.CameraWorldOrigin, View.ViewForward);
 			float DistanceFade = saturate(Depth * ShadowInjectParams.z + ShadowInjectParams.w);
 			DynamicShadowFactor = lerp(DynamicShadowFactor, 1.0f, DistanceFade * DistanceFade);
 		#endif
diff --git a/Engine/Shaders/VectorFieldVisualizationVertexFactory.usf b/Engine/Shaders/VectorFieldVisualizationVertexFactory.usf
--- a/Engine/Shaders/VectorFieldVisualizationVertexFactory.usf
+++ b/Engine/Shaders/VectorFieldVisualizationVertexFactory.usf
@@ -61,7 +61,7 @@ FVertexFactoryIntermediates GetVertexFactoryIntermediates( FVertexFactoryInput I
 	const float3 Vec = normalize(Force) * ForceScale;
 
 	// The base position for this vector.
-	const float3 BasePosition = mul(float4(VoxelUV.xyz,1), VectorFieldVis.VolumeToWorld).xyz - View.ViewOrigin.xyz;
+	const float3 BasePosition = mul(float4(VoxelUV.xyz,1), VectorFieldVis.VolumeToWorld).xyz - View.CameraWorldOrigin;
 	const float3 Position = BasePosition + float(Input.VertexId) * 
 		mul(float4(Vec,0), VectorFieldVis.VolumeToWorldNoScale).xyz;
 
diff --git a/Engine/Source/Runtime/Engine/Classes/Materials/MaterialExpressionTransform.h b/Engine/Source/Runtime/Engine/Classes/Materials/MaterialExpressionTransform.h
--- a/Engine/Source/Runtime/Engine/Classes/Materials/MaterialExpressionTransform.h
+++ b/Engine/Source/Runtime/Engine/Classes/Materials/MaterialExpressionTransform.h
@@ -17,9 +17,12 @@ enum EMaterialVectorCoordTransformSource
 	/** World space */
 	TRANSFORMSOURCE_World UMETA(DisplayName="World Space"),
 	
-	/** View space */
+	/** View space (differs from camera space in the shadow passes) */
 	TRANSFORMSOURCE_View UMETA(DisplayName="View Space"),
 
+	/** Camera space */
+	TRANSFORMSOURCE_Camera  UMETA(DisplayName="Camera Space"),
+
 	TRANSFORMSOURCE_MAX,
 };
 
@@ -35,9 +38,12 @@ enum EMaterialVectorCoordTransform
 	/** World space */
 	TRANSFORM_World UMETA(DisplayName="World Space"),
 	
-	/** View space */
+	/** View space (differs from camera space in the shadow passes) */
 	TRANSFORM_View UMETA(DisplayName="View Space"),
 
+	/** Camera space */
+	TRANSFORM_Camera  UMETA(DisplayName="Camera Space"),
+	
 	TRANSFORM_MAX,
 };
 
diff --git a/Engine/Source/Runtime/Engine/Classes/Materials/MaterialExpressionTransformPosition.h b/Engine/Source/Runtime/Engine/Classes/Materials/MaterialExpressionTransformPosition.h
--- a/Engine/Source/Runtime/Engine/Classes/Materials/MaterialExpressionTransformPosition.h
+++ b/Engine/Source/Runtime/Engine/Classes/Materials/MaterialExpressionTransformPosition.h
@@ -8,17 +8,20 @@
 UENUM()
 enum EMaterialPositionTransformSource
 {
-	/** Local Space */
+	/** Local space */
 	TRANSFORMPOSSOURCE_Local UMETA(DisplayName="Local Space"),
 	
-	/** Absolute World Space */
+	/** Absolute world space */
 	TRANSFORMPOSSOURCE_World UMETA(DisplayName="Absolute World Space"),
 	
-	/** Camera Relative World Space */
+	/** Camera relative world space */
 	TRANSFORMPOSSOURCE_TranslatedWorld  UMETA(DisplayName="Camera Relative World Space"),
 
-	/** View Space */
+	/** View space (differs from camera space in the shadow passes) */
 	TRANSFORMPOSSOURCE_View  UMETA(DisplayName="View Space"),
+
+	/** Camera space */
+	TRANSFORMPOSSOURCE_Camera  UMETA(DisplayName="Camera Space"),
 	
 	TRANSFORMPOSSOURCE_MAX,
 };
diff --git a/Engine/Source/Runtime/Engine/Classes/Materials/MaterialExpressionViewProperty.h b/Engine/Source/Runtime/Engine/Classes/Materials/MaterialExpressionViewProperty.h
--- a/Engine/Source/Runtime/Engine/Classes/Materials/MaterialExpressionViewProperty.h
+++ b/Engine/Source/Runtime/Engine/Classes/Materials/MaterialExpressionViewProperty.h
@@ -16,6 +16,8 @@ enum EMaterialExposedViewProperty
 	MEVP_TanHalfFieldOfView UMETA(DisplayName="TanHalfFieldOfView"),
 	/** Horizontal and vertical size of the view in pixels */
 	MEVP_ViewSize UMETA(DisplayName="ViewSize"),
+	/** Absolute world space view position (differs from the camera position in the shadow passes) */
+	MEVP_WorldSpaceViewPosition UMETA(DisplayName="ViewPosition (Absolute World Space)"),
 	/** Absolute world space camera position */
 	MEVP_WorldSpaceCameraPosition UMETA(DisplayName="CameraPosition (Absolute World Space)"),
 
diff --git a/Engine/Source/Runtime/Engine/Private/Materials/HLSLMaterialTranslator.h b/Engine/Source/Runtime/Engine/Private/Materials/HLSLMaterialTranslator.h
--- a/Engine/Source/Runtime/Engine/Private/Materials/HLSLMaterialTranslator.h
+++ b/Engine/Source/Runtime/Engine/Private/Materials/HLSLMaterialTranslator.h
@@ -1663,7 +1663,8 @@ protected:
 			{MEVP_FieldOfView, MCT_Float2, TEXT("View.<PREV>FieldOfViewWideAngles"), nullptr},
 			{MEVP_TanHalfFieldOfView, MCT_Float2, TEXT("Get<PREV>TanHalfFieldOfView()"), TEXT("Get<PREV>CotanHalfFieldOfView()")},
 			{MEVP_ViewSize, MCT_Float2, TEXT("View.ViewSizeAndInvSize.xy"), TEXT("View.ViewSizeAndInvSize.zw")},
-			{MEVP_WorldSpaceCameraPosition, MCT_Float3, TEXT("View.<PREV>ViewOrigin"), nullptr},
+			{MEVP_WorldSpaceViewPosition, MCT_Float3, TEXT("View.<PREV>ViewWorldOrigin"), nullptr},
+			{MEVP_WorldSpaceCameraPosition, MCT_Float3, TEXT("View.<PREV>CameraWorldOrigin"), nullptr},
 		};
 		static_assert((sizeof(ViewPropertyMetaArray) / sizeof(ViewPropertyMetaArray[0])) == MEVP_MAX, "incoherency between EMaterialExposedViewProperty and ViewPropertyMetaArray");
 
@@ -3319,10 +3320,14 @@ protected:
 						CodeStr = TEXT("<A>");
 					}
 				}
-				else if (DestCoordBasis == MCB_View)
+				else if (DestCoordBasis == MCB_Camera)
 				{
 					CodeStr = TEXT("mul(<A>, <MATRIX>(View.<PREV>TranslatedWorldToCameraView))");
 				}
+				else if (DestCoordBasis == MCB_View)
+				{
+					CodeStr = TEXT("mul(<A>, <MATRIX>(View.<PREV>TranslatedWorldToView))");
+				}
 				// else use MCB_World as intermediary basis
 				break;
 			}
@@ -3362,7 +3367,7 @@ protected:
 				IntermediaryBasis = MCB_TranslatedWorld;
 				break;
 			}
-			case MCB_View:
+			case MCB_Camera:
 			{
 				if (DestCoordBasis == MCB_TranslatedWorld)
 				{
@@ -3372,6 +3377,16 @@ protected:
 				IntermediaryBasis = MCB_TranslatedWorld;
 				break;
 			}
+			case MCB_View:
+			{
+				if (DestCoordBasis == MCB_TranslatedWorld)
+				{
+					CodeStr = TEXT("mul(<A>, <MATRIX>(View.<PREV>ViewToTranslatedWorld))");
+				}
+				// else use MCB_TranslatedWorld as intermediary basis
+				IntermediaryBasis = MCB_TranslatedWorld;
+				break;
+			}
 			default:
 				check(0);
 				break;
diff --git a/Engine/Source/Runtime/Engine/Private/Materials/MaterialExpressions.cpp b/Engine/Source/Runtime/Engine/Private/Materials/MaterialExpressions.cpp
--- a/Engine/Source/Runtime/Engine/Private/Materials/MaterialExpressions.cpp
+++ b/Engine/Source/Runtime/Engine/Private/Materials/MaterialExpressions.cpp
@@ -5340,10 +5340,11 @@ void UMaterialExpressionAbs::GetCaption(TArray<FString>& OutCaptions) const
 static EMaterialCommonBasis GetMaterialCommonBasis(EMaterialVectorCoordTransformSource X)
 {
 	static const EMaterialCommonBasis ConversionTable[TRANSFORMSOURCE_MAX] = {
-		MCB_Tangent,
-		MCB_Local,
-		MCB_World,
-		MCB_View,
+		MCB_Tangent,					// TRANSFORMSOURCE_Tangent
+		MCB_Local,						// TRANSFORMSOURCE_Local
+		MCB_World,						// TRANSFORMSOURCE_World
+		MCB_View,						// TRANSFORMSOURCE_View
+		MCB_Camera,						// TRANSFORMSOURCE_Camera
 	};
 	return ConversionTable[X];
 }
@@ -5351,10 +5352,11 @@ static EMaterialCommonBasis GetMaterialCommonBasis(EMaterialVectorCoordTransform
 static EMaterialCommonBasis GetMaterialCommonBasis(EMaterialVectorCoordTransform X)
 {
 	static const EMaterialCommonBasis ConversionTable[TRANSFORM_MAX] = {
-		MCB_Tangent,
-		MCB_Local,
-		MCB_World,
-		MCB_View,
+		MCB_Tangent,					// TRANSFORM_Tangent
+		MCB_Local,						// TRANSFORM_Local
+		MCB_World,						// TRANSFORM_World
+		MCB_View,						// TRANSFORM_View
+		MCB_Camera,						// TRANSFORM_Camera
 	};
 	return ConversionTable[X];
 }
@@ -5443,10 +5445,11 @@ UMaterialExpressionTransformPosition::UMaterialExpressionTransformPosition(const
 static EMaterialCommonBasis GetMaterialCommonBasis(EMaterialPositionTransformSource X)
 {
 	static const EMaterialCommonBasis ConversionTable[TRANSFORMPOSSOURCE_MAX] = {
-		MCB_Local,
-		MCB_World,
-		MCB_TranslatedWorld,
-		MCB_View,
+		MCB_Local,						// TRANSFORMPOSSOURCE_Local
+		MCB_World,						// TRANSFORMPOSSOURCE_World
+		MCB_TranslatedWorld,			// TRANSFORMPOSSOURCE_TranslatedWorld
+		MCB_View,						// TRANSFORMPOSSOURCE_View
+		MCB_Camera,						// TRANSFORMPOSSOURCE_Camera
 	};
 	return ConversionTable[X];
 }
diff --git a/Engine/Source/Runtime/Engine/Public/MaterialShared.h b/Engine/Source/Runtime/Engine/Public/MaterialShared.h
--- a/Engine/Source/Runtime/Engine/Public/MaterialShared.h
+++ b/Engine/Source/Runtime/Engine/Public/MaterialShared.h
@@ -93,6 +93,7 @@ enum EMaterialCommonBasis
 	MCB_TranslatedWorld,
 	MCB_World,
 	MCB_View,
+	MCB_Camera,
 	MCB_MAX,
 };
 
diff --git a/Engine/Source/Runtime/Engine/Public/SceneView.h b/Engine/Source/Runtime/Engine/Public/SceneView.h
--- a/Engine/Source/Runtime/Engine/Public/SceneView.h
+++ b/Engine/Source/Runtime/Engine/Public/SceneView.h
@@ -371,8 +371,9 @@ BEGIN_UNIFORM_BUFFER_STRUCT_WITH_CONSTRUCTOR(FViewUniformShaderParameters,ENGINE
 	DECLARE_UNIFORM_BUFFER_STRUCT_MEMBER_EX(FVector4,ViewRectMin, EShaderPrecisionModifier::Half)
 	DECLARE_UNIFORM_BUFFER_STRUCT_MEMBER(FVector4,ViewSizeAndInvSize)
 	DECLARE_UNIFORM_BUFFER_STRUCT_MEMBER(FVector4,BufferSizeAndInvSize)
-	DECLARE_UNIFORM_BUFFER_STRUCT_MEMBER(FVector4,ViewOrigin)
-	DECLARE_UNIFORM_BUFFER_STRUCT_MEMBER(FVector4,TranslatedViewOrigin)
+	DECLARE_UNIFORM_BUFFER_STRUCT_MEMBER(FVector,CameraWorldOrigin)
+	DECLARE_UNIFORM_BUFFER_STRUCT_MEMBER(FVector,CameraTranslatedWorldOrigin)
+	DECLARE_UNIFORM_BUFFER_STRUCT_MEMBER(FVector,ViewWorldOrigin)
 	// The exposure scale is just a scalar but needs to be a float4 to workaround a driver bug on IOS.
 	// After 4.2 we can put the workaround in the cross compiler.
 	DECLARE_UNIFORM_BUFFER_STRUCT_MEMBER_EX(FVector4,ExposureScale, EShaderPrecisionModifier::Half)
@@ -384,7 +385,7 @@ BEGIN_UNIFORM_BUFFER_STRUCT_WITH_CONSTRUCTOR(FViewUniformShaderParameters,ENGINE
 	DECLARE_UNIFORM_BUFFER_STRUCT_MEMBER(float, PrevFrameRealTime)
 	DECLARE_UNIFORM_BUFFER_STRUCT_MEMBER(FVector,PreViewTranslation)
 	DECLARE_UNIFORM_BUFFER_STRUCT_MEMBER_EX(float,OutOfBoundsMask, EShaderPrecisionModifier::Half)
-	DECLARE_UNIFORM_BUFFER_STRUCT_MEMBER(FVector,ViewOriginDelta)
+	DECLARE_UNIFORM_BUFFER_STRUCT_MEMBER(FVector,CameraWorldOriginDelta)
 	DECLARE_UNIFORM_BUFFER_STRUCT_MEMBER(float,CullingSign)
 	DECLARE_UNIFORM_BUFFER_STRUCT_MEMBER_EX(float,NearPlane, EShaderPrecisionModifier::Half)
 	DECLARE_UNIFORM_BUFFER_STRUCT_MEMBER(float,AdaptiveTessellationFactor)
@@ -425,7 +426,8 @@ BEGIN_UNIFORM_BUFFER_STRUCT_WITH_CONSTRUCTOR(FViewUniformShaderParameters,ENGINE
 	DECLARE_UNIFORM_BUFFER_STRUCT_MEMBER(FMatrix,PrevViewToTranslatedWorld)
 	DECLARE_UNIFORM_BUFFER_STRUCT_MEMBER(FMatrix,PrevTranslatedWorldToCameraView)
 	DECLARE_UNIFORM_BUFFER_STRUCT_MEMBER(FMatrix,PrevCameraViewToTranslatedWorld)
-	DECLARE_UNIFORM_BUFFER_STRUCT_MEMBER(FVector,PrevViewOrigin)
+	DECLARE_UNIFORM_BUFFER_STRUCT_MEMBER(FVector,PrevCameraWorldOrigin)
+	DECLARE_UNIFORM_BUFFER_STRUCT_MEMBER(FVector,PrevViewWorldOrigin)
 	DECLARE_UNIFORM_BUFFER_STRUCT_MEMBER(FVector,PrevPreViewTranslation)
 	DECLARE_UNIFORM_BUFFER_STRUCT_MEMBER(FMatrix,PrevInvViewProj)
 	DECLARE_UNIFORM_BUFFER_STRUCT_MEMBER(FMatrix,PrevScreenToTranslatedWorld)
diff --git a/Engine/Source/Runtime/Renderer/Private/SceneRendering.cpp b/Engine/Source/Runtime/Renderer/Private/SceneRendering.cpp
--- a/Engine/Source/Runtime/Renderer/Private/SceneRendering.cpp
+++ b/Engine/Source/Runtime/Renderer/Private/SceneRendering.cpp
@@ -479,8 +479,9 @@ TUniformBufferRef<FViewUniformShaderParameters> FViewInfo::CreateUniformBuffer(
 	ViewUniformShaderParameters.ViewRectMin = FVector4(ViewRect.Min.X, ViewRect.Min.Y, 0.0f, 0.0f);
 	ViewUniformShaderParameters.ViewSizeAndInvSize = FVector4(ViewRect.Width(), ViewRect.Height(), 1.0f / float(ViewRect.Width()), 1.0f / float(ViewRect.Height()));
 	ViewUniformShaderParameters.BufferSizeAndInvSize = FVector4(BufferSize.X, BufferSize.Y, InvBufferSizeX, InvBufferSizeY);
-	ViewUniformShaderParameters.ViewOrigin = ViewMatrices.ViewOrigin;
-	ViewUniformShaderParameters.TranslatedViewOrigin = ViewMatrices.ViewOrigin + ViewMatrices.PreViewTranslation;
+	ViewUniformShaderParameters.ViewWorldOrigin = EffectiveViewToTranslatedWorld.TransformPosition(FVector(0)) - ViewMatrices.PreViewTranslation;
+	ViewUniformShaderParameters.CameraWorldOrigin = ViewMatrices.ViewOrigin;
+	ViewUniformShaderParameters.CameraTranslatedWorldOrigin = ViewMatrices.ViewOrigin + ViewMatrices.PreViewTranslation;
 	ViewUniformShaderParameters.DiffuseOverrideParameter = LocalDiffuseOverrideParameter;
 	ViewUniformShaderParameters.SpecularOverrideParameter = SpecularOverrideParameter;
 	ViewUniformShaderParameters.NormalOverrideParameter = NormalOverrideParameter;
@@ -488,7 +489,7 @@ TUniformBufferRef<FViewUniformShaderParameters> FViewInfo::CreateUniformBuffer(
 	ViewUniformShaderParameters.PrevFrameGameTime = Family->CurrentWorldTime - Family->DeltaWorldTime;
 	ViewUniformShaderParameters.PrevFrameRealTime = Family->CurrentRealTime - Family->DeltaWorldTime;
 	ViewUniformShaderParameters.PreViewTranslation = ViewMatrices.PreViewTranslation;
-	ViewUniformShaderParameters.ViewOriginDelta = ViewMatrices.ViewOrigin - PrevViewMatrices.ViewOrigin;
+	ViewUniformShaderParameters.CameraWorldOriginDelta = ViewMatrices.ViewOrigin - PrevViewMatrices.ViewOrigin;
 	ViewUniformShaderParameters.CullingSign = bReverseCulling ? -1.0f : 1.0f;
 	ViewUniformShaderParameters.NearPlane = GNearClippingPlane;
 	ViewUniformShaderParameters.PrevProjection = PrevViewMatrices.ProjMatrix;
@@ -497,11 +498,18 @@ TUniformBufferRef<FViewUniformShaderParameters> FViewInfo::CreateUniformBuffer(
 	ViewUniformShaderParameters.PrevViewToClip = PrevViewMatrices.ProjMatrix;
 	ViewUniformShaderParameters.PrevClipToView = PrevViewMatrices.GetInvProjMatrix();
 	ViewUniformShaderParameters.PrevTranslatedWorldToClip = PrevViewMatrices.TranslatedViewProjectionMatrix;
+	// PrevTranslatedWorldToView is not set as same as TranslatedWorldToView
+	// (set by EffectiveTranslatedViewMatrix != ViewMatrices.TranslatedViewMatrix),
+	// but that is fine because that means there are going to be inconsistently 
+	// different only in the shadow pass that doesn't require previous frame computation.
 	ViewUniformShaderParameters.PrevTranslatedWorldToView = PrevViewMatrices.TranslatedViewMatrix;
 	ViewUniformShaderParameters.PrevViewToTranslatedWorld = ViewUniformShaderParameters.PrevTranslatedWorldToView.Inverse();
 	ViewUniformShaderParameters.PrevTranslatedWorldToCameraView = PrevViewMatrices.TranslatedViewMatrix;
 	ViewUniformShaderParameters.PrevCameraViewToTranslatedWorld = ViewUniformShaderParameters.PrevTranslatedWorldToCameraView.Inverse();
-	ViewUniformShaderParameters.PrevViewOrigin = PrevViewMatrices.ViewOrigin;
+	ViewUniformShaderParameters.PrevCameraWorldOrigin = PrevViewMatrices.ViewOrigin;
+	// previous view world origin is going to be needed only in the base pass or shadow pass
+	// therefore is same as previous camera world origin.
+	ViewUniformShaderParameters.PrevViewWorldOrigin = ViewUniformShaderParameters.PrevCameraWorldOrigin;
 	ViewUniformShaderParameters.PrevPreViewTranslation = PrevViewMatrices.PreViewTranslation;
 	// can be optimized
 	ViewUniformShaderParameters.PrevInvViewProj = PrevViewProjMatrix.Inverse();
diff --git a/Engine/Source/Runtime/SlateRHIRenderer/Private/SlateRHIRenderingPolicy.cpp b/Engine/Source/Runtime/SlateRHIRenderer/Private/SlateRHIRenderingPolicy.cpp
--- a/Engine/Source/Runtime/SlateRHIRenderer/Private/SlateRHIRenderingPolicy.cpp
+++ b/Engine/Source/Runtime/SlateRHIRenderer/Private/SlateRHIRenderingPolicy.cpp
@@ -482,8 +482,9 @@ static FSceneView& CreateSceneView( FSceneViewFamilyContext& ViewFamilyContext,
 	ViewUniformShaderParameters.ViewRectMin = FVector4(ViewRect.Min.X, ViewRect.Min.Y, 0.0f, 0.0f);
 	ViewUniformShaderParameters.ViewSizeAndInvSize = FVector4(ViewRect.Width(), ViewRect.Height(), 1.0f/ViewRect.Width(), 1.0f/ViewRect.Height() );
 	ViewUniformShaderParameters.BufferSizeAndInvSize = ViewUniformShaderParameters.ViewSizeAndInvSize;
-	ViewUniformShaderParameters.ViewOrigin = View->ViewMatrices.ViewOrigin;
-	ViewUniformShaderParameters.TranslatedViewOrigin = View->ViewMatrices.ViewOrigin + View->ViewMatrices.PreViewTranslation;
+	ViewUniformShaderParameters.ViewWorldOrigin = View->ViewMatrices.ViewOrigin;
+	ViewUniformShaderParameters.CameraWorldOrigin = View->ViewMatrices.ViewOrigin;
+	ViewUniformShaderParameters.CameraTranslatedWorldOrigin = ViewUniformShaderParameters.CameraWorldOrigin + View->ViewMatrices.PreViewTranslation;
 	ViewUniformShaderParameters.DiffuseOverrideParameter = View->DiffuseOverrideParameter;
 	ViewUniformShaderParameters.SpecularOverrideParameter = View->SpecularOverrideParameter;
 	ViewUniformShaderParameters.NormalOverrideParameter = View->NormalOverrideParameter;
