#ue4 - Optimized allocations in UpdatePrimitiveAttachment.

--------

diff --git a/Engine/Source/Runtime/Renderer/Private/RendererScene.cpp b/Engine/Source/Runtime/Renderer/Private/RendererScene.cpp
--- a/Engine/Source/Runtime/Renderer/Private/RendererScene.cpp
+++ b/Engine/Source/Runtime/Renderer/Private/RendererScene.cpp
@@ -715,26 +715,20 @@ void FScene::UpdatePrimitiveAttachment(UPrimitiveComponent* Primitive)
 	// Walk down the tree updating, because the scene's attachment data structures must be updated if the root of the attachment tree changes
 	while (ProcessStack.Num() > 0)
 	{
-		USceneComponent* Current = ProcessStack.Pop();
-		UPrimitiveComponent* CurrentPrimitive = Cast<UPrimitiveComponent>(Current);
-		check(Current);
-
-		if (CurrentPrimitive
-			&& CurrentPrimitive->GetWorld() 
-			&& CurrentPrimitive->GetWorld()->Scene 
-			&& CurrentPrimitive->GetWorld()->Scene == this
-			&& CurrentPrimitive->ShouldComponentAddToScene())
+		USceneComponent* Current = ProcessStack.Pop(/*bAllowShrinking=*/ false);
+		if (Current)
 		{
-			UpdatePrimitiveLightingAttachmentRoot(CurrentPrimitive);
-		}
+			UPrimitiveComponent* CurrentPrimitive = Cast<UPrimitiveComponent>(Current);
 
-		for (int32 ChildIndex = 0; ChildIndex < Current->AttachChildren.Num(); ChildIndex++)
-		{
-			USceneComponent* ChildComponent = Current->AttachChildren[ChildIndex];
-			if (ChildComponent)
+			if (CurrentPrimitive
+				&& CurrentPrimitive->GetWorld() 
+				&& CurrentPrimitive->GetWorld()->Scene == this
+				&& CurrentPrimitive->ShouldComponentAddToScene())
 			{
-				ProcessStack.Push(ChildComponent);
+				UpdatePrimitiveLightingAttachmentRoot(CurrentPrimitive);
 			}
+
+			ProcessStack.Append(Current->AttachChildren);
 		}
 	}
 }
