Merge sanity check removal. Fixes crash from GetRefCount not being threadsafe.

diff --git a/Engine/Source/Runtime/Renderer/Private/SceneVisibility.cpp b/Engine/Source/Runtime/Renderer/Private/SceneVisibility.cpp
--- a/Engine/Source/Runtime/Renderer/Private/SceneVisibility.cpp
+++ b/Engine/Source/Runtime/Renderer/Private/SceneVisibility.cpp
@@ -667,7 +667,7 @@ static void FetchVisibilityForPrimitives_Range(FVisForPrimParams& Params)
 						FRenderQueryRHIRef& PastQuery = PrimitiveOcclusionHistory->GetPastQuery(OcclusionFrameCounter, NumBufferedFrames);
 						if (IsValidRef(PastQuery))
 						{
-							int32 RefCount = PastQuery.GetRefCount();
+							//int32 RefCount = PastQuery.GetReference()->GetRefCount();
 							// NOTE: RHIGetOcclusionQueryResult should never fail when using a blocking call, rendering artifacts may show up.
 							//if (RHICmdList.GetRenderQueryResult(PastQuery, NumSamples, true))
 							if (GDynamicRHI->RHIGetRenderQueryResult(PastQuery.GetReference(), NumSamples, true))
@@ -699,8 +699,7 @@ static void FetchVisibilityForPrimitives_Range(FVisForPrimParams& Params)
 								// If the occlusion query failed, treat the primitive as visible.  
 								// already set bIsOccluded = false;
 							}
-
-							checkf(RefCount == PastQuery.GetRefCount(), TEXT("Ref count on prim: %i, old: %i, new: %i"), PrimitiveOcclusionHistory->PrimitiveId.PrimIDValue, RefCount, PastQuery.GetRefCount());
+							//checkf(RefCount == PastQuery.GetReference()->GetRefCount(), TEXT("Ref count on prim: %i, old: %i, new: %i"), PrimitiveOcclusionHistory->PrimitiveId.PrimIDValue, RefCount, PastQuery.GetReference()->GetRefCount());
 						}
 						else
 						{
