Removing unneeded telemetry

diff --git a/Engine/Source/Runtime/Renderer/Private/DeferredShadingRenderer.cpp b/Engine/Source/Runtime/Renderer/Private/DeferredShadingRenderer.cpp
--- a/Engine/Source/Runtime/Renderer/Private/DeferredShadingRenderer.cpp
+++ b/Engine/Source/Runtime/Renderer/Private/DeferredShadingRenderer.cpp
@@ -1640,9 +1640,6 @@ bool FDeferredShadingSceneRenderer::RenderPrePassHMD(FRHICommandListImmediate& R
 		return false;
 	}
 
-	SCOPED_DRAW_EVENT(RHICmdList, PrePassHMD);
-	SCOPE_CYCLE_COUNTER(STAT_DepthDrawTime);
-
 	FSceneRenderTargets& SceneContext = FSceneRenderTargets::Get(RHICmdList);
 	SceneContext.BeginRenderingPrePass(RHICmdList, true);
 
