UE4 - Fixed rare case of neededing to send end of frame updates before a lighting build.

diff --git a/Engine/Source/Runtime/Renderer/Private/ReflectionEnvironmentCapture.cpp b/Engine/Source/Runtime/Renderer/Private/ReflectionEnvironmentCapture.cpp
--- a/Engine/Source/Runtime/Renderer/Private/ReflectionEnvironmentCapture.cpp
+++ b/Engine/Source/Runtime/Renderer/Private/ReflectionEnvironmentCapture.cpp
@@ -1377,7 +1377,14 @@ void FScene::UpdateSkyCaptureContents(const USkyLightComponent* CaptureComponent
 	if (GSupportsRenderTargetFormat_PF_FloatRGBA || GetFeatureLevel() >= ERHIFeatureLevel::SM4)
 	{
 		QUICK_SCOPE_CYCLE_COUNTER(STAT_UpdateSkyCaptureContents);
-
+		{
+			World = GetWorld();
+			if (World)
+			{
+				//guarantee that all render proxies are up to date before kicking off this render
+				World->SendAllEndOfFrameUpdates();
+			}
+		}
 		ENQUEUE_UNIQUE_RENDER_COMMAND( 
 			ClearCommand,
 		{
