Merging changes that ensure that vertex declarations match the definitions used in the compiled shader states as required by Metal.

diff --git a/Engine/Source/Runtime/Renderer/Private/PostProcess/PostProcessMaterial.cpp b/Engine/Source/Runtime/Renderer/Private/PostProcess/PostProcessMaterial.cpp
--- a/Engine/Source/Runtime/Renderer/Private/PostProcess/PostProcessMaterial.cpp
+++ b/Engine/Source/Runtime/Renderer/Private/PostProcess/PostProcessMaterial.cpp
@@ -126,6 +126,30 @@ FRCPassPostProcessMaterial::FRCPassPostProcessMaterial(UMaterialInterface* InMat
 		MaterialInterface = UMaterial::GetDefaultMaterial(MD_PostProcess);
 	}
 }
+		
+/** The filter vertex declaration resource type. */
+class FPostProcessMaterialVertexDeclaration : public FRenderResource
+{
+public:
+	FVertexDeclarationRHIRef VertexDeclarationRHI;
+	
+	/** Destructor. */
+	virtual ~FPostProcessMaterialVertexDeclaration() {}
+	
+	virtual void InitRHI()
+	{
+		FVertexDeclarationElementList Elements;
+		uint32 Stride = sizeof(FFilterVertex);
+		Elements.Add(FVertexElement(0,STRUCT_OFFSET(FFilterVertex,Position),VET_Float4,0,Stride));
+		VertexDeclarationRHI = RHICreateVertexDeclaration(Elements);
+	}
+	
+	virtual void ReleaseRHI()
+	{
+		VertexDeclarationRHI.SafeRelease();
+	}
+};
+TGlobalResource<FPostProcessMaterialVertexDeclaration> GPostProcessMaterialVertexDeclaration;
 
 void FRCPassPostProcessMaterial::Process(FRenderingCompositePassContext& Context)
 {
@@ -182,7 +206,7 @@ void FRCPassPostProcessMaterial::Process(FRenderingCompositePassContext& Context
 	FPostProcessMaterialPS* PixelShader = MaterialShaderMap->GetShader<FPostProcessMaterialPS>();
 	FPostProcessMaterialVS* VertexShader = MaterialShaderMap->GetShader<FPostProcessMaterialVS>();
 
-	Context.RHICmdList.SetLocalBoundShaderState(Context.RHICmdList.BuildLocalBoundShaderState(GetVertexDeclarationFVector4(), VertexShader->GetVertexShader(), FHullShaderRHIRef(), FDomainShaderRHIRef(), PixelShader->GetPixelShader(), FGeometryShaderRHIRef()));
+	Context.RHICmdList.SetLocalBoundShaderState(Context.RHICmdList.BuildLocalBoundShaderState(GPostProcessMaterialVertexDeclaration.VertexDeclarationRHI, VertexShader->GetVertexShader(), FHullShaderRHIRef(), FDomainShaderRHIRef(), PixelShader->GetPixelShader(), FGeometryShaderRHIRef()));
 
 	VertexShader->SetParameters(Context.RHICmdList, Context);
 	PixelShader->SetParameters(Context.RHICmdList, Context, MaterialInterface->GetRenderProxy(false));
diff --git a/Engine/Source/Runtime/Renderer/Private/ReflectionEnvironment.cpp b/Engine/Source/Runtime/Renderer/Private/ReflectionEnvironment.cpp
--- a/Engine/Source/Runtime/Renderer/Private/ReflectionEnvironment.cpp
+++ b/Engine/Source/Runtime/Renderer/Private/ReflectionEnvironment.cpp
@@ -1046,7 +1046,7 @@ void FDeferredShadingSceneRenderer::RenderStandardDeferredImageBasedReflections(
 
 						static FGlobalBoundShaderState BoundShaderState;
 						
-						SetGlobalBoundShaderState(RHICmdList, FeatureLevel, BoundShaderState, GFilterVertexDeclaration.VertexDeclarationRHI, *VertexShader, *PixelShader);
+						SetGlobalBoundShaderState(RHICmdList, FeatureLevel, BoundShaderState, GetVertexDeclarationFVector4(), *VertexShader, *PixelShader);
 
 						PixelShader->SetParameters(RHICmdList, View, ReflectionCapture);
 					}
@@ -1056,7 +1056,7 @@ void FDeferredShadingSceneRenderer::RenderStandardDeferredImageBasedReflections(
 
 						static FGlobalBoundShaderState BoundShaderState;
 						
-						SetGlobalBoundShaderState(RHICmdList, FeatureLevel, BoundShaderState, GFilterVertexDeclaration.VertexDeclarationRHI, *VertexShader, *PixelShader);
+						SetGlobalBoundShaderState(RHICmdList, FeatureLevel, BoundShaderState, GetVertexDeclarationFVector4(), *VertexShader, *PixelShader);
 
 						PixelShader->SetParameters(RHICmdList, View, ReflectionCapture);
 					}
