Merging changes necessary to support layered rendering on platforms where layer specification is done in vertex shader as geometry shaders don't exist.

diff --git a/Engine/Shaders/AtmosphereCommon.usf b/Engine/Shaders/AtmosphereCommon.usf
--- a/Engine/Shaders/AtmosphereCommon.usf
+++ b/Engine/Shaders/AtmosphereCommon.usf
@@ -82,7 +82,12 @@ SamplerState AtmosphereInscatterTextureSampler;
 struct FAtmosphereVSOutput
 {
 	float2 OutTexCoord : TEXCOORD0;
-	float4 OutPosition : SV_POSITION;
+    float4 OutPosition : SV_POSITION;
+#if USING_VERTEX_SHADER_LAYER
+    uint LayerIndex : SV_RenderTargetArrayIndex;
+#else
+    uint LayerIndex : TEXCOORD1;
+#endif
 };
 
 float2 GetTransmittanceUV(float Radius, float Mu) 
diff --git a/Engine/Shaders/AtmospherePrecompute.usf b/Engine/Shaders/AtmospherePrecompute.usf
--- a/Engine/Shaders/AtmospherePrecompute.usf
+++ b/Engine/Shaders/AtmospherePrecompute.usf
@@ -21,6 +21,13 @@
 =============================================================================*/
 
 #include "Common.usf"
+
+// Disable USING_VERTEX_SHADER_LAYER as it would be an error for these shaders as they are intended to render to 2D targets only.
+#if USING_VERTEX_SHADER_LAYER != 0
+	#undef USING_VERTEX_SHADER_LAYER
+	#define USING_VERTEX_SHADER_LAYER 0
+#endif
+
 #include "AtmosphereCommon.usf"
 #include "AtmospherePrecomputeCommon.usf"
 
@@ -49,6 +56,7 @@ void MainVS(
 	)
 {
 	DrawRectangle(InPosition, InTexCoord, Out.OutPosition, Out.OutTexCoord );
+    Out.LayerIndex = 0;
 }
 
 // pixel shader entry point
diff --git a/Engine/Shaders/AtmospherePrecomputeInscatter.usf b/Engine/Shaders/AtmospherePrecomputeInscatter.usf
--- a/Engine/Shaders/AtmospherePrecomputeInscatter.usf
+++ b/Engine/Shaders/AtmospherePrecomputeInscatter.usf
@@ -36,12 +36,13 @@ void MainVS(
 {
 	Out.OutPosition = InPosition;
 	Out.OutTexCoord = InTexCoord;
+    Out.LayerIndex = AtmosphereLayer;
 }
 
 struct FAtmosphereGSOut
 {
-	FAtmosphereVSOutput Vertex;
-	uint LayerIndex : SV_RenderTargetArrayIndex;
+    FAtmosphereVSOutput Vertex;
+    uint LayerIndex : SV_RenderTargetArrayIndex;
 };
 
 [maxvertexcount(3)]
diff --git a/Engine/Shaders/Definitions.usf b/Engine/Shaders/Definitions.usf
--- a/Engine/Shaders/Definitions.usf
+++ b/Engine/Shaders/Definitions.usf
@@ -74,6 +74,10 @@
 #define COMPILER_GLSL_ES2									0
 #endif
 
+#ifndef COMPILER_METAL
+#define COMPILER_METAL									0
+#endif
+
 #ifndef COMPILER_SUPPORTS_ATTRIBUTES
 #define COMPILER_SUPPORTS_ATTRIBUTES					0
 #endif
@@ -102,6 +106,10 @@
 #define METAL_MRT_PROFILE								0
 #endif
 
+#ifndef METAL_SM4_PROFILE
+#define METAL_SM4_PROFILE								0
+#endif
+
 #ifndef METAL_SM5_PROFILE
 #define METAL_SM5_PROFILE								0
 #endif
@@ -211,8 +219,13 @@
 #define HAS_INVERTED_Z_BUFFER	1
 #endif
 
+/** When defined sets the layer index for volumetric rendering in the vertex shader, rather than using a geometry shader */
+#ifndef USING_VERTEX_SHADER_LAYER
+#define USING_VERTEX_SHADER_LAYER	0
+#endif
+
 // 'static' asserts
-#if COMPILER_GLSL || COMPILER_GLSL_ES2 || VULKAN_PROFILE || METAL_PROFILE || METAL_MRT_PROFILE || METAL_SM5_PROFILE
+#if COMPILER_GLSL || COMPILER_GLSL_ES2 || VULKAN_PROFILE || METAL_PROFILE || METAL_MRT_PROFILE || METAL_SM4_PROFILE || METAL_SM5_PROFILE
 #if !COMPILER_HLSLCC
 #error "Missing COMPILER_HLSLCC define!"
 #endif
diff --git a/Engine/Shaders/TranslucentLightInjectionShaders.usf b/Engine/Shaders/TranslucentLightInjectionShaders.usf
--- a/Engine/Shaders/TranslucentLightInjectionShaders.usf
+++ b/Engine/Shaders/TranslucentLightInjectionShaders.usf
@@ -45,7 +45,7 @@ float CalcSimpleLightAttenuation(float3 WorldPosition)
 void SimpleLightInjectMainPS(
 	FWriteToSliceGeometryOutput Input,
 	out float4 OutColor0 : SV_Target0,
-	out float3 OutColor1 : SV_Target1
+	out float4 OutColor1 : SV_Target1
 	)
 {
 	OutColor0 = 0;
@@ -75,7 +75,7 @@ void SimpleLightInjectMainPS(
 	float3 Coefficient0 = float3(SHLighting.R.V.y, SHLighting.G.V.y, SHLighting.B.V.y);
 	float3 Coefficient1 = float3(SHLighting.R.V.z, SHLighting.G.V.z, SHLighting.B.V.z);
 	float3 Coefficient2 = float3(SHLighting.R.V.w, SHLighting.G.V.w, SHLighting.B.V.w);
-	OutColor1 = float3(dot(Coefficient0, LuminanceWeights), dot(Coefficient1, LuminanceWeights), dot(Coefficient2, LuminanceWeights));
+	OutColor1 = float4(dot(Coefficient0, LuminanceWeights), dot(Coefficient1, LuminanceWeights), dot(Coefficient2, LuminanceWeights), 0);
 }
 
 #if INJECTION_PIXEL_SHADER
@@ -221,7 +221,7 @@ float ComputeVolumeShadowing(float3 WorldPositionForLighting)
 void InjectMainPS(
 	FWriteToSliceGeometryOutput Input,
 	out float4 OutColor0 : SV_Target0,
-	out float3 OutColor1 : SV_Target1
+	out float4 OutColor1 : SV_Target1
 	)
 {
 	OutColor0 = 0;
@@ -322,7 +322,7 @@ void InjectMainPS(
 		float3 Coefficient0 = float3(SHLighting.R.V.y, SHLighting.G.V.y, SHLighting.B.V.y);
 		float3 Coefficient1 = float3(SHLighting.R.V.z, SHLighting.G.V.z, SHLighting.B.V.z);
 		float3 Coefficient2 = float3(SHLighting.R.V.w, SHLighting.G.V.w, SHLighting.B.V.w);
-		OutColor1 = float3(dot(Coefficient0, LuminanceWeights), dot(Coefficient1, LuminanceWeights), dot(Coefficient2, LuminanceWeights));
+		OutColor1 = float4(dot(Coefficient0, LuminanceWeights), dot(Coefficient1, LuminanceWeights), dot(Coefficient2, LuminanceWeights), 0);
 	}
 
 	// debug, make inner cascase green
diff --git a/Engine/Shaders/TranslucentLightingShaders.usf b/Engine/Shaders/TranslucentLightingShaders.usf
--- a/Engine/Shaders/TranslucentLightingShaders.usf
+++ b/Engine/Shaders/TranslucentLightingShaders.usf
@@ -1,5 +1,6 @@
 // Copyright 1998-2015 Epic Games, Inc. All Rights Reserved.
 
+
 /**
  * TranslucentLightingShaders.usf: Shaders for calculating lighting in a volume to use on translucency
  */
@@ -14,9 +15,16 @@
 struct FWriteToSliceVertexOutput
 {
 	FScreenVertexOutput Vertex;
+#if USING_VERTEX_SHADER_LAYER
+	uint LayerIndex : SV_RenderTargetArrayIndex;
+#else
 	uint LayerIndex : TEXCOORD1;
+#endif
 };
 
+/** Z index of the minimum slice in the range. */
+int MinZ;
+
 float4 UVScaleBias;
 uint VolumeCascadeIndex;
 
@@ -31,7 +39,7 @@ void WriteToSliceMainVS(
 	Output.Vertex.Position = float4( InPosition, 0, 1 );
 	// Remap UVs based on the subregion of the volume texture being rendered to
     Output.Vertex.UV = InUV * UVScaleBias.xy + UVScaleBias.zw;
-	Output.LayerIndex = LayerIndex;
+	Output.LayerIndex = LayerIndex + MinZ;
 }
 
 void CopySceneAlphaMain(
@@ -51,9 +59,6 @@ void CopySceneColorMain(
 	OutColor = float4(EncodeSceneColorForMaterialNode(LinearColor.rgb), 0);
 }
 
-/** Z index of the minimum slice in the range. */
-int MinZ;
-
 #if FEATURE_LEVEL >= FEATURE_LEVEL_SM4
 /** Geometry shader that writes to a range of slices of a volume texture. */
 [maxvertexcount(3)]
@@ -61,15 +66,15 @@ void WriteToSliceMainGS(triangle FWriteToSliceVertexOutput Input[3], inout Trian
 {
 	FWriteToSliceGeometryOutput Vertex0;
 	Vertex0.Vertex = Input[0].Vertex;
-	Vertex0.LayerIndex = Input[0].LayerIndex + MinZ;
+	Vertex0.LayerIndex = Input[0].LayerIndex;
 
 	FWriteToSliceGeometryOutput Vertex1;
 	Vertex1.Vertex = Input[1].Vertex;
-	Vertex1.LayerIndex = Input[1].LayerIndex + MinZ;
+	Vertex1.LayerIndex = Input[1].LayerIndex;
 
 	FWriteToSliceGeometryOutput Vertex2;
 	Vertex2.Vertex = Input[2].Vertex;
-	Vertex2.LayerIndex = Input[2].LayerIndex + MinZ;
+	Vertex2.LayerIndex = Input[2].LayerIndex;
 
 	OutStream.Append(Vertex0);
 	OutStream.Append(Vertex1);
diff --git a/Engine/Source/Runtime/Renderer/Private/AtmosphereRendering.cpp b/Engine/Source/Runtime/Renderer/Private/AtmosphereRendering.cpp
--- a/Engine/Source/Runtime/Renderer/Private/AtmosphereRendering.cpp
+++ b/Engine/Source/Runtime/Renderer/Private/AtmosphereRendering.cpp
@@ -387,7 +387,7 @@ void FDeferredShadingSceneRenderer::RenderAtmosphere(FRHICommandListImmediate& R
 	// Atmospheric fog?
 	if (Scene->GetFeatureLevel() >= ERHIFeatureLevel::SM4 && Scene->HasAtmosphericFog())
 	{
-		SCOPED_DRAW_EVENT(RHICmdList, Fog);
+		SCOPED_DRAW_EVENT(RHICmdList, AtmosphericFog);
 
 		static const FVector2D Vertices[4] =
 		{
@@ -885,6 +885,42 @@ public:
 		return bShaderHasOutdatedParameters;
 	}
 };
+		
+class FAtmospherePrecomputeInscatterVS : public FGlobalShader
+{
+	DECLARE_SHADER_TYPE(FAtmospherePrecomputeInscatterVS,Global);
+	
+	static bool ShouldCache(EShaderPlatform Platform)
+	{
+		return IsFeatureLevelSupported(Platform, ERHIFeatureLevel::SM4);
+	}
+	
+	/* Default constructor. */
+	FAtmospherePrecomputeInscatterVS() {}
+	
+	FShaderParameter AtmosphereLayerParameter;
+	
+public:
+	
+	/* Initialization constructor. */
+	FAtmospherePrecomputeInscatterVS(const ShaderMetaType::CompiledShaderInitializerType& Initializer)
+	:	FGlobalShader(Initializer)
+	{
+		AtmosphereLayerParameter.Bind(Initializer.ParameterMap, TEXT("AtmosphereLayer"));
+	}
+
+	void SetParameters(FRHICommandList& RHICmdList, int AtmosphereLayer)
+	{
+		SetShaderValue(RHICmdList, GetVertexShader(), AtmosphereLayerParameter, AtmosphereLayer);
+	}
+
+	virtual bool Serialize(FArchive& Ar) override
+	{
+		bool bShaderHasOutdatedParameters = FGlobalShader::Serialize(Ar);
+		Ar << AtmosphereLayerParameter;
+		return bShaderHasOutdatedParameters;
+	}
+};
 
 // Final Fix
 class FAtmosphereCopyInscatterFPS : public FGlobalShader
@@ -990,6 +1026,7 @@ IMPLEMENT_SHADER_TYPE(,FAtmosphereInscatterNPS,TEXT("AtmospherePrecomputeInscatt
 IMPLEMENT_SHADER_TYPE(,FAtmosphereCopyInscatterFPS,TEXT("AtmospherePrecomputeInscatter"),TEXT("CopyInscatterFPS"),SF_Pixel);
 IMPLEMENT_SHADER_TYPE(,FAtmosphereCopyInscatterFBackPS,TEXT("AtmospherePrecomputeInscatter"),TEXT("CopyInscatterFBackPS"),SF_Pixel);
 IMPLEMENT_SHADER_TYPE(,FAtmospherePrecomputeVS,TEXT("AtmospherePrecompute"),TEXT("MainVS"),SF_Vertex);
+IMPLEMENT_SHADER_TYPE(,FAtmospherePrecomputeInscatterVS,TEXT("AtmospherePrecomputeInscatter"),TEXT("MainVS"),SF_Vertex);
 
 namespace
 {
@@ -1124,13 +1161,15 @@ void FAtmosphericFogSceneInfo::RenderAtmosphereShaders(FRHICommandList& RHICmdLi
 		{
 			int32 Layer = Atmosphere3DTextureIndex;
 			{
-				FTextureRHIParamRef RenderTargets[2];
-				RenderTargets[0] = AtmosphereTextures->AtmosphereDeltaSR->GetRenderTargetItem().TargetableTexture;
-				RenderTargets[1] = AtmosphereTextures->AtmosphereDeltaSM->GetRenderTargetItem().TargetableTexture;
-				SetRenderTargets(RHICmdList, ARRAY_COUNT(RenderTargets), RenderTargets, FTextureRHIRef(), 0, NULL);
+				FRHIDepthRenderTargetView DepthRTV;
+				FRHIRenderTargetView RenderTargets[2];
+				RenderTargets[0] = FRHIRenderTargetView(AtmosphereTextures->AtmosphereDeltaSR->GetRenderTargetItem().TargetableTexture);
+				RenderTargets[1] = FRHIRenderTargetView(AtmosphereTextures->AtmosphereDeltaSM->GetRenderTargetItem().TargetableTexture);
+				
+				RHICmdList.SetRenderTargets(ARRAY_COUNT(RenderTargets), RenderTargets, &DepthRTV, 0, nullptr);
 
-				TShaderMapRef<FAtmospherePrecomputeVS> VertexShader(ShaderMap);
-				TShaderMapRef<FAtmosphereGS> GeometryShader(ShaderMap);
+				TShaderMapRef<FAtmospherePrecomputeInscatterVS> VertexShader(ShaderMap);
+				TOptionalShaderMapRef<FAtmosphereGS> GeometryShader(ShaderMap);
 				TShaderMapRef<FAtmosphereInscatter1PS> PixelShader(ShaderMap);
 
 				static FGlobalBoundShaderState BoundShaderState;
@@ -1140,8 +1179,12 @@ void FAtmosphericFogSceneInfo::RenderAtmosphereShaders(FRHICommandList& RHICmdLi
 				//
 				float r;
 				FVector4 DhdH;
-				GetLayerValue(Layer, r, DhdH);
-				GeometryShader->SetParameters(RHICmdList, Layer);
+                GetLayerValue(Layer, r, DhdH);
+                VertexShader->SetParameters(RHICmdList, Layer);
+				if (GeometryShader.IsValid())
+				{
+					GeometryShader->SetParameters(RHICmdList, Layer);
+				}
 				PixelShader->SetParameters(RHICmdList, View, r, DhdH, AtmosphereTextures);
 				DrawQuad(RHICmdList, ViewRect, *VertexShader);
 
@@ -1170,8 +1213,8 @@ void FAtmosphericFogSceneInfo::RenderAtmosphereShaders(FRHICommandList& RHICmdLi
 				const FSceneRenderTargetItem& DestRenderTarget = AtmosphereTextures->AtmosphereInscatter->GetRenderTargetItem();
 				SetRenderTarget(RHICmdList, DestRenderTarget.TargetableTexture, FTextureRHIRef());
 
-				TShaderMapRef<FAtmospherePrecomputeVS> VertexShader(ShaderMap);
-				TShaderMapRef<FAtmosphereGS> GeometryShader(ShaderMap);
+				TShaderMapRef<FAtmospherePrecomputeInscatterVS> VertexShader(ShaderMap);
+				TOptionalShaderMapRef<FAtmosphereGS> GeometryShader(ShaderMap);
 				TShaderMapRef<FAtmosphereCopyInscatter1PS> PixelShader(ShaderMap);
 
 				static FGlobalBoundShaderState BoundShaderState;
@@ -1180,8 +1223,12 @@ void FAtmosphericFogSceneInfo::RenderAtmosphereShaders(FRHICommandList& RHICmdLi
 
 				float r;
 				FVector4 DhdH;
-				GetLayerValue(Layer, r, DhdH);
-				GeometryShader->SetParameters(RHICmdList, Layer);
+                GetLayerValue(Layer, r, DhdH);
+                VertexShader->SetParameters(RHICmdList, Layer);
+				if (GeometryShader.IsValid())
+				{
+					GeometryShader->SetParameters(RHICmdList, Layer);
+				}
 				PixelShader->SetParameters(RHICmdList, View, r, DhdH, Layer, AtmosphereTextures);
 				DrawQuad(RHICmdList, ViewRect, *VertexShader);
 
@@ -1199,8 +1246,8 @@ void FAtmosphericFogSceneInfo::RenderAtmosphereShaders(FRHICommandList& RHICmdLi
 				const FSceneRenderTargetItem& DestRenderTarget = AtmosphereTextures->AtmosphereDeltaJ->GetRenderTargetItem();
 				SetRenderTarget(RHICmdList, DestRenderTarget.TargetableTexture, FTextureRHIRef());
 
-				TShaderMapRef<FAtmospherePrecomputeVS> VertexShader(ShaderMap);
-				TShaderMapRef<FAtmosphereGS> GeometryShader(ShaderMap);
+				TShaderMapRef<FAtmospherePrecomputeInscatterVS> VertexShader(ShaderMap);
+				TOptionalShaderMapRef<FAtmosphereGS> GeometryShader(ShaderMap);
 				TShaderMapRef<FAtmosphereInscatterSPS> PixelShader(ShaderMap);
 
 				static FGlobalBoundShaderState BoundShaderState;
@@ -1209,8 +1256,12 @@ void FAtmosphericFogSceneInfo::RenderAtmosphereShaders(FRHICommandList& RHICmdLi
 
 				float r;
 				FVector4 DhdH;
-				GetLayerValue(Layer, r, DhdH);
-				GeometryShader->SetParameters(RHICmdList, Layer);
+                GetLayerValue(Layer, r, DhdH);
+                VertexShader->SetParameters(RHICmdList, Layer);
+				if (GeometryShader.IsValid())
+				{
+					GeometryShader->SetParameters(RHICmdList, Layer);
+				}
 				PixelShader->SetParameters(RHICmdList, View, r, DhdH, AtmoshpereOrder == 2 ? 1.f : 0.f, AtmosphereTextures);
 				DrawQuad(RHICmdList, ViewRect, *VertexShader);
 
@@ -1248,8 +1299,8 @@ void FAtmosphericFogSceneInfo::RenderAtmosphereShaders(FRHICommandList& RHICmdLi
 				const FSceneRenderTargetItem& DestRenderTarget = AtmosphereTextures->AtmosphereDeltaSR->GetRenderTargetItem();
 				SetRenderTarget(RHICmdList, DestRenderTarget.TargetableTexture, FTextureRHIRef());
 
-				TShaderMapRef<FAtmospherePrecomputeVS> VertexShader(ShaderMap);
-				TShaderMapRef<FAtmosphereGS> GeometryShader(ShaderMap);
+				TShaderMapRef<FAtmospherePrecomputeInscatterVS> VertexShader(ShaderMap);
+				TOptionalShaderMapRef<FAtmosphereGS> GeometryShader(ShaderMap);
 				TShaderMapRef<FAtmosphereInscatterNPS> PixelShader(ShaderMap);
 
 				static FGlobalBoundShaderState BoundShaderState;
@@ -1258,8 +1309,12 @@ void FAtmosphericFogSceneInfo::RenderAtmosphereShaders(FRHICommandList& RHICmdLi
 
 				float r;
 				FVector4 DhdH;
-				GetLayerValue(Layer, r, DhdH);
-				GeometryShader->SetParameters(RHICmdList, Layer);
+                GetLayerValue(Layer, r, DhdH);
+                VertexShader->SetParameters(RHICmdList, Layer);
+				if (GeometryShader.IsValid())
+				{
+					GeometryShader->SetParameters(RHICmdList, Layer);
+				}
 				PixelShader->SetParameters(RHICmdList, View, r, DhdH, AtmoshpereOrder == 2 ? 1.f : 0.f, AtmosphereTextures);
 				DrawQuad(RHICmdList, ViewRect, *VertexShader);
 
@@ -1303,8 +1358,8 @@ void FAtmosphericFogSceneInfo::RenderAtmosphereShaders(FRHICommandList& RHICmdLi
 				const FSceneRenderTargetItem& DestRenderTarget = AtmosphereTextures->AtmosphereInscatter->GetRenderTargetItem();
 				SetRenderTarget(RHICmdList, DestRenderTarget.TargetableTexture, FTextureRHIRef());
 
-				TShaderMapRef<FAtmospherePrecomputeVS> VertexShader(ShaderMap);
-				TShaderMapRef<FAtmosphereGS> GeometryShader(ShaderMap);
+				TShaderMapRef<FAtmospherePrecomputeInscatterVS> VertexShader(ShaderMap);
+				TOptionalShaderMapRef<FAtmosphereGS> GeometryShader(ShaderMap);
 				TShaderMapRef<FAtmosphereCopyInscatterNPS> PixelShader(ShaderMap);
 
 				static FGlobalBoundShaderState BoundShaderState;
@@ -1313,8 +1368,12 @@ void FAtmosphericFogSceneInfo::RenderAtmosphereShaders(FRHICommandList& RHICmdLi
 
 				float r;
 				FVector4 DhdH;
-				GetLayerValue(Layer, r, DhdH);
-				GeometryShader->SetParameters(RHICmdList, Layer);
+                GetLayerValue(Layer, r, DhdH);
+                VertexShader->SetParameters(RHICmdList, Layer);
+				if (GeometryShader.IsValid())
+				{
+					GeometryShader->SetParameters(RHICmdList, Layer);
+				}
 				PixelShader->SetParameters(RHICmdList, View, r, DhdH, Layer, AtmosphereTextures);
 				DrawQuad(RHICmdList, ViewRect, *VertexShader);
 
@@ -1335,8 +1394,8 @@ void FAtmosphericFogSceneInfo::RenderAtmosphereShaders(FRHICommandList& RHICmdLi
 				const FSceneRenderTargetItem& DestRenderTarget = AtmosphereTextures->AtmosphereDeltaSR->GetRenderTargetItem();
 				SetRenderTarget(RHICmdList, DestRenderTarget.TargetableTexture, FTextureRHIRef());
 
-				TShaderMapRef<FAtmospherePrecomputeVS> VertexShader(ShaderMap);
-				TShaderMapRef<FAtmosphereGS> GeometryShader(ShaderMap);
+				TShaderMapRef<FAtmospherePrecomputeInscatterVS> VertexShader(ShaderMap);
+				TOptionalShaderMapRef<FAtmosphereGS> GeometryShader(ShaderMap);
 				TShaderMapRef<FAtmosphereCopyInscatterFPS> PixelShader(ShaderMap);
 
 				static FGlobalBoundShaderState BoundShaderState;
@@ -1345,8 +1404,12 @@ void FAtmosphericFogSceneInfo::RenderAtmosphereShaders(FRHICommandList& RHICmdLi
 
 				float r;
 				FVector4 DhdH;
-				GetLayerValue(Layer, r, DhdH);
-				GeometryShader->SetParameters(RHICmdList, Layer);
+                GetLayerValue(Layer, r, DhdH);
+                VertexShader->SetParameters(RHICmdList, Layer);
+				if (GeometryShader.IsValid())
+				{
+					GeometryShader->SetParameters(RHICmdList, Layer);
+				}
 				PixelShader->SetParameters(RHICmdList, r, DhdH, Layer, AtmosphereTextures);
 				DrawQuad(RHICmdList, ViewRect, *VertexShader);
 
@@ -1364,8 +1427,8 @@ void FAtmosphericFogSceneInfo::RenderAtmosphereShaders(FRHICommandList& RHICmdLi
 				const FSceneRenderTargetItem& DestRenderTarget = AtmosphereTextures->AtmosphereInscatter->GetRenderTargetItem();
 				SetRenderTarget(RHICmdList, DestRenderTarget.TargetableTexture, FTextureRHIRef());
 
-				TShaderMapRef<FAtmospherePrecomputeVS> VertexShader(ShaderMap);
-				TShaderMapRef<FAtmosphereGS> GeometryShader(ShaderMap);
+				TShaderMapRef<FAtmospherePrecomputeInscatterVS> VertexShader(ShaderMap);
+				TOptionalShaderMapRef<FAtmosphereGS> GeometryShader(ShaderMap);
 				TShaderMapRef<FAtmosphereCopyInscatterFBackPS> PixelShader(ShaderMap);
 
 				static FGlobalBoundShaderState BoundShaderState;
@@ -1374,8 +1437,12 @@ void FAtmosphericFogSceneInfo::RenderAtmosphereShaders(FRHICommandList& RHICmdLi
 
 				float r;
 				FVector4 DhdH;
-				GetLayerValue(Layer, r, DhdH);
-				GeometryShader->SetParameters(RHICmdList, Layer);
+                GetLayerValue(Layer, r, DhdH);
+                VertexShader->SetParameters(RHICmdList, Layer);
+				if (GeometryShader.IsValid())
+				{
+					GeometryShader->SetParameters(RHICmdList, Layer);
+				}
 				PixelShader->SetParameters(RHICmdList, r, DhdH, Layer, AtmosphereTextures);
 				DrawQuad(RHICmdList, ViewRect, *VertexShader);
 
diff --git a/Engine/Source/Runtime/Renderer/Private/TranslucentLighting.cpp b/Engine/Source/Runtime/Renderer/Private/TranslucentLighting.cpp
--- a/Engine/Source/Runtime/Renderer/Private/TranslucentLighting.cpp
+++ b/Engine/Source/Runtime/Renderer/Private/TranslucentLighting.cpp
@@ -1069,18 +1069,12 @@ void RasterizeToVolumeTexture(FRHICommandList& RHICmdList, FVolumeBounds VolumeB
 	// Render a quad per slice affected by the given bounds
 	RHICmdList.DrawPrimitive(PT_TriangleStrip, 0, 2, NumInstances);
 }
-
+        
 /** Helper function that clears the given volume texture render targets. */
 template<int32 NumRenderTargets>
 void ClearVolumeTextures(FRHICommandList& RHICmdList, ERHIFeatureLevel::Type FeatureLevel, const FTextureRHIParamRef* RenderTargets, const FLinearColor* ClearColors)
 {
 	SetRenderTargets(RHICmdList, NumRenderTargets, RenderTargets, FTextureRHIRef(), 0, NULL);
-
-	static FGlobalBoundShaderState VolumeClearBoundShaderState;
-	
-
-
-
 #if PLATFORM_XBOXONE
 	// ClearMRT is faster on Xbox
 	if (true)
@@ -1093,6 +1087,8 @@ void ClearVolumeTextures(FRHICommandList& RHICmdList, ERHIFeatureLevel::Type Fea
 	}
 	else
 	{
+		static FGlobalBoundShaderState VolumeClearBoundShaderState;
+
 		RHICmdList.SetRasterizerState(TStaticRasterizerState<FM_Solid, CM_None>::GetRHI());
 		RHICmdList.SetDepthStencilState(TStaticDepthStencilState<false, CF_Always>::GetRHI());
 		RHICmdList.SetBlendState(TStaticBlendState<>::GetRHI());
@@ -1100,13 +1096,16 @@ void ClearVolumeTextures(FRHICommandList& RHICmdList, ERHIFeatureLevel::Type Fea
 		const FVolumeBounds VolumeBounds(GTranslucencyLightingVolumeDim);
 		auto ShaderMap = GetGlobalShaderMap(FeatureLevel);
 		TShaderMapRef<FWriteToSliceVS> VertexShader(ShaderMap);
-		TShaderMapRef<FWriteToSliceGS> GeometryShader(ShaderMap);
+		TOptionalShaderMapRef<FWriteToSliceGS> GeometryShader(ShaderMap);
 		TShaderMapRef<TOneColorPixelShaderMRT<NumRenderTargets> > PixelShader(ShaderMap);
 		
 		SetGlobalBoundShaderState(RHICmdList, FeatureLevel, VolumeClearBoundShaderState, GScreenVertexDeclaration.VertexDeclarationRHI, *VertexShader, *PixelShader, *GeometryShader);
 
 		VertexShader->SetParameters(RHICmdList, VolumeBounds, GTranslucencyLightingVolumeDim);
-		GeometryShader->SetParameters(RHICmdList, VolumeBounds);		
+		if(GeometryShader.IsValid())
+		{
+			GeometryShader->SetParameters(RHICmdList, VolumeBounds);
+		}
 		PixelShader->SetColors(RHICmdList, ClearColors, NumRenderTargets);
 
 		RasterizeToVolumeTexture(RHICmdList, VolumeBounds);
@@ -1192,10 +1191,6 @@ void FDeferredShadingSceneRenderer::InjectAmbientCubemapTranslucentVolumeLightin
 	{
 		SCOPED_DRAW_EVENT(RHICmdList, InjectAmbientCubemapTranslucentVolumeLighting);
 
-		RHICmdList.SetRasterizerState(TStaticRasterizerState<FM_Solid, CM_None>::GetRHI());
-		RHICmdList.SetDepthStencilState(TStaticDepthStencilState<false, CF_Always>::GetRHI());
-		RHICmdList.SetBlendState(TStaticBlendState<CW_RGBA, BO_Add, BF_One, BF_One, BO_Add, BF_One, BF_One>::GetRHI());
-
 		const FVolumeBounds VolumeBounds(GTranslucencyLightingVolumeDim);
 
 		auto ShaderMap = GetGlobalShaderMap(FeatureLevel);
@@ -1206,8 +1201,12 @@ void FDeferredShadingSceneRenderer::InjectAmbientCubemapTranslucentVolumeLightin
 			// we don't update the directional volume (could be a HQ option)
 			SetRenderTarget(RHICmdList, SceneContext.TranslucencyLightingVolumeAmbient[VolumeCascadeIndex]->GetRenderTargetItem().TargetableTexture, FTextureRHIRef());
 
+			RHICmdList.SetRasterizerState(TStaticRasterizerState<FM_Solid, CM_None>::GetRHI());
+			RHICmdList.SetDepthStencilState(TStaticDepthStencilState<false, CF_Always>::GetRHI());
+			RHICmdList.SetBlendState(TStaticBlendState<CW_RGBA, BO_Add, BF_One, BF_One, BO_Add, BF_One, BF_One>::GetRHI());
+
 			TShaderMapRef<FWriteToSliceVS> VertexShader(ShaderMap);
-			TShaderMapRef<FWriteToSliceGS> GeometryShader(ShaderMap);
+			TOptionalShaderMapRef<FWriteToSliceGS> GeometryShader(ShaderMap);
 			TShaderMapRef<FInjectAmbientCubemapPS> PixelShader(ShaderMap);
 
 			static FGlobalBoundShaderState BoundShaderState;
@@ -1215,7 +1214,10 @@ void FDeferredShadingSceneRenderer::InjectAmbientCubemapTranslucentVolumeLightin
 			SetGlobalBoundShaderState(RHICmdList, FeatureLevel, BoundShaderState, GScreenVertexDeclaration.VertexDeclarationRHI, *VertexShader, *PixelShader, *GeometryShader);
 
 			VertexShader->SetParameters(RHICmdList, VolumeBounds, GTranslucencyLightingVolumeDim);
-			GeometryShader->SetParameters(RHICmdList, VolumeBounds);
+			if(GeometryShader.IsValid())
+			{
+				GeometryShader->SetParameters(RHICmdList, VolumeBounds);
+			}
 
 			uint32 Count = View.FinalPostProcessSettings.ContributingCubemaps.Num();
 			for(uint32 i = 0; i < Count; ++i)
@@ -1226,7 +1228,7 @@ void FDeferredShadingSceneRenderer::InjectAmbientCubemapTranslucentVolumeLightin
 
 				RasterizeToVolumeTexture(RHICmdList, VolumeBounds);
 			}
-			
+
 			RHICmdList.CopyToResolveTarget(SceneContext.TranslucencyLightingVolumeAmbient[VolumeCascadeIndex]->GetRenderTargetItem().TargetableTexture,
 				SceneContext.TranslucencyLightingVolumeAmbient[VolumeCascadeIndex]->GetRenderTargetItem().ShaderResourceTexture, true, FResolveParams());
 		}
@@ -1295,12 +1297,6 @@ void FDeferredShadingSceneRenderer::AccumulateTranslucentVolumeObjectShadowing(F
 	{
 		SCOPED_DRAW_EVENT(RHICmdList, AccumulateTranslucentVolumeShadowing);
 
-		RHICmdList.SetRasterizerState(TStaticRasterizerState<FM_Solid, CM_None>::GetRHI());
-		RHICmdList.SetDepthStencilState(TStaticDepthStencilState<false, CF_Always>::GetRHI());
-
-		// Modulate the contribution of multiple object shadows in rgb
-		RHICmdList.SetBlendState(TStaticBlendState<CW_RGB, BO_Add, BF_DestColor, BF_Zero>::GetRHI());
-
 		auto ShaderMap = GetGlobalShaderMap(FeatureLevel);
 
 		FSceneRenderTargets& SceneContext = FSceneRenderTargets::Get(RHICmdList);
@@ -1327,13 +1323,22 @@ void FDeferredShadingSceneRenderer::AccumulateTranslucentVolumeObjectShadowing(F
 
 				SetRenderTarget(RHICmdList, RenderTarget, FTextureRHIRef());
 
+				RHICmdList.SetRasterizerState(TStaticRasterizerState<FM_Solid, CM_None>::GetRHI());
+				RHICmdList.SetDepthStencilState(TStaticDepthStencilState<false, CF_Always>::GetRHI());
+
+				// Modulate the contribution of multiple object shadows in rgb
+				RHICmdList.SetBlendState(TStaticBlendState<CW_RGB, BO_Add, BF_DestColor, BF_Zero>::GetRHI());
+
 				TShaderMapRef<FWriteToSliceVS> VertexShader(ShaderMap);
-				TShaderMapRef<FWriteToSliceGS> GeometryShader(ShaderMap);
+				TOptionalShaderMapRef<FWriteToSliceGS> GeometryShader(ShaderMap);
 				TShaderMapRef<FTranslucentObjectShadowingPS> PixelShader(ShaderMap);
 
 				SetGlobalBoundShaderState(RHICmdList, FeatureLevel, ObjectShadowingBoundShaderState, GScreenVertexDeclaration.VertexDeclarationRHI, *VertexShader, *PixelShader, *GeometryShader);
 				VertexShader->SetParameters(RHICmdList, VolumeBounds, GTranslucencyLightingVolumeDim);
-				GeometryShader->SetParameters(RHICmdList, VolumeBounds);
+				if(GeometryShader.IsValid())
+				{
+					GeometryShader->SetParameters(RHICmdList, VolumeBounds);
+				}
 				PixelShader->SetParameters(RHICmdList, View, LightSceneInfo, InProjectedShadowInfo, VolumeCascadeIndex);
 				
 				RasterizeToVolumeTexture(RHICmdList, VolumeBounds);
@@ -1417,8 +1422,7 @@ void SetInjectionShader(
 		CachedShaderMap = MaterialShaderMap;
 		check(IsInRenderingThread()); // I didn't know quite how to deal with this caching. It won't work with threads.
 		BoundShaderState = 
-			RHICreateBoundShaderState(
-			GScreenVertexDeclaration.VertexDeclarationRHI, VertexShader->GetVertexShader(), FHullShaderRHIRef(), FDomainShaderRHIRef(), PixelShader->GetPixelShader(), GeometryShader->GetGeometryShader());
+			RHICreateBoundShaderState( GScreenVertexDeclaration.VertexDeclarationRHI, VertexShader->GetVertexShader(), FHullShaderRHIRef(), FDomainShaderRHIRef(), PixelShader->GetPixelShader(), GeometryShader ? GeometryShader->GetGeometryShader() : nullptr);
 	}
 
 	RHICmdList.SetBoundShaderState(BoundShaderState);
@@ -1545,7 +1549,7 @@ static void InjectTranslucentLightArray(FRHICommandListImmediate& RHICmdList, co
 			if (VolumeBounds.IsValid())
 			{
 				TShaderMapRef<FWriteToSliceVS> VertexShader(View.ShaderMap);
-				TShaderMapRef<FWriteToSliceGS> GeometryShader(View.ShaderMap);
+				TOptionalShaderMapRef<FWriteToSliceGS> GeometryShader(View.ShaderMap);
 
 				if (bDirectionalLight)
 				{
@@ -1554,7 +1558,7 @@ static void InjectTranslucentLightArray(FRHICommandListImmediate& RHICmdList, co
 					RHICmdList.SetBlendState(TStaticBlendState<
 						CW_RGBA, BO_Add, BF_One, BF_One, BO_Add, BF_One, BF_One,
 						CW_RGB, BO_Add, BF_One, BF_One, BO_Add, BF_One, BF_One>::GetRHI());
-					
+
 					if (InjectionData.ProjectedShadowInfo)
 					{
 						// shadows, restricting light contribution to the cascade bounds (except last cascade far to get light functions and no shadows there)
@@ -1592,7 +1596,10 @@ static void InjectTranslucentLightArray(FRHICommandListImmediate& RHICmdList, co
 				}
 
 				VertexShader->SetParameters(RHICmdList, VolumeBounds, GTranslucencyLightingVolumeDim);
-				GeometryShader->SetParameters(RHICmdList, VolumeBounds);
+				if(GeometryShader.IsValid())
+				{
+					GeometryShader->SetParameters(RHICmdList, VolumeBounds);
+				}
 				RasterizeToVolumeTexture(RHICmdList, VolumeBounds);
 			}
 		}
@@ -1759,12 +1766,15 @@ void FDeferredShadingSceneRenderer::InjectSimpleTranslucentVolumeLightingArray(F
 					if (VolumeBounds.IsValid())
 					{
 						TShaderMapRef<FWriteToSliceVS> VertexShader(View.ShaderMap);
-						TShaderMapRef<FWriteToSliceGS> GeometryShader(View.ShaderMap);
+						TOptionalShaderMapRef<FWriteToSliceGS> GeometryShader(View.ShaderMap);
 						TShaderMapRef<FSimpleLightTranslucentLightingInjectPS> PixelShader(View.ShaderMap);
 						SetGlobalBoundShaderState(RHICmdList, FeatureLevel, InjectSimpleLightBoundShaderState, GScreenVertexDeclaration.VertexDeclarationRHI, *VertexShader, *PixelShader, *GeometryShader);
 
 						VertexShader->SetParameters(RHICmdList, VolumeBounds, GTranslucencyLightingVolumeDim);
-						GeometryShader->SetParameters(RHICmdList, VolumeBounds);
+						if(GeometryShader.IsValid())
+						{
+							GeometryShader->SetParameters(RHICmdList, VolumeBounds);
+						}
 						PixelShader->SetParameters(RHICmdList, View, SimpleLight, SimpleLightPerViewData, VolumeCascadeIndex);
 
 						// Accumulate the contribution of multiple lights
@@ -1805,10 +1815,6 @@ void FDeferredShadingSceneRenderer::FilterTranslucentVolumeLighting(FRHICommandL
 		{
 			SCOPED_DRAW_EVENT(RHICmdList, FilterTranslucentVolume);
 
-			RHICmdList.SetRasterizerState(TStaticRasterizerState<FM_Solid, CM_None>::GetRHI());
-			RHICmdList.SetDepthStencilState(TStaticDepthStencilState<false, CF_Always>::GetRHI());
-			RHICmdList.SetBlendState(TStaticBlendState<>::GetRHI());
-
 			// Filter each cascade
 			for (int32 VolumeCascadeIndex = 0; VolumeCascadeIndex < TVC_MAX; VolumeCascadeIndex++)
 			{
@@ -1829,12 +1835,20 @@ void FDeferredShadingSceneRenderer::FilterTranslucentVolumeLighting(FRHICommandL
 
 				const FVolumeBounds VolumeBounds(GTranslucencyLightingVolumeDim);
 				TShaderMapRef<FWriteToSliceVS> VertexShader(View.ShaderMap);
-				TShaderMapRef<FWriteToSliceGS> GeometryShader(View.ShaderMap);
+				TOptionalShaderMapRef<FWriteToSliceGS> GeometryShader(View.ShaderMap);
 				TShaderMapRef<FFilterTranslucentVolumePS> PixelShader(View.ShaderMap);
+				
+				RHICmdList.SetRasterizerState(TStaticRasterizerState<FM_Solid, CM_None>::GetRHI());
+				RHICmdList.SetDepthStencilState(TStaticDepthStencilState<false, CF_Always>::GetRHI());
+				RHICmdList.SetBlendState(TStaticBlendState<>::GetRHI());
+
 				SetGlobalBoundShaderState(RHICmdList, FeatureLevel, FilterBoundShaderState, GScreenVertexDeclaration.VertexDeclarationRHI, *VertexShader, *PixelShader, *GeometryShader);
 
 				VertexShader->SetParameters(RHICmdList, VolumeBounds, GTranslucencyLightingVolumeDim);
-				GeometryShader->SetParameters(RHICmdList, VolumeBounds);
+				if(GeometryShader.IsValid())
+				{
+					GeometryShader->SetParameters(RHICmdList, VolumeBounds);
+				}
 				PixelShader->SetParameters(RHICmdList, View, VolumeCascadeIndex);
 
 				RasterizeToVolumeTexture(RHICmdList, VolumeBounds);
diff --git a/Engine/Source/Runtime/Renderer/Private/TranslucentRendering.h b/Engine/Source/Runtime/Renderer/Private/TranslucentRendering.h
--- a/Engine/Source/Runtime/Renderer/Private/TranslucentRendering.h
+++ b/Engine/Source/Runtime/Renderer/Private/TranslucentRendering.h
@@ -170,8 +170,9 @@ public:
 
 	FWriteToSliceVS(const ShaderMetaType::CompiledShaderInitializerType& Initializer):
 		FGlobalShader(Initializer)
-	{
-		UVScaleBias.Bind(Initializer.ParameterMap, TEXT("UVScaleBias"));
+    {
+        UVScaleBias.Bind(Initializer.ParameterMap, TEXT("UVScaleBias"));
+        MinZ.Bind(Initializer.ParameterMap, TEXT("MinZ"));
 	}
 
 	FWriteToSliceVS() {}
@@ -184,17 +185,20 @@ public:
 			(VolumeBounds.MaxY - VolumeBounds.MinY) * InvVolumeResolution,
 			VolumeBounds.MinX * InvVolumeResolution,
 			VolumeBounds.MinY * InvVolumeResolution));
+        SetShaderValue(RHICmdList, GetVertexShader(), MinZ, VolumeBounds.MinZ);
 	}
 
 	virtual bool Serialize(FArchive& Ar) override
 	{
 		bool bShaderHasOutdatedParameters = FGlobalShader::Serialize(Ar);
 		Ar << UVScaleBias;
-		return bShaderHasOutdatedParameters;
+        Ar << MinZ;
+        return bShaderHasOutdatedParameters;
 	}
 
 private:
-	FShaderParameter UVScaleBias;
+    FShaderParameter UVScaleBias;
+    FShaderParameter MinZ;
 };
 
 /** Geometry shader used to write to a range of slices of a 3d volume texture. */
diff --git a/Engine/Source/Runtime/ShaderCore/Public/Shader.h b/Engine/Source/Runtime/ShaderCore/Public/Shader.h
--- a/Engine/Source/Runtime/ShaderCore/Public/Shader.h
+++ b/Engine/Source/Runtime/ShaderCore/Public/Shader.h
@@ -1253,6 +1253,30 @@ private:
 	ShaderType* Shader;
 };
 
+/** A reference to an optional shader, initialized with a shader type from a shader map if it is available or nullptr if it is not. */
+template<typename ShaderType>
+class TOptionalShaderMapRef
+{
+public:
+	TOptionalShaderMapRef(const TShaderMap<typename ShaderType::ShaderMetaType>* ShaderIndex):
+	Shader((ShaderType*)ShaderIndex->GetShader(&ShaderType::StaticType)) // gcc3 needs the template quantifier so it knows the < is not a less-than
+	{}
+	FORCEINLINE bool IsValid() const
+	{
+		return Shader != nullptr;
+	}
+	FORCEINLINE ShaderType* operator->() const
+	{
+		return Shader;
+	}
+	FORCEINLINE ShaderType* operator*() const
+	{
+		return Shader;
+	}
+private:
+	ShaderType* Shader;
+};
+
 /** Tracks state when traversing a FSerializationHistory. */
 class FSerializationHistoryTraversalState
 {
