Implements mesh's receive decal boolean using the stencil buffer instead of the custom stencil buffer to prune out one full-res pass between base passe and decals' pass.

diff --git a/Engine/Shaders/BasePassPixelShader.usf b/Engine/Shaders/BasePassPixelShader.usf
--- a/Engine/Shaders/BasePassPixelShader.usf
+++ b/Engine/Shaders/BasePassPixelShader.usf
@@ -833,7 +833,6 @@ void Main(
 		GBuffer.PrecomputedShadowFactors = PrecomputedShadowFactors;
 		GBuffer.GBufferAO = MaterialAO;
 		GBuffer.Opacity = Opacity;
-		GBuffer.DecalMask = Primitive.DecalReceiverMask;
 		GBuffer.HasDistanceFieldRepresentation = Primitive.HasDistanceFieldRepresentation;
 		GBuffer.HasHeightfieldRepresentation = Primitive.HasHeightfieldRepresentation;
 
diff --git a/Engine/Shaders/DeferredDecal.usf b/Engine/Shaders/DeferredDecal.usf
--- a/Engine/Shaders/DeferredDecal.usf
+++ b/Engine/Shaders/DeferredDecal.usf
@@ -28,22 +28,6 @@ float4x4 FrustumComponentToClip;
 // fade alpha from screen size fading
 float FadeAlpha;
 
-// only for StencilDecalMaskMain, to get ShaderComplexity for Decals (doesn't look at the DecalMask yet)
-float MaskComparison;
-
-void StencilDecalMaskMain(
-	FScreenVertexOutput Input,
-	out float4 OutColor : SV_Target0
-	)
-{
-	FGBufferData InGBufferData = GetGBufferData(Input.UV);
-    
-	// Write 1 to the stencil buffer for pixels that are masked out of receiving decals
-	clip(MaskComparison - InGBufferData.DecalMask);
-
-	OutColor = 0;
-}
-
 half3 FrameBufferDecalBlendOp(half4 Source)
 {
 #if (COMPILER_GLSL_ES2)
@@ -355,7 +339,6 @@ void MainPS(
 	Data.GBufferAO = 1;
 	Data.Opacity = Opacity;
 	Data.ShadingModelID = SHADINGMODELID_DEFAULT_LIT;
-	Data.DecalMask = 1;
 	Data.HasDistanceFieldRepresentation = 1;
 	Data.HasHeightfieldRepresentation = 1;
 
diff --git a/Engine/Shaders/DeferredShadingCommon.usf b/Engine/Shaders/DeferredShadingCommon.usf
--- a/Engine/Shaders/DeferredShadingCommon.usf
+++ b/Engine/Shaders/DeferredShadingCommon.usf
@@ -261,8 +261,6 @@ struct FGBufferData
 	float GBufferAO;
 	// 0..255 
 	uint ShadingModelID;
-	// 0..1 decal receiver mask
-	uint DecalMask;
 	// 0..1
 	uint HasDistanceFieldRepresentation; 
 	// 0..1
@@ -355,7 +353,7 @@ void EncodeGBuffer(
 		OutGBufferD.gba = GBuffer.CustomData;
 	}
 
-	OutGBufferB.a = Encode71(EncodeShadingModelId(GBuffer.ShadingModelID), GBuffer.DecalMask);
+	OutGBufferB.a = EncodeShadingModelId(GBuffer.ShadingModelID);
 
 	OutGBufferE = GBuffer.PrecomputedShadowFactors;
 
@@ -420,8 +418,7 @@ FGBufferData DecodeGBufferData(
 	uint PackedAlpha = (uint)( InGBufferA.a * 3.999 );
 	GBuffer.HasHeightfieldRepresentation = (PackedAlpha & 2) != 0;
 	GBuffer.HasDistanceFieldRepresentation = PackedAlpha & 1;
-	GBuffer.DecalMask = 0;
-	GBuffer.ShadingModelID = DecodeShadingModelId(Decode71(InGBufferB.a, GBuffer.DecalMask));
+	GBuffer.ShadingModelID = DecodeShadingModelId(InGBufferB.a);
 	
 	GBuffer.StoredBaseColor = DecodeBaseColor(InGBufferC.rgb);
 	GBuffer.BaseColor = GBuffer.StoredBaseColor;
diff --git a/Engine/Shaders/MaterialTemplate.usf b/Engine/Shaders/MaterialTemplate.usf
--- a/Engine/Shaders/MaterialTemplate.usf
+++ b/Engine/Shaders/MaterialTemplate.usf
@@ -1130,7 +1130,7 @@ float4 SceneTextureLookup(float2 UV, int SceneTextureIndex, bool bFiltered)
 	else if (SceneTextureIndex == 19)	if (bFiltered) { return Texture2DSample(PostprocessInput5, BilinearTextureSampler0, UV);} else { return Texture2DSample(PostprocessInput5, PostprocessInput5Sampler, UV);} 
 	else if (SceneTextureIndex == 20)	if (bFiltered) { return Texture2DSample(PostprocessInput6, BilinearTextureSampler0, UV);} else { return Texture2DSample(PostprocessInput6, PostprocessInput6Sampler, UV);} 
 #endif // POST_PROCESS_MATERIAL
-	else if (SceneTextureIndex == 21)	return ScreenSpaceData.GBuffer.DecalMask;
+	else if (SceneTextureIndex == 21)	return 0; // material compiler will return an error
 	else if (SceneTextureIndex == 22)	return float4(GetShadingModelColor(ScreenSpaceData.GBuffer), 1);
 	else if (SceneTextureIndex == 23)	return ScreenSpaceData.AmbientOcclusion;
 	else if (SceneTextureIndex == 24)	return ScreenSpaceData.GBuffer.CustomStencil;
@@ -1188,7 +1188,7 @@ float4 SceneTextureLookup(float2 UV, int SceneTextureIndex, bool bFiltered)
 		case 20: return Texture2DSample(PostprocessInput6, bFiltered ? BilinearTextureSampler0 : PostprocessInput6Sampler, UV);
 #endif // __POST_PROCESS_COMMON__
 		// PPI_DecalMask
-		case 21: return ScreenSpaceData.GBuffer.DecalMask;
+		case 21: return 0;  // material compiler will return an error
 		// PPI_ShadingModel
 		case 22: return float4(GetShadingModelColor(ScreenSpaceData.GBuffer), 1);
 		// PPI_AmbientOcclusion
diff --git a/Engine/Shaders/PostProcessVisualizeGBuffer.usf b/Engine/Shaders/PostProcessVisualizeGBuffer.usf
--- a/Engine/Shaders/PostProcessVisualizeGBuffer.usf
+++ b/Engine/Shaders/PostProcessVisualizeGBuffer.usf
@@ -59,7 +59,7 @@ float3 GetGBufferChannelData(FGBufferData GBuffer, float2 UV, int ChannelIndex)
 		case GBUFFER_VISUALIZATION_MODE_SEP_TRANS_RGB: return Texture2DSample(PostprocessInput1, PostprocessInput1Sampler, UV).rgb*20;
 		case GBUFFER_VISUALIZATION_MODE_OPACITY: return GBuffer.Opacity;
 		case GBUFFER_VISUALIZATION_MODE_SEP_TRANS_A: return Texture2DSample(PostprocessInput1, PostprocessInput1Sampler, UV).a;
-		case GBUFFER_VISUALIZATION_MODE_DECAL_MASK: return GBuffer.DecalMask;
+		case GBUFFER_VISUALIZATION_MODE_DECAL_MASK: return 0; // not available anymore because moved into the stencil buffer
 		case GBUFFER_VISUALIZATION_MODE_SCENE_DEPTH: return frac(GBuffer.Depth * 0.0001f);
 		case GBUFFER_VISUALIZATION_MODE_METALLIC: return GBuffer.Metallic;
 		case GBUFFER_VISUALIZATION_MODE_SPECULAR: return GBuffer.Specular;
diff --git a/Engine/Source/Runtime/Engine/Private/Materials/HLSLMaterialTranslator.h b/Engine/Source/Runtime/Engine/Private/Materials/HLSLMaterialTranslator.h
--- a/Engine/Source/Runtime/Engine/Private/Materials/HLSLMaterialTranslator.h
+++ b/Engine/Source/Runtime/Engine/Private/Materials/HLSLMaterialTranslator.h
@@ -2564,6 +2564,11 @@ protected:
 			return NonPixelShaderExpressionError();
 		}
 		
+		if (InSceneTextureId == PPI_DecalMask)
+		{
+			return Error(TEXT("Decal Mask bit was move out of GBuffer to the stencil buffer for performance optimisation and is therefor no longer available"));
+		}
+
 		ESceneTextureId SceneTextureId = (ESceneTextureId)InSceneTextureId;
 
 		UseSceneTextureId(SceneTextureId, true);
diff --git a/Engine/Source/Runtime/Renderer/Private/BasePassRendering.cpp b/Engine/Source/Runtime/Renderer/Private/BasePassRendering.cpp
--- a/Engine/Source/Runtime/Renderer/Private/BasePassRendering.cpp
+++ b/Engine/Source/Runtime/Renderer/Private/BasePassRendering.cpp
@@ -266,7 +266,11 @@ public:
 				Parameters.BlendMode,
 				Parameters.TextureMode,
 				Parameters.ShadingModel != MSM_Unlit && Scene->SkyLight && Scene->SkyLight->bWantsStaticShadowing && !Scene->SkyLight->bHasStaticLighting,
-				IsTranslucentBlendMode(Parameters.BlendMode) && Scene->HasAtmosphericFog()
+				IsTranslucentBlendMode(Parameters.BlendMode) && Scene->HasAtmosphericFog(),
+				/* bOverrideWithShaderComplexity = */ false,
+				/* bInAllowGlobalFog = */ false,
+				/* bInEnableEditorPrimitiveDepthTest = */ false,
+				/* bInEnableReceiveDecalOutput = */ true
 				),
 				Scene->GetFeatureLevel()
 				);
@@ -375,7 +379,8 @@ public:
 			IsTranslucentBlendMode(Parameters.BlendMode) && (Scene && Scene->HasAtmosphericFog()) && View.Family->EngineShowFlags.AtmosphericFog,
 			View.Family->EngineShowFlags.ShaderComplexity,
 			false,
-			Parameters.bEditorCompositeDepthTest
+			Parameters.bEditorCompositeDepthTest,
+			/* bInEnableReceiveDecalOutput = */ Scene != nullptr
 			);
 		RHICmdList.BuildAndSetLocalBoundShaderState(DrawingPolicy.GetBoundShaderStateInput(View.GetFeatureLevel()));
 		DrawingPolicy.SetSharedState(RHICmdList, &View, typename TBasePassDrawingPolicy<LightMapPolicyType>::ContextDataType());
diff --git a/Engine/Source/Runtime/Renderer/Private/BasePassRendering.h b/Engine/Source/Runtime/Renderer/Private/BasePassRendering.h
--- a/Engine/Source/Runtime/Renderer/Private/BasePassRendering.h
+++ b/Engine/Source/Runtime/Renderer/Private/BasePassRendering.h
@@ -510,7 +510,8 @@ public:
 		bool bInEnableAtmosphericFog,
 		bool bOverrideWithShaderComplexity = false,
 		bool bInAllowGlobalFog = false,
-		bool bInEnableEditorPrimitiveDepthTest = false
+		bool bInEnableEditorPrimitiveDepthTest = false,
+		bool bInEnableReceiveDecalOutput = false
 		):
 		FMeshDrawingPolicy(InVertexFactory,InMaterialRenderProxy,InMaterialResource,bOverrideWithShaderComplexity),
 		LightMapPolicy(InLightMapPolicy),
@@ -519,7 +520,8 @@ public:
 		bAllowGlobalFog(bInAllowGlobalFog),
 		bEnableSkyLight(bInEnableSkyLight),
 		bEnableEditorPrimitiveDepthTest(bInEnableEditorPrimitiveDepthTest),
-		bEnableAtmosphericFog(bInEnableAtmosphericFog)
+		bEnableAtmosphericFog(bInEnableAtmosphericFog),
+		bEnableReceiveDecalOutput(bInEnableReceiveDecalOutput)
 	{
 		HullShader = NULL;
 		DomainShader = NULL;
@@ -724,6 +726,19 @@ public:
 			PixelShader->SetMesh(RHICmdList, VertexFactory,View,PrimitiveSceneProxy,BatchElement,BlendMode,DitheredLODTransitionValue);
 		}
 
+		if (bEnableReceiveDecalOutput)
+		{
+			// we hash the stencil group because we only have 6 bits.
+			const uint8 StencilValue = GET_STENCIL_BIT_MASK(RECEIVE_DECAL, PrimitiveSceneProxy ? !!PrimitiveSceneProxy->ReceivesDecals() : 0x00);
+
+			RHICmdList.SetDepthStencilState(TStaticDepthStencilState<
+				true, CF_GreaterEqual,
+				true, CF_Always, SO_Keep, SO_Keep, SO_Replace,
+				false, CF_Always, SO_Keep, SO_Keep, SO_Keep,
+				0xFF, GET_STENCIL_BIT_MASK(RECEIVE_DECAL, 1)
+			>::GetRHI(), StencilValue);
+		}
+
 		FMeshDrawingPolicy::SetMeshRenderState(RHICmdList, View,PrimitiveSceneProxy,Mesh,BatchElementIndex,bBackFace, DitheredLODTransitionValue,FMeshDrawingPolicy::ElementDataType(),PolicyContext);
 	}
 
@@ -738,6 +753,7 @@ public:
 		COMPAREDRAWINGPOLICYMEMBERS(SceneTextureMode);
 		COMPAREDRAWINGPOLICYMEMBERS(bAllowGlobalFog);
 		COMPAREDRAWINGPOLICYMEMBERS(bEnableSkyLight);
+		COMPAREDRAWINGPOLICYMEMBERS(bEnableReceiveDecalOutput);
 
 		return CompareDrawingPolicy(A.LightMapPolicy,B.LightMapPolicy);
 	}
@@ -761,6 +777,9 @@ protected:
 	uint32 bEnableEditorPrimitiveDepthTest : 1;
 	/** Whether or not this policy enables atmospheric fog */
 	uint32 bEnableAtmosphericFog : 1;
+
+	/** Whether or not outputing the receive decal boolean */
+	uint32 bEnableReceiveDecalOutput : 1;
 };
 
 /**
diff --git a/Engine/Source/Runtime/Renderer/Private/CompositionLighting/PostProcessDeferredDecals.cpp b/Engine/Source/Runtime/Renderer/Private/CompositionLighting/PostProcessDeferredDecals.cpp
--- a/Engine/Source/Runtime/Renderer/Private/CompositionLighting/PostProcessDeferredDecals.cpp
+++ b/Engine/Source/Runtime/Renderer/Private/CompositionLighting/PostProcessDeferredDecals.cpp
@@ -246,91 +246,6 @@ void SetDecalBlendState(FRHICommandList& RHICmdList, const ERHIFeatureLevel::Typ
 	}
 }
 
-/** Pixel shader used to setup the decal receiver mask */
-class FStencilDecalMaskPS : public FGlobalShader
-{
-	DECLARE_SHADER_TYPE(FStencilDecalMaskPS,Global);
-public:
-
-	static bool ShouldCache(EShaderPlatform Platform) 
-	{ 
-		return IsFeatureLevelSupported(Platform, ERHIFeatureLevel::SM4); 
-	}
-
-	static void ModifyCompilationEnvironment(EShaderPlatform Platform, FShaderCompilerEnvironment& OutEnvironment)
-	{
-		FGlobalShader::ModifyCompilationEnvironment(Platform,OutEnvironment);
-	}
-
-	FStencilDecalMaskPS(const ShaderMetaType::CompiledShaderInitializerType& Initializer):
-		FGlobalShader(Initializer)
-	{
-		DeferredParameters.Bind(Initializer.ParameterMap);
-		MaskComparison.Bind(Initializer.ParameterMap,TEXT("MaskComparison"));
-	}
-	FStencilDecalMaskPS() {}
-
-	void SetParameters(FRHICommandList& RHICmdList, const FSceneView& View)
-	{
-		const FPixelShaderRHIParamRef ShaderRHI = GetPixelShader();
-
-		FGlobalShader::SetParameters(RHICmdList, ShaderRHI, View);
-		DeferredParameters.Set(RHICmdList, ShaderRHI, View);
-		SetShaderValue(RHICmdList, ShaderRHI, MaskComparison, View.Family->EngineShowFlags.ShaderComplexity ? -1.0f : 0.5f);
-	}
-
-	virtual bool Serialize(FArchive& Ar) override
-	{
-		bool bShaderHasOutdatedParameters = FGlobalShader::Serialize(Ar);
-		Ar << DeferredParameters << MaskComparison;
-		return bShaderHasOutdatedParameters;
-	}
-
-private:
-	FDeferredPixelShaderParameters DeferredParameters;
-	FShaderParameter MaskComparison;
-};
-
-IMPLEMENT_SHADER_TYPE(,FStencilDecalMaskPS,TEXT("DeferredDecal"),TEXT("StencilDecalMaskMain"),SF_Pixel);
-
-FGlobalBoundShaderState StencilDecalMaskBoundShaderState;
-
-/** Draws a full view quad that sets stencil to 1 anywhere that decals should not be projected. */
-void StencilDecalMask(FRHICommandList& RHICmdList, const FViewInfo& View, bool bUseHmdMesh)
-{
-	SCOPED_DRAW_EVENT(RHICmdList, StencilDecalMask);
-	FSceneRenderTargets& SceneContext = FSceneRenderTargets::Get(RHICmdList);
-	SetRenderTarget(RHICmdList, NULL, SceneContext.GetSceneDepthSurface(), ESimpleRenderTargetMode::EUninitializedColorExistingDepth, FExclusiveDepthStencil::DepthRead_StencilWrite);
-	RHICmdList.SetRasterizerState(TStaticRasterizerState<FM_Solid, CM_None>::GetRHI());
-	RHICmdList.SetBlendState(TStaticBlendState<CW_NONE>::GetRHI());
-	RHICmdList.SetViewport(View.ViewRect.Min.X, View.ViewRect.Min.Y, 0.0f, View.ViewRect.Max.X, View.ViewRect.Max.Y, 1.0f);
-
-	// Write 1 to highest bit of stencil to areas that should not receive decals
-	RHICmdList.SetDepthStencilState(TStaticDepthStencilState<false, CF_Always, true, CF_Always, SO_Replace, SO_Replace, SO_Replace>::GetRHI(), 0x80);
-
-	const auto FeatureLevel = View.GetFeatureLevel();
-	auto ShaderMap = View.ShaderMap;
-	TShaderMapRef<FScreenVS> ScreenVertexShader(ShaderMap);
-	TShaderMapRef<FStencilDecalMaskPS> PixelShader(ShaderMap);
-	
-	SetGlobalBoundShaderState(RHICmdList, FeatureLevel, StencilDecalMaskBoundShaderState, GFilterVertexDeclaration.VertexDeclarationRHI, *ScreenVertexShader, *PixelShader);
-
-	PixelShader->SetParameters(RHICmdList, View);
-
-	DrawPostProcessPass(
-		RHICmdList,
-		0, 0,
-		View.ViewRect.Width(), View.ViewRect.Height(),
-		View.ViewRect.Min.X, View.ViewRect.Min.Y,
-		View.ViewRect.Width(), View.ViewRect.Height(),
-		FIntPoint(View.ViewRect.Width(), View.ViewRect.Height()),
-		SceneContext.GetBufferSizeXY(),
-		*ScreenVertexShader,
-		View.StereoPass,
-		bUseHmdMesh,
-		EDRF_UseTriangleOptimization);
-}
-
 bool RenderPreStencil(FRenderingCompositePassContext& Context, const FMatrix& ComponentToWorldMatrix, const FMatrix& FrustumComponentToClip)
 {
 	const FViewInfo& View = Context.View;
@@ -367,13 +282,13 @@ bool RenderPreStencil(FRenderingCompositePassContext& Context, const FMatrix& Co
 		CW_NONE, BO_Add, BF_SourceAlpha, BF_InverseSourceAlpha,		BO_Add, BF_Zero, BF_One		// BaseColor
 	>::GetRHI() );
 
-	// Carmack's reverse on the bounds
+	// Carmack's reverse the sandbox stencil bit on the bounds
 	Context.RHICmdList.SetDepthStencilState(TStaticDepthStencilState<
 		false,CF_LessEqual,
-		true,CF_Equal,SO_Keep,SO_Keep,SO_Increment,
-		true,CF_Equal,SO_Keep,SO_Keep,SO_Decrement,
-		0x80,0x7f
-	>::GetRHI());
+		true,CF_Always,SO_Keep,SO_Keep,SO_Invert,
+		true,CF_Always,SO_Keep,SO_Keep,SO_Invert,
+		STENCIL_SANDBOX_MASK,STENCIL_SANDBOX_MASK
+	>::GetRHI(), 0);
 
 	// Render decal mask
 	Context.RHICmdList.DrawIndexedPrimitive(GetUnitCubeIndexBuffer(), PT_TriangleList, 0, 0, 8, 0, ARRAY_COUNT(GCubeIndices) / 3, 1);
@@ -474,7 +389,7 @@ static void SetDecalDepthState(FDecalDepthState DecalDepthState, FRHICommandList
 				false,CF_Always,
 				true,CF_Equal,SO_Zero,SO_Zero,SO_Zero,
 				true,CF_Equal,SO_Zero,SO_Zero,SO_Zero,
-				0xff, 0x7f>::GetRHI(), 1);
+				0xFF,STENCIL_SANDBOX_MASK>::GetRHI(), STENCIL_SANDBOX_MASK | GET_STENCIL_BIT_MASK(RECEIVE_DECAL, 1));
 			break;
 
 		case DDS_DepthAlways_StencilEqual0:
@@ -483,7 +398,7 @@ static void SetDecalDepthState(FDecalDepthState DecalDepthState, FRHICommandList
 				false,CF_Always,
 				true,CF_Equal,SO_Keep,SO_Keep,SO_Keep,
 				false,CF_Always,SO_Keep,SO_Keep,SO_Keep,
-				0x80,0x00>::GetRHI(), 0);
+				0xFF,0x00>::GetRHI(), GET_STENCIL_BIT_MASK(RECEIVE_DECAL, 1));
 			break;
 
 		case DDS_Always:
@@ -503,7 +418,7 @@ static void SetDecalDepthState(FDecalDepthState DecalDepthState, FRHICommandList
 				false,CF_DepthNearOrEqual,
 				true,CF_Equal,SO_Zero,SO_Zero,SO_Zero,
 				true,CF_Equal,SO_Zero,SO_Zero,SO_Zero,
-				0xff, 0x7f>::GetRHI(), 1);
+				0xFF,STENCIL_SANDBOX_MASK>::GetRHI(), STENCIL_SANDBOX_MASK | GET_STENCIL_BIT_MASK(RECEIVE_DECAL, 1));
 			break;
 
 		case DDS_DepthTest_StencilEqual0:
@@ -512,7 +427,7 @@ static void SetDecalDepthState(FDecalDepthState DecalDepthState, FRHICommandList
 				false,CF_DepthNearOrEqual,
 				true,CF_Equal,SO_Keep,SO_Keep,SO_Keep,
 				false,CF_Always,SO_Keep,SO_Keep,SO_Keep,
-				0x80,0x00>::GetRHI(), 0);
+				0xFF,0x00>::GetRHI(), GET_STENCIL_BIT_MASK(RECEIVE_DECAL, 1));
 			break;
 
 		case DDS_DepthTest:
@@ -652,12 +567,6 @@ void FRCPassPostProcessDeferredDecals::Process(FRenderingCompositePassContext& C
 			}
 #endif
 
-			// Setup a stencil mask to prevent certain pixels from receiving deferred decals
-			if (bStencilDecalsInThisStage)
-			{
-				StencilDecalMask(RHICmdList, View, Context.HasHmdMesh());
-			}
-
 			// optimization to have less state changes
 			EDecalRasterizerState LastDecalRasterizerState = DRS_Undefined;
 			FDecalDepthState LastDecalDepthState;
diff --git a/Engine/Source/Runtime/Renderer/Private/PostProcess/SceneRenderTargets.h b/Engine/Source/Runtime/Renderer/Private/PostProcess/SceneRenderTargets.h
--- a/Engine/Source/Runtime/Renderer/Private/PostProcess/SceneRenderTargets.h
+++ b/Engine/Source/Runtime/Renderer/Private/PostProcess/SceneRenderTargets.h
@@ -103,6 +103,33 @@ BEGIN_UNIFORM_BUFFER_STRUCT(FGBufferResourceStruct, )
 	DECLARE_UNIFORM_BUFFER_STRUCT_MEMBER_SAMPLER( SamplerState,			GBufferVelocityTextureSampler)
 END_UNIFORM_BUFFER_STRUCT( FGBufferResourceStruct )
 
+/*
+* Stencil layout:
+*		BIT ID    | USE
+*		[0]       | sandbox bit
+*		[1]       | unallocated
+*		[2]       | unallocated
+*		[3]       | unallocated
+*		[4]       | unallocated
+*		[5]       | unallocated
+*		[6]       | unallocated
+*		[7]       | primitive receive decal bit
+*
+* The sandbox bit is dedicated to be used by any render passes which would
+* need to generate non persistent stencil and use it locally, but this is
+* the render pass' responsibility to set it back to ZERO.
+*/
+#define STENCIL_SANDBOX_BIT_ID				0
+#define STENCIL_RECEIVE_DECAL_BIT_ID		7
+
+// Outputs a compile-time constant stencil's bit mask ready to be used
+// in TStaticDepthStencilState<> template parameter. It also takes care
+// of masking the Value macro parameter to only keep the low significant
+// bit to ensure to not overflow on other bits.
+#define GET_STENCIL_BIT_MASK(BIT_NAME,Value) uint8((uint8(Value) & uint8(0x01)) << (STENCIL_##BIT_NAME##_BIT_ID))
+
+#define STENCIL_SANDBOX_MASK GET_STENCIL_BIT_MASK(SANDBOX,1)
+
 /**
  * Encapsulates the render targets used for scene rendering.
  */
