Exposes the decal stencil optimization on dbuffer decals

diff --git a/Engine/Source/Runtime/Renderer/Private/CompositionLighting/PostProcessDeferredDecals.cpp b/Engine/Source/Runtime/Renderer/Private/CompositionLighting/PostProcessDeferredDecals.cpp
--- a/Engine/Source/Runtime/Renderer/Private/CompositionLighting/PostProcessDeferredDecals.cpp
+++ b/Engine/Source/Runtime/Renderer/Private/CompositionLighting/PostProcessDeferredDecals.cpp
@@ -13,9 +13,6 @@
 #include "SceneUtils.h"
 #include "DecalRenderingShared.h"
 
-#define DBUFFER_DONT_USE_STENCIL_YET 1
-
-
 static TAutoConsoleVariable<float> CVarStencilSizeThreshold(
 	TEXT("r.Decal.StencilSizeThreshold"),
 	0.1f,
@@ -30,9 +27,11 @@ enum EDecalDepthInputState
 	DDS_Undefined,
 	DDS_Always,
 	DDS_DepthTest,
-	DDS_DepthAlways_StencilEqual1Write,
+	DDS_DepthAlways_StencilEqual1,
+	DDS_DepthAlways_StencilEqual1_IgnoreMask,
 	DDS_DepthAlways_StencilEqual0,
-	DDS_DepthTest_StencilEqual1Write,
+	DDS_DepthTest_StencilEqual1,
+	DDS_DepthTest_StencilEqual1_IgnoreMask,
 	DDS_DepthTest_StencilEqual0,
 };
 
@@ -317,15 +316,15 @@ static EDecalRasterizerState ComputeDecalRasterizerState(bool bInsideDecal, cons
 }
 
 FRCPassPostProcessDeferredDecals::FRCPassPostProcessDeferredDecals(EDecalRenderStage InDecalRenderStage)
-	: DecalRenderStage(InDecalRenderStage)
+	: CurrentStage(InDecalRenderStage)
 {
 }
 
-static FDecalDepthState ComputeDecalDepthState(EDecalBlendMode DecalBlendMode, bool bInsideDecal, bool bStencilDecalsInThisStage, bool bThisDecalUsesStencil)
+static FDecalDepthState ComputeDecalDepthState(EDecalRenderStage LocalDecalStage, bool bInsideDecal, bool bThisDecalUsesStencil)
 {
 	FDecalDepthState Ret;
 
-	Ret.bDepthOutput = DecalBlendMode == DBM_Volumetric_DistanceFunction;
+	Ret.bDepthOutput = (LocalDecalStage == DRS_AfterBasePass);
 				
 	if(Ret.bDepthOutput)
 	{
@@ -334,45 +333,28 @@ static FDecalDepthState ComputeDecalDepthState(EDecalBlendMode DecalBlendMode, b
 		return Ret;
 	}
 
+	const bool bGBufferDecal = LocalDecalStage == DRS_BeforeLighting;
+
 	if(bInsideDecal)
 	{
-		// Render backfaces with depth tests disabled since the camera is inside (or close to inside)
-		if(bStencilDecalsInThisStage)
+		if(bThisDecalUsesStencil)
 		{
-			// Enable stencil testing, only write to pixels with stencil of 0
-			if(bThisDecalUsesStencil)
-			{
-				Ret.DepthTest = DDS_DepthAlways_StencilEqual1Write;
-			}
-			else
-			{
-				Ret.DepthTest = DDS_DepthAlways_StencilEqual0;
-			}
+			Ret.DepthTest = bGBufferDecal ? DDS_DepthAlways_StencilEqual1 : DDS_DepthAlways_StencilEqual1_IgnoreMask;
 		}
 		else
 		{
-			Ret.DepthTest = DDS_Always;
+			Ret.DepthTest = bGBufferDecal ? DDS_DepthAlways_StencilEqual0 : DDS_Always;
 		}
 	}
 	else
 	{
-		// Render frontfaces with depth tests on to get the speedup from HiZ since the camera is outside
-		if(bStencilDecalsInThisStage)
+		if(bThisDecalUsesStencil)
 		{
-			// Render frontfaces with depth tests on to get the speedup from HiZ since the camera is outside
-			// Enable stencil testing, only write to pixels with stencil of 0
-			if(bThisDecalUsesStencil)
-			{
-				Ret.DepthTest = DDS_DepthTest_StencilEqual1Write;
-			}
-			else
-			{
-				Ret.DepthTest = DDS_DepthTest_StencilEqual0;
-			}
+			Ret.DepthTest = bGBufferDecal ? DDS_DepthTest_StencilEqual1 : DDS_DepthTest_StencilEqual1_IgnoreMask;
 		}
 		else
 		{
-			Ret.DepthTest = DDS_DepthTest;
+			Ret.DepthTest = bGBufferDecal ? DDS_DepthTest_StencilEqual0 : DDS_DepthTest;
 		}
 	}
 
@@ -383,13 +365,24 @@ static void SetDecalDepthState(FDecalDepthState DecalDepthState, FRHICommandList
 {
 	switch(DecalDepthState.DepthTest)
 	{
-		case DDS_DepthAlways_StencilEqual1Write:
+		case DDS_DepthAlways_StencilEqual1:
 			check(!DecalDepthState.bDepthOutput);			// todo
 			RHICmdList.SetDepthStencilState(TStaticDepthStencilState<
 				false,CF_Always,
 				true,CF_Equal,SO_Zero,SO_Zero,SO_Zero,
 				true,CF_Equal,SO_Zero,SO_Zero,SO_Zero,
-				0xFF,STENCIL_SANDBOX_MASK>::GetRHI(), STENCIL_SANDBOX_MASK | GET_STENCIL_BIT_MASK(RECEIVE_DECAL, 1));
+				STENCIL_SANDBOX_MASK | GET_STENCIL_BIT_MASK(RECEIVE_DECAL, 1),STENCIL_SANDBOX_MASK>::GetRHI(),
+				STENCIL_SANDBOX_MASK | GET_STENCIL_BIT_MASK(RECEIVE_DECAL, 1));
+			break;
+			
+		case DDS_DepthAlways_StencilEqual1_IgnoreMask:
+			check(!DecalDepthState.bDepthOutput);			// todo
+			RHICmdList.SetDepthStencilState(TStaticDepthStencilState<
+				false,CF_Always,
+				true,CF_Equal,SO_Zero,SO_Zero,SO_Zero,
+				true,CF_Equal,SO_Zero,SO_Zero,SO_Zero,
+				STENCIL_SANDBOX_MASK,STENCIL_SANDBOX_MASK>::GetRHI(),
+				STENCIL_SANDBOX_MASK);
 			break;
 
 		case DDS_DepthAlways_StencilEqual0:
@@ -398,27 +391,33 @@ static void SetDecalDepthState(FDecalDepthState DecalDepthState, FRHICommandList
 				false,CF_Always,
 				true,CF_Equal,SO_Keep,SO_Keep,SO_Keep,
 				false,CF_Always,SO_Keep,SO_Keep,SO_Keep,
-				0xFF,0x00>::GetRHI(), GET_STENCIL_BIT_MASK(RECEIVE_DECAL, 1));
+				STENCIL_SANDBOX_MASK | GET_STENCIL_BIT_MASK(RECEIVE_DECAL, 1),0x00>::GetRHI(),
+				GET_STENCIL_BIT_MASK(RECEIVE_DECAL, 1));
 			break;
 
 		case DDS_Always:
-			if(DecalDepthState.bDepthOutput)
-			{
-				RHICmdList.SetDepthStencilState(TStaticDepthStencilState<true, CF_Always, true>::GetRHI(), 0);
-			}
-			else
-			{
-				RHICmdList.SetDepthStencilState(TStaticDepthStencilState<false, CF_Always, true>::GetRHI(), 0);
-			}
+			check(!DecalDepthState.bDepthOutput);			// todo
+			RHICmdList.SetDepthStencilState(TStaticDepthStencilState<false,CF_Always>::GetRHI());
 			break;
 
-		case DDS_DepthTest_StencilEqual1Write:
+		case DDS_DepthTest_StencilEqual1:
+			check(!DecalDepthState.bDepthOutput);			// todo
+			RHICmdList.SetDepthStencilState(TStaticDepthStencilState<
+				false,CF_DepthNearOrEqual,
+				true,CF_Equal,SO_Zero,SO_Zero,SO_Zero,
+				true,CF_Equal,SO_Zero,SO_Zero,SO_Zero,
+				STENCIL_SANDBOX_MASK | GET_STENCIL_BIT_MASK(RECEIVE_DECAL, 1),STENCIL_SANDBOX_MASK>::GetRHI(),
+				STENCIL_SANDBOX_MASK | GET_STENCIL_BIT_MASK(RECEIVE_DECAL, 1));
+			break;
+			
+		case DDS_DepthTest_StencilEqual1_IgnoreMask:
 			check(!DecalDepthState.bDepthOutput);			// todo
 			RHICmdList.SetDepthStencilState(TStaticDepthStencilState<
 				false,CF_DepthNearOrEqual,
 				true,CF_Equal,SO_Zero,SO_Zero,SO_Zero,
 				true,CF_Equal,SO_Zero,SO_Zero,SO_Zero,
-				0xFF,STENCIL_SANDBOX_MASK>::GetRHI(), STENCIL_SANDBOX_MASK | GET_STENCIL_BIT_MASK(RECEIVE_DECAL, 1));
+				STENCIL_SANDBOX_MASK,STENCIL_SANDBOX_MASK>::GetRHI(),
+				STENCIL_SANDBOX_MASK);
 			break;
 
 		case DDS_DepthTest_StencilEqual0:
@@ -427,7 +426,8 @@ static void SetDecalDepthState(FDecalDepthState DecalDepthState, FRHICommandList
 				false,CF_DepthNearOrEqual,
 				true,CF_Equal,SO_Keep,SO_Keep,SO_Keep,
 				false,CF_Always,SO_Keep,SO_Keep,SO_Keep,
-				0xFF,0x00>::GetRHI(), GET_STENCIL_BIT_MASK(RECEIVE_DECAL, 1));
+				STENCIL_SANDBOX_MASK | GET_STENCIL_BIT_MASK(RECEIVE_DECAL, 1),0x00>::GetRHI(),
+				GET_STENCIL_BIT_MASK(RECEIVE_DECAL, 1));
 			break;
 
 		case DDS_DepthTest:
@@ -456,6 +456,11 @@ static void SetDecalRasterizerState(EDecalRasterizerState DecalRasterizerState,
 	}
 }
 
+static inline bool IsStencilOptimizationAvailable(EDecalRenderStage RenderStage)
+{
+	return RenderStage == DRS_BeforeLighting || RenderStage == DRS_BeforeBasePass;
+}
+
 void FRCPassPostProcessDeferredDecals::Process(FRenderingCompositePassContext& Context)
 {
 	FRHICommandListImmediate& RHICmdList = Context.RHICmdList;
@@ -481,7 +486,7 @@ void FRCPassPostProcessDeferredDecals::Process(FRenderingCompositePassContext& C
 
 	FTextureRHIParamRef TargetsToResolve[ResolveBufferMax] = { nullptr };
 
-	if(DecalRenderStage == DRS_BeforeBasePass)
+	if(CurrentStage == DRS_BeforeBasePass)
 	{
 		// before BasePass, only if DBuffer is enabled
 
@@ -551,22 +556,13 @@ void FRCPassPostProcessDeferredDecals::Process(FRenderingCompositePassContext& C
 	{
 		// Build a list of decals that need to be rendered for this view
 		FTransientDecalRenderDataList SortedDecals;
-		FDecalRendering::BuildVisibleDecalList(Scene, View, DecalRenderStage, SortedDecals);
+		FDecalRendering::BuildVisibleDecalList(Scene, View, CurrentStage, SortedDecals);
 
 		if (SortedDecals.Num() > 0)
 		{
 			FIntRect SrcRect = View.ViewRect;
 			FIntRect DestRect = View.ViewRect;
 
-			bool bStencilDecalsInThisStage = true;
-
-#if DBUFFER_DONT_USE_STENCIL_YET
-			if (DecalRenderStage != DRS_BeforeLighting)
-			{
-				bStencilDecalsInThisStage = false;
-			}
-#endif
-
 			// optimization to have less state changes
 			EDecalRasterizerState LastDecalRasterizerState = DRS_Undefined;
 			FDecalDepthState LastDecalDepthState;
@@ -586,14 +582,8 @@ void FRCPassPostProcessDeferredDecals::Process(FRenderingCompositePassContext& C
 				const FMatrix FrustumComponentToClip = FDecalRendering::ComputeComponentToClipMatrix(View, ComponentToWorldMatrix);
 
 				EDecalBlendMode DecalBlendMode = DecalData.DecalBlendMode;
-				bool bStencilThisDecal = bStencilDecalsInThisStage;
-
-#if DBUFFER_DONT_USE_STENCIL_YET
-				if (FDecalRendering::ComputeRenderStage(View.GetShaderPlatform(), DecalBlendMode) != DRS_BeforeLighting)
-				{
-					bStencilThisDecal = false;
-				}
-#endif				
+				EDecalRenderStage LocalDecalStage = FDecalRendering::ComputeRenderStage(View.GetShaderPlatform(), DecalBlendMode);
+				bool bStencilThisDecal = IsStencilOptimizationAvailable(LocalDecalStage);
 
 				FDecalRendering::ERenderTargetMode CurrentRenderTargetMode = FDecalRendering::ComputeRenderTargetMode(View.GetShaderPlatform(), DecalBlendMode);
 
@@ -664,17 +654,14 @@ void FRCPassPostProcessDeferredDecals::Process(FRenderingCompositePassContext& C
 
 				bool bThisDecalUsesStencil = false;
 
-				if (bStencilThisDecal)
+				if (bStencilThisDecal && bStencilSizeThreshold)
 				{
-					if (bStencilSizeThreshold)
-					{
-						// note this is after a SetStreamSource (in if CurrentRenderTargetMode != LastRenderTargetMode) call as it needs to get the VB input
-						bThisDecalUsesStencil = RenderPreStencil(Context, ComponentToWorldMatrix, FrustumComponentToClip);
+					// note this is after a SetStreamSource (in if CurrentRenderTargetMode != LastRenderTargetMode) call as it needs to get the VB input
+					bThisDecalUsesStencil = RenderPreStencil(Context, ComponentToWorldMatrix, FrustumComponentToClip);
 
-						LastDecalRasterizerState = DRS_Undefined;
-						LastDecalDepthState = FDecalDepthState();
-						LastDecalBlendMode = -1;
-					}
+					LastDecalRasterizerState = DRS_Undefined;
+					LastDecalDepthState = FDecalDepthState();
+					LastDecalBlendMode = -1;
 				}
 
 				const bool bBlendStateChange = DecalBlendMode != LastDecalBlendMode;// Has decal mode changed.
@@ -688,7 +675,7 @@ void FRCPassPostProcessDeferredDecals::Process(FRenderingCompositePassContext& C
 					LastDecalBlendMode = DecalBlendMode;
 					LastDecalHasNormal = (int32)DecalData.bHasNormal;
 
-					SetDecalBlendState(RHICmdList, SMFeatureLevel, DecalRenderStage, (EDecalBlendMode)LastDecalBlendMode, DecalData.bHasNormal);
+					SetDecalBlendState(RHICmdList, SMFeatureLevel, CurrentStage, (EDecalBlendMode)LastDecalBlendMode, DecalData.bHasNormal);
 				}
 
 
@@ -711,7 +698,7 @@ void FRCPassPostProcessDeferredDecals::Process(FRenderingCompositePassContext& C
 
 				// update DepthStencil state if needed
 				{
-					FDecalDepthState DecalDepthState = ComputeDecalDepthState(DecalBlendMode, bInsideDecal, bStencilDecalsInThisStage, bThisDecalUsesStencil);
+					FDecalDepthState DecalDepthState = ComputeDecalDepthState(LocalDecalStage, bInsideDecal, bThisDecalUsesStencil);
 
 					if (LastDecalDepthState != DecalDepthState)
 					{
@@ -729,7 +716,7 @@ void FRCPassPostProcessDeferredDecals::Process(FRenderingCompositePassContext& C
 			// Clear stencil to 0, which is the assumed default by other passes
 			RHICmdList.Clear(false, FLinearColor::White, false, (float)ERHIZBuffer::FarPlane, true, 0, FIntRect());
 
-			if (DecalRenderStage == DRS_BeforeBasePass)
+			if (CurrentStage == DRS_BeforeBasePass)
 			{
 				// before BasePass
 				GRenderTargetPool.VisualizeTexture.SetCheckPoint(RHICmdList, SceneContext.DBufferA);
diff --git a/Engine/Source/Runtime/Renderer/Private/CompositionLighting/PostProcessDeferredDecals.h b/Engine/Source/Runtime/Renderer/Private/CompositionLighting/PostProcessDeferredDecals.h
--- a/Engine/Source/Runtime/Renderer/Private/CompositionLighting/PostProcessDeferredDecals.h
+++ b/Engine/Source/Runtime/Renderer/Private/CompositionLighting/PostProcessDeferredDecals.h
@@ -24,7 +24,7 @@ public:
 
 private:
 	// see EDecalRenderStage
-	EDecalRenderStage DecalRenderStage;
+	EDecalRenderStage CurrentStage;
 };
 
 bool IsDBufferEnabled();
