CPU optimized PostProcessMaterial blending, GetReusableMID() makes a strong commitment on the returned data which requires less fixup, by knowing the param hierarchy we don't need to copy all paramaters, only the one from the last MID parent. The blueprint code for material copying can be optimized but it will be less general or needs a different implementation - that can be done later.

diff --git a/Engine/Source/Runtime/Engine/Classes/Materials/MaterialInstanceDynamic.h b/Engine/Source/Runtime/Engine/Classes/Materials/MaterialInstanceDynamic.h
--- a/Engine/Source/Runtime/Engine/Classes/Materials/MaterialInstanceDynamic.h
+++ b/Engine/Source/Runtime/Engine/Classes/Materials/MaterialInstanceDynamic.h
@@ -47,12 +47,22 @@ class UMaterialInstanceDynamic : public UMaterialInstance
 	void K2_InterpolateMaterialInstanceParams(UMaterialInstance* SourceA, UMaterialInstance* SourceB, float Alpha);
 
 	/**
-	 * Copies over parameters given a material interface.
+	 * Copies over parameters given a material interface (deep, following the hierarchy)
+	 * Very slow implementation, avoid using at runtime. Hopefully we can replace ity later with somthing like CopyInterpParameters()
 	 * The output is the object itself (this).
 	 */
 	UFUNCTION(BlueprintCallable, meta=(DisplayName = "CopyMaterialInstanceParameters"), Category="Rendering|Material")
 	void K2_CopyMaterialInstanceParameters(UMaterialInterface* Source);
 
+	/**
+	 * Copies over parameters given a material instance (shallow, not deep following the hierarchy)
+	 * much faster than K2_CopyMaterialInstanceParameters(), 
+	 * The output is the object itself (this).
+	 * @param Source ignores the call if 0
+	 */
+	UFUNCTION(meta=(DisplayName = "CopyMaterialInstanceParameters"), Category="Rendering|Material")
+	ENGINE_API void CopyInterpParameters(UMaterialInstance* Source);
+
 	/**
 	 * Create a material instance dynamic parented to the specified material.
 	 */
diff --git a/Engine/Source/Runtime/Engine/Private/Materials/Material.cpp b/Engine/Source/Runtime/Engine/Private/Materials/Material.cpp
--- a/Engine/Source/Runtime/Engine/Private/Materials/Material.cpp
+++ b/Engine/Source/Runtime/Engine/Private/Materials/Material.cpp
@@ -1145,8 +1145,6 @@ void UMaterialInterface::OverrideBlendableSettings(class FSceneView& View, float
 
 		if(MID)
 		{
-			MID->K2_CopyMaterialInstanceParameters((UMaterialInterface*)this);
-
 			FPostProcessMaterialNode NewNode(MID, Material->BlendableLocation, Material->BlendablePriority);
 
 			// a material already exists, blend with existing ones
diff --git a/Engine/Source/Runtime/Engine/Private/Materials/MaterialInstance.cpp b/Engine/Source/Runtime/Engine/Private/Materials/MaterialInstance.cpp
--- a/Engine/Source/Runtime/Engine/Private/Materials/MaterialInstance.cpp
+++ b/Engine/Source/Runtime/Engine/Private/Materials/MaterialInstance.cpp
@@ -2329,12 +2329,10 @@ void UMaterialInstance::OverrideBlendableSettings(class FSceneView& View, float
 		{
 			UMaterial* Base = Material->GetBaseMaterial();
 
-			UMaterialInstanceDynamic* MID = View.State->GetReusableMID((UMaterialInterface*)this);
+			UMaterialInstanceDynamic* MID = View.State->GetReusableMID(Base);
 
 			if(MID)
 			{
-				MID->K2_CopyMaterialInstanceParameters((UMaterialInterface*)Base);
-
 				FPostProcessMaterialNode NewNode(MID, Base->BlendableLocation, Base->BlendablePriority);
 
 				// it's the first material, no blending needed
@@ -2364,7 +2362,6 @@ void UMaterialInstance::OverrideBlendableSettings(class FSceneView& View, float
 
 		if(MID)
 		{
-			MID->K2_CopyMaterialInstanceParameters((UMaterialInterface*)this);
 			SetTonemapperPostprocessMaterialSettings(View, *MID);
 
 			FPostProcessMaterialNode NewNode(MID, Material->BlendableLocation, Material->BlendablePriority);
diff --git a/Engine/Source/Runtime/Engine/Private/Materials/MaterialInstanceDynamic.cpp b/Engine/Source/Runtime/Engine/Private/Materials/MaterialInstanceDynamic.cpp
--- a/Engine/Source/Runtime/Engine/Private/Materials/MaterialInstanceDynamic.cpp
+++ b/Engine/Source/Runtime/Engine/Private/Materials/MaterialInstanceDynamic.cpp
@@ -215,6 +215,27 @@ void UMaterialInstanceDynamic::K2_CopyMaterialInstanceParameters(UMaterialInterf
 	CopyMaterialInstanceParameters(Source);
 }
 
+void UMaterialInstanceDynamic::CopyInterpParameters(UMaterialInstance* Source)
+{
+	// we might expose as blueprint function so we have the input a pointer instead of a reference
+	if(Source)
+	{
+		// copy the array and update the renderer data structures
+
+		for (auto& it : Source->ScalarParameterValues)
+		{
+			SetScalarParameterValue(it.ParameterName, it.ParameterValue);
+		}
+
+		for (auto& it : Source->VectorParameterValues)
+		{
+			SetVectorParameterValue(it.ParameterName, it.ParameterValue);
+		}
+
+		// the other parameters don't interpolate but if we expose this to blueprints we could copy them for consistency
+	}
+}
+
 void UMaterialInstanceDynamic::CopyParameterOverrides(UMaterialInstance* MaterialInstance)
 {
 	ClearParameterValues();
diff --git a/Engine/Source/Runtime/Engine/Public/SceneManagement.h b/Engine/Source/Runtime/Engine/Public/SceneManagement.h
--- a/Engine/Source/Runtime/Engine/Public/SceneManagement.h
+++ b/Engine/Source/Runtime/Engine/Public/SceneManagement.h
@@ -141,9 +141,9 @@ public:
 	/**
 	 * Allows MIDs being created and released during view rendering without the overhead of creating and releasing objects
 	 * As MID are not allowed to be parent of MID this gets fixed up by parenting it to the next Material or MIC
-	 * @param ParentMaterial can be Material, MIC or MID, must not be 0
+	 * @param InSource can be Material, MIC or MID, must not be 0
 	 */
-	virtual UMaterialInstanceDynamic* GetReusableMID(class UMaterialInterface* ParentMaterial) = 0;
+	virtual UMaterialInstanceDynamic* GetReusableMID(class UMaterialInterface* InSource) = 0;
 	/** Returns the temporal LOD struct from the viewstate */
 	virtual FTemporalLODState& GetTemporalLODState() = 0;
 	virtual const FTemporalLODState& GetTemporalLODState() const = 0;
diff --git a/Engine/Source/Runtime/Renderer/Private/ScenePrivate.h b/Engine/Source/Runtime/Renderer/Private/ScenePrivate.h
--- a/Engine/Source/Runtime/Renderer/Private/ScenePrivate.h
+++ b/Engine/Source/Runtime/Renderer/Private/ScenePrivate.h
@@ -779,45 +779,57 @@ public:
 	}
 
 	// Note: OnStartPostProcessing() needs to be called each frame for each view
-	virtual UMaterialInstanceDynamic* GetReusableMID(class UMaterialInterface* InParentMaterial) override
+	virtual UMaterialInstanceDynamic* GetReusableMID(class UMaterialInterface* InSource) override
 	{		
 		check(IsInGameThread());
-		check(InParentMaterial);
+		check(InSource);
 
-		auto ParentAsMaterialInstance = Cast<UMaterialInstanceDynamic>(InParentMaterial);
+		// 0 or MID (MaterialInstanceDynamic) pointer
+		auto InputAsMID = Cast<UMaterialInstanceDynamic>(InSource);
 
 		// fixup MID parents as this is not allowed, take the next MIC or Material.
-		UMaterialInterface* ParentMaterial = ParentAsMaterialInstance ? ParentAsMaterialInstance->Parent : InParentMaterial;
+		UMaterialInterface* ParentOfTheNewMID = InputAsMID ? InputAsMID->Parent : InSource;
 
 		// this is not allowed and would cause an error later in the code
-		check(!ParentMaterial->IsA(UMaterialInstanceDynamic::StaticClass()));
+		check(!ParentOfTheNewMID->IsA(UMaterialInstanceDynamic::StaticClass()));
+
+		UMaterialInstanceDynamic* NewMID = 0;
 
 		if(MIDUsedCount < (uint32)MIDPool.Num())
 		{
-			UMaterialInstanceDynamic* MID = MIDPool[MIDUsedCount];
+			NewMID = MIDPool[MIDUsedCount];
 
-			if(MID->Parent != ParentMaterial)
+			if(NewMID->Parent != ParentOfTheNewMID)
 			{
 				// create a new one
 				// garbage collector will remove the old one
 				// this should not happen too often
-				MID = UMaterialInstanceDynamic::Create(ParentMaterial, 0);
-				MIDPool[MIDUsedCount] = MID;
+				NewMID = UMaterialInstanceDynamic::Create(ParentOfTheNewMID, 0);
+				MIDPool[MIDUsedCount] = NewMID;
 			}
+
+			// reusing an existing object means we need to clear out the Vector and Scalar parameters
+			NewMID->ClearParameterValues();
 		}
 		else
 		{
-			UMaterialInstanceDynamic* MID = UMaterialInstanceDynamic::Create(ParentMaterial, 0);
-			check(MID);
+			NewMID = UMaterialInstanceDynamic::Create(ParentOfTheNewMID, 0);
+			check(NewMID);
 
-			MIDPool.Add(MID);
+			MIDPool.Add(NewMID);
 		}
 
-		UMaterialInstanceDynamic* Ret = MIDPool[MIDUsedCount++];
+		if(InputAsMID)
+		{
+			auto ParentAsMI = Cast<UMaterialInstance>(InputAsMID->Parent);
+			check(ParentAsMI);
 
-		check(Ret->GetRenderProxy(false));
+			// parent is an MID so we need to copy the MID Vector and Scalar parameters over
+			NewMID->CopyInterpParameters(ParentAsMI);
+		}
 
-		return Ret;
+		check(NewMID->GetRenderProxy(false));
+		return NewMID;
 	}
 
 	virtual FTemporalLODState& GetTemporalLODState() override
