Integrated from UE4-Orion:
 CL 2674843
	better debug/profile output for PostProcessAO

diff --git a/Engine/Source/Runtime/Renderer/Private/CompositionLighting/PostProcessAmbientOcclusion.cpp b/Engine/Source/Runtime/Renderer/Private/CompositionLighting/PostProcessAmbientOcclusion.cpp
--- a/Engine/Source/Runtime/Renderer/Private/CompositionLighting/PostProcessAmbientOcclusion.cpp
+++ b/Engine/Source/Runtime/Renderer/Private/CompositionLighting/PostProcessAmbientOcclusion.cpp
@@ -405,7 +405,7 @@ bool FRCPassPostProcessAmbientOcclusionSetup::IsInitialPass() const
 
 // ---------------------------------
 template <uint32 bTAOSetupAsInput, uint32 bDoUpsample, uint32 SampleSetQuality>
-void FRCPassPostProcessAmbientOcclusion::SetShaderTempl(const FRenderingCompositePassContext& Context)
+FShader* FRCPassPostProcessAmbientOcclusion::SetShaderTempl(const FRenderingCompositePassContext& Context)
 {
 	TShaderMapRef<FPostProcessVS> VertexShader(Context.GetShaderMap());
 	TShaderMapRef<FPostProcessAmbientOcclusionPS<bTAOSetupAsInput, bDoUpsample, SampleSetQuality> > PixelShader(Context.GetShaderMap());
@@ -419,12 +419,12 @@ void FRCPassPostProcessAmbientOcclusion::SetShaderTempl(const FRenderingComposit
 
 	VertexShader->SetParameters(Context);
 	PixelShader->SetParameters(Context, InputDesc0->Extent);
+
+	return *VertexShader;
 }
 
 void FRCPassPostProcessAmbientOcclusion::Process(FRenderingCompositePassContext& Context)
 {
-	SCOPED_DRAW_EVENT(Context.RHICmdList, AmbientOcclusion);
-
 	const FSceneView& View = Context.View;
 
 	const FPooledRenderTargetDesc* InputDesc0 = GetInputDesc(ePId_Input0);
@@ -465,18 +465,23 @@ void FRCPassPostProcessAmbientOcclusion::Process(FRenderingCompositePassContext&
 		(QualityPercent > 60.0f) +
 		(QualityPercent > 25.0f);
 	bool bDoUpsample = (InputDesc2 != 0);
+	
+	SCOPED_DRAW_EVENTF(Context.RHICmdList, AmbientOcclusion, TEXT("AmbientOcclusion %dx%d SetupAsInput=%d Upsample=%d SampleSetQuality=%d"), 
+		ViewRect.Width(), ViewRect.Height(), bAOSetupAsInput, bDoUpsample, QualitySet);
+
+	FShader* VertexShader = 0;
 
 #define SET_SHADER_CASE(Quality)                                \
 	case Quality:			                                    \
 	if(bAOSetupAsInput)                                         \
     {                                                           \
-		if(bDoUpsample) SetShaderTempl<1, 1, Quality>(Context); \
-		else SetShaderTempl<1, 0, Quality>(Context);            \
+		if(bDoUpsample) VertexShader = SetShaderTempl<1, 1, Quality>(Context); \
+		else VertexShader = SetShaderTempl<1, 0, Quality>(Context);            \
 	}                                                           \
 	else                                                        \
     {                                                           \
-		if(bDoUpsample) SetShaderTempl<0, 1, Quality>(Context); \
-		else SetShaderTempl<0, 0, Quality>(Context);            \
+		if(bDoUpsample) VertexShader = SetShaderTempl<0, 1, Quality>(Context); \
+		else VertexShader = SetShaderTempl<0, 0, Quality>(Context);            \
 	}                                                           \
 	break
 
@@ -491,8 +496,6 @@ void FRCPassPostProcessAmbientOcclusion::Process(FRenderingCompositePassContext&
 
 #undef SET_SHADER_CASE
 
-	TShaderMapRef<FPostProcessVS> VertexShader(Context.GetShaderMap());
-
 	DrawPostProcessPass(
 		Context.RHICmdList,
 		0, 0,
@@ -501,7 +504,7 @@ void FRCPassPostProcessAmbientOcclusion::Process(FRenderingCompositePassContext&
 		ViewRect.Width(), ViewRect.Height(),
 		ViewRect.Size(),
 		TexSize,
-		*VertexShader,
+		VertexShader,
 		View.StereoPass,
 		Context.HasHmdMesh(),
 		EDRF_UseTriangleOptimization);
@@ -648,4 +651,4 @@ TUniformBufferRef<FCameraMotionParameters> CreateCameraMotionParametersUniformBu
 		(float)(0));
 
 	return TUniformBufferRef<FCameraMotionParameters>::CreateUniformBufferImmediate(LocalCameraMotion, UniformBuffer_SingleFrame);
-}
\ No newline at end of file
+}
diff --git a/Engine/Source/Runtime/Renderer/Private/CompositionLighting/PostProcessAmbientOcclusion.h b/Engine/Source/Runtime/Renderer/Private/CompositionLighting/PostProcessAmbientOcclusion.h
--- a/Engine/Source/Runtime/Renderer/Private/CompositionLighting/PostProcessAmbientOcclusion.h
+++ b/Engine/Source/Runtime/Renderer/Private/CompositionLighting/PostProcessAmbientOcclusion.h
@@ -51,7 +51,7 @@ public:
 private:
 	
 	template <uint32 bAOSetupAsInput, uint32 bDoUpsample, uint32 SampleSetQuality>
-		void SetShaderTempl(const FRenderingCompositePassContext& Context);
+		FShader* SetShaderTempl(const FRenderingCompositePassContext& Context);
 
 	bool bAOSetupAsInput;
 };
