integrated from Orion

CL 2676126

	added more useful GPU profile into to many passes

diff --git a/Engine/Source/Runtime/RHI/Private/GPUProfiler.cpp b/Engine/Source/Runtime/RHI/Private/GPUProfiler.cpp
--- a/Engine/Source/Runtime/RHI/Private/GPUProfiler.cpp
+++ b/Engine/Source/Runtime/RHI/Private/GPUProfiler.cpp
@@ -92,15 +92,25 @@ static void DumpStatsEventNode(FGPUProfilerEventNode* Node, float RootResult, in
 
 		if (bMatchesFilter)
 		{
+			FString Extra;
+			
+			if(Node->TimingResult >= 0.1f && Node->NumVertices * Node->NumDraws > 100)
+			{
+				Extra = FString::Printf(TEXT(" %.0f prims/ms %.0f verts/ms"),
+					Node->NumPrimitives / Node->TimingResult,
+					Node->NumVertices / Node->TimingResult);
+			}
+
 			// Print information about this node, padded to its depth in the tree
-			UE_LOG(LogRHI, Warning, TEXT("%s%4.1f%%%5.2fms   %s %u draws %u prims %u verts"), 
+			UE_LOG(LogRHI, Warning, TEXT("%s%4.1f%%%5.2fms   %s %u draws %u prims %u verts%s"), 
 				*FString(TEXT("")).LeftPad(EffectiveDepth * 3), 
 				Percent,
 				Node->TimingResult,
 				*Node->Name,
 				Node->NumDraws,
 				Node->NumPrimitives,
-				Node->NumVertices
+				Node->NumVertices,
+				*Extra
 				);
 		}
 
diff --git a/Engine/Source/Runtime/Renderer/Private/AtmosphereRendering.cpp b/Engine/Source/Runtime/Renderer/Private/AtmosphereRendering.cpp
--- a/Engine/Source/Runtime/Renderer/Private/AtmosphereRendering.cpp
+++ b/Engine/Source/Runtime/Renderer/Private/AtmosphereRendering.cpp
@@ -387,8 +387,6 @@ void FDeferredShadingSceneRenderer::RenderAtmosphere(FRHICommandListImmediate& R
 	// Atmospheric fog?
 	if (Scene->GetFeatureLevel() >= ERHIFeatureLevel::SM4 && Scene->HasAtmosphericFog())
 	{
-		SCOPED_DRAW_EVENT(RHICmdList, AtmosphericFog);
-
 		static const FVector2D Vertices[4] =
 		{
 			FVector2D(-1,-1),
@@ -409,6 +407,8 @@ void FDeferredShadingSceneRenderer::RenderAtmosphere(FRHICommandListImmediate& R
 		{
 			const FViewInfo& View = Views[ViewIndex];
 
+			SCOPED_DRAW_EVENTF(RHICmdList, Atmosphere, TEXT("Atmosphere %dx%d"), View.ViewRect.Width(), View.ViewRect.Height());
+
 			// Set the device viewport for the view.
 			RHICmdList.SetViewport(View.ViewRect.Min.X, View.ViewRect.Min.Y, 0.0f, View.ViewRect.Max.X, View.ViewRect.Max.Y, 1.0f);
 			
diff --git a/Engine/Source/Runtime/Renderer/Private/CompositionLighting/PostProcessAmbientOcclusion.cpp b/Engine/Source/Runtime/Renderer/Private/CompositionLighting/PostProcessAmbientOcclusion.cpp
--- a/Engine/Source/Runtime/Renderer/Private/CompositionLighting/PostProcessAmbientOcclusion.cpp
+++ b/Engine/Source/Runtime/Renderer/Private/CompositionLighting/PostProcessAmbientOcclusion.cpp
@@ -308,8 +308,6 @@ FShader* FRCPassPostProcessAmbientOcclusionSetup::SetShaderSetupTempl(const FRen
 
 void FRCPassPostProcessAmbientOcclusionSetup::Process(FRenderingCompositePassContext& Context)
 {
-	SCOPED_DRAW_EVENT(Context.RHICmdList, AmbientOcclusionSetup);
-
 	const FSceneView& View = Context.View;
 
 	const FSceneRenderTargetItem& DestRenderTarget = PassOutputs[0].RequestSurface(Context);
@@ -322,6 +320,8 @@ void FRCPassPostProcessAmbientOcclusionSetup::Process(FRenderingCompositePassCon
 	FIntRect SrcRect = View.ViewRect;
 	FIntRect DestRect = SrcRect  / ScaleFactor;
 
+	SCOPED_DRAW_EVENTF(Context.RHICmdList, AmbientOcclusionSetup, TEXT("AmbientOcclusionSetup %dx%d"), DestRect.Width(), DestRect.Height());
+
 	// Set the view family's render target/viewport.
 	SetRenderTarget(Context.RHICmdList, DestRenderTarget.TargetableTexture, FTextureRHIParamRef());
 
@@ -347,7 +347,7 @@ void FRCPassPostProcessAmbientOcclusionSetup::Process(FRenderingCompositePassCon
 		Context.RHICmdList,
 		0, 0,
 		DestRect.Width(), DestRect.Height(),
-		SrcRect.Min.X, SrcRect.Min.Y,
+		SrcRect.Min.X, SrcRect.Min.Y, 
 		SrcRect.Width(), SrcRect.Height(),
 		DestRect.Size(),
 		FSceneRenderTargets::Get(Context.RHICmdList).GetBufferSizeXY(),
@@ -540,9 +540,10 @@ FPooledRenderTargetDesc FRCPassPostProcessAmbientOcclusion::ComputeOutputDesc(EP
 
 void FRCPassPostProcessBasePassAO::Process(FRenderingCompositePassContext& Context)
 {
-	SCOPED_DRAW_EVENT(Context.RHICmdList, ApplyAOToBasePassSceneColor);
-
 	const FSceneView& View = Context.View;
+
+	SCOPED_DRAW_EVENTF(Context.RHICmdList, ApplyAOToBasePassSceneColor, TEXT("ApplyAOToBasePassSceneColor %dx%d"), View.ViewRect.Width(), View.ViewRect.Height());
+
 	FSceneRenderTargets& SceneContext = FSceneRenderTargets::Get(Context.RHICmdList);
 
 	const FSceneRenderTargetItem& DestRenderTarget = SceneContext.GetSceneColor()->GetRenderTargetItem();
diff --git a/Engine/Source/Runtime/Renderer/Private/FogRendering.cpp b/Engine/Source/Runtime/Renderer/Private/FogRendering.cpp
--- a/Engine/Source/Runtime/Renderer/Private/FogRendering.cpp
+++ b/Engine/Source/Runtime/Renderer/Private/FogRendering.cpp
@@ -298,8 +298,6 @@ bool FDeferredShadingSceneRenderer::RenderFog(FRHICommandListImmediate& RHICmdLi
 {
 	if (Scene->ExponentialFogs.Num() > 0)
 	{
-		SCOPED_DRAW_EVENT(RHICmdList, Fog);
-
 		static const FVector2D Vertices[4] =
 		{
 			FVector2D(-1,-1),
@@ -319,6 +317,8 @@ bool FDeferredShadingSceneRenderer::RenderFog(FRHICommandListImmediate& RHICmdLi
 		{
 			const FViewInfo& View = Views[ViewIndex];
 
+			SCOPED_DRAW_EVENTF(RHICmdList, Fog, TEXT("Fog %dx%d"), View.ViewRect.Width(), View.ViewRect.Height());
+
 			if (View.IsPerspectiveProjection() == false)
 			{
 				continue; // Do not render exponential fog in orthographic views.
diff --git a/Engine/Source/Runtime/Renderer/Private/LightShaftRendering.cpp b/Engine/Source/Runtime/Renderer/Private/LightShaftRendering.cpp
--- a/Engine/Source/Runtime/Renderer/Private/LightShaftRendering.cpp
+++ b/Engine/Source/Runtime/Renderer/Private/LightShaftRendering.cpp
@@ -735,8 +735,6 @@ FLightShaftsOutput FDeferredShadingSceneRenderer::RenderLightShaftOcclusion(FRHI
 
 			if (bEnableOcclusion && LightSceneInfo->Proxy->GetLightType() == LightType_Directional)
 			{
-				SCOPED_DRAW_EVENT(RHICmdList, RenderLightShaftOcclusion);
-
 				// Allocate light shaft render targets on demand, using the pool
 				// Need two targets to ping pong between
 				AllocateOrReuseLightShaftRenderTarget(RHICmdList, LightShafts0, TEXT("LightShafts0"));
@@ -746,6 +744,8 @@ FLightShaftsOutput FDeferredShadingSceneRenderer::RenderLightShaftOcclusion(FRHI
 				{
 					FViewInfo& View = Views[ViewIndex];
 					
+					SCOPED_DRAW_EVENTF(RHICmdList, RenderLightShaftOcclusion, TEXT("RenderLightShaftOcclusion %dx%d (multiple passes)"), View.ViewRect.Width(), View.ViewRect.Height());
+
 					if (ShouldRenderLightShaftsForLight(View, LightSceneInfo))
 					{
 						INC_DWORD_STAT(STAT_LightShaftsLights);
@@ -913,8 +913,6 @@ void FDeferredShadingSceneRenderer::RenderLightShaftBloom(FRHICommandListImmedia
 
 			if (LightSceneInfo->bEnableLightShaftBloom)
 			{
-				SCOPED_DRAW_EVENT(RHICmdList, RenderLightShaftBloom);
-
 				// Allocate light shaft render targets on demand, using the pool
 				AllocateOrReuseLightShaftRenderTarget(RHICmdList, LightShafts0, TEXT("LightShafts0"));
 				AllocateOrReuseLightShaftRenderTarget(RHICmdList, LightShafts1, TEXT("LightShafts1"));
@@ -923,6 +921,8 @@ void FDeferredShadingSceneRenderer::RenderLightShaftBloom(FRHICommandListImmedia
 				{
 					FViewInfo& View = Views[ViewIndex];
 
+					SCOPED_DRAW_EVENTF(RHICmdList, RenderLightShaftBloom, TEXT("RenderLightShaftBloom %dx%d"), View.ViewRect.Width(), View.ViewRect.Height());
+
 					if (ShouldRenderLightShaftsForLight(View, LightSceneInfo))
 					{
 						INC_DWORD_STAT(STAT_LightShaftsLights);
diff --git a/Engine/Source/Runtime/Renderer/Private/PostProcess/PostProcessBloomSetup.cpp b/Engine/Source/Runtime/Renderer/Private/PostProcess/PostProcessBloomSetup.cpp
--- a/Engine/Source/Runtime/Renderer/Private/PostProcess/PostProcessBloomSetup.cpp
+++ b/Engine/Source/Runtime/Renderer/Private/PostProcess/PostProcessBloomSetup.cpp
@@ -138,8 +138,6 @@ IMPLEMENT_SHADER_TYPE(,FPostProcessBloomSetupVS,TEXT("PostProcessBloom"),TEXT("M
 
 void FRCPassPostProcessBloomSetup::Process(FRenderingCompositePassContext& Context)
 {
-	SCOPED_DRAW_EVENT(Context.RHICmdList, PostProcessBloomSetup);
-
 	const FPooledRenderTargetDesc* InputDesc = GetInputDesc(ePId_Input0);
 
 	if(!InputDesc)
@@ -160,6 +158,8 @@ void FRCPassPostProcessBloomSetup::Process(FRenderingCompositePassContext& Conte
 	FIntRect SrcRect = View.ViewRect / ScaleFactor;
 	FIntRect DestRect = SrcRect;
 
+	SCOPED_DRAW_EVENTF(Context.RHICmdList, PostProcessBloomSetup, TEXT("PostProcessBloomSetup %dx%d"), DestRect.Width(), DestRect.Height());
+
 	const FSceneRenderTargetItem& DestRenderTarget = PassOutputs[0].RequestSurface(Context);
 
 	// Set the view family's render target/viewport.
diff --git a/Engine/Source/Runtime/Renderer/Private/PostProcess/PostProcessCombineLUTs.cpp b/Engine/Source/Runtime/Renderer/Private/PostProcess/PostProcessCombineLUTs.cpp
--- a/Engine/Source/Runtime/Renderer/Private/PostProcess/PostProcessCombineLUTs.cpp
+++ b/Engine/Source/Runtime/Renderer/Private/PostProcess/PostProcessCombineLUTs.cpp
@@ -645,7 +645,7 @@ uint32 FRCPassPostProcessCombineLUTs::GenerateFinalTable(const FFinalPostProcess
 
 void FRCPassPostProcessCombineLUTs::Process(FRenderingCompositePassContext& Context)
 {
-	SCOPED_DRAW_EVENT(Context.RHICmdList, PostProcessCombineLUTs);
+	SCOPED_DRAW_EVENTF(Context.RHICmdList, PostProcessCombineLUTs, TEXT("PostProcessCombineLUTs %dx%dx%d"), GLUTSize, GLUTSize, GLUTSize);
 
 	FTexture* LocalTextures[GMaxLUTBlendCount];
 	float LocalWeights[GMaxLUTBlendCount];
diff --git a/Engine/Source/Runtime/Renderer/Private/PostProcess/PostProcessDownsample.cpp b/Engine/Source/Runtime/Renderer/Private/PostProcess/PostProcessDownsample.cpp
--- a/Engine/Source/Runtime/Renderer/Private/PostProcess/PostProcessDownsample.cpp
+++ b/Engine/Source/Runtime/Renderer/Private/PostProcess/PostProcessDownsample.cpp
@@ -154,8 +154,6 @@ void FRCPassPostProcessDownsample::SetShader(const FRenderingCompositePassContex
 
 void FRCPassPostProcessDownsample::Process(FRenderingCompositePassContext& Context)
 {
-	SCOPED_DRAW_EVENT(Context.RHICmdList, Downsample);
-
 	const FPooledRenderTargetDesc* InputDesc = GetInputDesc(ePId_Input0);
 
 	if(!InputDesc)
@@ -173,6 +171,12 @@ void FRCPassPostProcessDownsample::Process(FRenderingCompositePassContext& Conte
 	// e.g. 4 means the input texture is 4x smaller than the buffer size
 	uint32 ScaleFactor = FSceneRenderTargets::Get(Context.RHICmdList).GetBufferSizeXY().X / SrcSize.X;
 
+	FIntRect SrcRect = View.ViewRect / ScaleFactor;
+	FIntRect DestRect = FIntRect::DivideAndRoundUp(SrcRect, 2);
+	SrcRect = DestRect * 2;
+
+	SCOPED_DRAW_EVENTF(Context.RHICmdList, Downsample, TEXT("Downsample %dx%d"), DestRect.Width(), DestRect.Height());
+
 	const FSceneRenderTargetItem& DestRenderTarget = PassOutputs[0].RequestSurface(Context);
 
 	// Set the view family's render target/viewport.
@@ -221,10 +225,6 @@ void FRCPassPostProcessDownsample::Process(FRenderingCompositePassContext& Conte
 
 	TShaderMapRef<FPostProcessDownsampleVS> VertexShader(Context.GetShaderMap());
 
-	FIntRect SrcRect = View.ViewRect / ScaleFactor;
-	FIntRect DestRect = FIntRect::DivideAndRoundUp(SrcRect, 2);
-	SrcRect = DestRect * 2;
-
 	if (!bHasCleared)
 	{
 		Context.RHICmdList.Clear(true, FLinearColor(0, 0, 0, 0), false, 1.0f, false, 0, DestRect);
diff --git a/Engine/Source/Runtime/Renderer/Private/PostProcess/PostProcessMotionBlur.cpp b/Engine/Source/Runtime/Renderer/Private/PostProcess/PostProcessMotionBlur.cpp
--- a/Engine/Source/Runtime/Renderer/Private/PostProcess/PostProcessMotionBlur.cpp
+++ b/Engine/Source/Runtime/Renderer/Private/PostProcess/PostProcessMotionBlur.cpp
@@ -144,8 +144,6 @@ IMPLEMENT_SHADER_TYPE(,FPostProcessMotionBlurSetupPS, TEXT("PostProcessMotionBlu
 
 void FRCPassPostProcessMotionBlurSetup::Process(FRenderingCompositePassContext& Context)
 {
-	SCOPED_DRAW_EVENT(Context.RHICmdList, MotionBlurSetup);
-
 	const FPooledRenderTargetDesc* InputDesc = GetInputDesc(ePId_Input0);
 
 	if(!InputDesc)
@@ -168,6 +166,8 @@ void FRCPassPostProcessMotionBlurSetup::Process(FRenderingCompositePassContext&
 	// Viewport size not even also causes issue
 	FIntRect DestRect = FIntRect::DivideAndRoundUp(SrcRect, 2);
 
+	SCOPED_DRAW_EVENTF(Context.RHICmdList, MotionBlurSetup, TEXT("MotionBlurSetup %dx%d"), DestRect.Width(), DestRect.Height());
+	
 	const FSceneRenderTargetItem& DestRenderTarget0 = PassOutputs[0].RequestSurface(Context);
 	const FSceneRenderTargetItem& DestRenderTarget1 = PassOutputs[1].RequestSurface(Context);
 
@@ -600,8 +600,6 @@ IMPLEMENT_SHADER_TYPE(,FPostProcessMotionBlurRecombinePS,TEXT("PostProcessMotion
 
 void FRCPassPostProcessMotionBlurRecombine::Process(FRenderingCompositePassContext& Context)
 {
-	SCOPED_DRAW_EVENT(Context.RHICmdList, MotionBlurRecombine);
-
 	const FPooledRenderTargetDesc* InputDesc = GetInputDesc(ePId_Input0);
 
 	if(!InputDesc)
@@ -625,6 +623,8 @@ void FRCPassPostProcessMotionBlurRecombine::Process(FRenderingCompositePassConte
 	FIntRect SrcRect = View.ViewRect / ScaleFactor;
 	FIntRect DestRect = SrcRect;
 
+	SCOPED_DRAW_EVENTF(Context.RHICmdList, MotionBlurRecombine, TEXT("MotionBlurRecombine %dx%d"), DestRect.Width(), DestRect.Height());
+
 	const FSceneRenderTargetItem& DestRenderTarget = PassOutputs[0].RequestSurface(Context);
 
 	// Set the view family's render target/viewport.
@@ -1320,7 +1320,7 @@ VARIATION1(0)			VARIATION1(1)			VARIATION1(2)			VARIATION1(3)			VARIATION1(4)
 template< uint32 Quality >
 static void SetMotionBlurShaderNewTempl(const FRenderingCompositePassContext& Context)
 {
-	TShaderMapRef< FPostProcessVS >							VertexShader( Context.GetShaderMap() );
+	TShaderMapRef< FPostProcessVS > VertexShader( Context.GetShaderMap() );
 	TShaderMapRef< FPostProcessMotionBlurNewPS< Quality > >	PixelShader( Context.GetShaderMap() );
 
 	static FGlobalBoundShaderState BoundShaderState;
@@ -1339,8 +1339,6 @@ FRCPassPostProcessMotionBlurNew::FRCPassPostProcessMotionBlurNew(uint32 InQualit
 
 void FRCPassPostProcessMotionBlurNew::Process(FRenderingCompositePassContext& Context)
 {
-	SCOPED_DRAW_EVENT(Context.RHICmdList, MotionBlur);
-
 	const FPooledRenderTargetDesc* InputDesc = GetInputDesc(ePId_Input0);
 
 	if(!InputDesc)
@@ -1362,6 +1360,8 @@ void FRCPassPostProcessMotionBlurNew::Process(FRenderingCompositePassContext& Co
 	FIntRect SrcRect = View.ViewRect / ScaleFactor;
 	FIntRect DestRect = SrcRect;
 
+	SCOPED_DRAW_EVENTF(Context.RHICmdList, MotionBlurNew, TEXT("MotionBlurNew %dx%d"), SrcRect.Width(), SrcRect.Height());
+
 	const FSceneRenderTargetItem& DestRenderTarget = PassOutputs[0].RequestSurface(Context);
 
 	// Set the view family's render target/viewport.
diff --git a/Engine/Source/Runtime/Renderer/Private/PostProcess/PostProcessTemporalAA.cpp b/Engine/Source/Runtime/Renderer/Private/PostProcess/PostProcessTemporalAA.cpp
--- a/Engine/Source/Runtime/Renderer/Private/PostProcess/PostProcessTemporalAA.cpp
+++ b/Engine/Source/Runtime/Renderer/Private/PostProcess/PostProcessTemporalAA.cpp
@@ -537,8 +537,6 @@ FPooledRenderTargetDesc FRCPassPostProcessLightShaftTemporalAA::ComputeOutputDes
 
 void FRCPassPostProcessTemporalAA::Process(FRenderingCompositePassContext& Context)
 {
-	SCOPED_DRAW_EVENT(Context.RHICmdList, TemporalAA);
-
 	const FPooledRenderTargetDesc* InputDesc = GetInputDesc(ePId_Input0);
 
 	if(!InputDesc)
@@ -564,6 +562,8 @@ void FRCPassPostProcessTemporalAA::Process(FRenderingCompositePassContext& Conte
 	FIntRect SrcRect = View.ViewRect / ScaleFactor;
 	FIntRect DestRect = SrcRect;
 
+	SCOPED_DRAW_EVENTF(Context.RHICmdList, TemporalAA, TEXT("TemporalAA %dx%d"), SrcRect.Width(), SrcRect.Height());
+
 	const FSceneRenderTargetItem& DestRenderTarget = PassOutputs[0].RequestSurface(Context);
 
 	//Context.SetRenderTarget(RHICmdList, DestRenderTarget.TargetableTexture, FTextureRHIRef());
diff --git a/Engine/Source/Runtime/Renderer/Private/PostProcess/PostProcessTonemap.cpp b/Engine/Source/Runtime/Renderer/Private/PostProcess/PostProcessTonemap.cpp
--- a/Engine/Source/Runtime/Renderer/Private/PostProcess/PostProcessTonemap.cpp
+++ b/Engine/Source/Runtime/Renderer/Private/PostProcess/PostProcessTonemap.cpp
@@ -1076,8 +1076,6 @@ static void SetShaderTempl(const FRenderingCompositePassContext& Context)
 
 void FRCPassPostProcessTonemap::Process(FRenderingCompositePassContext& Context)
 {
-	SCOPED_DRAW_EVENTF(Context.RHICmdList, PostProcessTonemap, TEXT("Tonemapper#%d%s"), ConfigIndexPC, bDoGammaOnly ? TEXT(" GammaOnly") : TEXT(""));
-
 	const FPooledRenderTargetDesc* InputDesc = GetInputDesc(ePId_Input0);
 
 	if(!InputDesc)
@@ -1089,6 +1087,8 @@ void FRCPassPostProcessTonemap::Process(FRenderingCompositePassContext& Context)
 	const FSceneView& View = Context.View;
 	const FSceneViewFamily& ViewFamily = *(View.Family);
 
+	SCOPED_DRAW_EVENTF(Context.RHICmdList, PostProcessTonemap, TEXT("Tonemapper#%d%s %dx%d"), ConfigIndexPC, bDoGammaOnly ? TEXT(" GammaOnly") : TEXT(""), View.ViewRect.Width(), View.ViewRect.Height());
+
 	FIntRect SrcRect = View.ViewRect;
 	FIntRect DestRect = View.ViewRect;
 	FIntPoint SrcSize = InputDesc->Extent;
diff --git a/Engine/Source/Runtime/Renderer/Private/PostProcess/PostProcessWeightedSampleSum.cpp b/Engine/Source/Runtime/Renderer/Private/PostProcess/PostProcessWeightedSampleSum.cpp
--- a/Engine/Source/Runtime/Renderer/Private/PostProcess/PostProcessWeightedSampleSum.cpp
+++ b/Engine/Source/Runtime/Renderer/Private/PostProcess/PostProcessWeightedSampleSum.cpp
@@ -545,7 +545,9 @@ void FRCPassPostProcessWeightedSampleSum::Process(FRenderingCompositePassContext
 
 	uint32 NumSamples = Compute1DGaussianFilterKernel(FeatureLevel, EffectiveBlurRadius, OffsetAndWeight, MaxNumSamples, FilterShape, CrossCenterWeight);
 
-	SCOPED_DRAW_EVENTF(Context.RHICmdList, PostProcessWeightedSampleSum, TEXT("PostProcessWeightedSampleSum#%d"), NumSamples);
+	FIntRect DestRect = FIntRect::DivideAndRoundUp(View.ViewRect, DstScaleFactor);
+
+	SCOPED_DRAW_EVENTF(Context.RHICmdList, PostProcessWeightedSampleSum, TEXT("PostProcessWeightedSampleSum#%d %dx%d"), NumSamples, DestRect.Width(), DestRect.Height());
 
 	// compute weights as weighted contributions of the TintValue
 	for(uint32 i = 0; i < NumSamples; ++i)
@@ -629,7 +631,6 @@ void FRCPassPostProcessWeightedSampleSum::Process(FRenderingCompositePassContext
 	}
 
 	FIntRect SrcRect =  FIntRect::DivideAndRoundUp(View.ViewRect, SrcScaleFactor);
-	FIntRect DestRect = FIntRect::DivideAndRoundUp(View.ViewRect, DstScaleFactor);
 
 	DrawQuad(Context.RHICmdList, bDoFastBlur, SrcRect, DestRect, bRequiresClear, DestSize, SrcSize, VertexShader);
 
diff --git a/Engine/Source/Runtime/Renderer/Private/PostProcess/ScreenSpaceReflections.cpp b/Engine/Source/Runtime/Renderer/Private/PostProcess/ScreenSpaceReflections.cpp
--- a/Engine/Source/Runtime/Renderer/Private/PostProcess/ScreenSpaceReflections.cpp
+++ b/Engine/Source/Runtime/Renderer/Private/PostProcess/ScreenSpaceReflections.cpp
@@ -321,7 +321,7 @@ void FRCPassPostProcessScreenSpaceReflections::Process(FRenderingCompositePassCo
 	
 	if (SSRStencilPrePass)
 	{ // ScreenSpaceReflectionsStencil draw event
-		SCOPED_DRAW_EVENT(RHICmdList, ScreenSpaceReflectionsStencil);
+		SCOPED_DRAW_EVENTF(Context.RHICmdList, ScreenSpaceReflectionsStencil, TEXT("ScreenSpaceReflectionsStencil %dx%d"), View.ViewRect.Width(), View.ViewRect.Height());
 
 		TShaderMapRef< FPostProcessVS > VertexShader(Context.GetShaderMap());
 		TShaderMapRef< FPostProcessScreenSpaceReflectionsStencilPS > PixelShader(Context.GetShaderMap());
@@ -364,7 +364,7 @@ void FRCPassPostProcessScreenSpaceReflections::Process(FRenderingCompositePassCo
 	} // ScreenSpaceReflectionsStencil draw event
 
 	{ // ScreenSpaceReflections draw event
-		SCOPED_DRAW_EVENT(Context.RHICmdList, ScreenSpaceReflections);
+		SCOPED_DRAW_EVENTF(Context.RHICmdList, ScreenSpaceReflections, TEXT("ScreenSpaceReflections %dx%d"), View.ViewRect.Width(), View.ViewRect.Height());
 
 		if (SSRStencilPrePass)
 		{
diff --git a/Engine/Source/Runtime/Renderer/Private/ReflectionEnvironment.cpp b/Engine/Source/Runtime/Renderer/Private/ReflectionEnvironment.cpp
--- a/Engine/Source/Runtime/Renderer/Private/ReflectionEnvironment.cpp
+++ b/Engine/Source/Runtime/Renderer/Private/ReflectionEnvironment.cpp
@@ -869,12 +869,10 @@ void FDeferredShadingSceneRenderer::RenderTiledDeferredImageBasedReflections(FRH
 
 		// ReflectionEnv is assumed to be on when going into this method
 		{
-			// Render the reflection environment with tiled deferred culling
-			SCOPED_DRAW_EVENT(RHICmdList, ReflectionEnvironmentGather);
-
 			SetRenderTarget(RHICmdList, NULL, NULL);
 
 			FReflectionEnvironmentTiledDeferredCS* ComputeShader = NULL;
+			// Render the reflection environment with tiled deferred culling
 			TArray<FReflectionCaptureSortData> SortData;
 			int32 NumBoxCaptures = 0;
 			int32 NumSphereCaptures = 0;
@@ -883,6 +881,11 @@ void FDeferredShadingSceneRenderer::RenderTiledDeferredImageBasedReflections(FRH
 			bool bHasBoxCaptures = (NumBoxCaptures > 0);
 			bool bHasSphereCaptures = (NumSphereCaptures > 0);
 			bool bNeedsClearCoat = (View.ShadingModelMaskInView & (1 << MSM_ClearCoat)) != 0;
+
+			SCOPED_DRAW_EVENTF(RHICmdList, ReflectionEnvironment, TEXT("ReflectionEnvironment ComputeShader %dx%d Tile:%dx%d Box:%d Sphere:%d ClearCoat:%d"), 
+				View.ViewRect.Width(), View.ViewRect.Height(), AdjustedReflectionTileSizeX, GReflectionEnvironmentTileSizeY,
+				NumBoxCaptures, NumSphereCaptures, bNeedsClearCoat);
+
 			ComputeShader = SelectReflectionEnvironmentTiledDeferredCS(View.ShaderMap, bUseLightmaps, bNeedsClearCoat, bHasBoxCaptures, bHasSphereCaptures);
 
 			RHICmdList.SetComputeShader(ComputeShader->GetComputeShader());
@@ -1009,7 +1012,7 @@ void FDeferredShadingSceneRenderer::RenderStandardDeferredImageBasedReflections(
 		{
 			bRequiresApply = true;
 
-			SCOPED_DRAW_EVENT(RHICmdList, StandardDeferredReflectionEnvironment);
+			SCOPED_DRAW_EVENTF(RHICmdList, ReflectionEnvironment, TEXT("ReflectionEnvironment PixelShader"));
 
 			{
 				// Clear to no reflection contribution, alpha of 1 indicates full background contribution
diff --git a/Engine/Source/Runtime/Renderer/Private/SceneOcclusion.cpp b/Engine/Source/Runtime/Renderer/Private/SceneOcclusion.cpp
--- a/Engine/Source/Runtime/Renderer/Private/SceneOcclusion.cpp
+++ b/Engine/Source/Runtime/Renderer/Private/SceneOcclusion.cpp
@@ -975,7 +975,6 @@ IMPLEMENT_SHADER_TYPE(template<>,THZBBuildPS<1>,TEXT("HZBOcclusion"),TEXT("HZBBu
 void BuildHZB( FRHICommandListImmediate& RHICmdList, FViewInfo& View )
 {
 	QUICK_SCOPE_CYCLE_COUNTER(STAT_BuildHZB);
-	SCOPED_DRAW_EVENT(RHICmdList, BuildHZB);
 	
 	// View.ViewRect.{Width,Height}() are most likely to be < 2^24, so the float
 	// conversion won't loss any precision (assuming float have 23bits for mantissa)
@@ -995,6 +994,8 @@ void BuildHZB( FRHICommandListImmediate& RHICmdList, FViewInfo& View )
 	
 	// Mip 0
 	{
+		SCOPED_DRAW_EVENTF(RHICmdList, BuildHZB, TEXT("HZB SetupMip 0 %dx%d"), HZBSize.X, HZBSize.Y);
+
 		SetRenderTarget(RHICmdList, HZBRenderTarget.TargetableTexture, 0, NULL);
 		RHICmdList.SetBlendState(TStaticBlendState<>::GetRHI());
 		RHICmdList.SetRasterizerState(TStaticRasterizerState<>::GetRHI());
@@ -1029,6 +1030,8 @@ void BuildHZB( FRHICommandListImmediate& RHICmdList, FViewInfo& View )
 	FIntPoint SrcSize = HZBSize;
 	FIntPoint DstSize = SrcSize / 2;
 	
+	SCOPED_DRAW_EVENTF(RHICmdList, BuildHZB, TEXT("HZB SetupMips 1..%d %dx%d Mips:%d"), NumMips - 1, DstSize.X, DstSize.Y);
+
 	// Downsampling...
 	for( uint8 MipIndex = 1; MipIndex < NumMips; MipIndex++ )
 	{
diff --git a/Engine/Source/Runtime/Renderer/Private/TranslucentLighting.cpp b/Engine/Source/Runtime/Renderer/Private/TranslucentLighting.cpp
--- a/Engine/Source/Runtime/Renderer/Private/TranslucentLighting.cpp
+++ b/Engine/Source/Runtime/Renderer/Private/TranslucentLighting.cpp
@@ -1813,7 +1813,11 @@ void FDeferredShadingSceneRenderer::FilterTranslucentVolumeLighting(FRHICommandL
 
 		if (GUseTranslucencyVolumeBlur)
 		{
-			SCOPED_DRAW_EVENT(RHICmdList, FilterTranslucentVolume);
+			//@todo - support multiple views
+			const FViewInfo& View = Views[0];
+
+			SCOPED_DRAW_EVENTF(RHICmdList, FilterTranslucentVolume, TEXT("FilterTranslucentVolume %dx%dx%d Cascades:%d"),
+				GTranslucencyLightingVolumeDim, GTranslucencyLightingVolumeDim, GTranslucencyLightingVolumeDim, TVC_MAX);
 
 			// Filter each cascade
 			for (int32 VolumeCascadeIndex = 0; VolumeCascadeIndex < TVC_MAX; VolumeCascadeIndex++)
@@ -1830,8 +1834,6 @@ void FDeferredShadingSceneRenderer::FilterTranslucentVolumeLighting(FRHICommandL
 
 				SetRenderTargets(RHICmdList, ARRAY_COUNT(RenderTargets), RenderTargets, FTextureRHIRef(), 0, NULL);
 
-				//@todo - support multiple views
-				const FViewInfo& View = Views[0];
 
 				const FVolumeBounds VolumeBounds(GTranslucencyLightingVolumeDim);
 				TShaderMapRef<FWriteToSliceVS> VertexShader(View.ShaderMap);
