Fixes a performance regression cause by memory cache trashing on SSAO

diff --git a/Engine/Shaders/PostProcessAmbientOcclusion.usf b/Engine/Shaders/PostProcessAmbientOcclusion.usf
--- a/Engine/Shaders/PostProcessAmbientOcclusion.usf
+++ b/Engine/Shaders/PostProcessAmbientOcclusion.usf
@@ -344,7 +344,7 @@ float2 Wedge(float Scale, float2 LocalRandom, float3 FovFix, float3 InvFovFix, f
 
 	float2 ScreenSpacePos = ReconstructSSPosFromCS(SamplePosition2 * FovFix);
 
-#if USE_estmap 
+#if USE_HZB
 	float2 UV = ScreenSpacePos * float2(0.5f, -0.5f) + 0.5f;
 
 	float Abs = ConvertFromDeviceZ(Texture2DSampleLevel( PostprocessInput3, PostprocessInput3Sampler, HzbUvAndStepMipLevelFactor.xy * UV, MipLevel).r);
diff --git a/Engine/Source/Runtime/Renderer/Private/CompositionLighting/PostProcessAmbientOcclusion.cpp b/Engine/Source/Runtime/Renderer/Private/CompositionLighting/PostProcessAmbientOcclusion.cpp
--- a/Engine/Source/Runtime/Renderer/Private/CompositionLighting/PostProcessAmbientOcclusion.cpp
+++ b/Engine/Source/Runtime/Renderer/Private/CompositionLighting/PostProcessAmbientOcclusion.cpp
@@ -192,7 +192,7 @@ public:
 			HzbStepMipLevelFactorValue
 			);
 		
-		SetShaderValue(Context.RHICmdList, ShaderRHI, HzbUvAndStepMipLevelFactor, HzbStepMipLevelFactorValue );
+		SetShaderValue(Context.RHICmdList, ShaderRHI, HzbUvAndStepMipLevelFactor, HzbUvAndStepMipLevelFactorValue );
 	}
 	
 	// FShader interface.
