fixed build

diff --git a/Engine/Source/Runtime/Renderer/Private/ReflectionEnvironment.cpp b/Engine/Source/Runtime/Renderer/Private/ReflectionEnvironment.cpp
--- a/Engine/Source/Runtime/Renderer/Private/ReflectionEnvironment.cpp
+++ b/Engine/Source/Runtime/Renderer/Private/ReflectionEnvironment.cpp
@@ -883,7 +883,7 @@ void FDeferredShadingSceneRenderer::RenderTiledDeferredImageBasedReflections(FRH
 			bool bNeedsClearCoat = (View.ShadingModelMaskInView & (1 << MSM_ClearCoat)) != 0;
 
 			SCOPED_DRAW_EVENTF(RHICmdList, ReflectionEnvironment, TEXT("ReflectionEnvironment ComputeShader %dx%d Tile:%dx%d Box:%d Sphere:%d ClearCoat:%d"), 
-				View.ViewRect.Width(), View.ViewRect.Height(), AdjustedReflectionTileSizeX, GReflectionEnvironmentTileSizeY,
+				View.ViewRect.Width(), View.ViewRect.Height(), GReflectionEnvironmentTileSizeX, GReflectionEnvironmentTileSizeY,
 				NumBoxCaptures, NumSphereCaptures, bNeedsClearCoat);
 
 			ComputeShader = SelectReflectionEnvironmentTiledDeferredCS(View.ShaderMap, bUseLightmaps, bNeedsClearCoat, bHasBoxCaptures, bHasSphereCaptures);
