No geometry shaders in Mac Metal which is an editor platform.

diff --git a/Engine/Source/Runtime/Renderer/Private/VertexDensityRendering.cpp b/Engine/Source/Runtime/Renderer/Private/VertexDensityRendering.cpp
--- a/Engine/Source/Runtime/Renderer/Private/VertexDensityRendering.cpp
+++ b/Engine/Source/Runtime/Renderer/Private/VertexDensityRendering.cpp
@@ -68,7 +68,7 @@ public:
 	static bool ShouldCache(EShaderPlatform Platform, const FMaterial* Material, const FVertexFactoryType* VertexFactoryType)
 	{
 		return (Material->IsSpecialEngineMaterial() || Material->IsMasked() || Material->MaterialMayModifyMeshPosition())
-			&& IsFeatureLevelSupported(Platform, ERHIFeatureLevel::SM4);
+			&& IsFeatureLevelSupported(Platform, ERHIFeatureLevel::SM4) && RHISupportsGeometryShaders(Platform);
 	}
 
 	static void ModifyCompilationEnvironment(EShaderPlatform Platform, const FMaterial* Material, FShaderCompilerEnvironment& OutEnvironment)
@@ -342,7 +342,7 @@ IMPLEMENT_MATERIAL_SHADER_TYPE(, FVertexDensityPS, TEXT("VertexDensityShader"),
 
 bool FDeferredShadingSceneRenderer::RenderVertexDensities(FRHICommandListImmediate& RHICmdList)
 {
-	if (Scene->GetFeatureLevel() >= ERHIFeatureLevel::SM4)
+	if (Scene->GetFeatureLevel() >= ERHIFeatureLevel::SM4 && RHISupportsGeometryShaders(Scene->GetShaderPlatform()))
 	{
 		SCOPED_DRAW_EVENT(RHICmdList, VertexDensity);
 
