integrate from Orion
CL 2678931
	some minor refactor

diff --git a/Engine/Source/Runtime/Renderer/Private/PrimitiveSceneInfo.cpp b/Engine/Source/Runtime/Renderer/Private/PrimitiveSceneInfo.cpp
--- a/Engine/Source/Runtime/Renderer/Private/PrimitiveSceneInfo.cpp
+++ b/Engine/Source/Runtime/Renderer/Private/PrimitiveSceneInfo.cpp
@@ -90,7 +90,6 @@ FPrimitiveSceneInfo::FPrimitiveSceneInfo(UPrimitiveComponent* InComponent,FScene
 	IndirectLightingCacheAllocation(NULL),
 	CachedReflectionCaptureProxy(NULL),
 	bNeedsCachedReflectionCaptureUpdate(true),
-	bVelocityIsSupressed(false),
 	DefaultDynamicHitProxy(NULL),
 	LightList(NULL),
 	LastRenderTime(-FLT_MAX),
diff --git a/Engine/Source/Runtime/Renderer/Private/VelocityRendering.cpp b/Engine/Source/Runtime/Renderer/Private/VelocityRendering.cpp
--- a/Engine/Source/Runtime/Renderer/Private/VelocityRendering.cpp
+++ b/Engine/Source/Runtime/Renderer/Private/VelocityRendering.cpp
@@ -315,11 +315,6 @@ bool FVelocityDrawingPolicy::HasVelocity(const FViewInfo& View, const FPrimitive
 		return true;
 	}
 
-	if (PrimitiveSceneInfo->bVelocityIsSupressed)
-	{
-		return false;
-	}
-
 	// check if the primitive has moved
 	{
 		FMatrix PreviousLocalToWorld;
@@ -357,12 +352,6 @@ bool FVelocityDrawingPolicy::HasVelocityOnBasePass(const FViewInfo& View,const F
 		return false;
 	}
 
-	//@todo-rco: Where is this set?
-	if (PrimitiveSceneInfo->bVelocityIsSupressed)
-	{
-		return false;
-	}
-
 	// hack
 	FScene* Scene = PrimitiveSceneInfo->Scene;
 	if (Scene->MotionBlurInfoData.GetPrimitiveMotionBlurInfo(PrimitiveSceneInfo, OutTransform))
diff --git a/Engine/Source/Runtime/Renderer/Public/PrimitiveSceneInfo.h b/Engine/Source/Runtime/Renderer/Public/PrimitiveSceneInfo.h
--- a/Engine/Source/Runtime/Renderer/Public/PrimitiveSceneInfo.h
+++ b/Engine/Source/Runtime/Renderer/Public/PrimitiveSceneInfo.h
@@ -210,9 +210,6 @@ public:
 	/** Whether the primitive is newly registered or moved and CachedReflectionCaptureProxy needs to be updated on the next render. */
 	uint32 bNeedsCachedReflectionCaptureUpdate : 1;
 
-	/** This primitive has the Motion Blur explicitly disabled */
-	uint32 bVelocityIsSupressed : 1;
-
 	/** The hit proxies used by the primitive. */
 	TArray<TRefCountPtr<HHitProxy> > HitProxies;
 
