Enable DFAO on Xbox One.

diff --git a/Engine/Source/Runtime/Engine/Public/SceneManagement.h b/Engine/Source/Runtime/Engine/Public/SceneManagement.h
--- a/Engine/Source/Runtime/Engine/Public/SceneManagement.h
+++ b/Engine/Source/Runtime/Engine/Public/SceneManagement.h
@@ -624,7 +624,7 @@ public:
 inline bool DoesPlatformSupportDistanceFieldShadowing(EShaderPlatform Platform)
 {
 	// Hasn't been tested elsewhere yet
-	return Platform == SP_PCD3D_SM5 || Platform == SP_PS4 || Platform == SP_METAL_SM5;
+	return Platform == SP_PCD3D_SM5 || Platform == SP_PS4 || Platform == SP_METAL_SM5 || Platform == SP_XBOXONE;
 }
 
 /** Represents a USkyLightComponent to the rendering thread. */
diff --git a/Engine/Source/Runtime/Renderer/Private/DistanceFieldSurfaceCacheLighting.h b/Engine/Source/Runtime/Renderer/Private/DistanceFieldSurfaceCacheLighting.h
--- a/Engine/Source/Runtime/Renderer/Private/DistanceFieldSurfaceCacheLighting.h
+++ b/Engine/Source/Runtime/Renderer/Private/DistanceFieldSurfaceCacheLighting.h
@@ -21,7 +21,7 @@ const float GDefaultDFAOMaxOcclusionDistance = 600.0f;
 
 inline bool DoesPlatformSupportDistanceFieldAO(EShaderPlatform Platform)
 {
-	return Platform == SP_PCD3D_SM5 || Platform == SP_PS4 || Platform == SP_METAL_SM5;
+	return Platform == SP_PCD3D_SM5 || Platform == SP_PS4 || Platform == SP_METAL_SM5 || Platform == SP_XBOXONE;
 }
 
 extern FIntPoint GetBufferSizeForAO();
