Don't match rendertargets against FastVRAM unless the platform requires it.

diff --git a/Engine/Source/Runtime/Renderer/Private/PostProcess/RenderTargetPool.cpp b/Engine/Source/Runtime/Renderer/Private/PostProcess/RenderTargetPool.cpp
--- a/Engine/Source/Runtime/Renderer/Private/PostProcess/RenderTargetPool.cpp
+++ b/Engine/Source/Runtime/Renderer/Private/PostProcess/RenderTargetPool.cpp
@@ -208,7 +208,8 @@ bool FRenderTargetPool::FindFreeElement(const FPooledRenderTargetDesc& Desc, TRe
 
 	// try to find a suitable element in the pool
 	{
-		uint32 PassCount = (Desc.Flags & TexCreate_FastVRAM) ? 2 : 1;
+		//don't spend time doing 2 passes if the platform doesn't support fastvram
+		uint32 PassCount = ((Desc.Flags & TexCreate_FastVRAM) && FPlatformProperties::SupportsFastVRAMMemory()) ? 2 : 1;			
 
 		// first we try exact, if that fails we try without TexCreate_FastVRAM
 		// (easily we can run out of VRam, if this search becomes a performance problem we can optimize or we should use less TexCreate_FastVRAM)
@@ -228,11 +229,11 @@ bool FRenderTargetPool::FindFreeElement(const FPooledRenderTargetDesc& Desc, TRe
 					break;
 				}
 			}
-		}
+		}		
 	}
 
 	if(!Found)
-	{
+	{		
 		UE_LOG(LogRenderTargetPool, Display, TEXT("%d MB, NewRT %s %s"), (AllocationLevelInKB + 1023) / 1024, *Desc.GenerateInfoString(), InDebugName);
 
 		// not found in the pool, create a new element
diff --git a/Engine/Source/Runtime/Renderer/Public/RendererInterface.h b/Engine/Source/Runtime/Renderer/Public/RendererInterface.h
--- a/Engine/Source/Runtime/Renderer/Public/RendererInterface.h
+++ b/Engine/Source/Runtime/Renderer/Public/RendererInterface.h
@@ -159,7 +159,7 @@ public:
 		auto LhsFlags = Flags;
 		auto RhsFlags = rhs.Flags;
 
-		if(!bExact)
+		if (!bExact || !FPlatformProperties::SupportsFastVRAMMemory())
 		{
 			LhsFlags &= (~TexCreate_FastVRAM);
 			RhsFlags &= (~TexCreate_FastVRAM);
