prervent artifact that would happen with splitscreen and BloomThreshold <= -1

diff --git a/Engine/Source/Runtime/Renderer/Private/PostProcess/PostProcessing.cpp b/Engine/Source/Runtime/Renderer/Private/PostProcess/PostProcessing.cpp
--- a/Engine/Source/Runtime/Renderer/Private/PostProcess/PostProcessing.cpp
+++ b/Engine/Source/Runtime/Renderer/Private/PostProcess/PostProcessing.cpp
@@ -151,9 +151,10 @@ FPostprocessContext::FPostprocessContext(class FRenderingCompositionGraph& InGra
 
 static FRenderingCompositeOutputRef RenderHalfResBloomThreshold(FPostprocessContext& Context, FRenderingCompositeOutputRef SceneColorHalfRes, FRenderingCompositeOutputRef EyeAdaptation)
 {
-	if(Context.View.FinalPostProcessSettings.BloomThreshold <= -1)
+	// with multiple view ports the Setup pass also isolates the view from the others which allows for simpler simpler/faster blur passes.
+	if(Context.View.FinalPostProcessSettings.BloomThreshold <= -1 && Context.View.Family->Views.Num() == 1)
 	{
-		// no need for threshold
+		// no need for threshold, we don't need this pass
 		return SceneColorHalfRes;
 	}
 	else
