Postprocessing optimization: In rare cases some composition graph nodes have been processed without the actual need

diff --git a/Engine/Source/Runtime/Renderer/Private/PostProcess/RenderingCompositionGraph.cpp b/Engine/Source/Runtime/Renderer/Private/PostProcess/RenderingCompositionGraph.cpp
--- a/Engine/Source/Runtime/Renderer/Private/PostProcess/RenderingCompositionGraph.cpp
+++ b/Engine/Source/Runtime/Renderer/Private/PostProcess/RenderingCompositionGraph.cpp
@@ -169,21 +169,23 @@ void FRenderingCompositePassContext::Process(FRenderingCompositePass* Root, cons
 
 		bool bNewOrder = CVarCompositionGraphOrder.GetValueOnRenderThread() != 0;
 
+		Graph.RecursivelyGatherDependencies(Root);
+
 		if(bNewOrder)
 		{
+			// process in the order the nodes have been created (for more control), unless the dependencies require it differently
 			for (FRenderingCompositePass* Node : Graph.Nodes)
 			{
-				Graph.RecursivelyGatherDependencies(Node);
-			}
-
-			for (FRenderingCompositePass* Node : Graph.Nodes)
-			{
-				Graph.RecursivelyProcess(Node, *this);
+				// only if this is true the node is actually needed - no need to compute it when it's not needed
+				if(Node->WasComputeOutputDescCalled())
+				{
+					Graph.RecursivelyProcess(Node, *this);
+				}
 			}
 		}
 		else
 		{
-			Graph.RecursivelyGatherDependencies(Root);
+			// process in the order of the dependencies, starting from the root (without processing unreferenced nodes)
 			Graph.RecursivelyProcess(Root, *this);
 		}
 
diff --git a/Engine/Source/Runtime/Renderer/Private/PostProcess/RenderingCompositionGraph.h b/Engine/Source/Runtime/Renderer/Private/PostProcess/RenderingCompositionGraph.h
--- a/Engine/Source/Runtime/Renderer/Private/PostProcess/RenderingCompositionGraph.h
+++ b/Engine/Source/Runtime/Renderer/Private/PostProcess/RenderingCompositionGraph.h
@@ -265,6 +265,9 @@ struct FRenderingCompositePass
 	/** */
 	virtual void Release() = 0;
 
+	/** can be called after RecursivelyGatherDependencies to detect if the node is reference by any other node - if not we don't need to run it */
+	bool WasComputeOutputDescCalled() const { return bComputeOutputDescWasCalled; }
+
 protected:
 
 	/** to avoid wasteful recomputation and to support graph/DAG traversal, if ComputeOutputDesc() was called */
