Missing optimization: disabled bloom if requested by post process setting

diff --git a/Engine/Source/Runtime/Renderer/Private/PostProcess/PostProcessing.cpp b/Engine/Source/Runtime/Renderer/Private/PostProcess/PostProcessing.cpp
--- a/Engine/Source/Runtime/Renderer/Private/PostProcess/PostProcessing.cpp
+++ b/Engine/Source/Runtime/Renderer/Private/PostProcess/PostProcessing.cpp
@@ -1356,7 +1356,7 @@ void FPostProcessing::Process(FRHICommandListImmediate& RHICmdList, FViewInfo& V
 				EyeAdaptation = AddPostProcessEyeAdaptation(Context, Histogram);
 			}
 
-			if(View.Family->EngineShowFlags.Bloom)
+			if(View.FinalPostProcessSettings.BloomIntensity > 0.0f)
 			{
 				if (CVarUseMobileBloom.GetValueOnRenderThread() == 0)
 				{
