nicer VisualizeTexture printout

diff --git a/Engine/Source/Runtime/Renderer/Private/PostProcess/VisualizeTexture.cpp b/Engine/Source/Runtime/Renderer/Private/PostProcess/VisualizeTexture.cpp
--- a/Engine/Source/Runtime/Renderer/Private/PostProcess/VisualizeTexture.cpp
+++ b/Engine/Source/Runtime/Renderer/Private/PostProcess/VisualizeTexture.cpp
@@ -771,21 +771,31 @@ void FVisualizeTexture::DebugLog(bool bExtended)
 				// sort by index
 				Element.SortIndex = i;
 				
+				FString Info = Desc.GenerateInfoString();
 				if(SortOrder == -1)
 				{
+					// constant works well with the average name length
+					const uint32 TotelSpacerSize = 36;
+					uint32 SpaceCount = FMath::Max<int32>(0, TotelSpacerSize - Info.Len());
+
+					for(uint32 Space = 0; Space < SpaceCount; ++Space)
+					{
+						Info.AppendChar((TCHAR)' ');
+					}
+
 					// sort by index
-					Element.Line = FString::Printf(TEXT("%s %s %d KB%s"), *Desc.GenerateInfoString(), Desc.DebugName, SizeInKB, *UnusedStr);
+					Element.Line = FString::Printf(TEXT("%s %s %d KB%s"), *Info, Desc.DebugName, SizeInKB, *UnusedStr);
 				}
 				else if(SortOrder == 0)
 				{
 					// sort by name
-					Element.Line = FString::Printf(TEXT("%s %s %d KB%s"), Desc.DebugName, *Desc.GenerateInfoString(), SizeInKB, *UnusedStr);
+					Element.Line = FString::Printf(TEXT("%s %s %d KB%s"), Desc.DebugName, *Info, SizeInKB, *UnusedStr);
 					Element.SortIndex = 0;
 				}
 				else if(SortOrder == 1)
 				{
 					// sort by size (large ones first)
-					Element.Line = FString::Printf(TEXT("%d KB %s %s%s"), SizeInKB, *Desc.GenerateInfoString(), Desc.DebugName, *UnusedStr);
+					Element.Line = FString::Printf(TEXT("%d KB %s %s%s"), SizeInKB, *Info, Desc.DebugName, *UnusedStr);
 					Element.SortIndex = -(int32)SizeInKB;
 				}
 				else
@@ -870,31 +880,64 @@ void FVisualizeTexture::DebugLog(bool bExtended)
 		}
 
 		Entries.Sort();
-				
+
+		// that number works well with the name length we have
+		const uint32 ColumnCount = 5;
+		const uint32 SpaceBetweenColumns = 1;
+		uint32 ColumnHeight = FMath::DivideAndRoundUp((uint32)Entries.Num(), ColumnCount);
+
+		// width of the column in characters, init with 0
+		uint32 ColumnWidths[ColumnCount] = {};
+
+		for(int32 Index = 0; Index < Entries.Num(); ++Index)
+		{
+			uint32 Column = Index / ColumnHeight;
+			
+			const FString& Entry = *Entries[Index];
+
+			ColumnWidths[Column] = FMath::Max(ColumnWidths[Column], (uint32)Entry.Len());
+		}
+		
 		// print them sorted, if possible multiple in a line
 		{
 			FString Line;
-			FString Separator = " ,    ";
 
-			for(int32 Index=0; Index < Entries.Num(); Index++ )
+			for(int32 OutputIndex = 0; OutputIndex < Entries.Num(); ++OutputIndex)
 			{
-				const FString& Entry = *Entries[Index];
+				// 0..ColumnCount-1
+				uint32 Column = OutputIndex % ColumnCount;
+				int32 Row = OutputIndex / ColumnCount;
+
+				uint32 Index = Row + Column * ColumnHeight;
+
+				bool bLineEnd = true;
 
-				if(Line.Len() + 2 + Entry.Len() > 80)
+				if(Index < (uint32)Entries.Num())
 				{
+					bLineEnd = (Column + 1 == ColumnCount);
+
+					// for human readability we order them to be per column
+					const FString& Entry = *Entries[Index];
+
+					Line += Entry;
+
+					int32 SpaceCount = ColumnWidths[Column] + SpaceBetweenColumns - Entry.Len();
+
+					// otehrwise a fomer pass was producing bad data
+					check(SpaceCount >= 0);
+
+					for(int32 Space = 0; Space < SpaceCount; ++Space)
+					{
+						Line.AppendChar((TCHAR)' ');
+					}
+				}
+
+				if(bLineEnd)
+				{
+					Line = Line.TrimTrailing();
 					UE_LOG(LogConsoleResponse, Log, TEXT("   %s"), *Line);
 					Line.Empty();
 				}
-
-				Line += Entry;
-				Line += Separator;
-			}
-
-			if(!Line.IsEmpty())
-			{	
-				// remove separator in the end
-				Line = Line.Left(Line.Len() - Separator.Len());
-				UE_LOG(LogConsoleResponse, Log, TEXT("   %s"), *Line);
 			}
 		}
 	}
