minor performance optimization for postprocessing

diff --git a/Engine/Source/Runtime/Renderer/Private/PostProcess/PostProcessDownsample.cpp b/Engine/Source/Runtime/Renderer/Private/PostProcess/PostProcessDownsample.cpp
--- a/Engine/Source/Runtime/Renderer/Private/PostProcess/PostProcessDownsample.cpp
+++ b/Engine/Source/Runtime/Renderer/Private/PostProcess/PostProcessDownsample.cpp
@@ -58,7 +58,13 @@ public:
 
 		FGlobalShader::SetParameters(Context.RHICmdList, ShaderRHI, Context.View);
 		DeferredParameters.Set(Context.RHICmdList, ShaderRHI, Context.View);
-		PostprocessParameter.SetPS(ShaderRHI, Context, TStaticSamplerState<SF_Bilinear,AM_Clamp,AM_Clamp,AM_Clamp>::GetRHI());
+
+		// filter only if needed for better performance
+		FSamplerStateRHIParamRef Filter = (Method == 2) ? 
+			TStaticSamplerState<SF_Point, AM_Clamp, AM_Clamp, AM_Clamp>::GetRHI():
+			TStaticSamplerState<SF_Bilinear, AM_Clamp, AM_Clamp, AM_Clamp>::GetRHI();
+
+		PostprocessParameter.SetPS(ShaderRHI, Context, Filter);
 	}
 
 	static const TCHAR* GetSourceFilename()
