Fix shadowed variable

diff --git a/Engine/Source/Runtime/Renderer/Private/PostProcess/VisualizeTexture.cpp b/Engine/Source/Runtime/Renderer/Private/PostProcess/VisualizeTexture.cpp
--- a/Engine/Source/Runtime/Renderer/Private/PostProcess/VisualizeTexture.cpp
+++ b/Engine/Source/Runtime/Renderer/Private/PostProcess/VisualizeTexture.cpp
@@ -771,31 +771,31 @@ void FVisualizeTexture::DebugLog(bool bExtended)
 				// sort by index
 				Element.SortIndex = i;
 				
-				FString Info = Desc.GenerateInfoString();
+				FString InfoString = Desc.GenerateInfoString();
 				if(SortOrder == -1)
 				{
 					// constant works well with the average name length
 					const uint32 TotelSpacerSize = 36;
-					uint32 SpaceCount = FMath::Max<int32>(0, TotelSpacerSize - Info.Len());
+					uint32 SpaceCount = FMath::Max<int32>(0, TotelSpacerSize - InfoString.Len());
 
 					for(uint32 Space = 0; Space < SpaceCount; ++Space)
 					{
-						Info.AppendChar((TCHAR)' ');
+						InfoString.AppendChar((TCHAR)' ');
 					}
 
 					// sort by index
-					Element.Line = FString::Printf(TEXT("%s %s %d KB%s"), *Info, Desc.DebugName, SizeInKB, *UnusedStr);
+					Element.Line = FString::Printf(TEXT("%s %s %d KB%s"), *InfoString, Desc.DebugName, SizeInKB, *UnusedStr);
 				}
 				else if(SortOrder == 0)
 				{
 					// sort by name
-					Element.Line = FString::Printf(TEXT("%s %s %d KB%s"), Desc.DebugName, *Info, SizeInKB, *UnusedStr);
+					Element.Line = FString::Printf(TEXT("%s %s %d KB%s"), Desc.DebugName, *InfoString, SizeInKB, *UnusedStr);
 					Element.SortIndex = 0;
 				}
 				else if(SortOrder == 1)
 				{
 					// sort by size (large ones first)
-					Element.Line = FString::Printf(TEXT("%d KB %s %s%s"), SizeInKB, *Info, Desc.DebugName, *UnusedStr);
+					Element.Line = FString::Printf(TEXT("%d KB %s %s%s"), SizeInKB, *InfoString, Desc.DebugName, *UnusedStr);
 					Element.SortIndex = -(int32)SizeInKB;
 				}
 				else
