Changed default SamplerState anisotropy from 0 (use cvar setting) to 1 prevent accidential use of anisotropy which could cost performance and cause artifacts.
Other minor refactorings.

diff --git a/Engine/Documentation/Source/Support/Builds/ReleaseNotes/2013/Oct/ReleaseNotesOctober2013.INT.udn b/Engine/Documentation/Source/Support/Builds/ReleaseNotes/2013/Oct/ReleaseNotesOctober2013.INT.udn
--- a/Engine/Documentation/Source/Support/Builds/ReleaseNotes/2013/Oct/ReleaseNotesOctober2013.INT.udn
+++ b/Engine/Documentation/Source/Support/Builds/ReleaseNotes/2013/Oct/ReleaseNotesOctober2013.INT.udn
@@ -472,7 +472,7 @@ Availability:Licensee
 * Duplicate, via **Ctrl + W**, only requires 1 undo step again.
 * _SoundCue_ and _SoundWave_ no longer play when double clicked, instead opens **SoundCue Editor**.
 * _DecalActors_ marked as **Hidden In Game** will now draw properly in the editor.
-* Optimizations for Gaussian DOF effect. If **Near Blur Size** is less than the threshold defined by the console variable `r.DepthOfFieldNearBlurSizeThreshold`, the near blur effect will be disabled and cheaper shaders used instead.
+* Optimizations for Gaussian DOF effect. If **Near Blur Size** is less than the threshold defined by the console variable `r.DepthOfField.NearBlurSizeThreshold`, the near blur effect will be disabled and cheaper shaders used instead.
 * Buffer visualization system now works correctly in orthographic view modes.
 * New _Landscape_ doesn't try to resize during alt-drag view rotation any more.
 * Optimized _Landscape_ spline rendering in editor by reducing complexity of the lines and optimizing `DrawDashedLine()`.
diff --git a/Engine/Source/Developer/StandaloneRenderer/Private/Windows/D3D/SlateD3DRenderingPolicy.cpp b/Engine/Source/Developer/StandaloneRenderer/Private/Windows/D3D/SlateD3DRenderingPolicy.cpp
--- a/Engine/Source/Developer/StandaloneRenderer/Private/Windows/D3D/SlateD3DRenderingPolicy.cpp
+++ b/Engine/Source/Developer/StandaloneRenderer/Private/Windows/D3D/SlateD3DRenderingPolicy.cpp
@@ -54,7 +54,7 @@ void FSlateD3D11RenderingPolicy::InitResources()
 	SamplerDesc.AddressU = D3D11_TEXTURE_ADDRESS_WRAP;
 	SamplerDesc.AddressV = D3D11_TEXTURE_ADDRESS_WRAP;
 	SamplerDesc.AddressW = D3D11_TEXTURE_ADDRESS_WRAP;
-	SamplerDesc.MaxAnisotropy = 0;
+	SamplerDesc.MaxAnisotropy = 1;
 	SamplerDesc.ComparisonFunc = D3D11_COMPARISON_NEVER;
 	SamplerDesc.MaxLOD = D3D11_FLOAT32_MAX;
 	SamplerDesc.MinLOD = 0;
diff --git a/Engine/Source/Runtime/Core/Private/HAL/ConsoleManager.cpp b/Engine/Source/Runtime/Core/Private/HAL/ConsoleManager.cpp
--- a/Engine/Source/Runtime/Core/Private/HAL/ConsoleManager.cpp
+++ b/Engine/Source/Runtime/Core/Private/HAL/ConsoleManager.cpp
@@ -1877,7 +1877,7 @@ static TAutoConsoleVariable<int32> CVarFinishCurrentFrame(
 static TAutoConsoleVariable<int32> CVarMaxAnistropy(
 	TEXT("r.MaxAnisotropy"),
 	4,
-	TEXT("MaxAnisotropy should range from 1 to 16. Higher values mean better texure quality when using anisotropic filtering but at a cost to performance."),
+	TEXT("MaxAnisotropy should range from 1 to 16. Higher values mean better texure quality when using anisotropic filtering but at a cost to performance. Default is 4."),
 	ECVF_Scalability | ECVF_RenderThreadSafe);
 
 static TAutoConsoleVariable<int32> CVarShadowMaxResolution(
diff --git a/Engine/Source/Runtime/Engine/Private/UnrealEngine.cpp b/Engine/Source/Runtime/Engine/Private/UnrealEngine.cpp
--- a/Engine/Source/Runtime/Engine/Private/UnrealEngine.cpp
+++ b/Engine/Source/Runtime/Engine/Private/UnrealEngine.cpp
@@ -202,13 +202,6 @@ static FAutoConsoleVariable CVarSystemResolution(
 	TEXT("     1920x1080wm for windowed mirror")
 	);
 
-static TAutoConsoleVariable<float> CVarDepthOfFieldNearBlurSizeThreshold(
-	TEXT("r.DepthOfFieldNearBlurSizeThreshold"),
-	0.01f,
-	TEXT("Sets the minimum near blur size before the effect is forcably disabled. Currently only affects Gaussian DOF.\n")
-	TEXT(" (default = 0.01f)"),
-	ECVF_RenderThreadSafe);
-
 #if !(UE_BUILD_SHIPPING || UE_BUILD_TEST)
 static TAutoConsoleVariable<float> CVarSetOverrideFPS(
 	TEXT("t.OverrideFPS"),
@@ -268,7 +261,6 @@ FCachedSystemScalabilityCVars::FCachedSystemScalabilityCVars()
 	, MaxShadowResolution(-1)
 	, ViewDistanceScale(-1)
 	, ViewDistanceScaleSquared(-1)
-	, GaussianDOFNearThreshold(-1)
 	, SimpleDynamicLighting(-1)
 {
 
@@ -301,10 +293,6 @@ void ScalabilityCVarsSinkCallback()
 		LocalScalabilityCVars.ViewDistanceScaleSquared = FMath::Square(LocalScalabilityCVars.ViewDistanceScale);
 	}
 
-	{
-		LocalScalabilityCVars.GaussianDOFNearThreshold = CVarDepthOfFieldNearBlurSizeThreshold.GetValueOnGameThread();
-	}
-
 	{
 		static const auto MaterialQualityLevelVar = IConsoleManager::Get().FindTConsoleVariableDataInt(TEXT("r.MaterialQualityLevel"));
 		LocalScalabilityCVars.MaterialQualityLevel = (EMaterialQualityLevel::Type)FMath::Clamp(MaterialQualityLevelVar->GetValueOnGameThread(), 0, 1);
@@ -419,6 +407,7 @@ void RefreshSamplerStatesCallback()
 		static const auto CVar = IConsoleManager::Get().FindTConsoleVariableDataInt(TEXT("r.MaxAnisotropy"));
 		int32 MaxAnisotropy = CVar->GetValueOnGameThread();
 		// compare against the default so with that number we avoid RefreshSamplerStates() calls on startup
+		// todo: This can be improved since we now have many defaults (see BaseScalability.ini)
 		static int32 LastMaxAnisotropy = 4;
 
 		if(LastMaxAnisotropy != MaxAnisotropy)
diff --git a/Engine/Source/Runtime/Engine/Public/UnrealEngine.h b/Engine/Source/Runtime/Engine/Public/UnrealEngine.h
--- a/Engine/Source/Runtime/Engine/Public/UnrealEngine.h
+++ b/Engine/Source/Runtime/Engine/Public/UnrealEngine.h
@@ -399,7 +399,8 @@ public:
 };
 
 /**
- * Cache some of the scalability CVars to avoid some virtual function calls and to detect changes and act accordingly if needed.
+ * Cache some of the scalability CVars to avoid some virtual function calls (no longer the case with the new console variable system, but we only have 1 render thread)
+ * and to detect changes and act accordingly if needed.
  * read by rendering thread[s], written by main thread, uses FlushRenderingCommands() to avoid conflict
  */
 struct FCachedSystemScalabilityCVars
@@ -407,15 +408,19 @@ struct FCachedSystemScalabilityCVars
 	bool bInitialized;
 	int32 DetailMode;
 	EMaterialQualityLevel::Type MaterialQualityLevel;
-	int32 MaxAnisotropy;
 	int32 MaxShadowResolution;
 	float ViewDistanceScale;
 	float ViewDistanceScaleSquared;
-	float GaussianDOFNearThreshold;
 	// 0:off, 1:on, -1:unknown
 	int32 SimpleDynamicLighting;
 
 	FCachedSystemScalabilityCVars();
+
+protected:
+	// This isn't public as it's only used to detect the change. Use ComputeAnisotropyRT()
+	int32 MaxAnisotropy;
+
+	friend void ScalabilityCVarsSinkCallback();
 };
 
 ENGINE_API bool AllowHighQualityLightmaps(ERHIFeatureLevel::Type FeatureLevel);
diff --git a/Engine/Source/Runtime/IOS/MetalRHI/Private/MetalState.cpp b/Engine/Source/Runtime/IOS/MetalRHI/Private/MetalState.cpp
--- a/Engine/Source/Runtime/IOS/MetalRHI/Private/MetalState.cpp
+++ b/Engine/Source/Runtime/IOS/MetalRHI/Private/MetalState.cpp
@@ -29,12 +29,12 @@ static MTLSamplerMinMagFilter TranslateFilterMode(ESamplerFilter Filter)
 	}
 }
 
-static uint32 GetMaxAnisotropy(ESamplerFilter Filter, uint32 MaxAniso)
+static uint32 GetMetalMaxAnisotropy(ESamplerFilter Filter, uint32 MaxAniso)
 {
 	switch (Filter)
 	{
 		case SF_AnisotropicPoint:
-		case SF_AnisotropicLinear:	return MaxAniso > 0 ? MaxAniso : GetCachedScalabilityCVars().MaxAnisotropy;
+		case SF_AnisotropicLinear:	return ComputeAnisotropyRT(MaxAniso);
 		default:					return 1;
 	}
 }
@@ -150,7 +150,7 @@ FMetalSamplerState::FMetalSamplerState(const FSamplerStateInitializerRHI& Initia
 	
 	Desc.minFilter = Desc.magFilter = TranslateFilterMode(Initializer.Filter);
 	Desc.mipFilter = TranslateMipFilterMode(Initializer.Filter);
-	Desc.maxAnisotropy = GetMaxAnisotropy(Initializer.Filter, Initializer.MaxAnisotropy);
+	Desc.maxAnisotropy = GetMetalMaxAnisotropy(Initializer.Filter, Initializer.MaxAnisotropy);
 	Desc.sAddressMode = TranslateWrapMode(Initializer.AddressU);
 	Desc.tAddressMode = TranslateWrapMode(Initializer.AddressV);
 	Desc.rAddressMode = TranslateWrapMode(Initializer.AddressW);
diff --git a/Engine/Source/Runtime/OpenGLDrv/Private/OpenGLState.cpp b/Engine/Source/Runtime/OpenGLDrv/Private/OpenGLState.cpp
--- a/Engine/Source/Runtime/OpenGLDrv/Private/OpenGLState.cpp
+++ b/Engine/Source/Runtime/OpenGLDrv/Private/OpenGLState.cpp
@@ -193,7 +193,7 @@ FSamplerStateRHIRef FOpenGLDynamicRHI::RHICreateSamplerState(const FSamplerState
 	case SF_AnisotropicLinear:
 		SamplerState->Data.MagFilter	= GL_LINEAR;
 		SamplerState->Data.MinFilter	= bComparisonEnabled ? GL_LINEAR : GL_LINEAR_MIPMAP_LINEAR;
-		SamplerState->Data.MaxAnisotropy = FMath::Clamp(Initializer.MaxAnisotropy > 0 ? Initializer.MaxAnisotropy : GetCachedScalabilityCVars().MaxAnisotropy, 1, GMaxOpenGLTextureFilterAnisotropic);
+		SamplerState->Data.MaxAnisotropy = FMath::Min<uint32>(ComputeAnisotropyRT(Initializer.MaxAnisotropy), GMaxOpenGLTextureFilterAnisotropic);
 		break;
 	case SF_Trilinear:
 		SamplerState->Data.MagFilter	= GL_LINEAR;
diff --git a/Engine/Source/Runtime/RHI/Public/RHIStaticStates.h b/Engine/Source/Runtime/RHI/Public/RHIStaticStates.h
--- a/Engine/Source/Runtime/RHI/Public/RHIStaticStates.h
+++ b/Engine/Source/Runtime/RHI/Public/RHIStaticStates.h
@@ -142,10 +142,10 @@ template<ESamplerFilter Filter=SF_Point,
 	ESamplerAddressMode AddressV=AM_Clamp,
 	ESamplerAddressMode AddressW=AM_Clamp, 
 	int32 MipBias = 0,
-	// Note: setting to a different value than GSystemSettings.MaxAnisotropy is only supported in D3D11
+	// Note: setting to a different value than GSystemSettings.MaxAnisotropy is only supported in D3D11 (is that still true?)
 	// A value of 0 will use GSystemSettings.MaxAnisotropy
-	int32 MaxAnisotropy=0,
-	uint32 BorderColor=0,
+	int32 MaxAnisotropy = 1,
+	uint32 BorderColor = 0,
 	/** Only supported in D3D11 */
 	ESamplerCompareFunction SamplerComparisonFunction=SCF_Never>
 class TStaticSamplerState : public TStaticStateRHI<TStaticSamplerState<Filter,AddressU,AddressV,AddressW,MipBias,MaxAnisotropy,BorderColor,SamplerComparisonFunction>,FSamplerStateRHIRef,FSamplerStateRHIParamRef>
diff --git a/Engine/Source/Runtime/RHI/Public/RHIUtilities.h b/Engine/Source/Runtime/RHI/Public/RHIUtilities.h
--- a/Engine/Source/Runtime/RHI/Public/RHIUtilities.h
+++ b/Engine/Source/Runtime/RHI/Public/RHIUtilities.h
@@ -506,3 +506,11 @@ inline void DrawIndexedPrimitiveUP(
 	FMemory::Memcpy( IndexBuffer, IndexData, NumIndices * IndexDataStride );
 	RHICmdList.EndDrawIndexedPrimitiveUP();
 }
+
+inline uint32 ComputeAnisotropyRT(int32 InitializerMaxAnisotropy)
+{
+	static const auto CVar = IConsoleManager::Get().FindTConsoleVariableDataInt(TEXT("r.MaxAnisotropy"));
+	int32 CVarValue = CVar->GetValueOnRenderThread();
+
+	return FMath::Clamp(InitializerMaxAnisotropy > 0 ? InitializerMaxAnisotropy : CVarValue, 1, 16);
+}
\ No newline at end of file
diff --git a/Engine/Source/Runtime/Renderer/Private/PostProcess/PostProcessing.cpp b/Engine/Source/Runtime/Renderer/Private/PostProcess/PostProcessing.cpp
--- a/Engine/Source/Runtime/Renderer/Private/PostProcess/PostProcessing.cpp
+++ b/Engine/Source/Runtime/Renderer/Private/PostProcess/PostProcessing.cpp
@@ -58,6 +58,13 @@ static TAutoConsoleVariable<int32> CVarUseMobileBloom(
 	TEXT("HACK: Set to 1 to use mobile bloom."),
 	ECVF_Scalability | ECVF_RenderThreadSafe);
 
+static TAutoConsoleVariable<float> CVarDepthOfFieldNearBlurSizeThreshold(
+	TEXT("r.DepthOfField.NearBlurSizeThreshold"),
+	0.01f,
+	TEXT("Sets the minimum near blur size before the effect is forcably disabled. Currently only affects Gaussian DOF.\n")
+	TEXT(" (default: 0.01)"),
+	ECVF_RenderThreadSafe);
+
 static TAutoConsoleVariable<float> CVarDepthOfFieldMaxSize(
 	TEXT("r.DepthOfField.MaxSize"),
 	100.0f,
@@ -78,30 +85,31 @@ static TAutoConsoleVariable<float> CVarUpscalePaniniD(
 	0,
 	TEXT("Allow and configure to apply a panini distortion to the rendered image. Values between 0 and 1 allow to fade the effect (lerp).\n")
 	TEXT("Implementation from research paper \"Pannini: A New Projection for Rendering Wide Angle Perspective Images\"\n")
-	TEXT("0: off (default)\n")
+	TEXT(" 0: off (default)\n")
 	TEXT(">0: enabled (requires an extra post processing pass if upsampling wasn't used - see r.ScreenPercentage)\n")
-	TEXT("1: Panini cylindrical stereographic projection"),
+	TEXT(" 1: Panini cylindrical stereographic projection"),
 	ECVF_RenderThreadSafe);
 
 static TAutoConsoleVariable<float> CVarUpscalePaniniS(
 	TEXT("r.Upscale.Panini.S"),
 	0,
 	TEXT("Panini projection's hard vertical compression factor.\n")
-	TEXT("0: no vertical compression factor (default)\n")
-	TEXT("1: Hard vertical compression"),
+	TEXT(" 0: no vertical compression factor (default)\n")
+	TEXT(" 1: Hard vertical compression"),
 	ECVF_RenderThreadSafe);
 
 static TAutoConsoleVariable<float> CVarUpscalePaniniScreenFit(
 	TEXT("r.Upscale.Panini.ScreenFit"),
 	1.0f,
 	TEXT("Panini projection screen fit effect factor (lerp).\n")
-	TEXT("0: fit vertically\n")
-	TEXT("1: fit horizontally (default)"),
+	TEXT(" 0: fit vertically\n")
+	TEXT(" 1: fit horizontally (default)"),
 	ECVF_RenderThreadSafe);
 
 static TAutoConsoleVariable<int32> CVarUpscaleQuality(
 	TEXT("r.Upscale.Quality"),
 	3,
+	TEXT("Defines the quality in which ScreenPercentage and WindowedFullscreen scales the 3d rendering.\n")
 	TEXT(" 0: Nearest filtering\n")
 	TEXT(" 1: Simple Bilinear\n")
 	TEXT(" 2: 4 tap bilinear\n")
@@ -356,7 +364,12 @@ static bool AddPostProcessDepthOfFieldGaussian(FPostprocessContext& Context, FDe
 	NearSize = FMath::Min(NearSize, MaxSize);
 
 	Out.bFar = FarSize >= 0.01f;
-	Out.bNear = NearSize >= GetCachedScalabilityCVars().GaussianDOFNearThreshold;
+
+	{
+		const float CVarThreshold = CVarDepthOfFieldNearBlurSizeThreshold.GetValueOnRenderThread();
+
+		Out.bNear = (NearSize >= CVarThreshold);
+	}
 
 	if(Context.View.Family->EngineShowFlags.VisualizeDOF)
 	{
diff --git a/Engine/Source/Runtime/Windows/D3D11RHI/Private/D3D11State.cpp b/Engine/Source/Runtime/Windows/D3D11RHI/Private/D3D11State.cpp
--- a/Engine/Source/Runtime/Windows/D3D11RHI/Private/D3D11State.cpp
+++ b/Engine/Source/Runtime/Windows/D3D11RHI/Private/D3D11State.cpp
@@ -129,13 +129,11 @@ FSamplerStateRHIRef FD3D11DynamicRHI::RHICreateSamplerState(const FSamplerStateI
 	D3D11_SAMPLER_DESC SamplerDesc;
 	FMemory::Memzero(&SamplerDesc,sizeof(D3D11_SAMPLER_DESC));
 
-	int32 MaxAnisotropyCVar = GetCachedScalabilityCVars().MaxAnisotropy;
-
 	SamplerDesc.AddressU = TranslateAddressMode(Initializer.AddressU);
 	SamplerDesc.AddressV = TranslateAddressMode(Initializer.AddressV);
 	SamplerDesc.AddressW = TranslateAddressMode(Initializer.AddressW);
 	SamplerDesc.MipLODBias = Initializer.MipBias;
-	SamplerDesc.MaxAnisotropy = FMath::Clamp(Initializer.MaxAnisotropy > 0 ? Initializer.MaxAnisotropy : MaxAnisotropyCVar,1,16);
+	SamplerDesc.MaxAnisotropy = ComputeAnisotropyRT(Initializer.MaxAnisotropy);
 	SamplerDesc.MinLOD = Initializer.MinMipLevel;
 	SamplerDesc.MaxLOD = Initializer.MaxMipLevel;
 
diff --git a/Engine/Source/Runtime/Windows/D3D12RHI/Private/D3D12State.cpp b/Engine/Source/Runtime/Windows/D3D12RHI/Private/D3D12State.cpp
--- a/Engine/Source/Runtime/Windows/D3D12RHI/Private/D3D12State.cpp
+++ b/Engine/Source/Runtime/Windows/D3D12RHI/Private/D3D12State.cpp
@@ -132,13 +132,11 @@ FSamplerStateRHIRef FD3D12Device::CreateSamplerState(const FSamplerStateInitiali
     D3D12_SAMPLER_DESC SamplerDesc;
     FMemory::Memzero(&SamplerDesc, sizeof(D3D12_SAMPLER_DESC));
 
-    int32 MaxAnisotropyCVar = GetCachedScalabilityCVars().MaxAnisotropy;
-
     SamplerDesc.AddressU = TranslateAddressMode(Initializer.AddressU);
     SamplerDesc.AddressV = TranslateAddressMode(Initializer.AddressV);
     SamplerDesc.AddressW = TranslateAddressMode(Initializer.AddressW);
     SamplerDesc.MipLODBias = Initializer.MipBias;
-    SamplerDesc.MaxAnisotropy = FMath::Clamp(Initializer.MaxAnisotropy > 0 ? Initializer.MaxAnisotropy : MaxAnisotropyCVar, 1, 16);
+    SamplerDesc.MaxAnisotropy = ComputeAnisotropyRT(Initializer.MaxAnisotropy);
     SamplerDesc.MinLOD = Initializer.MinMipLevel;
     SamplerDesc.MaxLOD = Initializer.MaxMipLevel;
 
