UE4 - Integrate removing post process tonemap ES2 on console (2688834)

diff --git a/Engine/Source/Runtime/Renderer/Private/PostProcess/PostProcessTonemap.cpp b/Engine/Source/Runtime/Renderer/Private/PostProcess/PostProcessTonemap.cpp
--- a/Engine/Source/Runtime/Renderer/Private/PostProcess/PostProcessTonemap.cpp
+++ b/Engine/Source/Runtime/Renderer/Private/PostProcess/PostProcessTonemap.cpp
@@ -1182,7 +1182,7 @@ class FPostProcessTonemapPS_ES2 : public FGlobalShader
 	{
 		// This is only used on ES2.
 		// TODO: Make this only compile on PC/Mobile (and not console).
-		return true;
+		return !IsConsolePlatform(Platform);
 	}
 
 	static void ModifyCompilationEnvironment(EShaderPlatform Platform, FShaderCompilerEnvironment& OutEnvironment)
@@ -1369,7 +1369,7 @@ class FPostProcessTonemapVS_ES2 : public FGlobalShader
 
 	static bool ShouldCache(EShaderPlatform Platform)
 	{
-		return true;
+		return !IsConsolePlatform(Platform);
 	}
 
 	FPostProcessTonemapVS_ES2() { }
