Integrated from Orion

CL 2688832

	fixed issues with volumetric deferred decals (quality, baked shadows, depth out)
	  partly from Inigo Quilez

diff --git a/Engine/Shaders/BasePassPixelShader.usf b/Engine/Shaders/BasePassPixelShader.usf
--- a/Engine/Shaders/BasePassPixelShader.usf
+++ b/Engine/Shaders/BasePassPixelShader.usf
@@ -637,7 +637,7 @@ void Main(
 	#endif
 #endif
 
-	OPTIONAL_OutDepth
+	OPTIONAL_OutDepthConservative
 
 #if EDITOR_ALPHA2COVERAGE != 0
    ,in uint InCoverage : SV_Coverage
diff --git a/Engine/Shaders/DeferredDecal.usf b/Engine/Shaders/DeferredDecal.usf
--- a/Engine/Shaders/DeferredDecal.usf
+++ b/Engine/Shaders/DeferredDecal.usf
@@ -89,6 +89,9 @@ float DistanceFunction(float4 ScreenPosition, float3 Pos)
 	return 0;
 }
 
+// could be exposed to the user
+const static float coneSpread = 0.00002f;
+
 // @return t=0..1, -1: no hit
 float Raycast(float4 ScreenPosition, float3 StartPos, float3 EndPos)
 {
@@ -103,11 +106,16 @@ float Raycast(float4 ScreenPosition, float3 StartPos, float3 EndPos)
 	{
 		float Distance = DistanceFunction(ScreenPosition, StartPos + t * Direction);
 
-		if(Distance <= i * 0.001f)
+		if(Distance <= t * coneSpread)
 		{
 			return t / RayLength;
 		}
 
+		if(t > RayLength)
+		{
+		   break;
+		}
+
 		t += Distance;
 	}
 
@@ -115,9 +123,9 @@ float Raycast(float4 ScreenPosition, float3 StartPos, float3 EndPos)
 }
 
 // like explained here http://iquilezles.org/www/articles/derivative/derivative.htm
-float3 GradientNotNormalized(float4 ScreenPosition, float3 OSPos)
+float3 GradientNotNormalized(float4 ScreenPosition, float3 OSPos, float t)
 {
-	float eps = 0.1f;
+	float eps = t * coneSpread;
 
 	float2 e = float2(eps, 0);
 
@@ -131,9 +139,7 @@ float3 GradientNotNormalized(float4 ScreenPosition, float3 OSPos)
 // decal pixel shader
 #if DECAL_PROJECTION
 
-#if DECAL_BLEND_MODE == DECALBLENDMODEID_VOLUMETRIC
-// should work, but doesn't 	EARLYDEPTHSTENCIL
-#endif
+// EARLYDEPTHSTENCIL for volumetric decals could work here, but only if we are ouside of the decal
 
 void MainPS(
 	in float4 ScreenPosition : TEXCOORD0			// is ClipPosition
@@ -156,14 +162,25 @@ void MainPS(
 	, out float4 OutTarget3 : SV_Target3
 #endif
 
+#if DECAL_RENDERTARGET_COUNT > 4
+	, out float4 OutTarget4 : SV_Target4
+#endif
+
+	// todo: if we are not inside a volumetric decal we could use OPTIONAL_OutDepthConservative for better performance, this would require a shader permutation
 	OPTIONAL_OutDepth
 
 	)
 {
+#if DECAL_RENDERTARGET_COUNT > 4
+	// default to no shadows, good for volumetric decals
+	OutTarget4 = 1;
+#endif
+
 	bool bVolumetric = DECAL_BLEND_MODE == DECALBLENDMODEID_VOLUMETRIC;
 
 	float3 OSPosition;
 	float3 WSPosition;
+	float HitT;
 
 	if(bVolumetric)
 	{
@@ -191,7 +208,7 @@ void MainPS(
 		float3 OSBoxStartRay = OSStartRay + OSBoxMinMax.x * (OSEndRay - OSStartRay);
 		float3 OSBoxEndRay = OSStartRay + OSBoxMinMax.y * (OSEndRay - OSStartRay);
 
-		float HitT = Raycast(ScreenPosition, OSBoxStartRay, OSBoxEndRay);
+		HitT = Raycast(ScreenPosition, OSBoxStartRay, OSBoxEndRay);
 
 		clip(HitT);
 
@@ -251,7 +268,7 @@ void MainPS(
 		if(bVolumetric)
 		{
 			// not normalized
-			float3 OSNormal = GradientNotNormalized(ScreenPosition, OSPosition);
+			float3 OSNormal = GradientNotNormalized(ScreenPosition, OSPosition, HitT);
 			
 			// not normalized
 			float3 WSNormal = mul(float4(OSNormal, 0), transpose(WorldToDecal)).xyz;
diff --git a/Engine/Shaders/DepthOnlyPixelShader.usf b/Engine/Shaders/DepthOnlyPixelShader.usf
--- a/Engine/Shaders/DepthOnlyPixelShader.usf
+++ b/Engine/Shaders/DepthOnlyPixelShader.usf
@@ -14,7 +14,7 @@ void Main(
 	float4 PixelPosition : TEXCOORD6,
 	in INPUT_POSITION_QUALIFIERS float4 SvPosition : SV_Position
 	OPTIONAL_IsFrontFace
-	OPTIONAL_OutDepth,
+	OPTIONAL_OutDepthConservative,
 #endif
 	out float4 OutColor : SV_Target0
 	)
diff --git a/Engine/Shaders/HitProxyPixelShader.usf b/Engine/Shaders/HitProxyPixelShader.usf
--- a/Engine/Shaders/HitProxyPixelShader.usf
+++ b/Engine/Shaders/HitProxyPixelShader.usf
@@ -21,7 +21,7 @@ void Main(
 	float4 PixelPosition	: TEXCOORD6,
 	in INPUT_POSITION_QUALIFIERS float4 SvPosition	: SV_Position
 	OPTIONAL_IsFrontFace
-	OPTIONAL_OutDepth
+	OPTIONAL_OutDepthConservative
 	,out float4 OutColor		: SV_Target0
 	)
 {
diff --git a/Engine/Shaders/MaterialTemplate.usf b/Engine/Shaders/MaterialTemplate.usf
--- a/Engine/Shaders/MaterialTemplate.usf
+++ b/Engine/Shaders/MaterialTemplate.usf
@@ -1719,8 +1719,10 @@ half3x3 AssembleTangentToWorld( half3 TangentToWorld0, half4 TangentToWorld2 )
 #endif
 
 #if OUTPUT_PIXEL_DEPTH_OFFSET
-	#define OPTIONAL_OutDepth ,out float OutDepth : DEPTH_WRITE_SEMANTIC
+	#define OPTIONAL_OutDepthConservative ,out float OutDepth : DEPTH_WRITE_SEMANTIC
+	#define OPTIONAL_OutDepth ,out float OutDepth : SV_DEPTH
 #else
+	#define OPTIONAL_OutDepthConservative
 	#define OPTIONAL_OutDepth
 #endif
 
diff --git a/Engine/Source/Runtime/Renderer/Private/CompositionLighting/PostProcessDeferredDecals.cpp b/Engine/Source/Runtime/Renderer/Private/CompositionLighting/PostProcessDeferredDecals.cpp
--- a/Engine/Source/Runtime/Renderer/Private/CompositionLighting/PostProcessDeferredDecals.cpp
+++ b/Engine/Source/Runtime/Renderer/Private/CompositionLighting/PostProcessDeferredDecals.cpp
@@ -478,6 +478,7 @@ void FRCPassPostProcessDeferredDecals::Process(FRenderingCompositePassContext& C
 		GBufferAIndex,
 		GBufferBIndex,
 		GBufferCIndex,
+		GBufferEIndex,
 		DBufferAIndex,
 		DBufferBIndex,
 		DBufferCIndex,
@@ -601,50 +602,43 @@ void FRCPassPostProcessDeferredDecals::Process(FRenderingCompositePassContext& C
 
 					switch (CurrentRenderTargetMode)
 					{
-					case FDecalRendering::RTM_SceneColorAndGBuffer:
-					{
-																	  TargetsToResolve[SceneColorIndex] = SceneContext.GetSceneColor()->GetRenderTargetItem().TargetableTexture;
-																	  TargetsToResolve[GBufferAIndex] = SceneContext.GBufferA->GetRenderTargetItem().TargetableTexture;
-																	  TargetsToResolve[GBufferBIndex] = SceneContext.GBufferB->GetRenderTargetItem().TargetableTexture;
-																	  TargetsToResolve[GBufferCIndex] = SceneContext.GBufferC->GetRenderTargetItem().TargetableTexture;
-
-																	  SetRenderTargets(RHICmdList, 4, TargetsToResolve, SceneContext.GetSceneDepthSurface(), ESimpleRenderTargetMode::EExistingColorAndDepth, FExclusiveDepthStencil::DepthRead_StencilWrite);
-					}
-						break;
-
-					case FDecalRendering::RTM_SceneColorAndGBufferDepthWrite:
-					{
-																				TargetsToResolve[SceneColorIndex] = SceneContext.GetSceneColor()->GetRenderTargetItem().TargetableTexture;
-																				TargetsToResolve[GBufferAIndex] = SceneContext.GBufferA->GetRenderTargetItem().TargetableTexture;
-																				TargetsToResolve[GBufferBIndex] = SceneContext.GBufferB->GetRenderTargetItem().TargetableTexture;
-																				TargetsToResolve[GBufferCIndex] = SceneContext.GBufferC->GetRenderTargetItem().TargetableTexture;
-
-																				SetRenderTargets(RHICmdList, 4, TargetsToResolve, SceneContext.GetSceneDepthSurface(), ESimpleRenderTargetMode::EExistingColorAndDepth, FExclusiveDepthStencil::DepthWrite_StencilWrite);
-					}
-						break;
-
-					case FDecalRendering::RTM_GBufferNormal:
-						TargetsToResolve[GBufferAIndex] = SceneContext.GBufferA->GetRenderTargetItem().TargetableTexture;
-						SetRenderTarget(RHICmdList, TargetsToResolve[GBufferAIndex], SceneContext.GetSceneDepthSurface(), ESimpleRenderTargetMode::EExistingColorAndDepth, FExclusiveDepthStencil::DepthRead_StencilWrite);
-						break;
-
-					case FDecalRendering::RTM_SceneColor:
-						TargetsToResolve[SceneColorIndex] = SceneContext.GetSceneColor()->GetRenderTargetItem().TargetableTexture;
-						SetRenderTarget(RHICmdList, TargetsToResolve[SceneColorIndex], SceneContext.GetSceneDepthSurface(), ESimpleRenderTargetMode::EExistingColorAndDepth, FExclusiveDepthStencil::DepthRead_StencilWrite);
-						break;
-
-					case FDecalRendering::RTM_DBuffer:
-					{
-														 TargetsToResolve[DBufferAIndex] = SceneContext.DBufferA->GetRenderTargetItem().TargetableTexture;
-														 TargetsToResolve[DBufferBIndex] = SceneContext.DBufferB->GetRenderTargetItem().TargetableTexture;
-														 TargetsToResolve[DBufferCIndex] = SceneContext.DBufferC->GetRenderTargetItem().TargetableTexture;
-														 SetRenderTargets(RHICmdList, 3, &TargetsToResolve[DBufferAIndex], SceneContext.GetSceneDepthSurface(), ESimpleRenderTargetMode::EExistingColorAndDepth, FExclusiveDepthStencil::DepthRead_StencilWrite);
-					}
-						break;
-
-					default:
-						check(0);
-						break;
+						case FDecalRendering::RTM_SceneColorAndGBuffer:
+							TargetsToResolve[SceneColorIndex] = SceneContext.GetSceneColor()->GetRenderTargetItem().TargetableTexture;
+							TargetsToResolve[GBufferAIndex] = SceneContext.GBufferA->GetRenderTargetItem().TargetableTexture;
+							TargetsToResolve[GBufferBIndex] = SceneContext.GBufferB->GetRenderTargetItem().TargetableTexture;
+							TargetsToResolve[GBufferCIndex] = SceneContext.GBufferC->GetRenderTargetItem().TargetableTexture;
+							SetRenderTargets(RHICmdList, 4, TargetsToResolve, SceneContext.GetSceneDepthSurface(), ESimpleRenderTargetMode::EExistingColorAndDepth, FExclusiveDepthStencil::DepthRead_StencilWrite);
+							break;
+
+						case FDecalRendering::RTM_SceneColorAndGBufferDepthWrite:
+							TargetsToResolve[SceneColorIndex] = SceneContext.GetSceneColor()->GetRenderTargetItem().TargetableTexture;
+							TargetsToResolve[GBufferAIndex] = SceneContext.GBufferA->GetRenderTargetItem().TargetableTexture;
+							TargetsToResolve[GBufferBIndex] = SceneContext.GBufferB->GetRenderTargetItem().TargetableTexture;
+							TargetsToResolve[GBufferCIndex] = SceneContext.GBufferC->GetRenderTargetItem().TargetableTexture;
+							TargetsToResolve[GBufferEIndex] = SceneContext.GBufferE->GetRenderTargetItem().TargetableTexture;
+							SetRenderTargets(RHICmdList, 5, TargetsToResolve, SceneContext.GetSceneDepthSurface(), ESimpleRenderTargetMode::EExistingColorAndDepth, FExclusiveDepthStencil::DepthWrite_StencilWrite);
+							break;
+
+						case FDecalRendering::RTM_GBufferNormal:
+							TargetsToResolve[GBufferAIndex] = SceneContext.GBufferA->GetRenderTargetItem().TargetableTexture;
+							SetRenderTarget(RHICmdList, TargetsToResolve[GBufferAIndex], SceneContext.GetSceneDepthSurface(), ESimpleRenderTargetMode::EExistingColorAndDepth, FExclusiveDepthStencil::DepthRead_StencilWrite);
+							break;
+
+						case FDecalRendering::RTM_SceneColor:
+							TargetsToResolve[SceneColorIndex] = SceneContext.GetSceneColor()->GetRenderTargetItem().TargetableTexture;
+							SetRenderTarget(RHICmdList, TargetsToResolve[SceneColorIndex], SceneContext.GetSceneDepthSurface(), ESimpleRenderTargetMode::EExistingColorAndDepth, FExclusiveDepthStencil::DepthRead_StencilWrite);
+							break;
+
+						case FDecalRendering::RTM_DBuffer:
+							TargetsToResolve[DBufferAIndex] = SceneContext.DBufferA->GetRenderTargetItem().TargetableTexture;
+							TargetsToResolve[DBufferBIndex] = SceneContext.DBufferB->GetRenderTargetItem().TargetableTexture;
+							TargetsToResolve[DBufferCIndex] = SceneContext.DBufferC->GetRenderTargetItem().TargetableTexture;
+							SetRenderTargets(RHICmdList, 3, &TargetsToResolve[DBufferAIndex], SceneContext.GetSceneDepthSurface(), ESimpleRenderTargetMode::EExistingColorAndDepth, FExclusiveDepthStencil::DepthRead_StencilWrite);
+							break;
+
+						default:
+							check(0);
+							break;
 					}
 					Context.SetViewportAndCallRHI(DestRect);
 
diff --git a/Engine/Source/Runtime/Renderer/Private/DecalRenderingShared.cpp b/Engine/Source/Runtime/Renderer/Private/DecalRenderingShared.cpp
--- a/Engine/Source/Runtime/Renderer/Private/DecalRenderingShared.cpp
+++ b/Engine/Source/Runtime/Renderer/Private/DecalRenderingShared.cpp
@@ -403,7 +403,7 @@ uint32 FDecalRendering::ComputeRenderTargetCount(EShaderPlatform Platform, ERend
 	switch(RenderTargetMode)
 	{
 		case RTM_SceneColorAndGBuffer:				return 4;
-		case RTM_SceneColorAndGBufferDepthWrite:	return 4;
+		case RTM_SceneColorAndGBufferDepthWrite:	return 5;
 		case RTM_DBuffer:							return 3;
 		case RTM_GBufferNormal:						return 1;
 		case RTM_SceneColor:						return 1;
