Move IndirectLighitngCache update to a task to overlap with shadow setup.
Also some small improvements to the ILC update itself.

diff --git a/Engine/Source/Runtime/RenderCore/Private/RenderCore.cpp b/Engine/Source/Runtime/RenderCore/Private/RenderCore.cpp
--- a/Engine/Source/Runtime/RenderCore/Private/RenderCore.cpp
+++ b/Engine/Source/Runtime/RenderCore/Private/RenderCore.cpp
@@ -72,6 +72,10 @@ DEFINE_STAT(STAT_LightInteractionMemory);
 
 DEFINE_STAT(STAT_GatherShadowPrimitivesTime);
 DEFINE_STAT(STAT_UpdateIndirectLightingCache);
+DEFINE_STAT(STAT_UpdateIndirectLightingCachePrims);
+DEFINE_STAT(STAT_UpdateIndirectLightingCacheBlocks);
+DEFINE_STAT(STAT_UpdateIndirectLightingCacheTransitions);
+DEFINE_STAT(STAT_UpdateIndirectLightingCacheFinalize);
 DEFINE_STAT(STAT_SortStaticDrawLists);
 DEFINE_STAT(STAT_InitDynamicShadowsTime);
 DEFINE_STAT(STAT_InitProjectedShadowVisibility);
diff --git a/Engine/Source/Runtime/RenderCore/Public/RenderCore.h b/Engine/Source/Runtime/RenderCore/Public/RenderCore.h
--- a/Engine/Source/Runtime/RenderCore/Public/RenderCore.h
+++ b/Engine/Source/Runtime/RenderCore/Public/RenderCore.h
@@ -59,6 +59,10 @@ DECLARE_CYCLE_STAT_EXTERN(TEXT("Add View Whole Scene Shadows"),STAT_AddViewDepen
 DECLARE_CYCLE_STAT_EXTERN(TEXT("Setup Interaction Shadows"),STAT_SetupInteractionShadows,STATGROUP_InitViews, RENDERCORE_API);
 DECLARE_CYCLE_STAT_EXTERN(TEXT("SortStaticDrawLists"),STAT_SortStaticDrawLists,STATGROUP_InitViews, RENDERCORE_API);
 DECLARE_CYCLE_STAT_EXTERN(TEXT("Update Indirect Lighting Cache"),STAT_UpdateIndirectLightingCache,STATGROUP_InitViews, RENDERCORE_API);
+DECLARE_CYCLE_STAT_EXTERN(TEXT("Update Indirect Lighting Cache Prims"), STAT_UpdateIndirectLightingCachePrims, STATGROUP_InitViews, RENDERCORE_API);
+DECLARE_CYCLE_STAT_EXTERN(TEXT("Update Indirect Lighting Cache Blocks"), STAT_UpdateIndirectLightingCacheBlocks, STATGROUP_InitViews, RENDERCORE_API);
+DECLARE_CYCLE_STAT_EXTERN(TEXT("Update Indirect Lighting Cache Transitions"), STAT_UpdateIndirectLightingCacheTransitions, STATGROUP_InitViews, RENDERCORE_API);
+DECLARE_CYCLE_STAT_EXTERN(TEXT("Update Indirect Lighting Cache Finalize"), STAT_UpdateIndirectLightingCacheFinalize, STATGROUP_InitViews, RENDERCORE_API);
 DECLARE_CYCLE_STAT_EXTERN(TEXT("GatherShadowPrimitives"),STAT_GatherShadowPrimitivesTime,STATGROUP_InitViews, RENDERCORE_API);
 DECLARE_DWORD_COUNTER_STAT_EXTERN(TEXT("Processed primitives"),STAT_ProcessedPrimitives,STATGROUP_InitViews, RENDERCORE_API);
 DECLARE_DWORD_COUNTER_STAT_EXTERN(TEXT("Frustum Culled primitives"),STAT_CulledPrimitives,STATGROUP_InitViews, RENDERCORE_API);
diff --git a/Engine/Source/Runtime/Renderer/Private/ForwardShadingRenderer.cpp b/Engine/Source/Runtime/Renderer/Private/ForwardShadingRenderer.cpp
--- a/Engine/Source/Runtime/Renderer/Private/ForwardShadingRenderer.cpp
+++ b/Engine/Source/Runtime/Renderer/Private/ForwardShadingRenderer.cpp
@@ -35,9 +35,10 @@ void FForwardShadingSceneRenderer::InitViews(FRHICommandListImmediate& RHICmdLis
 
 	SCOPE_CYCLE_COUNTER(STAT_InitViewsTime);
 
+	FILCUpdatePrimTaskData ILCTaskData;
 	PreVisibilityFrameSetup(RHICmdList);
 	ComputeViewVisibility(RHICmdList);
-	PostVisibilityFrameSetup();
+	PostVisibilityFrameSetup(ILCTaskData);
 
 	bool bDynamicShadows = ViewFamily.EngineShowFlags.DynamicShadows && GetShadowQuality() > 0;
 
@@ -47,6 +48,12 @@ void FForwardShadingSceneRenderer::InitViews(FRHICommandListImmediate& RHICmdLis
 		InitDynamicShadows(RHICmdList);		
 	}
 
+	// if we kicked off ILC update via task, wait and finalize.
+	if (ILCTaskData.TaskRef.IsValid())
+	{
+		Scene->IndirectLightingCache.FinalizeCacheUpdates(Scene, *this, ILCTaskData);
+	}
+
 	// initialize per-view uniform buffer.  Pass in shadow info as necessary.
 	for (int32 ViewIndex = 0; ViewIndex < Views.Num(); ViewIndex++)
 	{
diff --git a/Engine/Source/Runtime/Renderer/Private/IndirectLightingCache.cpp b/Engine/Source/Runtime/Renderer/Private/IndirectLightingCache.cpp
--- a/Engine/Source/Runtime/Renderer/Private/IndirectLightingCache.cpp
+++ b/Engine/Source/Runtime/Renderer/Private/IndirectLightingCache.cpp
@@ -15,6 +15,7 @@
  * This provides some stability as bounds get larger and smaller, although by adding some waste.
  */
 const float BoundSizeRoundUpBase = FMath::Sqrt(2);
+const float LogEBoundSizeRoundUpBase = FMath::Loge(BoundSizeRoundUpBase);
 
 //////////////////////////////////////////////////////////////////////////////////////////////////////////
 // Console variables that can be changed at runtime to configure or debug the indirect lighting cache
@@ -182,7 +183,7 @@ static bool IsTexelMinValid(FIntVector TexelMin)
 }
 
 FIndirectLightingCacheBlock& FIndirectLightingCache::FindBlock(FIntVector TexelMin)
-{
+{	
 	checkSlow(IsTexelMinValid(TexelMin));
 	return VolumeBlocks.FindChecked(TexelMin);
 }
@@ -221,15 +222,15 @@ bool FIndirectLightingCache::AllocateBlock(int32 Size, FIntVector& OutMin)
 }
 
 void FIndirectLightingCache::CalculateBlockPositionAndSize(const FBoxSphereBounds& Bounds, int32 TexelSize, FVector& OutMin, FVector& OutSize) const
-{
+{	
 	FVector RoundedBoundsSize;
 
-	// Find the exponent needed to represent the bounds size if BoundSizeRoundUpBase is the base
-	RoundedBoundsSize.X = FMath::Max(1.f, FMath::LogX(BoundSizeRoundUpBase, Bounds.BoxExtent.X * 2));
-	RoundedBoundsSize.Y = FMath::Max(1.f, FMath::LogX(BoundSizeRoundUpBase, Bounds.BoxExtent.Y * 2));
-	RoundedBoundsSize.Z = FMath::Max(1.f, FMath::LogX(BoundSizeRoundUpBase, Bounds.BoxExtent.Z * 2));
+	// Find the exponent needed to represent the bounds size if BoundSizeRoundUpBase is the base	
+	RoundedBoundsSize.X = FMath::Max(1.f, FMath::Loge(Bounds.BoxExtent.X * 2) / LogEBoundSizeRoundUpBase);
+	RoundedBoundsSize.Y = FMath::Max(1.f, FMath::Loge(Bounds.BoxExtent.Y * 2) / LogEBoundSizeRoundUpBase);
+	RoundedBoundsSize.Z = FMath::Max(1.f, FMath::Loge(Bounds.BoxExtent.Z * 2) / LogEBoundSizeRoundUpBase);
 
-	// Round up to the next integer exponent to provide stability even when Bounds.BoxExtent is changing
+	// Round up to the next integer exponent to provide stability even when Bounds.BoxExtent is changing	
 	RoundedBoundsSize.X = FMath::Pow(BoundSizeRoundUpBase, FMath::TruncToInt(RoundedBoundsSize.X) + 1);
 	RoundedBoundsSize.Y = FMath::Pow(BoundSizeRoundUpBase, FMath::TruncToInt(RoundedBoundsSize.Y) + 1);
 	RoundedBoundsSize.Z = FMath::Pow(BoundSizeRoundUpBase, FMath::TruncToInt(RoundedBoundsSize.Z) + 1);
@@ -322,7 +323,7 @@ FIndirectLightingCacheAllocation* FIndirectLightingCache::AllocatePrimitive(cons
 }
 
 FIndirectLightingCacheAllocation* FIndirectLightingCache::CreateAllocation(int32 BlockSize, const FBoxSphereBounds& Bounds, bool bPointSample, bool bUnbuiltPreview)
-{
+{	
 	check(BlockSize > 1 || bPointSample);
 
 	FIndirectLightingCacheAllocation* NewAllocation = new FIndirectLightingCacheAllocation();
@@ -369,61 +370,154 @@ FIndirectLightingCacheAllocation* FIndirectLightingCache::FindPrimitiveAllocatio
 	return PrimitiveAllocations.FindRef(PrimitiveId);
 }
 
+class FUpdateCachePrimitivesTask
+{
+	FIndirectLightingCache* ILC;
+	FScene* Scene;
+	FSceneRenderer& Renderer;
+	TMap<FIntVector, FBlockUpdateInfo>& BlocksToUpdate;
+	TArray<FIndirectLightingCacheAllocation*>& TransitionsOverTimeToUpdate;
+	bool bAllowUnbuiltPreview;
+
+public:
+
+	FUpdateCachePrimitivesTask(FIndirectLightingCache* InILC, FScene* InScene, FSceneRenderer& InRenderer, bool bInAllowUnbuiltPreview, TMap<FIntVector, FBlockUpdateInfo>& OutBlocksToUpdate, TArray<FIndirectLightingCacheAllocation*>& OutTransitionsOverTimeToUpdate)
+		: ILC(InILC)
+		, Scene(InScene)
+		, Renderer(InRenderer)
+		, BlocksToUpdate(OutBlocksToUpdate)
+		, TransitionsOverTimeToUpdate(OutTransitionsOverTimeToUpdate)
+		, bAllowUnbuiltPreview(bInAllowUnbuiltPreview)
+	{
+	}
+
+	FORCEINLINE TStatId GetStatId() const
+	{
+		RETURN_QUICK_DECLARE_CYCLE_STAT(FUpdateCachePrimitivesTask, STATGROUP_TaskGraphTasks);
+	}
+
+	ENamedThreads::Type GetDesiredThread()
+	{
+		return ENamedThreads::AnyThread;
+	}
+
+	static ESubsequentsMode::Type GetSubsequentsMode() { return ESubsequentsMode::TrackSubsequents; }
+
+	void DoTask(ENamedThreads::Type CurrentThread, const FGraphEventRef& MyCompletionGraphEvent)
+	{
+		ILC->UpdateCachePrimitivesInternal(Scene, Renderer, bAllowUnbuiltPreview, BlocksToUpdate, TransitionsOverTimeToUpdate);
+	}
+};
+
+void FIndirectLightingCache::StartUpdateCachePrimitivesTask(FScene* Scene, FSceneRenderer& Renderer, bool bAllowUnbuiltPreview, FILCUpdatePrimTaskData& OutTaskData)
+{
+	OutTaskData.TaskRef = TGraphTask<FUpdateCachePrimitivesTask>::CreateTask().ConstructAndDispatchWhenReady(this, Scene, Renderer, bAllowUnbuiltPreview, OutTaskData.OutBlocksToUpdate, OutTaskData.OutTransitionsOverTimeToUpdate);
+}
+
+void FIndirectLightingCache::FinalizeCacheUpdates(FScene* Scene, FSceneRenderer& Renderer, FILCUpdatePrimTaskData& TaskData)
+{
+	SCOPE_CYCLE_COUNTER(STAT_UpdateIndirectLightingCacheFinalize);	
+	FTaskGraphInterface::Get().WaitUntilTaskCompletes(TaskData.TaskRef, ENamedThreads::AnyThread);
+	FinalizeUpdateInternal_RenderThread(Scene, Renderer, TaskData.OutBlocksToUpdate, TaskData.OutTransitionsOverTimeToUpdate);
+}
+
 void FIndirectLightingCache::UpdateCache(FScene* Scene, FSceneRenderer& Renderer, bool bAllowUnbuiltPreview)
 {
+	SCOPE_CYCLE_COUNTER(STAT_UpdateIndirectLightingCache);
+
+	TMap<FIntVector, FBlockUpdateInfo>BlocksToUpdate;
+	TArray<FIndirectLightingCacheAllocation*> TransitionsOverTimeToUpdate;
+
+	UpdateCachePrimitivesInternal(Scene, Renderer, bAllowUnbuiltPreview, BlocksToUpdate, TransitionsOverTimeToUpdate);
+	FinalizeUpdateInternal_RenderThread(Scene, Renderer, BlocksToUpdate, TransitionsOverTimeToUpdate);
+}
+
+bool FIndirectLightingCache::IndirectLightingAllowed(FScene* Scene, FSceneRenderer& Renderer) const
+{
+	bool bAnyViewAllowsIndirectLightingCache = false;
 	if (IsIndirectLightingCacheAllowed(GetFeatureLevel()) && Scene->PrecomputedLightVolumes.Num() > 0)
 	{
-		bool bAnyViewAllowsIndirectLightingCache = false;
-
 		for (int32 ViewIndex = 0; ViewIndex < Renderer.Views.Num(); ViewIndex++)
 		{
 			bAnyViewAllowsIndirectLightingCache |= (bool)Renderer.Views[ViewIndex].Family->EngineShowFlags.IndirectLightingCache;
 		}
+	}
+	return bAnyViewAllowsIndirectLightingCache;
+}
 
-		if (bAnyViewAllowsIndirectLightingCache)
+void FIndirectLightingCache::UpdateCachePrimitivesInternal(FScene* Scene, FSceneRenderer& Renderer, bool bAllowUnbuiltPreview, TMap<FIntVector, FBlockUpdateInfo>& OutBlocksToUpdate, TArray<FIndirectLightingCacheAllocation*>& OutTransitionsOverTimeToUpdate)
+{
+	SCOPE_CYCLE_COUNTER(STAT_UpdateIndirectLightingCachePrims);
+	const TMap<FPrimitiveComponentId, FAttachmentGroupSceneInfo>& AttachmentGroups = Scene->AttachmentGroups;
+	if (IndirectLightingAllowed(Scene, Renderer))
+	{		
+		if (bUpdateAllCacheEntries)
 		{
-			SCOPE_CYCLE_COUNTER(STAT_UpdateIndirectLightingCache);
+			const uint32 PrimitiveCount = Scene->Primitives.Num();
 
-			TMap<FIntVector, FBlockUpdateInfo> BlocksToUpdate;
-			TArray<FIndirectLightingCacheAllocation*> TransitionsOverTimeToUpdate;
+			for (uint32 PrimitiveIndex = 0; PrimitiveIndex < PrimitiveCount; ++PrimitiveIndex)
+			{
+				FPrimitiveSceneInfo* PrimitiveSceneInfo = Scene->Primitives[PrimitiveIndex];
+				UpdateCachePrimitive(AttachmentGroups, PrimitiveSceneInfo, false, true, OutBlocksToUpdate, OutTransitionsOverTimeToUpdate);
+			}
+		}
 
-			if (bUpdateAllCacheEntries)
+		TArray<uint32> SetBitIndices[4];
+		{
+			QUICK_SCOPE_CYCLE_COUNTER(STAT_UpdateCachePreWalk);
+			
+			for (int32 ViewIndex = 0; ViewIndex < Renderer.Views.Num(); ViewIndex++)
 			{
-				const uint32 PrimitiveCount = Scene->Primitives.Num();
+				FViewInfo& View = Renderer.Views[ViewIndex];
+				SetBitIndices[ViewIndex].Reserve(View.PrimitiveVisibilityMap.Num());
 
-				for (uint32 PrimitiveIndex = 0; PrimitiveIndex < PrimitiveCount; ++PrimitiveIndex)
+				for (FSceneSetBitIterator BitIt(View.PrimitiveVisibilityMap); BitIt; ++BitIt)
 				{
-					FPrimitiveSceneInfo* PrimitiveSceneInfo = Scene->Primitives[PrimitiveIndex];
-
-					UpdateCachePrimitive(Scene, PrimitiveSceneInfo, false, true, BlocksToUpdate, TransitionsOverTimeToUpdate);
+					uint32 PrimitiveIndex = BitIt.GetIndex();
+					SetBitIndices[ViewIndex].Add(PrimitiveIndex);					
 				}
-			}
+			}			
+		}
 
+		{			
 			// Go over the views and operate on any relevant visible primitives
 			for (int32 ViewIndex = 0; ViewIndex < Renderer.Views.Num(); ViewIndex++)
 			{
 				FViewInfo& View = Renderer.Views[ViewIndex];
 
 				if (!bUpdateAllCacheEntries)
-				{
-					for (FSceneSetBitIterator BitIt(View.PrimitiveVisibilityMap); BitIt; ++BitIt)
+				{					
+					const TArray<uint32>& SetBits = SetBitIndices[ViewIndex];
+					for (int32 i = 0; i < SetBits.Num(); ++i)
 					{
-						uint32 PrimitiveIndex = BitIt.GetIndex();
+						uint32 PrimitiveIndex = SetBits[i];
 						FPrimitiveSceneInfo* PrimitiveSceneInfo = Scene->Primitives[PrimitiveIndex];
 						const FPrimitiveViewRelevance& PrimitiveRelevance = View.PrimitiveViewRelevanceMap[PrimitiveIndex];
-
-						UpdateCachePrimitive(Scene, PrimitiveSceneInfo, bAllowUnbuiltPreview, PrimitiveRelevance.bOpaqueRelevance, BlocksToUpdate, TransitionsOverTimeToUpdate);
+						UpdateCachePrimitive(AttachmentGroups, PrimitiveSceneInfo, bAllowUnbuiltPreview, PrimitiveRelevance.bOpaqueRelevance, OutBlocksToUpdate, OutTransitionsOverTimeToUpdate);
 					}
 				}
 			}
+		}
+		bUpdateAllCacheEntries = false;
+	}
+}
+
+void FIndirectLightingCache::FinalizeUpdateInternal_RenderThread(FScene* Scene, FSceneRenderer& Renderer, TMap<FIntVector, FBlockUpdateInfo>& BlocksToUpdate, const TArray<FIndirectLightingCacheAllocation*>& TransitionsOverTimeToUpdate)
+{
+	check(IsInRenderingThread());
 
+	if (IndirectLightingAllowed(Scene, Renderer))
+	{
+		{
+			SCOPE_CYCLE_COUNTER(STAT_UpdateIndirectLightingCacheBlocks);
 			UpdateBlocks(Scene, Renderer.Views.GetData(), BlocksToUpdate);
+		}
 
+		{
+			SCOPE_CYCLE_COUNTER(STAT_UpdateIndirectLightingCacheTransitions);
 			UpdateTransitionsOverTime(TransitionsOverTimeToUpdate, Renderer.ViewFamily.DeltaWorldTime);
 		}
-
-		bUpdateAllCacheEntries = false;
-	}
+	}	
 
 	if (GCacheDrawLightingSamples || Renderer.ViewFamily.EngineShowFlags.VolumeLightingSamples || GCacheDrawDirectionalShadowing)
 	{
@@ -439,11 +533,11 @@ void FIndirectLightingCache::UpdateCache(FScene* Scene, FSceneRenderer& Renderer
 }
 
 void FIndirectLightingCache::UpdateCacheAllocation(
-	const FBoxSphereBounds& Bounds, 
+	const FBoxSphereBounds& Bounds,
 	int32 BlockSize,
 	bool bPointSample,
 	bool bUnbuiltPreview,
-	FIndirectLightingCacheAllocation*& Allocation, 
+	FIndirectLightingCacheAllocation*& Allocation,
 	TMap<FIntVector, FBlockUpdateInfo>& BlocksToUpdate,
 	TArray<FIndirectLightingCacheAllocation*>& TransitionsOverTimeToUpdate)
 {
@@ -453,7 +547,7 @@ void FIndirectLightingCache::UpdateCacheAllocation(
 
 		// Calculate a potentially new min and size based on the current bounds
 		FVector NewMin;
-		FVector NewSize;
+		FVector NewSize;		
 		CalculateBlockPositionAndSize(Bounds, Block.TexelSize, NewMin, NewSize);
 
 		// If the primitive has moved enough to change its block min and size, we need to interpolate it again
@@ -492,17 +586,20 @@ void FIndirectLightingCache::UpdateCacheAllocation(
 }
 
 void FIndirectLightingCache::UpdateCachePrimitive(
-	FScene* Scene, 
+	const TMap<FPrimitiveComponentId, FAttachmentGroupSceneInfo>& AttachmentGroups,
 	FPrimitiveSceneInfo* PrimitiveSceneInfo,
 	bool bAllowUnbuiltPreview,
-	bool bOpaqueRelevance, 
+	bool bOpaqueRelevance,
 	TMap<FIntVector, FBlockUpdateInfo>& BlocksToUpdate,
 	TArray<FIndirectLightingCacheAllocation*>& TransitionsOverTimeToUpdate)
 {
-	FPrimitiveSceneProxy* PrimitiveSceneProxy = PrimitiveSceneInfo->Proxy;
+	
+	FPrimitiveSceneProxy* PrimitiveSceneProxy = PrimitiveSceneInfo->Proxy;	
 	FIndirectLightingCacheAllocation** PrimitiveAllocationPtr = PrimitiveAllocations.Find(PrimitiveSceneInfo->PrimitiveComponentId);
 	FIndirectLightingCacheAllocation* PrimitiveAllocation = PrimitiveAllocationPtr != NULL ? *PrimitiveAllocationPtr : NULL;
 
+	const bool bIsMovable = PrimitiveSceneProxy->IsMovable();
+
 	if (PrimitiveSceneProxy->WillEverBeLit()
 		&& ((bAllowUnbuiltPreview && PrimitiveSceneProxy->HasStaticLighting() && PrimitiveAllocation && PrimitiveAllocation->bIsDirty)
 		|| (PrimitiveSceneProxy->IsMovable() && PrimitiveSceneProxy->GetIndirectLightingCacheQuality() != ILCQ_Off)))
@@ -510,8 +607,8 @@ void FIndirectLightingCache::UpdateCachePrimitive(
 		const FIndirectLightingCacheAllocation* AttachmentParentAllocation = NULL;
 
 		if (PrimitiveSceneInfo->LightingAttachmentRoot.IsValid())
-		{
-			FAttachmentGroupSceneInfo& AttachmentGroup = Scene->AttachmentGroups.FindChecked(PrimitiveSceneInfo->LightingAttachmentRoot);
+		{			
+			const FAttachmentGroupSceneInfo& AttachmentGroup = AttachmentGroups.FindChecked(PrimitiveSceneInfo->LightingAttachmentRoot);
 
 			if (AttachmentGroup.ParentSceneInfo && AttachmentGroup.ParentSceneInfo->Proxy->LightAttachmentsAsGroup())
 			{
@@ -524,10 +621,10 @@ void FIndirectLightingCache::UpdateCachePrimitive(
 			// Reuse the attachment parent's lighting allocation if part of an attachment group
 			PrimitiveSceneInfo->IndirectLightingCacheAllocation = AttachmentParentAllocation;
 		}
-		else 
+		else
 		{
 			FIndirectLightingCacheAllocation* OriginalAllocation = PrimitiveAllocation;
-			const bool bUnbuiltPreview = bAllowUnbuiltPreview && !PrimitiveSceneProxy->IsMovable();
+			const bool bUnbuiltPreview = bAllowUnbuiltPreview && !bIsMovable;
 			const bool bPointSample = PrimitiveSceneProxy->GetIndirectLightingCacheQuality() == ILCQ_Point || bUnbuiltPreview || !bOpaqueRelevance;
 			const int32 BlockSize = bPointSample ? 1 : GLightingCacheMovableObjectAllocationSize;
 
diff --git a/Engine/Source/Runtime/Renderer/Private/ScenePrivate.h b/Engine/Source/Runtime/Renderer/Private/ScenePrivate.h
--- a/Engine/Source/Runtime/Renderer/Private/ScenePrivate.h
+++ b/Engine/Source/Runtime/Renderer/Private/ScenePrivate.h
@@ -1141,13 +1141,36 @@ public:
 	FIndirectLightingCacheAllocation* Allocation;
 };
 
+/** Information about the primitives that are attached together. */
+class FAttachmentGroupSceneInfo
+{
+public:
+
+	/** The parent primitive, which is the root of the attachment tree. */
+	FPrimitiveSceneInfo* ParentSceneInfo;
+
+	/** The primitives in the attachment group. */
+	TArray<FPrimitiveSceneInfo*> Primitives;
+
+	FAttachmentGroupSceneInfo() :
+		ParentSceneInfo(nullptr)
+	{}
+};
+
+struct FILCUpdatePrimTaskData
+{
+	FGraphEventRef TaskRef;
+	TMap<FIntVector, FBlockUpdateInfo> OutBlocksToUpdate;
+	TArray<FIndirectLightingCacheAllocation*> OutTransitionsOverTimeToUpdate;
+};
+
 /** 
  * Implements a volume texture atlas for caching indirect lighting on a per-object basis.
  * The indirect lighting is interpolated from Lightmass SH volume lighting samples.
  */
 class FIndirectLightingCache : public FRenderResource
 {
-public:
+public:	
 
 	/** true for the editor case where we want a better preview for object that have no valid lightmaps */
 	FIndirectLightingCache(ERHIFeatureLevel::Type InFeatureLevel);
@@ -1162,11 +1185,17 @@ public:
 	/** Releases the indirect lighting allocation for the given primitive. */
 	void ReleasePrimitive(FPrimitiveComponentId PrimitiveId);
 
-	FIndirectLightingCacheAllocation* FindPrimitiveAllocation(FPrimitiveComponentId PrimitiveId);
+	FIndirectLightingCacheAllocation* FindPrimitiveAllocation(FPrimitiveComponentId PrimitiveId);	
 
-	/** Updates indirect lighting in the cache based on visibility. */
+	/** Updates indirect lighting in the cache based on visibility syncronously. */
 	void UpdateCache(FScene* Scene, FSceneRenderer& Renderer, bool bAllowUnbuiltPreview);
 
+	/** Starts a task to update the cache primitives.  Results and task ref returned in the FILCUpdatePrimTaskData structure */
+	void StartUpdateCachePrimitivesTask(FScene* Scene, FSceneRenderer& Renderer, bool bAllowUnbuiltPreview, FILCUpdatePrimTaskData& OutTaskData);
+
+	/** Wait on a previously started task and complete any block updates and debug draw */
+	void FinalizeCacheUpdates(FScene* Scene, FSceneRenderer& Renderer, FILCUpdatePrimTaskData& TaskData);
+
 	/** Force all primitive allocations to be re-interpolated. */
 	void SetLightingCacheDirty();
 
@@ -1176,6 +1205,14 @@ public:
 	FSceneRenderTargetItem& GetTexture2() { return Texture2->GetRenderTargetItem(); }
 
 private:
+	/** Internal helper to determine if indirect lighting is enabled at all */
+	bool IndirectLightingAllowed(FScene* Scene, FSceneRenderer& Renderer) const;
+
+	/** Internal helper to perform the work of updating the cache primitives.  Can be done on any thread as a task */
+	void UpdateCachePrimitivesInternal(FScene* Scene, FSceneRenderer& Renderer, bool bAllowUnbuiltPreview, TMap<FIntVector, FBlockUpdateInfo>& OutBlocksToUpdate, TArray<FIndirectLightingCacheAllocation*>& OutTransitionsOverTimeToUpdate);
+
+	/** Internal helper to perform blockupdates and transition updates on the results of UpdateCachePrimitivesInternal.  Must be on render thread. */
+	void FinalizeUpdateInternal_RenderThread(FScene* Scene, FSceneRenderer& Renderer, TMap<FIntVector, FBlockUpdateInfo>& BlocksToUpdate, const TArray<FIndirectLightingCacheAllocation*>& TransitionsOverTimeToUpdate);
 
 	/** Internal helper which adds an entry to the update lists for this allocation, if needed (due to movement, etc). */
 	void UpdateCacheAllocation(
@@ -1185,19 +1222,19 @@ private:
 		bool bUnbuiltPreview,
 		FIndirectLightingCacheAllocation*& Allocation, 
 		TMap<FIntVector, FBlockUpdateInfo>& BlocksToUpdate,
-		TArray<FIndirectLightingCacheAllocation*>& TransitionsOverTimeToUpdate);
+		TArray<FIndirectLightingCacheAllocation*>& TransitionsOverTimeToUpdate);	
 
 	/** 
 	 * Creates a new allocation if needed, caches the result in PrimitiveSceneInfo->IndirectLightingCacheAllocation, 
 	 * And adds an entry to the update lists when an update is needed. 
 	 */
 	void UpdateCachePrimitive(
-		FScene* Scene, 
+		const TMap<FPrimitiveComponentId, FAttachmentGroupSceneInfo>& AttachmentGroups,
 		FPrimitiveSceneInfo* PrimitiveSceneInfo,
 		bool bAllowUnbuiltPreview, 
 		bool bOpaqueRelevance, 
 		TMap<FIntVector, FBlockUpdateInfo>& BlocksToUpdate, 
-		TArray<FIndirectLightingCacheAllocation*>& TransitionsOverTimeToUpdate);
+		TArray<FIndirectLightingCacheAllocation*>& TransitionsOverTimeToUpdate);	
 
 	/** Updates the contents of the volume texture blocks in BlocksToUpdate. */
 	void UpdateBlocks(FScene* Scene, FViewInfo* DebugDrawingView, TMap<FIntVector, FBlockUpdateInfo>& BlocksToUpdate);
@@ -1206,7 +1243,7 @@ private:
 	void UpdateTransitionsOverTime(const TArray<FIndirectLightingCacheAllocation*>& TransitionsOverTimeToUpdate, float DeltaWorldTime) const;
 
 	/** Creates an allocation to be used outside the indirect lighting cache and a block to be used internally. */
-	FIndirectLightingCacheAllocation* CreateAllocation(int32 BlockSize, const FBoxSphereBounds& Bounds, bool bPointSample, bool bUnbuiltPreview);
+	FIndirectLightingCacheAllocation* CreateAllocation(int32 BlockSize, const FBoxSphereBounds& Bounds, bool bPointSample, bool bUnbuiltPreview);	
 
 	/** Block accessors. */
 	FIndirectLightingCacheBlock& FindBlock(FIntVector TexelMin);
@@ -1281,6 +1318,8 @@ private:
 
 	/** Tracks primitive allocations by component, so that they persist across re-registers. */
 	TMap<FPrimitiveComponentId, FIndirectLightingCacheAllocation*> PrimitiveAllocations;
+
+	friend class FUpdateCachePrimitivesTask;
 };
 
 /**
@@ -1331,22 +1370,6 @@ namespace EOcclusionFlags
 	};
 };
 
-/** Information about the primitives that are attached together. */
-class FAttachmentGroupSceneInfo
-{
-public:
-
-	/** The parent primitive, which is the root of the attachment tree. */
-	FPrimitiveSceneInfo* ParentSceneInfo;
-
-	/** The primitives in the attachment group. */
-	TArray<FPrimitiveSceneInfo*> Primitives;
-
-	FAttachmentGroupSceneInfo() :
-		ParentSceneInfo(nullptr)
-	{}
-};
-
 class FLODSceneTree
 {
 public:
diff --git a/Engine/Source/Runtime/Renderer/Private/SceneRendering.h b/Engine/Source/Runtime/Renderer/Private/SceneRendering.h
--- a/Engine/Source/Runtime/Renderer/Private/SceneRendering.h
+++ b/Engine/Source/Runtime/Renderer/Private/SceneRendering.h
@@ -14,6 +14,7 @@
 
 // Forward declarations.
 class FPostprocessContext;
+struct FILCUpdatePrimTaskData;
 
 /** Information about a visible light which is specific to the view it's visible in. */
 class FVisibleLightViewInfo
@@ -826,7 +827,7 @@ protected:
 	void ComputeViewVisibility(FRHICommandListImmediate& RHICmdList);
 
 	/** Performs once per frame setup after to visibility determination. */
-	void PostVisibilityFrameSetup();
+	void PostVisibilityFrameSetup(FILCUpdatePrimTaskData& OutILCTaskData);
 
 	void GatherDynamicMeshElements(
 		TArray<FViewInfo>& InViews, 
diff --git a/Engine/Source/Runtime/Renderer/Private/SceneVisibility.cpp b/Engine/Source/Runtime/Renderer/Private/SceneVisibility.cpp
--- a/Engine/Source/Runtime/Renderer/Private/SceneVisibility.cpp
+++ b/Engine/Source/Runtime/Renderer/Private/SceneVisibility.cpp
@@ -108,6 +108,14 @@ static FAutoConsoleVariableRef CVarOcclusionCullParallelPrimFetch(
 	ECVF_RenderThreadSafe
 	);
 
+static int32 GILCUpdatePrimTaskEnabled = 1;
+static FAutoConsoleVariableRef CVarILCUpdatePrimitivesTask(
+	TEXT("r.Cache.UpdatePrimsTaskEnabled"),
+	GILCUpdatePrimTaskEnabled,
+	TEXT("Enable threading for ILC primitive update.  Will overlap with the rest the end of InitViews."),
+	ECVF_RenderThreadSafe
+	);
+
 /** Distance fade cvars */
 static int32 GDisableLODFade = false;
 static FAutoConsoleVariableRef CVarDisableLODFade( TEXT("r.DisableLODFade"), GDisableLODFade, TEXT("Disable fading for distance culling"), ECVF_RenderThreadSafe );
@@ -2428,7 +2436,7 @@ void FSceneRenderer::ComputeViewVisibility(FRHICommandListImmediate& RHICmdList)
 	INC_DWORD_STAT_BY(STAT_OccludedPrimitives,NumOccludedPrimitives);
 }
 
-void FSceneRenderer::PostVisibilityFrameSetup()
+void FSceneRenderer::PostVisibilityFrameSetup(FILCUpdatePrimTaskData& OutILCTaskData)
 {
 	for (int32 ViewIndex = 0; ViewIndex < Views.Num(); ViewIndex++)
 	{		
@@ -2462,7 +2470,14 @@ void FSceneRenderer::PostVisibilityFrameSetup()
 
 	if (ViewFamily.EngineShowFlags.HitProxies == 0)
 	{
-		Scene->IndirectLightingCache.UpdateCache(Scene, *this, true);
+		if (GILCUpdatePrimTaskEnabled)
+		{
+			Scene->IndirectLightingCache.StartUpdateCachePrimitivesTask(Scene, *this, true, OutILCTaskData);
+		}
+		else
+		{
+			Scene->IndirectLightingCache.UpdateCache(Scene, *this, true);
+		}		
 	}
 
 	// determine visibility of each light
@@ -2628,10 +2643,10 @@ void FDeferredShadingSceneRenderer::InitViews(FRHICommandListImmediate& RHICmdLi
 			View.FinalPostProcessSettings.AntiAliasingMethod = AAM_None;
 		}
 	}
-
+	FILCUpdatePrimTaskData ILCTaskData;
 	PreVisibilityFrameSetup(RHICmdList);
 	ComputeViewVisibility(RHICmdList);
-	PostVisibilityFrameSetup();
+	PostVisibilityFrameSetup(ILCTaskData);
 
 	FVector AverageViewPosition(0);
 
@@ -2648,7 +2663,7 @@ void FDeferredShadingSceneRenderer::InitViews(FRHICommandListImmediate& RHICmdLi
 	}
 	else
 	{
-	SortBasePassStaticData(AverageViewPosition);
+		SortBasePassStaticData(AverageViewPosition);
 	}
 
 	// this cannot be moved later because of static mesh updates for stuff that is only visible in shadows
@@ -2666,6 +2681,12 @@ void FDeferredShadingSceneRenderer::InitViews(FRHICommandListImmediate& RHICmdLi
 		InitDynamicShadows(RHICmdList);
 	}
 
+	// if we kicked off ILC update via task, wait and finalize.
+	if (ILCTaskData.TaskRef.IsValid())
+	{
+		Scene->IndirectLightingCache.FinalizeCacheUpdates(Scene, *this, ILCTaskData);
+	}
+
 	// initialize per-view uniform buffer.
 	for (int32 ViewIndex = 0; ViewIndex < Views.Num(); ViewIndex++)
 	{
