Enable ILC update task only for PS4

diff --git a/Engine/Source/Runtime/Renderer/Private/SceneVisibility.cpp b/Engine/Source/Runtime/Renderer/Private/SceneVisibility.cpp
--- a/Engine/Source/Runtime/Renderer/Private/SceneVisibility.cpp
+++ b/Engine/Source/Runtime/Renderer/Private/SceneVisibility.cpp
@@ -108,7 +108,7 @@ static FAutoConsoleVariableRef CVarOcclusionCullParallelPrimFetch(
 	ECVF_RenderThreadSafe
 	);
 
-static int32 GILCUpdatePrimTaskEnabled = 1;
+static int32 GILCUpdatePrimTaskEnabled = 0;
 static FAutoConsoleVariableRef CVarILCUpdatePrimitivesTask(
 	TEXT("r.Cache.UpdatePrimsTaskEnabled"),
 	GILCUpdatePrimTaskEnabled,
@@ -2628,7 +2628,7 @@ uint32 GetShadowQuality();
  * Check visibility, sort translucent items, etc.
  */
 void FDeferredShadingSceneRenderer::InitViews(FRHICommandListImmediate& RHICmdList)
-{
+{	
 	SCOPED_DRAW_EVENT(RHICmdList, InitViews);
 
 	SCOPE_CYCLE_COUNTER(STAT_InitViewsTime);
