missing changed from integration on shader optimization from Orion

diff --git a/Engine/Source/Runtime/Renderer/Private/PostProcess/PostProcessDownsample.cpp b/Engine/Source/Runtime/Renderer/Private/PostProcess/PostProcessDownsample.cpp
--- a/Engine/Source/Runtime/Renderer/Private/PostProcess/PostProcessDownsample.cpp
+++ b/Engine/Source/Runtime/Renderer/Private/PostProcess/PostProcessDownsample.cpp
@@ -35,6 +35,7 @@ class FPostProcessDownsamplePS : public FGlobalShader
 public:
 	FPostProcessPassParameters PostprocessParameter;
 	FDeferredPixelShaderParameters DeferredParameters;
+	FShaderParameter DownsampleParams;
 
 	/** Initialization constructor. */
 	FPostProcessDownsamplePS(const ShaderMetaType::CompiledShaderInitializerType& Initializer)
@@ -42,17 +43,18 @@ public:
 	{
 		PostprocessParameter.Bind(Initializer.ParameterMap);
 		DeferredParameters.Bind(Initializer.ParameterMap);
+		DownsampleParams.Bind(Initializer.ParameterMap, TEXT("DownsampleParams"));
 	}
 
 	// FShader interface.
 	virtual bool Serialize(FArchive& Ar) override
 	{
 		bool bShaderHasOutdatedParameters = FGlobalShader::Serialize(Ar);
-		Ar << PostprocessParameter << DeferredParameters;
+		Ar << PostprocessParameter << DeferredParameters << DownsampleParams;
 		return bShaderHasOutdatedParameters;
 	}
 
-	void SetParameters(const FRenderingCompositePassContext& Context)
+	void SetParameters(const FRenderingCompositePassContext& Context, const FPooledRenderTargetDesc* InputDesc)
 	{
 		const FPixelShaderRHIParamRef ShaderRHI = GetPixelShader();
 
@@ -63,6 +65,13 @@ public:
 		FSamplerStateRHIParamRef Filter = (Method == 2) ? 
 			TStaticSamplerState<SF_Point, AM_Clamp, AM_Clamp, AM_Clamp>::GetRHI():
 			TStaticSamplerState<SF_Bilinear, AM_Clamp, AM_Clamp, AM_Clamp>::GetRHI();
+		
+		{
+			float PixelScale = (Method == 2) ? 0.5f : 1.0f;
+ 
+			FVector4 DownsampleParamsValue(PixelScale / InputDesc->Extent.X, PixelScale / InputDesc->Extent.Y, 0, 0);
+			SetShaderValue(Context.RHICmdList, ShaderRHI, DownsampleParams, DownsampleParamsValue);
+		}
 
 		PostprocessParameter.SetPS(ShaderRHI, Context, Filter);
 	}
@@ -142,18 +151,17 @@ FRCPassPostProcessDownsample::FRCPassPostProcessDownsample(EPixelFormat InOverri
 
 
 template <uint32 Method>
-void FRCPassPostProcessDownsample::SetShader(const FRenderingCompositePassContext& Context)
+void FRCPassPostProcessDownsample::SetShader(const FRenderingCompositePassContext& Context, const FPooledRenderTargetDesc* InputDesc)
 {
 	auto ShaderMap = Context.GetShaderMap();
 	TShaderMapRef<FPostProcessDownsampleVS> VertexShader(ShaderMap);
 	TShaderMapRef<FPostProcessDownsamplePS<Method> > PixelShader(ShaderMap);
 
 	static FGlobalBoundShaderState BoundShaderState;
-	
 
 	SetGlobalBoundShaderState(Context.RHICmdList, Context.GetFeatureLevel(), BoundShaderState, GFilterVertexDeclaration.VertexDeclarationRHI, *VertexShader, *PixelShader);
 
-	PixelShader->SetParameters(Context);
+	PixelShader->SetParameters(Context, InputDesc);
 	VertexShader->SetParameters(Context);
 }
 
@@ -202,18 +210,18 @@ void FRCPassPostProcessDownsample::Process(FRenderingCompositePassContext& Conte
 	{
 		// also put depth in alpha
 		InflateSize = 2;
-		SetShader<2>(Context);
+		SetShader<2>(Context, InputDesc);
 	}
 	else
 	{
 		if (Quality == 0)
 		{
-			SetShader<0>(Context);
+			SetShader<0>(Context, InputDesc);
 			InflateSize = 1;
 		}
 		else
 		{
-			SetShader<1>(Context);
+			SetShader<1>(Context, InputDesc);
 			InflateSize = 2;
 		}
 	}
diff --git a/Engine/Source/Runtime/Renderer/Private/PostProcess/PostProcessDownsample.h b/Engine/Source/Runtime/Renderer/Private/PostProcess/PostProcessDownsample.h
--- a/Engine/Source/Runtime/Renderer/Private/PostProcess/PostProcessDownsample.h
+++ b/Engine/Source/Runtime/Renderer/Private/PostProcess/PostProcessDownsample.h
@@ -11,13 +11,13 @@
 
 // derives from TRenderingCompositePassBase<InputCount, OutputCount>
 // ePId_Input0: Color input
-// ePId_Input1: optional depth input
+// ePId_Input1: optional depth input (then quality is ignores and it uses the a 4 sample unfiltered samples method)
 class FRCPassPostProcessDownsample : public TRenderingCompositePassBase<2, 1>
 {
 public:
 	// constructor
 	// @param InDebugName we store the pointer so don't release this string
-	// @param Quality 0:1 sample, 1:4 samples with blurring
+	// @param Quality only used if ePId_Input1 is not set, 0:one filtered sample, 1:four filtered samples
 	FRCPassPostProcessDownsample(EPixelFormat InOverrideFormat = PF_Unknown,
 			uint32 InQuality = 1,
 			const TCHAR *InDebugName = TEXT("Downsample"));
@@ -29,7 +29,7 @@ public:
 	virtual FPooledRenderTargetDesc ComputeOutputDesc(EPassOutputId InPassOutputId) const override;
 
 private:
-	template <uint32 Method> static void SetShader(const FRenderingCompositePassContext& Context);
+	template <uint32 Method> static void SetShader(const FRenderingCompositePassContext& Context, const FPooledRenderTargetDesc* InputDesc);
 
 	EPixelFormat OverrideFormat;
 	// explained in constructor
diff --git a/Engine/Source/Runtime/Renderer/Private/PostProcess/PostProcessing.cpp b/Engine/Source/Runtime/Renderer/Private/PostProcess/PostProcessing.cpp
--- a/Engine/Source/Runtime/Renderer/Private/PostProcess/PostProcessing.cpp
+++ b/Engine/Source/Runtime/Renderer/Private/PostProcess/PostProcessing.cpp
@@ -116,6 +116,14 @@ static TAutoConsoleVariable<int32> CVarUpscaleQuality(
 	TEXT(" 3: Directional blur with unsharp mask upsample. (default)"),
 	ECVF_Scalability | ECVF_RenderThreadSafe);
 
+static TAutoConsoleVariable<int32> CDownsampleQuality(
+	TEXT("r.Downsample.Quality"),
+	3,
+	TEXT("Defines the quality in which the Downsample passes. we might add more quality levels later.\n")
+	TEXT(" 0: low quality\n")
+	TEXT(">0: high quality (default: 3)\n"),
+	ECVF_Scalability | ECVF_RenderThreadSafe);
+
 static TAutoConsoleVariable<float> CVarMotionBlurSoftEdgeSize(
 	TEXT("r.MotionBlurSoftEdgeSize"),
 	1.0f,
@@ -1318,8 +1326,10 @@ void FPostProcessing::Process(FRHICommandListImmediate& RHICmdList, FViewInfo& V
 			// down sample Scene color from full to half res
 			FRenderingCompositeOutputRef SceneColorHalfRes;
 			{
+				int32 DownsampleQuality = FMath::Clamp(CDownsampleQuality.GetValueOnRenderThread(), 0, 1);
+
 				// doesn't have to be as high quality as the Scene color
-				FRenderingCompositePass* HalfResPass = Context.Graph.RegisterPass(new(FMemStack::Get()) FRCPassPostProcessDownsample(PF_FloatRGB, 1, TEXT("SceneColorHalfRes")));
+				FRenderingCompositePass* HalfResPass = Context.Graph.RegisterPass(new(FMemStack::Get()) FRCPassPostProcessDownsample(PF_FloatRGB, DownsampleQuality, TEXT("SceneColorHalfRes")));
 				HalfResPass->SetInput(ePId_Input0, FRenderingCompositeOutputRef(Context.FinalOutput));
 
 				SceneColorHalfRes = FRenderingCompositeOutputRef(HalfResPass);
