fixed SubsurfaceScattering rendering which got dither pattern from recent change

diff --git a/Engine/Source/Runtime/Renderer/Private/CompositionLighting/CompositionLighting.cpp b/Engine/Source/Runtime/Renderer/Private/CompositionLighting/CompositionLighting.cpp
--- a/Engine/Source/Runtime/Renderer/Private/CompositionLighting/CompositionLighting.cpp
+++ b/Engine/Source/Runtime/Renderer/Private/CompositionLighting/CompositionLighting.cpp
@@ -425,12 +425,14 @@ void FCompositionLighting::ProcessAfterLighting(FRHICommandListImmediate& RHICmd
 					RecombinePass->SetInput(ePId_Input0, Context.FinalOutput);
 					RecombinePass->SetInput(ePId_Input1, PassY);
 					RecombinePass->SetInput(ePId_Input2, PassSetup);
+					Context.FinalOutput = FRenderingCompositeOutputRef(RecombinePass);
 				}
 				else
 				{
 					// needed for Scalability
 					FRenderingCompositePass* RecombinePass = Context.Graph.RegisterPass(new(FMemStack::Get()) FRCPassPostProcessSubsurfaceRecombine(bHalfRes, bSingleViewportMode));
 					RecombinePass->SetInput(ePId_Input0, Context.FinalOutput);
+					Context.FinalOutput = FRenderingCompositeOutputRef(RecombinePass);
 				}
 			}
 		}
