Merging

//depot/UE4-Orion/Engine/Source/Runtime/Renderer/Private/ConvertToUniformMesh.cpp

to //depot/UE4/Engine/Source/Runtime/Renderer/Private/ConvertToUniformMesh.cpp

fixed comment

diff --git a/Engine/Source/Runtime/Renderer/Private/ConvertToUniformMesh.cpp b/Engine/Source/Runtime/Renderer/Private/ConvertToUniformMesh.cpp
--- a/Engine/Source/Runtime/Renderer/Private/ConvertToUniformMesh.cpp
+++ b/Engine/Source/Runtime/Renderer/Private/ConvertToUniformMesh.cpp
@@ -1,7 +1,7 @@
 // Copyright 1998-2014 Epic Games, Inc. All Rights Reserved.
 
 /*=============================================================================
-	SurfelTree.cpp
+	ConvertToUniformMesh.cpp
 =============================================================================*/
 
 #include "RendererPrivate.h"
