First pass at new RHITransitionResources API.
Rendertarget transitions are completed and validated on DX11 and PS4 RHI's.

diff --git a/Engine/Source/Runtime/Engine/Private/Particles/FXSystem.cpp b/Engine/Source/Runtime/Engine/Private/Particles/FXSystem.cpp
--- a/Engine/Source/Runtime/Engine/Private/Particles/FXSystem.cpp
+++ b/Engine/Source/Runtime/Engine/Private/Particles/FXSystem.cpp
@@ -332,8 +332,13 @@ void FFXSystem::PreRender(FRHICommandListImmediate& RHICmdList, const FGlobalDis
 {
 	if (RHISupportsGPUParticles(FeatureLevel))
 	{
+		PrepareGPUSimulation(RHICmdList);
+
 		SimulateGPUParticles(RHICmdList, EParticleSimulatePhase::Main, NULL, NULL, FTexture2DRHIParamRef(), FTexture2DRHIParamRef());
 		SimulateGPUParticles(RHICmdList, EParticleSimulatePhase::CollisionDistanceField, NULL, GlobalDistanceFieldParameterData, FTexture2DRHIParamRef(), FTexture2DRHIParamRef());
+
+		//particles rendered during basepass may need to read pos/velocity buffers.  must finalize unless we know for sure that nothingin base pass will read.
+		FinalizeGPUSimulation(RHICmdList);
 	}
 }
 
@@ -341,7 +346,11 @@ void FFXSystem::PostRenderOpaque(FRHICommandListImmediate& RHICmdList, const cla
 {
 	if (RHISupportsGPUParticles(FeatureLevel))
 	{
+		PrepareGPUSimulation(RHICmdList);
+
 		SimulateGPUParticles(RHICmdList, EParticleSimulatePhase::CollisionDepthBuffer, CollisionView, NULL, SceneDepthTexture, GBufferATexture);
 		SortGPUParticles(RHICmdList);
+
+		FinalizeGPUSimulation(RHICmdList);
 	}
 }
diff --git a/Engine/Source/Runtime/Engine/Private/Particles/FXSystemPrivate.h b/Engine/Source/Runtime/Engine/Private/Particles/FXSystemPrivate.h
--- a/Engine/Source/Runtime/Engine/Private/Particles/FXSystemPrivate.h
+++ b/Engine/Source/Runtime/Engine/Private/Particles/FXSystemPrivate.h
@@ -112,6 +112,9 @@ public:
 	 */
 	int32 AddSortedGPUSimulation(FParticleSimulationGPU* Simulation, const FVector& ViewOrigin);
 
+	void PrepareGPUSimulation(FRHICommandListImmediate& RHICmdList);
+	void FinalizeGPUSimulation(FRHICommandListImmediate& RHICmdList);
+
 private:
 
 	/*--------------------------------------------------------------------------
diff --git a/Engine/Source/Runtime/Engine/Private/Particles/ParticleGpuSimulation.cpp b/Engine/Source/Runtime/Engine/Private/Particles/ParticleGpuSimulation.cpp
--- a/Engine/Source/Runtime/Engine/Private/Particles/ParticleGpuSimulation.cpp
+++ b/Engine/Source/Runtime/Engine/Private/Particles/ParticleGpuSimulation.cpp
@@ -218,6 +218,11 @@ public:
 			VelocityTextureRHI
 			);
 
+		static FName PositionTextureName(TEXT("ParticleStatePosition"));
+		static FName VelocityTextureName(TEXT("ParticleStateVelocity"));
+		PositionTextureTargetRHI->SetName(PositionTextureName);
+		VelocityTextureTargetRHI->SetName(VelocityTextureName);
+
 		bTexturesCleared = false;
 	}
 
@@ -266,6 +271,9 @@ public:
 			TextureTargetRHI,
 			TextureRHI
 			);
+
+		static FName AttributesTextureName(TEXT("ParticleAttributes"));	
+		TextureTargetRHI->SetName(AttributesTextureName);		
 	}
 
 	/**
@@ -4184,6 +4192,38 @@ bool FFXSystem::UsesGlobalDistanceFieldInternal() const
 	return false;
 }
 
+void FFXSystem::PrepareGPUSimulation(FRHICommandListImmediate& RHICmdList)
+{
+	// Grab resources.
+	FParticleStateTextures& CurrentStateTextures = ParticleSimulationResources->GetCurrentStateTextures();
+	FParticleStateTextures& PrevStateTextures = ParticleSimulationResources->GetPreviousStateTextures();
+
+	// Setup render states.
+	FTextureRHIParamRef RenderTargets[2] =
+	{
+		CurrentStateTextures.PositionTextureTargetRHI,
+		CurrentStateTextures.VelocityTextureTargetRHI,
+	};
+
+	RHICmdList.TransitionResources(EResourceTransitionAccess::EWritable, RenderTargets, 2);
+}
+
+void FFXSystem::FinalizeGPUSimulation(FRHICommandListImmediate& RHICmdList)
+{
+	// Grab resources.
+	FParticleStateTextures& CurrentStateTextures = ParticleSimulationResources->GetCurrentStateTextures();
+	FParticleStateTextures& PrevStateTextures = ParticleSimulationResources->GetPreviousStateTextures();
+
+	// Setup render states.
+	FTextureRHIParamRef RenderTargets[2] =
+	{
+		CurrentStateTextures.PositionTextureTargetRHI,
+		CurrentStateTextures.VelocityTextureTargetRHI,
+	};
+	
+	RHICmdList.TransitionResources(EResourceTransitionAccess::EReadable, RenderTargets, 2);	
+}
+
 void FFXSystem::SimulateGPUParticles(
 	FRHICommandListImmediate& RHICmdList,
 	EParticleSimulatePhase::Type Phase,
@@ -4410,7 +4450,7 @@ void FFXSystem::SimulateGPUParticles(
 			ParticleSimulationResources->RenderAttributesTexture.TextureTargetRHI,
 			ParticleSimulationResources->SimulationAttributesTexture.TextureTargetRHI
 		};
-		SetRenderTargets(RHICmdList, 4, InjectRenderTargets, FTextureRHIParamRef(), 0, NULL);
+		SetRenderTargets(RHICmdList, 4, InjectRenderTargets, FTextureRHIParamRef(), 0, NULL, true);
 		RHICmdList.SetViewport(0, 0, 0.0f, GParticleSimulationTextureSizeX, GParticleSimulationTextureSizeY, 1.0f);
 		RHICmdList.SetDepthStencilState(TStaticDepthStencilState<false, CF_Always>::GetRHI());
 		RHICmdList.SetRasterizerState(TStaticRasterizerState<FM_Solid, CM_None>::GetRHI());
@@ -4438,10 +4478,6 @@ void FFXSystem::SimulateGPUParticles(
 	TilesToClear.Reset();
 	NewParticles.Reset();
 
-	// Resolve all textures.
-	RHICmdList.CopyToResolveTarget(CurrentStateTextures.PositionTextureTargetRHI, CurrentStateTextures.PositionTextureRHI, /*bKeepOriginalSurface=*/ false, FResolveParams());
-	RHICmdList.CopyToResolveTarget(CurrentStateTextures.VelocityTextureTargetRHI, CurrentStateTextures.VelocityTextureRHI, /*bKeepOriginalSurface=*/ false, FResolveParams());
-
 	// Clear render targets so we can safely read from them.
 	SetRenderTarget(RHICmdList, FTextureRHIParamRef(), FTextureRHIParamRef());
 
diff --git a/Engine/Source/Runtime/Engine/Private/Rendering/SubsurfaceProfile.cpp b/Engine/Source/Runtime/Engine/Private/Rendering/SubsurfaceProfile.cpp
--- a/Engine/Source/Runtime/Engine/Private/Rendering/SubsurfaceProfile.cpp
+++ b/Engine/Source/Runtime/Engine/Private/Rendering/SubsurfaceProfile.cpp
@@ -142,7 +142,7 @@ void FSubsurfaceProfileTexture::CreateTexture(FRHICommandListImmediate& RHICmdLi
 		Desc.Format = PF_A16B16G16R16;
 	}
 
-	GetRendererModule().RenderTargetPoolFindFreeElement(Desc, GSSProfiles, TEXT("SSProfiles"));
+	GetRendererModule().RenderTargetPoolFindFreeElement(RHICmdList, Desc, GSSProfiles, TEXT("SSProfiles"));
 
 	// Write the contents of the texture.
 	uint32 DestStride;
diff --git a/Engine/Source/Runtime/Engine/Private/SceneUtils.cpp b/Engine/Source/Runtime/Engine/Private/SceneUtils.cpp
--- a/Engine/Source/Runtime/Engine/Private/SceneUtils.cpp
+++ b/Engine/Source/Runtime/Engine/Private/SceneUtils.cpp
@@ -25,4 +25,24 @@ void FDrawEvent::Stop()
 	RHICmdList->PopEvent();
 }
 
+void FDrawEventRHIExecute::Start(IRHICommandContext& InRHICommandContext, const TCHAR* Fmt, ...)
+{
+	check(IsInParallelRenderingThread() || IsInRHIThread() || (!GRHIThread && IsInRenderingThread()));
+	{
+		
+		va_list ptr;
+		va_start(ptr, Fmt);
+		TCHAR TempStr[256];
+		// Build the string in the temp buffer
+		FCString::GetVarArgs(TempStr, ARRAY_COUNT(TempStr), ARRAY_COUNT(TempStr) - 1, Fmt, ptr);
+		RHICommandContext = &InRHICommandContext;
+		RHICommandContext->RHIPushEvent(TempStr);		
+	}
+}
+
+void FDrawEventRHIExecute::Stop()
+{
+	RHICommandContext->RHIPopEvent();
+}
+
 #endif // WANTS_DRAW_MESH_EVENTS && PLATFORM_SUPPORTS_DRAW_MESH_EVENTS
\ No newline at end of file
diff --git a/Engine/Source/Runtime/Engine/Private/Slate/SceneViewport.cpp b/Engine/Source/Runtime/Engine/Private/Slate/SceneViewport.cpp
--- a/Engine/Source/Runtime/Engine/Private/Slate/SceneViewport.cpp
+++ b/Engine/Source/Runtime/Engine/Private/Slate/SceneViewport.cpp
@@ -1249,7 +1249,8 @@ void FSceneViewport::BeginRenderFrame(FRHICommandListImmediate& RHICmdList)
 	check( IsInRenderingThread() );
 	if (bUseSeparateRenderTarget)
 	{		
-		SetRenderTarget(RHICmdList,  RenderTargetTextureRenderThreadRHI,  FTexture2DRHIRef() );
+		RHICmdList.TransitionResource(EResourceTransitionAccess::EWritable, RenderTargetTextureRenderThreadRHI);
+		SetRenderTarget(RHICmdList,  RenderTargetTextureRenderThreadRHI,  FTexture2DRHIRef(), true);
 	}
 	else if( IsValidRef( ViewportRHI ) ) 
 	{
@@ -1259,7 +1260,7 @@ void FSceneViewport::BeginRenderFrame(FRHICommandListImmediate& RHICmdList)
 		if (GRHIRequiresEarlyBackBufferRenderTarget)
 		{
 			// unused set render targets are bad on Metal
-			SetRenderTarget(RHICmdList, RenderTargetTextureRenderThreadRHI, FTexture2DRHIRef());
+			SetRenderTarget(RHICmdList, RenderTargetTextureRenderThreadRHI, FTexture2DRHIRef(), true);
 		}
 	}
 }
diff --git a/Engine/Source/Runtime/Engine/Private/UnrealClient.cpp b/Engine/Source/Runtime/Engine/Private/UnrealClient.cpp
--- a/Engine/Source/Runtime/Engine/Private/UnrealClient.cpp
+++ b/Engine/Source/Runtime/Engine/Private/UnrealClient.cpp
@@ -1113,7 +1113,7 @@ const TArray<FColor>& FViewport::GetRawHitProxyData(FIntRect InRect)
 			{
 			// Set the hit proxy map's render target.
 			// Clear the hit proxy map to white, which is overloaded to mean no hit proxy.
-			SetRenderTarget(RHICmdList, Viewport->HitProxyMap.GetRenderTargetTexture(), FTextureRHIRef(), ESimpleRenderTargetMode::EClearColorExistingDepth);
+			SetRenderTarget(RHICmdList, Viewport->HitProxyMap.GetRenderTargetTexture(), FTextureRHIRef(), ESimpleRenderTargetMode::EClearColorExistingDepth, FExclusiveDepthStencil::DepthWrite_StencilWrite, true);
 		});
 
 		// Let the viewport client draw its hit proxies.
diff --git a/Engine/Source/Runtime/Engine/Private/UserInterface/Canvas.cpp b/Engine/Source/Runtime/Engine/Private/UserInterface/Canvas.cpp
--- a/Engine/Source/Runtime/Engine/Private/UserInterface/Canvas.cpp
+++ b/Engine/Source/Runtime/Engine/Private/UserInterface/Canvas.cpp
@@ -760,7 +760,7 @@ void FCanvas::Flush_GameThread(bool bForce)
 		SCOPED_DRAW_EVENT(RHICmdList, CanvasFlush);
 
 		// Set the RHI render target.
-		::SetRenderTarget(RHICmdList, Parameters.CanvasRenderTarget->GetRenderTargetTexture(), FTextureRHIRef());
+		::SetRenderTarget(RHICmdList, Parameters.CanvasRenderTarget->GetRenderTargetTexture(), FTextureRHIRef(), true);
 		// disable depth test & writes
 		RHICmdList.SetDepthStencilState(TStaticDepthStencilState<false,CF_Always>::GetRHI());
 
@@ -897,7 +897,7 @@ void FCanvas::Clear(const FLinearColor& LinearColor)
 		SCOPED_DRAW_EVENT(RHICmdList, CanvasClear);
 		if( CanvasRenderTarget )
 		{
-			::SetRenderTarget(RHICmdList, CanvasRenderTarget->GetRenderTargetTexture(), FTextureRHIRef());
+			::SetRenderTarget(RHICmdList, CanvasRenderTarget->GetRenderTargetTexture(), FTextureRHIRef(), true);
 			RHICmdList.SetViewport(0, 0, 0.0f, CanvasRenderTarget->GetSizeXY().X, CanvasRenderTarget->GetSizeXY().Y, 1.0f);
 		}
 		RHICmdList.Clear(true, ClearColor, false, 0.0f, false, 0, FIntRect());
diff --git a/Engine/Source/Runtime/Engine/Public/SceneUtils.h b/Engine/Source/Runtime/Engine/Public/SceneUtils.h
--- a/Engine/Source/Runtime/Engine/Public/SceneUtils.h
+++ b/Engine/Source/Runtime/Engine/Public/SceneUtils.h
@@ -51,12 +51,46 @@
 		void Stop();
 	};
 
-	// Macros to allow for scoping of draw events
+	struct ENGINE_API FDrawEventRHIExecute
+	{
+		/** Context to execute on*/
+		class IRHICommandContext* RHICommandContext;
+
+		/** Default constructor, initializing all member variables. */
+		FORCEINLINE FDrawEventRHIExecute()
+			: RHICommandContext(nullptr)
+		{}
+
+		/**
+		* Terminate the event based upon scope
+		*/
+		FORCEINLINE ~FDrawEventRHIExecute()
+		{
+			if (RHICommandContext)
+			{
+				Stop();
+			}
+		}
+
+		/**
+		* Function for logging a PIX event with var args
+		*/
+		void CDECL Start(IRHICommandContext& InRHICommandContext, const TCHAR* Fmt, ...);
+		void Stop();
+	};
+
+	// Macros to allow for scoping of draw events outside of RHI function implementations
 	#define SCOPED_DRAW_EVENT(RHICmdList, Name) FDrawEvent PREPROCESSOR_JOIN(Event_##Name,__LINE__); if(GEmitDrawEvents) PREPROCESSOR_JOIN(Event_##Name,__LINE__).Start(RHICmdList, TEXT(#Name));
 	#define SCOPED_DRAW_EVENTF(RHICmdList, Name, Format, ...) FDrawEvent PREPROCESSOR_JOIN(Event_##Name,__LINE__); if(GEmitDrawEvents) PREPROCESSOR_JOIN(Event_##Name,__LINE__).Start(RHICmdList, Format, ##__VA_ARGS__);
 	#define SCOPED_CONDITIONAL_DRAW_EVENT(RHICmdList, Name, Condition) FDrawEvent PREPROCESSOR_JOIN(Event_##Name,__LINE__); if(GEmitDrawEvents && (Condition)) PREPROCESSOR_JOIN(Event_##Name,__LINE__).Start(RHICmdList, TEXT(#Name));
 	#define SCOPED_CONDITIONAL_DRAW_EVENTF(RHICmdList, Name, Condition, Format, ...) FDrawEvent PREPROCESSOR_JOIN(Event_##Name,__LINE__); if(GEmitDrawEvents && (Condition)) PREPROCESSOR_JOIN(Event_##Name,__LINE__).Start(RHICmdList, Format, ##__VA_ARGS__);
 
+	// Macros to allow for scoping of draw events within RHI function implementations
+	#define SCOPED_RHI_DRAW_EVENT(RHICmdContext, Name) FDrawEventRHIExecute PREPROCESSOR_JOIN(Event_##Name,__LINE__); if(GEmitDrawEvents) PREPROCESSOR_JOIN(Event_##Name,__LINE__).Start(RHICmdContext, TEXT(#Name));
+	#define SCOPED_RHI_DRAW_EVENTF(RHICmdContext, Name, Format, ...) FDrawEventRHIExecute PREPROCESSOR_JOIN(Event_##Name,__LINE__); if(GEmitDrawEvents) PREPROCESSOR_JOIN(Event_##Name,__LINE__).Start(RHICmdContext, Format, ##__VA_ARGS__);
+	#define SCOPED_RHI_CONDITIONAL_DRAW_EVENT(RHICmdContext, Name, Condition) FDrawEventRHIExecute PREPROCESSOR_JOIN(Event_##Name,__LINE__); if(GEmitDrawEvents && (Condition)) PREPROCESSOR_JOIN(Event_##Name,__LINE__).Start(RHICmdContext, TEXT(#Name));
+	#define SCOPED_RHI_CONDITIONAL_DRAW_EVENTF(RHICmdContext, Name, Condition, Format, ...) FDrawEventRHIExecute PREPROCESSOR_JOIN(Event_##Name,__LINE__); if(GEmitDrawEvents && (Condition)) PREPROCESSOR_JOIN(Event_##Name,__LINE__).Start(RHICmdContext, Format, ##__VA_ARGS__);
+
 #else
 
 	#define SCOPED_DRAW_EVENT(...)
@@ -64,4 +98,9 @@
 	#define SCOPED_CONDITIONAL_DRAW_EVENT(...)
 	#define SCOPED_CONDITIONAL_DRAW_EVENTF(...)
 
+	#define SCOPED_RHI_DRAW_EVENT(...)
+	#define SCOPED_RHI_DRAW_EVENTF(...)
+	#define SCOPED_RHI_CONDITIONAL_DRAW_EVENT(...)
+	#define SCOPED_RHI_CONDITIONAL_DRAW_EVENTF(...)
+
 #endif
diff --git a/Engine/Source/Runtime/GameLiveStreaming/GameLiveStreaming.cpp b/Engine/Source/Runtime/GameLiveStreaming/GameLiveStreaming.cpp
--- a/Engine/Source/Runtime/GameLiveStreaming/GameLiveStreaming.cpp
+++ b/Engine/Source/Runtime/GameLiveStreaming/GameLiveStreaming.cpp
@@ -318,7 +318,7 @@ void FGameLiveStreaming::StartCopyingNextGameVideoFrame( const FViewportRHIRef&
 		const auto FeatureLevel = GMaxRHIFeatureLevel;
 
 		TRefCountPtr<IPooledRenderTarget> ResampleTexturePooledRenderTarget;
-		Context.RendererModule->RenderTargetPoolFindFreeElement(OutputDesc, ResampleTexturePooledRenderTarget, TEXT("ResampleTexture"));
+		Context.RendererModule->RenderTargetPoolFindFreeElement(RHICmdList, OutputDesc, ResampleTexturePooledRenderTarget, TEXT("ResampleTexture"));
 		check( ResampleTexturePooledRenderTarget );
 
 		const FSceneRenderTargetItem& DestRenderTarget = ResampleTexturePooledRenderTarget->GetRenderTargetItem();
diff --git a/Engine/Source/Runtime/RHI/Private/GPUProfiler.cpp b/Engine/Source/Runtime/RHI/Private/GPUProfiler.cpp
--- a/Engine/Source/Runtime/RHI/Private/GPUProfiler.cpp
+++ b/Engine/Source/Runtime/RHI/Private/GPUProfiler.cpp
@@ -31,6 +31,18 @@ static TAutoConsoleVariable<FString> GProfileGPURootCVar(
 	TEXT("Allows to filter the tree when using ProfileGPU, the pattern match is case sensitive."),
 	ECVF_Default);
 
+static TAutoConsoleVariable<int32> GProfileGPUShowEvents(
+	TEXT("r.ProfileGPU.ShowLeafEvents"),
+	0,
+	TEXT("Allows profileGPU to display event-only leaf nodes with no draws associated."),
+	ECVF_Default);
+
+static TAutoConsoleVariable<int32> GProfileGPUTransitions(
+	TEXT("r.ProfileGPU.ShowTransitions"),
+	0,
+	TEXT("Allows profileGPU to display resource transition events."),
+	ECVF_Default);
+
 struct FNodeStatsCompare
 {
 	/** Sorts nodes by descending durations. */
@@ -77,10 +89,10 @@ static void GatherStatsEventNode(FGPUProfilerEventNode* Node, int32 Depth, TMap<
 }
 
 /** Recursively dumps stats for each node with a depth first traversal. */
-static void DumpStatsEventNode(FGPUProfilerEventNode* Node, float RootResult, int32 Depth, const FWildcardString& WildcardFilter, bool bParentMatchedFilter, int32& NumNodes, int32& NumDraws)
+static void DumpStatsEventNode(FGPUProfilerEventNode* Node, float RootResult, int32 Depth, const FWildcardString& WildcardFilter, bool bParentMatchedFilter, int32& NumNodes, int32& NumDraws, bool bDumpEventLeafNodes)
 {
 	NumNodes++;
-	if (Node->NumDraws > 0 || Node->Children.Num() > 0)
+	if (Node->NumDraws > 0 || Node->Children.Num() > 0 || bDumpEventLeafNodes)
 	{
 		NumDraws += Node->NumDraws;
 		// Percent that this node was of the total frame time
@@ -121,8 +133,8 @@ static void DumpStatsEventNode(FGPUProfilerEventNode* Node, float RootResult, in
 			FGPUProfilerEventNode* ChildNode = Node->Children[ChildIndex];
 
 			int32 NumChildDraws = 0;
-			// Traverse children
-			DumpStatsEventNode(Node->Children[ChildIndex], RootResult, Depth + 1, WildcardFilter, bMatchesFilter, NumNodes, NumChildDraws);
+			// Traverse children			
+			DumpStatsEventNode(Node->Children[ChildIndex], RootResult, Depth + 1, WildcardFilter, bMatchesFilter, NumNodes, NumChildDraws, bDumpEventLeafNodes);
 			NumDraws += NumChildDraws;
 
 			TotalChildTime += ChildNode->TimingResult;
@@ -225,9 +237,10 @@ void FGPUProfilerEventNodeFrame::DumpEventTree()
 
 		int32 NumNodes = 0;
 		int32 NumDraws = 0;
+		bool bDumpEventLeafNodes = GProfileGPUShowEvents.GetValueOnRenderThread() != 0;
 		for (int32 BaseNodeIndex = 0; BaseNodeIndex < EventTree.Num(); BaseNodeIndex++)
 		{
-			DumpStatsEventNode(EventTree[BaseNodeIndex], RootResult, 0, RootWildcard, false, NumNodes, NumDraws);
+			DumpStatsEventNode(EventTree[BaseNodeIndex], RootResult, 0, RootWildcard, false, NumNodes, NumDraws, bDumpEventLeafNodes);
 		}
 
 		//@todo - calculate overhead instead of hardcoding
diff --git a/Engine/Source/Runtime/RHI/Private/RHI.cpp b/Engine/Source/Runtime/RHI/Private/RHI.cpp
--- a/Engine/Source/Runtime/RHI/Private/RHI.cpp
+++ b/Engine/Source/Runtime/RHI/Private/RHI.cpp
@@ -30,6 +30,15 @@ DEFINE_STAT(STAT_VertexBufferMemory);
 DEFINE_STAT(STAT_StructuredBufferMemory);
 DEFINE_STAT(STAT_PixelBufferMemory);
 
+const FString FResourceTransitionUtility::ResourceTransitionAccessStrings[(int32)EResourceTransitionAccess::EMaxAccess + 1] =
+{
+	FString(TEXT("EReadable")),
+	FString(TEXT("EWritable")),	
+	FString(TEXT("ERWBarrier")),
+	FString(TEXT("ERWNoBarrier")),
+	FString(TEXT("EMaxAccess")),
+};
+
 #if STATS
 #include "StatsData.h"
 static void DumpRHIMemory(FOutputDevice& OutputDevice)
diff --git a/Engine/Source/Runtime/RHI/Public/DynamicRHI.h b/Engine/Source/Runtime/RHI/Public/DynamicRHI.h
--- a/Engine/Source/Runtime/RHI/Public/DynamicRHI.h
+++ b/Engine/Source/Runtime/RHI/Public/DynamicRHI.h
@@ -704,6 +704,42 @@ public:
 	*/
 	virtual void RHICopyToResolveTarget(FTextureRHIParamRef SourceTexture, FTextureRHIParamRef DestTexture, bool bKeepOriginalSurface, const FResolveParams& ResolveParams) = 0;
 
+	/**
+	 * Explicitly transition a texture resource from readable -> writable by the GPU or vice versa.
+	 * We know rendertargets are only used as rendered targets on the Gfx pipeline, so these transitions are assumed to be implemented such
+	 * Gfx->Gfx and Gfx->Compute pipeline transitions are both handled by this call by the RHI implementation.  Hence, no pipeline parameter on this call.
+	 *
+	 * @param TransitionType - direction of the transition	 
+	 * @param InTextures - array of texture objects to transition	 
+	 * @param NumTextures - number of textures to transition
+	 */	
+	virtual void RHITransitionResources(EResourceTransitionAccess TransitionType, FTextureRHIParamRef* InTextures, int32 NumTextures)
+	{
+		if (TransitionType == EResourceTransitionAccess::EReadable)
+		{
+			const FResolveParams ResolveParams;
+			for (int32 i = 0; i < NumTextures; ++i)
+			{
+				RHICopyToResolveTarget(InTextures[i], InTextures[i], true, ResolveParams);
+			}
+		}
+	}
+
+	/**
+	* Explicitly transition a UAV from readable -> writable by the GPU or vice versa.
+	* Also explicitly states which pipeline the UAV can be used on next.  For example, if a Compute job just wrote this UAV for a Pixel shader to read
+	* you would do EResourceTransitionAccess::Readable and EResourceTransitionPipeline::EComputeToGfx
+	*
+	* @param TransitionType - direction of the transition
+	* @param EResourceTransitionPipeline - How this UAV is transitioning between Gfx and Compute, if at all.
+	* @param InUAVs - array of UAV objects to transition
+	* @param NumUAVs - number of UAVs to transition
+	*/
+	virtual void RHITransitionResources(EResourceTransitionAccess TransitionType, EResourceTransitionPipeline TransitionPipeline, FUnorderedAccessViewRHIParamRef* InUAVs, int32 NumUAVs)
+	{		
+	}
+
+
 	virtual void RHIBeginRenderQuery(FRenderQueryRHIParamRef RenderQuery) = 0;
 
 	virtual void RHIEndRenderQuery(FRenderQueryRHIParamRef RenderQuery) = 0;
@@ -827,11 +863,22 @@ public:
 	 * @param NewState			The new sampler state.
 	 */
 	virtual void RHISetShaderSampler(FPixelShaderRHIParamRef PixelShader, uint32 SamplerIndex, FSamplerStateRHIParamRef NewState) = 0;
-
-	/** Sets a compute shader UAV parameter. */
+	
+	/**
+	* Sets a compute shader UAV parameter.
+	* @param ComputeShader	The compute shader to set the UAV for.
+	* @param UAVIndex		The index of the UAVIndex.
+	* @param UAV			The new UAV.	
+	*/
 	virtual void RHISetUAVParameter(FComputeShaderRHIParamRef ComputeShader, uint32 UAVIndex, FUnorderedAccessViewRHIParamRef UAV) = 0;
-
-	/** Sets a compute shader UAV parameter and initial count */
+	
+	/**
+	* Sets a compute shader counted UAV parameter and initial count
+	* @param ComputeShader	The compute shader to set the UAV for.
+	* @param UAVIndex		The index of the UAVIndex.
+	* @param UAV			The new UAV.
+	* @param InitialCount	The initial number of items in the UAV.	
+	*/
 	virtual void RHISetUAVParameter(FComputeShaderRHIParamRef ComputeShader, uint32 UAVIndex, FUnorderedAccessViewRHIParamRef UAV, uint32 InitialCount) = 0;
 
 	virtual void RHISetShaderResourceViewParameter(FPixelShaderRHIParamRef PixelShader, uint32 SamplerIndex, FShaderResourceViewRHIParamRef SRV) = 0;
diff --git a/Engine/Source/Runtime/RHI/Public/RHI.h b/Engine/Source/Runtime/RHI/Public/RHI.h
--- a/Engine/Source/Runtime/RHI/Public/RHI.h
+++ b/Engine/Source/Runtime/RHI/Public/RHI.h
@@ -970,6 +970,29 @@ struct FResolveParams
 	{}
 };
 
+enum class EResourceTransitionAccess
+{
+	EReadable, //transition from write-> read
+	EWritable, //transition from read -> write	
+	ERWBarrier, // Mostly for UAVs.  Transition to read/write state and always insert a resource barrier.
+	ERWNoBarrier, //Mostly UAVs.  Indicates we want R/W access and do not require synchronization for the duration of the RW state.  The initial transition from writable->RWNoBarrier and readable->RWNoBarrier still requires a sync
+	EMaxAccess,
+};
+
+class RHI_API FResourceTransitionUtility
+{
+public:
+	static const FString ResourceTransitionAccessStrings[(int32)EResourceTransitionAccess::EMaxAccess + 1];
+};
+
+enum class EResourceTransitionPipeline
+{
+	EGfxToCompute,
+	EComputeToGfx,
+	EGfxToGfx,
+	EComputeToCompute,	
+};
+
 /** specifies an update region for a texture */
 struct FUpdateTextureRegion2D
 {
diff --git a/Engine/Source/Runtime/RHI/Public/RHICommandList.h b/Engine/Source/Runtime/RHI/Public/RHICommandList.h
--- a/Engine/Source/Runtime/RHI/Public/RHICommandList.h
+++ b/Engine/Source/Runtime/RHI/Public/RHICommandList.h
@@ -788,6 +788,56 @@ struct FRHICommandCopyToResolveTarget : public FRHICommand<FRHICommandCopyToReso
 	RHI_API void Execute(FRHICommandListBase& CmdList);
 };
 
+struct FRHICommandTransitionTextures : public FRHICommand<FRHICommandTransitionTextures>
+{
+	static const int32 MaxTexturesToTransition = 16;
+	int32 NumTextures;
+	FTextureRHIParamRef Textures[MaxTexturesToTransition];
+	EResourceTransitionAccess TransitionType;
+	FORCEINLINE_DEBUGGABLE FRHICommandTransitionTextures(EResourceTransitionAccess InTransitionType, FTextureRHIParamRef* InTextures, int32 InNumTextures)
+		: NumTextures(InNumTextures)
+		, TransitionType(InTransitionType)
+	{
+		for (int32 i = 0; i < NumTextures; ++i)
+		{
+			Textures[i] = InTextures[i];
+		}
+	}
+	RHI_API void Execute(FRHICommandListBase& CmdList);
+};
+
+struct FRHICommandTransitionTexturesArray : public FRHICommand<FRHICommandTransitionTexturesArray>
+{	
+	TArray<FTextureRHIParamRef>& Textures;
+	EResourceTransitionAccess TransitionType;
+	FORCEINLINE_DEBUGGABLE FRHICommandTransitionTexturesArray(EResourceTransitionAccess InTransitionType, TArray<FTextureRHIParamRef>& InTextures)
+		: Textures(InTextures)
+		, TransitionType(InTransitionType)
+	{		
+	}
+	RHI_API void Execute(FRHICommandListBase& CmdList);
+};
+
+struct FRHICommandTransitionUAVs : public FRHICommand<FRHICommandTransitionUAVs>
+{
+	static const int32 MaxUAVsToTransition = 16;
+	int32 NumUAVs;
+	FUnorderedAccessViewRHIParamRef UAVs[MaxUAVsToTransition];
+	EResourceTransitionAccess TransitionType;
+	EResourceTransitionPipeline TransitionPipeline;
+	FORCEINLINE_DEBUGGABLE FRHICommandTransitionUAVs(EResourceTransitionAccess InTransitionType, EResourceTransitionPipeline InTransitionPipeline, FUnorderedAccessViewRHIParamRef* InUAVs, int32 InNumUAVs)
+		: NumUAVs(InNumUAVs)
+		, TransitionType(InTransitionType)
+		, TransitionPipeline(InTransitionPipeline)
+	{
+		for (int32 i = 0; i < NumUAVs; ++i)
+		{
+			UAVs[i] = InUAVs[i];
+		}
+	}
+	RHI_API void Execute(FRHICommandListBase& CmdList);
+};
+
 struct FRHICommandClear : public FRHICommand<FRHICommandClear>
 {
 	FLinearColor Color;
@@ -1796,7 +1846,38 @@ public:
 			return;
 		}
 		new (AllocCommand<FRHICommandSubmitCommandsHint>()) FRHICommandSubmitCommandsHint();
-	}	
+	}
+
+	FORCEINLINE_DEBUGGABLE void TransitionResource(EResourceTransitionAccess TransitionType, FTextureRHIParamRef InTexture)
+	{
+		FTextureRHIParamRef Texture = InTexture;
+		if (Bypass())
+		{
+			CMD_CONTEXT(TransitionResources)(TransitionType, &Texture, 1);
+			return;
+		}
+		new (AllocCommand<FRHICommandTransitionTextures>()) FRHICommandTransitionTextures(TransitionType, &Texture, 1);
+	}
+
+	FORCEINLINE_DEBUGGABLE void TransitionResources(EResourceTransitionAccess TransitionType, FTextureRHIParamRef* InTextures, int32 NumTextures)
+	{
+		if (Bypass())
+		{
+			CMD_CONTEXT(TransitionResources)(TransitionType, InTextures, NumTextures);
+			return;
+		}
+		new (AllocCommand<FRHICommandTransitionTextures>()) FRHICommandTransitionTextures(TransitionType, InTextures, NumTextures);
+	}
+
+	FORCEINLINE_DEBUGGABLE void TransitionResourceArrayNoCopy(EResourceTransitionAccess TransitionType, TArray<FTextureRHIParamRef>& InTextures)
+	{
+		if (Bypass())
+		{
+			CMD_CONTEXT(TransitionResources)(TransitionType, &InTextures[0], InTextures.Num());
+			return;
+		}
+		new (AllocCommand<FRHICommandTransitionTexturesArray>()) FRHICommandTransitionTexturesArray(TransitionType, InTextures);
+	}
 
 // These 6 are special in that they must be called on the immediate command list and they force a flush only when we are not doing RHI thread
 	void BeginScene();
diff --git a/Engine/Source/Runtime/RHI/Public/RHICommandListCommandExecutes.inl b/Engine/Source/Runtime/RHI/Public/RHICommandListCommandExecutes.inl
--- a/Engine/Source/Runtime/RHI/Public/RHICommandListCommandExecutes.inl
+++ b/Engine/Source/Runtime/RHI/Public/RHICommandListCommandExecutes.inl
@@ -265,6 +265,24 @@ void FRHICommandCopyToResolveTarget::Execute(FRHICommandListBase& CmdList)
 	INTERNAL_DECORATOR(CopyToResolveTarget)(SourceTexture, DestTexture, bKeepOriginalSurface, ResolveParams);
 }
 
+void FRHICommandTransitionTextures::Execute(FRHICommandListBase& CmdList)
+{
+	RHISTAT(TransitionTextures);
+	INTERNAL_DECORATOR(TransitionResources)(TransitionType, &Textures[0], NumTextures);
+}
+
+void FRHICommandTransitionTexturesArray::Execute(FRHICommandListBase& CmdList)
+{
+	RHISTAT(TransitionTextures);
+	INTERNAL_DECORATOR(TransitionResources)(TransitionType, &Textures[0], Textures.Num());
+}
+
+void FRHICommandTransitionUAVs::Execute(FRHICommandListBase& CmdList)
+{
+	RHISTAT(TransitionUAVs);
+	INTERNAL_DECORATOR(TransitionResources)(TransitionType, TransitionPipeline, UAVs, NumUAVs);
+}
+
 void FRHICommandClear::Execute(FRHICommandListBase& CmdList)
 {
 	RHISTAT(Clear);
diff --git a/Engine/Source/Runtime/RHI/Public/RHIResources.h b/Engine/Source/Runtime/RHI/Public/RHIResources.h
--- a/Engine/Source/Runtime/RHI/Public/RHIResources.h
+++ b/Engine/Source/Runtime/RHI/Public/RHIResources.h
@@ -825,6 +825,12 @@ public:
 	{
 		return ExtractStencil() == StencilWrite;
 	}
+
+	inline bool IsAnyWrite() const
+	{
+		return IsDepthWrite() || IsStencilWrite();
+	}
+
 	inline void SetDepthWrite()
 	{
 		Value = (Type)(ExtractStencil() | DepthWrite);
diff --git a/Engine/Source/Runtime/RHI/Public/RHIUtilities.h b/Engine/Source/Runtime/RHI/Public/RHIUtilities.h
--- a/Engine/Source/Runtime/RHI/Public/RHIUtilities.h
+++ b/Engine/Source/Runtime/RHI/Public/RHIUtilities.h
@@ -171,40 +171,103 @@ inline void DecodeRenderTargetMode(ESimpleRenderTargetMode Mode, ERenderTargetLo
 	}
 }
 
+inline void TransitionSetRenderTargetsHelper(FRHICommandList& RHICmdList, FTextureRHIParamRef NewRenderTarget, FTextureRHIParamRef NewDepthStencilTarget, FExclusiveDepthStencil DepthStencilAccess)
+{
+	int32 TransitionIndex = 0;
+	FTextureRHIParamRef Transitions[2];
+	if (NewRenderTarget)
+	{
+		Transitions[TransitionIndex] = NewRenderTarget;
+		++TransitionIndex;
+	}
+	if (NewDepthStencilTarget && DepthStencilAccess.IsDepthWrite())
+	{
+		Transitions[TransitionIndex] = NewDepthStencilTarget;
+		++TransitionIndex;
+	}
+	RHICmdList.TransitionResources(EResourceTransitionAccess::EWritable, Transitions, TransitionIndex);
+}
+
+inline void TransitionSetRenderTargetsHelper(FRHICommandList& RHICmdList, uint32 NumRenderTargets, const FTextureRHIParamRef* NewRenderTargetsRHI, const FTextureRHIParamRef NewDepthStencilTargetRHI, FExclusiveDepthStencil DepthStencilAccess)
+{
+	FRHIRenderTargetView RTVs[MaxSimultaneousRenderTargets];
+	FTextureRHIParamRef Transitions[MaxSimultaneousRenderTargets + 1];
+	int32 TransitionIndex = 0;
+	for (uint32 Index = 0; Index < NumRenderTargets; Index++)
+	{
+		if (NewRenderTargetsRHI[Index])
+		{
+			Transitions[TransitionIndex] = NewRenderTargetsRHI[Index];
+			++TransitionIndex;
+		}
+	}
+
+	if (NewDepthStencilTargetRHI && DepthStencilAccess.IsDepthWrite())
+	{
+		Transitions[TransitionIndex] = NewDepthStencilTargetRHI;
+		++TransitionIndex;
+	}	
+	RHICmdList.TransitionResources(EResourceTransitionAccess::EWritable, Transitions, TransitionIndex);
+}
+
 /** Helper for the common case of using a single color and depth render target. */
-inline void SetRenderTarget(FRHICommandList& RHICmdList, FTextureRHIParamRef NewRenderTarget, FTextureRHIParamRef NewDepthStencilTarget)
+inline void SetRenderTarget(FRHICommandList& RHICmdList, FTextureRHIParamRef NewRenderTarget, FTextureRHIParamRef NewDepthStencilTarget, bool bWritableBarrier = false)
 {
 	FRHIRenderTargetView RTV(NewRenderTarget);
 	FRHIDepthRenderTargetView DepthRTV(NewDepthStencilTarget);
+
+	//make these rendertargets safely writable
+	if (bWritableBarrier)
+	{
+		TransitionSetRenderTargetsHelper(RHICmdList, NewRenderTarget, NewDepthStencilTarget, FExclusiveDepthStencil::DepthWrite_StencilWrite);
+	}
 	RHICmdList.SetRenderTargets(1, &RTV, &DepthRTV, 0, NULL);
 }
 
 /** Helper for the common case of using a single color and depth render target. */
-inline void SetRenderTarget(FRHICommandList& RHICmdList, FTextureRHIParamRef NewRenderTarget, FTextureRHIParamRef NewDepthStencilTarget, ESimpleRenderTargetMode Mode, FExclusiveDepthStencil DepthStencilAccess = FExclusiveDepthStencil::DepthWrite_StencilWrite)
+inline void SetRenderTarget(FRHICommandList& RHICmdList, FTextureRHIParamRef NewRenderTarget, FTextureRHIParamRef NewDepthStencilTarget, ESimpleRenderTargetMode Mode, FExclusiveDepthStencil DepthStencilAccess = FExclusiveDepthStencil::DepthWrite_StencilWrite, bool bWritableBarrier = false)
 {
 	ERenderTargetLoadAction ColorLoadAction, DepthLoadAction;
 	ERenderTargetStoreAction ColorStoreAction, DepthStoreAction;	
 	DecodeRenderTargetMode(Mode, ColorLoadAction, ColorStoreAction, DepthLoadAction, DepthStoreAction, DepthStencilAccess);
 
+	//make these rendertargets safely writable
+	if (bWritableBarrier)
+	{
+		TransitionSetRenderTargetsHelper(RHICmdList, NewRenderTarget, NewDepthStencilTarget, DepthStencilAccess);
+	}
+
 	// now make the FRHISetRenderTargetsInfo that encapsulates all of the info
 	FRHIRenderTargetView ColorView(NewRenderTarget, 0, -1, ColorLoadAction, ColorStoreAction);
-	FRHISetRenderTargetsInfo Info(1, &ColorView, FRHIDepthRenderTargetView(NewDepthStencilTarget, DepthLoadAction, DepthStoreAction, DepthStencilAccess));
+	FRHISetRenderTargetsInfo Info(1, &ColorView, FRHIDepthRenderTargetView(NewDepthStencilTarget, DepthLoadAction, DepthStoreAction, DepthStencilAccess));	
 	RHICmdList.SetRenderTargetsAndClear(Info);
 }
 
 /** Helper for the common case of using a single color and depth render target, with a mip index for the color target. */
-inline void SetRenderTarget(FRHICommandList& RHICmdList, FTextureRHIParamRef NewRenderTarget, int32 MipIndex, FTextureRHIParamRef NewDepthStencilTarget)
+inline void SetRenderTarget(FRHICommandList& RHICmdList, FTextureRHIParamRef NewRenderTarget, int32 MipIndex, FTextureRHIParamRef NewDepthStencilTarget, bool bWritableBarrier = false)
 {
 	FRHIRenderTargetView RTV(NewRenderTarget, MipIndex, -1);
 	FRHIDepthRenderTargetView DepthRTV(NewDepthStencilTarget);
+	
+	//make these rendertargets safely writable
+	if (bWritableBarrier)
+	{
+		TransitionSetRenderTargetsHelper(RHICmdList, NewRenderTarget, NewDepthStencilTarget, FExclusiveDepthStencil::DepthWrite_StencilWrite);
+	}
 	RHICmdList.SetRenderTargets(1, &RTV, &DepthRTV, 0, nullptr);
 }
 
 /** Helper for the common case of using a single color and depth render target, with a mip index for the color target. */
-inline void SetRenderTarget(FRHICommandList& RHICmdList, FTextureRHIParamRef NewRenderTarget, int32 MipIndex, int32 ArraySliceIndex, FTextureRHIParamRef NewDepthStencilTarget)
+inline void SetRenderTarget(FRHICommandList& RHICmdList, FTextureRHIParamRef NewRenderTarget, int32 MipIndex, int32 ArraySliceIndex, FTextureRHIParamRef NewDepthStencilTarget, bool bWritableBarrier = false)
 {
 	FRHIRenderTargetView RTV(NewRenderTarget, MipIndex, ArraySliceIndex);
 	FRHIDepthRenderTargetView DepthRTV(NewDepthStencilTarget);
+
+	//make these rendertargets safely writable
+	if (bWritableBarrier)
+	{
+		TransitionSetRenderTargetsHelper(RHICmdList, NewRenderTarget, NewDepthStencilTarget, FExclusiveDepthStencil::DepthWrite_StencilWrite);
+	}
 	RHICmdList.SetRenderTargets(1, &RTV, &DepthRTV, 0, nullptr);
 }
 
@@ -215,14 +278,21 @@ inline void SetRenderTargets(
 	const FTextureRHIParamRef* NewRenderTargetsRHI,
 	FTextureRHIParamRef NewDepthStencilTargetRHI,
 	uint32 NewNumUAVs,
-	const FUnorderedAccessViewRHIParamRef* UAVs
+	const FUnorderedAccessViewRHIParamRef* UAVs,
+	bool bWritableBarrier = false
 	)
 {
 	FRHIRenderTargetView RTVs[MaxSimultaneousRenderTargets];
-
 	for (uint32 Index = 0; Index < NewNumSimultaneousRenderTargets; Index++)
 	{
 		RTVs[Index] = FRHIRenderTargetView(NewRenderTargetsRHI[Index]);
+	
+	}
+
+	//make these rendertargets safely writable
+	if (bWritableBarrier)
+	{
+		TransitionSetRenderTargetsHelper(RHICmdList, NewNumSimultaneousRenderTargets, NewRenderTargetsRHI, NewDepthStencilTargetRHI, FExclusiveDepthStencil::DepthWrite_StencilWrite);
 	}
 
 	FRHIDepthRenderTargetView DepthRTV(NewDepthStencilTargetRHI);
@@ -236,7 +306,8 @@ inline void SetRenderTargets(
 	const FTextureRHIParamRef* NewRenderTargetsRHI,
 	FTextureRHIParamRef NewDepthStencilTargetRHI,
 	ESimpleRenderTargetMode Mode,
-	FExclusiveDepthStencil DepthStencilAccess
+	FExclusiveDepthStencil DepthStencilAccess,
+	bool bWritableBarrier = false
 	)
 {
 	ERenderTargetLoadAction ColorLoadAction, DepthLoadAction;
@@ -244,10 +315,16 @@ inline void SetRenderTargets(
 	DecodeRenderTargetMode(Mode, ColorLoadAction, ColorStoreAction, DepthLoadAction, DepthStoreAction, DepthStencilAccess);
 
 	FRHIRenderTargetView RTVs[MaxSimultaneousRenderTargets];
-
+		
 	for (uint32 Index = 0; Index < NewNumSimultaneousRenderTargets; Index++)
 	{
-		RTVs[Index] = FRHIRenderTargetView(NewRenderTargetsRHI[Index], 0, -1, ColorLoadAction, ColorStoreAction);
+		RTVs[Index] = FRHIRenderTargetView(NewRenderTargetsRHI[Index], 0, -1, ColorLoadAction, ColorStoreAction);	
+	}
+
+	//make these rendertargets safely writable
+	if (bWritableBarrier)
+	{
+		TransitionSetRenderTargetsHelper(RHICmdList, NewNumSimultaneousRenderTargets, NewRenderTargetsRHI, NewDepthStencilTargetRHI, DepthStencilAccess);
 	}
 
 	FRHIDepthRenderTargetView DepthRTV(NewDepthStencilTargetRHI, DepthLoadAction, DepthStoreAction, DepthStencilAccess);
diff --git a/Engine/Source/Runtime/Renderer/Private/AtmosphereTextures.cpp b/Engine/Source/Runtime/Renderer/Private/AtmosphereTextures.cpp
--- a/Engine/Source/Runtime/Renderer/Private/AtmosphereTextures.cpp
+++ b/Engine/Source/Runtime/Renderer/Private/AtmosphereTextures.cpp
@@ -13,13 +13,13 @@ void FAtmosphereTextures::InitDynamicRHI()
 	check(PrecomputeParams != NULL);
 	// Allocate atmosphere precompute textures...
 	{
+		FRHICommandListImmediate& RHICmdList = FRHICommandListExecutor::GetImmediateCommandList();
+
 		// todo: Expose
 		// Transmittance
 		FIntPoint GTransmittanceTexSize(PrecomputeParams->TransmittanceTexWidth, PrecomputeParams->TransmittanceTexHeight);
 		FPooledRenderTargetDesc TransmittanceDesc(FPooledRenderTargetDesc::Create2DDesc(GTransmittanceTexSize, PF_FloatRGBA, FClearValueBinding::None, TexCreate_None, TexCreate_RenderTargetable, false));
-		GRenderTargetPool.FindFreeElement(TransmittanceDesc, AtmosphereTransmittance, TEXT("AtmosphereTransmittance"));
-
-		FRHICommandListImmediate& RHICmdList = FRHICommandListExecutor::GetImmediateCommandList();
+		GRenderTargetPool.FindFreeElement(RHICmdList, TransmittanceDesc, AtmosphereTransmittance, TEXT("AtmosphereTransmittance"));
 
 		SetRenderTarget(RHICmdList, AtmosphereTransmittance->GetRenderTargetItem().TargetableTexture, FTextureRHIRef());
 		RHICmdList.Clear(true, FLinearColor(0, 0, 0, 0), false, 0, false, 0, FIntRect());
@@ -28,28 +28,28 @@ void FAtmosphereTextures::InitDynamicRHI()
 		// Irradiance
 		FIntPoint GIrradianceTexSize(PrecomputeParams->IrradianceTexWidth, PrecomputeParams->IrradianceTexHeight);
 		FPooledRenderTargetDesc IrradianceDesc(FPooledRenderTargetDesc::Create2DDesc(GIrradianceTexSize, PF_FloatRGBA, FClearValueBinding::Black, TexCreate_None, TexCreate_RenderTargetable, false));
-		GRenderTargetPool.FindFreeElement(IrradianceDesc, AtmosphereIrradiance, TEXT("AtmosphereIrradiance"));
+		GRenderTargetPool.FindFreeElement(RHICmdList, IrradianceDesc, AtmosphereIrradiance, TEXT("AtmosphereIrradiance"));
 
 		SetRenderTarget(RHICmdList, AtmosphereIrradiance->GetRenderTargetItem().TargetableTexture, FTextureRHIRef());
 		RHICmdList.Clear(true, FLinearColor(0, 0, 0, 0), false, 0, false, 0, FIntRect());
 		RHICmdList.CopyToResolveTarget(AtmosphereIrradiance->GetRenderTargetItem().TargetableTexture, AtmosphereIrradiance->GetRenderTargetItem().ShaderResourceTexture, true, FResolveParams());
 
 		// DeltaE
-		GRenderTargetPool.FindFreeElement(IrradianceDesc, AtmosphereDeltaE, TEXT("AtmosphereDeltaE"));
+		GRenderTargetPool.FindFreeElement(RHICmdList, IrradianceDesc, AtmosphereDeltaE, TEXT("AtmosphereDeltaE"));
 
 		// 3D Texture
 		// Inscatter
 		FPooledRenderTargetDesc InscatterDesc(FPooledRenderTargetDesc::CreateVolumeDesc(PrecomputeParams->InscatterMuSNum * PrecomputeParams->InscatterNuNum, PrecomputeParams->InscatterMuNum, PrecomputeParams->InscatterAltitudeSampleNum, PF_FloatRGBA, FClearValueBinding::None, TexCreate_None, TexCreate_ShaderResource | TexCreate_RenderTargetable, false));
-		GRenderTargetPool.FindFreeElement(InscatterDesc, AtmosphereInscatter, TEXT("AtmosphereInscatter"));
+		GRenderTargetPool.FindFreeElement(RHICmdList, InscatterDesc, AtmosphereInscatter, TEXT("AtmosphereInscatter"));
 
 		// DeltaSR
-		GRenderTargetPool.FindFreeElement(InscatterDesc, AtmosphereDeltaSR, TEXT("AtmosphereDeltaSR"));
+		GRenderTargetPool.FindFreeElement(RHICmdList, InscatterDesc, AtmosphereDeltaSR, TEXT("AtmosphereDeltaSR"));
 
 		// DeltaSM
-		GRenderTargetPool.FindFreeElement(InscatterDesc, AtmosphereDeltaSM, TEXT("AtmosphereDeltaSM"));
+		GRenderTargetPool.FindFreeElement(RHICmdList, InscatterDesc, AtmosphereDeltaSM, TEXT("AtmosphereDeltaSM"));
 
 		// DeltaJ
-		GRenderTargetPool.FindFreeElement(InscatterDesc, AtmosphereDeltaJ, TEXT("AtmosphereDeltaJ"));
+		GRenderTargetPool.FindFreeElement(RHICmdList, InscatterDesc, AtmosphereDeltaJ, TEXT("AtmosphereDeltaJ"));
 	}
 }
 
diff --git a/Engine/Source/Runtime/Renderer/Private/CompositionLighting/CompositionLighting.cpp b/Engine/Source/Runtime/Renderer/Private/CompositionLighting/CompositionLighting.cpp
--- a/Engine/Source/Runtime/Renderer/Private/CompositionLighting/CompositionLighting.cpp
+++ b/Engine/Source/Runtime/Renderer/Private/CompositionLighting/CompositionLighting.cpp
@@ -256,7 +256,7 @@ void FCompositionLighting::ProcessBeforeBasePass(FRHICommandListImmediate& RHICm
 		FMemMark Mark(FMemStack::Get());
 		FRenderingCompositePassContext CompositeContext(RHICmdList, View);
 
-		FPostprocessContext Context(CompositeContext.Graph, View);
+		FPostprocessContext Context(RHICmdList, CompositeContext.Graph, View);
 
 		// Add the passes we want to add to the graph (commenting a line means the pass is not inserted into the graph) ----------
 
@@ -301,7 +301,7 @@ void FCompositionLighting::ProcessAfterBasePass(FRHICommandListImmediate& RHICmd
 		FMemMark Mark(FMemStack::Get());
 		FRenderingCompositePassContext CompositeContext(RHICmdList, View);
 
-		FPostprocessContext Context(CompositeContext.Graph, View);
+		FPostprocessContext Context(RHICmdList, CompositeContext.Graph, View);
 
 		// Add the passes we want to add to the graph ----------
 		
@@ -356,7 +356,7 @@ void FCompositionLighting::ProcessLpvIndirect(FRHICommandListImmediate& RHICmdLi
 	
 	FMemMark Mark(FMemStack::Get());
 	FRenderingCompositePassContext CompositeContext(RHICmdList, View);
-	FPostprocessContext Context(CompositeContext.Graph, View);
+	FPostprocessContext Context(RHICmdList, CompositeContext.Graph, View);
 
 	if(IsLpvIndirectPassRequired(Context))
 	{
@@ -393,7 +393,7 @@ void FCompositionLighting::ProcessAfterLighting(FRHICommandListImmediate& RHICmd
 	{
 		FMemMark Mark(FMemStack::Get());
 		FRenderingCompositePassContext CompositeContext(RHICmdList, View);
-		FPostprocessContext Context(CompositeContext.Graph, View);
+		FPostprocessContext Context(RHICmdList, CompositeContext.Graph, View);
 		FRenderingCompositeOutputRef AmbientOcclusion;
 
 		// Screen Space Subsurface Scattering
diff --git a/Engine/Source/Runtime/Renderer/Private/CompositionLighting/PostProcessAmbient.cpp b/Engine/Source/Runtime/Renderer/Private/CompositionLighting/PostProcessAmbient.cpp
--- a/Engine/Source/Runtime/Renderer/Private/CompositionLighting/PostProcessAmbient.cpp
+++ b/Engine/Source/Runtime/Renderer/Private/CompositionLighting/PostProcessAmbient.cpp
@@ -122,7 +122,7 @@ void FRCPassPostProcessAmbient::Process(FRenderingCompositePassContext& Context)
 	const FSceneRenderTargetItem& DestRenderTarget = FSceneRenderTargets::Get(Context.RHICmdList).GetSceneColor()->GetRenderTargetItem();
 
 	// Set the view family's render target/viewport.
-	SetRenderTarget(Context.RHICmdList, DestRenderTarget.TargetableTexture, FTextureRHIRef());
+	SetRenderTarget(Context.RHICmdList, DestRenderTarget.TargetableTexture, FTextureRHIRef(), true);
 	Context.SetViewportAndCallRHI(View.ViewRect);
 
 	// set the state
diff --git a/Engine/Source/Runtime/Renderer/Private/CompositionLighting/PostProcessAmbientOcclusion.cpp b/Engine/Source/Runtime/Renderer/Private/CompositionLighting/PostProcessAmbientOcclusion.cpp
--- a/Engine/Source/Runtime/Renderer/Private/CompositionLighting/PostProcessAmbientOcclusion.cpp
+++ b/Engine/Source/Runtime/Renderer/Private/CompositionLighting/PostProcessAmbientOcclusion.cpp
@@ -574,6 +574,7 @@ void FRCPassPostProcessBasePassAO::Process(FRenderingCompositePassContext& Conte
 	const FSceneRenderTargetItem& DestRenderTarget = SceneContext.GetSceneColor()->GetRenderTargetItem();
 
 	// Set the view family's render target/viewport.
+	Context.RHICmdList.TransitionResource(EResourceTransitionAccess::EWritable, DestRenderTarget.TargetableTexture);
 	SetRenderTarget(Context.RHICmdList, DestRenderTarget.TargetableTexture,	FTextureRHIParamRef(), ESimpleRenderTargetMode::EExistingColorAndDepth);
 	Context.SetViewportAndCallRHI(View.ViewRect);
 
diff --git a/Engine/Source/Runtime/Renderer/Private/CompositionLighting/PostProcessDeferredDecals.cpp b/Engine/Source/Runtime/Renderer/Private/CompositionLighting/PostProcessDeferredDecals.cpp
--- a/Engine/Source/Runtime/Renderer/Private/CompositionLighting/PostProcessDeferredDecals.cpp
+++ b/Engine/Source/Runtime/Renderer/Private/CompositionLighting/PostProcessDeferredDecals.cpp
@@ -485,6 +485,7 @@ void FRCPassPostProcessDeferredDecals::Process(FRenderingCompositePassContext& C
 		ResolveBufferMax,
 	};
 
+	bool TargetsToTransitionWritable[ResolveBufferMax] = { true };
 	FTextureRHIParamRef TargetsToResolve[ResolveBufferMax] = { nullptr };
 
 	if(CurrentStage == DRS_BeforeBasePass)
@@ -504,13 +505,13 @@ void FRCPassPostProcessDeferredDecals::Process(FRenderingCompositePassContext& C
 		if(!SceneContext.DBufferA)
 		{
 			Desc.ClearValue = FClearValueBinding::Black;
-			GRenderTargetPool.FindFreeElement(Desc, SceneContext.DBufferA, TEXT("DBufferA"));
+			GRenderTargetPool.FindFreeElement(Context.RHICmdList, Desc, SceneContext.DBufferA, TEXT("DBufferA"));
 		}
 
 		if(!SceneContext.DBufferB)
 		{
 			Desc.ClearValue = FClearValueBinding(FLinearColor(128.0f / 255.0f, 128.0f / 255.0f, 128.0f / 255.0f, 1));
-			GRenderTargetPool.FindFreeElement(Desc, SceneContext.DBufferB, TEXT("DBufferB"));
+			GRenderTargetPool.FindFreeElement(Context.RHICmdList, Desc, SceneContext.DBufferB, TEXT("DBufferB"));
 		}
 
 		Desc.Format = PF_R8G8;
@@ -518,7 +519,7 @@ void FRCPassPostProcessDeferredDecals::Process(FRenderingCompositePassContext& C
 		if(!SceneContext.DBufferC)
 		{
 			Desc.ClearValue = FClearValueBinding(FLinearColor(0, 1, 0, 1));
-			GRenderTargetPool.FindFreeElement(Desc, SceneContext.DBufferC, TEXT("DBufferC"));
+			GRenderTargetPool.FindFreeElement(Context.RHICmdList, Desc, SceneContext.DBufferC, TEXT("DBufferC"));
 		}
 
 		// we assume views are non overlapping, then we need to clear only once in the beginning, otherwise we would need to set scissor rects
@@ -607,7 +608,7 @@ void FRCPassPostProcessDeferredDecals::Process(FRenderingCompositePassContext& C
 							TargetsToResolve[GBufferAIndex] = SceneContext.GBufferA->GetRenderTargetItem().TargetableTexture;
 							TargetsToResolve[GBufferBIndex] = SceneContext.GBufferB->GetRenderTargetItem().TargetableTexture;
 							TargetsToResolve[GBufferCIndex] = SceneContext.GBufferC->GetRenderTargetItem().TargetableTexture;
-							SetRenderTargets(RHICmdList, 4, TargetsToResolve, SceneContext.GetSceneDepthSurface(), ESimpleRenderTargetMode::EExistingColorAndDepth, FExclusiveDepthStencil::DepthRead_StencilWrite);
+							SetRenderTargets(RHICmdList, 4, TargetsToResolve, SceneContext.GetSceneDepthSurface(), ESimpleRenderTargetMode::EExistingColorAndDepth, FExclusiveDepthStencil::DepthRead_StencilWrite, TargetsToTransitionWritable[CurrentRenderTargetMode]);							
 							break;
 
 						case FDecalRendering::RTM_SceneColorAndGBufferDepthWrite:
@@ -616,30 +617,31 @@ void FRCPassPostProcessDeferredDecals::Process(FRenderingCompositePassContext& C
 							TargetsToResolve[GBufferBIndex] = SceneContext.GBufferB->GetRenderTargetItem().TargetableTexture;
 							TargetsToResolve[GBufferCIndex] = SceneContext.GBufferC->GetRenderTargetItem().TargetableTexture;
 							TargetsToResolve[GBufferEIndex] = SceneContext.GBufferE->GetRenderTargetItem().TargetableTexture;
-							SetRenderTargets(RHICmdList, 5, TargetsToResolve, SceneContext.GetSceneDepthSurface(), ESimpleRenderTargetMode::EExistingColorAndDepth, FExclusiveDepthStencil::DepthWrite_StencilWrite);
+							SetRenderTargets(RHICmdList, 5, TargetsToResolve, SceneContext.GetSceneDepthSurface(), ESimpleRenderTargetMode::EExistingColorAndDepth, FExclusiveDepthStencil::DepthWrite_StencilWrite, TargetsToTransitionWritable[CurrentRenderTargetMode]);
 							break;
 
 						case FDecalRendering::RTM_GBufferNormal:
 							TargetsToResolve[GBufferAIndex] = SceneContext.GBufferA->GetRenderTargetItem().TargetableTexture;
-							SetRenderTarget(RHICmdList, TargetsToResolve[GBufferAIndex], SceneContext.GetSceneDepthSurface(), ESimpleRenderTargetMode::EExistingColorAndDepth, FExclusiveDepthStencil::DepthRead_StencilWrite);
+							SetRenderTarget(RHICmdList, TargetsToResolve[GBufferAIndex], SceneContext.GetSceneDepthSurface(), ESimpleRenderTargetMode::EExistingColorAndDepth, FExclusiveDepthStencil::DepthRead_StencilWrite, TargetsToTransitionWritable[CurrentRenderTargetMode]);
 							break;
 
 						case FDecalRendering::RTM_SceneColor:
 							TargetsToResolve[SceneColorIndex] = SceneContext.GetSceneColor()->GetRenderTargetItem().TargetableTexture;
-							SetRenderTarget(RHICmdList, TargetsToResolve[SceneColorIndex], SceneContext.GetSceneDepthSurface(), ESimpleRenderTargetMode::EExistingColorAndDepth, FExclusiveDepthStencil::DepthRead_StencilWrite);
+							SetRenderTarget(RHICmdList, TargetsToResolve[SceneColorIndex], SceneContext.GetSceneDepthSurface(), ESimpleRenderTargetMode::EExistingColorAndDepth, FExclusiveDepthStencil::DepthRead_StencilWrite, TargetsToTransitionWritable[CurrentRenderTargetMode]);
 							break;
 
 						case FDecalRendering::RTM_DBuffer:
 							TargetsToResolve[DBufferAIndex] = SceneContext.DBufferA->GetRenderTargetItem().TargetableTexture;
 							TargetsToResolve[DBufferBIndex] = SceneContext.DBufferB->GetRenderTargetItem().TargetableTexture;
 							TargetsToResolve[DBufferCIndex] = SceneContext.DBufferC->GetRenderTargetItem().TargetableTexture;
-							SetRenderTargets(RHICmdList, 3, &TargetsToResolve[DBufferAIndex], SceneContext.GetSceneDepthSurface(), ESimpleRenderTargetMode::EExistingColorAndDepth, FExclusiveDepthStencil::DepthRead_StencilWrite);
+							SetRenderTargets(RHICmdList, 3, &TargetsToResolve[DBufferAIndex], SceneContext.GetSceneDepthSurface(), ESimpleRenderTargetMode::EExistingColorAndDepth, FExclusiveDepthStencil::DepthRead_StencilWrite, TargetsToTransitionWritable[CurrentRenderTargetMode]);
 							break;
 
 						default:
 							check(0);
 							break;
 					}
+					TargetsToTransitionWritable[CurrentRenderTargetMode] = false;
 					Context.SetViewportAndCallRHI(DestRect);
 
 					// we need to reset the stream source after any call to SetRenderTarget (at least for Metal, which doesn't queue up VB assignments)
diff --git a/Engine/Source/Runtime/Renderer/Private/DeferredShadingRenderer.cpp b/Engine/Source/Runtime/Renderer/Private/DeferredShadingRenderer.cpp
--- a/Engine/Source/Runtime/Renderer/Private/DeferredShadingRenderer.cpp
+++ b/Engine/Source/Runtime/Renderer/Private/DeferredShadingRenderer.cpp
@@ -605,7 +605,7 @@ void FDeferredShadingSceneRenderer::RenderVisualizeTexturePool(FRHICommandListIm
 	};
 
 	FPooledRenderTargetDesc Desc(FPooledRenderTargetDesc::Create2DDesc(FIntPoint(TexturePoolVisualizerSizeX, TexturePoolVisualizerSizeY), PF_B8G8R8A8, FClearValueBinding::None, TexCreate_None, TexCreate_None, false));
-	GRenderTargetPool.FindFreeElement(Desc, VisualizeTexturePool, TEXT("VisualizeTexturePool"));
+	GRenderTargetPool.FindFreeElement(RHICmdList, Desc, VisualizeTexturePool, TEXT("VisualizeTexturePool"));
 	
 	uint32 Pitch;
 	FColor* TextureData = (FColor*)RHICmdList.LockTexture2D((FTexture2DRHIRef&)VisualizeTexturePool->GetRenderTargetItem().ShaderResourceTexture, 0, RLM_WriteOnly, Pitch, false);
@@ -788,6 +788,9 @@ void FDeferredShadingSceneRenderer::Render(FRHICommandListImmediate& RHICmdList)
 {
 	FSceneRenderTargets& SceneContext = FSceneRenderTargets::Get(RHICmdList);
 
+	//make sure all the targets we're going to use will be safely writable.
+	GRenderTargetPool.TransitionTargetsWritable(RHICmdList);
+
 	// this way we make sure the SceneColor format is the correct one and not the one from the end of frame before
 	SceneContext.ReleaseSceneColor();
 
@@ -814,7 +817,7 @@ void FDeferredShadingSceneRenderer::Render(FRHICommandListImmediate& RHICmdList)
 		GSystemTextures.InitializeTextures(RHICmdList, FeatureLevel);
 
 		// Allocate the maximum scene render target space for the current view family.
-		SceneContext.Allocate(ViewFamily);
+		SceneContext.Allocate(RHICmdList, ViewFamily);
 	}
 
 	// Find the visible primitives.
@@ -832,7 +835,7 @@ void FDeferredShadingSceneRenderer::Render(FRHICommandListImmediate& RHICmdList)
 		// there are dynamic attempts to get this target during parallel rendering
 		for (int32 ViewIndex = 0; ViewIndex < Views.Num(); ViewIndex++)
 		{
-			Views[ViewIndex].GetEyeAdaptation();
+			Views[ViewIndex].GetEyeAdaptation(RHICmdList);
 		}	
 	}
 
@@ -937,7 +940,7 @@ void FDeferredShadingSceneRenderer::Render(FRHICommandListImmediate& RHICmdList)
 	}
 
 	// Notify the FX system that the scene is about to be rendered.
-	if (Scene->FXSystem)
+	if (Scene->FXSystem && Views.IsValidIndex(0))
 	{
 		QUICK_SCOPE_CYCLE_COUNTER(STAT_FDeferredShadingSceneRenderer_FXSystem_PreRender);
 		Scene->FXSystem->PreRender(RHICmdList, &Views[0].GlobalDistanceFieldInfo.ParameterData);
@@ -987,7 +990,7 @@ void FDeferredShadingSceneRenderer::Render(FRHICommandListImmediate& RHICmdList)
 	{
 		QUICK_SCOPE_CYCLE_COUNTER(STAT_FDeferredShadingSceneRenderer_AllocGBufferTargets);
 		SceneContext.PreallocGBufferTargets(bShouldRenderVelocities && bUseVelocityGBuffer);
-		SceneContext.AllocGBufferTargets();
+		SceneContext.AllocGBufferTargets(RHICmdList);
 	}
 	
 	// Clear LPVs for all views
@@ -1036,7 +1039,7 @@ void FDeferredShadingSceneRenderer::Render(FRHICommandListImmediate& RHICmdList)
 	if (bIsWireframe && FSceneRenderer::ShouldCompositeEditorPrimitives(Views[0]))
 	{
 		// In Editor we want wire frame view modes to be MSAA for better quality. Resolve will be done with EditorPrimitives
-		SetRenderTarget(RHICmdList, SceneContext.GetEditorPrimitivesColor(), SceneContext.GetEditorPrimitivesDepth(), ESimpleRenderTargetMode::EClearColorAndDepth);
+		SetRenderTarget(RHICmdList, SceneContext.GetEditorPrimitivesColor(RHICmdList), SceneContext.GetEditorPrimitivesDepth(RHICmdList), ESimpleRenderTargetMode::EClearColorAndDepth);
 	}
 	else if (!bIsGBufferCurrent)
 	{
@@ -1078,19 +1081,29 @@ void FDeferredShadingSceneRenderer::Render(FRHICommandListImmediate& RHICmdList)
 
 		bRequiresFarZQuadClear = false;
 	}
-	
-	{
-		QUICK_SCOPE_CYCLE_COUNTER(STAT_FDeferredShadingSceneRenderer_Resolve_After_Basepass);
 
-		SceneContext.ResolveSceneColor(RHICmdList, FResolveRect(0, 0, ViewFamily.FamilySizeX, ViewFamily.FamilySizeY));
+	{
+		QUICK_SCOPE_CYCLE_COUNTER(STAT_FDeferredShadingSceneRenderer_ResolveDepth_After_Basepass);
 		SceneContext.ResolveSceneDepthTexture(RHICmdList);
 		SceneContext.ResolveSceneDepthToAuxiliaryTexture(RHICmdList);
+	}
+		
+	
+
+	bool bOcclusionAfterBasePass = bIsOcclusionTesting && !bOcclusionBeforeBasePass;
+	bool bHZBAfterBasePass = true;
+	RenderOcclusion(RHICmdList, bOcclusionAfterBasePass, bHZBAfterBasePass);
+	ServiceLocalQueue();
 
+	{
+		QUICK_SCOPE_CYCLE_COUNTER(STAT_FDeferredShadingSceneRenderer_Resolve_After_Basepass);
+		SceneContext.ResolveSceneColor(RHICmdList, FResolveRect(0, 0, ViewFamily.FamilySizeX, ViewFamily.FamilySizeY));
 		SceneContext.FinishRenderingGBuffer(RHICmdList);
 
 		RenderCustomDepthPass(RHICmdList);
 		ServiceLocalQueue();
 	}
+
 	// Notify the FX system that opaque primitives have been rendered and we now have a valid depth buffer.
 	if (Scene->FXSystem && Views.IsValidIndex(0))
 	{
@@ -1104,11 +1117,6 @@ void FDeferredShadingSceneRenderer::Render(FRHICommandListImmediate& RHICmdList)
 		ServiceLocalQueue();
 	}
 
-	bool bOcclusionAfterBasePass = bIsOcclusionTesting && !bOcclusionBeforeBasePass;
-	bool bHZBAfterBasePass = true;
-	RenderOcclusion(RHICmdList, bOcclusionAfterBasePass, bHZBAfterBasePass);
-	ServiceLocalQueue();
-
 	TRefCountPtr<IPooledRenderTarget> VelocityRT;
 
 	if (bUseVelocityGBuffer)
@@ -1349,7 +1357,7 @@ void FDeferredShadingSceneRenderer::Render(FRHICommandListImmediate& RHICmdList)
 	else
 	{
 		// Release the original reference on the scene render targets
-		SceneContext.AdjustGBufferRefCount(-1);
+		SceneContext.AdjustGBufferRefCount(RHICmdList, -1);
 	}
 
 	//grab the new transform out of the proxies for next frame
diff --git a/Engine/Source/Runtime/Renderer/Private/DistanceFieldLightingPost.cpp b/Engine/Source/Runtime/Renderer/Private/DistanceFieldLightingPost.cpp
--- a/Engine/Source/Runtime/Renderer/Private/DistanceFieldLightingPost.cpp
+++ b/Engine/Source/Runtime/Renderer/Private/DistanceFieldLightingPost.cpp
@@ -542,14 +542,14 @@ IMPLEMENT_SHADER_TYPE(template<>,TFilterHistoryPS<true>,TEXT("DistanceFieldLight
 IMPLEMENT_SHADER_TYPE(template<>,TFilterHistoryPS<false>,TEXT("DistanceFieldLightingPost"),TEXT("FilterHistoryPS"),SF_Pixel);
 
 
-void AllocateOrReuseAORenderTarget(TRefCountPtr<IPooledRenderTarget>& Target, const TCHAR* Name, bool bWithAlpha)
+void AllocateOrReuseAORenderTarget(FRHICommandList& RHICmdList, TRefCountPtr<IPooledRenderTarget>& Target, const TCHAR* Name, bool bWithAlpha)
 {
 	if (!Target)
 	{
 		FIntPoint BufferSize = GetBufferSizeForAO();
 
 		FPooledRenderTargetDesc Desc(FPooledRenderTargetDesc::Create2DDesc(BufferSize, bWithAlpha ? PF_FloatRGBA : PF_FloatRGB, FClearValueBinding::None, TexCreate_None, TexCreate_RenderTargetable | TexCreate_UAV, false));
-		GRenderTargetPool.FindFreeElement(Desc, Target, Name);
+		GRenderTargetPool.FindFreeElement(RHICmdList, Desc, Target, Name);
 	}
 }
 
@@ -582,13 +582,13 @@ void UpdateHistory(
 		{
 			// Reuse a render target from the pool with a consistent name, for vis purposes
 			TRefCountPtr<IPooledRenderTarget> NewBentNormalHistory;
-			AllocateOrReuseAORenderTarget(NewBentNormalHistory, BentNormalHistoryRTName, true);
+			AllocateOrReuseAORenderTarget(RHICmdList, NewBentNormalHistory, BentNormalHistoryRTName, true);
 
 			TRefCountPtr<IPooledRenderTarget> NewIrradianceHistory;
 
 			if (bUseDistanceFieldGI)
 			{
-				AllocateOrReuseAORenderTarget(NewIrradianceHistory, IrradianceHistoryRTName, false);
+				AllocateOrReuseAORenderTarget(RHICmdList, NewIrradianceHistory, IrradianceHistoryRTName, false);
 			}
 
 			SCOPED_DRAW_EVENT(RHICmdList, UpdateHistory);
@@ -656,13 +656,13 @@ void UpdateHistory(
 					GRenderTargetPool.FreeUnusedResource(*BentNormalHistoryState);
 					*BentNormalHistoryState = NULL;
 					// Update the view state's render target reference with the new history
-					AllocateOrReuseAORenderTarget(*BentNormalHistoryState, BentNormalHistoryRTName, true);
+					AllocateOrReuseAORenderTarget(RHICmdList, *BentNormalHistoryState, BentNormalHistoryRTName, true);
 
 					if (bUseDistanceFieldGI)
 					{
 						GRenderTargetPool.FreeUnusedResource(*IrradianceHistoryState);
 						*IrradianceHistoryState = NULL;
-						AllocateOrReuseAORenderTarget(*IrradianceHistoryState, IrradianceHistoryRTName, false);
+						AllocateOrReuseAORenderTarget(RHICmdList, *IrradianceHistoryState, IrradianceHistoryRTName, false);
 					}
 				}
 
@@ -767,11 +767,11 @@ void PostProcessBentNormalAOSurfaceCache(
 
 	TRefCountPtr<IPooledRenderTarget> DistanceFieldAOBentNormal;
 	TRefCountPtr<IPooledRenderTarget> DistanceFieldIrradiance;
-	AllocateOrReuseAORenderTarget(DistanceFieldAOBentNormal, TEXT("DistanceFieldBentNormalAO"), true);
+	AllocateOrReuseAORenderTarget(RHICmdList, DistanceFieldAOBentNormal, TEXT("DistanceFieldBentNormalAO"), true);
 
 	if (bUseDistanceFieldGI)
 	{
-		AllocateOrReuseAORenderTarget(DistanceFieldIrradiance, TEXT("DistanceFieldIrradiance"), false);
+		AllocateOrReuseAORenderTarget(RHICmdList, DistanceFieldIrradiance, TEXT("DistanceFieldIrradiance"), false);
 	}
 
 	TRefCountPtr<IPooledRenderTarget> DistanceFieldAOBentNormal2;
@@ -779,11 +779,11 @@ void PostProcessBentNormalAOSurfaceCache(
 
 	if (GAOFillGaps)
 	{
-		AllocateOrReuseAORenderTarget(DistanceFieldAOBentNormal2, TEXT("DistanceFieldBentNormalAO2"), true);
+		AllocateOrReuseAORenderTarget(RHICmdList, DistanceFieldAOBentNormal2, TEXT("DistanceFieldBentNormalAO2"), true);
 
 		if (bUseDistanceFieldGI)
 		{
-			AllocateOrReuseAORenderTarget(DistanceFieldIrradiance2, TEXT("DistanceFieldIrradiance2"), false);
+			AllocateOrReuseAORenderTarget(RHICmdList, DistanceFieldIrradiance2, TEXT("DistanceFieldIrradiance2"), false);
 		}
 	}
 
diff --git a/Engine/Source/Runtime/Renderer/Private/DistanceFieldLightingPost.h b/Engine/Source/Runtime/Renderer/Private/DistanceFieldLightingPost.h
--- a/Engine/Source/Runtime/Renderer/Private/DistanceFieldLightingPost.h
+++ b/Engine/Source/Runtime/Renderer/Private/DistanceFieldLightingPost.h
@@ -6,7 +6,7 @@
 
 #pragma once
 
-extern void AllocateOrReuseAORenderTarget(TRefCountPtr<IPooledRenderTarget>& Target, const TCHAR* Name, bool bWithAlpha);
+extern void AllocateOrReuseAORenderTarget(FRHICommandList& RHICmdList, TRefCountPtr<IPooledRenderTarget>& Target, const TCHAR* Name, bool bWithAlpha);
 
 extern void UpdateHistory(
 	FRHICommandList& RHICmdList,
diff --git a/Engine/Source/Runtime/Renderer/Private/DistanceFieldScreenGridLighting.cpp b/Engine/Source/Runtime/Renderer/Private/DistanceFieldScreenGridLighting.cpp
--- a/Engine/Source/Runtime/Renderer/Private/DistanceFieldScreenGridLighting.cpp
+++ b/Engine/Source/Runtime/Renderer/Private/DistanceFieldScreenGridLighting.cpp
@@ -625,7 +625,7 @@ IMPLEMENT_AWARE_UPSAMPLE_PS_TYPE(false, true);
 IMPLEMENT_AWARE_UPSAMPLE_PS_TYPE(false, false);
 
 void PostProcessBentNormalAOScreenGrid(
-	FRHICommandList& RHICmdList, 
+	FRHICommandListImmediate& RHICmdList, 
 	const FDistanceFieldAOParameters& Parameters, 
 	const FViewInfo& View, 
 	IPooledRenderTarget* VelocityTexture,
@@ -640,11 +640,11 @@ void PostProcessBentNormalAOScreenGrid(
 
 	TRefCountPtr<IPooledRenderTarget> DistanceFieldAOBentNormal;
 	TRefCountPtr<IPooledRenderTarget> DistanceFieldIrradiance;
-	AllocateOrReuseAORenderTarget(DistanceFieldAOBentNormal, TEXT("DistanceFieldBentNormalAO"), true);
+	AllocateOrReuseAORenderTarget(RHICmdList, DistanceFieldAOBentNormal, TEXT("DistanceFieldBentNormalAO"), true);
 
 	if (bUseDistanceFieldGI)
 	{
-		AllocateOrReuseAORenderTarget(DistanceFieldIrradiance, TEXT("DistanceFieldIrradiance"), false);
+		AllocateOrReuseAORenderTarget(RHICmdList, DistanceFieldIrradiance, TEXT("DistanceFieldIrradiance"), false);
 	}
 
 	{
@@ -851,7 +851,7 @@ void FDeferredShadingSceneRenderer::RenderDistanceFieldAOScreenGrid(
 	{
 		{
 			FPooledRenderTargetDesc Desc(FPooledRenderTargetDesc::Create2DDesc(ConeTraceBufferSize, PF_FloatRGBA, FClearValueBinding::None, TexCreate_None, TexCreate_RenderTargetable | TexCreate_UAV, false));
-			GRenderTargetPool.FindFreeElement(Desc, DownsampledIrradiance, TEXT("DownsampledIrradiance"));
+			GRenderTargetPool.FindFreeElement(RHICmdList, Desc, DownsampledIrradiance, TEXT("DownsampledIrradiance"));
 		}
 
 		extern void ComputeIrradianceForScreenGrid(
@@ -873,7 +873,7 @@ void FDeferredShadingSceneRenderer::RenderDistanceFieldAOScreenGrid(
 
 	{
 		FPooledRenderTargetDesc Desc(FPooledRenderTargetDesc::Create2DDesc(ConeTraceBufferSize, PF_FloatRGBA, FClearValueBinding::None, TexCreate_None, TexCreate_RenderTargetable | TexCreate_UAV, false));
-		GRenderTargetPool.FindFreeElement(Desc, DownsampledBentNormal, TEXT("DownsampledBentNormal"));
+		GRenderTargetPool.FindFreeElement(RHICmdList, Desc, DownsampledBentNormal, TEXT("DownsampledBentNormal"));
 	}
 
 	{
diff --git a/Engine/Source/Runtime/Renderer/Private/DistanceFieldShadowing.cpp b/Engine/Source/Runtime/Renderer/Private/DistanceFieldShadowing.cpp
--- a/Engine/Source/Runtime/Renderer/Private/DistanceFieldShadowing.cpp
+++ b/Engine/Source/Runtime/Renderer/Private/DistanceFieldShadowing.cpp
@@ -826,7 +826,7 @@ void FProjectedShadowInfo::RenderRayTracedDistanceFieldProjection(FRHICommandLis
 			{
 				const FIntPoint BufferSize = GetBufferSizeForDFShadows();
 				FPooledRenderTargetDesc Desc(FPooledRenderTargetDesc::Create2DDesc(BufferSize, PF_G16R16F, FClearValueBinding::None, TexCreate_None, TexCreate_RenderTargetable | TexCreate_UAV, false));
-				GRenderTargetPool.FindFreeElement(Desc, RayTracedShadowsRT, TEXT("RayTracedShadows"));
+				GRenderTargetPool.FindFreeElement(RHICmdList, Desc, RayTracedShadowsRT, TEXT("RayTracedShadows"));
 			}
 
 			FIntRect ScissorRect;
diff --git a/Engine/Source/Runtime/Renderer/Private/DistanceFieldSurfaceCacheLighting.cpp b/Engine/Source/Runtime/Renderer/Private/DistanceFieldSurfaceCacheLighting.cpp
--- a/Engine/Source/Runtime/Renderer/Private/DistanceFieldSurfaceCacheLighting.cpp
+++ b/Engine/Source/Runtime/Renderer/Private/DistanceFieldSurfaceCacheLighting.cpp
@@ -2344,7 +2344,7 @@ void RenderIrradianceCacheInterpolation(
 		if (bFinalInterpolation && (GAOInterpolationDepthTesting || GAOInterpolationStencilTesting))
 		{
 			FPooledRenderTargetDesc Desc(FPooledRenderTargetDesc::Create2DDesc(BentNormalInterpolationTarget->GetDesc().Extent, PF_DepthStencil, FClearValueBinding::DepthZero, TexCreate_None, TexCreate_DepthStencilTargetable, false));
-			GRenderTargetPool.FindFreeElement(Desc, SplatDepthStencilBuffer, TEXT("DistanceFieldAOSplatDepthBuffer"));
+			GRenderTargetPool.FindFreeElement(RHICmdList, Desc, SplatDepthStencilBuffer, TEXT("DistanceFieldAOSplatDepthBuffer"));
 
 			SetupDepthStencil(RHICmdList, View, SplatDepthStencilBuffer->GetRenderTargetItem(), DistanceFieldNormal, DepthLevel, DestLevelDownsampleFactor);
 
@@ -2735,7 +2735,7 @@ void RenderDistanceFieldAOSurfaceCache(
 		{
 			FIntPoint AOBufferSize = FIntPoint::DivideAndRoundUp(SceneContext.GetBufferSizeXY(), DestLevelDownsampleFactor);
 			FPooledRenderTargetDesc Desc(FPooledRenderTargetDesc::Create2DDesc(AOBufferSize, PF_FloatRGBA, FClearValueBinding::Transparent, TexCreate_None, TexCreate_RenderTargetable | TexCreate_UAV, false));
-			GRenderTargetPool.FindFreeElement(Desc, DistanceFieldAOBentNormalSplat, TEXT("DistanceFieldAOBentNormalSplat"));
+			GRenderTargetPool.FindFreeElement(RHICmdList, Desc, DistanceFieldAOBentNormalSplat, TEXT("DistanceFieldAOBentNormalSplat"));
 		}
 
 		// Splat / interpolate the surface cache records onto the buffer sized for the current depth level
@@ -2848,11 +2848,11 @@ void RenderDistanceFieldAOSurfaceCache(
 	{
 		FIntPoint BufferSize = GetBufferSizeForAO();
 		FPooledRenderTargetDesc Desc(FPooledRenderTargetDesc::Create2DDesc(BufferSize, PF_FloatRGBA, FClearValueBinding::Transparent, TexCreate_None, TexCreate_RenderTargetable | TexCreate_UAV, false));
-		GRenderTargetPool.FindFreeElement(Desc, BentNormalAccumulation, TEXT("BentNormalAccumulation"));
+		GRenderTargetPool.FindFreeElement(RHICmdList, Desc, BentNormalAccumulation, TEXT("BentNormalAccumulation"));
 
 		if (bUseDistanceFieldGI)
 		{
-			GRenderTargetPool.FindFreeElement(Desc, IrradianceAccumulation, TEXT("IrradianceAccumulation"));
+			GRenderTargetPool.FindFreeElement(RHICmdList, Desc, IrradianceAccumulation, TEXT("IrradianceAccumulation"));
 		}
 	}
 
@@ -2943,7 +2943,7 @@ bool FDeferredShadingSceneRenderer::RenderDistanceFieldLighting(
 			{
 				const FIntPoint BufferSize = GetBufferSizeForAO();
 				FPooledRenderTargetDesc Desc(FPooledRenderTargetDesc::Create2DDesc(BufferSize, PF_FloatRGBA, FClearValueBinding::Transparent, TexCreate_None, TexCreate_RenderTargetable | TexCreate_UAV, false));
-				GRenderTargetPool.FindFreeElement(Desc, DistanceFieldNormal, TEXT("DistanceFieldNormal"));
+				GRenderTargetPool.FindFreeElement(RHICmdList, Desc, DistanceFieldNormal, TEXT("DistanceFieldNormal"));
 			}
 
 			ComputeDistanceFieldNormal(RHICmdList, Views, DistanceFieldNormal->GetRenderTargetItem(), Parameters);
@@ -3004,12 +3004,12 @@ bool FDeferredShadingSceneRenderer::RenderDistanceFieldLighting(
 				FPooledRenderTargetDesc Desc = SceneContext.GetSceneColor()->GetDesc();
 				// Make sure we get a signed format
 				Desc.Format = PF_FloatRGBA;
-				GRenderTargetPool.FindFreeElement(Desc, OutDynamicBentNormalAO, TEXT("DynamicBentNormalAO"));
+				GRenderTargetPool.FindFreeElement(RHICmdList, Desc, OutDynamicBentNormalAO, TEXT("DynamicBentNormalAO"));
 
 				if (bUseDistanceFieldGI)
 				{
 					Desc.Format = PF_FloatRGB;
-					GRenderTargetPool.FindFreeElement(Desc, OutDynamicIrradiance, TEXT("DynamicIrradiance"));
+					GRenderTargetPool.FindFreeElement(RHICmdList, Desc, OutDynamicIrradiance, TEXT("DynamicIrradiance"));
 				}
 
 				FTextureRHIParamRef RenderTargets[2] =
@@ -3429,7 +3429,7 @@ void FDeferredShadingSceneRenderer::RenderMeshDistanceFieldVisualization(FRHICom
 			{
 				const FIntPoint BufferSize = GetBufferSizeForAO();
 				FPooledRenderTargetDesc Desc(FPooledRenderTargetDesc::Create2DDesc(BufferSize, PF_FloatRGBA, FClearValueBinding::None, TexCreate_None, TexCreate_RenderTargetable | TexCreate_UAV, false));
-				GRenderTargetPool.FindFreeElement(Desc, VisualizeResultRT, TEXT("VisualizeDistanceField"));
+				GRenderTargetPool.FindFreeElement(RHICmdList, Desc, VisualizeResultRT, TEXT("VisualizeDistanceField"));
 			}
 
 			{
diff --git a/Engine/Source/Runtime/Renderer/Private/DistortionRendering.cpp b/Engine/Source/Runtime/Renderer/Private/DistortionRendering.cpp
--- a/Engine/Source/Runtime/Renderer/Private/DistortionRendering.cpp
+++ b/Engine/Source/Runtime/Renderer/Private/DistortionRendering.cpp
@@ -846,7 +846,7 @@ void FSceneRenderer::RenderDistortion(FRHICommandListImmediate& RHICmdList)
 		{
 			FPooledRenderTargetDesc Desc(FPooledRenderTargetDesc::Create2DDesc(SceneContext.GetBufferSizeXY(), PF_B8G8R8A8, FClearValueBinding::None, TexCreate_None, TexCreate_RenderTargetable, false));
 			Desc.Flags |= TexCreate_FastVRAM;
-			GRenderTargetPool.FindFreeElement(Desc, DistortionRT, TEXT("Distortion"));
+			GRenderTargetPool.FindFreeElement(RHICmdList, Desc, DistortionRT, TEXT("Distortion"));
 
 			// use RGBA8 light target for accumulating distortion offsets	
 			// R = positive X offset
@@ -900,7 +900,7 @@ void FSceneRenderer::RenderDistortion(FRHICommandListImmediate& RHICmdList)
 // OCULUS BEGIN: select ONE render target for all views (eyes)
 		TRefCountPtr<IPooledRenderTarget> NewSceneColor;
 		// we don't create a new name to make it easier to use "vis SceneColor" and get the last HDRSceneColor
-		GRenderTargetPool.FindFreeElement(SceneContext.GetSceneColor()->GetDesc(), NewSceneColor, TEXT("SceneColor"));
+		GRenderTargetPool.FindFreeElement(RHICmdList, SceneContext.GetSceneColor()->GetDesc(), NewSceneColor, TEXT("SceneColor"));
 		const FSceneRenderTargetItem& DestRenderTarget = NewSceneColor->GetRenderTargetItem();
 // OCULUS END
 
@@ -949,7 +949,7 @@ void FSceneRenderer::RenderDistortion(FRHICommandListImmediate& RHICmdList)
 
 		// Resolve and merge distortion to main scene
 		RHICmdList.CopyToResolveTarget(DestRenderTarget.TargetableTexture, DestRenderTarget.ShaderResourceTexture, false, FResolveParams());
-		SetRenderTarget(RHICmdList, SceneContext.GetSceneColor()->GetRenderTargetItem().TargetableTexture, SceneContext.GetSceneDepthSurface(), ESimpleRenderTargetMode::EExistingColorAndDepth, FExclusiveDepthStencil::DepthRead_StencilWrite);
+		SetRenderTarget(RHICmdList, SceneContext.GetSceneColor()->GetRenderTargetItem().TargetableTexture, SceneContext.GetSceneDepthSurface(), ESimpleRenderTargetMode::EExistingColorAndDepth, FExclusiveDepthStencil::DepthRead_StencilWrite, true);
 
 		for(int32 ViewIndex = 0, Num = Views.Num(); ViewIndex < Num; ++ViewIndex)
 		{
diff --git a/Engine/Source/Runtime/Renderer/Private/FogRendering.cpp b/Engine/Source/Runtime/Renderer/Private/FogRendering.cpp
--- a/Engine/Source/Runtime/Renderer/Private/FogRendering.cpp
+++ b/Engine/Source/Runtime/Renderer/Private/FogRendering.cpp
@@ -312,7 +312,7 @@ bool FDeferredShadingSceneRenderer::RenderFog(FRHICommandListImmediate& RHICmdLi
 		};
 		FSceneRenderTargets& SceneContext = FSceneRenderTargets::Get(RHICmdList);
 
-		SceneContext.BeginRenderingSceneColor(RHICmdList, ESimpleRenderTargetMode::EExistingColorAndDepth, FExclusiveDepthStencil::DepthRead_StencilWrite);
+		SceneContext.BeginRenderingSceneColor(RHICmdList, ESimpleRenderTargetMode::EExistingColorAndDepth, FExclusiveDepthStencil::DepthRead_StencilWrite, true);
 		for(int32 ViewIndex = 0;ViewIndex < Views.Num();ViewIndex++)
 		{
 			const FViewInfo& View = Views[ViewIndex];
diff --git a/Engine/Source/Runtime/Renderer/Private/ForwardShadingRenderer.cpp b/Engine/Source/Runtime/Renderer/Private/ForwardShadingRenderer.cpp
--- a/Engine/Source/Runtime/Renderer/Private/ForwardShadingRenderer.cpp
+++ b/Engine/Source/Runtime/Renderer/Private/ForwardShadingRenderer.cpp
@@ -100,7 +100,10 @@ void FForwardShadingSceneRenderer::Render(FRHICommandListImmediate& RHICmdList)
 	FSceneRenderTargets& SceneContext = FSceneRenderTargets::Get(RHICmdList);
 
 	// Allocate the maximum scene render target space for the current view family.
-	SceneContext.Allocate(ViewFamily);
+	SceneContext.Allocate(RHICmdList, ViewFamily);
+
+	//make sure all the targets we're going to use will be safely writable.
+	GRenderTargetPool.TransitionTargetsWritable(RHICmdList);
 
 	// Find the visible primitives.
 	InitViews(RHICmdList);
@@ -239,7 +242,7 @@ void FForwardShadingSceneRenderer::Render(FRHICommandListImmediate& RHICmdList)
 void FForwardShadingSceneRenderer::BasicPostProcess(FRHICommandListImmediate& RHICmdList, FViewInfo &View, bool bDoUpscale, bool bDoEditorPrimitives)
 {
 	FRenderingCompositePassContext CompositeContext(RHICmdList, View);
-	FPostprocessContext Context(CompositeContext.Graph, View);
+	FPostprocessContext Context(RHICmdList, CompositeContext.Graph, View);
 
 	if (bDoUpscale)
 	{	// simple bilinear upscaling for ES2.
diff --git a/Engine/Source/Runtime/Renderer/Private/GPUBenchmark.cpp b/Engine/Source/Runtime/Renderer/Private/GPUBenchmark.cpp
--- a/Engine/Source/Runtime/Renderer/Private/GPUBenchmark.cpp
+++ b/Engine/Source/Runtime/Renderer/Private/GPUBenchmark.cpp
@@ -314,14 +314,14 @@ void RendererGPUBenchmark(FRHICommandListImmediate& RHICmdList, FSynthBenchmarkR
 	TRefCountPtr<IPooledRenderTarget> RTItems[3];
 	{
 		FPooledRenderTargetDesc Desc(FPooledRenderTargetDesc::Create2DDesc(FIntPoint(GBenchmarkResolution, GBenchmarkResolution), PF_B8G8R8A8, FClearValueBinding::None, TexCreate_None, TexCreate_RenderTargetable | TexCreate_ShaderResource, false));
-		GRenderTargetPool.FindFreeElement(Desc, RTItems[0], TEXT("Benchmark0"));
-		GRenderTargetPool.FindFreeElement(Desc, RTItems[1], TEXT("Benchmark1"));
+		GRenderTargetPool.FindFreeElement(RHICmdList, Desc, RTItems[0], TEXT("Benchmark0"));
+		GRenderTargetPool.FindFreeElement(RHICmdList, Desc, RTItems[1], TEXT("Benchmark1"));
 
 		Desc.Extent = FIntPoint(1, 1);
 		Desc.Flags = TexCreate_CPUReadback;	// needs TexCreate_ResolveTargetable?
 		Desc.TargetableFlags = TexCreate_None;
 
-		GRenderTargetPool.FindFreeElement(Desc, RTItems[2], TEXT("BenchmarkReadback"));
+		GRenderTargetPool.FindFreeElement(RHICmdList, Desc, RTItems[2], TEXT("BenchmarkReadback"));
 	}
 
 	// set the state
diff --git a/Engine/Source/Runtime/Renderer/Private/GlobalDistanceField.cpp b/Engine/Source/Runtime/Renderer/Private/GlobalDistanceField.cpp
--- a/Engine/Source/Runtime/Renderer/Private/GlobalDistanceField.cpp
+++ b/Engine/Source/Runtime/Renderer/Private/GlobalDistanceField.cpp
@@ -512,7 +512,7 @@ void AddUpdateRegionForPrimitive(const FVector4& Bounds, float MaxSphereQueryRad
 	}
 }
 
-void AllocateClipmapTexture(int32 ClipmapIndex, TRefCountPtr<IPooledRenderTarget>& Texture)
+void AllocateClipmapTexture(FRHICommandListImmediate& RHICmdList, int32 ClipmapIndex, TRefCountPtr<IPooledRenderTarget>& Texture)
 {
 	const TCHAR* TextureName = TEXT("GlobalDistanceField0");
 
@@ -530,6 +530,7 @@ void AllocateClipmapTexture(int32 ClipmapIndex, TRefCountPtr<IPooledRenderTarget
 	}
 
 	GRenderTargetPool.FindFreeElement(
+		RHICmdList,
 		FPooledRenderTargetDesc(FPooledRenderTargetDesc::CreateVolumeDesc(
 			GAOGlobalDFResolution,
 			GAOGlobalDFResolution,
@@ -571,7 +572,7 @@ float ComputeClipmapExtent(int32 ClipmapIndex)
 	return GAOInnerGlobalDFClipmapDistance * FMath::Pow(GAOGlobalDFClipmapDistanceExponent, ClipmapIndex);
 }
 
-void ComputeUpdateRegionsAndUpdateViewState(const FViewInfo& View, FGlobalDistanceFieldInfo& GlobalDistanceFieldInfo, int32 NumClipmaps, const TArray<FVector4>& PrimitiveModifiedBounds, float MaxOcclusionDistance)
+void ComputeUpdateRegionsAndUpdateViewState(FRHICommandListImmediate& RHICmdList, const FViewInfo& View, FGlobalDistanceFieldInfo& GlobalDistanceFieldInfo, int32 NumClipmaps, const TArray<FVector4>& PrimitiveModifiedBounds, float MaxOcclusionDistance)
 {
 	GlobalDistanceFieldInfo.Clipmaps.AddZeroed(NumClipmaps);
 
@@ -593,7 +594,7 @@ void ComputeUpdateRegionsAndUpdateViewState(const FViewInfo& View, FGlobalDistan
 
 			if (!RenderTarget || RenderTarget->GetDesc().Extent.X != GAOGlobalDFResolution)
 			{
-				AllocateClipmapTexture(ClipmapIndex, RenderTarget);
+				AllocateClipmapTexture(RHICmdList, ClipmapIndex, RenderTarget);
 				bReallocated = true;
 			}
 
@@ -715,7 +716,7 @@ void ComputeUpdateRegionsAndUpdateViewState(const FViewInfo& View, FGlobalDistan
 		for (int32 ClipmapIndex = 0; ClipmapIndex < NumClipmaps; ClipmapIndex++)
 		{
 			FGlobalDistanceFieldClipmap& Clipmap = GlobalDistanceFieldInfo.Clipmaps[ClipmapIndex];
-			AllocateClipmapTexture(ClipmapIndex, Clipmap.RenderTarget);
+			AllocateClipmapTexture(RHICmdList, ClipmapIndex, Clipmap.RenderTarget);
 			Clipmap.ScrollOffset = FVector::ZeroVector;
 
 			const float Extent = ComputeClipmapExtent(ClipmapIndex);
@@ -757,7 +758,7 @@ void UpdateGlobalDistanceFieldVolume(
 {
 	if (Scene->DistanceFieldSceneData.NumObjectsInBuffer > 0)
 	{
-		ComputeUpdateRegionsAndUpdateViewState(View, GlobalDistanceFieldInfo, GMaxGlobalDistanceFieldClipmaps, Scene->DistanceFieldSceneData.PrimitiveModifiedBounds, MaxOcclusionDistance);
+		ComputeUpdateRegionsAndUpdateViewState(RHICmdList, View, GlobalDistanceFieldInfo, GMaxGlobalDistanceFieldClipmaps, Scene->DistanceFieldSceneData.PrimitiveModifiedBounds, MaxOcclusionDistance);
 
 		bool bHasUpdateRegions = false;
 
diff --git a/Engine/Source/Runtime/Renderer/Private/HeightfieldLighting.cpp b/Engine/Source/Runtime/Renderer/Private/HeightfieldLighting.cpp
--- a/Engine/Source/Runtime/Renderer/Private/HeightfieldLighting.cpp
+++ b/Engine/Source/Runtime/Renderer/Private/HeightfieldLighting.cpp
@@ -65,6 +65,7 @@ FAutoConsoleVariableRef CVarHeightfieldTargetUnitsPerTexel(
 
 void FHeightfieldLightingAtlas::InitDynamicRHI()
 {
+	FRHICommandListImmediate& RHICmdList = FRHICommandListExecutor::GetImmediateCommandList();
 	if (AtlasSize.GetMin() > 0)
 	{
 		FPooledRenderTargetDesc Desc(FPooledRenderTargetDesc::Create2DDesc(
@@ -75,7 +76,7 @@ void FHeightfieldLightingAtlas::InitDynamicRHI()
 			TexCreate_RenderTargetable,
 			false));
 
-		GRenderTargetPool.FindFreeElement(Desc, Height, TEXT("HeightAtlas"));
+		GRenderTargetPool.FindFreeElement(RHICmdList, Desc, Height, TEXT("HeightAtlas"));
 
 		FPooledRenderTargetDesc Desc2(FPooledRenderTargetDesc::Create2DDesc(
 			AtlasSize,
@@ -85,7 +86,7 @@ void FHeightfieldLightingAtlas::InitDynamicRHI()
 			TexCreate_RenderTargetable,
 			false));
 
-		GRenderTargetPool.FindFreeElement(Desc2, Normal, TEXT("NormalAtlas"));
+		GRenderTargetPool.FindFreeElement(RHICmdList, Desc2, Normal, TEXT("NormalAtlas"));
 
 		FPooledRenderTargetDesc Desc3(FPooledRenderTargetDesc::Create2DDesc(
 			AtlasSize,
@@ -95,7 +96,7 @@ void FHeightfieldLightingAtlas::InitDynamicRHI()
 			TexCreate_RenderTargetable,
 			false));
 
-		GRenderTargetPool.FindFreeElement(Desc3, DiffuseColor, TEXT("DiffuseColorAtlas"));
+		GRenderTargetPool.FindFreeElement(RHICmdList, Desc3, DiffuseColor, TEXT("DiffuseColorAtlas"));
 
 		FPooledRenderTargetDesc Desc4(FPooledRenderTargetDesc::Create2DDesc(
 			AtlasSize,
@@ -105,7 +106,7 @@ void FHeightfieldLightingAtlas::InitDynamicRHI()
 			TexCreate_RenderTargetable,
 			false));
 
-		GRenderTargetPool.FindFreeElement(Desc4, DirectionalLightShadowing, TEXT("HeightfieldShadowingAtlas"));
+		GRenderTargetPool.FindFreeElement(RHICmdList, Desc4, DirectionalLightShadowing, TEXT("HeightfieldShadowingAtlas"));
 
 		FPooledRenderTargetDesc Desc5(FPooledRenderTargetDesc::Create2DDesc(
 			AtlasSize,
@@ -115,7 +116,7 @@ void FHeightfieldLightingAtlas::InitDynamicRHI()
 			TexCreate_RenderTargetable,
 			false));
 
-		GRenderTargetPool.FindFreeElement(Desc5, Lighting, TEXT("HeightfieldLightingAtlas"));
+		GRenderTargetPool.FindFreeElement(RHICmdList, Desc5, Lighting, TEXT("HeightfieldLightingAtlas"));
 	}
 }
 
diff --git a/Engine/Source/Runtime/Renderer/Private/IndirectLightingCache.cpp b/Engine/Source/Runtime/Renderer/Private/IndirectLightingCache.cpp
--- a/Engine/Source/Runtime/Renderer/Private/IndirectLightingCache.cpp
+++ b/Engine/Source/Runtime/Renderer/Private/IndirectLightingCache.cpp
@@ -151,6 +151,7 @@ void FIndirectLightingCache::InitDynamicRHI()
 {
 	if (CanIndirectLightingCacheUseVolumeTexture(GetFeatureLevel()))
 	{
+		FRHICommandListImmediate& RHICmdList = FRHICommandListExecutor::GetImmediateCommandList();
 		uint32 Flags = TexCreate_ShaderResource | TexCreate_NoTiling;
 
 		FPooledRenderTargetDesc Desc(FPooledRenderTargetDesc::CreateVolumeDesc(
@@ -164,9 +165,9 @@ void FIndirectLightingCache::InitDynamicRHI()
 			false, 
 			1));
 
-		GRenderTargetPool.FindFreeElement(Desc, Texture0, TEXT("IndirectLightingCache_0"));
-		GRenderTargetPool.FindFreeElement(Desc, Texture1, TEXT("IndirectLightingCache_1"));
-		GRenderTargetPool.FindFreeElement(Desc, Texture2, TEXT("IndirectLightingCache_2"));
+		GRenderTargetPool.FindFreeElement(RHICmdList, Desc, Texture0, TEXT("IndirectLightingCache_0"));
+		GRenderTargetPool.FindFreeElement(RHICmdList, Desc, Texture1, TEXT("IndirectLightingCache_1"));
+		GRenderTargetPool.FindFreeElement(RHICmdList, Desc, Texture2, TEXT("IndirectLightingCache_2"));
 	}
 }
 
diff --git a/Engine/Source/Runtime/Renderer/Private/LightPropagationVolume.cpp b/Engine/Source/Runtime/Renderer/Private/LightPropagationVolume.cpp
--- a/Engine/Source/Runtime/Renderer/Private/LightPropagationVolume.cpp
+++ b/Engine/Source/Runtime/Renderer/Private/LightPropagationVolume.cpp
@@ -832,7 +832,7 @@ FLightPropagationVolume::~FLightPropagationVolume()
 /**
 * Sets up the LPV at the beginning of the frame
 */
-void FLightPropagationVolume::InitSettings(FRHICommandList& RHICmdList, const FSceneView& View)
+void FLightPropagationVolume::InitSettings(FRHICommandListImmediate& RHICmdList, const FSceneView& View)
 {
 	int32 NumFastLpvTextures = 7;
 	int32 NumFastGvTextures = 2;
@@ -865,7 +865,7 @@ void FLightPropagationVolume::InitSettings(FRHICommandList& RHICmdList, const FS
 			{
 				for ( int j = 0; j < 7; j++ )
 				{
-					GRenderTargetPool.FindFreeElement(Desc, LpvVolumeTextures[i][j], Names[j * 2 + i] );
+					GRenderTargetPool.FindFreeElement(RHICmdList, Desc, LpvVolumeTextures[i][j], Names[j * 2 + i] );
 				}
 			}
 		}
@@ -875,7 +875,7 @@ void FLightPropagationVolume::InitSettings(FRHICommandList& RHICmdList, const FS
 
 			for ( int i = 0; i < NUM_GV_TEXTURES; i++ )
 			{
-				GRenderTargetPool.FindFreeElement(Desc, GvVolumeTextures[i], Names[i]);
+				GRenderTargetPool.FindFreeElement(RHICmdList, Desc, GvVolumeTextures[i], Names[i]);
 			}
 		}
 
@@ -890,7 +890,7 @@ void FLightPropagationVolume::InitSettings(FRHICommandList& RHICmdList, const FS
 				TexCreate_ShaderResource | TexCreate_UAV | TexCreate_FastVRAM,
 				false,
 				1));
-			GRenderTargetPool.FindFreeElement(AODesc, AOVolumeTexture, TEXT("LPVAOVolume"));
+			GRenderTargetPool.FindFreeElement(RHICmdList, AODesc, AOVolumeTexture, TEXT("LPVAOVolume"));
 		}
 
 		bInitialized = true;
diff --git a/Engine/Source/Runtime/Renderer/Private/LightPropagationVolume.h b/Engine/Source/Runtime/Renderer/Private/LightPropagationVolume.h
--- a/Engine/Source/Runtime/Renderer/Private/LightPropagationVolume.h
+++ b/Engine/Source/Runtime/Renderer/Private/LightPropagationVolume.h
@@ -137,7 +137,7 @@ public:
 	FLightPropagationVolume();
 	virtual ~FLightPropagationVolume();
 
-	void InitSettings(FRHICommandList& RHICmdList, const FSceneView& View);
+	void InitSettings(FRHICommandListImmediate& RHICmdList, const FSceneView& View);
 
 	void Clear(FRHICommandListImmediate& RHICmdList, FViewInfo& View);
 
diff --git a/Engine/Source/Runtime/Renderer/Private/LightRendering.cpp b/Engine/Source/Runtime/Renderer/Private/LightRendering.cpp
--- a/Engine/Source/Runtime/Renderer/Private/LightRendering.cpp
+++ b/Engine/Source/Runtime/Renderer/Private/LightRendering.cpp
@@ -444,7 +444,7 @@ void FDeferredShadingSceneRenderer::RenderLights(FRHICommandListImmediate& RHICm
 				SCOPED_DRAW_EVENT(RHICmdList, StandardDeferredLighting);
 
 				// make sure we don't clear the depth
-				SceneContext.BeginRenderingSceneColor(RHICmdList, ESimpleRenderTargetMode::EExistingColorAndDepth, FExclusiveDepthStencil::DepthRead_StencilWrite);
+				SceneContext.BeginRenderingSceneColor(RHICmdList, ESimpleRenderTargetMode::EExistingColorAndDepth, FExclusiveDepthStencil::DepthRead_StencilWrite, true);
 
 				// Draw non-shadowed non-light function lights without changing render targets between them
 				for (int32 LightIndex = StandardDeferredStart; LightIndex < AttenuationLightStart; LightIndex++)
diff --git a/Engine/Source/Runtime/Renderer/Private/LightShaftRendering.cpp b/Engine/Source/Runtime/Renderer/Private/LightShaftRendering.cpp
--- a/Engine/Source/Runtime/Renderer/Private/LightShaftRendering.cpp
+++ b/Engine/Source/Runtime/Renderer/Private/LightShaftRendering.cpp
@@ -449,7 +449,7 @@ void AllocateOrReuseLightShaftRenderTarget(FRHICommandListImmediate& RHICmdList,
 		const FIntPoint BufferSize = FSceneRenderTargets::Get(RHICmdList).GetBufferSizeXY();
 		FIntPoint LightShaftSize(FMath::Max<uint32>(BufferSize.X / GetLightShaftDownsampleFactor(), 1), FMath::Max<uint32>(BufferSize.Y / GetLightShaftDownsampleFactor(), 1));
 		FPooledRenderTargetDesc Desc(FPooledRenderTargetDesc::Create2DDesc(LightShaftSize, LightShaftFilterBufferFormat, FClearValueBinding::Black, TexCreate_None, TexCreate_RenderTargetable, false));
-		GRenderTargetPool.FindFreeElement(Desc, Target, Name);
+		GRenderTargetPool.FindFreeElement(RHICmdList, Desc, Target, Name);
 
 		SetRenderTarget(RHICmdList, Target->GetRenderTargetItem().TargetableTexture, FTextureRHIRef(), ESimpleRenderTargetMode::EClearColorExistingDepth);
 	}
@@ -542,7 +542,7 @@ void ApplyTemporalAA(
 		{
 			FMemMark Mark(FMemStack::Get());
 			FRenderingCompositePassContext CompositeContext(RHICmdList, View);
-			FPostprocessContext Context(CompositeContext.Graph, View);
+			FPostprocessContext Context(RHICmdList, CompositeContext.Graph, View);
 
 			// Nodes for input render targets
 			FRenderingCompositePass* LightShaftSetup = Context.Graph.RegisterPass( new(FMemStack::Get()) FRCPassPostProcessInput( LightShaftsSource ) );
diff --git a/Engine/Source/Runtime/Renderer/Private/PostProcess/PostProcessBloomSetup.cpp b/Engine/Source/Runtime/Renderer/Private/PostProcess/PostProcessBloomSetup.cpp
--- a/Engine/Source/Runtime/Renderer/Private/PostProcess/PostProcessBloomSetup.cpp
+++ b/Engine/Source/Runtime/Renderer/Private/PostProcess/PostProcessBloomSetup.cpp
@@ -115,7 +115,7 @@ public:
 		{
 			if (Context.View.HasValidEyeAdaptation())
 			{
-				IPooledRenderTarget* EyeAdaptationRT = Context.View.GetEyeAdaptation();
+				IPooledRenderTarget* EyeAdaptationRT = Context.View.GetEyeAdaptation(Context.RHICmdList);
 				SetTextureParameter(Context.RHICmdList, ShaderRHI, EyeAdaptation, EyeAdaptationRT->GetRenderTargetItem().TargetableTexture);
 			}
 			else
diff --git a/Engine/Source/Runtime/Renderer/Private/PostProcess/PostProcessCompositeEditorPrimitives.cpp b/Engine/Source/Runtime/Renderer/Private/PostProcess/PostProcessCompositeEditorPrimitives.cpp
--- a/Engine/Source/Runtime/Renderer/Private/PostProcess/PostProcessCompositeEditorPrimitives.cpp
+++ b/Engine/Source/Runtime/Renderer/Private/PostProcess/PostProcessCompositeEditorPrimitives.cpp
@@ -204,8 +204,12 @@ void FRCPassPostProcessCompositeEditorPrimitives::Process(FRenderingCompositePas
 	bool bClearIsNeeded = !IsValidRef(SceneContext.EditorPrimitivesColor);
 
 	// Get or create the msaa depth and color buffer
-	FTexture2DRHIRef ColorTarget = SceneContext.GetEditorPrimitivesColor();
-	FTexture2DRHIRef DepthTarget = SceneContext.GetEditorPrimitivesDepth();
+	FTexture2DRHIRef ColorTarget = SceneContext.GetEditorPrimitivesColor(Context.RHICmdList);
+	FTexture2DRHIRef DepthTarget = SceneContext.GetEditorPrimitivesDepth(Context.RHICmdList);
+
+	FTextureRHIParamRef EditorRenderTargets[2];
+	EditorRenderTargets[0] = ColorTarget;
+	EditorRenderTargets[1] = DepthTarget;
 
 	const uint32 MSAASampleCount = SceneContext.EditorPrimitivesColor->GetDesc().NumSamples;
 
@@ -234,6 +238,7 @@ void FRCPassPostProcessCompositeEditorPrimitives::Process(FRenderingCompositePas
 		}
 
 		GRenderTargetPool.VisualizeTexture.SetCheckPoint(Context.RHICmdList, SceneContext.EditorPrimitivesColor);
+		Context.RHICmdList.TransitionResources(EResourceTransitionAccess::EReadable, EditorRenderTargets, 2);
 	}
 
 	//	Context.RHICmdList.CopyToResolveTarget(SceneContext.EditorPrimitivesColor->GetRenderTargetItem().TargetableTexture, SceneContext.EditorPrimitivesColor->GetRenderTargetItem().ShaderResourceTexture, false, FResolveParams());
diff --git a/Engine/Source/Runtime/Renderer/Private/PostProcess/PostProcessEyeAdaptation.cpp b/Engine/Source/Runtime/Renderer/Private/PostProcess/PostProcessEyeAdaptation.cpp
--- a/Engine/Source/Runtime/Renderer/Private/PostProcess/PostProcessEyeAdaptation.cpp
+++ b/Engine/Source/Runtime/Renderer/Private/PostProcess/PostProcessEyeAdaptation.cpp
@@ -76,14 +76,14 @@ void FRCPassPostProcessEyeAdaptation::Process(FRenderingCompositePassContext& Co
 	const FSceneView& View = Context.View;
 	const FSceneViewFamily& ViewFamily = *(View.Family);
 
-	IPooledRenderTarget* EyeAdaptation = Context.View.GetEyeAdaptation();
+	IPooledRenderTarget* EyeAdaptation = Context.View.GetEyeAdaptation(Context.RHICmdList);
 	check(EyeAdaptation);
 
 	FIntPoint DestSize = EyeAdaptation->GetDesc().Extent;
 
 	// we render to our own output render target, not the intermediate one created by the compositing system
 	// Set the view family's render target/viewport.
-	SetRenderTarget(Context.RHICmdList, EyeAdaptation->GetRenderTargetItem().TargetableTexture, FTextureRHIRef());
+	SetRenderTarget(Context.RHICmdList, EyeAdaptation->GetRenderTargetItem().TargetableTexture, FTextureRHIRef(), true);
 	Context.SetViewportAndCallRHI(0, 0, 0.0f, DestSize.X, DestSize.Y, 1.0f );
 
 	// set the state
diff --git a/Engine/Source/Runtime/Renderer/Private/PostProcess/PostProcessGBufferHints.cpp b/Engine/Source/Runtime/Renderer/Private/PostProcess/PostProcessGBufferHints.cpp
--- a/Engine/Source/Runtime/Renderer/Private/PostProcess/PostProcessGBufferHints.cpp
+++ b/Engine/Source/Runtime/Renderer/Private/PostProcess/PostProcessGBufferHints.cpp
@@ -78,10 +78,10 @@ public:
 IMPLEMENT_SHADER_TYPE(,FPostProcessGBufferHintsPS,TEXT("PostProcessGBufferHints"),TEXT("MainPS"),SF_Pixel);
 
 
-FRCPassPostProcessGBufferHints::FRCPassPostProcessGBufferHints()
+FRCPassPostProcessGBufferHints::FRCPassPostProcessGBufferHints(FRHICommandList& RHICmdList)
 {
 	// AdjustGBufferRefCount(-1) call is done when the pass gets executed
-	FSceneRenderTargets::Get_Todo_PassContext().AdjustGBufferRefCount(1);
+	FSceneRenderTargets::Get_Todo_PassContext().AdjustGBufferRefCount(RHICmdList, 1);
 }
 
 
@@ -183,7 +183,7 @@ void FRCPassPostProcessGBufferHints::Process(FRenderingCompositePassContext& Con
 	Context.RHICmdList.CopyToResolveTarget(DestRenderTarget.TargetableTexture, DestRenderTarget.ShaderResourceTexture, false, FResolveParams());
 
 	// AdjustGBufferRefCount(1) call is done in constructor
-	FSceneRenderTargets::Get(Context.RHICmdList).AdjustGBufferRefCount(-1);
+	FSceneRenderTargets::Get(Context.RHICmdList).AdjustGBufferRefCount(Context.RHICmdList, -1);
 }
 
 FPooledRenderTargetDesc FRCPassPostProcessGBufferHints::ComputeOutputDesc(EPassOutputId InPassOutputId) const
diff --git a/Engine/Source/Runtime/Renderer/Private/PostProcess/PostProcessGBufferHints.h b/Engine/Source/Runtime/Renderer/Private/PostProcess/PostProcessGBufferHints.h
--- a/Engine/Source/Runtime/Renderer/Private/PostProcess/PostProcessGBufferHints.h
+++ b/Engine/Source/Runtime/Renderer/Private/PostProcess/PostProcessGBufferHints.h
@@ -15,7 +15,7 @@ class FRCPassPostProcessGBufferHints : public TRenderingCompositePassBase<2, 1>
 {
 public:
 	// constructor
-	FRCPassPostProcessGBufferHints();
+	FRCPassPostProcessGBufferHints(FRHICommandList& RHICmdList);
 
 	// interface FRenderingCompositePass ---------
 	virtual void Process(FRenderingCompositePassContext& Context) override;
diff --git a/Engine/Source/Runtime/Renderer/Private/PostProcess/PostProcessHistogramReduce.cpp b/Engine/Source/Runtime/Renderer/Private/PostProcess/PostProcessHistogramReduce.cpp
--- a/Engine/Source/Runtime/Renderer/Private/PostProcess/PostProcessHistogramReduce.cpp
+++ b/Engine/Source/Runtime/Renderer/Private/PostProcess/PostProcessHistogramReduce.cpp
@@ -60,7 +60,7 @@ public:
 		{
 			if (Context.View.HasValidEyeAdaptation())
 			{
-				IPooledRenderTarget* EyeAdaptationRT = Context.View.GetEyeAdaptation();
+				IPooledRenderTarget* EyeAdaptationRT = Context.View.GetEyeAdaptation(Context.RHICmdList);
 				SetTextureParameter(Context.RHICmdList, ShaderRHI, EyeAdaptationTexture, EyeAdaptationRT->GetRenderTargetItem().TargetableTexture);
 			}
 			else
diff --git a/Engine/Source/Runtime/Renderer/Private/PostProcess/PostProcessMaterial.cpp b/Engine/Source/Runtime/Renderer/Private/PostProcess/PostProcessMaterial.cpp
--- a/Engine/Source/Runtime/Renderer/Private/PostProcess/PostProcessMaterial.cpp
+++ b/Engine/Source/Runtime/Renderer/Private/PostProcess/PostProcessMaterial.cpp
@@ -228,7 +228,7 @@ void FRCPassPostProcessMaterial::Process(FRenderingCompositePassContext& Context
 
 	if(Material->NeedsGBuffer())
 	{
-		FSceneRenderTargets::Get(Context.RHICmdList).AdjustGBufferRefCount(-1);
+		FSceneRenderTargets::Get(Context.RHICmdList).AdjustGBufferRefCount(Context.RHICmdList,-1);
 	}
 }
 
diff --git a/Engine/Source/Runtime/Renderer/Private/PostProcess/PostProcessMotionBlur.cpp b/Engine/Source/Runtime/Renderer/Private/PostProcess/PostProcessMotionBlur.cpp
--- a/Engine/Source/Runtime/Renderer/Private/PostProcess/PostProcessMotionBlur.cpp
+++ b/Engine/Source/Runtime/Renderer/Private/PostProcess/PostProcessMotionBlur.cpp
@@ -1015,7 +1015,7 @@ void FRCPassPostProcessVelocityScatter::Process(FRenderingCompositePassContext&
 
 	TRefCountPtr<IPooledRenderTarget> DepthTarget;
 	FPooledRenderTargetDesc Desc( FPooledRenderTargetDesc::Create2DDesc( DestSize, PF_ShadowDepth, FClearValueBinding::DepthOne, TexCreate_None, TexCreate_DepthStencilTargetable, false ) );
-	GRenderTargetPool.FindFreeElement( Desc, DepthTarget, TEXT("VelocityScatterDepth") );
+	GRenderTargetPool.FindFreeElement(Context.RHICmdList, Desc, DepthTarget, TEXT("VelocityScatterDepth") );
 
 	// Set the view family's render target/viewport.
 	FRHIRenderTargetView ColorView(DestRenderTarget.TargetableTexture, 0, -1, ERenderTargetLoadAction::ELoad, ERenderTargetStoreAction::EStore);
diff --git a/Engine/Source/Runtime/Renderer/Private/PostProcess/PostProcessSubsurface.cpp b/Engine/Source/Runtime/Renderer/Private/PostProcess/PostProcessSubsurface.cpp
--- a/Engine/Source/Runtime/Renderer/Private/PostProcess/PostProcessSubsurface.cpp
+++ b/Engine/Source/Runtime/Renderer/Private/PostProcess/PostProcessSubsurface.cpp
@@ -188,10 +188,10 @@ void SetSubsurfaceVisualizeShader(const FRenderingCompositePassContext& Context)
 	VertexShader->SetParameters(Context);
 }
 
-FRCPassPostProcessSubsurfaceVisualize::FRCPassPostProcessSubsurfaceVisualize()
+FRCPassPostProcessSubsurfaceVisualize::FRCPassPostProcessSubsurfaceVisualize(FRHICommandList& RHICmdList)
 {
 	// we need the GBuffer, we release it Process()
-	FSceneRenderTargets::Get_Todo_PassContext().AdjustGBufferRefCount(1);
+	FSceneRenderTargets::Get_Todo_PassContext().AdjustGBufferRefCount(RHICmdList, 1);
 }
 
 void FRCPassPostProcessSubsurfaceVisualize::Process(FRenderingCompositePassContext& Context)
@@ -296,7 +296,7 @@ void FRCPassPostProcessSubsurfaceVisualize::Process(FRenderingCompositePassConte
 	Context.RHICmdList.CopyToResolveTarget(DestRenderTarget.TargetableTexture, DestRenderTarget.ShaderResourceTexture, false, FResolveParams());
 
 	// we no longer need the GBuffer
-	SceneContext.AdjustGBufferRefCount(-1);
+	SceneContext.AdjustGBufferRefCount(Context.RHICmdList, -1);
 }
 
 FPooledRenderTargetDesc FRCPassPostProcessSubsurfaceVisualize::ComputeOutputDesc(EPassOutputId InPassOutputId) const
diff --git a/Engine/Source/Runtime/Renderer/Private/PostProcess/PostProcessSubsurface.h b/Engine/Source/Runtime/Renderer/Private/PostProcess/PostProcessSubsurface.h
--- a/Engine/Source/Runtime/Renderer/Private/PostProcess/PostProcessSubsurface.h
+++ b/Engine/Source/Runtime/Renderer/Private/PostProcess/PostProcessSubsurface.h
@@ -15,7 +15,7 @@
 class FRCPassPostProcessSubsurfaceVisualize : public TRenderingCompositePassBase<1, 1>
 {
 public:
-	FRCPassPostProcessSubsurfaceVisualize();
+	FRCPassPostProcessSubsurfaceVisualize(FRHICommandList& RHICmdList);
 
 	// interface FRenderingCompositePass ---------
 	virtual void Process(FRenderingCompositePassContext& Context) override;
diff --git a/Engine/Source/Runtime/Renderer/Private/PostProcess/PostProcessTemporalAA.cpp b/Engine/Source/Runtime/Renderer/Private/PostProcess/PostProcessTemporalAA.cpp
--- a/Engine/Source/Runtime/Renderer/Private/PostProcess/PostProcessTemporalAA.cpp
+++ b/Engine/Source/Runtime/Renderer/Private/PostProcess/PostProcessTemporalAA.cpp
@@ -289,7 +289,7 @@ FPooledRenderTargetDesc FRCPassPostProcessSSRTemporalAA::ComputeOutputDesc(EPass
 	FPooledRenderTargetDesc Ret = GetInput(ePId_Input0)->GetOutput()->RenderTargetDesc;
 
 	Ret.DebugName = TEXT("SSRTemporalAA");
-
+	Ret.AutoWritable = false;
 	return Ret;
 }
 
@@ -756,6 +756,7 @@ FPooledRenderTargetDesc FRCPassPostProcessTemporalAA::ComputeOutputDesc(EPassOut
 	Ret.Reset();
 	//Ret.Format = PF_FloatRGBA;
 	Ret.DebugName = TEXT("TemporalAA");
+	Ret.AutoWritable = false;
 
 	return Ret;
 }
diff --git a/Engine/Source/Runtime/Renderer/Private/PostProcess/PostProcessTonemap.h b/Engine/Source/Runtime/Renderer/Private/PostProcess/PostProcessTonemap.h
--- a/Engine/Source/Runtime/Renderer/Private/PostProcess/PostProcessTonemap.h
+++ b/Engine/Source/Runtime/Renderer/Private/PostProcess/PostProcessTonemap.h
@@ -129,7 +129,12 @@ public:
 		{
 			if (Context.View.HasValidEyeAdaptation())
 			{
-				IPooledRenderTarget* EyeAdaptationRT = Context.View.GetEyeAdaptation();
+				IPooledRenderTarget* EyeAdaptationRT = Context.View.GetEyeAdaptation(Context.RHICmdList);
+				FTextureRHIParamRef EyeAdaptationRTRef = EyeAdaptationRT->GetRenderTargetItem().TargetableTexture;
+				if (EyeAdaptationRTRef)
+				{
+					Context.RHICmdList.TransitionResources(EResourceTransitionAccess::EReadable, &EyeAdaptationRTRef, 1);
+				}
 				SetTextureParameter(Context.RHICmdList, ShaderRHI, EyeAdaptation, EyeAdaptationRT->GetRenderTargetItem().TargetableTexture);
 			}
 			else
diff --git a/Engine/Source/Runtime/Renderer/Private/PostProcess/PostProcessing.cpp b/Engine/Source/Runtime/Renderer/Private/PostProcess/PostProcessing.cpp
--- a/Engine/Source/Runtime/Renderer/Private/PostProcess/PostProcessing.cpp
+++ b/Engine/Source/Runtime/Renderer/Private/PostProcess/PostProcessing.cpp
@@ -148,8 +148,9 @@ IMPLEMENT_SHADER_TYPE(,FPostProcessVS,TEXT("PostProcessBloom"),TEXT("MainPostpro
 
 // -------------------------------------------------------
 
-FPostprocessContext::FPostprocessContext(class FRenderingCompositionGraph& InGraph, const FViewInfo& InView)
-	: Graph(InGraph)
+FPostprocessContext::FPostprocessContext(FRHICommandListImmediate& InRHICmdList, FRenderingCompositionGraph& InGraph, const FViewInfo& InView)
+: RHICmdList(InRHICmdList)
+	, Graph(InGraph)
 	, View(InView)
 	, SceneColor(0)
 	, SceneDepth(0)
@@ -735,7 +736,7 @@ static FRenderingCompositePass* AddSinglePostProcessMaterial(FPostprocessContext
 		if(Material->NeedsGBuffer())
 		{
 			// AdjustGBufferRefCount(-1) call is done when the pass gets executed
-			FSceneRenderTargets::Get_Todo_PassContext().AdjustGBufferRefCount(1);
+			FSceneRenderTargets::Get_Todo_PassContext().AdjustGBufferRefCount(Context.RHICmdList, 1);
 		}
 
 		FRenderingCompositePass* Node = Context.Graph.RegisterPass(new(FMemStack::Get()) FRCPassPostProcessMaterial(MaterialInterface, Context.View.GetFeatureLevel()));
@@ -805,7 +806,7 @@ static void AddPostProcessMaterial(FPostprocessContext& Context, EBlendableLocat
 		if(Material->NeedsGBuffer())
 		{
 			// AdjustGBufferRefCount(-1) call is done when the pass gets executed
-			FSceneRenderTargets::Get_Todo_PassContext().AdjustGBufferRefCount(1);
+			FSceneRenderTargets::Get_Todo_PassContext().AdjustGBufferRefCount(Context.RHICmdList, 1);
 		}
 
 		FRenderingCompositePass* Node = Context.Graph.RegisterPass(new(FMemStack::Get()) FRCPassPostProcessMaterial(MaterialInterface,FeatureLevel));
@@ -865,7 +866,7 @@ static void AddHighResScreenshotMask(FPostprocessContext& Context, FRenderingCom
 		if (RendererMaterial->NeedsGBuffer())
 		{
 			// AdjustGBufferRefCount(-1) call is done when the pass gets executed
-			FSceneRenderTargets::Get_Todo_PassContext().AdjustGBufferRefCount(1);
+			FSceneRenderTargets::Get_Todo_PassContext().AdjustGBufferRefCount(Context.RHICmdList, 1);
 		}
 	}
 }
@@ -915,7 +916,7 @@ static void AddGBufferVisualizationOverview(FPostprocessContext& Context, FRende
 					if (Material->NeedsGBuffer())
 					{
 						// AdjustGBufferRefCount(-1) call is done when the pass gets executed
-						FSceneRenderTargets::Get_Todo_PassContext().AdjustGBufferRefCount(1);
+						FSceneRenderTargets::Get_Todo_PassContext().AdjustGBufferRefCount(Context.RHICmdList, 1);
 					}
 
 					if (BaseFilename.Len())
@@ -1034,7 +1035,7 @@ void FPostProcessing::Process(FRHICommandListImmediate& RHICmdList, FViewInfo& V
 		FMemMark Mark(FMemStack::Get());
 		FRenderingCompositePassContext CompositeContext(RHICmdList, View);
 
-		FPostprocessContext Context(CompositeContext.Graph, View);
+		FPostprocessContext Context(RHICmdList, CompositeContext.Graph, View);
 		
 		// not always valid
 		FRenderingCompositeOutputRef HDRColor; 
@@ -1570,7 +1571,7 @@ void FPostProcessing::Process(FRHICommandListImmediate& RHICmdList, FViewInfo& V
 
 		if (View.Family->EngineShowFlags.GBufferHints && FeatureLevel >= ERHIFeatureLevel::SM4)
 		{
-			FRenderingCompositePass* Node = Context.Graph.RegisterPass(new(FMemStack::Get()) FRCPassPostProcessGBufferHints());
+			FRenderingCompositePass* Node = Context.Graph.RegisterPass(new(FMemStack::Get()) FRCPassPostProcessGBufferHints(RHICmdList));
 			Node->SetInput(ePId_Input0, FRenderingCompositeOutputRef(Context.FinalOutput));
 			// Ideally without lighting as we want the emissive, we should do that later.
 			Node->SetInput(ePId_Input1, FRenderingCompositeOutputRef(Context.SceneColor));
@@ -1587,7 +1588,7 @@ void FPostProcessing::Process(FRHICommandListImmediate& RHICmdList, FViewInfo& V
 		if (View.Family->EngineShowFlags.VisualizeSSS)
 		{
 			// the setup pass also does visualization, based on EngineShowFlags.VisualizeSSS
-			FRenderingCompositePass* PassVisualize = Context.Graph.RegisterPass(new(FMemStack::Get()) FRCPassPostProcessSubsurfaceVisualize());
+			FRenderingCompositePass* PassVisualize = Context.Graph.RegisterPass(new(FMemStack::Get()) FRCPassPostProcessSubsurfaceVisualize(RHICmdList));
 			PassVisualize->SetInput(ePId_Input0, Context.FinalOutput);
 			Context.FinalOutput = FRenderingCompositeOutputRef(PassVisualize);
 		}
@@ -1674,7 +1675,7 @@ void FPostProcessing::Process(FRHICommandListImmediate& RHICmdList, FViewInfo& V
 		{
 			// Generally we no longer need the GBuffers, anyone that wants to keep the GBuffers for longer should have called AdjustGBufferRefCount(1) to keep it for longer
 			// and call AdjustGBufferRefCount(-1) once it's consumed. This needs to happen each frame. PostProcessMaterial do that automatically
-			FSceneRenderTargets::Get_Todo_PassContext().AdjustGBufferRefCount(-1);
+			FSceneRenderTargets::Get_Todo_PassContext().AdjustGBufferRefCount(RHICmdList, -1);
 		}
 
 		// The graph setup should be finished before this line ----------------------------------------
@@ -1732,7 +1733,7 @@ void FPostProcessing::ProcessES2(FRHICommandListImmediate& RHICmdList, FViewInfo
 		FMemMark Mark(FMemStack::Get());
 		FRenderingCompositePassContext CompositeContext(RHICmdList, View);
 
-		FPostprocessContext Context(CompositeContext.Graph, View);
+		FPostprocessContext Context(RHICmdList, CompositeContext.Graph, View);
 		FRenderingCompositeOutputRef BloomOutput;
 		FRenderingCompositeOutputRef DofOutput;
 
diff --git a/Engine/Source/Runtime/Renderer/Private/PostProcess/PostProcessing.h b/Engine/Source/Runtime/Renderer/Private/PostProcess/PostProcessing.h
--- a/Engine/Source/Runtime/Renderer/Private/PostProcess/PostProcessing.h
+++ b/Engine/Source/Runtime/Renderer/Private/PostProcess/PostProcessing.h
@@ -15,8 +15,9 @@ class FPostprocessContext
 {
 public:
 
-	FPostprocessContext(FRenderingCompositionGraph& InGraph, const FViewInfo& InView);
+	FPostprocessContext(FRHICommandListImmediate& InRHICmdList, FRenderingCompositionGraph& InGraph, const FViewInfo& InView);
 
+	FRHICommandListImmediate& RHICmdList;
 	FRenderingCompositionGraph& Graph;
 	const FViewInfo& View;
 
diff --git a/Engine/Source/Runtime/Renderer/Private/PostProcess/RenderTargetPool.cpp b/Engine/Source/Runtime/Renderer/Private/PostProcess/RenderTargetPool.cpp
--- a/Engine/Source/Runtime/Renderer/Private/PostProcess/RenderTargetPool.cpp
+++ b/Engine/Source/Runtime/Renderer/Private/PostProcess/RenderTargetPool.cpp
@@ -159,7 +159,52 @@ static void LogVRamUsage(FPooledRenderTarget& Ref)
 	}
 }
 
-bool FRenderTargetPool::FindFreeElement(const FPooledRenderTargetDesc& Desc, TRefCountPtr<IPooledRenderTarget> &Out, const TCHAR* InDebugName)
+void FRenderTargetPool::TransitionTargetsWritable(FRHICommandListImmediate& RHICmdList)
+{
+	QUICK_SCOPE_CYCLE_COUNTER(STAT_RenderTargetPoolTransition);
+	check(IsInRenderingThread());
+	WaitForTransitionFence();
+	
+	TransitionTargets.Reset();	
+
+	for (int32 i = 0; i < PooledRenderTargets.Num(); ++i)
+	{
+		FPooledRenderTarget* PooledRT = PooledRenderTargets[i];
+		if (PooledRT && PooledRT->GetDesc().AutoWritable)
+		{
+			FTextureRHIParamRef RenderTarget = PooledRT->GetRenderTargetItem().TargetableTexture;
+			if (RenderTarget)
+			{				
+				TransitionTargets.Add(RenderTarget);
+			}
+		}
+	}
+
+	if (TransitionTargets.Num() > 0)
+	{
+		RHICmdList.TransitionResourceArrayNoCopy(EResourceTransitionAccess::EWritable, TransitionTargets);
+		if (GRHIThread)
+		{
+			TransitionFence = RHICmdList.RHIThreadFence(false);
+		}
+	}
+}
+
+void FRenderTargetPool::WaitForTransitionFence()
+{
+	QUICK_SCOPE_CYCLE_COUNTER(STAT_RenderTargetPoolTransitionWait);
+	check(IsInRenderingThread());
+	if (TransitionFence)
+	{		
+		check(IsInRenderingThread());		
+		FRHICommandListExecutor::WaitOnRHIThreadFence(TransitionFence);
+		TransitionFence = nullptr;		
+	}
+	TransitionTargets.Reset();
+	DeferredDeleteArray.Reset();
+}
+
+bool FRenderTargetPool::FindFreeElement(FRHICommandList& RHICmdList, const FPooledRenderTargetDesc& Desc, TRefCountPtr<IPooledRenderTarget> &Out, const TCHAR* InDebugName, bool bDoWritableBarrier)
 {
 	check(IsInRenderingThread());
 
@@ -205,7 +250,7 @@ bool FRenderTargetPool::FindFreeElement(const FPooledRenderTargetDesc& Desc, TRe
 
 	FPooledRenderTarget* Found = 0;
 	uint32 FoundIndex = -1;
-
+	bool bReusingExistingTarget = false;
 	// try to find a suitable element in the pool
 	{
 		//don't spend time doing 2 passes if the platform doesn't support fastvram
@@ -226,6 +271,7 @@ bool FRenderTargetPool::FindFreeElement(const FPooledRenderTargetDesc& Desc, TRe
 					check(!Element->IsSnapshot());
 					Found = Element;
 					FoundIndex = i;
+					bReusingExistingTarget = true;
 					break;
 				}
 			}
@@ -388,8 +434,6 @@ bool FRenderTargetPool::FindFreeElement(const FPooledRenderTargetDesc& Desc, TRe
 #if !(UE_BUILD_SHIPPING || UE_BUILD_TEST)
 	{
 		
-		FRHICommandListImmediate& RHICmdList = FRHICommandListExecutor::GetImmediateCommandList();
-
 		if(CVarRenderTargetPoolTest.GetValueOnRenderThread())
 		{
 			if(Found->GetDesc().TargetableFlags & TexCreate_RenderTargetable)
@@ -425,6 +469,10 @@ bool FRenderTargetPool::FindFreeElement(const FPooledRenderTargetDesc& Desc, TRe
 
 	check(!Found->IsFree());
 
+	if (bReusingExistingTarget && bDoWritableBarrier)
+	{
+		RHICmdList.TransitionResource(EResourceTransitionAccess::EWritable, Found->GetRenderTargetItem().TargetableTexture);
+	}
 	return false;
 }
 
@@ -1028,6 +1076,7 @@ void FRenderTargetPool::AddAllocEventsFromCurrentState()
 void FRenderTargetPool::TickPoolElements()
 {
 	check(IsInRenderingThread());
+	WaitForTransitionFence();
 
 	if(bStartEventRecordingNextTick)
 	{
@@ -1172,6 +1221,7 @@ void FRenderTargetPool::FreeUnusedResource(TRefCountPtr<IPooledRenderTarget>& In
 			AllocationLevelInKB -= ComputeSizeInKB(*Element);
 			// we assume because of reference counting the resource gets released when not needed any more
 			// we don't use Remove() to not shuffle around the elements for better transparency on RenderTargetPoolEvents
+			DeferredDeleteArray.Add(PooledRenderTargets[Index]);
 			PooledRenderTargets[Index] = 0;
 
 			In.SafeRelease();
@@ -1195,6 +1245,7 @@ void FRenderTargetPool::FreeUnusedResources()
 			AllocationLevelInKB -= ComputeSizeInKB(*Element);
 			// we assume because of reference counting the resource gets released when not needed any more
 			// we don't use Remove() to not shuffle around the elements for better transparency on RenderTargetPoolEvents
+			DeferredDeleteArray.Add(PooledRenderTargets[i]);
 			PooledRenderTargets[i] = 0;
 		}
 	}
@@ -1284,6 +1335,8 @@ const FPooledRenderTargetDesc& FPooledRenderTarget::GetDesc() const
 void FRenderTargetPool::ReleaseDynamicRHI()
 {
 	check(IsInRenderingThread());
+	WaitForTransitionFence();
+
 	PooledRenderTargets.Empty();
 	if (PooledRenderTargetSnapshots.Num())
 	{
diff --git a/Engine/Source/Runtime/Renderer/Private/PostProcess/RenderTargetPool.h b/Engine/Source/Runtime/Renderer/Private/PostProcess/RenderTargetPool.h
--- a/Engine/Source/Runtime/Renderer/Private/PostProcess/RenderTargetPool.h
+++ b/Engine/Source/Runtime/Renderer/Private/PostProcess/RenderTargetPool.h
@@ -188,13 +188,16 @@ class FRenderTargetPool : public FRenderResource
 public:
 	FRenderTargetPool();
 
+	/** Transitions all targets in the pool to writable. */
+	void TransitionTargetsWritable(FRHICommandListImmediate& RHICmdList);
+
 	/**
 	 * @param DebugName must not be 0, we only store the pointer
 	 * @param Out is not the return argument to avoid double allocation because of wrong reference counting
 	 * call from RenderThread only
 	 * @return true if the old element was still valid, false if a new one was assigned
 	 */
-	bool FindFreeElement(const FPooledRenderTargetDesc& Desc, TRefCountPtr<IPooledRenderTarget>& Out, const TCHAR* InDebugName);
+	bool FindFreeElement(FRHICommandList& RHICmdList, const FPooledRenderTargetDesc& Desc, TRefCountPtr<IPooledRenderTarget>& Out, const TCHAR* InDebugName, bool bDoWritableBarrier = true);
 
 	void CreateUntrackedElement(const FPooledRenderTargetDesc& Desc, TRefCountPtr<IPooledRenderTarget>& Out, const FSceneRenderTargetItem& Item);
 
@@ -260,6 +263,8 @@ private:
 
 	/** Elements can be 0, we compact the buffer later. */
 	TArray< TRefCountPtr<FPooledRenderTarget> > PooledRenderTargets;
+	TArray< TRefCountPtr<FPooledRenderTarget> > DeferredDeleteArray;
+	TArray< FTextureRHIParamRef > TransitionTargets;	
 
 	/** These are snapshots, have odd life times, live in the scene allocator, and don't contibute to any accounting or other management. */
 	TArray<FPooledRenderTarget*> PooledRenderTargetSnapshots;
@@ -267,6 +272,8 @@ private:
 	// redundant, can always be computed with GetStats(), to debug "out of memory" situations and used for r.RenderTargetPoolMin
 	uint32 AllocationLevelInKB;
 
+	FGraphEventRef TransitionFence;
+
 	// to avoid log spam
 	bool bCurrentlyOverBudget;
 
@@ -276,6 +283,8 @@ private:
 	// could be done on the fly but that makes the RenderTargetPoolEvents harder to read
 	void CompactPool();
 
+	void WaitForTransitionFence();
+
 	// the following is used for Event recording --------------------------------
 
 	struct SMemoryStats
diff --git a/Engine/Source/Runtime/Renderer/Private/PostProcess/RenderingCompositionGraph.cpp b/Engine/Source/Runtime/Renderer/Private/PostProcess/RenderingCompositionGraph.cpp
--- a/Engine/Source/Runtime/Renderer/Private/PostProcess/RenderingCompositionGraph.cpp
+++ b/Engine/Source/Runtime/Renderer/Private/PostProcess/RenderingCompositionGraph.cpp
@@ -894,6 +894,7 @@ const FSceneRenderTargetItem& FRenderingCompositeOutput::RequestSurface(const FR
 {
 	if(PooledRenderTarget)
 	{
+		Context.RHICmdList.TransitionResource(EResourceTransitionAccess::EWritable, PooledRenderTarget->GetRenderTargetItem().TargetableTexture);
 		return PooledRenderTarget->GetRenderTargetItem();
 	}
 
@@ -907,7 +908,7 @@ const FSceneRenderTargetItem& FRenderingCompositeOutput::RequestSurface(const FR
 
 	if(!PooledRenderTarget)
 	{
-		GRenderTargetPool.FindFreeElement(RenderTargetDesc, PooledRenderTarget, RenderTargetDesc.DebugName);
+		GRenderTargetPool.FindFreeElement(Context.RHICmdList, RenderTargetDesc, PooledRenderTarget, RenderTargetDesc.DebugName);
 	}
 
 	check(!PooledRenderTarget->IsFree());
diff --git a/Engine/Source/Runtime/Renderer/Private/PostProcess/SceneRenderTargets.cpp b/Engine/Source/Runtime/Renderer/Private/PostProcess/SceneRenderTargets.cpp
--- a/Engine/Source/Runtime/Renderer/Private/PostProcess/SceneRenderTargets.cpp
+++ b/Engine/Source/Runtime/Renderer/Private/PostProcess/SceneRenderTargets.cpp
@@ -344,7 +344,7 @@ FIntPoint FSceneRenderTargets::ComputeDesiredSize(const FSceneViewFamily& ViewFa
 	return DesiredBufferSize;
 }
 
-void FSceneRenderTargets::Allocate(const FSceneViewFamily& ViewFamily)
+void FSceneRenderTargets::Allocate(FRHICommandList& RHICmdList, const FSceneViewFamily& ViewFamily)
 {
 	check(IsInRenderingThread());
 	// ViewFamily setup wasn't complete
@@ -443,16 +443,15 @@ void FSceneRenderTargets::Allocate(const FSceneViewFamily& ViewFamily)
 
 	// Do allocation of render targets if they aren't available for the current shading path
 	CurrentFeatureLevel = NewFeatureLevel;
-	AllocateRenderTargets();
+	AllocateRenderTargets(RHICmdList);
 }
 
-void FSceneRenderTargets::BeginRenderingSceneColor(FRHICommandList& RHICmdList, ESimpleRenderTargetMode RenderTargetMode/*=EUninitializedColorExistingDepth*/, FExclusiveDepthStencil DepthStencilAccess)
+void FSceneRenderTargets::BeginRenderingSceneColor(FRHICommandList& RHICmdList, ESimpleRenderTargetMode RenderTargetMode/*=EUninitializedColorExistingDepth*/, FExclusiveDepthStencil DepthStencilAccess, bool bTransitionWritable)
 {
 	SCOPED_DRAW_EVENT(RHICmdList, BeginRenderingSceneColor);
 
-	AllocSceneColor();
-	
-	SetRenderTarget(RHICmdList, GetSceneColorSurface(), GetSceneDepthSurface(), RenderTargetMode, DepthStencilAccess);
+	AllocSceneColor(RHICmdList);
+	SetRenderTarget(RHICmdList, GetSceneColorSurface(), GetSceneDepthSurface(), RenderTargetMode, DepthStencilAccess, bTransitionWritable);
 } 
 
 int32 FSceneRenderTargets::GetGBufferRenderTargets(ERenderTargetLoadAction ColorLoadAction, FRHIRenderTargetView OutRenderTargets[MaxSimultaneousRenderTargets], int32& OutVelocityRTIndex)
@@ -497,7 +496,7 @@ void FSceneRenderTargets::BeginRenderingGBuffer(FRHICommandList& RHICmdList, ERe
 		return;
 	}
 
-	AllocSceneColor();
+	AllocSceneColor(RHICmdList);
 
 	// Set the scene color surface as the render target, and the scene depth surface as the depth-stencil target.
 	if (CurrentFeatureLevel >= ERHIFeatureLevel::SM4)
@@ -610,7 +609,7 @@ inline uint16 GetNumSceneColorMSAASamples(ERHIFeatureLevel::Type InFeatureLevel)
 	return NumSamples;
 }
 
-void FSceneRenderTargets::AllocSceneColor()
+void FSceneRenderTargets::AllocSceneColor(FRHICommandList& RHICmdList)
 {
 	if (GetSceneColorForCurrentShadingPath())
 	{
@@ -642,14 +641,14 @@ void FSceneRenderTargets::AllocSceneColor()
 			Desc.NumSamples = GetNumSceneColorMSAASamples(CurrentFeatureLevel);
 		}
 
-		GRenderTargetPool.FindFreeElement(Desc, GetSceneColorForCurrentShadingPath(), GetSceneColorTargetName(CurrentShadingPath));
+		GRenderTargetPool.FindFreeElement(RHICmdList, Desc, GetSceneColorForCurrentShadingPath(), GetSceneColorTargetName(CurrentShadingPath));
 	}
 
 	// otherwise we have a severe problem
 	check(GetSceneColorForCurrentShadingPath());
 }
 
-void FSceneRenderTargets::AllocLightAttenuation()
+void FSceneRenderTargets::AllocLightAttenuation(FRHICommandList& RHICmdList)
 {
 	if(LightAttenuation)
 	{
@@ -665,7 +664,7 @@ void FSceneRenderTargets::AllocLightAttenuation()
 	{
 		FPooledRenderTargetDesc Desc(FPooledRenderTargetDesc::Create2DDesc(BufferSize, PF_B8G8R8A8, FClearValueBinding::White, TexCreate_None, TexCreate_RenderTargetable, false));
 		Desc.Flags |= TexCreate_FastVRAM;
-		GRenderTargetPool.FindFreeElement(Desc, LightAttenuation, TEXT("LightAttenuation"));
+		GRenderTargetPool.FindFreeElement(RHICmdList, Desc, LightAttenuation, TEXT("LightAttenuation"));
 
 		// the channel assignment is documented in ShadowRendering.cpp (look for Light Attenuation channel assignment)
 	}
@@ -704,7 +703,7 @@ void FSceneRenderTargets::PreallocGBufferTargets(bool bShouldRenderVelocities)
 	}
 }
 
-void FSceneRenderTargets::AllocGBufferTargets()
+void FSceneRenderTargets::AllocGBufferTargets(FRHICommandList& RHICmdList)
 {
 	// AdjustGBufferRefCount +1 doesn't match -1 (within the same frame)
 	ensure(GBufferRefCount == 0);
@@ -732,7 +731,7 @@ void FSceneRenderTargets::AllocGBufferTargets()
 		}
 
 		FPooledRenderTargetDesc Desc(FPooledRenderTargetDesc::Create2DDesc(BufferSize, NormalGBufferFormat, FClearValueBinding::Transparent, TexCreate_None, TexCreate_RenderTargetable, false));
-		GRenderTargetPool.FindFreeElement(Desc, GBufferA, TEXT("GBufferA"));
+		GRenderTargetPool.FindFreeElement(RHICmdList, Desc, GBufferA, TEXT("GBufferA"));
 	}
 
 	// Create the specular color and power g-buffer.
@@ -740,33 +739,33 @@ void FSceneRenderTargets::AllocGBufferTargets()
 		const EPixelFormat SpecularGBufferFormat = bHighPrecisionGBuffers ? PF_FloatRGBA : PF_B8G8R8A8;
 
 		FPooledRenderTargetDesc Desc(FPooledRenderTargetDesc::Create2DDesc(BufferSize, SpecularGBufferFormat, FClearValueBinding::Transparent, TexCreate_None, TexCreate_RenderTargetable, false));
-		GRenderTargetPool.FindFreeElement(Desc, GBufferB, TEXT("GBufferB"));
+		GRenderTargetPool.FindFreeElement(RHICmdList, Desc, GBufferB, TEXT("GBufferB"));
 	}
 
 	// Create the diffuse color g-buffer.
 	{
 		const EPixelFormat DiffuseGBufferFormat = bHighPrecisionGBuffers ? PF_FloatRGBA : PF_B8G8R8A8;
 		FPooledRenderTargetDesc Desc(FPooledRenderTargetDesc::Create2DDesc(BufferSize, DiffuseGBufferFormat, FClearValueBinding::Transparent, TexCreate_SRGB, TexCreate_RenderTargetable, false));
-		GRenderTargetPool.FindFreeElement(Desc, GBufferC, TEXT("GBufferC"));
+		GRenderTargetPool.FindFreeElement(RHICmdList, Desc, GBufferC, TEXT("GBufferC"));
 	}
 
 	// Create the mask g-buffer (e.g. SSAO, subsurface scattering, wet surface mask, skylight mask, ...).
 	{
 		FPooledRenderTargetDesc Desc(FPooledRenderTargetDesc::Create2DDesc(BufferSize, PF_B8G8R8A8, FClearValueBinding(FLinearColor(0, 1, 1, 1)), TexCreate_None, TexCreate_RenderTargetable, false));
-		GRenderTargetPool.FindFreeElement(Desc, GBufferD, TEXT("GBufferD"));
+		GRenderTargetPool.FindFreeElement(RHICmdList, Desc, GBufferD, TEXT("GBufferD"));
 	}
 
 	if (bAllowStaticLighting)
 	{
 		FPooledRenderTargetDesc Desc(FPooledRenderTargetDesc::Create2DDesc(BufferSize, PF_B8G8R8A8, FClearValueBinding(FLinearColor(1, 1, 1, 1)), TexCreate_None, TexCreate_RenderTargetable, false));
-		GRenderTargetPool.FindFreeElement(Desc, GBufferE, TEXT("GBufferE"));
+		GRenderTargetPool.FindFreeElement(RHICmdList, Desc, GBufferE, TEXT("GBufferE"));
 	}
 
 	GBasePassOutputsVelocityDebug = CVarBasePassOutputsVelocityDebug.GetValueOnRenderThread();
 	if (bAllocateVelocityGBuffer)
 	{
 		FPooledRenderTargetDesc VelocityRTDesc = FVelocityRendering::GetRenderTargetDesc();
-		GRenderTargetPool.FindFreeElement(VelocityRTDesc, GBufferVelocity, TEXT("GBufferVelocity"));
+		GRenderTargetPool.FindFreeElement(RHICmdList, VelocityRTDesc, GBufferVelocity, TEXT("GBufferVelocity"));
 	}
 
 	// otherwise we have a severe problem
@@ -917,12 +916,12 @@ TRefCountPtr<IPooledRenderTarget>& FSceneRenderTargets::GetLightAttenuation()
 	return LightAttenuation;
 }
 
-void FSceneRenderTargets::AdjustGBufferRefCount(int Delta)
+void FSceneRenderTargets::AdjustGBufferRefCount(FRHICommandList& RHICmdList, int Delta)
 { 
 	
 	if(Delta > 0 && GBufferRefCount == 0)
 	{
-		AllocGBufferTargets();
+		AllocGBufferTargets(RHICmdList);
 	}
 	else
 	{
@@ -947,7 +946,7 @@ void FSceneRenderTargets::FinishRenderingSceneColor(FRHICommandListImmediate& RH
 
 bool FSceneRenderTargets::BeginRenderingCustomDepth(FRHICommandListImmediate& RHICmdList, bool bPrimitives)
 {
-	IPooledRenderTarget* CustomDepthRenderTarget = RequestCustomDepth(bPrimitives);
+	IPooledRenderTarget* CustomDepthRenderTarget = RequestCustomDepth(RHICmdList, bPrimitives);
 
 	if(CustomDepthRenderTarget)
 	{
@@ -1119,8 +1118,9 @@ void FSceneRenderTargets::BeginRenderingShadowDepth(FRHICommandList& RHICmdList,
 	{
 		Info.NumColorRenderTargets = 1;
 		Info.ColorRenderTarget[0].Texture = GetOptionalShadowDepthColorSurface();
+		RHICmdList.TransitionResource(EResourceTransitionAccess::EWritable, Info.ColorRenderTarget[0].Texture);
 	}
-
+	RHICmdList.TransitionResource(EResourceTransitionAccess::EWritable, Info.DepthStencilRenderTarget.Texture);
 	RHICmdList.SetRenderTargetsAndClear(Info);
 }
 
@@ -1190,12 +1190,12 @@ void FSceneRenderTargets::BeginRenderingLightAttenuation(FRHICommandList& RHICmd
 {
 	SCOPED_DRAW_EVENT(RHICmdList, BeginRenderingLightAttenuation);
 
-	AllocLightAttenuation();
+	AllocLightAttenuation(RHICmdList);
 
 	GRenderTargetPool.VisualizeTexture.SetCheckPoint(RHICmdList, GetLightAttenuation());
 
 	// Set the light attenuation surface as the render target, and the scene depth buffer as the depth-stencil surface.
-	SetRenderTarget(RHICmdList, GetLightAttenuationSurface(), GetSceneDepthSurface(), bClearToWhite ? ESimpleRenderTargetMode::EClearColorExistingDepth : ESimpleRenderTargetMode::EExistingColorAndDepth, FExclusiveDepthStencil::DepthRead_StencilWrite);
+	SetRenderTarget(RHICmdList, GetLightAttenuationSurface(), GetSceneDepthSurface(), bClearToWhite ? ESimpleRenderTargetMode::EClearColorExistingDepth : ESimpleRenderTargetMode::EExistingColorAndDepth, FExclusiveDepthStencil::DepthRead_StencilWrite, true);
 }
 
 void FSceneRenderTargets::FinishRenderingLightAttenuation(FRHICommandList& RHICmdList)
@@ -1241,7 +1241,7 @@ bool FSceneRenderTargets::BeginRenderingSeparateTranslucency(FRHICommandList& RH
 		else
 		{
 			FSceneViewState* ViewState = (FSceneViewState*)View.State;
-			SeparateTranslucency = &ViewState->GetSeparateTranslucency(GetBufferSizeXY());
+			SeparateTranslucency = &ViewState->GetSeparateTranslucency(RHICmdList, GetBufferSizeXY());
 		}
 
 		check((*SeparateTranslucency)->GetRenderTargetItem().TargetableTexture->GetClearColor() == FLinearColor::Black);
@@ -1280,7 +1280,7 @@ void FSceneRenderTargets::FinishRenderingSeparateTranslucency(FRHICommandList& R
 		else
 		{
 			FSceneViewState* ViewState = (FSceneViewState*)View.State;
-			SeparateTranslucency = &ViewState->GetSeparateTranslucency(GetBufferSizeXY());
+			SeparateTranslucency = &ViewState->GetSeparateTranslucency(RHICmdList, GetBufferSizeXY());
 		}
 
 		RHICmdList.CopyToResolveTarget((*SeparateTranslucency)->GetRenderTargetItem().TargetableTexture, (*SeparateTranslucency)->GetRenderTargetItem().ShaderResourceTexture, true, FResolveParams());
@@ -1345,28 +1345,28 @@ int32 FSceneRenderTargets::GetEditorMSAACompositingSampleCount() const
 	return Value;
 }
 
-const FTexture2DRHIRef& FSceneRenderTargets::GetEditorPrimitivesColor()
+const FTexture2DRHIRef& FSceneRenderTargets::GetEditorPrimitivesColor(FRHICommandList& RHICmdList)
 {
 	const bool bIsValid = IsValidRef(EditorPrimitivesColor);
 
 	if( !bIsValid || EditorPrimitivesColor->GetDesc().NumSamples != GetEditorMSAACompositingSampleCount() )
 	{
 		// If the target is does not match the MSAA settings it needs to be recreated
-		InitEditorPrimitivesColor();
+		InitEditorPrimitivesColor(RHICmdList);
 	}
 
 	return (const FTexture2DRHIRef&)EditorPrimitivesColor->GetRenderTargetItem().TargetableTexture;
 }
 
 
-const FTexture2DRHIRef& FSceneRenderTargets::GetEditorPrimitivesDepth()
+const FTexture2DRHIRef& FSceneRenderTargets::GetEditorPrimitivesDepth(FRHICommandList& RHICmdList)
 {
 	const bool bIsValid = IsValidRef(EditorPrimitivesDepth);
 
 	if( !bIsValid || EditorPrimitivesDepth->GetDesc().NumSamples != GetEditorMSAACompositingSampleCount() )
 	{
 		// If the target is does not match the MSAA settings it needs to be recreated
-		InitEditorPrimitivesDepth();
+		InitEditorPrimitivesDepth(RHICmdList);
 	}
 
 	return (const FTexture2DRHIRef&)EditorPrimitivesDepth->GetRenderTargetItem().TargetableTexture;
@@ -1392,7 +1392,7 @@ bool FSceneRenderTargets::IsSeparateTranslucencyActive(const FViewInfo& View) co
 		&& View.State != NULL;	// We require a ViewState in order for separate translucency to work (it keeps track of our SeparateTranslucencyRT)
 }
 
-void FSceneRenderTargets::InitEditorPrimitivesColor()
+void FSceneRenderTargets::InitEditorPrimitivesColor(FRHICommandList& RHICmdList)
 {
 	FPooledRenderTargetDesc Desc(FPooledRenderTargetDesc::Create2DDesc(BufferSize, 
 		PF_B8G8R8A8,
@@ -1403,10 +1403,10 @@ void FSceneRenderTargets::InitEditorPrimitivesColor()
 
 	Desc.NumSamples = GetEditorMSAACompositingSampleCount();
 
-	GRenderTargetPool.FindFreeElement(Desc, EditorPrimitivesColor, TEXT("EditorPrimitivesColor"));
+	GRenderTargetPool.FindFreeElement(RHICmdList, Desc, EditorPrimitivesColor, TEXT("EditorPrimitivesColor"));
 }
 
-void FSceneRenderTargets::InitEditorPrimitivesDepth()
+void FSceneRenderTargets::InitEditorPrimitivesDepth(FRHICommandList& RHICmdList)
 {
 	FPooledRenderTargetDesc Desc(FPooledRenderTargetDesc::Create2DDesc(BufferSize, 
 		PF_DepthStencil,
@@ -1417,7 +1417,7 @@ void FSceneRenderTargets::InitEditorPrimitivesDepth()
 
 	Desc.NumSamples = GetEditorMSAACompositingSampleCount();
 
-	GRenderTargetPool.FindFreeElement(Desc, EditorPrimitivesDepth, TEXT("EditorPrimitivesDepth"));
+	GRenderTargetPool.FindFreeElement(RHICmdList, Desc, EditorPrimitivesDepth, TEXT("EditorPrimitivesDepth"));
 }
 
 void FSceneRenderTargets::QuantizeBufferSize(int32& InOutBufferSizeX, int32& InOutBufferSizeY)
@@ -1437,12 +1437,12 @@ void FSceneRenderTargets::SetBufferSize(int32 InBufferSizeX, int32 InBufferSizeY
 	BufferSize.Y = InBufferSizeY;
 }
 
-void FSceneRenderTargets::AllocateForwardShadingPathRenderTargets()
+void FSceneRenderTargets::AllocateForwardShadingPathRenderTargets(FRHICommandList& RHICmdList)
 {
 	// on ES2 we don't do on demand allocation of SceneColor yet (in non ES2 it's released in the Tonemapper Process())
-	AllocSceneColor();
-	AllocateCommonDepthTargets();
-	AllocateReflectionTargets();
+	AllocSceneColor(RHICmdList);
+	AllocateCommonDepthTargets(RHICmdList);
+	AllocateReflectionTargets(RHICmdList);
 
 	EPixelFormat Format = GetSceneColor()->GetDesc().Format;
 
@@ -1453,7 +1453,7 @@ void FSceneRenderTargets::AllocateForwardShadingPathRenderTargets()
 		// creating a PF_R16F (a true one-channel renderable fp texture) is only supported on GL if EXT_texture_rg is available.  It's present
 		// on iOS, but not in WebGL or Android.
 		FPooledRenderTargetDesc Desc(FPooledRenderTargetDesc::Create2DDesc(BufferSize, PF_FloatRGBA, FClearValueBinding::None, TexCreate_None, TexCreate_RenderTargetable, false));
-		GRenderTargetPool.FindFreeElement(Desc, SceneAlphaCopy, TEXT("SceneAlphaCopy"));
+		GRenderTargetPool.FindFreeElement(RHICmdList, Desc, SceneAlphaCopy, TEXT("SceneAlphaCopy"));
 	}
 	else
 #endif
@@ -1462,17 +1462,17 @@ void FSceneRenderTargets::AllocateForwardShadingPathRenderTargets()
 	}
 }
 
-void FSceneRenderTargets::AllocateForwardShadingShadowDepthTarget(const FIntPoint& ShadowBufferResolution)
+void FSceneRenderTargets::AllocateForwardShadingShadowDepthTarget(FRHICommandListImmediate& RHICmdList, const FIntPoint& ShadowBufferResolution)
 {
 	{
 		FPooledRenderTargetDesc Desc(FPooledRenderTargetDesc::Create2DDesc(ShadowBufferResolution, PF_ShadowDepth, FClearValueBinding::DepthOne, TexCreate_None, TexCreate_DepthStencilTargetable, false));
-		GRenderTargetPool.FindFreeElement(Desc, ShadowDepthZ, TEXT("ShadowDepthZ"));
+		GRenderTargetPool.FindFreeElement(RHICmdList, Desc, ShadowDepthZ, TEXT("ShadowDepthZ"));
 	}
 
 	if (!GSupportsDepthRenderTargetWithoutColorRenderTarget)
 	{
 		FPooledRenderTargetDesc Desc(FPooledRenderTargetDesc::Create2DDesc(ShadowBufferResolution, PF_B8G8R8A8, FClearValueBinding::None, TexCreate_None, TexCreate_RenderTargetable, false));
-		GRenderTargetPool.FindFreeElement(Desc, OptionalShadowDepthColor, TEXT("OptionalShadowDepthColor"));
+		GRenderTargetPool.FindFreeElement(RHICmdList, Desc, OptionalShadowDepthColor, TEXT("OptionalShadowDepthColor"));
 	}
 }
 
@@ -1508,7 +1508,7 @@ static TCHAR* const GetTranslucencyShadowTransmissionName(uint32 Id)
 	return (TCHAR*)TEXT("InvalidName");
 }
 
-void FSceneRenderTargets::AllocateReflectionTargets()
+void FSceneRenderTargets::AllocateReflectionTargets(FRHICommandList& RHICmdList)
 {	
 	if (GSupportsRenderTargetFormat_PF_FloatRGBA)
 	{
@@ -1527,8 +1527,8 @@ void FSceneRenderTargets::AllocateReflectionTargets()
 			{
 				// Create scratch cubemaps for filtering passes
 				FPooledRenderTargetDesc Desc2(FPooledRenderTargetDesc::CreateCubemapDesc(GReflectionCaptureSize, PF_FloatRGBA, FClearValueBinding::None, CubeTexFlags, TexCreate_RenderTargetable, false, 1, NumReflectionCaptureMips));
-				GRenderTargetPool.FindFreeElement(Desc2, ReflectionColorScratchCubemap[0], TEXT("ReflectionColorScratchCubemap0"));
-				GRenderTargetPool.FindFreeElement(Desc2, ReflectionColorScratchCubemap[1], TEXT("ReflectionColorScratchCubemap1"));
+				GRenderTargetPool.FindFreeElement(RHICmdList, Desc2, ReflectionColorScratchCubemap[0], TEXT("ReflectionColorScratchCubemap0"));
+				GRenderTargetPool.FindFreeElement(RHICmdList, Desc2, ReflectionColorScratchCubemap[1], TEXT("ReflectionColorScratchCubemap1"));
 			}
 
 			extern int32 GDiffuseIrradianceCubemapSize;
@@ -1536,13 +1536,13 @@ void FSceneRenderTargets::AllocateReflectionTargets()
 
 			{
 				FPooledRenderTargetDesc Desc2(FPooledRenderTargetDesc::CreateCubemapDesc(GDiffuseIrradianceCubemapSize, PF_FloatRGBA, FClearValueBinding::None, CubeTexFlags, TexCreate_RenderTargetable, false, 1, NumDiffuseIrradianceMips));
-				GRenderTargetPool.FindFreeElement(Desc2, DiffuseIrradianceScratchCubemap[0], TEXT("DiffuseIrradianceScratchCubemap0"));
-				GRenderTargetPool.FindFreeElement(Desc2, DiffuseIrradianceScratchCubemap[1], TEXT("DiffuseIrradianceScratchCubemap1"));
+				GRenderTargetPool.FindFreeElement(RHICmdList, Desc2, DiffuseIrradianceScratchCubemap[0], TEXT("DiffuseIrradianceScratchCubemap0"));
+				GRenderTargetPool.FindFreeElement(RHICmdList, Desc2, DiffuseIrradianceScratchCubemap[1], TEXT("DiffuseIrradianceScratchCubemap1"));
 			}
 
 			{
 				FPooledRenderTargetDesc Desc(FPooledRenderTargetDesc::Create2DDesc(FIntPoint(FSHVector3::MaxSHBasis, 1), PF_FloatRGBA, FClearValueBinding::None, TexCreate_None, TexCreate_RenderTargetable, false));
-				GRenderTargetPool.FindFreeElement(Desc, SkySHIrradianceMap, TEXT("SkySHIrradianceMap"));
+				GRenderTargetPool.FindFreeElement(RHICmdList, Desc, SkySHIrradianceMap, TEXT("SkySHIrradianceMap"));
 			}
 		}
 
@@ -1559,12 +1559,12 @@ void FSceneRenderTargets::AllocateReflectionTargets()
 			int32 ReflectionBrightnessIndex = bSupportsR32Float ? 0 : 1;
 			EPixelFormat BrightnessFormat = bSupportsR32Float ? PF_R32_FLOAT : PF_FloatRGBA;
 			FPooledRenderTargetDesc Desc3(FPooledRenderTargetDesc::Create2DDesc(FIntPoint(1, 1), BrightnessFormat, FClearValueBinding::None, TexCreate_None, TexCreate_RenderTargetable, false));
-			GRenderTargetPool.FindFreeElement(Desc3, ReflectionBrightnessTarget, GetReflectionBrightnessTargetName(ReflectionBrightnessIndex));
+			GRenderTargetPool.FindFreeElement(RHICmdList, Desc3, ReflectionBrightnessTarget, GetReflectionBrightnessTargetName(ReflectionBrightnessIndex));
 		}
 	}
 }
 
-void FSceneRenderTargets::AllocateCommonDepthTargets()
+void FSceneRenderTargets::AllocateCommonDepthTargets(FRHICommandList& RHICmdList)
 {
 	if (!SceneDepthZ)
 	{
@@ -1572,27 +1572,27 @@ void FSceneRenderTargets::AllocateCommonDepthTargets()
 		FPooledRenderTargetDesc Desc(FPooledRenderTargetDesc::Create2DDesc(BufferSize, PF_DepthStencil, FClearValueBinding::DepthFar, TexCreate_None, TexCreate_DepthStencilTargetable, false));
 		Desc.NumSamples = GetNumSceneColorMSAASamples(CurrentFeatureLevel);
 		Desc.Flags |= TexCreate_FastVRAM;
-		GRenderTargetPool.FindFreeElement(Desc, SceneDepthZ, TEXT("SceneDepthZ"));
+		GRenderTargetPool.FindFreeElement(RHICmdList, Desc, SceneDepthZ, TEXT("SceneDepthZ"));
 	}
 
 	// When targeting DX Feature Level 10, create an auxiliary texture to store the resolved scene depth, and a render-targetable surface to hold the unresolved scene depth.
 	if (!AuxiliarySceneDepthZ && !GSupportsDepthFetchDuringDepthTest)
 	{
 		FPooledRenderTargetDesc Desc(FPooledRenderTargetDesc::Create2DDesc(BufferSize, PF_DepthStencil, FClearValueBinding::DepthFar, TexCreate_None, TexCreate_DepthStencilTargetable, false));
-		GRenderTargetPool.FindFreeElement(Desc, AuxiliarySceneDepthZ, TEXT("AuxiliarySceneDepthZ"));
+		GRenderTargetPool.FindFreeElement(RHICmdList, Desc, AuxiliarySceneDepthZ, TEXT("AuxiliarySceneDepthZ"));
 	}
 }
 
-void FSceneRenderTargets::AllocateDeferredShadingPathRenderTargets()
+void FSceneRenderTargets::AllocateDeferredShadingPathRenderTargets(FRHICommandList& RHICmdList)
 {
-	AllocateCommonDepthTargets();
+	AllocateCommonDepthTargets(RHICmdList);
 
 	// Create a quarter-sized version of the scene depth.
 	{
 		FIntPoint SmallDepthZSize(FMath::Max<uint32>(BufferSize.X / SmallColorDepthDownsampleFactor, 1), FMath::Max<uint32>(BufferSize.Y / SmallColorDepthDownsampleFactor, 1));
 
 		FPooledRenderTargetDesc Desc(FPooledRenderTargetDesc::Create2DDesc(SmallDepthZSize, PF_DepthStencil, FClearValueBinding::None, TexCreate_None, TexCreate_DepthStencilTargetable, true));
-		GRenderTargetPool.FindFreeElement(Desc, SmallDepthZ, TEXT("SmallDepthZ"));
+		GRenderTargetPool.FindFreeElement(RHICmdList, Desc, SmallDepthZ, TEXT("SmallDepthZ"));
 	}
 
 	// Set up quarter size scene color shared texture
@@ -1608,7 +1608,7 @@ void FSceneRenderTargets::AllocateDeferredShadingPathRenderTargets()
 			{
 				// Using PF_FloatRGBA because Fourier coefficients used by Fourier opacity maps have a large range and can be negative
 				FPooledRenderTargetDesc Desc(FPooledRenderTargetDesc::Create2DDesc(TranslucentShadowBufferResolution, PF_FloatRGBA, FClearValueBinding::None, TexCreate_None, TexCreate_RenderTargetable, false));
-				GRenderTargetPool.FindFreeElement(Desc, TranslucencyShadowTransmission[SurfaceIndex], GetTranslucencyShadowTransmissionName(SurfaceIndex));
+				GRenderTargetPool.FindFreeElement(RHICmdList, Desc, TranslucencyShadowTransmission[SurfaceIndex], GetTranslucencyShadowTransmissionName(SurfaceIndex));
 			}
 		}
 	}
@@ -1621,19 +1621,20 @@ void FSceneRenderTargets::AllocateDeferredShadingPathRenderTargets()
 			const int32 SurfaceResolution = GetCubeShadowDepthZResolution(SurfaceIndex);
 
 			FPooledRenderTargetDesc Desc(FPooledRenderTargetDesc::CreateCubemapDesc(SurfaceResolution, PF_ShadowDepth, FClearValueBinding::None, TexCreate_None, TexCreate_DepthStencilTargetable | TexCreate_NoFastClear, false));
-			GRenderTargetPool.FindFreeElement(Desc, CubeShadowDepthZ[SurfaceIndex], TEXT("CubeShadowDepthZ[]"));
+			GRenderTargetPool.FindFreeElement(RHICmdList, Desc, CubeShadowDepthZ[SurfaceIndex], TEXT("CubeShadowDepthZ[]"));
 		}
 	}
 
 	//create the shadow depth texture and/or surface
 	{
 		FPooledRenderTargetDesc Desc(FPooledRenderTargetDesc::Create2DDesc(ShadowBufferResolution, PF_ShadowDepth, FClearValueBinding::DepthOne, TexCreate_None, TexCreate_DepthStencilTargetable, false));
-		GRenderTargetPool.FindFreeElement(Desc, ShadowDepthZ, TEXT("ShadowDepthZ"));
+		GRenderTargetPool.FindFreeElement(RHICmdList, Desc, ShadowDepthZ, TEXT("ShadowDepthZ"));
 	}
 		
 	{
 		FPooledRenderTargetDesc Desc(FPooledRenderTargetDesc::Create2DDesc(GetPreShadowCacheTextureResolution(), PF_ShadowDepth, FClearValueBinding::None, TexCreate_None, TexCreate_DepthStencilTargetable, false));
-		GRenderTargetPool.FindFreeElement(Desc, PreShadowCacheDepthZ, TEXT("PreShadowCacheDepthZ"));
+		Desc.AutoWritable = false;
+		GRenderTargetPool.FindFreeElement(RHICmdList, Desc, PreShadowCacheDepthZ, TEXT("PreShadowCacheDepthZ"));
 		// Mark the preshadow cache as newly allocated, so the cache will know to update
 		bPreshadowCacheNewlyAllocated = true;
 	}
@@ -1644,13 +1645,14 @@ void FSceneRenderTargets::AllocateDeferredShadingPathRenderTargets()
 		// Create the screen space ambient occlusion buffer
 		{
 			FPooledRenderTargetDesc Desc(FPooledRenderTargetDesc::Create2DDesc(BufferSize, PF_G8, FClearValueBinding::None, TexCreate_None, TexCreate_RenderTargetable, false));
-			GRenderTargetPool.FindFreeElement(Desc, ScreenSpaceAO, TEXT("ScreenSpaceAO"));
+			GRenderTargetPool.FindFreeElement(RHICmdList, Desc, ScreenSpaceAO, TEXT("ScreenSpaceAO"));
 		}
 
 		{
 			for (int32 RTSetIndex = 0; RTSetIndex < NumTranslucentVolumeRenderTargetSets; RTSetIndex++)
 			{
 				GRenderTargetPool.FindFreeElement(
+					RHICmdList,
 					FPooledRenderTargetDesc(FPooledRenderTargetDesc::CreateVolumeDesc(
 						GTranslucencyLightingVolumeDim,
 						GTranslucencyLightingVolumeDim,
@@ -1665,6 +1667,7 @@ void FSceneRenderTargets::AllocateDeferredShadingPathRenderTargets()
 					);
 
 				GRenderTargetPool.FindFreeElement(
+					RHICmdList,
 					FPooledRenderTargetDesc(FPooledRenderTargetDesc::CreateVolumeDesc(
 						GTranslucencyLightingVolumeDim,
 						GTranslucencyLightingVolumeDim,
@@ -1685,7 +1688,7 @@ void FSceneRenderTargets::AllocateDeferredShadingPathRenderTargets()
 	if(UseLightPropagationVolumeRT(CurrentFeatureLevel))
 	{
 		FPooledRenderTargetDesc Desc(FPooledRenderTargetDesc::Create2DDesc(BufferSize, PF_R8G8, FClearValueBinding::None, TexCreate_None, TexCreate_RenderTargetable, false));
-		GRenderTargetPool.FindFreeElement(Desc, DirectionalOcclusion, TEXT("DirectionalOcclusion"));
+		GRenderTargetPool.FindFreeElement(RHICmdList, Desc, DirectionalOcclusion, TEXT("DirectionalOcclusion"));
 	}
 
 	{
@@ -1694,10 +1697,10 @@ void FSceneRenderTargets::AllocateDeferredShadingPathRenderTargets()
 		{
 			Desc.TargetableFlags |= TexCreate_UAV;
 		}
-		GRenderTargetPool.FindFreeElement(Desc, LightAccumulation, TEXT("LightAccumulation"));
+		GRenderTargetPool.FindFreeElement(RHICmdList, Desc, LightAccumulation, TEXT("LightAccumulation"));
 	}
 
-	AllocateReflectionTargets();
+	AllocateReflectionTargets(RHICmdList);
 
 	if (CurrentFeatureLevel >= ERHIFeatureLevel::SM5)
 	{
@@ -1709,17 +1712,17 @@ void FSceneRenderTargets::AllocateDeferredShadingPathRenderTargets()
 
 			{
 				FPooledRenderTargetDesc Desc(FPooledRenderTargetDesc::Create2DDesc(Extent, PF_R8G8B8A8, FClearValueBinding::None, TexCreate_None, TexCreate_RenderTargetable, false));
-				GRenderTargetPool.FindFreeElement(Desc, ReflectiveShadowMapNormal, TEXT("RSMNormal"));
+				GRenderTargetPool.FindFreeElement(RHICmdList, Desc, ReflectiveShadowMapNormal, TEXT("RSMNormal"));
 			}
 
 			{
 				FPooledRenderTargetDesc Desc(FPooledRenderTargetDesc::Create2DDesc(Extent, PF_FloatR11G11B10, FClearValueBinding::None, TexCreate_None, TexCreate_RenderTargetable, false));
-				GRenderTargetPool.FindFreeElement(Desc, ReflectiveShadowMapDiffuse, TEXT("RSMDiffuse"));
+				GRenderTargetPool.FindFreeElement(RHICmdList, Desc, ReflectiveShadowMapDiffuse, TEXT("RSMDiffuse"));
 			}
 
 			{
 				FPooledRenderTargetDesc Desc(FPooledRenderTargetDesc::Create2DDesc(Extent, PF_DepthStencil, FClearValueBinding::None, TexCreate_None, TexCreate_DepthStencilTargetable, false));
-				GRenderTargetPool.FindFreeElement(Desc, ReflectiveShadowMapDepth, TEXT("RSMDepth"));
+				GRenderTargetPool.FindFreeElement(RHICmdList, Desc, ReflectiveShadowMapDepth, TEXT("RSMDepth"));
 			}
 		}
 	}
@@ -1727,7 +1730,7 @@ void FSceneRenderTargets::AllocateDeferredShadingPathRenderTargets()
 	if (bAllocateVelocityGBuffer)
 	{
 		FPooledRenderTargetDesc VelocityRTDesc = FVelocityRendering::GetRenderTargetDesc();
-		GRenderTargetPool.FindFreeElement(VelocityRTDesc, GBufferVelocity, TEXT("GBufferVelocity"));
+		GRenderTargetPool.FindFreeElement(RHICmdList, VelocityRTDesc, GBufferVelocity, TEXT("GBufferVelocity"));
 	}
 }
 
@@ -1767,17 +1770,17 @@ EPixelFormat FSceneRenderTargets::GetSceneColorFormat() const
 	return SceneColorBufferFormat;
 }
 
-void FSceneRenderTargets::AllocateRenderTargets()
+void FSceneRenderTargets::AllocateRenderTargets(FRHICommandList& RHICmdList)
 {
 	if (BufferSize.X > 0 && BufferSize.Y > 0 && !AreShadingPathRenderTargetsAllocated(CurrentShadingPath))
 	{
 		if ((EShadingPath)CurrentShadingPath == EShadingPath::Forward)
 		{
-			AllocateForwardShadingPathRenderTargets();
+			AllocateForwardShadingPathRenderTargets(RHICmdList);
 		}
 		else
 		{
-			AllocateDeferredShadingPathRenderTargets();
+			AllocateDeferredShadingPathRenderTargets(RHICmdList);
 		}
 	}
 }
@@ -1972,7 +1975,7 @@ IPooledRenderTarget* FSceneRenderTargets::GetGBufferVelocityRT()
 	return GBufferVelocity;
 }
 
-IPooledRenderTarget* FSceneRenderTargets::RequestCustomDepth(bool bPrimitives)
+IPooledRenderTarget* FSceneRenderTargets::RequestCustomDepth(FRHICommandListImmediate& RHICmdList, bool bPrimitives)
 {
 	int Value = CVarCustomDepth.GetValueOnRenderThread();
 
@@ -1980,7 +1983,7 @@ IPooledRenderTarget* FSceneRenderTargets::RequestCustomDepth(bool bPrimitives)
 	{
 		// Todo: Could check if writes stencil here and create min viable target
 		FPooledRenderTargetDesc Desc(FPooledRenderTargetDesc::Create2DDesc(BufferSize, PF_DepthStencil, FClearValueBinding::DepthFar, TexCreate_None, TexCreate_DepthStencilTargetable, false));
-		GRenderTargetPool.FindFreeElement(Desc, CustomDepth, TEXT("CustomDepth"));
+		GRenderTargetPool.FindFreeElement(RHICmdList, Desc, CustomDepth, TEXT("CustomDepth"));
 		CustomStencilSRV = RHICreateShaderResourceView((FTexture2DRHIRef&)CustomDepth->GetRenderTargetItem().TargetableTexture, 0, 1, PF_X24_G8);
 		return CustomDepth;
 	}
diff --git a/Engine/Source/Runtime/Renderer/Private/PostProcess/SceneRenderTargets.h b/Engine/Source/Runtime/Renderer/Private/PostProcess/SceneRenderTargets.h
--- a/Engine/Source/Runtime/Renderer/Private/PostProcess/SceneRenderTargets.h
+++ b/Engine/Source/Runtime/Renderer/Private/PostProcess/SceneRenderTargets.h
@@ -199,12 +199,12 @@ public:
 	 * Checks that scene render targets are ready for rendering a view family of the given dimensions.
 	 * If the allocated render targets are too small, they are reallocated.
 	 */
-	void Allocate(const FSceneViewFamily& ViewFamily);
+	void Allocate(FRHICommandList& RHICmdList, const FSceneViewFamily& ViewFamily);
 
 	/**
 	 * Forward shading can't know how big the optimal atlased shadow buffer will be, so provide a set it up per frame.
 	 */
-	void AllocateForwardShadingShadowDepthTarget(const FIntPoint& ShadowBufferResolution);
+	void AllocateForwardShadingShadowDepthTarget(FRHICommandListImmediate& RHICmdList, const FIntPoint& ShadowBufferResolution);
 
 	/**
 	 *
@@ -217,7 +217,7 @@ public:
 	/**
 	 * Sets the scene color target and restores its contents if necessary
 	 */
-	void BeginRenderingSceneColor(FRHICommandList& RHICmdList, ESimpleRenderTargetMode RenderTargetMode = ESimpleRenderTargetMode::EUninitializedColorExistingDepth, FExclusiveDepthStencil DepthStencilAccess = FExclusiveDepthStencil::DepthWrite_StencilWrite);
+	void BeginRenderingSceneColor(FRHICommandList& FRHICommandListImmediate, ESimpleRenderTargetMode RenderTargetMode = ESimpleRenderTargetMode::EUninitializedColorExistingDepth, FExclusiveDepthStencil DepthStencilAccess = FExclusiveDepthStencil::DepthWrite_StencilWrite, bool bTransitionWritable = true);
 	
 	/**
 	 * Called when finished rendering to the scene color surface
@@ -271,7 +271,7 @@ public:
 	void BeginRenderingSceneAlphaCopy(FRHICommandListImmediate& RHICmdList);
 	void FinishRenderingSceneAlphaCopy(FRHICommandListImmediate& RHICmdList);
 
-	void BeginRenderingLightAttenuation(FRHICommandList& RHICmdList, bool bClearToWhite=false);
+	void BeginRenderingLightAttenuation(FRHICommandList& RHICmdList, bool bClearToWhite = false);
 	void FinishRenderingLightAttenuation(FRHICommandList& RHICmdList);
 
 	/**
@@ -291,13 +291,13 @@ public:
 	 * Gets the editor primitives color target/shader resource.  This may recreate the target
 	 * if the msaa settings dont match
 	 */
-	const FTexture2DRHIRef& GetEditorPrimitivesColor();
+	const FTexture2DRHIRef& GetEditorPrimitivesColor(FRHICommandList& RHICmdList);
 
 	/**
 	 * Gets the editor primitives depth target/shader resource.  This may recreate the target
 	 * if the msaa settings dont match
 	 */
-	const FTexture2DRHIRef& GetEditorPrimitivesDepth();
+	const FTexture2DRHIRef& GetEditorPrimitivesDepth(FRHICommandList& RHICmdList);
 
 
 	// FRenderResource interface.
@@ -408,7 +408,7 @@ public:
 	}
 
 	// @return can be 0 if the feature is disabled
-	IPooledRenderTarget* RequestCustomDepth(bool bPrimitives);
+	IPooledRenderTarget* RequestCustomDepth(FRHICommandListImmediate& RHICmdList, bool bPrimitives);
 
 	bool IsCustomDepthPassWritingStencil() const;
 
@@ -480,15 +480,15 @@ public:
 
 	// allows to release the GBuffer once post process materials no longer need it
 	// @param 1: add a reference, -1: remove a reference
-	void AdjustGBufferRefCount(int Delta);
+	void AdjustGBufferRefCount(FRHICommandList& RHICmdList, int Delta);
 
 	//
 	void PreallocGBufferTargets(bool bShouldRenderVelocities);
-	void AllocGBufferTargets();
+	void AllocGBufferTargets(FRHICommandList& RHICmdList);
 
-	void AllocLightAttenuation();
+	void AllocLightAttenuation(FRHICommandList& RHICmdList);
 
-	void AllocateReflectionTargets();
+	void AllocateReflectionTargets(FRHICommandList& RHICmdList);
 
 	TRefCountPtr<IPooledRenderTarget>& GetReflectionBrightnessTarget();
 
@@ -626,29 +626,29 @@ private:
 	/**
 	 * Initializes the editor primitive color render target
 	 */
-	void InitEditorPrimitivesColor();
+	void InitEditorPrimitivesColor(FRHICommandList& RHICmdList);
 
 	/**
 	 * Initializes the editor primitive depth buffer
 	 */
-	void InitEditorPrimitivesDepth();
+	void InitEditorPrimitivesDepth(FRHICommandList& RHICmdList);
 
 	/** Allocates render targets for use with the forward shading path. */
-	void AllocateForwardShadingPathRenderTargets();	
+	void AllocateForwardShadingPathRenderTargets(FRHICommandList& RHICmdList);
 
 	/** Allocates render targets for use with the deferred shading path. */
-	void AllocateDeferredShadingPathRenderTargets();
+	void AllocateDeferredShadingPathRenderTargets(FRHICommandList& RHICmdList);
 
 	/** Allocates render targets for use with the current shading path. */
-	void AllocateRenderTargets();
+	void AllocateRenderTargets(FRHICommandList& RHICmdList);
 
 	/** Allocates common depth render targets that are used by both forward and deferred rendering paths */
-	void AllocateCommonDepthTargets();
+	void AllocateCommonDepthTargets(FRHICommandList& RHICmdList);
 
 	/** Determine the appropriate render target dimensions. */
 	FIntPoint ComputeDesiredSize(const FSceneViewFamily& ViewFamily);
 
-	void AllocSceneColor();
+	void AllocSceneColor(FRHICommandList& RHICmdList);
 
 	// internal method, used by AdjustGBufferRefCount()
 	void ReleaseGBufferTargets();
diff --git a/Engine/Source/Runtime/Renderer/Private/PostProcess/ScreenSpaceReflections.cpp b/Engine/Source/Runtime/Renderer/Private/PostProcess/ScreenSpaceReflections.cpp
--- a/Engine/Source/Runtime/Renderer/Private/PostProcess/ScreenSpaceReflections.cpp
+++ b/Engine/Source/Runtime/Renderer/Private/PostProcess/ScreenSpaceReflections.cpp
@@ -442,7 +442,7 @@ FPooledRenderTargetDesc FRCPassPostProcessScreenSpaceReflections::ComputeOutputD
 void ScreenSpaceReflections(FRHICommandListImmediate& RHICmdList, FViewInfo& View, TRefCountPtr<IPooledRenderTarget>& SSROutput)
 {
 	FRenderingCompositePassContext CompositeContext(RHICmdList, View);	
-	FPostprocessContext Context( CompositeContext.Graph, View );
+	FPostprocessContext Context(RHICmdList, CompositeContext.Graph, View );
 
 	FSceneViewState* ViewState = (FSceneViewState*)Context.View.State;
 
diff --git a/Engine/Source/Runtime/Renderer/Private/PostProcess/VisualizeTexture.cpp b/Engine/Source/Runtime/Renderer/Private/PostProcess/VisualizeTexture.cpp
--- a/Engine/Source/Runtime/Renderer/Private/PostProcess/VisualizeTexture.cpp
+++ b/Engine/Source/Runtime/Renderer/Private/PostProcess/VisualizeTexture.cpp
@@ -347,7 +347,7 @@ void FVisualizeTexture::GenerateContent(FRHICommandListImmediate& RHICmdList, co
 
 	FPooledRenderTargetDesc OutputDesc(FPooledRenderTargetDesc::Create2DDesc(Size, PF_B8G8R8A8, FClearValueBinding::None, TexCreate_None, TexCreate_RenderTargetable | TexCreate_ShaderResource, false));
 	
-	GRenderTargetPool.FindFreeElement(OutputDesc, VisualizeTextureContent, TEXT("VisualizeTexture"));
+	GRenderTargetPool.FindFreeElement(RHICmdList, OutputDesc, VisualizeTextureContent, TEXT("VisualizeTexture"));
 
 	if(!VisualizeTextureContent)
 	{
diff --git a/Engine/Source/Runtime/Renderer/Private/ReflectionEnvironment.cpp b/Engine/Source/Runtime/Renderer/Private/ReflectionEnvironment.cpp
--- a/Engine/Source/Runtime/Renderer/Private/ReflectionEnvironment.cpp
+++ b/Engine/Source/Runtime/Renderer/Private/ReflectionEnvironment.cpp
@@ -89,6 +89,7 @@ void FReflectionEnvironmentCubemapArray::InitDynamicRHI()
 {
 	if (GetFeatureLevel() >= ERHIFeatureLevel::SM5)
 	{
+
 		const int32 NumReflectionCaptureMips = FMath::CeilLogTwo(GReflectionCaptureSize) + 1;
 
 		ReleaseCubeArray();
@@ -104,9 +105,13 @@ void FReflectionEnvironmentCubemapArray::InitDynamicRHI()
 		// Cubemap array of 1 produces a regular cubemap, so guarantee it will be allocated as an array
 		FMath::Max<uint32>(MaxCubemaps, 2),
 		NumReflectionCaptureMips));
+
+		Desc.AutoWritable = false;
 	
+		FRHICommandListImmediate& RHICmdList = FRHICommandListExecutor::GetImmediateCommandList();
+
 		// Allocate TextureCubeArray for the scene's reflection captures
-		GRenderTargetPool.FindFreeElement(Desc, ReflectionEnvs, TEXT("ReflectionEnvs"));
+		GRenderTargetPool.FindFreeElement(RHICmdList, Desc, ReflectionEnvs, TEXT("ReflectionEnvs"));
 	}
 }
 
@@ -852,7 +857,7 @@ void FDeferredShadingSceneRenderer::RenderTiledDeferredImageBasedReflections(FRH
 		Desc.ClearValue = FClearValueBinding::None;
 
 		// we don't create a new name to make it easier to use "vis SceneColor" and get the last HDRSceneColor
-		GRenderTargetPool.FindFreeElement( Desc, NewSceneColor, TEXT("SceneColor") );
+		GRenderTargetPool.FindFreeElement(RHICmdList, Desc, NewSceneColor, TEXT("SceneColor") );
 	}
 
 	// If we are in SM5, use the compute shader gather method
@@ -1080,7 +1085,7 @@ void FDeferredShadingSceneRenderer::RenderStandardDeferredImageBasedReflections(
 			// Apply reflections to screen
 			SCOPED_DRAW_EVENT(RHICmdList, ReflectionApply);
 
-			SceneContext.BeginRenderingSceneColor(RHICmdList, ESimpleRenderTargetMode::EUninitializedColorExistingDepth, FExclusiveDepthStencil::DepthRead_StencilWrite);
+			SceneContext.BeginRenderingSceneColor(RHICmdList, ESimpleRenderTargetMode::EUninitializedColorExistingDepth, FExclusiveDepthStencil::DepthRead_StencilWrite, true);
 
 			RHICmdList.SetViewport(View.ViewRect.Min.X, View.ViewRect.Min.Y, 0.0f, View.ViewRect.Max.X, View.ViewRect.Max.Y, 1.0f);
 			RHICmdList.SetRasterizerState(TStaticRasterizerState<FM_Solid, CM_None>::GetRHI());
diff --git a/Engine/Source/Runtime/Renderer/Private/ReflectionEnvironmentCapture.cpp b/Engine/Source/Runtime/Renderer/Private/ReflectionEnvironmentCapture.cpp
--- a/Engine/Source/Runtime/Renderer/Private/ReflectionEnvironmentCapture.cpp
+++ b/Engine/Source/Runtime/Renderer/Private/ReflectionEnvironmentCapture.cpp
@@ -667,7 +667,7 @@ int32 FindOrAllocateCubemapIndex(FScene* Scene, const UReflectionCaptureComponen
 void ClearScratchCubemaps(FRHICommandList& RHICmdList)
 {
 	FSceneRenderTargets& SceneContext = FSceneRenderTargets::Get(RHICmdList);
-	SceneContext.AllocateReflectionTargets();
+	SceneContext.AllocateReflectionTargets(RHICmdList);
 	// Clear scratch render targets to a consistent but noticeable value
 	// This makes debugging capture issues much easier, otherwise the random contents from previous captures is shown
 
diff --git a/Engine/Source/Runtime/Renderer/Private/Renderer.cpp b/Engine/Source/Runtime/Renderer/Private/Renderer.cpp
--- a/Engine/Source/Runtime/Renderer/Private/Renderer.cpp
+++ b/Engine/Source/Runtime/Renderer/Private/Renderer.cpp
@@ -86,9 +86,9 @@ void FRendererModule::DrawTileMesh(FRHICommandListImmediate& RHICmdList, const F
 	}
 }
 
-void FRendererModule::RenderTargetPoolFindFreeElement(const FPooledRenderTargetDesc& Desc, TRefCountPtr<IPooledRenderTarget> &Out, const TCHAR* InDebugName)
+void FRendererModule::RenderTargetPoolFindFreeElement(FRHICommandListImmediate& RHICmdList, const FPooledRenderTargetDesc& Desc, TRefCountPtr<IPooledRenderTarget> &Out, const TCHAR* InDebugName)
 {
-	GRenderTargetPool.FindFreeElement(Desc, Out, InDebugName);
+	GRenderTargetPool.FindFreeElement(RHICmdList, Desc, Out, InDebugName);
 }
 
 void FRendererModule::TickRenderTargetPool()
diff --git a/Engine/Source/Runtime/Renderer/Private/RendererPrivate.h b/Engine/Source/Runtime/Renderer/Private/RendererPrivate.h
--- a/Engine/Source/Runtime/Renderer/Private/RendererPrivate.h
+++ b/Engine/Source/Runtime/Renderer/Private/RendererPrivate.h
@@ -32,7 +32,7 @@ public:
 	virtual void ReallocateSceneRenderTargets() override;
 	virtual void SceneRenderTargetsSetBufferSize(uint32 SizeX, uint32 SizeY) override;
 	virtual void DrawTileMesh(FRHICommandListImmediate& RHICmdList, const FSceneView& View, const FMeshBatch& Mesh, bool bIsHitTesting, const FHitProxyId& HitProxyId) override;
-	virtual void RenderTargetPoolFindFreeElement(const FPooledRenderTargetDesc& Desc, TRefCountPtr<IPooledRenderTarget> &Out, const TCHAR* InDebugName) override;
+	virtual void RenderTargetPoolFindFreeElement(FRHICommandListImmediate& RHICmdList, const FPooledRenderTargetDesc& Desc, TRefCountPtr<IPooledRenderTarget> &Out, const TCHAR* InDebugName) override;
 	virtual void TickRenderTargetPool() override;
 	virtual void DebugLogOnCrash() override;
 	virtual void GPUBenchmark(FSynthBenchmarkResults& InOut, float WorkScale) override;
diff --git a/Engine/Source/Runtime/Renderer/Private/SceneCaptureRendering.cpp b/Engine/Source/Runtime/Renderer/Private/SceneCaptureRendering.cpp
--- a/Engine/Source/Runtime/Renderer/Private/SceneCaptureRendering.cpp
+++ b/Engine/Source/Runtime/Renderer/Private/SceneCaptureRendering.cpp
@@ -95,7 +95,7 @@ static void UpdateSceneCaptureContent_RenderThread(FRHICommandListImmediate& RHI
 				TexCreate_None, 
 				TexCreate_RenderTargetable,
 				false));
-			GRenderTargetPool.FindFreeElement(Desc, FlippedPooledRenderTarget, TEXT("SceneCaptureFlipped"));
+			GRenderTargetPool.FindFreeElement(RHICmdList, Desc, FlippedPooledRenderTarget, TEXT("SceneCaptureFlipped"));
 		}
 
 		// Helper class to allow setting render target
diff --git a/Engine/Source/Runtime/Renderer/Private/SceneHitProxyRendering.cpp b/Engine/Source/Runtime/Renderer/Private/SceneHitProxyRendering.cpp
--- a/Engine/Source/Runtime/Renderer/Private/SceneHitProxyRendering.cpp
+++ b/Engine/Source/Runtime/Renderer/Private/SceneHitProxyRendering.cpp
@@ -334,7 +334,7 @@ TRefCountPtr<IPooledRenderTarget> InitHitProxyRender(FRHICommandListImmediate& R
 
 	FSceneRenderTargets& SceneContext = FSceneRenderTargets::Get(RHICmdList);
 	// Allocate the maximum scene render target space for the current view family.
-	SceneContext.Allocate(ViewFamily);
+	SceneContext.Allocate(RHICmdList, ViewFamily);
 
 	TRefCountPtr<IPooledRenderTarget> HitProxyRT;
 
@@ -342,7 +342,7 @@ TRefCountPtr<IPooledRenderTarget> InitHitProxyRender(FRHICommandListImmediate& R
 	{
 		FPooledRenderTargetDesc Desc(FPooledRenderTargetDesc::Create2DDesc(SceneContext.GetBufferSizeXY(), PF_B8G8R8A8, FClearValueBinding::None, TexCreate_None, TexCreate_RenderTargetable, false));
 		Desc.Flags |= TexCreate_FastVRAM;
-		GRenderTargetPool.FindFreeElement(Desc, HitProxyRT, TEXT("HitProxy"));
+		GRenderTargetPool.FindFreeElement(RHICmdList, Desc, HitProxyRT, TEXT("HitProxy"));
 	}
 
 	if (!HitProxyRT)
@@ -351,7 +351,7 @@ TRefCountPtr<IPooledRenderTarget> InitHitProxyRender(FRHICommandListImmediate& R
 		return HitProxyRT;
 	}
 
-	SetRenderTarget(RHICmdList, HitProxyRT->GetRenderTargetItem().TargetableTexture, SceneContext.GetSceneDepthSurface(), ESimpleRenderTargetMode::EExistingColorAndDepth, FExclusiveDepthStencil::DepthWrite_StencilWrite);
+	SetRenderTarget(RHICmdList, HitProxyRT->GetRenderTargetItem().TargetableTexture, SceneContext.GetSceneDepthSurface(), ESimpleRenderTargetMode::EExistingColorAndDepth, FExclusiveDepthStencil::DepthWrite_StencilWrite, true);
 
 	// Clear color for each view.
 	auto& Views = SceneRenderer->Views;
diff --git a/Engine/Source/Runtime/Renderer/Private/SceneOcclusion.cpp b/Engine/Source/Runtime/Renderer/Private/SceneOcclusion.cpp
--- a/Engine/Source/Runtime/Renderer/Private/SceneOcclusion.cpp
+++ b/Engine/Source/Runtime/Renderer/Private/SceneOcclusion.cpp
@@ -520,8 +520,9 @@ void FHZBOcclusionTester::InitDynamicRHI()
 {
 	if (GetFeatureLevel() >= ERHIFeatureLevel::SM4)
 	{
+		FRHICommandListImmediate& RHICmdList = FRHICommandListExecutor::GetImmediateCommandList();
 		FPooledRenderTargetDesc Desc( FPooledRenderTargetDesc::Create2DDesc( FIntPoint( SizeX, SizeY ), PF_B8G8R8A8, FClearValueBinding::None, TexCreate_CPUReadback | TexCreate_HideInVisualizeTexture, TexCreate_None, false ) );
-		GRenderTargetPool.FindFreeElement( Desc, ResultsTextureCPU, TEXT("HZBResultsCPU") );
+		GRenderTargetPool.FindFreeElement(RHICmdList, Desc, ResultsTextureCPU, TEXT("HZBResultsCPU") );
 	}
 }
 
@@ -718,14 +719,14 @@ void FHZBOcclusionTester::Submit(FRHICommandListImmediate& RHICmdList, const FVi
 		uint32 Flags = TexCreate_ShaderResource | TexCreate_Dynamic;
 		FPooledRenderTargetDesc Desc( FPooledRenderTargetDesc::Create2DDesc( FIntPoint( SizeX, SizeY ), PF_A32B32G32R32F, FClearValueBinding::None, Flags, TexCreate_None, false ) );
 
-		GRenderTargetPool.FindFreeElement( Desc, BoundsCenterTexture, TEXT("HZBBoundsCenter") );
-		GRenderTargetPool.FindFreeElement( Desc, BoundsExtentTexture, TEXT("HZBBoundsExtent") );
+		GRenderTargetPool.FindFreeElement(RHICmdList, Desc, BoundsCenterTexture, TEXT("HZBBoundsCenter") );
+		GRenderTargetPool.FindFreeElement(RHICmdList, Desc, BoundsExtentTexture, TEXT("HZBBoundsExtent") );
 	}
 
 	TRefCountPtr< IPooledRenderTarget >	ResultsTextureGPU;
 	{
 		FPooledRenderTargetDesc Desc( FPooledRenderTargetDesc::Create2DDesc( FIntPoint( SizeX, SizeY ), PF_B8G8R8A8, FClearValueBinding::None, TexCreate_None, TexCreate_RenderTargetable, false ) );
-		GRenderTargetPool.FindFreeElement( Desc, ResultsTextureGPU, TEXT("HZBResultsGPU") );
+		GRenderTargetPool.FindFreeElement(RHICmdList, Desc, ResultsTextureGPU, TEXT("HZBResultsGPU") );
 	}
 
 	{
@@ -988,10 +989,11 @@ void BuildHZB( FRHICommandListImmediate& RHICmdList, FViewInfo& View )
 
 	FPooledRenderTargetDesc Desc(FPooledRenderTargetDesc::Create2DDesc(HZBSize, PF_R16F, FClearValueBinding::None, TexCreate_None, TexCreate_RenderTargetable | TexCreate_ShaderResource | TexCreate_NoFastClear, false, NumMips));
 	Desc.Flags |= TexCreate_FastVRAM;
-	GRenderTargetPool.FindFreeElement( Desc, View.HZB, TEXT("HZB") );
+	GRenderTargetPool.FindFreeElement(RHICmdList, Desc, View.HZB, TEXT("HZB") );
 	
 	FSceneRenderTargetItem& HZBRenderTarget = View.HZB->GetRenderTargetItem();
 	
+	FTextureRHIParamRef HZBRenderTargetRef = HZBRenderTarget.TargetableTexture.GetReference();
 	// Mip 0
 	{
 		SCOPED_DRAW_EVENTF(RHICmdList, BuildHZB, TEXT("HZB SetupMip 0 %dx%d"), HZBSize.X, HZBSize.Y);
@@ -1024,7 +1026,8 @@ void BuildHZB( FRHICommandListImmediate& RHICmdList, FViewInfo& View )
 			*VertexShader,
 			EDRF_UseTriangleOptimization);
 
-		RHICmdList.CopyToResolveTarget(HZBRenderTarget.TargetableTexture, HZBRenderTarget.ShaderResourceTexture, false, FResolveParams(FResolveRect(), CubeFace_PosX, 0));
+		//Use RWBarrier since we don't transition individual subresources.  Basically treat the whole texture as R/W as we walk down the mip chain.
+		RHICmdList.TransitionResources(EResourceTransitionAccess::ERWBarrier, &HZBRenderTargetRef, 1);
 	}
 
 	FIntPoint SrcSize = HZBSize;
@@ -1065,10 +1068,11 @@ void BuildHZB( FRHICommandListImmediate& RHICmdList, FViewInfo& View )
 			*VertexShader,
 			EDRF_UseTriangleOptimization);
 
-		RHICmdList.CopyToResolveTarget(HZBRenderTarget.TargetableTexture, HZBRenderTarget.ShaderResourceTexture, false, FResolveParams(FResolveRect(), CubeFace_PosX, MipIndex));
-
 		SrcSize /= 2;
 		DstSize /= 2;
+
+		//Use RWBarrier since we don't transition individual subresources.  Basically treat the whole texture as R/W as we walk down the mip chain.
+		RHICmdList.TransitionResources(EResourceTransitionAccess::ERWBarrier, &HZBRenderTargetRef, 1);
 	}
 
 	GRenderTargetPool.VisualizeTexture.SetCheckPoint( RHICmdList, View.HZB );
diff --git a/Engine/Source/Runtime/Renderer/Private/ScenePrivate.h b/Engine/Source/Runtime/Renderer/Private/ScenePrivate.h
--- a/Engine/Source/Runtime/Renderer/Private/ScenePrivate.h
+++ b/Engine/Source/Runtime/Renderer/Private/ScenePrivate.h
@@ -665,13 +665,13 @@ public:
 	 */
 	bool IsShadowOccluded(FRHICommandListImmediate& RHICmdList, FPrimitiveComponentId PrimitiveId, const ULightComponent* Light, int32 SplitIndex, bool bTranslucentShadow, int32 NumBufferedFrames) const;
 
-	TRefCountPtr<IPooledRenderTarget>& GetEyeAdaptation()
+	TRefCountPtr<IPooledRenderTarget>& GetEyeAdaptation(FRHICommandList& RHICmdList)
 	{
 		if (!EyeAdaptationRT)
 		{
 			// Create the texture needed for EyeAdaptation
 			FPooledRenderTargetDesc Desc(FPooledRenderTargetDesc::Create2DDesc(FIntPoint(1, 1), PF_R32_FLOAT, FClearValueBinding::None, TexCreate_None, TexCreate_RenderTargetable, false));
-			GRenderTargetPool.FindFreeElement(Desc, EyeAdaptationRT, TEXT("EyeAdaptation"));
+			GRenderTargetPool.FindFreeElement(RHICmdList, Desc, EyeAdaptationRT, TEXT("EyeAdaptation"));
 		}
 		return EyeAdaptationRT;
 	}
@@ -686,13 +686,13 @@ public:
 		bValidEyeAdaptation = true;
 	}
 
-	TRefCountPtr<IPooledRenderTarget>& GetSeparateTranslucency(FIntPoint Size)
+	TRefCountPtr<IPooledRenderTarget>& GetSeparateTranslucency(FRHICommandList& RHICmdList, FIntPoint Size)
 	{
 		if (!SeparateTranslucencyRT || SeparateTranslucencyRT->GetDesc().Extent != Size)
 		{
 			// Create the SeparateTranslucency render target (alpha is needed to lerping)
 			FPooledRenderTargetDesc Desc(FPooledRenderTargetDesc::Create2DDesc(Size, PF_FloatRGBA, FClearValueBinding::Black, TexCreate_None, TexCreate_RenderTargetable, false));
-			GRenderTargetPool.FindFreeElement(Desc, SeparateTranslucencyRT, TEXT("SeparateTranslucency"));
+			GRenderTargetPool.FindFreeElement(RHICmdList, Desc, SeparateTranslucencyRT, TEXT("SeparateTranslucency"));
 		}
 		return SeparateTranslucencyRT;
 	}
diff --git a/Engine/Source/Runtime/Renderer/Private/SceneRendering.cpp b/Engine/Source/Runtime/Renderer/Private/SceneRendering.cpp
--- a/Engine/Source/Runtime/Renderer/Private/SceneRendering.cpp
+++ b/Engine/Source/Runtime/Renderer/Private/SceneRendering.cpp
@@ -1110,13 +1110,13 @@ FSceneViewState* FViewInfo::GetEffectiveViewState() const
 	return EffectiveViewState;
 }
 
-IPooledRenderTarget* FViewInfo::GetEyeAdaptation() const
+IPooledRenderTarget* FViewInfo::GetEyeAdaptation(FRHICommandList& RHICmdList) const
 {
 	FSceneViewState* EffectiveViewState = GetEffectiveViewState();
 
 	if (EffectiveViewState)
 	{
-		TRefCountPtr<IPooledRenderTarget>& EyeAdaptRef = EffectiveViewState->GetEyeAdaptation();
+		TRefCountPtr<IPooledRenderTarget>& EyeAdaptRef = EffectiveViewState->GetEyeAdaptation(RHICmdList);
 		if( IsValidRef(EyeAdaptRef) )
 		{
 			return EyeAdaptRef.GetReference();
diff --git a/Engine/Source/Runtime/Renderer/Private/SceneRendering.h b/Engine/Source/Runtime/Renderer/Private/SceneRendering.h
--- a/Engine/Source/Runtime/Renderer/Private/SceneRendering.h
+++ b/Engine/Source/Runtime/Renderer/Private/SceneRendering.h
@@ -593,7 +593,7 @@ public:
 	bool IsDistanceCulled(float DistanceSquared, float MaxDrawDistance, float MinDrawDistance, const FPrimitiveSceneInfo* PrimitiveSceneInfo);
 
 	/** Gets the eye adaptation render target for this view. */
-	IPooledRenderTarget* GetEyeAdaptation() const;
+	IPooledRenderTarget* GetEyeAdaptation(FRHICommandList& RHICmdList) const;
 
 	/** Tells if the eyeadaptation texture exists without attempting to allocate it. */
 	bool HasValidEyeAdaptation() const;
diff --git a/Engine/Source/Runtime/Renderer/Private/ShaderBaseClasses.cpp b/Engine/Source/Runtime/Renderer/Private/ShaderBaseClasses.cpp
--- a/Engine/Source/Runtime/Renderer/Private/ShaderBaseClasses.cpp
+++ b/Engine/Source/Runtime/Renderer/Private/ShaderBaseClasses.cpp
@@ -271,7 +271,7 @@ void FMaterialShader::SetParameters(
 	//Use of the eye adaptation texture here is experimental and potentially dangerous as it can introduce a feedback loop. May be removed.
 	if(EyeAdaptation.IsBound())
 	{
-		FTextureRHIRef& EyeAdaptationTex = GetEyeAdaptation(View);
+		FTextureRHIRef& EyeAdaptationTex = GetEyeAdaptation(RHICmdList, View);
 		SetTextureParameter(RHICmdList, ShaderRHI, EyeAdaptation, EyeAdaptationTex);
 	}
 
@@ -355,13 +355,13 @@ bool FMaterialShader::Serialize(FArchive& Ar)
 	return bShaderHasOutdatedParameters;
 }
 
-FTextureRHIRef& FMaterialShader::GetEyeAdaptation(const FSceneView& View)
+FTextureRHIRef& FMaterialShader::GetEyeAdaptation(FRHICommandList& RHICmdList, const FSceneView& View)
 {
 	IPooledRenderTarget* EyeAdaptationRT = NULL;
 	if( View.bIsViewInfo )
 	{
 		const FViewInfo& ViewInfo = static_cast<const FViewInfo&>(View);
-		EyeAdaptationRT = ViewInfo.GetEyeAdaptation();
+		EyeAdaptationRT = ViewInfo.GetEyeAdaptation(RHICmdList);
 	}
 
 	if( EyeAdaptationRT )
diff --git a/Engine/Source/Runtime/Renderer/Private/ShadowRendering.cpp b/Engine/Source/Runtime/Renderer/Private/ShadowRendering.cpp
--- a/Engine/Source/Runtime/Renderer/Private/ShadowRendering.cpp
+++ b/Engine/Source/Runtime/Renderer/Private/ShadowRendering.cpp
@@ -3647,8 +3647,11 @@ bool FDeferredShadingSceneRenderer::RenderCachedPreshadows(FRHICommandListImmedi
 					bool bPerformClear = true;;
 					auto SetShadowRenderTargets = [this, &bPerformClear, ProjectedShadowInfo, &SceneContext](FRHICommandList& InRHICmdList)
 					{
+						FTextureRHIParamRef PreShadowCacheDepthZ = SceneContext.PreShadowCacheDepthZ->GetRenderTargetItem().TargetableTexture.GetReference();
+						InRHICmdList.TransitionResources(EResourceTransitionAccess::EWritable, &PreShadowCacheDepthZ, 1);
+
 						// Must preserve existing contents as the clear will be scissored
-						SetRenderTarget(InRHICmdList, FTextureRHIRef(), SceneContext.PreShadowCacheDepthZ->GetRenderTargetItem().TargetableTexture, ESimpleRenderTargetMode::EExistingColorAndDepth);
+						SetRenderTarget(InRHICmdList, FTextureRHIRef(), PreShadowCacheDepthZ, ESimpleRenderTargetMode::EExistingColorAndDepth);
 						ProjectedShadowInfo->ClearDepth(InRHICmdList, this, bPerformClear);
 					};					
 					SetShadowRenderTargets(RHICmdList); // run it now, maybe run it later for parallel command lists
@@ -3921,7 +3924,7 @@ void FForwardShadingSceneRenderer::RenderModulatedShadowDepthMaps(FRHICommandLis
 	FSceneRenderTargets& SceneContext = FSceneRenderTargets::Get(RHICmdList);
 	TArray<FProjectedShadowInfo*, SceneRenderingAllocator> Shadows;
 	const FIntPoint ShadowBufferResolution = SceneContext.GetShadowDepthTextureResolution();
-	SceneContext.AllocateForwardShadingShadowDepthTarget(ShadowBufferResolution);
+	SceneContext.AllocateForwardShadingShadowDepthTarget(RHICmdList, ShadowBufferResolution);
 	FTextureLayout ShadowLayout(1, 1, ShadowBufferResolution.X, ShadowBufferResolution.Y, false, false);
 	
 	// render shadowmaps for relevant lights.
diff --git a/Engine/Source/Runtime/Renderer/Private/ShadowSetup.cpp b/Engine/Source/Runtime/Renderer/Private/ShadowSetup.cpp
--- a/Engine/Source/Runtime/Renderer/Private/ShadowSetup.cpp
+++ b/Engine/Source/Runtime/Renderer/Private/ShadowSetup.cpp
@@ -2404,7 +2404,7 @@ void FForwardShadingSceneRenderer::InitDynamicShadows(FRHICommandListImmediate&
 				const int32 NumHigh = FMath::Min(((NumWholeSceneShadows - 1) / MaxWide) + 1, MaxHigh);
 
 				const FIntPoint AtlasShadowBufferResolution(ShadowBufferResolution.X * NumWide, ShadowBufferResolution.Y * NumHigh);
-				SceneContext.AllocateForwardShadingShadowDepthTarget(AtlasShadowBufferResolution);
+				SceneContext.AllocateForwardShadingShadowDepthTarget(RHICmdList, AtlasShadowBufferResolution);
 
 				// Allocate atlas shadow texture space to the shadows.
 				FTextureLayout ShadowLayout(1, 1, AtlasShadowBufferResolution.X, AtlasShadowBufferResolution.Y, false, false);
@@ -2428,7 +2428,7 @@ void FForwardShadingSceneRenderer::InitDynamicShadows(FRHICommandListImmediate&
 				// Per obj projected shadows are in use. Ensure the shadow depth target is available for modulated shadow use later.
 				FSceneRenderTargets& SceneContext = FSceneRenderTargets::Get(RHICmdList);
 				const FIntPoint ShadowBufferResolution = SceneContext.GetShadowDepthTextureResolution();
-				SceneContext.AllocateForwardShadingShadowDepthTarget(ShadowBufferResolution);
+				SceneContext.AllocateForwardShadingShadowDepthTarget(RHICmdList, ShadowBufferResolution);
 			}
 		}
 
diff --git a/Engine/Source/Runtime/Renderer/Private/SystemTextures.cpp b/Engine/Source/Runtime/Renderer/Private/SystemTextures.cpp
--- a/Engine/Source/Runtime/Renderer/Private/SystemTextures.cpp
+++ b/Engine/Source/Runtime/Renderer/Private/SystemTextures.cpp
@@ -29,7 +29,8 @@ void FSystemTextures::InitializeTextures(FRHICommandListImmediate& RHICmdList, E
 		// Create a WhiteDummy texture
 		{
 			FPooledRenderTargetDesc Desc(FPooledRenderTargetDesc::Create2DDesc(FIntPoint(1, 1), PF_B8G8R8A8, FClearValueBinding::White, TexCreate_HideInVisualizeTexture, TexCreate_RenderTargetable | TexCreate_NoFastClear, false));
-			GRenderTargetPool.FindFreeElement(Desc, WhiteDummy, TEXT("WhiteDummy"));
+			Desc.AutoWritable = false;
+			GRenderTargetPool.FindFreeElement(RHICmdList, Desc, WhiteDummy, TEXT("WhiteDummy"));
 
 			SetRenderTarget(RHICmdList, WhiteDummy->GetRenderTargetItem().TargetableTexture, FTextureRHIRef(), ESimpleRenderTargetMode::EClearColorExistingDepth);
 			RHICmdList.CopyToResolveTarget(WhiteDummy->GetRenderTargetItem().TargetableTexture, WhiteDummy->GetRenderTargetItem().ShaderResourceTexture, true, FResolveParams());
@@ -38,7 +39,8 @@ void FSystemTextures::InitializeTextures(FRHICommandListImmediate& RHICmdList, E
 		// Create a BlackDummy texture
 		{
 			FPooledRenderTargetDesc Desc(FPooledRenderTargetDesc::Create2DDesc(FIntPoint(1, 1), PF_B8G8R8A8, FClearValueBinding::Black, TexCreate_HideInVisualizeTexture, TexCreate_RenderTargetable | TexCreate_NoFastClear, false));
-			GRenderTargetPool.FindFreeElement(Desc, BlackDummy, TEXT("BlackDummy"));
+			Desc.AutoWritable = false;
+			GRenderTargetPool.FindFreeElement(RHICmdList, Desc, BlackDummy, TEXT("BlackDummy"));
 
 			SetRenderTarget(RHICmdList, BlackDummy->GetRenderTargetItem().TargetableTexture, FTextureRHIRef(), ESimpleRenderTargetMode::EClearColorExistingDepth);
 			RHICmdList.CopyToResolveTarget(BlackDummy->GetRenderTargetItem().TargetableTexture, BlackDummy->GetRenderTargetItem().ShaderResourceTexture, true, FResolveParams());
@@ -47,7 +49,8 @@ void FSystemTextures::InitializeTextures(FRHICommandListImmediate& RHICmdList, E
 		// Create a BlackAlphaOneDummy texture
 		{
 			FPooledRenderTargetDesc Desc(FPooledRenderTargetDesc::Create2DDesc(FIntPoint(1, 1), PF_B8G8R8A8, FClearValueBinding::None, TexCreate_HideInVisualizeTexture, TexCreate_RenderTargetable | TexCreate_NoFastClear, false));
-			GRenderTargetPool.FindFreeElement(Desc, BlackAlphaOneDummy, TEXT("BlackAlphaOneDummy"));
+			Desc.AutoWritable = false;
+			GRenderTargetPool.FindFreeElement(RHICmdList, Desc, BlackAlphaOneDummy, TEXT("BlackAlphaOneDummy"));
 
 			SetRenderTarget(RHICmdList, BlackAlphaOneDummy->GetRenderTargetItem().TargetableTexture, FTextureRHIRef());
 			RHICmdList.Clear(true, FLinearColor(0, 0, 0, 1), false, 0, false, 0, FIntRect());
@@ -57,7 +60,8 @@ void FSystemTextures::InitializeTextures(FRHICommandListImmediate& RHICmdList, E
 		// Create the PerlinNoiseGradient texture
 		{
 			FPooledRenderTargetDesc Desc(FPooledRenderTargetDesc::Create2DDesc(FIntPoint(128, 128), PF_B8G8R8A8, FClearValueBinding::None, TexCreate_HideInVisualizeTexture, TexCreate_None | TexCreate_NoFastClear, false));
-			GRenderTargetPool.FindFreeElement(Desc, PerlinNoiseGradient, TEXT("PerlinNoiseGradient"));
+			Desc.AutoWritable = false;
+			GRenderTargetPool.FindFreeElement(RHICmdList, Desc, PerlinNoiseGradient, TEXT("PerlinNoiseGradient"));
 			// Write the contents of the texture.
 			uint32 DestStride;
 			uint8* DestBuffer = (uint8*)RHICmdList.LockTexture2D((FTexture2DRHIRef&)PerlinNoiseGradient->GetRenderTargetItem().ShaderResourceTexture, 0, RLM_WriteOnly, DestStride, false);
@@ -88,7 +92,8 @@ void FSystemTextures::InitializeTextures(FRHICommandListImmediate& RHICmdList, E
 		if (!GSupportsShaderFramebufferFetch && GPixelFormats[PF_FloatRGBA].Supported)
 		{
 			FPooledRenderTargetDesc Desc(FPooledRenderTargetDesc::Create2DDesc(FIntPoint(1, 1), PF_FloatRGBA, FClearValueBinding::None, TexCreate_HideInVisualizeTexture, TexCreate_RenderTargetable | TexCreate_NoFastClear, false));
-			GRenderTargetPool.FindFreeElement(Desc, MaxFP16Depth, TEXT("MaxFP16Depth"));
+			Desc.AutoWritable = false;
+			GRenderTargetPool.FindFreeElement(RHICmdList, Desc, MaxFP16Depth, TEXT("MaxFP16Depth"));
 
 			SetRenderTarget(RHICmdList, MaxFP16Depth->GetRenderTargetItem().TargetableTexture, FTextureRHIRef());
 			RHICmdList.Clear(true, FLinearColor(65000.0f, 65000.0f, 65000.0f, 65000.0f), false, 0, false, 0, FIntRect());
@@ -99,7 +104,8 @@ void FSystemTextures::InitializeTextures(FRHICommandListImmediate& RHICmdList, E
 		if (IsMobilePlatform(GShaderPlatformForFeatureLevel[InFeatureLevel])) // currently used only for mobile
 		{
 			FPooledRenderTargetDesc Desc(FPooledRenderTargetDesc::Create2DDesc(FIntPoint(1, 1), PF_DepthStencil, FClearValueBinding::DepthFar, TexCreate_None, TexCreate_DepthStencilTargetable, false));
-			GRenderTargetPool.FindFreeElement(Desc, DepthDummy, TEXT("DepthDummy"));
+			Desc.AutoWritable = false;
+			GRenderTargetPool.FindFreeElement(RHICmdList, Desc, DepthDummy, TEXT("DepthDummy"));
 
 			SetRenderTarget(RHICmdList, FTextureRHIRef(), DepthDummy->GetRenderTargetItem().TargetableTexture);
 			RHICmdList.Clear(false, FLinearColor::White, true, (float)ERHIZBuffer::FarPlane, true, 0, FIntRect());
@@ -115,7 +121,8 @@ void FSystemTextures::InitializeTextures(FRHICommandListImmediate& RHICmdList, E
 		const uint32 Square = Extent * Extent;
 
 		FPooledRenderTargetDesc Desc(FPooledRenderTargetDesc::CreateVolumeDesc(Extent, Extent, Extent, PF_B8G8R8A8, FClearValueBinding::None, TexCreate_ShaderResource | TexCreate_HideInVisualizeTexture | TexCreate_NoTiling, TexCreate_None, false));
-		GRenderTargetPool.FindFreeElement(Desc, PerlinNoise3D, TEXT("PerlinNoise3D"));
+		Desc.AutoWritable = false;
+		GRenderTargetPool.FindFreeElement(RHICmdList, Desc, PerlinNoise3D, TEXT("PerlinNoise3D"));
 		// Write the contents of the texture.
 		TArray<uint32> DestBuffer;
 
@@ -248,7 +255,8 @@ void FSystemTextures::InitializeTextures(FRHICommandListImmediate& RHICmdList, E
 
 		{
 			FPooledRenderTargetDesc Desc(FPooledRenderTargetDesc::Create2DDesc(FIntPoint(64, 64), PF_B8G8R8A8, FClearValueBinding::None, TexCreate_HideInVisualizeTexture, TexCreate_None | TexCreate_NoFastClear, false));
-			GRenderTargetPool.FindFreeElement(Desc, SSAORandomization, TEXT("SSAORandomization"));
+			Desc.AutoWritable = false;
+			GRenderTargetPool.FindFreeElement(RHICmdList, Desc, SSAORandomization, TEXT("SSAORandomization"));
 			// Write the contents of the texture.
 			uint32 DestStride;
 			uint8* DestBuffer = (uint8*)RHICmdList.LockTexture2D((FTexture2DRHIRef&)SSAORandomization->GetRenderTargetItem().ShaderResourceTexture, 0, RLM_WriteOnly, DestStride, false);
@@ -279,14 +287,14 @@ void FSystemTextures::InitializeTextures(FRHICommandListImmediate& RHICmdList, E
 			}
 
 			FPooledRenderTargetDesc Desc(FPooledRenderTargetDesc::Create2DDesc(FIntPoint(128, 32), Format, FClearValueBinding::None, TexCreate_FastVRAM, TexCreate_None, false));
-
+			Desc.AutoWritable = false;
 			if (bReference)
 			{
 				Desc.Extent.X = 128;
 				Desc.Extent.Y = 128;
 			}
 
-			GRenderTargetPool.FindFreeElement(Desc, PreintegratedGF, TEXT("PreintegratedGF"));
+			GRenderTargetPool.FindFreeElement(RHICmdList, Desc, PreintegratedGF, TEXT("PreintegratedGF"));
 			// Write the contents of the texture.
 			uint32 DestStride;
 			uint8* DestBuffer = (uint8*)RHICmdList.LockTexture2D((FTexture2DRHIRef&)PreintegratedGF->GetRenderTargetItem().ShaderResourceTexture, 0, RLM_WriteOnly, DestStride, false);
diff --git a/Engine/Source/Runtime/Renderer/Private/TiledDeferredLightRendering.cpp b/Engine/Source/Runtime/Renderer/Private/TiledDeferredLightRendering.cpp
--- a/Engine/Source/Runtime/Renderer/Private/TiledDeferredLightRendering.cpp
+++ b/Engine/Source/Runtime/Renderer/Private/TiledDeferredLightRendering.cpp
@@ -348,7 +348,7 @@ void FDeferredShadingSceneRenderer::RenderTiledDeferredLighting(FRHICommandListI
 				FPooledRenderTargetDesc Desc = SceneContext.GetSceneColor()->GetDesc();
 				Desc.TargetableFlags |= TexCreate_UAV;
 
-				GRenderTargetPool.FindFreeElement( Desc, OutTexture, TEXT("SceneColorTiled") );
+				GRenderTargetPool.FindFreeElement(RHICmdList, Desc, OutTexture, TEXT("SceneColorTiled") );
 			}
 
 			{
diff --git a/Engine/Source/Runtime/Renderer/Private/TranslucentLighting.cpp b/Engine/Source/Runtime/Renderer/Private/TranslucentLighting.cpp
--- a/Engine/Source/Runtime/Renderer/Private/TranslucentLighting.cpp
+++ b/Engine/Source/Runtime/Renderer/Private/TranslucentLighting.cpp
@@ -1229,8 +1229,8 @@ void FDeferredShadingSceneRenderer::InjectAmbientCubemapTranslucentVolumeLightin
 				RasterizeToVolumeTexture(RHICmdList, VolumeBounds);
 			}
 
-			RHICmdList.CopyToResolveTarget(SceneContext.TranslucencyLightingVolumeAmbient[VolumeCascadeIndex]->GetRenderTargetItem().TargetableTexture,
-				SceneContext.TranslucencyLightingVolumeAmbient[VolumeCascadeIndex]->GetRenderTargetItem().ShaderResourceTexture, true, FResolveParams());
+			//RHICmdList.CopyToResolveTarget(SceneContext.TranslucencyLightingVolumeAmbient[VolumeCascadeIndex]->GetRenderTargetItem().TargetableTexture,
+				//SceneContext.TranslucencyLightingVolumeAmbient[VolumeCascadeIndex]->GetRenderTargetItem().ShaderResourceTexture, true, FResolveParams());
 		}
 	}
 }
@@ -1604,8 +1604,8 @@ static void InjectTranslucentLightArray(FRHICommandListImmediate& RHICmdList, co
 			}
 		}
 
-		RHICmdList.CopyToResolveTarget(RT0->GetRenderTargetItem().TargetableTexture, RT0->GetRenderTargetItem().ShaderResourceTexture, true, FResolveParams());
-		RHICmdList.CopyToResolveTarget(RT1->GetRenderTargetItem().TargetableTexture, RT1->GetRenderTargetItem().ShaderResourceTexture, true, FResolveParams());
+		//RHICmdList.CopyToResolveTarget(RT0->GetRenderTargetItem().TargetableTexture, RT0->GetRenderTargetItem().ShaderResourceTexture, true, FResolveParams());
+		//RHICmdList.CopyToResolveTarget(RT1->GetRenderTargetItem().TargetableTexture, RT1->GetRenderTargetItem().ShaderResourceTexture, true, FResolveParams());
 	}
 }
 
@@ -1787,8 +1787,8 @@ void FDeferredShadingSceneRenderer::InjectSimpleTranslucentVolumeLightingArray(F
 				}
 			}
 
-			RHICmdList.CopyToResolveTarget(RT0->GetRenderTargetItem().TargetableTexture, RT0->GetRenderTargetItem().ShaderResourceTexture, true, FResolveParams());
-			RHICmdList.CopyToResolveTarget(RT1->GetRenderTargetItem().TargetableTexture, RT1->GetRenderTargetItem().ShaderResourceTexture, true, FResolveParams());
+			//RHICmdList.CopyToResolveTarget(RT0->GetRenderTargetItem().TargetableTexture, RT0->GetRenderTargetItem().ShaderResourceTexture, true, FResolveParams());
+			//RHICmdList.CopyToResolveTarget(RT1->GetRenderTargetItem().TargetableTexture, RT1->GetRenderTargetItem().ShaderResourceTexture, true, FResolveParams());
 		}
 	}
 }
@@ -1800,6 +1800,7 @@ void FDeferredShadingSceneRenderer::FilterTranslucentVolumeLighting(FRHICommandL
 	if (GUseTranslucentLightingVolumes && GSupportsVolumeTextureRendering)
 	{
 		FSceneRenderTargets& SceneContext = FSceneRenderTargets::Get(RHICmdList);
+#if 0
 		// textures have to be finalized before reading.
 		for (int32 VolumeCascadeIndex = 0; VolumeCascadeIndex < TVC_MAX; VolumeCascadeIndex++)
 		{
@@ -1810,6 +1811,7 @@ void FDeferredShadingSceneRenderer::FilterTranslucentVolumeLighting(FRHICommandL
 			RHICmdList.CopyToResolveTarget(TargetTexture0, TargetTexture0, true, FResolveParams());
 			RHICmdList.CopyToResolveTarget(TargetTexture1, TargetTexture1, true, FResolveParams());
 		}
+#endif
 
 		if (GUseTranslucencyVolumeBlur)
 		{
@@ -1825,6 +1827,9 @@ void FDeferredShadingSceneRenderer::FilterTranslucentVolumeLighting(FRHICommandL
 				const IPooledRenderTarget* RT0 = SceneContext.GetTranslucencyVolumeAmbient((ETranslucencyVolumeCascade)VolumeCascadeIndex);
 				const IPooledRenderTarget* RT1 = SceneContext.GetTranslucencyVolumeDirectional((ETranslucencyVolumeCascade)VolumeCascadeIndex);
 
+				const IPooledRenderTarget* Input0 = SceneContext.TranslucencyLightingVolumeAmbient[VolumeCascadeIndex];
+				const IPooledRenderTarget* Input1 = SceneContext.TranslucencyLightingVolumeDirectional[VolumeCascadeIndex];								
+
 				GRenderTargetPool.VisualizeTexture.SetCheckPoint(RHICmdList, RT0);
 				GRenderTargetPool.VisualizeTexture.SetCheckPoint(RHICmdList, RT1);
 
@@ -1832,6 +1837,19 @@ void FDeferredShadingSceneRenderer::FilterTranslucentVolumeLighting(FRHICommandL
 				RenderTargets[0] = RT0->GetRenderTargetItem().TargetableTexture;
 				RenderTargets[1] = RT1->GetRenderTargetItem().TargetableTexture;
 
+				FTextureRHIParamRef Inputs[2];
+				Inputs[0] = Input0->GetRenderTargetItem().TargetableTexture;
+				Inputs[1] = Input1->GetRenderTargetItem().TargetableTexture;
+
+				static_assert(TVC_MAX == 2, "Final transition logic should change");
+
+				//the volume textures should still be writable from the injection phase on the first loop.
+				if (VolumeCascadeIndex > 0)
+				{
+					RHICmdList.TransitionResources(EResourceTransitionAccess::EWritable, RenderTargets, 2);
+				}
+				RHICmdList.TransitionResources(EResourceTransitionAccess::EReadable, Inputs, 2);
+
 				SetRenderTargets(RHICmdList, ARRAY_COUNT(RenderTargets), RenderTargets, FTextureRHIRef(), 0, NULL);
 
 
@@ -1855,8 +1873,11 @@ void FDeferredShadingSceneRenderer::FilterTranslucentVolumeLighting(FRHICommandL
 
 				RasterizeToVolumeTexture(RHICmdList, VolumeBounds);
 
-				RHICmdList.CopyToResolveTarget(RT0->GetRenderTargetItem().TargetableTexture, RT0->GetRenderTargetItem().ShaderResourceTexture, true, FResolveParams());
-				RHICmdList.CopyToResolveTarget(RT1->GetRenderTargetItem().TargetableTexture, RT1->GetRenderTargetItem().ShaderResourceTexture, true, FResolveParams());
+				//only do readable transition on the final loop since the other ones will do this up front.
+				//if (VolumeCascadeIndex == TVC_MAX - 1)
+				{					
+					RHICmdList.TransitionResources(EResourceTransitionAccess::EReadable, RenderTargets, 2);
+				}
 			}
 		}
 	}
diff --git a/Engine/Source/Runtime/Renderer/Private/TranslucentRendering.cpp b/Engine/Source/Runtime/Renderer/Private/TranslucentRendering.cpp
--- a/Engine/Source/Runtime/Renderer/Private/TranslucentRendering.cpp
+++ b/Engine/Source/Runtime/Renderer/Private/TranslucentRendering.cpp
@@ -916,7 +916,7 @@ static FTranslucencyDrawingPolicyFactory::ContextType GParallelTranslucencyConte
 void FDeferredShadingSceneRenderer::RenderTranslucencyParallel(FRHICommandListImmediate& RHICmdList)
 {
 	FSceneRenderTargets& SceneContext = FSceneRenderTargets::Get(RHICmdList);
-	SceneContext.AllocLightAttenuation(); // materials will attempt to get this texture before the deferred command to set it up executes
+	SceneContext.AllocLightAttenuation(RHICmdList); // materials will attempt to get this texture before the deferred command to set it up executes
 	check(IsInRenderingThread());
 
 
@@ -935,7 +935,7 @@ void FDeferredShadingSceneRenderer::RenderTranslucencyParallel(FRHICommandListIm
 			{
 				FSceneViewState* ViewState = (FSceneViewState*)View.State;
 				// we need to allocate this now so it ends up in the snapshot
-				ViewState->GetSeparateTranslucency(SceneContext.GetBufferSizeXY());
+				ViewState->GetSeparateTranslucency(RHICmdList, SceneContext.GetBufferSizeXY());
 			}
 			FTranslucencyPassParallelCommandListSet ParallelCommandListSet(View, RHICmdList, 
 				CVarRHICmdTranslucencyPassDeferredContexts.GetValueOnRenderThread() > 0, 
diff --git a/Engine/Source/Runtime/Renderer/Private/VelocityRendering.cpp b/Engine/Source/Runtime/Renderer/Private/VelocityRendering.cpp
--- a/Engine/Source/Runtime/Renderer/Private/VelocityRendering.cpp
+++ b/Engine/Source/Runtime/Renderer/Private/VelocityRendering.cpp
@@ -736,7 +736,7 @@ void FDeferredShadingSceneRenderer::RenderVelocities(FRHICommandListImmediate& R
 	SCOPED_DRAW_EVENT(RHICmdList, RenderVelocities);
 
 	FPooledRenderTargetDesc Desc = FVelocityRendering::GetRenderTargetDesc();
-	GRenderTargetPool.FindFreeElement(Desc, VelocityRT, TEXT("Velocity"));
+	GRenderTargetPool.FindFreeElement(RHICmdList, Desc, VelocityRT, TEXT("Velocity"));
 
 	{
 		static const auto MotionBlurDebugVar = IConsoleManager::Get().FindTConsoleVariableDataInt(TEXT("r.MotionBlurDebug"));
diff --git a/Engine/Source/Runtime/Renderer/Public/MaterialShader.h b/Engine/Source/Runtime/Renderer/Public/MaterialShader.h
--- a/Engine/Source/Runtime/Renderer/Public/MaterialShader.h
+++ b/Engine/Source/Runtime/Renderer/Public/MaterialShader.h
@@ -130,7 +130,7 @@ public:
 		bool bDeferredPass, 
 		ESceneRenderTargetsMode::Type TextureMode );
 
-	FTextureRHIRef& GetEyeAdaptation(const FSceneView& View);
+	FTextureRHIRef& GetEyeAdaptation(FRHICommandList& RHICmdList, const FSceneView& View);
 
 	// FShader interface.
 	virtual bool Serialize(FArchive& Ar) override;
diff --git a/Engine/Source/Runtime/Renderer/Public/RendererInterface.h b/Engine/Source/Runtime/Renderer/Public/RendererInterface.h
--- a/Engine/Source/Runtime/Renderer/Public/RendererInterface.h
+++ b/Engine/Source/Runtime/Renderer/Public/RendererInterface.h
@@ -46,6 +46,7 @@ public:
 		, TargetableFlags(TexCreate_None)
 		, bForceSeparateTargetAndShaderResource(false)
 		, DebugName(TEXT("UnknownTexture"))
+		, AutoWritable(true)
 	{
 		check(!IsValid());
 	}
@@ -175,7 +176,8 @@ public:
 			&& LhsFlags == RhsFlags
 			&& TargetableFlags == rhs.TargetableFlags
 			&& bForceSeparateTargetAndShaderResource == rhs.bForceSeparateTargetAndShaderResource
-			&& ClearValue == rhs.ClearValue;
+			&& ClearValue == rhs.ClearValue
+			&& AutoWritable == AutoWritable;
 	}
 
 	bool IsCubemap() const
@@ -285,6 +287,7 @@ public:
 		NumSamples = 1;
 
 		bForceSeparateTargetAndShaderResource = false;
+		AutoWritable = true;
 
 		// Remove UAV flag for rendertargets that don't need it (some formats are incompatible)
 		TargetableFlags &= (~TexCreate_UAV);
@@ -317,6 +320,8 @@ public:
 	bool bForceSeparateTargetAndShaderResource;
 	/** only set a pointer to memory that never gets released */
 	const TCHAR *DebugName;
+	/** automatically set to writable via barrier during */
+	bool AutoWritable;
 };
 
 
@@ -541,7 +546,7 @@ public:
 	virtual void DrawTileMesh(FRHICommandListImmediate& RHICmdList, const FSceneView& View, const FMeshBatch& Mesh, bool bIsHitTesting, const class FHitProxyId& HitProxyId) = 0;
 
 	/** Render thread side, use TRefCountPtr<IPooledRenderTarget>, allows to use sharing and VisualizeTexture */
-	virtual void RenderTargetPoolFindFreeElement(const FPooledRenderTargetDesc& Desc, TRefCountPtr<IPooledRenderTarget> &Out, const TCHAR* InDebugName) = 0;
+	virtual void RenderTargetPoolFindFreeElement(FRHICommandListImmediate& RHICmdList, const FPooledRenderTargetDesc& Desc, TRefCountPtr<IPooledRenderTarget> &Out, const TCHAR* InDebugName) = 0;
 	
 	/** Render thread side, to age the pool elements so they get released at some point */
 	virtual void TickRenderTargetPool() = 0;
diff --git a/Engine/Source/Runtime/SlateRHIRenderer/Private/SlateRHIRenderer.cpp b/Engine/Source/Runtime/SlateRHIRenderer/Private/SlateRHIRenderer.cpp
--- a/Engine/Source/Runtime/SlateRHIRenderer/Private/SlateRHIRenderer.cpp
+++ b/Engine/Source/Runtime/SlateRHIRenderer/Private/SlateRHIRenderer.cpp
@@ -427,7 +427,7 @@ void FSlateRHIRenderer::DrawWindow_RenderThread(FRHICommandListImmediate& RHICmd
 
 		RHICmdList.BeginDrawingViewport( ViewportInfo.ViewportRHI, FTextureRHIRef() );
 		RHICmdList.SetViewport(0, 0, 0, ViewportWidth, ViewportHeight, 0.0f);
-
+		RHICmdList.TransitionResource(EResourceTransitionAccess::EWritable, BackBuffer);
 		if( ViewportInfo.bRequiresStencilTest )
 		{
 			check(IsValidRef( ViewportInfo.DepthStencil ));
diff --git a/Engine/Source/Runtime/Windows/D3D11RHI/Private/D3D11Commands.cpp b/Engine/Source/Runtime/Windows/D3D11RHI/Private/D3D11Commands.cpp
--- a/Engine/Source/Runtime/Windows/D3D11RHI/Private/D3D11Commands.cpp
+++ b/Engine/Source/Runtime/Windows/D3D11RHI/Private/D3D11Commands.cpp
@@ -261,10 +261,14 @@ void FD3D11DynamicRHI::RHISetShaderTexture(FVertexShaderRHIParamRef VertexShader
 	FD3D11TextureBase* NewTexture = GetD3D11TextureFromRHITexture(NewTextureRHI);
 	ID3D11ShaderResourceView* ShaderResourceView = NewTexture ? NewTexture->GetShaderResourceView() : NULL;
 
-	if (  ( NewTexture == NULL) || ( NewTexture->GetRenderTargetView( 0, 0 ) !=NULL) || ( NewTexture->HasDepthStencilView()) )
-		SetShaderResourceView<SF_Vertex>(NewTexture, ShaderResourceView, TextureIndex, FD3D11StateCache::SRV_Dynamic);
+	if ((NewTexture == NULL) || (NewTexture->GetRenderTargetView(0, 0) != NULL) || (NewTexture->HasDepthStencilView()))
+	{
+		SetShaderResourceView<SF_Vertex>(NewTexture, ShaderResourceView, TextureIndex, NewTextureRHI ? NewTextureRHI->GetName() : NAME_None, FD3D11StateCache::SRV_Dynamic);
+	}
 	else
-		SetShaderResourceView<SF_Vertex>(NewTexture, ShaderResourceView, TextureIndex, FD3D11StateCache::SRV_Static);
+	{
+		SetShaderResourceView<SF_Vertex>(NewTexture, ShaderResourceView, TextureIndex, NewTextureRHI->GetName(), FD3D11StateCache::SRV_Static);
+	}
 }
 
 void FD3D11DynamicRHI::RHISetShaderTexture(FHullShaderRHIParamRef HullShaderRHI,uint32 TextureIndex,FTextureRHIParamRef NewTextureRHI)
@@ -274,10 +278,14 @@ void FD3D11DynamicRHI::RHISetShaderTexture(FHullShaderRHIParamRef HullShaderRHI,
 	FD3D11TextureBase* NewTexture = GetD3D11TextureFromRHITexture(NewTextureRHI);
 	ID3D11ShaderResourceView* ShaderResourceView = NewTexture ? NewTexture->GetShaderResourceView() : NULL;
 
-	if (  ( NewTexture == NULL) || ( NewTexture->GetRenderTargetView( 0, 0 ) !=NULL) || ( NewTexture->HasDepthStencilView()) )
-		SetShaderResourceView<SF_Hull>(NewTexture, ShaderResourceView, TextureIndex, FD3D11StateCache::SRV_Dynamic);
+	if ((NewTexture == NULL) || (NewTexture->GetRenderTargetView(0, 0) != NULL) || (NewTexture->HasDepthStencilView()))
+	{
+		SetShaderResourceView<SF_Hull>(NewTexture, ShaderResourceView, TextureIndex, NewTextureRHI ? NewTextureRHI->GetName() : NAME_None, FD3D11StateCache::SRV_Dynamic);
+	}
 	else
-		SetShaderResourceView<SF_Hull>(NewTexture, ShaderResourceView, TextureIndex, FD3D11StateCache::SRV_Static);
+	{
+		SetShaderResourceView<SF_Hull>(NewTexture, ShaderResourceView, TextureIndex, NewTextureRHI->GetName(), FD3D11StateCache::SRV_Static);
+	}
 }
 
 void FD3D11DynamicRHI::RHISetShaderTexture(FDomainShaderRHIParamRef DomainShaderRHI,uint32 TextureIndex,FTextureRHIParamRef NewTextureRHI)
@@ -287,10 +295,14 @@ void FD3D11DynamicRHI::RHISetShaderTexture(FDomainShaderRHIParamRef DomainShader
 	FD3D11TextureBase* NewTexture = GetD3D11TextureFromRHITexture(NewTextureRHI);
 	ID3D11ShaderResourceView* ShaderResourceView = NewTexture ? NewTexture->GetShaderResourceView() : NULL;
 
-	if (  ( NewTexture == NULL) || ( NewTexture->GetRenderTargetView( 0, 0 ) !=NULL) || ( NewTexture->HasDepthStencilView()) )
-		SetShaderResourceView<SF_Domain>(NewTexture, ShaderResourceView, TextureIndex, FD3D11StateCache::SRV_Dynamic);
+	if ((NewTexture == NULL) || (NewTexture->GetRenderTargetView(0, 0) != NULL) || (NewTexture->HasDepthStencilView()))
+	{
+		SetShaderResourceView<SF_Domain>(NewTexture, ShaderResourceView, TextureIndex, NewTextureRHI ? NewTextureRHI->GetName() : NAME_None, FD3D11StateCache::SRV_Dynamic);
+	}
 	else
-		SetShaderResourceView<SF_Domain>(NewTexture, ShaderResourceView, TextureIndex, FD3D11StateCache::SRV_Static);
+	{
+		SetShaderResourceView<SF_Domain>(NewTexture, ShaderResourceView, TextureIndex, NewTextureRHI->GetName(), FD3D11StateCache::SRV_Static);
+	}
 }
 
 void FD3D11DynamicRHI::RHISetShaderTexture(FGeometryShaderRHIParamRef GeometryShaderRHI,uint32 TextureIndex,FTextureRHIParamRef NewTextureRHI)
@@ -300,10 +312,14 @@ void FD3D11DynamicRHI::RHISetShaderTexture(FGeometryShaderRHIParamRef GeometrySh
 	FD3D11TextureBase* NewTexture = GetD3D11TextureFromRHITexture(NewTextureRHI);
 	ID3D11ShaderResourceView* ShaderResourceView = NewTexture ? NewTexture->GetShaderResourceView() : NULL;
 
-	if (  ( NewTexture == NULL) || ( NewTexture->GetRenderTargetView( 0, 0 ) !=NULL) || ( NewTexture->HasDepthStencilView()) )
-		SetShaderResourceView<SF_Geometry>(NewTexture, ShaderResourceView, TextureIndex, FD3D11StateCache::SRV_Dynamic);
+	if ((NewTexture == NULL) || (NewTexture->GetRenderTargetView(0, 0) != NULL) || (NewTexture->HasDepthStencilView()))
+	{
+		SetShaderResourceView<SF_Geometry>(NewTexture, ShaderResourceView, TextureIndex, NewTextureRHI ? NewTextureRHI->GetName() : NAME_None, FD3D11StateCache::SRV_Dynamic);
+	}
 	else
-		SetShaderResourceView<SF_Geometry>(NewTexture, ShaderResourceView, TextureIndex, FD3D11StateCache::SRV_Static);
+	{
+		SetShaderResourceView<SF_Geometry>(NewTexture, ShaderResourceView, TextureIndex, NewTextureRHI->GetName(), FD3D11StateCache::SRV_Static);
+	}
 }
 
 void FD3D11DynamicRHI::RHISetShaderTexture(FPixelShaderRHIParamRef PixelShaderRHI,uint32 TextureIndex,FTextureRHIParamRef NewTextureRHI)
@@ -313,10 +329,14 @@ void FD3D11DynamicRHI::RHISetShaderTexture(FPixelShaderRHIParamRef PixelShaderRH
 
 	FD3D11TextureBase* NewTexture = GetD3D11TextureFromRHITexture(NewTextureRHI);
 	ID3D11ShaderResourceView* ShaderResourceView = NewTexture ? NewTexture->GetShaderResourceView() : NULL;
-	if ( ( NewTexture == NULL) ||  ( NewTexture->GetRenderTargetView( 0, 0 ) !=NULL) || ( NewTexture->HasDepthStencilView()) )
-		SetShaderResourceView<SF_Pixel>(NewTexture, ShaderResourceView, TextureIndex, FD3D11StateCache::SRV_Dynamic);
+	if ((NewTexture == NULL) || (NewTexture->GetRenderTargetView(0, 0) != NULL) || (NewTexture->HasDepthStencilView()))
+	{
+		SetShaderResourceView<SF_Pixel>(NewTexture, ShaderResourceView, TextureIndex, NewTextureRHI ? NewTextureRHI->GetName() : NAME_None, FD3D11StateCache::SRV_Dynamic);
+	}
 	else
-		SetShaderResourceView<SF_Pixel>(NewTexture, ShaderResourceView, TextureIndex, FD3D11StateCache::SRV_Static);
+	{
+		SetShaderResourceView<SF_Pixel>(NewTexture, ShaderResourceView, TextureIndex, NewTextureRHI->GetName(), FD3D11StateCache::SRV_Static);
+	}
 }
 
 void FD3D11DynamicRHI::RHISetShaderTexture(FComputeShaderRHIParamRef ComputeShaderRHI,uint32 TextureIndex,FTextureRHIParamRef NewTextureRHI)
@@ -326,10 +346,14 @@ void FD3D11DynamicRHI::RHISetShaderTexture(FComputeShaderRHIParamRef ComputeShad
 	FD3D11TextureBase* NewTexture = GetD3D11TextureFromRHITexture(NewTextureRHI);
 	ID3D11ShaderResourceView* ShaderResourceView = NewTexture ? NewTexture->GetShaderResourceView() : NULL;
 
-	if ( ( NewTexture == NULL) || ( NewTexture->GetRenderTargetView( 0, 0 ) !=NULL) || ( NewTexture->HasDepthStencilView()) )
-		SetShaderResourceView<SF_Compute>(NewTexture, ShaderResourceView, TextureIndex, FD3D11StateCache::SRV_Dynamic);
+	if ((NewTexture == NULL) || (NewTexture->GetRenderTargetView(0, 0) != NULL) || (NewTexture->HasDepthStencilView()))
+	{
+		SetShaderResourceView<SF_Compute>(NewTexture, ShaderResourceView, TextureIndex, NewTextureRHI ? NewTextureRHI->GetName() : NAME_None, FD3D11StateCache::SRV_Dynamic);
+	}
 	else
-		SetShaderResourceView<SF_Compute>(NewTexture, ShaderResourceView, TextureIndex, FD3D11StateCache::SRV_Static);
+	{
+		SetShaderResourceView<SF_Compute>(NewTexture, ShaderResourceView, TextureIndex, NewTextureRHI->GetName(), FD3D11StateCache::SRV_Static);
+	}
 }
 
 void FD3D11DynamicRHI::RHISetUAVParameter(FComputeShaderRHIParamRef ComputeShaderRHI,uint32 UAVIndex,FUnorderedAccessViewRHIParamRef UAVRHI)
@@ -340,7 +364,8 @@ void FD3D11DynamicRHI::RHISetUAVParameter(FComputeShaderRHIParamRef ComputeShade
 
 	if(UAV)
 	{
-		ConditionalClearShaderResource(UAV->Resource);
+		ConditionalClearShaderResource(UAV->Resource);		
+		UAV->Resource->SetDirty();
 	}
 
 	ID3D11UnorderedAccessView* D3D11UAV = UAV ? UAV->View : NULL;
@@ -358,6 +383,7 @@ void FD3D11DynamicRHI::RHISetUAVParameter(FComputeShaderRHIParamRef ComputeShade
 	if(UAV)
 	{
 		ConditionalClearShaderResource(UAV->Resource);
+		UAV->Resource->SetDirty();
 	}
 
 	ID3D11UnorderedAccessView* D3D11UAV = UAV ? UAV->View : NULL;
@@ -379,7 +405,7 @@ void FD3D11DynamicRHI::RHISetShaderResourceViewParameter(FPixelShaderRHIParamRef
 		D3D11SRV = SRV->View;
 	}
 
-	SetShaderResourceView<SF_Pixel>(Resource, D3D11SRV, TextureIndex);
+	SetShaderResourceView<SF_Pixel>(Resource, D3D11SRV, TextureIndex, NAME_None);
 }
 
 void FD3D11DynamicRHI::RHISetShaderResourceViewParameter(FVertexShaderRHIParamRef VertexShaderRHI,uint32 TextureIndex,FShaderResourceViewRHIParamRef SRVRHI)
@@ -397,7 +423,7 @@ void FD3D11DynamicRHI::RHISetShaderResourceViewParameter(FVertexShaderRHIParamRe
 		D3D11SRV = SRV->View;
 	}
 
-	SetShaderResourceView<SF_Vertex>(Resource, D3D11SRV, TextureIndex);
+	SetShaderResourceView<SF_Vertex>(Resource, D3D11SRV, TextureIndex, NAME_None);
 }
 
 void FD3D11DynamicRHI::RHISetShaderResourceViewParameter(FComputeShaderRHIParamRef ComputeShaderRHI,uint32 TextureIndex,FShaderResourceViewRHIParamRef SRVRHI)
@@ -415,7 +441,7 @@ void FD3D11DynamicRHI::RHISetShaderResourceViewParameter(FComputeShaderRHIParamR
 		D3D11SRV = SRV->View;
 	}
 
-	SetShaderResourceView<SF_Compute>(Resource, D3D11SRV, TextureIndex);
+	SetShaderResourceView<SF_Compute>(Resource, D3D11SRV, TextureIndex, NAME_None);
 }
 
 void FD3D11DynamicRHI::RHISetShaderResourceViewParameter(FHullShaderRHIParamRef HullShaderRHI,uint32 TextureIndex,FShaderResourceViewRHIParamRef SRVRHI)
@@ -433,7 +459,7 @@ void FD3D11DynamicRHI::RHISetShaderResourceViewParameter(FHullShaderRHIParamRef
 		D3D11SRV = SRV->View;
 	}
 
-	SetShaderResourceView<SF_Hull>(Resource, D3D11SRV, TextureIndex);
+	SetShaderResourceView<SF_Hull>(Resource, D3D11SRV, TextureIndex, NAME_None);
 }
 
 void FD3D11DynamicRHI::RHISetShaderResourceViewParameter(FDomainShaderRHIParamRef DomainShaderRHI,uint32 TextureIndex,FShaderResourceViewRHIParamRef SRVRHI)
@@ -451,7 +477,7 @@ void FD3D11DynamicRHI::RHISetShaderResourceViewParameter(FDomainShaderRHIParamRe
 		D3D11SRV = SRV->View;
 	}
 
-	SetShaderResourceView<SF_Domain>(Resource, D3D11SRV, TextureIndex);
+	SetShaderResourceView<SF_Domain>(Resource, D3D11SRV, TextureIndex, NAME_None);
 }
 
 void FD3D11DynamicRHI::RHISetShaderResourceViewParameter(FGeometryShaderRHIParamRef GeometryShaderRHI,uint32 TextureIndex,FShaderResourceViewRHIParamRef SRVRHI)
@@ -469,7 +495,7 @@ void FD3D11DynamicRHI::RHISetShaderResourceViewParameter(FGeometryShaderRHIParam
 		D3D11SRV = SRV->View;
 	}
 
-	SetShaderResourceView<SF_Geometry>(Resource, D3D11SRV, TextureIndex);
+	SetShaderResourceView<SF_Geometry>(Resource, D3D11SRV, TextureIndex, NAME_None);
 }
 
 void FD3D11DynamicRHI::RHISetShaderSampler(FVertexShaderRHIParamRef VertexShaderRHI,uint32 SamplerIndex,FSamplerStateRHIParamRef NewStateRHI)
@@ -833,6 +859,26 @@ void FD3D11DynamicRHI::RHISetRenderTargets(
 		bTargetChanged = true;
 	}
 
+	if (NewDepthStencilTarget)
+	{
+		uint32 CurrentFrame = PresentCounter;
+		const EResourceTransitionAccess CurrentAccess = NewDepthStencilTarget->GetCurrentGPUAccess();
+		const uint32 LastFrameWritten = NewDepthStencilTarget->GetLastFrameWritten();
+		const bool bReadable = CurrentAccess == EResourceTransitionAccess::EReadable;
+		const bool bDepthWrite = NewDepthStencilTargetRHI->GetDepthStencilAccess().IsDepthWrite();
+		const bool bAccessValid =	!bReadable ||
+									LastFrameWritten != CurrentFrame || 
+									!bDepthWrite;
+
+		ensureMsgf(bAccessValid, TEXT("DepthTarget %s is not GPU writable."), *NewDepthStencilTargetRHI->Texture->GetName().ToString());		
+
+		//switch to writable state if this is the first render of the frame.  Don't switch if it's a later render and this is a depth test only situation
+		if (!bAccessValid || (bReadable && bDepthWrite))
+		{
+			NewDepthStencilTarget->SetCurrentGPUAccess(EResourceTransitionAccess::EWritable, CurrentFrame);			
+		}
+	}
+
 	// Gather the render target views for the new render targets.
 	ID3D11RenderTargetView* NewRenderTargetViews[MaxSimultaneousRenderTargets];
 	for(uint32 RenderTargetIndex = 0;RenderTargetIndex < MaxSimultaneousRenderTargets;++RenderTargetIndex)
@@ -845,6 +891,21 @@ void FD3D11DynamicRHI::RHISetRenderTargets(
 			FD3D11TextureBase* NewRenderTarget = GetD3D11TextureFromRHITexture(NewRenderTargetsRHI[RenderTargetIndex].Texture);
 			RenderTargetView = NewRenderTarget->GetRenderTargetView(RTMipIndex, RTSliceIndex);
 
+			if (NewRenderTarget)
+			{
+				uint32 CurrentFrame = PresentCounter;
+				const EResourceTransitionAccess CurrentAccess = NewRenderTarget->GetCurrentGPUAccess();
+				const uint32 LastFrameWritten = NewRenderTarget->GetLastFrameWritten();
+				const bool bReadable = CurrentAccess == EResourceTransitionAccess::EReadable;
+				const bool bAccessValid = !bReadable || LastFrameWritten != CurrentFrame;
+				ensureMsgf(bAccessValid, TEXT("RenderTarget %s is not GPU writable."), *NewRenderTargetsRHI[RenderTargetIndex].Texture->GetName().ToString());
+								
+				if (!bAccessValid || bReadable)
+				{
+					NewRenderTarget->SetCurrentGPUAccess(EResourceTransitionAccess::EWritable, CurrentFrame);
+				}
+			}
+
 			ensureMsgf(RenderTargetView, TEXT("Texture being set as render target has no RTV"));
 #if CHECK_SRV_TRANSITIONS			
 			if (RenderTargetView)
@@ -1148,15 +1209,15 @@ void FD3D11DynamicRHI::CommitComputeShaderConstants()
 }
 
 template <EShaderFrequency Frequency>
-FORCEINLINE void SetResource(FD3D11DynamicRHI* RESTRICT D3D11RHI, FD3D11StateCache* RESTRICT StateCache, uint32 BindIndex, FD3D11BaseShaderResource* RESTRICT ShaderResource, ID3D11ShaderResourceView* RESTRICT SRV)
+FORCEINLINE void SetResource(FD3D11DynamicRHI* RESTRICT D3D11RHI, FD3D11StateCache* RESTRICT StateCache, uint32 BindIndex, FD3D11BaseShaderResource* RESTRICT ShaderResource, ID3D11ShaderResourceView* RESTRICT SRV, FName ResourceName)
 {
 	// We set the resource through the RHI to track state for the purposes of unbinding SRVs when a UAV or RTV is bound.
 	// todo: need to support SRV_Static for faster calls when possible
-	D3D11RHI->SetShaderResourceView<Frequency>(ShaderResource,SRV,BindIndex,FD3D11StateCache::SRV_Unknown);
+	D3D11RHI->SetShaderResourceView<Frequency>(ShaderResource, SRV, BindIndex, ResourceName,FD3D11StateCache::SRV_Unknown);
 }
 
 template <EShaderFrequency Frequency>
-FORCEINLINE void SetResource(FD3D11DynamicRHI* RESTRICT D3D11RHI, FD3D11StateCache* RESTRICT StateCache, uint32 BindIndex, FD3D11BaseShaderResource* RESTRICT ShaderResource, ID3D11SamplerState* RESTRICT SamplerState)
+FORCEINLINE void SetResource(FD3D11DynamicRHI* RESTRICT D3D11RHI, FD3D11StateCache* RESTRICT StateCache, uint32 BindIndex, FD3D11BaseShaderResource* RESTRICT ShaderResource, ID3D11SamplerState* RESTRICT SamplerState, FName ResourceName)
 {
 	StateCache->SetSamplerState<Frequency>(SamplerState,BindIndex);
 }
@@ -1180,7 +1241,7 @@ inline int32 SetShaderResourcesFromBuffer(FD3D11DynamicRHI* RESTRICT D3D11RHI, F
 			FD3D11UniformBuffer::FResourcePair* RESTRICT ResourcePair = &Buffer->RawResourceTable[ResourceIndex];
 			FD3D11BaseShaderResource* ShaderResource = ResourcePair->ShaderResource;
 			D3DResourceType* D3D11Resource = (D3DResourceType*)ResourcePair->D3D11Resource;
-			SetResource<ShaderFrequency>(D3D11RHI, StateCache, BindIndex, ShaderResource, D3D11Resource);
+			SetResource<ShaderFrequency>(D3D11RHI, StateCache, BindIndex, ShaderResource, D3D11Resource, ResourcePair->ResourceName);
 			NumSetCalls++;
 			ResourceInfo = *ResourceInfos++;
 		} while (FRHIResourceTableEntry::GetUniformBufferIndex(ResourceInfo) == BufferIndex);
@@ -2029,3 +2090,46 @@ IRHICommandContextContainer* FD3D11DynamicRHI::RHIGetCommandContextContainer()
 	return nullptr;
 }
 
+void FD3D11DynamicRHI::RHITransitionResources(EResourceTransitionAccess TransitionType, FTextureRHIParamRef* InTextures, int32 NumTextures)
+{
+	static IConsoleVariable* CVarShowTransitions = IConsoleManager::Get().FindConsoleVariable(TEXT("r.ProfileGPU.ShowTransitions"));
+	bool bShowTransitionEvents = CVarShowTransitions->GetInt() != 0;
+
+	SCOPED_RHI_CONDITIONAL_DRAW_EVENTF(*this, RHITransitionResources, bShowTransitionEvents, TEXT("TransitionTo: %s: %i Textures"), *FResourceTransitionUtility::ResourceTransitionAccessStrings[(int32)TransitionType], NumTextures);
+	for (int32 i = 0; i < NumTextures; ++i)
+	{				
+		FTextureRHIParamRef RenderTarget = InTextures[i];
+		SCOPED_RHI_CONDITIONAL_DRAW_EVENTF(*this, RHITransitionResourcesLoop, bShowTransitionEvents, TEXT("To:%i - %s"), i, *RenderTarget->GetName().ToString());
+
+		FD3D11BaseShaderResource* Resource = nullptr;
+		FD3D11Texture2D* SourceTexture2D = static_cast<FD3D11Texture2D*>(RenderTarget->GetTexture2D());
+		if (SourceTexture2D)
+		{
+			Resource = SourceTexture2D;
+		}
+		FD3D11TextureCube* SourceTextureCube = static_cast<FD3D11TextureCube*>(RenderTarget->GetTextureCube());
+		if (SourceTextureCube)
+		{
+			Resource = SourceTextureCube;			
+		}
+		FD3D11Texture3D* SourceTexture3D = static_cast<FD3D11Texture3D*>(RenderTarget->GetTexture3D());
+		if (SourceTexture3D)
+		{
+			Resource = SourceTexture3D;			
+		}
+
+		Resource->SetCurrentGPUAccess(TransitionType, PresentCounter);
+	}
+}
+
+void FD3D11DynamicRHI::RHITransitionResources(EResourceTransitionAccess TransitionType, EResourceTransitionPipeline TransitionPipeline, FUnorderedAccessViewRHIParamRef* InUAVs, int32 NumUAVs)
+{
+	for (int32 i = 0; i < NumUAVs; ++i)
+	{
+		FD3D11UnorderedAccessView* UAV = ResourceCast(InUAVs[i]);
+		if (UAV->Resource)
+		{
+			UAV->Resource->SetCurrentGPUAccess(TransitionType, PresentCounter);
+		}		
+	}	
+}
\ No newline at end of file
diff --git a/Engine/Source/Runtime/Windows/D3D11RHI/Private/D3D11Device.cpp b/Engine/Source/Runtime/Windows/D3D11RHI/Private/D3D11Device.cpp
--- a/Engine/Source/Runtime/Windows/D3D11RHI/Private/D3D11Device.cpp
+++ b/Engine/Source/Runtime/Windows/D3D11RHI/Private/D3D11Device.cpp
@@ -42,6 +42,7 @@ FD3D11DynamicRHI::FD3D11DynamicRHI(IDXGIFactory1* InDXGIFactory1,D3D_FEATURE_LEV
 	NumSimultaneousRenderTargets(0),
 	NumUAVs(0),
 	SceneFrameCounter(0),
+	PresentCounter(0),
 	ResourceTableFrameCounter(INDEX_NONE),
 	CurrentDSVAccessType(FExclusiveDepthStencil::DepthWrite_StencilWrite),
 	bDiscardSharedConstants(false),
diff --git a/Engine/Source/Runtime/Windows/D3D11RHI/Private/D3D11RHI.cpp b/Engine/Source/Runtime/Windows/D3D11RHI/Private/D3D11RHI.cpp
--- a/Engine/Source/Runtime/Windows/D3D11RHI/Private/D3D11RHI.cpp
+++ b/Engine/Source/Runtime/Windows/D3D11RHI/Private/D3D11RHI.cpp
@@ -176,12 +176,17 @@ void FD3D11DynamicRHI::CheckIfSRVIsResolved(ID3D11ShaderResourceView* SRV)
 }
 
 template <EShaderFrequency ShaderFrequency>
-void FD3D11DynamicRHI::InternalSetShaderResourceView(FD3D11BaseShaderResource* Resource, ID3D11ShaderResourceView* SRV, int32 ResourceIndex, FD3D11StateCache::ESRV_Type SrvType)
+void FD3D11DynamicRHI::InternalSetShaderResourceView(FD3D11BaseShaderResource* Resource, ID3D11ShaderResourceView* SRV, int32 ResourceIndex, FName SRVName, FD3D11StateCache::ESRV_Type SrvType)
 {
 	// Check either both are set, or both are null.
 	check((Resource && SRV) || (!Resource && !SRV));
 	CheckIfSRVIsResolved(SRV);
 
+	if (Resource)
+	{
+		ensureMsgf(Resource->GetCurrentGPUAccess() == EResourceTransitionAccess::EReadable || Resource->GetCurrentGPUAccess() == EResourceTransitionAccess::ERWBarrier || Resource->GetLastFrameWritten() != PresentCounter, TEXT("Shader resource %s is not GPU readable.  Missing a call to RHITransitionResources()"), *SRVName.ToString());
+	}
+
 	FD3D11BaseShaderResource*& ResourceSlot = CurrentResourcesBoundAsSRVs[ShaderFrequency][ResourceIndex];
 	int32& MaxResourceIndex = MaxBoundShaderResourcesIndex[ShaderFrequency];
 
@@ -223,7 +228,7 @@ void FD3D11DynamicRHI::ClearShaderResourceViews(FD3D11BaseShaderResource* Resour
 		if (CurrentResourcesBoundAsSRVs[ShaderFrequency][ResourceIndex] == Resource)
 		{
 			// Unset the SRV from the device context
-			InternalSetShaderResourceView<ShaderFrequency>(nullptr, nullptr, ResourceIndex);
+			InternalSetShaderResourceView<ShaderFrequency>(nullptr, nullptr, ResourceIndex, NAME_None);
 		}
 	}
 }
@@ -249,7 +254,7 @@ void FD3D11DynamicRHI::ClearAllShaderResourcesForFrequency()
 		if (CurrentResourcesBoundAsSRVs[ShaderFrequency][ResourceIndex] != nullptr)
 		{
 			// Unset the SRV from the device context
-			InternalSetShaderResourceView<ShaderFrequency>(nullptr, nullptr, ResourceIndex);
+			InternalSetShaderResourceView<ShaderFrequency>(nullptr, nullptr, ResourceIndex, NAME_None);
 		}
 	}
 }
diff --git a/Engine/Source/Runtime/Windows/D3D11RHI/Private/D3D11RHIPrivate.h b/Engine/Source/Runtime/Windows/D3D11RHI/Private/D3D11RHIPrivate.h
--- a/Engine/Source/Runtime/Windows/D3D11RHI/Private/D3D11RHIPrivate.h
+++ b/Engine/Source/Runtime/Windows/D3D11RHI/Private/D3D11RHIPrivate.h
@@ -425,6 +425,8 @@ public:
 	virtual void RHISetMultipleViewports(uint32 Count, const FViewportBounds* Data) final override;
 	virtual void RHIClearUAV(FUnorderedAccessViewRHIParamRef UnorderedAccessViewRHI, const uint32* Values) final override;
 	virtual void RHICopyToResolveTarget(FTextureRHIParamRef SourceTexture, FTextureRHIParamRef DestTexture, bool bKeepOriginalSurface, const FResolveParams& ResolveParams) final override;
+	virtual void RHITransitionResources(EResourceTransitionAccess TransitionType, FTextureRHIParamRef* InTextures, int32 NumTextures) final override;
+	virtual void RHITransitionResources(EResourceTransitionAccess TransitionType, EResourceTransitionPipeline TransitionPipeline, FUnorderedAccessViewRHIParamRef* InUAVs, int32 NumUAVs) final override;
 	virtual void RHIBeginRenderQuery(FRenderQueryRHIParamRef RenderQuery) final override;
 	virtual void RHIEndRenderQuery(FRenderQueryRHIParamRef RenderQuery) final override;
 	virtual void RHIBeginOcclusionQueryBatch() final override;
@@ -510,7 +512,7 @@ public:
 	{
 		return DXGIFactory1;
 	}
-private:
+private:	
 	virtual void RHIClearMRTImpl(bool bClearColor, int32 NumClearColors, const FLinearColor* ColorArray, bool bClearDepth, float Depth, bool bClearStencil, uint32 Stencil, FIntRect ExcludeRect, bool bForceShaderClear);
 
 	template <EShaderFrequency ShaderFrequency>
@@ -522,7 +524,7 @@ private:
 	void CheckIfSRVIsResolved(ID3D11ShaderResourceView* SRV);
 
 	template <EShaderFrequency ShaderFrequency>
-	void InternalSetShaderResourceView(FD3D11BaseShaderResource* Resource, ID3D11ShaderResourceView* SRV, int32 ResourceIndex, FD3D11StateCache::ESRV_Type SrvType = FD3D11StateCache::SRV_Unknown);
+	void InternalSetShaderResourceView(FD3D11BaseShaderResource* Resource, ID3D11ShaderResourceView* SRV, int32 ResourceIndex, FName SRVName, FD3D11StateCache::ESRV_Type SrvType = FD3D11StateCache::SRV_Unknown);
 
 	void SetCurrentComputeShader(FComputeShaderRHIParamRef ComputeShader)
 	{
@@ -537,9 +539,9 @@ private:
 public:
 
 	template <EShaderFrequency ShaderFrequency>
-	void SetShaderResourceView(FD3D11BaseShaderResource* Resource, ID3D11ShaderResourceView* SRV, int32 ResourceIndex, FD3D11StateCache::ESRV_Type SrvType = FD3D11StateCache::SRV_Unknown )
+	void SetShaderResourceView(FD3D11BaseShaderResource* Resource, ID3D11ShaderResourceView* SRV, int32 ResourceIndex, FName SRVName, FD3D11StateCache::ESRV_Type SrvType = FD3D11StateCache::SRV_Unknown)
 	{
-		InternalSetShaderResourceView<ShaderFrequency>(Resource, SRV, ResourceIndex, SrvType);
+		InternalSetShaderResourceView<ShaderFrequency>(Resource, SRV, ResourceIndex, SRVName, SrvType);
 	}
 
 	void ClearState();
@@ -604,6 +606,9 @@ protected:
 	/** Internal frame counter, incremented on each call to RHIBeginScene. */
 	uint32 SceneFrameCounter;
 
+	/** Internal frame counter that just counts calls to Present */
+	uint32 PresentCounter;
+
 	/**
 	 * Internal counter used for resource table caching.
 	 * INDEX_NONE means caching is not allowed.
diff --git a/Engine/Source/Runtime/Windows/D3D11RHI/Private/D3D11RenderTarget.cpp b/Engine/Source/Runtime/Windows/D3D11RHI/Private/D3D11RenderTarget.cpp
--- a/Engine/Source/Runtime/Windows/D3D11RHI/Private/D3D11RenderTarget.cpp
+++ b/Engine/Source/Runtime/Windows/D3D11RHI/Private/D3D11RenderTarget.cpp
@@ -143,7 +143,7 @@ void FD3D11DynamicRHI::ResolveTextureUsingShader(
 
 	if (SourceTexture)
 	{
-		SetShaderResourceView<SF_Pixel>(SourceTexture, SourceTexture->GetShaderResourceView(), TextureIndex);
+		SetShaderResourceView<SF_Pixel>(SourceTexture, SourceTexture->GetShaderResourceView(), TextureIndex, SourceTexture->GetName());
 	}
 
 	// Generate the vertices used
@@ -192,14 +192,16 @@ void FD3D11DynamicRHI::ResolveTextureUsingShader(
 */
 void FD3D11DynamicRHI::RHICopyToResolveTarget(FTextureRHIParamRef SourceTextureRHI, FTextureRHIParamRef DestTextureRHI, bool bKeepOriginalSurface, const FResolveParams& ResolveParams)
 {
-	if(!SourceTextureRHI || !DestTextureRHI)
+	if (!SourceTextureRHI || !DestTextureRHI)
 	{
 		// no need to do anything (sliently ignored)
 		return;
 	}
 
-	FRHICommandList_RecursiveHazardous RHICmdList(this);
+	RHITransitionResources(EResourceTransitionAccess::EReadable, &SourceTextureRHI, 1);
 
+	FRHICommandList_RecursiveHazardous RHICmdList(this);
+	
 
 	FD3D11Texture2D* SourceTexture2D = static_cast<FD3D11Texture2D*>(SourceTextureRHI->GetTexture2D());
 	FD3D11Texture2D* DestTexture2D = static_cast<FD3D11Texture2D*>(DestTextureRHI->GetTexture2D());
@@ -209,39 +211,6 @@ void FD3D11DynamicRHI::RHICopyToResolveTarget(FTextureRHIParamRef SourceTextureR
 
 	FD3D11Texture3D* SourceTexture3D = static_cast<FD3D11Texture3D*>(SourceTextureRHI->GetTexture3D());
 	FD3D11Texture3D* DestTexture3D = static_cast<FD3D11Texture3D*>(DestTextureRHI->GetTexture3D());
-
-#if CHECK_SRV_TRANSITIONS
-	//check first, as some of the resolve shaders could trigger the resource we're currently resolving.
-	ID3D11Resource* SourceResource = nullptr;
-	int32 ResolveSlice = -1;
-	if (SourceTexture2D)
-	{
-		SourceResource = SourceTexture2D->GetResource();
-	}
-	if (SourceTextureCube)
-	{
-		SourceResource = SourceTextureCube->GetResource();
-		ResolveSlice = ResolveParams.SourceArrayIndex * 6 + ResolveParams.CubeFace;
-	}
-	if (SourceTexture3D)
-	{
-		SourceResource = SourceTexture3D->GetResource();
-	}
-
-	if (SourceResource)
-	{
-		// if we don't have specific resolve data remove all references.
-		// this is helpful when you clear many mips/slices but only want to do one resolve.
-		if (ResolveParams.MipIndex == -1 && ResolveParams.SourceArrayIndex == -1)
-		{
-			UnresolvedTargets.Remove(SourceResource);
-		}
-		else
-		{
-			UnresolvedTargets.Remove(SourceResource, FUnresolvedRTInfo(SourceTextureRHI->GetName(), ResolveParams.MipIndex, 1, ResolveSlice, 1));
-		}
-	}
-#endif
 		
 	if(SourceTexture2D && DestTexture2D)
 	{
diff --git a/Engine/Source/Runtime/Windows/D3D11RHI/Private/D3D11Texture.cpp b/Engine/Source/Runtime/Windows/D3D11RHI/Private/D3D11Texture.cpp
--- a/Engine/Source/Runtime/Windows/D3D11RHI/Private/D3D11Texture.cpp
+++ b/Engine/Source/Runtime/Windows/D3D11RHI/Private/D3D11Texture.cpp
@@ -573,9 +573,9 @@ TD3D11Texture2D<BaseResourceType>* FD3D11DynamicRHI::CreateD3D11Texture2D(uint32
 	if(Flags & TexCreate_RenderTargetable)
 	{
 		check(!(Flags & TexCreate_DepthStencilTargetable));
-		check(!(Flags & TexCreate_ResolveTargetable));
+		check(!(Flags & TexCreate_ResolveTargetable));		
 		TextureDesc.BindFlags |= D3D11_BIND_RENDER_TARGET;
-		bCreateRTV = true;
+		bCreateRTV = true;		
 	}
 	else if(Flags & TexCreate_DepthStencilTargetable)
 	{
@@ -843,6 +843,11 @@ TD3D11Texture2D<BaseResourceType>* FD3D11DynamicRHI::CreateD3D11Texture2D(uint32
 
 	Texture2D->ResourceInfo.VRamAllocation = VRamAllocation;
 
+	if (Flags & TexCreate_RenderTargetable)
+	{
+		Texture2D->SetCurrentGPUAccess(EResourceTransitionAccess::EWritable, -1);		
+	}
+
 	D3D11TextureAllocated(*Texture2D);
 
 	return Texture2D;
@@ -886,7 +891,7 @@ FD3D11Texture3D* FD3D11DynamicRHI::CreateD3D11Texture3D(uint32 SizeX,uint32 Size
 
 	if(Flags & TexCreate_RenderTargetable)
 	{
-		TextureDesc.BindFlags |= D3D11_BIND_RENDER_TARGET;
+		TextureDesc.BindFlags |= D3D11_BIND_RENDER_TARGET;		
 		bCreateRTV = true;
 	}
 
@@ -965,6 +970,11 @@ FD3D11Texture3D* FD3D11DynamicRHI::CreateD3D11Texture3D(uint32 SizeX,uint32 Size
 
 	Texture3D->ResourceInfo.VRamAllocation = VRamAllocation;
 
+	if (Flags & TexCreate_RenderTargetable)
+	{
+		Texture3D->SetCurrentGPUAccess(EResourceTransitionAccess::EWritable, -1);		
+	}
+
 	D3D11TextureAllocated(*Texture3D);
 
 	return Texture3D;
diff --git a/Engine/Source/Runtime/Windows/D3D11RHI/Private/D3D11UniformBuffer.cpp b/Engine/Source/Runtime/Windows/D3D11RHI/Private/D3D11UniformBuffer.cpp
--- a/Engine/Source/Runtime/Windows/D3D11RHI/Private/D3D11UniformBuffer.cpp
+++ b/Engine/Source/Runtime/Windows/D3D11RHI/Private/D3D11UniformBuffer.cpp
@@ -264,13 +264,14 @@ void FD3D11UniformBuffer::CacheResourcesInternal()
 	// Texture streaming makes textures complicated :)
 	for (int32 i = 0; i < NumResources; ++i)
 	{
+		RawResources[i].ResourceName = NAME_None;
 		switch (ResourceTypes[i])
 		{
 		case UBMT_SRV:
 			{
 				FD3D11ShaderResourceView* ShaderResourceViewRHI = (FD3D11ShaderResourceView*)Resources[i].GetReference();
 				RawResources[i].ShaderResource = ShaderResourceViewRHI->Resource.GetReference();
-				RawResources[i].D3D11Resource = (IUnknown*)ShaderResourceViewRHI->View.GetReference();
+				RawResources[i].D3D11Resource = (IUnknown*)ShaderResourceViewRHI->View.GetReference();				
 			}
 			break;
 
@@ -283,6 +284,7 @@ void FD3D11UniformBuffer::CacheResourcesInternal()
 				FD3D11TextureBase* TextureD3D11 = GetD3D11TextureFromRHITexture(TextureRHI);
 				RawResources[i].ShaderResource = TextureD3D11->GetBaseShaderResource();
 				RawResources[i].D3D11Resource = (IUnknown*)TextureD3D11->GetShaderResourceView();
+				RawResources[i].ResourceName = TextureRHI->GetName();
 			}
 			break;
 
diff --git a/Engine/Source/Runtime/Windows/D3D11RHI/Private/D3D11Viewport.cpp b/Engine/Source/Runtime/Windows/D3D11RHI/Private/D3D11Viewport.cpp
--- a/Engine/Source/Runtime/Windows/D3D11RHI/Private/D3D11Viewport.cpp
+++ b/Engine/Source/Runtime/Windows/D3D11RHI/Private/D3D11Viewport.cpp
@@ -511,6 +511,7 @@ void FD3D11DynamicRHI::RHIBeginDrawingViewport(FViewportRHIParamRef ViewportRHI,
 	if( RenderTarget == NULL )
 	{
 		RenderTarget = Viewport->GetBackBuffer();
+		RHITransitionResources(EResourceTransitionAccess::EWritable, &RenderTarget, 1);
 	}
 	FRHIRenderTargetView View(RenderTarget);
 	RHISetRenderTargets(1,&View,nullptr,0,NULL);
@@ -521,6 +522,7 @@ void FD3D11DynamicRHI::RHIBeginDrawingViewport(FViewportRHIParamRef ViewportRHI,
 
 void FD3D11DynamicRHI::RHIEndDrawingViewport(FViewportRHIParamRef ViewportRHI,bool bPresent,bool bLockToVsync)
 {
+	++PresentCounter;
 	FD3D11Viewport* Viewport = ResourceCast(ViewportRHI);
 
 	SCOPE_CYCLE_COUNTER(STAT_D3D11PresentTime);
diff --git a/Engine/Source/Runtime/Windows/D3D11RHI/Public/D3D11Resources.h b/Engine/Source/Runtime/Windows/D3D11RHI/Public/D3D11Resources.h
--- a/Engine/Source/Runtime/Windows/D3D11RHI/Public/D3D11Resources.h
+++ b/Engine/Source/Runtime/Windows/D3D11RHI/Public/D3D11Resources.h
@@ -163,6 +163,57 @@ public:
 /** The base class of resources that may be bound as shader resources. */
 class FD3D11BaseShaderResource : public IRefCountedObject
 {
+public:
+	FD3D11BaseShaderResource()
+		: CurrentGPUAccess(EResourceTransitionAccess::EReadable)
+		, LastFrameWritten(-1)
+		, bDirty(false)
+	{}
+
+	void SetCurrentGPUAccess(EResourceTransitionAccess Access, uint32 CurrentFrame)
+	{
+		if (Access == EResourceTransitionAccess::EReadable)
+		{
+			bDirty = false;
+		}
+		else
+		{
+			LastFrameWritten = CurrentFrame;
+			bDirty = true;
+		}		
+		CurrentGPUAccess = Access;
+	}
+
+	EResourceTransitionAccess GetCurrentGPUAccess() const
+	{
+		return CurrentGPUAccess;
+	}
+
+	uint32 GetLastFrameWritten() const
+	{
+		return LastFrameWritten;
+	}
+
+	void SetDirty()
+	{
+		bDirty = true;
+	}
+
+	bool IsDirty() const
+	{
+		return bDirty;
+	}
+
+private:
+	
+	/** Whether the current resource is logically GPU readable or writable.  Mostly for validation for newer RHI's*/
+	EResourceTransitionAccess CurrentGPUAccess;
+
+	/** Most recent frame this resource was written to. */
+	uint32 LastFrameWritten;
+
+	/** Resource has been written to without a subsequent read barrier.  Mostly for UAVs */
+	bool bDirty;
 };
 
 /** Texture base class. */
@@ -285,8 +336,7 @@ protected:
 	TRefCountPtr<ID3D11DepthStencilView> DepthStencilViews[FExclusiveDepthStencil::MaxIndex];
 
 	/** Number of Depth Stencil Views - used for fast call tracking. */
-	uint32	NumDepthStencilViews;
-
+	uint32	NumDepthStencilViews;	
 };
 
 /** 2D texture (vanilla, cubemap or 2D array) */
@@ -618,7 +668,8 @@ public:
 
 	/** Cached resources need to retain the associated shader resource for bookkeeping purposes. */
 	struct FResourcePair
-	{
+	{		
+		FName ResourceName;
 		FD3D11BaseShaderResource* ShaderResource;
 		IUnknown* D3D11Resource;
 	};
@@ -699,7 +750,9 @@ public:
 	FD3D11StructuredBuffer(ID3D11Buffer* InResource, uint32 InStride, uint32 InSize, uint32 InUsage)
 	: FRHIStructuredBuffer(InStride,InSize,InUsage)
 	, Resource(InResource)
-	{}
+	{
+		SetCurrentGPUAccess(EResourceTransitionAccess::ERWBarrier, -1);
+	}
 
 	virtual ~FD3D11StructuredBuffer()
 	{
@@ -773,8 +826,8 @@ class FD3D11UnorderedAccessView : public FRHIUnorderedAccessView
 public:
 	
 	TRefCountPtr<ID3D11UnorderedAccessView> View;
-	TRefCountPtr<FD3D11BaseShaderResource> Resource;
-
+	TRefCountPtr<FD3D11BaseShaderResource> Resource;	
+	
 	FD3D11UnorderedAccessView(ID3D11UnorderedAccessView* InView,FD3D11BaseShaderResource* InResource)
 	: View(InView)
 	, Resource(InResource)
