Renamed some shadow variables for clarity

diff --git a/Engine/Source/Runtime/Renderer/Private/ShadowRendering.cpp b/Engine/Source/Runtime/Renderer/Private/ShadowRendering.cpp
--- a/Engine/Source/Runtime/Renderer/Private/ShadowRendering.cpp
+++ b/Engine/Source/Runtime/Renderer/Private/ShadowRendering.cpp
@@ -1318,7 +1318,7 @@ template <bool bReflectiveShadowmap>
 void DrawShadowMeshElements(FRHICommandList& RHICmdList, const FViewInfo& View, const FProjectedShadowInfo& ShadowInfo)
 {
 	FShadowDepthDrawingPolicyContext PolicyContext(&ShadowInfo);
-	const FShadowStaticMeshElement& FirstShadowMesh = ShadowInfo.SubjectMeshElements[0];
+	const FShadowStaticMeshElement& FirstShadowMesh = ShadowInfo.StaticSubjectMeshElements[0];
 	const FMaterial* FirstMaterialResource = FirstShadowMesh.MaterialResource;
 	auto FeatureLevel = View.GetFeatureLevel();
 
@@ -1331,10 +1331,10 @@ void DrawShadowMeshElements(FRHICommandList& RHICmdList, const FViewInfo& View,
 
 	FShadowStaticMeshElement OldState;
 
-	uint32 ElementCount = ShadowInfo.SubjectMeshElements.Num();
+	uint32 ElementCount = ShadowInfo.StaticSubjectMeshElements.Num();
 	for(uint32 ElementIndex = 0; ElementIndex < ElementCount; ++ElementIndex)
 	{
-		const FShadowStaticMeshElement& ShadowMesh = ShadowInfo.SubjectMeshElements[ElementIndex];
+		const FShadowStaticMeshElement& ShadowMesh = ShadowInfo.StaticSubjectMeshElements[ElementIndex];
 
 		if(!View.StaticMeshShadowDepthMap[ShadowMesh.Mesh->Id])
 		{
@@ -1644,14 +1644,14 @@ void FProjectedShadowInfo::RenderDepthInner(FRHICommandList& RHICmdList, FSceneR
 				}
 			}
 			// Draw the subject's static elements using manual state filtering
-			else if (SubjectMeshElements.Num() > 0)
+			else if (StaticSubjectMeshElements.Num() > 0)
 			{
 				FRHICommandList* CmdList = ParallelCommandListSet.NewParallelCommandList();
 
 				FGraphEventRef AnyThreadCompletionEvent = TGraphTask<FDrawShadowMeshElementsThreadTask>::CreateTask(ParallelCommandListSet.GetPrereqs(), ENamedThreads::RenderThread)
 					.ConstructAndDispatchWhenReady(*this, *CmdList, *FoundView, bReflectiveShadowmap && !CascadeSettings.bOnePassPointLightShadow);
 
-				ParallelCommandListSet.AddParallelCommandList(CmdList, AnyThreadCompletionEvent, SubjectMeshElements.Num());
+				ParallelCommandListSet.AddParallelCommandList(CmdList, AnyThreadCompletionEvent, StaticSubjectMeshElements.Num());
 			}
 			if (DynamicSubjectMeshElements.Num())
 			{
@@ -1685,7 +1685,7 @@ void FProjectedShadowInfo::RenderDepthInner(FRHICommandList& RHICmdList, FSceneR
 			}
 		}
 		// Draw the subject's static elements using manual state filtering
-		else if (SubjectMeshElements.Num() > 0)
+		else if (StaticSubjectMeshElements.Num() > 0)
 		{
 			SCOPE_CYCLE_COUNTER(STAT_WholeSceneStaticShadowDepthsTime);
 
@@ -1779,8 +1779,8 @@ void FProjectedShadowInfo::RenderDepth(FRHICommandList& RHICmdList, FSceneRender
 	switch(RenderMode)
 	{
 	case ShadowDepthRenderMode_Dynamic:
-		PtrCurrentMeshElements = &SubjectMeshElements;
-		PtrCurrentPrimitives = &SubjectPrimitives;
+		PtrCurrentMeshElements = &StaticSubjectMeshElements;
+		PtrCurrentPrimitives = &DynamicSubjectPrimitives;
 		break;
 	case ShadowDepthRenderMode_EmissiveOnly:
 		PtrCurrentMeshElements = &EmissiveOnlyMeshElements;
@@ -2102,7 +2102,7 @@ void FProjectedShadowInfo::RenderProjection(FRHICommandListImmediate& RHICmdList
 
 		// Pre-shadows mask by receiver elements, self-shadow mask by subject elements.
 		// Note that self-shadow pre-shadows still mask by receiver elements.
-		const PrimitiveArrayType& MaskPrimitives = bPreShadow ? ReceiverPrimitives : SubjectPrimitives;
+		const PrimitiveArrayType& MaskPrimitives = bPreShadow ? ReceiverPrimitives : DynamicSubjectPrimitives;
 
 		for (int32 PrimitiveIndex = 0, PrimitiveCount = MaskPrimitives.Num(); PrimitiveIndex < PrimitiveCount; PrimitiveIndex++)
 		{
@@ -2140,25 +2140,21 @@ void FProjectedShadowInfo::RenderProjection(FRHICommandListImmediate& RHICmdList
 
 		if (bSelfShadowOnly && !bPreShadow)
 		{
-			for (int32 ElementIndex = 0; ElementIndex < SubjectMeshElements.Num(); ++ElementIndex)
+			for (int32 ElementIndex = 0; ElementIndex < StaticSubjectMeshElements.Num(); ++ElementIndex)
 			{
-				const FStaticMesh& StaticMesh = *SubjectMeshElements[ElementIndex].Mesh;
- 
-				if (View->StaticMeshShadowDepthMap[StaticMesh.Id])
-				{
-					const float DitherValue = View->GetDitheredLODTransitionValue(StaticMesh);
-					FDepthDrawingPolicyFactory::DrawStaticMesh(
-						RHICmdList, 
-						*View,
-						FDepthDrawingPolicyFactory::ContextType(DDM_AllOccluders),
-						StaticMesh,
-						StaticMesh.Elements.Num() == 1 ? 1 : View->StaticMeshBatchVisibility[StaticMesh.Id],
-						true,
-						DitherValue,
-						StaticMesh.PrimitiveSceneInfo->Proxy,
-						StaticMesh.BatchHitProxyId
-						);
-				}
+				const FStaticMesh& StaticMesh = *StaticSubjectMeshElements[ElementIndex].Mesh;
+				const float DitherValue = View->GetDitheredLODTransitionValue(StaticMesh);
+				FDepthDrawingPolicyFactory::DrawStaticMesh(
+					RHICmdList, 
+					*View,
+					FDepthDrawingPolicyFactory::ContextType(DDM_AllOccluders),
+					StaticMesh,
+					StaticMesh.Elements.Num() == 1 ? 1 : View->StaticMeshBatchVisibility[StaticMesh.Id],
+					true,
+					DitherValue,
+					StaticMesh.PrimitiveSceneInfo->Proxy,
+					StaticMesh.BatchHitProxyId
+					);
 			}
 		}
 	}
@@ -2518,9 +2514,9 @@ void FProjectedShadowInfo::RenderFrustumWireframe(FPrimitiveDrawInterface* PDI)
 {
 	// Find the ID of an arbitrary subject primitive to use to color the shadow frustum.
 	int32 SubjectPrimitiveId = 0;
-	if(SubjectPrimitives.Num())
+	if(DynamicSubjectPrimitives.Num())
 	{
-		SubjectPrimitiveId = SubjectPrimitives[0]->GetIndex();
+		SubjectPrimitiveId = DynamicSubjectPrimitives[0]->GetIndex();
 	}
 
 	const FMatrix InvShadowTransform = (bWholeSceneShadow || bPreShadow) ? SubjectAndReceiverMatrix.InverseFast() : InvReceiverMatrix;
@@ -2769,7 +2765,7 @@ void FProjectedShadowInfo::SortSubjectMeshElements()
 		}
 	};
 
-	SubjectMeshElements.Sort( FCompareFShadowStaticMeshElement() );
+	StaticSubjectMeshElements.Sort( FCompareFShadowStaticMeshElement() );
 }
 
 void FProjectedShadowInfo::GetShadowTypeNameForDrawEvent(FString& TypeName) const
diff --git a/Engine/Source/Runtime/Renderer/Private/ShadowRendering.h b/Engine/Source/Runtime/Renderer/Private/ShadowRendering.h
--- a/Engine/Source/Runtime/Renderer/Private/ShadowRendering.h
+++ b/Engine/Source/Runtime/Renderer/Private/ShadowRendering.h
@@ -733,7 +733,7 @@ public:
 		return bWholeSceneShadow && LightSceneInfo->Proxy->GetLightType() == LightType_Point;
 	}
 
-	/** Sorts SubjectMeshElements based on state so that rendering the static elements will set as little state as possible. */
+	/** Sorts StaticSubjectMeshElements based on state so that rendering the static elements will set as little state as possible. */
 	void SortSubjectMeshElements();
 
 	// 0 if Setup...() wasn't called yet
@@ -757,7 +757,7 @@ private:
 	const FPrimitiveSceneInfo* ParentSceneInfo;
 
 	/** dynamic shadow casting elements */
-	PrimitiveArrayType SubjectPrimitives;
+	PrimitiveArrayType DynamicSubjectPrimitives;
 	/** For preshadows, this contains the receiver primitives to mask the projection to. */
 	PrimitiveArrayType ReceiverPrimitives;
 	/** Subject primitives with translucent relevance. */
@@ -774,7 +774,7 @@ private:
 	TArray<FShadowStaticMeshElement,SceneRenderingAllocator> GIBlockingMeshElements;
 
 	/** Static shadow casting elements. */
-	TArray<FShadowStaticMeshElement,SceneRenderingAllocator> SubjectMeshElements;
+	TArray<FShadowStaticMeshElement,SceneRenderingAllocator> StaticSubjectMeshElements;
 
 	/** Dynamic mesh elements for subject primitives. */
 	TArray<FMeshBatchAndRelevance,SceneRenderingAllocator> DynamicSubjectMeshElements;
diff --git a/Engine/Source/Runtime/Renderer/Private/ShadowSetup.cpp b/Engine/Source/Runtime/Renderer/Private/ShadowSetup.cpp
--- a/Engine/Source/Runtime/Renderer/Private/ShadowSetup.cpp
+++ b/Engine/Source/Runtime/Renderer/Private/ShadowSetup.cpp
@@ -850,7 +850,7 @@ void FProjectedShadowInfo::AddSubjectPrimitive(FPrimitiveSceneInfo* PrimitiveSce
 							{
 								const bool bTwoSided = Material->IsTwoSided() || PrimitiveSceneInfo->Proxy->CastsShadowAsTwoSided();
 								OverrideWithDefaultMaterialForShadowDepth(MaterialRenderProxy, Material, bReflectiveShadowmap, FeatureLevel);
-								SubjectMeshElements.Add(FShadowStaticMeshElement(MaterialRenderProxy, Material, &StaticMesh,bTwoSided));
+								StaticSubjectMeshElements.Add(FShadowStaticMeshElement(MaterialRenderProxy, Material, &StaticMesh,bTwoSided));
 							}
 						}
 					}
@@ -859,7 +859,7 @@ void FProjectedShadowInfo::AddSubjectPrimitive(FPrimitiveSceneInfo* PrimitiveSce
 			else
 			{
 				// Add the primitive to the subject primitive list.
-				SubjectPrimitives.Add(PrimitiveSceneInfo);
+				DynamicSubjectPrimitives.Add(PrimitiveSceneInfo);
 			}
 		}
 
@@ -909,8 +909,8 @@ void FProjectedShadowInfo::AddSubjectPrimitive(FPrimitiveSceneInfo* PrimitiveSce
 
 bool FProjectedShadowInfo::HasSubjectPrims() const
 {
-	return SubjectPrimitives.Num() > 0
-		|| SubjectMeshElements.Num() > 0
+	return DynamicSubjectPrimitives.Num() > 0
+		|| StaticSubjectMeshElements.Num() > 0
 		|| EmissiveOnlyPrimitives.Num() > 0 
 		|| EmissiveOnlyMeshElements.Num() > 0
 		|| GIBlockingMeshElements.Num() > 0
@@ -926,11 +926,11 @@ void FProjectedShadowInfo::AddReceiverPrimitive(FPrimitiveSceneInfo* PrimitiveSc
 static TAutoConsoleVariable<int32> CVarDisableCullShadows(
 	TEXT("foliage.DisableCullShadows"),
 	0,
-	TEXT("First three bits are disable SubjectPrimitives, ReceiverPrimitives, SubjectTranslucentPrimitives"));
+	TEXT("First three bits are disable DynamicSubjectPrimitives, ReceiverPrimitives, SubjectTranslucentPrimitives"));
 
 void FProjectedShadowInfo::GatherDynamicMeshElements(FSceneRenderer& Renderer, FVisibleLightInfo& VisibleLightInfo, TArray<const FSceneView*>& ReusedViewsArray)
 {
-	if (SubjectPrimitives.Num() > 0 || ReceiverPrimitives.Num() > 0 || SubjectTranslucentPrimitives.Num() > 0)
+	if (DynamicSubjectPrimitives.Num() > 0 || ReceiverPrimitives.Num() > 0 || SubjectTranslucentPrimitives.Num() > 0)
 	{
 		// Choose an arbitrary view where this shadow's subject is relevant.
 		FViewInfo* FoundView = NULL;
@@ -975,14 +975,14 @@ void FProjectedShadowInfo::GatherDynamicMeshElements(FSceneRenderer& Renderer, F
 		    {
 			    FoundView->ViewMatrices.PreShadowTranslation = FVector(0,0,0);
 			    FoundView->ViewMatrices.GetDynamicMeshElementsShadowCullFrustum = (Disable & 1) ? &NoCull : &CascadeSettings.ShadowBoundsAccurate;
-			    GatherDynamicMeshElementsArray(FoundView, Renderer, SubjectPrimitives, DynamicSubjectMeshElements, ReusedViewsArray);
+			    GatherDynamicMeshElementsArray(FoundView, Renderer, DynamicSubjectPrimitives, DynamicSubjectMeshElements, ReusedViewsArray);
 			    FoundView->ViewMatrices.PreShadowTranslation = PreShadowTranslation;
 		    }
 		    else
 		    {
 			    FoundView->ViewMatrices.PreShadowTranslation = PreShadowTranslation;
 			    FoundView->ViewMatrices.GetDynamicMeshElementsShadowCullFrustum = (Disable & 1) ? &NoCull : &CasterFrustum;
-			    GatherDynamicMeshElementsArray(FoundView, Renderer, SubjectPrimitives, DynamicSubjectMeshElements, ReusedViewsArray);
+			    GatherDynamicMeshElementsArray(FoundView, Renderer, DynamicSubjectPrimitives, DynamicSubjectMeshElements, ReusedViewsArray);
 		    }
     
 		    FoundView->ViewMatrices.GetDynamicMeshElementsShadowCullFrustum = (Disable & 2) ? &NoCull : &ReceiverFrustum;
@@ -1045,9 +1045,9 @@ void FProjectedShadowInfo::GatherDynamicMeshElementsArray(
 bool FProjectedShadowInfo::SubjectsVisible(const FViewInfo& View) const
 {
 	checkSlow(!IsWholeSceneDirectionalShadow());
-	for(int32 PrimitiveIndex = 0;PrimitiveIndex < SubjectPrimitives.Num();PrimitiveIndex++)
+	for(int32 PrimitiveIndex = 0;PrimitiveIndex < DynamicSubjectPrimitives.Num();PrimitiveIndex++)
 	{
-		const FPrimitiveSceneInfo* SubjectPrimitiveSceneInfo = SubjectPrimitives[PrimitiveIndex];
+		const FPrimitiveSceneInfo* SubjectPrimitiveSceneInfo = DynamicSubjectPrimitives[PrimitiveIndex];
 		if(View.PrimitiveVisibilityMap[SubjectPrimitiveSceneInfo->GetIndex()])
 		{
 			return true;
@@ -1063,9 +1063,9 @@ bool FProjectedShadowInfo::SubjectsVisible(const FViewInfo& View) const
 void FProjectedShadowInfo::ClearTransientArrays()
 {
 	SubjectTranslucentPrimitives.Empty();
-	SubjectPrimitives.Empty();
+	DynamicSubjectPrimitives.Empty();
 	ReceiverPrimitives.Empty();
-	SubjectMeshElements.Empty();
+	StaticSubjectMeshElements.Empty();
 	EmissiveOnlyPrimitives.Empty();
 	EmissiveOnlyMeshElements.Empty();
 	DynamicSubjectMeshElements.Empty();
