Fix reflection transition ensures.

diff --git a/Engine/Source/Runtime/Renderer/Private/ReflectionEnvironmentCapture.cpp b/Engine/Source/Runtime/Renderer/Private/ReflectionEnvironmentCapture.cpp
--- a/Engine/Source/Runtime/Renderer/Private/ReflectionEnvironmentCapture.cpp
+++ b/Engine/Source/Runtime/Renderer/Private/ReflectionEnvironmentCapture.cpp
@@ -286,7 +286,7 @@ IMPLEMENT_SHADER_TYPE(,FComputeBrightnessPS,TEXT("ReflectionEnvironmentShaders")
 void ComputeAverageBrightness(FRHICommandList& RHICmdList, ERHIFeatureLevel::Type FeatureLevel)
 {
 	FTextureRHIRef& BrightnessTarget = FSceneRenderTargets::Get(RHICmdList).GetReflectionBrightnessTarget()->GetRenderTargetItem().TargetableTexture;
-	SetRenderTarget(RHICmdList, BrightnessTarget, NULL);
+	SetRenderTarget(RHICmdList, BrightnessTarget, NULL, true);
 	RHICmdList.SetRasterizerState(TStaticRasterizerState<FM_Solid, CM_None>::GetRHI());
 	RHICmdList.SetDepthStencilState(TStaticDepthStencilState<false, CF_Always>::GetRHI());
 	RHICmdList.SetBlendState(TStaticBlendState<>::GetRHI());
@@ -348,7 +348,7 @@ void FilterReflectionEnvironment(FRHICommandListImmediate& RHICmdList, ERHIFeatu
 			
 			for (int32 CubeFace = 0; CubeFace < CubeFace_MAX; CubeFace++)
 			{
-				SetRenderTarget(RHICmdList, EffectiveRT.TargetableTexture, MipIndex, CubeFace, NULL);
+				SetRenderTarget(RHICmdList, EffectiveRT.TargetableTexture, MipIndex, CubeFace, NULL, true);
 
 				const FIntRect ViewRect(0, 0, MipSize, MipSize);
 				RHICmdList.SetViewport(0, 0, 0.0f, MipSize, MipSize, 1.0f);
@@ -406,7 +406,7 @@ void FilterReflectionEnvironment(FRHICommandListImmediate& RHICmdList, ERHIFeatu
 
 			for (int32 CubeFace = 0; CubeFace < CubeFace_MAX; CubeFace++)
 			{
-				SetRenderTarget(RHICmdList, EffectiveRT.TargetableTexture, MipIndex, CubeFace, NULL);
+				SetRenderTarget(RHICmdList, EffectiveRT.TargetableTexture, MipIndex, CubeFace, NULL, true);
 
 				const FIntRect ViewRect(0, 0, MipSize, MipSize);
 				RHICmdList.SetViewport(0, 0, 0.0f, MipSize, MipSize, 1.0f);
@@ -678,7 +678,7 @@ void ClearScratchCubemaps(FRHICommandList& RHICmdList)
 	{
 		for (int32 CubeFace = 0; CubeFace < CubeFace_MAX; CubeFace++)
 		{
-			SetRenderTarget(RHICmdList, RT0.TargetableTexture, MipIndex, CubeFace, NULL);
+			SetRenderTarget(RHICmdList, RT0.TargetableTexture, MipIndex, CubeFace, NULL, true);
 			RHICmdList.Clear(true, FLinearColor(0, 10000, 0, 0), false, (float)ERHIZBuffer::FarPlane, false, 0, FIntRect());
 		}
 	}
@@ -689,7 +689,7 @@ void ClearScratchCubemaps(FRHICommandList& RHICmdList)
 	{
 		for (int32 CubeFace = 0; CubeFace < CubeFace_MAX; CubeFace++)
 		{
-			SetRenderTarget(RHICmdList, RT1.TargetableTexture, MipIndex, CubeFace, NULL);
+			SetRenderTarget(RHICmdList, RT1.TargetableTexture, MipIndex, CubeFace, NULL, true);
 			RHICmdList.Clear(true, FLinearColor(0, 10000, 0, 0), false, (float)ERHIZBuffer::FarPlane, false, 0, FIntRect());
 		}
 	}
diff --git a/Engine/Source/Runtime/Renderer/Private/ReflectionEnvironmentDiffuseIrradiance.cpp b/Engine/Source/Runtime/Renderer/Private/ReflectionEnvironmentDiffuseIrradiance.cpp
--- a/Engine/Source/Runtime/Renderer/Private/ReflectionEnvironmentDiffuseIrradiance.cpp
+++ b/Engine/Source/Runtime/Renderer/Private/ReflectionEnvironmentDiffuseIrradiance.cpp
@@ -250,7 +250,7 @@ void ComputeDiffuseIrradiance(FRHICommandListImmediate& RHICmdList, ERHIFeatureL
 			
 			for (int32 CubeFace = 0; CubeFace < CubeFace_MAX; CubeFace++)
 			{
-				SetRenderTarget(RHICmdList, EffectiveRT.TargetableTexture, 0, CubeFace, NULL);
+				SetRenderTarget(RHICmdList, EffectiveRT.TargetableTexture, 0, CubeFace, NULL, true);
 					
 				const FIntRect ViewRect(0, 0, MipSize, MipSize);
 				RHICmdList.SetViewport(0, 0, 0.0f, MipSize, MipSize, 1.0f);
@@ -294,7 +294,7 @@ void ComputeDiffuseIrradiance(FRHICommandListImmediate& RHICmdList, ERHIFeatureL
 				
 				for (int32 CubeFace = 0; CubeFace < CubeFace_MAX; CubeFace++)
 				{
-					SetRenderTarget(RHICmdList, EffectiveRT.TargetableTexture, MipIndex, CubeFace, NULL);
+					SetRenderTarget(RHICmdList, EffectiveRT.TargetableTexture, MipIndex, CubeFace, NULL, true);
 						
 					const FIntRect ViewRect(0, 0, MipSize, MipSize);
 					RHICmdList.SetViewport(0, 0, 0.0f, MipSize, MipSize, 1.0f);
@@ -330,6 +330,7 @@ void ComputeDiffuseIrradiance(FRHICommandListImmediate& RHICmdList, ERHIFeatureL
 
 			//load/store actions so we don't lose results as we render one pixel at a time on tile renderers.
 			FRHIRenderTargetView RTV(EffectiveRT.TargetableTexture, 0, -1, ERenderTargetLoadAction::ELoad, ERenderTargetStoreAction::EStore);
+			RHICmdList.TransitionResource(EResourceTransitionAccess::EWritable, EffectiveRT.TargetableTexture);
 			RHICmdList.SetRenderTargets(1, &RTV, nullptr, 0, nullptr);
 
 			const FIntRect ViewRect(CoefficientIndex, 0, CoefficientIndex + 1, 1);
