Temp workaround for editor transition ensures.

diff --git a/Engine/Source/Runtime/Renderer/Private/PostProcess/PostProcessBokehDOFRecombine.cpp b/Engine/Source/Runtime/Renderer/Private/PostProcess/PostProcessBokehDOFRecombine.cpp
--- a/Engine/Source/Runtime/Renderer/Private/PostProcess/PostProcessBokehDOFRecombine.cpp
+++ b/Engine/Source/Runtime/Renderer/Private/PostProcess/PostProcessBokehDOFRecombine.cpp
@@ -194,6 +194,7 @@ FPooledRenderTargetDesc FRCPassPostProcessBokehDOFRecombine::ComputeOutputDesc(E
 	FPooledRenderTargetDesc Ret = GetInput(ePId_Input0)->GetOutput()->RenderTargetDesc;
 
 	Ret.Reset();
+	Ret.AutoWritable = false;
 	Ret.DebugName = TEXT("BokehDOFRecombine");
 
 	return Ret;
