Removed duplicate resolve code for the velocity buffer in order to be compliant with the  new RHITransitionResources API.

diff --git a/Engine/Source/Runtime/Renderer/Private/PostProcess/SceneRenderTargets.cpp b/Engine/Source/Runtime/Renderer/Private/PostProcess/SceneRenderTargets.cpp
--- a/Engine/Source/Runtime/Renderer/Private/PostProcess/SceneRenderTargets.cpp
+++ b/Engine/Source/Runtime/Renderer/Private/PostProcess/SceneRenderTargets.cpp
@@ -462,10 +462,12 @@ int32 FSceneRenderTargets::GetGBufferRenderTargets(ERenderTargetLoadAction Color
 	OutRenderTargets[MRTCount++] = FRHIRenderTargetView(GBufferB->GetRenderTargetItem().TargetableTexture, 0, -1, ColorLoadAction, ERenderTargetStoreAction::EStore);
 	OutRenderTargets[MRTCount++] = FRHIRenderTargetView(GBufferC->GetRenderTargetItem().TargetableTexture, 0, -1, ColorLoadAction, ERenderTargetStoreAction::EStore);
 
+	// The velocity buffer needs to be bound before other optionnal rendertargets (when UseSelecUseSelectiveBasePassOutputs() is true).
+	// Otherwise there is an issue on some AMD hardware where the target does not get updated. Seems to be related to the velocity buffer format as it works fine with other targets.
 	if (bAllocateVelocityGBuffer)
 	{
 		OutVelocityRTIndex = MRTCount;
-		// check(OutVelocityRTIndex == GetGBufferVelocityIndex());
+		check(OutVelocityRTIndex == 4); // As defined in BasePassPixelShader.usf
 		OutRenderTargets[MRTCount++] = FRHIRenderTargetView(GBufferVelocity->GetRenderTargetItem().TargetableTexture, 0, -1, ColorLoadAction, ERenderTargetStoreAction::EStore);
 	}
 	else
@@ -478,7 +480,7 @@ int32 FSceneRenderTargets::GetGBufferRenderTargets(ERenderTargetLoadAction Color
 
 	if (bAllowStaticLighting)
 	{
-		// check(MRTCount == GetGBufferEIndex());
+		check(MRTCount == (bAllocateVelocityGBuffer ? 6 : 5)); // As defined in BasePassPixelShader.usf
 		OutRenderTargets[MRTCount++] = FRHIRenderTargetView(GBufferE->GetRenderTargetItem().TargetableTexture, 0, -1, ColorLoadAction, ERenderTargetStoreAction::EStore);
 	}
 
@@ -576,7 +578,11 @@ void FSceneRenderTargets::FinishRenderingGBuffer(FRHICommandListImmediate& RHICm
 	FResolveParams ResolveParams;
 	for (int32 i = 0; i < NumMRTs; ++i)
 	{
-		RHICmdList.CopyToResolveTarget(RenderTargets[i].Texture, RenderTargets[i].Texture, true, ResolveParams);
+		 // When the basepass outputs to the velocity buffer, don't resolve it yet if selective outputs are enabled, as it will be resolved after the velocity pass.
+		if (i != VelocityRTIndex || !UseSelectiveBasePassOutputs())
+		{
+			RHICmdList.CopyToResolveTarget(RenderTargets[i].Texture, RenderTargets[i].Texture, true, ResolveParams);
+		}
 	}
 	FTextureRHIParamRef DepthSurface = GetSceneDepthSurface();
 	RHICmdList.CopyToResolveTarget(DepthSurface, DepthSurface, true, ResolveParams);
