Refraction on mobile, enabled by default for high-end devices

diff --git a/Engine/Config/BaseDeviceProfiles.ini b/Engine/Config/BaseDeviceProfiles.ini
--- a/Engine/Config/BaseDeviceProfiles.ini
+++ b/Engine/Config/BaseDeviceProfiles.ini
@@ -115,6 +115,7 @@ BaseProfileName=IOS
 +CVars=r.BloomQuality=1
 +CVars=r.DepthOfFieldQuality=1
 +CVars=r.LightShaftQuality=1
++CVars=r.RefractionQuality=1
 
 [iPadMini DeviceProfile]
 DeviceType=IOS
@@ -147,6 +148,7 @@ BaseProfileName=IOS
 +CVars=r.BloomQuality=1
 +CVars=r.DepthOfFieldQuality=1
 +CVars=r.LightShaftQuality=1
++CVars=r.RefractionQuality=1
 
 [iPodTouch5 DeviceProfile]
 DeviceType=IOS
@@ -159,6 +161,7 @@ BaseProfileName=IOS
 +CVars=r.BloomQuality=1
 +CVars=r.DepthOfFieldQuality=1
 +CVars=r.LightShaftQuality=1
++CVars=r.RefractionQuality=1
 
 [iPhone6Plus DeviceProfile]
 DeviceType=IOS
@@ -166,6 +169,7 @@ BaseProfileName=IOS
 +CVars=r.BloomQuality=1
 +CVars=r.DepthOfFieldQuality=1
 +CVars=r.LightShaftQuality=1
++CVars=r.RefractionQuality=1
 
 [/Script/AndroidDeviceProfileSelector.AndroidDeviceProfileMatchingRules]
 MatchProfile=(Profile="Android_PowerVR540",Match=((SourceType=SRC_GpuFamily,CompareType=CMP_Equal,MatchString="PowerVR SGX 540")))
@@ -213,6 +217,7 @@ BaseProfileName=Android
 +CVars=r.BloomQuality=1
 +CVars=r.DepthOfFieldQuality=1
 +CVars=r.LightShaftQuality=1
++CVars=r.RefractionQuality=1
 +CVars=r.MobileContentScaleFactor=1.0
 
 [Android_Default DeviceProfile]
diff --git a/Engine/Shaders/DistortAccumulatePS.usf b/Engine/Shaders/DistortAccumulatePS.usf
--- a/Engine/Shaders/DistortAccumulatePS.usf
+++ b/Engine/Shaders/DistortAccumulatePS.usf
@@ -60,12 +60,13 @@ void Main(
 
 	// Sample depth at distortion offset
 	float2 DistortScreenUV = ScreenUV + Distortion;
-#if COMPILER_GLSL_ES2
-	float DistortSceneDepth = Texture2DSample(SceneColorTexture, SceneColorTextureSampler, DistortScreenUV).w;
+
+#if (ES2_PROFILE || ES3_1_PROFILE) 
+	float DistortSceneDepth = ConvertFromDeviceZ(Texture2DSampleLevel(SceneDepthTexture, SceneDepthTextureSampler, DistortScreenUV, 0).r);
 #else
 	float DistortSceneDepth = CalcSceneDepth(DistortScreenUV);
-#endif
-
+#endif // (ES2_PROFILE || ES3_1_PROFILE) 
+	
 	// Soft thresholding 
 	float Bias = -Refraction.y;
 	float Range = clamp(abs(Bias * 0.5f), 0, 50);
@@ -74,6 +75,11 @@ void Main(
 	float InvWidth = 1.0f / max(1.0f, Range );
 	Distortion *= saturate((Z - ZCompare) * InvWidth + Bias);
 
+#if (ES2_PROFILE || ES3_1_PROFILE) 
+	// mobile path direclty outputs distorted colors
+	DistortScreenUV = ScreenUV + Distortion;
+	OutColor = CalcFullSceneColor(DistortScreenUV);
+#else
 	//Scale up for better precision in low/subtle refractions at the expense of artefacts at higher refraction.
 	Distortion*= DistortionScaleBias;
 
@@ -83,6 +89,7 @@ void Main(
 
 	// output positives in R|G channels and negatives in B|A channels
 	OutColor = float4(PosOffset.x,PosOffset.y,NegOffset.x,NegOffset.y);
+#endif // (ES2_PROFILE || ES3_1_PROFILE) 
 }
 
 
diff --git a/Engine/Source/Runtime/Renderer/Private/DistortionRendering.cpp b/Engine/Source/Runtime/Renderer/Private/DistortionRendering.cpp
--- a/Engine/Source/Runtime/Renderer/Private/DistortionRendering.cpp
+++ b/Engine/Source/Runtime/Renderer/Private/DistortionRendering.cpp
@@ -995,3 +995,96 @@ void FSceneRenderer::RenderDistortion(FRHICommandListImmediate& RHICmdList)
 	}
 }
 
+void FSceneRenderer::RenderDistortionES2(FRHICommandListImmediate& RHICmdList)
+{
+	// do we need to render the distortion pass?
+	bool bRender=false;
+	for(int32 ViewIndex = 0;ViewIndex < Views.Num();ViewIndex++)
+	{
+		const FViewInfo& View = Views[ViewIndex];
+		if( View.DistortionPrimSet.NumPrims() > 0)
+		{
+			bRender=true;
+			break;
+		}
+	}
+			
+	if (bRender)
+	{
+		// Apply distortion
+		SCOPED_DRAW_EVENT(RHICmdList, Distortion);
+		
+		FSceneRenderTargets& SceneContext = FSceneRenderTargets::Get(RHICmdList);
+
+		SceneContext.ResolveSceneColor(RHICmdList, FResolveRect(0, 0, ViewFamily.FamilySizeX, ViewFamily.FamilySizeY));
+	
+		TRefCountPtr<IPooledRenderTarget> SceneColorDistorted;
+		GRenderTargetPool.FindFreeElement(RHICmdList, SceneContext.GetSceneColor()->GetDesc(), SceneColorDistorted, TEXT("SceneColorDistorted"));
+		const FSceneRenderTargetItem& DistortedRenderTarget = SceneColorDistorted->GetRenderTargetItem();
+
+		SetRenderTarget(RHICmdList, DistortedRenderTarget.TargetableTexture, SceneContext.GetSceneDepthSurface(), ESimpleRenderTargetMode::EClearColorExistingDepth, FExclusiveDepthStencil::DepthRead_StencilNop);
+
+		// Copy scene color to a new render target
+		for(int32 ViewIndex = 0, Num = Views.Num(); ViewIndex < Num; ++ViewIndex)
+		{
+			SCOPED_CONDITIONAL_DRAW_EVENTF(RHICmdList, EventView, Views.Num() > 1, TEXT("View%d"), ViewIndex);
+
+			FViewInfo& View = Views[ViewIndex];
+
+			// useful when we move this into the compositing graph
+			FRenderingCompositePassContext Context(RHICmdList, View);
+
+			// Set the view family's render target/viewport.
+			Context.SetViewportAndCallRHI(View.ViewRect);				
+
+			RHICmdList.SetBlendState(TStaticBlendState<>::GetRHI());
+			RHICmdList.SetRasterizerState(TStaticRasterizerState<>::GetRHI());
+			RHICmdList.SetDepthStencilState(TStaticDepthStencilState<false, CF_Always>::GetRHI());
+					
+			static FGlobalBoundShaderState BoundShaderState;
+				
+			TShaderMapRef<FPostProcessVS> VertexShader(View.ShaderMap);
+			TShaderMapRef<FDistortionMergePS> PixelShader(View.ShaderMap);
+			SetGlobalBoundShaderState(RHICmdList, View.GetFeatureLevel(), BoundShaderState, GFilterVertexDeclaration.VertexDeclarationRHI, *VertexShader, *PixelShader);
+
+			VertexShader->SetParameters(Context);
+			PixelShader->SetParameters(Context, SceneContext.GetSceneColor()->GetRenderTargetItem().ShaderResourceTexture);
+
+			DrawRectangle(
+				RHICmdList,
+				0, 0,
+				View.ViewRect.Width(), View.ViewRect.Height(),
+				View.ViewRect.Min.X, View.ViewRect.Min.Y, 
+				View.ViewRect.Width(), View.ViewRect.Height(),
+				View.ViewRect.Size(),
+				SceneContext.GetBufferSizeXY(),
+				*VertexShader,
+				EDRF_UseTriangleOptimization);
+		}
+	
+		// Distort scene color in place
+		for(int32 ViewIndex = 0, Num = Views.Num(); ViewIndex < Num; ++ViewIndex)
+		{
+			SCOPED_CONDITIONAL_DRAW_EVENTF(RHICmdList, EventView, Views.Num() > 1, TEXT("View%d"), ViewIndex);
+
+			FViewInfo& View = Views[ViewIndex];
+
+			// useful when we move this into the compositing graph
+			FRenderingCompositePassContext Context(RHICmdList, View);
+
+			// Set the view family's render target/viewport.
+			Context.SetViewportAndCallRHI(View.ViewRect);
+
+			// test against depth
+			RHICmdList.SetBlendState(TStaticBlendState<>::GetRHI());
+			RHICmdList.SetRasterizerState(TStaticRasterizerState<>::GetRHI());
+			RHICmdList.SetDepthStencilState(TStaticDepthStencilState<false, CF_DepthNearOrEqual>::GetRHI());
+
+			// draw only distortion meshes
+			View.DistortionPrimSet.DrawAccumulatedOffsets(RHICmdList, View, false);
+		}
+	
+		// Set distorted scene color as main
+		SceneContext.SetSceneColor(SceneColorDistorted);
+	}
+}
\ No newline at end of file
diff --git a/Engine/Source/Runtime/Renderer/Private/ForwardShadingRenderer.cpp b/Engine/Source/Runtime/Renderer/Private/ForwardShadingRenderer.cpp
--- a/Engine/Source/Runtime/Renderer/Private/ForwardShadingRenderer.cpp
+++ b/Engine/Source/Runtime/Renderer/Private/ForwardShadingRenderer.cpp
@@ -172,10 +172,10 @@ void FForwardShadingSceneRenderer::Render(FRHICommandListImmediate& RHICmdList)
 		// Having the distortion applied between two different translucency passes would make it consistent with the deferred pass.
 		// This is not done yet.
 
-		if(GetRefractionQuality(ViewFamily) > 0)
+		if (GetRefractionQuality(ViewFamily) > 0)
 		{
 			// to apply refraction effect by distorting the scene color
-			RenderDistortion(RHICmdList);
+			RenderDistortionES2(RHICmdList);
 		}
 		RenderTranslucency(RHICmdList);
 	}
diff --git a/Engine/Source/Runtime/Renderer/Private/PostProcess/SceneRenderTargets.cpp b/Engine/Source/Runtime/Renderer/Private/PostProcess/SceneRenderTargets.cpp
--- a/Engine/Source/Runtime/Renderer/Private/PostProcess/SceneRenderTargets.cpp
+++ b/Engine/Source/Runtime/Renderer/Private/PostProcess/SceneRenderTargets.cpp
@@ -1953,14 +1953,16 @@ const FTexture2DRHIRef* FSceneRenderTargets::GetActualDepthTexture() const
 	}
 	else if (IsMobilePlatform(GShaderPlatformForFeatureLevel[CurrentFeatureLevel]))
 	{
-		bool bSceneDepthInAlpha = (GetSceneColor()->GetDesc().Format == PF_FloatRGBA);
-		bool bOnChipDepthFetch = (GSupportsShaderDepthStencilFetch || (bSceneDepthInAlpha && GSupportsShaderFramebufferFetch));
-		
-		if (bOnChipDepthFetch)
-		{
-			DepthTexture = (const FTexture2DRHIRef*)(&GSystemTextures.DepthDummy->GetRenderTargetItem().ShaderResourceTexture);
-		}
-		else
+		// TODO: avoid depth texture fetch when shader needs fragment previous depth and device supports framebuffer fetch
+
+		//bool bSceneDepthInAlpha = (GetSceneColor()->GetDesc().Format == PF_FloatRGBA);
+		//bool bOnChipDepthFetch = (GSupportsShaderDepthStencilFetch || (bSceneDepthInAlpha && GSupportsShaderFramebufferFetch));
+		//
+		//if (bOnChipDepthFetch)
+		//{
+		//	DepthTexture = (const FTexture2DRHIRef*)(&GSystemTextures.DepthDummy->GetRenderTargetItem().ShaderResourceTexture);
+		//}
+		//else
 		{
 			DepthTexture = &GetSceneDepthTexture();
 		}
diff --git a/Engine/Source/Runtime/Renderer/Private/SceneRendering.h b/Engine/Source/Runtime/Renderer/Private/SceneRendering.h
--- a/Engine/Source/Runtime/Renderer/Private/SceneRendering.h
+++ b/Engine/Source/Runtime/Renderer/Private/SceneRendering.h
@@ -858,6 +858,7 @@ protected:
 
 	/** Renders the scene's distortion */
 	void RenderDistortion(FRHICommandListImmediate& RHICmdList);
+	void RenderDistortionES2(FRHICommandListImmediate& RHICmdList);
 
 	static int32 GetRefractionQuality(const FSceneViewFamily& ViewFamily);
 };
