Jira UE-21208

temporarily inserting ensureMsgf to capture and report geometry with bad bounds.

diff --git a/Engine/Source/Runtime/Renderer/Private/PrimitiveSceneInfo.cpp b/Engine/Source/Runtime/Renderer/Private/PrimitiveSceneInfo.cpp
--- a/Engine/Source/Runtime/Renderer/Private/PrimitiveSceneInfo.cpp
+++ b/Engine/Source/Runtime/Renderer/Private/PrimitiveSceneInfo.cpp
@@ -403,10 +403,40 @@ void FPrimitiveSceneInfo::UnlinkAttachmentGroup()
 	}
 }
 
+// Temporary:  Used to gather data for a bug QA hasn't been able to reproduce internally.
+#define TEST_BBOX_NANS_UE_21208
+
 void FPrimitiveSceneInfo::GatherLightingAttachmentGroupPrimitives(TArray<FPrimitiveSceneInfo*, SceneRenderingAllocator>& OutChildSceneInfos)
 {
+
+#ifdef TEST_BBOX_NANS_UE_21208
+	
+	// local function that returns the name of object
+	auto GetObjectName = [](const UPrimitiveComponent* InPrimitive)->FString
+	{
+		return (InPrimitive) ? InPrimitive->GetFullName() : FString(TEXT("Unknown Object"));
+	};
+
+	// verify that the current object has a valid bbox before adding it
+	const float& BoundsRadius = this->Proxy->GetBounds().SphereRadius;
+	if (ensureMsgf(!FMath::IsNaN(BoundsRadius) && FMath::IsFinite(BoundsRadius),
+		TEXT("%s had an ill-formed bbox and was skipped during shadow setup, contact DavidH."), *GetObjectName(this->ComponentForDebuggingOnly)))
+	{
+
+		OutChildSceneInfos.Add(this);
+	} 
+	else
+	{
+		// return, leaving the TArray empty
+		return;
+	}
+
+#else
+	// add self at the head of the quie
 	OutChildSceneInfos.Add(this);
 
+#endif
+
 	if (!LightingAttachmentRoot.IsValid() && Proxy->LightAttachmentsAsGroup())
 	{
 		const FAttachmentGroupSceneInfo* AttachmentGroup = Scene->AttachmentGroups.Find(PrimitiveComponentId);
@@ -416,13 +446,35 @@ void FPrimitiveSceneInfo::GatherLightingAttachmentGroupPrimitives(TArray<FPrimit
 			for (int32 ChildIndex = 0; ChildIndex < AttachmentGroup->Primitives.Num(); ChildIndex++)
 			{
 				FPrimitiveSceneInfo* ShadowChild = AttachmentGroup->Primitives[ChildIndex];
-				checkSlow(!OutChildSceneInfos.Contains(ShadowChild))
-				OutChildSceneInfos.Add(ShadowChild);
+				for (int32 ChildIndex = 0, ChildIndexMax = AttachmentGroup->Primitives.Num(); ChildIndex < ChildIndexMax; ChildIndex++)
+				{
+					FPrimitiveSceneInfo* ShadowChild = AttachmentGroup->Primitives[ChildIndex];
+
+#ifdef TEST_BBOX_NANS_UE_21208
+					// Only enqueue objects with valid bounds using the normality of the SphereRaduis as criteria.
+
+					const float& ShadowChildBoundsRadius = ShadowChild->Proxy->GetBounds().SphereRadius;
+					 
+					if (ensureMsgf(!FMath::IsNaN(ShadowChildBoundsRadius) && FMath::IsFinite(ShadowChildBoundsRadius), 
+						TEXT("%s had an ill-formed bbox and was skipped during shadow setup, contact DavidH."), *GetObjectName(ShadowChild->ComponentForDebuggingOnly)))
+					{
+						checkSlow(!OutChildSceneInfos.Contains(ShadowChild))
+						OutChildSceneInfos.Add(ShadowChild);
+					}
+#else
+					// enqueue all objects.
+
+					checkSlow(!OutChildSceneInfos.Contains(ShadowChild))
+					OutChildSceneInfos.Add(ShadowChild);
+#endif
+				}
 			}
 		}
 	}
 }
 
+#undef TEST_BBOX_NANS_UE_21208
+
 FBoxSphereBounds FPrimitiveSceneInfo::GetAttachmentGroupBounds() const
 {
 	FBoxSphereBounds Bounds = Proxy->GetBounds();
diff --git a/Engine/Source/Runtime/Renderer/Private/ShadowSetup.cpp b/Engine/Source/Runtime/Renderer/Private/ShadowSetup.cpp
--- a/Engine/Source/Runtime/Renderer/Private/ShadowSetup.cpp
+++ b/Engine/Source/Runtime/Renderer/Private/ShadowSetup.cpp
@@ -1336,6 +1336,19 @@ void FSceneRenderer::CreatePerObjectProjectedShadow(
 	TArray<FPrimitiveSceneInfo*, SceneRenderingAllocator> ShadowGroupPrimitives;
 	PrimitiveSceneInfo->GatherLightingAttachmentGroupPrimitives(ShadowGroupPrimitives);
 
+
+#define TEST_BBOX_NANS_UE_21208
+
+#ifdef TEST_BBOX_NANS_UE_21208
+	// allow for silent failure 
+	if (ShadowGroupPrimitives.Num() == 0)
+	{
+		return;
+	}
+#endif
+
+#undef TEST_BBOX_NANS_UE_21208
+
 	// Compute the composite bounds of this group of shadow primitives.
 	FBoxSphereBounds OriginalBounds = ShadowGroupPrimitives[0]->Proxy->GetBounds();
 
