Jira UE-21208
corrected a merge error.

diff --git a/Engine/Source/Runtime/Renderer/Private/PrimitiveSceneInfo.cpp b/Engine/Source/Runtime/Renderer/Private/PrimitiveSceneInfo.cpp
--- a/Engine/Source/Runtime/Renderer/Private/PrimitiveSceneInfo.cpp
+++ b/Engine/Source/Runtime/Renderer/Private/PrimitiveSceneInfo.cpp
@@ -443,31 +443,28 @@ void FPrimitiveSceneInfo::GatherLightingAttachmentGroupPrimitives(TArray<FPrimit
 
 		if (AttachmentGroup)
 		{
-			for (int32 ChildIndex = 0; ChildIndex < AttachmentGroup->Primitives.Num(); ChildIndex++)
+			
+			for (int32 ChildIndex = 0, ChildIndexMax = AttachmentGroup->Primitives.Num(); ChildIndex < ChildIndexMax; ChildIndex++)
 			{
 				FPrimitiveSceneInfo* ShadowChild = AttachmentGroup->Primitives[ChildIndex];
-				for (int32 ChildIndex = 0, ChildIndexMax = AttachmentGroup->Primitives.Num(); ChildIndex < ChildIndexMax; ChildIndex++)
-				{
-					FPrimitiveSceneInfo* ShadowChild = AttachmentGroup->Primitives[ChildIndex];
 
 #ifdef TEST_BBOX_NANS_UE_21208
-					// Only enqueue objects with valid bounds using the normality of the SphereRaduis as criteria.
+				// Only enqueue objects with valid bounds using the normality of the SphereRaduis as criteria.
+
+				const float& ShadowChildBoundsRadius = ShadowChild->Proxy->GetBounds().SphereRadius;
 
-					const float& ShadowChildBoundsRadius = ShadowChild->Proxy->GetBounds().SphereRadius;
-					 
-					if (ensureMsgf(!FMath::IsNaN(ShadowChildBoundsRadius) && FMath::IsFinite(ShadowChildBoundsRadius), 
-						TEXT("%s had an ill-formed bbox and was skipped during shadow setup, contact DavidH."), *GetObjectName(ShadowChild->ComponentForDebuggingOnly)))
-					{
+				if (ensureMsgf(!FMath::IsNaN(ShadowChildBoundsRadius) && FMath::IsFinite(ShadowChildBoundsRadius),
+					TEXT("%s had an ill-formed bbox and was skipped during shadow setup, contact DavidH."), *GetObjectName(ShadowChild->ComponentForDebuggingOnly)))
+				{
 						checkSlow(!OutChildSceneInfos.Contains(ShadowChild))
-						OutChildSceneInfos.Add(ShadowChild);
-					}
+							OutChildSceneInfos.Add(ShadowChild);
+				}
 #else
-					// enqueue all objects.
+				// enqueue all objects.
 
-					checkSlow(!OutChildSceneInfos.Contains(ShadowChild))
-					OutChildSceneInfos.Add(ShadowChild);
+				checkSlow(!OutChildSceneInfos.Contains(ShadowChild))
+			    OutChildSceneInfos.Add(ShadowChild);
 #endif
-				}
 			}
 		}
 	}
