Remove HDR_LINEAR_32 shader permutations for MobileHDR on Android

#UEMOB-45

diff --git a/Engine/Shaders/BasePassForForwardShadingPixelShader.usf b/Engine/Shaders/BasePassForForwardShadingPixelShader.usf
--- a/Engine/Shaders/BasePassForForwardShadingPixelShader.usf
+++ b/Engine/Shaders/BasePassForForwardShadingPixelShader.usf
@@ -369,12 +369,16 @@ void Main(
 		OutColor = half4(FoggedColor, Opacity);
 	#else
 		OutColor.rgb = Color * VertexFog.a + VertexFog.rgb;
-		#if USE_32BPP_HDR || OUTPUT_GAMMA_SPACE 
-			// Scene color alpha is not used yet so we set it to 0
-			OutColor.a = 0.0;
-		#else
-			// Place Z in FP16 alpha value.
-			OutColor.a = MaterialParameters.ScreenPosition.w;
+
+		// Scene color alpha is not used yet so we set it to 0
+		OutColor.a = 0.0;
+
+		#if ES2_PROFILE && !OUTPUT_GAMMA_SPACE 
+			if(GetHDR32bppEncodeMode() == HDR_ENCODE_NONE)
+			{
+				// FP16 HDR stores Z in framebuffer alpha
+				OutColor.a = MaterialParameters.ScreenPosition.w;
+			}
 		#endif
 	#endif
 
@@ -388,14 +392,16 @@ void Main(
 		OutColor.rgb = sqrt( OutColor.rgb );
 	#endif
 
-	#if (ES2_PROFILE || ES3_1_PROFILE) && (!OUTPUT_GAMMA_SPACE)
-		#if (COMPILER_GLSL_ES2 && USE_32BPP_HDR) || (MOBILE_EMULATION)
-			if (GetHDR32bppEncodeMode() == HDR_ENCODE_RGBA)
+	#if ES2_PROFILE && !OUTPUT_GAMMA_SPACE
+		float Mode = GetHDR32bppEncodeMode();
+		if (Mode != HDR_ENCODE_NONE)
+		{
+			if (Mode == HDR_ENCODE_RGBA)
 			{
 				OutColor.rgb = FrameBufferBlendOp(OutColor);
 			}
 			OutColor = Encode32BPPHDR(OutColor, SvPosition.xy);
-		#endif
+		}
 	#endif
 
 	// Editor primitive depth testing
diff --git a/Engine/Shaders/Common.usf b/Engine/Shaders/Common.usf
--- a/Engine/Shaders/Common.usf
+++ b/Engine/Shaders/Common.usf
@@ -830,9 +830,9 @@ MaterialFloat3 RGBTDecode8BPC(MaterialFloat4 RGBT, MaterialFloat Range)
 	return RGBT.rgb * RGBT.a;
 }
 
-#define HDR_ENCODE_NONE 0.0
-#define HDR_ENCODE_MOSAIC 1.0
-#define HDR_ENCODE_RGBA 2.0
+#define HDR_ENCODE_NONE 0.0		// 64bpp HDR
+#define HDR_ENCODE_MOSAIC 1.0	// 32bpp HDR using Mosaic encoding
+#define HDR_ENCODE_RGBA 2.0		// 32bpp HDR using RGBA encoding
 
 float GetHDR32bppEncodeMode()
 {
diff --git a/Engine/Source/Runtime/Renderer/Private/ForwardBasePassRendering.h b/Engine/Source/Runtime/Renderer/Private/ForwardBasePassRendering.h
--- a/Engine/Source/Runtime/Renderer/Private/ForwardBasePassRendering.h
+++ b/Engine/Source/Runtime/Renderer/Private/ForwardBasePassRendering.h
@@ -32,15 +32,17 @@ public:
 	FVector4 LightColorAndFalloffExponent[MAX_BASEPASS_DYNAMIC_POINT_LIGHTS];
 };
 
-static bool ShouldCacheShaderByOutputFormat(EOutputFormat OutputFormat)
+static bool ShouldCacheShaderByPlatformAndOutputFormat(EShaderPlatform Platform, EOutputFormat OutputFormat)
 {
 	bool bSupportsMobileHDR = IsMobileHDR();
 	bool bShaderUsesLDR = (OutputFormat == LDR_GAMMA_32);
 	bool bShaderUsesHDR = !bShaderUsesLDR;
+	// Android ES2 uses intrinsic_GetHDR32bppEncodeModeES2 so doesn't need a HDR_LINEAR_32 permutation
+	bool bIsAndroid32bpp = (OutputFormat == HDR_LINEAR_32) && (Platform == SP_OPENGL_ES2_ANDROID);
 
 	// only cache this shader if the LDR/HDR output matches what we currently support.  IsMobileHDR can't change, so we don't need
-	// the LDR shaders if we are doing HDR, and vice-versa.	
-	return (bShaderUsesLDR && !bSupportsMobileHDR) || (bShaderUsesHDR && bSupportsMobileHDR);
+	// the LDR shaders if we are doing HDR, and vice-versa.	Android doesn't need HDR_LINEAR_32 as it
+	return (bShaderUsesLDR && !bSupportsMobileHDR) || (bShaderUsesHDR && bSupportsMobileHDR && !bIsAndroid32bpp);
 }
 
 /**
@@ -112,7 +114,7 @@ public:
 	
 	static bool ShouldCache(EShaderPlatform Platform, const FMaterial* Material, const FVertexFactoryType* VertexFactoryType)
 	{		
-		return TBasePassForForwardShadingVSBaseType<LightMapPolicyType>::ShouldCache(Platform, Material, VertexFactoryType) && ShouldCacheShaderByOutputFormat(OutputFormat);
+		return TBasePassForForwardShadingVSBaseType<LightMapPolicyType>::ShouldCache(Platform, Material, VertexFactoryType) && ShouldCacheShaderByPlatformAndOutputFormat(Platform,OutputFormat);
 	}
 
 	static void ModifyCompilationEnvironment(EShaderPlatform Platform, const FMaterial* Material, FShaderCompilerEnvironment& OutEnvironment)
@@ -270,7 +272,7 @@ public:
 		// Only compile skylight version for lit materials on ES2 (Metal) or higher
 		const bool bShouldCacheBySkylight = !bEnableSkyLight || (Material->GetShadingModel() != MSM_Unlit);
 
-		return TBasePassForForwardShadingPSBaseType<LightMapPolicyType, NumDynamicPointLights>::ShouldCache(Platform, Material, VertexFactoryType) && ShouldCacheShaderByOutputFormat(OutputFormat) && bShouldCacheBySkylight;
+		return TBasePassForForwardShadingPSBaseType<LightMapPolicyType, NumDynamicPointLights>::ShouldCache(Platform, Material, VertexFactoryType) && ShouldCacheShaderByPlatformAndOutputFormat(Platform, OutputFormat) && bShouldCacheBySkylight;
 	}
 	
 	static void ModifyCompilationEnvironment(EShaderPlatform Platform, const FMaterial* Material, FShaderCompilerEnvironment& OutEnvironment)
@@ -336,7 +338,7 @@ public:
 		SceneTextureMode(InSceneTextureMode),
 		bEnableEditorPrimitiveDepthTest(bInEnableEditorPrimitiveDepthTest)
 	{
-		if (IsMobileHDR32bpp())
+		if (IsMobileHDR32bpp() && !GSupportsHDR32bppEncodeModeIntrinsic)
 		{
 			VertexShader = InMaterialResource.GetShader<TBasePassForForwardShadingVS<LightMapPolicyType, HDR_LINEAR_64> >(InVertexFactory->GetType());
 
