Integrated

 CL 2705306

improved CircleDOF quality, added r.DepthOfFieldQuality 4 to trigger slower but higher quality

diff --git a/Engine/Shaders/PostProcessCircleDOF.usf b/Engine/Shaders/PostProcessCircleDOF.usf
--- a/Engine/Shaders/PostProcessCircleDOF.usf
+++ b/Engine/Shaders/PostProcessCircleDOF.usf
@@ -129,8 +129,7 @@ void CircleDilatePS(float4 UVAndScreenPos : TEXCOORD0, out float OutColor : SV_T
 }
 
 
-// pixel shader entry point
-void CirclePS(float4 UVAndScreenPos : TEXCOORD0, out float4 OutColor0 : SV_Target0)
+void Circle4Samples(float4 UVAndScreenPos : TEXCOORD0, out float4 OutColor0 : SV_Target0, float InRand)
 {
 	float2 UV = UVAndScreenPos.xy;
 
@@ -143,18 +142,20 @@ void CirclePS(float4 UVAndScreenPos : TEXCOORD0, out float4 OutColor0 : SV_Targe
 	// Higher than 8 is too noizy for 4 samples.
 	float Coc = 8.0 * CircleDofParams.w;
 
+	float LocalRand = RandomOffset.x + InRand;
+
 	// Get base semi-random direction and dither along radius.
 	// Reused throughout the rest of the algorithm.
 	float TwoPi = 2.0 * 3.14159;
-	float RadianBase = NoizSnorm(UVAndScreenPos.xy, 0.010 * RandomOffset.x) * TwoPi;
-	float RadiusBase = NoizNorm(UVAndScreenPos.xy, 0.013 * RandomOffset.x);
+	float RadianBase = NoizSnorm(UVAndScreenPos.xy, 0.010 * LocalRand) * TwoPi;
+	float RadiusBase = NoizNorm(UVAndScreenPos.xy, 0.013 * LocalRand);
 
 	// Radius 
 	float RadiusBase2 = RadiusBase * (1.0/4.0);
 	float R1 = sqrt(RadiusBase2 + 3.0/4.0) * Coc;
 	float R2 = sqrt(RadiusBase2 + 2.0/4.0) * Coc;
 	float R3 = sqrt(RadiusBase2 + 1.0/4.0) * Coc;
-	float R4 = sqrt(RadiusBase2 + 0.0/4.0) * Coc;
+	float R4 = sqrt(RadiusBase2 + 0.6/4.0) * Coc;
 
 	float2 UV1 = RotVec(R1, RadianBase + TwoPi * 0.0/4.0);
 	float2 UV2 = RotVec(R2, RadianBase + TwoPi * 2.0/4.0);
@@ -407,7 +408,28 @@ void CirclePS(float4 UVAndScreenPos : TEXCOORD0, out float4 OutColor0 : SV_Targe
 	Weight += W1+W2+W3+W4;
 
 	OutColor0.rgb *= (1.0/Weight);
+}
 
+// pixel shader entry point
+void CirclePS(float4 UVAndScreenPos : TEXCOORD0, out float4 OutColor0 : SV_Target0)
+{
+#if QUALITY
+	// 2 or higher, gets slower but less noisy
+	const uint Count = 4;
+#else
+	const uint Count = 1;
+#endif
+
+	OutColor0 = 0;
+	for(uint i = 0; i < Count; ++i)
+	{
+		float4 Color;
+		Circle4Samples(UVAndScreenPos, Color, i);
+
+		OutColor0 += Color;
+	}
+
+	OutColor0 /= Count;
 }
 
 
@@ -530,6 +552,4 @@ void MainCircleRecombinePS(in float4 UVAndScreenPos : TEXCOORD0, out float4 OutC
 	// Blend in the limited version quickly to remove HDR jitter artifacts and noise.
 	float4 OutLimited = min(max(OutColor,HMin),HMax);
 	OutColor = lerp(OutColor, OutLimited, saturate(PixCoc*PixCoc*4.0));
-}
-
-
+}
\ No newline at end of file
diff --git a/Engine/Source/Runtime/Core/Private/HAL/ConsoleManager.cpp b/Engine/Source/Runtime/Core/Private/HAL/ConsoleManager.cpp
--- a/Engine/Source/Runtime/Core/Private/HAL/ConsoleManager.cpp
+++ b/Engine/Source/Runtime/Core/Private/HAL/ConsoleManager.cpp
@@ -1722,7 +1722,8 @@ static TAutoConsoleVariable<int32> CVarDepthOfFieldQuality(
 	TEXT("Allows to adjust the depth of field quality. Currently only fully affects BokehDOF. GaussianDOF is either 0 for off, otherwise on.\n")
 	TEXT(" 0: Off\n")
 	TEXT(" 1: Low\n")
-	TEXT(" 2: high quality (default, adaptive, can be 4x slower)"),
+	TEXT(" 2: high quality (default, adaptive, can be 4x slower)\n")
+	TEXT(" 3: Special mode only affecting CircleDOF for very high quality but slow rendering"),
 	ECVF_Scalability | ECVF_RenderThreadSafe);
 
 static TAutoConsoleVariable<float> CVarScreenPercentage(
diff --git a/Engine/Source/Runtime/Renderer/Private/PostProcess/PostProcessCircleDOF.cpp b/Engine/Source/Runtime/Renderer/Private/PostProcess/PostProcessCircleDOF.cpp
--- a/Engine/Source/Runtime/Renderer/Private/PostProcess/PostProcessCircleDOF.cpp
+++ b/Engine/Source/Runtime/Renderer/Private/PostProcess/PostProcessCircleDOF.cpp
@@ -479,7 +479,7 @@ static void TemporalRandom2(FVector2D* RESTRICT const Constant, uint32 FrameNumb
 	Constant->Y = TemporalHalton2(FrameNumber & 1023, 3);
 }
 
-template <uint32 NearBlurEnable>
+template <uint32 NearBlurEnable, uint32 Quality>
 class FPostProcessCircleDOFPS : public FGlobalShader
 {
 	DECLARE_SHADER_TYPE(FPostProcessCircleDOFPS, Global);
@@ -493,6 +493,7 @@ class FPostProcessCircleDOFPS : public FGlobalShader
 	{
 		FGlobalShader::ModifyCompilationEnvironment(Platform,OutEnvironment);
 		OutEnvironment.SetDefine(TEXT("ENABLE_NEAR_BLUR"), NearBlurEnable);
+		OutEnvironment.SetDefine(TEXT("QUALITY"), Quality);
 	}
 
 	/** Default constructor. */
@@ -548,10 +549,43 @@ public:
 
 		SetShaderValue(Context.RHICmdList, ShaderRHI, CircleDofParams, CircleDofCoc(Context.View));
 	}
+	
+	static const TCHAR* GetSourceFilename()
+	{
+		return TEXT("PostProcessCircleDOF");
+	}
+
+	static const TCHAR* GetFunctionName()
+	{
+		return TEXT("CirclePS");
+	}
 };
 
-IMPLEMENT_SHADER_TYPE(template<>,FPostProcessCircleDOFPS<0>,TEXT("PostProcessCircleDOF"),TEXT("CirclePS"),SF_Pixel);
-IMPLEMENT_SHADER_TYPE(template<>,FPostProcessCircleDOFPS<1>,TEXT("PostProcessCircleDOF"),TEXT("CirclePS"),SF_Pixel);
+
+// #define avoids a lot of code duplication
+#define VARIATION1(A, B) typedef FPostProcessCircleDOFPS<A, B> FPostProcessCircleDOFPS##A##B; \
+	IMPLEMENT_SHADER_TYPE2(FPostProcessCircleDOFPS##A##B, SF_Pixel);
+
+	VARIATION1(0,0)			VARIATION1(1,0)
+	VARIATION1(0,1)			VARIATION1(1,1)
+
+#undef VARIATION1
+
+template <uint32 NearBlurEnable, uint32 Quality>
+FShader* FRCPassPostProcessCircleDOF::SetShaderTempl(const FRenderingCompositePassContext& Context)
+{
+	TShaderMapRef<FPostProcessVS> VertexShader(Context.GetShaderMap());
+	TShaderMapRef<FPostProcessCircleDOFPS<NearBlurEnable, Quality> > PixelShader(Context.GetShaderMap());
+
+	static FGlobalBoundShaderState BoundShaderState;
+
+	SetGlobalBoundShaderState(Context.RHICmdList, Context.GetFeatureLevel(), BoundShaderState, GFilterVertexDeclaration.VertexDeclarationRHI, *VertexShader, *PixelShader);
+
+	VertexShader->SetParameters(Context);
+	PixelShader->SetParameters(Context);
+	
+	return *VertexShader;
+}
 
 void FRCPassPostProcessCircleDOF::Process(FRenderingCompositePassContext& Context)
 {
@@ -607,31 +641,27 @@ void FRCPassPostProcessCircleDOF::Process(FRenderingCompositePassContext& Contex
 	Context.RHICmdList.SetBlendState(TStaticBlendState<>::GetRHI());
 	Context.RHICmdList.SetRasterizerState(TStaticRasterizerState<>::GetRHI());
 	Context.RHICmdList.SetDepthStencilState(TStaticDepthStencilState<false, CF_Always>::GetRHI());
+		
+	static const auto CVar = IConsoleManager::Get().FindTConsoleVariableDataInt(TEXT("r.DepthOfFieldQuality"));
+	check(CVar);
+	int32 DOFQualityCVarValue = CVar->GetValueOnRenderThread();
+	
+	// 0:normal / 1:slow but very high quality
+	uint32 Quality = DOFQualityCVarValue >= 3;
 
-	TShaderMapRef<FPostProcessVS> VertexShader(ShaderMap);
+	FShader* VertexShader = 0;
 
 	if (bNearBlurEnabled)
 	{
-		static FGlobalBoundShaderState BoundShaderState;
-
-
-		TShaderMapRef< FPostProcessCircleDOFPS<1> > PixelShader(ShaderMap);
-		SetGlobalBoundShaderState(Context.RHICmdList, FeatureLevel, BoundShaderState, GFilterVertexDeclaration.VertexDeclarationRHI, *VertexShader, *PixelShader);
-
-		PixelShader->SetParameters(Context);
+		if(Quality) VertexShader = SetShaderTempl<1, 1>(Context);
+		  else		VertexShader = SetShaderTempl<1, 0>(Context);
 	}
 	else
 	{
-		static FGlobalBoundShaderState BoundShaderState;
-
-		TShaderMapRef< FPostProcessCircleDOFPS<0> > PixelShader(ShaderMap);
-		SetGlobalBoundShaderState(Context.RHICmdList, FeatureLevel, BoundShaderState, GFilterVertexDeclaration.VertexDeclarationRHI, *VertexShader, *PixelShader);
-
-		PixelShader->SetParameters(Context);
+		if(Quality)	VertexShader = SetShaderTempl<0, 1>(Context);
+		  else		VertexShader = SetShaderTempl<0, 0>(Context);
 	}
 
-	VertexShader->SetParameters(Context);
-
 	DrawPostProcessPass(
 		Context.RHICmdList,
 		DestRect.Min.X, DestRect.Min.Y,
@@ -640,7 +670,7 @@ void FRCPassPostProcessCircleDOF::Process(FRenderingCompositePassContext& Contex
 		SrcRect.Width() + 1, SrcRect.Height() + 1,
 		DestSize,
 		SrcSize,
-		*VertexShader,
+		VertexShader,
 		View.StereoPass,
 		Context.HasHmdMesh(),
 		EDRF_UseTriangleOptimization);
diff --git a/Engine/Source/Runtime/Renderer/Private/PostProcess/PostProcessCircleDOF.h b/Engine/Source/Runtime/Renderer/Private/PostProcess/PostProcessCircleDOF.h
--- a/Engine/Source/Runtime/Renderer/Private/PostProcess/PostProcessCircleDOF.h
+++ b/Engine/Source/Runtime/Renderer/Private/PostProcess/PostProcessCircleDOF.h
@@ -66,6 +66,9 @@ public:
 private:
 
 	bool bNearBlurEnabled;
+
+	template <uint32 NearBlurEnable, uint32 Quality>
+	FShader* SetShaderTempl(const FRenderingCompositePassContext& Context);
 };
 
 // derives from TRenderingCompositePassBase<InputCount, OutputCount> 
