Fix decal ensure.

diff --git a/Engine/Source/Runtime/Renderer/Private/CompositionLighting/PostProcessDeferredDecals.cpp b/Engine/Source/Runtime/Renderer/Private/CompositionLighting/PostProcessDeferredDecals.cpp
--- a/Engine/Source/Runtime/Renderer/Private/CompositionLighting/PostProcessDeferredDecals.cpp
+++ b/Engine/Source/Runtime/Renderer/Private/CompositionLighting/PostProcessDeferredDecals.cpp
@@ -485,7 +485,7 @@ void FRCPassPostProcessDeferredDecals::Process(FRenderingCompositePassContext& C
 		ResolveBufferMax,
 	};
 
-	bool TargetsToTransitionWritable[ResolveBufferMax] = { true };
+	bool TargetsToTransitionWritable[ResolveBufferMax] = { true, true, true, true, true, true, true };
 	FTextureRHIParamRef TargetsToResolve[ResolveBufferMax] = { nullptr };
 
 	if(CurrentStage == DRS_BeforeBasePass)
