Fix lightshaft ensures.

diff --git a/Engine/Source/Runtime/Renderer/Private/LightShaftRendering.cpp b/Engine/Source/Runtime/Renderer/Private/LightShaftRendering.cpp
--- a/Engine/Source/Runtime/Renderer/Private/LightShaftRendering.cpp
+++ b/Engine/Source/Runtime/Renderer/Private/LightShaftRendering.cpp
@@ -449,6 +449,7 @@ void AllocateOrReuseLightShaftRenderTarget(FRHICommandListImmediate& RHICmdList,
 		const FIntPoint BufferSize = FSceneRenderTargets::Get(RHICmdList).GetBufferSizeXY();
 		FIntPoint LightShaftSize(FMath::Max<uint32>(BufferSize.X / GetLightShaftDownsampleFactor(), 1), FMath::Max<uint32>(BufferSize.Y / GetLightShaftDownsampleFactor(), 1));
 		FPooledRenderTargetDesc Desc(FPooledRenderTargetDesc::Create2DDesc(LightShaftSize, LightShaftFilterBufferFormat, FClearValueBinding::Black, TexCreate_None, TexCreate_RenderTargetable, false));
+		Desc.AutoWritable = false;
 		GRenderTargetPool.FindFreeElement(RHICmdList, Desc, Target, Name);
 
 		SetRenderTarget(RHICmdList, Target->GetRenderTargetItem().TargetableTexture, FTextureRHIRef(), ESimpleRenderTargetMode::EClearColorExistingDepth);
@@ -613,7 +614,7 @@ void ApplyRadialBlurPasses(
 
 	for (uint32 PassIndex = 0; PassIndex < NumPasses; PassIndex++)
 	{
-		SetRenderTarget(RHICmdList, LightShaftsDest->GetRenderTargetItem().TargetableTexture, FTextureRHIRef());
+		SetRenderTarget(RHICmdList, LightShaftsDest->GetRenderTargetItem().TargetableTexture, FTextureRHIRef(), true);
 		RHICmdList.SetViewport(0, 0, 0.0f, FilterBufferSize.X, FilterBufferSize.Y, 1.0f);
 
 		RHICmdList.SetBlendState(TStaticBlendState<>::GetRHI());
@@ -658,7 +659,7 @@ void FinishOcclusionTerm(FRHICommandList& RHICmdList, const FViewInfo& View, con
 	const uint32 DownsampledSizeX = View.ViewRect.Width() / DownsampleFactor;
 	const uint32 DownsampledSizeY = View.ViewRect.Height() / DownsampleFactor;
 
-	SetRenderTarget(RHICmdList, LightShaftsDest->GetRenderTargetItem().TargetableTexture, FTextureRHIRef());
+	SetRenderTarget(RHICmdList, LightShaftsDest->GetRenderTargetItem().TargetableTexture, FTextureRHIRef(), true);
 	RHICmdList.SetViewport(0, 0, 0.0f, FilterBufferSize.X, FilterBufferSize.Y, 1.0f);
 
 	RHICmdList.SetBlendState(TStaticBlendState<>::GetRHI());
