UAV resource transition support.  Validation for UAV transitions.
Some RenderTarget validation fixes.

diff --git a/Engine/Source/Runtime/Engine/Private/GPUSort.cpp b/Engine/Source/Runtime/Engine/Private/GPUSort.cpp
--- a/Engine/Source/Runtime/Engine/Private/GPUSort.cpp
+++ b/Engine/Source/Runtime/Engine/Private/GPUSort.cpp
@@ -735,6 +735,7 @@ int32 SortGPUBuffers(FRHICommandListImmediate& RHICmdList, FGPUSortBuffers SortB
 	check( UpsweepCS->RequiresConstantBufferWorkaround() == DownsweepCS->RequiresConstantBufferWorkaround() );
 	const bool bUseConstantBufferWorkaround = UpsweepCS->RequiresConstantBufferWorkaround();
 
+
 	
 	// Execute each pass as needed.
 	uint32 PassBits = DIGIT_COUNT - 1;
@@ -755,12 +756,18 @@ int32 SortGPUBuffers(FRHICommandListImmediate& RHICmdList, FGPUSortBuffers SortB
 				SortUniformBufferRef = FRadixSortUniformBufferRef::CreateUniformBufferImmediate( SortParameters, UniformBuffer_SingleDraw );
 			}
 
-			// Clear the offsets buffer.
-			RHICmdList.SetComputeShader(ClearOffsetsCS->GetComputeShader());
+			//make UAV safe for clear
+			RHICmdList.TransitionResource(EResourceTransitionAccess::ERWBarrier, EResourceTransitionPipeline::EComputeToCompute, GSortOffsetBuffers.BufferUAVs[0]);
+			
+			// Clear the offsets buffer.			
+			RHICmdList.SetComputeShader(ClearOffsetsCS->GetComputeShader());			
 			ClearOffsetsCS->SetOutput(RHICmdList, GSortOffsetBuffers.BufferUAVs[0]);
 			DispatchComputeShader(RHICmdList, *ClearOffsetsCS, 1, 1 ,1 );
 			ClearOffsetsCS->UnbindBuffers(RHICmdList);
 
+			//make UAV safe for readback
+			RHICmdList.TransitionResource(EResourceTransitionAccess::ERWBarrier, EResourceTransitionPipeline::EComputeToCompute, GSortOffsetBuffers.BufferUAVs[0]);
+
 			// Phase 1: Scan upsweep to compute per-digit totals.
 			RHICmdList.SetComputeShader(UpsweepCS->GetComputeShader());
 			UpsweepCS->SetOutput(RHICmdList, GSortOffsetBuffers.BufferUAVs[0]);
@@ -768,6 +775,13 @@ int32 SortGPUBuffers(FRHICommandListImmediate& RHICmdList, FGPUSortBuffers SortB
 			DispatchComputeShader(RHICmdList, *UpsweepCS, GroupCount, 1, 1 );
 			UpsweepCS->UnbindBuffers(RHICmdList);
 
+			//barrier both UAVS since for next step.
+			FUnorderedAccessViewRHIParamRef PrePhase2BarrierUAVS[2];
+			PrePhase2BarrierUAVS[0] = GSortOffsetBuffers.BufferUAVs[0];
+			PrePhase2BarrierUAVS[1] = GSortOffsetBuffers.BufferUAVs[1];
+
+			RHICmdList.TransitionResources(EResourceTransitionAccess::ERWBarrier, EResourceTransitionPipeline::EComputeToCompute, PrePhase2BarrierUAVS, 2);
+
 			if (bDebugOffsets)
 			{
 				UE_LOG(LogGPUSort, Log, TEXT("\n========== UPSWEEP =========="));
@@ -787,6 +801,14 @@ int32 SortGPUBuffers(FRHICommandListImmediate& RHICmdList, FGPUSortBuffers SortB
 				GSortOffsetBuffers.DumpOffsets(1);
 			}
 
+			//UAV is going to SRV, so transition to Readable.
+			RHICmdList.TransitionResource(EResourceTransitionAccess::EReadable, EResourceTransitionPipeline::EComputeToCompute, GSortOffsetBuffers.BufferUAVs[1]);
+
+			FUnorderedAccessViewRHIParamRef PrePhase3BarrierUAVS[2];
+			PrePhase3BarrierUAVS[0] = SortBuffers.RemoteKeyUAVs[BufferIndex ^ 0x1];
+			PrePhase3BarrierUAVS[1] = SortBuffers.RemoteValueUAVs[BufferIndex ^ 0x1];
+			RHICmdList.TransitionResources(EResourceTransitionAccess::ERWBarrier, EResourceTransitionPipeline::EComputeToCompute, PrePhase3BarrierUAVS, 2);
+
 			// Phase 3: Downsweep to compute final offsets and scatter keys.
 			RHICmdList.SetComputeShader(DownsweepCS->GetComputeShader());
 			DownsweepCS->SetOutput(RHICmdList, SortBuffers.RemoteKeyUAVs[BufferIndex ^ 0x1], SortBuffers.RemoteValueUAVs[BufferIndex ^ 0x1]);
@@ -794,6 +816,8 @@ int32 SortGPUBuffers(FRHICommandListImmediate& RHICmdList, FGPUSortBuffers SortB
 			DispatchComputeShader(RHICmdList, *DownsweepCS, GroupCount, 1, 1 );
 			DownsweepCS->UnbindBuffers(RHICmdList);
 
+			RHICmdList.TransitionResources(EResourceTransitionAccess::ERWBarrier, EResourceTransitionPipeline::EComputeToCompute, PrePhase3BarrierUAVS, 2);
+
 			// Flip buffers.
 			BufferIndex ^= 0x1;
 
diff --git a/Engine/Source/Runtime/Engine/Private/Particles/ParticleSortingGPU.cpp b/Engine/Source/Runtime/Engine/Private/Particles/ParticleSortingGPU.cpp
--- a/Engine/Source/Runtime/Engine/Private/Particles/ParticleSortingGPU.cpp
+++ b/Engine/Source/Runtime/Engine/Private/Particles/ParticleSortingGPU.cpp
@@ -192,6 +192,13 @@ static int32 GenerateParticleSortKeys(
 	const uint32 MaxGroupCount = 128;
 	int32 TotalParticleCount = 0;
 
+	FUnorderedAccessViewRHIParamRef OutputUAVs[2];
+	OutputUAVs[0] = KeyBufferUAV;
+	OutputUAVs[1] = SortedVertexBufferUAV;
+
+	//make sure our outputs are safe to write to.
+	RHICmdList.TransitionResources(EResourceTransitionAccess::ERWBarrier, EResourceTransitionPipeline::EGfxToCompute, OutputUAVs, 2);
+
 	// Grab the shader, set output.
 	TShaderMapRef<FParticleSortKeyGenCS> KeyGenCS(GetGlobalShaderMap(FeatureLevel));
 	RHICmdList.SetComputeShader(KeyGenCS->GetComputeShader());
@@ -221,6 +228,9 @@ static int32 GenerateParticleSortKeys(
 		KeyGenCS->SetParameters(RHICmdList, KeyGenUniformBuffer, SortInfo.VertexBufferSRV);
 		DispatchComputeShader(RHICmdList, *KeyGenCS, GroupCount, 1, 1);
 
+		//we may be able to remove this transition if each step isn't dependent on the previous one.
+		RHICmdList.TransitionResources(EResourceTransitionAccess::ERWBarrier, EResourceTransitionPipeline::EComputeToCompute, OutputUAVs, 2);
+
 		// Update offset in to the buffer.
 		TotalParticleCount += ParticleCount;
 	}
@@ -228,6 +238,9 @@ static int32 GenerateParticleSortKeys(
 	// Clear the output buffer.
 	KeyGenCS->UnbindBuffers(RHICmdList);
 
+	//make sure our outputs are readable as SRVs to further gfx steps.
+	RHICmdList.TransitionResources(EResourceTransitionAccess::EReadable, EResourceTransitionPipeline::EComputeToGfx, OutputUAVs, 2);
+
 	return TotalParticleCount;
 }
 
diff --git a/Engine/Source/Runtime/Engine/Private/VectorField.cpp b/Engine/Source/Runtime/Engine/Private/VectorField.cpp
--- a/Engine/Source/Runtime/Engine/Private/VectorField.cpp
+++ b/Engine/Source/Runtime/Engine/Private/VectorField.cpp
@@ -877,6 +877,7 @@ public:
 				NoiseVolumeTextureRHI = AnimatedVectorField->NoiseField->Resource->VolumeTextureRHI;
 			}
 
+			RHICmdList.TransitionResource(EResourceTransitionAccess::ERWBarrier, EResourceTransitionPipeline::EGfxToCompute, VolumeTextureUAV);
 			RHICmdList.SetComputeShader(CompositeCS->GetComputeShader());
 			CompositeCS->SetOutput(RHICmdList, VolumeTextureUAV);
 			/// ?
@@ -892,6 +893,7 @@ public:
 				SizeY / THREADS_PER_AXIS,
 				SizeZ / THREADS_PER_AXIS );
 			CompositeCS->UnbindBuffers(RHICmdList);
+			RHICmdList.TransitionResource(EResourceTransitionAccess::EReadable, EResourceTransitionPipeline::EComputeToGfx, VolumeTextureUAV);
 		}
 	}
 
diff --git a/Engine/Source/Runtime/RHI/Public/RHI.h b/Engine/Source/Runtime/RHI/Public/RHI.h
--- a/Engine/Source/Runtime/RHI/Public/RHI.h
+++ b/Engine/Source/Runtime/RHI/Public/RHI.h
@@ -976,6 +976,7 @@ enum class EResourceTransitionAccess
 	EWritable, //transition from read -> write	
 	ERWBarrier, // Mostly for UAVs.  Transition to read/write state and always insert a resource barrier.
 	ERWNoBarrier, //Mostly UAVs.  Indicates we want R/W access and do not require synchronization for the duration of the RW state.  The initial transition from writable->RWNoBarrier and readable->RWNoBarrier still requires a sync
+	ERWSubResBarrier, //For special cases where read/write happens to different subresources of the same resource in the same call.  Inserts a barrier, but read validation will pass.  Temporary until we pass full subresource info to all transition calls.
 	EMaxAccess,
 };
 
diff --git a/Engine/Source/Runtime/RHI/Public/RHICommandList.h b/Engine/Source/Runtime/RHI/Public/RHICommandList.h
--- a/Engine/Source/Runtime/RHI/Public/RHICommandList.h
+++ b/Engine/Source/Runtime/RHI/Public/RHICommandList.h
@@ -1879,6 +1879,29 @@ public:
 		new (AllocCommand<FRHICommandTransitionTexturesArray>()) FRHICommandTransitionTexturesArray(TransitionType, InTextures);
 	}
 
+	FORCEINLINE_DEBUGGABLE void TransitionResource(EResourceTransitionAccess TransitionType, EResourceTransitionPipeline TransitionPipeline, FUnorderedAccessViewRHIParamRef InUAV)
+	{
+		FUnorderedAccessViewRHIParamRef UAV = InUAV;
+		if (Bypass())
+		{
+			CMD_CONTEXT(TransitionResources)(TransitionType, TransitionPipeline, &UAV, 1);
+			return;
+		}
+		new (AllocCommand<FRHICommandTransitionUAVs>()) FRHICommandTransitionUAVs(TransitionType, TransitionPipeline, &UAV, 1);
+	}
+
+	FORCEINLINE_DEBUGGABLE void TransitionResources(EResourceTransitionAccess TransitionType, EResourceTransitionPipeline TransitionPipeline, FUnorderedAccessViewRHIParamRef* InUAVs, int32 NumUAVs)
+	{
+		if (Bypass())
+		{
+			CMD_CONTEXT(TransitionResources)(TransitionType, TransitionPipeline, InUAVs, NumUAVs);
+			return;
+		}
+		new (AllocCommand<FRHICommandTransitionUAVs>()) FRHICommandTransitionUAVs(TransitionType, TransitionPipeline, InUAVs, NumUAVs);
+	}
+
+	
+
 // These 6 are special in that they must be called on the immediate command list and they force a flush only when we are not doing RHI thread
 	void BeginScene();
 	void EndScene();
diff --git a/Engine/Source/Runtime/Renderer/Private/DistanceFieldGlobalIllumination.cpp b/Engine/Source/Runtime/Renderer/Private/DistanceFieldGlobalIllumination.cpp
--- a/Engine/Source/Runtime/Renderer/Private/DistanceFieldGlobalIllumination.cpp
+++ b/Engine/Source/Runtime/Renderer/Private/DistanceFieldGlobalIllumination.cpp
@@ -158,6 +158,11 @@ public:
 		SetShaderValue(RHICmdList, ShaderRHI, LightDirectionAndTraceDistance, FVector4(LightSceneProxy->GetDirection(), GVPLDirectionalLightTraceDistance));
 		SetShaderValue(RHICmdList, ShaderRHI, LightColor, LightSceneProxy->GetColor() * LightSceneProxy->GetIndirectLightingScale());
 
+		FUnorderedAccessViewRHIParamRef OutUAVs[2];
+		OutUAVs[0] = GVPLResources.VPLParameterBuffer.UAV;
+		OutUAVs[1] = GVPLResources.VPLData.UAV;
+		RHICmdList.TransitionResources(EResourceTransitionAccess::ERWBarrier, EResourceTransitionPipeline::EComputeToCompute, OutUAVs, ARRAY_COUNT(OutUAVs));
+
 		VPLParameterBuffer.SetBuffer(RHICmdList, ShaderRHI, GVPLResources.VPLParameterBuffer);
 		VPLData.SetBuffer(RHICmdList, ShaderRHI, GVPLResources.VPLData);
 
@@ -177,6 +182,11 @@ public:
 	{
 		VPLParameterBuffer.UnsetUAV(RHICmdList, GetComputeShader());
 		VPLData.UnsetUAV(RHICmdList, GetComputeShader());
+
+		FUnorderedAccessViewRHIParamRef OutUAVs[2];
+		OutUAVs[0] = GVPLResources.VPLParameterBuffer.UAV;
+		OutUAVs[1] = GVPLResources.VPLData.UAV;
+		RHICmdList.TransitionResources(EResourceTransitionAccess::EReadable, EResourceTransitionPipeline::EComputeToCompute, OutUAVs, ARRAY_COUNT(OutUAVs));
 	}
 
 	// FShader interface.
@@ -251,6 +261,7 @@ public:
 		FComputeShaderRHIParamRef ShaderRHI = GetComputeShader();
 		FGlobalShader::SetParameters(RHICmdList, ShaderRHI, View);
 
+		RHICmdList.TransitionResource(EResourceTransitionAccess::ERWBarrier, EResourceTransitionPipeline::EComputeToCompute, GVPLResources.VPLDispatchIndirectBuffer.UAV);
 		DispatchParameters.SetBuffer(RHICmdList, ShaderRHI, GVPLResources.VPLDispatchIndirectBuffer);
 		SetSRVParameter(RHICmdList, ShaderRHI, VPLParameterBuffer, GVPLResources.VPLParameterBuffer.SRV);
 	}
@@ -258,6 +269,7 @@ public:
 	void UnsetParameters(FRHICommandList& RHICmdList)
 	{
 		DispatchParameters.UnsetUAV(RHICmdList, GetComputeShader());
+		RHICmdList.TransitionResource(EResourceTransitionAccess::EReadable, EResourceTransitionPipeline::EComputeToCompute, GVPLResources.VPLDispatchIndirectBuffer.UAV);
 	}
 
 	virtual bool Serialize(FArchive& Ar) override
@@ -314,6 +326,11 @@ public:
 		FComputeShaderRHIParamRef ShaderRHI = GetComputeShader();
 		FGlobalShader::SetParameters(RHICmdList, ShaderRHI, View);
 
+		FUnorderedAccessViewRHIParamRef OutUAVs[2];
+		OutUAVs[0] = GCulledVPLResources.VPLParameterBuffer.UAV;
+		OutUAVs[1] = GCulledVPLResources.VPLData.UAV;
+		RHICmdList.TransitionResources(EResourceTransitionAccess::ERWBarrier, EResourceTransitionPipeline::EComputeToCompute, OutUAVs, ARRAY_COUNT(OutUAVs));
+
 		CulledVPLParameterBuffer.SetBuffer(RHICmdList, ShaderRHI, GCulledVPLResources.VPLParameterBuffer);
 		CulledVPLData.SetBuffer(RHICmdList, ShaderRHI, GCulledVPLResources.VPLData);
 		AOParameters.Set(RHICmdList, ShaderRHI, Parameters);
@@ -578,6 +595,7 @@ public:
 		FComputeShaderRHIParamRef ShaderRHI = GetComputeShader();
 		FGlobalShader::SetParameters(RHICmdList, ShaderRHI, View);
 
+		RHICmdList.TransitionResource(EResourceTransitionAccess::ERWBarrier, EResourceTransitionPipeline::EComputeToCompute, GAOCulledObjectBuffers.Buffers.ObjectIndirectDispatch.UAV);
 		DispatchParameters.SetBuffer(RHICmdList, ShaderRHI, GAOCulledObjectBuffers.Buffers.ObjectIndirectDispatch);
 		ObjectParameters.Set(RHICmdList, ShaderRHI, GAOCulledObjectBuffers.Buffers);
 
@@ -586,7 +604,11 @@ public:
 
 	void UnsetParameters(FRHICommandList& RHICmdList)
 	{
-		DispatchParameters.UnsetUAV(RHICmdList, GetComputeShader());
+		FComputeShaderRHIParamRef ShaderRHI = GetComputeShader();
+		DispatchParameters.UnsetUAV(RHICmdList, ShaderRHI);
+
+		ObjectParameters.UnsetParameters(RHICmdList, ShaderRHI, GAOCulledObjectBuffers.Buffers);
+		RHICmdList.TransitionResource(EResourceTransitionAccess::EReadable, EResourceTransitionPipeline::EComputeToCompute, GAOCulledObjectBuffers.Buffers.ObjectIndirectDispatch.UAV);
 	}
 
 	virtual bool Serialize(FArchive& Ar) override
@@ -920,14 +942,24 @@ public:
 		SetSRVParameter(RHICmdList, ShaderRHI, ScatterDrawParameters, SurfaceCacheResources.Level[DepthLevel]->ScatterDrawParameters.SRV);
 		SetSRVParameter(RHICmdList, ShaderRHI, SavedStartIndex, SurfaceCacheResources.Level[DepthLevel]->SavedStartIndex.SRV);
 
+		FUnorderedAccessViewRHIParamRef OutUAVs[2];
+		OutUAVs[0] = TemporaryIrradianceCacheResources->SurfelIrradiance.UAV;
+		OutUAVs[1] = TemporaryIrradianceCacheResources->HeightfieldIrradiance.UAV;
+		RHICmdList.TransitionResources(EResourceTransitionAccess::ERWBarrier, EResourceTransitionPipeline::EComputeToCompute, OutUAVs, ARRAY_COUNT(OutUAVs));
+
 		SurfelIrradiance.SetBuffer(RHICmdList, ShaderRHI, TemporaryIrradianceCacheResources->SurfelIrradiance);
 		HeightfieldIrradiance.SetBuffer(RHICmdList, ShaderRHI, TemporaryIrradianceCacheResources->HeightfieldIrradiance);
 	}
 
-	void UnsetParameters(FRHICommandList& RHICmdList)
+	void UnsetParameters(FRHICommandList& RHICmdList, FTemporaryIrradianceCacheResources* TemporaryIrradianceCacheResources)
 	{
 		SurfelIrradiance.UnsetUAV(RHICmdList, GetComputeShader());
 		HeightfieldIrradiance.UnsetUAV(RHICmdList, GetComputeShader());
+
+		FUnorderedAccessViewRHIParamRef OutUAVs[2];
+		OutUAVs[0] = TemporaryIrradianceCacheResources->SurfelIrradiance.UAV;
+		OutUAVs[1] = TemporaryIrradianceCacheResources->HeightfieldIrradiance.UAV;
+		RHICmdList.TransitionResources(EResourceTransitionAccess::EReadable, EResourceTransitionPipeline::EComputeToCompute, OutUAVs, ARRAY_COUNT(OutUAVs));
 	}
 
 	virtual bool Serialize(FArchive& Ar) override
@@ -993,6 +1025,7 @@ public:
 		SetSRVParameter(RHICmdList, ShaderRHI, SavedStartIndex, SurfaceCacheResources.Level[DepthLevel]->SavedStartIndex.SRV);
 		SetSRVParameter(RHICmdList, ShaderRHI, RecordConeData, TemporaryIrradianceCacheResources->ConeData.SRV);
 
+		RHICmdList.TransitionResource(EResourceTransitionAccess::ERWBarrier, EResourceTransitionPipeline::EComputeToCompute, TemporaryIrradianceCacheResources->StepBentNormal.UAV);
 		StepBentNormal.SetBuffer(RHICmdList, ShaderRHI, TemporaryIrradianceCacheResources->StepBentNormal);
 
 		FAOSampleData2 AOSampleData;
@@ -1018,9 +1051,10 @@ public:
 		SetShaderValue(RHICmdList, ShaderRHI, BentNormalNormalizeFactor, BentNormalNormalizeFactorValue);
 	}
 
-	void UnsetParameters(FRHICommandList& RHICmdList)
+	void UnsetParameters(FRHICommandList& RHICmdList, FTemporaryIrradianceCacheResources* TemporaryIrradianceCacheResources)
 	{
 		StepBentNormal.UnsetUAV(RHICmdList, GetComputeShader());
+		RHICmdList.TransitionResource(EResourceTransitionAccess::EReadable, EResourceTransitionPipeline::EComputeToCompute, TemporaryIrradianceCacheResources->StepBentNormal.UAV);
 	}
 
 	virtual bool Serialize(FArchive& Ar) override
@@ -1124,6 +1158,7 @@ public:
 		SetSRVParameter(RHICmdList, ShaderRHI, ScatterDrawParameters, SurfaceCacheResources.Level[DepthLevel]->ScatterDrawParameters.SRV);
 		SetSRVParameter(RHICmdList, ShaderRHI, SavedStartIndex, SurfaceCacheResources.Level[DepthLevel]->SavedStartIndex.SRV);
 
+		RHICmdList.TransitionResource(EResourceTransitionAccess::ERWBarrier, EResourceTransitionPipeline::EComputeToCompute, TemporaryIrradianceCacheResources->SurfelIrradiance.UAV);
 		SurfelIrradiance.SetBuffer(RHICmdList, ShaderRHI, TemporaryIrradianceCacheResources->SurfelIrradiance);
 
 		GlobalObjectParameters.Set(RHICmdList, ShaderRHI, *Scene->DistanceFieldSceneData.ObjectBuffers, Scene->DistanceFieldSceneData.NumObjectsInBuffer);
@@ -1147,9 +1182,10 @@ public:
 		SetSRVParameter(RHICmdList, ShaderRHI, StepBentNormalBuffer, TemporaryIrradianceCacheResources->StepBentNormal.SRV);
 	}
 
-	void UnsetParameters(FRHICommandList& RHICmdList)
+	void UnsetParameters(FRHICommandList& RHICmdList, FTemporaryIrradianceCacheResources* TemporaryIrradianceCacheResources)
 	{
 		SurfelIrradiance.UnsetUAV(RHICmdList, GetComputeShader());
+		RHICmdList.TransitionResource(EResourceTransitionAccess::EReadable, EResourceTransitionPipeline::EComputeToCompute, TemporaryIrradianceCacheResources->SurfelIrradiance.UAV);
 	}
 
 	virtual bool Serialize(FArchive& Ar) override
@@ -1253,12 +1289,17 @@ public:
 		SetSRVParameter(RHICmdList, ShaderRHI, ScatterDrawParameters, SurfaceCacheResources.Level[DepthLevel]->ScatterDrawParameters.SRV);
 		SetSRVParameter(RHICmdList, ShaderRHI, SavedStartIndex, SurfaceCacheResources.Level[DepthLevel]->SavedStartIndex.SRV);
 
+		RHICmdList.TransitionResource(EResourceTransitionAccess::ERWBarrier, EResourceTransitionPipeline::EComputeToCompute, SurfaceCacheResources.Level[DepthLevel]->Irradiance.UAV);
 		IrradianceCacheIrradiance.SetBuffer(RHICmdList, ShaderRHI, SurfaceCacheResources.Level[DepthLevel]->Irradiance);
 	}
 
-	void UnsetParameters(FRHICommandList& RHICmdList)
+	void UnsetParameters(FRHICommandList& RHICmdList, const FSceneView& View, int32 DepthLevel)
 	{
 		IrradianceCacheIrradiance.UnsetUAV(RHICmdList, GetComputeShader());
+
+		const FScene* Scene = (const FScene*)View.Family->Scene;
+		FSurfaceCacheResources& SurfaceCacheResources = *Scene->SurfaceCacheResources;
+		RHICmdList.TransitionResource(EResourceTransitionAccess::EReadable, EResourceTransitionPipeline::EComputeToCompute, SurfaceCacheResources.Level[DepthLevel]->Irradiance.UAV);
 	}
 
 	virtual bool Serialize(FArchive& Ar) override
@@ -1302,7 +1343,7 @@ void ComputeIrradianceForSamples(
 			RHICmdList.SetComputeShader(ComputeShader->GetComputeShader());
 			ComputeShader->SetParameters(RHICmdList, View, DepthLevel);
 			DispatchComputeShader(RHICmdList, *ComputeShader, 1, 1, 1);
-			ComputeShader->UnsetParameters(RHICmdList);
+			ComputeShader->UnsetParameters(RHICmdList, View);
 		}
 
 		{
@@ -1310,7 +1351,7 @@ void ComputeIrradianceForSamples(
 			RHICmdList.SetComputeShader(ComputeShader->GetComputeShader());
 			ComputeShader->SetParameters(RHICmdList, View, DepthLevel, TemporaryIrradianceCacheResources);
 			DispatchIndirectComputeShader(RHICmdList, *ComputeShader, SurfaceCacheResources.DispatchParameters.Buffer, 0);
-			ComputeShader->UnsetParameters(RHICmdList);
+			ComputeShader->UnsetParameters(RHICmdList, TemporaryIrradianceCacheResources);
 		}
 		
 		{
@@ -1318,7 +1359,7 @@ void ComputeIrradianceForSamples(
 			RHICmdList.SetComputeShader(ComputeShader->GetComputeShader());
 			ComputeShader->SetParameters(RHICmdList, View, DepthLevel, TemporaryIrradianceCacheResources);
 			DispatchIndirectComputeShader(RHICmdList, *ComputeShader, SurfaceCacheResources.DispatchParameters.Buffer, 0);
-			ComputeShader->UnsetParameters(RHICmdList);
+			ComputeShader->UnsetParameters(RHICmdList, TemporaryIrradianceCacheResources);
 		}
 	}
 
@@ -1333,7 +1374,7 @@ void ComputeIrradianceForSamples(
 			RHICmdList.SetComputeShader(ComputeShader->GetComputeShader());
 			ComputeShader->SetParameters(RHICmdList, View, DepthLevel);
 			DispatchComputeShader(RHICmdList, *ComputeShader, 1, 1, 1);
-			ComputeShader->UnsetParameters(RHICmdList);
+			ComputeShader->UnsetParameters(RHICmdList, View);
 		}
 	
 		if (GVPLSurfelRepresentation)
@@ -1344,7 +1385,7 @@ void ComputeIrradianceForSamples(
 			ComputeShader->SetParameters(RHICmdList, View, DepthLevel, Parameters, TemporaryIrradianceCacheResources);
 			DispatchIndirectComputeShader(RHICmdList, *ComputeShader, SurfaceCacheResources.DispatchParameters.Buffer, 0);
 
-			ComputeShader->UnsetParameters(RHICmdList);
+			ComputeShader->UnsetParameters(RHICmdList, TemporaryIrradianceCacheResources);
 		}
 		else
 		{
@@ -1354,7 +1395,7 @@ void ComputeIrradianceForSamples(
 			ComputeShader->SetParameters(RHICmdList, View, DepthLevel, Parameters, TemporaryIrradianceCacheResources);
 			DispatchIndirectComputeShader(RHICmdList, *ComputeShader, SurfaceCacheResources.DispatchParameters.Buffer, 0);
 
-			ComputeShader->UnsetParameters(RHICmdList);
+			ComputeShader->UnsetParameters(RHICmdList, TemporaryIrradianceCacheResources);
 		}
 	}
 
@@ -1363,7 +1404,7 @@ void ComputeIrradianceForSamples(
 		RHICmdList.SetComputeShader(ComputeShader->GetComputeShader());
 		ComputeShader->SetParameters(RHICmdList, View, DepthLevel);
 		DispatchComputeShader(RHICmdList, *ComputeShader, 1, 1, 1);
-		ComputeShader->UnsetParameters(RHICmdList);
+		ComputeShader->UnsetParameters(RHICmdList, View);
 	}
 		
 	{
@@ -1371,7 +1412,7 @@ void ComputeIrradianceForSamples(
 		RHICmdList.SetComputeShader(ComputeShader->GetComputeShader());
 		ComputeShader->SetParameters(RHICmdList, View, DepthLevel, TemporaryIrradianceCacheResources);
 		DispatchIndirectComputeShader(RHICmdList, *ComputeShader, SurfaceCacheResources.DispatchParameters.Buffer, 0);
-		ComputeShader->UnsetParameters(RHICmdList);
+		ComputeShader->UnsetParameters(RHICmdList, View, DepthLevel);
 	}
 }
 
@@ -1446,12 +1487,14 @@ public:
 
 		SetSRVParameter(RHICmdList, ShaderRHI, ConeDepthVisibilityFunction, ScreenGridResources.ConeDepthVisibilityFunction.SRV);
 
+		RHICmdList.TransitionResource(EResourceTransitionAccess::ERWBarrier, EResourceTransitionPipeline::EComputeToCompute, ScreenGridResources.StepBentNormal.UAV);
 		StepBentNormal.SetBuffer(RHICmdList, ShaderRHI, ScreenGridResources.StepBentNormal);
 	}
 
-	void UnsetParameters(FRHICommandList& RHICmdList)
+	void UnsetParameters(FRHICommandList& RHICmdList, const FAOScreenGridResources& ScreenGridResources)
 	{
 		StepBentNormal.UnsetUAV(RHICmdList, GetComputeShader());
+		RHICmdList.TransitionResource(EResourceTransitionAccess::EReadable, EResourceTransitionPipeline::EComputeToCompute, ScreenGridResources.StepBentNormal.UAV);
 	}
 
 	virtual bool Serialize(FArchive& Ar) override
@@ -1548,12 +1591,17 @@ public:
 		FAOScreenGridResources* ScreenGridResources = View.ViewState->AOScreenGridResources;
 
 		SetSRVParameter(RHICmdList, ShaderRHI, StepBentNormalBuffer, ScreenGridResources->StepBentNormal.SRV);
+
+		RHICmdList.TransitionResource(EResourceTransitionAccess::ERWBarrier, EResourceTransitionPipeline::EComputeToCompute, ScreenGridResources->SurfelIrradiance.UAV);
 		SurfelIrradiance.SetBuffer(RHICmdList, ShaderRHI, ScreenGridResources->SurfelIrradiance);
 	}
 
-	void UnsetParameters(FRHICommandList& RHICmdList)
+	void UnsetParameters(FRHICommandList& RHICmdList, const FViewInfo& View)
 	{
 		SurfelIrradiance.UnsetUAV(RHICmdList, GetComputeShader());
+
+		FAOScreenGridResources* ScreenGridResources = View.ViewState->AOScreenGridResources;
+		RHICmdList.TransitionResource(EResourceTransitionAccess::EReadable, EResourceTransitionPipeline::EComputeToCompute, ScreenGridResources->SurfelIrradiance.UAV);
 	}
 
 	virtual bool Serialize(FArchive& Ar) override
@@ -1637,14 +1685,16 @@ public:
 		SetSRVParameter(RHICmdList, ShaderRHI, SurfelIrradiance, ScreenGridResources.SurfelIrradiance.SRV);
 		SetSRVParameter(RHICmdList, ShaderRHI, HeightfieldIrradiance, ScreenGridResources.HeightfieldIrradiance.SRV);
 
+		RHICmdList.TransitionResource(EResourceTransitionAccess::ERWBarrier, EResourceTransitionPipeline::EComputeToCompute, IrradianceTextureValue.UAV);
 		IrradianceTexture.SetTexture(RHICmdList, ShaderRHI, IrradianceTextureValue.ShaderResourceTexture, IrradianceTextureValue.UAV);
 
 		SetShaderValue(RHICmdList, ShaderRHI, ScreenGridConeVisibilitySize, ScreenGridResources.ScreenGridDimensions);
 	}
 
-	void UnsetParameters(FRHICommandList& RHICmdList)
+	void UnsetParameters(FRHICommandList& RHICmdList, FSceneRenderTargetItem& IrradianceTextureValue)
 	{
 		IrradianceTexture.UnsetUAV(RHICmdList, GetComputeShader());
+		RHICmdList.TransitionResource(EResourceTransitionAccess::EReadable, EResourceTransitionPipeline::EComputeToCompute, IrradianceTextureValue.UAV);
 	}
 
 	virtual bool Serialize(FArchive& Ar) override
@@ -1696,7 +1746,7 @@ void ComputeIrradianceForScreenGrid(
 			RHICmdList.SetComputeShader(ComputeShader->GetComputeShader());
 			ComputeShader->SetParameters(RHICmdList, View, DistanceFieldNormal, ScreenGridResources);
 			DispatchComputeShader(RHICmdList, *ComputeShader, GroupSizeX, GroupSizeY, 1);
-			ComputeShader->UnsetParameters(RHICmdList);
+			ComputeShader->UnsetParameters(RHICmdList, ScreenGridResources);
 		}
 
 		if (GVPLSurfelRepresentation)
@@ -1710,7 +1760,7 @@ void ComputeIrradianceForScreenGrid(
 			const uint32 ComputeIrradianceGroupSizeX = View.ViewRect.Size().X / GAODownsampleFactor;
 			const uint32 ComputeIrradianceGroupSizeY = View.ViewRect.Size().Y / GAODownsampleFactor;
 			DispatchComputeShader(RHICmdList, *ComputeShader, ComputeIrradianceGroupSizeX, ComputeIrradianceGroupSizeY, 1);
-			ComputeShader->UnsetParameters(RHICmdList);
+			ComputeShader->UnsetParameters(RHICmdList, View);
 		}
 	}
 		
@@ -1719,7 +1769,7 @@ void ComputeIrradianceForScreenGrid(
 		RHICmdList.SetComputeShader(ComputeShader->GetComputeShader());
 		ComputeShader->SetParameters(RHICmdList, View, ScreenGridResources, IrradianceTexture);
 		DispatchComputeShader(RHICmdList, *ComputeShader, GroupSizeX, GroupSizeY, 1);
-		ComputeShader->UnsetParameters(RHICmdList);
+		ComputeShader->UnsetParameters(RHICmdList, IrradianceTexture);
 	}
 }
 
diff --git a/Engine/Source/Runtime/Renderer/Private/DistanceFieldLightingPost.cpp b/Engine/Source/Runtime/Renderer/Private/DistanceFieldLightingPost.cpp
--- a/Engine/Source/Runtime/Renderer/Private/DistanceFieldLightingPost.cpp
+++ b/Engine/Source/Runtime/Renderer/Private/DistanceFieldLightingPost.cpp
@@ -549,6 +549,7 @@ void AllocateOrReuseAORenderTarget(FRHICommandList& RHICmdList, TRefCountPtr<IPo
 		FIntPoint BufferSize = GetBufferSizeForAO();
 
 		FPooledRenderTargetDesc Desc(FPooledRenderTargetDesc::Create2DDesc(BufferSize, bWithAlpha ? PF_FloatRGBA : PF_FloatRGB, FClearValueBinding::None, TexCreate_None, TexCreate_RenderTargetable | TexCreate_UAV, false));
+		Desc.AutoWritable = false;
 		GRenderTargetPool.FindFreeElement(RHICmdList, Desc, Target, Name);
 	}
 }
@@ -673,7 +674,7 @@ void UpdateHistory(
 						bUseDistanceFieldGI ? (*IrradianceHistoryState)->GetRenderTargetItem().TargetableTexture : NULL,
 					};
 
-					SetRenderTargets(RHICmdList, ARRAY_COUNT(RenderTargets) - (bUseDistanceFieldGI ? 0 : 1), RenderTargets, FTextureRHIParamRef(), 0, NULL);
+					SetRenderTargets(RHICmdList, ARRAY_COUNT(RenderTargets) - (bUseDistanceFieldGI ? 0 : 1), RenderTargets, FTextureRHIParamRef(), 0, NULL, true);
 
 					RHICmdList.SetViewport(0, 0, 0.0f, View.ViewRect.Width() / GAODownsampleFactor, View.ViewRect.Height() / GAODownsampleFactor, 1.0f);
 					RHICmdList.SetRasterizerState(TStaticRasterizerState<FM_Solid, CM_None>::GetRHI());
diff --git a/Engine/Source/Runtime/Renderer/Private/DistanceFieldLightingShared.h b/Engine/Source/Runtime/Renderer/Private/DistanceFieldLightingShared.h
--- a/Engine/Source/Runtime/Renderer/Private/DistanceFieldLightingShared.h
+++ b/Engine/Source/Runtime/Renderer/Private/DistanceFieldLightingShared.h
@@ -147,8 +147,16 @@ public:
 	}
 
 	template<typename TParamRef>
-	void Set(FRHICommandList& RHICmdList, const TParamRef& ShaderRHI, const FDistanceFieldObjectBuffers& ObjectBuffers, int32 NumObjectsValue)
+	void Set(FRHICommandList& RHICmdList, const TParamRef& ShaderRHI, const FDistanceFieldObjectBuffers& ObjectBuffers, int32 NumObjectsValue, bool bBarrier = false)
 	{
+		if (bBarrier)
+		{
+			FUnorderedAccessViewRHIParamRef OutUAVs[2];
+			OutUAVs[0] = ObjectBuffers.Bounds.UAV;
+			OutUAVs[1] = ObjectBuffers.Data.UAV;
+			RHICmdList.TransitionResources(EResourceTransitionAccess::ERWBarrier, EResourceTransitionPipeline::EComputeToCompute, OutUAVs, ARRAY_COUNT(OutUAVs));
+		}
+
 		ObjectBounds.SetBuffer(RHICmdList, ShaderRHI, ObjectBuffers.Bounds);
 		ObjectData.SetBuffer(RHICmdList, ShaderRHI, ObjectBuffers.Data);
 		SetShaderValue(RHICmdList, ShaderRHI, NumSceneObjects, NumObjectsValue);
@@ -170,10 +178,18 @@ public:
 	}
 
 	template<typename TParamRef>
-	void UnsetParameters(FRHICommandList& RHICmdList, const TParamRef& ShaderRHI)
+	void UnsetParameters(FRHICommandList& RHICmdList, const TParamRef& ShaderRHI, const FDistanceFieldObjectBuffers& ObjectBuffers, bool bBarrier = false)
 	{
 		ObjectBounds.UnsetUAV(RHICmdList, ShaderRHI);
 		ObjectData.UnsetUAV(RHICmdList, ShaderRHI);
+
+		if (bBarrier)
+		{
+			FUnorderedAccessViewRHIParamRef OutUAVs[2];
+			OutUAVs[0] = ObjectBuffers.Bounds.UAV;
+			OutUAVs[1] = ObjectBuffers.Data.UAV;
+			RHICmdList.TransitionResources(EResourceTransitionAccess::EReadable, EResourceTransitionPipeline::EComputeToCompute, OutUAVs, ARRAY_COUNT(OutUAVs));
+		}
 	}
 
 	friend FArchive& operator<<(FArchive& Ar, FDistanceFieldObjectBufferParameters& P)
@@ -315,6 +331,18 @@ public:
 	template<typename TParamRef>
 	void Set(FRHICommandList& RHICmdList, const TParamRef& ShaderRHI, const FDistanceFieldCulledObjectBuffers& ObjectBuffers)
 	{
+		int32 NumOutUAVs = 0;
+		FUnorderedAccessViewRHIParamRef OutUAVs[4];
+		OutUAVs[NumOutUAVs++] = ObjectBuffers.ObjectIndirectArguments.UAV;
+		OutUAVs[NumOutUAVs++] = ObjectBuffers.Bounds.UAV;
+		OutUAVs[NumOutUAVs++] = ObjectBuffers.Data.UAV;
+
+		if (CulledObjectBoxBounds.IsBound()) 
+		{
+			OutUAVs[NumOutUAVs++] = ObjectBuffers.BoxBounds.UAV;
+		}
+		RHICmdList.TransitionResources(EResourceTransitionAccess::ERWBarrier, EResourceTransitionPipeline::EComputeToCompute, OutUAVs, NumOutUAVs);
+
 		ObjectIndirectArguments.SetBuffer(RHICmdList, ShaderRHI, ObjectBuffers.ObjectIndirectArguments);
 		CulledObjectBounds.SetBuffer(RHICmdList, ShaderRHI, ObjectBuffers.Bounds);
 		CulledObjectData.SetBuffer(RHICmdList, ShaderRHI, ObjectBuffers.Data);
@@ -342,12 +370,24 @@ public:
 	}
 
 	template<typename TParamRef>
-	void UnsetParameters(FRHICommandList& RHICmdList, const TParamRef& ShaderRHI)
+	void UnsetParameters(FRHICommandList& RHICmdList, const TParamRef& ShaderRHI, const FDistanceFieldCulledObjectBuffers& ObjectBuffers)
 	{
 		ObjectIndirectArguments.UnsetUAV(RHICmdList, ShaderRHI);
 		CulledObjectBounds.UnsetUAV(RHICmdList, ShaderRHI);
 		CulledObjectData.UnsetUAV(RHICmdList, ShaderRHI);
 		CulledObjectBoxBounds.UnsetUAV(RHICmdList, ShaderRHI);
+
+		int32 NumOutUAVs = 0;
+		FUnorderedAccessViewRHIParamRef OutUAVs[3];
+		OutUAVs[NumOutUAVs++] = ObjectBuffers.ObjectIndirectArguments.UAV;
+		OutUAVs[NumOutUAVs++] = ObjectBuffers.Bounds.UAV;
+		OutUAVs[NumOutUAVs++] = ObjectBuffers.Data.UAV;
+
+		if (CulledObjectBoxBounds.IsBound())
+		{
+			OutUAVs[NumOutUAVs++] = ObjectBuffers.BoxBounds.UAV;
+		}
+		RHICmdList.TransitionResources(EResourceTransitionAccess::EReadable, EResourceTransitionPipeline::EComputeToCompute, OutUAVs, NumOutUAVs);
 	}
 
 	friend FArchive& operator<<(FArchive& Ar, FDistanceFieldCulledObjectBufferParameters& P)
diff --git a/Engine/Source/Runtime/Renderer/Private/DistanceFieldObjectManagement.cpp b/Engine/Source/Runtime/Renderer/Private/DistanceFieldObjectManagement.cpp
--- a/Engine/Source/Runtime/Renderer/Private/DistanceFieldObjectManagement.cpp
+++ b/Engine/Source/Runtime/Renderer/Private/DistanceFieldObjectManagement.cpp
@@ -154,12 +154,19 @@ public:
 		SetSRVParameter(RHICmdList, ShaderRHI, UploadOperationIndices, InUploadOperationIndices);
 		SetSRVParameter(RHICmdList, ShaderRHI, UploadOperationData, InUploadOperationData);
 
-		ObjectBufferParameters.Set(RHICmdList, ShaderRHI, *(Scene->DistanceFieldSceneData.ObjectBuffers), Scene->DistanceFieldSceneData.NumObjectsInBuffer);
+
+		ObjectBufferParameters.Set(RHICmdList, ShaderRHI, *(Scene->DistanceFieldSceneData.ObjectBuffers), Scene->DistanceFieldSceneData.NumObjectsInBuffer, true);
 	}
 
-	void UnsetParameters(FRHICommandList& RHICmdList)
+	void UnsetParameters(FRHICommandList& RHICmdList, const FScene* Scene)
 	{
-		ObjectBufferParameters.UnsetParameters(RHICmdList, GetComputeShader());
+		ObjectBufferParameters.UnsetParameters(RHICmdList, GetComputeShader(), *(Scene->DistanceFieldSceneData.ObjectBuffers), true);
+
+		const FDistanceFieldObjectBuffers& ObjectBuffers = *(Scene->DistanceFieldSceneData.ObjectBuffers);
+		FUnorderedAccessViewRHIParamRef OutUAVs[2];
+		OutUAVs[0] = ObjectBuffers.Bounds.UAV;
+		OutUAVs[1] = ObjectBuffers.Data.UAV;
+		RHICmdList.TransitionResources(EResourceTransitionAccess::EReadable, EResourceTransitionPipeline::EComputeToCompute, OutUAVs, ARRAY_COUNT(OutUAVs));
 	}
 
 	virtual bool Serialize(FArchive& Ar) override
@@ -214,17 +221,27 @@ public:
 	{
 		FComputeShaderRHIParamRef ShaderRHI = GetComputeShader();
 
+		FUnorderedAccessViewRHIParamRef OutUAVs[2];
+		OutUAVs[0] = ObjectBuffersDest.Bounds.UAV;
+		OutUAVs[1] = ObjectBuffersDest.Data.UAV;
+		RHICmdList.TransitionResources(EResourceTransitionAccess::ERWBarrier, EResourceTransitionPipeline::EComputeToCompute, OutUAVs, ARRAY_COUNT(OutUAVs));
+
 		CopyObjectBounds.SetBuffer(RHICmdList, ShaderRHI, ObjectBuffersDest.Bounds);
 		CopyObjectData.SetBuffer(RHICmdList, ShaderRHI, ObjectBuffersDest.Data);
 
 		ObjectBufferParameters.Set(RHICmdList, ShaderRHI, ObjectBuffersSource, NumObjectsValue);
 	}
 
-	void UnsetParameters(FRHICommandList& RHICmdList)
+	void UnsetParameters(FRHICommandList& RHICmdList, FDistanceFieldObjectBuffers& ObjectBuffersDest)
 	{
-		ObjectBufferParameters.UnsetParameters(RHICmdList, GetComputeShader());
+		ObjectBufferParameters.UnsetParameters(RHICmdList, GetComputeShader(), ObjectBuffersDest);
 		CopyObjectBounds.UnsetUAV(RHICmdList, GetComputeShader());
 		CopyObjectData.UnsetUAV(RHICmdList, GetComputeShader());
+
+		FUnorderedAccessViewRHIParamRef OutUAVs[2];
+		OutUAVs[0] = ObjectBuffersDest.Bounds.UAV;
+		OutUAVs[1] = ObjectBuffersDest.Data.UAV;
+		RHICmdList.TransitionResources(EResourceTransitionAccess::EReadable, EResourceTransitionPipeline::EComputeToCompute, OutUAVs, ARRAY_COUNT(OutUAVs));
 	}
 
 	virtual bool Serialize(FArchive& Ar) override
@@ -277,17 +294,28 @@ public:
 	void SetParameters(FRHICommandList& RHICmdList, const FSurfelBuffers& SurfelBuffersSource, const FInstancedSurfelBuffers& InstancedSurfelBuffersSource, FSurfelBuffers& SurfelBuffersDest, int32 NumSurfelsValue)
 	{
 		FComputeShaderRHIParamRef ShaderRHI = GetComputeShader();
+
+		FUnorderedAccessViewRHIParamRef OutUAVs[2];
+		OutUAVs[0] = SurfelBuffersDest.InterpolatedVertexData.UAV;
+		OutUAVs[1] = SurfelBuffersDest.Surfels.UAV;		
+		RHICmdList.TransitionResources(EResourceTransitionAccess::ERWBarrier, EResourceTransitionPipeline::EComputeToCompute, OutUAVs, ARRAY_COUNT(OutUAVs));
+
 		CopyInterpolatedVertexData.SetBuffer(RHICmdList, ShaderRHI, SurfelBuffersDest.InterpolatedVertexData);
 		CopySurfelData.SetBuffer(RHICmdList, ShaderRHI, SurfelBuffersDest.Surfels);
 		SurfelBufferParameters.Set(RHICmdList, ShaderRHI, SurfelBuffersSource, InstancedSurfelBuffersSource);
 		SetShaderValue(RHICmdList, ShaderRHI, NumSurfels, NumSurfelsValue);
 	}
 
-	void UnsetParameters(FRHICommandList& RHICmdList)
+	void UnsetParameters(FRHICommandList& RHICmdList, FSurfelBuffers& SurfelBuffersDest)
 	{
 		SurfelBufferParameters.UnsetParameters(RHICmdList, GetComputeShader());
 		CopyInterpolatedVertexData.UnsetUAV(RHICmdList, GetComputeShader());
 		CopySurfelData.UnsetUAV(RHICmdList, GetComputeShader());
+
+		FUnorderedAccessViewRHIParamRef OutUAVs[2];
+		OutUAVs[0] = SurfelBuffersDest.InterpolatedVertexData.UAV;
+		OutUAVs[1] = SurfelBuffersDest.Surfels.UAV;
+		RHICmdList.TransitionResources(EResourceTransitionAccess::EReadable, EResourceTransitionPipeline::EComputeToCompute, OutUAVs, ARRAY_COUNT(OutUAVs));
 	}
 
 	virtual bool Serialize(FArchive& Ar) override
@@ -342,15 +370,18 @@ public:
 	void SetParameters(FRHICommandList& RHICmdList, const FSurfelBuffers& SurfelBuffersSource, const FInstancedSurfelBuffers& InstancedSurfelBuffersSource, FInstancedSurfelBuffers& InstancedSurfelBuffersDest, int32 NumSurfelsValue)
 	{
 		FComputeShaderRHIParamRef ShaderRHI = GetComputeShader();
+
+		RHICmdList.TransitionResource(EResourceTransitionAccess::ERWBarrier, EResourceTransitionPipeline::EComputeToCompute, InstancedSurfelBuffersDest.VPLFlux.UAV);
 		CopyVPLFlux.SetBuffer(RHICmdList, ShaderRHI, InstancedSurfelBuffersDest.VPLFlux);
 		SurfelBufferParameters.Set(RHICmdList, ShaderRHI, SurfelBuffersSource, InstancedSurfelBuffersSource);
 		SetShaderValue(RHICmdList, ShaderRHI, NumSurfels, NumSurfelsValue);
 	}
 
-	void UnsetParameters(FRHICommandList& RHICmdList)
+	void UnsetParameters(FRHICommandList& RHICmdList, FInstancedSurfelBuffers& InstancedSurfelBuffersDest)
 	{
 		SurfelBufferParameters.UnsetParameters(RHICmdList, GetComputeShader());
 		CopyVPLFlux.UnsetUAV(RHICmdList, GetComputeShader());
+		RHICmdList.TransitionResource(EResourceTransitionAccess::EReadable, EResourceTransitionPipeline::EComputeToCompute, InstancedSurfelBuffersDest.VPLFlux.UAV);
 	}
 
 	virtual bool Serialize(FArchive& Ar) override
@@ -415,14 +446,14 @@ public:
 
 		SetShaderValue(RHICmdList, ShaderRHI, NumRemoveOperations, NumRemoveOperationsValue);
 		SetSRVParameter(RHICmdList, ShaderRHI, RemoveOperationIndices, InRemoveOperationIndices);
-		ObjectBufferParameters.Set(RHICmdList, ShaderRHI, *(Scene->DistanceFieldSceneData.ObjectBuffers), Scene->DistanceFieldSceneData.NumObjectsInBuffer);
+		ObjectBufferParameters.Set(RHICmdList, ShaderRHI, *(Scene->DistanceFieldSceneData.ObjectBuffers), Scene->DistanceFieldSceneData.NumObjectsInBuffer, true);
 		SetSRVParameter(RHICmdList, ShaderRHI, ObjectBounds2, InObjectBounds2);
 		SetSRVParameter(RHICmdList, ShaderRHI, ObjectData2, InObjectData2);
 	}
 
-	void UnsetParameters(FRHICommandList& RHICmdList)
+	void UnsetParameters(FRHICommandList& RHICmdList, const FScene* Scene)
 	{
-		ObjectBufferParameters.UnsetParameters(RHICmdList, GetComputeShader());
+		ObjectBufferParameters.UnsetParameters(RHICmdList, GetComputeShader(), *(Scene->DistanceFieldSceneData.ObjectBuffers), true);
 	}
 
 	virtual bool Serialize(FArchive& Ar) override
@@ -685,7 +716,7 @@ void UpdateGlobalDistanceFieldObjectRemoves(FRHICommandListImmediate& RHICmdList
 					ComputeShader->SetParameters(RHICmdList, Scene, RemoveObjectIndices.Num(), GDistanceFieldRemoveIndices.RemoveIndices.BufferSRV, NULL, NULL);
 
 					DispatchComputeShader(RHICmdList, *ComputeShader, FMath::DivideAndRoundUp<uint32>(RemoveObjectIndices.Num(), UpdateObjectsGroupSize), 1, 1);
-					ComputeShader->UnsetParameters(RHICmdList);
+					ComputeShader->UnsetParameters(RHICmdList, Scene);
 				}
 				else
 				{
@@ -695,7 +726,7 @@ void UpdateGlobalDistanceFieldObjectRemoves(FRHICommandListImmediate& RHICmdList
 					ComputeShader->SetParameters(RHICmdList, Scene, RemoveObjectIndices.Num(), GDistanceFieldRemoveIndices.RemoveIndices.BufferSRV, TemporaryCopySourceBuffers->Bounds.SRV, TemporaryCopySourceBuffers->Data.SRV);
 
 					DispatchComputeShader(RHICmdList, *ComputeShader, FMath::DivideAndRoundUp<uint32>(RemoveObjectIndices.Num(), UpdateObjectsGroupSize), 1, 1);
-					ComputeShader->UnsetParameters(RHICmdList);
+					ComputeShader->UnsetParameters(RHICmdList, Scene);
 
 					TemporaryCopySourceBuffers->Release();
 					delete TemporaryCopySourceBuffers;
@@ -1003,7 +1034,7 @@ void FDeferredShadingSceneRenderer::UpdateGlobalDistanceFieldObjectBuffers(FRHIC
 							ComputeShader->SetParameters(RHICmdList, *(DistanceFieldSceneData.SurfelBuffers), *(DistanceFieldSceneData.InstancedSurfelBuffers), *NewSurfelBuffers, OriginalNumSurfels);
 
 							DispatchComputeShader(RHICmdList, *ComputeShader, FMath::DivideAndRoundUp<uint32>(OriginalNumSurfels, UpdateObjectsGroupSize), 1, 1);
-							ComputeShader->UnsetParameters(RHICmdList);
+							ComputeShader->UnsetParameters(RHICmdList, *NewSurfelBuffers);
 						}
 
 						DistanceFieldSceneData.SurfelBuffers->Release();
@@ -1033,7 +1064,7 @@ void FDeferredShadingSceneRenderer::UpdateGlobalDistanceFieldObjectBuffers(FRHIC
 							ComputeShader->SetParameters(RHICmdList, *(DistanceFieldSceneData.SurfelBuffers), *(DistanceFieldSceneData.InstancedSurfelBuffers), *NewInstancedSurfelBuffers, OriginalNumInstancedSurfels);
 
 							DispatchComputeShader(RHICmdList, *ComputeShader, FMath::DivideAndRoundUp<uint32>(OriginalNumInstancedSurfels, UpdateObjectsGroupSize), 1, 1);
-							ComputeShader->UnsetParameters(RHICmdList);
+							ComputeShader->UnsetParameters(RHICmdList, *NewInstancedSurfelBuffers);
 						}
 
 						DistanceFieldSceneData.InstancedSurfelBuffers->Release();
@@ -1101,7 +1132,7 @@ void FDeferredShadingSceneRenderer::UpdateGlobalDistanceFieldObjectBuffers(FRHIC
 						ComputeShader->SetParameters(RHICmdList, *(DistanceFieldSceneData.ObjectBuffers), *NewObjectBuffers, OriginalNumObjects);
 
 						DispatchComputeShader(RHICmdList, *ComputeShader, FMath::DivideAndRoundUp<uint32>(OriginalNumObjects, UpdateObjectsGroupSize), 1, 1);
-						ComputeShader->UnsetParameters(RHICmdList);
+						ComputeShader->UnsetParameters(RHICmdList, *NewObjectBuffers);
 					}
 
 					DistanceFieldSceneData.ObjectBuffers->Release();
@@ -1148,7 +1179,7 @@ void FDeferredShadingSceneRenderer::UpdateGlobalDistanceFieldObjectBuffers(FRHIC
 				ComputeShader->SetParameters(RHICmdList, Scene, UploadObjectIndices.Num(), GDistanceFieldUploadIndices.UploadIndices.BufferSRV, GDistanceFieldUploadData.UploadData.BufferSRV);
 
 				DispatchComputeShader(RHICmdList, *ComputeShader, FMath::DivideAndRoundUp<uint32>(UploadObjectIndices.Num(), UpdateObjectsGroupSize), 1, 1);
-				ComputeShader->UnsetParameters(RHICmdList);
+				ComputeShader->UnsetParameters(RHICmdList, Scene);
 			}
 		}
 
diff --git a/Engine/Source/Runtime/Renderer/Private/DistanceFieldScreenGridLighting.cpp b/Engine/Source/Runtime/Renderer/Private/DistanceFieldScreenGridLighting.cpp
--- a/Engine/Source/Runtime/Renderer/Private/DistanceFieldScreenGridLighting.cpp
+++ b/Engine/Source/Runtime/Renderer/Private/DistanceFieldScreenGridLighting.cpp
@@ -172,18 +172,37 @@ public:
 
 		FAOScreenGridResources* ScreenGridResources = View.ViewState->AOScreenGridResources;
 
-		ScreenGridConeVisibility.SetBuffer(RHICmdList, ShaderRHI, ScreenGridResources->ScreenGridConeVisibility);
+		int32 NumOutUAVs = 0;
+		FUnorderedAccessViewRHIParamRef OutUAVs[2];
+		OutUAVs[NumOutUAVs++] = ScreenGridResources->ScreenGridConeVisibility.UAV;
+		if (bSupportIrradiance)
+		{
+			OutUAVs[NumOutUAVs++] = ScreenGridResources->ConeDepthVisibilityFunction.UAV;			
+		}
+		RHICmdList.TransitionResources(EResourceTransitionAccess::ERWBarrier, EResourceTransitionPipeline::EComputeToCompute, OutUAVs, NumOutUAVs);
 
+		ScreenGridConeVisibility.SetBuffer(RHICmdList, ShaderRHI, ScreenGridResources->ScreenGridConeVisibility);
 		if (bSupportIrradiance)
 		{
 			ConeDepthVisibilityFunction.SetBuffer(RHICmdList, ShaderRHI, ScreenGridResources->ConeDepthVisibilityFunction);
 		}
 	}
 
-	void UnsetParameters(FRHICommandList& RHICmdList)
+	void UnsetParameters(FRHICommandList& RHICmdList, const FViewInfo& View)
 	{
 		ScreenGridConeVisibility.UnsetUAV(RHICmdList, GetComputeShader());
 		ConeDepthVisibilityFunction.UnsetUAV(RHICmdList, GetComputeShader());
+
+		FAOScreenGridResources* ScreenGridResources = View.ViewState->AOScreenGridResources;
+
+		int32 NumOutUAVs = 0;
+		FUnorderedAccessViewRHIParamRef OutUAVs[2];
+		OutUAVs[NumOutUAVs++] = ScreenGridResources->ScreenGridConeVisibility.UAV;
+		if (bSupportIrradiance)
+		{
+			OutUAVs[NumOutUAVs++] = ScreenGridResources->ConeDepthVisibilityFunction.UAV;
+		}
+		RHICmdList.TransitionResources(EResourceTransitionAccess::EReadable, EResourceTransitionPipeline::EComputeToCompute, OutUAVs, NumOutUAVs);
 	}
 
 	virtual bool Serialize(FArchive& Ar) override
@@ -329,12 +348,16 @@ public:
 
 		FAOScreenGridResources* ScreenGridResources = View.ViewState->AOScreenGridResources;
 
+		RHICmdList.TransitionResource(EResourceTransitionAccess::ERWBarrier, EResourceTransitionPipeline::EComputeToCompute, ScreenGridResources->ScreenGridConeVisibility.UAV);
 		ScreenGridConeVisibility.SetBuffer(RHICmdList, ShaderRHI, ScreenGridResources->ScreenGridConeVisibility);
 	}
 
-	void UnsetParameters(FRHICommandList& RHICmdList)
+	void UnsetParameters(FRHICommandList& RHICmdList, const FViewInfo& View)
 	{
 		ScreenGridConeVisibility.UnsetUAV(RHICmdList, GetComputeShader());
+
+		FAOScreenGridResources* ScreenGridResources = View.ViewState->AOScreenGridResources;
+		RHICmdList.TransitionResource(EResourceTransitionAccess::EReadable, EResourceTransitionPipeline::EComputeToCompute, ScreenGridResources->ScreenGridConeVisibility.UAV);
 	}
 
 	virtual bool Serialize(FArchive& Ar) override
@@ -449,14 +472,16 @@ public:
 
 		FAOScreenGridResources* ScreenGridResources = View.ViewState->AOScreenGridResources;
 
+		RHICmdList.TransitionResource(EResourceTransitionAccess::ERWBarrier, EResourceTransitionPipeline::EComputeToCompute, DownsampledBentNormal.UAV);
 		DistanceFieldBentNormal.SetTexture(RHICmdList, ShaderRHI, DownsampledBentNormal.ShaderResourceTexture, DownsampledBentNormal.UAV);
 
 		SetSRVParameter(RHICmdList, ShaderRHI, ScreenGridConeVisibility, ScreenGridResources->ScreenGridConeVisibility.SRV);
 	}
 
-	void UnsetParameters(FRHICommandList& RHICmdList)
+	void UnsetParameters(FRHICommandList& RHICmdList, FSceneRenderTargetItem& DownsampledBentNormal)
 	{
 		DistanceFieldBentNormal.UnsetUAV(RHICmdList, GetComputeShader());
+		RHICmdList.TransitionResource(EResourceTransitionAccess::EReadable, EResourceTransitionPipeline::EComputeToCompute, DownsampledBentNormal.UAV);
 	}
 
 	virtual bool Serialize(FArchive& Ar) override
@@ -781,7 +806,7 @@ void FDeferredShadingSceneRenderer::RenderDistanceFieldAOScreenGrid(
 				RHICmdList.SetComputeShader(ComputeShader->GetComputeShader());
 				ComputeShader->SetParameters(RHICmdList, View, TileListGroupSize, DistanceFieldNormal->GetRenderTargetItem(), Parameters, View.GlobalDistanceFieldInfo);
 				DispatchComputeShader(RHICmdList, *ComputeShader, GroupSizeX, GroupSizeY, 1);
-				ComputeShader->UnsetParameters(RHICmdList);
+				ComputeShader->UnsetParameters(RHICmdList, View);
 			}
 			else
 			{
@@ -790,7 +815,7 @@ void FDeferredShadingSceneRenderer::RenderDistanceFieldAOScreenGrid(
 				RHICmdList.SetComputeShader(ComputeShader->GetComputeShader());
 				ComputeShader->SetParameters(RHICmdList, View, TileListGroupSize, DistanceFieldNormal->GetRenderTargetItem(), Parameters, View.GlobalDistanceFieldInfo);
 				DispatchComputeShader(RHICmdList, *ComputeShader, GroupSizeX, GroupSizeY, 1);
-				ComputeShader->UnsetParameters(RHICmdList);
+				ComputeShader->UnsetParameters(RHICmdList, View);
 			}
 		}
 		else
@@ -802,7 +827,7 @@ void FDeferredShadingSceneRenderer::RenderDistanceFieldAOScreenGrid(
 				RHICmdList.SetComputeShader(ComputeShader->GetComputeShader());
 				ComputeShader->SetParameters(RHICmdList, View, TileListGroupSize, DistanceFieldNormal->GetRenderTargetItem(), Parameters, View.GlobalDistanceFieldInfo);
 				DispatchComputeShader(RHICmdList, *ComputeShader, GroupSizeX, GroupSizeY, 1);
-				ComputeShader->UnsetParameters(RHICmdList);
+				ComputeShader->UnsetParameters(RHICmdList, View);
 			}
 			else
 			{
@@ -811,7 +836,7 @@ void FDeferredShadingSceneRenderer::RenderDistanceFieldAOScreenGrid(
 				RHICmdList.SetComputeShader(ComputeShader->GetComputeShader());
 				ComputeShader->SetParameters(RHICmdList, View, TileListGroupSize, DistanceFieldNormal->GetRenderTargetItem(), Parameters, View.GlobalDistanceFieldInfo);
 				DispatchComputeShader(RHICmdList, *ComputeShader, GroupSizeX, GroupSizeY, 1);
-				ComputeShader->UnsetParameters(RHICmdList);
+				ComputeShader->UnsetParameters(RHICmdList, View);
 			}
 		}
 	}
@@ -832,7 +857,7 @@ void FDeferredShadingSceneRenderer::RenderDistanceFieldAOScreenGrid(
 			RHICmdList.SetComputeShader(ComputeShader->GetComputeShader());
 			ComputeShader->SetParameters(RHICmdList, View, TileListGroupSize, DistanceFieldNormal->GetRenderTargetItem(), Parameters, View.GlobalDistanceFieldInfo);
 			DispatchComputeShader(RHICmdList, *ComputeShader, GroupSizeX, GroupSizeY, 1);
-			ComputeShader->UnsetParameters(RHICmdList);
+			ComputeShader->UnsetParameters(RHICmdList, View);
 		}
 		else
 		{
@@ -841,7 +866,7 @@ void FDeferredShadingSceneRenderer::RenderDistanceFieldAOScreenGrid(
 			RHICmdList.SetComputeShader(ComputeShader->GetComputeShader());
 			ComputeShader->SetParameters(RHICmdList, View, TileListGroupSize, DistanceFieldNormal->GetRenderTargetItem(), Parameters, View.GlobalDistanceFieldInfo);
 			DispatchComputeShader(RHICmdList, *ComputeShader, GroupSizeX, GroupSizeY, 1);
-			ComputeShader->UnsetParameters(RHICmdList);
+			ComputeShader->UnsetParameters(RHICmdList, View);
 		}
 	}
 
@@ -886,7 +911,7 @@ void FDeferredShadingSceneRenderer::RenderDistanceFieldAOScreenGrid(
 		RHICmdList.SetComputeShader(ComputeShader->GetComputeShader());
 		ComputeShader->SetParameters(RHICmdList, View, DistanceFieldNormal->GetRenderTargetItem(), DownsampledBentNormal->GetRenderTargetItem());
 		DispatchComputeShader(RHICmdList, *ComputeShader, GroupSizeX, GroupSizeY, 1);
-		ComputeShader->UnsetParameters(RHICmdList);
+		ComputeShader->UnsetParameters(RHICmdList, DownsampledBentNormal->GetRenderTargetItem());
 	}
 
 	GRenderTargetPool.VisualizeTexture.SetCheckPoint(RHICmdList, DownsampledBentNormal);
diff --git a/Engine/Source/Runtime/Renderer/Private/DistanceFieldShadowing.cpp b/Engine/Source/Runtime/Renderer/Private/DistanceFieldShadowing.cpp
--- a/Engine/Source/Runtime/Renderer/Private/DistanceFieldShadowing.cpp
+++ b/Engine/Source/Runtime/Renderer/Private/DistanceFieldShadowing.cpp
@@ -109,6 +109,13 @@ public:
 		FGlobalShader::SetParameters(RHICmdList, ShaderRHI, View);
 		ObjectBufferParameters.Set(RHICmdList, ShaderRHI, *(Scene->DistanceFieldSceneData.ObjectBuffers), Scene->DistanceFieldSceneData.NumObjectsInBuffer);
 
+		FUnorderedAccessViewRHIParamRef OutUAVs[4];
+		OutUAVs[0] = GShadowCulledObjectBuffers.Buffers.ObjectIndirectArguments.UAV;
+		OutUAVs[1] = GShadowCulledObjectBuffers.Buffers.Bounds.UAV;
+		OutUAVs[2] = GShadowCulledObjectBuffers.Buffers.Data.UAV;
+		OutUAVs[3] = GShadowCulledObjectBuffers.Buffers.BoxBounds.UAV;		
+		RHICmdList.TransitionResources(EResourceTransitionAccess::ERWBarrier, EResourceTransitionPipeline::EComputeToCompute, OutUAVs, ARRAY_COUNT(OutUAVs));
+
 		ObjectIndirectArguments.SetBuffer(RHICmdList, ShaderRHI, GShadowCulledObjectBuffers.Buffers.ObjectIndirectArguments);
 		CulledObjectBounds.SetBuffer(RHICmdList, ShaderRHI, GShadowCulledObjectBuffers.Buffers.Bounds);
 		CulledObjectData.SetBuffer(RHICmdList, ShaderRHI, GShadowCulledObjectBuffers.Buffers.Data);
@@ -129,13 +136,20 @@ public:
 		}
 	}
 
-	void UnsetParameters(FRHICommandList& RHICmdList)
+	void UnsetParameters(FRHICommandList& RHICmdList, const FScene* Scene)
 	{
-		ObjectBufferParameters.UnsetParameters(RHICmdList, GetComputeShader());
+		ObjectBufferParameters.UnsetParameters(RHICmdList, GetComputeShader(), *(Scene->DistanceFieldSceneData.ObjectBuffers));
 		ObjectIndirectArguments.UnsetUAV(RHICmdList, GetComputeShader());
 		CulledObjectBounds.UnsetUAV(RHICmdList, GetComputeShader());
 		CulledObjectData.UnsetUAV(RHICmdList, GetComputeShader());
 		CulledObjectBoxBounds.UnsetUAV(RHICmdList, GetComputeShader());
+
+		FUnorderedAccessViewRHIParamRef OutUAVs[4];
+		OutUAVs[0] = GShadowCulledObjectBuffers.Buffers.ObjectIndirectArguments.UAV;
+		OutUAVs[1] = GShadowCulledObjectBuffers.Buffers.Bounds.UAV;
+		OutUAVs[2] = GShadowCulledObjectBuffers.Buffers.Data.UAV;
+		OutUAVs[3] = GShadowCulledObjectBuffers.Buffers.BoxBounds.UAV;
+		RHICmdList.TransitionResources(EResourceTransitionAccess::EReadable, EResourceTransitionPipeline::EComputeToCompute, OutUAVs, ARRAY_COUNT(OutUAVs));
 	}
 
 	virtual bool Serialize(FArchive& Ar) override
@@ -205,14 +219,17 @@ public:
 		FComputeShaderRHIParamRef ShaderRHI = GetComputeShader();
 
 		FGlobalShader::SetParameters(RHICmdList, ShaderRHI, View);
+
+		RHICmdList.TransitionResource(EResourceTransitionAccess::ERWBarrier, EResourceTransitionPipeline::EComputeToCompute, TileIntersectionResources->TileHeadDataUnpacked.UAV);
 		ShadowTileHeadDataUnpacked.SetBuffer(RHICmdList, ShaderRHI, TileIntersectionResources->TileHeadDataUnpacked);
 
 		SetShaderValue(RHICmdList, ShaderRHI, NumGroups, NumGroupsValue);
 	}
 
-	void UnsetParameters(FRHICommandList& RHICmdList)
+	void UnsetParameters(FRHICommandList& RHICmdList, FLightTileIntersectionResources* TileIntersectionResources)
 	{
 		ShadowTileHeadDataUnpacked.UnsetUAV(RHICmdList, GetComputeShader());
+		RHICmdList.TransitionResource(EResourceTransitionAccess::EReadable, EResourceTransitionPipeline::EComputeToCompute, TileIntersectionResources->TileHeadDataUnpacked.UAV);
 	}
 
 	virtual bool Serialize(FArchive& Ar) override
@@ -438,6 +455,7 @@ public:
 
 		FGlobalShader::SetParameters(RHICmdList, ShaderRHI, View);
 
+		RHICmdList.TransitionResource(EResourceTransitionAccess::ERWBarrier, EResourceTransitionPipeline::EComputeToCompute, ShadowFactorsValue.UAV);
 		ShadowFactors.SetTexture(RHICmdList, ShaderRHI, ShadowFactorsValue.ShaderResourceTexture, ShadowFactorsValue.UAV);
 
 		ObjectParameters.Set(RHICmdList, ShaderRHI, GShadowCulledObjectBuffers.Buffers);
@@ -487,9 +505,10 @@ public:
 		SetShaderValue(RHICmdList, ShaderRHI, DownsampleFactor, GetDFShadowDownsampleFactor());
 	}
 
-	void UnsetParameters(FRHICommandList& RHICmdList)
+	void UnsetParameters(FRHICommandList& RHICmdList, FSceneRenderTargetItem& ShadowFactorsValue)
 	{
 		ShadowFactors.UnsetUAV(RHICmdList, GetComputeShader());
+		RHICmdList.TransitionResource(EResourceTransitionAccess::ERWBarrier, EResourceTransitionPipeline::EComputeToCompute, ShadowFactorsValue.UAV);
 	}
 
 	// FShader interface.
@@ -654,7 +673,7 @@ void CullDistanceFieldObjectsForLight(
 			ComputeShader->SetParameters(RHICmdList, Scene, View, WorldToShadowValue, NumPlanes, PlaneData, ShadowBoundingSphereValue);
 
 			DispatchComputeShader(RHICmdList, *ComputeShader, FMath::DivideAndRoundUp<uint32>(Scene->DistanceFieldSceneData.NumObjectsInBuffer, UpdateObjectsGroupSize), 1, 1);
-			ComputeShader->UnsetParameters(RHICmdList);
+			ComputeShader->UnsetParameters(RHICmdList, Scene);
 		}
 	}
 
@@ -691,7 +710,7 @@ void CullDistanceFieldObjectsForLight(
 			ComputeShader->SetParameters(RHICmdList, View, FVector2D(LightTileDimensions.X, LightTileDimensions.Y), TileIntersectionResources);
 			DispatchComputeShader(RHICmdList, *ComputeShader, GroupSizeX, GroupSizeY, 1);
 
-			ComputeShader->UnsetParameters(RHICmdList);
+			ComputeShader->UnsetParameters(RHICmdList, TileIntersectionResources);
 		}
 		
 		{
@@ -700,6 +719,7 @@ void CullDistanceFieldObjectsForLight(
 
 			TArray<FUnorderedAccessViewRHIParamRef> UAVs;
 			PixelShader->GetUAVs(View, TileIntersectionResources, UAVs);
+			RHICmdList.TransitionResources(EResourceTransitionAccess::ERWBarrier, EResourceTransitionPipeline::EComputeToCompute, UAVs.GetData(), UAVs.Num());
 			RHICmdList.SetRenderTargets(0, (const FRHIRenderTargetView *)NULL, NULL, UAVs.Num(), UAVs.GetData());
 
 			RHICmdList.SetViewport(0, 0, 0.0f, LightTileDimensions.X, LightTileDimensions.Y, 1.0f);
@@ -725,6 +745,7 @@ void CullDistanceFieldObjectsForLight(
 				0);
 
 			SetRenderTarget(RHICmdList, NULL, NULL);
+			RHICmdList.TransitionResources(EResourceTransitionAccess::ERWBarrier, EResourceTransitionPipeline::EComputeToCompute, UAVs.GetData(), UAVs.Num());
 		}
 	}
 }
@@ -845,29 +866,30 @@ void FProjectedShadowInfo::RenderRayTracedDistanceFieldProjection(FRHICommandLis
 
 					SetRenderTarget(RHICmdList, NULL, NULL);
 
+					FSceneRenderTargetItem& RayTracedShadowsRTI = RayTracedShadowsRT->GetRenderTargetItem();
 					if (bDirectionalLight && GShadowScatterTileCulling)
 					{
 						TShaderMapRef<TDistanceFieldShadowingCS<DFS_DirectionalLightScatterTileCulling> > ComputeShader(View.ShaderMap);
 						RHICmdList.SetComputeShader(ComputeShader->GetComputeShader());
-						ComputeShader->SetParameters(RHICmdList, View, this, RayTracedShadowsRT->GetRenderTargetItem(), FVector2D(GroupSizeX, GroupSizeY), ScissorRect, TileIntersectionResources);
+						ComputeShader->SetParameters(RHICmdList, View, this, RayTracedShadowsRTI, FVector2D(GroupSizeX, GroupSizeY), ScissorRect, TileIntersectionResources);
 						DispatchComputeShader(RHICmdList, *ComputeShader, GroupSizeX, GroupSizeY, 1);
-						ComputeShader->UnsetParameters(RHICmdList);
+						ComputeShader->UnsetParameters(RHICmdList, RayTracedShadowsRTI);
 					}
 					else if (bDirectionalLight)
 					{
 						TShaderMapRef<TDistanceFieldShadowingCS<DFS_DirectionalLightTiledCulling> > ComputeShader(View.ShaderMap);
 						RHICmdList.SetComputeShader(ComputeShader->GetComputeShader());
-						ComputeShader->SetParameters(RHICmdList, View, this, RayTracedShadowsRT->GetRenderTargetItem(), FVector2D(GroupSizeX, GroupSizeY), ScissorRect, TileIntersectionResources);
+						ComputeShader->SetParameters(RHICmdList, View, this, RayTracedShadowsRTI, FVector2D(GroupSizeX, GroupSizeY), ScissorRect, TileIntersectionResources);
 						DispatchComputeShader(RHICmdList, *ComputeShader, GroupSizeX, GroupSizeY, 1);
-						ComputeShader->UnsetParameters(RHICmdList);
+						ComputeShader->UnsetParameters(RHICmdList, RayTracedShadowsRTI);
 					}
 					else
 					{
 						TShaderMapRef<TDistanceFieldShadowingCS<DFS_PointLightTiledCulling> > ComputeShader(View.ShaderMap);
 						RHICmdList.SetComputeShader(ComputeShader->GetComputeShader());
-						ComputeShader->SetParameters(RHICmdList, View, this, RayTracedShadowsRT->GetRenderTargetItem(), FVector2D(GroupSizeX, GroupSizeY), ScissorRect, TileIntersectionResources);
+						ComputeShader->SetParameters(RHICmdList, View, this, RayTracedShadowsRTI, FVector2D(GroupSizeX, GroupSizeY), ScissorRect, TileIntersectionResources);
 						DispatchComputeShader(RHICmdList, *ComputeShader, GroupSizeX, GroupSizeY, 1);
-						ComputeShader->UnsetParameters(RHICmdList);
+						ComputeShader->UnsetParameters(RHICmdList, RayTracedShadowsRTI);
 					}
 				}
 			}
diff --git a/Engine/Source/Runtime/Renderer/Private/DistanceFieldSurfaceCacheLighting.cpp b/Engine/Source/Runtime/Renderer/Private/DistanceFieldSurfaceCacheLighting.cpp
--- a/Engine/Source/Runtime/Renderer/Private/DistanceFieldSurfaceCacheLighting.cpp
+++ b/Engine/Source/Runtime/Renderer/Private/DistanceFieldSurfaceCacheLighting.cpp
@@ -404,10 +404,21 @@ public:
 
 	void SetParameters(FRHICommandList& RHICmdList, const FScene* Scene, const FSceneView& View, const FDistanceFieldAOParameters& Parameters)
 	{
+		FUnorderedAccessViewRHIParamRef OutUAVs[6];
+		OutUAVs[0] = GAOCulledObjectBuffers.Buffers.ObjectIndirectArguments.UAV;
+		OutUAVs[1] = GAOCulledObjectBuffers.Buffers.Bounds.UAV;
+		OutUAVs[2] = GAOCulledObjectBuffers.Buffers.Data.UAV;
+		OutUAVs[3] = GAOCulledObjectBuffers.Buffers.BoxBounds.UAV;
+		OutUAVs[4] = Scene->DistanceFieldSceneData.ObjectBuffers->Data.UAV;
+		OutUAVs[5] = Scene->DistanceFieldSceneData.ObjectBuffers->Bounds.UAV;
+		RHICmdList.TransitionResources(EResourceTransitionAccess::ERWBarrier, EResourceTransitionPipeline::EComputeToCompute, OutUAVs, ARRAY_COUNT(OutUAVs));
+
 		FComputeShaderRHIParamRef ShaderRHI = GetComputeShader();
 		FGlobalShader::SetParameters(RHICmdList, ShaderRHI, View);
 		ObjectBufferParameters.Set(RHICmdList, ShaderRHI, *(Scene->DistanceFieldSceneData.ObjectBuffers), Scene->DistanceFieldSceneData.NumObjectsInBuffer);
 
+		
+
 		ObjectIndirectArguments.SetBuffer(RHICmdList, ShaderRHI, GAOCulledObjectBuffers.Buffers.ObjectIndirectArguments);
 		CulledObjectBounds.SetBuffer(RHICmdList, ShaderRHI, GAOCulledObjectBuffers.Buffers.Bounds);
 		CulledObjectData.SetBuffer(RHICmdList, ShaderRHI, GAOCulledObjectBuffers.Buffers.Data);
@@ -421,13 +432,22 @@ public:
 		SetShaderValue(RHICmdList, ShaderRHI, ObjectBoundingGeometryIndexCount, StencilingGeometry::GLowPolyStencilSphereIndexBuffer.GetIndexCount());
 	}
 
-	void UnsetParameters(FRHICommandList& RHICmdList)
+	void UnsetParameters(FRHICommandList& RHICmdList, const FScene* Scene)
 	{
-		ObjectBufferParameters.UnsetParameters(RHICmdList, GetComputeShader());
+		ObjectBufferParameters.UnsetParameters(RHICmdList, GetComputeShader(), *(Scene->DistanceFieldSceneData.ObjectBuffers));
 		ObjectIndirectArguments.UnsetUAV(RHICmdList, GetComputeShader());
 		CulledObjectBounds.UnsetUAV(RHICmdList, GetComputeShader());
 		CulledObjectData.UnsetUAV(RHICmdList, GetComputeShader());
 		CulledObjectBoxBounds.UnsetUAV(RHICmdList, GetComputeShader());
+
+		FUnorderedAccessViewRHIParamRef OutUAVs[6];
+		OutUAVs[0] = GAOCulledObjectBuffers.Buffers.ObjectIndirectArguments.UAV;
+		OutUAVs[1] = GAOCulledObjectBuffers.Buffers.Bounds.UAV;
+		OutUAVs[2] = GAOCulledObjectBuffers.Buffers.Data.UAV;
+		OutUAVs[3] = GAOCulledObjectBuffers.Buffers.BoxBounds.UAV;
+		OutUAVs[4] = Scene->DistanceFieldSceneData.ObjectBuffers->Data.UAV;
+		OutUAVs[5] = Scene->DistanceFieldSceneData.ObjectBuffers->Bounds.UAV;		
+		RHICmdList.TransitionResources(EResourceTransitionAccess::ERWBarrier, EResourceTransitionPipeline::EComputeToCompute, OutUAVs, ARRAY_COUNT(OutUAVs));
 	}
 
 	virtual bool Serialize(FArchive& Ar)
@@ -509,6 +529,12 @@ public:
 
 		FTileIntersectionResources* TileIntersectionResources = ((FSceneViewState*)View.State)->AOTileIntersectionResources;
 
+		FUnorderedAccessViewRHIParamRef OutUAVs[3];
+		OutUAVs[0] = TileIntersectionResources->TileConeAxisAndCos.UAV;
+		OutUAVs[1] = TileIntersectionResources->TileConeDepthRanges.UAV;
+		OutUAVs[2] = TileIntersectionResources->TileHeadDataUnpacked.UAV;
+		RHICmdList.TransitionResources(EResourceTransitionAccess::ERWBarrier, EResourceTransitionPipeline::EComputeToCompute, OutUAVs, ARRAY_COUNT(OutUAVs));
+
 		TileConeAxisAndCos.SetBuffer(RHICmdList, ShaderRHI, TileIntersectionResources->TileConeAxisAndCos);
 		TileConeDepthRanges.SetBuffer(RHICmdList, ShaderRHI, TileIntersectionResources->TileConeDepthRanges);
 		TileHeadDataUnpacked.SetBuffer(RHICmdList, ShaderRHI, TileIntersectionResources->TileHeadDataUnpacked);
@@ -527,11 +553,19 @@ public:
 			);
 	}
 
-	void UnsetParameters(FRHICommandList& RHICmdList)
+	void UnsetParameters(FRHICommandList& RHICmdList, const FSceneView& View)
 	{
 		TileConeAxisAndCos.UnsetUAV(RHICmdList, GetComputeShader());
 		TileConeDepthRanges.UnsetUAV(RHICmdList, GetComputeShader());
 		TileHeadDataUnpacked.UnsetUAV(RHICmdList, GetComputeShader());
+
+		FTileIntersectionResources* TileIntersectionResources = ((FSceneViewState*)View.State)->AOTileIntersectionResources;
+
+		FUnorderedAccessViewRHIParamRef OutUAVs[3];
+		OutUAVs[0] = TileIntersectionResources->TileConeAxisAndCos.UAV;
+		OutUAVs[1] = TileIntersectionResources->TileConeDepthRanges.UAV;
+		OutUAVs[2] = TileIntersectionResources->TileHeadDataUnpacked.UAV;
+		RHICmdList.TransitionResources(EResourceTransitionAccess::EReadable, EResourceTransitionPipeline::EComputeToCompute, OutUAVs, ARRAY_COUNT(OutUAVs));
 	}
 
 	virtual bool Serialize(FArchive& Ar)
@@ -766,8 +800,15 @@ public:
 		ObjectParameters.Set(RHICmdList, ShaderRHI, GAOCulledObjectBuffers.Buffers);
 		AOParameters.Set(RHICmdList, ShaderRHI, Parameters);
 
+
 		FTileIntersectionResources* TileIntersectionResources = ((FSceneViewState*)View.State)->AOTileIntersectionResources;
 
+		FUnorderedAccessViewRHIParamRef OutUAVs[3];
+		OutUAVs[0] = TileIntersectionResources->TileHeadDataUnpacked.UAV;
+		OutUAVs[1] = TileIntersectionResources->TileArrayData.UAV;
+		OutUAVs[2] = TileIntersectionResources->TileArrayNextAllocation.UAV;
+		RHICmdList.TransitionResources(EResourceTransitionAccess::ERWBarrier, EResourceTransitionPipeline::EComputeToCompute, OutUAVs, ARRAY_COUNT(OutUAVs));
+
 		TileHeadDataUnpacked.SetBuffer(RHICmdList, ShaderRHI, TileIntersectionResources->TileHeadDataUnpacked);
 		TileArrayData.SetBuffer(RHICmdList, ShaderRHI, TileIntersectionResources->TileArrayData);
 		TileArrayNextAllocation.SetBuffer(RHICmdList, ShaderRHI, TileIntersectionResources->TileArrayNextAllocation);
@@ -777,11 +818,19 @@ public:
 		SetShaderValue(RHICmdList, ShaderRHI, NumGroups, NumGroupsValue);
 	}
 
-	void UnsetParameters(FRHICommandList& RHICmdList)
+	void UnsetParameters(FRHICommandList& RHICmdList, const FSceneView& View)
 	{
 		TileHeadDataUnpacked.UnsetUAV(RHICmdList, GetComputeShader());
 		TileArrayData.UnsetUAV(RHICmdList, GetComputeShader());
 		TileArrayNextAllocation.UnsetUAV(RHICmdList, GetComputeShader());
+
+		FTileIntersectionResources* TileIntersectionResources = ((FSceneViewState*)View.State)->AOTileIntersectionResources;
+
+		FUnorderedAccessViewRHIParamRef OutUAVs[3];
+		OutUAVs[0] = TileIntersectionResources->TileHeadDataUnpacked.UAV;
+		OutUAVs[1] = TileIntersectionResources->TileArrayData.UAV;
+		OutUAVs[2] = TileIntersectionResources->TileArrayNextAllocation.UAV;
+		RHICmdList.TransitionResources(EResourceTransitionAccess::EReadable, EResourceTransitionPipeline::EComputeToCompute, OutUAVs, ARRAY_COUNT(OutUAVs));
 	}
 
 	virtual bool Serialize(FArchive& Ar) override
@@ -901,14 +950,16 @@ public:
 
 		FGlobalShader::SetParameters(RHICmdList, ShaderRHI, View);
 
+		RHICmdList.TransitionResource(EResourceTransitionAccess::ERWBarrier, EResourceTransitionPipeline::EComputeToCompute, DistanceFieldNormalValue.UAV);
 		DistanceFieldNormal.SetTexture(RHICmdList, ShaderRHI, DistanceFieldNormalValue.ShaderResourceTexture, DistanceFieldNormalValue.UAV);
 		AOParameters.Set(RHICmdList, ShaderRHI, Parameters);
 		DeferredParameters.Set(RHICmdList, ShaderRHI, View);
 	}
 
-	void UnsetParameters(FRHICommandList& RHICmdList)
+	void UnsetParameters(FRHICommandList& RHICmdList, FSceneRenderTargetItem& DistanceFieldNormalValue)
 	{
 		DistanceFieldNormal.UnsetUAV(RHICmdList, GetComputeShader());
+		RHICmdList.TransitionResource(EResourceTransitionAccess::EReadable, EResourceTransitionPipeline::EComputeToCompute, DistanceFieldNormalValue.UAV);
 	}
 
 	// FShader interface.
@@ -951,13 +1002,13 @@ void ComputeDistanceFieldNormal(FRHICommandListImmediate& RHICmdList, const TArr
 				ComputeShader->SetParameters(RHICmdList, View, DistanceFieldNormal, Parameters);
 				DispatchComputeShader(RHICmdList, *ComputeShader, GroupSizeX, GroupSizeY, 1);
 
-				ComputeShader->UnsetParameters(RHICmdList);
+				ComputeShader->UnsetParameters(RHICmdList, DistanceFieldNormal);
 			}
 		}
 	}
 	else
 	{
-		SetRenderTarget(RHICmdList, DistanceFieldNormal.TargetableTexture, NULL);
+		SetRenderTarget(RHICmdList, DistanceFieldNormal.TargetableTexture, NULL, true);
 
 		for (int32 ViewIndex = 0; ViewIndex < Views.Num(); ViewIndex++)
 		{
@@ -969,26 +1020,28 @@ void ComputeDistanceFieldNormal(FRHICommandListImmediate& RHICmdList, const TArr
 			RHICmdList.SetRasterizerState(TStaticRasterizerState<FM_Solid, CM_None>::GetRHI());
 			RHICmdList.SetDepthStencilState(TStaticDepthStencilState<false, CF_Always>::GetRHI());
 			RHICmdList.SetBlendState(TStaticBlendState<>::GetRHI());
-			
+
 			TShaderMapRef<FPostProcessVS> VertexShader(View.ShaderMap);
 			TShaderMapRef<FComputeDistanceFieldNormalPS> PixelShader(View.ShaderMap);
 
 			static FGlobalBoundShaderState BoundShaderState;
-			
+
 			SetGlobalBoundShaderState(RHICmdList, View.FeatureLevel, BoundShaderState, GFilterVertexDeclaration.VertexDeclarationRHI, *VertexShader, *PixelShader);
 
 			PixelShader->SetParameters(RHICmdList, View, Parameters);
 
-			DrawRectangle( 
+			DrawRectangle(
 				RHICmdList,
-				0, 0, 
+				0, 0,
 				View.ViewRect.Width() / GAODownsampleFactor, View.ViewRect.Height() / GAODownsampleFactor,
-				0, 0, 
+				0, 0,
 				View.ViewRect.Width(), View.ViewRect.Height(),
 				FIntPoint(View.ViewRect.Width() / GAODownsampleFactor, View.ViewRect.Height() / GAODownsampleFactor),
 				FSceneRenderTargets::Get(RHICmdList).GetBufferSizeXY(),
 				*VertexShader);
 		}
+
+		RHICmdList.TransitionResource(EResourceTransitionAccess::EReadable, DistanceFieldNormal.TargetableTexture);
 	}
 }
 
@@ -1031,16 +1084,29 @@ public:
 
 		SetSRVParameter(RHICmdList, ShaderRHI, DrawParameters, SurfaceCacheResources.Level[DepthLevel]->ScatterDrawParameters.SRV);
 
+		FUnorderedAccessViewRHIParamRef OutUAVs[2];
+		OutUAVs[0] = SurfaceCacheResources.DispatchParameters.UAV;
+		OutUAVs[1] = SurfaceCacheResources.TempResources->ScatterDrawParameters.UAV;
+		RHICmdList.TransitionResources(EResourceTransitionAccess::ERWBarrier, EResourceTransitionPipeline::EComputeToCompute, OutUAVs, ARRAY_COUNT(OutUAVs));
+
 		DispatchParameters.SetBuffer(RHICmdList, ShaderRHI, SurfaceCacheResources.DispatchParameters);
 		ScatterDrawParameters.SetBuffer(RHICmdList, ShaderRHI, SurfaceCacheResources.TempResources->ScatterDrawParameters);
 
 		SetShaderValue(RHICmdList, ShaderRHI, TrimFraction, GAOTrimOldRecordsFraction);
 	}
 
-	void UnsetParameters(FRHICommandList& RHICmdList)
+	void UnsetParameters(FRHICommandList& RHICmdList, const FSceneView& View)
 	{
 		DispatchParameters.UnsetUAV(RHICmdList, GetComputeShader());
 		ScatterDrawParameters.UnsetUAV(RHICmdList, GetComputeShader());
+
+		const FScene* Scene = (const FScene*)View.Family->Scene;
+		FSurfaceCacheResources& SurfaceCacheResources = *Scene->SurfaceCacheResources;
+
+		FUnorderedAccessViewRHIParamRef OutUAVs[2];
+		OutUAVs[0] = SurfaceCacheResources.DispatchParameters.UAV;
+		OutUAVs[1] = SurfaceCacheResources.TempResources->ScatterDrawParameters.UAV;
+		RHICmdList.TransitionResources(EResourceTransitionAccess::EReadable, EResourceTransitionPipeline::EComputeToCompute, OutUAVs, ARRAY_COUNT(OutUAVs));
 	}
 
 	virtual bool Serialize(FArchive& Ar) override
@@ -1111,6 +1177,16 @@ public:
 		CacheParameters.Set(RHICmdList, ShaderRHI, *SurfaceCacheResources.Level[DepthLevel]);
 		SetSRVParameter(RHICmdList, ShaderRHI, DrawParameters, SurfaceCacheResources.Level[DepthLevel]->ScatterDrawParameters.SRV);
 
+		FUnorderedAccessViewRHIParamRef OutUAVs[7];
+		OutUAVs[0] = SurfaceCacheResources.TempResources->PositionAndRadius.UAV;
+		OutUAVs[1] = SurfaceCacheResources.TempResources->Normal.UAV;
+		OutUAVs[2] = SurfaceCacheResources.TempResources->OccluderRadius.UAV;
+		OutUAVs[3] = SurfaceCacheResources.TempResources->BentNormal.UAV;
+		OutUAVs[4] = SurfaceCacheResources.TempResources->Irradiance.UAV;
+		OutUAVs[5] = SurfaceCacheResources.TempResources->TileCoordinate.UAV;
+		OutUAVs[6] = SurfaceCacheResources.TempResources->ScatterDrawParameters.UAV;
+		RHICmdList.TransitionResources(EResourceTransitionAccess::ERWBarrier, EResourceTransitionPipeline::EComputeToCompute, OutUAVs, ARRAY_COUNT(OutUAVs));
+
 		CopyIrradianceCachePositionRadius.SetBuffer(RHICmdList, ShaderRHI, SurfaceCacheResources.TempResources->PositionAndRadius);
 		CopyIrradianceCacheNormal.SetBuffer(RHICmdList, ShaderRHI, SurfaceCacheResources.TempResources->Normal);
 		CopyOccluderRadius.SetBuffer(RHICmdList, ShaderRHI, SurfaceCacheResources.TempResources->OccluderRadius);
@@ -1122,7 +1198,7 @@ public:
 		SetShaderValue(RHICmdList, ShaderRHI, TrimFraction, GAOTrimOldRecordsFraction);
 	}
 
-	void UnsetParameters(FRHICommandList& RHICmdList)
+	void UnsetParameters(FRHICommandList& RHICmdList, const FSceneView& View)
 	{
 		CopyIrradianceCachePositionRadius.UnsetUAV(RHICmdList, GetComputeShader());
 		CopyIrradianceCacheNormal.UnsetUAV(RHICmdList, GetComputeShader());
@@ -1131,6 +1207,19 @@ public:
 		CopyIrradianceCacheIrradiance.UnsetUAV(RHICmdList, GetComputeShader());
 		CopyIrradianceCacheTileCoordinate.UnsetUAV(RHICmdList, GetComputeShader());
 		ScatterDrawParameters.UnsetUAV(RHICmdList, GetComputeShader());
+
+		const FScene* Scene = (const FScene*)View.Family->Scene;
+		FSurfaceCacheResources& SurfaceCacheResources = *Scene->SurfaceCacheResources;
+
+		FUnorderedAccessViewRHIParamRef OutUAVs[7];
+		OutUAVs[0] = SurfaceCacheResources.TempResources->PositionAndRadius.UAV;
+		OutUAVs[1] = SurfaceCacheResources.TempResources->Normal.UAV;
+		OutUAVs[2] = SurfaceCacheResources.TempResources->OccluderRadius.UAV;
+		OutUAVs[3] = SurfaceCacheResources.TempResources->BentNormal.UAV;
+		OutUAVs[4] = SurfaceCacheResources.TempResources->Irradiance.UAV;
+		OutUAVs[5] = SurfaceCacheResources.TempResources->TileCoordinate.UAV;
+		OutUAVs[6] = SurfaceCacheResources.TempResources->ScatterDrawParameters.UAV;
+		RHICmdList.TransitionResources(EResourceTransitionAccess::ERWBarrier, EResourceTransitionPipeline::EComputeToCompute, OutUAVs, ARRAY_COUNT(OutUAVs));
 	}
 
 	virtual bool Serialize(FArchive& Ar) override
@@ -1196,12 +1285,18 @@ public:
 		FSurfaceCacheResources& SurfaceCacheResources = *Scene->SurfaceCacheResources;
 
 		SetSRVParameter(RHICmdList, ShaderRHI, DrawParameters, SurfaceCacheResources.Level[DepthLevel]->ScatterDrawParameters.SRV);
+
+		RHICmdList.TransitionResource(EResourceTransitionAccess::ERWBarrier, EResourceTransitionPipeline::EComputeToCompute, SurfaceCacheResources.Level[DepthLevel]->SavedStartIndex.UAV);
 		SavedStartIndex.SetBuffer(RHICmdList, ShaderRHI, SurfaceCacheResources.Level[DepthLevel]->SavedStartIndex);
 	}
 
-	void UnsetParameters(FRHICommandList& RHICmdList)
+	void UnsetParameters(FRHICommandList& RHICmdList, const FSceneView& View, int32 DepthLevel)
 	{
 		SavedStartIndex.UnsetUAV(RHICmdList, GetComputeShader());
+
+		const FScene* Scene = (const FScene*)View.Family->Scene;
+		FSurfaceCacheResources& SurfaceCacheResources = *Scene->SurfaceCacheResources;
+		RHICmdList.TransitionResource(EResourceTransitionAccess::EReadable, EResourceTransitionPipeline::EComputeToCompute, SurfaceCacheResources.Level[DepthLevel]->SavedStartIndex.UAV);
 	}
 
 	virtual bool Serialize(FArchive& Ar)
@@ -1284,6 +1379,13 @@ public:
 		const FScene* Scene = (const FScene*)View.Family->Scene;
 		FSurfaceCacheResources& SurfaceCacheResources = *Scene->SurfaceCacheResources;
 
+		FUnorderedAccessViewRHIParamRef OutUAVs[4];
+		OutUAVs[0] = SurfaceCacheResources.Level[DepthLevel]->PositionAndRadius.UAV;
+		OutUAVs[1] = SurfaceCacheResources.Level[DepthLevel]->Normal.UAV;
+		OutUAVs[2] = SurfaceCacheResources.Level[DepthLevel]->TileCoordinate.UAV;
+		OutUAVs[3] = SurfaceCacheResources.Level[DepthLevel]->ScatterDrawParameters.UAV;
+		RHICmdList.TransitionResources(EResourceTransitionAccess::ERWBarrier, EResourceTransitionPipeline::EComputeToCompute, OutUAVs, ARRAY_COUNT(OutUAVs));
+
 		IrradianceCachePositionRadius.SetBuffer(RHICmdList, ShaderRHI, SurfaceCacheResources.Level[DepthLevel]->PositionAndRadius);
 		IrradianceCacheNormal.SetBuffer(RHICmdList, ShaderRHI, SurfaceCacheResources.Level[DepthLevel]->Normal);
 		IrradianceCacheTileCoordinate.SetBuffer(RHICmdList, ShaderRHI, SurfaceCacheResources.Level[DepthLevel]->TileCoordinate);
@@ -1329,12 +1431,22 @@ public:
 			);
 	}
 
-	void UnsetParameters(FRHICommandList& RHICmdList)
+	void UnsetParameters(FRHICommandList& RHICmdList, const FSceneView& View, int32 DepthLevel)
 	{
 		IrradianceCachePositionRadius.UnsetUAV(RHICmdList, GetComputeShader());
 		IrradianceCacheNormal.UnsetUAV(RHICmdList, GetComputeShader());
 		IrradianceCacheTileCoordinate.UnsetUAV(RHICmdList, GetComputeShader());
 		ScatterDrawParameters.UnsetUAV(RHICmdList, GetComputeShader());
+
+		const FScene* Scene = (const FScene*)View.Family->Scene;
+		FSurfaceCacheResources& SurfaceCacheResources = *Scene->SurfaceCacheResources;
+
+		FUnorderedAccessViewRHIParamRef OutUAVs[4];
+		OutUAVs[0] = SurfaceCacheResources.Level[DepthLevel]->PositionAndRadius.UAV;
+		OutUAVs[1] = SurfaceCacheResources.Level[DepthLevel]->Normal.UAV;
+		OutUAVs[2] = SurfaceCacheResources.Level[DepthLevel]->TileCoordinate.UAV;
+		OutUAVs[3] = SurfaceCacheResources.Level[DepthLevel]->ScatterDrawParameters.UAV;
+		RHICmdList.TransitionResources(EResourceTransitionAccess::EReadable, EResourceTransitionPipeline::EComputeToCompute, OutUAVs, ARRAY_COUNT(OutUAVs));
 	}
 
 	virtual bool Serialize(FArchive& Ar)
@@ -1463,6 +1575,12 @@ public:
 		SetSRVParameter(RHICmdList, ShaderRHI, SavedStartIndex, SurfaceCacheResources.Level[DepthLevel]->SavedStartIndex.SRV);
 		SetSRVParameter(RHICmdList, ShaderRHI, IrradianceCachePositionRadius, SurfaceCacheResources.Level[DepthLevel]->PositionAndRadius.SRV);
 
+		FUnorderedAccessViewRHIParamRef OutUAVs[3];
+		OutUAVs[0] = SurfaceCacheResources.Level[DepthLevel]->OccluderRadius.UAV;
+		OutUAVs[1] = GTemporaryIrradianceCacheResources.ConeVisibility.UAV;
+		OutUAVs[2] = GTemporaryIrradianceCacheResources.ConeData.UAV;
+		RHICmdList.TransitionResources(EResourceTransitionAccess::ERWBarrier, EResourceTransitionPipeline::EComputeToCompute, OutUAVs, ARRAY_COUNT(OutUAVs));
+
 		OccluderRadius.SetBuffer(RHICmdList, ShaderRHI, SurfaceCacheResources.Level[DepthLevel]->OccluderRadius);
 		RecordConeVisibility.SetBuffer(RHICmdList, ShaderRHI, GTemporaryIrradianceCacheResources.ConeVisibility);
 		RecordConeData.SetBuffer(RHICmdList, ShaderRHI, GTemporaryIrradianceCacheResources.ConeData);
@@ -1532,11 +1650,20 @@ public:
 		SetShaderValue(RHICmdList, ShaderRHI, RecordRadiusScale, GAORecordRadiusScale);
 	}
 
-	void UnsetParameters(FRHICommandList& RHICmdList)
+	void UnsetParameters(FRHICommandList& RHICmdList, const FSceneView& View, int32 DepthLevel)
 	{
 		OccluderRadius.UnsetUAV(RHICmdList, GetComputeShader());
 		RecordConeVisibility.UnsetUAV(RHICmdList, GetComputeShader());
 		RecordConeData.UnsetUAV(RHICmdList, GetComputeShader());
+
+		const FScene* Scene = (const FScene*)View.Family->Scene;
+		FSurfaceCacheResources& SurfaceCacheResources = *Scene->SurfaceCacheResources;
+
+		FUnorderedAccessViewRHIParamRef OutUAVs[3];
+		OutUAVs[0] = SurfaceCacheResources.Level[DepthLevel]->OccluderRadius.UAV;
+		OutUAVs[1] = GTemporaryIrradianceCacheResources.ConeVisibility.UAV;
+		OutUAVs[2] = GTemporaryIrradianceCacheResources.ConeData.UAV;
+		RHICmdList.TransitionResources(EResourceTransitionAccess::EReadable, EResourceTransitionPipeline::EComputeToCompute, OutUAVs, ARRAY_COUNT(OutUAVs));
 	}
 
 	virtual bool Serialize(FArchive& Ar) override
@@ -1662,6 +1789,7 @@ public:
 		SetSRVParameter(RHICmdList, ShaderRHI, SavedStartIndex, SurfaceCacheResources.Level[DepthLevel]->SavedStartIndex.SRV);
 		SetSRVParameter(RHICmdList, ShaderRHI, RecordConeVisibility, GTemporaryIrradianceCacheResources.ConeVisibility.SRV);
 
+		RHICmdList.TransitionResource(EResourceTransitionAccess::ERWBarrier, EResourceTransitionPipeline::EComputeToCompute, SurfaceCacheResources.Level[DepthLevel]->BentNormal.UAV);
 		IrradianceCacheBentNormal.SetBuffer(RHICmdList, ShaderRHI, SurfaceCacheResources.Level[DepthLevel]->BentNormal);
 
 		FAOSampleData2 AOSampleData;
@@ -1689,9 +1817,13 @@ public:
 		SetShaderValue(RHICmdList, ShaderRHI, TanConeHalfAngle, FMath::Tan(GAOConeHalfAngle));
 	}
 
-	void UnsetParameters(FRHICommandList& RHICmdList)
+	void UnsetParameters(FRHICommandList& RHICmdList, const FSceneView& View, int32 DepthLevel)
 	{
 		IrradianceCacheBentNormal.UnsetUAV(RHICmdList, GetComputeShader());
+
+		const FScene* Scene = (const FScene*)View.Family->Scene;
+		FSurfaceCacheResources& SurfaceCacheResources = *Scene->SurfaceCacheResources;
+		RHICmdList.TransitionResource(EResourceTransitionAccess::EReadable, EResourceTransitionPipeline::EComputeToCompute, SurfaceCacheResources.Level[DepthLevel]->BentNormal.UAV);
 	}
 
 	virtual bool Serialize(FArchive& Ar) override
@@ -2160,7 +2292,7 @@ FIntPoint BuildTileObjectLists(FRHICommandListImmediate& RHICmdList, FScene* Sce
 				ComputeShader->SetParameters(RHICmdList, View, DistanceFieldNormal, Scene, FVector2D(GroupSizeX, GroupSizeY), Parameters);
 				DispatchComputeShader(RHICmdList, *ComputeShader, GroupSizeX, GroupSizeY, 1);
 
-				ComputeShader->UnsetParameters(RHICmdList);
+				ComputeShader->UnsetParameters(RHICmdList, View);
 			}
 
 			{
@@ -2171,6 +2303,7 @@ FIntPoint BuildTileObjectLists(FRHICommandListImmediate& RHICmdList, FScene* Sce
 
 				TArray<FUnorderedAccessViewRHIParamRef> UAVs;
 				PixelShader->GetUAVs(Views[0], UAVs);
+				RHICmdList.TransitionResources(EResourceTransitionAccess::ERWBarrier, EResourceTransitionPipeline::EComputeToCompute, UAVs.GetData(), UAVs.Num());
 				RHICmdList.SetRenderTargets(0, (const FRHIRenderTargetView*)NULL, NULL, UAVs.Num(), UAVs.GetData());
 
 				RHICmdList.SetViewport(0, 0, 0.0f, GroupSizeX, GroupSizeY, 1.0f);
@@ -2194,6 +2327,7 @@ FIntPoint BuildTileObjectLists(FRHICommandListImmediate& RHICmdList, FScene* Sce
 					StencilingGeometry::GLowPolyStencilSphereIndexBuffer.IndexBufferRHI, 
 					GAOCulledObjectBuffers.Buffers.ObjectIndirectArguments.Buffer,
 					0);
+				RHICmdList.TransitionResources(EResourceTransitionAccess::ERWBarrier, EResourceTransitionPipeline::EComputeToCompute, UAVs.GetData(), UAVs.Num());
 			}
 		}
 	}
@@ -2236,7 +2370,7 @@ FIntPoint BuildTileObjectLists(FRHICommandListImmediate& RHICmdList, FScene* Sce
 			ComputeShader->SetParameters(RHICmdList, View, Scene, FVector2D(GroupSizeX, GroupSizeY), Parameters);
 			DispatchComputeShader(RHICmdList, *ComputeShader, GroupSizeX, GroupSizeY, 1);
 
-			ComputeShader->UnsetParameters(RHICmdList);
+			ComputeShader->UnsetParameters(RHICmdList, View);
 		}
 	}
 
@@ -2692,7 +2826,7 @@ void RenderDistanceFieldAOSurfaceCache(
 					ComputeShader->SetParameters(RHICmdList, View, DepthLevel);
 					DispatchComputeShader(RHICmdList, *ComputeShader, 1, 1, 1);
 
-					ComputeShader->UnsetParameters(RHICmdList);
+					ComputeShader->UnsetParameters(RHICmdList, View);
 				}
 
 				if (bUseDistanceFieldGI)
@@ -2703,7 +2837,7 @@ void RenderDistanceFieldAOSurfaceCache(
 					ComputeShader->SetParameters(RHICmdList, View, DepthLevel);
 					DispatchIndirectComputeShader(RHICmdList, *ComputeShader, SurfaceCacheResources.DispatchParameters.Buffer, 0);
 
-					ComputeShader->UnsetParameters(RHICmdList);
+					ComputeShader->UnsetParameters(RHICmdList, View);
 				}
 				else
 				{
@@ -2713,7 +2847,7 @@ void RenderDistanceFieldAOSurfaceCache(
 					ComputeShader->SetParameters(RHICmdList, View, DepthLevel);
 					DispatchIndirectComputeShader(RHICmdList, *ComputeShader, SurfaceCacheResources.DispatchParameters.Buffer, 0);
 
-					ComputeShader->UnsetParameters(RHICmdList);
+					ComputeShader->UnsetParameters(RHICmdList, View);
 				}
 
 				Swap(SurfaceCacheResources.Level[DepthLevel], SurfaceCacheResources.TempResources);
@@ -2753,7 +2887,7 @@ void RenderDistanceFieldAOSurfaceCache(
 				ComputeShader->SetParameters(RHICmdList, View, DepthLevel);
 				DispatchComputeShader(RHICmdList, *ComputeShader, 1, 1, 1);
 
-				ComputeShader->UnsetParameters(RHICmdList);
+				ComputeShader->UnsetParameters(RHICmdList, View, DepthLevel);
 			}
 
 			// Create new records which haven't been shaded yet for shading points which don't have a valid interpolation from existing records
@@ -2768,7 +2902,7 @@ void RenderDistanceFieldAOSurfaceCache(
 		        ComputeShader->SetParameters(RHICmdList, View, DistanceFieldAOBentNormalSplat->GetRenderTargetItem(), DistanceFieldNormal->GetRenderTargetItem(), DestLevelDownsampleFactor, DepthLevel, TileListGroupSize, Parameters);
 				DispatchComputeShader(RHICmdList, *ComputeShader, GroupSizeX, GroupSizeY, 1);
 
-				ComputeShader->UnsetParameters(RHICmdList);
+				ComputeShader->UnsetParameters(RHICmdList, View, DepthLevel);
 			}
 
 			{	
@@ -2778,7 +2912,7 @@ void RenderDistanceFieldAOSurfaceCache(
 				ComputeShader->SetParameters(RHICmdList, View, DepthLevel);
 				DispatchComputeShader(RHICmdList, *ComputeShader, 1, 1, 1);
 
-				ComputeShader->UnsetParameters(RHICmdList);
+				ComputeShader->UnsetParameters(RHICmdList, View);
 			}
 
 			// Compute lighting for the new surface cache records by cone-stepping through the object distance fields
@@ -2790,7 +2924,7 @@ void RenderDistanceFieldAOSurfaceCache(
 				ComputeShader->SetParameters(RHICmdList, View, DistanceFieldAOBentNormalSplat->GetRenderTargetItem(), DistanceFieldNormal->GetRenderTargetItem(), DestLevelDownsampleFactor, DepthLevel, TileListGroupSize, Parameters);
 				DispatchIndirectComputeShader(RHICmdList, *ComputeShader, SurfaceCacheResources.DispatchParameters.Buffer, 0);
 
-				ComputeShader->UnsetParameters(RHICmdList);
+				ComputeShader->UnsetParameters(RHICmdList, View, DepthLevel);
 			}
 			else
 			{
@@ -2800,7 +2934,7 @@ void RenderDistanceFieldAOSurfaceCache(
 				ComputeShader->SetParameters(RHICmdList, View, DistanceFieldAOBentNormalSplat->GetRenderTargetItem(), DistanceFieldNormal->GetRenderTargetItem(), DestLevelDownsampleFactor, DepthLevel, TileListGroupSize, Parameters);
 				DispatchIndirectComputeShader(RHICmdList, *ComputeShader, SurfaceCacheResources.DispatchParameters.Buffer, 0);
 
-				ComputeShader->UnsetParameters(RHICmdList);
+				ComputeShader->UnsetParameters(RHICmdList, View, DepthLevel);
 			}
 		}
 
@@ -2827,7 +2961,7 @@ void RenderDistanceFieldAOSurfaceCache(
 			ComputeShader->SetParameters(RHICmdList, View, DepthLevel);
 			DispatchComputeShader(RHICmdList, *ComputeShader, 1, 1, 1);
 
-			ComputeShader->UnsetParameters(RHICmdList);
+			ComputeShader->UnsetParameters(RHICmdList, View);
 		}
 
 		// Compute and store the final bent normal now that all occlusion sources have been computed (distance fields, heightfields)
@@ -2838,7 +2972,7 @@ void RenderDistanceFieldAOSurfaceCache(
 			ComputeShader->SetParameters(RHICmdList, View, DepthLevel, Parameters, NULL, FMatrix::Identity, NULL, NULL);
 			DispatchIndirectComputeShader(RHICmdList, *ComputeShader, SurfaceCacheResources.DispatchParameters.Buffer, 0);
 
-			ComputeShader->UnsetParameters(RHICmdList);
+			ComputeShader->UnsetParameters(RHICmdList, View, DepthLevel);
 		}
 	}
 
@@ -2934,7 +3068,7 @@ bool FDeferredShadingSceneRenderer::RenderDistanceFieldLighting(
 					ComputeShader->SetParameters(RHICmdList, Scene, View, Parameters);
 
 					DispatchComputeShader(RHICmdList, *ComputeShader, FMath::DivideAndRoundUp<uint32>(Scene->DistanceFieldSceneData.NumObjectsInBuffer, UpdateObjectsGroupSize), 1, 1);
-					ComputeShader->UnsetParameters(RHICmdList);
+					ComputeShader->UnsetParameters(RHICmdList, Scene);
 				}
 			}
 			
@@ -3161,7 +3295,7 @@ void FDeferredShadingSceneRenderer::RenderDynamicSkyLighting(FRHICommandListImme
 		}
 		FSceneRenderTargets& SceneContext = FSceneRenderTargets::Get(RHICmdList);
 
-		SceneContext.BeginRenderingSceneColor(RHICmdList);
+		SceneContext.BeginRenderingSceneColor(RHICmdList, ESimpleRenderTargetMode::EExistingColorAndDepth, FExclusiveDepthStencil::DepthRead_StencilRead);
 
 		for( int32 ViewIndex = 0; ViewIndex < Views.Num(); ViewIndex++ )
 		{
@@ -3281,6 +3415,7 @@ public:
 
 		FGlobalShader::SetParameters(RHICmdList, ShaderRHI, View);
 
+		RHICmdList.TransitionResource(EResourceTransitionAccess::ERWBarrier, EResourceTransitionPipeline::EComputeToCompute, VisualizeMeshDistanceFieldsValue.UAV);
 		VisualizeMeshDistanceFields.SetTexture(RHICmdList, ShaderRHI, VisualizeMeshDistanceFieldsValue.ShaderResourceTexture, VisualizeMeshDistanceFieldsValue.UAV);
 
 		ObjectParameters.Set(RHICmdList, ShaderRHI, GAOCulledObjectBuffers.Buffers);
@@ -3295,8 +3430,9 @@ public:
 		SetShaderValue(RHICmdList, ShaderRHI, NumGroups, NumGroupsValue);
 	}
 
-	void UnsetParameters(FRHICommandList& RHICmdList)
+	void UnsetParameters(FRHICommandList& RHICmdList, FSceneRenderTargetItem& VisualizeMeshDistanceFieldsValue)
 	{
+		RHICmdList.TransitionResource(EResourceTransitionAccess::EReadable, EResourceTransitionPipeline::EComputeToCompute, VisualizeMeshDistanceFieldsValue.UAV);
 		VisualizeMeshDistanceFields.UnsetUAV(RHICmdList, GetComputeShader());
 	}
 
@@ -3420,7 +3556,7 @@ void FDeferredShadingSceneRenderer::RenderMeshDistanceFieldVisualization(FRHICom
 					ComputeShader->SetParameters(RHICmdList, Scene, View, Parameters);
 
 					DispatchComputeShader(RHICmdList, *ComputeShader, FMath::DivideAndRoundUp<uint32>(Scene->DistanceFieldSceneData.NumObjectsInBuffer, UpdateObjectsGroupSize), 1, 1);
-					ComputeShader->UnsetParameters(RHICmdList);
+					ComputeShader->UnsetParameters(RHICmdList, Scene);
 				}
 			}
 
@@ -3444,6 +3580,7 @@ void FDeferredShadingSceneRenderer::RenderMeshDistanceFieldVisualization(FRHICom
 
 					SCOPED_DRAW_EVENT(RHICmdList, VisualizeMeshDistanceFieldCS);
 
+					FSceneRenderTargetItem& VisualizeResultRTI = VisualizeResultRT->GetRenderTargetItem();
 					if (bUseGlobalDistanceField)
 					{
 						check(View.GlobalDistanceFieldInfo.Clipmaps.Num() > 0);
@@ -3451,20 +3588,20 @@ void FDeferredShadingSceneRenderer::RenderMeshDistanceFieldVisualization(FRHICom
 						TShaderMapRef<TVisualizeMeshDistanceFieldCS<true> > ComputeShader(ViewInfo.ShaderMap);
 
 						RHICmdList.SetComputeShader(ComputeShader->GetComputeShader());
-						ComputeShader->SetParameters(RHICmdList, ViewInfo, VisualizeResultRT->GetRenderTargetItem(), FVector2D(GroupSizeX, GroupSizeY), Parameters, View.GlobalDistanceFieldInfo);
+						ComputeShader->SetParameters(RHICmdList, ViewInfo, VisualizeResultRTI, FVector2D(GroupSizeX, GroupSizeY), Parameters, View.GlobalDistanceFieldInfo);
 						DispatchComputeShader(RHICmdList, *ComputeShader, GroupSizeX, GroupSizeY, 1);
 
-						ComputeShader->UnsetParameters(RHICmdList);
+						ComputeShader->UnsetParameters(RHICmdList, VisualizeResultRTI);
 					}
 					else
 					{
 						TShaderMapRef<TVisualizeMeshDistanceFieldCS<false> > ComputeShader(ViewInfo.ShaderMap);
 
 						RHICmdList.SetComputeShader(ComputeShader->GetComputeShader());
-						ComputeShader->SetParameters(RHICmdList, ViewInfo, VisualizeResultRT->GetRenderTargetItem(), FVector2D(GroupSizeX, GroupSizeY), Parameters, View.GlobalDistanceFieldInfo);
+						ComputeShader->SetParameters(RHICmdList, ViewInfo, VisualizeResultRTI, FVector2D(GroupSizeX, GroupSizeY), Parameters, View.GlobalDistanceFieldInfo);
 						DispatchComputeShader(RHICmdList, *ComputeShader, GroupSizeX, GroupSizeY, 1);
 
-						ComputeShader->UnsetParameters(RHICmdList);
+						ComputeShader->UnsetParameters(RHICmdList, VisualizeResultRTI);
 					}
 				}
 			}
diff --git a/Engine/Source/Runtime/Renderer/Private/DistanceFieldSurfaceCacheLighting.h b/Engine/Source/Runtime/Renderer/Private/DistanceFieldSurfaceCacheLighting.h
--- a/Engine/Source/Runtime/Renderer/Private/DistanceFieldSurfaceCacheLighting.h
+++ b/Engine/Source/Runtime/Renderer/Private/DistanceFieldSurfaceCacheLighting.h
@@ -306,12 +306,17 @@ public:
 		SetSRVParameter(RHICmdList, ShaderRHI, DrawParameters, SurfaceCacheResources.Level[DepthLevel]->ScatterDrawParameters.SRV);
 		SetSRVParameter(RHICmdList, ShaderRHI, SavedStartIndex, SurfaceCacheResources.Level[DepthLevel]->SavedStartIndex.SRV);
 
+		RHICmdList.TransitionResource(EResourceTransitionAccess::ERWBarrier, EResourceTransitionPipeline::EComputeToCompute, SurfaceCacheResources.DispatchParameters.UAV);
 		DispatchParameters.SetBuffer(RHICmdList, ShaderRHI, SurfaceCacheResources.DispatchParameters);
 	}
 
-	void UnsetParameters(FRHICommandList& RHICmdList)
+	void UnsetParameters(FRHICommandList& RHICmdList, const FSceneView& View)
 	{
+		const FScene* Scene = (const FScene*)View.Family->Scene;
+		FSurfaceCacheResources& SurfaceCacheResources = *Scene->SurfaceCacheResources;
+
 		DispatchParameters.UnsetUAV(RHICmdList, GetComputeShader());
+		RHICmdList.TransitionResource(EResourceTransitionAccess::EReadable, EResourceTransitionPipeline::EComputeToCompute, SurfaceCacheResources.DispatchParameters.UAV);
 	}
 
 	virtual bool Serialize(FArchive& Ar) override
diff --git a/Engine/Source/Runtime/Renderer/Private/GlobalDistanceField.cpp b/Engine/Source/Runtime/Renderer/Private/GlobalDistanceField.cpp
--- a/Engine/Source/Runtime/Renderer/Private/GlobalDistanceField.cpp
+++ b/Engine/Source/Runtime/Renderer/Private/GlobalDistanceField.cpp
@@ -163,6 +163,13 @@ public:
 		FGlobalShader::SetParameters(RHICmdList, ShaderRHI, View);
 		ObjectBufferParameters.Set(RHICmdList, ShaderRHI, *(Scene->DistanceFieldSceneData.ObjectBuffers), Scene->DistanceFieldSceneData.NumObjectsInBuffer);
 
+		FUnorderedAccessViewRHIParamRef OutUAVs[4];
+		OutUAVs[0] = GGlobalDistanceFieldCulledObjectBuffers.Buffers.ObjectIndirectArguments.UAV;
+		OutUAVs[1] = GGlobalDistanceFieldCulledObjectBuffers.Buffers.Bounds.UAV;
+		OutUAVs[2] = GGlobalDistanceFieldCulledObjectBuffers.Buffers.Data.UAV;
+		OutUAVs[3] = GGlobalDistanceFieldCulledObjectBuffers.Buffers.BoxBounds.UAV;
+		RHICmdList.TransitionResources(EResourceTransitionAccess::ERWBarrier, EResourceTransitionPipeline::EComputeToCompute, OutUAVs, ARRAY_COUNT(OutUAVs));
+
 		ObjectIndirectArguments.SetBuffer(RHICmdList, ShaderRHI, GGlobalDistanceFieldCulledObjectBuffers.Buffers.ObjectIndirectArguments);
 		CulledObjectBounds.SetBuffer(RHICmdList, ShaderRHI, GGlobalDistanceFieldCulledObjectBuffers.Buffers.Bounds);
 		CulledObjectData.SetBuffer(RHICmdList, ShaderRHI, GGlobalDistanceFieldCulledObjectBuffers.Buffers.Data);
@@ -174,13 +181,20 @@ public:
 		SetShaderValue(RHICmdList, ShaderRHI, VolumeBounds, VolumeBoundsValue);
 	}
 
-	void UnsetParameters(FRHICommandList& RHICmdList)
+	void UnsetParameters(FRHICommandList& RHICmdList, const FScene* Scene)
 	{
-		ObjectBufferParameters.UnsetParameters(RHICmdList, GetComputeShader());
+		ObjectBufferParameters.UnsetParameters(RHICmdList, GetComputeShader(), *(Scene->DistanceFieldSceneData.ObjectBuffers));
 		ObjectIndirectArguments.UnsetUAV(RHICmdList, GetComputeShader());
 		CulledObjectBounds.UnsetUAV(RHICmdList, GetComputeShader());
 		CulledObjectData.UnsetUAV(RHICmdList, GetComputeShader());
 		CulledObjectBoxBounds.UnsetUAV(RHICmdList, GetComputeShader());
+
+		FUnorderedAccessViewRHIParamRef OutUAVs[4];
+		OutUAVs[0] = GGlobalDistanceFieldCulledObjectBuffers.Buffers.ObjectIndirectArguments.UAV;
+		OutUAVs[1] = GGlobalDistanceFieldCulledObjectBuffers.Buffers.Bounds.UAV;
+		OutUAVs[2] = GGlobalDistanceFieldCulledObjectBuffers.Buffers.Data.UAV;
+		OutUAVs[3] = GGlobalDistanceFieldCulledObjectBuffers.Buffers.BoxBounds.UAV;
+		RHICmdList.TransitionResources(EResourceTransitionAccess::EReadable, EResourceTransitionPipeline::EComputeToCompute, OutUAVs, ARRAY_COUNT(OutUAVs));
 	}
 
 	virtual bool Serialize(FArchive& Ar)
@@ -294,6 +308,7 @@ public:
 		CulledObjectBufferParameters.Set(RHICmdList, ShaderRHI, GGlobalDistanceFieldCulledObjectBuffers.Buffers);
 		GlobalDistanceFieldParameters.Set(RHICmdList, ShaderRHI, GlobalDistanceFieldInfo.ParameterData);
 
+		RHICmdList.TransitionResource(EResourceTransitionAccess::ERWBarrier, EResourceTransitionPipeline::EComputeToCompute, GObjectGridBuffers.CulledObjectGrid.UAV);
 		CulledObjectGrid.SetBuffer(RHICmdList, ShaderRHI, GObjectGridBuffers.CulledObjectGrid);
 
 		const FIntVector GridDimensionValue(
@@ -314,6 +329,7 @@ public:
 	void UnsetParameters(FRHICommandList& RHICmdList)
 	{
 		CulledObjectGrid.UnsetUAV(RHICmdList, GetComputeShader());
+		RHICmdList.TransitionResource(EResourceTransitionAccess::EReadable, EResourceTransitionPipeline::EComputeToCompute, GObjectGridBuffers.CulledObjectGrid.UAV);
 	}
 
 	virtual bool Serialize(FArchive& Ar)
@@ -397,7 +413,9 @@ public:
 		CulledObjectBufferParameters.Set(RHICmdList, ShaderRHI, GGlobalDistanceFieldCulledObjectBuffers.Buffers);
 		GlobalDistanceFieldParameters.Set(RHICmdList, ShaderRHI, GlobalDistanceFieldInfo.ParameterData);
 
-		GlobalDistanceFieldTexture.SetTexture(RHICmdList, ShaderRHI, GlobalDistanceFieldInfo.Clipmaps[ClipmapIndexValue].RenderTarget->GetRenderTargetItem().ShaderResourceTexture, GlobalDistanceFieldInfo.Clipmaps[ClipmapIndexValue].RenderTarget->GetRenderTargetItem().UAV);
+		const FSceneRenderTargetItem& ClipMapRTI = GlobalDistanceFieldInfo.Clipmaps[ClipmapIndexValue].RenderTarget->GetRenderTargetItem();
+		RHICmdList.TransitionResource(EResourceTransitionAccess::ERWBarrier, EResourceTransitionPipeline::EComputeToCompute, ClipMapRTI.UAV);
+		GlobalDistanceFieldTexture.SetTexture(RHICmdList, ShaderRHI, ClipMapRTI.ShaderResourceTexture, ClipMapRTI.UAV);
 
 		SetSRVParameter(RHICmdList, ShaderRHI, CulledObjectGrid, GObjectGridBuffers.CulledObjectGrid.SRV);
 
@@ -418,9 +436,12 @@ public:
 		SetShaderValue(RHICmdList, ShaderRHI, AOGlobalMaxSphereQueryRadius, GlobalMaxSphereQueryRadius);
 	}
 
-	void UnsetParameters(FRHICommandList& RHICmdList)
+	void UnsetParameters(FRHICommandList& RHICmdList, const FGlobalDistanceFieldInfo& GlobalDistanceFieldInfo, int32 ClipmapIndexValue)
 	{
 		GlobalDistanceFieldTexture.UnsetUAV(RHICmdList, GetComputeShader());
+
+		const FSceneRenderTargetItem& ClipMapRTI = GlobalDistanceFieldInfo.Clipmaps[ClipmapIndexValue].RenderTarget->GetRenderTargetItem();
+		RHICmdList.TransitionResource(EResourceTransitionAccess::EReadable, EResourceTransitionPipeline::EComputeToCompute, ClipMapRTI.UAV);
 	}
 
 	virtual bool Serialize(FArchive& Ar)
@@ -529,17 +550,20 @@ void AllocateClipmapTexture(FRHICommandListImmediate& RHICmdList, int32 ClipmapI
 		TextureName = TEXT("GlobalDistanceField3");
 	}
 
+	FPooledRenderTargetDesc VolumeDesc = FPooledRenderTargetDesc(FPooledRenderTargetDesc::CreateVolumeDesc(
+		GAOGlobalDFResolution,
+		GAOGlobalDFResolution,
+		GAOGlobalDFResolution,
+		PF_R16F,
+		FClearValueBinding::None,
+		0,
+		TexCreate_ShaderResource | TexCreate_RenderTargetable | TexCreate_UAV,
+		false));
+	VolumeDesc.AutoWritable = false;
+
 	GRenderTargetPool.FindFreeElement(
 		RHICmdList,
-		FPooledRenderTargetDesc(FPooledRenderTargetDesc::CreateVolumeDesc(
-			GAOGlobalDFResolution,
-			GAOGlobalDFResolution,
-			GAOGlobalDFResolution,
-			PF_R16F,
-			FClearValueBinding::None,
-			0,
-			TexCreate_ShaderResource | TexCreate_RenderTargetable | TexCreate_UAV,
-			false)),
+		VolumeDesc,
 		Texture,
 		TextureName 
 	);
@@ -811,7 +835,7 @@ void UpdateGlobalDistanceFieldVolume(
 							ComputeShader->SetParameters(RHICmdList, Scene, View, MaxOcclusionDistance, VolumeBounds);
 
 							DispatchComputeShader(RHICmdList, *ComputeShader, FMath::DivideAndRoundUp<uint32>(Scene->DistanceFieldSceneData.NumObjectsInBuffer, CullObjectsGroupSize), 1, 1);
-							ComputeShader->UnsetParameters(RHICmdList);
+							ComputeShader->UnsetParameters(RHICmdList, Scene);
 						}
 
 						// Further cull the objects into a low resolution grid
@@ -842,7 +866,7 @@ void UpdateGlobalDistanceFieldVolume(
 						const uint32 NumGroupsZ = FMath::DivideAndRoundUp<int32>(UpdateRegion.CellsSize.Z, CompositeTileSize);
 
 						DispatchComputeShader(RHICmdList, *ComputeShader, NumGroupsX, NumGroupsY, NumGroupsZ);
-						ComputeShader->UnsetParameters(RHICmdList);
+						ComputeShader->UnsetParameters(RHICmdList, GlobalDistanceFieldInfo, ClipmapIndex);
 					}
 				}
 			}
diff --git a/Engine/Source/Runtime/Renderer/Private/HeightfieldLighting.cpp b/Engine/Source/Runtime/Renderer/Private/HeightfieldLighting.cpp
--- a/Engine/Source/Runtime/Renderer/Private/HeightfieldLighting.cpp
+++ b/Engine/Source/Runtime/Renderer/Private/HeightfieldLighting.cpp
@@ -115,6 +115,7 @@ void FHeightfieldLightingAtlas::InitDynamicRHI()
 			TexCreate_None,
 			TexCreate_RenderTargetable,
 			false));
+		Desc5.AutoWritable = false;
 
 		GRenderTargetPool.FindFreeElement(RHICmdList, Desc5, Lighting, TEXT("HeightfieldLightingAtlas"));
 	}
@@ -541,6 +542,8 @@ void FHeightfieldLightingViewInfo::SetupVisibleHeightfields(const FViewInfo& Vie
 							RHICmdList.DrawPrimitive(PT_TriangleList, 0, 2, NumQuads);
 						}
 					}
+
+					RHICmdList.TransitionResources(EResourceTransitionAccess::EReadable, RenderTargets, ARRAY_COUNT(RenderTargets));
 				}
 			}
 		}
@@ -1273,8 +1276,15 @@ public:
 		SetSRVParameter(RHICmdList, ShaderRHI, ScatterDrawParameters, SurfaceCacheResources.Level[DepthLevel]->ScatterDrawParameters.SRV);
 		SetSRVParameter(RHICmdList, ShaderRHI, SavedStartIndex, SurfaceCacheResources.Level[DepthLevel]->SavedStartIndex.SRV);
 
-		OccluderRadius.SetBuffer(RHICmdList, ShaderRHI, SurfaceCacheResources.Level[DepthLevel]->OccluderRadius);
-		RecordConeVisibility.SetBuffer(RHICmdList, ShaderRHI, TemporaryIrradianceCacheResources.ConeVisibility);
+		const FRWBuffer& OccluderRadiusOut = SurfaceCacheResources.Level[DepthLevel]->OccluderRadius;
+		const FRWBuffer& ConeVisibiliyOut = TemporaryIrradianceCacheResources.ConeVisibility;
+		FUnorderedAccessViewRHIParamRef UAVs[2];
+		UAVs[0] = OccluderRadiusOut.UAV;
+		UAVs[1] = ConeVisibiliyOut.UAV;
+
+		RHICmdList.TransitionResources(EResourceTransitionAccess::ERWBarrier, EResourceTransitionPipeline::EComputeToCompute, UAVs, ARRAY_COUNT(UAVs));
+		OccluderRadius.SetBuffer(RHICmdList, ShaderRHI, OccluderRadiusOut);
+		RecordConeVisibility.SetBuffer(RHICmdList, ShaderRHI, ConeVisibiliyOut);
 
 		FAOSampleData2 AOSampleData;
 
@@ -1294,10 +1304,21 @@ public:
 		SetShaderValue(RHICmdList, ShaderRHI, RecordRadiusScale, GAORecordRadiusScale);
 	}
 
-	void UnsetParameters(FRHICommandList& RHICmdList)
+	void UnsetParameters(FRHICommandList& RHICmdList, const FSceneView& View, int32 DepthLevel, const FTemporaryIrradianceCacheResources& TemporaryIrradianceCacheResources)
 	{
 		OccluderRadius.UnsetUAV(RHICmdList, GetComputeShader());
 		RecordConeVisibility.UnsetUAV(RHICmdList, GetComputeShader());
+
+		const FScene* Scene = (const FScene*)View.Family->Scene;
+		FSurfaceCacheResources& SurfaceCacheResources = *Scene->SurfaceCacheResources;
+
+		const FRWBuffer& OccluderRadiusOut = SurfaceCacheResources.Level[DepthLevel]->OccluderRadius;
+		const FRWBuffer& ConeVisibiliyOut = TemporaryIrradianceCacheResources.ConeVisibility;
+		FUnorderedAccessViewRHIParamRef UAVs[2];
+		UAVs[0] = OccluderRadiusOut.UAV;
+		UAVs[1] = ConeVisibiliyOut.UAV;
+
+		RHICmdList.TransitionResources(EResourceTransitionAccess::EReadable, EResourceTransitionPipeline::EComputeToCompute, UAVs, ARRAY_COUNT(UAVs));
 	}
 
 	virtual bool Serialize(FArchive& Ar) override
@@ -1353,7 +1374,7 @@ void FHeightfieldLightingViewInfo::ComputeOcclusionForSamples(
 			ComputeShader->SetParameters(RHICmdList, View, DepthLevel);
 			DispatchComputeShader(RHICmdList, *ComputeShader, 1, 1, 1);
 
-			ComputeShader->UnsetParameters(RHICmdList);
+			ComputeShader->UnsetParameters(RHICmdList, View);
 		}
 
 		SCOPED_DRAW_EVENT(RHICmdList, HeightfieldOcclusion);
@@ -1379,7 +1400,7 @@ void FHeightfieldLightingViewInfo::ComputeOcclusionForSamples(
 						FSurfaceCacheResources& SurfaceCacheResources = *Scene->SurfaceCacheResources;
 						DispatchIndirectComputeShader(RHICmdList, *ComputeShader, SurfaceCacheResources.DispatchParameters.Buffer, 0);
 
-						ComputeShader->UnsetParameters(RHICmdList);
+						ComputeShader->UnsetParameters(RHICmdList, View, DepthLevel, TemporaryIrradianceCacheResources);
 					}
 				}
 			}
@@ -1455,6 +1476,7 @@ public:
 		SetSRVParameter(RHICmdList, ShaderRHI, SavedStartIndex, SurfaceCacheResources.Level[DepthLevel]->SavedStartIndex.SRV);
 		SetSRVParameter(RHICmdList, ShaderRHI, RecordConeVisibility, TemporaryIrradianceCacheResources.ConeVisibility.SRV);
 
+		RHICmdList.TransitionResource(EResourceTransitionAccess::ERWBarrier, EResourceTransitionPipeline::EComputeToCompute, TemporaryIrradianceCacheResources.HeightfieldIrradiance.UAV);
 		HeightfieldIrradiance.SetBuffer(RHICmdList, ShaderRHI, TemporaryIrradianceCacheResources.HeightfieldIrradiance);
 
 		{
@@ -1493,9 +1515,10 @@ public:
 		SetShaderValue(RHICmdList, ShaderRHI, OuterLightTransferDistanceScale, GHeightfieldOuterBounceDistanceScale);
 	}
 
-	void UnsetParameters(FRHICommandList& RHICmdList)
+	void UnsetParameters(FRHICommandList& RHICmdList, const FTemporaryIrradianceCacheResources& TemporaryIrradianceCacheResources)
 	{
 		HeightfieldIrradiance.UnsetUAV(RHICmdList, GetComputeShader());
+		RHICmdList.TransitionResource(EResourceTransitionAccess::EReadable, EResourceTransitionPipeline::EComputeToCompute, TemporaryIrradianceCacheResources.HeightfieldIrradiance.UAV);
 	}
 
 	virtual bool Serialize(FArchive& Ar) override
@@ -1562,7 +1585,7 @@ void FHeightfieldLightingViewInfo::ComputeIrradianceForSamples(
 			ComputeShader->SetParameters(RHICmdList, View, DepthLevel);
 			DispatchComputeShader(RHICmdList, *ComputeShader, 1, 1, 1);
 
-			ComputeShader->UnsetParameters(RHICmdList);
+			ComputeShader->UnsetParameters(RHICmdList, View);
 		}
 
 		SCOPED_DRAW_EVENT(RHICmdList, HeightfieldIrradiance);
@@ -1595,7 +1618,7 @@ void FHeightfieldLightingViewInfo::ComputeIrradianceForSamples(
 					FSurfaceCacheResources& SurfaceCacheResources = *Scene->SurfaceCacheResources;
 					DispatchIndirectComputeShader(RHICmdList, *ComputeShader, SurfaceCacheResources.DispatchParameters.Buffer, 0);
 
-					ComputeShader->UnsetParameters(RHICmdList);
+					ComputeShader->UnsetParameters(RHICmdList, TemporaryIrradianceCacheResources);
 				}
 			}
 		}
@@ -1658,6 +1681,7 @@ public:
 		HeightfieldDescriptionParameters.Set(RHICmdList, ShaderRHI, NumHeightfieldsValue);
 		HeightfieldTextureParameters.Set(RHICmdList, ShaderRHI, HeightfieldTextureValue, NULL);
 
+		RHICmdList.TransitionResource(EResourceTransitionAccess::ERWBarrier, EResourceTransitionPipeline::EComputeToCompute, ScreenGridResources.ScreenGridConeVisibility.UAV);
 		ScreenGridConeVisibility.SetBuffer(RHICmdList, ShaderRHI, ScreenGridResources.ScreenGridConeVisibility);
 
 		FAOSampleData2 AOSampleData;
@@ -1676,9 +1700,10 @@ public:
 		SetShaderValue(RHICmdList, ShaderRHI, TanConeHalfAngle, FMath::Tan(GAOConeHalfAngle));
 	}
 
-	void UnsetParameters(FRHICommandList& RHICmdList)
+	void UnsetParameters(FRHICommandList& RHICmdList, const FAOScreenGridResources& ScreenGridResources)
 	{
 		ScreenGridConeVisibility.UnsetUAV(RHICmdList, GetComputeShader());
+		RHICmdList.TransitionResource(EResourceTransitionAccess::EReadable, EResourceTransitionPipeline::EComputeToCompute, ScreenGridResources.ScreenGridConeVisibility.UAV);
 	}
 
 	virtual bool Serialize(FArchive& Ar) override
@@ -1739,7 +1764,7 @@ void FHeightfieldLightingViewInfo::ComputeOcclusionForScreenGrid(
 					RHICmdList.SetComputeShader(ComputeShader->GetComputeShader());
 					ComputeShader->SetParameters(RHICmdList, View, HeightfieldTexture, HeightfieldDescriptions.Num(), DistanceFieldNormal, ScreenGridResources, Parameters);
 					DispatchComputeShader(RHICmdList, *ComputeShader, GroupSizeX, GroupSizeY, 1);
-					ComputeShader->UnsetParameters(RHICmdList);
+					ComputeShader->UnsetParameters(RHICmdList, ScreenGridResources);
 				}
 			}
 		}
@@ -1806,6 +1831,7 @@ public:
 		ScreenGridParameters.Set(RHICmdList, ShaderRHI, View, DistanceFieldNormal);
 		HeightfieldDescriptionParameters.Set(RHICmdList, ShaderRHI, NumHeightfieldsValue);
 
+		RHICmdList.TransitionResource(EResourceTransitionAccess::ERWBarrier, EResourceTransitionPipeline::EComputeToCompute, ScreenGridResources.HeightfieldIrradiance.UAV);
 		HeightfieldIrradiance.SetBuffer(RHICmdList, ShaderRHI, ScreenGridResources.HeightfieldIrradiance);
 
 		extern float GAOConeHalfAngle;
@@ -1846,9 +1872,10 @@ public:
 		SetSRVParameter(RHICmdList, ShaderRHI, RecordConeVisibility, ScreenGridResources.ConeDepthVisibilityFunction.SRV);
 	}
 
-	void UnsetParameters(FRHICommandList& RHICmdList)
+	void UnsetParameters(FRHICommandList& RHICmdList, const FAOScreenGridResources& ScreenGridResources)
 	{
 		HeightfieldIrradiance.UnsetUAV(RHICmdList, GetComputeShader());
+		RHICmdList.TransitionResource(EResourceTransitionAccess::EReadable, EResourceTransitionPipeline::EComputeToCompute, ScreenGridResources.HeightfieldIrradiance.UAV);
 	}
 
 	virtual bool Serialize(FArchive& Ar) override
@@ -1933,7 +1960,7 @@ void FHeightfieldLightingViewInfo::ComputeIrradianceForScreenGrid(
 
 					DispatchComputeShader(RHICmdList, *ComputeShader, GroupSizeX, GroupSizeY, 1);
 
-					ComputeShader->UnsetParameters(RHICmdList);
+					ComputeShader->UnsetParameters(RHICmdList, ScreenGridResources);
 				}
 			}
 		}
diff --git a/Engine/Source/Runtime/Renderer/Private/LightPropagationVolume.cpp b/Engine/Source/Runtime/Renderer/Private/LightPropagationVolume.cpp
--- a/Engine/Source/Runtime/Renderer/Private/LightPropagationVolume.cpp
+++ b/Engine/Source/Runtime/Renderer/Private/LightPropagationVolume.cpp
@@ -188,6 +188,10 @@ public:
 	{
 		FComputeShaderRHIParamRef ShaderRHI = GetComputeShader();
 		SetUniformBufferParameter(RHICmdList, ShaderRHI, GetUniformBufferParameter<FLpvWriteUniformBufferParameters>(), Params.UniformBuffer );
+
+		TArray<int32> ResourceIndices;
+		TArray<FUnorderedAccessViewRHIParamRef> UAVs;
+
 		for(int i  =0; i < 7; i++)
 		{
 			if ( LpvBufferSRVParameters[i].IsBound() )
@@ -196,9 +200,9 @@ public:
 			}
 			if ( LpvBufferUAVs[i].IsBound() )
 			{
-				RHICmdList.SetUAVParameter( ShaderRHI, LpvBufferUAVs[i].GetBaseIndex(), Params.LpvBufferUAVs[i] );
+				ResourceIndices.Add(LpvBufferUAVs[i].GetBaseIndex());
+				UAVs.Add(Params.LpvBufferUAVs[i]);
 			}
-
 			SetTextureParameter(RHICmdList, ShaderRHI, LpvBufferSRVParameters[i], LpvVolumeTextureSampler, TStaticSamplerState<SF_Bilinear,AM_Clamp,AM_Clamp,AM_Clamp>::GetRHI(), Params.LpvBufferSRVs[i] );
 		}
 		if ( VplListHeadBufferSRV.IsBound() )
@@ -207,15 +211,17 @@ public:
 		}
 		if ( VplListHeadBufferUAV.IsBound() )
 		{
-			RHICmdList.SetUAVParameter( ShaderRHI, VplListHeadBufferUAV.GetBaseIndex(), Params.VplListHeadBufferUAV );
+			ResourceIndices.Add(VplListHeadBufferUAV.GetBaseIndex());
+			UAVs.Add(Params.VplListHeadBufferUAV);
 		}
 		if ( VplListBufferSRV.IsBound() )
 		{
 			RHICmdList.SetShaderResourceViewParameter( ShaderRHI, VplListBufferSRV.GetBaseIndex(), Params.VplListBufferSRV );
 		}
 		if ( VplListBufferUAV.IsBound() )
-		{
-			RHICmdList.SetUAVParameter( ShaderRHI, VplListBufferUAV.GetBaseIndex(), Params.VplListBufferUAV ); 
+		{			
+			ResourceIndices.Add(VplListBufferUAV.GetBaseIndex());
+			UAVs.Add(Params.VplListBufferUAV);
 		}
 
 		// GV Volume texture
@@ -227,7 +233,8 @@ public:
 			}
 			if ( GvBufferUAVs[i].IsBound() )
 			{
-				RHICmdList.SetUAVParameter( ShaderRHI, GvBufferUAVs[i].GetBaseIndex(), Params.GvBufferUAVs[i] );
+				ResourceIndices.Add(GvBufferUAVs[i].GetBaseIndex());
+				UAVs.Add(Params.GvBufferUAVs[i]);				
 			}
 
 			SetTextureParameter(RHICmdList, ShaderRHI, GvBufferSRVParameters[i], LpvVolumeTextureSampler, TStaticSamplerState<SF_Bilinear,AM_Clamp,AM_Clamp,AM_Clamp>::GetRHI(), Params.GvBufferSRVs[i] );
@@ -235,11 +242,13 @@ public:
 
 		if(GvListBufferUAV.IsBound())
 		{
-			RHICmdList.SetUAVParameter(ShaderRHI, GvListBufferUAV.GetBaseIndex(), Params.GvListBufferUAV); 
+			ResourceIndices.Add(GvListBufferUAV.GetBaseIndex());
+			UAVs.Add(Params.GvListBufferUAV);
 		}
 		if(GvListHeadBufferUAV.IsBound())
 		{
-			RHICmdList.SetUAVParameter(ShaderRHI, GvListHeadBufferUAV.GetBaseIndex(), Params.GvListHeadBufferUAV); 
+			ResourceIndices.Add(GvListHeadBufferUAV.GetBaseIndex());
+			UAVs.Add(Params.GvListHeadBufferUAV);			
 		}
 		if(GvListBufferSRV.IsBound())
 		{
@@ -251,17 +260,28 @@ public:
 		}
 		if ( AOVolumeTextureUAV.IsBound() )
 		{
-			RHICmdList.SetUAVParameter( ShaderRHI, AOVolumeTextureUAV.GetBaseIndex(), Params.AOVolumeTextureUAV ); 
+			ResourceIndices.Add(AOVolumeTextureUAV.GetBaseIndex());
+			UAVs.Add(Params.AOVolumeTextureUAV);			
 		}
 		if ( AOVolumeTextureSRV.IsBound() )
 		{
 			RHICmdList.SetShaderTexture(ShaderRHI, AOVolumeTextureSRV.GetBaseIndex(), Params.AOVolumeTextureSRV );
 		}
+
+		check(ResourceIndices.Num() == UAVs.Num());
+		RHICmdList.TransitionResources(EResourceTransitionAccess::ERWBarrier, EResourceTransitionPipeline::EGfxToCompute, UAVs.GetData(), UAVs.Num());
+		for (int32 i = 0; i < ResourceIndices.Num(); ++i)
+		{
+			RHICmdList.SetUAVParameter(ShaderRHI, ResourceIndices[i], UAVs[i]);
+		}
 	}
 
 	// Unbinds any buffers that have been bound.
-	void UnbindBuffers(FRHICommandList& RHICmdList)
+	void UnbindBuffers(FRHICommandList& RHICmdList, const FLpvBaseWriteShaderParams& Params)
 	{
+		TArray<int32> ResourceIndices;
+		TArray<FUnorderedAccessViewRHIParamRef> UAVs;
+
 		FComputeShaderRHIParamRef ShaderRHI = GetComputeShader();
 		for ( int i = 0; i < 7; i++ )
 		{
@@ -271,7 +291,8 @@ public:
 		    }
 			if ( LpvBufferUAVs[i].IsBound() )
 		    {
-			    RHICmdList.SetUAVParameter( ShaderRHI, LpvBufferUAVs[i].GetBaseIndex(), FUnorderedAccessViewRHIParamRef() );
+				ResourceIndices.Add(LpvBufferUAVs[i].GetBaseIndex());
+				UAVs.Add(Params.LpvBufferUAVs[i]);			    
 		    }
 	    }
 		if ( VplListHeadBufferSRV.IsBound() )
@@ -280,7 +301,8 @@ public:
 		}
 		if ( VplListHeadBufferUAV.IsBound() )
 		{
-			RHICmdList.SetUAVParameter( ShaderRHI, VplListHeadBufferUAV.GetBaseIndex(), FUnorderedAccessViewRHIParamRef() );
+			ResourceIndices.Add(VplListHeadBufferUAV.GetBaseIndex());
+			UAVs.Add(Params.VplListHeadBufferUAV);			
 		}
 		if ( VplListBufferSRV.IsBound() )
 		{
@@ -288,7 +310,8 @@ public:
 		}
 		if ( VplListBufferUAV.IsBound() )
 		{
-			RHICmdList.SetUAVParameter( ShaderRHI, VplListBufferUAV.GetBaseIndex(), FUnorderedAccessViewRHIParamRef() );
+			ResourceIndices.Add(VplListBufferUAV.GetBaseIndex());
+			UAVs.Add(Params.VplListBufferUAV);			
 		}
 		for ( int i = 0; i < NUM_GV_TEXTURES; i++ )
 		{
@@ -298,13 +321,15 @@ public:
 			}
 			if ( GvBufferUAVs[i].IsBound() )
 			{
-					RHICmdList.SetUAVParameter( ShaderRHI, GvBufferUAVs[i].GetBaseIndex(), FUnorderedAccessViewRHIParamRef() );
+				ResourceIndices.Add(GvBufferUAVs[i].GetBaseIndex());
+				UAVs.Add(Params.GvBufferUAVs[i]);
 			}
 		}
 
 		if ( AOVolumeTextureUAV.IsBound() )
 		{
-			RHICmdList.SetUAVParameter( ShaderRHI, AOVolumeTextureUAV.GetBaseIndex(), FUnorderedAccessViewRHIParamRef() );
+			ResourceIndices.Add(AOVolumeTextureUAV.GetBaseIndex());
+			UAVs.Add(Params.AOVolumeTextureUAV);
 		}
 		if ( AOVolumeTextureSRV.IsBound() )
 		{
@@ -312,11 +337,13 @@ public:
 		}
 		if(GvListBufferUAV.IsBound())
 		{
-			RHICmdList.SetUAVParameter(ShaderRHI, GvListBufferUAV.GetBaseIndex(), FUnorderedAccessViewRHIParamRef());
+			ResourceIndices.Add(GvListBufferUAV.GetBaseIndex());
+			UAVs.Add(Params.GvListBufferUAV);			
 		}
 		if(GvListHeadBufferUAV.IsBound())
 		{
-			RHICmdList.SetUAVParameter(ShaderRHI, GvListHeadBufferUAV.GetBaseIndex(), FUnorderedAccessViewRHIParamRef());
+			ResourceIndices.Add(GvListHeadBufferUAV.GetBaseIndex());
+			UAVs.Add(Params.GvListHeadBufferUAV);
 		}
 		if ( GvListBufferSRV.IsBound() )
 		{
@@ -326,6 +353,14 @@ public:
 		{
 			RHICmdList.SetShaderResourceViewParameter( ShaderRHI, GvListHeadBufferSRV.GetBaseIndex(), FShaderResourceViewRHIParamRef() );
 		}
+
+		check(ResourceIndices.Num() == UAVs.Num());
+		RHICmdList.TransitionResources(EResourceTransitionAccess::EReadable, EResourceTransitionPipeline::EComputeToGfx, UAVs.GetData(), UAVs.Num());
+		FUnorderedAccessViewRHIParamRef NullUAV = nullptr;
+		for (int32 i = 0; i < ResourceIndices.Num(); ++i)
+		{
+			RHICmdList.SetUAVParameter(ShaderRHI, ResourceIndices[i], NullUAV);
+		}
 	}
 
 protected:
@@ -470,9 +505,9 @@ public:
 	}
 
 	// Unbinds any buffers that have been bound.
-	void UnbindBuffers(FRHICommandList& RHICmdList)
+	void UnbindBuffers(FRHICommandList& RHICmdList, FLpvBaseWriteShaderParams& BaseParams)
 	{
-		FLpvWriteShaderCSBase::UnbindBuffers(RHICmdList);
+		FLpvWriteShaderCSBase::UnbindBuffers(RHICmdList, BaseParams);
 	}
 
 
@@ -1009,7 +1044,7 @@ void FLightPropagationVolume::Clear(FRHICommandListImmediate& RHICmdList, FViewI
 		GetShaderParams( ShaderParams );
 		Shader->SetParameters(RHICmdList, ShaderParams );
 		DispatchComputeShader(RHICmdList, *Shader, LPV_GRIDRES/4, LPV_GRIDRES/4, LPV_GRIDRES/4 );
-		Shader->UnbindBuffers(RHICmdList);
+		Shader->UnbindBuffers(RHICmdList, ShaderParams );
 	}
 
 	// Clear the LPV (or fade, if REFINE_OVER_TIME is enabled)
@@ -1021,7 +1056,7 @@ void FLightPropagationVolume::Clear(FRHICommandListImmediate& RHICmdList, FViewI
 		GetShaderParams( ShaderParams );
 		Shader->SetParameters(RHICmdList, ShaderParams );
 		DispatchComputeShader(RHICmdList, *Shader, LPV_GRIDRES/4, LPV_GRIDRES/4, LPV_GRIDRES/4 );
-		Shader->UnbindBuffers(RHICmdList);
+		Shader->UnbindBuffers(RHICmdList, ShaderParams);
 	}
 
 	// Clear the geometry volume if necessary
@@ -1034,7 +1069,7 @@ void FLightPropagationVolume::Clear(FRHICommandListImmediate& RHICmdList, FViewI
 		GetShaderParams( ShaderParams );
 		Shader->SetParameters(RHICmdList, ShaderParams );
 		DispatchComputeShader(RHICmdList, *Shader, LPV_GRIDRES/4, LPV_GRIDRES/4, LPV_GRIDRES/4 );
-		Shader->UnbindBuffers(RHICmdList);
+		Shader->UnbindBuffers(RHICmdList, ShaderParams);
 	}
 	RHICmdList.AutomaticCacheFlushAfterComputeShader(true);
 	RHICmdList.FlushComputeShaderCache();
@@ -1124,7 +1159,7 @@ void FLightPropagationVolume::InjectDirectionalLightRSM(
 		// todo: what if not divisble by 8?
 		DispatchComputeShader(RHICmdList, *Shader, RSMResolution / 8, RSMResolution / 8, 1 ); 
 
-		Shader->UnbindBuffers(RHICmdList); 
+		Shader->UnbindBuffers(RHICmdList, ShaderParams);
 	}
 
 	// If this is the first directional light, build the geometry volume with it
@@ -1150,7 +1185,7 @@ void FLightPropagationVolume::InjectDirectionalLightRSM(
 
 		DispatchComputeShader(RHICmdList, *Shader, LPV_GRIDRES/4, LPV_GRIDRES/4, LPV_GRIDRES/4 );
 
-		Shader->UnbindBuffers(RHICmdList);
+		Shader->UnbindBuffers(RHICmdList, ShaderParams);
 	}
 
 	mInjectedLightCount++;
@@ -1176,7 +1211,7 @@ void FLightPropagationVolume::ComputeDirectionalOcclusion( FRHICommandListImmedi
 			LpvWriteUniformBuffer.SetContents( *LpvWriteUniformBufferParams );
 
 			DispatchComputeShader( RHICmdList, *Shader, LPV_GRIDRES/4, LPV_GRIDRES/4, LPV_GRIDRES/4 );
-			Shader->UnbindBuffers(RHICmdList); 
+			Shader->UnbindBuffers(RHICmdList, ShaderParams);
 		}
 	}
 	RHICmdList.FlushComputeShaderCache();
@@ -1220,7 +1255,7 @@ void FLightPropagationVolume::Update( FRHICommandListImmediate& RHICmdList, FVie
 		DispatchComputeShader(RHICmdList, *Shader, LPV_GRIDRES/4, LPV_GRIDRES/4, LPV_GRIDRES/4 );
 		RHICmdList.FlushComputeShaderCache();
 
-		Shader->UnbindBuffers(RHICmdList);
+		Shader->UnbindBuffers(RHICmdList, ShaderParams);
 	}
 
 	// Propagate lighting, ping-ponging between the two buffers
@@ -1265,7 +1300,7 @@ void FLightPropagationVolume::Update( FRHICommandListImmediate& RHICmdList, FVie
 				RHICmdList.FlushComputeShaderCache();
 			}
 
-			Shader->UnbindBuffers(RHICmdList); 
+			Shader->UnbindBuffers(RHICmdList, ShaderParams);
 		}
 	}
 
@@ -1283,7 +1318,7 @@ void FLightPropagationVolume::Update( FRHICommandListImmediate& RHICmdList, FVie
 		Shader->SetParameters( RHICmdList, ShaderParams );
 		LpvWriteUniformBuffer.SetContents( *LpvWriteUniformBufferParams );
 		DispatchComputeShader( RHICmdList, *Shader, LPV_GRIDRES/4, LPV_GRIDRES/4, LPV_GRIDRES/4 );
-		Shader->UnbindBuffers(RHICmdList); 
+		Shader->UnbindBuffers(RHICmdList, ShaderParams);
 	}
 
 	// End the async compute job
@@ -1415,7 +1450,7 @@ void FLightPropagationVolume::InjectLightDirect(FRHICommandListImmediate& RHICmd
 
 		Shader->SetParameters(RHICmdList, ShaderParams, InjectUniformBuffer );
 		DispatchComputeShader(RHICmdList, Shader, LPV_GRIDRES / 4, LPV_GRIDRES / 4, LPV_GRIDRES / 4 );
-		Shader->UnbindBuffers(RHICmdList); 
+		Shader->UnbindBuffers(RHICmdList, ShaderParams);
 	}
 }
 
diff --git a/Engine/Source/Runtime/Renderer/Private/PostProcess/PostProcessBloomSetup.cpp b/Engine/Source/Runtime/Renderer/Private/PostProcess/PostProcessBloomSetup.cpp
--- a/Engine/Source/Runtime/Renderer/Private/PostProcess/PostProcessBloomSetup.cpp
+++ b/Engine/Source/Runtime/Renderer/Private/PostProcess/PostProcessBloomSetup.cpp
@@ -207,7 +207,7 @@ FPooledRenderTargetDesc FRCPassPostProcessBloomSetup::ComputeOutputDesc(EPassOut
 
 	Ret.Reset();
 	Ret.DebugName = TEXT("BloomSetup");
-
+	Ret.AutoWritable = false;
 	return Ret;
 }
 
diff --git a/Engine/Source/Runtime/Renderer/Private/PostProcess/PostProcessHistogram.cpp b/Engine/Source/Runtime/Renderer/Private/PostProcess/PostProcessHistogram.cpp
--- a/Engine/Source/Runtime/Renderer/Private/PostProcess/PostProcessHistogram.cpp
+++ b/Engine/Source/Runtime/Renderer/Private/PostProcess/PostProcessHistogram.cpp
@@ -116,6 +116,7 @@ void FRCPassPostProcessHistogram::Process(FRenderingCompositePassContext& Contex
 
 	// set destination
 	check(DestRenderTarget.UAV);
+	Context.RHICmdList.TransitionResource(EResourceTransitionAccess::ERWBarrier, EResourceTransitionPipeline::EGfxToCompute, DestRenderTarget.UAV);
 	Context.RHICmdList.SetUAVParameter(ComputeShader->GetComputeShader(), ComputeShader->HistogramRWTexture.GetBaseIndex(), DestRenderTarget.UAV);
 
 	FIntPoint GatherExtent = ComputeGatherExtent(View);
@@ -128,7 +129,8 @@ void FRCPassPostProcessHistogram::Process(FRenderingCompositePassContext& Contex
 	// un-set destination
 	Context.RHICmdList.SetUAVParameter(ComputeShader->GetComputeShader(), ComputeShader->HistogramRWTexture.GetBaseIndex(), NULL);
 
-	Context.RHICmdList.CopyToResolveTarget(DestRenderTarget.TargetableTexture, DestRenderTarget.ShaderResourceTexture, false, FResolveParams());
+	Context.RHICmdList.TransitionResource(EResourceTransitionAccess::EReadable, EResourceTransitionPipeline::EComputeToGfx, DestRenderTarget.UAV);
+	ensureMsgf(DestRenderTarget.TargetableTexture == DestRenderTarget.ShaderResourceTexture, TEXT("%s should be resolved to a separate SRV"), *DestRenderTarget.TargetableTexture->GetName().ToString());	
 }
 
 
diff --git a/Engine/Source/Runtime/Renderer/Private/PostProcess/PostProcessMotionBlur.cpp b/Engine/Source/Runtime/Renderer/Private/PostProcess/PostProcessMotionBlur.cpp
--- a/Engine/Source/Runtime/Renderer/Private/PostProcess/PostProcessMotionBlur.cpp
+++ b/Engine/Source/Runtime/Renderer/Private/PostProcess/PostProcessMotionBlur.cpp
@@ -540,6 +540,7 @@ FPooledRenderTargetDesc FRCPassPostProcessMotionBlur::ComputeOutputDesc(EPassOut
 
 	Ret.Reset();
 	Ret.DebugName = TEXT("MotionBlur");
+	Ret.AutoWritable = false;
 
 	return Ret;
 }
@@ -1422,6 +1423,7 @@ FPooledRenderTargetDesc FRCPassPostProcessMotionBlurNew::ComputeOutputDesc(EPass
 	// we don't need the alpha channel and 32bit is faster and costs less memory
 	Ret.Format = PF_FloatRGB;
 	Ret.DebugName = TEXT("MotionBlur");
+	Ret.AutoWritable = false;
 
 	return Ret;
 }
@@ -1579,6 +1581,7 @@ FPooledRenderTargetDesc FRCPassPostProcessVisualizeMotionBlur::ComputeOutputDesc
 
 	Ret.Reset();
 	Ret.DebugName = TEXT("MotionBlur");
+	Ret.AutoWritable = false;
 
 	return Ret;
 }
diff --git a/Engine/Source/Runtime/Renderer/Private/PostProcess/PostProcessWeightedSampleSum.cpp b/Engine/Source/Runtime/Renderer/Private/PostProcess/PostProcessWeightedSampleSum.cpp
--- a/Engine/Source/Runtime/Renderer/Private/PostProcess/PostProcessWeightedSampleSum.cpp
+++ b/Engine/Source/Runtime/Renderer/Private/PostProcess/PostProcessWeightedSampleSum.cpp
@@ -656,6 +656,7 @@ FPooledRenderTargetDesc FRCPassPostProcessWeightedSampleSum::ComputeOutputDesc(E
 
 	Ret.Reset();
 	Ret.DebugName = DebugName;
+	Ret.AutoWritable = false;
 
 	return Ret;
 }
diff --git a/Engine/Source/Runtime/Renderer/Private/PostProcess/SceneRenderTargets.cpp b/Engine/Source/Runtime/Renderer/Private/PostProcess/SceneRenderTargets.cpp
--- a/Engine/Source/Runtime/Renderer/Private/PostProcess/SceneRenderTargets.cpp
+++ b/Engine/Source/Runtime/Renderer/Private/PostProcess/SceneRenderTargets.cpp
@@ -1134,7 +1134,7 @@ void FSceneRenderTargets::BeginRenderingShadowDepth(FRHICommandList& RHICmdList,
 void FSceneRenderTargets::BeginRenderingCubeShadowDepth(FRHICommandList& RHICmdList, int32 ShadowResolution)
 {
 	SCOPED_DRAW_EVENT(RHICmdList, BeginRenderingCubeShadowDepth);
-	SetRenderTarget(RHICmdList, FTextureRHIRef(), GetCubeShadowDepthZSurface(ShadowResolution));
+	SetRenderTarget(RHICmdList, FTextureRHIRef(), GetCubeShadowDepthZSurface(ShadowResolution), true);
 }
 
 void FSceneRenderTargets::FinishRenderingShadowDepth(FRHICommandList& RHICmdList, const FResolveRect& ResolveRect)
@@ -1156,16 +1156,22 @@ void FSceneRenderTargets::BeginRenderingReflectiveShadowMap(FRHICommandList& RHI
 	Uavs[2] = Lpv->GetVplListBufferUav();
 	Uavs[3] = Lpv->GetVplListHeadBufferUav();
 	
+	RHICmdList.TransitionResources(EResourceTransitionAccess::ERWBarrier, EResourceTransitionPipeline::EGfxToGfx, Uavs, ARRAY_COUNT(Uavs));
 	SetRenderTargets(RHICmdList, ARRAY_COUNT(RenderTargets), RenderTargets, GetReflectiveShadowMapDepthSurface(), 4, Uavs);
 }
 
-void FSceneRenderTargets::FinishRenderingReflectiveShadowMap(FRHICommandList& RHICmdList, const FResolveRect& ResolveRect)
+void FSceneRenderTargets::FinishRenderingReflectiveShadowMap(FRHICommandList& RHICmdList, FLightPropagationVolume* Lpv, const FResolveRect& ResolveRect)
 {
 	// Resolve the shadow depth z surface.
 	RHICmdList.CopyToResolveTarget(GetReflectiveShadowMapDepthSurface(), GetReflectiveShadowMapDepthTexture(), false, FResolveParams(ResolveRect));
 	RHICmdList.CopyToResolveTarget(GetReflectiveShadowMapDiffuseSurface(), GetReflectiveShadowMapDiffuseTexture(), false, FResolveParams(ResolveRect));
 	RHICmdList.CopyToResolveTarget(GetReflectiveShadowMapNormalSurface(), GetReflectiveShadowMapNormalTexture(), false, FResolveParams(ResolveRect));
 
+	FUnorderedAccessViewRHIParamRef UavsToReadable[2];
+	UavsToReadable[0] = Lpv->GetGvListBufferUav();
+	UavsToReadable[1] = Lpv->GetGvListHeadBufferUav();	
+	RHICmdList.TransitionResources(EResourceTransitionAccess::EReadable, EResourceTransitionPipeline::EGfxToGfx, UavsToReadable, ARRAY_COUNT(UavsToReadable));
+
 	// Unset render targets
 	FTextureRHIParamRef RenderTargets[2] = {NULL};
 	FUnorderedAccessViewRHIParamRef Uavs[2] = {NULL};
@@ -1586,6 +1592,7 @@ void FSceneRenderTargets::AllocateCommonDepthTargets(FRHICommandList& RHICmdList
 	if (!AuxiliarySceneDepthZ && !GSupportsDepthFetchDuringDepthTest)
 	{
 		FPooledRenderTargetDesc Desc(FPooledRenderTargetDesc::Create2DDesc(BufferSize, PF_DepthStencil, FClearValueBinding::DepthFar, TexCreate_None, TexCreate_DepthStencilTargetable, false));
+		Desc.AutoWritable = false;
 		GRenderTargetPool.FindFreeElement(RHICmdList, Desc, AuxiliarySceneDepthZ, TEXT("AuxiliarySceneDepthZ"));
 	}
 }
diff --git a/Engine/Source/Runtime/Renderer/Private/PostProcess/SceneRenderTargets.h b/Engine/Source/Runtime/Renderer/Private/PostProcess/SceneRenderTargets.h
--- a/Engine/Source/Runtime/Renderer/Private/PostProcess/SceneRenderTargets.h
+++ b/Engine/Source/Runtime/Renderer/Private/PostProcess/SceneRenderTargets.h
@@ -250,7 +250,7 @@ public:
 	void FinishRenderingShadowDepth(FRHICommandList& RHICmdList, const FResolveRect& ResolveRect = FResolveRect());
 
 	void BeginRenderingReflectiveShadowMap(FRHICommandList& RHICmdList, class FLightPropagationVolume* Lpv);
-	void FinishRenderingReflectiveShadowMap(FRHICommandList& RHICmdList, const FResolveRect& ResolveRect = FResolveRect());
+	void FinishRenderingReflectiveShadowMap(FRHICommandList& RHICmdList, class FLightPropagationVolume* Lpv, const FResolveRect& ResolveRect = FResolveRect());
 
 	/** Resolves the appropriate shadow depth cube map and restores default state. */
 	void FinishRenderingCubeShadowDepth(FRHICommandList& RHICmdList, int32 ShadowResolution);
diff --git a/Engine/Source/Runtime/Renderer/Private/PostProcess/VisualizeTexture.cpp b/Engine/Source/Runtime/Renderer/Private/PostProcess/VisualizeTexture.cpp
--- a/Engine/Source/Runtime/Renderer/Private/PostProcess/VisualizeTexture.cpp
+++ b/Engine/Source/Runtime/Renderer/Private/PostProcess/VisualizeTexture.cpp
@@ -356,7 +356,7 @@ void FVisualizeTexture::GenerateContent(FRHICommandListImmediate& RHICmdList, co
 
 	const FSceneRenderTargetItem& DestRenderTarget = VisualizeTextureContent->GetRenderTargetItem();
 
-	SetRenderTarget(RHICmdList, DestRenderTarget.TargetableTexture, FTextureRHIRef());
+	SetRenderTarget(RHICmdList, DestRenderTarget.TargetableTexture, FTextureRHIRef(), true);
 	RHICmdList.Clear(true, FLinearColor(1, 1, 0, 1), false, 0.0f, false, 0, FIntRect());
 	RHICmdList.SetBlendState(TStaticBlendState<>::GetRHI());
 	RHICmdList.SetRasterizerState(TStaticRasterizerState<>::GetRHI());
diff --git a/Engine/Source/Runtime/Renderer/Private/ReflectionEnvironment.cpp b/Engine/Source/Runtime/Renderer/Private/ReflectionEnvironment.cpp
--- a/Engine/Source/Runtime/Renderer/Private/ReflectionEnvironment.cpp
+++ b/Engine/Source/Runtime/Renderer/Private/ReflectionEnvironment.cpp
@@ -288,6 +288,8 @@ public:
 		SetTextureParameter(RHICmdList, ShaderRHI, ScreenSpaceReflections, SSRTexture );
 
 		SetTextureParameter(RHICmdList, ShaderRHI, InSceneColor, FSceneRenderTargets::Get(RHICmdList).GetSceneColor()->GetRenderTargetItem().ShaderResourceTexture );
+
+		RHICmdList.TransitionResource(EResourceTransitionAccess::ERWBarrier, EResourceTransitionPipeline::EGfxToCompute, OutSceneColorUAV);
 		OutSceneColor.SetTexture(RHICmdList, ShaderRHI, NULL, OutSceneColorUAV);
 
 		SetShaderValue(RHICmdList, ShaderRHI, ViewDimensionsParameter, View.ViewRect);
@@ -313,9 +315,10 @@ public:
 		SpecularOcclusionParameters.SetParameters(RHICmdList, ShaderRHI, DynamicBentNormalAO, CVarSkySpecularOcclusionStrength.GetValueOnRenderThread(), MinOcclusion);
 	}
 
-	void UnsetParameters(FRHICommandList& RHICmdList)
+	void UnsetParameters(FRHICommandList& RHICmdList, FUnorderedAccessViewRHIParamRef OutSceneColorUAV)
 	{
 		const FComputeShaderRHIParamRef ShaderRHI = GetComputeShader();
+		RHICmdList.TransitionResource(EResourceTransitionAccess::EReadable, EResourceTransitionPipeline::EComputeToGfx, OutSceneColorUAV);
 		OutSceneColor.UnsetUAV(RHICmdList, ShaderRHI);
 	}
 
@@ -896,13 +899,14 @@ void FDeferredShadingSceneRenderer::RenderTiledDeferredImageBasedReflections(FRH
 
 			RHICmdList.SetComputeShader(ComputeShader->GetComputeShader());
 
-			ComputeShader->SetParameters(RHICmdList, View, SSROutput->GetRenderTargetItem().ShaderResourceTexture, SortData, NewSceneColor->GetRenderTargetItem().UAV, DynamicBentNormalAO);
+			FUnorderedAccessViewRHIParamRef OutUAV = NewSceneColor->GetRenderTargetItem().UAV;
+			ComputeShader->SetParameters(RHICmdList, View, SSROutput->GetRenderTargetItem().ShaderResourceTexture, SortData, OutUAV, DynamicBentNormalAO);
 
 			uint32 GroupSizeX = (View.ViewRect.Size().X + GReflectionEnvironmentTileSizeX - 1) / GReflectionEnvironmentTileSizeX;
 			uint32 GroupSizeY = (View.ViewRect.Size().Y + GReflectionEnvironmentTileSizeY - 1) / GReflectionEnvironmentTileSizeY;
 			DispatchComputeShader(RHICmdList, ComputeShader, GroupSizeX, GroupSizeY, 1);
 
-			ComputeShader->UnsetParameters(RHICmdList);
+			ComputeShader->UnsetParameters(RHICmdList, OutUAV);
 		}
 	}
 
diff --git a/Engine/Source/Runtime/Renderer/Private/SceneCaptureRendering.cpp b/Engine/Source/Runtime/Renderer/Private/SceneCaptureRendering.cpp
--- a/Engine/Source/Runtime/Renderer/Private/SceneCaptureRendering.cpp
+++ b/Engine/Source/Runtime/Renderer/Private/SceneCaptureRendering.cpp
@@ -116,7 +116,7 @@ static void UpdateSceneCaptureContent_RenderThread(FRHICommandListImmediate& RHI
 		FViewInfo& View = SceneRenderer->Views[0];
 		FIntRect ViewRect = View.ViewRect;
 		FIntRect UnconstrainedViewRect = View.UnconstrainedViewRect;
-		SetRenderTarget(RHICmdList, Target->GetRenderTargetTexture(), NULL);
+		SetRenderTarget(RHICmdList, Target->GetRenderTargetTexture(), NULL, true);
 		RHICmdList.Clear(true, FLinearColor::Black, false, (float)ERHIZBuffer::FarPlane, false, 0, ViewRect);
 
 		// Render the scene normally
diff --git a/Engine/Source/Runtime/Renderer/Private/SceneOcclusion.cpp b/Engine/Source/Runtime/Renderer/Private/SceneOcclusion.cpp
--- a/Engine/Source/Runtime/Renderer/Private/SceneOcclusion.cpp
+++ b/Engine/Source/Runtime/Renderer/Private/SceneOcclusion.cpp
@@ -1027,13 +1027,13 @@ void BuildHZB( FRHICommandListImmediate& RHICmdList, FViewInfo& View )
 			EDRF_UseTriangleOptimization);
 
 		//Use RWBarrier since we don't transition individual subresources.  Basically treat the whole texture as R/W as we walk down the mip chain.
-		RHICmdList.TransitionResources(EResourceTransitionAccess::ERWBarrier, &HZBRenderTargetRef, 1);
+		RHICmdList.TransitionResources(EResourceTransitionAccess::ERWSubResBarrier, &HZBRenderTargetRef, 1);
 	}
 
 	FIntPoint SrcSize = HZBSize;
 	FIntPoint DstSize = SrcSize / 2;
 	
-	SCOPED_DRAW_EVENTF(RHICmdList, BuildHZB, TEXT("HZB SetupMips 1..%d %dx%d Mips:%d"), NumMips - 1, DstSize.X, DstSize.Y);
+	SCOPED_DRAW_EVENTF(RHICmdList, BuildHZB, TEXT("HZB SetupMips 1..%d %dx%d Mips:%d"), NumMips - 1, DstSize.X, DstSize.Y);	
 
 	// Downsampling...
 	for( uint8 MipIndex = 1; MipIndex < NumMips; MipIndex++ )
@@ -1071,8 +1071,8 @@ void BuildHZB( FRHICommandListImmediate& RHICmdList, FViewInfo& View )
 		SrcSize /= 2;
 		DstSize /= 2;
 
-		//Use RWBarrier since we don't transition individual subresources.  Basically treat the whole texture as R/W as we walk down the mip chain.
-		RHICmdList.TransitionResources(EResourceTransitionAccess::ERWBarrier, &HZBRenderTargetRef, 1);
+		//Use ERWSubResBarrier since we don't transition individual subresources.  Basically treat the whole texture as R/W as we walk down the mip chain.
+		RHICmdList.TransitionResources(EResourceTransitionAccess::ERWSubResBarrier, &HZBRenderTargetRef, 1);
 	}
 
 	GRenderTargetPool.VisualizeTexture.SetCheckPoint( RHICmdList, View.HZB );
diff --git a/Engine/Source/Runtime/Renderer/Private/ScenePrivate.h b/Engine/Source/Runtime/Renderer/Private/ScenePrivate.h
--- a/Engine/Source/Runtime/Renderer/Private/ScenePrivate.h
+++ b/Engine/Source/Runtime/Renderer/Private/ScenePrivate.h
@@ -692,6 +692,7 @@ public:
 		{
 			// Create the SeparateTranslucency render target (alpha is needed to lerping)
 			FPooledRenderTargetDesc Desc(FPooledRenderTargetDesc::Create2DDesc(Size, PF_FloatRGBA, FClearValueBinding::Black, TexCreate_None, TexCreate_RenderTargetable, false));
+			Desc.AutoWritable = false;
 			GRenderTargetPool.FindFreeElement(RHICmdList, Desc, SeparateTranslucencyRT, TEXT("SeparateTranslucency"));
 		}
 		return SeparateTranslucencyRT;
diff --git a/Engine/Source/Runtime/Renderer/Private/ShadowRendering.cpp b/Engine/Source/Runtime/Renderer/Private/ShadowRendering.cpp
--- a/Engine/Source/Runtime/Renderer/Private/ShadowRendering.cpp
+++ b/Engine/Source/Runtime/Renderer/Private/ShadowRendering.cpp
@@ -3312,7 +3312,7 @@ bool FDeferredShadingSceneRenderer::RenderReflectiveShadowMaps(FRHICommandListIm
 					ProjectedShadowInfo->RenderDepth(RHICmdList, this, SetShadowRenderTargets, ShadowDepthRenderMode_GIBlockingVolumes);
 				}
 
-				SceneContext.FinishRenderingReflectiveShadowMap(RHICmdList);
+				SceneContext.FinishRenderingReflectiveShadowMap(RHICmdList, LightPropagationVolume);
 			}
 		}
 
diff --git a/Engine/Source/Runtime/Renderer/Private/SurfelTree.cpp b/Engine/Source/Runtime/Renderer/Private/SurfelTree.cpp
--- a/Engine/Source/Runtime/Renderer/Private/SurfelTree.cpp
+++ b/Engine/Source/Runtime/Renderer/Private/SurfelTree.cpp
@@ -238,9 +238,11 @@ public:
 		SetSRVParameter(RHICmdList, ShaderRHI, TriangleCDFs, UniformMeshBuffers.TriangleCDFs.SRV);
 	}
 
-	void UnsetParameters(FRHICommandList& RHICmdList)
+	void UnsetParameters(FRHICommandList& RHICmdList, const FSceneView& View)
 	{
 		FComputeShaderRHIParamRef ShaderRHI = GetComputeShader();
+
+		const FScene* Scene = (const FScene*)View.Family->Scene;
 		SurfelBufferParameters.UnsetParameters(RHICmdList, ShaderRHI);
 		// RHISetStreamOutTargets doesn't unbind existing uses like render targets do 
 		SetSRVParameter(RHICmdList, ShaderRHI, TriangleVertexData, NULL);
@@ -295,6 +297,11 @@ void GenerateSurfelRepresentation(FRHICommandListImmediate& RHICmdList, FSceneRe
 
 		if (NumUniformTriangles > 0 && MaterialRenderProxy && PrimitiveUniformBuffer)
 		{
+			FUnorderedAccessViewRHIParamRef UniformMeshUAVs[2];
+			UniformMeshUAVs[0] = UniformMeshBuffers->TriangleAreas.UAV;
+			UniformMeshUAVs[1] = UniformMeshBuffers->TriangleCDFs.UAV;
+			RHICmdList.TransitionResources(EResourceTransitionAccess::ERWBarrier, EResourceTransitionPipeline::EComputeToCompute, UniformMeshUAVs, ARRAY_COUNT(UniformMeshUAVs));
+
 			{
 				TShaderMapRef<FComputeTriangleAreasCS> ComputeShader(GetGlobalShaderMap(View.GetFeatureLevel()));
 
@@ -303,7 +310,8 @@ void GenerateSurfelRepresentation(FRHICommandListImmediate& RHICmdList, FSceneRe
 				DispatchComputeShader(RHICmdList, *ComputeShader, FMath::DivideAndRoundUp(NumUniformTriangles, GDistanceFieldAOTileSizeX * GDistanceFieldAOTileSizeY), 1, 1);
 
 				ComputeShader->UnsetParameters(RHICmdList);
-			}
+				RHICmdList.TransitionResource(EResourceTransitionAccess::EReadable, EResourceTransitionPipeline::EComputeToCompute, UniformMeshBuffers->TriangleAreas.UAV);
+			}			
 
 			{
 				TShaderMapRef<FComputeTriangleCDFsCS> ComputeShader(GetGlobalShaderMap(View.GetFeatureLevel()));
@@ -313,7 +321,8 @@ void GenerateSurfelRepresentation(FRHICommandListImmediate& RHICmdList, FSceneRe
 				DispatchComputeShader(RHICmdList, *ComputeShader, FMath::DivideAndRoundUp(NumUniformTriangles, GDistanceFieldAOTileSizeX * GDistanceFieldAOTileSizeY), 1, 1);
 
 				ComputeShader->UnsetParameters(RHICmdList);
-			}
+				RHICmdList.TransitionResource(EResourceTransitionAccess::EReadable, EResourceTransitionPipeline::EComputeToCompute, UniformMeshBuffers->TriangleCDFs.UAV);
+			}			
 
 			{
 				TShaderMapRef<FSampleTrianglesCS> ComputeShader(GetGlobalShaderMap(View.GetFeatureLevel()));
@@ -322,7 +331,7 @@ void GenerateSurfelRepresentation(FRHICommandListImmediate& RHICmdList, FSceneRe
 				ComputeShader->SetParameters(RHICmdList, View, Allocation.Offset, Allocation.NumLOD0, NumUniformTriangles, *UniformMeshBuffers);
 				DispatchComputeShader(RHICmdList, *ComputeShader, FMath::DivideAndRoundUp(Allocation.NumLOD0, GDistanceFieldAOTileSizeX * GDistanceFieldAOTileSizeY), 1, 1);
 
-				ComputeShader->UnsetParameters(RHICmdList);
+				ComputeShader->UnsetParameters(RHICmdList, View);
 			}
 
 			FUniformMeshConverter::GenerateSurfels(RHICmdList, View, PrimitiveSceneInfo, MaterialRenderProxy, PrimitiveUniformBuffer, Instance0Transform, Allocation.Offset, Allocation.NumLOD0);
@@ -338,7 +347,7 @@ void GenerateSurfelRepresentation(FRHICommandListImmediate& RHICmdList, FSceneRe
 					ComputeShader->SetParameters(RHICmdList, View, Allocation.Offset + Allocation.NumLOD0, NumLOD1, NumUniformTriangles, *UniformMeshBuffers);
 					DispatchComputeShader(RHICmdList, *ComputeShader, FMath::DivideAndRoundUp(NumLOD1, GDistanceFieldAOTileSizeX * GDistanceFieldAOTileSizeY), 1, 1);
 
-					ComputeShader->UnsetParameters(RHICmdList);
+					ComputeShader->UnsetParameters(RHICmdList, View);
 				}
 
 				FUniformMeshConverter::GenerateSurfels(RHICmdList, View, PrimitiveSceneInfo, MaterialRenderProxy, PrimitiveUniformBuffer, Instance0Transform, Allocation.Offset + Allocation.NumLOD0, NumLOD1);
diff --git a/Engine/Source/Runtime/Renderer/Private/TiledDeferredLightRendering.cpp b/Engine/Source/Runtime/Renderer/Private/TiledDeferredLightRendering.cpp
--- a/Engine/Source/Runtime/Renderer/Private/TiledDeferredLightRendering.cpp
+++ b/Engine/Source/Runtime/Renderer/Private/TiledDeferredLightRendering.cpp
@@ -123,7 +123,10 @@ public:
 		FGlobalShader::SetParameters(RHICmdList, ShaderRHI, View);
 		DeferredParameters.Set(RHICmdList, ShaderRHI, View);
 		SetTextureParameter(RHICmdList, ShaderRHI, InTexture, InTextureValue.GetRenderTargetItem().ShaderResourceTexture);
-		OutTexture.SetTexture(RHICmdList, ShaderRHI, 0, OutTextureValue.GetRenderTargetItem().UAV);
+
+		FUnorderedAccessViewRHIParamRef OutUAV = OutTextureValue.GetRenderTargetItem().UAV;
+		RHICmdList.TransitionResources(EResourceTransitionAccess::ERWBarrier, EResourceTransitionPipeline::EGfxToCompute, &OutUAV, 1);
+		OutTexture.SetTexture(RHICmdList, ShaderRHI, 0, OutUAV);
 
 		SetShaderValue(RHICmdList, ShaderRHI, ViewDimensions, View.ViewRect);
 
@@ -229,9 +232,12 @@ public:
 		SetShaderValue(RHICmdList, ShaderRHI, NumLights, NumThisPass);
 	}
 
-	void UnsetParameters(FRHICommandList& RHICmdList)
+	void UnsetParameters(FRHICommandList& RHICmdList, IPooledRenderTarget& OutTextureValue)
 	{
 		OutTexture.UnsetUAV(RHICmdList, GetComputeShader());
+
+		FUnorderedAccessViewRHIParamRef OutUAV = OutTextureValue.GetRenderTargetItem().UAV;
+		RHICmdList.TransitionResources(EResourceTransitionAccess::EReadable, EResourceTransitionPipeline::EComputeToGfx, &OutUAV, 1);
 	}
 
 	virtual bool Serialize(FArchive& Ar) override
@@ -311,7 +317,7 @@ static void SetShaderTemplTiledLighting(
 	uint32 GroupSizeY = (View.ViewRect.Size().Y + GDeferredLightTileSizeY - 1) / GDeferredLightTileSizeY;
 	DispatchComputeShader(RHICmdList, *ComputeShader, GroupSizeX, GroupSizeY, 1);
 
-	ComputeShader->UnsetParameters(RHICmdList);
+	ComputeShader->UnsetParameters(RHICmdList, OutTexture);
 }
 
 
diff --git a/Engine/Source/Runtime/UtilityShaders/Private/ClearReplacementShaders.cpp b/Engine/Source/Runtime/UtilityShaders/Private/ClearReplacementShaders.cpp
--- a/Engine/Source/Runtime/UtilityShaders/Private/ClearReplacementShaders.cpp
+++ b/Engine/Source/Runtime/UtilityShaders/Private/ClearReplacementShaders.cpp
@@ -17,20 +17,39 @@ void FClearTexture2DReplacementCS::SetParameters( FRHICommandList& RHICmdList, F
 {
 	FComputeShaderRHIParamRef ComputeShaderRHI = GetComputeShader();
 	SetShaderValue(RHICmdList, ComputeShaderRHI, ClearColor, Value);
+
+	RHICmdList.TransitionResource(EResourceTransitionAccess::ERWBarrier, EResourceTransitionPipeline::EGfxToCompute, TextureRW);
 	RHICmdList.SetUAVParameter(ComputeShaderRHI, ClearTextureRW.GetBaseIndex(), TextureRW);
 }
 
+void FClearTexture2DReplacementCS::FinalizeParameters(FRHICommandList& RHICmdList, FUnorderedAccessViewRHIParamRef TextureRW)
+{
+	RHICmdList.TransitionResource(EResourceTransitionAccess::EReadable, EResourceTransitionPipeline::EComputeToGfx, TextureRW);
+}
+
 void FClearTexture2DReplacementScissorCS::SetParameters(FRHICommandList& RHICmdList, FUnorderedAccessViewRHIParamRef TextureRW, FLinearColor InClearColor, const FVector4& InTargetBounds)
 {
 	FComputeShaderRHIParamRef ComputeShaderRHI = GetComputeShader();
 	SetShaderValue(RHICmdList, ComputeShaderRHI, ClearColor, InClearColor);
 	SetShaderValue(RHICmdList, ComputeShaderRHI, TargetBounds, InTargetBounds);
+	RHICmdList.TransitionResource(EResourceTransitionAccess::ERWBarrier, EResourceTransitionPipeline::EGfxToCompute, TextureRW);
 	RHICmdList.SetUAVParameter(ComputeShaderRHI, ClearTextureRW.GetBaseIndex(), TextureRW);
 }
 
+void FClearTexture2DReplacementScissorCS::FinalizeParameters(FRHICommandList& RHICmdList, FUnorderedAccessViewRHIParamRef TextureRW)
+{
+	RHICmdList.TransitionResource(EResourceTransitionAccess::EReadable, EResourceTransitionPipeline::EComputeToGfx, TextureRW);
+}
+
 void FClearBufferReplacementCS::SetParameters( FRHICommandList& RHICmdList, FUnorderedAccessViewRHIParamRef BufferRW, uint32 Dword )
 {
 	FComputeShaderRHIParamRef ComputeShaderRHI = GetComputeShader();
 	SetShaderValue(RHICmdList, ComputeShaderRHI, ClearDword, Dword);
+	RHICmdList.TransitionResource(EResourceTransitionAccess::ERWBarrier, EResourceTransitionPipeline::EGfxToCompute, BufferRW);
 	RHICmdList.SetUAVParameter(ComputeShaderRHI, ClearBufferRW.GetBaseIndex(), BufferRW);
 }
+
+void FClearBufferReplacementCS::FinalizeParameters(FRHICommandList& RHICmdList, FUnorderedAccessViewRHIParamRef BufferRW)
+{
+	RHICmdList.TransitionResource(EResourceTransitionAccess::EReadable, EResourceTransitionPipeline::EComputeToGfx, BufferRW);
+}
diff --git a/Engine/Source/Runtime/UtilityShaders/Public/ClearReplacementShaders.h b/Engine/Source/Runtime/UtilityShaders/Public/ClearReplacementShaders.h
--- a/Engine/Source/Runtime/UtilityShaders/Public/ClearReplacementShaders.h
+++ b/Engine/Source/Runtime/UtilityShaders/Public/ClearReplacementShaders.h
@@ -84,6 +84,7 @@ public:
 	}
 	
 	void SetParameters(FRHICommandList& RHICmdList, FUnorderedAccessViewRHIParamRef TextureRW, FLinearColor Value);
+	void FinalizeParameters(FRHICommandList& RHICmdList, FUnorderedAccessViewRHIParamRef TextureRW);
 	
 	static bool ShouldCache(EShaderPlatform Platform)
 	{
@@ -127,6 +128,7 @@ public:
 	}
 
 	void SetParameters(FRHICommandList& RHICmdList, FUnorderedAccessViewRHIParamRef TextureRW, FLinearColor ClearColor, const FVector4& InTargetBounds);
+	void FinalizeParameters(FRHICommandList& RHICmdList, FUnorderedAccessViewRHIParamRef TextureRW);
 
 	static bool ShouldCache(EShaderPlatform Platform)
 	{
@@ -167,6 +169,7 @@ public:
 	}
 	
 	void SetParameters(FRHICommandList& RHICmdList, FUnorderedAccessViewRHIParamRef BufferRW, uint32 Dword);
+	void FinalizeParameters(FRHICommandList& RHICmdList, FUnorderedAccessViewRHIParamRef BufferRW);
 	
 	// FShader interface.
 	virtual bool Serialize(FArchive& Ar) override
diff --git a/Engine/Source/Runtime/Windows/D3D11RHI/Private/D3D11Commands.cpp b/Engine/Source/Runtime/Windows/D3D11RHI/Private/D3D11Commands.cpp
--- a/Engine/Source/Runtime/Windows/D3D11RHI/Private/D3D11Commands.cpp
+++ b/Engine/Source/Runtime/Windows/D3D11RHI/Private/D3D11Commands.cpp
@@ -365,7 +365,14 @@ void FD3D11DynamicRHI::RHISetUAVParameter(FComputeShaderRHIParamRef ComputeShade
 	if(UAV)
 	{
 		ConditionalClearShaderResource(UAV->Resource);		
-		UAV->Resource->SetDirty();
+
+		//check it's safe for r/w for this UAV
+		const EResourceTransitionAccess CurrentUAVAccess = UAV->Resource->GetCurrentGPUAccess();
+		const bool UAVDirty = UAV->Resource->IsDirty();
+		ensureMsgf(!UAVDirty || (CurrentUAVAccess == EResourceTransitionAccess::ERWNoBarrier), TEXT("UAV: %i is in unsafe state for GPU R/W: %s, Dirty: %i"), UAVIndex, *FResourceTransitionUtility::ResourceTransitionAccessStrings[(int32)CurrentUAVAccess], (int32)UAVDirty);
+
+		//UAVs always dirty themselves. If a shader wanted to just read, it should use an SRV.
+		UAV->Resource->SetDirty(true, PresentCounter);
 	}
 
 	ID3D11UnorderedAccessView* D3D11UAV = UAV ? UAV->View : NULL;
@@ -383,7 +390,14 @@ void FD3D11DynamicRHI::RHISetUAVParameter(FComputeShaderRHIParamRef ComputeShade
 	if(UAV)
 	{
 		ConditionalClearShaderResource(UAV->Resource);
-		UAV->Resource->SetDirty();
+
+		//check it's safe for r/w for this UAV
+		const EResourceTransitionAccess CurrentUAVAccess = UAV->Resource->GetCurrentGPUAccess();
+		const bool UAVDirty = UAV->Resource->IsDirty();
+		ensureMsgf(!UAVDirty || (CurrentUAVAccess == EResourceTransitionAccess::ERWNoBarrier), TEXT("UAV: %i is in unsafe state for GPU R/W: %s, Dirty: %i"), UAVIndex, *FResourceTransitionUtility::ResourceTransitionAccessStrings[(int32)CurrentUAVAccess], (int32)UAVDirty);
+
+		//UAVs always dirty themselves. If a shader wanted to just read, it should use an SRV.
+		UAV->Resource->SetDirty(true, PresentCounter);
 	}
 
 	ID3D11UnorderedAccessView* D3D11UAV = UAV ? UAV->View : NULL;
@@ -875,7 +889,12 @@ void FD3D11DynamicRHI::RHISetRenderTargets(
 		//switch to writable state if this is the first render of the frame.  Don't switch if it's a later render and this is a depth test only situation
 		if (!bAccessValid || (bReadable && bDepthWrite))
 		{
-			NewDepthStencilTarget->SetCurrentGPUAccess(EResourceTransitionAccess::EWritable, CurrentFrame);			
+			NewDepthStencilTarget->SetCurrentGPUAccess(EResourceTransitionAccess::EWritable);
+		}
+
+		if (bDepthWrite)
+		{
+			NewDepthStencilTarget->SetDirty(true, CurrentFrame);
 		}
 	}
 
@@ -902,8 +921,9 @@ void FD3D11DynamicRHI::RHISetRenderTargets(
 								
 				if (!bAccessValid || bReadable)
 				{
-					NewRenderTarget->SetCurrentGPUAccess(EResourceTransitionAccess::EWritable, CurrentFrame);
+					NewRenderTarget->SetCurrentGPUAccess(EResourceTransitionAccess::EWritable);
 				}
+				NewRenderTarget->SetDirty(true, CurrentFrame);
 			}
 
 			ensureMsgf(RenderTargetView, TEXT("Texture being set as render target has no RTV"));
@@ -972,6 +992,18 @@ void FD3D11DynamicRHI::RHISetRenderTargets(
 			FD3D11UnorderedAccessView* RHIUAV = (FD3D11UnorderedAccessView*)UAVs[UAVIndex];
 			UAV = RHIUAV->View;
 
+			if (UAV)
+			{
+				//check it's safe for r/w for this UAV
+				const EResourceTransitionAccess CurrentUAVAccess = RHIUAV->Resource->GetCurrentGPUAccess();
+				const bool UAVDirty = RHIUAV->Resource->IsDirty();
+				const bool bAccessPass = (CurrentUAVAccess == EResourceTransitionAccess::ERWBarrier && !UAVDirty) || (CurrentUAVAccess == EResourceTransitionAccess::ERWNoBarrier);
+				ensureMsgf(bAccessPass, TEXT("UAV: %i is in unsafe state for GPU R/W: %s"), UAVIndex, *FResourceTransitionUtility::ResourceTransitionAccessStrings[(int32)CurrentUAVAccess]);
+
+				//UAVs get set to dirty.  If the shader just wanted to read it should have used an SRV.
+				RHIUAV->Resource->SetDirty(true, PresentCounter);
+			}
+
 			// Unbind any shader views of the UAV's resource.
 			ConditionalClearShaderResource(RHIUAV->Resource);
 		}
@@ -2136,7 +2168,7 @@ void FD3D11DynamicRHI::RHITransitionResources(EResourceTransitionAccess Transiti
 				Resource = SourceTexture3D;
 			}
 
-			Resource->SetCurrentGPUAccess(TransitionType, PresentCounter);
+			Resource->SetCurrentGPUAccess(TransitionType);
 		}
 	}
 }
@@ -2145,10 +2177,17 @@ void FD3D11DynamicRHI::RHITransitionResources(EResourceTransitionAccess Transiti
 {
 	for (int32 i = 0; i < NumUAVs; ++i)
 	{
-		FD3D11UnorderedAccessView* UAV = ResourceCast(InUAVs[i]);
-		if (UAV->Resource)
+		if (InUAVs[i])
 		{
-			UAV->Resource->SetCurrentGPUAccess(TransitionType, PresentCounter);
-		}		
+			FD3D11UnorderedAccessView* UAV = ResourceCast(InUAVs[i]);
+			if (UAV && UAV->Resource)
+			{
+				UAV->Resource->SetCurrentGPUAccess(TransitionType);
+				if (TransitionType != EResourceTransitionAccess::ERWNoBarrier)
+				{
+					UAV->Resource->SetDirty(false, PresentCounter);
+				}
+			}
+		}
 	}	
 }
\ No newline at end of file
diff --git a/Engine/Source/Runtime/Windows/D3D11RHI/Private/D3D11RHI.cpp b/Engine/Source/Runtime/Windows/D3D11RHI/Private/D3D11RHI.cpp
--- a/Engine/Source/Runtime/Windows/D3D11RHI/Private/D3D11RHI.cpp
+++ b/Engine/Source/Runtime/Windows/D3D11RHI/Private/D3D11RHI.cpp
@@ -183,8 +183,10 @@ void FD3D11DynamicRHI::InternalSetShaderResourceView(FD3D11BaseShaderResource* R
 	CheckIfSRVIsResolved(SRV);
 
 	if (Resource)
-	{
-		ensureMsgf(Resource->GetCurrentGPUAccess() == EResourceTransitionAccess::EReadable || Resource->GetCurrentGPUAccess() == EResourceTransitionAccess::ERWBarrier || Resource->GetLastFrameWritten() != PresentCounter, TEXT("Shader resource %s is not GPU readable.  Missing a call to RHITransitionResources()"), *SRVName.ToString());
+	{		
+		const EResourceTransitionAccess CurrentAccess = Resource->GetCurrentGPUAccess();
+		const bool bAccessPass = CurrentAccess == EResourceTransitionAccess::EReadable || (CurrentAccess == EResourceTransitionAccess::ERWBarrier && !Resource->IsDirty()) || CurrentAccess == EResourceTransitionAccess::ERWSubResBarrier;
+		ensureMsgf(bAccessPass || Resource->GetLastFrameWritten() != PresentCounter, TEXT("Shader resource %s is not GPU readable.  Missing a call to RHITransitionResources()"), *SRVName.ToString());
 	}
 
 	FD3D11BaseShaderResource*& ResourceSlot = CurrentResourcesBoundAsSRVs[ShaderFrequency][ResourceIndex];
diff --git a/Engine/Source/Runtime/Windows/D3D11RHI/Private/D3D11Texture.cpp b/Engine/Source/Runtime/Windows/D3D11RHI/Private/D3D11Texture.cpp
--- a/Engine/Source/Runtime/Windows/D3D11RHI/Private/D3D11Texture.cpp
+++ b/Engine/Source/Runtime/Windows/D3D11RHI/Private/D3D11Texture.cpp
@@ -845,7 +845,7 @@ TD3D11Texture2D<BaseResourceType>* FD3D11DynamicRHI::CreateD3D11Texture2D(uint32
 
 	if (Flags & TexCreate_RenderTargetable)
 	{
-		Texture2D->SetCurrentGPUAccess(EResourceTransitionAccess::EWritable, -1);		
+		Texture2D->SetCurrentGPUAccess(EResourceTransitionAccess::EWritable);
 	}
 
 	D3D11TextureAllocated(*Texture2D);
@@ -972,7 +972,7 @@ FD3D11Texture3D* FD3D11DynamicRHI::CreateD3D11Texture3D(uint32 SizeX,uint32 Size
 
 	if (Flags & TexCreate_RenderTargetable)
 	{
-		Texture3D->SetCurrentGPUAccess(EResourceTransitionAccess::EWritable, -1);		
+		Texture3D->SetCurrentGPUAccess(EResourceTransitionAccess::EWritable);
 	}
 
 	D3D11TextureAllocated(*Texture3D);
diff --git a/Engine/Source/Runtime/Windows/D3D11RHI/Public/D3D11Resources.h b/Engine/Source/Runtime/Windows/D3D11RHI/Public/D3D11Resources.h
--- a/Engine/Source/Runtime/Windows/D3D11RHI/Public/D3D11Resources.h
+++ b/Engine/Source/Runtime/Windows/D3D11RHI/Public/D3D11Resources.h
@@ -170,17 +170,12 @@ public:
 		, bDirty(false)
 	{}
 
-	void SetCurrentGPUAccess(EResourceTransitionAccess Access, uint32 CurrentFrame)
+	void SetCurrentGPUAccess(EResourceTransitionAccess Access)
 	{
 		if (Access == EResourceTransitionAccess::EReadable)
 		{
 			bDirty = false;
 		}
-		else
-		{
-			LastFrameWritten = CurrentFrame;
-			bDirty = true;
-		}		
 		CurrentGPUAccess = Access;
 	}
 
@@ -194,9 +189,14 @@ public:
 		return LastFrameWritten;
 	}
 
-	void SetDirty()
+	void SetDirty(bool bInDirty, uint32 CurrentFrame)
 	{
-		bDirty = true;
+		bDirty = bInDirty;
+		if (bDirty)
+		{
+			LastFrameWritten = CurrentFrame;
+		}
+		ensureMsgf(!(CurrentGPUAccess == EResourceTransitionAccess::EReadable && bDirty), TEXT("UAV is dirty, but set to Readable."));
 	}
 
 	bool IsDirty() const
@@ -751,7 +751,7 @@ public:
 	: FRHIStructuredBuffer(InStride,InSize,InUsage)
 	, Resource(InResource)
 	{
-		SetCurrentGPUAccess(EResourceTransitionAccess::ERWBarrier, -1);
+		SetCurrentGPUAccess(EResourceTransitionAccess::ERWBarrier);
 	}
 
 	virtual ~FD3D11StructuredBuffer()
