Merging
Optimized shader by using more SvPosition instead of ScreenPosition, also optimized some ALU in LightFunction

diff --git a/Engine/Shaders/DeferredDecal.usf b/Engine/Shaders/DeferredDecal.usf
--- a/Engine/Shaders/DeferredDecal.usf
+++ b/Engine/Shaders/DeferredDecal.usf
@@ -10,7 +10,7 @@
 
 #include "Common.usf"
 
-// from SvPosition to decal space (for position)
+// from SvPosition to decal space (for position), used like this: mul(float4(SvPosition.xyz, 1), SvPositionToDecal);
 float4x4 SvPositionToDecal;
 // to transform position from decal space to world space, for normals use transpose of inverse
 float4x4 DecalToWorld;
@@ -36,7 +36,7 @@ half3 FrameBufferDecalBlendOp(half4 Source)
 #else
 	half4 Dest = half4(0, 0, 0, 0);
 #endif
-
+    
 #if DECAL_BLEND_MODE == DECALBLENDMODEID_STAIN
 	return (Source.rgb*Dest.rgb) + (Dest.rgb*(1.0 - Source.a));
 #elif DECAL_BLEND_MODE == DECALBLENDMODEID_TRANSLUCENT
@@ -57,7 +57,7 @@ void MainVS(
 	OutPosition = mul(InPosition, FrustumComponentToClip);
 }
 
-float DistanceFunction(float4 ScreenPosition, float3 Pos)
+float DistanceFunction(float4 SvPosition, float3 Pos)
 {
 	float3 SwizzlePos = Pos.zyx;
 
@@ -72,7 +72,8 @@ float DistanceFunction(float4 ScreenPosition, float3 Pos)
 #endif
 
 	MaterialParameters.VertexColor = 1;
-	MaterialParameters.ScreenPosition = ScreenPosition;
+	MaterialParameters.SvPosition = SvPosition;
+	MaterialParameters.ScreenPosition = SvPositionToScreenPosition(SvPosition);
 	MaterialParameters.LightVector = SwizzlePos * 0.5f + 0.5f;		// to be compatible with decals but *0.5+0.5 seems wrong
 	MaterialParameters.AbsoluteWorldPosition = 
 	MaterialParameters.WorldPosition_CamRelative = 
@@ -92,7 +93,7 @@ float DistanceFunction(float4 ScreenPosition, float3 Pos)
 const static float coneSpread = 0.00002f;
 
 // @return t=0..1, -1: no hit
-float Raycast(float4 ScreenPosition, float3 StartPos, float3 EndPos)
+float Raycast(float4 SvPosition, float3 StartPos, float3 EndPos)
 {
 	const uint steps = 300; 
 
@@ -103,7 +104,7 @@ float Raycast(float4 ScreenPosition, float3 StartPos, float3 EndPos)
 
 	LOOP for(uint i = 0; i < steps; ++i)
 	{
-		float Distance = DistanceFunction(ScreenPosition, StartPos + t * Direction);
+		float Distance = DistanceFunction(SvPosition, StartPos + t * Direction);
 
 		if(Distance <= t * coneSpread)
 		{
@@ -122,15 +123,15 @@ float Raycast(float4 ScreenPosition, float3 StartPos, float3 EndPos)
 }
 
 // like explained here http://iquilezles.org/www/articles/derivative/derivative.htm
-float3 GradientNotNormalized(float4 ScreenPosition, float3 OSPos, float t)
+float3 GradientNotNormalized(float4 SvPosition, float3 OSPos, float t)
 {
 	float eps = t * coneSpread;
 
 	float2 e = float2(eps, 0);
 
-    return float3(DistanceFunction(ScreenPosition, OSPos + e.xyy) - DistanceFunction(ScreenPosition, OSPos - e.xyy),
-                  DistanceFunction(ScreenPosition, OSPos + e.yxy) - DistanceFunction(ScreenPosition, OSPos - e.yxy),
-                  DistanceFunction(ScreenPosition, OSPos + e.yyx) - DistanceFunction(ScreenPosition, OSPos - e.yyx));
+    return float3(DistanceFunction(SvPosition, OSPos + e.xyy) - DistanceFunction(SvPosition, OSPos - e.xyy),
+                  DistanceFunction(SvPosition, OSPos + e.yxy) - DistanceFunction(SvPosition, OSPos - e.yxy),
+                  DistanceFunction(SvPosition, OSPos + e.yyx) - DistanceFunction(SvPosition, OSPos - e.yyx));
 }
 
 
@@ -191,8 +192,6 @@ void MainPS(
 	// make SvPosition appear to be rasterized with the depth from the depth buffer
 	SvPosition.z = LookupDeviceZ(ScreenUV);
 
-	float4 ScreenPosition = SvPositionToScreenPosition(SvPosition);
-
 #if DECAL_RENDERTARGET_COUNT > 4
 	// default to no shadows, good for volumetric decals
 	OutTarget4 = 1;
@@ -210,12 +209,8 @@ void MainPS(
 		float SceneW = 100000000;
 
 		float4 SvPositionInDistance = float4(SvPosition.xy, ConvertToDeviceZ(SceneW), 1);
-//		float4 ProjectedScreenPosition = float4(ScreenPosition.xy / ScreenPosition.w * SceneW, SceneW, 1);		// is ProjectedClipPosition
-//		ProjectedScreenPosition = float4(ScreenPosition.xy / ScreenPosition.w * SceneW, SceneW, 1);
-//		float4 WorldPosHom = mul(ProjectedScreenPosition, View.ScreenToWorld);
 
 		float3 WSStartRay = View.WorldCameraOrigin;
-//		float3 WSEndRay = WorldPosHom.xyz;
 		float3 WSEndRay = SvPositionToWorld(SvPositionInDistance);
 
 		// in object/decal space
@@ -230,7 +225,7 @@ void MainPS(
 		float3 OSBoxStartRay = OSStartRay + OSBoxMinMax.x * (OSEndRay - OSStartRay);
 		float3 OSBoxEndRay = OSStartRay + OSBoxMinMax.y * (OSEndRay - OSStartRay);
 
-		HitT = Raycast(ScreenPosition, OSBoxStartRay, OSBoxEndRay);
+		HitT = Raycast(SvPosition, OSBoxStartRay, OSBoxEndRay);
 
 		clip(HitT);
 
@@ -252,9 +247,8 @@ void MainPS(
 		clip(OSPosition.xyz + 1.0f);
 		clip(1.0f - OSPosition.xyz);
 
-		float4 Hom = mul(float4(SvPosition.xyz, 1), View.SVPositionToTranslatedWorld);
-		
-		WSPosition = Hom.xyz / Hom.w - View.PreViewTranslation;
+		// todo: TranslatedWorld would be better for quality
+		WSPosition = SvPositionToWorld(SvPosition);
 	}
 
 	float3 CameraVector = normalize(View.WorldCameraOrigin - WSPosition);
@@ -278,7 +272,7 @@ void MainPS(
 	MaterialParameters.VertexColor = 1;
 	MaterialParameters.CameraVector = CameraVector;
 	MaterialParameters.SvPosition = SvPosition;
-	MaterialParameters.ScreenPosition = ScreenPosition;
+	MaterialParameters.ScreenPosition = SvPositionToScreenPosition(SvPosition);
 	MaterialParameters.LightVector = SwizzlePos;
 	MaterialParameters.AbsoluteWorldPosition = 
 		MaterialParameters.WorldPosition_CamRelative = 
@@ -289,7 +283,7 @@ void MainPS(
 		if(bVolumetric)
 		{
 			// not normalized
-			float3 OSNormal = GradientNotNormalized(ScreenPosition, OSPosition, HitT);
+			float3 OSNormal = GradientNotNormalized(SvPosition, OSPosition, HitT);
 			
 			// not normalized
 			float3 WSNormal = mul(float4(OSNormal, 0), transpose(WorldToDecal)).xyz;
diff --git a/Engine/Shaders/LightFunctionPixelShader.usf b/Engine/Shaders/LightFunctionPixelShader.usf
--- a/Engine/Shaders/LightFunctionPixelShader.usf
+++ b/Engine/Shaders/LightFunctionPixelShader.usf
@@ -8,7 +8,7 @@
 #include "Material.usf"
 #include "LightFunctionCommon.usf"
 
-float4x4 ScreenToLight;
+float4x4 SvPositionToLight;
 
 /** Fade distance in x, disabled brightness in y, output for preview shadows mask in z. */
 float3 LightFunctionParameters2;
@@ -19,26 +19,36 @@ void Main(
 	out float4 OutColor : SV_Target0
 	)
 {
-	half SceneW = CalcSceneDepth(SvPositionToBufferUV(SvPosition));
-	float4 ProjectedScreenPosition = float4(ScreenPosition.xy / ScreenPosition.w * SceneW, SceneW, 1);
-	float4 LightVector = mul(ProjectedScreenPosition, ScreenToLight);
-	LightVector.xyz /= LightVector.w;
+	float2 ScreenUV = SvPositionToBufferUV(SvPosition);
+
+	// make SvPosition appear to be rasterized with the depth from the depth buffer
+	SvPosition.z = LookupDeviceZ(ScreenUV);
+
+	float3 LightVector;
+	{
+		float4 Hom = mul(float4(SvPosition.xyz, 1), SvPositionToLight);
+		LightVector = Hom.xyz / Hom.w;
+	}
+	
+	// todo: TranslatedWorld would be better for quality
+	float3 AbsoluteWorldPosition = SvPositionToWorld(SvPosition);
 
-	float4 HomogeneousWorldPosition = mul(ProjectedScreenPosition, View.ScreenToWorld);
-	float3 AbsoluteWorldPosition = HomogeneousWorldPosition.xyz / HomogeneousWorldPosition.w;
 	// Calculate radial view distance for stable fading
 	float ViewDistance = length(View.WorldCameraOrigin - AbsoluteWorldPosition);
 
-	float3 Color = GetLightFunctionColor(LightVector.xyz, AbsoluteWorldPosition);
+	float GreyScale;
+	{
+		float3 Color = GetLightFunctionColor(LightVector, AbsoluteWorldPosition);
+		GreyScale = dot(Color, .3333f);
+	}
 
 	float DistanceFadeAlpha = saturate((LightFunctionParameters2.x - ViewDistance) / (LightFunctionParameters2.x * .2f));
 	// Fade to disabled based on LightFunctionFadeDistance
-	Color = lerp(LightFunctionParameters2.yyy, Color, DistanceFadeAlpha);
+	GreyScale = lerp(LightFunctionParameters2.y, GreyScale, DistanceFadeAlpha);
 
 	// Fade to disabled based on ShadowFadeFraction
-	Color = lerp(LightFunctionParameters2.yyy, Color, LightFunctionParameters.y);
+	GreyScale = lerp(LightFunctionParameters2.y, GreyScale, LightFunctionParameters.y);
 
-	float GreyScale = dot(Color, .3333f);
 	float EncodedLightAttenuation = EncodeLightAttenuation(GreyScale);
 
 	// Light function shadows write to the blue channel.
diff --git a/Engine/Shaders/LightFunctionVertexShader.usf b/Engine/Shaders/LightFunctionVertexShader.usf
--- a/Engine/Shaders/LightFunctionVertexShader.usf
+++ b/Engine/Shaders/LightFunctionVertexShader.usf
@@ -10,11 +10,10 @@ float4 StencilingGeometryPosAndScale;
 
 void Main(
 	in float4 InPosition : ATTRIBUTE0,
-	out float4 OutScreenPosition : TEXCOORD0,
 	out float4 OutPosition : SV_POSITION
 	)
 {
 	float3 WorldPosition = InPosition * StencilingGeometryPosAndScale.w + StencilingGeometryPosAndScale.xyz;
-	OutPosition = OutScreenPosition = mul(float4(WorldPosition,1), View.TranslatedWorldToClip);
+	OutPosition = mul(float4(WorldPosition,1), View.TranslatedWorldToClip);
 }
 
diff --git a/Engine/Source/Runtime/Renderer/Private/LightFunctionRendering.cpp b/Engine/Source/Runtime/Renderer/Private/LightFunctionRendering.cpp
--- a/Engine/Source/Runtime/Renderer/Private/LightFunctionRendering.cpp
+++ b/Engine/Source/Runtime/Renderer/Private/LightFunctionRendering.cpp
@@ -34,20 +34,20 @@ public:
 		StencilingGeometryParameters.Bind(Initializer.ParameterMap);
 	}
 
-	void SetParameters(FRHICommandList& RHICmdList, const FSceneView* View, const FLightSceneInfo* LightSceneInfo )
+	void SetParameters(FRHICommandList& RHICmdList, const FSceneView& View, const FLightSceneInfo* LightSceneInfo )
 	{
-		FMaterialShader::SetParameters(RHICmdList, GetVertexShader(), *View);
+		FMaterialShader::SetParameters(RHICmdList, GetVertexShader(), View);
 		
 		// Light functions are projected using a bounding sphere.
 		// Calculate transform for bounding stencil sphere.
 		FSphere LightBounds = LightSceneInfo->Proxy->GetBoundingSphere();
 		if (LightSceneInfo->Proxy->GetLightType() == LightType_Directional)
 		{
-			LightBounds.Center = View->ViewMatrices.ViewOrigin;
+			LightBounds.Center = View.ViewMatrices.ViewOrigin;
 		}
 
 		FVector4 StencilingSpherePosAndScale;
-		StencilingGeometry::GStencilSphereVertexBuffer.CalcTransform(StencilingSpherePosAndScale, LightBounds, View->ViewMatrices.PreViewTranslation);
+		StencilingGeometry::GStencilSphereVertexBuffer.CalcTransform(StencilingSpherePosAndScale, LightBounds, View.ViewMatrices.PreViewTranslation);
 		StencilingGeometryParameters.Set(RHICmdList, this, StencilingSpherePosAndScale);
 	}
 
@@ -88,34 +88,49 @@ public:
 	FLightFunctionPS(const ShaderMetaType::CompiledShaderInitializerType& Initializer)
 		: FMaterialShader(Initializer)
 	{
-		ScreenToLight.Bind(Initializer.ParameterMap,TEXT("ScreenToLight"));
+		SvPositionToLight.Bind(Initializer.ParameterMap,TEXT("SvPositionToLight"));
 		LightFunctionParameters.Bind(Initializer.ParameterMap);
 		LightFunctionParameters2.Bind(Initializer.ParameterMap,TEXT("LightFunctionParameters2"));
 		DeferredParameters.Bind(Initializer.ParameterMap);
 	}
 
-	void SetParameters(FRHICommandList& RHICmdList, const FSceneView* View, const FLightSceneInfo* LightSceneInfo, const FMaterialRenderProxy* MaterialProxy, bool bRenderingPreviewShadowIndicator, float ShadowFadeFraction )
+	void SetParameters(FRHICommandList& RHICmdList, const FSceneView& View, const FLightSceneInfo* LightSceneInfo, const FMaterialRenderProxy* MaterialProxy, bool bRenderingPreviewShadowIndicator, float ShadowFadeFraction )
 	{
 		const FPixelShaderRHIParamRef ShaderRHI = GetPixelShader();
 
-		FMaterialShader::SetParameters(RHICmdList, ShaderRHI, MaterialProxy, *MaterialProxy->GetMaterial(View->GetFeatureLevel()), *View, true, ESceneRenderTargetsMode::SetTextures);
+		FMaterialShader::SetParameters(RHICmdList, ShaderRHI, MaterialProxy, *MaterialProxy->GetMaterial(View.GetFeatureLevel()), View, true, ESceneRenderTargetsMode::SetTextures);
 
 		// Set the transform from screen space to light space.
-		if ( ScreenToLight.IsBound() )
+		if ( SvPositionToLight.IsBound() )
 		{
 			const FVector Scale = LightSceneInfo->Proxy->GetLightFunctionScale();
 			// Switch x and z so that z of the user specified scale affects the distance along the light direction
 			const FVector InverseScale = FVector( 1.f / Scale.Z, 1.f / Scale.Y, 1.f / Scale.X );
 			const FMatrix WorldToLight = LightSceneInfo->Proxy->GetWorldToLight() * FScaleMatrix(FVector(InverseScale));	
-			const FMatrix ScreenToLightValue = 
+
+			FVector2D InvViewSize = FVector2D(1.0f / View.ViewRect.Width(), 1.0f / View.ViewRect.Height());
+
+			// setup a matrix to transform float4(SVPosition.xyz,1) directly to Light (quality, performance as we don't need to convert or use interpolator)
+
+			//	new_xy = (xy - ViewRectMin.xy) * ViewSizeAndInvSize.zw * float2(2,-2) + float2(-1, 1);
+
+			//  transformed into one MAD:  new_xy = xy * ViewSizeAndInvSize.zw * float2(2,-2)      +       (-ViewRectMin.xy) * ViewSizeAndInvSize.zw * float2(2,-2) + float2(-1, 1);
+
+			float Mx = 2.0f * InvViewSize.X;
+			float My = -2.0f * InvViewSize.Y;
+			float Ax = -1.0f - 2.0f * View.ViewRect.Min.X * InvViewSize.X;
+			float Ay = 1.0f + 2.0f * View.ViewRect.Min.Y * InvViewSize.Y;
+
+			// todo: we could use InvTranslatedViewProjectionMatrix and TranslatedWorldToLight for better quality
+			const FMatrix SvPositionToLightValue = 
 				FMatrix(
-					FPlane(1,0,0,0),
-					FPlane(0,1,0,0),
-					FPlane(0,0,View->ViewMatrices.ProjMatrix.M[2][2],1),
-					FPlane(0,0,View->ViewMatrices.ProjMatrix.M[3][2],0))
-				* View->InvViewProjectionMatrix * WorldToLight;
+					FPlane(Mx,  0,   0,  0),
+					FPlane( 0, My,   0,  0),
+					FPlane( 0,  0,   1,  0),
+					FPlane(Ax, Ay,   0,  1)
+				) * View.InvViewProjectionMatrix * WorldToLight;
 
-			SetShaderValue(RHICmdList, ShaderRHI, ScreenToLight, ScreenToLightValue );
+			SetShaderValue(RHICmdList, ShaderRHI, SvPositionToLight, SvPositionToLightValue );
 		}
 
 		LightFunctionParameters.Set(RHICmdList, ShaderRHI, LightSceneInfo, ShadowFadeFraction);
@@ -125,13 +140,13 @@ public:
 			LightSceneInfo->Proxy->GetLightFunctionDisabledBrightness(),
 			bRenderingPreviewShadowIndicator ? 1.0f : 0.0f));
 
-		DeferredParameters.Set(RHICmdList, ShaderRHI, *View);
+		DeferredParameters.Set(RHICmdList, ShaderRHI, View);
 	}
 
 	virtual bool Serialize(FArchive& Ar) override
 	{
 		bool bShaderHasOutdatedParameters = FMaterialShader::Serialize(Ar);
-		Ar << ScreenToLight;
+		Ar << SvPositionToLight;
 		Ar << LightFunctionParameters;
 		Ar << LightFunctionParameters2;
 		Ar << DeferredParameters;
@@ -139,7 +154,7 @@ public:
 	}
 
 private:
-	FShaderParameter ScreenToLight;
+	FShaderParameter SvPositionToLight;
 	FLightFunctionSharedParameters LightFunctionParameters;
 	FShaderParameter LightFunctionParameters2;
 	FDeferredPixelShaderParameters DeferredParameters;
@@ -324,8 +339,8 @@ bool FDeferredShadingSceneRenderer::RenderLightFunctionForMaterial(FRHICommandLi
 
 					// Render a bounding light sphere.
 					RHICmdList.SetLocalBoundShaderState(LightFunctionBoundShaderState);
-					VertexShader->SetParameters(RHICmdList, &View, LightSceneInfo);
-					PixelShader->SetParameters(RHICmdList, &View, LightSceneInfo, MaterialProxy, bRenderingPreviewShadowsIndicator, FadeAlpha);
+					VertexShader->SetParameters(RHICmdList, View, LightSceneInfo);
+					PixelShader->SetParameters(RHICmdList, View, LightSceneInfo, MaterialProxy, bRenderingPreviewShadowsIndicator, FadeAlpha);
 
 					// Project the light function using a sphere around the light
 					//@todo - could use a cone for spotlights
