Integrate: StaticMesh Triangle PreCulling prototype
* New command 'r.PreCullIndexBuffers' which will process all static mesh triangles in the scene, generating culled per-instance index buffers that remove triangles completely inside another mesh (using distance fields), completely underneath the Landscape, or completely inside a PreCullVolume.
* New command 'r.TogglePreCulledIndexBuffers' to toggle the use of the culled index buffers.  This is useful for inspecting triangle savings at runtime along with 'stat rhi'
* Instanced index buffer memory is reported by 'stat memorystaticmesh'
* Requires the Renderer Project setting 'GenerateMeshDistanceFields' to be enabled locally, and only works on PC D3D11
* Results are not saved or outdated correctly on static mesh reimport / edit

diff --git a/Engine/Shaders/HeightfieldLighting.usf b/Engine/Shaders/HeightfieldLighting.usf
--- a/Engine/Shaders/HeightfieldLighting.usf
+++ b/Engine/Shaders/HeightfieldLighting.usf
@@ -990,4 +990,77 @@ void CalculateHeightfieldIrradianceScreenGridCS(
 			RWHeightfieldIrradiance[OutputIndex] = float4(Irradiance, 0);
 		}
 	}
-}
\ No newline at end of file
+}
+
+// In float's, must match FConvertToUniformMeshVertex and ComputeUniformVertexStride()
+#define TRIANGLE_VERTEX_DATA_STRIDE 24
+
+uint StartIndex;
+uint NumTriangles;
+Buffer<float> TriangleVertexData;
+
+float LoadVertexAttribute(uint TriangleIndex, uint VertexIndex, uint AttributeIndex)
+{
+	return TriangleVertexData[(TriangleIndex * 3 + VertexIndex) * TRIANGLE_VERTEX_DATA_STRIDE + AttributeIndex];
+}
+
+RWBuffer<uint> RWTriangleVisibleMask;
+
+#ifndef PRECULL_TRIANGLES_TO_HEIGHTFIELDS_GROUP_SIZE
+#define PRECULL_TRIANGLES_TO_HEIGHTFIELDS_GROUP_SIZE 1
+#endif
+
+/**  */
+[numthreads(PRECULL_TRIANGLES_TO_HEIGHTFIELDS_GROUP_SIZE, 1, 1)]
+void PreCullTrianglesToHeightfieldsCS(
+	uint3 GroupId : SV_GroupID,
+	uint3 DispatchThreadId : SV_DispatchThreadID,
+    uint3 GroupThreadId : SV_GroupThreadID) 
+{
+	uint TriangleIndex = DispatchThreadId.x;
+
+	if (TriangleIndex < NumTriangles)
+	{
+		uint GlobalTriangleIndex = TriangleIndex + StartIndex / 3;
+		float3 Position0 = float3(LoadVertexAttribute(GlobalTriangleIndex, 0, 0), LoadVertexAttribute(GlobalTriangleIndex, 0, 1), LoadVertexAttribute(GlobalTriangleIndex, 0, 2));
+		float3 Position1 = float3(LoadVertexAttribute(GlobalTriangleIndex, 1, 0), LoadVertexAttribute(GlobalTriangleIndex, 1, 1), LoadVertexAttribute(GlobalTriangleIndex, 1, 2));
+		float3 Position2 = float3(LoadVertexAttribute(GlobalTriangleIndex, 2, 0), LoadVertexAttribute(GlobalTriangleIndex, 2, 1), LoadVertexAttribute(GlobalTriangleIndex, 2, 2));
+		/*
+		float3 TriangleNormal = cross(Position2 - Position0, Position1 - Position0);
+		float TriangleArea = .5f * length(TriangleNormal);
+		uint NumSamplesOneDim = 10;
+		float SampleArea = TriangleArea / (NumSamplesOneDim * NumSamplesOneDim);
+		float SampleRadius = .5f * sqrt(2 * SampleArea);
+		*/
+		float3 InterpolatedPosition = (Position0 + Position1 + Position2) / 3.0f;
+		float SampleRadius = length(InterpolatedPosition - Position0);
+
+		bool bTriangleIsVisible = true;
+
+		LOOP
+		for (uint HeightfieldIndex = 0; HeightfieldIndex < NumHeightfields; HeightfieldIndex++)
+		{
+			float3 LocalPosition = mul(float4(InterpolatedPosition, 1), GetWorldToLocal(HeightfieldIndex)).xyz;
+			float4 MinMaxHeightfieldUV;
+			float2 HeightfieldUV = GetHeightfieldUV(HeightfieldIndex, LocalPosition.xy, MinMaxHeightfieldUV);
+
+			BRANCH
+			if (all(HeightfieldUV > MinMaxHeightfieldUV.xy) 
+				&& all(HeightfieldUV < MinMaxHeightfieldUV.zw))
+			{
+				float3 WorldHeightfieldShadingPosition = GetHeightfieldWorldPosition(HeightfieldIndex, LocalPosition.xy, HeightfieldUV);
+
+				//@todo - multiple samples to reduce error
+				if (InterpolatedPosition.z + SampleRadius * 1.5f < WorldHeightfieldShadingPosition.z)
+				{
+					bTriangleIsVisible = false;
+				}
+			}
+		}
+		
+		if (!bTriangleIsVisible)
+		{
+			RWTriangleVisibleMask[GlobalTriangleIndex] = 0;
+		}
+	}
+}
diff --git a/Engine/Shaders/PreCullTriangles.usf b/Engine/Shaders/PreCullTriangles.usf
new file mode 100644
--- /dev/null
+++ b/Engine/Shaders/PreCullTriangles.usf
@@ -0,0 +1,262 @@
+// Copyright 1998-2014 Epic Games, Inc. All Rights Reserved.
+
+/*=============================================================================
+	PreCullTriangles.usf
+=============================================================================*/
+
+#include "Common.usf"
+#include "DeferredShadingCommon.usf"
+#include "DistanceFieldLightingShared.usf"
+#include "DistanceFieldAOShared.usf"
+
+
+RWBuffer<float4> RWCulledObjectBounds;
+RWBuffer<float4> RWCulledObjectData;
+
+uint ObjectBoundingGeometryIndexCount;
+float3 ObjectBoundsCenter;
+float3 ObjectBoundsExtent;
+
+groupshared uint NumGroupObjects;
+
+groupshared uint GroupBaseIndex;
+groupshared uint GroupObjectIndices[UPDATEOBJECTS_THREADGROUP_SIZE];
+
+float4 FetchObjectDataFloat4(uint SourceIndex)
+{
+	return float4(ObjectData[4 * SourceIndex + 0], ObjectData[4 * SourceIndex + 1], ObjectData[4 * SourceIndex + 2], ObjectData[4 * SourceIndex + 3]);
+}
+
+void CopyCulledObjectData(uint DestIndex, uint SourceIndex)
+{
+	RWCulledObjectBounds[DestIndex] = float4(ObjectBounds[4 * SourceIndex + 0], ObjectBounds[4 * SourceIndex + 1], ObjectBounds[4 * SourceIndex + 2], ObjectBounds[4 * SourceIndex + 3]);
+
+	UNROLL
+	for (uint VectorIndex = 0; VectorIndex < CULLED_OBJECT_DATA_STRIDE; VectorIndex++)
+	{
+		float4 Data = FetchObjectDataFloat4(SourceIndex * OBJECT_DATA_STRIDE + VectorIndex);
+
+		// Note: only copying the first CULLED_OBJECT_DATA_STRIDE of the original object data
+		RWCulledObjectData[DestIndex * CULLED_OBJECT_DATA_STRIDE + VectorIndex] = Data;
+	}
+}
+
+float PreCullMaxDistance;
+
+[numthreads(UPDATEOBJECTS_THREADGROUP_SIZE, 1, 1)]
+void CullObjectsForBoundsCS(
+	uint3 GroupId : SV_GroupID,
+	uint3 DispatchThreadId : SV_DispatchThreadID,
+    uint3 GroupThreadId : SV_GroupThreadID) 
+{
+	uint ObjectIndex = DispatchThreadId.x;
+
+#define USE_FRUSTUM_CULLING 1
+#if USE_FRUSTUM_CULLING
+	if (DispatchThreadId.x == 0)
+	{
+		// RWObjectIndirectArguments is zeroed by a clear before this shader, only need to set things that are non-zero (and are not read by this shader as that would be a race condition)
+		// IndexCount, NumInstances, StartIndex, BaseVertexIndex, FirstInstance
+		RWObjectIndirectArguments[0] = ObjectBoundingGeometryIndexCount;
+	}
+
+	if (GroupThreadId.x == 0)
+	{
+		NumGroupObjects = 0;
+	}
+
+	GroupMemoryBarrierWithGroupSync();
+
+	if (ObjectIndex < NumSceneObjects)
+	{
+		uint SourceIndex = ObjectIndex;
+
+		float4 ObjectBoundingSphere = float4(ObjectBounds[4 * SourceIndex + 0], ObjectBounds[4 * SourceIndex + 1], ObjectBounds[4 * SourceIndex + 2], ObjectBounds[4 * SourceIndex + 3]);
+
+		if (ComputeSquaredDistanceFromBoxToPoint(ObjectBoundsCenter, ObjectBoundsExtent, ObjectBoundingSphere.xyz) < Square(PreCullMaxDistance + ObjectBoundingSphere.w))
+		{
+			uint DestIndex;
+			InterlockedAdd(NumGroupObjects, 1U, DestIndex);
+			GroupObjectIndices[DestIndex] = SourceIndex;
+		}
+	}
+
+	GroupMemoryBarrierWithGroupSync();
+
+	if (GroupThreadId.x == 0)
+	{
+		InterlockedAdd(RWObjectIndirectArguments[1], NumGroupObjects, GroupBaseIndex);
+	}
+
+	GroupMemoryBarrierWithGroupSync();
+
+	if (GroupThreadId.x < NumGroupObjects)
+	{
+		uint SourceIndex = GroupObjectIndices[GroupThreadId.x];
+		uint DestIndex = GroupBaseIndex + GroupThreadId.x;
+		CopyCulledObjectData(DestIndex, SourceIndex);
+	}
+
+#else
+
+	if (DispatchThreadId.x == 0)
+	{
+		// IndexCount, NumInstances, StartIndex, BaseVertexIndex, FirstInstance
+		RWObjectIndirectArguments[0] = ObjectBoundingGeometryIndexCount;
+		RWObjectIndirectArguments[1] = NumSceneObjects;
+	}
+
+	GroupMemoryBarrierWithGroupSync();
+
+	if (ObjectIndex < NumSceneObjects)
+	{
+		uint SourceIndex = ObjectIndex;
+		uint DestIndex = ObjectIndex;
+
+		CopyCulledObjectData(DestIndex, SourceIndex);
+	}
+
+#endif
+}
+
+// In float's, must match FConvertToUniformMeshVertex and ComputeUniformVertexStride()
+#define TRIANGLE_VERTEX_DATA_STRIDE 24
+
+uint StartIndex;
+uint NumTriangles;
+Buffer<float> TriangleVertexData;
+
+float LoadVertexAttribute(uint TriangleIndex, uint VertexIndex, uint AttributeIndex)
+{
+	return TriangleVertexData[(TriangleIndex * 3 + VertexIndex) * TRIANGLE_VERTEX_DATA_STRIDE + AttributeIndex];
+}
+
+uint NumCullVolumes;
+Buffer<uint> CullVolumePlaneOffsets;
+Buffer<float4> CullVolumePlaneData;
+
+RWBuffer<uint> RWTriangleVisibleMask;
+
+bool IsTriangleInsideAnyCullVolume(float3 Position0, float3 Position1, float3 Position2)
+{
+	for (uint VolumeIndex = 0; VolumeIndex < NumCullVolumes; VolumeIndex++)
+	{
+		uint NumPlanes = CullVolumePlaneOffsets[VolumeIndex * 2 + 0];
+		uint StartPlaneIndex = CullVolumePlaneOffsets[VolumeIndex * 2 + 1];
+
+		bool bInsideAllPlanes = true;
+
+		for (uint PlaneIndex = 0; PlaneIndex < NumPlanes; PlaneIndex++)
+		{
+			float4 PlaneData = CullVolumePlaneData[StartPlaneIndex + PlaneIndex];
+			float PlaneDot0 = dot(PlaneData.xyz, Position0) - PlaneData.w;
+			float PlaneDot1 = dot(PlaneData.xyz, Position1) - PlaneData.w;
+			float PlaneDot2 = dot(PlaneData.xyz, Position2) - PlaneData.w;
+
+			bInsideAllPlanes = bInsideAllPlanes && PlaneDot0 < 0.0f && PlaneDot1 < 0.0f && PlaneDot2 < 0.0f;
+		}
+
+		if (bInsideAllPlanes)
+		{
+			// Note: assumes convex volumes, all triangle vertices being inside the volume means the whole triangle is inside
+			return true;
+		}
+	}
+
+	return false;
+}
+
+[numthreads(THREADGROUP_TOTALSIZE, 1, 1)]
+void PreCullTrianglesCS(
+	uint3 GroupId : SV_GroupID,
+	uint3 DispatchThreadId : SV_DispatchThreadID,
+    uint3 GroupThreadId : SV_GroupThreadID) 
+{
+	uint TriangleIndex = DispatchThreadId.x;
+
+	if (TriangleIndex < NumTriangles)
+	{
+		uint GlobalTriangleIndex = TriangleIndex + StartIndex / 3;
+		float3 Position0 = float3(LoadVertexAttribute(GlobalTriangleIndex, 0, 0), LoadVertexAttribute(GlobalTriangleIndex, 0, 1), LoadVertexAttribute(GlobalTriangleIndex, 0, 2));
+		float3 Position1 = float3(LoadVertexAttribute(GlobalTriangleIndex, 1, 0), LoadVertexAttribute(GlobalTriangleIndex, 1, 1), LoadVertexAttribute(GlobalTriangleIndex, 1, 2));
+		float3 Position2 = float3(LoadVertexAttribute(GlobalTriangleIndex, 2, 0), LoadVertexAttribute(GlobalTriangleIndex, 2, 1), LoadVertexAttribute(GlobalTriangleIndex, 2, 2));
+
+		bool bTriangleIsVisible = false;
+
+		float3 TriangleNormal = cross(Position2 - Position0, Position1 - Position0);
+		float TriangleArea = .5f * length(TriangleNormal);
+		uint NumSamplesOneDim = 10;
+		float SampleArea = TriangleArea / (NumSamplesOneDim * NumSamplesOneDim);
+		float SampleRadius = .5f * sqrt(2 * SampleArea);
+
+		//float3 InterpolatedPosition = (Position0 + Position1 + Position2) / 3.0f;
+		//float SampleRadius = length(InterpolatedPosition - Position0);
+
+		if (SampleRadius < PreCullMaxDistance)
+		{
+			uint NumCulledObjects = GetCulledNumObjects();
+
+			for (uint SampleIndex = 0; SampleIndex < NumSamplesOneDim && !bTriangleIsVisible; SampleIndex++)
+			{
+				for (uint SampleIndex2 = 0; SampleIndex2 < NumSamplesOneDim && !bTriangleIsVisible; SampleIndex2++)
+				{
+					float U = sqrt(SampleIndex / (float)NumSamplesOneDim);
+					float V = U * SampleIndex2 / (float)NumSamplesOneDim;
+					U = 1 - U;
+					float W = 1 - U - V;
+
+					float3 InterpolatedPosition = Position0 * U + Position1 * V + Position2 * W;
+
+					float MinDistanceToSurface = PreCullMaxDistance;
+
+					LOOP
+					for (uint ObjectIndex = 0; ObjectIndex < NumCulledObjects; ObjectIndex++)
+					{
+						float4 SphereCenterAndRadius = LoadObjectPositionAndRadius(ObjectIndex);
+
+						float3 LocalPositionExtent = LoadObjectLocalPositionExtent(ObjectIndex);
+						float4x4 WorldToVolume = LoadObjectWorldToVolume(ObjectIndex);
+						bool bGeneratedAsTwoSided;
+						float4 UVScaleAndVolumeScale = LoadObjectUVScale(ObjectIndex, bGeneratedAsTwoSided);
+						float3 WorldTexelSize = DistanceFieldAtlasTexelSize / UVScaleAndVolumeScale.xyz * UVScaleAndVolumeScale.w;
+						float MaxWorldTexelSize = .5f * max(WorldTexelSize.x, max(WorldTexelSize.y, WorldTexelSize.z));
+/*
+						float3 VolumeShadingPosition = mul(float4(InterpolatedPosition, 1), WorldToVolume).xyz;
+						float BoxDistance = ComputeDistanceFromBoxToPoint(-LocalPositionExtent, LocalPositionExtent, VolumeShadingPosition) * UVScaleAndVolumeScale.w;
+
+						BRANCH
+						if (BoxDistance > 0)
+						{
+							bTriangleIsVisible = true;
+						}*/
+
+						{
+							float3 UVAdd = LoadObjectUVAdd(ObjectIndex);
+							float3 VolumeSamplePosition = mul(float4(InterpolatedPosition, 1), WorldToVolume).xyz;
+							float3 ClampedSamplePosition = clamp(VolumeSamplePosition, -LocalPositionExtent, LocalPositionExtent);
+							float DistanceToClamped = length(VolumeSamplePosition - ClampedSamplePosition);
+
+							float3 StepVolumeUV = DistanceFieldVolumePositionToUV(ClampedSamplePosition, UVScaleAndVolumeScale.xyz, UVAdd);
+							float DistanceToOccluder = (Texture3DSampleLevel(DistanceFieldTexture, DistanceFieldSampler, StepVolumeUV, 0).x + DistanceToClamped) * UVScaleAndVolumeScale.w;
+
+							MinDistanceToSurface = min(MinDistanceToSurface, DistanceToOccluder + MaxWorldTexelSize);
+						}
+					}
+
+					if (MinDistanceToSurface > -SampleRadius)
+					{
+						bTriangleIsVisible = true;
+					}
+				}
+			}
+		}
+		else
+		{
+			bTriangleIsVisible = true;
+		}
+		
+		bool bInsideCullVolume = IsTriangleInsideAnyCullVolume(Position0, Position1, Position2);
+
+		RWTriangleVisibleMask[GlobalTriangleIndex] = bTriangleIsVisible && !bInsideCullVolume ? 1 : 0;
+	}
+}
\ No newline at end of file
diff --git a/Engine/Source/Runtime/Engine/Classes/Components/StaticMeshComponent.h b/Engine/Source/Runtime/Engine/Classes/Components/StaticMeshComponent.h
--- a/Engine/Source/Runtime/Engine/Classes/Components/StaticMeshComponent.h
+++ b/Engine/Source/Runtime/Engine/Classes/Components/StaticMeshComponent.h
@@ -5,7 +5,7 @@
 #include "SceneTypes.h"
 #include "Components/MeshComponent.h"
 #include "Runtime/RenderCore/Public/PackedNormal.h"
-
+#include "RawIndexBuffer.h"
 #include "StaticMeshComponent.generated.h"
 
 class FColorVertexBuffer;
@@ -48,6 +48,13 @@ struct FPaintedVertex
 	
 };
 
+struct FPreCulledStaticMeshSection
+{
+	/** Range of vertices and indices used when rendering this section. */
+	uint32 FirstIndex;
+	uint32 NumTriangles;
+};
+
 USTRUCT()
 struct FStaticMeshComponentLODInfo
 {
@@ -64,6 +71,11 @@ struct FStaticMeshComponentLODInfo
 	/** Vertex colors to use for this mesh LOD */
 	FColorVertexBuffer* OverrideVertexColors;
 
+	/** Information for each section about what range of PreCulledIndexBuffer to use.  If no preculled index data is available, PreCulledSections will be empty. */
+	TArray<FPreCulledStaticMeshSection> PreCulledSections;
+
+	FRawStaticIndexBuffer PreCulledIndexBuffer;
+
 #if WITH_EDITORONLY_DATA
 	/** Owner of this FStaticMeshComponentLODInfo */
 	class UStaticMeshComponent* Owner;
@@ -412,13 +424,14 @@ public:
 	 */
 	void RemoveInstanceVertexColors();
 
+	void UpdatePreCulledData(int32 LODIndex, const TArray<uint32>& PreCulledData, const TArray<int32>& NumTrianglesPerSection);
+
 	/**
 	*	Sets the value of the SectionIndexPreview flag and reattaches the component as necessary.
 	*	@param	InSectionIndexPreview		New value of SectionIndexPreview.
 	*/
 	void SetSectionPreview(int32 InSectionIndexPreview);
 
-
 private:
 	/** Initializes the resources used by the static mesh component. */
 	void InitResources();
diff --git a/Engine/Source/Runtime/Engine/Classes/Engine/PreCullTrianglesVolume.h b/Engine/Source/Runtime/Engine/Classes/Engine/PreCullTrianglesVolume.h
new file mode 100644
--- /dev/null
+++ b/Engine/Source/Runtime/Engine/Classes/Engine/PreCullTrianglesVolume.h
@@ -0,0 +1,19 @@
+// Copyright 1998-2015 Epic Games, Inc. All Rights Reserved.
+
+//=============================================================================
+// PreCullTrianglesVolume:  
+//=============================================================================
+#pragma once
+
+#include "GameFramework/Volume.h"
+#include "PreCullTrianglesVolume.generated.h"
+
+UCLASS(hidecategories=(Collision, Brush, Attachment, Physics, Volume), MinimalAPI)
+class APreCullTrianglesVolume : public AVolume
+{
+	GENERATED_UCLASS_BODY()
+
+};
+
+
+
diff --git a/Engine/Source/Runtime/Engine/Private/Components/StaticMeshComponent.cpp b/Engine/Source/Runtime/Engine/Private/Components/StaticMeshComponent.cpp
--- a/Engine/Source/Runtime/Engine/Private/Components/StaticMeshComponent.cpp
+++ b/Engine/Source/Runtime/Engine/Private/Components/StaticMeshComponent.cpp
@@ -24,6 +24,7 @@
 #define LOCTEXT_NAMESPACE "StaticMeshComponent"
 
 DECLARE_MEMORY_STAT( TEXT( "StaticMesh VxColor Inst Mem" ), STAT_InstVertexColorMemory, STATGROUP_MemoryStaticMesh );
+DECLARE_MEMORY_STAT( TEXT( "StaticMesh PreCulled Index Memory" ), STAT_StaticMeshPreCulledIndexMemory, STATGROUP_MemoryStaticMesh );
 
 class FStaticMeshComponentInstanceData : public FSceneComponentInstanceData
 {
@@ -950,13 +951,71 @@ void UStaticMeshComponent::InitResources()
 	}
 }
 
+float GKeepPreCulledIndicesThreshold = .95f;
 
+FAutoConsoleVariableRef CKeepPreCulledIndicesThreshold(
+	TEXT("r.KeepPreCulledIndicesThreshold"),
+	GKeepPreCulledIndicesThreshold,
+	TEXT(""),
+	ECVF_Scalability | ECVF_RenderThreadSafe
+	);
+
+void UStaticMeshComponent::UpdatePreCulledData(int32 LODIndex, const TArray<uint32>& PreCulledData, const TArray<int32>& NumTrianglesPerSection)
+{
+	const FStaticMeshLODResources& StaticMeshLODResources = StaticMesh->RenderData->LODResources[LODIndex];
+
+	int32 NumOriginalTriangles = 0;
+	int32 NumVisibleTriangles = 0;
+
+	for (int32 SectionIndex = 0; SectionIndex < StaticMeshLODResources.Sections.Num(); SectionIndex++)
+	{
+		const FStaticMeshSection& Section = StaticMeshLODResources.Sections[SectionIndex];
+		NumOriginalTriangles += Section.NumTriangles;
+		NumVisibleTriangles += NumTrianglesPerSection[SectionIndex];
+	}
+
+	if (NumVisibleTriangles / (float)NumOriginalTriangles < GKeepPreCulledIndicesThreshold)
+	{
+		SetLODDataCount(LODIndex + 1, LODData.Num());
+
+		DEC_DWORD_STAT_BY(STAT_StaticMeshPreCulledIndexMemory, LODData[LODIndex].PreCulledIndexBuffer.GetAllocatedSize());
+		//@todo - game thread
+		check(IsInRenderingThread());
+		LODData[LODIndex].PreCulledIndexBuffer.ReleaseResource();
+		LODData[LODIndex].PreCulledIndexBuffer.SetIndices(PreCulledData, EIndexBufferStride::AutoDetect);
+		LODData[LODIndex].PreCulledIndexBuffer.InitResource();
+
+		INC_DWORD_STAT_BY(STAT_StaticMeshPreCulledIndexMemory, LODData[LODIndex].PreCulledIndexBuffer.GetAllocatedSize());
+		LODData[LODIndex].PreCulledSections.Empty(StaticMeshLODResources.Sections.Num());
+
+		int32 FirstIndex = 0;
+
+		for (int32 SectionIndex = 0; SectionIndex < StaticMeshLODResources.Sections.Num(); SectionIndex++)
+		{
+			const FStaticMeshSection& Section = StaticMeshLODResources.Sections[SectionIndex];
+			FPreCulledStaticMeshSection PreCulledSection;
+			PreCulledSection.FirstIndex = FirstIndex;
+			PreCulledSection.NumTriangles = NumTrianglesPerSection[SectionIndex];
+			FirstIndex += PreCulledSection.NumTriangles * 3;
+			LODData[LODIndex].PreCulledSections.Add(PreCulledSection);
+		}
+	}
+	else if (LODIndex < LODData.Num())
+	{
+		LODData[LODIndex].PreCulledIndexBuffer.ReleaseResource();
+		TArray<uint32> EmptyIndices;
+		LODData[LODIndex].PreCulledIndexBuffer.SetIndices(EmptyIndices, EIndexBufferStride::AutoDetect);
+		LODData[LODIndex].PreCulledSections.Empty(StaticMeshLODResources.Sections.Num());
+	}
+}
 
 void UStaticMeshComponent::ReleaseResources()
 {
 	for(int32 LODIndex = 0;LODIndex < LODData.Num();LODIndex++)
 	{
 		LODData[LODIndex].BeginReleaseOverrideVertexColors();
+		DEC_DWORD_STAT_BY(STAT_StaticMeshPreCulledIndexMemory, LODData[LODIndex].PreCulledIndexBuffer.GetAllocatedSize());
+		BeginReleaseResource(&LODData[LODIndex].PreCulledIndexBuffer);
 	}
 
 	DetachFence.BeginFence();
diff --git a/Engine/Source/Runtime/Engine/Private/HierarchicalInstancedStaticMesh.cpp b/Engine/Source/Runtime/Engine/Private/HierarchicalInstancedStaticMesh.cpp
--- a/Engine/Source/Runtime/Engine/Private/HierarchicalInstancedStaticMesh.cpp
+++ b/Engine/Source/Runtime/Engine/Private/HierarchicalInstancedStaticMesh.cpp
@@ -1097,7 +1097,7 @@ void FHierarchicalStaticMeshSceneProxy::FillDynamicMeshElements(FMeshElementColl
 					FMeshBatch& MeshElement = Collector.AllocateMesh();
 					INC_DWORD_STAT(STAT_FoliageMeshBatches);
 
-					if (!FStaticMeshSceneProxy::GetMeshElement(LODIndex, 0, SectionIndex, GetDepthPriorityGroup(ElementParams.View), ElementParams.BatchRenderSelection[SelectionGroupIndex], ElementParams.bUseHoveredMaterial, MeshElement))
+					if (!FStaticMeshSceneProxy::GetMeshElement(LODIndex, 0, SectionIndex, GetDepthPriorityGroup(ElementParams.View), ElementParams.BatchRenderSelection[SelectionGroupIndex], ElementParams.bUseHoveredMaterial, true, MeshElement))
 					{
 						continue;
 					}
diff --git a/Engine/Source/Runtime/Engine/Private/InstancedStaticMesh.cpp b/Engine/Source/Runtime/Engine/Private/InstancedStaticMesh.cpp
--- a/Engine/Source/Runtime/Engine/Private/InstancedStaticMesh.cpp
+++ b/Engine/Source/Runtime/Engine/Private/InstancedStaticMesh.cpp
@@ -531,7 +531,7 @@ void FInstancedStaticMeshSceneProxy::GetDynamicMeshElements(const TArray<const F
 					{
 						FMeshBatch& MeshElement = Collector.AllocateMesh();
 
-						if (GetMeshElement(LODIndex, BatchIndex, SectionIndex, GetDepthPriorityGroup(View), BatchRenderSelection[SelectionGroupIndex], IsHovered(), MeshElement))
+						if (GetMeshElement(LODIndex, BatchIndex, SectionIndex, GetDepthPriorityGroup(View), BatchRenderSelection[SelectionGroupIndex], IsHovered(), true, MeshElement))
 						{
 							//@todo-rco this is only supporting selection on the first element
 							MeshElement.Elements[0].UserData = PassUserData[SelectionGroupIndex];
@@ -629,9 +629,9 @@ bool FInstancedStaticMeshSceneProxy::GetShadowMeshElement(int32 LODIndex, int32
 }
 
 /** Sets up a FMeshBatch for a specific LOD and element. */
-bool FInstancedStaticMeshSceneProxy::GetMeshElement(int32 LODIndex, int32 BatchIndex, int32 ElementIndex, uint8 InDepthPriorityGroup, const bool bUseSelectedMaterial, const bool bUseHoveredMaterial, FMeshBatch& OutMeshBatch) const
+bool FInstancedStaticMeshSceneProxy::GetMeshElement(int32 LODIndex, int32 BatchIndex, int32 ElementIndex, uint8 InDepthPriorityGroup, bool bUseSelectedMaterial, bool bUseHoveredMaterial, bool bAllowPreCulledIndices, FMeshBatch& OutMeshBatch) const
 {
-	if (LODIndex < InstancedRenderData.VertexFactories.Num() && FStaticMeshSceneProxy::GetMeshElement(LODIndex, BatchIndex, ElementIndex, InDepthPriorityGroup, bUseSelectedMaterial, bUseHoveredMaterial, OutMeshBatch))
+	if (LODIndex < InstancedRenderData.VertexFactories.Num() && FStaticMeshSceneProxy::GetMeshElement(LODIndex, BatchIndex, ElementIndex, InDepthPriorityGroup, bUseSelectedMaterial, bUseHoveredMaterial, bAllowPreCulledIndices, OutMeshBatch))
 	{
 		SetupInstancedMeshBatch(LODIndex, BatchIndex, OutMeshBatch);
 		return true;
@@ -640,9 +640,9 @@ bool FInstancedStaticMeshSceneProxy::GetMeshElement(int32 LODIndex, int32 BatchI
 };
 
 /** Sets up a wireframe FMeshBatch for a specific LOD. */
-bool FInstancedStaticMeshSceneProxy::GetWireframeMeshElement(int32 LODIndex, int32 BatchIndex, const FMaterialRenderProxy* WireframeRenderProxy, uint8 InDepthPriorityGroup, FMeshBatch& OutMeshBatch) const
+bool FInstancedStaticMeshSceneProxy::GetWireframeMeshElement(int32 LODIndex, int32 BatchIndex, const FMaterialRenderProxy* WireframeRenderProxy, uint8 InDepthPriorityGroup, bool bAllowPreCulledIndices, FMeshBatch& OutMeshBatch) const
 {
-	if (LODIndex < InstancedRenderData.VertexFactories.Num() && FStaticMeshSceneProxy::GetWireframeMeshElement(LODIndex, BatchIndex, WireframeRenderProxy, InDepthPriorityGroup, OutMeshBatch))
+	if (LODIndex < InstancedRenderData.VertexFactories.Num() && FStaticMeshSceneProxy::GetWireframeMeshElement(LODIndex, BatchIndex, WireframeRenderProxy, InDepthPriorityGroup, bAllowPreCulledIndices, OutMeshBatch))
 	{
 		SetupInstancedMeshBatch(LODIndex, BatchIndex, OutMeshBatch);
 		return true;
diff --git a/Engine/Source/Runtime/Engine/Private/InstancedStaticMesh.h b/Engine/Source/Runtime/Engine/Private/InstancedStaticMesh.h
--- a/Engine/Source/Runtime/Engine/Private/InstancedStaticMesh.h
+++ b/Engine/Source/Runtime/Engine/Private/InstancedStaticMesh.h
@@ -553,10 +553,10 @@ public:
 	virtual bool GetShadowMeshElement(int32 LODIndex, int32 BatchIndex, uint8 InDepthPriorityGroup, FMeshBatch& OutMeshBatch, bool bDitheredLODTransition) const override;
 
 	/** Sets up a FMeshBatch for a specific LOD and element. */
-	virtual bool GetMeshElement(int32 LODIndex, int32 BatchIndex, int32 ElementIndex, uint8 InDepthPriorityGroup, const bool bUseSelectedMaterial, const bool bUseHoveredMaterial, FMeshBatch& OutMeshBatch) const override;
+	virtual bool GetMeshElement(int32 LODIndex, int32 BatchIndex, int32 ElementIndex, uint8 InDepthPriorityGroup, bool bUseSelectedMaterial, bool bUseHoveredMaterial, bool bAllowPreCulledIndices, FMeshBatch& OutMeshBatch) const override;
 
 	/** Sets up a wireframe FMeshBatch for a specific LOD. */
-	virtual bool GetWireframeMeshElement(int32 LODIndex, int32 BatchIndex, const FMaterialRenderProxy* WireframeRenderProxy, uint8 InDepthPriorityGroup, FMeshBatch& OutMeshBatch) const override;
+	virtual bool GetWireframeMeshElement(int32 LODIndex, int32 BatchIndex, const FMaterialRenderProxy* WireframeRenderProxy, uint8 InDepthPriorityGroup, bool bAllowPreCulledIndices, FMeshBatch& OutMeshBatch) const override;
 
 	virtual void GetDistancefieldAtlasData(FBox& LocalVolumeBounds, FIntVector& OutBlockMin, FIntVector& OutBlockSize, bool& bOutBuiltAsIfTwoSided, bool& bMeshWasPlane, TArray<FMatrix>& ObjectLocalToWorldTransforms) const override;
 
diff --git a/Engine/Source/Runtime/Engine/Private/PreCullTrianglesSetup.cpp b/Engine/Source/Runtime/Engine/Private/PreCullTrianglesSetup.cpp
new file mode 100644
--- /dev/null
+++ b/Engine/Source/Runtime/Engine/Private/PreCullTrianglesSetup.cpp
@@ -0,0 +1,56 @@
+// Copyright 1998-2015 Epic Games, Inc. All Rights Reserved.
+
+#include "EnginePrivate.h"
+#include "StaticMeshResources.h"
+#include "ComponentRecreateRenderStateContext.h"
+#include "Engine/PreCullTrianglesVolume.h"
+
+APreCullTrianglesVolume::APreCullTrianglesVolume(const FObjectInitializer& ObjectInitializer)
+	: Super(ObjectInitializer)
+{
+}
+
+void PreCullTrianglesExec(UWorld* InWorld)
+{
+	TArray<UStaticMeshComponent*> ComponentsToPreCull;
+
+	for (TObjectIterator<UStaticMeshComponent> It; It; ++It)
+	{
+		UStaticMeshComponent* Component = *It;
+
+		if (!Component->IsPendingKill() 
+			&& Component->Mobility == EComponentMobility::Static
+			&& Component->GetOwner() 
+			&& InWorld->ContainsActor(Component->GetOwner()) 
+			&& Component->StaticMesh 
+			&& Component->IsRegistered()
+			&& Component->SceneProxy)
+		{
+			ComponentsToPreCull.Add(Component);
+		}
+	}
+
+	TArray<TArray<FPlane> > CullVolumes;
+
+	for (TObjectIterator<APreCullTrianglesVolume> It; It; ++It)
+	{
+		APreCullTrianglesVolume* CullVolume = *It;
+
+		if (InWorld->ContainsActor(CullVolume) && !CullVolume->IsPendingKill())
+		{
+			TArray<FPlane> Planes;
+			CullVolume->Brush->GetSurfacePlanes(CullVolume, Planes);
+
+			CullVolumes.Add(Planes);
+		}
+	}
+
+	InWorld->Scene->PreCullStaticMeshes(ComponentsToPreCull, CullVolumes);
+	FGlobalComponentRecreateRenderStateContext Context;
+}
+
+FAutoConsoleCommandWithWorld GPreCullIndexBuffersCmd(
+	TEXT("r.PreCullIndexBuffers"),
+	TEXT("Runs a slow operation to remove any static mesh triangles that are invisible or inside a precull volume."),
+	FConsoleCommandWithWorldDelegate::CreateStatic(PreCullTrianglesExec)
+	);
\ No newline at end of file
diff --git a/Engine/Source/Runtime/Engine/Private/SplineMeshSceneProxy.cpp b/Engine/Source/Runtime/Engine/Private/SplineMeshSceneProxy.cpp
--- a/Engine/Source/Runtime/Engine/Private/SplineMeshSceneProxy.cpp
+++ b/Engine/Source/Runtime/Engine/Private/SplineMeshSceneProxy.cpp
@@ -81,11 +81,11 @@ bool FSplineMeshSceneProxy::GetShadowMeshElement(int32 LODIndex, int32 BatchInde
 	return false;
 }
 
-bool FSplineMeshSceneProxy::GetMeshElement(int32 LODIndex, int32 BatchIndex, int32 SectionIndex, uint8 InDepthPriorityGroup, const bool bUseSelectedMaterial, const bool bUseHoveredMaterial, FMeshBatch& OutMeshBatch) const
+bool FSplineMeshSceneProxy::GetMeshElement(int32 LODIndex, int32 BatchIndex, int32 SectionIndex, uint8 InDepthPriorityGroup, bool bUseSelectedMaterial, bool bUseHoveredMaterial, bool bAllowPreCulledIndices, FMeshBatch& OutMeshBatch) const
 {
 	//checkf(LODIndex == 0 /*&& SectionIndex == 0*/, TEXT("Getting spline static mesh element with invalid params [%d, %d]"), LODIndex, SectionIndex);
 
-	if (FStaticMeshSceneProxy::GetMeshElement(LODIndex, BatchIndex, SectionIndex, InDepthPriorityGroup, bUseSelectedMaterial, bUseHoveredMaterial, OutMeshBatch))
+	if (FStaticMeshSceneProxy::GetMeshElement(LODIndex, BatchIndex, SectionIndex, InDepthPriorityGroup, bUseSelectedMaterial, bUseHoveredMaterial, bAllowPreCulledIndices, OutMeshBatch))
 	{
 		OutMeshBatch.VertexFactory = LODResources[LODIndex].VertexFactory;
 		OutMeshBatch.ReverseCulling ^= (SplineParams.StartScale.X < 0) ^ (SplineParams.StartScale.Y < 0);
@@ -94,11 +94,11 @@ bool FSplineMeshSceneProxy::GetMeshElement(int32 LODIndex, int32 BatchIndex, int
 	return false;
 }
 
-bool FSplineMeshSceneProxy::GetWireframeMeshElement(int32 LODIndex, int32 BatchIndex, const FMaterialRenderProxy* WireframeRenderProxy, uint8 InDepthPriorityGroup, FMeshBatch& OutMeshBatch) const
+bool FSplineMeshSceneProxy::GetWireframeMeshElement(int32 LODIndex, int32 BatchIndex, const FMaterialRenderProxy* WireframeRenderProxy, uint8 InDepthPriorityGroup, bool bAllowPreCulledIndices, FMeshBatch& OutMeshBatch) const
 {
 	//checkf(LODIndex == 0, TEXT("Getting spline static mesh element with invalid LOD [%d]"), LODIndex);
 
-	if (FStaticMeshSceneProxy::GetWireframeMeshElement(LODIndex, BatchIndex, WireframeRenderProxy, InDepthPriorityGroup, OutMeshBatch))
+	if (FStaticMeshSceneProxy::GetWireframeMeshElement(LODIndex, BatchIndex, WireframeRenderProxy, InDepthPriorityGroup, bAllowPreCulledIndices, OutMeshBatch))
 	{
 		OutMeshBatch.VertexFactory = LODResources[LODIndex].VertexFactory;
 		OutMeshBatch.ReverseCulling ^= (SplineParams.StartScale.X < 0) ^ (SplineParams.StartScale.Y < 0);
diff --git a/Engine/Source/Runtime/Engine/Private/SplineMeshSceneProxy.h b/Engine/Source/Runtime/Engine/Private/SplineMeshSceneProxy.h
--- a/Engine/Source/Runtime/Engine/Private/SplineMeshSceneProxy.h
+++ b/Engine/Source/Runtime/Engine/Private/SplineMeshSceneProxy.h
@@ -157,10 +157,10 @@ public:
 	virtual bool GetShadowMeshElement(int32 LODIndex, int32 BatchIndex, uint8 InDepthPriorityGroup, FMeshBatch& OutMeshBatch, bool bDitheredLODTransition) const override;
 
 	/** Sets up a FMeshBatch for a specific LOD and element. */
-	virtual bool GetMeshElement(int32 LODIndex, int32 BatchIndex, int32 ElementIndex, uint8 InDepthPriorityGroup, const bool bUseSelectedMaterial, const bool bUseHoveredMaterial, FMeshBatch& OutMeshBatch) const override;
+	virtual bool GetMeshElement(int32 LODIndex, int32 BatchIndex, int32 ElementIndex, uint8 InDepthPriorityGroup, bool bUseSelectedMaterial, bool bUseHoveredMaterial, bool bAllowPreCulledIndices, FMeshBatch& OutMeshBatch) const override;
 
 	/** Sets up a wireframe FMeshBatch for a specific LOD. */
-	virtual bool GetWireframeMeshElement(int32 LODIndex, int32 BatchIndex, const FMaterialRenderProxy* WireframeRenderProxy, uint8 InDepthPriorityGroup, FMeshBatch& OutMeshBatch) const override;
+	virtual bool GetWireframeMeshElement(int32 LODIndex, int32 BatchIndex, const FMaterialRenderProxy* WireframeRenderProxy, uint8 InDepthPriorityGroup, bool bAllowPreCulledIndices, FMeshBatch& OutMeshBatch) const override;
 
 	virtual FPrimitiveViewRelevance GetViewRelevance(const FSceneView* View) const override
 	{
diff --git a/Engine/Source/Runtime/Engine/Private/StaticMeshRender.cpp b/Engine/Source/Runtime/Engine/Private/StaticMeshRender.cpp
--- a/Engine/Source/Runtime/Engine/Private/StaticMeshRender.cpp
+++ b/Engine/Source/Runtime/Engine/Private/StaticMeshRender.cpp
@@ -16,6 +16,7 @@
 #include "ComponentReregisterContext.h"
 #include "Components/BrushComponent.h"
 #include "AI/Navigation/NavCollision.h"
+#include "ComponentRecreateRenderStateContext.h"
 
 /** If true, optimized depth-only index buffers are used for shadow rendering. */
 static bool GUseShadowIndexBuffer = true;
@@ -45,6 +46,20 @@ static FAutoConsoleCommand GToggleShadowIndexBuffersCmd(
 	FConsoleCommandDelegate::CreateStatic(ToggleShadowIndexBuffers)
 	);
 
+static bool GUsePreCulledIndexBuffer = true;
+
+void TogglePreCulledIndexBuffers( UWorld* InWorld )
+{
+	FGlobalComponentRecreateRenderStateContext Context;
+	FlushRenderingCommands();
+	GUsePreCulledIndexBuffer = !GUsePreCulledIndexBuffer;
+}
+
+FAutoConsoleCommandWithWorld GToggleUsePreCulledIndexBuffersCmd(
+	TEXT("r.TogglePreCulledIndexBuffers"),
+	TEXT("Toggles use of preculled index buffers from the command 'PreCullIndexBuffers'"),
+	FConsoleCommandWithWorldDelegate::CreateStatic(TogglePreCulledIndexBuffers));
+
 static FAutoConsoleCommand GToggleReversedIndexBuffersCmd(
 	TEXT("ToggleReversedIndexBuffers"),
 	TEXT("Render static meshes with negative transform determinants using a reversed index buffer."),
@@ -239,7 +254,15 @@ bool FStaticMeshSceneProxy::GetShadowMeshElement(int32 LODIndex, int32 BatchInde
 }
 
 /** Sets up a FMeshBatch for a specific LOD and element. */
-bool FStaticMeshSceneProxy::GetMeshElement(int32 LODIndex, int32 BatchIndex, int32 SectionIndex, uint8 InDepthPriorityGroup, const bool bUseSelectedMaterial, const bool bUseHoveredMaterial, FMeshBatch& OutMeshBatch) const
+bool FStaticMeshSceneProxy::GetMeshElement(
+	int32 LODIndex, 
+	int32 BatchIndex, 
+	int32 SectionIndex, 
+	uint8 InDepthPriorityGroup, 
+	bool bUseSelectedMaterial, 
+	bool bUseHoveredMaterial, 
+	bool bAllowPreCulledIndices, 
+	FMeshBatch& OutMeshBatch) const
 {
 	const FStaticMeshLODResources& LOD = RenderData->LODResources[LODIndex];
 	const FStaticMeshSection& Section = LOD.Sections[SectionIndex];
@@ -277,7 +300,7 @@ bool FStaticMeshSceneProxy::GetMeshElement(int32 LODIndex, int32 BatchIndex, int
 	// Two sided material use bIsFrontFace which is wrong with Reversed Indices. AdjacencyInformation use another index buffer.
 	const bool bUseReversedIndices = GUseReversedIndexBuffer && IsLocalToWorldDeterminantNegative() && LOD.bHasReversedIndices && !bWireframe && !bRequiresAdjacencyInformation && !Material->IsTwoSided();
 
-	SetIndexSource(LODIndex, SectionIndex, OutMeshBatch, bWireframe, bRequiresAdjacencyInformation, bUseReversedIndices);
+	SetIndexSource(LODIndex, SectionIndex, OutMeshBatch, bWireframe, bRequiresAdjacencyInformation, bUseReversedIndices, bAllowPreCulledIndices);
 
 	FMeshBatchElement& OutBatchElement = OutMeshBatch.Elements[0];
 
@@ -323,7 +346,7 @@ bool FStaticMeshSceneProxy::GetMeshElement(int32 LODIndex, int32 BatchIndex, int
 }
 
 /** Sets up a wireframe FMeshBatch for a specific LOD. */
-bool FStaticMeshSceneProxy::GetWireframeMeshElement(int32 LODIndex, int32 BatchIndex, const FMaterialRenderProxy* WireframeRenderProxy, uint8 InDepthPriorityGroup, FMeshBatch& OutMeshBatch) const
+bool FStaticMeshSceneProxy::GetWireframeMeshElement(int32 LODIndex, int32 BatchIndex, const FMaterialRenderProxy* WireframeRenderProxy, uint8 InDepthPriorityGroup, bool bAllowPreCulledIndices, FMeshBatch& OutMeshBatch) const
 {
 	const FStaticMeshLODResources& LODModel = RenderData->LODResources[LODIndex];
 	
@@ -356,7 +379,7 @@ bool FStaticMeshSceneProxy::GetWireframeMeshElement(int32 LODIndex, int32 BatchI
 	const bool bRequiresAdjacencyInformation = false;
 	const bool bUseReversedIndices = false;
 
-	SetIndexSource(LODIndex, 0, OutMeshBatch, bWireframe, bRequiresAdjacencyInformation, bUseReversedIndices);
+	SetIndexSource(LODIndex, 0, OutMeshBatch, bWireframe, bRequiresAdjacencyInformation, bUseReversedIndices, bAllowPreCulledIndices);
 
 	return OutBatchElement.NumPrimitives > 0;
 }
@@ -364,7 +387,7 @@ bool FStaticMeshSceneProxy::GetWireframeMeshElement(int32 LODIndex, int32 BatchI
 /**
  * Sets IndexBuffer, FirstIndex and NumPrimitives of OutMeshElement.
  */
-void FStaticMeshSceneProxy::SetIndexSource(int32 LODIndex, int32 SectionIndex, FMeshBatch& OutMeshElement, bool bWireframe, bool bRequiresAdjacencyInformation, bool bUseReversedIndices) const
+void FStaticMeshSceneProxy::SetIndexSource(int32 LODIndex, int32 SectionIndex, FMeshBatch& OutMeshElement, bool bWireframe, bool bRequiresAdjacencyInformation, bool bUseReversedIndices, bool bAllowPreCulledIndices) const
 {
 	FMeshBatchElement& OutElement = OutMeshElement.Elements[0];
 	const FStaticMeshLODResources& LODModel = RenderData->LODResources[LODIndex];
@@ -382,9 +405,22 @@ void FStaticMeshSceneProxy::SetIndexSource(int32 LODIndex, int32 SectionIndex, F
 		else
 		{
 			OutMeshElement.Type = PT_TriangleList;
-			OutElement.FirstIndex = 0;
-			OutElement.IndexBuffer = &LODModel.IndexBuffer;
-			OutElement.NumPrimitives = LODModel.IndexBuffer.GetNumIndices() / 3;
+
+			if (bAllowPreCulledIndices
+				&& GUsePreCulledIndexBuffer 
+				&& LODs[LODIndex].Sections[SectionIndex].NumPreCulledTriangles >= 0)
+			{
+				OutElement.IndexBuffer = LODs[LODIndex].PreCulledIndexBuffer;
+				OutElement.FirstIndex = 0;
+				OutElement.NumPrimitives = LODs[LODIndex].PreCulledIndexBuffer->GetNumIndices() / 3;
+			}
+			else
+			{
+				OutElement.FirstIndex = 0;
+				OutElement.IndexBuffer = &LODModel.IndexBuffer;
+				OutElement.NumPrimitives = LODModel.IndexBuffer.GetNumIndices() / 3;
+			}
+
 			OutMeshElement.bWireframe = true;
 			OutMeshElement.bDisableBackfaceCulling = true;
 		}
@@ -393,9 +429,21 @@ void FStaticMeshSceneProxy::SetIndexSource(int32 LODIndex, int32 SectionIndex, F
 	{
 		const FStaticMeshSection& Section = LODModel.Sections[SectionIndex];
 		OutMeshElement.Type = PT_TriangleList;
-		OutElement.IndexBuffer = bUseReversedIndices ? &LODModel.ReversedIndexBuffer : &LODModel.IndexBuffer;
-		OutElement.FirstIndex = Section.FirstIndex;
-		OutElement.NumPrimitives = Section.NumTriangles;
+
+		if (bAllowPreCulledIndices
+			&& GUsePreCulledIndexBuffer 
+			&& LODs[LODIndex].Sections[SectionIndex].NumPreCulledTriangles >= 0)
+		{
+			OutElement.IndexBuffer = LODs[LODIndex].PreCulledIndexBuffer;
+			OutElement.FirstIndex = LODs[LODIndex].Sections[SectionIndex].FirstPreCulledIndex;
+			OutElement.NumPrimitives = LODs[LODIndex].Sections[SectionIndex].NumPreCulledTriangles;
+		}
+		else
+		{
+			OutElement.IndexBuffer = bUseReversedIndices ? &LODModel.ReversedIndexBuffer : &LODModel.IndexBuffer;
+			OutElement.FirstIndex = Section.FirstIndex;
+			OutElement.NumPrimitives = Section.NumTriangles;
+		}
 	}
 
 	if ( bRequiresAdjacencyInformation )
@@ -513,7 +561,7 @@ void FStaticMeshSceneProxy::DrawStaticElements(FStaticPrimitiveDrawInterface* PD
 				{
 					FMeshBatch MeshBatch;
 
-					if (GetMeshElement(LODIndex, BatchIndex, SectionIndex, PrimitiveDPG, bUseSelectedMaterial, bUseHoveredMaterial, MeshBatch))
+					if (GetMeshElement(LODIndex, BatchIndex, SectionIndex, PrimitiveDPG, bUseSelectedMaterial, bUseHoveredMaterial, true, MeshBatch))
 					{
 						PDI->DrawMesh(MeshBatch, FLT_MAX);
 					}
@@ -611,7 +659,7 @@ void FStaticMeshSceneProxy::DrawStaticElements(FStaticPrimitiveDrawInterface* PD
 					{
 						FMeshBatch MeshBatch;
 
-						if (GetMeshElement(LODIndex, BatchIndex, SectionIndex, PrimitiveDPG, bUseSelectedMaterial, bUseHoveredMaterial, MeshBatch))
+						if (GetMeshElement(LODIndex, BatchIndex, SectionIndex, PrimitiveDPG, bUseSelectedMaterial, bUseHoveredMaterial, true, MeshBatch))
 						{
 							// If we have submitted an optimized shadow-only mesh, remaining mesh elements must not cast shadows.
 							MeshBatch.CastShadow &= !bUseUnifiedMeshForShadow;
@@ -651,6 +699,27 @@ bool FStaticMeshSceneProxy::IsCollisionView(const FEngineShowFlags& EngineShowFl
 	return bInCollisionView;
 }
 
+void FStaticMeshSceneProxy::GetMeshDescription(int32 LODIndex, TArray<FMeshBatch>& OutMeshElements) const
+{
+	const FStaticMeshLODResources& LODModel = RenderData->LODResources[LODIndex];
+	const FLODInfo& ProxyLODInfo = LODs[LODIndex];
+
+	for (int32 SectionIndex = 0; SectionIndex < LODModel.Sections.Num(); SectionIndex++)
+	{
+		const int32 NumBatches = GetNumMeshBatches();
+
+		for (int32 BatchIndex = 0; BatchIndex < NumBatches; BatchIndex++)
+		{
+			FMeshBatch MeshElement; 
+
+			if (GetMeshElement(LODIndex, BatchIndex, SectionIndex, SDPG_World, false, false, false, MeshElement))
+			{
+				OutMeshElements.Add(MeshElement);
+			}
+		}
+	}
+}
+
 void FStaticMeshSceneProxy::GetDynamicMeshElements(const TArray<const FSceneView*>& Views, const FSceneViewFamily& ViewFamily, uint32 VisibilityMap, FMeshElementCollector& Collector) const
 {
 	QUICK_SCOPE_CYCLE_COUNTER(STAT_StaticMeshSceneProxy_GetMeshElements);
@@ -721,7 +790,7 @@ void FStaticMeshSceneProxy::GetDynamicMeshElements(const TArray<const FSceneView
 							{
 								FMeshBatch& Mesh = Collector.AllocateMesh();
 
-								if (GetWireframeMeshElement(LODIndex, BatchIndex, WireframeMaterialInstance, SDPG_World, Mesh))
+								if (GetWireframeMeshElement(LODIndex, BatchIndex, WireframeMaterialInstance, SDPG_World, true, Mesh))
 								{
 									// We implemented our own wireframe
 									Mesh.bCanApplyViewModeOverrides = false;
@@ -754,7 +823,7 @@ void FStaticMeshSceneProxy::GetDynamicMeshElements(const TArray<const FSceneView
 									}
 	#endif // WITH_EDITOR
 								
-									if (GetMeshElement(LODIndex, BatchIndex, SectionIndex, SDPG_World, bSectionIsSelected, IsHovered(), MeshElement))
+									if (GetMeshElement(LODIndex, BatchIndex, SectionIndex, SDPG_World, bSectionIsSelected, IsHovered(), true, MeshElement))
 									{
 										bool bAddMesh = true;
 
@@ -815,7 +884,7 @@ void FStaticMeshSceneProxy::GetDynamicMeshElements(const TArray<const FSceneView
 													if (bDrawComplexWireframeCollision)
 													{
 														FMeshBatch& CollisionElement = Collector.AllocateMesh();
-														GetMeshElement(LODIndex, BatchIndex, SectionIndex, SDPG_World, bSectionIsSelected, IsHovered(), CollisionElement);
+														GetMeshElement(LODIndex, BatchIndex, SectionIndex, SDPG_World, bSectionIsSelected, IsHovered(), true, CollisionElement);
 														CollisionElement.MaterialRenderProxy = CollisionMaterialInstance;
 														Collector.AddMesh(ViewIndex, CollisionElement);
 														INC_DWORD_STAT_BY(STAT_StaticMeshTriangles, CollisionElement.GetNumPrimitives());
@@ -1101,6 +1170,7 @@ bool FStaticMeshSceneProxy::HasDistanceFieldRepresentation() const
 /** Initialization constructor. */
 FStaticMeshSceneProxy::FLODInfo::FLODInfo(const UStaticMeshComponent* InComponent,int32 LODIndex):
 	OverrideColorVertexBuffer(0),
+	PreCulledIndexBuffer(NULL),
 	LightMap(NULL),
 	ShadowMap(NULL),
 	bUsesMeshModifyingMaterials(false)
@@ -1116,6 +1186,7 @@ FStaticMeshSceneProxy::FLODInfo::FLODInfo(const UStaticMeshComponent* InComponen
 		LightMap = ComponentLODInfo.LightMap;
 		ShadowMap = ComponentLODInfo.ShadowMap;
 		IrrelevantLights = InComponent->IrrelevantLights;
+		PreCulledIndexBuffer = &ComponentLODInfo.PreCulledIndexBuffer;
 
 		// Initialize this LOD's overridden vertex colors, if it has any
 		if( ComponentLODInfo.OverrideVertexColors )
@@ -1182,6 +1253,17 @@ FStaticMeshSceneProxy::FLODInfo::FLODInfo(const UStaticMeshComponent* InComponen
 		}
 #endif
 
+		if (LODIndex < InComponent->LODData.Num())
+		{
+			const FStaticMeshComponentLODInfo& ComponentLODInfo = InComponent->LODData[LODIndex];
+
+			if (SectionIndex < ComponentLODInfo.PreCulledSections.Num())
+			{
+				SectionInfo.FirstPreCulledIndex = ComponentLODInfo.PreCulledSections[SectionIndex].FirstIndex;
+				SectionInfo.NumPreCulledTriangles = ComponentLODInfo.PreCulledSections[SectionIndex].NumTriangles;
+			}
+		}
+
 		// Store the element info.
 		Sections.Add(SectionInfo);
 
diff --git a/Engine/Source/Runtime/Engine/Public/PrimitiveSceneProxy.h b/Engine/Source/Runtime/Engine/Public/PrimitiveSceneProxy.h
--- a/Engine/Source/Runtime/Engine/Public/PrimitiveSceneProxy.h
+++ b/Engine/Source/Runtime/Engine/Public/PrimitiveSceneProxy.h
@@ -173,6 +173,9 @@ public:
 	 */
 	virtual void DrawStaticElements(FStaticPrimitiveDrawInterface* PDI) {}
 
+	/** Gathers a description of the mesh elements to be rendered for the given LOD index, without consideration for views. */
+	virtual void GetMeshDescription(int32 LODIndex, TArray<FMeshBatch>& OutMeshElements) const {};
+
 	/** 
 	 * Gathers the primitive's dynamic mesh elements.  This will only be called if GetViewRelevance declares dynamic relevance.
 	 * This is called from the rendering thread for each set of views that might be rendered.  
diff --git a/Engine/Source/Runtime/Engine/Public/SceneInterface.h b/Engine/Source/Runtime/Engine/Public/SceneInterface.h
--- a/Engine/Source/Runtime/Engine/Public/SceneInterface.h
+++ b/Engine/Source/Runtime/Engine/Public/SceneInterface.h
@@ -106,6 +106,9 @@ public:
 	 */
 	virtual void UpdateSkyCaptureContents(const USkyLightComponent* CaptureComponent, bool bCaptureEmissiveOnly, UTextureCube* SourceCubemap, FTexture* OutProcessedTexture, FSHVectorRGB3& OutIrradianceEnvironmentMap) {}
 
+	/** Runs a slow preculling operation on static meshes, removing triangles that are invisible or inside a precull volume. */
+	virtual void PreCullStaticMeshes(const TArray<UStaticMeshComponent*>& ComponentsToPreCull, const TArray<TArray<FPlane> >& CullVolumes) {}
+
 	/** 
 	* Updates the contents of the given scene capture by rendering the scene. 
 	* This must be called on the game thread.
diff --git a/Engine/Source/Runtime/Engine/Public/StaticMeshResources.h b/Engine/Source/Runtime/Engine/Public/StaticMeshResources.h
--- a/Engine/Source/Runtime/Engine/Public/StaticMeshResources.h
+++ b/Engine/Source/Runtime/Engine/Public/StaticMeshResources.h
@@ -766,17 +766,25 @@ public:
 	virtual bool GetShadowMeshElement(int32 LODIndex, int32 BatchIndex, uint8 InDepthPriorityGroup, FMeshBatch& OutMeshBatch, bool bDitheredLODTransition) const;
 
 	/** Sets up a FMeshBatch for a specific LOD and element. */
-	virtual bool GetMeshElement(int32 LODIndex, int32 BatchIndex, int32 ElementIndex, uint8 InDepthPriorityGroup, const bool bUseSelectedMaterial, const bool bUseHoveredMaterial, FMeshBatch& OutMeshBatch) const;
+	virtual bool GetMeshElement(
+		int32 LODIndex, 
+		int32 BatchIndex, 
+		int32 ElementIndex, 
+		uint8 InDepthPriorityGroup, 
+		bool bUseSelectedMaterial, 
+		bool bUseHoveredMaterial, 
+		bool bAllowPreCulledIndices,
+		FMeshBatch& OutMeshBatch) const;
 
 	/** Sets up a wireframe FMeshBatch for a specific LOD. */
-	virtual bool GetWireframeMeshElement(int32 LODIndex, int32 BatchIndex, const FMaterialRenderProxy* WireframeRenderProxy, uint8 InDepthPriorityGroup, FMeshBatch& OutMeshBatch) const;
+	virtual bool GetWireframeMeshElement(int32 LODIndex, int32 BatchIndex, const FMaterialRenderProxy* WireframeRenderProxy, uint8 InDepthPriorityGroup, bool bAllowPreCulledIndices, FMeshBatch& OutMeshBatch) const;
 
 
 protected:
 	/**
 	 * Sets IndexBuffer, FirstIndex and NumPrimitives of OutMeshElement.
 	 */
-	virtual void SetIndexSource(int32 LODIndex, int32 ElementIndex, FMeshBatch& OutMeshElement, bool bWireframe, bool bRequiresAdjacencyInformation, bool bUseInversedIndices) const;
+	virtual void SetIndexSource(int32 LODIndex, int32 ElementIndex, FMeshBatch& OutMeshElement, bool bWireframe, bool bRequiresAdjacencyInformation, bool bUseInversedIndices, bool bAllowPreCulledIndices) const;
 	bool IsCollisionView(const FEngineShowFlags& EngineShowFlags, bool& bDrawSimpleCollision, bool& bDrawComplexCollision) const;
 
 public:
@@ -796,6 +804,8 @@ public:
 	virtual uint32 GetMemoryFootprint( void ) const override { return( sizeof( *this ) + GetAllocatedSize() ); }
 	uint32 GetAllocatedSize( void ) const { return( FPrimitiveSceneProxy::GetAllocatedSize() + LODs.GetAllocatedSize() ); }
 
+	virtual void GetMeshDescription(int32 LODIndex, TArray<FMeshBatch>& OutMeshElements) const override;
+
 	virtual void GetDynamicMeshElements(const TArray<const FSceneView*>& Views, const FSceneViewFamily& ViewFamily, uint32 VisibilityMap, FMeshElementCollector& Collector) const override;
 protected:
 
@@ -814,6 +824,8 @@ protected:
 #if WITH_EDITOR
 				, HitProxy(NULL)
 #endif
+				, FirstPreCulledIndex(0)
+				, NumPreCulledTriangles(-1)
 			{}
 
 			/** The material with which to render this section. */
@@ -826,6 +838,9 @@ protected:
 			/** The editor needs to be able to individual sub-mesh hit detection, so we store a hit proxy on each mesh. */
 			HHitProxy* HitProxy;
 #endif
+
+			int32 FirstPreCulledIndex;
+			int32 NumPreCulledTriangles;
 		};
 
 		/** Per-section information. */
@@ -834,6 +849,8 @@ protected:
 		/** Vertex color data for this LOD (or NULL when not overridden), FStaticMeshComponentLODInfo handle the release of the memory */
 		FColorVertexBuffer* OverrideColorVertexBuffer;
 
+		const FRawStaticIndexBuffer* PreCulledIndexBuffer;
+
 		/** When the mesh component has overridden the LOD's vertex colors, this vertex factory will be created
 		    and passed along to the renderer instead of the mesh's stock vertex factory */
 		TScopedPointer< FLocalVertexFactory > OverrideColorVertexFactory;
diff --git a/Engine/Source/Runtime/Renderer/Private/ConvertToUniformMesh.cpp b/Engine/Source/Runtime/Renderer/Private/ConvertToUniformMesh.cpp
--- a/Engine/Source/Runtime/Renderer/Private/ConvertToUniformMesh.cpp
+++ b/Engine/Source/Runtime/Renderer/Private/ConvertToUniformMesh.cpp
@@ -17,7 +17,6 @@
 #include "SceneUtils.h"
 #include "DistanceFieldAtlas.h"
 
-
 class FConvertToUniformMeshVS : public FMeshMaterialShader
 {
 	DECLARE_SHADER_TYPE(FConvertToUniformMeshVS,MeshMaterial);
@@ -293,15 +292,22 @@ bool ShouldGenerateSurfelsOnMesh(const FMeshBatch& Mesh, ERHIFeatureLevel::Type
 		&& Mesh.MaterialRenderProxy->GetMaterial(FeatureLevel)->GetShadingModel() != MSM_Unlit;
 }
 
-FUniformMeshBuffers GUniformMeshTemporaryBuffers;
+bool ShouldConvertMesh(const FMeshBatch& Mesh)
+{
+	return Mesh.Type == PT_TriangleList
+		//@todo - import types and compare directly
+		&& (FCString::Strstr(Mesh.VertexFactory->GetType()->GetName(), TEXT("LocalVertexFactory")) != NULL
+			|| FCString::Strstr(Mesh.VertexFactory->GetType()->GetName(), TEXT("InstancedStaticMeshVertexFactory")) != NULL);
+}
 
-TArray<TArray<FMeshBatchAndRelevance,SceneRenderingAllocator>> GDynamicMeshElements;
+FUniformMeshBuffers GUniformMeshTemporaryBuffers;
 
 int32 FUniformMeshConverter::Convert(
 	FRHICommandListImmediate& RHICmdList, 
 	FSceneRenderer& Renderer,
 	FViewInfo& View, 
 	const FPrimitiveSceneInfo* PrimitiveSceneInfo, 
+	int32 LODIndex,
 	FUniformMeshBuffers*& OutUniformMeshBuffers,
 	const FMaterialRenderProxy*& OutMaterialRenderProxy,
 	FUniformBufferRHIParamRef& OutPrimitiveUniformBuffer)
@@ -309,68 +315,19 @@ int32 FUniformMeshConverter::Convert(
 	const FPrimitiveSceneProxy* PrimitiveSceneProxy = PrimitiveSceneInfo->Proxy;
 	const auto FeatureLevel = View.GetFeatureLevel();
 
-	if (GDynamicMeshElements.Num() == 0)
-	{
-		GDynamicMeshElements.Empty(10000);
-	}
-
-	FSimpleElementCollector DynamicSimpleElements;
-	GDynamicMeshElements.Add(TArray<FMeshBatchAndRelevance,SceneRenderingAllocator>());
-	TArray<FMeshBatchAndRelevance,SceneRenderingAllocator>& DynamicMeshElements = GDynamicMeshElements.Last();
-	TArray<const FSceneView*> ViewsArray;
-	ViewsArray.Add(&View);
-
-	Renderer.MeshCollector.ClearViewMeshArrays();
-	Renderer.MeshCollector.AddViewMeshArrays(&View, &DynamicMeshElements, &DynamicSimpleElements, FeatureLevel);
-	
-	View.bRenderFirstInstanceOnly = true;
-	Renderer.MeshCollector.SetPrimitive(PrimitiveSceneInfo->Proxy, PrimitiveSceneInfo->DefaultDynamicHitProxyId);
-	PrimitiveSceneInfo->Proxy->GetDynamicMeshElements(ViewsArray, *(View.Family), 0x1, Renderer.MeshCollector);
-	View.bRenderFirstInstanceOnly = false;
-
-	Renderer.MeshCollector.ProcessTasks();
+	TArray<FMeshBatch> MeshElements;
+	PrimitiveSceneInfo->Proxy->GetMeshDescription(LODIndex, MeshElements);
 
 	int32 NumTriangles = 0;
 
-	FPrimitiveViewRelevance ViewRelevance = PrimitiveSceneInfo->Proxy->GetViewRelevance(&View);
-
-	if (ViewRelevance.bDynamicRelevance)
-	{
-		for (int32 MeshBatchIndex = 0; MeshBatchIndex < DynamicMeshElements.Num(); MeshBatchIndex++)
-		{
-			const FMeshBatchAndRelevance& MeshBatchAndRelevance = DynamicMeshElements[MeshBatchIndex];
-
-			if (ShouldGenerateSurfelsOnMesh(*MeshBatchAndRelevance.Mesh, FeatureLevel))
-			{
-				NumTriangles += MeshBatchAndRelevance.Mesh->GetNumPrimitives();
-			}
-		}
-	}
-
-	if (ViewRelevance.bStaticRelevance)
+	for (int32 MeshIndex = 0; MeshIndex < MeshElements.Num(); MeshIndex++)
 	{
-		float LargestScreenSize = 0;
-
-		for (int32 StaticMeshIdx=0; StaticMeshIdx < PrimitiveSceneInfo->StaticMeshes.Num(); StaticMeshIdx++)
-		{
-			const FStaticMesh& StaticMesh = PrimitiveSceneInfo->StaticMeshes[StaticMeshIdx];
-			LargestScreenSize = FMath::Max(LargestScreenSize, StaticMesh.ScreenSize);
-		}
-
-		for (int32 StaticMeshIdx=0; StaticMeshIdx < PrimitiveSceneInfo->StaticMeshes.Num(); StaticMeshIdx++)
+		if (ShouldConvertMesh(MeshElements[MeshIndex]))
 		{
-			const FStaticMesh& StaticMesh = PrimitiveSceneInfo->StaticMeshes[StaticMeshIdx];
-
-			if (ShouldGenerateSurfelsOnMesh(StaticMesh, FeatureLevel)
-				&& StaticMesh.bUseForMaterial
-				// LOD0 only
-				&& StaticMesh.ScreenSize == LargestScreenSize)
-			{
-				NumTriangles += StaticMesh.GetNumPrimitives();
-			}
+			NumTriangles += MeshElements[MeshIndex].GetNumPrimitives();
 		}
 	}
-
+	
 	if (NumTriangles > 0)
 	{
 		if (GUniformMeshTemporaryBuffers.MaxElements < NumTriangles * 3)
@@ -386,85 +343,33 @@ int32 FUniformMeshConverter::Convert(
 		const FVertexBufferRHIParamRef StreamOutTargets[1] = {GUniformMeshTemporaryBuffers.TriangleData.GetReference()};
 		RHICmdList.SetStreamOutTargets(1, StreamOutTargets, Offsets);
 
-		if (ViewRelevance.bDynamicRelevance)
-		{
-			for (int32 MeshBatchIndex = 0; MeshBatchIndex < DynamicMeshElements.Num(); MeshBatchIndex++)
-			{
-				const FMeshBatchAndRelevance& MeshBatchAndRelevance = DynamicMeshElements[MeshBatchIndex];
-
-				if (ShouldGenerateSurfelsOnMesh(*MeshBatchAndRelevance.Mesh, FeatureLevel))
-				{
-					const FMeshBatch& Mesh = *MeshBatchAndRelevance.Mesh;
-
-					FConvertToUniformMeshDrawingPolicy DrawingPolicy(
-						Mesh.VertexFactory,
-						Mesh.MaterialRenderProxy,
-						*Mesh.MaterialRenderProxy->GetMaterial(FeatureLevel),
-						FeatureLevel);
-
-					//@todo - fix
-					OutMaterialRenderProxy = Mesh.MaterialRenderProxy;
-
-					RHICmdList.BuildAndSetLocalBoundShaderState(DrawingPolicy.GetBoundShaderStateInput(FeatureLevel));
-					DrawingPolicy.SetSharedState(RHICmdList, &View, FConvertToUniformMeshDrawingPolicy::ContextDataType());
-
-					for (int32 BatchElementIndex = 0; BatchElementIndex < Mesh.Elements.Num(); BatchElementIndex++)
-					{
-						//@todo - fix
-						OutPrimitiveUniformBuffer = IsValidRef(Mesh.Elements[BatchElementIndex].PrimitiveUniformBuffer) 
-							? Mesh.Elements[BatchElementIndex].PrimitiveUniformBuffer
-							: *Mesh.Elements[BatchElementIndex].PrimitiveUniformBufferResource;
-
-						DrawingPolicy.SetMeshRenderState(RHICmdList, View,PrimitiveSceneProxy,Mesh,BatchElementIndex,false,0.0f,FConvertToUniformMeshDrawingPolicy::ElementDataType(), FConvertToUniformMeshDrawingPolicy::ContextDataType());
-						DrawingPolicy.DrawMesh(RHICmdList, Mesh, BatchElementIndex);
-					}
-				}
-			}
-		}
-
-		if (ViewRelevance.bStaticRelevance)
+		for (int32 MeshIndex = 0; MeshIndex < MeshElements.Num(); MeshIndex++)
 		{
-			float LargestScreenSize = 0;
+			const FMeshBatch& Mesh = MeshElements[MeshIndex];
 
-			for (int32 StaticMeshIdx=0; StaticMeshIdx < PrimitiveSceneInfo->StaticMeshes.Num(); StaticMeshIdx++)
+			if (ShouldConvertMesh(Mesh))
 			{
-				const FStaticMesh& StaticMesh = PrimitiveSceneInfo->StaticMeshes[StaticMeshIdx];
-				LargestScreenSize = FMath::Max(LargestScreenSize, StaticMesh.ScreenSize);
-			}
+				FConvertToUniformMeshDrawingPolicy DrawingPolicy(
+					Mesh.VertexFactory,
+					Mesh.MaterialRenderProxy,
+					*Mesh.MaterialRenderProxy->GetMaterial(FeatureLevel),
+					FeatureLevel);
 
-			for (int32 StaticMeshIdx = 0; StaticMeshIdx < PrimitiveSceneInfo->StaticMeshes.Num(); StaticMeshIdx++)
-			{
-				const FStaticMesh& StaticMesh = PrimitiveSceneInfo->StaticMeshes[StaticMeshIdx];
+				//@todo - fix
+				OutMaterialRenderProxy = Mesh.MaterialRenderProxy;
 
-				if (ShouldGenerateSurfelsOnMesh(StaticMesh, FeatureLevel) 
-					&& StaticMesh.bUseForMaterial
-					&& StaticMesh.ScreenSize == LargestScreenSize)
-				{
-					FConvertToUniformMeshDrawingPolicy DrawingPolicy(
-						StaticMesh.VertexFactory,
-						StaticMesh.MaterialRenderProxy,
-						*StaticMesh.MaterialRenderProxy->GetMaterial(FeatureLevel),
-						FeatureLevel
-						);
-					RHICmdList.BuildAndSetLocalBoundShaderState(DrawingPolicy.GetBoundShaderStateInput(FeatureLevel));
-					DrawingPolicy.SetSharedState(RHICmdList, &View, FConvertToUniformMeshDrawingPolicy::ContextDataType());
+				RHICmdList.BuildAndSetLocalBoundShaderState(DrawingPolicy.GetBoundShaderStateInput(FeatureLevel));
+				DrawingPolicy.SetSharedState(RHICmdList, &View, FConvertToUniformMeshDrawingPolicy::ContextDataType());
 
+				for (int32 BatchElementIndex = 0; BatchElementIndex < Mesh.Elements.Num(); BatchElementIndex++)
+				{
 					//@todo - fix
-					OutMaterialRenderProxy = StaticMesh.MaterialRenderProxy;
-
-					for (int32 BatchElementIndex = 0; BatchElementIndex < StaticMesh.Elements.Num(); BatchElementIndex++)
-					{
-						//@todo - fix
-						OutPrimitiveUniformBuffer = IsValidRef(StaticMesh.Elements[BatchElementIndex].PrimitiveUniformBuffer) 
-							? StaticMesh.Elements[BatchElementIndex].PrimitiveUniformBuffer
-							: *StaticMesh.Elements[BatchElementIndex].PrimitiveUniformBufferResource;
-
-						DrawingPolicy.SetMeshRenderState(RHICmdList, View,PrimitiveSceneProxy,StaticMesh,BatchElementIndex,false,0.0f,
-							FConvertToUniformMeshDrawingPolicy::ElementDataType(),
-							FConvertToUniformMeshDrawingPolicy::ContextDataType()
-							);
-						DrawingPolicy.DrawMesh(RHICmdList, StaticMesh, BatchElementIndex);
-					}
+					OutPrimitiveUniformBuffer = IsValidRef(Mesh.Elements[BatchElementIndex].PrimitiveUniformBuffer) 
+						? Mesh.Elements[BatchElementIndex].PrimitiveUniformBuffer
+						: *Mesh.Elements[BatchElementIndex].PrimitiveUniformBufferResource;
+
+					DrawingPolicy.SetMeshRenderState(RHICmdList, View,PrimitiveSceneProxy,Mesh,BatchElementIndex,false,0.0f,FConvertToUniformMeshDrawingPolicy::ElementDataType(), FConvertToUniformMeshDrawingPolicy::ContextDataType());
+					DrawingPolicy.DrawMesh(RHICmdList, Mesh, BatchElementIndex);
 				}
 			}
 		}
@@ -525,6 +430,11 @@ public:
 		SetUniformBufferParameter(RHICmdList, ShaderRHI,GetUniformBufferParameter<FPrimitiveUniformShaderParameters>(),PrimitiveUniformBuffer);
 
 		const FScene* Scene = (const FScene*)View.Family->Scene;
+
+		FUnorderedAccessViewRHIParamRef UniformMeshUAVs[1];
+		UniformMeshUAVs[0] = Scene->DistanceFieldSceneData.SurfelBuffers->Surfels.UAV;
+		RHICmdList.TransitionResources(EResourceTransitionAccess::ERWBarrier, EResourceTransitionPipeline::EComputeToCompute, UniformMeshUAVs, ARRAY_COUNT(UniformMeshUAVs));
+
 		SurfelBufferParameters.Set(RHICmdList, ShaderRHI, *Scene->DistanceFieldSceneData.SurfelBuffers, *Scene->DistanceFieldSceneData.InstancedSurfelBuffers);
 		
 		SetShaderValue(RHICmdList, ShaderRHI, SurfelStartIndex, SurfelStartIndexValue);
@@ -533,10 +443,15 @@ public:
 		SetShaderValue(RHICmdList, ShaderRHI, Instance0InverseTransform, Instance0Transform.Inverse());
 	}
 
-	void UnsetParameters(FRHICommandList& RHICmdList)
+	void UnsetParameters(FRHICommandList& RHICmdList, FViewInfo& View)
 	{
 		FComputeShaderRHIParamRef ShaderRHI = GetComputeShader();
 		SurfelBufferParameters.UnsetParameters(RHICmdList, ShaderRHI);
+
+		const FScene* Scene = (const FScene*)View.Family->Scene;
+		FUnorderedAccessViewRHIParamRef UniformMeshUAVs[1];
+		UniformMeshUAVs[0] = Scene->DistanceFieldSceneData.SurfelBuffers->Surfels.UAV;
+		RHICmdList.TransitionResources(EResourceTransitionAccess::EReadable, EResourceTransitionPipeline::EComputeToCompute, UniformMeshUAVs, ARRAY_COUNT(UniformMeshUAVs));
 	}
 
 	virtual bool Serialize(FArchive& Ar) override
@@ -577,5 +492,5 @@ void FUniformMeshConverter::GenerateSurfels(
 	ComputeShader->SetParameters(RHICmdList, View, SurfelOffset, NumSurfels, MaterialProxy, PrimitiveUniformBuffer, Instance0Transform);
 	DispatchComputeShader(RHICmdList, ComputeShader, FMath::DivideAndRoundUp(NumSurfels, GEvaluateSurfelMaterialGroupSize), 1, 1);
 
-	ComputeShader->UnsetParameters(RHICmdList);
+	ComputeShader->UnsetParameters(RHICmdList, View);
 }
diff --git a/Engine/Source/Runtime/Renderer/Private/DeferredShadingRenderer.h b/Engine/Source/Runtime/Renderer/Private/DeferredShadingRenderer.h
--- a/Engine/Source/Runtime/Renderer/Private/DeferredShadingRenderer.h
+++ b/Engine/Source/Runtime/Renderer/Private/DeferredShadingRenderer.h
@@ -116,6 +116,8 @@ public:
 	void RenderVisualizeTexturePool(FRHICommandListImmediate& RHICmdList);
 #endif
 
+	/** Offline culling of static triangles that won't be seen at runtime. */
+	void PreCullStaticMeshes(FRHICommandListImmediate& RHICmdList, const TArray<UStaticMeshComponent*>& ComponentsToPreCull, const TArray<TArray<FPlane> >& CullVolumes);
 
 	/** bound shader state for occlusion test prims */
 	static FGlobalBoundShaderState OcclusionTestBoundShaderState;
diff --git a/Engine/Source/Runtime/Renderer/Private/DistanceFieldLightingShared.h b/Engine/Source/Runtime/Renderer/Private/DistanceFieldLightingShared.h
--- a/Engine/Source/Runtime/Renderer/Private/DistanceFieldLightingShared.h
+++ b/Engine/Source/Runtime/Renderer/Private/DistanceFieldLightingShared.h
@@ -522,6 +522,7 @@ public:
 		FSceneRenderer& Renderer,
 		FViewInfo& View, 
 		const FPrimitiveSceneInfo* PrimitiveSceneInfo, 
+		int32 LODIndex,
 		FUniformMeshBuffers*& OutUniformMeshBuffers,
 		const FMaterialRenderProxy*& OutMaterialRenderProxy,
 		FUniformBufferRHIParamRef& OutPrimitiveUniformBuffer);
@@ -537,4 +538,36 @@ public:
 		int32 NumSurfels);
 };
 
+class FPreCulledTriangleBuffers
+{
+public:
+
+	int32 MaxIndices;
+
+	FRWBuffer TriangleVisibleMask;
+
+	FPreCulledTriangleBuffers()
+	{
+		MaxIndices = 0;
+	}
+
+	void Initialize()
+	{
+		if (MaxIndices > 0)
+		{
+			TriangleVisibleMask.Initialize(sizeof(uint32), MaxIndices / 3, PF_R32_UINT);
+		}
+	}
+
+	void Release()
+	{
+		TriangleVisibleMask.Release();
+	}
+
+	size_t GetSizeBytes() const
+	{
+		return TriangleVisibleMask.NumBytes;
+	}
+};
+
 extern TGlobalResource<FDistanceFieldObjectBufferResource> GAOCulledObjectBuffers;
\ No newline at end of file
diff --git a/Engine/Source/Runtime/Renderer/Private/HeightfieldLighting.cpp b/Engine/Source/Runtime/Renderer/Private/HeightfieldLighting.cpp
--- a/Engine/Source/Runtime/Renderer/Private/HeightfieldLighting.cpp
+++ b/Engine/Source/Runtime/Renderer/Private/HeightfieldLighting.cpp
@@ -550,6 +550,45 @@ void FHeightfieldLightingViewInfo::SetupVisibleHeightfields(const FViewInfo& Vie
 	}
 }
 
+void FHeightfieldLightingViewInfo::SetupHeightfieldsForScene(const FScene& Scene, FRHICommandListImmediate& RHICmdList)
+{
+	const int32 NumPrimitives = Scene.DistanceFieldSceneData.HeightfieldPrimitives.Num();
+
+	if (NumPrimitives > 0
+		&& SupportsHeightfieldLighting(Scene.GetFeatureLevel(), Scene.GetShaderPlatform()))
+	{
+		const float MaxDistanceSquared = FMath::Square(GetMaxAOViewDistance() + GetGHeightfieldBounceDistance());
+		float LocalToWorldScale = 1;
+
+		for (int32 HeightfieldPrimitiveIndex = 0; HeightfieldPrimitiveIndex < NumPrimitives; HeightfieldPrimitiveIndex++)
+		{
+			const FPrimitiveSceneInfo* HeightfieldPrimitive = Scene.DistanceFieldSceneData.HeightfieldPrimitives[HeightfieldPrimitiveIndex];
+
+			UTexture2D* HeightfieldTexture = NULL;
+			UTexture2D* DiffuseColorTexture = NULL;
+			FHeightfieldComponentDescription NewComponentDescription(HeightfieldPrimitive->Proxy->GetLocalToWorld());
+			HeightfieldPrimitive->Proxy->GetHeightfieldRepresentation(HeightfieldTexture, DiffuseColorTexture, NewComponentDescription);
+
+			if (HeightfieldTexture && HeightfieldTexture->Resource->TextureRHI)
+			{
+				const FIntPoint HeightfieldSize = NewComponentDescription.HeightfieldRect.Size();
+
+				if (Heightfield.Rect.Area() == 0)
+				{
+					Heightfield.Rect = NewComponentDescription.HeightfieldRect;
+					LocalToWorldScale = NewComponentDescription.LocalToWorld.GetScaleVector().X;
+				}
+				else
+				{
+					Heightfield.Rect.Union(NewComponentDescription.HeightfieldRect);
+				}
+
+				TArray<FHeightfieldComponentDescription>& ComponentDescriptions = Heightfield.ComponentDescriptions.FindOrAdd(FHeightfieldComponentTextures(HeightfieldTexture, DiffuseColorTexture));
+				ComponentDescriptions.Add(NewComponentDescription);
+			}
+		}
+	}
+}
 
 class FHeightfieldDescriptionsResource : public FRenderResource
 {
@@ -1965,4 +2004,146 @@ void FHeightfieldLightingViewInfo::ComputeIrradianceForScreenGrid(
 			}
 		}
 	}
+}
+
+
+const int32 GPreCullTrianglesToHeightfieldsGroupSize = 64;
+
+class FPreCullTrianglesToHeightfields : public FGlobalShader
+{
+	DECLARE_SHADER_TYPE(FPreCullTrianglesToHeightfields,Global)
+public:
+
+	static bool ShouldCache(EShaderPlatform Platform)
+	{
+		return IsFeatureLevelSupported(Platform, ERHIFeatureLevel::SM5) && DoesPlatformSupportDistanceFieldGI(Platform);
+	}
+
+	static void ModifyCompilationEnvironment(EShaderPlatform Platform, FShaderCompilerEnvironment& OutEnvironment)
+	{
+		FGlobalShader::ModifyCompilationEnvironment(Platform,OutEnvironment);
+
+		OutEnvironment.SetDefine(TEXT("PRECULL_TRIANGLES_TO_HEIGHTFIELDS_GROUP_SIZE"), GPreCullTrianglesToHeightfieldsGroupSize);
+
+		// To reduce shader compile time of compute shaders with shared memory, doesn't have an impact on generated code with current compiler (June 2010 DX SDK)
+		OutEnvironment.CompilerFlags.Add(CFLAG_StandardOptimization);
+	}
+
+	FPreCullTrianglesToHeightfields(const ShaderMetaType::CompiledShaderInitializerType& Initializer)
+		: FGlobalShader(Initializer)
+	{
+		HeightfieldDescriptionParameters.Bind(Initializer.ParameterMap);
+		HeightfieldTextureParameters.Bind(Initializer.ParameterMap);
+		TriangleVisibleMask.Bind(Initializer.ParameterMap, TEXT("TriangleVisibleMask"));
+		StartIndex.Bind(Initializer.ParameterMap, TEXT("StartIndex"));
+		NumTriangles.Bind(Initializer.ParameterMap, TEXT("NumTriangles"));
+		TriangleVertexData.Bind(Initializer.ParameterMap, TEXT("TriangleVertexData"));
+		PreCullMaxDistance.Bind(Initializer.ParameterMap, TEXT("PreCullMaxDistance"));
+	}
+
+	FPreCullTrianglesToHeightfields()
+	{
+	}
+
+	void SetParameters(
+		FRHICommandList& RHICmdList, 
+		const FViewInfo& View, 
+		UTexture2D* HeightfieldTextureValue, 
+		int32 NumHeightfieldsValue,
+		FPreCulledTriangleBuffers& PreCulledTriangleBuffers,
+		int32 StartIndexValue,
+		int32 NumTrianglesValue,
+		const FUniformMeshBuffers& UniformMeshBuffers)
+	{
+		FComputeShaderRHIParamRef ShaderRHI = GetComputeShader();
+		FGlobalShader::SetParameters(RHICmdList, ShaderRHI, View);
+
+		HeightfieldDescriptionParameters.Set(RHICmdList, ShaderRHI, NumHeightfieldsValue);
+		HeightfieldTextureParameters.Set(RHICmdList, ShaderRHI, HeightfieldTextureValue, NULL);
+
+		RHICmdList.TransitionResource(EResourceTransitionAccess::ERWBarrier, EResourceTransitionPipeline::EComputeToCompute, PreCulledTriangleBuffers.TriangleVisibleMask.UAV);
+		TriangleVisibleMask.SetBuffer(RHICmdList, ShaderRHI, PreCulledTriangleBuffers.TriangleVisibleMask);
+
+		SetShaderValue(RHICmdList, ShaderRHI, StartIndex, StartIndexValue);
+		SetShaderValue(RHICmdList, ShaderRHI, NumTriangles, NumTrianglesValue);
+		SetSRVParameter(RHICmdList, ShaderRHI, TriangleVertexData, UniformMeshBuffers.TriangleDataSRV);
+		extern float GPreCullMaxDistance;
+		SetShaderValue(RHICmdList, ShaderRHI, PreCullMaxDistance, GPreCullMaxDistance);
+	}
+
+	void UnsetParameters(FRHICommandList& RHICmdList, FPreCulledTriangleBuffers& PreCulledTriangleBuffers)
+	{
+		FComputeShaderRHIParamRef ShaderRHI = GetComputeShader();
+		TriangleVisibleMask.UnsetUAV(RHICmdList, ShaderRHI);
+		// RHISetStreamOutTargets doesn't unbind existing uses like render targets do 
+		SetSRVParameter(RHICmdList, ShaderRHI, TriangleVertexData, NULL);
+
+		RHICmdList.TransitionResource(EResourceTransitionAccess::EReadable, EResourceTransitionPipeline::EComputeToCompute, PreCulledTriangleBuffers.TriangleVisibleMask.UAV);
+	}
+
+	virtual bool Serialize(FArchive& Ar) override
+	{		
+		bool bShaderHasOutdatedParameters = FGlobalShader::Serialize(Ar);
+		Ar << HeightfieldDescriptionParameters;
+		Ar << HeightfieldTextureParameters;
+		Ar << TriangleVisibleMask;
+		Ar << StartIndex;
+		Ar << NumTriangles;
+		Ar << TriangleVertexData;
+		Ar << PreCullMaxDistance;
+		return bShaderHasOutdatedParameters;
+	}
+
+private:
+
+	FHeightfieldDescriptionParameters HeightfieldDescriptionParameters;
+	FHeightfieldTextureParameters HeightfieldTextureParameters;
+	FRWShaderParameter TriangleVisibleMask;
+	FShaderParameter StartIndex;
+	FShaderParameter NumTriangles;
+	FShaderResourceParameter TriangleVertexData;
+	FShaderParameter PreCullMaxDistance;
+};
+
+IMPLEMENT_SHADER_TYPE(,FPreCullTrianglesToHeightfields,TEXT("HeightfieldLighting"),TEXT("PreCullTrianglesToHeightfieldsCS"),SF_Compute);
+
+
+void FHeightfieldLightingViewInfo::PreCullTriangles(
+	const FViewInfo& View, 
+	FRHICommandListImmediate& RHICmdList,
+	FPreCulledTriangleBuffers& PreCulledTriangleBuffers,
+	int32 StartIndexValue,
+	int32 NumTrianglesValue,
+	const FUniformMeshBuffers& UniformMeshBuffers) const
+{
+	const FScene* Scene = (const FScene*)View.Family->Scene;
+
+	if (Heightfield.ComponentDescriptions.Num() > 0)
+	{
+		SCOPED_DRAW_EVENT(RHICmdList, PreCullTrianglesToHeightfields);
+
+		FSceneViewState* ViewState = (FSceneViewState*)View.State;
+
+		{
+			for (TMap<FHeightfieldComponentTextures, TArray<FHeightfieldComponentDescription>>::TConstIterator It(Heightfield.ComponentDescriptions); It; ++It)
+			{
+				const TArray<FHeightfieldComponentDescription>& HeightfieldDescriptions = It.Value();
+
+				if (HeightfieldDescriptions.Num() > 0)
+				{
+					UploadHeightfieldDescriptions(HeightfieldDescriptions, FVector2D(1, 1), 1.0f / Heightfield.DownsampleFactor);
+
+					UTexture2D* HeightfieldTexture = It.Key().HeightAndNormal;
+
+					const uint32 GroupSize = FMath::DivideAndRoundUp(NumTrianglesValue, GPreCullTrianglesToHeightfieldsGroupSize);
+
+					TShaderMapRef<FPreCullTrianglesToHeightfields> ComputeShader(View.ShaderMap);
+					RHICmdList.SetComputeShader(ComputeShader->GetComputeShader());
+					ComputeShader->SetParameters(RHICmdList, View, HeightfieldTexture, HeightfieldDescriptions.Num(), PreCulledTriangleBuffers, StartIndexValue, NumTrianglesValue, UniformMeshBuffers);
+					DispatchComputeShader(RHICmdList, *ComputeShader, GroupSize, 1, 1);
+					ComputeShader->UnsetParameters(RHICmdList, PreCulledTriangleBuffers);
+				}
+			}
+		}
+	}
 }
\ No newline at end of file
diff --git a/Engine/Source/Runtime/Renderer/Private/HeightfieldLighting.h b/Engine/Source/Runtime/Renderer/Private/HeightfieldLighting.h
--- a/Engine/Source/Runtime/Renderer/Private/HeightfieldLighting.h
+++ b/Engine/Source/Runtime/Renderer/Private/HeightfieldLighting.h
@@ -80,6 +80,8 @@ public:
 
 	void SetupVisibleHeightfields(const FViewInfo& View, FRHICommandListImmediate& RHICmdList);
 
+	void SetupHeightfieldsForScene(const FScene& Scene, FRHICommandListImmediate& RHICmdList);
+
 	void ClearShadowing(const FViewInfo& View, FRHICommandListImmediate& RHICmdList, const FLightSceneInfo& LightSceneInfo) const;
 
 	void ComputeShadowMapShadowing(const FViewInfo& View, FRHICommandListImmediate& RHICmdList, const FProjectedShadowInfo* ProjectedShadowInfo) const;
@@ -121,6 +123,14 @@ public:
 		const FAOScreenGridResources& ScreenGridResources,
 		const FDistanceFieldAOParameters& Parameters) const;
 
+	void PreCullTriangles(
+		const FViewInfo& View, 
+		FRHICommandListImmediate& RHICmdList,
+		class FPreCulledTriangleBuffers& PreCulledTriangleBuffers,
+		int32 StartIndexValue,
+		int32 NumTrianglesValue,
+		const class FUniformMeshBuffers& UniformMeshBuffers) const;
+
 private:
 
 	FHeightfieldDescription Heightfield;
diff --git a/Engine/Source/Runtime/Renderer/Private/PreCullTriangles.cpp b/Engine/Source/Runtime/Renderer/Private/PreCullTriangles.cpp
new file mode 100644
--- /dev/null
+++ b/Engine/Source/Runtime/Renderer/Private/PreCullTriangles.cpp
@@ -0,0 +1,627 @@
+// Copyright 1998-2014 Epic Games, Inc. All Rights Reserved.
+
+/*=============================================================================
+	PreCullTriangles.cpp : implementation of offline visibility culling of static triangles
+=============================================================================*/
+
+#include "RendererPrivate.h"
+#include "ScenePrivate.h"
+#include "UniformBuffer.h"
+#include "ShaderParameters.h"
+#include "PostProcessing.h"
+#include "SceneFilterRendering.h"
+#include "DistanceFieldLightingShared.h"
+#include "DistanceFieldSurfaceCacheLighting.h"
+#include "DistanceFieldGlobalIllumination.h"
+#include "RHICommandList.h"
+#include "SceneUtils.h"
+#include "DistanceFieldAtlas.h"
+#include "StaticMeshResources.h"
+
+TGlobalResource<FDistanceFieldObjectBufferResource> GPreCullTrianglesCulledObjectBuffers;
+
+float GPreCullMaxDistance = 1000;
+
+FAutoConsoleVariableRef CPreCullMaxDistance(
+	TEXT("r.PreCullMaxDistance"),
+	GPreCullMaxDistance,
+	TEXT(""),
+	ECVF_Scalability | ECVF_RenderThreadSafe
+	);
+
+class FCullObjectsForBoundsCS : public FGlobalShader
+{
+	DECLARE_SHADER_TYPE(FCullObjectsForBoundsCS,Global)
+public:
+
+	static bool ShouldCache(EShaderPlatform Platform)
+	{
+		return IsFeatureLevelSupported(Platform, ERHIFeatureLevel::SM5) && DoesPlatformSupportDistanceFieldGI(Platform);
+	}
+
+	static void ModifyCompilationEnvironment(EShaderPlatform Platform, FShaderCompilerEnvironment& OutEnvironment)
+	{
+		FGlobalShader::ModifyCompilationEnvironment(Platform,OutEnvironment);
+		OutEnvironment.SetDefine(TEXT("UPDATEOBJECTS_THREADGROUP_SIZE"), UpdateObjectsGroupSize);
+	}
+
+	FCullObjectsForBoundsCS(const ShaderMetaType::CompiledShaderInitializerType& Initializer)
+		: FGlobalShader(Initializer)
+	{
+		ObjectBufferParameters.Bind(Initializer.ParameterMap);
+		ObjectIndirectArguments.Bind(Initializer.ParameterMap, TEXT("ObjectIndirectArguments"));
+		CulledObjectBounds.Bind(Initializer.ParameterMap, TEXT("CulledObjectBounds"));
+		CulledObjectData.Bind(Initializer.ParameterMap, TEXT("CulledObjectData"));
+		AOParameters.Bind(Initializer.ParameterMap);
+		ObjectBoundsCenter.Bind(Initializer.ParameterMap, TEXT("ObjectBoundsCenter"));
+		ObjectBoundsExtent.Bind(Initializer.ParameterMap, TEXT("ObjectBoundsExtent"));
+		ObjectBoundingGeometryIndexCount.Bind(Initializer.ParameterMap, TEXT("ObjectBoundingGeometryIndexCount"));
+		PreCullMaxDistance.Bind(Initializer.ParameterMap, TEXT("PreCullMaxDistance"));
+	}
+
+	FCullObjectsForBoundsCS()
+	{
+	}
+
+	void SetParameters(FRHICommandList& RHICmdList, const FScene* Scene, const FSceneView& View, FVector ObjectBoundsCenterValue, FVector ObjectBoundsExtentValue, const FDistanceFieldAOParameters& Parameters)
+	{
+		FComputeShaderRHIParamRef ShaderRHI = GetComputeShader();
+		FGlobalShader::SetParameters(RHICmdList, ShaderRHI, View);
+		ObjectBufferParameters.Set(RHICmdList, ShaderRHI, *(Scene->DistanceFieldSceneData.ObjectBuffers), Scene->DistanceFieldSceneData.NumObjectsInBuffer);
+
+		FUnorderedAccessViewRHIParamRef OutUAVs[3];
+		OutUAVs[0] = GPreCullTrianglesCulledObjectBuffers.Buffers.ObjectIndirectArguments.UAV;
+		OutUAVs[1] = GPreCullTrianglesCulledObjectBuffers.Buffers.Bounds.UAV;
+		OutUAVs[2] = GPreCullTrianglesCulledObjectBuffers.Buffers.Data.UAV;
+		RHICmdList.TransitionResources(EResourceTransitionAccess::ERWBarrier, EResourceTransitionPipeline::EComputeToCompute, OutUAVs, ARRAY_COUNT(OutUAVs));
+
+		ObjectIndirectArguments.SetBuffer(RHICmdList, ShaderRHI, GPreCullTrianglesCulledObjectBuffers.Buffers.ObjectIndirectArguments);
+		CulledObjectBounds.SetBuffer(RHICmdList, ShaderRHI, GPreCullTrianglesCulledObjectBuffers.Buffers.Bounds);
+		CulledObjectData.SetBuffer(RHICmdList, ShaderRHI, GPreCullTrianglesCulledObjectBuffers.Buffers.Data);
+		AOParameters.Set(RHICmdList, ShaderRHI, Parameters);
+
+		SetShaderValue(RHICmdList, ShaderRHI, ObjectBoundsCenter, ObjectBoundsCenterValue);
+		SetShaderValue(RHICmdList, ShaderRHI, ObjectBoundsExtent, ObjectBoundsExtentValue);
+
+		SetShaderValue(RHICmdList, ShaderRHI, ObjectBoundingGeometryIndexCount, StencilingGeometry::GLowPolyStencilSphereIndexBuffer.GetIndexCount());
+		SetShaderValue(RHICmdList, ShaderRHI, PreCullMaxDistance, GPreCullMaxDistance);
+	}
+
+	void UnsetParameters(FRHICommandList& RHICmdList, const FScene* Scene)
+	{
+		ObjectBufferParameters.UnsetParameters(RHICmdList, GetComputeShader(), *(Scene->DistanceFieldSceneData.ObjectBuffers));
+		ObjectIndirectArguments.UnsetUAV(RHICmdList, GetComputeShader());
+		CulledObjectBounds.UnsetUAV(RHICmdList, GetComputeShader());
+		CulledObjectData.UnsetUAV(RHICmdList, GetComputeShader());
+
+		FUnorderedAccessViewRHIParamRef OutUAVs[3];
+		OutUAVs[0] = GPreCullTrianglesCulledObjectBuffers.Buffers.ObjectIndirectArguments.UAV;
+		OutUAVs[1] = GPreCullTrianglesCulledObjectBuffers.Buffers.Bounds.UAV;
+		OutUAVs[2] = GPreCullTrianglesCulledObjectBuffers.Buffers.Data.UAV;	
+		RHICmdList.TransitionResources(EResourceTransitionAccess::EReadable, EResourceTransitionPipeline::EComputeToCompute, OutUAVs, ARRAY_COUNT(OutUAVs));
+	}
+
+	virtual bool Serialize(FArchive& Ar)
+	{		
+		bool bShaderHasOutdatedParameters = FGlobalShader::Serialize(Ar);
+		Ar << ObjectBufferParameters;
+		Ar << ObjectIndirectArguments;
+		Ar << CulledObjectBounds;
+		Ar << CulledObjectData;
+		Ar << AOParameters;
+		Ar << ObjectBoundsCenter;
+		Ar << ObjectBoundsExtent;
+		Ar << ObjectBoundingGeometryIndexCount;
+		Ar << PreCullMaxDistance;
+		return bShaderHasOutdatedParameters;
+	}
+
+private:
+
+	FDistanceFieldObjectBufferParameters ObjectBufferParameters;
+	FRWShaderParameter ObjectIndirectArguments;
+	FRWShaderParameter CulledObjectBounds;
+	FRWShaderParameter CulledObjectData;
+	FAOParameters AOParameters;
+	FShaderParameter ObjectBoundsCenter;
+	FShaderParameter ObjectBoundsExtent;
+	FShaderParameter ObjectBoundingGeometryIndexCount;
+	FShaderParameter PreCullMaxDistance;
+};
+
+IMPLEMENT_SHADER_TYPE(,FCullObjectsForBoundsCS,TEXT("PreCullTriangles"),TEXT("CullObjectsForBoundsCS"),SF_Compute);
+
+class FCullVolumeBuffersResource : public FRenderResource
+{
+public:
+
+	FCPUUpdatedBuffer CullVolumePlaneOffsets;
+	FCPUUpdatedBuffer CullVolumePlaneData;
+
+	FCullVolumeBuffersResource()
+	{
+		CullVolumePlaneOffsets.Format = PF_R32_UINT;
+		CullVolumePlaneOffsets.Stride = 2;
+
+		CullVolumePlaneData.Format = PF_A32B32G32R32F;
+		CullVolumePlaneData.Stride = 1;
+	}
+
+	virtual void InitDynamicRHI()  override
+	{
+		CullVolumePlaneOffsets.Initialize();
+		CullVolumePlaneData.Initialize();
+	}
+
+	virtual void ReleaseDynamicRHI() override
+	{
+		CullVolumePlaneOffsets.Release();
+		CullVolumePlaneData.Release();
+	}
+};
+
+TGlobalResource<FCullVolumeBuffersResource> GCullVolumeBuffers;
+
+class FPreCulledTriangleBufferResource : public FRenderResource
+{
+public:
+	FPreCulledTriangleBuffers Buffers;
+
+	virtual void InitDynamicRHI()  override
+	{
+		Buffers.Initialize();
+	}
+
+	virtual void ReleaseDynamicRHI() override
+	{
+		Buffers.Release();
+	}
+};
+
+class FPreCullTrianglesCS : public FGlobalShader
+{
+	DECLARE_SHADER_TYPE(FPreCullTrianglesCS,Global)
+public:
+
+	static bool ShouldCache(EShaderPlatform Platform)
+	{
+		return IsFeatureLevelSupported(Platform, ERHIFeatureLevel::SM5) && DoesPlatformSupportDistanceFieldGI(Platform);
+	}
+
+	static void ModifyCompilationEnvironment(EShaderPlatform Platform, FShaderCompilerEnvironment& OutEnvironment)
+	{
+		FGlobalShader::ModifyCompilationEnvironment(Platform,OutEnvironment);
+		OutEnvironment.SetDefine(TEXT("THREADGROUP_SIZEX"), GDistanceFieldAOTileSizeX);
+		OutEnvironment.SetDefine(TEXT("THREADGROUP_SIZEY"), GDistanceFieldAOTileSizeY);
+	}
+
+	FPreCullTrianglesCS(const ShaderMetaType::CompiledShaderInitializerType& Initializer)
+		: FGlobalShader(Initializer)
+	{
+		ObjectParameters.Bind(Initializer.ParameterMap);
+		TriangleVisibleMask.Bind(Initializer.ParameterMap, TEXT("TriangleVisibleMask"));
+		StartIndex.Bind(Initializer.ParameterMap, TEXT("StartIndex"));
+		NumTriangles.Bind(Initializer.ParameterMap, TEXT("NumTriangles"));
+		TriangleVertexData.Bind(Initializer.ParameterMap, TEXT("TriangleVertexData"));
+		PreCullMaxDistance.Bind(Initializer.ParameterMap, TEXT("PreCullMaxDistance"));
+		NumCullVolumes.Bind(Initializer.ParameterMap, TEXT("NumCullVolumes"));
+		CullVolumePlaneOffsets.Bind(Initializer.ParameterMap, TEXT("CullVolumePlaneOffsets"));
+		CullVolumePlaneData.Bind(Initializer.ParameterMap, TEXT("CullVolumePlaneData"));
+	}
+
+	FPreCullTrianglesCS()
+	{
+	}
+
+	void SetParameters(
+		FRHICommandList& RHICmdList, 
+		const FSceneView& View, 
+		FPreCulledTriangleBuffers& PreCulledTriangleBuffers,
+		int32 StartIndexValue,
+		int32 NumTrianglesValue,
+		const FUniformMeshBuffers& UniformMeshBuffers,
+		int32 NumCullVolumesValue)
+	{
+		FComputeShaderRHIParamRef ShaderRHI = GetComputeShader();
+		FGlobalShader::SetParameters(RHICmdList, ShaderRHI, View);
+
+		ObjectParameters.Set(RHICmdList, ShaderRHI, GPreCullTrianglesCulledObjectBuffers.Buffers);
+
+		FUnorderedAccessViewRHIParamRef OutUAVs[1];
+		OutUAVs[0] = PreCulledTriangleBuffers.TriangleVisibleMask.UAV;
+		RHICmdList.TransitionResources(EResourceTransitionAccess::ERWBarrier, EResourceTransitionPipeline::EComputeToCompute, OutUAVs, ARRAY_COUNT(OutUAVs));
+
+		TriangleVisibleMask.SetBuffer(RHICmdList, ShaderRHI, PreCulledTriangleBuffers.TriangleVisibleMask);
+
+		SetShaderValue(RHICmdList, ShaderRHI, StartIndex, StartIndexValue);
+		SetShaderValue(RHICmdList, ShaderRHI, NumTriangles, NumTrianglesValue);
+		SetSRVParameter(RHICmdList, ShaderRHI, TriangleVertexData, UniformMeshBuffers.TriangleDataSRV);
+		SetShaderValue(RHICmdList, ShaderRHI, PreCullMaxDistance, GPreCullMaxDistance);
+
+		SetShaderValue(RHICmdList, ShaderRHI, NumCullVolumes, NumCullVolumesValue);
+		SetSRVParameter(RHICmdList, ShaderRHI, CullVolumePlaneOffsets, GCullVolumeBuffers.CullVolumePlaneOffsets.BufferSRV);
+		SetSRVParameter(RHICmdList, ShaderRHI, CullVolumePlaneData, GCullVolumeBuffers.CullVolumePlaneData.BufferSRV);
+	}
+
+	void UnsetParameters(FRHICommandList& RHICmdList, FPreCulledTriangleBuffers& PreCulledTriangleBuffers)
+	{
+		FComputeShaderRHIParamRef ShaderRHI = GetComputeShader();
+		TriangleVisibleMask.UnsetUAV(RHICmdList, ShaderRHI);
+		// RHISetStreamOutTargets doesn't unbind existing uses like render targets do 
+		SetSRVParameter(RHICmdList, ShaderRHI, TriangleVertexData, NULL);
+
+		FUnorderedAccessViewRHIParamRef OutUAVs[1];
+		OutUAVs[0] = PreCulledTriangleBuffers.TriangleVisibleMask.UAV;
+		RHICmdList.TransitionResources(EResourceTransitionAccess::EReadable, EResourceTransitionPipeline::EComputeToCompute, OutUAVs, ARRAY_COUNT(OutUAVs));
+	}
+
+	virtual bool Serialize(FArchive& Ar) override
+	{		
+		bool bShaderHasOutdatedParameters = FGlobalShader::Serialize(Ar);
+		Ar << ObjectParameters;
+		Ar << TriangleVisibleMask;
+		Ar << StartIndex;
+		Ar << NumTriangles;
+		Ar << TriangleVertexData;
+		Ar << PreCullMaxDistance;
+		Ar << NumCullVolumes;
+		Ar << CullVolumePlaneOffsets;
+		Ar << CullVolumePlaneData;
+		return bShaderHasOutdatedParameters;
+	}
+
+private:
+
+	FDistanceFieldCulledObjectBufferParameters ObjectParameters;
+	FRWShaderParameter TriangleVisibleMask;
+	FShaderParameter StartIndex;
+	FShaderParameter NumTriangles;
+	FShaderResourceParameter TriangleVertexData;
+	FShaderParameter PreCullMaxDistance;
+	FShaderParameter NumCullVolumes;
+	FShaderResourceParameter CullVolumePlaneOffsets;
+	FShaderResourceParameter CullVolumePlaneData;
+};
+
+IMPLEMENT_SHADER_TYPE(,FPreCullTrianglesCS,TEXT("PreCullTriangles"),TEXT("PreCullTrianglesCS"),SF_Compute);
+
+TGlobalResource<FPreCulledTriangleBufferResource> GPreCulledTriangleBuffers;
+TGlobalResource<FPreCulledTriangleBufferResource> GPreCulledTriangleBuffersFromHeightfields;
+
+extern bool SupportsDistanceFieldAO(ERHIFeatureLevel::Type FeatureLevel, EShaderPlatform ShaderPlatform);
+
+void FDeferredShadingSceneRenderer::PreCullStaticMeshes(FRHICommandListImmediate& RHICmdList, const TArray<UStaticMeshComponent*>& ComponentsToPreCull, const TArray<TArray<FPlane> >& CullVolumes)
+{
+	FViewInfo& View = Views[0];
+	FDistanceFieldAOParameters Parameters(1000);
+
+	if (SupportsDistanceFieldAO(Scene->GetFeatureLevel(), Scene->GetShaderPlatform()))
+	{
+		QUICK_SCOPE_CYCLE_COUNTER(STAT_FDeferredShadingSceneRenderer_PreCullStaticMeshes);
+		GDistanceFieldVolumeTextureAtlas.UpdateAllocations();
+		UpdateGlobalDistanceFieldObjectBuffers(RHICmdList);
+
+		FHeightfieldLightingViewInfo HeightfieldScene;
+		HeightfieldScene.SetupHeightfieldsForScene(*Scene, RHICmdList);
+
+		{
+			int32 NumCullVolumePlanes = 0;
+
+			for (int32 VolumeIndex = 0; VolumeIndex < CullVolumes.Num(); VolumeIndex++)
+			{
+				NumCullVolumePlanes += CullVolumes[VolumeIndex].Num();
+			}
+
+			if (CullVolumes.Num() > GCullVolumeBuffers.CullVolumePlaneOffsets.MaxElements)
+			{
+				GCullVolumeBuffers.CullVolumePlaneOffsets.MaxElements = CullVolumes.Num() * 5 / 4;
+				GCullVolumeBuffers.CullVolumePlaneOffsets.Release();
+				GCullVolumeBuffers.CullVolumePlaneOffsets.Initialize();
+			}
+
+			if (NumCullVolumePlanes > GCullVolumeBuffers.CullVolumePlaneData.MaxElements)
+			{
+				GCullVolumeBuffers.CullVolumePlaneData.MaxElements = NumCullVolumePlanes * 5 / 4;
+				GCullVolumeBuffers.CullVolumePlaneData.Release();
+				GCullVolumeBuffers.CullVolumePlaneData.Initialize();
+			}
+
+			if (NumCullVolumePlanes > 0)
+			{
+				uint32* CullVolumePlaneOffsetsData = (uint32*)RHILockVertexBuffer(GCullVolumeBuffers.CullVolumePlaneOffsets.Buffer, 0, GCullVolumeBuffers.CullVolumePlaneOffsets.Buffer->GetSize(), RLM_WriteOnly);
+				FVector4* CullVolumePlaneData = (FVector4*)RHILockVertexBuffer(GCullVolumeBuffers.CullVolumePlaneData.Buffer, 0, GCullVolumeBuffers.CullVolumePlaneData.Buffer->GetSize(), RLM_WriteOnly);
+			
+				int32 StartOffset = 0;
+
+				for (int32 VolumeIndex = 0; VolumeIndex < CullVolumes.Num(); VolumeIndex++)
+				{
+					const TArray<FPlane>& PlaneData = CullVolumes[VolumeIndex];
+					CullVolumePlaneOffsetsData[VolumeIndex * 2 + 0] = PlaneData.Num();
+					CullVolumePlaneOffsetsData[VolumeIndex * 2 + 1] = StartOffset;
+
+					for (int32 PlaneIndex = 0; PlaneIndex < PlaneData.Num(); PlaneIndex++)
+					{
+						CullVolumePlaneData[StartOffset + PlaneIndex] = FVector4(PlaneData[PlaneIndex], PlaneData[PlaneIndex].W);
+					}
+
+					StartOffset += PlaneData.Num();
+				}
+
+				RHIUnlockVertexBuffer(GCullVolumeBuffers.CullVolumePlaneOffsets.Buffer);
+				RHIUnlockVertexBuffer(GCullVolumeBuffers.CullVolumePlaneData.Buffer);
+			}
+		}
+
+		{
+			SCOPED_DRAW_EVENT(RHICmdList, ObjectFrustumCulling);
+
+			if (GPreCullTrianglesCulledObjectBuffers.Buffers.MaxObjects < Scene->DistanceFieldSceneData.NumObjectsInBuffer
+				|| GPreCullTrianglesCulledObjectBuffers.Buffers.MaxObjects > 3 * Scene->DistanceFieldSceneData.NumObjectsInBuffer)
+			{
+				GPreCullTrianglesCulledObjectBuffers.Buffers.MaxObjects = Scene->DistanceFieldSceneData.NumObjectsInBuffer * 5 / 4;
+				GPreCullTrianglesCulledObjectBuffers.Buffers.Release();
+				GPreCullTrianglesCulledObjectBuffers.Buffers.Initialize();
+			}
+
+			UE_LOG(LogShaders, Log, TEXT("Processing %u meshes"), ComponentsToPreCull.Num());
+			int32 TotalNumTrianglesProcessed = 0;
+			int32 TotalNumTrianglesVisible = 0;
+
+			for (int32 ComponentIndex = 0; ComponentIndex < ComponentsToPreCull.Num(); ComponentIndex++)
+			{
+				if (ComponentsToPreCull.Num() > 200 && (ComponentIndex % FMath::Max(ComponentsToPreCull.Num() / 10, 1)) == 0)
+				{
+					UE_LOG(LogShaders, Log, TEXT("PreCulled %u percent"), (uint32)(100 * ComponentIndex / (float)ComponentsToPreCull.Num()));
+				}
+
+				UStaticMeshComponent* StaticMeshComponent = ComponentsToPreCull[ComponentIndex];
+				FStaticMeshSceneProxy* StaticMeshProxy = (FStaticMeshSceneProxy*)StaticMeshComponent->SceneProxy;
+				FBoxSphereBounds Bounds = StaticMeshProxy->GetBounds();
+
+				for (int32 LODIndex = 0; LODIndex < StaticMeshComponent->StaticMesh->RenderData->LODResources.Num(); LODIndex++)
+				{
+					{
+						uint32 ClearValues[4] = { 0 };
+						RHICmdList.ClearUAV(GPreCullTrianglesCulledObjectBuffers.Buffers.ObjectIndirectArguments.UAV, ClearValues);
+
+						TShaderMapRef<FCullObjectsForBoundsCS> ComputeShader(GetGlobalShaderMap(Scene->GetFeatureLevel()));
+						RHICmdList.SetComputeShader(ComputeShader->GetComputeShader());
+						ComputeShader->SetParameters(RHICmdList, Scene, View, Bounds.GetBox().GetCenter(), Bounds.GetBox().GetExtent(), Parameters);
+
+						DispatchComputeShader(RHICmdList, *ComputeShader, FMath::DivideAndRoundUp<uint32>(Scene->DistanceFieldSceneData.NumObjectsInBuffer, UpdateObjectsGroupSize), 1, 1);
+						ComputeShader->UnsetParameters(RHICmdList, Scene);
+					}
+
+					FUniformMeshBuffers* UniformMeshBuffers = NULL;
+					const FMaterialRenderProxy* MaterialRenderProxy = NULL;
+					FUniformBufferRHIParamRef PrimitiveUniformBuffer = NULL;
+					const int32 NumUniformTriangles = FUniformMeshConverter::Convert(RHICmdList, *this, View, StaticMeshProxy->GetPrimitiveSceneInfo(), LODIndex, UniformMeshBuffers, MaterialRenderProxy, PrimitiveUniformBuffer);
+			
+					if (NumUniformTriangles > 0 && UniformMeshBuffers)
+					{
+						if (GPreCulledTriangleBuffers.Buffers.MaxIndices < NumUniformTriangles * 3)
+						{
+							GPreCulledTriangleBuffers.Buffers.MaxIndices = NumUniformTriangles * 3 * 5 / 4;
+							GPreCulledTriangleBuffers.ReleaseResource();
+							GPreCulledTriangleBuffers.InitResource();
+						}
+
+						uint32 MaskClearValues[4] = { 255 };
+						RHICmdList.ClearUAV(GPreCulledTriangleBuffers.Buffers.TriangleVisibleMask.UAV, MaskClearValues);
+
+						const FStaticMeshLODResources& StaticMeshLODResources = StaticMeshComponent->StaticMesh->RenderData->LODResources[LODIndex];
+
+						for (int32 SectionIndex = 0; SectionIndex < StaticMeshLODResources.Sections.Num(); SectionIndex++)
+						{
+							const FStaticMeshSection& Section = StaticMeshLODResources.Sections[SectionIndex];
+
+							{
+								TShaderMapRef<FPreCullTrianglesCS> ComputeShader(GetGlobalShaderMap(Scene->GetFeatureLevel()));
+								RHICmdList.SetComputeShader(ComputeShader->GetComputeShader());
+								ComputeShader->SetParameters(RHICmdList, View, GPreCulledTriangleBuffers.Buffers, Section.FirstIndex, Section.NumTriangles, *UniformMeshBuffers, CullVolumes.Num());
+
+								uint32 GroupSize = FMath::DivideAndRoundUp((int32)Section.NumTriangles, GDistanceFieldAOTileSizeX * GDistanceFieldAOTileSizeY);
+								DispatchComputeShader(RHICmdList, *ComputeShader, GroupSize, 1, 1);
+								ComputeShader->UnsetParameters(RHICmdList, GPreCulledTriangleBuffers.Buffers);
+							}
+						}
+
+						if (GPreCulledTriangleBuffersFromHeightfields.Buffers.MaxIndices < NumUniformTriangles * 3)
+						{
+							GPreCulledTriangleBuffersFromHeightfields.Buffers.MaxIndices = NumUniformTriangles * 3 * 5 / 4;
+							GPreCulledTriangleBuffersFromHeightfields.ReleaseResource();
+							GPreCulledTriangleBuffersFromHeightfields.InitResource();
+						}
+
+						uint32 HeightfieldMaskClearValues[4] = { 1 };
+						RHICmdList.ClearUAV(GPreCulledTriangleBuffersFromHeightfields.Buffers.TriangleVisibleMask.UAV, HeightfieldMaskClearValues);
+
+						for (int32 SectionIndex = 0; SectionIndex < StaticMeshLODResources.Sections.Num(); SectionIndex++)
+						{
+							const FStaticMeshSection& Section = StaticMeshLODResources.Sections[SectionIndex];
+							HeightfieldScene.PreCullTriangles(View, RHICmdList, GPreCulledTriangleBuffersFromHeightfields.Buffers, Section.FirstIndex, Section.NumTriangles, *UniformMeshBuffers);
+						}
+
+						{
+							uint32* TriangleVisibleMaskData = (uint32*)RHILockVertexBuffer(GPreCulledTriangleBuffers.Buffers.TriangleVisibleMask.Buffer, 0, NumUniformTriangles * sizeof(uint32), RLM_ReadOnly);
+							uint32* TriangleVisibleMaskDataFromHeightfields = (uint32*)RHILockVertexBuffer(GPreCulledTriangleBuffersFromHeightfields.Buffers.TriangleVisibleMask.Buffer, 0, NumUniformTriangles * sizeof(uint32), RLM_ReadOnly);
+
+							FIndexArrayView IndexBufferView = StaticMeshLODResources.IndexBuffer.GetArrayView();
+							int32 GlobalTriangleIndex = 0;
+
+							TArray<int32> NumTrianglesPerSection;
+							NumTrianglesPerSection.Empty(StaticMeshLODResources.Sections.Num());
+
+							TArray<uint32> CopiedIndices;
+							CopiedIndices.Empty(NumUniformTriangles * 3);
+
+							for (int32 SectionIndex = 0; SectionIndex < StaticMeshLODResources.Sections.Num(); SectionIndex++)
+							{
+								const FStaticMeshSection& Section = StaticMeshLODResources.Sections[SectionIndex];
+								int32 NumVisibleTrianglesInSection = 0;
+
+								TotalNumTrianglesProcessed += Section.NumTriangles;
+
+								for (uint32 TriangleIndex = 0; TriangleIndex < Section.NumTriangles; TriangleIndex++)
+								{
+									uint32 TriangleMask = TriangleVisibleMaskData[GlobalTriangleIndex];
+									check(TriangleMask == 0 || TriangleMask == 1);
+
+									uint32 TriangleMaskFromHeightfields = TriangleVisibleMaskDataFromHeightfields[GlobalTriangleIndex];
+									check(TriangleMaskFromHeightfields == 0 || TriangleMaskFromHeightfields == 1);
+
+									if (TriangleMask > 0 && TriangleMaskFromHeightfields > 0)
+									{
+										CopiedIndices.Add(IndexBufferView[GlobalTriangleIndex * 3 + 0]);
+										CopiedIndices.Add(IndexBufferView[GlobalTriangleIndex * 3 + 1]);
+										CopiedIndices.Add(IndexBufferView[GlobalTriangleIndex * 3 + 2]);
+										NumVisibleTrianglesInSection++;
+									}
+
+									GlobalTriangleIndex++;
+								}
+
+								TotalNumTrianglesVisible += NumVisibleTrianglesInSection;
+								NumTrianglesPerSection.Add(NumVisibleTrianglesInSection);
+							}
+
+							RHIUnlockVertexBuffer(GPreCulledTriangleBuffers.Buffers.TriangleVisibleMask.Buffer);
+							RHIUnlockVertexBuffer(GPreCulledTriangleBuffersFromHeightfields.Buffers.TriangleVisibleMask.Buffer);
+
+							StaticMeshComponent->UpdatePreCulledData(LODIndex, CopiedIndices, NumTrianglesPerSection);
+						}
+					}
+				}
+			}
+
+			UE_LOG(LogShaders, Log, TEXT("%.3fm triangles visible out of %.3fm total (%u percent removed)"), TotalNumTrianglesVisible / 1000000.0f, TotalNumTrianglesProcessed / 1000000.0f, (int32)(100 * (1 - TotalNumTrianglesVisible / (float)TotalNumTrianglesProcessed)));
+		}
+	}
+}
+
+void PreCullStaticMeshesOuter(FRHICommandListImmediate& RHICmdList, FDeferredShadingSceneRenderer* SceneRenderer, const TArray<UStaticMeshComponent*>& ComponentsToPreCull, const TArray<TArray<FPlane> >& CullVolumes)
+{
+	FMemMark MemStackMark(FMemStack::Get());
+	
+	SceneRenderer->PreCullStaticMeshes(RHICmdList, ComponentsToPreCull, CullVolumes);
+
+	check(&RHICmdList == &FRHICommandListExecutor::GetImmediateCommandList());
+	check(IsInRenderingThread());
+	{
+		QUICK_SCOPE_CYCLE_COUNTER(STAT_CaptureSceneToScratchCubemap_Flush);
+		FRHICommandListExecutor::GetImmediateCommandList().ImmediateFlush(EImmediateFlushType::FlushRHIThread);
+	}
+
+	FSceneRenderer::WaitForTasksClearSnapshotsAndDeleteSceneRenderer(RHICmdList, SceneRenderer);
+}
+
+/** 
+ * Render target class required for rendering the scene.
+ */
+class FPreCullTrianglesRenderTarget : public FRenderResource, public FRenderTarget
+{
+public:
+
+	FPreCullTrianglesRenderTarget() :
+		Size(128)
+	{}
+
+	virtual const FTexture2DRHIRef& GetRenderTargetTexture() const 
+	{
+		static FTexture2DRHIRef DummyTexture;
+		return DummyTexture;
+	}
+
+	virtual FIntPoint GetSizeXY() const { return FIntPoint(Size, Size); }
+	virtual float GetDisplayGamma() const { return 1.0f; }
+
+private:
+
+	int32 Size;
+};
+
+TGlobalResource<FPreCullTrianglesRenderTarget> GPreCullTrianglesRenderTarget;
+
+void FScene::PreCullStaticMeshes(const TArray<UStaticMeshComponent*>& ComponentsToPreCull, const TArray<TArray<FPlane> >& CullVolumes)
+{
+	// Alert the RHI that we're rendering a new frame
+	// Not really a new frame, but it will allow pooling mechanisms to update, like the uniform buffer pool
+	ENQUEUE_UNIQUE_RENDER_COMMAND(
+		BeginFrame,
+	{
+		GFrameNumberRenderThread++;
+		RHICmdList.BeginFrame();
+	})
+
+	FSceneViewFamilyContext ViewFamily( FSceneViewFamily::ConstructionValues(
+		&GPreCullTrianglesRenderTarget,
+		this,
+		FEngineShowFlags(ESFIM_Game))
+		.SetWorldTimes( 0.0f, 0.0f, 0.0f )				
+		.SetResolveScene(false) );
+
+	// Disable features that are not desired when capturing the scene
+	ViewFamily.EngineShowFlags.PostProcessing = 0;
+	ViewFamily.EngineShowFlags.MotionBlur = 0;
+	ViewFamily.EngineShowFlags.SetOnScreenDebug(false);
+	ViewFamily.EngineShowFlags.HMDDistortion = 0;
+	// Exclude particles and light functions as they are usually dynamic, and can't be captured well
+	ViewFamily.EngineShowFlags.Particles = 0;
+	ViewFamily.EngineShowFlags.LightFunctions = 0;
+	ViewFamily.EngineShowFlags.CompositeEditorPrimitives = 0;
+	// These are highly dynamic and can't be captured effectively
+	ViewFamily.EngineShowFlags.LightShafts = 0;
+
+	FSceneViewInitOptions ViewInitOptions;
+	ViewInitOptions.ViewFamily = &ViewFamily;
+	ViewInitOptions.BackgroundColor = FLinearColor::Black;
+	ViewInitOptions.OverlayColor = FLinearColor::Black;
+	ViewInitOptions.SetViewRectangle(FIntRect(0, 0, 128 * 1, 128 * 1));
+
+	// Projection matrix based on the fov, near / far clip settings
+	// Each face always uses a 90 degree field of view
+	if ((int32)ERHIZBuffer::IsInverted != 0)
+	{
+		ViewInitOptions.ProjectionMatrix = FReversedZPerspectiveMatrix(
+			90.0f * (float)PI / 360.0f,
+			(float)128 * 1,
+			(float)128 * 1,
+			10
+			);
+	}
+	else
+	{
+		ViewInitOptions.ProjectionMatrix = FPerspectiveMatrix(
+			90.0f * (float)PI / 360.0f,
+			(float)128 * 1,
+			(float)128 * 1,
+			10
+			);
+	}
+
+	ViewInitOptions.ViewRotationMatrix = FMatrix::Identity;
+	ViewInitOptions.ViewOrigin = FVector::ZeroVector;
+
+	FSceneView* View = new FSceneView(ViewInitOptions);
+
+	View->StartFinalPostprocessSettings(FVector::ZeroVector);
+	View->EndFinalPostprocessSettings(ViewInitOptions);
+
+	ViewFamily.Views.Add(View);
+
+	FSceneRenderer* SceneRenderer = FSceneRenderer::CreateSceneRenderer(&ViewFamily, NULL);
+
+	ENQUEUE_UNIQUE_RENDER_COMMAND_THREEPARAMETER( 
+		PreCullTrianglesCommand,
+		FSceneRenderer*, SceneRenderer, SceneRenderer,
+		// Note: passing these arrays which are on the GT stack, only because this is an offline operation and the game thread will flush rendering commands before exiting the scope
+		const TArray<UStaticMeshComponent*>&, ComponentsToPreCull, ComponentsToPreCull,
+		const TArray<TArray<FPlane> >&, CullVolumes, CullVolumes,
+	{
+		PreCullStaticMeshesOuter(RHICmdList, (FDeferredShadingSceneRenderer*)SceneRenderer, ComponentsToPreCull, CullVolumes);
+
+		RHICmdList.EndFrame();
+	});
+
+	FlushRenderingCommands();
+}
\ No newline at end of file
diff --git a/Engine/Source/Runtime/Renderer/Private/ScenePrivate.h b/Engine/Source/Runtime/Renderer/Private/ScenePrivate.h
--- a/Engine/Source/Runtime/Renderer/Private/ScenePrivate.h
+++ b/Engine/Source/Runtime/Renderer/Private/ScenePrivate.h
@@ -1671,6 +1671,7 @@ public:
 	virtual void UpdateSceneCaptureContents(class USceneCaptureComponentCube* CaptureComponent) override;
 	virtual void AllocateReflectionCaptures(const TArray<UReflectionCaptureComponent*>& NewCaptures) override;
 	virtual void UpdateSkyCaptureContents(const USkyLightComponent* CaptureComponent, bool bCaptureEmissiveOnly, UTextureCube* SourceCubemap, FTexture* OutProcessedTexture, FSHVectorRGB3& OutIrradianceEnvironmentMap) override; 
+	virtual void PreCullStaticMeshes(const TArray<UStaticMeshComponent*>& ComponentsToPreCull, const TArray<TArray<FPlane> >& CullVolumes) override;
 	virtual void AddPrecomputedLightVolume(const class FPrecomputedLightVolume* Volume) override;
 	virtual void RemovePrecomputedLightVolume(const class FPrecomputedLightVolume* Volume) override;
 	virtual void UpdateLightTransform(ULightComponent* Light) override;
diff --git a/Engine/Source/Runtime/Renderer/Private/SurfelTree.cpp b/Engine/Source/Runtime/Renderer/Private/SurfelTree.cpp
--- a/Engine/Source/Runtime/Renderer/Private/SurfelTree.cpp
+++ b/Engine/Source/Runtime/Renderer/Private/SurfelTree.cpp
@@ -228,6 +228,11 @@ public:
 		FGlobalShader::SetParameters(RHICmdList, ShaderRHI, View);
 
 		const FScene* Scene = (const FScene*)View.Family->Scene;
+
+		FUnorderedAccessViewRHIParamRef UniformMeshUAVs[1];
+		UniformMeshUAVs[0] = Scene->DistanceFieldSceneData.SurfelBuffers->InterpolatedVertexData.UAV;
+		RHICmdList.TransitionResources(EResourceTransitionAccess::ERWBarrier, EResourceTransitionPipeline::EComputeToCompute, UniformMeshUAVs, ARRAY_COUNT(UniformMeshUAVs));
+
 		SurfelBufferParameters.Set(RHICmdList, ShaderRHI, *Scene->DistanceFieldSceneData.SurfelBuffers, *Scene->DistanceFieldSceneData.InstancedSurfelBuffers);
 		
 		SetShaderValue(RHICmdList, ShaderRHI, SurfelStartIndex, SurfelStartIndexValue);
@@ -246,6 +251,10 @@ public:
 		SurfelBufferParameters.UnsetParameters(RHICmdList, ShaderRHI);
 		// RHISetStreamOutTargets doesn't unbind existing uses like render targets do 
 		SetSRVParameter(RHICmdList, ShaderRHI, TriangleVertexData, NULL);
+
+		FUnorderedAccessViewRHIParamRef UniformMeshUAVs[1];
+		UniformMeshUAVs[0] = Scene->DistanceFieldSceneData.SurfelBuffers->InterpolatedVertexData.UAV;
+		RHICmdList.TransitionResources(EResourceTransitionAccess::EReadable, EResourceTransitionPipeline::EComputeToCompute, UniformMeshUAVs, ARRAY_COUNT(UniformMeshUAVs));
 	}
 
 	virtual bool Serialize(FArchive& Ar) override
@@ -293,7 +302,7 @@ void GenerateSurfelRepresentation(FRHICommandListImmediate& RHICmdList, FSceneRe
 		FUniformMeshBuffers* UniformMeshBuffers = NULL;
 		const FMaterialRenderProxy* MaterialRenderProxy = NULL;
 		FUniformBufferRHIParamRef PrimitiveUniformBuffer = NULL;
-		const int32 NumUniformTriangles = FUniformMeshConverter::Convert(RHICmdList, Renderer, View, PrimitiveSceneInfo, UniformMeshBuffers, MaterialRenderProxy, PrimitiveUniformBuffer);
+		const int32 NumUniformTriangles = FUniformMeshConverter::Convert(RHICmdList, Renderer, View, PrimitiveSceneInfo, 0, UniformMeshBuffers, MaterialRenderProxy, PrimitiveUniformBuffer);
 
 		if (NumUniformTriangles > 0 && MaterialRenderProxy && PrimitiveUniformBuffer)
 		{
