Worked around AMD driver bug with distance field shadowing, so that the feature works again without having to get the latest driver

diff --git a/Engine/Shaders/DistanceFieldShadowing.usf b/Engine/Shaders/DistanceFieldShadowing.usf
--- a/Engine/Shaders/DistanceFieldShadowing.usf
+++ b/Engine/Shaders/DistanceFieldShadowing.usf
@@ -163,6 +163,16 @@ bool RayHitSphere(float3 RayOrigin, float3 UnitRayDirection, float RayLength, fl
 RWBuffer<uint> RWShadowTileHeadDataUnpacked;
 RWBuffer<uint> RWShadowTileArrayData;
 
+[numthreads(1, 1, 1)]
+void WorkaroundAMDBugCS(
+	uint3 GroupId : SV_GroupID,
+	uint3 DispatchThreadId : SV_DispatchThreadID,
+    uint3 GroupThreadId : SV_GroupThreadID) 
+{
+	RWShadowTileHeadDataUnpacked[0] = 0;
+	RWShadowTileArrayData[0] = 0;
+}
+
 float2 NumGroups;
 
 [numthreads(THREADGROUP_SIZEX, THREADGROUP_SIZEY, 1)]
diff --git a/Engine/Source/Runtime/Renderer/Private/DistanceFieldShadowing.cpp b/Engine/Source/Runtime/Renderer/Private/DistanceFieldShadowing.cpp
--- a/Engine/Source/Runtime/Renderer/Private/DistanceFieldShadowing.cpp
+++ b/Engine/Source/Runtime/Renderer/Private/DistanceFieldShadowing.cpp
@@ -391,6 +391,70 @@ IMPLEMENT_SHADER_TYPE(,FShadowObjectCullPS,TEXT("DistanceFieldShadowing"),TEXT("
 
 
 
+/**  */
+class FWorkaroundAMDBugCS : public FGlobalShader
+{
+	DECLARE_SHADER_TYPE(FWorkaroundAMDBugCS,Global)
+public:
+
+	static bool ShouldCache(EShaderPlatform Platform)
+	{
+		return IsFeatureLevelSupported(Platform, ERHIFeatureLevel::SM5) && DoesPlatformSupportDistanceFieldShadowing(Platform);
+	}
+
+	FWorkaroundAMDBugCS(const ShaderMetaType::CompiledShaderInitializerType& Initializer)
+		: FGlobalShader(Initializer)
+	{
+		ShadowTileHeadDataUnpacked.Bind(Initializer.ParameterMap, TEXT("ShadowTileHeadDataUnpacked"));
+		ShadowTileArrayData.Bind(Initializer.ParameterMap, TEXT("ShadowTileArrayData"));
+	}
+
+	FWorkaroundAMDBugCS()
+	{
+	}
+
+	void SetParameters(FRHICommandList& RHICmdList, const FSceneView& View, FLightTileIntersectionResources* TileIntersectionResources)
+	{
+		FComputeShaderRHIParamRef ShaderRHI = GetComputeShader();
+
+		FGlobalShader::SetParameters(RHICmdList, ShaderRHI, View);
+
+		FUnorderedAccessViewRHIParamRef OutUAVs[2];
+		OutUAVs[0] = TileIntersectionResources->TileHeadDataUnpacked.UAV;
+		OutUAVs[1] = TileIntersectionResources->TileArrayData.UAV;
+		RHICmdList.TransitionResources(EResourceTransitionAccess::ERWBarrier, EResourceTransitionPipeline::EComputeToCompute, OutUAVs, ARRAY_COUNT(OutUAVs));
+
+		ShadowTileHeadDataUnpacked.SetBuffer(RHICmdList, ShaderRHI, TileIntersectionResources->TileHeadDataUnpacked);
+		ShadowTileArrayData.SetBuffer(RHICmdList, ShaderRHI, TileIntersectionResources->TileArrayData);
+	}
+
+	void UnsetParameters(FRHICommandList& RHICmdList, FLightTileIntersectionResources* TileIntersectionResources)
+	{
+		ShadowTileHeadDataUnpacked.UnsetUAV(RHICmdList, GetComputeShader());
+		ShadowTileArrayData.UnsetUAV(RHICmdList, GetComputeShader());
+
+		FUnorderedAccessViewRHIParamRef OutUAVs[2];
+		OutUAVs[0] = TileIntersectionResources->TileHeadDataUnpacked.UAV;
+		OutUAVs[1] = TileIntersectionResources->TileArrayData.UAV;
+		RHICmdList.TransitionResources(EResourceTransitionAccess::EReadable, EResourceTransitionPipeline::EComputeToCompute, OutUAVs, ARRAY_COUNT(OutUAVs));
+	}
+
+	virtual bool Serialize(FArchive& Ar) override
+	{		
+		bool bShaderHasOutdatedParameters = FGlobalShader::Serialize(Ar);
+		Ar << ShadowTileHeadDataUnpacked;
+		Ar << ShadowTileArrayData;
+		return bShaderHasOutdatedParameters;
+	}
+
+private:
+
+	FRWShaderParameter ShadowTileHeadDataUnpacked;
+	FRWShaderParameter ShadowTileArrayData;
+};
+
+IMPLEMENT_SHADER_TYPE(,FWorkaroundAMDBugCS,TEXT("DistanceFieldShadowing"),TEXT("WorkaroundAMDBugCS"),SF_Compute);
+
 enum EDistanceFieldShadowingType
 {
 	DFS_DirectionalLightScatterTileCulling,
@@ -699,7 +763,20 @@ void CullDistanceFieldObjectsForLight(
 
 			TileIntersectionResources->Initialize();
 		}
+		
+		if (View.GetShaderPlatform() == SP_PCD3D_SM5)
+		{
+			// AMD PC driver versions before 15.10 have a bug where the UAVs are not updated correctly in FShadowObjectCullPS unless we touch them here (tested on 15.7)
+			// This bug is fixed in 15.10, but the workaround makes sure the feature works everywhere
+			TShaderMapRef<FWorkaroundAMDBugCS> ComputeShader(View.ShaderMap);
 
+			RHICmdList.SetComputeShader(ComputeShader->GetComputeShader());
+			ComputeShader->SetParameters(RHICmdList, View, TileIntersectionResources);
+			DispatchComputeShader(RHICmdList, *ComputeShader, 1, 1, 1);
+
+			ComputeShader->UnsetParameters(RHICmdList, TileIntersectionResources);
+		}
+		
 		{
 			TShaderMapRef<FClearTilesCS> ComputeShader(View.ShaderMap);
 
@@ -712,7 +789,7 @@ void CullDistanceFieldObjectsForLight(
 
 			ComputeShader->UnsetParameters(RHICmdList, TileIntersectionResources);
 		}
-		
+
 		{
 			TShaderMapRef<FShadowObjectCullVS> VertexShader(View.ShaderMap);
 			TShaderMapRef<FShadowObjectCullPS> PixelShader(View.ShaderMap);
