Using the visible area mesh for this pass creates incorrect results.

diff --git a/Engine/Source/Runtime/Renderer/Private/PostProcess/PostProcessHistogramReduce.cpp b/Engine/Source/Runtime/Renderer/Private/PostProcess/PostProcessHistogramReduce.cpp
--- a/Engine/Source/Runtime/Renderer/Private/PostProcess/PostProcessHistogramReduce.cpp
+++ b/Engine/Source/Runtime/Renderer/Private/PostProcess/PostProcessHistogramReduce.cpp
@@ -128,7 +128,7 @@ void FRCPassPostProcessHistogramReduce::Process(FRenderingCompositePassContext&
 
 	PixelShader->SetPS(Context, LoopSizeValue);
 
-	DrawPostProcessPass(
+	DrawRectangle(
 		Context.RHICmdList,
 		0, 0,
 		DestSize.X, DestSize.Y,
@@ -137,8 +137,6 @@ void FRCPassPostProcessHistogramReduce::Process(FRenderingCompositePassContext&
 		DestSize,
 		SrcSize,
 		*VertexShader,
-		View.StereoPass,
-		Context.HasHmdMesh(),
 		EDRF_UseTriangleOptimization);
 
 	Context.RHICmdList.CopyToResolveTarget(DestRenderTarget.TargetableTexture, DestRenderTarget.ShaderResourceTexture, false, FResolveParams());
