Fixed compile error

diff --git a/Engine/Source/Runtime/Renderer/Private/PreCullTriangles.cpp b/Engine/Source/Runtime/Renderer/Private/PreCullTriangles.cpp
--- a/Engine/Source/Runtime/Renderer/Private/PreCullTriangles.cpp
+++ b/Engine/Source/Runtime/Renderer/Private/PreCullTriangles.cpp
@@ -568,7 +568,7 @@ void FScene::PreCullStaticMeshes(const TArray<UStaticMeshComponent*>& Components
 	// Exclude particles and light functions as they are usually dynamic, and can't be captured well
 	ViewFamily.EngineShowFlags.Particles = 0;
 	ViewFamily.EngineShowFlags.LightFunctions = 0;
-	ViewFamily.EngineShowFlags.CompositeEditorPrimitives = 0;
+	ViewFamily.EngineShowFlags.SetCompositeEditorPrimitives(false);
 	// These are highly dynamic and can't be captured effectively
 	ViewFamily.EngineShowFlags.LightShafts = 0;
 
