Fix orioneditor transition ensures

diff --git a/Engine/Source/Runtime/Renderer/Private/DeferredShadingRenderer.cpp b/Engine/Source/Runtime/Renderer/Private/DeferredShadingRenderer.cpp
--- a/Engine/Source/Runtime/Renderer/Private/DeferredShadingRenderer.cpp
+++ b/Engine/Source/Runtime/Renderer/Private/DeferredShadingRenderer.cpp
@@ -1016,6 +1016,8 @@ void FDeferredShadingSceneRenderer::Render(FRHICommandListImmediate& RHICmdList)
 
 			GCompositionLighting.ProcessBeforeBasePass(RHICmdList, Views[ViewIndex]);
 		}
+		//GBuffer pass will want to write to SceneDepthZ
+		RHICmdList.TransitionResource(EResourceTransitionAccess::EWritable, SceneContext.GetSceneDepthTexture());
 		ServiceLocalQueue();
 	}
 
diff --git a/Engine/Source/Runtime/Renderer/Private/ScenePrivate.h b/Engine/Source/Runtime/Renderer/Private/ScenePrivate.h
--- a/Engine/Source/Runtime/Renderer/Private/ScenePrivate.h
+++ b/Engine/Source/Runtime/Renderer/Private/ScenePrivate.h
@@ -673,6 +673,7 @@ public:
 		{
 			// Create the texture needed for EyeAdaptation
 			FPooledRenderTargetDesc Desc(FPooledRenderTargetDesc::Create2DDesc(FIntPoint(1, 1), PF_R32_FLOAT, FClearValueBinding::None, TexCreate_None, TexCreate_RenderTargetable, false));
+			Desc.AutoWritable = false;
 			GRenderTargetPool.FindFreeElement(RHICmdList, Desc, EyeAdaptationRT, TEXT("EyeAdaptation"));
 		}
 		return EyeAdaptationRT;
