Transition ensure fixes

diff --git a/Engine/Source/Runtime/Renderer/Private/GPUBenchmark.cpp b/Engine/Source/Runtime/Renderer/Private/GPUBenchmark.cpp
--- a/Engine/Source/Runtime/Renderer/Private/GPUBenchmark.cpp
+++ b/Engine/Source/Runtime/Renderer/Private/GPUBenchmark.cpp
@@ -314,6 +314,8 @@ void RendererGPUBenchmark(FRHICommandListImmediate& RHICmdList, FSynthBenchmarkR
 	TRefCountPtr<IPooledRenderTarget> RTItems[3];
 	{
 		FPooledRenderTargetDesc Desc(FPooledRenderTargetDesc::Create2DDesc(FIntPoint(GBenchmarkResolution, GBenchmarkResolution), PF_B8G8R8A8, FClearValueBinding::None, TexCreate_None, TexCreate_RenderTargetable | TexCreate_ShaderResource, false));
+		Desc.AutoWritable = false;
+
 		GRenderTargetPool.FindFreeElement(RHICmdList, Desc, RTItems[0], TEXT("Benchmark0"));
 		GRenderTargetPool.FindFreeElement(RHICmdList, Desc, RTItems[1], TEXT("Benchmark1"));
 
@@ -435,7 +437,7 @@ void RendererGPUBenchmark(FRHICommandListImmediate& RHICmdList, FSynthBenchmarkR
 
 				GRenderTargetPool.VisualizeTexture.SetCheckPoint(RHICmdList, RTItems[DestRTIndex]);
 
-				SetRenderTarget(RHICmdList, RTItems[DestRTIndex]->GetRenderTargetItem().TargetableTexture, FTextureRHIRef());	
+				SetRenderTarget(RHICmdList, RTItems[DestRTIndex]->GetRenderTargetItem().TargetableTexture, FTextureRHIRef(), true);	
 
 				// decide how much work we do in this pass
 				LocalWorkScale[Iteration] = (Iteration / 10.f + 1.f) * WorkScale;
diff --git a/Engine/Source/Runtime/Renderer/Private/LightShaftRendering.cpp b/Engine/Source/Runtime/Renderer/Private/LightShaftRendering.cpp
--- a/Engine/Source/Runtime/Renderer/Private/LightShaftRendering.cpp
+++ b/Engine/Source/Runtime/Renderer/Private/LightShaftRendering.cpp
@@ -467,7 +467,7 @@ void DownsamplePass(FRHICommandListImmediate& RHICmdList, const FViewInfo& View,
 	const uint32 DownsampledSizeX = View.ViewRect.Width() / DownsampleFactor;
 	const uint32 DownsampledSizeY = View.ViewRect.Height() / DownsampleFactor;
 
-	SetRenderTarget(RHICmdList, LightShaftsDest->GetRenderTargetItem().TargetableTexture, FTextureRHIRef());
+	SetRenderTarget(RHICmdList, LightShaftsDest->GetRenderTargetItem().TargetableTexture, FTextureRHIRef(), true);
 	RHICmdList.SetViewport(DownSampledXY.X, DownSampledXY.Y, 0.0f, DownSampledXY.X + DownsampledSizeX, DownSampledXY.Y + DownsampledSizeY, 1.0f);
 
 	// Set shaders and texture
@@ -872,7 +872,7 @@ void ApplyLightShaftBloom(FRHICommandListImmediate& RHICmdList, const FViewInfo&
 		*ScreenVertexShader,
 		EDRF_UseTriangleOptimization);
 
-	SceneContext.FinishRenderingSceneColor(RHICmdList, false);
+	SceneContext.FinishRenderingSceneColor(RHICmdList, true);
 }
 
 void FSceneViewState::TrimHistoryRenderTargets(const FScene* Scene)
diff --git a/Engine/Source/Runtime/Renderer/Private/PostProcess/PostProcessTemporalAA.cpp b/Engine/Source/Runtime/Renderer/Private/PostProcess/PostProcessTemporalAA.cpp
--- a/Engine/Source/Runtime/Renderer/Private/PostProcess/PostProcessTemporalAA.cpp
+++ b/Engine/Source/Runtime/Renderer/Private/PostProcess/PostProcessTemporalAA.cpp
@@ -368,7 +368,7 @@ void FRCPassPostProcessDOFTemporalAA::Process(FRenderingCompositePassContext& Co
 FPooledRenderTargetDesc FRCPassPostProcessDOFTemporalAA::ComputeOutputDesc(EPassOutputId InPassOutputId) const
 {
 	FPooledRenderTargetDesc Ret = GetInput(ePId_Input0)->GetOutput()->RenderTargetDesc;
-
+	Ret.AutoWritable = false;
 	Ret.DebugName = TEXT("BokehDOFTemporalAA");
 
 	return Ret;
diff --git a/Engine/Source/Runtime/Renderer/Private/TranslucentLighting.cpp b/Engine/Source/Runtime/Renderer/Private/TranslucentLighting.cpp
--- a/Engine/Source/Runtime/Renderer/Private/TranslucentLighting.cpp
+++ b/Engine/Source/Runtime/Renderer/Private/TranslucentLighting.cpp
@@ -551,7 +551,7 @@ void FProjectedShadowInfo::RenderTranslucencyDepths(FRHICommandList& RHICmdList,
 			ShadowTransmission0,
 			ShadowTransmission1
 		};
-		SetRenderTargets(RHICmdList, ARRAY_COUNT(RenderTargets), RenderTargets, FTextureRHIParamRef(), 0, NULL);
+		SetRenderTargets(RHICmdList, ARRAY_COUNT(RenderTargets), RenderTargets, FTextureRHIParamRef(), 0, NULL, true);
 
 		// Clear the shadow and its border
 		RHICmdList.SetViewport(
