fixed UE-22333 Selection outline does not appear on any interactable object

diff --git a/Engine/Source/Runtime/Renderer/Private/ScenePrivate.h b/Engine/Source/Runtime/Renderer/Private/ScenePrivate.h
--- a/Engine/Source/Runtime/Renderer/Private/ScenePrivate.h
+++ b/Engine/Source/Runtime/Renderer/Private/ScenePrivate.h
@@ -846,12 +846,8 @@ public:
 
 		if(InputAsMID)
 		{
-			auto ParentAsMI = Cast<UMaterialInstance>(InputAsMID->Parent);
-			if(ParentAsMI)
-			{
-				// parent is an MID so we need to copy the MID Vector and Scalar parameters over
-				NewMID->CopyInterpParameters(ParentAsMI);
-			}
+			// parent is an MID so we need to copy the MID Vector and Scalar parameters over
+			NewMID->CopyInterpParameters(InputAsMID);
 		}
 
 		check(NewMID->GetRenderProxy(false));
