Merging CL#2724751

fixed  UE-20693 Light Propagation Volume crash with Translucent material applied to mesh

diff --git a/Engine/Source/Runtime/Renderer/Private/ShadowRendering.cpp b/Engine/Source/Runtime/Renderer/Private/ShadowRendering.cpp
--- a/Engine/Source/Runtime/Renderer/Private/ShadowRendering.cpp
+++ b/Engine/Source/Runtime/Renderer/Private/ShadowRendering.cpp
@@ -1092,8 +1092,14 @@ bool FShadowDepthDrawingPolicyFactory::DrawDynamicMesh(
 		const bool bLocalOnePassPointLightShadow = Context.ShadowInfo->CascadeSettings.bOnePassPointLightShadow;
 		const bool bReflectiveShadowmap = Context.ShadowInfo->bReflectiveShadowmap && !bLocalOnePassPointLightShadow;
 
-		if ( ((!IsTranslucentBlendMode(BlendMode) || bReflectiveShadowmap ) && ShadingModel != MSM_Unlit)
-			|| (Material->ShouldInjectEmissiveIntoLPV() && bReflectiveShadowmap) )
+		bool bProcess = !IsTranslucentBlendMode(BlendMode) && ShadingModel != MSM_Unlit;
+
+		if (bReflectiveShadowmap && Material->ShouldInjectEmissiveIntoLPV())
+		{
+			bProcess = true;
+		}
+
+		if (bProcess)
 		{
 			const bool bLocalDirectionalLight = Context.ShadowInfo->bDirectionalLight;
 			const bool bPreShadow = Context.ShadowInfo->bPreShadow;
