Resource transition fixes

diff --git a/Engine/Source/Runtime/Renderer/Private/DeferredShadingRenderer.cpp b/Engine/Source/Runtime/Renderer/Private/DeferredShadingRenderer.cpp
--- a/Engine/Source/Runtime/Renderer/Private/DeferredShadingRenderer.cpp
+++ b/Engine/Source/Runtime/Renderer/Private/DeferredShadingRenderer.cpp
@@ -1279,7 +1279,7 @@ void FDeferredShadingSceneRenderer::Render(FRHICommandListImmediate& RHICmdList)
 			GetRendererModule().RenderPostOpaqueExtensions(View, RHICmdList, SceneContext);
 		}
 
-		SceneContext.FinishRenderingSceneColor(RHICmdList, false);
+		SceneContext.FinishRenderingSceneColor(RHICmdList, true);
 	}
 
 	// No longer needed, release
@@ -1904,4 +1904,6 @@ void FDeferredShadingSceneRenderer::DownsampleDepthSurface(FRHICommandList& RHIC
 		SceneContext.GetBufferSizeXY(),
 		*ScreenVertexShader,
 		EDRF_UseTriangleOptimization);
+
+	RHICmdList.TransitionResource(EResourceTransitionAccess::EReadable, RenderTarget);
 }
diff --git a/Engine/Source/Runtime/Renderer/Private/DistortionRendering.cpp b/Engine/Source/Runtime/Renderer/Private/DistortionRendering.cpp
--- a/Engine/Source/Runtime/Renderer/Private/DistortionRendering.cpp
+++ b/Engine/Source/Runtime/Renderer/Private/DistortionRendering.cpp
@@ -991,7 +991,7 @@ void FSceneRenderer::RenderDistortion(FRHICommandListImmediate& RHICmdList)
 		RHICmdList.SetDepthStencilState(TStaticDepthStencilState<false, CF_DepthNearOrEqual>::GetRHI());		
 			
 		// Distortions RT is no longer needed, buffer can be reused by the pool, see BeginRenderingDistortionAccumulation() call above
-		SceneContext.FinishRenderingSceneColor(RHICmdList, false);
+		SceneContext.FinishRenderingSceneColor(RHICmdList, true);
 	}
 }
 
diff --git a/Engine/Source/Runtime/Renderer/Private/LightShaftRendering.cpp b/Engine/Source/Runtime/Renderer/Private/LightShaftRendering.cpp
--- a/Engine/Source/Runtime/Renderer/Private/LightShaftRendering.cpp
+++ b/Engine/Source/Runtime/Renderer/Private/LightShaftRendering.cpp
@@ -734,43 +734,58 @@ void FDeferredShadingSceneRenderer::RenderLightShaftOcclusion(FRHICommandListImm
 
 			if (bEnableOcclusion && LightSceneInfo->Proxy->GetLightType() == LightType_Directional)
 			{
-				// Allocate light shaft render targets on demand, using the pool
-				// Need two targets to ping pong between
-				AllocateOrReuseLightShaftRenderTarget(RHICmdList, LightShafts0, TEXT("LightShafts0"));
-				AllocateOrReuseLightShaftRenderTarget(RHICmdList, LightShafts1, TEXT("LightShafts1"));
+				bool bWillRenderLightShafts = false;
 
-				for (int ViewIndex = 0;ViewIndex < Views.Num();ViewIndex++)
+				for (int ViewIndex = 0; ViewIndex < Views.Num(); ViewIndex++)
 				{
 					FViewInfo& View = Views[ViewIndex];
-					
-					SCOPED_DRAW_EVENTF(RHICmdList, RenderLightShaftOcclusion, TEXT("RenderLightShaftOcclusion %dx%d (multiple passes)"), View.ViewRect.Width(), View.ViewRect.Height());
 
 					if (ShouldRenderLightShaftsForLight(View, LightSceneInfo))
 					{
-						INC_DWORD_STAT(STAT_LightShaftsLights);
+						bWillRenderLightShafts = true;
+					}
+				}
 
-						// Create a downsampled occlusion mask from scene depth, result will be in LightShafts0
-						DownsamplePass<true>(RHICmdList, View, LightSceneInfo, LightShafts0, LightShafts1);
+				if (bWillRenderLightShafts)
+				{
+					// Allocate light shaft render targets on demand, using the pool
+					// Need two targets to ping pong between
+					AllocateOrReuseLightShaftRenderTarget(RHICmdList, LightShafts0, TEXT("LightShafts0"));
+					AllocateOrReuseLightShaftRenderTarget(RHICmdList, LightShafts1, TEXT("LightShafts1"));
 
-						FSceneViewState* ViewState = (FSceneViewState*)View.State;
-						// Find the previous frame's occlusion mask
-						TRefCountPtr<IPooledRenderTarget>* HistoryState = ViewState ? &ViewState->LightShaftOcclusionHistoryRT : NULL;
-						TRefCountPtr<IPooledRenderTarget> HistoryOutput;
+					for (int ViewIndex = 0; ViewIndex < Views.Num(); ViewIndex++)
+					{
+						FViewInfo& View = Views[ViewIndex];
 
-						// Apply temporal AA to the occlusion mask
-						// Result will be in HistoryOutput
-						ApplyTemporalAA(RHICmdList, View, TEXT("LSOcclusionHistory"), HistoryState, LightShafts0, HistoryOutput);
+						SCOPED_DRAW_EVENTF(RHICmdList, RenderLightShaftOcclusion, TEXT("RenderLightShaftOcclusion %dx%d (multiple passes)"), View.ViewRect.Width(), View.ViewRect.Height());
 
-						// Apply radial blur passes
-						// Send HistoryOutput in as the first pass input only, so it will not be overwritten by any subsequent passes, since it is needed for next frame
-						ApplyRadialBlurPasses(RHICmdList, View, LightSceneInfo, HistoryOutput, LightShafts0, LightShafts1);
+						if (ShouldRenderLightShaftsForLight(View, LightSceneInfo))
+						{
+							INC_DWORD_STAT(STAT_LightShaftsLights);
 
-						// Apply post-blur masking
-						FinishOcclusionTerm(RHICmdList, View, LightSceneInfo, LightShafts0, LightShafts1);
+							// Create a downsampled occlusion mask from scene depth, result will be in LightShafts0
+							DownsamplePass<true>(RHICmdList, View, LightSceneInfo, LightShafts0, LightShafts1);
 
-						//@todo - different views could have different result render targets
-						Output.LightShaftOcclusion = LightShafts1;
-					}
+							FSceneViewState* ViewState = (FSceneViewState*)View.State;
+							// Find the previous frame's occlusion mask
+							TRefCountPtr<IPooledRenderTarget>* HistoryState = ViewState ? &ViewState->LightShaftOcclusionHistoryRT : NULL;
+							TRefCountPtr<IPooledRenderTarget> HistoryOutput;
+
+							// Apply temporal AA to the occlusion mask
+							// Result will be in HistoryOutput
+							ApplyTemporalAA(RHICmdList, View, TEXT("LSOcclusionHistory"), HistoryState, LightShafts0, HistoryOutput);
+
+							// Apply radial blur passes
+							// Send HistoryOutput in as the first pass input only, so it will not be overwritten by any subsequent passes, since it is needed for next frame
+							ApplyRadialBlurPasses(RHICmdList, View, LightSceneInfo, HistoryOutput, LightShafts0, LightShafts1);
+
+							// Apply post-blur masking
+							FinishOcclusionTerm(RHICmdList, View, LightSceneInfo, LightShafts0, LightShafts1);
+
+							//@todo - different views could have different result render targets
+							Output.LightShaftOcclusion = LightShafts1;
+						}
+					}					
 				}
 			}
 		}
@@ -909,44 +924,59 @@ void FDeferredShadingSceneRenderer::RenderLightShaftBloom(FRHICommandListImmedia
 
 			if (LightSceneInfo->bEnableLightShaftBloom)
 			{
-				// Allocate light shaft render targets on demand, using the pool
-				AllocateOrReuseLightShaftRenderTarget(RHICmdList, LightShafts0, TEXT("LightShafts0"));
-				AllocateOrReuseLightShaftRenderTarget(RHICmdList, LightShafts1, TEXT("LightShafts1"));
+				bool bWillRenderLightShafts = false;
 
-				for (int ViewIndex = 0;ViewIndex < Views.Num();ViewIndex++)
+				for (int ViewIndex = 0; ViewIndex < Views.Num(); ViewIndex++)
 				{
 					FViewInfo& View = Views[ViewIndex];
 
-					SCOPED_DRAW_EVENTF(RHICmdList, RenderLightShaftBloom, TEXT("RenderLightShaftBloom %dx%d"), View.ViewRect.Width(), View.ViewRect.Height());
-
 					if (ShouldRenderLightShaftsForLight(View, LightSceneInfo))
 					{
-						INC_DWORD_STAT(STAT_LightShaftsLights);
+						bWillRenderLightShafts = true;
+					}
+				}
+
+				if (bWillRenderLightShafts)
+				{
+					// Allocate light shaft render targets on demand, using the pool
+					AllocateOrReuseLightShaftRenderTarget(RHICmdList, LightShafts0, TEXT("LightShafts0"));
+					AllocateOrReuseLightShaftRenderTarget(RHICmdList, LightShafts1, TEXT("LightShafts1"));
 
-						// Generate the bloom source from scene color, masked by depth and downsampled
-						DownsamplePass<false>(RHICmdList, View, LightSceneInfo, LightShafts0, LightShafts1);
+					for (int ViewIndex = 0; ViewIndex < Views.Num(); ViewIndex++)
+					{
+						FViewInfo& View = Views[ViewIndex];
 
-						FSceneViewState* ViewState = (FSceneViewState*)View.State;
-						TRefCountPtr<IPooledRenderTarget>* HistoryState = NULL;
+						SCOPED_DRAW_EVENTF(RHICmdList, RenderLightShaftBloom, TEXT("RenderLightShaftBloom %dx%d"), View.ViewRect.Width(), View.ViewRect.Height());
 
-						if (ViewState)
+						if (ShouldRenderLightShaftsForLight(View, LightSceneInfo))
 						{
-							// Find the previous frame's bloom source for this light
-							HistoryState = &ViewState->LightShaftBloomHistoryRTs.FindOrAdd(LightSceneInfo->Proxy->GetLightComponent());
-						}
+							INC_DWORD_STAT(STAT_LightShaftsLights);
 
-						TRefCountPtr<IPooledRenderTarget> HistoryOutput;
+							// Generate the bloom source from scene color, masked by depth and downsampled
+							DownsamplePass<false>(RHICmdList, View, LightSceneInfo, LightShafts0, LightShafts1);
 
-						// Apply temporal AA to the occlusion mask
-						// Result will be in HistoryOutput
-						ApplyTemporalAA(RHICmdList, View, TEXT("LSBloomHistory"), HistoryState, LightShafts0, HistoryOutput);
+							FSceneViewState* ViewState = (FSceneViewState*)View.State;
+							TRefCountPtr<IPooledRenderTarget>* HistoryState = NULL;
 
-						// Apply radial blur passes
-						// Send HistoryOutput in as the first pass input only, so it will not be overwritten by any subsequent passes, since it is needed for next frame
-						ApplyRadialBlurPasses(RHICmdList, View, LightSceneInfo, HistoryOutput, LightShafts0, LightShafts1);
-						
-						// Add light shaft bloom to scene color in full res
-						ApplyLightShaftBloom(RHICmdList, View, LightSceneInfo, LightShafts0);
+							if (ViewState)
+							{
+								// Find the previous frame's bloom source for this light
+								HistoryState = &ViewState->LightShaftBloomHistoryRTs.FindOrAdd(LightSceneInfo->Proxy->GetLightComponent());
+							}
+
+							TRefCountPtr<IPooledRenderTarget> HistoryOutput;
+
+							// Apply temporal AA to the occlusion mask
+							// Result will be in HistoryOutput
+							ApplyTemporalAA(RHICmdList, View, TEXT("LSBloomHistory"), HistoryState, LightShafts0, HistoryOutput);
+
+							// Apply radial blur passes
+							// Send HistoryOutput in as the first pass input only, so it will not be overwritten by any subsequent passes, since it is needed for next frame
+							ApplyRadialBlurPasses(RHICmdList, View, LightSceneInfo, HistoryOutput, LightShafts0, LightShafts1);
+
+							// Add light shaft bloom to scene color in full res
+							ApplyLightShaftBloom(RHICmdList, View, LightSceneInfo, LightShafts0);
+						}
 					}
 				}
 			}
diff --git a/Engine/Source/Runtime/Renderer/Private/SceneHitProxyRendering.cpp b/Engine/Source/Runtime/Renderer/Private/SceneHitProxyRendering.cpp
--- a/Engine/Source/Runtime/Renderer/Private/SceneHitProxyRendering.cpp
+++ b/Engine/Source/Runtime/Renderer/Private/SceneHitProxyRendering.cpp
@@ -461,6 +461,8 @@ void RenderHitProxies(FRHICommandListImmediate& RHICmdList, const FSceneRenderer
 
 	// Finish drawing to the hit proxy render target.
 	RHICmdList.CopyToResolveTarget(HitProxyRT->GetRenderTargetItem().TargetableTexture, HitProxyRT->GetRenderTargetItem().ShaderResourceTexture, false, FResolveParams());
+	RHICmdList.CopyToResolveTarget(SceneContext.GetSceneDepthSurface(), SceneContext.GetSceneDepthSurface(), true, FResolveParams());
+
 	// to be able to observe results with VisualizeTexture
 	GRenderTargetPool.VisualizeTexture.SetCheckPoint(RHICmdList, HitProxyRT);
 
