Fix potential crash from threaded function local static

diff --git a/Engine/Source/Runtime/Renderer/Private/PostProcess/SceneRenderTargets.cpp b/Engine/Source/Runtime/Renderer/Private/PostProcess/SceneRenderTargets.cpp
--- a/Engine/Source/Runtime/Renderer/Private/PostProcess/SceneRenderTargets.cpp
+++ b/Engine/Source/Runtime/Renderer/Private/PostProcess/SceneRenderTargets.cpp
@@ -1402,7 +1402,7 @@ const FTexture2DRHIRef& FSceneRenderTargets::GetEditorPrimitivesDepth(FRHIComman
 	return (const FTexture2DRHIRef&)EditorPrimitivesDepth->GetRenderTargetItem().TargetableTexture;
 }
 
-static TAutoConsoleVariable<int32> CVarSetSeperateTranslucencyEnabled(
+TAutoConsoleVariable<int32> FSceneRenderTargets::CVarSetSeperateTranslucencyEnabled(
 	TEXT("r.SeparateTranslucency"),
 	1,
 	TEXT("Allows to disable the separate translucency feature (all translucency is rendered in separate RT and composited\n")
@@ -1412,9 +1412,8 @@ static TAutoConsoleVariable<int32> CVarSetSeperateTranslucencyEnabled(
 	ECVF_RenderThreadSafe);
 
 bool FSceneRenderTargets::IsSeparateTranslucencyActive(const FViewInfo& View) const
-{
-	static const auto CVar = IConsoleManager::Get().FindTConsoleVariableDataInt(TEXT("r.SeparateTranslucency"));
-	int32 Value = CVar->GetValueOnRenderThread();
+{	
+	int32 Value = FSceneRenderTargets::CVarSetSeperateTranslucencyEnabled.GetValueOnRenderThread();
 
 	return (Value != 0) && CurrentFeatureLevel >= ERHIFeatureLevel::SM4
 		&& View.Family->EngineShowFlags.PostProcessing
diff --git a/Engine/Source/Runtime/Renderer/Private/PostProcess/SceneRenderTargets.h b/Engine/Source/Runtime/Renderer/Private/PostProcess/SceneRenderTargets.h
--- a/Engine/Source/Runtime/Renderer/Private/PostProcess/SceneRenderTargets.h
+++ b/Engine/Source/Runtime/Renderer/Private/PostProcess/SceneRenderTargets.h
@@ -154,6 +154,8 @@ public:
 	/** Destruct all snapshots */
 	void DestroyAllSnapshots();
 
+	static TAutoConsoleVariable<int32> CVarSetSeperateTranslucencyEnabled;
+
 protected:
 	/** Constructor */
 	FSceneRenderTargets(): 
diff --git a/Engine/Source/Runtime/Renderer/Private/TranslucentRendering.cpp b/Engine/Source/Runtime/Renderer/Private/TranslucentRendering.cpp
--- a/Engine/Source/Runtime/Renderer/Private/TranslucentRendering.cpp
+++ b/Engine/Source/Runtime/Renderer/Private/TranslucentRendering.cpp
@@ -685,9 +685,8 @@ void FTranslucentPrimSet::RenderPrimitive(
 		FTranslucencyDrawingPolicyFactory::ContextType Context(TranslucentSelfShadow, TranslucenyPassType);
 
 		// need to chec further down if we can skip rendering ST primitives, because we need to make sure they render in the normal translucency pass otherwise
-		// getting the cvar here and passing it down to be more efficient
-		static const auto CVar = IConsoleManager::Get().FindTConsoleVariableDataInt(TEXT("r.SeparateTranslucency"));
-		bool bSeparateTranslucencyPossible = (CVar->GetValueOnRenderThread() == 0 ? false : true) && (View.Family->EngineShowFlags.SeparateTranslucency && View.Family->EngineShowFlags.PostProcessing);
+		// getting the cvar here and passing it down to be more efficient		
+		bool bSeparateTranslucencyPossible = (FSceneRenderTargets::CVarSetSeperateTranslucencyEnabled.GetValueOnRenderThread() == 0 ? false : true) && (View.Family->EngineShowFlags.SeparateTranslucency && View.Family->EngineShowFlags.PostProcessing);
 
 
 		//@todo parallelrendering - come up with a better way to filter these by primitive
@@ -790,9 +789,8 @@ void FTranslucentPrimSet::AppendScenePrimitives(FSortedPrim* Normal, int32 NumNo
 }
 
 void FTranslucentPrimSet::PlaceScenePrimitive(FPrimitiveSceneInfo* PrimitiveSceneInfo, const FViewInfo& ViewInfo, bool bUseNormalTranslucency, bool bUseSeparateTranslucency, void *NormalPlace, int32& NormalNum, void* SeparatePlace, int32& SeparateNum)
-{
-	static const auto CVar = IConsoleManager::Get().FindTConsoleVariableDataInt(TEXT("r.SeparateTranslucency"));
-	int32 Value = CVar->GetValueOnRenderThread();
+{	
+	int32 Value = FSceneRenderTargets::CVarSetSeperateTranslucencyEnabled.GetValueOnRenderThread();
 
 	const float SortKey = CalculateTranslucentSortKey(PrimitiveSceneInfo, ViewInfo);
 
