UE4 - Fix speedtree crash on unitialized uniform buffer. This is a temp fix that just fills the uniform buffer with zeros.

diff --git a/Engine/Source/Runtime/RenderCore/Public/UniformBuffer.h b/Engine/Source/Runtime/RenderCore/Public/UniformBuffer.h
--- a/Engine/Source/Runtime/RenderCore/Public/UniformBuffer.h
+++ b/Engine/Source/Runtime/RenderCore/Public/UniformBuffer.h
@@ -46,6 +46,16 @@ public:
 		SetContentsNoUpdate(NewContents);
 		UpdateRHI();
 	}
+	/** Sets the contents of the uniform buffer to all zeros. */
+	void SetContentsToZero()
+	{
+		if (!Contents)
+		{
+			Contents = (uint8*)FMemory::Malloc(sizeof(TBufferStruct), UNIFORM_BUFFER_STRUCT_ALIGNMENT);
+		}
+		FMemory::Memzero(Contents, sizeof(TBufferStruct));
+		UpdateRHI();
+	}
 
 	// FRenderResource interface.
 	virtual void InitDynamicRHI() override
diff --git a/Engine/Source/Runtime/Renderer/Private/RendererScene.cpp b/Engine/Source/Runtime/Renderer/Private/RendererScene.cpp
--- a/Engine/Source/Runtime/Renderer/Private/RendererScene.cpp
+++ b/Engine/Source/Runtime/Renderer/Private/RendererScene.cpp
@@ -1572,6 +1572,7 @@ void FScene::AddSpeedTreeWind(FVertexFactory* VertexFactory, const UStaticMesh*
 				{
 					FSpeedTreeWindComputation* WindComputation = new FSpeedTreeWindComputation;
 					WindComputation->Wind = *(StaticMesh->SpeedTreeWind.Get( ));
+					WindComputation->UniformBuffer.SetContentsToZero();
 					WindComputation->UniformBuffer.InitResource();
 					Scene->SpeedTreeWindComputationMap.Add(StaticMesh, WindComputation);
 				}
