Fix Uniform Buffer assertion on startup on ES2 because GBlackVolumeTexture is unavailable on ES2.

diff --git a/Engine/Source/Runtime/Renderer/Private/LightMapRendering.cpp b/Engine/Source/Runtime/Renderer/Private/LightMapRendering.cpp
--- a/Engine/Source/Runtime/Renderer/Private/LightMapRendering.cpp
+++ b/Engine/Source/Runtime/Renderer/Private/LightMapRendering.cpp
@@ -131,6 +131,7 @@ void GetPrecomputedLightingParameters(
 			Parameters.IndirectLightingCacheTextureSampler2 = TStaticSamplerState<SF_Bilinear,AM_Clamp,AM_Clamp,AM_Clamp>::GetRHI();
 		}
 		else
+		if (FeatureLevel >= ERHIFeatureLevel::ES3_1)
 		{
 			Parameters.IndirectLightingCacheTexture0 = GBlackVolumeTexture->TextureRHI;
 			Parameters.IndirectLightingCacheTexture1 = GBlackVolumeTexture->TextureRHI;
@@ -140,6 +141,16 @@ void GetPrecomputedLightingParameters(
 			Parameters.IndirectLightingCacheTextureSampler1 = GBlackVolumeTexture->SamplerStateRHI;
 			Parameters.IndirectLightingCacheTextureSampler2 = GBlackVolumeTexture->SamplerStateRHI;
 		}
+		else
+		{
+			Parameters.IndirectLightingCacheTexture0 = GBlackTexture->TextureRHI;
+			Parameters.IndirectLightingCacheTexture1 = GBlackTexture->TextureRHI;
+			Parameters.IndirectLightingCacheTexture2 = GBlackTexture->TextureRHI;
+
+			Parameters.IndirectLightingCacheTextureSampler0 = GBlackTexture->SamplerStateRHI;
+			Parameters.IndirectLightingCacheTextureSampler1 = GBlackTexture->SamplerStateRHI;
+			Parameters.IndirectLightingCacheTextureSampler2 = GBlackTexture->SamplerStateRHI;
+		}
 	}
 
 	// TDistanceFieldShadowsAndLightMapPolicy
