Fix for Precomputed lighting renders gray on ES2

diff --git a/Engine/Source/Runtime/Renderer/Private/ForwardShadingRenderer.cpp b/Engine/Source/Runtime/Renderer/Private/ForwardShadingRenderer.cpp
--- a/Engine/Source/Runtime/Renderer/Private/ForwardShadingRenderer.cpp
+++ b/Engine/Source/Runtime/Renderer/Private/ForwardShadingRenderer.cpp
@@ -78,6 +78,9 @@ void FForwardShadingSceneRenderer::InitViews(FRHICommandListImmediate& RHICmdLis
 		Views[ViewIndex].InitRHIResources(DirectionalLightShadowInfo);
 	}
 
+	// Now that the indirect lighting cache is updated, we can update the primitive precomputed lighting buffers.
+	UpdatePrimitivePrecomputedLightingBuffers();
+	
 	OnStartFrame();
 }
 
