Copying //UE4/Dev-Sequencer to Dev-Main (//UE4/Dev-Main)

diff --git a/Engine/Source/Runtime/MovieSceneTracks/Public/MovieSceneHitProxy.h b/Engine/Source/Runtime/MovieSceneTracks/Public/MovieSceneHitProxy.h
new file mode 100644
--- /dev/null
+++ b/Engine/Source/Runtime/MovieSceneTracks/Public/MovieSceneHitProxy.h
@@ -0,0 +1,28 @@
+// Copyright 1998-2015 Epic Games, Inc. All Rights Reserved.
+
+#pragma once
+
+struct HMovieSceneKeyProxy : public HHitProxy
+{
+	DECLARE_HIT_PROXY( MOVIESCENETRACKS_API );
+
+	/** The track that contains the section */
+	class UMovieSceneTrack* MovieSceneTrack;
+	/** The section that contains the keyframe */
+	class UMovieSceneSection* MovieSceneSection;
+	/** The time of the key that is selected */
+	float Time;
+
+	HMovieSceneKeyProxy(class UMovieSceneTrack* InTrack, class UMovieSceneSection* InSection, float InTime) :
+		HHitProxy(HPP_UI),
+		MovieSceneTrack(InTrack),
+		MovieSceneSection(InSection),
+		Time(InTime)
+	{}
+
+	virtual EMouseCursor::Type GetMouseCursor() override
+	{
+		return EMouseCursor::Crosshairs;
+	}
+};
+
diff --git a/Engine/Source/Runtime/MovieSceneTracks/Public/Sections/MovieScene3DTransformSection.h b/Engine/Source/Runtime/MovieSceneTracks/Public/Sections/MovieScene3DTransformSection.h
--- a/Engine/Source/Runtime/MovieSceneTracks/Public/Sections/MovieScene3DTransformSection.h
+++ b/Engine/Source/Runtime/MovieSceneTracks/Public/Sections/MovieScene3DTransformSection.h
@@ -16,12 +16,6 @@ namespace EKey3DTransformChannel
 		Scale = 0x00000004,
 		All = Translation | Rotation | Scale
 	};
-
-	enum ValueType
-	{
-		Key,
-		Default
-	};
 }
 
 /**
@@ -61,16 +55,14 @@ struct FTransformData
 
 struct FTransformKey
 {
-	FTransformKey( EKey3DTransformChannel::Type InChannel, EKey3DTransformChannel::ValueType InChannelValueType, EAxis::Type InAxis, float InValue, bool InbUnwindRotation )
+	FTransformKey( EKey3DTransformChannel::Type InChannel, EAxis::Type InAxis, float InValue, bool InbUnwindRotation )
 	{
 		Channel = InChannel;
-		ChannelValueType = InChannelValueType;
 		Axis = InAxis;
 		Value = InValue;
 		bUnwindRotation = InbUnwindRotation;
 	}
 	EKey3DTransformChannel::Type Channel;
-	EKey3DTransformChannel::ValueType ChannelValueType;
 	EAxis::Type Axis;
 	float Value;
 	bool bUnwindRotation;
@@ -100,7 +92,7 @@ public:
 	 * @param Time				The position in time within the movie scene
 	 * @param OutTranslation	The evaluated translation.  Note: will remain unchanged if there were no keys to evaluate
 	 */
-	void EvalTranslation( float Time, FVector& OutTranslation ) const;
+	MOVIESCENETRACKS_API void EvalTranslation( float Time, FVector& OutTranslation ) const;
 
 	/**
 	 * Evaluates the rotation component of the transform
@@ -108,7 +100,7 @@ public:
 	 * @param Time				The position in time within the movie scene
 	 * @param OutRotation		The evaluated rotation.  Note: will remain unchanged if there were no keys to evaluate
 	 */
-	void EvalRotation( float Time, FRotator& OutRotation ) const;
+	MOVIESCENETRACKS_API void EvalRotation( float Time, FRotator& OutRotation ) const;
 
 	/**
 	 * Evaluates the scale component of the transform
@@ -116,7 +108,7 @@ public:
 	 * @param Time				The position in time within the movie scene
 	 * @param OutScale			The evaluated scale.  Note: will remain unchanged if there were no keys to evaluate
 	 */
-	void EvalScale( float Time, FVector& OutScale ) const;
+	MOVIESCENETRACKS_API void EvalScale( float Time, FVector& OutScale ) const;
 
 	/** 
 	 * Returns the translation curve for a specific axis
diff --git a/Engine/Source/Runtime/MovieSceneTracks/Public/Sections/MovieSceneVectorSection.h b/Engine/Source/Runtime/MovieSceneTracks/Public/Sections/MovieSceneVectorSection.h
--- a/Engine/Source/Runtime/MovieSceneTracks/Public/Sections/MovieSceneVectorSection.h
+++ b/Engine/Source/Runtime/MovieSceneTracks/Public/Sections/MovieSceneVectorSection.h
@@ -14,23 +14,16 @@ enum class EKeyVectorChannel
 	W
 };
 
-enum class EKeyVectorValueType
-{
-	Key,
-	Default
-};
 
 struct MOVIESCENETRACKS_API FVectorKey
 {
-	FVectorKey( EKeyVectorChannel InChannel, float InValue, EKeyVectorValueType InValueType )
+	FVectorKey( EKeyVectorChannel InChannel, float InValue )
 	{
 		Channel = InChannel;
 		Value = InValue;
-		ValueType = InValueType;
 	}
 	EKeyVectorChannel Channel;
 	float Value;
-	EKeyVectorValueType ValueType;
 };
 
 
diff --git a/Engine/Source/Runtime/MovieSceneTracks/Public/TrackInstances/MovieScene3DTransformTrackInstance.h b/Engine/Source/Runtime/MovieSceneTracks/Public/TrackInstances/MovieScene3DTransformTrackInstance.h
--- a/Engine/Source/Runtime/MovieSceneTracks/Public/TrackInstances/MovieScene3DTransformTrackInstance.h
+++ b/Engine/Source/Runtime/MovieSceneTracks/Public/TrackInstances/MovieScene3DTransformTrackInstance.h
@@ -23,14 +23,19 @@ public:
 	virtual void RestoreState (const TArray<UObject*>& RuntimeObjects, IMovieScenePlayer& Player, FMovieSceneSequenceInstance& SequenceInstance) override;
 	virtual void Update( float Position, float LastPosition, const TArray<UObject*>& RuntimeObjects, IMovieScenePlayer& Player, FMovieSceneSequenceInstance& SequenceInstance, EMovieSceneUpdatePass UpdatePass ) override;
 	virtual EMovieSceneUpdatePass HasUpdatePasses() override { return (EMovieSceneUpdatePass)(MSUP_PreUpdate | MSUP_Update); }
-	virtual void RefreshInstance( const TArray<UObject*>& RuntimeObjects, IMovieScenePlayer& Player, FMovieSceneSequenceInstance& SequenceInstance ) override {}
+	virtual void RefreshInstance( const TArray<UObject*>& RuntimeObjects, IMovieScenePlayer& Player, FMovieSceneSequenceInstance& SequenceInstance ) override;
 	virtual void ClearInstance( IMovieScenePlayer& Player, FMovieSceneSequenceInstance& SequenceInstance ) override {}
 
 private:
 
+	void UpdateRuntimeMobility(const TArray<UObject*>& RuntimeObjects);
+
 	/** Track that is being instanced */
 	UMovieScene3DTransformTrack* TransformTrack;
 
 	/** Map from object to initial state */
 	TMap< TWeakObjectPtr<UObject>, FTransform > InitTransformMap;
+
+	/** Map from object to initial mobility state */
+	TMap< TWeakObjectPtr<UObject>, EComponentMobility::Type > InitMobilityMap;
 };
diff --git a/Engine/Source/Runtime/MovieSceneTracks/Public/Tracks/MovieSceneSubTrack.h b/Engine/Source/Runtime/MovieSceneTracks/Public/Tracks/MovieSceneSubTrack.h
--- a/Engine/Source/Runtime/MovieSceneTracks/Public/Tracks/MovieSceneSubTrack.h
+++ b/Engine/Source/Runtime/MovieSceneTracks/Public/Tracks/MovieSceneSubTrack.h
@@ -53,6 +53,10 @@ public:
 	virtual void RemoveSection(UMovieSceneSection& Section) override;
 	virtual bool SupportsMultipleRows() const override;
 
+#if WITH_EDITORONLY_DATA
+	virtual FText GetDisplayName() const override;
+#endif
+
 private:
 
 	/** All movie scene sections. */
diff --git a/Engine/Source/Runtime/Renderer/Private/PostProcess/RenderingCompositionGraph.cpp b/Engine/Source/Runtime/Renderer/Private/PostProcess/RenderingCompositionGraph.cpp
--- a/Engine/Source/Runtime/Renderer/Private/PostProcess/RenderingCompositionGraph.cpp
+++ b/Engine/Source/Runtime/Renderer/Private/PostProcess/RenderingCompositionGraph.cpp
@@ -274,7 +274,63 @@ void FRenderingCompositionGraph::RecursivelyGatherDependencies(FRenderingComposi
 	}
 }
 
-void FRenderingCompositionGraph::DumpOutputToFile(FRenderingCompositePassContext& Context, const FString& Filename, FRenderingCompositeOutput* Output) const
+template<typename TColor> struct TAsyncBufferWrite;
+
+struct FAsyncBufferWriteQueue
+{
+	template<typename T>
+	static TFuture<void> Dispatch(TAsyncBufferWrite<T>&& In)
+	{
+		NumInProgressWrites.Increment();
+
+		while (NumInProgressWrites.GetValue() >= MaxAsyncWrites)
+		{
+			// Yield until we can write another
+			FPlatformProcess::Sleep(0.f);
+		}
+
+		return Async<void>(EAsyncExecution::ThreadPool, MoveTemp(In));
+	}
+	
+	static FThreadSafeCounter NumInProgressWrites;
+	static const int32 MaxAsyncWrites = 6;
+};
+FThreadSafeCounter FAsyncBufferWriteQueue::NumInProgressWrites;
+
+/** Callable type used to save a color buffer on an async task without allocating/copying into a new one */
+template<typename TColor>
+struct TAsyncBufferWrite
+{
+	TAsyncBufferWrite(FString InFilename, FIntPoint InDestSize, TArray<TColor> InBitmap) : Filename(MoveTemp(InFilename)), DestSize(InDestSize), Bitmap(MoveTemp(InBitmap)) {}
+
+	TAsyncBufferWrite(TAsyncBufferWrite&& In) : Filename(MoveTemp(In.Filename)), DestSize(In.DestSize), Bitmap(MoveTemp(In.Bitmap)) {}
+	TAsyncBufferWrite& operator=(TAsyncBufferWrite&& In) { Filename = MoveTemp(In.Filename); DestSize = In.DestSize; Bitmap = MoveTemp(In.Bitmap); return *this; }
+
+	/** Call operator that saves the color buffer data */
+	void operator()()
+	{
+		FString ResultPath;
+		GetHighResScreenshotConfig().SaveImage(Filename, Bitmap, DestSize, &ResultPath);
+		UE_LOG(LogConsoleResponse, Display, TEXT("Content was saved to \"%s\""), *ResultPath);
+
+		FAsyncBufferWriteQueue::NumInProgressWrites.Decrement();
+	}
+
+	/** Copy semantics are only defined to appease TFunction, whose virtual CloneByCopy function is always instantiated, even if the TFunction itself is never copied */
+	TAsyncBufferWrite(const TAsyncBufferWrite& In) : Filename(In.Filename), DestSize(In.DestSize), Bitmap(In.Bitmap) { ensureMsgf(false, TEXT("Type should not be copied")); }
+	TAsyncBufferWrite& operator=(const TAsyncBufferWrite& In) { Filename = In.Filename; DestSize = In.DestSize; Bitmap = In.Bitmap; ensureMsgf(false, TEXT("Type should not be copied")); return *this; }
+
+private:
+
+	/** The filename to save to */
+	FString Filename;
+	/** The size of the bitmap */
+	FIntPoint DestSize;
+	/** The bitmap data itself */
+	TArray<TColor> Bitmap;
+};
+
+TFuture<void> FRenderingCompositionGraph::DumpOutputToFile(FRenderingCompositePassContext& Context, const FString& Filename, FRenderingCompositeOutput* Output) const
 {
 	FSceneRenderTargetItem& RenderTargetItem = Output->PooledRenderTarget->GetRenderTargetItem();
 	FHighResScreenshotConfig& HighResScreenshotConfig = GetHighResScreenshotConfig();
@@ -297,16 +353,17 @@ void FRenderingCompositionGraph::DumpOutputToFile(FRenderingCompositePassContext
 	FIntPoint DestSize(SourceRect.Width(), SourceRect.Height());
 
 	EPixelFormat PixelFormat = Texture->GetFormat();
-	FString ResultPath;
+	
 	switch (PixelFormat)
 	{
 		case PF_FloatRGBA:
 		{
 			TArray<FFloat16Color> Bitmap;
 			Context.RHICmdList.ReadSurfaceFloatData(Texture, SourceRect, Bitmap, (ECubeFace)0, 0, 0);
-			HighResScreenshotConfig.SaveImage(Filename, Bitmap, DestSize, &ResultPath);
+
+			return FAsyncBufferWriteQueue::Dispatch(TAsyncBufferWrite<FFloat16Color>(Filename, DestSize, MoveTemp(Bitmap)));
 		}
-		break;
+
 		case PF_R8G8B8A8:
 		case PF_B8G8R8A8:
 		{
@@ -319,12 +376,12 @@ void FRenderingCompositionGraph::DumpOutputToFile(FRenderingCompositePassContext
 			{
 				Pixel->A = 255;
 			}
-			HighResScreenshotConfig.SaveImage(Filename, Bitmap, DestSize, &ResultPath);
+
+			return FAsyncBufferWriteQueue::Dispatch(TAsyncBufferWrite<FColor>(Filename, DestSize, MoveTemp(Bitmap)));
 		}
-		break;
 	}
 
-	UE_LOG(LogConsoleResponse, Display, TEXT("Content was saved to \"%s\""), *ResultPath);
+	return TFuture<void>();
 }
 
 void FRenderingCompositionGraph::RecursivelyProcess(const FRenderingCompositeOutputRef& InOutputRef, FRenderingCompositePassContext& Context) const
diff --git a/Engine/Source/Runtime/Renderer/Private/PostProcess/RenderingCompositionGraph.h b/Engine/Source/Runtime/Renderer/Private/PostProcess/RenderingCompositionGraph.h
--- a/Engine/Source/Runtime/Renderer/Private/PostProcess/RenderingCompositionGraph.h
+++ b/Engine/Source/Runtime/Renderer/Private/PostProcess/RenderingCompositionGraph.h
@@ -61,7 +61,7 @@ private:
 	void RecursivelyProcess(const FRenderingCompositeOutputRef& InOutputRef, FRenderingCompositePassContext& Context) const;
 
 	/** Write the contents of the specified output to a file */
-	void DumpOutputToFile(FRenderingCompositePassContext& Context, const FString& Filename, FRenderingCompositeOutput* Output) const;
+	TFuture<void> DumpOutputToFile(FRenderingCompositePassContext& Context, const FString& Filename, FRenderingCompositeOutput* Output) const;
 
 	/**
 	 * for debugging purpose O(n)
diff --git a/Engine/Source/Runtime/Slate/Private/Widgets/Input/SEditableLabel.cpp b/Engine/Source/Runtime/Slate/Private/Widgets/Input/SEditableLabel.cpp
--- a/Engine/Source/Runtime/Slate/Private/Widgets/Input/SEditableLabel.cpp
+++ b/Engine/Source/Runtime/Slate/Private/Widgets/Input/SEditableLabel.cpp
@@ -4,6 +4,9 @@
 #include "SEditableLabel.h"
 
 
+#define LOCTEXT_NAMESPACE "SEditableLabel"
+
+
 /* SEditableLabel interface
  *****************************************************************************/
 
@@ -12,7 +15,6 @@ void SEditableLabel::Construct(const FArguments& InArgs)
 	CanEditAttribute = InArgs._CanEdit;
 	OnTextChanged = InArgs._OnTextChanged;
 	TextAttribute = InArgs._Text;
-	WasFocused = false;
 
 	ChildSlot
 	[
@@ -29,6 +31,7 @@ void SEditableLabel::Construct(const FArguments& InArgs)
 					.HighlightShape(InArgs._HighlightShape)
 					.HighlightText(InArgs._HighlightText)
 					.MinDesiredWidth(InArgs._MinDesiredWidth)
+					.OnDoubleClicked(this, &SEditableLabel::HandleTextBlockDoubleClicked)
 					.ShadowColorAndOpacity(InArgs._ShadowColorAndOpacity)
 					.ShadowOffset(InArgs._ShadowOffset)
 					.TextStyle(InArgs._TextStyle)
@@ -61,7 +64,8 @@ void SEditableLabel::Construct(const FArguments& InArgs)
 			.VAlign(VAlign_Center)
 			[
 				SNew(SImage)
-					.Image(FCoreStyle::Get().GetBrush(TEXT("EditableLabel.EditIcon")))
+					.Image(FCoreStyle::Get().GetBrush(TEXT("Icons.Rename")))
+					.ToolTipText(LOCTEXT("RenameToolTip", "Press F2 or double-click the text label to rename it"))
 					.Visibility(this, &SEditableLabel::HandleIconVisibility)
 			]
 	];
@@ -77,6 +81,7 @@ void SEditableLabel::EnterTextMode()
 
 	TextBlock->SetVisibility(EVisibility::Collapsed);
 	EditableText->SetVisibility(EVisibility::Visible);
+	FSlateApplication::Get().SetAllUserFocus(EditableText);
 }
 	
 
@@ -84,6 +89,7 @@ void SEditableLabel::ExitTextMode()
 {
 	TextBlock->SetVisibility(EVisibility::Visible);
 	EditableText->SetVisibility(EVisibility::Collapsed);
+	FSlateApplication::Get().SetAllUserFocus(AsShared());
 }
 
 
@@ -108,7 +114,7 @@ FReply SEditableLabel::OnKeyDown(const FGeometry& MyGeometry, const FKeyEvent& I
 		return FReply::Handled().SetUserFocus(AsShared(), EFocusCause::Navigation);
 	}
 
-	if (WasFocused && (Key == EKeys::Enter))
+	if ((Key == EKeys::F2) && CanEditAttribute.Get())
 	{
 		EnterTextMode();
 		return FReply::Handled().SetUserFocus(EditableText.ToSharedRef(), EFocusCause::Navigation);
@@ -118,34 +124,6 @@ FReply SEditableLabel::OnKeyDown(const FGeometry& MyGeometry, const FKeyEvent& I
 }
 
 
-FReply SEditableLabel::OnMouseButtonDown(const FGeometry& MyGeometry, const FPointerEvent& MouseEvent)
-{
-	if (MouseEvent.GetEffectingButton() == EKeys::LeftMouseButton)
-	{
-		WasFocused = HasKeyboardFocus();
-	}
-
-	if (WasFocused)
-	{
-		return FReply::Handled();
-	}
-	
-	return FReply::Unhandled();
-}
-
-
-FReply SEditableLabel::OnMouseButtonUp(const FGeometry& MyGeometry, const FPointerEvent& MouseEvent )
-{
-	if (WasFocused && (MouseEvent.GetEffectingButton() == EKeys::LeftMouseButton))
-	{
-		EnterTextMode();
-		return FReply::Handled().SetUserFocus(EditableText.ToSharedRef(), EFocusCause::Mouse);;
-	}
-
-	return FReply::Unhandled();
-}
-
-
 bool SEditableLabel::SupportsKeyboardFocus() const
 {
 	return true;
@@ -162,9 +140,19 @@ void SEditableLabel::HandleEditableTextTextCommitted(const FText& NewText, EText
 }
 
 
+FReply SEditableLabel::HandleTextBlockDoubleClicked()
+{
+	EnterTextMode();
+	return FReply::Handled().SetUserFocus(EditableText.ToSharedRef(), EFocusCause::Navigation);
+}
+
+
 EVisibility SEditableLabel::HandleIconVisibility() const
 {
 	return (IsHovered() && !EditableText->HasKeyboardFocus() && CanEditAttribute.Get())
 		? EVisibility::Visible
 		: EVisibility::Collapsed;
 }
+
+
+#undef LOCTEXT_NAMESPACE
diff --git a/Engine/Source/Runtime/Slate/Private/Widgets/Text/STextBlock.cpp b/Engine/Source/Runtime/Slate/Private/Widgets/Text/STextBlock.cpp
--- a/Engine/Source/Runtime/Slate/Private/Widgets/Text/STextBlock.cpp
+++ b/Engine/Source/Runtime/Slate/Private/Widgets/Text/STextBlock.cpp
@@ -196,9 +196,7 @@ FReply STextBlock::OnMouseButtonDoubleClick( const FGeometry& InMyGeometry, cons
 	{
 		if( OnDoubleClicked.IsBound() )
 		{
-			OnDoubleClicked.Execute();
-
-			return FReply::Handled();
+			return OnDoubleClicked.Execute();
 		}
 	}
 
diff --git a/Engine/Source/Runtime/Slate/Public/Widgets/Input/SEditableLabel.h b/Engine/Source/Runtime/Slate/Public/Widgets/Input/SEditableLabel.h
--- a/Engine/Source/Runtime/Slate/Public/Widgets/Input/SEditableLabel.h
+++ b/Engine/Source/Runtime/Slate/Public/Widgets/Input/SEditableLabel.h
@@ -84,13 +84,12 @@ public:
 
 	virtual bool HasKeyboardFocus() const override;
 	virtual FReply OnKeyDown(const FGeometry& MyGeometry, const FKeyEvent& InKeyEvent) override;
-	virtual FReply OnMouseButtonDown(const FGeometry& MyGeometry, const FPointerEvent& MouseEvent) override;
-	virtual FReply OnMouseButtonUp(const FGeometry& MyGeometry, const FPointerEvent& MouseEvent) override;
 	virtual bool SupportsKeyboardFocus() const override;
 
 private:
 
 	void HandleEditableTextTextCommitted(const FText& NewText, ETextCommit::Type CommitInfo);
+	FReply HandleTextBlockDoubleClicked();
 	EVisibility HandleIconVisibility() const;
 
 private:
@@ -98,7 +97,6 @@ private:
 	TAttribute<bool> CanEditAttribute;
 	FOnTextChanged OnTextChanged;
 	TAttribute<FText> TextAttribute;
-	bool WasFocused;
 
 	TSharedPtr<STextBlock> TextBlock;
 	TSharedPtr<SEditableText> EditableText;
diff --git a/Engine/Source/Runtime/SlateCore/Private/Styling/CoreStyle.cpp b/Engine/Source/Runtime/SlateCore/Private/Styling/CoreStyle.cpp
--- a/Engine/Source/Runtime/SlateCore/Private/Styling/CoreStyle.cpp
+++ b/Engine/Source/Runtime/SlateCore/Private/Styling/CoreStyle.cpp
@@ -269,6 +269,7 @@ TSharedRef<ISlateStyle> FCoreStyle::Create( const FName& InStyleSetName )
 		Style->Set( "Icons.Info", new IMAGE_BRUSH( "Icons/icon_info_16x", Icon16x16) );
 		Style->Set( "Icons.Warning", new IMAGE_BRUSH( "Icons/icon_warning_16x", Icon16x16) );
 		Style->Set( "Icons.Download", new IMAGE_BRUSH( "Icons/icon_Downloads_16x", Icon16x16) );
+		Style->Set( "Icons.Rename", new IMAGE_BRUSH( "Icons/icon_rename_12x", Icon12x12) );
 	}
 
 	// Tool panels
@@ -414,11 +415,6 @@ TSharedRef<ISlateStyle> FCoreStyle::Create( const FName& InStyleSetName )
 		Style->Set( "RadioButton", BasicRadioButtonStyle );
 	}
 
-	// SEditableLabel defaults...
-	{
-		Style->Set("EditableLabel.EditIcon", new IMAGE_BRUSH("Icons/ellipsis_12x", Icon12x12));
-	}
-
 	// SEditableText defaults...
 	{
 		FSlateBrush* SelectionBackground = new BOX_BRUSH( "Common/EditableTextSelectionBackground",  FMargin(4.f/16.f) );
diff --git a/Engine/Source/Runtime/UMG/Private/Components/ListView.cpp b/Engine/Source/Runtime/UMG/Private/Components/ListView.cpp
--- a/Engine/Source/Runtime/UMG/Private/Components/ListView.cpp
+++ b/Engine/Source/Runtime/UMG/Private/Components/ListView.cpp
@@ -73,6 +73,13 @@ const FText UListView::GetPaletteCategory()
 
 #endif
 
+void UListView::ReleaseSlateResources(bool bReleaseChildren)
+{
+	Super::ReleaseSlateResources(bReleaseChildren);
+
+	MyListView.Reset();
+}
+
 /////////////////////////////////////////////////////
 
 #undef LOCTEXT_NAMESPACE
diff --git a/Engine/Source/Runtime/UMG/Private/Components/Slider.cpp b/Engine/Source/Runtime/UMG/Private/Components/Slider.cpp
--- a/Engine/Source/Runtime/UMG/Private/Components/Slider.cpp
+++ b/Engine/Source/Runtime/UMG/Private/Components/Slider.cpp
@@ -43,6 +43,13 @@ void USlider::SynchronizeProperties()
 	MySlider->SetIndentHandle(IndentHandle);
 }
 
+void USlider::ReleaseSlateResources(bool bReleaseChildren)
+{
+	Super::ReleaseSlateResources(bReleaseChildren);
+
+	MySlider.Reset();
+}
+
 void USlider::HandleOnValueChanged(float InValue)
 {
 	OnValueChanged.Broadcast(InValue);
diff --git a/Engine/Source/Runtime/UMG/Private/Components/TileView.cpp b/Engine/Source/Runtime/UMG/Private/Components/TileView.cpp
--- a/Engine/Source/Runtime/UMG/Private/Components/TileView.cpp
+++ b/Engine/Source/Runtime/UMG/Private/Components/TileView.cpp
@@ -76,6 +76,13 @@ void UTileView::RequestListRefresh()
 	MyTileView->RequestListRefresh();
 }
 
+void UTileView::ReleaseSlateResources(bool bReleaseChildren)
+{
+	Super::ReleaseSlateResources(bReleaseChildren);
+
+	MyTileView.Reset();
+}
+
 #if WITH_EDITOR
 
 const FSlateBrush* UTileView::GetEditorIcon()
diff --git a/Engine/Source/Runtime/UMG/Public/Animation/MovieScene2DTransformSection.h b/Engine/Source/Runtime/UMG/Public/Animation/MovieScene2DTransformSection.h
--- a/Engine/Source/Runtime/UMG/Public/Animation/MovieScene2DTransformSection.h
+++ b/Engine/Source/Runtime/UMG/Public/Animation/MovieScene2DTransformSection.h
@@ -23,25 +23,17 @@ enum class EKey2DTransformAxis
 	None
 };
 
-enum class EKey2DTransformValueType
-{
-	Key,
-	Default
-};
-
 struct F2DTransformKey
 {
-	F2DTransformKey( EKey2DTransformChannel InChannel, EKey2DTransformAxis InAxis, float InValue, EKey2DTransformValueType InValueType )
+	F2DTransformKey( EKey2DTransformChannel InChannel, EKey2DTransformAxis InAxis, float InValue )
 	{
 		Channel = InChannel;
 		Axis = InAxis;
 		Value = InValue;
-		ValueType = InValueType;
 	}
 	EKey2DTransformChannel Channel;
 	EKey2DTransformAxis Axis;
 	float Value;
-	EKey2DTransformValueType ValueType;	
 };
 
 
diff --git a/Engine/Source/Runtime/UMG/Public/Animation/MovieSceneMarginSection.h b/Engine/Source/Runtime/UMG/Public/Animation/MovieSceneMarginSection.h
--- a/Engine/Source/Runtime/UMG/Public/Animation/MovieSceneMarginSection.h
+++ b/Engine/Source/Runtime/UMG/Public/Animation/MovieSceneMarginSection.h
@@ -15,24 +15,16 @@ enum class EKeyMarginChannel
 	Bottom
 };
 
-enum class EKeyMarginValueType
-{
-	Key,
-	Default
-};
-
 
 struct FMarginKey
 {
-	FMarginKey( EKeyMarginChannel InChannel, float InValue, EKeyMarginValueType InValueType )
+	FMarginKey( EKeyMarginChannel InChannel, float InValue )
 	{
 		Channel = InChannel;
 		Value = InValue;
-		ValueType = InValueType;
 	}
 	EKeyMarginChannel Channel;
 	float Value;
-	EKeyMarginValueType ValueType;
 };
 
 
diff --git a/Engine/Source/Runtime/UMG/Public/Components/ListView.h b/Engine/Source/Runtime/UMG/Public/Components/ListView.h
--- a/Engine/Source/Runtime/UMG/Public/Components/ListView.h
+++ b/Engine/Source/Runtime/UMG/Public/Components/ListView.h
@@ -35,6 +35,10 @@ public:
 	virtual const FText GetPaletteCategory() override;
 #endif
 
+	//~ Begin UVisual Interface
+	virtual void ReleaseSlateResources(bool bReleaseChildren) override;
+	//~ End UVisual Interface
+
 protected:
 	TSharedPtr< SListView<UObject*> > MyListView;
 
diff --git a/Engine/Source/Runtime/UMG/Public/Components/Slider.h b/Engine/Source/Runtime/UMG/Public/Components/Slider.h
--- a/Engine/Source/Runtime/UMG/Public/Components/Slider.h
+++ b/Engine/Source/Runtime/UMG/Public/Components/Slider.h
@@ -87,6 +87,10 @@ public:
 	virtual void SynchronizeProperties() override;
 	// End of UWidget interface
 
+	// UVisual interface
+	virtual void ReleaseSlateResources(bool bReleaseChildren) override;
+	// End of UVisual interface
+
 #if WITH_EDITOR
 	virtual const FSlateBrush* GetEditorIcon() override;
 	virtual const FText GetPaletteCategory() override;
diff --git a/Engine/Source/Runtime/UMG/Public/Components/TileView.h b/Engine/Source/Runtime/UMG/Public/Components/TileView.h
--- a/Engine/Source/Runtime/UMG/Public/Components/TileView.h
+++ b/Engine/Source/Runtime/UMG/Public/Components/TileView.h
@@ -48,6 +48,10 @@ public:
 	UFUNCTION(BlueprintCallable, Category="Behavior")
 	void RequestListRefresh();
 
+	// UVisual interface
+	virtual void ReleaseSlateResources(bool bReleaseChildren) override;
+	// End of UVisual interface
+
 #if WITH_EDITOR
 	virtual const FSlateBrush* GetEditorIcon() override;
 	virtual const FText GetPaletteCategory() override;
