Fixed duplicate variable "r.Streaming.UseAABB"
Fixed shader compilation for material with tessellation

diff --git a/Engine/Shaders/DebugViewModeVertexShader.usf b/Engine/Shaders/DebugViewModeVertexShader.usf
--- a/Engine/Shaders/DebugViewModeVertexShader.usf
+++ b/Engine/Shaders/DebugViewModeVertexShader.usf
@@ -49,9 +49,9 @@
 		}
 
 #if NUM_MATERIAL_TEXCOORDS > 0
-	Output.TexCoords = VertexParameters.TexCoords[0];
+		O.TexCoords = MaterialParameters.TexCoords[0];
 #else
-	Output.TexCoords = 0;
+		O.TexCoords = 0;
 #endif
 
 		
diff --git a/Engine/Source/Runtime/Engine/Private/ContentStreaming.cpp b/Engine/Source/Runtime/Engine/Private/ContentStreaming.cpp
--- a/Engine/Source/Runtime/Engine/Private/ContentStreaming.cpp
+++ b/Engine/Source/Runtime/Engine/Private/ContentStreaming.cpp
@@ -156,7 +156,7 @@ static TAutoConsoleVariable<int32> CVarStreamingShowWantedMips(
 	TEXT("If non-zero, will limit resolution to wanted mip."),
 	ECVF_Cheat);
 
-static TAutoConsoleVariable<int32> CVarStreamingUseAABB(
+ENGINE_API TAutoConsoleVariable<int32> CVarStreamingUseAABB(
 	TEXT("r.Streaming.UseAABB"),
 	0,
 	TEXT("If non-zero, will use AABB to compute distance."),
diff --git a/Engine/Source/Runtime/Renderer/Private/WantedMipsAccuracyRendering.cpp b/Engine/Source/Runtime/Renderer/Private/WantedMipsAccuracyRendering.cpp
--- a/Engine/Source/Runtime/Renderer/Private/WantedMipsAccuracyRendering.cpp
+++ b/Engine/Source/Runtime/Renderer/Private/WantedMipsAccuracyRendering.cpp
@@ -7,12 +7,7 @@ WantedMipsAccuracyRendering.cpp: Contains definitions for rendering the viewmode
 #include "RendererPrivate.h"
 #include "ScenePrivate.h"
 
-static TAutoConsoleVariable<int32> CVarStreamingUseAABBRenderer(
-	TEXT("r.Streaming.UseAABB"),
-	0,
-	TEXT("If non-zero, will use AABB to compute distance."),
-	ECVF_Default);
-
+extern ENGINE_API TAutoConsoleVariable<int32> CVarStreamingUseAABB;
 
 IMPLEMENT_SHADER_TYPE(,FWantedMipsAccuracyPS,TEXT("WantedMipsAccuracyPixelShader"),TEXT("Main"),SF_Pixel);
 
@@ -54,7 +49,7 @@ void FWantedMipsAccuracyPS::SetMesh(
 		float DistSqMinusRadiusSq = 0;
 		float ViewSize = View.ViewRect.Size().X;
 
-		if (CVarStreamingUseAABBRenderer.GetValueOnRenderThread())
+		if (CVarStreamingUseAABB.GetValueOnRenderThread())
 		{
 			ViewToObject = ViewToObject.GetAbs();
 			FVector BoxViewToObject = ViewToObject.ComponentMin(Proxy->GetBounds().BoxExtent);
