Copying //UE4/Dev-Rendering to Dev-Main (//UE4/Dev-Main)

diff --git a/Engine/Source/Runtime/Engine/Private/PhysicsEngine/PhysSubstepTasks.cpp b/Engine/Source/Runtime/Engine/Private/PhysicsEngine/PhysSubstepTasks.cpp
--- a/Engine/Source/Runtime/Engine/Private/PhysicsEngine/PhysSubstepTasks.cpp
+++ b/Engine/Source/Runtime/Engine/Private/PhysicsEngine/PhysSubstepTasks.cpp
@@ -348,6 +348,14 @@ void FPhysSubstepTask::StepSimulation(PhysXCompletionTask * Task)
 DECLARE_CYCLE_STAT(TEXT("Phys SubstepStart"), STAT_SubstepSimulationStart, STATGROUP_Physics);
 DECLARE_CYCLE_STAT(TEXT("Phys SubstepEnd"), STAT_SubstepSimulationEnd, STATGROUP_Physics);
 
+FAutoConsoleTaskPriority CPrio_PhyXSubstepSimulationEnd(
+	TEXT("TaskGraph.TaskPriorities.PhyXSubstepSimulationEnd"),
+	TEXT("Task and thread priority for FPhysSubstepTask::SubstepSimulationEnd."),
+	ENamedThreads::HighThreadPriority, // if we have high priority task threads, then use them...
+	ENamedThreads::NormalTaskPriority, // .. at normal task priority
+	ENamedThreads::HighTaskPriority // if we don't have hi pri threads, then use normal priority threads at high task priority instead
+	);
+
 void FPhysSubstepTask::SubstepSimulationStart()
 {
 	SCOPE_CYCLE_COUNTER(STAT_TotalPhysicsTime);
@@ -361,7 +369,7 @@ void FPhysSubstepTask::SubstepSimulationStart()
 	PhysXCompletionTask* SubstepTask = new PhysXCompletionTask(CompletionEvent,
 		 PST_MAX //we don't care about sub-step time. The full time is recorded by FullSimulationTask
 		,PAScene->getTaskManager());
-	ENamedThreads::Type NamedThread = PhysSingleThreadedMode() ? ENamedThreads::GameThread : ENamedThreads::HiPri(ENamedThreads::AnyThread);
+	ENamedThreads::Type NamedThread = PhysSingleThreadedMode() ? ENamedThreads::GameThread : CPrio_PhyXSubstepSimulationEnd.Get();
 
 	DECLARE_CYCLE_STAT(TEXT("FDelegateGraphTask.ProcessPhysSubstepSimulation"),
 		STAT_FDelegateGraphTask_ProcessPhysSubstepSimulation,
@@ -369,7 +377,7 @@ void FPhysSubstepTask::SubstepSimulationStart()
 
 	FDelegateGraphTask::CreateAndDispatchWhenReady(
 		FDelegateGraphTask::FDelegate::CreateRaw(this, &FPhysSubstepTask::SubstepSimulationEnd),
-		GET_STATID(STAT_FDelegateGraphTask_ProcessPhysSubstepSimulation), CompletionEvent, NamedThread, NamedThread);
+		GET_STATID(STAT_FDelegateGraphTask_ProcessPhysSubstepSimulation), CompletionEvent, ENamedThreads::AnyThread, NamedThread);
 
 	++CurrentSubStep;	
 
diff --git a/Engine/Source/Runtime/Engine/Private/PrimitiveSceneProxy.cpp b/Engine/Source/Runtime/Engine/Private/PrimitiveSceneProxy.cpp
--- a/Engine/Source/Runtime/Engine/Private/PrimitiveSceneProxy.cpp
+++ b/Engine/Source/Runtime/Engine/Private/PrimitiveSceneProxy.cpp
@@ -202,8 +202,8 @@ void FPrimitiveSceneProxy::UpdateActorPosition(FVector InActorPosition)
 
 static TAutoConsoleVariable<int32> CVarDeferUniformBufferUpdatesUntilVisible(
 	TEXT("r.DeferUniformBufferUpdatesUntilVisible"),
-	0,
-	TEXT("If > 0, then don't update the primitive uniform buffer until it is visible. Experimental option."));
+	!WITH_EDITOR,
+	TEXT("If > 0, then don't update the primitive uniform buffer until it is visible."));
 
 void FPrimitiveSceneProxy::UpdateUniformBufferMaybeLazy()
 {
diff --git a/Engine/Source/Runtime/Engine/Private/SceneView.cpp b/Engine/Source/Runtime/Engine/Private/SceneView.cpp
--- a/Engine/Source/Runtime/Engine/Private/SceneView.cpp
+++ b/Engine/Source/Runtime/Engine/Private/SceneView.cpp
@@ -173,6 +173,13 @@ static TAutoConsoleVariable<int32> CVarDefaultAutoExposure(
 	TEXT(" 0: off, sets AutoExposureMinBrightness and AutoExposureMaxBrightness to 1\n")
 	TEXT(" 1: on (default)"));
 
+static TAutoConsoleVariable<int32> CVarDefaultAutoExposureMethod(
+	TEXT("r.DefaultFeature.AutoExposure.Method"),
+	0,
+	TEXT("Engine default (project setting) for AutoExposure Method (postprocess volume/camera/game setting still can override)\n")
+	TEXT(" 0: Histogram based (requires compute shader, default)\n")
+	TEXT(" 1: Basic AutoExposure"));
+
 static TAutoConsoleVariable<int32> CVarDefaultMotionBlur(
 	TEXT("r.DefaultFeature.MotionBlur"),
 	1,
@@ -1241,6 +1248,15 @@ void FSceneView::StartFinalPostprocessSettings(FVector InViewLocation)
 			FinalPostProcessSettings.AutoExposureMinBrightness = 1;
 			FinalPostProcessSettings.AutoExposureMaxBrightness = 1;
 		}
+		else 
+		{
+			int32 Value = CVarDefaultAutoExposureMethod.GetValueOnGameThread();
+			if (Value >= 0 && Value < AEM_MAX)
+			{
+				FinalPostProcessSettings.AutoExposureMethod = (EAutoExposureMethod)Value;
+			}
+		}
+
 		if (!CVarDefaultMotionBlur.GetValueOnGameThread())
 		{
 			FinalPostProcessSettings.MotionBlurAmount = 0;
@@ -1289,6 +1305,20 @@ void FSceneView::EndFinalPostprocessSettings(const FSceneViewInitOptions& ViewIn
 {
 	const auto SceneViewFeatureLevel = GetFeatureLevel();
 
+#if !(UE_BUILD_SHIPPING || UE_BUILD_TEST)
+	static const auto CVar = IConsoleManager::Get().FindTConsoleVariableDataInt(TEXT("r.EyeAdaptation.MethodOveride"));
+	if (CVar->GetValueOnGameThread() == -2)
+	{
+		// seemed to be good setting for Paragon, we might want to remove or adjust this later on
+		FinalPostProcessSettings.AutoExposureMethod = AEM_Basic;
+		FinalPostProcessSettings.AutoExposureBias = -0.6f;
+		FinalPostProcessSettings.AutoExposureMaxBrightness = 2.f;
+		FinalPostProcessSettings.AutoExposureMinBrightness = 0.05;
+		FinalPostProcessSettings.AutoExposureSpeedDown = 1.f;
+		FinalPostProcessSettings.AutoExposureSpeedUp = 3.f;
+	}
+#endif
+
 	// will be deprecated soon, use the new asset LightPropagationVolumeBlendable instead
 	{
 		FLightPropagationVolumeSettings& Dest = FinalPostProcessSettings.BlendableManager.GetSingleFinalData<FLightPropagationVolumeSettings>();
diff --git a/Engine/Source/Runtime/Engine/Private/SkeletalMeshComponentPhysics.cpp b/Engine/Source/Runtime/Engine/Private/SkeletalMeshComponentPhysics.cpp
--- a/Engine/Source/Runtime/Engine/Private/SkeletalMeshComponentPhysics.cpp
+++ b/Engine/Source/Runtime/Engine/Private/SkeletalMeshComponentPhysics.cpp
@@ -3622,6 +3622,13 @@ void USkeletalMeshComponent::UpdateClothMorphTarget()
 	}
 }
 
+FAutoConsoleTaskPriority CPrio_FParallelClothTask(
+	TEXT("TaskGraph.TaskPriorities.ParallelClothTask"),
+	TEXT("Task and thread priority for parallel cloth."),
+	ENamedThreads::NormalThreadPriority,
+	ENamedThreads::NormalTaskPriority
+	);
+
 class FParallelClothTask
 {
 	const USkeletalMeshComponent& SkeletalMeshComponent;
@@ -3642,7 +3649,7 @@ public:
 	}
 	static ENamedThreads::Type GetDesiredThread()
 	{
-		return ENamedThreads::AnyThread;
+		return CPrio_FParallelClothTask.Get();
 	}
 	static ESubsequentsMode::Type GetSubsequentsMode()
 	{
diff --git a/Engine/Source/Runtime/Engine/Private/SkeletalRenderGPUSkin.cpp b/Engine/Source/Runtime/Engine/Private/SkeletalRenderGPUSkin.cpp
--- a/Engine/Source/Runtime/Engine/Private/SkeletalRenderGPUSkin.cpp
+++ b/Engine/Source/Runtime/Engine/Private/SkeletalRenderGPUSkin.cpp
@@ -320,7 +320,7 @@ void FSkeletalMeshObjectGPUSkin::ProcessUpdatedDynamicData(FRHICommandListImmedi
 
 			bool bClothFactory = (FeatureLevel >= ERHIFeatureLevel::SM4) && (DynamicData->ClothSimulUpdateData.Num() > 0) && Chunk.HasApexClothData();
 
-			FGPUBaseSkinVertexFactory::ShaderDataType& ShaderData = 
+			FGPUBaseSkinVertexFactory::FShaderDataType& ShaderData = 
 				bClothFactory ? VertexFactoryData.ClothVertexFactories[ChunkIdx]->GetVertexFactory()->GetShaderData() :
 				(DynamicData->NumWeightedActiveVertexAnims > 0 ? VertexFactoryData.MorphVertexFactories[ChunkIdx]->GetShaderData() : VertexFactoryData.VertexFactories[ChunkIdx]->GetShaderData());
 
@@ -392,11 +392,12 @@ void FSkeletalMeshObjectGPUSkin::ProcessUpdatedDynamicData(FRHICommandListImmedi
 	}
 }
 
+TArray<FVector> FSkeletalMeshObjectGPUSkin::FSkeletalMeshObjectLOD::MorphDeltaTangentZAccumulationArray;
+TArray<float> FSkeletalMeshObjectGPUSkin::FSkeletalMeshObjectLOD::MorphAccumulatedWeightArray;
+
 void FSkeletalMeshObjectGPUSkin::FSkeletalMeshObjectLOD::UpdateMorphVertexBuffer(const TArray<FActiveVertexAnim>& ActiveVertexAnims)
 {
 	SCOPE_CYCLE_COUNTER(STAT_MorphVertexBuffer_Update);
-	static TArray<FVector> DeltaTangentZAccumulationArray;
-	static TArray<float> AccumulatedWeightArray;
 
 	if( IsValidRef(MorphVertexBuffer.VertexBufferRHI) )
 	{
@@ -410,16 +411,16 @@ void FSkeletalMeshObjectGPUSkin::FSkeletalMeshObjectLOD::UpdateMorphVertexBuffer
 			SCOPE_CYCLE_COUNTER(STAT_MorphVertexBuffer_Init);
 
 			// zero everything
-			int32 vertsToAdd = static_cast<int32>(LodModel.NumVertices) - DeltaTangentZAccumulationArray.Num();
+			int32 vertsToAdd = static_cast<int32>(LodModel.NumVertices) - MorphDeltaTangentZAccumulationArray.Num();
 			if(vertsToAdd > 0) 
 			{
 				// we're memzero-ing afterwards anyway, so add uninitalized
-				DeltaTangentZAccumulationArray.AddUninitialized(vertsToAdd);
-				AccumulatedWeightArray.AddUninitialized(vertsToAdd);
+				MorphDeltaTangentZAccumulationArray.AddUninitialized(vertsToAdd);
+				MorphAccumulatedWeightArray.AddUninitialized(vertsToAdd);
 			}
 
-			FMemory::Memzero(DeltaTangentZAccumulationArray.GetData(), sizeof(FVector)*LodModel.NumVertices);
-			FMemory::Memzero(AccumulatedWeightArray.GetData(), sizeof(float)*LodModel.NumVertices);
+			FMemory::Memzero(MorphDeltaTangentZAccumulationArray.GetData(), sizeof(FVector)*LodModel.NumVertices);
+			FMemory::Memzero(MorphAccumulatedWeightArray.GetData(), sizeof(float)*LodModel.NumVertices);
 
 			// PackedNormals will be wrong init with 0, but they'll be overwritten later
 			FMemory::Memzero(&Buffer[0], sizeof(FMorphGPUSkinVertex)*LodModel.NumVertices);
@@ -457,9 +458,9 @@ void FSkeletalMeshObjectGPUSkin::FSkeletalMeshObjectLOD::UpdateMorphVertexBuffer
 
 						DestVertex.DeltaPosition += MorphVertex.PositionDelta * VertAnim.Weight;
 						// add to accumulated tangent Z
-						DeltaTangentZAccumulationArray[MorphVertex.SourceIdx] += MorphVertex.TangentZDelta * VertAnim.Weight;
+						MorphDeltaTangentZAccumulationArray[MorphVertex.SourceIdx] += MorphVertex.TangentZDelta * VertAnim.Weight;
 						// accumulate the weight so we can normalized it later
-						AccumulatedWeightArray[MorphVertex.SourceIdx] += VertAnimAbsWeight;
+						MorphAccumulatedWeightArray[MorphVertex.SourceIdx] += VertAnimAbsWeight;
 					}
 				} // for all vertices
 
@@ -470,14 +471,14 @@ void FSkeletalMeshObjectGPUSkin::FSkeletalMeshObjectLOD::UpdateMorphVertexBuffer
 			for(uint32 iVertex = 0; iVertex < LodModel.NumVertices; ++iVertex) 
 			{
 				FMorphGPUSkinVertex& DestVertex = Buffer[iVertex];
-				float AccumulatedWeight = AccumulatedWeightArray[iVertex];
+				float AccumulatedWeight = MorphAccumulatedWeightArray[iVertex];
 
 				if (AccumulatedWeight > MinVertexAnimBlendWeight)
 				{
 					// when copy back, make sure to normalize by accumulated weight
 					// since delta diff of normal is (-2, 2), we divide by 2 to packed into packed normal
 					// when we unpack, we'll apply *2. 
-					DestVertex.DeltaTangentZ = (DeltaTangentZAccumulationArray[iVertex] / AccumulatedWeight)/2;
+					DestVertex.DeltaTangentZ = (MorphDeltaTangentZAccumulationArray[iVertex] / AccumulatedWeight)/2;
 					// we now add W as how much alpha of DeltaTangentZ we're apply to the original tangent
 					DestVertex.DeltaTangentZ.Vector.W = FMath::Min(1.0f, AccumulatedWeight) * 255.9999f;
 				}
@@ -545,7 +546,7 @@ const FVertexFactory* FSkeletalMeshObjectGPUSkin::GetVertexFactory(int32 LODInde
  * @param bUseInstancedVertexWeights - use instanced influence weights instead of default weights
  */
 template<class VertexFactoryType>
-void InitGPUSkinVertexFactoryComponents(typename VertexFactoryType::DataType* VertexFactoryData, 
+void InitGPUSkinVertexFactoryComponents(typename VertexFactoryType::FDataType* VertexFactoryData, 
 										const FSkeletalMeshObjectGPUSkin::FVertexFactoryBuffers& VertexBuffers)
 {
 	typedef TGPUSkinVertexBase<VertexFactoryType::HasExtraBoneInfluences> BaseVertexType;
@@ -617,7 +618,7 @@ void InitGPUSkinVertexFactoryComponents(typename VertexFactoryType::DataType* Ve
  * @param bUseInstancedVertexWeights - use instanced influence weights instead of default weights
  */
 template<class VertexFactoryType>
-void InitMorphVertexFactoryComponents(typename VertexFactoryType::DataType* VertexFactoryData, 
+void InitMorphVertexFactoryComponents(typename VertexFactoryType::FDataType* VertexFactoryData, 
 										const FSkeletalMeshObjectGPUSkin::FVertexFactoryBuffers& VertexBuffers)
 {
 	// delta positions
@@ -636,7 +637,7 @@ void InitMorphVertexFactoryComponents(typename VertexFactoryType::DataType* Vert
  * @param bUseInstancedVertexWeights - use instanced influence weights instead of default weights
  */
 template<class VertexFactoryType>
-void InitAPEXClothVertexFactoryComponents(typename VertexFactoryType::DataType* VertexFactoryData, 
+void InitAPEXClothVertexFactoryComponents(typename VertexFactoryType::FDataType* VertexFactoryData, 
 										const FSkeletalMeshObjectGPUSkin::FVertexFactoryBuffers& VertexBuffers)
 {
 	// barycentric coord for positions
@@ -675,7 +676,7 @@ ENQUEUE_UNIQUE_RENDER_COMMAND_ONEPARAMETER_DECLARE_TEMPLATE(
 InitGPUSkinVertexFactory,VertexFactoryType,
 TDynamicUpdateVertexFactoryData<VertexFactoryType>,VertexUpdateData,VertexUpdateData,
 {
-	typename VertexFactoryType::DataType Data;
+	typename VertexFactoryType::FDataType Data;
 	InitGPUSkinVertexFactoryComponents<VertexFactoryType>(&Data,VertexUpdateData.VertexBuffers);
 	VertexUpdateData.VertexFactory->SetData(Data);
 	VertexUpdateData.VertexFactory->GetShaderData().MeshOrigin = VertexUpdateData.VertexBuffers.VertexBufferGPUSkin->GetMeshOrigin();
@@ -737,7 +738,7 @@ ENQUEUE_UNIQUE_RENDER_COMMAND_ONEPARAMETER_DECLARE_TEMPLATE(
 InitGPUSkinVertexFactoryMorph,VertexFactoryType,
 TDynamicUpdateVertexFactoryData<VertexFactoryType>,VertexUpdateData,VertexUpdateData,
 {
-	typename VertexFactoryType::DataType Data;
+	typename VertexFactoryType::FDataType Data;
 	InitGPUSkinVertexFactoryComponents<VertexFactoryType>(&Data,VertexUpdateData.VertexBuffers);
 	InitMorphVertexFactoryComponents<VertexFactoryType>(&Data,VertexUpdateData.VertexBuffers);
 	VertexUpdateData.VertexFactory->SetData(Data);
@@ -749,7 +750,7 @@ TDynamicUpdateVertexFactoryData<VertexFactoryType>,VertexUpdateData,VertexUpdate
  * Creates a vertex factory entry for the given type and initialize it on the render thread
  */
 template <class VertexFactoryTypeBase, class VertexFactoryType>
-static void CreateVertexFactoryMorph(TArray<TUniquePtr<VertexFactoryTypeBase>>& VertexFactories,
+static VertexFactoryType* CreateVertexFactoryMorph(TArray<TUniquePtr<VertexFactoryTypeBase>>& VertexFactories,
 						 const FSkeletalMeshObjectGPUSkin::FVertexFactoryBuffers& InVertexBuffers,
 						 ERHIFeatureLevel::Type FeatureLevel
 						 )
@@ -768,6 +769,8 @@ static void CreateVertexFactoryMorph(TArray<TUniquePtr<VertexFactoryTypeBase>>&
 
 	// init rendering resource	
 	BeginInitResource(VertexFactory);
+
+	return VertexFactory;
 }
 
 // APEX cloth
@@ -775,7 +778,7 @@ ENQUEUE_UNIQUE_RENDER_COMMAND_ONEPARAMETER_DECLARE_TEMPLATE(
 InitGPUSkinAPEXClothVertexFactory,VertexFactoryType,
 TDynamicUpdateVertexFactoryData<VertexFactoryType>,VertexUpdateData,VertexUpdateData,
 {
-	typename VertexFactoryType::DataType Data;
+	typename VertexFactoryType::FDataType Data;
 	InitGPUSkinVertexFactoryComponents<VertexFactoryType>(&Data,VertexUpdateData.VertexBuffers);
 	InitAPEXClothVertexFactoryComponents<VertexFactoryType>(&Data,VertexUpdateData.VertexBuffers);
 	VertexUpdateData.VertexFactory->SetData(Data);
@@ -911,7 +914,6 @@ void FSkeletalMeshObjectGPUSkin::FVertexFactoryData::InitAPEXClothVertexFactorie
 	const TArray<FSkelMeshChunk>& Chunks,
 	ERHIFeatureLevel::Type InFeatureLevel)
 {
-
 	// clear existing factories (resources assumed to have been released already)
 	ClothVertexFactories.Empty(Chunks.Num());
 	for( int32 FactoryIdx=0; FactoryIdx < Chunks.Num(); FactoryIdx++ )
@@ -1053,7 +1055,7 @@ const FTwoVectors& FSkeletalMeshObjectGPUSkin::GetCustomLeftRightVectors(int32 S
 FDynamicSkelMeshObjectDataGPUSkin
 -----------------------------------------------------------------------------*/
 
-static TLockFreePointerListUnordered<FDynamicSkelMeshObjectDataGPUSkin>	FreeDynamicSkelMeshObjectDataGPUSkins;
+static TLockFreePointerListUnordered<FDynamicSkelMeshObjectDataGPUSkin, PLATFORM_CACHE_LINE_SIZE>	FreeDynamicSkelMeshObjectDataGPUSkins;
 
 void FDynamicSkelMeshObjectDataGPUSkin::Clear()
 {
diff --git a/Engine/Source/Runtime/Engine/Private/SkeletalRenderGPUSkin.h b/Engine/Source/Runtime/Engine/Private/SkeletalRenderGPUSkin.h
--- a/Engine/Source/Runtime/Engine/Private/SkeletalRenderGPUSkin.h
+++ b/Engine/Source/Runtime/Engine/Private/SkeletalRenderGPUSkin.h
@@ -422,6 +422,10 @@ private:
 		 * @param OutVertexBuffers output vertex buffers
 		 */
 		void GetVertexBuffers(FVertexFactoryBuffers& OutVertexBuffers,FStaticLODModel& LODModel,const FSkelMeshObjectLODInfo& MeshLODInfo);
+
+		// Temporary arrays used on UpdateMorphVertexBuffer(); these grow to the max and are not thread safe.
+		static TArray<FVector> MorphDeltaTangentZAccumulationArray;
+		static TArray<float> MorphAccumulatedWeightArray;
 	};
 
 	/** 
diff --git a/Engine/Source/Runtime/Engine/Private/SplineMeshSceneProxy.h b/Engine/Source/Runtime/Engine/Private/SplineMeshSceneProxy.h
--- a/Engine/Source/Runtime/Engine/Private/SplineMeshSceneProxy.h
+++ b/Engine/Source/Runtime/Engine/Private/SplineMeshSceneProxy.h
@@ -42,7 +42,7 @@ public:
 		ENQUEUE_UNIQUE_RENDER_COMMAND_TWOPARAMETER(
 			FSplineMeshVertexFactoryCopyData,
 			FSplineMeshVertexFactory*, VertexFactory, this,
-			const DataType*, DataCopy, &Other.Data,
+			const FDataType*, DataCopy, &Other.Data,
 			{
 			VertexFactory->Data = *DataCopy;
 		});
diff --git a/Engine/Source/Runtime/Engine/Private/StaticMeshRender.cpp b/Engine/Source/Runtime/Engine/Private/StaticMeshRender.cpp
--- a/Engine/Source/Runtime/Engine/Private/StaticMeshRender.cpp
+++ b/Engine/Source/Runtime/Engine/Private/StaticMeshRender.cpp
@@ -762,6 +762,8 @@ void FStaticMeshSceneProxy::GetDynamicMeshElements(const TArray<const FSceneView
 		{
 			if (VisibilityMap & (1 << ViewIndex))
 			{
+				FFrozenSceneViewMatricesGuard FrozenMatricesGuard(*const_cast<FSceneView*>(Views[ViewIndex]));
+
 				const FSceneView* View = Views[ViewIndex];
 				FLODMask LODMask = GetLODMask(View);
 
diff --git a/Engine/Source/Runtime/Engine/Private/SystemSettings.cpp b/Engine/Source/Runtime/Engine/Private/SystemSettings.cpp
--- a/Engine/Source/Runtime/Engine/Private/SystemSettings.cpp
+++ b/Engine/Source/Runtime/Engine/Private/SystemSettings.cpp
@@ -153,7 +153,11 @@ void FSystemSettings::RegisterShowFlagConsoleVariables()
 				ShowFlagIndex,
 				(uint8*)&This.Force0Mask,
 				(uint8*)&This.Force1Mask,
-				TEXT("ShowFlag 0:force off, 1:force on, 2 or other value to not apply any adjustments"),
+				TEXT("Allows to override a specific showflag (works in editor and game, \"show\" only works in game and UI only in editor)\n")
+				TEXT("Useful to run a build many time with the same showflags (when put in consolevariables.ini like \"showflag.abc=0\")\n")
+				TEXT(" 0: force the showflag to be OFF\n")
+				TEXT(" 1: force the showflag to be ON\n")
+				TEXT(" 2: do not override this showflag (default)"),
 				ECVF_Cheat);
 
 			return true;
diff --git a/Engine/Source/Runtime/Engine/Private/TextureRenderTarget.cpp b/Engine/Source/Runtime/Engine/Private/TextureRenderTarget.cpp
--- a/Engine/Source/Runtime/Engine/Private/TextureRenderTarget.cpp
+++ b/Engine/Source/Runtime/Engine/Private/TextureRenderTarget.cpp
@@ -26,7 +26,8 @@ UTextureRenderTarget::UTextureRenderTarget(const FObjectInitializer& ObjectIniti
 
 FTextureRenderTargetResource* UTextureRenderTarget::GetRenderTargetResource()
 {
-	check( IsInRenderingThread() );
+	check(IsInRenderingThread() || 
+		(IsInParallelRenderingThread() && (!Resource || Resource->IsInitialized()))); // we allow this in parallel, but only if the resource is initialized...otherwise it might be a race on intialization
 	FTextureRenderTargetResource* Result = NULL;
 	if( Resource &&
 		Resource->IsInitialized() )
diff --git a/Engine/Source/Runtime/Engine/Private/TickTaskManager.cpp b/Engine/Source/Runtime/Engine/Private/TickTaskManager.cpp
--- a/Engine/Source/Runtime/Engine/Private/TickTaskManager.cpp
+++ b/Engine/Source/Runtime/Engine/Private/TickTaskManager.cpp
@@ -56,6 +56,38 @@ static TAutoConsoleVariable<int32> CVarAllowAsyncTickDispatch(
 	0,
 	TEXT("If true, ticks are dispatched in a task thread."));
 
+FAutoConsoleTaskPriority CPrio_DispatchTaskPriority(
+	TEXT("TaskGraph.TaskPriorities.TickDispatchTaskPriority"),
+	TEXT("Task and thread priority for tick tasks dispatch."),
+	ENamedThreads::HighThreadPriority, // if we have high priority task threads, then use them...
+	ENamedThreads::NormalTaskPriority, // .. at normal task priority
+	ENamedThreads::HighTaskPriority // if we don't have hi pri threads, then use normal priority threads at high task priority instead
+	);
+
+FAutoConsoleTaskPriority CPrio_CleanupTaskPriority(
+	TEXT("TaskGraph.TaskPriorities.TickCleanupTaskPriority"),
+	TEXT("Task and thread priority for tick cleanup."),
+	ENamedThreads::BackgroundThreadPriority, // if we have background priority task threads, then use them...
+	ENamedThreads::NormalTaskPriority, // .. at normal task priority
+	ENamedThreads::NormalTaskPriority // if we don't have background threads, then use normal priority threads at normal task priority instead
+	);
+
+FAutoConsoleTaskPriority CPrio_NormalAsyncTickTaskPriority(
+	TEXT("TaskGraph.TaskPriorities.NormalAsyncTickTaskPriority"),
+	TEXT("Task and thread priority for async ticks that are not high priority."),
+	ENamedThreads::NormalThreadPriority, 
+	ENamedThreads::NormalTaskPriority
+	);
+
+FAutoConsoleTaskPriority CPrio_HiPriAsyncTickTaskPriority(
+	TEXT("TaskGraph.TaskPriorities.HiPriAsyncTickTaskPriority"),
+	TEXT("Task and thread priority for async ticks that are high priority."),
+	ENamedThreads::HighThreadPriority, // if we have high priority task threads, then use them...
+	ENamedThreads::NormalTaskPriority, // .. at normal task priority
+	ENamedThreads::HighTaskPriority // if we don't have hi pri threads, then use normal priority threads at high task priority instead
+	);
+
+
 FORCEINLINE bool CanDemoteIntoTickGroup(ETickingGroup TickGroup)
 {
 	switch (TickGroup)
@@ -192,20 +224,16 @@ public:
 	, bLogTicksShowPrerequistes(bInLogTicksShowPrerequistes)
 	{
 	}
-	FORCEINLINE TStatId GetStatId() const
+	static FORCEINLINE TStatId GetStatId()
 	{
-	RETURN_QUICK_DECLARE_CYCLE_STAT(FTickFunctionTask, STATGROUP_TaskGraphTasks);
+		RETURN_QUICK_DECLARE_CYCLE_STAT(FTickFunctionTask, STATGROUP_TaskGraphTasks);
 	}
 	/** return the thread for this task **/
 	FORCEINLINE ENamedThreads::Type GetDesiredThread()
 	{
-		if (Target->bHighPriority)
-		{
-			return ENamedThreads::HiPri(Context.Thread);
-		}
 		return Context.Thread;
 	}
-	FORCEINLINE static ESubsequentsMode::Type GetSubsequentsMode() 
+	static FORCEINLINE ESubsequentsMode::Type GetSubsequentsMode()
 	{ 
 		return ESubsequentsMode::TrackSubsequents; 
 	}
@@ -255,15 +283,15 @@ class FTickTaskSequencer
 			, WorldTickGroup(InWorldTickGroup)
 		{
 		}
-		FORCEINLINE TStatId GetStatId() const
+		static FORCEINLINE TStatId GetStatId()
 		{
 			RETURN_QUICK_DECLARE_CYCLE_STAT(FDipatchTickGroupTask, STATGROUP_TaskGraphTasks);
 		}
-		FORCEINLINE ENamedThreads::Type GetDesiredThread()
+		static FORCEINLINE ENamedThreads::Type GetDesiredThread()
 		{
-			return ENamedThreads::HiPri(ENamedThreads::AnyThread);
+			return CPrio_DispatchTaskPriority.Get();
 		}
-		FORCEINLINE static ESubsequentsMode::Type GetSubsequentsMode() 
+		static FORCEINLINE ESubsequentsMode::Type GetSubsequentsMode()
 		{ 
 			return ESubsequentsMode::TrackSubsequents; 
 		}
@@ -291,13 +319,13 @@ class FTickTaskSequencer
 			, WorldTickGroup(InWorldTickGroup)
 		{
 		}
-		FORCEINLINE TStatId GetStatId() const
+		static FORCEINLINE TStatId GetStatId()
 		{
 			RETURN_QUICK_DECLARE_CYCLE_STAT(FResetTickGroupTask, STATGROUP_TaskGraphTasks);
 		}
-		FORCEINLINE ENamedThreads::Type GetDesiredThread()
+		static FORCEINLINE ENamedThreads::Type GetDesiredThread()
 		{
-			return ENamedThreads::AnyThread;
+			return CPrio_CleanupTaskPriority.Get();
 		}
 		FORCEINLINE static ESubsequentsMode::Type GetSubsequentsMode() 
 		{ 
@@ -369,10 +397,20 @@ public:
 	   
 		bool bIsOriginalTickGroup = (TickFunction->ActualStartTickGroup == TickFunction->TickGroup);
 
-		UseContext.Thread = ENamedThreads::GameThread;
 		if (TickFunction->bRunOnAnyThread && bAllowConcurrentTicks && bIsOriginalTickGroup)
 		{
-			UseContext.Thread = ENamedThreads::HiPri(ENamedThreads::AnyThread);
+			if (TickFunction->bHighPriority)
+			{
+				UseContext.Thread = CPrio_HiPriAsyncTickTaskPriority.Get();
+			}
+			else
+			{
+				UseContext.Thread = CPrio_NormalAsyncTickTaskPriority.Get();
+			}
+		}
+		else
+		{
+			UseContext.Thread = ENamedThreads::SetTaskPriority(ENamedThreads::GameThread, TickFunction->bHighPriority ? ENamedThreads::HighTaskPriority : ENamedThreads::NormalTaskPriority);
 		}
 
 		TickFunction->TaskPointer = TGraphTask<FTickFunctionTask>::CreateTask(Prerequisites, TickContext.Thread).ConstructAndHold(TickFunction, &UseContext, bLogTicks, bLogTicksShowPrerequistes);
diff --git a/Engine/Source/Runtime/Engine/Public/GPUSkinCache.h b/Engine/Source/Runtime/Engine/Public/GPUSkinCache.h
--- a/Engine/Source/Runtime/Engine/Public/GPUSkinCache.h
+++ b/Engine/Source/Runtime/Engine/Public/GPUSkinCache.h
@@ -12,12 +12,16 @@
 #include "RenderResource.h"
 #include "UniformBuffer.h"
 
+
+extern ENGINE_API int32 GEnableGPUSkinCacheShaders;
+extern ENGINE_API int32 GEnableGPUSkinCache;
+
 class FSkeletalMeshObjectGPUSkin;
 
 class FGPUSkinCache : public FRenderResource
 {
 public:
-	enum SkinCacheInitSettings
+	enum ESkinCacheInitSettings
 	{
 		// max 256 bones as we use a byte to index
 		MaxUniformBufferBones = 256,
@@ -42,7 +46,15 @@ public:
 	void Initialize(FRHICommandListImmediate& RHICmdList);
 	void Cleanup();
 
-	bool IsElementProcessed(int32 Key) const;
+	inline bool IsElementProcessed(int32 Key) const
+	{
+		if (!GEnableGPUSkinCache)
+		{
+			return false;
+		}
+
+		return InternalIsElementProcessed(Key);
+	}
 
 	// For each SkeletalMeshObject:
 	//	Call Begin*
@@ -130,6 +142,7 @@ private:
 	bool	bInitialized : 1;
 
 	void TransitionAllToWriteable(FRHICommandList& RHICmdList);
+	bool InternalIsElementProcessed(int32 Key) const;
 
 	uint32	FrameCounter;
 
@@ -152,8 +165,6 @@ BEGIN_UNIFORM_BUFFER_STRUCT(GPUSkinCacheBonesUniformShaderParameters,)
 	DECLARE_UNIFORM_BUFFER_STRUCT_MEMBER_ARRAY(FVector4,BoneVectors,[FGPUSkinCache::MaxUniformBufferBones*3])
 END_UNIFORM_BUFFER_STRUCT(GPUSkinCacheBonesUniformShaderParameters)
 
-extern ENGINE_API int32 GEnableGPUSkinCache;
-extern ENGINE_API int32 GEnableGPUSkinCacheShaders;
 extern ENGINE_API TGlobalResource<FGPUSkinCache> GGPUSkinCache;
 
 DECLARE_STATS_GROUP(TEXT("GPU Skin Cache"), STATGROUP_GPUSkinCache, STATCAT_Advanced);
diff --git a/Engine/Source/Runtime/Engine/Public/GPUSkinVertexFactory.h b/Engine/Source/Runtime/Engine/Public/GPUSkinVertexFactory.h
--- a/Engine/Source/Runtime/Engine/Public/GPUSkinVertexFactory.h
+++ b/Engine/Source/Runtime/Engine/Public/GPUSkinVertexFactory.h
@@ -22,8 +22,7 @@
 
 =============================================================================*/
 
-#ifndef __GPUSKINVERTEXFACTORY_H__
-#define __GPUSKINVERTEXFACTORY_H__
+#pragma once
 
 #include "GPUSkinPublicDefs.h"
 #include "ResourcePool.h"
@@ -278,9 +277,9 @@ private: // -------------------------------------------------------
 class FGPUBaseSkinVertexFactory : public FVertexFactory
 {
 public:
-	struct ShaderDataType
+	struct FShaderDataType
 	{
-		ShaderDataType()
+		FShaderDataType()
 			: CurrentBuffer(0)
 			, PreviousFrameNumber(0)
 			, CurrentFrameNumber(0)
@@ -351,7 +350,7 @@ public:
 		// @return IsValid() can fail, then you have to create the buffers first (or if the size changes)
 		FVertexBufferAndSRV& GetBoneBufferForWriting(bool bPrevious, uint32 FrameNumber)
 		{
-			const ShaderDataType* This = (const ShaderDataType*)this;
+			const FShaderDataType* This = (const FShaderDataType*)this;
 
 			// non const version maps to const version
 			return (FVertexBufferAndSRV&)This->GetBoneBufferInternal(bPrevious, FrameNumber);
@@ -398,12 +397,12 @@ public:
 	{}
 
 	/** accessor */
-	FORCEINLINE ShaderDataType& GetShaderData()
+	FORCEINLINE FShaderDataType& GetShaderData()
 	{
 		return ShaderData;
 	}
 
-	FORCEINLINE const ShaderDataType& GetShaderData() const
+	FORCEINLINE const FShaderDataType& GetShaderData() const
 	{
 		return ShaderData;
 	}
@@ -419,7 +418,7 @@ public:
 
 protected:
 	/** dynamic data need for setting the shader */ 
-	ShaderDataType ShaderData;
+	FShaderDataType ShaderData;
 	/** Pool of buffers for bone matrices. */
 	static TGlobalResource<FBoneBufferPool> BoneBufferPool;
 };
@@ -437,7 +436,7 @@ public:
 		HasExtraBoneInfluences = bExtraBoneInfluencesT,
 	};
 
-	struct DataType
+	struct FDataType
 	{
 		/** The stream to read the vertex position from. */
 		FVertexStreamComponent PositionComponent;
@@ -486,7 +485,7 @@ public:
 	* update the resource with new data from the game thread.
 	* @param	InData - new stream component data
 	*/
-	void SetData(const DataType& InData)
+	void SetData(const FDataType& InData)
 	{
 		Data = InData;
 		FGPUBaseSkinVertexFactory::UpdateRHI();
@@ -507,13 +506,13 @@ protected:
 	* @param InData - type with stream components
 	* @param OutElements - vertex decl list to modify
 	*/
-	void AddVertexElements(DataType& InData, FVertexDeclarationElementList& OutElements);
+	void AddVertexElements(FDataType& InData, FVertexDeclarationElementList& OutElements);
 
-	inline const DataType& GetData() const { return Data; }
+	inline const FDataType& GetData() const { return Data; }
 
 private:
 	/** stream component data bound to this vertex factory */
-	DataType Data;  
+	FDataType Data;  
 };
 
 /** 
@@ -567,7 +566,7 @@ class TGPUSkinMorphVertexFactory : public TGPUSkinVertexFactory<bExtraBoneInflue
 	typedef TGPUSkinVertexFactory<bExtraBoneInfluencesT> Super;
 public:
 
-	struct DataType : TGPUSkinVertexFactory<bExtraBoneInfluencesT>::DataType
+	struct FDataType : TGPUSkinVertexFactory<bExtraBoneInfluencesT>::FDataType
 	{
 		/** stream which has the position deltas to add to the vertex position */
 		FVertexStreamComponent DeltaPositionComponent;
@@ -592,7 +591,7 @@ public:
 	* update the resource with new data from the game thread.
 	* @param	InData - new stream component data
 	*/
-	void SetData(const DataType& InData)
+	void SetData(const FDataType& InData)
 	{
 		MorphData = InData;
 		FGPUBaseSkinVertexFactory::UpdateRHI();
@@ -614,11 +613,11 @@ protected:
 	* @param InData - type with stream components
 	* @param OutElements - vertex decl list to modify
 	*/
-	void AddVertexElements(DataType& InData, FVertexDeclarationElementList& OutElements);
+	void AddVertexElements(FDataType& InData, FVertexDeclarationElementList& OutElements);
 
 private:
 	/** stream component data bound to this vertex factory */
-	DataType MorphData; 
+	FDataType MorphData; 
 
 };
 
@@ -804,7 +803,7 @@ class TGPUSkinAPEXClothVertexFactory : public FGPUBaseSkinAPEXClothVertexFactory
 
 public:
 
-	struct DataType : TGPUSkinVertexFactory<bExtraBoneInfluencesT>::DataType
+	struct FDataType : TGPUSkinVertexFactory<bExtraBoneInfluencesT>::FDataType
 	{
 		/** stream which has the physical mesh position + height offset */
 		FVertexStreamComponent CoordPositionComponent;
@@ -833,7 +832,7 @@ public:
 	* update the resource with new data from the game thread.
 	* @param	InData - new stream component data
 	*/
-	void SetData(const DataType& InData)
+	void SetData(const FDataType& InData)
 	{
 		MeshMappingData = InData;
 		FGPUBaseSkinVertexFactory::UpdateRHI();
@@ -866,11 +865,10 @@ protected:
 	* @param InData - type with stream components
 	* @param OutElements - vertex decl list to modify
 	*/
-	void AddVertexElements(DataType& InData, FVertexDeclarationElementList& OutElements);
+	void AddVertexElements(FDataType& InData, FVertexDeclarationElementList& OutElements);
 
 private:
 	/** stream component data bound to this vertex factory */
-	DataType MeshMappingData; 
+	FDataType MeshMappingData; 
 };
 
-#endif // __GPUSKINVERTEXFACTORY_H__
diff --git a/Engine/Source/Runtime/Engine/Public/LocalVertexFactory.h b/Engine/Source/Runtime/Engine/Public/LocalVertexFactory.h
--- a/Engine/Source/Runtime/Engine/Public/LocalVertexFactory.h
+++ b/Engine/Source/Runtime/Engine/Public/LocalVertexFactory.h
@@ -22,7 +22,7 @@ public:
 	{
 	}
 
-	struct DataType : public FVertexFactory::DataType
+	struct FDataType
 	{
 		/** The stream to read the vertex position from. */
 		FVertexStreamComponent PositionComponent;
@@ -40,6 +40,9 @@ public:
 		FVertexStreamComponent ColorComponent;
 	};
 
+	DEPRECATED(4.11, "DataType has been renamed to FDataType.")
+	typedef FDataType DataType;
+
 	/**
 	 * Should we cache the material's shadertype on this platform with this vertex factory? 
 	 */
@@ -53,7 +56,7 @@ public:
 	/**
 	 * An implementation of the interface used by TSynchronizedResource to update the resource with new data from the game thread.
 	 */
-	void SetData(const DataType& InData);
+	void SetData(const FDataType& InData);
 
 	/**
 	* Copy the data from another vertex factory
@@ -76,10 +79,10 @@ public:
 	}
 
 protected:
-	DataType Data;
+	FDataType Data;
 	int32 ColorStreamIndex;
 
-	const DataType& GetData() const { return Data; }
+	const FDataType& GetData() const { return Data; }
 };
 
 /**
diff --git a/Engine/Source/Runtime/Engine/Public/MeshParticleVertexFactory.h b/Engine/Source/Runtime/Engine/Public/MeshParticleVertexFactory.h
--- a/Engine/Source/Runtime/Engine/Public/MeshParticleVertexFactory.h
+++ b/Engine/Source/Runtime/Engine/Public/MeshParticleVertexFactory.h
@@ -28,7 +28,7 @@ class FMeshParticleVertexFactory : public FParticleVertexFactoryBase
 	DECLARE_VERTEX_FACTORY_TYPE(FMeshParticleVertexFactory);
 public:
 
-	struct DataType : public FVertexFactory::DataType
+	struct FDataType
 	{
 		/** The stream to read the vertex position from. */
 		FVertexStreamComponent PositionComponent;
@@ -60,7 +60,7 @@ public:
 		/** Flag to mark as initialized */
 		bool bInitialized;
 
-		DataType()
+		FDataType()
 		: bInitialized(false)
 		{
 
@@ -109,7 +109,7 @@ public:
 	/**
 	 * An implementation of the interface used by TSynchronizedResource to update the resource with new data from the game thread.
 	 */
-	ENGINE_API void SetData(const DataType& InData);
+	ENGINE_API void SetData(const FDataType& InData);
 
 	/**
 	 * Set the uniform buffer for this vertex factory.
@@ -160,14 +160,14 @@ public:
 	static FVertexFactoryShaderParameters* ConstructShaderParameters(EShaderFrequency ShaderFrequency);
 
 protected:
-	DataType Data;
+	FDataType Data;
 	
 	/** Stride information for instanced mesh particles */
 	int32 DynamicVertexStride;
 	int32 DynamicParameterVertexStride;
 	
 
-	/** Uniform buffer with mesh particle paramters. */
+	/** Uniform buffer with mesh particle parameters. */
 	FUniformBufferRHIParamRef MeshParticleUniformBuffer;
 };
 
diff --git a/Engine/Source/Runtime/Engine/Public/ParticleHelper.h b/Engine/Source/Runtime/Engine/Public/ParticleHelper.h
--- a/Engine/Source/Runtime/Engine/Public/ParticleHelper.h
+++ b/Engine/Source/Runtime/Engine/Public/ParticleHelper.h
@@ -1367,8 +1367,10 @@ struct FDynamicEmitterDataBase
 	void* operator new(size_t Size);
 	void operator delete(void *RawMemory, size_t Size);
 
-	virtual void CreateVertexFactory(const FSceneView *View)
-	{}
+	virtual FParticleVertexFactoryBase *CreateVertexFactory()
+	{
+		return nullptr;
+	}
 
 	/**
 	 *	Create the render thread resources for this emitter data
@@ -1402,7 +1404,7 @@ struct FDynamicEmitterDataBase
 		ReturnVertexFactory();
 	}
 
-	virtual void GetDynamicMeshElementsEmitter(const FParticleSystemSceneProxy* Proxy, const FSceneView* View, const FSceneViewFamily& ViewFamily, int32 ViewIndex, FMeshElementCollector& Collector) const {}
+	virtual void GetDynamicMeshElementsEmitter(const FParticleSystemSceneProxy* Proxy, const FSceneView* View, const FSceneViewFamily& ViewFamily, int32 ViewIndex, FMeshElementCollector& Collector, FParticleVertexFactoryBase *VertexFactory) const {}
 
 	/**
 	 *	Retrieve the material render proxy to use for rendering this emitter. PURE VIRTUAL
@@ -1429,6 +1431,7 @@ struct FDynamicEmitterDataBase
 	/** true if this emitter has valid rendering data */
 	uint32	bValid:1;
 
+	int32  EmitterIndex;
 protected:
 	/**
 	 *	Create the vertex factory for this emitter data
@@ -1639,20 +1642,14 @@ struct FDynamicSpriteEmitterData : public FDynamicSpriteEmitterDataBase
 {
 	FDynamicSpriteEmitterData(const UParticleModuleRequired* RequiredModule) :
 		FDynamicSpriteEmitterDataBase(RequiredModule)
-		, VertexFactory(nullptr)
 	{
 	}
 
 	~FDynamicSpriteEmitterData()
 	{
-		if (VertexFactory != nullptr)
-		{
-			VertexFactory->ReleaseResource();
-			delete VertexFactory;
-		}
 	}
 
-	virtual void CreateVertexFactory(const FSceneView *View);
+	virtual FParticleVertexFactoryBase *CreateVertexFactory() override;
 
 	/** Initialize this emitter's dynamic rendering data, called after source data has been filled in */
 	void Init( bool bInSelected );
@@ -1714,7 +1711,7 @@ struct FDynamicSpriteEmitterData : public FDynamicSpriteEmitterDataBase
 	/** Gathers simple lights for this emitter. */
 	virtual void GatherSimpleLights(const FParticleSystemSceneProxy* Proxy, const FSceneViewFamily& ViewFamily, FSimpleLightArray& OutParticleLights) const override;
 
-	virtual void GetDynamicMeshElementsEmitter(const FParticleSystemSceneProxy* Proxy, const FSceneView* View, const FSceneViewFamily& ViewFamily, int32 ViewIndex, FMeshElementCollector& Collector) const override;
+	virtual void GetDynamicMeshElementsEmitter(const FParticleSystemSceneProxy* Proxy, const FSceneView* View, const FSceneViewFamily& ViewFamily, int32 ViewIndex, FMeshElementCollector& Collector, FParticleVertexFactoryBase *VertexFactory) const override;
 
 	/**
 	 *	Create the render thread resources for this emitter data
@@ -1740,8 +1737,6 @@ struct FDynamicSpriteEmitterData : public FDynamicSpriteEmitterDataBase
 
 	/** Uniform parameters. Most fields are filled in when updates are sent to the rendering thread, some are per-view! */
 	FParticleSpriteUniformParameters UniformParameters;
-
-	class FParticleSpriteVertexFactory *VertexFactory;
 };
 
 /** Source data for Mesh emitters */
@@ -1801,7 +1796,7 @@ struct FDynamicMeshEmitterData : public FDynamicSpriteEmitterDataBase
 
 	virtual ~FDynamicMeshEmitterData();
 
-	void CreateVertexFactory(const FSceneView *View);
+	FParticleVertexFactoryBase *CreateVertexFactory() override;
 
 	/** Initialize this emitter's dynamic rendering data, called after source data has been filled in */
 	void Init(bool bInSelected,const FParticleMeshEmitterInstance* InEmitterInstance,UStaticMesh* InStaticMesh);
@@ -1826,7 +1821,7 @@ struct FDynamicMeshEmitterData : public FDynamicSpriteEmitterDataBase
 	 */
 	virtual void ReleaseRenderThreadResources(const FParticleSystemSceneProxy* InOwnerProxy) override;
 
-	virtual void GetDynamicMeshElementsEmitter(const FParticleSystemSceneProxy* Proxy, const FSceneView* View, const FSceneViewFamily& ViewFamily, int32 ViewIndex, FMeshElementCollector& Collector) const override;
+	virtual void GetDynamicMeshElementsEmitter(const FParticleSystemSceneProxy* Proxy, const FSceneView* View, const FSceneViewFamily& ViewFamily, int32 ViewIndex, FMeshElementCollector& Collector, FParticleVertexFactoryBase *VertexFactory) const override;
 
 	/**
 	 *	Retrieve the instance data required to render this emitter.
@@ -1926,8 +1921,6 @@ struct FDynamicMeshEmitterData : public FDynamicSpriteEmitterDataBase
 	uint32 bFaceCameraDirectionRatherThanPosition:1;
 	/** The EMeshCameraFacingOption setting to use if bUseCameraFacing is true. */
 	uint8 CameraFacingOption;
-
-	FMeshParticleVertexFactory *VertexFactory;
 };
 
 /** Source data for Beam emitters */
@@ -2083,19 +2076,18 @@ struct FDynamicBeam2EmitterData : public FDynamicSpriteEmitterDataBase
 		: 
 		  FDynamicSpriteEmitterDataBase(RequiredModule)
 		, LastFramePreRendered(-1)
-		, VertexFactory(nullptr)
 	{
 	}
 
 	~FDynamicBeam2EmitterData();
 
-	void CreateVertexFactory(const FSceneView *View);
+	virtual FParticleVertexFactoryBase *CreateVertexFactory() override;
 
 	/** Initialize this emitter's dynamic rendering data, called after source data has been filled in */
 	void Init( bool bInSelected );
 
 
-	virtual void GetDynamicMeshElementsEmitter(const FParticleSystemSceneProxy* Proxy, const FSceneView* View, const FSceneViewFamily& ViewFamily, int32 ViewIndex, FMeshElementCollector& Collector) const override;
+	virtual void GetDynamicMeshElementsEmitter(const FParticleSystemSceneProxy* Proxy, const FSceneView* View, const FSceneViewFamily& ViewFamily, int32 ViewIndex, FMeshElementCollector& Collector, FParticleVertexFactoryBase *VertexFactory) const override;
 
 	virtual FParticleVertexFactoryBase* BuildVertexFactory(const FParticleSystemSceneProxy* InOwnerProxy) override;
 
@@ -2147,8 +2139,6 @@ struct FDynamicBeam2EmitterData : public FDynamicSpriteEmitterDataBase
 	FDynamicBeam2EmitterReplayData Source;
 
 	int32									LastFramePreRendered;
-
-	class FParticleBeamTrailVertexFactory *VertexFactory;
 };
 
 /** Source data for trail-type emitters */
@@ -2235,20 +2225,19 @@ struct FDynamicTrailsEmitterData : public FDynamicSpriteEmitterDataBase
 		, DistanceTessellationStepSize(12.5f)
 		, TangentTessellationScalar(25.0f)
 		, TextureTileDistance(0.0f)
-		, VertexFactory(nullptr)
 	{
 	}
 
 	~FDynamicTrailsEmitterData();
 
-	virtual void CreateVertexFactory(const FSceneView *View);
+	virtual FParticleVertexFactoryBase *CreateVertexFactory() override;
 
 	/** Initialize this emitter's dynamic rendering data, called after source data has been filled in */
 	virtual void Init(bool bInSelected);
 
 	virtual FParticleVertexFactoryBase* BuildVertexFactory(const FParticleSystemSceneProxy* InOwnerProxy) override;
 
-	virtual void GetDynamicMeshElementsEmitter(const FParticleSystemSceneProxy* Proxy, const FSceneView* View, const FSceneViewFamily& ViewFamily, int32 ViewIndex, FMeshElementCollector& Collector) const override;
+	virtual void GetDynamicMeshElementsEmitter(const FParticleSystemSceneProxy* Proxy, const FSceneView* View, const FSceneViewFamily& ViewFamily, int32 ViewIndex, FMeshElementCollector& Collector, FParticleVertexFactoryBase *VertexFactory) const override;
 
 	virtual void RenderDebug(const FParticleSystemSceneProxy* Proxy, FPrimitiveDrawInterface* PDI, const FSceneView* View, bool bCrosses) const override;
 
@@ -2315,8 +2304,6 @@ struct FDynamicTrailsEmitterData : public FDynamicSpriteEmitterDataBase
 	float DistanceTessellationStepSize;
 	float TangentTessellationScalar;
 	float TextureTileDistance;
-
-	class FParticleBeamTrailVertexFactory *VertexFactory;
 };
 
 /** Dynamic emitter data for Ribbon emitters */
@@ -2525,6 +2512,61 @@ public:
 	/** Gets a mesh batch from the pool. */
 	FMeshBatch* GetPooledMeshBatch();
 
+	void MarkVertexFactoriesDirty()
+	{
+		bVertexFactoriesDirty = true;
+	}
+
+	void ClearVertexFactoriesIfDirty() const
+	{
+		if (bVertexFactoriesDirty)
+		{
+			ClearVertexFactories();
+		}
+	}
+
+	void ClearVertexFactories() const
+	{
+		for (int32 Index = 0; Index < EmitterVertexFactoryArray.Num(); Index++)
+		{
+			FParticleVertexFactoryBase *VertexFactory = EmitterVertexFactoryArray[Index];
+			if (VertexFactory)
+			{
+				VertexFactory->ReleaseResource();
+				delete VertexFactory;
+				EmitterVertexFactoryArray[Index] = nullptr;
+			}
+		}
+		bVertexFactoriesDirty = false;
+	}
+
+	void AddEmitterVertexFactory(FDynamicEmitterDataBase *InDynamicData) const
+	{
+		while(InDynamicData->EmitterIndex >= EmitterVertexFactoryArray.Num())
+		{
+			EmitterVertexFactoryArray.Add(nullptr);
+		}
+
+		if (EmitterVertexFactoryArray[InDynamicData->EmitterIndex] == nullptr)
+		{
+			EmitterVertexFactoryArray[InDynamicData->EmitterIndex] = InDynamicData->CreateVertexFactory();
+		}
+	}
+
+	void QueueVertexFactoryCreation(FDynamicEmitterDataBase *InDynamicData)
+	{
+		DynamicDataForThisFrame.Add(InDynamicData);
+	}
+
+	void UpdateVertexFactories() const
+	{
+		for (int32 Index = 0; Index < DynamicDataForThisFrame.Num(); Index++)
+		{
+			AddEmitterVertexFactory(DynamicDataForThisFrame[Index]);
+		}
+		DynamicDataForThisFrame.Empty();
+	}
+
 protected:
 
 	/**
@@ -2566,6 +2608,11 @@ protected:
 	TIndirectArray<FMeshBatch, TInlineAllocator<4> > MeshBatchPool;
 	int32 FirstFreeMeshBatch;
 
+	/** vertex factories for all emitters */
+	mutable TArray<FParticleVertexFactoryBase*> EmitterVertexFactoryArray;
+	mutable TArray<FDynamicEmitterDataBase*> DynamicDataForThisFrame; 
+	mutable bool bVertexFactoriesDirty : 1;
+
 	friend struct FDynamicSpriteEmitterDataBase;
 };
 
diff --git a/Engine/Source/Runtime/Engine/Public/SceneManagement.h b/Engine/Source/Runtime/Engine/Public/SceneManagement.h
--- a/Engine/Source/Runtime/Engine/Public/SceneManagement.h
+++ b/Engine/Source/Runtime/Engine/Public/SceneManagement.h
@@ -139,25 +139,39 @@ public:
 
 	/** Resets pool for GetReusableMID() */
 	virtual void OnStartPostProcessing(FSceneView& CurrentView) = 0;
+
 	/**
 	 * Allows MIDs being created and released during view rendering without the overhead of creating and releasing objects
 	 * As MID are not allowed to be parent of MID this gets fixed up by parenting it to the next Material or MIC
 	 * @param InSource can be Material, MIC or MID, must not be 0
 	 */
 	virtual UMaterialInstanceDynamic* GetReusableMID(class UMaterialInterface* InSource) = 0;
+
+#if !(UE_BUILD_SHIPPING || UE_BUILD_TEST)
+	/** If frozen view matrices are available, set those as active on the SceneView */
+	virtual void ActivateFrozenViewMatrices(FSceneView& SceneView) = 0;
+
+	/** If frozen view matrices were set, restore the previous view matrices */
+	virtual void RestoreUnfrozenViewMatrices(FSceneView& SceneView) = 0;
+#endif
+
 	/** Returns the temporal LOD struct from the viewstate */
 	virtual FTemporalLODState& GetTemporalLODState() = 0;
 	virtual const FTemporalLODState& GetTemporalLODState() const = 0;
+
 	/** 
 	 * Returns the blend factor between the last two LOD samples
 	 */
 	virtual float GetTemporalLODTransition() const = 0;
+
 	/** 
 	 * returns a unique key for the view state, non-zero
 	 */
 	virtual uint32 GetViewKey() const = 0;
+
 	//
 	virtual uint32 GetCurrentTemporalAASampleIndex() const = 0;
+
 	/** 
 	 * returns the occlusion frame counter 
 	 */
@@ -173,7 +187,33 @@ private:
 	int32							NumChildren;
 };
 
+class FFrozenSceneViewMatricesGuard
+{
+public:
+	FFrozenSceneViewMatricesGuard(FSceneView& SV)
+		: SceneView(SV)
+	{
+#if !(UE_BUILD_SHIPPING || UE_BUILD_TEST)
+		if (SceneView.State)
+		{
+			SceneView.State->ActivateFrozenViewMatrices(SceneView);
+		}
+#endif
+	}
+
+	~FFrozenSceneViewMatricesGuard()
+	{
+#if !(UE_BUILD_SHIPPING || UE_BUILD_TEST)
+		if (SceneView.State)
+		{
+			SceneView.State->RestoreUnfrozenViewMatrices(SceneView);
+		}
+#endif
+	}
 
+private:
+	FSceneView& SceneView;
+};
 
 /**
  * The types of interactions between a light and a primitive.
diff --git a/Engine/Source/Runtime/Landscape/Private/LandscapeRender.cpp b/Engine/Source/Runtime/Landscape/Private/LandscapeRender.cpp
--- a/Engine/Source/Runtime/Landscape/Private/LandscapeRender.cpp
+++ b/Engine/Source/Runtime/Landscape/Private/LandscapeRender.cpp
@@ -2560,7 +2560,7 @@ void FLandscapeVertexFactory::InitRHI()
 	Elements.Add(AccessStreamComponent(Data.PositionComponent, 0));
 
 	// create the actual device decls
-	InitDeclaration(Elements, FVertexFactory::DataType());
+	InitDeclaration(Elements);
 }
 
 FVertexFactoryShaderParameters* FLandscapeVertexFactory::ConstructShaderParameters(EShaderFrequency ShaderFrequency)
@@ -2596,7 +2596,7 @@ void FLandscapeVertexFactory::Copy(const FLandscapeVertexFactory& Other)
 	ENQUEUE_UNIQUE_RENDER_COMMAND_TWOPARAMETER(
 		FLandscapeVertexFactoryCopyData,
 		FLandscapeVertexFactory*, VertexFactory, this,
-		const DataType*, DataCopy, &Other.Data,
+		const FDataType*, DataCopy, &Other.Data,
 		{
 			VertexFactory->Data = *DataCopy;
 		});
diff --git a/Engine/Source/Runtime/Landscape/Private/LandscapeRenderMobile.cpp b/Engine/Source/Runtime/Landscape/Private/LandscapeRenderMobile.cpp
--- a/Engine/Source/Runtime/Landscape/Private/LandscapeRenderMobile.cpp
+++ b/Engine/Source/Runtime/Landscape/Private/LandscapeRenderMobile.cpp
@@ -28,7 +28,7 @@ void FLandscapeVertexFactoryMobile::InitRHI()
 	}
 
 	// create the actual device decls
-	InitDeclaration(Elements,FVertexFactory::DataType());
+	InitDeclaration(Elements);
 }
 
 /** Shader parameters for use with FLandscapeVertexFactory */
diff --git a/Engine/Source/Runtime/Landscape/Public/LandscapeRender.h b/Engine/Source/Runtime/Landscape/Public/LandscapeRender.h
--- a/Engine/Source/Runtime/Landscape/Public/LandscapeRender.h
+++ b/Engine/Source/Runtime/Landscape/Public/LandscapeRender.h
@@ -156,7 +156,7 @@ public:
 
 	static FVertexFactoryShaderParameters* ConstructShaderParameters(EShaderFrequency ShaderFrequency);
 
-	struct DataType
+	struct FDataType
 	{
 		/** The stream to read the vertex position from. */
 		FVertexStreamComponent PositionComponent;
@@ -192,7 +192,7 @@ public:
 	/**
 	 * An implementation of the interface used by TSynchronizedResource to update the resource with new data from the game thread.
 	 */
-	void SetData(const DataType& InData)
+	void SetData(const FDataType& InData)
 	{
 		Data = InData;
 		UpdateRHI();
@@ -201,7 +201,7 @@ public:
 	virtual uint64 GetStaticBatchElementVisibility(const FSceneView& View, const FMeshBatch* Batch) const override;
 
 	/** stream component data bound to this vertex factory */
-	DataType Data;
+	FDataType Data;
 };
 
 
diff --git a/Engine/Source/Runtime/Landscape/Public/LandscapeRenderMobile.h b/Engine/Source/Runtime/Landscape/Public/LandscapeRenderMobile.h
--- a/Engine/Source/Runtime/Landscape/Public/LandscapeRenderMobile.h
+++ b/Engine/Source/Runtime/Landscape/Public/LandscapeRenderMobile.h
@@ -28,7 +28,7 @@ class FLandscapeVertexFactoryMobile : public FLandscapeVertexFactory
 	typedef FLandscapeVertexFactory Super;
 public:
 
-	struct DataType : FLandscapeVertexFactory::DataType
+	struct FDataType : FLandscapeVertexFactory::FDataType
 	{
 		/** stream which has heights of each LOD levels */
 		TArray<FVertexStreamComponent,TFixedAllocator<LANDSCAPE_MAX_ES_LOD_COMP> > LODHeightsComponent;
@@ -67,7 +67,7 @@ public:
 	/**
 	 * An implementation of the interface used by TSynchronizedResource to update the resource with new data from the game thread.
 	 */
-	void SetData(const DataType& InData)
+	void SetData(const FDataType& InData)
 	{
 		MobileData = InData;
 		UpdateRHI();
@@ -75,7 +75,7 @@ public:
 
 private:
 	/** stream component data bound to this vertex factory */
-	DataType MobileData; 
+	FDataType MobileData; 
 
 	friend class FLandscapeComponentSceneProxyMobile;
 };
diff --git a/Engine/Source/Runtime/PakFile/Private/SignedArchiveReader.h b/Engine/Source/Runtime/PakFile/Private/SignedArchiveReader.h
--- a/Engine/Source/Runtime/PakFile/Private/SignedArchiveReader.h
+++ b/Engine/Source/Runtime/PakFile/Private/SignedArchiveReader.h
@@ -115,7 +115,7 @@ class FChunkCacheWorker : public FRunnable
 	/** Stops this thread */
 	FThreadSafeCounter StopTaskCounter;
 	/** Available chunk requests */
-	TLockFreePointerListUnordered<FChunkRequest> FreeChunkRequests;
+	TLockFreePointerListUnordered<FChunkRequest, PLATFORM_CACHE_LINE_SIZE> FreeChunkRequests;
 	/** Public decryption key */
 	FEncryptionKey DecryptionKey;
 
diff --git a/Engine/Source/Runtime/RHI/Private/RHI.cpp b/Engine/Source/Runtime/RHI/Private/RHI.cpp
--- a/Engine/Source/Runtime/RHI/Private/RHI.cpp
+++ b/Engine/Source/Runtime/RHI/Private/RHI.cpp
@@ -81,7 +81,7 @@ const FClearValueBinding FClearValueBinding::DepthNear((float)ERHIZBuffer::NearP
 const FClearValueBinding FClearValueBinding::DepthFar((float)ERHIZBuffer::FarPlane, 0);
 
 
-TLockFreePointerListUnordered<FRHIResource> FRHIResource::PendingDeletes;
+TLockFreePointerListUnordered<FRHIResource, PLATFORM_CACHE_LINE_SIZE> FRHIResource::PendingDeletes;
 FRHIResource* FRHIResource::CurrentlyDeleting = nullptr;
 TQueue<FRHIResource::ResourceToDelete> FRHIResource::DeferredDeletionQueue;
 uint32 FRHIResource::CurrentFrame = 0;
diff --git a/Engine/Source/Runtime/RHI/Private/RHICommandList.cpp b/Engine/Source/Runtime/RHI/Private/RHICommandList.cpp
--- a/Engine/Source/Runtime/RHI/Private/RHICommandList.cpp
+++ b/Engine/Source/Runtime/RHI/Private/RHICommandList.cpp
@@ -66,7 +66,7 @@ static TAutoConsoleVariable<int32> CVarRHICmdBufferWriteLocks(
 
 static TAutoConsoleVariable<int32> CVarRHICmdAsyncRHIThreadDispatch(
 	TEXT("r.RHICmdAsyncRHIThreadDispatch"),
-	0,
+	1,
 	TEXT("Experiemental option to do RHI dispatches async. This keeps data flowing to the RHI thread faster and avoid a block at the end of the frame."));
 
 static TAutoConsoleVariable<int32> CVarRHICmdCollectRHIThreadStatsFromHighLevel(
@@ -116,6 +116,14 @@ DECLARE_CYCLE_STAT(TEXT("RHI Thread Execute"), STAT_RHIThreadExecute, STATGROUP_
 
 static TStatId GCurrentExecuteStat;
 
+RHI_API FAutoConsoleTaskPriority CPrio_SceneRenderingTask(
+	TEXT("TaskGraph.TaskPriorities.SceneRenderingTask"),
+	TEXT("Task and thread priority for various scene rendering tasks."),
+	ENamedThreads::NormalThreadPriority, 
+	ENamedThreads::HighTaskPriority 
+	);
+
+
 struct FRHICommandStat : public FRHICommand<FRHICommandStat>
 {
 	TStatId CurrentExecuteStat;
@@ -131,7 +139,10 @@ struct FRHICommandStat : public FRHICommand<FRHICommandStat>
 
 void FRHICommandListImmediate::SetCurrentStat(TStatId Stat)
 {
-	new (AllocCommand<FRHICommandStat>()) FRHICommandStat(Stat);
+	if (!Bypass())
+	{
+		new (AllocCommand<FRHICommandStat>()) FRHICommandStat(Stat);
+	}
 }
 
 DECLARE_CYCLE_STAT(TEXT("FNullGraphTask.RenderThreadTaskFence"), STAT_RenderThreadTaskFence, STATGROUP_TaskGraphTasks);
@@ -193,11 +204,18 @@ void FRHIAsyncComputeCommandListImmediate::ImmediateDispatch(FRHIAsyncComputeCom
 
 FRHICommandBase* GCurrentCommand = nullptr;
 
+DECLARE_CYCLE_STAT(TEXT("BigList"), STAT_BigList, STATGROUP_RHICMDLIST);
+DECLARE_CYCLE_STAT(TEXT("SmallList"), STAT_SmallList, STATGROUP_RHICMDLIST);
+DECLARE_CYCLE_STAT(TEXT("PTrans"), STAT_PTrans, STATGROUP_RHICMDLIST);
+
 void FRHICommandListExecutor::ExecuteInner_DoExecute(FRHICommandListBase& CmdList)
 {
+	FScopeCycleCounter ScopeOuter(CmdList.ExecuteStat);
+
 	CmdList.bExecuting = true;
 	check(CmdList.Context || CmdList.ComputeContext);
 
+
 	FRHICommandListIterator Iter(CmdList);
 #if STATS
 	bool bDoStats =  CVarRHICmdCollectRHIThreadStatsFromHighLevel.GetValueOnRenderThread() > 0 && FThreadStats::IsCollectingData() && (IsInRenderingThread() || IsInRHIThread());
@@ -209,12 +227,12 @@ void FRHICommandListExecutor::ExecuteInner_DoExecute(FRHICommandListBase& CmdLis
 			FScopeCycleCounter Scope(GCurrentExecuteStat);
 			while (Iter.HasCommandsLeft() && Stat == GCurrentExecuteStat.GetRawPointer())
 			{
-			FRHICommandBase* Cmd = Iter.NextCommand();
-			//FPlatformMisc::Prefetch(Cmd->Next);
-			Cmd->CallExecuteAndDestruct(CmdList);
+				FRHICommandBase* Cmd = Iter.NextCommand();
+				//FPlatformMisc::Prefetch(Cmd->Next);
+				Cmd->CallExecuteAndDestruct(CmdList);
+			}
 		}
 	}
-	}
 	else
 #endif
 	{
@@ -689,34 +707,51 @@ void FRHICommandListBase::Reset()
 	{
 		RenderThreadContexts[Index] = nullptr;
 	}
+	ExecuteStat = TStatId();
 }
 
 
 DECLARE_CYCLE_STAT(TEXT("Parallel Async Chain Translate"), STAT_ParallelChainTranslate, STATGROUP_RHICMDLIST);
 
+FAutoConsoleTaskPriority CPrio_FParallelTranslateCommandList(
+	TEXT("TaskGraph.TaskPriorities.ParallelTranslateCommandList"),
+	TEXT("Task and thread priority for FParallelTranslateCommandList."),
+	ENamedThreads::NormalThreadPriority,
+	ENamedThreads::NormalTaskPriority
+	);
+
+FAutoConsoleTaskPriority CPrio_FParallelTranslateCommandListPrepass(
+	TEXT("TaskGraph.TaskPriorities.ParallelTranslateCommandListPrepass"),
+	TEXT("Task and thread priority for FParallelTranslateCommandList for the prepass, which we would like to get to the GPU asap."),
+	ENamedThreads::NormalThreadPriority,
+	ENamedThreads::HighTaskPriority
+	);
+
 class FParallelTranslateCommandList
 {
 	FRHICommandListBase** RHICmdLists;
 	int32 NumCommandLists;
 	IRHICommandContextContainer* ContextContainer;
+	bool bIsPrepass;
 public:
 
-	FParallelTranslateCommandList(FRHICommandListBase** InRHICmdLists, int32 InNumCommandLists, IRHICommandContextContainer* InContextContainer)
+	FParallelTranslateCommandList(FRHICommandListBase** InRHICmdLists, int32 InNumCommandLists, IRHICommandContextContainer* InContextContainer, bool bInIsPrepass)
 		: RHICmdLists(InRHICmdLists)
 		, NumCommandLists(InNumCommandLists)
 		, ContextContainer(InContextContainer)
+		, bIsPrepass(bInIsPrepass)
 	{
 		check(RHICmdLists && ContextContainer && NumCommandLists);
 	}
 
-	FORCEINLINE TStatId GetStatId() const
+	static FORCEINLINE TStatId GetStatId()
 	{
 		RETURN_QUICK_DECLARE_CYCLE_STAT(FParallelTranslateCommandList, STATGROUP_TaskGraphTasks);
 	}
 
 	ENamedThreads::Type GetDesiredThread()
 	{
-		return ENamedThreads::AnyThread;
+		return bIsPrepass ? CPrio_FParallelTranslateCommandListPrepass.Get() : CPrio_FParallelTranslateCommandList.Get();
 	}
 
 	static ESubsequentsMode::Type GetSubsequentsMode() { return ESubsequentsMode::TrackSubsequents; }
@@ -833,36 +868,46 @@ struct FRHICommandWaitForAndSubmitSubList : public FRHICommand<FRHICommandWaitFo
 
 DECLARE_CYCLE_STAT(TEXT("Parallel Setup Translate"), STAT_ParallelSetupTranslate, STATGROUP_RHICMDLIST);
 
+FAutoConsoleTaskPriority CPrio_FParallelTranslateSetupCommandList(
+	TEXT("TaskGraph.TaskPriorities.ParallelTranslateSetupCommandList"),
+	TEXT("Task and thread priority for FParallelTranslateSetupCommandList."),
+	ENamedThreads::HighThreadPriority, // if we have high priority task threads, then use them...
+	ENamedThreads::HighTaskPriority, // .. at high task priority
+	ENamedThreads::HighTaskPriority // if we don't have hi pri threads, then use normal priority threads at high task priority instead
+	);
+
 class FParallelTranslateSetupCommandList
 {
 	FRHICommandList* RHICmdList;
 	FRHICommandListBase** RHICmdLists;
 	int32 NumCommandLists;
+	bool bIsPrepass;
 	int32 MinSize;
 	int32 MinCount;
 public:
 
-	FParallelTranslateSetupCommandList(FRHICommandList* InRHICmdList, FRHICommandListBase** InRHICmdLists, int32 InNumCommandLists)
+	FParallelTranslateSetupCommandList(FRHICommandList* InRHICmdList, FRHICommandListBase** InRHICmdLists, int32 InNumCommandLists, bool bInIsPrepass)
 		: RHICmdList(InRHICmdList)
 		, RHICmdLists(InRHICmdLists)
 		, NumCommandLists(InNumCommandLists)
+		, bIsPrepass(bInIsPrepass)
 	{
 		check(RHICmdList && RHICmdLists && NumCommandLists);
 		MinSize = CVarRHICmdMinCmdlistSizeForParallelTranslate.GetValueOnRenderThread() * 1024;
 		MinCount = CVarRHICmdMinCmdlistForParallelTranslate.GetValueOnRenderThread();
 	}
 
-	FORCEINLINE TStatId GetStatId() const
+	static FORCEINLINE TStatId GetStatId()
 	{
 		RETURN_QUICK_DECLARE_CYCLE_STAT(FParallelTranslateSetupCommandList, STATGROUP_TaskGraphTasks);
 	}
 
-	ENamedThreads::Type GetDesiredThread()
+	static FORCEINLINE ENamedThreads::Type GetDesiredThread()
 	{
-		return ENamedThreads::AnyThread;
+		return CPrio_FParallelTranslateSetupCommandList.Get();
 	}
 
-	static ESubsequentsMode::Type GetSubsequentsMode() { return ESubsequentsMode::TrackSubsequents; }
+	static FORCEINLINE ESubsequentsMode::Type GetSubsequentsMode() { return ESubsequentsMode::TrackSubsequents; }
 
 	void DoTask(ENamedThreads::Type CurrentThread, const FGraphEventRef& MyCompletionGraphEvent)
 	{
@@ -922,7 +967,7 @@ public:
 				IRHICommandContextContainer* ContextContainer =  RHIGetCommandContextContainer();
 				check(ContextContainer);
 
-				FGraphEventRef TranslateCompletionEvent = TGraphTask<FParallelTranslateCommandList>::CreateTask(nullptr, ENamedThreads::RenderThread).ConstructAndDispatchWhenReady(&RHICmdLists[Start], 1 + Last - Start, ContextContainer);
+				FGraphEventRef TranslateCompletionEvent = TGraphTask<FParallelTranslateCommandList>::CreateTask(nullptr, ENamedThreads::RenderThread).ConstructAndDispatchWhenReady(&RHICmdLists[Start], 1 + Last - Start, ContextContainer, bIsPrepass);
 				MyCompletionGraphEvent->DontCompleteUntil(TranslateCompletionEvent);
 				new (RHICmdList->AllocCommand<FRHICommandWaitForAndSubmitSubListParallel>()) FRHICommandWaitForAndSubmitSubListParallel(TranslateCompletionEvent, ContextContainer, EffectiveThreads, ThreadIndex++);
 				Start = Last + 1;
@@ -932,7 +977,7 @@ public:
 	}
 };
 
-void FRHICommandListBase::QueueParallelAsyncCommandListSubmit(FGraphEventRef* AnyThreadCompletionEvents, FRHICommandList** CmdLists, int32* NumDrawsIfKnown, int32 Num, int32 MinDrawsPerTranslate, bool bSpewMerge)
+void FRHICommandListBase::QueueParallelAsyncCommandListSubmit(FGraphEventRef* AnyThreadCompletionEvents, bool bIsPrepass, FRHICommandList** CmdLists, int32* NumDrawsIfKnown, int32 Num, int32 MinDrawsPerTranslate, bool bSpewMerge)
 {
 	check(IsInRenderingThread() && IsImmediate() && Num);
 	if (GRHIThread)
@@ -979,7 +1024,7 @@ void FRHICommandListBase::QueueParallelAsyncCommandListSubmit(FGraphEventRef* An
 			}
 			FRHICommandList* CmdList = new FRHICommandList;
 			CmdList->CopyRenderThreadContexts(*this);
-			FGraphEventRef TranslateSetupCompletionEvent = TGraphTask<FParallelTranslateSetupCommandList>::CreateTask(&Prereq, ENamedThreads::RenderThread).ConstructAndDispatchWhenReady(CmdList, &RHICmdLists[0], Num);
+			FGraphEventRef TranslateSetupCompletionEvent = TGraphTask<FParallelTranslateSetupCommandList>::CreateTask(&Prereq, ENamedThreads::RenderThread).ConstructAndDispatchWhenReady(CmdList, &RHICmdLists[0], Num, bIsPrepass);
 			QueueCommandListSubmit(CmdList);
 			AllOutstandingTasks.Add(TranslateSetupCompletionEvent);
 			if (GRHIThread)
@@ -1072,7 +1117,7 @@ void FRHICommandListBase::QueueParallelAsyncCommandListSubmit(FGraphEventRef* An
 					Prereq.Add(RHIThreadBufferLockFence);
 				}
 
-				FGraphEventRef TranslateCompletionEvent = TGraphTask<FParallelTranslateCommandList>::CreateTask(&Prereq, ENamedThreads::RenderThread).ConstructAndDispatchWhenReady(&RHICmdLists[0], 1 + Last - Start, ContextContainer);
+				FGraphEventRef TranslateCompletionEvent = TGraphTask<FParallelTranslateCommandList>::CreateTask(&Prereq, ENamedThreads::RenderThread).ConstructAndDispatchWhenReady(&RHICmdLists[0], 1 + Last - Start, ContextContainer, bIsPrepass);
 
 				AllOutstandingTasks.Add(TranslateCompletionEvent);
 				new (AllocCommand<FRHICommandWaitForAndSubmitSubListParallel>()) FRHICommandWaitForAndSubmitSubListParallel(TranslateCompletionEvent, ContextContainer, EffectiveThreads, ThreadIndex++);
@@ -1579,8 +1624,8 @@ void FRHICommandListBase::HandleRTThreadTaskCompletion(const FGraphEventRef& MyC
 	RTTasks.Empty();
 }
 
-static TLockFreeFixedSizeAllocator<sizeof(FRHICommandList), FThreadSafeCounter> RHICommandListAllocator;
-static TLockFreeFixedSizeAllocator<sizeof(FRHIAsyncComputeCommandList), FThreadSafeCounter> RHIAsyncComputeCommandListAllocator;
+static TLockFreeFixedSizeAllocator<sizeof(FRHICommandList), PLATFORM_CACHE_LINE_SIZE, FThreadSafeCounter> RHICommandListAllocator;
+static TLockFreeFixedSizeAllocator<sizeof(FRHIAsyncComputeCommandList), PLATFORM_CACHE_LINE_SIZE, FThreadSafeCounter> RHIAsyncComputeCommandListAllocator;
 
 void* FRHICommandList::operator new(size_t Size)
 {
@@ -1898,7 +1943,9 @@ void* FDynamicRHI::LockTexture2D_RenderThread(class FRHICommandListImmediate& RH
 	{
 		QUICK_SCOPE_CYCLE_COUNTER(STAT_RHIMETHOD_LockTexture2D_Flush);
 		RHICmdList.ImmediateFlush(EImmediateFlushType::FlushRHIThread);
+		return GDynamicRHI->RHILockTexture2D(Texture, MipIndex, LockMode, DestStride, bLockWithinMiptail);
 	}
+	FScopedRHIThreadStaller StallRHIThread(RHICmdList);
 	return GDynamicRHI->RHILockTexture2D(Texture, MipIndex, LockMode, DestStride, bLockWithinMiptail);
 }
 
@@ -1908,7 +1955,10 @@ void FDynamicRHI::UnlockTexture2D_RenderThread(class FRHICommandListImmediate& R
 	{
 		QUICK_SCOPE_CYCLE_COUNTER(STAT_RHIMETHOD_UnlockTexture2D_Flush);
 		RHICmdList.ImmediateFlush(EImmediateFlushType::FlushRHIThread);
+		GDynamicRHI->RHIUnlockTexture2D(Texture, MipIndex, bLockWithinMiptail);
+		return;
 	}
+	FScopedRHIThreadStaller StallRHIThread(RHICmdList);
 	GDynamicRHI->RHIUnlockTexture2D(Texture, MipIndex, bLockWithinMiptail);
 }
 
diff --git a/Engine/Source/Runtime/RHI/Public/RHICommandList.h b/Engine/Source/Runtime/RHI/Public/RHICommandList.h
--- a/Engine/Source/Runtime/RHI/Public/RHICommandList.h
+++ b/Engine/Source/Runtime/RHI/Public/RHICommandList.h
@@ -124,7 +124,7 @@ public:
 
 	const int32 GetUsedMemory() const;
 	void QueueAsyncCommandListSubmit(FGraphEventRef& AnyThreadCompletionEvent, class FRHICommandList* CmdList);
-	void QueueParallelAsyncCommandListSubmit(FGraphEventRef* AnyThreadCompletionEvents, class FRHICommandList** CmdLists, int32* NumDrawsIfKnown, int32 Num, int32 MinDrawsPerTranslate, bool bSpewMerge);
+	void QueueParallelAsyncCommandListSubmit(FGraphEventRef* AnyThreadCompletionEvents, bool bIsPrepass, class FRHICommandList** CmdLists, int32* NumDrawsIfKnown, int32 Num, int32 MinDrawsPerTranslate, bool bSpewMerge);
 	void QueueRenderThreadCommandListSubmit(FGraphEventRef& RenderThreadCompletionEvent, class FRHICommandList* CmdList);
 	void QueueCommandListSubmit(class FRHICommandList* CmdList);
 	void WaitForTasks(bool bKnownToBeComplete = false);
@@ -134,13 +134,14 @@ public:
 
 	FORCEINLINE_DEBUGGABLE void* Alloc(int32 AllocSize, int32 Alignment)
 	{
+		checkSlow(!Bypass());
 		return MemManager.Alloc(AllocSize, Alignment);
 	}
 
 	template <typename T>
 	FORCEINLINE_DEBUGGABLE void* Alloc()
 	{
-		return MemManager.Alloc(sizeof(T), ALIGNOF(T));
+		return Alloc(sizeof(T), ALIGNOF(T));
 	}
 
 	template <typename TCmd>
@@ -251,6 +252,7 @@ private:
 	friend class FRHICommandListIterator;
 
 public:
+	TStatId	ExecuteStat;
 	enum class ERenderThreadContext
 	{
 		SceneRenderTargets,
@@ -3009,6 +3011,18 @@ private:
 
 extern RHI_API FRHICommandListExecutor GRHICommandList;
 
+extern RHI_API FAutoConsoleTaskPriority CPrio_SceneRenderingTask;
+
+class FRenderTask
+{
+public:
+	FORCEINLINE static ENamedThreads::Type GetDesiredThread()
+	{
+		return CPrio_SceneRenderingTask.Get();
+	}
+};
+
+
 FORCEINLINE_DEBUGGABLE FRHICommandListImmediate& FRHICommandListExecutor::GetImmediateCommandList()
 {
 	return GRHICommandList.CommandListImmediate;
diff --git a/Engine/Source/Runtime/RHI/Public/RHIResources.h b/Engine/Source/Runtime/RHI/Public/RHIResources.h
--- a/Engine/Source/Runtime/RHI/Public/RHIResources.h
+++ b/Engine/Source/Runtime/RHI/Public/RHIResources.h
@@ -87,7 +87,7 @@ private:
 	mutable FThreadSafeCounter NumRefs;
 	mutable int32 MarkedForDelete;
 	bool bDoNotDeferDelete;
-	static TLockFreePointerListUnordered<FRHIResource> PendingDeletes;
+	static TLockFreePointerListUnordered<FRHIResource, PLATFORM_CACHE_LINE_SIZE> PendingDeletes;
 	static FRHIResource* CurrentlyDeleting;
 
 	// Some APIs don't do internal reference counting, so we have to wait an extra couple of frames before deleting resources
diff --git a/Engine/Source/Runtime/RenderCore/Private/RenderingThread.cpp b/Engine/Source/Runtime/RenderCore/Private/RenderingThread.cpp
--- a/Engine/Source/Runtime/RenderCore/Private/RenderingThread.cpp
+++ b/Engine/Source/Runtime/RenderCore/Private/RenderingThread.cpp
@@ -866,7 +866,7 @@ void AddFrameRenderPrerequisite(const FGraphEventRef& TaskToAdd)
 void AdvanceFrameRenderPrerequisite()
 {
 	checkSlow(IsInGameThread()); 
-	FGraphEventRef PendingComplete = FrameRenderPrerequisites.CreatePrerequisiteCompletionHandle();
+	FGraphEventRef PendingComplete = FrameRenderPrerequisites.CreatePrerequisiteCompletionHandle(ENamedThreads::GameThread);
 	if (PendingComplete.GetReference())
 	{
 		GameThreadWaitForTask(PendingComplete);
@@ -927,7 +927,7 @@ FRHICommandListImmediate& GetImmediateCommandList_ForRenderCommand()
 }
 
 /** The set of deferred cleanup objects which are pending cleanup. */
-static TLockFreePointerListUnordered<FDeferredCleanupInterface>	PendingCleanupObjectsList;
+static TLockFreePointerListUnordered<FDeferredCleanupInterface, PLATFORM_CACHE_LINE_SIZE>	PendingCleanupObjectsList;
 
 FPendingCleanupObjects::FPendingCleanupObjects()
 {
diff --git a/Engine/Source/Runtime/RenderCore/Public/RenderResource.h b/Engine/Source/Runtime/RenderCore/Public/RenderResource.h
--- a/Engine/Source/Runtime/RenderCore/Public/RenderResource.h
+++ b/Engine/Source/Runtime/RenderCore/Public/RenderResource.h
@@ -626,7 +626,7 @@ private:
  * This is used to keep bound shader states that have been used recently from being freed, as they're likely to be used again soon.
  */
 
-template<uint32 Size>
+template<uint32 Size, bool TThreadSafe = true>
 class TBoundShaderStateHistory : public FRenderResource
 {
 public:
@@ -637,15 +637,15 @@ public:
 	{}
 
 	/** Adds a bound shader state to the history. */
-	void Add(FBoundShaderStateRHIParamRef BoundShaderState)
+	FORCEINLINE void Add(FBoundShaderStateRHIParamRef BoundShaderState)
 	{
-		if (GRHISupportsParallelRHIExecute)
+		if (TThreadSafe && GRHISupportsParallelRHIExecute)
 		{
 			BoundShaderStateHistoryLock.Lock();
 		}
 		BoundShaderStates[NextBoundShaderStateIndex] = BoundShaderState;
 		NextBoundShaderStateIndex = (NextBoundShaderStateIndex + 1) % Size;
-		if (GRHISupportsParallelRHIExecute)
+		if (TThreadSafe && GRHISupportsParallelRHIExecute)
 		{
 			BoundShaderStateHistoryLock.Unlock();
 		}
@@ -662,7 +662,7 @@ public:
 	// FRenderResource interface.
 	virtual void ReleaseRHI()
 	{
-		if (GRHISupportsParallelRHIExecute)
+		if (TThreadSafe && GRHISupportsParallelRHIExecute)
 		{
 			BoundShaderStateHistoryLock.Lock();
 		}
@@ -670,7 +670,7 @@ public:
 		{
 			BoundShaderStates[Index].SafeRelease();
 		}
-		if (GRHISupportsParallelRHIExecute)
+		if (TThreadSafe && GRHISupportsParallelRHIExecute)
 		{
 			BoundShaderStateHistoryLock.Unlock();
 		}
diff --git a/Engine/Source/Runtime/Renderer/Private/DeferredShadingRenderer.cpp b/Engine/Source/Runtime/Renderer/Private/DeferredShadingRenderer.cpp
--- a/Engine/Source/Runtime/Renderer/Private/DeferredShadingRenderer.cpp
+++ b/Engine/Source/Runtime/Renderer/Private/DeferredShadingRenderer.cpp
@@ -71,6 +71,13 @@ static TAutoConsoleVariable<int32> CVarRHICmdFlushRenderThreadTasksBasePass(
 	0,
 	TEXT("Wait for completion of parallel render thread tasks at the end of the base pass. A more granular version of r.RHICmdFlushRenderThreadTasks. If either r.RHICmdFlushRenderThreadTasks or r.RHICmdFlushRenderThreadTasksBasePass is > 0 we will flush."));
 
+static TAutoConsoleVariable<int32> CVarFXSystemPreRenderAfterPrepass(
+	TEXT("r.FXSystemPreRenderAfterPrepass"),
+	0,
+	TEXT("If > 0, then do the FX prerender after the prepass. This improves pipelining for greater performance. Experiemental option."),
+	ECVF_RenderThreadSafe
+	);
+
 DECLARE_CYCLE_STAT(TEXT("PostInitViews FlushDel"), STAT_PostInitViews_FlushDel, STATGROUP_InitViews);
 DECLARE_CYCLE_STAT(TEXT("InitViews Intentional Stall"), STAT_InitViews_Intentional_Stall, STATGROUP_InitViews);
 
@@ -300,6 +307,14 @@ void FDeferredShadingSceneRenderer::RenderBasePassStaticDataDefaultParallel(FPar
 	}
 }
 
+FAutoConsoleTaskPriority CPrio_FSortFrontToBackTask(
+	TEXT("TaskGraph.TaskPriorities.SortFrontToBackTask"),
+	TEXT("Task and thread priority for FSortFrontToBackTask."),
+	ENamedThreads::HighThreadPriority, // if we have high priority task threads, then use them...
+	ENamedThreads::NormalTaskPriority, // .. at normal task priority
+	ENamedThreads::HighTaskPriority // if we don't have hi pri threads, then use normal priority threads at high task priority instead
+	);
+
 template<typename StaticMeshDrawList>
 class FSortFrontToBackTask
 {
@@ -322,7 +337,7 @@ public:
 
 	ENamedThreads::Type GetDesiredThread()
 	{
-		return ENamedThreads::AnyThread;
+		return CPrio_FSortFrontToBackTask.Get();
 	}
 
 	static ESubsequentsMode::Type GetSubsequentsMode() { return ESubsequentsMode::TrackSubsequents; }
@@ -450,7 +465,7 @@ void FDeferredShadingSceneRenderer::RenderBasePassDynamicData(FRHICommandList& R
 	}
 }
 
-class FRenderBasePassDynamicDataThreadTask
+class FRenderBasePassDynamicDataThreadTask : public FRenderTask
 {
 	FDeferredShadingSceneRenderer& ThisRenderer;
 	FRHICommandList& RHICmdList;
@@ -474,11 +489,6 @@ public:
 		RETURN_QUICK_DECLARE_CYCLE_STAT(FRenderBasePassDynamicDataThreadTask, STATGROUP_TaskGraphTasks);
 	}
 
-	ENamedThreads::Type GetDesiredThread()
-	{
-		return ENamedThreads::AnyThread;
-	}
-
 	static ESubsequentsMode::Type GetSubsequentsMode() { return ESubsequentsMode::TrackSubsequents; }
 
 	void DoTask(ENamedThreads::Type CurrentThread, const FGraphEventRef& MyCompletionGraphEvent)
@@ -537,13 +547,15 @@ static void SetupBasePassView(FRHICommandList& RHICmdList, const FViewInfo& View
 	RHICmdList.SetRasterizerState(TStaticRasterizerState<FM_Solid, CM_None>::GetRHI());
 }
 
+DECLARE_CYCLE_STAT(TEXT("Basepass"), STAT_CLP_Basepass, STATGROUP_ParallelCommandListMarkers);
+
 class FBasePassParallelCommandListSet : public FParallelCommandListSet
 {
 public:
 	const FSceneViewFamily& ViewFamily;
 
 	FBasePassParallelCommandListSet(const FViewInfo& InView, FRHICommandListImmediate& InParentCmdList, bool bInParallelExecute, bool bInCreateSceneContext, const FSceneViewFamily& InViewFamily)
-		: FParallelCommandListSet(InView, InParentCmdList, bInParallelExecute, bInCreateSceneContext)
+		: FParallelCommandListSet(GET_STATID(STAT_CLP_Basepass), InView, InParentCmdList, bInParallelExecute, bInCreateSceneContext)
 		, ViewFamily(InViewFamily)
 	{
 		SetStateOnCommandList(ParentCmdList);
@@ -986,7 +998,9 @@ void FDeferredShadingSceneRenderer::Render(FRHICommandListImmediate& RHICmdList)
 	}
 
 	// Notify the FX system that the scene is about to be rendered.
-	if (Scene->FXSystem && Views.IsValidIndex(0))
+	bool bLateFXPrerender = CVarFXSystemPreRenderAfterPrepass.GetValueOnRenderThread() > 0;
+	bool bDoFXPrerender = Scene->FXSystem && Views.IsValidIndex(0);
+	if (!bLateFXPrerender && bDoFXPrerender)
 	{
 		SCOPE_CYCLE_COUNTER(STAT_FDeferredShadingSceneRenderer_FXSystem_PreRender);
 		Scene->FXSystem->PreRender(RHICmdList, &Views[0].GlobalDistanceFieldInfo.ParameterData);
@@ -1000,14 +1014,17 @@ void FDeferredShadingSceneRenderer::Render(FRHICommandListImmediate& RHICmdList)
 	}
 
 	// Draw the scene pre-pass / early z pass, populating the scene depth buffer and HiZ
-	bool bDepthWasCleared = RenderPrePassHMD(RHICmdList);
+	RHICmdList.SetCurrentStat(GET_STATID(STAT_CLM_PrePass));
 	const bool bNeedsPrePass = NeedsPrePass(this);
+	bool bDepthWasCleared;
 	if (bNeedsPrePass)
 	{
-		RHICmdList.SetCurrentStat(GET_STATID(STAT_CLM_PrePass));
-		RenderPrePass(RHICmdList, bDepthWasCleared);
-		// at this point, the depth was cleared
-		bDepthWasCleared = true;
+		bDepthWasCleared = RenderPrePass(bLateFXPrerender && bDoFXPrerender, RHICmdList);
+	}
+	else
+	{
+		// we didn't do the prepass, but we still want the HMD mask if there is one
+		bDepthWasCleared = RenderPrePassHMD(RHICmdList);
 	}
 
 	RHICmdList.SetCurrentStat(GET_STATID(STAT_CLM_AfterPrePass));
@@ -1581,7 +1598,7 @@ bool FDeferredShadingSceneRenderer::RenderPrePassView(FRHICommandList& RHICmdLis
 	return bDirty;
 }
 
-class FRenderPrepassDynamicDataThreadTask
+class FRenderPrepassDynamicDataThreadTask : public FRenderTask
 {
 	FDeferredShadingSceneRenderer& ThisRenderer;
 	FRHICommandList& RHICmdList;
@@ -1605,11 +1622,6 @@ public:
 		RETURN_QUICK_DECLARE_CYCLE_STAT(FRenderPrepassDynamicDataThreadTask, STATGROUP_TaskGraphTasks);
 	}
 
-	ENamedThreads::Type GetDesiredThread()
-	{
-		return ENamedThreads::AnyThread;
-	}
-
 	static ESubsequentsMode::Type GetSubsequentsMode() { return ESubsequentsMode::TrackSubsequents; }
 
 	void DoTask(ENamedThreads::Type CurrentThread, const FGraphEventRef& MyCompletionGraphEvent)
@@ -1619,18 +1631,22 @@ public:
 	}
 };
 
+DECLARE_CYCLE_STAT(TEXT("Prepass"), STAT_CLP_Prepass, STATGROUP_ParallelCommandListMarkers);
+
 class FPrePassParallelCommandListSet : public FParallelCommandListSet
 {
 public:
 	FPrePassParallelCommandListSet(const FViewInfo& InView, FRHICommandListImmediate& InParentCmdList, bool bInParallelExecute, bool bInCreateSceneContext)
-		: FParallelCommandListSet(InView, InParentCmdList, bInParallelExecute, bInCreateSceneContext)
+		: FParallelCommandListSet(GET_STATID(STAT_CLP_Prepass), InView, InParentCmdList, bInParallelExecute, bInCreateSceneContext)
 	{
-		SetStateOnCommandList(ParentCmdList);
+		// Do not copy-paste. this is a very unusual FParallelCommandListSet because it is a prepass and we want to do some work after starting some tasks
 	}
 
 	virtual ~FPrePassParallelCommandListSet()
 	{
-		Dispatch();
+		// Do not copy-paste. this is a very unusual FParallelCommandListSet because it is a prepass and we want to do some work after starting some tasks
+		SetStateOnCommandList(ParentCmdList);
+		Dispatch(true);
 	}
 
 	virtual void SetStateOnCommandList(FRHICommandList& CmdList) override
@@ -1655,9 +1671,10 @@ static TAutoConsoleVariable<int32> CVarRHICmdFlushRenderThreadTasksPrePass(
 	0,
 	TEXT("Wait for completion of parallel render thread tasks at the end of the pre pass.  A more granular version of r.RHICmdFlushRenderThreadTasks. If either r.RHICmdFlushRenderThreadTasks or r.RHICmdFlushRenderThreadTasksPrePass is > 0 we will flush."));
 
-void FDeferredShadingSceneRenderer::RenderPrePassViewParallel(const FViewInfo& View, FRHICommandListImmediate& ParentCmdList)
+bool FDeferredShadingSceneRenderer::RenderPrePassViewParallel(const FViewInfo& View, FRHICommandListImmediate& ParentCmdList, bool bDoFXPrerender)
 {
-	FPrePassParallelCommandListSet ParallelCommandListSet(View, ParentCmdList, 
+	bool bDepthWasCleared = false;
+	FPrePassParallelCommandListSet ParallelCommandListSet(View, ParentCmdList,
 		CVarRHICmdPrePassDeferredContexts.GetValueOnRenderThread() > 0, 
 		CVarRHICmdFlushRenderThreadTasksPrePass.GetValueOnRenderThread() == 0  && CVarRHICmdFlushRenderThreadTasks.GetValueOnRenderThread() == 0);
 
@@ -1691,6 +1708,12 @@ void FDeferredShadingSceneRenderer::RenderPrePassViewParallel(const FViewInfo& V
 		}
 	}
 
+	// we do this step here (awkwardly) so that the above tasks can be in flight while we get the particles (which must be dynamic) setup.
+	if (bDoFXPrerender)
+	{
+		bDepthWasCleared = PreRenderPrePass(ParentCmdList, bDoFXPrerender);
+	}
+
 	// Dynamic
 	FRHICommandList* CmdList = ParallelCommandListSet.NewParallelCommandList();
 
@@ -1698,6 +1721,8 @@ void FDeferredShadingSceneRenderer::RenderPrePassViewParallel(const FViewInfo& V
 		.ConstructAndDispatchWhenReady(*this, *CmdList, View);
 
 	ParallelCommandListSet.AddParallelCommandList(CmdList, AnyThreadCompletionEvent);
+
+	return bDepthWasCleared;
 }
 
 /** A pixel shader used to fill the stencil buffer with the current dithered transition mask. */
@@ -1742,24 +1767,29 @@ public:
 IMPLEMENT_SHADER_TYPE(, FDitheredTransitionStencilPS, TEXT("DitheredTransitionStencil"), TEXT("Main"), SF_Pixel);
 FGlobalBoundShaderState DitheredTransitionStencilBoundShaderState;
 
-/** Renders the scene's prepass and occlusion queries */
-bool FDeferredShadingSceneRenderer::RenderPrePass(FRHICommandListImmediate& RHICmdList, bool bDepthWasCleared)
+/** Possibly do the FX prerender and setup the prepass*/
+bool FDeferredShadingSceneRenderer::PreRenderPrePass(FRHICommandListImmediate& RHICmdList, bool bDoFXPrerender)
 {
-	SCOPED_DRAW_EVENT(RHICmdList, PrePass);
-	SCOPE_CYCLE_COUNTER(STAT_DepthDrawTime);
+	if (bDoFXPrerender)
+	{
+		check(Scene->FXSystem && Views.IsValidIndex(0));
+		SCOPE_CYCLE_COUNTER(STAT_FDeferredShadingSceneRenderer_FXSystem_PreRender);
+		Scene->FXSystem->PreRender(RHICmdList, &Views[0].GlobalDistanceFieldInfo.ParameterData);
+	}
+	bool bDepthWasCleared = RenderPrePassHMD(RHICmdList);
 
-	bool bDirty = false;
 	FSceneRenderTargets& SceneContext = FSceneRenderTargets::Get(RHICmdList);
 
 	SceneContext.BeginRenderingPrePass(RHICmdList, !bDepthWasCleared);
-	
+	bDepthWasCleared = true;
+
 	// Dithered transition stencil mask fill
-	if(bDitheredLODTransitionsUseStencil)
+	if (bDitheredLODTransitionsUseStencil)
 	{
 		SCOPED_DRAW_EVENT(RHICmdList, DitheredStencilPrePass);
 		FIntPoint BufferSizeXY = SceneContext.GetBufferSizeXY();
-		
-		for(int32 ViewIndex = 0; ViewIndex < Views.Num(); ++ViewIndex)
+
+		for (int32 ViewIndex = 0; ViewIndex < Views.Num(); ++ViewIndex)
 		{
 			SCOPED_CONDITIONAL_DRAW_EVENTF(RHICmdList, EventView, Views.Num() > 1, TEXT("View%d"), ViewIndex);
 
@@ -1780,7 +1810,7 @@ bool FDeferredShadingSceneRenderer::RenderPrePass(FRHICommandListImmediate& RHIC
 				true, CF_Always, SO_Keep, SO_Keep, SO_Replace,
 				false, CF_Always, SO_Keep, SO_Keep, SO_Keep,
 				STENCIL_SANDBOX_MASK, STENCIL_SANDBOX_MASK>::GetRHI(), STENCIL_SANDBOX_MASK);
-			
+
 			DrawRectangle(
 				RHICmdList,
 				0, 0,
@@ -1793,28 +1823,55 @@ bool FDeferredShadingSceneRenderer::RenderPrePass(FRHICommandListImmediate& RHIC
 				EDRF_UseTriangleOptimization);
 		}
 	}
+	return bDepthWasCleared;
+}
+
+bool FDeferredShadingSceneRenderer::RenderPrePass(bool bDoFXPrerender, FRHICommandListImmediate& RHICmdList)
+{
+	bool bDepthWasCleared = false;
+	SCOPED_DRAW_EVENT(RHICmdList, PrePass);
+	SCOPE_CYCLE_COUNTER(STAT_DepthDrawTime);
+
+	bool bDirty = false;
+	bool bDidPrePre = false;
+	FSceneRenderTargets& SceneContext = FSceneRenderTargets::Get(RHICmdList);
+
+	bool bParallel = GRHICommandList.UseParallelAlgorithms() && CVarParallelPrePass.GetValueOnRenderThread();
+
+	if (!bDoFXPrerender || !bParallel)
+	{
+		// nothing to be gained by delaying this.
+		bDepthWasCleared = PreRenderPrePass(RHICmdList, bDoFXPrerender);
+		bDidPrePre = true;
+		bDoFXPrerender = false; // only need to do this once
+	}
+	else
+	{
+		SceneContext.GetSceneDepthSurface(); // this probably isn't needed, but if there was some lazy allocation of the depth surface going on, we want it allocated now before we go wide. We may not have called BeginRenderingPrePass yet if bDoFXPrerender is true
+	}
 
 	// Draw a depth pass to avoid overdraw in the other passes.
 	if(EarlyZPassMode != DDM_None)
 	{
-		if (GRHICommandList.UseParallelAlgorithms() && CVarParallelPrePass.GetValueOnRenderThread())
+		if (bParallel)
 		{
 			FScopedCommandListWaitForTasks Flusher(CVarRHICmdFlushRenderThreadTasksPrePass.GetValueOnRenderThread() > 0 || CVarRHICmdFlushRenderThreadTasks.GetValueOnRenderThread() > 0, RHICmdList);
-
 			for(int32 ViewIndex = 0;ViewIndex < Views.Num();ViewIndex++)
 			{
 				SCOPED_CONDITIONAL_DRAW_EVENTF(RHICmdList, EventView, Views.Num() > 1, TEXT("View%d"), ViewIndex);
 				const FViewInfo& View = Views[ViewIndex];
 				if (View.ShouldRenderView())
 				{
-					RenderPrePassViewParallel(View, RHICmdList);
+					bDepthWasCleared = RenderPrePassViewParallel(View, RHICmdList, bDoFXPrerender) || bDepthWasCleared;
 					bDirty = true; // assume dirty since we are not going to wait
+					bDidPrePre = true;
+					bDoFXPrerender = false; // only need to do this once even if we have multiple views
 				}
 			}
 		}
 		else
 		{
-			for(int32 ViewIndex = 0;ViewIndex < Views.Num();ViewIndex++)
+			for (int32 ViewIndex = 0; ViewIndex < Views.Num(); ViewIndex++)
 			{
 				SCOPED_CONDITIONAL_DRAW_EVENTF(RHICmdList, EventView, Views.Num() > 1, TEXT("View%d"), ViewIndex);
 				const FViewInfo& View = Views[ViewIndex];
@@ -1825,6 +1882,12 @@ bool FDeferredShadingSceneRenderer::RenderPrePass(FRHICommandListImmediate& RHIC
 			}
 		}
 	}
+	if (!bDidPrePre)
+	{
+		// For some reason we haven't done this yet. Best do it now for consistency with the old code.
+		bDepthWasCleared = PreRenderPrePass(RHICmdList, bDoFXPrerender);
+		bDidPrePre = true;
+	}
 
 	// Dithered transition stencil mask clear
 	if(bDitheredLODTransitionsUseStencil)
@@ -1834,7 +1897,7 @@ bool FDeferredShadingSceneRenderer::RenderPrePass(FRHICommandListImmediate& RHIC
 
 	SceneContext.FinishRenderingPrePass(RHICmdList);
 
-	return bDirty;
+	return bDepthWasCleared;
 }
 
 bool FDeferredShadingSceneRenderer::RenderPrePassHMD(FRHICommandListImmediate& RHICmdList)
diff --git a/Engine/Source/Runtime/Renderer/Private/DeferredShadingRenderer.h b/Engine/Source/Runtime/Renderer/Private/DeferredShadingRenderer.h
--- a/Engine/Source/Runtime/Renderer/Private/DeferredShadingRenderer.h
+++ b/Engine/Source/Runtime/Renderer/Private/DeferredShadingRenderer.h
@@ -61,8 +61,9 @@ public:
 
 	/**
 	 * Renders the scene's prepass for a particular view in parallel
+	 * @return true if the depth was cleared
 	 */
-	void RenderPrePassViewParallel(const FViewInfo& View, FRHICommandListImmediate& ParentCmdList);
+	bool RenderPrePassViewParallel(const FViewInfo& View, FRHICommandListImmediate& ParentCmdList, bool bDoFXPrerender);
 
 	/** Renders the basepass for the static data of a given View. */
 	bool RenderBasePassStaticData(FRHICommandList& RHICmdList, FViewInfo& View);
@@ -161,11 +162,17 @@ private:
 
 	void CreateIndirectCapsuleShadows();
 
+	/**
+	* Setup the prepass. This is split out so that in parallel we can do the fx prerender after we start the parallel tasks
+	* @return true if the depth was cleared
+	*/
+	bool PreRenderPrePass(FRHICommandListImmediate& RHICmdList, bool bDoFXPrerender);
+
 	/**
 	 * Renders the scene's prepass and occlusion queries.
-	 * @return true if anything was rendered
+	 * @return true if the depth was cleared
 	 */
-	bool RenderPrePass(FRHICommandListImmediate& RHICmdList, bool bDepthWasCleared);
+	bool RenderPrePass(bool bDoFXPrerender, FRHICommandListImmediate& RHICmdList);
 
 	/**
 	 * Renders the active HMD's hidden area mask as a depth prepass, if available.
diff --git a/Engine/Source/Runtime/Renderer/Private/DynamicPrimitiveDrawing.inl b/Engine/Source/Runtime/Renderer/Private/DynamicPrimitiveDrawing.inl
--- a/Engine/Source/Runtime/Renderer/Private/DynamicPrimitiveDrawing.inl
+++ b/Engine/Source/Runtime/Renderer/Private/DynamicPrimitiveDrawing.inl
@@ -48,7 +48,7 @@ void DrawViewElementsInner(
 }
 
 template<typename DrawingPolicyFactoryType>
-class FDrawViewElementsAnyThreadTask
+class FDrawViewElementsAnyThreadTask : public FRenderTask
 {
 	FRHICommandList& RHICmdList;
 	const FViewInfo& View;
@@ -86,11 +86,6 @@ public:
 		RETURN_QUICK_DECLARE_CYCLE_STAT(FDrawViewElementsAnyThreadTask, STATGROUP_TaskGraphTasks);
 	}
 
-	ENamedThreads::Type GetDesiredThread()
-	{
-		return ENamedThreads::AnyThread;
-	}
-
 	static ESubsequentsMode::Type GetSubsequentsMode() { return ESubsequentsMode::TrackSubsequents; }
 
 	void DoTask(ENamedThreads::Type CurrentThread, const FGraphEventRef& MyCompletionGraphEvent)
diff --git a/Engine/Source/Runtime/Renderer/Private/IndirectLightingCache.cpp b/Engine/Source/Runtime/Renderer/Private/IndirectLightingCache.cpp
--- a/Engine/Source/Runtime/Renderer/Private/IndirectLightingCache.cpp
+++ b/Engine/Source/Runtime/Renderer/Private/IndirectLightingCache.cpp
@@ -371,6 +371,13 @@ FIndirectLightingCacheAllocation* FIndirectLightingCache::FindPrimitiveAllocatio
 	return PrimitiveAllocations.FindRef(PrimitiveId);
 }
 
+FAutoConsoleTaskPriority CPrio_FUpdateCachePrimitivesTask(
+	TEXT("TaskGraph.TaskPriorities.UpdateCachePrimitivesTask"),
+	TEXT("Task and thread priority for FUpdateCachePrimitivesTask."),
+	ENamedThreads::NormalThreadPriority,
+	ENamedThreads::NormalTaskPriority
+	);
+
 class FUpdateCachePrimitivesTask
 {
 	FIndirectLightingCache* ILC;
@@ -399,7 +406,7 @@ public:
 
 	ENamedThreads::Type GetDesiredThread()
 	{
-		return ENamedThreads::AnyThread;
+		return CPrio_FUpdateCachePrimitivesTask.Get();
 	}
 
 	static ESubsequentsMode::Type GetSubsequentsMode() { return ESubsequentsMode::TrackSubsequents; }
@@ -418,7 +425,7 @@ void FIndirectLightingCache::StartUpdateCachePrimitivesTask(FScene* Scene, FScen
 void FIndirectLightingCache::FinalizeCacheUpdates(FScene* Scene, FSceneRenderer& Renderer, FILCUpdatePrimTaskData& TaskData)
 {
 	SCOPE_CYCLE_COUNTER(STAT_UpdateIndirectLightingCacheFinalize);	
-	FTaskGraphInterface::Get().WaitUntilTaskCompletes(TaskData.TaskRef, ENamedThreads::AnyThread);
+	FTaskGraphInterface::Get().WaitUntilTaskCompletes(TaskData.TaskRef, ENamedThreads::RenderThread_Local);
 	FinalizeUpdateInternal_RenderThread(Scene, Renderer, TaskData.OutBlocksToUpdate, TaskData.OutTransitionsOverTimeToUpdate);
 }
 
diff --git a/Engine/Source/Runtime/Renderer/Private/PostProcess/PostProcessEyeAdaptation.cpp b/Engine/Source/Runtime/Renderer/Private/PostProcess/PostProcessEyeAdaptation.cpp
--- a/Engine/Source/Runtime/Renderer/Private/PostProcess/PostProcessEyeAdaptation.cpp
+++ b/Engine/Source/Runtime/Renderer/Private/PostProcess/PostProcessEyeAdaptation.cpp
@@ -89,9 +89,20 @@ static ERHIFeatureLevel::Type BasicEyeAdaptationMinFeatureLevel = ERHIFeatureLev
 TAutoConsoleVariable<int32> CVarEyeAdaptationMethodOveride(
 	TEXT("r.EyeAdaptation.MethodOveride"),
 	-1,
+	TEXT("Overide the eye adapation method set in post processing volumes\n")
+	TEXT("-2: override with custom settings (for testing Basic Mode)\n")
 	TEXT("-1: no override\n")
-	TEXT("1: Histogram-based. \n")
-	TEXT("2: Basic \n"),
+	TEXT(" 1: Histogram-based\n")
+	TEXT(" 2: Basic"),
+	ECVF_Scalability | ECVF_RenderThreadSafe);
+
+// Initialize the static CVar used in computing the weighting focus in basic eye-adaptation
+TAutoConsoleVariable<float> CVarEyeAdaptationFocus(
+	TEXT("r.EyeAdaptation.Focus"),
+	1.0f,
+	TEXT("Applies to basic adapation mode only\n")
+	TEXT(" 0: Uniform weighting\n")
+	TEXT(">0: Center focus, 1 is a good number (default)"),
 	ECVF_Scalability | ECVF_RenderThreadSafe);
 
 /** Encapsulates the histogram-based post processing eye adaptation pixel shader. */
@@ -384,6 +395,7 @@ public:
 		PostprocessParameter.Bind(Initializer.ParameterMap);
 		EyeAdaptationTexture.Bind(Initializer.ParameterMap, TEXT("EyeAdaptationTexture"));
 		EyeAdaptationParams.Bind(Initializer.ParameterMap, TEXT("EyeAdaptationParams"));
+		EyeAdaptationExtent.Bind(Initializer.ParameterMap, TEXT("EyeAdaptionSrcExtent"));
 	}
 
 public:
@@ -401,7 +413,7 @@ public:
 	}
 
 
-	void SetPS(const FRenderingCompositePassContext& Context, IPooledRenderTarget* EyeAdaptationLastFrameRT)
+	void SetPS(const FRenderingCompositePassContext& Context, const FIntPoint SrcSize, IPooledRenderTarget* EyeAdaptationLastFrameRT)
 	{
 		const FPixelShaderRHIParamRef ShaderRHI = GetPixelShader();
 
@@ -428,15 +440,22 @@ public:
 			ComputeEyeAdaptationValues(BasicEyeAdaptationMinFeatureLevel, Context.View, Temp);
 			// Log-based computation of the exposure scale has a built in scaling.
 			//Temp[1].X *= 0.16;  
+			//Encode the eye-focus slope
+			// Get the focus value for the eye-focus weighting
+			Temp[2].W = GetBasicAutoExposureFocus();
 			SetShaderValueArray(Context.RHICmdList, ShaderRHI, EyeAdaptationParams, Temp, 3);
 		}
+
+		// Set the src extent for the shader
+		SetShaderValue(Context.RHICmdList, ShaderRHI, EyeAdaptationExtent, SrcSize);
 	}
 
 	// FShader interface.
 	virtual bool Serialize(FArchive& Ar) override
 	{
 		bool bShaderHasOutdatedParameters = FGlobalShader::Serialize(Ar);
-		Ar << PostprocessParameter << EyeAdaptationTexture << EyeAdaptationParams;
+		Ar << PostprocessParameter << EyeAdaptationTexture 
+			<< EyeAdaptationParams << EyeAdaptationExtent;
 		return bShaderHasOutdatedParameters;
 	}
 
@@ -444,6 +463,7 @@ private:
 	FPostProcessPassParameters PostprocessParameter;
 	FShaderResourceParameter EyeAdaptationTexture;
 	FShaderParameter EyeAdaptationParams;
+	FShaderParameter EyeAdaptationExtent;
 };
 
 IMPLEMENT_SHADER_TYPE(, FPostProcessLogLuminance2ExposureScalePS, TEXT("PostProcessEyeAdaptation"), TEXT("MainLogLuminance2ExposureScalePS"), SF_Pixel);
@@ -464,6 +484,9 @@ void FRCPassPostProcessBasicEyeAdaptation::Process(FRenderingCompositePassContex
 
 	FIntPoint DestSize = EyeAdaptationThisFrameRT->GetDesc().Extent;
 
+	// The input texture sample size.  Averaged in the pixel shader.
+	const FIntPoint SrcSize = GetInputDesc(ePId_Input0)->Extent;
+
 	// we render to our own output render target, not the intermediate one created by the compositing system
 	// Set the view family's render target/viewport.
 	SetRenderTarget(Context.RHICmdList, EyeAdaptationThisFrameRT->GetRenderTargetItem().TargetableTexture, FTextureRHIRef(), true);
@@ -484,7 +507,7 @@ void FRCPassPostProcessBasicEyeAdaptation::Process(FRenderingCompositePassContex
 
 	// Set the parameters used by the pixel shader.
 
-	PixelShader->SetPS(Context, EyeAdaptationLastFrameRT);
+	PixelShader->SetPS(Context, SrcSize, EyeAdaptationLastFrameRT);
 
 	// Draw a quad mapping scene color to the view's render target
 	DrawRectangle(
diff --git a/Engine/Source/Runtime/Renderer/Private/PostProcess/PostProcessEyeAdaptation.h b/Engine/Source/Runtime/Renderer/Private/PostProcess/PostProcessEyeAdaptation.h
--- a/Engine/Source/Runtime/Renderer/Private/PostProcess/PostProcessEyeAdaptation.h
+++ b/Engine/Source/Runtime/Renderer/Private/PostProcess/PostProcessEyeAdaptation.h
@@ -107,6 +107,15 @@ static inline bool IsAutoExposureMethodSupported(const ERHIFeatureLevel::Type& F
 	return Result;
 }
 
+extern TAutoConsoleVariable<float> CVarEyeAdaptationFocus;
+static inline float GetBasicAutoExposureFocus()
+{
+	// Hard coded value camp.
+	static float clampValue = 10.f;
+	float FocusValue = CVarEyeAdaptationFocus.GetValueOnRenderThread();
+	
+	return FMath::Max(FMath::Min(FocusValue, clampValue), 0.0f);
+}
 
 // derives from TRenderingCompositePassBase<InputCount, OutputCount>
 // ePId_Input0: Color input
diff --git a/Engine/Source/Runtime/Renderer/Private/PostProcess/PostProcessVisualizeHDR.cpp b/Engine/Source/Runtime/Renderer/Private/PostProcess/PostProcessVisualizeHDR.cpp
--- a/Engine/Source/Runtime/Renderer/Private/PostProcess/PostProcessVisualizeHDR.cpp
+++ b/Engine/Source/Runtime/Renderer/Private/PostProcess/PostProcessVisualizeHDR.cpp
@@ -91,6 +91,15 @@ public:
 			FVector4 Temp[3];
 
 			FRCPassPostProcessEyeAdaptation::ComputeEyeAdaptationParamsValue(Context.View, Temp);
+			if (GetAutoExposureMethod(Context.View) == EAutoExposureMethod::AEM_Basic)
+			{
+				Temp[2].W = GetBasicAutoExposureFocus();
+			}
+			else
+			{
+				Temp[2].W = 0.0f;
+			}
+
 			SetShaderValueArray(Context.RHICmdList, ShaderRHI, EyeAdaptationParams, Temp, 3);
 		}
 
@@ -273,29 +282,29 @@ void FRCPassPostProcessVisualizeHDR::Process(FRenderingCompositePassContext& Con
 	Y += 3 * YStep;
 	if (GetAutoExposureMethod(ViewInfo) == EAutoExposureMethod::AEM_Basic)
 	{
-		Line = FString::Printf(TEXT("AutoExposure Basic"));
+		Line = FString::Printf(TEXT("Basic"));
 	} 
 	else
 	{
-		Line = FString::Printf(TEXT("AutoExposure Histogram"));
+		Line = FString::Printf(TEXT("Histogram"));
 	}
-	Canvas.DrawShadowedString(X, Y += YStep, TEXT("EyeAdaptationPercent Method:"), GetStatsFont(), FLinearColor(1, 1, 1));
+	Canvas.DrawShadowedString(X, Y += YStep, TEXT("Auto Exposure Method:"), GetStatsFont(), FLinearColor(1, 1, 1));
 	Canvas.DrawShadowedString(X + ColumnWidth, Y, *Line, GetStatsFont(), FLinearColor(1, 1, 1));
 
 	Line = FString::Printf(TEXT("%g%% .. %g%%"), View.FinalPostProcessSettings.AutoExposureLowPercent, View.FinalPostProcessSettings.AutoExposureHighPercent);
-	Canvas.DrawShadowedString( X, Y += YStep, TEXT("EyeAdaptationPercent Low/High:"), GetStatsFont(), FLinearColor(1, 1, 1));
+	Canvas.DrawShadowedString( X, Y += YStep, TEXT("Percent Low/High:"), GetStatsFont(), FLinearColor(1, 1, 1));
 	Canvas.DrawShadowedString( X + ColumnWidth, Y, *Line, GetStatsFont(), FLinearColor(1, 1, 1));
 
 	Line = FString::Printf(TEXT("%g .. %g"), View.FinalPostProcessSettings.AutoExposureMinBrightness, View.FinalPostProcessSettings.AutoExposureMaxBrightness);
-	Canvas.DrawShadowedString( X, Y += YStep, TEXT("EyeAdaptationBrightness Min/Max:"), GetStatsFont(), FLinearColor(1, 1, 1));
+	Canvas.DrawShadowedString( X, Y += YStep, TEXT("Brightness Min/Max:"), GetStatsFont(), FLinearColor(1, 1, 1));
 	Canvas.DrawShadowedString( X + ColumnWidth, Y, *Line, GetStatsFont(), FLinearColor(0.3f, 0.3f, 1));
 
 	Line = FString::Printf(TEXT("%g / %g"), View.FinalPostProcessSettings.AutoExposureSpeedUp, View.FinalPostProcessSettings.AutoExposureSpeedDown);
-	Canvas.DrawShadowedString( X, Y += YStep, TEXT("EyeAdaptionSpeed Up/Down:"), GetStatsFont(), FLinearColor(1, 1, 1));
+	Canvas.DrawShadowedString( X, Y += YStep, TEXT("Speed Up/Down:"), GetStatsFont(), FLinearColor(1, 1, 1));
 	Canvas.DrawShadowedString( X + ColumnWidth, Y, *Line, GetStatsFont(), FLinearColor(1, 1, 1));
 
 	Line = FString::Printf(TEXT("%g"), View.FinalPostProcessSettings.AutoExposureBias);
-	Canvas.DrawShadowedString( X, Y += YStep, TEXT("Exposure Offset: "), GetStatsFont(), FLinearColor(1, 1, 1));
+	Canvas.DrawShadowedString( X, Y += YStep, TEXT("Exposure Bias: "), GetStatsFont(), FLinearColor(1, 1, 1));
 	Canvas.DrawShadowedString( X + ColumnWidth, Y, *Line, GetStatsFont(), FLinearColor(1, 0.3f, 0.3f));
 
 	Line = FString::Printf(TEXT("%g .. %g (log2)"),
@@ -306,6 +315,13 @@ void FRCPassPostProcessVisualizeHDR::Process(FRenderingCompositePassContext& Con
 		*LogToString(View.FinalPostProcessSettings.HistogramLogMin), *LogToString(View.FinalPostProcessSettings.HistogramLogMax));
 	Canvas.DrawShadowedString( X + ColumnWidth, Y+= YStep, *Line, GetStatsFont(), FLinearColor(0.3f, 0.3f, 1));
 
+	if (GetAutoExposureMethod(ViewInfo) == EAutoExposureMethod::AEM_Basic)
+	{
+		Line = FString::Printf(TEXT("%g"), GetBasicAutoExposureFocus());
+		Canvas.DrawShadowedString(X, Y += YStep, TEXT("Weighting Focus: "), GetStatsFont(), FLinearColor(1, 1, 1));
+		Canvas.DrawShadowedString(X + ColumnWidth, Y, *Line, GetStatsFont(), FLinearColor(1, 0.3f, 0.3f));
+	}
+
 	Canvas.Flush_RenderThread(Context.RHICmdList);
 
 	Context.RHICmdList.CopyToResolveTarget(DestRenderTarget.TargetableTexture, DestRenderTarget.ShaderResourceTexture, false, FResolveParams());
