Copying //UE4/Dev-Rendering to Dev-Main (//UE4/Dev-Main)

diff --git a/Engine/Source/Runtime/Renderer/Private/PostProcess/PostProcessing.cpp b/Engine/Source/Runtime/Renderer/Private/PostProcess/PostProcessing.cpp
--- a/Engine/Source/Runtime/Renderer/Private/PostProcess/PostProcessing.cpp
+++ b/Engine/Source/Runtime/Renderer/Private/PostProcess/PostProcessing.cpp
@@ -148,14 +148,14 @@ static TAutoConsoleVariable<int32> CVarMotionBlurNew(
 static TAutoConsoleVariable<int32> CVarMotionBlurScatter(
 	TEXT("r.MotionBlurScatter"),
 	0,
-	TEXT(""),
+	TEXT("Forces scatter based max velocity method (slower)."),
 	ECVF_RenderThreadSafe
 );
 
 static TAutoConsoleVariable<int32> CVarMotionBlurSeparable(
 	TEXT("r.MotionBlurSeparable"),
 	0,
-	TEXT(""),
+	TEXT("Adds a second motion blur pass that smooths noise for a higher quality blur."),
 	ECVF_RenderThreadSafe
 );
 
@@ -196,27 +196,27 @@ public:
 	TBloomDownSampleArray(FPostprocessContext& InContext, FRenderingCompositeOutputRef SourceDownsample, bool bGenerateLog2Alpha) :
 		bHasLog2Alpha(bGenerateLog2Alpha), Context(InContext)
 	{
-		// Optionally encode log2 data in the alpha channel. Used for EyeAdaptation.
-		if (bHasLog2Alpha) {
-			FRenderingCompositePass* BasicEyeSetupPass = Context.Graph.RegisterPass(new(FMemStack::Get()) FRCPassPostProcessBasicEyeAdaptationSetUp());
-			BasicEyeSetupPass->SetInput(ePId_Input0, SourceDownsample);
-			PostProcessDownsamples[0] = FRenderingCompositeOutputRef(BasicEyeSetupPass); 
-		}
-		else
-		{
-			// The source becomes the 0th down sample.
-			PostProcessDownsamples[0] = SourceDownsample;
-		}
 
-		// Queue the additional down samples.  T
+		static const TCHAR* PassLabels[] =
+		{ NULL, TEXT("BloomDownsample1"), TEXT("BloomDownsample2"), TEXT("BloomDownsample3"), TEXT("BloomDownsample4"), TEXT("BloomDownsample5") };
+		static_assert(ARRAY_COUNT(PassLabels) == DownSampleStages, "PassLabel count must be equal to DownSampleStages.");
+
+		// The first down sample is the input
+		PostProcessDownsamples[0] = SourceDownsample;
+
+		// Queue the down samples. 
 		for (int i = 1; i < DownSampleStages; i++)
 		{
-			static const TCHAR* PassLabels[] =
-			{ NULL, TEXT("BloomDownsample1"), TEXT("BloomDownsample2"), TEXT("BloomDownsample3"), TEXT("BloomDownsample4"), TEXT("BloomDownsample5") };
-			static_assert(ARRAY_COUNT(PassLabels) == DownSampleStages, "PassLabel count must be equal to DownSampleStages.");
 			FRenderingCompositePass* Pass = Context.Graph.RegisterPass(new(FMemStack::Get()) FRCPassPostProcessDownsample(PF_Unknown, 1, PassLabels[i]));
 			Pass->SetInput(ePId_Input0, PostProcessDownsamples[i - 1]);
 			PostProcessDownsamples[i] = FRenderingCompositeOutputRef(Pass);
+
+			// Add log2 data to the alpha channel after doing the 1st (i==1) down sample pass
+			if (bHasLog2Alpha && i == 1 ) {
+				FRenderingCompositePass* BasicEyeSetupPass = Context.Graph.RegisterPass(new(FMemStack::Get()) FRCPassPostProcessBasicEyeAdaptationSetUp());
+				BasicEyeSetupPass->SetInput(ePId_Input0, PostProcessDownsamples[i]);
+				PostProcessDownsamples[i] = FRenderingCompositeOutputRef(BasicEyeSetupPass);
+			}
 		}
 	}
 
@@ -379,7 +379,7 @@ static FRenderingCompositeOutputRef AddPostProcessBasicEyeAdaptation(FViewInfo&
 	// Extract the last (i.e. smallest) down sample
 	static const int32 FinalDSIdx = FBloomDownSampleArray::Num() - 1;
 	FRenderingCompositeOutputRef PostProcessPriorReduction = BloomAndEyeDownSamples.PostProcessDownsamples[FinalDSIdx];
-	
+
 	// Compute the eye adaptation value based on average luminance from log2 luminance buffer, history, and specific shader parameters.
 	FRenderingCompositePass* Node = Context.Graph.RegisterPass(new(FMemStack::Get()) FRCPassPostProcessBasicEyeAdaptation());
 	Node->SetInput(ePId_Input0, PostProcessPriorReduction);
@@ -1158,12 +1158,12 @@ void FPostProcessing::Process(FRHICommandListImmediate& RHICmdList, FViewInfo& V
 		{
 		public:
 			FAutoExposure(const FViewInfo& InView) : 
-				MethodId(GetAutoExposureMethod(InView)) 
+				MethodId(GetAutoExposureMethod(InView))
 			{}
 			// distinguish between Basic and Histogram-based
 			EAutoExposureMethod          MethodId;
 			// not always valid
-			FRenderingCompositeOutputRef EyeAdaptation;			
+			FRenderingCompositeOutputRef EyeAdaptation;
 		} AutoExposure(View);
 	
 		// not always valid
@@ -1709,7 +1709,22 @@ void FPostProcessing::Process(FRHICommandListImmediate& RHICmdList, FViewInfo& V
 		else
 		{
 			check(!FSceneRenderTargets::Get(RHICmdList).SeparateTranslucencyRT);
-			AddGammaOnlyTonemapper(Context);
+
+			if (SeparateTranslucency.IsValid())
+			{
+				// separate translucency is done here or in AddPostProcessDepthOfFieldBokeh()
+				FRenderingCompositePass* NodeRecombined = Context.Graph.RegisterPass(new(FMemStack::Get()) FRCPassPostProcessBokehDOFRecombine());
+				NodeRecombined->SetInput(ePId_Input0, Context.FinalOutput);
+				NodeRecombined->SetInput(ePId_Input2, SeparateTranslucency);
+
+				Context.FinalOutput = FRenderingCompositeOutputRef(NodeRecombined);
+			}
+
+			const EDebugViewShaderMode DebugViewShaderMode = View.Family->GetDebugViewShaderMode();
+			if (DebugViewShaderMode == DVSM_None)
+			{
+				AddGammaOnlyTonemapper(Context);
+			}
 		}
 		
 		if(View.Family->EngineShowFlags.StationaryLightOverlap)
@@ -1724,14 +1739,14 @@ void FPostProcessing::Process(FRHICommandListImmediate& RHICmdList, FViewInfo& V
 		{
 			float ComplexityScale = 1.f / (float)(GEngine->QuadComplexityColors.Num() - 1) / NormalizedQuadComplexityValue; // .1f comes from the values used in LightAccumulator_GetResult
 			FRenderingCompositePass* Node = Context.Graph.RegisterPass(new(FMemStack::Get()) FRCPassPostProcessVisualizeComplexity(GEngine->QuadComplexityColors, FVisualizeComplexityApplyPS::CS_STAIR, ComplexityScale, true));
-			Node->SetInput(ePId_Input0, FRenderingCompositeOutputRef(Context.SceneColor));
+			Node->SetInput(ePId_Input0, FRenderingCompositeOutputRef(Context.FinalOutput));
 			Context.FinalOutput = FRenderingCompositeOutputRef(Node);
 		}
 
 		if(DebugViewShaderMode == DVSM_ShaderComplexity || DebugViewShaderMode == DVSM_ShaderComplexityContainedQuadOverhead || DebugViewShaderMode == DVSM_ShaderComplexityBleedingQuadOverhead)
 		{
 			FRenderingCompositePass* Node = Context.Graph.RegisterPass(new(FMemStack::Get()) FRCPassPostProcessVisualizeComplexity(GEngine->ShaderComplexityColors, FVisualizeComplexityApplyPS::CS_RAMP, 1.f, true));
-			Node->SetInput(ePId_Input0, FRenderingCompositeOutputRef(Context.SceneColor));
+			Node->SetInput(ePId_Input0, FRenderingCompositeOutputRef(Context.FinalOutput));
 			Context.FinalOutput = FRenderingCompositeOutputRef(Node);
 		}
 
diff --git a/Engine/Source/Runtime/Renderer/Private/PostProcess/SceneRenderTargets.cpp b/Engine/Source/Runtime/Renderer/Private/PostProcess/SceneRenderTargets.cpp
--- a/Engine/Source/Runtime/Renderer/Private/PostProcess/SceneRenderTargets.cpp
+++ b/Engine/Source/Runtime/Renderer/Private/PostProcess/SceneRenderTargets.cpp
@@ -1124,7 +1124,7 @@ void FSceneRenderTargets::BeginRenderingPrePass(FRHICommandList& RHICmdList, boo
 	{
 		// Set the scene depth surface and a DUMMY buffer as color buffer
 		// (as long as it's the same dimension as the depth buffer),	
-		SetRenderTarget(RHICmdList, FTextureRHIRef(), GetSceneDepthSurface());
+		SetRenderTarget(RHICmdList, FTextureRHIRef(), DepthTarget);
 		RHICmdList.BindClearMRTValues(false, true, true);
 	}
 }
@@ -1440,7 +1440,7 @@ bool FSceneRenderTargets::IsSeparateTranslucencyActive(const FViewInfo& View) co
 	int32 Value = FSceneRenderTargets::CVarSetSeperateTranslucencyEnabled.GetValueOnRenderThread();
 
 	return (Value != 0) && CurrentFeatureLevel >= ERHIFeatureLevel::SM4
-		&& View.Family->EngineShowFlags.PostProcessing
+		&& (View.Family->EngineShowFlags.PostProcessing || View.Family->EngineShowFlags.ShaderComplexity)
 		&& View.Family->EngineShowFlags.SeparateTranslucency;
 }
 
diff --git a/Engine/Source/Runtime/Renderer/Private/RendererScene.cpp b/Engine/Source/Runtime/Renderer/Private/RendererScene.cpp
--- a/Engine/Source/Runtime/Renderer/Private/RendererScene.cpp
+++ b/Engine/Source/Runtime/Renderer/Private/RendererScene.cpp
@@ -67,6 +67,7 @@ FSceneViewState::FSceneViewState()
 #if !(UE_BUILD_SHIPPING || UE_BUILD_TEST)
 	bIsFreezing = false;
 	bIsFrozen = false;
+	bIsFrozenViewMatricesCached = false;
 #endif
 	// Register this object as a resource, so it will receive device reset notifications.
 	if ( IsInGameThread() )
diff --git a/Engine/Source/Runtime/Renderer/Private/ScenePrivate.h b/Engine/Source/Runtime/Renderer/Private/ScenePrivate.h
--- a/Engine/Source/Runtime/Renderer/Private/ScenePrivate.h
+++ b/Engine/Source/Runtime/Renderer/Private/ScenePrivate.h
@@ -471,13 +471,19 @@ public:
 
 #if !(UE_BUILD_SHIPPING || UE_BUILD_TEST)
 	/** Are we currently in the state of freezing rendering? (1 frame where we gather what was rendered) */
-	bool bIsFreezing;
+	uint32 bIsFreezing : 1;
 
 	/** Is rendering currently frozen? */
-	bool bIsFrozen;
+	uint32 bIsFrozen : 1;
+
+	/** True if the CachedViewMatrices is holding frozen view matrices, otherwise false */
+	uint32 bIsFrozenViewMatricesCached : 1;
 
 	/** The set of primitives that were rendered the frame that we froze rendering */
 	TSet<FPrimitiveComponentId> FrozenPrimitives;
+
+	/** The cache view matrices at the time of freezing or the cached debug fly cam's view matrices. */
+	FViewMatrices CachedViewMatrices;
 #endif
 
 private:
@@ -694,6 +700,13 @@ public:
 	 */
 	void UpdateTemporalLODTransition(const FViewInfo& View)
 	{
+#if !(UE_BUILD_SHIPPING || UE_BUILD_TEST)
+		if (bIsFrozen)
+		{
+			return;
+		}
+#endif
+
 		TemporalLODState.UpdateTemporalLODTransition(View, LastRenderTime);
 	}
 
@@ -894,29 +907,57 @@ public:
 		return NewMID;
 	}
 
+#if !(UE_BUILD_SHIPPING || UE_BUILD_TEST)
+	virtual void ActivateFrozenViewMatrices(FSceneView& SceneView) override
+	{
+		auto* ViewState = static_cast<FSceneViewState*>(SceneView.State);
+		if (ViewState->bIsFrozen)
+		{
+			check(ViewState->bIsFrozenViewMatricesCached);
+
+			Swap(SceneView.ViewMatrices, ViewState->CachedViewMatrices);
+			ViewState->bIsFrozenViewMatricesCached = false;
+		}
+	}
+
+	virtual void RestoreUnfrozenViewMatrices(FSceneView& SceneView) override
+	{
+		auto* ViewState = static_cast<FSceneViewState*>(SceneView.State);
+		if (ViewState->bIsFrozen)
+		{
+			check(!ViewState->bIsFrozenViewMatricesCached);
+
+			Swap(SceneView.ViewMatrices, ViewState->CachedViewMatrices);
+			ViewState->bIsFrozenViewMatricesCached = true;
+		}
+	}
+#endif
+
 	virtual FTemporalLODState& GetTemporalLODState() override
 	{
 		return TemporalLODState;
 	}
+
 	virtual const FTemporalLODState& GetTemporalLODState() const override
 	{
 		return TemporalLODState;
 	}
+
 	float GetTemporalLODTransition() const override
 	{
 		return TemporalLODState.GetTemporalLODTransition(LastRenderTime);
 	}
+
 	uint32 GetViewKey() const override
 	{
 		return UniqueID;
 	}
+
 	uint32 GetOcclusionFrameCounter() const
 	{
 		return OcclusionFrameCounter;
 	}
 
-
-
 	// FDeferredCleanupInterface
 	virtual void FinishCleanup() override
 	{
diff --git a/Engine/Source/Runtime/Renderer/Private/SceneRendering.cpp b/Engine/Source/Runtime/Renderer/Private/SceneRendering.cpp
--- a/Engine/Source/Runtime/Renderer/Private/SceneRendering.cpp
+++ b/Engine/Source/Runtime/Renderer/Private/SceneRendering.cpp
@@ -129,6 +129,33 @@ static TAutoConsoleVariable<float> CVarTessellationAdaptivePixelsPerTriangle(
 	TEXT("Global tessellation factor multiplier"),
 	ECVF_RenderThreadSafe);
 
+/*-----------------------------------------------------------------------------
+BitCounting
+-----------------------------------------------------------------------------*/
+
+MS_ALIGN(64) uint8 CountBitsTable[64] GCC_ALIGN(64) = { 0 };
+
+static struct FInitBitCounts
+{
+	FInitBitCounts()
+	{
+		for (uint32 Index = 0; Index < 64; Index++)
+		{
+			uint32 LocalIndex = Index;
+			uint8 Result = 0;
+			while (LocalIndex)
+			{
+				if (LocalIndex & 1)
+				{
+					Result++;
+				}
+				LocalIndex >>= 1;
+			}
+			CountBitsTable[Index] = Result;
+		}
+	}
+} GInitBitCounts;
+
 /*-----------------------------------------------------------------------------
 	FParallelCommandListSet
 -----------------------------------------------------------------------------*/
@@ -143,7 +170,7 @@ static TAutoConsoleVariable<int32> CVarRHICmdSpewParallelListBalance(
 
 static TAutoConsoleVariable<int32> CVarRHICmdBalanceParallelLists(
 	TEXT("r.RHICmdBalanceParallelLists"),
-	1,
+	2,
 	TEXT("Allows to enable a preprocess of the drawlists to try to balance the load equally among the command lists.\n")
 	TEXT(" 0: off \n")
 	TEXT(" 1: enabled")
@@ -156,13 +183,14 @@ static TAutoConsoleVariable<int32> CVarRHICmdMinCmdlistForParallelSubmit(
 
 static TAutoConsoleVariable<int32> CVarRHICmdMinDrawsPerParallelCmdList(
 	TEXT("r.RHICmdMinDrawsPerParallelCmdList"),
-	32,
+	64,
 	TEXT("The minimum number of draws per cmdlist. If the total number of draws is less than this, then no parallel work will be done at all. This can't always be honored or done correctly. More effective with RHICmdBalanceParallelLists."));
 
-FParallelCommandListSet::FParallelCommandListSet(const FViewInfo& InView, FRHICommandListImmediate& InParentCmdList, bool bInParallelExecute, bool bInCreateSceneContext)
+FParallelCommandListSet::FParallelCommandListSet(TStatId InExecuteStat, const FViewInfo& InView, FRHICommandListImmediate& InParentCmdList, bool bInParallelExecute, bool bInCreateSceneContext)
 	: View(InView)
 	, ParentCmdList(InParentCmdList)
 	, Snapshot(nullptr)
+	, ExecuteStat(InExecuteStat)
 	, NumAlloc(0)
 	, bParallelExecute(GRHISupportsParallelRHIExecute && bInParallelExecute)
 	, bCreateSceneContext(bInCreateSceneContext)
@@ -184,7 +212,7 @@ FRHICommandList* FParallelCommandListSet::AllocCommandList()
 	return new FRHICommandList;
 }
 
-void FParallelCommandListSet::Dispatch()
+void FParallelCommandListSet::Dispatch(bool bHighPriority)
 {
 	check(IsInRenderingThread() && FMemStack::Get().GetNumMarks() == 1); // we do not want this popped before the end of the scene and it better be the scene allocator
 	check(CommandLists.Num() == Events.Num());
@@ -230,7 +258,7 @@ void FParallelCommandListSet::Dispatch()
 		UE_CLOG(bSpewBalance, LogTemp, Display, TEXT("%d cmdlists for parallel translate"), CommandLists.Num());
 		check(GRHISupportsParallelRHIExecute);
 		NumAlloc -= CommandLists.Num();
-		ParentCmdList.QueueParallelAsyncCommandListSubmit(&Events[0], &CommandLists[0], &NumDrawsIfKnown[0], CommandLists.Num(), (MinDrawsPerCommandList * 4) / 3, bSpewBalance);
+		ParentCmdList.QueueParallelAsyncCommandListSubmit(&Events[0], bHighPriority, &CommandLists[0], &NumDrawsIfKnown[0], CommandLists.Num(), (MinDrawsPerCommandList * 4) / 3, bSpewBalance);
 		SetStateOnCommandList(ParentCmdList);
 	}
 	else
@@ -259,7 +287,8 @@ FParallelCommandListSet::~FParallelCommandListSet()
 FRHICommandList* FParallelCommandListSet::NewParallelCommandList()
 {
 	FRHICommandList* Result = AllocCommandList();
-	SetStateOnCommandList(*Result); 
+	Result->ExecuteStat = ExecuteStat;
+	SetStateOnCommandList(*Result);
 	if (bCreateSceneContext)
 	{
 		FSceneRenderTargets& SceneContext = FSceneRenderTargets::Get(ParentCmdList);
@@ -1454,25 +1483,18 @@ void FSceneRenderer::RenderFinish(FRHICommandListImmediate& RHICmdList)
 			{
 				ViewState->bIsFreezing = false;
 				ViewState->bIsFrozen = true;
+				ViewState->bIsFrozenViewMatricesCached = true;
+				ViewState->CachedViewMatrices = View.ViewMatrices;
 			}
 
 			// handle freeze toggle request
 			if (bHasRequestedToggleFreeze)
 			{
-				// do we want to start freezing?
-				if (!ViewState->bIsFrozen)
-				{
-					ViewState->bIsFrozen = false;
-					ViewState->bIsFreezing = true;
-					ViewState->FrozenPrimitives.Empty();
-				}
-				// or stop?
-				else
-				{
-					ViewState->bIsFrozen = false;
-					ViewState->bIsFreezing = false;
-					ViewState->FrozenPrimitives.Empty();
-				}
+				// do we want to start freezing or stop?
+				ViewState->bIsFreezing = !ViewState->bIsFrozen;
+				ViewState->bIsFrozen = false;
+				ViewState->bIsFrozenViewMatricesCached = false;
+				ViewState->FrozenPrimitives.Empty();
 			}
 		}
 #endif
diff --git a/Engine/Source/Runtime/Renderer/Private/SceneRendering.h b/Engine/Source/Runtime/Renderer/Private/SceneRendering.h
--- a/Engine/Source/Runtime/Renderer/Private/SceneRendering.h
+++ b/Engine/Source/Runtime/Renderer/Private/SceneRendering.h
@@ -315,12 +315,15 @@ private:
 	uint32 ValidFrameNumber;
 };
 
+DECLARE_STATS_GROUP(TEXT("Parallel Command List Markers"), STATGROUP_ParallelCommandListMarkers, STATCAT_Advanced);
+
 class FParallelCommandListSet
 {
 public:
 	const FViewInfo& View;
 	FRHICommandListImmediate& ParentCmdList;
 	FSceneRenderTargets* Snapshot;
+	TStatId	ExecuteStat;
 	int32 Width;
 	int32 NumAlloc;
 	int32 MinDrawsPerCommandList;
@@ -337,12 +340,12 @@ public:
 protected:
 	//this must be called by deriving classes virtual destructor because it calls the virtual SetStateOnCommandList.
 	//C++ will not do dynamic dispatch of virtual calls from destructors so we can't call it in the base class.
-	void Dispatch();
+	void Dispatch(bool bHighPriority = false);
 	FRHICommandList* AllocCommandList();
 	bool bParallelExecute;
 	bool bCreateSceneContext;
 public:
-	FParallelCommandListSet(const FViewInfo& InView, FRHICommandListImmediate& InParentCmdList, bool bInParallelExecute, bool bInCreateSceneContext);
+	FParallelCommandListSet(TStatId InExecuteStat, const FViewInfo& InView, FRHICommandListImmediate& InParentCmdList, bool bInParallelExecute, bool bInCreateSceneContext);
 	virtual ~FParallelCommandListSet();
 	int32 NumParallelCommandLists() const
 	{
diff --git a/Engine/Source/Runtime/Renderer/Private/SceneVisibility.cpp b/Engine/Source/Runtime/Renderer/Private/SceneVisibility.cpp
--- a/Engine/Source/Runtime/Renderer/Private/SceneVisibility.cpp
+++ b/Engine/Source/Runtime/Renderer/Private/SceneVisibility.cpp
@@ -106,14 +106,6 @@ static FAutoConsoleVariableRef CVarOcclusionCullParallelPrimFetch(
 	ECVF_RenderThreadSafe
 	);
 
-static int32 GOcclusionCullParallelPrimFetchHiPri = 1;
-static FAutoConsoleVariableRef CVarOcclusionCullParallelPrimFetchHiPri(
-	TEXT("r.OcclusionCullParallelPrimFetchHiPri"),
-	GOcclusionCullParallelPrimFetchHiPri,
-	TEXT("When using r.OcclusionCullParallelPrimFetch, setting this to 1 uses hi priority tasks."),
-	ECVF_RenderThreadSafe
-	);
-
 static int32 GILCUpdatePrimTaskEnabled = 0;
 static FAutoConsoleVariableRef CVarILCUpdatePrimitivesTask(
 	TEXT("r.Cache.UpdatePrimsTaskEnabled"),
@@ -943,6 +935,14 @@ static void FetchVisibilityForPrimitives_Range(FVisForPrimParams& Params)
 	Params.NumOccludedPrims = NumOccludedPrimitives;	
 }
 
+FAutoConsoleTaskPriority CPrio_FetchVisibilityForPrimitivesTask(
+	TEXT("TaskGraph.TaskPriorities.FetchVisibilityForPrimitivesTask"),
+	TEXT("Task and thread priority for FetchVisibilityForPrimitivesTask."),
+	ENamedThreads::HighThreadPriority, // if we have high priority task threads, then use them...
+	ENamedThreads::NormalTaskPriority, // .. at normal task priority
+	ENamedThreads::HighTaskPriority // if we don't have hi pri threads, then use normal priority threads at high task priority instead
+	);
+
 class FetchVisibilityForPrimitivesTask
 {
 	FVisForPrimParams& Params;
@@ -961,7 +961,7 @@ public:
 
 	ENamedThreads::Type GetDesiredThread()
 	{
-		return GOcclusionCullParallelPrimFetchHiPri ? ENamedThreads::HiPri(ENamedThreads::AnyThread) : ENamedThreads::AnyThread;
+		return CPrio_FetchVisibilityForPrimitivesTask.Get();
 	}
 
 	static ESubsequentsMode::Type GetSubsequentsMode() { return ESubsequentsMode::TrackSubsequents; }
@@ -1071,7 +1071,7 @@ static int32 FetchVisibilityForPrimitives(const FScene* Scene, FViewInfo& View,
 				&SubIsOccluded
 				);
 
-			TaskRefArray[i] = TGraphTask<FetchVisibilityForPrimitivesTask>::CreateTask(nullptr, ENamedThreads::AnyThread).ConstructAndDispatchWhenReady(Params[i]);			
+			TaskRefArray[i] = TGraphTask<FetchVisibilityForPrimitivesTask>::CreateTask().ConstructAndDispatchWhenReady(Params[i]);			
 			TaskWaitArray.Add(TaskRefArray[i]);
 
 			StartIndex += NumToProcess;
@@ -1555,6 +1555,7 @@ struct FRelevancePacket
 		// using a local counter to reduce memory traffic
 		int32 NumVisibleStaticMeshElements = 0;
 		FViewInfo& WriteView = const_cast<FViewInfo&>(View);
+		FFrozenSceneViewMatricesGuard FrozenMatricesGuard(WriteView);
 
 		const bool bHLODActive = Scene->SceneLODHierarchy.IsActive();
 
diff --git a/Engine/Source/Runtime/Renderer/Private/ShadowRendering.cpp b/Engine/Source/Runtime/Renderer/Private/ShadowRendering.cpp
--- a/Engine/Source/Runtime/Renderer/Private/ShadowRendering.cpp
+++ b/Engine/Source/Runtime/Renderer/Private/ShadowRendering.cpp
@@ -1575,7 +1575,7 @@ void FProjectedShadowInfo::RenderDepthDynamic(FRHICommandList& RHICmdList, FScen
 	}
 }
 
-class FDrawShadowMeshElementsThreadTask
+class FDrawShadowMeshElementsThreadTask : public FRenderTask
 {
 	FProjectedShadowInfo& ThisShadow;
 	FRHICommandList& RHICmdList;
@@ -1602,11 +1602,6 @@ public:
 		RETURN_QUICK_DECLARE_CYCLE_STAT(FDrawShadowMeshElementsThreadTask, STATGROUP_TaskGraphTasks);
 	}
 
-	ENamedThreads::Type GetDesiredThread()
-	{
-		return ENamedThreads::AnyThread;
-	}
-
 	static ESubsequentsMode::Type GetSubsequentsMode() { return ESubsequentsMode::TrackSubsequents; }
 
 	void DoTask(ENamedThreads::Type CurrentThread, const FGraphEventRef& MyCompletionGraphEvent)
@@ -1627,7 +1622,7 @@ public:
 	}
 };
 
-class FRenderDepthDynamicThreadTask
+class FRenderDepthDynamicThreadTask : public FRenderTask
 {
 	FProjectedShadowInfo& ThisShadow;
 	FRHICommandList& RHICmdList;
@@ -1654,11 +1649,6 @@ public:
 		RETURN_QUICK_DECLARE_CYCLE_STAT(FRenderDepthDynamicThreadTask, STATGROUP_TaskGraphTasks);
 	}
 
-	ENamedThreads::Type GetDesiredThread()
-	{
-		return ENamedThreads::AnyThread;
-	}
-
 	static ESubsequentsMode::Type GetSubsequentsMode() { return ESubsequentsMode::TrackSubsequents; }
 
 	void DoTask(ENamedThreads::Type CurrentThread, const FGraphEventRef& MyCompletionGraphEvent)
@@ -1771,6 +1761,8 @@ static TAutoConsoleVariable<int32> CVarRHICmdFlushRenderThreadTasksShadowPass(
 	0,
 	TEXT("Wait for completion of parallel render thread tasks at the end of each shadow pass.  A more granular version of r.RHICmdFlushRenderThreadTasks. If either r.RHICmdFlushRenderThreadTasks or r.RHICmdFlushRenderThreadTasksShadowPass is > 0 we will flush."));
 
+DECLARE_CYCLE_STAT(TEXT("Shadow"), STAT_CLP_Shadow, STATGROUP_ParallelCommandListMarkers);
+
 class FShadowParallelCommandListSet : public FParallelCommandListSet
 {
 	FProjectedShadowInfo& ProjectedShadowInfo;
@@ -1779,7 +1771,7 @@ class FShadowParallelCommandListSet : public FParallelCommandListSet
 
 public:
 	FShadowParallelCommandListSet(const FViewInfo& InView, FRHICommandListImmediate& InParentCmdList, bool bInParallelExecute, bool bInCreateSceneContext, FProjectedShadowInfo& InProjectedShadowInfo, TFunctionRef<void (FRHICommandList& RHICmdList)> InSetShadowRenderTargets, EShadowDepthRenderMode RenderModeIn )
-		: FParallelCommandListSet(InView, InParentCmdList, bInParallelExecute, bInCreateSceneContext)
+		: FParallelCommandListSet(GET_STATID(STAT_CLP_Shadow), InView, InParentCmdList, bInParallelExecute, bInCreateSceneContext)
 		, ProjectedShadowInfo(InProjectedShadowInfo)
 		, SetShadowRenderTargets(InSetShadowRenderTargets)
 		, RenderMode(RenderModeIn)
diff --git a/Engine/Source/Runtime/Renderer/Private/ShadowSetup.cpp b/Engine/Source/Runtime/Renderer/Private/ShadowSetup.cpp
--- a/Engine/Source/Runtime/Renderer/Private/ShadowSetup.cpp
+++ b/Engine/Source/Runtime/Renderer/Private/ShadowSetup.cpp
@@ -442,7 +442,8 @@ bool FProjectedShadowInfo::SetupPerObjectProjection(
 		}
 		else
 		{
-			ProjectedBoundsPoints.Add(FVector(FLT_MAX, FLT_MAX, FLT_MAX));
+			//ProjectedBoundsPoints.Add(FVector(FLT_MAX, FLT_MAX, FLT_MAX));
+			return false;
 		}
 	}
 
@@ -1589,13 +1590,15 @@ void FSceneRenderer::CreatePerObjectProjectedShadow(
 				// Try to reuse a preshadow from the cache
 				TRefCountPtr<FProjectedShadowInfo> ProjectedPreShadowInfo = GetCachedPreshadow(Interaction, ShadowInitializer, OriginalBounds, PreshadowSizeX);
 
+				bool bOk = true;
+
 				if(!ProjectedPreShadowInfo)
 				{
 					// Create a new projected shadow for this interaction's preshadow
 					// Not using the scene rendering mem stack because this shadow info may need to persist for multiple frames if it gets cached
 					ProjectedPreShadowInfo = new FProjectedShadowInfo;
 
-					ProjectedPreShadowInfo->SetupPerObjectProjection(
+					bOk = ProjectedPreShadowInfo->SetupPerObjectProjection(
 						LightSceneInfo,
 						PrimitiveSceneInfo,
 						ShadowInitializer,
@@ -1606,23 +1609,28 @@ void FSceneRenderer::CreatePerObjectProjectedShadow(
 						false				// not translucent shadow
 						);
 				}
-				// Update fade alpha on the cached preshadow
-				ProjectedPreShadowInfo->FadeAlphas = ResolutionPreShadowFadeAlphas;
-
-				VisibleLightInfo.AllProjectedShadows.Add(ProjectedPreShadowInfo);
-				VisibleLightInfo.ProjectedPreShadows.Add(ProjectedPreShadowInfo);
 
-				// Only add to OutPreShadows if the preshadow doesn't already have depths cached, 
-				// Since OutPreShadows is used to generate information only used when rendering the shadow depths.
-				if (!ProjectedPreShadowInfo->bDepthsCached && ProjectedPreShadowInfo->CasterFrustum.PermutedPlanes.Num())
+				if (bOk)
 				{
-					OutPreShadows.Add(ProjectedPreShadowInfo);
-				}
 
-				for (int32 ChildIndex = 0; ChildIndex < ShadowGroupPrimitives.Num(); ChildIndex++)
-				{
-					FPrimitiveSceneInfo* ShadowChild = ShadowGroupPrimitives[ChildIndex];
-					ProjectedPreShadowInfo->AddReceiverPrimitive(ShadowChild);
+					// Update fade alpha on the cached preshadow
+					ProjectedPreShadowInfo->FadeAlphas = ResolutionPreShadowFadeAlphas;
+
+					VisibleLightInfo.AllProjectedShadows.Add(ProjectedPreShadowInfo);
+					VisibleLightInfo.ProjectedPreShadows.Add(ProjectedPreShadowInfo);
+
+					// Only add to OutPreShadows if the preshadow doesn't already have depths cached, 
+					// Since OutPreShadows is used to generate information only used when rendering the shadow depths.
+					if (!ProjectedPreShadowInfo->bDepthsCached && ProjectedPreShadowInfo->CasterFrustum.PermutedPlanes.Num())
+					{
+						OutPreShadows.Add(ProjectedPreShadowInfo);
+					}
+
+					for (int32 ChildIndex = 0; ChildIndex < ShadowGroupPrimitives.Num(); ChildIndex++)
+					{
+						FPrimitiveSceneInfo* ShadowChild = ShadowGroupPrimitives[ChildIndex];
+						ProjectedPreShadowInfo->AddReceiverPrimitive(ShadowChild);
+					}
 				}
 			}
 		}
diff --git a/Engine/Source/Runtime/Renderer/Private/StaticMeshDrawList.h b/Engine/Source/Runtime/Renderer/Private/StaticMeshDrawList.h
--- a/Engine/Source/Runtime/Renderer/Private/StaticMeshDrawList.h
+++ b/Engine/Source/Runtime/Renderer/Private/StaticMeshDrawList.h
@@ -287,7 +287,7 @@ private:
 	* @param bDrawnShared - determines whether to draw shared 
 	*/
 	template<InstancedStereoPolicy InstancedStereo>
-	void DrawElement(FRHICommandList& RHICmdList, const FViewInfo& View, const typename DrawingPolicyType::ContextDataType PolicyContext, const FElement& Element, uint64 BatchElementMask, FDrawingPolicyLink* DrawingPolicyLink, bool &bDrawnShared);
+	int32 DrawElement(FRHICommandList& RHICmdList, const FViewInfo& View, const typename DrawingPolicyType::ContextDataType PolicyContext, const FElement& Element, uint64 BatchElementMask, FDrawingPolicyLink* DrawingPolicyLink, bool &bDrawnShared);
 
 public:
 
diff --git a/Engine/Source/Runtime/Renderer/Private/StaticMeshDrawList.inl b/Engine/Source/Runtime/Renderer/Private/StaticMeshDrawList.inl
--- a/Engine/Source/Runtime/Renderer/Private/StaticMeshDrawList.inl
+++ b/Engine/Source/Runtime/Renderer/Private/StaticMeshDrawList.inl
@@ -59,7 +59,7 @@ void TStaticMeshDrawList<DrawingPolicyType>::FElementHandle::Remove()
 
 template<typename DrawingPolicyType>
 template<InstancedStereoPolicy InstancedStereo>
-void TStaticMeshDrawList<DrawingPolicyType>::DrawElement(
+int32 TStaticMeshDrawList<DrawingPolicyType>::DrawElement(
 	FRHICommandList& RHICmdList,
 	const FViewInfo& View,
 	const typename DrawingPolicyType::ContextDataType PolicyContext,
@@ -91,6 +91,8 @@ void TStaticMeshDrawList<DrawingPolicyType>::DrawElement(
 	uint32 BackFaceEnd = DrawingPolicyLink->DrawingPolicy.NeedsBackfacePass() ? 2 : 1;
 	const FMeshDrawingRenderState DrawRenderState(View.GetDitheredLODTransitionState(*Element.Mesh, View.bAllowStencilDither));
 
+	int32 DrawCount = 0;
+
 	int32 BatchElementIndex = 0;
 	do
 	{
@@ -99,16 +101,16 @@ void TStaticMeshDrawList<DrawingPolicyType>::DrawElement(
 			const FPrimitiveSceneProxy* Proxy = Element.Mesh->PrimitiveSceneInfo->Proxy;			
 			if (InstancedStereo == InstancedStereoPolicy::Enabled)
 			{
-				// We draw instanced static meshes twice when rendering with instanced stereo. Once for each eye.
-				const bool bIsInstancedMesh = Element.Mesh->Elements[BatchElementIndex].bIsInstancedMesh;
+			    // We draw instanced static meshes twice when rendering with instanced stereo. Once for each eye.
+			    const bool bIsInstancedMesh = Element.Mesh->Elements[BatchElementIndex].bIsInstancedMesh;
 				const uint32 InstancedStereoDrawCount = bIsInstancedMesh ? 2 : 1;
-				for (uint32 DrawCountIter = 0; DrawCountIter < InstancedStereoDrawCount; ++DrawCountIter)
-				{
+			    for (uint32 DrawCountIter = 0; DrawCountIter < InstancedStereoDrawCount; ++DrawCountIter)
+			    {
 					DrawingPolicyLink->DrawingPolicy.SetInstancedEyeIndex(RHICmdList, DrawCountIter);
 
 					for (uint32 BackFace = 0; BackFace < BackFaceEnd; ++BackFace)
 					{
-						INC_DWORD_STAT(STAT_StaticDrawListMeshDrawCalls);
+						DrawCount++;
 
 						TDrawEvent<FRHICommandList> MeshEvent;
 						BeginMeshDrawEvent(RHICmdList, Proxy, *Element.Mesh, MeshEvent);
@@ -130,25 +132,25 @@ void TStaticMeshDrawList<DrawingPolicyType>::DrawElement(
 			}
 			else
 			{
-				for (uint32 BackFace = 0; BackFace < BackFaceEnd; ++BackFace)
-				{
-					INC_DWORD_STAT(STAT_StaticDrawListMeshDrawCalls);
-
-					TDrawEvent<FRHICommandList> MeshEvent;
-					BeginMeshDrawEvent(RHICmdList, Proxy, *Element.Mesh, MeshEvent);
-
-					DrawingPolicyLink->DrawingPolicy.SetMeshRenderState(
-						RHICmdList,
-						View,
-						Proxy,
-						*Element.Mesh,
-						BatchElementIndex,
-						!!BackFace,
-						DrawRenderState,
-						Element.PolicyData,
-						PolicyContext
-						);
-					DrawingPolicyLink->DrawingPolicy.DrawMesh(RHICmdList, *Element.Mesh, BatchElementIndex, false);
+			    for (uint32 BackFace = 0; BackFace < BackFaceEnd; ++BackFace)
+			    {
+				    DrawCount++;
+    
+				    TDrawEvent<FRHICommandList> MeshEvent;
+				    BeginMeshDrawEvent(RHICmdList, Proxy, *Element.Mesh, MeshEvent);
+    
+				    DrawingPolicyLink->DrawingPolicy.SetMeshRenderState(
+					    RHICmdList, 
+					    View,
+					    Proxy,
+					    *Element.Mesh,
+					    BatchElementIndex,
+					    !!BackFace,
+					    DrawRenderState,
+					    Element.PolicyData,
+					    PolicyContext
+					    );
+					    DrawingPolicyLink->DrawingPolicy.DrawMesh(RHICmdList, *Element.Mesh, BatchElementIndex, false);
 				}
 			}
 		}
@@ -156,6 +158,8 @@ void TStaticMeshDrawList<DrawingPolicyType>::DrawElement(
 		BatchElementMask >>= 1;
 		BatchElementIndex++;
 	} while(BatchElementMask);
+	INC_DWORD_STAT_BY(STAT_StaticDrawListMeshDrawCalls, DrawCount);
+	return DrawCount;
 }
 
 template<typename DrawingPolicyType>
@@ -327,10 +331,9 @@ bool TStaticMeshDrawList<DrawingPolicyType>::DrawVisibleInner(
 					const FElement& Element = DrawingPolicyLink->Elements[ElementIndex];
 					STAT(StatInc += Element.Mesh->GetNumPrimitives();)
 					int32 SubCount = Element.Mesh->Elements.Num();
-					Count += SubCount;
 					// Avoid the cache miss looking up batch visibility if there is only one element.
 					uint64 BatchElementMask = SubCount == 1 ? 1 : (*BatchVisibilityArray)[Element.Mesh->Id];
-					DrawElement<InstancedStereoPolicy::Disabled>(RHICmdList, View, PolicyContext, Element, BatchElementMask, DrawingPolicyLink, bDrawnShared);
+					Count += DrawElement<InstancedStereoPolicy::Disabled>(RHICmdList, View, PolicyContext, Element, BatchElementMask, DrawingPolicyLink, bDrawnShared);
 				}
 			}
 
@@ -350,13 +353,12 @@ bool TStaticMeshDrawList<DrawingPolicyType>::DrawVisibleInner(
 
 				if (ResolvedVisiblityArray != nullptr)
 				{
-					const FElement& Element = DrawingPolicyLink->Elements[ElementIndex];
-					STAT(StatInc += Element.Mesh->GetNumPrimitives();)
-					int32 SubCount = Element.Mesh->Elements.Num();
-					Count += SubCount;
-					// Avoid the cache miss looking up batch visibility if there is only one element.
+				    const FElement& Element = DrawingPolicyLink->Elements[ElementIndex];
+					    STAT(StatInc += Element.Mesh->GetNumPrimitives();)
+				    int32 SubCount = Element.Mesh->Elements.Num();
+				    // Avoid the cache miss looking up batch visibility if there is only one element.
 					uint64 BatchElementMask = SubCount == 1 ? 1 : (*ResolvedVisiblityArray)[Element.Mesh->Id];
-					DrawElement<InstancedStereoPolicy::Enabled>(RHICmdList, View, PolicyContext, Element, BatchElementMask, DrawingPolicyLink, bDrawnShared);
+					Count += DrawElement<InstancedStereoPolicy::Enabled>(RHICmdList, View, PolicyContext, Element, BatchElementMask, DrawingPolicyLink, bDrawnShared);
 				}
 			}
 		}
@@ -383,7 +385,7 @@ bool TStaticMeshDrawList<DrawingPolicyType>::DrawVisible(
 }
 
 template<typename DrawingPolicyType>
-class FDrawVisibleAnyThreadTask
+class FDrawVisibleAnyThreadTask : public FRenderTask
 {
 	TStaticMeshDrawList<DrawingPolicyType>& Caller;
 	FRHICommandList& RHICmdList;
@@ -431,15 +433,11 @@ public:
 		RETURN_QUICK_DECLARE_CYCLE_STAT(FDrawVisibleAnyThreadTask, STATGROUP_TaskGraphTasks);
 	}
 
-	ENamedThreads::Type GetDesiredThread()
-	{
-		return ENamedThreads::AnyThread;
-	}
-
 	static ESubsequentsMode::Type GetSubsequentsMode() { return ESubsequentsMode::TrackSubsequents; }
 
 	void DoTask(ENamedThreads::Type CurrentThread, const FGraphEventRef& MyCompletionGraphEvent)
 	{
+		FScopeCycleCounter ScopeOuter(RHICmdList.ExecuteStat);
 		if (this->PerDrawingPolicyCounts)
 		{
 			int32 Start = this->FirstPolicy;
@@ -485,6 +483,24 @@ public:
 	}
 };
 
+ 
+extern MS_ALIGN(64) uint8 CountBitsTable[64] GCC_ALIGN(64);
+
+static FORCEINLINE int32 CountBits(uint64 Bits)
+{
+	int32 Count = 0;
+	while (true)
+	{
+		Count += CountBitsTable[Bits & 63];
+		if (!(Bits & ~63))
+		{
+			break;
+		}
+		Bits >>= 6;
+	}
+	return Count;
+}
+
 template<typename DrawingPolicyType>
 void TStaticMeshDrawList<DrawingPolicyType>::DrawVisibleParallelInternal(
 	const typename DrawingPolicyType::ContextDataType PolicyContext,
@@ -563,7 +579,9 @@ void TStaticMeshDrawList<DrawingPolicyType>::DrawVisibleParallelInternal(
 							if (bIsVisible)
 							{
 								const FElement& Element = DrawingPolicyLink->Elements[ElementIndex];
-								Count += Element.Mesh->Elements.Num();
+								int32 SubCount = Element.Mesh->Elements.Num();
+								// Avoid the cache miss looking up batch visibility if there is only one element.
+								Count += (SubCount == 1) ? 1 : CountBits((*BatchVisibilityArray)[Element.Mesh->Id]);
 							}
 						}
 						if (Count)
diff --git a/Engine/Source/Runtime/Renderer/Private/TranslucentRendering.cpp b/Engine/Source/Runtime/Renderer/Private/TranslucentRendering.cpp
--- a/Engine/Source/Runtime/Renderer/Private/TranslucentRendering.cpp
+++ b/Engine/Source/Runtime/Renderer/Private/TranslucentRendering.cpp
@@ -811,7 +811,7 @@ void FTranslucentPrimSet::PlaceScenePrimitive(FPrimitiveSceneInfo* PrimitiveScen
 
 	if (bUseSeparateTranslucency
 		&& ViewInfo.Family->EngineShowFlags.SeparateTranslucency
-		&& ViewInfo.Family->EngineShowFlags.PostProcessing
+		&& (ViewInfo.Family->EngineShowFlags.PostProcessing || ViewInfo.Family->EngineShowFlags.ShaderComplexity)
 		&& FeatureLevel >= ERHIFeatureLevel::SM4)
 	{
 		// add to list of translucent prims that use scene color
@@ -865,7 +865,7 @@ bool FSceneRenderer::ShouldRenderTranslucency() const
 	return bRender;
 }
 
-class FDrawSortedTransAnyThreadTask
+class FDrawSortedTransAnyThreadTask : public FRenderTask
 {
 	FDeferredShadingSceneRenderer& Renderer;
 	FRHICommandList& RHICmdList;
@@ -899,27 +899,26 @@ public:
 		RETURN_QUICK_DECLARE_CYCLE_STAT(FDrawSortedTransAnyThreadTask, STATGROUP_TaskGraphTasks);
 	}
 
-	ENamedThreads::Type GetDesiredThread()
-	{
-		return ENamedThreads::AnyThread;
-	}
-
 	static ESubsequentsMode::Type GetSubsequentsMode() { return ESubsequentsMode::TrackSubsequents; }
 
 	void DoTask(ENamedThreads::Type CurrentThread, const FGraphEventRef& MyCompletionGraphEvent)
 	{
+		FScopeCycleCounter ScopeOuter(RHICmdList.ExecuteStat);
 		View.TranslucentPrimSet.DrawPrimitivesParallel(RHICmdList, View, Renderer, TranslucenyPassType, FirstIndex, LastIndex);
 		RHICmdList.HandleRTThreadTaskCompletion(MyCompletionGraphEvent);
 	}
 };
 
+DECLARE_CYCLE_STAT(TEXT("Translucency"), STAT_CLP_Translucency, STATGROUP_ParallelCommandListMarkers);
+
+
 class FTranslucencyPassParallelCommandListSet : public FParallelCommandListSet
 {
 	ETranslucencyPassType TranslucenyPassType;
 	bool bFirstTimeThisFrame;
 public:
 	FTranslucencyPassParallelCommandListSet(const FViewInfo& InView, FRHICommandListImmediate& InParentCmdList, bool bInParallelExecute, bool bInCreateSceneContext, ETranslucencyPassType InTranslucenyPassType)
-		: FParallelCommandListSet(InView, InParentCmdList, bInParallelExecute, bInCreateSceneContext)
+		: FParallelCommandListSet(GET_STATID(STAT_CLP_Translucency), InView, InParentCmdList, bInParallelExecute, bInCreateSceneContext)
 		, TranslucenyPassType(InTranslucenyPassType)
 		, bFirstTimeThisFrame(InTranslucenyPassType == TPT_SeparateTransluceny)
 	{
diff --git a/Engine/Source/Runtime/Renderer/Private/VelocityRendering.cpp b/Engine/Source/Runtime/Renderer/Private/VelocityRendering.cpp
--- a/Engine/Source/Runtime/Renderer/Private/VelocityRendering.cpp
+++ b/Engine/Source/Runtime/Renderer/Private/VelocityRendering.cpp
@@ -554,7 +554,7 @@ void FDeferredShadingSceneRenderer::RenderDynamicVelocitiesMeshElementsInner(FRH
 	}
 }
 
-class FRenderVelocityDynamicThreadTask
+class FRenderVelocityDynamicThreadTask : public FRenderTask
 {
 	FDeferredShadingSceneRenderer& ThisRenderer;
 	FRHICommandList& RHICmdList;
@@ -583,11 +583,6 @@ public:
 		RETURN_QUICK_DECLARE_CYCLE_STAT(FRenderVelocityDynamicThreadTask, STATGROUP_TaskGraphTasks);
 	}
 
-	ENamedThreads::Type GetDesiredThread()
-	{
-		return ENamedThreads::AnyThread;
-	}
-
 	static ESubsequentsMode::Type GetSubsequentsMode() { return ESubsequentsMode::TrackSubsequents; }
 
 	void DoTask(ENamedThreads::Type CurrentThread, const FGraphEventRef& MyCompletionGraphEvent)
@@ -639,12 +634,15 @@ static void SetVelocitiesState(FRHICommandList& RHICmdList, const FViewInfo& Vie
 	RHICmdList.SetRasterizerState(GetStaticRasterizerState<true>(FM_Solid, CM_CW));
 }
 
+DECLARE_CYCLE_STAT(TEXT("Velocity"), STAT_CLP_Velocity, STATGROUP_ParallelCommandListMarkers);
+
+
 class FVelocityPassParallelCommandListSet : public FParallelCommandListSet
 {
 	TRefCountPtr<IPooledRenderTarget>& VelocityRT;
 public:
 	FVelocityPassParallelCommandListSet(const FViewInfo& InView, FRHICommandListImmediate& InParentCmdList, bool bInParallelExecute, bool bInCreateSceneContext, TRefCountPtr<IPooledRenderTarget>& InVelocityRT)
-		: FParallelCommandListSet(InView, InParentCmdList, bInParallelExecute, bInCreateSceneContext)
+		: FParallelCommandListSet(GET_STATID(STAT_CLP_Velocity), InView, InParentCmdList, bInParallelExecute, bInCreateSceneContext)
 		, VelocityRT(InVelocityRT)
 	{
 		SetStateOnCommandList(ParentCmdList);
diff --git a/Engine/Source/Runtime/ShaderCore/Private/VertexFactory.cpp b/Engine/Source/Runtime/ShaderCore/Private/VertexFactory.cpp
--- a/Engine/Source/Runtime/ShaderCore/Private/VertexFactory.cpp
+++ b/Engine/Source/Runtime/ShaderCore/Private/VertexFactory.cpp
@@ -291,13 +291,8 @@ FVertexElement FVertexFactory::AccessPositionStreamComponent(const FVertexStream
 	return FVertexElement(PositionStream.AddUnique(VertexStream),Component.Offset,Component.Type,AttributeIndex,VertexStream.Stride,Component.bUseInstanceIndex);
 }
 
-void FVertexFactory::InitDeclaration(
-	FVertexDeclarationElementList& Elements, 
-	const DataType& InData)
+void FVertexFactory::InitDeclaration(FVertexDeclarationElementList& Elements)
 {
-	// Make a copy of the vertex factory data.
-	Data = InData;
-
 	// Create the vertex declaration for rendering the factory normally.
 	Declaration = RHICreateVertexDeclaration(Elements);
 }
diff --git a/Engine/Source/Runtime/ShaderCore/Public/VertexFactory.h b/Engine/Source/Runtime/ShaderCore/Public/VertexFactory.h
--- a/Engine/Source/Runtime/ShaderCore/Public/VertexFactory.h
+++ b/Engine/Source/Runtime/ShaderCore/Public/VertexFactory.h
@@ -335,13 +335,6 @@ public:
 class SHADERCORE_API FVertexFactory : public FRenderResource
 {
 public:
-
-	struct DataType
-	{
-		/** Initialization constructor. */
-		DataType() {}
-	};
-
 	/**
 	 * Default constructor 
 	 */
@@ -447,9 +440,7 @@ protected:
 	 * Initializes the vertex declaration.
 	 * @param Elements - The elements of the vertex declaration.
 	 */
-	void InitDeclaration(
-		FVertexDeclarationElementList& Elements, 
-		const DataType& Data);
+	void InitDeclaration(FVertexDeclarationElementList& Elements);
 
 	/**
 	 * Initializes the position-only vertex declaration.
@@ -496,9 +487,6 @@ private:
 
 	/** The RHI vertex declaration used to render the factory during depth only passes. */
 	FVertexDeclarationRHIRef PositionDeclaration;
-
-	/** The vertex factory's data. */
-	DataType Data;
 };
 
 /**
diff --git a/Engine/Source/Runtime/SlateRHIRenderer/Private/SlateRHIRenderer.cpp b/Engine/Source/Runtime/SlateRHIRenderer/Private/SlateRHIRenderer.cpp
--- a/Engine/Source/Runtime/SlateRHIRenderer/Private/SlateRHIRenderer.cpp
+++ b/Engine/Source/Runtime/SlateRHIRenderer/Private/SlateRHIRenderer.cpp
@@ -1338,5 +1338,13 @@ void FSlateEndDrawingWindowsCommand::Execute(FRHICommandListBase& CmdList)
 
 void FSlateEndDrawingWindowsCommand::EndDrawingWindows(FRHICommandListImmediate& RHICmdList, FSlateDrawBuffer* DrawBuffer, FSlateRHIRenderingPolicy& Policy)
 {
-	new ( RHICmdList.AllocCommand<FSlateEndDrawingWindowsCommand>() ) FSlateEndDrawingWindowsCommand(Policy, DrawBuffer);
+	if (!RHICmdList.Bypass())
+	{
+		new (RHICmdList.AllocCommand<FSlateEndDrawingWindowsCommand>()) FSlateEndDrawingWindowsCommand(Policy, DrawBuffer);
+	}
+	else
+	{
+		FSlateEndDrawingWindowsCommand Cmd(Policy, DrawBuffer);
+		Cmd.Execute(RHICmdList);
+	}
 }
diff --git a/Engine/Source/Runtime/Windows/D3D12RHI/Private/D3D12Device.cpp b/Engine/Source/Runtime/Windows/D3D12RHI/Private/D3D12Device.cpp
--- a/Engine/Source/Runtime/Windows/D3D12RHI/Private/D3D12Device.cpp
+++ b/Engine/Source/Runtime/Windows/D3D12RHI/Private/D3D12Device.cpp
@@ -147,7 +147,7 @@ public:
 	}
 };
 
-static TLockFreeFixedSizeAllocator<sizeof(FD3D12CommandContextContainer), FThreadSafeCounter> FD3D12CommandContextContainerAllocator;
+static TLockFreeFixedSizeAllocator<sizeof(FD3D12CommandContextContainer), PLATFORM_CACHE_LINE_SIZE, FThreadSafeCounter> FD3D12CommandContextContainerAllocator;
 
 void* FD3D12CommandContextContainer::operator new(size_t Size)
 {
