Copying //UE4/Dev-Rendering to Dev-Main (//UE4/Dev-Main)

diff --git a/Engine/Source/Runtime/RHI/Public/RHIUtilities.h b/Engine/Source/Runtime/RHI/Public/RHIUtilities.h
--- a/Engine/Source/Runtime/RHI/Public/RHIUtilities.h
+++ b/Engine/Source/Runtime/RHI/Public/RHIUtilities.h
@@ -120,6 +120,54 @@ struct FRWBufferByteAddress
 	}
 };
 
+struct FDynamicReadBuffer : public FReadBuffer
+{
+	/** Pointer to the vertex buffer mapped in main memory. */
+	uint8* MappedBuffer;
+
+	/** Default constructor. */
+	FDynamicReadBuffer()
+		: MappedBuffer(nullptr)
+	{
+	}
+
+	~FDynamicReadBuffer()
+	{
+		Release();
+	}
+
+	virtual void Initialize(uint32 BytesPerElement, uint32 NumElements, EPixelFormat Format, uint32 AdditionalUsage = 0)
+	{
+		ensure(
+			AdditionalUsage & (BUF_Dynamic | BUF_Volatile) &&								// buffer should be Dynamic or Volatile
+			(AdditionalUsage & (BUF_Dynamic | BUF_Volatile)) ^ (BUF_Dynamic | BUF_Volatile) // buffer should not be both
+			);
+
+		FReadBuffer::Initialize(BytesPerElement, NumElements, Format, AdditionalUsage);
+	}
+
+	/**
+	* Locks the vertex buffer so it may be written to.
+	*/
+	void Lock()
+	{
+		check(MappedBuffer == nullptr);
+		check(IsValidRef(Buffer));
+		MappedBuffer = (uint8*)RHILockVertexBuffer(Buffer, 0, NumBytes, RLM_WriteOnly);
+	}
+
+	/**
+	* Unocks the buffer so the GPU may read from it.
+	*/
+	void Unlock()
+	{
+		check(MappedBuffer);
+		check(IsValidRef(Buffer));
+		RHIUnlockVertexBuffer(Buffer);
+		MappedBuffer = nullptr;
+	}
+};
+
 /**
  * Convert the ESimpleRenderTargetMode into usable values 
  * @todo: Can we easily put this into a .cpp somewhere?
diff --git a/Engine/Source/Runtime/RenderCore/Private/RenderCore.cpp b/Engine/Source/Runtime/RenderCore/Private/RenderCore.cpp
--- a/Engine/Source/Runtime/RenderCore/Private/RenderCore.cpp
+++ b/Engine/Source/Runtime/RenderCore/Private/RenderCore.cpp
@@ -41,6 +41,7 @@ DEFINE_STAT(STAT_DynamicShadowSetupTime);
 DEFINE_STAT(STAT_RenderQueryResultTime);
 // Use 'stat initviews' to get more detail
 DEFINE_STAT(STAT_InitViewsTime);
+DEFINE_STAT(STAT_InitViewsPossiblyAfterPrepass);
 // Measures the time spent in RenderViewFamily_RenderThread
 // Note that this is not the total rendering thread time, any other rendering commands will not be counted here
 DEFINE_STAT(STAT_TotalSceneRenderingTime);
diff --git a/Engine/Source/Runtime/RenderCore/Private/RenderingThread.cpp b/Engine/Source/Runtime/RenderCore/Private/RenderingThread.cpp
--- a/Engine/Source/Runtime/RenderCore/Private/RenderingThread.cpp
+++ b/Engine/Source/Runtime/RenderCore/Private/RenderingThread.cpp
@@ -274,7 +274,7 @@ public:
 
 	void Start()
 	{
-		Thread = FRunnableThread::Create(this, TEXT("RHIThread"), 512 * 1024, TPri_Normal, 
+		Thread = FRunnableThread::Create(this, TEXT("RHIThread"), 512 * 1024, TPri_SlightlyBelowNormal, 
 			FPlatformAffinity::GetRHIThreadMask()
 			);
 		check(Thread);
diff --git a/Engine/Source/Runtime/RenderCore/Public/RenderCore.h b/Engine/Source/Runtime/RenderCore/Public/RenderCore.h
--- a/Engine/Source/Runtime/RenderCore/Public/RenderCore.h
+++ b/Engine/Source/Runtime/RenderCore/Public/RenderCore.h
@@ -16,7 +16,8 @@ DECLARE_LOG_CATEGORY_EXTERN(LogRendererCore, Log, All);
 DECLARE_CYCLE_STAT_EXTERN(TEXT("BeginOcclusionTests"),STAT_BeginOcclusionTestsTime,STATGROUP_SceneRendering, RENDERCORE_API);
 DECLARE_CYCLE_STAT_EXTERN(TEXT("RenderQuery Result"),STAT_RenderQueryResultTime,STATGROUP_SceneRendering, RENDERCORE_API);
 DECLARE_CYCLE_STAT_EXTERN(TEXT("InitViews"),STAT_InitViewsTime,STATGROUP_SceneRendering, RENDERCORE_API);
-DECLARE_CYCLE_STAT_EXTERN(TEXT("Dynamic shadow setup"),STAT_DynamicShadowSetupTime,STATGROUP_SceneRendering, RENDERCORE_API);
+DECLARE_CYCLE_STAT_EXTERN(TEXT("InitViewsPossiblyAfterPrepass"), STAT_InitViewsPossiblyAfterPrepass, STATGROUP_SceneRendering, RENDERCORE_API);
+DECLARE_CYCLE_STAT_EXTERN(TEXT("Dynamic shadow setup"), STAT_DynamicShadowSetupTime, STATGROUP_SceneRendering, RENDERCORE_API);
 DECLARE_CYCLE_STAT_EXTERN(TEXT("Total GPU rendering"),STAT_TotalGPUFrameTime,STATGROUP_SceneRendering, RENDERCORE_API);
 DECLARE_CYCLE_STAT_EXTERN(TEXT("RenderViewFamily"),STAT_TotalSceneRenderingTime,STATGROUP_SceneRendering, RENDERCORE_API);
 DECLARE_CYCLE_STAT_EXTERN(TEXT("FinishRenderViewTarget"),STAT_FinishRenderViewTargetTime,STATGROUP_SceneRendering, RENDERCORE_API);
diff --git a/Engine/Source/Runtime/Renderer/Private/AtmosphereRendering.cpp b/Engine/Source/Runtime/Renderer/Private/AtmosphereRendering.cpp
--- a/Engine/Source/Runtime/Renderer/Private/AtmosphereRendering.cpp
+++ b/Engine/Source/Runtime/Renderer/Private/AtmosphereRendering.cpp
@@ -308,15 +308,14 @@ IMPLEMENT_SHADER_TYPE(,FAtmosphericVS,TEXT("AtmosphericFogShader"),TEXT("VSMain"
 /** Vertex declaration for the light function fullscreen 2D quad. */
 TGlobalResource<FAtmopshereVertexDeclaration> GAtmophereVertexDeclaration;
 
-void FSceneRenderer::InitAtmosphereConstants()
+void InitAtmosphereConstantsInView(FViewInfo& View)
 {
-	for(int32 ViewIndex = 0;ViewIndex < Views.Num();ViewIndex++)
+	bool bInitTextures = false;
+	if(ShouldRenderAtmosphere(*View.Family))
 	{
-		bool bInitTextures = false;
-		FViewInfo& View = Views[ViewIndex];
-		// set fog consts based on height fog components
-		if(ShouldRenderAtmosphere(*View.Family))
+		if (View.Family->Scene)
 		{
+			FScene* Scene = (FScene*)View.Family->Scene;
 			if (Scene->AtmosphericFog)
 			{
 				const FAtmosphericFogSceneInfo& FogInfo = *Scene->AtmosphericFog;
@@ -327,13 +326,13 @@ void FSceneRenderer::InitAtmosphereConstants()
 				bInitTextures = true;
 			}
 		}
+	}
 
-		if (!bInitTextures)
-		{
-			View.AtmosphereTransmittanceTexture = GBlackTexture->TextureRHI;
-			View.AtmosphereIrradianceTexture = GBlackTexture->TextureRHI;
-			View.AtmosphereInscatterTexture = GBlackVolumeTexture->TextureRHI;
-		}
+	if (!bInitTextures)
+	{
+		View.AtmosphereTransmittanceTexture = GBlackTexture->TextureRHI;
+		View.AtmosphereIrradianceTexture = GBlackTexture->TextureRHI;
+		View.AtmosphereInscatterTexture = GBlackVolumeTexture->TextureRHI;
 	}
 }
 
diff --git a/Engine/Source/Runtime/Renderer/Private/AtmosphereRendering.h b/Engine/Source/Runtime/Renderer/Private/AtmosphereRendering.h
--- a/Engine/Source/Runtime/Renderer/Private/AtmosphereRendering.h
+++ b/Engine/Source/Runtime/Renderer/Private/AtmosphereRendering.h
@@ -86,4 +86,5 @@ private:
 #endif
 };
 
-extern bool ShouldRenderAtmosphere(const FSceneViewFamily& Family);
\ No newline at end of file
+bool ShouldRenderAtmosphere(const FSceneViewFamily& Family);
+void InitAtmosphereConstantsInView(FViewInfo& View);
diff --git a/Engine/Source/Runtime/Renderer/Private/AtmosphereTextures.cpp b/Engine/Source/Runtime/Renderer/Private/AtmosphereTextures.cpp
--- a/Engine/Source/Runtime/Renderer/Private/AtmosphereTextures.cpp
+++ b/Engine/Source/Runtime/Renderer/Private/AtmosphereTextures.cpp
@@ -78,19 +78,6 @@ void FAtmosphereShaderTextureParameters::Bind(const FShaderParameterMap& Paramet
 	InscatterTextureSampler.Bind(ParameterMap,TEXT("AtmosphereInscatterTextureSampler"));
 }
 
-template< typename ShaderRHIParamRef >
-void FAtmosphereShaderTextureParameters::Set( FRHICommandList& RHICmdList, const ShaderRHIParamRef ShaderRHI, const FSceneView& View ) const
-{
-	if (TransmittanceTexture.IsBound() || IrradianceTexture.IsBound() || InscatterTexture.IsBound())
-	{
-		SetTextureParameter(RHICmdList, ShaderRHI, TransmittanceTexture, TransmittanceTextureSampler, 
-			TStaticSamplerState<SF_Bilinear>::GetRHI(), View.AtmosphereTransmittanceTexture);
-		SetTextureParameter(RHICmdList, ShaderRHI, IrradianceTexture, IrradianceTextureSampler, 
-			TStaticSamplerState<SF_Bilinear>::GetRHI(), View.AtmosphereIrradianceTexture);
-		SetTextureParameter(RHICmdList, ShaderRHI, InscatterTexture, InscatterTextureSampler, 
-			TStaticSamplerState<SF_Bilinear>::GetRHI(), View.AtmosphereInscatterTexture);
-	}
-}
 
 #define IMPLEMENT_ATMOSPHERE_TEXTURE_PARAM_SET( ShaderRHIParamRef ) \
 	template void FAtmosphereShaderTextureParameters::Set< ShaderRHIParamRef >( FRHICommandList& RHICmdList, const ShaderRHIParamRef ShaderRHI, const FSceneView& View ) const;
diff --git a/Engine/Source/Runtime/Renderer/Private/BasePassRendering.h b/Engine/Source/Runtime/Renderer/Private/BasePassRendering.h
--- a/Engine/Source/Runtime/Renderer/Private/BasePassRendering.h
+++ b/Engine/Source/Runtime/Renderer/Private/BasePassRendering.h
@@ -118,7 +118,6 @@ protected:
 	{
 		VertexParametersType::Bind(Initializer.ParameterMap);
 		HeightFogParameters.Bind(Initializer.ParameterMap);
-		AtmosphericFogTextureParameters.Bind(Initializer.ParameterMap);
 		TranslucentLightingVolumeParameters.Bind(Initializer.ParameterMap);
 		const bool bOutputsVelocityToGBuffer = FVelocityRendering::OutputsToGBuffer();
 		if (bOutputsVelocityToGBuffer)
@@ -144,7 +143,6 @@ public:
 		bool bShaderHasOutdatedParameters = FMeshMaterialShader::Serialize(Ar);
 		VertexParametersType::Serialize(Ar);
 		Ar << HeightFogParameters;
-		Ar << AtmosphericFogTextureParameters;
 		Ar << TranslucentLightingVolumeParameters;
 		Ar << PreviousLocalToWorldParameter;
 		Ar << SkipOutputVelocityParameter;
@@ -169,7 +167,6 @@ public:
 		if (bAllowGlobalFog)
 		{
 			HeightFogParameters.Set(RHICmdList, GetVertexShader(), &View);
-			AtmosphericFogTextureParameters.Set(RHICmdList, GetVertexShader(), View);
 		}
 
 		TranslucentLightingVolumeParameters.Set(RHICmdList, GetVertexShader());
@@ -188,7 +185,6 @@ private:
 	
 	/** The parameters needed to calculate the fog contribution from height fog layers. */
 	FHeightFogShaderParameters HeightFogParameters;
-	FAtmosphereShaderTextureParameters AtmosphericFogTextureParameters;
 	FTranslucentLightingVolumeParameters TranslucentLightingVolumeParameters;
 	// When outputting from base pass, the previous transform
 	FShaderParameter PreviousLocalToWorldParameter;
@@ -819,7 +815,7 @@ public:
 				RHICmdList.SetBlendState( TStaticBlendState<CW_RGBA, BO_Add, BF_One, BF_One, BO_Add, BF_Zero, BF_One>::GetRHI());
 			}
 
-			FDebugViewMode::GetPixelShader(View->ShaderMap, GetDebugViewShaderMode())->SetParameters(RHICmdList, VertexShader, PixelShader, *View);
+			FDebugViewMode::GetPSInterface(View->ShaderMap, MaterialResource, GetDebugViewShaderMode())->SetParameters(RHICmdList, VertexShader, PixelShader, MaterialRenderProxy, *MaterialResource, *View);
 		}
 		else
 #endif
@@ -946,7 +942,7 @@ public:
 				RHICmdList.SetBlendState(TStaticBlendState<CW_RGB,BO_Add,BF_One,BF_One>::GetRHI());
 			}
 
-			FDebugViewMode::GetPixelShader(View.ShaderMap, GetDebugViewShaderMode())->SetMesh(RHICmdList, VertexFactory, View, PrimitiveSceneProxy, BatchElement, DrawRenderState);
+			FDebugViewMode::GetPSInterface(View.ShaderMap, MaterialResource, GetDebugViewShaderMode())->SetMesh(RHICmdList, VertexFactory, View, PrimitiveSceneProxy, BatchElement, DrawRenderState);
 		}
 		else
 #endif
diff --git a/Engine/Source/Runtime/Renderer/Private/CompositionLighting/CompositionLighting.cpp b/Engine/Source/Runtime/Renderer/Private/CompositionLighting/CompositionLighting.cpp
--- a/Engine/Source/Runtime/Renderer/Private/CompositionLighting/CompositionLighting.cpp
+++ b/Engine/Source/Runtime/Renderer/Private/CompositionLighting/CompositionLighting.cpp
@@ -328,18 +328,21 @@ void FCompositionLighting::ProcessAfterBasePass(FRHICommandListImmediate& RHICmd
 
 		FRenderingCompositeOutputRef AmbientOcclusion;
 
-		uint32 Levels = ComputeAmbientOcclusionPassCount(Context.View);
-		if (Levels && !FSSAOHelper::IsAmbientOcclusionAsyncCompute(Context.View, Levels))
+		uint32 SSAOLevels = ComputeAmbientOcclusionPassCount(Context.View);
+		if (SSAOLevels)
 		{
-			AmbientOcclusion = AddPostProcessingAmbientOcclusion(RHICmdList, Context, Levels);
-		}
+			if(!FSSAOHelper::IsAmbientOcclusionAsyncCompute(Context.View, SSAOLevels))
+			{
+				AmbientOcclusion = AddPostProcessingAmbientOcclusion(RHICmdList, Context, SSAOLevels);
+			}
 
-		if (FSSAOHelper::IsBasePassAmbientOcclusionRequired(Context.View))
-		{
-			FRenderingCompositePass* Pass = Context.Graph.RegisterPass(new(FMemStack::Get()) FRCPassPostProcessBasePassAO());
-			Pass->AddDependency(Context.FinalOutput);
+			if (FSSAOHelper::IsBasePassAmbientOcclusionRequired(Context.View))
+			{
+				FRenderingCompositePass* Pass = Context.Graph.RegisterPass(new(FMemStack::Get()) FRCPassPostProcessBasePassAO());
+				Pass->AddDependency(Context.FinalOutput);
 
-			Context.FinalOutput = FRenderingCompositeOutputRef(Pass);
+				Context.FinalOutput = FRenderingCompositeOutputRef(Pass);
+			}
 		}
 
 		if (IsAmbientCubemapPassRequired(Context.View))
diff --git a/Engine/Source/Runtime/Renderer/Private/DebugViewModeRendering.cpp b/Engine/Source/Runtime/Renderer/Private/DebugViewModeRendering.cpp
--- a/Engine/Source/Runtime/Renderer/Private/DebugViewModeRendering.cpp
+++ b/Engine/Source/Runtime/Renderer/Private/DebugViewModeRendering.cpp
@@ -1,4 +1,4 @@
-// Copyright 1998-2015 Epic Games, Inc. All Rights Reserved.
+// Copyright 1998-2016 Epic Games, Inc. All Rights Reserved.
 
 /*=============================================================================
 DebugViewModeRendering.cpp: Contains definitions for rendering debug viewmodes.
@@ -9,11 +9,57 @@ DebugViewModeRendering.cpp: Contains definitions for rendering debug viewmodes.
 #include "SceneFilterRendering.h"
 #include "SceneUtils.h"
 #include "PostProcessing.h"
+#include "WantedMipsAccuracyRendering.h"
+#include "TexelFactorAccuracyRendering.h"
+#include "TexCoordScaleAnalysisRendering.h"
 
 IMPLEMENT_MATERIAL_SHADER_TYPE(,FDebugViewModeVS,TEXT("DebugViewModeVertexShader"),TEXT("Main"),SF_Vertex);	
 IMPLEMENT_MATERIAL_SHADER_TYPE(,FDebugViewModeHS,TEXT("DebugViewModeVertexShader"),TEXT("MainHull"),SF_Hull);	
 IMPLEMENT_MATERIAL_SHADER_TYPE(,FDebugViewModeDS,TEXT("DebugviewModeVertexShader"),TEXT("MainDomain"),SF_Domain);
 
+/**
+* Pixel shader that renders texture streamer wanted mips accuracy.
+*/
+class FMissingShaderPS : public FGlobalShader, public IDebugViewModePSInterface
+{
+	DECLARE_SHADER_TYPE(FMissingShaderPS,Global);
+
+public:
+
+	static bool ShouldCache(EShaderPlatform Platform) { return AllowDebugViewModeShader(Platform); }
+
+	FMissingShaderPS(const ShaderMetaType::CompiledShaderInitializerType& Initializer):	FGlobalShader(Initializer) {}
+	FMissingShaderPS() {}
+
+	virtual bool Serialize(FArchive& Ar) override { return FGlobalShader::Serialize(Ar); }
+
+	static void ModifyCompilationEnvironment(EShaderPlatform Platform, FShaderCompilerEnvironment& OutEnvironment) {}
+
+	virtual void SetParameters(
+		FRHICommandList& RHICmdList, 
+		const FShader* OriginalVS, 
+		const FShader* OriginalPS, 
+		const FMaterialRenderProxy* MaterialRenderProxy,
+		const FMaterial& Material,
+		const FSceneView& View
+		) override {}
+
+	virtual void SetMesh(
+		FRHICommandList& RHICmdList, 
+		const FVertexFactory* VertexFactory,
+		const FSceneView& View,
+		const FPrimitiveSceneProxy* Proxy,
+		const FMeshBatchElement& BatchElement, 
+		const FMeshDrawingRenderState& DrawRenderState
+		) override{}
+
+	virtual void SetMesh(FRHICommandList& RHICmdList, const FSceneView& View) override {}
+
+	virtual FShader* GetShader() override { return static_cast<FShader*>(this); }
+};
+
+IMPLEMENT_SHADER_TYPE(,FMissingShaderPS,TEXT("MissingShaderPixelShader"),TEXT("Main"),SF_Pixel);
+
 
 void FDebugViewMode::GetDebugMaterial(const FMaterialRenderProxy** MaterialRenderProxy, const FMaterial** Material, ERHIFeatureLevel::Type FeatureLevel)
 {
@@ -27,8 +73,15 @@ void FDebugViewMode::GetDebugMaterial(const FMaterialRenderProxy** MaterialRende
 	}
 }
 
-FDebugViewModePS* FDebugViewMode::GetPixelShader(TShaderMap<FGlobalShaderType>* ShaderMap, EDebugViewShaderMode DebugViewShaderMode)
+ENGINE_API const FMaterial* GetDebugViewMaterialPS(EDebugViewShaderMode DebugViewShaderMode, const FMaterial* Material);
+
+IDebugViewModePSInterface* FDebugViewMode::GetPSInterface(TShaderMap<FGlobalShaderType>* ShaderMap, const FMaterial* Material, EDebugViewShaderMode DebugViewShaderMode)
 {
+    // The mesh material shaders are only compiled with the local vertex factory to prevent multiple compilation.
+    // Nothing from the factory is actually used, but the shader must still derive from FMeshMaterialShader.
+    // This is required to call FMeshMaterialShader::SetMesh and bind primitive related data.
+	static FVertexFactoryType* LocalVertexFactory = FindVertexFactoryType(FName(TEXT("FLocalVertexFactory"), FNAME_Find));
+
 	switch (DebugViewShaderMode)
 	{
 	case DVSM_QuadComplexity:
@@ -41,10 +94,20 @@ FDebugViewModePS* FDebugViewMode::GetPixelShader(TShaderMap<FGlobalShaderType>*
 		return *TShaderMapRef<FWantedMipsAccuracyPS>(ShaderMap);
 	case DVSM_TexelFactorAccuracy:
 		return *TShaderMapRef<FTexelFactorAccuracyPS>(ShaderMap);
+	case DVSM_TexCoordScaleAnalysis:
+	{
+		const FMaterial* MaterialForPS = GetDebugViewMaterialPS(DebugViewShaderMode, Material);
+		if (MaterialForPS)
+		{
+			return MaterialForPS->GetShader<FTexCoordScaleAnalysisPS>(LocalVertexFactory);
+		}
+		break;
+	}
 	default:
-		check(false); // Going here will probably generate a null access
-		return nullptr;
+		break;
 	}
+
+	return *TShaderMapRef<FMissingShaderPS>(ShaderMap);
 }
 
 void FDebugViewMode::PatchBoundShaderState(
@@ -55,6 +118,7 @@ void FDebugViewMode::PatchBoundShaderState(
 	EDebugViewShaderMode DebugViewShaderMode
 	)
 {
+	const FMaterial* MaterialForPS = Material;
 	GetDebugMaterial(nullptr, &Material, FeatureLevel);
 
 	if (!Material->HasVertexPositionOffsetConnected() && Material->GetTessellationMode() == MTM_NoTessellation)
@@ -75,7 +139,7 @@ void FDebugViewMode::PatchBoundShaderState(
 		BoundShaderStateInput.DomainShaderRHI = Material->GetShader<FDebugViewModeDS>(VertexFactoryType)->GetDomainShader();
 	}
 
-	BoundShaderStateInput.PixelShaderRHI = GetPixelShader(GetGlobalShaderMap(FeatureLevel), DebugViewShaderMode)->GetPixelShader();
+	BoundShaderStateInput.PixelShaderRHI = GetPSInterface(GetGlobalShaderMap(FeatureLevel), MaterialForPS, DebugViewShaderMode)->GetShader()->GetPixelShader();
 }
 
 void FDebugViewMode::SetParametersVSHSDS(
@@ -124,4 +188,4 @@ void FDebugViewMode::SetMeshVSHSDS(
 		Material->GetShader<FDebugViewModeHS>(VertexFactoryType)->SetMesh(RHICmdList, VertexFactory, View, Proxy, BatchElement, DrawRenderState);
 		Material->GetShader<FDebugViewModeDS>(VertexFactoryType)->SetMesh(RHICmdList, VertexFactory, View, Proxy, BatchElement, DrawRenderState);
 	}
-}
+}
\ No newline at end of file
diff --git a/Engine/Source/Runtime/Renderer/Private/DebugViewModeRendering.h b/Engine/Source/Runtime/Renderer/Private/DebugViewModeRendering.h
--- a/Engine/Source/Runtime/Renderer/Private/DebugViewModeRendering.h
+++ b/Engine/Source/Runtime/Renderer/Private/DebugViewModeRendering.h
@@ -1,7 +1,7 @@
-// Copyright 1998-2015 Epic Games, Inc. All Rights Reserved.
+// Copyright 1998-2016 Epic Games, Inc. All Rights Reserved.
 
 /*=============================================================================
-ShaderComplexityRendering.h: Declarations used for the shader complexity viewmode.
+DebugViewModeRendering.h: Contains definitions for rendering debug viewmodes.
 =============================================================================*/
 
 #pragma once
@@ -19,6 +19,9 @@ FORCEINLINE bool RuntimeAllowDebugViewModeShader(ERHIFeatureLevel::Type InFeatur
 }
 #endif
 
+static const int32 NumStreamingAccuracyColors = 5;
+static const int32 MaxStreamingAccuracyMips = 11;
+
 /**
  * Vertex shader for quad overdraw. Required because overdraw shaders need to have SV_Position as first PS interpolant.
  */
@@ -82,17 +85,19 @@ public:
 	FDebugViewModeDS() {}
 };
 
-class FDebugViewModePS : public FGlobalShader
+// Interface for debug viewmode pixel shaders. Devired classes can be of global shader or material shaders.
+class IDebugViewModePSInterface
 {
 public:
 
-	FDebugViewModePS(const ShaderMetaType::CompiledShaderInitializerType& Initializer) : FGlobalShader(Initializer) {}
-	FDebugViewModePS() {}
+	virtual ~IDebugViewModePSInterface() {}
 
 	virtual void SetParameters(
 		FRHICommandList& RHICmdList, 
 		const FShader* OriginalVS, 
 		const FShader* OriginalPS, 
+		const FMaterialRenderProxy* MaterialRenderProxy,
+		const FMaterial& Material,
 		const FSceneView& View
 		) = 0;
 
@@ -107,11 +112,16 @@ public:
 
 	// Used for custom rendering like decals.
 	virtual void SetMesh(FRHICommandList& RHICmdList, const FSceneView& View) = 0;
+
+	virtual FShader* GetShader() = 0;
 };
 
+/**
+ * Namespace holding the interface for debugview modes.
+ */
 struct FDebugViewMode
 {
-	static FDebugViewModePS* GetPixelShader(TShaderMap<FGlobalShaderType>* ShaderMap, EDebugViewShaderMode DebugViewShaderMode);
+	static IDebugViewModePSInterface* GetPSInterface(TShaderMap<FGlobalShaderType>* ShaderMap, const FMaterial* Material, EDebugViewShaderMode DebugViewShaderMode);
 
 	static void GetDebugMaterial(const FMaterialRenderProxy** MaterialRenderProxy, const FMaterial** Material, ERHIFeatureLevel::Type FeatureLevel);
 
@@ -142,5 +152,4 @@ struct FDebugViewMode
 			ERHIFeatureLevel::Type FeatureLevel, 
 			EDebugViewShaderMode DebugViewShaderMode
 			);
-
 };
diff --git a/Engine/Source/Runtime/Renderer/Private/DecalRenderingShared.cpp b/Engine/Source/Runtime/Renderer/Private/DecalRenderingShared.cpp
--- a/Engine/Source/Runtime/Renderer/Private/DecalRenderingShared.cpp
+++ b/Engine/Source/Runtime/Renderer/Private/DecalRenderingShared.cpp
@@ -446,15 +446,16 @@ void FDecalRendering::SetShader(FRHICommandList& RHICmdList, const FViewInfo& Vi
 	if (DebugViewShaderMode != DVSM_None)
 	{
 		// For this to work, decal VS must output compatible interpolants. Currently this requires to use FDebugPSInLean.
-		FDebugViewModePS* DebugPixelShader = FDebugViewMode::GetPixelShader(View.ShaderMap, DebugViewShaderMode); 
+		// Here we pass nullptr for the material interface because the use of a static bound shader state is only compatible with unique shaders.
+		IDebugViewModePSInterface* DebugPixelShader = FDebugViewMode::GetPSInterface(View.ShaderMap, nullptr, DebugViewShaderMode); 
 
 		const uint32 NumPixelShaderInstructions = PixelShader->GetNumInstructions();
 		const uint32 NumVertexShaderInstructions = VertexShader->GetNumInstructions();
 
 		static FGlobalBoundShaderState BoundShaderState[DVSM_MAX];
-		SetGlobalBoundShaderState(RHICmdList, View.GetFeatureLevel(), BoundShaderState[(uint32)DebugViewShaderMode], GetVertexDeclarationFVector4(), *VertexShader, DebugPixelShader);
+		SetGlobalBoundShaderState(RHICmdList, View.GetFeatureLevel(), BoundShaderState[(uint32)DebugViewShaderMode], GetVertexDeclarationFVector4(), *VertexShader, DebugPixelShader->GetShader());
 
-		DebugPixelShader->SetParameters(RHICmdList, *VertexShader, PixelShader, View);
+		DebugPixelShader->SetParameters(RHICmdList, *VertexShader, PixelShader, DecalData.MaterialProxy, *DecalData.MaterialResource, View);
 		DebugPixelShader->SetMesh(RHICmdList, View);
 	}
 	else
diff --git a/Engine/Source/Runtime/Renderer/Private/DeferredShadingRenderer.cpp b/Engine/Source/Runtime/Renderer/Private/DeferredShadingRenderer.cpp
--- a/Engine/Source/Runtime/Renderer/Private/DeferredShadingRenderer.cpp
+++ b/Engine/Source/Runtime/Renderer/Private/DeferredShadingRenderer.cpp
@@ -704,6 +704,7 @@ void BuildHZB( FRHICommandListImmediate& RHICmdList, FViewInfo& View );
 DECLARE_STATS_GROUP(TEXT("Command List Markers"), STATGROUP_CommandListMarkers, STATCAT_Advanced);
 
 DECLARE_CYCLE_STAT(TEXT("PrePass"), STAT_CLM_PrePass, STATGROUP_CommandListMarkers);
+DECLARE_CYCLE_STAT(TEXT("FXPreRender"), STAT_CLM_FXPreRender, STATGROUP_CommandListMarkers);
 DECLARE_CYCLE_STAT(TEXT("AfterPrePass"), STAT_CLM_AfterPrePass, STATGROUP_CommandListMarkers);
 DECLARE_CYCLE_STAT(TEXT("BasePass"), STAT_CLM_BasePass, STATGROUP_CommandListMarkers);
 DECLARE_CYCLE_STAT(TEXT("AfterBasePass"), STAT_CLM_AfterBasePass, STATGROUP_CommandListMarkers);
@@ -829,8 +830,9 @@ void FDeferredShadingSceneRenderer::RenderOcclusion(FRHICommandListImmediate& RH
 			for (int32 Dest = 1; Dest < NumFrames; Dest++)
 			{
 				OcclusionSubmittedFence[Dest] = OcclusionSubmittedFence[Dest - 1];
-		}
+			}
 			OcclusionSubmittedFence[0] = RHICmdList.RHIThreadFence();
+			RHICmdList.ImmediateFlush(EImmediateFlushType::DispatchToRHIThread);
 		}
 	}
 }
@@ -891,8 +893,11 @@ void FDeferredShadingSceneRenderer::Render(FRHICommandListImmediate& RHICmdList)
 		SceneContext.Allocate(RHICmdList, ViewFamily);
 	}
 
+	FGraphEventArray SortEvents;
+	FILCUpdatePrimTaskData ILCTaskData;
+
 	// Find the visible primitives.
-	InitViews(RHICmdList);
+	bool bDoInitViewAftersPrepass = InitViews(RHICmdList, ILCTaskData, SortEvents);
 #if !UE_BUILD_SHIPPING
 	if (CVarStallInitViews.GetValueOnRenderThread() > 0.0f)
 	{
@@ -999,12 +1004,14 @@ void FDeferredShadingSceneRenderer::Render(FRHICommandListImmediate& RHICmdList)
 	const bool bIsOcclusionTesting = DoOcclusionQueries(FeatureLevel) && (!bIsWireframe || bIsViewFrozen || bHasViewParent);
 
 	// Dynamic vertex and index buffers need to be committed before rendering.
+	if (!bDoInitViewAftersPrepass)
 	{
-		SCOPE_CYCLE_COUNTER(STAT_FDeferredShadingSceneRenderer_FGlobalDynamicVertexBuffer_Commit);
-		FGlobalDynamicVertexBuffer::Get().Commit();
-		FGlobalDynamicIndexBuffer::Get().Commit();
+		{
+			SCOPE_CYCLE_COUNTER(STAT_FDeferredShadingSceneRenderer_FGlobalDynamicVertexBuffer_Commit);
+			FGlobalDynamicVertexBuffer::Get().Commit();
+			FGlobalDynamicIndexBuffer::Get().Commit();
+		}
 	}
-
 	{
 		SCOPE_CYCLE_COUNTER(STAT_FDeferredShadingSceneRenderer_MotionBlurStartFrame);
 		Scene->MotionBlurInfoData.StartFrame(ViewFamily.bWorldIsPaused);
@@ -1016,33 +1023,62 @@ void FDeferredShadingSceneRenderer::Render(FRHICommandListImmediate& RHICmdList)
 	if (!bLateFXPrerender && bDoFXPrerender)
 	{
 		SCOPE_CYCLE_COUNTER(STAT_FDeferredShadingSceneRenderer_FXSystem_PreRender);
+		RHICmdList.SetCurrentStat(GET_STATID(STAT_CLM_FXPreRender));
 		Scene->FXSystem->PreRender(RHICmdList, &Views[0].GlobalDistanceFieldInfo.ParameterData);
 	}
 
-	GRenderTargetPool.AddPhaseEvent(TEXT("EarlyZPass"));
-
 	if (GEnableGPUSkinCache)
 	{
 		GGPUSkinCache.TransitionToReadable(RHICmdList);
 	}
 
+	bool bDidAfterTaskWork = false;
+	auto AfterTasksAreStarted = [&bDidAfterTaskWork, bDoInitViewAftersPrepass, this, &RHICmdList, &ILCTaskData, &SortEvents, bLateFXPrerender, bDoFXPrerender]()
+	{
+		if (!bDidAfterTaskWork)
+		{
+			QUICK_SCOPE_CYCLE_COUNTER(STAT_AfterPrepassTasksWork);
+			bDidAfterTaskWork = true; // only do this once
+			if (bDoInitViewAftersPrepass)
+			{
+				InitViewsPossiblyAfterPrepass(RHICmdList, ILCTaskData, SortEvents);
+				{
+					SCOPE_CYCLE_COUNTER(STAT_FDeferredShadingSceneRenderer_FGlobalDynamicVertexBuffer_Commit);
+					FGlobalDynamicVertexBuffer::Get().Commit();
+					FGlobalDynamicIndexBuffer::Get().Commit();
+				}
+				ServiceLocalQueue();
+			}
+			if (bLateFXPrerender && bDoFXPrerender)
+			{
+				SCOPE_CYCLE_COUNTER(STAT_FDeferredShadingSceneRenderer_FXSystem_PreRender);
+				RHICmdList.SetCurrentStat(GET_STATID(STAT_CLM_FXPreRender));
+				Scene->FXSystem->PreRender(RHICmdList, &Views[0].GlobalDistanceFieldInfo.ParameterData);
+				ServiceLocalQueue();
+			}
+		}
+	};
+
 	// Draw the scene pre-pass / early z pass, populating the scene depth buffer and HiZ
-	RHICmdList.SetCurrentStat(GET_STATID(STAT_CLM_PrePass));
+	GRenderTargetPool.AddPhaseEvent(TEXT("EarlyZPass"));
 	const bool bNeedsPrePass = NeedsPrePass(this);
 	bool bDepthWasCleared;
 	if (bNeedsPrePass)
 	{
-		bDepthWasCleared = RenderPrePass(bLateFXPrerender && bDoFXPrerender, RHICmdList);
+		bDepthWasCleared = RenderPrePass(RHICmdList, AfterTasksAreStarted);
 	}
 	else
 	{
 		// we didn't do the prepass, but we still want the HMD mask if there is one
+		AfterTasksAreStarted();
+		RHICmdList.SetCurrentStat(GET_STATID(STAT_CLM_PrePass));
 		bDepthWasCleared = RenderPrePassHMD(RHICmdList);
 	}
-
+	check(bDidAfterTaskWork);
 	RHICmdList.SetCurrentStat(GET_STATID(STAT_CLM_AfterPrePass));
 	ServiceLocalQueue();
 
+
 	const bool bShouldRenderVelocities = ShouldRenderVelocities();
 	const bool bUseVelocityGBuffer = FVelocityRendering::OutputsToGBuffer();
 	const bool bUseSelectiveBasePassOutputs = UseSelectiveBasePassOutputs();
@@ -1684,7 +1720,7 @@ static TAutoConsoleVariable<int32> CVarRHICmdFlushRenderThreadTasksPrePass(
 	0,
 	TEXT("Wait for completion of parallel render thread tasks at the end of the pre pass.  A more granular version of r.RHICmdFlushRenderThreadTasks. If either r.RHICmdFlushRenderThreadTasks or r.RHICmdFlushRenderThreadTasksPrePass is > 0 we will flush."));
 
-bool FDeferredShadingSceneRenderer::RenderPrePassViewParallel(const FViewInfo& View, FRHICommandListImmediate& ParentCmdList, bool bDoFXPrerender)
+bool FDeferredShadingSceneRenderer::RenderPrePassViewParallel(const FViewInfo& View, FRHICommandListImmediate& ParentCmdList, TFunctionRef<void()> AfterTasksAreStarted, bool bDoPrePre)
 {
 	bool bDepthWasCleared = false;
 	FPrePassParallelCommandListSet ParallelCommandListSet(View, ParentCmdList,
@@ -1722,9 +1758,10 @@ bool FDeferredShadingSceneRenderer::RenderPrePassViewParallel(const FViewInfo& V
 	}
 
 	// we do this step here (awkwardly) so that the above tasks can be in flight while we get the particles (which must be dynamic) setup.
-	if (bDoFXPrerender)
+	if (bDoPrePre)
 	{
-		bDepthWasCleared = PreRenderPrePass(ParentCmdList, bDoFXPrerender);
+		AfterTasksAreStarted();
+		bDepthWasCleared = PreRenderPrePass(ParentCmdList);
 	}
 
 	// Dynamic
@@ -1781,14 +1818,9 @@ IMPLEMENT_SHADER_TYPE(, FDitheredTransitionStencilPS, TEXT("DitheredTransitionSt
 FGlobalBoundShaderState DitheredTransitionStencilBoundShaderState;
 
 /** Possibly do the FX prerender and setup the prepass*/
-bool FDeferredShadingSceneRenderer::PreRenderPrePass(FRHICommandListImmediate& RHICmdList, bool bDoFXPrerender)
+bool FDeferredShadingSceneRenderer::PreRenderPrePass(FRHICommandListImmediate& RHICmdList)
 {
-	if (bDoFXPrerender)
-	{
-		check(Scene->FXSystem && Views.IsValidIndex(0));
-		SCOPE_CYCLE_COUNTER(STAT_FDeferredShadingSceneRenderer_FXSystem_PreRender);
-		Scene->FXSystem->PreRender(RHICmdList, &Views[0].GlobalDistanceFieldInfo.ParameterData);
-	}
+	RHICmdList.SetCurrentStat(GET_STATID(STAT_CLM_PrePass));
 	bool bDepthWasCleared = RenderPrePassHMD(RHICmdList);
 
 	FSceneRenderTargets& SceneContext = FSceneRenderTargets::Get(RHICmdList);
@@ -1839,7 +1871,7 @@ bool FDeferredShadingSceneRenderer::PreRenderPrePass(FRHICommandListImmediate& R
 	return bDepthWasCleared;
 }
 
-bool FDeferredShadingSceneRenderer::RenderPrePass(bool bDoFXPrerender, FRHICommandListImmediate& RHICmdList)
+bool FDeferredShadingSceneRenderer::RenderPrePass(FRHICommandListImmediate& RHICmdList, TFunctionRef<void()> AfterTasksAreStarted)
 {
 	bool bDepthWasCleared = false;
 	SCOPED_DRAW_EVENT(RHICmdList, PrePass);
@@ -1851,12 +1883,12 @@ bool FDeferredShadingSceneRenderer::RenderPrePass(bool bDoFXPrerender, FRHIComma
 
 	bool bParallel = GRHICommandList.UseParallelAlgorithms() && CVarParallelPrePass.GetValueOnRenderThread();
 
-	if (!bDoFXPrerender || !bParallel)
+	if (!bParallel)
 	{
 		// nothing to be gained by delaying this.
-		bDepthWasCleared = PreRenderPrePass(RHICmdList, bDoFXPrerender);
+		AfterTasksAreStarted();
+		bDepthWasCleared = PreRenderPrePass(RHICmdList);
 		bDidPrePre = true;
-		bDoFXPrerender = false; // only need to do this once
 	}
 	else
 	{
@@ -1875,10 +1907,9 @@ bool FDeferredShadingSceneRenderer::RenderPrePass(bool bDoFXPrerender, FRHIComma
 				const FViewInfo& View = Views[ViewIndex];
 				if (View.ShouldRenderView())
 				{
-					bDepthWasCleared = RenderPrePassViewParallel(View, RHICmdList, bDoFXPrerender) || bDepthWasCleared;
+					bDepthWasCleared = RenderPrePassViewParallel(View, RHICmdList, AfterTasksAreStarted, !bDidPrePre) || bDepthWasCleared;
 					bDirty = true; // assume dirty since we are not going to wait
 					bDidPrePre = true;
-					bDoFXPrerender = false; // only need to do this once even if we have multiple views
 				}
 			}
 		}
@@ -1898,7 +1929,8 @@ bool FDeferredShadingSceneRenderer::RenderPrePass(bool bDoFXPrerender, FRHIComma
 	if (!bDidPrePre)
 	{
 		// For some reason we haven't done this yet. Best do it now for consistency with the old code.
-		bDepthWasCleared = PreRenderPrePass(RHICmdList, bDoFXPrerender);
+		AfterTasksAreStarted();
+		bDepthWasCleared = PreRenderPrePass(RHICmdList);
 		bDidPrePre = true;
 	}
 
diff --git a/Engine/Source/Runtime/Renderer/Private/DeferredShadingRenderer.h b/Engine/Source/Runtime/Renderer/Private/DeferredShadingRenderer.h
--- a/Engine/Source/Runtime/Renderer/Private/DeferredShadingRenderer.h
+++ b/Engine/Source/Runtime/Renderer/Private/DeferredShadingRenderer.h
@@ -63,7 +63,7 @@ public:
 	 * Renders the scene's prepass for a particular view in parallel
 	 * @return true if the depth was cleared
 	 */
-	bool RenderPrePassViewParallel(const FViewInfo& View, FRHICommandListImmediate& ParentCmdList, bool bDoFXPrerender);
+	bool RenderPrePassViewParallel(const FViewInfo& View, FRHICommandListImmediate& ParentCmdList, TFunctionRef<void()> AfterTasksAreStarted, bool bDoPrePre);
 
 	/** Renders the basepass for the static data of a given View. */
 	bool RenderBasePassStaticData(FRHICommandList& RHICmdList, FViewInfo& View);
@@ -158,7 +158,9 @@ private:
 	bool CheckForLightFunction(const FLightSceneInfo* LightSceneInfo) const;
 
 	/** Determines which primitives are visible for each view. */
-	void InitViews(FRHICommandListImmediate& RHICmdList);
+	bool InitViews(FRHICommandListImmediate& RHICmdList, struct FILCUpdatePrimTaskData& ILCTaskData, FGraphEventArray& SortEvents);
+
+	void InitViewsPossiblyAfterPrepass(FRHICommandListImmediate& RHICmdList, struct FILCUpdatePrimTaskData& ILCTaskData, FGraphEventArray& SortEvents);
 
 	void CreateIndirectCapsuleShadows();
 
@@ -166,13 +168,13 @@ private:
 	* Setup the prepass. This is split out so that in parallel we can do the fx prerender after we start the parallel tasks
 	* @return true if the depth was cleared
 	*/
-	bool PreRenderPrePass(FRHICommandListImmediate& RHICmdList, bool bDoFXPrerender);
+	bool PreRenderPrePass(FRHICommandListImmediate& RHICmdList);
 
 	/**
 	 * Renders the scene's prepass and occlusion queries.
 	 * @return true if the depth was cleared
 	 */
-	bool RenderPrePass(bool bDoFXPrerender, FRHICommandListImmediate& RHICmdList);
+	bool RenderPrePass(FRHICommandListImmediate& RHICmdList, TFunctionRef<void()> AfterTasksAreStarted);
 
 	/**
 	 * Renders the active HMD's hidden area mask as a depth prepass, if available.
diff --git a/Engine/Source/Runtime/Renderer/Private/DistortionRendering.cpp b/Engine/Source/Runtime/Renderer/Private/DistortionRendering.cpp
--- a/Engine/Source/Runtime/Renderer/Private/DistortionRendering.cpp
+++ b/Engine/Source/Runtime/Renderer/Private/DistortionRendering.cpp
@@ -739,19 +739,25 @@ bool TDistortionMeshDrawingPolicyFactory<DistortMeshPolicy>::DrawStaticMesh(
 
 bool FDistortionPrimSet::DrawAccumulatedOffsets(FRHICommandListImmediate& RHICmdList, const FViewInfo& View, bool bInitializeOffsets)
 {
-	bool bDirty=false;
-
-	// Draw the view's elements with the translucent drawing policy.
-	bDirty |= DrawViewElements<TDistortionMeshDrawingPolicyFactory<FDistortMeshAccumulatePolicy> >(
-		RHICmdList,
-		View,
-		bInitializeOffsets,
-		0,	// DPG Index?
-		false // Distortion is rendered post fog.
-		);
+	bool bDirty = false;
+
+	QUICK_SCOPE_CYCLE_COUNTER(STAT_FDistortionPrimSet_DrawAccumulatedOffsets);
+
+	{
+		QUICK_SCOPE_CYCLE_COUNTER(STAT_FDistortionPrimSet_DrawAccumulatedOffsets_View);
+		// Draw the view's elements with the translucent drawing policy.
+		bDirty |= DrawViewElements<TDistortionMeshDrawingPolicyFactory<FDistortMeshAccumulatePolicy> >(
+			RHICmdList,
+			View,
+			bInitializeOffsets,
+			0,	// DPG Index?
+			false // Distortion is rendered post fog.
+			);
+	}
 
 	if( Prims.Num() )
 	{
+		QUICK_SCOPE_CYCLE_COUNTER(STAT_FDistortionPrimSet_DrawAccumulatedOffsets_Prims);
 		// Draw sorted scene prims
 		for( int32 PrimIdx=0; PrimIdx < Prims.Num(); PrimIdx++ )
 		{
@@ -760,14 +766,16 @@ bool FDistortionPrimSet::DrawAccumulatedOffsets(FRHICommandListImmediate& RHICmd
 
 			TDistortionMeshDrawingPolicyFactory<FDistortMeshAccumulatePolicy>::ContextType Context(bInitializeOffsets);
 
-			for (int32 MeshBatchIndex = 0; MeshBatchIndex < View.DynamicMeshElements.Num(); MeshBatchIndex++)
 			{
-				const FMeshBatchAndRelevance& MeshBatchAndRelevance = View.DynamicMeshElements[MeshBatchIndex];
-
-				if (MeshBatchAndRelevance.PrimitiveSceneProxy == PrimitiveSceneProxy)
+				for (int32 MeshBatchIndex = 0; MeshBatchIndex < View.DynamicMeshElements.Num(); MeshBatchIndex++)
 				{
-					const FMeshBatch& MeshBatch = *MeshBatchAndRelevance.Mesh;
-					bDirty |= TDistortionMeshDrawingPolicyFactory<FDistortMeshAccumulatePolicy>::DrawDynamicMesh(RHICmdList, View, Context, MeshBatch, false, false, MeshBatchAndRelevance.PrimitiveSceneProxy, MeshBatch.BatchHitProxyId);
+					const FMeshBatchAndRelevance& MeshBatchAndRelevance = View.DynamicMeshElements[MeshBatchIndex];
+
+					if (MeshBatchAndRelevance.PrimitiveSceneProxy == PrimitiveSceneProxy)
+					{
+						const FMeshBatch& MeshBatch = *MeshBatchAndRelevance.Mesh;
+						bDirty |= TDistortionMeshDrawingPolicyFactory<FDistortMeshAccumulatePolicy>::DrawDynamicMesh(RHICmdList, View, Context, MeshBatch, false, false, MeshBatchAndRelevance.PrimitiveSceneProxy, MeshBatch.BatchHitProxyId);
+					}
 				}
 			}
 
@@ -819,6 +827,7 @@ int32 FSceneRenderer::GetRefractionQuality(const FSceneViewFamily& ViewFamily)
  */
 void FSceneRenderer::RenderDistortion(FRHICommandListImmediate& RHICmdList)
 {
+	QUICK_SCOPE_CYCLE_COUNTER(STAT_FSceneRenderer_RenderDistortion);
 	SCOPED_DRAW_EVENT(RHICmdList, Distortion);
 
 	// do we need to render the distortion pass?
@@ -846,6 +855,7 @@ void FSceneRenderer::RenderDistortion(FRHICommandListImmediate& RHICmdList)
 	// Render accumulated distortion offsets
 	if( bRender)
 	{
+		QUICK_SCOPE_CYCLE_COUNTER(STAT_FSceneRenderer_RenderDistortion_Render);
 		SCOPED_DRAW_EVENT(RHICmdList, DistortionAccum);
 
 		// Create a texture to store the resolved light attenuation values, and a render-targetable surface to hold the unresolved light attenuation values.
@@ -900,6 +910,7 @@ void FSceneRenderer::RenderDistortion(FRHICommandListImmediate& RHICmdList)
 
 	if(bDirty)
 	{
+		QUICK_SCOPE_CYCLE_COUNTER(STAT_FSceneRenderer_RenderDistortion_Post);
 		SCOPED_DRAW_EVENT(RHICmdList, DistortionApply);
 
 		SceneContext.ResolveSceneColor(RHICmdList, FResolveRect(0, 0, ViewFamily.FamilySizeX, ViewFamily.FamilySizeY));
@@ -916,6 +927,7 @@ void FSceneRenderer::RenderDistortion(FRHICommandListImmediate& RHICmdList)
 
 		for(int32 ViewIndex = 0, Num = Views.Num(); ViewIndex < Num; ++ViewIndex)
 		{
+			QUICK_SCOPE_CYCLE_COUNTER(STAT_FSceneRenderer_RenderDistortion_PostView1);
 			SCOPED_CONDITIONAL_DRAW_EVENTF(RHICmdList, EventView, Views.Num() > 1, TEXT("View%d"), ViewIndex);
 
 			FViewInfo& View = Views[ViewIndex];
@@ -960,6 +972,7 @@ void FSceneRenderer::RenderDistortion(FRHICommandListImmediate& RHICmdList)
 
 		for(int32 ViewIndex = 0, Num = Views.Num(); ViewIndex < Num; ++ViewIndex)
 		{
+			QUICK_SCOPE_CYCLE_COUNTER(STAT_FSceneRenderer_RenderDistortion_PostView2);
 			SCOPED_CONDITIONAL_DRAW_EVENTF(RHICmdList, EventView, Views.Num() > 1, TEXT("View%d"), ViewIndex);
 
 			FViewInfo& View = Views[ViewIndex];
diff --git a/Engine/Source/Runtime/Renderer/Private/DrawingPolicy.cpp b/Engine/Source/Runtime/Renderer/Private/DrawingPolicy.cpp
--- a/Engine/Source/Runtime/Renderer/Private/DrawingPolicy.cpp
+++ b/Engine/Source/Runtime/Renderer/Private/DrawingPolicy.cpp
@@ -87,8 +87,9 @@ void FMeshDrawingPolicy::DrawMesh(FRHICommandList& RHICmdList, const FMeshBatch&
 		if(BatchElement.IndexBuffer)
 		{
 			check(BatchElement.IndexBuffer->IsInitialized());
-			if (BatchElement.InstanceRuns)
+			if (BatchElement.bIsInstanceRuns)
 			{
+				checkSlow(BatchElement.bIsInstanceRuns);
 				if (!GRHISupportsFirstInstance)
 				{
 					if (bUsePositionOnlyVS)
diff --git a/Engine/Source/Runtime/Renderer/Private/ForwardBasePassRendering.h b/Engine/Source/Runtime/Renderer/Private/ForwardBasePassRendering.h
--- a/Engine/Source/Runtime/Renderer/Private/ForwardBasePassRendering.h
+++ b/Engine/Source/Runtime/Renderer/Private/ForwardBasePassRendering.h
@@ -665,7 +665,7 @@ public:
 				}
 			}
 
-			FDebugViewMode::GetPixelShader(View->ShaderMap, GetDebugViewShaderMode())->SetParameters(RHICmdList, VertexShader, PixelShader, *View);
+			FDebugViewMode::GetPSInterface(View->ShaderMap, MaterialResource, GetDebugViewShaderMode())->SetParameters(RHICmdList, VertexShader, PixelShader, MaterialRenderProxy, *MaterialResource, *View);
 		}
 		else
 #endif
@@ -714,7 +714,7 @@ public:
 #if !(UE_BUILD_SHIPPING || UE_BUILD_TEST)
 		if (GetDebugViewShaderMode() != DVSM_None)
 		{
-			PixelShaderRHIRef = FDebugViewMode::GetPixelShader(GetGlobalShaderMap(InFeatureLevel), GetDebugViewShaderMode())->GetPixelShader();
+			PixelShaderRHIRef = FDebugViewMode::GetPSInterface(GetGlobalShaderMap(InFeatureLevel), MaterialResource, GetDebugViewShaderMode())->GetShader()->GetPixelShader();
 		}
 #endif
 
@@ -765,7 +765,7 @@ public:
 				RHICmdList.SetBlendState(TStaticBlendState<CW_RGB,BO_Add,BF_One,BF_One>::GetRHI());
 			}
 
-			FDebugViewMode::GetPixelShader(GetGlobalShaderMap(View.FeatureLevel), GetDebugViewShaderMode())->SetParameters(RHICmdList, VertexShader, PixelShader, View);
+			FDebugViewMode::GetPSInterface(GetGlobalShaderMap(View.FeatureLevel), MaterialResource, GetDebugViewShaderMode())->SetMesh(RHICmdList, VertexFactory, View, PrimitiveSceneProxy, BatchElement, DrawRenderState);
 		}
 		else
 #endif
diff --git a/Engine/Source/Runtime/Renderer/Private/IndirectLightingCache.cpp b/Engine/Source/Runtime/Renderer/Private/IndirectLightingCache.cpp
--- a/Engine/Source/Runtime/Renderer/Private/IndirectLightingCache.cpp
+++ b/Engine/Source/Runtime/Renderer/Private/IndirectLightingCache.cpp
@@ -374,8 +374,9 @@ FIndirectLightingCacheAllocation* FIndirectLightingCache::FindPrimitiveAllocatio
 FAutoConsoleTaskPriority CPrio_FUpdateCachePrimitivesTask(
 	TEXT("TaskGraph.TaskPriorities.UpdateCachePrimitivesTask"),
 	TEXT("Task and thread priority for FUpdateCachePrimitivesTask."),
-	ENamedThreads::NormalThreadPriority,
-	ENamedThreads::NormalTaskPriority
+	ENamedThreads::HighThreadPriority, // if we have high priority task threads, then use them...
+	ENamedThreads::NormalTaskPriority, // .. at normal task priority
+	ENamedThreads::HighTaskPriority // if we don't have hi pri threads, then use normal priority threads at high task priority instead
 	);
 
 class FUpdateCachePrimitivesTask
diff --git a/Engine/Source/Runtime/Renderer/Private/LightFunctionRendering.cpp b/Engine/Source/Runtime/Renderer/Private/LightFunctionRendering.cpp
--- a/Engine/Source/Runtime/Renderer/Private/LightFunctionRendering.cpp
+++ b/Engine/Source/Runtime/Renderer/Private/LightFunctionRendering.cpp
@@ -165,14 +165,14 @@ IMPLEMENT_MATERIAL_SHADER_TYPE(,FLightFunctionPS,TEXT("LightFunctionPixelShader"
 /** Returns a fade fraction for a light function and a given view based on the appropriate fade settings. */
 static float GetLightFunctionFadeFraction(const FViewInfo& View, FSphere LightBounds)
 {
-	extern float CalculateShadowFadeAlpha(float MaxUnclampedResolution, int32 ShadowFadeResolution, int32 MinShadowResolution);
+	extern float CalculateShadowFadeAlpha(const float MaxUnclampedResolution, const uint32 ShadowFadeResolution, const uint32 MinShadowResolution);
 
 	// Override the global settings with the light's settings if the light has them specified
 	static auto CVarMinShadowResolution = IConsoleManager::Get().FindTConsoleVariableDataInt(TEXT("r.Shadow.MinResolution"));
 	static auto CVarShadowFadeResolution = IConsoleManager::Get().FindTConsoleVariableDataInt(TEXT("r.Shadow.FadeResolution"));
 
-	const uint32 MinShadowResolution = CVarMinShadowResolution->GetValueOnRenderThread();
-	const int32 ShadowFadeResolution = CVarShadowFadeResolution->GetValueOnRenderThread();
+	const uint32 MinShadowResolution  = FMath::Max<int32>(0, CVarMinShadowResolution->GetValueOnRenderThread());
+	const uint32 ShadowFadeResolution = FMath::Max<int32>(0, CVarShadowFadeResolution->GetValueOnRenderThread());
 
 	// Project the bounds onto the view
 	const FVector4 ScreenPosition = View.WorldToScreen(LightBounds.Center);
diff --git a/Engine/Source/Runtime/Renderer/Private/PostProcess/PostProcessEyeAdaptation.cpp b/Engine/Source/Runtime/Renderer/Private/PostProcess/PostProcessEyeAdaptation.cpp
--- a/Engine/Source/Runtime/Renderer/Private/PostProcess/PostProcessEyeAdaptation.cpp
+++ b/Engine/Source/Runtime/Renderer/Private/PostProcess/PostProcessEyeAdaptation.cpp
@@ -78,7 +78,7 @@ inline static void ComputeEyeAdaptationValues(const ERHIFeatureLevel::Type MinFe
 	float DeltaLog = HistogramLogMax - HistogramLogMin;
 	float Multiply = 1.0f / DeltaLog;
 	float Add = -HistogramLogMin * Multiply;
-	float MinIntensity = exp2(HistogramLogMin);
+	float MinIntensity = FMath::Exp2(HistogramLogMin);
 	Out[2] = FVector4(Multiply, Add, MinIntensity, 0);
 }
 
diff --git a/Engine/Source/Runtime/Renderer/Private/PostProcess/PostProcessSubsurface.cpp b/Engine/Source/Runtime/Renderer/Private/PostProcess/PostProcessSubsurface.cpp
--- a/Engine/Source/Runtime/Renderer/Private/PostProcess/PostProcessSubsurface.cpp
+++ b/Engine/Source/Runtime/Renderer/Private/PostProcess/PostProcessSubsurface.cpp
@@ -15,6 +15,15 @@
 ENGINE_API const IPooledRenderTarget* GetSubsufaceProfileTexture_RT(FRHICommandListImmediate& RHICmdList);
 
 
+static TAutoConsoleVariable<int32> CVarSSSQuality(
+	TEXT("r.SSS.Quality"),
+	0,
+	TEXT("Defines the quality of the recombine pass when using the SubsurfaceScatteringProfile shading model\n")
+	TEXT(" 0: low (faster, default)\n")
+	TEXT(" 1: high (sharper details but slower)\n")
+	TEXT("-1: auto, 1 if TemporalAA is disabled (without TemporalAA the quality is more noticable)"),
+	ECVF_RenderThreadSafe  | ECVF_Scalability);
+
 static TAutoConsoleVariable<int32> CVarSSSFilter(
 	TEXT("r.SSS.Filter"),
 	1,
@@ -196,7 +205,7 @@ FRCPassPostProcessSubsurfaceVisualize::FRCPassPostProcessSubsurfaceVisualize(FRH
 
 void FRCPassPostProcessSubsurfaceVisualize::Process(FRenderingCompositePassContext& Context)
 {
-	SCOPED_DRAW_EVENT(Context.RHICmdList, SubsurfaceSetup);
+	SCOPED_DRAW_EVENT(Context.RHICmdList, SubsurfaceVisualize);
 
 	const FPooledRenderTargetDesc* InputDesc = GetInputDesc(ePId_Input0);
 
@@ -718,7 +727,8 @@ FPooledRenderTargetDesc FRCPassPostProcessSubsurface::ComputeOutputDesc(EPassOut
 
 /** Encapsulates the post processing subsurface recombine pixel shader. */
 // @param RecombineMode 0:fullres, 1: halfres, 2:no scattering, just reconstruct the lighting (needed for scalability)
-template <uint32 RecombineMode>
+// @param RecombineQuality 0:low..1:high
+template <uint32 RecombineMode, uint32 RecombineQuality>
 class TPostProcessSubsurfaceRecombinePS : public FGlobalShader
 {
 	DECLARE_SHADER_TYPE(TPostProcessSubsurfaceRecombinePS, Global);
@@ -731,6 +741,7 @@ class TPostProcessSubsurfaceRecombinePS : public FGlobalShader
 	static void ModifyCompilationEnvironment(EShaderPlatform Platform, FShaderCompilerEnvironment& OutEnvironment)
 	{
 		FGlobalShader::ModifyCompilationEnvironment(Platform,OutEnvironment);
+		OutEnvironment.SetDefine(TEXT("RECOMBINE_QUALITY"), RecombineQuality);
 		OutEnvironment.SetDefine(TEXT("HALF_RES"), (uint32)(RecombineMode == 1));
 		OutEnvironment.SetDefine(TEXT("RECOMBINE_SUBSURFACESCATTER"), (uint32)(RecombineMode != 2));
 		OutEnvironment.SetDefine(TEXT("SUBSURFACE_RADIUS_SCALE"), SUBSURFACE_RADIUS_SCALE);
@@ -784,18 +795,19 @@ public:
 };
 
 // #define avoids a lot of code duplication
-#define VARIATION1(A) typedef TPostProcessSubsurfaceRecombinePS<A> TPostProcessSubsurfaceRecombinePS##A; \
-	IMPLEMENT_SHADER_TYPE2(TPostProcessSubsurfaceRecombinePS##A, SF_Pixel);
+#define VARIATION1(A)		VARIATION2(A,0)			VARIATION2(A,1)
+#define VARIATION2(A, B) typedef TPostProcessSubsurfaceRecombinePS<A, B> TPostProcessSubsurfaceRecombinePS##A##B; \
+	IMPLEMENT_SHADER_TYPE2(TPostProcessSubsurfaceRecombinePS##A##B, SF_Pixel);
 	VARIATION1(0) VARIATION1(1) VARIATION1(2)
 #undef VARIATION1
+#undef VARIATION2
 
 // @param RecombineMode 0:fullres, 1: halfres, 2:no scattering, just reconstruct the lighting (needed for scalability)
-template <uint32 RecombineMode>
+// @param RecombineQuality 0:low..1:high
+template <uint32 RecombineMode, uint32 RecombineQuality>
 void SetSubsurfaceRecombineShader(const FRenderingCompositePassContext& Context, TShaderMapRef<FPostProcessVS> &VertexShader)
 {
-	SCOPED_DRAW_EVENTF(Context.RHICmdList, SubsurfacePassRecombine, TEXT("SubsurfacePassRecombine%d"), RecombineMode);
-
-	TShaderMapRef<TPostProcessSubsurfaceRecombinePS<RecombineMode> > PixelShader(Context.GetShaderMap());
+	TShaderMapRef<TPostProcessSubsurfaceRecombinePS<RecombineMode, RecombineQuality> > PixelShader(Context.GetShaderMap());
 
 	static FGlobalBoundShaderState BoundShaderState;
 
@@ -849,21 +861,53 @@ void FRCPassPostProcessSubsurfaceRecombine::Process(FRenderingCompositePassConte
 
 	TShaderMapRef<FPostProcessVS> VertexShader(Context.GetShaderMap());
 
-	if(GetInput(ePId_Input1)->IsValid())
+	uint32 QualityCVar = CVarSSSQuality.GetValueOnRenderThread();
+
+	// 0:low / 1:high
+	uint32 RecombineQuality = 0;
 	{
-		if(bHalfRes)
+		if(QualityCVar == -1)
 		{
-			SetSubsurfaceRecombineShader<1>(Context, VertexShader);
+			RecombineQuality = (View.FinalPostProcessSettings.AntiAliasingMethod == AAM_TemporalAA) ? 0 : 1;
 		}
-		else
+		else if(QualityCVar == 1)
 		{
-			SetSubsurfaceRecombineShader<0>(Context, VertexShader);
+			RecombineQuality = 1;
 		}
 	}
-	else
+
+	// needed for Scalability
+	// 0:fullres, 1: halfres, 2:no scattering, just reconstruct the lighting (needed for scalability)
+	uint32 RecombineMode = 2;
+
+	if(GetInput(ePId_Input1)->IsValid())
 	{
-		// needed for Scalability
-		SetSubsurfaceRecombineShader<2>(Context, VertexShader);
+		RecombineMode = bHalfRes ? 1 : 0;
+	}
+
+	SCOPED_DRAW_EVENTF(Context.RHICmdList, SubsurfacePassRecombine, TEXT("SubsurfacePassRecombine Mode:%d Quality:%d"), RecombineMode, RecombineQuality);
+
+	{
+		if(RecombineQuality == 0)
+		{
+			switch(RecombineMode)
+			{
+				case 0: SetSubsurfaceRecombineShader<0, 0>(Context, VertexShader); break;
+				case 1: SetSubsurfaceRecombineShader<1, 0>(Context, VertexShader); break;
+				case 2: SetSubsurfaceRecombineShader<2, 0>(Context, VertexShader); break;
+				default: check(0);
+			}
+		}
+		else
+		{
+			switch(RecombineMode)
+			{
+				case 0: SetSubsurfaceRecombineShader<0, 1>(Context, VertexShader); break;
+				case 1: SetSubsurfaceRecombineShader<1, 1>(Context, VertexShader); break;
+				case 2: SetSubsurfaceRecombineShader<2, 1>(Context, VertexShader); break;
+				default: check(0);
+			}
+		}
 	}
 
 	DrawPostProcessPass(
diff --git a/Engine/Source/Runtime/Renderer/Private/PostProcess/RenderingCompositionGraph.cpp b/Engine/Source/Runtime/Renderer/Private/PostProcess/RenderingCompositionGraph.cpp
--- a/Engine/Source/Runtime/Renderer/Private/PostProcess/RenderingCompositionGraph.cpp
+++ b/Engine/Source/Runtime/Renderer/Private/PostProcess/RenderingCompositionGraph.cpp
@@ -89,12 +89,45 @@ bool ShouldDebugCompositionGraph()
 #endif
 }
 
+void Test()
+{
+	struct ObjectSize4
+	{
+		void SetBaseValues(){}
+		static FName GetFName()
+		{
+ 			static const FName Name(TEXT("ObjectSize4"));
+ 			return Name;
+		}
+		uint8 Data[4];
+	};
+ 
+	MS_ALIGN(16) struct ObjectAligned16
+	{
+		void SetBaseValues(){}
+		static FName GetFName()
+		{
+ 			static const FName Name(TEXT("ObjectAligned16"));
+ 			return Name;
+		}
+		uint8 Data[16];
+	} GCC_ALIGN(16);
+
+	// https://udn.unrealengine.com/questions/274066/fblendablemanager-returning-wrong-or-misaligned-da.html
+	FBlendableManager Manager;
+	Manager.GetSingleFinalData<ObjectSize4>();
+	ObjectAligned16& AlignedData = Manager.GetSingleFinalData<ObjectAligned16>();
+
+	check((reinterpret_cast<ptrdiff_t>(&AlignedData) & 16) == 0);
+}
+
 void ExecuteCompositionGraphDebug()
 {
 	ENQUEUE_UNIQUE_RENDER_COMMAND(
 		StartDebugCompositionGraph,
 	{
 		GDebugCompositionGraphFrames = 1;
+		Test();
 	}
 	);
 }
diff --git a/Engine/Source/Runtime/Renderer/Private/RendererScene.cpp b/Engine/Source/Runtime/Renderer/Private/RendererScene.cpp
--- a/Engine/Source/Runtime/Renderer/Private/RendererScene.cpp
+++ b/Engine/Source/Runtime/Renderer/Private/RendererScene.cpp
@@ -2454,7 +2454,7 @@ FSceneInterface* FRendererModule::AllocateScene(UWorld* World, bool bInRequiresH
 	// Create a full fledged scene if we have something to render.
 	if (GIsClient && FApp::CanEverRender() && !GUsingNullRHI)
 	{
-		FScene* NewScene = new FScene(World, bInRequiresHitProxies, GIsEditor && !World->IsGameWorld(), bCreateFXSystem, InFeatureLevel);
+		FScene* NewScene = new FScene(World, bInRequiresHitProxies, GIsEditor && (!World || !World->IsGameWorld()), bCreateFXSystem, InFeatureLevel);
 		AllocatedScenes.Add(NewScene);
 		return NewScene;
 	}
diff --git a/Engine/Source/Runtime/Renderer/Private/SceneHitProxyRendering.cpp b/Engine/Source/Runtime/Renderer/Private/SceneHitProxyRendering.cpp
--- a/Engine/Source/Runtime/Renderer/Private/SceneHitProxyRendering.cpp
+++ b/Engine/Source/Runtime/Renderer/Private/SceneHitProxyRendering.cpp
@@ -552,7 +552,13 @@ void FDeferredShadingSceneRenderer::RenderHitProxies(FRHICommandListImmediate& R
 	if (HitProxyRT)
 	{
 		// Find the visible primitives.
-		InitViews(RHICmdList);
+		FGraphEventArray SortEvents;
+		FILCUpdatePrimTaskData ILCTaskData;
+		bool bDoInitViewAftersPrepass = InitViews(RHICmdList, ILCTaskData, SortEvents);
+		if (bDoInitViewAftersPrepass)
+		{
+			InitViewsPossiblyAfterPrepass(RHICmdList, ILCTaskData, SortEvents);
+		}
 		::RenderHitProxies(RHICmdList, this, HitProxyRT);
 		ClearPrimitiveSingleFramePrecomputedLightingBuffers();
 	}
diff --git a/Engine/Source/Runtime/Renderer/Private/ScenePrivate.h b/Engine/Source/Runtime/Renderer/Private/ScenePrivate.h
--- a/Engine/Source/Runtime/Renderer/Private/ScenePrivate.h
+++ b/Engine/Source/Runtime/Renderer/Private/ScenePrivate.h
@@ -91,8 +91,6 @@ public:
 #include "SceneHitProxyRendering.h"
 #include "DebugViewModeRendering.h"
 #include "ShaderComplexityRendering.h"
-#include "WantedMipsAccuracyRendering.h"
-#include "TexelFactorAccuracyRendering.h"
 #include "ShadowRendering.h"
 #include "SceneRendering.h"
 #include "StaticMeshDrawList.h"
diff --git a/Engine/Source/Runtime/Renderer/Private/SceneRendering.cpp b/Engine/Source/Runtime/Renderer/Private/SceneRendering.cpp
--- a/Engine/Source/Runtime/Renderer/Private/SceneRendering.cpp
+++ b/Engine/Source/Runtime/Renderer/Private/SceneRendering.cpp
@@ -18,6 +18,7 @@
 #include "PostProcessCircleDOF.h"
 #include "SceneUtils.h"
 #include "LightGrid.h"
+#include "AtmosphereRendering.h"
 
 /*-----------------------------------------------------------------------------
 	Globals
@@ -214,6 +215,7 @@ FRHICommandList* FParallelCommandListSet::AllocCommandList()
 
 void FParallelCommandListSet::Dispatch(bool bHighPriority)
 {
+	QUICK_SCOPE_CYCLE_COUNTER(STAT_FParallelCommandListSet_Dispatch);
 	check(IsInRenderingThread() && FMemStack::Get().GetNumMarks() == 1); // we do not want this popped before the end of the scene and it better be the scene allocator
 	check(CommandLists.Num() == Events.Num());
 	check(CommandLists.Num() == NumAlloc);
@@ -397,6 +399,7 @@ void FViewInfo::Init()
 	bIsSnapshot = false;
 
 	bAllowStencilDither = false;
+	InitAtmosphereConstantsInView(*this);
 }
 
 FViewInfo::~FViewInfo()
@@ -474,6 +477,39 @@ void FViewInfo::SetupSkyIrradianceEnvironmentMapConstants(FVector4* OutSkyIrradi
 	}
 }
 
+static FTextureRHIParamRef OrBlack2DIfNull(FTextureRHIParamRef Tex)
+{
+	FTextureRHIParamRef Result = Tex ? Tex : GBlackTexture->TextureRHI.GetReference();
+	check(Result);
+	return Result;
+}
+
+static FTextureRHIParamRef OrBlack3DIfNull(FTextureRHIParamRef Tex)
+{
+	// we fall back to 2D which are unbound es2 parameters
+	return OrBlack2DIfNull(Tex ? Tex : GBlackVolumeTexture->TextureRHI.GetReference());
+}
+
+static void SetBlack2DIfNull(FTextureRHIParamRef& Tex)
+{
+	if (!Tex)
+	{
+		Tex = GBlackTexture->TextureRHI.GetReference();
+		check(Tex);
+	}
+}
+
+static void SetBlack3DIfNull(FTextureRHIParamRef& Tex)
+{
+	if (!Tex)
+	{
+		Tex = GBlackVolumeTexture->TextureRHI.GetReference();
+		// we fall back to 2D which are unbound es2 parameters
+		SetBlack2DIfNull(Tex);
+	}
+}
+
+
 /** Creates the view's uniform buffers given a set of view transforms. */
 void FViewInfo::CreateUniformBuffer(
 	TUniformBufferRef<FViewUniformShaderParameters>& OutViewUniformBuffer,
@@ -774,6 +810,49 @@ void FViewInfo::CreateUniformBuffer(
 		FrameUniformShaderParameters.AtmosphericFogInscatterAltitudeSampleNum = 0;
 	}
 
+	SetBlack2DIfNull(FrameUniformShaderParameters.AtmosphereTransmittanceTexture_UB);
+	SetBlack2DIfNull(FrameUniformShaderParameters.AtmosphereIrradianceTexture_UB);
+	SetBlack3DIfNull(FrameUniformShaderParameters.AtmosphereInscatterTexture_UB);
+
+	FrameUniformShaderParameters.AtmosphereTransmittanceTextureSampler_UB = TStaticSamplerState<SF_Bilinear>::GetRHI();
+	FrameUniformShaderParameters.AtmosphereIrradianceTextureSampler_UB = TStaticSamplerState<SF_Bilinear>::GetRHI();
+	FrameUniformShaderParameters.AtmosphereInscatterTextureSampler_UB = TStaticSamplerState<SF_Bilinear>::GetRHI();
+
+	if (GSystemTextures.PerlinNoiseGradient.GetReference())
+	{
+		FrameUniformShaderParameters.PerlinNoiseGradientTexture = (FTexture2DRHIRef&)GSystemTextures.PerlinNoiseGradient->GetRenderTargetItem().ShaderResourceTexture;
+		SetBlack2DIfNull(FrameUniformShaderParameters.PerlinNoiseGradientTexture);
+	}
+	check(FrameUniformShaderParameters.PerlinNoiseGradientTexture);
+	FrameUniformShaderParameters.PerlinNoiseGradientTextureSampler = TStaticSamplerState<SF_Point, AM_Wrap, AM_Wrap, AM_Wrap>::GetRHI();
+
+	if (GSystemTextures.PerlinNoise3D.GetReference())
+	{
+		FrameUniformShaderParameters.PerlinNoise3DTexture = (FTexture3DRHIRef&)GSystemTextures.PerlinNoise3D->GetRenderTargetItem().ShaderResourceTexture;
+		SetBlack3DIfNull(FrameUniformShaderParameters.PerlinNoise3DTexture);
+	}
+	check(FrameUniformShaderParameters.PerlinNoise3DTexture);
+	FrameUniformShaderParameters.PerlinNoise3DTextureSampler = TStaticSamplerState<SF_Bilinear, AM_Wrap, AM_Wrap, AM_Wrap>::GetRHI();
+
+	for (int32 Index = 0; Index < GMaxGlobalDistanceFieldClipmaps; Index++)
+	{
+		FrameUniformShaderParameters.GlobalVolumeCenterAndExtent_UB[Index] = GlobalDistanceFieldInfo.ParameterData.CenterAndExtent[Index];
+		FrameUniformShaderParameters.GlobalVolumeWorldToUVAddAndMul_UB[Index] = GlobalDistanceFieldInfo.ParameterData.WorldToUVAddAndMul[Index];
+	}
+	FrameUniformShaderParameters.GlobalVolumeDimension_UB = GlobalDistanceFieldInfo.ParameterData.GlobalDFResolution;
+	FrameUniformShaderParameters.GlobalVolumeTexelSize_UB = 1.0f / GlobalDistanceFieldInfo.ParameterData.GlobalDFResolution;
+	FrameUniformShaderParameters.MaxGlobalDistance_UB = GlobalDistanceFieldInfo.ParameterData.MaxDistance;
+
+	FrameUniformShaderParameters.GlobalDistanceFieldTexture0_UB = OrBlack3DIfNull(GlobalDistanceFieldInfo.ParameterData.Textures[0]);
+	FrameUniformShaderParameters.GlobalDistanceFieldSampler0_UB = TStaticSamplerState<SF_Bilinear, AM_Wrap, AM_Wrap, AM_Wrap>::GetRHI();
+	FrameUniformShaderParameters.GlobalDistanceFieldTexture1_UB = OrBlack3DIfNull(GlobalDistanceFieldInfo.ParameterData.Textures[1]);
+	FrameUniformShaderParameters.GlobalDistanceFieldSampler1_UB = TStaticSamplerState<SF_Bilinear, AM_Wrap, AM_Wrap, AM_Wrap>::GetRHI();
+	FrameUniformShaderParameters.GlobalDistanceFieldTexture2_UB = OrBlack3DIfNull(GlobalDistanceFieldInfo.ParameterData.Textures[2]);
+	FrameUniformShaderParameters.GlobalDistanceFieldSampler2_UB = TStaticSamplerState<SF_Bilinear, AM_Wrap, AM_Wrap, AM_Wrap>::GetRHI();
+	FrameUniformShaderParameters.GlobalDistanceFieldTexture3_UB = OrBlack3DIfNull(GlobalDistanceFieldInfo.ParameterData.Textures[3]);
+	FrameUniformShaderParameters.GlobalDistanceFieldSampler3_UB = TStaticSamplerState<SF_Bilinear, AM_Wrap, AM_Wrap, AM_Wrap>::GetRHI();
+
+
 	FrameUniformShaderParameters.UnlitViewmodeMask = bIsUnlitView ? 1 : 0;
 	FrameUniformShaderParameters.OutOfBoundsMask = Family->EngineShowFlags.VisualizeOutOfBoundsPixels ? 1 : 0;
 
diff --git a/Engine/Source/Runtime/Renderer/Private/SceneRendering.h b/Engine/Source/Runtime/Renderer/Private/SceneRendering.h
--- a/Engine/Source/Runtime/Renderer/Private/SceneRendering.h
+++ b/Engine/Source/Runtime/Renderer/Private/SceneRendering.h
@@ -686,6 +686,8 @@ public:
 		return DrawRenderState;
 	}
 
+	inline FVector GetPrevViewDirection() const { return PrevViewMatrices.ViewMatrix.GetColumn(2); }
+
 	/** Create a snapshot of this view info on the scene allocator. */
 	FViewInfo* CreateSnapshot() const;
 
@@ -904,9 +906,6 @@ protected:
 	/** Initialized the fog constants for each view. */
 	void InitFogConstants();
 
-	/** Initialized the atmopshere constants for each view. */
-	void InitAtmosphereConstants();
-
 	/** Returns whether there are translucent primitives to be renderered. */
 	bool ShouldRenderTranslucency() const;
 
diff --git a/Engine/Source/Runtime/Renderer/Private/SceneVisibility.cpp b/Engine/Source/Runtime/Renderer/Private/SceneVisibility.cpp
--- a/Engine/Source/Runtime/Renderer/Private/SceneVisibility.cpp
+++ b/Engine/Source/Runtime/Renderer/Private/SceneVisibility.cpp
@@ -114,6 +114,14 @@ static FAutoConsoleVariableRef CVarILCUpdatePrimitivesTask(
 	ECVF_RenderThreadSafe
 	);
 
+static int32 GDoInitViewsLightingAfterPrepass = 0;
+static FAutoConsoleVariableRef CVarDoInitViewsLightingAfterPrepass(
+	TEXT("r.DoInitViewsLightingAfterPrepass"),
+	GDoInitViewsLightingAfterPrepass,
+	TEXT("Delays the lighting part of InitViews until after the prepass. This improves the threading throughput and gets the prepass to the GPU ASAP. Experimental options; has an unknown race."),
+	ECVF_RenderThreadSafe
+	);
+
 /** Distance fade cvars */
 static int32 GDisableLODFade = false;
 static FAutoConsoleVariableRef CVarDisableLODFade( TEXT("r.DisableLODFade"), GDisableLODFade, TEXT("Disable fading for distance culling"), ECVF_RenderThreadSafe );
@@ -2499,6 +2507,7 @@ void FSceneRenderer::ComputeViewVisibility(FRHICommandListImmediate& RHICmdList)
 		FSceneBitArray PrimitiveHiddenByLODMap;
 		if (bLODSceneTreeActive)
 		{
+			QUICK_SCOPE_CYCLE_COUNTER(STAT_ViewVisibilityTime_HLOD);
 			Scene->SceneLODHierarchy.PopulateFadingFlags(View);
 
 			PrimitiveHiddenByLODMap.Init(false, View.PrimitiveVisibilityMap.Num());
@@ -2512,14 +2521,18 @@ void FSceneRenderer::ComputeViewVisibility(FRHICommandListImmediate& RHICmdList)
 		}
 
 		MarkAllPrimitivesForReflectionProxyUpdate(Scene);
-		Scene->ConditionalMarkStaticMeshElementsForUpdate();
+		{
+			QUICK_SCOPE_CYCLE_COUNTER(STAT_ViewVisibilityTime_ConditionalMarkStaticMeshElementsForUpdate);
+			Scene->ConditionalMarkStaticMeshElementsForUpdate();
+		}
 
 		{
 			SCOPE_CYCLE_COUNTER(STAT_ViewRelevance);
-				ComputeAndMarkRelevanceForViewParallel(RHICmdList, Scene, View, ViewBit, HasDynamicMeshElementsMasks, HasDynamicEditorMeshElementsMasks);
+			ComputeAndMarkRelevanceForViewParallel(RHICmdList, Scene, View, ViewBit, HasDynamicMeshElementsMasks, HasDynamicEditorMeshElementsMasks);
 
 			if (bLODSceneTreeActive)
 			{
+				QUICK_SCOPE_CYCLE_COUNTER(STAT_ViewVisibilityTime_HLOD_Prop);
 				for (FSceneSetBitIterator BitIt(PrimitiveHiddenByLODMap); BitIt; ++BitIt)
 				{
 					View.PrimitiveViewRelevanceMap[BitIt.GetIndex()].bInitializedThisFrame = true;
@@ -2563,16 +2576,21 @@ void FSceneRenderer::ComputeViewVisibility(FRHICommandListImmediate& RHICmdList)
 
 void FSceneRenderer::PostVisibilityFrameSetup(FILCUpdatePrimTaskData& OutILCTaskData)
 {
-	for (int32 ViewIndex = 0; ViewIndex < Views.Num(); ViewIndex++)
-	{		
-		FViewInfo& View = Views[ViewIndex];
-
-		// sort the translucent primitives
-		View.TranslucentPrimSet.SortPrimitives();
+	QUICK_SCOPE_CYCLE_COUNTER(STAT_PostVisibilityFrameSetup);
 
-		if (View.State)
+	{
+		QUICK_SCOPE_CYCLE_COUNTER(STAT_PostVisibilityFrameSetup_SortTranslucency);
+		for (int32 ViewIndex = 0; ViewIndex < Views.Num(); ViewIndex++)
 		{
-			((FSceneViewState*)View.State)->TrimHistoryRenderTargets(Scene);
+			FViewInfo& View = Views[ViewIndex];
+
+			// sort the translucent primitives
+			View.TranslucentPrimSet.SortPrimitives();
+
+			if (View.State)
+			{
+				((FSceneViewState*)View.State)->TrimHistoryRenderTargets(Scene);
+			}
 		}
 	}
 
@@ -2596,6 +2614,7 @@ void FSceneRenderer::PostVisibilityFrameSetup(FILCUpdatePrimTaskData& OutILCTask
 
 	if (ViewFamily.EngineShowFlags.HitProxies == 0)
 	{
+		QUICK_SCOPE_CYCLE_COUNTER(STAT_PostVisibilityFrameSetup_IndirectLightingCache_Update);
 		if (GILCUpdatePrimTaskEnabled)
 		{
 			Scene->IndirectLightingCache.StartUpdateCachePrimitivesTask(Scene, *this, true, OutILCTaskData);
@@ -2606,145 +2625,149 @@ void FSceneRenderer::PostVisibilityFrameSetup(FILCUpdatePrimTaskData& OutILCTask
 		}		
 	}
 
-	// determine visibility of each light
-	for(TSparseArray<FLightSceneInfoCompact>::TConstIterator LightIt(Scene->Lights);LightIt;++LightIt)
 	{
-		const FLightSceneInfoCompact& LightSceneInfoCompact = *LightIt;
-		const FLightSceneInfo* LightSceneInfo = LightSceneInfoCompact.LightSceneInfo;
-
-		// view frustum cull lights in each view
-		for(int32 ViewIndex = 0;ViewIndex < Views.Num();ViewIndex++)
-		{		
-			const FLightSceneProxy* Proxy = LightSceneInfo->Proxy;
-			FViewInfo& View = Views[ViewIndex];
-			FVisibleLightViewInfo& VisibleLightViewInfo = View.VisibleLightInfos[LightIt.GetIndex()];
-			// dir lights are always visible, and point/spot only if in the frustum
-			if (Proxy->GetLightType() == LightType_Point  
-				|| Proxy->GetLightType() == LightType_Spot)
-			{
-				const float Radius = Proxy->GetRadius();
-
-				if (View.ViewFrustum.IntersectSphere(Proxy->GetOrigin(), Radius))
+		QUICK_SCOPE_CYCLE_COUNTER(STAT_PostVisibilityFrameSetup_Light_Visibility);
+		// determine visibility of each light
+		for (TSparseArray<FLightSceneInfoCompact>::TConstIterator LightIt(Scene->Lights); LightIt; ++LightIt)
+		{
+			const FLightSceneInfoCompact& LightSceneInfoCompact = *LightIt;
+			const FLightSceneInfo* LightSceneInfo = LightSceneInfoCompact.LightSceneInfo;
+
+			// view frustum cull lights in each view
+			for (int32 ViewIndex = 0; ViewIndex < Views.Num(); ViewIndex++)
+			{
+				const FLightSceneProxy* Proxy = LightSceneInfo->Proxy;
+				FViewInfo& View = Views[ViewIndex];
+				FVisibleLightViewInfo& VisibleLightViewInfo = View.VisibleLightInfos[LightIt.GetIndex()];
+				// dir lights are always visible, and point/spot only if in the frustum
+				if (Proxy->GetLightType() == LightType_Point
+					|| Proxy->GetLightType() == LightType_Spot)
 				{
-					FSphere Bounds = Proxy->GetBoundingSphere();
-					float DistanceSquared = (Bounds.Center - View.ViewMatrices.ViewOrigin).SizeSquared();
-					const bool bDrawLight = FMath::Square( FMath::Min( 0.0002f, GMinScreenRadiusForLights / Bounds.W ) * View.LODDistanceFactor ) * DistanceSquared < 1.0f;
-					VisibleLightViewInfo.bInViewFrustum = bDrawLight;
+					const float Radius = Proxy->GetRadius();
+
+					if (View.ViewFrustum.IntersectSphere(Proxy->GetOrigin(), Radius))
+					{
+						FSphere Bounds = Proxy->GetBoundingSphere();
+						float DistanceSquared = (Bounds.Center - View.ViewMatrices.ViewOrigin).SizeSquared();
+						const bool bDrawLight = FMath::Square(FMath::Min(0.0002f, GMinScreenRadiusForLights / Bounds.W) * View.LODDistanceFactor) * DistanceSquared < 1.0f;
+						VisibleLightViewInfo.bInViewFrustum = bDrawLight;
+					}
 				}
-			}
-			else
-			{
-				VisibleLightViewInfo.bInViewFrustum = true;
+				else
+				{
+					VisibleLightViewInfo.bInViewFrustum = true;
 
-				static const auto CVarMobileMSAA = IConsoleManager::Get().FindTConsoleVariableDataInt(TEXT("r.MobileMSAA"));
-				bool bNotMobileMSAA = !(CVarMobileMSAA ? CVarMobileMSAA->GetValueOnRenderThread() > 1 : false);
+					static const auto CVarMobileMSAA = IConsoleManager::Get().FindTConsoleVariableDataInt(TEXT("r.MobileMSAA"));
+					bool bNotMobileMSAA = !(CVarMobileMSAA ? CVarMobileMSAA->GetValueOnRenderThread() > 1 : false);
 
-				// Setup single sun-shaft from direction lights for mobile.
-				if(bCheckLightShafts && LightSceneInfo->bEnableLightShaftBloom)
-				{
-					// Find directional light for sun shafts.
-					// Tweaked values from UE3 implementation.
-					const float PointLightFadeDistanceIncrease = 200.0f;
-					const float PointLightRadiusFadeFactor = 5.0f;
+					// Setup single sun-shaft from direction lights for mobile.
+					if (bCheckLightShafts && LightSceneInfo->bEnableLightShaftBloom)
+					{
+						// Find directional light for sun shafts.
+						// Tweaked values from UE3 implementation.
+						const float PointLightFadeDistanceIncrease = 200.0f;
+						const float PointLightRadiusFadeFactor = 5.0f;
 
-					const FVector WorldSpaceBlurOrigin = LightSceneInfo->Proxy->GetPosition();
-					// Transform into post projection space
-					FVector4 ProjectedBlurOrigin = View.WorldToScreen(WorldSpaceBlurOrigin);
+						const FVector WorldSpaceBlurOrigin = LightSceneInfo->Proxy->GetPosition();
+						// Transform into post projection space
+						FVector4 ProjectedBlurOrigin = View.WorldToScreen(WorldSpaceBlurOrigin);
 
-					const float DistanceToBlurOrigin = (View.ViewMatrices.ViewOrigin - WorldSpaceBlurOrigin).Size() + PointLightFadeDistanceIncrease;
+						const float DistanceToBlurOrigin = (View.ViewMatrices.ViewOrigin - WorldSpaceBlurOrigin).Size() + PointLightFadeDistanceIncrease;
 
-					// Don't render if the light's origin is behind the view
-					if(ProjectedBlurOrigin.W >= 0.0f
-						// Don't render point lights that have completely faded out
-							&& (LightSceneInfo->Proxy->GetLightType() == LightType_Directional 
+						// Don't render if the light's origin is behind the view
+						if (ProjectedBlurOrigin.W >= 0.0f
+							// Don't render point lights that have completely faded out
+							&& (LightSceneInfo->Proxy->GetLightType() == LightType_Directional
 							|| DistanceToBlurOrigin < LightSceneInfo->Proxy->GetRadius() * PointLightRadiusFadeFactor))
-					{
-						View.bLightShaftUse = bNotMobileMSAA;
-						View.LightShaftCenter.X = ProjectedBlurOrigin.X / ProjectedBlurOrigin.W;
-						View.LightShaftCenter.Y = ProjectedBlurOrigin.Y / ProjectedBlurOrigin.W;
-						// TODO: Might want to hookup different colors for these.
-						View.LightShaftColorMask = LightSceneInfo->BloomTint;
-						View.LightShaftColorApply = LightSceneInfo->BloomTint;
+						{
+							View.bLightShaftUse = bNotMobileMSAA;
+							View.LightShaftCenter.X = ProjectedBlurOrigin.X / ProjectedBlurOrigin.W;
+							View.LightShaftCenter.Y = ProjectedBlurOrigin.Y / ProjectedBlurOrigin.W;
+							// TODO: Might want to hookup different colors for these.
+							View.LightShaftColorMask = LightSceneInfo->BloomTint;
+							View.LightShaftColorApply = LightSceneInfo->BloomTint;
+						}
 					}
 				}
-			}
-
-			// Draw shapes for reflection captures
-			if( View.bIsReflectionCapture && VisibleLightViewInfo.bInViewFrustum && Proxy->HasStaticLighting() && Proxy->GetLightType() != LightType_Directional )
-			{
-				FVector Origin = Proxy->GetOrigin();
-				FVector ToLight = Origin - View.ViewMatrices.ViewOrigin;
-				float DistanceSqr = ToLight | ToLight;
-				float Radius = Proxy->GetRadius();
 
-				if( DistanceSqr < Radius * Radius )
+				// Draw shapes for reflection captures
+				if (View.bIsReflectionCapture && VisibleLightViewInfo.bInViewFrustum && Proxy->HasStaticLighting() && Proxy->GetLightType() != LightType_Directional)
 				{
-					FVector4	PositionAndInvRadius;
-					FVector4	ColorAndFalloffExponent;
-					FVector		Direction;
-					FVector2D	SpotAngles;
-					float		SourceRadius;
-					float		SourceLength;
-					float		MinRoughness;
-					Proxy->GetParameters( PositionAndInvRadius, ColorAndFalloffExponent, Direction, SpotAngles, SourceRadius, SourceLength, MinRoughness );
-
-					// Force to be at least 0.75 pixels
-					float CubemapSize = 128.0f;
-					float Distance = FMath::Sqrt( DistanceSqr );
-					float MinRadius = Distance * 0.75f / CubemapSize;
-					SourceRadius = FMath::Max( MinRadius, SourceRadius );
-
-					// Snap to cubemap pixel center to reduce aliasing
-					FVector Scale = ToLight.GetAbs();
-					int32 MaxComponent = Scale.X > Scale.Y ? ( Scale.X > Scale.Z ? 0 : 2 ) : ( Scale.Y > Scale.Z ? 1 : 2 );
-					for( int32 k = 1; k < 3; k++ )
-					{
-						float Projected = ToLight[ (MaxComponent + k) % 3 ] / Scale[ MaxComponent ];
-						float Quantized = ( FMath::RoundToFloat( Projected * (0.5f * CubemapSize) - 0.5f ) + 0.5f ) / (0.5f * CubemapSize);
-						ToLight[ (MaxComponent + k) % 3 ] = Quantized * Scale[ MaxComponent ];
-					}
-					Origin = ToLight + View.ViewMatrices.ViewOrigin;
-				
-					FLinearColor Color( ColorAndFalloffExponent );
+					FVector Origin = Proxy->GetOrigin();
+					FVector ToLight = Origin - View.ViewMatrices.ViewOrigin;
+					float DistanceSqr = ToLight | ToLight;
+					float Radius = Proxy->GetRadius();
 
-					if( Proxy->IsInverseSquared() )
+					if (DistanceSqr < Radius * Radius)
 					{
-						float LightRadiusMask = FMath::Square( 1.0f - FMath::Square( DistanceSqr * FMath::Square( PositionAndInvRadius.W ) ) );
-						Color *= LightRadiusMask;
-						
-						// Correction for lumen units
-						Color *= 16.0f;
-					}
-					else
-					{
-						// Remove inverse square falloff
-						Color *= DistanceSqr + 1.0f;
+						FVector4	PositionAndInvRadius;
+						FVector4	ColorAndFalloffExponent;
+						FVector		Direction;
+						FVector2D	SpotAngles;
+						float		SourceRadius;
+						float		SourceLength;
+						float		MinRoughness;
+						Proxy->GetParameters(PositionAndInvRadius, ColorAndFalloffExponent, Direction, SpotAngles, SourceRadius, SourceLength, MinRoughness);
+
+						// Force to be at least 0.75 pixels
+						float CubemapSize = 128.0f;
+						float Distance = FMath::Sqrt(DistanceSqr);
+						float MinRadius = Distance * 0.75f / CubemapSize;
+						SourceRadius = FMath::Max(MinRadius, SourceRadius);
+
+						// Snap to cubemap pixel center to reduce aliasing
+						FVector Scale = ToLight.GetAbs();
+						int32 MaxComponent = Scale.X > Scale.Y ? (Scale.X > Scale.Z ? 0 : 2) : (Scale.Y > Scale.Z ? 1 : 2);
+						for (int32 k = 1; k < 3; k++)
+						{
+							float Projected = ToLight[(MaxComponent + k) % 3] / Scale[MaxComponent];
+							float Quantized = (FMath::RoundToFloat(Projected * (0.5f * CubemapSize) - 0.5f) + 0.5f) / (0.5f * CubemapSize);
+							ToLight[(MaxComponent + k) % 3] = Quantized * Scale[MaxComponent];
+						}
+						Origin = ToLight + View.ViewMatrices.ViewOrigin;
 
-						// Apply falloff
-						Color *= FMath::Pow( 1.0f - DistanceSqr * FMath::Square( PositionAndInvRadius.W ), ColorAndFalloffExponent.W ); 
-					}
+						FLinearColor Color(ColorAndFalloffExponent);
 
-					// Spot falloff
-					FVector L = ToLight.GetSafeNormal();
-					Color *= FMath::Square( FMath::Clamp( ( (L | Direction) - SpotAngles.X ) * SpotAngles.Y, 0.0f, 1.0f ) );
+						if (Proxy->IsInverseSquared())
+						{
+							float LightRadiusMask = FMath::Square(1.0f - FMath::Square(DistanceSqr * FMath::Square(PositionAndInvRadius.W)));
+							Color *= LightRadiusMask;
 
-					// Scale by visible area
-					Color /= PI * FMath::Square( SourceRadius );
-				
-					// Always opaque
-					Color.A = 1.0f;
-				
-					FViewElementPDI LightPDI( &View, NULL );
-					FMaterialRenderProxy* const ColoredMeshInstance = new(FMemStack::Get()) FColoredMaterialRenderProxy( GEngine->DebugMeshMaterial->GetRenderProxy(false), Color );
-					DrawSphere( &LightPDI, Origin, FVector( SourceRadius, SourceRadius, SourceRadius ), 8, 6, ColoredMeshInstance, SDPG_World );
+							// Correction for lumen units
+							Color *= 16.0f;
+						}
+						else
+						{
+							// Remove inverse square falloff
+							Color *= DistanceSqr + 1.0f;
+
+							// Apply falloff
+							Color *= FMath::Pow(1.0f - DistanceSqr * FMath::Square(PositionAndInvRadius.W), ColorAndFalloffExponent.W);
+						}
+
+						// Spot falloff
+						FVector L = ToLight.GetSafeNormal();
+						Color *= FMath::Square(FMath::Clamp(((L | Direction) - SpotAngles.X) * SpotAngles.Y, 0.0f, 1.0f));
+
+						// Scale by visible area
+						Color /= PI * FMath::Square(SourceRadius);
+
+						// Always opaque
+						Color.A = 1.0f;
+
+						FViewElementPDI LightPDI(&View, NULL);
+						FMaterialRenderProxy* const ColoredMeshInstance = new(FMemStack::Get()) FColoredMaterialRenderProxy(GEngine->DebugMeshMaterial->GetRenderProxy(false), Color);
+						DrawSphere(&LightPDI, Origin, FVector(SourceRadius, SourceRadius, SourceRadius), 8, 6, ColoredMeshInstance, SDPG_World);
+					}
 				}
 			}
 		}
 	}
+	{
 
-	// Initialize the fog constants.
-	InitFogConstants();
-	InitAtmosphereConstants();
+		QUICK_SCOPE_CYCLE_COUNTER(STAT_PostVisibilityFrameSetup_InitFogConstants);
+		InitFogConstants();
+	}
 }
 
 uint32 GetShadowQuality();
@@ -2753,7 +2776,7 @@ uint32 GetShadowQuality();
  * Initialize scene's views.
  * Check visibility, sort translucent items, etc.
  */
-void FDeferredShadingSceneRenderer::InitViews(FRHICommandListImmediate& RHICmdList)
+bool FDeferredShadingSceneRenderer::InitViews(FRHICommandListImmediate& RHICmdList, struct FILCUpdatePrimTaskData& ILCTaskData, FGraphEventArray& SortEvents)
 {	
 	SCOPED_DRAW_EVENT(RHICmdList, InitViews);
 
@@ -2769,7 +2792,6 @@ void FDeferredShadingSceneRenderer::InitViews(FRHICommandListImmediate& RHICmdLi
 			View.FinalPostProcessSettings.AntiAliasingMethod = AAM_None;
 		}
 	}
-	FILCUpdatePrimTaskData ILCTaskData;
 	PreVisibilityFrameSetup(RHICmdList);
 	ComputeViewVisibility(RHICmdList);
 	PostVisibilityFrameSetup(ILCTaskData);
@@ -2782,7 +2804,6 @@ void FDeferredShadingSceneRenderer::InitViews(FRHICommandListImmediate& RHICmdLi
 		AverageViewPosition += View.ViewMatrices.ViewOrigin / Views.Num();
 	}
 
-	FGraphEventArray SortEvents;
 	if (FApp::ShouldUseThreadingForPerformance() && CVarParallelInitViews.GetValueOnRenderThread() > 0)
 	{
 		AsyncSortBasePassStaticData(AverageViewPosition, SortEvents);
@@ -2792,6 +2813,40 @@ void FDeferredShadingSceneRenderer::InitViews(FRHICommandListImmediate& RHICmdLi
 		SortBasePassStaticData(AverageViewPosition);
 	}
 
+	bool bDoInitViewAftersPrepass = !!GDoInitViewsLightingAfterPrepass;
+
+	if (!bDoInitViewAftersPrepass)
+	{
+		InitViewsPossiblyAfterPrepass(RHICmdList, ILCTaskData, SortEvents);
+	}
+
+	{
+		QUICK_SCOPE_CYCLE_COUNTER(STAT_InitViews_InitRHIResources);
+		// initialize per-view uniform buffer.
+		for (int32 ViewIndex = 0; ViewIndex < Views.Num(); ViewIndex++)
+		{
+			// Initialize the view's RHI resources.
+			Views[ViewIndex].InitRHIResources(nullptr);
+
+			// Possible stencil dither optimization approach
+			Views[ViewIndex].bAllowStencilDither = bDitheredLODTransitionsUseStencil;
+		}
+	}
+
+	{
+		QUICK_SCOPE_CYCLE_COUNTER(STAT_InitViews_OnStartFrame);
+		OnStartFrame();
+	}
+
+	return bDoInitViewAftersPrepass;
+}
+
+void FDeferredShadingSceneRenderer::InitViewsPossiblyAfterPrepass(FRHICommandListImmediate& RHICmdList, struct FILCUpdatePrimTaskData& ILCTaskData, FGraphEventArray& SortEvents)
+{
+	SCOPED_DRAW_EVENT(RHICmdList, InitViewsPossiblyAfterPrepass);
+
+	SCOPE_CYCLE_COUNTER(STAT_InitViewsPossiblyAfterPrepass);
+
 	// this cannot be moved later because of static mesh updates for stuff that is only visible in shadows
 	if (SortEvents.Num())
 	{
@@ -2813,20 +2868,11 @@ void FDeferredShadingSceneRenderer::InitViews(FRHICommandListImmediate& RHICmdLi
 		Scene->IndirectLightingCache.FinalizeCacheUpdates(Scene, *this, ILCTaskData);
 	}
 
-	// Now that the indirect lighting cache is updated, we can update the primitive precomputed lighting buffers.
-	UpdatePrimitivePrecomputedLightingBuffers();
-
-	// initialize per-view uniform buffer.
-	for (int32 ViewIndex = 0; ViewIndex < Views.Num(); ViewIndex++)
 	{
-		// Initialize the view's RHI resources.
-		Views[ViewIndex].InitRHIResources(nullptr);
-
-		// Possible stencil dither optimization approach
-		Views[ViewIndex].bAllowStencilDither = bDitheredLODTransitionsUseStencil;
+		QUICK_SCOPE_CYCLE_COUNTER(STAT_InitViews_UpdatePrimitivePrecomputedLightingBuffers);
+		// Now that the indirect lighting cache is updated, we can update the primitive precomputed lighting buffers.
+		UpdatePrimitivePrecomputedLightingBuffers();
 	}
-
-	OnStartFrame();
 }
 
 /*------------------------------------------------------------------------------
diff --git a/Engine/Source/Runtime/Renderer/Private/ShaderBaseClasses.cpp b/Engine/Source/Runtime/Renderer/Private/ShaderBaseClasses.cpp
--- a/Engine/Source/Runtime/Renderer/Private/ShaderBaseClasses.cpp
+++ b/Engine/Source/Runtime/Renderer/Private/ShaderBaseClasses.cpp
@@ -69,16 +69,8 @@ FMaterialShader::FMaterialShader(const FMaterialShaderType::CompiledShaderInitia
 	DeferredParameters.Bind(Initializer.ParameterMap);
 	LightAttenuation.Bind(Initializer.ParameterMap, TEXT("LightAttenuationTexture"));
 	LightAttenuationSampler.Bind(Initializer.ParameterMap, TEXT("LightAttenuationTextureSampler"));
-	AtmosphericFogTextureParameters.Bind(Initializer.ParameterMap);
 	EyeAdaptation.Bind(Initializer.ParameterMap, TEXT("EyeAdaptation"));
-	// only used it Material has expression that requires PerlinNoiseGradientTexture
-	PerlinNoiseGradientTexture.Bind(Initializer.ParameterMap,TEXT("PerlinNoiseGradientTexture"));
-	PerlinNoiseGradientTextureSampler.Bind(Initializer.ParameterMap,TEXT("PerlinNoiseGradientTextureSampler"));
-	// only used it Material has expression that requires PerlinNoise3DTexture
-	PerlinNoise3DTexture.Bind(Initializer.ParameterMap,TEXT("PerlinNoise3DTexture"));
-	PerlinNoise3DTextureSampler.Bind(Initializer.ParameterMap,TEXT("PerlinNoise3DTextureSampler"));
-
-	GlobalDistanceFieldParameters.Bind(Initializer.ParameterMap);
+
 }
 
 FUniformBufferRHIParamRef FMaterialShader::GetParameterCollectionBuffer(const FGuid& Id, const FSceneInterface* SceneInterface) const
@@ -87,21 +79,19 @@ FUniformBufferRHIParamRef FMaterialShader::GetParameterCollectionBuffer(const FG
 	return Scene ? Scene->GetParameterCollectionBuffer(Id) : FUniformBufferRHIParamRef();
 }
 
+static TLockFreeFixedSizeAllocator<sizeof(FUniformExpressionCache), PLATFORM_CACHE_LINE_SIZE, FNoopCounter> TempUniformCache;
+
 template<typename ShaderRHIParamRef>
 void FMaterialShader::SetParameters(
 	FRHICommandList& RHICmdList,
-	const ShaderRHIParamRef ShaderRHI, 
-	const FMaterialRenderProxy* MaterialRenderProxy, 
+	const ShaderRHIParamRef ShaderRHI,
+	const FMaterialRenderProxy* MaterialRenderProxy,
 	const FMaterial& Material,
-	const FSceneView& View, 
-	bool bDeferredPass, 
-	ESceneRenderTargetsMode::Type TextureMode, 
+	const FSceneView& View,
+	bool bDeferredPass,
+	ESceneRenderTargetsMode::Type TextureMode,
 	const bool bIsInstancedStereo)
 {
-	ERHIFeatureLevel::Type FeatureLevel = View.GetFeatureLevel();
-	FUniformExpressionCache TempUniformExpressionCache;
-	const FUniformExpressionCache* UniformExpressionCache = &MaterialRenderProxy->UniformExpressionCache[FeatureLevel];
-
 	SetParameters(RHICmdList, ShaderRHI, View);
 
 	// If the material has cached uniform expressions for selection or hover
@@ -112,113 +102,113 @@ void FMaterialShader::SetParameters(
 		!View.Family->EngineShowFlags.Selection &&
 		(MaterialRenderProxy->IsSelected() || MaterialRenderProxy->IsHovered());
 
+	ERHIFeatureLevel::Type FeatureLevel = View.GetFeatureLevel();
+	check(Material.GetRenderingThreadShaderMap() && Material.GetRenderingThreadShaderMap()->IsValidForRendering() && Material.GetFeatureLevel() == FeatureLevel);
+
+	FUniformExpressionCache* UniformExpressionCache = &MaterialRenderProxy->UniformExpressionCache[FeatureLevel];
+	bool bUniformExpressionCacheNeedsDelete = false;
+
 	if (!bAllowCachedUniformExpressions || !UniformExpressionCache->bUpToDate || bOverrideSelection)
 	{
 		FMaterialRenderContext MaterialRenderContext(MaterialRenderProxy, Material, &View);
-		MaterialRenderProxy->EvaluateUniformExpressions(TempUniformExpressionCache, MaterialRenderContext, &RHICmdList);
-		UniformExpressionCache = &TempUniformExpressionCache;
+		bUniformExpressionCacheNeedsDelete = true;
+		UniformExpressionCache = new (TempUniformCache.Allocate()) FUniformExpressionCache();
+		MaterialRenderProxy->EvaluateUniformExpressions(*UniformExpressionCache, MaterialRenderContext, &RHICmdList);
+		SetLocalUniformBufferParameter(RHICmdList, ShaderRHI, MaterialUniformBuffer, UniformExpressionCache->LocalUniformBuffer);
+	}
+	else
+	{
+		SetUniformBufferParameter(RHICmdList, ShaderRHI, MaterialUniformBuffer, UniformExpressionCache->UniformBuffer);
 	}
 
-	check(Material.GetRenderingThreadShaderMap());
-	check(Material.GetRenderingThreadShaderMap()->IsValidForRendering());
-	check(Material.GetFeatureLevel() == FeatureLevel);
-
-	// Validate that the shader is being used for a material that matches the uniform expression set the shader was compiled for.
-	const FUniformExpressionSet& MaterialUniformExpressionSet = Material.GetRenderingThreadShaderMap()->GetUniformExpressionSet();
 
-#if NO_LOGGING == 0
-	bool bUniformExpressionSetMismatch = !DebugUniformExpressionSet.Matches(MaterialUniformExpressionSet)
-		|| UniformExpressionCache->CachedUniformExpressionShaderMap != Material.GetRenderingThreadShaderMap();
-	if (!bUniformExpressionSetMismatch)
+#if !(UE_BUILD_TEST || UE_BUILD_SHIPPING || !WITH_EDITOR)
 	{
-		auto DumpUB = [](const FRHIUniformBufferLayout& Layout)
+		// Validate that the shader is being used for a material that matches the uniform expression set the shader was compiled for.
+		const FUniformExpressionSet& MaterialUniformExpressionSet = Material.GetRenderingThreadShaderMap()->GetUniformExpressionSet();
+		bool bUniformExpressionSetMismatch = !DebugUniformExpressionSet.Matches(MaterialUniformExpressionSet)
+			|| UniformExpressionCache->CachedUniformExpressionShaderMap != Material.GetRenderingThreadShaderMap();
+		if (!bUniformExpressionSetMismatch)
 		{
-			FString DebugName = Layout.GetDebugName().GetPlainNameString();
-			UE_LOG(LogShaders, Warning, TEXT("Layout %s, Hash %08x"), *DebugName, Layout.GetHash());
-			FString ResourcesString;
-			for (int32 Index = 0; Index < Layout.Resources.Num(); ++Index)
+			auto DumpUB = [](const FRHIUniformBufferLayout& Layout)
 			{
-				ResourcesString += FString::Printf(TEXT("%d "), Layout.Resources[Index]);
-			}
-			UE_LOG(LogShaders, Warning, TEXT("Layout CB Size %d Res Offs %d; %d Resources: %s"), Layout.ConstantBufferSize, Layout.ResourceOffset, Layout.Resources.Num(), *ResourcesString);
-		};
-		if (UniformExpressionCache->LocalUniformBuffer.IsValid())
-		{
-			if (UniformExpressionCache->LocalUniformBuffer.BypassUniform)
+				FString DebugName = Layout.GetDebugName().GetPlainNameString();
+				UE_LOG(LogShaders, Warning, TEXT("Layout %s, Hash %08x"), *DebugName, Layout.GetHash());
+				FString ResourcesString;
+				for (int32 Index = 0; Index < Layout.Resources.Num(); ++Index)
+				{
+					ResourcesString += FString::Printf(TEXT("%d "), Layout.Resources[Index]);
+				}
+				UE_LOG(LogShaders, Warning, TEXT("Layout CB Size %d Res Offs %d; %d Resources: %s"), Layout.ConstantBufferSize, Layout.ResourceOffset, Layout.Resources.Num(), *ResourcesString);
+			};
+			if (UniformExpressionCache->LocalUniformBuffer.IsValid())
 			{
-				if (DebugUniformExpressionUBLayout.GetHash() != UniformExpressionCache->LocalUniformBuffer.BypassUniform->GetLayout().GetHash())
+				if (UniformExpressionCache->LocalUniformBuffer.BypassUniform)
 				{
-					UE_LOG(LogShaders, Warning, TEXT("Material Expression UB mismatch!"));
-					DumpUB(DebugUniformExpressionUBLayout);
-					DumpUB(UniformExpressionCache->LocalUniformBuffer.BypassUniform->GetLayout());
-					bUniformExpressionSetMismatch = true;
+					if (DebugUniformExpressionUBLayout.GetHash() != UniformExpressionCache->LocalUniformBuffer.BypassUniform->GetLayout().GetHash())
+					{
+						UE_LOG(LogShaders, Warning, TEXT("Material Expression UB mismatch!"));
+						DumpUB(DebugUniformExpressionUBLayout);
+						DumpUB(UniformExpressionCache->LocalUniformBuffer.BypassUniform->GetLayout());
+						bUniformExpressionSetMismatch = true;
+					}
+				}
+				else
+				{
+					if (DebugUniformExpressionUBLayout.GetHash() != UniformExpressionCache->LocalUniformBuffer.WorkArea->Layout->GetHash())
+					{
+						UE_LOG(LogShaders, Warning, TEXT("Material Expression UB mismatch!"));
+						DumpUB(DebugUniformExpressionUBLayout);
+						DumpUB(*UniformExpressionCache->LocalUniformBuffer.WorkArea->Layout);
+						bUniformExpressionSetMismatch = true;
+					}
 				}
 			}
 			else
 			{
-				if (DebugUniformExpressionUBLayout.GetHash() != UniformExpressionCache->LocalUniformBuffer.WorkArea->Layout->GetHash())
+				if (DebugUniformExpressionUBLayout.GetHash() != UniformExpressionCache->UniformBuffer->GetLayout().GetHash())
 				{
 					UE_LOG(LogShaders, Warning, TEXT("Material Expression UB mismatch!"));
 					DumpUB(DebugUniformExpressionUBLayout);
-					DumpUB(*UniformExpressionCache->LocalUniformBuffer.WorkArea->Layout);
+					DumpUB(UniformExpressionCache->UniformBuffer->GetLayout());
 					bUniformExpressionSetMismatch = true;
 				}
 			}
 		}
-		else
+		if (bUniformExpressionSetMismatch)
 		{
-			if (DebugUniformExpressionUBLayout.GetHash() != UniformExpressionCache->UniformBuffer->GetLayout().GetHash())
-			{
-				UE_LOG(LogShaders, Warning, TEXT("Material Expression UB mismatch!"));
-				DumpUB(DebugUniformExpressionUBLayout);
-				DumpUB(UniformExpressionCache->UniformBuffer->GetLayout());
-				bUniformExpressionSetMismatch = true;
-			}
+			UE_LOG(
+				LogShaders,
+				Fatal,	// TEMP workaround only!!!!
+				TEXT("%s shader uniform expression set mismatch for material %s/%s.\n")
+				TEXT("Shader compilation info:                %s\n")
+				TEXT("Material render proxy compilation info: %s\n")
+				TEXT("Shader uniform expression set:   %u vectors, %u scalars, %u 2D textures, %u cube textures, %u scalars/frame, %u vectors/frame, shader map %p\n")
+				TEXT("Material uniform expression set: %u vectors, %u scalars, %u 2D textures, %u cube textures, %u scalars/frame, %u vectors/frame, shader map %p\n"),
+				GetType()->GetName(),
+				*MaterialRenderProxy->GetFriendlyName(),
+				*Material.GetFriendlyName(),
+				*DebugDescription,
+				*Material.GetRenderingThreadShaderMap()->GetDebugDescription(),
+				DebugUniformExpressionSet.NumVectorExpressions,
+				DebugUniformExpressionSet.NumScalarExpressions,
+				DebugUniformExpressionSet.Num2DTextureExpressions,
+				DebugUniformExpressionSet.NumCubeTextureExpressions,
+				DebugUniformExpressionSet.NumPerFrameScalarExpressions,
+				DebugUniformExpressionSet.NumPerFrameVectorExpressions,
+				UniformExpressionCache->CachedUniformExpressionShaderMap,
+				MaterialUniformExpressionSet.UniformVectorExpressions.Num(),
+				MaterialUniformExpressionSet.UniformScalarExpressions.Num(),
+				MaterialUniformExpressionSet.Uniform2DTextureExpressions.Num(),
+				MaterialUniformExpressionSet.UniformCubeTextureExpressions.Num(),
+				MaterialUniformExpressionSet.PerFrameUniformScalarExpressions.Num(),
+				MaterialUniformExpressionSet.PerFrameUniformVectorExpressions.Num(),
+				Material.GetRenderingThreadShaderMap()
+				);
 		}
 	}
-	if (bUniformExpressionSetMismatch)
-	{
-		UE_LOG(
-			LogShaders,
-			Fatal,	// TEMP workaround only!!!!
-			TEXT("%s shader uniform expression set mismatch for material %s/%s.\n")
-			TEXT("Shader compilation info:                %s\n")
-			TEXT("Material render proxy compilation info: %s\n")
-			TEXT("Shader uniform expression set:   %u vectors, %u scalars, %u 2D textures, %u cube textures, %u scalars/frame, %u vectors/frame, shader map %p\n")
-			TEXT("Material uniform expression set: %u vectors, %u scalars, %u 2D textures, %u cube textures, %u scalars/frame, %u vectors/frame, shader map %p\n"),
-			GetType()->GetName(),
-			*MaterialRenderProxy->GetFriendlyName(),
-			*Material.GetFriendlyName(),
-			*DebugDescription,
-			*Material.GetRenderingThreadShaderMap()->GetDebugDescription(),
-			DebugUniformExpressionSet.NumVectorExpressions,
-			DebugUniformExpressionSet.NumScalarExpressions,
-			DebugUniformExpressionSet.Num2DTextureExpressions,
-			DebugUniformExpressionSet.NumCubeTextureExpressions,
-			DebugUniformExpressionSet.NumPerFrameScalarExpressions,
-			DebugUniformExpressionSet.NumPerFrameVectorExpressions,
-			UniformExpressionCache->CachedUniformExpressionShaderMap,
-			MaterialUniformExpressionSet.UniformVectorExpressions.Num(),
-			MaterialUniformExpressionSet.UniformScalarExpressions.Num(),
-			MaterialUniformExpressionSet.Uniform2DTextureExpressions.Num(),
-			MaterialUniformExpressionSet.UniformCubeTextureExpressions.Num(),
-			MaterialUniformExpressionSet.PerFrameUniformScalarExpressions.Num(),
-			MaterialUniformExpressionSet.PerFrameUniformVectorExpressions.Num(),
-			Material.GetRenderingThreadShaderMap()
-			);
-	}
 #endif
 
-	if (UniformExpressionCache->LocalUniformBuffer.IsValid())
-	{
-		// Set the material uniform buffer.
-		SetLocalUniformBufferParameter(RHICmdList, ShaderRHI, MaterialUniformBuffer, UniformExpressionCache->LocalUniformBuffer);
-	}
-	else
-	{
-		// Set the material uniform buffer.
-		SetUniformBufferParameter(RHICmdList, ShaderRHI, MaterialUniformBuffer, UniformExpressionCache->UniformBuffer);
-	}
 
 	{
 		const TArray<FGuid>& ParameterCollections = UniformExpressionCache->ParameterCollections;
@@ -258,6 +248,8 @@ void FMaterialShader::SetParameters(
 
 		if (NumScalarExpressions > 0 || NumVectorExpressions > 0)
 		{
+
+			const FUniformExpressionSet& MaterialUniformExpressionSet = Material.GetRenderingThreadShaderMap()->GetUniformExpressionSet();
 			FMaterialRenderContext MaterialRenderContext(MaterialRenderProxy, Material, &View);
 			MaterialRenderContext.Time = View.Family->CurrentWorldTime;
 			MaterialRenderContext.RealTime = View.Family->CurrentRealTime;
@@ -318,8 +310,6 @@ void FMaterialShader::SetParameters(
 
 	DeferredParameters.Set(RHICmdList, ShaderRHI, View, TextureMode);
 
-	AtmosphericFogTextureParameters.Set(RHICmdList, ShaderRHI, View);
-
 	if (FeatureLevel >= ERHIFeatureLevel::SM4)
 	{
 		// for copied scene color
@@ -342,35 +332,11 @@ void FMaterialShader::SetParameters(
 		SetTextureParameter(RHICmdList, ShaderRHI, EyeAdaptation, EyeAdaptationTex);
 	}
 
-	if (PerlinNoiseGradientTexture.IsBound() && IsValidRef(GSystemTextures.PerlinNoiseGradient))
+	if (bUniformExpressionCacheNeedsDelete)
 	{
-		const FTexture2DRHIRef& Texture = (FTexture2DRHIRef&)GSystemTextures.PerlinNoiseGradient->GetRenderTargetItem().ShaderResourceTexture;
-		// Bind the PerlinNoiseGradientTexture as a texture
-		SetTextureParameter(
-			RHICmdList,
-			ShaderRHI,
-			PerlinNoiseGradientTexture,
-			PerlinNoiseGradientTextureSampler,
-			TStaticSamplerState<SF_Point,AM_Wrap,AM_Wrap,AM_Wrap>::GetRHI(),
-			Texture
-			);
+		UniformExpressionCache->~FUniformExpressionCache();
+		TempUniformCache.Free(UniformExpressionCache);
 	}
-
-	if (PerlinNoise3DTexture.IsBound() && IsValidRef(GSystemTextures.PerlinNoise3D))
-	{
-		const FTexture3DRHIRef& Texture = (FTexture3DRHIRef&)GSystemTextures.PerlinNoise3D->GetRenderTargetItem().ShaderResourceTexture;
-		// Bind the PerlinNoise3DTexture as a texture
-		SetTextureParameter(
-			RHICmdList,
-			ShaderRHI,
-			PerlinNoise3DTexture,
-			PerlinNoise3DTextureSampler,
-			TStaticSamplerState<SF_Bilinear,AM_Wrap,AM_Wrap,AM_Wrap>::GetRHI(),
-			Texture
-			);
-	}
-
-	GlobalDistanceFieldParameters.Set(RHICmdList, ShaderRHI, static_cast<const FViewInfo&>(View).GlobalDistanceFieldInfo.ParameterData);
 }
 
 // Doxygen struggles to parse these explicit specializations. Just ignore them for now.
@@ -425,18 +391,12 @@ bool FMaterialShader::Serialize(FArchive& Ar)
 		Ar << DebugUniformExpressionUBLayout.Resources;
 	}
 	Ar << DebugDescription;
-	Ar << AtmosphericFogTextureParameters;
 	Ar << EyeAdaptation;
-	Ar << PerlinNoiseGradientTexture;
-	Ar << PerlinNoiseGradientTextureSampler;
-	Ar << PerlinNoise3DTexture;
-	Ar << PerlinNoise3DTextureSampler;
 
 	Ar << PerFrameScalarExpressions;
 	Ar << PerFrameVectorExpressions;
 	Ar << PerFramePrevScalarExpressions;
 	Ar << PerFramePrevVectorExpressions;
-	Ar << GlobalDistanceFieldParameters;
 
 	return bShaderHasOutdatedParameters;
 }
diff --git a/Engine/Source/Runtime/Renderer/Private/ShaderComplexityRendering.cpp b/Engine/Source/Runtime/Renderer/Private/ShaderComplexityRendering.cpp
--- a/Engine/Source/Runtime/Renderer/Private/ShaderComplexityRendering.cpp
+++ b/Engine/Source/Runtime/Renderer/Private/ShaderComplexityRendering.cpp
@@ -19,6 +19,8 @@ void TComplexityAccumulatePS<bQuadComplexity>::SetParameters(
 	FRHICommandList& RHICmdList, 
 	const FShader* OriginalVS, 
 	const FShader* OriginalPS, 
+	const FMaterialRenderProxy* MaterialRenderProxy,
+	const FMaterial& Material,
 	const FSceneView& View
 	)
 {
diff --git a/Engine/Source/Runtime/Renderer/Private/ShaderComplexityRendering.h b/Engine/Source/Runtime/Renderer/Private/ShaderComplexityRendering.h
--- a/Engine/Source/Runtime/Renderer/Private/ShaderComplexityRendering.h
+++ b/Engine/Source/Runtime/Renderer/Private/ShaderComplexityRendering.h
@@ -7,7 +7,7 @@ ShaderComplexityRendering.h: Declarations used for the shader complexity viewmod
 #pragma once
 
 template <bool bQuadComplexity>
-class TComplexityAccumulatePS : public FDebugViewModePS
+class TComplexityAccumulatePS : public FGlobalShader, public IDebugViewModePSInterface
 {
 	DECLARE_SHADER_TYPE(TComplexityAccumulatePS,Global);
 public:
@@ -18,7 +18,7 @@ public:
 	}
 
 	TComplexityAccumulatePS(const ShaderMetaType::CompiledShaderInitializerType& Initializer):
-		FDebugViewModePS(Initializer)
+		FGlobalShader(Initializer)
 	{
 		NormalizedComplexity.Bind(Initializer.ParameterMap,TEXT("NormalizedComplexity"));
 		ShowQuadOverdraw.Bind(Initializer.ParameterMap,TEXT("bShowQuadOverdraw"));
@@ -45,6 +45,8 @@ public:
 		FRHICommandList& RHICmdList, 
 		const FShader* OriginalVS, 
 		const FShader* OriginalPS, 
+		const FMaterialRenderProxy* MaterialRenderProxy,
+		const FMaterial& Material,
 		const FSceneView& View
 		) override;
 
@@ -59,6 +61,8 @@ public:
 
 	virtual void SetMesh(FRHICommandList& RHICmdList, const FSceneView& View) override {}
 
+	virtual FShader* GetShader() override { return static_cast<FShader*>(this); }
+
 private:
 
 	FShaderParameter NormalizedComplexity;
diff --git a/Engine/Source/Runtime/Renderer/Private/ShadowRendering.cpp b/Engine/Source/Runtime/Renderer/Private/ShadowRendering.cpp
--- a/Engine/Source/Runtime/Renderer/Private/ShadowRendering.cpp
+++ b/Engine/Source/Runtime/Renderer/Private/ShadowRendering.cpp
@@ -1281,8 +1281,9 @@ void FShadowDepthDrawingPolicy<bRenderingReflectiveShadowMaps>::DrawMesh(FRHICom
 			if(BatchElement.IndexBuffer)
 			{
 				check(BatchElement.IndexBuffer->IsInitialized());
-				if (BatchElement.InstanceRuns)
+				if (BatchElement.bIsInstanceRuns)
 				{
+					checkSlow(BatchElement.bIsInstanceRuns);
 					if (bUsePositionOnlyVS)
 					{
 						for (uint32 Run = 0; Run < BatchElement.NumInstances; Run++)
