Copying //UE4/Dev-Rendering to Dev-Main (//UE4/Dev-Main)

==========================
MAJOR FEATURES + CHANGES
==========================

Change 2868708 on 2016/02/16 by [REDACTED_PERSON]

	Only add primitives to separate translucency draw list if ST is enabled; shouldn't make any visual difference, but save some time.

Change 2868751 on 2016/02/16 by [REDACTED_PERSON]

	fixed endless driver iteration in rare case (disabled code path)

Change 2868793 on 2016/02/16 by [REDACTED_PERSON]

	added missing code for CreateShaderResourceView

Change 2868814 on 2016/02/16 by [REDACTED_PERSON]

	removed unused code, cleanup, renamed parameter to be more clean

Change 2869124 on 2016/02/16 by [REDACTED_PERSON]

	code cleanupof SkinCache, using struct avoid redundant passing of many paramters

Change 2869162 on 2016/02/16 by [REDACTED_PERSON]

	added vectex/indexbuffer SRV caching for upcoming changes

Change 2869339 on 2016/02/16 by [REDACTED_PERSON]

	Rename r.streaming.UseAabb to r.streaming.UseNewMetrics as it nows controls old build and new build.
	Texture Streaming Build now rebuilds the streaming data per static meshes.
	Static Primitives bounds are now computed for each sections and lods.
	Scale and rotation do not increase bounds as before since they are computed on the geomtery transformed in world space.
	Fixed an issue with some debug view modes preventing the visualization in the static mesh editor

Change 2869431 on 2016/02/16 by [REDACTED_PERSON]

	Fix typo in original UE4 code - CreatePrimitiveUniformBufferImmediateà #2058
	GitHub pull request

Change 2869632 on 2016/02/16 by [REDACTED_PERSON]

	Fixed clear coat with movable skylight. Movable skylight + reflection captures is still broken.

Change 2870545 on 2016/02/17 by [REDACTED_PERSON]

	Fixing separate translucency in reflection captures (no post-processing case); clarifying and simplifying checks a little bit

Change 2871326 on 2016/02/17 by [REDACTED_PERSON]

	Better precomputed visibility log

Change 2871329 on 2016/02/17 by [REDACTED_PERSON]

	'r.VisualizeOccludedPrimitives 1' shows statically occluded primitives in dark red

Change 2871354 on 2016/02/17 by [REDACTED_PERSON]

	Nearest-Depth Neighbor upsampling for half res separate translucency
	* Enabled by default when r.SeparateTranslucencyScreenPercentage is 50
	* Produces full resolution silhouettes of opaque objects in front of half res translucency
	* Costs .1ms on PS4

Change 2871363 on 2016/02/17 by [REDACTED_PERSON]

	Particle cutout optimization works on SubUVMovie module as well

Change 2872099 on 2016/02/18 by [REDACTED_PERSON]

	UE4 - Improve the performance of SGameLayerManager::GetAspectRatioInset and perhaps fixed a bug.

Change 2872433 on 2016/02/18 by [REDACTED_PERSON]

	All world position related nodes were broken in UI material.  This was noticed in Noise Node in UI material rendering black - Jira 26931

	The bug resulted from to View vs ResolvedView changes.

Change 2872844 on 2016/02/18 by [REDACTED_PERSON]

	Ensure bDestroyOwnerAfterFade is only used when a Fade has actually be set.

Change 2873272 on 2016/02/19 by [REDACTED_PERSON]

	fixed SkinCache memory layout for tangents after recent change

Change 2873460 on 2016/02/19 by [REDACTED_PERSON]

	UAT - Add option to run commandlets with -AllowStdOutLogVerbosity

Change 2873468 on 2016/02/19 by [REDACTED_PERSON]

	remove RHI flush on texture streaming...this is a dup of a change made elsewhere

Change 2873971 on 2016/02/19 by [REDACTED_PERSON]

	Decals not rendering in orthographic lit views

Change 2874153 on 2016/02/19 by [REDACTED_PERSON]

	Automatically downsampled separate translucency
	* Off by default, controlled with r.SeparateTranslucencyAutoDownsample
	* The separate translucency pass is timed with GPU timestamps, and downsampling is automatically enabled if the smoothed time is greater than r.SeparateTranslucencyDurationDownsampleThreshold
	* Downsampling is then disabled if the GPU time is less than r.SeparateTranslucencyDurationUpsampleThreshold
	* Downsampling state will not change more often than r.SeparateTranslucencyMinDownsampleChangeTime
	* Implemented timestamp render queries for PS4

Change 2874286 on 2016/02/19 by [REDACTED_PERSON]

	Buffer one extra frame for translucency timestamps, to make sure they never block

Change 2874990 on 2016/02/20 by [REDACTED_PERSON]

	DR - Remove old files

Change 2875910 on 2016/02/22 by [REDACTED_PERSON]

	Enabled r.SeparateTranslucencyAutoDownsample for OrionGame

diff --git a/Engine/Shaders/BasePassPixelShader.usf b/Engine/Shaders/BasePassPixelShader.usf
--- a/Engine/Shaders/BasePassPixelShader.usf
+++ b/Engine/Shaders/BasePassPixelShader.usf
@@ -86,7 +86,7 @@ half3 GetImageBasedReflectionLighting(FMaterialPixelParameters MaterialParameter
 	}
 #endif
 
-#if ENABLE_SKY_LIGHT
+#if ENABLE_SKY_LIGHT || MATERIAL_SHADINGMODEL_CLEAR_COAT
 	BRANCH
 	if (SkyLightParameters.y > 0 && SpecularIBL.a < .999f)
 	{
diff --git a/Engine/Shaders/DownsampleDepthPixelShader.usf b/Engine/Shaders/DownsampleDepthPixelShader.usf
--- a/Engine/Shaders/DownsampleDepthPixelShader.usf
+++ b/Engine/Shaders/DownsampleDepthPixelShader.usf
@@ -9,7 +9,7 @@
 float2 ProjectionScaleBias;
 float4 SourceTexelOffsets01;
 float4 SourceTexelOffsets23;
-float MinMaxBlend;
+float UseMaxDepth;
 
 void Main(
 	noperspective float2 InUV : TEXCOORD0,
@@ -23,10 +23,11 @@ void Main(
 	float ViewSpaceZ2 = CalcSceneDepth(InUV + SourceTexelOffsets23.xy);
 	float ViewSpaceZ3 = CalcSceneDepth(InUV + SourceTexelOffsets23.zw);
 
-	// Use the furthest depth to shrink foreground object silhouettes
-	float MinDepth = min(min(ViewSpaceZ0, ViewSpaceZ1), min(ViewSpaceZ2, ViewSpaceZ3));
 	float MaxDepth = max(max(ViewSpaceZ0, ViewSpaceZ1), max(ViewSpaceZ2, ViewSpaceZ3));
-	float Depth = MinDepth*MinMaxBlend + MaxDepth*(1-MinMaxBlend);
+
+	// Max depth shrinks the silhouettes around foreground objects and is conservative for depth testing
+	// Sample 0 has consistent error, use whichever one is requested for this downsample
+	float Depth = UseMaxDepth > 0 ? MaxDepth : ViewSpaceZ0;
 	// Convert view space Z into post projection Z and output post projection Z / W as the depth buffer expects.
 	OutDepth = (Depth * ProjectionScaleBias.x + ProjectionScaleBias.y) / Depth;
 }
diff --git a/Engine/Shaders/GpuSkinCacheComputeShader.usf b/Engine/Shaders/GpuSkinCacheComputeShader.usf
--- a/Engine/Shaders/GpuSkinCacheComputeShader.usf
+++ b/Engine/Shaders/GpuSkinCacheComputeShader.usf
@@ -21,26 +21,26 @@ Buffer<float4> BoneMatrices;
 
 #if GPUSKIN_MORPH_BLEND
 // float3 DeltaPosition, PackedNormal, see FMorphGPUSkinVertex
-Buffer<float> MorphInputStreamBuffer;
-// data offset to start indexing the data in MorphInputStreamBuffer, in float units
-uint MorphInputStreamFloatOffset;
+Buffer<float> MorphBuffer;
+// data offset to start indexing the data in MorphBuffer, in float units
+uint MorphBufferOffset;
 	
-void GetMorphInputStreamBuffer(uint VertexIndex, out float3 DeltaPosition, out half3 DeltaTangentZ)
+void GetMorphBuffer(uint VertexIndex, out float3 DeltaPosition, out half3 DeltaTangentZ)
 {
 	// see CPU code: "check(MorphStride == sizeof(float) * 6);"
 	// in floats
 	// 3 + 3 floats because of the format defined in InitMorphVertexFactoryComponents / FMorphGPUSkinVertex
-	uint Offset = MorphInputStreamFloatOffset + VertexIndex * (3 + 3);
+	uint Offset = MorphBufferOffset + VertexIndex * (3 + 3);
 
 	DeltaPosition = float3(
-		MorphInputStreamBuffer[Offset + 0],
-		MorphInputStreamBuffer[Offset + 1],
-		MorphInputStreamBuffer[Offset + 2]);
+		MorphBuffer[Offset + 0],
+		MorphBuffer[Offset + 1],
+		MorphBuffer[Offset + 2]);
 
 	DeltaTangentZ = float3(
-		MorphInputStreamBuffer[Offset + 3],
-		MorphInputStreamBuffer[Offset + 4],
-		MorphInputStreamBuffer[Offset + 5]);
+		MorphBuffer[Offset + 3],
+		MorphBuffer[Offset + 4],
+		MorphBuffer[Offset + 5]);
 }
 #endif
 
@@ -150,7 +150,7 @@ FVertexUnpacked UnpackedVertex(uint VertexIndex, uint Offset )
 	Unpacked.BlendWeights2.w = float(BlendWeights >> 24 & 0xff) / 255.0f;
 #endif
 #if GPUSKIN_MORPH_BLEND
-	GetMorphInputStreamBuffer(VertexIndex, Unpacked.DeltaPosition, Unpacked.DeltaTangentZ);
+	GetMorphBuffer(VertexIndex, Unpacked.DeltaPosition, Unpacked.DeltaTangentZ);
 #endif
 	return Unpacked;
 }
diff --git a/Engine/Shaders/GpuSkinCommon.usf b/Engine/Shaders/GpuSkinCommon.usf
--- a/Engine/Shaders/GpuSkinCommon.usf
+++ b/Engine/Shaders/GpuSkinCommon.usf
@@ -72,6 +72,7 @@ float4 UnpackS8x4N(uint Packed)
 }
 
 // @param Unpacked -1..1
+// @return layout 0xWWZZYYXX where 0..0xff is mapped to -1..1
 uint PackS8x4N(float4 Unpacked)
 {
 	uint Packed = asuint(
@@ -83,8 +84,12 @@ uint PackS8x4N(float4 Unpacked)
 }
 
 // @param Unpacked -1..1
-// Leaves the high 8 bits unset
+// @return layout 0x00ZZYYXX where 0..0xff is mapped to -1..1
 uint PackS8x3N(float3 Unpacked)
 {
-	return PackS8x4N(float4(Unpacked, 0));
+	uint Packed = asuint(
+		(uint((Unpacked.x + 1.0f) * 127.5f) << 0) |
+		(uint((Unpacked.y + 1.0f) * 127.5f) << 8) |
+		(uint((Unpacked.z + 1.0f) * 127.5f) << 16));
+	return Packed;
 }
\ No newline at end of file
diff --git a/Engine/Shaders/PostProcessBokehDOF.usf b/Engine/Shaders/PostProcessBokehDOF.usf
--- a/Engine/Shaders/PostProcessBokehDOF.usf
+++ b/Engine/Shaders/PostProcessBokehDOF.usf
@@ -7,6 +7,7 @@
 #include "Common.usf"
 #include "PostProcessCommon.usf"
 #include "DepthOfFieldCommon.usf"
+#include "SeparateTranslucency.usf"
 
 // for the BokehDOF vertex shader, from postprocessing settings
 // .x: color threshold, .y:size threshold, zw: unused
@@ -21,9 +22,6 @@ float4 KernelSize;
 Texture2D LensTexture;
 SamplerState LensTextureSampler;
 
-// Multiplier for separate translucency render target resolution
-float4 SeparateTranslucencyResMult;
-
 // @return x:layer in front of focal plane, y: layer behind focal plane  1-x-y:layer in focus
 float2 ComputeLayerContributions(float Depth)
 {
@@ -377,7 +375,6 @@ void MainPS(
 	OutColor = InColor * Kernel;
 }
 
-
 // pixel shader to combine the full res scene and the blurred images behind and in front of the the focal plane
 void MainRecombinePS(
 	noperspective in float4 UVAndScreenPos : TEXCOORD0,
@@ -441,8 +438,7 @@ void MainRecombinePS(
 #endif
 
 #if RECOMBINE_METHOD == 2 || RECOMBINE_METHOD == 3
-	// sampling from pixel centers, so we need to offset if we're using lower resolution separate translucency
-	float4 SeparateTranslucency = Texture2DSample(PostprocessInput2, PostprocessInput2Sampler, PixelPosCenter * PostprocessInput2Size.zw * SeparateTranslucencyResMult.xxxx );
+	float4 SeparateTranslucency = UpsampleSeparateTranslucency(SvPosition.xy, FullResUV);
 
 	// add RGB, darken by A (this allows to represent translucent and additive blending)
 	OutColor.rgb = OutColor.rgb * SeparateTranslucency.a + SeparateTranslucency.rgb;
diff --git a/Engine/Shaders/PostProcessDOF.usf b/Engine/Shaders/PostProcessDOF.usf
--- a/Engine/Shaders/PostProcessDOF.usf
+++ b/Engine/Shaders/PostProcessDOF.usf
@@ -208,6 +208,7 @@ void MainRecombinePS(
 
 	// apply SeparateTransluceny (todo: shader permutation to save a bit of performance if the feature is disabled)
 	{
+		//@todo - use UpsampleSeparateTranslucency
 		float4 SeparateTranslucency = Texture2DSample(PostprocessInput3, PostprocessInput3Sampler, FullResUV);
 		       
 		// add RGB, darken by A (this allows to represent translucent and additive blending)
diff --git a/Engine/Shaders/ReflectionEnvironmentShaders.usf b/Engine/Shaders/ReflectionEnvironmentShaders.usf
--- a/Engine/Shaders/ReflectionEnvironmentShaders.usf
+++ b/Engine/Shaders/ReflectionEnvironmentShaders.usf
@@ -461,6 +461,15 @@ void ReflectionApplyPS(
 	// Only light pixels marked as lit
 	if( GBuffer.ShadingModelID > 0 && GBuffer.ShadingModelID != SHADINGMODELID_HAIR )
 	{
+		if( GBuffer.ShadingModelID == SHADINGMODELID_CLEAR_COAT )
+		{
+			const float ClearCoat			= GBuffer.CustomData.x;
+			const float ClearCoatRoughness	= GBuffer.CustomData.y;
+
+			GBuffer.Roughness = lerp( GBuffer.Roughness, ClearCoatRoughness, ClearCoat );
+			GBuffer.SpecularColor = lerp( GBuffer.SpecularColor, 0.04, ClearCoat );
+		}
+
 		// TODO if SpecularColor and Roughness were in same texture would remove 2 fetches
 		float3 V = -normalize( mul( float4(UVAndScreenPos.zw, 1, 0), View.ScreenToWorld ).xyz );
 		float NoV = saturate( dot( GBuffer.WorldNormal, V ) );
@@ -494,46 +503,10 @@ void ReflectionApplyPS(
 
 #if APPLY_SSR
 		float4 SSR = Texture2DSample( ScreenSpaceReflectionsTexture, ScreenSpaceReflectionsSampler, UV );
+		SpecularLighting.rgb = SSR.rgb + SpecularLighting.rgb * (1 - SSR.a);
 #endif
-
-		BRANCH
-		if( GBuffer.ShadingModelID == SHADINGMODELID_CLEAR_COAT )
-		{
-			const float ClearCoat			= GBuffer.CustomData.x;
-			const float ClearCoatRoughness	= GBuffer.CustomData.y;
-
-			// TODO EnvBRDF should have a mask param
-			float2 AB = PreIntegratedGF.SampleLevel( PreIntegratedGFSampler, float2( NoV, GBuffer.Roughness ), 0 ).rg;
-			SpecularColor = GBuffer.SpecularColor * AB.x + AB.y * saturate( 50 * GBuffer.SpecularColor.g ) * (1 - ClearCoat);
-
-			// F_Schlick
-			float F0 = 0.04;
-			float Fc = Pow5( 1 - NoV );
-			float F = Fc + (1 - Fc) * F0;
-			F *= ClearCoat;
-
-			float LayerAttenuation = (1 - F);
-			SpecularColor *= LayerAttenuation;
-			SpecularColor *= SpecularOcclusion;
-
-			#if APPLY_SSR
-				SpecularLighting.rgb *= lerp( 1 - SSR.a, 1, ClearCoat );
-				SpecularLighting.rgb += lerp( SSR.rgb, 0, ClearCoat );
-			#endif
-
-			SpecularLighting.rgb *= SpecularColor;
-
-			#if APPLY_SSR
-				SpecularLighting.rgb += SSR.rgb * F;
-			#endif
-		}
-		else
-		{
-			#if APPLY_SSR
-				SpecularLighting.rgb = SSR.rgb + SpecularLighting.rgb * (1 - SSR.a);
-			#endif
-			SpecularLighting.rgb *= SpecularColor;
-		}
+		
+		SpecularLighting.rgb *= SpecularColor;
 
 #if APPLY_SKYLIGHT
 
@@ -546,26 +519,6 @@ void ReflectionApplyPS(
 			float3 SkyLighting = GetSkyLightReflectionSupportingBlend(ReflectionVector, GBuffer.Roughness, bNormalize);
 			SkyLighting *= SpecularColor;
 
-			BRANCH
-			if( GBuffer.ShadingModelID == SHADINGMODELID_CLEAR_COAT )
-			{
-				const float ClearCoat			= GBuffer.CustomData.x;
-				const float ClearCoatRoughness	= GBuffer.CustomData.y;
-
-				// F_Schlick
-				float F0 = 0.04;
-				float Fc = Pow5( 1 - NoV );
-				float F = Fc + (1 - Fc) * F0;
-				F *= ClearCoat;
-
-				#if APPLY_SSR
-					F *= 1 - SSR.a;
-				#endif
-
-				float3 SampleColor = GetSkyLightReflectionSupportingBlend(ReflectionVector, ClearCoatRoughness, bNormalize);	
-				SkyLighting += SampleColor * F;
-			}
-
 		    float DirectionalOcclusion = ScreenSpaceData.DirectionalOcclusion.g;
 		    SkyLighting *= DirectionalOcclusion;
 
diff --git a/Engine/Shaders/SeparateTranslucency.usf b/Engine/Shaders/SeparateTranslucency.usf
new file mode 100644
--- /dev/null
+++ b/Engine/Shaders/SeparateTranslucency.usf
@@ -0,0 +1,78 @@
+// Copyright 1998-2016 Epic Games, Inc. All Rights Reserved.
+
+/*=============================================================================
+	SeparateTranslucency.usf: 
+=============================================================================*/
+
+Texture2D<float> LowResDepthTexture;
+
+void UpdateNearestSample(float Z, float2 UV, float FullResZ, inout float MinDist, inout float2 NearestUV)
+{
+    float DepthDelta = abs(Z - FullResZ);
+
+	FLATTEN
+    if (DepthDelta < MinDist)
+    {
+        MinDist = DepthDelta;
+        NearestUV = UV;
+    }
+}
+
+//@todo - support for all platforms, just skip the GatherRed optimization where not supported
+#define SUPPORTS_NEAREST_DEPTH_NEIGHBOR_UPSAMPLE (SM5_PROFILE && !METAL_SM5_PROFILE || PS4_PROFILE)
+
+float4 UpsampleSeparateTranslucency(float2 Position, float2 UV)
+{
+#if NEAREST_DEPTH_NEIGHBOR_UPSAMPLE && SUPPORTS_NEAREST_DEPTH_NEIGHBOR_UPSAMPLE
+
+	// The relative depth comparison breaks down at larger distances and incorrectly causes point sampling on the skybox pixels
+	float MaxOperationDepth = 2000000.0f;
+
+	// Note: this upsample is specialized for half res to full res
+    float4 LowResDepthBuffer = LowResDepthTexture.GatherRed(GetBilinearClampedSampler(), UV * .5f);
+    float4 LowResDepth = min(float4(ConvertFromDeviceZ(LowResDepthBuffer.x), ConvertFromDeviceZ(LowResDepthBuffer.y), ConvertFromDeviceZ(LowResDepthBuffer.z), ConvertFromDeviceZ(LowResDepthBuffer.w)), MaxOperationDepth.xxxx);
+	float FullResDepth = min(ConvertFromDeviceZ(SceneDepthTexture[uint2(Position.xy)].x), MaxOperationDepth);
+
+	float RelativeDepthThreshold = .1f;
+
+	// Search for the UV of the low res neighbor whose depth is closest to the full res depth
+    float MinDist = 1.e8f;
+    float2 LowResTexelSize = PostprocessInput2Size.zw;
+
+    float2 UV00 = UV - 0.5f * LowResTexelSize;
+    float2 NearestUV = UV00;
+    UpdateNearestSample(LowResDepth.w, UV00, FullResDepth, MinDist, NearestUV);
+
+    float2 UV10 = float2(UV00.x + LowResTexelSize.x, UV00.y);
+    UpdateNearestSample(LowResDepth.z, UV10, FullResDepth, MinDist, NearestUV);
+
+    float2 UV01 = float2(UV00.x, UV00.y + LowResTexelSize.y);
+    UpdateNearestSample(LowResDepth.x, UV01, FullResDepth, MinDist, NearestUV);
+
+    float2 UV11 = float2(UV00.x + LowResTexelSize.x, UV00.y + LowResTexelSize.y);
+    UpdateNearestSample(LowResDepth.y, UV11, FullResDepth, MinDist, NearestUV);
+	 
+    float4 Output = 0.0f;
+	float InvFullResDepth = 1.0f / FullResDepth;
+
+    BRANCH
+    if (abs(LowResDepth.w - FullResDepth) * InvFullResDepth < RelativeDepthThreshold 
+		&& abs(LowResDepth.z - FullResDepth) * InvFullResDepth < RelativeDepthThreshold 
+		&& abs(LowResDepth.x - FullResDepth) * InvFullResDepth < RelativeDepthThreshold
+		&& abs(LowResDepth.y - FullResDepth) * InvFullResDepth < RelativeDepthThreshold)
+    {
+        Output = Texture2DSampleLevel(PostprocessInput2, PostprocessInput2Sampler, UV, 0);
+		//Output = float4(1, 0, 0, .1f);
+    }
+    else
+    {
+        Output = Texture2DSampleLevel(PostprocessInput2, GetPointClampedSampler(), NearestUV, 0);
+		//Output = float4(0, 1, 0, .1f);
+    }
+	
+	return Output;
+
+#else
+	return Texture2DSample(PostprocessInput2, PostprocessInput2Sampler, UV);
+#endif
+}
diff --git a/Engine/Shaders/SlateElementPixelShader.usf b/Engine/Shaders/SlateElementPixelShader.usf
--- a/Engine/Shaders/SlateElementPixelShader.usf
+++ b/Engine/Shaders/SlateElementPixelShader.usf
@@ -84,6 +84,7 @@ float4 GetMaterialColor( VertexToPixelInterpolants InVertex, float2 MeshUV )
 
 	Parameters.VertexColor = InVertex.Color;
 
+    ResolvedView = ResolveView();
 	FPixelMaterialInputs PixelMaterialInputs;
 	CalcMaterialParameters( Parameters, PixelMaterialInputs, InVertex.Position, true );
 
diff --git a/Engine/Source/Editor/StaticMeshEditor/Private/SStaticMeshEditorViewport.cpp b/Engine/Source/Editor/StaticMeshEditor/Private/SStaticMeshEditorViewport.cpp
--- a/Engine/Source/Editor/StaticMeshEditor/Private/SStaticMeshEditorViewport.cpp
+++ b/Engine/Source/Editor/StaticMeshEditor/Private/SStaticMeshEditorViewport.cpp
@@ -279,6 +279,10 @@ void SStaticMeshEditorViewport::UpdatePreviewMesh(UStaticMesh* InStaticMesh)
 	PreviewMeshComponent = NewObject<UStaticMeshComponent>();
 
 	PreviewMeshComponent->SetStaticMesh(InStaticMesh);
+
+	// Update streaming data for debug viewmode feedback
+	PreviewMeshComponent->UpdateStreamingTextureInfos(true);
+
 	PreviewScene.AddComponent(PreviewMeshComponent,FTransform::Identity);
 
 	const int32 SocketCount = InStaticMesh->Sockets.Num();
diff --git a/Engine/Source/Programs/AutomationTool/AutomationUtils/Automation.cs b/Engine/Source/Programs/AutomationTool/AutomationUtils/Automation.cs
--- a/Engine/Source/Programs/AutomationTool/AutomationUtils/Automation.cs
+++ b/Engine/Source/Programs/AutomationTool/AutomationUtils/Automation.cs
@@ -120,7 +120,8 @@ namespace AutomationTool
 		public static CommandLineArg NoKill = new CommandLineArg("-NoKill");
 		public static CommandLineArg Installed = new CommandLineArg("-Installed");
 		public static CommandLineArg UTF8Output = new CommandLineArg("-UTF8Output");
-		public static CommandLineArg NoAutoSDK = new CommandLineArg("-NoAutoSDK");
+        public static CommandLineArg AllowStdOutLogVerbosity = new CommandLineArg("-AllowStdOutLogVerbosity");
+        public static CommandLineArg NoAutoSDK = new CommandLineArg("-NoAutoSDK");
 		public static CommandLineArg IgnoreJunk = new CommandLineArg("-ignorejunk");
         /// <summary>
         /// Allows you to use local storage for your root build storage dir (default of P:\Builds (on PC) is changed to Engine\Saved\LocalBuilds). Used for local testing.
diff --git a/Engine/Source/Programs/AutomationTool/AutomationUtils/CommandletUtils.cs b/Engine/Source/Programs/AutomationTool/AutomationUtils/CommandletUtils.cs
--- a/Engine/Source/Programs/AutomationTool/AutomationUtils/CommandletUtils.cs
+++ b/Engine/Source/Programs/AutomationTool/AutomationUtils/CommandletUtils.cs
@@ -238,7 +238,7 @@ namespace AutomationTool
 				String.IsNullOrEmpty(Parameters) ? "" : Parameters,
 				CommandUtils.MakePathSafeToUseWithCommandLine(LocalLogFile),
 				IsBuildMachine ? "-buildmachine" : "",
-				GlobalCommandLine.Verbose ? "-AllowStdOutLogVerbosity " : ""
+                (GlobalCommandLine.Verbose || GlobalCommandLine.AllowStdOutLogVerbosity) ? "-AllowStdOutLogVerbosity " : ""
 			);
 			ERunOptions Opts = ERunOptions.Default;
 			if (GlobalCommandLine.UTF8Output)
diff --git a/Engine/Source/Programs/UnrealLightmass/Private/Lighting/LightingSystem.cpp b/Engine/Source/Programs/UnrealLightmass/Private/Lighting/LightingSystem.cpp
--- a/Engine/Source/Programs/UnrealLightmass/Private/Lighting/LightingSystem.cpp
+++ b/Engine/Source/Programs/UnrealLightmass/Private/Lighting/LightingSystem.cpp
@@ -1112,7 +1112,7 @@ void FStaticLightingSystem::DumpStats(float TotalStaticLightingTime) const
 		const uint64 NumQueriesVisible = Stats.NumQueriesVisibleByDistanceRatio + Stats.NumQueriesVisibleExplicitSampling + Stats.NumQueriesVisibleImportanceSampling;
 		const uint64 TotalNumQueries = Stats.NumPrecomputedVisibilityQueries + Stats.NumPrecomputedVisibilityGroupQueries;
 		SolverStats += FString::Printf( TEXT("   %.3f million mesh queries, %.3f million group queries, %.1f%% visible, (%.1f%% trivially visible, %.1f%% explicit sampling, %.1f%% importance sampling)\n"), Stats.NumPrecomputedVisibilityQueries / 1000000.0f, Stats.NumPrecomputedVisibilityGroupQueries / 1000000.0f, 100.0f * NumQueriesVisible / TotalNumQueries, 100.0f * Stats.NumQueriesVisibleByDistanceRatio / NumQueriesVisible, 100.0f * Stats.NumQueriesVisibleExplicitSampling / NumQueriesVisible, 100.0f * Stats.NumQueriesVisibleImportanceSampling / NumQueriesVisible);
-		SolverStats += FString::Printf( TEXT("   %ux%ux%u group cells with %u occupied, %u meshes excluded, %.3f million mesh queries skipped\n"), GroupVisibilityGridSizeXY, GroupVisibilityGridSizeXY, GroupVisibilityGridSizeZ, VisibilityGroups.Num(), Stats.NumPrecomputedVisibilityMeshesExcludedFromGroups, Stats.NumPrecomputedVisibilityMeshQueriesSkipped / 1000000.0f);
+		SolverStats += FString::Printf( TEXT("   %ux%ux%u group cells with %u occupied, %u meshes individually queried, %.3f million mesh queries skipped\n"), GroupVisibilityGridSizeXY, GroupVisibilityGridSizeXY, GroupVisibilityGridSizeZ, VisibilityGroups.Num(), Stats.NumPrecomputedVisibilityMeshesExcludedFromGroups, Stats.NumPrecomputedVisibilityMeshQueriesSkipped / 1000000.0f);
 	}
 
 	// Send the message in multiple parts since it cuts off in the middle otherwise
diff --git a/Engine/Source/Runtime/Core/Private/Windows/WindowsPlatformMisc.cpp b/Engine/Source/Runtime/Core/Private/Windows/WindowsPlatformMisc.cpp
--- a/Engine/Source/Runtime/Core/Private/Windows/WindowsPlatformMisc.cpp
+++ b/Engine/Source/Runtime/Core/Private/Windows/WindowsPlatformMisc.cpp
@@ -2593,7 +2593,7 @@ FGPUDriverInfo FWindowsPlatformMisc::GetGPUDriverInfo(const FString& DeviceDescr
 
 	if(bIterateAvailableAndChoose)
 	{
-		for(uint32 i = 0;; ++i)
+		for(uint32 i = 0; i < 256; ++i)
 		{
 			// Iterate all installed display adapters
 			const FString Key = FString::Printf(TEXT("SYSTEM\\CurrentControlSet\\Control\\Class\\{4D36E968-E325-11CE-BFC1-08002BE10318}\\%04d"), i);
@@ -2603,7 +2603,9 @@ FGPUDriverInfo FWindowsPlatformMisc::GetGPUDriverInfo(const FString& DeviceDescr
 
 			if(!Local.IsValid())
 			{
+				// last device or error
 				DebugString += TEXT("GetVideoDriverDetailsInvalid ");
+				break;
 			}
 
 			if(Local.DeviceDescription == DeviceDescription)
diff --git a/Engine/Source/Runtime/Engine/Classes/Components/PrimitiveComponent.h b/Engine/Source/Runtime/Engine/Classes/Components/PrimitiveComponent.h
--- a/Engine/Source/Runtime/Engine/Classes/Components/PrimitiveComponent.h
+++ b/Engine/Source/Runtime/Engine/Classes/Components/PrimitiveComponent.h
@@ -19,11 +19,10 @@ struct FConvexVolume;
 struct FNavigableGeometryExport;
 
 /** Information about a streaming texture that a primitive uses for rendering. */
+USTRUCT()
 struct FStreamingTexturePrimitiveInfo
 {
-	UTexture* Texture;
-	FBoxSphereBounds Bounds;
-	float TexelFactor;
+	GENERATED_USTRUCT_BODY()
 
 	FStreamingTexturePrimitiveInfo()
 		: Texture(nullptr)
@@ -31,6 +30,15 @@ struct FStreamingTexturePrimitiveInfo
 		, TexelFactor(1.0f)
 	{
 	}
+
+	UPROPERTY()
+	UTexture* Texture;
+
+	UPROPERTY()
+	FBoxSphereBounds Bounds;
+
+	UPROPERTY()
+	float TexelFactor;
 };
 
 /** Determines whether a Character can attempt to step up onto a component when they walk in to it. */
diff --git a/Engine/Source/Runtime/Engine/Classes/Components/SplineMeshComponent.h b/Engine/Source/Runtime/Engine/Classes/Components/SplineMeshComponent.h
--- a/Engine/Source/Runtime/Engine/Classes/Components/SplineMeshComponent.h
+++ b/Engine/Source/Runtime/Engine/Classes/Components/SplineMeshComponent.h
@@ -328,7 +328,7 @@ class ENGINE_API USplineMeshComponent : public UStaticMeshComponent, public IInt
 	/** Returns a vector which, when componentwise-multiplied by another vector, will zero all the components not corresponding to the supplied ESplineMeshAxis */
 	inline static FVector GetAxisMask(ESplineMeshAxis::Type InAxis);
 
-	virtual bool GetStreamingTextureFactors(float& OutWorldTexelFactor, float& OutWorldLightmapFactor) const override;
+	virtual bool GetStreamingTextureFactors(float& OutWorldTexelFactor, float& OutWorldLightmapFactor, FBoxSphereBounds& OutBounds, int32 LODIndex, int32 ElementIndex) const override;
 };
 
 const float& USplineMeshComponent::GetAxisValue(const FVector& InVector, ESplineMeshAxis::Type InAxis)
diff --git a/Engine/Source/Runtime/Engine/Classes/Components/StaticMeshComponent.h b/Engine/Source/Runtime/Engine/Classes/Components/StaticMeshComponent.h
--- a/Engine/Source/Runtime/Engine/Classes/Components/StaticMeshComponent.h
+++ b/Engine/Source/Runtime/Engine/Classes/Components/StaticMeshComponent.h
@@ -237,6 +237,10 @@ class ENGINE_API UStaticMeshComponent : public UMeshComponent
 	UPROPERTY(transient)
 	TArray<struct FStaticMeshComponentLODInfo> LODData;
 
+	/** The streaming info for each LOD / Section. Needs to be persistent in order to be accessible when loading cooked maps. */
+	UPROPERTY()
+	TArray<FStreamingTexturePrimitiveInfo> StreamingTextureInfos;
+
 #if WITH_EDITORONLY_DATA
 	/** Derived data key of the static mesh, used to determine if an update from the source static mesh is required. */
 	UPROPERTY()
@@ -295,6 +299,7 @@ public:
 		// return IsCollisionEnabled() && (StaticMesh != NULL);
 		return false;
 	}
+
 	//~ End USceneComponent Interface
 
 	//~ Begin UActorComponent Interface.
@@ -322,7 +327,8 @@ public:
 		return LightmassSettings.bShadowIndirectOnly;
 	}
 	virtual ELightMapInteractionType GetStaticLightingType() const override;
-	virtual bool GetStreamingTextureFactors(float& OutWorldTexelFactor, float& OutWorldLightmapFactor) const;
+	void UpdateStreamingTextureInfos(bool bForce = false);
+	virtual bool GetStreamingTextureFactors(float& OutWorldTexelFactor, float& OutWorldLightmapFactor, FBoxSphereBounds& OutBounds, int32 LODIndex, int32 ElementIndex) const;
 	virtual void GetStreamingTextureInfo(TArray<FStreamingTexturePrimitiveInfo>& OutStreamingTextures) const override;
 	virtual class UBodySetup* GetBodySetup() override;
 	virtual bool CanEditSimulatePhysics() override;
diff --git a/Engine/Source/Runtime/Engine/Classes/Engine/LocalPlayer.h b/Engine/Source/Runtime/Engine/Classes/Engine/LocalPlayer.h
--- a/Engine/Source/Runtime/Engine/Classes/Engine/LocalPlayer.h
+++ b/Engine/Source/Runtime/Engine/Classes/Engine/LocalPlayer.h
@@ -252,6 +252,22 @@ public:
 	 */
 	UGameInstance* GetGameInstance() const;
 
+
+	/**
+	* Calculate the view init settings for drawing from this view actor
+	*
+	* @param	OutInitOptions - output view struct. Not every field is initialized, some of them are only filled in by CalcSceneView
+	* @param	Viewport - current client viewport
+	* @param	ViewDrawer - optional drawing in the view
+	* @param	StereoPass - whether we are drawing the full viewport, or a stereo left / right pass
+	* @return	true if the view options were filled in. false in various fail conditions.
+	*/
+	bool CalcSceneViewInitOptions(
+		struct FSceneViewInitOptions& OutInitOptions, 
+		FViewport* Viewport,
+		class FViewElementDrawer* ViewDrawer = NULL,
+		EStereoscopicPass StereoPass = eSSP_FULL);
+
 	/**
 	 * Calculate the view settings for drawing from this view actor
 	 *
diff --git a/Engine/Source/Runtime/Engine/Classes/Engine/StaticMesh.h b/Engine/Source/Runtime/Engine/Classes/Engine/StaticMesh.h
--- a/Engine/Source/Runtime/Engine/Classes/Engine/StaticMesh.h
+++ b/Engine/Source/Runtime/Engine/Classes/Engine/StaticMesh.h
@@ -500,12 +500,25 @@ public:
 	ENGINE_API virtual void ReleaseResources();
 
 	/**
-	 * Returns the scale dependent texture factor used by the texture streaLODg code.
+	 * Returns the scale dependent texture factor used by the texture streaming code.
 	 *
 	 * @param RequestedUVIndex UVIndex to look at
 	 * @return scale dependent texture factor
 	 */
-	float GetStreamingTextureFactor( int32 RequestedUVIndex );
+	float GetStreamingTextureFactor( int32 RequestedUVIndex ) const;
+
+	/**
+	 * Returns the scale dependent texture factor and bound used by the texture streaming code.
+	 *
+	 * @param OutTexelFactor		The requested texel factor
+	 * @param OutTexelFactor		The requested bound for this texel factor
+	 * @param CoordinateIndex		UV Index to look at
+	 * @param LODIndex				LOD index to look at
+	 * @param ElementIndex			Element index to look at
+	 * @param TransformMatrix		Matrix to be applied to the position before computing the bounds
+	 * @return false if some parameters are invalid
+	 */
+	bool GetStreamingTextureFactor( float& OutTexelFactor, FBoxSphereBounds& OutBounds, int32 CoordinateIndex, int32 LODIndex, int32 ElementIndex, const FTransform& Transform ) const;
 
 	/**
 	 * Returns the number of vertices for the specified LOD.
diff --git a/Engine/Source/Runtime/Engine/Classes/Particles/SubUV/ParticleModuleSubUVMovie.h b/Engine/Source/Runtime/Engine/Classes/Particles/SubUV/ParticleModuleSubUVMovie.h
--- a/Engine/Source/Runtime/Engine/Classes/Particles/SubUV/ParticleModuleSubUVMovie.h
+++ b/Engine/Source/Runtime/Engine/Classes/Particles/SubUV/ParticleModuleSubUVMovie.h
@@ -5,7 +5,7 @@
 #include "Particles/SubUV/ParticleModuleSubUV.h"
 #include "ParticleModuleSubUVMovie.generated.h"
 
-UCLASS(editinlinenew, hidecategories=Object, hidecategories=SubUV, meta=(DisplayName = "SubUV Movie"))
+UCLASS(editinlinenew, hidecategories=Object, meta=(DisplayName = "SubUV Movie"))
 class UParticleModuleSubUVMovie : public UParticleModuleSubUV
 {
 	GENERATED_UCLASS_BODY()
diff --git a/Engine/Source/Runtime/Engine/Private/Components/DecalComponent.cpp b/Engine/Source/Runtime/Engine/Private/Components/DecalComponent.cpp
--- a/Engine/Source/Runtime/Engine/Private/Components/DecalComponent.cpp
+++ b/Engine/Source/Runtime/Engine/Private/Components/DecalComponent.cpp
@@ -57,6 +57,8 @@ void FDeferredDecalProxy::InitializeFadingParameters(float AbsSpawnTime, float F
 UDecalComponent::UDecalComponent(const FObjectInitializer& ObjectInitializer)
 	: Super(ObjectInitializer)
 	, FadeScreenSize(0.01)
+	, FadeStartDelay(0.0f)
+	, FadeDuration(0.0f)
 	, bDestroyOwnerAfterFade(true)
 	, DecalSize(128.0f, 256.0f, 256.0f)
 {
@@ -90,7 +92,7 @@ void UDecalComponent::LifeSpanCallback()
 
 	auto* Owner = GetOwner();
 
-	if (bDestroyOwnerAfterFade && Owner)
+	if (bDestroyOwnerAfterFade && Owner && (FadeDuration > 0.0f || FadeStartDelay > 0.0f))
 	{
 		Owner->Destroy();
 	}
diff --git a/Engine/Source/Runtime/Engine/Private/Components/SplineMeshComponent.cpp b/Engine/Source/Runtime/Engine/Private/Components/SplineMeshComponent.cpp
--- a/Engine/Source/Runtime/Engine/Private/Components/SplineMeshComponent.cpp
+++ b/Engine/Source/Runtime/Engine/Private/Components/SplineMeshComponent.cpp
@@ -1161,9 +1161,9 @@ FStaticMeshStaticLightingMesh* USplineMeshComponent::AllocateStaticLightingMesh(
 }
 
 
-bool USplineMeshComponent::GetStreamingTextureFactors(float& OutWorldTexelFactor, float& OutWorldLightmapFactor) const
+bool USplineMeshComponent::GetStreamingTextureFactors(float& OutWorldTexelFactor, float& OutWorldLightmapFactor, FBoxSphereBounds& OutBounds, int32 LODIndex, int32 ElementIndex) const
 {
-	if (UStaticMeshComponent::GetStreamingTextureFactors(OutWorldTexelFactor, OutWorldLightmapFactor))
+	if (UStaticMeshComponent::GetStreamingTextureFactors(OutWorldTexelFactor, OutWorldLightmapFactor, OutBounds, LODIndex, ElementIndex))
 	{
 		// We need to come up with a compensation factor for spline deformed meshes
 
diff --git a/Engine/Source/Runtime/Engine/Private/Components/StaticMeshComponent.cpp b/Engine/Source/Runtime/Engine/Private/Components/StaticMeshComponent.cpp
--- a/Engine/Source/Runtime/Engine/Private/Components/StaticMeshComponent.cpp
+++ b/Engine/Source/Runtime/Engine/Private/Components/StaticMeshComponent.cpp
@@ -495,20 +495,79 @@ void UStaticMeshComponent::OnUnregister()
 	Super::OnUnregister();
 }
 
-bool UStaticMeshComponent::GetStreamingTextureFactors(float& OutWorldTexelFactor, float& OutWorldLightmapFactor) const
+void UStaticMeshComponent::UpdateStreamingTextureInfos(bool bForce)
 {
-	if (StaticMesh && StaticMesh && StaticMesh->RenderData && StaticMesh->RenderData->LODResources.Num() > 0)
+	// Only rebuild the data in editor 
+#if WITH_EDITORONLY_DATA
+	StreamingTextureInfos.Empty();
+
+	if ((bForce || (!bIgnoreInstanceForTextureStreaming && Mobility == EComponentMobility::Static)) && StaticMesh && StaticMesh->RenderData)
 	{
-		const auto FeatureLevel = GetWorld() ? GetWorld()->FeatureLevel : GMaxRHIFeatureLevel;
+		for (int32 LODIndex = 0; LODIndex < StaticMesh->RenderData->LODResources.Num(); ++LODIndex)
+		{
+			const FStaticMeshLODResources& LOD = StaticMesh->RenderData->LODResources[LODIndex];
 
-		bool bHasValidLightmapCoordinates = ((StaticMesh->LightMapCoordinateIndex >= 0)
-			&& ((uint32)StaticMesh->LightMapCoordinateIndex < StaticMesh->RenderData->LODResources[0].VertexBuffer.GetNumTexCoords()));
+			for( int32 ElementIndex = 0; ElementIndex < LOD.Sections.Num(); ElementIndex++ )
+			{
+				FStreamingTexturePrimitiveInfo Info;
 
-		const float LocalTexelFactor	= StaticMesh->GetStreamingTextureFactor(0) * FMath::Max(0.0f, StreamingDistanceMultiplier);
-		const float LocalLightmapFactor	= bHasValidLightmapCoordinates ? StaticMesh->GetStreamingTextureFactor(StaticMesh->LightMapCoordinateIndex) : 1.0f;
+				if (StaticMesh->GetStreamingTextureFactor(Info.TexelFactor, Info.Bounds, 0, LODIndex, ElementIndex, ComponentToWorld))
+				{
+					Info.TexelFactor *= FMath::Max(0.0f, StreamingDistanceMultiplier);
+					StreamingTextureInfos.Push(Info);
+				}
+				else
+				{
+					StreamingTextureInfos.Empty();
+					return;
+				}
+			}
+		}
+	}
+#endif
+}
 
-		OutWorldTexelFactor	= LocalTexelFactor * ComponentToWorld.GetMaximumAxisScale();
-		OutWorldLightmapFactor	= LocalLightmapFactor * ComponentToWorld.GetMaximumAxisScale();
+extern ENGINE_API TAutoConsoleVariable<int32> CVarStreamingUseNewMetrics;
+
+bool UStaticMeshComponent::GetStreamingTextureFactors(float& OutWorldTexelFactor, float& OutWorldLightmapFactor, FBoxSphereBounds& OutBounds, int32 LODIndex, int32 ElementIndex) const
+{
+	bool bUsePreciseStreamingBounds = false;
+
+	if (StaticMesh && StaticMesh->RenderData && LODIndex < StaticMesh->RenderData->LODResources.Num())
+	{
+		TIndirectArray<FStaticMeshLODResources>& LODResources = StaticMesh->RenderData->LODResources;
+
+		const ERHIFeatureLevel::Type FeatureLevel = GetWorld() ? GetWorld()->FeatureLevel : GMaxRHIFeatureLevel;
+
+		const bool bHasValidLightmapCoordinates = StaticMesh->LightMapCoordinateIndex >= 0 &&
+			(uint32)StaticMesh->LightMapCoordinateIndex < LODResources[0].VertexBuffer.GetNumTexCoords();
+
+		// Because there is a single array for each sections and lods, the index needs to be computed depending on the number of section per lod.
+		int32 InfoIndex = ElementIndex;
+		if (StreamingTextureInfos.Num() && LODResources.IsValidIndex(LODIndex) && LODResources[LODIndex].Sections.IsValidIndex(ElementIndex))
+		{
+			int32 DebugLODIndex = 0;
+			while (DebugLODIndex < LODIndex)
+			{
+				InfoIndex += LODResources[DebugLODIndex].Sections.Num();
+				++DebugLODIndex;
+			}
+		}
+
+		if (StreamingTextureInfos.IsValidIndex(InfoIndex) && CVarStreamingUseNewMetrics.GetValueOnAnyThread() > 0)
+		{
+			OutBounds = StreamingTextureInfos[InfoIndex].Bounds;
+			OutWorldTexelFactor = StreamingTextureInfos[InfoIndex].TexelFactor;
+		}
+		else
+		{
+			float LocalTexelFactor = StaticMesh->GetStreamingTextureFactor(0) * FMath::Max(0.0f, StreamingDistanceMultiplier);
+			OutWorldTexelFactor	= LocalTexelFactor * ComponentToWorld.GetMaximumAxisScale();
+			OutBounds = Bounds;
+		}
+
+		const float LocalLightmapFactor	= bHasValidLightmapCoordinates ? StaticMesh->GetStreamingTextureFactor(StaticMesh->LightMapCoordinateIndex) : 1.0f;
+		OutWorldLightmapFactor = LocalLightmapFactor * ComponentToWorld.GetMaximumAxisScale();
 
 		return true;
 	}
@@ -524,10 +583,11 @@ void UStaticMeshComponent::GetStreamingTextureInfo(TArray<FStreamingTexturePrimi
 	float WorldTexelFactor = 1.f;
 	float WorldLightmapFactor = 1.f;
 
-	if ( !bIgnoreInstanceForTextureStreaming && GetStreamingTextureFactors(WorldTexelFactor, WorldLightmapFactor) )
+	if (!bIgnoreInstanceForTextureStreaming && StaticMesh && StaticMesh->RenderData)
 	{
-		const auto FeatureLevel = GetWorld() ? GetWorld()->FeatureLevel : GMaxRHIFeatureLevel;
-		bool bHasValidLightmapCoordinates = ((StaticMesh->LightMapCoordinateIndex >= 0)
+		const ERHIFeatureLevel::Type FeatureLevel = GetWorld() ? GetWorld()->FeatureLevel : GMaxRHIFeatureLevel;
+
+		const bool bHasValidLightmapCoordinates = ((StaticMesh->LightMapCoordinateIndex >= 0)
 			&& ((uint32)StaticMesh->LightMapCoordinateIndex < StaticMesh->RenderData->LODResources[0].VertexBuffer.GetNumTexCoords()));
 
 		for (int32 LODIndex = 0; LODIndex < StaticMesh->RenderData->LODResources.Num(); ++LODIndex)
@@ -544,6 +604,10 @@ void UStaticMeshComponent::GetStreamingTextureInfo(TArray<FStreamingTexturePrimi
 					Material = UMaterial::GetDefaultMaterial(MD_Surface);
 				}
 
+				FBoxSphereBounds SectionBounds;
+				if (!GetStreamingTextureFactors(WorldTexelFactor, WorldLightmapFactor, SectionBounds, LODIndex, ElementIndex))
+					continue;
+
 				// Enumerate the textures used by the material.
 				TArray<UTexture*> Textures;
 
@@ -554,7 +618,7 @@ void UStaticMeshComponent::GetStreamingTextureInfo(TArray<FStreamingTexturePrimi
 				for(int32 TextureIndex = 0;TextureIndex < Textures.Num();TextureIndex++)
 				{
 					FStreamingTexturePrimitiveInfo& StreamingTexture = *new(OutStreamingTextures) FStreamingTexturePrimitiveInfo;
-					StreamingTexture.Bounds = Bounds;
+					StreamingTexture.Bounds = SectionBounds;
 					StreamingTexture.TexelFactor = WorldTexelFactor;
 					StreamingTexture.Texture = Textures[TextureIndex];
 				}
diff --git a/Engine/Source/Runtime/Engine/Private/ContentStreaming.cpp b/Engine/Source/Runtime/Engine/Private/ContentStreaming.cpp
--- a/Engine/Source/Runtime/Engine/Private/ContentStreaming.cpp
+++ b/Engine/Source/Runtime/Engine/Private/ContentStreaming.cpp
@@ -172,10 +172,10 @@ static TAutoConsoleVariable<int32> CVarStreamingShowWantedMips(
 	TEXT("If non-zero, will limit resolution to wanted mip."),
 	ECVF_Cheat);
 
-ENGINE_API TAutoConsoleVariable<int32> CVarStreamingUseAABB(
-	TEXT("r.Streaming.UseAABB"),
+ENGINE_API TAutoConsoleVariable<int32> CVarStreamingUseNewMetrics(
+	TEXT("r.Streaming.UseNewMetrics"),
 	0,
-	TEXT("If non-zero, will use AABB to compute distance."),
+	TEXT("If non-zero, will use tight AABB bounds and improved texture factors."),
 	ECVF_Default);
 
 static TAutoConsoleVariable<int32> CVarStreamingHLODStrategy(
@@ -5522,7 +5522,7 @@ FFloatMipLevel FStreamingHandlerTextureStatic::GetWantedMips( FStreamingManagerT
 	bool bShouldAbortLoop = false;
 	bool bEntryFound = false; // True an entry for this texture exists­.
 
-	const bool bUseAABB = CVarStreamingUseAABB.GetValueOnAnyThread() > 0;
+	const bool bUseNewMetrics = CVarStreamingUseNewMetrics.GetValueOnAnyThread() > 0;
 	const float HiddenScale = CVarStreamingHiddenPrimitiveScale.GetValueOnAnyThread();
 	const float MaxResolution = (float)(0x1 << (StreamingTexture.MaxAllowedMips - 1));
 
@@ -5597,7 +5597,7 @@ FFloatMipLevel FStreamingHandlerTextureStatic::GetWantedMips( FStreamingManagerT
 						const VectorRegister InRangeMask = VectorCompareEQ(DistSq, ClampedDistSq);
 
 						VectorRegister DistSqMinusRadiusSq = VectorZero();
-						if (bUseAABB)
+						if (bUseNewMetrics)
 						{
 							// In this case DistSqMinusRadiusSq will contain the distance to the box^2
 							Temp = VectorSubtract( ViewOriginX, CenterX );
diff --git a/Engine/Source/Runtime/Engine/Private/GPUSkinCache.cpp b/Engine/Source/Runtime/Engine/Private/GPUSkinCache.cpp
--- a/Engine/Source/Runtime/Engine/Private/GPUSkinCache.cpp
+++ b/Engine/Source/Runtime/Engine/Private/GPUSkinCache.cpp
@@ -36,6 +36,7 @@ static FAutoConsoleVariableRef CVarEnableGPUSkinCacheShaders(
 	ECVF_RenderThreadSafe | ECVF_ReadOnly
 	);
 
+// 0/1
 int32 GEnableGPUSkinCache = 0;
 static TAutoConsoleVariable<int32> CVarEnableGPUSkinCache(
 	TEXT("r.SkinCaching"),
@@ -43,7 +44,8 @@ static TAutoConsoleVariable<int32> CVarEnableGPUSkinCache(
 	TEXT("Whether or not to use the GPU compute skinning cache.\n")
 	TEXT("This will perform skinning on a compute job and not skin on the vertex shader.\n")
 	TEXT("Requires r.SkinCacheShaders=1.\n")
-	TEXT("0 is off(default), 1 is on"),
+	TEXT(" 0: off(default)\n")
+	TEXT(" 1: on"),
 	ECVF_RenderThreadSafe
 	);
 
@@ -85,31 +87,33 @@ public:
 
 		SkinInputStreamStride.Bind(Initializer.ParameterMap, TEXT("SkinCacheVertexStride"));
 		SkinInputStreamVertexCount.Bind(Initializer.ParameterMap, TEXT("SkinCacheVertexCount"));
-		SkinInputStreamFloatOffset.Bind(Initializer.ParameterMap, TEXT("SkinCacheInputStreamFloatOffset"));
 		SkinOutputBufferFloatOffset.Bind(Initializer.ParameterMap, TEXT("SkinCacheOutputBufferFloatOffset"));
 		BoneMatrices.Bind(Initializer.ParameterMap, TEXT("BoneMatrices"));
 		SkinInputStream.Bind(Initializer.ParameterMap, TEXT("SkinStreamInputBuffer"));
-
+		SkinInputStreamFloatOffset.Bind(Initializer.ParameterMap, TEXT("SkinCacheInputStreamFloatOffset"));
+		
 		SkinCacheBuffer.Bind(Initializer.ParameterMap, TEXT("SkinCacheBuffer"));
 
-		MorphInputStreamBuffer.Bind(Initializer.ParameterMap, TEXT("MorphInputStreamBuffer"));
-		MorphInputStreamFloatOffset.Bind(Initializer.ParameterMap, TEXT("MorphInputStreamFloatOffset"));
+		MorphBuffer.Bind(Initializer.ParameterMap, TEXT("MorphBuffer"));
+		MorphBufferOffset.Bind(Initializer.ParameterMap, TEXT("MorphBufferOffset"));
 	}
 
-	void SetParameters(FRHICommandList& RHICmdList, uint32 InputStreamStride,
+	void SetParameters(uint32 InputStreamStride,
 		uint32 InputStreamFloatOffset, uint32 InputStreamVertexCount, uint32 OutputBufferFloatOffset,
 		const FVertexBufferAndSRV& BoneBuffer, FUniformBufferRHIRef UniformBuffer,
 		FShaderResourceViewRHIRef VertexBufferSRV, FRWBuffer& SkinBuffer, const FVector& MeshOrigin, const FVector& MeshExtension,
-		uint32 SkinType, uint32 InMorphInputStreamFloatOffset, FShaderResourceViewRHIParamRef MorphSRV)
+		const FGPUSkinCache::FDispatchData& DispatchData)
 	{
 		FComputeShaderRHIParamRef ShaderRHI = GetComputeShader();
+
+		FRHICommandListImmediate& RHICmdList = DispatchData.RHICmdList;
+
 		SetShaderValue(RHICmdList, ShaderRHI, SkinMeshOriginParameter, MeshOrigin);
 		SetShaderValue(RHICmdList, ShaderRHI, SkinMeshExtensionParameter, MeshExtension);
 
 		SetShaderValue(RHICmdList, ShaderRHI, SkinInputStreamStride, InputStreamStride);
 		SetShaderValue(RHICmdList, ShaderRHI, SkinInputStreamVertexCount, InputStreamVertexCount);
 		SetShaderValue(RHICmdList, ShaderRHI, SkinInputStreamFloatOffset, InputStreamFloatOffset);
-		SetShaderValue(RHICmdList, ShaderRHI, SkinOutputBufferFloatOffset, OutputBufferFloatOffset);
 		//SetShaderValue(RHICmdList, ShaderRHI, DebugParameter, FUintVector4((uint32)(uint64)VertexBufferSRV.GetReference(), 0, 0, 1));
 
 		if (UniformBuffer)
@@ -118,29 +122,32 @@ public:
 		}
 		else
 		{
-			RHICmdList.SetShaderResourceViewParameter(ShaderRHI, BoneMatrices.GetBaseIndex(), BoneBuffer.VertexBufferSRV);
+			SetSRVParameter(RHICmdList, ShaderRHI, BoneMatrices, BoneBuffer.VertexBufferSRV);
 		}
 
-		RHICmdList.SetShaderResourceViewParameter(ShaderRHI, SkinInputStream.GetBaseIndex(), VertexBufferSRV);
+		SetSRVParameter(RHICmdList, ShaderRHI, SkinInputStream, VertexBufferSRV);
+		// output UAV
 		RHICmdList.SetUAVParameter(ShaderRHI, SkinCacheBuffer.GetBaseIndex(), SkinBuffer.UAV);
-
-		const bool bMorph = SkinType == 1;
+		SetShaderValue(RHICmdList, ShaderRHI, SkinOutputBufferFloatOffset, OutputBufferFloatOffset);
+	
+		const bool bMorph = DispatchData.SkinType == 1;
 		if(bMorph)
 		{
-			SetShaderValue(RHICmdList, ShaderRHI, MorphInputStreamFloatOffset, InMorphInputStreamFloatOffset);
-			RHICmdList.SetShaderResourceViewParameter(ShaderRHI, MorphInputStreamBuffer.GetBaseIndex(), MorphSRV);
+			SetSRVParameter(RHICmdList, ShaderRHI, MorphBuffer, DispatchData.MorphBuffer);
+			SetShaderValue(RHICmdList, ShaderRHI, MorphBufferOffset, DispatchData.MorphBufferOffset);
 		}
 	}
 
 	void UnsetParameters(FRHICommandList& RHICmdList)
 	{
 		FComputeShaderRHIParamRef ShaderRHI = GetComputeShader();
+
+		// are the next 3 lines needed?
 		FShaderResourceViewRHIParamRef NullSRV = FShaderResourceViewRHIParamRef();
 		RHICmdList.SetShaderResourceViewParameter(ShaderRHI, BoneMatrices.GetBaseIndex(), NullSRV);
-
 		RHICmdList.SetShaderResourceViewParameter(ShaderRHI, SkinInputStream.GetBaseIndex(), NullSRV);
 
-		RHICmdList.SetUAVParameter(ShaderRHI, SkinCacheBuffer.GetBaseIndex(), FUnorderedAccessViewRHIParamRef() );
+		RHICmdList.SetUAVParameter(ShaderRHI, SkinCacheBuffer.GetBaseIndex(), FUnorderedAccessViewRHIParamRef());
 	}
 
 	virtual bool Serialize(FArchive& Ar) override
@@ -148,7 +155,7 @@ public:
 		bool bShaderHasOutdatedParameters = FGlobalShader::Serialize(Ar);
 		Ar  << SkinMeshOriginParameter << SkinMeshExtensionParameter << SkinInputStreamStride
 			<< SkinInputStreamVertexCount << SkinInputStreamFloatOffset << SkinOutputBufferFloatOffset
-			<< SkinInputStream << SkinCacheBuffer << BoneMatrices << MorphInputStreamBuffer << MorphInputStreamFloatOffset;
+			<< SkinInputStream << SkinCacheBuffer << BoneMatrices << MorphBuffer << MorphBufferOffset;
 		//Ar << DebugParameter;
 		return bShaderHasOutdatedParameters;
 	}
@@ -171,8 +178,8 @@ private:
 	FShaderResourceParameter SkinInputStream;
 	FShaderResourceParameter SkinCacheBuffer;
 
-	FShaderResourceParameter MorphInputStreamBuffer;
-	FShaderParameter MorphInputStreamFloatOffset;
+	FShaderResourceParameter MorphBuffer;
+	FShaderParameter MorphBufferOffset;
 };
 
 /** Compute shader that skins a batch of vertices. */
@@ -269,7 +276,7 @@ void FGPUSkinCache::Initialize(FRHICommandListImmediate& RHICmdList)
 	TransitionAllToWriteable(RHICmdList);
 
 	CachedElements.Reserve(MaxCachedElements);
-	CachedVertexBuffers.Reserve(MaxCachedVertexBufferSRVs);
+	CachedSRVs.Reserve(MaxCachedVertexBufferSRVs);
 
 	bInitialized = true;
 }
@@ -287,9 +294,73 @@ void FGPUSkinCache::Cleanup()
 	}
 }
 
+FGPUSkinCache::FSRVCacheEntry* FGPUSkinCache::ReuseOrFindAndCacheSRV(int32& InOutIndex, const FVertexBuffer& InKey)
+{
+	FSRVCacheEntry* Ret = 0;
+	if (InOutIndex >= 0)
+	{
+		Ret = &(CachedSRVs.GetData()[InOutIndex]);
+	}
+	else
+	{
+		// time complexity: O(n) 
+		Ret = CachedSRVs.FindByKey(InKey);
+
+		if (Ret)
+		{
+			InOutIndex = Ret - CachedSRVs.GetData();
+		}
+	}
+
+	if (!Ret)
+	{
+		// Add new cache entry for vertex buffer and SRV
+		FSRVCacheEntry Info;
+		Info.Key = &InKey;
+		Info.SRVValue = RHICreateShaderResourceView(InKey.VertexBufferRHI, 4, PF_R32_FLOAT);
+
+		InOutIndex = CachedSRVs.Add(Info);
+		Ret = &(CachedSRVs.GetData()[InOutIndex]);
+	}
+
+	return Ret;
+}
+
+FGPUSkinCache::FSRVCacheEntry* FGPUSkinCache::ReuseOrFindAndCacheSRV(int32& InOutIndex, const FIndexBuffer& InKey)
+{
+	FSRVCacheEntry* Ret = 0;
+	if (InOutIndex >= 0)
+	{
+		Ret = &(CachedSRVs.GetData()[InOutIndex]);
+	}
+	else
+	{
+		// time complexity: O(n) 
+		Ret = CachedSRVs.FindByKey(InKey);
+
+		if (Ret)
+		{
+			InOutIndex = Ret - CachedSRVs.GetData();
+		}
+	}
+
+	if (!Ret)
+	{
+		// Add new cache entry for index buffer and SRV
+		FSRVCacheEntry Info;
+		Info.Key = &InKey;
+		Info.SRVValue = RHICreateShaderResourceView(InKey.IndexBufferRHI);
+
+		InOutIndex = CachedSRVs.Add(Info);
+		Ret = &(CachedSRVs.GetData()[InOutIndex]);
+	}
+
+	return Ret;
+}
+
 // Kick off compute shader skinning for incoming chunk, return key for fast lookup in draw passes
 int32 FGPUSkinCache::StartCacheMesh(FRHICommandListImmediate& RHICmdList, int32 Key, FGPUBaseSkinVertexFactory* VertexFactory,
-	FGPUSkinPassthroughVertexFactory* TargetVertexFactory, const FSkelMeshChunk& BatchElement, const FSkeletalMeshObjectGPUSkin* Skin,
+	FGPUSkinPassthroughVertexFactory* TargetVertexFactory, const FSkelMeshChunk& BatchElement, FSkeletalMeshObjectGPUSkin* Skin,
 	const FMorphVertexBuffer* MorphVertexBuffer)
 {
 	if (CachedChunksThisFrameCount >= GMaxGPUSkinCacheElementsPerFrame && GFrameNumberRenderThread <= FrameCounter)
@@ -422,15 +493,14 @@ int32 FGPUSkinCache::StartCacheMesh(FRHICommandListImmediate& RHICmdList, int32
 
 	const FVertexBuffer* SkinVertexBuffer = GPUVertexFactory->GetSkinVertexBuffer();
 
-	FVertexBufferInfo* VBInfo = ReuseOrFindAndCacheSRV(ECSInfo->VertexBufferSRVIndex, *SkinVertexBuffer);
+	FSRVCacheEntry* VBInfo = ReuseOrFindAndCacheSRV(ECSInfo->VertexBufferSRVIndex, *SkinVertexBuffer);
 	
-	FShaderResourceViewRHIRef MorphVertexBufferSRV;
-	uint32 MorphInputStreamFloatOffset = 0;
+	FDispatchData DispatchData(RHICmdList, Skin->GetFeatureLevel(), Skin);
 
-	if(MorphVertexBuffer != 0)
+	if(MorphVertexBuffer)
 	{
-		FVertexBufferInfo* MorphVBInfo = ReuseOrFindAndCacheSRV(ECSInfo->MorphVertexBufferSRVIndex, *MorphVertexBuffer);
-		MorphVertexBufferSRV = MorphVBInfo->VertexBufferSRV;
+		FSRVCacheEntry* MorphVBInfo = ReuseOrFindAndCacheSRV(ECSInfo->MorphVertexBufferSRVIndex, *MorphVertexBuffer);
+		DispatchData.MorphBuffer = MorphVBInfo->SRVValue;
 
 		// in bytes
 		const uint32 MorphStride = sizeof(FMorphGPUSkinVertex);
@@ -438,7 +508,7 @@ int32 FGPUSkinCache::StartCacheMesh(FRHICommandListImmediate& RHICmdList, int32
 		// see GPU code "check(MorphStride == sizeof(float) * 6);"
 		check(MorphStride == sizeof(float) * 6);
 
-		MorphInputStreamFloatOffset = (MorphStride * BatchElement.BaseVertexIndex) / sizeof(float);
+		DispatchData.MorphBufferOffset = (MorphStride * BatchElement.BaseVertexIndex) / sizeof(float);
 	}
 
 	INC_DWORD_STAT(STAT_GPUSkinCache_TotalNumChunks);
@@ -447,12 +517,12 @@ int32 FGPUSkinCache::StartCacheMesh(FRHICommandListImmediate& RHICmdList, int32
 	auto& ShaderData = GPUVertexFactory->GetShaderData();
 	
 	// SkinType 0:normal, 1:with morph target, 2:with APEX cloth (not yet implemented)
-	uint32 SkinType = MorphVertexBuffer ? 1 : 0;
+	DispatchData.SkinType = MorphVertexBuffer ? 1 : 0;
 
-	DispatchSkinCacheProcess(RHICmdList, InputStreamFloatOffset, ECSInfo->StreamOffset,
+	DispatchSkinCacheProcess(InputStreamFloatOffset, ECSInfo->StreamOffset,
 		ShaderData.GetBoneBufferForWriting(false, GFrameNumberRenderThread), ShaderData.GetUniformBuffer(),
-		VBInfo, StreamStrides[0], NumVertices, ShaderData.MeshOrigin, ShaderData.MeshExtension, bExtraBoneInfluences, Skin->GetFeatureLevel(),
-		SkinType, MorphInputStreamFloatOffset, MorphVertexBufferSRV);
+		VBInfo, StreamStrides[0], NumVertices, ShaderData.MeshOrigin, ShaderData.MeshExtension, bExtraBoneInfluences,
+		DispatchData);
 
 	TargetVertexFactory->UpdateVertexDeclaration(VertexFactory, &SkinCacheBuffer[FrameCounter % GPUSKINCACHE_FRAMES]);
 
@@ -491,39 +561,6 @@ void FGPUSkinCache::TransitionAllToWriteable(FRHICommandList& RHICmdList)
 	}
 }
 
-// @param InOutIndex index -1 if not yet set
-FGPUSkinCache::FVertexBufferInfo* FGPUSkinCache::ReuseOrFindAndCacheSRV(int32& InOutIndex, const FVertexBuffer& VertexBuffer)
-{
-	FVertexBufferInfo* Ret = 0;
-	if (InOutIndex >= 0)
-	{
-		Ret = &(CachedVertexBuffers.GetData()[InOutIndex]);
-	}
-	else
-	{
-		// time complexity: O(n) 
-		Ret = CachedVertexBuffers.FindByKey(VertexBuffer);
-
-		if (Ret)
-		{
-			InOutIndex = Ret - CachedVertexBuffers.GetData();
-		}
-	}
-
-	if (!Ret)
-	{
-		// Add new cache entry for vertex buffer and SRV
-		FVertexBufferInfo	Info;
-		Info.VertexBuffer = &VertexBuffer;
-		Info.VertexBufferSRV = RHICreateShaderResourceView(VertexBuffer.VertexBufferRHI, 4, PF_R32_FLOAT);
-
-		InOutIndex = CachedVertexBuffers.Add(Info);
-		Ret = &(CachedVertexBuffers.GetData()[InOutIndex]);
-	}
-
-	return Ret;
-}
-
 bool FGPUSkinCache::InternalIsElementProcessed(int32 Key) const
 {
 	if (Key >= 0 && Key < CachedElements.Num())
@@ -619,47 +656,51 @@ FGPUSkinCache::FElementCacheStatusInfo* FGPUSkinCache::FindEvictableCacheStatusI
 	}
 }
 
-void FGPUSkinCache::DispatchSkinCacheProcess(FRHICommandListImmediate& RHICmdList, uint32 InputStreamFloatOffset, uint32 OutputBufferFloatOffset,
-	const FVertexBufferAndSRV& BoneBuffer, FUniformBufferRHIRef UniformBuffer, const FGPUSkinCache::FVertexBufferInfo* VBInfo, uint32 VertexStride,
-	uint32 VertexCount, const FVector& MeshOrigin, const FVector& MeshExtension, bool bUseExtraBoneInfluences, ERHIFeatureLevel::Type FeatureLevel,
-	uint32 SkinType, uint32 InMorphInputStreamFloatOffset, FShaderResourceViewRHIParamRef MorphVertexBufferSRV)
+void FGPUSkinCache::DispatchSkinCacheProcess(uint32 InputStreamFloatOffset, uint32 OutputBufferFloatOffset,
+	const FVertexBufferAndSRV& BoneBuffer, FUniformBufferRHIRef UniformBuffer, const FGPUSkinCache::FSRVCacheEntry* VBInfo, uint32 VertexStride,
+	uint32 VertexCount, const FVector& MeshOrigin, const FVector& MeshExtension, bool bUseExtraBoneInfluences,
+	const FDispatchData& DispatchData)
 {
-	check(FeatureLevel >= ERHIFeatureLevel::SM5);
+	check(DispatchData.FeatureLevel >= ERHIFeatureLevel::SM5);
 
 	uint32 VertexCountAlign64 = FMath::DivideAndRoundUp(VertexCount, (uint32)64);
 
-	SCOPED_DRAW_EVENT(RHICmdList, SkinCacheDispatch);
+	SCOPED_DRAW_EVENT(DispatchData.RHICmdList, SkinCacheDispatch);
 
 	int32 UAVIndex = FrameCounter % GPUSKINCACHE_FRAMES;
 
-	TShaderMapRef<TGPUSkinCacheCS<true,0> > SkinCacheCS10(GetGlobalShaderMap(FeatureLevel));
-	TShaderMapRef<TGPUSkinCacheCS<false,0> > SkinCacheCS00(GetGlobalShaderMap(FeatureLevel));
-	TShaderMapRef<TGPUSkinCacheCS<true,1> > SkinCacheCS11(GetGlobalShaderMap(FeatureLevel));
-	TShaderMapRef<TGPUSkinCacheCS<false,1> > SkinCacheCS01(GetGlobalShaderMap(FeatureLevel));
+	TShaderMapRef<TGPUSkinCacheCS<true,0> > SkinCacheCS10(GetGlobalShaderMap(DispatchData.FeatureLevel));
+	TShaderMapRef<TGPUSkinCacheCS<false,0> > SkinCacheCS00(GetGlobalShaderMap(DispatchData.FeatureLevel));
+	TShaderMapRef<TGPUSkinCacheCS<true,1> > SkinCacheCS11(GetGlobalShaderMap(DispatchData.FeatureLevel));
+	TShaderMapRef<TGPUSkinCacheCS<false,1> > SkinCacheCS01(GetGlobalShaderMap(DispatchData.FeatureLevel));
 
 	FBaseGPUSkinCacheCS* Shader = 0;
 	
-	switch(SkinType)
+	switch(DispatchData.SkinType)
 	{
 		case 0: Shader = bUseExtraBoneInfluences ? (FBaseGPUSkinCacheCS*)*SkinCacheCS10 : (FBaseGPUSkinCacheCS*)*SkinCacheCS00; break;
 		case 1: Shader = bUseExtraBoneInfluences ? (FBaseGPUSkinCacheCS*)*SkinCacheCS11 : (FBaseGPUSkinCacheCS*)*SkinCacheCS01; break;
+		default:
+			check(0);
 	}
 	
 	check(Shader);
 
-	RHICmdList.SetComputeShader(Shader->GetComputeShader());
-	Shader->SetParameters(RHICmdList, VertexStride, InputStreamFloatOffset, VertexCount,
-		OutputBufferFloatOffset, BoneBuffer, UniformBuffer, VBInfo->VertexBufferSRV,
+	DispatchData.RHICmdList.SetComputeShader(Shader->GetComputeShader());
+	Shader->SetParameters(VertexStride, InputStreamFloatOffset, VertexCount,
+		OutputBufferFloatOffset, BoneBuffer, UniformBuffer, VBInfo->SRVValue,
 		SkinCacheBuffer[UAVIndex], MeshOrigin, MeshExtension,
-		SkinType, InMorphInputStreamFloatOffset, MorphVertexBufferSRV);
+		DispatchData);
 
-	RHICmdList.DispatchComputeShader(VertexCountAlign64, 1, 1);
-	Shader->UnsetParameters(RHICmdList);
+	DispatchData.RHICmdList.DispatchComputeShader(VertexCountAlign64, 1, 1);
+	Shader->UnsetParameters(DispatchData.RHICmdList);
 }
 
 void FGPUSkinCache::CVarSinkFunction()
 {
-	int32 NewGPUSkinCacheValue = CVarEnableGPUSkinCache.GetValueOnAnyThread();
+	// 0/1
+	int32 NewGPUSkinCacheValue = CVarEnableGPUSkinCache.GetValueOnAnyThread() != 0;
+
 	if (!GEnableGPUSkinCacheShaders)
 	{
 		NewGPUSkinCacheValue = 0;
diff --git a/Engine/Source/Runtime/Engine/Private/GPUSkinVertexFactory.cpp b/Engine/Source/Runtime/Engine/Private/GPUSkinVertexFactory.cpp
--- a/Engine/Source/Runtime/Engine/Private/GPUSkinVertexFactory.cpp
+++ b/Engine/Source/Runtime/Engine/Private/GPUSkinVertexFactory.cpp
@@ -580,8 +580,8 @@ public:
 	*/
 	virtual void Bind(const FShaderParameterMap& ParameterMap) override
 	{
-		GPUSkinCacheFloatOffset.Bind(ParameterMap,TEXT("GPUSkinCachePreviousFloatOffset"));
-		GPUSkinCacheBuffer.Bind(ParameterMap,TEXT("GPUSkinCachePreviousBuffer"));
+		GPUSkinCachePreviousFloatOffset.Bind(ParameterMap,TEXT("GPUSkinCachePreviousFloatOffset"));
+		GPUSkinCachePreviousBuffer.Bind(ParameterMap,TEXT("GPUSkinCachePreviousBuffer"));
 	}
 	/**
 	* Serialize shader params to an archive
@@ -589,8 +589,8 @@ public:
 	*/
 	virtual void Serialize(FArchive& Ar) override
 	{
-		Ar << GPUSkinCacheFloatOffset;
-		Ar << GPUSkinCacheBuffer;
+		Ar << GPUSkinCachePreviousFloatOffset;
+		Ar << GPUSkinCachePreviousBuffer;
 	}
 	/**
 	* Set any shader data specific to this vertex factory
@@ -601,7 +601,7 @@ public:
 		{
 			bool bIsGPUCached = false;
 			check(VertexFactory->GetType() == &FGPUSkinPassthroughVertexFactory::StaticType);
-			if (GGPUSkinCache.SetVertexStreamFromCache(RHICmdList, BatchElement.UserIndex, Shader, (FGPUSkinPassthroughVertexFactory*)VertexFactory, BatchElement.MinVertexIndex, GPUSkinCacheFloatOffset, GPUSkinCacheBuffer))
+			if (GGPUSkinCache.SetVertexStreamFromCache(RHICmdList, BatchElement.UserIndex, Shader, (FGPUSkinPassthroughVertexFactory*)VertexFactory, BatchElement.MinVertexIndex, GPUSkinCachePreviousFloatOffset, GPUSkinCachePreviousBuffer))
 			{
 				bIsGPUCached = true;
 			}
@@ -611,8 +611,8 @@ public:
 	virtual uint32 GetSize() const override { return sizeof(*this); }
 
 private:
-	FShaderParameter GPUSkinCacheFloatOffset;
-	FShaderResourceParameter GPUSkinCacheBuffer;
+	FShaderParameter GPUSkinCachePreviousFloatOffset;
+	FShaderResourceParameter GPUSkinCachePreviousBuffer;
 };
 
 /*-----------------------------------------------------------------------------
diff --git a/Engine/Source/Runtime/Engine/Private/LocalPlayer.cpp b/Engine/Source/Runtime/Engine/Private/LocalPlayer.cpp
--- a/Engine/Source/Runtime/Engine/Private/LocalPlayer.cpp
+++ b/Engine/Source/Runtime/Engine/Private/LocalPlayer.cpp
@@ -640,50 +640,37 @@ void ULocalPlayer::GetViewPoint(FMinimalViewInfo& OutViewInfo, EStereoscopicPass
     }
 }
 
-FSceneView* ULocalPlayer::CalcSceneView( class FSceneViewFamily* ViewFamily, 
-	FVector& OutViewLocation, 
-	FRotator& OutViewRotation, 
-	FViewport* Viewport, 
+bool ULocalPlayer::CalcSceneViewInitOptions(
+	struct FSceneViewInitOptions& ViewInitOptions,
+	FViewport* Viewport,
 	class FViewElementDrawer* ViewDrawer,
 	EStereoscopicPass StereoPass)
 {
-	SCOPE_CYCLE_COUNTER(STAT_CalcSceneView);
-
+	QUICK_SCOPE_CYCLE_COUNTER(STAT_CalcSceneViewInitOptions);
 	if ((PlayerController == NULL) || (Size.X <= 0.f) || (Size.Y <= 0.f) || (Viewport == NULL))
 	{
-		return NULL;
+		return false;
 	}
-
-	FSceneViewInitOptions ViewInitOptions;
-
 	// get the projection data
 	if (GetProjectionData(Viewport, StereoPass, /*inout*/ ViewInitOptions) == false)
 	{
 		// Return NULL if this we didn't get back the info we needed
-		return NULL;
+		return false;
 	}
-	
+
 	// return if we have an invalid view rect
 	if (!ViewInitOptions.IsValidViewRectangle())
 	{
-		return NULL;
+		return false;
 	}
 
-	// Get the viewpoint...technically doing this twice
-	// but it makes GetProjectionData better
-	FMinimalViewInfo ViewInfo;
-	GetViewPoint(ViewInfo, StereoPass);
-	
-	OutViewLocation = ViewInfo.Location;
-	OutViewRotation = ViewInfo.Rotation;
-
 	if (PlayerController->PlayerCameraManager != NULL)
 	{
 		// Apply screen fade effect to screen.
 		if (PlayerController->PlayerCameraManager->bEnableFading)
 		{
 			ViewInitOptions.OverlayColor = PlayerController->PlayerCameraManager->FadeColor;
-			ViewInitOptions.OverlayColor.A = FMath::Clamp(PlayerController->PlayerCameraManager->FadeAmount,0.0f,1.0f);
+			ViewInitOptions.OverlayColor.A = FMath::Clamp(PlayerController->PlayerCameraManager->FadeAmount, 0.0f, 1.0f);
 		}
 
 		// Do color scaling if desired.
@@ -699,11 +686,8 @@ FSceneView* ULocalPlayer::CalcSceneView( class FSceneViewFamily* ViewFamily,
 		// Was there a camera cut this frame?
 		ViewInitOptions.bInCameraCut = PlayerController->PlayerCameraManager->bGameCameraCutThisFrame;
 	}
-	
-	check(PlayerController && PlayerController->GetWorld());
 
-	// Fill out the rest of the view init options
-	ViewInitOptions.ViewFamily = ViewFamily;
+	check(PlayerController && PlayerController->GetWorld());
 	ViewInitOptions.SceneViewStateInterface = ((StereoPass != eSSP_RIGHT_EYE) ? ViewState.GetReference() : StereoViewState.GetReference());
 	ViewInitOptions.ViewActor = PlayerController->GetViewTarget();
 	ViewInitOptions.ViewElementDrawer = ViewDrawer;
@@ -713,8 +697,42 @@ FSceneView* ULocalPlayer::CalcSceneView( class FSceneViewFamily* ViewFamily,
 	ViewInitOptions.WorldToMetersScale = PlayerController->GetWorldSettings()->WorldToMeters;
 	ViewInitOptions.CursorPos = Viewport->HasMouseCapture() ? FIntPoint(-1, -1) : FIntPoint(Viewport->GetMouseX(), Viewport->GetMouseY());
 	ViewInitOptions.bOriginOffsetThisFrame = PlayerController->GetWorld()->bOriginOffsetThisFrame;
+
+	return true;
+}
+
+
+FSceneView* ULocalPlayer::CalcSceneView( class FSceneViewFamily* ViewFamily, 
+	FVector& OutViewLocation, 
+	FRotator& OutViewRotation, 
+	FViewport* Viewport, 
+	class FViewElementDrawer* ViewDrawer,
+	EStereoscopicPass StereoPass)
+{
+	SCOPE_CYCLE_COUNTER(STAT_CalcSceneView);
+
+	FSceneViewInitOptions ViewInitOptions;
+
+	if (!CalcSceneViewInitOptions(ViewInitOptions, Viewport, ViewDrawer, StereoPass))
+	{
+		return nullptr;
+	}
+
+	// Get the viewpoint...technically doing this twice
+	// but it makes GetProjectionData better
+	FMinimalViewInfo ViewInfo;
+	GetViewPoint(ViewInfo, StereoPass);
+	OutViewLocation = ViewInfo.Location;
+	OutViewRotation = ViewInfo.Rotation;
 	ViewInitOptions.bUseFieldOfViewForLOD = ViewInfo.bUseFieldOfViewForLOD;
-	PlayerController->BuildHiddenComponentList(OutViewLocation, /*out*/ ViewInitOptions.HiddenPrimitives);
+
+	// Fill out the rest of the view init options
+	ViewInitOptions.ViewFamily = ViewFamily;
+
+	{
+		QUICK_SCOPE_CYCLE_COUNTER(STAT_BuildHiddenComponentList);
+		PlayerController->BuildHiddenComponentList(OutViewLocation, /*out*/ ViewInitOptions.HiddenPrimitives);
+	}
 
 	FSceneView* const View = new FSceneView(ViewInitOptions);
 	
diff --git a/Engine/Source/Runtime/Engine/Private/Particles/ParticleTrail2EmitterInstance.cpp b/Engine/Source/Runtime/Engine/Private/Particles/ParticleTrail2EmitterInstance.cpp
--- a/Engine/Source/Runtime/Engine/Private/Particles/ParticleTrail2EmitterInstance.cpp
+++ b/Engine/Source/Runtime/Engine/Private/Particles/ParticleTrail2EmitterInstance.cpp
@@ -4272,7 +4272,7 @@ void FParticleAnimTrailEmitterInstance::PrintParticleData(FBaseParticle* Particl
 		);
 	UE_LOG(LogParticles, Log, TEXT("TI=%4d - TC=%4d, SpT=%.5f, SpD= %.5f, TU=%.4f, SpTP=%4d, RIntpC=%4d, PSF=%.4f, %s, %s"),
 		TrailData->TrailIndex, TrailData->TriangleCount, TrailData->SpawnTime, TrailData->SpawnDelta, TrailData->TiledU, TrailData->SpawnedTessellationPoints,
-		TrailData->RenderingInterpCount, TrailData->PinchScaleFactor, TrailData->bInterpolatedSpawn ? "1" : "0", TrailData->bMovementSpawned ? "1" : "0");
+		TrailData->RenderingInterpCount, TrailData->PinchScaleFactor, TrailData->bInterpolatedSpawn ? TEXT("1") : TEXT("0"), TrailData->bMovementSpawned ? TEXT("1") : TEXT("0"));
 }
 
 void FParticleAnimTrailEmitterInstance::PrintAllActiveParticles()
diff --git a/Engine/Source/Runtime/Engine/Private/PhysicsEngine/PhysDrawing.cpp b/Engine/Source/Runtime/Engine/Private/PhysicsEngine/PhysDrawing.cpp
--- a/Engine/Source/Runtime/Engine/Private/PhysicsEngine/PhysDrawing.cpp
+++ b/Engine/Source/Runtime/Engine/Private/PhysicsEngine/PhysDrawing.cpp
@@ -637,7 +637,7 @@ void FKAggregateGeom::GetAggGeom(const FTransform& Transform, const FColor Color
 				FBoxSphereBounds WorldBounds, LocalBounds;
 				CalcBoxSphereBounds(WorldBounds, LocalToWorld);
 				CalcBoxSphereBounds(LocalBounds, FTransform::Identity);
-			    BatchElement.PrimitiveUniformBuffer = CreatePrimitiveUniformBufferImmediate(LocalToWorld.ToMatrixWithScale(), WorldBounds, LocalBounds, 1.0f, true, bUseEditorDepthTest);
+			    BatchElement.PrimitiveUniformBuffer = CreatePrimitiveUniformBufferImmediate(LocalToWorld.ToMatrixWithScale(), WorldBounds, LocalBounds, true, bUseEditorDepthTest);
 			 	// previous l2w not used so treat as static
 				BatchElement.FirstIndex = 0;
 				BatchElement.NumPrimitives = RenderInfo->IndexBuffer->Indices.Num() / 3;
diff --git a/Engine/Source/Runtime/Engine/Private/PrimitiveSceneProxy.cpp b/Engine/Source/Runtime/Engine/Private/PrimitiveSceneProxy.cpp
--- a/Engine/Source/Runtime/Engine/Private/PrimitiveSceneProxy.cpp
+++ b/Engine/Source/Runtime/Engine/Private/PrimitiveSceneProxy.cpp
@@ -71,9 +71,6 @@ FPrimitiveSceneProxy::FPrimitiveSceneProxy(const UPrimitiveComponent* InComponen
 ,	LightingChannelMask(GetLightingChannelMaskForStruct(InComponent->LightingChannels))
 ,	LpvBiasMultiplier(InComponent->LpvBiasMultiplier)
 ,	IndirectLightingCacheQuality(InComponent->IndirectLightingCacheQuality)
-#if !(UE_BUILD_SHIPPING || UE_BUILD_TEST)
-,	WorldTexelFactor(0)
-#endif
 ,	Scene(InComponent->GetScene())
 ,	PrimitiveComponentId(InComponent->ComponentId)
 ,	PrimitiveSceneInfo(NULL)
diff --git a/Engine/Source/Runtime/Engine/Private/SkeletalRenderGPUSkin.h b/Engine/Source/Runtime/Engine/Private/SkeletalRenderGPUSkin.h
--- a/Engine/Source/Runtime/Engine/Private/SkeletalRenderGPUSkin.h
+++ b/Engine/Source/Runtime/Engine/Private/SkeletalRenderGPUSkin.h
@@ -380,14 +380,6 @@ private:
 		 */
 		void ReleaseMorphResources();
 
-		/** 
-		 * Update the contents of the vertex buffer with new data. Note that this
-		 * function is called from the render thread.
-		 * @param	NewVertices - array of new vertex data
-		 * @param	NumVertices - Number of vertices
-		 */
-		void UpdateShadowVertexBuffer( const FVector* NewVertices, uint32 NumVertices ) const;
-
 		/**
 		 * @return memory in bytes of size of the resources for this LOD
 		 */
diff --git a/Engine/Source/Runtime/Engine/Private/Slate/SGameLayerManager.cpp b/Engine/Source/Runtime/Engine/Private/Slate/SGameLayerManager.cpp
--- a/Engine/Source/Runtime/Engine/Private/Slate/SGameLayerManager.cpp
+++ b/Engine/Source/Runtime/Engine/Private/Slate/SGameLayerManager.cpp
@@ -320,29 +320,16 @@ void SGameLayerManager::AddOrUpdatePlayerLayers(const FGeometry& AllottedGeometr
 
 FVector2D SGameLayerManager::GetAspectRatioInset(ULocalPlayer* LocalPlayer) const
 {
+	QUICK_SCOPE_CYCLE_COUNTER(STAT_SGameLayerManager_GetAspectRatioInset);
 	FVector2D Offset(0.f, 0.f);
-
-	if ( LocalPlayer )
+	if (LocalPlayer)
 	{
-		// Create a view family for the game viewport
-		FSceneViewFamilyContext ViewFamily(FSceneViewFamily::ConstructionValues(
-			LocalPlayer->ViewportClient->Viewport,
-			LocalPlayer->GetWorld()->Scene,
-			LocalPlayer->ViewportClient->EngineShowFlags)
-			.SetRealtimeUpdate(true));
-
-		// Calculate a view where the player is to update the streaming from the players start location
-		FVector ViewLocation;
-		FRotator ViewRotation;
-		FSceneView* SceneView = LocalPlayer->CalcSceneView(&ViewFamily, /*out*/ ViewLocation, /*out*/ ViewRotation, LocalPlayer->ViewportClient->Viewport);
-
-		if ( SceneView )
+		FSceneViewInitOptions ViewInitOptions;
+		if (LocalPlayer->CalcSceneViewInitOptions(ViewInitOptions, LocalPlayer->ViewportClient->Viewport))
 		{
-			Offset.X = ( SceneView->ViewRect.Min.X - SceneView->UnscaledViewRect.Min.X ) // This accounts for the borders when the aspect ratio is locked
-				- SceneView->UnscaledViewRect.Min.X;									 // And this will deal with the viewport offset if its a split screen
-
-			Offset.Y = ( SceneView->ViewRect.Min.Y - SceneView->UnscaledViewRect.Min.Y )
-				- SceneView->UnscaledViewRect.Min.Y;
+			FIntRect UnscaledViewRect = ViewInitOptions.GetConstrainedViewRect();
+			Offset.X = -UnscaledViewRect.Min.X;
+			Offset.Y = -UnscaledViewRect.Min.Y;
 		}
 	}
 
diff --git a/Engine/Source/Runtime/Engine/Private/StaticMesh.cpp b/Engine/Source/Runtime/Engine/Private/StaticMesh.cpp
--- a/Engine/Source/Runtime/Engine/Private/StaticMesh.cpp
+++ b/Engine/Source/Runtime/Engine/Private/StaticMesh.cpp
@@ -1452,7 +1452,7 @@ FBox UStaticMesh::GetBoundingBox() const
 	return ExtendedBounds.GetBox();
 }
 
-float UStaticMesh::GetStreamingTextureFactor(int32 RequestedUVIndex)
+float UStaticMesh::GetStreamingTextureFactor(int32 RequestedUVIndex) const
 {
 	check(RequestedUVIndex >= 0);
 	check(RequestedUVIndex < MAX_STATIC_TEXCOORDS);
@@ -1476,6 +1476,118 @@ float UStaticMesh::GetStreamingTextureFactor(int32 RequestedUVIndex)
 	return StreamingTextureFactor;
 }
 
+bool UStaticMesh::GetStreamingTextureFactor(float& OutTexelFactor, FBoxSphereBounds& OutBounds, int32 CoordinateIndex, int32 LODIndex, int32 ElementIndex, const FTransform& Transform) const 
+{
+#if WITH_EDITORONLY_DATA
+	if (!GIsEditor || !RenderData || !RenderData->LODResources.IsValidIndex(LODIndex))
+	{
+		return false;
+	}
+
+	const FStaticMeshLODResources& LODModel = RenderData->LODResources[LODIndex];
+
+	if (CoordinateIndex < 0 || CoordinateIndex >= LODModel.GetNumTexCoords())
+	{
+		return false;
+	}
+
+	if (!LODModel.Sections.IsValidIndex(ElementIndex))
+	{
+		return false;
+	}
+
+	struct FTriangleInfo
+	{
+		FTriangleInfo(float InAera, float InTexelRatio) : Aera(InAera), TexelRatio(InTexelRatio) {}
+		float Aera;
+		float TexelRatio;
+	};
+
+	struct FCompareAera
+	{
+		FORCEINLINE bool operator()(FTriangleInfo const& A, FTriangleInfo const& B) const { return A.Aera < B.Aera; }
+	};
+
+	struct FCompareTexelRatio
+	{
+		FORCEINLINE bool operator()(FTriangleInfo const& A, FTriangleInfo const& B) const { return A.TexelRatio < B.TexelRatio; }
+	};
+
+	TArray<FTriangleInfo> TriangleInfos;
+
+	const FStaticMeshSection& SectionInfo = LODModel.Sections[ElementIndex];
+	const int32 SectionIndexCount = SectionInfo.NumTriangles * 3;
+	FIndexArrayView IndexBuffer = LODModel.IndexBuffer.GetArrayView();
+	
+	FBox TransformedSectionBox(ForceInit);
+
+	for (uint32 TriangleIndex = 0; TriangleIndex < SectionInfo.NumTriangles; ++TriangleIndex)
+	{
+		const int32 Index0 = IndexBuffer[SectionInfo.FirstIndex + TriangleIndex * 3 + 0];
+		const int32 Index1 = IndexBuffer[SectionInfo.FirstIndex + TriangleIndex * 3 + 1];
+		const int32 Index2 = IndexBuffer[SectionInfo.FirstIndex + TriangleIndex * 3 + 2];
+
+		FVector Pos0 = Transform.TransformPosition(LODModel.PositionVertexBuffer.VertexPosition(Index0));
+		FVector Pos1 = Transform.TransformPosition(LODModel.PositionVertexBuffer.VertexPosition(Index1));
+		FVector Pos2 = Transform.TransformPosition(LODModel.PositionVertexBuffer.VertexPosition(Index2));
+
+		TransformedSectionBox += Pos0;
+		TransformedSectionBox += Pos1;
+		TransformedSectionBox += Pos2;
+
+		FVector2D UV0 = LODModel.VertexBuffer.GetVertexUV(Index0, CoordinateIndex);
+		FVector2D UV1 = LODModel.VertexBuffer.GetVertexUV(Index1, CoordinateIndex);
+		FVector2D UV2 = LODModel.VertexBuffer.GetVertexUV(Index2, CoordinateIndex);
+
+		FVector P01 = Pos1 - Pos0;
+		FVector P02 = Pos2 - Pos0;
+
+		float Aera = FVector::CrossProduct(P01, P02).Size();
+
+		if (Aera > SMALL_NUMBER)
+		{
+			float L1 = P01.Size();
+			float L2 = P02.Size();
+
+			float T1 = (UV1 - UV0).Size();
+			float T2 = (UV2 - UV0).Size();
+
+			if (T1 > SMALL_NUMBER && T2 > SMALL_NUMBER)
+			{
+				float TexelRatio = FMath::Max(L1 / T1, L2 / T2);
+				TriangleInfos.Push(FTriangleInfo(Aera, TexelRatio));
+			}
+		}
+	}
+
+	TriangleInfos.Sort(FCompareTexelRatio());
+
+	float WeightedTexelFactorSum = 0;
+	float AreaSum = 0;
+
+	// Remove 10% of higher and lower texel factors.
+	int32 Threshold = FMath::FloorToInt(.10f * (float)TriangleInfos.Num());
+	for (int32 Index = Threshold; Index < TriangleInfos.Num() - Threshold; ++Index)
+	{
+		WeightedTexelFactorSum += TriangleInfos[Index].TexelRatio * TriangleInfos[Index].Aera;
+		AreaSum += TriangleInfos[Index].Aera;
+	}
+
+	if (AreaSum == 0)
+	{
+		return false;
+	}
+
+	OutBounds = TransformedSectionBox;
+	OutTexelFactor = WeightedTexelFactorSum / AreaSum;
+	return true;
+
+#else
+	return false;
+#endif
+}	
+
+
 /**
  * Releases the static mesh's render resources.
  */
diff --git a/Engine/Source/Runtime/Engine/Private/StaticMeshRender.cpp b/Engine/Source/Runtime/Engine/Private/StaticMeshRender.cpp
--- a/Engine/Source/Runtime/Engine/Private/StaticMeshRender.cpp
+++ b/Engine/Source/Runtime/Engine/Private/StaticMeshRender.cpp
@@ -100,6 +100,7 @@ FStaticMeshSceneProxy::FStaticMeshSceneProxy(UStaticMeshComponent* InComponent):
 	, SectionIndexPreview(InComponent->SectionIndexPreview)
 #endif
 #if !(UE_BUILD_SHIPPING || UE_BUILD_TEST)
+	, WorldTexelFactor(0)
 	, WorldLightmapFactor(0)
 	, LODForCollision(InComponent->StaticMesh->LODForCollision)
 	, bDrawMeshCollisionWireframe(InComponent->bDrawMeshCollisionWireframe)
@@ -189,8 +190,19 @@ FStaticMeshSceneProxy::FStaticMeshSceneProxy(UStaticMeshComponent* InComponent):
 #endif
 
 #if !(UE_BUILD_SHIPPING || UE_BUILD_TEST)
-	// Update texture streaming factors
-	InComponent->GetStreamingTextureFactors(WorldTexelFactor, WorldLightmapFactor);
+	if (AllowDebugViewmodes())
+	{
+		// Update texture streaming factors
+		StreamingTextureInfos = InComponent->StreamingTextureInfos;
+
+		FBoxSphereBounds DummyBounds;
+		InComponent->GetStreamingTextureFactors(WorldTexelFactor, WorldLightmapFactor, DummyBounds, INDEX_NONE, INDEX_NONE);
+	}
+	else
+	{
+		WorldTexelFactor = 0;
+		WorldLightmapFactor = 0;
+	}
 #endif
 
 	if (BodySetup)
@@ -334,6 +346,7 @@ bool FStaticMeshSceneProxy::GetMeshElement(
 		OutBatchElement.MaxVertexIndex = Section.MaxVertexIndex;
 		OutMeshBatch.LODIndex = LODIndex;
 #if !(UE_BUILD_SHIPPING || UE_BUILD_TEST)
+		OutBatchElement.VisualizeElementIndex = SectionIndex;
 		OutMeshBatch.VisualizeLODIndex = LODIndex;
 		OutMeshBatch.VisualizeHLODIndex = HierarchicalLODIndex;
 #endif
@@ -406,6 +419,34 @@ bool FStaticMeshSceneProxy::GetWireframeMeshElement(int32 LODIndex, int32 BatchI
 	return OutBatchElement.NumPrimitives > 0;
 }
 
+#if !(UE_BUILD_SHIPPING || UE_BUILD_TEST)
+bool FStaticMeshSceneProxy::GetStreamingTextureInfo(struct FStreamingTexturePrimitiveInfo& OutInfo, int32 LODIndex, int32 ElementIndex) const
+{
+	int32 DebugIndex = ElementIndex;
+	if (StreamingTextureInfos.Num() && RenderData->LODResources.IsValidIndex(LODIndex) && RenderData->LODResources[LODIndex].Sections.IsValidIndex(ElementIndex))
+	{
+		int32 DebugLODIndex = 0;
+		while (DebugLODIndex < LODIndex)
+		{
+			DebugIndex += RenderData->LODResources[DebugLODIndex].Sections.Num();
+			++DebugLODIndex;
+		}
+	}
+
+	if (StreamingTextureInfos.IsValidIndex(DebugIndex))
+	{
+		OutInfo = StreamingTextureInfos[DebugIndex];
+		return true;
+	}
+	else
+	{
+		OutInfo.Bounds = GetBounds();
+		OutInfo.TexelFactor = WorldTexelFactor;
+		return WorldTexelFactor != 0;
+	}
+}
+#endif
+
 /**
  * Sets IndexBuffer, FirstIndex and NumPrimitives of OutMeshElement.
  */
@@ -1471,3 +1512,4 @@ bool UStaticMeshComponent::ShouldRecreateProxyOnUpdateTransform() const
 {
 	return (Mobility != EComponentMobility::Movable);
 }
+
diff --git a/Engine/Source/Runtime/Engine/Private/TextureStreamingBuild.cpp b/Engine/Source/Runtime/Engine/Private/TextureStreamingBuild.cpp
--- a/Engine/Source/Runtime/Engine/Private/TextureStreamingBuild.cpp
+++ b/Engine/Source/Runtime/Engine/Private/TextureStreamingBuild.cpp
@@ -10,42 +10,87 @@ TextureStreamingBuild.cpp : Contains definitions to build texture streaming data
 
 ENGINE_API void BuildTextureStreamingData(UWorld* InWorld)
 {
+	if (!InWorld) return;
+
 #if WITH_EDITORONLY_DATA
 	FlushRenderingCommands();
-
+	
 	double StartTime = FPlatformTime::Seconds();
 
+	// ====================================================
+	// Build Shaders to compute tex coord scale per texture
+	// ====================================================
+
 	FDebugViewModeMaterialProxy::ClearAllShaders();
 
-	if (InWorld)
+	for (int32 LevelIndex = 0; LevelIndex < InWorld->GetNumLevels(); LevelIndex++)
 	{
-		for (int32 LevelIndex=0; LevelIndex < InWorld->GetNumLevels(); LevelIndex++)
-		{
-			ULevel* Level = InWorld->GetLevel(LevelIndex);
+		ULevel* Level = InWorld->GetLevel(LevelIndex);
 
-			TArray<UObject*> ObjectsInOuter;
-			GetObjectsWithOuter(Level, ObjectsInOuter);
+		TArray<UObject*> ObjectsInOuter;
+		GetObjectsWithOuter(Level, ObjectsInOuter);
 
-			for (int32 Index = 0; Index < ObjectsInOuter.Num(); Index++)
-			{
-				UPrimitiveComponent* Primitive = Cast<UPrimitiveComponent>(ObjectsInOuter[Index]);
-				if (!Primitive) continue;
+		for (int32 Index = 0; Index < ObjectsInOuter.Num(); Index++)
+		{
+			UPrimitiveComponent* Primitive = Cast<UPrimitiveComponent>(ObjectsInOuter[Index]);
+			if (!Primitive) continue;
 
-				TArray<UMaterialInterface*> Materials;
-				Primitive->GetUsedMaterials(Materials);
+			TArray<UMaterialInterface*> Materials;
+			Primitive->GetUsedMaterials(Materials);
 
-				for (int32 MaterialIndex = 0; MaterialIndex < Materials.Num(); ++MaterialIndex)
-				{
-					UMaterialInterface* MaterialInterface = Materials[MaterialIndex];
-					if (!MaterialInterface) continue;
+			for (int32 MaterialIndex = 0; MaterialIndex < Materials.Num(); ++MaterialIndex)
+			{
+				UMaterialInterface* MaterialInterface = Materials[MaterialIndex];
+				if (!MaterialInterface) continue;
 
-					FDebugViewModeMaterialProxy::AddShader(MaterialInterface, EMaterialShaderMapUsage::DebugViewModeTexCoordScale);
-				}
+				FDebugViewModeMaterialProxy::AddShader(MaterialInterface, EMaterialShaderMapUsage::DebugViewModeTexCoordScale);
 			}
 		}
 	}
 	GShaderCompilingManager->FinishAllCompilation();
 
+	// ====================================================
+	// Build per primitive data
+	// ====================================================
+
+	for (int32 LevelIndex = 0; LevelIndex < InWorld->GetNumLevels(); LevelIndex++)
+	{
+		ULevel* Level = InWorld->GetLevel(LevelIndex);
+
+		TArray<UObject*> ObjectsInOuter;
+		GetObjectsWithOuter(Level, ObjectsInOuter);
+
+		for (int32 Index = 0; Index < ObjectsInOuter.Num(); Index++)
+		{
+			UStaticMeshComponent* Primitive = Cast<UStaticMeshComponent>(ObjectsInOuter[Index]);
+			if (!Primitive) continue;
+
+			Primitive->UpdateStreamingTextureInfos();
+		}
+	}
+
+	// ====================================================
+	// Update texture streaming data
+	// ====================================================
+
+	ULevel::BuildStreamingData(InWorld);
+
+	// ====================================================
+	// Reregister everything for debug view modes to reflect changes
+	// ====================================================
+
+	for (int32 LevelIndex = 0; LevelIndex < InWorld->GetNumLevels(); LevelIndex++)
+	{
+		ULevel* Level = InWorld->GetLevel(LevelIndex);
+		if (!Level) continue;
+		
+		for (AActor* Actor : Level->Actors)
+		{
+			if (!Actor) continue;
+
+			Actor->ReregisterAllComponents();
+		}
+	}
 	UE_LOG(LogLevel, Display, TEXT("Build TextureStreaming took %.3f seconds."), FPlatformTime::Seconds() - StartTime);
 #else
 	UE_LOG(LogLevel, Fatal,TEXT("Build TextureStreaming should not be called on a console"));
diff --git a/Engine/Source/Runtime/Engine/Public/GPUSkinCache.h b/Engine/Source/Runtime/Engine/Public/GPUSkinCache.h
--- a/Engine/Source/Runtime/Engine/Public/GPUSkinCache.h
+++ b/Engine/Source/Runtime/Engine/Public/GPUSkinCache.h
@@ -62,10 +62,10 @@ public:
 	//	For each Chunk:
 	//		Call Add*
 	//	Call End*
-	// @param MorphVertexBuffer if no moprh targets are used
+	// @param MorphVertexBuffer if no morph targets are used
 	// @returns -1 if failed, otherwise index into CachedElements[]
 	int32 StartCacheMesh(FRHICommandListImmediate& RHICmdList, int32 Key, class FGPUBaseSkinVertexFactory* VertexFactory,
-		class FGPUSkinPassthroughVertexFactory* TargetVertexFactory, const struct FSkelMeshChunk& BatchElement, const FSkeletalMeshObjectGPUSkin* Skin,
+		class FGPUSkinPassthroughVertexFactory* TargetVertexFactory, const struct FSkelMeshChunk& BatchElement, FSkeletalMeshObjectGPUSkin* Skin,
 		const class FMorphVertexBuffer* MorphVertexBuffer);
 
 	inline bool SetVertexStreamFromCache(FRHICommandList& RHICmdList, int32 Key, FShader* Shader, const FGPUSkinPassthroughVertexFactory* VertexFactory, uint32 BaseVertexIndex, FShaderParameter PreviousStreamFloatOffset, FShaderResourceParameter PreviousStreamBuffer)
@@ -86,10 +86,11 @@ public:
 		const FSkeletalMeshObjectGPUSkin* Skin;
 
 		int32	Key;
-		// index into CachedVertexBuffers[] or -1 if it has no assignment yet
+		// index into CachedSRVs[] or -1 if it has no assignment yet
 		int32	VertexBufferSRVIndex;
+		// index into CachedSRVs[] or -1 if it has no assignment yet
 		int32	MorphVertexBufferSRVIndex;
-
+		
 		uint32	FrameUpdated;
 
 		uint32	InputVBStride;
@@ -127,6 +128,41 @@ public:
 			return (BatchElement == &OtherBatchElement);
 		}
 	};
+	
+	// used temporary when passing data to the Dispatch call chain
+	struct FDispatchData 
+	{
+		FRHICommandListImmediate& RHICmdList;
+
+		ERHIFeatureLevel::Type FeatureLevel;
+		
+		// must not be 0
+		FSkeletalMeshObjectGPUSkin* GPUSkin;
+
+		// 0:normal, 1:with morph target, 2:with APEX cloth (not yet implemented)
+		uint32 SkinType;
+
+		// SkinCache output input for RecomputeSkinTagents
+		FShaderResourceViewRHIParamRef VertexBuffer;
+		uint32 VertexBufferOffset;
+
+		// morph input
+		FShaderResourceViewRHIParamRef MorphBuffer;
+		uint32 MorphBufferOffset;
+
+		FDispatchData(FRHICommandListImmediate& InRHICmdList, ERHIFeatureLevel::Type InFeatureLevel, FSkeletalMeshObjectGPUSkin* InGPUSkin)
+			: RHICmdList(InRHICmdList)
+			, FeatureLevel(InFeatureLevel)
+			, GPUSkin(InGPUSkin)
+			, SkinType(0)
+			, VertexBuffer(0)
+			, VertexBufferOffset(0)
+			, MorphBuffer(0)
+			, MorphBufferOffset(0)
+		{
+			check(GPUSkin);
+		}
+	};
 
 	ENGINE_API void TransitionToReadable(FRHICommandList& RHICmdList);
 	ENGINE_API void TransitionToWriteable(FRHICommandList& RHICmdList);
@@ -135,13 +171,12 @@ public:
 private:
 	FElementCacheStatusInfo* FindEvictableCacheStatusInfo();
 
-	struct FVertexBufferInfo;
+	struct FSRVCacheEntry;
 
-	// @param SkinType 0:normal, 1:with morph target, 2:with APEX cloth (not yet implemented)
-	void DispatchSkinCacheProcess(FRHICommandListImmediate& RHICmdList, uint32 InputStreamFloatOffset, uint32 OutputBufferFloatOffset,
-		const struct FVertexBufferAndSRV& BoneBuffer, FUniformBufferRHIRef UniformBuffer, const FVertexBufferInfo* VBInfo,
-		uint32 VertexStride, uint32 VertexCount, const FVector& MeshOrigin, const FVector& MeshExtension, bool bUseExtraBoneInfluences, ERHIFeatureLevel::Type FeatureLevel,
-		uint32 SkinType, uint32 InMorphInputStreamFloatOffset, FShaderResourceViewRHIParamRef MorphVertexBufferSRV);
+	void DispatchSkinCacheProcess(uint32 InputStreamFloatOffset, uint32 OutputBufferFloatOffset,
+		const struct FVertexBufferAndSRV& BoneBuffer, FUniformBufferRHIRef UniformBuffer, const FSRVCacheEntry* VBInfo,
+		uint32 VertexStride, uint32 VertexCount, const FVector& MeshOrigin, const FVector& MeshExtension, bool bUseExtraBoneInfluences,
+		const FDispatchData& DispatchData);
 
 	bool InternalSetVertexStreamFromCache(FRHICommandList& RHICmdList, int32 Key, FShader* Shader, const class FGPUSkinPassthroughVertexFactory* VertexFactory, uint32 BaseVertexIndex, FShaderParameter PreviousStreamFloatOffset, FShaderResourceParameter PreviousStreamBuffer);
 
@@ -168,24 +203,29 @@ private:
 	// -------------------------------------------------
 	
 	// to cache the SRV needed to access the input vertex buffers
-	struct FVertexBufferInfo
+	struct FSRVCacheEntry
 	{
-		// key: input vertex buffer from which we need an SRV
-		const FVertexBuffer* VertexBuffer;
+		// key: input vertex/index buffer from which we need an SRV
+		const FRenderResource* Key;
 		// value: SRV created from the VertexBuffer
-		FShaderResourceViewRHIRef VertexBufferSRV;
+		FShaderResourceViewRHIRef SRVValue;
 		
 		// to support FindByKey()
-		bool operator == (const FVertexBuffer& OtherVertexBuffer) const 
+		bool operator == (const FRenderResource& Other) const 
 		{
-			return (VertexBuffer == &OtherVertexBuffer);
+			return (Key == &Other);
 		}
 	};
-	
+
 	// to cache the SRV needed to access the input vertex buffers
-	TArray<FVertexBufferInfo> CachedVertexBuffers;
+	TArray<FSRVCacheEntry> CachedSRVs;
 	
-	FVertexBufferInfo* ReuseOrFindAndCacheSRV(int32& InOutIndex, const FVertexBuffer& VertexBuffer);
+	// @param InOutIndex index -1 if not yet set
+	// @return 0 if we failed to find a buffer
+	FSRVCacheEntry* ReuseOrFindAndCacheSRV(int32& InOutIndex, const FVertexBuffer& InKey);
+	// @param InOutIndex index -1 if not yet set
+	// @return 0 if we failed to find a buffer
+	FSRVCacheEntry* ReuseOrFindAndCacheSRV(int32& InOutIndex, const FIndexBuffer& InKey);
 };
 
 BEGIN_UNIFORM_BUFFER_STRUCT(GPUSkinCacheBonesUniformShaderParameters,)
diff --git a/Engine/Source/Runtime/Engine/Public/GPUSkinVertexFactory.h b/Engine/Source/Runtime/Engine/Public/GPUSkinVertexFactory.h
--- a/Engine/Source/Runtime/Engine/Public/GPUSkinVertexFactory.h
+++ b/Engine/Source/Runtime/Engine/Public/GPUSkinVertexFactory.h
@@ -357,6 +357,7 @@ public:
 		}
 
 	private:
+		// double buffered bone positions+orientations to support normal rendering and velocity (new-old position) rendering
 		FVertexBufferAndSRV BoneBuffer[2];
 		// 0 / 1 to index into BoneBuffer
 		uint32 CurrentBuffer;
@@ -617,7 +618,7 @@ protected:
 
 private:
 	/** stream component data bound to this vertex factory */
-	FDataType MorphData; 
+	FDataType MorphData;
 
 };
 
diff --git a/Engine/Source/Runtime/Engine/Public/MeshBatch.h b/Engine/Source/Runtime/Engine/Public/MeshBatch.h
--- a/Engine/Source/Runtime/Engine/Public/MeshBatch.h
+++ b/Engine/Source/Runtime/Engine/Public/MeshBatch.h
@@ -51,6 +51,11 @@ struct FMeshBatchElement
 	uint8 bIsSplineProxy : 1;
 	uint8 bIsInstanceRuns : 1;
 
+#if !(UE_BUILD_SHIPPING || UE_BUILD_TEST)
+	/** Conceptual element index used for debug viewmodes. */
+	int8 VisualizeElementIndex;
+#endif
+
 	FMeshBatchElement()
 	:	PrimitiveUniformBufferResource(nullptr)
 	,	IndexBuffer(nullptr)
@@ -67,6 +72,10 @@ struct FMeshBatchElement
 	,   bIsInstancedMesh(false)
 	,	bIsSplineProxy(false)
 	,	bIsInstanceRuns(false)
+#if !(UE_BUILD_SHIPPING || UE_BUILD_TEST)
+	,	VisualizeElementIndex(INDEX_NONE)
+#endif
+
 	{
 	}
 };
@@ -84,8 +93,10 @@ struct FMeshBatch
 	/** LOD index of the mesh, used for fading LOD transitions. */
 	int8 LODIndex;
 
+#if !(UE_BUILD_SHIPPING || UE_BUILD_TEST)
 	/** Conceptual LOD index used for the LOD Coloration visualization. */
 	int8 VisualizeLODIndex;
+#endif
 
 	/** Conceptual HLOD index used for the HLOD Coloration visualization. */
 	int8 VisualizeHLODIndex;
@@ -214,7 +225,9 @@ struct FMeshBatch
 	FMeshBatch()
 	:	DynamicVertexStride(0)
 	,	LODIndex(INDEX_NONE)
+#if !(UE_BUILD_SHIPPING || UE_BUILD_TEST)
 	,	VisualizeLODIndex(INDEX_NONE)
+#endif
 	,	VisualizeHLODIndex(INDEX_NONE)
 	,	UseDynamicData(false)
 	,	ReverseCulling(false)
diff --git a/Engine/Source/Runtime/Engine/Public/PrimitiveSceneProxy.h b/Engine/Source/Runtime/Engine/Public/PrimitiveSceneProxy.h
--- a/Engine/Source/Runtime/Engine/Public/PrimitiveSceneProxy.h
+++ b/Engine/Source/Runtime/Engine/Public/PrimitiveSceneProxy.h
@@ -520,14 +520,9 @@ public:
 	ENGINE_API virtual void GetLCIs(FLCIArray& LCIs) {}
 
 	/**
-	 * Get the texture streaming texel factor for this primitive. Taken from ComputeStreamingTextureFactors
+	 * Get the texture streaming texel factor and bound for this primitive.
 	 */
-#if !(UE_BUILD_SHIPPING || UE_BUILD_TEST)
-	float GetWorldTexelFactor() const { return WorldTexelFactor; }
-#else
-	FORCEINLINE float GetWorldTexelFactor() const { return 0; }
-#endif
-
+	virtual bool GetStreamingTextureInfo(struct FStreamingTexturePrimitiveInfo& OutInfo, int32 LODIndex, int32 ElementIndex) const { return false; }
 
 protected:
 
@@ -725,11 +720,6 @@ protected:
 	/** Quality of interpolated indirect lighting for Movable components. */
 	EIndirectLightingCacheQuality IndirectLightingCacheQuality;
 
-#if !(UE_BUILD_SHIPPING || UE_BUILD_TEST)
-	/** Texture streaming factor used for material textures */
-	float WorldTexelFactor;
-#endif
-
 private:
 	/** The primitive's local to world transform. */
 	FMatrix LocalToWorld;
diff --git a/Engine/Source/Runtime/Engine/Public/SkeletalMeshTypes.h b/Engine/Source/Runtime/Engine/Public/SkeletalMeshTypes.h
--- a/Engine/Source/Runtime/Engine/Public/SkeletalMeshTypes.h
+++ b/Engine/Source/Runtime/Engine/Public/SkeletalMeshTypes.h
@@ -351,7 +351,9 @@ struct FSkelMeshChunk
 	/** max # of bones used to skin the vertices in this chunk */
 	int32 MaxBoneInfluences;
 
+	// INDEX_NONE if not set
 	int16 CorrespondClothAssetIndex;
+	// INDEX_NONE if not set
 	int16 ClothAssetSubmeshIndex;
 
 	FSkelMeshChunk()
@@ -1487,7 +1489,10 @@ public:
 	/** 
 	* Rendering data.
 	*/
+
+	// Index Buffer (MultiSize: 16bit or 32bit)
 	FMultiSizeIndexContainer	MultiSizeIndexContainer; 
+	//
 	uint32						NumVertices;
 	/** The number of unique texture coordinate sets in this lod */
 	uint32						NumTexCoords;
diff --git a/Engine/Source/Runtime/Engine/Public/StaticMeshResources.h b/Engine/Source/Runtime/Engine/Public/StaticMeshResources.h
--- a/Engine/Source/Runtime/Engine/Public/StaticMeshResources.h
+++ b/Engine/Source/Runtime/Engine/Public/StaticMeshResources.h
@@ -786,6 +786,9 @@ public:
 	/** Sets up a wireframe FMeshBatch for a specific LOD. */
 	virtual bool GetWireframeMeshElement(int32 LODIndex, int32 BatchIndex, const FMaterialRenderProxy* WireframeRenderProxy, uint8 InDepthPriorityGroup, bool bAllowPreCulledIndices, FMeshBatch& OutMeshBatch) const;
 
+#if !(UE_BUILD_SHIPPING || UE_BUILD_TEST)
+	virtual bool GetStreamingTextureInfo(struct FStreamingTexturePrimitiveInfo& OutInfo, int32 LODIndex, int32 ElementIndex) const override;
+#endif
 
 protected:
 	/**
@@ -914,7 +917,10 @@ protected:
 #endif
 
 #if !(UE_BUILD_SHIPPING || UE_BUILD_TEST)
+	/** The precise streaming info for each LOD and Section. */
+	TArray<struct FStreamingTexturePrimitiveInfo> StreamingTextureInfos;
 	/** Texture streaming factor used for lightmap textures */
+	float WorldTexelFactor;
 	float WorldLightmapFactor;
 
 	/** LOD used for collision */
diff --git a/Engine/Source/Runtime/OpenGLDrv/Private/OpenGLDevice.cpp b/Engine/Source/Runtime/OpenGLDrv/Private/OpenGLDevice.cpp
--- a/Engine/Source/Runtime/OpenGLDrv/Private/OpenGLDevice.cpp
+++ b/Engine/Source/Runtime/OpenGLDrv/Private/OpenGLDevice.cpp
@@ -698,6 +698,7 @@ static void InitRHICapabilitiesForGL()
 	GMaxShadowDepthBufferSizeY = FMath::Min<int32>(Value_GL_MAX_RENDERBUFFER_SIZE, 4096);
 	GHardwareHiddenSurfaceRemoval = FOpenGL::HasHardwareHiddenSurfaceRemoval();
 	GRHISupportsInstancing = FOpenGL::SupportsInstancing(); // HTML5 does not support it. Android supports it with OpenGL ES3.0+ 
+	GSupportsTimestampRenderQueries = FOpenGL::SupportsTimestampQueries();
 
 	GSupportsHDR32bppEncodeModeIntrinsic = FOpenGL::SupportsHDR32bppEncodeModeIntrinsic();
 
diff --git a/Engine/Source/Runtime/RHI/Private/DynamicRHI.cpp b/Engine/Source/Runtime/RHI/Private/DynamicRHI.cpp
--- a/Engine/Source/Runtime/RHI/Private/DynamicRHI.cpp
+++ b/Engine/Source/Runtime/RHI/Private/DynamicRHI.cpp
@@ -25,7 +25,7 @@ static TAutoConsoleVariable<int32> CVarWarnOfBadDrivers(
 	TEXT(" 2: Simulating the system has a blacklisted NVIDIA driver (UI should appear)\n")
 	TEXT(" 3: Simulating the system has a blacklisted AMD driver (UI should appear)\n")
 	TEXT(" 4: Simulating the system has a not blacklisted AMD driver (no UI should appear)\n")
-	TEXT(" 5: Simulating the system has a Indel driver (no UI should appear)"),
+	TEXT(" 5: Simulating the system has a Intel driver (no UI should appear)"),
 	ECVF_RenderThreadSafe
 	);
 
diff --git a/Engine/Source/Runtime/RHI/Private/RHI.cpp b/Engine/Source/Runtime/RHI/Private/RHI.cpp
--- a/Engine/Source/Runtime/RHI/Private/RHI.cpp
+++ b/Engine/Source/Runtime/RHI/Private/RHI.cpp
@@ -235,6 +235,7 @@ bool GSupportsRenderTargetFormat_PF_G8 = true;
 bool GSupportsRenderTargetFormat_PF_FloatRGBA = true;
 bool GSupportsShaderFramebufferFetch = false;
 bool GSupportsShaderDepthStencilFetch = false;
+bool GSupportsTimestampRenderQueries = false;
 bool GHardwareHiddenSurfaceRemoval = false;
 bool GRHISupportsAsyncTextureCreation = false;
 bool GSupportsQuads = false;
diff --git a/Engine/Source/Runtime/RHI/Public/RHI.h b/Engine/Source/Runtime/RHI/Public/RHI.h
--- a/Engine/Source/Runtime/RHI/Public/RHI.h
+++ b/Engine/Source/Runtime/RHI/Public/RHI.h
@@ -102,6 +102,9 @@ extern RHI_API bool GSupportsShaderFramebufferFetch;
 /** true if mobile depth & stencil fetch is supported */
 extern RHI_API bool GSupportsShaderDepthStencilFetch;
 
+/** true if RQT_AbsoluteTime is supported by RHICreateRenderQuery */
+extern RHI_API bool GSupportsTimestampRenderQueries;
+
 /** true if the GPU supports hidden surface removal in hardware. */
 extern RHI_API bool GHardwareHiddenSurfaceRemoval;
 
diff --git a/Engine/Source/Runtime/RHI/Public/RHICommandList.h b/Engine/Source/Runtime/RHI/Public/RHICommandList.h
--- a/Engine/Source/Runtime/RHI/Public/RHICommandList.h
+++ b/Engine/Source/Runtime/RHI/Public/RHICommandList.h
@@ -2562,6 +2562,11 @@ public:
 		return GDynamicRHI->CreateShaderResourceView_RenderThread(*this, VertexBuffer, Stride, Format);
 	}
 	
+	FORCEINLINE FShaderResourceViewRHIRef CreateShaderResourceView(FIndexBufferRHIParamRef Buffer)
+	{
+		return GDynamicRHI->CreateShaderResourceView_RenderThread(*this, Buffer);
+	}
+	
 	FORCEINLINE uint64 CalcTexture2DPlatformSize(uint32 SizeX, uint32 SizeY, uint8 Format, uint32 NumMips, uint32 NumSamples, uint32 Flags, uint32& OutAlign)
 	{
 		return RHICalcTexture2DPlatformSize(SizeX, SizeY, Format, NumMips, NumSamples, Flags, OutAlign);
@@ -3189,6 +3194,11 @@ FORCEINLINE FShaderResourceViewRHIRef RHICreateShaderResourceView(FVertexBufferR
 	return FRHICommandListExecutor::GetImmediateCommandList().CreateShaderResourceView(VertexBuffer, Stride, Format);
 }
 
+FORCEINLINE FShaderResourceViewRHIRef RHICreateShaderResourceView(FIndexBufferRHIParamRef Buffer)
+{
+	return FRHICommandListExecutor::GetImmediateCommandList().CreateShaderResourceView(Buffer);
+}
+
 FORCEINLINE FTextureReferenceRHIRef RHICreateTextureReference(FLastRenderTimeContainer* LastRenderTime)
 {
 	return FRHICommandListExecutor::GetImmediateCommandList().CreateTextureReference(LastRenderTime);
diff --git a/Engine/Source/Runtime/Renderer/Private/BasePassRendering.h b/Engine/Source/Runtime/Renderer/Private/BasePassRendering.h
--- a/Engine/Source/Runtime/Renderer/Private/BasePassRendering.h
+++ b/Engine/Source/Runtime/Renderer/Private/BasePassRendering.h
@@ -530,9 +530,9 @@ public:
 				{
 					check(GetUniformBufferParameter<FForwardLightData>().IsInitialized());
 
-					SetUniformBufferParameter(RHICmdList, ShaderRHI,GetUniformBufferParameter<FForwardLightData>(),View->ForwardLightData);
-					SetSRVParameter(RHICmdList, ShaderRHI, LightGrid, GLightGridVertexBuffer.VertexBufferSRV);
-				}
+						SetUniformBufferParameter(RHICmdList, ShaderRHI,GetUniformBufferParameter<FForwardLightData>(),View->ForwardLightData);
+						SetSRVParameter(RHICmdList, ShaderRHI, LightGrid, GLightGridVertexBuffer.VertexBufferSRV);
+					}
 			}
 		}
 		
@@ -942,7 +942,7 @@ public:
 				RHICmdList.SetBlendState(TStaticBlendState<CW_RGB,BO_Add,BF_One,BF_One>::GetRHI());
 			}
 
-			FDebugViewMode::GetPSInterface(View.ShaderMap, MaterialResource, GetDebugViewShaderMode())->SetMesh(RHICmdList, VertexFactory, View, PrimitiveSceneProxy, BatchElement, DrawRenderState);
+			FDebugViewMode::GetPSInterface(View.ShaderMap, MaterialResource, GetDebugViewShaderMode())->SetMesh(RHICmdList, VertexFactory, View, PrimitiveSceneProxy, Mesh.VisualizeLODIndex, BatchElement, DrawRenderState);
 		}
 		else
 #endif
diff --git a/Engine/Source/Runtime/Renderer/Private/DebugViewModeRendering.cpp b/Engine/Source/Runtime/Renderer/Private/DebugViewModeRendering.cpp
--- a/Engine/Source/Runtime/Renderer/Private/DebugViewModeRendering.cpp
+++ b/Engine/Source/Runtime/Renderer/Private/DebugViewModeRendering.cpp
@@ -49,6 +49,7 @@ public:
 		const FVertexFactory* VertexFactory,
 		const FSceneView& View,
 		const FPrimitiveSceneProxy* Proxy,
+		int32 VisualizeLODIndex,
 		const FMeshBatchElement& BatchElement, 
 		const FMeshDrawingRenderState& DrawRenderState
 		) override{}
diff --git a/Engine/Source/Runtime/Renderer/Private/DebugViewModeRendering.h b/Engine/Source/Runtime/Renderer/Private/DebugViewModeRendering.h
--- a/Engine/Source/Runtime/Renderer/Private/DebugViewModeRendering.h
+++ b/Engine/Source/Runtime/Renderer/Private/DebugViewModeRendering.h
@@ -6,7 +6,7 @@ DebugViewModeRendering.h: Contains definitions for rendering debug viewmodes.
 
 #pragma once
 
-// Used to know whether QuadOverdraw should be used at runtime.
+// Used to know whether debug viewmodes are allowed at runtime.
 #if !(UE_BUILD_SHIPPING || UE_BUILD_TEST)
 FORCEINLINE bool RuntimeAllowDebugViewModeShader(ERHIFeatureLevel::Type InFeatureLevel)
 {
@@ -106,6 +106,7 @@ public:
 		const FVertexFactory* VertexFactory,
 		const FSceneView& View,
 		const FPrimitiveSceneProxy* Proxy,
+		int32 VisualizeLODIndex,
 		const FMeshBatchElement& BatchElement, 
 		const FMeshDrawingRenderState& DrawRenderState
 		) = 0;
diff --git a/Engine/Source/Runtime/Renderer/Private/DecalRenderingShared.cpp b/Engine/Source/Runtime/Renderer/Private/DecalRenderingShared.cpp
--- a/Engine/Source/Runtime/Renderer/Private/DecalRenderingShared.cpp
+++ b/Engine/Source/Runtime/Renderer/Private/DecalRenderingShared.cpp
@@ -235,7 +235,9 @@ void FDecalRendering::BuildVisibleDecalList(const FScene& Scene, const FViewInfo
 
 	const float FadeMultiplier = CVarDecalFadeScreenSizeMultiplier.GetValueOnRenderThread();
 	const EShaderPlatform ShaderPlatform = View.GetShaderPlatform();
-	
+
+	const bool bIsPerspectiveProjection = View.IsPerspectiveProjection();
+
 	// Build a list of decals that need to be rendered for this view in SortedDecals
 	for (const FDeferredDecalProxy* DecalProxy : Scene.Decals)
 	{
@@ -273,7 +275,7 @@ void FDecalRendering::BuildVisibleDecalList(const FScene& Scene, const FViewInfo
 			// filter out decals with blend modes that are not supported on current platform
 			if (IsBlendModeSupported(ShaderPlatform, Data.DecalBlendMode))
 			{
-				if (Data.DecalProxy->Component->FadeScreenSize != 0.0f)
+				if (bIsPerspectiveProjection && Data.DecalProxy->Component->FadeScreenSize != 0.0f)
 				{
 					float Distance = (View.ViewMatrices.ViewOrigin - ComponentToWorldMatrix.GetOrigin()).Size();
 					float Radius = ComponentToWorldMatrix.GetMaximumAxisScale();
diff --git a/Engine/Source/Runtime/Renderer/Private/DeferredShadingRenderer.cpp b/Engine/Source/Runtime/Renderer/Private/DeferredShadingRenderer.cpp
--- a/Engine/Source/Runtime/Renderer/Private/DeferredShadingRenderer.cpp
+++ b/Engine/Source/Runtime/Renderer/Private/DeferredShadingRenderer.cpp
@@ -722,6 +722,7 @@ DECLARE_CYCLE_STAT(TEXT("RenderFinish"), STAT_CLM_RenderFinish, STATGROUP_Comman
 DECLARE_CYCLE_STAT(TEXT("AfterFrame"), STAT_CLM_AfterFrame, STATGROUP_CommandListMarkers);
 
 FGraphEventRef FDeferredShadingSceneRenderer::OcclusionSubmittedFence[FOcclusionQueryHelpers::MaxBufferedOcclusionFrames];
+FGraphEventRef FDeferredShadingSceneRenderer::TranslucencyTimestampQuerySubmittedFence[FOcclusionQueryHelpers::MaxBufferedOcclusionFrames + 1];
 
 /**
  * Returns true if the depth Prepass needs to run
@@ -2117,11 +2118,11 @@ public:
 		ProjectionScaleBias.Bind(Initializer.ParameterMap,TEXT("ProjectionScaleBias"));
 		SourceTexelOffsets01.Bind(Initializer.ParameterMap,TEXT("SourceTexelOffsets01"));
 		SourceTexelOffsets23.Bind(Initializer.ParameterMap,TEXT("SourceTexelOffsets23"));
-		MinMaxBlendParameter.Bind(Initializer.ParameterMap, TEXT("MinMaxBlend"));
+		UseMaxDepth.Bind(Initializer.ParameterMap, TEXT("UseMaxDepth"));
 	}
 	FDownsampleSceneDepthPS() {}
 
-	void SetParameters(FRHICommandList& RHICmdList, const FSceneView& View, float MinMaxBlend)
+	void SetParameters(FRHICommandList& RHICmdList, const FSceneView& View, bool bUseMaxDepth)
 	{
 		FGlobalShader::SetParameters(RHICmdList, GetPixelShader(), View);
 		FSceneRenderTargets& SceneContext = FSceneRenderTargets::Get(RHICmdList);
@@ -2129,7 +2130,7 @@ public:
 		// Used to remap view space Z (which is stored in scene color alpha) into post projection z and w so we can write z/w into the downsized depth buffer
 		const FVector2D ProjectionScaleBiasValue(View.ViewMatrices.ProjMatrix.M[2][2], View.ViewMatrices.ProjMatrix.M[3][2]);
 		SetShaderValue(RHICmdList, GetPixelShader(), ProjectionScaleBias, ProjectionScaleBiasValue);
-		SetShaderValue(RHICmdList, GetPixelShader(), MinMaxBlendParameter, MinMaxBlend);
+		SetShaderValue(RHICmdList, GetPixelShader(), UseMaxDepth, (bUseMaxDepth ? 1.0f : 0.0f));
 
 		FIntPoint BufferSize = SceneContext.GetBufferSizeXY();
 
@@ -2151,7 +2152,7 @@ public:
 		Ar << SourceTexelOffsets01;
 		Ar << SourceTexelOffsets23;
 		Ar << SceneTextureParameters;
-		Ar << MinMaxBlendParameter;
+		Ar << UseMaxDepth;
 		return bShaderHasOutdatedParameters;
 	}
 
@@ -2159,7 +2160,7 @@ public:
 	FShaderParameter SourceTexelOffsets01;
 	FShaderParameter SourceTexelOffsets23;
 	FSceneTextureShaderParameters SceneTextureParameters;
-	FShaderParameter MinMaxBlendParameter;
+	FShaderParameter UseMaxDepth;
 };
 
 IMPLEMENT_SHADER_TYPE(,FDownsampleSceneDepthPS,TEXT("DownsampleDepthPixelShader"),TEXT("Main"),SF_Pixel);
@@ -2177,16 +2178,14 @@ void FDeferredShadingSceneRenderer::UpdateDownsampledDepthSurface(FRHICommandLis
 		for (int32 ViewIndex = 0; ViewIndex < Views.Num(); ViewIndex++)
 		{
 			const FViewInfo& View = Views[ViewIndex];
-			DownsampleDepthSurface(RHICmdList, SceneContext.GetSmallDepthSurface(), View, 1.0f / SceneContext.GetSmallColorDepthDownsampleFactor());
+			DownsampleDepthSurface(RHICmdList, SceneContext.GetSmallDepthSurface(), View, 1.0f / SceneContext.GetSmallColorDepthDownsampleFactor(), true);
 		}
 	}
 }
 
-
-
 /** Downsample the scene depth with a specified scale factor to a specified render target
 */
-void FDeferredShadingSceneRenderer::DownsampleDepthSurface(FRHICommandList& RHICmdList, const FTexture2DRHIRef& RenderTarget, const FViewInfo &View, float ScaleFactor, float MinMaxFilterBlend)
+void FDeferredShadingSceneRenderer::DownsampleDepthSurface(FRHICommandList& RHICmdList, const FTexture2DRHIRef& RenderTarget, const FViewInfo& View, float ScaleFactor, bool bUseMaxDepth)
 {
 	FSceneRenderTargets& SceneContext = FSceneRenderTargets::Get(RHICmdList);
 	SetRenderTarget(RHICmdList, NULL, RenderTarget);
@@ -2204,7 +2203,7 @@ void FDeferredShadingSceneRenderer::DownsampleDepthSurface(FRHICommandList& RHIC
 	RHICmdList.SetRasterizerState(TStaticRasterizerState<FM_Solid, CM_None>::GetRHI());
 	RHICmdList.SetDepthStencilState(TStaticDepthStencilState<true, CF_Always>::GetRHI());
 
-	PixelShader->SetParameters(RHICmdList, View, MinMaxFilterBlend);
+	PixelShader->SetParameters(RHICmdList, View, bUseMaxDepth);
 	const uint32 DownsampledX = FMath::TruncToInt(View.ViewRect.Min.X * ScaleFactor);
 	const uint32 DownsampledY = FMath::TruncToInt(View.ViewRect.Min.Y * ScaleFactor);
 	const uint32 DownsampledSizeX = FMath::TruncToInt(View.ViewRect.Width() * ScaleFactor);
diff --git a/Engine/Source/Runtime/Renderer/Private/DeferredShadingRenderer.h b/Engine/Source/Runtime/Renderer/Private/DeferredShadingRenderer.h
--- a/Engine/Source/Runtime/Renderer/Private/DeferredShadingRenderer.h
+++ b/Engine/Source/Runtime/Renderer/Private/DeferredShadingRenderer.h
@@ -127,6 +127,7 @@ private:
 
 	// fences to make sure the rhi thread has digested the occlusion query renders before we attempt to read them back async
 	static FGraphEventRef OcclusionSubmittedFence[FOcclusionQueryHelpers::MaxBufferedOcclusionFrames];
+	static FGraphEventRef TranslucencyTimestampQuerySubmittedFence[FOcclusionQueryHelpers::MaxBufferedOcclusionFrames + 1];
 
 	/** Creates a per object projected shadow for the given interaction. */
 	void CreatePerObjectProjectedShadow(
@@ -162,6 +163,9 @@ private:
 
 	void InitViewsPossiblyAfterPrepass(FRHICommandListImmediate& RHICmdList, struct FILCUpdatePrimTaskData& ILCTaskData, FGraphEventArray& SortEvents);
 
+	/** Updates auto-downsampling of separate translucency and sets FSceneRenderTargets::SeparateTranslucencyBufferSize */
+	void UpdateSeparateTranslucencyBufferSize(FRHICommandListImmediate& RHICmdList);
+
 	void CreateIndirectCapsuleShadows();
 
 	/**
@@ -238,6 +242,12 @@ private:
 	/** Render stationary light overlap as complexity to scene color. */
 	void RenderStationaryLightOverlap(FRHICommandListImmediate& RHICmdList);
 	
+	/** Issues a timestamp query for the beginning of the separate translucency pass. */
+	void BeginTimingSeparateTranslucencyPass(FRHICommandListImmediate& RHICmdList, const FViewInfo& View);
+
+	/** Issues a timestamp query for the end of the separate translucency pass. */
+	void EndTimingSeparateTranslucencyPass(FRHICommandListImmediate& RHICmdList, const FViewInfo& View);
+
 	/** 
 	 * Renders the scene's translucency, parallel version
 	 */
@@ -279,7 +289,7 @@ private:
 	void UpdateDownsampledDepthSurface(FRHICommandList& RHICmdList);
 
 	/** Downsample the scene depth with a specified scale factor to a specified render target*/
-	void DownsampleDepthSurface(FRHICommandList& RHICmdList, const FTexture2DRHIRef& RenderTarget, const FViewInfo &View, float ScaleFactor, float MinMaxFilterBlend = 0.0f);
+	void DownsampleDepthSurface(FRHICommandList& RHICmdList, const FTexture2DRHIRef& RenderTarget, const FViewInfo& View, float ScaleFactor, bool bUseMaxDepth);
 
 	void CopyStencilToLightingChannelTexture(FRHICommandList& RHICmdList);
 
diff --git a/Engine/Source/Runtime/Renderer/Private/ForwardBasePassRendering.h b/Engine/Source/Runtime/Renderer/Private/ForwardBasePassRendering.h
--- a/Engine/Source/Runtime/Renderer/Private/ForwardBasePassRendering.h
+++ b/Engine/Source/Runtime/Renderer/Private/ForwardBasePassRendering.h
@@ -765,7 +765,7 @@ public:
 				RHICmdList.SetBlendState(TStaticBlendState<CW_RGB,BO_Add,BF_One,BF_One>::GetRHI());
 			}
 
-			FDebugViewMode::GetPSInterface(GetGlobalShaderMap(View.FeatureLevel), MaterialResource, GetDebugViewShaderMode())->SetMesh(RHICmdList, VertexFactory, View, PrimitiveSceneProxy, BatchElement, DrawRenderState);
+			FDebugViewMode::GetPSInterface(GetGlobalShaderMap(View.FeatureLevel), MaterialResource, GetDebugViewShaderMode())->SetMesh(RHICmdList, VertexFactory, View, PrimitiveSceneProxy, Mesh.VisualizeLODIndex, BatchElement, DrawRenderState);
 		}
 		else
 #endif
diff --git a/Engine/Source/Runtime/Renderer/Private/PostProcess/PostProcessBokehDOFRecombine.cpp b/Engine/Source/Runtime/Renderer/Private/PostProcess/PostProcessBokehDOFRecombine.cpp
--- a/Engine/Source/Runtime/Renderer/Private/PostProcess/PostProcessBokehDOFRecombine.cpp
+++ b/Engine/Source/Runtime/Renderer/Private/PostProcess/PostProcessBokehDOFRecombine.cpp
@@ -13,10 +13,16 @@
 #include "PostProcessBokehDOF.h"
 #include "SceneUtils.h"
 
+static TAutoConsoleVariable<int32> CVarSeparateTranslucencyUpsampleMode(
+	TEXT("r.SeparateTranslucencyUpsampleMode"),
+	1,
+	TEXT("Upsample method to use on separate translucency.  These are only used when r.SeparateTranslucencyScreenPercentage is less than 100.\n")
+	TEXT("0: bilinear 1: Nearest-Depth Neighbor (only when r.SeparateTranslucencyScreenPercentage is 50)"),
+	ECVF_Scalability | ECVF_Default);
 
 /**
  * Encapsulates a shader to combined Depth of field and separate translucency layers.
- * @param Method 1:DOF, 2:SeparateTranslucency, 3:DOF+SeparateTranslucency
+ * @param Method 1:DOF, 2:SeparateTranslucency, 3:DOF+SeparateTranslucency, 4:SeparateTranslucency with Nearest-Depth Neighbor, 5:DOF+SeparateTranslucency with Nearest-Depth Neighbor
  */
 template <uint32 Method>
 class FPostProcessBokehDOFRecombinePS : public FGlobalShader
@@ -28,10 +34,28 @@ class FPostProcessBokehDOFRecombinePS : public FGlobalShader
 		return IsFeatureLevelSupported(Platform, ERHIFeatureLevel::SM4);
 	}
 
+	static int32 GetCombineFeatureMethod()
+	{
+		if (Method <= 3)
+		{
+			return Method;
+		}
+		else
+		{
+			return Method - 2;
+		}
+	}
+
+	static bool UseNearestDepthNeighborUpsample()
+	{
+		return Method > 3;
+	}
+
 	static void ModifyCompilationEnvironment(EShaderPlatform Platform, FShaderCompilerEnvironment& OutEnvironment)
 	{
 		FGlobalShader::ModifyCompilationEnvironment(Platform, OutEnvironment);
-		OutEnvironment.SetDefine(TEXT("RECOMBINE_METHOD"), Method);
+		OutEnvironment.SetDefine(TEXT("RECOMBINE_METHOD"), GetCombineFeatureMethod());
+		OutEnvironment.SetDefine(TEXT("NEAREST_DEPTH_NEIGHBOR_UPSAMPLE"), (UseNearestDepthNeighborUpsample() ? 1 : 0));
 	}
 
 	/** Default constructor. */
@@ -42,6 +66,7 @@ public:
 	FDeferredPixelShaderParameters DeferredParameters;
 	FShaderParameter DepthOfFieldParams;
 	FShaderParameter SeparateTranslucencyResMultParam;
+	FShaderResourceParameter LowResDepthTexture;
 
 	/** Initialization constructor. */
 	FPostProcessBokehDOFRecombinePS(const ShaderMetaType::CompiledShaderInitializerType& Initializer)
@@ -51,28 +76,31 @@ public:
 		DeferredParameters.Bind(Initializer.ParameterMap);
 		DepthOfFieldParams.Bind(Initializer.ParameterMap,TEXT("DepthOfFieldParams"));
 		SeparateTranslucencyResMultParam.Bind(Initializer.ParameterMap, TEXT("SeparateTranslucencyResMult"));
+		LowResDepthTexture.Bind(Initializer.ParameterMap, TEXT("LowResDepthTexture"));
 	}
 
 	// FShader interface.
 	virtual bool Serialize(FArchive& Ar) override
 	{
 		bool bShaderHasOutdatedParameters = FGlobalShader::Serialize(Ar);
-		Ar << PostprocessParameter << DepthOfFieldParams << DeferredParameters << SeparateTranslucencyResMultParam;
+		Ar << PostprocessParameter << DepthOfFieldParams << DeferredParameters << SeparateTranslucencyResMultParam << LowResDepthTexture;
 		return bShaderHasOutdatedParameters;
 	}
 
 	void SetParameters(const FRenderingCompositePassContext& Context)
 	{
-		static IConsoleVariable* STSP_CVar = IConsoleManager::Get().FindConsoleVariable(TEXT("r.SeparateTranslucencyScreenPercentage"));
-
 		const FPixelShaderRHIParamRef ShaderRHI = GetPixelShader();
 
 		FGlobalShader::SetParameters(Context.RHICmdList, ShaderRHI, Context.View);
 		DeferredParameters.Set(Context.RHICmdList, ShaderRHI, Context.View);
 		PostprocessParameter.SetPS(ShaderRHI, Context, TStaticSamplerState<SF_Bilinear,AM_Clamp,AM_Clamp,AM_Clamp>::GetRHI());
 
-		float Scale = STSP_CVar->GetInt() / 100.0f;
-		SetShaderValue(Context.RHICmdList, ShaderRHI, SeparateTranslucencyResMultParam, FVector4(Scale, Scale, Scale, Scale));
+		FSceneRenderTargets& SceneContext = FSceneRenderTargets::Get(Context.RHICmdList);
+		FIntPoint OutScaledSize;
+		float OutScale;
+		SceneContext.GetSeparateTranslucencyDimensions(OutScaledSize, OutScale);
+
+		SetShaderValue(Context.RHICmdList, ShaderRHI, SeparateTranslucencyResMultParam, FVector4(OutScale, OutScale, OutScale, OutScale));
 
 		{
 			FVector4 DepthOfFieldParamValues[2];
@@ -81,6 +109,20 @@ public:
 
 			SetShaderValueArray(Context.RHICmdList, ShaderRHI, DepthOfFieldParams, DepthOfFieldParamValues, 2);
 		}
+
+		if (UseNearestDepthNeighborUpsample())
+		{
+			check(SceneContext.IsSeparateTranslucencyDepthValid());
+			FTextureRHIParamRef LowResDepth = SceneContext.GetSeparateTranslucencyDepthSurface();
+			SetTextureParameter(Context.RHICmdList, ShaderRHI, LowResDepthTexture, LowResDepth);
+
+			const auto& BuiltinSamplersUBParameter = GetUniformBufferParameter<FBuiltinSamplersParameters>();
+			SetUniformBufferParameter(Context.RHICmdList, ShaderRHI, BuiltinSamplersUBParameter, GBuiltinSamplersUniformBuffer.GetUniformBufferRHI());
+		}
+		else
+		{
+			checkSlow(!LowResDepthTexture.IsBound());
+		}
 	}
 
 	static const TCHAR* GetSourceFilename()
@@ -98,7 +140,7 @@ public:
 #define VARIATION1(A) typedef FPostProcessBokehDOFRecombinePS<A> FPostProcessBokehDOFRecombinePS##A; \
 	IMPLEMENT_SHADER_TYPE2(FPostProcessBokehDOFRecombinePS##A, SF_Pixel);
 
-VARIATION1(1)			VARIATION1(2)			VARIATION1(3)
+VARIATION1(1)			VARIATION1(2)			VARIATION1(3)			VARIATION1(4)			VARIATION1(5)
 #undef VARIATION1
 
 
@@ -139,6 +181,20 @@ void FRCPassPostProcessBokehDOFRecombine::Process(FRenderingCompositePassContext
 		check(GetInput(ePId_Input2)->GetPass());
 	}
 
+	FSceneRenderTargets& SceneContext = FSceneRenderTargets::Get(Context.RHICmdList);
+	FIntPoint OutScaledSize;
+	float OutScale;
+	SceneContext.GetSeparateTranslucencyDimensions(OutScaledSize, OutScale);
+
+	const bool bUseNearestDepthNeighborUpsample = 
+		CVarSeparateTranslucencyUpsampleMode.GetValueOnRenderThread() != 0
+		&& FMath::Abs(OutScale - .5f) < .001f;
+
+	if (Method != 1 && bUseNearestDepthNeighborUpsample)
+	{
+		Method += 2;
+	}
+
 	const FSceneView& View = Context.View;
 
 	SCOPED_DRAW_EVENTF(Context.RHICmdList, BokehDOFRecombine, TEXT("BokehDOFRecombine#%d %dx%d"), Method, View.ViewRect.Width(), View.ViewRect.Height());
@@ -174,6 +230,8 @@ void FRCPassPostProcessBokehDOFRecombine::Process(FRenderingCompositePassContext
 		case 1: SetShader<1>(Context); break;
 		case 2: SetShader<2>(Context); break;
 		case 3: SetShader<3>(Context); break;
+		case 4: SetShader<4>(Context); break;
+		case 5: SetShader<5>(Context); break;
 		default:
 			check(0);
 	}
diff --git a/Engine/Source/Runtime/Renderer/Private/PostProcess/PostProcessing.cpp b/Engine/Source/Runtime/Renderer/Private/PostProcess/PostProcessing.cpp
--- a/Engine/Source/Runtime/Renderer/Private/PostProcess/PostProcessing.cpp
+++ b/Engine/Source/Runtime/Renderer/Private/PostProcess/PostProcessing.cpp
@@ -1722,8 +1722,8 @@ void FPostProcessing::Process(FRHICommandListImmediate& RHICmdList, FViewInfo& V
 				Context.FinalOutput = FRenderingCompositeOutputRef(NodeRecombined);
 			}
 
-			const EDebugViewShaderMode DebugViewShaderMode = View.Family->GetDebugViewShaderMode();
-			if (DebugViewShaderMode == DVSM_None)
+			// Shader complexity does not actually output a color
+			if (!View.Family->EngineShowFlags.ShaderComplexity)
 			{
 				AddGammaOnlyTonemapper(Context);
 			}
@@ -1943,6 +1943,9 @@ void FPostProcessing::Process(FRHICommandListImmediate& RHICmdList, FViewInfo& V
 	}
 
 	GRenderTargetPool.AddPhaseEvent(TEXT("AfterPostprocessing"));
+
+	// End of frame, we don't need it anymore
+	FSceneRenderTargets::Get(RHICmdList).FreeSeparateTranslucencyDepth();
 }
 
 
diff --git a/Engine/Source/Runtime/Renderer/Private/PostProcess/SceneRenderTargets.cpp b/Engine/Source/Runtime/Renderer/Private/PostProcess/SceneRenderTargets.cpp
--- a/Engine/Source/Runtime/Renderer/Private/PostProcess/SceneRenderTargets.cpp
+++ b/Engine/Source/Runtime/Renderer/Private/PostProcess/SceneRenderTargets.cpp
@@ -222,6 +222,8 @@ FSceneRenderTargets::FSceneRenderTargets(const FViewInfo& View, const FSceneRend
 	, bSeparateTranslucencyPass(SnapshotSource.bSeparateTranslucencyPass)
 	, GBufferResourcesUniformBuffer(SnapshotSource.GBufferResourcesUniformBuffer)
 	, BufferSize(SnapshotSource.BufferSize)
+	, SeparateTranslucencyBufferSize(SnapshotSource.SeparateTranslucencyBufferSize)
+	, SeparateTranslucencyScale(SnapshotSource.SeparateTranslucencyScale)
 	, SmallColorDepthDownsampleFactor(SnapshotSource.SmallColorDepthDownsampleFactor)
 	, bLightAttenuationEnabled(SnapshotSource.bLightAttenuationEnabled)
 	, bUseDownsizedOcclusionQueries(SnapshotSource.bUseDownsizedOcclusionQueries)
@@ -1273,9 +1275,8 @@ bool FSceneRenderTargets::BeginRenderingSeparateTranslucency(FRHICommandList& RH
 	if(IsSeparateTranslucencyActive(View))
 	{
 		FIntPoint ScaledSize;
-		uint32 NumSamples = 1;
 		float Scale = 1.0f;
-		GetSeparateTranslucencyDimensionsAndSamplecount(ScaledSize, NumSamples, Scale);
+		GetSeparateTranslucencyDimensions(ScaledSize, Scale);
 
 		SCOPED_DRAW_EVENT(RHICmdList, BeginSeparateTranslucency);
 
@@ -1327,9 +1328,8 @@ void FSceneRenderTargets::FinishRenderingSeparateTranslucency(FRHICommandList& R
 		else
 		{
 			FIntPoint ScaledSize;
-			uint32 NumSamples = 1;
 			float Scale = 1.0f;
-			GetSeparateTranslucencyDimensionsAndSamplecount(ScaledSize, NumSamples, Scale);
+			GetSeparateTranslucencyDimensions(ScaledSize, Scale);
 			SeparateTranslucency = &GetSeparateTranslucency(RHICmdList, ScaledSize);
 			SeparateTranslucencyDepth = &GetSeparateTranslucencyDepth(RHICmdList, GetBufferSizeXY());
 		}
@@ -1489,6 +1489,24 @@ void FSceneRenderTargets::SetBufferSize(int32 InBufferSizeX, int32 InBufferSizeY
 	BufferSize.Y = InBufferSizeY;
 }
 
+void FSceneRenderTargets::SetSeparateTranslucencyBufferSize(bool bAnyViewWantsDownsampledSeparateTranslucency)
+{
+	static const auto CVar = IConsoleManager::Get().FindTConsoleVariableDataFloat(TEXT("r.SeparateTranslucencyScreenPercentage"));
+	const float CVarScale = FMath::Clamp(CVar->GetValueOnRenderThread() / 100.0f, 0.0f, 100.0f);
+	float EffectiveScale = CVarScale;
+
+	// 'r.SeparateTranslucencyScreenPercentage' CVar wins over automatic downsampling
+	if (FMath::Abs(CVarScale - 1.0f) < .001f && bAnyViewWantsDownsampledSeparateTranslucency)
+	{
+		EffectiveScale = .5f;
+	}
+
+	int32 ScaledX = GetBufferSizeXY().X * EffectiveScale;
+	int32 ScaledY = GetBufferSizeXY().Y * EffectiveScale;
+	SeparateTranslucencyBufferSize = FIntPoint(FMath::Max(ScaledX, 1), FMath::Max(ScaledY, 1));
+	SeparateTranslucencyScale = EffectiveScale;
+}
+
 void FSceneRenderTargets::AllocateForwardShadingPathRenderTargets(FRHICommandList& RHICmdList)
 {
 	// on ES2 we don't do on demand allocation of SceneColor yet (in non ES2 it's released in the Tonemapper Process())
@@ -2262,13 +2280,18 @@ void FSceneTextureShaderParameters::Set(
 
 		if(SceneDepthTextureParameter.IsBound() || SceneDepthTextureParameterSampler.IsBound())
 		{
-			static const auto CVar = IConsoleManager::Get().FindTConsoleVariableDataFloat(TEXT("r.SeparateTranslucencyScreenPercentage"));
 			const FTexture2DRHIRef* DepthTexture = SceneContext.GetActualDepthTexture();
-			if (SceneContext.IsSeparateTranslucencyPass()
-				&& CVar->GetValueOnRenderThread() < 100
-				&& SceneContext.IsSeparateTranslucencyDepthValid())
+
+			if (SceneContext.IsSeparateTranslucencyPass() && SceneContext.IsSeparateTranslucencyDepthValid())
 			{
-				DepthTexture = &SceneContext.GetSeparateTranslucencyDepthSurface();
+				FIntPoint OutScaledSize;
+				float OutScale;
+				SceneContext.GetSeparateTranslucencyDimensions(OutScaledSize, OutScale);
+
+				if (OutScale < 1.0f)
+				{
+					DepthTexture = &SceneContext.GetSeparateTranslucencyDepthSurface();
+				}
 			}
 
 			SetTextureParameter(
diff --git a/Engine/Source/Runtime/Renderer/Private/PostProcess/SceneRenderTargets.h b/Engine/Source/Runtime/Renderer/Private/PostProcess/SceneRenderTargets.h
--- a/Engine/Source/Runtime/Renderer/Private/PostProcess/SceneRenderTargets.h
+++ b/Engine/Source/Runtime/Renderer/Private/PostProcess/SceneRenderTargets.h
@@ -172,6 +172,8 @@ protected:
 		bVelocityPass(false),
 		bSeparateTranslucencyPass(false),
 		BufferSize(0, 0),
+		SeparateTranslucencyBufferSize(0, 0),
+		SeparateTranslucencyScale(1),
 		SmallColorDepthDownsampleFactor(2),
 		bLightAttenuationEnabled(true),
 		bUseDownsizedOcclusionQueries(true),
@@ -218,6 +220,8 @@ public:
 	 */
 	void SetBufferSize(int32 InBufferSizeX, int32 InBufferSizeY);
 
+	void SetSeparateTranslucencyBufferSize(bool bAnyViewWantsDownsampledSeparateTranslucency);
+
 	void BeginRenderingGBuffer(FRHICommandList& RHICmdList, ERenderTargetLoadAction ColorLoadAction, ERenderTargetLoadAction DepthLoadAction, bool bBindQuadOverdrawBuffers, const FLinearColor& ClearColor = FLinearColor(0, 0, 0, 1));
 	void FinishRenderingGBuffer(FRHICommandListImmediate& RHICmdList);
 
@@ -271,7 +275,10 @@ public:
 	{
 		SeparateTranslucencyRT.SafeRelease();
 		check(!SeparateTranslucencyRT);
+	}
 
+	void FreeSeparateTranslucencyDepth()
+	{
 		if (SeparateTranslucencyDepthRT.GetReference())
 		{
 			SeparateTranslucencyDepthRT.SafeRelease();
@@ -290,17 +297,13 @@ public:
 	void BeginRenderingLightAttenuation(FRHICommandList& RHICmdList, bool bClearToWhite = false);
 	void FinishRenderingLightAttenuation(FRHICommandList& RHICmdList);
 
-	void GetSeparateTranslucencyDimensionsAndSamplecount(FIntPoint &OutScaledSize, uint32 &OutNumSamples, float &OutScale)
+	void GetSeparateTranslucencyDimensions(FIntPoint& OutScaledSize, float& OutScale)
 	{
-		static const auto CVar = IConsoleManager::Get().FindTConsoleVariableDataFloat(TEXT("r.SeparateTranslucencyScreenPercentage"));
-		OutScale = FMath::Clamp(CVar->GetValueOnRenderThread() / 100.0f, 0.0f, 100.0f);
-		int32 ScaledX = GetBufferSizeXY().X * OutScale;
-		int32 ScaledY = GetBufferSizeXY().Y * OutScale;
-		OutScaledSize = FIntPoint(  FMath::Max(ScaledX, 1), FMath::Max(ScaledY, 1));
-		OutNumSamples = OutScale < 1.0f ? 4 : 1;
+		OutScaledSize = SeparateTranslucencyBufferSize;
+		OutScale = SeparateTranslucencyScale;
 	}
 
-	TRefCountPtr<IPooledRenderTarget>& GetSeparateTranslucency(FRHICommandList& RHICmdList, FIntPoint Size, uint32 NumSamples = 1)
+	TRefCountPtr<IPooledRenderTarget>& GetSeparateTranslucency(FRHICommandList& RHICmdList, FIntPoint Size)
 	{
 		if (!SeparateTranslucencyRT || SeparateTranslucencyRT->GetDesc().Extent != Size)
 		{
@@ -317,7 +320,7 @@ public:
 		return SeparateTranslucencyDepthRT != nullptr;
 	}
 
-	TRefCountPtr<IPooledRenderTarget>& GetSeparateTranslucencyDepth(FRHICommandList& RHICmdList, FIntPoint Size, uint32 NumSamples = 1)
+	TRefCountPtr<IPooledRenderTarget>& GetSeparateTranslucencyDepth(FRHICommandList& RHICmdList, FIntPoint Size)
 	{
 		if (!SeparateTranslucencyDepthRT || SeparateTranslucencyDepthRT->GetDesc().Extent != Size)
 		{
@@ -729,6 +732,8 @@ private:
 	FUniformBufferRHIRef GBufferResourcesUniformBuffer;
 	/** size of the back buffer, in editor this has to be >= than the biggest view port */
 	FIntPoint BufferSize;
+	FIntPoint SeparateTranslucencyBufferSize;
+	float SeparateTranslucencyScale;
 	/** e.g. 2 */
 	uint32 SmallColorDepthDownsampleFactor;
 	/** if true we use the light attenuation buffer otherwise the 1x1 white texture is used */
diff --git a/Engine/Source/Runtime/Renderer/Private/RendererScene.cpp b/Engine/Source/Runtime/Renderer/Private/RendererScene.cpp
--- a/Engine/Source/Runtime/Renderer/Private/RendererScene.cpp
+++ b/Engine/Source/Runtime/Renderer/Private/RendererScene.cpp
@@ -109,6 +109,17 @@ FSceneViewState::FSceneViewState()
 	ShadowOcclusionQueryMaps.AddZeroed(FOcclusionQueryHelpers::MaxBufferedOcclusionFrames);	
 
 	bValidEyeAdaptation = false;
+
+	LastAutoDownsampleChangeTime = 0;
+	SmoothedHalfResTranslucencyGPUDuration = 0;
+	SmoothedFullResTranslucencyGPUDuration = 0;
+	bShouldAutoDownsampleTranslucency = false;
+
+	PendingTranslucencyStartTimestamps.Empty(FOcclusionQueryHelpers::MaxBufferedOcclusionFrames + 1);
+	PendingTranslucencyStartTimestamps.AddZeroed(FOcclusionQueryHelpers::MaxBufferedOcclusionFrames + 1);
+
+	PendingTranslucencyEndTimestamps.Empty(FOcclusionQueryHelpers::MaxBufferedOcclusionFrames + 1);
+	PendingTranslucencyEndTimestamps.AddZeroed(FOcclusionQueryHelpers::MaxBufferedOcclusionFrames + 1);
 }
 
 void DestroyRenderResource(FRenderResource* RenderResource)
diff --git a/Engine/Source/Runtime/Renderer/Private/ScenePrivate.h b/Engine/Source/Runtime/Renderer/Private/ScenePrivate.h
--- a/Engine/Source/Runtime/Renderer/Private/ScenePrivate.h
+++ b/Engine/Source/Runtime/Renderer/Private/ScenePrivate.h
@@ -605,6 +605,18 @@ public:
 	FShaderResourceViewRHIRef IndirectShadowLightDirectionSRV;
 	FRWBuffer CapsuleTileIntersectionCountsBuffer;
 
+	/** Timestamp queries around separate translucency, used for auto-downsampling. */
+	TArray<FRenderQueryRHIRef, TInlineAllocator<FOcclusionQueryHelpers::MaxBufferedOcclusionFrames + 1> > PendingTranslucencyStartTimestamps;
+	TArray<FRenderQueryRHIRef, TInlineAllocator<FOcclusionQueryHelpers::MaxBufferedOcclusionFrames + 1> > PendingTranslucencyEndTimestamps;
+
+	/** This is float since it is derived off of UWorld::RealTimeSeconds, which is relative to BeginPlay time. */
+	float LastAutoDownsampleChangeTime;
+	float SmoothedHalfResTranslucencyGPUDuration;
+	float SmoothedFullResTranslucencyGPUDuration;
+
+	/** Current desired state of auto-downsampled separate translucency for this view. */
+	bool bShouldAutoDownsampleTranslucency;
+
 	// Is DOFHistoryRT set from DepthOfField?
 	bool bDOFHistory;
 	// Is DOFHistoryRT2 set from DepthOfField?
@@ -811,6 +823,16 @@ public:
 		IndirectShadowLightDirectionVertexBuffer.SafeRelease();
 		IndirectShadowLightDirectionSRV.SafeRelease();
 		CapsuleTileIntersectionCountsBuffer.Release();
+
+		for (int32 QueryIndex = 0; QueryIndex < PendingTranslucencyStartTimestamps.Num(); QueryIndex++)
+		{
+			PendingTranslucencyStartTimestamps[QueryIndex].SafeRelease();
+		}
+
+		for (int32 QueryIndex = 0; QueryIndex < PendingTranslucencyEndTimestamps.Num(); QueryIndex++)
+		{
+			PendingTranslucencyEndTimestamps[QueryIndex].SafeRelease();
+		}
 	}
 
 	// FSceneViewStateInterface
diff --git a/Engine/Source/Runtime/Renderer/Private/SceneVisibility.cpp b/Engine/Source/Runtime/Renderer/Private/SceneVisibility.cpp
--- a/Engine/Source/Runtime/Renderer/Private/SceneVisibility.cpp
+++ b/Engine/Source/Runtime/Renderer/Private/SceneVisibility.cpp
@@ -1225,6 +1225,7 @@ static int32 OcclusionCull(FRHICommandListImmediate& RHICmdList, const FScene* S
 	{
 		QUICK_SCOPE_CYCLE_COUNTER(STAT_LookupPrecomputedVisibility);
 
+		FViewElementPDI OcclusionPDI(&View, NULL);
 		uint8 PrecomputedVisibilityFlags = EOcclusionFlags::CanBeOccluded | EOcclusionFlags::HasPrecomputedVisibility;
 		for (FSceneSetBitIterator BitIt(View.PrimitiveVisibilityMap); BitIt; ++BitIt)
 		{
@@ -1236,6 +1237,12 @@ static int32 OcclusionCull(FRHICommandListImmediate& RHICmdList, const FScene* S
 					View.PrimitiveVisibilityMap.AccessCorrespondingBit(BitIt) = false;
 					INC_DWORD_STAT_BY(STAT_StaticallyOccludedPrimitives,1);
 					STAT(NumOccludedPrimitives++);
+
+					if (GVisualizeOccludedPrimitives)
+					{
+						const FBoxSphereBounds& Bounds = Scene->PrimitiveOcclusionBounds[BitIt.GetIndex()];
+						DrawWireBox(&OcclusionPDI, Bounds.GetBox(), FColor(100, 50, 50), SDPG_Foreground);
+					}
 				}
 			}
 		}
@@ -2873,6 +2880,8 @@ void FDeferredShadingSceneRenderer::InitViewsPossiblyAfterPrepass(FRHICommandLis
 		// Now that the indirect lighting cache is updated, we can update the primitive precomputed lighting buffers.
 		UpdatePrimitivePrecomputedLightingBuffers();
 	}
+
+	UpdateSeparateTranslucencyBufferSize(RHICmdList);
 }
 
 /*------------------------------------------------------------------------------
diff --git a/Engine/Source/Runtime/Renderer/Private/ShaderComplexityRendering.h b/Engine/Source/Runtime/Renderer/Private/ShaderComplexityRendering.h
--- a/Engine/Source/Runtime/Renderer/Private/ShaderComplexityRendering.h
+++ b/Engine/Source/Runtime/Renderer/Private/ShaderComplexityRendering.h
@@ -55,6 +55,7 @@ public:
 		const FVertexFactory* VertexFactory,
 		const FSceneView& View,
 		const FPrimitiveSceneProxy* Proxy,
+		int32 VisualizeLODIndex,
 		const FMeshBatchElement& BatchElement, 
 		const FMeshDrawingRenderState& DrawRenderState
 		) override {}
diff --git a/Engine/Source/Runtime/Renderer/Private/TexCoordScaleAnalysisRendering.cpp b/Engine/Source/Runtime/Renderer/Private/TexCoordScaleAnalysisRendering.cpp
--- a/Engine/Source/Runtime/Renderer/Private/TexCoordScaleAnalysisRendering.cpp
+++ b/Engine/Source/Runtime/Renderer/Private/TexCoordScaleAnalysisRendering.cpp
@@ -44,27 +44,11 @@ void FTexCoordScaleAnalysisPS::SetMesh(
 	const FVertexFactory* VertexFactory,
 	const FSceneView& View,
 	const FPrimitiveSceneProxy* Proxy,
+	int32 VisualizeLODIndex,
 	const FMeshBatchElement& BatchElement, 
 	const FMeshDrawingRenderState& DrawRenderState
 	)
 {
-	// This is taken from FStreamingHandlerTextureStatic::GetWantedMips
-
-	float CPUWantedMips = -1.f;
-	if (Proxy && Proxy->GetWorldTexelFactor() > 0)
-	{
-		FVector ViewToObject = Proxy->GetBounds().Origin - View.ViewMatrices.ViewOrigin;
-
-		float DistSqMinusRadiusSq = 0;
-		float ViewSize = View.ViewRect.Size().X;
-
-		DistSqMinusRadiusSq = FMath::Max(1.f, DistSqMinusRadiusSq);
-
-		float CoordSize = Proxy->GetWorldTexelFactor() * FMath::InvSqrt(DistSqMinusRadiusSq) * ViewSize;
-
-		CPUWantedMips = FMath::Clamp<float>(FMath::Log2(CoordSize), 0.f, (float)MaxStreamingAccuracyMips);
-	}
-
 	SetShaderValue(RHICmdList, FMeshMaterialShader::GetPixelShader(), TextureAnalysisIndexParameter, (int32)CVarStreamingAnalysisIndex.GetValueOnRenderThread());
 	SetShaderValue(RHICmdList, FMeshMaterialShader::GetPixelShader(), CPUScaleParameter, 1.f);
 
diff --git a/Engine/Source/Runtime/Renderer/Private/TexCoordScaleAnalysisRendering.h b/Engine/Source/Runtime/Renderer/Private/TexCoordScaleAnalysisRendering.h
--- a/Engine/Source/Runtime/Renderer/Private/TexCoordScaleAnalysisRendering.h
+++ b/Engine/Source/Runtime/Renderer/Private/TexCoordScaleAnalysisRendering.h
@@ -62,6 +62,7 @@ public:
 		const FVertexFactory* VertexFactory,
 		const FSceneView& View,
 		const FPrimitiveSceneProxy* Proxy,
+		int32 VisualizeLODIndex,
 		const FMeshBatchElement& BatchElement, 
 		const FMeshDrawingRenderState& DrawRenderState
 		) override;
diff --git a/Engine/Source/Runtime/Renderer/Private/TexelFactorAccuracyRendering.cpp b/Engine/Source/Runtime/Renderer/Private/TexelFactorAccuracyRendering.cpp
--- a/Engine/Source/Runtime/Renderer/Private/TexelFactorAccuracyRendering.cpp
+++ b/Engine/Source/Runtime/Renderer/Private/TexelFactorAccuracyRendering.cpp
@@ -10,6 +10,9 @@ TexelFactorAccuracyRendering.cpp: Contains definitions for rendering the viewmod
 
 IMPLEMENT_SHADER_TYPE(,FTexelFactorAccuracyPS,TEXT("TexelFactorAccuracyPixelShader"),TEXT("Main"),SF_Pixel);
 
+extern ENGINE_API TAutoConsoleVariable<int32> CVarStreamingUseNewMetrics;
+
+
 void FTexelFactorAccuracyPS::SetParameters(
 	FRHICommandList& RHICmdList, 
 	const FShader* OriginalVS, 
@@ -39,15 +42,21 @@ void FTexelFactorAccuracyPS::SetMesh(
 	const FVertexFactory* VertexFactory,
 	const FSceneView& View,
 	const FPrimitiveSceneProxy* Proxy,
+	int32 VisualizeLODIndex,
 	const FMeshBatchElement& BatchElement, 
 	const FMeshDrawingRenderState& DrawRenderState
 	)
 {
+	const bool bUseMetrics = CVarStreamingUseNewMetrics.GetValueOnRenderThread() != 0;
+
 	float CPUTexelFactor = -1.f;
-	if (Proxy && Proxy->GetWorldTexelFactor() > 0)
+#if !(UE_BUILD_SHIPPING || UE_BUILD_TEST)
+	FStreamingTexturePrimitiveInfo Info; 
+	if (Proxy && Proxy->GetStreamingTextureInfo(Info, bUseMetrics ? VisualizeLODIndex : INDEX_NONE, bUseMetrics ? BatchElement.VisualizeElementIndex : INDEX_NONE))
 	{
-		CPUTexelFactor = Proxy->GetWorldTexelFactor();
+		CPUTexelFactor = Info.TexelFactor;
 	}
+#endif
 	SetShaderValue(RHICmdList, FGlobalShader::GetPixelShader(), CPUTexelFactorParameter, CPUTexelFactor);
 }
 
diff --git a/Engine/Source/Runtime/Renderer/Private/TexelFactorAccuracyRendering.h b/Engine/Source/Runtime/Renderer/Private/TexelFactorAccuracyRendering.h
--- a/Engine/Source/Runtime/Renderer/Private/TexelFactorAccuracyRendering.h
+++ b/Engine/Source/Runtime/Renderer/Private/TexelFactorAccuracyRendering.h
@@ -55,6 +55,7 @@ public:
 		const FVertexFactory* VertexFactory,
 		const FSceneView& View,
 		const FPrimitiveSceneProxy* Proxy,
+		int32 VisualizeLODIndex,
 		const FMeshBatchElement& BatchElement, 
 		const FMeshDrawingRenderState& DrawRenderState
 		) override;
diff --git a/Engine/Source/Runtime/Renderer/Private/TranslucentRendering.cpp b/Engine/Source/Runtime/Renderer/Private/TranslucentRendering.cpp
--- a/Engine/Source/Runtime/Renderer/Private/TranslucentRendering.cpp
+++ b/Engine/Source/Runtime/Renderer/Private/TranslucentRendering.cpp
@@ -11,7 +11,194 @@
 #include "LightPropagationVolume.h"
 #include "SceneUtils.h"
 
-static float GScreenTextureUVScale;
+DECLARE_CYCLE_STAT(TEXT("TranslucencyTimestampQueryFence Wait"), STAT_TranslucencyTimestampQueryFence_Wait, STATGROUP_SceneRendering);
+DECLARE_CYCLE_STAT(TEXT("TranslucencyTimestampQuery Wait"), STAT_TranslucencyTimestampQuery_Wait, STATGROUP_SceneRendering);
+
+static TAutoConsoleVariable<int32> CVarSeparateTranslucencyAutoDownsample(
+	TEXT("r.SeparateTranslucencyAutoDownsample"),
+	0,
+	TEXT("Whether to automatically downsample separate translucency based on last frame's GPU time.\n")
+	TEXT("Automatic downsampling is only used when r.SeparateTranslucencyScreenPercentage is 100"),
+	ECVF_Scalability | ECVF_Default);
+
+static TAutoConsoleVariable<float> CVarSeparateTranslucencyDurationDownsampleThreshold(
+	TEXT("r.SeparateTranslucencyDurationDownsampleThreshold"),
+	1.5f,
+	TEXT("When smoothed full-res translucency GPU duration is larger than this value (ms), the entire pass will be downsampled by a factor of 2 in each dimension."),
+	ECVF_Scalability | ECVF_Default);
+
+static TAutoConsoleVariable<float> CVarSeparateTranslucencyDurationUpsampleThreshold(
+	TEXT("r.SeparateTranslucencyDurationUpsampleThreshold"),
+	.5f,
+	TEXT("When smoothed half-res translucency GPU duration is smaller than this value (ms), the entire pass will be restored to full resolution.\n")
+	TEXT("This should be around 1/4 of r.SeparateTranslucencyDurationDownsampleThreshold to avoid toggling downsampled state constantly."),
+	ECVF_Scalability | ECVF_Default);
+
+static TAutoConsoleVariable<float> CVarSeparateTranslucencyMinDownsampleChangeTime(
+	TEXT("r.SeparateTranslucencyMinDownsampleChangeTime"),
+	1.0f,
+	TEXT("Minimum time in seconds between changes to automatic downsampling state, used to prevent rapid swapping between half and full res."),
+	ECVF_Scalability | ECVF_Default);
+
+void FDeferredShadingSceneRenderer::UpdateSeparateTranslucencyBufferSize(FRHICommandListImmediate& RHICmdList)
+{
+	bool bAnyViewWantsDownsampledSeparateTranslucency = false;
+
+	if (CVarSeparateTranslucencyAutoDownsample.GetValueOnRenderThread() != 0)
+	{
+		for (int32 ViewIndex = 0; ViewIndex < Views.Num(); ViewIndex++)
+		{
+			const FViewInfo& View = Views[ViewIndex];
+			FSceneViewState* ViewState = View.ViewState;
+
+			// Buffer one extra frame for translucency timestamps because we don't care about latency (times are smoothed anyway)
+			// And we want to make sure we never block on the RHI thread or GPU
+			const int32 NumTimestampBufferedFrames = FOcclusionQueryHelpers::GetNumBufferedFrames() + 1;
+			const uint32 QueryIndex = FOcclusionQueryHelpers::GetQueryLookupIndex(View.ViewState->OcclusionFrameCounter, NumTimestampBufferedFrames);
+
+			if (ViewState
+				&& GSupportsTimestampRenderQueries
+				&& ViewState->PendingTranslucencyStartTimestamps[QueryIndex] 
+				&& ViewState->PendingTranslucencyEndTimestamps[QueryIndex])
+			{
+				if (GRHIThread)
+				{
+					// Block until the RHI thread has processed the previous query commands, if necessary
+					// Stat disabled since we buffer 2 frames minimum, it won't actually block
+					//SCOPE_CYCLE_COUNTER(STAT_TranslucencyTimestampQueryFence_Wait);
+					int32 BlockFrame = NumTimestampBufferedFrames - 1;
+					FRHICommandListExecutor::WaitOnRHIThreadFence(TranslucencyTimestampQuerySubmittedFence[BlockFrame]);
+					TranslucencyTimestampQuerySubmittedFence[BlockFrame] = nullptr;
+				}
+
+				uint64 StartMicroseconds;
+				uint64 EndMicroseconds;
+				bool bStartSuccess;
+				bool bEndSuccess;
+
+				{
+					// Block on the GPU until we have the timestamp query results, if necessary
+					// Stat disabled since we buffer 2 frames minimum, it won't actually block
+					//SCOPE_CYCLE_COUNTER(STAT_TranslucencyTimestampQuery_Wait);
+					bStartSuccess = RHICmdList.GetRenderQueryResult(ViewState->PendingTranslucencyStartTimestamps[QueryIndex] , StartMicroseconds, true);
+					bEndSuccess = RHICmdList.GetRenderQueryResult(ViewState->PendingTranslucencyEndTimestamps[QueryIndex], EndMicroseconds, true);
+				}
+
+				if (bStartSuccess && bEndSuccess)
+				{
+					const float LastFrameTranslucencyDurationMS = (EndMicroseconds - StartMicroseconds) / 1000.0f;
+					const bool bOriginalShouldAutoDownsampleTranslucency = ViewState->bShouldAutoDownsampleTranslucency;
+
+					//UE_LOG(LogRenderer, Log, TEXT("%u %.1fms"), bOriginalShouldAutoDownsampleTranslucency, LastFrameTranslucencyDurationMS);
+
+					if (ViewState->bShouldAutoDownsampleTranslucency)
+					{
+						ViewState->SmoothedFullResTranslucencyGPUDuration = 0;
+						const float LerpAlpha = ViewState->SmoothedHalfResTranslucencyGPUDuration == 0 ? 1.0f : .1f;
+						ViewState->SmoothedHalfResTranslucencyGPUDuration = FMath::Lerp(ViewState->SmoothedHalfResTranslucencyGPUDuration, LastFrameTranslucencyDurationMS, LerpAlpha);
+
+						// Don't re-asses switching for some time after the last switch
+						if (View.Family->CurrentRealTime - ViewState->LastAutoDownsampleChangeTime > CVarSeparateTranslucencyMinDownsampleChangeTime.GetValueOnRenderThread())
+						{
+							// Downsample if the smoothed time is larger than the threshold
+							ViewState->bShouldAutoDownsampleTranslucency = ViewState->SmoothedHalfResTranslucencyGPUDuration > CVarSeparateTranslucencyDurationUpsampleThreshold.GetValueOnRenderThread();
+
+							if (!ViewState->bShouldAutoDownsampleTranslucency)
+							{
+								// Do 'log LogRenderer verbose' to get these
+								UE_LOG(LogRenderer, Verbose, TEXT("Upsample: %.1fms < %.1fms"), ViewState->SmoothedHalfResTranslucencyGPUDuration, CVarSeparateTranslucencyDurationUpsampleThreshold.GetValueOnRenderThread());
+							}
+						}
+					}
+					else
+					{
+						ViewState->SmoothedHalfResTranslucencyGPUDuration = 0;
+						const float LerpAlpha = ViewState->SmoothedFullResTranslucencyGPUDuration == 0 ? 1.0f : .1f;
+						ViewState->SmoothedFullResTranslucencyGPUDuration = FMath::Lerp(ViewState->SmoothedFullResTranslucencyGPUDuration, LastFrameTranslucencyDurationMS, LerpAlpha);
+
+						if (View.Family->CurrentRealTime - ViewState->LastAutoDownsampleChangeTime > CVarSeparateTranslucencyMinDownsampleChangeTime.GetValueOnRenderThread())
+						{
+							ViewState->bShouldAutoDownsampleTranslucency = ViewState->SmoothedFullResTranslucencyGPUDuration > CVarSeparateTranslucencyDurationDownsampleThreshold.GetValueOnRenderThread();
+					
+							if (ViewState->bShouldAutoDownsampleTranslucency)
+							{
+								UE_LOG(LogRenderer, Verbose, TEXT("Downsample: %.1fms > %.1fms"), ViewState->SmoothedFullResTranslucencyGPUDuration, CVarSeparateTranslucencyDurationDownsampleThreshold.GetValueOnRenderThread());
+							}
+						}
+					}
+
+					if (bOriginalShouldAutoDownsampleTranslucency != ViewState->bShouldAutoDownsampleTranslucency)
+					{
+						ViewState->LastAutoDownsampleChangeTime = View.Family->CurrentRealTime;
+					}
+				}
+
+				bAnyViewWantsDownsampledSeparateTranslucency = bAnyViewWantsDownsampledSeparateTranslucency || ViewState->bShouldAutoDownsampleTranslucency;
+			}
+		}
+	}
+
+	FSceneRenderTargets& SceneContext = FSceneRenderTargets::Get(RHICmdList);
+	SceneContext.SetSeparateTranslucencyBufferSize(bAnyViewWantsDownsampledSeparateTranslucency);
+}
+
+void FDeferredShadingSceneRenderer::BeginTimingSeparateTranslucencyPass(FRHICommandListImmediate& RHICmdList, const FViewInfo& View)
+{
+	if (View.ViewState 
+		&& GSupportsTimestampRenderQueries
+		&& CVarSeparateTranslucencyAutoDownsample.GetValueOnRenderThread() != 0)
+	{
+		// Buffer one extra frame for translucency timestamps because we don't care about latency (times are smoothed anyway)
+		// And we want to make sure we never block on the RHI thread or GPU
+		const int32 NumTimestampBufferedFrames = FOcclusionQueryHelpers::GetNumBufferedFrames() + 1;
+		const uint32 QueryIndex = FOcclusionQueryHelpers::GetQueryIssueIndex(View.ViewState->OcclusionFrameCounter, NumTimestampBufferedFrames);
+
+		if (!View.ViewState->PendingTranslucencyStartTimestamps[QueryIndex])
+		{
+			View.ViewState->PendingTranslucencyStartTimestamps[QueryIndex] = RHICmdList.CreateRenderQuery(RQT_AbsoluteTime);
+		}
+						
+		// Insert a timestamp query at the beginning of separate translucency rendering
+		RHICmdList.EndRenderQuery(View.ViewState->PendingTranslucencyStartTimestamps[QueryIndex]);
+	}
+}
+
+void FDeferredShadingSceneRenderer::EndTimingSeparateTranslucencyPass(FRHICommandListImmediate& RHICmdList, const FViewInfo& View)
+{
+	if (View.ViewState 
+		&& GSupportsTimestampRenderQueries
+		&& CVarSeparateTranslucencyAutoDownsample.GetValueOnRenderThread() != 0)
+	{
+		// Buffer one extra frame for translucency timestamps because we don't care about latency (times are smoothed anyway)
+		// And we want to make sure we never block on the RHI thread or GPU
+		const int32 NumTimestampBufferedFrames = FOcclusionQueryHelpers::GetNumBufferedFrames() + 1;
+		const uint32 QueryIndex = FOcclusionQueryHelpers::GetQueryIssueIndex(View.ViewState->OcclusionFrameCounter, NumTimestampBufferedFrames);
+
+		if (!View.ViewState->PendingTranslucencyEndTimestamps[QueryIndex])
+		{
+			View.ViewState->PendingTranslucencyEndTimestamps[QueryIndex] = RHICmdList.CreateRenderQuery(RQT_AbsoluteTime);
+		}
+
+		// Insert a timestamp query at the end of separate translucency rendering
+		RHICmdList.EndRenderQuery(View.ViewState->PendingTranslucencyEndTimestamps[QueryIndex]);
+
+		// Hint to the RHI to submit commands up to this point to the GPU if possible.  Can help avoid CPU stalls next frame waiting
+		// for these query results on some platforms.
+		RHICmdList.SubmitCommandsHint();
+
+		if (GRHIThread)
+		{
+			int32 NumFrames = NumTimestampBufferedFrames;
+			for (int32 Dest = 1; Dest < NumFrames; Dest++)
+			{
+				TranslucencyTimestampQuerySubmittedFence[Dest] = TranslucencyTimestampQuerySubmittedFence[Dest - 1];
+			}
+			// Start an RHI thread fence so we can be sure the RHI thread has processed the EndRenderQuery before we ask for results
+			TranslucencyTimestampQuerySubmittedFence[0] = RHICmdList.RHIThreadFence();
+			RHICmdList.ImmediateFlush(EImmediateFlushType::DispatchToRHIThread);
+		}
+	}
+}
 
 static void SetTranslucentRenderTargetAndState(FRHICommandList& RHICmdList, const FViewInfo& View, ETranslucencyPassType TranslucenyPassType, bool bFirstTimeThisFrame = false)
 {
@@ -452,6 +639,11 @@ bool FTranslucencyDrawingPolicyFactory::DrawMesh(
 				RHICmdList.SetDepthStencilState(TStaticDepthStencilState<false,CF_Always>::GetRHI());
 			}
 
+			FSceneRenderTargets& SceneContext = FSceneRenderTargets::Get(RHICmdList);
+			FIntPoint OutScaledSize;
+			float OutScale;
+			SceneContext.GetSeparateTranslucencyDimensions(OutScaledSize, OutScale);
+
 			ProcessBasePassMesh(
 				RHICmdList, 
 				FProcessBasePassMeshParameters(
@@ -471,7 +663,7 @@ bool FTranslucencyDrawingPolicyFactory::DrawMesh(
 					HitProxyId,
 					DrawingContext.TranslucentSelfShadow,
 					PrimitiveSceneProxy && PrimitiveSceneProxy->CastsVolumetricTranslucentShadow(),
-					bCurrentlyRenderingSeparateTranslucency ? GScreenTextureUVScale : 1.0f
+					bCurrentlyRenderingSeparateTranslucency ? OutScale : 1.0f
 				)
 			);
 
@@ -807,35 +999,29 @@ void FTranslucentPrimSet::AppendScenePrimitives(FSortedPrim* Normal, int32 NumNo
 
 void FTranslucentPrimSet::PlaceScenePrimitive(FPrimitiveSceneInfo* PrimitiveSceneInfo, const FViewInfo& ViewInfo, bool bUseNormalTranslucency, bool bUseSeparateTranslucency, void *NormalPlace, int32& NormalNum, void* SeparatePlace, int32& SeparateNum)
 {	
-	int32 Value = FSceneRenderTargets::CVarSetSeperateTranslucencyEnabled.GetValueOnRenderThread();
-
 	const float SortKey = CalculateTranslucentSortKey(PrimitiveSceneInfo, ViewInfo);
-
 	const auto FeatureLevel = ViewInfo.GetFeatureLevel();
+	int32 CVarEnabled = FSceneRenderTargets::CVarSetSeperateTranslucencyEnabled.GetValueOnRenderThread();
 
-	if (bUseSeparateTranslucency
-		&& ViewInfo.Family->EngineShowFlags.SeparateTranslucency
+	bool bCanBeSeparate = CVarEnabled
+		&& FeatureLevel >= ERHIFeatureLevel::SM4
 		&& (ViewInfo.Family->EngineShowFlags.PostProcessing || ViewInfo.Family->EngineShowFlags.ShaderComplexity)
-		&& FeatureLevel >= ERHIFeatureLevel::SM4)
+		&& ViewInfo.Family->EngineShowFlags.SeparateTranslucency;
+
+	// add to list of sepaate translucency prims 
+	if (bUseSeparateTranslucency
+		&& bCanBeSeparate
+		)
 	{
-		// add to list of translucent prims that use scene color
 		new (SeparatePlace) FSortedPrim(PrimitiveSceneInfo,SortKey,PrimitiveSceneInfo->Proxy->GetTranslucencySortPriority());
 		SeparateNum++;
 	}
 
-	if (bUseNormalTranslucency 
-
-		// Force separate translucency to be rendered normally if separate translucency is disabled via cvar
-		|| Value == 0
-
-		// Force separate translucency to be rendered normally if the feature level does not support separate translucency
-		|| (bUseSeparateTranslucency && FeatureLevel < ERHIFeatureLevel::SM4)
-
-		// Force separate translucency to be rendered normally if separate translucency is disabled in the view
-		|| (!ViewInfo.Family->EngineShowFlags.SeparateTranslucency))
-
+	// add to list of translucent prims
+	else if (bUseNormalTranslucency
+		|| !bCanBeSeparate
+		)
 	{
-		// add to list of translucent prims
 		new (NormalPlace) FSortedPrim(PrimitiveSceneInfo,SortKey,PrimitiveSceneInfo->Proxy->GetTranslucencySortPriority());
 		NormalNum++;
 	}
@@ -960,7 +1146,6 @@ void FDeferredShadingSceneRenderer::RenderTranslucencyParallel(FRHICommandListIm
 	SceneContext.AllocLightAttenuation(RHICmdList); // materials will attempt to get this texture before the deferred command to set it up executes
 	check(IsInRenderingThread());
 
-
 	GParallelTranslucencyContext.TranslucentSelfShadow = nullptr;
 	GParallelTranslucencyContext.TranslucenyPassType = TPT_NonSeparateTransluceny;
 	GParallelTranslucencyContext.bSceneColorCopyIsUpToDate = false;
@@ -977,14 +1162,13 @@ void FDeferredShadingSceneRenderer::RenderTranslucencyParallel(FRHICommandListIm
 				QUICK_SCOPE_CYCLE_COUNTER(RenderTranslucencyParallel_Downsample);
 				// we need to allocate this now so it ends up in the snapshot
 				FIntPoint ScaledSize;
-				uint32 NumSamples = 1;
 				float Scale = 1.0f;
-				SceneContext.GetSeparateTranslucencyDimensionsAndSamplecount(ScaledSize, NumSamples, Scale);
+				SceneContext.GetSeparateTranslucencyDimensions(ScaledSize, Scale);
 
 				if (Scale<1.0f)
 				{
 					SceneContext.GetSeparateTranslucencyDepth(RHICmdList, SceneContext.GetBufferSizeXY());
-					DownsampleDepthSurface(RHICmdList, SceneContext.GetSeparateTranslucencyDepthSurface(), View, Scale, 1.0f);
+					DownsampleDepthSurface(RHICmdList, SceneContext.GetSeparateTranslucencyDepthSurface(), View, Scale, false);
 				}
 			}
 			FTranslucencyPassParallelCommandListSet ParallelCommandListSet(View, RHICmdList, 
@@ -1048,6 +1232,8 @@ void FDeferredShadingSceneRenderer::RenderTranslucencyParallel(FRHICommandListIm
 		}
 #endif
 		{
+			BeginTimingSeparateTranslucencyPass(RHICmdList, View);
+
 			{
 				// always call BeginRenderingSeparateTranslucency() even if there are no primitives to we keep the RT allocated
 				FTranslucencyPassParallelCommandListSet ParallelCommandListSet(View, 
@@ -1060,6 +1246,7 @@ void FDeferredShadingSceneRenderer::RenderTranslucencyParallel(FRHICommandListIm
 				if (View.TranslucentPrimSet.NumSeparateTranslucencyPrims() > 0)
 				{
 					QUICK_SCOPE_CYCLE_COUNTER(RenderTranslucencyParallel_Start_FDrawSortedTransAnyThreadTask_SeparateTransluceny);
+
 					int32 NumPrims = View.TranslucentPrimSet.NumSeparateTranslucencyPrims();
 					int32 EffectiveThreads = FMath::Min<int32>(FMath::DivideAndRoundUp(NumPrims, ParallelCommandListSet.MinDrawsPerCommandList), ParallelCommandListSet.Width);
 
@@ -1089,6 +1276,7 @@ void FDeferredShadingSceneRenderer::RenderTranslucencyParallel(FRHICommandListIm
 				}
 			}
 			SceneContext.FinishRenderingSeparateTranslucency(RHICmdList, View);
+			EndTimingSeparateTranslucencyPass(RHICmdList, View);
 		}
 	}
 }
@@ -1116,9 +1304,6 @@ void FDeferredShadingSceneRenderer::RenderTranslucency(FRHICommandListImmediate&
 {
 	if (ShouldRenderTranslucency())
 	{
-		static const auto CVar = IConsoleManager::Get().FindTConsoleVariableDataFloat(TEXT("r.SeparateTranslucencyScreenPercentage"));
-		GScreenTextureUVScale = 100.0f / CVar->GetValueOnRenderThread();
-
 		SCOPED_DRAW_EVENT(RHICmdList, Translucency);
 
 		if (GRHICommandList.UseParallelAlgorithms() && CVarParallelTranslucency.GetValueOnRenderThread())
@@ -1161,14 +1346,16 @@ void FDeferredShadingSceneRenderer::RenderTranslucency(FRHICommandListImmediate&
 				{
 					// always call BeginRenderingSeparateTranslucency() even if there are no primitives to we keep the RT allocated
 					FIntPoint ScaledSize;
-					uint32 NumSamples = 1;
 					float Scale = 1.0f;
-					SceneContext.GetSeparateTranslucencyDimensionsAndSamplecount(ScaledSize, NumSamples, Scale);
+					SceneContext.GetSeparateTranslucencyDimensions(ScaledSize, Scale);
 					if (Scale < 1.0f)
 					{
 						SceneContext.GetSeparateTranslucencyDepth(RHICmdList, SceneContext.GetBufferSizeXY());
-						DownsampleDepthSurface(RHICmdList, SceneContext.GetSeparateTranslucencyDepthSurface(), View, Scale, 1.0f);
+						DownsampleDepthSurface(RHICmdList, SceneContext.GetSeparateTranslucencyDepthSurface(), View, Scale, false);
 					}
+
+					BeginTimingSeparateTranslucencyPass(RHICmdList, View);
+
 					bool bFirstTimeThisFrame = (ViewIndex == 0);
 					bool bSetupTranslucency = SceneContext.BeginRenderingSeparateTranslucency(RHICmdList, View, bFirstTimeThisFrame);
 
@@ -1188,7 +1375,9 @@ void FDeferredShadingSceneRenderer::RenderTranslucency(FRHICommandListImmediate&
 						DrawAllTranslucencyPasses(RHICmdList, View, TPT_SeparateTransluceny);
 					}
 
-					SceneContext.FinishRenderingSeparateTranslucency(RHICmdList, View);
+					SceneContext.FinishRenderingSeparateTranslucency(RHICmdList, View);			
+
+					EndTimingSeparateTranslucencyPass(RHICmdList, View);
 				}
 			}
 		}
diff --git a/Engine/Source/Runtime/Renderer/Private/WantedMipsAccuracyRendering.cpp b/Engine/Source/Runtime/Renderer/Private/WantedMipsAccuracyRendering.cpp
--- a/Engine/Source/Runtime/Renderer/Private/WantedMipsAccuracyRendering.cpp
+++ b/Engine/Source/Runtime/Renderer/Private/WantedMipsAccuracyRendering.cpp
@@ -8,7 +8,7 @@ WantedMipsAccuracyRendering.cpp: Contains definitions for rendering the viewmode
 #include "ScenePrivate.h"
 #include "WantedMipsAccuracyRendering.h"
 
-extern ENGINE_API TAutoConsoleVariable<int32> CVarStreamingUseAABB;
+extern ENGINE_API TAutoConsoleVariable<int32> CVarStreamingUseNewMetrics;
 
 IMPLEMENT_SHADER_TYPE(,FWantedMipsAccuracyPS,TEXT("WantedMipsAccuracyPixelShader"),TEXT("Main"),SF_Pixel);
 
@@ -38,40 +38,43 @@ void FWantedMipsAccuracyPS::SetMesh(
 	const FVertexFactory* VertexFactory,
 	const FSceneView& View,
 	const FPrimitiveSceneProxy* Proxy,
+	int32 VisualizeLODIndex,
 	const FMeshBatchElement& BatchElement, 
 	const FMeshDrawingRenderState& DrawRenderState
 	)
 {
-	// This is taken from FStreamingHandlerTextureStatic::GetWantedMips
+	const bool bUseNewMetrics= CVarStreamingUseNewMetrics.GetValueOnRenderThread() != 0;
 
 	float CPUWantedMips = -1.f;
-	if (Proxy && Proxy->GetWorldTexelFactor() > 0)
+#if !(UE_BUILD_SHIPPING || UE_BUILD_TEST)
+	FStreamingTexturePrimitiveInfo Info;
+	if (Proxy && Proxy->GetStreamingTextureInfo(Info, bUseNewMetrics ? VisualizeLODIndex : INDEX_NONE, bUseNewMetrics ? BatchElement.VisualizeElementIndex : INDEX_NONE))
 	{
-		FVector ViewToObject = Proxy->GetBounds().Origin - View.ViewMatrices.ViewOrigin;
+		FVector ViewToObject = Info.Bounds.Origin - View.ViewMatrices.ViewOrigin;
 
 		float DistSqMinusRadiusSq = 0;
 		float ViewSize = View.ViewRect.Size().X;
 
-		if (CVarStreamingUseAABB.GetValueOnRenderThread())
+		if (bUseNewMetrics)
 		{
 			ViewToObject = ViewToObject.GetAbs();
-			FVector BoxViewToObject = ViewToObject.ComponentMin(Proxy->GetBounds().BoxExtent);
+			FVector BoxViewToObject = ViewToObject.ComponentMin(Info.Bounds.BoxExtent);
 
 			DistSqMinusRadiusSq = FVector::DistSquared(BoxViewToObject, ViewToObject);
 		}
 		else
 		{
 			float Distance = ViewToObject.Size();
-			DistSqMinusRadiusSq = FMath::Square(Distance) - FMath::Square(Proxy->GetBounds().SphereRadius);
+			DistSqMinusRadiusSq = FMath::Square(Distance) - FMath::Square(Info.Bounds.SphereRadius);
 		}
 
 		DistSqMinusRadiusSq = FMath::Max(1.f, DistSqMinusRadiusSq);
 
-		float CoordSize = Proxy->GetWorldTexelFactor() * FMath::InvSqrt(DistSqMinusRadiusSq) * ViewSize;
+		float CoordSize = Info.TexelFactor * FMath::InvSqrt(DistSqMinusRadiusSq) * ViewSize;
 
 		CPUWantedMips = FMath::Clamp<float>(FMath::Log2(CoordSize), 0.f, (float)MaxStreamingAccuracyMips);
 	}
-
+#endif
 	SetShaderValue(RHICmdList, FGlobalShader::GetPixelShader(), CPUWantedMipsParameter, FMath::FloorToFloat(CPUWantedMips));
 }
 
diff --git a/Engine/Source/Runtime/Renderer/Private/WantedMipsAccuracyRendering.h b/Engine/Source/Runtime/Renderer/Private/WantedMipsAccuracyRendering.h
--- a/Engine/Source/Runtime/Renderer/Private/WantedMipsAccuracyRendering.h
+++ b/Engine/Source/Runtime/Renderer/Private/WantedMipsAccuracyRendering.h
@@ -56,6 +56,7 @@ public:
 		const FVertexFactory* VertexFactory,
 		const FSceneView& View,
 		const FPrimitiveSceneProxy* Proxy,
+		int32 VisualizeLODIndex,
 		const FMeshBatchElement& BatchElement, 
 		const FMeshDrawingRenderState& DrawRenderState
 		) override;
diff --git a/Engine/Source/Runtime/Windows/D3D11RHI/Private/D3D11Device.cpp b/Engine/Source/Runtime/Windows/D3D11RHI/Private/D3D11Device.cpp
--- a/Engine/Source/Runtime/Windows/D3D11RHI/Private/D3D11Device.cpp
+++ b/Engine/Source/Runtime/Windows/D3D11RHI/Private/D3D11Device.cpp
@@ -187,6 +187,7 @@ FD3D11DynamicRHI::FD3D11DynamicRHI(IDXGIFactory1* InDXGIFactory1,D3D_FEATURE_LEV
 	GMaxShadowDepthBufferSizeX = 4096;
 	GMaxShadowDepthBufferSizeY = 4096;
 	GRHISupportsInstancing = GMaxRHIFeatureLevel > ERHIFeatureLevel::ES2;
+	GSupportsTimestampRenderQueries = true;
 
 	// Initialize the constant buffers.
 	InitConstantBuffers();
diff --git a/Engine/Source/Runtime/Windows/D3D12RHI/Private/D3D12Device.cpp b/Engine/Source/Runtime/Windows/D3D12RHI/Private/D3D12Device.cpp
--- a/Engine/Source/Runtime/Windows/D3D12RHI/Private/D3D12Device.cpp
+++ b/Engine/Source/Runtime/Windows/D3D12RHI/Private/D3D12Device.cpp
@@ -424,6 +424,7 @@ FD3D12DynamicRHI::FD3D12DynamicRHI(IDXGIFactory4* InDXGIFactory, FD3D12Adapter&
 	GMaxTextureDimensions = D3D12_REQ_TEXTURE2D_U_OR_V_DIMENSION;
 	GMaxCubeTextureDimensions = D3D12_REQ_TEXTURECUBE_DIMENSION;
 	GMaxTextureArrayLayers = D3D12_REQ_TEXTURE2D_ARRAY_AXIS_DIMENSION;
+	GSupportsTimestampRenderQueries = true;
 
 	GMaxShadowDepthBufferSizeX = 4096;
 	GMaxShadowDepthBufferSizeY = 4096;
