Fix for particle systems with both translucent and seperate translucent emitters.

#tests Fixes issue

diff --git a/Engine/Source/Runtime/Renderer/Private/TranslucentRendering.cpp b/Engine/Source/Runtime/Renderer/Private/TranslucentRendering.cpp
--- a/Engine/Source/Runtime/Renderer/Private/TranslucentRendering.cpp
+++ b/Engine/Source/Runtime/Renderer/Private/TranslucentRendering.cpp
@@ -1018,7 +1018,7 @@ void FTranslucentPrimSet::PlaceScenePrimitive(FPrimitiveSceneInfo* PrimitiveScen
 	}
 
 	// add to list of translucent prims
-	else if (bUseNormalTranslucency
+	if (bUseNormalTranslucency
 		|| !bCanBeSeparate
 		)
 	{
