Copying //UE4/Dev-VR to //UE4/Main

==========================
MAJOR FEATURES + CHANGES
==========================

Change 2797279 on 2015/12/09 by [REDACTED_PERSON][REDACTED_HANDLE]

	#psvr Modifying reprojection border type

Change 2797300 on 2015/12/09 by [REDACTED_PERSON][REDACTED_HANDLE]

	#psvr Moving frame tracking memory from onion to garlic, chaning frame mode read

Change 2818240 on 2016/01/06 by [REDACTED_PERSON][REDACTED_HANDLE]

	Resource icon for the Dualshock 4 movement controller plugin

Change 2819348 on 2016/01/07 by [REDACTED_PERSON][REDACTED_HANDLE]

	Adding support for an audio listener offset:  This adds support for adding an offset from the screen position for the audio listener (e.g. moving from the eyes to the ears).

Change 2881565 on 2016/02/25 by [REDACTED_PERSON][REDACTED_HANDLE]

	CountBits needs to handle the instanced stereo case.

diff --git a/Engine/Source/Runtime/Engine/Private/GameViewportClient.cpp b/Engine/Source/Runtime/Engine/Private/GameViewportClient.cpp
--- a/Engine/Source/Runtime/Engine/Private/GameViewportClient.cpp
+++ b/Engine/Source/Runtime/Engine/Private/GameViewportClient.cpp
@@ -976,6 +976,14 @@ void UGameViewportClient::Draw(FViewport* InViewport, FCanvas* SceneCanvas)
 								PlayerController->GetAudioListenerPosition(/*out*/ Location, /*out*/ ProjFront, /*out*/ ProjRight);
 
 								FTransform ListenerTransform(FRotationMatrix::MakeFromXY(ProjFront, ProjRight));
+
+								// Allow the HMD to adjust based on the head position of the player, as opposed to the view location
+								if (GEngine->HMDDevice.IsValid() && GEngine->HMDDevice->IsStereoEnabled())
+								{
+									const FVector Offset = GEngine->HMDDevice->GetAudioListenerOffset();
+									Location += ListenerTransform.TransformPositionNoScale(Offset);
+								}
+
 								ListenerTransform.SetTranslation(Location);
 								ListenerTransform.NormalizeRotation();
 
diff --git a/Engine/Source/Runtime/HeadMountedDisplay/Public/IHeadMountedDisplay.h b/Engine/Source/Runtime/HeadMountedDisplay/Public/IHeadMountedDisplay.h
--- a/Engine/Source/Runtime/HeadMountedDisplay/Public/IHeadMountedDisplay.h
+++ b/Engine/Source/Runtime/HeadMountedDisplay/Public/IHeadMountedDisplay.h
@@ -100,6 +100,11 @@ public:
 	 */
 	virtual void RebaseObjectOrientationAndPosition(FVector& Position, FQuat& Orientation) const = 0;
 
+	/**
+	 * Get the offset, in device space, of the reported device (screen / eye) position to the center of the head, if supported
+	 */
+	virtual FVector GetAudioListenerOffset() const { return FVector(0.f); }
+
 	/**
 	 * Get the ISceneViewExtension for this HMD, or none.
 	 */
diff --git a/Engine/Source/Runtime/Renderer/Private/StaticMeshDrawList.inl b/Engine/Source/Runtime/Renderer/Private/StaticMeshDrawList.inl
--- a/Engine/Source/Runtime/Renderer/Private/StaticMeshDrawList.inl
+++ b/Engine/Source/Runtime/Renderer/Private/StaticMeshDrawList.inl
@@ -579,9 +579,23 @@ void TStaticMeshDrawList<DrawingPolicyType>::DrawVisibleParallelInternal(
 							if (bIsVisible)
 							{
 								const FElement& Element = DrawingPolicyLink->Elements[ElementIndex];
-								int32 SubCount = Element.Mesh->Elements.Num();
+								const int32 SubCount = Element.Mesh->Elements.Num();
+
 								// Avoid the cache miss looking up batch visibility if there is only one element.
-								Count += (SubCount == 1) ? 1 : CountBits((*BatchVisibilityArray)[Element.Mesh->Id]);
+								if (SubCount == 1)
+								{
+									++Count;
+								}
+								else if (!bIsInstancedStereo)
+								{
+									Count += CountBits((*BatchVisibilityArray)[Element.Mesh->Id]);
+								}
+								else
+								{
+									const int32 LeftCount = CountBits((*StereoView->LeftViewBatchVisibilityArray)[Element.Mesh->Id]);
+									const int32 RightCount = CountBits((*StereoView->RightViewBatchVisibilityArray)[Element.Mesh->Id]);
+									Count += (LeftCount > RightCount) ? LeftCount : RightCount;
+								}
 							}
 						}
 						if (Count)
