Copying //UE4/Dev-Sequencer to //UE4/Main

diff --git a/Engine/Source/Runtime/MovieSceneTracks/Public/TrackInstances/MovieSceneSkeletalAnimationTrackInstance.h b/Engine/Source/Runtime/MovieSceneTracks/Public/TrackInstances/MovieSceneSkeletalAnimationTrackInstance.h
--- a/Engine/Source/Runtime/MovieSceneTracks/Public/TrackInstances/MovieSceneSkeletalAnimationTrackInstance.h
+++ b/Engine/Source/Runtime/MovieSceneTracks/Public/TrackInstances/MovieSceneSkeletalAnimationTrackInstance.h
@@ -20,11 +20,11 @@ public:
 	virtual ~FMovieSceneSkeletalAnimationTrackInstance();
 
 	/** IMovieSceneTrackInstance interface */
-	virtual void SaveState(const TArray<UObject*>& RuntimeObjects, IMovieScenePlayer& Player, FMovieSceneSequenceInstance& SequenceInstance) override {}
-	virtual void RestoreState(const TArray<UObject*>& RuntimeObjects, IMovieScenePlayer& Player, FMovieSceneSequenceInstance& SequenceInstance) override {}
-	virtual void Update(EMovieSceneUpdateData& UpdateData, const TArray<UObject*>& RuntimeObjects, IMovieScenePlayer& Player, FMovieSceneSequenceInstance& SequenceInstance) override;
-	virtual void RefreshInstance( const TArray<UObject*>& RuntimeObjects, IMovieScenePlayer& Player, FMovieSceneSequenceInstance& SequenceInstance ) override {}
-	virtual void ClearInstance( IMovieScenePlayer& Player, FMovieSceneSequenceInstance& SequenceInstance ) override {}
+	virtual void SaveState(const TArray<TWeakObjectPtr<UObject>>& RuntimeObjects, IMovieScenePlayer& Player, FMovieSceneSequenceInstance& SequenceInstance) override;
+	virtual void RestoreState(const TArray<TWeakObjectPtr<UObject>>& RuntimeObjects, IMovieScenePlayer& Player, FMovieSceneSequenceInstance& SequenceInstance) override;
+	virtual void Update(EMovieSceneUpdateData& UpdateData, const TArray<TWeakObjectPtr<UObject>>& RuntimeObjects, IMovieScenePlayer& Player, FMovieSceneSequenceInstance& SequenceInstance) override;
+	virtual void RefreshInstance(const TArray<TWeakObjectPtr<UObject>>& RuntimeObjects, IMovieScenePlayer& Player, FMovieSceneSequenceInstance& SequenceInstance ) override {}
+	virtual void ClearInstance(IMovieScenePlayer& Player, FMovieSceneSequenceInstance& SequenceInstance ) override {}
 
 private:
 
@@ -32,7 +32,7 @@ private:
 
 	bool CanPlayAnimation(USkeletalMeshComponent* SkeletalMeshComponent, class UAnimSequenceBase* AnimAssetBase = nullptr) const;
 
-	void SetAnimPosition(USkeletalMeshComponent* SkeletalMeshComponent, FName SlotName, int32 ChannelIndex, UAnimSequence* InAnimSequence, float InPosition, bool bFireNotifies, bool bLooping);
+	void SetAnimPosition(USkeletalMeshComponent* SkeletalMeshComponent, FName SlotName, int32 ChannelIndex, UAnimSequence* InAnimSequence, float InPosition, bool bLooping, bool bFireNotifies);
 
 	void FinishAnimControl(USkeletalMeshComponent* SkeletalMeshComponent);
 
@@ -40,13 +40,13 @@ private:
 
 	void PreviewFinishAnimControl(USkeletalMeshComponent* SkeletalMeshComponent);
 
-	void PreviewSetAnimPosition(USkeletalMeshComponent* SkeletalMeshComponent, FName SlotName, int32 ChannelIndex, UAnimSequence* InAnimSequence, float InPosition, bool bLooping, bool bFireNotifies, float DeltaTime);
+	void PreviewSetAnimPosition(USkeletalMeshComponent* SkeletalMeshComponent, FName SlotName, int32 ChannelIndex, UAnimSequence* InAnimSequence, float InPosition, bool bLooping, bool bFireNotifies, float DeltaTime, bool bPlaying);
 
 private:
 
 	/** Track that is being instanced */
 	UMovieSceneSkeletalAnimationTrack* AnimationTrack;
 
-	/** Keep track of the currently playhing montage */
+	/** Keep track of the currently playing montage */
 	TWeakObjectPtr<UAnimMontage> CurrentlyPlayingMontage;
 };
diff --git a/Engine/Source/Runtime/MovieSceneTracks/Public/TrackInstances/MovieSceneSlomoTrackInstance.h b/Engine/Source/Runtime/MovieSceneTracks/Public/TrackInstances/MovieSceneSlomoTrackInstance.h
--- a/Engine/Source/Runtime/MovieSceneTracks/Public/TrackInstances/MovieSceneSlomoTrackInstance.h
+++ b/Engine/Source/Runtime/MovieSceneTracks/Public/TrackInstances/MovieSceneSlomoTrackInstance.h
@@ -31,10 +31,10 @@ public:
 	// IMovieSceneTrackInstance interface
 
 	virtual void ClearInstance(IMovieScenePlayer& Player, FMovieSceneSequenceInstance& SequenceInstance) override { }
-	virtual void RefreshInstance(const TArray<UObject*>& RuntimeObjects, IMovieScenePlayer& Player, FMovieSceneSequenceInstance& SequenceInstance) override { }
-	virtual void RestoreState(const TArray<UObject*>& RuntimeObjects, IMovieScenePlayer& Player, FMovieSceneSequenceInstance& SequenceInstance) override;
-	virtual void SaveState(const TArray<UObject*>& RuntimeObjects, IMovieScenePlayer& Player, FMovieSceneSequenceInstance& SequenceInstance) override;
-	virtual void Update(EMovieSceneUpdateData& UpdateData, const TArray<UObject*>& RuntimeObjects, IMovieScenePlayer& Player, FMovieSceneSequenceInstance& SequenceInstance) override;
+	virtual void RefreshInstance(const TArray<TWeakObjectPtr<UObject>>& RuntimeObjects, IMovieScenePlayer& Player, FMovieSceneSequenceInstance& SequenceInstance) override { }
+	virtual void RestoreState(const TArray<TWeakObjectPtr<UObject>>& RuntimeObjects, IMovieScenePlayer& Player, FMovieSceneSequenceInstance& SequenceInstance) override;
+	virtual void SaveState(const TArray<TWeakObjectPtr<UObject>>& RuntimeObjects, IMovieScenePlayer& Player, FMovieSceneSequenceInstance& SequenceInstance) override;
+	virtual void Update(EMovieSceneUpdateData& UpdateData, const TArray<TWeakObjectPtr<UObject>>& RuntimeObjects, IMovieScenePlayer& Player, FMovieSceneSequenceInstance& SequenceInstance) override;
 
 protected:
 
diff --git a/Engine/Source/Runtime/MovieSceneTracks/Public/TrackInstances/MovieSceneSpawnTrackInstance.h b/Engine/Source/Runtime/MovieSceneTracks/Public/TrackInstances/MovieSceneSpawnTrackInstance.h
--- a/Engine/Source/Runtime/MovieSceneTracks/Public/TrackInstances/MovieSceneSpawnTrackInstance.h
+++ b/Engine/Source/Runtime/MovieSceneTracks/Public/TrackInstances/MovieSceneSpawnTrackInstance.h
@@ -15,13 +15,13 @@ public:
 	FMovieSceneSpawnTrackInstance(UMovieSceneSpawnTrack& InTrack);
 
 	/** IMovieSceneTrackInstance interface */
-	virtual void SaveState(const TArray<UObject*>& RuntimeObjects, IMovieScenePlayer& Player, FMovieSceneSequenceInstance& SequenceInstance) override {}
-	virtual void RestoreState(const TArray<UObject*>& RuntimeObjects, IMovieScenePlayer& Player, FMovieSceneSequenceInstance& SequenceInstance) override {}
-	virtual void RefreshInstance(const TArray<UObject*>& RuntimeObjects, class IMovieScenePlayer& Player, FMovieSceneSequenceInstance& SequenceInstance) override {}
+	virtual void SaveState(const TArray<TWeakObjectPtr<UObject>>& RuntimeObjects, IMovieScenePlayer& Player, FMovieSceneSequenceInstance& SequenceInstance) override {}
+	virtual void RestoreState(const TArray<TWeakObjectPtr<UObject>>& RuntimeObjects, IMovieScenePlayer& Player, FMovieSceneSequenceInstance& SequenceInstance) override {}
+	virtual void RefreshInstance(const TArray<TWeakObjectPtr<UObject>>& RuntimeObjects, class IMovieScenePlayer& Player, FMovieSceneSequenceInstance& SequenceInstance) override {}
 	virtual void ClearInstance(IMovieScenePlayer& Player, FMovieSceneSequenceInstance& SequenceInstance) override {}
 	virtual EMovieSceneUpdatePass HasUpdatePasses() override { return MSUP_PreUpdate; }
 
-	virtual void Update(EMovieSceneUpdateData& UpdateData, const TArray<UObject*>& RuntimeObjects, class IMovieScenePlayer& Player, FMovieSceneSequenceInstance& SequenceInstance) override;
+	virtual void Update(EMovieSceneUpdateData& UpdateData, const TArray<TWeakObjectPtr<UObject>>& RuntimeObjects, class IMovieScenePlayer& Player, FMovieSceneSequenceInstance& SequenceInstance) override;
 
 	virtual float EvalOrder() override { return -10000.f; }
 
diff --git a/Engine/Source/Runtime/MovieSceneTracks/Public/TrackInstances/MovieSceneSubTrackInstance.h b/Engine/Source/Runtime/MovieSceneTracks/Public/TrackInstances/MovieSceneSubTrackInstance.h
--- a/Engine/Source/Runtime/MovieSceneTracks/Public/TrackInstances/MovieSceneSubTrackInstance.h
+++ b/Engine/Source/Runtime/MovieSceneTracks/Public/TrackInstances/MovieSceneSubTrackInstance.h
@@ -26,17 +26,18 @@ public:
 	// IMovieSceneTrackInstance interface
 
 	virtual void ClearInstance(IMovieScenePlayer& Player, FMovieSceneSequenceInstance& SequenceInstance) override;
-	virtual void RefreshInstance(const TArray<UObject*>& RuntimeObjects, IMovieScenePlayer& Player, FMovieSceneSequenceInstance& SequenceInstance) override;
-	virtual void RestoreState (const TArray<UObject*>& RuntimeObjects, IMovieScenePlayer& Player, FMovieSceneSequenceInstance& SequenceInstance) override;
-	virtual void SaveState(const TArray<UObject*>& RuntimeObjects, IMovieScenePlayer& Player, FMovieSceneSequenceInstance& SequenceInstance) override;
-	virtual void Update(EMovieSceneUpdateData& UpdateData, const TArray<UObject*>& RuntimeObjects, IMovieScenePlayer& Player, FMovieSceneSequenceInstance& SequenceInstance) override;
-	virtual EMovieSceneUpdatePass HasUpdatePasses() override { return (EMovieSceneUpdatePass)(MSUP_PreUpdate | MSUP_Update); }
+	virtual void RefreshInstance(const TArray<TWeakObjectPtr<UObject>>& RuntimeObjects, IMovieScenePlayer& Player, FMovieSceneSequenceInstance& SequenceInstance) override;
+	virtual void RestoreState (const TArray<TWeakObjectPtr<UObject>>& RuntimeObjects, IMovieScenePlayer& Player, FMovieSceneSequenceInstance& SequenceInstance) override;
+	virtual void SaveState(const TArray<TWeakObjectPtr<UObject>>& RuntimeObjects, IMovieScenePlayer& Player, FMovieSceneSequenceInstance& SequenceInstance) override;
+	virtual void Update(EMovieSceneUpdateData& UpdateData, const TArray<TWeakObjectPtr<UObject>>& RuntimeObjects, IMovieScenePlayer& Player, FMovieSceneSequenceInstance& SequenceInstance) override;
+	virtual EMovieSceneUpdatePass HasUpdatePasses() override { return (EMovieSceneUpdatePass)(MSUP_All); }
 
 protected:
 
 	TArray<UMovieSceneSection*> GetAllTraversedSectionsWithPreroll( const TArray<UMovieSceneSection*>& Sections, float CurrentTime, float PreviousTime );
 	TArray<UMovieSceneSection*> GetTraversedSectionsWithPreroll( const TArray<UMovieSceneSection*>& Sections, float CurrentTime, float PreviousTime );
-
+	virtual bool ShouldEvaluateIfOverlapping(const TArray<UMovieSceneSection*>& Sections, UMovieSceneSection* Section) const;
+	
 protected:
 
 	/** Track that is being instanced */
diff --git a/Engine/Source/Runtime/MovieSceneTracks/Public/TrackInstances/MovieSceneVectorTrackInstance.h b/Engine/Source/Runtime/MovieSceneTracks/Public/TrackInstances/MovieSceneVectorTrackInstance.h
--- a/Engine/Source/Runtime/MovieSceneTracks/Public/TrackInstances/MovieSceneVectorTrackInstance.h
+++ b/Engine/Source/Runtime/MovieSceneTracks/Public/TrackInstances/MovieSceneVectorTrackInstance.h
@@ -20,11 +20,11 @@ public:
 	FMovieSceneVectorTrackInstance( UMovieSceneVectorTrack& InVectorTrack );
 
 	/** IMovieSceneTrackInstance interface */
-	virtual void SaveState (const TArray<UObject*>& RuntimeObjects, IMovieScenePlayer& Player, FMovieSceneSequenceInstance& SequenceInstance) override;
-	virtual void RestoreState (const TArray<UObject*>& RuntimeObjects, IMovieScenePlayer& Player, FMovieSceneSequenceInstance& SequenceInstance) override;
-	virtual void Update( EMovieSceneUpdateData& UpdateData, const TArray<UObject*>& RuntimeObjects, IMovieScenePlayer& Player, FMovieSceneSequenceInstance& SequenceInstance ) override;
-	virtual void RefreshInstance( const TArray<UObject*>& RuntimeObjects, IMovieScenePlayer& Player, FMovieSceneSequenceInstance& SequenceInstance ) override;
-	virtual void ClearInstance( IMovieScenePlayer& Player, FMovieSceneSequenceInstance& SequenceInstance ) override {}
+	virtual void SaveState(const TArray<TWeakObjectPtr<UObject>>& RuntimeObjects, IMovieScenePlayer& Player, FMovieSceneSequenceInstance& SequenceInstance) override;
+	virtual void RestoreState(const TArray<TWeakObjectPtr<UObject>>& RuntimeObjects, IMovieScenePlayer& Player, FMovieSceneSequenceInstance& SequenceInstance) override;
+	virtual void Update(EMovieSceneUpdateData& UpdateData, const TArray<TWeakObjectPtr<UObject>>& RuntimeObjects, IMovieScenePlayer& Player, FMovieSceneSequenceInstance& SequenceInstance) override;
+	virtual void RefreshInstance(const TArray<TWeakObjectPtr<UObject>>& RuntimeObjects, IMovieScenePlayer& Player, FMovieSceneSequenceInstance& SequenceInstance) override;
+	virtual void ClearInstance(IMovieScenePlayer& Player, FMovieSceneSequenceInstance& SequenceInstance) override {}
 
 private:
 
diff --git a/Engine/Source/Runtime/MovieSceneTracks/Public/TrackInstances/MovieSceneVisibilityTrackInstance.h b/Engine/Source/Runtime/MovieSceneTracks/Public/TrackInstances/MovieSceneVisibilityTrackInstance.h
--- a/Engine/Source/Runtime/MovieSceneTracks/Public/TrackInstances/MovieSceneVisibilityTrackInstance.h
+++ b/Engine/Source/Runtime/MovieSceneTracks/Public/TrackInstances/MovieSceneVisibilityTrackInstance.h
@@ -19,11 +19,11 @@ public:
 	FMovieSceneVisibilityTrackInstance( UMovieSceneVisibilityTrack& InVisibilityTrack );
 
 	/** IMovieSceneTrackInstance interface */
-	virtual void SaveState (const TArray<UObject*>& RuntimeObjects, IMovieScenePlayer& Player, FMovieSceneSequenceInstance& SequenceInstance) override;
-	virtual void RestoreState (const TArray<UObject*>& RuntimeObjects, IMovieScenePlayer& Player, FMovieSceneSequenceInstance& SequenceInstance) override;
-	virtual void Update( EMovieSceneUpdateData& UpdateData, const TArray<UObject*>& RuntimeObjects, class IMovieScenePlayer& Player, FMovieSceneSequenceInstance& SequenceInstance ) override;
-	virtual void RefreshInstance( const TArray<UObject*>& RuntimeObjects, class IMovieScenePlayer& Player, FMovieSceneSequenceInstance& SequenceInstance ) override;
-	virtual void ClearInstance( IMovieScenePlayer& Player, FMovieSceneSequenceInstance& SequenceInstance ) override {}
+	virtual void SaveState(const TArray<TWeakObjectPtr<UObject>>& RuntimeObjects, IMovieScenePlayer& Player, FMovieSceneSequenceInstance& SequenceInstance) override;
+	virtual void RestoreState(const TArray<TWeakObjectPtr<UObject>>& RuntimeObjects, IMovieScenePlayer& Player, FMovieSceneSequenceInstance& SequenceInstance) override;
+	virtual void Update(EMovieSceneUpdateData& UpdateData, const TArray<TWeakObjectPtr<UObject>>& RuntimeObjects, class IMovieScenePlayer& Player, FMovieSceneSequenceInstance& SequenceInstance) override;
+	virtual void RefreshInstance(const TArray<TWeakObjectPtr<UObject>>& RuntimeObjects, class IMovieScenePlayer& Player, FMovieSceneSequenceInstance& SequenceInstance) override;
+	virtual void ClearInstance(IMovieScenePlayer& Player, FMovieSceneSequenceInstance& SequenceInstance) override {}
 
 private:
 	/** Visibility track that is being instanced */
diff --git a/Engine/Source/Runtime/MovieSceneTracks/Public/Tracks/MovieScene3DConstraintTrack.h b/Engine/Source/Runtime/MovieSceneTracks/Public/Tracks/MovieScene3DConstraintTrack.h
--- a/Engine/Source/Runtime/MovieSceneTracks/Public/Tracks/MovieScene3DConstraintTrack.h
+++ b/Engine/Source/Runtime/MovieSceneTracks/Public/Tracks/MovieScene3DConstraintTrack.h
@@ -2,7 +2,7 @@
 
 #pragma once
 
-#include "MovieSceneNameableTrack.h"
+#include "MovieSceneTrack.h"
 #include "MovieScene3DConstraintTrack.generated.h"
 
 
@@ -14,7 +14,7 @@ class UMovieSceneSection;
  */
 UCLASS( MinimalAPI )
 class UMovieScene3DConstraintTrack
-	: public UMovieSceneNameableTrack
+	: public UMovieSceneTrack
 {
 	GENERATED_UCLASS_BODY()
 
@@ -43,10 +43,6 @@ public:
 	virtual TRange<float> GetSectionBoundaries() const override;
 	virtual const TArray<UMovieSceneSection*>& GetAllSections() const override;
 
-#if WITH_EDITORONLY_DATA
-	virtual FText GetDefaultDisplayName() const override;
-#endif
-
 protected:
 
 	/** List of all constraint sections. */
diff --git a/Engine/Source/Runtime/MovieSceneTracks/Public/Tracks/MovieSceneActorReferenceTrack.h b/Engine/Source/Runtime/MovieSceneTracks/Public/Tracks/MovieSceneActorReferenceTrack.h
new file mode 100644
--- /dev/null
+++ b/Engine/Source/Runtime/MovieSceneTracks/Public/Tracks/MovieSceneActorReferenceTrack.h
@@ -0,0 +1,32 @@
+// Copyright 1998-2016 Epic Games, Inc. All Rights Reserved.
+
+#pragma once
+
+#include "MovieScene.h"
+#include "MovieSceneTrack.h"
+#include "MovieScenePropertyTrack.h"
+#include "MovieSceneActorReferenceTrack.generated.h"
+
+/**
+ * Handles manipulation of actor reference properties in a movie scene
+ */
+UCLASS( MinimalAPI )
+class UMovieSceneActorReferenceTrack : public UMovieScenePropertyTrack
+{
+	GENERATED_UCLASS_BODY()
+
+public:
+	/** UMovieSceneTrack interface */
+	virtual UMovieSceneSection* CreateNewSection() override;
+	virtual TSharedPtr<IMovieSceneTrackInstance> CreateInstance() override;
+
+	/**
+	 * Evaluates the track at the playback position
+	 *
+	 * @param Position	The current playback position
+	 * @param LastPosition	The last plackback position
+	 * @param OutActorReference 	The value at the playback position
+	 * @return true if anything was evaluated. Note: if false is returned OutActorReference remains unchanged
+	 */
+	virtual bool Eval( float Position, float LastPosition, FGuid& OutActorReferenceGuid ) const;
+};
diff --git a/Engine/Source/Runtime/MovieSceneTracks/Public/Tracks/MovieSceneCameraAnimTrack.h b/Engine/Source/Runtime/MovieSceneTracks/Public/Tracks/MovieSceneCameraAnimTrack.h
new file mode 100644
--- /dev/null
+++ b/Engine/Source/Runtime/MovieSceneTracks/Public/Tracks/MovieSceneCameraAnimTrack.h
@@ -0,0 +1,46 @@
+// Copyright 1998-2016 Epic Games, Inc. All Rights Reserved.
+
+#pragma once
+
+#include "MovieSceneNameableTrack.h"
+#include "MovieSceneCameraAnimTrack.generated.h"
+
+class UMovieSceneCameraAnimSection;
+
+/**
+ * 
+ */
+UCLASS(MinimalAPI)
+class UMovieSceneCameraAnimTrack : public UMovieSceneNameableTrack
+{
+	GENERATED_BODY()
+
+public:
+	virtual void AddNewCameraAnim(float KeyTime, class UCameraAnim* CameraAnim);
+	void GetCameraAnimSectionsAtTime(float Time, TArray<UMovieSceneCameraAnimSection*>& OutSections);
+
+public:
+
+	// UMovieSceneTrack interface
+	virtual TSharedPtr<IMovieSceneTrackInstance> CreateInstance() override;
+
+	virtual bool HasSection(const UMovieSceneSection& Section) const override;
+	virtual void AddSection(UMovieSceneSection& Section) override;
+	virtual void RemoveSection(UMovieSceneSection& Section) override;
+	virtual bool IsEmpty() const override;
+	virtual TRange<float> GetSectionBoundaries() const override;
+	virtual bool SupportsMultipleRows() const override { return true; }
+	virtual const TArray<UMovieSceneSection*>& GetAllSections() const override;
+	virtual UMovieSceneSection* CreateNewSection() override;
+	virtual void RemoveAllAnimationData() override;
+
+	
+#if WITH_EDITORONLY_DATA
+	virtual FText GetDisplayName() const override;
+#endif
+
+private:
+	/** List of all sections */
+	UPROPERTY()
+	TArray<UMovieSceneSection*> CameraAnimSections;
+};
diff --git a/Engine/Source/Runtime/MovieSceneTracks/Public/Tracks/MovieSceneCameraShakeTrack.h b/Engine/Source/Runtime/MovieSceneTracks/Public/Tracks/MovieSceneCameraShakeTrack.h
new file mode 100644
--- /dev/null
+++ b/Engine/Source/Runtime/MovieSceneTracks/Public/Tracks/MovieSceneCameraShakeTrack.h
@@ -0,0 +1,47 @@
+// Copyright 1998-2016 Epic Games, Inc. All Rights Reserved.
+
+#pragma once
+
+#include "MovieSceneNameableTrack.h"
+#include "MovieSceneCameraShakeTrack.generated.h"
+
+class UCameraShake;
+class UMovieSceneCameraShakeSection;
+
+/**
+ * 
+ */
+UCLASS(MinimalAPI)
+class UMovieSceneCameraShakeTrack : public UMovieSceneNameableTrack
+{
+	GENERATED_BODY()
+
+public:
+	virtual void AddNewCameraShake(float KeyTime, TSubclassOf<UCameraShake> ShakeClass);
+	void GetCameraShakeSectionsAtTime(float Time, TArray<UMovieSceneCameraShakeSection*>& OutSections);
+	
+public:
+
+	// UMovieSceneTrack interface
+	virtual TSharedPtr<IMovieSceneTrackInstance> CreateInstance() override;
+
+	virtual bool HasSection(const UMovieSceneSection& Section) const override;
+	virtual void AddSection(UMovieSceneSection& Section) override;
+	virtual void RemoveSection(UMovieSceneSection& Section) override;
+	virtual bool IsEmpty() const override;
+	virtual TRange<float> GetSectionBoundaries() const override;
+	virtual bool SupportsMultipleRows() const override { return true; }
+	virtual const TArray<UMovieSceneSection*>& GetAllSections() const override;
+	virtual UMovieSceneSection* CreateNewSection() override;
+	virtual void RemoveAllAnimationData() override;
+
+#if WITH_EDITORONLY_DATA
+	virtual FText GetDisplayName() const override;
+#endif
+	
+private:
+	/** List of all sections */
+	UPROPERTY()
+	TArray<UMovieSceneSection*> CameraShakeSections;
+
+};
diff --git a/Engine/Source/Runtime/MovieSceneTracks/Public/Tracks/MovieSceneFadeTrack.h b/Engine/Source/Runtime/MovieSceneTracks/Public/Tracks/MovieSceneFadeTrack.h
--- a/Engine/Source/Runtime/MovieSceneTracks/Public/Tracks/MovieSceneFadeTrack.h
+++ b/Engine/Source/Runtime/MovieSceneTracks/Public/Tracks/MovieSceneFadeTrack.h
@@ -24,5 +24,6 @@ public:
 
 #if WITH_EDITORONLY_DATA
 	virtual FText GetDefaultDisplayName() const override;
+	virtual bool CanRename() const override { return true; }
 #endif
 };
diff --git a/Engine/Source/Runtime/MovieSceneTracks/Public/Tracks/MovieScenePropertyTrack.h b/Engine/Source/Runtime/MovieSceneTracks/Public/Tracks/MovieScenePropertyTrack.h
--- a/Engine/Source/Runtime/MovieSceneTracks/Public/Tracks/MovieScenePropertyTrack.h
+++ b/Engine/Source/Runtime/MovieSceneTracks/Public/Tracks/MovieScenePropertyTrack.h
@@ -34,6 +34,7 @@ public:
 
 #if WITH_EDITORONLY_DATA
 	virtual FText GetDefaultDisplayName() const override;
+	virtual bool CanRename() const override { return false; }
 #endif
 
 public:
diff --git a/Engine/Source/Runtime/MovieSceneTracks/Public/Tracks/MovieSceneSlomoTrack.h b/Engine/Source/Runtime/MovieSceneTracks/Public/Tracks/MovieSceneSlomoTrack.h
--- a/Engine/Source/Runtime/MovieSceneTracks/Public/Tracks/MovieSceneSlomoTrack.h
+++ b/Engine/Source/Runtime/MovieSceneTracks/Public/Tracks/MovieSceneSlomoTrack.h
@@ -28,5 +28,6 @@ public:
 
 #if WITH_EDITORONLY_DATA
 	virtual FText GetDefaultDisplayName() const override;
+	virtual bool CanRename() const override { return true; }
 #endif
 };
diff --git a/Engine/Source/Runtime/MovieSceneTracks/Public/Tracks/MovieSceneSpawnTrack.h b/Engine/Source/Runtime/MovieSceneTracks/Public/Tracks/MovieSceneSpawnTrack.h
--- a/Engine/Source/Runtime/MovieSceneTracks/Public/Tracks/MovieSceneSpawnTrack.h
+++ b/Engine/Source/Runtime/MovieSceneTracks/Public/Tracks/MovieSceneSpawnTrack.h
@@ -2,7 +2,7 @@
 
 #pragma once
 
-#include "MovieSceneNameableTrack.h"
+#include "MovieSceneTrack.h"
 #include "MovieSceneSpawnTrack.generated.h"
 
 
@@ -15,7 +15,7 @@ class UMovieSceneSection;
  */
 UCLASS(MinimalAPI)
 class UMovieSceneSpawnTrack
-	: public UMovieSceneNameableTrack
+	: public UMovieSceneTrack
 {
 public:
 
@@ -53,7 +53,7 @@ public:
 	virtual const TArray<UMovieSceneSection*>& GetAllSections() const override;
 
 #if WITH_EDITORONLY_DATA
-	virtual FText GetDefaultDisplayName() const override;
+	virtual FText GetDisplayName() const override;
 #endif
 
 protected:
diff --git a/Engine/Source/Runtime/MovieSceneTracks/Public/Tracks/MovieSceneSubTrack.h b/Engine/Source/Runtime/MovieSceneTracks/Public/Tracks/MovieSceneSubTrack.h
--- a/Engine/Source/Runtime/MovieSceneTracks/Public/Tracks/MovieSceneSubTrack.h
+++ b/Engine/Source/Runtime/MovieSceneTracks/Public/Tracks/MovieSceneSubTrack.h
@@ -44,6 +44,11 @@ public:
 	 */
 	MOVIESCENETRACKS_API bool ContainsSequence(const UMovieSceneSequence& Sequence, bool Recursively = false) const;
 
+	/**
+	 * Add a new sequence to record
+	 */
+	MOVIESCENETRACKS_API virtual UMovieSceneSubSection* AddSequenceToRecord();
+
 public:
 
 	// UMovieSceneTrack interface
@@ -62,7 +67,7 @@ public:
 #if WITH_EDITORONLY_DATA
 	virtual FText GetDefaultDisplayName() const override;
 #endif
-
+	
 protected:
 
 	/** All movie scene sections. */
diff --git a/Engine/Source/Runtime/Renderer/Private/PostProcess/PostProcessCircleDOF.cpp b/Engine/Source/Runtime/Renderer/Private/PostProcess/PostProcessCircleDOF.cpp
--- a/Engine/Source/Runtime/Renderer/Private/PostProcess/PostProcessCircleDOF.cpp
+++ b/Engine/Source/Runtime/Renderer/Private/PostProcess/PostProcessCircleDOF.cpp
@@ -31,8 +31,10 @@ float ComputeFocalLengthFromFov(const FSceneView& View)
 	//   f = focal length
 	// 
 	// f = 0.5 * d * (1/tan(fov/2))
-	float HalfFOV = FMath::Atan(1.0f / View.ViewMatrices.ProjMatrix.M[0][0]);
-	float FocalLength = 0.5f * 24.576f * (1.0f/FMath::Tan(HalfFOV));
+
+	float const d = View.FinalPostProcessSettings.DepthOfFieldSensorWidth;
+	float const HalfFOV = FMath::Atan(1.0f / View.ViewMatrices.ProjMatrix.M[0][0]);
+	float const FocalLength = 0.5f * d * (1.0f/FMath::Tan(HalfFOV));
 
 	return FocalLength;
 }
@@ -63,10 +65,10 @@ FVector4 CircleDofHalfCoc(const FSceneView& View)
 		//   n = fstop (where n is the "n" in "f/n")
 		float Radius = FMath::Square(FocalLengthInMM) / (View.FinalPostProcessSettings.DepthOfFieldFstop * (FocalDistanceInMM - FocalLengthInMM));
 
-		// Scale so that APS-C 24.576 mm = full frame.
 		// Convert mm to pixels.
-		float Width = (float)View.ViewRect.Width();
-		Radius = Radius * Width * (1.0f/24.576f);
+		float const Width = (float)View.ViewRect.Width();
+		float const SensorWidth = View.FinalPostProcessSettings.DepthOfFieldSensorWidth;
+		Radius = Radius * Width * (1.0f / SensorWidth);
 
 		// Convert diameter to radius at half resolution (algorithm radius is at half resolution).
 		Radius *= 0.25f;
diff --git a/Engine/Source/Runtime/Renderer/Private/PostProcess/PostProcessCircleDOF.h b/Engine/Source/Runtime/Renderer/Private/PostProcess/PostProcessCircleDOF.h
--- a/Engine/Source/Runtime/Renderer/Private/PostProcess/PostProcessCircleDOF.h
+++ b/Engine/Source/Runtime/Renderer/Private/PostProcess/PostProcessCircleDOF.h
@@ -95,7 +95,7 @@ private:
 };
 
 // to verify this can be used: http://www.radical.org/aov
-// @return in mm, assuming this sensor: APS-C 24.576 mm
+// @return in mm, assuming the sensor in the DepthOfField settings
 float ComputeFocalLengthFromFov(const FSceneView& View);
 
 FVector4 CircleDofHalfCoc(const FSceneView& View);
diff --git a/Engine/Source/Runtime/Renderer/Private/ScenePrivate.h b/Engine/Source/Runtime/Renderer/Private/ScenePrivate.h
--- a/Engine/Source/Runtime/Renderer/Private/ScenePrivate.h
+++ b/Engine/Source/Runtime/Renderer/Private/ScenePrivate.h
@@ -1510,24 +1510,26 @@ public:
 	/** Information about the primitives that are attached together. */
 	struct FLODSceneNode
 	{
-		/** The primitive. */
-		FPrimitiveSceneInfo* SceneInfo;
-
 		/** Children scene infos. */
 		TArray<FPrimitiveSceneInfo*> ChildrenSceneInfos;
 
+		/** The primitive. */
+		FPrimitiveSceneInfo* SceneInfo;
+
 		/** Last updated FrameCount */
 		int32 LatestUpdateCount;
 
 		/** Persistent visibility states */
 		bool bWasVisible;
 		bool bIsVisible;
-
-		FLODSceneNode() :
-			SceneInfo(nullptr), 
-			LatestUpdateCount(INDEX_NONE),
-			bWasVisible(false),
-			bIsVisible(false)
+		bool bIsFading;
+
+		FLODSceneNode()
+			: SceneInfo(nullptr)
+			, LatestUpdateCount(INDEX_NONE)
+			, bWasVisible(false)
+			, bIsVisible(false)
+			, bIsFading(false)
 		{}
 
 		void AddChild(FPrimitiveSceneInfo * NewChild)
@@ -1551,8 +1553,7 @@ public:
 	void RemoveChildNode(FPrimitiveComponentId NodeId, FPrimitiveSceneInfo* ChildSceneInfo);
 
 	void UpdateNodeSceneInfo(FPrimitiveComponentId NodeId, FPrimitiveSceneInfo* SceneInfo);
-	void PopulateFadingFlags(FViewInfo& View);
-	void PopulateHiddenFlags(FViewInfo& View, FSceneBitArray& HiddenFlags);
+	void UpdateAndApplyVisibilityStates(FViewInfo& View);
 
 	bool IsNodeFading(const int32 Index) const
 	{
@@ -1583,9 +1584,9 @@ private:
 	/**  Update Count. This is used to skip Child node that has been updated */
 	int32 UpdateCount;
 
-	/** Propagate flags to children */
-	void PropagateFadingFlagsToChildren(FViewInfo& View, FLODSceneNode& Node, bool bIsFading, bool bIsFadingOut);
-	void PropagateHiddenFlagsToChildren(FSceneBitArray& HiddenFlags, FLODSceneNode& Node);
+	/** Recursive state updates */
+	void ApplyNodeFadingToChildren(FLODSceneNode& Node, FSceneBitArray& VisibilityFlags, const bool bIsFading, const bool bIsFadingOut);
+	void HideNodeChildren(FLODSceneNode& Node, FSceneBitArray& VisibilityFlags);
 };
 
 typedef TMap<FMaterial*, FMaterialShaderMap*> FMaterialsToUpdateMap;
diff --git a/Engine/Source/Runtime/Renderer/Private/SceneRendering.cpp b/Engine/Source/Runtime/Renderer/Private/SceneRendering.cpp
--- a/Engine/Source/Runtime/Renderer/Private/SceneRendering.cpp
+++ b/Engine/Source/Runtime/Renderer/Private/SceneRendering.cpp
@@ -946,6 +946,7 @@ void FViewInfo::CreateUniformBuffer(
 	float ExposureScale = FRCPassPostProcessEyeAdaptation::ComputeExposureScaleValue( *this );
 	FrameUniformShaderParameters.ExposureScale = FVector4(ExposureScale, ExposureScale, ExposureScale, 1.0f);
 	FrameUniformShaderParameters.DepthOfFieldFocalDistance = FinalPostProcessSettings.DepthOfFieldFocalDistance;
+	FrameUniformShaderParameters.DepthOfFieldSensorWidth = FinalPostProcessSettings.DepthOfFieldSensorWidth;
 	FrameUniformShaderParameters.DepthOfFieldFocalRegion = FinalPostProcessSettings.DepthOfFieldFocalRegion;
 	// clamped to avoid div by 0 in shader
 	FrameUniformShaderParameters.DepthOfFieldNearTransitionRegion = FMath::Max(0.01f, FinalPostProcessSettings.DepthOfFieldNearTransitionRegion);
diff --git a/Engine/Source/Runtime/Renderer/Private/SceneVisibility.cpp b/Engine/Source/Runtime/Renderer/Private/SceneVisibility.cpp
--- a/Engine/Source/Runtime/Renderer/Private/SceneVisibility.cpp
+++ b/Engine/Source/Runtime/Renderer/Private/SceneVisibility.cpp
@@ -2510,21 +2510,12 @@ void FSceneRenderer::ComputeViewVisibility(FRHICommandListImmediate& RHICmdList)
 		}
 
 		// visibility test is done, so now build the hidden flags based on visibility set up
-		bool bLODSceneTreeActive = Scene->SceneLODHierarchy.IsActive();
-		FSceneBitArray PrimitiveHiddenByLODMap;
-		if (bLODSceneTreeActive)
+		FLODSceneTree& HLODTree = Scene->SceneLODHierarchy;
+
+		if (HLODTree.IsActive())
 		{
 			QUICK_SCOPE_CYCLE_COUNTER(STAT_ViewVisibilityTime_HLOD);
-			Scene->SceneLODHierarchy.PopulateFadingFlags(View);
-
-			PrimitiveHiddenByLODMap.Init(false, View.PrimitiveVisibilityMap.Num());
-			Scene->SceneLODHierarchy.PopulateHiddenFlags(View, PrimitiveHiddenByLODMap);
-
-			// now iterate through turn off visibility if hidden by LOD
-			for(FSceneSetBitIterator BitIt(PrimitiveHiddenByLODMap); BitIt; ++BitIt)
-			{
-				View.PrimitiveVisibilityMap.AccessCorrespondingBit(BitIt) = false;
-			}
+			HLODTree.UpdateAndApplyVisibilityStates(View);
 		}
 
 		MarkAllPrimitivesForReflectionProxyUpdate(Scene);
@@ -2536,15 +2527,6 @@ void FSceneRenderer::ComputeViewVisibility(FRHICommandListImmediate& RHICmdList)
 		{
 			SCOPE_CYCLE_COUNTER(STAT_ViewRelevance);
 			ComputeAndMarkRelevanceForViewParallel(RHICmdList, Scene, View, ViewBit, HasDynamicMeshElementsMasks, HasDynamicEditorMeshElementsMasks);
-
-			if (bLODSceneTreeActive)
-			{
-				QUICK_SCOPE_CYCLE_COUNTER(STAT_ViewVisibilityTime_HLOD_Prop);
-				for (FSceneSetBitIterator BitIt(PrimitiveHiddenByLODMap); BitIt; ++BitIt)
-				{
-					View.PrimitiveViewRelevanceMap[BitIt.GetIndex()].bInitializedThisFrame = true;
-				}
-			}
 		}
 
 #if !(UE_BUILD_SHIPPING || UE_BUILD_TEST)
@@ -2937,10 +2919,13 @@ void FLODSceneTree::UpdateNodeSceneInfo(FPrimitiveComponentId NodeId, FPrimitive
 	}
 }
 
-void FLODSceneTree::PopulateFadingFlags(FViewInfo& View)
+void FLODSceneTree::UpdateAndApplyVisibilityStates(FViewInfo& View)
 {
 	PrimitiveFadingLODMap.Init(false, View.PrimitiveVisibilityMap.Num());
 	PrimitiveFadingOutLODMap.Init(false, View.PrimitiveVisibilityMap.Num());
+	FSceneBitArray& VisibilityFlags = View.PrimitiveVisibilityMap;
+
+	++UpdateCount;
 
 	if (const FSceneViewState* ViewState = (FSceneViewState*)View.State)
 	{
@@ -2957,49 +2942,90 @@ void FLODSceneTree::PopulateFadingFlags(FViewInfo& View)
 		for (auto Iter = SceneNodes.CreateIterator(); Iter; ++Iter)
 		{
 			FLODSceneNode& Node = Iter.Value();
-			if (Node.SceneInfo && Node.LatestUpdateCount == UpdateCount)
+			const TIndirectArray<FStaticMesh>& NodeMeshes = Node.SceneInfo->StaticMeshes;
+
+			// Ignore already updated nodes, or those that we can't work with
+			if (Node.LatestUpdateCount == UpdateCount || !Node.SceneInfo || NodeMeshes.Num() == 0)
 			{
-				const int32 NodeIndex = Node.SceneInfo->GetIndex();
-				const TIndirectArray<FStaticMesh>& NodeMeshes = Node.SceneInfo->StaticMeshes;
+				continue;
+			}
 
-				if (NodeMeshes.Num() > 0 && NodeMeshes[0].bDitheredLODTransition)
-				{
-					if (bSyncFrame)
-					{
-						bool bChildrenFading = false;
+			const int32 NodeIndex = Node.SceneInfo->GetIndex();
+			bool bIsVisible = VisibilityFlags[NodeIndex];
 
-						// Note: Unless ending a fade, need to wait for all children to finish fading?
-						// Not sure this can actually happen as we wait for the syncs so will override the fade
-						if (Node.bWasVisible != Node.bIsVisible || !bChildrenFading)
-						{
-							Node.bWasVisible = Node.bIsVisible;
-							Node.bIsVisible = View.PrimitiveVisibilityMap[NodeIndex];
-						}
-					}
-
-					// Fade until state back in sync, then hold visibility
-					if (Node.bWasVisible != Node.bIsVisible)
+			// Update fading state
+			if (NodeMeshes[0].bDitheredLODTransition)
+			{
+				// Update with syncs
+				if (bSyncFrame)
+				{
+					// Determine desired HLOD state
+					const FPrimitiveBounds& Bounds = Scene->PrimitiveBounds[NodeIndex];
+					const float DistanceSquared = (Bounds.Origin - View.ViewMatrices.ViewOrigin).SizeSquared();
+					const bool bIsInDrawRange = DistanceSquared >= Bounds.MinDrawDistanceSq;
+
+					// Fade when HLODs change threshold on-screen, else snap
+					// TODO: This logic can still be improved to clear state and
+					//       transitions when off-screen, but needs better detection
+					const bool bChangedRange = bIsInDrawRange != Node.bWasVisible;
+					const bool bIsOnScreen = bIsVisible || Node.bWasVisible;
+				
+					if (Node.bIsFading)
 					{
-						PrimitiveFadingLODMap[NodeIndex] = true;
-						PrimitiveFadingOutLODMap[NodeIndex] = !Node.bIsVisible;
-
-						PropagateFadingFlagsToChildren(View, Node, true, Node.bIsVisible);
+						Node.bIsFading = false;
 					}
-					else
+					else if (bChangedRange && bIsOnScreen)
 					{
-						View.PrimitiveVisibilityMap[NodeIndex] = Node.bIsVisible;
+						Node.bIsFading = true;	
 					}
+
+					Node.bWasVisible = Node.bIsVisible;
+					Node.bIsVisible = bIsInDrawRange;
+				}
+
+				// Flag as fading or freeze visibility if waiting for a fade
+				if (Node.bIsFading)
+				{
+					PrimitiveFadingLODMap[NodeIndex] = true;
+					PrimitiveFadingOutLODMap[NodeIndex] = !Node.bIsVisible;
+				}
+				else
+				{
+					VisibilityFlags[NodeIndex] = Node.bWasVisible;
+					bIsVisible = Node.bWasVisible;
 				}
 			}
+#if WITH_EDITOR
+			else
+			{
+				// Force the default state in-case material edits cause an object to get stuck
+				Node.bIsFading = false;
+			}
+#endif
+
+			if (Node.bIsFading)
+			{
+				// Fade until state back in sync
+				ApplyNodeFadingToChildren(Node, VisibilityFlags, true, Node.bIsVisible);
+			}
+			else if (bIsVisible)
+			{
+				// If stable and visible, override hierarchy visibility
+				HideNodeChildren(Node, VisibilityFlags);
+			}
 		}
 	}
 }
 
-void FLODSceneTree::PropagateFadingFlagsToChildren(FViewInfo& View, FLODSceneNode& Node, bool bIsFading, bool bIsFadingOut)
+void FLODSceneTree::ApplyNodeFadingToChildren(FLODSceneNode& Node, FSceneBitArray& VisibilityFlags, const bool bIsFading, const bool bIsFadingOut)
 {
 	if (Node.SceneInfo)
 	{
+		Node.LatestUpdateCount = UpdateCount;
+
+		// Force visibility during fades
 		const int32 NodeIndex = Node.SceneInfo->GetIndex();
+		VisibilityFlags[NodeIndex] = true;
 
 		for (const auto& Child : Node.ChildrenSceneInfos)
 		{
@@ -3007,91 +3033,31 @@ void FLODSceneTree::PropagateFadingFlagsToChildren(FViewInfo& View, FLODSceneNod
 
 			PrimitiveFadingLODMap[ChildIndex] = bIsFading;
 			PrimitiveFadingOutLODMap[ChildIndex] = bIsFadingOut;
+			VisibilityFlags[ChildIndex] = true;
 
+			// Fading only occurs at the adjacent hierarchy level, below should be hidden
 			if (FLODSceneNode* ChildNode = SceneNodes.Find(Child->PrimitiveComponentId))
 			{
-				PropagateFadingFlagsToChildren(View, *ChildNode, bIsFading, bIsFadingOut);
+				HideNodeChildren(*ChildNode, VisibilityFlags);
 			}
 		}
-
-		// Force visibility during fades
-		View.PrimitiveVisibilityMap[NodeIndex] = true;
 	}
 }
 
-void FLODSceneTree::PropagateHiddenFlagsToChildren(FSceneBitArray& HiddenFlags, FLODSceneNode& Node)
+void FLODSceneTree::HideNodeChildren(FLODSceneNode& Node, FSceneBitArray& VisibilityFlags)
 {
-	// if already updated, no reason to do this
-	if(Node.LatestUpdateCount != UpdateCount)
+	if (Node.LatestUpdateCount != UpdateCount)
 	{
 		Node.LatestUpdateCount = UpdateCount;
-		// if node doesn't have scene info, that means it doesn't to populate, children is disconnected, so don't bother
-		// in this case we still update children when this node is missing scene info
-		// because parent is showing, so you don't have to show anyway anybody below
-		// although this node might be MIA at this moment
-		for(const auto& Child : Node.ChildrenSceneInfos)
+
+		for (const auto& Child : Node.ChildrenSceneInfos)
 		{
 			const int32 ChildIndex = Child->GetIndex();
+			VisibilityFlags[ChildIndex] = false;
 
-			// first update the flags
-			FRelativeBitReference BitRef(ChildIndex);
-			HiddenFlags.AccessCorrespondingBit(BitRef) = true;
-			// find the node for it
-			FLODSceneNode* ChildNode = SceneNodes.Find(Child->PrimitiveComponentId);
-			// if you have child, populate it again, 
-			if (ChildNode)
-			{
-				PropagateHiddenFlagsToChildren(HiddenFlags, *ChildNode);
-			}
-		}
-	}
-}
-
-void FLODSceneTree::PopulateHiddenFlags(FViewInfo& View, FSceneBitArray& HiddenFlags)
-{
-	++UpdateCount;
-
-	// @todo this is experimental code - hide the children if parent is showing
-	for(auto Iter = SceneNodes.CreateIterator(); Iter; ++Iter)
-	{
-		FLODSceneNode& Node = Iter.Value();
-		// if already updated, no reason to do this
-		if(Node.LatestUpdateCount != UpdateCount)
-		{
-			Node.LatestUpdateCount = UpdateCount;
-			// if node doesn't have scene info, that means it doesn't have any
-			if(Node.SceneInfo)
+			if (FLODSceneNode* ChildNode = SceneNodes.Find(Child->PrimitiveComponentId))
 			{
-				int32 NodeIndex = Node.SceneInfo->GetIndex();
-
-				// Only override visibility when fading isn't forcing state
-				if (!IsNodeFading(NodeIndex))
-				{
-					// if this node is visible, children shouldn't show up
-					if (View.PrimitiveVisibilityMap[NodeIndex])
-					{
-						for (const auto& Child : Node.ChildrenSceneInfos)
-						{
-							const int32 ChildIndex = Child->GetIndex();
-
-							// first update the flags
-							FRelativeBitReference BitRef(ChildIndex);
-							HiddenFlags.AccessCorrespondingBit(BitRef) = true;
-
-							// find the node for it
-							FLODSceneNode* ChildNode = SceneNodes.Find(Child->PrimitiveComponentId);
-							// if you have child, populate it again, 
-							if (ChildNode)
-							{
-								PropagateHiddenFlagsToChildren(HiddenFlags, *ChildNode);
-							}
-						}
-					}
-					else
-					{
-						HiddenFlags.AccessCorrespondingBit(FRelativeBitReference(NodeIndex)) = true;
-					}
-				}
+				HideNodeChildren(*ChildNode, VisibilityFlags);
 			}
 		}
 	}
diff --git a/Engine/Source/Runtime/Slate/Private/Widgets/Input/SComboButton.cpp b/Engine/Source/Runtime/Slate/Private/Widgets/Input/SComboButton.cpp
--- a/Engine/Source/Runtime/Slate/Private/Widgets/Input/SComboButton.cpp
+++ b/Engine/Source/Runtime/Slate/Private/Widgets/Input/SComboButton.cpp
@@ -24,6 +24,7 @@ void SComboButton::Construct( const FArguments& InArgs )
 		.Method(InArgs._Method)
 		.OnMenuOpenChanged(InArgs._OnMenuOpenChanged)
 		.OnGetMenuContent(InArgs._OnGetMenuContent)
+		.IsCollapsedByParent(InArgs._CollapseMenuOnParentFocus)
 		[
 			SNew( SButton )
 			.ButtonStyle( OurButtonStyle )
diff --git a/Engine/Source/Runtime/Slate/Private/Widgets/Input/SMenuAnchor.cpp b/Engine/Source/Runtime/Slate/Private/Widgets/Input/SMenuAnchor.cpp
--- a/Engine/Source/Runtime/Slate/Private/Widgets/Input/SMenuAnchor.cpp
+++ b/Engine/Source/Runtime/Slate/Private/Widgets/Input/SMenuAnchor.cpp
@@ -109,6 +109,7 @@ void SMenuAnchor::Construct( const FArguments& InArgs )
 	Placement = InArgs._Placement;
 	Method = InArgs._Method;
 	ShouldDeferPaintingAfterWindowContent = InArgs._ShouldDeferPaintingAfterWindowContent;
+	bIsCollapsedByParent = InArgs._IsCollapsedByParent;
 }
 
 
@@ -375,7 +376,6 @@ void SMenuAnchor::SetIsOpen( bool InIsOpen, const bool bFocusMenu )
 					if (MethodInUse.GetPopupMethod() == EPopupMethod::CreateNewWindow)
 					{
 						// Open the pop-up
-						const bool bIsCollapsedByParent = false;	// don't auto-close child menus when the parent gets focus
 						TSharedPtr<IMenu> NewMenu = FSlateApplication::Get().PushMenu(AsShared(), MyWidgetPath, MenuContentRef, NewPosition, TransitionEffect, bFocusMenu, SummonLocationSize, MethodInUse.GetPopupMethod(), bIsCollapsedByParent);
 
 						PopupMenuPtr = NewMenu;
@@ -398,7 +398,6 @@ void SMenuAnchor::SetIsOpen( bool InIsOpen, const bool bFocusMenu )
 
 						TSharedRef<SMenuAnchor> SharedThis = StaticCastSharedRef<SMenuAnchor>(AsShared());
 
-						const bool bIsCollapsedByParent = false;	// don't auto-close child menus when the parent gets focus
 						TSharedPtr<IMenu> NewMenu = FSlateApplication::Get().PushHostedMenu(
 							SharedThis, MyWidgetPath, SharedThis, MenuContentRef, WrappedContent, TransitionEffect, MethodInUse.GetShouldThrottle(), bIsCollapsedByParent);
 
diff --git a/Engine/Source/Runtime/Slate/Public/Widgets/Input/SComboButton.h b/Engine/Source/Runtime/Slate/Public/Widgets/Input/SComboButton.h
--- a/Engine/Source/Runtime/Slate/Public/Widgets/Input/SComboButton.h
+++ b/Engine/Source/Runtime/Slate/Public/Widgets/Input/SComboButton.h
@@ -25,6 +25,7 @@ public:
 		, _HAlign(HAlign_Fill)
 		, _VAlign(VAlign_Center)
 		, _Method()
+		, _CollapseMenuOnParentFocus(false)
 		{}
 
 		SLATE_STYLE_ARGUMENT( FComboButtonStyle, ComboButtonStyle )
@@ -54,6 +55,10 @@ public:
 
 		/** Spawn a new window or reuse current window for this combo*/
 		SLATE_ARGUMENT( TOptional<EPopupMethod>, Method )
+
+		/** True if this combo's menu should be collapsed when our parent receives focus, false (default) otherwise */
+		SLATE_ARGUMENT(bool, CollapseMenuOnParentFocus)
+
 	SLATE_END_ARGS()
 
 	// SMenuAnchor interface
diff --git a/Engine/Source/Runtime/Slate/Public/Widgets/Input/SMenuAnchor.h b/Engine/Source/Runtime/Slate/Public/Widgets/Input/SMenuAnchor.h
--- a/Engine/Source/Runtime/Slate/Public/Widgets/Input/SMenuAnchor.h
+++ b/Engine/Source/Runtime/Slate/Public/Widgets/Input/SMenuAnchor.h
@@ -31,6 +31,7 @@ public:
 		, _Placement( MenuPlacement_BelowAnchor )
 		, _Method()
 		, _ShouldDeferPaintingAfterWindowContent(true)
+		, _IsCollapsedByParent(false)
 		{}
 		
 		SLATE_DEFAULT_SLOT( FArguments, Content )
@@ -47,6 +48,9 @@ public:
 
 		SLATE_ARGUMENT(bool, ShouldDeferPaintingAfterWindowContent)
 
+		/** True if this menu anchor should be collapsed when its parent receives focus, false (default) otherwise */
+		SLATE_ARGUMENT(bool, IsCollapsedByParent)
+
 	SLATE_END_ARGS()
 
 	virtual ~SMenuAnchor();
@@ -142,6 +146,9 @@ protected:
 	/** Was the menu just dismissed this tick? */
 	bool bDismissedThisTick;
 
+	/** Whether this menu should be collapsed when its parent gets focus */
+	bool bIsCollapsedByParent;
+
 	/** Should we summon a new window for this popup or  */
 	TOptional<EPopupMethod> Method;
 
diff --git a/Engine/Source/Runtime/Slate/Public/Widgets/Input/SNumericDropDown.h b/Engine/Source/Runtime/Slate/Public/Widgets/Input/SNumericDropDown.h
--- a/Engine/Source/Runtime/Slate/Public/Widgets/Input/SNumericDropDown.h
+++ b/Engine/Source/Runtime/Slate/Public/Widgets/Input/SNumericDropDown.h
@@ -116,8 +116,10 @@ public:
 					.ButtonContent()
 					[
 						SNew( SEditableTextBox )
-						.Text( this, &SNumericDropDown<NumericType>::GetValueText )
 						.MinDesiredWidth( InArgs._MinDesiredValueWidth )
+						.RevertTextOnEscape(true)
+						.SelectAllTextWhenFocused(true)
+						.Text( this, &SNumericDropDown<NumericType>::GetValueText )
 						.OnTextCommitted( this, &SNumericDropDown<NumericType>::ValueTextComitted )
 					]
 				]
@@ -173,6 +175,7 @@ private:
 
 	void SetValue( NumericType InValue )
 	{
+		FSlateApplication::Get().ClearKeyboardFocus( EFocusCause::Cleared );
 		OnValueChanged.ExecuteIfBound(InValue);
 	}
 
diff --git a/Engine/Source/Runtime/UMG/Private/Animation/MovieScene2DTransformTrack.cpp b/Engine/Source/Runtime/UMG/Private/Animation/MovieScene2DTransformTrack.cpp
--- a/Engine/Source/Runtime/UMG/Private/Animation/MovieScene2DTransformTrack.cpp
+++ b/Engine/Source/Runtime/UMG/Private/Animation/MovieScene2DTransformTrack.cpp
@@ -11,7 +11,7 @@ UMovieScene2DTransformTrack::UMovieScene2DTransformTrack(const FObjectInitialize
 	: Super(ObjectInitializer)
 {
 #if WITH_EDITORONLY_DATA
-	TrackTint = FColor(48, 227, 255);
+	TrackTint = FColor(48, 227, 255, 65);
 #endif
 }
 
diff --git a/Engine/Source/Runtime/UMG/Private/Animation/MovieScene2DTransformTrackInstance.cpp b/Engine/Source/Runtime/UMG/Private/Animation/MovieScene2DTransformTrackInstance.cpp
--- a/Engine/Source/Runtime/UMG/Private/Animation/MovieScene2DTransformTrackInstance.cpp
+++ b/Engine/Source/Runtime/UMG/Private/Animation/MovieScene2DTransformTrackInstance.cpp
@@ -5,17 +5,21 @@
 #include "MovieScene2DTransformTrack.h"
 #include "MovieSceneCommonHelpers.h"
 
+
 FMovieScene2DTransformTrackInstance::FMovieScene2DTransformTrackInstance( UMovieScene2DTransformTrack& InTransformTrack )
 	: TransformTrack( &InTransformTrack )
 {
 	PropertyBindings = MakeShareable(new FTrackInstancePropertyBindings(TransformTrack->GetPropertyName(), TransformTrack->GetPropertyPath()));
 }
 
-void FMovieScene2DTransformTrackInstance::Update(EMovieSceneUpdateData& UpdateData, const TArray<UObject*>& RuntimeObjects, class IMovieScenePlayer& Player, FMovieSceneSequenceInstance& SequenceInstance) 
+
+void FMovieScene2DTransformTrackInstance::Update(EMovieSceneUpdateData& UpdateData, const TArray<TWeakObjectPtr<UObject>>& RuntimeObjects, class IMovieScenePlayer& Player, FMovieSceneSequenceInstance& SequenceInstance) 
 {
-	for(UObject* Object : RuntimeObjects)
+	for (auto ObjectPtr : RuntimeObjects)
 	{
+		UObject* Object = ObjectPtr.Get();
 		FWidgetTransform TransformValue = PropertyBindings->GetCurrentValue<FWidgetTransform>(Object);
+
 		if(TransformTrack->Eval(UpdateData.Position, UpdateData.LastPosition, TransformValue))
 		{
 			PropertyBindings->CallFunction<FWidgetTransform>(Object, &TransformValue);
@@ -23,8 +27,8 @@ void FMovieScene2DTransformTrackInstance::Update(EMovieSceneUpdateData& UpdateDa
 	}
 }
 
-void FMovieScene2DTransformTrackInstance::RefreshInstance(const TArray<UObject*>& RuntimeObjects, class IMovieScenePlayer& Player, FMovieSceneSequenceInstance& SequenceInstance)
+
+void FMovieScene2DTransformTrackInstance::RefreshInstance(const TArray<TWeakObjectPtr<UObject>>& RuntimeObjects, class IMovieScenePlayer& Player, FMovieSceneSequenceInstance& SequenceInstance)
 {
 	PropertyBindings->UpdateBindings(RuntimeObjects);
 }
-
diff --git a/Engine/Source/Runtime/UMG/Private/Animation/MovieScene2DTransformTrackInstance.h b/Engine/Source/Runtime/UMG/Private/Animation/MovieScene2DTransformTrackInstance.h
--- a/Engine/Source/Runtime/UMG/Private/Animation/MovieScene2DTransformTrackInstance.h
+++ b/Engine/Source/Runtime/UMG/Private/Animation/MovieScene2DTransformTrackInstance.h
@@ -13,10 +13,10 @@ public:
 	FMovieScene2DTransformTrackInstance( class UMovieScene2DTransformTrack& InTransformTrack );
 
 	/** IMovieSceneTrackInstance interface */
-	virtual void SaveState(const TArray<UObject*>& RuntimeObjects, IMovieScenePlayer& Player, FMovieSceneSequenceInstance& SequenceInstance) override {}
-	virtual void RestoreState(const TArray<UObject*>& RuntimeObjects, IMovieScenePlayer& Player, FMovieSceneSequenceInstance& SequenceInstance) override {}
-	virtual void Update(EMovieSceneUpdateData& UpdateData, const TArray<UObject*>& RuntimeObjects, class IMovieScenePlayer& Player, FMovieSceneSequenceInstance& SequenceInstance) override;
-	virtual void RefreshInstance( const TArray<UObject*>& RuntimeObjects, class IMovieScenePlayer& Player, FMovieSceneSequenceInstance& SequenceInstance ) override;
+	virtual void SaveState(const TArray<TWeakObjectPtr<UObject>>& RuntimeObjects, IMovieScenePlayer& Player, FMovieSceneSequenceInstance& SequenceInstance) override {}
+	virtual void RestoreState(const TArray<TWeakObjectPtr<UObject>>& RuntimeObjects, IMovieScenePlayer& Player, FMovieSceneSequenceInstance& SequenceInstance) override {}
+	virtual void Update(EMovieSceneUpdateData& UpdateData, const TArray<TWeakObjectPtr<UObject>>& RuntimeObjects, class IMovieScenePlayer& Player, FMovieSceneSequenceInstance& SequenceInstance) override;
+	virtual void RefreshInstance( const TArray<TWeakObjectPtr<UObject>>& RuntimeObjects, class IMovieScenePlayer& Player, FMovieSceneSequenceInstance& SequenceInstance ) override;
 	virtual void ClearInstance( IMovieScenePlayer& Player, FMovieSceneSequenceInstance& SequenceInstance ) override {};
 private:
 	/** Track that is being instanced */
diff --git a/Engine/Source/Runtime/UMG/Private/Animation/MovieSceneMarginTrackInstance.cpp b/Engine/Source/Runtime/UMG/Private/Animation/MovieSceneMarginTrackInstance.cpp
--- a/Engine/Source/Runtime/UMG/Private/Animation/MovieSceneMarginTrackInstance.cpp
+++ b/Engine/Source/Runtime/UMG/Private/Animation/MovieSceneMarginTrackInstance.cpp
@@ -5,6 +5,7 @@
 #include "MovieSceneMarginTrack.h"
 #include "MovieSceneCommonHelpers.h"
 
+
 FMovieSceneMarginTrackInstance::FMovieSceneMarginTrackInstance( UMovieSceneMarginTrack& InMarginTrack )
 {
 	MarginTrack = &InMarginTrack;
@@ -12,19 +13,23 @@ FMovieSceneMarginTrackInstance::FMovieSceneMarginTrackInstance( UMovieSceneMargi
 	PropertyBindings = MakeShareable( new FTrackInstancePropertyBindings( MarginTrack->GetPropertyName(), MarginTrack->GetPropertyPath() ) );
 }
 
-void FMovieSceneMarginTrackInstance::Update(EMovieSceneUpdateData& UpdateData, const TArray<UObject*>& RuntimeObjects, class IMovieScenePlayer& Player, FMovieSceneSequenceInstance& SequenceInstance) 
+
+void FMovieSceneMarginTrackInstance::Update(EMovieSceneUpdateData& UpdateData, const TArray<TWeakObjectPtr<UObject>>& RuntimeObjects, class IMovieScenePlayer& Player, FMovieSceneSequenceInstance& SequenceInstance) 
 {
-	for( UObject* Object : RuntimeObjects )
+	for (auto ObjectPtr : RuntimeObjects)
 	{
-		FMargin MarginValue = PropertyBindings->GetCurrentValue<FMargin>( Object );
-		if(MarginTrack->Eval(UpdateData.Position, UpdateData.LastPosition, MarginValue))
+		UObject* Object = ObjectPtr.Get();
+		FMargin MarginValue = PropertyBindings->GetCurrentValue<FMargin>(Object);
+
+		if (MarginTrack->Eval(UpdateData.Position, UpdateData.LastPosition, MarginValue))
 		{
 			PropertyBindings->CallFunction<FMargin>(Object, &MarginValue);
 		}
 	}
 }
 
-void FMovieSceneMarginTrackInstance::RefreshInstance( const TArray<UObject*>& RuntimeObjects, class IMovieScenePlayer& Player, FMovieSceneSequenceInstance& SequenceInstance )
+
+void FMovieSceneMarginTrackInstance::RefreshInstance( const TArray<TWeakObjectPtr<UObject>>& RuntimeObjects, class IMovieScenePlayer& Player, FMovieSceneSequenceInstance& SequenceInstance )
 {
 	PropertyBindings->UpdateBindings( RuntimeObjects );
 }
diff --git a/Engine/Source/Runtime/UMG/Private/Animation/MovieSceneMarginTrackInstance.h b/Engine/Source/Runtime/UMG/Private/Animation/MovieSceneMarginTrackInstance.h
--- a/Engine/Source/Runtime/UMG/Private/Animation/MovieSceneMarginTrackInstance.h
+++ b/Engine/Source/Runtime/UMG/Private/Animation/MovieSceneMarginTrackInstance.h
@@ -13,10 +13,10 @@ public:
 	FMovieSceneMarginTrackInstance( class UMovieSceneMarginTrack& InMarginTrack );
 
 	/** IMovieSceneTrackInstance interface */
-	virtual void SaveState( const TArray<UObject*>& RuntimeObjects, IMovieScenePlayer& Player, FMovieSceneSequenceInstance& SequenceInstance ) override {}
-	virtual void RestoreState( const TArray<UObject*>& RuntimeObjects, IMovieScenePlayer& Player, FMovieSceneSequenceInstance& SequenceInstance ) override {}
-	virtual void Update(EMovieSceneUpdateData& UpdateData, const TArray<UObject*>& RuntimeObjects, class IMovieScenePlayer& Player, FMovieSceneSequenceInstance& SequenceInstance) override;
-	virtual void RefreshInstance( const TArray<UObject*>& RuntimeObjects, class IMovieScenePlayer& Player, FMovieSceneSequenceInstance& SequenceInstance ) override;
+	virtual void SaveState( const TArray<TWeakObjectPtr<UObject>>& RuntimeObjects, IMovieScenePlayer& Player, FMovieSceneSequenceInstance& SequenceInstance ) override {}
+	virtual void RestoreState( const TArray<TWeakObjectPtr<UObject>>& RuntimeObjects, IMovieScenePlayer& Player, FMovieSceneSequenceInstance& SequenceInstance ) override {}
+	virtual void Update(EMovieSceneUpdateData& UpdateData, const TArray<TWeakObjectPtr<UObject>>& RuntimeObjects, class IMovieScenePlayer& Player, FMovieSceneSequenceInstance& SequenceInstance) override;
+	virtual void RefreshInstance( const TArray<TWeakObjectPtr<UObject>>& RuntimeObjects, class IMovieScenePlayer& Player, FMovieSceneSequenceInstance& SequenceInstance ) override;
 	virtual void ClearInstance( IMovieScenePlayer& Player, FMovieSceneSequenceInstance& SequenceInstance ) override {};
 private:
 	/** The track being instanced */
diff --git a/Engine/Source/Runtime/UMG/Private/Animation/UMGSequencePlayer.cpp b/Engine/Source/Runtime/UMG/Private/Animation/UMGSequencePlayer.cpp
--- a/Engine/Source/Runtime/UMG/Private/Animation/UMGSequencePlayer.cpp
+++ b/Engine/Source/Runtime/UMG/Private/Animation/UMGSequencePlayer.cpp
@@ -181,7 +181,7 @@ void UUMGSequencePlayer::SetNumLoopsToPlay(int32 InNumLoopsToPlay)
 	}
 }
 
-void UUMGSequencePlayer::GetRuntimeObjects(TSharedRef<FMovieSceneSequenceInstance> MovieSceneInstance, const FGuid& ObjectHandle, TArray< UObject* >& OutObjects) const
+void UUMGSequencePlayer::GetRuntimeObjects(TSharedRef<FMovieSceneSequenceInstance> MovieSceneInstance, const FGuid& ObjectHandle, TArray<TWeakObjectPtr<UObject>>& OutObjects) const
 {
 	const TArray<UObject*>* FoundObjects = GuidToRuntimeObjectMap.Find( ObjectHandle );
 
diff --git a/Engine/Source/Runtime/UMG/Public/Animation/UMGSequencePlayer.h b/Engine/Source/Runtime/UMG/Public/Animation/UMGSequencePlayer.h
--- a/Engine/Source/Runtime/UMG/Public/Animation/UMGSequencePlayer.h
+++ b/Engine/Source/Runtime/UMG/Public/Animation/UMGSequencePlayer.h
@@ -41,8 +41,8 @@ public:
 	void SetNumLoopsToPlay(int32 InNumLoopsToPlay);
 
 	/** IMovieScenePlayer interface */
-	virtual void GetRuntimeObjects( TSharedRef<FMovieSceneSequenceInstance> MovieSceneInstance, const FGuid& ObjectHandle, TArray< UObject* >& OutObjects ) const override;
-	virtual void UpdateCameraCut(UObject* CameraObject, UObject* UnlockIfCameraObject) const override {}
+	virtual void GetRuntimeObjects( TSharedRef<FMovieSceneSequenceInstance> MovieSceneInstance, const FGuid& ObjectHandle, TArray<TWeakObjectPtr<UObject>>& OutObjects ) const override;
+	virtual void UpdateCameraCut(UObject* CameraObject, UObject* UnlockIfCameraObject, bool bJumpCut) const override {}
 	virtual void SetViewportSettings(const TMap<FViewportClient*, EMovieSceneViewportParams>& ViewportParamsMap) override {}
 	virtual void GetViewportSettings(TMap<FViewportClient*, EMovieSceneViewportParams>& ViewportParamsMap) const override {}
 	virtual void AddOrUpdateMovieSceneInstance( class UMovieSceneSection& MovieSceneSection, TSharedRef<FMovieSceneSequenceInstance> InstanceToAdd ) override {}
