Copying //UE4/Dev-Rendering to Dev-Main (Source //UE4/Dev-Rendering@2932636)

diff --git a/Engine/Source/Runtime/Engine/Public/PlanarReflectionSceneProxy.h b/Engine/Source/Runtime/Engine/Public/PlanarReflectionSceneProxy.h
new file mode 100644
--- /dev/null
+++ b/Engine/Source/Runtime/Engine/Public/PlanarReflectionSceneProxy.h
@@ -0,0 +1,97 @@
+// Copyright 1998-2016 Epic Games, Inc. All Rights Reserved.
+
+/*=============================================================================
+	PlanarReflectionSceneProxy.h: 
+=============================================================================*/
+
+#pragma once
+
+class FPlanarReflectionRenderTarget : public FTexture, public FRenderTarget
+{
+public:
+
+	FPlanarReflectionRenderTarget(FIntPoint InSize) :
+		Size(InSize)
+	{}
+
+	virtual const FTexture2DRHIRef& GetRenderTargetTexture() const 
+	{
+		return (const FTexture2DRHIRef&)TextureRHI;
+	}
+
+	virtual void InitDynamicRHI()
+	{
+		// Create the sampler state RHI resource.
+		FSamplerStateInitializerRHI SamplerStateInitializer
+		(
+			SF_Bilinear,
+			AM_Clamp,
+			AM_Clamp,
+			AM_Clamp
+		);
+		SamplerStateRHI = RHICreateSamplerState(SamplerStateInitializer);
+
+		FTexture2DRHIRef Texture2DRHI;
+		FRHIResourceCreateInfo CreateInfo;
+
+		RHICreateTargetableShaderResource2D(
+			GetSizeX(), 
+			GetSizeY(), 
+			PF_FloatRGBA,
+			1,
+			0,
+			TexCreate_RenderTargetable,
+			false,
+			CreateInfo,
+			RenderTargetTextureRHI,
+			Texture2DRHI
+			);
+		TextureRHI = (FTextureRHIRef&)Texture2DRHI;
+	}
+
+	virtual FIntPoint GetSizeXY() const { return Size; }
+
+	/** Returns the width of the texture in pixels. */
+	virtual uint32 GetSizeX() const
+	{
+		return Size.X;
+	}
+	/** Returns the height of the texture in pixels. */
+	virtual uint32 GetSizeY() const
+	{
+		return Size.Y;
+	}
+
+	virtual float GetDisplayGamma() const { return 1.0f; }
+
+	virtual FString GetFriendlyName() const override { return TEXT("FPlanarReflectionRenderTarget"); }
+
+private:
+
+	FIntPoint Size;
+};
+
+class FPlanarReflectionSceneProxy
+{
+public:
+
+	FPlanarReflectionSceneProxy(UPlanarReflectionComponent* Component, FPlanarReflectionRenderTarget* InRenderTarget);
+
+	void UpdateTransform(const FMatrix& NewTransform)
+	{
+		ReflectionPlane = FPlane(NewTransform.TransformPosition(FVector::ZeroVector), NewTransform.TransformVector(FVector(0, 0, 1)));
+
+		const FMirrorMatrix MirrorMatrix(ReflectionPlane);
+		// Using TransposeAdjoint instead of full inverse because we only care about transforming normals
+		const FMatrix InverseTransposeMirrorMatrix4x4 = MirrorMatrix.TransposeAdjoint();
+		InverseTransposeMirrorMatrix4x4.GetScaledAxes((FVector&)InverseTransposeMirrorMatrix[0], (FVector&)InverseTransposeMirrorMatrix[1], (FVector&)InverseTransposeMirrorMatrix[2]);
+	}
+
+	FPlane ReflectionPlane;
+	FVector PlanarReflectionParameters;
+	FVector2D PlanarReflectionParameters2;
+	FMatrix ProjectionWithExtraFOV;
+	FVector4 InverseTransposeMirrorMatrix[3];
+	FPlanarReflectionRenderTarget* RenderTarget;
+	FName OwnerName;
+};
\ No newline at end of file
diff --git a/Engine/Source/Runtime/Engine/Public/PrimitiveSceneProxy.h b/Engine/Source/Runtime/Engine/Public/PrimitiveSceneProxy.h
--- a/Engine/Source/Runtime/Engine/Public/PrimitiveSceneProxy.h
+++ b/Engine/Source/Runtime/Engine/Public/PrimitiveSceneProxy.h
@@ -18,6 +18,8 @@ class FLightSceneProxy;
 class FLightSceneInfo;
 class FPrimitiveDrawInterface;
 
+struct FStreamingSectionBuildInfo;
+
 /** Data for a simple dynamic light. */
 class FSimpleLightEntry
 {
@@ -393,7 +395,6 @@ public:
 	inline bool HasStaticLighting() const { return bStaticLighting; }
 	inline bool NeedsUnbuiltPreviewLighting() const { return bNeedsUnbuiltPreviewLighting; }
 	inline bool CastsStaticShadow() const { return bCastStaticShadow; }
-	inline bool IsVisibleInPlanarReflection() const { return bVisibleInPlanarReflection; }
 	inline bool CastsDynamicShadow() const { return bCastDynamicShadow; }
 	inline bool AffectsDynamicIndirectLighting() const { return bAffectDynamicIndirectLighting; }
 	inline bool AffectsDistanceFieldLighting() const { return bAffectDistanceFieldLighting; }
@@ -519,9 +520,12 @@ public:
 	ENGINE_API virtual void GetLCIs(FLCIArray& LCIs) {}
 
 	/**
-	 * Get the texture streaming texel factor and bound for this primitive.
+	 * Get the temp data used in the texture streaming build. Used to debug accuracy.
+	 * @param LODIndex			LOD index of the query. INDEX_NONE for a value for all LODs.
+	 * @param ElementIndex		Element index of the query. INDEX_NONE for a value for all elements.
+	 * @return					The section data built in the texture streaming build.
 	 */
-	virtual bool GetStreamingTextureInfo(struct FStreamingTexturePrimitiveInfo& OutInfo, int32 LODIndex, int32 ElementIndex) const { return false; }
+	virtual const FStreamingSectionBuildInfo* GetStreamingSectionData(int32 LODIndex, int32 ElementIndex) const { return nullptr; }
 
 protected:
 
@@ -609,9 +613,6 @@ protected:
 	/** True if the primitive casts static shadows. */
 	uint32 bCastStaticShadow : 1;
 
-	/** True if the primitive is rendered into the planar reflection pass. */
-	uint32 bVisibleInPlanarReflection : 1;
-
 	/** 
 	 * Whether the object should cast a volumetric translucent shadow.
 	 * Volumetric translucent shadows are useful for primitives with smoothly changing opacity like particles representing a volume, 
diff --git a/Engine/Source/Runtime/Engine/Public/SceneInterface.h b/Engine/Source/Runtime/Engine/Public/SceneInterface.h
--- a/Engine/Source/Runtime/Engine/Public/SceneInterface.h
+++ b/Engine/Source/Runtime/Engine/Public/SceneInterface.h
@@ -112,12 +112,17 @@ public:
 	/** Runs a slow preculling operation on static meshes, removing triangles that are invisible or inside a precull volume. */
 	virtual void PreCullStaticMeshes(const TArray<UStaticMeshComponent*>& ComponentsToPreCull, const TArray<TArray<FPlane> >& CullVolumes) {}
 
+	virtual void AddPlanarReflection(class UPlanarReflectionComponent* Component) {}
+	virtual void RemovePlanarReflection(class UPlanarReflectionComponent* Component) {}
+	virtual void UpdatePlanarReflectionTransform(UPlanarReflectionComponent* Component) {}
+
 	/** 
 	* Updates the contents of the given scene capture by rendering the scene. 
 	* This must be called on the game thread.
 	*/
 	virtual void UpdateSceneCaptureContents(class USceneCaptureComponent2D* CaptureComponent) {}
 	virtual void UpdateSceneCaptureContents(class USceneCaptureComponentCube* CaptureComponent) {}
+	virtual void UpdatePlanarReflectionContents(class UPlanarReflectionComponent* CaptureComponent, class FSceneRenderer& MainSceneRenderer) {}
 
 	virtual void AddPrecomputedLightVolume(const class FPrecomputedLightVolume* Volume) {}
 	virtual void RemovePrecomputedLightVolume(const class FPrecomputedLightVolume* Volume) {}
@@ -264,6 +269,9 @@ public:
 	{
 		return NULL;
 	}
+
+	virtual void UpdateSceneSettings(AWorldSettings* WorldSettings) {}
+
 	/**
 	 * Sets the FX system associated with the scene.
 	 */
diff --git a/Engine/Source/Runtime/Engine/Public/SceneView.h b/Engine/Source/Runtime/Engine/Public/SceneView.h
--- a/Engine/Source/Runtime/Engine/Public/SceneView.h
+++ b/Engine/Source/Runtime/Engine/Public/SceneView.h
@@ -13,6 +13,7 @@
 #include "RendererInterface.h"
 #include "RHIStaticStates.h"
 #include "GlobalDistanceFieldParameters.h"
+#include "DebugViewModeHelpers.h"
 
 class FSceneViewStateInterface;
 class FViewUniformShaderParameters;
@@ -355,28 +356,6 @@ enum ETranslucencyVolumeCascade
 	TVC_MAX,
 };
 
-/** 
- * Enumeration for different Quad Overdraw visualization mode.
- */
-enum EDebugViewShaderMode
-{
-	DVSM_None,						// No debug view.
-	DVSM_ShaderComplexity,			// Default shader complexity viewmode
-	DVSM_ShaderComplexityContainedQuadOverhead,	// Show shader complexity with quad overdraw scaling the PS instruction count.
-	DVSM_ShaderComplexityBleedingQuadOverhead,	// Show shader complexity with quad overdraw bleeding the PS instruction count over the quad.
-	DVSM_QuadComplexity,			// Show quad overdraw only.
-	DVSM_WantedMipsAccuracy,		// Accuraty of the wanted mips computed by the texture streamer.
-	DVSM_TexelFactorAccuracy,		// Accuraty of the texel factor computed on each mesh.
-	DVSM_TexCoordScaleAnalysis,		// To view the material texture coordinate scale used to sample each texture.
-	DVSM_MAX
-};
-
-// This defines for what shader platform the FDebugViewModeVS/DS/HS will be compiled.
-FORCEINLINE bool AllowDebugViewModeShader(EShaderPlatform Platform)
-{
-	return Platform == SP_PCD3D_SM5;
-}
-
 /** The view dependent uniform shader parameters associated with a view. */
 BEGIN_UNIFORM_BUFFER_STRUCT_WITH_CONSTRUCTOR(FViewUniformShaderParameters, ENGINE_API)
 	DECLARE_UNIFORM_BUFFER_STRUCT_MEMBER(FMatrix, TranslatedWorldToClip)
@@ -416,6 +395,7 @@ BEGIN_UNIFORM_BUFFER_STRUCT_WITH_CONSTRUCTOR(FViewUniformShaderParameters, ENGIN
 	DECLARE_UNIFORM_BUFFER_STRUCT_MEMBER(FMatrix, PrevInvViewProj)
 	DECLARE_UNIFORM_BUFFER_STRUCT_MEMBER(FMatrix, PrevScreenToTranslatedWorld)
 	DECLARE_UNIFORM_BUFFER_STRUCT_MEMBER(FMatrix, ClipToPrevClip)
+	DECLARE_UNIFORM_BUFFER_STRUCT_MEMBER(FVector4, GlobalClippingPlane)
 END_UNIFORM_BUFFER_STRUCT(FViewUniformShaderParameters)
 
 /** Copy of the view dependent uniform shader parameters associated with a view for instanced stereo. */
@@ -457,6 +437,7 @@ BEGIN_UNIFORM_BUFFER_STRUCT_WITH_CONSTRUCTOR(FInstancedViewUniformShaderParamete
 	DECLARE_UNIFORM_BUFFER_STRUCT_MEMBER(FMatrix, PrevInvViewProj)
 	DECLARE_UNIFORM_BUFFER_STRUCT_MEMBER(FMatrix, PrevScreenToTranslatedWorld)
 	DECLARE_UNIFORM_BUFFER_STRUCT_MEMBER(FMatrix, ClipToPrevClip)
+	DECLARE_UNIFORM_BUFFER_STRUCT_MEMBER(FVector4, GlobalClippingPlane)
 END_UNIFORM_BUFFER_STRUCT(FInstancedViewUniformShaderParameters)
 
 /** The view independent uniform shader parameters associated with a view. */
@@ -733,12 +714,9 @@ public:
 	/** Whether this view is being used to render a reflection capture. */
 	bool bIsReflectionCapture;
 
-	/** Whether this view is being used to render a planar reflection capture. */
-	bool bIsPlanarReflectionCapture;
+	/** Whether this view is being used to render a planar reflection. */
+	bool bIsPlanarReflection;
 
-	/** World space reflection plane normal and origin height. Used when rendering a planar reflection capture. */
-	FVector4 ReflectionPlane;
-	
 	/** Whether this view was created from a locked viewpoint. */
 	bool bIsLocked;
 
@@ -751,6 +729,9 @@ public:
 	/** True if instanced stereo is enabled. */
 	bool bIsInstancedStereoEnabled;
 
+	/** Global clipping plane being applied to the scene, or all 0's if disabled.  This is used when rendering the planar reflection pass. */
+	FPlane GlobalClippingPlane;
+
 	/** Aspect ratio constrained view rect. In the editor, when attached to a camera actor and the camera black bar showflag is enabled, the normal viewrect 
 	  * remains as the full viewport, and the black bars are just simulated by drawing black bars. This member stores the effective constrained area within the
 	  * bars.
diff --git a/Engine/Source/Runtime/Engine/Public/ShaderCompiler.h b/Engine/Source/Runtime/Engine/Public/ShaderCompiler.h
--- a/Engine/Source/Runtime/Engine/Public/ShaderCompiler.h
+++ b/Engine/Source/Runtime/Engine/Public/ShaderCompiler.h
@@ -278,12 +278,14 @@ struct FShaderMapCompileResults
 	FShaderMapCompileResults() :
 		NumJobsQueued(0),
 		bAllJobsSucceeded(true),
-		bApplyCompletedShaderMapForRendering(true)
+		bApplyCompletedShaderMapForRendering(true),
+		bRecreateComponentRenderStateOnCompletion(false)
 	{}
 
 	int32 NumJobsQueued;
 	bool bAllJobsSucceeded;
 	bool bApplyCompletedShaderMapForRendering;
+	bool bRecreateComponentRenderStateOnCompletion;
 	TArray<FShaderCommonCompileJob*> FinishedJobs;
 };
 
@@ -368,6 +370,8 @@ private:
 	FString ShaderCompileWorkerName;
 	/** Whether the SCW has crashed and we should fall back to calling the compiler dll's directly. */
 	bool bFallBackToDirectCompiles;
+	/** Whether a recreate should be done when compiling is finished. */
+	bool bRecreateComponentRenderStateOutstanding;
 
 	/** 
 	 * Tracks the total time that shader compile workers have been busy since startup.  
@@ -385,7 +389,7 @@ private:
 	void BlockOnAllShaderMapCompletion(TMap<int32, FShaderMapFinalizeResults>& CompiledShaderMaps);
 
 	/** Finalizes the given shader map results and optionally assigns the affected shader maps to materials, while attempting to stay within an execution time budget. */
-	void ProcessCompiledShaderMaps(TMap<int32, FShaderMapFinalizeResults>& CompiledShaderMaps, float TimeBudget);
+	void ProcessCompiledShaderMaps(TMap<int32, FShaderMapFinalizeResults>& CompiledShaderMaps, float TimeBudget, bool bRecreateComponentRenderState);
 
 	/** Recompiles shader jobs with errors if requested, and returns true if a retry was needed. */
 	bool HandlePotentialRetryOnError(TMap<int32, FShaderMapFinalizeResults>& CompletedShaderMaps);
@@ -437,7 +441,7 @@ public:
 	 * Adds shader jobs to be asynchronously compiled. 
 	 * FinishCompilation or ProcessAsyncResults must be used to get the results.
 	 */
-	ENGINE_API void AddJobs(TArray<FShaderCommonCompileJob*>& NewJobs, bool bApplyCompletedShaderMapForRendering, bool bOptimizeForLowLatency);
+	ENGINE_API void AddJobs(TArray<FShaderCommonCompileJob*>& NewJobs, bool bApplyCompletedShaderMapForRendering, bool bOptimizeForLowLatency, bool bRecreateComponentRenderStateOnCompletion);
 
 	/**
 	* Removes all outstanding compile jobs for the passed shader maps.
diff --git a/Engine/Source/Runtime/Engine/Public/ShowFlags.h b/Engine/Source/Runtime/Engine/Public/ShowFlags.h
--- a/Engine/Source/Runtime/Engine/Public/ShowFlags.h
+++ b/Engine/Source/Runtime/Engine/Public/ShowFlags.h
@@ -192,6 +192,8 @@ struct FEngineShowFlags
 		SetStereoRendering(false);
 		SetDistanceFieldAO(false);
 		SetDistanceFieldGI(false);
+		// hiding the SkinCache would break thumbnail caching
+		SetSkinCache(true);
 	}
 
 	// ---------------------------------------------------------
@@ -355,6 +357,7 @@ private:
 		SetWantedMipsAccuracy(false);
 		SetTexelFactorAccuracy(false);
 		SetTexCoordScaleAccuracy(false);
+		SetTexCoordAnalysis(false);
 	}
 
 
diff --git a/Engine/Source/Runtime/Engine/Public/ShowFlagsValues.inl b/Engine/Source/Runtime/Engine/Public/ShowFlagsValues.inl
--- a/Engine/Source/Runtime/Engine/Public/ShowFlagsValues.inl
+++ b/Engine/Source/Runtime/Engine/Public/ShowFlagsValues.inl
@@ -140,6 +140,8 @@ SHOWFLAG_ALWAYS_ACCESSIBLE(DynamicShadows, SFG_LightingComponents, LOCTEXT("Dyna
 SHOWFLAG_ALWAYS_ACCESSIBLE(Particles, SFG_Normal, LOCTEXT("ParticlesSF", "Particles Sprite"))
 /** if SkeletalMeshes are getting rendered, for now SHOWFLAG_ALWAYS_ACCESSIBLE because it's exposed in SceneCapture */
 SHOWFLAG_ALWAYS_ACCESSIBLE(SkeletalMeshes, SFG_Normal, LOCTEXT("SkeletalMeshesSF", "Skeletal Meshes"))
+/** if elements using the SkinCache are getting rendered (cache is still populated but not rendered, c) */
+SHOWFLAG_FIXED_IN_SHIPPING(1, SkinCache, SFG_Advanced, LOCTEXT("SkinCacheSF", "Skin Cache (GPU vertex cache for skeletal meshes)"))
 /** if the builder brush (editor) is getting rendered */
 SHOWFLAG_FIXED_IN_SHIPPING(0, BuilderBrush, SFG_Hidden, LOCTEXT("BuilderBrushSF", "Builder Brush"))
 /** Render translucency, for now SHOWFLAG_ALWAYS_ACCESSIBLE because it's exposed in SceneCapture */
@@ -304,5 +306,7 @@ SHOWFLAG_FIXED_IN_SHIPPING(0, WantedMipsAccuracy, SFG_Hidden, LOCTEXT("WantedMip
 SHOWFLAG_FIXED_IN_SHIPPING(0, TexelFactorAccuracy, SFG_Hidden, LOCTEXT("TexelFactorAccuracy", "Texel Factor Accuracy"))
 /** Visualize the accuracy of the material texture coordinate scale factors */
 SHOWFLAG_FIXED_IN_SHIPPING(0, TexCoordScaleAccuracy, SFG_Hidden, LOCTEXT("TexCoordScaleAccuray", "Texture Coordinate Scale Accuracy"))
+/** Outputs each texture lookup texcoord scale. Used in the texture streaming build. */
+SHOWFLAG_FIXED_IN_SHIPPING(0, TexCoordAnalysis, SFG_Hidden, LOCTEXT("TexCoordScaleAnalysis", "Texture Coordinate Scale Analysis"))
 #undef SHOWFLAG_ALWAYS_ACCESSIBLE
 #undef SHOWFLAG_FIXED_IN_SHIPPING
diff --git a/Engine/Source/Runtime/Engine/Public/SkeletalMeshTypes.h b/Engine/Source/Runtime/Engine/Public/SkeletalMeshTypes.h
--- a/Engine/Source/Runtime/Engine/Public/SkeletalMeshTypes.h
+++ b/Engine/Source/Runtime/Engine/Public/SkeletalMeshTypes.h
@@ -497,9 +497,13 @@ struct FSkelMeshSection
 
 	/** Is this mesh selected? */
 	uint8 bSelected:1;
+	
+	/** This section will recompute tangent in runtime */
+	bool bRecomputeTangent;
 
 	/** This Section can be disabled for cloth simulation and corresponding Cloth Section will be enabled*/
 	bool bDisabled;
+	
 	/** Corresponding Section Index will be enabled when this section is disabled 
 		because corresponding cloth section will be showed instead of this
 		or disabled section index when this section is enabled for cloth simulation
@@ -517,6 +521,7 @@ struct FSkelMeshSection
 		, NumTriangles(0)
 		, TriangleSorting(0)
 		, bSelected(false)
+		, bRecomputeTangent(false)
 		, bDisabled(false)
 		, CorrespondClothSectionIndex(-1)
 	{}
@@ -1680,18 +1685,22 @@ public:
 		return false;
 	}
 
+	// O(1)
 	// @return -1 if not found
 	uint32 FindChunkIndex(const FSkelMeshChunk& Chunk) const
 	{
-		uint32 Ret = -1;
+		const FSkelMeshChunk* Start = Chunks.GetData();
 
-		for(int32 ChunkIdx = 0; ChunkIdx < Chunks.Num(); ++ChunkIdx)
+		if(Start == nullptr)
 		{
-			if(&Chunks[ChunkIdx] == &Chunk)
+			return -1;
+		}
+
+		uint32 Ret = &Chunk - Start;
+
+		if(Ret >= (uint32)Chunks.Num())
 			{
-				Ret = ChunkIdx;
-				break;
-			}
+			Ret = -1;
 		}
 
 		return Ret;
diff --git a/Engine/Source/Runtime/Engine/Public/SkeletalRenderPublic.h b/Engine/Source/Runtime/Engine/Public/SkeletalRenderPublic.h
--- a/Engine/Source/Runtime/Engine/Public/SkeletalRenderPublic.h
+++ b/Engine/Source/Runtime/Engine/Public/SkeletalRenderPublic.h
@@ -16,7 +16,8 @@ public:
 	enum
 	{
 		// Up to this number of chunks, we can use on the GPU skinning cache; no point in using TArray as it's using int16 elements
-		MAX_GPUSKINCACHE_CHUNKS_PER_LOD = 16,
+		// 32 is the conservative size needed for Senua in CharDemo with NinjaTheory (this should be a cvar and project specific or dynamic, with only using the feature on RecomputeTangents it can be much less)
+		MAX_GPUSKINCACHE_CHUNKS_PER_LOD = 32,
 	};
 
 	FSkeletalMeshObject(USkinnedMeshComponent* InMeshComponent, FSkeletalMeshResource* InSkeletalMeshResource, ERHIFeatureLevel::Type FeatureLevel);
@@ -40,6 +41,13 @@ public:
 	 */
 	virtual void Update(int32 LODIndex,USkinnedMeshComponent* InMeshComponent,const TArray<FActiveVertexAnim>& ActiveVertexAnims) = 0;
 
+	/**
+	* Called by the game thread for any update on RecomputeTangent
+	* @param	MaterialIndex : Material Index for the update
+	* @param	bRecomputeTangent : the new flag
+	*/
+	virtual void UpdateRecomputeTangent(int32 MaterialIndex, bool bRecomputeTangent) = 0;
+
 	/**
 	 * Called by FSkeletalMeshObject prior to GDME. This allows the GPU skin version to update bones etc now that we know we are going to render
 	 * @param FrameNumber from GFrameNumber
@@ -49,11 +57,12 @@ public:
 	}
 
 	/**
+	 * @param	View - View, must not be 0, allows to cull depending on showflags (normally not needed/used in SHIPPING)
 	 * @param	LODIndex - each LOD has its own vertex data
 	 * @param	ChunkIdx - not used
-	 * @return	vertex factory for rendering the LOD
+	 * @return	vertex factory for rendering the LOD, 0 to suppress rendering
 	 */
-	virtual const FVertexFactory* GetVertexFactory(int32 LODIndex,int32 ChunkIdx) const = 0;
+	virtual const FVertexFactory* GetSkinVertexFactory(const FSceneView* View, int32 LODIndex,int32 ChunkIdx) const = 0;
 
 	/**
 	 * Re-skin cached vertices for an LOD and update the vertex buffer. Note that this
diff --git a/Engine/Source/Runtime/Engine/Public/StaticMeshResources.h b/Engine/Source/Runtime/Engine/Public/StaticMeshResources.h
--- a/Engine/Source/Runtime/Engine/Public/StaticMeshResources.h
+++ b/Engine/Source/Runtime/Engine/Public/StaticMeshResources.h
@@ -786,8 +786,8 @@ public:
 	/** Sets up a wireframe FMeshBatch for a specific LOD. */
 	virtual bool GetWireframeMeshElement(int32 LODIndex, int32 BatchIndex, const FMaterialRenderProxy* WireframeRenderProxy, uint8 InDepthPriorityGroup, bool bAllowPreCulledIndices, FMeshBatch& OutMeshBatch) const;
 
-#if !(UE_BUILD_SHIPPING || UE_BUILD_TEST)
-	virtual bool GetStreamingTextureInfo(struct FStreamingTexturePrimitiveInfo& OutInfo, int32 LODIndex, int32 ElementIndex) const override;
+#if WITH_EDITORONLY_DATA
+	virtual const FStreamingSectionBuildInfo* GetStreamingSectionData(int32 LODIndex, int32 ElementIndex) const override;
 #endif
 
 protected:
@@ -912,17 +912,14 @@ protected:
 	FCollisionResponseContainer CollisionResponse;
 
 #if WITH_EDITORONLY_DATA
+	/** Data shared with the component */
+	TSharedPtr<TArray<FStreamingSectionBuildInfo>, ESPMode::NotThreadSafe> StreamingSectionData;
+
 	/** Index of the section to preview. If set to INDEX_NONE, all section will be rendered */
 	int32 SectionIndexPreview;
 #endif
 
 #if !(UE_BUILD_SHIPPING || UE_BUILD_TEST)
-	/** The precise streaming info for each LOD and Section. */
-	TArray<struct FStreamingTexturePrimitiveInfo> StreamingTextureInfos;
-	/** Texture streaming factor used for lightmap textures */
-	float WorldTexelFactor;
-	float WorldLightmapFactor;
-
 	/** LOD used for collision */
 	int32 LODForCollision;
 	/** If we want to draw the mesh collision for debugging */
diff --git a/Engine/Source/Runtime/Landscape/Private/LandscapeRender.cpp b/Engine/Source/Runtime/Landscape/Private/LandscapeRender.cpp
--- a/Engine/Source/Runtime/Landscape/Private/LandscapeRender.cpp
+++ b/Engine/Source/Runtime/Landscape/Private/LandscapeRender.cpp
@@ -2548,7 +2548,7 @@ ULandscapeMaterialInstanceConstant::ULandscapeMaterialInstanceConstant(const FOb
 	bIsLayerThumbnail = false;
 }
 
-void ULandscapeComponent::GetStreamingTextureInfo(TArray<FStreamingTexturePrimitiveInfo>& OutStreamingTextures) const
+void ULandscapeComponent::GetStreamingTextureInfo(FStreamingTextureLevelContext& LevelContext, TArray<FStreamingTexturePrimitiveInfo>& OutStreamingTextures) const
 {
 	ALandscapeProxy* Proxy = Cast<ALandscapeProxy>(GetOuter());
 	FSphere BoundingSphere = Bounds.GetSphere();
diff --git a/Engine/Source/Runtime/Renderer/Private/BasePassRendering.cpp b/Engine/Source/Runtime/Renderer/Private/BasePassRendering.cpp
--- a/Engine/Source/Runtime/Renderer/Private/BasePassRendering.cpp
+++ b/Engine/Source/Runtime/Renderer/Private/BasePassRendering.cpp
@@ -6,6 +6,7 @@
 
 #include "RendererPrivate.h"
 #include "ScenePrivate.h"
+#include "PlanarReflectionSceneProxy.h"
 
 // Changing this causes a full shader recompile
 static TAutoConsoleVariable<int32> CVarSelectiveBasePassOutputs(
@@ -16,6 +17,12 @@ static TAutoConsoleVariable<int32> CVarSelectiveBasePassOutputs(
 	TEXT(" 1: Export only into relevant rendertarget.\n"),
 	ECVF_ReadOnly | ECVF_RenderThreadSafe);
 
+// Changing this causes a full shader recompile
+static TAutoConsoleVariable<int32> CVarGlobalClipPlane(
+	TEXT("r.AllowGlobalClipPlane"),
+	0,
+	TEXT("Enables mesh shaders to support a global clip plane, needed for planar reflections, which adds about 15% BasePass GPU cost on PS4."),
+	ECVF_ReadOnly | ECVF_RenderThreadSafe);
 
 bool UseSelectiveBasePassOutputs()
 {
@@ -114,20 +121,49 @@ void FSkyLightReflectionParameters::GetSkyParametersFromScene(
 	}
 }
 
-void FReflectionParameters::Set(FRHICommandList& RHICmdList, FShader* Shader, const FViewInfo* View)
+void FPlanarReflectionParameters::SetParameters(FRHICommandList& RHICmdList, FPixelShaderRHIParamRef ShaderRHI, const FPlanarReflectionSceneProxy* ReflectionSceneProxy)
+{
+	// Degenerate plane causes shader to branch around the reflection lookup
+	FPlane ReflectionPlaneValue(FVector4(0, 0, 0, 0));
+	FTexture* PlanarReflectionTextureValue = GBlackTexture;
+
+	if (ReflectionSceneProxy)
+	{
+		ReflectionPlaneValue = ReflectionSceneProxy->ReflectionPlane;
+		PlanarReflectionTextureValue = ReflectionSceneProxy->RenderTarget;
+			
+		SetShaderValue(RHICmdList, ShaderRHI, InverseTransposeMirrorMatrix, ReflectionSceneProxy->InverseTransposeMirrorMatrix);
+		SetShaderValue(RHICmdList, ShaderRHI, PlanarReflectionParameters, ReflectionSceneProxy->PlanarReflectionParameters);
+		SetShaderValue(RHICmdList, ShaderRHI, PlanarReflectionParameters2, ReflectionSceneProxy->PlanarReflectionParameters2);
+		SetShaderValue(RHICmdList, ShaderRHI, ProjectionWithExtraFOV, ReflectionSceneProxy->ProjectionWithExtraFOV);
+	}
+
+	SetShaderValue(RHICmdList, ShaderRHI, ReflectionPlane, ReflectionPlaneValue);
+	SetTextureParameter(RHICmdList, ShaderRHI, PlanarReflectionTexture, PlanarReflectionSampler, PlanarReflectionTextureValue);
+}
+
+void FBasePassReflectionParameters::Set(FRHICommandList& RHICmdList, FPixelShaderRHIParamRef PixelShaderRHI, const FViewInfo* View)
 {
 	FSceneRenderTargets& SceneContext = FSceneRenderTargets::Get(RHICmdList);
-	auto PixelShader = Shader->GetPixelShader();
 
-	SkyLightReflectionParameters.SetParameters(RHICmdList, Shader->GetPixelShader(), (const FScene*)(View->Family->Scene), true);
+	SkyLightReflectionParameters.SetParameters(RHICmdList, PixelShaderRHI, (const FScene*)(View->Family->Scene), true);
 }
 
-void FReflectionParameters::SetMesh(FRHICommandList& RHICmdList, FShader* Shader, const FPrimitiveSceneProxy* Proxy, ERHIFeatureLevel::Type FeatureLevel)
+void FBasePassReflectionParameters::SetMesh(FRHICommandList& RHICmdList, FPixelShaderRHIParamRef PixelShaderRHI, const FPrimitiveSceneProxy* Proxy, ERHIFeatureLevel::Type FeatureLevel)
 {
+	const FPrimitiveSceneInfo* PrimitiveSceneInfo = Proxy ? Proxy->GetPrimitiveSceneInfo() : NULL;
+	const FPlanarReflectionSceneProxy* ReflectionSceneProxy = NULL;
+
+	if (PrimitiveSceneInfo && PrimitiveSceneInfo->CachedPlanarReflectionProxy)
+	{
+		ReflectionSceneProxy = PrimitiveSceneInfo->CachedPlanarReflectionProxy;
+	}
+
+	PlanarReflectionParameters.SetParameters(RHICmdList, PixelShaderRHI, ReflectionSceneProxy);
+
 	// Note: GBlackCubeArrayTexture has an alpha of 0, which is needed to represent invalid data so the sky cubemap can still be applied
 	FTextureRHIParamRef CubeArrayTexture = FeatureLevel >= ERHIFeatureLevel::SM5 ? GBlackCubeArrayTexture->TextureRHI : GBlackTextureCube->TextureRHI;
 	int32 ArrayIndex = 0;
-	const FPrimitiveSceneInfo* PrimitiveSceneInfo = Proxy ? Proxy->GetPrimitiveSceneInfo() : NULL;
 
 	if (PrimitiveSceneInfo && PrimitiveSceneInfo->CachedReflectionCaptureProxy)
 	{
@@ -136,27 +172,26 @@ void FReflectionParameters::SetMesh(FRHICommandList& RHICmdList, FShader* Shader
 
 	SetTextureParameter(
 		RHICmdList, 
-		Shader->GetPixelShader(), 
+		PixelShaderRHI, 
 		ReflectionCubemap, 
 		ReflectionCubemapSampler, 
 		TStaticSamplerState<SF_Trilinear, AM_Clamp, AM_Clamp, AM_Clamp>::GetRHI(), 
 		CubeArrayTexture);
 
-	SetShaderValue(RHICmdList, Shader->GetPixelShader(), CubemapArrayIndex, ArrayIndex);
+	SetShaderValue(RHICmdList, PixelShaderRHI, CubemapArrayIndex, ArrayIndex);
 }
 
-void FTranslucentLightingParameters::Set(FRHICommandList& RHICmdList, FShader* Shader, const FViewInfo* View)
+void FTranslucentLightingParameters::Set(FRHICommandList& RHICmdList, FPixelShaderRHIParamRef PixelShaderRHI, const FViewInfo* View)
 {
 	FSceneRenderTargets& SceneContext = FSceneRenderTargets::Get(RHICmdList);
-	auto PixelShader = Shader->GetPixelShader();
 
-	TranslucentLightingVolumeParameters.Set(RHICmdList, PixelShader);
+	TranslucentLightingVolumeParameters.Set(RHICmdList, PixelShaderRHI);
 
 	if (View->HZB)
 	{
 		SetTextureParameter(
 			RHICmdList, 
-			PixelShader, 
+			PixelShaderRHI, 
 			HZBTexture, 
 			HZBSampler, 
 			TStaticSamplerState<SF_Point, AM_Clamp, AM_Clamp, AM_Clamp>::GetRHI(), 
@@ -172,7 +207,7 @@ void FTranslucentLightingParameters::Set(FRHICommandList& RHICmdList, FShader* S
 
 		SetTextureParameter(
 			RHICmdList, 
-			PixelShader, 
+			PixelShaderRHI, 
 			PrevSceneColor, 
 			PrevSceneColorSampler, 
 			TStaticSamplerState<SF_Bilinear, AM_Clamp, AM_Clamp, AM_Clamp>::GetRHI(), 
@@ -189,14 +224,14 @@ void FTranslucentLightingParameters::Set(FRHICommandList& RHICmdList, FShader* S
 			1.0f / HZBUvFactor.Y
 			);
 			
-		SetShaderValue(RHICmdList, PixelShader, HZBUvFactorAndInvFactor, HZBUvFactorAndInvFactorValue);
+		SetShaderValue(RHICmdList, PixelShaderRHI, HZBUvFactorAndInvFactor, HZBUvFactorAndInvFactorValue);
 	}
 	else
 	{
 		//set dummies for platforms that require bound resources.
 		SetTextureParameter(
 			RHICmdList,
-			PixelShader,
+			PixelShaderRHI,
 			HZBTexture,
 			HZBSampler,
 			TStaticSamplerState<SF_Point, AM_Clamp, AM_Clamp, AM_Clamp>::GetRHI(),
@@ -204,7 +239,7 @@ void FTranslucentLightingParameters::Set(FRHICommandList& RHICmdList, FShader* S
 
 		SetTextureParameter(
 			RHICmdList,
-			PixelShader,
+			PixelShaderRHI,
 			PrevSceneColor,
 			PrevSceneColorSampler,
 			TStaticSamplerState<SF_Bilinear, AM_Clamp, AM_Clamp, AM_Clamp>::GetRHI(),
diff --git a/Engine/Source/Runtime/Renderer/Private/BasePassRendering.h b/Engine/Source/Runtime/Renderer/Private/BasePassRendering.h
--- a/Engine/Source/Runtime/Renderer/Private/BasePassRendering.h
+++ b/Engine/Source/Runtime/Renderer/Private/BasePassRendering.h
@@ -388,26 +388,70 @@ private:
 		float& OutBlendFraction);
 };
 
+/** Parameters needed for planar reflections, shared by multiple shaders. */
+class FPlanarReflectionParameters
+{
+public:
+
+	void Bind(const FShaderParameterMap& ParameterMap)
+	{
+		ReflectionPlane.Bind(ParameterMap, TEXT("ReflectionPlane"));
+		InverseTransposeMirrorMatrix.Bind(ParameterMap, TEXT("InverseTransposeMirrorMatrix"));
+		PlanarReflectionParameters.Bind(ParameterMap, TEXT("PlanarReflectionParameters"));
+		PlanarReflectionParameters2.Bind(ParameterMap, TEXT("PlanarReflectionParameters2"));
+		ProjectionWithExtraFOV.Bind(ParameterMap, TEXT("ProjectionWithExtraFOV"));
+		PlanarReflectionTexture.Bind(ParameterMap, TEXT("PlanarReflectionTexture"));
+		PlanarReflectionSampler.Bind(ParameterMap, TEXT("PlanarReflectionSampler"));
+	}
+
+	void SetParameters(FRHICommandList& RHICmdList, FPixelShaderRHIParamRef ShaderRHI, const class FPlanarReflectionSceneProxy* ReflectionSceneProxy);
+
+	/** Serializer. */
+	friend FArchive& operator<<(FArchive& Ar,FPlanarReflectionParameters& P)
+	{
+		Ar << P.ReflectionPlane;
+		Ar << P.InverseTransposeMirrorMatrix;
+		Ar << P.PlanarReflectionParameters;
+		Ar << P.PlanarReflectionParameters2;
+		Ar << P.ProjectionWithExtraFOV;
+		Ar << P.PlanarReflectionTexture;
+		Ar << P.PlanarReflectionSampler;
+		return Ar;
+	}
+
+private:
+
+	FShaderParameter ReflectionPlane;
+	FShaderParameter InverseTransposeMirrorMatrix;
+	FShaderParameter PlanarReflectionParameters;
+	FShaderParameter PlanarReflectionParameters2;
+	FShaderParameter ProjectionWithExtraFOV;
+	FShaderResourceParameter PlanarReflectionTexture;
+	FShaderResourceParameter PlanarReflectionSampler;
+};
+
 /** Parameters needed for reflections, shared by multiple shaders. */
-class FReflectionParameters
+class FBasePassReflectionParameters
 {
 public:
 
 	void Bind(const FShaderParameterMap& ParameterMap)
 	{
+		PlanarReflectionParameters.Bind(ParameterMap);
 		ReflectionCubemap.Bind(ParameterMap, TEXT("ReflectionCubemap"));
 		ReflectionCubemapSampler.Bind(ParameterMap, TEXT("ReflectionCubemapSampler"));
 		CubemapArrayIndex.Bind(ParameterMap, TEXT("CubemapArrayIndex"));
 		SkyLightReflectionParameters.Bind(ParameterMap);
 	}
 
-	void Set(FRHICommandList& RHICmdList, FShader* Shader, const FViewInfo* View);
+	void Set(FRHICommandList& RHICmdList, FPixelShaderRHIParamRef PixelShaderRHI, const FViewInfo* View);
 
-	void SetMesh(FRHICommandList& RHICmdList, FShader* Shader, const FPrimitiveSceneProxy* Proxy, ERHIFeatureLevel::Type FeatureLevel);
+	void SetMesh(FRHICommandList& RHICmdList, FPixelShaderRHIParamRef PixelShaderRH, const FPrimitiveSceneProxy* Proxy, ERHIFeatureLevel::Type FeatureLevel);
 
 	/** Serializer. */
-	friend FArchive& operator<<(FArchive& Ar,FReflectionParameters& P)
+	friend FArchive& operator<<(FArchive& Ar,FBasePassReflectionParameters& P)
 	{
+		Ar << P.PlanarReflectionParameters;
 		Ar << P.ReflectionCubemap;
 		Ar << P.ReflectionCubemapSampler;
 		Ar << P.CubemapArrayIndex;
@@ -417,6 +461,7 @@ public:
 
 private:
 
+	FPlanarReflectionParameters PlanarReflectionParameters;
 	FShaderResourceParameter ReflectionCubemap;
 	FShaderResourceParameter ReflectionCubemapSampler;
 	FShaderParameter CubemapArrayIndex;
@@ -438,7 +483,7 @@ public:
 		PrevSceneColorSampler.Bind(ParameterMap, TEXT("PrevSceneColorSampler"));
 	}
 
-	void Set(FRHICommandList& RHICmdList, FShader* Shader, const FViewInfo* View);
+	void Set(FRHICommandList& RHICmdList, FPixelShaderRHIParamRef PixelShaderRHI, const FViewInfo* View);
 
 	/** Serializer. */
 	friend FArchive& operator<<(FArchive& Ar,FTranslucentLightingParameters& P)
@@ -518,12 +563,12 @@ public:
 
 		if (View->GetFeatureLevel() >= ERHIFeatureLevel::SM4)
 		{
-			ReflectionParameters.Set(RHICmdList, this, View);
+			ReflectionParameters.Set(RHICmdList, ShaderRHI, View);
 
 			if (IsTranslucentBlendMode(BlendMode))
 			{
 				SetShaderValue(RHICmdList, ShaderRHI, ScreenTextureUVScale, FVector(ScreenTextureScaleFactor));
-				TranslucentLightingParameters.Set(RHICmdList, this, View);
+				TranslucentLightingParameters.Set(RHICmdList, ShaderRHI, View);
 
 				// Experimental dynamic forward lighting for translucency. Can be the base for opaque forward lighting which will allow more lighting models or rendering without a GBuffer
 				if(BlendMode == BLEND_Translucent && MaterialResource.GetTranslucencyLightingMode() == TLM_SurfacePerPixelLighting)
@@ -554,7 +599,7 @@ public:
 	}
 
 private:
-	FReflectionParameters ReflectionParameters;
+	FBasePassReflectionParameters ReflectionParameters;
 	FTranslucentLightingParameters TranslucentLightingParameters;
 	FEditorCompositingParameters EditorCompositeParams;
 	FShaderResourceParameter LightGrid;
@@ -782,14 +827,12 @@ public:
 
 	void SetSharedState(FRHICommandList& RHICmdList, const FViewInfo* View, const ContextDataType PolicyContext, float ScreenTextureScaleFactor = 1.0f) const
 	{
-#if !(UE_BUILD_SHIPPING || UE_BUILD_TEST)
-		// If debug view shader modes are allowed, different VS/DS/HS must be used (with only SV_POSITION as PS interpolant).
-		if (GetDebugViewShaderMode() != DVSM_None && RuntimeAllowDebugViewModeShader(View->GetFeatureLevel()))
+		// If the current debug view shader modes are allowed, different VS/DS/HS must be used (with only SV_POSITION as PS interpolant).
+		if (AllowDebugViewShaderMode(GetDebugViewShaderMode(), View->GetFeatureLevel()) && AllowDebugViewVSDSHS(View->GetFeatureLevel()))
 		{
 			FDebugViewMode::SetParametersVSHSDS(RHICmdList, MaterialRenderProxy, MaterialResource, *View, VertexFactory, HullShader && DomainShader);
 		}
 		else
-#endif
 		{
 			// Set the light-map policy.
 			LightMapPolicy.Set(RHICmdList, VertexShader, GetDebugViewShaderMode() == DVSM_None ? PixelShader : nullptr, VertexShader, PixelShader, VertexFactory, MaterialRenderProxy, View);
@@ -871,20 +914,10 @@ public:
 			FGeometryShaderRHIRef()
 			);
 
-#if !(UE_BUILD_SHIPPING || UE_BUILD_TEST)
-		if (GetDebugViewShaderMode() != DVSM_None)
+		if (AllowDebugViewShaderMode(GetDebugViewShaderMode(), InFeatureLevel))
 		{
-			if (RuntimeAllowDebugViewModeShader(InFeatureLevel))
-			{
-				FDebugViewMode::PatchBoundShaderState(BoundShaderStateInput, MaterialResource, VertexFactory, InFeatureLevel, GetDebugViewShaderMode());
-			}
-			else // Otherwise only shader complexity is available
-			{
-				TShaderMapRef<TShaderComplexityAccumulatePS> ShaderComplexityAccumulatePixelShader(GetGlobalShaderMap(InFeatureLevel));
-				BoundShaderStateInput.PixelShaderRHI = ShaderComplexityAccumulatePixelShader->FGlobalShader::GetPixelShader();
-			}
+			FDebugViewMode::PatchBoundShaderState(BoundShaderStateInput, MaterialResource, VertexFactory, InFeatureLevel, GetDebugViewShaderMode());
 		}
-#endif
 		return BoundShaderStateInput;
 	}
 
@@ -902,14 +935,12 @@ public:
 	{
 		const FMeshBatchElement& BatchElement = Mesh.Elements[BatchElementIndex];
 
-#if !(UE_BUILD_SHIPPING || UE_BUILD_TEST)
 		// If debug view shader mode are allowed, different VS/DS/HS must be used (with only SV_POSITION as PS interpolant).
-		if (GetDebugViewShaderMode() != DVSM_None && RuntimeAllowDebugViewModeShader(View.GetFeatureLevel()))
+		if (AllowDebugViewShaderMode(GetDebugViewShaderMode(), View.GetFeatureLevel()) && AllowDebugViewVSDSHS(View.GetFeatureLevel()))
 		{
 			FDebugViewMode::SetMeshVSHSDS(RHICmdList, VertexFactory, View, PrimitiveSceneProxy, BatchElement, DrawRenderState, MaterialResource, HullShader && DomainShader);
 		}
 		else
-#endif
 		{
 			// Set the light-map policy's mesh-specific settings.
 			LightMapPolicy.SetMesh(
diff --git a/Engine/Source/Runtime/Renderer/Private/BasePassRendering.inl b/Engine/Source/Runtime/Renderer/Private/BasePassRendering.inl
--- a/Engine/Source/Runtime/Renderer/Private/BasePassRendering.inl
+++ b/Engine/Source/Runtime/Renderer/Private/BasePassRendering.inl
@@ -59,7 +59,7 @@ void TBasePassPixelShaderPolicyParamType<PixelParametersType>::SetMesh(FRHIComma
 {
 	if (View.GetFeatureLevel() >= ERHIFeatureLevel::SM4)
 	{
-		ReflectionParameters.SetMesh(RHICmdList, this, Proxy, View.GetFeatureLevel());
+		ReflectionParameters.SetMesh(RHICmdList, GetPixelShader(), Proxy, View.GetFeatureLevel());
 	}
 
 	FMeshMaterialShader::SetMesh(RHICmdList, GetPixelShader(), VertexFactory, View, Proxy, BatchElement, DrawRenderState);
diff --git a/Engine/Source/Runtime/Renderer/Private/DebugViewModeRendering.cpp b/Engine/Source/Runtime/Renderer/Private/DebugViewModeRendering.cpp
--- a/Engine/Source/Runtime/Renderer/Private/DebugViewModeRendering.cpp
+++ b/Engine/Source/Runtime/Renderer/Private/DebugViewModeRendering.cpp
@@ -26,7 +26,7 @@ class FMissingShaderPS : public FGlobalShader, public IDebugViewModePSInterface
 
 public:
 
-	static bool ShouldCache(EShaderPlatform Platform) { return AllowDebugViewModeShader(Platform); }
+	static bool ShouldCache(EShaderPlatform Platform) { return AllowDebugViewVSDSHS(Platform); }
 
 	FMissingShaderPS(const ShaderMetaType::CompiledShaderInitializerType& Initializer):	FGlobalShader(Initializer) {}
 	FMissingShaderPS() {}
@@ -62,8 +62,11 @@ public:
 IMPLEMENT_SHADER_TYPE(,FMissingShaderPS,TEXT("MissingShaderPixelShader"),TEXT("Main"),SF_Pixel);
 
 
-void FDebugViewMode::GetDebugMaterial(const FMaterialRenderProxy** MaterialRenderProxy, const FMaterial** Material, ERHIFeatureLevel::Type FeatureLevel)
+void FDebugViewMode::GetMaterialForVSHSDS(const FMaterialRenderProxy** MaterialRenderProxy, const FMaterial** Material, ERHIFeatureLevel::Type FeatureLevel)
 {
+	check(AllowDebugViewVSDSHS(FeatureLevel)); // Shouldn't be called otherwise.
+
+	// If the material was compiled fo VS, return it, otherwise, return the default material.
 	if (!(*Material)->HasVertexPositionOffsetConnected() && (*Material)->GetTessellationMode() == MTM_NoTessellation)
 	{
 		if (MaterialRenderProxy)
@@ -95,6 +98,7 @@ IDebugViewModePSInterface* FDebugViewMode::GetPSInterface(TShaderMap<FGlobalShad
 		return *TShaderMapRef<FWantedMipsAccuracyPS>(ShaderMap);
 	case DVSM_TexelFactorAccuracy:
 		return *TShaderMapRef<FTexelFactorAccuracyPS>(ShaderMap);
+	case DVSM_TexCoordScaleAccuracy:
 	case DVSM_TexCoordScaleAnalysis:
 	{
 		const FMaterial* MaterialForPS = GetDebugViewMaterialPS(DebugViewShaderMode, Material);
@@ -119,25 +123,24 @@ void FDebugViewMode::PatchBoundShaderState(
 	EDebugViewShaderMode DebugViewShaderMode
 	)
 {
-	const FMaterial* MaterialForPS = Material;
-	GetDebugMaterial(nullptr, &Material, FeatureLevel);
+	const FMaterial* MaterialForPS = Material; // Backup before calling GetMaterialForVSHSDS
 
-	if (!Material->HasVertexPositionOffsetConnected() && Material->GetTessellationMode() == MTM_NoTessellation)
+	if (AllowDebugViewVSDSHS(FeatureLevel))
 	{
-		Material = UMaterial::GetDefaultMaterial(MD_Surface)->GetRenderProxy(false)->GetMaterial(FeatureLevel);
-	}
+		GetMaterialForVSHSDS(nullptr, &Material, FeatureLevel);
 
-	FVertexFactoryType* VertexFactoryType = VertexFactory->GetType();
+		FVertexFactoryType* VertexFactoryType = VertexFactory->GetType();
 
-	BoundShaderStateInput.VertexShaderRHI = Material->GetShader<FDebugViewModeVS>(VertexFactoryType)->GetVertexShader();
+		BoundShaderStateInput.VertexShaderRHI = Material->GetShader<FDebugViewModeVS>(VertexFactoryType)->GetVertexShader();
 
-	if (BoundShaderStateInput.HullShaderRHI)
-	{
-		BoundShaderStateInput.HullShaderRHI = Material->GetShader<FDebugViewModeHS>(VertexFactoryType)->GetHullShader();
-	}
-	if (BoundShaderStateInput.DomainShaderRHI)
-	{
-		BoundShaderStateInput.DomainShaderRHI = Material->GetShader<FDebugViewModeDS>(VertexFactoryType)->GetDomainShader();
+		if (BoundShaderStateInput.HullShaderRHI)
+		{
+			BoundShaderStateInput.HullShaderRHI = Material->GetShader<FDebugViewModeHS>(VertexFactoryType)->GetHullShader();
+		}
+		if (BoundShaderStateInput.DomainShaderRHI)
+		{
+			BoundShaderStateInput.DomainShaderRHI = Material->GetShader<FDebugViewModeDS>(VertexFactoryType)->GetDomainShader();
+		}
 	}
 
 	BoundShaderStateInput.PixelShaderRHI = GetPSInterface(GetGlobalShaderMap(FeatureLevel), MaterialForPS, DebugViewShaderMode)->GetShader()->GetPixelShader();
@@ -154,7 +157,7 @@ void FDebugViewMode::SetParametersVSHSDS(
 {
 	VertexFactory->Set(RHICmdList);
 
-	GetDebugMaterial(&MaterialRenderProxy, &Material, View.GetFeatureLevel());
+	GetMaterialForVSHSDS(&MaterialRenderProxy, &Material, View.GetFeatureLevel());
 
 	FVertexFactoryType* VertexFactoryType = VertexFactory->GetType();
 
@@ -178,7 +181,7 @@ void FDebugViewMode::SetMeshVSHSDS(
 	bool bHasHullAndDomainShader
 	)
 {
-	GetDebugMaterial(nullptr, &Material, View.GetFeatureLevel());
+	GetMaterialForVSHSDS(nullptr, &Material, View.GetFeatureLevel());
 
 	FVertexFactoryType* VertexFactoryType = VertexFactory->GetType();
 
diff --git a/Engine/Source/Runtime/Renderer/Private/DebugViewModeRendering.h b/Engine/Source/Runtime/Renderer/Private/DebugViewModeRendering.h
--- a/Engine/Source/Runtime/Renderer/Private/DebugViewModeRendering.h
+++ b/Engine/Source/Runtime/Renderer/Private/DebugViewModeRendering.h
@@ -6,18 +6,7 @@ DebugViewModeRendering.h: Contains definitions for rendering debug viewmodes.
 
 #pragma once
 
-// Used to know whether debug viewmodes are allowed at runtime.
-#if !(UE_BUILD_SHIPPING || UE_BUILD_TEST)
-FORCEINLINE bool RuntimeAllowDebugViewModeShader(ERHIFeatureLevel::Type InFeatureLevel)
-{
-	return AllowDebugViewModeShader(GetFeatureLevelShaderPlatform(InFeatureLevel));
-}
-#else
-FORCEINLINE bool RuntimeAllowDebugViewModeShader(ERHIFeatureLevel::Type InFeatureLevel)
-{
-	return false;
-}
-#endif
+#include "DebugViewModeHelpers.h"
 
 static const int32 NumStreamingAccuracyColors = 5;
 static const int32 MaxStreamingAccuracyMips = 11;
@@ -37,7 +26,7 @@ public:
 
 	static bool ShouldCache(EShaderPlatform Platform,const FMaterial* Material,const FVertexFactoryType* VertexFactoryType)
 	{
-		return AllowDebugViewModeShader(Platform) && (Material->IsSpecialEngineMaterial() || Material->HasVertexPositionOffsetConnected() || Material->GetTessellationMode() != MTM_NoTessellation);
+		return AllowDebugViewVSDSHS(Platform) && (Material->IsDefaultMaterial() || Material->HasVertexPositionOffsetConnected() || Material->GetTessellationMode() != MTM_NoTessellation);
 	}
 
 	void SetParameters(FRHICommandList& RHICmdList, const FMaterialRenderProxy* MaterialRenderProxy,const FMaterial& Material,const FSceneView& View)
@@ -124,7 +113,7 @@ struct FDebugViewMode
 {
 	static IDebugViewModePSInterface* GetPSInterface(TShaderMap<FGlobalShaderType>* ShaderMap, const FMaterial* Material, EDebugViewShaderMode DebugViewShaderMode);
 
-	static void GetDebugMaterial(const FMaterialRenderProxy** MaterialRenderProxy, const FMaterial** Material, ERHIFeatureLevel::Type FeatureLevel);
+	static void GetMaterialForVSHSDS(const FMaterialRenderProxy** MaterialRenderProxy, const FMaterial** Material, ERHIFeatureLevel::Type FeatureLevel);
 
 	static void SetParametersVSHSDS(
 		FRHICommandList& RHICmdList, 
diff --git a/Engine/Source/Runtime/Renderer/Private/DeferredShadingRenderer.cpp b/Engine/Source/Runtime/Renderer/Private/DeferredShadingRenderer.cpp
--- a/Engine/Source/Runtime/Renderer/Private/DeferredShadingRenderer.cpp
+++ b/Engine/Source/Runtime/Renderer/Private/DeferredShadingRenderer.cpp
@@ -797,7 +797,7 @@ void FDeferredShadingSceneRenderer::RenderOcclusion(FRHICommandListImmediate& RH
 				FViewInfo& View = Views[ViewIndex];
 				FSceneViewState* ViewState = (FSceneViewState*)View.State;
 				
-				const uint32 bSSR = DoScreenSpaceReflections( View );
+				const uint32 bSSR = ShouldRenderScreenSpaceReflections( View );
 				
 				if (bSSAO || bHZBOcclusion || bSSR)
 				{
@@ -931,7 +931,7 @@ void FDeferredShadingSceneRenderer::Render(FRHICommandListImmediate& RHICmdList)
 			if (UseGlobalDistanceField())
 			{
 				// Use the skylight's max distance if there is one
-				const float OcclusionMaxDistance = Scene->SkyLight && !Scene->SkyLight->bWantsStaticShadowing ? Scene->SkyLight->OcclusionMaxDistance : GDefaultDFAOMaxOcclusionDistance;
+				const float OcclusionMaxDistance = Scene->SkyLight && !Scene->SkyLight->bWantsStaticShadowing ? Scene->SkyLight->OcclusionMaxDistance : Scene->DefaultMaxDistanceFieldOcclusionDistance;
 				UpdateGlobalDistanceFieldVolume(RHICmdList, Views[ViewIndex], Scene, OcclusionMaxDistance, Views[ViewIndex].GlobalDistanceFieldInfo);
 			}
 		}	
@@ -1472,7 +1472,7 @@ void FDeferredShadingSceneRenderer::Render(FRHICommandListImmediate& RHICmdList)
 	if (ViewFamily.EngineShowFlags.VisualizeDistanceFieldAO || ViewFamily.EngineShowFlags.VisualizeDistanceFieldGI)
 	{
 		// Use the skylight's max distance if there is one, to be consistent with DFAO shadowing on the skylight
-		const float OcclusionMaxDistance = Scene->SkyLight && !Scene->SkyLight->bWantsStaticShadowing ? Scene->SkyLight->OcclusionMaxDistance : GDefaultDFAOMaxOcclusionDistance;
+		const float OcclusionMaxDistance = Scene->SkyLight && !Scene->SkyLight->bWantsStaticShadowing ? Scene->SkyLight->OcclusionMaxDistance : Scene->DefaultMaxDistanceFieldOcclusionDistance;
 		TRefCountPtr<IPooledRenderTarget> DummyOutput;
 		RHICmdList.SetCurrentStat(GET_STATID(STAT_CLM_RenderDistanceFieldLighting));
 		RenderDistanceFieldLighting(RHICmdList, FDistanceFieldAOParameters(OcclusionMaxDistance), VelocityRT, DummyOutput, DummyOutput, false, ViewFamily.EngineShowFlags.VisualizeDistanceFieldAO, ViewFamily.EngineShowFlags.VisualizeDistanceFieldGI); 
@@ -1481,7 +1481,7 @@ void FDeferredShadingSceneRenderer::Render(FRHICommandListImmediate& RHICmdList)
 
 	if (ViewFamily.EngineShowFlags.VisualizeMeshDistanceFields)
 	{
-		RenderMeshDistanceFieldVisualization(RHICmdList, FDistanceFieldAOParameters(GDefaultDFAOMaxOcclusionDistance));
+		RenderMeshDistanceFieldVisualization(RHICmdList, FDistanceFieldAOParameters(Scene->DefaultMaxDistanceFieldOcclusionDistance));
 		ServiceLocalQueue();
 	}
 
diff --git a/Engine/Source/Runtime/Renderer/Private/DeferredShadingRenderer.h b/Engine/Source/Runtime/Renderer/Private/DeferredShadingRenderer.h
--- a/Engine/Source/Runtime/Renderer/Private/DeferredShadingRenderer.h
+++ b/Engine/Source/Runtime/Renderer/Private/DeferredShadingRenderer.h
@@ -229,6 +229,14 @@ private:
 		TRefCountPtr<IPooledRenderTarget>& OutDynamicBentNormalAO, 
 		TRefCountPtr<IPooledRenderTarget>& OutDynamicIrradiance);
 
+	void RenderDistanceFieldSpecularOcclusion(
+		FRHICommandListImmediate& RHICmdList, 
+		const FViewInfo& View,
+		FIntPoint TileListGroupSize,
+		const FDistanceFieldAOParameters& Parameters, 
+		const TRefCountPtr<IPooledRenderTarget>& DistanceFieldNormal, 
+		TRefCountPtr<IPooledRenderTarget>& OutSpecularOcclusion);
+
 	void RenderMeshDistanceFieldVisualization(FRHICommandListImmediate& RHICmdList, const FDistanceFieldAOParameters& Parameters);
 
 	/** Whether tiled deferred is supported and can be used at all. */
@@ -394,6 +402,8 @@ private:
 	/** Render image based reflections (SSR, Env, SkyLight) without compute shaders */
 	void RenderStandardDeferredImageBasedReflections(FRHICommandListImmediate& RHICmdList, bool bReflectionEnv, const TRefCountPtr<IPooledRenderTarget>& DynamicBentNormalAO);
 
+	bool RenderDeferredPlanarReflections(FRHICommandListImmediate& RHICmdList, bool bLightAccumulationIsInUse, TRefCountPtr<IPooledRenderTarget>& Output);
+
 	bool ShouldDoReflectionEnvironment() const;
 	
 	bool ShouldRenderDistanceFieldAO() const;
diff --git a/Engine/Source/Runtime/Renderer/Private/DistanceFieldLightingPost.cpp b/Engine/Source/Runtime/Renderer/Private/DistanceFieldLightingPost.cpp
--- a/Engine/Source/Runtime/Renderer/Private/DistanceFieldLightingPost.cpp
+++ b/Engine/Source/Runtime/Renderer/Private/DistanceFieldLightingPost.cpp
@@ -997,7 +997,7 @@ enum EAOUpsampleType
 	AOUpsample_OutputIrradiance
 };
 
-template<EAOUpsampleType UpsampleType, bool bModulateToSceneColor>
+template<EAOUpsampleType UpsampleType, bool bModulateToSceneColor, bool bSupportSpecularOcclusion>
 class TDistanceFieldAOUpsamplePS : public FGlobalShader
 {
 	DECLARE_SHADER_TYPE(TDistanceFieldAOUpsamplePS, Global);
@@ -1014,6 +1014,7 @@ public:
 		OutEnvironment.SetDefine(TEXT("MODULATE_SCENE_COLOR"), bModulateToSceneColor ? TEXT("1") : TEXT("0"));
 		OutEnvironment.SetDefine(TEXT("OUTPUT_BENT_NORMAL"), (UpsampleType == AOUpsample_OutputBentNormal || UpsampleType == AOUpsample_OutputBentNormalAndIrradiance) ? TEXT("1") : TEXT("0"));
 		OutEnvironment.SetDefine(TEXT("SUPPORT_IRRADIANCE"), (UpsampleType == AOUpsample_OutputIrradiance || UpsampleType == AOUpsample_OutputBentNormalAndIrradiance) ? TEXT("1") : TEXT("0"));
+		OutEnvironment.SetDefine(TEXT("SUPPORT_SPECULAR_OCCLUSION"), bSupportSpecularOcclusion ? TEXT("1") : TEXT("0"));
 		OutEnvironment.SetDefine(TEXT("OUTPUT_AO"), (UpsampleType == AOUpsample_OutputAO) ? TEXT("1") : TEXT("0"));
 	}
 
@@ -1029,10 +1030,12 @@ public:
 		BentNormalAOSampler.Bind(Initializer.ParameterMap,TEXT("BentNormalAOSampler"));
 		IrradianceTexture.Bind(Initializer.ParameterMap,TEXT("IrradianceTexture"));
 		IrradianceSampler.Bind(Initializer.ParameterMap,TEXT("IrradianceSampler"));
+		SpecularOcclusionTexture.Bind(Initializer.ParameterMap,TEXT("SpecularOcclusionTexture"));
+		SpecularOcclusionSampler.Bind(Initializer.ParameterMap,TEXT("SpecularOcclusionSampler"));
 		MinIndirectDiffuseOcclusion.Bind(Initializer.ParameterMap,TEXT("MinIndirectDiffuseOcclusion"));
 	}
 
-	void SetParameters(FRHICommandList& RHICmdList, const FSceneView& View, TRefCountPtr<IPooledRenderTarget>& DistanceFieldAOBentNormal, IPooledRenderTarget* DistanceFieldIrradiance)
+	void SetParameters(FRHICommandList& RHICmdList, const FSceneView& View, TRefCountPtr<IPooledRenderTarget>& DistanceFieldAOBentNormal, IPooledRenderTarget* DistanceFieldIrradiance, IPooledRenderTarget* DistanceFieldSpecularOcclusion)
 	{
 		const FPixelShaderRHIParamRef ShaderRHI = GetPixelShader();
 
@@ -1046,6 +1049,11 @@ public:
 			SetTextureParameter(RHICmdList, ShaderRHI, IrradianceTexture, IrradianceSampler, TStaticSamplerState<SF_Bilinear>::GetRHI(), DistanceFieldIrradiance->GetRenderTargetItem().ShaderResourceTexture);
 		}
 
+		if (SpecularOcclusionTexture.IsBound())
+		{
+			SetTextureParameter(RHICmdList, ShaderRHI, SpecularOcclusionTexture, SpecularOcclusionSampler, TStaticSamplerState<SF_Bilinear>::GetRHI(), DistanceFieldSpecularOcclusion->GetRenderTargetItem().ShaderResourceTexture);
+		}
+
 		FScene* Scene = (FScene*)View.Family->Scene;
 		const float MinOcclusion = Scene->SkyLight ? Scene->SkyLight->MinOcclusion : 0;
 		SetShaderValue(RHICmdList, ShaderRHI, MinIndirectDiffuseOcclusion, MinOcclusion);
@@ -1059,6 +1067,8 @@ public:
 		Ar << BentNormalAOSampler;
 		Ar << IrradianceTexture;
 		Ar << IrradianceSampler;
+		Ar << SpecularOcclusionTexture;
+		Ar << SpecularOcclusionSampler;
 		Ar << MinIndirectDiffuseOcclusion;
 		return bShaderHasOutdatedParameters;
 	}
@@ -1070,28 +1080,123 @@ private:
 	FShaderResourceParameter BentNormalAOSampler;
 	FShaderResourceParameter IrradianceTexture;
 	FShaderResourceParameter IrradianceSampler;
+	FShaderResourceParameter SpecularOcclusionTexture;
+	FShaderResourceParameter SpecularOcclusionSampler;
 	FShaderParameter MinIndirectDiffuseOcclusion;
 };
 
-#define IMPLEMENT_UPSAMPLE_PS_TYPE(UpsampleType, bModulateToSceneColor) \
-	typedef TDistanceFieldAOUpsamplePS<UpsampleType, bModulateToSceneColor> TDistanceFieldAOUpsamplePS##UpsampleType##bModulateToSceneColor; \
-	IMPLEMENT_SHADER_TYPE(template<>,TDistanceFieldAOUpsamplePS##UpsampleType##bModulateToSceneColor,TEXT("DistanceFieldLightingPost"),TEXT("AOUpsamplePS"),SF_Pixel);
-
-IMPLEMENT_UPSAMPLE_PS_TYPE(AOUpsample_OutputBentNormal, true)
-IMPLEMENT_UPSAMPLE_PS_TYPE(AOUpsample_OutputAO, true)
-IMPLEMENT_UPSAMPLE_PS_TYPE(AOUpsample_OutputBentNormalAndIrradiance, true)
-IMPLEMENT_UPSAMPLE_PS_TYPE(AOUpsample_OutputIrradiance, true)
+#define IMPLEMENT_UPSAMPLE_PS_TYPE(UpsampleType, bModulateToSceneColor, bSupportSpecularOcclusion) \
+	typedef TDistanceFieldAOUpsamplePS<UpsampleType, bModulateToSceneColor, bSupportSpecularOcclusion> TDistanceFieldAOUpsamplePS##UpsampleType##bModulateToSceneColor##bSupportSpecularOcclusion; \
+	IMPLEMENT_SHADER_TYPE(template<>,TDistanceFieldAOUpsamplePS##UpsampleType##bModulateToSceneColor##bSupportSpecularOcclusion,TEXT("DistanceFieldLightingPost"),TEXT("AOUpsamplePS"),SF_Pixel);
+
+IMPLEMENT_UPSAMPLE_PS_TYPE(AOUpsample_OutputBentNormal, true, true)
+IMPLEMENT_UPSAMPLE_PS_TYPE(AOUpsample_OutputAO, true, true)
+IMPLEMENT_UPSAMPLE_PS_TYPE(AOUpsample_OutputBentNormalAndIrradiance, true, true)
+IMPLEMENT_UPSAMPLE_PS_TYPE(AOUpsample_OutputIrradiance, true, true)
+IMPLEMENT_UPSAMPLE_PS_TYPE(AOUpsample_OutputBentNormal, false, true)
+IMPLEMENT_UPSAMPLE_PS_TYPE(AOUpsample_OutputAO, false, true)
+IMPLEMENT_UPSAMPLE_PS_TYPE(AOUpsample_OutputBentNormalAndIrradiance, false, true)
+IMPLEMENT_UPSAMPLE_PS_TYPE(AOUpsample_OutputIrradiance, false, true)
+
+IMPLEMENT_UPSAMPLE_PS_TYPE(AOUpsample_OutputBentNormal, true, false)
+IMPLEMENT_UPSAMPLE_PS_TYPE(AOUpsample_OutputAO, true, false)
+IMPLEMENT_UPSAMPLE_PS_TYPE(AOUpsample_OutputBentNormalAndIrradiance, true, false)
+IMPLEMENT_UPSAMPLE_PS_TYPE(AOUpsample_OutputIrradiance, true, false)
+IMPLEMENT_UPSAMPLE_PS_TYPE(AOUpsample_OutputBentNormal, false, false)
+IMPLEMENT_UPSAMPLE_PS_TYPE(AOUpsample_OutputAO, false, false)
+IMPLEMENT_UPSAMPLE_PS_TYPE(AOUpsample_OutputBentNormalAndIrradiance, false, false)
+IMPLEMENT_UPSAMPLE_PS_TYPE(AOUpsample_OutputIrradiance, false, false)
+
+template<bool bSupportSpecularOcclusion>
+void SetUpsampleShaders(
+	FRHICommandList& RHICmdList, 
+	const FViewInfo& View, 
+	TShaderMapRef<FPostProcessVS>& VertexShader,
+	TRefCountPtr<IPooledRenderTarget>& DistanceFieldAOBentNormal, 
+	TRefCountPtr<IPooledRenderTarget>& DistanceFieldIrradiance,
+	TRefCountPtr<IPooledRenderTarget>& DistanceFieldSpecularOcclusion,
+	bool bModulateSceneColor,
+	bool bVisualizeAmbientOcclusion,
+	bool bVisualizeGlobalIllumination)
+{
+	const bool bUseDistanceFieldGI = IsDistanceFieldGIAllowed(View);
 
-IMPLEMENT_UPSAMPLE_PS_TYPE(AOUpsample_OutputBentNormal, false)
-IMPLEMENT_UPSAMPLE_PS_TYPE(AOUpsample_OutputAO, false)
-IMPLEMENT_UPSAMPLE_PS_TYPE(AOUpsample_OutputBentNormalAndIrradiance, false)
-IMPLEMENT_UPSAMPLE_PS_TYPE(AOUpsample_OutputIrradiance, false)
+	if (bModulateSceneColor)
+	{
+		if (bVisualizeAmbientOcclusion)
+		{
+			TShaderMapRef<TDistanceFieldAOUpsamplePS<AOUpsample_OutputAO, true, bSupportSpecularOcclusion> > PixelShader(View.ShaderMap);
+			static FGlobalBoundShaderState BoundShaderState;
+			SetGlobalBoundShaderState(RHICmdList, View.GetFeatureLevel(), BoundShaderState, GFilterVertexDeclaration.VertexDeclarationRHI, *VertexShader, *PixelShader);
+			PixelShader->SetParameters(RHICmdList, View, DistanceFieldAOBentNormal, DistanceFieldIrradiance, DistanceFieldSpecularOcclusion);
+		}
+		else if (bVisualizeGlobalIllumination && bUseDistanceFieldGI)
+		{
+			TShaderMapRef<TDistanceFieldAOUpsamplePS<AOUpsample_OutputIrradiance, true, bSupportSpecularOcclusion> > PixelShader(View.ShaderMap);
+			static FGlobalBoundShaderState BoundShaderState;
+			SetGlobalBoundShaderState(RHICmdList, View.GetFeatureLevel(), BoundShaderState, GFilterVertexDeclaration.VertexDeclarationRHI, *VertexShader, *PixelShader);
+			PixelShader->SetParameters(RHICmdList, View, DistanceFieldAOBentNormal, DistanceFieldIrradiance, DistanceFieldSpecularOcclusion);
+		}
+		else
+		{
+			if (bUseDistanceFieldGI)
+			{
+				TShaderMapRef<TDistanceFieldAOUpsamplePS<AOUpsample_OutputBentNormalAndIrradiance, true, bSupportSpecularOcclusion> > PixelShader(View.ShaderMap);
+				static FGlobalBoundShaderState BoundShaderState;
+				SetGlobalBoundShaderState(RHICmdList, View.GetFeatureLevel(), BoundShaderState, GFilterVertexDeclaration.VertexDeclarationRHI, *VertexShader, *PixelShader);
+				PixelShader->SetParameters(RHICmdList, View, DistanceFieldAOBentNormal, DistanceFieldIrradiance, DistanceFieldSpecularOcclusion);
+			}
+			else
+			{
+				TShaderMapRef<TDistanceFieldAOUpsamplePS<AOUpsample_OutputBentNormal, true, bSupportSpecularOcclusion> > PixelShader(View.ShaderMap);
+				static FGlobalBoundShaderState BoundShaderState;
+				SetGlobalBoundShaderState(RHICmdList, View.GetFeatureLevel(), BoundShaderState, GFilterVertexDeclaration.VertexDeclarationRHI, *VertexShader, *PixelShader);
+				PixelShader->SetParameters(RHICmdList, View, DistanceFieldAOBentNormal, DistanceFieldIrradiance, DistanceFieldSpecularOcclusion);
+			}
+		}
+	}
+	else
+	{
+		if (bVisualizeAmbientOcclusion)
+		{
+			TShaderMapRef<TDistanceFieldAOUpsamplePS<AOUpsample_OutputAO, false, bSupportSpecularOcclusion> > PixelShader(View.ShaderMap);
+			static FGlobalBoundShaderState BoundShaderState;
+			SetGlobalBoundShaderState(RHICmdList, View.GetFeatureLevel(), BoundShaderState, GFilterVertexDeclaration.VertexDeclarationRHI, *VertexShader, *PixelShader);
+			PixelShader->SetParameters(RHICmdList, View, DistanceFieldAOBentNormal, DistanceFieldIrradiance, DistanceFieldSpecularOcclusion);
+		}
+		else if (bVisualizeGlobalIllumination && bUseDistanceFieldGI)
+		{
+			TShaderMapRef<TDistanceFieldAOUpsamplePS<AOUpsample_OutputIrradiance, false, bSupportSpecularOcclusion> > PixelShader(View.ShaderMap);
+			static FGlobalBoundShaderState BoundShaderState;
+			SetGlobalBoundShaderState(RHICmdList, View.GetFeatureLevel(), BoundShaderState, GFilterVertexDeclaration.VertexDeclarationRHI, *VertexShader, *PixelShader);
+			PixelShader->SetParameters(RHICmdList, View, DistanceFieldAOBentNormal, DistanceFieldIrradiance, DistanceFieldSpecularOcclusion);
+		}
+		else
+		{
+			if (bUseDistanceFieldGI)
+			{
+				TShaderMapRef<TDistanceFieldAOUpsamplePS<AOUpsample_OutputBentNormalAndIrradiance, false, bSupportSpecularOcclusion> > PixelShader(View.ShaderMap);
+				static FGlobalBoundShaderState BoundShaderState;
+				SetGlobalBoundShaderState(RHICmdList, View.GetFeatureLevel(), BoundShaderState, GFilterVertexDeclaration.VertexDeclarationRHI, *VertexShader, *PixelShader);
+				PixelShader->SetParameters(RHICmdList, View, DistanceFieldAOBentNormal, DistanceFieldIrradiance, DistanceFieldSpecularOcclusion);
+			}
+			else
+			{
+				TShaderMapRef<TDistanceFieldAOUpsamplePS<AOUpsample_OutputBentNormal, false, bSupportSpecularOcclusion> > PixelShader(View.ShaderMap);
+				static FGlobalBoundShaderState BoundShaderState;
+				SetGlobalBoundShaderState(RHICmdList, View.GetFeatureLevel(), BoundShaderState, GFilterVertexDeclaration.VertexDeclarationRHI, *VertexShader, *PixelShader);
+				PixelShader->SetParameters(RHICmdList, View, DistanceFieldAOBentNormal, DistanceFieldIrradiance, DistanceFieldSpecularOcclusion);
+			}
+		}
+	}
+}
 
 void UpsampleBentNormalAO(
 	FRHICommandList& RHICmdList, 
 	const TArray<FViewInfo>& Views, 
 	TRefCountPtr<IPooledRenderTarget>& DistanceFieldAOBentNormal, 
 	TRefCountPtr<IPooledRenderTarget>& DistanceFieldIrradiance,
+	TRefCountPtr<IPooledRenderTarget>& DistanceFieldSpecularOcclusion,
 	bool bModulateSceneColor,
 	bool bVisualizeAmbientOcclusion,
 	bool bVisualizeGlobalIllumination)
@@ -1099,7 +1204,6 @@ void UpsampleBentNormalAO(
 	for (int32 ViewIndex = 0; ViewIndex < Views.Num(); ViewIndex++)
 	{
 		const FViewInfo& View = Views[ViewIndex];
-		const bool bUseDistanceFieldGI = IsDistanceFieldGIAllowed(View);
 
 		SCOPED_DRAW_EVENT(RHICmdList, UpsampleAO);
 
@@ -1124,73 +1228,13 @@ void UpsampleBentNormalAO(
 
 		TShaderMapRef<FPostProcessVS> VertexShader(View.ShaderMap);
 
-		if (bModulateSceneColor)
+		if (DistanceFieldSpecularOcclusion)
 		{
-			if (bVisualizeAmbientOcclusion)
-			{
-				TShaderMapRef<TDistanceFieldAOUpsamplePS<AOUpsample_OutputAO, true> > PixelShader(View.ShaderMap);
-				static FGlobalBoundShaderState BoundShaderState;
-				SetGlobalBoundShaderState(RHICmdList, View.GetFeatureLevel(), BoundShaderState, GFilterVertexDeclaration.VertexDeclarationRHI, *VertexShader, *PixelShader);
-				PixelShader->SetParameters(RHICmdList, View, DistanceFieldAOBentNormal, DistanceFieldIrradiance);
-			}
-			else if (bVisualizeGlobalIllumination && bUseDistanceFieldGI)
-			{
-				TShaderMapRef<TDistanceFieldAOUpsamplePS<AOUpsample_OutputIrradiance, true> > PixelShader(View.ShaderMap);
-				static FGlobalBoundShaderState BoundShaderState;
-				SetGlobalBoundShaderState(RHICmdList, View.GetFeatureLevel(), BoundShaderState, GFilterVertexDeclaration.VertexDeclarationRHI, *VertexShader, *PixelShader);
-				PixelShader->SetParameters(RHICmdList, View, DistanceFieldAOBentNormal, DistanceFieldIrradiance);
-			}
-			else
-			{
-				if (bUseDistanceFieldGI)
-				{
-					TShaderMapRef<TDistanceFieldAOUpsamplePS<AOUpsample_OutputBentNormalAndIrradiance, true> > PixelShader(View.ShaderMap);
-					static FGlobalBoundShaderState BoundShaderState;
-					SetGlobalBoundShaderState(RHICmdList, View.GetFeatureLevel(), BoundShaderState, GFilterVertexDeclaration.VertexDeclarationRHI, *VertexShader, *PixelShader);
-					PixelShader->SetParameters(RHICmdList, View, DistanceFieldAOBentNormal, DistanceFieldIrradiance);
-				}
-				else
-				{
-					TShaderMapRef<TDistanceFieldAOUpsamplePS<AOUpsample_OutputBentNormal, true> > PixelShader(View.ShaderMap);
-					static FGlobalBoundShaderState BoundShaderState;
-					SetGlobalBoundShaderState(RHICmdList, View.GetFeatureLevel(), BoundShaderState, GFilterVertexDeclaration.VertexDeclarationRHI, *VertexShader, *PixelShader);
-					PixelShader->SetParameters(RHICmdList, View, DistanceFieldAOBentNormal, DistanceFieldIrradiance);
-				}
-			}
+			SetUpsampleShaders<true>(RHICmdList, View, VertexShader, DistanceFieldAOBentNormal, DistanceFieldIrradiance, DistanceFieldSpecularOcclusion, bModulateSceneColor, bVisualizeAmbientOcclusion, bVisualizeGlobalIllumination);
 		}
 		else
 		{
-			if (bVisualizeAmbientOcclusion)
-			{
-				TShaderMapRef<TDistanceFieldAOUpsamplePS<AOUpsample_OutputAO, false> > PixelShader(View.ShaderMap);
-				static FGlobalBoundShaderState BoundShaderState;
-				SetGlobalBoundShaderState(RHICmdList, View.GetFeatureLevel(), BoundShaderState, GFilterVertexDeclaration.VertexDeclarationRHI, *VertexShader, *PixelShader);
-				PixelShader->SetParameters(RHICmdList, View, DistanceFieldAOBentNormal, DistanceFieldIrradiance);
-			}
-			else if (bVisualizeGlobalIllumination && bUseDistanceFieldGI)
-			{
-				TShaderMapRef<TDistanceFieldAOUpsamplePS<AOUpsample_OutputIrradiance, false> > PixelShader(View.ShaderMap);
-				static FGlobalBoundShaderState BoundShaderState;
-				SetGlobalBoundShaderState(RHICmdList, View.GetFeatureLevel(), BoundShaderState, GFilterVertexDeclaration.VertexDeclarationRHI, *VertexShader, *PixelShader);
-				PixelShader->SetParameters(RHICmdList, View, DistanceFieldAOBentNormal, DistanceFieldIrradiance);
-			}
-			else
-			{
-				if (bUseDistanceFieldGI)
-				{
-					TShaderMapRef<TDistanceFieldAOUpsamplePS<AOUpsample_OutputBentNormalAndIrradiance, false> > PixelShader(View.ShaderMap);
-					static FGlobalBoundShaderState BoundShaderState;
-					SetGlobalBoundShaderState(RHICmdList, View.GetFeatureLevel(), BoundShaderState, GFilterVertexDeclaration.VertexDeclarationRHI, *VertexShader, *PixelShader);
-					PixelShader->SetParameters(RHICmdList, View, DistanceFieldAOBentNormal, DistanceFieldIrradiance);
-				}
-				else
-				{
-					TShaderMapRef<TDistanceFieldAOUpsamplePS<AOUpsample_OutputBentNormal, false> > PixelShader(View.ShaderMap);
-					static FGlobalBoundShaderState BoundShaderState;
-					SetGlobalBoundShaderState(RHICmdList, View.GetFeatureLevel(), BoundShaderState, GFilterVertexDeclaration.VertexDeclarationRHI, *VertexShader, *PixelShader);
-					PixelShader->SetParameters(RHICmdList, View, DistanceFieldAOBentNormal, DistanceFieldIrradiance);
-				}
-			}
+			SetUpsampleShaders<false>(RHICmdList, View, VertexShader, DistanceFieldAOBentNormal, DistanceFieldIrradiance, DistanceFieldSpecularOcclusion, bModulateSceneColor, bVisualizeAmbientOcclusion, bVisualizeGlobalIllumination);
 		}
 
 		DrawRectangle( 
diff --git a/Engine/Source/Runtime/Renderer/Private/DistanceFieldLightingPost.h b/Engine/Source/Runtime/Renderer/Private/DistanceFieldLightingPost.h
--- a/Engine/Source/Runtime/Renderer/Private/DistanceFieldLightingPost.h
+++ b/Engine/Source/Runtime/Renderer/Private/DistanceFieldLightingPost.h
@@ -41,6 +41,7 @@ extern void UpsampleBentNormalAO(
 	const TArray<FViewInfo>& Views, 
 	TRefCountPtr<IPooledRenderTarget>& DistanceFieldAOBentNormal, 
 	TRefCountPtr<IPooledRenderTarget>& DistanceFieldIrradiance,
+	TRefCountPtr<IPooledRenderTarget>& DistanceFieldSpecularOcclusion,
 	bool bModulateSceneColor,
 	bool bVisualizeAmbientOcclusion,
 	bool bVisualizeGlobalIllumination);
\ No newline at end of file
diff --git a/Engine/Source/Runtime/Renderer/Private/DistanceFieldLightingShared.h b/Engine/Source/Runtime/Renderer/Private/DistanceFieldLightingShared.h
--- a/Engine/Source/Runtime/Renderer/Private/DistanceFieldLightingShared.h
+++ b/Engine/Source/Runtime/Renderer/Private/DistanceFieldLightingShared.h
@@ -570,4 +570,6 @@ public:
 	}
 };
 
-extern TGlobalResource<FDistanceFieldObjectBufferResource> GAOCulledObjectBuffers;
\ No newline at end of file
+extern TGlobalResource<FDistanceFieldObjectBufferResource> GAOCulledObjectBuffers;
+
+extern bool SupportsDistanceFieldAO(ERHIFeatureLevel::Type FeatureLevel, EShaderPlatform ShaderPlatform);
\ No newline at end of file
diff --git a/Engine/Source/Runtime/Renderer/Private/DistanceFieldObjectManagement.cpp b/Engine/Source/Runtime/Renderer/Private/DistanceFieldObjectManagement.cpp
--- a/Engine/Source/Runtime/Renderer/Private/DistanceFieldObjectManagement.cpp
+++ b/Engine/Source/Runtime/Renderer/Private/DistanceFieldObjectManagement.cpp
@@ -935,8 +935,6 @@ void FDeferredShadingSceneRenderer::UpdateGlobalDistanceFieldObjectBuffers(FRHIC
 {
 	FDistanceFieldSceneData& DistanceFieldSceneData = Scene->DistanceFieldSceneData;
 
-	DistanceFieldSceneData.PrimitiveModifiedBounds.Reset();
-
 	if (GDistanceFieldVolumeTextureAtlas.VolumeTextureRHI
 		&& (DistanceFieldSceneData.HasPendingOperations() || DistanceFieldSceneData.AtlasGeneration != GDistanceFieldVolumeTextureAtlas.GetGeneration()))
 	{
diff --git a/Engine/Source/Runtime/Renderer/Private/DistanceFieldSpecularOcclusion.cpp b/Engine/Source/Runtime/Renderer/Private/DistanceFieldSpecularOcclusion.cpp
new file mode 100644
--- /dev/null
+++ b/Engine/Source/Runtime/Renderer/Private/DistanceFieldSpecularOcclusion.cpp
@@ -0,0 +1,219 @@
+// Copyright 1998-2016 Epic Games, Inc. All Rights Reserved.
+
+/*=============================================================================
+	DistanceFieldSpecularOcclusion.cpp
+=============================================================================*/
+
+#include "RendererPrivate.h"
+#include "ScenePrivate.h"
+#include "UniformBuffer.h"
+#include "ShaderParameters.h"
+#include "PostProcessing.h"
+#include "SceneFilterRendering.h"
+#include "DistanceFieldLightingShared.h"
+#include "DistanceFieldSurfaceCacheLighting.h"
+#include "DistanceFieldLightingPost.h"
+#include "DistanceFieldGlobalIllumination.h"
+#include "GlobalDistanceField.h"
+#include "PostProcessAmbientOcclusion.h"
+#include "RHICommandList.h"
+#include "SceneUtils.h"
+#include "OneColorShader.h"
+#include "BasePassRendering.h"
+#include "HeightfieldLighting.h"
+
+// Also have to set USE_CONE_TRACED_SPECULAR_OCCLUSION to 1
+// Prototype feature
+int32 GDistanceFieldSpecularOcclusion = 0;
+FAutoConsoleVariableRef CVarDistanceFieldSpecularOcclusion(
+	TEXT("r.DistanceFieldSpecularOcclusion"),
+	GDistanceFieldSpecularOcclusion,
+	TEXT("Whether the distance field specular occlusion feature is allowed, which is used to shadow reflection captures and skylight reflections more accurately."),
+	ECVF_Scalability | ECVF_RenderThreadSafe
+	);
+
+int32 GSpecularOcclusionDownsampleFactor = 1;
+
+FIntPoint GetBufferSizeForSpecularOcclusion()
+{
+	return FIntPoint::DivideAndRoundDown(GetBufferSizeForAO(), GSpecularOcclusionDownsampleFactor);
+}
+
+class FConeTraceSpecularOcclusionCS : public FGlobalShader
+{
+	DECLARE_SHADER_TYPE(FConeTraceSpecularOcclusionCS,Global)
+public:
+
+	static bool ShouldCache(EShaderPlatform Platform)
+	{
+		return IsFeatureLevelSupported(Platform, ERHIFeatureLevel::SM5) && DoesPlatformSupportDistanceFieldAO(Platform);
+	}
+
+	static void ModifyCompilationEnvironment(EShaderPlatform Platform, FShaderCompilerEnvironment& OutEnvironment)
+	{
+		FGlobalShader::ModifyCompilationEnvironment(Platform,OutEnvironment);
+		OutEnvironment.SetDefine(TEXT("CULLED_TILE_SIZEX"), GDistanceFieldAOTileSizeX);
+		OutEnvironment.SetDefine(TEXT("TRACE_DOWNSAMPLE_FACTOR"), GSpecularOcclusionDownsampleFactor);
+
+		// To reduce shader compile time of compute shaders with shared memory, doesn't have an impact on generated code with current compiler (June 2010 DX SDK)
+		OutEnvironment.CompilerFlags.Add(CFLAG_StandardOptimization);
+	}
+
+	FConeTraceSpecularOcclusionCS(const ShaderMetaType::CompiledShaderInitializerType& Initializer)
+		: FGlobalShader(Initializer)
+	{
+		DeferredParameters.Bind(Initializer.ParameterMap);
+		ObjectParameters.Bind(Initializer.ParameterMap);
+		AOParameters.Bind(Initializer.ParameterMap);
+		ScreenGridParameters.Bind(Initializer.ParameterMap);
+		TileHeadDataUnpacked.Bind(Initializer.ParameterMap, TEXT("TileHeadDataUnpacked"));
+		TileArrayData.Bind(Initializer.ParameterMap, TEXT("TileArrayData"));
+		TileConeDepthRanges.Bind(Initializer.ParameterMap, TEXT("TileConeDepthRanges"));
+		TileListGroupSize.Bind(Initializer.ParameterMap, TEXT("TileListGroupSize"));
+		TanConeHalfAngle.Bind(Initializer.ParameterMap, TEXT("TanConeHalfAngle"));
+		BentNormalNormalizeFactor.Bind(Initializer.ParameterMap, TEXT("BentNormalNormalizeFactor"));
+		SpecularOcclusion.Bind(Initializer.ParameterMap, TEXT("SpecularOcclusion"));
+	}
+
+	FConeTraceSpecularOcclusionCS()
+	{
+	}
+
+	void SetParameters(
+		FRHICommandList& RHICmdList, 
+		const FViewInfo& View, 
+		FIntPoint TileListGroupSizeValue, 
+		FSceneRenderTargetItem& DistanceFieldNormal, 
+		const FDistanceFieldAOParameters& Parameters,
+		FSceneRenderTargetItem& SpecularOcclusionBuffer)
+	{
+		FComputeShaderRHIParamRef ShaderRHI = GetComputeShader();
+		FGlobalShader::SetParameters(RHICmdList, ShaderRHI, View);
+		DeferredParameters.Set(RHICmdList, ShaderRHI, View);
+		ObjectParameters.Set(RHICmdList, ShaderRHI, GAOCulledObjectBuffers.Buffers);
+		AOParameters.Set(RHICmdList, ShaderRHI, Parameters);
+		ScreenGridParameters.Set(RHICmdList, ShaderRHI, View, DistanceFieldNormal);
+
+		FAOSampleData2 AOSampleData;
+
+		TArray<FVector, TInlineAllocator<9> > SampleDirections;
+		GetSpacedVectors(SampleDirections);
+
+		for (int32 SampleIndex = 0; SampleIndex < NumConeSampleDirections; SampleIndex++)
+		{
+			AOSampleData.SampleDirections[SampleIndex] = FVector4(SampleDirections[SampleIndex]);
+		}
+
+		SetUniformBufferParameterImmediate(RHICmdList, ShaderRHI, GetUniformBufferParameter<FAOSampleData2>(), AOSampleData);
+
+		FTileIntersectionResources* TileIntersectionResources = View.ViewState->AOTileIntersectionResources;
+
+		SetSRVParameter(RHICmdList, ShaderRHI, TileHeadDataUnpacked, TileIntersectionResources->TileHeadDataUnpacked.SRV);
+		SetSRVParameter(RHICmdList, ShaderRHI, TileArrayData, TileIntersectionResources->TileArrayData.SRV);
+		SetSRVParameter(RHICmdList, ShaderRHI, TileConeDepthRanges, TileIntersectionResources->TileConeDepthRanges.SRV);
+
+		SetShaderValue(RHICmdList, ShaderRHI, TileListGroupSize, TileListGroupSizeValue);
+
+		extern float GAOConeHalfAngle;
+		SetShaderValue(RHICmdList, ShaderRHI, TanConeHalfAngle, FMath::Tan(GAOConeHalfAngle));
+
+		FVector UnoccludedVector(0);
+
+		for (int32 SampleIndex = 0; SampleIndex < NumConeSampleDirections; SampleIndex++)
+		{
+			UnoccludedVector += SampleDirections[SampleIndex];
+		}
+
+		float BentNormalNormalizeFactorValue = 1.0f / (UnoccludedVector / NumConeSampleDirections).Size();
+		SetShaderValue(RHICmdList, ShaderRHI, BentNormalNormalizeFactor, BentNormalNormalizeFactorValue);
+
+		int32 NumOutUAVs = 0;
+		FUnorderedAccessViewRHIParamRef OutUAVs[1];
+		OutUAVs[NumOutUAVs++] = SpecularOcclusionBuffer.UAV;
+
+		RHICmdList.TransitionResources(EResourceTransitionAccess::ERWBarrier, EResourceTransitionPipeline::EComputeToCompute, OutUAVs, NumOutUAVs);
+
+		SpecularOcclusion.SetTexture(RHICmdList, ShaderRHI, SpecularOcclusionBuffer.ShaderResourceTexture, SpecularOcclusionBuffer.UAV);
+	}
+
+	void UnsetParameters(FRHICommandList& RHICmdList, FSceneRenderTargetItem& SpecularOcclusionBuffer)
+	{
+		SpecularOcclusion.UnsetUAV(RHICmdList, GetComputeShader());
+
+		int32 NumOutUAVs = 0;
+		FUnorderedAccessViewRHIParamRef OutUAVs[1];
+		OutUAVs[NumOutUAVs++] = SpecularOcclusionBuffer.UAV;
+
+		RHICmdList.TransitionResources(EResourceTransitionAccess::EReadable, EResourceTransitionPipeline::EComputeToCompute, OutUAVs, NumOutUAVs);
+	}
+
+	virtual bool Serialize(FArchive& Ar) override
+	{		
+		bool bShaderHasOutdatedParameters = FGlobalShader::Serialize(Ar);
+		Ar << DeferredParameters;
+		Ar << ObjectParameters;
+		Ar << AOParameters;
+		Ar << ScreenGridParameters;
+		Ar << TileHeadDataUnpacked;
+		Ar << TileArrayData;
+		Ar << TileConeDepthRanges;
+		Ar << TileListGroupSize;
+		Ar << TanConeHalfAngle;
+		Ar << BentNormalNormalizeFactor;
+		Ar << SpecularOcclusion;
+		return bShaderHasOutdatedParameters;
+	}
+
+private:
+
+	FDeferredPixelShaderParameters DeferredParameters;
+	FDistanceFieldCulledObjectBufferParameters ObjectParameters;
+	FAOParameters AOParameters;
+	FScreenGridParameters ScreenGridParameters;
+	FShaderResourceParameter TileHeadDataUnpacked;
+	FShaderResourceParameter TileArrayData;
+	FShaderResourceParameter TileConeDepthRanges;
+	FShaderParameter TileListGroupSize;
+	FShaderParameter TanConeHalfAngle;
+	FShaderParameter BentNormalNormalizeFactor;
+	FRWShaderParameter SpecularOcclusion;
+};
+
+IMPLEMENT_SHADER_TYPE(,FConeTraceSpecularOcclusionCS,TEXT("DistanceFieldSpecularOcclusion"),TEXT("ConeTraceSpecularOcclusionCS"),SF_Compute);
+
+void FDeferredShadingSceneRenderer::RenderDistanceFieldSpecularOcclusion(
+	FRHICommandListImmediate& RHICmdList, 
+	const FViewInfo& View,
+	FIntPoint TileListGroupSize,
+	const FDistanceFieldAOParameters& Parameters, 
+	const TRefCountPtr<IPooledRenderTarget>& DistanceFieldNormal, 
+	TRefCountPtr<IPooledRenderTarget>& OutSpecularOcclusion)
+{
+	if (GDistanceFieldSpecularOcclusion)
+	{
+		const FIntPoint ConeTraceBufferSize = GetBufferSizeForSpecularOcclusion();
+
+		SetRenderTarget(RHICmdList, NULL, NULL);
+
+		{
+			FPooledRenderTargetDesc Desc(FPooledRenderTargetDesc::Create2DDesc(ConeTraceBufferSize, PF_R16F, FClearValueBinding::None, TexCreate_None, TexCreate_RenderTargetable | TexCreate_UAV, false));
+			GRenderTargetPool.FindFreeElement(RHICmdList, Desc, OutSpecularOcclusion, TEXT("DownsampledSpecularOcclusion"));
+		}
+
+		{
+			SCOPED_DRAW_EVENT(RHICmdList, ConeTraceSpecularOcclusion);
+
+			const uint32 GroupSizeX = FMath::DivideAndRoundUp(View.ViewRect.Size().X / GAODownsampleFactor, GDistanceFieldAOTileSizeX / GSpecularOcclusionDownsampleFactor);
+			const uint32 GroupSizeY = FMath::DivideAndRoundUp(View.ViewRect.Size().Y / GAODownsampleFactor, GDistanceFieldAOTileSizeY / GSpecularOcclusionDownsampleFactor);
+
+			TShaderMapRef<FConeTraceSpecularOcclusionCS> ComputeShader(View.ShaderMap);
+
+			RHICmdList.SetComputeShader(ComputeShader->GetComputeShader());
+			ComputeShader->SetParameters(RHICmdList, View, TileListGroupSize, DistanceFieldNormal->GetRenderTargetItem(), Parameters, OutSpecularOcclusion->GetRenderTargetItem());
+			DispatchComputeShader(RHICmdList, *ComputeShader, GroupSizeX, GroupSizeY, 1);
+			ComputeShader->UnsetParameters(RHICmdList, OutSpecularOcclusion->GetRenderTargetItem());
+		}
+
+		GRenderTargetPool.VisualizeTexture.SetCheckPoint(RHICmdList, OutSpecularOcclusion);
+	}
+}
diff --git a/Engine/Source/Runtime/Renderer/Private/DistanceFieldSurfaceCacheLighting.cpp b/Engine/Source/Runtime/Renderer/Private/DistanceFieldSurfaceCacheLighting.cpp
--- a/Engine/Source/Runtime/Renderer/Private/DistanceFieldSurfaceCacheLighting.cpp
+++ b/Engine/Source/Runtime/Renderer/Private/DistanceFieldSurfaceCacheLighting.cpp
@@ -3028,7 +3028,7 @@ void FDeferredShadingSceneRenderer::RenderDFAOAsIndirectShadowing(
 	if (GDistanceFieldAOApplyToStaticIndirect && ShouldRenderDistanceFieldAO())
 	{
 		// Use the skylight's max distance if there is one, to be consistent with DFAO shadowing on the skylight
-		const float OcclusionMaxDistance = Scene->SkyLight && !Scene->SkyLight->bWantsStaticShadowing ? Scene->SkyLight->OcclusionMaxDistance : GDefaultDFAOMaxOcclusionDistance;
+		const float OcclusionMaxDistance = Scene->SkyLight && !Scene->SkyLight->bWantsStaticShadowing ? Scene->SkyLight->OcclusionMaxDistance : Scene->DefaultMaxDistanceFieldOcclusionDistance;
 		TRefCountPtr<IPooledRenderTarget> DummyOutput;
 		RenderDistanceFieldLighting(RHICmdList, FDistanceFieldAOParameters(OcclusionMaxDistance), VelocityTexture, DynamicBentNormalAO, DummyOutput, true, false, false);
 	}
@@ -3154,6 +3154,9 @@ bool FDeferredShadingSceneRenderer::RenderDistanceFieldLighting(
 
 			GRenderTargetPool.VisualizeTexture.SetCheckPoint(RHICmdList, BentNormalOutput);
 
+			TRefCountPtr<IPooledRenderTarget> SpecularOcclusion;
+			RenderDistanceFieldSpecularOcclusion(RHICmdList, View, TileListGroupSize, Parameters, DistanceFieldNormal, SpecularOcclusion);
+
 			if (bVisualizeAmbientOcclusion || bVisualizeGlobalIllumination)
 			{
 				SceneContext.BeginRenderingSceneColor(RHICmdList, ESimpleRenderTargetMode::EExistingColorAndDepth, FExclusiveDepthStencil::DepthRead_StencilNop);
@@ -3196,7 +3199,7 @@ bool FDeferredShadingSceneRenderer::RenderDistanceFieldLighting(
 			}
 
 			// Upsample to full resolution, write to output
-			UpsampleBentNormalAO(RHICmdList, Views, BentNormalOutput, IrradianceOutput, bModulateToSceneColor, bVisualizeAmbientOcclusion, bVisualizeGlobalIllumination);
+			UpsampleBentNormalAO(RHICmdList, Views, BentNormalOutput, IrradianceOutput, SpecularOcclusion, bModulateToSceneColor, bVisualizeAmbientOcclusion, bVisualizeGlobalIllumination);
 
 			if (!bVisualizeAmbientOcclusion && !bVisualizeGlobalIllumination)
 			{
diff --git a/Engine/Source/Runtime/Renderer/Private/DistanceFieldSurfaceCacheLighting.h b/Engine/Source/Runtime/Renderer/Private/DistanceFieldSurfaceCacheLighting.h
--- a/Engine/Source/Runtime/Renderer/Private/DistanceFieldSurfaceCacheLighting.h
+++ b/Engine/Source/Runtime/Renderer/Private/DistanceFieldSurfaceCacheLighting.h
@@ -17,11 +17,9 @@ const int32 GAODownsampleFactor = 2;
 
 extern const uint32 UpdateObjectsGroupSize;
 
-const float GDefaultDFAOMaxOcclusionDistance = 600.0f;
-
 inline bool DoesPlatformSupportDistanceFieldAO(EShaderPlatform Platform)
 {
-	return Platform == SP_PCD3D_SM5 || Platform == SP_PS4 || Platform == SP_METAL_SM5 || Platform == SP_XBOXONE;
+	return Platform == SP_PCD3D_SM5 || Platform == SP_PS4 || Platform == SP_XBOXONE;
 }
 
 extern FIntPoint GetBufferSizeForAO();
diff --git a/Engine/Source/Runtime/Renderer/Private/ForwardBasePassRendering.h b/Engine/Source/Runtime/Renderer/Private/ForwardBasePassRendering.h
--- a/Engine/Source/Runtime/Renderer/Private/ForwardBasePassRendering.h
+++ b/Engine/Source/Runtime/Renderer/Private/ForwardBasePassRendering.h
@@ -60,7 +60,6 @@ protected:
 	{
 		VertexParametersType::Bind(Initializer.ParameterMap);
 		HeightFogParameters.Bind(Initializer.ParameterMap);
-		PlanarReflectionCaptureDataParameter.Bind(Initializer.ParameterMap, TEXT("PlanarReflectionCaptureData"));
 	}
 
 public:
@@ -77,7 +76,6 @@ public:
 		bool bShaderHasOutdatedParameters = FMeshMaterialShader::Serialize(Ar);
 		VertexParametersType::Serialize(Ar);
 		Ar << HeightFogParameters;
-		Ar << PlanarReflectionCaptureDataParameter;
 		return bShaderHasOutdatedParameters;
 	}
 
@@ -92,11 +90,6 @@ public:
 	{
 		HeightFogParameters.Set(RHICmdList, GetVertexShader(), &View);
 		FMeshMaterialShader::SetParameters(RHICmdList, GetVertexShader(),MaterialRenderProxy,InMaterialResource,View,TextureMode);
-
-		if (PlanarReflectionCaptureDataParameter.IsBound())
-		{
-			SetShaderValue(RHICmdList, GetVertexShader(), PlanarReflectionCaptureDataParameter, View.ReflectionPlane);
-		}
 	}
 
 	void SetMesh(FRHICommandList& RHICmdList, const FVertexFactory* VertexFactory,const FSceneView& View,const FPrimitiveSceneProxy* Proxy,const FMeshBatchElement& BatchElement,const FMeshDrawingRenderState& DrawRenderState)
@@ -106,7 +99,6 @@ public:
 
 private:
 	FHeightFogShaderParameters HeightFogParameters;
-	FShaderParameter PlanarReflectionCaptureDataParameter;
 };
 
 template<typename LightMapPolicyType>
@@ -222,7 +214,6 @@ public:
 		ReflectionSampler1.Bind(Initializer.ParameterMap, TEXT("ReflectionCubemapSampler1"));
 		ReflectionCubemap2.Bind(Initializer.ParameterMap, TEXT("ReflectionCubemap2"));
 		ReflectionSampler2.Bind(Initializer.ParameterMap, TEXT("ReflectionCubemapSampler2"));
-		PlanarReflectionCaptureDataParameter.Bind(Initializer.ParameterMap, TEXT("PlanarReflectionCaptureData"));
 	}
 	TBasePassForForwardShadingPSPolicyParamType() {}
 
@@ -230,12 +221,6 @@ public:
 	{
 		FMeshMaterialShader::SetParameters(RHICmdList, GetPixelShader(),MaterialRenderProxy,MaterialResource,*View,TextureMode);
 		EditorCompositeParams.SetParameters(RHICmdList, MaterialResource, View, bEnableEditorPrimitveDepthTest, GetPixelShader());
-
-		if (PlanarReflectionCaptureDataParameter.IsBound())
-		{
-			check(View != nullptr);
-			SetShaderValue(RHICmdList, GetPixelShader(), PlanarReflectionCaptureDataParameter, View->ReflectionPlane);
-		}
 	}
 
 	void SetMesh(FRHICommandList& RHICmdList, const FVertexFactory* VertexFactory,const FSceneView& View,const FPrimitiveSceneProxy* Proxy,const FMeshBatchElement& BatchElement,const FMeshDrawingRenderState& DrawRenderState)
@@ -336,7 +321,6 @@ public:
 		Ar << ReflectionPositionsAndRadii;
 		Ar << ReflectionSampler1;
 		Ar << ReflectionSampler2;
-		Ar << PlanarReflectionCaptureDataParameter;
 
 		return bShaderHasOutdatedParameters;
 	}
@@ -359,7 +343,6 @@ private:
 	FShaderResourceParameter ReflectionCubemap2;
 	FShaderResourceParameter ReflectionSampler2;
 	FShaderParameter ReflectionPositionsAndRadii;
-	FShaderParameter PlanarReflectionCaptureDataParameter;
 };
 
 template<typename LightMapPolicyType, int32 NumDynamicPointLights>
diff --git a/Engine/Source/Runtime/Renderer/Private/GlobalDistanceField.cpp b/Engine/Source/Runtime/Renderer/Private/GlobalDistanceField.cpp
--- a/Engine/Source/Runtime/Renderer/Private/GlobalDistanceField.cpp
+++ b/Engine/Source/Runtime/Renderer/Private/GlobalDistanceField.cpp
@@ -15,6 +15,7 @@
 #include "GlobalDistanceField.h"
 #include "RHICommandList.h"
 #include "SceneUtils.h"
+#include "HeightfieldLighting.h"
 
 int32 GAOGlobalDistanceField = 1;
 FAutoConsoleVariableRef CVarAOGlobalDistanceField(
@@ -41,6 +42,14 @@ FAutoConsoleVariableRef CVarAOGlobalDistanceFieldPartialUpdates(
 	ECVF_Cheat | ECVF_RenderThreadSafe
 	);
 
+int32 GAOGlobalDistanceFieldStaggeredUpdates = 1;
+FAutoConsoleVariableRef CVarAOGlobalDistanceFieldStaggeredUpdatess(
+	TEXT("r.AOGlobalDistanceFieldStaggeredUpdates"),
+	GAOGlobalDistanceFieldStaggeredUpdates,
+	TEXT("Whether to allow the larger clipmaps to be updated less frequently."),
+	ECVF_Cheat | ECVF_RenderThreadSafe
+	);
+
 int32 GAOLogGlobalDistanceFieldModifiedPrimitives = 0;
 FAutoConsoleVariableRef CVarAOLogGlobalDistanceFieldModifiedPrimitives(
 	TEXT("r.AOGlobalDistanceFieldLogModifiedPrimitives"),
@@ -122,7 +131,10 @@ void FGlobalDistanceFieldInfo::UpdateParameterData(float MaxOcclusionDistance)
 		}
 
 		ParameterData.GlobalDFResolution = GAOGlobalDFResolution;
-		ParameterData.MaxDistance = MaxOcclusionDistance;
+
+		extern float GAOConeHalfAngle;
+		const float GlobalMaxSphereQueryRadius = MaxOcclusionDistance / (1 + FMath::Tan(GAOConeHalfAngle));
+		ParameterData.MaxDistance = GlobalMaxSphereQueryRadius;
 	}
 	else
 	{
@@ -485,6 +497,196 @@ private:
 
 IMPLEMENT_SHADER_TYPE(,FCompositeObjectDistanceFieldsCS,TEXT("GlobalDistanceField"),TEXT("CompositeObjectDistanceFieldsCS"),SF_Compute);
 
+const int32 HeightfieldCompositeTileSize = 8;
+
+class FCompositeHeightfieldsIntoGlobalDistanceFieldCS : public FGlobalShader
+{
+	DECLARE_SHADER_TYPE(FCompositeHeightfieldsIntoGlobalDistanceFieldCS, Global)
+public:
+
+	static bool ShouldCache(EShaderPlatform Platform)
+	{
+		return IsFeatureLevelSupported(Platform, ERHIFeatureLevel::SM5) && DoesPlatformSupportDistanceFieldAO(Platform);
+	}
+
+	static void ModifyCompilationEnvironment(EShaderPlatform Platform, FShaderCompilerEnvironment& OutEnvironment)
+	{
+		FGlobalShader::ModifyCompilationEnvironment(Platform, OutEnvironment);
+		OutEnvironment.SetDefine(TEXT("COMPOSITE_HEIGHTFIELDS_THREADGROUP_SIZE"), HeightfieldCompositeTileSize);
+	}
+
+	FCompositeHeightfieldsIntoGlobalDistanceFieldCS(const ShaderMetaType::CompiledShaderInitializerType& Initializer)
+		: FGlobalShader(Initializer)
+	{
+		GlobalDistanceFieldParameters.Bind(Initializer.ParameterMap);
+		GlobalDistanceFieldTexture.Bind(Initializer.ParameterMap, TEXT("GlobalDistanceFieldTexture"));
+		UpdateRegionSize.Bind(Initializer.ParameterMap, TEXT("UpdateRegionSize"));
+		VolumeTexelSize.Bind(Initializer.ParameterMap, TEXT("VolumeTexelSize"));
+		UpdateRegionVolumeMin.Bind(Initializer.ParameterMap, TEXT("UpdateRegionVolumeMin"));
+		ClipmapIndex.Bind(Initializer.ParameterMap, TEXT("ClipmapIndex"));
+		AOGlobalMaxSphereQueryRadius.Bind(Initializer.ParameterMap, TEXT("AOGlobalMaxSphereQueryRadius"));
+		HeightfieldDescriptionParameters.Bind(Initializer.ParameterMap);
+		HeightfieldTextureParameters.Bind(Initializer.ParameterMap);
+	}
+
+	FCompositeHeightfieldsIntoGlobalDistanceFieldCS()
+	{
+	}
+
+	void SetParameters(
+		FRHICommandList& RHICmdList,
+		const FScene* Scene,
+		const FSceneView& View,
+		float MaxOcclusionDistance,
+		const FGlobalDistanceFieldInfo& GlobalDistanceFieldInfo,
+		int32 ClipmapIndexValue,
+		const FVolumeUpdateRegion& UpdateRegion,
+		UTexture2D* HeightfieldTextureValue,
+		int32 NumHeightfieldsValue)
+	{
+		FComputeShaderRHIParamRef ShaderRHI = GetComputeShader();
+		FGlobalShader::SetParameters(RHICmdList, ShaderRHI, View);
+		GlobalDistanceFieldParameters.Set(RHICmdList, ShaderRHI, GlobalDistanceFieldInfo.ParameterData);
+
+		const FSceneRenderTargetItem& ClipMapRTI = GlobalDistanceFieldInfo.Clipmaps[ClipmapIndexValue].RenderTarget->GetRenderTargetItem();
+		RHICmdList.TransitionResource(EResourceTransitionAccess::ERWBarrier, EResourceTransitionPipeline::EComputeToCompute, ClipMapRTI.UAV);
+		GlobalDistanceFieldTexture.SetTexture(RHICmdList, ShaderRHI, ClipMapRTI.ShaderResourceTexture, ClipMapRTI.UAV);
+
+		SetShaderValue(RHICmdList, ShaderRHI, UpdateRegionSize, UpdateRegion.CellsSize);
+		SetShaderValue(RHICmdList, ShaderRHI, VolumeTexelSize, FVector(1.0f / GAOGlobalDFResolution));
+		SetShaderValue(RHICmdList, ShaderRHI, UpdateRegionVolumeMin, UpdateRegion.Bounds.Min);
+		SetShaderValue(RHICmdList, ShaderRHI, ClipmapIndex, ClipmapIndexValue);
+
+		extern float GAOConeHalfAngle;
+		const float GlobalMaxSphereQueryRadius = MaxOcclusionDistance / (1 + FMath::Tan(GAOConeHalfAngle));
+		SetShaderValue(RHICmdList, ShaderRHI, AOGlobalMaxSphereQueryRadius, GlobalMaxSphereQueryRadius);
+
+		HeightfieldDescriptionParameters.Set(RHICmdList, ShaderRHI, GetHeightfieldDescriptionsSRV(), NumHeightfieldsValue);
+		HeightfieldTextureParameters.Set(RHICmdList, ShaderRHI, HeightfieldTextureValue, NULL);
+	}
+
+	void UnsetParameters(FRHICommandList& RHICmdList, const FGlobalDistanceFieldInfo& GlobalDistanceFieldInfo, int32 ClipmapIndexValue)
+	{
+		GlobalDistanceFieldTexture.UnsetUAV(RHICmdList, GetComputeShader());
+
+		const FSceneRenderTargetItem& ClipMapRTI = GlobalDistanceFieldInfo.Clipmaps[ClipmapIndexValue].RenderTarget->GetRenderTargetItem();
+		RHICmdList.TransitionResource(EResourceTransitionAccess::EReadable, EResourceTransitionPipeline::EComputeToCompute, ClipMapRTI.UAV);
+	}
+
+	virtual bool Serialize(FArchive& Ar)
+	{
+		bool bShaderHasOutdatedParameters = FGlobalShader::Serialize(Ar);
+		Ar << GlobalDistanceFieldParameters;
+		Ar << GlobalDistanceFieldTexture;
+		Ar << UpdateRegionSize;
+		Ar << VolumeTexelSize;
+		Ar << UpdateRegionVolumeMin;
+		Ar << ClipmapIndex;
+		Ar << AOGlobalMaxSphereQueryRadius;
+		Ar << HeightfieldDescriptionParameters;
+		Ar << HeightfieldTextureParameters;
+		return bShaderHasOutdatedParameters;
+	}
+
+private:
+
+	FGlobalDistanceFieldParameters GlobalDistanceFieldParameters;
+	FRWShaderParameter GlobalDistanceFieldTexture;
+	FShaderParameter UpdateRegionSize;
+	FShaderParameter VolumeTexelSize;
+	FShaderParameter UpdateRegionVolumeMin;
+	FShaderParameter ClipmapIndex;
+	FShaderParameter AOGlobalMaxSphereQueryRadius;
+	FHeightfieldDescriptionParameters HeightfieldDescriptionParameters;
+	FHeightfieldTextureParameters HeightfieldTextureParameters;
+};
+
+IMPLEMENT_SHADER_TYPE(, FCompositeHeightfieldsIntoGlobalDistanceFieldCS, TEXT("GlobalDistanceField"), TEXT("CompositeHeightfieldsIntoGlobalDistanceFieldCS"), SF_Compute);
+
+extern int32 GAOHeightfieldOcclusion;
+extern void UploadHeightfieldDescriptions(const TArray<FHeightfieldComponentDescription>& HeightfieldDescriptions, FVector2D InvLightingAtlasSize, float InvDownsampleFactor);
+
+void FHeightfieldLightingViewInfo::CompositeHeightfieldsIntoGlobalDistanceField(
+	FRHICommandList& RHICmdList,
+	const FScene* Scene,
+	const FViewInfo& View,
+	float MaxOcclusionDistance,
+	const FGlobalDistanceFieldInfo& GlobalDistanceFieldInfo,
+	int32 ClipmapIndexValue,
+	const FVolumeUpdateRegion& UpdateRegion) const
+{
+	const int32 NumPrimitives = Scene->DistanceFieldSceneData.HeightfieldPrimitives.Num();
+
+	if (GAOHeightfieldOcclusion
+		&& NumPrimitives > 0
+		&& SupportsDistanceFieldAO(Scene->GetFeatureLevel(), Scene->GetShaderPlatform()))
+	{
+		FHeightfieldDescription UpdateRegionHeightfield;
+		float LocalToWorldScale = 1;
+
+		for (int32 HeightfieldPrimitiveIndex = 0; HeightfieldPrimitiveIndex < NumPrimitives; HeightfieldPrimitiveIndex++)
+		{
+			const FPrimitiveSceneInfo* HeightfieldPrimitive = Scene->DistanceFieldSceneData.HeightfieldPrimitives[HeightfieldPrimitiveIndex];
+			const FBoxSphereBounds& PrimitiveBounds = HeightfieldPrimitive->Proxy->GetBounds();
+			const float DistanceToPrimitiveSq = (PrimitiveBounds.Origin - View.ViewMatrices.ViewOrigin).SizeSquared();
+
+			if (UpdateRegion.Bounds.Intersect(PrimitiveBounds.GetBox()))
+			{
+				UTexture2D* HeightfieldTexture = NULL;
+				UTexture2D* DiffuseColorTexture = NULL;
+				FHeightfieldComponentDescription NewComponentDescription(HeightfieldPrimitive->Proxy->GetLocalToWorld());
+				HeightfieldPrimitive->Proxy->GetHeightfieldRepresentation(HeightfieldTexture, DiffuseColorTexture, NewComponentDescription);
+
+				if (HeightfieldTexture && HeightfieldTexture->Resource->TextureRHI)
+				{
+					const FIntPoint HeightfieldSize = NewComponentDescription.HeightfieldRect.Size();
+
+					if (UpdateRegionHeightfield.Rect.Area() == 0)
+					{
+						UpdateRegionHeightfield.Rect = NewComponentDescription.HeightfieldRect;
+						LocalToWorldScale = NewComponentDescription.LocalToWorld.GetScaleVector().X;
+					}
+					else
+					{
+						UpdateRegionHeightfield.Rect.Union(NewComponentDescription.HeightfieldRect);
+					}
+
+					TArray<FHeightfieldComponentDescription>& ComponentDescriptions = UpdateRegionHeightfield.ComponentDescriptions.FindOrAdd(FHeightfieldComponentTextures(HeightfieldTexture, DiffuseColorTexture));
+					ComponentDescriptions.Add(NewComponentDescription);
+				}
+			}
+		}
+
+		if (UpdateRegionHeightfield.ComponentDescriptions.Num() > 0)
+		{
+			SCOPED_DRAW_EVENT(RHICmdList, CompositeHeightfields);
+
+			for (TMap<FHeightfieldComponentTextures, TArray<FHeightfieldComponentDescription>>::TConstIterator It(UpdateRegionHeightfield.ComponentDescriptions); It; ++It)
+			{
+				const TArray<FHeightfieldComponentDescription>& HeightfieldDescriptions = It.Value();
+
+				if (HeightfieldDescriptions.Num() > 0)
+				{
+					UploadHeightfieldDescriptions(HeightfieldDescriptions, FVector2D(1, 1), 1.0f / UpdateRegionHeightfield.DownsampleFactor);
+
+					UTexture2D* HeightfieldTexture = It.Key().HeightAndNormal;
+
+					TShaderMapRef<FCompositeHeightfieldsIntoGlobalDistanceFieldCS> ComputeShader(View.ShaderMap);
+					RHICmdList.SetComputeShader(ComputeShader->GetComputeShader());
+					ComputeShader->SetParameters(RHICmdList, Scene, View, MaxOcclusionDistance, GlobalDistanceFieldInfo, ClipmapIndexValue, UpdateRegion, HeightfieldTexture, HeightfieldDescriptions.Num());
+
+					//@todo - match typical update sizes.  Camera movement creates narrow slabs.
+					const uint32 NumGroupsX = FMath::DivideAndRoundUp<int32>(UpdateRegion.CellsSize.X, HeightfieldCompositeTileSize);
+					const uint32 NumGroupsY = FMath::DivideAndRoundUp<int32>(UpdateRegion.CellsSize.Y, HeightfieldCompositeTileSize);
+
+					DispatchComputeShader(RHICmdList, *ComputeShader, NumGroupsX, NumGroupsY, 1);
+					ComputeShader->UnsetParameters(RHICmdList, GlobalDistanceFieldInfo, ClipmapIndexValue);
+				}
+			}
+		}
+	}
+}
+
 /** Constructs and adds an update region based on camera movement for the given axis. */
 void AddUpdateRegionForAxis(FIntVector Movement, const FBox& ClipmapBounds, float CellSize, int32 ComponentIndex, TArray<FVolumeUpdateRegion, TInlineAllocator<3> >& UpdateRegions)
 {
@@ -581,7 +783,7 @@ void AllocateClipmapTexture(FRHICommandListImmediate& RHICmdList, int32 ClipmapI
 /** Staggers clipmap updates so there are only 2 per frame */
 bool ShouldUpdateClipmapThisFrame(int32 ClipmapIndex, int32 GlobalDistanceFieldUpdateIndex)
 {
-	if (ClipmapIndex == 0)
+	if (ClipmapIndex == 0 || !GAOGlobalDistanceFieldStaggeredUpdates)
 	{
 		return true;
 	}
@@ -639,9 +841,13 @@ void ComputeUpdateRegionsAndUpdateViewState(FRHICommandListImmediate& RHICmdList
 			// Accumulate primitive modifications in the viewstate in case we don't update the clipmap this frame
 			ClipmapViewState.PrimitiveModifiedBounds.Append(PrimitiveModifiedBounds);
 
-			if (ShouldUpdateClipmapThisFrame(ClipmapIndex, View.ViewState->GlobalDistanceFieldUpdateIndex)
+			const bool bForceFullUpdate = bReallocated
 				|| !View.ViewState->bIntializedGlobalDistanceFieldOrigins
-				|| bReallocated)
+				// Detect when max occlusion distance has changed
+				|| ClipmapViewState.CachedMaxOcclusionDistance != MaxOcclusionDistance;
+
+			if (ShouldUpdateClipmapThisFrame(ClipmapIndex, View.ViewState->GlobalDistanceFieldUpdateIndex)
+				|| bForceFullUpdate)
 			{
 				const FVector NewCenter = View.ViewMatrices.ViewOrigin;
 
@@ -654,8 +860,7 @@ void ComputeUpdateRegionsAndUpdateViewState(FRHICommandListImmediate& RHICmdList
 				const FBox ClipmapBounds(SnappedCenter - Extent, SnappedCenter + Extent);
 
 				bool bUsePartialUpdates = GAOGlobalDistanceFieldPartialUpdates
-					&& View.ViewState->bIntializedGlobalDistanceFieldOrigins 
-					&& !bReallocated
+					&& !bForceFullUpdate
 					// Only use partial updates with small numbers of primitive modifications
 					&& ClipmapViewState.PrimitiveModifiedBounds.Num() < 100;
 
@@ -667,7 +872,7 @@ void ComputeUpdateRegionsAndUpdateViewState(FRHICommandListImmediate& RHICmdList
 					AddUpdateRegionForAxis(Movement, ClipmapBounds, CellSize, 0, Clipmap.UpdateRegions);
 					AddUpdateRegionForAxis(Movement, ClipmapBounds, CellSize, 1, Clipmap.UpdateRegions);
 					AddUpdateRegionForAxis(Movement, ClipmapBounds, CellSize, 2, Clipmap.UpdateRegions);
-
+					
 					extern float GAOConeHalfAngle;
 					const float GlobalMaxSphereQueryRadius = MaxOcclusionDistance / (1 + FMath::Tan(GAOConeHalfAngle));
 
@@ -711,7 +916,7 @@ void ComputeUpdateRegionsAndUpdateViewState(FRHICommandListImmediate& RHICmdList
 						}
 					}
 
-					if (TotalTexelsBeingUpdated > GAOGlobalDFResolution * GAOGlobalDFResolution * GAOGlobalDFResolution)
+					if (TotalTexelsBeingUpdated >= GAOGlobalDFResolution * GAOGlobalDFResolution * GAOGlobalDFResolution)
 					{
 						// Fall back to a full update if the partial updates were going to do more work
 						bUsePartialUpdates = false;
@@ -742,6 +947,8 @@ void ComputeUpdateRegionsAndUpdateViewState(FRHICommandListImmediate& RHICmdList
 			Clipmap.Bounds = FBox(Center - Extent, Center + Extent);
 			// Scroll offset so the contents of the global distance field don't have to be moved as the camera moves around, only updated in slabs
 			Clipmap.ScrollOffset = FVector(ClipmapViewState.LastPartialUpdateOrigin - ClipmapViewState.FullUpdateOrigin) * CellSize;
+
+			ClipmapViewState.CachedMaxOcclusionDistance = MaxOcclusionDistance;
 		}
 	}
 	else
@@ -784,7 +991,7 @@ void ComputeUpdateRegionsAndUpdateViewState(FRHICommandListImmediate& RHICmdList
  **/
 void UpdateGlobalDistanceFieldVolume(
 	FRHICommandListImmediate& RHICmdList, 
-	const FViewInfo& View, 
+	FViewInfo& View, 
 	const FScene* Scene, 
 	float MaxOcclusionDistance, 
 	FGlobalDistanceFieldInfo& GlobalDistanceFieldInfo)
@@ -793,6 +1000,10 @@ void UpdateGlobalDistanceFieldVolume(
 	{
 		ComputeUpdateRegionsAndUpdateViewState(RHICmdList, View, GlobalDistanceFieldInfo, GMaxGlobalDistanceFieldClipmaps, Scene->DistanceFieldSceneData.PrimitiveModifiedBounds, MaxOcclusionDistance);
 
+		// Recreate the view uniform buffer now that we have updated GlobalDistanceFieldInfo
+		//@todo - minimal recreate
+		View.InitRHIResources(nullptr);
+
 		bool bHasUpdateRegions = false;
 
 		for (int32 ClipmapIndex = 0; ClipmapIndex < GlobalDistanceFieldInfo.Clipmaps.Num(); ClipmapIndex++)
@@ -877,6 +1088,8 @@ void UpdateGlobalDistanceFieldVolume(
 						DispatchComputeShader(RHICmdList, *ComputeShader, NumGroupsX, NumGroupsY, NumGroupsZ);
 						ComputeShader->UnsetParameters(RHICmdList, GlobalDistanceFieldInfo, ClipmapIndex);
 					}
+
+					View.HeightfieldLightingViewInfo.CompositeHeightfieldsIntoGlobalDistanceField(RHICmdList, Scene, View, MaxOcclusionDistance, GlobalDistanceFieldInfo, ClipmapIndex, UpdateRegion);
 				}
 			}
 		}
diff --git a/Engine/Source/Runtime/Renderer/Private/GlobalDistanceField.h b/Engine/Source/Runtime/Renderer/Private/GlobalDistanceField.h
--- a/Engine/Source/Runtime/Renderer/Private/GlobalDistanceField.h
+++ b/Engine/Source/Runtime/Renderer/Private/GlobalDistanceField.h
@@ -28,7 +28,7 @@ inline bool UseGlobalDistanceField(const FDistanceFieldAOParameters& Parameters)
  **/
 extern void UpdateGlobalDistanceFieldVolume(
 	FRHICommandListImmediate& RHICmdList, 
-	const FViewInfo& View, 
+	FViewInfo& View, 
 	const FScene* Scene, 
 	float MaxOcclusionDistance, 
 	FGlobalDistanceFieldInfo& Info);
\ No newline at end of file
