Copying //UE4/Dev-Rendering to Dev-Main (Source //UE4/Dev-Rendering@2932636)

diff --git a/Engine/Source/Runtime/Renderer/Private/HeightfieldLighting.cpp b/Engine/Source/Runtime/Renderer/Private/HeightfieldLighting.cpp
--- a/Engine/Source/Runtime/Renderer/Private/HeightfieldLighting.cpp
+++ b/Engine/Source/Runtime/Renderer/Private/HeightfieldLighting.cpp
@@ -174,49 +174,6 @@ public:
 
 TGlobalResource<FSubsectionHeightfieldDescriptionsResource> GSubsectionHeightfieldDescriptions;
 
-class FHeightfieldTextureParameters
-{
-public:
-	void Bind(const FShaderParameterMap& ParameterMap)
-	{
-		HeightfieldTexture.Bind(ParameterMap,TEXT("HeightfieldTexture"));
-		HeightfieldSampler.Bind(ParameterMap,TEXT("HeightfieldSampler"));
-		DiffuseColorTexture.Bind(ParameterMap,TEXT("DiffuseColorTexture"));
-		DiffuseColorSampler.Bind(ParameterMap,TEXT("DiffuseColorSampler"));
-	}
-
-	friend FArchive& operator<<(FArchive& Ar,FHeightfieldTextureParameters& Parameters)
-	{
-		Ar << Parameters.HeightfieldTexture;
-		Ar << Parameters.HeightfieldSampler;
-		Ar << Parameters.DiffuseColorTexture;
-		Ar << Parameters.DiffuseColorSampler;
-		return Ar;
-	}
-
-	template<typename ShaderRHIParamRef>
-	void Set(FRHICommandList& RHICmdList, const ShaderRHIParamRef ShaderRHI, UTexture2D* HeightfieldTextureValue, UTexture2D* DiffuseColorTextureValue)
-	{
-		//@todo - shouldn't filter the heightfield, it's packed
-		SetTextureParameter(RHICmdList, ShaderRHI, HeightfieldTexture, HeightfieldSampler, TStaticSamplerState<SF_Bilinear>::GetRHI(), HeightfieldTextureValue->Resource->TextureRHI);
-
-		if (DiffuseColorTextureValue)
-		{
-			SetTextureParameter(RHICmdList, ShaderRHI, DiffuseColorTexture, DiffuseColorSampler, TStaticSamplerState<SF_Bilinear>::GetRHI(), DiffuseColorTextureValue->Resource->TextureRHI);
-		}
-		else
-		{
-			SetTextureParameter(RHICmdList, ShaderRHI, DiffuseColorTexture, DiffuseColorSampler, TStaticSamplerState<SF_Bilinear>::GetRHI(), GBlackTexture->TextureRHI);
-		}
-	}
-
-private:
-	FShaderResourceParameter HeightfieldTexture;
-	FShaderResourceParameter HeightfieldSampler;
-	FShaderResourceParameter DiffuseColorTexture;
-	FShaderResourceParameter DiffuseColorSampler;
-};
-
 class FSubsectionHeightfieldDescriptionParameters
 {
 public:
@@ -616,6 +573,11 @@ public:
 
 TGlobalResource<FHeightfieldDescriptionsResource> GHeightfieldDescriptions;
 
+FShaderResourceViewRHIParamRef GetHeightfieldDescriptionsSRV()
+{
+	return GHeightfieldDescriptions.Data.BufferSRV;
+}
+
 void UploadHeightfieldDescriptions(const TArray<FHeightfieldComponentDescription>& HeightfieldDescriptions, FVector2D InvLightingAtlasSize, float InvDownsampleFactor)
 {
 	TArray<FVector4, SceneRenderingAllocator> HeightfieldDescriptionData;
@@ -676,34 +638,6 @@ void UploadHeightfieldDescriptions(const TArray<FHeightfieldComponentDescription
 	RHIUnlockVertexBuffer(GHeightfieldDescriptions.Data.Buffer);
 }
 
-class FHeightfieldDescriptionParameters
-{
-public:
-	void Bind(const FShaderParameterMap& ParameterMap)
-	{
-		HeightfieldDescriptions.Bind(ParameterMap,TEXT("HeightfieldDescriptions"));
-		NumHeightfields.Bind(ParameterMap,TEXT("NumHeightfields"));
-	}
-
-	friend FArchive& operator<<(FArchive& Ar,FHeightfieldDescriptionParameters& Parameters)
-	{
-		Ar << Parameters.HeightfieldDescriptions;
-		Ar << Parameters.NumHeightfields;
-		return Ar;
-	}
-
-	template<typename ShaderRHIParamRef>
-	void Set(FRHICommandList& RHICmdList, const ShaderRHIParamRef ShaderRHI, int32 NumHeightfieldsValue)
-	{
-		SetSRVParameter(RHICmdList, ShaderRHI, HeightfieldDescriptions, GHeightfieldDescriptions.Data.BufferSRV);
-		SetShaderValue(RHICmdList, ShaderRHI, NumHeightfields, NumHeightfieldsValue);
-	}
-
-private:
-	FShaderResourceParameter HeightfieldDescriptions;
-	FShaderParameter NumHeightfields;
-};
-
 class FGlobalHeightfieldParameters
 {
 public:
@@ -764,7 +698,7 @@ public:
 		const FVertexShaderRHIParamRef ShaderRHI = GetVertexShader();
 
 		FGlobalShader::SetParameters(RHICmdList, ShaderRHI, View);
-		HeightfieldDescriptionParameters.Set(RHICmdList, ShaderRHI, NumHeightfieldsValue);
+		HeightfieldDescriptionParameters.Set(RHICmdList, ShaderRHI, GetHeightfieldDescriptionsSRV(), NumHeightfieldsValue);
 	}
 
 	// FShader interface.
@@ -814,7 +748,7 @@ public:
 		const FPixelShaderRHIParamRef ShaderRHI = GetPixelShader();
 
 		FGlobalShader::SetParameters(RHICmdList, ShaderRHI, View);
-		HeightfieldDescriptionParameters.Set(RHICmdList, ShaderRHI, NumHeightfieldsValue);
+		HeightfieldDescriptionParameters.Set(RHICmdList, ShaderRHI, GetHeightfieldDescriptionsSRV(), NumHeightfieldsValue);
 		GlobalHeightfieldParameters.Set(RHICmdList, ShaderRHI, Atlas);
 
 		FVector4 ShadowmapMinMaxValue;
@@ -968,7 +902,7 @@ public:
 		const FPixelShaderRHIParamRef ShaderRHI = GetPixelShader();
 
 		FGlobalShader::SetParameters(RHICmdList, ShaderRHI, View);
-		HeightfieldDescriptionParameters.Set(RHICmdList, ShaderRHI, NumHeightfieldsValue);
+		HeightfieldDescriptionParameters.Set(RHICmdList, ShaderRHI, GetHeightfieldDescriptionsSRV(), NumHeightfieldsValue);
 		GlobalHeightfieldParameters.Set(RHICmdList, ShaderRHI, Atlas);
 		ObjectParameters.Set(RHICmdList, ShaderRHI, CulledObjectBuffers.Buffers);
 
@@ -1132,7 +1066,7 @@ public:
 
 		FMaterialShader::SetParameters(RHICmdList, ShaderRHI, MaterialProxy, *MaterialProxy->GetMaterial(View.GetFeatureLevel()), View, true, ESceneRenderTargetsMode::SetTextures);
 
-		HeightfieldDescriptionParameters.Set(RHICmdList, ShaderRHI, NumHeightfieldsValue);
+		HeightfieldDescriptionParameters.Set(RHICmdList, ShaderRHI, GetHeightfieldDescriptionsSRV(), NumHeightfieldsValue);
 		GlobalHeightfieldParameters.Set(RHICmdList, ShaderRHI, Atlas);
 
 		SetShaderValue(RHICmdList, ShaderRHI, LightDirection, LightSceneInfo.Proxy->GetDirection());
@@ -1304,7 +1238,7 @@ public:
 		FGlobalShader::SetParameters(RHICmdList, ShaderRHI, View);
 
 		AOParameters.Set(RHICmdList, ShaderRHI, Parameters);
-		HeightfieldDescriptionParameters.Set(RHICmdList, ShaderRHI, NumHeightfieldsValue);
+		HeightfieldDescriptionParameters.Set(RHICmdList, ShaderRHI, GetHeightfieldDescriptionsSRV(), NumHeightfieldsValue);
 		HeightfieldTextureParameters.Set(RHICmdList, ShaderRHI, HeightfieldTextureValue, NULL);
 
 		const FScene* Scene = (const FScene*)View.Family->Scene;
@@ -1503,7 +1437,7 @@ public:
 		FGlobalShader::SetParameters(RHICmdList, ShaderRHI, View);
 
 		AOParameters.Set(RHICmdList, ShaderRHI, Parameters);
-		HeightfieldDescriptionParameters.Set(RHICmdList, ShaderRHI, NumHeightfieldsValue);
+		HeightfieldDescriptionParameters.Set(RHICmdList, ShaderRHI, GetHeightfieldDescriptionsSRV(), NumHeightfieldsValue);
 		GlobalHeightfieldParameters.Set(RHICmdList, ShaderRHI, Atlas);
 
 		const FScene* Scene = (const FScene*)View.Family->Scene;
@@ -1717,7 +1651,7 @@ public:
 
 		AOParameters.Set(RHICmdList, ShaderRHI, Parameters);
 		ScreenGridParameters.Set(RHICmdList, ShaderRHI, View, DistanceFieldNormal);
-		HeightfieldDescriptionParameters.Set(RHICmdList, ShaderRHI, NumHeightfieldsValue);
+		HeightfieldDescriptionParameters.Set(RHICmdList, ShaderRHI, GetHeightfieldDescriptionsSRV(), NumHeightfieldsValue);
 		HeightfieldTextureParameters.Set(RHICmdList, ShaderRHI, HeightfieldTextureValue, NULL);
 
 		RHICmdList.TransitionResource(EResourceTransitionAccess::ERWBarrier, EResourceTransitionPipeline::EComputeToCompute, ScreenGridResources.ScreenGridConeVisibility.UAV);
@@ -1868,7 +1802,7 @@ public:
 
 		AOParameters.Set(RHICmdList, ShaderRHI, Parameters);
 		ScreenGridParameters.Set(RHICmdList, ShaderRHI, View, DistanceFieldNormal);
-		HeightfieldDescriptionParameters.Set(RHICmdList, ShaderRHI, NumHeightfieldsValue);
+		HeightfieldDescriptionParameters.Set(RHICmdList, ShaderRHI, GetHeightfieldDescriptionsSRV(), NumHeightfieldsValue);
 
 		RHICmdList.TransitionResource(EResourceTransitionAccess::ERWBarrier, EResourceTransitionPipeline::EComputeToCompute, ScreenGridResources.HeightfieldIrradiance.UAV);
 		HeightfieldIrradiance.SetBuffer(RHICmdList, ShaderRHI, ScreenGridResources.HeightfieldIrradiance);
@@ -2058,7 +1992,7 @@ public:
 		FComputeShaderRHIParamRef ShaderRHI = GetComputeShader();
 		FGlobalShader::SetParameters(RHICmdList, ShaderRHI, View);
 
-		HeightfieldDescriptionParameters.Set(RHICmdList, ShaderRHI, NumHeightfieldsValue);
+		HeightfieldDescriptionParameters.Set(RHICmdList, ShaderRHI, GetHeightfieldDescriptionsSRV(), NumHeightfieldsValue);
 		HeightfieldTextureParameters.Set(RHICmdList, ShaderRHI, HeightfieldTextureValue, NULL);
 
 		RHICmdList.TransitionResource(EResourceTransitionAccess::ERWBarrier, EResourceTransitionPipeline::EComputeToCompute, PreCulledTriangleBuffers.TriangleVisibleMask.UAV);
diff --git a/Engine/Source/Runtime/Renderer/Private/HeightfieldLighting.h b/Engine/Source/Runtime/Renderer/Private/HeightfieldLighting.h
--- a/Engine/Source/Runtime/Renderer/Private/HeightfieldLighting.h
+++ b/Engine/Source/Runtime/Renderer/Private/HeightfieldLighting.h
@@ -131,8 +131,89 @@ public:
 		int32 NumTrianglesValue,
 		const class FUniformMeshBuffers& UniformMeshBuffers) const;
 
+	void CompositeHeightfieldsIntoGlobalDistanceField(
+		FRHICommandList& RHICmdList,
+		const FScene* Scene,
+		const FViewInfo& View,
+		float MaxOcclusionDistance,
+		const class FGlobalDistanceFieldInfo& GlobalDistanceFieldInfo,
+		int32 ClipmapIndexValue,
+		const class FVolumeUpdateRegion& UpdateRegion) const;
+
 private:
 
 	FHeightfieldDescription Heightfield;
 };
 
+extern FShaderResourceViewRHIParamRef GetHeightfieldDescriptionsSRV();
+
+class FHeightfieldDescriptionParameters
+{
+public:
+	void Bind(const FShaderParameterMap& ParameterMap)
+	{
+		HeightfieldDescriptions.Bind(ParameterMap, TEXT("HeightfieldDescriptions"));
+		NumHeightfields.Bind(ParameterMap, TEXT("NumHeightfields"));
+	}
+
+	friend FArchive& operator<<(FArchive& Ar, FHeightfieldDescriptionParameters& Parameters)
+	{
+		Ar << Parameters.HeightfieldDescriptions;
+		Ar << Parameters.NumHeightfields;
+		return Ar;
+	}
+
+	template<typename ShaderRHIParamRef>
+	void Set(FRHICommandList& RHICmdList, const ShaderRHIParamRef ShaderRHI, FShaderResourceViewRHIParamRef HeightfieldDescriptionsValue, int32 NumHeightfieldsValue)
+	{
+		SetSRVParameter(RHICmdList, ShaderRHI, HeightfieldDescriptions, HeightfieldDescriptionsValue);
+		SetShaderValue(RHICmdList, ShaderRHI, NumHeightfields, NumHeightfieldsValue);
+	}
+
+private:
+	FShaderResourceParameter HeightfieldDescriptions;
+	FShaderParameter NumHeightfields;
+};
+
+class FHeightfieldTextureParameters
+{
+public:
+	void Bind(const FShaderParameterMap& ParameterMap)
+	{
+		HeightfieldTexture.Bind(ParameterMap, TEXT("HeightfieldTexture"));
+		HeightfieldSampler.Bind(ParameterMap, TEXT("HeightfieldSampler"));
+		DiffuseColorTexture.Bind(ParameterMap, TEXT("DiffuseColorTexture"));
+		DiffuseColorSampler.Bind(ParameterMap, TEXT("DiffuseColorSampler"));
+	}
+
+	friend FArchive& operator<<(FArchive& Ar, FHeightfieldTextureParameters& Parameters)
+	{
+		Ar << Parameters.HeightfieldTexture;
+		Ar << Parameters.HeightfieldSampler;
+		Ar << Parameters.DiffuseColorTexture;
+		Ar << Parameters.DiffuseColorSampler;
+		return Ar;
+	}
+
+	template<typename ShaderRHIParamRef>
+	void Set(FRHICommandList& RHICmdList, const ShaderRHIParamRef ShaderRHI, UTexture2D* HeightfieldTextureValue, UTexture2D* DiffuseColorTextureValue)
+	{
+		//@todo - shouldn't filter the heightfield, it's packed
+		SetTextureParameter(RHICmdList, ShaderRHI, HeightfieldTexture, HeightfieldSampler, TStaticSamplerState<SF_Bilinear>::GetRHI(), HeightfieldTextureValue->Resource->TextureRHI);
+
+		if (DiffuseColorTextureValue)
+		{
+			SetTextureParameter(RHICmdList, ShaderRHI, DiffuseColorTexture, DiffuseColorSampler, TStaticSamplerState<SF_Bilinear>::GetRHI(), DiffuseColorTextureValue->Resource->TextureRHI);
+		}
+		else
+		{
+			SetTextureParameter(RHICmdList, ShaderRHI, DiffuseColorTexture, DiffuseColorSampler, TStaticSamplerState<SF_Bilinear>::GetRHI(), GBlackTexture->TextureRHI);
+		}
+	}
+
+private:
+	FShaderResourceParameter HeightfieldTexture;
+	FShaderResourceParameter HeightfieldSampler;
+	FShaderResourceParameter DiffuseColorTexture;
+	FShaderResourceParameter DiffuseColorSampler;
+};
diff --git a/Engine/Source/Runtime/Renderer/Private/PlanarReflectionRendering.cpp b/Engine/Source/Runtime/Renderer/Private/PlanarReflectionRendering.cpp
new file mode 100644
--- /dev/null
+++ b/Engine/Source/Runtime/Renderer/Private/PlanarReflectionRendering.cpp
@@ -0,0 +1,156 @@
+// Copyright 1998-2016 Epic Games, Inc. All Rights Reserved.
+
+/*=============================================================================
+ PlanarReflectionRendering.cpp
+=============================================================================*/
+
+#include "RendererPrivate.h"
+#include "ScenePrivate.h"
+#include "SceneFilterRendering.h"
+#include "PostProcessing.h"
+#include "UniformBuffer.h"
+#include "ShaderParameters.h"
+#include "ScreenRendering.h"
+#include "ShaderParameterUtils.h"
+#include "LightRendering.h"
+#include "SceneUtils.h"
+#include "PlanarReflectionSceneProxy.h"
+
+class FPlanarReflectionPS : public FGlobalShader
+{
+	DECLARE_SHADER_TYPE(FPlanarReflectionPS, Global);
+public:
+
+	static bool ShouldCache(EShaderPlatform Platform)
+	{
+		return IsFeatureLevelSupported(Platform, ERHIFeatureLevel::SM4);
+	}
+
+	static void ModifyCompilationEnvironment(EShaderPlatform Platform, FShaderCompilerEnvironment& OutEnvironment)
+	{
+		FGlobalShader::ModifyCompilationEnvironment(Platform, OutEnvironment);
+	}
+
+	/** Default constructor. */
+	FPlanarReflectionPS() {}
+
+	/** Initialization constructor. */
+	FPlanarReflectionPS(const ShaderMetaType::CompiledShaderInitializerType& Initializer)
+		: FGlobalShader(Initializer)
+	{
+		DeferredParameters.Bind(Initializer.ParameterMap);
+		PlanarReflectionParameters.Bind(Initializer.ParameterMap);
+	}
+
+	void SetParameters(FRHICommandList& RHICmdList, const FSceneView& View, FPlanarReflectionSceneProxy* ReflectionSceneProxy)
+	{
+		const FPixelShaderRHIParamRef ShaderRHI = GetPixelShader();
+		FGlobalShader::SetParameters(RHICmdList, ShaderRHI, View);
+		DeferredParameters.Set(RHICmdList, ShaderRHI, View);
+		PlanarReflectionParameters.SetParameters(RHICmdList, ShaderRHI, ReflectionSceneProxy);
+	}
+
+	// FShader interface.
+	virtual bool Serialize(FArchive& Ar) override
+	{
+		bool bShaderHasOutdatedParameters = FGlobalShader::Serialize(Ar);
+		Ar << PlanarReflectionParameters;
+		Ar << DeferredParameters;
+		return bShaderHasOutdatedParameters;
+	}
+
+private:
+
+	FPlanarReflectionParameters PlanarReflectionParameters;
+	FDeferredPixelShaderParameters DeferredParameters;
+};
+
+IMPLEMENT_SHADER_TYPE(,FPlanarReflectionPS,TEXT("PlanarReflectionShaders"),TEXT("PlanarReflectionPS"),SF_Pixel);
+
+bool FDeferredShadingSceneRenderer::RenderDeferredPlanarReflections(FRHICommandListImmediate& RHICmdList, bool bLightAccumulationIsInUse, TRefCountPtr<IPooledRenderTarget>& Output)
+{
+	bool bAnyViewIsReflectionCapture = false;
+	for (int32 ViewIndex = 0, Num = Views.Num(); ViewIndex < Num; ViewIndex++)
+	{
+		const FViewInfo& View = Views[ViewIndex];
+		bAnyViewIsReflectionCapture = bAnyViewIsReflectionCapture || View.bIsPlanarReflection || View.bIsReflectionCapture;
+	}
+
+	// Prevent reflection recursion, or view-dependent planar reflections being seen in reflection captures
+	if (Scene->PlanarReflections.Num() > 0 && !bAnyViewIsReflectionCapture)
+	{
+		bool bSSRAsInput = true;
+
+		if (Output == GSystemTextures.BlackDummy)
+		{
+			bSSRAsInput = false;
+			FSceneRenderTargets& SceneContext = FSceneRenderTargets::Get(RHICmdList);
+
+			if (bLightAccumulationIsInUse)
+			{
+				FPooledRenderTargetDesc Desc(FPooledRenderTargetDesc::Create2DDesc(SceneContext.GetBufferSizeXY(), PF_FloatRGBA, FClearValueBinding::Black, TexCreate_None, TexCreate_RenderTargetable, false));
+				GRenderTargetPool.FindFreeElement(RHICmdList, Desc, Output, TEXT("PlanarReflectionComposite"));
+			}
+			else
+			{
+				Output = SceneContext.LightAccumulation;
+			}
+		}
+
+		SCOPED_DRAW_EVENT(RHICmdList, CompositePlanarReflections);
+
+		SetRenderTarget(RHICmdList, Output->GetRenderTargetItem().TargetableTexture, NULL);
+
+		if (!bSSRAsInput)
+		{
+			RHICmdList.Clear(true, FLinearColor(0, 0, 0, 0), false, 0, false, 0, FIntRect());
+		}
+
+		for (int32 ViewIndex = 0; ViewIndex < Views.Num(); ViewIndex++)
+		{
+			FViewInfo& View = Views[ViewIndex];
+
+			RHICmdList.SetViewport(View.ViewRect.Min.X, View.ViewRect.Min.Y, 0.0f, View.ViewRect.Max.X, View.ViewRect.Max.Y, 1.0f);
+
+			// Blend over previous reflections in the output target (SSR or planar reflections that have already been rendered)
+			// Planar reflections win over SSR and reflection environment
+			//@todo - this is order dependent blending, but ordering is coming from registration order
+			RHICmdList.SetBlendState(TStaticBlendState<CW_RGBA, BO_Add, BF_One, BF_InverseSourceAlpha, BO_Max, BF_One, BF_One>::GetRHI());
+			RHICmdList.SetRasterizerState(TStaticRasterizerState<FM_Solid, CM_None>::GetRHI());
+			RHICmdList.SetDepthStencilState(TStaticDepthStencilState<false, CF_Always>::GetRHI());
+
+			for (int32 PlanarReflectionIndex = 0; PlanarReflectionIndex < Scene->PlanarReflections.Num(); PlanarReflectionIndex++)
+			{
+				FPlanarReflectionSceneProxy* ReflectionSceneProxy = Scene->PlanarReflections[PlanarReflectionIndex];
+
+				SCOPED_DRAW_EVENTF(RHICmdList, PlanarReflection, *ReflectionSceneProxy->OwnerName.ToString());
+
+				TShaderMapRef<TDeferredLightVS<false> > VertexShader(View.ShaderMap);
+				TShaderMapRef<FPlanarReflectionPS> PixelShader(View.ShaderMap);
+
+				static FGlobalBoundShaderState BoundShaderState;
+				SetGlobalBoundShaderState(RHICmdList, FeatureLevel, BoundShaderState, GFilterVertexDeclaration.VertexDeclarationRHI, *VertexShader, *PixelShader);
+
+				PixelShader->SetParameters(RHICmdList, View, ReflectionSceneProxy);
+				VertexShader->SetSimpleLightParameters(RHICmdList, View, FSphere(0));
+
+				DrawRectangle( 
+					RHICmdList,
+					0, 0,
+					View.ViewRect.Width(), View.ViewRect.Height(),
+					View.ViewRect.Min.X, View.ViewRect.Min.Y, 
+					View.ViewRect.Width(), View.ViewRect.Height(),
+					View.ViewRect.Size(),
+					FSceneRenderTargets::Get(RHICmdList).GetBufferSizeXY(),
+					*VertexShader,
+					EDRF_UseTriangleOptimization);
+			}
+		}
+
+		RHICmdList.CopyToResolveTarget(Output->GetRenderTargetItem().TargetableTexture, Output->GetRenderTargetItem().ShaderResourceTexture, false, FResolveParams());
+
+		return true;
+	}
+
+	return false;
+}
diff --git a/Engine/Source/Runtime/Renderer/Private/PostProcess/PostProcessVisualizeComplexity.h b/Engine/Source/Runtime/Renderer/Private/PostProcess/PostProcessVisualizeComplexity.h
--- a/Engine/Source/Runtime/Renderer/Private/PostProcess/PostProcessVisualizeComplexity.h
+++ b/Engine/Source/Runtime/Renderer/Private/PostProcess/PostProcessVisualizeComplexity.h
@@ -53,7 +53,7 @@ public:
 	static void ModifyCompilationEnvironment(EShaderPlatform Platform, FShaderCompilerEnvironment& OutEnvironment)
 	{
 		FGlobalShader::ModifyCompilationEnvironment(Platform, OutEnvironment);
-		OutEnvironment.SetDefine(TEXT("READ_QUAD_OVERDRAW"), AllowDebugViewModeShader(Platform) ? TEXT("1") : TEXT("0"));
+		OutEnvironment.SetDefine(TEXT("READ_QUAD_OVERDRAW"), AllowDebugViewShaderMode(DVSM_QuadComplexity, Platform) ? TEXT("1") : TEXT("0"));
 		OutEnvironment.SetDefine(TEXT("MAX_NUM_COMPLEXITY_COLORS"), MaxNumShaderComplexityColors);
 		// EColorSampling values
 		OutEnvironment.SetDefine(TEXT("CS_RAMP"), (uint32)CS_RAMP);
diff --git a/Engine/Source/Runtime/Renderer/Private/PostProcess/SceneRenderTargets.cpp b/Engine/Source/Runtime/Renderer/Private/PostProcess/SceneRenderTargets.cpp
--- a/Engine/Source/Runtime/Renderer/Private/PostProcess/SceneRenderTargets.cpp
+++ b/Engine/Source/Runtime/Renderer/Private/PostProcess/SceneRenderTargets.cpp
@@ -567,7 +567,7 @@ void FSceneRenderTargets::BeginRenderingGBuffer(FRHICommandList& RHICmdList, ERe
 			bSceneDepthCleared = true;			
 		}		
 
-		if (bBindQuadOverdrawBuffers && RuntimeAllowDebugViewModeShader(CurrentFeatureLevel))
+		if (bBindQuadOverdrawBuffers && AllowDebugViewShaderMode(DVSM_QuadComplexity, CurrentFeatureLevel))
 		{
 			if (QuadOverdrawBuffer.IsValid() && QuadOverdrawBuffer->GetRenderTargetItem().UAV.IsValid())
 			{
@@ -1647,7 +1647,7 @@ void FSceneRenderTargets::AllocateReflectionTargets(FRHICommandList& RHICmdList)
 void FSceneRenderTargets::AllocateDebugViewModeTargets(FRHICommandList& RHICmdList)
 {
 	// If the shader/quad complexity shader need a quad overdraw buffer to be bind, allocate it.
-	if (RuntimeAllowDebugViewModeShader(CurrentFeatureLevel) && AllowDebugViewmodes())
+	if (AllowDebugViewShaderMode(DVSM_QuadComplexity, CurrentFeatureLevel))
 	{
 		FIntPoint QuadOverdrawSize;
 		QuadOverdrawSize.X = 2 * FMath::Max<uint32>((BufferSize.X + 1) / 2, 1); // The size is time 2 since left side is QuadDescriptor, and right side QuadComplexity.
diff --git a/Engine/Source/Runtime/Renderer/Private/PostProcess/ScreenSpaceReflections.cpp b/Engine/Source/Runtime/Renderer/Private/PostProcess/ScreenSpaceReflections.cpp
--- a/Engine/Source/Runtime/Renderer/Private/PostProcess/ScreenSpaceReflections.cpp
+++ b/Engine/Source/Runtime/Renderer/Private/PostProcess/ScreenSpaceReflections.cpp
@@ -48,7 +48,7 @@ static TAutoConsoleVariable<int32> CVarSSRCone(
 	TEXT(" 0 is off (default), 1 is on"),
 	ECVF_RenderThreadSafe);
 
-bool DoScreenSpaceReflections(const FViewInfo& View)
+bool ShouldRenderScreenSpaceReflections(const FViewInfo& View)
 {
 	if(!View.Family->EngineShowFlags.ScreenSpaceReflections)
 	{
@@ -439,7 +439,7 @@ FPooledRenderTargetDesc FRCPassPostProcessScreenSpaceReflections::ComputeOutputD
 	return Ret;
 }
 
-void ScreenSpaceReflections(FRHICommandListImmediate& RHICmdList, FViewInfo& View, TRefCountPtr<IPooledRenderTarget>& SSROutput)
+void RenderScreenSpaceReflections(FRHICommandListImmediate& RHICmdList, FViewInfo& View, TRefCountPtr<IPooledRenderTarget>& SSROutput)
 {
 	FRenderingCompositePassContext CompositeContext(RHICmdList, View);	
 	FPostprocessContext Context(RHICmdList, CompositeContext.Graph, View );
diff --git a/Engine/Source/Runtime/Renderer/Private/PostProcess/ScreenSpaceReflections.h b/Engine/Source/Runtime/Renderer/Private/PostProcess/ScreenSpaceReflections.h
--- a/Engine/Source/Runtime/Renderer/Private/PostProcess/ScreenSpaceReflections.h
+++ b/Engine/Source/Runtime/Renderer/Private/PostProcess/ScreenSpaceReflections.h
@@ -50,6 +50,6 @@ public:
 	virtual FPooledRenderTargetDesc ComputeOutputDesc(EPassOutputId InPassOutputId) const override;
 };
 
-void ScreenSpaceReflections(FRHICommandListImmediate& RHICmdList, FViewInfo& View, TRefCountPtr<IPooledRenderTarget>& SSROutput);
+void RenderScreenSpaceReflections(FRHICommandListImmediate& RHICmdList, FViewInfo& View, TRefCountPtr<IPooledRenderTarget>& SSROutput);
 
-bool DoScreenSpaceReflections(const FViewInfo& View);
+bool ShouldRenderScreenSpaceReflections(const FViewInfo& View);
diff --git a/Engine/Source/Runtime/Renderer/Private/PreCullTriangles.cpp b/Engine/Source/Runtime/Renderer/Private/PreCullTriangles.cpp
--- a/Engine/Source/Runtime/Renderer/Private/PreCullTriangles.cpp
+++ b/Engine/Source/Runtime/Renderer/Private/PreCullTriangles.cpp
@@ -290,8 +290,6 @@ IMPLEMENT_SHADER_TYPE(,FPreCullTrianglesCS,TEXT("PreCullTriangles"),TEXT("PreCul
 TGlobalResource<FPreCulledTriangleBufferResource> GPreCulledTriangleBuffers;
 TGlobalResource<FPreCulledTriangleBufferResource> GPreCulledTriangleBuffersFromHeightfields;
 
-extern bool SupportsDistanceFieldAO(ERHIFeatureLevel::Type FeatureLevel, EShaderPlatform ShaderPlatform);
-
 void FDeferredShadingSceneRenderer::PreCullStaticMeshes(FRHICommandListImmediate& RHICmdList, const TArray<UStaticMeshComponent*>& ComponentsToPreCull, const TArray<TArray<FPlane> >& CullVolumes)
 {
 	FViewInfo& View = Views[0];
diff --git a/Engine/Source/Runtime/Renderer/Private/PrimitiveSceneInfo.cpp b/Engine/Source/Runtime/Renderer/Private/PrimitiveSceneInfo.cpp
--- a/Engine/Source/Runtime/Renderer/Private/PrimitiveSceneInfo.cpp
+++ b/Engine/Source/Runtime/Renderer/Private/PrimitiveSceneInfo.cpp
@@ -84,6 +84,7 @@ FPrimitiveSceneInfo::FPrimitiveSceneInfo(UPrimitiveComponent* InComponent,FScene
 	PrimitiveComponentId(InComponent->ComponentId),
 	ComponentLastRenderTime(&InComponent->LastRenderTime),
 	IndirectLightingCacheAllocation(NULL),
+	CachedPlanarReflectionProxy(NULL),
 	CachedReflectionCaptureProxy(NULL),
 	bNeedsCachedReflectionCaptureUpdate(true),
 	DefaultDynamicHitProxy(NULL),
diff --git a/Engine/Source/Runtime/Renderer/Private/ReflectionEnvironment.cpp b/Engine/Source/Runtime/Renderer/Private/ReflectionEnvironment.cpp
--- a/Engine/Source/Runtime/Renderer/Private/ReflectionEnvironment.cpp
+++ b/Engine/Source/Runtime/Renderer/Private/ReflectionEnvironment.cpp
@@ -909,14 +909,16 @@ void FDeferredShadingSceneRenderer::RenderTiledDeferredImageBasedReflections(FRH
 	{
 		FViewInfo& View = Views[ViewIndex];
 
-		const uint32 bSSR = DoScreenSpaceReflections(Views[ViewIndex]);
+		const uint32 bSSR = ShouldRenderScreenSpaceReflections(Views[ViewIndex]);
 
 		TRefCountPtr<IPooledRenderTarget> SSROutput = GSystemTextures.BlackDummy;
 		if( bSSR )
 		{
-			ScreenSpaceReflections(RHICmdList, View, SSROutput);
+			RenderScreenSpaceReflections(RHICmdList, View, SSROutput);
 		}
 
+		RenderDeferredPlanarReflections(RHICmdList, false, SSROutput);
+
 		// ReflectionEnv is assumed to be on when going into this method
 		{
 			SetRenderTarget(RHICmdList, NULL, NULL);
@@ -1067,14 +1069,19 @@ void FDeferredShadingSceneRenderer::RenderStandardDeferredImageBasedReflections(
 			// If Reflection Environment is active and mixed with indirect lighting (Ambient + LPV), apply is required!
 			|| (View.Family->EngineShowFlags.ReflectionEnvironment && (bReflectionEnv || bEnvironmentMixing) );
 
-		const bool bSSR = DoScreenSpaceReflections(View);
+		const bool bSSR = ShouldRenderScreenSpaceReflections(View);
 
 		TRefCountPtr<IPooledRenderTarget> SSROutput = GSystemTextures.BlackDummy;
 		if (bSSR)
 		{
 			bRequiresApply = true;
 
-			ScreenSpaceReflections(RHICmdList, View, SSROutput);
+			RenderScreenSpaceReflections(RHICmdList, View, SSROutput);
+		}
+
+		if (RenderDeferredPlanarReflections(RHICmdList, true, SSROutput))
+		{
+			bRequiresApply = true;
 		}
 
 	    /* Light Accumulation moved to SceneRenderTargets */
diff --git a/Engine/Source/Runtime/Renderer/Private/RendererScene.cpp b/Engine/Source/Runtime/Renderer/Private/RendererScene.cpp
--- a/Engine/Source/Runtime/Renderer/Private/RendererScene.cpp
+++ b/Engine/Source/Runtime/Renderer/Private/RendererScene.cpp
@@ -17,6 +17,7 @@
 #include "SpeedTreeWind.h"
 #include "HeightfieldLighting.h"
 #include "Components/WindDirectionalSourceComponent.h"
+#include "PlanarReflectionSceneProxy.h"
 
 // Enable this define to do slow checks for components being added to the wrong
 // world's scene, when using PIE. This can happen if a PIE component is reattached
@@ -261,6 +262,17 @@ void FDistanceFieldSceneData::VerifyIntegrity()
 	}
 }
 
+void FScene::UpdateSceneSettings(AWorldSettings* WorldSettings)
+{
+	ENQUEUE_UNIQUE_RENDER_COMMAND_TWOPARAMETER(
+		UpdateSceneSettings,
+		FScene*, Scene, this,
+		float, DefaultMaxDistanceFieldOcclusionDistance, WorldSettings->DefaultMaxDistanceFieldOcclusionDistance,
+	{
+		Scene->DefaultMaxDistanceFieldOcclusionDistance = DefaultMaxDistanceFieldOcclusionDistance;
+	});
+}
+
 /**
  * Sets the FX system associated with the scene.
  */
@@ -416,6 +428,7 @@ FScene::FScene(UWorld* InWorld, bool bInRequiresHitProxies, bool bInIsEditorScen
 ,	UpperDynamicSkylightColor(FLinearColor::Black)
 ,	LowerDynamicSkylightColor(FLinearColor::Black)
 ,	SceneLODHierarchy(this)
+,	DefaultMaxDistanceFieldOcclusionDistance(InWorld->GetWorldSettings()->DefaultMaxDistanceFieldOcclusionDistance)
 ,	NumVisibleLights_GameThread(0)
 ,	NumEnabledSkylights_GameThread(0)
 {
@@ -1209,6 +1222,28 @@ const FReflectionCaptureProxy* FScene::FindClosestReflectionCapture(FVector Posi
 	return ClosestCaptureIndex != INDEX_NONE ? ReflectionSceneData.RegisteredReflectionCaptures[ClosestCaptureIndex] : NULL;
 }
 
+const FPlanarReflectionSceneProxy* FScene::FindClosestPlanarReflection(FVector Position) const
+{
+	checkSlow(IsInParallelRenderingThread());
+	const FPlanarReflectionSceneProxy* ClosestPlanarReflection = NULL;
+	float ClosestDistance = FLT_MAX;
+
+	// Linear search through the scene's planar reflections
+	for (int32 CaptureIndex = 0; CaptureIndex < PlanarReflections.Num(); CaptureIndex++)
+	{
+		FPlanarReflectionSceneProxy* CurrentPlanarReflection = PlanarReflections[CaptureIndex];
+		const float Distance = FMath::Abs(CurrentPlanarReflection->ReflectionPlane.PlaneDot(Position));
+
+		if (Distance < ClosestDistance)
+		{
+			ClosestDistance = Distance;
+			ClosestPlanarReflection = CurrentPlanarReflection;
+		}
+	}
+
+	return ClosestPlanarReflection;
+}
+
 void FScene::FindClosestReflectionCaptures(FVector Position, const FReflectionCaptureProxy* (&SortedByDistanceOUT)[FPrimitiveSceneInfo::MaxCachedReflectionCaptureProxies]) const
 {
 	checkSlow(IsInParallelRenderingThread());
diff --git a/Engine/Source/Runtime/Renderer/Private/SceneCaptureRendering.cpp b/Engine/Source/Runtime/Renderer/Private/SceneCaptureRendering.cpp
--- a/Engine/Source/Runtime/Renderer/Private/SceneCaptureRendering.cpp
+++ b/Engine/Source/Runtime/Renderer/Private/SceneCaptureRendering.cpp
@@ -13,6 +13,8 @@
 #include "PostProcessAmbient.h"
 #include "PostProcessing.h"
 #include "SceneUtils.h"
+#include "Components/PlanarReflectionComponent.h"
+#include "PlanarReflectionSceneProxy.h"
 
 // Copies into render target, optionally flipping it in the Y-axis
 static void CopyCaptureToTarget(FRHICommandListImmediate& RHICmdList, const FRenderTarget* Target, const FIntPoint& TargetSize, FViewInfo& View, const FIntRect& ViewRect, FTextureRHIParamRef SourceTextureRHI, bool bNeedsFlippedRenderTarget)
@@ -63,27 +65,8 @@ static void CopyCaptureToTarget(FRHICommandListImmediate& RHICmdList, const FRen
 	}
 }
 
-static void UpdateSceneCaptureContent_RenderThread(
-	FRHICommandListImmediate& RHICmdList, 
-	FSceneRenderer* SceneRenderer, 
-	FTextureRenderTargetResource* TextureRenderTarget, 
-	const FName OwnerName, 
-	const FResolveParams& ResolveParams, 
-	bool bUseSceneColorTexture, 
-	bool bIsPlanarReflection, 
-	FVector4& ReflectionPlane
-	)
+static void UpdateSceneCaptureContent_RenderThread(FRHICommandListImmediate& RHICmdList, FSceneRenderer* SceneRenderer, FRenderTarget* RenderTarget, FTexture* RenderTargetTexture, const FName OwnerName, const FResolveParams& ResolveParams, bool bUseSceneColorTexture)
 {
-	// Early out?
-	if (bIsPlanarReflection)
-	{
-		static const auto* const PlanarReflectionCVar = IConsoleManager::Get().FindTConsoleVariableDataInt(TEXT("r.EnablePlanarReflections"));
-		if (PlanarReflectionCVar && PlanarReflectionCVar->GetValueOnRenderThread() == 0)
-		{
-			return;
-		}
-	}
-
 	FMemMark MemStackMark(FMemStack::Get());
 
 	// update any resources that needed a deferred update
@@ -110,9 +93,9 @@ static void UpdateSceneCaptureContent_RenderThread(
 		if (bNeedsFlippedRenderTarget)
 		{
 			// We need to use an intermediate render target since the result will be flipped
-			auto& RenderTarget = Target->GetRenderTargetTexture();
+			auto& RenderTargetRHI = Target->GetRenderTargetTexture();
 			FPooledRenderTargetDesc Desc(FPooledRenderTargetDesc::Create2DDesc(Target->GetSizeXY(), 
-				RenderTarget.GetReference()->GetFormat(), 
+				RenderTargetRHI.GetReference()->GetFormat(), 
 				FClearValueBinding::None,
 				TexCreate_None, 
 				TexCreate_RenderTargetable,
@@ -151,10 +134,6 @@ static void UpdateSceneCaptureContent_RenderThread(
 				SceneRenderer->ViewFamily.RenderTarget = &FlippedRenderTarget; //-V506
 			}
 
-			// Setup planar reflection
-			View.bIsPlanarReflectionCapture = bIsPlanarReflection;
-			View.ReflectionPlane = ReflectionPlane;
-
 			SceneRenderer->Render(RHICmdList);
 
 			if (bNeedsFlippedRenderTarget)
@@ -171,32 +150,70 @@ static void UpdateSceneCaptureContent_RenderThread(
 			SCOPED_DRAW_EVENT(RHICmdList, FlipCapture);
 			CopyCaptureToTarget(RHICmdList, Target, TargetSize, View, ViewRect, FlippedRenderTarget.GetTextureParamRef(), true);
 		}
-		else if (bUseSceneColorTexture && (!bIsPlanarReflection || SceneRenderer->FeatureLevel >= ERHIFeatureLevel::SM4))
+		else if (bUseSceneColorTexture)
 		{
 			// Copy the captured scene into the destination texture (only required on HDR or deferred as that implies post-processing)
 			SCOPED_DRAW_EVENT(RHICmdList, CaptureSceneColor);
 			CopyCaptureToTarget(RHICmdList, Target, TargetSize, View, ViewRect, FSceneRenderTargets::Get(RHICmdList).GetSceneColorTexture(), false);
 		}
 
-		RHICmdList.CopyToResolveTarget(TextureRenderTarget->GetRenderTargetTexture(), TextureRenderTarget->TextureRHI, false, ResolveParams);
+		RHICmdList.CopyToResolveTarget(RenderTarget->GetRenderTargetTexture(), RenderTargetTexture->TextureRHI, false, ResolveParams);
 	}
 	FSceneRenderer::WaitForTasksClearSnapshotsAndDeleteSceneRenderer(RHICmdList, SceneRenderer);
 }
 
-
-FSceneRenderer* FScene::CreateSceneRenderer( USceneCaptureComponent* SceneCaptureComponent, UTextureRenderTarget* TextureTarget, const FMatrix& ViewRotationMatrix, const FVector& ViewLocation, float FOV, float MaxViewDistance, bool bCaptureSceneColour, FPostProcessSettings* PostProcessSettings, float PostProcessBlendWeight )
+void BuildProjectionMatrix(FIntPoint RenderTargetSize, float FOV, FMatrix& ProjectionMatrix)
 {
-	FIntPoint CaptureSize(TextureTarget->GetSurfaceWidth(), TextureTarget->GetSurfaceHeight());
+	float XAxisMultiplier;
+	float YAxisMultiplier;
+
+	if (RenderTargetSize.X > RenderTargetSize.Y)
+	{
+		// if the viewport is wider than it is tall
+		XAxisMultiplier = 1.0f;
+		YAxisMultiplier = RenderTargetSize.X / (float)RenderTargetSize.Y;
+	}
+	else
+	{
+		// if the viewport is taller than it is wide
+		XAxisMultiplier = RenderTargetSize.Y / (float)RenderTargetSize.X;
+		YAxisMultiplier = 1.0f;
+	}
 
-	FTextureRenderTargetResource* Resource = TextureTarget->GameThread_GetRenderTargetResource();
+	if ((int32)ERHIZBuffer::IsInverted != 0)
+	{
+		ProjectionMatrix = FReversedZPerspectiveMatrix(
+			FOV,
+			FOV,
+			XAxisMultiplier,
+			YAxisMultiplier,
+			GNearClippingPlane,
+			GNearClippingPlane
+			);
+	}
+	else
+	{
+		ProjectionMatrix = FPerspectiveMatrix(
+			FOV,
+			FOV,
+			XAxisMultiplier,
+			YAxisMultiplier,
+			GNearClippingPlane,
+			GNearClippingPlane
+			);
+	}
+}
+
+FSceneRenderer* CreateSceneRendererForSceneCapture(FScene* Scene, USceneCaptureComponent* SceneCaptureComponent, FRenderTarget* RenderTarget, FIntPoint RenderTargetSize, const FMatrix& ViewRotationMatrix, const FVector& ViewLocation, const FMatrix& ProjectionMatrix, float MaxViewDistance, bool bCaptureSceneColour, FPostProcessSettings* PostProcessSettings, float PostProcessBlendWeight)
+{
 	FSceneViewFamilyContext ViewFamily(FSceneViewFamily::ConstructionValues(
-		Resource,
-		this,
+		RenderTarget,
+		Scene,
 		SceneCaptureComponent->ShowFlags)
 		.SetResolveScene(!bCaptureSceneColour));
 
 	FSceneViewInitOptions ViewInitOptions;
-	ViewInitOptions.SetViewRectangle(FIntRect(0, 0, CaptureSize.X, CaptureSize.Y));
+	ViewInitOptions.SetViewRectangle(FIntRect(0, 0, RenderTargetSize.X, RenderTargetSize.Y));
 	ViewInitOptions.ViewFamily = &ViewFamily;
 	ViewInitOptions.ViewOrigin = ViewLocation;
 	ViewInitOptions.ViewRotationMatrix = ViewRotationMatrix;
@@ -204,6 +221,7 @@ FSceneRenderer* FScene::CreateSceneRenderer( USceneCaptureComponent* SceneCaptur
 	ViewInitOptions.OverrideFarClippingPlaneDistance = MaxViewDistance;
 	ViewInitOptions.SceneViewStateInterface = SceneCaptureComponent->GetViewState();
     ViewInitOptions.StereoPass = SceneCaptureComponent->CaptureStereoPass;
+	ViewInitOptions.ProjectionMatrix = ProjectionMatrix;
 
 	if (bCaptureSceneColour)
 	{
@@ -211,48 +229,6 @@ FSceneRenderer* FScene::CreateSceneRenderer( USceneCaptureComponent* SceneCaptur
 		ViewInitOptions.OverlayColor = FLinearColor::Black;
 	}
 
-	// Build projection matrix
-	{
-		float XAxisMultiplier;
-		float YAxisMultiplier;
-
-		if (CaptureSize.X > CaptureSize.Y)
-		{
-			// if the viewport is wider than it is tall
-			XAxisMultiplier = 1.0f;
-			YAxisMultiplier = CaptureSize.X / (float)CaptureSize.Y;
-		}
-		else
-		{
-			// if the viewport is taller than it is wide
-			XAxisMultiplier = CaptureSize.Y / (float)CaptureSize.X;
-			YAxisMultiplier = 1.0f;
-		}
-
-		if ((int32)ERHIZBuffer::IsInverted != 0)
-		{
-			ViewInitOptions.ProjectionMatrix = FReversedZPerspectiveMatrix(
-				FOV,
-				FOV,
-				XAxisMultiplier,
-				YAxisMultiplier,
-				GNearClippingPlane,
-				GNearClippingPlane
-				);
-		}
-		else
-		{
-			ViewInitOptions.ProjectionMatrix = FPerspectiveMatrix(
-				FOV,
-				FOV,
-				XAxisMultiplier,
-				YAxisMultiplier,
-				GNearClippingPlane,
-				GNearClippingPlane
-				);
-		}
-	}
-
 	FSceneView* View = new FSceneView(ViewInitOptions);
 
 	View->bIsSceneCapture = true;
@@ -281,77 +257,41 @@ void FScene::UpdateSceneCaptureContents(USceneCaptureComponent2D* CaptureCompone
 {
 	check(CaptureComponent);
 
-	FVector ViewLocation;
-	FMatrix ViewRotationMatrix;
-	float FOV = CaptureComponent->FOVAngle * (float)PI / 360.0f;
 	if (CaptureComponent->TextureTarget)
 	{
-		//!! HACK!
-
-		// Attempt to query the local player so we can use the players view for planar reflections.
-		const APlayerController* const LocalPlayer = UGameplayStatics::GetPlayerController(CaptureComponent, 0);
-		
-		// Standard capture
-		if (!CaptureComponent->bIsPlanarReflection || LocalPlayer == nullptr)
-		{
-			FTransform Transform = CaptureComponent->GetComponentToWorld();
-			ViewLocation = Transform.GetTranslation();
-
-			// Remove the translation from Transform because we only need rotation.
-			Transform.SetTranslation(FVector::ZeroVector);
-			ViewRotationMatrix = Transform.ToInverseMatrixWithScale();
-
-			// swap axis st. x=z,y=x,z=y (unreal coord space) so that z is up
-			ViewRotationMatrix = ViewRotationMatrix * FMatrix(
-				FPlane(0, 0, 1, 0),
-				FPlane(1, 0, 0, 0),
-				FPlane(0, 1, 0, 0),
-				FPlane(0, 0, 0, 1));
-		}
-
-		// Planar reflection capture
-		else
-		{
-			FRotator ViewRotation;
-			LocalPlayer->GetPlayerViewPoint(ViewLocation, ViewRotation);
-
-			const FTransform Rotation(ViewRotation);
-			ViewRotationMatrix = Rotation.ToInverseMatrixWithScale();
-
-			// swap axis st. x=z,y=x,z=y (unreal coord space) so that z is up
-			ViewRotationMatrix *= FMatrix(
-				FPlane(0, 0, 1, 0),
-				FPlane(1, 0, 0, 0),
-				FPlane(0, 1, 0, 0),
-				FPlane(0, 0, 0, 1));
-
-			if (LocalPlayer->PlayerCameraManager != nullptr)
-			{
-				FOV = LocalPlayer->PlayerCameraManager->GetFOVAngle() * (float)PI / 360.0f;
-			}
-		}
-		
+		FTransform Transform = CaptureComponent->GetComponentToWorld();
+		FVector ViewLocation = Transform.GetTranslation();
+
+		// Remove the translation from Transform because we only need rotation.
+		Transform.SetTranslation(FVector::ZeroVector);
+		FMatrix ViewRotationMatrix = Transform.ToInverseMatrixWithScale();
+
+		// swap axis st. x=z,y=x,z=y (unreal coord space) so that z is up
+		ViewRotationMatrix = ViewRotationMatrix * FMatrix(
+			FPlane(0, 0, 1, 0),
+			FPlane(1, 0, 0, 0),
+			FPlane(0, 1, 0, 0),
+			FPlane(0, 0, 0, 1));
+		const float FOV = CaptureComponent->FOVAngle * (float)PI / 360.0f;
 		const bool bUseSceneColorTexture = CaptureComponent->CaptureSource == SCS_SceneColorHDR;
-		const bool bIsPlanarReflection = CaptureComponent->bIsPlanarReflection;
+		FIntPoint CaptureSize(CaptureComponent->TextureTarget->GetSurfaceWidth(), CaptureComponent->TextureTarget->GetSurfaceHeight());
 
-		const FVector ReflectionPlaneNormal(CaptureComponent->ReflectionPlaneNormal.GetSafeNormal() * static_cast<float>(bIsPlanarReflection));
-		const FVector4 ReflectionPlane(ReflectionPlaneNormal, CaptureComponent->ReflectionPlaneHeight);
-				
-		FSceneRenderer* SceneRenderer = CreateSceneRenderer(CaptureComponent, CaptureComponent->TextureTarget, ViewRotationMatrix , ViewLocation, FOV, CaptureComponent->MaxViewDistanceOverride, bUseSceneColorTexture, &CaptureComponent->PostProcessSettings, CaptureComponent->PostProcessBlendWeight);
+		FMatrix ProjectionMatrix;
+		BuildProjectionMatrix(CaptureSize, FOV, ProjectionMatrix);
 
+		FSceneRenderer* SceneRenderer = CreateSceneRendererForSceneCapture(this, CaptureComponent, CaptureComponent->TextureTarget->GameThread_GetRenderTargetResource(), CaptureSize, ViewRotationMatrix, ViewLocation, ProjectionMatrix, CaptureComponent->MaxViewDistanceOverride, bUseSceneColorTexture, &CaptureComponent->PostProcessSettings, CaptureComponent->PostProcessBlendWeight);
+				
 		FTextureRenderTargetResource* TextureRenderTarget = CaptureComponent->TextureTarget->GameThread_GetRenderTargetResource();
 		const FName OwnerName = CaptureComponent->GetOwner() ? CaptureComponent->GetOwner()->GetFName() : NAME_None;
 
-		ENQUEUE_UNIQUE_RENDER_COMMAND_SIXPARAMETER( 
+		ENQUEUE_UNIQUE_RENDER_COMMAND_FOURPARAMETER(
 			CaptureCommand,
 			FSceneRenderer*, SceneRenderer, SceneRenderer,
 			FTextureRenderTargetResource*, TextureRenderTarget, TextureRenderTarget,
 			FName, OwnerName, OwnerName,
 			bool, bUseSceneColorTexture, bUseSceneColorTexture, 
-			bool, bIsPlanarReflection, bIsPlanarReflection,
-			FVector4, ReflectionPlane, ReflectionPlane,
 		{
-			UpdateSceneCaptureContent_RenderThread(RHICmdList, SceneRenderer, TextureRenderTarget, OwnerName, FResolveParams(), bUseSceneColorTexture, bIsPlanarReflection, ReflectionPlane);
+			UpdateSceneCaptureContent_RenderThread(RHICmdList, SceneRenderer, TextureRenderTarget, TextureRenderTarget, OwnerName, FResolveParams(), bUseSceneColorTexture);
 		});
 	}
 }
@@ -410,23 +350,153 @@ void FScene::UpdateSceneCaptureContents(USceneCaptureComponentCube* CaptureCompo
 			const ECubeFace TargetFace = (ECubeFace)faceidx;
 			const FVector Location = CaptureComponent->GetComponentToWorld().GetTranslation();
 			const FMatrix ViewRotationMatrix = FLocal::CalcCubeFaceTransform(TargetFace);
-			FSceneRenderer* SceneRenderer = CreateSceneRenderer(CaptureComponent, CaptureComponent->TextureTarget, ViewRotationMatrix, Location, FOV, CaptureComponent->MaxViewDistanceOverride);
+			FIntPoint CaptureSize(CaptureComponent->TextureTarget->GetSurfaceWidth(), CaptureComponent->TextureTarget->GetSurfaceHeight());
+			FMatrix ProjectionMatrix;
+			BuildProjectionMatrix(CaptureSize, FOV, ProjectionMatrix);
+
+			FSceneRenderer* SceneRenderer = CreateSceneRendererForSceneCapture(this, CaptureComponent, CaptureComponent->TextureTarget->GameThread_GetRenderTargetResource(), CaptureSize, ViewRotationMatrix, Location, ProjectionMatrix, CaptureComponent->MaxViewDistanceOverride, true, NULL, 0);
 
 			FTextureRenderTargetCubeResource* TextureRenderTarget = static_cast<FTextureRenderTargetCubeResource*>(CaptureComponent->TextureTarget->GameThread_GetRenderTargetResource());
 			const FName OwnerName = CaptureComponent->GetOwner() ? CaptureComponent->GetOwner()->GetFName() : NAME_None;
-			const FVector4 ReflectionPlane(0.0f);
 
-			ENQUEUE_UNIQUE_RENDER_COMMAND_FIVEPARAMETER( 
+			ENQUEUE_UNIQUE_RENDER_COMMAND_FOURPARAMETER(
 				CaptureCommand,
 				FSceneRenderer*, SceneRenderer, SceneRenderer,
 				FTextureRenderTargetCubeResource*, TextureRenderTarget, TextureRenderTarget,
 				FName, OwnerName, OwnerName,
 				ECubeFace, TargetFace, TargetFace,
-				FVector4, ReflectionPlane, ReflectionPlane,
 			{
-				UpdateSceneCaptureContent_RenderThread(RHICmdList, SceneRenderer, TextureRenderTarget, OwnerName, FResolveParams(FResolveRect(), TargetFace), true, false, ReflectionPlane);
+				UpdateSceneCaptureContent_RenderThread(RHICmdList, SceneRenderer, TextureRenderTarget, TextureRenderTarget, OwnerName, FResolveParams(FResolveRect(), TargetFace), true);
+			});
+		}
+	}
+}
+
+void FScene::UpdatePlanarReflectionContents(UPlanarReflectionComponent* CaptureComponent, FSceneRenderer& MainSceneRenderer)
+{
+	check(CaptureComponent);
+
+	{
+		//@todo - splitscreen / VR
+		const FSceneView& ParentView = *MainSceneRenderer.ViewFamily.Views[0];
+
+		if (CaptureComponent->RenderTarget != NULL && CaptureComponent->RenderTarget->GetSizeXY() != ParentView.ViewRect.Size())
+		{
+			ENQUEUE_UNIQUE_RENDER_COMMAND_ONEPARAMETER( 
+				ReleaseRenderTargetCommand,
+				FPlanarReflectionRenderTarget*, RenderTarget, CaptureComponent->RenderTarget,
+			{
+				RenderTarget->ReleaseResource();
+				delete RenderTarget;
+			});
+
+			CaptureComponent->RenderTarget = NULL;
+		}
+
+		if (CaptureComponent->RenderTarget == NULL)
+		{
+			CaptureComponent->RenderTarget = new FPlanarReflectionRenderTarget(ParentView.ViewRect.Size());
+
+			ENQUEUE_UNIQUE_RENDER_COMMAND_TWOPARAMETER( 
+				InitRenderTargetCommand,
+				FPlanarReflectionRenderTarget*, RenderTarget, CaptureComponent->RenderTarget,
+				FPlanarReflectionSceneProxy*, SceneProxy, CaptureComponent->SceneProxy,
+			{
+				RenderTarget->InitResource();
+				
+				if (SceneProxy)
+				{
+					SceneProxy->RenderTarget = RenderTarget;
+				}
 			});
 		}
+
+		FMatrix ComponentTransform = CaptureComponent->ComponentToWorld.ToMatrixWithScale();
+		FPlane MirrorPlane = FPlane(ComponentTransform.TransformPosition(FVector::ZeroVector), ComponentTransform.TransformVector(FVector(0, 0, 1)));
+
+		// Create a mirror matrix and premultiply the view transform by it
+		FMirrorMatrix MirrorMatrix(MirrorPlane);
+		FMatrix ViewMatrix(MirrorMatrix * ParentView.ViewMatrices.ViewMatrix);
+		FVector ViewLocation = ViewMatrix.InverseTransformPosition(FVector::ZeroVector);
+		FMatrix ViewRotationMatrix = ViewMatrix.RemoveTranslation();
+		float FOV = FMath::Atan(1.0f / ParentView.ViewMatrices.ProjMatrix.M[0][0]);
+
+		FMatrix ProjectionMatrix;
+		BuildProjectionMatrix(ParentView.ViewRect.Size(), FOV + CaptureComponent->ExtraFOV * (float)PI / 180.0f, ProjectionMatrix);
+
+		FSceneRenderer* SceneRenderer = CreateSceneRendererForSceneCapture(this, CaptureComponent, CaptureComponent->RenderTarget, ParentView.ViewRect.Size(), ViewRotationMatrix, ViewLocation, ProjectionMatrix, CaptureComponent->MaxViewDistanceOverride, true, NULL, 0);
+
+		CaptureComponent->ProjectionWithExtraFOV = ProjectionMatrix;
+
+		if (CaptureComponent->SceneProxy)
+		{
+			ENQUEUE_UNIQUE_RENDER_COMMAND_TWOPARAMETER( 
+				UpdateProxyCommand,
+				FMatrix, ProjectionMatrix, ProjectionMatrix,
+				FPlanarReflectionSceneProxy*, SceneProxy, CaptureComponent->SceneProxy,
+			{
+				SceneProxy->ProjectionWithExtraFOV = ProjectionMatrix;
+			});
+		}
+		
+		SceneRenderer->Views[0].bIsPlanarReflection = true;
+		SceneRenderer->Views[0].GlobalClippingPlane = MirrorPlane;
+
+		const FName OwnerName = CaptureComponent->GetOwner() ? CaptureComponent->GetOwner()->GetFName() : NAME_None;
+
+		ENQUEUE_UNIQUE_RENDER_COMMAND_THREEPARAMETER( 
+			CaptureCommand,
+			FSceneRenderer*, SceneRenderer, SceneRenderer,
+			FPlanarReflectionRenderTarget*, RenderTarget, CaptureComponent->RenderTarget,
+			FName, OwnerName, OwnerName,
+		{
+			QUICK_SCOPE_CYCLE_COUNTER(STAT_RenderPlanarReflection);
+			UpdateSceneCaptureContent_RenderThread(RHICmdList, SceneRenderer, RenderTarget, RenderTarget, OwnerName, FResolveParams(), true);
+		});
 	}
 }
 
+void FScene::AddPlanarReflection(UPlanarReflectionComponent* Component)
+{
+	check(Component->SceneProxy);
+	PlanarReflections_GameThread.Add(Component);
+
+	ENQUEUE_UNIQUE_RENDER_COMMAND_TWOPARAMETER(
+		FAddPlanarReflectionCommand,
+		FPlanarReflectionSceneProxy*,SceneProxy,Component->SceneProxy,
+		FScene*,Scene,this,
+	{
+		Scene->ReflectionSceneData.bRegisteredReflectionCapturesHasChanged = true;
+		Scene->PlanarReflections.Add(SceneProxy);
+	});
+}
+
+void FScene::RemovePlanarReflection(UPlanarReflectionComponent* Component) 
+{
+	check(Component->SceneProxy);
+	PlanarReflections_GameThread.Remove(Component);
+
+	ENQUEUE_UNIQUE_RENDER_COMMAND_TWOPARAMETER(
+		FRemovePlanarReflectionCommand,
+		FPlanarReflectionSceneProxy*,SceneProxy,Component->SceneProxy,
+		FScene*,Scene,this,
+	{
+		Scene->ReflectionSceneData.bRegisteredReflectionCapturesHasChanged = true;
+		Scene->PlanarReflections.Remove(SceneProxy);
+	});
+}
+
+void FScene::UpdatePlanarReflectionTransform(UPlanarReflectionComponent* Component)
+{	
+	check(Component->SceneProxy);
+
+	ENQUEUE_UNIQUE_RENDER_COMMAND_THREEPARAMETER(
+		FUpdatePlanarReflectionCommand,
+		FPlanarReflectionSceneProxy*,SceneProxy,Component->SceneProxy,
+		FMatrix,Transform,Component->ComponentToWorld.ToMatrixWithScale(),
+		FScene*,Scene,this,
+	{
+		Scene->ReflectionSceneData.bRegisteredReflectionCapturesHasChanged = true;
+		SceneProxy->UpdateTransform(Transform);
+	});
+}
diff --git a/Engine/Source/Runtime/Renderer/Private/ScenePrivate.h b/Engine/Source/Runtime/Renderer/Private/ScenePrivate.h
--- a/Engine/Source/Runtime/Renderer/Private/ScenePrivate.h
+++ b/Engine/Source/Runtime/Renderer/Private/ScenePrivate.h
@@ -321,11 +321,13 @@ public:
 	{
 		FullUpdateOrigin = FIntVector::ZeroValue;
 		LastPartialUpdateOrigin = FIntVector::ZeroValue;
+		CachedMaxOcclusionDistance = 0;
 	}
 
 	FIntVector FullUpdateOrigin;
 	FIntVector LastPartialUpdateOrigin;
 	TArray<FVector4> PrimitiveModifiedBounds;
+	float CachedMaxOcclusionDistance;
 	TRefCountPtr<IPooledRenderTarget> VolumeTexture;
 };
 
@@ -1703,6 +1705,9 @@ public:
 	/** Potential capsule shadow casters registered to the scene. */
 	TArray<FPrimitiveSceneInfo*> CapsuleIndirectCasterPrimitives; 
 
+	TArray<class FPlanarReflectionSceneProxy*> PlanarReflections;
+	TArray<class UPlanarReflectionComponent*> PlanarReflections_GameThread;
+
 	/** State needed for the reflection environment feature. */
 	FReflectionEnvironmentSceneData ReflectionSceneData;
 
@@ -1775,6 +1780,8 @@ public:
 	/** LOD Tree Holder for massive LOD system */
 	FLODSceneTree SceneLODHierarchy;
 
+	float DefaultMaxDistanceFieldOcclusionDistance;
+
 	/** Initialization constructor. */
 	FScene(UWorld* InWorld, bool bInRequiresHitProxies,bool bInIsEditorScene, bool bCreateFXSystem, ERHIFeatureLevel::Type InFeatureLevel);
 
@@ -1800,8 +1807,12 @@ public:
 	virtual void GetReflectionCaptureData(UReflectionCaptureComponent* Component, class FReflectionCaptureFullHDRDerivedData& OutDerivedData) override;
 	virtual void UpdateReflectionCaptureTransform(UReflectionCaptureComponent* Component) override;
 	virtual void ReleaseReflectionCubemap(UReflectionCaptureComponent* CaptureComponent) override;
+	virtual void AddPlanarReflection(class UPlanarReflectionComponent* Component) override;
+	virtual void RemovePlanarReflection(UPlanarReflectionComponent* Component) override;
+	virtual void UpdatePlanarReflectionTransform(UPlanarReflectionComponent* Component) override;
 	virtual void UpdateSceneCaptureContents(class USceneCaptureComponent2D* CaptureComponent) override;
 	virtual void UpdateSceneCaptureContents(class USceneCaptureComponentCube* CaptureComponent) override;
+	virtual void UpdatePlanarReflectionContents(UPlanarReflectionComponent* CaptureComponent, FSceneRenderer& MainSceneRenderer) override;
 	virtual void AllocateReflectionCaptures(const TArray<UReflectionCaptureComponent*>& NewCaptures) override;
 	virtual void UpdateSkyCaptureContents(const USkyLightComponent* CaptureComponent, bool bCaptureEmissiveOnly, UTextureCube* SourceCubemap, FTexture* OutProcessedTexture, FSHVectorRGB3& OutIrradianceEnvironmentMap) override; 
 	virtual void PreCullStaticMeshes(const TArray<UStaticMeshComponent*>& ComponentsToPreCull, const TArray<TArray<FPlane> >& CullVolumes) override;
@@ -1866,6 +1877,7 @@ public:
 	/** Finds the closest reflection capture to a point in space. */
 	const FReflectionCaptureProxy* FindClosestReflectionCapture(FVector Position) const;
 
+	const class FPlanarReflectionSceneProxy* FindClosestPlanarReflection(FVector Position) const;
 
 	void FindClosestReflectionCaptures(FVector Position, const FReflectionCaptureProxy* (&SortedByDistanceOUT)[FPrimitiveSceneInfo::MaxCachedReflectionCaptureProxies]) const;
 	
@@ -1895,6 +1907,8 @@ public:
 		return this;
 	}
 
+	virtual void UpdateSceneSettings(AWorldSettings* WorldSettings) override;
+
 	/**
 	 * Sets the FX system associated with the scene.
 	 */
@@ -2031,22 +2045,6 @@ private:
 	 */
 	void OnLevelAddedToWorld_RenderThread(FName InLevelName);
 
-	/**
-	 * Builds a FSceneRenderer instance for a given view and render target. Helper function for Scene capture code.
-	 *
-	 * @param	SceneCaptureComponent - The scene capture component for which to create a scene renderer.
-	 * @param	TextureTarget - render target to draw to.
-	 * @param	ViewRotationMatrix - Camera rotation matrix.
-	 * @param	ViewLocation - Camera location.
-	 * @param	FOV - Camera field of view.
-	 * @param	MaxViewDistance - Far draw distance.
-	 * @param	bCaptureSceneColour if true the returned scenerenderer will ignore post process operations as only the raw scene color is used.
-	 * @param	PostProcessSettings - pointer to post process structure, NULL for no effect.
-	 * @param	PostProcessBlendWeight - Blendweight for PostProcessSettings.
-	 * @return	pointer to a configured SceneRenderer instance.
-	 */
-	FSceneRenderer* CreateSceneRenderer(USceneCaptureComponent* SceneCaptureComponent, UTextureRenderTarget* TextureTarget, const FMatrix& ViewRotationMatrix, const FVector& ViewLocation, float FOV, float MaxViewDistance, bool bCaptureSceneColour = true, FPostProcessSettings* PostProcessSettings = NULL, float PostProcessBlendWeight = 0);
-
 private:
 	/** 
 	 * The number of visible lights in the scene
diff --git a/Engine/Source/Runtime/Renderer/Private/SceneRendering.cpp b/Engine/Source/Runtime/Renderer/Private/SceneRendering.cpp
--- a/Engine/Source/Runtime/Renderer/Private/SceneRendering.cpp
+++ b/Engine/Source/Runtime/Renderer/Private/SceneRendering.cpp
@@ -19,6 +19,7 @@
 #include "SceneUtils.h"
 #include "LightGrid.h"
 #include "AtmosphereRendering.h"
+#include "Components/PlanarReflectionComponent.h"
 
 /*-----------------------------------------------------------------------------
 	Globals
@@ -72,11 +73,6 @@ static TAutoConsoleVariable<int32> CVarInstancedStereo(
 	TEXT("0 to disable instanced stereo (default), 1 to enable."),
 	ECVF_ReadOnly | ECVF_RenderThreadSafe);
 
-static TAutoConsoleVariable<int32> CVarEnablePlanarReflections(
-	TEXT("r.EnablePlanarReflections"),
-	1,
-	TEXT("0: Disable planar reflection captures. 1: Enable (Default)"));
-
 #if !(UE_BUILD_SHIPPING || UE_BUILD_TEST)
 static TAutoConsoleVariable<float> CVarGeneralPurposeTweak(
 	TEXT("r.GeneralPurposeTweak"),
@@ -671,6 +667,7 @@ void FViewInfo::CreateUniformBuffer(
 	// can be optimized
 	ViewUniformShaderParameters.PrevInvViewProj = PrevViewProjMatrix.Inverse();
 	ViewUniformShaderParameters.ScreenPositionScaleBias = ScreenPositionScaleBias;
+	ViewUniformShaderParameters.GlobalClippingPlane = FVector4(GlobalClippingPlane.X, GlobalClippingPlane.Y, GlobalClippingPlane.Z, -GlobalClippingPlane.W);
 
 	FrameUniformShaderParameters.FieldOfViewWideAngles = 2.f * ViewMatrices.GetHalfFieldOfViewPerAxis();
 	FrameUniformShaderParameters.PrevFieldOfViewWideAngles = 2.f * PrevViewMatrices.GetHalfFieldOfViewPerAxis();
@@ -1538,6 +1535,17 @@ void FSceneRenderer::RenderFinish(FRHICommandListImmediate& RHICmdList)
 			bShowPrecomputedVisibilityWarning = !bUsedPrecomputedVisibility;
 		}
 
+		bool bShowGlobalClipPlaneWarning = false;
+
+		if (Scene->PlanarReflections.Num() > 0)
+		{
+			static const auto* CVarClipPlane = IConsoleManager::Get().FindTConsoleVariableDataInt(TEXT("r.AllowGlobalClipPlane"));
+			if (CVarClipPlane && CVarClipPlane->GetValueOnRenderThread() == 0)
+			{
+				bShowGlobalClipPlaneWarning = true;
+			}
+		}
+
 		for(int32 ViewIndex = 0;ViewIndex < Views.Num();ViewIndex++)
 		{	
 			FViewInfo& View = Views[ViewIndex];
@@ -1548,7 +1556,7 @@ void FSceneRenderer::RenderFinish(FRHICommandListImmediate& RHICmdList)
 				FSceneViewState* ViewState = (FSceneViewState*)View.State;
 				bool bViewParentOrFrozen = ViewState && (ViewState->HasViewParent() || ViewState->bIsFrozen);
 				bool bLocked = View.bIsLocked;
-				if (bViewParentOrFrozen || bShowPrecomputedVisibilityWarning || bLocked)
+				if (bViewParentOrFrozen || bShowPrecomputedVisibilityWarning || bLocked || bShowGlobalClipPlaneWarning)
 				{
 					// this is a helper class for FCanvas to be able to get screen size
 					class FRenderTargetTemp : public FRenderTarget
@@ -1589,6 +1597,12 @@ void FSceneRenderer::RenderFinish(FRHICommandListImmediate& RHICmdList)
 						Canvas.DrawShadowedText(10, Y, Message, GetStatsFont(), FLinearColor(1.0, 0.05, 0.05, 1.0));
 						Y += 14;
 					}
+					if (bShowGlobalClipPlaneWarning)
+					{
+						const FText Message = NSLOCTEXT("Renderer", "NoGlobalClipPlane", "GLOBAL CLIP PLANE PROJECT SETTING NOT ENABLED");
+						Canvas.DrawShadowedText(10, Y, Message, GetStatsFont(), FLinearColor(1.0, 0.05, 0.05, 1.0));
+						Y += 14;
+					}
 					if (bLocked)
 					{
 						const FText Message = NSLOCTEXT("Renderer", "ViewLocked", "VIEW LOCKED");
@@ -1913,6 +1927,9 @@ static void RenderViewFamily_RenderThread(FRHICommandListImmediate& RHICmdList,
 			SceneRenderer->Render(RHICmdList);
 		}
 
+		// Only reset per-frame scene state once all views have processed their frame, including those in planar reflections
+		SceneRenderer->Scene->DistanceFieldSceneData.PrimitiveModifiedBounds.Reset();
+
 #if STATS
 		{
 			QUICK_SCOPE_CYCLE_COUNTER(STAT_RenderViewFamily_RenderThread_MemStats);
@@ -2005,6 +2022,12 @@ void FRendererModule::BeginRenderingViewFamily(FCanvas* Canvas,FSceneViewFamily*
 		// Construct the scene renderer.  This copies the view family attributes into its own structures.
 		FSceneRenderer* SceneRenderer = FSceneRenderer::CreateSceneRenderer(ViewFamily, Canvas->GetHitProxyConsumer());
 
+		for (int32 ReflectionIndex = 0; ReflectionIndex < SceneRenderer->Scene->PlanarReflections_GameThread.Num(); ReflectionIndex++)
+		{
+			UPlanarReflectionComponent* ReflectionComponent = SceneRenderer->Scene->PlanarReflections_GameThread[ReflectionIndex];
+			SceneRenderer->Scene->UpdatePlanarReflectionContents(ReflectionComponent, *SceneRenderer);
+		}
+
 		ENQUEUE_UNIQUE_RENDER_COMMAND_ONEPARAMETER(
 			FDrawSceneCommand,
 			FSceneRenderer*,SceneRenderer,SceneRenderer,
diff --git a/Engine/Source/Runtime/Renderer/Private/SceneRendering.h b/Engine/Source/Runtime/Renderer/Private/SceneRendering.h
--- a/Engine/Source/Runtime/Renderer/Private/SceneRendering.h
+++ b/Engine/Source/Runtime/Renderer/Private/SceneRendering.h
@@ -935,6 +935,7 @@ protected:
 	void UpdatePrimitivePrecomputedLightingBuffers();
 	void ClearPrimitiveSingleFramePrecomputedLightingBuffers();
 
+	void RenderPlanarReflection(class FPlanarReflectionSceneProxy* ReflectionSceneProxy);
 };
 
 /**
diff --git a/Engine/Source/Runtime/Renderer/Private/SceneVisibility.cpp b/Engine/Source/Runtime/Renderer/Private/SceneVisibility.cpp
--- a/Engine/Source/Runtime/Renderer/Private/SceneVisibility.cpp
+++ b/Engine/Source/Runtime/Renderer/Private/SceneVisibility.cpp
@@ -10,6 +10,7 @@
 #include "FXSystem.h"
 #include "SceneUtils.h"
 #include "PostProcessing.h"
+#include "PlanarReflectionSceneProxy.h"
 
 /*------------------------------------------------------------------------------
 	Globals
@@ -304,7 +305,7 @@ static int32 FrustumCull(const FScene* Scene, FViewInfo& View)
 				uint32 Mask = 0x1;
 				uint32 VisBits = 0;
 				uint32 FadingBits = 0;
-				for(int32 BitSubIndex = 0; BitSubIndex < NumBitsPerDWORD && WordIndex * NumBitsPerDWORD + BitSubIndex < BitArrayNumInner; BitSubIndex++, Mask <<= 1)
+				for (int32 BitSubIndex = 0; BitSubIndex < NumBitsPerDWORD && WordIndex * NumBitsPerDWORD + BitSubIndex < BitArrayNumInner; BitSubIndex++, Mask <<= 1)
 				{
 					int32 Index = WordIndex * NumBitsPerDWORD + BitSubIndex;
 					const FPrimitiveBounds& Bounds = Scene->PrimitiveBounds[Index];
@@ -325,16 +326,11 @@ static int32 FrustumCull(const FScene* Scene, FViewInfo& View)
 						MaxDrawDistance = FLT_MAX;
 					}
 
-					// The primitive is always culled if it exceeds the max fade distance or lay outside the view frustum.
-					const bool bShouldCull = 
-						(DistanceSquared > FMath::Square(MaxDrawDistance + FadeRadius) ||
-						DistanceSquared < Bounds.MinDrawDistanceSq || 
+					if (DistanceSquared > FMath::Square(MaxDrawDistance + FadeRadius) ||
+						DistanceSquared < Bounds.MinDrawDistanceSq ||
 						(UseCustomCulling && !View.CustomVisibilityQuery->IsVisible(VisibilityId, FBoxSphereBounds(Bounds.Origin, Bounds.BoxExtent, Bounds.SphereRadius))) ||
 						(bAlsoUseSphereTest && View.ViewFrustum.IntersectSphere(Bounds.Origin, Bounds.SphereRadius) == false) ||
-						View.ViewFrustum.IntersectBox(Bounds.Origin, Bounds.BoxExtent) == false);
-
-					// Don't cull primitives when rendering planar reflections. The reflection of the primitive may be visible when the primitive isn't.
-					if (bShouldCull && !View.bIsPlanarReflectionCapture)
+						View.ViewFrustum.IntersectBox(Bounds.Origin, Bounds.BoxExtent) == false)
 					{
 						STAT(NumCulledPrimitives.Increment());
 					}
@@ -1551,6 +1547,7 @@ struct FRelevancePacket
 				&& (!Scene->ShouldUseDeferredRenderer() || bTranslucentRelevance || bHasClearCoat))
 			{
 				PrimitiveSceneInfo->CachedReflectionCaptureProxy = Scene->FindClosestReflectionCapture(Scene->PrimitiveBounds[BitIndex].Origin);
+				PrimitiveSceneInfo->CachedPlanarReflectionProxy = Scene->FindClosestPlanarReflection(Scene->PrimitiveBounds[BitIndex].Origin);
 
 				if (!Scene->ShouldUseDeferredRenderer())
 				{
@@ -2493,18 +2490,6 @@ void FSceneRenderer::ComputeViewVisibility(FRHICommandListImmediate& RHICmdList)
 			}
 		}
 
-		// When rendering a planar reflection capture, hide primitives not flagged as visible in the reflection.
-		if (View.bIsPlanarReflectionCapture)
-		{
-			for (FSceneSetBitIterator BitIt(View.PrimitiveVisibilityMap); BitIt; ++BitIt)
-			{
-				if (!Scene->Primitives[BitIt.GetIndex()]->Proxy->IsVisibleInPlanarReflection())
-				{
-					View.PrimitiveVisibilityMap.AccessCorrespondingBit(BitIt) = false;
-				}
-			}
-		}
-
 		// Cull small objects in wireframe in ortho views
 		// This is important for performance in the editor because wireframe disables any kind of occlusion culling
 		if (View.Family->EngineShowFlags.Wireframe)
diff --git a/Engine/Source/Runtime/Renderer/Private/ShaderBaseClasses.cpp b/Engine/Source/Runtime/Renderer/Private/ShaderBaseClasses.cpp
--- a/Engine/Source/Runtime/Renderer/Private/ShaderBaseClasses.cpp
+++ b/Engine/Source/Runtime/Renderer/Private/ShaderBaseClasses.cpp
@@ -401,10 +401,12 @@ bool FMaterialShader::Serialize(FArchive& Ar)
 FTextureRHIRef& FMaterialShader::GetEyeAdaptation(FRHICommandList& RHICmdList, const FSceneView& View)
 {
 	IPooledRenderTarget* EyeAdaptationRT = NULL;
-	if( View.bIsViewInfo )
+	if (View.bIsViewInfo)
 	{
 		const FViewInfo& ViewInfo = static_cast<const FViewInfo&>(View);
-		EyeAdaptationRT = ViewInfo.GetEyeAdaptation(RHICmdList);
+		if (ViewInfo.HasValidEyeAdaptation()) {
+			EyeAdaptationRT = ViewInfo.GetEyeAdaptation(RHICmdList);
+		}
 	}
 
 	if( EyeAdaptationRT )
diff --git a/Engine/Source/Runtime/Renderer/Private/ShaderComplexityRendering.h b/Engine/Source/Runtime/Renderer/Private/ShaderComplexityRendering.h
--- a/Engine/Source/Runtime/Renderer/Private/ShaderComplexityRendering.h
+++ b/Engine/Source/Runtime/Renderer/Private/ShaderComplexityRendering.h
@@ -14,7 +14,7 @@ public:
 
 	static bool ShouldCache(EShaderPlatform Platform)
 	{
-		return AllowDebugViewModeShader(Platform) || !bQuadComplexity; // Shader complexity is always available
+		return AllowDebugViewShaderMode(bQuadComplexity ? DVSM_QuadComplexity : DVSM_ShaderComplexity, Platform);
 	}
 
 	TComplexityAccumulatePS(const ShaderMetaType::CompiledShaderInitializerType& Initializer):
@@ -38,7 +38,7 @@ public:
 
 	static void ModifyCompilationEnvironment(EShaderPlatform Platform, FShaderCompilerEnvironment& OutEnvironment)
 	{
-		OutEnvironment.SetDefine(TEXT("OUTPUT_QUAD_OVERDRAW"), AllowDebugViewModeShader(Platform) ? TEXT("1") : TEXT("0"));
+		OutEnvironment.SetDefine(TEXT("OUTPUT_QUAD_OVERDRAW"), AllowDebugViewShaderMode(DVSM_QuadComplexity, Platform) ? TEXT("1") : TEXT("0"));
 	}
 
 	virtual void SetParameters(
diff --git a/Engine/Source/Runtime/Renderer/Private/TexCoordScaleAnalysisRendering.cpp b/Engine/Source/Runtime/Renderer/Private/TexCoordScaleAnalysisRendering.cpp
--- a/Engine/Source/Runtime/Renderer/Private/TexCoordScaleAnalysisRendering.cpp
+++ b/Engine/Source/Runtime/Renderer/Private/TexCoordScaleAnalysisRendering.cpp
@@ -10,7 +10,7 @@ WantedMipsAccuracyRendering.cpp: Contains definitions for rendering the viewmode
 
 static TAutoConsoleVariable<int32> CVarStreamingAnalysisIndex(
 	TEXT("r.Streaming.AnalysisIndex"),
-	0,
+	FTexCoordScaleAnalysisPS::MAX_NUM_TEXTURE_REGISTER,
 	TEXT("Analysis index of the texture coordinate scale accuracy viewmode."),
 	ECVF_Default);
 
@@ -49,8 +49,44 @@ void FTexCoordScaleAnalysisPS::SetMesh(
 	const FMeshDrawingRenderState& DrawRenderState
 	)
 {
-	SetShaderValue(RHICmdList, FMeshMaterialShader::GetPixelShader(), TextureAnalysisIndexParameter, (int32)CVarStreamingAnalysisIndex.GetValueOnRenderThread());
-	SetShaderValue(RHICmdList, FMeshMaterialShader::GetPixelShader(), CPUScaleParameter, 1.f);
+	int32 AnalysisIndex = CVarStreamingAnalysisIndex.GetValueOnRenderThread();
+	if (AnalysisIndex >= 0 && AnalysisIndex < MAX_NUM_TEXTURE_REGISTER)
+	{
+		AnalysisIndex = FMath::Clamp<int32>(AnalysisIndex, 0, MAX_NUM_TEXTURE_REGISTER - 1);;
+	}
+	else
+	{
+		AnalysisIndex = MAX_NUM_TEXTURE_REGISTER;
+	}
+
+	FVector4 OneOverCPUTexCoordScales[MAX_NUM_TEXTURE_REGISTER / 4];
+	FMemory::Memzero(OneOverCPUTexCoordScales); // 0 remap to irrelevant data.
+
+	FIntVector4 TexCoordIndices[MAX_NUM_TEXTURE_REGISTER / 4];
+	FMemory::Memzero(TexCoordIndices);
+
+#if !(UE_BUILD_SHIPPING || UE_BUILD_TEST)
+	const bool bUseNewMetrics = CVarStreamingUseNewMetrics.GetValueOnRenderThread() != 0;
+	const FStreamingSectionBuildInfo* SectionData = Proxy ? Proxy->GetStreamingSectionData(bUseNewMetrics ? VisualizeLODIndex : INDEX_NONE, bUseNewMetrics ? BatchElement.VisualizeElementIndex : INDEX_NONE) : nullptr;
+	if (SectionData)
+	{
+		for (int32 ScaleIndex = 0; ScaleIndex < MAX_NUM_TEXTURE_REGISTER && ScaleIndex < SectionData->TexCoordData.Num(); ++ScaleIndex)
+		{
+			float Scale = SectionData->TexCoordData[ScaleIndex].Scale;
+			if (Scale > 0)
+			{
+				OneOverCPUTexCoordScales[ScaleIndex / 4][ScaleIndex % 4] = 1.f / Scale;
+				TexCoordIndices[ScaleIndex / 4][ScaleIndex % 4] = SectionData->TexCoordData[ScaleIndex].Index;
+			}
+		}
+	}
+#endif
+
+	const bool bOutputScales = View.Family->GetDebugViewShaderMode() == DVSM_TexCoordScaleAnalysis;
+
+	SetShaderValueArray(RHICmdList, FMeshMaterialShader::GetPixelShader(), OneOverCPUTexCoordScalesParameter, OneOverCPUTexCoordScales, ARRAY_COUNT(OneOverCPUTexCoordScales));
+	SetShaderValueArray(RHICmdList, FMeshMaterialShader::GetPixelShader(), TexCoordIndicesParameter, TexCoordIndices, ARRAY_COUNT(TexCoordIndices));
+	SetShaderValue(RHICmdList, FMeshMaterialShader::GetPixelShader(), TextureAnalysisIndexParameter, bOutputScales ? (int32)INDEX_NONE : AnalysisIndex);
 
 	FMeshMaterialShader::SetMesh(RHICmdList, FMeshMaterialShader::GetPixelShader(), VertexFactory, View, Proxy, BatchElement, DrawRenderState);
 }
diff --git a/Engine/Source/Runtime/Renderer/Private/TexCoordScaleAnalysisRendering.h b/Engine/Source/Runtime/Renderer/Private/TexCoordScaleAnalysisRendering.h
--- a/Engine/Source/Runtime/Renderer/Private/TexCoordScaleAnalysisRendering.h
+++ b/Engine/Source/Runtime/Renderer/Private/TexCoordScaleAnalysisRendering.h
@@ -17,11 +17,13 @@ class FTexCoordScaleAnalysisPS : public FMeshMaterialShader, public IDebugViewMo
 	DECLARE_SHADER_TYPE(FTexCoordScaleAnalysisPS,MeshMaterial);
 
 public:
+	enum { MAX_NUM_TEXTURE_REGISTER = FMaterialTexCoordBuildInfo::MAX_NUM_TEXTURE_REGISTER };
 
 	static bool ShouldCache(EShaderPlatform Platform, const FMaterial* Material, const FVertexFactoryType* VertexFactoryType)
 	{
-		// Debug viewmode pixel shaders have predefined interpolants, so we only need to compile it for one vertex factory.
-		return Material->GetFriendlyName().Contains(TEXT("FDebugViewModeMaterialProxy")) &&  AllowDebugViewModeShader(Platform) && VertexFactoryType == FindVertexFactoryType(FName(TEXT("FLocalVertexFactory"), FNAME_Find));
+		return AllowDebugViewShaderMode(DVSM_TexCoordScaleAnalysis, Platform) && 
+			Material->GetFriendlyName().Contains(TEXT("FDebugViewModeMaterialProxy")) && 
+			VertexFactoryType == FindVertexFactoryType(FName(TEXT("FLocalVertexFactory"), FNAME_Find));
 	}
 
 	FTexCoordScaleAnalysisPS(const ShaderMetaType::CompiledShaderInitializerType& Initializer):
@@ -29,7 +31,8 @@ public:
 	{
 		AccuracyColorsParameter.Bind(Initializer.ParameterMap,TEXT("AccuracyColors"));
 		TextureAnalysisIndexParameter.Bind(Initializer.ParameterMap,TEXT("TextureAnalysisIndex"));
-		CPUScaleParameter.Bind(Initializer.ParameterMap,TEXT("CPUScale"));
+		OneOverCPUTexCoordScalesParameter.Bind(Initializer.ParameterMap,TEXT("OneOverCPUTexCoordScales"));
+		TexCoordIndicesParameter.Bind(Initializer.ParameterMap,TEXT("TexCoordIndices"));
 	}
 
 	FTexCoordScaleAnalysisPS() {}
@@ -39,12 +42,17 @@ public:
 		bool bShaderHasOutdatedParameters = FMeshMaterialShader::Serialize(Ar);
 		Ar << AccuracyColorsParameter;
 		Ar << TextureAnalysisIndexParameter;
-		Ar << CPUScaleParameter;
+		Ar << OneOverCPUTexCoordScalesParameter;
+		Ar << TexCoordIndicesParameter;
 		return bShaderHasOutdatedParameters;
 	}
 
 	static void ModifyCompilationEnvironment(EShaderPlatform Platform, const FMaterial* Material, FShaderCompilerEnvironment& OutEnvironment)
 	{
+		OutEnvironment.SetDefine(TEXT("MAX_NUM_TEX_COORD"), (uint32)FMaterialTexCoordBuildInfo::MAX_NUM_TEX_COORD);
+		OutEnvironment.SetDefine(TEXT("INITIAL_GPU_SCALE"), (uint32)FMaterialTexCoordBuildInfo::INITIAL_GPU_SCALE);
+		OutEnvironment.SetDefine(TEXT("TILE_RESOLUTION"), (uint32)FMaterialTexCoordBuildInfo::TILE_RESOLUTION);
+		OutEnvironment.SetDefine(TEXT("MAX_NUM_TEXTURE_REGISTER"), (uint32)FMaterialTexCoordBuildInfo::MAX_NUM_TEXTURE_REGISTER);
 		FMeshMaterialShader::ModifyCompilationEnvironment(Platform, OutEnvironment);
 	}
 
@@ -75,5 +83,6 @@ private:
 
 	FShaderParameter AccuracyColorsParameter;
 	FShaderParameter TextureAnalysisIndexParameter;
-	FShaderParameter CPUScaleParameter;
+	FShaderParameter OneOverCPUTexCoordScalesParameter;
+	FShaderParameter TexCoordIndicesParameter;
 };
diff --git a/Engine/Source/Runtime/Renderer/Private/TexelFactorAccuracyRendering.cpp b/Engine/Source/Runtime/Renderer/Private/TexelFactorAccuracyRendering.cpp
--- a/Engine/Source/Runtime/Renderer/Private/TexelFactorAccuracyRendering.cpp
+++ b/Engine/Source/Runtime/Renderer/Private/TexelFactorAccuracyRendering.cpp
@@ -10,9 +10,6 @@ TexelFactorAccuracyRendering.cpp: Contains definitions for rendering the viewmod
 
 IMPLEMENT_SHADER_TYPE(,FTexelFactorAccuracyPS,TEXT("TexelFactorAccuracyPixelShader"),TEXT("Main"),SF_Pixel);
 
-extern ENGINE_API TAutoConsoleVariable<int32> CVarStreamingUseNewMetrics;
-
-
 void FTexelFactorAccuracyPS::SetParameters(
 	FRHICommandList& RHICmdList, 
 	const FShader* OriginalVS, 
@@ -47,14 +44,13 @@ void FTexelFactorAccuracyPS::SetMesh(
 	const FMeshDrawingRenderState& DrawRenderState
 	)
 {
-	const bool bUseMetrics = CVarStreamingUseNewMetrics.GetValueOnRenderThread() != 0;
-
 	float CPUTexelFactor = -1.f;
 #if !(UE_BUILD_SHIPPING || UE_BUILD_TEST)
-	FStreamingTexturePrimitiveInfo Info; 
-	if (Proxy && Proxy->GetStreamingTextureInfo(Info, bUseMetrics ? VisualizeLODIndex : INDEX_NONE, bUseMetrics ? BatchElement.VisualizeElementIndex : INDEX_NONE))
+	const bool bUseNewMetrics = CVarStreamingUseNewMetrics.GetValueOnRenderThread() != 0;
+	const FStreamingSectionBuildInfo* SectionData = Proxy ? Proxy->GetStreamingSectionData(bUseNewMetrics ? VisualizeLODIndex : INDEX_NONE, bUseNewMetrics ? BatchElement.VisualizeElementIndex : INDEX_NONE) : nullptr;
+	if (SectionData)
 	{
-		CPUTexelFactor = Info.TexelFactor;
+		CPUTexelFactor = SectionData->TexelFactors[0];
 	}
 #endif
 	SetShaderValue(RHICmdList, FGlobalShader::GetPixelShader(), CPUTexelFactorParameter, CPUTexelFactor);
diff --git a/Engine/Source/Runtime/Renderer/Private/TexelFactorAccuracyRendering.h b/Engine/Source/Runtime/Renderer/Private/TexelFactorAccuracyRendering.h
--- a/Engine/Source/Runtime/Renderer/Private/TexelFactorAccuracyRendering.h
+++ b/Engine/Source/Runtime/Renderer/Private/TexelFactorAccuracyRendering.h
@@ -17,7 +17,7 @@ public:
 
 	static bool ShouldCache(EShaderPlatform Platform)
 	{
-		return AllowDebugViewModeShader(Platform);
+		return AllowDebugViewShaderMode(DVSM_TexelFactorAccuracy, Platform);
 	}
 
 	FTexelFactorAccuracyPS(const ShaderMetaType::CompiledShaderInitializerType& Initializer):
diff --git a/Engine/Source/Runtime/Renderer/Private/WantedMipsAccuracyRendering.cpp b/Engine/Source/Runtime/Renderer/Private/WantedMipsAccuracyRendering.cpp
--- a/Engine/Source/Runtime/Renderer/Private/WantedMipsAccuracyRendering.cpp
+++ b/Engine/Source/Runtime/Renderer/Private/WantedMipsAccuracyRendering.cpp
@@ -8,8 +8,6 @@ WantedMipsAccuracyRendering.cpp: Contains definitions for rendering the viewmode
 #include "ScenePrivate.h"
 #include "WantedMipsAccuracyRendering.h"
 
-extern ENGINE_API TAutoConsoleVariable<int32> CVarStreamingUseNewMetrics;
-
 IMPLEMENT_SHADER_TYPE(,FWantedMipsAccuracyPS,TEXT("WantedMipsAccuracyPixelShader"),TEXT("Main"),SF_Pixel);
 
 void FWantedMipsAccuracyPS::SetParameters(
@@ -43,14 +41,13 @@ void FWantedMipsAccuracyPS::SetMesh(
 	const FMeshDrawingRenderState& DrawRenderState
 	)
 {
-	const bool bUseNewMetrics= CVarStreamingUseNewMetrics.GetValueOnRenderThread() != 0;
-
 	float CPUWantedMips = -1.f;
 #if !(UE_BUILD_SHIPPING || UE_BUILD_TEST)
-	FStreamingTexturePrimitiveInfo Info;
-	if (Proxy && Proxy->GetStreamingTextureInfo(Info, bUseNewMetrics ? VisualizeLODIndex : INDEX_NONE, bUseNewMetrics ? BatchElement.VisualizeElementIndex : INDEX_NONE))
+	const bool bUseNewMetrics = CVarStreamingUseNewMetrics.GetValueOnRenderThread() != 0;
+	const FStreamingSectionBuildInfo* SectionData = Proxy ? Proxy->GetStreamingSectionData(bUseNewMetrics ? VisualizeLODIndex : INDEX_NONE, bUseNewMetrics ? BatchElement.VisualizeElementIndex : INDEX_NONE) : nullptr;
+	if (SectionData)
 	{
-		FVector ViewToObject = Info.Bounds.Origin - View.ViewMatrices.ViewOrigin;
+		FVector ViewToObject = SectionData->BoxOrigin - View.ViewMatrices.ViewOrigin;
 
 		float DistSqMinusRadiusSq = 0;
 		float ViewSize = View.ViewRect.Size().X;
@@ -58,24 +55,24 @@ void FWantedMipsAccuracyPS::SetMesh(
 		if (bUseNewMetrics)
 		{
 			ViewToObject = ViewToObject.GetAbs();
-			FVector BoxViewToObject = ViewToObject.ComponentMin(Info.Bounds.BoxExtent);
-
+			FVector BoxViewToObject = ViewToObject.ComponentMin(SectionData->BoxExtent);
 			DistSqMinusRadiusSq = FVector::DistSquared(BoxViewToObject, ViewToObject);
 		}
 		else
 		{
 			float Distance = ViewToObject.Size();
-			DistSqMinusRadiusSq = FMath::Square(Distance) - FMath::Square(Info.Bounds.SphereRadius);
+			DistSqMinusRadiusSq = FMath::Square(Distance) - SectionData->BoxExtent.SizeSquared();
 		}
 
 		DistSqMinusRadiusSq = FMath::Max(1.f, DistSqMinusRadiusSq);
 
-		float CoordSize = Info.TexelFactor * FMath::InvSqrt(DistSqMinusRadiusSq) * ViewSize;
+		float CoordSize = SectionData->TexelFactors[0] * FMath::InvSqrt(DistSqMinusRadiusSq) * ViewSize;
 
 		CPUWantedMips = FMath::Clamp<float>(FMath::Log2(CoordSize), 0.f, (float)MaxStreamingAccuracyMips);
 	}
 #endif
-	SetShaderValue(RHICmdList, FGlobalShader::GetPixelShader(), CPUWantedMipsParameter, FMath::FloorToFloat(CPUWantedMips));
+	// Because the streamer use FMath::FloorToFloat, here we need to use -1 to have a useful result.
+	SetShaderValue(RHICmdList, FGlobalShader::GetPixelShader(), CPUWantedMipsParameter, CPUWantedMips - 1.f);
 }
 
 void FWantedMipsAccuracyPS::SetMesh(FRHICommandList& RHICmdList, const FSceneView& View)
diff --git a/Engine/Source/Runtime/Renderer/Private/WantedMipsAccuracyRendering.h b/Engine/Source/Runtime/Renderer/Private/WantedMipsAccuracyRendering.h
--- a/Engine/Source/Runtime/Renderer/Private/WantedMipsAccuracyRendering.h
+++ b/Engine/Source/Runtime/Renderer/Private/WantedMipsAccuracyRendering.h
@@ -17,7 +17,7 @@ public:
 
 	static bool ShouldCache(EShaderPlatform Platform)
 	{
-		return AllowDebugViewModeShader(Platform);
+		return AllowDebugViewShaderMode(DVSM_WantedMipsAccuracy, Platform);
 	}
 
 	FWantedMipsAccuracyPS(const ShaderMetaType::CompiledShaderInitializerType& Initializer):
diff --git a/Engine/Source/Runtime/Renderer/Public/PrimitiveSceneInfo.h b/Engine/Source/Runtime/Renderer/Public/PrimitiveSceneInfo.h
--- a/Engine/Source/Runtime/Renderer/Public/PrimitiveSceneInfo.h
+++ b/Engine/Source/Runtime/Renderer/Public/PrimitiveSceneInfo.h
@@ -206,6 +206,11 @@ public:
 	 */
 	FUniformBufferRHIRef IndirectLightingCacheUniformBuffer;
 
+	/** 
+	 * Planar reflection that was closest to this primitive, used for forward reflections.
+	 */
+	const class FPlanarReflectionSceneProxy* CachedPlanarReflectionProxy;
+
 	/** 
 	 * Reflection capture proxy that was closest to this primitive, used for the forward shading rendering path. 
 	 */
diff --git a/Engine/Source/Runtime/ShaderCore/Private/Shader.cpp b/Engine/Source/Runtime/ShaderCore/Private/Shader.cpp
--- a/Engine/Source/Runtime/ShaderCore/Private/Shader.cpp
+++ b/Engine/Source/Runtime/ShaderCore/Private/Shader.cpp
@@ -112,6 +112,7 @@ void FShaderParameterMap::UpdateHash(FSHA1& HashState) const
 bool FShaderType::bInitializedSerializationHistory = false;
 
 FShaderType::FShaderType(
+	EShaderTypeForDynamicCast InShaderTypeForDynamicCast,
 	const TCHAR* InName,
 	const TCHAR* InSourceFilename,
 	const TCHAR* InFunctionName,
@@ -119,6 +120,7 @@ FShaderType::FShaderType(
 	ConstructSerializedType InConstructSerializedRef,
 	GetStreamOutElementsType InGetStreamOutElementsRef
 	):
+	ShaderTypeForDynamicCast(InShaderTypeForDynamicCast),
 	Name(InName),
 	TypeName(InName),
 	SourceFilename(InSourceFilename),
@@ -1640,6 +1642,11 @@ void ShaderMapAppendKeyString(EShaderPlatform Platform, FString& KeyString)
 		KeyString += (CVar && CVar->GetInt() != 0) ? TEXT("_DBuf") : TEXT("_NoDBuf");
 	}
 
+	{
+		static IConsoleVariable* CVar = IConsoleManager::Get().FindConsoleVariable(TEXT("r.AllowGlobalClipPlane"));
+		KeyString += (CVar && CVar->GetInt() != 0) ? TEXT("_ClipP") : TEXT("");
+	}
+
 	{
 		static const auto CVar = IConsoleManager::Get().FindConsoleVariable(TEXT("r.Shaders.KeepDebugInfo"));
 		KeyString += (CVar && CVar->GetInt() != 0) ? TEXT("_NoStrip") : TEXT("");
diff --git a/Engine/Source/Runtime/ShaderCore/Public/Shader.h b/Engine/Source/Runtime/ShaderCore/Public/Shader.h
--- a/Engine/Source/Runtime/ShaderCore/Public/Shader.h
+++ b/Engine/Source/Runtime/ShaderCore/Public/Shader.h
@@ -758,6 +758,13 @@ public:
 class SHADERCORE_API FShaderType
 {
 public:
+	enum class EShaderTypeForDynamicCast : uint32
+	{
+		Global,
+		Material,
+		MeshMaterial
+	};
+
 	typedef class FShader* (*ConstructSerializedType)();
 	typedef void (*GetStreamOutElementsType)(FStreamOutElementList& ElementList, TArray<uint32>& StreamStrides, int32& RasterizedStream);
 
@@ -784,6 +791,7 @@ public:
 
 	/** Minimal initialization constructor. */
 	FShaderType(
+		EShaderTypeForDynamicCast InShaderTypeForDynamicCast,
 		const TCHAR* InName,
 		const TCHAR* InSourceFilename,
 		const TCHAR* InFunctionName,
@@ -815,13 +823,31 @@ public:
 	}
 
 	// Dynamic casts.
-	virtual FGlobalShaderType* GetGlobalShaderType() { return nullptr; }
-	virtual const FGlobalShaderType* GetGlobalShaderType() const { return nullptr; }
-	virtual FMaterialShaderType* GetMaterialShaderType() { return nullptr; }
-	virtual const FMaterialShaderType* GetMaterialShaderType() const { return nullptr; }
-	virtual FMeshMaterialShaderType* GetMeshMaterialShaderType() { return nullptr; }
-	virtual const FMeshMaterialShaderType* GetMeshMaterialShaderType() const { return nullptr; }
-	
+	FORCEINLINE FGlobalShaderType* GetGlobalShaderType() 
+	{ 
+		return (ShaderTypeForDynamicCast == EShaderTypeForDynamicCast::Global) ? reinterpret_cast<FGlobalShaderType*>(this) : nullptr;
+	}
+	FORCEINLINE const FGlobalShaderType* GetGlobalShaderType() const
+	{
+		return (ShaderTypeForDynamicCast == EShaderTypeForDynamicCast::Global) ? reinterpret_cast<const FGlobalShaderType*>(this) : nullptr;
+	}
+	FORCEINLINE FMaterialShaderType* GetMaterialShaderType()
+	{
+		return (ShaderTypeForDynamicCast == EShaderTypeForDynamicCast::Material) ? reinterpret_cast<FMaterialShaderType*>(this) : nullptr;
+	}
+	FORCEINLINE const FMaterialShaderType* GetMaterialShaderType() const
+	{
+		return (ShaderTypeForDynamicCast == EShaderTypeForDynamicCast::Material) ? reinterpret_cast<const FMaterialShaderType*>(this) : nullptr;
+	}
+	FORCEINLINE FMeshMaterialShaderType* GetMeshMaterialShaderType()
+	{
+		return (ShaderTypeForDynamicCast == EShaderTypeForDynamicCast::MeshMaterial) ? reinterpret_cast<FMeshMaterialShaderType*>(this) : nullptr;
+	}
+	FORCEINLINE const FMeshMaterialShaderType* GetMeshMaterialShaderType() const
+	{
+		return (ShaderTypeForDynamicCast == EShaderTypeForDynamicCast::MeshMaterial) ? reinterpret_cast<const FMeshMaterialShaderType*>(this) : nullptr;
+	}
+
 	// Accessors.
 	inline EShaderFrequency GetFrequency() const
 	{ 
@@ -905,6 +931,7 @@ public:
 	}
 
 private:
+	EShaderTypeForDynamicCast ShaderTypeForDynamicCast;
 	uint32 HashIndex;
 	const TCHAR* Name;
 	FName TypeName;
@@ -1360,6 +1387,12 @@ public:
 	/** Default constructor. */
 	TShaderMap() {}
 
+	/** Destructor ensures pipelines cleared up. */
+	~TShaderMap()
+	{
+		EmptyShaderPipelines();
+	}
+
 	/** Finds the shader with the given type.  Asserts on failure. */
 	template<typename ShaderType>
 	ShaderType* GetShader() const
@@ -1669,11 +1702,11 @@ public:
 		return ShaderPipelines.Num();
 	}
 
-	/** clears out all shaders held in the map */
+	/** clears out all shaders and deletes shader pipelines held in the map */
 	void Empty()
 	{
 		Shaders.Empty();
-		ShaderPipelines.Empty();
+		EmptyShaderPipelines();
 	}
 
 	inline FShaderPipeline* GetShaderPipeline(const FShaderPipelineType* PipelineType)
@@ -1724,6 +1757,19 @@ public:
 	}
 
 protected:
+	inline void EmptyShaderPipelines()
+	{
+		for (auto& Pair : ShaderPipelines)
+		{
+			if (FShaderPipeline* Pipeline = Pair.Value)
+			{
+				delete Pipeline;
+				Pipeline = nullptr;
+			}
+		}
+		ShaderPipelines.Empty();
+	}
+
 	TMap<FShaderType*, TRefCountPtr<FShader> > Shaders;
 	TMap<const FShaderPipelineType*, FShaderPipeline*> ShaderPipelines;
 };
diff --git a/Engine/Source/Runtime/ShaderCore/Public/ShaderCore.h b/Engine/Source/Runtime/ShaderCore/Public/ShaderCore.h
--- a/Engine/Source/Runtime/ShaderCore/Public/ShaderCore.h
+++ b/Engine/Source/Runtime/ShaderCore/Public/ShaderCore.h
@@ -342,30 +342,15 @@ struct FShaderCompilerEnvironment : public FRefCountedObject
 	
 	/**
 	 * Works for TCHAR
-	 * e.g. SetDefine(TEXT("NUM_SAMPLES"), TEXT("1"));
-	 */
-	void SetDefine(const TCHAR* Name, const TCHAR* Value)
-	{
-		Definitions.SetDefine(Name, Value);
-	}
-
-	/**
-	 * Works for uint32 and bool
-	 * e.g. OutEnvironment.SetDefine(TEXT("REALLY"), bReally);
-	 * e.g. OutEnvironment.SetDefine(TEXT("NUM_SAMPLES"), NumSamples);
+	 * e.g. SetDefine(TEXT("NAME"), TEXT("Test"));
+	 * e.g. SetDefine(TEXT("NUM_SAMPLES"), 1);
+	 * e.g. SetDefine(TEXT("DOIT"), true);
 	 */
-	void SetDefine(const TCHAR* Name, uint32 Value)
-	{
-		Definitions.SetDefine(Name, Value);
-	}
-
-	/**
-	 * Works for float
-	 */
-	void SetFloatDefine(const TCHAR* Name, float Value)
-	{
-		Definitions.SetFloatDefine(Name, Value);
-	}
+	void SetDefine(const TCHAR* Name, const TCHAR* Value)	{ Definitions.SetDefine(Name, Value); }
+	void SetDefine(const TCHAR* Name, uint32 Value)			{ Definitions.SetDefine(Name, Value); }
+	void SetDefine(const TCHAR* Name, int32 Value)			{ Definitions.SetDefine(Name, (uint32)Value); }
+	void SetDefine(const TCHAR* Name, bool Value)			{ Definitions.SetDefine(Name, (uint32)Value); }
+	void SetDefine(const TCHAR* Name, float Value)			{ Definitions.SetFloatDefine(Name, Value); }
 
 	const TMap<FString,FString>& GetDefinitions() const
 	{
diff --git a/Engine/Source/Runtime/VulkanRHI/Private/VulkanCommandBuffer.cpp b/Engine/Source/Runtime/VulkanRHI/Private/VulkanCommandBuffer.cpp
--- a/Engine/Source/Runtime/VulkanRHI/Private/VulkanCommandBuffer.cpp
+++ b/Engine/Source/Runtime/VulkanRHI/Private/VulkanCommandBuffer.cpp
@@ -7,12 +7,18 @@
 #include "VulkanRHIPrivate.h"
 #include "VulkanCommandBuffer.h"
 
-FVulkanCmdBuffer::FVulkanCmdBuffer(FVulkanDevice* InDevice, FVulkanCommandBufferManager* InCommandBufferManager) :
-	Device(InDevice),
-	CommandBufferManager(InCommandBufferManager),
-	CommandBufferHandle(VK_NULL_HANDLE),
-	IsWriting(VK_FALSE),
-	IsEmpty(VK_TRUE)
+FVulkanCmdBuffer::FVulkanCmdBuffer(FVulkanDevice* InDevice, FVulkanCommandBufferManager* InCommandBufferManager)
+	: Device(InDevice)
+	, CommandBufferManager(InCommandBufferManager)
+	, CommandBufferHandle(VK_NULL_HANDLE)
+#if VULKAN_USE_NEW_COMMAND_BUFFERS
+	, State(EState::ReadyForBegin)
+	, Fence(nullptr)
+	, FenceSignaledCounter(0)
+#else
+	, IsWriting(VK_FALSE)
+	, IsEmpty(VK_TRUE)
+#endif
 {
 	check(Device);
 
@@ -25,42 +31,97 @@ FVulkanCmdBuffer::FVulkanCmdBuffer(FVulkanDevice* InDevice, FVulkanCommandBuffer
 	CreateCmdBufInfo.commandPool = CommandBufferManager->GetHandle();
 
 	VERIFYVULKANRESULT(vkAllocateCommandBuffers(Device->GetInstanceHandle(), &CreateCmdBufInfo, &CommandBufferHandle));
+#if VULKAN_USE_NEW_COMMAND_BUFFERS
+	Fence = Device->GetFenceManager().AllocateFence();
+#endif
 }
 
 FVulkanCmdBuffer::~FVulkanCmdBuffer()
 {
+#if VULKAN_USE_NEW_COMMAND_BUFFERS
+	auto& FenceManager = Device->GetFenceManager();
+	FenceManager.WaitAndReleaseFence(Fence, 0xffffffff);
+#else
 	check(IsWriting == VK_FALSE);
 	if (CommandBufferHandle != VK_NULL_HANDLE)
+#endif
 	{
 		vkFreeCommandBuffers(Device->GetInstanceHandle(), CommandBufferManager->GetHandle(), 1, &CommandBufferHandle);
 		CommandBufferHandle = VK_NULL_HANDLE;
 	}
 }
 
+#if VULKAN_USE_NEW_COMMAND_BUFFERS
+void FVulkanCmdBuffer::BeginRenderPass(const FVulkanRenderTargetLayout& Layout, VkRenderPass RenderPass, VkFramebuffer Framebuffer, const VkClearValue* AttachmentClearValues)
+{
+	check(IsOutsideRenderPass());
+
+	VkRenderPassBeginInfo Info;
+	FMemory::Memzero(Info);
+	Info.sType = VK_STRUCTURE_TYPE_RENDER_PASS_BEGIN_INFO;
+	Info.renderPass = RenderPass;
+	Info.framebuffer = Framebuffer;
+	Info.renderArea.offset.x = 0;
+	Info.renderArea.offset.y = 0;
+	Info.renderArea.extent = Layout.GetExtent2D();
+	Info.clearValueCount = Layout.GetNumAttachments();
+	Info.pClearValues = AttachmentClearValues;
+
+	vkCmdBeginRenderPass(CommandBufferHandle, &Info, VK_SUBPASS_CONTENTS_INLINE);
+
+	State = EState::IsInsideRenderPass;
+}
+#else
 VkCommandBuffer& FVulkanCmdBuffer::GetHandle(const VkBool32 WritingToCommandBuffer /* = true */)
 {
 	check(CommandBufferHandle != VK_NULL_HANDLE);
 	IsEmpty &= !WritingToCommandBuffer;
 	return CommandBufferHandle;
 }
+#endif
+
 
 void FVulkanCmdBuffer::Begin()
 {
+#if VULKAN_USE_NEW_COMMAND_BUFFERS
+	check(State == EState::ReadyForBegin);
+#else
 	check(!IsWriting);
 	check(IsEmpty);
-
+#endif
 	VkCommandBufferBeginInfo CmdBufBeginInfo;
 	FMemory::Memzero(CmdBufBeginInfo);
 	CmdBufBeginInfo.sType = VK_STRUCTURE_TYPE_COMMAND_BUFFER_BEGIN_INFO;
-	CmdBufBeginInfo.pNext = nullptr;
 	CmdBufBeginInfo.flags = VK_COMMAND_BUFFER_USAGE_ONE_TIME_SUBMIT_BIT;
-	CmdBufBeginInfo.pInheritanceInfo = nullptr;
 
-	VERIFYVULKANRESULT(vkBeginCommandBuffer(GetHandle(false), &CmdBufBeginInfo));
+	VERIFYVULKANRESULT(vkBeginCommandBuffer(CommandBufferHandle, &CmdBufBeginInfo));
 
+#if VULKAN_USE_NEW_COMMAND_BUFFERS
+	State = EState::IsInsideBegin;
+#else
 	IsWriting = VK_TRUE;
+#endif
 }
 
+#if VULKAN_USE_NEW_COMMAND_BUFFERS
+void FVulkanCmdBuffer::RefreshFenceStatus()
+{
+	if (State == EState::Submitted)
+	{
+		if (Fence->GetOwner()->IsFenceSignaled(Fence))
+		{
+			State = EState::ReadyForBegin;
+			vkResetCommandBuffer(CommandBufferHandle, VK_COMMAND_BUFFER_RESET_RELEASE_RESOURCES_BIT);
+			Fence->GetOwner()->ResetFence(Fence);
+			++FenceSignaledCounter;
+		}
+	}
+	else
+	{
+		check(!Fence->IsSignaled());
+	}
+}
+#else
 VkBool32 FVulkanCmdBuffer::GetIsWriting() const
 {
 	return IsWriting;
@@ -89,26 +150,41 @@ void FVulkanCmdBuffer::End()
 	IsWriting = VK_FALSE;
 	IsEmpty = VK_TRUE;
 }
-
-FVulkanCommandBufferManager::FVulkanCommandBufferManager(FVulkanDevice* InDevice):
-	Device(InDevice),
-	Handle(VK_NULL_HANDLE)
+#endif
+
+FVulkanCommandBufferManager::FVulkanCommandBufferManager(FVulkanDevice* InDevice)
+	: Device(InDevice)
+	, Handle(VK_NULL_HANDLE)
+#if VULKAN_USE_NEW_COMMAND_BUFFERS
+	, ActiveCmdBuffer(nullptr)
+#else
+#endif
 {
 	check(Device);
 
 	VkCommandPoolCreateInfo CmdPoolInfo;
 	FMemory::Memzero(CmdPoolInfo);
 	CmdPoolInfo.sType = VK_STRUCTURE_TYPE_COMMAND_POOL_CREATE_INFO;
-	CmdPoolInfo.pNext = nullptr;
-	CmdPoolInfo.queueFamilyIndex = Device->GetQueue()->GetNodeIndex();
+	CmdPoolInfo.queueFamilyIndex = Device->GetQueue()->GetFamilyIndex();
+	//#todo-rco: Should we use VK_COMMAND_POOL_CREATE_TRANSIENT_BIT?
 	CmdPoolInfo.flags = VK_COMMAND_POOL_CREATE_RESET_COMMAND_BUFFER_BIT;
 	VERIFYVULKANRESULT(vkCreateCommandPool(Device->GetInstanceHandle(), &CmdPoolInfo, nullptr, &Handle));
+
+#if VULKAN_USE_NEW_COMMAND_BUFFERS
+	ActiveCmdBuffer = Create();
+	ActiveCmdBuffer->Begin();
+#else
+#endif
 }
 
 FVulkanCommandBufferManager::~FVulkanCommandBufferManager()
 {
-	check(Device);
-	check(Handle != VK_NULL_HANDLE);
+	for (int32 Index = 0; Index < CmdBuffers.Num(); ++Index)
+	{
+		FVulkanCmdBuffer* CmdBuffer = CmdBuffers[Index];
+		delete CmdBuffer;
+	}
+
 	vkDestroyCommandPool(Device->GetInstanceHandle(), Handle, nullptr);
 }
 
@@ -121,6 +197,39 @@ FVulkanCmdBuffer* FVulkanCommandBufferManager::Create()
 	return CmdBuffer;
 }
 
+#if VULKAN_USE_NEW_COMMAND_BUFFERS
+void FVulkanCommandBufferManager::RefreshFenceStatus()
+{
+	for (int32 Index = 0; Index < CmdBuffers.Num(); ++Index)
+	{
+		FVulkanCmdBuffer* CmdBuffer = CmdBuffers[Index];
+		CmdBuffer->RefreshFenceStatus();
+	}
+}
+
+void FVulkanCommandBufferManager::PrepareForNewActiveCommandBuffer()
+{
+	for (int32 Index = 0; Index < CmdBuffers.Num(); ++Index)
+	{
+		FVulkanCmdBuffer* CmdBuffer = CmdBuffers[Index];
+		CmdBuffer->RefreshFenceStatus();
+		if (CmdBuffer->State == FVulkanCmdBuffer::EState::ReadyForBegin)
+		{
+			ActiveCmdBuffer = CmdBuffer;
+			ActiveCmdBuffer->Begin();
+			return;
+		}
+		else
+		{
+			check(CmdBuffer->State == FVulkanCmdBuffer::EState::Submitted);
+		}
+	}
+
+	// All cmd buffers are being executed still
+	ActiveCmdBuffer = Create();
+	ActiveCmdBuffer->Begin();
+}
+#else
 void FVulkanCommandBufferManager::Destroy(FVulkanCmdBuffer* CmdBuffer)
 {
 	CmdBuffers.Remove(CmdBuffer);
@@ -138,3 +247,4 @@ void FVulkanCommandBufferManager::Submit(FVulkanCmdBuffer* CmdBuffer)
 {
 	Device->GetQueue()->Submit(CmdBuffer);
 }
+#endif
diff --git a/Engine/Source/Runtime/VulkanRHI/Private/VulkanCommandBuffer.h b/Engine/Source/Runtime/VulkanRHI/Private/VulkanCommandBuffer.h
--- a/Engine/Source/Runtime/VulkanRHI/Private/VulkanCommandBuffer.h
+++ b/Engine/Source/Runtime/VulkanRHI/Private/VulkanCommandBuffer.h
@@ -8,21 +8,96 @@
 
 #include "VulkanRHIPrivate.h"
 
-struct FVulkanDevice;
+class FVulkanDevice;
+class FVulkanCommandBufferManager;
+
+namespace VulkanRHI
+{
+	class FFence;
+}
 
 class FVulkanCmdBuffer
 {
-	friend struct FVulkanCommandBufferManager;
-	friend struct FVulkanQueue;
+protected:
+	friend class FVulkanCommandBufferManager;
+	friend class FVulkanQueue;
 
 	FVulkanCmdBuffer(FVulkanDevice* InDevice, FVulkanCommandBufferManager* InCommandBufferManager);
-
 	~FVulkanCmdBuffer();
 
 public:
+#if VULKAN_USE_NEW_COMMAND_BUFFERS
+/*
+	inline bool IsWritable() const
+	{
+		return State == EState::IsInsideBegin || State == EState::IsInsideRenderPass;
+	}
+*/
+	FVulkanCommandBufferManager* GetOwner()
+	{
+		return CommandBufferManager;
+	}
+
+	inline bool IsInsideRenderPass() const
+	{
+		return State == EState::IsInsideRenderPass;
+	}
+
+	inline bool IsOutsideRenderPass() const
+	{
+		return State == EState::IsInsideBegin;
+	}
+
+	inline bool HasEnded() const
+	{
+		return State == EState::HasEnded;
+	}
+
+	inline VkCommandBuffer GetHandle()
+	{
+		return CommandBufferHandle;
+	}
+
+	void BeginRenderPass(const FVulkanRenderTargetLayout& Layout, VkRenderPass RenderPass, VkFramebuffer Framebuffer, const VkClearValue* AttachmentClearValues);
+
+	void EndRenderPass()
+	{
+		check(IsInsideRenderPass());
+		vkCmdEndRenderPass(CommandBufferHandle);
+		State = EState::IsInsideBegin;
+	}
+
+	void End()
+	{
+		check(IsOutsideRenderPass());
+		VERIFYVULKANRESULT(vkEndCommandBuffer(GetHandle()));
+		State = EState::HasEnded;
+	}
+
+	inline VulkanRHI::FFence* GetFence()
+	{
+		return Fence;
+	}
+
+	inline uint64 GetFenceSignaledCounter() const
+	{
+		return FenceSignaledCounter;
+	}
+#else
 	VkCommandBuffer& GetHandle(const VkBool32 MakeDirty = true);
+#endif
 
 	void Begin();
+#if VULKAN_USE_NEW_COMMAND_BUFFERS
+	enum class EState
+	{
+		ReadyForBegin,
+		IsInsideBegin,
+		IsInsideRenderPass,
+		HasEnded,
+		Submitted,
+	};
+#else
 
 	void Reset();
 
@@ -30,22 +105,39 @@ public:
 
 	VkBool32 GetIsWriting() const;
 	VkBool32 GetIsEmpty() const;
-
+#endif
 private:
 	FVulkanDevice* Device;
 	VkCommandBuffer CommandBufferHandle;
+#if VULKAN_USE_NEW_COMMAND_BUFFERS
+	EState State;
+	VulkanRHI::FFence* Fence;
+	uint64 FenceSignaledCounter;
+	void RefreshFenceStatus();
+#else
 	VkBool32 IsWriting;
 	VkBool32 IsEmpty;
+#endif
 
 	FVulkanCommandBufferManager* CommandBufferManager;
 };
 
-struct FVulkanCommandBufferManager
+class FVulkanCommandBufferManager
 {
+public:
 	FVulkanCommandBufferManager(FVulkanDevice* InDevice);
 
 	~FVulkanCommandBufferManager();
 
+#if VULKAN_USE_NEW_COMMAND_BUFFERS
+	FVulkanCmdBuffer* GetActiveCmdBuffer()
+	{
+		return ActiveCmdBuffer;
+	}
+
+	void RefreshFenceStatus();
+	void PrepareForNewActiveCommandBuffer();
+#else
 	FVulkanCmdBuffer* Create();
 
 	void Destroy(FVulkanCmdBuffer* CmdBuffer);
@@ -55,16 +147,19 @@ struct FVulkanCommandBufferManager
 private:
 	friend class FVulkanCmdBuffer;
 	friend class FVulkanDynamicRHI;
-
+#endif
 	inline VkCommandPool GetHandle() const
 	{
 		check(Handle != VK_NULL_HANDLE);
 		return Handle;
 	}
-
 private:
 	FVulkanDevice* Device;
 	VkCommandPool Handle;
+#if VULKAN_USE_NEW_COMMAND_BUFFERS
+	FVulkanCmdBuffer* ActiveCmdBuffer;
+	FVulkanCmdBuffer* Create();
+#else
+#endif
 	TArray<FVulkanCmdBuffer*> CmdBuffers;
 };
-
