Copying from Release-4.12

* added ensures to track down UE-22073 [CrashReport] Crash when changing Editor Window Size - FDeferredShadingSceneRenderer::ClearTranslucentVolumeLighting()
* fixed UE-30208 Ensure taking high resolution screenshot

diff --git a/Engine/Source/Runtime/Engine/Private/BufferVisualizationData.cpp b/Engine/Source/Runtime/Engine/Private/BufferVisualizationData.cpp
--- a/Engine/Source/Runtime/Engine/Private/BufferVisualizationData.cpp
+++ b/Engine/Source/Runtime/Engine/Private/BufferVisualizationData.cpp
@@ -20,12 +20,6 @@ static TAutoConsoleVariable<int32> CVarBufferVisualizationDumpFrames(
 	TEXT("1:on"),
 	ECVF_RenderThreadSafe);
 
-static TAutoConsoleVariable<int32> CVarSaveUncompressedEXRFrames(
-	TEXT("r.SaveUncompressedEXRFrames"),
-	0,
-	TEXT("If enabled will save EXR screenshots with no compression, rather then the zip default"),
-	ECVF_RenderThreadSafe);
-
 void FBufferVisualizationData::Initialize()
 {
 	if (!bIsInitialized)
diff --git a/Engine/Source/Runtime/Engine/Private/HighResScreenshot.cpp b/Engine/Source/Runtime/Engine/Private/HighResScreenshot.cpp
--- a/Engine/Source/Runtime/Engine/Private/HighResScreenshot.cpp
+++ b/Engine/Source/Runtime/Engine/Private/HighResScreenshot.cpp
@@ -4,6 +4,14 @@
 #include "Slate/SceneViewport.h"
 #include "ImageWrapper.h"
 
+static TAutoConsoleVariable<int32> CVarSaveEXRCompressionQuality(
+	TEXT("r.SaveEXR.CompressionQuality"),
+	1,
+	TEXT("Defines how we save HDR screenshots in the EXR format.\n")
+	TEXT(" 0: no compression\n")
+	TEXT(" 1: default compression which can be slow (default)"),
+	ECVF_RenderThreadSafe);
+
 FHighResScreenshotConfig& GetHighResScreenshotConfig()
 {
 	static FHighResScreenshotConfig Instance;
@@ -204,19 +212,18 @@ bool FHighResScreenshotConfig::SaveImage(const FString& File, const TArray<TPixe
 
 	if (ImageWriter && ImageWriter->ImageWrapper->SetRaw((void*)&Bitmap[0], sizeof(TPixelType)* x * y, x, y, Traits::SourceChannelLayout, BitsPerPixel))
 	{
-		// Optionally disable EXR Compression
-		static const auto CVarSaveUncompressedEXRFrames = IConsoleManager::Get().FindTConsoleVariableDataInt(TEXT("r.SaveUncompressedEXRFrames"));
-		const bool bSaveUncompressedEXR = (CVarSaveUncompressedEXRFrames->GetValueOnRenderThread() > 0);
-
-		const int32 CompressionQuality = (bIsWritingHDRImage && bSaveUncompressedEXR)
-			? ImageCompression::CompressionQuality::Uncompressed
-			: ImageCompression::CompressionQuality::Default;
+		ImageCompression::CompressionQuality LocalCompressionQuality = ImageCompression::Default;
+		
+		if(bIsWritingHDRImage && CVarSaveEXRCompressionQuality.GetValueOnGameThread() == 0)
+		{
+			LocalCompressionQuality = ImageCompression::Uncompressed;
+		}
 
 		// Compress and write image
 		FArchive* Ar = FileManager->CreateFileWriter(Filename.GetCharArray().GetData());
 		if (Ar != nullptr)
 		{
-			const TArray<uint8>& CompressedData = ImageWriter->ImageWrapper->GetCompressed(CompressionQuality);
+		const TArray<uint8>& CompressedData = ImageWriter->ImageWrapper->GetCompressed(LocalCompressionQuality);
 			int32 CompressedSize = CompressedData.Num();
 			Ar->Serialize((void*)CompressedData.GetData(), CompressedSize);
 			delete Ar;
diff --git a/Engine/Source/Runtime/Renderer/Private/PostProcess/SceneRenderTargets.h b/Engine/Source/Runtime/Renderer/Private/PostProcess/SceneRenderTargets.h
--- a/Engine/Source/Runtime/Renderer/Private/PostProcess/SceneRenderTargets.h
+++ b/Engine/Source/Runtime/Renderer/Private/PostProcess/SceneRenderTargets.h
@@ -575,6 +575,8 @@ public:
 	// Can be called when the Scene Color content is no longer needed. As we create SceneColor on demand we can make sure it is created with the right format.
 	// (as a call to SetSceneColor() can override it with a different format)
 	void ReleaseSceneColor();
+	
+	ERHIFeatureLevel::Type GetCurrentFeatureLevel() const { return CurrentFeatureLevel; }
 
 private: // Get...() methods instead of direct access
 
diff --git a/Engine/Source/Runtime/Renderer/Private/TranslucentLighting.cpp b/Engine/Source/Runtime/Renderer/Private/TranslucentLighting.cpp
--- a/Engine/Source/Runtime/Renderer/Private/TranslucentLighting.cpp
+++ b/Engine/Source/Runtime/Renderer/Private/TranslucentLighting.cpp
@@ -1144,22 +1144,35 @@ void FDeferredShadingSceneRenderer::ClearTranslucentVolumeLighting(FRHICommandLi
 		SCOPED_DRAW_EVENT(RHICmdList, ClearTranslucentVolumeLighting);
 		FSceneRenderTargets& SceneContext = FSceneRenderTargets::Get(RHICmdList);
 
-		// Clear all volume textures in the same draw with MRT, which is faster than individually
-
-		static_assert(TVC_MAX == 2, "Only expecting two translucency lighting cascades.");
-		FTextureRHIParamRef RenderTargets[4];
-		RenderTargets[0] = SceneContext.TranslucencyLightingVolumeAmbient[0]->GetRenderTargetItem().TargetableTexture;
-		RenderTargets[1] = SceneContext.TranslucencyLightingVolumeDirectional[0]->GetRenderTargetItem().TargetableTexture;
-		RenderTargets[2] = SceneContext.TranslucencyLightingVolumeAmbient[1]->GetRenderTargetItem().TargetableTexture;
-		RenderTargets[3] = SceneContext.TranslucencyLightingVolumeDirectional[1]->GetRenderTargetItem().TargetableTexture;
-
-		FLinearColor ClearColors[4];
-		ClearColors[0] = FLinearColor(0, 0, 0, 0);
-		ClearColors[1] = FLinearColor(0, 0, 0, 0);
-		ClearColors[2] = FLinearColor(0, 0, 0, 0);
-		ClearColors[3] = FLinearColor(0, 0, 0, 0);
-
-		ClearVolumeTextures<ARRAY_COUNT(RenderTargets)>(RHICmdList, FeatureLevel, RenderTargets, ClearColors);
+		// to track down Jira UE-22073
+		ensure(SceneContext.GetCurrentFeatureLevel() >= ERHIFeatureLevel::SM4);
+		ensure(SceneContext.TranslucencyLightingVolumeAmbient[0]);
+		ensure(SceneContext.TranslucencyLightingVolumeDirectional[0]);
+		ensure(SceneContext.TranslucencyLightingVolumeAmbient[1]);
+		ensure(SceneContext.TranslucencyLightingVolumeDirectional[1]);
+		// if to prevent crash (Jira UE-22073) but it happen later
+		if( SceneContext.TranslucencyLightingVolumeAmbient[0] &&
+			SceneContext.TranslucencyLightingVolumeDirectional[0] &&
+			SceneContext.TranslucencyLightingVolumeAmbient[1] &&
+			SceneContext.TranslucencyLightingVolumeDirectional[1])
+		{
+			// Clear all volume textures in the same draw with MRT, which is faster than individually
+
+			static_assert(TVC_MAX == 2, "Only expecting two translucency lighting cascades.");
+			FTextureRHIParamRef RenderTargets[4];
+			RenderTargets[0] = SceneContext.TranslucencyLightingVolumeAmbient[0]->GetRenderTargetItem().TargetableTexture;
+			RenderTargets[1] = SceneContext.TranslucencyLightingVolumeDirectional[0]->GetRenderTargetItem().TargetableTexture;
+			RenderTargets[2] = SceneContext.TranslucencyLightingVolumeAmbient[1]->GetRenderTargetItem().TargetableTexture;
+			RenderTargets[3] = SceneContext.TranslucencyLightingVolumeDirectional[1]->GetRenderTargetItem().TargetableTexture;
+
+			FLinearColor ClearColors[4];
+			ClearColors[0] = FLinearColor(0, 0, 0, 0);
+			ClearColors[1] = FLinearColor(0, 0, 0, 0);
+			ClearColors[2] = FLinearColor(0, 0, 0, 0);
+			ClearColors[3] = FLinearColor(0, 0, 0, 0);
+
+			ClearVolumeTextures<ARRAY_COUNT(RenderTargets)>(RHICmdList, FeatureLevel, RenderTargets, ClearColors);
+		}
 	}
 }
 
@@ -1214,13 +1227,17 @@ void FDeferredShadingSceneRenderer::InjectAmbientCubemapTranslucentVolumeLightin
 
 	if (GUseTranslucentLightingVolumes && GSupportsVolumeTextureRendering && View.FinalPostProcessSettings.ContributingCubemaps.Num() && !IsSimpleDynamicLightingEnabled())
 	{
+		FSceneRenderTargets& SceneContext = FSceneRenderTargets::Get(RHICmdList);
+
+		// to track down Jira UE-22073
+		ensure(SceneContext.GetCurrentFeatureLevel() >= ERHIFeatureLevel::SM4);
+
 		SCOPED_DRAW_EVENT(RHICmdList, InjectAmbientCubemapTranslucentVolumeLighting);
 
 		const FVolumeBounds VolumeBounds(GTranslucencyLightingVolumeDim);
 
 		auto ShaderMap = GetGlobalShaderMap(FeatureLevel);
 		
-		FSceneRenderTargets& SceneContext = FSceneRenderTargets::Get(RHICmdList);
 		for (int32 VolumeCascadeIndex = 0; VolumeCascadeIndex < TVC_MAX; VolumeCascadeIndex++)
 		{
 			// we don't update the directional volume (could be a HQ option)
@@ -1544,6 +1561,9 @@ static void InjectTranslucentLightArray(FRHICommandListImmediate& RHICmdList, co
 	FSceneRenderTargets& SceneContext = FSceneRenderTargets::Get(RHICmdList);
 	INC_DWORD_STAT_BY(STAT_NumLightsInjectedIntoTranslucency, LightInjectionData.Num());
 
+	// to track down Jira UE-22073
+	ensure(SceneContext.GetCurrentFeatureLevel() >= ERHIFeatureLevel::SM4);
+
 	// Inject into each volume cascade
 	// Operate on one cascade at a time to reduce render target switches
 	for (uint32 VolumeCascadeIndex = 0; VolumeCascadeIndex < TVC_MAX; VolumeCascadeIndex++)
