Copying //UE4/Dev-Rendering to //UE4/Dev-Main (Source: //UE4/Dev-Rendering @ 2981742)

diff --git a/Engine/Source/Runtime/Renderer/Private/CompositionLighting/CompositionLighting.cpp b/Engine/Source/Runtime/Renderer/Private/CompositionLighting/CompositionLighting.cpp
--- a/Engine/Source/Runtime/Renderer/Private/CompositionLighting/CompositionLighting.cpp
+++ b/Engine/Source/Runtime/Renderer/Private/CompositionLighting/CompositionLighting.cpp
@@ -47,11 +47,18 @@ static TAutoConsoleVariable<int32> CVarSSSHalfRes(
 	TEXT(" 1: parts of the algorithm runs in half resolution which is lower quality but faster (default)"),
 	ECVF_RenderThreadSafe | ECVF_Scalability);
 
+static TAutoConsoleVariable<int32> CVarSubsurfaceScattering(
+	TEXT("r.SubsurfaceScattering"),
+	1,
+	TEXT(" 0: disabled\n")
+	TEXT(" 1: enabled (default)"),
+	ECVF_RenderThreadSafe | ECVF_Scalability);
+
 bool IsAmbientCubemapPassRequired(const FSceneView& View)
 {
 	FScene* Scene = (FScene*)View.Family->Scene;
 
-	return View.FinalPostProcessSettings.ContributingCubemaps.Num() != 0 && !IsSimpleDynamicLightingEnabled();
+	return View.FinalPostProcessSettings.ContributingCubemaps.Num() != 0 && !IsSimpleForwardShadingEnabled(View.GetShaderPlatform());
 }
 
 bool IsLpvIndirectPassRequired(const FViewInfo& View)
@@ -91,7 +98,7 @@ static bool IsReflectionEnvironmentActive(const FSceneView& View)
 	bool HasReflectionCaptures = (Scene->ReflectionSceneData.RegisteredReflectionCaptures.Num() > 0);
 	bool HasSSR = View.Family->EngineShowFlags.ScreenSpaceReflections;
 
-	return (Scene->GetFeatureLevel() == ERHIFeatureLevel::SM5 && IsReflectingEnvironment && (HasReflectionCaptures || HasSSR) && !IsSimpleDynamicLightingEnabled());
+	return (Scene->GetFeatureLevel() == ERHIFeatureLevel::SM5 && IsReflectingEnvironment && (HasReflectionCaptures || HasSSR) && !IsSimpleForwardShadingEnabled(View.GetShaderPlatform()));
 }
 
 static bool IsSkylightActive(const FViewInfo& View)
@@ -116,7 +123,7 @@ uint32 ComputeAmbientOcclusionPassCount(const FViewInfo& View)
 			&& View.FinalPostProcessSettings.AmbientOcclusionRadius >= 0.1f
 			&& !View.Family->UseDebugViewPS()
 			&& (FSSAOHelper::IsBasePassAmbientOcclusionRequired(View) || IsAmbientCubemapPassRequired(View) || IsReflectionEnvironmentActive(View) || IsSkylightActive(View) || View.Family->EngineShowFlags.VisualizeBuffer)
-			&& !IsSimpleDynamicLightingEnabled();
+			&& !IsSimpleForwardShadingEnabled(View.GetShaderPlatform());
 	}
 
 	if (bEnabled)
@@ -429,11 +436,13 @@ void FCompositionLighting::ProcessAfterLighting(FRHICommandListImmediate& RHICmd
 		// Screen Space Subsurface Scattering
 		{
 			float Radius = CVarSSSScale.GetValueOnRenderThread();
-
-			bool bSimpleDynamicLighting = IsSimpleDynamicLightingEnabled();
-
+			bool bSimpleDynamicLighting = IsSimpleForwardShadingEnabled(View.GetShaderPlatform());
 			bool bScreenSpaceSubsurfacePassNeeded = (View.ShadingModelMaskInView & (1 << MSM_SubsurfaceProfile)) != 0;
-			if (bScreenSpaceSubsurfacePassNeeded && !bSimpleDynamicLighting)
+			bool bSubsurfaceAllowed = CVarSubsurfaceScattering.GetValueOnRenderThread() == 1;
+
+			if (bScreenSpaceSubsurfacePassNeeded 
+				&& !bSimpleDynamicLighting 
+				&& bSubsurfaceAllowed)
 			{
 				bool bHalfRes = CVarSSSHalfRes.GetValueOnRenderThread() != 0;
 				bool bSingleViewportMode = View.Family->Views.Num() == 1;
diff --git a/Engine/Source/Runtime/Renderer/Private/CompositionLighting/PostProcessAmbientOcclusion.cpp b/Engine/Source/Runtime/Renderer/Private/CompositionLighting/PostProcessAmbientOcclusion.cpp
--- a/Engine/Source/Runtime/Renderer/Private/CompositionLighting/PostProcessAmbientOcclusion.cpp
+++ b/Engine/Source/Runtime/Renderer/Private/CompositionLighting/PostProcessAmbientOcclusion.cpp
@@ -120,7 +120,7 @@ EAsyncComputeBudget FSSAOHelper::GetAmbientOcclusionAsyncComputeBudget()
 bool FSSAOHelper::IsBasePassAmbientOcclusionRequired(const FViewInfo& View)
 {
 	// the BaseAO pass is only worth with some AO
-	return (View.FinalPostProcessSettings.AmbientOcclusionStaticFraction >= 1 / 100.0f) && !IsSimpleDynamicLightingEnabled();
+	return (View.FinalPostProcessSettings.AmbientOcclusionStaticFraction >= 1 / 100.0f) && !IsSimpleForwardShadingEnabled(View.GetShaderPlatform());
 }
 
 bool FSSAOHelper::IsAmbientOcclusionAsyncCompute(const FViewInfo& View, uint32 AOPassCount)
diff --git a/Engine/Source/Runtime/Renderer/Private/DeferredShadingRenderer.cpp b/Engine/Source/Runtime/Renderer/Private/DeferredShadingRenderer.cpp
--- a/Engine/Source/Runtime/Renderer/Private/DeferredShadingRenderer.cpp
+++ b/Engine/Source/Runtime/Renderer/Private/DeferredShadingRenderer.cpp
@@ -918,8 +918,7 @@ void FDeferredShadingSceneRenderer::Render(FRHICommandListImmediate& RHICmdList)
 	bool bRequiresRHIClear = true;
 	bool bRequiresFarZQuadClear = false;
 
-	static const auto GBufferCVar = IConsoleManager::Get().FindTConsoleVariableDataInt(TEXT("r.GBuffer"));
-	bool bGBuffer = GBufferCVar ? (GBufferCVar->GetValueOnRenderThread() != 0) : true;
+	bool bUseGBuffer = !IsSimpleForwardShadingEnabled(GetFeatureLevelShaderPlatform(FeatureLevel));
 
 	if (ClearMethodCVar)
 	{
@@ -1224,7 +1223,7 @@ void FDeferredShadingSceneRenderer::Render(FRHICommandListImmediate& RHICmdList)
 	// Pre-lighting composition lighting stage
 	// e.g. deferred decals, SSAO
 	if (FeatureLevel >= ERHIFeatureLevel::SM4
-		&& bGBuffer)
+		&& bUseGBuffer)
 	{
 		SCOPE_CYCLE_COUNTER(STAT_FDeferredShadingSceneRenderer_AfterBasePass);
 
@@ -1264,8 +1263,7 @@ void FDeferredShadingSceneRenderer::Render(FRHICommandListImmediate& RHICmdList)
 	if (ViewFamily.EngineShowFlags.Lighting
 		&& FeatureLevel >= ERHIFeatureLevel::SM4
 		&& ViewFamily.EngineShowFlags.DeferredLighting
-		&& bGBuffer
-		)
+		&& bUseGBuffer)
 	{
 		SCOPE_CYCLE_COUNTER(STAT_FDeferredShadingSceneRenderer_Lighting);
 
diff --git a/Engine/Source/Runtime/Renderer/Private/DistanceFieldSurfaceCacheLighting.cpp b/Engine/Source/Runtime/Renderer/Private/DistanceFieldSurfaceCacheLighting.cpp
--- a/Engine/Source/Runtime/Renderer/Private/DistanceFieldSurfaceCacheLighting.cpp
+++ b/Engine/Source/Runtime/Renderer/Private/DistanceFieldSurfaceCacheLighting.cpp
@@ -2740,7 +2740,7 @@ bool ShouldRenderDynamicSkyLight(const FScene* Scene, const FSceneViewFamily& Vi
 		&& !Scene->SkyLight->bHasStaticLighting
 		&& ViewFamily.EngineShowFlags.SkyLighting
 		&& Scene->GetFeatureLevel() >= ERHIFeatureLevel::SM4
-		&& !IsSimpleDynamicLightingEnabled() 
+		&& !IsSimpleForwardShadingEnabled(Scene->GetShaderPlatform()) 
 		&& !ViewFamily.EngineShowFlags.VisualizeLightCulling;
 }
 
diff --git a/Engine/Source/Runtime/Renderer/Private/FogRendering.cpp b/Engine/Source/Runtime/Renderer/Private/FogRendering.cpp
--- a/Engine/Source/Runtime/Renderer/Private/FogRendering.cpp
+++ b/Engine/Source/Runtime/Renderer/Private/FogRendering.cpp
@@ -28,6 +28,13 @@ static TAutoConsoleVariable<float> CVarFogDensity(
 	ECVF_Cheat | ECVF_RenderThreadSafe);
 #endif
 
+static TAutoConsoleVariable<int32> CVarFog(
+	TEXT("r.Fog"),
+	1,
+	TEXT(" 0: disabled\n")
+	TEXT(" 1: enabled (default)"),
+	ECVF_RenderThreadSafe | ECVF_Scalability);
+
 /** Binds the parameters. */
 void FExponentialHeightFogShaderParameters::Bind(const FShaderParameterMap& ParameterMap)
 {
@@ -357,6 +364,7 @@ bool ShouldRenderFog(const FSceneViewFamily& Family)
 	return EngineShowFlags.Fog
 		&& EngineShowFlags.Materials 
 		&& !Family.UseDebugViewPS()
+		&& CVarFog.GetValueOnRenderThread() == 1
 		&& !EngineShowFlags.StationaryLightOverlap 
 		&& !EngineShowFlags.VertexDensities
 		&& !EngineShowFlags.LightMapDensity;
diff --git a/Engine/Source/Runtime/Renderer/Private/ForwardBasePassRendering.cpp b/Engine/Source/Runtime/Renderer/Private/ForwardBasePassRendering.cpp
--- a/Engine/Source/Runtime/Renderer/Private/ForwardBasePassRendering.cpp
+++ b/Engine/Source/Runtime/Renderer/Private/ForwardBasePassRendering.cpp
@@ -46,10 +46,10 @@ IMPLEMENT_FORWARD_SHADING_BASEPASS_LIGHTMAPPED_SHADER_TYPE(TUniformLightMapPolic
 IMPLEMENT_FORWARD_SHADING_BASEPASS_LIGHTMAPPED_SHADER_TYPE(TUniformLightMapPolicy<LMP_LQ_LIGHTMAP>, TLightMapPolicyLQ);
 IMPLEMENT_FORWARD_SHADING_BASEPASS_LIGHTMAPPED_SHADER_TYPE(TUniformLightMapPolicy<LMP_DISTANCE_FIELD_SHADOWS_AND_LQ_LIGHTMAP>, TDistanceFieldShadowsAndLightMapPolicyLQ);
 IMPLEMENT_FORWARD_SHADING_BASEPASS_LIGHTMAPPED_SHADER_TYPE(TUniformLightMapPolicy<LMP_DISTANCE_FIELD_SHADOWS_LIGHTMAP_AND_CSM>, FDistanceFieldShadowsLightMapAndCSMLightingPolicy);
-IMPLEMENT_FORWARD_SHADING_BASEPASS_LIGHTMAPPED_SHADER_TYPE(TUniformLightMapPolicy<LMP_SIMPLE_DIRECTIONAL_LIGHT_AND_SH_INDIRECT>, FSimpleDirectionalLightAndSHIndirectPolicy);
-IMPLEMENT_FORWARD_SHADING_BASEPASS_LIGHTMAPPED_SHADER_TYPE(TUniformLightMapPolicy<LMP_SIMPLE_DIRECTIONAL_LIGHT_AND_SH_DIRECTIONAL_INDIRECT>, FSimpleDirectionalLightAndSHDirectionalIndirectPolicy);
-IMPLEMENT_FORWARD_SHADING_BASEPASS_LIGHTMAPPED_SHADER_TYPE(TUniformLightMapPolicy<LMP_SIMPLE_DIRECTIONAL_LIGHT_AND_SH_DIRECTIONAL_CSM_INDIRECT>, FSimpleDirectionalLightAndSHDirectionalCSMIndirectPolicy);
-IMPLEMENT_FORWARD_SHADING_BASEPASS_LIGHTMAPPED_SHADER_TYPE(TUniformLightMapPolicy<LMP_SIMPLE_DIRECTIONAL_LIGHT_AND_SH_CSM_INDIRECT>, FSimpleDirectionalLightAndSHCSMIndirectPolicy);
+IMPLEMENT_FORWARD_SHADING_BASEPASS_LIGHTMAPPED_SHADER_TYPE(TUniformLightMapPolicy<LMP_FORWARD_DIRECTIONAL_LIGHT_AND_SH_INDIRECT>, FForwardDirectionalLightAndSHIndirectPolicy);
+IMPLEMENT_FORWARD_SHADING_BASEPASS_LIGHTMAPPED_SHADER_TYPE(TUniformLightMapPolicy<LMP_FORWARD_DIRECTIONAL_LIGHT_AND_SH_DIRECTIONAL_INDIRECT>, FForwardDirectionalLightAndSHDirectionalIndirectPolicy);
+IMPLEMENT_FORWARD_SHADING_BASEPASS_LIGHTMAPPED_SHADER_TYPE(TUniformLightMapPolicy<LMP_FORWARD_DIRECTIONAL_LIGHT_AND_SH_DIRECTIONAL_CSM_INDIRECT>, FForwardDirectionalLightAndSHDirectionalCSMIndirectPolicy);
+IMPLEMENT_FORWARD_SHADING_BASEPASS_LIGHTMAPPED_SHADER_TYPE(TUniformLightMapPolicy<LMP_FORWARD_DIRECTIONAL_LIGHT_AND_SH_CSM_INDIRECT>, FForwardDirectionalLightAndSHCSMIndirectPolicy);
 IMPLEMENT_FORWARD_SHADING_BASEPASS_LIGHTMAPPED_SHADER_TYPE(TUniformLightMapPolicy<LMP_MOVABLE_DIRECTIONAL_LIGHT>, FMovableDirectionalLightLightingPolicy);
 IMPLEMENT_FORWARD_SHADING_BASEPASS_LIGHTMAPPED_SHADER_TYPE(TUniformLightMapPolicy<LMP_MOVABLE_DIRECTIONAL_LIGHT_CSM>, FMovableDirectionalLightCSMLightingPolicy);
 IMPLEMENT_FORWARD_SHADING_BASEPASS_LIGHTMAPPED_SHADER_TYPE(TUniformLightMapPolicy<LMP_MOVABLE_DIRECTIONAL_LIGHT_WITH_LIGHTMAP>, FMovableDirectionalLightWithLightmapLightingPolicy);
diff --git a/Engine/Source/Runtime/Renderer/Private/ForwardBasePassRendering.h b/Engine/Source/Runtime/Renderer/Private/ForwardBasePassRendering.h
--- a/Engine/Source/Runtime/Renderer/Private/ForwardBasePassRendering.h
+++ b/Engine/Source/Runtime/Renderer/Private/ForwardBasePassRendering.h
@@ -551,17 +551,17 @@ GetBasePassForForwardShadingShaders<FUniformLightMapPolicy, NumDynamicPointLight
 	case LMP_DISTANCE_FIELD_SHADOWS_LIGHTMAP_AND_CSM:
 		GetUniformBasePassorForwardShadingShaders<LMP_DISTANCE_FIELD_SHADOWS_LIGHTMAP_AND_CSM, NumDynamicPointLights>(Material, VertexFactoryType, bEnableSkyLight, VertexShader, PixelShader);
 		break;
-	case LMP_SIMPLE_DIRECTIONAL_LIGHT_AND_SH_INDIRECT:
-		GetUniformBasePassorForwardShadingShaders<LMP_SIMPLE_DIRECTIONAL_LIGHT_AND_SH_INDIRECT, NumDynamicPointLights>(Material, VertexFactoryType, bEnableSkyLight, VertexShader, PixelShader);
+	case LMP_FORWARD_DIRECTIONAL_LIGHT_AND_SH_INDIRECT:
+		GetUniformBasePassorForwardShadingShaders<LMP_FORWARD_DIRECTIONAL_LIGHT_AND_SH_INDIRECT, NumDynamicPointLights>(Material, VertexFactoryType, bEnableSkyLight, VertexShader, PixelShader);
 		break;
-	case LMP_SIMPLE_DIRECTIONAL_LIGHT_AND_SH_DIRECTIONAL_INDIRECT:
-		GetUniformBasePassorForwardShadingShaders<LMP_SIMPLE_DIRECTIONAL_LIGHT_AND_SH_DIRECTIONAL_INDIRECT, NumDynamicPointLights>(Material, VertexFactoryType, bEnableSkyLight, VertexShader, PixelShader);
+	case LMP_FORWARD_DIRECTIONAL_LIGHT_AND_SH_DIRECTIONAL_INDIRECT:
+		GetUniformBasePassorForwardShadingShaders<LMP_FORWARD_DIRECTIONAL_LIGHT_AND_SH_DIRECTIONAL_INDIRECT, NumDynamicPointLights>(Material, VertexFactoryType, bEnableSkyLight, VertexShader, PixelShader);
 		break;
-	case LMP_SIMPLE_DIRECTIONAL_LIGHT_AND_SH_DIRECTIONAL_CSM_INDIRECT:
-		GetUniformBasePassorForwardShadingShaders<LMP_SIMPLE_DIRECTIONAL_LIGHT_AND_SH_DIRECTIONAL_CSM_INDIRECT, NumDynamicPointLights>(Material, VertexFactoryType, bEnableSkyLight, VertexShader, PixelShader);
+	case LMP_FORWARD_DIRECTIONAL_LIGHT_AND_SH_DIRECTIONAL_CSM_INDIRECT:
+		GetUniformBasePassorForwardShadingShaders<LMP_FORWARD_DIRECTIONAL_LIGHT_AND_SH_DIRECTIONAL_CSM_INDIRECT, NumDynamicPointLights>(Material, VertexFactoryType, bEnableSkyLight, VertexShader, PixelShader);
 		break;
-	case LMP_SIMPLE_DIRECTIONAL_LIGHT_AND_SH_CSM_INDIRECT:
-		GetUniformBasePassorForwardShadingShaders<LMP_SIMPLE_DIRECTIONAL_LIGHT_AND_SH_CSM_INDIRECT, NumDynamicPointLights>(Material, VertexFactoryType, bEnableSkyLight, VertexShader, PixelShader);
+	case LMP_FORWARD_DIRECTIONAL_LIGHT_AND_SH_CSM_INDIRECT:
+		GetUniformBasePassorForwardShadingShaders<LMP_FORWARD_DIRECTIONAL_LIGHT_AND_SH_CSM_INDIRECT, NumDynamicPointLights>(Material, VertexFactoryType, bEnableSkyLight, VertexShader, PixelShader);
 		break;
 	case LMP_MOVABLE_DIRECTIONAL_LIGHT:
 		GetUniformBasePassorForwardShadingShaders<LMP_MOVABLE_DIRECTIONAL_LIGHT, NumDynamicPointLights>(Material, VertexFactoryType, bEnableSkyLight, VertexShader, PixelShader);
@@ -949,22 +949,22 @@ void ProcessBasePassMeshForForwardShading(
 			{
 				if (SimpleDirectionalLight && SimpleDirectionalLight->ShouldRenderViewIndependentWholeSceneShadows())
 				{
-					Action.template Process<FUniformLightMapPolicy, NumDynamicPointLights>(RHICmdList, Parameters, FUniformLightMapPolicy(LMP_SIMPLE_DIRECTIONAL_LIGHT_AND_SH_DIRECTIONAL_CSM_INDIRECT), Parameters.Mesh.LCI);
+					Action.template Process<FUniformLightMapPolicy, NumDynamicPointLights>(RHICmdList, Parameters, FUniformLightMapPolicy(LMP_FORWARD_DIRECTIONAL_LIGHT_AND_SH_DIRECTIONAL_CSM_INDIRECT), Parameters.Mesh.LCI);
 				}
 				else
 				{
-					Action.template Process<FUniformLightMapPolicy, NumDynamicPointLights>(RHICmdList, Parameters, FUniformLightMapPolicy(LMP_SIMPLE_DIRECTIONAL_LIGHT_AND_SH_DIRECTIONAL_INDIRECT), Parameters.Mesh.LCI);
+					Action.template Process<FUniformLightMapPolicy, NumDynamicPointLights>(RHICmdList, Parameters, FUniformLightMapPolicy(LMP_FORWARD_DIRECTIONAL_LIGHT_AND_SH_DIRECTIONAL_INDIRECT), Parameters.Mesh.LCI);
 				}
 			}
 			else
 			{
 				if (bUseCSM)
 				{
-					Action.template Process<FUniformLightMapPolicy, NumDynamicPointLights>(RHICmdList, Parameters, FUniformLightMapPolicy(LMP_SIMPLE_DIRECTIONAL_LIGHT_AND_SH_CSM_INDIRECT), Parameters.Mesh.LCI);
+					Action.template Process<FUniformLightMapPolicy, NumDynamicPointLights>(RHICmdList, Parameters, FUniformLightMapPolicy(LMP_FORWARD_DIRECTIONAL_LIGHT_AND_SH_CSM_INDIRECT), Parameters.Mesh.LCI);
 				}
 				else
 				{
-					Action.template Process<FUniformLightMapPolicy, NumDynamicPointLights>(RHICmdList, Parameters, FUniformLightMapPolicy(LMP_SIMPLE_DIRECTIONAL_LIGHT_AND_SH_INDIRECT), Parameters.Mesh.LCI);
+					Action.template Process<FUniformLightMapPolicy, NumDynamicPointLights>(RHICmdList, Parameters, FUniformLightMapPolicy(LMP_FORWARD_DIRECTIONAL_LIGHT_AND_SH_INDIRECT), Parameters.Mesh.LCI);
 				}
 			}
 
diff --git a/Engine/Source/Runtime/Renderer/Private/ForwardShadingRenderer.cpp b/Engine/Source/Runtime/Renderer/Private/ForwardShadingRenderer.cpp
--- a/Engine/Source/Runtime/Renderer/Private/ForwardShadingRenderer.cpp
+++ b/Engine/Source/Runtime/Renderer/Private/ForwardShadingRenderer.cpp
@@ -42,7 +42,7 @@ void FForwardShadingSceneRenderer::InitViews(FRHICommandListImmediate& RHICmdLis
 
 	bool bDynamicShadows = ViewFamily.EngineShowFlags.DynamicShadows && GetShadowQuality() > 0;
 
-	if (bDynamicShadows && !IsSimpleDynamicLightingEnabled())
+	if (bDynamicShadows && !IsSimpleForwardShadingEnabled(GetFeatureLevelShaderPlatform(FeatureLevel)))
 	{
 		// Setup dynamic shadows.
 		InitDynamicShadows(RHICmdList);		
diff --git a/Engine/Source/Runtime/Renderer/Private/IndirectLightingCache.cpp b/Engine/Source/Runtime/Renderer/Private/IndirectLightingCache.cpp
--- a/Engine/Source/Runtime/Renderer/Private/IndirectLightingCache.cpp
+++ b/Engine/Source/Runtime/Renderer/Private/IndirectLightingCache.cpp
@@ -507,6 +507,13 @@ void FIndirectLightingCache::UpdateCachePrimitivesInternal(FScene* Scene, FScene
 						uint32 PrimitiveIndex = BitIt.GetIndex();
 						SetBitIndices[ViewIndex].Add(PrimitiveIndex);					
 					}
+
+					// Any visible primitives with an indirect shadow need their ILC updated, since that determines the indirect shadow direction
+					for (int32 IndirectPrimitiveIndex = 0; IndirectPrimitiveIndex < View.IndirectShadowPrimitives.Num(); IndirectPrimitiveIndex++)
+					{
+						int32 PrimitiveIndex = View.IndirectShadowPrimitives[IndirectPrimitiveIndex]->GetIndex();
+						SetBitIndices[ViewIndex].AddUnique(PrimitiveIndex);
+					}
 				}			
 			}
 
diff --git a/Engine/Source/Runtime/Renderer/Private/LightMapRendering.h b/Engine/Source/Runtime/Renderer/Private/LightMapRendering.h
--- a/Engine/Source/Runtime/Renderer/Private/LightMapRendering.h
+++ b/Engine/Source/Runtime/Renderer/Private/LightMapRendering.h
@@ -355,69 +355,152 @@ struct FCachedPointIndirectLightingPolicy
 	static void ModifyCompilationEnvironment(EShaderPlatform Platform, const FMaterial* Material, FShaderCompilerEnvironment& OutEnvironment);
 };
 
+/**
+ * Renders the base pass without outputting to GBuffers, used to support low end hardware where deferred shading is too expensive.
+ */
+struct FSimpleNoLightmapLightingPolicy : public FNoLightMapPolicy
+{
+	static bool ShouldCache(EShaderPlatform Platform,const FMaterial* Material,const FVertexFactoryType* VertexFactoryType)
+	{
+		return PlatformSupportsSimpleForwardShading(Platform)
+			&& FNoLightMapPolicy::ShouldCache(Platform, Material, VertexFactoryType);
+	}
+
+	static void ModifyCompilationEnvironment(EShaderPlatform Platform, const FMaterial* Material, FShaderCompilerEnvironment& OutEnvironment)
+	{
+		OutEnvironment.SetDefine(TEXT("SIMPLE_FORWARD_SHADING"),TEXT("1"));
+		FNoLightMapPolicy::ModifyCompilationEnvironment(Platform, Material, OutEnvironment);
+	}
+};
+
+/**
+ * Renders lightmaps without outputting to GBuffers, used to support low end hardware where deferred shading is too expensive.
+ */
+struct FSimpleLightmapOnlyLightingPolicy : public TLightMapPolicy<HQ_LIGHTMAP>
+{
+	static bool ShouldCache(EShaderPlatform Platform,const FMaterial* Material,const FVertexFactoryType* VertexFactoryType)
+	{
+		static const auto AllowStaticLightingVar = IConsoleManager::Get().FindTConsoleVariableDataInt(TEXT("r.AllowStaticLighting"));
+
+		return AllowStaticLightingVar->GetValueOnAnyThread() != 0
+			&& PlatformSupportsSimpleForwardShading(Platform)
+			&& TLightMapPolicy<HQ_LIGHTMAP>::ShouldCache(Platform, Material, VertexFactoryType);
+	}
+
+	static void ModifyCompilationEnvironment(EShaderPlatform Platform, const FMaterial* Material, FShaderCompilerEnvironment& OutEnvironment)
+	{
+		OutEnvironment.SetDefine(TEXT("SIMPLE_FORWARD_SHADING"),TEXT("1"));
+		TLightMapPolicy<HQ_LIGHTMAP>::ModifyCompilationEnvironment(Platform, Material, OutEnvironment);
+	}
+};
 
 /**
  * Renders an unshadowed directional light in the base pass, used to support low end hardware where deferred shading is too expensive.
  */
-struct FSimpleDynamicLightingPolicy
+struct FSimpleDirectionalLightLightingPolicy
 {
 	static bool ShouldCache(EShaderPlatform Platform,const FMaterial* Material,const FVertexFactoryType* VertexFactoryType)
 	{
-		return Material->GetShadingModel() != MSM_Unlit;
+		return PlatformSupportsSimpleForwardShading(Platform)
+			&& Material->GetShadingModel() != MSM_Unlit;
 	}
 
 	static void ModifyCompilationEnvironment(EShaderPlatform Platform, const FMaterial* Material, FShaderCompilerEnvironment& OutEnvironment)
 	{
-		OutEnvironment.SetDefine(TEXT("SIMPLE_DYNAMIC_LIGHTING"),TEXT("1"));
+		OutEnvironment.SetDefine(TEXT("SIMPLE_FORWARD_SHADING"),TEXT("1"));
+		OutEnvironment.SetDefine(TEXT("SIMPLE_FORWARD_DIRECTIONAL_LIGHT"),TEXT("1"));
+	}
+};
+
+/**
+ * Renders a stationary directional light in the base pass with distance field precomputed shadows without outputting to GBuffers, used to support low end hardware where deferred shading is too expensive.
+ */
+struct FSimpleStationaryLightPrecomputedShadowsLightingPolicy : public TDistanceFieldShadowsAndLightMapPolicy<HQ_LIGHTMAP>
+{
+	static bool ShouldCache(EShaderPlatform Platform,const FMaterial* Material,const FVertexFactoryType* VertexFactoryType)
+	{
+		static const auto AllowStaticLightingVar = IConsoleManager::Get().FindTConsoleVariableDataInt(TEXT("r.AllowStaticLighting"));
+
+		return AllowStaticLightingVar->GetValueOnAnyThread() != 0
+			&& PlatformSupportsSimpleForwardShading(Platform)
+			&& TDistanceFieldShadowsAndLightMapPolicy<HQ_LIGHTMAP>::ShouldCache(Platform, Material, VertexFactoryType);
+	}
+
+	static void ModifyCompilationEnvironment(EShaderPlatform Platform, const FMaterial* Material, FShaderCompilerEnvironment& OutEnvironment)
+	{
+		OutEnvironment.SetDefine(TEXT("SIMPLE_FORWARD_SHADING"),TEXT("1"));
+		OutEnvironment.SetDefine(TEXT("SIMPLE_FORWARD_DIRECTIONAL_LIGHT"),TEXT("1"));
+		TDistanceFieldShadowsAndLightMapPolicy<HQ_LIGHTMAP>::ModifyCompilationEnvironment(Platform, Material, OutEnvironment);
+	}
+};
+
+/**
+ * Renders a stationary directional light in the base pass with single sample shadows without outputting to GBuffers, used to support low end hardware where deferred shading is too expensive.
+ */
+struct FSimpleStationaryLightSingleSampleShadowsLightingPolicy : public FCachedPointIndirectLightingPolicy
+{
+	static bool ShouldCache(EShaderPlatform Platform,const FMaterial* Material,const FVertexFactoryType* VertexFactoryType)
+	{
+		static const auto AllowStaticLightingVar = IConsoleManager::Get().FindTConsoleVariableDataInt(TEXT("r.AllowStaticLighting"));
+
+		return AllowStaticLightingVar->GetValueOnAnyThread() != 0
+			&& PlatformSupportsSimpleForwardShading(Platform)
+			&& FCachedPointIndirectLightingPolicy::ShouldCache(Platform, Material, VertexFactoryType);
+	}
+
+	static void ModifyCompilationEnvironment(EShaderPlatform Platform, const FMaterial* Material, FShaderCompilerEnvironment& OutEnvironment)
+	{
+		OutEnvironment.SetDefine(TEXT("SIMPLE_FORWARD_SHADING"),TEXT("1"));
+		OutEnvironment.SetDefine(TEXT("SIMPLE_FORWARD_DIRECTIONAL_LIGHT"),TEXT("1"));
+		FCachedPointIndirectLightingPolicy::ModifyCompilationEnvironment(Platform, Material, OutEnvironment);
 	}
 };
 
 /** Combines an unshadowed directional light with indirect lighting from a single SH sample. */
-struct FSimpleDirectionalLightAndSHIndirectPolicy
+struct FForwardDirectionalLightAndSHIndirectPolicy
 {
 	static bool ShouldCache(EShaderPlatform Platform,const FMaterial* Material,const FVertexFactoryType* VertexFactoryType)
 	{
-		return FSimpleDynamicLightingPolicy::ShouldCache(Platform, Material, VertexFactoryType) && FCachedPointIndirectLightingPolicy::ShouldCache(Platform, Material, VertexFactoryType);
+		return Material->GetShadingModel() != MSM_Unlit && FCachedPointIndirectLightingPolicy::ShouldCache(Platform, Material, VertexFactoryType);
 	}
 
 	static void ModifyCompilationEnvironment(EShaderPlatform Platform, const FMaterial* Material, FShaderCompilerEnvironment& OutEnvironment)
 	{
-		FSimpleDynamicLightingPolicy::ModifyCompilationEnvironment(Platform, Material, OutEnvironment);
 		FCachedPointIndirectLightingPolicy::ModifyCompilationEnvironment(Platform, Material, OutEnvironment);
 	}
 };
 
 /** Combines a movable directional light with indirect lighting from a single SH sample. */
-struct FSimpleDirectionalLightAndSHDirectionalIndirectPolicy : public FSimpleDirectionalLightAndSHIndirectPolicy
+struct FForwardDirectionalLightAndSHDirectionalIndirectPolicy : public FForwardDirectionalLightAndSHIndirectPolicy
 {
 	static void ModifyCompilationEnvironment(EShaderPlatform Platform, const FMaterial* Material, FShaderCompilerEnvironment& OutEnvironment)
 	{
 		OutEnvironment.SetDefine(TEXT("MOVABLE_DIRECTIONAL_LIGHT"), TEXT("1"));
 		OutEnvironment.SetDefine(TEXT(PREPROCESSOR_TO_STRING(MAX_FORWARD_SHADOWCASCADES)), MAX_FORWARD_SHADOWCASCADES);
-		FSimpleDirectionalLightAndSHIndirectPolicy::ModifyCompilationEnvironment(Platform, Material, OutEnvironment);
+		FForwardDirectionalLightAndSHIndirectPolicy::ModifyCompilationEnvironment(Platform, Material, OutEnvironment);
 	}
 };
 
 /** Combines a movable directional light with CSM with indirect lighting from a single SH sample. */
-struct FSimpleDirectionalLightAndSHDirectionalCSMIndirectPolicy : public FSimpleDirectionalLightAndSHDirectionalIndirectPolicy
+struct FForwardDirectionalLightAndSHDirectionalCSMIndirectPolicy : public FForwardDirectionalLightAndSHDirectionalIndirectPolicy
 {
 	static void ModifyCompilationEnvironment(EShaderPlatform Platform, const FMaterial* Material, FShaderCompilerEnvironment& OutEnvironment)
 	{
 		OutEnvironment.SetDefine(TEXT("DIRECTIONAL_LIGHT_CSM"), TEXT("1"));
 		OutEnvironment.SetDefine(TEXT(PREPROCESSOR_TO_STRING(MAX_FORWARD_SHADOWCASCADES)), MAX_FORWARD_SHADOWCASCADES);
-		FSimpleDirectionalLightAndSHDirectionalIndirectPolicy::ModifyCompilationEnvironment(Platform, Material, OutEnvironment);
+		FForwardDirectionalLightAndSHDirectionalIndirectPolicy::ModifyCompilationEnvironment(Platform, Material, OutEnvironment);
 	}
 };
 
 /** Combines an directional light with CSM with indirect lighting from a single SH sample. */
-class FSimpleDirectionalLightAndSHCSMIndirectPolicy : public FSimpleDirectionalLightAndSHIndirectPolicy
+class FForwardDirectionalLightAndSHCSMIndirectPolicy : public FForwardDirectionalLightAndSHIndirectPolicy
 {
 public:
 	static void ModifyCompilationEnvironment(EShaderPlatform Platform, const FMaterial* Material, FShaderCompilerEnvironment& OutEnvironment)
 	{
 		OutEnvironment.SetDefine(TEXT("DIRECTIONAL_LIGHT_CSM"), TEXT("1"));
 		OutEnvironment.SetDefine(TEXT(PREPROCESSOR_TO_STRING(MAX_FORWARD_SHADOWCASCADES)), MAX_FORWARD_SHADOWCASCADES);
-		FSimpleDirectionalLightAndSHIndirectPolicy::ModifyCompilationEnvironment(Platform, Material, OutEnvironment);
+		FForwardDirectionalLightAndSHIndirectPolicy::ModifyCompilationEnvironment(Platform, Material, OutEnvironment);
 	}
 };
 
@@ -477,17 +560,21 @@ enum ELightMapPolicyType
 	LMP_NO_LIGHTMAP,
 	LMP_CACHED_VOLUME_INDIRECT_LIGHTING,
 	LMP_CACHED_POINT_INDIRECT_LIGHTING,
-	LMP_SIMPLE_DYNAMIC_LIGHTING,
+	LMP_SIMPLE_NO_LIGHTMAP,
+	LMP_SIMPLE_LIGHTMAP_ONLY_LIGHTING,
+	LMP_SIMPLE_DIRECTIONAL_LIGHT_LIGHTING,
+	LMP_SIMPLE_STATIONARY_PRECOMPUTED_SHADOW_LIGHTING,
+	LMP_SIMPLE_STATIONARY_SINGLESAMPLE_SHADOW_LIGHTING,
 	LMP_LQ_LIGHTMAP,
 	LMP_HQ_LIGHTMAP,
 	LMP_DISTANCE_FIELD_SHADOWS_AND_HQ_LIGHTMAP,
 	// Forward shading specific
 	LMP_DISTANCE_FIELD_SHADOWS_AND_LQ_LIGHTMAP,
 	LMP_DISTANCE_FIELD_SHADOWS_LIGHTMAP_AND_CSM,
-	LMP_SIMPLE_DIRECTIONAL_LIGHT_AND_SH_INDIRECT,
-	LMP_SIMPLE_DIRECTIONAL_LIGHT_AND_SH_DIRECTIONAL_INDIRECT,
-	LMP_SIMPLE_DIRECTIONAL_LIGHT_AND_SH_DIRECTIONAL_CSM_INDIRECT,
-	LMP_SIMPLE_DIRECTIONAL_LIGHT_AND_SH_CSM_INDIRECT,
+	LMP_FORWARD_DIRECTIONAL_LIGHT_AND_SH_INDIRECT,
+	LMP_FORWARD_DIRECTIONAL_LIGHT_AND_SH_DIRECTIONAL_INDIRECT,
+	LMP_FORWARD_DIRECTIONAL_LIGHT_AND_SH_DIRECTIONAL_CSM_INDIRECT,
+	LMP_FORWARD_DIRECTIONAL_LIGHT_AND_SH_CSM_INDIRECT,
 	LMP_MOVABLE_DIRECTIONAL_LIGHT,
 	LMP_MOVABLE_DIRECTIONAL_LIGHT_CSM,
 	LMP_MOVABLE_DIRECTIONAL_LIGHT_WITH_LIGHTMAP,
@@ -590,8 +677,16 @@ public:
 			return FCachedVolumeIndirectLightingPolicy::ShouldCache(Platform, Material, VertexFactoryType);
 		case LMP_CACHED_POINT_INDIRECT_LIGHTING:
 			return FCachedPointIndirectLightingPolicy::ShouldCache(Platform, Material, VertexFactoryType);
-		case LMP_SIMPLE_DYNAMIC_LIGHTING:
-			return FSimpleDynamicLightingPolicy::ShouldCache(Platform, Material, VertexFactoryType);
+		case LMP_SIMPLE_NO_LIGHTMAP:
+			return FSimpleNoLightmapLightingPolicy::ShouldCache(Platform, Material, VertexFactoryType);
+		case LMP_SIMPLE_LIGHTMAP_ONLY_LIGHTING:
+			return FSimpleLightmapOnlyLightingPolicy::ShouldCache(Platform, Material, VertexFactoryType);
+		case LMP_SIMPLE_DIRECTIONAL_LIGHT_LIGHTING:
+			return FSimpleDirectionalLightLightingPolicy::ShouldCache(Platform, Material, VertexFactoryType);
+		case LMP_SIMPLE_STATIONARY_PRECOMPUTED_SHADOW_LIGHTING:
+			return FSimpleStationaryLightPrecomputedShadowsLightingPolicy::ShouldCache(Platform, Material, VertexFactoryType);
+		case LMP_SIMPLE_STATIONARY_SINGLESAMPLE_SHADOW_LIGHTING:
+			return FSimpleStationaryLightSingleSampleShadowsLightingPolicy::ShouldCache(Platform, Material, VertexFactoryType);
 		case LMP_LQ_LIGHTMAP:
 			return TLightMapPolicy<LQ_LIGHTMAP>::ShouldCache(Platform, Material, VertexFactoryType);
 		case LMP_HQ_LIGHTMAP:
@@ -605,14 +700,14 @@ public:
 			return TDistanceFieldShadowsAndLightMapPolicy<LQ_LIGHTMAP>::ShouldCache(Platform, Material, VertexFactoryType);
 		case LMP_DISTANCE_FIELD_SHADOWS_LIGHTMAP_AND_CSM:
 			return FDistanceFieldShadowsLightMapAndCSMLightingPolicy::ShouldCache(Platform, Material, VertexFactoryType);
-		case LMP_SIMPLE_DIRECTIONAL_LIGHT_AND_SH_INDIRECT:
-			return FSimpleDirectionalLightAndSHIndirectPolicy::ShouldCache(Platform, Material, VertexFactoryType);
-		case LMP_SIMPLE_DIRECTIONAL_LIGHT_AND_SH_DIRECTIONAL_INDIRECT:
-			return FSimpleDirectionalLightAndSHDirectionalIndirectPolicy::ShouldCache(Platform, Material, VertexFactoryType);
-		case LMP_SIMPLE_DIRECTIONAL_LIGHT_AND_SH_CSM_INDIRECT:
-			return FSimpleDirectionalLightAndSHCSMIndirectPolicy::ShouldCache(Platform, Material, VertexFactoryType);
-		case LMP_SIMPLE_DIRECTIONAL_LIGHT_AND_SH_DIRECTIONAL_CSM_INDIRECT:
-			return FSimpleDirectionalLightAndSHDirectionalCSMIndirectPolicy::ShouldCache(Platform, Material, VertexFactoryType);
+		case LMP_FORWARD_DIRECTIONAL_LIGHT_AND_SH_INDIRECT:
+			return FForwardDirectionalLightAndSHIndirectPolicy::ShouldCache(Platform, Material, VertexFactoryType);
+		case LMP_FORWARD_DIRECTIONAL_LIGHT_AND_SH_DIRECTIONAL_INDIRECT:
+			return FForwardDirectionalLightAndSHDirectionalIndirectPolicy::ShouldCache(Platform, Material, VertexFactoryType);
+		case LMP_FORWARD_DIRECTIONAL_LIGHT_AND_SH_CSM_INDIRECT:
+			return FForwardDirectionalLightAndSHCSMIndirectPolicy::ShouldCache(Platform, Material, VertexFactoryType);
+		case LMP_FORWARD_DIRECTIONAL_LIGHT_AND_SH_DIRECTIONAL_CSM_INDIRECT:
+			return FForwardDirectionalLightAndSHDirectionalCSMIndirectPolicy::ShouldCache(Platform, Material, VertexFactoryType);
 		case LMP_MOVABLE_DIRECTIONAL_LIGHT:
 			return FMovableDirectionalLightLightingPolicy::ShouldCache(Platform, Material, VertexFactoryType);
 		case LMP_MOVABLE_DIRECTIONAL_LIGHT_CSM:
@@ -648,8 +743,20 @@ public:
 		case LMP_CACHED_POINT_INDIRECT_LIGHTING:
 			FCachedPointIndirectLightingPolicy::ModifyCompilationEnvironment(Platform, Material, OutEnvironment);
 			break;
-		case LMP_SIMPLE_DYNAMIC_LIGHTING:
-			FSimpleDynamicLightingPolicy::ModifyCompilationEnvironment(Platform, Material, OutEnvironment);
+		case LMP_SIMPLE_NO_LIGHTMAP:
+			return FSimpleNoLightmapLightingPolicy::ModifyCompilationEnvironment(Platform, Material, OutEnvironment);
+			break;
+		case LMP_SIMPLE_LIGHTMAP_ONLY_LIGHTING:
+			return FSimpleLightmapOnlyLightingPolicy::ModifyCompilationEnvironment(Platform, Material, OutEnvironment);
+			break;
+		case LMP_SIMPLE_DIRECTIONAL_LIGHT_LIGHTING:
+			FSimpleDirectionalLightLightingPolicy::ModifyCompilationEnvironment(Platform, Material, OutEnvironment);
+			break;
+		case LMP_SIMPLE_STATIONARY_PRECOMPUTED_SHADOW_LIGHTING:
+			return FSimpleStationaryLightPrecomputedShadowsLightingPolicy::ModifyCompilationEnvironment(Platform, Material, OutEnvironment);
+			break;
+		case LMP_SIMPLE_STATIONARY_SINGLESAMPLE_SHADOW_LIGHTING:
+			return FSimpleStationaryLightSingleSampleShadowsLightingPolicy::ModifyCompilationEnvironment(Platform, Material, OutEnvironment);
 			break;
 		case LMP_LQ_LIGHTMAP:
 			TLightMapPolicy<LQ_LIGHTMAP>::ModifyCompilationEnvironment(Platform, Material, OutEnvironment);
@@ -668,17 +775,17 @@ public:
 		case LMP_DISTANCE_FIELD_SHADOWS_LIGHTMAP_AND_CSM:
 			FDistanceFieldShadowsLightMapAndCSMLightingPolicy::ModifyCompilationEnvironment(Platform, Material, OutEnvironment);
 			break;
-		case LMP_SIMPLE_DIRECTIONAL_LIGHT_AND_SH_INDIRECT:
-			FSimpleDirectionalLightAndSHIndirectPolicy::ModifyCompilationEnvironment(Platform, Material, OutEnvironment);
+		case LMP_FORWARD_DIRECTIONAL_LIGHT_AND_SH_INDIRECT:
+			FForwardDirectionalLightAndSHIndirectPolicy::ModifyCompilationEnvironment(Platform, Material, OutEnvironment);
 			break;
-		case LMP_SIMPLE_DIRECTIONAL_LIGHT_AND_SH_DIRECTIONAL_INDIRECT:
-			FSimpleDirectionalLightAndSHDirectionalIndirectPolicy::ModifyCompilationEnvironment(Platform, Material, OutEnvironment);
+		case LMP_FORWARD_DIRECTIONAL_LIGHT_AND_SH_DIRECTIONAL_INDIRECT:
+			FForwardDirectionalLightAndSHDirectionalIndirectPolicy::ModifyCompilationEnvironment(Platform, Material, OutEnvironment);
 			break;
-		case LMP_SIMPLE_DIRECTIONAL_LIGHT_AND_SH_CSM_INDIRECT:
-			FSimpleDirectionalLightAndSHCSMIndirectPolicy::ModifyCompilationEnvironment(Platform, Material, OutEnvironment);
+		case LMP_FORWARD_DIRECTIONAL_LIGHT_AND_SH_CSM_INDIRECT:
+			FForwardDirectionalLightAndSHCSMIndirectPolicy::ModifyCompilationEnvironment(Platform, Material, OutEnvironment);
 			break;
-		case LMP_SIMPLE_DIRECTIONAL_LIGHT_AND_SH_DIRECTIONAL_CSM_INDIRECT:
-			FSimpleDirectionalLightAndSHDirectionalCSMIndirectPolicy::ModifyCompilationEnvironment(Platform, Material, OutEnvironment);
+		case LMP_FORWARD_DIRECTIONAL_LIGHT_AND_SH_DIRECTIONAL_CSM_INDIRECT:
+			FForwardDirectionalLightAndSHDirectionalCSMIndirectPolicy::ModifyCompilationEnvironment(Platform, Material, OutEnvironment);
 			break;
 		case LMP_MOVABLE_DIRECTIONAL_LIGHT:
 			FMovableDirectionalLightLightingPolicy::ModifyCompilationEnvironment(Platform, Material, OutEnvironment);
diff --git a/Engine/Source/Runtime/Renderer/Private/LightRendering.cpp b/Engine/Source/Runtime/Renderer/Private/LightRendering.cpp
--- a/Engine/Source/Runtime/Renderer/Private/LightRendering.cpp
+++ b/Engine/Source/Runtime/Renderer/Private/LightRendering.cpp
@@ -322,11 +322,6 @@ void FDeferredShadingSceneRenderer::RenderLights(FRHICommandListImmediate& RHICm
 {
 	SCOPED_DRAW_EVENT(RHICmdList, Lights);
 
-	if(IsSimpleDynamicLightingEnabled())
-	{
-		return;
-	}
-
 	bool bStencilBufferDirty = false;	// The stencil buffer should've been cleared to 0 already
 
 	SCOPE_CYCLE_COUNTER(STAT_LightingDrawTime);
diff --git a/Engine/Source/Runtime/Renderer/Private/MaterialTexCoordScalesRendering.cpp b/Engine/Source/Runtime/Renderer/Private/MaterialTexCoordScalesRendering.cpp
--- a/Engine/Source/Runtime/Renderer/Private/MaterialTexCoordScalesRendering.cpp
+++ b/Engine/Source/Runtime/Renderer/Private/MaterialTexCoordScalesRendering.cpp
@@ -66,8 +66,8 @@ void FMaterialTexCoordScalePS::SetMesh(
 	FMemory::Memzero(TexCoordIndices);
 
 #if !(UE_BUILD_SHIPPING || UE_BUILD_TEST)
-	float DistanceMultiplier = 1.f;
-	const FStreamingSectionBuildInfo* SectionData = Proxy ? Proxy->GetStreamingSectionData(DistanceMultiplier, VisualizeLODIndex, BatchElement.VisualizeElementIndex) : nullptr;
+	float ComponentExtraScale = 1.f, MeshExtraScale = 1.f;
+	const FStreamingSectionBuildInfo* SectionData = Proxy ? Proxy->GetStreamingSectionData(ComponentExtraScale, MeshExtraScale, VisualizeLODIndex, BatchElement.VisualizeElementIndex) : nullptr;
 	if (SectionData)
 	{
 		for (int32 ScaleIndex = 0; ScaleIndex < MAX_NUM_TEXTURE_REGISTER && ScaleIndex < SectionData->TexCoordData.Num(); ++ScaleIndex)
diff --git a/Engine/Source/Runtime/Renderer/Private/MeshTexCoordSizeAccuracyRendering.cpp b/Engine/Source/Runtime/Renderer/Private/MeshTexCoordSizeAccuracyRendering.cpp
--- a/Engine/Source/Runtime/Renderer/Private/MeshTexCoordSizeAccuracyRendering.cpp
+++ b/Engine/Source/Runtime/Renderer/Private/MeshTexCoordSizeAccuracyRendering.cpp
@@ -46,11 +46,11 @@ void FMeshTexCoordSizeAccuracyPS::SetMesh(
 {
 	float CPUTexelFactor = -1.f;
 #if !(UE_BUILD_SHIPPING || UE_BUILD_TEST)
-	float DistanceMultiplier = 1.f;
-	const FStreamingSectionBuildInfo* SectionData = Proxy ? Proxy->GetStreamingSectionData(DistanceMultiplier, VisualizeLODIndex, BatchElement.VisualizeElementIndex) : nullptr;
+	float ComponentExtraScale = 1.f, MeshExtraScale = 1.f;
+	const FStreamingSectionBuildInfo* SectionData = Proxy ? Proxy->GetStreamingSectionData(ComponentExtraScale, MeshExtraScale, VisualizeLODIndex, BatchElement.VisualizeElementIndex) : nullptr;
 	if (SectionData)
 	{
-		CPUTexelFactor = SectionData->TexelFactors[0];
+		CPUTexelFactor = MeshExtraScale * SectionData->TexelFactors[0];
 	}
 #endif
 	SetShaderValue(RHICmdList, FGlobalShader::GetPixelShader(), CPUTexelFactorParameter, CPUTexelFactor);
diff --git a/Engine/Source/Runtime/Renderer/Private/PostProcess/PostProcessTonemap.cpp b/Engine/Source/Runtime/Renderer/Private/PostProcess/PostProcessTonemap.cpp
--- a/Engine/Source/Runtime/Renderer/Private/PostProcess/PostProcessTonemap.cpp
+++ b/Engine/Source/Runtime/Renderer/Private/PostProcess/PostProcessTonemap.cpp
@@ -1220,7 +1220,7 @@ void FRCPassPostProcessTonemap::Process(FRenderingCompositePassContext& Context)
 	FIntRect DestRect = bDoScreenPercentageInTonemapper ? View.UnscaledViewRect : View.ViewRect;
 	FIntPoint SrcSize = InputDesc->Extent;
 	
-	SCOPED_DRAW_EVENTF(Context.RHICmdList, PostProcessTonemap, TEXT("Tonemapper#%d GammaOnly=%d ScreenPercentage=%d  %dx%d"),
+	SCOPED_DRAW_EVENTF(Context.RHICmdList, PostProcessTonemap, TEXT("Tonemapper#%d GammaOnly=%d HandleScreenPercentage=%d  %dx%d"),
 		ConfigIndexPC, bDoGammaOnly, bDoScreenPercentageInTonemapper, DestRect.Width(), DestRect.Height());
 
 	const FSceneRenderTargetItem& DestRenderTarget = PassOutputs[0].RequestSurface(Context);
diff --git a/Engine/Source/Runtime/Renderer/Private/PostProcess/SceneRenderTargets.cpp b/Engine/Source/Runtime/Renderer/Private/PostProcess/SceneRenderTargets.cpp
--- a/Engine/Source/Runtime/Renderer/Private/PostProcess/SceneRenderTargets.cpp
+++ b/Engine/Source/Runtime/Renderer/Private/PostProcess/SceneRenderTargets.cpp
@@ -522,7 +522,7 @@ void FSceneRenderTargets::BeginRenderingGBuffer(FRHICommandList& RHICmdList, ERe
 {
 	SCOPED_DRAW_EVENT(RHICmdList, BeginRenderingSceneColor);
 
-	if (IsSimpleDynamicLightingEnabled())
+	if (IsSimpleForwardShadingEnabled(GetFeatureLevelShaderPlatform(CurrentFeatureLevel)))
 	{
 		// in this non-standard case, just render to scene color with default mode
 		BeginRenderingSceneColor(RHICmdList);
@@ -612,7 +612,7 @@ void FSceneRenderTargets::BeginRenderingGBuffer(FRHICommandList& RHICmdList, ERe
 
 void FSceneRenderTargets::FinishRenderingGBuffer(FRHICommandListImmediate& RHICmdList)
 {
-	if (IsSimpleDynamicLightingEnabled())
+	if (IsSimpleForwardShadingEnabled(GetFeatureLevelShaderPlatform(CurrentFeatureLevel)))
 	{
 		// in this non-standard case, just render to scene color with default mode
 		FinishRenderingSceneColor(RHICmdList, true);
@@ -642,7 +642,7 @@ int32 FSceneRenderTargets::GetNumGBufferTargets() const
 {
 	int32 NumGBufferTargets = 1;
 
-	if (CurrentFeatureLevel >= ERHIFeatureLevel::SM4 && !IsSimpleDynamicLightingEnabled())
+	if (CurrentFeatureLevel >= ERHIFeatureLevel::SM4 && !IsSimpleForwardShadingEnabled(GetFeatureLevelShaderPlatform(CurrentFeatureLevel)))
 	{
 		// This needs to match TBasePassPixelShaderBaseType::ModifyCompilationEnvironment()
 		NumGBufferTargets = bAllowStaticLighting ? 6 : 5;
diff --git a/Engine/Source/Runtime/Renderer/Private/PrimitiveDistanceAccuracyRendering.cpp b/Engine/Source/Runtime/Renderer/Private/PrimitiveDistanceAccuracyRendering.cpp
--- a/Engine/Source/Runtime/Renderer/Private/PrimitiveDistanceAccuracyRendering.cpp
+++ b/Engine/Source/Runtime/Renderer/Private/PrimitiveDistanceAccuracyRendering.cpp
@@ -47,8 +47,8 @@ void FPrimitiveDistanceAccuracyPS::SetMesh(
 	float CPULogDistance = -1.f;
 #if !(UE_BUILD_SHIPPING || UE_BUILD_TEST)
 	const bool bUseNewMetrics = CVarStreamingUseNewMetrics.GetValueOnRenderThread() != 0;
-	float DistanceMultiplier = 1.f;
-	const FStreamingSectionBuildInfo* SectionData = Proxy ? Proxy->GetStreamingSectionData(DistanceMultiplier, VisualizeLODIndex, BatchElement.VisualizeElementIndex) : nullptr;
+	float ComponentExtraScale = 1.f, MeshExtraScale = 1.f;
+	const FStreamingSectionBuildInfo* SectionData = Proxy ? Proxy->GetStreamingSectionData(ComponentExtraScale, MeshExtraScale, VisualizeLODIndex, BatchElement.VisualizeElementIndex) : nullptr;
 	if (SectionData)
 	{
 		FVector ViewToObject = SectionData->BoxOrigin - View.ViewMatrices.ViewOrigin;
@@ -67,7 +67,7 @@ void FPrimitiveDistanceAccuracyPS::SetMesh(
 			DistSqMinusRadiusSq = FMath::Square(Distance) - SectionData->BoxExtent.SizeSquared();
 		}
 
-		const float OneOverDistanceMultiplier = 1.f / FMath::Max<float>(SMALL_NUMBER, DistanceMultiplier);
+		const float OneOverDistanceMultiplier = 1.f / FMath::Max<float>(SMALL_NUMBER, ComponentExtraScale);
 		CPULogDistance =  FMath::Log2(OneOverDistanceMultiplier * FMath::Sqrt(FMath::Max<float>(1.f, DistSqMinusRadiusSq)));
 	}
 #endif
diff --git a/Engine/Source/Runtime/Renderer/Private/ReflectionEnvironment.cpp b/Engine/Source/Runtime/Renderer/Private/ReflectionEnvironment.cpp
--- a/Engine/Source/Runtime/Renderer/Private/ReflectionEnvironment.cpp
+++ b/Engine/Source/Runtime/Renderer/Private/ReflectionEnvironment.cpp
@@ -1010,6 +1010,7 @@ void FDeferredShadingSceneRenderer::RenderStandardDeferredImageBasedReflections(
 
 	const bool bSkyLight = Scene->SkyLight
 		&& Scene->SkyLight->ProcessedTexture
+		&& !Scene->SkyLight->bHasStaticLighting
 		&& ViewFamily.EngineShowFlags.SkyLighting;
 
 	static TArray<FReflectionCaptureSortData> SortData;
@@ -1259,7 +1260,7 @@ void FDeferredShadingSceneRenderer::RenderStandardDeferredImageBasedReflections(
 
 void FDeferredShadingSceneRenderer::RenderDeferredReflections(FRHICommandListImmediate& RHICmdList, const TRefCountPtr<IPooledRenderTarget>& DynamicBentNormalAO)
 {
-	if (IsSimpleDynamicLightingEnabled() || ViewFamily.EngineShowFlags.VisualizeLightCulling)
+	if (ViewFamily.EngineShowFlags.VisualizeLightCulling)
 	{
 		return;
 	}
diff --git a/Engine/Source/Runtime/Renderer/Private/Renderer.cpp b/Engine/Source/Runtime/Renderer/Private/Renderer.cpp
--- a/Engine/Source/Runtime/Renderer/Private/Renderer.cpp
+++ b/Engine/Source/Runtime/Renderer/Private/Renderer.cpp
@@ -53,6 +53,8 @@ void FRendererModule::DrawTileMesh(FRHICommandListImmediate& RHICmdList, const F
 
 	if (!GUsingNullRHI)
 	{
+		GSystemTextures.InitializeTextures(RHICmdList, FeatureLevel);
+
 		// handle translucent material blend modes
 		if (IsTranslucentBlendMode(MaterialBlendMode))
 		{
diff --git a/Engine/Source/Runtime/Renderer/Private/RendererPrivate.h b/Engine/Source/Runtime/Renderer/Private/RendererPrivate.h
--- a/Engine/Source/Runtime/Renderer/Private/RendererPrivate.h
+++ b/Engine/Source/Runtime/Renderer/Private/RendererPrivate.h
@@ -20,6 +20,7 @@ DECLARE_LOG_CATEGORY_EXTERN(LogRenderer, Log, All);
 class FRendererModule : public IRendererModule
 {
 public:
+	FRendererModule();
 	virtual bool SupportsDynamicReloading() override { return true; }
 
 	virtual void BeginRenderingViewFamily(FCanvas* Canvas,FSceneViewFamily* ViewFamily) override;
diff --git a/Engine/Source/Runtime/Renderer/Private/RendererScene.cpp b/Engine/Source/Runtime/Renderer/Private/RendererScene.cpp
--- a/Engine/Source/Runtime/Renderer/Private/RendererScene.cpp
+++ b/Engine/Source/Runtime/Renderer/Private/RendererScene.cpp
@@ -511,8 +511,6 @@ FScene::FScene(UWorld* InWorld, bool bInRequiresHitProxies, bool bInIsEditorScen
 ,	bRequiresHitProxies(bInRequiresHitProxies)
 ,	bIsEditorScene(bInIsEditorScene)
 ,	NumUncachedStaticLightingInteractions(0)
-,	UpperDynamicSkylightColor(FLinearColor::Black)
-,	LowerDynamicSkylightColor(FLinearColor::Black)
 ,	SceneLODHierarchy(this)
 ,	DefaultMaxDistanceFieldOcclusionDistance(InWorld->GetWorldSettings()->DefaultMaxDistanceFieldOcclusionDistance)
 ,	GlobalDistanceFieldViewDistance(InWorld->GetWorldSettings()->GlobalDistanceFieldViewDistance)
@@ -1552,20 +1550,6 @@ void FScene::UpdateLightColorAndBrightness(ULightComponent* Light)
 	}
 }
 
-/** Updates the scene's dynamic skylight. */
-void FScene::UpdateDynamicSkyLight(const FLinearColor& UpperColor, const FLinearColor& LowerColor)
-{
-	ENQUEUE_UNIQUE_RENDER_COMMAND_THREEPARAMETER(
-		UpdateDynamicSkyLight,
-		FScene*,Scene,this,
-		FLinearColor,UpperColor,UpperColor,
-		FLinearColor,LowerColor,LowerColor,
-	{
-		Scene->UpperDynamicSkylightColor = UpperColor;
-		Scene->LowerDynamicSkylightColor = LowerColor;
-	});
-}
-
 void FScene::RemoveLightSceneInfo_RenderThread(FLightSceneInfo* LightSceneInfo)
 {
 	SCOPE_CYCLE_COUNTER(STAT_RemoveSceneLightTime);
diff --git a/Engine/Source/Runtime/Renderer/Private/ScenePrivate.h b/Engine/Source/Runtime/Renderer/Private/ScenePrivate.h
--- a/Engine/Source/Runtime/Renderer/Private/ScenePrivate.h
+++ b/Engine/Source/Runtime/Renderer/Private/ScenePrivate.h
@@ -1863,9 +1863,6 @@ public:
 	/** Set by the rendering thread to signal to the game thread that the scene needs a static lighting build. */
 	volatile mutable int32 NumUncachedStaticLightingInteractions;
 
-	FLinearColor UpperDynamicSkylightColor;
-	FLinearColor LowerDynamicSkylightColor;
-
 	FMotionBlurInfoData MotionBlurInfoData;
 
 	/** Uniform buffers for parameter collections with the corresponding Ids. */
@@ -1921,7 +1918,6 @@ public:
 	virtual void RemovePrecomputedLightVolume(const class FPrecomputedLightVolume* Volume) override;
 	virtual void UpdateLightTransform(ULightComponent* Light) override;
 	virtual void UpdateLightColorAndBrightness(ULightComponent* Light) override;
-	virtual void UpdateDynamicSkyLight(const FLinearColor& UpperColor, const FLinearColor& LowerColor) override;
 	virtual void AddExponentialHeightFog(UExponentialHeightFogComponent* FogComponent) override;
 	virtual void RemoveExponentialHeightFog(UExponentialHeightFogComponent* FogComponent) override;
 	virtual void AddAtmosphericFog(UAtmosphericFogComponent* FogComponent) override;
@@ -1943,12 +1939,6 @@ public:
 	virtual void SetClearMotionBlurInfoGameThread() override;
 	virtual void UpdateParameterCollections(const TArray<FMaterialParameterCollectionInstanceResource*>& InParameterCollections) override;
 
-	/** Determines whether the scene has dynamic sky lighting. */
-	bool HasDynamicSkyLighting() const
-	{
-		return (!UpperDynamicSkylightColor.Equals(FLinearColor::Black) || !LowerDynamicSkylightColor.Equals(FLinearColor::Black));
-	}
-
 	/** Determines whether the scene has atmospheric fog and sun light. */
 	bool HasAtmosphericFog() const
 	{
diff --git a/Engine/Source/Runtime/Renderer/Private/SceneRendering.cpp b/Engine/Source/Runtime/Renderer/Private/SceneRendering.cpp
--- a/Engine/Source/Runtime/Renderer/Private/SceneRendering.cpp
+++ b/Engine/Source/Runtime/Renderer/Private/SceneRendering.cpp
@@ -21,6 +21,7 @@
 #include "AtmosphereRendering.h"
 #include "Components/PlanarReflectionComponent.h"
 #include "Matinee/MatineeActor.h"
+#include "ComponentRecreateRenderStateContext.h"
 
 /*-----------------------------------------------------------------------------
 	Globals
@@ -142,6 +143,20 @@ static TAutoConsoleVariable<float> CVarTessellationAdaptivePixelsPerTriangle(
 
 extern ENGINE_API TAutoConsoleVariable<int32> CVarReflectionCaptureSize;
 
+static TAutoConsoleVariable<int32> CVarSupportSimpleForwardShading(
+	TEXT("r.SupportSimpleForwardShading"),
+	0,
+	TEXT("Whether to compile the shaders to support r.SimpleForwardShading being enabled (PC only)."),
+	ECVF_RenderThreadSafe | ECVF_ReadOnly);
+
+static TAutoConsoleVariable<int32> CVarSimpleForwardShading(
+	TEXT("r.SimpleForwardShading"),
+	0,
+	TEXT("Whether to use the simple forward shading base pass shaders which only support lightmaps + stationary directional light + stationary skylight\n")
+	TEXT("All other lighting features are disabled when true.  This is useful for supporting very low end hardware, and is only supported on PC platforms.\n")
+	TEXT("0:off, 1:on"),
+	ECVF_RenderThreadSafe | ECVF_Scalability);
+
 /*-----------------------------------------------------------------------------
 BitCounting
 -----------------------------------------------------------------------------*/
@@ -821,9 +836,6 @@ void FViewInfo::CreateUniformBuffer(
 			FrameUniformShaderParameters.DirectionalLightColor = FLinearColor::Black;
 			FrameUniformShaderParameters.DirectionalLightDirection = FVector::ZeroVector;
 		}
-		
-		FrameUniformShaderParameters.UpperSkyColor = Scene->UpperDynamicSkylightColor;
-		FrameUniformShaderParameters.LowerSkyColor = Scene->LowerDynamicSkylightColor;
 
 		// Atmospheric fog parameters
 		if (ShouldRenderAtmosphere(*Family) && Scene->AtmosphericFog)
@@ -867,8 +879,6 @@ void FViewInfo::CreateUniformBuffer(
 	{
 		FrameUniformShaderParameters.DirectionalLightDirection = FVector::ZeroVector;
 		FrameUniformShaderParameters.DirectionalLightColor = FLinearColor::Black;
-		FrameUniformShaderParameters.UpperSkyColor = FLinearColor::Black;
-		FrameUniformShaderParameters.LowerSkyColor = FLinearColor::Black;
 
 		// Atmospheric fog parameters
 		FrameUniformShaderParameters.AtmosphericFogSunPower = 0.f;
@@ -1439,6 +1449,9 @@ void FDisplayInternalsData::Setup(UWorld *World)
 		}
 
 		check(IsValid());
+		
+		extern ENGINE_API uint32 GStreamAllResourcesStillInFlight;
+		NumPendingStreamingRequests = GStreamAllResourcesStillInFlight;
 	}
 #endif
 }
@@ -2034,6 +2047,32 @@ static void RenderViewFamily_RenderThread(FRHICommandListImmediate& RHICmdList,
 #endif
 }
 
+void OnChangeSimpleForwardShading(IConsoleVariable* Var)
+{
+	static const auto SupportSimpleForwardShadingCVar = IConsoleManager::Get().FindTConsoleVariableDataInt(TEXT("r.SupportSimpleForwardShading"));
+	static const auto SimpleForwardShadingCVar = IConsoleManager::Get().FindTConsoleVariableDataInt(TEXT("r.SimpleForwardShading"));
+
+	if (SimpleForwardShadingCVar->GetValueOnAnyThread() != 0)
+	{
+		if (SupportSimpleForwardShadingCVar->GetValueOnAnyThread() == 0)
+		{
+			UE_LOG(LogRenderer, Warning, TEXT("r.SimpleForwardShading ignored as r.SupportSimpleForwardShading is not enabled"));
+		}
+		else if (!PlatformSupportsSimpleForwardShading(GMaxRHIShaderPlatform))
+		{
+			UE_LOG(LogRenderer, Warning, TEXT("r.SimpleForwardShading ignored, only supported on PC shader platforms.  Current shader platform %s"), *LegacyShaderPlatformToShaderFormat(GMaxRHIShaderPlatform).ToString());
+		}
+	}
+
+	// Propgate cvar change to static draw lists
+	FGlobalComponentRecreateRenderStateContext Context;
+}
+
+FRendererModule::FRendererModule()
+{
+	CVarSimpleForwardShading.AsVariable()->SetOnChangedCallback(FConsoleVariableDelegate::CreateStatic(&OnChangeSimpleForwardShading));
+}
+
 void FRendererModule::CreateAndInitSingleView(FRHICommandListImmediate& RHICmdList, class FSceneViewFamily* ViewFamily, const struct FSceneViewInitOptions* ViewInitOptions)
 {
 	// Create and add the new view
diff --git a/Engine/Source/Runtime/Renderer/Private/SceneVisibility.cpp b/Engine/Source/Runtime/Renderer/Private/SceneVisibility.cpp
--- a/Engine/Source/Runtime/Renderer/Private/SceneVisibility.cpp
+++ b/Engine/Source/Runtime/Renderer/Private/SceneVisibility.cpp
@@ -2664,7 +2664,7 @@ void FSceneRenderer::PostVisibilityFrameSetup(FILCUpdatePrimTaskData& OutILCTask
 					// Don't render if the light's origin is behind the view
 					if(ProjectedBlurOrigin.W >= 0.0f
 						// Don't render point lights that have completely faded out
-							&& (LightSceneInfo->Proxy->GetLightType() == LightType_Directional 
+						&& (LightSceneInfo->Proxy->GetLightType() == LightType_Directional 
 							|| DistanceToBlurOrigin < LightSceneInfo->Proxy->GetRadius() * PointLightRadiusFadeFactor))
 					{
 						View.bLightShaftUse = bNotMobileMSAA;
@@ -2757,7 +2757,7 @@ void FSceneRenderer::PostVisibilityFrameSetup(FILCUpdatePrimTaskData& OutILCTask
 	{
 
 		QUICK_SCOPE_CYCLE_COUNTER(STAT_PostVisibilityFrameSetup_InitFogConstants);
-	InitFogConstants();
+		InitFogConstants();
 	}
 }
 
@@ -2787,6 +2787,10 @@ bool FDeferredShadingSceneRenderer::InitViews(FRHICommandListImmediate& RHICmdLi
 	}
 	PreVisibilityFrameSetup(RHICmdList);
 	ComputeViewVisibility(RHICmdList);
+
+	// This has to happen before Scene->IndirectLightingCache.UpdateCache, since primitives in View.IndirectShadowPrimitives need ILC updates
+	CreateIndirectCapsuleShadows();
+
 	PostVisibilityFrameSetup(ILCTaskData);
 
 	FVector AverageViewPosition(0);
@@ -2849,7 +2853,7 @@ void FDeferredShadingSceneRenderer::InitViewsPossiblyAfterPrepass(FRHICommandLis
 
 	bool bDynamicShadows = ViewFamily.EngineShowFlags.DynamicShadows && GetShadowQuality() > 0;
 
-	if (bDynamicShadows && !IsSimpleDynamicLightingEnabled())
+	if (bDynamicShadows && !IsSimpleForwardShadingEnabled(GetFeatureLevelShaderPlatform(FeatureLevel)))
 	{
 		// Setup dynamic shadows.
 		InitDynamicShadows(RHICmdList);
@@ -2863,8 +2867,8 @@ void FDeferredShadingSceneRenderer::InitViewsPossiblyAfterPrepass(FRHICommandLis
 
 	{
 		QUICK_SCOPE_CYCLE_COUNTER(STAT_InitViews_UpdatePrimitivePrecomputedLightingBuffers);
-	// Now that the indirect lighting cache is updated, we can update the primitive precomputed lighting buffers.
-	UpdatePrimitivePrecomputedLightingBuffers();
+		// Now that the indirect lighting cache is updated, we can update the primitive precomputed lighting buffers.
+		UpdatePrimitivePrecomputedLightingBuffers();
 	}
 
 	UpdateTranslucencyTimersAndSeparateTranslucencyBufferSize(RHICmdList);
diff --git a/Engine/Source/Runtime/Renderer/Private/ShaderBaseClasses.cpp b/Engine/Source/Runtime/Renderer/Private/ShaderBaseClasses.cpp
--- a/Engine/Source/Runtime/Renderer/Private/ShaderBaseClasses.cpp
+++ b/Engine/Source/Runtime/Renderer/Private/ShaderBaseClasses.cpp
@@ -6,6 +6,7 @@
 
 #include "RendererPrivate.h"
 #include "ScenePrivate.h"
+#include "ParameterCollection.h"
 
 /** If true, cached uniform expressions are allowed. */
 int32 FMaterialShader::bAllowCachedUniformExpressions = true;
@@ -76,7 +77,14 @@ FMaterialShader::FMaterialShader(const FMaterialShaderType::CompiledShaderInitia
 FUniformBufferRHIParamRef FMaterialShader::GetParameterCollectionBuffer(const FGuid& Id, const FSceneInterface* SceneInterface) const
 {
 	const FScene* Scene = (const FScene*)SceneInterface;
-	return Scene ? Scene->GetParameterCollectionBuffer(Id) : FUniformBufferRHIParamRef();
+	FUniformBufferRHIParamRef UniformBuffer = Scene ? Scene->GetParameterCollectionBuffer(Id) : FUniformBufferRHIParamRef();
+
+	if (!UniformBuffer)
+	{
+		UniformBuffer = GDefaultMaterialParameterCollectionInstances.FindChecked(Id)->GetUniformBuffer();
+	}
+
+	return UniformBuffer;
 }
 
 template<typename ShaderRHIParamRef>
diff --git a/Engine/Source/Runtime/Renderer/Private/ShadowRendering.cpp b/Engine/Source/Runtime/Renderer/Private/ShadowRendering.cpp
--- a/Engine/Source/Runtime/Renderer/Private/ShadowRendering.cpp
+++ b/Engine/Source/Runtime/Renderer/Private/ShadowRendering.cpp
@@ -4042,7 +4042,7 @@ void FForwardShadingSceneRenderer::RenderShadowDepthMaps(FRHICommandListImmediat
 	SCOPED_DRAW_EVENT(RHICmdList, Lights);
 
 	bool bDynamicShadows = ViewFamily.EngineShowFlags.DynamicShadows && GetShadowQuality() > 0;
-	if (IsSimpleDynamicLightingEnabled() || !bDynamicShadows)
+	if (IsSimpleForwardShadingEnabled(GetFeatureLevelShaderPlatform(FeatureLevel)) || !bDynamicShadows)
 	{
 		return;
 	}	
@@ -4143,7 +4143,7 @@ void FForwardShadingSceneRenderer::RenderModulatedShadowDepthMaps(FRHICommandLis
 	
 	bool bDynamicShadows = ViewFamily.EngineShowFlags.DynamicShadows && GetShadowQuality() > 0;
 	
-	if (IsSimpleDynamicLightingEnabled() || !bDynamicShadows)
+	if (IsSimpleForwardShadingEnabled(GetFeatureLevelShaderPlatform(FeatureLevel)) || !bDynamicShadows)
 	{
 		bModulatedShadowsInUse = false;
 	}
@@ -4247,7 +4247,7 @@ void FForwardShadingSceneRenderer::RenderAllocatedModulatedShadowProjections(FRH
 
 	bool bDynamicShadows = ViewFamily.EngineShowFlags.DynamicShadows && GetShadowQuality() > 0;
 
-	if (IsSimpleDynamicLightingEnabled() || !bDynamicShadows)
+	if (IsSimpleForwardShadingEnabled(GetFeatureLevelShaderPlatform(FeatureLevel)) || !bDynamicShadows)
 	{
 		return;
 	}
diff --git a/Engine/Source/Runtime/Renderer/Private/ShadowSetup.cpp b/Engine/Source/Runtime/Renderer/Private/ShadowSetup.cpp
--- a/Engine/Source/Runtime/Renderer/Private/ShadowSetup.cpp
+++ b/Engine/Source/Runtime/Renderer/Private/ShadowSetup.cpp
@@ -2650,6 +2650,4 @@ void FDeferredShadingSceneRenderer::InitDynamicShadows(FRHICommandListImmediate&
 
 	// Generate mesh element arrays from shadow primitive arrays
 	GatherShadowDynamicMeshElements();
-
-	CreateIndirectCapsuleShadows();
 }
diff --git a/Engine/Source/Runtime/ShaderCore/Private/Shader.cpp b/Engine/Source/Runtime/ShaderCore/Private/Shader.cpp
--- a/Engine/Source/Runtime/ShaderCore/Private/Shader.cpp
+++ b/Engine/Source/Runtime/ShaderCore/Private/Shader.cpp
@@ -1617,14 +1617,6 @@ void ShaderMapAppendKeyString(EShaderPlatform Platform, FString& KeyString)
 		}
 	}
 
-	{
-		static const auto CVar = IConsoleManager::Get().FindTConsoleVariableDataInt(TEXT("r.GBuffer"));
-		if (CVar ? CVar->GetValueOnAnyThread() == 0 : false)
-		{
-			KeyString += TEXT("_NoGB");
-		}
-	}
-
 	{
 		static IConsoleVariable* CVar = IConsoleManager::Get().FindConsoleVariable(TEXT("r.DBuffer"));
 		KeyString += (CVar && CVar->GetInt() != 0) ? TEXT("_DBuf") : TEXT("_NoDBuf");
