Copying //UE4/Dev-Rendering to //UE4/Dev-Main (Source: //UE4/Dev-Rendering @ 2998063)

diff --git a/Engine/Source/Runtime/Renderer/Private/PostProcess/PostProcessMotionBlur.cpp b/Engine/Source/Runtime/Renderer/Private/PostProcess/PostProcessMotionBlur.cpp
--- a/Engine/Source/Runtime/Renderer/Private/PostProcess/PostProcessMotionBlur.cpp
+++ b/Engine/Source/Runtime/Renderer/Private/PostProcess/PostProcessMotionBlur.cpp
@@ -32,726 +32,6 @@ static TAutoConsoleVariable<float> CVarMotionBlur2ndScale(
 	ECVF_Cheat | ECVF_RenderThreadSafe);
 
 
-/** Encapsulates the post processing motion blur vertex shader. */
-class FPostProcessMotionBlurSetupVS : public FGlobalShader
-{
-	DECLARE_SHADER_TYPE(FPostProcessMotionBlurSetupVS,Global);
-
-	static bool ShouldCache(EShaderPlatform Platform)
-	{
-		return IsFeatureLevelSupported(Platform, ERHIFeatureLevel::SM4);
-	}
-
-	/** Default constructor. */
-	FPostProcessMotionBlurSetupVS() {}
-
-	// FShader interface.
-	virtual bool Serialize(FArchive& Ar) override
-	{
-		bool bShaderHasOutdatedParameters = FGlobalShader::Serialize(Ar);
-		Ar << PostprocessParameter;
-		return bShaderHasOutdatedParameters;
-	}
-
-	/** to have a similar interface as all other shaders */
-	void SetParameters(const FRenderingCompositePassContext& Context)
-	{
-		const auto ShaderRHI = GetVertexShader();
-		FGlobalShader::SetParameters(Context.RHICmdList, ShaderRHI, Context.View);
-		PostprocessParameter.SetVS(ShaderRHI, Context, TStaticSamplerState<SF_Point, AM_Clamp, AM_Clamp, AM_Clamp>::GetRHI());
-	}
-
-public:
-	FPostProcessPassParameters PostprocessParameter;
-
-	/** Initialization constructor. */
-	FPostProcessMotionBlurSetupVS(const ShaderMetaType::CompiledShaderInitializerType& Initializer)
-		: FGlobalShader(Initializer)
-	{
-		PostprocessParameter.Bind(Initializer.ParameterMap);
-	}
-};
-
-IMPLEMENT_SHADER_TYPE(,FPostProcessMotionBlurSetupVS,TEXT("PostProcessMotionBlur"),TEXT("SetupVS"),SF_Vertex);
-
-
-/** Encapsulates the post processing motion blur pixel shader. */
-class FPostProcessMotionBlurSetupPS : public FGlobalShader
-{
-	DECLARE_SHADER_TYPE(FPostProcessMotionBlurSetupPS, Global);
-
-	static bool ShouldCache(EShaderPlatform Platform)
-	{
-		return IsFeatureLevelSupported(Platform, ERHIFeatureLevel::SM4);
-	}
-
-	/** Default constructor. */
-	FPostProcessMotionBlurSetupPS() {}
-
-public:
-	FPostProcessPassParameters PostprocessParameter;
-	FShaderParameter VelocityScale;
-
-	/** Initialization constructor. */
-	FPostProcessMotionBlurSetupPS(const ShaderMetaType::CompiledShaderInitializerType& Initializer)
-		: FGlobalShader(Initializer)
-	{
-		PostprocessParameter.Bind(Initializer.ParameterMap);
-		VelocityScale.Bind( Initializer.ParameterMap, TEXT("VelocityScale") );
-	}
-
-	// FShader interface.
-	virtual bool Serialize(FArchive& Ar) override
-	{
-		bool bShaderHasOutdatedParameters = FGlobalShader::Serialize(Ar);
-		Ar << PostprocessParameter << VelocityScale;
-		return bShaderHasOutdatedParameters;
-	}
-
-	void SetParameters(const FRenderingCompositePassContext& Context)
-	{
-		const FPixelShaderRHIParamRef ShaderRHI = GetPixelShader();
-
-		FGlobalShader::SetParameters(Context.RHICmdList, ShaderRHI, Context.View);
-
-		PostprocessParameter.SetPS(ShaderRHI, Context, TStaticSamplerState<SF_Point, AM_Clamp, AM_Clamp, AM_Clamp>::GetRHI());
-
-		{
-			const float SizeX = Context.View.ViewRect.Width();
-			const float SizeY = Context.View.ViewRect.Height();
-			const float InvAspectRatio = SizeY / SizeX;
-
-			const FSceneViewState* ViewState = (FSceneViewState*) Context.View.State;
-			const float MotionBlurTimeScale = ViewState ? ViewState->MotionBlurTimeScale : 1.0f;
-
-			const float ViewMotionBlurScale = 0.5f * MotionBlurTimeScale * Context.View.FinalPostProcessSettings.MotionBlurAmount;
-
-			// 0:no 1:full screen width
-			float MaxVelocity = Context.View.FinalPostProcessSettings.MotionBlurMax / 100.0f;
-			float InvMaxVelocity = 1.0f / MaxVelocity;
-			float ObjectScaleX = ViewMotionBlurScale * InvMaxVelocity;
-			float ObjectScaleY = ViewMotionBlurScale * InvMaxVelocity * InvAspectRatio;
-
-			SetShaderValue(Context.RHICmdList, ShaderRHI, VelocityScale, FVector4(ObjectScaleX, -ObjectScaleY, 0, 0));
-		}
-	}
-};
-
-IMPLEMENT_SHADER_TYPE(,FPostProcessMotionBlurSetupPS, TEXT("PostProcessMotionBlur"), TEXT("SetupPS"), SF_Pixel);
-
-
-void FRCPassPostProcessMotionBlurSetup::Process(FRenderingCompositePassContext& Context)
-{
-	const FPooledRenderTargetDesc* InputDesc = GetInputDesc(ePId_Input0);
-
-	if(!InputDesc)
-	{
-		// input is not hooked up correctly
-		return;
-	}
-
-	const FSceneView& View = Context.View;
-	const FSceneViewFamily& ViewFamily = *(View.Family);
-
-	FIntPoint SrcSize = InputDesc->Extent;
-	FIntPoint DestSize = PassOutputs[0].RenderTargetDesc.Extent;
-
-	// e.g. 4 means the input texture is 4x smaller than the buffer size
-	uint32 ScaleFactor = FSceneRenderTargets::Get(Context.RHICmdList).GetBufferSizeXY().X / SrcSize.X;
-
-	FIntRect SrcRect = View.ViewRect / ScaleFactor;
-
-	// Viewport size not even also causes issue
-	FIntRect DestRect = FIntRect::DivideAndRoundUp(SrcRect, 2);
-
-	SCOPED_DRAW_EVENTF(Context.RHICmdList, MotionBlurSetup, TEXT("MotionBlurSetup %dx%d"), DestRect.Width(), DestRect.Height());
-	
-	const FSceneRenderTargetItem& DestRenderTarget0 = PassOutputs[0].RequestSurface(Context);
-	const FSceneRenderTargetItem& DestRenderTarget1 = PassOutputs[1].RequestSurface(Context);
-
-	// Set the view family's render target/viewport.
-	FTextureRHIParamRef RenderTargets[] =
-	{
-		DestRenderTarget0.TargetableTexture,
-		DestRenderTarget1.TargetableTexture
-	};
-	SetRenderTargets(Context.RHICmdList, ARRAY_COUNT(RenderTargets), RenderTargets, FTextureRHIParamRef(), 0, NULL);
-
-	// is optimized away if possible (RT size=view size, )
-	FLinearColor ClearColors[2] = {FLinearColor(0,0,0,0), FLinearColor(0,0,0,0)};
-	Context.RHICmdList.ClearMRT(true, 2, ClearColors, false, 1.0f, false, 0, DestRect);
-
-	Context.SetViewportAndCallRHI(0, 0, 0.0f, DestSize.X, DestSize.Y, 1.0f );
-
-	// set the state
-	Context.RHICmdList.SetBlendState(TStaticBlendState<>::GetRHI());
-	Context.RHICmdList.SetRasterizerState(TStaticRasterizerState<>::GetRHI());
-	Context.RHICmdList.SetDepthStencilState(TStaticDepthStencilState<false, CF_Always>::GetRHI());
-
-	TShaderMapRef<FPostProcessMotionBlurSetupVS> VertexShader(Context.GetShaderMap());
-
-	{
-		TShaderMapRef<FPostProcessMotionBlurSetupPS > PixelShader(Context.GetShaderMap());
-		static FGlobalBoundShaderState BoundShaderState;
-		
-
-		SetGlobalBoundShaderState(Context.RHICmdList, Context.GetFeatureLevel(), BoundShaderState, GFilterVertexDeclaration.VertexDeclarationRHI, *VertexShader, *PixelShader);
-
-		PixelShader->SetParameters(Context);
-		VertexShader->SetParameters(Context);
-	}
-
-	DrawPostProcessPass(
-		Context.RHICmdList,
-		DestRect.Min.X, DestRect.Min.Y,
-		DestRect.Width(), DestRect.Height(),
-		SrcRect.Min.X, SrcRect.Min.Y,
-		SrcRect.Width(), SrcRect.Height(),
-		DestSize,
-		SrcSize,
-		*VertexShader,
-		View.StereoPass,
-		Context.HasHmdMesh(),
-		EDRF_UseTriangleOptimization);
-
-	Context.RHICmdList.CopyToResolveTarget(DestRenderTarget0.TargetableTexture, DestRenderTarget0.ShaderResourceTexture, false, FResolveParams());
-	Context.RHICmdList.CopyToResolveTarget(DestRenderTarget1.TargetableTexture, DestRenderTarget1.ShaderResourceTexture, false, FResolveParams());
-}
-
-FPooledRenderTargetDesc FRCPassPostProcessMotionBlurSetup::ComputeOutputDesc(EPassOutputId InPassOutputId) const
-{
-	if(InPassOutputId == ePId_Output0)
-	{
-		// downsampled velocity
-
-		FPooledRenderTargetDesc Ret = GetInput(ePId_Input0)->GetOutput()->RenderTargetDesc;
-
-		Ret.Reset();
-		Ret.Extent /= 2;
-		Ret.Extent.X = FMath::Max(1, Ret.Extent.X);
-		Ret.Extent.Y = FMath::Max(1, Ret.Extent.Y);
-
-		// we need at least a format in the range 0..1 with RGB channels, A is unused
-		//	Ret.Format = PF_A2B10G10R10;
-		// we need alpha to renormalize
-		Ret.Format = PF_FloatRGBA;
-
-		Ret.TargetableFlags &= ~TexCreate_UAV;
-		Ret.TargetableFlags |= TexCreate_RenderTargetable;
-		Ret.DebugName = TEXT("MotionBlurSetup0");
-
-		return Ret;
-	}
-	else 
-	{
-		check(InPassOutputId == ePId_Output1);
-
-		// scene color with depth in alpha
-		FPooledRenderTargetDesc Ret = GetInput(ePId_Input1)->GetOutput()->RenderTargetDesc;
-
-		Ret.Reset();
-		Ret.Extent /= 2;
-		Ret.Extent.X = FMath::Max(1, Ret.Extent.X);
-		Ret.Extent.Y = FMath::Max(1, Ret.Extent.Y);
-
-		Ret.Format = PF_FloatRGBA;
-
-		Ret.TargetableFlags &= ~TexCreate_UAV;
-		Ret.TargetableFlags |= TexCreate_RenderTargetable;
-		Ret.DebugName = TEXT("MotionBlurSetup1");
-
-		return Ret;
-	}
-}
-
-
-/**
- * @param Quality 0: visualize, 1:low, 2:medium, 3:high, 4:very high
- * Encapsulates a MotionBlur pixel shader.
- */
-template <uint32 Quality>
-class FPostProcessMotionBlurPS : public FGlobalShader
-{
-	DECLARE_SHADER_TYPE(FPostProcessMotionBlurPS, Global);
-
-	static bool ShouldCache(EShaderPlatform Platform)
-	{
-		return IsFeatureLevelSupported(Platform, ERHIFeatureLevel::SM4);
-	}
-
-	static void ModifyCompilationEnvironment(EShaderPlatform Platform, FShaderCompilerEnvironment& OutEnvironment)
-	{
-		FGlobalShader::ModifyCompilationEnvironment(Platform,OutEnvironment);
-		OutEnvironment.SetDefine(TEXT("MOTION_BLUR_QUALITY"), Quality);
-	}
-
-	/** Default constructor. */
-	FPostProcessMotionBlurPS() {}
-
-public:
-	FPostProcessPassParameters PostprocessParameter;
-	FDeferredPixelShaderParameters DeferredParameters;
-	FShaderParameter PrevViewProjMatrix;
-	FShaderParameter TextureViewMad;
-	FShaderParameter MotionBlurParameters;
-
-	/** Initialization constructor. */
-	FPostProcessMotionBlurPS(const ShaderMetaType::CompiledShaderInitializerType& Initializer)
-		: FGlobalShader(Initializer)
-	{
-		PostprocessParameter.Bind(Initializer.ParameterMap);
-		DeferredParameters.Bind(Initializer.ParameterMap);
-		PrevViewProjMatrix.Bind(Initializer.ParameterMap, TEXT("PrevViewProjMatrix"));
-		TextureViewMad.Bind(Initializer.ParameterMap, TEXT("TextureViewMad"));
-		MotionBlurParameters.Bind(Initializer.ParameterMap, TEXT("MotionBlurParameters"));
-	}
-
-	// FShader interface.
-	virtual bool Serialize(FArchive& Ar) override
-	{
-		bool bShaderHasOutdatedParameters = FGlobalShader::Serialize(Ar);
-		Ar << PostprocessParameter << DeferredParameters << PrevViewProjMatrix << TextureViewMad << MotionBlurParameters;
-		return bShaderHasOutdatedParameters;
-	}
-
-	void SetParameters(const FRenderingCompositePassContext& Context)
-	{
-		const FPixelShaderRHIParamRef ShaderRHI = GetPixelShader();
-
-		FGlobalShader::SetParameters(Context.RHICmdList, ShaderRHI, Context.View);
-
-		DeferredParameters.Set(Context.RHICmdList, ShaderRHI, Context.View);
-
-		{
-			bool bFiltered = false;
-
-#if !(UE_BUILD_SHIPPING || UE_BUILD_TEST)
-			bFiltered = CVarMotionBlurFiltering.GetValueOnRenderThread() != 0;
-#endif // !(UE_BUILD_SHIPPING || UE_BUILD_TEST)
-
-			if(bFiltered)
-			{
-				PostprocessParameter.SetPS(ShaderRHI, Context, TStaticSamplerState<SF_Bilinear,AM_Border,AM_Border,AM_Clamp>::GetRHI());
-			}
-			else
-			{
-				PostprocessParameter.SetPS(ShaderRHI, Context, TStaticSamplerState<SF_Point,AM_Border,AM_Border,AM_Clamp>::GetRHI());
-			}
-		}
-		
-		if( Context.View.Family->EngineShowFlags.CameraInterpolation )
-		{
-			// Instead of finding the world space position of the current pixel, calculate the world space position offset by the camera position, 
-			// then translate by the difference between last frame's camera position and this frame's camera position,
-			// then apply the rest of the transforms.  This effectively avoids precision issues near the extents of large levels whose world space position is very large.
-			FVector ViewOriginDelta = Context.View.ViewMatrices.ViewOrigin - Context.View.PrevViewMatrices.ViewOrigin;
-			SetShaderValue(Context.RHICmdList, ShaderRHI, PrevViewProjMatrix, FTranslationMatrix(ViewOriginDelta) * Context.View.PrevViewRotationProjMatrix);
-		}
-		else
-		{
-			SetShaderValue(Context.RHICmdList, ShaderRHI, PrevViewProjMatrix, Context.View.ViewMatrices.GetViewRotationProjMatrix());
-		}
-
-		// to mask out samples from outside of the view
-		{
-			FIntPoint BufferSize = FSceneRenderTargets::Get(Context.RHICmdList).GetBufferSizeXY();
-			FVector2D InvBufferSize(1.0f / BufferSize.X, 1.0f / BufferSize.Y);
-
-			FIntRect ClipRect = Context.View.ViewRect;
-
-			// to avoid leaking in content from the outside because of bilinear filtering, shrink
-			ClipRect.InflateRect(-1);
-
-			FVector2D MinUV(ClipRect.Min.X * InvBufferSize.X, ClipRect.Min.Y * InvBufferSize.Y);
-			FVector2D MaxUV(ClipRect.Max.X * InvBufferSize.X, ClipRect.Max.Y * InvBufferSize.Y);
-			FVector2D SizeUV = MaxUV - MinUV;
-
-			FVector2D Mul(1.0f / SizeUV.X, 1.0f / SizeUV.Y);
-			FVector2D Add = - MinUV * Mul;
-			FVector4 TextureViewMadValue(Mul.X, Mul.Y, Add.X, Add.Y);
-			SetShaderValue(Context.RHICmdList, ShaderRHI, TextureViewMad, TextureViewMadValue);
-		}
-
-		{
-			const float SizeX = Context.View.ViewRect.Width();
-			const float SizeY = Context.View.ViewRect.Height();
-			const float AspectRatio = SizeX / SizeY;
-			const float InvAspectRatio = SizeY / SizeX;
-
-			const FSceneViewState* ViewState = (FSceneViewState*) Context.View.State;
-			float MotionBlurTimeScale = ViewState ? ViewState->MotionBlurTimeScale : 1.0f;
-
-			float ViewMotionBlurScale = 0.5f * MotionBlurTimeScale * Context.View.FinalPostProcessSettings.MotionBlurAmount;
-
-			// MotionBlurInstanceScale was needed to hack some cases where motion blur wasn't working well, this shouldn't be needed any more, can clean this up later
-			float MotionBlurInstanceScale = 1;
-
-			float ObjectMotionBlurScale	= MotionBlurInstanceScale * ViewMotionBlurScale;
-			// 0:no 1:full screen width, percent conversion
-			float MaxVelocity = Context.View.FinalPostProcessSettings.MotionBlurMax / 100.0f;
-			float InvMaxVelocity = 1.0f / MaxVelocity;
-
-			// *2 to convert to -1..1 -1..1 screen space
-			// / MaxFraction to map screenpos to -1..1 normalized MaxFraction
-			FVector4 MotionBlurParametersValue(
-				ObjectMotionBlurScale * InvMaxVelocity,
-				- ObjectMotionBlurScale * InvMaxVelocity * InvAspectRatio,
-				MaxVelocity * 2,
-				- MaxVelocity * 2 * AspectRatio);
-			SetShaderValue(Context.RHICmdList, ShaderRHI, MotionBlurParameters, MotionBlurParametersValue);
-		}
-	}
-
-	static const TCHAR* GetSourceFilename()
-	{
-		return TEXT("PostProcessMotionBlur");
-	}
-
-	static const TCHAR* GetFunctionName()
-	{
-		return TEXT("MainPS");
-	}
-};
-
-// #define avoids a lot of code duplication
-#define VARIATION1(A) typedef FPostProcessMotionBlurPS<A> FPostProcessMotionBlurPS##A; \
-	IMPLEMENT_SHADER_TYPE2(FPostProcessMotionBlurPS##A, SF_Pixel);
-
-VARIATION1(0)			VARIATION1(1)			VARIATION1(2)			VARIATION1(3)			VARIATION1(4)
-#undef VARIATION1
-
-
-class FPostProcessVisualizeMotionBlurPS : public FGlobalShader
-{
-	DECLARE_SHADER_TYPE(FPostProcessVisualizeMotionBlurPS, Global);
-
-	static bool ShouldCache(EShaderPlatform Platform)
-	{
-		return IsFeatureLevelSupported(Platform, ERHIFeatureLevel::SM4);
-	}
-
-	/** Default constructor. */
-	FPostProcessVisualizeMotionBlurPS() {}
-
-public:
-	FPostProcessPassParameters PostprocessParameter;
-	FDeferredPixelShaderParameters DeferredParameters;
-	FShaderParameter PrevViewProjMatrix;
-
-	/** Initialization constructor. */
-	FPostProcessVisualizeMotionBlurPS(const ShaderMetaType::CompiledShaderInitializerType& Initializer)
-		: FGlobalShader(Initializer)
-	{
-		PostprocessParameter.Bind(Initializer.ParameterMap);
-		DeferredParameters.Bind(Initializer.ParameterMap);
-		PrevViewProjMatrix.Bind(Initializer.ParameterMap, TEXT("PrevViewProjMatrix"));
-	}
-
-	// FShader interface.
-	virtual bool Serialize(FArchive& Ar) override
-	{
-		bool bShaderHasOutdatedParameters = FGlobalShader::Serialize(Ar);
-		Ar << PostprocessParameter << DeferredParameters << PrevViewProjMatrix;
-		return bShaderHasOutdatedParameters;
-	}
-
-	void SetParameters(const FRenderingCompositePassContext& Context)
-	{
-		const FPixelShaderRHIParamRef ShaderRHI = GetPixelShader();
-
-		FGlobalShader::SetParameters(Context.RHICmdList, ShaderRHI, Context.View);
-
-		DeferredParameters.Set(Context.RHICmdList, ShaderRHI, Context.View);
-
-		SetUniformBufferParameter(Context.RHICmdList, ShaderRHI, GetUniformBufferParameter<FCameraMotionParameters>(), CreateCameraMotionParametersUniformBuffer(Context.View));
-
-		{
-			bool bFiltered = false;
-
-#if !(UE_BUILD_SHIPPING || UE_BUILD_TEST)
-			bFiltered = CVarMotionBlurFiltering.GetValueOnRenderThread() != 0;
-#endif // !(UE_BUILD_SHIPPING || UE_BUILD_TEST)
-
-			if(bFiltered)
-			{
-				PostprocessParameter.SetPS(ShaderRHI, Context, TStaticSamplerState<SF_Bilinear,AM_Border,AM_Border,AM_Clamp>::GetRHI());
-			}
-			else
-			{
-				PostprocessParameter.SetPS(ShaderRHI, Context, TStaticSamplerState<SF_Point,AM_Border,AM_Border,AM_Clamp>::GetRHI());
-			}
-		}
-		
-		if( Context.View.Family->EngineShowFlags.CameraInterpolation )
-		{
-			// Instead of finding the world space position of the current pixel, calculate the world space position offset by the camera position, 
-			// then translate by the difference between last frame's camera position and this frame's camera position,
-			// then apply the rest of the transforms.  This effectively avoids precision issues near the extents of large levels whose world space position is very large.
-			FVector ViewOriginDelta = Context.View.ViewMatrices.ViewOrigin - Context.View.PrevViewMatrices.ViewOrigin;
-			SetShaderValue(Context.RHICmdList, ShaderRHI, PrevViewProjMatrix, FTranslationMatrix(ViewOriginDelta) * Context.View.PrevViewRotationProjMatrix);
-		}
-		else
-		{
-			SetShaderValue(Context.RHICmdList, ShaderRHI, PrevViewProjMatrix, Context.View.ViewMatrices.GetViewRotationProjMatrix());
-		}
-	}
-
-	static const TCHAR* GetSourceFilename()
-	{
-		return TEXT("PostProcessMotionBlur");
-	}
-
-	static const TCHAR* GetFunctionName()
-	{
-		return TEXT("VisualizeMotionBlurPS");
-	}
-};
-
-IMPLEMENT_SHADER_TYPE3(FPostProcessVisualizeMotionBlurPS, SF_Pixel);
-
-
-// @param Quality 0: visualize, 1:low, 2:medium, 3:high, 4:very high
-template <uint32 Quality>
-static void SetMotionBlurShaderTempl(const FRenderingCompositePassContext& Context)
-{
-	TShaderMapRef<FPostProcessVS> VertexShader(Context.GetShaderMap());
-	TShaderMapRef<FPostProcessMotionBlurPS<Quality> > PixelShader(Context.GetShaderMap());
-
-	static FGlobalBoundShaderState BoundShaderState;
-	
-	SetGlobalBoundShaderState(Context.RHICmdList, Context.GetFeatureLevel(), BoundShaderState, GFilterVertexDeclaration.VertexDeclarationRHI, *VertexShader, *PixelShader);
-
-	VertexShader->SetParameters(Context);
-	PixelShader->SetParameters(Context);
-}
-
-FRCPassPostProcessMotionBlur::FRCPassPostProcessMotionBlur(uint32 InQuality) 
-	: Quality(InQuality)
-{
-	// internal error
-	check(Quality >= 1 && Quality <= 4);
-}
-
-void FRCPassPostProcessMotionBlur::Process(FRenderingCompositePassContext& Context)
-{
-	const FPooledRenderTargetDesc* InputDesc = GetInputDesc(ePId_Input0);
-
-	if(!InputDesc)
-	{
-		// input is not hooked up correctly
-		return;
-	}
-
-	const FSceneView& View = Context.View;
-
-	FIntPoint TexSize = InputDesc->Extent;
-
-	// we assume the input and output is full resolution
-
-	FIntPoint SrcSize = InputDesc->Extent;
-	FIntPoint DestSize = PassOutputs[0].RenderTargetDesc.Extent;
-
-	// e.g. 4 means the input texture is 4x smaller than the buffer size
-	uint32 ScaleFactor = FSceneRenderTargets::Get(Context.RHICmdList).GetBufferSizeXY().X / SrcSize.X;
-
-	FIntRect SrcRect = FIntRect::DivideAndRoundUp(View.ViewRect, ScaleFactor);
-
-	SCOPED_DRAW_EVENTF(Context.RHICmdList, MotionBlur, TEXT("MotionBlur(Old) %dx%d"), SrcRect.Width(), SrcRect.Height());
-
-	FIntRect DestRect = SrcRect;
-
-	const FSceneRenderTargetItem& DestRenderTarget = PassOutputs[0].RequestSurface(Context);
-
-	// Set the view family's render target/viewport.
-	SetRenderTarget(Context.RHICmdList, DestRenderTarget.TargetableTexture, FTextureRHIRef());
-
-	// is optimized away if possible (RT size=view size, )
-	Context.RHICmdList.Clear(true, FLinearColor::Black, false, 1.0f, false, 0, SrcRect);
-
-	Context.SetViewportAndCallRHI(SrcRect);
-
-	// set the state
-	Context.RHICmdList.SetBlendState(TStaticBlendState<>::GetRHI());
-	Context.RHICmdList.SetRasterizerState(TStaticRasterizerState<>::GetRHI());
-	Context.RHICmdList.SetDepthStencilState(TStaticDepthStencilState<false, CF_Always>::GetRHI());
-
-	if(Quality == 1)
-	{
-		SetMotionBlurShaderTempl<1>(Context);
-	}
-	else if(Quality == 2)
-	{
-		SetMotionBlurShaderTempl<2>(Context);
-	}
-	else if(Quality == 3)
-	{
-		SetMotionBlurShaderTempl<3>(Context);
-	}
-	else
-	{
-		check(Quality == 4);
-		SetMotionBlurShaderTempl<4>(Context);
-	}
-
-	TShaderMapRef<FPostProcessVS> VertexShader(Context.GetShaderMap());
-
-	DrawPostProcessPass(
-		Context.RHICmdList,
-		0, 0,
-		SrcRect.Width(), SrcRect.Height(),
-		SrcRect.Min.X, SrcRect.Min.Y,
-		SrcRect.Width(), SrcRect.Height(),
-		SrcRect.Size(),
-		SrcSize,
-		*VertexShader,
-		View.StereoPass,
-		Context.HasHmdMesh(),
-		EDRF_UseTriangleOptimization);
-
-	Context.RHICmdList.CopyToResolveTarget(DestRenderTarget.TargetableTexture, DestRenderTarget.ShaderResourceTexture, false, FResolveParams());
-}
-
-FPooledRenderTargetDesc FRCPassPostProcessMotionBlur::ComputeOutputDesc(EPassOutputId InPassOutputId) const
-{
-	FPooledRenderTargetDesc Ret = GetInput(ePId_Input0)->GetOutput()->RenderTargetDesc;
-
-	Ret.Reset();
-	Ret.DebugName = TEXT("MotionBlur");
-	Ret.AutoWritable = false;
-
-	return Ret;
-}
-
-/** Encapsulates a MotionBlur recombine pixel shader. */
-class FPostProcessMotionBlurRecombinePS : public FGlobalShader
-{
-	DECLARE_SHADER_TYPE(FPostProcessMotionBlurRecombinePS, Global);
-
-	static bool ShouldCache(EShaderPlatform Platform)
-	{
-		return IsFeatureLevelSupported(Platform, ERHIFeatureLevel::SM4);
-	}
-
-	static void ModifyCompilationEnvironment(EShaderPlatform Platform, FShaderCompilerEnvironment& OutEnvironment)
-	{
-		FGlobalShader::ModifyCompilationEnvironment(Platform, OutEnvironment);
-	}
-
-	/** Default constructor. */
-	FPostProcessMotionBlurRecombinePS() {}
-
-public:
-	FPostProcessPassParameters PostprocessParameter;
-
-	/** Initialization constructor. */
-	FPostProcessMotionBlurRecombinePS(const ShaderMetaType::CompiledShaderInitializerType& Initializer)
-		: FGlobalShader(Initializer)
-	{
-		PostprocessParameter.Bind(Initializer.ParameterMap);
-	}
-
-	// FShader interface.
-	virtual bool Serialize(FArchive& Ar) override
-	{
-		bool bShaderHasOutdatedParameters = FGlobalShader::Serialize(Ar);
-		Ar << PostprocessParameter;
-		return bShaderHasOutdatedParameters;
-	}
-
-	void SetParameters(const FRenderingCompositePassContext& Context)
-	{
-		const FPixelShaderRHIParamRef ShaderRHI = GetPixelShader();
-
-		FGlobalShader::SetParameters(Context.RHICmdList, ShaderRHI, Context.View);
-
-		PostprocessParameter.SetPS(ShaderRHI, Context, TStaticSamplerState<SF_Bilinear, AM_Border, AM_Border, AM_Clamp>::GetRHI());
-	}
-};
-
-IMPLEMENT_SHADER_TYPE(,FPostProcessMotionBlurRecombinePS,TEXT("PostProcessMotionBlur"),TEXT("MainRecombinePS"),SF_Pixel);
-
-void FRCPassPostProcessMotionBlurRecombine::Process(FRenderingCompositePassContext& Context)
-{
-	const FPooledRenderTargetDesc* InputDesc = GetInputDesc(ePId_Input0);
-
-	if(!InputDesc)
-	{
-		// input is not hooked up correctly
-		return;
-	}
-
-	const FSceneView& View = Context.View;
-
-	FIntPoint TexSize = InputDesc->Extent;
-
-	// we assume the input and output is full resolution
-
-	FIntPoint SrcSize = InputDesc->Extent;
-	FIntPoint DestSize = PassOutputs[0].RenderTargetDesc.Extent;
-
-	// e.g. 4 means the input texture is 4x smaller than the buffer size
-	uint32 ScaleFactor = FSceneRenderTargets::Get(Context.RHICmdList).GetBufferSizeXY().X / SrcSize.X;
-
-	FIntRect SrcRect = View.ViewRect / ScaleFactor;
-	FIntRect DestRect = SrcRect;
-
-	SCOPED_DRAW_EVENTF(Context.RHICmdList, MotionBlurRecombine, TEXT("MotionBlurRecombine %dx%d"), DestRect.Width(), DestRect.Height());
-
-	const FSceneRenderTargetItem& DestRenderTarget = PassOutputs[0].RequestSurface(Context);
-
-	// Set the view family's render target/viewport.
-	SetRenderTarget(Context.RHICmdList, DestRenderTarget.TargetableTexture, FTextureRHIRef());
-
-	// is optimized away if possible (RT size=view size, )
-	Context.RHICmdList.Clear(true, FLinearColor::Black, false, 1.0f, false, 0, SrcRect);
-
-	Context.SetViewportAndCallRHI(SrcRect);
-
-	// set the state
-	Context.RHICmdList.SetBlendState(TStaticBlendState<>::GetRHI());
-	Context.RHICmdList.SetRasterizerState(TStaticRasterizerState<>::GetRHI());
-	Context.RHICmdList.SetDepthStencilState(TStaticDepthStencilState<false, CF_Always>::GetRHI());
-
-	TShaderMapRef<FPostProcessVS> VertexShader(Context.GetShaderMap());
-	TShaderMapRef<FPostProcessMotionBlurRecombinePS> PixelShader(Context.GetShaderMap());
-
-	static FGlobalBoundShaderState BoundShaderState;
-	
-	SetGlobalBoundShaderState(Context.RHICmdList, Context.GetFeatureLevel(), BoundShaderState, GFilterVertexDeclaration.VertexDeclarationRHI, *VertexShader, *PixelShader);
-
-	VertexShader->SetParameters(Context);
-	PixelShader->SetParameters(Context);
-
-	DrawPostProcessPass(
-		Context.RHICmdList,
-		0, 0,
-		SrcRect.Width(), SrcRect.Height(),
-		SrcRect.Min.X, SrcRect.Min.Y,
-		SrcRect.Width(), SrcRect.Height(),
-		SrcRect.Size(),
-		SrcSize,
-		*VertexShader,
-		View.StereoPass,
-		Context.HasHmdMesh(),
-		EDRF_UseTriangleOptimization);
-
-	Context.RHICmdList.CopyToResolveTarget(DestRenderTarget.TargetableTexture, DestRenderTarget.ShaderResourceTexture, false, FResolveParams());
-}
-
-FPooledRenderTargetDesc FRCPassPostProcessMotionBlurRecombine::ComputeOutputDesc(EPassOutputId InPassOutputId) const
-{
-	FPooledRenderTargetDesc Ret = GetInput(ePId_Input0)->GetOutput()->RenderTargetDesc;
-
-	Ret.Reset();
-	// we don't need the alpha channel and 32bit is faster and costs less memory
-	Ret.Format = PF_FloatRGB;
-	Ret.AutoWritable = false;
-	Ret.DebugName = TEXT("MotionBlurRecombine");
-
-	return Ret;
-}
-
-
 FIntPoint GetNumTiles16x16( FIntPoint PixelExtent )
 {
 	uint32 TilesX = (PixelExtent.X + 15) / 16;
@@ -926,11 +206,6 @@ FPooledRenderTargetDesc FRCPassPostProcessVelocityFlatten::ComputeOutputDesc(EPa
 }
 
 
-
-
-
-
-
 class FScatterQuadIndexBuffer : public FIndexBuffer
 {
 public:
@@ -964,7 +239,7 @@ class FPostProcessVelocityScatterVS : public FGlobalShader
 
 	static bool ShouldCache(EShaderPlatform Platform)
 	{
-		return IsFeatureLevelSupported(Platform, ERHIFeatureLevel::SM4);
+		return IsFeatureLevelSupported(Platform, ERHIFeatureLevel::SM5);
 	}
 
 	/** Default constructor. */
@@ -1021,7 +296,7 @@ class FPostProcessVelocityScatterPS : public FGlobalShader
 
 	static bool ShouldCache(EShaderPlatform Platform)
 	{
-		return IsFeatureLevelSupported(Platform, ERHIFeatureLevel::SM4);
+		return IsFeatureLevelSupported(Platform, ERHIFeatureLevel::SM5);
 	}
 
 	/** Default constructor. */
@@ -1239,13 +514,13 @@ FPooledRenderTargetDesc FRCPassPostProcessVelocityGather::ComputeOutputDesc(EPas
  * @param Quality 0: visualize, 1:low, 2:medium, 3:high, 4:very high
  */
 template< uint32 Quality >
-class FPostProcessMotionBlurNewPS : public FGlobalShader
+class FPostProcessMotionBlurPS : public FGlobalShader
 {
-	DECLARE_SHADER_TYPE(FPostProcessMotionBlurNewPS, Global);
+	DECLARE_SHADER_TYPE(FPostProcessMotionBlurPS, Global);
 
 	static bool ShouldCache(EShaderPlatform Platform)
 	{
-		return IsFeatureLevelSupported(Platform, ERHIFeatureLevel::SM4);
+		return IsFeatureLevelSupported(Platform, ERHIFeatureLevel::SM5);
 	}
 
 	static void ModifyCompilationEnvironment(EShaderPlatform Platform, FShaderCompilerEnvironment& OutEnvironment)
@@ -1255,7 +530,7 @@ class FPostProcessMotionBlurNewPS : public FGlobalShader
 	}
 
 	/** Default constructor. */
-	FPostProcessMotionBlurNewPS() {}
+	FPostProcessMotionBlurPS() {}
 
 public:
 	FPostProcessPassParameters PostprocessParameter;
@@ -1263,7 +538,7 @@ public:
 	FShaderParameter MotionBlurParameters;
 
 	/** Initialization constructor. */
-	FPostProcessMotionBlurNewPS(const ShaderMetaType::CompiledShaderInitializerType& Initializer)
+	FPostProcessMotionBlurPS(const ShaderMetaType::CompiledShaderInitializerType& Initializer)
 		: FGlobalShader(Initializer)
 	{
 		PostprocessParameter.Bind(Initializer.ParameterMap);
@@ -1324,13 +599,13 @@ public:
 
 	static const TCHAR* GetFunctionName()
 	{
-		return TEXT("MainNewPS");
+		return TEXT("MainPS");
 	}
 };
 
 // #define avoids a lot of code duplication
-#define VARIATION1(A) typedef FPostProcessMotionBlurNewPS<A> FPostProcessMotionBlurNewPS##A; \
-	IMPLEMENT_SHADER_TYPE2(FPostProcessMotionBlurNewPS##A, SF_Pixel);
+#define VARIATION1(A) typedef FPostProcessMotionBlurPS<A> FPostProcessMotionBlurPS##A; \
+	IMPLEMENT_SHADER_TYPE2(FPostProcessMotionBlurPS##A, SF_Pixel);
 
 VARIATION1(1)			VARIATION1(2)			VARIATION1(3)			VARIATION1(4)
 #undef VARIATION1
@@ -1342,7 +617,7 @@ template< uint32 Quality >
 static void SetMotionBlurShaderNewTempl( const FRenderingCompositePassContext& Context, float Scale )
 {
 	TShaderMapRef< FPostProcessVS > VertexShader( Context.GetShaderMap() );
-	TShaderMapRef< FPostProcessMotionBlurNewPS< Quality > >	PixelShader( Context.GetShaderMap() );
+	TShaderMapRef< FPostProcessMotionBlurPS< Quality > >	PixelShader( Context.GetShaderMap() );
 
 	static FGlobalBoundShaderState BoundShaderState;
 	SetGlobalBoundShaderState(Context.RHICmdList, Context.GetFeatureLevel(), BoundShaderState, GFilterVertexDeclaration.VertexDeclarationRHI, *VertexShader, *PixelShader);
@@ -1351,7 +626,7 @@ static void SetMotionBlurShaderNewTempl( const FRenderingCompositePassContext& C
 	PixelShader->SetParameters(Context, Scale);
 }
 
-void FRCPassPostProcessMotionBlurNew::Process(FRenderingCompositePassContext& Context)
+void FRCPassPostProcessMotionBlur::Process(FRenderingCompositePassContext& Context)
 {
 	const FPooledRenderTargetDesc* InputDesc = GetInputDesc(ePId_Input0);
 
@@ -1374,7 +649,7 @@ void FRCPassPostProcessMotionBlurNew::Process(FRenderingCompositePassContext& Co
 	FIntRect SrcRect = View.ViewRect / ScaleFactor;
 	FIntRect DestRect = SrcRect;
 
-	SCOPED_DRAW_EVENTF(Context.RHICmdList, MotionBlurNew, TEXT("MotionBlur(New) %dx%d"), SrcRect.Width(), SrcRect.Height());
+	SCOPED_DRAW_EVENTF(Context.RHICmdList, MotionBlurNew, TEXT("MotionBlur %dx%d"), SrcRect.Width(), SrcRect.Height());
 
 	const FSceneRenderTargetItem& DestRenderTarget = PassOutputs[0].RequestSurface(Context);
 
@@ -1440,7 +715,7 @@ void FRCPassPostProcessMotionBlurNew::Process(FRenderingCompositePassContext& Co
 	Context.RHICmdList.CopyToResolveTarget(DestRenderTarget.TargetableTexture, DestRenderTarget.ShaderResourceTexture, false, FResolveParams());
 }
 
-FPooledRenderTargetDesc FRCPassPostProcessMotionBlurNew::ComputeOutputDesc(EPassOutputId InPassOutputId) const
+FPooledRenderTargetDesc FRCPassPostProcessMotionBlur::ComputeOutputDesc(EPassOutputId InPassOutputId) const
 {
 	FPooledRenderTargetDesc Ret = GetInput(ePId_Input0)->GetOutput()->RenderTargetDesc;
 
@@ -1453,6 +728,96 @@ FPooledRenderTargetDesc FRCPassPostProcessMotionBlurNew::ComputeOutputDesc(EPass
 }
 
 
+
+
+class FPostProcessVisualizeMotionBlurPS : public FGlobalShader
+{
+	DECLARE_SHADER_TYPE(FPostProcessVisualizeMotionBlurPS, Global);
+
+	static bool ShouldCache(EShaderPlatform Platform)
+	{
+		return IsFeatureLevelSupported(Platform, ERHIFeatureLevel::SM4);
+	}
+
+	/** Default constructor. */
+	FPostProcessVisualizeMotionBlurPS() {}
+
+public:
+	FPostProcessPassParameters PostprocessParameter;
+	FDeferredPixelShaderParameters DeferredParameters;
+	FShaderParameter PrevViewProjMatrix;
+
+	/** Initialization constructor. */
+	FPostProcessVisualizeMotionBlurPS(const ShaderMetaType::CompiledShaderInitializerType& Initializer)
+		: FGlobalShader(Initializer)
+	{
+		PostprocessParameter.Bind(Initializer.ParameterMap);
+		DeferredParameters.Bind(Initializer.ParameterMap);
+		PrevViewProjMatrix.Bind(Initializer.ParameterMap, TEXT("PrevViewProjMatrix"));
+	}
+
+	// FShader interface.
+	virtual bool Serialize(FArchive& Ar) override
+	{
+		bool bShaderHasOutdatedParameters = FGlobalShader::Serialize(Ar);
+		Ar << PostprocessParameter << DeferredParameters << PrevViewProjMatrix;
+		return bShaderHasOutdatedParameters;
+	}
+
+	void SetParameters(const FRenderingCompositePassContext& Context)
+	{
+		const FPixelShaderRHIParamRef ShaderRHI = GetPixelShader();
+
+		FGlobalShader::SetParameters(Context.RHICmdList, ShaderRHI, Context.View);
+
+		DeferredParameters.Set(Context.RHICmdList, ShaderRHI, Context.View);
+
+		SetUniformBufferParameter(Context.RHICmdList, ShaderRHI, GetUniformBufferParameter<FCameraMotionParameters>(), CreateCameraMotionParametersUniformBuffer(Context.View));
+
+		{
+			bool bFiltered = false;
+
+#if !(UE_BUILD_SHIPPING || UE_BUILD_TEST)
+			bFiltered = CVarMotionBlurFiltering.GetValueOnRenderThread() != 0;
+#endif // !(UE_BUILD_SHIPPING || UE_BUILD_TEST)
+
+			if(bFiltered)
+			{
+				PostprocessParameter.SetPS(ShaderRHI, Context, TStaticSamplerState<SF_Bilinear,AM_Border,AM_Border,AM_Clamp>::GetRHI());
+			}
+			else
+			{
+				PostprocessParameter.SetPS(ShaderRHI, Context, TStaticSamplerState<SF_Point,AM_Border,AM_Border,AM_Clamp>::GetRHI());
+			}
+		}
+		
+		if( Context.View.Family->EngineShowFlags.CameraInterpolation )
+		{
+			// Instead of finding the world space position of the current pixel, calculate the world space position offset by the camera position, 
+			// then translate by the difference between last frame's camera position and this frame's camera position,
+			// then apply the rest of the transforms.  This effectively avoids precision issues near the extents of large levels whose world space position is very large.
+			FVector ViewOriginDelta = Context.View.ViewMatrices.ViewOrigin - Context.View.PrevViewMatrices.ViewOrigin;
+			SetShaderValue(Context.RHICmdList, ShaderRHI, PrevViewProjMatrix, FTranslationMatrix(ViewOriginDelta) * Context.View.PrevViewRotationProjMatrix);
+		}
+		else
+		{
+			SetShaderValue(Context.RHICmdList, ShaderRHI, PrevViewProjMatrix, Context.View.ViewMatrices.GetViewRotationProjMatrix());
+		}
+	}
+
+	static const TCHAR* GetSourceFilename()
+	{
+		return TEXT("PostProcessMotionBlur");
+	}
+
+	static const TCHAR* GetFunctionName()
+	{
+		return TEXT("VisualizeMotionBlurPS");
+	}
+};
+
+IMPLEMENT_SHADER_TYPE3(FPostProcessVisualizeMotionBlurPS, SF_Pixel);
+
 void FRCPassPostProcessVisualizeMotionBlur::Process(FRenderingCompositePassContext& Context)
 {
 	SCOPED_DRAW_EVENT(Context.RHICmdList, VisualizeMotionBlur);
diff --git a/Engine/Source/Runtime/Renderer/Private/PostProcess/PostProcessMotionBlur.h b/Engine/Source/Runtime/Renderer/Private/PostProcess/PostProcessMotionBlur.h
--- a/Engine/Source/Runtime/Renderer/Private/PostProcess/PostProcessMotionBlur.h
+++ b/Engine/Source/Runtime/Renderer/Private/PostProcess/PostProcessMotionBlur.h
@@ -8,71 +8,6 @@
 
 #include "RenderingCompositionGraph.h"
 
-// down sample and prepare for soft masked blurring
-// derives from TRenderingCompositePassBase<InputCount, OutputCount>
-// MRT0: half res velocity, MRT1: half res scene color with depth in alpha
-// ePId_Input0: optional Full Res velocity input
-// ePId_Input1: Full Res Scene Color
-// ePId_Input2: Full Res Scene Depth
-class FRCPassPostProcessMotionBlurSetup : public TRenderingCompositePassBase<3, 2>
-{
-public:
-	// interface FRenderingCompositePass ---------
-
-	virtual void Process(FRenderingCompositePassContext& Context) override;
-	virtual void Release() override { delete this; }
-	virtual FPooledRenderTargetDesc ComputeOutputDesc(EPassOutputId InPassOutputId) const override;
-};
-
-
-// input half res scene with depth, velocity, output half res motion blurred version
-// ePId_Input0: Half res SceneColor with depth in alpha
-// ePId_Input1: optional quarter res blurred screen space velocity
-// ePId_Input2: optional half res screen space masked velocity
-// derives from TRenderingCompositePassBase<InputCount, OutputCount> 
-class FRCPassPostProcessMotionBlur : public TRenderingCompositePassBase<3, 1>
-{
-public:
-	// @param InQuality 0xffffffff to visualize, 0:off(no shader is used), 1:low, 2:medium, 3:high, 4:very high
-	FRCPassPostProcessMotionBlur(uint32 InQuality);
-
-	// interface FRenderingCompositePass ---------
-	virtual void Process(FRenderingCompositePassContext& Context) override;
-	virtual void Release() override { delete this; }
-	virtual FPooledRenderTargetDesc ComputeOutputDesc(EPassOutputId InPassOutputId) const override;
-
-	// 1:low, 2:medium, 3:high, 4: very high
-	uint32 Quality;
-};
-
-
-// ePId_Input0: Full Res Scene color
-// ePId_Input1: output from FRCPassPostProcessMotionBlur
-// derives from TRenderingCompositePassBase<InputCount, OutputCount> 
-class FRCPassPostProcessMotionBlurRecombine : public TRenderingCompositePassBase<2, 1>
-{
-public:
-	// interface FRenderingCompositePass ---------
-	virtual void Process(FRenderingCompositePassContext& Context) override;
-	virtual void Release() override { delete this; }
-	virtual FPooledRenderTargetDesc ComputeOutputDesc(EPassOutputId InPassOutputId) const override;
-};
-
-
-// ePId_Input0: Full Res Scene Color
-// ePId_Input1: Full Res Scene Depth
-// ePId_Input2: Full Res velocity input
-// derives from TRenderingCompositePassBase<InputCount, OutputCount> 
-class FRCPassPostProcessVisualizeMotionBlur : public TRenderingCompositePassBase<3, 1>
-{
-public:
-	// interface FRenderingCompositePass ---------
-	virtual void Process(FRenderingCompositePassContext& Context) override;
-	virtual void Release() override { delete this; }
-	virtual FPooledRenderTargetDesc ComputeOutputDesc(EPassOutputId InPassOutputId) const override;
-};
-
-
 // Find max velocity per 16x16 tile
 // derives from TRenderingCompositePassBase<InputCount, OutputCount>
 class FRCPassPostProcessVelocityFlatten : public TRenderingCompositePassBase<2, 2>
@@ -112,17 +47,16 @@
 };
 
 
-// input half res scene with depth, velocity, output half res motion blurred version
 // ePId_Input0: Full Res Scene Color
 // ePId_Input1: Full Res Scene Depth
 // ePId_Input2: Velocity
 // ePId_Input3: Max tile velocity
 // derives from TRenderingCompositePassBase<InputCount, OutputCount> 
-class FRCPassPostProcessMotionBlurNew : public TRenderingCompositePassBase<4, 1>
+class FRCPassPostProcessMotionBlur : public TRenderingCompositePassBase<4, 1>
 {
 public:
 	// @param InQuality 0xffffffff to visualize, 0:off(no shader is used), 1:low, 2:medium, 3:high, 4:very high
-	FRCPassPostProcessMotionBlurNew( uint32 InQuality, int32 InPass )
+	FRCPassPostProcessMotionBlur( uint32 InQuality, int32 InPass )
 		: Quality(InQuality)
 		, Pass(InPass)
 	{
@@ -139,3 +73,17 @@
 	uint32	Quality;
 	int32	Pass;
 };
+
+
+// ePId_Input0: Full Res Scene Color
+// ePId_Input1: Full Res Scene Depth
+// ePId_Input2: Full Res velocity input
+// derives from TRenderingCompositePassBase<InputCount, OutputCount> 
+class FRCPassPostProcessVisualizeMotionBlur : public TRenderingCompositePassBase<3, 1>
+{
+public:
+	// interface FRenderingCompositePass ---------
+	virtual void Process(FRenderingCompositePassContext& Context) override;
+	virtual void Release() override { delete this; }
+	virtual FPooledRenderTargetDesc ComputeOutputDesc(EPassOutputId InPassOutputId) const override;
+};
diff --git a/Engine/Source/Runtime/Renderer/Private/PostProcess/PostProcessing.cpp b/Engine/Source/Runtime/Renderer/Private/PostProcess/PostProcessing.cpp
--- a/Engine/Source/Runtime/Renderer/Private/PostProcess/PostProcessing.cpp
+++ b/Engine/Source/Runtime/Renderer/Private/PostProcess/PostProcessing.cpp
@@ -143,13 +143,6 @@ static TAutoConsoleVariable<int32> CVarTonemapperScreenPercentage(
 	TEXT(" 1: enabled if nothing prevents it which is safe but it might do the two pass method"),
 	ECVF_Scalability | ECVF_RenderThreadSafe);
 
-static TAutoConsoleVariable<int32> CVarMotionBlurNew(
-	TEXT("r.MotionBlurNew"),
-	1,
-	TEXT(""),
-	ECVF_RenderThreadSafe
-);
-
 static TAutoConsoleVariable<int32> CVarMotionBlurScatter(
 	TEXT("r.MotionBlurScatter"),
 	0,
@@ -376,9 +369,8 @@ static void AddPostProcessAA(FPostprocessContext& Context)
 }
 
 
-static FRenderingCompositeOutputRef AddPostProcessBasicEyeAdaptation(FViewInfo& View, FBloomDownSampleArray& BloomAndEyeDownSamples)
+static FRenderingCompositeOutputRef AddPostProcessBasicEyeAdaptation(const FViewInfo& View, FBloomDownSampleArray& BloomAndEyeDownSamples)
 {
-
 	// Extract the context
 	FPostprocessContext& Context = BloomAndEyeDownSamples.Context;
 
@@ -1142,7 +1134,7 @@ bool FPostProcessing::AllowFullPostProcessing(const FViewInfo& View, ERHIFeature
 		&& !View.Family->EngineShowFlags.VisualizeMeshDistanceFields;
 }
 
-void FPostProcessing::Process(FRHICommandListImmediate& RHICmdList, FViewInfo& View, TRefCountPtr<IPooledRenderTarget>& VelocityRT)
+void FPostProcessing::Process(FRHICommandListImmediate& RHICmdList, const FViewInfo& View, TRefCountPtr<IPooledRenderTarget>& VelocityRT)
 {
 	QUICK_SCOPE_CYCLE_COUNTER( STAT_PostProcessing_Process );
 
@@ -1250,18 +1242,6 @@ void FPostProcessing::Process(FRHICommandListImmediate& RHICmdList, FViewInfo& V
 			{
 				VelocityInput = Context.Graph.RegisterPass(new(FMemStack::Get()) FRCPassPostProcessInput(VelocityRT));
 			}
-			
-			if( Context.View.FinalPostProcessSettings.AntiAliasingMethod != AAM_TemporalAA && ViewState )			
-			{
-				if(ViewState->DOFHistoryRT)
-				{
-					ViewState->DOFHistoryRT.SafeRelease();
-				}
-				if(ViewState->TemporalAAHistoryRT)
-				{
-					ViewState->TemporalAAHistoryRT.SafeRelease();
-				}
-			}
 
 			AddPostProcessMaterial(Context, BL_BeforeTranslucency, SeparateTranslucency);
 
@@ -1366,115 +1346,59 @@ void FPostProcessing::Process(FRHICommandListImmediate& RHICmdList, FViewInfo& V
 			{
 				// Motion blur
 
-				if( CVarMotionBlurNew.GetValueOnRenderThread() && FeatureLevel >= ERHIFeatureLevel::SM5 )
-				{
-					FRenderingCompositeOutputRef MaxTileVelocity;
-
-					{
-						check(!VelocityFlattenPass);
-						VelocityFlattenPass = Context.Graph.RegisterPass( new(FMemStack::Get()) FRCPassPostProcessVelocityFlatten() );
-						VelocityFlattenPass->SetInput( ePId_Input0, VelocityInput );
-						VelocityFlattenPass->SetInput( ePId_Input1, Context.SceneDepth );
-
-						VelocityInput	= FRenderingCompositeOutputRef( VelocityFlattenPass, ePId_Output0 );
-						MaxTileVelocity	= FRenderingCompositeOutputRef( VelocityFlattenPass, ePId_Output1 );
-					}
-
-					const float SizeX = View.ViewRect.Width();
+				FRenderingCompositeOutputRef MaxTileVelocity;
 
-					// 0:no 1:full screen width, percent conversion
-					float MaxVelocity = View.FinalPostProcessSettings.MotionBlurMax / 100.0f;
-					float MaxVelocityTiles = MaxVelocity * SizeX * (0.5f / 16.0f);
-					float MaxTileDistGathered = 3.0f;
-					if( MaxVelocityTiles > MaxTileDistGathered || CVarMotionBlurScatter.GetValueOnRenderThread() )
-					{
-						FRenderingCompositePass* VelocityScatterPass = Context.Graph.RegisterPass( new(FMemStack::Get()) FRCPassPostProcessVelocityScatter() );
-						VelocityScatterPass->SetInput( ePId_Input0, MaxTileVelocity );
+				{
+					check(!VelocityFlattenPass);
+					VelocityFlattenPass = Context.Graph.RegisterPass( new(FMemStack::Get()) FRCPassPostProcessVelocityFlatten() );
+					VelocityFlattenPass->SetInput( ePId_Input0, VelocityInput );
+					VelocityFlattenPass->SetInput( ePId_Input1, Context.SceneDepth );
 
-						MaxTileVelocity	= FRenderingCompositeOutputRef( VelocityScatterPass );
-					}
-					else
-					{
-						FRenderingCompositePass* VelocityGatherPass = Context.Graph.RegisterPass( new(FMemStack::Get()) FRCPassPostProcessVelocityGather() );
-						VelocityGatherPass->SetInput( ePId_Input0, MaxTileVelocity );
+					VelocityInput	= FRenderingCompositeOutputRef( VelocityFlattenPass, ePId_Output0 );
+					MaxTileVelocity	= FRenderingCompositeOutputRef( VelocityFlattenPass, ePId_Output1 );
+				}
 
-						MaxTileVelocity	= FRenderingCompositeOutputRef( VelocityGatherPass );
-					}
+				const float SizeX = View.ViewRect.Width();
 
-					bool bTwoPass = CVarMotionBlurSeparable.GetValueOnRenderThread() != 0;
-					{
-						FRenderingCompositePass* MotionBlurPass = Context.Graph.RegisterPass( new(FMemStack::Get()) FRCPassPostProcessMotionBlurNew( GetMotionBlurQualityFromCVar(), bTwoPass ? 0 : -1 ) );
-						MotionBlurPass->SetInput( ePId_Input0, Context.FinalOutput );
-						MotionBlurPass->SetInput( ePId_Input1, Context.SceneDepth );
-						MotionBlurPass->SetInput( ePId_Input2, VelocityInput );
-						MotionBlurPass->SetInput( ePId_Input3, MaxTileVelocity );
-					
-						Context.FinalOutput = FRenderingCompositeOutputRef( MotionBlurPass );
-					}
+				// 0:no 1:full screen width, percent conversion
+				float MaxVelocity = View.FinalPostProcessSettings.MotionBlurMax / 100.0f;
+				float MaxVelocityTiles = MaxVelocity * SizeX * (0.5f / 16.0f);
+				float MaxTileDistGathered = 3.0f;
+				if( MaxVelocityTiles > MaxTileDistGathered || CVarMotionBlurScatter.GetValueOnRenderThread() )
+				{
+					FRenderingCompositePass* VelocityScatterPass = Context.Graph.RegisterPass( new(FMemStack::Get()) FRCPassPostProcessVelocityScatter() );
+					VelocityScatterPass->SetInput( ePId_Input0, MaxTileVelocity );
 
-					if( bTwoPass )
-					{
-						FRenderingCompositePass* MotionBlurPass = Context.Graph.RegisterPass( new(FMemStack::Get()) FRCPassPostProcessMotionBlurNew( GetMotionBlurQualityFromCVar(), 1 ) );
-						MotionBlurPass->SetInput( ePId_Input0, Context.FinalOutput );
-						MotionBlurPass->SetInput( ePId_Input1, Context.SceneDepth );
-						MotionBlurPass->SetInput( ePId_Input2, VelocityInput );
-						MotionBlurPass->SetInput( ePId_Input3, MaxTileVelocity );
-					
-						Context.FinalOutput = FRenderingCompositeOutputRef( MotionBlurPass );
-					}
+					MaxTileVelocity	= FRenderingCompositeOutputRef( VelocityScatterPass );
 				}
 				else
 				{
-					FRenderingCompositeOutputRef SoftEdgeVelocity;
+					FRenderingCompositePass* VelocityGatherPass = Context.Graph.RegisterPass( new(FMemStack::Get()) FRCPassPostProcessVelocityGather() );
+					VelocityGatherPass->SetInput( ePId_Input0, MaxTileVelocity );
 
-					FRenderingCompositeOutputRef MotionBlurHalfVelocity;
-					FRenderingCompositeOutputRef MotionBlurColorDepth;
+					MaxTileVelocity	= FRenderingCompositeOutputRef( VelocityGatherPass );
+				}
 
-					// down sample screen space velocity and extend outside of borders to allows soft edge motion blur
-					{
-						// Down sample and prepare for soft masked blurring
-						FRenderingCompositePass* HalfResVelocity = Context.Graph.RegisterPass(new(FMemStack::Get()) FRCPassPostProcessMotionBlurSetup());
-						HalfResVelocity->SetInput(ePId_Input0, VelocityInput);
-						HalfResVelocity->SetInput(ePId_Input1, FRenderingCompositeOutputRef(Context.FinalOutput));
-						// only for Depth of Field we need the depth in the alpha channel
-						HalfResVelocity->SetInput(ePId_Input2, FRenderingCompositeOutputRef(Context.SceneDepth));
+				bool bTwoPass = CVarMotionBlurSeparable.GetValueOnRenderThread() != 0;
+				{
+					FRenderingCompositePass* MotionBlurPass = Context.Graph.RegisterPass( new(FMemStack::Get()) FRCPassPostProcessMotionBlur( GetMotionBlurQualityFromCVar(), bTwoPass ? 0 : -1 ) );
+					MotionBlurPass->SetInput( ePId_Input0, Context.FinalOutput );
+					MotionBlurPass->SetInput( ePId_Input1, Context.SceneDepth );
+					MotionBlurPass->SetInput( ePId_Input2, VelocityInput );
+					MotionBlurPass->SetInput( ePId_Input3, MaxTileVelocity );
 					
-						MotionBlurHalfVelocity = FRenderingCompositeOutputRef(HalfResVelocity, ePId_Output0);
-						MotionBlurColorDepth = FRenderingCompositeOutputRef(HalfResVelocity, ePId_Output1);
-
-						FRenderingCompositePass* QuarterResVelocity = Context.Graph.RegisterPass(new(FMemStack::Get()) FRCPassPostProcessDownsample(PF_Unknown, 0, TEXT("QuarterResVelocity")));
-						QuarterResVelocity->SetInput(ePId_Input0, HalfResVelocity);
-
-						SoftEdgeVelocity = FRenderingCompositeOutputRef(QuarterResVelocity);
-
-						float MotionBlurSoftEdgeSize = CVarMotionBlurSoftEdgeSize.GetValueOnRenderThread();
-
-						if(MotionBlurSoftEdgeSize > 0.01f)
-						{
-							SoftEdgeVelocity = RenderGaussianBlur(Context, TEXT("VelocityBlurX"), TEXT("VelocityBlurY"), QuarterResVelocity, MotionBlurSoftEdgeSize);
-						}
-					}
-
-					// doing the actual motion blur sampling, alpha to mask out the blurred areas
-					FRenderingCompositePass* MotionBlurPass;
-
-					int32 MotionBlurQuality = GetMotionBlurQualityFromCVar();
-
-					MotionBlurPass = Context.Graph.RegisterPass(new(FMemStack::Get()) FRCPassPostProcessMotionBlur(MotionBlurQuality));
-					MotionBlurPass->SetInput(ePId_Input0, MotionBlurColorDepth);
-
-					if(VelocityInput.IsValid())
-					{
-						// blurred screen space velocity for soft masked motion blur
-						MotionBlurPass->SetInput(ePId_Input1, SoftEdgeVelocity);
-						// screen space velocity input from per object velocity rendering
-						MotionBlurPass->SetInput(ePId_Input2, MotionBlurHalfVelocity);
-					}
+					Context.FinalOutput = FRenderingCompositeOutputRef( MotionBlurPass );
+				}
 
-					FRenderingCompositePass* MotionBlurRecombinePass = Context.Graph.RegisterPass(new(FMemStack::Get()) FRCPassPostProcessMotionBlurRecombine());
-					MotionBlurRecombinePass->SetInput(ePId_Input0, Context.FinalOutput);
-					MotionBlurRecombinePass->SetInput(ePId_Input1, MotionBlurPass);
-					Context.FinalOutput = FRenderingCompositeOutputRef(MotionBlurRecombinePass);
+				if( bTwoPass )
+				{
+					FRenderingCompositePass* MotionBlurPass = Context.Graph.RegisterPass( new(FMemStack::Get()) FRCPassPostProcessMotionBlur( GetMotionBlurQualityFromCVar(), 1 ) );
+					MotionBlurPass->SetInput( ePId_Input0, Context.FinalOutput );
+					MotionBlurPass->SetInput( ePId_Input1, Context.SceneDepth );
+					MotionBlurPass->SetInput( ePId_Input2, VelocityInput );
+					MotionBlurPass->SetInput( ePId_Input3, MaxTileVelocity );
+					
+					Context.FinalOutput = FRenderingCompositeOutputRef( MotionBlurPass );
 				}
 			}
 
@@ -1898,7 +1822,7 @@ void FPostProcessing::Process(FRHICommandListImmediate& RHICmdList, FViewInfo& V
 			Context.FinalOutput = FRenderingCompositeOutputRef(Node);
 		}
 
-		if(View.Family->EngineShowFlags.TestImage && FeatureLevel >= ERHIFeatureLevel::SM5)
+		if(View.Family->EngineShowFlags.TestImage && FeatureLevel >= ERHIFeatureLevel::SM4)
 		{
 			FRenderingCompositePass* Node = Context.Graph.RegisterPass(new(FMemStack::Get()) FRCPassPostProcessTestImage());
 			Node->SetInput(ePId_Input0, FRenderingCompositeOutputRef(Context.FinalOutput));
@@ -2006,7 +1930,7 @@ static bool IsGaussianActive(FPostprocessContext& Context)
 	return true;
 }
 
-void FPostProcessing::ProcessES2(FRHICommandListImmediate& RHICmdList, FViewInfo& View, bool bUsedFramebufferFetch)
+void FPostProcessing::ProcessES2(FRHICommandListImmediate& RHICmdList, const FViewInfo& View, bool bUsedFramebufferFetch)
 {
 	check(IsInRenderingThread());
 
diff --git a/Engine/Source/Runtime/Renderer/Private/PostProcess/PostProcessing.h b/Engine/Source/Runtime/Renderer/Private/PostProcess/PostProcessing.h
--- a/Engine/Source/Runtime/Renderer/Private/PostProcess/PostProcessing.h
+++ b/Engine/Source/Runtime/Renderer/Private/PostProcess/PostProcessing.h
@@ -74,9 +74,9 @@ public:
 	bool AllowFullPostProcessing(const FViewInfo& View, ERHIFeatureLevel::Type FeatureLevel);
 
 	// @param VelocityRT only valid if motion blur is supported
-	void Process(FRHICommandListImmediate& RHICmdList, FViewInfo& View, TRefCountPtr<IPooledRenderTarget>& VelocityRT);
+	void Process(FRHICommandListImmediate& RHICmdList, const FViewInfo& View, TRefCountPtr<IPooledRenderTarget>& VelocityRT);
 
-	void ProcessES2(FRHICommandListImmediate& RHICmdList, FViewInfo& View, bool bViewRectSource);
+	void ProcessES2(FRHICommandListImmediate& RHICmdList, const FViewInfo& View, bool bViewRectSource);
 
 	void ProcessPlanarReflection(FRHICommandListImmediate& RHICmdList, FViewInfo& View, TRefCountPtr<IPooledRenderTarget>& VelocityRT, TRefCountPtr<IPooledRenderTarget>& OutFilteredSceneColor);
 };
diff --git a/Engine/Source/Runtime/Renderer/Private/PostProcess/RenderingCompositionGraph.cpp b/Engine/Source/Runtime/Renderer/Private/PostProcess/RenderingCompositionGraph.cpp
--- a/Engine/Source/Runtime/Renderer/Private/PostProcess/RenderingCompositionGraph.cpp
+++ b/Engine/Source/Runtime/Renderer/Private/PostProcess/RenderingCompositionGraph.cpp
@@ -148,7 +148,7 @@ void CompositionGraph_OnStartFrame()
 #endif
 }
 
-FRenderingCompositePassContext::FRenderingCompositePassContext(FRHICommandListImmediate& InRHICmdList, FViewInfo& InView/*, const FSceneRenderTargetItem& InRenderTargetItem*/)
+FRenderingCompositePassContext::FRenderingCompositePassContext(FRHICommandListImmediate& InRHICmdList, const FViewInfo& InView)
 	: View(InView)
 	, ViewState((FSceneViewState*)InView.State)
 	, Pass(0)
diff --git a/Engine/Source/Runtime/Renderer/Private/PostProcess/RenderingCompositionGraph.h b/Engine/Source/Runtime/Renderer/Private/PostProcess/RenderingCompositionGraph.h
--- a/Engine/Source/Runtime/Renderer/Private/PostProcess/RenderingCompositionGraph.h
+++ b/Engine/Source/Runtime/Renderer/Private/PostProcess/RenderingCompositionGraph.h
@@ -80,7 +80,7 @@ private:
 struct FRenderingCompositePassContext
 {
 	// constructor
-	FRenderingCompositePassContext(FRHICommandListImmediate& RHICmdList, FViewInfo& InView);
+	FRenderingCompositePassContext(FRHICommandListImmediate& RHICmdList, const FViewInfo& InView);
 
 	// destructor
 	~FRenderingCompositePassContext();
@@ -139,7 +139,7 @@ struct FRenderingCompositePassContext
 	TShaderMap<FGlobalShaderType>* GetShaderMap() const { check(ShaderMap); return ShaderMap; }
 
 	//
-	FViewInfo& View;
+	const FViewInfo& View;
 	//
 	FSceneViewState* ViewState;
 	// is updated before each Pass->Process() call
diff --git a/Engine/Source/Runtime/Renderer/Private/PostProcess/SceneRenderTargets.cpp b/Engine/Source/Runtime/Renderer/Private/PostProcess/SceneRenderTargets.cpp
--- a/Engine/Source/Runtime/Renderer/Private/PostProcess/SceneRenderTargets.cpp
+++ b/Engine/Source/Runtime/Renderer/Private/PostProcess/SceneRenderTargets.cpp
@@ -98,8 +98,6 @@ static TAutoConsoleVariable<int32> CVarGBufferFormat(
 	TEXT(" 5: high precision"),
 	ECVF_RenderThreadSafe);
 
-extern ENGINE_API TAutoConsoleVariable<int32> CVarReflectionCaptureSize;
-
 /** The global render targets used for scene rendering. */
 static TGlobalResource<FSceneRenderTargets> SceneRenderTargetsSingleton;
 
@@ -1533,7 +1531,6 @@ void FSceneRenderTargets::AllocateForwardShadingPathRenderTargets(FRHICommandLis
 	// on ES2 we don't do on demand allocation of SceneColor yet (in non ES2 it's released in the Tonemapper Process())
 	AllocSceneColor(RHICmdList);
 	AllocateCommonDepthTargets(RHICmdList);
-	AllocateReflectionTargets(RHICmdList, CVarReflectionCaptureSize.GetValueOnRenderThread());
 	AllocateDebugViewModeTargets(RHICmdList);
 
 	EPixelFormat Format = GetSceneColor()->GetDesc().Format;
@@ -1601,7 +1598,7 @@ static TCHAR* const GetTranslucencyShadowTransmissionName(uint32 Id)
 }
 
 void FSceneRenderTargets::AllocateReflectionTargets(FRHICommandList& RHICmdList, int32 TargetSize)
-{	
+{
 	if (GSupportsRenderTargetFormat_PF_FloatRGBA)
 	{
 		const int32 NumReflectionCaptureMips = FMath::CeilLogTwo(TargetSize) + 1;
@@ -1845,8 +1842,9 @@ void FSceneRenderTargets::AllocateDeferredShadingPathRenderTargets(FRHICommandLi
 		GRenderTargetPool.FindFreeElement(RHICmdList, Desc, DirectionalOcclusion, TEXT("DirectionalOcclusion"));
 	}
 
+	if (CurrentFeatureLevel >= ERHIFeatureLevel::SM4) 
 	{
-		FPooledRenderTargetDesc Desc(FPooledRenderTargetDesc::Create2DDesc(BufferSize, GetSceneColorFormat(), FClearValueBinding::Black, TexCreate_None, TexCreate_RenderTargetable, false));
+		FPooledRenderTargetDesc Desc(FPooledRenderTargetDesc::Create2DDesc(BufferSize, PF_FloatRGBA, FClearValueBinding::Black, TexCreate_None, TexCreate_RenderTargetable, false));
 		if (CurrentFeatureLevel >= ERHIFeatureLevel::SM5)
 		{
 			Desc.TargetableFlags |= TexCreate_UAV;
@@ -1854,7 +1852,6 @@ void FSceneRenderTargets::AllocateDeferredShadingPathRenderTargets(FRHICommandLi
 		GRenderTargetPool.FindFreeElement(RHICmdList, Desc, LightAccumulation, TEXT("LightAccumulation"));
 	}
 
-	AllocateReflectionTargets(RHICmdList, CVarReflectionCaptureSize.GetValueOnRenderThread());
 	AllocateDebugViewModeTargets(RHICmdList);
 
 	if (CurrentFeatureLevel >= ERHIFeatureLevel::SM5)
diff --git a/Engine/Source/Runtime/Renderer/Private/PrimitiveSceneInfo.cpp b/Engine/Source/Runtime/Renderer/Private/PrimitiveSceneInfo.cpp
--- a/Engine/Source/Runtime/Renderer/Private/PrimitiveSceneInfo.cpp
+++ b/Engine/Source/Runtime/Renderer/Private/PrimitiveSceneInfo.cpp
@@ -83,6 +83,7 @@ FPrimitiveSceneInfo::FPrimitiveSceneInfo(UPrimitiveComponent* InComponent,FScene
 	Proxy(InComponent->SceneProxy),
 	PrimitiveComponentId(InComponent->ComponentId),
 	ComponentLastRenderTime(&InComponent->LastRenderTime),
+	ComponentLastRenderTimeOnScreen(&InComponent->LastRenderTimeOnScreen),
 	IndirectLightingCacheAllocation(NULL),
 	CachedPlanarReflectionProxy(NULL),
 	CachedReflectionCaptureProxy(NULL),
diff --git a/Engine/Source/Runtime/Renderer/Private/ReflectionEnvironmentCapture.cpp b/Engine/Source/Runtime/Renderer/Private/ReflectionEnvironmentCapture.cpp
--- a/Engine/Source/Runtime/Renderer/Private/ReflectionEnvironmentCapture.cpp
+++ b/Engine/Source/Runtime/Renderer/Private/ReflectionEnvironmentCapture.cpp
@@ -37,8 +37,6 @@ float ReflectionCaptureRoughnessMipScale = 1.2f;
 
 int32 GDiffuseIrradianceCubemapSize = 32;
 
-extern ENGINE_API TAutoConsoleVariable<int32> CVarReflectionCaptureSize;
-
 void OnUpdateReflectionCaptures( UWorld* InWorld )
 {
 	InWorld->UpdateAllReflectionCaptures();
@@ -50,7 +48,6 @@ FAutoConsoleCommandWithWorld CaptureConsoleCommand(
 	FConsoleCommandWithWorldDelegate::CreateStatic(OnUpdateReflectionCaptures)
 	);
 
-
 /** Encapsulates render target picking logic for cubemap mip generation. */
 FSceneRenderTargetItem& GetEffectiveRenderTarget(FSceneRenderTargets& SceneContext, bool bDownsamplePass, int32 TargetMipIndex)
 {
@@ -720,6 +717,7 @@ void ClearScratchCubemaps(FRHICommandList& RHICmdList, int32 TargetSize)
 void CaptureSceneToScratchCubemap(FRHICommandListImmediate& RHICmdList, FSceneRenderer* SceneRenderer, ECubeFace CubeFace, int32 CubemapSize, bool bCapturingForSkyLight, bool bLowerHemisphereIsBlack, const FLinearColor& LowerHemisphereColor)
 {
 	FMemMark MemStackMark(FMemStack::Get());
+
 	// update any resources that needed a deferred update
 	FDeferredUpdateResource::UpdateResources(RHICmdList);
 
@@ -743,8 +741,10 @@ void CaptureSceneToScratchCubemap(FRHICommandListImmediate& RHICmdList, FSceneRe
 		TEMP_PostReflectionCaptureRender();
 #endif
 
-		auto ShaderMap = GetGlobalShaderMap(FeatureLevel);
 		FSceneRenderTargets& SceneContext = FSceneRenderTargets::Get(RHICmdList);
+		SceneContext.AllocateReflectionTargets(RHICmdList, CubemapSize);
+
+		auto ShaderMap = GetGlobalShaderMap(FeatureLevel);
 
 		const int32 EffectiveSize = CubemapSize;
 		FSceneRenderTargetItem& EffectiveColorRT =  SceneContext.ReflectionColorScratchCubemap[0]->GetRenderTargetItem();
@@ -872,7 +872,8 @@ void FScene::AllocateReflectionCaptures(const TArray<UReflectionCaptureComponent
 
 			DesiredMaxCubemaps = FMath::Min(DesiredMaxCubemaps, GMaxNumReflectionCaptures);
 
-			if (DesiredMaxCubemaps != ReflectionSceneData.MaxAllocatedReflectionCubemapsGameThread || CVarReflectionCaptureSize.GetValueOnGameThread() != ReflectionSceneData.CubemapArray.GetCubemapSize())
+			const int32 ReflectionCaptureSize = UReflectionCaptureComponent::GetReflectionCaptureSize_GameThread();
+			if (DesiredMaxCubemaps != ReflectionSceneData.MaxAllocatedReflectionCubemapsGameThread || ReflectionCaptureSize != ReflectionSceneData.CubemapArray.GetCubemapSize())
 			{
 				ReflectionSceneData.MaxAllocatedReflectionCubemapsGameThread = DesiredMaxCubemaps;
 
@@ -880,10 +881,10 @@ void FScene::AllocateReflectionCaptures(const TArray<UReflectionCaptureComponent
 					ResizeArrayCommand,
 					FScene*, Scene, this,
 					uint32, MaxSize, ReflectionSceneData.MaxAllocatedReflectionCubemapsGameThread,
-					int32, CubemapSize, CVarReflectionCaptureSize.GetValueOnGameThread(),
+					int32, ReflectionCaptureSize, ReflectionCaptureSize,
 				{
 					// Update the scene's cubemap array, which will reallocate it, so we no longer have the contents of existing entries
-					Scene->ReflectionSceneData.CubemapArray.UpdateMaxCubemaps(MaxSize, CubemapSize);
+					Scene->ReflectionSceneData.CubemapArray.UpdateMaxCubemaps(MaxSize, ReflectionCaptureSize);
 				});
 
 				// Recapture all reflection captures now that we have reallocated the cubemap array
@@ -1184,12 +1185,15 @@ void CaptureSceneIntoScratchCubemap(
 
 		GReflectionCaptureRenderTarget.SetSize(CubemapSize);
 
-		FSceneViewFamilyContext ViewFamily( FSceneViewFamily::ConstructionValues(
-			&GReflectionCaptureRenderTarget,
-			Scene,
-			FEngineShowFlags(ESFIM_Game))
+		FSceneViewFamilyContext ViewFamily( 
+			FSceneViewFamily::ConstructionValues(
+				&GReflectionCaptureRenderTarget,
+				Scene,
+				FEngineShowFlags(ESFIM_Game)
+				)
 			.SetWorldTimes( 0.0f, 0.0f, 0.0f )				
-			.SetResolveScene(false) );
+			.SetResolveScene(false) 
+			);
 
 		// Disable features that are not desired when capturing the scene
 		ViewFamily.EngineShowFlags.PostProcessing = 0;
@@ -1275,7 +1279,7 @@ void CaptureSceneIntoScratchCubemap(
 
 void CopyToSceneArray(FRHICommandListImmediate& RHICmdList, FScene* Scene, FReflectionCaptureProxy* ReflectionProxy)
 {
-	const int32 EffectiveTopMipSize = CVarReflectionCaptureSize.GetValueOnRenderThread();
+	const int32 EffectiveTopMipSize = UReflectionCaptureComponent::GetReflectionCaptureSize_RenderThread();
 	const int32 NumMips = FMath::CeilLogTwo(EffectiveTopMipSize) + 1;
 
 	const int32 CaptureIndex = FindOrAllocateCubemapIndex(Scene, ReflectionProxy->Component);
@@ -1298,7 +1302,8 @@ void CopyToSceneArray(FRHICommandListImmediate& RHICmdList, FScene* Scene, FRefl
 void CopyToComponentTexture(FRHICommandList& RHICmdList, FScene* Scene, FReflectionCaptureProxy* ReflectionProxy)
 {
 	check(ReflectionProxy->SM4FullHDRCubemap);
-	const int32 EffectiveTopMipSize = CVarReflectionCaptureSize.GetValueOnRenderThread();
+
+	const int32 EffectiveTopMipSize = UReflectionCaptureComponent::GetReflectionCaptureSize_RenderThread();
 	const int32 NumMips = FMath::CeilLogTwo(EffectiveTopMipSize) + 1;
 	FSceneRenderTargets& SceneContext = FSceneRenderTargets::Get(RHICmdList);
 
@@ -1322,6 +1327,7 @@ void CopyToComponentTexture(FRHICommandList& RHICmdList, FScene* Scene, FReflect
 void FScene::UpdateReflectionCaptureContents(UReflectionCaptureComponent* CaptureComponent)
 {
 	const bool bCubemapSpecified = CaptureComponent->ReflectionSourceType == EReflectionSourceType::SpecifiedCubemap && CaptureComponent->Cubemap;
+	const int32 ReflectionCaptureSize = UReflectionCaptureComponent::GetReflectionCaptureSize_GameThread();
 
 	if (IsReflectionEnvironmentAvailable(GetFeatureLevel()) || bCubemapSpecified)
 	{
@@ -1352,25 +1358,25 @@ void FScene::UpdateReflectionCaptureContents(UReflectionCaptureComponent* Captur
 
 			ENQUEUE_UNIQUE_RENDER_COMMAND_ONEPARAMETER( 
 				ClearCommand,
-				int32, CubemapSize, CVarReflectionCaptureSize.GetValueOnGameThread(),
+				int32, ReflectionCaptureSize, ReflectionCaptureSize,
 			{
-				ClearScratchCubemaps(RHICmdList, CubemapSize);
+				ClearScratchCubemaps(RHICmdList, ReflectionCaptureSize);
 			});
 			
 			if (CaptureComponent->ReflectionSourceType == EReflectionSourceType::CapturedScene)
 			{
-				CaptureSceneIntoScratchCubemap(this, CaptureComponent->GetComponentLocation() + CaptureComponent->CaptureOffset, CVarReflectionCaptureSize.GetValueOnGameThread(), false, true, 0, false, false, FLinearColor());
+				CaptureSceneIntoScratchCubemap(this, CaptureComponent->GetComponentLocation() + CaptureComponent->CaptureOffset, ReflectionCaptureSize, false, true, 0, false, false, FLinearColor());
 			}
 			else if (CaptureComponent->ReflectionSourceType == EReflectionSourceType::SpecifiedCubemap)
 			{
 				ENQUEUE_UNIQUE_RENDER_COMMAND_FOURPARAMETER(
 					CopyCubemapCommand,
 					UTextureCube*, SourceTexture, CaptureComponent->Cubemap,
-					int32, CubemapSize, CVarReflectionCaptureSize.GetValueOnGameThread(),
+					int32, ReflectionCaptureSize, ReflectionCaptureSize,
 					float, SourceCubemapRotation, CaptureComponent->SourceCubemapAngle * (PI / 180.f),
 					ERHIFeatureLevel::Type, FeatureLevel, GetFeatureLevel(),
 					{
-						CopyCubemapToScratchCubemap(RHICmdList, FeatureLevel, SourceTexture, CubemapSize, false, false, SourceCubemapRotation, FLinearColor());
+						CopyCubemapToScratchCubemap(RHICmdList, FeatureLevel, SourceTexture, ReflectionCaptureSize, false, false, SourceCubemapRotation, FLinearColor());
 					});
 			}
 			else
@@ -1381,10 +1387,10 @@ void FScene::UpdateReflectionCaptureContents(UReflectionCaptureComponent* Captur
 			ENQUEUE_UNIQUE_RENDER_COMMAND_TWOPARAMETER( 
 				FilterCommand,
 				ERHIFeatureLevel::Type, FeatureLevel, GetFeatureLevel(),
-				int32, CubemapSize, CVarReflectionCaptureSize.GetValueOnGameThread(),
+				int32, ReflectionCaptureSize, ReflectionCaptureSize,
 				{
 				bool bNormalize = true;
-				FilterReflectionEnvironment(RHICmdList, FeatureLevel, CubemapSize, NULL, bNormalize);
+				FilterReflectionEnvironment(RHICmdList, FeatureLevel, ReflectionCaptureSize, NULL, bNormalize);
 			});
 
 			// Create a proxy to represent the reflection capture to the rendering thread
diff --git a/Engine/Source/Runtime/Renderer/Private/Renderer.cpp b/Engine/Source/Runtime/Renderer/Private/Renderer.cpp
--- a/Engine/Source/Runtime/Renderer/Private/Renderer.cpp
+++ b/Engine/Source/Runtime/Renderer/Private/Renderer.cpp
@@ -55,8 +55,8 @@ void FRendererModule::DrawTileMesh(FRHICommandListImmediate& RHICmdList, const F
 	{
 		GSystemTextures.InitializeTextures(RHICmdList, FeatureLevel);
 
-		// handle translucent material blend modes
-		if (IsTranslucentBlendMode(MaterialBlendMode))
+		// handle translucent material blend modes, not relevant in MaterialTexCoordScalesAnalysis since it outputs the scales.é
+		if (IsTranslucentBlendMode(MaterialBlendMode) && View.Family->GetDebugViewShaderMode() != DVSM_MaterialTexCoordScalesAnalysis)
 		{
 			if (FeatureLevel >= ERHIFeatureLevel::SM4)
 			{
diff --git a/Engine/Source/Runtime/Renderer/Private/SceneCaptureRendering.cpp b/Engine/Source/Runtime/Renderer/Private/SceneCaptureRendering.cpp
--- a/Engine/Source/Runtime/Renderer/Private/SceneCaptureRendering.cpp
+++ b/Engine/Source/Runtime/Renderer/Private/SceneCaptureRendering.cpp
@@ -14,6 +14,147 @@
 #include "PostProcessing.h"
 #include "SceneUtils.h"
 
+/**
+ * A pixel shader for capturing a component of the rendered scene for a scene capture.
+ */
+template<ESceneCaptureSource CaptureSource>
+class TSceneCapturePS : public FGlobalShader
+{
+	DECLARE_SHADER_TYPE(TSceneCapturePS,Global);
+public:
+
+	static bool ShouldCache(EShaderPlatform Platform) 
+	{ 
+		return IsFeatureLevelSupported(Platform, ERHIFeatureLevel::SM4);
+	}
+
+	static void ModifyCompilationEnvironment(EShaderPlatform Platform, FShaderCompilerEnvironment& OutEnvironment)
+	{
+		switch(CaptureSource)
+		{
+			case SCS_SceneColorHDR: OutEnvironment.SetDefine(TEXT("SOURCE_MODE_SCENE_COLOR_AND_OPACITY"), 1); break;
+			case SCS_SceneColorSceneDepth: OutEnvironment.SetDefine(TEXT("SOURCE_MODE_SCENE_COLOR_SCENE_DEPTH"), 1); break;
+			case SCS_SceneDepth: OutEnvironment.SetDefine(TEXT("SOURCE_MODE_SCENE_DEPTH"), 1); break;
+			case SCS_Normal: OutEnvironment.SetDefine(TEXT("SOURCE_MODE_NORMAL"), 1); break;
+			case SCS_BaseColor: OutEnvironment.SetDefine(TEXT("SOURCE_MODE_BASE_COLOR"), 1); break;
+		}
+	}
+
+	TSceneCapturePS(const ShaderMetaType::CompiledShaderInitializerType& Initializer):
+		FGlobalShader(Initializer)
+	{
+		DeferredParameters.Bind(Initializer.ParameterMap);
+	}
+	TSceneCapturePS() {}
+
+	void SetParameters(FRHICommandList& RHICmdList, const FSceneView& View)
+	{
+		FGlobalShader::SetParameters(RHICmdList, GetPixelShader(), View);
+		DeferredParameters.Set(RHICmdList, GetPixelShader(), View);
+	}
+
+	virtual bool Serialize(FArchive& Ar) override
+	{
+		bool bShaderHasOutdatedParameters = FGlobalShader::Serialize(Ar);
+		Ar << DeferredParameters;
+		return bShaderHasOutdatedParameters;
+	}
+
+private:
+	FDeferredPixelShaderParameters DeferredParameters;
+};
+
+IMPLEMENT_SHADER_TYPE(template<>,TSceneCapturePS<SCS_SceneColorHDR>,TEXT("SceneCapturePixelShader"),TEXT("Main"),SF_Pixel);
+IMPLEMENT_SHADER_TYPE(template<>,TSceneCapturePS<SCS_SceneColorSceneDepth>,TEXT("SceneCapturePixelShader"),TEXT("Main"),SF_Pixel);
+IMPLEMENT_SHADER_TYPE(template<>,TSceneCapturePS<SCS_SceneDepth>,TEXT("SceneCapturePixelShader"),TEXT("Main"),SF_Pixel);
+IMPLEMENT_SHADER_TYPE(template<>,TSceneCapturePS<SCS_Normal>,TEXT("SceneCapturePixelShader"),TEXT("Main"),SF_Pixel);
+IMPLEMENT_SHADER_TYPE(template<>,TSceneCapturePS<SCS_BaseColor>,TEXT("SceneCapturePixelShader"),TEXT("Main"),SF_Pixel);
+
+void FDeferredShadingSceneRenderer::CopySceneCaptureComponentToTarget(FRHICommandListImmediate& RHICmdList)
+{
+	if (ViewFamily.SceneCaptureSource != SCS_FinalColorLDR)
+	{
+		SCOPED_DRAW_EVENT(RHICmdList, CaptureSceneComponent);
+
+		for (int32 ViewIndex = 0; ViewIndex < Views.Num(); ViewIndex++)
+		{
+			FViewInfo& View = Views[ViewIndex];
+			FRHIRenderTargetView ColorView(ViewFamily.RenderTarget->GetRenderTargetTexture(), 0, -1, ERenderTargetLoadAction::ENoAction, ERenderTargetStoreAction::EStore);
+			FRHISetRenderTargetsInfo Info(1, &ColorView, FRHIDepthRenderTargetView());
+			RHICmdList.SetRenderTargetsAndClear(Info);
+
+			RHICmdList.SetRasterizerState(TStaticRasterizerState<FM_Solid, CM_None>::GetRHI());
+			RHICmdList.SetDepthStencilState(TStaticDepthStencilState<false, CF_Always>::GetRHI());
+			RHICmdList.SetBlendState(TStaticBlendState<>::GetRHI());
+
+			TShaderMapRef<FScreenVS> VertexShader(View.ShaderMap);
+
+			if (ViewFamily.SceneCaptureSource == SCS_SceneColorHDR)
+			{
+				TShaderMapRef<TSceneCapturePS<SCS_SceneColorHDR> > PixelShader(View.ShaderMap);
+
+				static FGlobalBoundShaderState BoundShaderState;
+				SetGlobalBoundShaderState(RHICmdList, View.GetFeatureLevel(), BoundShaderState, GFilterVertexDeclaration.VertexDeclarationRHI, *VertexShader, *PixelShader);
+
+				PixelShader->SetParameters(RHICmdList, View);
+			}
+			else if (ViewFamily.SceneCaptureSource == SCS_SceneColorSceneDepth)
+			{
+				TShaderMapRef<TSceneCapturePS<SCS_SceneColorSceneDepth> > PixelShader(View.ShaderMap);
+
+				static FGlobalBoundShaderState BoundShaderState;
+				SetGlobalBoundShaderState(RHICmdList, View.GetFeatureLevel(), BoundShaderState, GFilterVertexDeclaration.VertexDeclarationRHI, *VertexShader, *PixelShader);
+
+				PixelShader->SetParameters(RHICmdList, View);
+			}
+			else if (ViewFamily.SceneCaptureSource == SCS_SceneDepth)
+			{
+				TShaderMapRef<TSceneCapturePS<SCS_SceneDepth> > PixelShader(View.ShaderMap);
+
+				static FGlobalBoundShaderState BoundShaderState;
+				SetGlobalBoundShaderState(RHICmdList, View.GetFeatureLevel(), BoundShaderState, GFilterVertexDeclaration.VertexDeclarationRHI, *VertexShader, *PixelShader);
+
+				PixelShader->SetParameters(RHICmdList, View);
+			}
+			else if (ViewFamily.SceneCaptureSource == SCS_Normal)
+			{
+				TShaderMapRef<TSceneCapturePS<SCS_Normal> > PixelShader(View.ShaderMap);
+
+				static FGlobalBoundShaderState BoundShaderState;
+				SetGlobalBoundShaderState(RHICmdList, View.GetFeatureLevel(), BoundShaderState, GFilterVertexDeclaration.VertexDeclarationRHI, *VertexShader, *PixelShader);
+
+				PixelShader->SetParameters(RHICmdList, View);
+			}
+			else if (ViewFamily.SceneCaptureSource == SCS_BaseColor)
+			{
+				TShaderMapRef<TSceneCapturePS<SCS_BaseColor> > PixelShader(View.ShaderMap);
+
+				static FGlobalBoundShaderState BoundShaderState;
+				SetGlobalBoundShaderState(RHICmdList, View.GetFeatureLevel(), BoundShaderState, GFilterVertexDeclaration.VertexDeclarationRHI, *VertexShader, *PixelShader);
+
+				PixelShader->SetParameters(RHICmdList, View);
+			}
+			else
+			{
+				check(0);
+			}
+
+			VertexShader->SetParameters(RHICmdList, View);
+
+			DrawRectangle(
+				RHICmdList,
+				View.ViewRect.Min.X, View.ViewRect.Min.Y,
+				View.ViewRect.Width(), View.ViewRect.Height(),
+				View.ViewRect.Min.X, View.ViewRect.Min.Y,
+				View.ViewRect.Width(), View.ViewRect.Height(),
+				View.UnconstrainedViewRect.Size(),
+				FSceneRenderTargets::Get(RHICmdList).GetBufferSizeXY(),
+				*VertexShader,
+				EDRF_UseTriangleOptimization);
+		}
+	}
+}
+
 // Copies into render target, optionally flipping it in the Y-axis
 static void CopyCaptureToTarget(FRHICommandListImmediate& RHICmdList, const FRenderTarget* Target, const FIntPoint& TargetSize, FViewInfo& View, const FIntRect& ViewRect, FTextureRHIParamRef SourceTextureRHI, bool bNeedsFlippedRenderTarget)
 {
@@ -30,8 +171,6 @@ static void CopyCaptureToTarget(FRHICommandListImmediate& RHICmdList, const FRen
 	static FGlobalBoundShaderState BoundShaderState;
 	SetGlobalBoundShaderState(RHICmdList, View.GetFeatureLevel(), BoundShaderState, GFilterVertexDeclaration.VertexDeclarationRHI, *VertexShader, *PixelShader);
 
-	FRenderingCompositePassContext Context(RHICmdList, View);
-
 	VertexShader->SetParameters(RHICmdList, View);
 	PixelShader->SetParameters(RHICmdList, TStaticSamplerState<SF_Point>::GetRHI(), SourceTextureRHI);
 
@@ -63,12 +202,20 @@ static void CopyCaptureToTarget(FRHICommandListImmediate& RHICmdList, const FRen
 	}
 }
 
-static void UpdateSceneCaptureContent_RenderThread(FRHICommandListImmediate& RHICmdList, FSceneRenderer* SceneRenderer, FRenderTarget* RenderTarget, FTexture* RenderTargetTexture, const FName OwnerName, const FResolveParams& ResolveParams, bool bUseSceneColorTexture)
+static void UpdateSceneCaptureContent_RenderThread(
+	FRHICommandListImmediate& RHICmdList, 
+	FSceneRenderer* SceneRenderer, 
+	FRenderTarget* RenderTarget, 
+	FTexture* RenderTargetTexture, 
+	const FName OwnerName, 
+	const FResolveParams& ResolveParams, 
+	ESceneCaptureSource CaptureSource)
 {
 	FMemMark MemStackMark(FMemStack::Get());
 
 	// update any resources that needed a deferred update
 	FDeferredUpdateResource::UpdateResources(RHICmdList);
+	bool bUseSceneTextures = CaptureSource != SCS_FinalColorLDR;
 
 	{
 #if WANTS_DRAW_MESH_EVENTS
@@ -82,7 +229,7 @@ static void UpdateSceneCaptureContent_RenderThread(FRHICommandListImmediate& RHI
 		const bool bIsMobileHDR = IsMobileHDR();
 		const bool bRHINeedsFlip = RHINeedsToSwitchVerticalAxis(GMaxRHIShaderPlatform);
 		// note that ES2 will flip the image during post processing. this needs flipping again so it is correct for texture addressing.
-		const bool bNeedsFlippedRenderTarget = (!bIsMobileHDR || !bUseSceneColorTexture) && bRHINeedsFlip;
+		const bool bNeedsFlippedRenderTarget = (!bIsMobileHDR || !bUseSceneTextures) && bRHINeedsFlip;
 
 		// Intermediate render target that will need to be flipped (needed on !IsMobileHDR())
 		TRefCountPtr<IPooledRenderTarget> FlippedPooledRenderTarget;
@@ -148,19 +295,23 @@ static void UpdateSceneCaptureContent_RenderThread(FRHICommandListImmediate& RHI
 			SCOPED_DRAW_EVENT(RHICmdList, FlipCapture);
 			CopyCaptureToTarget(RHICmdList, Target, TargetSize, View, ViewRect, FlippedRenderTarget.GetTextureParamRef(), true);
 		}
-		else if (bUseSceneColorTexture)
+		else if (!FSceneInterface::ShouldUseDeferredRenderer(View.GetFeatureLevel()))
 		{
-			// Copy the captured scene into the destination texture (only required on HDR or deferred as that implies post-processing)
+			// Copy the captured scene into the destination texture
 			SCOPED_DRAW_EVENT(RHICmdList, CaptureSceneColor);
 			CopyCaptureToTarget(RHICmdList, Target, TargetSize, View, ViewRect, FSceneRenderTargets::Get(RHICmdList).GetSceneColorTexture(), false);
 		}
+		else if (bUseSceneTextures)
+		{
+			// Note: When the CaptureSource requires scene textures, the copy will be done in CopySceneCaptureComponentToTarget while the GBuffers are still alive for the frame
+		}
 
 		RHICmdList.CopyToResolveTarget(RenderTarget->GetRenderTargetTexture(), RenderTargetTexture->TextureRHI, false, ResolveParams);
 	}
 	FSceneRenderer::WaitForTasksClearSnapshotsAndDeleteSceneRenderer(RHICmdList, SceneRenderer);
 }
 
-void BuildProjectionMatrix(FIntPoint RenderTargetSize, float FOV, FMatrix& ProjectionMatrix)
+void BuildProjectionMatrix(FIntPoint RenderTargetSize, ECameraProjectionMode::Type ProjectionType, float FOV, float InOrthoWidth, FMatrix& ProjectionMatrix)
 {
 	float XAxisMultiplier;
 	float YAxisMultiplier;
@@ -178,7 +329,28 @@ void BuildProjectionMatrix(FIntPoint RenderTargetSize, float FOV, FMatrix& Proje
 		YAxisMultiplier = 1.0f;
 	}
 
-	if ((bool)ERHIZBuffer::IsInverted)
+	if (ProjectionType == ECameraProjectionMode::Orthographic)
+	{
+		check((int32)ERHIZBuffer::IsInverted != 0);
+		const float OrthoWidth = InOrthoWidth / 2.0f;
+		const float OrthoHeight = InOrthoWidth / 2.0f * YAxisMultiplier;
+
+		const float NearPlane = 0;
+		const float FarPlane = WORLD_MAX / 8.0f;
+
+		const float ZScale = 1.0f / (FarPlane - NearPlane);
+		const float ZOffset = -NearPlane;
+
+		ProjectionMatrix = FReversedZOrthoMatrix(
+			OrthoWidth,
+			OrthoHeight,
+			ZScale,
+			ZOffset
+			);
+	}
+	else
+	{
+	if ((int32)ERHIZBuffer::IsInverted != 0)
 	{
 		ProjectionMatrix = FReversedZPerspectiveMatrix(
 			FOV,
@@ -201,6 +373,7 @@ void BuildProjectionMatrix(FIntPoint RenderTargetSize, float FOV, FMatrix& Proje
 			);
 	}
 }
+}
 
 FSceneRenderer* CreateSceneRendererForSceneCapture(
 	FScene* Scene, 
@@ -211,7 +384,7 @@ FSceneRenderer* CreateSceneRendererForSceneCapture(
 	const FVector& ViewLocation, 
 	const FMatrix& ProjectionMatrix, 
 	float MaxViewDistance, 
-	bool bCaptureSceneColour, 
+	bool bCaptureSceneColor, 
 	bool bIsPlanarReflection,
 	FPostProcessSettings* PostProcessSettings, 
 	float PostProcessBlendWeight)
@@ -220,7 +393,7 @@ FSceneRenderer* CreateSceneRendererForSceneCapture(
 		RenderTarget,
 		Scene,
 		SceneCaptureComponent->ShowFlags)
-		.SetResolveScene(!bCaptureSceneColour)
+		.SetResolveScene(!bCaptureSceneColor)
 		.SetRealtimeUpdate(bIsPlanarReflection));
 
 	FSceneViewInitOptions ViewInitOptions;
@@ -234,7 +407,7 @@ FSceneRenderer* CreateSceneRendererForSceneCapture(
     ViewInitOptions.StereoPass = SceneCaptureComponent->CaptureStereoPass;
 	ViewInitOptions.ProjectionMatrix = ProjectionMatrix;
 
-	if (bCaptureSceneColour)
+	if (bCaptureSceneColor)
 	{
 		ViewFamily.EngineShowFlags.PostProcessing = 0;
 		ViewInitOptions.OverlayColor = FLinearColor::Black;
@@ -326,13 +499,28 @@ void FScene::UpdateSceneCaptureContents(USceneCaptureComponent2D* CaptureCompone
 			FPlane(0, 1, 0, 0),
 			FPlane(0, 0, 0, 1));
 		const float FOV = CaptureComponent->FOVAngle * (float)PI / 360.0f;
-		const bool bUseSceneColorTexture = CaptureComponent->CaptureSource == SCS_SceneColorHDR;
 		FIntPoint CaptureSize(CaptureComponent->TextureTarget->GetSurfaceWidth(), CaptureComponent->TextureTarget->GetSurfaceHeight());
 
 		FMatrix ProjectionMatrix;
-		BuildProjectionMatrix(CaptureSize, FOV, ProjectionMatrix);
-
-		FSceneRenderer* SceneRenderer = CreateSceneRendererForSceneCapture(this, CaptureComponent, CaptureComponent->TextureTarget->GameThread_GetRenderTargetResource(), CaptureSize, ViewRotationMatrix, ViewLocation, ProjectionMatrix, CaptureComponent->MaxViewDistanceOverride, bUseSceneColorTexture, false, &CaptureComponent->PostProcessSettings, CaptureComponent->PostProcessBlendWeight);
+		BuildProjectionMatrix(CaptureSize, CaptureComponent->ProjectionType, FOV, CaptureComponent->OrthoWidth, ProjectionMatrix);
+
+		const bool bUseSceneColorTexture = CaptureComponent->CaptureSource != SCS_FinalColorLDR;
+
+		FSceneRenderer* SceneRenderer = CreateSceneRendererForSceneCapture(
+			this, 
+			CaptureComponent, 
+			CaptureComponent->TextureTarget->GameThread_GetRenderTargetResource(), 
+			CaptureSize, 
+			ViewRotationMatrix, 
+			ViewLocation, 
+			ProjectionMatrix, 
+			CaptureComponent->MaxViewDistanceOverride, 
+			bUseSceneColorTexture, 
+			false, 
+			&CaptureComponent->PostProcessSettings, 
+			CaptureComponent->PostProcessBlendWeight);
+
+		SceneRenderer->ViewFamily.SceneCaptureSource = CaptureComponent->CaptureSource;
 				
 		FTextureRenderTargetResource* TextureRenderTarget = CaptureComponent->TextureTarget->GameThread_GetRenderTargetResource();
 		const FName OwnerName = CaptureComponent->GetOwner() ? CaptureComponent->GetOwner()->GetFName() : NAME_None;
@@ -342,9 +530,9 @@ void FScene::UpdateSceneCaptureContents(USceneCaptureComponent2D* CaptureCompone
 			FSceneRenderer*, SceneRenderer, SceneRenderer,
 			FTextureRenderTargetResource*, TextureRenderTarget, TextureRenderTarget,
 			FName, OwnerName, OwnerName,
-			bool, bUseSceneColorTexture, bUseSceneColorTexture, 
+			ESceneCaptureSource, CaptureSource, CaptureComponent->CaptureSource, 
 		{
-			UpdateSceneCaptureContent_RenderThread(RHICmdList, SceneRenderer, TextureRenderTarget, TextureRenderTarget, OwnerName, FResolveParams(), bUseSceneColorTexture);
+			UpdateSceneCaptureContent_RenderThread(RHICmdList, SceneRenderer, TextureRenderTarget, TextureRenderTarget, OwnerName, FResolveParams(), CaptureSource);
 		});
 	}
 }
@@ -405,7 +593,7 @@ void FScene::UpdateSceneCaptureContents(USceneCaptureComponentCube* CaptureCompo
 			const FMatrix ViewRotationMatrix = FLocal::CalcCubeFaceTransform(TargetFace);
 			FIntPoint CaptureSize(CaptureComponent->TextureTarget->GetSurfaceWidth(), CaptureComponent->TextureTarget->GetSurfaceHeight());
 			FMatrix ProjectionMatrix;
-			BuildProjectionMatrix(CaptureSize, FOV, ProjectionMatrix);
+			BuildProjectionMatrix(CaptureSize, ECameraProjectionMode::Perspective, FOV, 1.0f, ProjectionMatrix);
 
 			FSceneRenderer* SceneRenderer = CreateSceneRendererForSceneCapture(this, CaptureComponent, CaptureComponent->TextureTarget->GameThread_GetRenderTargetResource(), CaptureSize, ViewRotationMatrix, Location, ProjectionMatrix, CaptureComponent->MaxViewDistanceOverride, true, false, NULL, 0);
 
@@ -419,7 +607,7 @@ void FScene::UpdateSceneCaptureContents(USceneCaptureComponentCube* CaptureCompo
 				FName, OwnerName, OwnerName,
 				ECubeFace, TargetFace, TargetFace,
 			{
-				UpdateSceneCaptureContent_RenderThread(RHICmdList, SceneRenderer, TextureRenderTarget, TextureRenderTarget, OwnerName, FResolveParams(FResolveRect(), TargetFace), true);
+				UpdateSceneCaptureContent_RenderThread(RHICmdList, SceneRenderer, TextureRenderTarget, TextureRenderTarget, OwnerName, FResolveParams(FResolveRect(), TargetFace), SCS_SceneColorHDR);
 			});
 		}
 	}
diff --git a/Engine/Source/Runtime/Renderer/Private/SceneCore.cpp b/Engine/Source/Runtime/Renderer/Private/SceneCore.cpp
--- a/Engine/Source/Runtime/Renderer/Private/SceneCore.cpp
+++ b/Engine/Source/Runtime/Renderer/Private/SceneCore.cpp
@@ -321,10 +321,24 @@ void FStaticMesh::AddToDrawLists(FRHICommandListImmediate& RHICmdList, FScene* S
 
 			extern int32 GEarlyZPassMovable;
 
+			EDepthDrawingMode EarlyZPassMode = (EDepthDrawingMode)EarlyZPass;
+			bool bEarlyZPassMovable = GEarlyZPassMovable != 0;
+
+			static IConsoleVariable* CDBufferVar = IConsoleManager::Get().FindConsoleVariable(TEXT("r.DBuffer"));
+			bool bDBufferAllowed = CDBufferVar ? CDBufferVar->GetInt() != 0 : false;
+
+			if (bDBufferAllowed)
+			{
+				// DBuffer decals force a full prepass
+				EarlyZPassMode = DDM_AllOccluders;
+				bEarlyZPassMovable = true;
+			}
+
 			// WARNING : If you change this condition, also change the logic in FStaticMeshSceneProxy::DrawStaticElements.
+			// Warning: also mirrored in FDeferredShadingSceneRenderer::FDeferredShadingSceneRenderer
 			if (PrimitiveSceneInfo->Proxy->ShouldUseAsOccluder() 
-				&& (!IsMasked(FeatureLevel) || EarlyZPass == 2)
-				&& (!PrimitiveSceneInfo->Proxy->IsMovable() || GEarlyZPassMovable))
+				&& (!IsMasked(FeatureLevel) || EarlyZPassMode == DDM_AllOccluders)
+				&& (!PrimitiveSceneInfo->Proxy->IsMovable() || bEarlyZPassMovable))
 			{
 				FDepthDrawingPolicyFactory::AddStaticMesh(Scene,this);
 			}
diff --git a/Engine/Source/Runtime/Renderer/Private/SceneRendering.cpp b/Engine/Source/Runtime/Renderer/Private/SceneRendering.cpp
--- a/Engine/Source/Runtime/Renderer/Private/SceneRendering.cpp
+++ b/Engine/Source/Runtime/Renderer/Private/SceneRendering.cpp
@@ -141,8 +141,6 @@ static TAutoConsoleVariable<float> CVarTessellationAdaptivePixelsPerTriangle(
 	TEXT("Global tessellation factor multiplier"),
 	ECVF_RenderThreadSafe);
 
-extern ENGINE_API TAutoConsoleVariable<int32> CVarReflectionCaptureSize;
-
 static TAutoConsoleVariable<int32> CVarSupportSimpleForwardShading(
 	TEXT("r.SupportSimpleForwardShading"),
 	0,
@@ -438,6 +436,7 @@ void FViewInfo::Init()
 	ExponentialFogParameters = FVector4(0,1,1,0);
 	ExponentialFogColor = FVector::ZeroVector;
 	FogMaxOpacity = 1;
+	ExponentialFogParameters3 = FVector2D(0, 0);
 
 	bUseDirectionalInscattering = false;
 	DirectionalInscatteringExponent = 0;
@@ -1084,7 +1083,12 @@ void FViewInfo::CreateUniformBuffer(
 		FrameUniformShaderParameters.HMDEyePaddingOffset = 1.0f;
 	}
 
-	FrameUniformShaderParameters.ReflectionCubemapMaxMip = FMath::FloorLog2(CVarReflectionCaptureSize.GetValueOnRenderThread()) - 1;
+	FrameUniformShaderParameters.ReflectionCubemapMaxMip = FMath::FloorLog2(UReflectionCaptureComponent::GetReflectionCaptureSize_RenderThread()) - 1;
+
+	FrameUniformShaderParameters.ShowDecalsMask = Family->EngineShowFlags.Decals ? 1.0f : 0.0f;
+
+	extern int32 GDistanceFieldAOSpecularOcclusionMode;
+	FrameUniformShaderParameters.DistanceFieldAOSpecularOcclusionMode = GDistanceFieldAOSpecularOcclusionMode;
 
 	OutViewUniformBuffer = TUniformBufferRef<FViewUniformShaderParameters>::CreateUniformBufferImmediate(ViewUniformShaderParameters, UniformBuffer_SingleFrame);
 	OutFrameUniformBuffer = TUniformBufferRef<FFrameUniformShaderParameters>::CreateUniformBufferImmediate(FrameUniformShaderParameters, UniformBuffer_SingleFrame);
@@ -2067,9 +2071,18 @@ void OnChangeSimpleForwardShading(IConsoleVariable* Var)
 	FGlobalComponentRecreateRenderStateContext Context;
 }
 
+void OnChangeCVarRequiringRecreateRenderState(IConsoleVariable* Var)
+{
+	// Propgate cvar change to static draw lists
+	FGlobalComponentRecreateRenderStateContext Context;
+}
+
 FRendererModule::FRendererModule()
 {
 	CVarSimpleForwardShading.AsVariable()->SetOnChangedCallback(FConsoleVariableDelegate::CreateStatic(&OnChangeSimpleForwardShading));
+
+	static auto CVarEarlyZPass = IConsoleManager::Get().FindConsoleVariable(TEXT("r.EarlyZPass"));
+	CVarEarlyZPass->SetOnChangedCallback(FConsoleVariableDelegate::CreateStatic(&OnChangeCVarRequiringRecreateRenderState));
 }
 
 void FRendererModule::CreateAndInitSingleView(FRHICommandListImmediate& RHICmdList, class FSceneViewFamily* ViewFamily, const struct FSceneViewInitOptions* ViewInitOptions)
diff --git a/Engine/Source/Runtime/Renderer/Private/SceneRendering.h b/Engine/Source/Runtime/Renderer/Private/SceneRendering.h
--- a/Engine/Source/Runtime/Renderer/Private/SceneRendering.h
+++ b/Engine/Source/Runtime/Renderer/Private/SceneRendering.h
@@ -138,6 +138,43 @@
 class FTranslucentPrimSet
 {
 public:
+	/** contains a scene prim and its sort key */
+	struct FSortedPrim
+	{
+		/** Default constructor. */
+		FSortedPrim() {}
+
+		// @param InPass (first we sort by this)
+		// @param InSortPriority SHRT_MIN .. SHRT_MAX (then we sort by this)
+		// @param InSortKey from UPrimitiveComponent::TranslucencySortPriority e.g. SortByDistance/SortAlongAxis (then by this)
+		FSortedPrim(FPrimitiveSceneInfo* InPrimitiveSceneInfo, ETranslucencyPass::Type InPass, int16 InSortPriority, float InSortKey)
+			:	PrimitiveSceneInfo(InPrimitiveSceneInfo)
+			,	SortKey(InSortKey)
+		{
+			SetSortOrder(InPass, InSortPriority);
+		}
+
+		void SetSortOrder(ETranslucencyPass::Type InPass, int16 InSortPriority)
+		{
+			uint32 UpperShort = (uint32)InPass;
+			// 0 .. 0xffff
+			int32 SortPriorityWithoutSign = (int32)InSortPriority - (int32)SHRT_MIN;
+			uint32 LowerShort = SortPriorityWithoutSign;
+
+			check(LowerShort <= 0xffff);
+
+			// top 8 bits are currently unused
+			SortOrder = (UpperShort << 16) | LowerShort;
+		}
+
+		//
+		FPrimitiveSceneInfo* PrimitiveSceneInfo;
+		// single 32bit sort order containing Pass and SortPriority (first we sort by this)
+		uint32 SortOrder;
+		// from UPrimitiveComponent::TranslucencySortPriority (then by this)
+		float SortKey;
+	};
+
 	/** 
 	* Iterate over the sorted list of prims and draw them
 	* @param View - current view used to draw items
@@ -172,16 +209,11 @@
 	void DrawPrimitivesForForwardShading(FRHICommandListImmediate& RHICmdList, const class FViewInfo& View, const bool bRenderSeparateTranslucency) const;
 
 	/**
-	* Add a new primitive to the list of sorted prims
-	* @param PrimitiveSceneInfo - primitive info to add. Origin of bounds is used for sort.
-	* @param ViewInfo - used to transform bounds to view space
+	* Insert a primitive to the translucency rendering list[s]
 	*/
-	void AddScenePrimitive(FPrimitiveSceneInfo* PrimitiveSceneInfo, const FViewInfo& ViewInfo, bool bUseNormalTranslucency, bool bUseSeparateTranslucency, bool bUseMobileSeparateTranslucency);
+	
+	static void PlaceScenePrimitive(FPrimitiveSceneInfo* PrimitiveSceneInfo, const FViewInfo& ViewInfo, bool bUseNormalTranslucency, bool bUseSeparateTranslucency, bool bUseMobileSeparateTranslucency,
+		FSortedPrim* InArrayStart, int32& InOutArrayNum, FTranslucenyPrimCount& OutCount);
 
 	/**
-	* Similar to AddScenePrimitive, but we are doing placement news and increasing counts when that happens
-	*/
-	static void PlaceScenePrimitive(FPrimitiveSceneInfo* PrimitiveSceneInfo, const FViewInfo& ViewInfo, bool bUseNormalTranslucency, bool bUseSeparateTranslucency, bool bUseMobileSeparateTranslucency, void *InPlace, int32& InOutNum, FTranslucenyPrimCount& OutCount);
-
-	/**
 	* Sort any primitives that were added to the set back-to-front
@@ -196,43 +228,6 @@
 		return SortedPrims.Num();
 	}
 
-	/** contains a scene prim and its sort key */
-	struct FSortedPrim
-	{
-		/** Default constructor. */
-		FSortedPrim() {}
-
-		// @param InPass (first we sort by this)
-		// @param InSortPriority SHRT_MIN .. SHRT_MAX (then we sort by this)
-		// @param InSortKey from UPrimitiveComponent::TranslucencySortPriority e.g. SortByDistance/SortAlongAxis (then by this)
-		FSortedPrim(FPrimitiveSceneInfo* InPrimitiveSceneInfo, ETranslucencyPass::Type InPass, int16 InSortPriority, float InSortKey)
-			:	PrimitiveSceneInfo(InPrimitiveSceneInfo)
-			,	SortKey(InSortKey)
-		{
-			SetSortOrder(InPass, InSortPriority);
-		}
-
-		void SetSortOrder(ETranslucencyPass::Type InPass, int16 InSortPriority)
-		{
-			uint32 UpperShort = (uint32)InPass;
-			// 0 .. 0xffff
-			int32 SortPriorityWithoutSign = (int32)InSortPriority - (int32)SHRT_MIN;
-			uint32 LowerShort = SortPriorityWithoutSign;
-
-			check(LowerShort <= 0xffff);
-
-			// top 8 bits are currently unused
-			SortOrder = (UpperShort << 16) | LowerShort;
-		}
-
-		//
-		FPrimitiveSceneInfo* PrimitiveSceneInfo;
-		// single 32bit sort order containing Pass and SortPriority (first we sort by this)
-		uint32 SortOrder;
-		// from UPrimitiveComponent::TranslucencySortPriority (then by this)
-		float SortKey;
-	};
-
 	/**
 	* Adds primitives originally created with PlaceScenePrimitive
 	*/
@@ -551,6 +546,7 @@
 	FVector4 ExponentialFogParameters;
 	FVector ExponentialFogColor;
 	float FogMaxOpacity;
+	FVector2D ExponentialFogParameters3;
 
 	/** Parameters for directional inscattering of exponential height fog. */
 	bool bUseDirectionalInscattering;
diff --git a/Engine/Source/Runtime/Renderer/Private/SceneVisibility.cpp b/Engine/Source/Runtime/Renderer/Private/SceneVisibility.cpp
--- a/Engine/Source/Runtime/Renderer/Private/SceneVisibility.cpp
+++ b/Engine/Source/Runtime/Renderer/Private/SceneVisibility.cpp
@@ -1284,28 +1284,31 @@ static int32 OcclusionCull(FRHICommandListImmediate& RHICmdList, const FScene* S
 	return NumOccludedPrimitives;
 }
 
-template<class T>
+template<class T, int TAmplifyFactor = 1>
 struct FRelevancePrimSet
 {
 	enum
 	{
-		MaxPrims = 127 //like 128, but we leave space for NumPrims
+		MaxInputPrims = 127, //like 128, but we leave space for NumPrims
+		MaxOutputPrims = MaxInputPrims * TAmplifyFactor
 	};
 	int32 NumPrims;
-	T Prims[MaxPrims];
+
+	T Prims[MaxOutputPrims];
+
 	FORCEINLINE FRelevancePrimSet()
 		: NumPrims(0)
 	{
-		//FMemory::Memzero(Prims, sizeof(T) * MaxPrims);
+		//FMemory::Memzero(Prims, sizeof(T) * GetMaxOutputPrim());
 	}
 	FORCEINLINE void AddPrim(T Prim)
 	{
-		checkSlow(NumPrims < MaxPrims);
+		checkSlow(NumPrims < MaxOutputPrims);
 		Prims[NumPrims++] = Prim;
 	}
 	FORCEINLINE bool IsFull() const
 	{
-		return NumPrims >= MaxPrims;
+		return NumPrims >= MaxOutputPrims;
 	}
 	template<class TARRAY>
 	FORCEINLINE void AppendTo(TARRAY& DestArray)
@@ -1376,7 +1379,7 @@ struct FRelevancePacket
 	FRelevancePrimSet<int32> RelevantStaticPrimitives;
 	FRelevancePrimSet<int32> NotDrawRelevant;
 	FRelevancePrimSet<FPrimitiveSceneInfo*> VisibleDynamicPrimitives;
-	FRelevancePrimSet<FTranslucentPrimSet::FSortedPrim> TranslucencyPrims;
+	FRelevancePrimSet<FTranslucentPrimSet::FSortedPrim, ETranslucencyPass::TPT_MAX> TranslucencyPrims;
 	// belongs to TranslucencyPrims
 	FTranslucenyPrimCount TranslucencyPrimCount;
 	FRelevancePrimSet<FPrimitiveSceneProxy*> DistortionPrimSet;
@@ -1475,7 +1478,7 @@ struct FRelevancePacket
 				// Add to set of dynamic translucent primitives
 				FTranslucentPrimSet::PlaceScenePrimitive(PrimitiveSceneInfo, View, 
 					ViewRelevance.bNormalTranslucencyRelevance, ViewRelevance.bSeparateTranslucencyRelevance, ViewRelevance.bMobileSeparateTranslucencyRelevance, 
-					&TranslucencyPrims.Prims[TranslucencyPrims.NumPrims], TranslucencyPrims.NumPrims, TranslucencyPrimCount);
+					&TranslucencyPrims.Prims[0], TranslucencyPrims.NumPrims, TranslucencyPrimCount);
 
 				if (ViewRelevance.bDistortionRelevance)
 				{
@@ -1527,6 +1530,7 @@ struct FRelevancePacket
 			{
 				// Update the PrimitiveComponent's LastRenderTime.
 				*(PrimitiveSceneInfo->ComponentLastRenderTime) = CurrentWorldTime;
+				*(PrimitiveSceneInfo->ComponentLastRenderTimeOnScreen) = CurrentWorldTime;
 			}
 
 			uint32 bHasClearCoat = ViewRelevance.ShadingModelMaskRelevance & (1 << MSM_ClearCoat);
@@ -1733,7 +1737,7 @@ static void ComputeAndMarkRelevanceForViewParallel(
 	uint8* RESTRICT MarkMasks = (uint8*)FMemStack::Get().Alloc(NumMesh + 31 , 8); // some padding to simplify the high speed transpose
 	FMemory::Memzero(MarkMasks, NumMesh + 31);
 
-	int32 EstimateOfNumPackets = NumMesh / (FRelevancePrimSet<int32>::MaxPrims * 4);
+	int32 EstimateOfNumPackets = NumMesh / (FRelevancePrimSet<int32>::MaxInputPrims * 4);
 
 	TArray<FRelevancePacket*,SceneRenderingAllocator> Packets;
 
@@ -2144,21 +2148,23 @@ void FSceneRenderer::PreVisibilityFrameSetup(FRHICommandListImmediate& RHICmdLis
 					SampleX = SamplesX[ Index ];
 					SampleY = SamplesY[ Index ];
 				}
-				else if( TemporalAASamples == 8 )
-				{
-					// This works better than various orderings of 8xMSAA.
-					ViewState->SetupTemporalAA(8, ViewFamily);
-					uint32 Index = ViewState->GetCurrentTemporalAASampleIndex();
-					SampleX = Halton( Index, 2 ) - 0.5f;
-					SampleY = Halton( Index, 3 ) - 0.5f;
-				}
 				else
 				{
+					static auto CVar = IConsoleManager::Get().FindConsoleVariable(TEXT("r.TemporalAASharpness"));
+					float Scale = ( 2.0f - CVar->GetFloat() ) * 0.3f;
+
 					// More than 8 samples can improve quality.
 					ViewState->SetupTemporalAA(TemporalAASamples, ViewFamily);
 					uint32 Index = ViewState->GetCurrentTemporalAASampleIndex();
-					SampleX = Halton( Index, 2 ) - 0.5f;
-					SampleY = Halton( Index, 3 ) - 0.5f;
+
+					float u1 = Halton( Index + 1, 2 );
+					float u2 = Halton( Index + 1, 3 );
+
+					// Gaussian sample
+					float phi = 2.0f * PI * u2;
+					float r = Scale * FMath::Sqrt( -2.0f * FMath::Loge( FMath::Max( u1, 1e-6f ) ) );
+					SampleX = r * FMath::Cos( phi );
+					SampleY = r * FMath::Sin( phi );
 				}
 
 				View.TemporalJitterPixelsX = SampleX;
@@ -2191,6 +2197,9 @@ void FSceneRenderer::PreVisibilityFrameSetup(FRHICommandListImmediate& RHICmdLis
 		{
 			// no TemporalAA
 			ViewState->SetupTemporalAA(1, ViewFamily);
+
+			ViewState->TemporalAAHistoryRT.SafeRelease();
+			ViewState->PendingTemporalAAHistoryRT.SafeRelease();
 		}
 
 		if ( ViewState )
@@ -2851,9 +2860,7 @@ void FDeferredShadingSceneRenderer::InitViewsPossiblyAfterPrepass(FRHICommandLis
 		FTaskGraphInterface::Get().WaitUntilTasksComplete(SortEvents, ENamedThreads::RenderThread);
 	}
 
-	bool bDynamicShadows = ViewFamily.EngineShowFlags.DynamicShadows && GetShadowQuality() > 0;
-
-	if (bDynamicShadows && !IsSimpleForwardShadingEnabled(GetFeatureLevelShaderPlatform(FeatureLevel)))
+	if (ViewFamily.EngineShowFlags.DynamicShadows && !IsSimpleForwardShadingEnabled(GetFeatureLevelShaderPlatform(FeatureLevel)))
 	{
 		// Setup dynamic shadows.
 		InitDynamicShadows(RHICmdList);
diff --git a/Engine/Source/Runtime/Renderer/Private/ShadowRendering.cpp b/Engine/Source/Runtime/Renderer/Private/ShadowRendering.cpp
--- a/Engine/Source/Runtime/Renderer/Private/ShadowRendering.cpp
+++ b/Engine/Source/Runtime/Renderer/Private/ShadowRendering.cpp
@@ -4041,8 +4041,7 @@ void FForwardShadingSceneRenderer::RenderShadowDepthMaps(FRHICommandListImmediat
 {
 	SCOPED_DRAW_EVENT(RHICmdList, Lights);
 
-	bool bDynamicShadows = ViewFamily.EngineShowFlags.DynamicShadows && GetShadowQuality() > 0;
-	if (IsSimpleForwardShadingEnabled(GetFeatureLevelShaderPlatform(FeatureLevel)) || !bDynamicShadows)
+	if (IsSimpleForwardShadingEnabled(GetFeatureLevelShaderPlatform(FeatureLevel)) || !ViewFamily.EngineShowFlags.DynamicShadows)
 	{
 		return;
 	}	
@@ -4143,9 +4142,7 @@ void FForwardShadingSceneRenderer::RenderModulatedShadowDepthMaps(FRHICommandLis
 {
 	SCOPED_DRAW_EVENT(RHICmdList, Lights);
 	
-	bool bDynamicShadows = ViewFamily.EngineShowFlags.DynamicShadows && GetShadowQuality() > 0;
-	
-	if (IsSimpleForwardShadingEnabled(GetFeatureLevelShaderPlatform(FeatureLevel)) || !bDynamicShadows)
+	if (IsSimpleForwardShadingEnabled(GetFeatureLevelShaderPlatform(FeatureLevel)) || !ViewFamily.EngineShowFlags.DynamicShadows)
 	{
 		bModulatedShadowsInUse = false;
 	}
@@ -4247,9 +4244,7 @@ void FForwardShadingSceneRenderer::RenderAllocatedModulatedShadowProjections(FRH
 {
 	SCOPED_DRAW_EVENT(RHICmdList, Lights);
 
-	bool bDynamicShadows = ViewFamily.EngineShowFlags.DynamicShadows && GetShadowQuality() > 0;
-
-	if (IsSimpleForwardShadingEnabled(GetFeatureLevelShaderPlatform(FeatureLevel)) || !bDynamicShadows)
+	if (IsSimpleForwardShadingEnabled(GetFeatureLevelShaderPlatform(FeatureLevel)) || !ViewFamily.EngineShowFlags.DynamicShadows)
 	{
 		return;
 	}
