Copying //UE4/Dev-Rendering to //UE4/Dev-Main (Source: //UE4/Dev-Rendering @ 2998063)

diff --git a/Engine/Source/Runtime/Renderer/Private/TranslucentRendering.cpp b/Engine/Source/Runtime/Renderer/Private/TranslucentRendering.cpp
--- a/Engine/Source/Runtime/Renderer/Private/TranslucentRendering.cpp
+++ b/Engine/Source/Runtime/Renderer/Private/TranslucentRendering.cpp
@@ -370,7 +370,7 @@ public:
 	bool bBackFace;
 	FMeshDrawingRenderState DrawRenderState;
 	bool bUseTranslucentSelfShadowing;
-	float SeparateTranslucencyScreenTextureScaleFactor;
+	float DownsampleFactorFromSceneBufferSize;
 
 	/** Initialization constructor. */
 	FDrawTranslucentMeshAction(
@@ -380,7 +380,7 @@ public:
 		FHitProxyId InHitProxyId,
 		const FProjectedShadowInfo* InTranslucentSelfShadow,
 		bool bInUseTranslucentSelfShadowing,
-		float ScreenTextureUVScaleFactor
+		float InDownsampleFactorFromSceneBufferSize
 		) :
 		View(InView),
 		TranslucentSelfShadow(InTranslucentSelfShadow),
@@ -388,7 +388,7 @@ public:
 		bBackFace(bInBackFace),
 		DrawRenderState(InDrawRenderState),
 		bUseTranslucentSelfShadowing(bInUseTranslucentSelfShadowing),
-		SeparateTranslucencyScreenTextureScaleFactor(ScreenTextureUVScaleFactor)
+		DownsampleFactorFromSceneBufferSize(InDownsampleFactorFromSceneBufferSize)
 	{}
 
 	bool UseTranslucentSelfShadowing() const 
@@ -442,7 +442,7 @@ public:
 			false,
 			false);
 		RHICmdList.BuildAndSetLocalBoundShaderState(DrawingPolicy.GetBoundShaderStateInput(View.GetFeatureLevel()));
-		DrawingPolicy.SetSharedState(RHICmdList, &View, typename TBasePassDrawingPolicy<LightMapPolicyType>::ContextDataType(), SeparateTranslucencyScreenTextureScaleFactor);
+		DrawingPolicy.SetSharedState(RHICmdList, &View, typename TBasePassDrawingPolicy<LightMapPolicyType>::ContextDataType(), DownsampleFactorFromSceneBufferSize);
 
 		int32 BatchElementIndex = 0;
 		uint64 BatchElementMask = Parameters.BatchElementMask;
@@ -613,7 +613,7 @@ bool FTranslucencyDrawingPolicyFactory::DrawMesh(
 					HitProxyId,
 					DrawingContext.TranslucentSelfShadow,
 					PrimitiveSceneProxy && PrimitiveSceneProxy->CastsVolumetricTranslucentShadow(),
-					bCurrentlyRenderingSeparateTranslucency ? OutScale : 1.0f
+					bCurrentlyRenderingSeparateTranslucency ? 1.0f / OutScale : 1.0f
 				)
 			);
 
@@ -903,36 +903,14 @@ inline float CalculateTranslucentSortKey(FPrimitiveSceneInfo* PrimitiveSceneInfo
 	return SortKey;
 }
 
-void FTranslucentPrimSet::AddScenePrimitive(FPrimitiveSceneInfo* PrimitiveSceneInfo, const FViewInfo& ViewInfo, bool bUseNormalTranslucency, bool bUseSeparateTranslucency, bool bUseMobileSeparateTranslucency)
-{
-	const float SortKey = CalculateTranslucentSortKey(PrimitiveSceneInfo, ViewInfo);
-
-	const auto FeatureLevel = ViewInfo.GetFeatureLevel();
-
-	const bool bIsSeparateTranslucency = (bUseSeparateTranslucency && FeatureLevel >= ERHIFeatureLevel::SM4) || (bUseMobileSeparateTranslucency && FeatureLevel < ERHIFeatureLevel::SM4);
-	// Force separate translucency to be rendered normally if the feature level does not support separate translucency
-	const bool bIsNormalTranslucency = bUseNormalTranslucency || (bUseSeparateTranslucency && !bUseMobileSeparateTranslucency && FeatureLevel < ERHIFeatureLevel::SM4);
-
-	// can be optimized
-	const bool bIsTranslucency = bIsSeparateTranslucency || bIsNormalTranslucency;
-
-	if (bIsTranslucency)
-	{
-		ETranslucencyPass::Type Pass = bIsSeparateTranslucency ? ETranslucencyPass::TPT_SeparateTransluceny : ETranslucencyPass::TPT_NonSeparateTransluceny;
-
-		// add to list of translucent prims that use scene color
-		new(SortedPrims) FSortedPrim(PrimitiveSceneInfo, Pass, PrimitiveSceneInfo->Proxy->GetTranslucencySortPriority(), SortKey);
-		SortedPrimsNum.Add(Pass);
-	}
-}
-
 void FTranslucentPrimSet::AppendScenePrimitives(FSortedPrim* Elements, int32 Num, const FTranslucenyPrimCount& TranslucentPrimitiveCountPerPass)
 {
 	SortedPrims.Append(Elements, Num);
 	SortedPrimsNum.Append(TranslucentPrimitiveCountPerPass);
 }
 
-void FTranslucentPrimSet::PlaceScenePrimitive(FPrimitiveSceneInfo* PrimitiveSceneInfo, const FViewInfo& ViewInfo, bool bUseNormalTranslucency, bool bUseSeparateTranslucency, bool bUseMobileSeparateTranslucency, void *InPlace, int32& InOutNum, FTranslucenyPrimCount& OutCount)
+void FTranslucentPrimSet::PlaceScenePrimitive(FPrimitiveSceneInfo* PrimitiveSceneInfo, const FViewInfo& ViewInfo, bool bUseNormalTranslucency, bool bUseSeparateTranslucency, bool bUseMobileSeparateTranslucency,
+	FTranslucentPrimSet::FSortedPrim *InArrayStart, int32& InOutArrayNum, FTranslucenyPrimCount& OutCount)
 {
 	const float SortKey = CalculateTranslucentSortKey(PrimitiveSceneInfo, ViewInfo);
 	const auto FeatureLevel = ViewInfo.GetFeatureLevel();
@@ -944,20 +922,21 @@ void FTranslucentPrimSet::PlaceScenePrimitive(FPrimitiveSceneInfo* PrimitiveScen
 		&& !ViewInfo.Family->EngineShowFlags.ShaderComplexity
 		&& ViewInfo.Family->EngineShowFlags.SeparateTranslucency;
 
-	bool bLocalSeparateTranslucency = bUseSeparateTranslucency && bCanBeSeparate;
-	bool bLocalNonSeparateTranslucency = bUseNormalTranslucency || !bCanBeSeparate;
+	bool bIsSeparateTranslucency = bUseSeparateTranslucency && bCanBeSeparate;
+	bool bIsNonSeparateTranslucency = bUseNormalTranslucency || !bCanBeSeparate;
 
-	// can be optimized
-	bool bLocalTranslucency = bLocalSeparateTranslucency || bLocalNonSeparateTranslucency;
-
-	// add to list of translucency prims, can be separate or non separate translucency
-	if (bLocalTranslucency)
+	if (bIsSeparateTranslucency)
 	{
-		ETranslucencyPass::Type TranslucencyPass = bLocalSeparateTranslucency ? ETranslucencyPass::TPT_SeparateTransluceny : ETranslucencyPass::TPT_NonSeparateTransluceny;
+		ETranslucencyPass::Type TranslucencyPass = ETranslucencyPass::TPT_SeparateTransluceny;
 
-		new (InPlace) FSortedPrim(PrimitiveSceneInfo, TranslucencyPass, PrimitiveSceneInfo->Proxy->GetTranslucencySortPriority(), SortKey);
-		InOutNum++;
+		new(&InArrayStart[InOutArrayNum++]) FSortedPrim(PrimitiveSceneInfo, TranslucencyPass, PrimitiveSceneInfo->Proxy->GetTranslucencySortPriority(), SortKey);
+		OutCount.Add(TranslucencyPass);
+	}
+	if (bIsNonSeparateTranslucency)
+	{
+		ETranslucencyPass::Type TranslucencyPass = ETranslucencyPass::TPT_NonSeparateTransluceny;
 
+		new(&InArrayStart[InOutArrayNum++]) FSortedPrim(PrimitiveSceneInfo, TranslucencyPass, PrimitiveSceneInfo->Proxy->GetTranslucencySortPriority(), SortKey);
 		OutCount.Add(TranslucencyPass);
 	}
 }
diff --git a/Engine/Source/Runtime/Renderer/Private/VelocityRendering.cpp b/Engine/Source/Runtime/Renderer/Private/VelocityRendering.cpp
--- a/Engine/Source/Runtime/Renderer/Private/VelocityRendering.cpp
+++ b/Engine/Source/Runtime/Renderer/Private/VelocityRendering.cpp
@@ -538,7 +538,7 @@ int32 GetMotionBlurQualityFromCVar()
 
 bool IsMotionBlurEnabled(const FViewInfo& View)
 {
-	if (!(View.GetFeatureLevel() >= ERHIFeatureLevel::SM4))
+	if (View.GetFeatureLevel() < ERHIFeatureLevel::SM5)
 	{
 		return false; 
 	}
diff --git a/Engine/Source/Runtime/Renderer/Public/PrimitiveSceneInfo.h b/Engine/Source/Runtime/Renderer/Public/PrimitiveSceneInfo.h
--- a/Engine/Source/Runtime/Renderer/Public/PrimitiveSceneInfo.h
+++ b/Engine/Source/Runtime/Renderer/Public/PrimitiveSceneInfo.h
@@ -184,6 +184,9 @@ public:
 	 */
 	float* ComponentLastRenderTime;
 
+	/** Same as ComponentLastRenderTime but only updated if the component is on screen. Used by the texture streamer. */
+	float* ComponentLastRenderTimeOnScreen;
+
 	/** 
 	 * The root attachment component id for use with lighting, if valid.
 	 * If the root id is not valid, this is a parent primitive.
diff --git a/Engine/Source/Runtime/Windows/D3D11RHI/Private/D3D11RenderTarget.cpp b/Engine/Source/Runtime/Windows/D3D11RHI/Private/D3D11RenderTarget.cpp
--- a/Engine/Source/Runtime/Windows/D3D11RHI/Private/D3D11RenderTarget.cpp
+++ b/Engine/Source/Runtime/Windows/D3D11RHI/Private/D3D11RenderTarget.cpp
@@ -337,7 +337,23 @@ void FD3D11DynamicRHI::RHICopyToResolveTarget(FTextureRHIParamRef SourceTextureR
 			}
 			else
 			{
-				Direct3DDeviceIMContext->CopySubresourceRegion(DestTextureCube->GetResource(),DestSubresource,0,0,0,SourceTextureCube->GetResource(),SourceSubresource,NULL);
+				if (ResolveParams.Rect.IsValid())
+				{
+					D3D11_BOX SrcBox;
+
+					SrcBox.left = ResolveParams.Rect.X1;
+					SrcBox.top = ResolveParams.Rect.Y1;
+					SrcBox.front = 0;
+					SrcBox.right = ResolveParams.Rect.X2;
+					SrcBox.bottom = ResolveParams.Rect.Y2;
+					SrcBox.back = 1;
+
+					Direct3DDeviceIMContext->CopySubresourceRegion(DestTextureCube->GetResource(), DestSubresource, 0, 0, 0, SourceTextureCube->GetResource(), SourceSubresource, &SrcBox);
+				}
+				else
+				{
+					Direct3DDeviceIMContext->CopySubresourceRegion(DestTextureCube->GetResource(), DestSubresource, 0, 0, 0, SourceTextureCube->GetResource(), SourceSubresource, NULL);
+				}
 			}
 		}
 	}
diff --git a/Engine/Source/Runtime/Windows/D3D11RHI/Private/Windows/WindowsD3D11Device.cpp b/Engine/Source/Runtime/Windows/D3D11RHI/Private/Windows/WindowsD3D11Device.cpp
--- a/Engine/Source/Runtime/Windows/D3D11RHI/Private/Windows/WindowsD3D11Device.cpp
+++ b/Engine/Source/Runtime/Windows/D3D11RHI/Private/Windows/WindowsD3D11Device.cpp
@@ -45,6 +45,20 @@ static TAutoConsoleVariable<int32> CVarForceNvidiaToSM4(
 	TEXT("Forces Nvidia devices to use SM4.0/D3D10.0 feature level."),
 	ECVF_RenderThreadSafe);
 
+static TAutoConsoleVariable<int32> CVarAMDUseMultiThreadedDevice(
+	TEXT("r.AMDD3D11MultiThreadedDevice"),
+	1,
+	TEXT("If true, creates a multithreaded D3D11 device on AMD hardware (workaround for driver bug)\n")
+	TEXT("Changes will only take effect in new game/editor instances - can't be changed at runtime.\n"),
+	ECVF_Default);
+
+static TAutoConsoleVariable<int32> CVarAMDDisableAsyncTextureCreation(
+	TEXT("r.AMDDisableAsyncTextureCreation"),
+	1,
+	TEXT("If true, uses synchronous texture creation on AMD hardware (workaround for driver bug)\n")
+	TEXT("Changes will only take effect in new game/editor instances - can't be changed at runtime.\n"),
+	ECVF_Default);
+
 /**
  * Console variables used by the D3D11 RHI device.
  */
@@ -631,7 +645,7 @@ void FD3D11DynamicRHI::InitD3DDevice()
 
 		D3D_FEATURE_LEVEL ActualFeatureLevel = (D3D_FEATURE_LEVEL)0;
 
-		if(IsRHIDeviceAMD())
+		if (IsRHIDeviceAMD() && CVarAMDUseMultiThreadedDevice.GetValueOnAnyThread())
 		{
 			DeviceFlags &= ~D3D11_CREATE_DEVICE_SINGLETHREADED;
 		}
@@ -671,7 +685,7 @@ void FD3D11DynamicRHI::InitD3DDevice()
 		GShaderPlatformForFeatureLevel[ERHIFeatureLevel::SM4] = SP_PCD3D_SM4;
 		GShaderPlatformForFeatureLevel[ERHIFeatureLevel::SM5] = SP_PCD3D_SM5;
 
-		if(IsRHIDeviceAMD())
+		if (IsRHIDeviceAMD() && CVarAMDDisableAsyncTextureCreation.GetValueOnAnyThread())
 		{
 			GRHISupportsAsyncTextureCreation = false;
 		}
