[Integrate] Fixed planar reflection clip plane that was broken in the merge from Dev-VR

diff --git a/Engine/Source/Runtime/Renderer/Private/SceneRendering.cpp b/Engine/Source/Runtime/Renderer/Private/SceneRendering.cpp
--- a/Engine/Source/Runtime/Renderer/Private/SceneRendering.cpp
+++ b/Engine/Source/Runtime/Renderer/Private/SceneRendering.cpp
@@ -641,6 +641,7 @@ void FViewInfo::CreateUniformBuffer(
 	// can be optimized
 	ViewUniformShaderParameters.PrevInvViewProj = PrevViewProjMatrix.Inverse();
 	ViewUniformShaderParameters.ScreenPositionScaleBias = ScreenPositionScaleBias;
+	ViewUniformShaderParameters.GlobalClippingPlane = FVector4(GlobalClippingPlane.X, GlobalClippingPlane.Y, GlobalClippingPlane.Z, -GlobalClippingPlane.W);
 
 	ViewUniformShaderParameters.FieldOfViewWideAngles = 2.f * ViewMatrices.GetHalfFieldOfViewPerAxis();
 	ViewUniformShaderParameters.PrevFieldOfViewWideAngles = 2.f * PrevViewMatrices.GetHalfFieldOfViewPerAxis();
