Fix broken skylight rendering on translucent materials in projects with static lighting disabled.
Made IsTranslucentBlendMode inline in the header.

diff --git a/Engine/Source/Runtime/Engine/Private/Materials/MaterialShared.cpp b/Engine/Source/Runtime/Engine/Private/Materials/MaterialShared.cpp
--- a/Engine/Source/Runtime/Engine/Private/Materials/MaterialShared.cpp
+++ b/Engine/Source/Runtime/Engine/Private/Materials/MaterialShared.cpp
@@ -811,11 +811,6 @@ void FMaterial::ReleaseShaderMap()
 	}
 }
 
-bool IsTranslucentBlendMode(EBlendMode BlendMode)
-{
-	return BlendMode != BLEND_Opaque && BlendMode != BLEND_Masked;
-}
-
 int32 FMaterialResource::GetMaterialDomain() const { return Material->MaterialDomain; }
 bool FMaterialResource::IsTangentSpaceNormal() const { return Material->bTangentSpaceNormal || (!Material->Normal.IsConnected() && !Material->bUseMaterialAttributes); }
 bool FMaterialResource::ShouldInjectEmissiveIntoLPV() const { return Material->bUseEmissiveForDynamicAreaLighting; }
diff --git a/Engine/Source/Runtime/Engine/Public/MaterialShared.h b/Engine/Source/Runtime/Engine/Public/MaterialShared.h
--- a/Engine/Source/Runtime/Engine/Public/MaterialShared.h
+++ b/Engine/Source/Runtime/Engine/Public/MaterialShared.h
@@ -1561,7 +1561,10 @@ public:
 /**
  * @return True if BlendMode is translucent (should be part of the translucent rendering).
  */
-extern ENGINE_API bool IsTranslucentBlendMode(enum EBlendMode BlendMode);
+inline bool IsTranslucentBlendMode(enum EBlendMode BlendMode)
+{
+	return BlendMode != BLEND_Opaque && BlendMode != BLEND_Masked;
+}
 
 /**
  * Implementation of the FMaterial interface for a UMaterial or UMaterialInstance.
diff --git a/Engine/Source/Runtime/Renderer/Private/BasePassRendering.cpp b/Engine/Source/Runtime/Renderer/Private/BasePassRendering.cpp
--- a/Engine/Source/Runtime/Renderer/Private/BasePassRendering.cpp
+++ b/Engine/Source/Runtime/Renderer/Private/BasePassRendering.cpp
@@ -349,7 +349,7 @@ public:
 			TStaticMeshDrawList<TBasePassDrawingPolicy<LightMapPolicyType> >& DrawList =
 				Scene->GetBasePassDrawList<LightMapPolicyType>(DrawType);
 
-			const bool bRenderSkylight = Scene->ShouldRenderSkylight() && Parameters.ShadingModel != MSM_Unlit;
+			const bool bRenderSkylight = Scene->ShouldRenderSkylight(Parameters.BlendMode) && Parameters.ShadingModel != MSM_Unlit;
 			const bool bRenderAtmosphericFog = IsTranslucentBlendMode(Parameters.BlendMode) && Scene->HasAtmosphericFog() && Scene->ReadOnlyCVARCache.bEnableAtmosphericFog;
 
 			// Add the static mesh to the draw list.
@@ -466,7 +466,7 @@ public:
 #endif
 		const FScene* Scene = Parameters.PrimitiveSceneProxy ? Parameters.PrimitiveSceneProxy->GetPrimitiveSceneInfo()->Scene : NULL;
 
-		const bool bRenderSkylight = Scene && Scene->ShouldRenderSkylight() && Parameters.ShadingModel != MSM_Unlit;
+		const bool bRenderSkylight = Scene && Scene->ShouldRenderSkylight(Parameters.BlendMode) && Parameters.ShadingModel != MSM_Unlit;
 		const bool bRenderAtmosphericFog = IsTranslucentBlendMode(Parameters.BlendMode) && (Scene && Scene->HasAtmosphericFog() && Scene->ReadOnlyCVARCache.bEnableAtmosphericFog) && View.Family->EngineShowFlags.AtmosphericFog;
 
 		TBasePassDrawingPolicy<LightMapPolicyType> DrawingPolicy(
diff --git a/Engine/Source/Runtime/Renderer/Private/MobileBasePassRendering.cpp b/Engine/Source/Runtime/Renderer/Private/MobileBasePassRendering.cpp
--- a/Engine/Source/Runtime/Renderer/Private/MobileBasePassRendering.cpp
+++ b/Engine/Source/Runtime/Renderer/Private/MobileBasePassRendering.cpp
@@ -178,7 +178,7 @@ public:
 				LightMapPolicy,
 				Parameters.BlendMode,
 				Parameters.TextureMode,
-				Parameters.ShadingModel != MSM_Unlit && Scene->ShouldRenderSkylight(),
+				Parameters.ShadingModel != MSM_Unlit && Scene->ShouldRenderSkylight(Parameters.BlendMode),
 				DVSM_None,
 				FeatureLevel,
 				Parameters.bEditorCompositeDepthTest,
@@ -283,7 +283,7 @@ public:
 			LightMapPolicy,
 			Parameters.BlendMode,
 			Parameters.TextureMode,
-			Parameters.ShadingModel != MSM_Unlit && Scene && Scene->ShouldRenderSkylight(),
+			Parameters.ShadingModel != MSM_Unlit && Scene && Scene->ShouldRenderSkylight(Parameters.BlendMode),
 			View.Family->GetDebugViewShaderMode(),
 			View.GetFeatureLevel(),
 			Parameters.bEditorCompositeDepthTest,
diff --git a/Engine/Source/Runtime/Renderer/Private/MobileTranslucentRendering.cpp b/Engine/Source/Runtime/Renderer/Private/MobileTranslucentRendering.cpp
--- a/Engine/Source/Runtime/Renderer/Private/MobileTranslucentRendering.cpp
+++ b/Engine/Source/Runtime/Renderer/Private/MobileTranslucentRendering.cpp
@@ -142,7 +142,7 @@ public:
 			LightMapPolicy,
 			Parameters.BlendMode,
 			Parameters.TextureMode,
-			Parameters.ShadingModel != MSM_Unlit && Scene && Scene->ShouldRenderSkylight(),
+			Parameters.ShadingModel != MSM_Unlit && Scene && Scene->ShouldRenderSkylight(Parameters.BlendMode),
 			View.Family->GetDebugViewShaderMode(),
 			View.GetFeatureLevel()
 			);
diff --git a/Engine/Source/Runtime/Renderer/Private/ScenePrivate.h b/Engine/Source/Runtime/Renderer/Private/ScenePrivate.h
--- a/Engine/Source/Runtime/Renderer/Private/ScenePrivate.h
+++ b/Engine/Source/Runtime/Renderer/Private/ScenePrivate.h
@@ -2108,19 +2108,26 @@ public:
 
 	virtual ERHIFeatureLevel::Type GetFeatureLevel() const override { return FeatureLevel; }
 
-	bool ShouldRenderSkylight() const
+	bool ShouldRenderSkylight(EBlendMode BlendMode) const
 	{
-		return ShouldRenderSkylight_Internal() && ReadOnlyCVARCache.bEnableStationarySkylight;
+		return ShouldRenderSkylight_Internal(BlendMode) && ReadOnlyCVARCache.bEnableStationarySkylight;
 	}
 
-	bool ShouldRenderSkylight_Internal() const
+	bool ShouldRenderSkylight_Internal(EBlendMode BlendMode) const
 	{
-		const bool bRenderSkylight = SkyLight
-			&& !SkyLight->bHasStaticLighting
-			// The deferred shading renderer does movable skylight diffuse in a later deferred pass, not in the base pass
-			&& (SkyLight->bWantsStaticShadowing || IsSimpleForwardShadingEnabled(GetShaderPlatform()));
+		if (IsTranslucentBlendMode(BlendMode))
+		{
+			return SkyLight && !SkyLight->bHasStaticLighting;
+		}
+		else
+		{
+			const bool bRenderSkylight = SkyLight
+				&& !SkyLight->bHasStaticLighting
+				// The deferred shading renderer does movable skylight diffuse in a later deferred pass, not in the base pass
+				&& (SkyLight->bWantsStaticShadowing || IsSimpleForwardShadingEnabled(GetShaderPlatform()));
 
-		return bRenderSkylight;
+			return bRenderSkylight;
+		}
 	}
 
 	virtual TArray<FPrimitiveComponentId> GetScenePrimitiveComponentIds() const override
diff --git a/Engine/Source/Runtime/Renderer/Private/SceneRendering.cpp b/Engine/Source/Runtime/Renderer/Private/SceneRendering.cpp
--- a/Engine/Source/Runtime/Renderer/Private/SceneRendering.cpp
+++ b/Engine/Source/Runtime/Renderer/Private/SceneRendering.cpp
@@ -1428,7 +1428,7 @@ void FSceneRenderer::RenderFinish(FRHICommandListImmediate& RHICmdList)
 		bool bShowDFAODisabledWarning = !GDistanceFieldAO && (ViewFamily.EngineShowFlags.VisualizeMeshDistanceFields || ViewFamily.EngineShowFlags.VisualizeDistanceFieldAO);
 
 		const bool bShowAtmosphericFogWarning = Scene->AtmosphericFog != nullptr && !Scene->ReadOnlyCVARCache.bEnableAtmosphericFog;
-		const bool bShowSkylightWarning = Scene->ShouldRenderSkylight_Internal() && !Scene->ReadOnlyCVARCache.bEnableStationarySkylight;
+		const bool bShowSkylightWarning = (Scene->ShouldRenderSkylight_Internal(BLEND_Opaque) || Scene->ShouldRenderSkylight_Internal(BLEND_Translucent)) && !Scene->ReadOnlyCVARCache.bEnableStationarySkylight;
 		const bool bShowPointLightWarning = UsedWholeScenePointLightNames.Num() > 0 && !Scene->ReadOnlyCVARCache.bEnablePointLightShadows;
 
 		const bool bAnyWarning = bShowPrecomputedVisibilityWarning || bShowGlobalClipPlaneWarning || bShowAtmosphericFogWarning || bShowSkylightWarning || bShowPointLightWarning || bShowDFAODisabledWarning;
diff --git a/Engine/Source/Runtime/Renderer/Private/TranslucentRendering.cpp b/Engine/Source/Runtime/Renderer/Private/TranslucentRendering.cpp
--- a/Engine/Source/Runtime/Renderer/Private/TranslucentRendering.cpp
+++ b/Engine/Source/Runtime/Renderer/Private/TranslucentRendering.cpp
@@ -363,7 +363,7 @@ public:
 
 		const FScene* Scene = Parameters.PrimitiveSceneProxy ? Parameters.PrimitiveSceneProxy->GetPrimitiveSceneInfo()->Scene : NULL;
 
-		const bool bRenderSkylight = Scene && Scene->ShouldRenderSkylight() && bIsLitMaterial;
+		const bool bRenderSkylight = Scene && Scene->ShouldRenderSkylight(Parameters.BlendMode) && bIsLitMaterial;
 		const bool bRenderAtmosphericFog =(Scene && Scene->HasAtmosphericFog() && Scene->ReadOnlyCVARCache.bEnableAtmosphericFog) && View.Family->EngineShowFlags.AtmosphericFog && View.Family->EngineShowFlags.Fog;
 
 		TBasePassDrawingPolicy<LightMapPolicyType> DrawingPolicy(
