Copying //UE4/Dev-Rendering to //UE4/Dev-Main (Source: //UE4/Dev-Rendering @ 3045398)

diff --git a/Engine/Source/Runtime/Renderer/Private/LightGridInjection.cpp b/Engine/Source/Runtime/Renderer/Private/LightGridInjection.cpp
new file mode 100644
--- /dev/null
+++ b/Engine/Source/Runtime/Renderer/Private/LightGridInjection.cpp
@@ -0,0 +1,576 @@
+// Copyright 1998-2015 Epic Games, Inc. All Rights Reserved.
+
+/*=============================================================================
+	LightGridInjection.cpp
+=============================================================================*/
+
+#include "RendererPrivate.h"
+#include "ScenePrivate.h"
+#include "SceneUtils.h"
+
+int32 GLightGridPixelSize = 64;
+FAutoConsoleVariableRef CVarLightGridPixelSize(
+	TEXT("r.Forward.LightGridPixelSize"),
+	GLightGridPixelSize,
+	TEXT("Size of a cell in the light grid, in pixels."),
+	ECVF_Scalability | ECVF_RenderThreadSafe
+	);
+
+int32 GLightGridSizeZ = 32;
+FAutoConsoleVariableRef CVarLightGridSizeZ(
+	TEXT("r.Forward.LightGridSizeZ"),
+	GLightGridSizeZ,
+	TEXT("Number of Z slices in the light grid."),
+	ECVF_Scalability | ECVF_RenderThreadSafe
+	);
+
+int32 GMaxCulledLightsPerCell = 32;
+FAutoConsoleVariableRef CVarMaxCulledLightsPerCell(
+	TEXT("r.Forward.MaxCulledLightsPerCell"),
+	GMaxCulledLightsPerCell,
+	TEXT("Controls how much memory is allocated for each cell for light culling.  When r.Forward.LightLinkedListCulling is enabled, this is used to compute a global max instead of a per-cell limit on culled lights."),
+	ECVF_Scalability | ECVF_RenderThreadSafe
+	);
+
+int32 GLightLinkedListCulling = 1;
+FAutoConsoleVariableRef CVarLightLinkedListCulling(
+	TEXT("r.Forward.LightLinkedListCulling"),
+	GLightLinkedListCulling,
+	TEXT("Uses a reverse linked list to store culled lights, removing the fixed limit on how many lights can affect a cell - it becomes a global limit instead."),
+	ECVF_Scalability | ECVF_RenderThreadSafe
+	);
+
+IMPLEMENT_UNIFORM_BUFFER_STRUCT(FForwardGlobalLightData,TEXT("ForwardGlobalLightData"));
+
+FForwardGlobalLightData::FForwardGlobalLightData()
+{
+	NumLocalLights = 0;
+	HasDirectionalLight = 0;
+}
+
+int32 NumCulledLightsGridStride = 2;
+int32 LightLinkStride = 2;
+// 65k indexable light limit
+typedef uint16 FLightIndexType;
+
+/**  */
+class FForwardCullingParameters
+{
+public:
+
+	static void ModifyCompilationEnvironment(EShaderPlatform Platform, FShaderCompilerEnvironment& OutEnvironment)
+	{
+		OutEnvironment.SetDefine(TEXT("LIGHT_LINK_STRIDE"), LightLinkStride);
+	}
+
+	void Bind(const FShaderParameterMap& ParameterMap)
+	{
+		NextCulledLightLink.Bind(ParameterMap, TEXT("NextCulledLightLink"));
+		StartOffsetGrid.Bind(ParameterMap, TEXT("StartOffsetGrid"));
+		CulledLightLinks.Bind(ParameterMap, TEXT("CulledLightLinks"));
+		NextCulledLightData.Bind(ParameterMap, TEXT("NextCulledLightData"));
+	}
+
+	template<typename ShaderRHIParamRef>
+	void Set(FRHICommandList& RHICmdList, const ShaderRHIParamRef& ShaderRHI, FShader* Shader, const FViewInfo& View)
+	{
+		NextCulledLightLink.SetBuffer(RHICmdList, ShaderRHI, View.ForwardLightingResources->NextCulledLightLink);
+		StartOffsetGrid.SetBuffer(RHICmdList, ShaderRHI, View.ForwardLightingResources->StartOffsetGrid);
+		CulledLightLinks.SetBuffer(RHICmdList, ShaderRHI, View.ForwardLightingResources->CulledLightLinks);
+		NextCulledLightData.SetBuffer(RHICmdList, ShaderRHI, View.ForwardLightingResources->NextCulledLightData);
+	}
+
+	template<typename ShaderRHIParamRef>
+	void UnsetParameters(FRHICommandList& RHICmdList, const ShaderRHIParamRef& ShaderRHI, const FViewInfo& View)
+	{
+		NextCulledLightLink.UnsetUAV(RHICmdList, ShaderRHI);
+		StartOffsetGrid.UnsetUAV(RHICmdList, ShaderRHI);
+		CulledLightLinks.UnsetUAV(RHICmdList, ShaderRHI);
+		NextCulledLightData.UnsetUAV(RHICmdList, ShaderRHI);
+
+		TArray<FUnorderedAccessViewRHIParamRef, TInlineAllocator<4>> OutUAVs;
+
+		if (NextCulledLightLink.IsBound())
+		{
+			OutUAVs.Add(View.ForwardLightingResources->NextCulledLightLink.UAV);
+		}
+
+		if (StartOffsetGrid.IsBound())
+		{
+			OutUAVs.Add(View.ForwardLightingResources->StartOffsetGrid.UAV);
+		}
+
+		if (CulledLightLinks.IsBound())
+		{
+			OutUAVs.Add(View.ForwardLightingResources->CulledLightLinks.UAV);
+		}
+
+		if (NextCulledLightData.IsBound())
+		{
+			OutUAVs.Add(View.ForwardLightingResources->NextCulledLightData.UAV);
+		}
+
+		if (OutUAVs.Num() > 0)
+		{
+			RHICmdList.TransitionResources(EResourceTransitionAccess::EReadable, EResourceTransitionPipeline::EComputeToGfx, OutUAVs.GetData(), OutUAVs.Num());
+		}
+	}
+
+	/** Serializer. */
+	friend FArchive& operator<<(FArchive& Ar,FForwardCullingParameters& P)
+	{
+		Ar << P.NextCulledLightLink;
+		Ar << P.StartOffsetGrid;
+		Ar << P.CulledLightLinks;
+		Ar << P.NextCulledLightData;
+		return Ar;
+	}
+
+private:
+
+	FRWShaderParameter NextCulledLightLink;
+	FRWShaderParameter StartOffsetGrid;
+	FRWShaderParameter CulledLightLinks;
+	FRWShaderParameter NextCulledLightData;
+};
+
+
+uint32 LightGridInjectionGroupSize = 4;
+
+template<bool bLightLinkedListCulling>
+class TLightGridInjectionCS : public FGlobalShader
+{
+	DECLARE_SHADER_TYPE(TLightGridInjectionCS,Global)
+public:
+
+	static bool ShouldCache(EShaderPlatform Platform)
+	{
+		return IsFeatureLevelSupported(Platform, ERHIFeatureLevel::SM5);
+	}
+
+	static void ModifyCompilationEnvironment(EShaderPlatform Platform, FShaderCompilerEnvironment& OutEnvironment)
+	{
+		FGlobalShader::ModifyCompilationEnvironment(Platform,OutEnvironment);
+		OutEnvironment.SetDefine(TEXT("THREADGROUP_SIZE"), LightGridInjectionGroupSize);
+		FForwardLightingParameters::ModifyCompilationEnvironment(Platform, OutEnvironment);
+		FForwardCullingParameters::ModifyCompilationEnvironment(Platform, OutEnvironment);
+		OutEnvironment.SetDefine(TEXT("USE_LINKED_CULL_LIST"), bLightLinkedListCulling);
+	}
+
+	TLightGridInjectionCS(const ShaderMetaType::CompiledShaderInitializerType& Initializer)
+		: FGlobalShader(Initializer)
+	{
+		ForwardLightingParameters.Bind(Initializer.ParameterMap);
+		ForwardCullingParameters.Bind(Initializer.ParameterMap);
+	}
+
+	TLightGridInjectionCS()
+	{
+	}
+
+	void SetParameters(FRHICommandList& RHICmdList, const FViewInfo& View)
+	{
+		FComputeShaderRHIParamRef ShaderRHI = GetComputeShader();
+		FGlobalShader::SetParameters(RHICmdList, ShaderRHI, View);
+		ForwardLightingParameters.Set(RHICmdList, ShaderRHI, this, View);
+		ForwardCullingParameters.Set(RHICmdList, ShaderRHI, this, View);
+	}
+
+	void UnsetParameters(FRHICommandList& RHICmdList, const FViewInfo& View)
+	{
+		ForwardLightingParameters.UnsetParameters(RHICmdList, GetComputeShader(), View);
+		ForwardCullingParameters.UnsetParameters(RHICmdList, GetComputeShader(), View);
+	}
+
+	virtual bool Serialize(FArchive& Ar)
+	{		
+		bool bShaderHasOutdatedParameters = FGlobalShader::Serialize(Ar);
+		Ar << ForwardLightingParameters;
+		Ar << ForwardCullingParameters;
+		return bShaderHasOutdatedParameters;
+	}
+
+private:
+
+	FForwardLightingParameters ForwardLightingParameters;
+	FForwardCullingParameters ForwardCullingParameters;
+};
+
+IMPLEMENT_SHADER_TYPE(template<>,TLightGridInjectionCS<true>,TEXT("LightGridInjection"),TEXT("LightGridInjectionCS"),SF_Compute);
+IMPLEMENT_SHADER_TYPE(template<>,TLightGridInjectionCS<false>,TEXT("LightGridInjection"),TEXT("LightGridInjectionCS"),SF_Compute);
+
+class FLightGridCompactCS : public FGlobalShader
+{
+	DECLARE_SHADER_TYPE(FLightGridCompactCS,Global)
+public:
+
+	static bool ShouldCache(EShaderPlatform Platform)
+	{
+		return IsFeatureLevelSupported(Platform, ERHIFeatureLevel::SM5);
+	}
+
+	static void ModifyCompilationEnvironment(EShaderPlatform Platform, FShaderCompilerEnvironment& OutEnvironment)
+	{
+		FGlobalShader::ModifyCompilationEnvironment(Platform,OutEnvironment);
+		OutEnvironment.SetDefine(TEXT("THREADGROUP_SIZE"), LightGridInjectionGroupSize);
+		FForwardLightingParameters::ModifyCompilationEnvironment(Platform, OutEnvironment);
+		FForwardCullingParameters::ModifyCompilationEnvironment(Platform, OutEnvironment);
+	}
+
+	FLightGridCompactCS(const ShaderMetaType::CompiledShaderInitializerType& Initializer)
+		: FGlobalShader(Initializer)
+	{
+		ForwardLightingParameters.Bind(Initializer.ParameterMap);
+		ForwardCullingParameters.Bind(Initializer.ParameterMap);
+	}
+
+	FLightGridCompactCS()
+	{
+	}
+
+	void SetParameters(FRHICommandList& RHICmdList, const FViewInfo& View)
+	{
+		FComputeShaderRHIParamRef ShaderRHI = GetComputeShader();
+		FGlobalShader::SetParameters(RHICmdList, ShaderRHI, View);
+		ForwardLightingParameters.Set(RHICmdList, ShaderRHI, this, View);
+		ForwardCullingParameters.Set(RHICmdList, ShaderRHI, this, View);
+	}
+
+	void UnsetParameters(FRHICommandList& RHICmdList, const FViewInfo& View)
+	{
+		ForwardLightingParameters.UnsetParameters(RHICmdList, GetComputeShader(), View);
+		ForwardCullingParameters.UnsetParameters(RHICmdList, GetComputeShader(), View);
+	}
+
+	virtual bool Serialize(FArchive& Ar)
+	{		
+		bool bShaderHasOutdatedParameters = FGlobalShader::Serialize(Ar);
+		Ar << ForwardLightingParameters;
+		Ar << ForwardCullingParameters;
+		return bShaderHasOutdatedParameters;
+	}
+
+private:
+
+	FForwardLightingParameters ForwardLightingParameters;
+	FForwardCullingParameters ForwardCullingParameters;
+};
+
+IMPLEMENT_SHADER_TYPE(,FLightGridCompactCS,TEXT("LightGridInjection"),TEXT("LightGridCompactCS"),SF_Compute);
+
+FVector GetLightGridZParams(float NearPlane, float FarPlane)
+{
+	// S = distribution scale
+	// B, O are solved for given the z distances of the first+last slice, and the # of slices.
+	//
+	// slice = log2(z*B + O) * S
+
+	// Don't spent lots of resolution right in front of the near plane
+	double NearOffset = .095 * 100;
+	// Space out the slices so they aren't all clustered at the near plane
+	double S = 4.05;
+
+	double N = NearPlane + NearOffset;
+	double F = FarPlane;
+
+	double O = (F - N * exp2((GLightGridSizeZ - 1) / S)) / (F - N);
+	double B = (1 - O) / N;
+
+	return FVector(B, O, S);
+}
+
+void FDeferredShadingSceneRenderer::ComputeLightGrid(FRHICommandListImmediate& RHICmdList)
+{
+	bool bAnyViewUsesTranslucentSurfaceLighting = false;
+
+	for (int32 ViewIndex = 0; ViewIndex < Views.Num(); ViewIndex++)
+	{
+		const FViewInfo& View = Views[ViewIndex];
+		bAnyViewUsesTranslucentSurfaceLighting |= View.bTranslucentSurfaceLighting;
+	}
+
+	if (IsForwardShadingEnabled(FeatureLevel) || (bAnyViewUsesTranslucentSurfaceLighting && FeatureLevel >= ERHIFeatureLevel::SM5))
+	{
+		QUICK_SCOPE_CYCLE_COUNTER(STAT_ComputeLightGrid);
+		SCOPED_DRAW_EVENT(RHICmdList, ComputeLightGrid);
+
+		static const auto AllowStaticLightingVar = IConsoleManager::Get().FindTConsoleVariableDataInt(TEXT("r.AllowStaticLighting"));
+		const bool bAllowStaticLighting = (!AllowStaticLightingVar || AllowStaticLightingVar->GetValueOnRenderThread() != 0);
+
+		FSimpleLightArray SimpleLights;
+		GatherSimpleLights(ViewFamily, Views, SimpleLights);
+
+		for (int32 ViewIndex = 0; ViewIndex < Views.Num(); ViewIndex++)
+		{
+			FViewInfo& View = Views[ViewIndex];
+
+			FForwardGlobalLightData GlobalLightData;
+			TArray<FForwardLocalLightData, SceneRenderingAllocator> ForwardLocalLightData;
+			ForwardLocalLightData.Empty(Scene->Lights.Num());
+
+			float FurthestLight = 1000;
+
+			for (TSparseArray<FLightSceneInfoCompact>::TConstIterator LightIt(Scene->Lights); LightIt; ++LightIt)
+			{
+				const FLightSceneInfoCompact& LightSceneInfoCompact = *LightIt;
+				const FLightSceneInfo* const LightSceneInfo = LightSceneInfoCompact.LightSceneInfo;
+
+				if (LightSceneInfo->ShouldRenderLightViewIndependent()
+					&& LightSceneInfo->ShouldRenderLight(View)
+					// Reflection override skips direct specular because it tends to be blindingly bright with a perfectly smooth surface
+					&& !ViewFamily.EngineShowFlags.ReflectionOverride)
+				{
+					FVector4 LightPositionAndInvRadius;
+					FVector4 LightColorAndFalloffExponent;
+					FVector NormalizedLightDirection;
+					FVector2D SpotAngles;
+					float SourceRadius;
+					float SourceLength;
+					float MinRoughness;
+
+					// Get the light parameters
+					LightSceneInfo->Proxy->GetParameters(
+						LightPositionAndInvRadius,
+						LightColorAndFalloffExponent,
+						NormalizedLightDirection,
+						SpotAngles,
+						SourceRadius,
+						SourceLength,
+						MinRoughness);
+
+					if (LightSceneInfo->Proxy->IsInverseSquared())
+					{
+						// Correction for lumen units
+						LightColorAndFalloffExponent.X *= 16.0f;
+						LightColorAndFalloffExponent.Y *= 16.0f;
+						LightColorAndFalloffExponent.Z *= 16.0f;
+						LightColorAndFalloffExponent.W = 0;
+					}
+
+					// When rendering reflection captures, the direct lighting of the light is actually the indirect specular from the main view
+					if (View.bIsReflectionCapture)
+					{
+						LightColorAndFalloffExponent.X *= LightSceneInfo->Proxy->GetIndirectLightingScale();
+						LightColorAndFalloffExponent.Y *= LightSceneInfo->Proxy->GetIndirectLightingScale();
+						LightColorAndFalloffExponent.Z *= LightSceneInfo->Proxy->GetIndirectLightingScale();
+					}
+
+					int32 ShadowMapChannel = LightSceneInfo->Proxy->GetShadowMapChannel();
+					int32 PreviewShadowMapChannel = LightSceneInfo->Proxy->GetPreviewShadowMapChannel();
+
+					if (!bAllowStaticLighting)
+					{
+						ShadowMapChannel = INDEX_NONE;
+						PreviewShadowMapChannel = INDEX_NONE;
+					}
+
+					// Static shadowing uses ShadowMapChannel, dynamic shadows are packed into light attenuation using PreviewShadowMapChannel
+					const uint32 ShadowMapChannelMaskPacked =
+						(ShadowMapChannel == 0 ? 1 : 0) |
+						(ShadowMapChannel == 1 ? 2 : 0) |
+						(ShadowMapChannel == 2 ? 4 : 0) |
+						(ShadowMapChannel == 3 ? 8 : 0) |
+						(PreviewShadowMapChannel == 0 ? 16 : 0) |
+						(PreviewShadowMapChannel == 1 ? 32 : 0) |
+						(PreviewShadowMapChannel == 2 ? 64 : 0) |
+						(PreviewShadowMapChannel == 3 ? 128 : 0);
+
+					if (LightSceneInfoCompact.LightType == LightType_Point || LightSceneInfoCompact.LightType == LightType_Spot)
+					{
+						ForwardLocalLightData.AddUninitialized(1);
+						FForwardLocalLightData& LightData = ForwardLocalLightData.Last();
+
+						LightData.LightPositionAndInvRadius = LightPositionAndInvRadius;
+						LightData.LightColorAndFalloffExponent = LightColorAndFalloffExponent;
+						LightData.LightDirectionAndShadowMapChannelMask = FVector4(NormalizedLightDirection, *((float*)&ShadowMapChannelMaskPacked));
+						LightData.SpotAnglesAndSourceRadius = FVector4(SpotAngles.X, SpotAngles.Y, SourceRadius, SourceLength);
+
+						const FSphere BoundingSphere = LightSceneInfo->Proxy->GetBoundingSphere();
+						const float Distance = View.ViewMatrices.ViewMatrix.TransformPosition(BoundingSphere.Center).Z + BoundingSphere.W;
+						FurthestLight = FMath::Max(FurthestLight, Distance);
+					}
+					else if (LightSceneInfoCompact.LightType == LightType_Directional)
+					{
+						GlobalLightData.HasDirectionalLight = 1;
+						GlobalLightData.DirectionalLightColor = LightColorAndFalloffExponent;
+						GlobalLightData.DirectionalLightDirection = NormalizedLightDirection;
+						GlobalLightData.DirectionalLightShadowMapChannelMask = ShadowMapChannelMaskPacked;
+
+						const FVector2D FadeParams = LightSceneInfo->Proxy->GetDirectionalLightDistanceFadeParameters(View.GetFeatureLevel(), LightSceneInfo->IsPrecomputedLightingValid());
+
+						GlobalLightData.DirectionalLightDistanceFadeMAD = FVector2D(FadeParams.Y, -FadeParams.X * FadeParams.Y);
+					}
+				}
+			}
+
+			for (int32 SimpleLightIndex = 0; SimpleLightIndex < SimpleLights.InstanceData.Num(); SimpleLightIndex++)
+			{	
+				ForwardLocalLightData.AddUninitialized(1);
+				FForwardLocalLightData& LightData = ForwardLocalLightData.Last();
+
+				const FSimpleLightEntry& SimpleLight = SimpleLights.InstanceData[SimpleLightIndex];
+				const FSimpleLightPerViewEntry& SimpleLightPerViewData = SimpleLights.GetViewDependentData(SimpleLightIndex, ViewIndex, Views.Num());
+				LightData.LightPositionAndInvRadius = FVector4(SimpleLightPerViewData.Position, 1.0f / FMath::Max(SimpleLight.Radius, KINDA_SMALL_NUMBER));
+				LightData.LightColorAndFalloffExponent = FVector4(SimpleLight.Color, SimpleLight.Exponent);
+
+				uint32 ShadowMapChannelMask = 0;
+				LightData.LightDirectionAndShadowMapChannelMask = FVector4(FVector(1, 0, 0), *((float*)&ShadowMapChannelMask));
+				LightData.SpotAnglesAndSourceRadius = FVector4(-2, 1, 0, 0);
+
+				if( SimpleLight.Exponent == 0.0f )
+				{
+					// Correction for lumen units
+					LightData.LightColorAndFalloffExponent *= 16.0f;
+				}
+			}
+
+			if (ForwardLocalLightData.Num() > 0)
+			{
+				const uint32 NumBytesRequired = ForwardLocalLightData.Num() * ForwardLocalLightData.GetTypeSize();
+
+				if (View.ForwardLightingResources->ForwardLocalLightBuffer.NumBytes < NumBytesRequired)
+				{
+					View.ForwardLightingResources->ForwardLocalLightBuffer.Release();
+					View.ForwardLightingResources->ForwardLocalLightBuffer.Initialize(sizeof(FVector4), NumBytesRequired / sizeof(FVector4), PF_A32B32G32R32F, BUF_Volatile);
+				}
+
+				View.ForwardLightingResources->ForwardLocalLightBuffer.Lock();
+				FPlatformMemory::Memcpy(View.ForwardLightingResources->ForwardLocalLightBuffer.MappedBuffer, ForwardLocalLightData.GetData(), ForwardLocalLightData.Num() * ForwardLocalLightData.GetTypeSize());
+				View.ForwardLightingResources->ForwardLocalLightBuffer.Unlock();
+			}
+
+			const FIntPoint LightGridSizeXY = FIntPoint::DivideAndRoundUp(View.ViewRect.Size(), GLightGridPixelSize);
+			GlobalLightData.NumLocalLights = ForwardLocalLightData.Num();
+			GlobalLightData.CulledGridSize = FIntVector(LightGridSizeXY.X, LightGridSizeXY.Y, GLightGridSizeZ);
+			GlobalLightData.MaxCulledLightsPerCell = GMaxCulledLightsPerCell;
+			GlobalLightData.LightGridPixelSizeShift = FMath::FloorLog2(GLightGridPixelSize);
+
+			float FarPlane = FurthestLight;
+			FVector ZParams = GetLightGridZParams(View.NearClippingDistance, FarPlane + 10.f);
+			GlobalLightData.LightGridZParams = ZParams;
+
+			const uint32 NumIndexableLights = 1 << (sizeof(FLightIndexType) * 8);
+
+			if ((uint32)ForwardLocalLightData.Num() > NumIndexableLights)
+			{
+				static bool bWarned = false;
+
+				if (!bWarned)
+				{
+					UE_LOG(LogRenderer, Warning, TEXT("Exceeded indexable light count, glitches will be visible (%u / %u)"), ForwardLocalLightData.Num(), NumIndexableLights);
+					bWarned = true;
+				}
+			}
+
+			View.ForwardLightingResources->ForwardGlobalLightData = TUniformBufferRef<FForwardGlobalLightData>::CreateUniformBufferImmediate(GlobalLightData, UniformBuffer_SingleFrame);
+		}
+
+		for (int32 ViewIndex = 0; ViewIndex < Views.Num(); ViewIndex++)
+		{
+			FViewInfo& View = Views[ViewIndex];
+
+			const FIntPoint LightGridSizeXY = FIntPoint::DivideAndRoundUp(View.ViewRect.Size(), GLightGridPixelSize);
+			const int32 NumCells = LightGridSizeXY.X * LightGridSizeXY.Y * GLightGridSizeZ;
+
+			if (View.ForwardLightingResources->NumCulledLightsGrid.NumBytes != NumCells * NumCulledLightsGridStride * sizeof(uint32))
+			{
+				View.ForwardLightingResources->NumCulledLightsGrid.Initialize(sizeof(uint32), NumCells * NumCulledLightsGridStride, PF_R32_UINT);
+				View.ForwardLightingResources->NextCulledLightLink.Initialize(sizeof(uint32), 1, PF_R32_UINT);
+				View.ForwardLightingResources->StartOffsetGrid.Initialize(sizeof(uint32), NumCells, PF_R32_UINT);
+				View.ForwardLightingResources->NextCulledLightData.Initialize(sizeof(uint32), 1, PF_R32_UINT);
+			}
+
+			if (View.ForwardLightingResources->CulledLightDataGrid.NumBytes != NumCells * GMaxCulledLightsPerCell * sizeof(FLightIndexType))
+			{
+				View.ForwardLightingResources->CulledLightDataGrid.Initialize(sizeof(FLightIndexType), NumCells * GMaxCulledLightsPerCell, sizeof(FLightIndexType) == sizeof(uint16) ? PF_R16_UINT : PF_R32_UINT);
+				View.ForwardLightingResources->CulledLightLinks.Initialize(sizeof(uint32), NumCells * GMaxCulledLightsPerCell * LightLinkStride, PF_R32_UINT);
+			}
+
+			const FIntVector NumGroups = FIntVector::DivideAndRoundUp(FIntVector(LightGridSizeXY.X, LightGridSizeXY.Y, GLightGridSizeZ), LightGridInjectionGroupSize);
+
+			{
+				SCOPED_DRAW_EVENT(RHICmdList, CullLights);
+
+				if (GLightLinkedListCulling)
+				{
+					uint32 LinkClearValues[4] = { 0xFFFFFFFF };
+					RHICmdList.ClearUAV(View.ForwardLightingResources->StartOffsetGrid.UAV, LinkClearValues);
+
+					uint32 ClearValues[4] = { 0 };
+					RHICmdList.ClearUAV(View.ForwardLightingResources->NextCulledLightLink.UAV, ClearValues);
+					RHICmdList.ClearUAV(View.ForwardLightingResources->NextCulledLightData.UAV, ClearValues);
+
+					TShaderMapRef<TLightGridInjectionCS<true> > ComputeShader(View.ShaderMap);
+					RHICmdList.SetComputeShader(ComputeShader->GetComputeShader());
+					ComputeShader->SetParameters(RHICmdList, View);
+					DispatchComputeShader(RHICmdList, *ComputeShader, NumGroups.X, NumGroups.Y, NumGroups.Z);
+					ComputeShader->UnsetParameters(RHICmdList, View);
+				}
+				else
+				{
+					uint32 ClearValues[4] = { 0 };
+					RHICmdList.ClearUAV(View.ForwardLightingResources->NumCulledLightsGrid.UAV, ClearValues);
+
+					TShaderMapRef<TLightGridInjectionCS<false> > ComputeShader(View.ShaderMap);
+					RHICmdList.SetComputeShader(ComputeShader->GetComputeShader());
+					ComputeShader->SetParameters(RHICmdList, View);
+					DispatchComputeShader(RHICmdList, *ComputeShader, NumGroups.X, NumGroups.Y, NumGroups.Z);
+					ComputeShader->UnsetParameters(RHICmdList, View);
+				}
+			}
+
+			if (GLightLinkedListCulling)
+			{
+				SCOPED_DRAW_EVENT(RHICmdList, Compact);
+
+				TShaderMapRef<FLightGridCompactCS> ComputeShader(View.ShaderMap);
+				RHICmdList.SetComputeShader(ComputeShader->GetComputeShader());
+				ComputeShader->SetParameters(RHICmdList, View);
+				DispatchComputeShader(RHICmdList, *ComputeShader, NumGroups.X, NumGroups.Y, NumGroups.Z);
+				ComputeShader->UnsetParameters(RHICmdList, View);
+			}
+		}
+	}
+}
+
+void FDeferredShadingSceneRenderer::RenderForwardShadingShadowProjections(FRHICommandListImmediate& RHICmdList)
+{
+	bool bLightAttenuationNeeded = false;
+
+	for (TSparseArray<FLightSceneInfoCompact>::TConstIterator LightIt(Scene->Lights); LightIt; ++LightIt)
+	{
+		const FLightSceneInfoCompact& LightSceneInfoCompact = *LightIt;
+		const FLightSceneInfo* const LightSceneInfo = LightSceneInfoCompact.LightSceneInfo;
+		const FVisibleLightInfo& VisibleLightInfo = VisibleLightInfos[LightSceneInfo->Id];
+
+		bLightAttenuationNeeded = bLightAttenuationNeeded || VisibleLightInfo.ShadowsToProject.Num() > 0 || VisibleLightInfo.CapsuleShadowsToProject.Num() > 0;
+	}
+
+	FSceneRenderTargets& SceneRenderTargets = FSceneRenderTargets::Get(RHICmdList);
+	SceneRenderTargets.SetLightAttenuationMode(bLightAttenuationNeeded);
+
+	if (bLightAttenuationNeeded)
+	{
+		SCOPED_DRAW_EVENT(RHICmdList, ShadowProjectionOnOpaque);
+
+		// All shadows render with min blending
+		bool bClearToWhite = true;
+		SceneRenderTargets.BeginRenderingLightAttenuation(RHICmdList, bClearToWhite);
+
+		for (TSparseArray<FLightSceneInfoCompact>::TConstIterator LightIt(Scene->Lights); LightIt; ++LightIt)
+		{
+			const FLightSceneInfoCompact& LightSceneInfoCompact = *LightIt;
+			const FLightSceneInfo* const LightSceneInfo = LightSceneInfoCompact.LightSceneInfo;
+			FVisibleLightInfo& VisibleLightInfo = VisibleLightInfos[LightSceneInfo->Id];
+
+			if (VisibleLightInfo.ShadowsToProject.Num() > 0)
+			{
+				FSceneRenderer::RenderShadowProjections(RHICmdList, LightSceneInfo, true, false);
+			}
+		
+			RenderCapsuleDirectShadows(*LightSceneInfo, RHICmdList, VisibleLightInfo.CapsuleShadowsToProject, true);
+		}
+
+		SceneRenderTargets.FinishRenderingLightAttenuation(RHICmdList);
+	}
+}
\ No newline at end of file
diff --git a/Engine/Source/Runtime/Renderer/Private/LightMapRendering.h b/Engine/Source/Runtime/Renderer/Private/LightMapRendering.h
--- a/Engine/Source/Runtime/Renderer/Private/LightMapRendering.h
+++ b/Engine/Source/Runtime/Renderer/Private/LightMapRendering.h
@@ -4,8 +4,7 @@
 	LightMapRendering.h: Light map rendering definitions.
 =============================================================================*/
 
-#ifndef __LIGHTMAPRENDERING_H__
-#define __LIGHTMAPRENDERING_H__
+#pragma once
 
 #include "Engine/ShadowMapTexture2D.h"
 
@@ -868,4 +867,3 @@ public:
 		) const;
 };
 
-#endif // __LIGHTMAPRENDERING_H__
diff --git a/Engine/Source/Runtime/Renderer/Private/LightPropagationVolume.cpp b/Engine/Source/Runtime/Renderer/Private/LightPropagationVolume.cpp
--- a/Engine/Source/Runtime/Renderer/Private/LightPropagationVolume.cpp
+++ b/Engine/Source/Runtime/Renderer/Private/LightPropagationVolume.cpp
@@ -796,6 +796,7 @@ FLightPropagationVolume::FLightPropagationVolume() :
 	, SecondaryBounceStrength( 0.0f )
 	, CubeSize( 5312.0f )
 	, bEnabled( false )
+	, bDirectionalOcclusionEnabled( false )
 	, bGeometryVolumeNeeded( false )
 	, mWriteBufferIndex( 0 )
 	, bNeedsBufferClear( true )
@@ -932,11 +933,12 @@ void FLightPropagationVolume::InitSettings(FRHICommandListImmediate& RHICmdList,
 	Strength	 = LPVSettings.LPVIntensity;
 	bEnabled     = Strength > 0.0f;
 	CubeSize	 = LPVSettings.LPVSize;
+	bDirectionalOcclusionEnabled = bEnabled && ( LPVSettings.LPVDirectionalOcclusionIntensity > 0.0001f );
 
-	SecondaryOcclusionStrength =LPVSettings.LPVSecondaryOcclusionIntensity;
-	SecondaryBounceStrength =LPVSettings.LPVSecondaryBounceIntensity;
+	SecondaryOcclusionStrength = LPVSettings.LPVSecondaryOcclusionIntensity;
+	SecondaryBounceStrength = LPVSettings.LPVSecondaryBounceIntensity;
 
-	bGeometryVolumeNeeded =LPVSettings.LPVSecondaryOcclusionIntensity > 0.001f ||LPVSettings.LPVDirectionalOcclusionIntensity > 0.001;
+	bGeometryVolumeNeeded = LPVSettings.LPVSecondaryOcclusionIntensity > 0.001f || bDirectionalOcclusionEnabled;
 	GeometryVolumeGenerated = false;
 
 	if ( !bEnabled )
diff --git a/Engine/Source/Runtime/Renderer/Private/LightPropagationVolume.h b/Engine/Source/Runtime/Renderer/Private/LightPropagationVolume.h
--- a/Engine/Source/Runtime/Renderer/Private/LightPropagationVolume.h
+++ b/Engine/Source/Runtime/Renderer/Private/LightPropagationVolume.h
@@ -175,6 +175,8 @@ public:
 	FUnorderedAccessViewRHIParamRef GetGvListBufferUav()				{ return GvListBuffer->UAV; }
 	FUnorderedAccessViewRHIParamRef GetGvListHeadBufferUav()			{ return GvListHeadBuffer->UAV; }
 
+	bool IsEnabled() const												{ return bEnabled;	}
+	bool IsDirectionalOcclusionEnabled() const							{ return bDirectionalOcclusionEnabled; }
 
 	const FBox&	GetBoundingBox()										{ return BoundingBox; }
 
@@ -215,6 +217,7 @@ public:
 	float								CubeSize;
 	float								Strength;
 	bool								bEnabled;
+	bool								bDirectionalOcclusionEnabled;
 	bool								bGeometryVolumeNeeded;
 
 	uint32								mWriteBufferIndex;
diff --git a/Engine/Source/Runtime/Renderer/Private/LightRendering.cpp b/Engine/Source/Runtime/Renderer/Private/LightRendering.cpp
--- a/Engine/Source/Runtime/Renderer/Private/LightRendering.cpp
+++ b/Engine/Source/Runtime/Renderer/Private/LightRendering.cpp
@@ -315,6 +315,8 @@ void FSceneRenderer::GetLightNameForDrawEvent(const FLightSceneProxy* LightProxy
 #endif
 }
 
+extern int32 GbEnableAsyncComputeTranslucencyLightingVolumeClear;
+
 uint32 GetShadowQuality();
 
 /** Renders the scene's lighting. */
@@ -416,6 +418,12 @@ void FDeferredShadingSceneRenderer::RenderLights(FRHICommandListImmediate& RHICm
 			}
 		}
 		
+		if (GbEnableAsyncComputeTranslucencyLightingVolumeClear && GSupportsEfficientAsyncCompute)
+		{
+			//Gfx pipe must wait for the async compute clear of the translucency volume clear.
+			RHICmdList.WaitComputeFence(TranslucencyLightingVolumeClearEndFence);
+		}
+
 		if(ViewFamily.EngineShowFlags.DirectLighting)
 		{
 			SCOPED_DRAW_EVENT(RHICmdList, NonShadowedLights);
diff --git a/Engine/Source/Runtime/Renderer/Private/LightRendering.h b/Engine/Source/Runtime/Renderer/Private/LightRendering.h
--- a/Engine/Source/Runtime/Renderer/Private/LightRendering.h
+++ b/Engine/Source/Runtime/Renderer/Private/LightRendering.h
@@ -9,119 +9,6 @@
 
 #include "UniformBuffer.h"
 
-/** Uniform buffer for rendering deferred lights. */
-BEGIN_UNIFORM_BUFFER_STRUCT(FDeferredLightUniformStruct,)
-	DECLARE_UNIFORM_BUFFER_STRUCT_MEMBER(FVector,LightPosition)
-	DECLARE_UNIFORM_BUFFER_STRUCT_MEMBER(float,LightInvRadius)
-	DECLARE_UNIFORM_BUFFER_STRUCT_MEMBER(FVector,LightColor)
-	DECLARE_UNIFORM_BUFFER_STRUCT_MEMBER(float,LightFalloffExponent)
-	DECLARE_UNIFORM_BUFFER_STRUCT_MEMBER(FVector,NormalizedLightDirection)
-	DECLARE_UNIFORM_BUFFER_STRUCT_MEMBER(FVector2D,SpotAngles)
-	DECLARE_UNIFORM_BUFFER_STRUCT_MEMBER(float,SourceRadius)
-	DECLARE_UNIFORM_BUFFER_STRUCT_MEMBER(float,SourceLength)
-	DECLARE_UNIFORM_BUFFER_STRUCT_MEMBER(float,MinRoughness)
-	DECLARE_UNIFORM_BUFFER_STRUCT_MEMBER(float,ContactShadowLength)
-	DECLARE_UNIFORM_BUFFER_STRUCT_MEMBER(FVector2D,DistanceFadeMAD)
-	DECLARE_UNIFORM_BUFFER_STRUCT_MEMBER(FVector4,ShadowMapChannelMask)
-	DECLARE_UNIFORM_BUFFER_STRUCT_MEMBER(uint32,ShadowedBits)
-	DECLARE_UNIFORM_BUFFER_STRUCT_MEMBER(uint32,LightingChannelMask)
-END_UNIFORM_BUFFER_STRUCT(FDeferredLightUniformStruct)
-
-extern float GMinScreenRadiusForLights;
-extern uint32 GetShadowQuality();
-
-template<typename ShaderRHIParamRef>
-void SetDeferredLightParameters(
-	FRHICommandList& RHICmdList, 
-	const ShaderRHIParamRef ShaderRHI, 
-	const TShaderUniformBufferParameter<FDeferredLightUniformStruct>& DeferredLightUniformBufferParameter, 
-	const FLightSceneInfo* LightSceneInfo,
-	const FSceneView& View)
-{
-	FVector4 LightPositionAndInvRadius;
-	FVector4 LightColorAndFalloffExponent;
-		
-	FDeferredLightUniformStruct DeferredLightUniformsValue;
-
-	// Get the light parameters
-	LightSceneInfo->Proxy->GetParameters(
-		LightPositionAndInvRadius,
-		LightColorAndFalloffExponent,
-		DeferredLightUniformsValue.NormalizedLightDirection,
-		DeferredLightUniformsValue.SpotAngles,
-		DeferredLightUniformsValue.SourceRadius,
-		DeferredLightUniformsValue.SourceLength,
-		DeferredLightUniformsValue.MinRoughness);
-	
-	DeferredLightUniformsValue.LightPosition = LightPositionAndInvRadius;
-	DeferredLightUniformsValue.LightInvRadius = LightPositionAndInvRadius.W;
-	DeferredLightUniformsValue.LightColor = LightColorAndFalloffExponent;
-	DeferredLightUniformsValue.LightFalloffExponent = LightColorAndFalloffExponent.W;
-
-	const FVector2D FadeParams = LightSceneInfo->Proxy->GetDirectionalLightDistanceFadeParameters(View.GetFeatureLevel(), LightSceneInfo->IsPrecomputedLightingValid());
-
-	// use MAD for efficiency in the shader
-	DeferredLightUniformsValue.DistanceFadeMAD = FVector2D(FadeParams.Y, -FadeParams.X * FadeParams.Y);
-
-	int32 ShadowMapChannel = LightSceneInfo->Proxy->GetShadowMapChannel();
-
-	static const auto AllowStaticLightingVar = IConsoleManager::Get().FindTConsoleVariableDataInt(TEXT("r.AllowStaticLighting"));
-	const bool bAllowStaticLighting = (!AllowStaticLightingVar || AllowStaticLightingVar->GetValueOnRenderThread() != 0);
-
-	if (!bAllowStaticLighting)
-	{
-		ShadowMapChannel = INDEX_NONE;
-	}
-
-	DeferredLightUniformsValue.ShadowMapChannelMask = FVector4(
-		ShadowMapChannel == 0 ? 1 : 0,
-		ShadowMapChannel == 1 ? 1 : 0,
-		ShadowMapChannel == 2 ? 1 : 0,
-		ShadowMapChannel == 3 ? 1 : 0);
-
-	const bool bDynamicShadows = View.Family->EngineShowFlags.DynamicShadows && GetShadowQuality() > 0;
-	const bool bHasLightFunction = LightSceneInfo->Proxy->GetLightFunctionMaterial() != NULL;
-	DeferredLightUniformsValue.ShadowedBits  = LightSceneInfo->Proxy->CastsStaticShadow() || bHasLightFunction ? 1 : 0;
-	DeferredLightUniformsValue.ShadowedBits |= LightSceneInfo->Proxy->CastsDynamicShadow() && View.Family->EngineShowFlags.DynamicShadows ? 3 : 0;
-
-	static auto* ContactShadowsCVar = IConsoleManager::Get().FindTConsoleVariableDataInt(TEXT("r.ContactShadows"));
-	DeferredLightUniformsValue.ContactShadowLength = 0;
-
-	if (ContactShadowsCVar && ContactShadowsCVar->GetValueOnRenderThread() != 0)
-	{
-		DeferredLightUniformsValue.ContactShadowLength = LightSceneInfo->Proxy->GetContactShadowLength();
-	}
-
-	if( LightSceneInfo->Proxy->IsInverseSquared() )
-	{
-		// Correction for lumen units
-		DeferredLightUniformsValue.LightColor *= 16.0f;
-	}
-
-	// When rendering reflection captures, the direct lighting of the light is actually the indirect specular from the main view
-	if (View.bIsReflectionCapture)
-	{
-		DeferredLightUniformsValue.LightColor *= LightSceneInfo->Proxy->GetIndirectLightingScale();
-	}
-
-	const ELightComponentType LightType = (ELightComponentType)LightSceneInfo->Proxy->GetLightType();
-
-	if (LightType == LightType_Point || LightType == LightType_Spot)
-	{
-		// Distance fade
-		FSphere Bounds = LightSceneInfo->Proxy->GetBoundingSphere();
-
-		const float DistanceSquared = ( Bounds.Center - View.ViewMatrices.ViewOrigin ).SizeSquared();
-		float Fade = FMath::Square( FMath::Min( 0.0002f, GMinScreenRadiusForLights / Bounds.W ) * View.LODDistanceFactor ) * DistanceSquared;
-		Fade = FMath::Clamp( 6.0f - 6.0f * Fade, 0.0f, 1.0f );
-		DeferredLightUniformsValue.LightColor *= Fade;
-	}
-
-	DeferredLightUniformsValue.LightingChannelMask = LightSceneInfo->Proxy->GetLightingChannelMask();
-
-	SetUniformBufferParameterImmediate(RHICmdList, ShaderRHI,DeferredLightUniformBufferParameter,DeferredLightUniformsValue);
-}
-
 template<typename ShaderRHIParamRef>
 void SetSimpleDeferredLightParameters(
 	FRHICommandList& RHICmdList, 
diff --git a/Engine/Source/Runtime/Renderer/Private/LightSceneInfo.h b/Engine/Source/Runtime/Renderer/Private/LightSceneInfo.h
--- a/Engine/Source/Runtime/Renderer/Private/LightSceneInfo.h
+++ b/Engine/Source/Runtime/Renderer/Private/LightSceneInfo.h
@@ -4,8 +4,7 @@
 	LightSceneInfo.h: Light scene info definitions.
 =============================================================================*/
 
-#ifndef __LIGHTSCENEINFO_H__
-#define __LIGHTSCENEINFO_H__
+#pragma once
 
 #include "StaticArray.h"
 #include "SceneManagement.h"
@@ -262,4 +261,16 @@ struct FLightOctreeSemantics
 	}
 };
 
-#endif
+/** Stores lighting information for the clustered forward shading path */
+struct FClusteredLightsSceneInfo
+{
+	FIntPoint TileSize;			// In pixels
+	FIntVector GridSize;		// In tiles (x,y) + slices (z)
+	FVector4 LightGridZParams;
+
+	// Index of the light in this array corresponds to the bit set in the grid.
+	TArray<FLightSceneInfoCompact> ClusteredLights;
+
+	// The light grid.  Size is >= GridSize.
+	FTexture3DRHIRef LightGridTex;
+};
\ No newline at end of file
diff --git a/Engine/Source/Runtime/Renderer/Private/MaterialTexCoordScalesRendering.cpp b/Engine/Source/Runtime/Renderer/Private/MaterialTexCoordScalesRendering.cpp
--- a/Engine/Source/Runtime/Renderer/Private/MaterialTexCoordScalesRendering.cpp
+++ b/Engine/Source/Runtime/Renderer/Private/MaterialTexCoordScalesRendering.cpp
@@ -87,6 +87,7 @@ void FMaterialTexCoordScalePS::SetMesh(
 	SetShaderValueArray(RHICmdList, FMeshMaterialShader::GetPixelShader(), OneOverCPUTexCoordScalesParameter, OneOverCPUTexCoordScales, ARRAY_COUNT(OneOverCPUTexCoordScales));
 	SetShaderValueArray(RHICmdList, FMeshMaterialShader::GetPixelShader(), TexCoordIndicesParameter, TexCoordIndices, ARRAY_COUNT(TexCoordIndices));
 	SetShaderValue(RHICmdList, FMeshMaterialShader::GetPixelShader(), TextureAnalysisIndexParameter, bOutputScales ? (int32)INDEX_NONE : AnalysisIndex);
+	SetShaderValue(RHICmdList, FMeshMaterialShader::GetPixelShader(), PrimitiveAlphaParameter, (!Proxy || Proxy->IsSelected()) ? 1.f : .2f);
 
 	FMeshMaterialShader::SetMesh(RHICmdList, FMeshMaterialShader::GetPixelShader(), VertexFactory, View, Proxy, BatchElement, DrawRenderState);
 }
diff --git a/Engine/Source/Runtime/Renderer/Private/MaterialTexCoordScalesRendering.h b/Engine/Source/Runtime/Renderer/Private/MaterialTexCoordScalesRendering.h
--- a/Engine/Source/Runtime/Renderer/Private/MaterialTexCoordScalesRendering.h
+++ b/Engine/Source/Runtime/Renderer/Private/MaterialTexCoordScalesRendering.h
@@ -32,7 +32,8 @@ public:
 		AccuracyColorsParameter.Bind(Initializer.ParameterMap,TEXT("AccuracyColors"));
 		TextureAnalysisIndexParameter.Bind(Initializer.ParameterMap,TEXT("TextureAnalysisIndex"));
 		OneOverCPUTexCoordScalesParameter.Bind(Initializer.ParameterMap,TEXT("OneOverCPUTexCoordScales"));
-		TexCoordIndicesParameter.Bind(Initializer.ParameterMap,TEXT("TexCoordIndices"));
+		TexCoordIndicesParameter.Bind(Initializer.ParameterMap, TEXT("TexCoordIndices"));
+		PrimitiveAlphaParameter.Bind(Initializer.ParameterMap, TEXT("PrimitiveAlpha"));
 	}
 
 	FMaterialTexCoordScalePS() {}
@@ -44,6 +45,7 @@ public:
 		Ar << TextureAnalysisIndexParameter;
 		Ar << OneOverCPUTexCoordScalesParameter;
 		Ar << TexCoordIndicesParameter;
+		Ar << PrimitiveAlphaParameter;
 		return bShaderHasOutdatedParameters;
 	}
 
@@ -85,4 +87,5 @@ private:
 	FShaderParameter TextureAnalysisIndexParameter;
 	FShaderParameter OneOverCPUTexCoordScalesParameter;
 	FShaderParameter TexCoordIndicesParameter;
+	FShaderParameter PrimitiveAlphaParameter;
 };
diff --git a/Engine/Source/Runtime/Renderer/Private/MeshTexCoordSizeAccuracyRendering.cpp b/Engine/Source/Runtime/Renderer/Private/MeshTexCoordSizeAccuracyRendering.cpp
--- a/Engine/Source/Runtime/Renderer/Private/MeshTexCoordSizeAccuracyRendering.cpp
+++ b/Engine/Source/Runtime/Renderer/Private/MeshTexCoordSizeAccuracyRendering.cpp
@@ -19,9 +19,6 @@ void FMeshTexCoordSizeAccuracyPS::SetParameters(
 	const FSceneView& View
 	)
 {
-	const auto& ViewUniformBufferParameter = GetUniformBufferParameter<FViewUniformShaderParameters>();
-	SetUniformBufferParameter(RHICmdList, FGlobalShader::GetPixelShader(), ViewUniformBufferParameter, View.ViewUniformBuffer);
-
 	const int32 NumEngineColors = FMath::Min<int32>(GEngine->StreamingAccuracyColors.Num(), NumStreamingAccuracyColors);
 	int32 ColorIndex = 0;
 	for (; ColorIndex < NumEngineColors; ++ColorIndex)
@@ -32,6 +29,9 @@ void FMeshTexCoordSizeAccuracyPS::SetParameters(
 	{
 		SetShaderValue(RHICmdList, FGlobalShader::GetPixelShader(), AccuracyColorsParameter, FLinearColor::Black, ColorIndex);
 	}
+
+	// Bind view params
+	FGlobalShader::SetParameters(RHICmdList, FGlobalShader::GetPixelShader(), View);
 }
 
 void FMeshTexCoordSizeAccuracyPS::SetMesh(
@@ -54,9 +54,11 @@ void FMeshTexCoordSizeAccuracyPS::SetMesh(
 	}
 #endif
 	SetShaderValue(RHICmdList, FGlobalShader::GetPixelShader(), CPUTexelFactorParameter, CPUTexelFactor);
+	SetShaderValue(RHICmdList, FGlobalShader::GetPixelShader(), PrimitiveAlphaParameter, (!Proxy || Proxy->IsSelected()) ? 1.f : .2f);
 }
 
 void FMeshTexCoordSizeAccuracyPS::SetMesh(FRHICommandList& RHICmdList, const FSceneView& View)
 {
 	SetShaderValue(RHICmdList, FGlobalShader::GetPixelShader(), CPUTexelFactorParameter, -1.f);
+	SetShaderValue(RHICmdList, FGlobalShader::GetPixelShader(), PrimitiveAlphaParameter, 1.f);
 }
diff --git a/Engine/Source/Runtime/Renderer/Private/MeshTexCoordSizeAccuracyRendering.h b/Engine/Source/Runtime/Renderer/Private/MeshTexCoordSizeAccuracyRendering.h
--- a/Engine/Source/Runtime/Renderer/Private/MeshTexCoordSizeAccuracyRendering.h
+++ b/Engine/Source/Runtime/Renderer/Private/MeshTexCoordSizeAccuracyRendering.h
@@ -25,6 +25,7 @@ public:
 	{
 		AccuracyColorsParameter.Bind(Initializer.ParameterMap,TEXT("AccuracyColors"));
 		CPUTexelFactorParameter.Bind(Initializer.ParameterMap,TEXT("CPUTexelFactor"));
+		PrimitiveAlphaParameter.Bind(Initializer.ParameterMap, TEXT("PrimitiveAlpha"));
 	}
 
 	FMeshTexCoordSizeAccuracyPS() {}
@@ -34,6 +35,7 @@ public:
 		bool bShaderHasOutdatedParameters = FGlobalShader::Serialize(Ar);
 		Ar << AccuracyColorsParameter;
 		Ar << CPUTexelFactorParameter;
+		Ar << PrimitiveAlphaParameter;
 		return bShaderHasOutdatedParameters;
 	}
 
@@ -68,4 +70,5 @@ private:
 
 	FShaderParameter AccuracyColorsParameter;
 	FShaderParameter CPUTexelFactorParameter;
+	FShaderParameter PrimitiveAlphaParameter;
 };
diff --git a/Engine/Source/Runtime/Renderer/Private/PlanarReflectionRendering.cpp b/Engine/Source/Runtime/Renderer/Private/PlanarReflectionRendering.cpp
--- a/Engine/Source/Runtime/Renderer/Private/PlanarReflectionRendering.cpp
+++ b/Engine/Source/Runtime/Renderer/Private/PlanarReflectionRendering.cpp
@@ -220,7 +220,8 @@ static void UpdatePlanarReflectionContents_RenderThread(
 					SCOPED_DRAW_EVENT(RHICmdList, RenderScene);
 					SceneRenderer->Render(RHICmdList);
 				}
-				if(MainSceneRenderer->Scene->ShouldUseDeferredRenderer())
+
+				if (MainSceneRenderer->Scene->GetShadingPath() == EShadingPath::Deferred)
 				{
 					PrefilterPlanarReflection<true>(RHICmdList, View, SceneProxy, Target);
 					RHICmdList.CopyToResolveTarget(RenderTarget->GetRenderTargetTexture(), RenderTargetTexture->TextureRHI, false, ResolveParams);
diff --git a/Engine/Source/Runtime/Renderer/Private/PostProcess/PostProcessDOF.cpp b/Engine/Source/Runtime/Renderer/Private/PostProcess/PostProcessDOF.cpp
--- a/Engine/Source/Runtime/Renderer/Private/PostProcess/PostProcessDOF.cpp
+++ b/Engine/Source/Runtime/Renderer/Private/PostProcess/PostProcessDOF.cpp
@@ -29,7 +29,7 @@ class FPostProcessDOFSetupPS : public FGlobalShader
 	static void ModifyCompilationEnvironment(EShaderPlatform Platform, FShaderCompilerEnvironment& OutEnvironment)
 	{
 		FGlobalShader::ModifyCompilationEnvironment(Platform,OutEnvironment);
-		OutEnvironment.SetDefine(TEXT("FORWARD_SHADING"), IsFeatureLevelSupported(Platform, ERHIFeatureLevel::SM4) ? 0u : 1u);
+		OutEnvironment.SetDefine(TEXT("MOBILE_SHADING"), IsFeatureLevelSupported(Platform, ERHIFeatureLevel::SM4) ? 0u : 1u);
 		OutEnvironment.SetDefine(TEXT("NEAR_BLUR"), (uint32)(NearBlur >= 1));
 		OutEnvironment.SetDefine(TEXT("DOF_VIGNETTE"), (uint32)(NearBlur == 2));
 		OutEnvironment.SetDefine(TEXT("MRT_COUNT"), FarBlur + (NearBlur > 0));
diff --git a/Engine/Source/Runtime/Renderer/Private/PostProcess/PostProcessTonemap.cpp b/Engine/Source/Runtime/Renderer/Private/PostProcess/PostProcessTonemap.cpp
--- a/Engine/Source/Runtime/Renderer/Private/PostProcess/PostProcessTonemap.cpp
+++ b/Engine/Source/Runtime/Renderer/Private/PostProcess/PostProcessTonemap.cpp
@@ -1386,7 +1386,7 @@ void FRCPassPostProcessTonemap::Process(FRenderingCompositePassContext& Context)
 		SceneContext.SetSceneColor(0);
 	}
 
-	if (ViewFamily.Scene && !ViewFamily.Scene->ShouldUseDeferredRenderer())
+	if (ViewFamily.Scene && ViewFamily.Scene->GetShadingPath() == EShadingPath::Mobile)
 	{
 		// Double buffer tonemapper output for temporal AA.
 		if(View.FinalPostProcessSettings.AntiAliasingMethod == AAM_TemporalAA)
diff --git a/Engine/Source/Runtime/Renderer/Private/PostProcess/SceneRenderTargets.cpp b/Engine/Source/Runtime/Renderer/Private/PostProcess/SceneRenderTargets.cpp
--- a/Engine/Source/Runtime/Renderer/Private/PostProcess/SceneRenderTargets.cpp
+++ b/Engine/Source/Runtime/Renderer/Private/PostProcess/SceneRenderTargets.cpp
@@ -264,12 +264,12 @@ inline const TCHAR* GetReflectionBrightnessTargetName(uint32 Index)
 	return Names[Index];
 }
 
-inline const TCHAR* GetSceneColorTargetName(FSceneRenderTargets::EShadingPath ShadingPath)
+inline const TCHAR* GetSceneColorTargetName(EShadingPath ShadingPath)
 {
-	const TCHAR* SceneColorNames[(uint32)FSceneRenderTargets::EShadingPath::Num] =
+	const TCHAR* SceneColorNames[(uint32)EShadingPath::Num] =
 	{ 
 		TEXT("SceneColorMobile"), 
-		TEXT("SceneColorDeferred") 
+		TEXT("SceneColorDeferred")
 	};
 	check((uint32)ShadingPath < ARRAY_COUNT(SceneColorNames));
 	return SceneColorNames[(uint32)ShadingPath];
@@ -375,7 +375,7 @@ void FSceneRenderTargets::Allocate(FRHICommandList& RHICmdList, const FSceneView
 
 	// If feature level has changed, release all previously allocated targets to the pool. If feature level has changed but
 	const auto NewFeatureLevel = ViewFamily.Scene->GetFeatureLevel();
-	CurrentShadingPath = ViewFamily.Scene->ShouldUseDeferredRenderer() ? EShadingPath::Deferred : EShadingPath::Mobile;
+	CurrentShadingPath = ViewFamily.Scene->GetShadingPath();
 
 	FIntPoint DesiredBufferSize = ComputeDesiredSize(ViewFamily);
 	check(DesiredBufferSize.X > 0 && DesiredBufferSize.Y > 0);
@@ -403,7 +403,7 @@ void FSceneRenderTargets::Allocate(FRHICommandList& RHICmdList, const FSceneView
 
 	int32 RSMResolution = FMath::Clamp(CVarRSMResolution.GetValueOnRenderThread(), 1, 2048);
 
-	if (!ViewFamily.Scene->ShouldUseDeferredRenderer())
+	if (ViewFamily.Scene->GetShadingPath() == EShadingPath::Mobile)
 	{
 		// ensure there is always enough space for mobile renderer's tiled shadow maps
 		// by reducing the shadow map resolution.
@@ -514,10 +514,10 @@ void FSceneRenderTargets::BeginRenderingGBuffer(FRHICommandList& RHICmdList, ERe
 {
 	SCOPED_DRAW_EVENT(RHICmdList, BeginRenderingSceneColor);
 
-	if (IsSimpleForwardShadingEnabled(GetFeatureLevelShaderPlatform(CurrentFeatureLevel)))
+	if (IsAnyForwardShadingEnabled(GetFeatureLevelShaderPlatform(CurrentFeatureLevel)))
 	{
 		// in this non-standard case, just render to scene color with default mode
-		BeginRenderingSceneColor(RHICmdList);
+		BeginRenderingSceneColor(RHICmdList, ColorLoadAction == ERenderTargetLoadAction::EClear ? ESimpleRenderTargetMode::EClearColorExistingDepth : ESimpleRenderTargetMode::EUninitializedColorExistingDepth);
 		return;
 	}
 
@@ -604,7 +604,7 @@ void FSceneRenderTargets::BeginRenderingGBuffer(FRHICommandList& RHICmdList, ERe
 
 void FSceneRenderTargets::FinishRenderingGBuffer(FRHICommandListImmediate& RHICmdList)
 {
-	if (IsSimpleForwardShadingEnabled(GetFeatureLevelShaderPlatform(CurrentFeatureLevel)))
+	if (IsAnyForwardShadingEnabled(GetFeatureLevelShaderPlatform(CurrentFeatureLevel)))
 	{
 		// in this non-standard case, just render to scene color with default mode
 		FinishRenderingSceneColor(RHICmdList, true);
@@ -634,7 +634,7 @@ int32 FSceneRenderTargets::GetNumGBufferTargets() const
 {
 	int32 NumGBufferTargets = 1;
 
-	if (CurrentFeatureLevel >= ERHIFeatureLevel::SM4 && !IsSimpleForwardShadingEnabled(GetFeatureLevelShaderPlatform(CurrentFeatureLevel)))
+	if (CurrentFeatureLevel >= ERHIFeatureLevel::SM4 && !IsAnyForwardShadingEnabled(GetFeatureLevelShaderPlatform(CurrentFeatureLevel)))
 	{
 		// This needs to match TBasePassPixelShaderBaseType::ModifyCompilationEnvironment()
 		NumGBufferTargets = bAllowStaticLighting ? 6 : 5;
@@ -928,7 +928,7 @@ const TRefCountPtr<IPooledRenderTarget>& FSceneRenderTargets::GetLightAttenuatio
 		{
 			bFirst = false;
 
-			// First we need to call AllocLightAttenuation(), contact MartinM if that happens
+			// First we need to call AllocLightAttenuation()
 			ensure(LightAttenuation);
 		}
 
@@ -948,7 +948,7 @@ TRefCountPtr<IPooledRenderTarget>& FSceneRenderTargets::GetLightAttenuation()
 		{
 			bFirst = false;
 
-			// the first called should be AllocLightAttenuation(), contact MartinM if that happens
+			// the first called should be AllocLightAttenuation()
 			ensure(LightAttenuation);
 		}
 
@@ -1164,7 +1164,8 @@ void FSceneRenderTargets::FinishRenderingSceneAlphaCopy(FRHICommandListImmediate
 
 void FSceneRenderTargets::BeginRenderingLightAttenuation(FRHICommandList& RHICmdList, bool bClearToWhite)
 {
-	SCOPED_DRAW_EVENT(RHICmdList, BeginRenderingLightAttenuation);
+	SCOPED_CONDITIONAL_DRAW_EVENT(RHICmdList, ClearLightAttenuation, bClearToWhite);
+	SCOPED_CONDITIONAL_DRAW_EVENT(RHICmdList, BeginRenderingLightAttenuation, !bClearToWhite);
 
 	AllocLightAttenuation(RHICmdList);
 
@@ -1654,6 +1655,13 @@ void FSceneRenderTargets::AllocateDeferredShadingPathRenderTargets(FRHICommandLi
 		}
 		
 		{
+			uint32 TranslucencyTargetFlags = TexCreate_ShaderResource | TexCreate_RenderTargetable;
+
+			if (CurrentFeatureLevel >= ERHIFeatureLevel::SM5)
+			{
+				TranslucencyTargetFlags |= TexCreate_UAV;
+			}
+
 			for (int32 RTSetIndex = 0; RTSetIndex < NumTranslucentVolumeRenderTargetSets; RTSetIndex++)
 			{
 				GRenderTargetPool.FindFreeElement(
@@ -1665,7 +1673,7 @@ void FSceneRenderTargets::AllocateDeferredShadingPathRenderTargets(FRHICommandLi
 						PF_FloatRGBA,
 						FClearValueBinding::None,
 						0,
-						TexCreate_ShaderResource | TexCreate_RenderTargetable,
+						TranslucencyTargetFlags,
 						false,
 						1,
 						false)),
@@ -1682,7 +1690,7 @@ void FSceneRenderTargets::AllocateDeferredShadingPathRenderTargets(FRHICommandLi
 						PF_FloatRGBA,
 						FClearValueBinding::None,
 						0,
-						TexCreate_ShaderResource | TexCreate_RenderTargetable,
+						TranslucencyTargetFlags,
 						false,
 						1,
 						false)),
@@ -2221,12 +2229,13 @@ void FSceneTextureShaderParameters::Set(
 	{
 		bool bDirectionalOcclusion = false;
 		FSceneViewState* ViewState = (FSceneViewState*)View.State;
-
-		if(ViewState && ViewState->GetLightPropagationVolume(View.GetFeatureLevel()))
+		if (ViewState != nullptr)
 		{
-			const FLightPropagationVolumeSettings& LPVSettings = View.FinalPostProcessSettings.BlendableManager.GetSingleFinalDataConst<FLightPropagationVolumeSettings>();
-
-			bDirectionalOcclusion = LPVSettings.LPVIntensity > 0.0f && LPVSettings.LPVDirectionalOcclusionIntensity > 0.0001f;
+			FLightPropagationVolume* Lpv = ViewState->GetLightPropagationVolume(View.GetFeatureLevel());
+		    if(Lpv != nullptr)
+		    {
+			    bDirectionalOcclusion = Lpv->IsDirectionalOcclusionEnabled(); 
+		    }
 		}
 
 		FTextureRHIParamRef DirectionalOcclusion = nullptr;
diff --git a/Engine/Source/Runtime/Renderer/Private/PostProcess/SceneRenderTargets.h b/Engine/Source/Runtime/Renderer/Private/PostProcess/SceneRenderTargets.h
--- a/Engine/Source/Runtime/Renderer/Private/PostProcess/SceneRenderTargets.h
+++ b/Engine/Source/Runtime/Renderer/Private/PostProcess/SceneRenderTargets.h
@@ -196,14 +196,6 @@ protected:
 	FSceneRenderTargets(const FViewInfo& InView, const FSceneRenderTargets& SnapshotSource);
 public:
 
-	enum class EShadingPath
-	{
-		Mobile,
-		Deferred,
-
-		Num,
-	};
-
 	/**
 	 * Checks that scene render targets are ready for rendering a view family of the given dimensions.
 	 * If the allocated render targets are too small, they are reallocated.
@@ -670,7 +662,11 @@ private:
 	bool AreShadingPathRenderTargetsAllocated(EShadingPath InShadingPath) const;
 
 	/** Determine whether the render targets for any shading path have been allocated */
-	bool AreAnyShadingPathRenderTargetsAllocated() const { return AreShadingPathRenderTargetsAllocated(EShadingPath::Deferred) || AreShadingPathRenderTargetsAllocated(EShadingPath::Mobile); }
+	bool AreAnyShadingPathRenderTargetsAllocated() const 
+	{ 
+		return AreShadingPathRenderTargetsAllocated(EShadingPath::Deferred) 
+			|| AreShadingPathRenderTargetsAllocated(EShadingPath::Mobile); 
+	}
 
 	/** Gets all GBuffers to use.  Returns the number actually used. */
 	int32 GetGBufferRenderTargets(ERenderTargetLoadAction ColorLoadAction, FRHIRenderTargetView OutRenderTargets[MaxSimultaneousRenderTargets], int32& OutVelocityRTIndex);
diff --git a/Engine/Source/Runtime/Renderer/Private/PrimitiveDistanceAccuracyRendering.cpp b/Engine/Source/Runtime/Renderer/Private/PrimitiveDistanceAccuracyRendering.cpp
--- a/Engine/Source/Runtime/Renderer/Private/PrimitiveDistanceAccuracyRendering.cpp
+++ b/Engine/Source/Runtime/Renderer/Private/PrimitiveDistanceAccuracyRendering.cpp
@@ -53,7 +53,6 @@ void FPrimitiveDistanceAccuracyPS::SetMesh(
 	{
 		FVector ViewToObject = SectionData->BoxOrigin - View.ViewMatrices.ViewOrigin;
 
-
 		float DistSqMinusRadiusSq = 0;
 		if (bUseNewMetrics)
 		{
@@ -73,9 +72,11 @@ void FPrimitiveDistanceAccuracyPS::SetMesh(
 #endif
 	// Because the streamer use FMath::FloorToFloat, here we need to use -1 to have a useful result.
 	SetShaderValue(RHICmdList, FGlobalShader::GetPixelShader(), CPULogDistanceParameter, CPULogDistance);
+	SetShaderValue(RHICmdList, FGlobalShader::GetPixelShader(), PrimitiveAlphaParameter, (!Proxy || Proxy->IsSelected()) ? 1.f : .2f);
 }
 
 void FPrimitiveDistanceAccuracyPS::SetMesh(FRHICommandList& RHICmdList, const FSceneView& View)
 {
 	SetShaderValue(RHICmdList, FGlobalShader::GetPixelShader(), CPULogDistanceParameter, -1.f);
+	SetShaderValue(RHICmdList, FGlobalShader::GetPixelShader(), PrimitiveAlphaParameter, 1.f);
 }
diff --git a/Engine/Source/Runtime/Renderer/Private/PrimitiveDistanceAccuracyRendering.h b/Engine/Source/Runtime/Renderer/Private/PrimitiveDistanceAccuracyRendering.h
--- a/Engine/Source/Runtime/Renderer/Private/PrimitiveDistanceAccuracyRendering.h
+++ b/Engine/Source/Runtime/Renderer/Private/PrimitiveDistanceAccuracyRendering.h
@@ -24,7 +24,8 @@ public:
 		FGlobalShader(Initializer)
 	{
 		AccuracyColorsParameter.Bind(Initializer.ParameterMap,TEXT("AccuracyColors"));
-		CPULogDistanceParameter.Bind(Initializer.ParameterMap,TEXT("CPULogDistance"));
+		CPULogDistanceParameter.Bind(Initializer.ParameterMap, TEXT("CPULogDistance"));
+		PrimitiveAlphaParameter.Bind(Initializer.ParameterMap, TEXT("PrimitiveAlpha"));
 	}
 
 	FPrimitiveDistanceAccuracyPS() {}
@@ -34,6 +35,7 @@ public:
 		bool bShaderHasOutdatedParameters = FGlobalShader::Serialize(Ar);
 		Ar << AccuracyColorsParameter;
 		Ar << CPULogDistanceParameter;
+		Ar << PrimitiveAlphaParameter;
 		return bShaderHasOutdatedParameters;
 	}
 
@@ -64,4 +66,5 @@ private:
 
 	FShaderParameter AccuracyColorsParameter;
 	FShaderParameter CPULogDistanceParameter;
+	FShaderParameter PrimitiveAlphaParameter;
 };
diff --git a/Engine/Source/Runtime/Renderer/Private/ReflectionEnvironment.cpp b/Engine/Source/Runtime/Renderer/Private/ReflectionEnvironment.cpp
--- a/Engine/Source/Runtime/Renderer/Private/ReflectionEnvironment.cpp
+++ b/Engine/Source/Runtime/Renderer/Private/ReflectionEnvironment.cpp
@@ -153,11 +153,11 @@ public:
 		BentNormalAOSampler.Bind(ParameterMap, TEXT("BentNormalAOSampler"));
 		ApplyBentNormalAO.Bind(ParameterMap, TEXT("ApplyBentNormalAO"));
 		InvSkySpecularOcclusionStrength.Bind(ParameterMap, TEXT("InvSkySpecularOcclusionStrength"));
-		MinSkySpecularOcclusion.Bind(ParameterMap, TEXT("MinSkySpecularOcclusion"));
+		OcclusionTintAndMinOcclusion.Bind(ParameterMap, TEXT("OcclusionTintAndMinOcclusion"));
 	}
 
 	template<typename ShaderRHIParamRef, typename TRHICmdList>
-	void SetParameters(TRHICmdList& RHICmdList, const ShaderRHIParamRef& ShaderRHI, const TRefCountPtr<IPooledRenderTarget>& DynamicBentNormalAO, float SkySpecularOcclusionStrength, float MinOcclusionValue)
+	void SetParameters(TRHICmdList& RHICmdList, const ShaderRHIParamRef& ShaderRHI, const TRefCountPtr<IPooledRenderTarget>& DynamicBentNormalAO, float SkySpecularOcclusionStrength, const FVector4& OcclusionTintAndMinOcclusionValue)
 	{
 		FTextureRHIParamRef BentNormalAO = GWhiteTexture->TextureRHI;
 		bool bApplyBentNormalAO = false;
@@ -171,12 +171,12 @@ public:
 		SetTextureParameter(RHICmdList, ShaderRHI, BentNormalAOTexture, BentNormalAOSampler, TStaticSamplerState<SF_Point>::GetRHI(), BentNormalAO);
 		SetShaderValue(RHICmdList, ShaderRHI, ApplyBentNormalAO, bApplyBentNormalAO ? 1.0f : 0.0f);
 		SetShaderValue(RHICmdList, ShaderRHI, InvSkySpecularOcclusionStrength, 1.0f / FMath::Max(SkySpecularOcclusionStrength, .1f));
-		SetShaderValue(RHICmdList, ShaderRHI, MinSkySpecularOcclusion, MinOcclusionValue);
+		SetShaderValue(RHICmdList, ShaderRHI, OcclusionTintAndMinOcclusion, OcclusionTintAndMinOcclusionValue);
 	}
 
 	friend FArchive& operator<<(FArchive& Ar,FDistanceFieldAOSpecularOcclusionParameters& P)
 	{
-		Ar << P.BentNormalAOTexture << P.BentNormalAOSampler << P.ApplyBentNormalAO << P.InvSkySpecularOcclusionStrength << P.MinSkySpecularOcclusion;
+		Ar << P.BentNormalAOTexture << P.BentNormalAOSampler << P.ApplyBentNormalAO << P.InvSkySpecularOcclusionStrength << P.OcclusionTintAndMinOcclusion;
 		return Ar;
 	}
 
@@ -185,7 +185,7 @@ private:
 	FShaderResourceParameter BentNormalAOSampler;
 	FShaderParameter ApplyBentNormalAO;
 	FShaderParameter InvSkySpecularOcclusionStrength;
-	FShaderParameter MinSkySpecularOcclusion;
+	FShaderParameter OcclusionTintAndMinOcclusion;
 };
 
 struct FReflectionCaptureSortData
@@ -324,7 +324,8 @@ public:
 		SkyLightParameters.SetParameters(RHICmdList, ShaderRHI, Scene, View.Family->EngineShowFlags.SkyLighting);
 
 		const float MinOcclusion = Scene->SkyLight ? Scene->SkyLight->MinOcclusion : 0;
-		SpecularOcclusionParameters.SetParameters(RHICmdList, ShaderRHI, DynamicBentNormalAO, CVarSkySpecularOcclusionStrength.GetValueOnRenderThread(), MinOcclusion);
+		const FVector OcclusionTint = Scene->SkyLight ? (const FVector&)Scene->SkyLight->OcclusionTint : FVector::ZeroVector;
+		SpecularOcclusionParameters.SetParameters(RHICmdList, ShaderRHI, DynamicBentNormalAO, CVarSkySpecularOcclusionStrength.GetValueOnRenderThread(), FVector4(OcclusionTint, MinOcclusion));
 	}
 
 	void UnsetParameters(FRHIAsyncComputeCommandListImmediate& RHICmdList, FUnorderedAccessViewRHIParamRef OutSceneColorUAV)
@@ -487,7 +488,8 @@ public:
 		
 		FScene* Scene = (FScene*)View.Family->Scene;
 		const float MinOcclusion = Scene->SkyLight ? Scene->SkyLight->MinOcclusion : 0;
-		SpecularOcclusionParameters.SetParameters(RHICmdList, ShaderRHI, DynamicBentNormalAO, CVarSkySpecularOcclusionStrength.GetValueOnRenderThread(), MinOcclusion);
+		const FVector OcclusionTint = Scene->SkyLight ? (const FVector&)Scene->SkyLight->OcclusionTint : FVector::ZeroVector;
+		SpecularOcclusionParameters.SetParameters(RHICmdList, ShaderRHI, DynamicBentNormalAO, CVarSkySpecularOcclusionStrength.GetValueOnRenderThread(), FVector4(OcclusionTint, MinOcclusion));
 	}
 
 	// FShader interface.
diff --git a/Engine/Source/Runtime/Renderer/Private/RendererScene.cpp b/Engine/Source/Runtime/Renderer/Private/RendererScene.cpp
--- a/Engine/Source/Runtime/Renderer/Private/RendererScene.cpp
+++ b/Engine/Source/Runtime/Renderer/Private/RendererScene.cpp
@@ -998,7 +998,7 @@ void FScene::AddLightSceneInfo_RenderThread(FLightSceneInfo* LightSceneInfo)
 		SimpleDirectionalLight = LightSceneInfo;
 
 		// if we are mobile rendered and this light is a dynamic shadowcast then we need to update the static draw lists to pick a new lightingpolicy:
-		bool bMobileRendererRequiresLightPolicyChange = !ShouldUseDeferredRenderer() &&
+		bool bMobileRendererRequiresLightPolicyChange = GetShadingPath() == EShadingPath::Mobile &&
 			(
 			// this light is a dynamic shadowcast 
 			!SimpleDirectionalLight->Proxy->HasStaticShadowing() || 
@@ -1443,6 +1443,23 @@ void FScene::GetCaptureParameters(const FReflectionCaptureProxy* ReflectionProxy
 	}
 }
 
+int64 FScene::GetCachedWholeSceneShadowMapsSize() const
+{
+	int64 CachedShadowmapMemory = 0;
+
+	for (TMap<int32, FCachedShadowMapData>::TConstIterator CachedShadowMapIt(CachedShadowMaps); CachedShadowMapIt; ++CachedShadowMapIt)
+	{
+		const FCachedShadowMapData& ShadowMapData = CachedShadowMapIt.Value();
+
+		if (ShadowMapData.ShadowMap.IsValid())
+		{
+			CachedShadowmapMemory += ShadowMapData.ShadowMap.ComputeMemorySize();
+		}
+	}
+
+	return CachedShadowmapMemory;
+}
+
 void FScene::AddPrecomputedLightVolume(const FPrecomputedLightVolume* Volume)
 {
 	ENQUEUE_UNIQUE_RENDER_COMMAND_TWOPARAMETER(
@@ -1559,7 +1576,7 @@ void FScene::UpdateLightColorAndBrightness(ULightComponent* Light)
 					// thus, lights that change state in this way must update the draw lists.
 					Scene->bScenesPrimitivesNeedStaticMeshElementUpdate =
 						Scene->bScenesPrimitivesNeedStaticMeshElementUpdate ||
-						( !Scene->ShouldUseDeferredRenderer() 
+						( Scene->GetShadingPath() == EShadingPath::Mobile 
 						&& Parameters.NewColor.IsAlmostBlack() != LightSceneInfo->Proxy->GetColor().IsAlmostBlack() );
 
 					LightSceneInfo->Proxy->SetColor(Parameters.NewColor);
@@ -1584,7 +1601,7 @@ void FScene::RemoveLightSceneInfo_RenderThread(FLightSceneInfo* LightSceneInfo)
 		if (LightSceneInfo == SimpleDirectionalLight)
 		{
 			// if we are forward rendered and this light is a dynamic shadowcast then we need to update the static draw lists to pick a new lightingpolicy
-			bScenesPrimitivesNeedStaticMeshElementUpdate = bScenesPrimitivesNeedStaticMeshElementUpdate  || (!ShouldUseDeferredRenderer() && (!SimpleDirectionalLight->Proxy->HasStaticShadowing() || SimpleDirectionalLight->Proxy->UseCSMForDynamicObjects()));
+			bScenesPrimitivesNeedStaticMeshElementUpdate = bScenesPrimitivesNeedStaticMeshElementUpdate  || (GetShadingPath() == EShadingPath::Mobile && (!SimpleDirectionalLight->Proxy->HasStaticShadowing() || SimpleDirectionalLight->Proxy->UseCSMForDynamicObjects()));
 			SimpleDirectionalLight = NULL;
 		}
 
@@ -2088,7 +2105,7 @@ void FScene::UpdateStaticDrawListsForMaterials_RenderThread(FRHICommandListImmed
 	TArray<FPrimitiveSceneInfo*> PrimitivesToUpdate;
 	auto SceneFeatureLevel = GetFeatureLevel();
 
-	if (ShouldUseDeferredRenderer())
+	if (GetShadingPath() == EShadingPath::Deferred)
 	{
 		for (int32 DrawType = 0; DrawType < EBasePass_MAX; DrawType++)
 		{
@@ -2098,7 +2115,7 @@ void FScene::UpdateStaticDrawListsForMaterials_RenderThread(FRHICommandListImmed
 			BasePassUniformLightMapPolicyDrawList[DrawType].GetUsedPrimitivesBasedOnMaterials(SceneFeatureLevel, Materials, PrimitivesToUpdate);
 		}
 	}
-	else
+	else if (GetShadingPath() == EShadingPath::Mobile)
 	{
 		for (int32 DrawType = 0; DrawType < EBasePass_MAX; DrawType++)
 		{
diff --git a/Engine/Source/Runtime/Renderer/Private/SceneCaptureRendering.cpp b/Engine/Source/Runtime/Renderer/Private/SceneCaptureRendering.cpp
--- a/Engine/Source/Runtime/Renderer/Private/SceneCaptureRendering.cpp
+++ b/Engine/Source/Runtime/Renderer/Private/SceneCaptureRendering.cpp
@@ -302,7 +302,7 @@ static void UpdateSceneCaptureContent_RenderThread(
 			SCOPED_DRAW_EVENT(RHICmdList, FlipCapture);
 			CopyCaptureToTarget(RHICmdList, Target, TargetSize, View, ViewRect, FlippedRenderTarget.GetTextureParamRef(), true);
 		}
-		else if (!FSceneInterface::ShouldUseDeferredRenderer(View.GetFeatureLevel()))
+		else if (FSceneInterface::GetShadingPath(View.GetFeatureLevel()) == EShadingPath::Mobile)
 		{
 			// Copy the captured scene into the destination texture
 			SCOPED_DRAW_EVENT(RHICmdList, CaptureSceneColor);
diff --git a/Engine/Source/Runtime/Renderer/Private/SceneCore.cpp b/Engine/Source/Runtime/Renderer/Private/SceneCore.cpp
--- a/Engine/Source/Runtime/Renderer/Private/SceneCore.cpp
+++ b/Engine/Source/Runtime/Renderer/Private/SceneCore.cpp
@@ -211,15 +211,7 @@ FLightPrimitiveInteraction::FLightPrimitiveInteraction(
 		}
 	}
 
-	if (bCastShadow && !PrimitiveSceneInfo->Proxy->IsMeshShapeOftenMoving())
-	{
-		FCachedShadowMapData* CachedShadowMapData = PrimitiveSceneInfo->Scene->CachedShadowMaps.Find(LightSceneInfo->Id);
-
-		if (CachedShadowMapData)
-		{
-			CachedShadowMapData->ShadowMap.Release();
-		}
-	}
+	FlushCachedShadowMapData();
 
 	NextPrimitive = *PrevPrimitiveLink;
 	if(*PrevPrimitiveLink)
@@ -254,15 +246,7 @@ FLightPrimitiveInteraction::~FLightPrimitiveInteraction()
 	}
 #endif
 
-	if (bCastShadow && !PrimitiveSceneInfo->Proxy->IsMeshShapeOftenMoving())
-	{
-		FCachedShadowMapData* CachedShadowMapData = PrimitiveSceneInfo->Scene->CachedShadowMaps.Find(LightSceneInfo->Id);
-
-		if (CachedShadowMapData)
-		{
-			CachedShadowMapData->ShadowMap.Release();
-		}
-	}
+	FlushCachedShadowMapData();
 
 	// Track ES2 dynamic point light count
 	if (bES2DynamicPointLight)
@@ -289,6 +273,19 @@ FLightPrimitiveInteraction::~FLightPrimitiveInteraction()
 	*PrevLightLink = NextLight;
 }
 
+void FLightPrimitiveInteraction::FlushCachedShadowMapData()
+{
+	if (bCastShadow && !PrimitiveSceneInfo->Proxy->IsMeshShapeOftenMoving())
+	{
+		FCachedShadowMapData* CachedShadowMapData = PrimitiveSceneInfo->Scene->CachedShadowMaps.Find(LightSceneInfo->Id);
+
+		if (CachedShadowMapData)
+		{
+			CachedShadowMapData->ShadowMap.Release();
+		}
+	}
+}
+
 /*-----------------------------------------------------------------------------
 	FStaticMesh
 -----------------------------------------------------------------------------*/
@@ -331,7 +328,7 @@ void FStaticMesh::AddToDrawLists(FRHICommandListImmediate& RHICmdList, FScene* S
 		return;
 	}
 
-	if (Scene->ShouldUseDeferredRenderer())
+	if (Scene->GetShadingPath() == EShadingPath::Deferred)
 	{
 		if (bUseAsOccluder)
 		{
@@ -344,13 +341,8 @@ void FStaticMesh::AddToDrawLists(FRHICommandListImmediate& RHICmdList, FScene* S
 			EDepthDrawingMode EarlyZPassMode = (EDepthDrawingMode)EarlyZPass;
 			bool bEarlyZPassMovable = GEarlyZPassMovable != 0;
 
-			static IConsoleVariable* CDBufferVar = IConsoleManager::Get().FindConsoleVariable(TEXT("r.DBuffer"));
-			bool bDBufferAllowed = CDBufferVar ? CDBufferVar->GetInt() != 0 : false;
-
-			static const auto StencilLODDitherCVar = IConsoleManager::Get().FindTConsoleVariableDataInt(TEXT("r.StencilForLODDither"));
-			bool bStencilLODDither = StencilLODDitherCVar->GetValueOnAnyThread() != 0;
-
-			if (bDBufferAllowed || bStencilLODDither)
+			extern bool ShouldForceFullDepthPass(ERHIFeatureLevel::Type FeatureLevel);
+			if (ShouldForceFullDepthPass(Scene->GetFeatureLevel()))
 			{
 				// DBuffer decals force a full prepass
 				EarlyZPassMode = DDM_AllOccluders;
@@ -374,7 +366,7 @@ void FStaticMesh::AddToDrawLists(FRHICommandListImmediate& RHICmdList, FScene* S
 			FVelocityDrawingPolicyFactory::AddStaticMesh(Scene, this);
 		}
 	}
-	else
+	else if (Scene->GetShadingPath() == EShadingPath::Mobile)
 	{
 		if (bUseForMaterial)
 		{
diff --git a/Engine/Source/Runtime/Renderer/Private/SceneCore.h b/Engine/Source/Runtime/Renderer/Private/SceneCore.h
--- a/Engine/Source/Runtime/Renderer/Private/SceneCore.h
+++ b/Engine/Source/Runtime/Renderer/Private/SceneCore.h
@@ -50,6 +50,9 @@ public:
 		return (uint32)Interaction->LightId;
 	}
 
+	/** Clears cached shadow maps, if possible */
+	void FlushCachedShadowMapData();
+
 	/** Custom new/delete */
 	void* operator new(size_t Size);
 	void operator delete(void *RawMemory);
diff --git a/Engine/Source/Runtime/Renderer/Private/ScenePrivate.h b/Engine/Source/Runtime/Renderer/Private/ScenePrivate.h
--- a/Engine/Source/Runtime/Renderer/Private/ScenePrivate.h
+++ b/Engine/Source/Runtime/Renderer/Private/ScenePrivate.h
@@ -674,6 +674,8 @@ public:
 	FTextureRHIRef SelectionOutlineCacheKey;
 	TRefCountPtr<FRHIShaderResourceView> SelectionOutlineCacheValue;
 
+	FForwardLightingViewResources ForwardLightingResources;
+
 	/** Distance field AO tile intersection GPU resources.  Last frame's state is not used, but they must be sized exactly to the view so stored here. */
 	class FTileIntersectionResources* AOTileIntersectionResources;
 
@@ -927,9 +929,9 @@ public:
 		IndirectShadowLightDirectionVertexBuffer.SafeRelease();
 		IndirectShadowLightDirectionSRV.SafeRelease();
 		CapsuleTileIntersectionCountsBuffer.Release();
-
 		TranslucencyTimer.Release();
 		SeparateTranslucencyTimer.Release();
+		ForwardLightingResources.Release();
 	}
 
 	// FSceneViewStateInterface
@@ -2043,6 +2045,8 @@ public:
 	 */
 	void GetCaptureParameters(const FReflectionCaptureProxy* ReflectionProxy, FTextureRHIParamRef& ReflectionCubemapArray, int32& ArrayIndex) const;
 
+	int64 GetCachedWholeSceneShadowMapsSize() const;
+
 	/**
 	 * Marks static mesh elements as needing an update if necessary.
 	 */
@@ -2120,14 +2124,14 @@ public:
 			return SkyLight && !SkyLight->bHasStaticLighting;
 		}
 		else
-		{
-			const bool bRenderSkylight = SkyLight
-				&& !SkyLight->bHasStaticLighting
-				// The deferred shading renderer does movable skylight diffuse in a later deferred pass, not in the base pass
-				&& (SkyLight->bWantsStaticShadowing || IsSimpleForwardShadingEnabled(GetShaderPlatform()));
+	{
+		const bool bRenderSkylight = SkyLight
+			&& !SkyLight->bHasStaticLighting
+			// The deferred shading renderer does movable skylight diffuse in a later deferred pass, not in the base pass
+			&& (SkyLight->bWantsStaticShadowing || IsAnyForwardShadingEnabled(GetShaderPlatform()));
 
-			return bRenderSkylight;
-		}
+		return bRenderSkylight;
+	}
 	}
 
 	virtual TArray<FPrimitiveComponentId> GetScenePrimitiveComponentIds() const override
diff --git a/Engine/Source/Runtime/Renderer/Private/SceneRendering.cpp b/Engine/Source/Runtime/Renderer/Private/SceneRendering.cpp
--- a/Engine/Source/Runtime/Renderer/Private/SceneRendering.cpp
+++ b/Engine/Source/Runtime/Renderer/Private/SceneRendering.cpp
@@ -127,6 +127,12 @@ static TAutoConsoleVariable<float> CVarTessellationAdaptivePixelsPerTriangle(
 	TEXT("Global tessellation factor multiplier"),
 	ECVF_RenderThreadSafe);
 
+static TAutoConsoleVariable<int32> CVarForwardShading(
+	TEXT("r.ForwardShading"),
+	0,
+	TEXT("Whether to use forward shading on desktop platforms.  Requires Shader Model 5 hardware.  Forward shading has lower constant cost, but fewer features supported.  0:off, 1:on"),
+	ECVF_RenderThreadSafe | ECVF_ReadOnly);
+
 static TAutoConsoleVariable<int32> CVarSupportSimpleForwardShading(
 	TEXT("r.SupportSimpleForwardShading"),
 	0,
@@ -382,6 +388,7 @@ void FViewInfo::Init()
 	
 	bUsesGlobalDistanceField = false;
 	bUsesLightingChannels = false;
+	bTranslucentSurfaceLighting = false;
 
 	ExponentialFogParameters = FVector4(0,1,1,0);
 	ExponentialFogColor = FVector::ZeroVector;
@@ -411,6 +418,8 @@ void FViewInfo::Init()
 	bIsSnapshot = false;
 
 	bAllowStencilDither = false;
+
+	ForwardLightingResources = &ForwardLightingResourcesStorage;
 }
 
 FViewInfo::~FViewInfo()
@@ -1622,13 +1631,15 @@ void FSceneRenderer::RenderFinish(FRHICommandListImmediate& RHICmdList)
 
 FSceneRenderer* FSceneRenderer::CreateSceneRenderer(const FSceneViewFamily* InViewFamily, FHitProxyConsumer* HitProxyConsumer)
 {
-	bool bUseDeferred = InViewFamily->Scene->ShouldUseDeferredRenderer();
-	if (bUseDeferred)
+	EShadingPath ShadingPath = InViewFamily->Scene->GetShadingPath();
+
+	if (ShadingPath == EShadingPath::Deferred)
 	{
 		return new FDeferredShadingSceneRenderer(InViewFamily, HitProxyConsumer);
 	}
-	else
+	else 
 	{
+		check(ShadingPath == EShadingPath::Mobile);
 		return new FMobileSceneRenderer(InViewFamily, HitProxyConsumer);
 	}
 }
diff --git a/Engine/Source/Runtime/Renderer/Private/SceneRendering.h b/Engine/Source/Runtime/Renderer/Private/SceneRendering.h
--- a/Engine/Source/Runtime/Renderer/Private/SceneRendering.h
+++ b/Engine/Source/Runtime/Renderer/Private/SceneRendering.h
@@ -512,6 +512,44 @@ public:
 	void UpdateParameterData(float MaxOcclusionDistance);
 };
 
+BEGIN_UNIFORM_BUFFER_STRUCT_WITH_CONSTRUCTOR(FForwardGlobalLightData,)
+	DECLARE_UNIFORM_BUFFER_STRUCT_MEMBER(uint32,NumLocalLights)
+	DECLARE_UNIFORM_BUFFER_STRUCT_MEMBER(uint32,HasDirectionalLight)
+	DECLARE_UNIFORM_BUFFER_STRUCT_MEMBER(FIntVector,CulledGridSize)
+	DECLARE_UNIFORM_BUFFER_STRUCT_MEMBER(uint32,MaxCulledLightsPerCell)
+	DECLARE_UNIFORM_BUFFER_STRUCT_MEMBER(uint32,LightGridPixelSizeShift)
+	DECLARE_UNIFORM_BUFFER_STRUCT_MEMBER(FVector,LightGridZParams)
+	DECLARE_UNIFORM_BUFFER_STRUCT_MEMBER(FVector,DirectionalLightDirection)
+	DECLARE_UNIFORM_BUFFER_STRUCT_MEMBER(FVector,DirectionalLightColor)
+	DECLARE_UNIFORM_BUFFER_STRUCT_MEMBER(uint32,DirectionalLightShadowMapChannelMask)
+	DECLARE_UNIFORM_BUFFER_STRUCT_MEMBER(FVector2D,DirectionalLightDistanceFadeMAD)
+END_UNIFORM_BUFFER_STRUCT(FForwardGlobalLightData)
+
+class FForwardLightingViewResources
+{
+public:
+	TUniformBufferRef<FForwardGlobalLightData> ForwardGlobalLightData;
+	FDynamicReadBuffer ForwardLocalLightBuffer;
+	FRWBuffer NumCulledLightsGrid;
+	FRWBuffer CulledLightDataGrid;
+	FRWBuffer NextCulledLightLink;
+	FRWBuffer StartOffsetGrid;
+	FRWBuffer CulledLightLinks;
+	FRWBuffer NextCulledLightData;
+
+	void Release()
+	{
+		ForwardGlobalLightData.SafeRelease();
+		ForwardLocalLightBuffer.Release();
+		NumCulledLightsGrid.Release();
+		CulledLightDataGrid.Release();
+		NextCulledLightLink.Release();
+		StartOffsetGrid.Release();
+		CulledLightLinks.Release();
+		NextCulledLightData.Release();
+	}
+};
+
 /** A FSceneView with additional state used by the scene renderer. */
 class FViewInfo : public FSceneView
 {
@@ -645,6 +683,7 @@ public:
 	/** Whether the view has any materials that use the global distance field. */
 	uint32 bUsesGlobalDistanceField : 1;
 	uint32 bUsesLightingChannels : 1;
+	uint32 bTranslucentSurfaceLighting : 1;
 	/** 
 	 * true if the scene has at least one decal. Used to disable stencil operations in the mobile base pass when the scene has no decals.
 	 * TODO: Right now decal visibility is computed right before rendering them. Ideally it should be done in InitViews and this flag should be replaced with list of visible decals  
@@ -673,6 +712,12 @@ public:
 	// Hierarchical Z Buffer
 	TRefCountPtr<IPooledRenderTarget> HZB;
 
+	/** Points to the view state's resources if a view state exists. */
+	FForwardLightingViewResources* ForwardLightingResources;
+
+	/** Used when there is no view state, buffers reallocate every frame. */
+	FForwardLightingViewResources ForwardLightingResourcesStorage;
+
 	// Size of the HZB's mipmap 0
 	// NOTE: the mipmap 0 is downsampled version of the depth buffer
 	FIntPoint HZBMipmap0Size;
@@ -1059,7 +1104,7 @@ protected:
 
 	void InitDynamicShadows(FRHICommandListImmediate& RHICmdList);
 
-	bool RenderShadowProjections(FRHICommandListImmediate& RHICmdList, const FLightSceneInfo* LightSceneInfo, bool bMobileModulatedProjections);
+	bool RenderShadowProjections(FRHICommandListImmediate& RHICmdList, const FLightSceneInfo* LightSceneInfo, bool bProjectingForForwardShading, bool bMobileModulatedProjections);
 
 	/** Finds a matching cached preshadow, if one exists. */
 	TRefCountPtr<FProjectedShadowInfo> GetCachedPreshadow(
diff --git a/Engine/Source/Runtime/Renderer/Private/SceneVisibility.cpp b/Engine/Source/Runtime/Renderer/Private/SceneVisibility.cpp
--- a/Engine/Source/Runtime/Renderer/Private/SceneVisibility.cpp
+++ b/Engine/Source/Runtime/Renderer/Private/SceneVisibility.cpp
@@ -1392,6 +1392,7 @@ struct FRelevancePacket
 	uint16 CombinedShadingModelMask;
 	bool bUsesGlobalDistanceField;
 	bool bUsesLightingChannels;
+	bool bTranslucentSurfaceLighting;
 
 	FRelevancePacket(
 		FRHICommandListImmediate& InRHICmdList,
@@ -1416,6 +1417,7 @@ struct FRelevancePacket
 		, CombinedShadingModelMask(0)
 		, bUsesGlobalDistanceField(false)
 		, bUsesLightingChannels(false)
+		, bTranslucentSurfaceLighting(false)
 	{
 	}
 
@@ -1430,6 +1432,7 @@ struct FRelevancePacket
 		CombinedShadingModelMask = 0;
 		bUsesGlobalDistanceField = false;
 		bUsesLightingChannels = false;
+		bTranslucentSurfaceLighting = false;
 
 		SCOPE_CYCLE_COUNTER(STAT_ComputeViewRelevance);
 		for (int32 Index = 0; Index < Input.NumPrims; Index++)
@@ -1497,6 +1500,7 @@ struct FRelevancePacket
 			CombinedShadingModelMask |= ViewRelevance.ShadingModelMaskRelevance;
 			bUsesGlobalDistanceField |= ViewRelevance.bUsesGlobalDistanceField;
 			bUsesLightingChannels |= ViewRelevance.bUsesLightingChannels;
+			bTranslucentSurfaceLighting |= ViewRelevance.bTranslucentSurfaceLighting;
 
 			if (ViewRelevance.bRenderCustomDepth)
 			{
@@ -1526,17 +1530,15 @@ struct FRelevancePacket
 				*(PrimitiveSceneInfo->ComponentLastRenderTimeOnScreen) = CurrentWorldTime;
 			}
 
-			uint32 bHasClearCoat = ViewRelevance.ShadingModelMaskRelevance & (1 << MSM_ClearCoat);
-
 			// Cache the nearest reflection proxy if needed
 			if (PrimitiveSceneInfo->bNeedsCachedReflectionCaptureUpdate
 				// For mobile, the per-object reflection is used for everything
-				&& (!Scene->ShouldUseDeferredRenderer() || bTranslucentRelevance || bHasClearCoat))
+				&& (Scene->GetShadingPath() == EShadingPath::Mobile || bTranslucentRelevance || IsForwardShadingEnabled(Scene->GetFeatureLevel())))
 			{
 				PrimitiveSceneInfo->CachedReflectionCaptureProxy = Scene->FindClosestReflectionCapture(Scene->PrimitiveBounds[BitIndex].Origin);
 				PrimitiveSceneInfo->CachedPlanarReflectionProxy = Scene->FindClosestPlanarReflection(Scene->PrimitiveBounds[BitIndex]);
 
-				if (!Scene->ShouldUseDeferredRenderer())
+				if (Scene->GetShadingPath() == EShadingPath::Mobile)
 				{
 					// mobile HQ reflections
 					Scene->FindClosestReflectionCaptures(Scene->PrimitiveBounds[BitIndex].Origin, PrimitiveSceneInfo->CachedReflectionCaptureProxies);
@@ -1685,6 +1687,7 @@ struct FRelevancePacket
 		WriteView.ShadingModelMaskInView |= CombinedShadingModelMask;
 		WriteView.bUsesGlobalDistanceField |= bUsesGlobalDistanceField;
 		WriteView.bUsesLightingChannels |= bUsesLightingChannels;
+		WriteView.bTranslucentSurfaceLighting |= bTranslucentSurfaceLighting;
 		VisibleEditorPrimitives.AppendTo(WriteView.VisibleEditorPrimitives);
 		VisibleDynamicPrimitives.AppendTo(WriteView.VisibleDynamicPrimitives);
 		WriteView.TranslucentPrimSet.AppendScenePrimitives(TranslucencyPrims.Prims, TranslucencyPrims.NumPrims, TranslucencyPrimCount);
@@ -2819,11 +2822,15 @@ bool FDeferredShadingSceneRenderer::InitViews(FRHICommandListImmediate& RHICmdLi
 		// initialize per-view uniform buffer.
 		for (int32 ViewIndex = 0; ViewIndex < Views.Num(); ViewIndex++)
 		{
-			// Initialize the view's RHI resources.
-			Views[ViewIndex].InitRHIResources(nullptr);
+			FViewInfo& View = Views[ViewIndex];
+
+			View.ForwardLightingResources = View.ViewState ? &View.ViewState->ForwardLightingResources : &View.ForwardLightingResourcesStorage;
 
 			// Possible stencil dither optimization approach
-			Views[ViewIndex].bAllowStencilDither = bDitheredLODTransitionsUseStencil;
+			View.bAllowStencilDither = bDitheredLODTransitionsUseStencil;
+
+			// Initialize the view's RHI resources.
+			View.InitRHIResources(nullptr);
 		}
 	}
 
@@ -2925,6 +2932,7 @@ void FLODSceneTree::UpdateAndApplyVisibilityStates(FViewInfo& View)
 	PrimitiveFadingLODMap.Init(false, View.PrimitiveVisibilityMap.Num());
 	PrimitiveFadingOutLODMap.Init(false, View.PrimitiveVisibilityMap.Num());
 	FSceneBitArray& VisibilityFlags = View.PrimitiveVisibilityMap;
+	TArray<FPrimitiveViewRelevance, SceneRenderingAllocator>& RelevanceMap = View.PrimitiveViewRelevanceMap;
 
 	++UpdateCount;
 
@@ -2954,14 +2962,16 @@ void FLODSceneTree::UpdateAndApplyVisibilityStates(FViewInfo& View)
 			const int32 NodeIndex = Node.SceneInfo->GetIndex();
 			bool bIsVisible = VisibilityFlags[NodeIndex];
 
+			// Determine desired HLOD state
+			const FPrimitiveBounds& Bounds = Scene->PrimitiveBounds[NodeIndex];
+			const float DistanceSquared = (Bounds.Origin - View.ViewMatrices.ViewOrigin).SizeSquared();
+			const bool bIsInDrawRange = DistanceSquared >= Bounds.MinDrawDistanceSq;
+
+			const bool bWasFadingPreUpdate = Node.bIsFading;
+
 			// Update fading state
 			if (NodeMeshes[0].bDitheredLODTransition)
-			{	
-				// Determine desired HLOD state
-				const FPrimitiveBounds& Bounds = Scene->PrimitiveBounds[NodeIndex];
-				const float DistanceSquared = (Bounds.Origin - View.ViewMatrices.ViewOrigin).SizeSquared();
-				const bool bIsInDrawRange = DistanceSquared >= Bounds.MinDrawDistanceSq;
-
+			{
 				// Fade when HLODs change threshold on-screen, else snap
 				// TODO: This logic can still be improved to clear state and
 				//       transitions when off-screen, but needs better detection
@@ -2996,13 +3006,13 @@ void FLODSceneTree::UpdateAndApplyVisibilityStates(FViewInfo& View)
 					bIsVisible = Node.bWasVisible;
 				}
 			}
-#if WITH_EDITOR
 			else
 			{
-				// Force the default state in-case material edits cause an object to get stuck
+				// Instant transitions without dithering
+				Node.bWasVisible = Node.bIsVisible;
+				Node.bIsVisible = bIsInDrawRange;
 				Node.bIsFading = false;
 			}
-#endif
 
 			if (Node.bIsFading)
 			{
@@ -3014,6 +3024,25 @@ void FLODSceneTree::UpdateAndApplyVisibilityStates(FViewInfo& View)
 				// If stable and visible, override hierarchy visibility
 				HideNodeChildren(Node, VisibilityFlags);
 			}
+
+			// Flush cached lighting data when changing visible contents
+			if (Node.bIsVisible != Node.bWasVisible || bWasFadingPreUpdate || Node.bIsFading)
+			{
+				FLightPrimitiveInteraction* NodeLightList = Node.SceneInfo->LightList;
+				while (NodeLightList)
+				{
+					NodeLightList->FlushCachedShadowMapData();
+					NodeLightList = NodeLightList->GetNextLight();
+				}
+			}
+
+			// Force fully disabled view relevance so shadows don't attempt to recompute
+			if (!Node.bIsVisible)
+			{
+				FPrimitiveViewRelevance& ViewRelevance = RelevanceMap[NodeIndex];
+				FMemory::Memzero(&ViewRelevance, sizeof(FPrimitiveViewRelevance));
+				ViewRelevance.bInitializedThisFrame = true;
+			}
 		}
 	}
 }
diff --git a/Engine/Source/Runtime/Renderer/Private/ShaderBaseClasses.cpp b/Engine/Source/Runtime/Renderer/Private/ShaderBaseClasses.cpp
--- a/Engine/Source/Runtime/Renderer/Private/ShaderBaseClasses.cpp
+++ b/Engine/Source/Runtime/Renderer/Private/ShaderBaseClasses.cpp
@@ -68,8 +68,8 @@ FMaterialShader::FMaterialShader(const FMaterialShaderType::CompiledShaderInitia
 	}
 
 	DeferredParameters.Bind(Initializer.ParameterMap);
-	LightAttenuation.Bind(Initializer.ParameterMap, TEXT("LightAttenuationTexture"));
-	LightAttenuationSampler.Bind(Initializer.ParameterMap, TEXT("LightAttenuationTextureSampler"));
+	SceneColorCopyTexture.Bind(Initializer.ParameterMap, TEXT("SceneColorCopyTexture"));
+	SceneColorCopyTextureSampler.Bind(Initializer.ParameterMap, TEXT("SceneColorCopyTextureSampler"));
 	EyeAdaptation.Bind(Initializer.ParameterMap, TEXT("EyeAdaptation"));
 
 }
@@ -319,13 +319,13 @@ void FMaterialShader::SetParameters(
 	if (FeatureLevel >= ERHIFeatureLevel::SM4)
 	{
 		// for copied scene color
-		if(LightAttenuation.IsBound())
+		if(SceneColorCopyTexture.IsBound())
 		{
 			SetTextureParameter(
 				RHICmdList,
 				ShaderRHI,
-				LightAttenuation,
-				LightAttenuationSampler,
+				SceneColorCopyTexture,
+				SceneColorCopyTextureSampler,
 				TStaticSamplerState<SF_Bilinear,AM_Clamp,AM_Clamp,AM_Clamp>::GetRHI(),
 				FSceneRenderTargets::Get(RHICmdList).GetLightAttenuationTexture());
 		}
@@ -374,8 +374,8 @@ bool FMaterialShader::Serialize(FArchive& Ar)
 	Ar << MaterialUniformBuffer;
 	Ar << ParameterCollectionUniformBuffers;
 	Ar << DeferredParameters;
-	Ar << LightAttenuation;
-	Ar << LightAttenuationSampler;
+	Ar << SceneColorCopyTexture;
+	Ar << SceneColorCopyTextureSampler;
 	Ar << DebugUniformExpressionSet;
 	if (Ar.IsLoading())
 	{
diff --git a/Engine/Source/Runtime/Renderer/Private/ShadowRendering.cpp b/Engine/Source/Runtime/Renderer/Private/ShadowRendering.cpp
--- a/Engine/Source/Runtime/Renderer/Private/ShadowRendering.cpp
+++ b/Engine/Source/Runtime/Renderer/Private/ShadowRendering.cpp
@@ -296,7 +296,117 @@ static void SetShadowProjectionShaderTemplNew(FRHICommandList& RHICmdList, int32
 	}
 }
 
-void FProjectedShadowInfo::RenderProjection(FRHICommandListImmediate& RHICmdList, int32 ViewIndex, const FViewInfo* View, bool bMobileModulatedProjections) const
+void FProjectedShadowInfo::SetBlendStateForProjection(
+	FRHICommandListImmediate& RHICmdList, 
+	int32 ShadowMapChannel, 
+	bool bIsWholeSceneDirectionalShadow,
+	bool bUseFadePlane,
+	bool bProjectingForForwardShading, 
+	bool bMobileModulatedProjections)
+{
+	// With forward shading we are packing shadowing for all 4 possible stationary lights affecting each pixel into channels of the same texture, based on assigned shadowmap channels.
+	// With deferred shading we have 4 channels for each light.  
+	//	* CSM and per-object shadows are kept in separate channels to allow fading CSM out to precomputed shadowing while keeping per-object shadows past the fade distance.
+	//	* Subsurface shadowing requires an extra channel for each
+
+	if (bProjectingForForwardShading)
+	{
+		if (bUseFadePlane)
+		{
+			// alpha is used to fade between cascades
+			check(ShadowMapChannel == 0);
+			RHICmdList.SetBlendState(TStaticBlendState<CW_RED, BO_Add, BF_SourceAlpha, BF_InverseSourceAlpha>::GetRHI());
+		}
+		else
+		{
+			FBlendStateRHIParamRef BlendState = NULL;
+
+			if (ShadowMapChannel == 0)
+			{
+				BlendState = TStaticBlendState<CW_RED, BO_Min, BF_One, BF_One, BO_Min, BF_One, BF_One>::GetRHI();
+			}
+			else if (ShadowMapChannel == 1)
+			{
+				BlendState = TStaticBlendState<CW_GREEN, BO_Min, BF_One, BF_One, BO_Min, BF_One, BF_One>::GetRHI();
+			}
+			else if (ShadowMapChannel == 2)
+			{
+				BlendState = TStaticBlendState<CW_BLUE, BO_Min, BF_One, BF_One, BO_Min, BF_One, BF_One>::GetRHI();
+			}
+			else if (ShadowMapChannel == 3)
+			{
+				BlendState = TStaticBlendState<CW_ALPHA, BO_Min, BF_One, BF_One, BO_Min, BF_One, BF_One>::GetRHI();
+			}
+
+			checkf(BlendState, TEXT("Only shadows whose stationary lights have a valid ShadowMapChannel can be projected with forward shading"));
+			RHICmdList.SetBlendState(BlendState);
+		}
+	}
+	else
+	{
+		// Light Attenuation channel assignment:
+		//  R:     WholeSceneShadows, non SSS
+		//  G:     WholeSceneShadows,     SSS
+		//  B: non WholeSceneShadows, non SSS
+		//  A: non WholeSceneShadows,     SSS
+		//
+		// SSS: SubsurfaceScattering materials
+		// non SSS: shadow for opaque materials
+		// WholeSceneShadows: directional light CSM
+		// non WholeSceneShadows: spotlight, per object shadows, translucency lighting, omni-directional lights
+
+		if (bIsWholeSceneDirectionalShadow)
+		{
+			// use R and G in Light Attenuation
+			if (bUseFadePlane)
+			{
+				// alpha is used to fade between cascades, we don't don't need to do BO_Min as we leave B and A untouched which has translucency shadow
+				RHICmdList.SetBlendState(TStaticBlendState<CW_RG, BO_Add, BF_SourceAlpha, BF_InverseSourceAlpha>::GetRHI());
+			}
+			else
+			{
+				// first cascade rendered doesn't require fading (CO_Min is needed to combine multiple shadow passes)
+				// RTDF shadows: CO_Min is needed to combine with far shadows which overlap the same depth range
+				RHICmdList.SetBlendState(TStaticBlendState<CW_RG, BO_Min, BF_One, BF_One>::GetRHI());
+			}
+		}
+		else
+		{
+			if (bMobileModulatedProjections)
+			{
+				bool bEncodedHDR = IsMobileHDR32bpp() && !IsMobileHDRMosaic();
+				if (bEncodedHDR)
+				{
+					RHICmdList.SetBlendState(TStaticBlendState<>::GetRHI());
+				}
+				else
+				{
+					// Color modulate shadows, ignore alpha.
+					RHICmdList.SetBlendState(TStaticBlendState<CW_RGB, BO_Add, BF_Zero, BF_SourceColor, BO_Add, BF_Zero, BF_One>::GetRHI());
+				}
+			}
+			else
+			{
+				// use B and A in Light Attenuation
+				// CO_Min is needed to combine multiple shadow passes
+				RHICmdList.SetBlendState(TStaticBlendState<CW_BA, BO_Min, BF_One, BF_One, BO_Min, BF_One, BF_One>::GetRHI());
+			}
+		}
+	}
+}
+
+void FProjectedShadowInfo::SetBlendStateForProjection(FRHICommandListImmediate& RHICmdList, bool bProjectingForForwardShading, bool bMobileModulatedProjections) const
+{
+	SetBlendStateForProjection(
+		RHICmdList, 
+		GetLightSceneInfo().Proxy->GetPreviewShadowMapChannel(), 
+		IsWholeSceneDirectionalShadow(),
+		CascadeSettings.FadePlaneLength > 0 && !bRayTracedDistanceField,
+		bProjectingForForwardShading, 
+		bMobileModulatedProjections);
+}
+
+void FProjectedShadowInfo::RenderProjection(FRHICommandListImmediate& RHICmdList, int32 ViewIndex, const FViewInfo* View, bool bProjectingForForwardShading, bool bMobileModulatedProjections) const
 {
 #if WANTS_DRAW_MESH_EVENTS
 	FString EventName;
@@ -652,57 +762,7 @@ void FProjectedShadowInfo::RenderProjection(FRHICommandListImmediate& RHICmdList
 			>::GetRHI());
 	}
 
-	// Light Attenuation channel assignment:
-	//  R:     WholeSceneShadows, non SSS
-	//  G:     WholeSceneShadows,     SSS
-	//  B: non WholeSceneShadows, non SSS
-	//  A: non WholeSceneShadows,     SSS
-	//
-	// SSS: SubsurfaceScattering materials
-	// non SSS: shadow for opaque materials
-	// WholeSceneShadows: directional light
-	// non WholeSceneShadows: spotlight, per object shadows, translucency lighting, omni-directional lights
-
-	//@todo - get rid of this extra channel requirement by ordering the light mask accumulation
-
-	if (IsWholeSceneDirectionalShadow())
-	{
-		// use R and G in Light Attenuation (see above)
-
-		if(CascadeSettings.FadePlaneLength > 0)
-		{
-			// alpha is used to fade between cascades, we don't don't need to do BO_Min as we leave B and A untouched which has translucency shadow
-			RHICmdList.SetBlendState(TStaticBlendState<CW_RG, BO_Add, BF_SourceAlpha, BF_InverseSourceAlpha>::GetRHI());
-		}
-		else
-		{
-			// first cascade rendered or old method doesn't require fading (CO_Min is needed to combine multiple shadow passes)
-			RHICmdList.SetBlendState(TStaticBlendState<CW_RG, BO_Min, BF_One, BF_One>::GetRHI());
-		}
-	}
-	else
-	{
-		if (bMobileModulatedProjections)
-		{
-			bool bEncodedHDR = IsMobileHDR32bpp() && !IsMobileHDRMosaic();
-			if (bEncodedHDR)
-			{
-				RHICmdList.SetBlendState(TStaticBlendState<>::GetRHI());
-			}
-			else
-			{
-				// Color modulate shadows, ignore alpha.
-				RHICmdList.SetBlendState(TStaticBlendState<CW_RGB, BO_Add, BF_Zero, BF_SourceColor, BO_Add, BF_Zero, BF_One>::GetRHI());
-			}
-		}
-		else
-		{
-			// use B and A in Light Attenuation (see above)
-
-			// (CO_Min is needed to combine multiple shadow passes)
-			RHICmdList.SetBlendState(TStaticBlendState<CW_BA, BO_Min, BF_One, BF_One, BO_Min, BF_One, BF_One>::GetRHI());
-		}
-	}
+	SetBlendStateForProjection(RHICmdList, bProjectingForForwardShading, bMobileModulatedProjections);
 
 	{
 		uint32 LocalQuality = GetShadowQuality();
@@ -794,7 +854,7 @@ static void SetPointLightShaderTempl(FRHICommandList& RHICmdList, int32 ViewInde
 }
 
 /** Render one pass point light shadow projections. */
-void FProjectedShadowInfo::RenderOnePassPointLightProjection(FRHICommandListImmediate& RHICmdList, int32 ViewIndex, const FViewInfo& View) const
+void FProjectedShadowInfo::RenderOnePassPointLightProjection(FRHICommandListImmediate& RHICmdList, int32 ViewIndex, const FViewInfo& View, bool bProjectingForForwardShading) const
 {
 	SCOPE_CYCLE_COUNTER(STAT_RenderWholeSceneShadowProjectionsTime);
 
@@ -802,8 +862,7 @@ void FProjectedShadowInfo::RenderOnePassPointLightProjection(FRHICommandListImme
 	
 	const FSphere LightBounds = LightSceneInfo->Proxy->GetBoundingSphere();
 
-	// Use min blending to RGB since shadows may overlap, preserve Alpha
-	RHICmdList.SetBlendState(TStaticBlendState<CW_RGBA, BO_Min, BF_One, BF_One, BO_Add, BF_One, BF_Zero>::GetRHI());
+	SetBlendStateForProjection(RHICmdList, bProjectingForForwardShading, false);
 
 	const bool bCameraInsideLightGeometry = ((FVector)View.ViewMatrices.ViewOrigin - LightBounds.Center).SizeSquared() < FMath::Square(LightBounds.W * 1.05f + View.NearClippingDistance * 2.0f);
 
@@ -1150,7 +1209,7 @@ bool FDeferredShadingSceneRenderer::InjectReflectiveShadowMaps(FRHICommandListIm
 	
 extern int32 GCapsuleShadows;
 
-bool FSceneRenderer::RenderShadowProjections(FRHICommandListImmediate& RHICmdList, const FLightSceneInfo* LightSceneInfo, bool bMobileModulatedProjections)
+bool FSceneRenderer::RenderShadowProjections(FRHICommandListImmediate& RHICmdList, const FLightSceneInfo* LightSceneInfo, bool bProjectingForForwardShading, bool bMobileModulatedProjections)
 {
 	FVisibleLightInfo& VisibleLightInfo = VisibleLightInfos[LightSceneInfo->Id];
 	FSceneRenderTargets& SceneContext = FSceneRenderTargets::Get(RHICmdList);
@@ -1184,7 +1243,7 @@ bool FSceneRenderer::RenderShadowProjections(FRHICommandListImmediate& RHICmdLis
 
 			if (ProjectedShadowInfo->bRayTracedDistanceField)
 			{
-				ProjectedShadowInfo->RenderRayTracedDistanceFieldProjection(RHICmdList, View);
+				ProjectedShadowInfo->RenderRayTracedDistanceFieldProjection(RHICmdList, View, bProjectingForForwardShading);
 			}
 			else if (ProjectedShadowInfo->bAllocated)
 			{
@@ -1193,11 +1252,11 @@ bool FSceneRenderer::RenderShadowProjections(FRHICommandListImmediate& RHICmdLis
 				{
 					if (ProjectedShadowInfo->bOnePassPointLightShadow)
 					{
-						ProjectedShadowInfo->RenderOnePassPointLightProjection(RHICmdList, ViewIndex, View);
+						ProjectedShadowInfo->RenderOnePassPointLightProjection(RHICmdList, ViewIndex, View, bProjectingForForwardShading);
 					}
 					else 
 					{
-						ProjectedShadowInfo->RenderProjection(RHICmdList, ViewIndex, &View, bMobileModulatedProjections);
+						ProjectedShadowInfo->RenderProjection(RHICmdList, ViewIndex, &View, bProjectingForForwardShading, bMobileModulatedProjections);
 					}
 
 					if (!bMobileModulatedProjections)
@@ -1222,7 +1281,7 @@ bool FDeferredShadingSceneRenderer::RenderShadowProjections(FRHICommandListImmed
 
 	FVisibleLightInfo& VisibleLightInfo = VisibleLightInfos[LightSceneInfo->Id];
 
-	FSceneRenderer::RenderShadowProjections(RHICmdList, LightSceneInfo, false);
+	FSceneRenderer::RenderShadowProjections(RHICmdList, LightSceneInfo, false, false);
 
 	for (int32 ShadowIndex = 0; ShadowIndex < VisibleLightInfo.ShadowsToProject.Num(); ShadowIndex++)
 	{
@@ -1244,7 +1303,7 @@ bool FDeferredShadingSceneRenderer::RenderShadowProjections(FRHICommandListImmed
 		}
 	}
 
-	RenderCapsuleDirectShadows(*LightSceneInfo, RHICmdList, VisibleLightInfo.CapsuleShadowsToProject);
+	RenderCapsuleDirectShadows(*LightSceneInfo, RHICmdList, VisibleLightInfo.CapsuleShadowsToProject, false);
 
 	for (int32 ViewIndex = 0; ViewIndex < Views.Num(); ViewIndex++)
 	{
@@ -1280,6 +1339,6 @@ void FMobileSceneRenderer::RenderModulatedShadowProjections(FRHICommandListImmed
 		FLightSceneInfo* LightSceneInfo = LightSceneInfoCompact.LightSceneInfo;
 		TArray<FProjectedShadowInfo*, SceneRenderingAllocator> Shadows;
 		SCOPE_CYCLE_COUNTER(STAT_ProjectedShadowDrawTime);
-		FSceneRenderer::RenderShadowProjections(RHICmdList, LightSceneInfo, true);
+		FSceneRenderer::RenderShadowProjections(RHICmdList, LightSceneInfo, false, true);
 	}
 }
