Copying //UE4/Dev-Platform to //UE4/Dev-Main (Source: //UE4/Dev-Platform @ 3046626)

diff --git a/Engine/Source/Runtime/Renderer/Private/ShadowDepthRendering.cpp b/Engine/Source/Runtime/Renderer/Private/ShadowDepthRendering.cpp
--- a/Engine/Source/Runtime/Renderer/Private/ShadowDepthRendering.cpp
+++ b/Engine/Source/Runtime/Renderer/Private/ShadowDepthRendering.cpp
@@ -182,6 +182,11 @@ public:
 		OutEnvironment.SetDefine(TEXT("ONEPASS_POINTLIGHT_SHADOW"), (uint32)(ShaderMode == VertexShadowDepth_OnePassPointLight));
 		OutEnvironment.SetDefine(TEXT("REFLECTIVE_SHADOW_MAP"), (uint32)bRenderReflectiveShadowMap);
 		OutEnvironment.SetDefine(TEXT("POSITION_ONLY"), (uint32)bUsePositionOnlyStream);
+
+		if( bIsForGeometryShader )
+		{
+			OutEnvironment.CompilerFlags.Add( CFLAG_VertexToGeometryShader );
+		}
 	}
 };
 
@@ -1684,25 +1689,52 @@ void FProjectedShadowInfo::CopyCachedShadowMap(FRHICommandList& RHICmdList, FSce
 
 		if (bOnePassPointLightShadow)
 		{
-			// Set shaders and texture
-			TShaderMapRef<FScreenVSForGS> ScreenVertexShader(View.ShaderMap);
-			TShaderMapRef<FCopyShadowMapsCubeGS> GeometryShader(View.ShaderMap);
-			TShaderMapRef<FCopyShadowMapsCubePS> PixelShader(View.ShaderMap);
-
-			SetGlobalBoundShaderState(RHICmdList, SceneRenderer->FeatureLevel, CopyShadowMapsCubeBoundShaderState, GFilterVertexDeclaration.VertexDeclarationRHI, *ScreenVertexShader, *PixelShader, *GeometryShader);
-
-			PixelShader->SetParameters(RHICmdList, View, CachedShadowMapData.ShadowMap.DepthTarget.GetReference());
-
-			DrawRectangle(
-				RHICmdList,
-				0, 0,
-				ResolutionX, ResolutionY,
-				BorderSize, BorderSize,
-				ResolutionX, ResolutionY,
-				FIntPoint(ResolutionX, ResolutionY),
-				CachedShadowMapData.ShadowMap.GetSize(),
-				*ScreenVertexShader,
-				EDRF_Default);
+			if (RHISupportsGeometryShaders(GShaderPlatformForFeatureLevel[SceneRenderer->FeatureLevel]))
+			{
+				// Set shaders and texture
+				TShaderMapRef<TScreenVSForGS<false>> ScreenVertexShader(View.ShaderMap);
+				TShaderMapRef<FCopyShadowMapsCubeGS> GeometryShader(View.ShaderMap);
+				TShaderMapRef<FCopyShadowMapsCubePS> PixelShader(View.ShaderMap);
+
+				SetGlobalBoundShaderState(RHICmdList, SceneRenderer->FeatureLevel, CopyShadowMapsCubeBoundShaderState, GFilterVertexDeclaration.VertexDeclarationRHI, *ScreenVertexShader, *PixelShader, *GeometryShader);
+
+				PixelShader->SetParameters(RHICmdList, View, CachedShadowMapData.ShadowMap.DepthTarget.GetReference());
+
+				DrawRectangle(
+					RHICmdList,
+					0, 0,
+					ResolutionX, ResolutionY,
+					BorderSize, BorderSize,
+					ResolutionX, ResolutionY,
+					FIntPoint(ResolutionX, ResolutionY),
+					CachedShadowMapData.ShadowMap.GetSize(),
+					*ScreenVertexShader,
+					EDRF_Default);
+			}
+			else
+			{
+				check(RHISupportsVertexShaderLayer(GShaderPlatformForFeatureLevel[SceneRenderer->FeatureLevel]));
+				
+				// Set shaders and texture
+				TShaderMapRef<TScreenVSForGS<true>> ScreenVertexShader(View.ShaderMap);
+				TShaderMapRef<FCopyShadowMapsCubePS> PixelShader(View.ShaderMap);
+				
+				SetGlobalBoundShaderState(RHICmdList, SceneRenderer->FeatureLevel, CopyShadowMapsCubeBoundShaderState, GFilterVertexDeclaration.VertexDeclarationRHI, *ScreenVertexShader, *PixelShader);
+				
+				PixelShader->SetParameters(RHICmdList, View, CachedShadowMapData.ShadowMap.DepthTarget.GetReference());
+				
+				DrawRectangle(
+					  RHICmdList,
+					  0, 0,
+					  ResolutionX, ResolutionY,
+					  BorderSize, BorderSize,
+					  ResolutionX, ResolutionY,
+					  FIntPoint(ResolutionX, ResolutionY),
+					  CachedShadowMapData.ShadowMap.GetSize(),
+					  *ScreenVertexShader,
+					  EDRF_Default,
+					  6);
+			}
 		}
 		else
 		{
diff --git a/Engine/Source/Runtime/Renderer/Private/ShadowSetup.cpp b/Engine/Source/Runtime/Renderer/Private/ShadowSetup.cpp
--- a/Engine/Source/Runtime/Renderer/Private/ShadowSetup.cpp
+++ b/Engine/Source/Runtime/Renderer/Private/ShadowSetup.cpp
@@ -1750,7 +1750,7 @@ void ComputeWholeSceneShadowCacheModes(
 	EShadowDepthCacheMode* OutCacheModes)
 {
 	if (GCacheWholeSceneShadows 
-		&& (!bCubeShadowMap || RHISupportsGeometryShaders(GShaderPlatformForFeatureLevel[Scene->GetFeatureLevel()])))
+		&& (!bCubeShadowMap || RHISupportsGeometryShaders(GShaderPlatformForFeatureLevel[Scene->GetFeatureLevel()]) || RHISupportsVertexShaderLayer(GShaderPlatformForFeatureLevel[Scene->GetFeatureLevel()])))
 	{
 		FCachedShadowMapData* CachedShadowMapData = Scene->CachedShadowMaps.Find(LightSceneInfo->Id);
 
diff --git a/Engine/Source/Runtime/Renderer/Private/TranslucentRendering.h b/Engine/Source/Runtime/Renderer/Private/TranslucentRendering.h
--- a/Engine/Source/Runtime/Renderer/Private/TranslucentRendering.h
+++ b/Engine/Source/Runtime/Renderer/Private/TranslucentRendering.h
@@ -176,6 +176,12 @@ public:
 		return IsFeatureLevelSupported(Platform, ERHIFeatureLevel::SM4); 
 	}
 
+	static void ModifyCompilationEnvironment( EShaderPlatform Platform, FShaderCompilerEnvironment& OutEnvironment )
+	{
+		FGlobalShader::ModifyCompilationEnvironment( Platform, OutEnvironment );
+		OutEnvironment.CompilerFlags.Add( CFLAG_VertexToGeometryShader );
+	}
+
 	FWriteToSliceVS(const ShaderMetaType::CompiledShaderInitializerType& Initializer):
 		FGlobalShader(Initializer)
     {
diff --git a/Engine/Source/Runtime/Renderer/Private/VelocityRendering.cpp b/Engine/Source/Runtime/Renderer/Private/VelocityRendering.cpp
--- a/Engine/Source/Runtime/Renderer/Private/VelocityRendering.cpp
+++ b/Engine/Source/Runtime/Renderer/Private/VelocityRendering.cpp
@@ -897,7 +897,7 @@ FPooledRenderTargetDesc FVelocityRendering::GetRenderTargetDesc()
 {
 	const FIntPoint BufferSize = FSceneRenderTargets::Get_FrameConstantsOnly().GetBufferSizeXY();
 	const FIntPoint VelocityBufferSize = BufferSize;		// full resolution so we can reuse the existing full res z buffer
-	return FPooledRenderTargetDesc(FPooledRenderTargetDesc::Create2DDesc(VelocityBufferSize, PF_G16R16, FClearValueBinding::Transparent, TexCreate_None, TexCreate_RenderTargetable | TexCreate_DeltaColorCompression, false));
+	return FPooledRenderTargetDesc(FPooledRenderTargetDesc::Create2DDesc(VelocityBufferSize, PF_G16R16, FClearValueBinding::Transparent, TexCreate_None, TexCreate_RenderTargetable, false));
 }
 
 bool FVelocityRendering::OutputsToGBuffer()
diff --git a/Engine/Source/Runtime/ShaderCore/Private/ShaderCache.cpp b/Engine/Source/Runtime/ShaderCore/Private/ShaderCache.cpp
--- a/Engine/Source/Runtime/ShaderCore/Private/ShaderCache.cpp
+++ b/Engine/Source/Runtime/ShaderCore/Private/ShaderCache.cpp
@@ -41,7 +41,7 @@ static TCHAR const* GCookedCodeCacheFileName = TEXT("ByteCodeCache.ushadercode")
 #endif
 
 // Only the Mac build defaults to using the shader cache for now, Editor uses a separate cache from the game to avoid ever-growing cache being propagated to the game.
-int32 FShaderCache::bUseShaderCaching = (PLATFORM_MAC && !WITH_EDITOR) ? 1 : 0;
+int32 FShaderCache::bUseShaderCaching = 0;
 FAutoConsoleVariableRef FShaderCache::CVarUseShaderCaching(
 	TEXT("r.UseShaderCaching"),
 	bUseShaderCaching,
@@ -51,7 +51,7 @@ FAutoConsoleVariableRef FShaderCache::CVarUseShaderCaching(
 
 // Predrawing takes an existing shader cache with draw log & renders each shader + draw-state combination before use to avoid in-driver recompilation
 // This requires plenty of setup & is done in batches at frame-end.
-int32 FShaderCache::bUseShaderPredraw = (PLATFORM_MAC && !WITH_EDITOR) ? 1 : 0;
+int32 FShaderCache::bUseShaderPredraw = 0;
 FAutoConsoleVariableRef FShaderCache::CVarUseShaderPredraw(
 	TEXT("r.UseShaderPredraw"),
 	bUseShaderPredraw,
@@ -60,7 +60,7 @@ FAutoConsoleVariableRef FShaderCache::CVarUseShaderPredraw(
 	);
 
 // The actual draw loggging is even more expensive as it has to cache all the RHI draw state & is disabled by default.
-int32 FShaderCache::bUseShaderDrawLog = (PLATFORM_MAC && !WITH_EDITOR) ? 1 : 0;
+int32 FShaderCache::bUseShaderDrawLog = 0;
 FAutoConsoleVariableRef FShaderCache::CVarUseShaderDrawLog(
 	TEXT("r.UseShaderDrawLog"),
 	bUseShaderDrawLog,
@@ -78,7 +78,7 @@ FAutoConsoleVariableRef FShaderCache::CVarPredrawBatchTime(
 	);
 
 // A separate cache of used shader binaries for even earlier submission - may be platform or even device specific.
-int32 FShaderCache::bUseShaderBinaryCache = (PLATFORM_MAC && !WITH_EDITOR) ? 1 : 0;
+int32 FShaderCache::bUseShaderBinaryCache = 0;
 FAutoConsoleVariableRef FShaderCache::CVarUseShaderBinaryCache(
 	TEXT("r.UseShaderBinaryCache"),
 	bUseShaderBinaryCache,
diff --git a/Engine/Source/Runtime/ShaderCore/Public/ShaderCore.h b/Engine/Source/Runtime/ShaderCore/Public/ShaderCore.h
--- a/Engine/Source/Runtime/ShaderCore/Public/ShaderCore.h
+++ b/Engine/Source/Runtime/ShaderCore/Public/ShaderCore.h
@@ -123,6 +123,8 @@ enum ECompilerFlags
 	CFLAG_ForceRemoveUnusedInterpolators,
 	// Set default precision to highp in a pixel shader (default is mediump on ES2 platforms)
 	CFLAG_UseFullPrecisionInPS,
+	// Hint that its a vertex to geometry shader
+	CFLAG_VertexToGeometryShader
 };
 
 /**
diff --git a/Engine/Source/Runtime/Slate/Private/Framework/Application/SlateApplication.cpp b/Engine/Source/Runtime/Slate/Private/Framework/Application/SlateApplication.cpp
--- a/Engine/Source/Runtime/Slate/Private/Framework/Application/SlateApplication.cpp
+++ b/Engine/Source/Runtime/Slate/Private/Framework/Application/SlateApplication.cpp
@@ -1721,6 +1721,7 @@ TSharedRef< FGenericWindow > FSlateApplication::MakeWindow( TSharedRef<SWindow>
 	Definition->IsRegularWindow = InSlateWindow->IsRegularWindow();
 	Definition->HasSizingFrame = InSlateWindow->HasSizingFrame();
 	Definition->SizeWillChangeOften = InSlateWindow->SizeWillChangeOften();
+	Definition->ShouldPreserveAspectRatio = InSlateWindow->ShouldPreserveAspectRatio();
 	Definition->ExpectedMaxWidth = InSlateWindow->GetExpectedMaxWidth();
 	Definition->ExpectedMaxHeight = InSlateWindow->GetExpectedMaxHeight();
 
diff --git a/Engine/Source/Runtime/SlateCore/Private/Widgets/SWindow.cpp b/Engine/Source/Runtime/SlateCore/Private/Widgets/SWindow.cpp
--- a/Engine/Source/Runtime/SlateCore/Private/Widgets/SWindow.cpp
+++ b/Engine/Source/Runtime/SlateCore/Private/Widgets/SWindow.cpp
@@ -236,6 +236,7 @@ void SWindow::Construct(const FArguments& InArgs)
 	this->bHasMinimizeButton = InArgs._SupportsMinimize;
 	this->bHasMaximizeButton = InArgs._SupportsMaximize;
 	this->bHasSizingFrame = !InArgs._IsPopupWindow && InArgs._SizingRule == ESizingRule::UserSized;
+	this->bShouldPreserveAspectRatio = InArgs._ShouldPreserveAspectRatio;
 	this->LayoutBorder = InArgs._LayoutBorder;
 	this->UserResizeBorder = InArgs._UserResizeBorder;
 	this->SizeLimits = FWindowSizeLimits()
@@ -1725,6 +1726,7 @@ SWindow::SWindow()
 	, bHasMaximizeButton( false )
 	, bHasSizingFrame( false )
 	, bIsModalWindow( false )
+	, bShouldPreserveAspectRatio( false )
 	, InitialDesiredScreenPosition( FVector2D::ZeroVector )
 	, InitialDesiredSize( FVector2D::ZeroVector )
 	, ScreenPosition( FVector2D::ZeroVector )
diff --git a/Engine/Source/Runtime/SlateCore/Public/Widgets/SWindow.h b/Engine/Source/Runtime/SlateCore/Public/Widgets/SWindow.h
--- a/Engine/Source/Runtime/SlateCore/Public/Widgets/SWindow.h
+++ b/Engine/Source/Runtime/SlateCore/Public/Widgets/SWindow.h
@@ -113,6 +113,7 @@ public:
 		, _HasCloseButton( true )
 		, _SupportsMaximize( true )
 		, _SupportsMinimize( true )
+		, _ShouldPreserveAspectRatio( false )
 		, _CreateTitleBar( true )
 		, _SaneWindowPlacement( true )
 		, _LayoutBorder(FMargin(5, 5, 5, 5))
@@ -180,6 +181,9 @@ public:
 		/** Can this window be minimized? */
 		SLATE_ARGUMENT( bool, SupportsMinimize )
 
+		/** Should this window preserve its aspect ratio when resized by user? */
+		SLATE_ARGUMENT( bool, ShouldPreserveAspectRatio )
+
 		/** The smallest width this window can be in Desktop Pixel Units. */
 		SLATE_ARGUMENT( TOptional<float>, MinWidth )
 		
@@ -601,6 +605,11 @@ public:
 		return bSizeWillChangeOften;
 	}
 
+	bool ShouldPreserveAspectRatio() const
+	{
+		return bShouldPreserveAspectRatio;
+	}
+
 	/** @return Returns the configured expected maximum width of the window, or INDEX_NONE if not specified.  Can be used to optimize performance for window size animation */
 	int32 GetExpectedMaxWidth() const
 	{
@@ -901,6 +910,9 @@ protected:
 	/** True if the window is a mirror window for HMD content */
 	bool bIsMirrorWindow : 1;
 
+	/** True if the window should preserve its aspect ratio when resized by user */
+	bool bShouldPreserveAspectRatio : 1;
+
 	/** Initial desired position of the window's content in screen space */
 	FVector2D InitialDesiredScreenPosition;
 
