Copying //UE4/Dev-Rendering to //UE4/Dev-Main (Source: //UE4/Dev-Rendering @ 3054480)

diff --git a/Engine/Source/Runtime/OpenGLDrv/Private/Windows/OpenGLWindows.cpp b/Engine/Source/Runtime/OpenGLDrv/Private/Windows/OpenGLWindows.cpp
--- a/Engine/Source/Runtime/OpenGLDrv/Private/Windows/OpenGLWindows.cpp
+++ b/Engine/Source/Runtime/OpenGLDrv/Private/Windows/OpenGLWindows.cpp
@@ -82,6 +82,11 @@ public:
 		}
 	}
 
+	bool ContextWasAlreadyActive() const
+	{
+		return bSameDCAndContext;
+	}
+
 private:
 	HDC					PrevDC;
 	HGLRC				PrevContext;
@@ -182,8 +187,13 @@ static void PlatformCreateDummyGLWindow(FPlatformOpenGLContext* OutContext)
 static bool PlatformOpenGL3()
 {
 	// OpenGL3 is our default platform for Windows XP
-	return FWindowsPlatformMisc::VerifyWindowsVersion(6, 0) == false || FParse::Param(FCommandLine::Get(),TEXT("opengl")) || FParse::Param(FCommandLine::Get(),TEXT("opengl3"));
+#if (WINVER < 0x0600)
+	return true;
+#else
+	return FParse::Param(FCommandLine::Get(),TEXT("opengl")) || FParse::Param(FCommandLine::Get(),TEXT("opengl3"));
+#endif
 }
+
 static bool PlatformOpenGL4()
 {
 	return FParse::Param(FCommandLine::Get(),TEXT("opengl4"));
@@ -382,11 +392,15 @@ void PlatformReleaseOpenGLContext(FPlatformOpenGLDevice* Device, FPlatformOpenGL
 	Device->ViewportContexts.RemoveSingle(Context);
 	Device->TargetDirty = true;
 
+	bool bActiveContextWillBeReleased = false;
+
 	{
 		FScopeLock ScopeLock(Device->ContextUsageGuard);
 		{
 			FScopeContext ScopeContext(Context);
 
+			bActiveContextWillBeReleased = ScopeContext.ContextWasAlreadyActive();
+
 			DeleteQueriesForCurrentContext(Context->OpenGLContext);
 			glBindVertexArray(0);
 			glDeleteVertexArrays(1, &Context->VertexArrayObject);
@@ -404,7 +418,7 @@ void PlatformReleaseOpenGLContext(FPlatformOpenGLDevice* Device, FPlatformOpenGL
 
 	check(Context->DeviceContext);
 
-	if (wglGetCurrentDC() == Context->DeviceContext)
+	if (bActiveContextWillBeReleased)
 	{
 		wglMakeCurrent( NULL, NULL );
 	}
diff --git a/Engine/Source/Runtime/RHI/Private/RHI.cpp b/Engine/Source/Runtime/RHI/Private/RHI.cpp
--- a/Engine/Source/Runtime/RHI/Private/RHI.cpp
+++ b/Engine/Source/Runtime/RHI/Private/RHI.cpp
@@ -283,6 +283,7 @@ bool GRHISupportsRHIThread = false;
 bool GRHISupportsParallelRHIExecute = false;
 bool GSupportsHDR32bppEncodeModeIntrinsic = false;
 bool GSupportsParallelOcclusionQueries = false;
+bool GSupportsRenderTargetWriteMask = false;
 
 bool GRHISupportsMSAADepthSampleAccess = false;
 
diff --git a/Engine/Source/Runtime/RHI/Public/DynamicRHI.h b/Engine/Source/Runtime/RHI/Public/DynamicRHI.h
--- a/Engine/Source/Runtime/RHI/Public/DynamicRHI.h
+++ b/Engine/Source/Runtime/RHI/Public/DynamicRHI.h
@@ -686,6 +686,15 @@ public:
 	// FlushType: Wait RHI Thread
 	virtual FTexture2DRHIRef RHICreateTexture2D(uint32 SizeX, uint32 SizeY, uint8 Format, uint32 NumMips, uint32 NumSamples, uint32 Flags, FRHIResourceCreateInfo& CreateInfo) = 0;
 
+	/**
+	* Creates an FStructuredBuffer for the RT write mask of a render target 
+	* @param RenderTarget - the RT to create the buffer for
+	*/
+	virtual FStructuredBufferRHIRef RHICreateRTWriteMaskBuffer(FTexture2DRHIParamRef RenderTarget)
+	{
+		return nullptr;
+	}
+
 	/**
 	 * Thread-safe function that can be used to create a texture outside of the
 	 * rendering thread. This function can ONLY be called if GRHISupportsAsyncTextureCreation
@@ -1022,7 +1031,7 @@ public:
 	// FlushType: Flush Immediate (seems wrong)
 	virtual void RHISetStreamOutTargets(uint32 NumTargets, const FVertexBufferRHIParamRef* VertexBuffers, const uint32* Offsets) = 0;
 
-	// FlushType: Flush Immediate (seems wrong)
+	// Each RHI should flush if it needs to when implementing this method.
 	virtual void RHIDiscardRenderTargets(bool Depth, bool Stencil, uint32 ColorBitMask) = 0;
 
 	// Blocks the CPU until the GPU catches up and goes idle.
diff --git a/Engine/Source/Runtime/RHI/Public/RHI.h b/Engine/Source/Runtime/RHI/Public/RHI.h
--- a/Engine/Source/Runtime/RHI/Public/RHI.h
+++ b/Engine/Source/Runtime/RHI/Public/RHI.h
@@ -138,6 +138,9 @@ extern RHI_API bool GSupportsWideMRT;
 /** True if the RHI and current hardware supports supports depth bounds testing */
 extern RHI_API bool GSupportsDepthBoundsTest;
 
+/** True if the RHI and current hardware support a render target write mask */
+extern RHI_API bool GSupportsRenderTargetWriteMask;
+
 /** True if the RHI and current hardware supports efficient AsyncCompute (by default we assume false and later we can enable this for more hardware) */
 extern RHI_API bool GSupportsEfficientAsyncCompute;
 
@@ -1006,6 +1009,7 @@ enum class EResourceTransitionAccess
 	ERWBarrier, // Mostly for UAVs.  Transition to read/write state and always insert a resource barrier.
 	ERWNoBarrier, //Mostly UAVs.  Indicates we want R/W access and do not require synchronization for the duration of the RW state.  The initial transition from writable->RWNoBarrier and readable->RWNoBarrier still requires a sync
 	ERWSubResBarrier, //For special cases where read/write happens to different subresources of the same resource in the same call.  Inserts a barrier, but read validation will pass.  Temporary until we pass full subresource info to all transition calls.
+	EMetaData,		  // For transitioning texture meta data, for example for making readable in shaders
 	EMaxAccess,
 };
 
diff --git a/Engine/Source/Runtime/RHI/Public/RHICommandList.h b/Engine/Source/Runtime/RHI/Public/RHICommandList.h
--- a/Engine/Source/Runtime/RHI/Public/RHICommandList.h
+++ b/Engine/Source/Runtime/RHI/Public/RHICommandList.h
@@ -2552,7 +2552,12 @@ public:
 	{
 		return GDynamicRHI->RHICreateTexture2D_RenderThread(*this, SizeX, SizeY, Format, NumMips, NumSamples, Flags, CreateInfo);
 	}
-	
+
+	FORCEINLINE FStructuredBufferRHIRef CreateRTWriteMaskBuffer(FTexture2DRHIRef RenderTarget)
+	{
+		return GDynamicRHI->RHICreateRTWriteMaskBuffer(RenderTarget);
+	}
+
 	FORCEINLINE FTexture2DRHIRef AsyncCreateTexture2D(uint32 SizeX, uint32 SizeY, uint8 Format, uint32 NumMips, uint32 Flags, void** InitialMipData, uint32 NumInitialMips)
 	{
 		return GDynamicRHI->RHIAsyncCreateTexture2D(SizeX, SizeY, Format, NumMips, Flags, InitialMipData, NumInitialMips);
@@ -2818,8 +2823,6 @@ public:
 	
 	FORCEINLINE void DiscardRenderTargets(bool Depth,bool Stencil,uint32 ColorBitMask)
 	{
-		QUICK_SCOPE_CYCLE_COUNTER(STAT_RHIMETHOD_DiscardRenderTargets_Flush);
-		ImmediateFlush(EImmediateFlushType::FlushRHIThread);  
 		GDynamicRHI->RHIDiscardRenderTargets(Depth,Stencil,ColorBitMask);
 	}
 	
@@ -3184,6 +3187,11 @@ FORCEINLINE FTexture2DRHIRef RHICreateTexture2D(uint32 SizeX, uint32 SizeY, uint
 	return FRHICommandListExecutor::GetImmediateCommandList().CreateTexture2D(SizeX, SizeY, Format, NumMips, NumSamples, Flags, CreateInfo);
 }
 
+FORCEINLINE FStructuredBufferRHIRef RHICreateRTWriteMaskBuffer(FTexture2DRHIRef RenderTarget)
+{
+	return FRHICommandListExecutor::GetImmediateCommandList().CreateRTWriteMaskBuffer(RenderTarget);
+}
+
 FORCEINLINE FTexture2DRHIRef RHIAsyncCreateTexture2D(uint32 SizeX, uint32 SizeY, uint8 Format, uint32 NumMips, uint32 Flags, void** InitialMipData, uint32 NumInitialMips)
 {
 	return FRHICommandListExecutor::GetImmediateCommandList().AsyncCreateTexture2D(SizeX, SizeY, Format, NumMips, Flags, InitialMipData, NumInitialMips);
diff --git a/Engine/Source/Runtime/RHI/Public/RHIDefinitions.h b/Engine/Source/Runtime/RHI/Public/RHIDefinitions.h
--- a/Engine/Source/Runtime/RHI/Public/RHIDefinitions.h
+++ b/Engine/Source/Runtime/RHI/Public/RHIDefinitions.h
@@ -542,6 +542,8 @@ enum ETextureCreateFlags
 	TexCreate_NoFastClear = 1 << 25,
 	// Texture is a depth stencil resolve target
 	TexCreate_DepthStencilResolveTarget = 1 << 26,
+	// Render target will not FinalizeFastClear; Caches and meta data will be flushed, but clearing will be skipped (avoids potentially trashing metadata)
+	TexCreate_NoFastClearFinalize = 1 << 28,
 };
 
 enum EAsyncComputePriority
diff --git a/Engine/Source/Runtime/RenderCore/Private/RenderUtils.cpp b/Engine/Source/Runtime/RenderCore/Private/RenderUtils.cpp
--- a/Engine/Source/Runtime/RenderCore/Private/RenderUtils.cpp
+++ b/Engine/Source/Runtime/RenderCore/Private/RenderUtils.cpp
@@ -124,6 +124,7 @@ FPixelFormatInfo	GPixelFormats[PF_MAX] =
 
 	{ TEXT("BC6H"),				4,			4,			1,			16,			3,				0,				1,				PF_BC6H				},
 	{ TEXT("BC7"),				4,			4,			1,			16,			4,				0,				1,				PF_BC7				},
+	{ TEXT("R8_UINT"),			1,			1,			1,			1,			1,				0,				1,				PF_R8_UINT },
 };
 
 static struct FValidatePixelFormats
diff --git a/Engine/Source/Runtime/Renderer/Private/AtmosphereRendering.cpp b/Engine/Source/Runtime/Renderer/Private/AtmosphereRendering.cpp
--- a/Engine/Source/Runtime/Renderer/Private/AtmosphereRendering.cpp
+++ b/Engine/Source/Runtime/Renderer/Private/AtmosphereRendering.cpp
@@ -13,6 +13,8 @@
 #include "ShowFlags.h"
 #include "SceneUtils.h"
 
+DECLARE_FLOAT_COUNTER_STAT(TEXT("Atmosphere"), Stat_GPU_Atmosphere, STATGROUP_GPU);
+
 class FAtmosphereShaderPrecomputeTextureParameters
 {
 public:
@@ -409,6 +411,7 @@ void FDeferredShadingSceneRenderer::RenderAtmosphere(FRHICommandListImmediate& R
 			const FViewInfo& View = Views[ViewIndex];
 
 			SCOPED_DRAW_EVENTF(RHICmdList, Atmosphere, TEXT("Atmosphere %dx%d"), View.ViewRect.Width(), View.ViewRect.Height());
+			SCOPED_GPU_STAT(RHICmdList, Stat_GPU_Atmosphere);
 
 			// Set the device viewport for the view.
 			RHICmdList.SetViewport(View.ViewRect.Min.X, View.ViewRect.Min.Y, 0.0f, View.ViewRect.Max.X, View.ViewRect.Max.Y, 1.0f);
diff --git a/Engine/Source/Runtime/Renderer/Private/BasePassRendering.cpp b/Engine/Source/Runtime/Renderer/Private/BasePassRendering.cpp
--- a/Engine/Source/Runtime/Renderer/Private/BasePassRendering.cpp
+++ b/Engine/Source/Runtime/Renderer/Private/BasePassRendering.cpp
@@ -123,6 +123,8 @@ IMPLEMENT_BASEPASS_LIGHTMAPPED_SHADER_TYPE( TUniformLightMapPolicy<LMP_LQ_LIGHTM
 IMPLEMENT_BASEPASS_LIGHTMAPPED_SHADER_TYPE( TUniformLightMapPolicy<LMP_HQ_LIGHTMAP>, TLightMapPolicyHQ );
 IMPLEMENT_BASEPASS_LIGHTMAPPED_SHADER_TYPE( TUniformLightMapPolicy<LMP_DISTANCE_FIELD_SHADOWS_AND_HQ_LIGHTMAP>, TDistanceFieldShadowsAndLightMapPolicyHQ  );
 
+DECLARE_FLOAT_COUNTER_STAT(TEXT("Basepass"), Stat_GPU_Basepass, STATGROUP_GPU);
+
 void FSkyLightReflectionParameters::GetSkyParametersFromScene(
 	const FScene* Scene, 
 	bool bApplySkyLight, 
@@ -729,6 +731,7 @@ bool FDeferredShadingSceneRenderer::RenderBasePass(FRHICommandListImmediate& RHI
 	{
 		SCOPED_DRAW_EVENT(RHICmdList, BasePass);
 		SCOPE_CYCLE_COUNTER(STAT_BasePassDrawTime);
+		SCOPED_GPU_STAT(RHICmdList, Stat_GPU_Basepass );
 
 		if (GRHICommandList.UseParallelAlgorithms() && CVarParallelBasePass.GetValueOnRenderThread())
 		{
diff --git a/Engine/Source/Runtime/Renderer/Private/BasePassRendering.h b/Engine/Source/Runtime/Renderer/Private/BasePassRendering.h
--- a/Engine/Source/Runtime/Renderer/Private/BasePassRendering.h
+++ b/Engine/Source/Runtime/Renderer/Private/BasePassRendering.h
@@ -996,6 +996,10 @@ public:
 				// Modulate with the existing scene color, preserve destination alpha.
 				RHICmdList.SetBlendState( TStaticBlendState<CW_RGB, BO_Add, BF_DestColor, BF_Zero>::GetRHI());
 				break;
+			case BLEND_AlphaComposite:
+				// Blend with existing scene color. New color is already pre-multiplied by alpha.
+				RHICmdList.SetBlendState(TStaticBlendState<CW_RGBA, BO_Add, BF_One, BF_InverseSourceAlpha, BO_Add, BF_Zero, BF_InverseSourceAlpha>::GetRHI());
+				break;
 			};
 		}
 	}
diff --git a/Engine/Source/Runtime/Renderer/Private/CapsuleShadowRendering.cpp b/Engine/Source/Runtime/Renderer/Private/CapsuleShadowRendering.cpp
--- a/Engine/Source/Runtime/Renderer/Private/CapsuleShadowRendering.cpp
+++ b/Engine/Source/Runtime/Renderer/Private/CapsuleShadowRendering.cpp
@@ -15,6 +15,8 @@
 #include "DistanceFieldSurfaceCacheLighting.h"
 #include "DistanceFieldLightingPost.h"
 
+DECLARE_FLOAT_COUNTER_STAT(TEXT("Capsule Shadows"), Stat_GPU_CapsuleShadows, STATGROUP_GPU);
+
 int32 GCapsuleShadows = 1;
 FAutoConsoleVariableRef CVarCapsuleShadows(
 	TEXT("r.CapsuleShadows"),
@@ -645,6 +647,7 @@ bool FDeferredShadingSceneRenderer::RenderCapsuleDirectShadows(
 		{
 			const FViewInfo& View = Views[ViewIndex];
 			SCOPED_DRAW_EVENT(RHICmdList, CapsuleShadows);
+			SCOPED_GPU_STAT(RHICmdList, Stat_GPU_CapsuleShadows);
 
 			static TArray<FCapsuleShape> CapsuleShapeData;
 			CapsuleShapeData.Reset();
@@ -1059,6 +1062,7 @@ void FDeferredShadingSceneRenderer::RenderIndirectCapsuleShadows(
 				if (View.IndirectShadowPrimitives.Num() > 0 && View.ViewState)
 				{
 					SCOPED_DRAW_EVENT(RHICmdList, IndirectCapsuleShadows);
+					SCOPED_GPU_STAT(RHICmdList, Stat_GPU_CapsuleShadows);
 		
 					int32 NumCapsuleShapes = 0;
 					SetupIndirectCapsuleShadows(RHICmdList, View, true, NumCapsuleShapes);
@@ -1218,6 +1222,7 @@ void FDeferredShadingSceneRenderer::RenderCapsuleShadowsForMovableSkylight(FRHIC
 				if (View.IndirectShadowPrimitives.Num() > 0 && View.ViewState)
 				{
 					SCOPED_DRAW_EVENT(RHICmdList, IndirectCapsuleShadows);
+					SCOPED_GPU_STAT(RHICmdList, Stat_GPU_CapsuleShadows);
 		
 					int32 NumCapsuleShapes = 0;
 					SetupIndirectCapsuleShadows(RHICmdList, View, true, NumCapsuleShapes);
diff --git a/Engine/Source/Runtime/Renderer/Private/CompositionLighting/CompositionLighting.cpp b/Engine/Source/Runtime/Renderer/Private/CompositionLighting/CompositionLighting.cpp
--- a/Engine/Source/Runtime/Renderer/Private/CompositionLighting/CompositionLighting.cpp
+++ b/Engine/Source/Runtime/Renderer/Private/CompositionLighting/CompositionLighting.cpp
@@ -26,6 +26,11 @@
 /** The global center for all deferred lighting activities. */
 FCompositionLighting GCompositionLighting;
 
+DECLARE_FLOAT_COUNTER_STAT(TEXT("Composition BeforeBasePass"), Stat_GPU_CompositionBeforeBasePass, STATGROUP_GPU);
+DECLARE_FLOAT_COUNTER_STAT(TEXT("Composition PreLighting"), Stat_GPU_CompositionPreLighting, STATGROUP_GPU);
+DECLARE_FLOAT_COUNTER_STAT(TEXT("Composition LpvIndirect"), Stat_GPU_CompositionLpvIndirect, STATGROUP_GPU);
+DECLARE_FLOAT_COUNTER_STAT(TEXT("Composition PostLighting"), Stat_GPU_CompositionPostLighting, STATGROUP_GPU);
+
 // -------------------------------------------------------
 
 static TAutoConsoleVariable<float> CVarSSSScale(
@@ -120,6 +125,7 @@ uint32 ComputeAmbientOcclusionPassCount(const FViewInfo& View)
 	if (!IsLpvIndirectPassRequired(View))
 	{
 		bEnabled = View.FinalPostProcessSettings.AmbientOcclusionIntensity > 0
+			&& View.Family->EngineShowFlags.Lighting
 			&& View.FinalPostProcessSettings.AmbientOcclusionRadius >= 0.1f
 			&& !View.Family->UseDebugViewPS()
 			&& (FSSAOHelper::IsBasePassAmbientOcclusionRequired(View) || IsAmbientCubemapPassRequired(View) || IsReflectionEnvironmentActive(View) || IsSkylightActive(View) || View.Family->EngineShowFlags.VisualizeBuffer)
@@ -281,6 +287,7 @@ void FCompositionLighting::ProcessBeforeBasePass(FRHICommandListImmediate& RHICm
 		// The graph setup should be finished before this line ----------------------------------------
 
 		SCOPED_DRAW_EVENT(RHICmdList, CompositionBeforeBasePass);
+		SCOPED_GPU_STAT(RHICmdList, Stat_GPU_CompositionBeforeBasePass);
 
 		CompositeContext.Process(Context.FinalOutput.GetPass(), TEXT("Composition_BeforeBasePass"));
 	}
@@ -360,6 +367,7 @@ void FCompositionLighting::ProcessAfterBasePass(FRHICommandListImmediate& RHICmd
 		// The graph setup should be finished before this line ----------------------------------------
 
 		SCOPED_DRAW_EVENT(RHICmdList, LightCompositionTasks_PreLighting);
+		SCOPED_GPU_STAT(RHICmdList, Stat_GPU_CompositionPreLighting);
 
 		TRefCountPtr<IPooledRenderTarget>& SceneColor = SceneContext.GetSceneColor();
 
@@ -394,6 +402,7 @@ void FCompositionLighting::ProcessLpvIndirect(FRHICommandListImmediate& RHICmdLi
 	// The graph setup should be finished before this line ----------------------------------------
 
 	SCOPED_DRAW_EVENT(RHICmdList, CompositionLpvIndirect);
+	SCOPED_GPU_STAT(RHICmdList, Stat_GPU_CompositionLpvIndirect);
 
 	// we don't replace the final element with the scenecolor because this is what those passes should do by themself
 
@@ -457,6 +466,7 @@ void FCompositionLighting::ProcessAfterLighting(FRHICommandListImmediate& RHICmd
 		// The graph setup should be finished before this line ----------------------------------------
 
 		SCOPED_DRAW_EVENT(RHICmdList, CompositionAfterLighting);
+		SCOPED_GPU_STAT(RHICmdList, Stat_GPU_CompositionPostLighting);
 
 		// we don't replace the final element with the scenecolor because this is what those passes should do by themself
 
diff --git a/Engine/Source/Runtime/Renderer/Private/CompositionLighting/PostProcessDeferredDecals.cpp b/Engine/Source/Runtime/Renderer/Private/CompositionLighting/PostProcessDeferredDecals.cpp
--- a/Engine/Source/Runtime/Renderer/Private/CompositionLighting/PostProcessDeferredDecals.cpp
+++ b/Engine/Source/Runtime/Renderer/Private/CompositionLighting/PostProcessDeferredDecals.cpp
@@ -13,6 +13,89 @@
 #include "SceneUtils.h"
 #include "DecalRenderingShared.h"
 
+
+
+
+static TAutoConsoleVariable<int32> CVarGenerateDecalRTWriteMaskTexture(
+	TEXT("r.Decal.GenerateRTWriteMaskTexture"),
+	1,
+	TEXT("Turn on or off generation of the RT write mask texture for decals\n"),
+	ECVF_Default);
+
+
+
+
+class FRTWriteMaskDecodeCS : public FGlobalShader
+{
+	DECLARE_SHADER_TYPE(FRTWriteMaskDecodeCS, Global);
+
+	static bool ShouldCache(EShaderPlatform Platform)
+	{
+		return IsFeatureLevelSupported(Platform, ERHIFeatureLevel::SM5);
+	}
+
+	static void ModifyCompilationEnvironment(EShaderPlatform Platform, FShaderCompilerEnvironment& OutEnvironment)
+	{
+		FGlobalShader::ModifyCompilationEnvironment(Platform, OutEnvironment);
+	}
+
+	FRTWriteMaskDecodeCS() {}
+
+public:
+	FShaderParameter OutCombinedRTWriteMask;		// UAV
+	FShaderResourceParameter RTWriteMaskInput0;				// SRV 
+	FShaderResourceParameter RTWriteMaskInput1;				// SRV 
+	FShaderResourceParameter RTWriteMaskInput2;				// SRV 
+	FShaderParameter UtilizeMask;				// SRV 
+
+	FRTWriteMaskDecodeCS(const ShaderMetaType::CompiledShaderInitializerType& Initializer)
+		: FGlobalShader(Initializer)
+	{
+		RTWriteMaskDimensions.Bind(Initializer.ParameterMap, TEXT("RTWriteMaskDimensions"));
+		OutCombinedRTWriteMask.Bind(Initializer.ParameterMap, TEXT("OutCombinedRTWriteMask"));
+		RTWriteMaskInput0.Bind(Initializer.ParameterMap, TEXT("RTWriteMaskInput0"));
+		RTWriteMaskInput1.Bind(Initializer.ParameterMap, TEXT("RTWriteMaskInput1"));
+		RTWriteMaskInput2.Bind(Initializer.ParameterMap, TEXT("RTWriteMaskInput2"));
+		UtilizeMask.Bind(Initializer.ParameterMap, TEXT("UtilizeMask"));
+	}
+
+	void SetCS(FRHICommandList& RHICmdList, const FRenderingCompositePassContext& Context, const FSceneView& View)
+	{
+		const FComputeShaderRHIParamRef ShaderRHI = GetComputeShader();
+
+		FGlobalShader::SetParameters(RHICmdList, ShaderRHI, Context.View);
+		//PostprocessParameter.SetCS(ShaderRHI, Context, Context.RHICmdList, TStaticSamplerState<SF_Point, AM_Clamp, AM_Clamp, AM_Clamp>::GetRHI());
+
+		FSceneRenderTargets& SceneContext = FSceneRenderTargets::Get(RHICmdList);
+		FTextureRHIRef DBufferTex = SceneContext.DBufferA->GetRenderTargetItem().TargetableTexture;
+		FIntPoint Dims(FPlatformMath::CeilToInt(DBufferTex->GetTexture2D()->GetSizeX() / 8.0f), FPlatformMath::CeilToInt(DBufferTex->GetTexture2D()->GetSizeY() / 8.0f));
+
+		SetShaderValue(Context.RHICmdList, ShaderRHI, RTWriteMaskDimensions, Dims);
+		SetSRVParameter(Context.RHICmdList, ShaderRHI, RTWriteMaskInput0, SceneContext.DBufferA->GetRenderTargetItem().RTWriteMaskBufferRHI_SRV);
+		SetSRVParameter(Context.RHICmdList, ShaderRHI, RTWriteMaskInput1, SceneContext.DBufferB->GetRenderTargetItem().RTWriteMaskBufferRHI_SRV);
+		SetSRVParameter(Context.RHICmdList, ShaderRHI, RTWriteMaskInput2, SceneContext.DBufferC->GetRenderTargetItem().RTWriteMaskBufferRHI_SRV);
+		int32 UseMask = CVarGenerateDecalRTWriteMaskTexture.GetValueOnRenderThread();
+		SetShaderValue(Context.RHICmdList, ShaderRHI, UtilizeMask, UseMask);
+	}
+
+	virtual bool Serialize(FArchive& Ar) override
+	{
+		bool bShaderHasOutdatedParameters = FGlobalShader::Serialize(Ar);
+		Ar << RTWriteMaskDimensions;
+		Ar << OutCombinedRTWriteMask;
+		Ar << RTWriteMaskInput0;
+		Ar << RTWriteMaskInput1;
+		Ar << RTWriteMaskInput2;
+		Ar << UtilizeMask;
+		return bShaderHasOutdatedParameters;
+	}
+
+private:
+	FShaderParameter			RTWriteMaskDimensions;
+};
+
+IMPLEMENT_SHADER_TYPE(, FRTWriteMaskDecodeCS, TEXT("RTWriteMaskDecode"), TEXT("RTWriteMaskCombineMain"), SF_Compute);
+
 static TAutoConsoleVariable<float> CVarStencilSizeThreshold(
 	TEXT("r.Decal.StencilSizeThreshold"),
 	0.1f,
@@ -477,6 +560,52 @@ const TCHAR* GetStageName(EDecalRenderStage Stage)
 	return TEXT("<UNKNOWN>");
 }
 
+
+void FRCPassPostProcessDeferredDecals::DecodeRTWriteMask(FRenderingCompositePassContext& Context)
+{
+	if (GSupportsRenderTargetWriteMask)
+	{
+		FRHICommandListImmediate& RHICmdList = Context.RHICmdList;
+		FSceneRenderTargets& SceneContext = FSceneRenderTargets::Get(RHICmdList);
+		FTextureRHIRef DBufferTex = SceneContext.DBufferA->GetRenderTargetItem().TargetableTexture;
+		int32 W = FPlatformMath::CeilToInt(DBufferTex->GetTexture2D()->GetSizeX() / 16.0f);
+		int32 H = FPlatformMath::CeilToInt(DBufferTex->GetTexture2D()->GetSizeY() / 8.0f);
+		FIntPoint RTWriteMaskDims(W, H);
+
+		FIntRect ViewRect(0, 0, DBufferTex->GetTexture2D()->GetSizeX(), DBufferTex->GetTexture2D()->GetSizeY());
+
+		TShaderMapRef< FRTWriteMaskDecodeCS > ComputeShader(Context.GetShaderMap());
+
+		SetRenderTarget(Context.RHICmdList, FTextureRHIRef(), FTextureRHIRef());
+		Context.SetViewportAndCallRHI(ViewRect);
+		Context.RHICmdList.SetComputeShader(ComputeShader->GetComputeShader());
+
+		// set destination
+		Context.RHICmdList.SetUAVParameter(ComputeShader->GetComputeShader(), ComputeShader->OutCombinedRTWriteMask.GetBaseIndex(), SceneContext.DBufferMask->GetRenderTargetItem().UAV);
+		ComputeShader->SetCS(Context.RHICmdList, Context, Context.View);
+
+		RHICmdList.TransitionResource(EResourceTransitionAccess::EWritable, EResourceTransitionPipeline::EGfxToCompute, SceneContext.DBufferMask->GetRenderTargetItem().UAV);
+		{
+			SCOPED_DRAW_EVENTF(Context.RHICmdList, DeferredDecals, TEXT("Combine DBuffer RTWriteMasks"));
+			FIntPoint ThreadGroupCountValue(FPlatformMath::CeilToInt(RTWriteMaskDims.X / 8.0f), FPlatformMath::CeilToInt(RTWriteMaskDims.Y / 16.0f));
+			DispatchComputeShader(Context.RHICmdList, *ComputeShader, ThreadGroupCountValue.X, ThreadGroupCountValue.Y, 1);
+		}
+
+		//	void FD3D11DynamicRHI::RHIGraphicsWaitOnAsyncComputeJob( uint32 FenceIndex )
+		Context.RHICmdList.FlushComputeShaderCache();
+
+		RHICmdList.TransitionResource(EResourceTransitionAccess::EReadable, EResourceTransitionPipeline::EComputeToGfx, SceneContext.DBufferMask->GetRenderTargetItem().UAV);
+
+		RHICmdList.TransitionResource(EResourceTransitionAccess::EMetaData, SceneContext.DBufferA->GetRenderTargetItem().TargetableTexture);
+		RHICmdList.TransitionResource(EResourceTransitionAccess::EMetaData, SceneContext.DBufferB->GetRenderTargetItem().TargetableTexture);
+		RHICmdList.TransitionResource(EResourceTransitionAccess::EMetaData, SceneContext.DBufferC->GetRenderTargetItem().TargetableTexture);
+
+		// un-set destination
+		Context.RHICmdList.SetUAVParameter(ComputeShader->GetComputeShader(), ComputeShader->OutCombinedRTWriteMask.GetBaseIndex(), NULL);
+	}
+}
+
+
 void FRCPassPostProcessDeferredDecals::Process(FRenderingCompositePassContext& Context)
 {
 	FRHICommandListImmediate& RHICmdList = Context.RHICmdList;
@@ -488,27 +617,45 @@ void FRCPassPostProcessDeferredDecals::Process(FRenderingCompositePassContext& C
 
 	SCOPED_DRAW_EVENTF(RHICmdList, DeferredDecals, TEXT("DeferredDecals %s"), GetStageName(CurrentStage));
 
-	if(CurrentStage == DRS_BeforeBasePass)
+	if (CurrentStage == DRS_BeforeBasePass)
 	{
 		// before BasePass, only if DBuffer is enabled
 
 		check(bDBuffer);
 
 		// DBuffer: Decal buffer
-		FPooledRenderTargetDesc Desc(FPooledRenderTargetDesc::Create2DDesc(SceneContext.GBufferA->GetDesc().Extent, 
+		FPooledRenderTargetDesc Desc(FPooledRenderTargetDesc::Create2DDesc(SceneContext.GBufferA->GetDesc().Extent,
 			PF_B8G8R8A8,
 			FClearValueBinding::None,
-			TexCreate_None, 
+			TexCreate_None,
 			TexCreate_ShaderResource | TexCreate_RenderTargetable,
-			false));
+			false,
+			1, 
+			true, 
+			true));
+
+		if (GSupportsRenderTargetWriteMask)
+		{
+			int32 W = FPlatformMath::CeilToInt(SceneContext.GBufferA->GetDesc().Extent.X / 8.0f);
+			int32 H = FPlatformMath::CeilToInt(SceneContext.GBufferA->GetDesc().Extent.Y / 8.0f);
+			FIntPoint RTWriteMaskDims(W, H);
+			FPooledRenderTargetDesc MaskDesc(FPooledRenderTargetDesc::Create2DDesc(RTWriteMaskDims,
+				PF_R8_UINT,
+				FClearValueBinding::White,
+				TexCreate_None,
+				TexCreate_UAV | TexCreate_RenderTargetable,
+				false));
+
+			GRenderTargetPool.FindFreeElement(Context.RHICmdList, MaskDesc, SceneContext.DBufferMask, TEXT("DBufferMask"));
+		}
 
-		if(!SceneContext.DBufferA)
+		if (!SceneContext.DBufferA)
 		{
 			Desc.ClearValue = FClearValueBinding::Black;
 			GRenderTargetPool.FindFreeElement(RHICmdList, Desc, SceneContext.DBufferA, TEXT("DBufferA"));
 		}
 
-		if(!SceneContext.DBufferB)
+		if (!SceneContext.DBufferB)
 		{
 			Desc.ClearValue = FClearValueBinding(FLinearColor(128.0f / 255.0f, 128.0f / 255.0f, 128.0f / 255.0f, 1));
 			GRenderTargetPool.FindFreeElement(RHICmdList, Desc, SceneContext.DBufferB, TEXT("DBufferB"));
@@ -516,7 +663,7 @@ void FRCPassPostProcessDeferredDecals::Process(FRenderingCompositePassContext& C
 
 		Desc.Format = PF_R8G8;
 
-		if(!SceneContext.DBufferC)
+		if (!SceneContext.DBufferC)
 		{
 			Desc.ClearValue = FClearValueBinding(FLinearColor(0, 1, 0, 1));
 			GRenderTargetPool.FindFreeElement(RHICmdList, Desc, SceneContext.DBufferC, TEXT("DBufferC"));
@@ -526,7 +673,7 @@ void FRCPassPostProcessDeferredDecals::Process(FRenderingCompositePassContext& C
 		// and don't get FastClear any more.
 		bool bFirstView = Context.View.Family->Views[0] == &Context.View;
 
-		if(bFirstView)
+		if (bFirstView)
 		{
 			SCOPED_DRAW_EVENT(RHICmdList, DBufferClear);
 
@@ -680,7 +827,6 @@ void FRCPassPostProcessDeferredDecals::Process(FRenderingCompositePassContext& C
 					FDecalRendering::SetShader(RHICmdList, View, DecalData, FrustumComponentToClip);
 
 					RHICmdList.DrawIndexedPrimitive(GetUnitCubeIndexBuffer(), PT_TriangleList, 0, 0, 8, 0, ARRAY_COUNT(GCubeIndices) / 3, 1);
-		
 					RenderTargetManager.bGufferADirty |= (RenderTargetManager.TargetsToResolve[FDecalRenderTargetManager::GBufferAIndex] != nullptr);
 				}
 
@@ -688,6 +834,15 @@ void FRCPassPostProcessDeferredDecals::Process(FRenderingCompositePassContext& C
 				// Clear stencil to 0, which is the assumed default by other passes
 				RHICmdList.Clear(false, FLinearColor::White, false, (float)ERHIZBuffer::FarPlane, true, 0, FIntRect());
 			}
+
+			if (CurrentStage == DRS_BeforeBasePass)
+			{
+				// combine DBuffer RTWriteMasks; will end up in one texture we can load from in the base pass PS and decide whether to do the actual work or not
+				RenderTargetManager.FlushMetaData();
+				DecodeRTWriteMask(Context);
+			}
+
+			RenderTargetManager.ResolveTargets();
 		}
 
 		if (CurrentStage == DRS_BeforeBasePass)
@@ -708,4 +863,4 @@ FPooledRenderTargetDesc FRCPassPostProcessDeferredDecals::ComputeOutputDesc(EPas
 	Ret.DebugName = TEXT("DeferredDecals");
 
 	return Ret;
-}
\ No newline at end of file
+}
diff --git a/Engine/Source/Runtime/Renderer/Private/CompositionLighting/PostProcessDeferredDecals.h b/Engine/Source/Runtime/Renderer/Private/CompositionLighting/PostProcessDeferredDecals.h
--- a/Engine/Source/Runtime/Renderer/Private/CompositionLighting/PostProcessDeferredDecals.h
+++ b/Engine/Source/Runtime/Renderer/Private/CompositionLighting/PostProcessDeferredDecals.h
@@ -25,6 +25,7 @@ public:
 private:
 	// see EDecalRenderStage
 	EDecalRenderStage CurrentStage;
+	void DecodeRTWriteMask(FRenderingCompositePassContext& Context);
 };
 
 bool IsDBufferEnabled();
@@ -88,6 +89,11 @@ struct FDecalRenderTargetManager
 
 	// destructor
 	~FDecalRenderTargetManager()
+	{
+
+	}
+
+	void ResolveTargets()
 	{
 		FSceneRenderTargets& SceneContext = FSceneRenderTargets::Get(RHICmdList);
 
@@ -105,8 +111,12 @@ struct FDecalRenderTargetManager
 			{
 				RHICmdList.CopyToResolveTarget(TargetsToResolve[i], TargetsToResolve[i], true, ResolveParams);
 			}
-		}		
+		}
+	}
 
+	void FlushMetaData()
+	{
+		RHICmdList.TransitionResource(EResourceTransitionAccess::EMetaData, nullptr);
 	}
 
 	void SetRenderTargetMode(FDecalRenderingCommon::ERenderTargetMode CurrentRenderTargetMode, bool bHasNormal)
diff --git a/Engine/Source/Runtime/Renderer/Private/DeferredShadingRenderer.cpp b/Engine/Source/Runtime/Renderer/Private/DeferredShadingRenderer.cpp
--- a/Engine/Source/Runtime/Renderer/Private/DeferredShadingRenderer.cpp
+++ b/Engine/Source/Runtime/Renderer/Private/DeferredShadingRenderer.cpp
@@ -128,6 +128,10 @@ DECLARE_CYCLE_STAT(TEXT("DeferredShadingSceneRenderer RenderFinish"), STAT_FDefe
 DECLARE_CYCLE_STAT(TEXT("OcclusionSubmittedFence Dispatch"), STAT_OcclusionSubmittedFence_Dispatch, STATGROUP_SceneRendering);
 DECLARE_CYCLE_STAT(TEXT("OcclusionSubmittedFence Wait"), STAT_OcclusionSubmittedFence_Wait, STATGROUP_SceneRendering);
 
+DECLARE_FLOAT_COUNTER_STAT(TEXT("Postprocessing"), Stat_GPU_Postprocessing, STATGROUP_GPU);
+DECLARE_FLOAT_COUNTER_STAT(TEXT("HZB"), Stat_GPU_HZB, STATGROUP_GPU);
+DECLARE_FLOAT_COUNTER_STAT(TEXT("[unaccounted]"), Stat_GPU_Unaccounted, STATGROUP_GPU);
+
 bool ShouldForceFullDepthPass(ERHIFeatureLevel::Type FeatureLevel)
 {
 	static IConsoleVariable* CDBufferVar = IConsoleManager::Get().FindConsoleVariable(TEXT("r.DBuffer"));
@@ -431,6 +435,8 @@ void FDeferredShadingSceneRenderer::RenderOcclusion(FRHICommandListImmediate& RH
 	{
 		FSceneRenderTargets& SceneContext = FSceneRenderTargets::Get(RHICmdList);
 
+		SCOPED_GPU_STAT(RHICmdList, Stat_GPU_HZB);
+
 		{
 			// Update the quarter-sized depth buffer with the current contents of the scene depth texture.
 			// This needs to happen before occlusion tests, which makes use of the small depth buffer.
@@ -543,6 +549,10 @@ void FDeferredShadingSceneRenderer::Render(FRHICommandListImmediate& RHICmdList)
 		return;
 	}
 	SCOPED_DRAW_EVENT(RHICmdList, Scene);
+
+	// Anything rendered inside Render() which isn't accounted for will fall into this stat
+	// This works because child stat events do not contribute to their parents' times (see GPU_STATS_CHILD_TIMES_INCLUDED)
+	SCOPED_GPU_STAT(RHICmdList, Stat_GPU_Unaccounted);
 	
 	{
 		SCOPE_CYCLE_COUNTER(STAT_FDeferredShadingSceneRenderer_Render_Init);
@@ -1186,6 +1196,8 @@ void FDeferredShadingSceneRenderer::Render(FRHICommandListImmediate& RHICmdList)
 	if (ViewFamily.bResolveScene)
 	{
 		SCOPED_DRAW_EVENT(RHICmdList, PostProcessing);
+   		SCOPED_GPU_STAT(RHICmdList, Stat_GPU_Postprocessing);
+
 		SCOPE_CYCLE_COUNTER(STAT_FinishRenderViewTargetTime);
 
 		RHICmdList.SetCurrentStat(GET_STATID(STAT_CLM_PostProcessing));
@@ -1458,4 +1470,4 @@ void FDeferredShadingSceneRenderer::CopyStencilToLightingChannelTexture(FRHIComm
 			true,
 			ResolveParams);
 	}
-}
\ No newline at end of file
+}
diff --git a/Engine/Source/Runtime/Renderer/Private/DepthRendering.cpp b/Engine/Source/Runtime/Renderer/Private/DepthRendering.cpp
--- a/Engine/Source/Runtime/Renderer/Private/DepthRendering.cpp
+++ b/Engine/Source/Runtime/Renderer/Private/DepthRendering.cpp
@@ -45,6 +45,8 @@ const TCHAR* GetDepthDrawingModeString(EDepthDrawingMode Mode)
 	return TEXT("");
 }
 
+DECLARE_FLOAT_COUNTER_STAT(TEXT("Prepass"), Stat_GPU_Prepass, STATGROUP_GPU);
+
 /**
  * A vertex shader for rendering the depth of a mesh.
  */
@@ -196,19 +198,33 @@ public:
 
 	FDepthOnlyPS(const ShaderMetaType::CompiledShaderInitializerType& Initializer):
 		FMeshMaterialShader(Initializer)
-	{}
+	{
+		ApplyDepthOffsetParameter.Bind(Initializer.ParameterMap, TEXT("bApplyDepthOffset"));
+	}
 
 	FDepthOnlyPS() {}
 
 	void SetParameters(FRHICommandList& RHICmdList, const FMaterialRenderProxy* MaterialRenderProxy,const FMaterial& MaterialResource,const FSceneView* View)
 	{
 		FMeshMaterialShader::SetParameters(RHICmdList, GetPixelShader(),MaterialRenderProxy,MaterialResource,*View,ESceneRenderTargetsMode::DontSet);
+
+		// For debug view shaders, don't apply the depth offset as their base pass PS are using global shaders with depth equal.
+		SetShaderValue(RHICmdList, GetPixelShader(), ApplyDepthOffsetParameter, !View || !View->Family->UseDebugViewPS());
 	}
 
 	void SetMesh(FRHICommandList& RHICmdList, const FVertexFactory* VertexFactory,const FSceneView& View,const FPrimitiveSceneProxy* Proxy,const FMeshBatchElement& BatchElement,const FMeshDrawingRenderState& DrawRenderState)
 	{
 		FMeshMaterialShader::SetMesh(RHICmdList, GetPixelShader(),VertexFactory,View,Proxy,BatchElement,DrawRenderState);
 	}
+
+	virtual bool Serialize(FArchive& Ar) override
+	{
+		bool bShaderHasOutdatedParameters = FMeshMaterialShader::Serialize(Ar);
+		Ar << ApplyDepthOffsetParameter;
+		return bShaderHasOutdatedParameters;
+	}
+
+	FShaderParameter ApplyDepthOffsetParameter;
 };
 
 IMPLEMENT_MATERIAL_SHADER_TYPE(,FDepthOnlyPS,TEXT("DepthOnlyPixelShader"),TEXT("Main"),SF_Pixel);
@@ -1071,6 +1087,7 @@ bool FDeferredShadingSceneRenderer::RenderPrePass(FRHICommandListImmediate& RHIC
 	SCOPED_DRAW_EVENTF(RHICmdList, PrePass, TEXT("PrePass %s %s"), GetDepthDrawingModeString(EarlyZPassMode), GetDepthPassReason(bDitheredLODTransitionsUseStencil, FeatureLevel));
 
 	SCOPE_CYCLE_COUNTER(STAT_DepthDrawTime);
+	SCOPED_GPU_STAT(RHICmdList, Stat_GPU_Prepass);
 
 	bool bDirty = false;
 	bool bDidPrePre = false;
diff --git a/Engine/Source/Runtime/Renderer/Private/DistortionRendering.cpp b/Engine/Source/Runtime/Renderer/Private/DistortionRendering.cpp
--- a/Engine/Source/Runtime/Renderer/Private/DistortionRendering.cpp
+++ b/Engine/Source/Runtime/Renderer/Private/DistortionRendering.cpp
@@ -13,6 +13,8 @@
 #include "SceneFilterRendering.h"
 #include "SceneUtils.h"
 
+DECLARE_FLOAT_COUNTER_STAT(TEXT("Distortion"), Stat_GPU_Distortion, STATGROUP_GPU);
+
 /**
 * A pixel shader for rendering the full screen refraction pass
 */
@@ -832,6 +834,7 @@ void FSceneRenderer::RenderDistortion(FRHICommandListImmediate& RHICmdList)
 {
 	QUICK_SCOPE_CYCLE_COUNTER(STAT_FSceneRenderer_RenderDistortion);
 	SCOPED_DRAW_EVENT(RHICmdList, Distortion);
+	SCOPED_GPU_STAT(RHICmdList, Stat_GPU_Distortion);
 
 	// do we need to render the distortion pass?
 	bool bRender=false;
diff --git a/Engine/Source/Runtime/Renderer/Private/FogRendering.cpp b/Engine/Source/Runtime/Renderer/Private/FogRendering.cpp
--- a/Engine/Source/Runtime/Renderer/Private/FogRendering.cpp
+++ b/Engine/Source/Runtime/Renderer/Private/FogRendering.cpp
@@ -8,6 +8,8 @@
 #include "ScenePrivate.h"
 #include "SceneUtils.h"
 
+DECLARE_FLOAT_COUNTER_STAT(TEXT("Fog"), Stat_GPU_Fog, STATGROUP_GPU);
+
 #if !(UE_BUILD_SHIPPING || UE_BUILD_TEST)
 static TAutoConsoleVariable<float> CVarFogStartDistance(
 	TEXT("r.FogStartDistance"),
@@ -320,6 +322,7 @@ bool FDeferredShadingSceneRenderer::RenderFog(FRHICommandListImmediate& RHICmdLi
 			const FViewInfo& View = Views[ViewIndex];
 
 			SCOPED_DRAW_EVENTF(RHICmdList, Fog, TEXT("ExponentialHeightFog %dx%d"), View.ViewRect.Width(), View.ViewRect.Height());
+			SCOPED_GPU_STAT(RHICmdList, Stat_GPU_Fog);
 
 			if (View.IsPerspectiveProjection() == false)
 			{
diff --git a/Engine/Source/Runtime/Renderer/Private/LightPropagationVolume.cpp b/Engine/Source/Runtime/Renderer/Private/LightPropagationVolume.cpp
--- a/Engine/Source/Runtime/Renderer/Private/LightPropagationVolume.cpp
+++ b/Engine/Source/Runtime/Renderer/Private/LightPropagationVolume.cpp
@@ -17,6 +17,8 @@
 #include "SceneUtils.h"
 #include "LightPropagationVolumeBlendable.h"
 
+DECLARE_FLOAT_COUNTER_STAT(TEXT("LPV"), Stat_GPU_LPV, STATGROUP_GPU);
+
 static TAutoConsoleVariable<int32> CVarLightPropagationVolume(
 	TEXT("r.LightPropagationVolume"),
 	0,
@@ -1570,6 +1572,7 @@ void FDeferredShadingSceneRenderer::ClearLPVs(FRHICommandListImmediate& RHICmdLi
 
 				if(LightPropagationVolume)
 				{
+					SCOPED_GPU_STAT(RHICmdList, Stat_GPU_LPV);
 					LightPropagationVolume->InitSettings(RHICmdList, Views[ViewIndex]);
 					LightPropagationVolume->Clear(RHICmdList, View);
 				}
@@ -1598,7 +1601,8 @@ void FDeferredShadingSceneRenderer::UpdateLPVs(FRHICommandListImmediate& RHICmdL
 			{
 //				SCOPED_DRAW_EVENT(RHICmdList, UpdateLPVs);
 //				SCOPE_CYCLE_COUNTER(STAT_UpdateLPVs);
-				
+				SCOPED_GPU_STAT(RHICmdList, Stat_GPU_LPV);
+
 				LightPropagationVolume->Update(RHICmdList, View);
 			}
 		}
diff --git a/Engine/Source/Runtime/Renderer/Private/LightRendering.cpp b/Engine/Source/Runtime/Renderer/Private/LightRendering.cpp
--- a/Engine/Source/Runtime/Renderer/Private/LightRendering.cpp
+++ b/Engine/Source/Runtime/Renderer/Private/LightRendering.cpp
@@ -13,6 +13,7 @@
 #include "LightPropagationVolume.h"
 #include "SceneUtils.h"
 
+DECLARE_FLOAT_COUNTER_STAT(TEXT("Lights"), Stat_GPU_Lights, STATGROUP_GPU);
 
 IMPLEMENT_UNIFORM_BUFFER_STRUCT(FDeferredLightUniformStruct,TEXT("DeferredLightUniforms"));
 
@@ -323,6 +324,8 @@ uint32 GetShadowQuality();
 void FDeferredShadingSceneRenderer::RenderLights(FRHICommandListImmediate& RHICmdList)
 {
 	SCOPED_DRAW_EVENT(RHICmdList, Lights);
+	SCOPED_GPU_STAT(RHICmdList, Stat_GPU_Lights);
+
 
 	bool bStencilBufferDirty = false;	// The stencil buffer should've been cleared to 0 already
 
diff --git a/Engine/Source/Runtime/Renderer/Private/MobileBasePassRendering.h b/Engine/Source/Runtime/Renderer/Private/MobileBasePassRendering.h
--- a/Engine/Source/Runtime/Renderer/Private/MobileBasePassRendering.h
+++ b/Engine/Source/Runtime/Renderer/Private/MobileBasePassRendering.h
@@ -708,6 +708,10 @@ public:
 					// Modulate with the existing scene color
 					RHICmdList.SetBlendState(TStaticBlendState<CW_RGB, BO_Add, BF_DestColor, BF_Zero>::GetRHI());
 					break;
+				case BLEND_AlphaComposite:
+					// Blend with existing scene color. New color is already pre-multiplied by alpha.
+					RHICmdList.SetBlendState(TStaticBlendState<CW_RGBA, BO_Add, BF_One, BF_InverseSourceAlpha, BO_Add, BF_Zero, BF_InverseSourceAlpha>::GetRHI());
+					break;
 				};
 			}
 		}
diff --git a/Engine/Source/Runtime/Renderer/Private/PostProcess/RenderTargetPool.cpp b/Engine/Source/Runtime/Renderer/Private/PostProcess/RenderTargetPool.cpp
--- a/Engine/Source/Runtime/Renderer/Private/PostProcess/RenderTargetPool.cpp
+++ b/Engine/Source/Runtime/Renderer/Private/PostProcess/RenderTargetPool.cpp
@@ -314,6 +314,13 @@ bool FRenderTargetPool::FindFreeElement(FRHICommandList& RHICmdList, const FPool
 					Desc.NumSamples
 					);
 
+				if (GSupportsRenderTargetWriteMask && Desc.bCreateRenderTargetWriteMask)
+				{
+					Found->RenderTargetItem.RTWriteMaskDataBufferRHI = RHICreateRTWriteMaskBuffer((FTexture2DRHIRef&)Found->RenderTargetItem.TargetableTexture);
+					Found->RenderTargetItem.RTWriteMaskBufferRHI_SRV = RHICreateShaderResourceView(Found->RenderTargetItem.RTWriteMaskDataBufferRHI);
+				}
+
+
 				if( Desc.NumMips > 1 )
 				{
 					Found->RenderTargetItem.MipSRVs.SetNum( Desc.NumMips );
diff --git a/Engine/Source/Runtime/Renderer/Private/PostProcess/SceneRenderTargets.cpp b/Engine/Source/Runtime/Renderer/Private/PostProcess/SceneRenderTargets.cpp
--- a/Engine/Source/Runtime/Renderer/Private/PostProcess/SceneRenderTargets.cpp
+++ b/Engine/Source/Runtime/Renderer/Private/PostProcess/SceneRenderTargets.cpp
@@ -199,6 +199,7 @@ FSceneRenderTargets::FSceneRenderTargets(const FViewInfo& View, const FSceneRend
 	, DBufferA(GRenderTargetPool.MakeSnapshot(SnapshotSource.DBufferA))
 	, DBufferB(GRenderTargetPool.MakeSnapshot(SnapshotSource.DBufferB))
 	, DBufferC(GRenderTargetPool.MakeSnapshot(SnapshotSource.DBufferC))
+	, DBufferMask(GRenderTargetPool.MakeSnapshot(SnapshotSource.DBufferMask))
 	, ScreenSpaceAO(GRenderTargetPool.MakeSnapshot(SnapshotSource.ScreenSpaceAO))
 	, QuadOverdrawBuffer(GRenderTargetPool.MakeSnapshot(SnapshotSource.QuadOverdrawBuffer))
 	, CustomDepth(GRenderTargetPool.MakeSnapshot(SnapshotSource.CustomDepth))
@@ -833,7 +834,7 @@ void FSceneRenderTargets::AllocGBufferTargets(FRHICommandList& RHICmdList)
 		GBufferResourceStruct.GBufferDTexture = GBufferDToUse.ShaderResourceTexture;
 		GBufferResourceStruct.GBufferETexture = GBufferEToUse.ShaderResourceTexture;
 		GBufferResourceStruct.GBufferVelocityTexture = GBufferVelocityToUse.ShaderResourceTexture;
-
+		
 		GBufferResourceStruct.GBufferATextureNonMS = GBufferAToUse.ShaderResourceTexture;
 		GBufferResourceStruct.GBufferBTextureNonMS = GBufferBToUse.ShaderResourceTexture;
 		GBufferResourceStruct.GBufferCTextureNonMS = GBufferCToUse.ShaderResourceTexture;
@@ -2320,6 +2321,7 @@ void FDeferredPixelShaderParameters::Bind(const FShaderParameterMap& ParameterMa
 	CustomDepthTexture.Bind(ParameterMap,TEXT("CustomDepthTexture"));
 	CustomDepthTextureSampler.Bind(ParameterMap,TEXT("CustomDepthTextureSampler"));
 	CustomStencilTexture.Bind(ParameterMap,TEXT("CustomStencilTexture"));
+	DBufferRenderMask.Bind(ParameterMap, TEXT("DBufferMask"));
 }
 
 bool IsDBufferEnabled();
@@ -2332,7 +2334,7 @@ void FDeferredPixelShaderParameters::Set(TRHICmdList& RHICmdList, const ShaderRH
 	SceneTextureParameters.Set(RHICmdList, ShaderRHI, View, TextureMode, SF_Point);
 
 	// if() is purely an optimization and could be removed
-	if(IsDBufferEnabled())
+	if (IsDBufferEnabled())
 	{
 		IPooledRenderTarget* DBufferA = SceneContext.DBufferA ? SceneContext.DBufferA : GSystemTextures.BlackDummy;
 		IPooledRenderTarget* DBufferB = SceneContext.DBufferB ? SceneContext.DBufferB : GSystemTextures.BlackDummy;
@@ -2342,6 +2344,19 @@ void FDeferredPixelShaderParameters::Set(TRHICmdList& RHICmdList, const ShaderRH
 		SetTextureParameter(RHICmdList, ShaderRHI, DBufferATexture, DBufferATextureSampler, TStaticSamplerState<>::GetRHI(), DBufferA->GetRenderTargetItem().ShaderResourceTexture);
 		SetTextureParameter(RHICmdList, ShaderRHI, DBufferBTexture, DBufferBTextureSampler, TStaticSamplerState<>::GetRHI(), DBufferB->GetRenderTargetItem().ShaderResourceTexture);
 		SetTextureParameter(RHICmdList, ShaderRHI, DBufferCTexture, DBufferCTextureSampler, TStaticSamplerState<>::GetRHI(), DBufferC->GetRenderTargetItem().ShaderResourceTexture);
+
+		if (GSupportsRenderTargetWriteMask)
+		{
+			if (SceneContext.DBufferMask)
+			{
+				SetTextureParameter(RHICmdList, ShaderRHI, DBufferRenderMask, SceneContext.DBufferMask->GetRenderTargetItem().TargetableTexture);
+			}
+			else
+			{
+				SetTextureParameter(RHICmdList, ShaderRHI, DBufferRenderMask, GSystemTextures.WhiteDummy->GetRenderTargetItem().TargetableTexture);
+			}
+		}
+
 		SetTextureParameter(RHICmdList, ShaderRHI, DBufferATextureMS, DBufferA->GetRenderTargetItem().TargetableTexture);
 		SetTextureParameter(RHICmdList, ShaderRHI, DBufferBTextureMS, DBufferB->GetRenderTargetItem().TargetableTexture);
 		SetTextureParameter(RHICmdList, ShaderRHI, DBufferCTextureMS, DBufferC->GetRenderTargetItem().TargetableTexture);
@@ -2441,6 +2456,7 @@ FArchive& operator<<(FArchive& Ar,FDeferredPixelShaderParameters& Parameters)
 	Ar << Parameters.ScreenSpaceAOTextureNonMS;
 	Ar << Parameters.CustomDepthTextureNonMS;
 	Ar << Parameters.DBufferATexture;
+	Ar << Parameters.DBufferRenderMask;
 	Ar << Parameters.DBufferATextureSampler;
 	Ar << Parameters.DBufferBTexture;
 	Ar << Parameters.DBufferBTextureSampler;
diff --git a/Engine/Source/Runtime/Renderer/Private/PostProcess/SceneRenderTargets.h b/Engine/Source/Runtime/Renderer/Private/PostProcess/SceneRenderTargets.h
--- a/Engine/Source/Runtime/Renderer/Private/PostProcess/SceneRenderTargets.h
+++ b/Engine/Source/Runtime/Renderer/Private/PostProcess/SceneRenderTargets.h
@@ -554,6 +554,7 @@ public:
 	TRefCountPtr<IPooledRenderTarget> DBufferA;
 	TRefCountPtr<IPooledRenderTarget> DBufferB;
 	TRefCountPtr<IPooledRenderTarget> DBufferC;
+	TRefCountPtr<IPooledRenderTarget> DBufferMask;
 
 	// for AmbientOcclusion, only valid for a short time during the frame to allow reuse
 	TRefCountPtr<IPooledRenderTarget> ScreenSpaceAO;
diff --git a/Engine/Source/Runtime/Renderer/Private/PostProcess/ScreenSpaceReflections.cpp b/Engine/Source/Runtime/Renderer/Private/PostProcess/ScreenSpaceReflections.cpp
--- a/Engine/Source/Runtime/Renderer/Private/PostProcess/ScreenSpaceReflections.cpp
+++ b/Engine/Source/Runtime/Renderer/Private/PostProcess/ScreenSpaceReflections.cpp
@@ -48,6 +48,8 @@ static TAutoConsoleVariable<int32> CVarSSRCone(
 	TEXT(" 0 is off (default), 1 is on"),
 	ECVF_RenderThreadSafe);
 
+DECLARE_FLOAT_COUNTER_STAT(TEXT("ScreenSpace Reflections"), Stat_GPU_ScreenSpaceReflections, STATGROUP_GPU);
+
 bool ShouldRenderScreenSpaceReflections(const FViewInfo& View)
 {
 	if(!View.Family->EngineShowFlags.ScreenSpaceReflections)
@@ -335,6 +337,8 @@ void FRCPassPostProcessScreenSpaceReflections::Process(FRenderingCompositePassCo
 	}
 	
 	const FSceneRenderTargetItem& DestRenderTarget = PassOutputs[0].RequestSurface(Context);
+
+	SCOPED_GPU_STAT(RHICmdList, Stat_GPU_ScreenSpaceReflections);
 	
 	if (SSRStencilPrePass)
 	{ // ScreenSpaceReflectionsStencil draw event
diff --git a/Engine/Source/Runtime/Renderer/Private/ReflectionEnvironment.cpp b/Engine/Source/Runtime/Renderer/Private/ReflectionEnvironment.cpp
--- a/Engine/Source/Runtime/Renderer/Private/ReflectionEnvironment.cpp
+++ b/Engine/Source/Runtime/Renderer/Private/ReflectionEnvironment.cpp
@@ -20,6 +20,8 @@
 #include "SceneUtils.h"
 #include "LightPropagationVolumeBlendable.h"
 
+DECLARE_FLOAT_COUNTER_STAT(TEXT("Reflection Environment"), Stat_GPU_ReflectionEnvironment, STATGROUP_GPU);
+
 /** Tile size for the reflection environment compute shader, tweaked for 680 GTX. */
 const int32 GReflectionEnvironmentTileSizeX = 16;
 const int32 GReflectionEnvironmentTileSizeY = 16;
@@ -936,6 +938,7 @@ void FDeferredShadingSceneRenderer::RenderTiledDeferredImageBasedReflections(FRH
 			RenderScreenSpaceReflections(RHICmdList, View, SSROutput);
 		}
 
+		SCOPED_GPU_STAT(RHICmdList, Stat_GPU_ReflectionEnvironment)
 		RenderDeferredPlanarReflections(RHICmdList, false, SSROutput);
 
 		// ReflectionEnv is assumed to be on when going into this method
@@ -1100,6 +1103,7 @@ void FDeferredShadingSceneRenderer::RenderStandardDeferredImageBasedReflections(
 			RenderScreenSpaceReflections(RHICmdList, View, SSROutput);
 		}
 
+		SCOPED_GPU_STAT(RHICmdList, Stat_GPU_ReflectionEnvironment)
 		bool bApplyFromSSRTexture = bSSR;
 
 		if (RenderDeferredPlanarReflections(RHICmdList, true, SSROutput))
@@ -1188,6 +1192,7 @@ void FDeferredShadingSceneRenderer::RenderStandardDeferredImageBasedReflections(
 		{
 			// Apply reflections to screen
 			SCOPED_DRAW_EVENT(RHICmdList, ReflectionApply);
+			SCOPED_GPU_STAT(RHICmdList, Stat_GPU_ReflectionEnvironment);
 
 			SceneContext.BeginRenderingSceneColor(RHICmdList, ESimpleRenderTargetMode::EUninitializedColorExistingDepth, FExclusiveDepthStencil::DepthRead_StencilWrite, true);
 
diff --git a/Engine/Source/Runtime/Renderer/Private/RendererScene.cpp b/Engine/Source/Runtime/Renderer/Private/RendererScene.cpp
--- a/Engine/Source/Runtime/Renderer/Private/RendererScene.cpp
+++ b/Engine/Source/Runtime/Renderer/Private/RendererScene.cpp
@@ -509,7 +509,7 @@ FScene::FReadOnlyCVARCache::FReadOnlyCVARCache()
 	const bool bShowMissmatchedLowQualityLightmapsWarning = (bEnableLowQualityLightmaps) != (GEngine->bShouldGenerateLowQualityLightmaps_DEPRECATED);
 	if ( bShowMissmatchedLowQualityLightmapsWarning )
 	{
-		UE_LOG(LogRenderer, Warning, TEXT("Missmatch between bShouldGenerateLowQualityLightmaps(%d) and r.SupportLowQualityLightmaps(%d), UEngine::bShouldGenerateLowQualityLightmaps has been depricated please use r.SupportLowQualityLightmaps instead"), GEngine->bShouldGenerateLowQualityLightmaps_DEPRECATED, bEnableLowQualityLightmaps);
+		UE_LOG(LogRenderer, Warning, TEXT("Mismatch between bShouldGenerateLowQualityLightmaps(%d) and r.SupportLowQualityLightmaps(%d), UEngine::bShouldGenerateLowQualityLightmaps has been deprecated please use r.SupportLowQualityLightmaps instead"), GEngine->bShouldGenerateLowQualityLightmaps_DEPRECATED, bEnableLowQualityLightmaps);
 	}
 }
 
diff --git a/Engine/Source/Runtime/Renderer/Private/SceneCaptureRendering.cpp b/Engine/Source/Runtime/Renderer/Private/SceneCaptureRendering.cpp
--- a/Engine/Source/Runtime/Renderer/Private/SceneCaptureRendering.cpp
+++ b/Engine/Source/Runtime/Renderer/Private/SceneCaptureRendering.cpp
@@ -215,14 +215,13 @@ static void UpdateSceneCaptureContent_RenderThread(
 	FRenderTarget* RenderTarget, 
 	FTexture* RenderTargetTexture, 
 	const FName OwnerName, 
-	const FResolveParams& ResolveParams, 
-	ESceneCaptureSource CaptureSource)
+	const FResolveParams& ResolveParams)
 {
 	FMemMark MemStackMark(FMemStack::Get());
 
 	// update any resources that needed a deferred update
 	FDeferredUpdateResource::UpdateResources(RHICmdList);
-	bool bUseSceneTextures = CaptureSource != SCS_FinalColorLDR;
+	bool bUseSceneTextures = SceneRenderer->ViewFamily.SceneCaptureSource != SCS_FinalColorLDR;
 
 	{
 #if WANTS_DRAW_MESH_EVENTS
@@ -310,7 +309,7 @@ static void UpdateSceneCaptureContent_RenderThread(
 		}
 		else if (bUseSceneTextures)
 		{
-			// Note: When the CaptureSource requires scene textures, the copy will be done in CopySceneCaptureComponentToTarget while the GBuffers are still alive for the frame
+			// Note: When the ViewFamily.SceneCaptureSource requires scene textures, the copy will be done in CopySceneCaptureComponentToTarget while the GBuffers are still alive for the frame
 		}
 
 		RHICmdList.CopyToResolveTarget(RenderTarget->GetRenderTargetTexture(), RenderTargetTexture->TextureRHI, false, ResolveParams);
@@ -532,14 +531,13 @@ void FScene::UpdateSceneCaptureContents(USceneCaptureComponent2D* CaptureCompone
 		FTextureRenderTargetResource* TextureRenderTarget = CaptureComponent->TextureTarget->GameThread_GetRenderTargetResource();
 		const FName OwnerName = CaptureComponent->GetOwner() ? CaptureComponent->GetOwner()->GetFName() : NAME_None;
 
-		ENQUEUE_UNIQUE_RENDER_COMMAND_FOURPARAMETER(
+		ENQUEUE_UNIQUE_RENDER_COMMAND_THREEPARAMETER(
 			CaptureCommand,
 			FSceneRenderer*, SceneRenderer, SceneRenderer,
 			FTextureRenderTargetResource*, TextureRenderTarget, TextureRenderTarget,
 			FName, OwnerName, OwnerName,
-			ESceneCaptureSource, CaptureSource, CaptureComponent->CaptureSource, 
 		{
-			UpdateSceneCaptureContent_RenderThread(RHICmdList, SceneRenderer, TextureRenderTarget, TextureRenderTarget, OwnerName, FResolveParams(), CaptureSource);
+			UpdateSceneCaptureContent_RenderThread(RHICmdList, SceneRenderer, TextureRenderTarget, TextureRenderTarget, OwnerName, FResolveParams());
 		});
 	}
 }
@@ -603,6 +601,7 @@ void FScene::UpdateSceneCaptureContents(USceneCaptureComponentCube* CaptureCompo
 			BuildProjectionMatrix(CaptureSize, ECameraProjectionMode::Perspective, FOV, 1.0f, ProjectionMatrix);
 
 			FSceneRenderer* SceneRenderer = CreateSceneRendererForSceneCapture(this, CaptureComponent, CaptureComponent->TextureTarget->GameThread_GetRenderTargetResource(), CaptureSize, ViewRotationMatrix, Location, ProjectionMatrix, CaptureComponent->MaxViewDistanceOverride, true, false, NULL, 0);
+			SceneRenderer->ViewFamily.SceneCaptureSource = SCS_SceneColorHDR;
 
 			FTextureRenderTargetCubeResource* TextureRenderTarget = static_cast<FTextureRenderTargetCubeResource*>(CaptureComponent->TextureTarget->GameThread_GetRenderTargetResource());
 			const FName OwnerName = CaptureComponent->GetOwner() ? CaptureComponent->GetOwner()->GetFName() : NAME_None;
@@ -614,7 +613,7 @@ void FScene::UpdateSceneCaptureContents(USceneCaptureComponentCube* CaptureCompo
 				FName, OwnerName, OwnerName,
 				ECubeFace, TargetFace, TargetFace,
 			{
-				UpdateSceneCaptureContent_RenderThread(RHICmdList, SceneRenderer, TextureRenderTarget, TextureRenderTarget, OwnerName, FResolveParams(FResolveRect(), TargetFace), SCS_SceneColorHDR);
+				UpdateSceneCaptureContent_RenderThread(RHICmdList, SceneRenderer, TextureRenderTarget, TextureRenderTarget, OwnerName, FResolveParams(FResolveRect(), TargetFace));
 			});
 		}
 	}
diff --git a/Engine/Source/Runtime/Renderer/Private/SceneRendering.cpp b/Engine/Source/Runtime/Renderer/Private/SceneRendering.cpp
--- a/Engine/Source/Runtime/Renderer/Private/SceneRendering.cpp
+++ b/Engine/Source/Runtime/Renderer/Private/SceneRendering.cpp
@@ -29,6 +29,8 @@
 extern ENGINE_API FLightMap2D* GDebugSelectedLightmap;
 extern ENGINE_API UPrimitiveComponent* GDebugSelectedComponent;
 
+DECLARE_FLOAT_COUNTER_STAT(TEXT("Custom Depth"), Stat_GPU_CustomDepth, STATGROUP_GPU);
+
 /**
  * Console variable controlling whether or not occlusion queries are allowed.
  */
@@ -1688,6 +1690,7 @@ void FSceneRenderer::RenderCustomDepthPass(FRHICommandListImmediate& RHICmdList)
 	if (SceneContext.BeginRenderingCustomDepth(RHICmdList, bPrimitives))
 	{
 		SCOPED_DRAW_EVENT(RHICmdList, CustomDepth);
+		SCOPED_GPU_STAT(RHICmdList, Stat_GPU_CustomDepth);
 
 		for(int32 ViewIndex = 0;ViewIndex < Views.Num();ViewIndex++)
 		{
@@ -1870,6 +1873,8 @@ static void RenderViewFamily_RenderThread(FRHICommandListImmediate& RHICmdList,
 	{
 		SCOPE_CYCLE_COUNTER(STAT_TotalSceneRenderingTime);
 		
+		GPU_STATS_UPDATE(RHICmdList);
+		
 		if(SceneRenderer->ViewFamily.EngineShowFlags.HitProxies)
 		{
 			// Render the scene's hit proxies.
@@ -2025,6 +2030,9 @@ void FRendererModule::BeginRenderingViewFamily(FCanvas* Canvas, FSceneViewFamily
 
 		if (!SceneRenderer->ViewFamily.EngineShowFlags.HitProxies)
 		{
+			USceneCaptureComponent2D::UpdateDeferredCaptures(Scene);
+			USceneCaptureComponentCube::UpdateDeferredCaptures(Scene);
+
 			for (int32 ReflectionIndex = 0; ReflectionIndex < SceneRenderer->Scene->PlanarReflections_GameThread.Num(); ReflectionIndex++)
 			{
 				UPlanarReflectionComponent* ReflectionComponent = SceneRenderer->Scene->PlanarReflections_GameThread[ReflectionIndex];
diff --git a/Engine/Source/Runtime/Renderer/Private/SceneVisibility.cpp b/Engine/Source/Runtime/Renderer/Private/SceneVisibility.cpp
--- a/Engine/Source/Runtime/Renderer/Private/SceneVisibility.cpp
+++ b/Engine/Source/Runtime/Renderer/Private/SceneVisibility.cpp
@@ -2638,7 +2638,10 @@ void FSceneRenderer::PostVisibilityFrameSetup(FILCUpdatePrimTaskData& OutILCTask
 				{
 					FSphere Bounds = Proxy->GetBoundingSphere();
 					float DistanceSquared = (Bounds.Center - View.ViewMatrices.ViewOrigin).SizeSquared();
-					const bool bDrawLight = FMath::Square( FMath::Min( 0.0002f, GMinScreenRadiusForLights / Bounds.W ) * View.LODDistanceFactor ) * DistanceSquared < 1.0f;
+					float MaxDistSquared = Proxy->GetMaxDrawDistance() * Proxy->GetMaxDrawDistance();
+					const bool bDrawLight = (FMath::Square(FMath::Min(0.0002f, GMinScreenRadiusForLights / Bounds.W) * View.LODDistanceFactor) * DistanceSquared < 1.0f)
+												& (MaxDistSquared == 0 || DistanceSquared < MaxDistSquared);
+							
 					VisibleLightViewInfo.bInViewFrustum = bDrawLight;
 				}
 			}
diff --git a/Engine/Source/Runtime/Renderer/Private/ShadowDepthRendering.cpp b/Engine/Source/Runtime/Renderer/Private/ShadowDepthRendering.cpp
--- a/Engine/Source/Runtime/Renderer/Private/ShadowDepthRendering.cpp
+++ b/Engine/Source/Runtime/Renderer/Private/ShadowDepthRendering.cpp
@@ -12,6 +12,8 @@
 #include "SceneFilterRendering.h"
 #include "ScreenRendering.h"
 
+DECLARE_FLOAT_COUNTER_STAT(TEXT("Shadow Depths"), Stat_GPU_ShadowDepths, STATGROUP_GPU);
+
 /**
  * A vertex shader for rendering the depth of a mesh.
  */
@@ -2143,6 +2145,7 @@ void FSceneRenderer::RenderShadowDepthMaps(FRHICommandListImmediate& RHICmdList)
 	FSceneRenderTargets& SceneContext = FSceneRenderTargets::Get(RHICmdList);
 
 	SCOPED_DRAW_EVENT(RHICmdList, ShadowDepths);
+	SCOPED_GPU_STAT(RHICmdList, Stat_GPU_ShadowDepths);
 
 	FSceneRenderer::RenderShadowDepthMapAtlases(RHICmdList);
 
diff --git a/Engine/Source/Runtime/Renderer/Private/ShadowRendering.cpp b/Engine/Source/Runtime/Renderer/Private/ShadowRendering.cpp
--- a/Engine/Source/Runtime/Renderer/Private/ShadowRendering.cpp
+++ b/Engine/Source/Runtime/Renderer/Private/ShadowRendering.cpp
@@ -67,6 +67,8 @@ static TAutoConsoleVariable<int32> CVarEnableModulatedSelfShadow(
 	TEXT("Allows modulated shadows to affect the shadow caster. (mobile only)"),
 	ECVF_RenderThreadSafe);
 
+DECLARE_FLOAT_COUNTER_STAT(TEXT("ShadowProjection"), Stat_GPU_ShadowProjection, STATGROUP_GPU);
+
 // 0:off, 1:low, 2:med, 3:high, 4:very high, 5:max
 uint32 GetShadowQuality()
 {
@@ -1278,6 +1280,7 @@ bool FDeferredShadingSceneRenderer::RenderShadowProjections(FRHICommandListImmed
 {
 	SCOPE_CYCLE_COUNTER(STAT_ProjectedShadowDrawTime);
 	SCOPED_DRAW_EVENT(RHICmdList, ShadowProjectionOnOpaque);
+	SCOPED_GPU_STAT(RHICmdList, Stat_GPU_ShadowProjection);
 
 	FVisibleLightInfo& VisibleLightInfo = VisibleLightInfos[LightSceneInfo->Id];
 
diff --git a/Engine/Source/Runtime/Renderer/Private/ShadowRendering.h b/Engine/Source/Runtime/Renderer/Private/ShadowRendering.h
--- a/Engine/Source/Runtime/Renderer/Private/ShadowRendering.h
+++ b/Engine/Source/Runtime/Renderer/Private/ShadowRendering.h
@@ -108,15 +108,22 @@ void SetDeferredLightParameters(
 
 	const ELightComponentType LightType = (ELightComponentType)LightSceneInfo->Proxy->GetLightType();
 
+
 	if (LightType == LightType_Point || LightType == LightType_Spot)
 	{
 		// Distance fade
 		FSphere Bounds = LightSceneInfo->Proxy->GetBoundingSphere();
 
-		const float DistanceSquared = ( Bounds.Center - View.ViewMatrices.ViewOrigin ).SizeSquared();
-		float Fade = FMath::Square( FMath::Min( 0.0002f, GMinScreenRadiusForLights / Bounds.W ) * View.LODDistanceFactor ) * DistanceSquared;
-		Fade = FMath::Clamp( 6.0f - 6.0f * Fade, 0.0f, 1.0f );
-		DeferredLightUniformsValue.LightColor *= Fade;
+		const float DistanceSquared = (Bounds.Center - View.ViewMatrices.ViewOrigin).SizeSquared();
+		float SizeFade = FMath::Square(FMath::Min(0.0002f, GMinScreenRadiusForLights / Bounds.W) * View.LODDistanceFactor) * DistanceSquared;
+		SizeFade = FMath::Clamp(6.0f - 6.0f * SizeFade, 0.0f, 1.0f);
+
+		float MaxDist = LightSceneInfo->Proxy->GetMaxDrawDistance();
+		float Range = LightSceneInfo->Proxy->GetFadeRange();
+		float DistanceFade = MaxDist ? (MaxDist - FMath::Sqrt(DistanceSquared)) / Range : 1.0f;
+		DistanceFade = FMath::Clamp(DistanceFade, 0.0f, 1.0f);
+
+		DeferredLightUniformsValue.LightColor *= SizeFade * DistanceFade;
 	}
 
 	DeferredLightUniformsValue.LightingChannelMask = LightSceneInfo->Proxy->GetLightingChannelMask();
diff --git a/Engine/Source/Runtime/Renderer/Private/TranslucentLighting.cpp b/Engine/Source/Runtime/Renderer/Private/TranslucentLighting.cpp
--- a/Engine/Source/Runtime/Renderer/Private/TranslucentLighting.cpp
+++ b/Engine/Source/Runtime/Renderer/Private/TranslucentLighting.cpp
@@ -17,6 +17,8 @@ class FMaterial;
 
 /** Whether to allow rendering translucency shadow depths. */
 bool GUseTranslucencyShadowDepths = true;
+
+DECLARE_FLOAT_COUNTER_STAT(TEXT("Translucent Lighting"), Stat_GPU_TranslucentLighting, STATGROUP_GPU);
  
 int32 GUseTranslucentLightingVolumes = 1;
 FAutoConsoleVariableRef CVarUseTranslucentLightingVolumes(
@@ -1075,6 +1077,8 @@ void FDeferredShadingSceneRenderer::ClearTranslucentVolumeLighting(FRHICommandLi
 	if (GUseTranslucentLightingVolumes && GSupportsVolumeTextureRendering)
 	{
 		SCOPED_DRAW_EVENT(RHICmdList, ClearTranslucentVolumeLighting);
+		SCOPED_GPU_STAT(RHICmdList, Stat_GPU_TranslucentLighting);
+
 		FSceneRenderTargets& SceneContext = FSceneRenderTargets::Get(RHICmdList);
 
 		// to track down Jira UE-22073
@@ -1292,6 +1296,7 @@ void FDeferredShadingSceneRenderer::InjectAmbientCubemapTranslucentVolumeLightin
 		ensure(SceneContext.GetCurrentFeatureLevel() >= ERHIFeatureLevel::SM4);
 
 		SCOPED_DRAW_EVENT(RHICmdList, InjectAmbientCubemapTranslucentVolumeLighting);
+		SCOPED_GPU_STAT(RHICmdList, Stat_GPU_TranslucentLighting);
 
 		const FVolumeBounds VolumeBounds(GTranslucencyLightingVolumeDim);
 
@@ -1351,6 +1356,7 @@ void FDeferredShadingSceneRenderer::ClearTranslucentVolumePerObjectShadowing(FRH
 	{
 		FSceneRenderTargets& SceneContext = FSceneRenderTargets::Get(RHICmdList);
 		SCOPED_DRAW_EVENT(RHICmdList, ClearTranslucentVolumePerLightShadowing);
+		SCOPED_GPU_STAT(RHICmdList, Stat_GPU_TranslucentLighting);
 
 		static_assert(TVC_MAX == 2, "Only expecting two translucency lighting cascades.");
 		FTextureRHIParamRef RenderTargets[2];
@@ -1406,6 +1412,7 @@ void FDeferredShadingSceneRenderer::AccumulateTranslucentVolumeObjectShadowing(F
 	if (GUseTranslucentLightingVolumes && GSupportsVolumeTextureRendering)
 	{
 		SCOPED_DRAW_EVENT(RHICmdList, AccumulateTranslucentVolumeShadowing);
+		SCOPED_GPU_STAT(RHICmdList, Stat_GPU_TranslucentLighting);
 
 		auto ShaderMap = GetGlobalShaderMap(FeatureLevel);
 
@@ -1952,6 +1959,8 @@ void FDeferredShadingSceneRenderer::FilterTranslucentVolumeLighting(FRHICommandL
 			SCOPED_DRAW_EVENTF(RHICmdList, FilterTranslucentVolume, TEXT("FilterTranslucentVolume %dx%dx%d Cascades:%d"),
 				GTranslucencyLightingVolumeDim, GTranslucencyLightingVolumeDim, GTranslucencyLightingVolumeDim, TVC_MAX);
 
+			SCOPED_GPU_STAT(RHICmdList, Stat_GPU_TranslucentLighting);
+
 			// Filter each cascade
 			for (int32 VolumeCascadeIndex = 0; VolumeCascadeIndex < TVC_MAX; VolumeCascadeIndex++)
 			{
diff --git a/Engine/Source/Runtime/Renderer/Private/TranslucentRendering.cpp b/Engine/Source/Runtime/Renderer/Private/TranslucentRendering.cpp
--- a/Engine/Source/Runtime/Renderer/Private/TranslucentRendering.cpp
+++ b/Engine/Source/Runtime/Renderer/Private/TranslucentRendering.cpp
@@ -16,6 +16,9 @@ DECLARE_CYCLE_STAT(TEXT("TranslucencyTimestampQuery Wait"), STAT_TranslucencyTim
 
 DECLARE_FLOAT_COUNTER_STAT(TEXT("Translucency GPU Time (MS)"), STAT_TranslucencyGPU, STATGROUP_SceneRendering);
 
+DECLARE_FLOAT_COUNTER_STAT(TEXT("Translucency"), Stat_GPU_Translucency, STATGROUP_GPU);
+
+
 static TAutoConsoleVariable<int32> CVarSeparateTranslucencyAutoDownsample(
 	TEXT("r.SeparateTranslucencyAutoDownsample"),
 	0,
@@ -1177,6 +1180,7 @@ void FDeferredShadingSceneRenderer::RenderTranslucency(FRHICommandListImmediate&
 	if (ShouldRenderTranslucency())
 	{
 		SCOPED_DRAW_EVENT(RHICmdList, Translucency);
+		SCOPED_GPU_STAT(RHICmdList, Stat_GPU_Translucency);
 
 		if (GRHICommandList.UseParallelAlgorithms() && CVarParallelTranslucency.GetValueOnRenderThread())
 		{
diff --git a/Engine/Source/Runtime/Renderer/Private/VelocityRendering.cpp b/Engine/Source/Runtime/Renderer/Private/VelocityRendering.cpp
--- a/Engine/Source/Runtime/Renderer/Private/VelocityRendering.cpp
+++ b/Engine/Source/Runtime/Renderer/Private/VelocityRendering.cpp
@@ -28,6 +28,8 @@ static TAutoConsoleVariable<int32> CVarRHICmdVelocityPassDeferredContexts(
 	1,
 	TEXT("True to use deferred contexts to parallelize velocity pass command list execution."));
 
+DECLARE_FLOAT_COUNTER_STAT(TEXT("Render Velocities"), Stat_GPU_RenderVelocities, STATGROUP_GPU);
+
 bool IsParallelVelocity()
 {
 	return GRHICommandList.UseParallelAlgorithms() && CVarParallelVelocity.GetValueOnRenderThread();
@@ -849,6 +851,7 @@ void FDeferredShadingSceneRenderer::RenderVelocities(FRHICommandListImmediate& R
 	check(!Views[0].bIsSceneCapture);
 
 	SCOPED_DRAW_EVENT(RHICmdList, RenderVelocities);
+	SCOPED_GPU_STAT(RHICmdList, Stat_GPU_RenderVelocities);
 
 	FPooledRenderTargetDesc Desc = FVelocityRendering::GetRenderTargetDesc();
 	GRenderTargetPool.FindFreeElement(RHICmdList, Desc, VelocityRT, TEXT("Velocity"));
diff --git a/Engine/Source/Runtime/Renderer/Public/RendererInterface.h b/Engine/Source/Runtime/Renderer/Public/RendererInterface.h
--- a/Engine/Source/Runtime/Renderer/Public/RendererInterface.h
+++ b/Engine/Source/Runtime/Renderer/Public/RendererInterface.h
@@ -47,6 +47,7 @@ public:
 		, bForceSeparateTargetAndShaderResource(false)
 		, DebugName(TEXT("UnknownTexture"))
 		, AutoWritable(true)
+		, bCreateRenderTargetWriteMask(false)
 	{
 		check(!IsValid());
 	}
@@ -63,7 +64,8 @@ public:
 		uint32 InTargetableFlags,
 		bool bInForceSeparateTargetAndShaderResource,
 		uint16 InNumMips = 1,
-		bool InAutowritable = true)
+		bool InAutowritable = true,
+		bool InCreateRTWriteMask = false)
 	{
 		check(InExtent.X);
 		check(InExtent.Y);
@@ -82,6 +84,7 @@ public:
 		NewDesc.bForceSeparateTargetAndShaderResource = bInForceSeparateTargetAndShaderResource;
 		NewDesc.DebugName = TEXT("UnknownTexture2D");
 		NewDesc.AutoWritable = InAutowritable;
+		NewDesc.bCreateRenderTargetWriteMask = InCreateRTWriteMask;
 		check(NewDesc.Is2DTexture());
 		return NewDesc;
 	}
@@ -328,6 +331,8 @@ public:
 	const TCHAR *DebugName;
 	/** automatically set to writable via barrier during */
 	bool AutoWritable;
+	/** create render target write mask (supported only on specific platforms) */
+	bool bCreateRenderTargetWriteMask;
 };
 
 
@@ -373,6 +378,9 @@ struct FSceneRenderTargetItem
 	FUnorderedAccessViewRHIRef UAV;
 	/** only created if requested through the flag  */
 	TArray< FShaderResourceViewRHIRef > MipSRVs;
+
+	FShaderResourceViewRHIRef RTWriteMaskBufferRHI_SRV;
+	FStructuredBufferRHIRef RTWriteMaskDataBufferRHI;
 };
 
 /**
diff --git a/Engine/Source/Runtime/Renderer/Public/SceneRenderTargetParameters.h b/Engine/Source/Runtime/Renderer/Public/SceneRenderTargetParameters.h
--- a/Engine/Source/Runtime/Renderer/Public/SceneRenderTargetParameters.h
+++ b/Engine/Source/Runtime/Renderer/Public/SceneRenderTargetParameters.h
@@ -86,6 +86,7 @@ private:
 	FShaderResourceParameter ScreenSpaceAOTextureNonMS;
 	FShaderResourceParameter CustomDepthTextureNonMS;
 	FShaderResourceParameter DBufferATexture;
+	FShaderResourceParameter DBufferRenderMask;
 	FShaderResourceParameter DBufferATextureSampler;
 	FShaderResourceParameter DBufferBTexture;
 	FShaderResourceParameter DBufferBTextureSampler;
diff --git a/Engine/Source/Runtime/ShaderCore/Private/ShaderCore.cpp b/Engine/Source/Runtime/ShaderCore/Private/ShaderCore.cpp
--- a/Engine/Source/Runtime/ShaderCore/Private/ShaderCore.cpp
+++ b/Engine/Source/Runtime/ShaderCore/Private/ShaderCore.cpp
@@ -434,7 +434,10 @@ void GetShaderIncludes(const TCHAR* Filename, TArray<FString>& IncludeFilenames,
 			
 					// Some headers aren't required to be found (platforms that the user doesn't have access to)
 					// @todo: Is there some way to generalize this"
-					const bool bIsOptionalInclude = (ExtractedIncludeFilename == TEXT("PS4/PS4Common.usf") || ExtractedIncludeFilename == TEXT("PS4/PostProcessHMDMorpheus.usf"));
+					const bool bIsOptionalInclude = (ExtractedIncludeFilename == TEXT("PS4/PS4Common.usf") 
+						|| ExtractedIncludeFilename == TEXT("PS4/PostProcessHMDMorpheus.usf")
+						|| ExtractedIncludeFilename == TEXT("PS4/RTWriteMaskProcessing.usf")
+						);
 					// ignore the header if it's optional and doesn't exist
 					if (bIsOptionalInclude)
 					{
diff --git a/Engine/Source/Runtime/ShaderCore/Public/CrossCompilerCommon.h b/Engine/Source/Runtime/ShaderCore/Public/CrossCompilerCommon.h
--- a/Engine/Source/Runtime/ShaderCore/Public/CrossCompilerCommon.h
+++ b/Engine/Source/Runtime/ShaderCore/Public/CrossCompilerCommon.h
@@ -37,6 +37,22 @@ namespace CrossCompiler
 		PACKED_TYPEINDEX_MAX		= 5,
 	};
 
+	static FORCEINLINE uint8 ShaderStageIndexToTypeName(uint8 ShaderStage)
+	{
+		switch (ShaderStage)
+		{
+		case SHADER_STAGE_VERTEX:	return 'v';
+		case SHADER_STAGE_PIXEL:	return 'p';
+		case SHADER_STAGE_GEOMETRY:	return 'g';
+		case SHADER_STAGE_HULL:		return 'h';
+		case SHADER_STAGE_DOMAIN:	return 'd';
+		case SHADER_STAGE_COMPUTE:	return 'c';
+		default: break;
+		}
+		check(0);
+		return 0;
+	}
+
 	static FORCEINLINE uint8 PackedTypeIndexToTypeName(uint8 ArrayType)
 	{
 		switch (ArrayType)
diff --git a/Engine/Source/Runtime/SlateRHIRenderer/Private/SlateMaterialShader.cpp b/Engine/Source/Runtime/SlateRHIRenderer/Private/SlateMaterialShader.cpp
--- a/Engine/Source/Runtime/SlateRHIRenderer/Private/SlateMaterialShader.cpp
+++ b/Engine/Source/Runtime/SlateRHIRenderer/Private/SlateMaterialShader.cpp
@@ -108,6 +108,10 @@ void FSlateMaterialShaderPS::SetParameters(FRHICommandList& RHICmdList, const FS
 		// Modulate with the existing scene color
 		RHICmdList.SetBlendState(TStaticBlendState<CW_RGB, BO_Add, BF_DestColor, BF_Zero>::GetRHI());
 		break;
+	case BLEND_AlphaComposite:
+		// Blend with existing scene color. New color is already pre-multiplied by alpha.
+		RHICmdList.SetBlendState(TStaticBlendState<CW_RGBA, BO_Add, BF_One, BF_InverseSourceAlpha, BO_Add, BF_One, BF_InverseSourceAlpha>::GetRHI());
+		break;
 	};
 
 	SetShaderValue( RHICmdList, ShaderRHI, ShaderParams, InShaderParams );
diff --git a/Engine/Source/Runtime/SlateRHIRenderer/Private/SlateRHIRenderer.cpp b/Engine/Source/Runtime/SlateRHIRenderer/Private/SlateRHIRenderer.cpp
--- a/Engine/Source/Runtime/SlateRHIRenderer/Private/SlateRHIRenderer.cpp
+++ b/Engine/Source/Runtime/SlateRHIRenderer/Private/SlateRHIRenderer.cpp
@@ -16,6 +16,8 @@ DECLARE_CYCLE_STAT(TEXT("Slate RT: Create Batches"), STAT_SlateRTCreateBatches,
 DECLARE_CYCLE_STAT(TEXT("Slate RT: Fill Vertex & Index Buffers"), STAT_SlateRTFillVertexIndexBuffers, STATGROUP_Slate);
 DECLARE_CYCLE_STAT(TEXT("Slate RT: Draw Batches"), STAT_SlateRTDrawBatches, STATGROUP_Slate);
 
+DECLARE_FLOAT_COUNTER_STAT(TEXT("Slate UI/present"), Stat_GPU_SlateUI, STATGROUP_GPU); 
+
 // Defines the maximum size that a slate viewport will create
 #define MAX_VIEWPORT_SIZE 16384
 
@@ -396,6 +398,7 @@ void FSlateRHIRenderer::OnWindowDestroyed( const TSharedRef<SWindow>& InWindow )
 void FSlateRHIRenderer::DrawWindow_RenderThread(FRHICommandListImmediate& RHICmdList, const FViewportInfo& ViewportInfo, FSlateWindowElementList& WindowElementList, bool bLockToVsync, bool bClear)
 {
 	SCOPED_DRAW_EVENT(RHICmdList, SlateUI);
+	//SCOPED_GPU_STAT(RHICmdList, Stat_GPU_SlateUI); // Disabled since this seems to be distorted by the present time
 
 	// Should only be called by the rendering thread
 	check(IsInRenderingThread());
diff --git a/Engine/Source/Runtime/VulkanRHI/Private/VulkanCommandWrappers.h b/Engine/Source/Runtime/VulkanRHI/Private/VulkanCommandWrappers.h
--- a/Engine/Source/Runtime/VulkanRHI/Private/VulkanCommandWrappers.h
+++ b/Engine/Source/Runtime/VulkanRHI/Private/VulkanCommandWrappers.h
@@ -21,10 +21,10 @@ namespace VulkanRHI
 	void DumpPhysicalDeviceProperties(VkPhysicalDeviceMemoryProperties* Properties);
 	void DumpAllocateMemory(VkDevice Device, const VkMemoryAllocateInfo* AllocateInfo, VkDeviceMemory* Memory);
 	void DumpMemoryRequirements(VkMemoryRequirements* MemoryRequirements);
-	void DumpBufferCreate(VkDevice Device, const VkBufferCreateInfo* CreateInfo, VkBuffer* Buffer);
-	void DumpBufferViewCreate(VkDevice Device, const VkBufferViewCreateInfo* CreateInfo, VkBufferView* BufferView);
-	void DumpImageCreate(VkDevice Device, const VkImageCreateInfo* CreateInfo, VkImage* Image);
-	void DumpImageViewCreate(VkDevice Device, const VkImageViewCreateInfo* CreateInfo, VkImageView* ImageView);
+	void DumpCreateBuffer(VkDevice Device, const VkBufferCreateInfo* CreateInfo, VkBuffer* Buffer);
+	void DumpCreateBufferView(VkDevice Device, const VkBufferViewCreateInfo* CreateInfo, VkBufferView* BufferView);
+	void DumpCreateImage(VkDevice Device, const VkImageCreateInfo* CreateInfo, VkImage* Image);
+	void DumpCreateImageView(VkDevice Device, const VkImageViewCreateInfo* CreateInfo, VkImageView* ImageView);
 	void DumpFenceCreate(VkDevice Device, const VkFenceCreateInfo* CreateInfo, VkFence* Fence);
 	void DumpFenceList(uint32 FenceCount, const VkFence* Fences);
 	void DumpSemaphoreCreate(VkDevice Device, const VkSemaphoreCreateInfo* CreateInfo, VkSemaphore* Semaphore);
@@ -73,10 +73,10 @@ namespace VulkanRHI
 	#define DumpPhysicalDeviceProperties(x)
 	#define DumpAllocateMemory(d, i, m)
 	#define DumpMemoryRequirements(x)
-	#define DumpBufferCreate(d, i, b)
-	#define DumpBufferViewCreate(d, i, v)
-	#define DumpImageCreate(d, x, i)
-	#define DumpImageViewCreate(d, i, v)
+	#define DumpCreateBuffer(d, i, b)
+	#define DumpCreateBufferView(d, i, v)
+	#define DumpCreateImage(d, x, i)
+	#define DumpCreateImageView(d, i, v)
 	#define DumpFenceCreate(d, x, f)
 	#define DumpFenceList(c, p)
 	#define DumpSemaphoreCreate(d, i, s)
@@ -542,7 +542,7 @@ namespace VulkanRHI
 
 	static FORCEINLINE_DEBUGGABLE VkResult  vkCreateBuffer(VkDevice Device, const VkBufferCreateInfo* CreateInfo, const VkAllocationCallbacks* Allocator, VkBuffer* Buffer)
 	{
-		DumpBufferCreate(Device, CreateInfo, Buffer);
+		DumpCreateBuffer(Device, CreateInfo, Buffer);
 
 		VkResult Result = ::vkCreateBuffer(Device, CreateInfo, Allocator, Buffer);
 
@@ -559,7 +559,7 @@ namespace VulkanRHI
 
 	static FORCEINLINE_DEBUGGABLE VkResult  vkCreateBufferView(VkDevice Device, const VkBufferViewCreateInfo* CreateInfo, const VkAllocationCallbacks* Allocator, VkBufferView* View)
 	{
-		DumpBufferViewCreate(Device, CreateInfo, View);
+		DumpCreateBufferView(Device, CreateInfo, View);
 		DevicePrintfBeginResult(Device, FString::Printf(TEXT("vkCreateBufferView(Info=%p, OutView=%p)"), CreateInfo, View));
 
 		VkResult Result = ::vkCreateBufferView(Device, CreateInfo, Allocator, View);
@@ -577,7 +577,7 @@ namespace VulkanRHI
 
 	static FORCEINLINE_DEBUGGABLE VkResult  vkCreateImage(VkDevice Device, const VkImageCreateInfo* CreateInfo, const VkAllocationCallbacks* Allocator, VkImage* Image)
 	{
-		DumpImageCreate(Device, CreateInfo, Image);
+		DumpCreateImage(Device, CreateInfo, Image);
 
 		VkResult Result = ::vkCreateImage(Device, CreateInfo, Allocator, Image);
 
@@ -603,7 +603,7 @@ namespace VulkanRHI
 
 	static FORCEINLINE_DEBUGGABLE VkResult  vkCreateImageView(VkDevice Device, const VkImageViewCreateInfo* CreateInfo, const VkAllocationCallbacks* Allocator, VkImageView* View)
 	{
-		DumpImageViewCreate(Device, CreateInfo, View);
+		DumpCreateImageView(Device, CreateInfo, View);
 
 		VkResult Result = ::vkCreateImageView(Device, CreateInfo, Allocator, View);
 
diff --git a/Engine/Source/Runtime/VulkanRHI/Private/VulkanCommands.cpp b/Engine/Source/Runtime/VulkanRHI/Private/VulkanCommands.cpp
--- a/Engine/Source/Runtime/VulkanRHI/Private/VulkanCommands.cpp
+++ b/Engine/Source/Runtime/VulkanRHI/Private/VulkanCommands.cpp
@@ -511,7 +511,7 @@ void FVulkanCommandListContext::RHIDrawPrimitive(uint32 PrimitiveType, uint32 Ba
 
 	FVulkanCmdBuffer* CmdBuffer = CommandBufferManager->GetActiveCmdBuffer();
 	PendingState->PrepareDraw(this, CmdBuffer, UEToVulkanType((EPrimitiveType)PrimitiveType));
-	NumInstances = FMath::Min(1U, NumInstances);
+	NumInstances = FMath::Max(1U, NumInstances);
 	VulkanRHI::vkCmdDraw(CmdBuffer->GetHandle(), NumVertices, NumInstances, BaseVertexIndex, 0);
 
 	//if (IsImmediate())
@@ -556,10 +556,10 @@ void FVulkanCommandListContext::RHIDrawIndexedPrimitive(FIndexBufferRHIParamRef
 	VulkanRHI::vkCmdBindIndexBuffer(CmdBuffer, IndexBuffer->GetHandle(), IndexBuffer->GetOffset(), IndexBuffer->GetIndexType());
 
 	uint32 NumIndices = GetVertexCountForPrimitiveCount(NumPrimitives, PrimitiveType);
-	NumInstances = FMath::Min(1U, NumInstances);
+	NumInstances = FMath::Max(1U, NumInstances);
 	VulkanRHI::vkCmdDrawIndexed(CmdBuffer, NumIndices, NumInstances, StartIndex, BaseVertexIndex, FirstInstance);
 
-	//if (IsImmediate())
+	if (IsImmediate())
 	{
 		VulkanRHI::GManager.GPUProfilingData.RegisterGPUWork(NumPrimitives * NumInstances, NumVertices * NumInstances);
 	}
@@ -583,7 +583,7 @@ void FVulkanCommandListContext::RHIDrawIndexedIndirect(FIndexBufferRHIParamRef I
 #endif
 	VULKAN_SIGNAL_UNIMPLEMENTED();
 
-	//if (IsImmediate())
+	if (IsImmediate())
 	{
 		VulkanRHI::GManager.GPUProfilingData.RegisterGPUWork(0);
 	}
@@ -606,7 +606,7 @@ void FVulkanCommandListContext::RHIDrawIndexedPrimitiveIndirect(uint32 Primitive
 #endif
 	VULKAN_SIGNAL_UNIMPLEMENTED();
 
-	//if (IsImmediate())
+	if (IsImmediate())
 	{
 		VulkanRHI::GManager.GPUProfilingData.RegisterGPUWork(0);
 	}
@@ -643,7 +643,7 @@ void FVulkanCommandListContext::RHIEndDrawPrimitiveUP()
 	PendingState->PrepareDraw(this, CmdBuffer, UEToVulkanType((EPrimitiveType)PendingPrimitiveType));
 	VulkanRHI::vkCmdDraw(CmdBuffer->GetHandle(), PendingNumVertices, 1, PendingMinVertexIndex, 0);
 
-	//if (IsImmediate())
+	if (IsImmediate())
 	{
 		VulkanRHI::GManager.GPUProfilingData.RegisterGPUWork(PendingNumPrimitives, PendingNumVertices);
 	}
@@ -688,7 +688,7 @@ void FVulkanCommandListContext::RHIEndDrawIndexedPrimitiveUP()
 	VulkanRHI::vkCmdBindIndexBuffer(Cmd, PendingDrawPrimUPIndexAllocInfo.GetHandle(), PendingDrawPrimUPIndexAllocInfo.GetBindOffset(), PendingPrimitiveIndexType);
 	VulkanRHI::vkCmdDrawIndexed(Cmd, NumIndices, 1, PendingMinVertexIndex, 0, 0);
 
-	//if (IsImmediate())
+	if (IsImmediate())
 	{
 		VulkanRHI::GManager.GPUProfilingData.RegisterGPUWork(PendingNumPrimitives, PendingNumVertices);
 	}
@@ -721,7 +721,7 @@ void FVulkanCommandListContext::RHIClearMRT(bool bClearColor, int32 NumClearColo
 	PendingState->UpdateRenderPass(CmdBuffer);
 
 	const uint32 NumColorAttachments = PendingState->GetFrameBuffer()->GetNumColorAttachments();
-	check((uint32)NumClearColors <= NumColorAttachments);
+	check(!bClearColor || (uint32)NumClearColors <= NumColorAttachments);
 	InternalClearMRT(CmdBuffer, bClearColor, NumClearColors, ClearColorArray, bClearDepth, Depth, bClearStencil, Stencil, ExcludeRect);
 }
 
