Copying //UE4/Dev-Rendering to //UE4/Dev-Main (Source: //UE4/Dev-Rendering @ 3072736)

diff --git a/Engine/Source/Runtime/Renderer/Private/PostProcess/SceneRenderTargets.cpp b/Engine/Source/Runtime/Renderer/Private/PostProcess/SceneRenderTargets.cpp
--- a/Engine/Source/Runtime/Renderer/Private/PostProcess/SceneRenderTargets.cpp
+++ b/Engine/Source/Runtime/Renderer/Private/PostProcess/SceneRenderTargets.cpp
@@ -14,6 +14,7 @@
 #include "HdrCustomResolveShaders.h"
 #include "Public/LightPropagationVolumeBlendable.h"
 #include "Engine/EngineTypes.h"
+#include "WideCustomResolveShaders.h"
 
 IMPLEMENT_UNIFORM_BUFFER_STRUCT(FGBufferResourceStruct,TEXT("GBuffers"));
 
@@ -57,16 +58,6 @@ static TAutoConsoleVariable<int32> CVarSceneTargetsResizingMethod(
 	ECVF_RenderThreadSafe
 	);
 
-static TAutoConsoleVariable<int32> CVarOptimizeForUAVPerformance(
-	TEXT("r.OptimizeForUAVPerformance"),
-	0,
-	TEXT("Allows to profile if hardware has a performance cost due to render target reuse (more info: search for GCNPerformanceTweets.pdf Tip 37)\n")
-	TEXT("If we see a noticeable difference on some hardware we can add another option like -1 (meaning auto) and make it the new default.\n")
-	TEXT("0: Optimize for GPU memory savings and reuse render targets (default)\n")
-	TEXT("1: Optimize for GPU performance (might render faster but can require more GPU memory)"),
-	ECVF_RenderThreadSafe
-	);
-
 static TAutoConsoleVariable<int32> CVarCustomDepth(
 	TEXT("r.CustomDepth"),
 	1,
@@ -77,6 +68,27 @@ static TAutoConsoleVariable<int32> CVarCustomDepth(
 	ECVF_RenderThreadSafe
 	);
 
+static TAutoConsoleVariable<int32> CVarMSAACount(
+	TEXT("r.MSAACount"),
+	4,
+	TEXT("Number of MSAA samples to use with the forward renderer.  Only used when MSAA is enabled in the rendering project settings.\n")
+	TEXT("1: Use Temporal AA (MSAA disabled)\n")
+	TEXT("2: Use 2x MSAA (Temporal AA disabled)\n")
+	TEXT("4: Use 4x MSAA (Temporal AA disabled)"),
+	ECVF_RenderThreadSafe | ECVF_Scalability
+	);
+
+static TAutoConsoleVariable<int32> CVarWideCustomResolve(
+	TEXT("r.WideCustomResolve"),
+	0,
+	TEXT("Use a wide custom resolve filter when MSAA is enabled")
+	TEXT("0: Disabled [hardware box filter]")
+	TEXT("1: Wide (r=1.25, 12 samples)")
+	TEXT("2: Wider (r=1.4, 16 samples)")
+	TEXT("3: Widest (r=1.5, 20 samples)"),
+	ECVF_RenderThreadSafe | ECVF_Scalability
+	);
+
 static TAutoConsoleVariable<int32> CVarMobileMSAA(
 	TEXT("r.MobileMSAA"),
 	0,
@@ -186,7 +198,6 @@ FSceneRenderTargets::FSceneRenderTargets(const FViewInfo& View, const FSceneRend
 	, DirectionalOcclusion(GRenderTargetPool.MakeSnapshot(SnapshotSource.DirectionalOcclusion))
 	, SceneDepthZ(GRenderTargetPool.MakeSnapshot(SnapshotSource.SceneDepthZ))
 	, LightingChannels(GRenderTargetPool.MakeSnapshot(SnapshotSource.LightingChannels))
-	, NoMSAASceneDepthZ(GRenderTargetPool.MakeSnapshot(SnapshotSource.NoMSAASceneDepthZ))
 	, SceneAlphaCopy(GRenderTargetPool.MakeSnapshot(SnapshotSource.SceneAlphaCopy))
 	, AuxiliarySceneDepthZ(GRenderTargetPool.MakeSnapshot(SnapshotSource.AuxiliarySceneDepthZ))
 	, SmallDepthZ(GRenderTargetPool.MakeSnapshot(SnapshotSource.SmallDepthZ))
@@ -205,7 +216,6 @@ FSceneRenderTargets::FSceneRenderTargets(const FViewInfo& View, const FSceneRend
 	, CustomDepth(GRenderTargetPool.MakeSnapshot(SnapshotSource.CustomDepth))
 	, CustomStencilSRV(SnapshotSource.CustomStencilSRV)
 	, OptionalShadowDepthColor(GRenderTargetPool.MakeSnapshot(SnapshotSource.OptionalShadowDepthColor))
-	, PreShadowCacheDepthZ(GRenderTargetPool.MakeSnapshot(SnapshotSource.PreShadowCacheDepthZ))
 	, SkySHIrradianceMap(GRenderTargetPool.MakeSnapshot(SnapshotSource.SkySHIrradianceMap))
 	, EditorPrimitivesColor(GRenderTargetPool.MakeSnapshot(SnapshotSource.EditorPrimitivesColor))
 	, EditorPrimitivesDepth(GRenderTargetPool.MakeSnapshot(SnapshotSource.EditorPrimitivesDepth))
@@ -213,7 +223,6 @@ FSceneRenderTargets::FSceneRenderTargets(const FViewInfo& View, const FSceneRend
 	, SeparateTranslucencyDepthRT(SnapshotSource.SeparateTranslucencyDepthRT)
 	, bScreenSpaceAOIsValid(SnapshotSource.bScreenSpaceAOIsValid)
 	, bCustomDepthIsValid(SnapshotSource.bCustomDepthIsValid)
-	, bPreshadowCacheNewlyAllocated(SnapshotSource.bPreshadowCacheNewlyAllocated)
 	, GBufferRefCount(SnapshotSource.GBufferRefCount)
 	, LargestDesiredSizeThisFrame(SnapshotSource.LargestDesiredSizeThisFrame)
 	, LargestDesiredSizeLastFrame(SnapshotSource.LargestDesiredSizeLastFrame)
@@ -234,7 +243,7 @@ FSceneRenderTargets::FSceneRenderTargets(const FViewInfo& View, const FSceneRend
 	, CurrentRSMResolution(SnapshotSource.CurrentRSMResolution)
 	, CurrentTranslucencyLightingVolumeDim(SnapshotSource.CurrentTranslucencyLightingVolumeDim)
 	, CurrentMobile32bpp(SnapshotSource.CurrentMobile32bpp)
-	, CurrentMobileMSAA(SnapshotSource.CurrentMobileMSAA)
+	, CurrentMSAACount(SnapshotSource.CurrentMSAACount)
 	, CurrentMinShadowResolution(SnapshotSource.CurrentMinShadowResolution)
 	, bCurrentLightPropagationVolume(SnapshotSource.bCurrentLightPropagationVolume)
 	, CurrentFeatureLevel(SnapshotSource.CurrentFeatureLevel)
@@ -248,23 +257,10 @@ FSceneRenderTargets::FSceneRenderTargets(const FViewInfo& View, const FSceneRend
 	SnapshotArray(SceneColor, SnapshotSource.SceneColor);
 	SnapshotArray(ReflectionColorScratchCubemap, SnapshotSource.ReflectionColorScratchCubemap);
 	SnapshotArray(DiffuseIrradianceScratchCubemap, SnapshotSource.DiffuseIrradianceScratchCubemap);
-	SnapshotArray(ReflectionBrightness, SnapshotSource.ReflectionBrightness);
 	SnapshotArray(TranslucencyLightingVolumeAmbient, SnapshotSource.TranslucencyLightingVolumeAmbient);
 	SnapshotArray(TranslucencyLightingVolumeDirectional, SnapshotSource.TranslucencyLightingVolumeDirectional);
 }
 
-
-inline const TCHAR* GetReflectionBrightnessTargetName(uint32 Index)
-{
-	const TCHAR* Names[2] = 
-	{ 
-		TEXT("ReflectionBrightness0"), 
-		TEXT("ReflectionBrightness1") 
-	};
-	check(Index < ARRAY_COUNT(Names));
-	return Names[Index];
-}
-
 inline const TCHAR* GetSceneColorTargetName(EShadingPath ShadingPath)
 {
 	const TCHAR* SceneColorNames[(uint32)EShadingPath::Num] =
@@ -351,20 +347,37 @@ FIntPoint FSceneRenderTargets::ComputeDesiredSize(const FSceneViewFamily& ViewFa
 
 	return DesiredBufferSize;
 }
+
 inline uint16 GetNumSceneColorMSAASamples(ERHIFeatureLevel::Type InFeatureLevel)
 {
-	if (InFeatureLevel > ERHIFeatureLevel::ES3_1 ||
-		//@todo-rco: Fix when iOS OpenGL supports MSAA
-		GShaderPlatformForFeatureLevel[InFeatureLevel] == SP_OPENGL_ES2_IOS)
+	uint16 NumSamples = 1;
+
+	if (InFeatureLevel >= ERHIFeatureLevel::SM4)
+	{
+		static IConsoleVariable* CVarDefaultAntiAliasing = IConsoleManager::Get().FindConsoleVariable(TEXT("r.DefaultFeature.AntiAliasing"));
+		EAntiAliasingMethod Method = (EAntiAliasingMethod)CVarDefaultAntiAliasing->GetInt();
+
+		if (IsForwardShadingEnabled(InFeatureLevel) && Method == AAM_MSAA)
+		{
+			NumSamples = CVarMSAACount.GetValueOnRenderThread();
+		}
+	}
+	else
 	{
-		return 1;
+		//@todo-rco: Fix when iOS OpenGL supports MSAA
+		if (GShaderPlatformForFeatureLevel[InFeatureLevel] == SP_OPENGL_ES2_IOS)
+		{
+			return 1;
+		}
+
+		NumSamples = CVarMobileMSAA.GetValueOnRenderThread();
 	}
 
-	uint16 NumSamples = CVarMobileMSAA.GetValueOnRenderThread();
 	if (NumSamples != 1 && NumSamples != 2 && NumSamples != 4)
 	{
 		NumSamples = 1;
 	}
+
 	return NumSamples;
 }
 
@@ -419,7 +432,7 @@ void FSceneRenderTargets::Allocate(FRHICommandList& RHICmdList, const FSceneView
 
 	uint32 Mobile32bpp = !IsMobileHDR() || IsMobileHDR32bpp();
 
-	int32 MobileMSAA = GetNumSceneColorMSAASamples(NewFeatureLevel);
+	int32 MSAACount = GetNumSceneColorMSAASamples(NewFeatureLevel);
 
 	bool bLightPropagationVolume = UseLightPropagationVolumeRT(NewFeatureLevel);
 
@@ -440,7 +453,7 @@ void FSceneRenderTargets::Allocate(FRHICommandList& RHICmdList, const FSceneView
  		(CurrentRSMResolution != RSMResolution) ||
 		(CurrentTranslucencyLightingVolumeDim != TranslucencyLightingVolumeDim) ||
 		(CurrentMobile32bpp != Mobile32bpp) ||
-		(CurrentMobileMSAA != MobileMSAA) ||
+		(CurrentMSAACount != MSAACount) ||
 		(bCurrentLightPropagationVolume != bLightPropagationVolume) ||
 		(CurrentMinShadowResolution != MinShadowResolution))
 	{
@@ -452,7 +465,7 @@ void FSceneRenderTargets::Allocate(FRHICommandList& RHICmdList, const FSceneView
 		CurrentRSMResolution = RSMResolution;
 		CurrentTranslucencyLightingVolumeDim = TranslucencyLightingVolumeDim;
 		CurrentMobile32bpp = Mobile32bpp;
-		CurrentMobileMSAA = MobileMSAA;
+		CurrentMSAACount = MSAACount;
 		CurrentMinShadowResolution = MinShadowResolution;
 		bCurrentLightPropagationVolume = bLightPropagationVolume;
 
@@ -495,12 +508,11 @@ int32 FSceneRenderTargets::GetGBufferRenderTargets(ERenderTargetLoadAction Color
 	}
 	else
 	{
-	OutVelocityRTIndex = -1;
+		OutVelocityRTIndex = -1;
 	}
 
 	OutRenderTargets[MRTCount++] = FRHIRenderTargetView(GBufferD->GetRenderTargetItem().TargetableTexture, 0, -1, ColorLoadAction, ERenderTargetStoreAction::EStore);
 
-
 	if (bAllowStaticLighting)
 	{
 		check(MRTCount == (bAllocateVelocityGBuffer ? 6 : 5)); // As defined in BasePassPixelShader.usf
@@ -553,8 +565,6 @@ void FSceneRenderTargets::BeginRenderingGBuffer(FRHICommandList& RHICmdList, ERe
 
 		//make sure our conditions for shader clear fallback are valid.
 		check(RenderTargets[0].Texture == SceneColorTex);
-		
-		
 
 		FRHIDepthRenderTargetView DepthView(GetSceneDepthSurface(), DepthLoadAction, ERenderTargetStoreAction::EStore);
 		FRHISetRenderTargetsInfo Info(MRTCount, RenderTargets, DepthView);		
@@ -605,13 +615,6 @@ void FSceneRenderTargets::BeginRenderingGBuffer(FRHICommandList& RHICmdList, ERe
 
 void FSceneRenderTargets::FinishRenderingGBuffer(FRHICommandListImmediate& RHICmdList)
 {
-	if (IsAnyForwardShadingEnabled(GetFeatureLevelShaderPlatform(CurrentFeatureLevel)))
-	{
-		// in this non-standard case, just render to scene color with default mode
-		FinishRenderingSceneColor(RHICmdList, true);
-		return;
-	}
-
 	int32 VelocityRTIndex;
 	FRHIRenderTargetView RenderTargets[MaxSimultaneousRenderTargets];
 	int32 NumMRTs = GetGBufferRenderTargets(ERenderTargetLoadAction::ELoad, RenderTargets, VelocityRTIndex);
@@ -622,8 +625,8 @@ void FSceneRenderTargets::FinishRenderingGBuffer(FRHICommandListImmediate& RHICm
 		 // When the basepass outputs to the velocity buffer, don't resolve it yet if selective outputs are enabled, as it will be resolved after the velocity pass.
 		if (i != VelocityRTIndex || !UseSelectiveBasePassOutputs())
 		{
-		RHICmdList.CopyToResolveTarget(RenderTargets[i].Texture, RenderTargets[i].Texture, true, ResolveParams);
-	}
+			RHICmdList.CopyToResolveTarget(RenderTargets[i].Texture, RenderTargets[i].Texture, true, ResolveParams);
+		}
 	}
 	FTextureRHIParamRef DepthSurface = GetSceneDepthSurface();
 	RHICmdList.CopyToResolveTarget(DepthSurface, DepthSurface, true, ResolveParams);
@@ -652,12 +655,9 @@ void FSceneRenderTargets::AllocSceneColor(FRHICommandList& RHICmdList)
 {
 	if (GetSceneColorForCurrentShadingPath())
 	{
-		// no work needed
 		return;
 	}
 
-	// create SceneColor on demand so it can be shared with other pooled RT
-
 	EPixelFormat SceneColorBufferFormat = GetSceneColorFormat();
 
 	// Create the scene color.
@@ -665,25 +665,17 @@ void FSceneRenderTargets::AllocSceneColor(FRHICommandList& RHICmdList)
 		FPooledRenderTargetDesc Desc(FPooledRenderTargetDesc::Create2DDesc(BufferSize, SceneColorBufferFormat, FClearValueBinding::Black, TexCreate_None, TexCreate_RenderTargetable, false));
 
 		Desc.Flags |= TexCreate_FastVRAM;
+		Desc.NumSamples = GetNumSceneColorMSAASamples(CurrentFeatureLevel);
 
-		int32 OptimizeForUAVPerformance = CVarOptimizeForUAVPerformance.GetValueOnRenderThread();
-
-		// with TexCreate_UAV it would allow better sharing with later elements but it might come at a high cost:
-		// GCNPerformanceTweets.pdf Tip 37: Warning: Causes additional synchronization between draw calls when using a render target allocated with this flag, use sparingly
-		if (CurrentFeatureLevel >= ERHIFeatureLevel::SM5 && !OptimizeForUAVPerformance)
+		if (CurrentFeatureLevel >= ERHIFeatureLevel::SM5 && Desc.NumSamples == 1)
 		{
+			// GCNPerformanceTweets.pdf Tip 37: Warning: Causes additional synchronization between draw calls when using a render target allocated with this flag, use sparingly
 			Desc.TargetableFlags |= TexCreate_UAV;
 		}
 
-		if (CurrentFeatureLevel < ERHIFeatureLevel::SM4)
-		{
-			Desc.NumSamples = GetNumSceneColorMSAASamples(CurrentFeatureLevel);
-		}
-
 		GRenderTargetPool.FindFreeElement(RHICmdList, Desc, GetSceneColorForCurrentShadingPath(), GetSceneColorTargetName(CurrentShadingPath));
 	}
 
-	// otherwise we have a severe problem
 	check(GetSceneColorForCurrentShadingPath());
 }
 
@@ -697,28 +689,17 @@ void FSceneRenderTargets::AllocLightAttenuation(FRHICommandList& RHICmdList)
 
 	check(IsInRenderingThread());
 
-	// create LightAttenuation on demand so it can be shared with other pooled RT
-
 	// Create a texture to store the resolved light attenuation values, and a render-targetable surface to hold the unresolved light attenuation values.
 	{
 		FPooledRenderTargetDesc Desc(FPooledRenderTargetDesc::Create2DDesc(BufferSize, PF_B8G8R8A8, FClearValueBinding::White, TexCreate_None, TexCreate_RenderTargetable, false));
 		Desc.Flags |= TexCreate_FastVRAM;
 		GRenderTargetPool.FindFreeElement(RHICmdList, Desc, LightAttenuation, TEXT("LightAttenuation"));
-
-		// the channel assignment is documented in ShadowRendering.cpp (look for Light Attenuation channel assignment)
 	}
 
 	// otherwise we have a severe problem
 	check(LightAttenuation);
 }
 
-TRefCountPtr<IPooledRenderTarget>& FSceneRenderTargets::GetReflectionBrightnessTarget()
-{
-	bool bSupportsR32Float = CurrentFeatureLevel > ERHIFeatureLevel::ES2;
-	int32 ReflectionBrightnessIndex = bSupportsR32Float ? 0 : 1;
-	return ReflectionBrightness[ReflectionBrightnessIndex];
-}
-
 void FSceneRenderTargets::ReleaseGBufferTargets()
 {
 	GBufferResourcesUniformBuffer.SafeRelease();
@@ -742,8 +723,33 @@ void FSceneRenderTargets::PreallocGBufferTargets(bool bShouldRenderVelocities)
 	}
 }
 
+void FSceneRenderTargets::GetGBufferADesc(FPooledRenderTargetDesc& Desc) const
+{
+	// good to see the quality loss due to precision in the gbuffer
+	const bool bHighPrecisionGBuffers = (CurrentGBufferFormat >= EGBufferFormat::Force16BitsPerChannel);
+	// good to profile the impact of non 8 bit formats
+	const bool bEnforce8BitPerChannel = (CurrentGBufferFormat == EGBufferFormat::Force8BitsPerChannel);
+
+	// Create the world-space normal g-buffer.
+	{
+		EPixelFormat NormalGBufferFormat = bHighPrecisionGBuffers ? PF_FloatRGBA : PF_A2B10G10R10;
+
+		if (bEnforce8BitPerChannel)
+		{
+			NormalGBufferFormat = PF_B8G8R8A8;
+		}
+		else if (CurrentGBufferFormat == EGBufferFormat::HighPrecisionNormals)
+		{
+			NormalGBufferFormat = PF_FloatRGBA;
+		}
+
+		Desc = FPooledRenderTargetDesc::Create2DDesc(BufferSize, NormalGBufferFormat, FClearValueBinding::Transparent, TexCreate_None, TexCreate_RenderTargetable, false);
+	}
+}
+
 void FSceneRenderTargets::AllocGBufferTargets(FRHICommandList& RHICmdList)
 {
+	check(!IsAnyForwardShadingEnabled(GetFeatureLevelShaderPlatform(CurrentFeatureLevel)));
 	// AdjustGBufferRefCount +1 doesn't match -1 (within the same frame)
 	ensure(GBufferRefCount == 0);
 
@@ -762,18 +768,8 @@ void FSceneRenderTargets::AllocGBufferTargets(FRHICommandList& RHICmdList)
 
 	// Create the world-space normal g-buffer.
 	{
-		EPixelFormat NormalGBufferFormat = bHighPrecisionGBuffers ? PF_FloatRGBA : PF_A2B10G10R10;
-
-		if(bEnforce8BitPerChannel)
-		{
-			NormalGBufferFormat = PF_B8G8R8A8;
-		}
-		else if (CurrentGBufferFormat == EGBufferFormat::HighPrecisionNormals)
-		{
-			NormalGBufferFormat = PF_FloatRGBA;
-		}
-
-		FPooledRenderTargetDesc Desc(FPooledRenderTargetDesc::Create2DDesc(BufferSize, NormalGBufferFormat, FClearValueBinding::Transparent, TexCreate_None, TexCreate_RenderTargetable, false));
+		FPooledRenderTargetDesc Desc;
+		GetGBufferADesc(Desc);
 		GRenderTargetPool.FindFreeElement(RHICmdList, Desc, GBufferA, TEXT("GBufferA"));
 	}
 
@@ -859,7 +855,6 @@ void FSceneRenderTargets::AllocGBufferTargets(FRHICommandList& RHICmdList)
 		GBufferResourcesUniformBuffer = FGBufferResourceStruct::CreateUniformBuffer(GBufferResourceStruct, UniformBuffer_SingleFrame);
 	}
 
-	// so that 
 	GBufferRefCount = 1;
 }
 
@@ -961,30 +956,29 @@ TRefCountPtr<IPooledRenderTarget>& FSceneRenderTargets::GetLightAttenuation()
 
 void FSceneRenderTargets::AdjustGBufferRefCount(FRHICommandList& RHICmdList, int Delta)
 { 
-	
-	if(Delta > 0 && GBufferRefCount == 0)
-	{
-		AllocGBufferTargets(RHICmdList);
-	}
-	else
+	if (!IsAnyForwardShadingEnabled(GetFeatureLevelShaderPlatform(CurrentFeatureLevel)))
 	{
-		GBufferRefCount += Delta;
-
-		if (GBufferRefCount == 0)
+		if (Delta > 0 && GBufferRefCount == 0)
+		{
+			AllocGBufferTargets(RHICmdList);
+		}
+		else
 		{
-			ReleaseGBufferTargets();
+			GBufferRefCount += Delta;
+
+			if (GBufferRefCount == 0)
+			{
+				ReleaseGBufferTargets();
+			}
 		}
 	}
 }
 
-void FSceneRenderTargets::FinishRenderingSceneColor(FRHICommandListImmediate& RHICmdList, bool bKeepChanges, const FResolveRect& ResolveRect)
+void FSceneRenderTargets::FinishRenderingSceneColor(FRHICommandListImmediate& RHICmdList, const FResolveRect& ResolveRect)
 {
 	SCOPED_DRAW_EVENT(RHICmdList, FinishRenderingSceneColor);
 
-	if(bKeepChanges)
-	{
-		ResolveSceneColor(RHICmdList);
-	}
+	ResolveSceneColor(RHICmdList);
 }
 
 bool FSceneRenderTargets::BeginRenderingCustomDepth(FRHICommandListImmediate& RHICmdList, bool bPrimitives)
@@ -1025,11 +1019,10 @@ void FSceneRenderTargets::ResolveSceneColor(FRHICommandList& RHICmdList, const F
 	SCOPED_DRAW_EVENT(RHICmdList, ResolveSceneColor);
 
 	auto& CurrentSceneColor = GetSceneColor();
-
-	uint32 samples = CurrentSceneColor->GetDesc().NumSamples;
+	uint32 CurrentNumSamples = CurrentSceneColor->GetDesc().NumSamples;
 
 	const EShaderPlatform CurrentShaderPlatform = GShaderPlatformForFeatureLevel[CurrentFeatureLevel];
-	if (samples <= 1 || CurrentShaderPlatform == SP_METAL || IsVulkanMobilePlatform(CurrentShaderPlatform))
+	if (CurrentNumSamples <= 1 || !RHISupportsSeparateMSAAAndResolveTextures(CurrentShaderPlatform))
 	{
 		RHICmdList.CopyToResolveTarget(GetSceneColorSurface(), GetSceneColorTexture(), true, FResolveParams(ResolveRect));
 	}
@@ -1051,61 +1044,51 @@ void FSceneRenderTargets::ResolveSceneColor(FRHICommandList& RHICmdList, const F
 		RHICmdList.SetRasterizerState(TStaticRasterizerState<>::GetRHI());
 		RHICmdList.SetDepthStencilState(TStaticDepthStencilState<false, CF_Always>::GetRHI());
 
-		auto ShaderMap = GetGlobalShaderMap(CurrentFeatureLevel);
-		TShaderMapRef<FHdrCustomResolveVS> VertexShader(ShaderMap);
+		int32 ResolveWidth = CVarWideCustomResolve.GetValueOnRenderThread();
 
-		if(samples == 2)
+		if (CurrentNumSamples <= 1)
 		{
-			TShaderMapRef<FHdrCustomResolve2xPS> PixelShader(ShaderMap);
-			static FGlobalBoundShaderState BoundShaderState;
-			SetGlobalBoundShaderState(RHICmdList, CurrentFeatureLevel, BoundShaderState, GetVertexDeclarationFVector4(), *VertexShader, *PixelShader);
-			PixelShader->SetParameters(RHICmdList, CurrentSceneColor->GetRenderTargetItem().TargetableTexture);
-			RHICmdList.DrawPrimitive(PT_TriangleList, 0, 1, 1);
+			ResolveWidth = 0;
 		}
-		else if(samples == 4)
-		{
-			TShaderMapRef<FHdrCustomResolve4xPS> PixelShader(ShaderMap);
-			static FGlobalBoundShaderState BoundShaderState;
-			SetGlobalBoundShaderState(RHICmdList, CurrentFeatureLevel, BoundShaderState, GetVertexDeclarationFVector4(), *VertexShader, *PixelShader);
-			PixelShader->SetParameters(RHICmdList, CurrentSceneColor->GetRenderTargetItem().TargetableTexture);
-			RHICmdList.DrawPrimitive(PT_TriangleList, 0, 1, 1);
-		}
-		else if(samples == 8)
+
+		if (ResolveWidth != 0)
 		{
-			TShaderMapRef<FHdrCustomResolve8xPS> PixelShader(ShaderMap);
-			static FGlobalBoundShaderState BoundShaderState;
-			SetGlobalBoundShaderState(RHICmdList, CurrentFeatureLevel, BoundShaderState, GetVertexDeclarationFVector4(), *VertexShader, *PixelShader);
-			PixelShader->SetParameters(RHICmdList, CurrentSceneColor->GetRenderTargetItem().TargetableTexture);
-			RHICmdList.DrawPrimitive(PT_TriangleList, 0, 1, 1);
+			ResolveFilterWide(RHICmdList, CurrentFeatureLevel, CurrentSceneColor->GetRenderTargetItem().TargetableTexture, FIntPoint(0, 0), CurrentNumSamples, ResolveWidth);
 		}
 		else
 		{
-			// Everything other than 2,4,8 samples is not implemented.
-			check(0);
-		}
-	}
-}
-
-/** Resolves the GBuffer targets so that their resolved textures can be sampled. */
-void FSceneRenderTargets::ResolveGBufferSurfaces(FRHICommandList& RHICmdList, const FResolveRect& ResolveRect)
-{
-	if (CurrentFeatureLevel >= ERHIFeatureLevel::SM4)
-	{
-		SCOPED_DRAW_EVENT(RHICmdList, ResolveGBufferSurfaces);
-
-		RHICmdList.CopyToResolveTarget(GBufferA->GetRenderTargetItem().TargetableTexture, GBufferA->GetRenderTargetItem().ShaderResourceTexture, true, FResolveParams(ResolveRect));
-		RHICmdList.CopyToResolveTarget(GBufferB->GetRenderTargetItem().TargetableTexture, GBufferB->GetRenderTargetItem().ShaderResourceTexture, true, FResolveParams(ResolveRect));
-		RHICmdList.CopyToResolveTarget(GBufferC->GetRenderTargetItem().TargetableTexture, GBufferC->GetRenderTargetItem().ShaderResourceTexture, true, FResolveParams(ResolveRect));
-		RHICmdList.CopyToResolveTarget(GBufferD->GetRenderTargetItem().TargetableTexture, GBufferD->GetRenderTargetItem().ShaderResourceTexture, true, FResolveParams(ResolveRect));
-		
-		if (bAllowStaticLighting)
-		{
-			RHICmdList.CopyToResolveTarget(GBufferE->GetRenderTargetItem().TargetableTexture, GBufferE->GetRenderTargetItem().ShaderResourceTexture, true, FResolveParams(ResolveRect));
-		}
+			auto ShaderMap = GetGlobalShaderMap(CurrentFeatureLevel);
+			TShaderMapRef<FHdrCustomResolveVS> VertexShader(ShaderMap);
 
-		if (bAllocateVelocityGBuffer)
-		{
-			RHICmdList.CopyToResolveTarget(GBufferVelocity->GetRenderTargetItem().TargetableTexture, GBufferVelocity->GetRenderTargetItem().ShaderResourceTexture, true, FResolveParams(ResolveRect));
+			if (CurrentNumSamples == 2)
+			{
+				TShaderMapRef<FHdrCustomResolve2xPS> PixelShader(ShaderMap);
+				static FGlobalBoundShaderState BoundShaderState;
+				SetGlobalBoundShaderState(RHICmdList, CurrentFeatureLevel, BoundShaderState, GetVertexDeclarationFVector4(), *VertexShader, *PixelShader);
+				PixelShader->SetParameters(RHICmdList, CurrentSceneColor->GetRenderTargetItem().TargetableTexture);
+				RHICmdList.DrawPrimitive(PT_TriangleList, 0, 1, 1);
+			}
+			else if (CurrentNumSamples == 4)
+			{
+				TShaderMapRef<FHdrCustomResolve4xPS> PixelShader(ShaderMap);
+				static FGlobalBoundShaderState BoundShaderState;
+				SetGlobalBoundShaderState(RHICmdList, CurrentFeatureLevel, BoundShaderState, GetVertexDeclarationFVector4(), *VertexShader, *PixelShader);
+				PixelShader->SetParameters(RHICmdList, CurrentSceneColor->GetRenderTargetItem().TargetableTexture);
+				RHICmdList.DrawPrimitive(PT_TriangleList, 0, 1, 1);
+			}
+			else if (CurrentNumSamples == 8)
+			{
+				TShaderMapRef<FHdrCustomResolve8xPS> PixelShader(ShaderMap);
+				static FGlobalBoundShaderState BoundShaderState;
+				SetGlobalBoundShaderState(RHICmdList, CurrentFeatureLevel, BoundShaderState, GetVertexDeclarationFVector4(), *VertexShader, *PixelShader);
+				PixelShader->SetParameters(RHICmdList, CurrentSceneColor->GetRenderTargetItem().TargetableTexture);
+				RHICmdList.DrawPrimitive(PT_TriangleList, 0, 1, 1);
+			}
+			else
+			{
+				// Everything other than 2,4,8 samples is not implemented.
+				check(0);
+			}
 		}
 	}
 }
@@ -1173,7 +1156,7 @@ void FSceneRenderTargets::BeginRenderingLightAttenuation(FRHICommandList& RHICmd
 	GRenderTargetPool.VisualizeTexture.SetCheckPoint(RHICmdList, GetLightAttenuation());
 
 	// Set the light attenuation surface as the render target, and the scene depth buffer as the depth-stencil surface.
-	SetRenderTarget(RHICmdList, GetLightAttenuationSurface(), GetSceneDepthSurface(), bClearToWhite ? ESimpleRenderTargetMode::EClearColorExistingDepth : ESimpleRenderTargetMode::EExistingColorAndDepth, FExclusiveDepthStencil::DepthRead_StencilWrite, true);
+	SetRenderTarget(RHICmdList, GetLightAttenuationSurface(), GetSceneDepthTexture(), bClearToWhite ? ESimpleRenderTargetMode::EClearColorExistingDepth : ESimpleRenderTargetMode::EExistingColorAndDepth, FExclusiveDepthStencil::DepthRead_StencilWrite, true);
 }
 
 void FSceneRenderTargets::FinishRenderingLightAttenuation(FRHICommandList& RHICmdList)
@@ -1186,6 +1169,33 @@ void FSceneRenderTargets::FinishRenderingLightAttenuation(FRHICommandList& RHICm
 	GRenderTargetPool.VisualizeTexture.SetCheckPoint(RHICmdList, GetLightAttenuation());
 }
 
+TRefCountPtr<IPooledRenderTarget>& FSceneRenderTargets::GetSeparateTranslucency(FRHICommandList& RHICmdList, FIntPoint Size)
+{
+	if (!SeparateTranslucencyRT || SeparateTranslucencyRT->GetDesc().Extent != Size)
+	{
+		uint32 Flags = TexCreate_RenderTargetable;
+
+		// Create the SeparateTranslucency render target (alpha is needed to lerping)
+		FPooledRenderTargetDesc Desc(FPooledRenderTargetDesc::Create2DDesc(Size, PF_FloatRGBA, FClearValueBinding::Black, TexCreate_None, Flags, false));
+		Desc.AutoWritable = false;
+		Desc.NumSamples = GetNumSceneColorMSAASamples(CurrentFeatureLevel);
+		GRenderTargetPool.FindFreeElement(RHICmdList, Desc, SeparateTranslucencyRT, TEXT("SeparateTranslucency"));
+	}
+	return SeparateTranslucencyRT;
+}
+
+TRefCountPtr<IPooledRenderTarget>& FSceneRenderTargets::GetSeparateTranslucencyDepth(FRHICommandList& RHICmdList, FIntPoint Size)
+{
+	if (!SeparateTranslucencyDepthRT || SeparateTranslucencyDepthRT->GetDesc().Extent != Size)
+	{
+		// Create the SeparateTranslucency depth render target 
+		FPooledRenderTargetDesc Desc(FPooledRenderTargetDesc::Create2DDesc(Size, PF_DepthStencil, FClearValueBinding::None, TexCreate_None, TexCreate_DepthStencilTargetable, false));
+		Desc.NumSamples = GetNumSceneColorMSAASamples(CurrentFeatureLevel);
+		GRenderTargetPool.FindFreeElement(RHICmdList, Desc, SeparateTranslucencyDepthRT, TEXT("SeparateTranslucencyDepth"));
+	}
+	return SeparateTranslucencyDepthRT;
+}
+
 void FSceneRenderTargets::BeginRenderingTranslucency(FRHICommandList& RHICmdList, const FViewInfo& View, bool bFirstTimeThisFrame)
 {
 	// Use the scene color buffer.
@@ -1203,7 +1213,7 @@ void FSceneRenderTargets::BeginRenderingTranslucency(FRHICommandList& RHICmdList
 
 void FSceneRenderTargets::FinishRenderingTranslucency(FRHICommandListImmediate& RHICmdList, const class FViewInfo& View)
 {
-	FinishRenderingSceneColor(RHICmdList, true);
+	FinishRenderingSceneColor(RHICmdList);
 }
 
 
@@ -1528,22 +1538,6 @@ void FSceneRenderTargets::AllocateReflectionTargets(FRHICommandList& RHICmdList,
 				GRenderTargetPool.FindFreeElement(RHICmdList, Desc, SkySHIrradianceMap, TEXT("SkySHIrradianceMap"));
 			}
 		}
-
-		auto& ReflectionBrightnessTarget = GetReflectionBrightnessTarget();
-		if (!ReflectionBrightnessTarget)
-		{
-			bool bSupportsR32Float = CurrentFeatureLevel > ERHIFeatureLevel::ES2;
-			//@todo: FScene::UpdateSkyCaptureContents() is called before FSceneRenderTargets::AllocateRenderTargets()
-			// so CurrentFeatureLevel == ERHIFeatureLevel::Num and that crashes OpenGL ES2 as R32F is not valid
-			if (CurrentFeatureLevel == ERHIFeatureLevel::Num)
-			{
-				bSupportsR32Float = GMaxRHIFeatureLevel > ERHIFeatureLevel::ES2;
-			}
-			int32 ReflectionBrightnessIndex = bSupportsR32Float ? 0 : 1;
-			EPixelFormat BrightnessFormat = bSupportsR32Float ? PF_R32_FLOAT : PF_FloatRGBA;
-			FPooledRenderTargetDesc Desc3(FPooledRenderTargetDesc::Create2DDesc(FIntPoint(1, 1), BrightnessFormat, FClearValueBinding::None, TexCreate_None, TexCreate_RenderTargetable, false));
-			GRenderTargetPool.FindFreeElement(RHICmdList, Desc3, ReflectionBrightnessTarget, GetReflectionBrightnessTargetName(ReflectionBrightnessIndex));
-		}
 	}
 }
 
@@ -1579,18 +1573,6 @@ void FSceneRenderTargets::AllocateCommonDepthTargets(FRHICommandList& RHICmdList
 		Desc.Flags |= TexCreate_FastVRAM;
 		GRenderTargetPool.FindFreeElement(RHICmdList, Desc, SceneDepthZ, TEXT("SceneDepthZ"));
 		SceneStencilSRV = RHICreateShaderResourceView((FTexture2DRHIRef&)SceneDepthZ->GetRenderTargetItem().TargetableTexture, 0, 1, PF_X24_G8);
-
-		// create non-MSAA version for hit proxies on PC
-		const EShaderPlatform CurrentShaderPlatform = GShaderPlatformForFeatureLevel[CurrentFeatureLevel];
-		if (Desc.NumSamples > 1 && CurrentShaderPlatform != SP_METAL && !IsVulkanMobilePlatform(CurrentShaderPlatform))
-		{
-			Desc.NumSamples = 1;
-			GRenderTargetPool.FindFreeElement(RHICmdList, Desc, NoMSAASceneDepthZ, TEXT("NoMSAASceneDepthZ"));
-		}
-		else
-		{
-			NoMSAASceneDepthZ = SceneDepthZ;
-		}
 	}
 
 	// When targeting DX Feature Level 10, create an auxiliary texture to store the resolved scene depth, and a render-targetable surface to hold the unresolved scene depth.
@@ -1601,14 +1583,6 @@ void FSceneRenderTargets::AllocateCommonDepthTargets(FRHICommandList& RHICmdList
 		GRenderTargetPool.FindFreeElement(RHICmdList, Desc, AuxiliarySceneDepthZ, TEXT("AuxiliarySceneDepthZ"));
 	}
 
-	{
-		FPooledRenderTargetDesc Desc(FPooledRenderTargetDesc::Create2DDesc(GetPreShadowCacheTextureResolution(), PF_ShadowDepth, FClearValueBinding::None, TexCreate_None, TexCreate_DepthStencilTargetable, false));
-		Desc.AutoWritable = false;
-		GRenderTargetPool.FindFreeElement(RHICmdList, Desc, PreShadowCacheDepthZ, TEXT("PreShadowCacheDepthZ"));
-		// Mark the preshadow cache as newly allocated, so the cache will know to update
-		bPreshadowCacheNewlyAllocated = true;
-	}
-
 	if (!GSupportsDepthRenderTargetWithoutColorRenderTarget)
 	{
 		const FIntPoint ShadowBufferResolution = GetShadowDepthTextureResolution();
@@ -1804,7 +1778,6 @@ void FSceneRenderTargets::ReleaseAllTargets()
 	SceneDepthZ.SafeRelease();
 	SceneStencilSRV.SafeRelease();
 	LightingChannels.SafeRelease();
-	NoMSAASceneDepthZ.SafeRelease();
 	AuxiliarySceneDepthZ.SafeRelease();
 	SmallDepthZ.SafeRelease();
 	DBufferA.SafeRelease();
@@ -1820,16 +1793,11 @@ void FSceneRenderTargets::ReleaseAllTargets()
 
 	OptionalShadowDepthColor.SafeRelease();
 
-	PreShadowCacheDepthZ.SafeRelease();
-	
 	for (int32 i = 0; i < ARRAY_COUNT(ReflectionColorScratchCubemap); i++)
 	{
 		ReflectionColorScratchCubemap[i].SafeRelease();
 	}
 
-	ReflectionBrightness[0].SafeRelease();
-	ReflectionBrightness[1].SafeRelease();
-
 	for (int32 i = 0; i < ARRAY_COUNT(DiffuseIrradianceScratchCubemap); i++)
 	{
 		DiffuseIrradianceScratchCubemap[i].SafeRelease();
diff --git a/Engine/Source/Runtime/Renderer/Private/PostProcess/SceneRenderTargets.h b/Engine/Source/Runtime/Renderer/Private/PostProcess/SceneRenderTargets.h
--- a/Engine/Source/Runtime/Renderer/Private/PostProcess/SceneRenderTargets.h
+++ b/Engine/Source/Runtime/Renderer/Private/PostProcess/SceneRenderTargets.h
@@ -164,7 +164,6 @@ protected:
 	FSceneRenderTargets(): 
 		bScreenSpaceAOIsValid(false),
 		bCustomDepthIsValid(false),
-		bPreshadowCacheNewlyAllocated(false),
 		GBufferRefCount(0),
 		LargestDesiredSizeThisFrame( 0, 0 ),
 		LargestDesiredSizeLastFrame( 0, 0 ),
@@ -219,9 +218,8 @@ public:
 	
 	/**
 	 * Called when finished rendering to the scene color surface
-	 * @param bKeepChanges - if true then the SceneColorSurface is resolved to the SceneColorTexture
 	 */
-	void FinishRenderingSceneColor(FRHICommandListImmediate& RHICmdList, bool bKeepChanges = true, const FResolveRect& ResolveRect = FResolveRect());
+	void FinishRenderingSceneColor(FRHICommandListImmediate& RHICmdList, const FResolveRect& ResolveRect = FResolveRect());
 
 	// @return true: call FinishRenderingCustomDepth after rendering, false: don't render it, feature is disabled
 	bool BeginRenderingCustomDepth(FRHICommandListImmediate& RHICmdList, bool bPrimitives);
@@ -233,9 +231,6 @@ public:
 	 */
 	void ResolveSceneColor(FRHICommandList& RHICmdList, const FResolveRect& ResolveRect = FResolveRect());
 
-	/** Resolves the GBuffer targets so that their resolved textures can be sampled. */
-	void ResolveGBufferSurfaces(FRHICommandList& RHICmdList, const FResolveRect& ResolveRect = FResolveRect());
-
 	/** Binds the appropriate shadow depth cube map for rendering. */
 	void BeginRenderingCubeShadowDepth(FRHICommandList& RHICmdList, int32 ShadowResolution);
 
@@ -276,42 +271,20 @@ public:
 		OutScale = SeparateTranslucencyScale;
 	}
 
-	TRefCountPtr<IPooledRenderTarget>& GetSeparateTranslucency(FRHICommandList& RHICmdList, FIntPoint Size)
-	{
-		if (!SeparateTranslucencyRT || SeparateTranslucencyRT->GetDesc().Extent != Size)
-		{
-			uint32 Flags = TexCreate_RenderTargetable;
-
-			// Create the SeparateTranslucency render target (alpha is needed to lerping)
-			FPooledRenderTargetDesc Desc(FPooledRenderTargetDesc::Create2DDesc(Size, PF_FloatRGBA, FClearValueBinding::Black, TexCreate_None, Flags, false));
-			Desc.AutoWritable = false;
-			GRenderTargetPool.FindFreeElement(RHICmdList, Desc, SeparateTranslucencyRT, TEXT("SeparateTranslucency"));
-		}
-		return SeparateTranslucencyRT;
-	}
+	TRefCountPtr<IPooledRenderTarget>& GetSeparateTranslucency(FRHICommandList& RHICmdList, FIntPoint Size);
 
 	bool IsSeparateTranslucencyDepthValid()
 	{
 		return SeparateTranslucencyDepthRT != nullptr;
 	}
 
-	TRefCountPtr<IPooledRenderTarget>& GetSeparateTranslucencyDepth(FRHICommandList& RHICmdList, FIntPoint Size)
-	{
-		if (!SeparateTranslucencyDepthRT || SeparateTranslucencyDepthRT->GetDesc().Extent != Size)
-		{
-			// Create the SeparateTranslucency depth render target 
-			FPooledRenderTargetDesc Desc(FPooledRenderTargetDesc::Create2DDesc(Size, PF_DepthStencil, FClearValueBinding::None, TexCreate_None, TexCreate_DepthStencilTargetable, false));
-			GRenderTargetPool.FindFreeElement(RHICmdList, Desc, SeparateTranslucencyDepthRT, TEXT("SeparateTranslucencyDepth"));
-		}
-		return SeparateTranslucencyDepthRT;
-	}
+	TRefCountPtr<IPooledRenderTarget>& GetSeparateTranslucencyDepth(FRHICommandList& RHICmdList, FIntPoint Size);
+
 	const FTexture2DRHIRef& GetSeparateTranslucencyDepthSurface()
 	{
 		return (const FTexture2DRHIRef&)SeparateTranslucencyDepthRT->GetRenderTargetItem().TargetableTexture;
 	}
 
-
-
 	/**
 	 * Cleans up editor primitive targets that we no longer need
 	 */
@@ -347,7 +320,6 @@ public:
 	const FTexture2DRHIRef& GetSceneAlphaCopyTexture() const { return (const FTexture2DRHIRef&)SceneAlphaCopy->GetRenderTargetItem().ShaderResourceTexture; }
 	bool HasSceneAlphaCopyTexture() const { return SceneAlphaCopy.GetReference() != 0; }
 	const FTexture2DRHIRef& GetSceneDepthTexture() const { return (const FTexture2DRHIRef&)SceneDepthZ->GetRenderTargetItem().ShaderResourceTexture; }
-	const FTexture2DRHIRef& GetNoMSAASceneDepthTexture() const { return (const FTexture2DRHIRef&)NoMSAASceneDepthZ->GetRenderTargetItem().ShaderResourceTexture; }
 	const FTexture2DRHIRef& GetAuxiliarySceneDepthTexture() const
 	{ 
 		check(!GSupportsDepthFetchDuringDepthTest);
@@ -381,7 +353,6 @@ public:
 	const FTextureRHIRef& GetSceneColorSurface() const;
 	const FTexture2DRHIRef& GetSceneAlphaCopySurface() const						{ return (const FTexture2DRHIRef&)SceneAlphaCopy->GetRenderTargetItem().TargetableTexture; }
 	const FTexture2DRHIRef& GetSceneDepthSurface() const							{ return (const FTexture2DRHIRef&)SceneDepthZ->GetRenderTargetItem().TargetableTexture; }
-	const FTexture2DRHIRef& GetNoMSAASceneDepthSurface() const						{ return (const FTexture2DRHIRef&)NoMSAASceneDepthZ->GetRenderTargetItem().TargetableTexture; }
 	const FTexture2DRHIRef& GetSmallDepthSurface() const							{ return (const FTexture2DRHIRef&)SmallDepthZ->GetRenderTargetItem().TargetableTexture; }
 	const FTexture2DRHIRef& GetOptionalShadowDepthColorSurface() const 
 	{ 
@@ -483,6 +454,7 @@ public:
 
 	//
 	void PreallocGBufferTargets(bool bShouldRenderVelocities);
+	void GetGBufferADesc(FPooledRenderTargetDesc& Desc) const;
 	void AllocGBufferTargets(FRHICommandList& RHICmdList);
 
 	void AllocLightAttenuation(FRHICommandList& RHICmdList);
@@ -532,8 +504,6 @@ public:
 	TRefCountPtr<IPooledRenderTarget> SceneDepthZ;
 	TRefCountPtr<FRHIShaderResourceView> SceneStencilSRV;
 	TRefCountPtr<IPooledRenderTarget> LightingChannels;
-	// Used when MSAA is enabled but rendering to a non-MSAA render target.
-	TRefCountPtr<IPooledRenderTarget> NoMSAASceneDepthZ;
 	// Mobile without frame buffer fetch (to get depth from alpha).
 	TRefCountPtr<IPooledRenderTarget> SceneAlphaCopy;
 	// Auxiliary scene depth target. The scene depth is resolved to this surface when targeting SM4. 
@@ -566,9 +536,6 @@ public:
 	TRefCountPtr<FRHIShaderResourceView> CustomStencilSRV;
 	// optional in case this RHI requires a color render target
 	TRefCountPtr<IPooledRenderTarget> OptionalShadowDepthColor;
-	// Cache of preshadow depths
-	//@todo - this should go in FScene
-	TRefCountPtr<IPooledRenderTarget> PreShadowCacheDepthZ;
 
 	/** 2 scratch cubemaps used for filtering reflections. */
 	TRefCountPtr<IPooledRenderTarget> ReflectionColorScratchCubemap[2];
@@ -579,12 +546,6 @@ public:
 	/** Temporary storage during SH irradiance map generation. */
 	TRefCountPtr<IPooledRenderTarget> SkySHIrradianceMap;
 
-	/** Temporary storage, used during reflection capture filtering. 
-	  * 0 - R32 version for > ES2
-	  * 1 - RGBAF version for ES2
-	  */
-	TRefCountPtr<IPooledRenderTarget> ReflectionBrightness[2];
-
 	/** Volume textures used for lighting translucency. */
 	TRefCountPtr<IPooledRenderTarget> TranslucencyLightingVolumeAmbient[NumTranslucentVolumeRenderTargetSets];
 	TRefCountPtr<IPooledRenderTarget> TranslucencyLightingVolumeDirectional[NumTranslucentVolumeRenderTargetSets];
@@ -605,9 +566,6 @@ public:
 	// todo: free ScreenSpaceAO so pool can reuse
 	bool bCustomDepthIsValid;
 
-	/** Whether the preshadow cache render target has been newly allocated and cached shadows need to be re-rendered. */
-	bool bPreshadowCacheNewlyAllocated;
-
 private:
 	/** used by AdjustGBufferRefCount */
 	int32 GBufferRefCount;
@@ -698,8 +656,8 @@ private:
 	int32 CurrentTranslucencyLightingVolumeDim;
 	/** To detect a change of the CVar r.MobileHDR / r.MobileHDR32bppMode */
 	int32 CurrentMobile32bpp;
-	/** To detect a change of the CVar r.MobileMSAA */
-	int32 CurrentMobileMSAA;
+	/** To detect a change of the CVar r.MobileMSAA or r.MSAA */
+	int32 CurrentMSAACount;
 	/** To detect a change of the CVar r.Shadow.MinResolution */
 	int32 CurrentMinShadowResolution;
 	/** To detect a change of the CVar r.LightPropagationVolume */
diff --git a/Engine/Source/Runtime/Renderer/Private/PostProcess/ScreenSpaceReflections.cpp b/Engine/Source/Runtime/Renderer/Private/PostProcess/ScreenSpaceReflections.cpp
--- a/Engine/Source/Runtime/Renderer/Private/PostProcess/ScreenSpaceReflections.cpp
+++ b/Engine/Source/Runtime/Renderer/Private/PostProcess/ScreenSpaceReflections.cpp
@@ -98,7 +98,7 @@ FLinearColor ComputeSSRParams(const FRenderingCompositePassContext& Context, uin
 
 	if(Context.ViewState)
 	{
-		bool bTemporalAAIsOn = Context.View.FinalPostProcessSettings.AntiAliasingMethod == AAM_TemporalAA;
+		bool bTemporalAAIsOn = Context.View.AntiAliasingMethod == AAM_TemporalAA;
 
 		if(bTemporalAAIsOn)
 		{
@@ -493,7 +493,7 @@ void RenderScreenSpaceReflections(FRHICommandListImmediate& RHICmdList, FViewInf
 		Context.FinalOutput = FRenderingCompositeOutputRef( TracePass );
 	}
 
-	const bool bTemporalFilter = View.FinalPostProcessSettings.AntiAliasingMethod != AAM_TemporalAA || CVarSSRTemporal.GetValueOnRenderThread() != 0;
+	const bool bTemporalFilter = View.AntiAliasingMethod != AAM_TemporalAA || CVarSSRTemporal.GetValueOnRenderThread() != 0;
 
 	if( ViewState && bTemporalFilter )
 	{
diff --git a/Engine/Source/Runtime/Renderer/Private/PostProcess/VisualizeShadingModels.cpp b/Engine/Source/Runtime/Renderer/Private/PostProcess/VisualizeShadingModels.cpp
--- a/Engine/Source/Runtime/Renderer/Private/PostProcess/VisualizeShadingModels.cpp
+++ b/Engine/Source/Runtime/Renderer/Private/PostProcess/VisualizeShadingModels.cpp
@@ -133,28 +133,7 @@ void FRCPassPostProcessVisualizeShadingModels::Process(FRenderingCompositePassCo
 		*VertexShader,
 		EDRF_UseTriangleOptimization);
 
-
-	// this is a helper class for FCanvas to be able to get screen size
-	class FRenderTargetTemp : public FRenderTarget
-	{
-	public:
-		const FSceneView& View;
-		const FTexture2DRHIRef Texture;
-
-		FRenderTargetTemp(const FSceneView& InView, const FTexture2DRHIRef InTexture)
-			: View(InView), Texture(InTexture)
-		{
-		}
-		virtual FIntPoint GetSizeXY() const
-		{
-			return View.ViewRect.Size();
-		};
-		virtual const FTexture2DRHIRef& GetRenderTargetTexture() const
-		{
-			return Texture;
-		}
-	} TempRenderTarget(View, (const FTexture2DRHIRef&)DestRenderTarget.TargetableTexture);
-
+	FRenderTargetTemp TempRenderTarget(View, (const FTexture2DRHIRef&)DestRenderTarget.TargetableTexture);
 	FCanvas Canvas(&TempRenderTarget, NULL, 0, 0, 0, Context.GetFeatureLevel());
 
 	float X = 30;
diff --git a/Engine/Source/Runtime/Renderer/Private/PostProcess/VisualizeTexture.cpp b/Engine/Source/Runtime/Renderer/Private/PostProcess/VisualizeTexture.cpp
--- a/Engine/Source/Runtime/Renderer/Private/PostProcess/VisualizeTexture.cpp
+++ b/Engine/Source/Runtime/Renderer/Private/PostProcess/VisualizeTexture.cpp
@@ -553,25 +553,7 @@ void FVisualizeTexture::PresentContent(FRHICommandListImmediate& RHICmdList, con
 			EDRF_Default);
 	}
 
-	// this is a helper class for FCanvas to be able to get screen size
-	class FRenderTargetTemp : public FRenderTarget
-	{
-	public:
-		const FSceneView& View;
-
-		FRenderTargetTemp(const FSceneView& InView) : View(InView)
-		{
-		}
-		virtual FIntPoint GetSizeXY() const
-		{
-			return View.UnscaledViewRect.Size();
-		};
-		virtual const FTexture2DRHIRef& GetRenderTargetTexture() const
-		{
-			return View.Family->RenderTarget->GetRenderTargetTexture();
-		}
-	} TempRenderTarget(View);
-
+	FRenderTargetTemp TempRenderTarget(View, View.UnscaledViewRect.Size());
 	FCanvas Canvas(&TempRenderTarget, NULL, View.Family->CurrentRealTime, View.Family->CurrentWorldTime, View.Family->DeltaWorldTime, View.GetFeatureLevel());
 
 	float X = 100 + View.ViewRect.Min.X;
diff --git a/Engine/Source/Runtime/Renderer/Private/PrimitiveDistanceAccuracyRendering.h b/Engine/Source/Runtime/Renderer/Private/PrimitiveDistanceAccuracyRendering.h
--- a/Engine/Source/Runtime/Renderer/Private/PrimitiveDistanceAccuracyRendering.h
+++ b/Engine/Source/Runtime/Renderer/Private/PrimitiveDistanceAccuracyRendering.h
@@ -39,6 +39,11 @@ public:
 		return bShaderHasOutdatedParameters;
 	}
 
+	static void ModifyCompilationEnvironment(EShaderPlatform Platform, FShaderCompilerEnvironment& OutEnvironment)
+	{
+		OutEnvironment.SetDefine(TEXT("UNDEFINED_ACCURACY"), UndefinedStreamingAccuracyIntensity);
+	}
+
 	virtual void SetParameters(
 		FRHICommandList& RHICmdList, 
 		const FShader* OriginalVS, 
diff --git a/Engine/Source/Runtime/Renderer/Private/ReflectionEnvironment.cpp b/Engine/Source/Runtime/Renderer/Private/ReflectionEnvironment.cpp
--- a/Engine/Source/Runtime/Renderer/Private/ReflectionEnvironment.cpp
+++ b/Engine/Source/Runtime/Renderer/Private/ReflectionEnvironment.cpp
@@ -22,19 +22,12 @@
 
 DECLARE_FLOAT_COUNTER_STAT(TEXT("Reflection Environment"), Stat_GPU_ReflectionEnvironment, STATGROUP_GPU);
 
-/** Tile size for the reflection environment compute shader, tweaked for 680 GTX. */
-const int32 GReflectionEnvironmentTileSizeX = 16;
-const int32 GReflectionEnvironmentTileSizeY = 16;
+/** Tile size for the reflection environment compute shader, tweaked for PS4. */
+const int32 GReflectionEnvironmentTileSizeX = 8;
+const int32 GReflectionEnvironmentTileSizeY = 8;
 
 extern TAutoConsoleVariable<int32> CVarLPVMixing;
 
-static TAutoConsoleVariable<int32> CVarDiffuseFromCaptures(
-	TEXT("r.DiffuseFromCaptures"),
-	0,
-	TEXT("Apply indirect diffuse lighting from captures instead of lightmaps.\n")
-	TEXT(" 0 is off (default), 1 is on"),
-	ECVF_RenderThreadSafe);
-
 static TAutoConsoleVariable<int32> CVarReflectionEnvironment(
 	TEXT("r.ReflectionEnvironment"),
 	1,
@@ -44,6 +37,29 @@ static TAutoConsoleVariable<int32> CVarReflectionEnvironment(
 	TEXT(" 2: on and overwrite scene (only in non-shipping builds)"),
 	ECVF_RenderThreadSafe | ECVF_Scalability);
 
+int32 GReflectionEnvironmentLightmapMixing = 1;
+FAutoConsoleVariableRef CVarReflectionEnvironmentLightmapMixing(
+	TEXT("r.ReflectionEnvironmentLightmapMixing"),
+	GReflectionEnvironmentLightmapMixing,
+	TEXT("Whether to mix indirect specular from reflection captures with indirect diffuse from lightmaps for rough surfaces."),
+	ECVF_Scalability | ECVF_RenderThreadSafe
+	);
+
+float GReflectionEnvironmentBeginMixingRoughness = .1f;
+FAutoConsoleVariableRef CVarReflectionEnvironmentBeginMixingRoughness(
+	TEXT("r.ReflectionEnvironmentBeginMixingRoughness"),
+	GReflectionEnvironmentBeginMixingRoughness,
+	TEXT("Min roughness value at which to begin mixing reflection captures with lightmap indirect diffuse."),
+	ECVF_Scalability | ECVF_RenderThreadSafe
+	);
+
+float GReflectionEnvironmentEndMixingRoughness = .3f;
+FAutoConsoleVariableRef CVarReflectionEnvironmentEndMixingRoughness(
+	TEXT("r.ReflectionEnvironmentEndMixingRoughness"),
+	GReflectionEnvironmentEndMixingRoughness,
+	TEXT("Min roughness value at which to end mixing reflection captures with lightmap indirect diffuse."),
+	ECVF_Scalability | ECVF_RenderThreadSafe
+	);
 
 static TAutoConsoleVariable<int32> CVarDoTiledReflections(
 	TEXT("r.DoTiledReflections"),
@@ -75,6 +91,24 @@ static int GetReflectionEnvironmentCVar()
 	return RetVal;
 }
 
+FVector2D GetReflectionEnvironmentRoughnessMixingScaleBias()
+{
+	float RoughnessMixingRange = 1.0f / FMath::Max(GReflectionEnvironmentEndMixingRoughness - GReflectionEnvironmentBeginMixingRoughness, .001f);
+
+	if (GReflectionEnvironmentLightmapMixing == 0)
+	{
+		RoughnessMixingRange = 0.0f;
+	}
+
+	if (GReflectionEnvironmentEndMixingRoughness == 0.0f && GReflectionEnvironmentBeginMixingRoughness == 0.0f)
+	{
+		// Make sure a Roughness of 0 results in full mixing when disabling roughness-based mixing
+		return FVector2D(0, 1);
+	}
+
+	return FVector2D(RoughnessMixingRange, -GReflectionEnvironmentBeginMixingRoughness * RoughnessMixingRange);
+}
+
 bool IsReflectionEnvironmentAvailable(ERHIFeatureLevel::Type InFeatureLevel)
 {
 	return (InFeatureLevel >= ERHIFeatureLevel::SM4) && (GetReflectionEnvironmentCVar() != 0);
@@ -198,7 +232,7 @@ struct FReflectionCaptureSortData
 	FVector4 CaptureProperties;
 	FMatrix BoxTransform;
 	FVector4 BoxScales;
-	FVector4 CaptureOffset;
+	FVector4 CaptureOffsetAndAverageBrightness;
 	FTexture* SM4FullHDRCubemap;
 
 	bool operator < (const FReflectionCaptureSortData& Other) const
@@ -221,7 +255,7 @@ BEGIN_UNIFORM_BUFFER_STRUCT(FReflectionCaptureData,)
 	DECLARE_UNIFORM_BUFFER_STRUCT_MEMBER_ARRAY(FVector4,CaptureProperties,[GMaxNumReflectionCaptures])
 	// Stores the box transform for a box shape, other data is packed for other shapes
 	DECLARE_UNIFORM_BUFFER_STRUCT_MEMBER_ARRAY(FMatrix,BoxTransform,[GMaxNumReflectionCaptures])
-	DECLARE_UNIFORM_BUFFER_STRUCT_MEMBER_ARRAY(FVector4,CaptureOffset,[GMaxNumReflectionCaptures])
+	DECLARE_UNIFORM_BUFFER_STRUCT_MEMBER_ARRAY(FVector4,CaptureOffsetAndAverageBrightness,[GMaxNumReflectionCaptures])
 	DECLARE_UNIFORM_BUFFER_STRUCT_MEMBER_ARRAY(FVector4,BoxScales,[GMaxNumReflectionCaptures])
 END_UNIFORM_BUFFER_STRUCT(FReflectionCaptureData)
 
@@ -314,7 +348,7 @@ public:
 			SamplePositionsBuffer.PositionAndRadius[CaptureIndex] = SortData[CaptureIndex].PositionAndRadius;
 			SamplePositionsBuffer.CaptureProperties[CaptureIndex] = SortData[CaptureIndex].CaptureProperties;
 			SamplePositionsBuffer.BoxTransform[CaptureIndex] = SortData[CaptureIndex].BoxTransform;
-			SamplePositionsBuffer.CaptureOffset[CaptureIndex] = SortData[CaptureIndex].CaptureOffset;
+			SamplePositionsBuffer.CaptureOffsetAndAverageBrightness[CaptureIndex] = SortData[CaptureIndex].CaptureOffsetAndAverageBrightness;
 			SamplePositionsBuffer.BoxScales[CaptureIndex] = SortData[CaptureIndex].BoxScales;
 		}
 
@@ -622,7 +656,7 @@ public:
 		CaptureProperties.Bind(Initializer.ParameterMap, TEXT("CaptureProperties"));
 		CaptureBoxTransform.Bind(Initializer.ParameterMap, TEXT("CaptureBoxTransform"));
 		CaptureBoxScales.Bind(Initializer.ParameterMap, TEXT("CaptureBoxScales"));
-		CaptureOffset.Bind(Initializer.ParameterMap, TEXT("CaptureOffset"));
+		CaptureOffsetAndAverageBrightness.Bind(Initializer.ParameterMap, TEXT("CaptureOffsetAndAverageBrightness"));
 		CaptureArrayIndex.Bind(Initializer.ParameterMap, TEXT("CaptureArrayIndex"));
 		ReflectionEnvironmentColorTexture.Bind(Initializer.ParameterMap, TEXT("ReflectionEnvironmentColorTexture"));
 		ReflectionEnvironmentColorTextureArray.Bind(Initializer.ParameterMap, TEXT("ReflectionEnvironmentColorTextureArray"));
@@ -657,7 +691,7 @@ public:
 		SetShaderValue(RHICmdList, ShaderRHI, CaptureProperties, SortData.CaptureProperties);
 		SetShaderValue(RHICmdList, ShaderRHI, CaptureBoxTransform, SortData.BoxTransform);
 		SetShaderValue(RHICmdList, ShaderRHI, CaptureBoxScales, SortData.BoxScales);
-		SetShaderValue(RHICmdList, ShaderRHI, CaptureOffset, FVector(SortData.CaptureOffset));
+		SetShaderValue(RHICmdList, ShaderRHI, CaptureOffsetAndAverageBrightness, SortData.CaptureOffsetAndAverageBrightness);
 	}
 
 	// FShader interface.
@@ -668,7 +702,7 @@ public:
 		Ar << CaptureProperties;
 		Ar << CaptureBoxTransform;
 		Ar << CaptureBoxScales;
-		Ar << CaptureOffset;
+		Ar << CaptureOffsetAndAverageBrightness;
 		Ar << CaptureArrayIndex;
 		Ar << ReflectionEnvironmentColorTexture;
 		Ar << ReflectionEnvironmentColorTextureArray;
@@ -683,7 +717,7 @@ private:
 	FShaderParameter CaptureProperties;
 	FShaderParameter CaptureBoxTransform;
 	FShaderParameter CaptureBoxScales;
-	FShaderParameter CaptureOffset;
+	FShaderParameter CaptureOffsetAndAverageBrightness;
 	FShaderParameter CaptureArrayIndex;
 	FShaderResourceParameter ReflectionEnvironmentColorTexture;
 	FShaderResourceParameter ReflectionEnvironmentColorTextureArray;
@@ -772,7 +806,7 @@ void GatherAndSortReflectionCaptures(const FScene* Scene, TArray<FReflectionCapt
 			NewSortEntry.PositionAndRadius = FVector4(CurrentCapture->Position, CurrentCapture->InfluenceRadius);
 			float ShapeTypeValue = (float)CurrentCapture->Shape;
 			NewSortEntry.CaptureProperties = FVector4(CurrentCapture->Brightness, CubemapIndex, ShapeTypeValue, 0);
-			NewSortEntry.CaptureOffset = FVector4(CurrentCapture->CaptureOffset, 0);
+			NewSortEntry.CaptureOffsetAndAverageBrightness = FVector4(CurrentCapture->CaptureOffset, CurrentCapture->AverageBrightness);
 
 			if (CurrentCapture->Shape == EReflectionCaptureShape::Plane)
 			{
@@ -910,7 +944,7 @@ void FDeferredShadingSceneRenderer::RenderTiledDeferredImageBasedReflections(FRH
 {
 	FSceneRenderTargets& SceneContext = FSceneRenderTargets::Get(RHICmdList);
 	static const auto AllowStaticLightingVar = IConsoleManager::Get().FindTConsoleVariableDataInt(TEXT("r.AllowStaticLighting"));
-	const bool bUseLightmaps = (AllowStaticLightingVar->GetValueOnRenderThread() == 1) && (CVarDiffuseFromCaptures.GetValueOnRenderThread() == 0);
+	const bool bUseLightmaps = (AllowStaticLightingVar->GetValueOnRenderThread() == 1);
 
 	TRefCountPtr<IPooledRenderTarget> NewSceneColor;
 	{
@@ -1045,7 +1079,7 @@ void FDeferredShadingSceneRenderer::RenderStandardDeferredImageBasedReflections(
 			NewSortEntry.PositionAndRadius = FVector4(CurrentCapture->Position, CurrentCapture->InfluenceRadius);
 			float ShapeTypeValue = (float)CurrentCapture->Shape;
 			NewSortEntry.CaptureProperties = FVector4(CurrentCapture->Brightness, 0, ShapeTypeValue, 0);
-			NewSortEntry.CaptureOffset = FVector4(CurrentCapture->CaptureOffset);
+			NewSortEntry.CaptureOffsetAndAverageBrightness = FVector4(CurrentCapture->CaptureOffset, CurrentCapture->AverageBrightness);
 
 			if (CurrentCapture->Shape == EReflectionCaptureShape::Plane)
 			{
diff --git a/Engine/Source/Runtime/Renderer/Private/ReflectionEnvironmentCapture.cpp b/Engine/Source/Runtime/Renderer/Private/ReflectionEnvironmentCapture.cpp
--- a/Engine/Source/Runtime/Renderer/Private/ReflectionEnvironmentCapture.cpp
+++ b/Engine/Source/Runtime/Renderer/Private/ReflectionEnvironmentCapture.cpp
@@ -14,6 +14,7 @@
 #include "ReflectionEnvironment.h"
 #include "ReflectionEnvironmentCapture.h"
 #include "SceneUtils.h"
+#include "OneColorShader.h"
 
 /** Near plane to use when capturing the scene. */
 float GReflectionCaptureNearPlane = 5;
@@ -157,8 +158,6 @@ public:
 	FCubeFilterPS(const ShaderMetaType::CompiledShaderInitializerType& Initializer):
 		FDownsamplePS(Initializer)
 	{
-		AverageBrightnessTexture.Bind(Initializer.ParameterMap,TEXT("AverageBrightnessTexture"));
-		AverageBrightnessSampler.Bind(Initializer.ParameterMap,TEXT("AverageBrightnessSampler"));
 		CubemapMaxMipParameter.Bind(Initializer.ParameterMap, TEXT("CubemapMaxMip"));
 	}
 	FCubeFilterPS() {}
@@ -167,29 +166,17 @@ public:
 	{
 		FDownsamplePS::SetParameters(RHICmdList, CubeFaceValue, SourceMipIndexValue, SourceTextureValue);
 
-		SetTextureParameter(
-			RHICmdList,
-			GetPixelShader(), 
-			AverageBrightnessTexture, 
-			AverageBrightnessSampler, 
-			TStaticSamplerState<SF_Bilinear, AM_Clamp, AM_Clamp, AM_Clamp>::GetRHI(), 
-			FSceneRenderTargets::Get(RHICmdList).GetReflectionBrightnessTarget()->GetRenderTargetItem().ShaderResourceTexture);
-
 		SetShaderValue(RHICmdList, GetPixelShader(), CubemapMaxMipParameter, NumMips - 1.0f);
 	}
 
 	virtual bool Serialize(FArchive& Ar) override
 	{
 		bool bShaderHasOutdatedParameters = FDownsamplePS::Serialize(Ar);
-		Ar << AverageBrightnessTexture;
-		Ar << AverageBrightnessSampler;
 		Ar << CubemapMaxMipParameter;
 		return bShaderHasOutdatedParameters;
 	}
 
 private:
-	FShaderResourceParameter AverageBrightnessTexture;
-	FShaderResourceParameter AverageBrightnessSampler;
 	FShaderParameter CubemapMaxMipParameter;
 };
 
@@ -232,7 +219,6 @@ public:
 	{
 		FGlobalShader::ModifyCompilationEnvironment(Platform, OutEnvironment);
 		OutEnvironment.SetDefine(TEXT("COMPUTEBRIGHTNESS_PIXELSHADER"), 1);
-		OutEnvironment.SetRenderTargetOutputFormat(0, PF_R32_FLOAT);
 	}
 
 	FComputeBrightnessPS(const ShaderMetaType::CompiledShaderInitializerType& Initializer)
@@ -287,9 +273,13 @@ private:
 IMPLEMENT_SHADER_TYPE(,FComputeBrightnessPS,TEXT("ReflectionEnvironmentShaders"),TEXT("ComputeBrightnessMain"),SF_Pixel);
 
 /** Computes the average brightness of the given reflection capture and stores it in the scene. */
-void ComputeAverageBrightness(FRHICommandList& RHICmdList, ERHIFeatureLevel::Type FeatureLevel, int32 TargetSize)
+float ComputeSingleAverageBrightnessFromCubemap(FRHICommandListImmediate& RHICmdList, ERHIFeatureLevel::Type FeatureLevel, int32 TargetSize)
 {
-	FTextureRHIRef& BrightnessTarget = FSceneRenderTargets::Get(RHICmdList).GetReflectionBrightnessTarget()->GetRenderTargetItem().TargetableTexture;
+	TRefCountPtr<IPooledRenderTarget> ReflectionBrightnessTarget;
+	FPooledRenderTargetDesc Desc(FPooledRenderTargetDesc::Create2DDesc(FIntPoint(1, 1), PF_FloatRGBA, FClearValueBinding::None, TexCreate_None, TexCreate_RenderTargetable, false));
+	GRenderTargetPool.FindFreeElement(RHICmdList, Desc, ReflectionBrightnessTarget, TEXT("ReflectionBrightness"));
+
+	FTextureRHIRef& BrightnessTarget = ReflectionBrightnessTarget->GetRenderTargetItem().TargetableTexture;
 	SetRenderTarget(RHICmdList, BrightnessTarget, NULL, true);
 	RHICmdList.SetRasterizerState(TStaticRasterizerState<FM_Solid, CM_None>::GetRHI());
 	RHICmdList.SetDepthStencilState(TStaticDepthStencilState<false, CF_Always>::GetRHI());
@@ -316,20 +306,24 @@ void ComputeAverageBrightness(FRHICommandList& RHICmdList, ERHIFeatureLevel::Typ
 		*VertexShader);
 
 	RHICmdList.CopyToResolveTarget(BrightnessTarget, BrightnessTarget, true, FResolveParams());
+
+	FSceneRenderTargetItem& EffectiveRT = ReflectionBrightnessTarget->GetRenderTargetItem();
+	check(EffectiveRT.ShaderResourceTexture->GetFormat() == PF_FloatRGBA);
+
+	TArray<FFloat16Color> SurfaceData;
+	RHICmdList.ReadSurfaceFloatData(EffectiveRT.ShaderResourceTexture, FIntRect(0, 0, 1, 1), SurfaceData, CubeFace_PosX, 0, 0);
+
+	float AverageBrightness = SurfaceData[0].R.GetFloat();
+	return AverageBrightness;
 }
 
-/** Generates mips for glossiness and filters the cubemap for a given reflection. */
-void FilterReflectionEnvironment(FRHICommandListImmediate& RHICmdList, ERHIFeatureLevel::Type FeatureLevel, int32 CubmapSize, FSHVectorRGB3* OutIrradianceEnvironmentMap, bool bNormalize)
+void ComputeAverageBrightness(FRHICommandListImmediate& RHICmdList, ERHIFeatureLevel::Type FeatureLevel, int32 CubmapSize, float& OutAverageBrightness)
 {
+	SCOPED_DRAW_EVENT(RHICmdList, ComputeAverageBrightness);
+
 	const int32 EffectiveTopMipSize = CubmapSize;
 	const int32 NumMips = FMath::CeilLogTwo(EffectiveTopMipSize) + 1;
 
-	int32 DiffuseConvolutionSourceMip = INDEX_NONE;
-	FSceneRenderTargetItem* DiffuseConvolutionSource = NULL;
-
-	static const auto CVar = IConsoleManager::Get().FindConsoleVariable(TEXT("r.DiffuseFromCaptures"));
-	bNormalize = bNormalize && CVar->GetInt() == 0;
-
 	// necessary to resolve the clears which touched all the mips.  scene rendering only resolves mip 0.
 	FullyResolveReflectionScratchCubes(RHICmdList);	
 
@@ -337,6 +331,101 @@ void FilterReflectionEnvironment(FRHICommandListImmediate& RHICmdList, ERHIFeatu
 
 	FSceneRenderTargets& SceneContext = FSceneRenderTargets::Get(RHICmdList);
 
+	{	
+		SCOPED_DRAW_EVENT(RHICmdList, DownsampleCubeMips);
+
+		// Downsample all the mips, each one reads from the mip above it
+		for (int32 MipIndex = 1; MipIndex < NumMips; MipIndex++)
+		{
+			const int32 SourceMipIndex = FMath::Max(MipIndex - 1, 0);
+			const int32 MipSize = 1 << (NumMips - MipIndex - 1);
+
+			FSceneRenderTargetItem& EffectiveRT = GetEffectiveRenderTarget(SceneContext, true, MipIndex);
+			FSceneRenderTargetItem& EffectiveSource = GetEffectiveSourceTexture(SceneContext, true, MipIndex);
+			check(EffectiveRT.TargetableTexture != EffectiveSource.ShaderResourceTexture);
+			
+			for (int32 CubeFace = 0; CubeFace < CubeFace_MAX; CubeFace++)
+			{
+				SetRenderTarget(RHICmdList, EffectiveRT.TargetableTexture, MipIndex, CubeFace, NULL, true);
+
+				const FIntRect ViewRect(0, 0, MipSize, MipSize);
+				RHICmdList.SetViewport(0, 0, 0.0f, MipSize, MipSize, 1.0f);
+				RHICmdList.SetRasterizerState(TStaticRasterizerState<FM_Solid, CM_None>::GetRHI());
+				RHICmdList.SetDepthStencilState(TStaticDepthStencilState<false, CF_Always>::GetRHI());
+				RHICmdList.SetBlendState(TStaticBlendState<>::GetRHI());
+
+				TShaderMapRef<FScreenVS> VertexShader(GetGlobalShaderMap(FeatureLevel));
+				TShaderMapRef<FDownsamplePS> PixelShader(GetGlobalShaderMap(FeatureLevel));
+
+				SetGlobalBoundShaderState(RHICmdList, FeatureLevel, DownsampleBoundShaderState, GFilterVertexDeclaration.VertexDeclarationRHI, *VertexShader, *PixelShader);
+
+				PixelShader->SetParameters(RHICmdList, CubeFace, SourceMipIndex, EffectiveSource);
+
+				DrawRectangle( 
+					RHICmdList,
+					ViewRect.Min.X, ViewRect.Min.Y, 
+					ViewRect.Width(), ViewRect.Height(),
+					ViewRect.Min.X, ViewRect.Min.Y, 
+					ViewRect.Width(), ViewRect.Height(),
+					FIntPoint(ViewRect.Width(), ViewRect.Height()),
+					FIntPoint(MipSize, MipSize),
+					*VertexShader);
+
+				RHICmdList.CopyToResolveTarget(EffectiveRT.TargetableTexture, EffectiveRT.ShaderResourceTexture, true, FResolveParams(FResolveRect(), (ECubeFace)CubeFace, MipIndex));
+			}
+		}
+	}
+
+	OutAverageBrightness = ComputeSingleAverageBrightnessFromCubemap(RHICmdList, FeatureLevel, CubmapSize);
+}
+
+/** Generates mips for glossiness and filters the cubemap for a given reflection. */
+void FilterReflectionEnvironment(FRHICommandListImmediate& RHICmdList, ERHIFeatureLevel::Type FeatureLevel, int32 CubmapSize, FSHVectorRGB3* OutIrradianceEnvironmentMap)
+{
+	const int32 EffectiveTopMipSize = CubmapSize;
+	const int32 NumMips = FMath::CeilLogTwo(EffectiveTopMipSize) + 1;
+
+	FSceneRenderTargetItem& EffectiveColorRT = FSceneRenderTargets::Get(RHICmdList).ReflectionColorScratchCubemap[0]->GetRenderTargetItem();
+
+	// Premultiply alpha in-place using alpha blending
+	for (uint32 CubeFace = 0; CubeFace < CubeFace_MAX; CubeFace++)
+	{
+		SetRenderTarget(RHICmdList, EffectiveColorRT.TargetableTexture, 0, CubeFace, NULL, true);
+
+		const FIntPoint SourceDimensions(CubmapSize, CubmapSize);
+		const FIntRect ViewRect(0, 0, EffectiveTopMipSize, EffectiveTopMipSize);
+		RHICmdList.SetViewport(0, 0, 0.0f, EffectiveTopMipSize, EffectiveTopMipSize, 1.0f);
+		RHICmdList.SetRasterizerState(TStaticRasterizerState<FM_Solid, CM_None>::GetRHI());
+		RHICmdList.SetDepthStencilState(TStaticDepthStencilState<false, CF_Always>::GetRHI());
+		RHICmdList.SetBlendState(TStaticBlendState<CW_RGBA, BO_Add, BF_Zero, BF_DestAlpha, BO_Add, BF_Zero, BF_One>::GetRHI());
+
+		TShaderMapRef<FScreenVS> VertexShader(GetGlobalShaderMap(FeatureLevel));
+		TShaderMapRef<FOneColorPS> PixelShader(GetGlobalShaderMap(FeatureLevel));
+
+		static FGlobalBoundShaderState BoundShaderState;
+		SetGlobalBoundShaderState(RHICmdList, FeatureLevel, BoundShaderState, GFilterVertexDeclaration.VertexDeclarationRHI, *VertexShader, *PixelShader);
+		FLinearColor UnusedColors[1] = { FLinearColor::Black };
+		PixelShader->SetColors(RHICmdList, UnusedColors, ARRAY_COUNT(UnusedColors));
+
+		DrawRectangle(
+			RHICmdList,
+			ViewRect.Min.X, ViewRect.Min.Y, 
+			ViewRect.Width(), ViewRect.Height(),
+			0, 0, 
+			SourceDimensions.X, SourceDimensions.Y,
+			FIntPoint(ViewRect.Width(), ViewRect.Height()),
+			SourceDimensions,
+			*VertexShader);
+
+		RHICmdList.CopyToResolveTarget(EffectiveColorRT.TargetableTexture, EffectiveColorRT.ShaderResourceTexture, true, FResolveParams(FResolveRect(), (ECubeFace)CubeFace));
+	}
+
+	int32 DiffuseConvolutionSourceMip = INDEX_NONE;
+	FSceneRenderTargetItem* DiffuseConvolutionSource = NULL;
+
+	auto ShaderMap = GetGlobalShaderMap(FeatureLevel);
+	FSceneRenderTargets& SceneContext = FSceneRenderTargets::Get(RHICmdList);
+
 	{	
 		SCOPED_DRAW_EVENT(RHICmdList, DownsampleCubeMips);
 		// Downsample all the mips, each one reads from the mip above it
@@ -394,8 +483,6 @@ void FilterReflectionEnvironment(FRHICommandListImmediate& RHICmdList, ERHIFeatu
 		check(DiffuseConvolutionSource != NULL);
 		ComputeDiffuseIrradiance(RHICmdList, FeatureLevel, DiffuseConvolutionSource->ShaderResourceTexture, DiffuseConvolutionSourceMip, OutIrradianceEnvironmentMap);
 	}
-	
-	ComputeAverageBrightness(RHICmdList, FeatureLevel, CubmapSize);
 
 	{	
 		SCOPED_DRAW_EVENT(RHICmdList, FilterCubeMap);
@@ -422,21 +509,10 @@ void FilterReflectionEnvironment(FRHICommandListImmediate& RHICmdList, ERHIFeatu
 				TShaderMapRef< TCubeFilterPS<1> > CaptureCubemapArrayPixelShader(GetGlobalShaderMap(FeatureLevel));
 
 				FCubeFilterPS* PixelShader;
-				if( bNormalize )
-				{
-					PixelShader = *TShaderMapRef< TCubeFilterPS<1> >(ShaderMap);
-					static FGlobalBoundShaderState BoundShaderState;
-					
-					SetGlobalBoundShaderState(RHICmdList, FeatureLevel, BoundShaderState, GFilterVertexDeclaration.VertexDeclarationRHI, *VertexShader, PixelShader);
-				}
-				else
-				{
-					PixelShader = *TShaderMapRef< TCubeFilterPS<0> >(ShaderMap);
-					static FGlobalBoundShaderState BoundShaderState;
-					SetGlobalBoundShaderState(RHICmdList, FeatureLevel, BoundShaderState, GFilterVertexDeclaration.VertexDeclarationRHI, *VertexShader, PixelShader);
-						
-					SetGlobalBoundShaderState(RHICmdList, FeatureLevel, BoundShaderState, GFilterVertexDeclaration.VertexDeclarationRHI, *VertexShader, PixelShader);
-				}
+
+				PixelShader = *TShaderMapRef< TCubeFilterPS<0> >(ShaderMap);
+				static FGlobalBoundShaderState BoundShaderState;
+				SetGlobalBoundShaderState(RHICmdList, FeatureLevel, BoundShaderState, GFilterVertexDeclaration.VertexDeclarationRHI, *VertexShader, PixelShader);
 
 				PixelShader->SetParameters(RHICmdList, NumMips, CubeFace, MipIndex, EffectiveSource);
 
@@ -1384,14 +1460,16 @@ void FScene::UpdateReflectionCaptureContents(UReflectionCaptureComponent* Captur
 				check(!TEXT("Unknown reflection source type"));
 			}
 
-			ENQUEUE_UNIQUE_RENDER_COMMAND_TWOPARAMETER( 
+			ENQUEUE_UNIQUE_RENDER_COMMAND_THREEPARAMETER( 
 				FilterCommand,
 				ERHIFeatureLevel::Type, FeatureLevel, GetFeatureLevel(),
 				int32, ReflectionCaptureSize, ReflectionCaptureSize,
+				float&, AverageBrightness, *CaptureComponent->GetAverageBrightnessPtr(),
 				{
-				bool bNormalize = true;
-				FilterReflectionEnvironment(RHICmdList, FeatureLevel, ReflectionCaptureSize, NULL, bNormalize);
-			});
+					ComputeAverageBrightness(RHICmdList, FeatureLevel, ReflectionCaptureSize, AverageBrightness);
+					FilterReflectionEnvironment(RHICmdList, FeatureLevel, ReflectionCaptureSize, NULL);
+				}
+			);
 
 			// Create a proxy to represent the reflection capture to the rendering thread
 			// The rendering thread will be responsible for deleting this when done with the filtering operation
@@ -1443,7 +1521,7 @@ void CopyToSkyTexture(FRHICommandList& RHICmdList, FScene* Scene, FTexture* Proc
 }
 
 // Warning: returns before writes to OutIrradianceEnvironmentMap have completed, as they are queued on the rendering thread
-void FScene::UpdateSkyCaptureContents(const USkyLightComponent* CaptureComponent, bool bCaptureEmissiveOnly, UTextureCube* SourceCubemap, FTexture* OutProcessedTexture, FSHVectorRGB3& OutIrradianceEnvironmentMap)
+void FScene::UpdateSkyCaptureContents(const USkyLightComponent* CaptureComponent, bool bCaptureEmissiveOnly, UTextureCube* SourceCubemap, FTexture* OutProcessedTexture, float& OutAverageBrightness, FSHVectorRGB3& OutIrradianceEnvironmentMap)
 {	
 	if (GSupportsRenderTargetFormat_PF_FloatRGBA || GetFeatureLevel() >= ERHIFeatureLevel::SM4)
 	{
@@ -1487,15 +1565,15 @@ void FScene::UpdateSkyCaptureContents(const USkyLightComponent* CaptureComponent
 			check(0);
 		}
 
-		ENQUEUE_UNIQUE_RENDER_COMMAND_THREEPARAMETER( 
+		ENQUEUE_UNIQUE_RENDER_COMMAND_FOURPARAMETER( 
 			FilterCommand,
 			int32, CubemapSize, CaptureComponent->CubemapResolution,
+			float&, AverageBrightness, OutAverageBrightness,
 			FSHVectorRGB3*, IrradianceEnvironmentMap, &OutIrradianceEnvironmentMap,
 			ERHIFeatureLevel::Type, FeatureLevel, GetFeatureLevel(),
 		{
-			// Skylight is normalized manually in the shader
-			bool bNormalize = false;
-			FilterReflectionEnvironment(RHICmdList, FeatureLevel, CubemapSize, IrradianceEnvironmentMap, bNormalize);
+			ComputeAverageBrightness(RHICmdList, FeatureLevel, CubemapSize, AverageBrightness);
+			FilterReflectionEnvironment(RHICmdList, FeatureLevel, CubemapSize, IrradianceEnvironmentMap);
 		});
 
 		// Optionally copy the filtered mip chain to the output texture
diff --git a/Engine/Source/Runtime/Renderer/Private/RenderTargetTemp.h b/Engine/Source/Runtime/Renderer/Private/RenderTargetTemp.h
new file mode 100644
--- /dev/null
+++ b/Engine/Source/Runtime/Renderer/Private/RenderTargetTemp.h
@@ -0,0 +1,44 @@
+// Copyright 1998 - 2016 Epic Games, Inc.All Rights Reserved.
+
+#pragma once
+
+/*=============================================================================
+RenderTargetTemp.h : Helper for canvas rendering
+=============================================================================*/
+
+
+// this is a helper class for FCanvas to be able to get screen size
+class FRenderTargetTemp : public FRenderTarget
+{
+	const FTexture2DRHIRef Texture;
+	const FIntPoint SizeXY;
+
+public:
+
+	FRenderTargetTemp(const FSceneView& InView, const FIntPoint& InSizeXY)
+	: Texture(InView.Family->RenderTarget->GetRenderTargetTexture())
+	, SizeXY(InSizeXY)
+	{}
+
+	FRenderTargetTemp(const FTexture2DRHIRef InTexture, const FIntPoint& InSizeXY)
+	: Texture(InTexture)
+	, SizeXY(InSizeXY)
+	{}
+
+	FRenderTargetTemp(const FSceneView& InView, const FTexture2DRHIRef InTexture) 
+	: Texture(InTexture)
+	, SizeXY(InView.ViewRect.Size()) 
+	{}
+
+	FRenderTargetTemp(const FSceneView& InView)
+	: Texture(InView.Family->RenderTarget->GetRenderTargetTexture())
+	, SizeXY(InView.ViewRect.Size())
+	{}
+
+	virtual FIntPoint GetSizeXY() const { return SizeXY; }
+
+	virtual const FTexture2DRHIRef& GetRenderTargetTexture() const
+	{
+		return Texture;
+	}
+};
\ No newline at end of file
diff --git a/Engine/Source/Runtime/Renderer/Private/RendererPrivate.h b/Engine/Source/Runtime/Renderer/Private/RendererPrivate.h
--- a/Engine/Source/Runtime/Renderer/Private/RendererPrivate.h
+++ b/Engine/Source/Runtime/Renderer/Private/RendererPrivate.h
@@ -13,6 +13,7 @@
 #include "RHIStaticStates.h"
 #include "ScenePrivate.h"
 #include "RendererInterface.h"
+#include "RenderTargetTemp.h"
 
 DECLARE_LOG_CATEGORY_EXTERN(LogRenderer, Log, All);
 
diff --git a/Engine/Source/Runtime/Renderer/Private/RendererScene.cpp b/Engine/Source/Runtime/Renderer/Private/RendererScene.cpp
--- a/Engine/Source/Runtime/Renderer/Private/RendererScene.cpp
+++ b/Engine/Source/Runtime/Renderer/Private/RendererScene.cpp
@@ -498,13 +498,14 @@ FScene::FReadOnlyCVARCache::FReadOnlyCVARCache()
 	static const auto CVarSupportLowQualityLightmaps = IConsoleManager::Get().FindTConsoleVariableDataInt(TEXT("r.SupportLowQualityLightmaps"));
 	static const auto CVarSupportPointLightWholeSceneShadows = IConsoleManager::Get().FindTConsoleVariableDataInt(TEXT("r.SupportPointLightWholeSceneShadows"));
 	static const auto CVarSupportAllShaderPermutations = IConsoleManager::Get().FindTConsoleVariableDataInt(TEXT("r.SupportAllShaderPermutations"));	
+	static const auto CVarVertexFoggingForOpaque = IConsoleManager::Get().FindTConsoleVariableDataInt(TEXT("r.VertexFoggingForOpaque"));	
 	const bool bForceAllPermutations = CVarSupportAllShaderPermutations && CVarSupportAllShaderPermutations->GetValueOnAnyThread() != 0;
 
 	bEnableAtmosphericFog = !CVarSupportAtmosphericFog || CVarSupportAtmosphericFog->GetValueOnAnyThread() != 0 || bForceAllPermutations;
 	bEnableStationarySkylight = !CVarSupportStationarySkylight || CVarSupportStationarySkylight->GetValueOnAnyThread() != 0 || bForceAllPermutations;
 	bEnablePointLightShadows = !CVarSupportPointLightWholeSceneShadows || CVarSupportPointLightWholeSceneShadows->GetValueOnAnyThread() != 0 || bForceAllPermutations;
 	bEnableLowQualityLightmaps = !CVarSupportLowQualityLightmaps || CVarSupportLowQualityLightmaps->GetValueOnAnyThread() != 0 || bForceAllPermutations;
-
+	bEnableVertexFoggingForOpaque = !CVarVertexFoggingForOpaque || CVarVertexFoggingForOpaque->GetValueOnAnyThread() != 0;
 
 	const bool bShowMissmatchedLowQualityLightmapsWarning = (bEnableLowQualityLightmaps) != (GEngine->bShouldGenerateLowQualityLightmaps_DEPRECATED);
 	if ( bShowMissmatchedLowQualityLightmapsWarning )
@@ -1286,14 +1287,16 @@ void FScene::UpdateReflectionCaptureTransform(UReflectionCaptureComponent* Compo
 {
 	if (Component->SceneProxy)
 	{
-		ENQUEUE_UNIQUE_RENDER_COMMAND_THREEPARAMETER(
+		ENQUEUE_UNIQUE_RENDER_COMMAND_FOURPARAMETER(
 			UpdateTransformCommand,
 			FReflectionCaptureProxy*,Proxy,Component->SceneProxy,
 			FMatrix,Transform,Component->ComponentToWorld.ToMatrixWithScale(),
+			const float*,AverageBrightness,Component->GetAverageBrightnessPtr(),
 			FScene*,Scene,this,
 		{
 			Scene->ReflectionSceneData.bRegisteredReflectionCapturesHasChanged = true;
 			Proxy->SetTransform(Transform);
+			Proxy->InitializeAverageBrightness(*AverageBrightness);
 		});
 	}
 }
diff --git a/Engine/Source/Runtime/Renderer/Private/SceneCaptureRendering.cpp b/Engine/Source/Runtime/Renderer/Private/SceneCaptureRendering.cpp
--- a/Engine/Source/Runtime/Renderer/Private/SceneCaptureRendering.cpp
+++ b/Engine/Source/Runtime/Renderer/Private/SceneCaptureRendering.cpp
@@ -599,8 +599,9 @@ void FScene::UpdateSceneCaptureContents(USceneCaptureComponentCube* CaptureCompo
 			FIntPoint CaptureSize(CaptureComponent->TextureTarget->GetSurfaceWidth(), CaptureComponent->TextureTarget->GetSurfaceHeight());
 			FMatrix ProjectionMatrix;
 			BuildProjectionMatrix(CaptureSize, ECameraProjectionMode::Perspective, FOV, 1.0f, ProjectionMatrix);
+			FPostProcessSettings PostProcessSettings;
 
-			FSceneRenderer* SceneRenderer = CreateSceneRendererForSceneCapture(this, CaptureComponent, CaptureComponent->TextureTarget->GameThread_GetRenderTargetResource(), CaptureSize, ViewRotationMatrix, Location, ProjectionMatrix, CaptureComponent->MaxViewDistanceOverride, true, false, NULL, 0);
+			FSceneRenderer* SceneRenderer = CreateSceneRendererForSceneCapture(this, CaptureComponent, CaptureComponent->TextureTarget->GameThread_GetRenderTargetResource(), CaptureSize, ViewRotationMatrix, Location, ProjectionMatrix, CaptureComponent->MaxViewDistanceOverride, true, false, &PostProcessSettings, 0);
 			SceneRenderer->ViewFamily.SceneCaptureSource = SCS_SceneColorHDR;
 
 			FTextureRenderTargetCubeResource* TextureRenderTarget = static_cast<FTextureRenderTargetCubeResource*>(CaptureComponent->TextureTarget->GameThread_GetRenderTargetResource());
diff --git a/Engine/Source/Runtime/Renderer/Private/SceneHitProxyRendering.cpp b/Engine/Source/Runtime/Renderer/Private/SceneHitProxyRendering.cpp
--- a/Engine/Source/Runtime/Renderer/Private/SceneHitProxyRendering.cpp
+++ b/Engine/Source/Runtime/Renderer/Private/SceneHitProxyRendering.cpp
@@ -339,12 +339,27 @@ TRefCountPtr<IPooledRenderTarget> InitHitProxyRender(FRHICommandListImmediate& R
 	SceneContext.Allocate(RHICmdList, ViewFamily);
 
 	TRefCountPtr<IPooledRenderTarget> HitProxyRT;
+	TRefCountPtr<IPooledRenderTarget> HitProxyDepthRT;
 
 	// Create a texture to store the resolved light attenuation values, and a render-targetable surface to hold the unresolved light attenuation values.
 	{
 		FPooledRenderTargetDesc Desc(FPooledRenderTargetDesc::Create2DDesc(SceneContext.GetBufferSizeXY(), PF_B8G8R8A8, FClearValueBinding::None, TexCreate_None, TexCreate_RenderTargetable, false));
 		Desc.Flags |= TexCreate_FastVRAM;
 		GRenderTargetPool.FindFreeElement(RHICmdList, Desc, HitProxyRT, TEXT("HitProxy"));
+
+		// create non-MSAA version for hit proxies on PC if needed
+		const EShaderPlatform CurrentShaderPlatform = GShaderPlatformForFeatureLevel[FeatureLevel];
+		FPooledRenderTargetDesc DepthDesc = SceneContext.SceneDepthZ->GetDesc();
+
+		if (DepthDesc.NumSamples > 1 && RHISupportsSeparateMSAAAndResolveTextures(CurrentShaderPlatform))
+		{
+			DepthDesc.NumSamples = 1;
+			GRenderTargetPool.FindFreeElement(RHICmdList, DepthDesc, HitProxyDepthRT, TEXT("NoMSAASceneDepthZ"));
+		}
+		else
+		{
+			HitProxyDepthRT = SceneContext.SceneDepthZ;
+		}
 	}
 
 	if (!HitProxyRT)
@@ -353,7 +368,7 @@ TRefCountPtr<IPooledRenderTarget> InitHitProxyRender(FRHICommandListImmediate& R
 		return HitProxyRT;
 	}
 
-	SetRenderTarget(RHICmdList, HitProxyRT->GetRenderTargetItem().TargetableTexture, SceneContext.GetSceneDepthSurface(), ESimpleRenderTargetMode::EExistingColorAndDepth, FExclusiveDepthStencil::DepthWrite_StencilWrite, true);
+	SetRenderTarget(RHICmdList, HitProxyRT->GetRenderTargetItem().TargetableTexture, HitProxyDepthRT->GetRenderTargetItem().TargetableTexture, ESimpleRenderTargetMode::EExistingColorAndDepth, FExclusiveDepthStencil::DepthWrite_StencilWrite, true);
 
 	// Clear color for each view.
 	auto& Views = SceneRenderer->Views;
diff --git a/Engine/Source/Runtime/Renderer/Private/ScenePrivate.h b/Engine/Source/Runtime/Renderer/Private/ScenePrivate.h
--- a/Engine/Source/Runtime/Renderer/Private/ScenePrivate.h
+++ b/Engine/Source/Runtime/Renderer/Private/ScenePrivate.h
@@ -1756,6 +1756,7 @@ public:
 		bool bEnableStationarySkylight;
 		bool bEnableAtmosphericFog;
 		bool bEnableLowQualityLightmaps;
+		bool bEnableVertexFoggingForOpaque;
 	};
 
 	/** An optional world associated with the scene. */
@@ -1837,9 +1838,6 @@ public:
 	 */
 	TSparseArray<FLightSceneInfoCompact> InvisibleLights;
 
-	/** Map from light id to the cached shadowmap data for that light. */
-	TMap<int32, FCachedShadowMapData> CachedShadowMaps;
-
 	/** The mobile quality level for which static draw lists have been built. */
 	bool bStaticDrawListsMobileHDR;
 	bool bStaticDrawListsMobileHDR32bpp;
@@ -1894,6 +1892,11 @@ public:
 
 	/** Distance field object scene data. */
 	FDistanceFieldSceneData DistanceFieldSceneData;
+	
+	/** Map from light id to the cached shadowmap data for that light. */
+	TMap<int32, FCachedShadowMapData> CachedShadowMaps;
+
+	TRefCountPtr<IPooledRenderTarget> PreShadowCacheDepthZ;
 
 	/** Preshadows that are currently cached in the PreshadowCache render target. */
 	TArray<TRefCountPtr<FProjectedShadowInfo> > CachedPreshadows;
@@ -1989,7 +1992,7 @@ public:
 	virtual void UpdateSceneCaptureContents(class USceneCaptureComponentCube* CaptureComponent) override;
 	virtual void UpdatePlanarReflectionContents(UPlanarReflectionComponent* CaptureComponent, FSceneRenderer& MainSceneRenderer) override;
 	virtual void AllocateReflectionCaptures(const TArray<UReflectionCaptureComponent*>& NewCaptures) override;
-	virtual void UpdateSkyCaptureContents(const USkyLightComponent* CaptureComponent, bool bCaptureEmissiveOnly, UTextureCube* SourceCubemap, FTexture* OutProcessedTexture, FSHVectorRGB3& OutIrradianceEnvironmentMap) override; 
+	virtual void UpdateSkyCaptureContents(const USkyLightComponent* CaptureComponent, bool bCaptureEmissiveOnly, UTextureCube* SourceCubemap, FTexture* OutProcessedTexture, float& OutAverageBrightness, FSHVectorRGB3& OutIrradianceEnvironmentMap) override; 
 	virtual void PreCullStaticMeshes(const TArray<UStaticMeshComponent*>& ComponentsToPreCull, const TArray<TArray<FPlane> >& CullVolumes) override;
 	virtual void AddPrecomputedLightVolume(const class FPrecomputedLightVolume* Volume) override;
 	virtual void RemovePrecomputedLightVolume(const class FPrecomputedLightVolume* Volume) override;
diff --git a/Engine/Source/Runtime/Renderer/Private/SceneRendering.cpp b/Engine/Source/Runtime/Renderer/Private/SceneRendering.cpp
--- a/Engine/Source/Runtime/Renderer/Private/SceneRendering.cpp
+++ b/Engine/Source/Runtime/Renderer/Private/SceneRendering.cpp
@@ -139,7 +139,9 @@ static TAutoConsoleVariable<float> CVarTessellationAdaptivePixelsPerTriangle(
 static TAutoConsoleVariable<int32> CVarForwardShading(
 	TEXT("r.ForwardShading"),
 	0,
-	TEXT("Whether to use forward shading on desktop platforms.  Requires Shader Model 5 hardware.  Forward shading has lower constant cost, but fewer features supported.  0:off, 1:on"),
+	TEXT("Whether to use forward shading on desktop platforms - requires Shader Model 5 hardware.\n")
+	TEXT("Forward shading has lower constant cost, but fewer features supported. 0:off, 1:on\n")
+	TEXT("This rendering path is a work in progress with many unimplemented features, notably only a single reflection capture is applied per object and no translucency dynamic shadow receiving."),
 	ECVF_RenderThreadSafe | ECVF_ReadOnly);
 
 static TAutoConsoleVariable<int32> CVarSupportSimpleForwardShading(
@@ -156,6 +158,17 @@ static TAutoConsoleVariable<int32> CVarSimpleForwardShading(
 	TEXT("0:off, 1:on"),
 	ECVF_RenderThreadSafe | ECVF_Scalability);
 
+static TAutoConsoleVariable<float> CVarNormalCurvatureToRoughnessBias(
+	TEXT("r.NormalCurvatureToRoughnessBias"),
+	0.0f,
+	TEXT("Biases the roughness resulting from screen space normal changes for materials with NormalCurvatureToRoughness enabled.  Valid range [-1, 1]"),
+	ECVF_RenderThreadSafe | ECVF_Scalability);
+
+static TAutoConsoleVariable<float> CVarNormalCurvatureToRoughnessScale(
+	TEXT("r.NormalCurvatureToRoughnessScale"),
+	1.0f,
+	TEXT("Scales the roughness resulting from screen space normal changes for materials with NormalCurvatureToRoughness enabled.  Valid range [0, 2]"),
+	ECVF_RenderThreadSafe | ECVF_Scalability);
 
 /*-----------------------------------------------------------------------------
 	FParallelCommandListSet
@@ -445,7 +458,6 @@ FViewInfo::~FViewInfo()
 
 void FViewInfo::SetupSkyIrradianceEnvironmentMapConstants(FVector4* OutSkyIrradianceEnvironmentMap) const
 {
-
 	FScene* Scene = nullptr;
 
 	if (Family->Scene)
@@ -855,13 +867,7 @@ void FViewInfo::CreateUniformBuffer(
 	ViewUniformShaderParameters.Random = FMath::Rand();
 	ViewUniformShaderParameters.FrameNumber = Family->FrameNumber;
 
-	// Lets not use Lightmaps if we don't allow static lighting, shall we?
-	static const auto AllowStaticLightingVar = IConsoleManager::Get().FindConsoleVariable(TEXT("r.AllowStaticLighting"));
-	static const auto CVar = IConsoleManager::Get().FindConsoleVariable(TEXT("r.DiffuseFromCaptures"));
-	const bool bUseLightmaps = (AllowStaticLightingVar->GetInt() == 1) && (CVar->GetInt() == 0);
-
 	ViewUniformShaderParameters.CameraCut = bCameraCut ? 1 : 0;
-	ViewUniformShaderParameters.UseLightmaps = bUseLightmaps ? 1 : 0;
 
 	uint32 StateFrameIndexMod8 = 0;
 
@@ -947,6 +953,11 @@ void FViewInfo::CreateUniformBuffer(
 		FinalPostProcessSettings.IndirectLightingColor.G * FinalPostProcessSettings.IndirectLightingIntensity,
 		FinalPostProcessSettings.IndirectLightingColor.B * FinalPostProcessSettings.IndirectLightingIntensity);
 
+	ViewUniformShaderParameters.NormalCurvatureToRoughnessScaleBias.X = FMath::Clamp(CVarNormalCurvatureToRoughnessScale.GetValueOnAnyThread(), 0.0f, 2.0f);
+	ViewUniformShaderParameters.NormalCurvatureToRoughnessScaleBias.Y = FMath::Clamp(CVarNormalCurvatureToRoughnessBias.GetValueOnAnyThread(), -1.0f, 1.0f);
+
+	ViewUniformShaderParameters.RenderingReflectionCaptureMask = bIsReflectionCapture ? 1.0f : 0.0f;
+
 	ViewUniformShaderParameters.AmbientCubemapTint = FinalPostProcessSettings.AmbientCubemapTint;
 	ViewUniformShaderParameters.AmbientCubemapIntensity = FinalPostProcessSettings.AmbientCubemapIntensity;
 
@@ -1022,6 +1033,9 @@ void FViewInfo::CreateUniformBuffer(
 	extern int32 GDistanceFieldAOSpecularOcclusionMode;
 	ViewUniformShaderParameters.DistanceFieldAOSpecularOcclusionMode = GDistanceFieldAOSpecularOcclusionMode;
 
+	extern FVector2D GetReflectionEnvironmentRoughnessMixingScaleBias();
+	ViewUniformShaderParameters.ReflectionEnvironmentRoughnessMixingScaleBias = GetReflectionEnvironmentRoughnessMixingScaleBias();
+
 	ViewUniformShaderParameters.StereoPassIndex = (StereoPass != eSSP_RIGHT_EYE) ? 0 : 1;
 
 	OutViewUniformBuffer = TUniformBufferRef<FViewUniformShaderParameters>::CreateUniformBufferImmediate(ViewUniformShaderParameters, UniformBuffer_SingleFrame);
@@ -1434,24 +1448,7 @@ void FSceneRenderer::RenderFinish(FRHICommandListImmediate& RHICmdList)
 				{
 					SCOPED_CONDITIONAL_DRAW_EVENTF(RHICmdList, EventView, Views.Num() > 1, TEXT("View%d"), ViewIndex);
 
-					// this is a helper class for FCanvas to be able to get screen size
-					class FRenderTargetTemp : public FRenderTarget
-					{
-					public:
-						FViewInfo& View;
-
-						FRenderTargetTemp(FViewInfo& InView) : View(InView)
-						{
-						}
-						virtual FIntPoint GetSizeXY() const
-						{
-							return View.ViewRect.Size();
-						};
-						virtual const FTexture2DRHIRef& GetRenderTargetTexture() const
-						{
-							return View.Family->RenderTarget->GetRenderTargetTexture();
-						}
-					} TempRenderTarget(View);
+					FRenderTargetTemp TempRenderTarget(View);
 
 					// create a temporary FCanvas object with the temp render target
 					// so it can get the screen size
diff --git a/Engine/Source/Runtime/Renderer/Private/SceneVisibility.cpp b/Engine/Source/Runtime/Renderer/Private/SceneVisibility.cpp
--- a/Engine/Source/Runtime/Renderer/Private/SceneVisibility.cpp
+++ b/Engine/Source/Runtime/Renderer/Private/SceneVisibility.cpp
@@ -2063,7 +2063,7 @@ void FSceneRenderer::PreVisibilityFrameSetup(FRHICommandListImmediate& RHICmdLis
 			ViewState->SetupDistanceFieldTemporalOffset(ViewFamily);
 		}
 
-		if( View.FinalPostProcessSettings.AntiAliasingMethod == AAM_TemporalAA && ViewState )
+		if( View.AntiAliasingMethod == AAM_TemporalAA && ViewState )
 		{
 			// Subpixel jitter for temporal AA
 			int32 TemporalAASamples = CVarTemporalAASamples.GetValueOnRenderThread();
@@ -2760,7 +2760,7 @@ void FSceneRenderer::PostVisibilityFrameSetup(FILCUpdatePrimTaskData& OutILCTask
 				
 					FViewElementPDI LightPDI( &View, NULL );
 					FMaterialRenderProxy* const ColoredMeshInstance = new(FMemStack::Get()) FColoredMaterialRenderProxy( GEngine->DebugMeshMaterial->GetRenderProxy(false), Color );
-					DrawSphere( &LightPDI, Origin, FVector( SourceRadius, SourceRadius, SourceRadius ), 8, 6, ColoredMeshInstance, SDPG_World );
+					DrawSphere( &LightPDI, Origin, FVector( SourceRadius, SourceRadius, SourceRadius ), 36, 24, ColoredMeshInstance, SDPG_World );
 				}
 			}
 		}
@@ -2792,7 +2792,7 @@ bool FDeferredShadingSceneRenderer::InitViews(FRHICommandListImmediate& RHICmdLi
 		if (!bWillApplyTemporalAA)
 		{
 			// Disable anti-aliasing if we are not going to be able to apply final post process effects
-			View.FinalPostProcessSettings.AntiAliasingMethod = AAM_None;
+			View.AntiAliasingMethod = AAM_None;
 		}
 	}
 	PreVisibilityFrameSetup(RHICmdList);
diff --git a/Engine/Source/Runtime/Renderer/Private/ShadowDepthRendering.cpp b/Engine/Source/Runtime/Renderer/Private/ShadowDepthRendering.cpp
--- a/Engine/Source/Runtime/Renderer/Private/ShadowDepthRendering.cpp
+++ b/Engine/Source/Runtime/Renderer/Private/ShadowDepthRendering.cpp
@@ -2189,10 +2189,10 @@ void FSceneRenderer::RenderShadowDepthMaps(FRHICommandListImmediate& RHICmdList)
 
 			if (ProjectedShadowInfo->CacheMode == SDCM_MovablePrimitivesOnly 
 				&& Scene->CachedShadowMaps.FindChecked(ProjectedShadowInfo->GetLightSceneInfo().Id).bCachedShadowMapHasPrimitives)
-				{
+			{
 				// Skip the clear when we'll copy from a cached shadowmap
 				bDoClear = false;
-				}
+			}
 
 			SCOPED_CONDITIONAL_DRAW_EVENT(RHICmdList, Clear, bDoClear);
 			SetShadowRenderTargets(RHICmdList, bDoClear);	
@@ -2215,9 +2215,9 @@ void FSceneRenderer::RenderShadowDepthMaps(FRHICommandListImmediate& RHICmdList)
 		{
 			FProjectedShadowInfo* ProjectedShadowInfo = SortedShadowsForShadowDepthPass.PreshadowCache.Shadows[ShadowIndex];
 
-			auto SetShadowRenderTargets = [this, ProjectedShadowInfo, &SceneContext](FRHICommandList& InRHICmdList, bool bPerformClear)
+			auto SetShadowRenderTargets = [this, ProjectedShadowInfo](FRHICommandList& InRHICmdList, bool bPerformClear)
 			{
-				FTextureRHIParamRef PreShadowCacheDepthZ = SceneContext.PreShadowCacheDepthZ->GetRenderTargetItem().TargetableTexture.GetReference();
+				FTextureRHIParamRef PreShadowCacheDepthZ = Scene->PreShadowCacheDepthZ->GetRenderTargetItem().TargetableTexture.GetReference();
 				InRHICmdList.TransitionResources(EResourceTransitionAccess::EWritable, &PreShadowCacheDepthZ, 1);
 
 				// Must preserve existing contents as the clear will be scissored
diff --git a/Engine/Source/Runtime/Renderer/Private/ShadowSetup.cpp b/Engine/Source/Runtime/Renderer/Private/ShadowSetup.cpp
--- a/Engine/Source/Runtime/Renderer/Private/ShadowSetup.cpp
+++ b/Engine/Source/Runtime/Renderer/Private/ShadowSetup.cpp
@@ -823,11 +823,15 @@ void FProjectedShadowInfo::AddSubjectPrimitive(FPrimitiveSceneInfo* PrimitiveSce
 					FViewInfo& CurrentView = *Views[ViewIndex];
 
 					const float DistanceSquared = ( Bounds.Origin - CurrentView.ShadowViewMatrices.ViewOrigin ).SizeSquared();
-					const bool bDrawShadowDepth = FMath::Square( Bounds.SphereRadius ) > FMath::Square( GMinScreenRadiusForShadowCaster ) * DistanceSquared;
-					if( !bDrawShadowDepth )
+
+					if (bWholeSceneShadow)
 					{
-						// cull object if it's too small to be considered as shadow caster
-						continue;
+						const bool bDrawShadowDepth = FMath::Square( Bounds.SphereRadius ) > FMath::Square( GMinScreenRadiusForShadowCaster ) * DistanceSquared;
+						if( !bDrawShadowDepth )
+						{
+							// cull object if it's too small to be considered as shadow caster
+							continue;
+						}
 					}
 
 					// Update visibility for meshes which weren't visible in the main views or were visible with static relevance
@@ -1278,14 +1282,8 @@ void FSceneRenderer::UpdatePreshadowCache(FSceneRenderTargets& SceneContext)
 					CachedShadow->ResolutionY + CachedShadow->BorderSize * 2));
 				Scene->CachedPreshadows.RemoveAt(CachedShadowIndex);
 			}
-			else if (SceneContext.bPreshadowCacheNewlyAllocated)
-			{
-				CachedShadow->bDepthsCached = false;
-			}
 		}
 
-		SceneContext.bPreshadowCacheNewlyAllocated = false;
-
 		TArray<TRefCountPtr<FProjectedShadowInfo>, SceneRenderingAllocator> UncachedPreShadows;
 
 		// Gather a list of preshadows that can be cached
@@ -2805,17 +2803,27 @@ void FSceneRenderer::AllocateShadowDepthTargets(FRHICommandListImmediate& RHICmd
 		AllocateCSMDepthTargets(RHICmdList, WholeSceneDirectionalShadows);
 	}
 
-	SortedShadowsForShadowDepthPass.PreshadowCache.RenderTargets.DepthTarget = SceneContext.PreShadowCacheDepthZ;
-
-	for (int32 ShadowIndex = 0; ShadowIndex < CachedPreShadows.Num(); ShadowIndex++)
+	if (CachedPreShadows.Num() > 0)
 	{
-		FProjectedShadowInfo* ProjectedShadowInfo = CachedPreShadows[ShadowIndex];
-		ProjectedShadowInfo->RenderTargets.DepthTarget = SceneContext.PreShadowCacheDepthZ.GetReference();
+		if (!Scene->PreShadowCacheDepthZ)
+		{
+			FPooledRenderTargetDesc Desc(FPooledRenderTargetDesc::Create2DDesc(SceneContext.GetPreShadowCacheTextureResolution(), PF_ShadowDepth, FClearValueBinding::None, TexCreate_None, TexCreate_DepthStencilTargetable, false));
+			Desc.AutoWritable = false;
+			GRenderTargetPool.FindFreeElement(RHICmdList, Desc, Scene->PreShadowCacheDepthZ, TEXT("PreShadowCacheDepthZ"));
+		}
 
-		if (!ProjectedShadowInfo->bDepthsCached)
+		SortedShadowsForShadowDepthPass.PreshadowCache.RenderTargets.DepthTarget = Scene->PreShadowCacheDepthZ;
+
+		for (int32 ShadowIndex = 0; ShadowIndex < CachedPreShadows.Num(); ShadowIndex++)
 		{
-			ProjectedShadowInfo->SetupShadowDepthView(RHICmdList, this);
-			SortedShadowsForShadowDepthPass.PreshadowCache.Shadows.Add(ProjectedShadowInfo);
+			FProjectedShadowInfo* ProjectedShadowInfo = CachedPreShadows[ShadowIndex];
+			ProjectedShadowInfo->RenderTargets.DepthTarget = Scene->PreShadowCacheDepthZ.GetReference();
+
+			if (!ProjectedShadowInfo->bDepthsCached)
+			{
+				ProjectedShadowInfo->SetupShadowDepthView(RHICmdList, this);
+				SortedShadowsForShadowDepthPass.PreshadowCache.Shadows.Add(ProjectedShadowInfo);
+			}
 		}
 	}
 
diff --git a/Engine/Source/Runtime/Renderer/Private/TranslucentRendering.cpp b/Engine/Source/Runtime/Renderer/Private/TranslucentRendering.cpp
--- a/Engine/Source/Runtime/Renderer/Private/TranslucentRendering.cpp
+++ b/Engine/Source/Runtime/Renderer/Private/TranslucentRendering.cpp
@@ -215,7 +215,10 @@ const FProjectedShadowInfo* FDeferredShadingSceneRenderer::PrepareTranslucentSha
 				{
 					FProjectedShadowInfo* CurrentShadowInfo = VisibleLightInfo->AllProjectedShadows[ShadowIndex];
 
-					if (CurrentShadowInfo && CurrentShadowInfo->bTranslucentShadow && CurrentShadowInfo->GetParentSceneInfo() == PrimitiveSceneInfo)
+					if (CurrentShadowInfo 
+						&& CurrentShadowInfo->bTranslucentShadow 
+						&& CurrentShadowInfo->GetParentSceneInfo() == PrimitiveSceneInfo
+						&& CurrentShadowInfo->ShadowDepthView)
 					{
 						check(CurrentShadowInfo->RenderTargets.ColorTargets.Num() > 0);
 						TranslucentSelfShadow = CurrentShadowInfo;
@@ -268,7 +271,7 @@ void FTranslucencyDrawingPolicyFactory::CopySceneColor(FRHICommandList& RHICmdLi
 {
 	FSceneRenderTargets& SceneContext = FSceneRenderTargets::Get(RHICmdList);
 
-	SCOPED_DRAW_EVENTF(RHICmdList, EventCopy, TEXT("CopySceneColor for %s %s"), *PrimitiveSceneProxy->GetOwnerName().ToString(), *PrimitiveSceneProxy->GetResourceName().ToString());
+	SCOPED_DRAW_EVENTF(RHICmdList, EventCopy, TEXT("CopySceneColor from SceneColor node for %s %s"), *PrimitiveSceneProxy->GetOwnerName().ToString(), *PrimitiveSceneProxy->GetResourceName().ToString());
 	RHICmdList.SetRasterizerState(TStaticRasterizerState<FM_Solid, CM_None>::GetRHI());
 	RHICmdList.SetDepthStencilState(TStaticDepthStencilState<false, CF_Always>::GetRHI());
 	RHICmdList.SetBlendState(TStaticBlendState<>::GetRHI());
diff --git a/Engine/Source/Runtime/Renderer/Private/WideCustomResolveShaders.cpp b/Engine/Source/Runtime/Renderer/Private/WideCustomResolveShaders.cpp
new file mode 100644
--- /dev/null
+++ b/Engine/Source/Runtime/Renderer/Private/WideCustomResolveShaders.cpp
@@ -0,0 +1,85 @@
+
+// Copyright 1998-2015 Epic Games, Inc. All Rights Reserved.
+
+#include "RendererPrivate.h"
+#include "WideCustomResolveShaders.h"
+#include "ShaderParameterUtils.h"
+
+IMPLEMENT_SHADER_TYPE(, FWideCustomResolveVS, TEXT("WideCustomResolveShaders"), TEXT("WideCustomResolveVS"), SF_Vertex);
+
+#define IMPLEMENT_RESOLVE_SHADER(Width, MSAA) \
+	typedef FWideCustomResolvePS<Width,MSAA> FWideCustomResolve##Width##_##MSAA##xPS; \
+	IMPLEMENT_SHADER_TYPE(template<>, FWideCustomResolve##Width##_##MSAA##xPS, TEXT("WideCustomResolveShaders"), TEXT("WideCustomResolvePS"), SF_Pixel)
+
+IMPLEMENT_RESOLVE_SHADER(0, 1);
+IMPLEMENT_RESOLVE_SHADER(2, 0);
+IMPLEMENT_RESOLVE_SHADER(2, 1);
+IMPLEMENT_RESOLVE_SHADER(2, 2);
+IMPLEMENT_RESOLVE_SHADER(2, 3);
+IMPLEMENT_RESOLVE_SHADER(4, 0);
+IMPLEMENT_RESOLVE_SHADER(4, 1);
+IMPLEMENT_RESOLVE_SHADER(4, 2);
+IMPLEMENT_RESOLVE_SHADER(4, 3);
+
+
+template <unsigned Width, unsigned MSAA>
+static void ResolveColorWideInternal2(
+	FRHICommandList& RHICmdList,
+	const ERHIFeatureLevel::Type CurrentFeatureLevel,
+	const FTextureRHIRef& SrcTexture,
+	const FIntPoint& SrcOrigin)
+{
+	auto ShaderMap = GetGlobalShaderMap(CurrentFeatureLevel);
+
+	TShaderMapRef<FWideCustomResolveVS> VertexShader(ShaderMap);
+	TShaderMapRef<FWideCustomResolvePS<MSAA, Width>> PixelShader(ShaderMap);
+	static FGlobalBoundShaderState BoundShaderState;
+	SetGlobalBoundShaderState(RHICmdList, CurrentFeatureLevel, BoundShaderState, GetVertexDeclarationFVector4(), *VertexShader, *PixelShader);
+	PixelShader->SetParameters(RHICmdList, SrcTexture, SrcOrigin);
+
+	RHICmdList.DrawPrimitive(PT_TriangleList, 0, 1, 1);
+}
+
+template <unsigned MSAA>
+static void ResolveColorWideInternal(
+	FRHICommandList& RHICmdList, 
+	const ERHIFeatureLevel::Type CurrentFeatureLevel,
+	const FTextureRHIRef& SrcTexture,
+	const FIntPoint& SrcOrigin, 
+	int32 WideFilterWidth)
+{
+	switch (WideFilterWidth)
+	{
+	case 0: ResolveColorWideInternal2<0, MSAA>(RHICmdList, CurrentFeatureLevel, SrcTexture, SrcOrigin); break;
+	case 1: ResolveColorWideInternal2<1, MSAA>(RHICmdList, CurrentFeatureLevel, SrcTexture, SrcOrigin); break;
+	case 2: ResolveColorWideInternal2<2, MSAA>(RHICmdList, CurrentFeatureLevel, SrcTexture, SrcOrigin); break;
+	case 3: ResolveColorWideInternal2<3, MSAA>(RHICmdList, CurrentFeatureLevel, SrcTexture, SrcOrigin); break;
+	}
+}
+
+void ResolveFilterWide(
+	FRHICommandList& RHICmdList, 
+	const ERHIFeatureLevel::Type CurrentFeatureLevel,
+	const FTextureRHIRef& SrcTexture,
+	const FIntPoint& SrcOrigin, 
+	int32 NumSamples,
+	int32 WideFilterWidth)
+{				
+	if (NumSamples <= 1)
+	{
+		ResolveColorWideInternal2<1, 0>(RHICmdList, CurrentFeatureLevel, SrcTexture, SrcOrigin);
+	}
+	else if (NumSamples == 2)
+	{
+		ResolveColorWideInternal<2>(RHICmdList, CurrentFeatureLevel, SrcTexture, SrcOrigin, WideFilterWidth);
+	}
+	else if (NumSamples == 4)
+	{
+		ResolveColorWideInternal<4>(RHICmdList, CurrentFeatureLevel, SrcTexture, SrcOrigin, WideFilterWidth);
+	}
+	else
+	{
+		// Need to implement more.
+		check(0);
+	}
+}
\ No newline at end of file
diff --git a/Engine/Source/Runtime/Renderer/Private/WideCustomResolveShaders.h b/Engine/Source/Runtime/Renderer/Private/WideCustomResolveShaders.h
new file mode 100644
--- /dev/null
+++ b/Engine/Source/Runtime/Renderer/Private/WideCustomResolveShaders.h
@@ -0,0 +1,84 @@
+// Copyright 1998-2015 Epic Games, Inc. All Rights Reserved.
+
+#pragma once
+
+#include "GlobalShader.h"
+#include "ShaderParameters.h"
+
+class FWideCustomResolveVS : public FGlobalShader
+{
+	DECLARE_SHADER_TYPE(FWideCustomResolveVS,Global);
+public:
+	FWideCustomResolveVS() {}
+	FWideCustomResolveVS(const ShaderMetaType::CompiledShaderInitializerType& Initializer)
+		: FGlobalShader( Initializer )
+	{
+	}
+
+	virtual bool Serialize(FArchive& Ar) override
+	{
+		bool bShaderHasOutdatedParameters = FGlobalShader::Serialize(Ar);
+		return bShaderHasOutdatedParameters;
+	}
+
+	static bool ShouldCache(EShaderPlatform Platform)
+	{
+		return IsFeatureLevelSupported(Platform, ERHIFeatureLevel::SM5) || (Platform == SP_PCD3D_ES2);
+	}
+};
+
+template <unsigned MSAASampleCount, unsigned Width>
+class FWideCustomResolvePS : public FGlobalShader
+{
+	DECLARE_SHADER_TYPE(FWideCustomResolvePS,Global);
+public:
+	FWideCustomResolvePS() {}
+	FWideCustomResolvePS(const ShaderMetaType::CompiledShaderInitializerType& Initializer)
+		: FGlobalShader( Initializer )
+	{
+		static_assert(Width >= 0 && Width <= 3, "Invalid width");
+		static_assert(MSAASampleCount == 0 || MSAASampleCount == 2 || MSAASampleCount == 4, "Invalid sample count");
+
+		Tex.Bind(Initializer.ParameterMap, TEXT("Tex"), SPF_Mandatory);
+		ResolveOrigin.Bind(Initializer.ParameterMap, TEXT("ResolveOrigin"));
+	}
+
+	virtual bool Serialize(FArchive& Ar) override
+	{
+		bool bShaderHasOutdatedParameters = FGlobalShader::Serialize(Ar);
+		Ar << Tex;
+		Ar << ResolveOrigin;
+		return bShaderHasOutdatedParameters;
+	}
+
+	static bool ShouldCache(EShaderPlatform Platform)
+	{
+		return IsFeatureLevelSupported(Platform, ERHIFeatureLevel::SM5);
+	}
+
+	void SetParameters(FRHICommandList& RHICmdList, FTextureRHIParamRef Texture2DMS, FIntPoint Origin)
+	{
+		FPixelShaderRHIParamRef PixelShaderRHI = GetPixelShader();
+		SetTextureParameter(RHICmdList, PixelShaderRHI, Tex, Texture2DMS);
+		SetShaderValue(RHICmdList, PixelShaderRHI, ResolveOrigin, FVector2D(Origin));
+	}
+
+	static void ModifyCompilationEnvironment(EShaderPlatform Platform, FShaderCompilerEnvironment& OutEnvironment)
+	{
+		FGlobalShader::ModifyCompilationEnvironment(Platform, OutEnvironment);
+		OutEnvironment.SetDefine(TEXT("WIDE_RESOLVE_WIDTH"), Width);
+		OutEnvironment.SetDefine(TEXT("MSAA_SAMPLE_COUNT"), MSAASampleCount);
+	}
+
+protected:
+	FShaderResourceParameter Tex;
+	FShaderParameter ResolveOrigin;
+};
+
+extern void ResolveFilterWide(
+	FRHICommandList& RHICmdList,
+	const ERHIFeatureLevel::Type CurrentFeatureLevel,
+	const FTextureRHIRef& SrcTexture,
+	const FIntPoint& SrcOrigin,
+	int32 NumSamples,
+	int32 WideFilterWidth);
\ No newline at end of file
diff --git a/Engine/Source/Runtime/Renderer/Public/DecalRenderingCommon.h b/Engine/Source/Runtime/Renderer/Public/DecalRenderingCommon.h
--- a/Engine/Source/Runtime/Renderer/Public/DecalRenderingCommon.h
+++ b/Engine/Source/Runtime/Renderer/Public/DecalRenderingCommon.h
@@ -62,6 +62,15 @@ struct FDecalRenderingCommon
 			return RTM_SceneColor;
 		}
 	
+		// Can't modify GBuffers when forward shading, just modify scene color
+		if (IsAnyForwardShadingEnabled(Platform)
+			&& (DecalBlendMode == DBM_Translucent
+				|| DecalBlendMode == DBM_Stain
+				|| DecalBlendMode == DBM_Normal))
+		{
+			return RTM_SceneColor;
+		}
+
 		switch(DecalBlendMode)
 		{
 			case DBM_Translucent:
