Fixed duplicate localization tag.

diff --git a/Engine/Source/Runtime/Renderer/Private/PostProcess/PostProcessStreamingAccuracyLegend.cpp b/Engine/Source/Runtime/Renderer/Private/PostProcess/PostProcessStreamingAccuracyLegend.cpp
--- a/Engine/Source/Runtime/Renderer/Private/PostProcess/PostProcessStreamingAccuracyLegend.cpp
+++ b/Engine/Source/Runtime/Renderer/Private/PostProcess/PostProcessStreamingAccuracyLegend.cpp
@@ -54,7 +54,7 @@ void FRCPassPostProcessStreamingAccuracyLegend::DrawCustom(FRenderingCompositePa
 	}
 	else if (ViewFamily.GetDebugViewShaderMode() == DVSM_MeshTexCoordSizeAccuracy)
 	{
-		DrawDesc(Canvas, DestRect.Min.X + 115, DestRect.Max.Y - 75, LOCTEXT("DescMaterialTexCoordScalesAccuracy", "Shows under/over texture streaming caused by mesh texcoord mappings in world units."));
+		DrawDesc(Canvas, DestRect.Min.X + 115, DestRect.Max.Y - 75, LOCTEXT("DescMeshTexCoordSizeAccuracy", "Shows under/over texture streaming caused by mesh texcoord mappings in world units."));
 	}
 	else if (ViewFamily.GetDebugViewShaderMode() == DVSM_PrimitiveDistanceAccuracy)
 	{
