Copying //UE4/Dev-Mobile to //UE4/Dev-Main (Source: //UE4/Dev-Mobile @ 3155909)

==========================
MAJOR FEATURES + CHANGES
==========================

Change 3118534 on 2016/09/08 by [REDACTED_PERSON]

	Certain non-looping SFX would not properly fire "Stop" events when the effect was finished, and therefore Sound Mixes that should end at the end of a SFX (such as ducking the BGM volume) were not properly finishing. Fixing the IsSourceFinished call to check the Position of the player to determine if we are actually done playing, which was not working properly before for PCM sounds.

	#ue4
	#android

Change 3119125 on 2016/09/09 by [REDACTED_PERSON]

	Mobile launcher profile wizard:
	Fixed case where DLC will not be built if user selects non Development build configuration
	Fixed case where project maps will be empty if UE4 and project located on different drives

Change 3122584 on 2016/09/13 by [REDACTED_PERSON]

	Add simple clip plane for planar reflections.

Change 3129390 on 2016/09/16 by [REDACTED_PERSON]

	Fixed ES 3.1 detection to also support devices returning ES 3.2 support (eg Note 7)

Change 3129867 on 2016/09/18 by [REDACTED_PERSON]

	Fixed ES 3.1 detection to also support devices returning ES 3.2 support (eg Note 7)

Change 3131961 on 2016/09/20 by [REDACTED_PERSON]

	Fix missing editor widgets and gamma incorrectness when mobileHDR == false.
	Fix editor widget rendering when hdr encoding is active.

Change 3132717 on 2016/09/20 by [REDACTED_PERSON]

	Add $S(ProjectDir) to UPL
	#ue4

Change 3132940 on 2016/09/20 by [REDACTED_PERSON]

	Corrected case for some include files (contributed by [REDACTED_PERSON])
	#PR #2636
	#ue4
	#android

Change 3134098 on 2016/09/21 by [REDACTED_PERSON]

	Mobile CSM shadow quality controllable via quality levels.

Change 3134931 on 2016/09/21 by [REDACTED_PERSON]

	Allow Windows types in vulkan.h
	#ue4
	#vulkan

Change 3135380 on 2016/09/21 by [REDACTED_PERSON]

	Plugin which exposes some of BuildPatchServices functionality to BP. Inteded to be used on mobile platforms for donwloading game content.
	Right now misses: IOS download directory and iOS WiFi detection

Change 3136004 on 2016/09/22 by [REDACTED_PERSON]

	Add project option to disable vertex fog on mobile.
	Vertex fog is now enabled even when mobile HDR is not.

Change 3137377 on 2016/09/22 by [REDACTED_PERSON]

	Fix compile error from CL# 3135380

Change 3139571 on 2016/09/26 by [REDACTED_PERSON]

	Applied deferred change CL 3101462 to mobile to make planar reflections no longer update GPU particles

Change 3139663 on 2016/09/26 by [REDACTED_PERSON]

	Include Android shader cache files when packaging

Change 3142839 on 2016/09/28 by [REDACTED_PERSON]

	Added WiFi connection detection on iOS

Change 3142845 on 2016/09/28 by [REDACTED_PERSON]

	Fixed various issues with TcpMessageTransport discovered when transferring automation testing screenshots from mobile devices
	- socket not readable or writable is not an error condition if output buffer is full
	- messages were previously limited to 64kb but screenshots overflowed this
	- messages over 8kb were not reliably received as the inbound buffer was full so the available bytes was always less than the message length
	- sending large messages was not reliable due to the output buffer being full

Change 3143280 on 2016/09/28 by [REDACTED_PERSON]

	Clear out UnbuiltInstanceBoundsList when async building a tree with no instances

Change 3143282 on 2016/09/28 by [REDACTED_PERSON]

	Fix issue where client functional tests in the current map do not appear on clients running with cooked content.
	Problem is that the AssetRegistry uses in-memory metadata created on load for currently-loaded assets, but cooked content only has the serialized AssetRegistry and individual assets do not contain any metadata.

Change 3143808 on 2016/09/28 by [REDACTED_PERSON]

	Assume that the app starts in focus at startup and don't wait for an "APP_EVENT_STATE_WINDOW_GAINED_FOCUS" event to fire, as this event will not come down from SplashActivity since it is not a NativeActivity. If the user then rotates the device in Sensor or FullSensor orientation during SplashActivity and forces an eglSurface recreation, the initial Create will properly execute if we're "in focus". Previously, the create-destroy-create cycle would not properly execute due to the EventManager thinking the app was not yet in focus, and would cause the second create to get a 0x3003 error (EGL_BAD_ALLOC)

	#ue4
	#android

Change 3144880 on 2016/09/29 by [REDACTED_PERSON]

	Merging //UE4/Dev-Main to Dev-Mobile (//UE4/Dev-Mobile)

Change 3146220 on 2016/09/30 by [REDACTED_PERSON]

	Adjusted device button style to make it more readable.

Change 3146280 on 2016/09/30 by [REDACTED_PERSON]

	Replaced IBuildManifest::ComputeDownloadSize with a new function

Change 3146302 on 2016/09/30 by [REDACTED_PERSON]

	Added more stringent checks for ES3.1 compatibility

Change 3146435 on 2016/09/30 by [REDACTED_PERSON]

	Prevent landscape grass being duplicated for PIE, causing ensure

Change 3147532 on 2016/09/30 by [REDACTED_PERSON]

	Use .sh extension for Android install scripts on Linux
	#ue4
	#android
	#linux

Change 3149851 on 2016/10/04 by [REDACTED_PERSON]

	Mobile: Added custom depth rendering
	Mobile: Added support for CustomDepth and SceneDepth in post-process materails

Change 3149852 on 2016/10/04 by [REDACTED_PERSON]

	Fixed comments for SortBasePass console variable

Change 3149857 on 2016/10/04 by [REDACTED_PERSON]

	Remove dead code in ProceduralFoliageComponentDetails.cpp

Change 3149863 on 2016/10/04 by [REDACTED_PERSON]

	Merging //UE4/Dev-Main to Dev-Mobile (//UE4/Dev-Mobile)

Change 3149896 on 2016/10/04 by [REDACTED_PERSON]

	Fixed: SkyLight makes level entire green on Android devices

Change 3150102 on 2016/10/04 by [REDACTED_PERSON]

	Bring Protostar 4.13 fixes back to Dev-Mobile

	Engine
	- MaxDescriptorSets = 16384 to prevent crash on Mali
	- Include texture formats in FRenderTargetLayoutHashableStruct to solve RT aliasing issue
	- Use ERenderTargetLoadAction::EClear for planar reflection target to work around Adreno issue
	- Default Adreno to SPIR-V

	Contents
	- Disable fog, reduce CSM shadow quality, fix device profiles
	- Add PSO cache

Change 3150113 on 2016/10/04 by [REDACTED_PERSON]

	Ensure automation testing screenshots have Alpha=255 (fixes automation screenshots on Mobile)

Change 3150231 on 2016/10/04 by [REDACTED_PERSON]

	Use a new SessionID GUID each time you use the launcher to launch a session.

Change 3150608 on 2016/10/04 by [REDACTED_PERSON]

	Changes for automated testing screenshots on Android.

	- Prevent automation screenshots from changing resolution on platforms with fixed resolution
	- Set GRHIAdapterInternalDriverVersion for OpenGL and Vulkan
	- Parse ImgTec/ARM/Qualcomm GRHIVendorId on OpenGL
	- Added helper to convert GRHIVendorId to string

Change 3151318 on 2016/10/04 by [REDACTED_PERSON]

	Fixed compile error with AdapterVendor

Change 3151366 on 2016/10/04 by [REDACTED_PERSON]

	Prevent FTcpMessageTransportConnection deadlock on device disconnect

Change 3151397 on 2016/10/05 by [REDACTED_PERSON]

	More consistent BP categories for Mobile Patching utils

Change 3151576 on 2016/10/05 by [REDACTED_PERSON]

	Added on screen warning for invalid reflection captures, can be seen only in game running with FeatureLevel < SM4 and no valid capture data

Change 3151795 on 2016/10/05 by [REDACTED_PERSON]

	Linux: update UBT to use a v8 multiarch toolchain.

	- Also added toolchain build scripts and ct-ng configs.

Change 3151966 on 2016/10/05 by [REDACTED_PERSON]

	Add mobile support for inverse opacity to mobile scene captures as well as SCS_SceneColorSceneDepth and SCS_SceneDepth.

Change 3152664 on 2016/10/05 by [REDACTED_PERSON]

	Merging //UE4/Main to Dev-Mobile (//UE4/Dev-Mobile)

Change 3152675 on 2016/10/05 by [REDACTED_PERSON]

	Fixed patching so that it searches for pak files as well as pkg files.
	#test Patch for QAGame

Change 3152728 on 2016/10/05 by [REDACTED_PERSON]

	Update ReflectionCaptureDDCVer (need to resave maps)

Change 3152910 on 2016/10/05 by [REDACTED_PERSON]

	Linux: Fix toolchain for non-AutoSDKs (github) case (UE-36899).

Change 3152966 on 2016/10/05 by [REDACTED_PERSON]

	Linux: Fix test for the installed SDK (UE-36899).

Change 3153004 on 2016/10/05 by [REDACTED_PERSON]

	Linux: fix CIS (UT server case-sens errors).

Change 3153694 on 2016/10/06 by [REDACTED_PERSON]

	Rollback ReflectionCaptureDDCVer change as bug intended to fix UE-36919 does not repro

Change 3154766 on 2016/10/07 by [REDACTED_PERSON]

	Merging //UE4/Dev-Main to Dev-Mobile (//UE4/Dev-Mobile)

Change 3154833 on 2016/10/07 by [REDACTED_PERSON]

	Fix merge error of MobileShadingRenderer.cpp

Change 3154848 on 2016/10/07 by [REDACTED_PERSON]

	Fix mobile scene capture's clear code

Change 3154875 on 2016/10/07 by [REDACTED_PERSON]

	fix vk build issues

Change 3154941 on 2016/10/07 by [REDACTED_PERSON]

	Fix gearvr build fail

Change 3154950 on 2016/10/07 by [REDACTED_PERSON]

	Fix shadowed local variable vk build warning on android.

Change 3155909 on 2016/10/07 by [REDACTED_PERSON]

	UBT: Attempt to work around C1076 error ("internal heap limit reached: use /Zm to specify a higher limit"), encountered when building with XGE. Specify the AutoReserveMemory attribute on XGE tool tasks that manipulate precompiled headers.

diff --git a/Engine/Build/Android/Java/src/com/epicgames/ue4/GameActivity.java b/Engine/Build/Android/Java/src/com/epicgames/ue4/GameActivity.java
--- a/Engine/Build/Android/Java/src/com/epicgames/ue4/GameActivity.java
+++ b/Engine/Build/Android/Java/src/com/epicgames/ue4/GameActivity.java
@@ -59,6 +59,9 @@ import android.widget.PopupWindow;
 
 import android.media.AudioManager;
 
+import android.net.ConnectivityManager;
+import android.net.NetworkInfo;
+
 import com.google.android.gms.auth.GoogleAuthUtil;
 import com.google.android.gms.common.api.GoogleApiClient;
 import com.google.android.gms.common.GooglePlayServicesUtil;
@@ -2121,6 +2124,19 @@ public class GameActivity extends NativeActivity implements SurfaceHolder.Callba
 		return false;
 	}
 
+	public boolean AndroidThunkJava_HasActiveWiFiConnection()
+	{
+		ConnectivityManager cm = (ConnectivityManager)this.getSystemService(Context.CONNECTIVITY_SERVICE);
+		NetworkInfo activeNetwork = cm.getActiveNetworkInfo();
+		
+		boolean isConnected = (activeNetwork != null && activeNetwork.isConnectedOrConnecting());
+		if (isConnected)
+		{
+			return (activeNetwork.getType() == ConnectivityManager.TYPE_WIFI);
+		}
+		return false;
+	}
+	
 	public boolean AndroidThunkJava_HasMetaDataKey(String key)
 	{
 		if (_bundle == null || key == null)
diff --git a/Engine/Build/BatchFiles/Linux/Toolchain_x86-64/BuildCrossToolchain.sh b/Engine/Build/BatchFiles/Linux/Toolchain_x86-64/BuildCrossToolchain.sh
--- a/Engine/Build/BatchFiles/Linux/Toolchain_x86-64/BuildCrossToolchain.sh
+++ b/Engine/Build/BatchFiles/Linux/Toolchain_x86-64/BuildCrossToolchain.sh
@@ -3,7 +3,10 @@
 
 SelfName=`basename $0`
 LinuxBuildTemp=linux-host-build
-Win64BuildTemp=win64-host-build
+Win64BuildTemp_x86_64=win64-host-build-x86_64
+Win64BuildTemp_ARM32=win64-host-build-ARM32
+Win64BuildTemp_i686=win64-host-build-i686
+OutputFolder=CrossToolchainMultiarch
 
 # prereq:
 # mercurial autoconf gperf bison flex libtool ncurses-dev
@@ -12,9 +15,15 @@ Win64BuildTemp=win64-host-build
 #############
 # build recent ct-ng
 
-hg clone http://crosstool-ng.org/hg/crosstool-ng -r "1.19.0"
+rm -rf crosstool-ng
+git clone [REDACTED_PROFILE_URL]/crosstool-ng -b 1.22
 cd crosstool-ng
 
+# install Ubuntu (16.04) packages
+UbuntuPackages="autoconf gperf bison flex help2man libtool-bin libncurses5-dev"
+echo Installing packages $UbuntuPackages
+sudo apt install -y $UbuntuPackages
+
 ./bootstrap
 [ $? -ne 0 ] && echo "$SelfName: Error encountered, exiting at line $LINENO" && exit 1
 
@@ -43,31 +52,68 @@ cd ..
 
 
 #############
-# build Windows-hosted toolchain that targets Linux
+# build Windows-hosted toolchain that targets Linux x86_64
 
-rm -rf $Win64BuildTemp
-mkdir $Win64BuildTemp
+rm -rf $Win64BuildTemp_x86_64
+mkdir $Win64BuildTemp_x86_64
 
-cp win64-host.config $Win64BuildTemp/.config
+cp win64-host_x86_64.config $Win64BuildTemp_x86_64/.config
 [ $? -ne 0 ] && echo "$SelfName: Error encountered, exiting at line $LINENO" && exit 1
 
-cd $Win64BuildTemp
+cd $Win64BuildTemp_x86_64
+
 ../crosstool-ng/ct-ng build
 [ $? -ne 0 ] && echo "$SelfName: Error encountered, exiting at line $LINENO" && exit 1
 
 cd ..
 
 #############
-# copy the GNU toolchain here
+# build Windows-hosted toolchain that targets Linux ARM32
+
+rm -rf $Win64BuildTemp_ARM32
+mkdir $Win64BuildTemp_ARM32
+
+cp win64-host_ARM32.config $Win64BuildTemp_ARM32/.config
+[ $? -ne 0 ] && echo "$SelfName: Error encountered, exiting at line $LINENO" && exit 1
+
+cd $Win64BuildTemp_ARM32
+
+../crosstool-ng/ct-ng build
+[ $? -ne 0 ] && echo "$SelfName: Error encountered, exiting at line $LINENO" && exit 1
+
+cd ..
+
+#############
+# build Windows-hosted toolchain that targets Linux i686
+
+rm -rf $Win64BuildTemp_i686
+mkdir $Win64BuildTemp_i686
+
+cp win64-host_i686.config $Win64BuildTemp_i686/.config
+[ $? -ne 0 ] && echo "$SelfName: Error encountered, exiting at line $LINENO" && exit 1
+
+cd $Win64BuildTemp_i686
+
+../crosstool-ng/ct-ng build
+[ $? -ne 0 ] && echo "$SelfName: Error encountered, exiting at line $LINENO" && exit 1
+
+cd ..
+
+if [ -d $OutputFolder ]; then
+	chmod -R 777 $OutputFolder
+	rm -rf $OutputFolder
+fi
 
-mkdir CrossToolchain
-cp -r ~/x-tools/x86_64-unknown-linux-gnu/*/ CrossToolchain
+mkdir $OutputFolder
+cp -r ~/x-tools/x86_64-unknown-linux-gnu $OutputFolder
+cp -r ~/x-tools/arm-unknown-linux-gnueabihf $OutputFolder
+cp -r ~/x-tools/i686-unknown-linux-gnu $OutputFolder
 [ $? -ne 0 ] && echo "$SelfName: Error encountered, exiting at line $LINENO" && exit 1
 
 #############
 # rearrange files the way we need them
 
-cd CrossToolchain
+cd $OutputFolder/x86_64-unknown-linux-gnu
 find . -type d -exec chmod +w {} \;
 cp -r -L x86_64-unknown-linux-gnu/include .
 cp -r -L x86_64-unknown-linux-gnu/sysroot/lib64 .
@@ -76,6 +122,26 @@ cp -r -L x86_64-unknown-linux-gnu/sysroot/usr/lib64 usr
 cp -r -L x86_64-unknown-linux-gnu/sysroot/usr/include usr
 rm -rf x86_64-unknown-linux-gnu
 
-cd ..
+cd ../..
+
+cd $OutputFolder/arm-unknown-linux-gnueabihf
+find . -type d -exec chmod +w {} \;
+cp -r -L arm-unknown-linux-gnueabihf/include .
+cp -r -L arm-unknown-linux-gnueabihf/sysroot/lib .
+mkdir -p usr
+cp -r -L arm-unknown-linux-gnueabihf/sysroot/usr/lib usr
+cp -r -L arm-unknown-linux-gnueabihf/sysroot/usr/include usr
+rm -rf arm-unknown-linux-gnueabihf
 
+cd ../..
+
+cd $OutputFolder/i686-unknown-linux-gnu
+find . -type d -exec chmod +w {} \;
+cp -r -L i686-unknown-linux-gnu/include .
+cp -r -L i686-unknown-linux-gnu/sysroot/lib .
+mkdir -p usr
+cp -r -L i686-unknown-linux-gnu/sysroot/usr/lib usr
+cp -r -L i686-unknown-linux-gnu/sysroot/usr/include usr
+rm -rf i686-unknown-linux-gnu
 
+cd ../..
diff --git a/Engine/Build/BatchFiles/Linux/Toolchain_x86-64/README.md b/Engine/Build/BatchFiles/Linux/Toolchain_x86-64/README.md
--- a/Engine/Build/BatchFiles/Linux/Toolchain_x86-64/README.md
+++ b/Engine/Build/BatchFiles/Linux/Toolchain_x86-64/README.md
@@ -10,21 +10,15 @@ Big picture view, steps are described in detail later:
 How to use the cross-toolchain.
 =================================
 
-2.1 Set LINUX_ROOT environment variable to point to asbolute (Windows) path of the toolchain.
+2.1 Set LINUX_MULTIARCH_ROOT environment variable to point to asbolute (Windows) path of the toolchain.
 2.2 Restart Visual Studio and regenerate projects with UBT.
 
 
 Building gcc-based cross-toolchain
 ==================================
 
-Copy BuildCrossToolchain.sh together with linux-host.config and win64-host.config to
-a Linux machine running reasonably recent distro. Make sure that the following
-pre-requisites are installed (using Debian package names):
-
- mercurial autoconf gperf bison flex libtool ncurses-dev
-
-and of course make/gcc.
-Run the script, which will clone current crosstool-ng and will use it to perform
+Copy BuildCrossToolchain.sh together with all .config files to a Linux machine running Ubuntu 16.04. 
+Run the script, which will clone a fork of crosstool-ng and will use it to perform
 so called "canadian cross", i.e. will build Linux-hosted toolchain which targets
 mingw64 first, and then will build Windows-hosted toolchain that targets Linux.
 
@@ -33,19 +27,17 @@ to experiment (using "ct-ng menuconfig") if you need more recent ones. Not every
 combination is buildable though.
 
 If the script finishes without error, the new toolchain will be copied to subfolder
-called CrossToolchain of the current folder. Copy it to your Windows machine,
-e.g. to C:\CrossToolchain
+called CrossToolchainMultiarch of the current folder. Copy it to your Windows machine,
+e.g. to C:\CrossToolchainMultiarch
 
 Building clang
 ==============
 
-Note that you can grab a pre-built clang from here: http://llvm.org/releases/download.html, but
-we haven't actually tried that, building ourselves for various reasons.
-
+Note that you can grab a pre-built clang from here: http://llvm.org/releases/download.html
 Detailed instructions how to build clang are given here: http://clang.llvm.org/get_started.html
 
 To recap, you will need: cmake, python and a supported compiler (recent Visual Studio is Ok).
-Grab stable release of llvm and clang from here: http://llvm.org/releases/download.html#3.6.0
+Grab stable release of llvm and clang from here: http://llvm.org/releases/download.html
 and unpack it to your local folder (unpack llvm first, then unpack clang into llvm/tools subdirectory).
 
 Use cmake to generate project files and then build Release (or MinSizeRel) x64 configuration of "install" 
@@ -58,18 +50,18 @@ Having done that, copy from install directory the following
 
 - bin/clang.exe
 - bin/clang++.exe
-- lib/clang/3.6.0/include 
+- lib/clang/x.y.0/include 
 
-into the toolchain (so their relative path from toolchain root stays the same).
+into all the architecture-specific toolchains (so their relative path from toolchain root stays the same).
 
 
 Using toolchain to build UE4 Linux targets
 ==========================================
 
 Set environment variable (Control Panel->System->Advanced system settings->Advanced->Environment variables)
-LINUX_ROOT to absolute path to your toolchain, without trailing backslash, so it looks like this:
+LINUX_MULTIARCH_ROOT to absolute path to your toolchain, without trailing backslash, so it looks like this:
 
-LINUX_ROOT=C:\CrossToolchain
+LINUX_MULTIARCH_ROOT=C:\CrossToolchainMultiarch
 
 After you made sure that variable is set, you can restart Visual Studio and regenerate UE4 project files
 (using GenerateProjectFiles.bat). After this, you should have "Linux" available among Win32/Win64
diff --git a/Engine/Build/BatchFiles/Linux/Toolchain_x86-64/linux-host.config b/Engine/Build/BatchFiles/Linux/Toolchain_x86-64/linux-host.config
--- a/Engine/Build/BatchFiles/Linux/Toolchain_x86-64/linux-host.config
+++ b/Engine/Build/BatchFiles/Linux/Toolchain_x86-64/linux-host.config
@@ -1,7 +1,6 @@
 #
-# Automatically generated make config: don't edit
-# crosstool-NG hg+default-9838aecd6340 Configuration
-# Thu Oct 24 01:39:25 2013
+# Automatically generated file; DO NOT EDIT.
+# Crosstool-NG Configuration
 #
 CT_CONFIGURE_has_xz=y
 CT_CONFIGURE_has_svn=y
@@ -31,7 +30,8 @@ CT_INSTALL_DIR="${CT_PREFIX_DIR}"
 CT_RM_RF_PREFIX_DIR=y
 CT_REMOVE_DOCS=y
 CT_INSTALL_DIR_RO=y
-CT_STRIP_ALL_TOOLCHAIN_EXECUTABLES=y
+CT_STRIP_HOST_TOOLCHAIN_EXECUTABLES=y
+# CT_STRIP_TARGET_TOOLCHAIN_EXECUTABLES is not set
 
 #
 # Downloading
@@ -108,28 +108,30 @@ CT_TARGET_CFLAGS=""
 CT_TARGET_LDFLAGS=""
 # CT_ARCH_alpha is not set
 # CT_ARCH_arm is not set
-# CT_ARCH_avr32 is not set
-# CT_ARCH_blackfin is not set
+# CT_ARCH_avr is not set
 # CT_ARCH_m68k is not set
 # CT_ARCH_microblaze is not set
 # CT_ARCH_mips is not set
+# CT_ARCH_nios2 is not set
 # CT_ARCH_powerpc is not set
 # CT_ARCH_s390 is not set
 # CT_ARCH_sh is not set
 # CT_ARCH_sparc is not set
 CT_ARCH_x86=y
+# CT_ARCH_xtensa is not set
 CT_ARCH_alpha_AVAILABLE=y
 CT_ARCH_arm_AVAILABLE=y
-CT_ARCH_avr32_AVAILABLE=y
-CT_ARCH_blackfin_AVAILABLE=y
+CT_ARCH_avr_AVAILABLE=y
 CT_ARCH_m68k_AVAILABLE=y
 CT_ARCH_microblaze_AVAILABLE=y
 CT_ARCH_mips_AVAILABLE=y
+CT_ARCH_nios2_AVAILABLE=y
 CT_ARCH_powerpc_AVAILABLE=y
 CT_ARCH_s390_AVAILABLE=y
 CT_ARCH_sh_AVAILABLE=y
 CT_ARCH_sparc_AVAILABLE=y
 CT_ARCH_x86_AVAILABLE=y
+CT_ARCH_xtensa_AVAILABLE=y
 CT_ARCH_SUFFIX=""
 
 #
@@ -162,7 +164,7 @@ CT_TOOLCHAIN_BUGURL=""
 #
 # Tuple completion and aliasing
 #
-CT_TARGET_VENDOR="pc"
+CT_TARGET_VENDOR="w64"
 CT_TARGET_ALIAS_SED_EXPR=""
 CT_TARGET_ALIAS=""
 
@@ -213,6 +215,10 @@ CT_BINUTILS_binutils=y
 #
 # GNU binutils
 #
+# CT_CC_BINUTILS_SHOW_LINARO is not set
+# CT_BINUTILS_V_2_25_1 is not set
+# CT_BINUTILS_V_2_25 is not set
+# CT_BINUTILS_V_2_24 is not set
 # CT_BINUTILS_V_2_23_2 is not set
 # CT_BINUTILS_V_2_23_1 is not set
 # CT_BINUTILS_V_2_22 is not set
@@ -230,6 +236,7 @@ CT_BINUTILS_2_18_or_later=y
 CT_BINUTILS_HAS_HASH_STYLE=y
 CT_BINUTILS_HAS_GOLD=y
 CT_BINUTILS_GOLD_SUPPORTS_ARCH=y
+CT_BINUTILS_GOLD_SUPPORT=y
 CT_BINUTILS_HAS_PLUGINS=y
 CT_BINUTILS_HAS_PKGVERSION_BUGURL=y
 CT_BINUTILS_LINKER_LD=y
@@ -246,86 +253,86 @@ CT_BINUTILS_EXTRA_CONFIG_ARRAY=""
 # binutils other options
 #
 
+#
+# C-library
+#
+CT_LIBC="mingw"
+CT_LIBC_mingw=y
+CT_LIBC_avr_libc_AVAILABLE=y
+CT_LIBC_glibc_AVAILABLE=y
+CT_THREADS="win32"
+CT_LIBC_mingw_AVAILABLE=y
+# CT_WINAPI_V_DEVEL is not set
+CT_WINAPI_V_4_0_4=y
+# CT_WINAPI_V_4_0_2 is not set
+# CT_WINAPI_V_4_0_1 is not set
+# CT_WINAPI_V_4_0_0 is not set
+# CT_WINAPI_V_3_1_0 is not set
+# CT_WINAPI_V_3_0_0 is not set
+# CT_WINAPI_V_2_0_8 is not set
+# CT_WINAPI_V_2_0_7 is not set
+# CT_WINAPI_V_select is not set
+CT_WINAPI_VERSION="4.0.4"
+# CT_MINGW_DIRECTX is not set
+# CT_MINGW_DDK is not set
+CT_LIBC_musl_AVAILABLE=y
+CT_LIBC_newlib_AVAILABLE=y
+CT_LIBC_none_AVAILABLE=y
+CT_LIBC_uClibc_AVAILABLE=y
+CT_LIBC_SUPPORT_THREADS_ANY=y
+CT_LIBC_SUPPORT_THREADS_NATIVE=y
+
+#
+# Common C library options
+#
+CT_THREADS_NATIVE=y
+
 #
 # C compiler
 #
 CT_CC="gcc"
-CT_CC_VERSION="4.5.2"
-CT_CC_CORE_PASSES_NEEDED=y
+CT_CC_CORE_PASS_2_NEEDED=y
 CT_CC_gcc=y
 # CT_CC_GCC_SHOW_LINARO is not set
-# CT_CC_V_4_8_1 is not set
-# CT_CC_V_4_8_0 is not set
-# CT_CC_V_4_7_3 is not set
-# CT_CC_V_4_7_2 is not set
-# CT_CC_V_4_7_1 is not set
-# CT_CC_V_4_7_0 is not set
-# CT_CC_V_4_6_4 is not set
-# CT_CC_V_4_6_3 is not set
-# CT_CC_V_4_6_2 is not set
-# CT_CC_V_4_6_1 is not set
-# CT_CC_V_4_6_0 is not set
-# CT_CC_V_4_5_3 is not set
-CT_CC_V_4_5_2=y
-# CT_CC_V_4_5_1 is not set
-# CT_CC_V_4_5_0 is not set
-# CT_CC_V_4_4_7 is not set
-# CT_CC_V_4_4_6 is not set
-# CT_CC_V_4_4_5 is not set
-# CT_CC_V_4_4_4 is not set
-# CT_CC_V_4_4_3 is not set
-# CT_CC_V_4_4_2 is not set
-# CT_CC_V_4_4_1 is not set
-# CT_CC_V_4_4_0 is not set
-# CT_CC_V_4_3_6 is not set
-# CT_CC_V_4_3_5 is not set
-# CT_CC_V_4_3_4 is not set
-# CT_CC_V_4_3_3 is not set
-# CT_CC_V_4_3_2 is not set
-# CT_CC_V_4_3_1 is not set
-# CT_CC_V_4_2_4 is not set
-# CT_CC_V_4_2_2 is not set
-# CT_CC_CUSTOM is not set
+CT_CC_GCC_V_5_2_0=y
+# CT_CC_GCC_V_4_9_3 is not set
+# CT_CC_GCC_V_4_8_5 is not set
+# CT_CC_GCC_V_4_7_4 is not set
+# CT_CC_GCC_V_4_6_4 is not set
+# CT_CC_GCC_V_4_5_4 is not set
+# CT_CC_GCC_V_4_4_7 is not set
+# CT_CC_GCC_V_4_3_6 is not set
+# CT_CC_GCC_V_4_2_4 is not set
+# CT_CC_GCC_CUSTOM is not set
 CT_CC_GCC_4_2_or_later=y
 CT_CC_GCC_4_3_or_later=y
 CT_CC_GCC_4_4_or_later=y
-CT_CC_GCC_4_5=y
 CT_CC_GCC_4_5_or_later=y
+CT_CC_GCC_4_6_or_later=y
+CT_CC_GCC_4_7_or_later=y
+CT_CC_GCC_4_8_or_later=y
+CT_CC_GCC_4_9_or_later=y
+CT_CC_GCC_5=y
+CT_CC_GCC_5_or_later=y
 CT_CC_GCC_HAS_GRAPHITE=y
 # CT_CC_GCC_USE_GRAPHITE is not set
 CT_CC_GCC_HAS_LTO=y
 CT_CC_GCC_USE_LTO=y
 CT_CC_GCC_HAS_PKGVERSION_BUGURL=y
 CT_CC_GCC_HAS_BUILD_ID=y
+CT_CC_GCC_HAS_LNK_HASH_STYLE=y
 CT_CC_GCC_ENABLE_PLUGINS=y
 CT_CC_GCC_USE_GMP_MPFR=y
 CT_CC_GCC_USE_MPC=y
-CT_CC_GCC_USE_LIBELF=y
+CT_CC_GCC_HAS_LIBQUADMATH=y
+CT_CC_GCC_HAS_LIBSANITIZER=y
+CT_CC_GCC_VERSION="5.2.0"
 # CT_CC_LANG_FORTRAN is not set
-CT_CC_SUPPORT_CXX=y
-CT_CC_SUPPORT_FORTRAN=y
-CT_CC_SUPPORT_JAVA=y
-CT_CC_SUPPORT_ADA=y
-CT_CC_SUPPORT_OBJC=y
-CT_CC_SUPPORT_OBJCXX=y
-
-#
-# Additional supported languages:
-#
-CT_CC_LANG_CXX=y
-# CT_CC_LANG_JAVA is not set
-# CT_CC_LANG_ADA is not set
-# CT_CC_LANG_OBJC is not set
-# CT_CC_LANG_OBJCXX is not set
-CT_CC_LANG_OTHERS=""
-
-#
-# gcc other options
-#
-CT_CC_ENABLE_CXX_FLAGS=""
-CT_CC_CORE_EXTRA_CONFIG_ARRAY=""
-CT_CC_EXTRA_CONFIG_ARRAY=""
-CT_CC_STATIC_LIBSTDCXX=y
+CT_CC_GCC_ENABLE_CXX_FLAGS=""
+CT_CC_GCC_CORE_EXTRA_CONFIG_ARRAY=""
+CT_CC_GCC_EXTRA_CONFIG_ARRAY=""
+CT_CC_GCC_EXTRA_ENV_ARRAY=""
+CT_CC_GCC_STATIC_LIBSTDCXX=y
 # CT_CC_GCC_SYSTEM_ZLIB is not set
 
 #
@@ -339,6 +346,8 @@ CT_CC_GCC_ENABLE_TARGET_OPTSPACE=y
 # CT_CC_GCC_LIBMUDFLAP is not set
 # CT_CC_GCC_LIBGOMP is not set
 # CT_CC_GCC_LIBSSP is not set
+# CT_CC_GCC_LIBQUADMATH is not set
+# CT_CC_GCC_LIBSANITIZER is not set
 
 #
 # Misc. obscure options.
@@ -348,31 +357,33 @@ CT_CC_CXA_ATEXIT=y
 CT_CC_GCC_SJLJ_EXCEPTIONS=m
 CT_CC_GCC_LDBL_128=m
 # CT_CC_GCC_BUILD_ID is not set
+CT_CC_GCC_LNK_HASH_STYLE_DEFAULT=y
+# CT_CC_GCC_LNK_HASH_STYLE_SYSV is not set
+# CT_CC_GCC_LNK_HASH_STYLE_GNU is not set
+# CT_CC_GCC_LNK_HASH_STYLE_BOTH is not set
+CT_CC_GCC_LNK_HASH_STYLE=""
+CT_CC_GCC_DEC_FLOAT_AUTO=y
+# CT_CC_GCC_DEC_FLOAT_BID is not set
+# CT_CC_GCC_DEC_FLOAT_DPD is not set
+# CT_CC_GCC_DEC_FLOATS_NO is not set
+CT_CC_SUPPORT_CXX=y
+CT_CC_SUPPORT_FORTRAN=y
+CT_CC_SUPPORT_JAVA=y
+CT_CC_SUPPORT_ADA=y
+CT_CC_SUPPORT_OBJC=y
+CT_CC_SUPPORT_OBJCXX=y
+CT_CC_SUPPORT_GOLANG=y
 
 #
-# C-library
-#
-CT_LIBC="mingw"
-CT_LIBC_mingw=y
-CT_LIBC_eglibc_AVAILABLE=y
-CT_LIBC_glibc_AVAILABLE=y
-CT_LIBC_mingw_AVAILABLE=y
-CT_WINAPI_V_2_0_7=y
-# CT_WINAPI_V_select is not set
-CT_WINAPI_VERSION="2.0.7"
-CT_MINGW_DIRECTX=y
-CT_MINGW_DDK=y
-CT_LIBC_newlib_AVAILABLE=y
-CT_LIBC_none_AVAILABLE=y
-CT_LIBC_uClibc_AVAILABLE=y
-CT_LIBC_SUPPORT_THREADS_ANY=y
-CT_LIBC_SUPPORT_WIN32THREADS=y
-CT_THREADS="win32"
-
-#
-# Common C library options
+# Additional supported languages:
 #
-CT_THREADS_WIN32THREADS=y
+CT_CC_LANG_CXX=y
+# CT_CC_LANG_JAVA is not set
+# CT_CC_LANG_ADA is not set
+# CT_CC_LANG_OBJC is not set
+# CT_CC_LANG_OBJCXX is not set
+# CT_CC_LANG_GOLANG is not set
+CT_CC_LANG_OTHERS=""
 
 #
 # Debug facilities
@@ -390,12 +401,12 @@ CT_COMPLIBS_NEEDED=y
 CT_GMP_NEEDED=y
 CT_MPFR_NEEDED=y
 CT_MPC_NEEDED=y
-CT_LIBELF_NEEDED=y
 CT_COMPLIBS=y
 CT_GMP=y
 CT_MPFR=y
 CT_MPC=y
-CT_LIBELF=y
+# CT_GMP_V_6_0_0 is not set
+# CT_GMP_V_5_1_3 is not set
 # CT_GMP_V_5_1_1 is not set
 # CT_GMP_V_5_0_2 is not set
 CT_GMP_V_5_0_1=y
@@ -403,6 +414,7 @@ CT_GMP_V_5_0_1=y
 # CT_GMP_V_4_3_1 is not set
 # CT_GMP_V_4_3_0 is not set
 CT_GMP_VERSION="5.0.1"
+# CT_MPFR_V_3_1_3 is not set
 # CT_MPFR_V_3_1_2 is not set
 # CT_MPFR_V_3_1_0 is not set
 # CT_MPFR_V_3_0_1 is not set
@@ -411,16 +423,8 @@ CT_MPFR_V_2_4_2=y
 # CT_MPFR_V_2_4_1 is not set
 # CT_MPFR_V_2_4_0 is not set
 CT_MPFR_VERSION="2.4.2"
-# CT_PPL_V_0_11_2 is not set
-# CT_PPL_V_0_11_1 is not set
-# CT_PPL_V_0_11 is not set
-# CT_PPL_V_0_10_2 is not set
-# CT_CLOOG_V_0_15_11 is not set
-# CT_CLOOG_V_0_15_10 is not set
-# CT_CLOOG_V_0_15_9 is not set
-# CT_CLOOG_V_0_15_8 is not set
-# CT_CLOOG_V_0_15_7 is not set
-# CT_CLOOG_V_0_15_6 is not set
+# CT_MPC_V_1_0_3 is not set
+# CT_MPC_V_1_0_2 is not set
 # CT_MPC_V_1_0_1 is not set
 # CT_MPC_V_1_0 is not set
 CT_MPC_V_0_9=y
@@ -428,9 +432,6 @@ CT_MPC_V_0_9=y
 # CT_MPC_V_0_8_1 is not set
 # CT_MPC_V_0_7 is not set
 CT_MPC_VERSION="0.9"
-CT_LIBELF_V_0_8_13=y
-# CT_LIBELF_V_0_8_12 is not set
-CT_LIBELF_VERSION="0.8.13"
 
 #
 # Companion libraries common options
diff --git a/Engine/Build/BatchFiles/Linux/Toolchain_x86-64/win64-host.config b/Engine/Build/BatchFiles/Linux/Toolchain_x86-64/win64-host.config
--- a/Engine/Build/BatchFiles/Linux/Toolchain_x86-64/win64-host.config
+++ b/Engine/Build/BatchFiles/Linux/Toolchain_x86-64/win64-host.config
@@ -1,9 +1,9 @@
 #
-# Automatically generated make config: don't edit
-# crosstool-NG hg+default-529a71ea091e Configuration
-# Thu Feb  6 11:32:50 2014
+# Automatically generated file; DO NOT EDIT.
+# Crosstool-NG Configuration
 #
 CT_CONFIGURE_has_xz=y
+CT_CONFIGURE_has_svn=y
 CT_MODULES=y
 
 #
@@ -30,7 +30,8 @@
 CT_RM_RF_PREFIX_DIR=y
 CT_REMOVE_DOCS=y
 CT_INSTALL_DIR_RO=y
-CT_STRIP_ALL_TOOLCHAIN_EXECUTABLES=y
+CT_STRIP_HOST_TOOLCHAIN_EXECUTABLES=y
+# CT_STRIP_TARGET_TOOLCHAIN_EXECUTABLES is not set
 
 #
 # Downloading
@@ -107,28 +108,30 @@
 CT_TARGET_LDFLAGS=""
 # CT_ARCH_alpha is not set
 # CT_ARCH_arm is not set
-# CT_ARCH_avr32 is not set
-# CT_ARCH_blackfin is not set
+# CT_ARCH_avr is not set
 # CT_ARCH_m68k is not set
 # CT_ARCH_microblaze is not set
 # CT_ARCH_mips is not set
+# CT_ARCH_nios2 is not set
 # CT_ARCH_powerpc is not set
 # CT_ARCH_s390 is not set
 # CT_ARCH_sh is not set
 # CT_ARCH_sparc is not set
 CT_ARCH_x86=y
+# CT_ARCH_xtensa is not set
 CT_ARCH_alpha_AVAILABLE=y
 CT_ARCH_arm_AVAILABLE=y
-CT_ARCH_avr32_AVAILABLE=y
-CT_ARCH_blackfin_AVAILABLE=y
+CT_ARCH_avr_AVAILABLE=y
 CT_ARCH_m68k_AVAILABLE=y
 CT_ARCH_microblaze_AVAILABLE=y
 CT_ARCH_mips_AVAILABLE=y
+CT_ARCH_nios2_AVAILABLE=y
 CT_ARCH_powerpc_AVAILABLE=y
 CT_ARCH_s390_AVAILABLE=y
 CT_ARCH_sh_AVAILABLE=y
 CT_ARCH_sparc_AVAILABLE=y
 CT_ARCH_x86_AVAILABLE=y
+CT_ARCH_xtensa_AVAILABLE=y
 CT_ARCH_SUFFIX=""
 
 #
@@ -184,8 +187,8 @@
 #
 # Host system
 #
-CT_HOST="x86_64-pc-mingw32"
-CT_HOST_PREFIX="$HOME/x-tools/x86_64-pc-mingw32/bin/x86_64-pc-mingw32-"
+CT_HOST="x86_64-w64-mingw32"
+CT_HOST_PREFIX="$HOME/x-tools/x86_64-w64-mingw32/bin/x86_64-w64-mingw32-"
 CT_HOST_SUFFIX=""
 
 #
@@ -198,33 +201,22 @@
 #
 CT_KERNEL_SUPPORTS_SHARED_LIBS=y
 CT_KERNEL="linux"
-CT_KERNEL_VERSION="2.6.27.62"
+CT_KERNEL_VERSION="3.10.93"
 # CT_KERNEL_bare_metal is not set
 CT_KERNEL_linux=y
 # CT_KERNEL_windows is not set
 CT_KERNEL_bare_metal_AVAILABLE=y
 CT_KERNEL_linux_AVAILABLE=y
+# CT_KERNEL_V_4_3 is not set
+# CT_KERNEL_V_4_2 is not set
+# CT_KERNEL_V_4_1 is not set
+# CT_KERNEL_V_3_18 is not set
+# CT_KERNEL_V_3_14 is not set
 # CT_KERNEL_V_3_12 is not set
-# CT_KERNEL_V_3_11 is not set
-# CT_KERNEL_V_3_10 is not set
-# CT_KERNEL_V_3_9 is not set
-# CT_KERNEL_V_3_8 is not set
-# CT_KERNEL_V_3_7 is not set
-# CT_KERNEL_V_3_6 is not set
-# CT_KERNEL_V_3_5 is not set
+CT_KERNEL_V_3_10=y
 # CT_KERNEL_V_3_4 is not set
-# CT_KERNEL_V_3_3 is not set
 # CT_KERNEL_V_3_2 is not set
-# CT_KERNEL_V_3_1 is not set
-# CT_KERNEL_V_3_0 is not set
-# CT_KERNEL_V_2_6_39 is not set
-# CT_KERNEL_V_2_6_38 is not set
-# CT_KERNEL_V_2_6_37 is not set
-# CT_KERNEL_V_2_6_36 is not set
-# CT_KERNEL_V_2_6_33 is not set
 # CT_KERNEL_V_2_6_32 is not set
-# CT_KERNEL_V_2_6_31 is not set
-CT_KERNEL_V_2_6_27=y
 # CT_KERNEL_LINUX_CUSTOM is not set
 CT_KERNEL_windows_AVAILABLE=y
 
@@ -252,7 +244,10 @@
 #
 # GNU binutils
 #
-CT_BINUTILS_V_2_24=y
+# CT_CC_BINUTILS_SHOW_LINARO is not set
+CT_BINUTILS_V_2_25_1=y
+# CT_BINUTILS_V_2_25 is not set
+# CT_BINUTILS_V_2_24 is not set
 # CT_BINUTILS_V_2_23_2 is not set
 # CT_BINUTILS_V_2_23_1 is not set
 # CT_BINUTILS_V_2_22 is not set
@@ -262,7 +257,9 @@
 # CT_BINUTILS_V_2_19_1a is not set
 # CT_BINUTILS_V_2_18a is not set
 # CT_BINUTILS_CUSTOM is not set
-CT_BINUTILS_VERSION="2.24"
+CT_BINUTILS_VERSION="2.25.1"
+CT_BINUTILS_2_25_1_or_later=y
+CT_BINUTILS_2_25_or_later=y
 CT_BINUTILS_2_24_or_later=y
 CT_BINUTILS_2_23_or_later=y
 CT_BINUTILS_2_22_or_later=y
@@ -277,8 +274,6 @@
 CT_BINUTILS_HAS_PKGVERSION_BUGURL=y
 CT_BINUTILS_FORCE_LD_BFD=y
 CT_BINUTILS_LINKER_LD=y
-# CT_BINUTILS_LINKER_LD_GOLD is not set
-# CT_BINUTILS_LINKER_GOLD_LD is not set
 CT_BINUTILS_LINKERS_LIST="ld"
 CT_BINUTILS_LINKER_DEFAULT="bfd"
 CT_BINUTILS_EXTRA_CONFIG_ARRAY=""
@@ -289,54 +284,101 @@
 #
 
 #
+# C-library
+#
+CT_LIBC="glibc"
+CT_LIBC_VERSION="2.17"
+CT_LIBC_glibc=y
+# CT_LIBC_musl is not set
+# CT_LIBC_uClibc is not set
+CT_LIBC_avr_libc_AVAILABLE=y
+CT_LIBC_glibc_AVAILABLE=y
+CT_THREADS="nptl"
+# CT_CC_GLIBC_SHOW_LINARO is not set
+# CT_LIBC_GLIBC_V_2_22 is not set
+# CT_LIBC_GLIBC_V_2_21 is not set
+# CT_LIBC_GLIBC_V_2_20 is not set
+# CT_LIBC_GLIBC_V_2_19 is not set
+# CT_LIBC_GLIBC_V_2_18 is not set
+CT_LIBC_GLIBC_V_2_17=y
+# CT_LIBC_GLIBC_V_2_16_0 is not set
+# CT_LIBC_GLIBC_V_2_15 is not set
+# CT_LIBC_GLIBC_V_2_14_1 is not set
+# CT_LIBC_GLIBC_V_2_14 is not set
+# CT_LIBC_GLIBC_V_2_13 is not set
+# CT_LIBC_GLIBC_V_2_12_2 is not set
+# CT_LIBC_GLIBC_V_2_12_1 is not set
+# CT_LIBC_GLIBC_V_2_11_1 is not set
+# CT_LIBC_GLIBC_V_2_11 is not set
+# CT_LIBC_GLIBC_V_2_10_1 is not set
+# CT_LIBC_GLIBC_V_2_9 is not set
+# CT_LIBC_GLIBC_V_2_8 is not set
+# CT_LIBC_GLIBC_CUSTOM is not set
+CT_LIBC_GLIBC_2_17_or_later=y
+CT_LIBC_mingw_AVAILABLE=y
+CT_LIBC_musl_AVAILABLE=y
+CT_LIBC_newlib_AVAILABLE=y
+CT_LIBC_none_AVAILABLE=y
+CT_LIBC_uClibc_AVAILABLE=y
+CT_LIBC_SUPPORT_THREADS_ANY=y
+CT_LIBC_SUPPORT_THREADS_NATIVE=y
+
+#
+# Common C library options
+#
+CT_THREADS_NATIVE=y
+CT_LIBC_XLDD=y
+
+#
+# glibc other options
+#
+# CT_LIBC_GLIBC_PORTS_EXTERNAL is not set
+CT_LIBC_glibc_familly=y
+CT_LIBC_GLIBC_EXTRA_CONFIG_ARRAY=""
+CT_LIBC_GLIBC_CONFIGPARMS=""
+CT_LIBC_GLIBC_EXTRA_CFLAGS=""
+CT_LIBC_EXTRA_CC_ARGS=""
+# CT_LIBC_ENABLE_FORTIFIED_BUILD is not set
+# CT_LIBC_DISABLE_VERSIONING is not set
+CT_LIBC_OLDEST_ABI=""
+CT_LIBC_GLIBC_FORCE_UNWIND=y
+# CT_LIBC_GLIBC_USE_PORTS is not set
+CT_LIBC_ADDONS_LIST=""
+# CT_LIBC_LOCALES is not set
+# CT_LIBC_GLIBC_KERNEL_VERSION_NONE is not set
+CT_LIBC_GLIBC_KERNEL_VERSION_AS_HEADERS=y
+# CT_LIBC_GLIBC_KERNEL_VERSION_CHOSEN is not set
+CT_LIBC_GLIBC_MIN_KERNEL="3.10.93"
+
+#
 # C compiler
 #
 CT_CC="gcc"
-CT_CC_VERSION="4.8.1"
 CT_CC_CORE_PASSES_NEEDED=y
+CT_CC_CORE_PASS_1_NEEDED=y
+CT_CC_CORE_PASS_2_NEEDED=y
 CT_CC_gcc=y
 # CT_CC_GCC_SHOW_LINARO is not set
-# CT_CC_V_4_8_2 is not set
-CT_CC_V_4_8_1=y
-# CT_CC_V_4_8_0 is not set
-# CT_CC_V_4_7_3 is not set
-# CT_CC_V_4_7_2 is not set
-# CT_CC_V_4_7_1 is not set
-# CT_CC_V_4_7_0 is not set
-# CT_CC_V_4_6_4 is not set
-# CT_CC_V_4_6_3 is not set
-# CT_CC_V_4_6_2 is not set
-# CT_CC_V_4_6_1 is not set
-# CT_CC_V_4_6_0 is not set
-# CT_CC_V_4_5_3 is not set
-# CT_CC_V_4_5_2 is not set
-# CT_CC_V_4_5_1 is not set
-# CT_CC_V_4_5_0 is not set
-# CT_CC_V_4_4_7 is not set
-# CT_CC_V_4_4_6 is not set
-# CT_CC_V_4_4_5 is not set
-# CT_CC_V_4_4_4 is not set
-# CT_CC_V_4_4_3 is not set
-# CT_CC_V_4_4_2 is not set
-# CT_CC_V_4_4_1 is not set
-# CT_CC_V_4_4_0 is not set
-# CT_CC_V_4_3_6 is not set
-# CT_CC_V_4_3_5 is not set
-# CT_CC_V_4_3_4 is not set
-# CT_CC_V_4_3_3 is not set
-# CT_CC_V_4_3_2 is not set
-# CT_CC_V_4_3_1 is not set
-# CT_CC_V_4_2_4 is not set
-# CT_CC_V_4_2_2 is not set
-# CT_CC_CUSTOM is not set
+CT_CC_GCC_V_5_2_0=y
+# CT_CC_GCC_V_4_9_3 is not set
+# CT_CC_GCC_V_4_8_5 is not set
+# CT_CC_GCC_V_4_7_4 is not set
+# CT_CC_GCC_V_4_6_4 is not set
+# CT_CC_GCC_V_4_5_4 is not set
+# CT_CC_GCC_V_4_4_7 is not set
+# CT_CC_GCC_V_4_3_6 is not set
+# CT_CC_GCC_V_4_2_4 is not set
+# CT_CC_GCC_CUSTOM is not set
 CT_CC_GCC_4_2_or_later=y
 CT_CC_GCC_4_3_or_later=y
 CT_CC_GCC_4_4_or_later=y
 CT_CC_GCC_4_5_or_later=y
 CT_CC_GCC_4_6_or_later=y
 CT_CC_GCC_4_7_or_later=y
-CT_CC_GCC_4_8=y
 CT_CC_GCC_4_8_or_later=y
+CT_CC_GCC_4_9_or_later=y
+CT_CC_GCC_5=y
+CT_CC_GCC_5_or_later=y
 CT_CC_GCC_HAS_GRAPHITE=y
 # CT_CC_GCC_USE_GRAPHITE is not set
 CT_CC_GCC_HAS_LTO=y
@@ -347,33 +389,14 @@
 CT_CC_GCC_USE_GMP_MPFR=y
 CT_CC_GCC_USE_MPC=y
 CT_CC_GCC_HAS_LIBQUADMATH=y
+CT_CC_GCC_HAS_LIBSANITIZER=y
+CT_CC_GCC_VERSION="5.2.0"
 # CT_CC_LANG_FORTRAN is not set
-CT_CC_SUPPORT_CXX=y
-CT_CC_SUPPORT_FORTRAN=y
-CT_CC_SUPPORT_JAVA=y
-CT_CC_SUPPORT_ADA=y
-CT_CC_SUPPORT_OBJC=y
-CT_CC_SUPPORT_OBJCXX=y
-CT_CC_SUPPORT_GOLANG=y
-
-#
-# Additional supported languages:
-#
-CT_CC_LANG_CXX=y
-# CT_CC_LANG_JAVA is not set
-# CT_CC_LANG_ADA is not set
-# CT_CC_LANG_OBJC is not set
-# CT_CC_LANG_OBJCXX is not set
-# CT_CC_LANG_GOLANG is not set
-CT_CC_LANG_OTHERS=""
-
-#
-# gcc other options
-#
-CT_CC_ENABLE_CXX_FLAGS=""
-CT_CC_CORE_EXTRA_CONFIG_ARRAY=""
-CT_CC_EXTRA_CONFIG_ARRAY=""
-CT_CC_STATIC_LIBSTDCXX=y
+CT_CC_GCC_ENABLE_CXX_FLAGS=""
+CT_CC_GCC_CORE_EXTRA_CONFIG_ARRAY=""
+CT_CC_GCC_EXTRA_CONFIG_ARRAY=""
+CT_CC_GCC_EXTRA_ENV_ARRAY=""
+CT_CC_GCC_STATIC_LIBSTDCXX=y
 # CT_CC_GCC_SYSTEM_ZLIB is not set
 
 #
@@ -388,6 +411,7 @@
 # CT_CC_GCC_LIBGOMP is not set
 # CT_CC_GCC_LIBSSP is not set
 # CT_CC_GCC_LIBQUADMATH is not set
+# CT_CC_GCC_LIBSANITIZER is not set
 
 #
 # Misc. obscure options.
@@ -406,65 +430,24 @@
 # CT_CC_GCC_DEC_FLOAT_BID is not set
 # CT_CC_GCC_DEC_FLOAT_DPD is not set
 # CT_CC_GCC_DEC_FLOATS_NO is not set
+CT_CC_SUPPORT_CXX=y
+CT_CC_SUPPORT_FORTRAN=y
+CT_CC_SUPPORT_JAVA=y
+CT_CC_SUPPORT_ADA=y
+CT_CC_SUPPORT_OBJC=y
+CT_CC_SUPPORT_OBJCXX=y
+CT_CC_SUPPORT_GOLANG=y
 
 #
-# C-library
+# Additional supported languages:
 #
-CT_LIBC="glibc"
-CT_LIBC_VERSION="2.12.2"
-# CT_LIBC_eglibc is not set
-CT_LIBC_glibc=y
-# CT_LIBC_uClibc is not set
-CT_LIBC_eglibc_AVAILABLE=y
-CT_LIBC_glibc_AVAILABLE=y
-# CT_LIBC_GLIBC_V_2_18 is not set
-# CT_LIBC_GLIBC_V_2_17 is not set
-# CT_LIBC_GLIBC_V_2_16_0 is not set
-# CT_LIBC_GLIBC_V_2_15 is not set
-# CT_LIBC_GLIBC_V_2_14_1 is not set
-# CT_LIBC_GLIBC_V_2_14 is not set
-# CT_LIBC_GLIBC_V_2_13 is not set
-CT_LIBC_GLIBC_V_2_12_2=y
-# CT_LIBC_GLIBC_V_2_12_1 is not set
-# CT_LIBC_GLIBC_V_2_11_1 is not set
-# CT_LIBC_GLIBC_V_2_11 is not set
-# CT_LIBC_GLIBC_V_2_10_1 is not set
-# CT_LIBC_GLIBC_V_2_9 is not set
-# CT_LIBC_GLIBC_V_2_8 is not set
-CT_LIBC_mingw_AVAILABLE=y
-CT_LIBC_newlib_AVAILABLE=y
-CT_LIBC_none_AVAILABLE=y
-CT_LIBC_uClibc_AVAILABLE=y
-CT_LIBC_SUPPORT_THREADS_ANY=y
-CT_LIBC_SUPPORT_NPTL=y
-CT_THREADS="nptl"
+CT_CC_LANG_CXX=y
+# CT_CC_LANG_JAVA is not set
+# CT_CC_LANG_ADA is not set
+# CT_CC_LANG_OBJC is not set
+# CT_CC_LANG_OBJCXX is not set
+# CT_CC_LANG_GOLANG is not set
+CT_CC_LANG_OTHERS=""
 
 #
-# Common C library options
-#
-CT_THREADS_NPTL=y
-CT_LIBC_XLDD=y
-CT_LIBC_GLIBC_PORTS_EXTERNAL=y
-CT_LIBC_glibc_familly=y
-CT_LIBC_GLIBC_EXTRA_CONFIG_ARRAY=""
-CT_LIBC_GLIBC_CONFIGPARMS=""
-CT_LIBC_GLIBC_EXTRA_CFLAGS=""
-CT_LIBC_EXTRA_CC_ARGS=""
-# CT_LIBC_ENABLE_FORTIFIED_BUILD is not set
-# CT_LIBC_DISABLE_VERSIONING is not set
-CT_LIBC_OLDEST_ABI=""
-CT_LIBC_GLIBC_FORCE_UNWIND=y
-# CT_LIBC_GLIBC_USE_PORTS is not set
-CT_LIBC_ADDONS_LIST=""
-# CT_LIBC_LOCALES is not set
-# CT_LIBC_GLIBC_KERNEL_VERSION_NONE is not set
-CT_LIBC_GLIBC_KERNEL_VERSION_AS_HEADERS=y
-# CT_LIBC_GLIBC_KERNEL_VERSION_CHOSEN is not set
-CT_LIBC_GLIBC_MIN_KERNEL="2.6.27.62"
-
-#
-# glibc other options
-#
-
-#
 # Debug facilities
@@ -479,13 +462,23 @@
 # Companion libraries
 #
 CT_COMPLIBS_NEEDED=y
+CT_LIBICONV_NEEDED=y
+CT_GETTEXT_NEEDED=y
 CT_GMP_NEEDED=y
 CT_MPFR_NEEDED=y
 CT_MPC_NEEDED=y
 CT_COMPLIBS=y
+CT_LIBICONV=y
+CT_GETTEXT=y
 CT_GMP=y
 CT_MPFR=y
 CT_MPC=y
+CT_LIBICONV_V_1_14=y
+CT_LIBICONV_VERSION="1.14"
+CT_GETTEXT_V_0_19_6=y
+CT_GETTEXT_VERSION="0.19.6"
+# CT_GMP_V_6_0_0 is not set
+# CT_GMP_V_5_1_3 is not set
 # CT_GMP_V_5_1_1 is not set
 # CT_GMP_V_5_0_2 is not set
 CT_GMP_V_5_0_1=y
@@ -493,6 +486,7 @@
 # CT_GMP_V_4_3_1 is not set
 # CT_GMP_V_4_3_0 is not set
 CT_GMP_VERSION="5.0.1"
+# CT_MPFR_V_3_1_3 is not set
 # CT_MPFR_V_3_1_2 is not set
 # CT_MPFR_V_3_1_0 is not set
 # CT_MPFR_V_3_0_1 is not set
@@ -501,6 +495,8 @@
 # CT_MPFR_V_2_4_1 is not set
 # CT_MPFR_V_2_4_0 is not set
 CT_MPFR_VERSION="2.4.2"
+# CT_MPC_V_1_0_3 is not set
+# CT_MPC_V_1_0_2 is not set
 # CT_MPC_V_1_0_1 is not set
 # CT_MPC_V_1_0 is not set
 CT_MPC_V_0_9=y
diff --git a/Engine/Config/BaseDeviceProfiles.ini b/Engine/Config/BaseDeviceProfiles.ini
--- a/Engine/Config/BaseDeviceProfiles.ini
+++ b/Engine/Config/BaseDeviceProfiles.ini
@@ -377,7 +377,6 @@ BaseProfileName=Android_Adreno330
 [Android_Adreno5xx DeviceProfile]
 DeviceType=Android
 BaseProfileName=Android_High
-+CVars=r.Vulkan.UseGLSL=2
 
 [Android_Adreno4xx DeviceProfile]
 DeviceType=Android
diff --git a/Engine/Plugins/Editor/MobileLauncherProfileWizard/Source/MobileLauncherProfileWizard/Private/AndroidProfileWizard.cpp b/Engine/Plugins/Editor/MobileLauncherProfileWizard/Source/MobileLauncherProfileWizard/Private/AndroidProfileWizard.cpp
--- a/Engine/Plugins/Editor/MobileLauncherProfileWizard/Source/MobileLauncherProfileWizard/Private/AndroidProfileWizard.cpp
+++ b/Engine/Plugins/Editor/MobileLauncherProfileWizard/Source/MobileLauncherProfileWizard/Private/AndroidProfileWizard.cpp
@@ -83,7 +83,8 @@ static void SetupAndroidDLCProfile(ILauncherProfileRef& DLCProfile, const FProfi
 	DLCProfile->SetBuildUAT(true);
 	// App build configuration
 	DLCProfile->SetBuildGame(false);
-	
+	DLCProfile->SetBuildConfiguration(Params.BuildConfiguration);
+		
 	//// Cooking
 	DLCProfile->SetCookMode(ELauncherProfileCookModes::ByTheBook);
 	DLCProfile->SetCookConfiguration(Params.BuildConfiguration);
diff --git a/Engine/Plugins/Editor/MobileLauncherProfileWizard/Source/MobileLauncherProfileWizard/Private/IOSProfileWizard.cpp b/Engine/Plugins/Editor/MobileLauncherProfileWizard/Source/MobileLauncherProfileWizard/Private/IOSProfileWizard.cpp
--- a/Engine/Plugins/Editor/MobileLauncherProfileWizard/Source/MobileLauncherProfileWizard/Private/IOSProfileWizard.cpp
+++ b/Engine/Plugins/Editor/MobileLauncherProfileWizard/Source/MobileLauncherProfileWizard/Private/IOSProfileWizard.cpp
@@ -82,6 +82,7 @@ static void SetupIOSDLCProfile(ILauncherProfileRef& DLCProfile, const FProfilePa
 	DLCProfile->SetBuildUAT(true);
 	// App build configuration
 	DLCProfile->SetBuildGame(false);
+	DLCProfile->SetBuildConfiguration(Params.BuildConfiguration);
 	
 	//// Cooking
 	DLCProfile->SetCookMode(ELauncherProfileCookModes::ByTheBook);
diff --git a/Engine/Plugins/Editor/MobileLauncherProfileWizard/Source/MobileLauncherProfileWizard/Private/SProfileWizardUI.cpp b/Engine/Plugins/Editor/MobileLauncherProfileWizard/Source/MobileLauncherProfileWizard/Private/SProfileWizardUI.cpp
--- a/Engine/Plugins/Editor/MobileLauncherProfileWizard/Source/MobileLauncherProfileWizard/Private/SProfileWizardUI.cpp
+++ b/Engine/Plugins/Editor/MobileLauncherProfileWizard/Source/MobileLauncherProfileWizard/Private/SProfileWizardUI.cpp
@@ -700,16 +700,12 @@ void SProfileWizardUI::CacheProjectMapList()
 {
 	ProjectMapList.Reset();
 
-	FString BaseProjectPath = ProjectPath;
-	if (FPaths::MakePathRelativeTo(BaseProjectPath, *FPaths::RootDir()))
-	{
-		BaseProjectPath = FPaths::GetPath(BaseProjectPath);
-		TArray<FString> AvailableMaps = FGameProjectHelper::GetAvailableMaps(BaseProjectPath, false, true);
+	FString BaseProjectPath = FPaths::GetPath(ProjectPath);
+	TArray<FString> AvailableMaps = FGameProjectHelper::GetAvailableMaps(BaseProjectPath, false, true);
 
-		for (const FString& MapName : AvailableMaps)
-		{
-			ProjectMapList.Add(MakeShareable(new FString(MapName)));
-		}
+	for (const FString& MapName : AvailableMaps)
+	{
+		ProjectMapList.Add(MakeShareable(new FString(MapName)));
 	}
 }
 
diff --git a/Engine/Plugins/Messaging/TcpMessaging/Source/TcpMessaging/Private/TcpMessagingModule.cpp b/Engine/Plugins/Messaging/TcpMessaging/Source/TcpMessaging/Private/TcpMessagingModule.cpp
--- a/Engine/Plugins/Messaging/TcpMessaging/Source/TcpMessaging/Private/TcpMessagingModule.cpp
+++ b/Engine/Plugins/Messaging/TcpMessaging/Source/TcpMessaging/Private/TcpMessagingModule.cpp
@@ -37,7 +37,7 @@ public:
 			UTcpMessagingSettings* Settings = GetMutableDefault<UTcpMessagingSettings>();
 
 			// general information
-			Ar.Logf(TEXT("Protocol Version: %i"), TCP_MESSAGING_TRANSPORT_PROTOCOL_VERSION);
+			Ar.Logf(TEXT("Protocol Version: %d"), (int32)ETcpMessagingVersion::LatestVersion);
 
 			// bridge status
 			if (MessageBridge.IsValid())
diff --git a/Engine/Plugins/Messaging/TcpMessaging/Source/TcpMessaging/Private/TcpMessagingPrivatePCH.h b/Engine/Plugins/Messaging/TcpMessaging/Source/TcpMessaging/Private/TcpMessagingPrivatePCH.h
--- a/Engine/Plugins/Messaging/TcpMessaging/Source/TcpMessaging/Private/TcpMessagingPrivatePCH.h
+++ b/Engine/Plugins/Messaging/TcpMessaging/Source/TcpMessaging/Private/TcpMessagingPrivatePCH.h
@@ -29,12 +29,26 @@ DECLARE_LOG_CATEGORY_EXTERN(LogTcpMessaging, Log, All);
 /** Defines the desired size of socket receive buffers (in bytes). */
 #define TCP_MESSAGING_RECEIVE_BUFFER_SIZE 2 * 1024 * 1024
 
-/** Defines the protocol version of the TCP message transport. */
-#define TCP_MESSAGING_TRANSPORT_PROTOCOL_VERSION 1
-
 /** Defines a magic number for the the TCP message transport. */
 #define TCP_MESSAGING_TRANSPORT_PROTOCOL_MAGIC 0x45504943
 
+/** Defines the protocol version of the TCP message transport. */
+namespace ETcpMessagingVersion
+{
+	enum Type
+	{
+		Initial,
+		ChangedMessageLengthToInt32,
+
+		// -----<new versions can be added before this line>-------------------------------------------------
+		// - this needs to be the last line (see note below)
+		VersionPlusOne,
+		LatestVersion = VersionPlusOne - 1,
+		// bump this when break need to break compatibility.
+		OldestSupportedVersion = ChangedMessageLengthToInt32
+	};
+}
+
 /* Private includes
  *****************************************************************************/
 
diff --git a/Engine/Plugins/Messaging/TcpMessaging/Source/TcpMessaging/Private/Transport/TcpMessageTransport.cpp b/Engine/Plugins/Messaging/TcpMessaging/Source/TcpMessaging/Private/Transport/TcpMessageTransport.cpp
--- a/Engine/Plugins/Messaging/TcpMessaging/Source/TcpMessaging/Private/Transport/TcpMessageTransport.cpp
+++ b/Engine/Plugins/Messaging/TcpMessaging/Source/TcpMessaging/Private/Transport/TcpMessageTransport.cpp
@@ -123,8 +123,11 @@ bool FTcpMessageTransport::TransportMessage(const IMessageContextRef& Context, c
 
 	if (Recipients.Num() == 0)
 	{
-		// broadcast the message to all connections
-		RecipientConnections = Connections;
+		// broadcast the message to all valid connections
+		RecipientConnections = Connections.FilterByPredicate([&](const TSharedPtr<FTcpMessageTransportConnection>& Connection) -> bool
+		{
+			return Connection->GetConnectionState() == FTcpMessageTransportConnection::STATE_Connected;
+		});
 	}
 	else
 	{
@@ -132,18 +135,20 @@ bool FTcpMessageTransport::TransportMessage(const IMessageContextRef& Context, c
 		for (auto& Recipient : Recipients)
 		{
 			TSharedPtr<FTcpMessageTransportConnection>* RecipientConnection = NodeConnectionMap.Find(Recipient);
-			if (RecipientConnection)
+			if (RecipientConnection && (*RecipientConnection)->GetConnectionState() == FTcpMessageTransportConnection::STATE_Connected)
 			{
 				RecipientConnections.AddUnique(*RecipientConnection);
 			}
 		}
 	}
 
-	if(RecipientConnections.Num() == 0)
+	if (RecipientConnections.Num() == 0)
 	{
 		return false;
 	}
 
+	UE_LOG(LogTcpMessaging, Verbose, TEXT("Transporting message '%s' to %d connections"), *Context->GetMessageType().ToString(), RecipientConnections.Num());
+
 	FTcpSerializedMessageRef SerializedMessage = MakeShareable(new FTcpSerializedMessage());
 
 	TGraphTask<FTcpSerializeMessageTask>::CreateTask().ConstructAndDispatchWhenReady(Context, SerializedMessage, RecipientConnections);
@@ -220,6 +225,7 @@ uint32 FTcpMessageTransport::Run()
 				FGuid SenderNodeId;
 				while (Connection->Receive(Message, SenderNodeId))
 				{
+					UE_LOG(LogTcpMessaging, Verbose, TEXT("Received message '%s'"), *Message->GetMessageType().ToString());
 					MessageReceivedDelegate.ExecuteIfBound(Message.ToSharedRef(), SenderNodeId);
 				}
 			}
diff --git a/Engine/Plugins/Messaging/TcpMessaging/Source/TcpMessaging/Private/Transport/TcpMessageTransportConnection.cpp b/Engine/Plugins/Messaging/TcpMessaging/Source/TcpMessaging/Private/Transport/TcpMessageTransportConnection.cpp
--- a/Engine/Plugins/Messaging/TcpMessaging/Source/TcpMessaging/Private/Transport/TcpMessageTransportConnection.cpp
+++ b/Engine/Plugins/Messaging/TcpMessaging/Source/TcpMessaging/Private/Transport/TcpMessageTransportConnection.cpp
@@ -16,23 +16,28 @@ struct FTcpMessageHeader
 
 	FTcpMessageHeader(const FGuid& InNodeId)
 	:	MagicNumber(TCP_MESSAGING_TRANSPORT_PROTOCOL_MAGIC)
-	,	Version(TCP_MESSAGING_TRANSPORT_PROTOCOL_VERSION)
+	,	Version(ETcpMessagingVersion::LatestVersion)
 	,	NodeId(InNodeId)
 	{}
 
-	bool IsValid()
+	bool IsValid() const
 	{
 		return
 			MagicNumber == TCP_MESSAGING_TRANSPORT_PROTOCOL_MAGIC &&
-			Version == TCP_MESSAGING_TRANSPORT_PROTOCOL_VERSION &&
+			Version == ETcpMessagingVersion::OldestSupportedVersion &&
 			NodeId.IsValid();
 	}
 
-	FGuid GetNodeId()
+	FGuid GetNodeId() const
 	{
 		return NodeId;
 	}
 
+	uint32 GetVersion() const
+	{
+		return Version;
+	}
+
 	// Serializer
 	friend FArchive& operator<<(FArchive& Ar, FTcpMessageHeader& H)
 	{
@@ -51,12 +56,14 @@ FTcpMessageTransportConnection::FTcpMessageTransportConnection(FSocket* InSocket
 	, LocalNodeId(FGuid::NewGuid())
 	, bSentHeader(false)
 	, bReceivedHeader(false)
+	, RemoteProtocolVersion(0)
 	, Socket(InSocket)
 	, Thread(nullptr)
 	, TotalBytesReceived(0)
 	, TotalBytesSent(0)
 	, bRun(false)
 	, ConnectionRetryDelay(InConnectionRetryDelay)
+	, RecvMessageDataRemaining(0)
 {
 	int32 NewSize = 0;
 	Socket->SetReceiveBufferSize(2*1024*1024, NewSize);
@@ -67,8 +74,11 @@ FTcpMessageTransportConnection::~FTcpMessageTransportConnection()
 {
 	if (Thread != nullptr)
 	{
-		bRun = false;
-		Thread->WaitForCompletion();
+		if(bRun)
+		{
+			bRun = false;
+			Thread->WaitForCompletion();
+		}
 		delete Thread;
 	}
 
@@ -122,10 +132,8 @@ uint32 FTcpMessageTransportConnection::Run()
 {
 	while (bRun)
 	{
-		// Check if we can are ready to send or receive, and if so try sending and receiving messages.
-		// Otherwise, check for a connection error.
-		if (((Socket->Wait(ESocketWaitConditions::WaitForReadOrWrite, FTimespan::Zero()) && (!ReceiveMessages() || !SendMessages())) ||
-			(Socket->GetConnectionState() == SCS_ConnectionError)) && bRun)
+		// Try sending and receiving messages and detect if they fail or if another connection error is reported.
+		if ((!ReceiveMessages() || !SendMessages() || Socket->GetConnectionState() == SCS_ConnectionError) && bRun)
 		{
 			// Disconnected. Reconnect if requested.
 			if (ConnectionRetryDelay > 0)
@@ -284,6 +292,7 @@ bool FTcpMessageTransportConnection::ReceiveMessages()
 			else
 			{
 				RemoteNodeId = MessageHeader.GetNodeId();
+				RemoteProtocolVersion = MessageHeader.GetVersion();
 				bReceivedHeader = true;
 				OpenedTime = FDateTime::UtcNow();
 				UpdateConnectionState(STATE_Connected);
@@ -296,62 +305,92 @@ bool FTcpMessageTransportConnection::ReceiveMessages()
 		}
 	}
 
-	// if we received a payload size...
-	while (Socket->HasPendingData(PendingDataSize) && (PendingDataSize >= sizeof(uint16)))
+	// keep going until we have no data.
+	for(;;)
 	{
 		int32 BytesRead = 0;
-
-		FArrayReader MessagesizeData = FArrayReader(true);
-		MessagesizeData.SetNumUninitialized(sizeof(uint16));
-
-		// ... read it from the stream without removing it...
-		if (!Socket->Recv(MessagesizeData.GetData(), sizeof(uint16), BytesRead, ESocketReceiveFlags::Peek))
+		// See if we're in the process of receiving a (large) message
+		if (RecvMessageDataRemaining == 0)
 		{
-			return false;
-		}
-
-		check(BytesRead == sizeof(uint16));
+			// no partial message. Try to receive the size of a message
+			if (!Socket->HasPendingData(PendingDataSize) || (PendingDataSize < sizeof(uint32)))
+			{
+				// no messages
+				return true;
+			}
 
-		uint16 Messagesize;
-		MessagesizeData << Messagesize;
+			FArrayReader MessagesizeData = FArrayReader(true);
+			MessagesizeData.SetNumUninitialized(sizeof(uint32));
 
-		// ... and if we received the complete payload...
-		if (Socket->HasPendingData(PendingDataSize) && (PendingDataSize >= sizeof(uint16) + Messagesize))
-		{
-			// ... then remove the payload size from the stream...
-			if (!Socket->Recv((uint8*)&Messagesize, sizeof(uint16), BytesRead))
+			// read message size from the stream
+			BytesRead = 0;
+			if (!Socket->Recv(MessagesizeData.GetData(), sizeof(uint32), BytesRead))
 			{
 				return false;
 			}
 
-			check(BytesRead == sizeof(uint16));
+			check(BytesRead == sizeof(uint32));
 			TotalBytesReceived += BytesRead;
 
-			// ... receive the payload
-			FArrayReaderPtr Payload = MakeShareable(new FArrayReader(true));
-			Payload->SetNumUninitialized(Messagesize);
+			// Setup variables to receive the message
+			MessagesizeData << RecvMessageDataRemaining;
 
-			if (!Socket->Recv(Payload->GetData(), Payload->Num(), BytesRead))
-			{
-				return false;
-			}
+			RecvMessageData = MakeShareable(new FArrayReader(true));
+			RecvMessageData->SetNumUninitialized(RecvMessageDataRemaining);
+			check(RecvMessageDataRemaining > 0);
+		}
 
-			check(BytesRead == Messagesize);
-			TotalBytesReceived += BytesRead;
+		BytesRead = 0;
+		if (!Socket->Recv(RecvMessageData->GetData() + RecvMessageData->Num() - RecvMessageDataRemaining, RecvMessageDataRemaining, BytesRead))
+		{
+			return false;
+		}
 
-			// @todo gmp: move message deserialization into an async task
-			FTcpDeserializedMessage* DeserializedMessage = new FTcpDeserializedMessage(nullptr);
-			if (DeserializedMessage->Deserialize(Payload))
+		if (BytesRead > 0)
+		{
+			check(BytesRead <= RecvMessageDataRemaining);
+			TotalBytesReceived += BytesRead;
+			RecvMessageDataRemaining -= BytesRead;
+			if (RecvMessageDataRemaining == 0)
 			{
-				Inbox.Enqueue(MakeShareable(DeserializedMessage));
+				// @todo gmp: move message deserialization into an async task
+				FTcpDeserializedMessage* DeserializedMessage = new FTcpDeserializedMessage(nullptr);
+				if (DeserializedMessage->Deserialize(RecvMessageData))
+				{
+					Inbox.Enqueue(MakeShareable(DeserializedMessage));
+				}
+				RecvMessageData.Reset();
 			}
 		}
 		else
 		{
-			break;
+			// no data
+			return true;
 		}
 	}
+}
+
+bool FTcpMessageTransportConnection::BlockingSend(const uint8* Data, int32 BytesToSend)
+{
+	int32 TotalBytes = BytesToSend;
+	while (BytesToSend > 0)
+	{
+		while (!Socket->Wait(ESocketWaitConditions::WaitForWrite, FTimespan(0, 0, 1)))
+		{
+			if (Socket->GetConnectionState() == SCS_ConnectionError)
+			{
+				return false;
+			}
+		}
 
+		int32 BytesSent = 0;
+		if (!Socket->Send(Data, BytesToSend, BytesSent))
+		{
+			return false;
+		}
+		BytesToSend -= BytesSent;
+		Data += BytesSent;
+	}
 	return true;
 }
 
@@ -373,14 +412,13 @@ bool FTcpMessageTransportConnection::SendMessages()
 		FTcpMessageHeader MessageHeader(LocalNodeId);
 		HeaderData << MessageHeader;
 
-		int32 BytesSent = 0;
-		if (!Socket->Send(HeaderData.GetData(), sizeof(FTcpMessageHeader), BytesSent))
+		if (!BlockingSend(HeaderData.GetData(), sizeof(FTcpMessageHeader)))
 		{
 			return false;
 		}
 
 		bSentHeader = true;
-		TotalBytesSent += BytesSent;
+		TotalBytesSent += sizeof(FTcpMessageHeader);
 	}
 	else
 	{
@@ -392,23 +430,23 @@ bool FTcpMessageTransportConnection::SendMessages()
 
 			// send the payload size
 			FArrayWriter MessagesizeData = FArrayWriter(true);
-			uint16 Messagesize = Payload.Num();
+			uint32 Messagesize = Payload.Num();
 			MessagesizeData << Messagesize;
 
-			if (!Socket->Send(MessagesizeData.GetData(), sizeof(uint16), BytesSent))
+			if (!BlockingSend(MessagesizeData.GetData(), sizeof(uint32)))
 			{
 				return false;
 			}
 
-			TotalBytesSent += BytesSent;
+			TotalBytesSent += sizeof(uint32);
 
 			// send the payload
-			if (!Socket->Send(Payload.GetData(), Payload.Num(), BytesSent))
+			if (!BlockingSend(Payload.GetData(), Payload.Num()))
 			{
 				return false;
 			}
 
-			TotalBytesSent += BytesSent;
+			TotalBytesSent += Payload.Num();
 
 			// return if the socket is no longer writable, or that was the last message
 			if (Outbox.IsEmpty() || !Socket->Wait(ESocketWaitConditions::WaitForWrite, FTimespan::Zero()))
diff --git a/Engine/Plugins/Messaging/TcpMessaging/Source/TcpMessaging/Private/Transport/TcpMessageTransportConnection.h b/Engine/Plugins/Messaging/TcpMessaging/Source/TcpMessaging/Private/Transport/TcpMessageTransportConnection.h
--- a/Engine/Plugins/Messaging/TcpMessaging/Source/TcpMessaging/Private/Transport/TcpMessageTransportConnection.h
+++ b/Engine/Plugins/Messaging/TcpMessaging/Source/TcpMessaging/Private/Transport/TcpMessageTransportConnection.h
@@ -159,7 +159,9 @@ private:
 	virtual void Exit() override;
 
 private:
-	
+	/** Try to send data, but if all data is not sent in one go, block on send until data is sent or an error occurs */
+	bool BlockingSend(const uint8* Data, int32 BytesToSend);
+
 	/** Connection state changed delegate */
 	FOnTcpMessageTransportConnectionStateChanged ConnectionStateChangedDelegate;
 
@@ -193,6 +195,9 @@ private:
 	/** Whether we've received the initial header from the remote end */
 	bool bReceivedHeader;
 
+	/** Peer's value for of ETcpMessagingVersion::LatestVersion */
+	uint32 RemoteProtocolVersion;
+
 	/** Holds the connection socket. */
 	FSocket* Socket;
 
@@ -210,5 +215,11 @@ private:
 
 	/** Delay before re-establishing connection if it drops, 0 disables */
 	int32 ConnectionRetryDelay;
+
+	/** Message data we're currently in the process of receiving, if any */
+	FArrayReaderPtr RecvMessageData;
+
+	/** The number of bytes of incoming message data we're still waiting on before we have a complete message */
+	int32 RecvMessageDataRemaining;
 };
 
diff --git a/Engine/Plugins/Runtime/GearVR/Source/GearVR/Private/GearVRRender.cpp b/Engine/Plugins/Runtime/GearVR/Source/GearVR/Private/GearVRRender.cpp
--- a/Engine/Plugins/Runtime/GearVR/Source/GearVR/Private/GearVRRender.cpp
+++ b/Engine/Plugins/Runtime/GearVR/Source/GearVR/Private/GearVRRender.cpp
@@ -1280,7 +1280,7 @@ void FCustomPresent::CopyTexture_RenderThread(FRHICommandListImmediate& RHICmdLi
 
 	SetRenderTarget(RHICmdList, DstTexture, FTextureRHIRef());
 	//RHICmdList.Clear(true, FLinearColor(1.0f, 0.0f, 0.0f, 1.0f), false, 0.0f, false, 0, FIntRect()); // @DBG
-	RHICmdList.ClearColor(FLinearColor(0.0f, 0.0f, 0.0f, 0.0f), FIntRect());
+	RHICmdList.ClearColorTexture(DstTexture, FLinearColor(0.0f, 0.0f, 0.0f, 0.0f), FIntRect());
 	RHICmdList.SetViewport(DstRect.Min.X, DstRect.Min.Y, 0, DstRect.Max.X, DstRect.Max.Y, 1.0f);
 
 	if (bAlphaPremultiply)
diff --git a/Engine/Plugins/Runtime/MobilePatchingUtils/MobilePatchingUtils.uplugin b/Engine/Plugins/Runtime/MobilePatchingUtils/MobilePatchingUtils.uplugin
new file mode 100644
--- /dev/null
+++ b/Engine/Plugins/Runtime/MobilePatchingUtils/MobilePatchingUtils.uplugin
@@ -0,0 +1,25 @@
+{
+	"FileVersion" : 3,
+	"Version" : 1,
+	"VersionName" : "1.0",
+	"FriendlyName" : "Mobile Patching Utilities",
+	"Description" : "Blueprint exposed functionality for downloading and patching content on mobile platforms",
+	"Category" : "Mobile",
+	"CreatedBy" : "Epic Games, Inc.",
+	"CreatedByURL" : "http://epicgames.com",
+	"DocsURL" : "",
+	"MarketplaceURL" : "",
+	"SupportURL" : "",
+	"EnabledByDefault" : true,
+	"CanContainContent" : false,
+	"IsBetaVersion" : false,
+	"Installed" : false,
+	"Modules" :
+	[
+		{
+			"Name" : "MobilePatchingUtils",
+			"Type" : "Runtime",
+			"LoadingPhase" : "Default"
+		}
+	]
+}
\ No newline at end of file
diff --git a/Engine/Plugins/Runtime/MobilePatchingUtils/Source/MobilePatchingUtils/MobilePatchingUtils.Build.cs b/Engine/Plugins/Runtime/MobilePatchingUtils/Source/MobilePatchingUtils/MobilePatchingUtils.Build.cs
new file mode 100644
--- /dev/null
+++ b/Engine/Plugins/Runtime/MobilePatchingUtils/Source/MobilePatchingUtils/MobilePatchingUtils.Build.cs
@@ -0,0 +1,26 @@
+// Copyright 1998-2016 Epic Games, Inc. All Rights Reserved.
+
+namespace UnrealBuildTool.Rules
+{
+    public class MobilePatchingUtils : ModuleRules
+    {
+        public MobilePatchingUtils(TargetInfo Target)
+        {
+            PrivateIncludePaths.Add("MobilePatchingUtils/Private");
+
+            PublicDependencyModuleNames.AddRange(new string[]
+            {
+                "Core",
+                "CoreUObject",
+                "Engine",
+            });
+
+            PrivateDependencyModuleNames.AddRange(new string[]
+            {
+                "PakFile",
+                "HTTP",
+                "BuildPatchServices"
+            });
+        }
+    }
+}
diff --git a/Engine/Plugins/Runtime/MobilePatchingUtils/Source/MobilePatchingUtils/Private/MobilePatchingLibrary.cpp b/Engine/Plugins/Runtime/MobilePatchingUtils/Source/MobilePatchingUtils/Private/MobilePatchingLibrary.cpp
new file mode 100644
--- /dev/null
+++ b/Engine/Plugins/Runtime/MobilePatchingUtils/Source/MobilePatchingUtils/Private/MobilePatchingLibrary.cpp
@@ -0,0 +1,377 @@
+// Copyright 1998-2016 Epic Games, Inc. All Rights Reserved.
+
+#include "MobilePatchingUtilsPrivatePCH.h"
+
+#include "MobilePatchingLibrary.h"
+#include "Http.h"
+#include "IPlatformFilePak.h"
+#include "UnitConversion.h"
+#include "BuildPatchServicesSingleton.h"
+#include "ScopeExit.h"
+
+#define LOCTEXT_NAMESPACE "MobilePatchingUtils"
+
+DEFINE_LOG_CATEGORY_STATIC(LogMobilePatchingUtils, Log, All);
+
+static FString GetStagingDir()
+{
+	return FPaths::GamePersistentDownloadDir() / TEXT("PatchStaging");
+}
+
+static IBuildPatchServicesModule* GetBuildPatchServices()
+{
+	static IBuildPatchServicesModule* BuildPatchServices = nullptr;
+	if (BuildPatchServices == nullptr)
+	{
+		BuildPatchServices = &FModuleManager::LoadModuleChecked<IBuildPatchServicesModule>(TEXT("BuildPatchServices"));
+	}
+	return BuildPatchServices;
+}
+
+static const FText& GetRequestContentErrorText(ERequestContentError ErrorCode)
+{
+	static const FText NoError(LOCTEXT("RequestContentError_NoError", "The operation was successful."));
+	static const FText InvalidInstallationDirectory(LOCTEXT("RequestContentError_InvalidInstallationDirectory", "Invalid installation directory"));
+	static const FText InvalidCloudURL(LOCTEXT("RequestContentError_InvalidCloudURL", "Invalid cloud URL"));
+	static const FText InvalidManifestURL(LOCTEXT("RequestContentError_InvalidManifestURL", "Invalid manifest URL"));
+	static const FText FailedToDownloadManifestNoResponse(LOCTEXT("RequestContentError_FailedToDownloadManifestNoResponse", "Failed to download manifest file. No response"));
+	static const FText FailedToDownloadManifest(LOCTEXT("RequestContentError_FailedToDownloadManifest", "Failed to download manifest file"));
+	static const FText FailedToReadManifest(LOCTEXT("RequestContentError_FailedToReadManifest", "Failed to reconstruct downloaded manifest file"));
+	static const FText InvalidOrMax(LOCTEXT("RequestContentError_InvalidOrMax", "An unknown error ocurred"));
+
+	switch (ErrorCode)
+	{
+	case ERequestContentError::NoError: return NoError;
+	case ERequestContentError::InvalidInstallationDirectory: return InvalidInstallationDirectory;
+	case ERequestContentError::InvalidCloudURL: return InvalidCloudURL;
+	case ERequestContentError::InvalidManifestURL: return InvalidManifestURL;
+	case ERequestContentError::FailedToDownloadManifestNoResponse: return FailedToDownloadManifestNoResponse;
+	case ERequestContentError::FailedToDownloadManifest: return FailedToDownloadManifest;
+	case ERequestContentError::FailedToReadManifest: return FailedToReadManifest;
+	default: return InvalidOrMax;
+	}
+}
+
+float UMobileInstalledContent::GetInstalledContentSize()
+{
+	check(InstalledManifest.IsValid())
+	
+	double BuildSize = (double)InstalledManifest->GetBuildSize();
+	double BuildSizeMB = FUnitConversion::Convert(BuildSize, EUnit::Bytes, EUnit::Megabytes);
+	return (float)BuildSizeMB;
+}
+
+float UMobileInstalledContent::GetDiskFreeSpace()
+{
+	uint64 TotalDiskSpace = 0;
+	uint64 TotalDiskFreeSpace = 0;
+	if (FPlatformMisc::GetDiskTotalAndFreeSpace(InstallDir, TotalDiskSpace, TotalDiskFreeSpace))
+	{
+		double DiskFreeSpace = (double)TotalDiskFreeSpace;
+		double DiskFreeSpaceMB = FUnitConversion::Convert(DiskFreeSpace, EUnit::Bytes, EUnit::Megabytes);
+		return (float)DiskFreeSpaceMB;
+	}
+	return 0.0f;
+}
+
+bool UMobileInstalledContent::Mount(int32 InPakOrder, const FString& InMountPoint)
+{
+	// Mount all pak files found in this content
+	FPakPlatformFile* PakFileMgr = (FPakPlatformFile*)(FPlatformFileManager::Get().FindPlatformFile(TEXT("PakFile")));
+	if (PakFileMgr == nullptr)
+	{
+		return false;
+	}
+
+	bool bMounted = false;
+	uint32 PakOrder = (uint32)FMath::Max(0, InPakOrder);
+	const TCHAR* MountPount = InMountPoint.GetCharArray().GetData();
+
+	TArray<FString> InstalledFileNames = InstalledManifest->GetBuildFileList();
+	for (const FString& FileName : InstalledFileNames)
+	{
+		if (FPaths::GetExtension(FileName) == TEXT("pak"))
+		{
+			FString PakFullName = InstallDir / FileName;
+			if (PakFileMgr->Mount(*PakFullName, PakOrder, MountPount))
+			{
+				UE_LOG(LogMobilePatchingUtils, Log, TEXT("Mounted = %s, Order = %d, MountPoint = %s"), *PakFullName, PakOrder, MountPount ? TEXT("(null)") : MountPount);
+				bMounted = true;
+			}
+			else
+			{
+				UE_LOG(LogMobilePatchingUtils, Error, TEXT("Failed to mount pak = %s"), *PakFullName);
+				bMounted = false;
+				break;
+			}
+		}
+	}
+	
+	return bMounted;
+}
+
+float UMobilePendingContent::GetDownloadSize()
+{
+	double DonwloadSize = 0.0;
+
+	if (RemoteManifest.IsValid())
+	{
+		if (InstalledManifest.IsValid())
+		{
+			TSet<FString> Tags;
+			DonwloadSize = (double)RemoteManifest->GetDeltaDownloadSize(Tags, InstalledManifest.ToSharedRef());
+		}
+		else
+		{
+			DonwloadSize = (double)RemoteManifest->GetDownloadSize();
+		}
+	}
+	
+	double DonwloadSizeMB = FUnitConversion::Convert(DonwloadSize, EUnit::Bytes, EUnit::Megabytes);	
+	return (float)DonwloadSizeMB;
+}
+
+float UMobilePendingContent::GetRequiredDiskSpace()
+{
+	if (RemoteManifest.IsValid())
+	{
+		double BuildSize = (double)RemoteManifest->GetBuildSize();
+		double BuildSizeMB = FUnitConversion::Convert(BuildSize, EUnit::Bytes, EUnit::Megabytes);
+		return (float)BuildSizeMB;
+	}
+
+	return 0.0f;
+}
+
+float UMobilePendingContent::GetDownloadSpeed()
+{
+	if (Installer.IsValid())
+	{
+		double SpeedBytes = (double)Installer->GetDownloadSpeed();
+		double SpeedMB = FUnitConversion::Convert(SpeedBytes, EUnit::Bytes, EUnit::Megabytes);
+		return (float)SpeedMB;
+	}
+	return 0.0f;
+}
+
+float UMobilePendingContent::GetTotalDownloadedSize()
+{
+	if (Installer.IsValid())
+	{
+		double SizeBytes = (double)Installer->GetTotalDownloaded();
+		double SizeMB = FUnitConversion::Convert(SizeBytes, EUnit::Bytes, EUnit::Megabytes);
+		return (float)SizeMB;
+	}
+	return 0.0f;
+}
+
+FText UMobilePendingContent::GetDownloadStatusText()
+{
+	if (Installer.IsValid())
+	{
+		return Installer->GetStatusText();
+	}
+	return FText::GetEmpty();
+}
+
+float UMobilePendingContent::GetInstallProgress()
+{
+	if (Installer.IsValid())
+	{
+		return Installer->GetUpdateProgress();
+	}
+	return 0.0f;
+}
+
+static void OnInstallComplete(bool bSuccess, IBuildManifestRef RemoteManifest, UMobilePendingContent* MobilePendingContent, FOnContentInstallSucceeded OnSucceeded, FOnContentInstallFailed OnFailed)
+{
+	if (bSuccess)
+	{
+		IBuildPatchServicesModule* BuildPatchServices = GetBuildPatchServices();
+		FString ManifestFilename = FPaths::GetCleanFilename(MobilePendingContent->RemoteManifestURL);
+		FString ManifestFullFilename = MobilePendingContent->InstallDir / ManifestFilename;
+		
+		if (!BuildPatchServices->SaveManifestToFile(ManifestFullFilename, RemoteManifest))
+		{
+			UE_LOG(LogMobilePatchingUtils, Error, TEXT("Failed to save manifest to disk. %s"), *ManifestFullFilename);
+		}
+
+		// Installed content updated
+		MobilePendingContent->InstalledManifest = RemoteManifest;
+		OnSucceeded.ExecuteIfBound();
+	}
+	else
+	{
+		FText ErrorText = MobilePendingContent->Installer->GetErrorText();
+		EBuildPatchInstallError ErrorCode = MobilePendingContent->Installer->GetErrorType();
+		UE_LOG(LogMobilePatchingUtils, Error, TEXT("Installation Failed. Code %d Msg: %s"), (int32)ErrorCode, *ErrorText.ToString());
+		OnFailed.ExecuteIfBound(ErrorText, (int32)ErrorCode);
+	}
+
+	// Release installer
+	MobilePendingContent->Installer.Reset();
+}
+
+void UMobilePendingContent::StartInstall(FOnContentInstallSucceeded OnSucceeded, FOnContentInstallFailed OnFailed)
+{
+	const FString StageDir = GetStagingDir();
+	IBuildPatchServicesModule* BuildPatchServices = GetBuildPatchServices();
+	
+	float DownloadSize = GetDownloadSize();
+	float RequiredDiskSpace = GetRequiredDiskSpace();
+	float DiskFreeSpace = GetDiskFreeSpace();
+	UE_LOG(LogMobilePatchingUtils, Log, TEXT("Download size = %.2f MB"), DownloadSize);
+	UE_LOG(LogMobilePatchingUtils, Log, TEXT("Required disk space = %.2f MB"), RequiredDiskSpace);
+	UE_LOG(LogMobilePatchingUtils, Log, TEXT("Disk free space = %.2f MB"), DiskFreeSpace);
+
+	BuildPatchServices->SetCloudDirectory(CloudURL);
+	BuildPatchServices->SetStagingDirectory(StageDir);
+	Installer = BuildPatchServices->StartBuildInstall(InstalledManifest, RemoteManifest, InstallDir, FBuildPatchBoolManifestDelegate::CreateStatic(&OnInstallComplete, this, OnSucceeded, OnFailed));
+}
+
+void UMobilePendingContent::BeginDestroy()
+{
+	Super::BeginDestroy();
+	
+	if (Installer.IsValid())
+	{
+		Installer->CancelInstall();
+		Installer.Reset();
+	}
+}
+
+static void OnRemoteManifestReady(FHttpRequestPtr Request, FHttpResponsePtr Response, bool bSucceeded, UMobilePendingContent* MobilePendingContent, FOnRequestContentSucceeded OnSucceeded, FOnRequestContentFailed OnFailed)
+{
+	ERequestContentError ErrorCode = ERequestContentError::NoError;
+	ON_SCOPE_EXIT
+	{
+		if (ErrorCode != ERequestContentError::NoError)
+		{
+			FText ErrorText = GetRequestContentErrorText(ErrorCode);
+			UE_LOG(LogMobilePatchingUtils, Error, TEXT("ErrorText: %s. Code %d"), *ErrorText.ToString(), (int32)ErrorCode);
+			OnFailed.ExecuteIfBound(ErrorText, (int32)ErrorCode);
+		}
+	};
+		
+	if (!bSucceeded || !Response.IsValid())
+	{
+		ErrorCode = ERequestContentError::FailedToDownloadManifestNoResponse;
+		return;
+	}
+
+	int32 ResponseCode = Response->GetResponseCode();
+	if (!EHttpResponseCodes::IsOk(ResponseCode))
+	{
+		UE_LOG(LogMobilePatchingUtils, Error, TEXT("Failed to download manifest. ResponseCode = %d, ResponseMsg = %s"), ResponseCode, *Response->GetContentAsString());
+		ErrorCode = ERequestContentError::FailedToDownloadManifest;
+		return;
+	}
+	
+	IBuildPatchServicesModule* BuildPatchServices = GetBuildPatchServices();
+	IBuildManifestPtr RemoteManifest = BuildPatchServices->MakeManifestFromData(Response->GetContent());
+	if (!RemoteManifest.IsValid())
+	{
+		ErrorCode = ERequestContentError::FailedToReadManifest;
+		return;
+	}
+
+	MobilePendingContent->RemoteManifest = RemoteManifest;
+	OnSucceeded.ExecuteIfBound(MobilePendingContent);
+}
+
+static IBuildManifestPtr GetInstalledManifest(const FString& InstallDirectory)
+{
+	IBuildManifestPtr InstalledManifest;
+	
+	FString FullInstallDir = FPaths::GamePersistentDownloadDir() / InstallDirectory;
+	TArray<FString> InstalledManifestNames;
+	IFileManager::Get().FindFiles(InstalledManifestNames, *(FullInstallDir / TEXT("*.manifest")), true, false);
+
+	if (InstalledManifestNames.Num()) 
+	{
+		const FString& InstalledManifestName = InstalledManifestNames[0];// should we warn if there more than one manifest?
+		IBuildPatchServicesModule* BuildPatchServices = GetBuildPatchServices();
+		InstalledManifest = BuildPatchServices->LoadManifestFromFile(FullInstallDir / InstalledManifestName);
+	}
+
+	return InstalledManifest;
+}
+
+UMobileInstalledContent* UMobilePatchingLibrary::GetInstalledContent(const FString& InstallDirectory)
+{
+	UMobileInstalledContent* InstalledContent = nullptr;
+	// Look for installed manifest
+	IBuildManifestPtr InstalledManifest = GetInstalledManifest(InstallDirectory);
+	if (InstalledManifest.IsValid())
+	{
+		InstalledContent = NewObject<UMobileInstalledContent>();
+		InstalledContent->InstallDir = FPaths::GamePersistentDownloadDir() / InstallDirectory;
+		InstalledContent->InstalledManifest = InstalledManifest;
+	}
+	
+	return InstalledContent;
+}
+
+void UMobilePatchingLibrary::RequestContent(const FString& RemoteManifestURL, const FString& CloudURL, const FString& InstallDirectory, FOnRequestContentSucceeded OnSucceeded, FOnRequestContentFailed OnFailed)
+{
+	ERequestContentError ErrorCode = ERequestContentError::NoError;
+	ON_SCOPE_EXIT
+	{
+		if (ErrorCode != ERequestContentError::NoError)
+		{
+			FText ErrorText = GetRequestContentErrorText(ErrorCode);
+			UE_LOG(LogMobilePatchingUtils, Error, TEXT("ErrorText: %s. Code %d"), *ErrorText.ToString(), (int32)ErrorCode);
+			OnFailed.ExecuteIfBound(ErrorText, (int32)ErrorCode);
+		}
+	};
+		
+	if (InstallDirectory.IsEmpty())
+	{
+		ErrorCode = ERequestContentError::InvalidInstallationDirectory;
+		return;
+	}
+
+	if (RemoteManifestURL.IsEmpty())
+	{
+		ErrorCode = ERequestContentError::InvalidManifestURL;
+		return;
+	}
+
+	if (CloudURL.IsEmpty())
+	{
+		ErrorCode = ERequestContentError::InvalidCloudURL;
+		return;
+	}
+		
+	UMobilePendingContent* MobilePendingContent = NewObject<UMobilePendingContent>();
+	MobilePendingContent->InstallDir = FPaths::GamePersistentDownloadDir() / InstallDirectory;
+	MobilePendingContent->RemoteManifestURL = RemoteManifestURL;
+	MobilePendingContent->CloudURL = CloudURL;
+	MobilePendingContent->InstalledManifest = GetInstalledManifest(InstallDirectory);
+
+	// Request remote manifest
+	TSharedRef<class IHttpRequest> HttpRequest = FHttpModule::Get().CreateRequest();
+	HttpRequest->OnProcessRequestComplete().BindStatic(OnRemoteManifestReady, MobilePendingContent, OnSucceeded, OnFailed);
+	HttpRequest->SetURL(*RemoteManifestURL);
+	HttpRequest->SetVerb(TEXT("GET"));
+	HttpRequest->ProcessRequest();
+}
+
+bool UMobilePatchingLibrary::HasActiveWiFiConnection()
+{
+	return FPlatformMisc::HasActiveWiFiConnection();
+}
+
+FString UMobilePatchingLibrary::GetActiveDeviceProfileName()
+{
+	UDeviceProfile* DeviceProfile = UDeviceProfileManager::Get().GetActiveProfile();
+	check(DeviceProfile);
+	return DeviceProfile->GetName();
+}
+
+TArray<FString> UMobilePatchingLibrary::GetSupportedPlatformNames()
+{
+	TArray<FString> Result;
+	FPlatformMisc::GetValidTargetPlatforms(Result);
+	return Result;
+}
\ No newline at end of file
diff --git a/Engine/Plugins/Runtime/MobilePatchingUtils/Source/MobilePatchingUtils/Private/MobilePatchingLibrary.h b/Engine/Plugins/Runtime/MobilePatchingUtils/Source/MobilePatchingUtils/Private/MobilePatchingLibrary.h
new file mode 100644
--- /dev/null
+++ b/Engine/Plugins/Runtime/MobilePatchingUtils/Source/MobilePatchingUtils/Private/MobilePatchingLibrary.h
@@ -0,0 +1,146 @@
+// Copyright 1998-2016 Epic Games, Inc. All Rights Reserved.
+
+#pragma once
+
+#include "IBuildPatchServicesModule.h"
+#include "MobilePatchingLibrary.generated.h"
+
+DECLARE_DYNAMIC_DELEGATE(FOnContentInstallSucceeded);
+DECLARE_DYNAMIC_DELEGATE_TwoParams(FOnContentInstallFailed, FText, ErrorText, int32, ErrorCode);
+
+UCLASS(BlueprintType)
+class UMobileInstalledContent : public UObject
+{
+	GENERATED_BODY()
+
+public:
+	/** Get the disk free space in megabytes where content is installed  */
+	UFUNCTION(BlueprintPure, Category="Mobile Patching")
+	float GetDiskFreeSpace();
+
+	/** Get the installed content size in megabytes */
+	UFUNCTION(BlueprintPure, Category="Mobile Patching")
+	float GetInstalledContentSize();
+
+	/** 
+	 * Mount installed content
+	 * @param	PakOrder : Content pak priority
+	 * @param	MountPoint : Path to mount the pak at
+	 */
+	UFUNCTION(BlueprintCallable, Category="Mobile Patching", meta=(AdvancedDisplay="PakOrder,MountPoint", PakOrder="1"))
+	bool Mount(int32 PakOrder, const FString& MountPoint);
+
+public:	
+	// User specified directory where content should be/already installed
+	FString InstallDir;
+	// Currently installed manifest
+	IBuildManifestPtr InstalledManifest;
+};
+
+UCLASS(BlueprintType)
+class UMobilePendingContent : public UMobileInstalledContent
+{
+	GENERATED_BODY()
+
+public:	
+	/** Get the total download size for this content installation */
+	UFUNCTION(BlueprintPure, Category="Mobile Patching")
+	float GetDownloadSize();
+	
+	/** Get the required disk space in megabytes for this content installation */
+	UFUNCTION(BlueprintPure, Category="Mobile Patching")
+	float GetRequiredDiskSpace();
+
+	/** Get the total downloaded size in megabytes. Valid during installation */
+	UFUNCTION(BlueprintPure, Category="Mobile Patching|Progress")
+	float GetTotalDownloadedSize();
+
+	/** Get the current download speed in megabytes per second. Valid during installation */
+	UFUNCTION(BlueprintPure, Category="Mobile Patching|Progress")
+	float GetDownloadSpeed();
+
+	/** Get the current installation status text. Valid during installation */
+	UFUNCTION(BlueprintPure, Category="Mobile Patching|Progress")
+	FText GetDownloadStatusText();
+	
+	/** Get the current installation progress. Between 0 and 1 for known progress, or less than 0 for unknown progress */
+	UFUNCTION(BlueprintPure, Category="Mobile Patching|Progress")
+	float GetInstallProgress();
+	
+	 /** 
+	 * Attempt to download and install remote content.
+	 * User can choose to mount installed content into the game.
+	 * @param	OnSucceeded: Callback on installation success. 
+	 * @param	OnFailed: Callback on installation fail. Will return error message text and error integer code
+	 */
+	UFUNCTION(BlueprintCallable, Category="Mobile Patching")
+	void StartInstall(FOnContentInstallSucceeded OnSucceeded, FOnContentInstallFailed OnFailed);
+
+	//~ Begin UObject Interface
+	virtual void BeginDestroy() override;
+	//~ End UObject Interface.
+
+public:
+	// User specified URL from where manifest can be downloaded
+	FString RemoteManifestURL;
+	// User specified cloud URL from where content chunks can be downloaded
+	FString CloudURL;
+	// Content installer, only valid during installation
+	IBuildInstallerPtr Installer;
+	// Manifest downloaded from a cloud
+	IBuildManifestPtr RemoteManifest;
+};
+
+enum class ERequestContentError
+{
+	NoError,
+	InvalidInstallationDirectory,
+	InvalidCloudURL,
+	InvalidManifestURL,
+	FailedToDownloadManifestNoResponse,
+	FailedToDownloadManifest,
+	FailedToReadManifest,
+};
+
+DECLARE_DYNAMIC_DELEGATE_OneParam(FOnRequestContentSucceeded, UMobilePendingContent*, MobilePendingContent);
+DECLARE_DYNAMIC_DELEGATE_TwoParams(FOnRequestContentFailed, FText, ErrorText, int32, ErrorCode);
+
+UCLASS()
+class UMobilePatchingLibrary : public UBlueprintFunctionLibrary
+{
+	GENERATED_BODY()
+
+public:
+	/** 
+	 * Get the installed content. Will return non-null object if there is an installed content at specified directory
+	 * User can choose to mount installed content into the game
+	 */
+	UFUNCTION(BlueprintPure, Category="Mobile Patching")
+	static UMobileInstalledContent* GetInstalledContent(const FString& InstallDirectory);
+	
+	/** 
+	 * Attempt to download manifest file using specified manifest URL. 
+	 * On success it will return an object that represents remote content. This object can be queried for additional information, like total content size, download size, etc.
+	 * User can choose to download and install remote content.
+	 * @param	RemoteManifestURL : URL from where manifest file can be downloaded. (http://my-server.com/awesomecontent.manifest)
+	 * @param	CloudURL :  URL from where content chunk data can be downloaded. (http://my-server.com/awesomecontent/clouddir)
+	 * @param	InstallDirectory: Relative directory to where downloaded content should be installed (MyContent/AwesomeContent)
+	 * @param	OnSucceeded: Callback on manifest download success. Will return object that represents remote content 
+	 * @param	OnFailed: Callback on manifest download fail. Will return error message text and error integer code
+	 */
+	UFUNCTION(BlueprintCallable, Category="Mobile Patching")
+	static void RequestContent(const FString& RemoteManifestURL, const FString& CloudURL, const FString& InstallDirectory, FOnRequestContentSucceeded OnSucceeded, FOnRequestContentFailed OnFailed);
+
+	// Whether WiFi connection is currently available
+	UFUNCTION(BlueprintPure, Category="Mobile Patching")
+	static bool HasActiveWiFiConnection();
+
+	// Get the name of currently selected device profile name
+	UFUNCTION(BlueprintPure, Category="Mobile Patching")
+	static FString GetActiveDeviceProfileName();
+	
+	// Get the list of supported platform names on this device. Example: Android_ETC2, Android_ASTC
+	UFUNCTION(BlueprintPure, Category="Mobile Patching")
+	static TArray<FString> GetSupportedPlatformNames();
+};
+
diff --git a/Engine/Plugins/Runtime/MobilePatchingUtils/Source/MobilePatchingUtils/Private/MobilePatchingUtils.cpp b/Engine/Plugins/Runtime/MobilePatchingUtils/Source/MobilePatchingUtils/Private/MobilePatchingUtils.cpp
new file mode 100644
--- /dev/null
+++ b/Engine/Plugins/Runtime/MobilePatchingUtils/Source/MobilePatchingUtils/Private/MobilePatchingUtils.cpp
@@ -0,0 +1,17 @@
+// Copyright 1998-2016 Epic Games, Inc. All Rights Reserved.
+
+#include "MobilePatchingUtilsPrivatePCH.h"
+
+class FMobilePatchingUtilsModule : public FDefaultModuleImpl
+{
+public:
+	virtual void StartupModule() override
+	{
+	}
+
+	virtual void ShutdownModule() override
+	{
+	}
+};
+
+IMPLEMENT_MODULE( FMobilePatchingUtilsModule, MobilePatchingUtils )
diff --git a/Engine/Plugins/Runtime/MobilePatchingUtils/Source/MobilePatchingUtils/Private/MobilePatchingUtilsPrivatePCH.h b/Engine/Plugins/Runtime/MobilePatchingUtils/Source/MobilePatchingUtils/Private/MobilePatchingUtilsPrivatePCH.h
new file mode 100644
--- /dev/null
+++ b/Engine/Plugins/Runtime/MobilePatchingUtils/Source/MobilePatchingUtils/Private/MobilePatchingUtilsPrivatePCH.h
@@ -0,0 +1,3 @@
+// Copyright 1998-2016 Epic Games, Inc. All Rights Reserved.
+
+#include "Engine.h"
diff --git a/Engine/Shaders/Common.usf b/Engine/Shaders/Common.usf
--- a/Engine/Shaders/Common.usf
+++ b/Engine/Shaders/Common.usf
@@ -218,6 +218,10 @@
 #define MOBILE_EMULATION ((FEATURE_LEVEL == FEATURE_LEVEL_ES2 || FEATURE_LEVEL == FEATURE_LEVEL_ES3_1) && (!(COMPILER_GLSL_ES3_1 || COMPILER_GLSL_ES2) && USE_DEVELOPMENT_SHADERS && !(METAL_PROFILE && MAC) && !VULKAN_PROFILE && !WOLF_PROFILE))
 #endif
 
+#ifndef MOBILE_FORCE_DEPTH_TEXTURE_READS
+#define MOBILE_FORCE_DEPTH_TEXTURE_READS 0
+#endif
+
 // If we didn't request color component swizzling, just make it empty
 #ifndef FCOLOR_COMPONENT_SWIZZLE
 #define FCOLOR_COMPONENT_SWIZZLE
@@ -573,6 +577,8 @@ MaterialFloat PhongShadingPow(MaterialFloat X, MaterialFloat Y)
 // depth in the red channel in DeviceZ
 Texture2D		SceneDepthTexture;
 SamplerState	SceneDepthTextureSampler;
+Texture2D		CustomDepthTexture;
+SamplerState	CustomDepthTextureSampler;
 // scene HDR color
 Texture2D		SceneColorTexture;
 SamplerState	SceneColorTextureSampler;
@@ -1074,27 +1080,26 @@ float CalcSceneDepth(float2 ScreenUV)
 #if SCENE_TEXTURES_DISABLED
 	return 0.0f;
 #else
-
-#if FEATURE_LEVEL > FEATURE_LEVEL_ES3_1
-	return ConvertFromDeviceZ(Texture2DSampleLevel(SceneDepthTexture, SceneDepthTextureSampler, ScreenUV, 0).r);
-#else
-	#if COMPILER_GLSL_ES2
-		#if IOS
-			// Only call FramebufferFetch when actually compiling for IOS ES2.
-			return FramebufferFetchES2().w;
-		#elif WEBGL
-			return Texture2DSampleLevel(SceneAlphaCopyTexture, SceneAlphaCopyTextureSampler, ScreenUV, 0).r;
-		#else 
-			float SceneW = ConvertFromDeviceZ(Texture2DSampleLevel(SceneDepthTexture, SceneDepthTextureSampler, ScreenUV, 0).r);		
-			return DepthbufferFetchES2(SceneW, View.InvDeviceZToWorldZTransform[2], View.InvDeviceZToWorldZTransform[3]);
-		#endif 
-	#elif METAL_PROFILE && !MAC
-			return FramebufferFetchES2().w;
-	#else
+	#if FEATURE_LEVEL > FEATURE_LEVEL_ES3_1 || MOBILE_FORCE_DEPTH_TEXTURE_READS
 		return ConvertFromDeviceZ(Texture2DSampleLevel(SceneDepthTexture, SceneDepthTextureSampler, ScreenUV, 0).r);
+	#else
+		#if COMPILER_GLSL_ES2
+			#if IOS
+				// Only call FramebufferFetch when actually compiling for IOS ES2.
+				return FramebufferFetchES2().w;
+			#elif WEBGL
+				return Texture2DSampleLevel(SceneAlphaCopyTexture, SceneAlphaCopyTextureSampler, ScreenUV, 0).r;
+			#else 
+				float SceneW = ConvertFromDeviceZ(Texture2DSampleLevel(SceneDepthTexture, SceneDepthTextureSampler, ScreenUV, 0).r);		
+				return DepthbufferFetchES2(SceneW, View.InvDeviceZToWorldZTransform[2], View.InvDeviceZToWorldZTransform[3]);
+			#endif 
+		#elif METAL_PROFILE && !MAC
+				return FramebufferFetchES2().w;
+		#else
+			return ConvertFromDeviceZ(Texture2DSampleLevel(SceneDepthTexture, SceneDepthTextureSampler, ScreenUV, 0).r);
+		#endif
 	#endif
 #endif
-#endif
 }
 
 /** Returns DeviceZ which is the z value stored in the depth buffer. */
diff --git a/Engine/Shaders/DeferredShadingCommon.usf b/Engine/Shaders/DeferredShadingCommon.usf
--- a/Engine/Shaders/DeferredShadingCommon.usf
+++ b/Engine/Shaders/DeferredShadingCommon.usf
@@ -673,8 +673,6 @@ SamplerState DBufferCTextureSampler;
 
 Texture2D ScreenSpaceAOTexture;
 SamplerState ScreenSpaceAOTextureSampler;
-Texture2D CustomDepthTexture;
-SamplerState CustomDepthTextureSampler;
 
 #if FEATURE_LEVEL >= FEATURE_LEVEL_SM4 
 Texture2D<uint2> CustomStencilTexture;
diff --git a/Engine/Shaders/MaterialTemplate.usf b/Engine/Shaders/MaterialTemplate.usf
--- a/Engine/Shaders/MaterialTemplate.usf
+++ b/Engine/Shaders/MaterialTemplate.usf
@@ -1093,19 +1093,37 @@ float4 GetPostProcessInputMinMax(uint Index)
 	return 1.0f;
 }
 
-MaterialFloat4 MobileLookupPostProcessInput0(inout FMaterialPixelParameters Parameters, float2 UV)
+MaterialFloat4 MobileSceneTextureLookup(inout FMaterialPixelParameters Parameters, int SceneTextureId, float2 UV)
 {
-#if POST_PROCESS_MATERIAL && (FEATURE_LEVEL <= FEATURE_LEVEL_ES3_1)
-	MaterialFloat4 Input0 = Texture2DSample(PostprocessInput0, PostprocessInput0Sampler, UV);
-	#if POST_PROCESS_MATERIAL_BEFORE_TONEMAP
+#if (FEATURE_LEVEL <= FEATURE_LEVEL_ES3_1)
+	//PPI_SceneDepth
+	if (SceneTextureId == 1)
+	{
+		MaterialFloat Depth = ConvertFromDeviceZ(Texture2DSample(SceneDepthTexture, SceneDepthTextureSampler, UV).r);
+		return MaterialFloat4(Depth.rrr, 0.f);
+	}
+	//PPI_CustomDepth
+	else if (SceneTextureId == 13)
+	{
+		MaterialFloat Depth = ConvertFromDeviceZ(Texture2DSample(CustomDepthTexture, CustomDepthTextureSampler, UV).r);
+		return MaterialFloat4(Depth.rrr, 0.f);
+	}
+	// PPI_PostprocessInput
+	else if (SceneTextureId == 14)
+	{
+#if POST_PROCESS_MATERIAL	
+		MaterialFloat4 Input0 = Texture2DSample(PostprocessInput0, PostprocessInput0Sampler, UV);
+		#if POST_PROCESS_MATERIAL_BEFORE_TONEMAP
 		Input0 = Decode32BPPHDR(Input0);
-	#endif
-	// We need to preserve original SceneColor Alpha as it's used by tonemaper on mobile
-	Parameters.BackupSceneColorAlpha = Input0.a;
-	return Input0;
-#else
+		#endif
+		// We need to preserve original SceneColor Alpha as it's used by tonemaper on mobile
+		Parameters.BackupSceneColorAlpha = Input0.a;
+		return Input0;
+#endif// POST_PROCESS_MATERIAL	
+	}
+#endif// FEATURE_LEVEL
+
 	return MaterialFloat4(0.0f, 0.0f, 0.0f, 0.0f);
-#endif
 }
 
 /** Applies an offset to the scene texture lookup and decodes the HDR linear space color. */
diff --git a/Engine/Shaders/MobileBasePassCommon.usf b/Engine/Shaders/MobileBasePassCommon.usf
--- a/Engine/Shaders/MobileBasePassCommon.usf
+++ b/Engine/Shaders/MobileBasePassCommon.usf
@@ -4,12 +4,18 @@
 	MobileBasePassCommon.usf: Base pass definitions used by both vertex and pixel shader
 =============================================================================*/
 
+#define USE_PS_CLIP_PLANE PROJECT_ALLOW_GLOBAL_CLIP_PLANE
+
 #undef NEEDS_LIGHTMAP_COORDINATE
 
 #define NEEDS_LIGHTMAP_COORDINATE		(LQ_TEXTURE_LIGHTMAP)
 #define NEEDS_LIGHTMAP					(NEEDS_LIGHTMAP_COORDINATE)
 
-#define USE_VERTEX_FOG					(!OUTPUT_GAMMA_SPACE)
+#ifndef PROJECT_MOBILE_DISABLE_VERTEX_FOG
+	#define PROJECT_MOBILE_DISABLE_VERTEX_FOG 0
+#endif
+
+#define USE_VERTEX_FOG					(!PROJECT_MOBILE_DISABLE_VERTEX_FOG)
 
 #define PACK_INTERPOLANTS (USE_VERTEX_FOG && NUM_VF_PACKED_INTERPOLANTS > 0 && (ES2_PROFILE || ES3_1_PROFILE))
 #define LANDSCAPE_BUG_WORKAROUND (IOS && IS_MOBILE_BASEPASS_VERTEX_SHADER && PACK_INTERPOLANTS)
@@ -34,6 +40,11 @@ struct FSharedMobileBasePassInterpolants
 #if USE_WORLD_POSITION_EXCLUDING_SHADER_OFFSETS
 	float3 PixelPositionExcludingWPO	: TEXCOORD9;
 #endif
+
+#if USE_PS_CLIP_PLANE
+	float OutClipDistance : OUTCLIPDIST;
+#endif
+
 };
 
 #define FMobileBasePassInterpolantsVSToPS FSharedMobileBasePassInterpolants
diff --git a/Engine/Shaders/MobileBasePassPixelShader.usf b/Engine/Shaders/MobileBasePassPixelShader.usf
--- a/Engine/Shaders/MobileBasePassPixelShader.usf
+++ b/Engine/Shaders/MobileBasePassPixelShader.usf
@@ -34,6 +34,10 @@ float4 LightColorAndFalloffExponent[MAX_DYNAMIC_POINT_LIGHTS];
 #define MOBILE_QL_FORCE_LQ_REFLECTIONS 0
 #endif
 
+#ifndef MOBILE_CSM_QUALITY
+#define MOBILE_CSM_QUALITY 2
+#endif
+
 #define FULLY_ROUGH (MATERIAL_FULLY_ROUGH || MOBILE_QL_FORCE_FULLY_ROUGH)
 #define NONMETAL (MATERIAL_NONMETAL || MOBILE_QL_FORCE_NONMETAL)
 #define HQ_REFLECTIONS (MATERIAL_HQ_FORWARD_REFLECTIONS && !MOBILE_QL_FORCE_LQ_REFLECTIONS)
@@ -305,6 +309,10 @@ void Main(
 	ResolvedView = ResolveView();
 #endif
 
+#if USE_PS_CLIP_PLANE
+	clip(BasePassInterpolants.OutClipDistance);
+#endif
+
 #if PACK_INTERPOLANTS
 	float4 PackedInterpolants[NUM_VF_PACKED_INTERPOLANTS];
 	VertexFactoryUnpackInterpolants(Interpolants, PackedInterpolants);
@@ -456,10 +464,14 @@ void Main(
 			float LightSpacePixelDepthForOpaque = min(ShadowPosition.z, 0.99999f);
 			Settings.SceneDepth = LightSpacePixelDepthForOpaque;
 
-			#if 0
+			#if MOBILE_CSM_QUALITY == 0
+				half ShadowMap = ManualNoFiltering(ShadowPosition.xy, Settings);
+			#elif MOBILE_CSM_QUALITY == 1
 				half ShadowMap = Manual1x1PCF(ShadowPosition.xy, Settings);
-			#else
+			#elif MOBILE_CSM_QUALITY == 2
 				half ShadowMap = Manual2x2PCF(ShadowPosition.xy, Settings);
+			#else
+				#error Unsupported MOBILE_CSM_QUALITY value.
 			#endif
 
 			#if MOVABLE_DIRECTIONAL_LIGHT
@@ -594,10 +606,17 @@ void Main(
 		OutColor.rgba *= ResolvedView.ExposureScale.xyzw;
 	#endif
 
-	#if OUTPUT_GAMMA_SPACE
-		OutColor.rgb = sqrt( OutColor.rgb );
+#if USE_EDITOR_COMPOSITING && (MOBILE_EMULATION)
+	// Editor primitive depth testing
+	OutColor.a = 1.0;
+	#if MATERIALBLENDING_MASKED
+		// some material might have an opacity value
+		OutColor.a = GetMaterialMaskInputRaw(PixelMaterialInputs);
 	#endif
-
+		// we output premultiplied alpha to we have to darken all 4 channels
+	OutColor *= ClipForEditorPrimitives(MaterialParameters);
+	clip(OutColor.a - GetMaterialOpacityMaskClipValue());
+#else
 	#if ES2_PROFILE && !OUTPUT_GAMMA_SPACE
 		float Mode = GetHDR32bppEncodeMode();
 		if (Mode != HDR_ENCODE_NONE)
@@ -609,15 +628,8 @@ void Main(
 			OutColor = Encode32BPPHDR(OutColor, SvPosition.xy);
 		}
 	#endif
-
-	// Editor primitive depth testing
-	#if USE_EDITOR_COMPOSITING && (MOBILE_EMULATION)
-		#if MATERIALBLENDING_MASKED
-			// some material might have a opacity value
-			OutColor.a = GetMaterialMaskInputRaw(PixelMaterialInputs);
-		#endif
-		// we output premultiplied alpha to we have to darken all 4 channels
-		OutColor *= ClipForEditorPrimitives(MaterialParameters);
-		clip(OutColor.a - GetMaterialOpacityMaskClipValue());
+	#if OUTPUT_GAMMA_SPACE
+		OutColor.rgb = sqrt(OutColor.rgb);
 	#endif
+#endif
 }
diff --git a/Engine/Shaders/MobileBasePassVertexShader.usf b/Engine/Shaders/MobileBasePassVertexShader.usf
--- a/Engine/Shaders/MobileBasePassVertexShader.usf
+++ b/Engine/Shaders/MobileBasePassVertexShader.usf
@@ -82,6 +82,10 @@ void Main(
 	Output.BasePassInterpolants.PixelPositionExcludingWPO = WorldPositionExcludingWPO.xyz;
 #endif
 
+#if USE_PS_CLIP_PLANE
+	Output.BasePassInterpolants.OutClipDistance = dot(ResolvedView.GlobalClippingPlane, float4(WorldPosition.xyz - ResolvedView.PreViewTranslation.xyz, 1));
+#endif
+
 #if USE_VERTEX_FOG
 #if PACK_INTERPOLANTS
 	PackedInterps[0] = CalculateVertexHeightFog(WorldPosition.xyz - float4(ResolvedView.TranslatedWorldCameraOrigin, 1));
diff --git a/Engine/Shaders/MobileOpacityShaders.usf b/Engine/Shaders/MobileOpacityShaders.usf
new file mode 100644
--- /dev/null
+++ b/Engine/Shaders/MobileOpacityShaders.usf
@@ -0,0 +1,50 @@
+// Copyright 1998-2016 Epic Games, Inc. All Rights Reserved.
+
+/*=============================================================================
+MobileOpacityShaders.usf: Renders opacity to alpha channel. Mobile only.
+=============================================================================*/
+
+#include "Common.usf"
+#include "Material.usf"
+#include "VertexFactory.usf"
+
+
+#define VertexFactoryGetInterpolants VertexFactoryGetInterpolantsVSToPS
+
+void MainVS(
+	FVertexFactoryInput Input
+	, out FVertexFactoryInterpolantsVSToPS FactoryInterpolants
+	, out float4 Position : SV_POSITION
+)
+{
+	ResolvedView = ResolveView();
+
+	FVertexFactoryIntermediates VFIntermediates = GetVertexFactoryIntermediates(Input);
+	float4 WorldPos = VertexFactoryGetWorldPosition(Input, VFIntermediates);
+	half3x3 TangentToLocal = VertexFactoryGetTangentToLocal(Input, VFIntermediates);
+	FMaterialVertexParameters VertexParameters = GetMaterialVertexParameters(Input, VFIntermediates, WorldPos.xyz, TangentToLocal);
+
+	half3 WorldPositionOffset = GetMaterialWorldPositionOffset(VertexParameters);
+	WorldPos.xyz += WorldPositionOffset;
+
+	float4 RasterizedWorldPosition = VertexFactoryGetRasterizedWorldPosition(Input, VFIntermediates, WorldPos);
+	Position = mul(RasterizedWorldPosition, ResolvedView.TranslatedWorldToClip);
+
+	FactoryInterpolants = VertexFactoryGetInterpolants(Input, VFIntermediates, VertexParameters);
+}
+
+
+void MainPS(
+	FVertexFactoryInterpolantsVSToPS FactoryInterpolants
+	, in float4 SvPosition : SV_Position
+	OPTIONAL_IsFrontFace
+	, out half4 OutColor : SV_Target0
+)
+{
+	FMaterialPixelParameters MaterialParameters = GetMaterialPixelParameters(FactoryInterpolants, SvPosition);
+	FPixelMaterialInputs PixelMaterialInputs;
+	CalcMaterialParameters(MaterialParameters, PixelMaterialInputs, SvPosition, bIsFrontFace);
+	OutColor = half4(0,0,0, GetMaterialOpacity(PixelMaterialInputs));
+}
+
+
diff --git a/Engine/Shaders/MobileSceneCapture.usf b/Engine/Shaders/MobileSceneCapture.usf
new file mode 100644
--- /dev/null
+++ b/Engine/Shaders/MobileSceneCapture.usf
@@ -0,0 +1,63 @@
+// Copyright 1998-2016 Epic Games, Inc. All Rights Reserved.
+
+/*=============================================================================
+	Decode32bppHDR.usf: Decode 32bpp HDR encoding
+=============================================================================*/
+
+#include "Common.usf"
+
+float2 InvTexSize;
+
+void MainCopyVS(
+	float2 InPosition	: ATTRIBUTE0,
+	float2 InUV			: ATTRIBUTE1,
+#if DECODING_MOSAIC
+	out float2 OutTexCoords[4] : TEXCOORD4,
+#endif
+	out FScreenVertexOutput Output
+	)
+{
+	DrawRectangle( float4( InPosition, 0, 1 ), InUV, Output.Position, Output.UV);
+
+#if DECODING_MOSAIC
+	OutTexCoords[0] = Output.UV + InvTexSize.xy * float2( 0,-1);
+	OutTexCoords[1] = Output.UV + InvTexSize.xy * float2( 1, 0);
+	OutTexCoords[2] = Output.UV + InvTexSize.xy * float2(-1, 0);
+	OutTexCoords[3] = Output.UV + InvTexSize.xy * float2( 0, 1);
+#endif
+}
+
+
+Texture2D InTexture;
+SamplerState InTextureSampler;
+
+void MainCopyPS(
+#if DECODING_MOSAIC
+	in float2 InTexCoords[4] : TEXCOORD4,
+#endif
+	FScreenVertexOutput Input,
+	out float4 OutColor : SV_Target0
+	)
+{
+    float4 EncodedColor = Texture2DSample(InTexture, InTextureSampler, Input.UV);
+	// Decode EncodedColor, OutColor will saturate on 8 bit RTs.
+#if DECODING_MOSAIC
+	OutColor = Decode32BPPHDR(EncodedColor, Input.Position, InTexture, InTextureSampler, InTexCoords);
+#else
+	OutColor = Decode32BPPHDR(EncodedColor);
+#endif
+
+	float SceneW = CalcSceneDepth(Input.UV);
+
+#if SOURCE_MODE_SCENE_COLOR_AND_OPACITY
+	// Put inverse opacity in alpha channel. (determined via depth, translucency is combined with another pass.)
+	float MaxSceneDepth = 65000.0f;
+	OutColor.a = SceneW >= MaxSceneDepth ? 1.0f : 0.0f;
+#elif SOURCE_MODE_SCENE_COLOR_NO_ALPHA
+	OutColor.a = 0;
+#elif SOURCE_MODE_SCENE_COLOR_SCENE_DEPTH
+	OutColor.a = SceneW;
+#elif SOURCE_MODE_SCENE_DEPTH
+	OutColor.r = float4(SceneW,0,0,0);
+#endif
+}
diff --git a/Engine/Shaders/SceneCapturePixelShader.usf b/Engine/Shaders/SceneCapturePixelShader.usf
--- a/Engine/Shaders/SceneCapturePixelShader.usf
+++ b/Engine/Shaders/SceneCapturePixelShader.usf
@@ -18,6 +18,8 @@ void Main(
 	// Also to be consistent with DrawMaterialToRenderTarget, which doesn't have this full texture copy in which to invert
 	// Artists will have to invert the alpha to get foreground opacity
 	OutColor = float4(SceneColorValue.rgb, saturate(SceneColorValue.a));
+#elif SOURCE_MODE_SCENE_COLOR_NO_ALPHA
+	OutColor = float4(CalcSceneColor(Input.UV), 0);
 #elif SOURCE_MODE_SCENE_COLOR_SCENE_DEPTH
 	OutColor = float4(CalcSceneColor(Input.UV), CalcSceneDepth(Input.UV));
 #elif SOURCE_MODE_SCENE_DEPTH 
diff --git a/Engine/Source/Developer/FunctionalTesting/Private/AutomationBlueprintFunctionLibrary.cpp b/Engine/Source/Developer/FunctionalTesting/Private/AutomationBlueprintFunctionLibrary.cpp
--- a/Engine/Source/Developer/FunctionalTesting/Private/AutomationBlueprintFunctionLibrary.cpp
+++ b/Engine/Source/Developer/FunctionalTesting/Private/AutomationBlueprintFunctionLibrary.cpp
@@ -204,17 +204,25 @@ bool UAutomationBlueprintFunctionLibrary::TakeAutomationScreenshotInternal(const
 	FAutomationScreenshotTaker* TempObject = new FAutomationScreenshotTaker(Name, Options);
 #endif
 
-	FHighResScreenshotConfig& Config = GetHighResScreenshotConfig();
-	if ( Config.SetResolution(ResolutionX, ResolutionY, 1.0f) )
+    if (FPlatformProperties::HasFixedResolution())
+    {
+	    FScreenshotRequest::RequestScreenshot(false);
+	    return true;
+    }
+	else
 	{
-		if ( !GEngine->GameViewport->GetGameViewport()->TakeHighResScreenShot() )
-		{
-			// If we failed to take the screenshot, we're going to need to cleanup the automation screenshot taker.
-			delete TempObject;
-			return false;
-		}
+	    FHighResScreenshotConfig& Config = GetHighResScreenshotConfig();
+	    if ( Config.SetResolution(ResolutionX, ResolutionY, 1.0f) )
+	    {
+			if ( !GEngine->GameViewport->GetGameViewport()->TakeHighResScreenShot() )
+			{
+				// If we failed to take the screenshot, we're going to need to cleanup the automation screenshot taker.
+				delete TempObject;
+				return false;
+			}
 
-		return true;
+			return true;
+		}
 	}
 
 	return false;
diff --git a/Engine/Source/Developer/FunctionalTesting/Private/Tests/ClientFuncTestPerforming.cpp b/Engine/Source/Developer/FunctionalTesting/Private/Tests/ClientFuncTestPerforming.cpp
--- a/Engine/Source/Developer/FunctionalTesting/Private/Tests/ClientFuncTestPerforming.cpp
+++ b/Engine/Source/Developer/FunctionalTesting/Private/Tests/ClientFuncTestPerforming.cpp
@@ -131,7 +131,11 @@ void FClientFunctionalTestingMaps::GetTests(TArray<FString>& OutBeautifiedNames,
 	if ( !AssetRegistry.IsLoadingAssets() )
 	{
 		TArray<FAssetData> MapList;
-		if ( AssetRegistry.GetAssetsByClass(UWorld::StaticClass()->GetFName(), /*out*/ MapList) )
+		FARFilter Filter;
+		Filter.ClassNames.Add(UWorld::StaticClass()->GetFName());
+		Filter.bRecursiveClasses = true;
+		Filter.bIncludeOnlyOnDiskAssets = true;
+		if ( AssetRegistry.GetAssets(Filter, /*out*/ MapList) )
 		{
 			for ( const FAssetData& MapAsset : MapList )
 			{
diff --git a/Engine/Source/Developer/LauncherServices/Private/Launcher/LauncherWorker.cpp b/Engine/Source/Developer/LauncherServices/Private/Launcher/LauncherWorker.cpp
--- a/Engine/Source/Developer/LauncherServices/Private/Launcher/LauncherWorker.cpp
+++ b/Engine/Source/Developer/LauncherServices/Private/Launcher/LauncherWorker.cpp
@@ -241,7 +241,7 @@ FString FLauncherWorker::CreateUATCommand( const ILauncherProfileRef& InProfile,
 {
 	CommandStart = TEXT("");
 	FString UATCommand = TEXT(" -utf8output");
-	static FGuid SessionId(FGuid::NewGuid());
+	FGuid SessionId(FGuid::NewGuid());
 	FString InitialMap = InProfile->GetDefaultLaunchRole()->GetInitialMap();
 	if (InitialMap.IsEmpty() && InProfile->GetCookedMaps().Num() == 1)
 	{
diff --git a/Engine/Source/Developer/Linux/LinuxTargetPlatform/Private/LinuxTargetPlatform.h b/Engine/Source/Developer/Linux/LinuxTargetPlatform/Private/LinuxTargetPlatform.h
--- a/Engine/Source/Developer/Linux/LinuxTargetPlatform/Private/LinuxTargetPlatform.h
+++ b/Engine/Source/Developer/Linux/LinuxTargetPlatform/Private/LinuxTargetPlatform.h
@@ -156,10 +156,19 @@ public:
 	{
 		if (!PLATFORM_LINUX)
 		{
-			// check for LINUX_ROOT when targeting Linux from Win/Mac
+			// check for LINUX_MULTIARCH_ROOT or for legacy LINUX_ROOT when targeting Linux from Win/Mac
 			TCHAR ToolchainRoot[32768] = { 0 };
-			FPlatformMisc::GetEnvironmentVariable(TEXT("LINUX_ROOT"), ToolchainRoot, ARRAY_COUNT(ToolchainRoot));
-
+			FPlatformMisc::GetEnvironmentVariable(TEXT("LINUX_MULTIARCH_ROOT"), ToolchainRoot, ARRAY_COUNT(ToolchainRoot));
+			// proceed with any value for MULTIARCH root, because checking exact architecture is not possible at this point
+			FString ToolchainMultiarchRoot = ToolchainRoot;
+			if (ToolchainMultiarchRoot.Len() > 0 && FPaths::DirectoryExists(ToolchainMultiarchRoot))
+			{
+				return true;
+			}
+			
+			// else check for legacy LINUX_ROOT
+			ToolchainRoot[ 0 ] = 0;
+			FPlatformMisc::GetEnvironmentVariable(TEXT("LINUX_MULTIARCH_ROOT"), ToolchainRoot, ARRAY_COUNT(ToolchainRoot));            
 			FString ToolchainCompiler = ToolchainRoot;
 			if (PLATFORM_WINDOWS)
 			{
diff --git a/Engine/Source/Developer/OutputLog/Private/SDeviceOutputLog.cpp b/Engine/Source/Developer/OutputLog/Private/SDeviceOutputLog.cpp
--- a/Engine/Source/Developer/OutputLog/Private/SDeviceOutputLog.cpp
+++ b/Engine/Source/Developer/OutputLog/Private/SDeviceOutputLog.cpp
@@ -43,20 +43,19 @@ void SDeviceOutputLog::Construct( const FArguments& InArgs )
 			// The console input box
 			+SVerticalBox::Slot()
 			.AutoHeight()
-			.Padding(FMargin(0.0f, 4.0f, 0.0f, 0.0f))
+			.Padding(FMargin(0.0f, 2.0f, 0.0f, 0.0f))
 			[
 				SNew(SHorizontalBox)
 
 				+SHorizontalBox::Slot()
 				.AutoWidth()
 				[
-					SAssignNew(TargetDeviceComboBox, SComboBox<FTargetDeviceEntryPtr>)
-					.OptionsSource(&DeviceList)
-					.ButtonStyle(FEditorStyle::Get(), "ToolBar.Button")
-					.OnSelectionChanged(this, &SDeviceOutputLog::OnDeviceSelectionChanged)
-					.OnGenerateWidget(this, &SDeviceOutputLog::GenerateWidgetForDeviceComboBox)
+					SAssignNew(TargetDeviceComboButton, SComboButton)
+					.ComboButtonStyle(FEditorStyle::Get(), "OutputLog.Filters.Style")
+					.ForegroundColor(FLinearColor::White)
+					.OnGetMenuContent(this, &SDeviceOutputLog::MakeDeviceComboButtonMenu)
 					.ContentPadding(FMargin(4.0f, 0.0f))
-					.Content()
+					.ButtonContent()
 					[
 						SNew(SHorizontalBox)
 						+SHorizontalBox::Slot()
@@ -73,14 +72,17 @@ void SDeviceOutputLog::Construct( const FArguments& InArgs )
 						+SHorizontalBox::Slot()
 						.VAlign(VAlign_Center)
 						[
-							SNew(STextBlock).Text(this, &SDeviceOutputLog::GetSelectedTargetDeviceText)
+							SNew(STextBlock)
+							.TextStyle(FEditorStyle::Get(), "OutputLog.Filters.Text")
+							.Text(this, &SDeviceOutputLog::GetSelectedTargetDeviceText)
 						]
 					]
 				]
 
 				+SHorizontalBox::Slot()
-				.Padding(FMargin(4.0f, 0.0f, 0.0f, 4.0f))
+				.Padding(FMargin(4.0f, 0.0f, 0.0f, 0.0f))
 				.FillWidth(1)
+				.VAlign(VAlign_Center)
 				[
 					SNew(SConsoleInputBox)
 					.ConsoleCommandCustomExec(this, &SDeviceOutputLog::ExecuteConsoleCommand)
@@ -171,10 +173,9 @@ bool SDeviceOutputLog::CanBeUsedOnAnyThread() const
 
 void SDeviceOutputLog::ExecuteConsoleCommand(const FString& ExecCommand)
 {
-	FTargetDeviceEntryPtr SelectedDeviceEntry = TargetDeviceComboBox->GetSelectedItem();
-	if (SelectedDeviceEntry.IsValid())
+	if (CurrentDevicePtr.IsValid())
 	{
-		ITargetDevicePtr PinnedPtr = SelectedDeviceEntry->DeviceWeakPtr.Pin();
+		ITargetDevicePtr PinnedPtr = CurrentDevicePtr->DeviceWeakPtr.Pin();
 		if (PinnedPtr.IsValid())
 		{
 			PinnedPtr->ExecuteConsoleCommand(ExecCommand);
@@ -184,10 +185,9 @@ void SDeviceOutputLog::ExecuteConsoleCommand(const FString& ExecCommand)
 
 void SDeviceOutputLog::HandleTargetPlatformDeviceLost(ITargetDeviceRef LostDevice)
 {
-	auto LostDeviceId = LostDevice->GetId();
-	FTargetDeviceEntryPtr SelectedDeviceEntry = TargetDeviceComboBox->GetSelectedItem();
-
-	if (SelectedDeviceEntry.IsValid() && SelectedDeviceEntry->DeviceId == LostDeviceId)
+	FTargetDeviceId LostDeviceId = LostDevice->GetId();
+	
+	if (CurrentDevicePtr.IsValid() && CurrentDevicePtr->DeviceId == LostDeviceId)
 	{
 		// Kill device output object, but do not clean up output in the window
 		CurrentDeviceOutputPtr.Reset();
@@ -215,7 +215,7 @@ void SDeviceOutputLog::HandleTargetPlatformDeviceDiscovered(ITargetDeviceRef Dis
 	{
 		DeviceList[ExistingEntryIdx]->DeviceWeakPtr = DiscoveredDevice;
 		
-		if (TargetDeviceComboBox->GetSelectedItem() == DeviceList[ExistingEntryIdx])
+		if (CurrentDevicePtr.IsValid() && CurrentDevicePtr->DeviceId == DeviceList[ExistingEntryIdx]->DeviceId)
 		{
 			CurrentDeviceOutputPtr = DiscoveredDevice->CreateDeviceOutputRouter(this);
 		}
@@ -223,7 +223,6 @@ void SDeviceOutputLog::HandleTargetPlatformDeviceDiscovered(ITargetDeviceRef Dis
 	else
 	{
 		AddDeviceEntry(DiscoveredDevice);
-		TargetDeviceComboBox->RefreshOptions();
 	}
 }
 
@@ -242,10 +241,11 @@ void SDeviceOutputLog::AddDeviceEntry(ITargetDeviceRef TargetDevice)
 	DeviceList.Add(DeviceEntry);
 }
 
-void SDeviceOutputLog::OnDeviceSelectionChanged(FTargetDeviceEntryPtr DeviceEntry, ESelectInfo::Type SelectInfo)
+void SDeviceOutputLog::OnDeviceSelectionChanged(FTargetDeviceEntryPtr DeviceEntry)
 {
 	CurrentDeviceOutputPtr.Reset();
 	OnClearLog();
+	CurrentDevicePtr = DeviceEntry;
 	
 	if (DeviceEntry.IsValid())
 	{
@@ -257,7 +257,22 @@ void SDeviceOutputLog::OnDeviceSelectionChanged(FTargetDeviceEntryPtr DeviceEntr
 	}
 }
 
-TSharedRef<SWidget> SDeviceOutputLog::GenerateWidgetForDeviceComboBox(FTargetDeviceEntryPtr DeviceEntry) const
+TSharedRef<SWidget> SDeviceOutputLog::MakeDeviceComboButtonMenu()
+{
+	FMenuBuilder MenuBuilder(/*bInShouldCloseWindowAfterMenuSelection=*/true, nullptr);
+	for (const FTargetDeviceEntryPtr& TargetDeviceEntryPtr : DeviceList)
+	{
+		TSharedRef<SWidget> MenuEntryWidget = GenerateWidgetForDeviceComboBox(TargetDeviceEntryPtr);
+				
+		MenuBuilder.AddMenuEntry(
+			FUIAction(FExecuteAction::CreateSP(this, &SDeviceOutputLog::OnDeviceSelectionChanged, TargetDeviceEntryPtr)), 
+			MenuEntryWidget
+			);
+	}
+	return MenuBuilder.MakeWidget();
+}
+
+TSharedRef<SWidget> SDeviceOutputLog::GenerateWidgetForDeviceComboBox(const FTargetDeviceEntryPtr& DeviceEntry) const
 {
 	return 
 		SNew(SBox)
@@ -298,8 +313,7 @@ const FSlateBrush* SDeviceOutputLog::GetTargetDeviceBrush(FTargetDeviceEntryPtr
 
 const FSlateBrush* SDeviceOutputLog::GetSelectedTargetDeviceBrush() const
 {
-	FTargetDeviceEntryPtr DeviceEntry = TargetDeviceComboBox->GetSelectedItem();
-	return GetTargetDeviceBrush(DeviceEntry);
+	return GetTargetDeviceBrush(CurrentDevicePtr);
 }
 
 FText SDeviceOutputLog::GetTargetDeviceText(FTargetDeviceEntryPtr DeviceEntry) const
@@ -324,6 +338,5 @@ FText SDeviceOutputLog::GetTargetDeviceText(FTargetDeviceEntryPtr DeviceEntry) c
 
 FText SDeviceOutputLog::GetSelectedTargetDeviceText() const
 {
-	FTargetDeviceEntryPtr DeviceEntry = TargetDeviceComboBox->GetSelectedItem();
-	return GetTargetDeviceText(DeviceEntry);
+	return GetTargetDeviceText(CurrentDevicePtr);
 }
diff --git a/Engine/Source/Developer/OutputLog/Private/SDeviceOutputLog.h b/Engine/Source/Developer/OutputLog/Private/SDeviceOutputLog.h
--- a/Engine/Source/Developer/OutputLog/Private/SDeviceOutputLog.h
+++ b/Engine/Source/Developer/OutputLog/Private/SDeviceOutputLog.h
@@ -50,8 +50,9 @@ protected:
 
 	void AddDeviceEntry(ITargetDeviceRef TargetDevice);
 
-	void OnDeviceSelectionChanged(FTargetDeviceEntryPtr DeviceEntry, ESelectInfo::Type SelectInfo);
-	TSharedRef<SWidget> GenerateWidgetForDeviceComboBox(FTargetDeviceEntryPtr DeviceEntry) const;
+	void OnDeviceSelectionChanged(FTargetDeviceEntryPtr DeviceEntry);
+	TSharedRef<SWidget> MakeDeviceComboButtonMenu();
+	TSharedRef<SWidget> GenerateWidgetForDeviceComboBox(const FTargetDeviceEntryPtr& DeviceEntry) const;
 
 	FText GetTargetDeviceText(FTargetDeviceEntryPtr DeviceEntry) const;
 	FText GetSelectedTargetDeviceText() const;
@@ -61,10 +62,10 @@ protected:
 		
 private:
 	TArray<FTargetDeviceEntryPtr>			DeviceList;
-	ITargetDeviceWeakPtr					CurrentDevicePtr;
+	FTargetDeviceEntryPtr					CurrentDevicePtr;
 	ITargetDeviceOutputPtr					CurrentDeviceOutputPtr;
 
-	TSharedPtr<SComboBox<FTargetDeviceEntryPtr>> TargetDeviceComboBox;
+	TSharedPtr<SComboButton> TargetDeviceComboButton;
 		
 	/** Synchronization object for access to buffered lines */
 	FCriticalSection		BufferedLinesSynch;
diff --git a/Engine/Source/Editor/FoliageEdit/Private/ProceduralFoliageComponentDetails.cpp b/Engine/Source/Editor/FoliageEdit/Private/ProceduralFoliageComponentDetails.cpp
--- a/Engine/Source/Editor/FoliageEdit/Private/ProceduralFoliageComponentDetails.cpp
+++ b/Engine/Source/Editor/FoliageEdit/Private/ProceduralFoliageComponentDetails.cpp
@@ -67,16 +67,10 @@ void FProceduralFoliageComponentDetails::CustomizeDetails(IDetailLayoutBuilder&
 
 FReply FProceduralFoliageComponentDetails::OnResimulateClicked()
 {
-	TSet<UProceduralFoliageSpawner*> UniqueFoliageSpawners;
 	for (TWeakObjectPtr<UProceduralFoliageComponent>& Component : SelectedComponents)
 	{
 		if (Component.IsValid() && Component->FoliageSpawner)
 		{
-			if (!UniqueFoliageSpawners.Contains(Component->FoliageSpawner))
-			{
-				UniqueFoliageSpawners.Add(Component->FoliageSpawner);
-			}
-
 			FScopedTransaction Transaction(LOCTEXT("Resimulate_Transaction", "Procedural Foliage Simulation"));
 			TArray <FDesiredFoliageInstance> DesiredFoliageInstances;
 			if (Component->GenerateProceduralContent(DesiredFoliageInstances))
diff --git a/Engine/Source/Programs/AutomationTool/Android/AndroidPlatform.Automation.cs b/Engine/Source/Programs/AutomationTool/Android/AndroidPlatform.Automation.cs
--- a/Engine/Source/Programs/AutomationTool/Android/AndroidPlatform.Automation.cs
+++ b/Engine/Source/Programs/AutomationTool/Android/AndroidPlatform.Automation.cs
@@ -105,9 +105,25 @@ public class AndroidPlatform : Platform
 		}
     }
 
-	private static string GetFinalBatchName(string ApkName, ProjectParams Params, string Architecture, string GPUArchitecture, bool bNoOBBInstall, bool bUninstall, bool bIsPC)
+	private static string GetFinalBatchName(string ApkName, ProjectParams Params, string Architecture, string GPUArchitecture, bool bNoOBBInstall, bool bUninstall, UnrealTargetPlatform Target)
 	{
-		return Path.Combine(Path.GetDirectoryName(ApkName), (bUninstall ? "Uninstall_" : "Install_") + Params.ShortProjectName + (!bNoOBBInstall ? "_" : "_NoOBBInstall_") + Params.ClientConfigsToBuild[0].ToString() + Architecture + GPUArchitecture + (bIsPC ? ".bat" : ".command"));
+		string Extension = ".bat";
+		switch (Target)
+		{
+			default:
+			case UnrealTargetPlatform.Win64:
+				Extension = ".bat";
+				break;
+
+			case UnrealTargetPlatform.Linux:
+				Extension = ".sh";
+				break;
+
+			case UnrealTargetPlatform.Mac:
+				Extension = ".command";
+				break;
+		}
+		return Path.Combine(Path.GetDirectoryName(ApkName), (bUninstall ? "Uninstall_" : "Install_") + Params.ShortProjectName + (!bNoOBBInstall ? "_" : "_NoOBBInstall_") + Params.ClientConfigsToBuild[0].ToString() + Architecture + GPUArchitecture + Extension);
 	}
 
 	private List<string> CollectPluginDataPaths(DeploymentContext SC)
@@ -227,42 +243,48 @@ public class AndroidPlatform : Platform
 
 				//figure out which platforms we need to create install files for
 				bool bNeedsPCInstall = false;
-				bool bNeedsMonoInstall = false;
-				GetPlatformInstallOptions(SC, out bNeedsPCInstall, out bNeedsMonoInstall);
+				bool bNeedsMacInstall = false;
+				bool bNeedsLinuxInstall = false;
+				GetPlatformInstallOptions(SC, out bNeedsPCInstall, out bNeedsMacInstall, out bNeedsLinuxInstall);
 
 				//helper delegate to prevent code duplication but allow us access to all the local variables we need
-				var CreateInstallFilesAction = new Action<bool>(bPCInstall =>
+				var CreateInstallFilesAction = new Action<UnrealTargetPlatform>(Target =>
+				{
+					bool bIsPC = (Target == UnrealTargetPlatform.Win64);
+					// Write install batch file(s).
+					string BatchName = GetFinalBatchName(ApkName, Params, bMakeSeparateApks ? Architecture : "", bMakeSeparateApks ? GPUArchitecture : "", false, false, Target);
+					string PackageName = GetPackageInfo(ApkName, false);
+					// make a batch file that can be used to install the .apk and .obb files
+					string[] BatchLines = GenerateInstallBatchFile(bPackageDataInsideApk, PackageName, ApkName, Params, ObbName, DeviceObbName, false, bIsPC);
+					File.WriteAllLines(BatchName, BatchLines);
+					// make a batch file that can be used to uninstall the .apk and .obb files
+					string UninstallBatchName = GetFinalBatchName(ApkName, Params, bMakeSeparateApks ? Architecture : "", bMakeSeparateApks ? GPUArchitecture : "", false, true, Target);
+					BatchLines = GenerateUninstallBatchFile(bPackageDataInsideApk, PackageName, ApkName, Params, ObbName, DeviceObbName, false, bIsPC);
+					File.WriteAllLines(UninstallBatchName, BatchLines);
+
+					if (Utils.IsRunningOnMono)
 					{
-						// Write install batch file(s).
-						string BatchName = GetFinalBatchName(ApkName, Params, bMakeSeparateApks ? Architecture : "", bMakeSeparateApks ? GPUArchitecture : "", false, false, bPCInstall);
-						string PackageName = GetPackageInfo(ApkName, false);
-						// make a batch file that can be used to install the .apk and .obb files
-						string[] BatchLines = GenerateInstallBatchFile(bPackageDataInsideApk, PackageName, ApkName, Params, ObbName, DeviceObbName, false, bPCInstall);
-						File.WriteAllLines(BatchName, BatchLines);
-						// make a batch file that can be used to uninstall the .apk and .obb files
-						string UninstallBatchName = GetFinalBatchName(ApkName, Params, bMakeSeparateApks ? Architecture : "", bMakeSeparateApks ? GPUArchitecture : "", false, true, bPCInstall);
-						BatchLines = GenerateUninstallBatchFile(bPackageDataInsideApk, PackageName, ApkName, Params, ObbName, DeviceObbName, false, bPCInstall);
-						File.WriteAllLines(UninstallBatchName, BatchLines);
-
-						if (Utils.IsRunningOnMono)
-						{
-							CommandUtils.FixUnixFilePermissions(BatchName);
-							CommandUtils.FixUnixFilePermissions(UninstallBatchName);
-							//if(File.Exists(NoInstallBatchName)) 
-							//{
-							//    CommandUtils.FixUnixFilePermissions(NoInstallBatchName);
-							//}
-						}
+						CommandUtils.FixUnixFilePermissions(BatchName);
+						CommandUtils.FixUnixFilePermissions(UninstallBatchName);
+						//if(File.Exists(NoInstallBatchName)) 
+						//{
+						//    CommandUtils.FixUnixFilePermissions(NoInstallBatchName);
+						//}
 					}
+				}
 				);
 
 				if (bNeedsPCInstall)
 				{
-					CreateInstallFilesAction.Invoke(true);
+					CreateInstallFilesAction.Invoke(UnrealTargetPlatform.Win64);
 				}
-				if(bNeedsMonoInstall)
+				if (bNeedsMacInstall)
 				{
-					CreateInstallFilesAction.Invoke(false);
+					CreateInstallFilesAction.Invoke(UnrealTargetPlatform.Mac);
+				}
+				if (bNeedsLinuxInstall)
+				{
+					CreateInstallFilesAction.Invoke(UnrealTargetPlatform.Linux);
 				}
 
 				// If we aren't packaging data in the APK then lets write out a bat file to also let us test without the OBB
@@ -464,48 +486,65 @@ public class AndroidPlatform : Platform
 				}
 
 				bool bNeedsPCInstall = false;
-				bool bNeedsMonoInstall = false;
-				GetPlatformInstallOptions(SC, out bNeedsPCInstall, out bNeedsMonoInstall);
+				bool bNeedsMacInstall = false;
+				bool bNeedsLinuxInstall = false;
+				GetPlatformInstallOptions(SC, out bNeedsPCInstall, out bNeedsMacInstall, out bNeedsLinuxInstall);
 
 				//helper delegate to prevent code duplication but allow us access to all the local variables we need
-				var CreateBatchFilesAndArchiveAction = new Action<bool>(bPCInstall =>
-					{
-						string BatchName = GetFinalBatchName(ApkName, Params, bMakeSeparateApks ? Architecture : "", bMakeSeparateApks ? GPUArchitecture : "", false, false, bPCInstall);
-						string UninstallBatchName = GetFinalBatchName(ApkName, Params, bMakeSeparateApks ? Architecture : "", bMakeSeparateApks ? GPUArchitecture : "", false, true, bPCInstall);
+				var CreateBatchFilesAndArchiveAction = new Action<UnrealTargetPlatform>(Target =>
+				{
+					string BatchName = GetFinalBatchName(ApkName, Params, bMakeSeparateApks ? Architecture : "", bMakeSeparateApks ? GPUArchitecture : "", false, false, Target);
+					string UninstallBatchName = GetFinalBatchName(ApkName, Params, bMakeSeparateApks ? Architecture : "", bMakeSeparateApks ? GPUArchitecture : "", false, true, Target);
 
-						SC.ArchiveFiles(Path.GetDirectoryName(BatchName), Path.GetFileName(BatchName));
-						SC.ArchiveFiles(Path.GetDirectoryName(UninstallBatchName), Path.GetFileName(UninstallBatchName));
-						//SC.ArchiveFiles(Path.GetDirectoryName(NoOBBBatchName), Path.GetFileName(NoOBBBatchName));
-					}
+					SC.ArchiveFiles(Path.GetDirectoryName(BatchName), Path.GetFileName(BatchName));
+					SC.ArchiveFiles(Path.GetDirectoryName(UninstallBatchName), Path.GetFileName(UninstallBatchName));
+					//SC.ArchiveFiles(Path.GetDirectoryName(NoOBBBatchName), Path.GetFileName(NoOBBBatchName));
+				}
 				);
 
 				//it's possible we will need both PC and Mac/Linux install files, do both
 				if (bNeedsPCInstall)
 				{
-					CreateBatchFilesAndArchiveAction(true);
+					CreateBatchFilesAndArchiveAction(UnrealTargetPlatform.Win64);
 				}
-				if(bNeedsMonoInstall)
+				if (bNeedsMacInstall)
 				{
-					CreateBatchFilesAndArchiveAction(false);
+					CreateBatchFilesAndArchiveAction(UnrealTargetPlatform.Mac);
+				}
+				if (bNeedsLinuxInstall)
+				{
+					CreateBatchFilesAndArchiveAction(UnrealTargetPlatform.Linux);
 				}
 			}
 		}
 	}
 
-	private void GetPlatformInstallOptions(DeploymentContext SC, out bool bNeedsPCInstall, out bool bNeedsMonoInstall)
+	private void GetPlatformInstallOptions(DeploymentContext SC, out bool bNeedsPCInstall, out bool bNeedsMacInstall, out bool bNeedsLinuxInstall)
 	{
-		ConfigCacheIni Ini = ConfigCacheIni.CreateConfigCacheIni(SC.StageTargetPlatform.PlatformType, "Engine", DirectoryReference.FromFile(SC.RawProjectPath));		
+		ConfigCacheIni Ini = ConfigCacheIni.CreateConfigCacheIni(SC.StageTargetPlatform.PlatformType, "Engine", DirectoryReference.FromFile(SC.RawProjectPath));
 		bool bGenerateAllPlatformInstall = false;
-	 	Ini.GetBool("/Script/AndroidRuntimeSettings.AndroidRuntimeSettings", "bCreateAllPlatformsInstall", out bGenerateAllPlatformInstall);
+		Ini.GetBool("/Script/AndroidRuntimeSettings.AndroidRuntimeSettings", "bCreateAllPlatformsInstall", out bGenerateAllPlatformInstall);
+
+		bNeedsPCInstall = bNeedsMacInstall = bNeedsLinuxInstall = false;
 
-		if(bGenerateAllPlatformInstall)
+		if (bGenerateAllPlatformInstall)
 		{
-			bNeedsMonoInstall = bNeedsPCInstall = true;
+			bNeedsPCInstall = bNeedsMacInstall = bNeedsLinuxInstall = true;
 		}
 		else
 		{
-			bNeedsMonoInstall = Utils.IsRunningOnMono;
-			bNeedsPCInstall = !bNeedsMonoInstall;
+			if (HostPlatform.Current.HostEditorPlatform == UnrealTargetPlatform.Mac)
+			{
+				bNeedsMacInstall = true;
+			}
+			else if (HostPlatform.Current.HostEditorPlatform == UnrealTargetPlatform.Linux)
+			{
+				bNeedsLinuxInstall = true;
+			}
+			else
+			{
+				bNeedsPCInstall = true;
+			}
 		}
 	}
 
@@ -1345,25 +1384,16 @@ public class AndroidPlatform : Platform
 
 	public override void GetFilesToDeployOrStage(ProjectParams Params, DeploymentContext SC)
 	{
-		// 		if (SC.StageExecutables.Count != 1 && Params.Package)
-		// 		{
-		// 			throw new AutomationException("Exactly one executable expected when staging Android. Had " + SC.StageExecutables.Count.ToString());
-		// 		}
-		// 
-		// 		// stage all built executables
-		// 		foreach (var Exe in SC.StageExecutables)
-		// 		{
-		// 			string ApkName = Exe + GetArchitecture(Params) + ".apk";
-		// 
-		// 			SC.StageFiles(StagedFileType.NonUFS, Params.ProjectBinariesFolder, ApkName);
-		// 		}
-	}
+        // Add any Android shader cache files
+        string ProjectShaderDir = Path.Combine(Path.GetDirectoryName(Path.GetFullPath(Params.RawProjectPath.ToString())), "Build/ShaderCaches/Android");
+        SC.StageFiles(StagedFileType.UFS, ProjectShaderDir, "*.*", true, null, null, true);
+    }
 
-	/// <summary>
-	/// Gets cook platform name for this platform.
-	/// </summary>
-	/// <returns>Cook platform string.</returns>
-	public override string GetCookPlatform(bool bDedicatedServer, bool bIsClientOnly)
+    /// <summary>
+    /// Gets cook platform name for this platform.
+    /// </summary>
+    /// <returns>Cook platform string.</returns>
+    public override string GetCookPlatform(bool bDedicatedServer, bool bIsClientOnly)
 	{
 		return "Android";
 	}
diff --git a/Engine/Source/Programs/AutomationTool/Scripts/CopyBuildToStagingDirectory.Automation.cs b/Engine/Source/Programs/AutomationTool/Scripts/CopyBuildToStagingDirectory.Automation.cs
--- a/Engine/Source/Programs/AutomationTool/Scripts/CopyBuildToStagingDirectory.Automation.cs
+++ b/Engine/Source/Programs/AutomationTool/Scripts/CopyBuildToStagingDirectory.Automation.cs
@@ -1061,7 +1061,6 @@ public partial class Project : CommandUtils
 
 				string BuildRoot = MakePathSafeToUseWithCommandLine(RawDataPath);
 				string CloudDir = MakePathSafeToUseWithCommandLine(CombinePaths(ChunkInstallBasePath, "CloudDir"));
-                InternalUtils.SafeDeleteDirectory(CloudDir, true);
 				string ManifestDir = CombinePaths(ChunkInstallBasePath, "ManifestDir");
 				var AppID = 1; // For a chunk install this value doesn't seem to matter
 				string AppName = String.Format("{0}_{1}", SC.ShortProjectName, PakName);
diff --git a/Engine/Source/Programs/GitDependencies/Program.cs b/Engine/Source/Programs/GitDependencies/Program.cs
--- a/Engine/Source/Programs/GitDependencies/Program.cs
+++ b/Engine/Source/Programs/GitDependencies/Program.cs
@@ -186,7 +186,7 @@ namespace GitDependencies
 				{
 					ExcludeFolders.Add("Android");
 				}
-				if(Environment.OSVersion.Platform == PlatformID.Win32NT && Environment.GetEnvironmentVariable("LINUX_ROOT") == null)
+				if(Environment.OSVersion.Platform == PlatformID.Win32NT && (Environment.GetEnvironmentVariable("LINUX_ROOT") == null && Environment.GetEnvironmentVariable("LINUX_MULTIARCH_ROOT") == null))
 				{
 					ExcludeFolders.Add("Linux");
 				}
diff --git a/Engine/Source/Programs/UnrealBuildTool/Android/UEDeployAndroid.cs b/Engine/Source/Programs/UnrealBuildTool/Android/UEDeployAndroid.cs
--- a/Engine/Source/Programs/UnrealBuildTool/Android/UEDeployAndroid.cs
+++ b/Engine/Source/Programs/UnrealBuildTool/Android/UEDeployAndroid.cs
@@ -1933,7 +1933,7 @@ namespace UnrealBuildTool
 					NDKArches.Add(NDKArch);
 				}
 			}
-			UPL.Init(NDKArches, bForDistribution, EngineDirectory, UE4BuildPath);
+			UPL.Init(NDKArches, bForDistribution, EngineDirectory, UE4BuildPath, ProjectDirectory);
 
 			IEnumerable<Tuple<string, string, string>> BuildList = null;
 
diff --git a/Engine/Source/Programs/UnrealBuildTool/IOS/UEDeployIOS.cs b/Engine/Source/Programs/UnrealBuildTool/IOS/UEDeployIOS.cs
--- a/Engine/Source/Programs/UnrealBuildTool/IOS/UEDeployIOS.cs
+++ b/Engine/Source/Programs/UnrealBuildTool/IOS/UEDeployIOS.cs
@@ -818,7 +818,7 @@ namespace UnrealBuildTool
 			string BundlePath = Path.Combine (ProjectDirectory, "Binaries", "IOS", "Payload", ProjectBaseName + ".app");
 
 			// Passing in true for distribution is not ideal here but given the way that ios packaging happens and this call chain it seems unavoidable for now, maybe there is a way to correctly pass it in that I can't find?
-			UPL.Init (PlatformContext.ProjectArches, true, BuildConfiguration.RelativeEnginePath, BundlePath);
+			UPL.Init (PlatformContext.ProjectArches, true, BuildConfiguration.RelativeEnginePath, BundlePath, ProjectDirectory);
 
 			if (BuildHostPlatform.Current.Platform == UnrealTargetPlatform.Mac && Environment.GetEnvironmentVariable("UBT_NO_POST_DEPLOY") != "true")
 			{
diff --git a/Engine/Source/Programs/UnrealBuildTool/Linux/LinuxToolChain.cs b/Engine/Source/Programs/UnrealBuildTool/Linux/LinuxToolChain.cs
--- a/Engine/Source/Programs/UnrealBuildTool/Linux/LinuxToolChain.cs
+++ b/Engine/Source/Programs/UnrealBuildTool/Linux/LinuxToolChain.cs
@@ -21,7 +21,7 @@ namespace UnrealBuildTool
 			if (!CrossCompiling())
 			{
 				// use native linux toolchain
-				string[] ClangNames = { "clang++", "clang++-3.8", "clang++-3.7", "clang++-3.6", "clang++-3.5" };
+				string[] ClangNames = { "clang++", "clang++-3.9", "clang++-3.8", "clang++-3.7", "clang++-3.6", "clang++-3.5" };
 				foreach (var ClangName in ClangNames)
 				{
 					ClangPath = Which(ClangName);
@@ -43,12 +43,29 @@ namespace UnrealBuildTool
 			}
 			else
 			{
-				// use cross linux toolchain if LINUX_ROOT is specified
-				BaseLinuxPath = Environment.GetEnvironmentVariable("LINUX_ROOT");
+				// if new multi-arch toolchain is used, prefer it
+				string MultiArchRoot = Environment.GetEnvironmentVariable("LINUX_MULTIARCH_ROOT");
+
+				if (MultiArchRoot != null)
+				{
+					// FIXME: UBT should loop across all the architectures and compile for all the selected ones.
+					BaseLinuxPath = Path.Combine(MultiArchRoot, LinuxPlatform.DefaultArchitecture);
+
+					Console.WriteLine("Using LINUX_MULTIARCH_ROOT, building with toolchain '{0}'", BaseLinuxPath);
+				}
+				else
+				{
+					// use cross linux toolchain if LINUX_ROOT is specified
+					BaseLinuxPath = Environment.GetEnvironmentVariable("LINUX_ROOT");
+
+					Console.WriteLine("Using LINUX_ROOT (deprecated, consider LINUX_MULTIARCH_ROOT), building with toolchain '{0}'", BaseLinuxPath);
+				}
 
 				// don't register if we don't have an LINUX_ROOT specified
 				if (String.IsNullOrEmpty(BaseLinuxPath))
+				{
 					throw new BuildException("LINUX_ROOT environment variable is not set; cannot instantiate Linux toolchain");
+				}
 
 				BaseLinuxPath = BaseLinuxPath.Replace("\"", "");
 
@@ -362,6 +379,15 @@ namespace UnrealBuildTool
 					Result += " -Wno-unused-local-typedef";	// clang is being overly strict here? PhysX headers trigger this.
 					Result += " -Wno-inconsistent-missing-override";	// these have to be suppressed for UE 4.8, should be fixed later.
 				}
+
+				if (CompilerVersionGreaterOrEqual(3, 9, 0))
+				{
+					Result += " -Wno-undefined-var-template"; // not really a good warning to disable
+					Result += " -Wno-delete-non-virtual-dtor";  // at least for 4.14
+					Result += " -Wno-expansion-to-defined";		// at least for 4.14
+					Result += " -Wno-null-dereference";         // at least for 4.14
+					Result += " -Wno-literal-conversion";       // at least for 4.14
+				}
 			}
 
 			Result += " -Wno-unused-variable";
diff --git a/Engine/Source/Programs/UnrealBuildTool/Linux/UEBuildLinux.cs b/Engine/Source/Programs/UnrealBuildTool/Linux/UEBuildLinux.cs
--- a/Engine/Source/Programs/UnrealBuildTool/Linux/UEBuildLinux.cs
+++ b/Engine/Source/Programs/UnrealBuildTool/Linux/UEBuildLinux.cs
@@ -335,11 +335,7 @@ namespace UnrealBuildTool
 		/// <summary>
 		/// This is the SDK version we support
 		/// </summary>
-		static private Dictionary<string, string> ExpectedSDKVersions = new Dictionary<string, string>()
-		{
-			{ "x86_64-unknown-linux-gnu", "v7_clang-3.7.0_ld-2.24_glibc-2.12.2" },
-			{ "arm-unknown-linux-gnueabihf", "arm-unknown-linux-gnueabihf_v5_clang-3.5.0-ld-2.23.1-glibc-2.13" },
-		};
+		static string ExpectedSDKVersion = "v8_clang-3.9.0-centos7";	// now unified for all the architectures
 
 		/// <summary>
 		/// Platform name (embeds architecture for now)
@@ -370,13 +366,7 @@ namespace UnrealBuildTool
 		/// <returns>Valid SDK string</returns>
 		protected override string GetRequiredSDKString()
 		{
-			string SDKString;
-			if (!ExpectedSDKVersions.TryGetValue(LinuxPlatform.DefaultArchitecture, out SDKString))
-			{
-				throw new BuildException("LinuxPlatform::GetRequiredSDKString: no toolchain set up for architecture '{0}'", LinuxPlatform.DefaultArchitecture);
-			}
-
-			return SDKString;
+			return ExpectedSDKVersion;
 		}
 
 		protected override String GetRequiredScriptVersionString()
@@ -400,18 +390,34 @@ namespace UnrealBuildTool
 				return SDKStatus.Valid;
 			}
 
-			string BaseLinuxPath = Environment.GetEnvironmentVariable("LINUX_ROOT");
+			string MultiArchRoot = Environment.GetEnvironmentVariable("LINUX_MULTIARCH_ROOT");
+			string BaseLinuxPath;
+
+			if (MultiArchRoot != null)
+			{
+				// FIXME: UBT should loop across all the architectures and compile for all the selected ones.
+				BaseLinuxPath = Path.Combine(MultiArchRoot, LinuxPlatform.DefaultArchitecture);
+			}
+			else
+			{
+				// support the existing, non-multiarch toolchains for continuity
+				BaseLinuxPath = Environment.GetEnvironmentVariable("LINUX_ROOT");
+			}
 
 			// we don't have an LINUX_ROOT specified
 			if (String.IsNullOrEmpty(BaseLinuxPath))
+			{
 				return SDKStatus.Invalid;
+			}
 
 			// paths to our toolchains
 			BaseLinuxPath = BaseLinuxPath.Replace("\"", "");
-			string ClangPath = Path.Combine(BaseLinuxPath, @"bin\Clang++.exe");
+			string ClangPath = Path.Combine(BaseLinuxPath, @"bin\clang++.exe");
 
 			if (File.Exists(ClangPath))
+			{
 				return SDKStatus.Valid;
+			}
 
 			return SDKStatus.Invalid;
 		}
diff --git a/Engine/Source/Programs/UnrealBuildTool/System/UnrealPluginLanguage.cs b/Engine/Source/Programs/UnrealBuildTool/System/UnrealPluginLanguage.cs
--- a/Engine/Source/Programs/UnrealBuildTool/System/UnrealPluginLanguage.cs
+++ b/Engine/Source/Programs/UnrealBuildTool/System/UnrealPluginLanguage.cs
@@ -2117,11 +2117,12 @@ namespace UnrealBuildTool
 			return GlobalContext.StringVariables["Output"];
 		}
 
-		public void Init(List<string> Architectures, bool bDistribution, string EngineDirectory, string BuildDirectory)
+		public void Init(List<string> Architectures, bool bDistribution, string EngineDirectory, string BuildDirectory, string ProjectDirectory)
 		{
 			GlobalContext.BoolVariables["Distribution"] = bDistribution;
 			GlobalContext.StringVariables["EngineDir"] = EngineDirectory;
 			GlobalContext.StringVariables["BuildDir"] = BuildDirectory;
+			GlobalContext.StringVariables["ProjectDir"] = ProjectDirectory;
 
 			foreach (string Arch in Architectures)
 			{
diff --git a/Engine/Source/Programs/UnrealBuildTool/System/XGE.cs b/Engine/Source/Programs/UnrealBuildTool/System/XGE.cs
--- a/Engine/Source/Programs/UnrealBuildTool/System/XGE.cs
+++ b/Engine/Source/Programs/UnrealBuildTool/System/XGE.cs
@@ -298,6 +298,10 @@ namespace UnrealBuildTool
 				{
 					ToolElement.SetAttribute("AutoReserveMemory", "*.gch");
 				}
+				else
+				{
+					ToolElement.SetAttribute("AutoReserveMemory", "*.pch");
+				}
 				ToolElement.SetAttribute(
 					"OutputFileMasks",
 					string.Join(
diff --git a/Engine/Source/Runtime/Android/AndroidAudio/Private/AndroidAudioSource.cpp b/Engine/Source/Runtime/Android/AndroidAudio/Private/AndroidAudioSource.cpp
--- a/Engine/Source/Runtime/Android/AndroidAudio/Private/AndroidAudioSource.cpp
+++ b/Engine/Source/Runtime/Android/AndroidAudio/Private/AndroidAudioSource.cpp
@@ -163,6 +163,7 @@ bool FSLESSoundSource::EnqueuePCMBuffer( bool bLoop)
 
 	bStreamedSound = false;
 	bHasLooped = false;
+	bHasPositionUpdated = false;
 	bBuffersToFlush = false;
 
 	return true;
@@ -252,6 +253,7 @@ bool FSLESSoundSource::EnqueuePCMRTBuffer( bool bLoop )
 	bStreamedSound = true;
 	bHasLooped = false;
 	bBuffersToFlush = false;
+	bHasPositionUpdated = false;
 	BufferInUse = 1;
 	return true;
 }
@@ -358,6 +360,7 @@ FSLESSoundSource::FSLESSoundSource( class FAudioDevice* InAudioDevice )
 		BufferInUse(0),
 		VolumePreviousUpdate(-1.0f),
 		bHasLooped(false),
+		bHasPositionUpdated(false),
 		SL_PlayerObject(NULL),
 		SL_PlayerPlayInterface(NULL),
 		SL_PlayerBufferQueue(NULL),
@@ -555,6 +558,32 @@ bool FSLESSoundSource::IsSourceFinished( void )
 		return true;
 	}
 
+	if (WaveInstance && WaveInstance->LoopingMode == LOOP_Never)
+	{
+		// if it wasn't that simple, see if we're at the end position
+		SLmillisecond PositionMs;
+		SLmillisecond DurationMs;
+
+		result = (*SL_PlayerPlayInterface)->GetPosition(SL_PlayerPlayInterface, &PositionMs);
+		check(SL_RESULT_SUCCESS == result);
+
+		result = (*SL_PlayerPlayInterface)->GetDuration(SL_PlayerPlayInterface, &DurationMs);
+		check(SL_RESULT_SUCCESS == result);
+
+		// on some android devices, the value for GetPosition wraps back to 0 when the playback is done, however it's very possible
+		// for us to try to check for IsSourceFinished when the Position is genuinely "0". Therefore, we'll flip bHasPositionUpdated once
+		// we've actually started the sound to denote a wrap-back 0 position versus a real 0 position
+		if ((DurationMs != SL_TIME_UNKNOWN && PositionMs == DurationMs) || (PositionMs == 0 && bHasPositionUpdated))
+		{
+			return true;
+		}
+		else if (!bHasPositionUpdated && PositionMs > 0)
+		{
+			bHasPositionUpdated = true;
+		}
+	}
+	
+
 	return false;
 }
 
diff --git a/Engine/Source/Runtime/Android/AndroidAudio/Public/AndroidAudioDevice.h b/Engine/Source/Runtime/Android/AndroidAudio/Public/AndroidAudioDevice.h
--- a/Engine/Source/Runtime/Android/AndroidAudio/Public/AndroidAudioDevice.h
+++ b/Engine/Source/Runtime/Android/AndroidAudio/Public/AndroidAudioDevice.h
@@ -241,6 +241,8 @@ protected:
 	int32						BufferInUse;
 	float						VolumePreviousUpdate;
 	bool						bHasLooped;
+	/** Lets us know if calls to SL_PlayerPlayInterface->GetPosition have progressed past 0 at some point so we can determine when a non-looping sound has finished */
+	bool						bHasPositionUpdated;
 
 	bool CreatePlayer();
 	void DestroyPlayer();
diff --git a/Engine/Source/Runtime/Core/Core.Build.cs b/Engine/Source/Runtime/Core/Core.Build.cs
--- a/Engine/Source/Runtime/Core/Core.Build.cs
+++ b/Engine/Source/Runtime/Core/Core.Build.cs
@@ -97,7 +97,7 @@ public class Core : ModuleRules
 			AddEngineThirdPartyPrivateStaticDependencies(Target, 
 				"zlib"
 				);
-			PublicFrameworks.AddRange(new string[] { "UIKit", "Foundation", "AudioToolbox", "AVFoundation", "GameKit", "StoreKit", "CoreVideo", "CoreMedia", "CoreGraphics", "GameController"});
+			PublicFrameworks.AddRange(new string[] { "UIKit", "Foundation", "AudioToolbox", "AVFoundation", "GameKit", "StoreKit", "CoreVideo", "CoreMedia", "CoreGraphics", "GameController", "SystemConfiguration" });
 			if (Target.Platform == UnrealTargetPlatform.IOS)
 			{
 				PublicFrameworks.AddRange(new string[] { "CoreMotion" });
diff --git a/Engine/Source/Runtime/Core/Private/Android/AndroidMisc.cpp b/Engine/Source/Runtime/Core/Private/Android/AndroidMisc.cpp
--- a/Engine/Source/Runtime/Core/Private/Android/AndroidMisc.cpp
+++ b/Engine/Source/Runtime/Core/Private/Android/AndroidMisc.cpp
@@ -1538,6 +1538,12 @@ int FAndroidMisc::GetVolumeState(double* OutTimeOfChangeInSec)
 	return v;
 }
 
+const TCHAR* FAndroidMisc::GamePersistentDownloadDir()
+{
+	extern FString GExternalFilePath;
+	return *GExternalFilePath;
+}
+
 FAndroidMisc::FBatteryState FAndroidMisc::GetBatteryState()
 {
 	FBatteryState CurState;
@@ -1552,3 +1558,9 @@ bool FAndroidMisc::AreHeadPhonesPluggedIn()
 	return HeadPhonesArePluggedIn;
 }
 
+bool FAndroidMisc::HasActiveWiFiConnection()
+{
+	extern bool AndroidThunkCpp_HasActiveWiFiConnection();
+	return AndroidThunkCpp_HasActiveWiFiConnection();
+}
+
diff --git a/Engine/Source/Runtime/Core/Private/IOS/IOSPlatformMisc.cpp b/Engine/Source/Runtime/Core/Private/IOS/IOSPlatformMisc.cpp
--- a/Engine/Source/Runtime/Core/Private/IOS/IOSPlatformMisc.cpp
+++ b/Engine/Source/Runtime/Core/Private/IOS/IOSPlatformMisc.cpp
@@ -15,7 +15,10 @@
 
 #include "Apple/ApplePlatformCrashContext.h"
 
-/** Amount of free memory in MB reported by the system at starup */
+#include <SystemConfiguration/SystemConfiguration.h>
+#include <netinet/in.h>
+
+/** Amount of free memory in MB reported by the system at startup */
 CORE_API int32 GStartupFreeMemoryMB;
 
 /** Global pointer to memory warning handler */
@@ -910,6 +913,32 @@ void FIOSPlatformMisc::GetValidTargetPlatforms(TArray<FString>& TargetPlatformNa
 #endif
 }
 
+bool FIOSPlatformMisc::HasActiveWiFiConnection()
+{
+	struct sockaddr_in ZeroAddress;
+	FMemory::Memzero(&ZeroAddress, sizeof(ZeroAddress));
+	ZeroAddress.sin_len = sizeof(ZeroAddress);
+	ZeroAddress.sin_family = AF_INET;
+
+	SCNetworkReachabilityRef ReachabilityRef = SCNetworkReachabilityCreateWithAddress(kCFAllocatorDefault, (const struct sockaddr*)&ZeroAddress);
+	SCNetworkReachabilityFlags ReachabilityFlags;
+	bool bFlagsAvailable = SCNetworkReachabilityGetFlags(ReachabilityRef, &ReachabilityFlags);
+	CFRelease(ReachabilityRef);
+	
+	bool bHasActiveWiFiConnection = false;
+	if (bFlagsAvailable)
+	{
+        bool bReachable =	(ReachabilityFlags & kSCNetworkReachabilityFlagsReachable) != 0 && 
+							(ReachabilityFlags & kSCNetworkReachabilityFlagsConnectionRequired) == 0 &&
+							// in case kSCNetworkReachabilityFlagsConnectionOnDemand  || kSCNetworkReachabilityFlagsConnectionOnTraffic
+							(ReachabilityFlags & kSCNetworkReachabilityFlagsInterventionRequired) == 0; 
+					
+		bHasActiveWiFiConnection = bReachable && (ReachabilityFlags & kSCNetworkReachabilityFlagsIsWWAN) == 0;
+	}
+	
+	return bHasActiveWiFiConnection; 
+}
+
 void FIOSPlatformMisc::ResetGamepadAssignments()
 {
 	UE_LOG(LogIOS, Warning, TEXT("Restting gamepad assignments is not allowed in IOS"))
diff --git a/Engine/Source/Runtime/Core/Public/Android/AndroidMisc.h b/Engine/Source/Runtime/Core/Public/Android/AndroidMisc.h
--- a/Engine/Source/Runtime/Core/Public/Android/AndroidMisc.h
+++ b/Engine/Source/Runtime/Core/Public/Android/AndroidMisc.h
@@ -83,6 +83,7 @@ struct CORE_API FAndroidMisc : public FGenericPlatformMisc
 	static bool IsControllerAssignedToGamepad(int32 ControllerId);
 	// Returns current volume, 0-100 (%)
 	static int GetVolumeState(double* OutTimeOfChangeInSec = nullptr);
+	static const TCHAR* GamePersistentDownloadDir();
 
 	enum EBatteryState
 	{
@@ -101,6 +102,7 @@ struct CORE_API FAndroidMisc : public FGenericPlatformMisc
 
 	static FBatteryState GetBatteryState();
 	static bool AreHeadPhonesPluggedIn();
+	static bool HasActiveWiFiConnection();
 
 	/** @return Memory representing a true type or open type font provided by the platform as a default font for unreal to consume; empty array if the default font failed to load. */
 	static TArray<uint8> GetSystemFontBytes();
@@ -118,6 +120,7 @@ struct CORE_API FAndroidMisc : public FGenericPlatformMisc
 	static FString GetGLVersion();
 	static bool SupportsFloatingPointRenderTargets();
 	static bool SupportsShaderFramebufferFetch();
+	static bool SupportsShaderIOBlocks();
 	static int GetAndroidBuildVersion();
 	static bool ShouldUseVulkan();
 	static FString GetVulkanVersion();
diff --git a/Engine/Source/Runtime/Core/Public/GenericPlatform/GenericPlatformMisc.h b/Engine/Source/Runtime/Core/Public/GenericPlatform/GenericPlatformMisc.h
--- a/Engine/Source/Runtime/Core/Public/GenericPlatform/GenericPlatformMisc.h
+++ b/Engine/Source/Runtime/Core/Public/GenericPlatform/GenericPlatformMisc.h
@@ -916,6 +916,14 @@ public:
 		return PLATFORM_HAS_TOUCH_MAIN_SCREEN;
 	}
 
+	/**
+	 * Returns whether WiFi connection is currently active
+	 */
+	static bool HasActiveWiFiConnection()
+	{
+		return false;
+	}
+
 	/**
 	 * Returns whether the given platform feature is currently available (for instance, Metal is only available in IOS8 and with A7 devices)
 	 */
diff --git a/Engine/Source/Runtime/Core/Public/IOS/IOSPlatformMisc.h b/Engine/Source/Runtime/Core/Public/IOS/IOSPlatformMisc.h
--- a/Engine/Source/Runtime/Core/Public/IOS/IOSPlatformMisc.h
+++ b/Engine/Source/Runtime/Core/Public/IOS/IOSPlatformMisc.h
@@ -132,6 +132,7 @@ struct CORE_API FIOSPlatformMisc : public FGenericPlatformMisc
 	static FString GetLocalCurrencyCode();
 	static FString GetLocalCurrencySymbol();
 	static void GetValidTargetPlatforms(class TArray<class FString>& TargetPlatformNames);
+	static bool HasActiveWiFiConnection();
 
 	static void ResetGamepadAssignments();
 	static void ResetGamepadAssignmentToController(int32 ControllerId);
diff --git a/Engine/Source/Runtime/Engine/Classes/Components/ReflectionCaptureComponent.h b/Engine/Source/Runtime/Engine/Classes/Components/ReflectionCaptureComponent.h
--- a/Engine/Source/Runtime/Engine/Classes/Components/ReflectionCaptureComponent.h
+++ b/Engine/Source/Runtime/Engine/Classes/Components/ReflectionCaptureComponent.h
@@ -159,6 +159,8 @@ class UReflectionCaptureComponent : public USceneComponent
 	virtual void CreateRenderState_Concurrent() override;
 	virtual void DestroyRenderState_Concurrent() override;
 	virtual void SendRenderTransform_Concurrent() override;
+	virtual void OnRegister() override;
+	virtual void OnUnregister() override;
 	//~ End UActorComponent Interface
 
 	//~ Begin UObject Interface
diff --git a/Engine/Source/Runtime/Engine/Classes/Engine/EngineTypes.h b/Engine/Source/Runtime/Engine/Classes/Engine/EngineTypes.h
--- a/Engine/Source/Runtime/Engine/Classes/Engine/EngineTypes.h
+++ b/Engine/Source/Runtime/Engine/Classes/Engine/EngineTypes.h
@@ -314,6 +314,7 @@ UENUM()
 enum ESceneCaptureSource 
 { 
 	SCS_SceneColorHDR UMETA(DisplayName="SceneColor (HDR) in RGB, Inv Opacity in A"),
+	SCS_SceneColorHDRNoAlpha UMETA(DisplayName="SceneColor (HDR) in RGB, 0 in A"),
 	SCS_FinalColorLDR UMETA(DisplayName="Final Color (LDR) in RGB"),
 	SCS_SceneColorSceneDepth UMETA(DisplayName="SceneColor (HDR) in RGB, SceneDepth in A"),
 	SCS_SceneDepth UMETA(DisplayName="SceneDepth in R"),
diff --git a/Engine/Source/Runtime/Engine/Classes/Engine/RendererSettings.h b/Engine/Source/Runtime/Engine/Classes/Engine/RendererSettings.h
--- a/Engine/Source/Runtime/Engine/Classes/Engine/RendererSettings.h
+++ b/Engine/Source/Runtime/Engine/Classes/Engine/RendererSettings.h
@@ -115,6 +115,12 @@ class ENGINE_API URendererSettings : public UDeveloperSettings
 		ConfigRestartRequired = true))
 		uint32 bMobileEnableStaticAndCSMShadowReceivers : 1;
 
+	UPROPERTY(config, EditAnywhere, Category = Mobile, meta = (
+		ConsoleVariable = "r.Mobile.DisableVertexFog", DisplayName = "Disable vertex fogging in mobile shaders",
+		ToolTip = "If true, vertex fog will be omitted from all mobile shaders, this can increase shading performance.",
+		ConfigRestartRequired = true))
+		uint32 bMobileDisableVertexFog : 1;
+
 	UPROPERTY(config, EditAnywhere, Category = Materials, meta = (
 		ConsoleVariable = "r.DiscardUnusedQuality", DisplayName = "Game Discards Unused Material Quality Levels",
 		ToolTip = "When running in game mode, whether to keep shaders for all quality levels in memory or only those needed for the current quality level.\nUnchecked: Keep all quality levels in memory allowing a runtime quality level change. (default)\nChecked: Discard unused quality levels when loading content for the game, saving some memory."))
diff --git a/Engine/Source/Runtime/Engine/Classes/Engine/World.h b/Engine/Source/Runtime/Engine/Classes/Engine/World.h
--- a/Engine/Source/Runtime/Engine/Classes/Engine/World.h
+++ b/Engine/Source/Runtime/Engine/Classes/Engine/World.h
@@ -1169,6 +1169,9 @@ public:
 	 * have lighting rebuilt.
 	 **/
 	uint32 NumLightingUnbuiltObjects;
+	
+	/** Num of reflection capture components missing valid data. This can be non-zero only in game with FeatureLevel < SM4*/
+	uint32 NumInvalidReflectionCaptureComponents;
 
 	/** Num of components missing valid texture streaming data. Updated in map check. */
 	int32 NumTextureStreamingUnbuiltComponents;
diff --git a/Engine/Source/Runtime/Engine/Private/BatchedElements.cpp b/Engine/Source/Runtime/Engine/Private/BatchedElements.cpp
--- a/Engine/Source/Runtime/Engine/Private/BatchedElements.cpp
+++ b/Engine/Source/Runtime/Engine/Private/BatchedElements.cpp
@@ -507,7 +507,7 @@ void FBatchedElements::PrepareShaders(
 	FMatrix ColorWeights( FPlane(1, 0, 0, 0), FPlane(0, 1, 0, 0), FPlane(0, 0, 1, 0), FPlane(0, 0, 0, 0) );
 
 	// bEncodedHDR requires that blend states are disabled.
-	bool bEncodedHDR = Is32BppHDREncoded(View);
+	bool bEncodedHDR = bEnableHDREncoding && Is32BppHDREncoded(View);
 
 	float GammaToUse = Gamma;
 
diff --git a/Engine/Source/Runtime/Engine/Private/Components/ReflectionCaptureComponent.cpp b/Engine/Source/Runtime/Engine/Private/Components/ReflectionCaptureComponent.cpp
--- a/Engine/Source/Runtime/Engine/Private/Components/ReflectionCaptureComponent.cpp
+++ b/Engine/Source/Runtime/Engine/Private/Components/ReflectionCaptureComponent.cpp
@@ -842,6 +842,34 @@ void UReflectionCaptureComponent::SendRenderTransform_Concurrent()
 	Super::SendRenderTransform_Concurrent();
 }
 
+void UReflectionCaptureComponent::OnRegister()
+{
+	Super::OnRegister();
+
+	UWorld* World = GetWorld();
+	if (World->IsGameWorld() && GMaxRHIFeatureLevel < ERHIFeatureLevel::SM4)
+	{
+		if (EncodedHDRDerivedData == nullptr)
+		{
+			World->NumInvalidReflectionCaptureComponents+= 1;
+		}
+	}
+}
+
+void UReflectionCaptureComponent::OnUnregister()
+{
+	UWorld* World = GetWorld();
+	if (World->IsGameWorld() && GMaxRHIFeatureLevel < ERHIFeatureLevel::SM4)
+	{
+		if (EncodedHDRDerivedData == nullptr && World->NumInvalidReflectionCaptureComponents > 0)
+		{
+			World->NumInvalidReflectionCaptureComponents-= 1;
+		}
+	}
+
+	Super::OnUnregister();
+}
+
 void UReflectionCaptureComponent::DestroyRenderState_Concurrent()
 {
 	Super::DestroyRenderState_Concurrent();
diff --git a/Engine/Source/Runtime/Engine/Private/GameViewportClient.cpp b/Engine/Source/Runtime/Engine/Private/GameViewportClient.cpp
--- a/Engine/Source/Runtime/Engine/Private/GameViewportClient.cpp
+++ b/Engine/Source/Runtime/Engine/Private/GameViewportClient.cpp
@@ -1293,8 +1293,7 @@ void UGameViewportClient::Draw(FViewport* InViewport, FCanvas* SceneCanvas)
 void UGameViewportClient::ProcessScreenShots(FViewport* InViewport)
 {
 	if (GIsDumpingMovie || FScreenshotRequest::IsScreenshotRequested() || GIsHighResScreenshot)
-	{
-	
+	{	
 		TArray<FColor> Bitmap;
 
 		bool bShowUI = false;
@@ -1322,6 +1321,12 @@ void UGameViewportClient::ProcessScreenShots(FViewport* InViewport)
 		{
 			if (ScreenshotCapturedDelegate.IsBound() && CVarScreenshotDelegate.GetValueOnGameThread())
 			{
+				// Ensure that all pixels' alpha is set to 255
+				for (auto& Color : Bitmap)
+				{
+					Color.A = 255;
+				}
+
 				// If delegate subscribed, fire it instead of writing out a file to disk
 				ScreenshotCapturedDelegate.Broadcast(Size.X, Size.Y, Bitmap);
 			}
diff --git a/Engine/Source/Runtime/Engine/Private/HierarchicalInstancedStaticMesh.cpp b/Engine/Source/Runtime/Engine/Private/HierarchicalInstancedStaticMesh.cpp
--- a/Engine/Source/Runtime/Engine/Private/HierarchicalInstancedStaticMesh.cpp
+++ b/Engine/Source/Runtime/Engine/Private/HierarchicalInstancedStaticMesh.cpp
@@ -2395,6 +2395,10 @@ void UHierarchicalInstancedStaticMeshComponent::BuildTreeAsync()
 		InstanceReorderTable.Empty();
 		SortedInstances.Empty();
 		RemovedInstances.Empty();
+
+		UnbuiltInstanceBoundsList.Empty();
+		BuiltInstanceBounds.Init();
+		ReleasePerInstanceRenderData();
 	}
 }
 
diff --git a/Engine/Source/Runtime/Engine/Private/Materials/HLSLMaterialTranslator.h b/Engine/Source/Runtime/Engine/Private/Materials/HLSLMaterialTranslator.h
--- a/Engine/Source/Runtime/Engine/Private/Materials/HLSLMaterialTranslator.h
+++ b/Engine/Source/Runtime/Engine/Private/Materials/HLSLMaterialTranslator.h
@@ -453,6 +453,8 @@ public:
 				Errorf(TEXT("Only transparent or postprocess materials can read from scene depth."));
 			}
 
+			MaterialCompilationOutput.bUsesSceneDepthLookup = bUsesSceneDepth;
+
 			if (MaterialCompilationOutput.bRequiresSceneColorCopy)
 			{
 				if (Domain != MD_Surface)
@@ -2963,8 +2965,11 @@ protected:
 	// @param SceneTextureId of type ESceneTextureId e.g. PPI_SubsurfaceColor
 	virtual int32 SceneTextureLookup(int32 UV, uint32 InSceneTextureId, bool bFiltered) override
 	{
-		if (InSceneTextureId != PPI_PostProcessInput0 // fetching PostProcessInput0 supported on all platforms
-			&& ErrorUnlessFeatureLevelSupported(ERHIFeatureLevel::SM4) == INDEX_NONE)
+		const bool bSupportedOnMobile = InSceneTextureId == PPI_PostProcessInput0 ||
+										InSceneTextureId == PPI_CustomDepth ||
+										InSceneTextureId == PPI_SceneDepth;
+
+		if (!bSupportedOnMobile	&& ErrorUnlessFeatureLevelSupported(ERHIFeatureLevel::SM4) == INDEX_NONE)
 		{
 			return INDEX_NONE;
 		}
@@ -3003,8 +3008,7 @@ protected:
 				UV = TextureCoordinate(0, false, false);
 			}
 			FString TexCoordCode = CoerceParameter(UV, MCT_Float2);
-			// Only PPI_PostProcessInput0, which holds scene color 
-			return AddCodeChunk(MCT_Float4,	TEXT("MobileLookupPostProcessInput0(Parameters, %s)"), *TexCoordCode);
+			return AddCodeChunk(MCT_Float4,	TEXT("MobileSceneTextureLookup(Parameters, %d, %s)"), (int32)SceneTextureId, *TexCoordCode);
 		}
 	}
 
@@ -3147,6 +3151,12 @@ protected:
 			bNeedsSceneTexturePostProcessInputs = bNeedsSceneTexturePostProcessInputs
 				|| ((SceneTextureId >= PPI_PostProcessInput0 && SceneTextureId <= PPI_PostProcessInput6)
 				|| SceneTextureId == PPI_SceneColor);
+
+		}
+
+		if (SceneTextureId == PPI_SceneDepth && bTextureLookup)
+		{
+			bUsesSceneDepth = true;
 		}
 
 		const bool bNeedsGBuffer = SceneTextureId == PPI_DiffuseColor 
diff --git a/Engine/Source/Runtime/Engine/Private/Materials/MaterialShared.cpp b/Engine/Source/Runtime/Engine/Private/Materials/MaterialShared.cpp
--- a/Engine/Source/Runtime/Engine/Private/Materials/MaterialShared.cpp
+++ b/Engine/Source/Runtime/Engine/Private/Materials/MaterialShared.cpp
@@ -313,6 +313,7 @@ void FMaterialCompilationOutput::Serialize(FArchive& Ar)
 	Ar << bNeedsGBuffer;
 	Ar << bUsesGlobalDistanceField;
 	Ar << bUsesPixelDepthOffset;
+	Ar << bUsesSceneDepthLookup;
 }
 
 void FMaterial::GetShaderMapId(EShaderPlatform Platform, FMaterialShaderMapId& OutId) const
@@ -585,6 +586,12 @@ bool FMaterial::MaterialUsesPixelDepthOffset() const
 	return RenderingThreadShaderMap ? RenderingThreadShaderMap->UsesPixelDepthOffset() : false;
 }
 
+bool FMaterial::MaterialUsesSceneDepthLookup() const
+{
+	check(IsInParallelRenderingThread());
+	return RenderingThreadShaderMap ? RenderingThreadShaderMap->UsesSceneDepthLookup() : false;
+}
+
 FMaterialShaderMap* FMaterial::GetRenderingThreadShaderMap() const 
 { 
 	check(IsInParallelRenderingThread());
@@ -2706,7 +2713,6 @@ bool FMaterialShaderMapId::ContainsVertexFactoryType(const FVertexFactoryType* V
 
 	return false;
 }
-
 //////////////////////////////////////////////////////////////////////////
 
 FMaterialAttributeDefintion::FMaterialAttributeDefintion(
diff --git a/Engine/Source/Runtime/Engine/Private/ShaderCompiler/ShaderCompiler.cpp b/Engine/Source/Runtime/Engine/Private/ShaderCompiler/ShaderCompiler.cpp
--- a/Engine/Source/Runtime/Engine/Private/ShaderCompiler/ShaderCompiler.cpp
+++ b/Engine/Source/Runtime/Engine/Private/ShaderCompiler/ShaderCompiler.cpp
@@ -2708,6 +2708,11 @@ void GlobalBeginCompileShader(
 		Input.Environment.SetDefine(TEXT("VERTEX_FOGGING_FOR_OPAQUE"), bForwardShading && (CVar ? (CVar->GetInt() != 0) : 0));
 	}
 
+	{
+		static IConsoleVariable* CVar = IConsoleManager::Get().FindConsoleVariable(TEXT("r.Mobile.DisableVertexFog"));
+		Input.Environment.SetDefine(TEXT("PROJECT_MOBILE_DISABLE_VERTEX_FOG"), CVar ? (CVar->GetInt() != 0) : 0);
+	}
+
 	if (GSupportsRenderTargetWriteMask)
 	{
 		Input.Environment.SetDefine(TEXT("PLATFORM_SUPPORTS_RENDERTARGET_WRITE_MASK"), 1);
diff --git a/Engine/Source/Runtime/Engine/Private/ShaderDerivedDataVersion.h b/Engine/Source/Runtime/Engine/Private/ShaderDerivedDataVersion.h
--- a/Engine/Source/Runtime/Engine/Private/ShaderDerivedDataVersion.h
+++ b/Engine/Source/Runtime/Engine/Private/ShaderDerivedDataVersion.h
@@ -10,4 +10,4 @@
 // guid as version
 
 #define GLOBALSHADERMAP_DERIVEDDATA_VER			TEXT("96B578BEE72B4D52A11B36DBAD651AE0")
-#define MATERIALSHADERMAP_DERIVEDDATA_VER		TEXT("F3AFA18D58E84ACD8191C2D585CEDD02")
\ No newline at end of file
+#define MATERIALSHADERMAP_DERIVEDDATA_VER		TEXT("DBBC4FB34C0F423698AA0729642543E3")
\ No newline at end of file
diff --git a/Engine/Source/Runtime/Engine/Private/Tests/AutomationCommon.cpp b/Engine/Source/Runtime/Engine/Private/Tests/AutomationCommon.cpp
--- a/Engine/Source/Runtime/Engine/Private/Tests/AutomationCommon.cpp
+++ b/Engine/Source/Runtime/Engine/Private/Tests/AutomationCommon.cpp
@@ -105,20 +105,7 @@ namespace AutomationCommon
 		Data.Rhi = FHardwareInfo::GetHardwareInfo(NAME_RHI);
 		GetFeatureLevelName(GMaxRHIFeatureLevel, Data.FeatureLevel);
 		Data.bIsStereo = GEngine->StereoRenderingDevice.IsValid();
-
-		if ( IsRHIDeviceNVIDIA() )
-		{
-			Data.Vendor = TEXT("NVIDIA");
-		}
-		else if ( IsRHIDeviceAMD() )
-		{
-			Data.Vendor = TEXT("AMD");
-		}
-		else if ( IsRHIDeviceIntel() )
-		{
-			Data.Vendor = TEXT("Intel");
-		}
-
+		Data.Vendor = RHIVendorIdToString();
 		Data.AdapterName = GRHIAdapterName;
 		Data.AdapterInternalDriverVersion = GRHIAdapterInternalDriverVersion;
 		Data.AdapterUserDriverVersion = GRHIAdapterUserDriverVersion;
diff --git a/Engine/Source/Runtime/Engine/Private/UnrealEngine.cpp b/Engine/Source/Runtime/Engine/Private/UnrealEngine.cpp
--- a/Engine/Source/Runtime/Engine/Private/UnrealEngine.cpp
+++ b/Engine/Source/Runtime/Engine/Private/UnrealEngine.cpp
@@ -7795,6 +7795,15 @@ float DrawMapWarnings(UWorld* World, FViewport* Viewport, FCanvas* Canvas, UCanv
 		MessageY += FontSizeY;
 	}
 
+	// Warn about invalid reflection captures, this can appear only in game with FeatureLevel < SM4
+	if (World->NumInvalidReflectionCaptureComponents > 0)
+	{
+		SmallTextItem.SetColor(FLinearColor::Red);
+		SmallTextItem.Text = FText::FromString(FString::Printf(TEXT("INVALID REFLECTION CAPTURES (%u Components, resave map in the editor)"), World->NumInvalidReflectionCaptureComponents));
+		Canvas->DrawItem(SmallTextItem, FVector2D(MessageX, MessageY));
+		MessageY += FontSizeY;
+	}
+	
 	// Check HLOD clusters and show warning if unbuilt
 #if WITH_EDITOR
 	if (World->GetWorldSettings()->bEnableHierarchicalLODSystem)
diff --git a/Engine/Source/Runtime/Engine/Public/BatchedElements.h b/Engine/Source/Runtime/Engine/Public/BatchedElements.h
--- a/Engine/Source/Runtime/Engine/Public/BatchedElements.h
+++ b/Engine/Source/Runtime/Engine/Public/BatchedElements.h
@@ -100,6 +100,7 @@ public:
 		:	MaxMeshIndicesAllowed(GDrawUPIndexCheckCount / sizeof(int32))
 			// the index buffer is 2 bytes, so make sure we only address 0xFFFF vertices in the index buffer
 		,	MaxMeshVerticesAllowed(FMath::Min<uint32>(0xFFFF, GDrawUPVertexCheckCount / sizeof(FSimpleElementVertex)))
+		,	bEnableHDREncoding(true)
 	{
 	}
 
@@ -234,6 +235,12 @@ public:
 		return sizeof(*this) + Points.GetAllocatedSize() + WireTris.GetAllocatedSize() + WireTriVerts.GetAllocatedSize() + ThickLines.GetAllocatedSize()
 			+ Sprites.GetAllocatedSize() + MeshElements.GetAllocatedSize() + MeshVertices.GetAllocatedSize();
 	}
+
+	void EnableMobileHDREncoding(bool bInEnableHDREncoding)
+	{
+		bEnableHDREncoding = bInEnableHDREncoding;
+	}
+
 private:
 
 	/**
@@ -376,5 +383,8 @@ private:
 		const FSceneView* View = NULL,
 		FTexture2DRHIRef DepthTexture = FTexture2DRHIRef()
 		) const;
+
+	/** if false then prevent the use of HDR encoded shaders. */
+	bool bEnableHDREncoding;
 };
 
diff --git a/Engine/Source/Runtime/Engine/Public/MaterialShared.h b/Engine/Source/Runtime/Engine/Public/MaterialShared.h
--- a/Engine/Source/Runtime/Engine/Public/MaterialShared.h
+++ b/Engine/Source/Runtime/Engine/Public/MaterialShared.h
@@ -305,7 +305,8 @@ public:
 		bUsesWorldPositionOffset(false),
 		bNeedsGBuffer(false),
 		bUsesGlobalDistanceField(false),
-		bUsesPixelDepthOffset(false)
+		bUsesPixelDepthOffset(false),
+		bUsesSceneDepthLookup(false)
 	{}
 
 	ENGINE_API void Serialize(FArchive& Ar);
@@ -337,6 +338,9 @@ public:
 
 	/** true if the material writes a pixel depth offset */
 	bool bUsesPixelDepthOffset;
+
+	/** true if the material uses the SceneDepth lookup */
+	bool bUsesSceneDepthLookup;
 };
 
 /** 
@@ -716,6 +720,7 @@ public:
 	bool UsesEyeAdaptation() const { return MaterialCompilationOutput.bUsesEyeAdaptation; }
 	bool ModifiesMeshPosition() const { return MaterialCompilationOutput.bModifiesMeshPosition; }
 	bool UsesPixelDepthOffset() const { return MaterialCompilationOutput.bUsesPixelDepthOffset; }
+	bool UsesSceneDepthLookup() const { return MaterialCompilationOutput.bUsesSceneDepthLookup; }
 
 	bool IsValidForRendering() const
 	{
@@ -1102,6 +1107,9 @@ public:
 	/** Does the material use a pixel depth offset. */
 	ENGINE_API bool MaterialUsesPixelDepthOffset() const;
 
+	/** Does the material use a SceneDepth lookup. */
+	ENGINE_API bool MaterialUsesSceneDepthLookup() const;
+
 	/** Note: This function is only intended for use in deciding whether or not shader permutations are required before material translation occurs. */
 	ENGINE_API bool MaterialMayModifyMeshPosition() const;
 
diff --git a/Engine/Source/Runtime/Landscape/Private/LandscapeGrass.cpp b/Engine/Source/Runtime/Landscape/Private/LandscapeGrass.cpp
--- a/Engine/Source/Runtime/Landscape/Private/LandscapeGrass.cpp
+++ b/Engine/Source/Runtime/Landscape/Private/LandscapeGrass.cpp
@@ -2111,6 +2111,7 @@ void ALandscapeProxy::UpdateGrass(const TArray<FVector>& Cameras, bool bForceSyn
 										{
 											QUICK_SCOPE_CYCLE_COUNTER(STAT_GrassCreateComp);
 											HierarchicalInstancedStaticMeshComponent = NewObject<UHierarchicalInstancedStaticMeshComponent>(this);
+											HierarchicalInstancedStaticMeshComponent->SetFlags(RF_Transient);
 										}
 										NewComp.Foliage = HierarchicalInstancedStaticMeshComponent;
 										FoliageCache.CachedGrassComps.Add(NewComp);
diff --git a/Engine/Source/Runtime/Launch/Private/Android/AndroidEventManager.cpp b/Engine/Source/Runtime/Launch/Private/Android/AndroidEventManager.cpp
--- a/Engine/Source/Runtime/Launch/Private/Android/AndroidEventManager.cpp
+++ b/Engine/Source/Runtime/Launch/Private/Android/AndroidEventManager.cpp
@@ -208,7 +208,7 @@ FAppEventManager::FAppEventManager():
 	,EmptyQueueHandlerEvent(nullptr)
 	,FirstInitialized(false)
 	,bCreateWindow(false)
-	,bWindowInFocus(false)
+	,bWindowInFocus(true)
 	,bSaveState(false)
 	,bAudioPaused(false)
 	,PendingWindow(NULL)
diff --git a/Engine/Source/Runtime/Launch/Private/Android/AndroidJNI.cpp b/Engine/Source/Runtime/Launch/Private/Android/AndroidJNI.cpp
--- a/Engine/Source/Runtime/Launch/Private/Android/AndroidJNI.cpp
+++ b/Engine/Source/Runtime/Launch/Private/Android/AndroidJNI.cpp
@@ -69,6 +69,7 @@ void FJavaWrapper::FindClassesAndMethods(JNIEnv* Env)
 	AndroidThunkJava_LocalNotificationClearAll = FindMethod(Env, GameActivityClassID, "AndroidThunkJava_LocalNotificationClearAll", "()V", bIsOptional);
 	AndroidThunkJava_LocalNotificationGetLaunchNotification = FindMethod(Env, GameActivityClassID, "AndroidThunkJava_LocalNotificationGetLaunchNotification", "()Lcom/epicgames/ue4/GameActivity$LaunchNotification;", bIsOptional);
 	//AndroidThunkJava_LocalNotificationDestroyIfExists = FindMethod(Env, GameActivityClassID, "AndroidThunkJava_LocalNotificationDestroyIfExists", "(I)Z", bIsOptional);
+	AndroidThunkJava_HasActiveWiFiConnection = FindMethod(Env, GameActivityClassID, "AndroidThunkJava_HasActiveWiFiConnection", "()Z", bIsOptional);
 
 	// this is optional - only inserted if GearVR plugin enabled
 	AndroidThunkJava_IsGearVRApplication = FindMethod(Env, GameActivityClassID, "AndroidThunkJava_IsGearVRApplication", "()Z", true);
@@ -237,6 +238,7 @@ jmethodID FJavaWrapper::AndroidThunkJava_LocalNotificationScheduleAtTime;
 jmethodID FJavaWrapper::AndroidThunkJava_LocalNotificationClearAll;
 jmethodID FJavaWrapper::AndroidThunkJava_LocalNotificationGetLaunchNotification;
 //jmethodID FJavaWrapper::AndroidThunkJava_LocalNotificationDestroyIfExists;
+jmethodID FJavaWrapper::AndroidThunkJava_HasActiveWiFiConnection;
 
 jclass FJavaWrapper::InputDeviceInfoClass;
 jfieldID FJavaWrapper::InputDeviceInfo_VendorId;
@@ -934,6 +936,17 @@ bool AndroidThunkCpp_DestroyScheduledNotificationIfExists(int32 NotificationId)
 }
 */
 
+bool AndroidThunkCpp_HasActiveWiFiConnection()
+{
+	bool bIsActive = false;
+	if (JNIEnv* Env = FAndroidApplication::GetJavaEnv())
+	{
+		bIsActive = FJavaWrapper::CallBooleanMethod(Env, FJavaWrapper::GameActivityThis, FJavaWrapper::AndroidThunkJava_HasActiveWiFiConnection);
+	}
+
+	return bIsActive;
+}
+
 //The JNI_OnLoad function is triggered by loading the game library from 
 //the Java source file.
 //	static
diff --git a/Engine/Source/Runtime/Launch/Public/Android/AndroidJNI.h b/Engine/Source/Runtime/Launch/Public/Android/AndroidJNI.h
--- a/Engine/Source/Runtime/Launch/Public/Android/AndroidJNI.h
+++ b/Engine/Source/Runtime/Launch/Public/Android/AndroidJNI.h
@@ -42,6 +42,7 @@ public:
 	static jmethodID AndroidThunkJava_LocalNotificationClearAll;
 	static jmethodID AndroidThunkJava_LocalNotificationGetLaunchNotification;
 	//static jmethodID AndroidThunkJava_LocalNotificationDestroyIfExists; - This is not needed yet but will be soon so just leaving commented out for now
+	static jmethodID AndroidThunkJava_HasActiveWiFiConnection;
 
 	// InputDeviceInfo member field ids
 	static jclass InputDeviceInfoClass;
diff --git a/Engine/Source/Runtime/MaterialShaderQualitySettings/Classes/ShaderPlatformQualitySettings.h b/Engine/Source/Runtime/MaterialShaderQualitySettings/Classes/ShaderPlatformQualitySettings.h
--- a/Engine/Source/Runtime/MaterialShaderQualitySettings/Classes/ShaderPlatformQualitySettings.h
+++ b/Engine/Source/Runtime/MaterialShaderQualitySettings/Classes/ShaderPlatformQualitySettings.h
@@ -8,6 +8,16 @@
 /**
 * 
 */
+UENUM()
+enum class EMobileCSMQuality : uint8
+{
+	// Lowest quality, no filtering.
+	NoFiltering,
+	// Medium quality, 1x1 PCF filtering.
+	PCF_1x1 UMETA(DisplayName = "1x1 PCF"),
+	// Highest quality, 2x2 PCF filtering.
+	PCF_2x2 UMETA(DisplayName = "2x2 PCF"),
+};
 
 // FMaterialQualityOverrides represents the full set of possible material overrides per quality level.
 USTRUCT()
@@ -22,6 +32,7 @@ public:
 		, bForceNonMetal(false)
 		, bForceDisableLMDirectionality(false)
 		, bForceLQReflections(false)
+		, MobileCSMQuality(EMobileCSMQuality::PCF_2x2)
 	{
 	}
 
@@ -40,6 +51,9 @@ public:
 	UPROPERTY(EditAnywhere, Config, Meta = (DisplayName = "Force low quality reflections"), Category = "Quality")
 	bool bForceLQReflections;
 
+	UPROPERTY(EditAnywhere, Config, Meta = (DisplayName = "Cascade shadow mapping quality"), Category = "Quality")
+	EMobileCSMQuality MobileCSMQuality;
+
 	bool HasAnyOverridesSet() const;
 };
 
diff --git a/Engine/Source/Runtime/MaterialShaderQualitySettings/Private/MaterialShaderQualitySettings.cpp b/Engine/Source/Runtime/MaterialShaderQualitySettings/Private/MaterialShaderQualitySettings.cpp
--- a/Engine/Source/Runtime/MaterialShaderQualitySettings/Private/MaterialShaderQualitySettings.cpp
+++ b/Engine/Source/Runtime/MaterialShaderQualitySettings/Private/MaterialShaderQualitySettings.cpp
@@ -148,5 +148,12 @@ void UShaderPlatformQualitySettings::AppendToHashState(EMaterialQualityLevel::Ty
 
 bool FMaterialQualityOverrides::HasAnyOverridesSet() const
 {
-	return bForceDisableLMDirectionality || bForceFullyRough || bForceNonMetal || bForceDisableLMDirectionality || bForceLQReflections;
+	static const FMaterialQualityOverrides DefaultOverrides;
+
+	return
+		MobileCSMQuality != DefaultOverrides.MobileCSMQuality
+		|| bForceDisableLMDirectionality != DefaultOverrides.bForceDisableLMDirectionality
+		|| bForceFullyRough != DefaultOverrides.bForceFullyRough
+		|| bForceNonMetal != DefaultOverrides.bForceNonMetal
+		|| bForceLQReflections != DefaultOverrides.bForceLQReflections;
 }
diff --git a/Engine/Source/Runtime/OpenGLDrv/Private/Android/AndroidESDeferredOpenGL.cpp b/Engine/Source/Runtime/OpenGLDrv/Private/Android/AndroidESDeferredOpenGL.cpp
--- a/Engine/Source/Runtime/OpenGLDrv/Private/Android/AndroidESDeferredOpenGL.cpp
+++ b/Engine/Source/Runtime/OpenGLDrv/Private/Android/AndroidESDeferredOpenGL.cpp
@@ -582,6 +582,11 @@ bool FAndroidMisc::SupportsShaderFramebufferFetch()
 	return FAndroidGPUInfo::Get().bSupportsFrameBufferFetch;
 }
 
+bool FAndroidMisc::SupportsShaderIOBlocks()
+{
+	return FAndroidGPUInfo::Get().bSupportsShaderIOBlocks;
+}
+
 void FAndroidMisc::GetValidTargetPlatforms(TArray<FString>& TargetPlatformNames)
 {
 	TargetPlatformNames = FAndroidGPUInfo::Get().TargetPlatformNames;
diff --git a/Engine/Source/Runtime/OpenGLDrv/Private/Android/AndroidOpenGL.cpp b/Engine/Source/Runtime/OpenGLDrv/Private/Android/AndroidOpenGL.cpp
--- a/Engine/Source/Runtime/OpenGLDrv/Private/Android/AndroidOpenGL.cpp
+++ b/Engine/Source/Runtime/OpenGLDrv/Private/Android/AndroidOpenGL.cpp
@@ -205,25 +205,44 @@ bool PlatformInitOpenGL()
 
 	{
 		// determine ES version. PlatformInitOpenGL happens before ProcessExtensions and therefore FAndroidOpenGL::bES31Support.
-		const bool bES31Supported = FAndroidGPUInfo::Get().GLVersion.Contains(TEXT("OpenGL ES 3.1"));
+		FString SubVersion;
+		const bool bES31Supported = FAndroidGPUInfo::Get().GLVersion.Split(TEXT("OpenGL ES 3."), nullptr, &SubVersion) && FCString::Atoi(*SubVersion) >= 1;
 		static const auto CVarDisableES31 = IConsoleManager::Get().FindTConsoleVariableDataInt(TEXT("r.Android.DisableOpenGLES31Support"));
 
 		bool bBuildForES31 = false;
 		GConfig->GetBool(TEXT("/Script/AndroidRuntimeSettings.AndroidRuntimeSettings"), TEXT("bBuildForES31"), bBuildForES31, GEngineIni);
+		const bool bSupportsFloatingPointRTs = FAndroidMisc::SupportsFloatingPointRenderTargets();
+		const bool bSupportsShaderIOBlocks = FAndroidMisc::SupportsShaderIOBlocks();
 
-		if (bES31Supported && bBuildForES31 && CVarDisableES31->GetValueOnAnyThread() == 0)
+		if (bBuildForES31 && bES31Supported && bSupportsFloatingPointRTs && bSupportsShaderIOBlocks && CVarDisableES31->GetValueOnAnyThread() == 0)
 		{
+			FAndroidOpenGL::CurrentFeatureLevelSupport = FAndroidOpenGL::EFeatureLevelSupport::ES31;
 			// shut down existing ES2 egl.
+			UE_LOG(LogRHI, Log, TEXT("App is packaged for OpenGL ES 3.1 and an ES 3.1-capable device was detected. Reinitializing OpenGL ES with a 3.1 context."));
 			FAndroidAppEntry::ReleaseEGL();
 			// Re-init gles for 3.1
 			AndroidEGL::GetInstance()->Init(AndroidEGL::AV_OpenGLES, 3, 1, false);
 		}
 		else
 		{
+			FAndroidOpenGL::CurrentFeatureLevelSupport = FAndroidOpenGL::EFeatureLevelSupport::ES2;
+			if (bBuildForES31)
+			{
+				UE_LOG(LogRHI, Log, TEXT("App is packaged for OpenGL ES 3.1 but device has not met all the requirements for ES 3.1 :"));
+				if (CVarDisableES31->GetValueOnAnyThread())
+				{
+					UE_LOG(LogRHI, Log, TEXT("	ES 3.1 support was disabled via r.Android.DisableOpenGLES31Support"));
+				}
+
+				UE_LOG(LogRHI, Log, TEXT("	Device has OpenGL ES 3.1 support: %s"), bES31Supported ? TEXT("YES") : TEXT("NO"));
+				UE_LOG(LogRHI, Log, TEXT("	Floating point render target support: %s"), bSupportsFloatingPointRTs ? TEXT("YES") : TEXT("NO"));
+				UE_LOG(LogRHI, Log, TEXT("	EXT_shader_io_blocks support: %s"), bSupportsShaderIOBlocks ? TEXT("YES") : TEXT("NO"));
+			}
+
 			bool bBuildForES2 = false;
 			GConfig->GetBool(TEXT("/Script/AndroidRuntimeSettings.AndroidRuntimeSettings"), TEXT("bBuildForES2"), bBuildForES2, GEngineIni);
 			// If we're here and there's no ES2 data then we're in trouble.
-			check(bBuildForES2);
+			checkf(bBuildForES2, TEXT("This device only supports OpenGL ES 2 but the app was not packaged with ES2 support."));
 		}
 	}
 	return true;
@@ -383,15 +402,16 @@ bool FAndroidOpenGL::bES31Support = false;
 bool FAndroidOpenGL::bSupportsInstancing = false;
 bool FAndroidOpenGL::bHasHardwareHiddenSurfaceRemoval = false;
 bool FAndroidOpenGL::bSupportsMobileMultiView = false;
+FAndroidOpenGL::EFeatureLevelSupport FAndroidOpenGL::CurrentFeatureLevelSupport = FAndroidOpenGL::EFeatureLevelSupport::Invalid;
 
 void FAndroidOpenGL::ProcessExtensions(const FString& ExtensionsString)
 {
 	FOpenGLES2::ProcessExtensions(ExtensionsString);
 
 	FString VersionString = FString(ANSI_TO_TCHAR((const ANSICHAR*)glGetString(GL_VERSION)));
-
-	bES30Support = VersionString.Contains(TEXT("OpenGL ES 3."));
-	bES31Support = VersionString.Contains(TEXT("OpenGL ES 3.1"));
+	FString SubVersionString;
+	bES30Support = VersionString.Split(TEXT("OpenGL ES 3."), nullptr, &SubVersionString);
+	bES31Support = bES30Support && FCString::Atoi(*SubVersionString) >= 1;
 
 	// Get procedures
 	if (bSupportsOcclusionQueries || bSupportsDisjointTimeQueries)
@@ -612,6 +632,11 @@ bool FAndroidMisc::SupportsShaderFramebufferFetch()
 	return FAndroidGPUInfo::Get().bSupportsFrameBufferFetch;
 }
 
+bool FAndroidMisc::SupportsShaderIOBlocks()
+{
+	return FAndroidGPUInfo::Get().bSupportsShaderIOBlocks;
+}
+
 void FAndroidMisc::GetValidTargetPlatforms(TArray<FString>& TargetPlatformNames)
 {
 	TargetPlatformNames = FAndroidGPUInfo::Get().TargetPlatformNames;
diff --git a/Engine/Source/Runtime/OpenGLDrv/Private/Android/AndroidOpenGL.h b/Engine/Source/Runtime/OpenGLDrv/Private/Android/AndroidOpenGL.h
--- a/Engine/Source/Runtime/OpenGLDrv/Private/Android/AndroidOpenGL.h
+++ b/Engine/Source/Runtime/OpenGLDrv/Private/Android/AndroidOpenGL.h
@@ -155,22 +155,10 @@ extern "C"
 
 struct FAndroidOpenGL : public FOpenGLES2
 {
-	static FORCEINLINE bool IsBuiltForES31()
-	{
-		static int32 ES31BuiltState = -1;
-		if(ES31BuiltState == -1)
-		{
-			bool bBuildForES31 = false;
-			GConfig->GetBool(TEXT("/Script/AndroidRuntimeSettings.AndroidRuntimeSettings"), TEXT("bBuildForES31"), bBuildForES31, GEngineIni);
-			ES31BuiltState = bBuildForES31 ? 1 : 0;
-		}
-		return ES31BuiltState == 1;
-	}
-
 	static FORCEINLINE bool IsES31Usable()
 	{
-		static const auto CVarDisableES31 = IConsoleManager::Get().FindTConsoleVariableDataInt(TEXT("r.Android.DisableOpenGLES31Support"));
-		return bES31Support && IsBuiltForES31() && CVarDisableES31->GetValueOnAnyThread() == 0;
+		check(CurrentFeatureLevelSupport != EFeatureLevelSupport::Invalid);
+		return CurrentFeatureLevelSupport == EFeatureLevelSupport::ES31;
 	}
 
 	static FORCEINLINE EShaderPlatform GetShaderPlatform()
@@ -486,6 +474,16 @@ struct FAndroidOpenGL : public FOpenGLES2
 
 	/** Whether device supports mobile multi-view */
 	static bool bSupportsMobileMultiView;
+
+	enum class EFeatureLevelSupport : uint8
+	{
+		Invalid,	// no feature level has yet been determined
+		ES2,
+		ES31,
+	};
+
+	/** Describes which feature level is currently being supported */
+	static EFeatureLevelSupport CurrentFeatureLevelSupport;
 };
 
 typedef FAndroidOpenGL FOpenGL;
diff --git a/Engine/Source/Runtime/OpenGLDrv/Private/Android/AndroidOpenGLPrivate.h b/Engine/Source/Runtime/OpenGLDrv/Private/Android/AndroidOpenGLPrivate.h
--- a/Engine/Source/Runtime/OpenGLDrv/Private/Android/AndroidOpenGLPrivate.h
+++ b/Engine/Source/Runtime/OpenGLDrv/Private/Android/AndroidOpenGLPrivate.h
@@ -25,6 +25,7 @@ public:
 	FString GLVersion;
 	bool bSupportsFloatingPointRenderTargets;
 	bool bSupportsFrameBufferFetch;
+	bool bSupportsShaderIOBlocks;
 	TArray<FString> TargetPlatformNames;
 
 private:
@@ -106,6 +107,8 @@ private:
 			|| (bES30Support && ExtensionsString.Contains(TEXT("GL_EXT_color_buffer_float")));
 
 		bSupportsFrameBufferFetch = ExtensionsString.Contains(TEXT("GL_EXT_shader_framebuffer_fetch")) || ExtensionsString.Contains(TEXT("GL_NV_shader_framebuffer_fetch")) || ExtensionsString.Contains(TEXT("GL_ARM_shader_framebuffer_fetch"));
+		bSupportsShaderIOBlocks = ExtensionsString.Contains(TEXT("GL_EXT_shader_io_blocks"));
+
 		GAndroidGPUInfoReady = true;
 	}
 };
\ No newline at end of file
diff --git a/Engine/Source/Runtime/OpenGLDrv/Private/OpenGLCommands.cpp b/Engine/Source/Runtime/OpenGLDrv/Private/OpenGLCommands.cpp
--- a/Engine/Source/Runtime/OpenGLDrv/Private/OpenGLCommands.cpp
+++ b/Engine/Source/Runtime/OpenGLDrv/Private/OpenGLCommands.cpp
@@ -3500,6 +3500,11 @@ void FOpenGLDynamicRHI::RHIBlockUntilGPUIdle()
 	// Not really supported
 }
 
+void FOpenGLDynamicRHI::RHISubmitCommandsAndFlushGPU()
+{
+	FOpenGL::Flush();
+}
+
 /**
  * Returns the total GPU time taken to render the last frame. Same metric as FPlatformTime::Cycles().
  */
diff --git a/Engine/Source/Runtime/OpenGLDrv/Private/OpenGLDevice.cpp b/Engine/Source/Runtime/OpenGLDrv/Private/OpenGLDevice.cpp
--- a/Engine/Source/Runtime/OpenGLDrv/Private/OpenGLDevice.cpp
+++ b/Engine/Source/Runtime/OpenGLDrv/Private/OpenGLDevice.cpp
@@ -567,6 +567,7 @@ static void InitRHICapabilitiesForGL()
 	#undef LOG_GL_STRING
 
 	GRHIAdapterName = FOpenGL::GetAdapterName();
+	GRHIAdapterInternalDriverVersion = ANSI_TO_TCHAR((const ANSICHAR*)glGetString(GL_VERSION));
 
 	// Log all supported extensions.
 #if PLATFORM_WINDOWS
diff --git a/Engine/Source/Runtime/OpenGLDrv/Private/OpenGLDrv.cpp b/Engine/Source/Runtime/OpenGLDrv/Private/OpenGLDrv.cpp
--- a/Engine/Source/Runtime/OpenGLDrv/Private/OpenGLDrv.cpp
+++ b/Engine/Source/Runtime/OpenGLDrv/Private/OpenGLDrv.cpp
@@ -501,22 +501,40 @@ void FOpenGLBase::ProcessExtensions( const FString& ExtensionsString )
 
 	bSupportsDrawBuffersBlend = ExtensionsString.Contains(TEXT("GL_ARB_draw_buffers_blend"));
 
-#if PLATFORM_WINDOWS || PLATFORM_LINUX
+#if PLATFORM_IOS
+	GRHIVendorId = 0x1010;
+#else
 	FString VendorName( ANSI_TO_TCHAR((const ANSICHAR*)glGetString(GL_VENDOR) ) );
 	if (VendorName.Contains(TEXT("ATI ")))
 	{
-		bAmdWorkaround = true;
 		GRHIVendorId = 0x1002;
+#if PLATFORM_WINDOWS || PLATFORM_LINUX
+		bAmdWorkaround = true;
+#endif
 	}
 	else if (VendorName.Contains(TEXT("Intel ")))
 	{
-		bAmdWorkaround = true;
 		GRHIVendorId = 0x8086;
+#if PLATFORM_WINDOWS || PLATFORM_LINUX
+		bAmdWorkaround = true;
+#endif
 	}
 	else if (VendorName.Contains(TEXT("NVIDIA ")))
 	{
 		GRHIVendorId = 0x10DE;
 	}
+	else if (VendorName.Contains(TEXT("ImgTec")))
+	{
+		GRHIVendorId = 0x1010;
+	}
+	else if (VendorName.Contains(TEXT("ARM")))
+	{
+		GRHIVendorId = 0x13B5;
+	}
+	else if (VendorName.Contains(TEXT("Qualcomm")))
+	{
+		GRHIVendorId = 0x5143;
+	}
 #endif
 
 	// Setup CVars that require the RHI initialized
diff --git a/Engine/Source/Runtime/OpenGLDrv/Public/OpenGLDrv.h b/Engine/Source/Runtime/OpenGLDrv/Public/OpenGLDrv.h
--- a/Engine/Source/Runtime/OpenGLDrv/Public/OpenGLDrv.h
+++ b/Engine/Source/Runtime/OpenGLDrv/Public/OpenGLDrv.h
@@ -391,6 +391,7 @@ public:
 	virtual void RHISetStreamOutTargets(uint32 NumTargets,const FVertexBufferRHIParamRef* VertexBuffers,const uint32* Offsets) final override;
 	virtual void RHIDiscardRenderTargets(bool Depth,bool Stencil,uint32 ColorBitMask) final override;
 	virtual void RHIBlockUntilGPUIdle() final override;
+	virtual void RHISubmitCommandsAndFlushGPU() final override;
 	virtual bool RHIGetAvailableResolutions(FScreenResolutionArray& Resolutions, bool bIgnoreRefreshRate) final override;
 	virtual void RHIGetSupportedResolution(uint32& Width, uint32& Height) final override;
 	virtual void RHIVirtualTextureSetFirstMipInMemory(FTexture2DRHIParamRef Texture, uint32 FirstMip) final override;
diff --git a/Engine/Source/Runtime/RHI/Private/RHI.cpp b/Engine/Source/Runtime/RHI/Private/RHI.cpp
--- a/Engine/Source/Runtime/RHI/Private/RHI.cpp
+++ b/Engine/Source/Runtime/RHI/Private/RHI.cpp
@@ -526,3 +526,17 @@ RHI_API bool IsRHIDeviceNVIDIA()
 	// NVIDIA GPUs are discrete and use DedicatedVideoMemory only.
 	return GRHIVendorId == 0x10DE;
 }
+
+RHI_API const TCHAR* RHIVendorIdToString()
+{
+	switch (GRHIVendorId)
+	{
+	case 0x1002: return TEXT("AMD");
+	case 0x1010: return TEXT("ImgTec");
+	case 0x10DE: return TEXT("NVIDIA");
+	case 0x13B5: return TEXT("ARM");
+	case 0x5143: return TEXT("Qualcomm");
+	case 0x8086: return TEXT("Intel");
+	default: return TEXT("Unknown");
+	}
+}
\ No newline at end of file
diff --git a/Engine/Source/Runtime/RHI/Public/RHI.h b/Engine/Source/Runtime/RHI/Public/RHI.h
--- a/Engine/Source/Runtime/RHI/Public/RHI.h
+++ b/Engine/Source/Runtime/RHI/Public/RHI.h
@@ -90,6 +90,9 @@ RHI_API bool IsRHIDeviceIntel();
 // to trigger GPU specific optimizations and fallbacks
 RHI_API bool IsRHIDeviceNVIDIA();
 
+// helper to convert GRHIVendorId into a printable string, or "Unknown" if unknown.
+RHI_API const TCHAR* RHIVendorIdToString();
+
 /** true if PF_G8 render targets are supported */
 extern RHI_API bool GSupportsRenderTargetFormat_PF_G8;
 
diff --git a/Engine/Source/Runtime/Renderer/Private/MobileBasePassRendering.cpp b/Engine/Source/Runtime/Renderer/Private/MobileBasePassRendering.cpp
--- a/Engine/Source/Runtime/Renderer/Private/MobileBasePassRendering.cpp
+++ b/Engine/Source/Runtime/Renderer/Private/MobileBasePassRendering.cpp
@@ -9,6 +9,13 @@
 #include "SceneUtils.h"
 #include "MaterialShaderQualitySettings.h"
 
+// Changing this causes a full shader recompile
+static TAutoConsoleVariable<int32> CVarMobileDisableVertexFog(
+	TEXT("r.Mobile.DisableVertexFog"),
+	0,
+	TEXT("Set to 1 to disable vertex fogging in all mobile shaders."),
+	ECVF_ReadOnly | ECVF_RenderThreadSafe);
+
 #define IMPLEMENT_MOBILE_SHADING_BASEPASS_LIGHTMAPPED_VERTEX_SHADER_TYPE(LightMapPolicyType,LightMapPolicyName) \
 	typedef TMobileBasePassVS< LightMapPolicyType, LDR_GAMMA_32 > TMobileBasePassVS##LightMapPolicyName##LDRGamma32; \
 	typedef TMobileBasePassVS< LightMapPolicyType, HDR_LINEAR_64 > TMobileBasePassVS##LightMapPolicyName##HDRLinear64; \
@@ -71,7 +78,8 @@ bool TMobileBasePassPSPolicyParamType<PixelParametersType, NumDynamicPointLights
 	OutEnvironment.SetDefine(TEXT("MOBILE_QL_FORCE_NONMETAL"), QualityOverrides.bEnableOverride && QualityOverrides.bForceNonMetal != 0 ? 1u : 0u);
 	OutEnvironment.SetDefine(TEXT("QL_FORCEDISABLE_LM_DIRECTIONALITY"), QualityOverrides.bEnableOverride && QualityOverrides.bForceDisableLMDirectionality != 0 ? 1u : 0u);
 	OutEnvironment.SetDefine(TEXT("MOBILE_QL_FORCE_LQ_REFLECTIONS"), QualityOverrides.bEnableOverride && QualityOverrides.bForceLQReflections != 0 ? 1u : 0u);
-	
+	OutEnvironment.SetDefine(TEXT("MOBILE_CSM_QUALITY"), (uint32)QualityOverrides.MobileCSMQuality);
+
 	return true;
 }
 
@@ -521,8 +529,9 @@ namespace EBasePassSort
 TAutoConsoleVariable<int32> GSortBasePass(TEXT("r.ForwardBasePassSort"),0,
 	TEXT("How to sort the mobile base pass:\n")
 	TEXT("\t0: Decide automatically based on the hardware.\n")
-	TEXT("\t1: Sort drawing policies.\n")
-	TEXT("\t2: Sort drawing policies and the meshes within them."),
+	TEXT("\t1: No sorting.\n")
+	TEXT("\t2: Sort drawing policies.\n")
+	TEXT("\t3: Sort drawing policies and the meshes within them."),
 	ECVF_RenderThreadSafe);
 TAutoConsoleVariable<int32> GMaxBasePassDraws(TEXT("r.MaxForwardBasePassDraws"),0,TEXT("Stops rendering static mobile base pass draws after the specified number of times. Useful for seeing the order in which meshes render when optimizing."),ECVF_RenderThreadSafe);
 
diff --git a/Engine/Source/Runtime/Renderer/Private/MobileSceneCaptureRendering.cpp b/Engine/Source/Runtime/Renderer/Private/MobileSceneCaptureRendering.cpp
new file mode 100644
--- /dev/null
+++ b/Engine/Source/Runtime/Renderer/Private/MobileSceneCaptureRendering.cpp
@@ -0,0 +1,397 @@
+// Copyright 1998-2016 Epic Games, Inc. All Rights Reserved.
+
+/*=============================================================================
+MobileSceneCaptureRendering.cpp - Mobile specific scene capture code.
+=============================================================================*/
+
+#include "RendererPrivate.h"
+#include "ScenePrivate.h"
+#include "SceneFilterRendering.h"
+#include "UniformBuffer.h"
+#include "ShaderParameters.h"
+#include "ScreenRendering.h"
+#include "PostProcessAmbient.h"
+#include "PostProcessing.h"
+#include "SceneUtils.h"
+
+
+/**
+* Shader set for the copy of scene color to capture target, decoding mosaic or RGBE encoded HDR image as part of a
+* copy operation. Alpha channel will contain opacity information. (Determined from depth buffer content)
+*/
+
+// Use same defines as deferred for capture source defines,
+extern const TCHAR* GShaderSourceModeDefineName[];
+
+template<bool bDemosaic, ESceneCaptureSource CaptureSource>
+class FMobileSceneCaptureCopyPS : public FGlobalShader
+{
+	DECLARE_SHADER_TYPE(FMobileSceneCaptureCopyPS, Global)
+public:
+
+	static bool ShouldCache(EShaderPlatform Platform) { return IsMobilePlatform(Platform); }
+
+	FMobileSceneCaptureCopyPS(const ShaderMetaType::CompiledShaderInitializerType& Initializer) :
+		FGlobalShader(Initializer)
+	{
+		InTexture.Bind(Initializer.ParameterMap, TEXT("InTexture"), SPF_Mandatory);
+		InTextureSampler.Bind(Initializer.ParameterMap, TEXT("InTextureSampler"));
+		SceneTextureParameters.Bind(Initializer.ParameterMap);
+	}
+	FMobileSceneCaptureCopyPS() {}
+
+	static void ModifyCompilationEnvironment(EShaderPlatform Platform, FShaderCompilerEnvironment& OutEnvironment)
+	{
+		OutEnvironment.SetDefine(TEXT("MOBILE_FORCE_DEPTH_TEXTURE_READS"), 1u);
+		OutEnvironment.SetDefine(TEXT("DECODING_MOSAIC"), bDemosaic ? 1u : 0u);
+		const TCHAR* DefineName = GShaderSourceModeDefineName[CaptureSource];
+		if (DefineName)
+		{
+			OutEnvironment.SetDefine(DefineName, 1);
+		}
+	}
+
+	void SetParameters(FRHICommandList& RHICmdList, const FViewInfo& View, FSamplerStateRHIParamRef SamplerStateRHI, FTextureRHIParamRef TextureRHI)
+	{
+		FGlobalShader::SetParameters(RHICmdList, GetPixelShader(), View);
+		SetTextureParameter(RHICmdList, GetPixelShader(), InTexture, InTextureSampler, SamplerStateRHI, TextureRHI);
+		SceneTextureParameters.Set(RHICmdList, GetPixelShader(), View);
+	}
+
+	virtual bool Serialize(FArchive& Ar) override
+	{
+		bool bShaderHasOutdatedParameters = FGlobalShader::Serialize(Ar);
+		Ar << InTexture;
+		Ar << InTextureSampler;
+		Ar << SceneTextureParameters;
+		return bShaderHasOutdatedParameters;
+	}
+
+private:
+	FShaderResourceParameter InTexture;
+	FShaderResourceParameter InTextureSampler;
+	FSceneTextureShaderParameters SceneTextureParameters;
+};
+
+/**
+* A vertex shader for rendering a textured screen element.
+* Additional texcoords are used when demosaic is required.
+*/
+template<bool bDemosaic>
+class FMobileSceneCaptureCopyVS : public FGlobalShader
+{
+	DECLARE_SHADER_TYPE(FMobileSceneCaptureCopyVS, Global)
+public:
+
+	static bool ShouldCache(EShaderPlatform Platform) { return IsMobilePlatform(Platform); }
+
+	FMobileSceneCaptureCopyVS(const ShaderMetaType::CompiledShaderInitializerType& Initializer) :
+		FGlobalShader(Initializer)
+	{
+		InvTexSizeParameter.Bind(Initializer.ParameterMap, TEXT("InvTexSize"));
+	}
+	FMobileSceneCaptureCopyVS() {}
+
+	static void ModifyCompilationEnvironment(EShaderPlatform Platform, FShaderCompilerEnvironment& OutEnvironment)
+	{
+		OutEnvironment.SetDefine(TEXT("DECODING_MOSAIC"), bDemosaic ? 1u : 0u);
+	}
+
+	void SetParameters(FRHICommandList& RHICmdList, const FViewInfo& View, const FIntPoint& SourceTexSize)
+	{
+		FGlobalShader::SetParameters(RHICmdList, GetVertexShader(), View);
+		if (InvTexSizeParameter.IsBound())
+		{
+			FVector2D InvTexSize(1.0f / SourceTexSize.X, 1.0f / SourceTexSize.Y);
+			SetShaderValue(RHICmdList, GetVertexShader(), InvTexSizeParameter, InvTexSize);
+		}
+	}
+
+	virtual bool Serialize(FArchive& Ar) override
+	{
+		bool bShaderHasOutdatedParameters = FGlobalShader::Serialize(Ar);
+		Ar << InvTexSizeParameter;
+		return bShaderHasOutdatedParameters;
+	}
+
+	FShaderParameter InvTexSizeParameter;
+};
+
+#define IMPLEMENT_MOBILE_SCENE_CAPTURECOPY(SCENETYPE) \
+typedef FMobileSceneCaptureCopyPS<false,SCENETYPE> FMobileSceneCaptureCopyPS##SCENETYPE;\
+typedef FMobileSceneCaptureCopyPS<true,SCENETYPE> FMobileSceneCaptureCopyPS_Mosaic##SCENETYPE;\
+IMPLEMENT_SHADER_TYPE(template<>, FMobileSceneCaptureCopyPS##SCENETYPE, TEXT("MobileSceneCapture"), TEXT("MainCopyPS"), SF_Pixel); \
+IMPLEMENT_SHADER_TYPE(template<>, FMobileSceneCaptureCopyPS_Mosaic##SCENETYPE, TEXT("MobileSceneCapture"), TEXT("MainCopyPS"), SF_Pixel);
+
+IMPLEMENT_MOBILE_SCENE_CAPTURECOPY(SCS_SceneColorHDR);
+IMPLEMENT_MOBILE_SCENE_CAPTURECOPY(SCS_SceneColorHDRNoAlpha);
+IMPLEMENT_MOBILE_SCENE_CAPTURECOPY(SCS_SceneColorSceneDepth);
+IMPLEMENT_MOBILE_SCENE_CAPTURECOPY(SCS_SceneDepth);
+IMPLEMENT_SHADER_TYPE(template<>, FMobileSceneCaptureCopyVS<false>, TEXT("MobileSceneCapture"), TEXT("MainCopyVS"), SF_Vertex);
+IMPLEMENT_SHADER_TYPE(template<>, FMobileSceneCaptureCopyVS<true>, TEXT("MobileSceneCapture"), TEXT("MainCopyVS"), SF_Vertex);
+ 
+
+template <bool bDemosaic, ESceneCaptureSource CaptureSource>
+static FShader* SetCaptureToTargetShaders(FRHICommandListImmediate& RHICmdList, FViewInfo& View, const FIntPoint& SourceTexSize, FTextureRHIParamRef SourceTextureRHI)
+{
+	TShaderMapRef<FMobileSceneCaptureCopyVS<bDemosaic>> VertexShader(View.ShaderMap);
+	TShaderMapRef<FMobileSceneCaptureCopyPS<bDemosaic, CaptureSource>> PixelShader(View.ShaderMap);
+	static FGlobalBoundShaderState BoundShaderState;
+	SetGlobalBoundShaderState(RHICmdList, View.GetFeatureLevel(), BoundShaderState, GFilterVertexDeclaration.VertexDeclarationRHI, *VertexShader, *PixelShader);
+
+	VertexShader->SetParameters(RHICmdList, View, SourceTexSize);
+	PixelShader->SetParameters(RHICmdList, View, TStaticSamplerState<SF_Point>::GetRHI(), SourceTextureRHI);
+
+	return *VertexShader;
+}
+
+template <bool bDemosaic>
+static FShader* SetCaptureToTargetShaders(FRHICommandListImmediate& RHICmdList, ESceneCaptureSource CaptureSource, FViewInfo& View, const FIntPoint& SourceTexSize, FTextureRHIParamRef SourceTextureRHI)
+{
+	switch (CaptureSource)
+	{
+		case SCS_SceneColorHDR:
+			return SetCaptureToTargetShaders<bDemosaic, SCS_SceneColorHDR>(RHICmdList, View, SourceTexSize, SourceTextureRHI);
+		case SCS_SceneColorHDRNoAlpha:
+			return SetCaptureToTargetShaders<bDemosaic, SCS_SceneColorHDRNoAlpha>(RHICmdList, View, SourceTexSize, SourceTextureRHI);
+		case SCS_SceneColorSceneDepth:
+			return SetCaptureToTargetShaders<bDemosaic, SCS_SceneColorSceneDepth>(RHICmdList, View, SourceTexSize, SourceTextureRHI);
+		case SCS_SceneDepth:
+			return SetCaptureToTargetShaders<bDemosaic, SCS_SceneDepth>(RHICmdList, View, SourceTexSize, SourceTextureRHI);
+		default:
+			checkNoEntry();
+			return nullptr;
+	}
+}
+
+// Copies into render target, optionally flipping it in the Y-axis
+static void CopyCaptureToTarget(
+	FRHICommandListImmediate& RHICmdList, 
+	const FRenderTarget* Target, 
+	const FIntPoint& TargetSize, 
+	FViewInfo& View, 
+	const FIntRect& ViewRect, 
+	FTextureRHIParamRef SourceTextureRHI, 
+	bool bNeedsFlippedRenderTarget,
+	FSceneRenderer* SceneRenderer)
+{
+	ESceneCaptureSource CaptureSource = View.Family->SceneCaptureSource;
+	ESceneCaptureCompositeMode CaptureCompositeMode = View.Family->SceneCaptureCompositeMode;
+
+	// Normal and BaseColor not supported on mobile, fall back to scene colour.
+	if (CaptureSource == SCS_Normal || CaptureSource == SCS_BaseColor)
+	{
+		CaptureSource = SCS_SceneColorHDR;
+	}
+
+	ERenderTargetLoadAction RTLoadAction;
+	if (CaptureSource == SCS_SceneColorHDR && CaptureCompositeMode == SCCM_Composite)
+	{
+		// Blend with existing render target color. Scene capture color is already pre-multiplied by alpha.
+		RHICmdList.SetBlendState(TStaticBlendState<CW_RGBA, BO_Add, BF_One, BF_SourceAlpha, BO_Add, BF_Zero, BF_SourceAlpha>::GetRHI());
+		RTLoadAction = ERenderTargetLoadAction::ELoad;
+	}
+	else if (CaptureSource == SCS_SceneColorHDR && CaptureCompositeMode == SCCM_Additive)
+	{
+		// Add to existing render target color. Scene capture color is already pre-multiplied by alpha.
+		RHICmdList.SetBlendState(TStaticBlendState<CW_RGBA, BO_Add, BF_One, BF_One, BO_Add, BF_Zero, BF_SourceAlpha>::GetRHI());
+		RTLoadAction = ERenderTargetLoadAction::ELoad;
+	}
+	else
+	{
+		RTLoadAction = ERenderTargetLoadAction::ENoAction;
+		RHICmdList.SetBlendState(TStaticBlendState<>::GetRHI());
+	}
+
+	RHICmdList.SetRasterizerState(TStaticRasterizerState<FM_Solid, CM_None>::GetRHI());
+	RHICmdList.SetDepthStencilState(TStaticDepthStencilState<false, CF_Always>::GetRHI());
+
+	FRHIRenderTargetView ColorView(Target->GetRenderTargetTexture(), 0, -1, RTLoadAction, ERenderTargetStoreAction::EStore);
+	FRHISetRenderTargetsInfo Info(1, &ColorView, FRHIDepthRenderTargetView());
+	RHICmdList.SetRenderTargetsAndClear(Info);
+
+	const bool bUsingDemosaic = IsMobileHDRMosaic();
+	FShader* VertexShader;
+	FIntPoint SourceTexSize = bNeedsFlippedRenderTarget ? TargetSize : FSceneRenderTargets::Get(RHICmdList).GetBufferSizeXY();
+	if (bUsingDemosaic)
+	{
+		VertexShader = SetCaptureToTargetShaders<true>(RHICmdList, CaptureSource, View, SourceTexSize, SourceTextureRHI);
+	}
+	else
+	{
+		VertexShader = SetCaptureToTargetShaders<false>(RHICmdList, CaptureSource, View, SourceTexSize, SourceTextureRHI);
+	}
+
+	if (bNeedsFlippedRenderTarget)
+	{
+		DrawRectangle(
+			RHICmdList,
+			ViewRect.Min.X, ViewRect.Min.Y,
+			ViewRect.Width(), ViewRect.Height(),
+			ViewRect.Min.X, ViewRect.Height() - ViewRect.Min.Y,
+			ViewRect.Width(), -ViewRect.Height(),
+			TargetSize,
+			SourceTexSize,
+			VertexShader,
+			EDRF_UseTriangleOptimization);
+	}
+	else
+	{
+		DrawRectangle(
+			RHICmdList,
+			ViewRect.Min.X, ViewRect.Min.Y,
+			ViewRect.Width(), ViewRect.Height(),
+			ViewRect.Min.X, ViewRect.Min.Y,
+			ViewRect.Width(), ViewRect.Height(),
+			TargetSize,
+			SourceTexSize,
+			VertexShader,
+			EDRF_UseTriangleOptimization);
+	}
+
+	// if opacity is needed.
+	if (CaptureSource == SCS_SceneColorHDR)
+	{
+		// render translucent opacity. (to scene color)
+		check(View.Family->Scene->GetShadingPath() == EShadingPath::Mobile);
+		FMobileSceneRenderer* MobileSceneRenderer = (FMobileSceneRenderer*)SceneRenderer;
+		FSceneRenderTargets& SceneContext = FSceneRenderTargets::Get(RHICmdList);
+		SceneContext.BeginRenderingSceneColor(RHICmdList, ESimpleRenderTargetMode::EClearColorExistingDepth);
+		MobileSceneRenderer->RenderInverseOpacity(RHICmdList, View);
+
+		// Set capture target.
+		FRHIRenderTargetView OpacityView(Target->GetRenderTargetTexture(), 0, -1, ERenderTargetLoadAction::ELoad, ERenderTargetStoreAction::EStore);
+		FRHISetRenderTargetsInfo OpacityInfo(1, &OpacityView, FRHIDepthRenderTargetView());
+		RHICmdList.SetRenderTargetsAndClear(OpacityInfo);
+
+		RHICmdList.SetRasterizerState(TStaticRasterizerState<FM_Solid, CM_None>::GetRHI());
+		RHICmdList.SetDepthStencilState(TStaticDepthStencilState<false, CF_Always>::GetRHI());
+		// Note lack of inverse, both the target and source images are already inverted.
+		RHICmdList.SetBlendState(TStaticBlendState<CW_ALPHA, BO_Add, BF_DestColor, BF_Zero, BO_Add, BF_Zero, BF_SourceAlpha>::GetRHI());
+
+		// Combine translucent opacity pass to earlier opaque pass to build final inverse opacity.
+		TShaderMapRef<FScreenVS> ScreenVertexShader(View.ShaderMap);
+		TShaderMapRef<FScreenPS> PixelShader(View.ShaderMap);
+		static FGlobalBoundShaderState BoundShaderState;
+		SetGlobalBoundShaderState(RHICmdList, View.GetFeatureLevel(), BoundShaderState, GFilterVertexDeclaration.VertexDeclarationRHI, *ScreenVertexShader, *PixelShader);
+
+		ScreenVertexShader->SetParameters(RHICmdList, View);
+		PixelShader->SetParameters(RHICmdList, TStaticSamplerState<SF_Point>::GetRHI(), SourceTextureRHI);
+
+		DrawRectangle(
+			RHICmdList,
+			ViewRect.Min.X, ViewRect.Min.Y,
+			ViewRect.Width(), ViewRect.Height(),
+			ViewRect.Min.X, ViewRect.Min.Y,
+			ViewRect.Width(), ViewRect.Height(),
+			TargetSize,
+			SourceTexSize,
+			*ScreenVertexShader,
+			EDRF_UseTriangleOptimization);
+	}
+}
+
+void UpdateSceneCaptureContentMobile_RenderThread(
+	FRHICommandListImmediate& RHICmdList,
+	FSceneRenderer* SceneRenderer,
+	FRenderTarget* RenderTarget,
+	FTexture* RenderTargetTexture,
+	const FName OwnerName,
+	const FResolveParams& ResolveParams)
+{
+	FMemMark MemStackMark(FMemStack::Get());
+
+	// update any resources that needed a deferred update
+	FDeferredUpdateResource::UpdateResources(RHICmdList);
+	bool bUseSceneTextures = SceneRenderer->ViewFamily.SceneCaptureSource != SCS_FinalColorLDR;
+
+	{
+#if WANTS_DRAW_MESH_EVENTS
+		FString EventName;
+		OwnerName.ToString(EventName);
+		SCOPED_DRAW_EVENTF(RHICmdList, SceneCaptureMobile, TEXT("SceneCaptureMobile %s"), *EventName);
+#else
+		SCOPED_DRAW_EVENT(RHICmdList, UpdateSceneCaptureContentMobile_RenderThread);
+#endif
+
+		const bool bIsMobileHDR = IsMobileHDR();
+		const bool bRHINeedsFlip = RHINeedsToSwitchVerticalAxis(GMaxRHIShaderPlatform);
+		// note that GLES will flip the image during post processing. this needs flipping again so it is correct for texture addressing.
+		const bool bNeedsFlippedRenderTarget = (!bIsMobileHDR || !bUseSceneTextures) && bRHINeedsFlip;
+
+		// Intermediate render target that will need to be flipped (needed on !IsMobileHDR())
+		TRefCountPtr<IPooledRenderTarget> FlippedPooledRenderTarget;
+
+		const FRenderTarget* Target = SceneRenderer->ViewFamily.RenderTarget;
+		if (bNeedsFlippedRenderTarget)
+		{
+			// We need to use an intermediate render target since the result will be flipped
+			auto& RenderTargetRHI = Target->GetRenderTargetTexture();
+			FPooledRenderTargetDesc Desc(FPooledRenderTargetDesc::Create2DDesc(Target->GetSizeXY(),
+				RenderTargetRHI.GetReference()->GetFormat(),
+				FClearValueBinding::None,
+				TexCreate_None,
+				TexCreate_RenderTargetable,
+				false));
+			GRenderTargetPool.FindFreeElement(RHICmdList, Desc, FlippedPooledRenderTarget, TEXT("SceneCaptureFlipped"));
+		}
+
+		// Helper class to allow setting render target
+		struct FRenderTargetOverride : public FRenderTarget
+		{
+			FRenderTargetOverride(FRHITexture2D* In)
+			{
+				RenderTargetTextureRHI = In;
+			}
+
+			virtual FIntPoint GetSizeXY() const { return FIntPoint(RenderTargetTextureRHI->GetSizeX(), RenderTargetTextureRHI->GetSizeY()); }
+
+			FTexture2DRHIRef GetTextureParamRef() { return RenderTargetTextureRHI; }
+		} FlippedRenderTarget(
+			FlippedPooledRenderTarget.GetReference()
+			? FlippedPooledRenderTarget.GetReference()->GetRenderTargetItem().TargetableTexture->GetTexture2D()
+			: nullptr);
+		FViewInfo& View = SceneRenderer->Views[0];
+		FIntRect ViewRect = View.ViewRect;
+		FIntRect UnconstrainedViewRect = View.UnconstrainedViewRect;
+		auto& RenderTargetRHI = Target->GetRenderTargetTexture();
+		SetRenderTarget(RHICmdList, RenderTargetRHI, NULL, true);
+		RHICmdList.ClearColorTexture(RenderTargetRHI, FLinearColor::Black, ViewRect);
+
+		// Render the scene normally
+		{
+			SCOPED_DRAW_EVENT(RHICmdList, RenderScene);
+
+			if (bNeedsFlippedRenderTarget)
+			{
+				// Hijack the render target
+				SceneRenderer->ViewFamily.RenderTarget = &FlippedRenderTarget; //-V506
+			}
+
+			SceneRenderer->Render(RHICmdList);
+
+			if (bNeedsFlippedRenderTarget)
+			{
+				// And restore it
+				SceneRenderer->ViewFamily.RenderTarget = Target;
+			}
+		}
+
+		const FIntPoint TargetSize(UnconstrainedViewRect.Width(), UnconstrainedViewRect.Height());
+		if (bNeedsFlippedRenderTarget)
+		{
+			// We need to flip this texture upside down (since we depended on tonemapping to fix this on the hdr path)
+			SCOPED_DRAW_EVENT(RHICmdList, FlipCapture);
+			CopyCaptureToTarget(RHICmdList, Target, TargetSize, View, ViewRect, FlippedRenderTarget.GetTextureParamRef(), true, SceneRenderer);
+		}
+		else if(bUseSceneTextures)
+		{
+			// Copy the captured scene into the destination texture
+			SCOPED_DRAW_EVENT(RHICmdList, CaptureSceneColor);
+			CopyCaptureToTarget(RHICmdList, Target, TargetSize, View, ViewRect, FSceneRenderTargets::Get(RHICmdList).GetSceneColorTexture(), false, SceneRenderer);
+		}
+
+		RHICmdList.CopyToResolveTarget(RenderTarget->GetRenderTargetTexture(), RenderTargetTexture->TextureRHI, false, ResolveParams);
+	}
+	FSceneRenderer::WaitForTasksClearSnapshotsAndDeleteSceneRenderer(RHICmdList, SceneRenderer);
+}
\ No newline at end of file
diff --git a/Engine/Source/Runtime/Renderer/Private/MobileSceneCaptureRendering.h b/Engine/Source/Runtime/Renderer/Private/MobileSceneCaptureRendering.h
new file mode 100644
--- /dev/null
+++ b/Engine/Source/Runtime/Renderer/Private/MobileSceneCaptureRendering.h
@@ -0,0 +1,11 @@
+// Copyright 1998 - 2016 Epic Games, Inc.All Rights Reserved.
+
+#pragma once
+
+void UpdateSceneCaptureContentMobile_RenderThread(
+	FRHICommandListImmediate& RHICmdList,
+	FSceneRenderer* SceneRenderer,
+	FRenderTarget* RenderTarget,
+	FTexture* RenderTargetTexture,
+	const FName OwnerName,
+	const FResolveParams& ResolveParams);
diff --git a/Engine/Source/Runtime/Renderer/Private/MobileShadingRenderer.cpp b/Engine/Source/Runtime/Renderer/Private/MobileShadingRenderer.cpp
--- a/Engine/Source/Runtime/Renderer/Private/MobileShadingRenderer.cpp
+++ b/Engine/Source/Runtime/Renderer/Private/MobileShadingRenderer.cpp
@@ -20,11 +20,18 @@
 
 uint32 GetShadowQuality();
 
+static TAutoConsoleVariable<int32> CVarMobileForceDepthResolve(
+	TEXT("r.Mobile.ForceDepthResolve"),
+	0,
+	TEXT("0: Depth buffer is resolved after opaque pass only when decals or modulated shadows are in use. (Default)\n")
+	TEXT("1: Depth buffer is always resolved after opaque pass.\n"),
+	ECVF_Scalability | ECVF_RenderThreadSafe);
 
 FMobileSceneRenderer::FMobileSceneRenderer(const FSceneViewFamily* InViewFamily,FHitProxyConsumer* HitProxyConsumer)
 	:	FSceneRenderer(InViewFamily, HitProxyConsumer)
 {
 	bModulatedShadowsInUse = false;
+	bPostProcessUsesDepthTexture = false;
 }
 
 /**
@@ -68,6 +75,8 @@ void FMobileSceneRenderer::InitViews(FRHICommandListImmediate& RHICmdList)
 
 	// Now that the indirect lighting cache is updated, we can update the primitive precomputed lighting buffers.
 	UpdatePrimitivePrecomputedLightingBuffers();
+
+	UpdatePostProcessUsageFlags();
 	
 	OnStartFrame();
 }
@@ -107,7 +116,7 @@ void FMobileSceneRenderer::Render(FRHICommandListImmediate& RHICmdList)
 	}
 
 	// Notify the FX system that the scene is about to be rendered.
-	if (Scene->FXSystem && ViewFamily.EngineShowFlags.Particles)
+	if (Scene->FXSystem && !Views[0].bIsPlanarReflection && ViewFamily.EngineShowFlags.Particles)
 	{
 		Scene->FXSystem->PreRender(RHICmdList, NULL);
 	}
@@ -131,6 +140,11 @@ void FMobileSceneRenderer::Render(FRHICommandListImmediate& RHICmdList)
 	const bool bStereoRenderingAndHMD = View.Family->EngineShowFlags.StereoRendering && View.Family->EngineShowFlags.HMDDistortion;
 	const bool bRenderToSceneColor = bStereoRenderingAndHMD || bRequiresUpscale || FSceneRenderer::ShouldCompositeEditorPrimitives(View) || View.bIsSceneCapture;
 
+	if (!bGammaSpace)
+	{
+		RenderCustomDepthPass(RHICmdList);
+	}
+
 	FTextureRHIParamRef SceneColor = nullptr;
 	if (bGammaSpace && !bRenderToSceneColor)
 	{
@@ -161,7 +175,7 @@ void FMobileSceneRenderer::Render(FRHICommandListImmediate& RHICmdList)
 	}
 
 	// Notify the FX system that opaque primitives have been rendered.
-	if (Scene->FXSystem && ViewFamily.EngineShowFlags.Particles)
+	if (Scene->FXSystem && !Views[0].bIsPlanarReflection && ViewFamily.EngineShowFlags.Particles)
 	{
 		//#todo-rco: This is switching to another RT!
 		Scene->FXSystem->PostRenderOpaque(RHICmdList);
@@ -216,8 +230,12 @@ void FMobileSceneRenderer::Render(FRHICommandListImmediate& RHICmdList)
 		// Resolve the scene color for post processing.
 		SceneContext.ResolveSceneColor(RHICmdList, FResolveRect(0, 0, ViewFamily.FamilySizeX, ViewFamily.FamilySizeY));
 
+		const bool bKeepDepthContent = bPostProcessUsesDepthTexture || (View.bIsSceneCapture && (ViewFamily.SceneCaptureSource == ESceneCaptureSource::SCS_SceneColorHDR || ViewFamily.SceneCaptureSource == ESceneCaptureSource::SCS_SceneColorSceneDepth));
 		// Drop depth and stencil before post processing to avoid export.
-		RHICmdList.DiscardRenderTargets(true, true, 0);
+		if (!bKeepDepthContent)
+		{
+			RHICmdList.DiscardRenderTargets(true, true, 0);
+		}
 	}
 
 	if (!bGammaSpace)
@@ -332,19 +350,20 @@ void FMobileSceneRenderer::ConditionalResolveSceneDepth(FRHICommandListImmediate
 		bool bSceneDepthInAlpha = (SceneContext.GetSceneColor()->GetDesc().Format == PF_FloatRGBA);
 		bool bOnChipDepthFetch = (GSupportsShaderDepthStencilFetch || (bSceneDepthInAlpha && GSupportsShaderFramebufferFetch));
 		
-		if (!bOnChipDepthFetch)
+		const bool bForceDepthResolve = CVarMobileForceDepthResolve.GetValueOnRenderThread() == 1;
+
+		if (!bOnChipDepthFetch || bForceDepthResolve )
 		{
 			// Only these features require depth texture
 			bool bDecals = ViewFamily.EngineShowFlags.Decals && Scene->Decals.Num();
 			bool bModulatedShadows = ViewFamily.EngineShowFlags.DynamicShadows && bModulatedShadowsInUse;
 
-			if (bDecals || bModulatedShadows)
+			if (bDecals || bModulatedShadows || bForceDepthResolve)
 			{
 				// Switch target to force hardware flush current depth to texture
 				FTextureRHIRef DummySceneColor = GSystemTextures.BlackDummy->GetRenderTargetItem().TargetableTexture;
 				FTextureRHIRef DummyDepthTarget = GSystemTextures.DepthDummy->GetRenderTargetItem().TargetableTexture;
 				SetRenderTarget(RHICmdList, DummySceneColor, DummyDepthTarget, ESimpleRenderTargetMode::EUninitializedColorClearDepth, FExclusiveDepthStencil::DepthWrite_StencilWrite);
-				RHICmdList.DiscardRenderTargets(true, true, 0);
 			}
 		}
 	}
@@ -397,7 +416,6 @@ void FMobileSceneRenderer::CreateDirectionalLightUniformBuffers(FSceneView& Scen
 		SceneView.MobileDirectionalLightUniformBuffers[ChannelIdx + 1] = TUniformBufferRef<FMobileDirectionalLightShaderParameters>::CreateUniformBufferImmediate(Params, UniformBuffer_SingleFrame);
 	}
 }
-
 class FCopyMobileMultiViewSceneColorPS : public FGlobalShader
 {
 	DECLARE_SHADER_TYPE(FCopyMobileMultiViewSceneColorPS, Global);
@@ -496,3 +514,32 @@ void FMobileSceneRenderer::CopyMobileMultiViewSceneColor(FRHICommandListImmediat
 			EDRF_UseTriangleOptimization);
 	}
 }
+
+void FMobileSceneRenderer::UpdatePostProcessUsageFlags()
+{
+	bPostProcessUsesDepthTexture = false;
+	// Find out whether post-process materials require SceneDepth lookups, otherwise renderer can discard depth buffer before starting post-processing pass
+	for (int32 ViewIndex = 0; ViewIndex < Views.Num(); ViewIndex++)
+	{
+		const FBlendableManager& BlendableManager = Views[ViewIndex].FinalPostProcessSettings.BlendableManager;
+		FBlendableEntry* BlendableIt = nullptr;
+
+		while (FPostProcessMaterialNode* DataPtr = BlendableManager.IterateBlendables<FPostProcessMaterialNode>(BlendableIt))
+		{
+			if (DataPtr->IsValid())
+			{
+				FMaterialRenderProxy* Proxy = DataPtr->GetMaterialInterface()->GetRenderProxy(false);
+				check(Proxy);
+
+				const FMaterial* Material = Proxy->GetMaterial(Views[ViewIndex].GetFeatureLevel());
+				check(Material);
+
+				if (Material->MaterialUsesSceneDepthLookup())
+				{
+					bPostProcessUsesDepthTexture = true;
+					break;
+				}
+			}
+		}
+	}
+}
diff --git a/Engine/Source/Runtime/Renderer/Private/MobileTranslucentRendering.cpp b/Engine/Source/Runtime/Renderer/Private/MobileTranslucentRendering.cpp
--- a/Engine/Source/Runtime/Renderer/Private/MobileTranslucentRendering.cpp
+++ b/Engine/Source/Runtime/Renderer/Private/MobileTranslucentRendering.cpp
@@ -243,43 +243,11 @@ bool FMobileTranslucencyDrawingPolicyFactory::DrawDynamicMesh(
 	return bDirty;
 }
 
-/**
- * Render a static mesh using a translucent draw policy
- * @return true if the mesh rendered
- */
-bool FMobileTranslucencyDrawingPolicyFactory::DrawStaticMesh(
-	FRHICommandList& RHICmdList, 
-	const FViewInfo& View,
-	ContextType DrawingContext,
-	const FStaticMesh& StaticMesh,
-	bool bPreFog,
-	const FPrimitiveSceneProxy* PrimitiveSceneProxy,
-	FHitProxyId HitProxyId
-	)
+template <class TDrawingPolicyFactory>
+void FTranslucentPrimSet::DrawPrimitivesForMobile(FRHICommandListImmediate& RHICmdList, const FViewInfo& View,
+	typename TDrawingPolicyFactory::ContextType& DrawingContext) const
 {
-	bool bDirty = false;
-	const FMaterial* Material = StaticMesh.MaterialRenderProxy->GetMaterial(View.GetFeatureLevel());
-
-	bDirty |= DrawDynamicMesh(
-		RHICmdList, 
-		View,
-		DrawingContext,
-		StaticMesh,
-		false,
-		bPreFog,
-		PrimitiveSceneProxy,
-		HitProxyId
-		);
-	return bDirty;
-}
-
-/*-----------------------------------------------------------------------------
-FTranslucentPrimSet
------------------------------------------------------------------------------*/
-
-void FTranslucentPrimSet::DrawPrimitivesForMobile(FRHICommandListImmediate& RHICmdList, const FViewInfo& View, const bool bRenderSeparateTranslucency) const
-{
-	ETranslucencyPass::Type TranslucencyPass = bRenderSeparateTranslucency ? ETranslucencyPass::TPT_SeparateTranslucency : ETranslucencyPass::TPT_StandardTranslucency;
+	ETranslucencyPass::Type TranslucencyPass = DrawingContext.ShouldRenderSeparateTranslucency() ? ETranslucencyPass::TPT_SeparateTranslucency : ETranslucencyPass::TPT_StandardTranslucency;
 
 	FInt32Range PassRange = SortedPrimsNum.GetPassRange(TranslucencyPass);
 
@@ -292,10 +260,8 @@ void FTranslucentPrimSet::DrawPrimitivesForMobile(FRHICommandListImmediate& RHIC
 
 		checkSlow(ViewRelevance.HasTranslucency());
 
-		if(ViewRelevance.bDrawRelevance)
+		if (ViewRelevance.bDrawRelevance)
 		{
-			FMobileTranslucencyDrawingPolicyFactory::ContextType Context(bRenderSeparateTranslucency);
-
 			// range in View.DynamicMeshElements[]
 			FInt32Range range = View.GetDynamicMeshElementRange(PrimitiveSceneInfo->GetIndex());
 
@@ -306,29 +272,21 @@ void FTranslucentPrimSet::DrawPrimitivesForMobile(FRHICommandListImmediate& RHIC
 				checkSlow(MeshBatchAndRelevance.PrimitiveSceneProxy == PrimitiveSceneInfo->Proxy);
 
 				const FMeshBatch& MeshBatch = *MeshBatchAndRelevance.Mesh;
-				FMobileTranslucencyDrawingPolicyFactory::DrawDynamicMesh(RHICmdList, View, Context, MeshBatch, false, false, MeshBatchAndRelevance.PrimitiveSceneProxy, MeshBatch.BatchHitProxyId);
+				TDrawingPolicyFactory::DrawDynamicMesh(RHICmdList, View, DrawingContext, MeshBatch, false, false, MeshBatchAndRelevance.PrimitiveSceneProxy, MeshBatch.BatchHitProxyId);
 			}
 
 			// Render static scene prim
-			if( ViewRelevance.bStaticRelevance )
+			if (ViewRelevance.bStaticRelevance)
 			{
 				// Render static meshes from static scene prim
-				for( int32 StaticMeshIdx=0; StaticMeshIdx < PrimitiveSceneInfo->StaticMeshes.Num(); StaticMeshIdx++ )
+				for (int32 StaticMeshIdx = 0; StaticMeshIdx < PrimitiveSceneInfo->StaticMeshes.Num(); StaticMeshIdx++)
 				{
 					FStaticMesh& StaticMesh = PrimitiveSceneInfo->StaticMeshes[StaticMeshIdx];
 					if (View.StaticMeshVisibilityMap[StaticMesh.Id]
 						// Only render static mesh elements using translucent materials
 						&& StaticMesh.IsTranslucent(View.GetFeatureLevel()))
 					{
-						FMobileTranslucencyDrawingPolicyFactory::DrawStaticMesh(
-							RHICmdList, 
-							View,
-							FMobileTranslucencyDrawingPolicyFactory::ContextType(bRenderSeparateTranslucency),
-							StaticMesh,
-							false,
-							PrimitiveSceneInfo->Proxy,
-							StaticMesh.BatchHitProxyId
-							);
+						TDrawingPolicyFactory::DrawDynamicMesh(RHICmdList, View, DrawingContext, StaticMesh, false, false, PrimitiveSceneInfo->Proxy, StaticMesh.BatchHitProxyId);
 					}
 				}
 			}
@@ -377,7 +335,8 @@ void FMobileSceneRenderer::RenderTranslucency(FRHICommandListImmediate& RHICmdLi
 			RHICmdList.SetDepthStencilState(TStaticDepthStencilState<false,CF_DepthNearOrEqual>::GetRHI());
 
 			// Draw only translucent prims that don't read from scene color
-			View.TranslucentPrimSet.DrawPrimitivesForMobile(RHICmdList, View, false);
+			FMobileTranslucencyDrawingPolicyFactory::ContextType DrawingContext(false);
+			View.TranslucentPrimSet.DrawPrimitivesForMobile<FMobileTranslucencyDrawingPolicyFactory>(RHICmdList, View, DrawingContext);
 			// Draw the view's mesh elements with the translucent drawing policy.
 			DrawViewElements<FMobileTranslucencyDrawingPolicyFactory>(RHICmdList, View, FMobileTranslucencyDrawingPolicyFactory::ContextType(false), SDPG_World, false);
 			// Draw the view's mesh elements with the translucent drawing policy.
@@ -385,3 +344,336 @@ void FMobileSceneRenderer::RenderTranslucency(FRHICommandListImmediate& RHICmdLi
 		}
 	}
 }
+
+//////////////////////////////////////////////////////////////////////////
+// Translucent material inverse opacity render code
+// Used to generate inverse opacity channel for scene captures that require opacity information.
+// See MobileSceneCaptureRendering for more details.
+
+/**
+* Vertex shader for mobile opacity only pass
+*/
+class FOpacityOnlyVS : public FMeshMaterialShader
+{
+	DECLARE_SHADER_TYPE(FOpacityOnlyVS, MeshMaterial);
+protected:
+
+	FOpacityOnlyVS() {}
+	FOpacityOnlyVS(const FMeshMaterialShaderType::CompiledShaderInitializerType& Initializer) :
+		FMeshMaterialShader(Initializer)
+	{
+	}
+
+public:
+
+	static bool ShouldCache(EShaderPlatform Platform, const FMaterial* Material, const FVertexFactoryType* VertexFactoryType)
+	{
+		return IsTranslucentBlendMode(Material->GetBlendMode()) && IsMobilePlatform(Platform);
+	}
+
+	virtual bool Serialize(FArchive& Ar) override
+	{
+		const bool result = FMeshMaterialShader::Serialize(Ar);
+		return result;
+	}
+
+	void SetParameters(
+		FRHICommandList& RHICmdList,
+		const FMaterialRenderProxy* MaterialRenderProxy,
+		const FMaterial& MaterialResource,
+		const FSceneView& View
+	)
+	{
+		FMeshMaterialShader::SetParameters(RHICmdList, GetVertexShader(), MaterialRenderProxy, MaterialResource, View, View.ViewUniformBuffer, ESceneRenderTargetsMode::DontSet);
+	}
+
+	void SetMesh(FRHICommandList& RHICmdList, const FVertexFactory* VertexFactory, const FSceneView& View, const FPrimitiveSceneProxy* Proxy, const FMeshBatchElement& BatchElement, const FMeshDrawingRenderState& DrawRenderState)
+	{
+		FMeshMaterialShader::SetMesh(RHICmdList, GetVertexShader(), VertexFactory, View, Proxy, BatchElement, DrawRenderState);
+	}
+};
+
+IMPLEMENT_MATERIAL_SHADER_TYPE(, FOpacityOnlyVS, TEXT("MobileOpacityShaders"), TEXT("MainVS"), SF_Vertex);
+
+/**
+* Pixel shader for mobile opacity only pass, writes opacity to alpha channel.
+*/
+class FOpacityOnlyPS : public FMeshMaterialShader
+{
+	DECLARE_SHADER_TYPE(FOpacityOnlyPS, MeshMaterial);
+public:
+
+	static bool ShouldCache(EShaderPlatform Platform, const FMaterial* Material, const FVertexFactoryType* VertexFactoryType)
+	{
+		return IsTranslucentBlendMode(Material->GetBlendMode()) && IsMobilePlatform(Platform);
+	}
+
+	static void ModifyCompilationEnvironment(EShaderPlatform Platform, const FMaterial* Material, FShaderCompilerEnvironment& OutEnvironment)
+	{
+		OutEnvironment.SetDefine(TEXT("MOBILE_FORCE_DEPTH_TEXTURE_READS"), 1u);
+	}
+
+	FOpacityOnlyPS(const ShaderMetaType::CompiledShaderInitializerType& Initializer) :
+		FMeshMaterialShader(Initializer)
+	{
+	}
+
+	FOpacityOnlyPS() {}
+
+	void SetParameters(FRHICommandList& RHICmdList, const FMaterialRenderProxy* MaterialRenderProxy, const FMaterial& MaterialResource, const FSceneView* View)
+	{
+		FMeshMaterialShader::SetParameters(RHICmdList, GetPixelShader(), MaterialRenderProxy, MaterialResource, *View, View->ViewUniformBuffer, ESceneRenderTargetsMode::DontSet);
+	}
+
+	void SetMesh(FRHICommandList& RHICmdList, const FVertexFactory* VertexFactory, const FSceneView& View, const FPrimitiveSceneProxy* Proxy, const FMeshBatchElement& BatchElement, const FMeshDrawingRenderState& DrawRenderState)
+	{
+		FMeshMaterialShader::SetMesh(RHICmdList, GetPixelShader(), VertexFactory, View, Proxy, BatchElement, DrawRenderState);
+	}
+
+	virtual bool Serialize(FArchive& Ar) override
+	{
+		bool bShaderHasOutdatedParameters = FMeshMaterialShader::Serialize(Ar);
+		return bShaderHasOutdatedParameters;
+	}
+};
+
+IMPLEMENT_MATERIAL_SHADER_TYPE(, FOpacityOnlyPS, TEXT("MobileOpacityShaders"), TEXT("MainPS"), SF_Pixel);
+IMPLEMENT_SHADERPIPELINE_TYPE_VSPS(MobileOpacityPipeline, FOpacityOnlyVS, FOpacityOnlyPS, true);
+
+class FMobileOpacityDrawingPolicy : public FMeshDrawingPolicy
+{
+public:
+
+	struct ContextDataType : public FMeshDrawingPolicy::ContextDataType
+	{
+	};
+
+	FMobileOpacityDrawingPolicy(
+		const FVertexFactory* InVertexFactory,
+		const FMaterialRenderProxy* InMaterialRenderProxy,
+		const FMaterial& InMaterialResource,
+		bool bIsTwoSided,
+		ERHIFeatureLevel::Type InFeatureLevel
+	) :
+		FMeshDrawingPolicy(InVertexFactory, InMaterialRenderProxy, InMaterialResource, DVSM_None,/*bInTwoSidedOverride=*/ bIsTwoSided)
+	{
+		const EMaterialTessellationMode TessellationMode = InMaterialResource.GetTessellationMode();
+		static const auto* CVar = IConsoleManager::Get().FindTConsoleVariableDataInt(TEXT("r.ShaderPipelines"));
+		bool bUseShaderPipelines = CVar && CVar->GetValueOnAnyThread() != 0;
+
+		ShaderPipeline = bUseShaderPipelines ? InMaterialResource.GetShaderPipeline(&MobileOpacityPipeline, InVertexFactory->GetType()) : nullptr;
+
+		if (ShaderPipeline)
+		{
+			VertexShader = ShaderPipeline->GetShader<FOpacityOnlyVS >();
+			PixelShader = ShaderPipeline->GetShader<FOpacityOnlyPS>();
+		}
+		else
+		{
+			VertexShader = InMaterialResource.GetShader<FOpacityOnlyVS >(VertexFactory->GetType());
+			PixelShader = InMaterialResource.GetShader<FOpacityOnlyPS>(InVertexFactory->GetType());
+		}
+	}
+
+	// FMeshDrawingPolicy interface.
+	FDrawingPolicyMatchResult Matches(const FMobileOpacityDrawingPolicy& Other) const
+	{
+		DRAWING_POLICY_MATCH_BEGIN
+			DRAWING_POLICY_MATCH(FMeshDrawingPolicy::Matches(Other)) &&
+			DRAWING_POLICY_MATCH(VertexShader == Other.VertexShader) &&
+			DRAWING_POLICY_MATCH(PixelShader == Other.PixelShader);
+		DRAWING_POLICY_MATCH_END
+	}
+
+	void SetSharedState(FRHICommandList& RHICmdList, const FSceneView* View, const FMobileOpacityDrawingPolicy::ContextDataType PolicyContext) const
+	{
+		VertexShader->SetParameters(RHICmdList, MaterialRenderProxy, *MaterialResource, *View);
+		PixelShader->SetParameters(RHICmdList, MaterialRenderProxy, *MaterialResource, View);
+
+		// Set the shared mesh resources.
+		FMeshDrawingPolicy::SetSharedState(RHICmdList, View, PolicyContext);
+	}
+
+	/**
+	* Create bound shader state using the vertex decl from the mesh draw policy
+	* as well as the shaders needed to draw the mesh
+	* @return new bound shader state object
+	*/
+	FBoundShaderStateInput GetBoundShaderStateInput(ERHIFeatureLevel::Type InFeatureLevel)
+	{
+		return FBoundShaderStateInput(
+			FMeshDrawingPolicy::GetVertexDeclaration(),
+			VertexShader->GetVertexShader(),
+			nullptr,
+			nullptr,
+			PixelShader->GetPixelShader(),
+			NULL);
+	}
+
+	void SetMeshRenderState(
+		FRHICommandList& RHICmdList,
+		const FSceneView& View,
+		const FPrimitiveSceneProxy* PrimitiveSceneProxy,
+		const FMeshBatch& Mesh,
+		int32 BatchElementIndex,
+		bool bBackFace,
+		const FMeshDrawingRenderState& DrawRenderState,
+		const ElementDataType& ElementData,
+		const ContextDataType PolicyContext
+	) const
+	{
+		const FMeshBatchElement& BatchElement = Mesh.Elements[BatchElementIndex];
+		VertexShader->SetMesh(RHICmdList, VertexFactory, View, PrimitiveSceneProxy, BatchElement, DrawRenderState);
+		PixelShader->SetMesh(RHICmdList, VertexFactory, View, PrimitiveSceneProxy, BatchElement, DrawRenderState);
+
+		FMeshDrawingPolicy::SetMeshRenderState(RHICmdList, View, PrimitiveSceneProxy, Mesh, BatchElementIndex, bBackFace, DrawRenderState, ElementData, PolicyContext);
+	}
+
+	friend int32 CompareDrawingPolicy(const FMobileOpacityDrawingPolicy& A, const FMobileOpacityDrawingPolicy& B);
+
+private:
+
+	FShaderPipeline* ShaderPipeline;
+	FOpacityOnlyVS* VertexShader;
+	FOpacityOnlyPS* PixelShader;
+};
+
+int32 CompareDrawingPolicy(const FMobileOpacityDrawingPolicy& A, const FMobileOpacityDrawingPolicy& B)
+{
+	COMPAREDRAWINGPOLICYMEMBERS(VertexShader);
+	COMPAREDRAWINGPOLICYMEMBERS(PixelShader);
+	COMPAREDRAWINGPOLICYMEMBERS(VertexFactory);
+	COMPAREDRAWINGPOLICYMEMBERS(MaterialRenderProxy);
+	COMPAREDRAWINGPOLICYMEMBERS(bIsTwoSidedMaterial);
+	return 0;
+}
+
+class FMobileOpacityDrawingPolicyFactory
+{
+public:
+
+	struct ContextType
+	{
+		bool ShouldRenderSeparateTranslucency() const { return false; }
+	};
+
+	static bool DrawDynamicMesh(
+		FRHICommandList& RHICmdList,
+		const FViewInfo& View,
+		ContextType DrawingContext,
+		const FMeshBatch& Mesh,
+		bool bBackFace,
+		bool bPreFog,
+		const FPrimitiveSceneProxy* PrimitiveSceneProxy,
+		FHitProxyId HitProxyId
+	)
+	{
+		return DrawMesh(
+			RHICmdList,
+			View,
+			DrawingContext,
+			Mesh,
+			Mesh.Elements.Num() == 1 ? 1ull : (1ull << Mesh.Elements.Num()) - 1ull,	// 1 bit set for each mesh element
+			bBackFace,
+			FMeshDrawingRenderState(),
+			bPreFog,
+			PrimitiveSceneProxy,
+			HitProxyId
+		);
+	}
+
+private:
+	/**
+	* Render a dynamic or static mesh using the opacity draw policy
+	* @return true if the mesh rendered
+	*/
+	static bool DrawMesh(
+		FRHICommandList& RHICmdList,
+		const FViewInfo& View,
+		ContextType DrawingContext,
+		const FMeshBatch& Mesh,
+		const uint64& BatchElementMask,
+		bool bBackFace,
+		const FMeshDrawingRenderState& DrawRenderState,
+		bool bPreFog,
+		const FPrimitiveSceneProxy* PrimitiveSceneProxy,
+		FHitProxyId HitProxyId
+	)
+	{
+		bool bDirty = false;
+
+		const FMaterialRenderProxy* MaterialRenderProxy = Mesh.MaterialRenderProxy;
+		const FMaterial* Material = MaterialRenderProxy->GetMaterial(View.GetFeatureLevel());
+		const EBlendMode BlendMode = Material->GetBlendMode();
+
+		bool bDraw = IsTranslucentBlendMode(BlendMode);
+		if (bDraw)
+		{
+			FMobileOpacityDrawingPolicy DrawingPolicy(Mesh.VertexFactory, MaterialRenderProxy, *MaterialRenderProxy->GetMaterial(View.GetFeatureLevel()), Material->IsTwoSided(), View.GetFeatureLevel());
+			RHICmdList.BuildAndSetLocalBoundShaderState(DrawingPolicy.GetBoundShaderStateInput(View.GetFeatureLevel()));
+			DrawingPolicy.SetSharedState(RHICmdList, &View, FMobileOpacityDrawingPolicy::ContextDataType());
+
+			int32 BatchElementIndex = 0;
+			uint64 Mask = BatchElementMask;
+			do
+			{
+				if (Mask & 1)
+				{
+					const bool bIsInstancedMesh = Mesh.Elements[BatchElementIndex].bIsInstancedMesh;
+					TDrawEvent<FRHICommandList> MeshEvent;
+					BeginMeshDrawEvent(RHICmdList, PrimitiveSceneProxy, Mesh, MeshEvent);
+
+					DrawingPolicy.SetMeshRenderState(RHICmdList, View, PrimitiveSceneProxy, Mesh, BatchElementIndex, bBackFace, DrawRenderState, FMeshDrawingPolicy::ElementDataType(), FMobileOpacityDrawingPolicy::ContextDataType());
+					DrawingPolicy.DrawMesh(RHICmdList, Mesh, BatchElementIndex);
+				}
+				Mask >>= 1;
+				BatchElementIndex++;
+			} while (Mask);
+
+			bDirty = true;
+		}		
+
+		return bDirty;
+	}
+};
+
+bool FMobileSceneRenderer::RenderInverseOpacity(FRHICommandListImmediate& RHICmdList, const FViewInfo& View)
+{
+	bool bDirty = false;
+
+	if (ShouldRenderTranslucency())
+	{
+		const bool bGammaSpace = !IsMobileHDR();
+		const bool bLinearHDR64 = !bGammaSpace && !IsMobileHDR32bpp();
+
+		{
+			if (!bGammaSpace)
+			{
+				FSceneRenderTargets::Get(RHICmdList).BeginRenderingTranslucency(RHICmdList, View);
+			}
+			else
+			{
+				// Mobile multi-view is not side by side stereo
+				const FViewInfo& TranslucentViewport = (View.bIsMobileMultiViewEnabled) ? Views[0] : View;
+				RHICmdList.SetViewport(TranslucentViewport.ViewRect.Min.X, TranslucentViewport.ViewRect.Min.Y, 0.0f, TranslucentViewport.ViewRect.Max.X, TranslucentViewport.ViewRect.Max.Y, 1.0f);
+			}
+			RHICmdList.SetRasterizerState(TStaticRasterizerState<FM_Solid, CM_None>::GetRHI());
+			// Enable depth test, disable depth writes.
+			RHICmdList.SetDepthStencilState(TStaticDepthStencilState<false, CF_DepthNearOrEqual>::GetRHI());
+			RHICmdList.SetBlendState(TStaticBlendState<CW_ALPHA, BO_Add, BF_DestColor, BF_Zero, BO_Add, BF_Zero, BF_InverseSourceAlpha>::GetRHI());
+
+			{
+				SCOPED_DRAW_EVENT(RHICmdList, InverseOpacity);
+				bDirty |= RenderInverseOpacityDynamic(RHICmdList, View);
+			}
+		}
+	}
+	return bDirty;
+}
+
+bool FMobileSceneRenderer::RenderInverseOpacityDynamic(FRHICommandListImmediate& RHICmdList, const FViewInfo& View)
+{
+	FMobileOpacityDrawingPolicyFactory::ContextType DrawingContext;
+	View.TranslucentPrimSet.DrawPrimitivesForMobile<FMobileOpacityDrawingPolicyFactory>(RHICmdList, View, DrawingContext);
+	return View.TranslucentPrimSet.NumPrims() > 0;
+}
diff --git a/Engine/Source/Runtime/Renderer/Private/PlanarReflectionRendering.cpp b/Engine/Source/Runtime/Renderer/Private/PlanarReflectionRendering.cpp
--- a/Engine/Source/Runtime/Renderer/Private/PlanarReflectionRendering.cpp
+++ b/Engine/Source/Runtime/Renderer/Private/PlanarReflectionRendering.cpp
@@ -113,7 +113,10 @@ void PrefilterPlanarReflection(FRHICommandListImmediate& RHICmdList, FViewInfo&
 	{
 		SCOPED_DRAW_EVENT(RHICmdList, PrefilterPlanarReflection);
 
-		FRHIRenderTargetView ColorView(Target->GetRenderTargetTexture(), 0, -1, ERenderTargetLoadAction::ENoAction, ERenderTargetStoreAction::EStore);
+		// Workaround for a possible driver bug on S7 Adreno, missing planar reflections
+		ERenderTargetLoadAction RTLoadAction = IsVulkanMobilePlatform(View.GetShaderPlatform()) ?  ERenderTargetLoadAction::EClear : ERenderTargetLoadAction::ENoAction;
+
+		FRHIRenderTargetView ColorView(Target->GetRenderTargetTexture(), 0, -1, RTLoadAction, ERenderTargetStoreAction::EStore);
 		FRHISetRenderTargetsInfo Info(1, &ColorView, FRHIDepthRenderTargetView());
 		RHICmdList.SetRenderTargetsAndClear(Info);
 
diff --git a/Engine/Source/Runtime/Renderer/Private/PostProcess/SceneRenderTargets.cpp b/Engine/Source/Runtime/Renderer/Private/PostProcess/SceneRenderTargets.cpp
--- a/Engine/Source/Runtime/Renderer/Private/PostProcess/SceneRenderTargets.cpp
+++ b/Engine/Source/Runtime/Renderer/Private/PostProcess/SceneRenderTargets.cpp
@@ -2439,7 +2439,7 @@ void FDeferredPixelShaderParameters::Set(TRHICmdList& RHICmdList, const ShaderRH
 
 	const auto FeatureLevel = View.GetFeatureLevel();
 
-	if (TextureMode == ESceneRenderTargetsMode::SetTextures && FeatureLevel >= ERHIFeatureLevel::SM4)
+	if (TextureMode == ESceneRenderTargetsMode::SetTextures)
 	{
 		// if there is no ambient occlusion it's better to have white there
 		IPooledRenderTarget* ScreenSpaceAO = SceneContext.ScreenSpaceAO;
@@ -2454,6 +2454,7 @@ void FDeferredPixelShaderParameters::Set(TRHICmdList& RHICmdList, const ShaderRH
 		{
 			CustomDepth = GSystemTextures.BlackDummy;
 		}
+		SetTextureParameter(RHICmdList, ShaderRHI, CustomDepthTexture, CustomDepthTextureSampler, TStaticSamplerState<>::GetRHI(), CustomDepth->GetRenderTargetItem().ShaderResourceTexture);
 
 		if (FeatureLevel >= ERHIFeatureLevel::SM4)
 		{
@@ -2466,7 +2467,6 @@ void FDeferredPixelShaderParameters::Set(TRHICmdList& RHICmdList, const ShaderRH
 			SetTextureParameter(RHICmdList, ShaderRHI, ScreenSpaceAOTextureMS, ScreenSpaceAO->GetRenderTargetItem().TargetableTexture);
 			SetTextureParameter(RHICmdList, ShaderRHI, ScreenSpaceAOTextureNonMS, ScreenSpaceAO->GetRenderTargetItem().ShaderResourceTexture);
 
-			SetTextureParameter(RHICmdList, ShaderRHI, CustomDepthTexture, CustomDepthTextureSampler, TStaticSamplerState<>::GetRHI(), CustomDepth->GetRenderTargetItem().ShaderResourceTexture);
 			SetTextureParameter(RHICmdList, ShaderRHI, CustomDepthTextureNonMS, CustomDepth->GetRenderTargetItem().ShaderResourceTexture);
 
 			if (CustomStencilTexture.IsBound())
diff --git a/Engine/Source/Runtime/Renderer/Private/SceneCaptureRendering.cpp b/Engine/Source/Runtime/Renderer/Private/SceneCaptureRendering.cpp
--- a/Engine/Source/Runtime/Renderer/Private/SceneCaptureRendering.cpp
+++ b/Engine/Source/Runtime/Renderer/Private/SceneCaptureRendering.cpp
@@ -13,10 +13,12 @@
 #include "PostProcessAmbient.h"
 #include "PostProcessing.h"
 #include "SceneUtils.h"
+#include "MobileSceneCaptureRendering.h"
 
 const TCHAR* GShaderSourceModeDefineName[] =
 {
 	TEXT("SOURCE_MODE_SCENE_COLOR_AND_OPACITY"),
+	TEXT("SOURCE_MODE_SCENE_COLOR_NO_ALPHA"),
 	nullptr,
 	TEXT("SOURCE_MODE_SCENE_COLOR_SCENE_DEPTH"),
 	TEXT("SOURCE_MODE_SCENE_DEPTH"),
@@ -71,7 +73,8 @@ private:
 	FDeferredPixelShaderParameters DeferredParameters;
 };
 
-IMPLEMENT_SHADER_TYPE(template<>,TSceneCapturePS<SCS_SceneColorHDR>,TEXT("SceneCapturePixelShader"),TEXT("Main"),SF_Pixel);
+IMPLEMENT_SHADER_TYPE(template<>, TSceneCapturePS<SCS_SceneColorHDR>, TEXT("SceneCapturePixelShader"), TEXT("Main"), SF_Pixel);
+IMPLEMENT_SHADER_TYPE(template<>, TSceneCapturePS<SCS_SceneColorHDRNoAlpha>, TEXT("SceneCapturePixelShader"), TEXT("Main"), SF_Pixel);
 IMPLEMENT_SHADER_TYPE(template<>,TSceneCapturePS<SCS_SceneColorSceneDepth>,TEXT("SceneCapturePixelShader"),TEXT("Main"),SF_Pixel);
 IMPLEMENT_SHADER_TYPE(template<>,TSceneCapturePS<SCS_SceneDepth>,TEXT("SceneCapturePixelShader"),TEXT("Main"),SF_Pixel);
 IMPLEMENT_SHADER_TYPE(template<>,TSceneCapturePS<SCS_Normal>,TEXT("SceneCapturePixelShader"),TEXT("Main"),SF_Pixel);
@@ -119,6 +122,15 @@ void FDeferredShadingSceneRenderer::CopySceneCaptureComponentToTarget(FRHIComman
 
 				PixelShader->SetParameters(RHICmdList, View);
 			}
+			else if (ViewFamily.SceneCaptureSource == SCS_SceneColorHDRNoAlpha)
+			{
+				TShaderMapRef<TSceneCapturePS<SCS_SceneColorHDRNoAlpha> > PixelShader(View.ShaderMap);
+
+				static FGlobalBoundShaderState BoundShaderState;
+				SetGlobalBoundShaderState(RHICmdList, View.GetFeatureLevel(), BoundShaderState, GFilterVertexDeclaration.VertexDeclarationRHI, *VertexShader, *PixelShader);
+
+				PixelShader->SetParameters(RHICmdList, View);
+			}
 			else if (ViewFamily.SceneCaptureSource == SCS_SceneColorSceneDepth)
 			{
 				TShaderMapRef<TSceneCapturePS<SCS_SceneColorSceneDepth> > PixelShader(View.ShaderMap);
@@ -176,165 +188,7 @@ void FDeferredShadingSceneRenderer::CopySceneCaptureComponentToTarget(FRHIComman
 	}
 }
 
-
-/**
-* Shader set for decoding a mosaic or RGBE encoded HDR image as part of a copy operation.
-*/
-template<bool bDemosaic>
-class FHDRDecodePS : public FGlobalShader
-{
-	DECLARE_SHADER_TYPE(FHDRDecodePS, Global)
-public:
-
-	static bool ShouldCache(EShaderPlatform Platform) { return IsMobilePlatform(Platform); }
-
-	FHDRDecodePS(const ShaderMetaType::CompiledShaderInitializerType& Initializer) :
-		FGlobalShader(Initializer)
-	{
-		InTexture.Bind(Initializer.ParameterMap, TEXT("InTexture"), SPF_Mandatory);
-		InTextureSampler.Bind(Initializer.ParameterMap, TEXT("InTextureSampler"));
-	}
-	FHDRDecodePS() {}
-
-	static void ModifyCompilationEnvironment(EShaderPlatform Platform, FShaderCompilerEnvironment& OutEnvironment)
-	{
-		OutEnvironment.SetDefine(TEXT("DECODING_MOSAIC"), bDemosaic ? 1u : 0u);
-	}
-
-	void SetParameters(FRHICommandList& RHICmdList, const FViewInfo& View, FSamplerStateRHIParamRef SamplerStateRHI, FTextureRHIParamRef TextureRHI)
-	{
-		FGlobalShader::SetParameters(RHICmdList, GetPixelShader(), View);
-		SetTextureParameter(RHICmdList, GetPixelShader(), InTexture, InTextureSampler, SamplerStateRHI, TextureRHI);
-	}
-
-	virtual bool Serialize(FArchive& Ar) override
-	{
-		bool bShaderHasOutdatedParameters = FGlobalShader::Serialize(Ar);
-		Ar << InTexture;
-		Ar << InTextureSampler;
-		return bShaderHasOutdatedParameters;
-	}
-
-private:
-	FShaderResourceParameter InTexture;
-	FShaderResourceParameter InTextureSampler;
-};
-
-/**
-* A vertex shader for rendering a textured screen element.
-*/
-template<bool bDemosaic>
-class FHDRDecodeVS : public FGlobalShader
-{
-	DECLARE_SHADER_TYPE(FHDRDecodeVS, Global)
-public:
-
-	static bool ShouldCache(EShaderPlatform Platform) { return IsMobilePlatform(Platform); }
-
-	FHDRDecodeVS(const ShaderMetaType::CompiledShaderInitializerType& Initializer) :
-		FGlobalShader(Initializer)
-	{
-		InvTexSizeParameter.Bind(Initializer.ParameterMap, TEXT("InvTexSize"));
-	}
-	FHDRDecodeVS() {}
-	
-	static void ModifyCompilationEnvironment(EShaderPlatform Platform, FShaderCompilerEnvironment& OutEnvironment)
-	{
-		OutEnvironment.SetDefine(TEXT("DECODING_MOSAIC"), bDemosaic ? 1u : 0u);
-	}
-
-	void SetParameters(FRHICommandList& RHICmdList, const FViewInfo& View, const FIntPoint& SourceTexSize)
-	{
-		FGlobalShader::SetParameters(RHICmdList, GetVertexShader(), View);
-		if(InvTexSizeParameter.IsBound())
-		{
-			FVector2D InvTexSize(1.0f / SourceTexSize.X, 1.0f / SourceTexSize.Y);
-			SetShaderValue(RHICmdList, GetVertexShader(), InvTexSizeParameter, InvTexSize);
-		}
-	}
-
-	virtual bool Serialize(FArchive& Ar) override
-	{
-		bool bShaderHasOutdatedParameters = FGlobalShader::Serialize(Ar);
-		Ar << InvTexSizeParameter;
-		return bShaderHasOutdatedParameters;
-	}
-
-	FShaderParameter InvTexSizeParameter;
-};
-
-IMPLEMENT_SHADER_TYPE(template<>, FHDRDecodePS<false>, TEXT("Decode32bppHDR"), TEXT("MainPS"), SF_Pixel);
-IMPLEMENT_SHADER_TYPE(template<>, FHDRDecodePS<true>, TEXT("Decode32bppHDR"), TEXT("MainPS"), SF_Pixel);
-
-IMPLEMENT_SHADER_TYPE(template<>, FHDRDecodeVS<false>, TEXT("Decode32bppHDR"), TEXT("MainVS"), SF_Vertex);
-IMPLEMENT_SHADER_TYPE(template<>, FHDRDecodeVS<true>, TEXT("Decode32bppHDR"), TEXT("MainVS"), SF_Vertex);
-
-template <bool bDemosaic>
-static FShader* SetCaptureToTargetShaders(FRHICommandListImmediate& RHICmdList, FViewInfo& View, const FIntPoint& SourceTexSize, FTextureRHIParamRef SourceTextureRHI)
-{
-	TShaderMapRef<FHDRDecodeVS<bDemosaic>> VertexShader(View.ShaderMap);
-	TShaderMapRef<FHDRDecodePS<bDemosaic>> PixelShader(View.ShaderMap);
-	static FGlobalBoundShaderState BoundShaderState;
-	SetGlobalBoundShaderState(RHICmdList, View.GetFeatureLevel(), BoundShaderState, GFilterVertexDeclaration.VertexDeclarationRHI, *VertexShader, *PixelShader);
-
-	VertexShader->SetParameters(RHICmdList, View, SourceTexSize);
-	PixelShader->SetParameters(RHICmdList, View, TStaticSamplerState<SF_Point>::GetRHI(), SourceTextureRHI);
-
-	return *VertexShader;
-}
-
-// Copies into render target, optionally flipping it in the Y-axis
-static void CopyCaptureToTarget(FRHICommandListImmediate& RHICmdList, const FRenderTarget* Target, const FIntPoint& TargetSize, FViewInfo& View, const FIntRect& ViewRect, FTextureRHIParamRef SourceTextureRHI, bool bNeedsFlippedRenderTarget)
-{
-	FRHIRenderTargetView ColorView(Target->GetRenderTargetTexture(), 0, -1, ERenderTargetLoadAction::ENoAction, ERenderTargetStoreAction::EStore);
-	FRHISetRenderTargetsInfo Info(1, &ColorView, FRHIDepthRenderTargetView());
-	RHICmdList.SetRenderTargetsAndClear(Info);
-
-	RHICmdList.SetRasterizerState(TStaticRasterizerState<FM_Solid, CM_None>::GetRHI());
-	RHICmdList.SetDepthStencilState(TStaticDepthStencilState<false, CF_Always>::GetRHI());
-	RHICmdList.SetBlendState(TStaticBlendState<>::GetRHI());
-
-	const bool bUsingDemosaic = IsMobileHDRMosaic();
-	FShader* VertexShader;
-	FIntPoint SourceTexSize = bNeedsFlippedRenderTarget ? TargetSize : FSceneRenderTargets::Get(RHICmdList).GetBufferSizeXY();
-	if (bUsingDemosaic)
-	{
-		VertexShader = SetCaptureToTargetShaders<true>(RHICmdList, View, SourceTexSize, SourceTextureRHI);
-	}
-	else
-	{
-		VertexShader = SetCaptureToTargetShaders<false>(RHICmdList, View, SourceTexSize, SourceTextureRHI);
-	}
-
-	if (bNeedsFlippedRenderTarget)
-	{
-		DrawRectangle(
-			RHICmdList,
-			ViewRect.Min.X, ViewRect.Min.Y,
-			ViewRect.Width(), ViewRect.Height(),
-			ViewRect.Min.X, ViewRect.Height() - ViewRect.Min.Y,
-			ViewRect.Width(), -ViewRect.Height(),
-			TargetSize,
-			SourceTexSize,
-			VertexShader,
-			EDRF_UseTriangleOptimization);
-	}
-	else
-	{
-		DrawRectangle(
-			RHICmdList,
-			ViewRect.Min.X, ViewRect.Min.Y,
-			ViewRect.Width(), ViewRect.Height(),
-			ViewRect.Min.X, ViewRect.Min.Y,
-			ViewRect.Width(), ViewRect.Height(),
-			TargetSize,
-			SourceTexSize,
-			VertexShader,
-			EDRF_UseTriangleOptimization);
-	}
-}
-
-static void UpdateSceneCaptureContent_RenderThread(
+static void UpdateSceneCaptureContentDeferred_RenderThread(
 	FRHICommandListImmediate& RHICmdList, 
 	FSceneRenderer* SceneRenderer, 
 	FRenderTarget* RenderTarget, 
@@ -346,8 +200,6 @@ static void UpdateSceneCaptureContent_RenderThread(
 
 	// update any resources that needed a deferred update
 	FDeferredUpdateResource::UpdateResources(RHICmdList);
-	bool bUseSceneTextures = SceneRenderer->ViewFamily.SceneCaptureSource != SCS_FinalColorLDR;
-
 	{
 #if WANTS_DRAW_MESH_EVENTS
 		FString EventName;
@@ -357,43 +209,8 @@ static void UpdateSceneCaptureContent_RenderThread(
 		SCOPED_DRAW_EVENT(RHICmdList, UpdateSceneCaptureContent_RenderThread);
 #endif
 
-		const bool bIsMobileHDR = IsMobileHDR();
-		const bool bRHINeedsFlip = RHINeedsToSwitchVerticalAxis(GMaxRHIShaderPlatform);
-		// note that ES2 will flip the image during post processing. this needs flipping again so it is correct for texture addressing.
-		const bool bNeedsFlippedRenderTarget = (!bIsMobileHDR || !bUseSceneTextures) && bRHINeedsFlip;
-
-		// Intermediate render target that will need to be flipped (needed on !IsMobileHDR())
-		TRefCountPtr<IPooledRenderTarget> FlippedPooledRenderTarget;
-
 		const FRenderTarget* Target = SceneRenderer->ViewFamily.RenderTarget;
-		if (bNeedsFlippedRenderTarget)
-		{
-			// We need to use an intermediate render target since the result will be flipped
-			auto& RenderTargetRHI = Target->GetRenderTargetTexture();
-			FPooledRenderTargetDesc Desc(FPooledRenderTargetDesc::Create2DDesc(Target->GetSizeXY(), 
-				RenderTargetRHI.GetReference()->GetFormat(), 
-				FClearValueBinding::None,
-				TexCreate_None, 
-				TexCreate_RenderTargetable,
-				false));
-			GRenderTargetPool.FindFreeElement(RHICmdList, Desc, FlippedPooledRenderTarget, TEXT("SceneCaptureFlipped"));
-		}
-
-		// Helper class to allow setting render target
-		struct FRenderTargetOverride : public FRenderTarget
-		{
-			FRenderTargetOverride(FRHITexture2D* In)
-			{
-				RenderTargetTextureRHI = In;
-			}
-
-			virtual FIntPoint GetSizeXY() const { return FIntPoint(RenderTargetTextureRHI->GetSizeX(), RenderTargetTextureRHI->GetSizeY()); }
 
-			FTexture2DRHIRef GetTextureParamRef() { return RenderTargetTextureRHI; }
-		} FlippedRenderTarget(
-			FlippedPooledRenderTarget.GetReference()
-			? FlippedPooledRenderTarget.GetReference()->GetRenderTargetItem().TargetableTexture->GetTexture2D() 
-			: nullptr);
 		FViewInfo& View = SceneRenderer->Views[0];
 		FIntRect ViewRect = View.ViewRect;
 		FIntRect UnconstrainedViewRect = View.UnconstrainedViewRect;
@@ -403,43 +220,53 @@ static void UpdateSceneCaptureContent_RenderThread(
 		// Render the scene normally
 		{
 			SCOPED_DRAW_EVENT(RHICmdList, RenderScene);
-
-			if (bNeedsFlippedRenderTarget)
-			{
-				// Hijack the render target
-				SceneRenderer->ViewFamily.RenderTarget = &FlippedRenderTarget; //-V506
-			}
-
 			SceneRenderer->Render(RHICmdList);
-
-			if (bNeedsFlippedRenderTarget)
-			{
-				// And restore it
-				SceneRenderer->ViewFamily.RenderTarget = Target;
-			}
 		}
 
-		const FIntPoint TargetSize(UnconstrainedViewRect.Width(), UnconstrainedViewRect.Height());
-		if (bNeedsFlippedRenderTarget)
-		{
-			// We need to flip this texture upside down (since we depended on tonemapping to fix this on the hdr path)
-			SCOPED_DRAW_EVENT(RHICmdList, FlipCapture);
-			CopyCaptureToTarget(RHICmdList, Target, TargetSize, View, ViewRect, FlippedRenderTarget.GetTextureParamRef(), true);
-		}
-		else if (FSceneInterface::GetShadingPath(View.GetFeatureLevel()) == EShadingPath::Mobile)
+		// Note: When the ViewFamily.SceneCaptureSource requires scene textures (i.e. SceneCaptureSource != SCS_FinalColorLDR), the copy to RenderTarget 
+		// will be done in CopySceneCaptureComponentToTarget while the GBuffers are still alive for the frame.
+		RHICmdList.CopyToResolveTarget(RenderTarget->GetRenderTargetTexture(), RenderTargetTexture->TextureRHI, false, ResolveParams);
+	}
+
+	FSceneRenderer::WaitForTasksClearSnapshotsAndDeleteSceneRenderer(RHICmdList, SceneRenderer);
+}
+
+static void UpdateSceneCaptureContent_RenderThread(
+	FRHICommandListImmediate& RHICmdList,
+	FSceneRenderer* SceneRenderer,
+	FRenderTarget* RenderTarget,
+	FTexture* RenderTargetTexture,
+	const FName OwnerName,
+	const FResolveParams& ResolveParams)
+{
+	switch (SceneRenderer->Scene->GetShadingPath())
+	{
+		case EShadingPath::Mobile:
 		{
-			// Copy the captured scene into the destination texture
-			SCOPED_DRAW_EVENT(RHICmdList, CaptureSceneColor);
-			CopyCaptureToTarget(RHICmdList, Target, TargetSize, View, ViewRect, FSceneRenderTargets::Get(RHICmdList).GetSceneColorTexture(), false);
+			UpdateSceneCaptureContentMobile_RenderThread(
+				RHICmdList,
+				SceneRenderer,
+				RenderTarget,
+				RenderTargetTexture,
+				OwnerName,
+				ResolveParams);
+			break;
 		}
-		else if (bUseSceneTextures)
+		case EShadingPath::Deferred:
 		{
-			// Note: When the ViewFamily.SceneCaptureSource requires scene textures, the copy will be done in CopySceneCaptureComponentToTarget while the GBuffers are still alive for the frame
+			UpdateSceneCaptureContentDeferred_RenderThread(
+				RHICmdList,
+				SceneRenderer,
+				RenderTarget,
+				RenderTargetTexture,
+				OwnerName,
+				ResolveParams);
+			break;
 		}
-
-		RHICmdList.CopyToResolveTarget(RenderTarget->GetRenderTargetTexture(), RenderTargetTexture->TextureRHI, false, ResolveParams);
+		default:
+			checkNoEntry();
+			break;
 	}
-	FSceneRenderer::WaitForTasksClearSnapshotsAndDeleteSceneRenderer(RHICmdList, SceneRenderer);
 }
 
 void BuildProjectionMatrix(FIntPoint RenderTargetSize, ECameraProjectionMode::Type ProjectionType, float FOV, float InOrthoWidth, FMatrix& ProjectionMatrix)
diff --git a/Engine/Source/Runtime/Renderer/Private/SceneRendering.cpp b/Engine/Source/Runtime/Renderer/Private/SceneRendering.cpp
--- a/Engine/Source/Runtime/Renderer/Private/SceneRendering.cpp
+++ b/Engine/Source/Runtime/Renderer/Private/SceneRendering.cpp
@@ -421,6 +421,9 @@ void FViewInfo::Init()
 	NumBoxReflectionCaptures = 0;
 	NumSphereReflectionCaptures = 0;
 	FurthestReflectionCaptureDistance = 0;
+
+	// Disable HDR encoding for editor elements.
+	EditorSimpleElementCollector.BatchedElements.EnableMobileHDREncoding(false);
 }
 
 FViewInfo::~FViewInfo()
@@ -1402,12 +1405,6 @@ void FSceneRenderer::RenderCustomDepthPassAtLocation(FRHICommandListImmediate& R
 
 void FSceneRenderer::RenderCustomDepthPass(FRHICommandListImmediate& RHICmdList)
 {
-	if(FeatureLevel < ERHIFeatureLevel::SM4)
-	{
-		// not yet supported on lower end platforms
-		return;
-	}
-
 	// do we have primitives in this pass?
 	bool bPrimitives = false;
 
diff --git a/Engine/Source/Runtime/Renderer/Private/SceneRendering.h b/Engine/Source/Runtime/Renderer/Private/SceneRendering.h
--- a/Engine/Source/Runtime/Renderer/Private/SceneRendering.h
+++ b/Engine/Source/Runtime/Renderer/Private/SceneRendering.h
@@ -333,7 +333,8 @@ public:
 	* Draw all the primitives in this set for the mobile pipeline. 
 	* @param bRenderSeparateTranslucency - If false, only primitives with materials without mobile separate translucency enabled are rendered. Opposite if true.
 	*/
-	void DrawPrimitivesForMobile(FRHICommandListImmediate& RHICmdList, const class FViewInfo& View, const bool bRenderSeparateTranslucency) const;
+	template <class TDrawingPolicyFactory>
+	void DrawPrimitivesForMobile(FRHICommandListImmediate& RHICmdList, const FViewInfo& View, typename TDrawingPolicyFactory::ContextType& DrawingContext) const;
 
 	/**
 	* Insert a primitive to the translucency rendering list[s]
@@ -1369,6 +1370,8 @@ public:
 
 	virtual void RenderHitProxies(FRHICommandListImmediate& RHICmdList) override;
 
+	bool RenderInverseOpacity(FRHICommandListImmediate& RHICmdList, const FViewInfo& View);
+
 protected:
 	/** Finds the visible dynamic shadows for each view. */
 	void InitDynamicShadows(FRHICommandListImmediate& RHICmdList);
@@ -1405,9 +1408,16 @@ protected:
 	/** Copy scene color from the mobile multi-view render targat array to side by side stereo scene color */
 	void CopyMobileMultiViewSceneColor(FRHICommandListImmediate& RHICmdList);
 
+	/** Gather information about post-processing pass, which can be used by render for optimizations. Called by InitViews */
+	void UpdatePostProcessUsageFlags();
+
+	/** Render inverse opacity for the dynamic meshes. */
+	bool RenderInverseOpacityDynamic(FRHICommandListImmediate& RHICmdList, const FViewInfo& View);
+	
 private:
 
 	bool bModulatedShadowsInUse;
+	bool bPostProcessUsesDepthTexture;
 };
 
 // The noise textures need to be set in Slate too.
diff --git a/Engine/Source/Runtime/Renderer/Private/TranslucentRendering.h b/Engine/Source/Runtime/Renderer/Private/TranslucentRendering.h
--- a/Engine/Source/Runtime/Renderer/Private/TranslucentRendering.h
+++ b/Engine/Source/Runtime/Renderer/Private/TranslucentRendering.h
@@ -100,6 +100,8 @@ public:
 		ContextType(bool InbRenderingSeparateTranslucency)
 		:	bRenderingSeparateTranslucency(InbRenderingSeparateTranslucency)
 		{}
+
+		bool ShouldRenderSeparateTranslucency() const { return bRenderingSeparateTranslucency; }
 	};
 
 	/**
@@ -116,20 +118,6 @@ public:
 		const FPrimitiveSceneProxy* PrimitiveSceneProxy,
 		FHitProxyId HitProxyId
 		);
-
-	/**
-	* Render a dynamic mesh using a translucent draw policy 
-	* @return true if the mesh rendered
-	*/
-	static bool DrawStaticMesh(
-		FRHICommandList& RHICmdList, 
-		const FViewInfo& View,
-		ContextType DrawingContext,
-		const FStaticMesh& StaticMesh,
-		bool bPreFog,
-		const FPrimitiveSceneProxy* PrimitiveSceneProxy,
-		FHitProxyId HitProxyId
-		);
 };
 
 /** Represents a subregion of a volume texture. */
diff --git a/Engine/Source/Runtime/ShaderCore/Private/Shader.cpp b/Engine/Source/Runtime/ShaderCore/Private/Shader.cpp
--- a/Engine/Source/Runtime/ShaderCore/Private/Shader.cpp
+++ b/Engine/Source/Runtime/ShaderCore/Private/Shader.cpp
@@ -1775,6 +1775,12 @@ void ShaderMapAppendKeyString(EShaderPlatform Platform, FString& KeyString)
 		}
 	}
 
+	if (IsMobilePlatform(Platform))
+	{
+		static IConsoleVariable* CVar = IConsoleManager::Get().FindConsoleVariable(TEXT("r.Mobile.DisableVertexFog"));
+		KeyString += (CVar && CVar->GetInt() != 0) ? TEXT("_NoVFog") : TEXT("");
+	}
+
 	if (Platform == SP_PS4)
 	{
 		{
diff --git a/Engine/Source/Runtime/VulkanRHI/Private/VulkanMemory.cpp b/Engine/Source/Runtime/VulkanRHI/Private/VulkanMemory.cpp
--- a/Engine/Source/Runtime/VulkanRHI/Private/VulkanMemory.cpp
+++ b/Engine/Source/Runtime/VulkanRHI/Private/VulkanMemory.cpp
@@ -612,12 +612,12 @@ namespace VulkanRHI
 			++NumTypesPerHeap[MemoryProperties.memoryTypes[Index].heapIndex];
 		}
 
-		auto GetMemoryTypesFromProperties = [MemoryProperties](uint32 TypeBits, VkMemoryPropertyFlags Properties, TArray<uint32>& OutTypeIndices)
+		auto GetMemoryTypesFromProperties = [MemoryProperties](uint32 InTypeBits, VkMemoryPropertyFlags Properties, TArray<uint32>& OutTypeIndices)
 		{
 			// Search memtypes to find first index with those properties
-			for (uint32 i = 0; i < MemoryProperties.memoryTypeCount && TypeBits; i++)
+			for (uint32 i = 0; i < MemoryProperties.memoryTypeCount && InTypeBits; i++)
 			{
-				if ((TypeBits & 1) == 1)
+				if ((InTypeBits & 1) == 1)
 				{
 					// Type is available, does it match user properties?
 					if ((MemoryProperties.memoryTypes[i].propertyFlags & Properties) == Properties)
@@ -625,7 +625,7 @@ namespace VulkanRHI
 						OutTypeIndices.Add(i);
 					}
 				}
-				TypeBits >>= 1;
+				InTypeBits >>= 1;
 			}
 
 			for (int32 Index = OutTypeIndices.Num() - 1; Index >= 1; --Index)
diff --git a/Engine/Source/Runtime/VulkanRHI/Private/VulkanPendingState.cpp b/Engine/Source/Runtime/VulkanRHI/Private/VulkanPendingState.cpp
--- a/Engine/Source/Runtime/VulkanRHI/Private/VulkanPendingState.cpp
+++ b/Engine/Source/Runtime/VulkanRHI/Private/VulkanPendingState.cpp
@@ -437,7 +437,9 @@ FVulkanDescriptorPool::FVulkanDescriptorPool(FVulkanDevice* InDevice)
 	, NumAllocatedDescriptorSets(0)
 	, PeakAllocatedDescriptorSets(0)
 {
-	MaxDescriptorSets = 8192;
+	// Increased from 8192 to prevent Protostar crashing on Mali
+	MaxDescriptorSets = 16384;
+
 	const VkPhysicalDeviceLimits& Limits = Device->GetLimits();
 	FMemory::Memzero(MaxAllocatedTypes);
 	FMemory::Memzero(NumAllocatedTypes);
diff --git a/Engine/Source/Runtime/VulkanRHI/Private/VulkanRHI.cpp b/Engine/Source/Runtime/VulkanRHI/Private/VulkanRHI.cpp
--- a/Engine/Source/Runtime/VulkanRHI/Private/VulkanRHI.cpp
+++ b/Engine/Source/Runtime/VulkanRHI/Private/VulkanRHI.cpp
@@ -572,6 +572,7 @@ void FVulkanDynamicRHI::InitInstance()
 		// Initialize the RHI capabilities.
 		GRHIVendorId = Device->GetDeviceProperties().vendorID;
 		GRHIAdapterName = ANSI_TO_TCHAR(Props.deviceName);
+		GRHIAdapterInternalDriverVersion = FString::Printf(TEXT("%d.%d.%d"), VK_VERSION_MAJOR(Props.apiVersion), VK_VERSION_MINOR(Props.apiVersion), VK_VERSION_PATCH(Props.apiVersion));
 		GSupportsRenderTargetFormat_PF_G8 = false;	// #todo-rco
 		GSupportsQuads = false;	// Not supported in Vulkan
 		GRHISupportsTextureStreaming = false;	// #todo-rco
diff --git a/Engine/Source/Runtime/VulkanRHI/Private/VulkanRHIPrivate.h b/Engine/Source/Runtime/VulkanRHI/Private/VulkanRHIPrivate.h
--- a/Engine/Source/Runtime/VulkanRHI/Private/VulkanRHIPrivate.h
+++ b/Engine/Source/Runtime/VulkanRHI/Private/VulkanRHIPrivate.h
@@ -34,16 +34,18 @@ DECLARE_LOG_CATEGORY_EXTERN(LogVulkanRHI, Log, All);
 	#define VULKAN_DYNAMICALLYLOADED 0
 #endif
 
+#if PLATFORM_WINDOWS
+#include "AllowWindowsPlatformTypes.h"
+#endif
+
 #if VULKAN_DYNAMICALLYLOADED
 	#include "VulkanLoader.h"
 #else
-	#if PLATFORM_WINDOWS
-		#include "AllowWindowsPlatformTypes.h"
-	#endif
-#include <vulkan.h>
-	#if PLATFORM_WINDOWS
-		#include "HideWindowsPlatformTypes.h"
-	#endif
+	#include <vulkan.h>
+#endif
+
+#if PLATFORM_WINDOWS
+#include "HideWindowsPlatformTypes.h"
 #endif
 
 #if VULKAN_COMMANDWRAPPERS_ENABLE
diff --git a/Engine/Source/Runtime/VulkanRHI/Private/VulkanRenderTarget.cpp b/Engine/Source/Runtime/VulkanRHI/Private/VulkanRenderTarget.cpp
--- a/Engine/Source/Runtime/VulkanRHI/Private/VulkanRenderTarget.cpp
+++ b/Engine/Source/Runtime/VulkanRHI/Private/VulkanRenderTarget.cpp
@@ -919,6 +919,7 @@ struct FRenderTargetLayoutHashableStruct
 	uint32 ArraySliceIndex[MaxSimultaneousRenderTargets];
 	ERenderTargetLoadAction LoadAction[MaxSimultaneousRenderTargets];
 	ERenderTargetStoreAction StoreAction[MaxSimultaneousRenderTargets];
+	EPixelFormat TextureFormat[MaxSimultaneousRenderTargets + 1];
 
 	ERenderTargetLoadAction		DepthLoadAction;
 	ERenderTargetStoreAction	DepthStoreAction;
@@ -943,9 +944,11 @@ struct FRenderTargetLayoutHashableStruct
 			ArraySliceIndex[Index] = RTInfo.ColorRenderTarget[Index].ArraySliceIndex;
 			LoadAction[Index] = RTInfo.ColorRenderTarget[Index].LoadAction;
 			StoreAction[Index] = RTInfo.ColorRenderTarget[Index].StoreAction;
+			TextureFormat[Index] = Texture[Index] ? Texture[Index]->GetFormat() : PF_Unknown;
 		}
 
 		Texture[MaxSimultaneousRenderTargets] = RTInfo.DepthStencilRenderTarget.Texture;
+		TextureFormat[MaxSimultaneousRenderTargets] = Texture[MaxSimultaneousRenderTargets] ? Texture[MaxSimultaneousRenderTargets]->GetFormat() : PF_Unknown;
 		DepthLoadAction = RTInfo.DepthStencilRenderTarget.DepthLoadAction;
 		DepthStoreAction = RTInfo.DepthStencilRenderTarget.DepthStoreAction;
 		StencilLoadAction = RTInfo.DepthStencilRenderTarget.StencilLoadAction;
