UE4 - Replicate change from 3202901

OS - Fix uninitialized variable causing scene color targets not available
#tests Run soaktest

diff --git a/Engine/Source/Runtime/Renderer/Private/PostProcess/SceneRenderTargets.cpp b/Engine/Source/Runtime/Renderer/Private/PostProcess/SceneRenderTargets.cpp
--- a/Engine/Source/Runtime/Renderer/Private/PostProcess/SceneRenderTargets.cpp
+++ b/Engine/Source/Runtime/Renderer/Private/PostProcess/SceneRenderTargets.cpp
@@ -251,6 +251,7 @@ FSceneRenderTargets::FSceneRenderTargets(const FViewInfo& View, const FSceneRend
 	, bCurrentLightPropagationVolume(SnapshotSource.bCurrentLightPropagationVolume)
 	, CurrentFeatureLevel(SnapshotSource.CurrentFeatureLevel)
 	, CurrentShadingPath(SnapshotSource.CurrentShadingPath)
+	, bRequireSceneColorAlpha(SnapshotSource.bRequireSceneColorAlpha)
 	, bAllocateVelocityGBuffer(SnapshotSource.bAllocateVelocityGBuffer)
 	, bGBuffersFastCleared(SnapshotSource.bGBuffersFastCleared)	
 	, bSceneDepthCleared(SnapshotSource.bSceneDepthCleared)	
