UE4 - Main non-unity compile fixes

diff --git a/Engine/Source/Runtime/Renderer/Private/CompositionLighting/PostProcessDeferredDecals.h b/Engine/Source/Runtime/Renderer/Private/CompositionLighting/PostProcessDeferredDecals.h
--- a/Engine/Source/Runtime/Renderer/Private/CompositionLighting/PostProcessDeferredDecals.h
+++ b/Engine/Source/Runtime/Renderer/Private/CompositionLighting/PostProcessDeferredDecals.h
@@ -13,6 +13,7 @@
 #include "PostProcess/SceneRenderTargets.h"
 #include "PostProcess/RenderingCompositionGraph.h"
 #include "DecalRenderingCommon.h"
+#include "DrawingPolicy.h"
 
 // ePId_Input0: SceneColor (not needed for DBuffer decals)
 // derives from TRenderingCompositePassBase<InputCount, OutputCount> 
diff --git a/Engine/Source/Runtime/Renderer/Private/CustomDepthRendering.h b/Engine/Source/Runtime/Renderer/Private/CustomDepthRendering.h
--- a/Engine/Source/Runtime/Renderer/Private/CustomDepthRendering.h
+++ b/Engine/Source/Runtime/Renderer/Private/CustomDepthRendering.h
@@ -9,6 +9,7 @@
 #include "CoreMinimal.h"
 #include "RendererInterface.h"
 #include "SceneView.h"
+#include "DrawingPolicy.h"
 
 class FPrimitiveSceneProxy;
 
