UE4 - Fix PSO validation asserting on debug on most games

diff --git a/Engine/Source/Runtime/Renderer/Private/StaticMeshDrawList.inl b/Engine/Source/Runtime/Renderer/Private/StaticMeshDrawList.inl
--- a/Engine/Source/Runtime/Renderer/Private/StaticMeshDrawList.inl
+++ b/Engine/Source/Runtime/Renderer/Private/StaticMeshDrawList.inl
@@ -132,7 +132,9 @@ int32 TStaticMeshDrawList<DrawingPolicyType>::DrawElement(
 		bDrawnShared = true;
 	}
 	
-	DrawingPolicyLink->DrawingPolicy.ApplyDitheredLODTransitionState(RHICmdList, DrawRenderState, View, *Element.Mesh, View.bAllowStencilDither);
+	FDrawingPolicyRenderState DrawRenderStateLocal(&RHICmdList, DrawRenderState);
+
+	DrawingPolicyLink->DrawingPolicy.ApplyDitheredLODTransitionState(RHICmdList, DrawRenderStateLocal, View, *Element.Mesh, View.bAllowStencilDither);
 
 	int32 DrawCount = 0;
 
@@ -162,7 +164,7 @@ int32 TStaticMeshDrawList<DrawingPolicyType>::DrawElement(
 						Proxy,
 						*Element.Mesh,
 						BatchElementIndex,
-						DrawRenderState,
+						DrawRenderStateLocal,
 						Element.PolicyData,
 						PolicyContext
 					);
@@ -183,7 +185,7 @@ int32 TStaticMeshDrawList<DrawingPolicyType>::DrawElement(
 					Proxy,
 					*Element.Mesh,
 					BatchElementIndex,
-					DrawRenderState,
+					DrawRenderStateLocal,
 					Element.PolicyData,
 					PolicyContext
 				);
