Fix for distortion rendering crash on ES2 introducted by merge Odin

diff --git a/Engine/Source/Runtime/Renderer/Private/DistortionRendering.cpp b/Engine/Source/Runtime/Renderer/Private/DistortionRendering.cpp
--- a/Engine/Source/Runtime/Renderer/Private/DistortionRendering.cpp
+++ b/Engine/Source/Runtime/Renderer/Private/DistortionRendering.cpp
@@ -1214,7 +1214,7 @@ void FSceneRenderer::RenderDistortionES2(FRHICommandListImmediate& RHICmdList)
 			static FGlobalBoundShaderState BoundShaderState;
 				
 			TShaderMapRef<FPostProcessVS> VertexShader(View.ShaderMap);
-			TShaderMapRef<TDistortionMergePS<1>> PixelShader(View.ShaderMap);
+			TShaderMapRef<TDistortionMergePS<false>> PixelShader(View.ShaderMap);
 			SetGlobalBoundShaderState(RHICmdList, View.GetFeatureLevel(), BoundShaderState, GFilterVertexDeclaration.VertexDeclarationRHI, *VertexShader, *PixelShader);
 
 			VertexShader->SetParameters(Context);
