UE - Fix spam for msaa

diff --git a/Engine/Source/Runtime/Renderer/Private/PostProcess/SceneRenderTargets.cpp b/Engine/Source/Runtime/Renderer/Private/PostProcess/SceneRenderTargets.cpp
--- a/Engine/Source/Runtime/Renderer/Private/PostProcess/SceneRenderTargets.cpp
+++ b/Engine/Source/Runtime/Renderer/Private/PostProcess/SceneRenderTargets.cpp
@@ -82,7 +82,7 @@ static TAutoConsoleVariable<int32> CVarWideCustomResolve(
 
 static TAutoConsoleVariable<int32> CVarMobileMSAA(
 	TEXT("r.MobileMSAA"),
-	0,
+	1,
 	TEXT("Use MSAA instead of Temporal AA on mobile:\n")
 	TEXT("1: Use Temporal AA (MSAA disabled)\n")
 	TEXT("2: Use 2x MSAA (Temporal AA disabled)\n")
@@ -355,7 +355,7 @@ uint16 FSceneRenderTargets::GetNumSceneColorMSAASamples(ERHIFeatureLevel::Type I
 
 		if (IsForwardShadingEnabled(InFeatureLevel) && Method == AAM_MSAA)
 		{
-			NumSamples = CVarMSAACount.GetValueOnRenderThread();
+			NumSamples = FMath::Max(1, CVarMSAACount.GetValueOnRenderThread());
 
 			if (NumSamples != 1 && NumSamples != 2 && NumSamples != 4)
 			{
