Copying //UE4/Dev-Mobile to //UE4/Dev-Main (Source: //UE4/Dev-Mobile @ 3281679)

==========================
MAJOR FEATURES + CHANGES
==========================

Change 3228537 on 2016/12/09 by [REDACTED_PERSON]

	Fix for MGD v4.3.0

Change 3228800 on 2016/12/09 by [REDACTED_PERSON]

	Update outstanding on-load reflection captures when mobile preview is active.

Change 3239410 on 2016/12/19 by [REDACTED_PERSON]

	Fixed: Exponential Height Fog renders incorrectly after Set World Origin Location

Change 3240878 on 2016/12/20 by [REDACTED_PERSON]

	Mobile Support for r.ScreenPercentage

Change 3241506 on 2016/12/20 by [REDACTED_PERSON]

	Changed constructor field initialization order to fix Mac warning

Change 3241564 on 2016/12/21 by [REDACTED_PERSON]

	Changed another constructor field initialization order to fix Mac warning

Change 3243237 on 2016/12/22 by [REDACTED_PERSON]

	Added CustomDepth/CustomStencil and Lighting Channels to Landscape Actor

Change 3244540 on 2017/01/02 by [REDACTED_PERSON]

	Hide Landscape filter for foliage fill tool, as that tool doesn't operate on Landscape.

Change 3245833 on 2017/01/03 by [REDACTED_PERSON]

	Allow vsync to be disabled on Android OpenGL ES.

	PR #3071: OpenGLES vsync fix (Contributed by [REDACTED_PERSON])

Change 3256593 on 2017/01/13 by [REDACTED_PERSON]

	Fix merge/compile error in PostProcessUpscale.cpp

Change 3269195 on 2017/01/23 by [REDACTED_PERSON]

	Fixed: Static-like noise on IOS

Change 3271148 on 2017/01/25 by [REDACTED_PERSON]

	Fixed issue where HTML5 video in Web Browser widget would only output sound on Android

Change 3273161 on 2017/01/26 by [REDACTED_PERSON]

	Fix for rare random web browser crash on iOS due to threading issue

Change 3281310 on 2017/02/01 by [REDACTED_PERSON]

	Correct line endings

diff --git a/Engine/Source/Developer/AudioFormatADPCM/Private/AudioFormatADPCM.cpp b/Engine/Source/Developer/AudioFormatADPCM/Private/AudioFormatADPCM.cpp
--- a/Engine/Source/Developer/AudioFormatADPCM/Private/AudioFormatADPCM.cpp
+++ b/Engine/Source/Developer/AudioFormatADPCM/Private/AudioFormatADPCM.cpp
@@ -248,8 +248,9 @@ namespace ADPCM
 		/* TODO::JTM - Dec 10, 2012 05:30PM - Calculate the optimal starting coefficient */
 		uint8 CoefficientIndex = 0;
 		Context.AdaptationDelta = Context.AdaptationTable[0];
-		Context.Sample1 = ReadFromArray<int16>(InputPCMSamples, ReadIndex, NumSamples, SampleStride);
+		// First PCM sample goes to Context.Sample2, decoder will reverse it
 		Context.Sample2 = ReadFromArray<int16>(InputPCMSamples, ReadIndex, NumSamples, SampleStride);
+		Context.Sample1 = ReadFromArray<int16>(InputPCMSamples, ReadIndex, NumSamples, SampleStride);
 		Context.Coefficient1 = Context.AdaptationCoefficient1[CoefficientIndex];
 		Context.Coefficient2 = Context.AdaptationCoefficient2[CoefficientIndex];
 
diff --git a/Engine/Source/Editor/FoliageEdit/Private/SFoliageEdit.cpp b/Engine/Source/Editor/FoliageEdit/Private/SFoliageEdit.cpp
--- a/Engine/Source/Editor/FoliageEdit/Private/SFoliageEdit.cpp
+++ b/Engine/Source/Editor/FoliageEdit/Private/SFoliageEdit.cpp
@@ -303,6 +303,7 @@ void SFoliageEdit::Construct(const FArguments& InArgs)
 							[
 								SNew(SBox)
 								.MinDesiredWidth(91)
+								.Visibility(this, &SFoliageEdit::GetVisibility_LandscapeFilter)
 								[
 									SNew(SCheckBox)
 									.Visibility(this, &SFoliageEdit::GetVisibility_Filters)
@@ -897,6 +898,17 @@ EVisibility SFoliageEdit::GetVisibility_Filters() const
 	return EVisibility::Visible;
 }
 
+EVisibility SFoliageEdit::GetVisibility_LandscapeFilter() const
+{
+	// Fill tool doesn't support Landscape
+	if (FoliageEditMode->UISettings.GetPaintBucketToolSelected())
+	{
+		return EVisibility::Collapsed;
+	}
+
+	return EVisibility::Visible;
+}
+
 EVisibility SFoliageEdit::GetVisibility_Actions() const
 {
 	if (FoliageEditMode->UISettings.GetSelectToolSelected() || FoliageEditMode->UISettings.GetLassoSelectToolSelected())
diff --git a/Engine/Source/Editor/FoliageEdit/Private/SFoliageEdit.h b/Engine/Source/Editor/FoliageEdit/Private/SFoliageEdit.h
--- a/Engine/Source/Editor/FoliageEdit/Private/SFoliageEdit.h
+++ b/Engine/Source/Editor/FoliageEdit/Private/SFoliageEdit.h
@@ -156,6 +156,9 @@ private:	// BRUSH SETTINGS
 	/** Checks if the filters should appear. Dependant on the current tool being used. */
 	EVisibility GetVisibility_Filters() const;
 
+	/** Checks if the Landscape filter should appear. Hidden for the Fill tool as it's not supported there. */
+	EVisibility GetVisibility_LandscapeFilter() const;
+	
 	/** Checks if the actions should appear. Dependant on the current tool being used. */
 	EVisibility GetVisibility_Actions() const;
 
diff --git a/Engine/Source/Editor/LandscapeEditor/Private/LandscapeEdMode.cpp b/Engine/Source/Editor/LandscapeEditor/Private/LandscapeEdMode.cpp
--- a/Engine/Source/Editor/LandscapeEditor/Private/LandscapeEdMode.cpp
+++ b/Engine/Source/Editor/LandscapeEditor/Private/LandscapeEdMode.cpp
@@ -3400,6 +3400,9 @@ ALandscape* FEdModeLandscape::ChangeComponentSetting(int32 NumComponentsX, int32
 			Landscape->StaticLightingResolution = OldLandscapeProxy->StaticLightingResolution;
 			Landscape->bCastStaticShadow = OldLandscapeProxy->bCastStaticShadow;
 			Landscape->bCastShadowAsTwoSided = OldLandscapeProxy->bCastShadowAsTwoSided;
+			Landscape->LightingChannels = OldLandscapeProxy->LightingChannels;
+			Landscape->bRenderCustomDepth = OldLandscapeProxy->bRenderCustomDepth;
+			Landscape->CustomDepthStencilValue = OldLandscapeProxy->CustomDepthStencilValue;
 			Landscape->LightmassSettings = OldLandscapeProxy->LightmassSettings;
 			Landscape->CollisionThickness = OldLandscapeProxy->CollisionThickness;
 			Landscape->BodyInstance.SetCollisionProfileName(OldLandscapeProxy->BodyInstance.GetCollisionProfileName());
diff --git a/Engine/Source/Editor/LandscapeEditor/Private/LandscapeEdModeComponentTools.cpp b/Engine/Source/Editor/LandscapeEditor/Private/LandscapeEdModeComponentTools.cpp
--- a/Engine/Source/Editor/LandscapeEditor/Private/LandscapeEdModeComponentTools.cpp
+++ b/Engine/Source/Editor/LandscapeEditor/Private/LandscapeEdModeComponentTools.cpp
@@ -870,8 +870,7 @@ public:
 						LandscapeComponent->AttachToComponent(Landscape->GetRootComponent(), FAttachmentTransformRules::KeepRelativeTransform);
 
 						// Assign shared properties
-						LandscapeComponent->bCastStaticShadow = Landscape->bCastStaticShadow;
-						LandscapeComponent->bCastShadowAsTwoSided = Landscape->bCastShadowAsTwoSided;
+						LandscapeComponent->UpdatedSharedPropertiesFromActor();
 
 						int32 ComponentVerts = (Landscape->SubsectionSizeQuads + 1) * Landscape->NumSubsections;
 						// Update Weightmap Scale Bias
diff --git a/Engine/Source/Editor/UnrealEd/Private/EditorEngine.cpp b/Engine/Source/Editor/UnrealEd/Private/EditorEngine.cpp
--- a/Engine/Source/Editor/UnrealEd/Private/EditorEngine.cpp
+++ b/Engine/Source/Editor/UnrealEd/Private/EditorEngine.cpp
@@ -1531,6 +1531,11 @@ void UEditorEngine::Tick( float DeltaSeconds, bool bIdleMode )
 	// Skip updating reflection captures on the first update as the level will not be ready to display
 	if (!bFirstTick)
 	{
+		if (UReflectionCaptureComponent::MobileReflectionCapturesNeedForcedUpdate(EditorContext.World()))
+		{
+			UpdateReflectionCaptures();
+		}
+
 		// Update sky light first because sky diffuse will be visible in reflection capture indirect specular
 		USkyLightComponent::UpdateSkyCaptureContents(EditorContext.World());
 		UReflectionCaptureComponent::UpdateReflectionCaptureContents(EditorContext.World());
diff --git a/Engine/Source/Programs/UnrealBuildTool/Android/UEDeployAndroid.cs b/Engine/Source/Programs/UnrealBuildTool/Android/UEDeployAndroid.cs
--- a/Engine/Source/Programs/UnrealBuildTool/Android/UEDeployAndroid.cs
+++ b/Engine/Source/Programs/UnrealBuildTool/Android/UEDeployAndroid.cs
@@ -1487,6 +1487,7 @@ namespace UnrealBuildTool
 					Text.AppendLine("\t             " + Line);
 				}
 			}
+            Text.AppendLine("\t             android:hardwareAccelerated=\"true\">");
 			Text.AppendLine("\t             android:hasCode=\"true\">");
 			if (bShowLaunchImage)
 			{
diff --git a/Engine/Source/Runtime/Core/Private/HAL/ConsoleManager.cpp b/Engine/Source/Runtime/Core/Private/HAL/ConsoleManager.cpp
--- a/Engine/Source/Runtime/Core/Private/HAL/ConsoleManager.cpp
+++ b/Engine/Source/Runtime/Core/Private/HAL/ConsoleManager.cpp
@@ -2148,6 +2148,14 @@ static TAutoConsoleVariable<float> CVarMobileContentScaleFactor(
 	TEXT("Content scale multiplier (equates to iOS's contentScaleFactor to support Retina displays"),
 	ECVF_Default);
 
+static TAutoConsoleVariable<int32> CVarMobileTonemapperUpscale(
+	TEXT("r.MobileTonemapperUpscale"),
+	0,
+	TEXT("On mobile, whether to allow upscaling as part of the tonemapper or as a separate pass when possible")
+	TEXT("0: separate pass (default)\n")
+	TEXT("1: as part of the tonemapper pass\n"),
+	ECVF_Default);
+
 // this cvar can be removed in shipping to not compile shaders for development (faster)
 static TAutoConsoleVariable<int32> CVarCompileShadersForDevelopment(
 	TEXT("r.CompileShadersForDevelopment"),
diff --git a/Engine/Source/Runtime/Engine/Classes/Components/ReflectionCaptureComponent.h b/Engine/Source/Runtime/Engine/Classes/Components/ReflectionCaptureComponent.h
--- a/Engine/Source/Runtime/Engine/Classes/Components/ReflectionCaptureComponent.h
+++ b/Engine/Source/Runtime/Engine/Classes/Components/ReflectionCaptureComponent.h
@@ -245,6 +245,9 @@ class UReflectionCaptureComponent : public USceneComponent
 	virtual void PostEditChangeProperty(FPropertyChangedEvent& PropertyChangedEvent) override;
 	virtual void PostEditImport() override;
 	virtual void PreFeatureLevelChange(ERHIFeatureLevel::Type PendingFeatureLevel) override;
+
+	// Test if editor needs to update mobile's reflection captures. 
+	ENGINE_API static bool MobileReflectionCapturesNeedForcedUpdate(UWorld* WorldToUpdate);
 #endif // WITH_EDITOR
 	virtual void BeginDestroy() override;
 	virtual bool IsReadyForFinishDestroy() override;
diff --git a/Engine/Source/Runtime/Engine/Private/Components/ReflectionCaptureComponent.cpp b/Engine/Source/Runtime/Engine/Private/Components/ReflectionCaptureComponent.cpp
--- a/Engine/Source/Runtime/Engine/Private/Components/ReflectionCaptureComponent.cpp
+++ b/Engine/Source/Runtime/Engine/Private/Components/ReflectionCaptureComponent.cpp
@@ -1516,6 +1516,28 @@ void UReflectionCaptureComponent::ReadbackFromGPU(UWorld* WorldToUpdate)
 	}
 }
 
+#if WITH_EDITOR
+// If the feature level preview has been set before on-load captures have been built then the editor must update them.
+bool UReflectionCaptureComponent::MobileReflectionCapturesNeedForcedUpdate(UWorld* WorldToUpdate)
+{
+	if (WorldToUpdate->Scene 
+		&& WorldToUpdate->Scene->GetFeatureLevel() <= ERHIFeatureLevel::ES3_1
+		&& (GShaderCompilingManager == NULL || !GShaderCompilingManager->IsCompiling()))
+	{
+		FScopeLock Lock(&ReflectionCapturesToUpdateForLoadLock);
+		for (int32 CaptureIndex = ReflectionCapturesToUpdateForLoad.Num() - 1; CaptureIndex >= 0; CaptureIndex--)
+		{
+			UReflectionCaptureComponent* CaptureComponent = ReflectionCapturesToUpdateForLoad[CaptureIndex];
+			if (!CaptureComponent->GetOwner() || WorldToUpdate->ContainsActor(CaptureComponent->GetOwner()))
+			{
+				return true;
+			}
+		}
+	}
+	return false;
+}
+#endif
+
 void UReflectionCaptureComponent::UpdateReflectionCaptureContents(UWorld* WorldToUpdate)
 {
 	if (WorldToUpdate->Scene 
@@ -1579,14 +1601,14 @@ void UReflectionCaptureComponent::UpdateReflectionCaptureContents(UWorld* WorldT
 		WorldToUpdate->Scene->AllocateReflectionCaptures(WorldCombinedCaptures);
 
 		if (FeatureLevel >= ERHIFeatureLevel::SM4 && !FPlatformProperties::RequiresCookedData())
-		{
-			for (int32 CaptureIndex = 0; CaptureIndex < WorldCapturesToUpdateForLoad.Num(); CaptureIndex++)
 			{
-				// Save the derived data for any captures that were dirty on load
-				// This allows the derived data to get cached without having to resave a map
-				WorldCapturesToUpdateForLoad[CaptureIndex]->ReadbackFromGPU(WorldToUpdate);
+				for (int32 CaptureIndex = 0; CaptureIndex < WorldCapturesToUpdateForLoad.Num(); CaptureIndex++)
+				{
+					// Save the derived data for any captures that were dirty on load
+					// This allows the derived data to get cached without having to resave a map
+					WorldCapturesToUpdateForLoad[CaptureIndex]->ReadbackFromGPU(WorldToUpdate);
+				}
 			}
-		}
 	}
 }
 
diff --git a/Engine/Source/Runtime/Engine/Private/SceneUtils.cpp b/Engine/Source/Runtime/Engine/Private/SceneUtils.cpp
--- a/Engine/Source/Runtime/Engine/Private/SceneUtils.cpp
+++ b/Engine/Source/Runtime/Engine/Private/SceneUtils.cpp
@@ -96,7 +96,7 @@ bool IsMobileHDR()
 bool IsMobileHDR32bpp()
 {
 	static auto* MobileHDR32bppModeCvar = IConsoleManager::Get().FindTConsoleVariableDataInt(TEXT("r.MobileHDR32bppMode"));
-	return IsMobileHDR() && (GSupportsRenderTargetFormat_PF_FloatRGBA == false || MobileHDR32bppModeCvar->GetValueOnRenderThread() != 0);
+	return IsMobileHDR() && (GSupportsRenderTargetFormat_PF_FloatRGBA == false || MobileHDR32bppModeCvar->GetValueOnAnyThread() != 0);
 }
 
 bool IsMobileHDRMosaic()
@@ -105,7 +105,7 @@ bool IsMobileHDRMosaic()
 		return false;
 
 	static auto* MobileHDR32bppMode = IConsoleManager::Get().FindTConsoleVariableDataInt(TEXT("r.MobileHDR32bppMode"));
-	switch (MobileHDR32bppMode->GetValueOnRenderThread())
+	switch (MobileHDR32bppMode->GetValueOnAnyThread())
 	{
 		case 1:
 			return true;
diff --git a/Engine/Source/Runtime/Engine/Private/SceneView.cpp b/Engine/Source/Runtime/Engine/Private/SceneView.cpp
--- a/Engine/Source/Runtime/Engine/Private/SceneView.cpp
+++ b/Engine/Source/Runtime/Engine/Private/SceneView.cpp
@@ -1799,12 +1799,11 @@ void FSceneView::EndFinalPostprocessSettings(const FSceneViewInitOptions& ViewIn
 		}
 	}
 
-	// Not supported in ES2/3.
-	bool bES2Or3 = (SceneViewFeatureLevel == ERHIFeatureLevel::ES2) || (SceneViewFeatureLevel == ERHIFeatureLevel::ES3_1);
-
-	if(!Family->EngineShowFlags.ScreenPercentage || bIsSceneCapture || bIsReflectionCapture || bES2Or3)
+	// ScreenPercentage is not supported in ES2/3.1 with MobileHDR = false
+	const bool bIsMobileLDR = (SceneViewFeatureLevel <= ERHIFeatureLevel::ES3_1 && !IsMobileHDR());
+	if(!Family->EngineShowFlags.ScreenPercentage || bIsSceneCapture || bIsReflectionCapture || bIsMobileLDR)
 	{
-		FinalPostProcessSettings.ScreenPercentage = 100;
+		FinalPostProcessSettings.ScreenPercentage = 100.0f;
 	}
 
 	if(!Family->EngineShowFlags.AmbientOcclusion || !Family->EngineShowFlags.ScreenSpaceAO)
@@ -1966,6 +1965,14 @@ void FSceneView::EndFinalPostprocessSettings(const FSceneViewInitOptions& ViewIn
 			// compute the view rectangle with the ScreenPercentage applied
 			FIntRect ScreenPercentageAffectedViewRect = ViewInitOptions.GetConstrainedViewRect().Scale(Fraction);
 			QuantizeSceneBufferSize(ScreenPercentageAffectedViewRect.Max.X, ScreenPercentageAffectedViewRect.Max.Y);
+			// Mosaic needs to the viewport height to be a multiple of 2.
+			if (SceneViewFeatureLevel <= ERHIFeatureLevel::ES3_1 && IsMobileHDRMosaic())
+			{
+				if (((ScreenPercentageAffectedViewRect.Size().Y) & 1) == 1)
+				{
+					ScreenPercentageAffectedViewRect.Max.Y -= 1;
+				}
+			}
 			SetScaledViewRect(ScreenPercentageAffectedViewRect);
 		}
 	}
diff --git a/Engine/Source/Runtime/Landscape/Classes/LandscapeComponent.h b/Engine/Source/Runtime/Landscape/Classes/LandscapeComponent.h
--- a/Engine/Source/Runtime/Landscape/Classes/LandscapeComponent.h
+++ b/Engine/Source/Runtime/Landscape/Classes/LandscapeComponent.h
@@ -622,6 +622,9 @@ public:
 
 	/** Updates navigation properties to match landscape's master switch */
 	void UpdateNavigationRelevance();
+	
+	/** Updates the values of component-level properties exposed by the Landscape Actor */
+	LANDSCAPE_API void UpdatedSharedPropertiesFromActor();
 #endif
 
 	friend class FLandscapeComponentSceneProxy;
diff --git a/Engine/Source/Runtime/Landscape/Classes/LandscapeProxy.h b/Engine/Source/Runtime/Landscape/Classes/LandscapeProxy.h
--- a/Engine/Source/Runtime/Landscape/Classes/LandscapeProxy.h
+++ b/Engine/Source/Runtime/Landscape/Classes/LandscapeProxy.h
@@ -449,12 +449,28 @@ public:
 	UPROPERTY(EditAnywhere, AdvancedDisplay, Category=Lighting, meta=(DisplayName = "Far Shadow"))
 	uint32 bCastFarShadow:1;
 	
+	/**
+	* Channels that this Landscape should be in.  Lights with matching channels will affect the Landscape.
+	* These channels only apply to opaque materials, direct lighting, and dynamic lighting and shadowing.
+	*/
+	UPROPERTY(EditAnywhere, AdvancedDisplay, BlueprintReadOnly, Category = Lighting)
+	FLightingChannels LightingChannels;
+
 	/** Whether to use the landscape material's vertical world position offset when calculating static lighting.
 		Note: Only z (vertical) offset is supported. XY offsets are ignored.
 		Does not work correctly with an XY offset map (mesh collision) */
 	UPROPERTY(EditAnywhere, AdvancedDisplay, Category=Lighting)
 	uint32 bUseMaterialPositionOffsetInStaticLighting:1;
 
+	/** If true, the Landscape will be rendered in the CustomDepth pass (usually used for outlines) */
+	UPROPERTY(EditAnywhere, AdvancedDisplay, BlueprintReadOnly, Category=Rendering, meta=(DisplayName = "Render CustomDepth Pass"))
+	uint32 bRenderCustomDepth:1;
+
+	/** Optionally write this 0-255 value to the stencil buffer in CustomDepth pass (Requires project setting or r.CustomDepth == 3) */
+	UPROPERTY(EditAnywhere, AdvancedDisplay, BlueprintReadOnly, Category=Rendering,  meta=(UIMin = "0", UIMax = "255", editcondition = "bRenderCustomDepth", DisplayName = "CustomDepth Stencil Value"))
+	int32 CustomDepthStencilValue;
+
+
 #if WITH_EDITORONLY_DATA
 	UPROPERTY(transient)
 	uint32 bIsMovingToLevel:1;    // Check for the Move to Current Level case
diff --git a/Engine/Source/Runtime/Landscape/Private/Landscape.cpp b/Engine/Source/Runtime/Landscape/Private/Landscape.cpp
--- a/Engine/Source/Runtime/Landscape/Private/Landscape.cpp
+++ b/Engine/Source/Runtime/Landscape/Private/Landscape.cpp
@@ -515,6 +515,18 @@ void ULandscapeInfo::UpdateDebugColorMaterial()
 	//GWarn->EndSlowTask();
 }
 
+void ULandscapeComponent::UpdatedSharedPropertiesFromActor()
+{
+	ALandscapeProxy* LandscapeProxy = GetLandscapeProxy();
+
+	bCastStaticShadow = LandscapeProxy->bCastStaticShadow;
+	bCastShadowAsTwoSided = LandscapeProxy->bCastShadowAsTwoSided;
+	bCastFarShadow = LandscapeProxy->bCastFarShadow;
+	bRenderCustomDepth = LandscapeProxy->bRenderCustomDepth;
+	CustomDepthStencilValue = LandscapeProxy->CustomDepthStencilValue;
+	LightingChannels = LandscapeProxy->LightingChannels;
+}
+
 void ULandscapeComponent::PostLoad()
 {
 	Super::PostLoad();
@@ -523,8 +535,7 @@ void ULandscapeComponent::PostLoad()
 	if (ensure(LandscapeProxy))
 	{
 		// Ensure that the component's lighting settings matches the actor's.
-		bCastStaticShadow = LandscapeProxy->bCastStaticShadow;
-		bCastShadowAsTwoSided = LandscapeProxy->bCastShadowAsTwoSided;
+		UpdatedSharedPropertiesFromActor();	
 
 		// check SectionBaseX/Y are correct
 		int32 CheckSectionBaseX = FMath::RoundToInt(RelativeLocation.X) + LandscapeProxy->LandscapeSectionOffset.X;
@@ -1586,6 +1597,9 @@ void ALandscapeProxy::GetSharedProperties(ALandscapeProxy* Landscape)
 		StaticLightingResolution = Landscape->StaticLightingResolution;
 		bCastStaticShadow = Landscape->bCastStaticShadow;
 		bCastShadowAsTwoSided = Landscape->bCastShadowAsTwoSided;
+		LightingChannels = Landscape->LightingChannels;
+		bRenderCustomDepth = Landscape->bRenderCustomDepth;
+		CustomDepthStencilValue = Landscape->CustomDepthStencilValue;
 		ComponentSizeQuads = Landscape->ComponentSizeQuads;
 		NumSubsections = Landscape->NumSubsections;
 		SubsectionSizeQuads = Landscape->SubsectionSizeQuads;
diff --git a/Engine/Source/Runtime/Landscape/Private/LandscapeEdit.cpp b/Engine/Source/Runtime/Landscape/Private/LandscapeEdit.cpp
--- a/Engine/Source/Runtime/Landscape/Private/LandscapeEdit.cpp
+++ b/Engine/Source/Runtime/Landscape/Private/LandscapeEdit.cpp
@@ -2070,8 +2070,7 @@ LANDSCAPE_API void ALandscapeProxy::Import(
 				);
 
 			// Assign shared properties
-			LandscapeComponent->bCastStaticShadow = bCastStaticShadow;
-			LandscapeComponent->bCastShadowAsTwoSided = bCastShadowAsTwoSided;
+			LandscapeComponent->UpdatedSharedPropertiesFromActor();
 		}
 	}
 
@@ -3664,7 +3663,11 @@ void ALandscapeProxy::PostEditChangeProperty(FPropertyChangedEvent& PropertyChan
 	}
 	else if(PropertyName == FName(TEXT("bCastStaticShadow")) ||
 		PropertyName == FName(TEXT("bCastShadowAsTwoSided")) ||
-		PropertyName == FName(TEXT("bCastFarShadow")))
+		PropertyName == FName(TEXT("bCastFarShadow")) ||
+		PropertyName == FName(TEXT("bRenderCustomDepth")) ||
+		PropertyName == FName(TEXT("CustomDepthStencilValue")) ||
+		PropertyName == FName(TEXT("LightingChannels"))
+		)
 	{
 		// Replicate shared properties to all components.
 		for (int32 ComponentIndex = 0; ComponentIndex < LandscapeComponents.Num(); ComponentIndex++)
@@ -3672,9 +3675,7 @@ void ALandscapeProxy::PostEditChangeProperty(FPropertyChangedEvent& PropertyChan
 			ULandscapeComponent* Comp = LandscapeComponents[ComponentIndex];
 			if (Comp)
 			{
-				Comp->bCastStaticShadow = bCastStaticShadow;
-				Comp->bCastShadowAsTwoSided = bCastShadowAsTwoSided;
-				Comp->bCastFarShadow = bCastFarShadow;
+				Comp->UpdatedSharedPropertiesFromActor();
 			}
 		}
 	}
diff --git a/Engine/Source/Runtime/Landscape/Private/LandscapeRender.cpp b/Engine/Source/Runtime/Landscape/Private/LandscapeRender.cpp
--- a/Engine/Source/Runtime/Landscape/Private/LandscapeRender.cpp
+++ b/Engine/Source/Runtime/Landscape/Private/LandscapeRender.cpp
@@ -923,6 +923,8 @@ FPrimitiveViewRelevance FLandscapeComponentSceneProxy::GetViewRelevance(const FS
 	FPrimitiveViewRelevance Result;
 	const bool bCollisionView = (View->Family->EngineShowFlags.CollisionVisibility || View->Family->EngineShowFlags.CollisionPawn);
 	Result.bDrawRelevance = (IsShown(View) || bCollisionView) && View->Family->EngineShowFlags.Landscape;
+	Result.bRenderCustomDepth = ShouldRenderCustomDepth();
+	Result.bUsesLightingChannels = GetLightingChannelMask() != GetDefaultLightingChannelMask();
 
 	auto FeatureLevel = View->GetFeatureLevel();
 
diff --git a/Engine/Source/Runtime/OpenGLDrv/Private/Android/AndroidEGL.cpp b/Engine/Source/Runtime/OpenGLDrv/Private/Android/AndroidEGL.cpp
--- a/Engine/Source/Runtime/OpenGLDrv/Private/Android/AndroidEGL.cpp
+++ b/Engine/Source/Runtime/OpenGLDrv/Private/Android/AndroidEGL.cpp
@@ -647,8 +647,10 @@ int32 AndroidEGL::GetError()
 	return eglGetError();
 }
 
-bool AndroidEGL::SwapBuffers()
+bool AndroidEGL::SwapBuffers(int32 SyncInterval)
 {
+	eglSwapInterval(PImplData->eglDisplay, SyncInterval);
+
 	if ( PImplData->eglSurface == NULL || !eglSwapBuffers(PImplData->eglDisplay, PImplData->eglSurface))
 	{
 		// shutdown if swapbuffering goes down
diff --git a/Engine/Source/Runtime/OpenGLDrv/Private/Android/AndroidEGL.h b/Engine/Source/Runtime/OpenGLDrv/Private/Android/AndroidEGL.h
--- a/Engine/Source/Runtime/OpenGLDrv/Private/Android/AndroidEGL.h
+++ b/Engine/Source/Runtime/OpenGLDrv/Private/Android/AndroidEGL.h
@@ -59,7 +59,7 @@ public:
 	void Init( APIVariant API, uint32 MajorVersion, uint32 MinorVersion, bool bDebug);
 	void ReInit();
 	void UnBind();
-	bool SwapBuffers();
+	bool SwapBuffers(int32 SyncInterval);
 	void Terminate();
 	void InitSurface(bool bUseSmallSurface, bool bCreateWndSurface);
 
diff --git a/Engine/Source/Runtime/OpenGLDrv/Private/Android/AndroidESDeferredOpenGL.cpp b/Engine/Source/Runtime/OpenGLDrv/Private/Android/AndroidESDeferredOpenGL.cpp
--- a/Engine/Source/Runtime/OpenGLDrv/Private/Android/AndroidESDeferredOpenGL.cpp
+++ b/Engine/Source/Runtime/OpenGLDrv/Private/Android/AndroidESDeferredOpenGL.cpp
@@ -308,7 +308,7 @@ bool PlatformBlitToViewport( FPlatformOpenGLDevice* Device, const FOpenGLViewpor
 {
 	if (FOpenGL::IsES2())
 	{
-		AndroidEGL::GetInstance()->SwapBuffers();
+		AndroidEGL::GetInstance()->SwapBuffers(bLockToVsync ? SyncInterval : 0);
 	}
 	else
 	{
@@ -341,7 +341,7 @@ bool PlatformBlitToViewport( FPlatformOpenGLDevice* Device, const FOpenGLViewpor
 		{
 			uint32 IdleStart = FPlatformTime::Cycles();
 
-			AndroidEGL::GetInstance()->SwapBuffers();
+			AndroidEGL::GetInstance()->SwapBuffers(bLockToVsync ? SyncInterval : 0);
 			REPORT_GL_END_BUFFER_EVENT_FOR_FRAME_DUMP();
 //			INITIATE_GL_FRAME_DUMP_EVERY_X_CALLS( 1000 );
 
diff --git a/Engine/Source/Runtime/OpenGLDrv/Private/Android/AndroidOpenGL.cpp b/Engine/Source/Runtime/OpenGLDrv/Private/Android/AndroidOpenGL.cpp
--- a/Engine/Source/Runtime/OpenGLDrv/Private/Android/AndroidOpenGL.cpp
+++ b/Engine/Source/Runtime/OpenGLDrv/Private/Android/AndroidOpenGL.cpp
@@ -167,7 +167,7 @@ bool PlatformBlitToViewport( FPlatformOpenGLDevice* Device, const FOpenGLViewpor
 	}
 	if (bPresent)
 	{
-		AndroidEGL::GetInstance()->SwapBuffers();
+		AndroidEGL::GetInstance()->SwapBuffers(bLockToVsync ? SyncInterval : 0);
 	}
 	return bPresent;
 }
diff --git a/Engine/Source/Runtime/Renderer/Private/PostProcess/PostProcessTonemap.cpp b/Engine/Source/Runtime/Renderer/Private/PostProcess/PostProcessTonemap.cpp
--- a/Engine/Source/Runtime/Renderer/Private/PostProcess/PostProcessTonemap.cpp
+++ b/Engine/Source/Runtime/Renderer/Private/PostProcess/PostProcessTonemap.cpp
@@ -1363,6 +1363,11 @@ FPooledRenderTargetDesc FRCPassPostProcessTonemap::ComputeOutputDesc(EPassOutput
 	Ret.DebugName = TEXT("Tonemap");
 	Ret.ClearValue = FClearValueBinding(FLinearColor(0, 0, 0, 0));
 
+	// Mobile needs to override the extent
+	if (bDoScreenPercentageInTonemapper && View.GetFeatureLevel() <= ERHIFeatureLevel::ES3_1)
+	{
+		Ret.Extent = View.UnscaledViewRect.Max;
+	}
 	return Ret;
 }
 
@@ -1659,11 +1664,9 @@ namespace PostProcessTonemap_ES2Util
 	}
 }
 
-
-FRCPassPostProcessTonemapES2::FRCPassPostProcessTonemapES2(const FViewInfo& View, FIntRect InViewRect, FIntPoint InDestSize, bool bInUsedFramebufferFetch) 
-	: bEnableExtentOverride(true)
-	, ViewRect(InViewRect)
-	, DestSize(InDestSize)
+FRCPassPostProcessTonemapES2::FRCPassPostProcessTonemapES2(const FViewInfo& InView, bool bInUsedFramebufferFetch)
+	: bDoScreenPercentageInTonemapper(false)
+	, View(InView)
 	, bUsedFramebufferFetch(bInUsedFramebufferFetch)
 {
 	uint32 ConfigBitmask = TonemapperGenerateBitmaskMobile(&View, false);
@@ -1682,14 +1685,13 @@ void FRCPassPostProcessTonemapES2::Process(FRenderingCompositePassContext& Conte
 		return;
 	}
 	
-	const FSceneView& View = Context.View;
 	const FSceneViewFamily& ViewFamily = *(View.Family);
 	const FSceneRenderTargetItem& DestRenderTarget = PassOutputs[0].RequestSurface(Context);
 	const FPooledRenderTargetDesc& OutputDesc = PassOutputs[0].RenderTargetDesc;
 
-	FIntRect SrcRect = ViewRect;
 	// no upscale if separate ren target is used.
-	FIntRect DestRect = (ViewFamily.bUseSeparateRenderTarget) ? ViewRect : View.UnscaledViewRect; // Simple upscaling, ES2 post process does not currently have a specific upscaling pass.
+	FIntRect SrcRect = View.ViewRect;
+	FIntRect DestRect = bDoScreenPercentageInTonemapper ? View.UnscaledViewRect : View.ViewRect;
 	FIntPoint SrcSize = InputDesc->Extent;
 	FIntPoint DstSize = OutputDesc.Extent;
 
@@ -1798,9 +1800,9 @@ FPooledRenderTargetDesc FRCPassPostProcessTonemapES2::ComputeOutputDesc(EPassOut
 	Ret.Format = PF_B8G8R8A8;
 	Ret.DebugName = TEXT("Tonemap");
 	Ret.ClearValue = FClearValueBinding(FLinearColor(0, 0, 0, 0));
-	if (bEnableExtentOverride)
+	if (bDoScreenPercentageInTonemapper)
 	{
-		Ret.Extent = DestSize;
+		Ret.Extent = View.UnscaledViewRect.Max;
 	}
 	return Ret;
 }
diff --git a/Engine/Source/Runtime/Renderer/Private/PostProcess/PostProcessTonemap.h b/Engine/Source/Runtime/Renderer/Private/PostProcess/PostProcessTonemap.h
--- a/Engine/Source/Runtime/Renderer/Private/PostProcess/PostProcessTonemap.h
+++ b/Engine/Source/Runtime/Renderer/Private/PostProcess/PostProcessTonemap.h
@@ -78,21 +78,18 @@ private:
 class FRCPassPostProcessTonemapES2 : public TRenderingCompositePassBase<3, 1>
 {
 public:
-	FRCPassPostProcessTonemapES2(const FViewInfo& View, FIntRect InViewRect, FIntPoint InDestSize, bool bInUsedFramebufferFetch);
+	FRCPassPostProcessTonemapES2(const FViewInfo& View, bool bInUsedFramebufferFetch);
 
 	// interface FRenderingCompositePass ---------
 
 	virtual void Process(FRenderingCompositePassContext& Context) override;
 	virtual void Release() override { delete this; }
 	virtual FPooledRenderTargetDesc ComputeOutputDesc(EPassOutputId InPassOutputId) const override;
-
-	// temporary for 4.15. 4.16 includes a better fix supporting ScreenPercentage
-	bool bEnableExtentOverride;
+	
+	bool bDoScreenPercentageInTonemapper;
 private:
+	const FViewInfo& View;
 
-
-	FIntRect ViewRect;
-	FIntPoint DestSize;
 	bool bUsedFramebufferFetch;
 	// set in constructor
 	uint32 ConfigIndexMobile;
diff --git a/Engine/Source/Runtime/Renderer/Private/PostProcess/PostProcessUpscale.cpp b/Engine/Source/Runtime/Renderer/Private/PostProcess/PostProcessUpscale.cpp
--- a/Engine/Source/Runtime/Renderer/Private/PostProcess/PostProcessUpscale.cpp
+++ b/Engine/Source/Runtime/Renderer/Private/PostProcess/PostProcessUpscale.cpp
@@ -297,7 +297,7 @@ void FRCPassPostProcessUpscale::Process(FRenderingCompositePassContext& Context)
 	
 	FIntRect SrcRect = View.ViewRect;
 	// no upscale if separate ren target is used.
-	FIntRect DestRect = (ViewFamily.bUseSeparateRenderTarget) ? View.ViewRect : View.UnscaledViewRect; // Simple upscaling, ES2 post process does not currently have a specific upscaling pass.
+	FIntRect DestRect = (ViewFamily.bUseSeparateRenderTarget) ? View.ViewRect : View.UnscaledViewRect;
 	FIntPoint SrcSize = InputDesc->Extent;
 
 	const FSceneRenderTargetItem& DestRenderTarget = PassOutputs[0].RequestSurface(Context);
@@ -383,3 +383,16 @@ FPooledRenderTargetDesc FRCPassPostProcessUpscale::ComputeOutputDesc(EPassOutput
 
 	return Ret;
 }
+
+FRCPassPostProcessUpscaleES2::FRCPassPostProcessUpscaleES2(const FViewInfo& InView)
+:	FRCPassPostProcessUpscale(InView, 1 /* bilinear */)
+,	View(InView)
+{
+}
+
+FPooledRenderTargetDesc FRCPassPostProcessUpscaleES2::ComputeOutputDesc(EPassOutputId InPassOutputId) const
+{
+	FPooledRenderTargetDesc Ret = FRCPassPostProcessUpscale::ComputeOutputDesc(InPassOutputId);
+	Ret.Extent = View.UnscaledViewRect.Max;
+	return Ret;
+}
\ No newline at end of file
diff --git a/Engine/Source/Runtime/Renderer/Private/PostProcess/PostProcessUpscale.h b/Engine/Source/Runtime/Renderer/Private/PostProcess/PostProcessUpscale.h
--- a/Engine/Source/Runtime/Renderer/Private/PostProcess/PostProcessUpscale.h
+++ b/Engine/Source/Runtime/Renderer/Private/PostProcess/PostProcessUpscale.h
@@ -77,3 +77,14 @@ private:
 	FIntPoint OutputExtent;
 };
 
+// Simple version used for ES2 forcing Bilinear and overriding the output extent
+class FRCPassPostProcessUpscaleES2 : public FRCPassPostProcessUpscale
+{
+public:
+	FRCPassPostProcessUpscaleES2(const FViewInfo& InView);
+
+	FPooledRenderTargetDesc ComputeOutputDesc(EPassOutputId InPassOutputId) const override;
+private:
+	const FViewInfo& View;
+};
+
diff --git a/Engine/Source/Runtime/Renderer/Private/PostProcess/PostProcessing.cpp b/Engine/Source/Runtime/Renderer/Private/PostProcess/PostProcessing.cpp
--- a/Engine/Source/Runtime/Renderer/Private/PostProcess/PostProcessing.cpp
+++ b/Engine/Source/Runtime/Renderer/Private/PostProcess/PostProcessing.cpp
@@ -2321,25 +2321,30 @@ void FPostProcessing::ProcessES2(FRHICommandListImmediate& RHICmdList, const FVi
 		static const auto VarTonemapperFilm = IConsoleManager::Get().FindTConsoleVariableDataInt(TEXT("r.TonemapperFilm"));
 		const bool bUseTonemapperFilm = IsMobileHDR() && GSupportsRenderTargetFormat_PF_FloatRGBA && (VarTonemapperFilm && VarTonemapperFilm->GetValueOnRenderThread());
 
-		// Temporary for 4.15. Remember the non-filmic tonemapper so we can disable the extent override if necessary.
-		FRCPassPostProcessTonemapES2* PostProcessTonemapES2 = nullptr;
+		static const auto VarTonemapperUpscale = IConsoleManager::Get().FindTConsoleVariableDataInt(TEXT("r.MobileTonemapperUpscale"));
+		bool bDisableUpscaleInTonemapper = Context.View.Family->bUseSeparateRenderTarget || IsMobileHDRMosaic() || !VarTonemapperUpscale || VarTonemapperUpscale->GetValueOnRenderThread() == 0;
 
+		bool* DoScreenPercentageInTonemapperPtr = nullptr;
 		if (bUseTonemapperFilm)
 		{
 			//@todo Ronin Set to EAutoExposureMethod::AEM_Basic for PC vk crash.
-			AddTonemapper(Context, BloomOutput, nullptr, EAutoExposureMethod::AEM_Histogram, false, false);
+			FRCPassPostProcessTonemap* PostProcessTonemap = AddTonemapper(Context, BloomOutput, nullptr, EAutoExposureMethod::AEM_Histogram, false, false);
+			DoScreenPercentageInTonemapperPtr = &PostProcessTonemap->bDoScreenPercentageInTonemapper;
 		}
 		else
 		{
 			// Must run to blit to back buffer even if post processing is off.
-			FRenderingCompositePass* PostProcessTonemap = Context.Graph.RegisterPass(new(FMemStack::Get()) FRCPassPostProcessTonemapES2(Context.View, FinalOutputViewRect, FinalTargetSize, bViewRectSource));
+			FRCPassPostProcessTonemapES2* PostProcessTonemap = (FRCPassPostProcessTonemapES2*)Context.Graph.RegisterPass(new(FMemStack::Get()) FRCPassPostProcessTonemapES2(Context.View,  bViewRectSource));
 			PostProcessTonemap->SetInput(ePId_Input0, Context.FinalOutput);
 			PostProcessTonemap->SetInput(ePId_Input1, BloomOutput);
 			PostProcessTonemap->SetInput(ePId_Input2, DofOutput);
 			Context.FinalOutput = FRenderingCompositeOutputRef(PostProcessTonemap);
-			PostProcessTonemapES2 = (FRCPassPostProcessTonemapES2*)PostProcessTonemap;
+			DoScreenPercentageInTonemapperPtr = &PostProcessTonemap->bDoScreenPercentageInTonemapper;
 		}
-			
+
+		// remember the tonemapper pass so we can check if it's last
+		FRenderingCompositePass* TonemapperPass = Context.FinalOutput.GetPass();
+		
 		// if Context.FinalOutput was the clipped result of sunmask stage then this stage also restores Context.FinalOutput back original target size.
 		FinalOutputViewRect = View.UnscaledViewRect;
 
@@ -2349,11 +2354,10 @@ void FPostProcessing::ProcessES2(FRHICommandListImmediate& RHICmdList, const FVi
 			{
 				AddPostProcessMaterial(Context, BL_AfterTonemapping, nullptr);
 
-				// Temporary for 4.15 to disable setting the extent in ComputeOutputDesc when there is a custom PP material.
-				// 4.16 has full support for ScreenPercentage
-				if (PostProcessTonemapES2 && Context.FinalOutput.GetPass() != PostProcessTonemapES2)
+				// Tonemapper is not the final pass so if we may need to use a separate upscale pass
+				if (Context.FinalOutput.GetPass() != TonemapperPass)
 				{
-					PostProcessTonemapES2->bEnableExtentOverride = false;
+					bDisableUpscaleInTonemapper = true;
 				}
 			}
 	
@@ -2378,6 +2382,28 @@ void FPostProcessing::ProcessES2(FRHICommandListImmediate& RHICmdList, const FVi
 			}
 		}
 
+		// Apply ScreenPercentage
+		if (View.UnscaledViewRect != View.ViewRect)
+		{
+			if (bDisableUpscaleInTonemapper)
+			{
+				FRenderingCompositePass* Node = Context.Graph.RegisterPass(new(FMemStack::Get()) FRCPassPostProcessUpscaleES2(View));
+				Node->SetInput(ePId_Input0, FRenderingCompositeOutputRef(Context.FinalOutput)); // Bilinear sampling.
+				Node->SetInput(ePId_Input1, FRenderingCompositeOutputRef(Context.FinalOutput)); // Point sampling.
+				Context.FinalOutput = FRenderingCompositeOutputRef(Node);
+				*DoScreenPercentageInTonemapperPtr = false;
+			}
+			else
+			{
+				check(DoScreenPercentageInTonemapperPtr != nullptr);
+				*DoScreenPercentageInTonemapperPtr = true;
+			}
+		}
+		else
+		{
+			*DoScreenPercentageInTonemapperPtr = false;
+		}
+
 #if WITH_EDITOR
 		// Show the selection outline if it is in the editor and we aren't in wireframe 
 		// If the engine is in demo mode and game view is on we also do not show the selection outline
diff --git a/Engine/Source/Runtime/Renderer/Private/RendererScene.cpp b/Engine/Source/Runtime/Renderer/Private/RendererScene.cpp
--- a/Engine/Source/Runtime/Renderer/Private/RendererScene.cpp
+++ b/Engine/Source/Runtime/Renderer/Private/RendererScene.cpp
@@ -2683,6 +2683,12 @@ void FScene::ApplyWorldOffset_RenderThread(FVector InOffset)
 		(*It)->SetTransform(NewTransform);
 	}
 
+	// Exponential Fog
+	for (FExponentialHeightFogSceneInfo& FogInfo : ExponentialFogs)
+	{
+		FogInfo.FogHeight+= InOffset.Z;
+	}
+	
 	// StaticMeshDrawLists
 	StaticMeshDrawListApplyWorldOffset(PositionOnlyDepthDrawList, InOffset);
 	StaticMeshDrawListApplyWorldOffset(DepthDrawList, InOffset);
diff --git a/Engine/Source/Runtime/WebBrowser/Private/IOS/IOSPlatformWebBrowser.cpp b/Engine/Source/Runtime/WebBrowser/Private/IOS/IOSPlatformWebBrowser.cpp
--- a/Engine/Source/Runtime/WebBrowser/Private/IOS/IOSPlatformWebBrowser.cpp
+++ b/Engine/Source/Runtime/WebBrowser/Private/IOS/IOSPlatformWebBrowser.cpp
@@ -95,13 +95,19 @@
 
 -(void)loadurl:(NSURL*)InURL;
 {
-	NextURL = InURL;
+	dispatch_async(dispatch_get_main_queue(), ^
+	{
+		NextURL = InURL;
+	});
 }
 
 -(void)loadstring:(NSString*)InString dummyurl:(NSURL*)InURL;
 {
-	NextContent = InString;
-	NextURL = InURL;
+	dispatch_async(dispatch_get_main_queue(), ^
+	{
+		NextContent = InString;
+		NextURL = InURL;
+	});
 }
 
 @end
diff --git a/Engine/Source/Runtime/WebBrowser/Private/IOS/IOSPlatformWebBrowser.h b/Engine/Source/Runtime/WebBrowser/Private/IOS/IOSPlatformWebBrowser.h
--- a/Engine/Source/Runtime/WebBrowser/Private/IOS/IOSPlatformWebBrowser.h
+++ b/Engine/Source/Runtime/WebBrowser/Private/IOS/IOSPlatformWebBrowser.h
@@ -29,8 +29,8 @@ class SWebBrowserView;
 #if !PLATFORM_TVOS
 @property(strong) UIWebView* WebView;
 #endif
-@property(strong) NSURL* NextURL;
-@property(strong) NSString* NextContent;
+@property(copy) NSURL* NextURL;
+@property(copy) NSString* NextContent;
 @property CGRect DesiredFrame;
 
 -(void)create:(TSharedPtr<SIOSWebBrowserWidget>)InWebBrowserWidget;
