Copying //UE4/Dev-Rendering to //UE4/Dev-Main (Source: //UE4/Dev-Rendering @ 3304653)

==========================
MAJOR FEATURES + CHANGES
==========================

Change 3288774 on 2017/02/06 by [REDACTED_PERSON]

	Added in min/max/avg values to table. Allow us to also set which stat groups we use explicitly to massively cut down on plugin overhead.
	#tests Several tests passes taken on ShooterGame.

Change 3289887 on 2017/02/07 by [REDACTED_PERSON]

	Eye Adaptation:  Fix basic mode.
	The last step of the basic eye adaptation needs to know the correct down-sampled viewport size

Change 3290281 on 2017/02/07 by [REDACTED_PERSON]

	Account for LogToLin approximation value offseting.

Change 3290282 on 2017/02/07 by [REDACTED_PERSON]

	Fixed missing 2000 nit branch in tonemapper.

Change 3290331 on 2017/02/07 by [REDACTED_PERSON]

	Minor optimization for BlendMaterialAttributes node compilation.

Change 3291140 on 2017/02/07 by [REDACTED_PERSON]

	Fixed distance field temporal filter using Point filtering
	* Fixed the manual bilinear filter in the high quality reprojection path, still doesn't solve the streaking artifacts though

Change 3292060 on 2017/02/08 by [REDACTED_PERSON]

	DR - Some enums for max/num bits on RHI definitions

Change 3292213 on 2017/02/08 by [REDACTED_PERSON]

	Updating TM-DistanceFields map in QAGame

Change 3292291 on 2017/02/08 by [REDACTED_PERSON]

	small tweaks to TM-DistanceFields

Change 3292399 on 2017/02/08 by [REDACTED_PERSON]

	Minor fix to HDR format type check.

Change 3293560 on 2017/02/08 by [REDACTED_PERSON]

	DR - vk - Events

Change 3293562 on 2017/02/08 by [REDACTED_PERSON]

	DR - vk - Disambiguate class name

Change 3295346 on 2017/02/09 by [REDACTED_PERSON]

	DR - Duplicate fix from 3295320

Change 3296930 on 2017/02/10 by [REDACTED_PERSON]

	Search keywords for pre-skinned position/normal nodes.

Change 3297162 on 2017/02/10 by [REDACTED_PERSON]

	Distance field temporal filter stores a confidence value, which is used to track leaking of occlusion during the upsample, and flush those leaked values through the history faster.  Reduces DFAO ghosting when the camera is moving.

Change 3297345 on 2017/02/10 by [REDACTED_PERSON]

	Added 'r.CompressMeshDistanceFields' to rendering project settings, defaults to off to prevent hitches when streaming in levels

Change 3297371 on 2017/02/10 by [REDACTED_PERSON]

	Custom vertex interpolator in/out node with auto-packing - Function support. Unified UV/CVI interpolator arrays for mixed custom UV support.

Change 3298013 on 2017/02/10 by [REDACTED_PERSON]

	Ray Traced Distance Field shadowing is overlapped with the shadow depth pass, controlled by r.DFShadowAsyncCompute.
	* Didn't save any GPU time on PS4, so currently disabled until further investigation.
	* This change breaks RTDF shadows in splitscreen / stereo

Change 3300028 on 2017/02/13 by [REDACTED_PERSON]

	Adding GPU/Render/Game thread timers by default. Remove pointless flavor whitelisting on the plugin. Add sample usage in ShooterGame.ini config.
	#tests Ran half a dozen perf passes, and preflighted my changes in EC

Change 3301571 on 2017/02/14 by [REDACTED_PERSON]

	Deleted this plugin since it is no longer needed as an example

Change 3301882 on 2017/02/14 by [REDACTED_PERSON]

	Fixed DistanceFieldAOConfidenceHistoryRT not getting released on RT

Change 3304283 on 2017/02/15 by [REDACTED_PERSON]

	Variable shadowing fix

Change 3304653 on 2017/02/15 by [REDACTED_PERSON]

	Fix thread safety issues in Performance Monitor - deal with floating point bug in determining when to record frames.
	#tests Preflighted changes several times, ran several runs on Showdown

diff --git a/Engine/Plugins/Performance/PerformanceMonitor/PerformanceMonitor.uplugin b/Engine/Plugins/Performance/PerformanceMonitor/PerformanceMonitor.uplugin
--- a/Engine/Plugins/Performance/PerformanceMonitor/PerformanceMonitor.uplugin
+++ b/Engine/Plugins/Performance/PerformanceMonitor/PerformanceMonitor.uplugin
@@ -19,12 +19,7 @@
 		{
 			"Name" : "PerformanceMonitor",
 			"Type" : "Runtime",
-			"LoadingPhase" : "Default",
-			"WhitelistPlatforms" :
-			[
-				"Win64",
-				"Linux"
-			]
+			"LoadingPhase" : "Default"
 		}
 	]
 }
\ No newline at end of file
diff --git a/Engine/Plugins/Performance/PerformanceMonitor/Source/PerformanceMonitor/Private/PerformanceMonitor.cpp b/Engine/Plugins/Performance/PerformanceMonitor/Source/PerformanceMonitor/Private/PerformanceMonitor.cpp
--- a/Engine/Plugins/Performance/PerformanceMonitor/Source/PerformanceMonitor/Private/PerformanceMonitor.cpp
+++ b/Engine/Plugins/Performance/PerformanceMonitor/Source/PerformanceMonitor/Private/PerformanceMonitor.cpp
@@ -19,7 +19,6 @@ FPerformanceMonitorModule::FPerformanceMonitorModule()
 {
 	bRecording = false;
 	FileToLogTo = nullptr;
-
 }
 
 FPerformanceMonitorModule::~FPerformanceMonitorModule()
@@ -36,11 +35,10 @@ bool FPerformanceMonitorModule::Tick(float DeltaTime)
 
 	if (bRecording)
 	{
-		float CurrentTime = FPlatformTime::Seconds();
-		if (CurrentTime - TimeBetweenRecords > TimeOfLastRecord)
+		if (FPlatformTime::Seconds() - TimeOfLastRecord > TimeBetweenRecords)
 		{
 			RecordFrame();
-			TimeOfLastRecord = CurrentTime;
+			TimeOfLastRecord = FPlatformTime::Seconds();
 		}
 	}
 #endif
@@ -62,6 +60,18 @@ void FPerformanceMonitorModule::Init()
 	
 }
 
+void FPerformanceMonitorModule::GetDataFromStatsThread(int64 CurrentFrame)
+{
+#if STATS
+	if (!bNewFrameDataReady && CurrentFrame >= 0)
+	{
+		FStatsThreadState& Stats = FStatsThreadState::GetLocalState();
+		ReceivedFramePayload = Stats.GetCondensedHistory(CurrentFrame);
+		bNewFrameDataReady = true;
+	}
+#endif
+}
+
 void FPerformanceMonitorModule::SetRecordInterval(float NewInterval)
 {
 	TimeBetweenRecords = NewInterval;
@@ -164,7 +174,7 @@ void FPerformanceMonitorModule::StartRecordingPerfTimers(FString FileNameToUse,
 		GLog->Log(TEXT("PerformanceMonitor"), ELogVerbosity::Warning, TEXT("Tried to start recording when we already have started! Don't do that."));
 		return;
 	}
-	FThreadStats::MasterEnableAdd(1);
+
 	if (FileNameToUse == TEXT(""))	
 	{
 		GLog->Log(TEXT("PerformanceMonitor"), ELogVerbosity::Warning, TEXT("Please set a file name."));
@@ -191,6 +201,7 @@ void FPerformanceMonitorModule::StartRecordingPerfTimers(FString FileNameToUse,
 	if (GConfig)
 	{
 		FString ConfigCategory = FString::Printf(TEXT("/Plugins/PerformanceMonitor/%s"), *FileNameToUse);
+		StatGroupsToUse.Empty();
 		float floatValueReceived = 0;
 		if (GConfig->GetFloat(*ConfigCategory,
 			TEXT("PerformanceMonitorInterval"),
@@ -215,6 +226,12 @@ void FPerformanceMonitorModule::StartRecordingPerfTimers(FString FileNameToUse,
 		{
 			DesiredStats = TimersOfInterest;
 		}
+		TArray<FString> TimerGroupsOfInterest;
+		if (GConfig->GetArray(*ConfigCategory,
+			TEXT("PerformanceMonitorStatGroups"), TimerGroupsOfInterest, GGameIni))
+		{
+			StatGroupsToUse = TimerGroupsOfInterest;
+		}
 		FString MapToLoad;
 		if (GConfig->GetString(*ConfigCategory,
 			TEXT("PerformanceMonitorMap"), MapToLoad, GGameIni))
@@ -244,6 +261,28 @@ void FPerformanceMonitorModule::StartRecordingPerfTimers(FString FileNameToUse,
 		DesiredStats.Add(TEXT("STAT_ParticleRenderingTime"));
 		DesiredStats.Add(TEXT("STAT_GPUParticleTickTime"));
 	}
+	FThreadStats::MasterEnableAdd(1);
+
+	FStatsThreadState& Stats = FStatsThreadState::GetLocalState();
+	// Set up our delegate to gather data from the stats thread for safe consumption on game thread.
+	Stats.NewFrameDelegate.AddRaw(this, &FPerformanceMonitorModule::GetDataFromStatsThread);
+	// Cut down the flow of stats if we can to make things work more efficiently.
+	if (StatGroupsToUse.Num())
+	{
+		for (const FWorldContext& WorldContext : GEngine->GetWorldContexts())
+		{
+			if (WorldContext.WorldType == EWorldType::Game || WorldContext.WorldType == EWorldType::PIE)
+			{
+				UWorld* World = WorldContext.World();
+				GEngine->Exec(World, TEXT("stat group none"));
+				for (int i = 0; i < StatGroupsToUse.Num(); i++)
+				{
+					FString StatGroupCommand = FString::Printf(TEXT("stat group enable %s"), *StatGroupsToUse[i]);
+					GEngine->Exec(World, *StatGroupCommand);
+				}
+			}
+		}
+	}
 	GeneratedStats.Empty();
 	for (int i = 0; i < DesiredStats.Num(); i++)
 	{
@@ -269,7 +308,7 @@ void FPerformanceMonitorModule::StartRecordingPerfTimers(FString FileNameToUse,
 
 void FPerformanceMonitorModule::RecordFrame()
 {
-	if (!bRecording)
+	if (!bRecording || !bNewFrameDataReady)
 	{
 		return;
 	}
@@ -288,7 +327,36 @@ void FPerformanceMonitorModule::GetStatsBreakdown()
 	int curFrame = Stats.GetLatestValidFrame();
 	if (curFrame >= 0)
 	{
-		StoredMessages.Add(Stats.GetCondensedHistory(curFrame));
+		//StoredMessages.Add(Stats.GetCondensedHistory(curFrame));
+		TArray<float> ArrayForStatName = GeneratedStats.FindOrAdd(TEXT("RenderThreadTime"));
+		ArrayForStatName.Add(FPlatformTime::ToMilliseconds(GRenderThreadTime));
+		GeneratedStats.Emplace(TEXT("RenderThreadTime"), ArrayForStatName);
+		ArrayForStatName = GeneratedStats.FindOrAdd(TEXT("GameThreadTime"));
+		ArrayForStatName.Add(FPlatformTime::ToMilliseconds(GGameThreadTime));
+		GeneratedStats.Emplace(TEXT("GameThreadTime"), ArrayForStatName);
+		ArrayForStatName = GeneratedStats.FindOrAdd(TEXT("GPUFrameTime"));
+		ArrayForStatName.Add(FPlatformTime::ToMilliseconds(GGPUFrameTime));
+		GeneratedStats.Emplace(TEXT("GPUFrameTime"), ArrayForStatName);
+		TArray<FString> StatsCoveredThisFrame = DesiredStats;
+		for (int j = 0; j < ReceivedFramePayload.Num(); j++)
+		{
+			FStatMessage TempMessage = ReceivedFramePayload[j];
+			FName StatFName = TempMessage.NameAndInfo.GetShortName();
+			if (!StatFName.IsValid())
+			{
+				return;
+			}
+			FString StatName = StatFName.ToString();
+			if (StatsCoveredThisFrame.Contains(StatName))
+			{
+				ArrayForStatName = GeneratedStats.FindOrAdd(StatName);
+				ArrayForStatName.Add(FPlatformTime::ToMilliseconds(TempMessage.GetValue_Duration()));
+				GeneratedStats.Emplace(StatName, ArrayForStatName);
+				StatsCoveredThisFrame.Remove(StatName);
+			}
+		}
+		ReceivedFramePayload.Empty();
+		bNewFrameDataReady = false;
 	}
 #endif
 }
@@ -319,6 +387,8 @@ void FPerformanceMonitorModule::StopRecordingPerformanceTimers()
 	}
 	RecordData();
 	FThreadStats::MasterEnableSubtract(1);
+	FStatsThreadState& Stats = FStatsThreadState::GetLocalState();
+	Stats.NewFrameDelegate.RemoveAll(this);
 	bRecording = false;
 	if (!GeneratedStats.Num())
 	{
@@ -339,6 +409,25 @@ void FPerformanceMonitorModule::StopRecordingPerformanceTimers()
 #endif
 }
 
+float GetAverageOfArray(TArray<float>ArrayToAvg)
+{
+	float RetVal = 0.f;
+	int NumValidValues = 0;
+	for (int i = 0; i < ArrayToAvg.Num(); i++)
+	{
+		if (ArrayToAvg[i] >= 0)
+		{
+			RetVal += ArrayToAvg[i];
+			NumValidValues++;
+		}
+	}
+	if (NumValidValues)
+	{
+		RetVal /= NumValidValues;
+	}
+	return RetVal;
+}
+
 void FPerformanceMonitorModule::RecordData()
 {
 #if STATS
@@ -346,29 +435,6 @@ void FPerformanceMonitorModule::RecordData()
 
 	FString StringToPrint = FString::Printf(TEXT("Interval (s),%0.4f\n"), TimeBetweenRecords);
 	FileToLogTo->Serialize(TCHAR_TO_ANSI(*StringToPrint), StringToPrint.Len());
-	for (int i = 0; i < StoredMessages.Num(); i++)
-	{
-		TArray<FStatMessage> CurrentFrame = StoredMessages[i];
-		TArray<FString> StatsCoveredThisFrame = DesiredStats;
-		for (int j = 0; j < CurrentFrame.Num(); j++)
-		{
-			FStatMessage TempMessage = CurrentFrame[j];
-			FString StatName = TempMessage.NameAndInfo.GetShortName().ToString();
-			if (StatsCoveredThisFrame.Contains(StatName))
-			{
-				TArray<float> ArrayForStatName = GeneratedStats.FindOrAdd(StatName);
-				ArrayForStatName.Add(FPlatformTime::ToMilliseconds(TempMessage.GetValue_Duration()));
-				GeneratedStats.Emplace(StatName, ArrayForStatName);
-				StatsCoveredThisFrame.Remove(StatName);
-			}			
-		}
-		for (int j = 0; j < StatsCoveredThisFrame.Num(); j++)
-		{
-			TArray<float> ArrayForStatName = GeneratedStats.FindOrAdd(StatsCoveredThisFrame[j]);
-			ArrayForStatName.Add(0.f);
-			GeneratedStats.Emplace(StatsCoveredThisFrame[j], ArrayForStatName);
-		}
-	}
 	for (TMap<FString, TArray<float>>::TConstIterator ThreadedItr(GeneratedStats); ThreadedItr; ++ThreadedItr)
 	{
 				StringToPrint = ThreadedItr.Key();
@@ -380,7 +446,26 @@ void FPerformanceMonitorModule::RecordData()
 				StringToPrint.Append(TEXT("\n"));
 				FileToLogTo->Serialize(TCHAR_TO_ANSI(*StringToPrint), StringToPrint.Len());
 	}
-
+	StringToPrint = TEXT("Timer Name, Min Val, Max Val, Avg Val, Timer Active Frames\n");
+	FileToLogTo->Serialize(TCHAR_TO_ANSI(*StringToPrint), StringToPrint.Len());
+	for (TMap<FString, TArray<float>>::TConstIterator ThreadedItr(GeneratedStats); ThreadedItr; ++ThreadedItr)
+	{
+		StringToPrint = ThreadedItr.Key();
+		TArray<float> StatValues = ThreadedItr.Value();
+		float StatMin = FMath::Min(StatValues);
+		float StatMax = FMath::Max(StatValues);
+		float StatAvg = GetAverageOfArray(StatValues);
+		int ActiveFrames = StatValues.Num();
+		for (int i = 0; i < StatValues.Num(); i++)
+		{
+			if (!StatValues[i])
+			{
+				ActiveFrames--;
+			}
+		}
+		StringToPrint = FString::Printf(TEXT("%s,%0.4f,%0.4f,%0.4f,%d\n"), *StringToPrint, StatMin, StatMax, StatAvg, ActiveFrames);
+		FileToLogTo->Serialize(TCHAR_TO_ANSI(*StringToPrint), StringToPrint.Len());
+	}
 #endif
 }
 
diff --git a/Engine/Plugins/Performance/PerformanceMonitor/Source/PerformanceMonitor/Private/PerformanceMonitor.h b/Engine/Plugins/Performance/PerformanceMonitor/Source/PerformanceMonitor/Private/PerformanceMonitor.h
--- a/Engine/Plugins/Performance/PerformanceMonitor/Source/PerformanceMonitor/Private/PerformanceMonitor.h
+++ b/Engine/Plugins/Performance/PerformanceMonitor/Source/PerformanceMonitor/Private/PerformanceMonitor.h
@@ -22,8 +22,11 @@ class FPerformanceMonitorModule : public IModuleInterface, public FSelfRegisteri
 	};
 private: 
 	TArray<FString> DesiredStats;
+	TArray<FString> StatGroupsToUse;
 #if STATS
 	TArray<TArray<FStatMessage>> StoredMessages;
+	// Received frame data. Cleared after parsing.
+	TArray<FStatMessage> ReceivedFramePayload;
 #endif
 	TMap<FString, TArray<float>> GeneratedStats;
 	FArchive * FileToLogTo;
@@ -39,7 +42,6 @@ private:
 	bool bHasWarnedAboutTime;
 	bool bExitOnCompletion;
 	bool bRequiresCutsceneStart;
-
 protected:
 public:
 
@@ -81,6 +83,9 @@ public:
 	///** [FTickableGameObject] tick when recording as long as we are not CDO. */
 	bool IsTickable() const { return bRecording; }
 
+	// Whether the received frame is ready to be interacted with on the game thread.
+	bool bNewFrameDataReady;
+	void GetDataFromStatsThread(int64 CurrentFrame);
 
 	bool IsRecordingPerfTimers() { return bRecording; }
 	void StartRecordingPerfTimers(FString FileName, TArray<FString> StatsToRecord);
diff --git a/Engine/Shaders/DistanceFieldLightingPost.usf b/Engine/Shaders/DistanceFieldLightingPost.usf
--- a/Engine/Shaders/DistanceFieldLightingPost.usf
+++ b/Engine/Shaders/DistanceFieldLightingPost.usf
@@ -9,231 +9,17 @@
 #include "DistanceFieldLightingShared.usf"
 #include "DistanceFieldAOShared.usf"
 
-// 1 / ((1 - FadeDistanceFraction) * AOMaxViewDistance)
-float DistanceFadeScale;
-float SelfOcclusionReplacement;
-
-/** Normalizes the splatted surface cache values, packs depth in alpha. */
-void AOCombinePS(
-	in float4 UVAndScreenPos : TEXCOORD0
-	,out float4 OutBentNormal : SV_Target0
-#if SUPPORT_IRRADIANCE
-	,out float4 OutIrradiance : SV_Target1
-#endif
-	)
-{
-#if SUPPORT_IRRADIANCE
-	OutIrradiance = 0;
-#endif
-
-	float3 WorldNormal;
-	float SceneDepth;
-	bool bHasDistanceFieldRepresentation;
-	bool bLit;
-	GetDownsampledGBuffer(UVAndScreenPos.xy, WorldNormal, SceneDepth, bHasDistanceFieldRepresentation, bLit);
-
-#define VISUALIZE_ACCUMULATED_WEIGHTS 0
-#if VISUALIZE_ACCUMULATED_WEIGHTS
-
-	float3 BentNormalAO = Texture2DSampleLevel(BentNormalAOTexture, BentNormalAOSampler, UVAndScreenPos.xy, 0).aaa;
-	OutBentNormal = float4(BentNormalAO, SceneDepth); 
-
-#else
-
-	float4 BentNormalAccumulation = Texture2DSampleLevel(BentNormalAOTexture, BentNormalAOSampler, UVAndScreenPos.xy, 0);
-
-	if (BentNormalAccumulation.w > 0)
-	{
-		OutBentNormal.rgb = BentNormalAccumulation.xyz / BentNormalAccumulation.w;
-		OutBentNormal.a = SceneDepth;
-
-		if (!bHasDistanceFieldRepresentation)
-		{
-			// Pixels without a distance field representation interpolate AO instead of a bent normal
-			// Construct a bent normal from the interpolated AO
-			OutBentNormal.rgb = OutBentNormal.r * WorldNormal;
-		}
-
-		#if SUPPORT_IRRADIANCE
-			OutIrradiance.rgb = Texture2DSampleLevel(IrradianceTexture, IrradianceSampler, UVAndScreenPos.xy, 0).xyz / BentNormalAccumulation.w;
-
-			if (!bHasDistanceFieldRepresentation/* && !bHasHeightfieldRepresentation*/)
-			{
-				OutIrradiance.rgb *= SelfOcclusionReplacement;
-			}
-		#endif
-	}
-	else
-	{
-		OutBentNormal.rgb = float3(0, 0, .1f);
-		// Sign bit stores whether texel is valid
-		OutBentNormal.a = -SceneDepth;
-	}
-
-	float BentNormalLength = length(OutBentNormal.rgb);
-	// Fade to unoccluded in the distance
-	float FadeAlpha = saturate((AOMaxViewDistance - SceneDepth) * DistanceFadeScale);
-	float3 NormalizedBentNormal = OutBentNormal.rgb / max(BentNormalLength, .0001f);
-	OutBentNormal.rgb = NormalizedBentNormal * lerp(1, BentNormalLength, FadeAlpha);
-
-	// Clamp Nan's before they get to the filter
-	OutBentNormal.rgb = clamp(OutBentNormal.rgb, -1, 1);
-
-	FLATTEN
-	if (SceneDepth > AOMaxViewDistance)
-	{
-		// Mark as valid for the gap filling pass, so we don't get overwritten
-		OutBentNormal.a = abs(OutBentNormal.a);
-	}
-
-#endif
-}
-
-float2 BentNormalAOTexelSize;
-float MinDownsampleFactorToBaseLevel;
-
-#define HALF_FILL_KERNEL_SIZE 2
-
-/** Fills in texels with no splatted weight from screenspace neighbors. */
-void FillGapsPS(
-	in float4 UVAndScreenPos : TEXCOORD0
-	,out float4 OutBentNormal : SV_Target0
-#if SUPPORT_IRRADIANCE
-	,out float4 OutIrradiance : SV_Target1
-#endif
-	)
-{
-	float4 CenterValue = Texture2DSampleLevel(BentNormalAOTexture, BentNormalAOSampler, UVAndScreenPos.xy, 0);
-	float SceneDepth = abs(CenterValue.w);
-
-#if HIGH_QUALITY_FILL_GAPS
-
-	float4 Accumulation = 0;
-	float3 IrradianceAccumulation = 0;
-	float Length = 0;
-	// Pixels without a valid interpolation need to find valid neighbors
-	// Pixels with a valid interpolation just get a slight spatial filter so don't rely on a large kernel
-	float HalfKernelSize = CenterValue.w < 0 ? HALF_FILL_KERNEL_SIZE : 1;
-
-	for (float y = -HalfKernelSize; y <= HalfKernelSize; y++)
-	{
-		float WeightY = exp2(-abs(y * 10.0f / HALF_FILL_KERNEL_SIZE));
-
-		for (float x = -HalfKernelSize; x <= HalfKernelSize; x++)
-		{
-			float2 UVOffset = BentNormalAOTexelSize * float2(x, y);
-			float4 TextureValue = Texture2DSampleLevel(BentNormalAOTexture, BentNormalAOSampler, UVAndScreenPos.xy + UVOffset, 0);
-
-			float SampleSceneDepth = abs(TextureValue.w);
-			float ValidMask = TextureValue.w > 0;
-
-			// Don't let depth weight go to 0 with huge depth differences
-			float DepthWeight = max(exp2(-abs(SceneDepth - SampleSceneDepth) * .01f), .001f);
-
-			float2 Weight2D = float2(exp2(-abs(x * 10.0f / HALF_FILL_KERNEL_SIZE)), WeightY);
-			float ScreenSpaceSpatialWeight = max(Weight2D.x, Weight2D.y);
-
-			float Weight = ValidMask * ScreenSpaceSpatialWeight * DepthWeight;
-
-			// Track length separately 
-			Length += length(TextureValue.rgb) * Weight;
-			Accumulation.rgb += TextureValue.rgb * Weight;
-			Accumulation.a += Weight;
-
-			#if SUPPORT_IRRADIANCE
-				float3 Irradiance = Texture2DSampleLevel(IrradianceTexture, IrradianceSampler, UVAndScreenPos.xy + UVOffset, 0).xyz;
-				IrradianceAccumulation += Irradiance * Weight;
-			#endif
-		}
-	}
-	
-	float InvSafeWeight = 1.0f / max(Accumulation.a, .00001f);
-	float AverageLength = Length * InvSafeWeight;
-	float3 AverageBentNormal = Accumulation.rgb * InvSafeWeight;
-	float BentNormalLength = length(AverageBentNormal);
-	float3 AverageIrradiance = IrradianceAccumulation * InvSafeWeight;
-
-	if (BentNormalLength < AverageLength && BentNormalLength > 0)
-	{
-		// Fixup normal shortening due to weighted average of vectors
-		AverageBentNormal = AverageBentNormal / BentNormalLength * AverageLength;
-	}
-
-#else
-
-	float3 AverageBentNormal;
-	float3 AverageIrradiance;
-
-	if (CenterValue.w < 0)
-	{
-		float2 LowResBufferSize = 1.0f / (MinDownsampleFactorToBaseLevel * BentNormalAOTexelSize);
-		float2 LowResTexelSize = BentNormalAOTexelSize * MinDownsampleFactorToBaseLevel;
-		float2 Corner00UV = floor(UVAndScreenPos.xy * LowResBufferSize) / LowResBufferSize + .5f * BentNormalAOTexelSize;
-		float2 BilinearWeights = (UVAndScreenPos.xy - Corner00UV) * LowResBufferSize;
-
-		float4 TextureValues00 = Texture2DSampleLevel(BentNormalAOTexture, BentNormalAOSampler, Corner00UV, 0);
-		float4 TextureValues10 = Texture2DSampleLevel(BentNormalAOTexture, BentNormalAOSampler, Corner00UV + float2(LowResTexelSize.x, 0), 0);
-		float4 TextureValues01 = Texture2DSampleLevel(BentNormalAOTexture, BentNormalAOSampler, Corner00UV + float2(0, LowResTexelSize.y), 0);
-		float4 TextureValues11 = Texture2DSampleLevel(BentNormalAOTexture, BentNormalAOSampler, Corner00UV + LowResTexelSize, 0);
-
-		float4 CornerWeights = float4(
-			(1 - BilinearWeights.y) * (1 - BilinearWeights.x), 
-			(1 - BilinearWeights.y) * BilinearWeights.x,
-			BilinearWeights.y * (1 - BilinearWeights.x),
-			BilinearWeights.y * BilinearWeights.x);
-
-		float Epsilon = .0001f;
-
-		float4 CornerDepths = abs(float4(TextureValues00.w, TextureValues10.w, TextureValues01.w, TextureValues11.w));
-		float4 DepthWeights = min(10.0f / (abs(CornerDepths - SceneDepth.xxxx) + Epsilon), 1);
-		float4 FinalWeights = CornerWeights * DepthWeights;
-
-		float InvSafeWeight = 1.0f / max(dot(FinalWeights, 1), .00001f);
-
-		AverageBentNormal = 
-			(FinalWeights.x * TextureValues00.xyz 
-			+ FinalWeights.y * TextureValues10.xyz
-			+ FinalWeights.z * TextureValues01.xyz 
-			+ FinalWeights.w * TextureValues11.xyz)
-			* InvSafeWeight;
-		
-		#if SUPPORT_IRRADIANCE
-			float4 IrradianceValues00 = Texture2DSampleLevel(IrradianceTexture, IrradianceSampler, Corner00UV, 0);
-			float4 IrradianceValues10 = Texture2DSampleLevel(IrradianceTexture, IrradianceSampler, Corner00UV + float2(LowResTexelSize.x, 0), 0);
-			float4 IrradianceValues01 = Texture2DSampleLevel(IrradianceTexture, IrradianceSampler, Corner00UV + float2(0, LowResTexelSize.y), 0);
-			float4 IrradianceValues11 = Texture2DSampleLevel(IrradianceTexture, IrradianceSampler, Corner00UV + LowResTexelSize, 0);
-
-			AverageIrradiance = 
-				(FinalWeights.x * IrradianceValues00.xyz 
-				+ FinalWeights.y * IrradianceValues10.xyz
-				+ FinalWeights.z * IrradianceValues01.xyz 
-				+ FinalWeights.w * IrradianceValues11.xyz)
-				* InvSafeWeight;
-		#endif
-	}
-	else
-	{
-		AverageBentNormal = CenterValue.xyz;
-
-		#if SUPPORT_IRRADIANCE
-			AverageIrradiance = Texture2DSampleLevel(IrradianceTexture, IrradianceSampler, UVAndScreenPos.xy, 0).xyz;
-		#endif
-	}
-
-#endif
-	
-	OutBentNormal = float4(AverageBentNormal, CenterValue.w);
-
-#if SUPPORT_IRRADIANCE
-	OutIrradiance = float4(AverageIrradiance, 0);
-#endif
-}
+// Stores a confidence value for each AO value
+// A confidence of 0 means the AO value was constructed during the upsample from samples that were distant, which causes leaking
+Texture2D ConfidenceTexture;
 
 Texture2D BentNormalHistoryTexture;
+Texture2D ConfidenceHistoryTexture;
 SamplerState BentNormalHistorySampler;
 Texture2D IrradianceHistoryTexture;
 SamplerState IrradianceHistorySampler;
 float HistoryWeight;
+float AOHistoryMinConfidenceScale;
 float HistoryDistanceThreshold;
 float UseHistoryFilter;
 
@@ -260,23 +46,22 @@ float ComputeHistoryWeightBasedOnPosition(float2 ScreenPosition, float SceneDept
 	return DistanceToHistoryValue / SceneDepth > RelativeHistoryDistanceThreshold ? 0.0f : 1.0f;
 }
 
-// Doesn't work yet
-#define HIGH_QUALITY_HISTORY_REJECTION 0
-
 /** Reproject the occlusion history. */
-void UpdateHistoryPS(
+void UpdateHistoryDepthRejectionPS(
 	in float4 UVAndScreenPos : TEXCOORD0
 	,out float4 OutBentNormal : SV_Target0
+	,out float OutConfidence : SV_Target1
 #if SUPPORT_IRRADIANCE
-	,out float4 OutIrradiance : SV_Target1
+	,out float4 OutIrradiance : SV_Target2
 #endif
 	)
 {
 	// Distance field AO was computed at 0,0 regardless of viewrect min
 	float2 DistanceFieldUVs = UVAndScreenPos.xy - View.ViewRectMin.xy * View.BufferSizeAndInvSize.zw;
 	float4 NewValue = Texture2DSampleLevel(BentNormalAOTexture, BentNormalAOSampler, DistanceFieldUVs, 0);
-	float SceneDepth = abs(NewValue.w);
-	float OutputDepthSign = 1.0f;
+	float NewConfidence = Texture2DSampleLevel(ConfidenceTexture, BentNormalAOSampler, DistanceFieldUVs, 0);
+	float SceneDepth = NewValue.w;
+	float PixelSpeed;
 
 	float2 OldUV;
 
@@ -302,6 +87,7 @@ void UpdateHistoryPS(
 			MotionVector = float2(-.5f, .5f) * DecodeVelocityFromTexture(VelocityN);
 		}
 
+		PixelSpeed = length(MotionVector);
 		OldUV = UVAndScreenPos.xy + MotionVector / float2(0.5, -0.5) * View.ScreenPositionScaleBias.xy;
 	}
 	
@@ -320,74 +106,107 @@ void UpdateHistoryPS(
 		EffectiveHistoryWeight = 0;
 	}
 
-#define MANUAL_HISTORY_RESAMPLE HIGH_QUALITY_HISTORY_REJECTION
-#if MANUAL_HISTORY_RESAMPLE
-
-		float3 Unused;
-		float Unused2;
-		//@todo - pack into BentNormalAOTexture
-		bool bHasDistanceFieldRepresentation;
-		bool bUnused;
-		GetDownsampledGBuffer(DistanceFieldUVs, Unused, Unused2, bHasDistanceFieldRepresentation, bUnused);
+	// Only reduce history weight when we are moving enough that the flickering will not be apparent
+	float EffectiveHistoryMinConfidenceScale = lerp(1, AOHistoryMinConfidenceScale, saturate(PixelSpeed));
 
-		OutputDepthSign = bHasDistanceFieldRepresentation ? 1.0f : -1.0f;
+// Manual resample disabled as it doesn't affect the artifacts that cause most ghosting
+#define MANUAL_HISTORY_RESAMPLE 0
+#if MANUAL_HISTORY_RESAMPLE
 
+		float2 HistoryBufferSize = floor(View.BufferSizeAndInvSize.xy / DOWNSAMPLE_FACTOR);
 		float2 HistoryTexelSize = View.BufferSizeAndInvSize.zw * DOWNSAMPLE_FACTOR;
-		float2 OldDistanceFieldUVsCorner00 = HistoryTexelSize * (floor(OldDistanceFieldUVs / HistoryTexelSize) + .5f);
+		float2 OldDistanceFieldUVsCorner00 = HistoryTexelSize * (floor(OldDistanceFieldUVs * HistoryBufferSize - .5f) + .5f);
 		float2 OldDistanceFieldUVsCorner01 = OldDistanceFieldUVsCorner00 + float2(0, HistoryTexelSize.y);
 		float2 OldDistanceFieldUVsCorner11 = OldDistanceFieldUVsCorner00 + HistoryTexelSize;
 		float2 OldDistanceFieldUVsCorner10 = OldDistanceFieldUVsCorner00 + float2(HistoryTexelSize.x, 0);
+		float2 LerpWeights = (OldDistanceFieldUVs - OldDistanceFieldUVsCorner00) * HistoryBufferSize;
 
 		float4 Corner00Value = Texture2DSampleLevel(BentNormalHistoryTexture, BentNormalHistorySampler, OldDistanceFieldUVsCorner00, 0);
 		float4 Corner01Value = Texture2DSampleLevel(BentNormalHistoryTexture, BentNormalHistorySampler, OldDistanceFieldUVsCorner01, 0);
 		float4 Corner11Value = Texture2DSampleLevel(BentNormalHistoryTexture, BentNormalHistorySampler, OldDistanceFieldUVsCorner11, 0);
 		float4 Corner10Value = Texture2DSampleLevel(BentNormalHistoryTexture, BentNormalHistorySampler, OldDistanceFieldUVsCorner10, 0);
 
-		float2 LerpWeights = OldDistanceFieldUVs - OldDistanceFieldUVsCorner00;
+		float Corner00Confidence = Texture2DSampleLevel(ConfidenceHistoryTexture, BentNormalHistorySampler, OldDistanceFieldUVsCorner00, 0);
+		float Corner01Confidence = Texture2DSampleLevel(ConfidenceHistoryTexture, BentNormalHistorySampler, OldDistanceFieldUVsCorner01, 0);
+		float Corner11Confidence = Texture2DSampleLevel(ConfidenceHistoryTexture, BentNormalHistorySampler, OldDistanceFieldUVsCorner11, 0);
+		float Corner10Confidence = Texture2DSampleLevel(ConfidenceHistoryTexture, BentNormalHistorySampler, OldDistanceFieldUVsCorner10, 0);
 
 	#define COMPUTE_SEPARATE_POSITION_WEIGHTS 1
 	#if COMPUTE_SEPARATE_POSITION_WEIGHTS
-		float PositionWeight00 = ComputeHistoryWeightBasedOnPosition(UVAndScreenPos.zw, SceneDepth, OldScreenPosition, abs(Corner00Value.w));
-		float PositionWeight01 = ComputeHistoryWeightBasedOnPosition(UVAndScreenPos.zw, SceneDepth, OldScreenPosition, abs(Corner01Value.w));
-		float PositionWeight11 = ComputeHistoryWeightBasedOnPosition(UVAndScreenPos.zw, SceneDepth, OldScreenPosition, abs(Corner11Value.w));
-		float PositionWeight10 = ComputeHistoryWeightBasedOnPosition(UVAndScreenPos.zw, SceneDepth, OldScreenPosition, abs(Corner10Value.w));
-	
-		float ClassificationWeight00 = (Corner00Value.w > 0 == bHasDistanceFieldRepresentation) ? PositionWeight00 : 0.0f;
-		float ClassificationWeight01 = (Corner01Value.w > 0 == bHasDistanceFieldRepresentation) ? PositionWeight01 : 0.0f;
-		float ClassificationWeight11 = (Corner11Value.w > 0 == bHasDistanceFieldRepresentation) ? PositionWeight11 : 0.0f;
-		float ClassificationWeight10 = (Corner10Value.w > 0 == bHasDistanceFieldRepresentation) ? PositionWeight10 : 0.0f;
+		float PositionWeight00 = ComputeHistoryWeightBasedOnPosition(UVAndScreenPos.zw, SceneDepth, OldScreenPosition, Corner00Value.w);
+		float PositionWeight01 = ComputeHistoryWeightBasedOnPosition(UVAndScreenPos.zw, SceneDepth, OldScreenPosition, Corner01Value.w);
+		float PositionWeight11 = ComputeHistoryWeightBasedOnPosition(UVAndScreenPos.zw, SceneDepth, OldScreenPosition, Corner11Value.w);
+		float PositionWeight10 = ComputeHistoryWeightBasedOnPosition(UVAndScreenPos.zw, SceneDepth, OldScreenPosition, Corner10Value.w);
+
+		float BilinearWeight00 = (1 - LerpWeights.y) * (1 - LerpWeights.x);
+		float BilinearWeight01 = (LerpWeights.y) * (1 - LerpWeights.x);
+		float BilinearWeight11 = (LerpWeights.y) * (LerpWeights.x);
+		float BilinearWeight10 = (1 - LerpWeights.y) * (LerpWeights.x);
+
+		float3 HistoryValue = (BilinearWeight00 * lerp(NewValue.xyz, Corner00Value.xyz, PositionWeight00)
+			+ BilinearWeight01 * lerp(NewValue.xyz, Corner01Value.xyz, PositionWeight01)
+			+ BilinearWeight11 * lerp(NewValue.xyz, Corner11Value.xyz, PositionWeight11)
+			+ BilinearWeight10 * lerp(NewValue.xyz, Corner10Value.xyz, PositionWeight10));
+		
+		float HistoryConfidence = (BilinearWeight00 * lerp(NewConfidence, Corner00Confidence, PositionWeight00)
+			+ BilinearWeight01 * lerp(NewConfidence, Corner01Confidence, PositionWeight01)
+			+ BilinearWeight11 * lerp(NewConfidence, Corner11Confidence, PositionWeight11)
+			+ BilinearWeight10 * lerp(NewConfidence, Corner10Confidence, PositionWeight10));
 
-		float FinalWeight00 = (1 - LerpWeights.y) * (1 - LerpWeights.x) * ClassificationWeight00;
-		float FinalWeight01 = (LerpWeights.y) * (1 - LerpWeights.x) * ClassificationWeight01;
-		float FinalWeight11 = (LerpWeights.y) * (LerpWeights.x) * ClassificationWeight11;
-		float FinalWeight10 = (1 - LerpWeights.y) * (LerpWeights.x) * ClassificationWeight10;
-		float TotalWeight = FinalWeight00 + FinalWeight01 + FinalWeight11 + FinalWeight10;
+		EffectiveHistoryWeight *= clamp(HistoryConfidence, EffectiveHistoryMinConfidenceScale, 1);
 
-		float3 HistoryValue = (Corner00Value.xyz * FinalWeight00
-			+ Corner01Value.xyz * FinalWeight01
-			+ Corner11Value.xyz * FinalWeight11
-			+ Corner10Value.xyz * FinalWeight10) / max(TotalWeight, .00001f);
+		OutBentNormal.rgb = lerp(NewValue.rgb, HistoryValue.rgb, EffectiveHistoryWeight);
+		OutConfidence = lerp(NewConfidence, HistoryConfidence, EffectiveHistoryWeight);
 
-		EffectiveHistoryWeight *= TotalWeight > 0 ? 1.0f : 0.0f;
+		EffectiveHistoryWeight *= (PositionWeight00 > 0 || PositionWeight01 > 0 || PositionWeight11 > 0 || PositionWeight10 > 0) ? 1.0f : 0.0f;
+		
+#define BENT_NORMAL_LENGTH_FIXUP 0
+#if BENT_NORMAL_LENGTH_FIXUP
+
+		float BentNormalLength = length(HistoryValue.xyz);
+
+		float LengthVerticalLerp0 = lerp(length(Corner00Value.xyz), length(Corner01Value.xyz), LerpWeights.y);
+		float LengthVerticalLerp1 = lerp(length(Corner10Value.xyz), length(Corner11Value.xyz), LerpWeights.y);
+		float AverageLength = lerp(LengthVerticalLerp0, LengthVerticalLerp1, LerpWeights.x);
+
+		if (BentNormalLength < AverageLength && BentNormalLength > 0)
+		{
+			// Fixup normal shortening due to weighted average of vectors
+			HistoryValue.xyz = HistoryValue.xyz / BentNormalLength * AverageLength;
+		}
+
+#endif
 
 	#else
-		// Filter encoded depth path, for testing
+		// Manually implement bilinear filtering with a single position weight, for testing
 		float4 VerticalLerp0 = lerp(Corner00Value, Corner01Value, LerpWeights.y);
 		float4 VerticalLerp1 = lerp(Corner10Value, Corner11Value, LerpWeights.y);
 		float4 HistoryValue = lerp(VerticalLerp0, VerticalLerp1, LerpWeights.x);
-		float PositionWeight = ComputeHistoryWeightBasedOnPosition(UVAndScreenPos.zw, SceneDepth, OldScreenPosition, abs(HistoryValue.w));
+
+		float VerticalLerp0Confidence = lerp(Corner00Confidence, Corner01Confidence, LerpWeights.y);
+		float VerticalLerp1Confidence = lerp(Corner10Confidence, Corner11Confidence, LerpWeights.y);
+		float HistoryConfidence = lerp(VerticalLerp0Confidence, VerticalLerp1Confidence, LerpWeights.x);
+
+		float PositionWeight = ComputeHistoryWeightBasedOnPosition(UVAndScreenPos.zw, SceneDepth, OldScreenPosition, HistoryValue.w);
 		EffectiveHistoryWeight *= PositionWeight;
-	#endif
+		EffectiveHistoryWeight *= clamp(HistoryConfidence, EffectiveHistoryMinConfidenceScale, 1);
 
-	OutBentNormal.rgb = lerp(NewValue.rgb, HistoryValue.rgb, EffectiveHistoryWeight);
+		OutBentNormal.rgb = lerp(NewValue.rgb, HistoryValue.rgb, EffectiveHistoryWeight);
+		OutConfidence = lerp(NewConfidence, HistoryConfidence, EffectiveHistoryWeight);
+	#endif
 
 #else
 	// Fast path that uses hardware interpolation
 	// This fails to reject based on depth consistently because the depth has been filtered
 	float4 HistoryValue = Texture2DSampleLevel(BentNormalHistoryTexture, BentNormalHistorySampler, OldDistanceFieldUVs, 0);
-	float PositionWeight = ComputeHistoryWeightBasedOnPosition(UVAndScreenPos.zw, SceneDepth, OldScreenPosition, abs(HistoryValue.w));
+	float HistoryConfidence = Texture2DSampleLevel(ConfidenceHistoryTexture, BentNormalHistorySampler, OldDistanceFieldUVs, 0);
+	float PositionWeight = ComputeHistoryWeightBasedOnPosition(UVAndScreenPos.zw, SceneDepth, OldScreenPosition, HistoryValue.w);
+
 	EffectiveHistoryWeight *= PositionWeight;
+	EffectiveHistoryWeight *= clamp(HistoryConfidence, EffectiveHistoryMinConfidenceScale, 1);
+
 	OutBentNormal.rgb = lerp(NewValue.rgb, HistoryValue.rgb, EffectiveHistoryWeight);
+	OutConfidence = lerp(NewConfidence, HistoryConfidence, EffectiveHistoryWeight);
 #endif
 
 #if SUPPORT_IRRADIANCE
@@ -397,8 +216,18 @@ void UpdateHistoryPS(
 	OutIrradiance = float4(lerp(NewIrradiance, HistoryIrradianceValue, EffectiveHistoryWeight), 0);
 #endif
 
+	// Only discard histories when we are moving enough that the jittering we let through will not be apparent
+	float HistoryConfidenceThreshold = lerp(0, .99f, saturate(PixelSpeed));
+
+	// Mark pixels with less than full confidence as not having a valid history, so the history filter will fill them in from stable neighbors
+	// This reduces ghosting where foreground occlusion values have been transfered to the background during upsampling
+	if (HistoryConfidence < HistoryConfidenceThreshold)
+	{
+		EffectiveHistoryWeight = 0;
+	}
+
 	OutBentNormal.rgb = isnan(OutBentNormal.rgb) ? 0 : OutBentNormal.rgb;
-	OutBentNormal.a = abs(NewValue.a);
+	OutBentNormal.a = SceneDepth;
 
 	FLATTEN
 	if (UseHistoryFilter > 0)
@@ -406,24 +235,23 @@ void UpdateHistoryPS(
 		// Sign bit of alpha stores whether the history was rejected or not, to be read by the history filter pass
 		OutBentNormal.a *= EffectiveHistoryWeight > 0 ? 1 : -1;
 	}
-	else
-	{
-		OutBentNormal.a *= OutputDepthSign;
-	}
 }
 
 #define HALF_HISTORY_FILL_KERNEL_SIZE 2
+float2 BentNormalAOTexelSize;
 
 /** Seeds newly rejected history values (which are sources of temporal instability) with the results of a spatial search from stable history values */
 void FilterHistoryPS(
 	in float4 UVAndScreenPos : TEXCOORD0
 	,out float4 OutBentNormal : SV_Target0
+	,out float OutConfidence : SV_Target1
 #if SUPPORT_IRRADIANCE
-	,out float4 OutIrradiance : SV_Target1
+	,out float4 OutIrradiance : SV_Target2
 #endif
 	)
 {
 	float4 HistoryValue = Texture2DSampleLevel(BentNormalAOTexture, BentNormalAOSampler, UVAndScreenPos.xy, 0);
+	float HistoryConfidence = Texture2DSampleLevel(ConfidenceTexture, BentNormalAOSampler, UVAndScreenPos.xy, 0);
 
 #if SUPPORT_IRRADIANCE
 	float3 IrradianceValue = Texture2DSampleLevel(IrradianceTexture, IrradianceSampler, UVAndScreenPos.xy, 0).xyz;
@@ -481,22 +309,10 @@ void FilterHistoryPS(
 	}
 
 	OutBentNormal = HistoryValue;
+	OutConfidence = HistoryConfidence;
 	// Remove sign bit so future reprojection interpolation isn't affected
 	OutBentNormal.w = abs(OutBentNormal.w);
 
-#if HIGH_QUALITY_HISTORY_REJECTION
-	float3 Unused;
-	float Unused2;
-	//@todo - pack into BentNormalAOTexture
-	bool bHasDistanceFieldRepresentation;
-	bool bUnused;
-	GetDownsampledGBuffer(UVAndScreenPos.xy, Unused, Unused2, bHasDistanceFieldRepresentation, bUnused);
-	// Store whether the pixel has a distance field representation for next frame
-	// This can only be supported with manual history filtering
-	OutBentNormal.w *= bHasDistanceFieldRepresentation ? 1.0f : -1.0f;
-
-#endif
-
 #if SUPPORT_IRRADIANCE
 	OutIrradiance = float4(IrradianceValue, 0);
 #endif
@@ -550,7 +366,7 @@ void AOUpsamplePS(
 
 	float Epsilon = .0001f;
 
-	float4 CornerDepths = abs(float4(TextureValues00.w, TextureValues10.w, TextureValues01.w, TextureValues11.w));
+	float4 CornerDepths = float4(TextureValues00.w, TextureValues10.w, TextureValues01.w, TextureValues11.w);
 	float SceneDepth = CalcSceneDepth(UVAndScreenPos.xy);
 	float4 DepthWeights = 1.0f / (abs(CornerDepths - SceneDepth.xxxx) + Epsilon);
 	float4 FinalWeights = CornerWeights * DepthWeights;
diff --git a/Engine/Shaders/DistanceFieldScreenGridLighting.usf b/Engine/Shaders/DistanceFieldScreenGridLighting.usf
--- a/Engine/Shaders/DistanceFieldScreenGridLighting.usf
+++ b/Engine/Shaders/DistanceFieldScreenGridLighting.usf
@@ -431,12 +431,26 @@ float4 GetNormalWeights(float2 Corner00UV, float2 LowResTexelSize, float3 WorldN
 // 1 / ((1 - FadeDistanceFraction) * AOMaxViewDistance)
 float DistanceFadeScale;
 
+float ComputeSampleWeightBasedOnPosition(float4 ReferencePlane, float2 SampleScreenUV, float SampleDepth)
+{
+	float2 SampleScreenPosition = (SampleScreenUV.xy - View.ScreenPositionScaleBias.wz) / View.ScreenPositionScaleBias.xy;
+	float4 SampleHomogeneousWorldPosition = mul(float4(SampleScreenPosition * SampleDepth, SampleDepth, 1), View.ScreenToWorld);
+	float3 SampleWorldPosition = SampleHomogeneousWorldPosition.xyz / SampleHomogeneousWorldPosition.w;
+
+	float PlaneDistance = dot(ReferencePlane, float4(SampleWorldPosition, 1));
+
+	float Epsilon = .0001f;
+	float RelativeDistance = 1000 * abs(PlaneDistance) / SampleDepth;
+	return min(10.0f / (RelativeDistance + Epsilon), 1);
+}
+
 /**  */
 void GeometryAwareUpsamplePS(
 	in float4 UVAndScreenPos : TEXCOORD0
 	,out float4 OutBentNormal : SV_Target0
+	,out float OutConfidence : SV_Target1
 #if SUPPORT_IRRADIANCE
-	,out float4 OutIrradiance : SV_Target1
+	,out float4 OutIrradiance : SV_Target2
 #endif
 	)
 {
@@ -446,6 +460,11 @@ void GeometryAwareUpsamplePS(
 	bool bLit;
 	GetDownsampledGBuffer(UVAndScreenPos.xy, WorldNormal, SceneDepth, bHasDistanceFieldRepresentation, bLit);
 
+	float4 HomogeneousWorldPosition = mul(float4(UVAndScreenPos.zw * SceneDepth, SceneDepth, 1), View.ScreenToWorld);
+	float3 WorldPosition = HomogeneousWorldPosition.xyz / HomogeneousWorldPosition.w;
+
+	float4 ReferencePlane = float4(WorldNormal, -dot(WorldPosition, WorldNormal));
+
 	float2 LowResBufferSize = BentNormalBufferAndTexelSize.xy;
 	float2 LowResTexelSize = BentNormalBufferAndTexelSize.zw;
 	float2 Corner00UV = floor((UVAndScreenPos.xy - JitterOffset * DistanceFieldGBufferTexelSize) * LowResBufferSize) * LowResTexelSize;
@@ -463,16 +482,22 @@ void GeometryAwareUpsamplePS(
 		BilinearWeights.y * (1 - BilinearWeights.x),
 		BilinearWeights.y * BilinearWeights.x);
 
-	float4 CornerDepths = abs(float4(TextureValues00.w, TextureValues10.w, TextureValues01.w, TextureValues11.w));
+	float4 CornerDepths = float4(TextureValues00.w, TextureValues10.w, TextureValues01.w, TextureValues11.w);
+
+	float4 PositionWeights;
+	PositionWeights.x = ComputeSampleWeightBasedOnPosition(ReferencePlane, LowResCorner00UV, CornerDepths.x);
+	PositionWeights.y = ComputeSampleWeightBasedOnPosition(ReferencePlane, LowResCorner00UV + float2(LowResTexelSize.x, 0), CornerDepths.y);
+	PositionWeights.z = ComputeSampleWeightBasedOnPosition(ReferencePlane, LowResCorner00UV + float2(0, LowResTexelSize.y), CornerDepths.z);
+	PositionWeights.w = ComputeSampleWeightBasedOnPosition(ReferencePlane, LowResCorner00UV + LowResTexelSize, CornerDepths.w);
 
 	//float4 DepthWeights = max(exp2(-abs(CornerDepths - SceneDepth.xxxx) * .01f), .001f);
 	float Epsilon = .0001f;
-	float4 DepthWeights = min(10.0f / (abs(CornerDepths - SceneDepth.xxxx) + Epsilon), 1);
+	//float4 DepthWeights = min(10.0f / (abs(CornerDepths - SceneDepth.xxxx) + Epsilon), 1);
 
 	float2 FullResCorner00UV = Corner00UV + (JitterOffset + .5f) * DistanceFieldGBufferTexelSize;
 	float4 NormalWeights = GetNormalWeights(FullResCorner00UV, LowResTexelSize, WorldNormal);
 	
-	float4 FinalWeights = CornerWeights * DepthWeights * NormalWeights;
+	float4 FinalWeights = CornerWeights * PositionWeights * NormalWeights;
 
 	float InvSafeWeight = 1.0f / max(dot(FinalWeights, 1), .00001f);
 
@@ -483,8 +508,12 @@ void GeometryAwareUpsamplePS(
 		+ FinalWeights.w * TextureValues11.xyz)
 		* InvSafeWeight;
 	
+	// Output a confidence value based on position weights
+	// Confidence will be 0 when we detect that foreground occlusion values have been leaked onto the background, because no better low res samples were available
+	OutConfidence = dot(PositionWeights, 1) < .5f ? 0 : 1;
+
 	float BentNormalLength = length(AverageBentNormal);
-	
+
 #if !PASS_THROUGH_DEBUG_VALUE
 	float AverageLength = 
 		(FinalWeights.x * length(TextureValues00.xyz) 
@@ -499,6 +528,7 @@ void GeometryAwareUpsamplePS(
 		AverageBentNormal = AverageBentNormal / BentNormalLength * AverageLength;
 	}
 #endif
+
 	OutBentNormal = float4(AverageBentNormal, SceneDepth);
 
 	// Fade to unoccluded in the distance
@@ -507,7 +537,6 @@ void GeometryAwareUpsamplePS(
 	float FadeAlpha = saturate((AOMaxViewDistance - SceneDepth) * DistanceFadeScale);
 	float3 NormalizedBentNormal = OutBentNormal.rgb / max(BentNormalLength, .0001f);
 	OutBentNormal.rgb = NormalizedBentNormal * lerp(1, BentNormalLength, FadeAlpha);
-
 	//OutBentNormal = float4(WorldNormal, SceneDepth);
 
 	#if SUPPORT_IRRADIANCE
diff --git a/Engine/Shaders/GpuSkinVertexFactory.usf b/Engine/Shaders/GpuSkinVertexFactory.usf
--- a/Engine/Shaders/GpuSkinVertexFactory.usf
+++ b/Engine/Shaders/GpuSkinVertexFactory.usf
@@ -119,8 +119,8 @@ struct FVertexFactoryInterpolantsVSToPS
 
 	float4 Color : COLOR0;
 
-#if NUM_MATERIAL_TEXCOORDS
-	float2	TexCoords[NUM_MATERIAL_TEXCOORDS]	: TEXCOORD0;
+#if NUM_TEX_COORD_INTERPOLATORS
+	float2	TexCoords[NUM_TEX_COORD_INTERPOLATORS]	: TEXCOORD0;
 #endif
 
 #if INSTANCED_STEREO
@@ -139,9 +139,9 @@ FMaterialPixelParameters GetMaterialPixelParameters(FVertexFactoryInterpolantsVS
 	// GetMaterialPixelParameters is responsible for fully initializing the result
 	FMaterialPixelParameters Result = MakeInitializedMaterialPixelParameters();
 
-#if NUM_MATERIAL_TEXCOORDS	
+#if NUM_TEX_COORD_INTERPOLATORS	
 	UNROLL
-	for(int CoordinateIndex = 0;CoordinateIndex < NUM_MATERIAL_TEXCOORDS;CoordinateIndex++)
+	for(int CoordinateIndex = 0;CoordinateIndex < NUM_TEX_COORD_INTERPOLATORS;CoordinateIndex++)
 	{
 		Result.TexCoords[CoordinateIndex] = Interpolants.TexCoords[CoordinateIndex];
 	}
@@ -623,9 +623,9 @@ FVertexFactoryInterpolantsVSToDS VertexFactoryGetInterpolantsVSToDS(FVertexFacto
 FMaterialTessellationParameters GetMaterialTessellationParameters(FVertexFactoryInterpolantsVSToDS Interpolants, float3 CameraLocalWorldPosition)
 {
 	FMaterialTessellationParameters	Result;
-#if NUM_MATERIAL_TEXCOORDS
+#if NUM_TEX_COORD_INTERPOLATORS
 	UNROLL
-	for (int CoordinateIndex = 0;CoordinateIndex < NUM_MATERIAL_TEXCOORDS;CoordinateIndex++)
+	for (int CoordinateIndex = 0;CoordinateIndex < NUM_TEX_COORD_INTERPOLATORS;CoordinateIndex++)
 	{
 		Result.TexCoords[CoordinateIndex] = Interpolants.InterpolantsVSToPS.TexCoords[CoordinateIndex];
 	}
@@ -654,9 +654,9 @@ FVertexFactoryInterpolantsVSToDS VertexFactoryInterpolate(FVertexFactoryInterpol
 	TESSELLATION_INTERPOLATE_MEMBER(InterpolantsVSToPS.TangentToWorld2);
 	TESSELLATION_INTERPOLATE_MEMBER(InterpolantsVSToPS.Color);
 
-#if NUM_MATERIAL_TEXCOORDS
+#if NUM_TEX_COORD_INTERPOLATORS
 	UNROLL
-	for(int tc = 0; tc < NUM_MATERIAL_TEXCOORDS; ++tc)
+	for(int tc = 0; tc < NUM_TEX_COORD_INTERPOLATORS; ++tc)
 	{
 		TESSELLATION_INTERPOLATE_MEMBER(InterpolantsVSToPS.TexCoords[tc]);
 	}
diff --git a/Engine/Shaders/LocalVertexFactory.usf b/Engine/Shaders/LocalVertexFactory.usf
--- a/Engine/Shaders/LocalVertexFactory.usf
+++ b/Engine/Shaders/LocalVertexFactory.usf
@@ -287,13 +287,13 @@ FMaterialPixelParameters GetMaterialPixelParameters(FVertexFactoryInterpolantsVS
 	// GetMaterialPixelParameters is responsible for fully initializing the result
 	FMaterialPixelParameters Result = MakeInitializedMaterialPixelParameters();
 
-#if NUM_MATERIAL_TEXCOORDS
+#if NUM_TEX_COORD_INTERPOLATORS
 	UNROLL
-	for( int CoordinateIndex = 0; CoordinateIndex < NUM_MATERIAL_TEXCOORDS; CoordinateIndex++ )
+	for( int CoordinateIndex = 0; CoordinateIndex < NUM_TEX_COORD_INTERPOLATORS; CoordinateIndex++ )
 	{
 		Result.TexCoords[CoordinateIndex] = GetUV(Interpolants, CoordinateIndex);
 	}
-#endif	// NUM_MATERIAL_TEXCOORDS
+#endif
 
 #if USE_PARTICLE_SUBUVS
 	// Output TexCoord0 for when previewing materials that use ParticleSubUV.
@@ -626,20 +626,20 @@ FVertexFactoryInterpolantsVSToPS VertexFactoryGetInterpolantsVSToPS(FVertexFacto
 	// Really only the last two components of the packed UVs have the opportunity to be uninitialized
 	Interpolants = (FVertexFactoryInterpolantsVSToPS)0;
 
-#if NUM_MATERIAL_TEXCOORDS
-
-	float2 CustomizedUVs[NUM_MATERIAL_TEXCOORDS];
+#if NUM_TEX_COORD_INTERPOLATORS
+	float2 CustomizedUVs[NUM_TEX_COORD_INTERPOLATORS];
 	GetMaterialCustomizedUVs(VertexParameters, CustomizedUVs);
+	GetCustomInterpolators(VertexParameters, CustomizedUVs);
 	
 	UNROLL
-	for (int CoordinateIndex = 0; CoordinateIndex < NUM_MATERIAL_TEXCOORDS; CoordinateIndex++)
+	for (int CoordinateIndex = 0; CoordinateIndex < NUM_TEX_COORD_INTERPOLATORS; CoordinateIndex++)
 	{
 		SetUV(Interpolants, CoordinateIndex, CustomizedUVs[CoordinateIndex]);
 	}
 
 #elif NUM_MATERIAL_TEXCOORDS_VERTEX == 0 && USE_PARTICLE_SUBUVS
 	SetUV(Interpolants, 0, Input.TexCoords[0]);
-#endif	// NUM_MATERIAL_TEXCOORDS
+#endif
 
 #if NEEDS_LIGHTMAP_COORDINATE
 	float2 LightMapCoordinate = 0;
@@ -745,17 +745,17 @@ float4 VertexFactoryGetPreviousWorldPosition(FVertexFactoryInput Input, FVertexF
 	FMaterialTessellationParameters GetMaterialTessellationParameters(FVertexFactoryInterpolantsVSToDS Interpolants, float3 CameraLocalWorldPosition)
 	{
 		FMaterialTessellationParameters	Result;
-	#if NUM_MATERIAL_TEXCOORDS
+	#if NUM_TEX_COORD_INTERPOLATORS
 		UNROLL
-		for(int CoordinateIndex = 0;CoordinateIndex < NUM_MATERIAL_TEXCOORDS;CoordinateIndex += 2)
+		for(int CoordinateIndex = 0;CoordinateIndex < NUM_TEX_COORD_INTERPOLATORS;CoordinateIndex += 2)
 		{
 			Result.TexCoords[CoordinateIndex] = Interpolants.InterpolantsVSToPS.TexCoords[CoordinateIndex/2].xy;
-			if(CoordinateIndex + 1 < NUM_MATERIAL_TEXCOORDS)
+			if(CoordinateIndex + 1 < NUM_TEX_COORD_INTERPOLATORS)
 			{
 				Result.TexCoords[CoordinateIndex + 1] = Interpolants.InterpolantsVSToPS.TexCoords[CoordinateIndex/2].zw;
 			}
 		}
-	#endif	// NUM_MATERIAL_TEXCOORDS
+	#endif // NUM_TEX_COORD_INTERPOLATORS
 
 		half3 TangentToWorld0 = Interpolants.InterpolantsVSToPS.TangentToWorld0.xyz;
 		half4 TangentToWorld2 = Interpolants.InterpolantsVSToPS.TangentToWorld2;
@@ -800,9 +800,9 @@ float4 VertexFactoryGetPreviousWorldPosition(FVertexFactoryInput Input, FVertexF
 		TESSELLATION_INTERPOLATE_MEMBER(InterpolantsVSToPS.LightMapCoordinate);
 	#endif
 
-	#if NUM_MATERIAL_TEXCOORDS
+	#if NUM_TEX_COORD_INTERPOLATORS
 		UNROLL
-		for(int tc = 0; tc < (NUM_MATERIAL_TEXCOORDS+1)/2; ++tc)
+		for(int tc = 0; tc < (NUM_TEX_COORD_INTERPOLATORS+1)/2; ++tc)
 		{
 			TESSELLATION_INTERPOLATE_MEMBER(InterpolantsVSToPS.TexCoords[tc]);
 		}
diff --git a/Engine/Shaders/LocalVertexFactoryCommon.usf b/Engine/Shaders/LocalVertexFactoryCommon.usf
--- a/Engine/Shaders/LocalVertexFactoryCommon.usf
+++ b/Engine/Shaders/LocalVertexFactoryCommon.usf
@@ -21,8 +21,8 @@ struct FVertexFactoryInterpolantsVSToPS
 	float4  PerInstanceParams : COLOR1;
 #endif
 
-#if NUM_MATERIAL_TEXCOORDS
-	float4	TexCoords[(NUM_MATERIAL_TEXCOORDS+1)/2]	: TEXCOORD0;
+#if NUM_TEX_COORD_INTERPOLATORS
+	float4	TexCoords[(NUM_TEX_COORD_INTERPOLATORS+1)/2]	: TEXCOORD0;
 #elif USE_PARTICLE_SUBUVS
 	float4	TexCoords[1] : TEXCOORD0;
 #endif
@@ -36,7 +36,7 @@ struct FVertexFactoryInterpolantsVSToPS
 #endif
 };
 
-#if NUM_MATERIAL_TEXCOORDS || USE_PARTICLE_SUBUVS
+#if NUM_TEX_COORD_INTERPOLATORS || USE_PARTICLE_SUBUVS
 float2 GetUV(FVertexFactoryInterpolantsVSToPS Interpolants, int UVIndex)
 {
 	float4 UVVector = Interpolants.TexCoords[UVIndex / 2];
@@ -111,7 +111,7 @@ void SetTangents(inout FVertexFactoryInterpolantsVSToPS Interpolants, float3 InT
 
 #else // #if FEATURE_LEVEL >= FEATURE_LEVEL_ES3_1
 
-#define NUM_PACKED_UVS (NUM_MATERIAL_TEXCOORDS > 1 ? 3 : 1)
+#define NUM_PACKED_UVS (NUM_TEX_COORD_INTERPOLATORS > 1 ? 3 : 1)
 
 #define USE_INTERPOLATOR_345 ((!MATERIAL_SHADINGMODEL_UNLIT) || USES_TRANSFORM_VECTOR || USES_DISTORTION)
 
diff --git a/Engine/Shaders/MaterialTemplate.usf b/Engine/Shaders/MaterialTemplate.usf
--- a/Engine/Shaders/MaterialTemplate.usf
+++ b/Engine/Shaders/MaterialTemplate.usf
@@ -18,6 +18,8 @@
 
 #define NUM_MATERIAL_TEXCOORDS_VERTEX %s
 #define NUM_MATERIAL_TEXCOORDS %s
+#define NUM_CUSTOM_VERTEX_INTERPOLATORS %s
+#define NUM_TEX_COORD_INTERPOLATORS %s
 
 #ifdef MIN_MATERIAL_TEXCOORDS 
 	#include "MinMaterialTexCoords.usf"
@@ -116,8 +118,8 @@ struct FPixelMaterialInputs
  */
 struct FMaterialPixelParameters
 {
-#if NUM_MATERIAL_TEXCOORDS
-	float2 TexCoords[NUM_MATERIAL_TEXCOORDS];
+#if NUM_TEX_COORD_INTERPOLATORS
+	float2 TexCoords[NUM_TEX_COORD_INTERPOLATORS];
 #endif
 
 	/** Interpolated vertex color, in linear color space. */
@@ -233,8 +235,8 @@ struct FMaterialTessellationParameters
 {
 	// Note: Customized UVs are only evaluated in the vertex shader, which is not really what you want with tessellation, but keeps the code simpler 
 	// (tessellation texcoords are the same as pixels shader texcoords)
-#if NUM_MATERIAL_TEXCOORDS
-	float2 TexCoords[NUM_MATERIAL_TEXCOORDS];
+#if NUM_TEX_COORD_INTERPOLATORS
+	float2 TexCoords[NUM_TEX_COORD_INTERPOLATORS];
 #endif
 	float4 VertexColor;
 	// TODO: Non translated world position
@@ -1538,8 +1540,13 @@ half2 GetMaterialRefraction(FPixelMaterialInputs PixelMaterialInputs)
 	return PixelMaterialInputs.Refraction;
 }
 
-#if NUM_MATERIAL_TEXCOORDS
-void GetMaterialCustomizedUVs(FMaterialVertexParameters Parameters, out float2 OutTexCoords[NUM_MATERIAL_TEXCOORDS])
+#if NUM_TEX_COORD_INTERPOLATORS
+void GetMaterialCustomizedUVs(FMaterialVertexParameters Parameters, inout float2 OutTexCoords[NUM_TEX_COORD_INTERPOLATORS])
+{
+%s
+}
+
+void GetCustomInterpolators(FMaterialVertexParameters Parameters, inout float2 OutTexCoords[NUM_TEX_COORD_INTERPOLATORS])
 {
 %s
 }
diff --git a/Engine/Shaders/MeshParticleVertexFactory.usf b/Engine/Shaders/MeshParticleVertexFactory.usf
--- a/Engine/Shaders/MeshParticleVertexFactory.usf
+++ b/Engine/Shaders/MeshParticleVertexFactory.usf
@@ -85,12 +85,12 @@ struct FVertexFactoryInterpolantsVSToPS
 	float4 SubUV0AndTexCoord0 : TEXCOORD1;
 	float4 SubUV1AndLerp : TEXCOORD2;
 #else
-	#if NUM_MATERIAL_TEXCOORDS
+	#if NUM_TEX_COORD_INTERPOLATORS
 			#if (ES2_PROFILE || ES3_1_PROFILE)
 				// Avoid dependent texture fetches, put all UVs in xy
-				float2	TexCoords[NUM_MATERIAL_TEXCOORDS] : TEXCOORD0;
+				float2	TexCoords[NUM_TEX_COORD_INTERPOLATORS] : TEXCOORD0;
 			#else
-				float4	TexCoords[(NUM_MATERIAL_TEXCOORDS+1)/2]	: TEXCOORD0;
+				float4	TexCoords[(NUM_TEX_COORD_INTERPOLATORS+1)/2] : TEXCOORD0;
 			#endif
 	#endif
 #endif
@@ -213,20 +213,20 @@ FMaterialPixelParameters GetMaterialPixelParameters(FVertexFactoryInterpolantsVS
 	Result.Particle.SubUVCoords[0] = Interpolants.SubUV0AndTexCoord0.xy;
 	Result.Particle.SubUVCoords[1] = Interpolants.SubUV1AndLerp.xy;
 	Result.Particle.SubUVLerp = Interpolants.SubUV1AndLerp.z;
-#elif NUM_MATERIAL_TEXCOORDS
+#elif NUM_TEX_COORD_INTERPOLATORS
 	#if (ES2_PROFILE || ES3_1_PROFILE)
 		UNROLL
-		for( int CoordinateIndex = 0; CoordinateIndex < NUM_MATERIAL_TEXCOORDS; CoordinateIndex++ )
+		for( int CoordinateIndex = 0; CoordinateIndex < NUM_TEX_COORD_INTERPOLATORS; CoordinateIndex++ )
 		{
 			Result.TexCoords[CoordinateIndex] = Interpolants.TexCoords[CoordinateIndex].xy;
 		}
 	#else
 		UNROLL
-		for(int CoordinateIndex = 0;CoordinateIndex < NUM_MATERIAL_TEXCOORDS;CoordinateIndex += 2)
+		for(int CoordinateIndex = 0;CoordinateIndex < NUM_TEX_COORD_INTERPOLATORS;CoordinateIndex += 2)
 		{
 			Result.TexCoords[CoordinateIndex] = Interpolants.TexCoords[CoordinateIndex/2].xy;
 
-			if(CoordinateIndex + 1 < NUM_MATERIAL_TEXCOORDS)
+			if(CoordinateIndex + 1 < NUM_TEX_COORD_INTERPOLATORS)
 			{
 				Result.TexCoords[CoordinateIndex + 1] = Interpolants.TexCoords[CoordinateIndex/2].wz;
 			}
@@ -490,23 +490,24 @@ FVertexFactoryInterpolantsVSToPS VertexFactoryGetInterpolantsVSToPS(FVertexFacto
 	#else
 		Interpolants.SubUV0AndTexCoord0.zw = 0;
 	#endif
-#elif NUM_MATERIAL_TEXCOORDS
+#elif NUM_TEX_COORD_INTERPOLATORS
 	#if (ES2_PROFILE || ES3_1_PROFILE)
 		GetMaterialCustomizedUVs(VertexParameters, Interpolants.TexCoords);
 	#else
 		// Ensure the unused components of the last packed texture coordinate are initialized.
-		Interpolants.TexCoords[(NUM_MATERIAL_TEXCOORDS + 1) / 2 - 1] = 0;
+		Interpolants.TexCoords[(NUM_TEX_COORD_INTERPOLATORS + 1) / 2 - 1] = 0;
 
-		float2 CustomizedUVs[NUM_MATERIAL_TEXCOORDS];
+		float2 CustomizedUVs[NUM_TEX_COORD_INTERPOLATORS];
 		GetMaterialCustomizedUVs(VertexParameters, CustomizedUVs);
+		GetCustomInterpolators(VertexParameters, CustomizedUVs);
 
 		UNROLL
-		for(int CoordinateIndex = 0;CoordinateIndex < NUM_MATERIAL_TEXCOORDS;CoordinateIndex += 2)
+		for(int CoordinateIndex = 0;CoordinateIndex < NUM_TEX_COORD_INTERPOLATORS;CoordinateIndex += 2)
 		{
-				Interpolants.TexCoords[CoordinateIndex / 2].xy = CustomizedUVs[CoordinateIndex];
-			if(CoordinateIndex + 1 < NUM_MATERIAL_TEXCOORDS)
+			Interpolants.TexCoords[CoordinateIndex / 2].xy = CustomizedUVs[CoordinateIndex];
+			if(CoordinateIndex + 1 < NUM_TEX_COORD_INTERPOLATORS)
 			{
-					Interpolants.TexCoords[CoordinateIndex / 2].wz = CustomizedUVs[CoordinateIndex + 1];
+				Interpolants.TexCoords[CoordinateIndex / 2].wz = CustomizedUVs[CoordinateIndex + 1];
 			}
 		}
 	#endif 
@@ -574,12 +575,12 @@ float4 VertexFactoryGetPreviousWorldPosition(FVertexFactoryInput Input, FVertexF
 	FMaterialTessellationParameters GetMaterialTessellationParameters(FVertexFactoryInterpolantsVSToDS Interpolants, float3 CameraLocalWorldPosition)
 	{
 		FMaterialTessellationParameters	Result;
-	#if NUM_MATERIAL_TEXCOORDS
+	#if NUM_TEX_COORD_INTERPOLATORS
 		UNROLL
-		for(int CoordinateIndex = 0;CoordinateIndex < NUM_MATERIAL_TEXCOORDS;CoordinateIndex += 2)
+		for(int CoordinateIndex = 0;CoordinateIndex < NUM_TEX_COORD_INTERPOLATORS;CoordinateIndex += 2)
 		{
 			Result.TexCoords[CoordinateIndex] = Interpolants.InterpolantsVSToPS.TexCoords[CoordinateIndex/2].xy;
-			if(CoordinateIndex + 1 < NUM_MATERIAL_TEXCOORDS)
+			if(CoordinateIndex + 1 < NUM_TEX_COORD_INTERPOLATORS)
 			{
 				Result.TexCoords[CoordinateIndex + 1] = Interpolants.InterpolantsVSToPS.TexCoords[CoordinateIndex/2].wz;
 			}
@@ -621,8 +622,8 @@ float4 VertexFactoryGetPreviousWorldPosition(FVertexFactoryInput Input, FVertexF
 		TESSELLATION_INTERPOLATE_MEMBER(InterpolantsVSToPS.ParticleColor);
 #endif
 
-#if NUM_MATERIAL_TEXCOORDS
-		for (int i = 0; i < (NUM_MATERIAL_TEXCOORDS + 1) / 2; ++i)
+#if NUM_TEX_COORD_INTERPOLATORS
+		for (int i = 0; i < (NUM_TEX_COORD_INTERPOLATORS + 1) / 2; ++i)
 		{
 			TESSELLATION_INTERPOLATE_MEMBER(InterpolantsVSToPS.TexCoords[i]);
 		}
diff --git a/Engine/Shaders/MinMaterialTexCoords.usf b/Engine/Shaders/MinMaterialTexCoords.usf
--- a/Engine/Shaders/MinMaterialTexCoords.usf
+++ b/Engine/Shaders/MinMaterialTexCoords.usf
@@ -14,5 +14,6 @@
 #if MIN_MATERIAL_TEXCOORDS > NUM_MATERIAL_TEXCOORDS
 	#undef NUM_MATERIAL_TEXCOORDS
 	#define NUM_MATERIAL_TEXCOORDS MIN_MATERIAL_TEXCOORDS
+	#undef NUM_TEX_COORD_INTERPOLATORS
+	#define NUM_TEX_COORD_INTERPOLATORS (MIN_MATERIAL_TEXCOORDS + NUM_CUSTOM_VERTEX_INTERPOLATORS)
 #endif
-
diff --git a/Engine/Shaders/PostProcessCombineLUTs.usf b/Engine/Shaders/PostProcessCombineLUTs.usf
--- a/Engine/Shaders/PostProcessCombineLUTs.usf
+++ b/Engine/Shaders/PostProcessCombineLUTs.usf
@@ -276,7 +276,7 @@ void MainPS(noperspective float4 InUV : TEXCOORD0, out float4 OutColor : SV_Targ
 		LinearColor = ST2084ToLinear( LUTEncodedColor );
 	// Decode log values
 	else
-		LinearColor = LogToLin( LUTEncodedColor );
+		LinearColor = LogToLin( LUTEncodedColor ) - LogToLin( 0 );
 
 	float3 BalancedColor = WhiteBalance( LinearColor );
 	float3 GradedColor = ColorGrade( BalancedColor );
diff --git a/Engine/Shaders/PostProcessEyeAdaptation.usf b/Engine/Shaders/PostProcessEyeAdaptation.usf
--- a/Engine/Shaders/PostProcessEyeAdaptation.usf
+++ b/Engine/Shaders/PostProcessEyeAdaptation.usf
@@ -10,17 +10,18 @@
 
 // prior frame's eye adaptation exposure scale
 Texture2D EyeAdaptationTexture;
-uint2 EyeAdaptionSrcExtent;
+uint4 EyeAdaptionSrcRect;
 
 
 // Compute the average alpha stored in a 2D texture.
-float ComputeWeightedTextureAverageAlpha(Texture2D InBuffer, float slope)
+// in uint4 ROI:  region of interest.  (min, min, max, max).
+// in float Slope 
+float ComputeWeightedTextureAverageAlpha(Texture2D InBuffer, in uint4 ROI, in float Slope)
 {
-	uint Xmax = EyeAdaptionSrcExtent.x;
-	uint Ymax = EyeAdaptionSrcExtent.y;
-
-	float InvXmax = 1.f / float(Xmax);
-	float InvYmax = 1.f / float(Ymax);
+	
+	// The inverse of the Region of Interest size.
+	const float InvXMax = 1.f / float(ROI.z - ROI.x);
+	const float InvYMax = 1.f / float(ROI.w - ROI.y);
 
 	// use product of linear weight in x and y.  
 
@@ -28,12 +29,12 @@ float ComputeWeightedTextureAverageAlpha(Texture2D InBuffer, float slope)
 	float WeightTotal = 0.0f;
 
 
-	for (uint i = 0; i < Xmax; ++i)
+	for (uint i = ROI.x; i < ROI.z; ++i)
 	{
-		float IWeight = AdaptationWeight( i * InvXmax, slope);  
-		for (uint j = 0; j < Ymax; ++j)
+		float IWeight = AdaptationWeight( i * InvXMax, Slope);  
+		for (uint j = ROI.y; j < ROI.w; ++j)
 		{
-			float JWeight = AdaptationWeight( j * InvYmax, slope);
+			float JWeight = AdaptationWeight( j * InvYMax, Slope);
 			float Weight = max(IWeight * JWeight, 0.05f);
 			
 			WeightTotal += Weight;
@@ -142,7 +143,7 @@ void MainLogLuminance2ExposureScalePS(
 	float WeightSlope = EyeAdaptationParams[2].w;
 	
 	// Compute scaled Log Luminance Average 
-	float LogLumAve = ComputeWeightedTextureAverageAlpha(PostprocessInput0, WeightSlope);
+	float LogLumAve = ComputeWeightedTextureAverageAlpha(PostprocessInput0,  EyeAdaptionSrcRect, WeightSlope);
 
 	// Correct for [0,1] scaling
 	LogLumAve = (LogLumAve - EyeAdaptationParams[2].y) / EyeAdaptationParams[2].x;
diff --git a/Engine/Shaders/PostProcessTonemap.usf b/Engine/Shaders/PostProcessTonemap.usf
--- a/Engine/Shaders/PostProcessTonemap.usf
+++ b/Engine/Shaders/PostProcessTonemap.usf
@@ -73,7 +73,7 @@ half3 ColorLookupTable( half3 LinearColor )
 		LUTEncodedColor = LinearToST2084( LinearColor );
 	// Encode as log values
 	else
-		LUTEncodedColor = LinToLog( LinearColor );
+		LUTEncodedColor = LinToLog( LinearColor + LogToLin( 0 ) );
 	
 	float3 UVW = LUTEncodedColor * ((LUTSize - 1) / LUTSize) + (0.5f / LUTSize);
 
diff --git a/Engine/Shaders/TonemapCommon.usf b/Engine/Shaders/TonemapCommon.usf
--- a/Engine/Shaders/TonemapCommon.usf
+++ b/Engine/Shaders/TonemapCommon.usf
@@ -326,7 +326,7 @@ float3 ST2084ToScRGB(float3 Color, uint OutputDevice)
 	float MaxODTNits = 1000.0f;
 	float MinODTNits = 0.0001f;
 
-	if (OutputDevice == 4)
+	if (OutputDevice == 4 || OutputDevice == 6)
 	{
 		// 2000 nit display
 		MaxODTNits = 2000.0f;
diff --git a/Engine/Source/Developer/MaterialUtilities/Private/MaterialUtilities.cpp b/Engine/Source/Developer/MaterialUtilities/Private/MaterialUtilities.cpp
--- a/Engine/Source/Developer/MaterialUtilities/Private/MaterialUtilities.cpp
+++ b/Engine/Source/Developer/MaterialUtilities/Private/MaterialUtilities.cpp
@@ -181,6 +181,11 @@ struct FExportMaterialCompiler : public FProxyMaterialCompiler
 		return Compiler->PreSkinnedNormal();
 	}
 
+	virtual int32 VertexInterpolator(uint32 InterpolatorIndex) override
+	{
+		return Compiler->VertexInterpolator(InterpolatorIndex);
+	}
+
 	virtual int32 LightVector() override
 	{
 		return Compiler->LightVector();
diff --git a/Engine/Source/Developer/MeshUtilities/Private/MeshUtilities.cpp b/Engine/Source/Developer/MeshUtilities/Private/MeshUtilities.cpp
--- a/Engine/Source/Developer/MeshUtilities/Private/MeshUtilities.cpp
+++ b/Engine/Source/Developer/MeshUtilities/Private/MeshUtilities.cpp
@@ -984,24 +984,39 @@ void FMeshUtilities::GenerateSignedDistanceFieldVolumeData(
 
 			if (DistanceFieldVolume.Num() > 0)
 			{
-				const int32 UncompressedSize = DistanceFieldVolume.Num() * DistanceFieldVolume.GetTypeSize();
-				TArray<uint8> TempCompressedMemory;
-				// Compressed can be slightly larger than uncompressed
-				TempCompressedMemory.Empty(UncompressedSize * 4 / 3);
-				TempCompressedMemory.AddUninitialized(UncompressedSize * 4 / 3);
-				int32 CompressedSize = TempCompressedMemory.Num() * TempCompressedMemory.GetTypeSize();
-
-				verify(FCompression::CompressMemory(
-					(ECompressionFlags)(COMPRESS_ZLIB | COMPRESS_BiasMemory), 
-					TempCompressedMemory.GetData(), 
-					CompressedSize, 
-					DistanceFieldVolume.GetData(), 
-					UncompressedSize));
-
-				OutData.CompressedDistanceFieldVolume.Empty(CompressedSize);
-				OutData.CompressedDistanceFieldVolume.AddUninitialized(CompressedSize);
-
-				FPlatformMemory::Memcpy(OutData.CompressedDistanceFieldVolume.GetData(), TempCompressedMemory.GetData(), CompressedSize);
+				static const auto CVarCompress = IConsoleManager::Get().FindTConsoleVariableDataInt(TEXT("r.CompressMeshDistanceFields"));
+				const bool bCompress = CVarCompress->GetValueOnAnyThread() != 0;
+
+				if (bCompress)
+				{
+					const int32 UncompressedSize = DistanceFieldVolume.Num() * DistanceFieldVolume.GetTypeSize();
+					TArray<uint8> TempCompressedMemory;
+					// Compressed can be slightly larger than uncompressed
+					TempCompressedMemory.Empty(UncompressedSize * 4 / 3);
+					TempCompressedMemory.AddUninitialized(UncompressedSize * 4 / 3);
+					int32 CompressedSize = TempCompressedMemory.Num() * TempCompressedMemory.GetTypeSize();
+
+					verify(FCompression::CompressMemory(
+						(ECompressionFlags)(COMPRESS_ZLIB | COMPRESS_BiasMemory), 
+						TempCompressedMemory.GetData(), 
+						CompressedSize, 
+						DistanceFieldVolume.GetData(), 
+						UncompressedSize));
+
+					OutData.CompressedDistanceFieldVolume.Empty(CompressedSize);
+					OutData.CompressedDistanceFieldVolume.AddUninitialized(CompressedSize);
+
+					FPlatformMemory::Memcpy(OutData.CompressedDistanceFieldVolume.GetData(), TempCompressedMemory.GetData(), CompressedSize);
+				}
+				else
+				{
+					int32 CompressedSize = DistanceFieldVolume.Num() * DistanceFieldVolume.GetTypeSize();
+
+					OutData.CompressedDistanceFieldVolume.Empty(CompressedSize);
+					OutData.CompressedDistanceFieldVolume.AddUninitialized(CompressedSize);
+
+					FPlatformMemory::Memcpy(OutData.CompressedDistanceFieldVolume.GetData(), DistanceFieldVolume.GetData(), CompressedSize);
+				}
 			}
 			else
 			{
diff --git a/Engine/Source/Editor/MaterialEditor/Private/MaterialEditor.cpp b/Engine/Source/Editor/MaterialEditor/Private/MaterialEditor.cpp
--- a/Engine/Source/Editor/MaterialEditor/Private/MaterialEditor.cpp
+++ b/Engine/Source/Editor/MaterialEditor/Private/MaterialEditor.cpp
@@ -1761,6 +1761,24 @@ void FMaterialEditor::UpdateMaterialInfoList(bool bForceDisplay)
 					Line->AddToken( FTextToken::Create( FText::FromString( SamplersString ) ) );
 					Messages.Add(Line);
 				}
+
+				// Display the number of custom/user interpolators used by the material.
+				uint32 UVScalarsUsed, CustomInterpolatorScalarsUsed;
+				MaterialResource->GetUserInterpolatorUsage(UVScalarsUsed, CustomInterpolatorScalarsUsed);
+
+				if (UVScalarsUsed > 0 || CustomInterpolatorScalarsUsed > 0)
+				{
+					uint32 TotalScalars = UVScalarsUsed + CustomInterpolatorScalarsUsed;
+					uint32 MaxScalars = FMath::DivideAndRoundUp(TotalScalars, 4u) * 4;
+
+					FString InterpolatorsString = FString::Printf(TEXT("User interpolators: %u/%u Scalars (%u/4 Vectors) (TexCoords: %i, Custom: %i)"),
+						TotalScalars, MaxScalars, MaxScalars / 4, UVScalarsUsed, CustomInterpolatorScalarsUsed);
+
+					TempMaterialInfoList.Add(MakeShareable(new FMaterialInfo(InterpolatorsString, FLinearColor::Yellow)));
+					TSharedRef<FTokenizedMessage> Line = FTokenizedMessage::Create( EMessageSeverity::Info );
+					Line->AddToken(FTextToken::Create(FText::FromString(InterpolatorsString)));
+					Messages.Add(Line);
+				}
 			}
 
 			FString FeatureLevelName;
diff --git a/Engine/Source/Editor/UnrealEd/Private/Lightmass/LightmassRender.cpp b/Engine/Source/Editor/UnrealEd/Private/Lightmass/LightmassRender.cpp
--- a/Engine/Source/Editor/UnrealEd/Private/Lightmass/LightmassRender.cpp
+++ b/Engine/Source/Editor/UnrealEd/Private/Lightmass/LightmassRender.cpp
@@ -168,6 +168,11 @@ struct FLightmassMaterialCompiler : public FProxyMaterialCompiler
 		return Compiler->Constant3(0.f,0.f,1.f);
 	}
 
+	virtual int32 VertexInterpolator(uint32 InterpolatorIndex) override
+	{
+		return Compiler->Constant4(0.f,0.f,0.f,0.f);
+	}
+
 	virtual int32 RealTime(bool bPeriodic, float Period) override
 	{
 		//UE_LOG(LogLightmassRender, Log, TEXT("Lightmass material compiler has encountered RealTime... Forcing constant 0.0f."));
diff --git a/Engine/Source/Runtime/Engine/Classes/Engine/RendererSettings.h b/Engine/Source/Runtime/Engine/Classes/Engine/RendererSettings.h
--- a/Engine/Source/Runtime/Engine/Classes/Engine/RendererSettings.h
+++ b/Engine/Source/Runtime/Engine/Classes/Engine/RendererSettings.h
@@ -227,6 +227,12 @@ class ENGINE_API URendererSettings : public UDeveloperSettings
 		ToolTip = "Whether to generate a low-resolution base color texture for landscapes for rendering real-time global illumination.  This feature requires GenerateMeshDistanceFields is also enabled, and will increase mesh build times and memory usage."))
 	uint32 bGenerateLandscapeGIData : 1;
 
+	UPROPERTY(config, EditAnywhere, Category=Lighting, meta=(
+		ConsoleVariable="r.CompressMeshDistanceFields",
+		ToolTip="Whether to store mesh distance fields compressed in memory, which reduces how much memory they take, but also causes serious hitches when making new levels visible.  Only enable if your project does not stream levels in-game.  Changing this setting requires restarting the editor.",
+		ConfigRestartRequired=true))
+	uint32 bCompressMeshDistanceFields:1;
+
 	UPROPERTY(config, EditAnywhere, Category=Tessellation, meta=(
 		ConsoleVariable="r.TessellationAdaptivePixelsPerTriangle",DisplayName="Adaptive pixels per triangle",
 		ToolTip="When adaptive tessellation is enabled it will try to tessellate a mesh so that each triangle contains the specified number of pixels. The tessellation multiplier specified in the material can increase or decrease the amount of tessellation."))
diff --git a/Engine/Source/Runtime/Engine/Classes/Materials/Material.h b/Engine/Source/Runtime/Engine/Classes/Materials/Material.h
--- a/Engine/Source/Runtime/Engine/Classes/Materials/Material.h
+++ b/Engine/Source/Runtime/Engine/Classes/Materials/Material.h
@@ -1372,6 +1372,7 @@ public:
 
 	/* Returns any UMaterialExpressionCustomOutput expressions */
 	ENGINE_API void GetAllCustomOutputExpressions(TArray<class UMaterialExpressionCustomOutput*>& OutCustomOutputs) const;
+	ENGINE_API void GetAllExpressionsForCustomInterpolators(TArray<class UMaterialExpression*>& OutExpressions) const;
 
 #if WITH_EDITOR
 	/**
diff --git a/Engine/Source/Runtime/Engine/Classes/Materials/MaterialExpressionCustomOutput.h b/Engine/Source/Runtime/Engine/Classes/Materials/MaterialExpressionCustomOutput.h
--- a/Engine/Source/Runtime/Engine/Classes/Materials/MaterialExpressionCustomOutput.h
+++ b/Engine/Source/Runtime/Engine/Classes/Materials/MaterialExpressionCustomOutput.h
@@ -17,6 +17,11 @@ class UMaterialExpressionCustomOutput : public UMaterialExpression
 	virtual int32 GetNumOutputs() const { return 1; };
 	virtual FString GetFunctionName() const PURE_VIRTUAL(UMaterialExpressionCustomOutput::GetFunctionName, return TEXT("GetCustomOutput"););
 	virtual FString GetDisplayName() const { return GetFunctionName(); }
+
+#if WITH_EDITOR
+	// Allow custom outputs to generate their own source code
+	virtual bool HasCustomSourceOutput() { return false; }
+#endif
 };
 
 
diff --git a/Engine/Source/Runtime/Engine/Classes/Materials/MaterialExpressionPreSkinnedNormal.h b/Engine/Source/Runtime/Engine/Classes/Materials/MaterialExpressionPreSkinnedNormal.h
--- a/Engine/Source/Runtime/Engine/Classes/Materials/MaterialExpressionPreSkinnedNormal.h
+++ b/Engine/Source/Runtime/Engine/Classes/Materials/MaterialExpressionPreSkinnedNormal.h
@@ -17,6 +17,7 @@ class UMaterialExpressionPreSkinnedNormal : public UMaterialExpression
 	virtual int32 Compile(class FMaterialCompiler* Compiler, int32 OutputIndex) override;
 	virtual void GetCaption(TArray<FString>& OutCaptions) const override;
 	virtual void GetExpressionToolTip(TArray<FString>& OutToolTip) override;
+	virtual FText GetKeywords() const override {return FText::FromString(TEXT("normal"));}
 #endif
 	//~ End UMaterialExpression Interface
 };
diff --git a/Engine/Source/Runtime/Engine/Classes/Materials/MaterialExpressionPreSkinnedPosition.h b/Engine/Source/Runtime/Engine/Classes/Materials/MaterialExpressionPreSkinnedPosition.h
--- a/Engine/Source/Runtime/Engine/Classes/Materials/MaterialExpressionPreSkinnedPosition.h
+++ b/Engine/Source/Runtime/Engine/Classes/Materials/MaterialExpressionPreSkinnedPosition.h
@@ -17,6 +17,7 @@ class UMaterialExpressionPreSkinnedPosition : public UMaterialExpression
 	virtual int32 Compile(class FMaterialCompiler* Compiler, int32 OutputIndex) override;
 	virtual void GetCaption(TArray<FString>& OutCaptions) const override;
 	virtual void GetExpressionToolTip(TArray<FString>& OutToolTip) override;
+	virtual FText GetKeywords() const override {return FText::FromString(TEXT("position"));}
 #endif
 	//~ End UMaterialExpression Interface
 };
diff --git a/Engine/Source/Runtime/Engine/Classes/Materials/MaterialExpressionVertexInterpolator.h b/Engine/Source/Runtime/Engine/Classes/Materials/MaterialExpressionVertexInterpolator.h
new file mode 100644
--- /dev/null
+++ b/Engine/Source/Runtime/Engine/Classes/Materials/MaterialExpressionVertexInterpolator.h
@@ -0,0 +1,34 @@
+// Copyright 1998-2015 Epic Games, Inc. All Rights Reserved.
+
+#pragma once
+
+#include "Materials/MaterialExpressionCustomOutput.h"
+#include "MaterialExpressionVertexInterpolator.generated.h"
+
+UCLASS()
+class UMaterialExpressionVertexInterpolator : public UMaterialExpressionCustomOutput
+{
+	GENERATED_UCLASS_BODY()
+
+	UPROPERTY(meta = (RequiredInput = "true"))
+	FExpressionInput Input;
+
+#if WITH_EDITOR
+	virtual int32 Compile(class FMaterialCompiler* Compiler, int32 OutputIndex) override;
+	int32 CompileInput(class FMaterialCompiler* Compiler, int32 AssignedInterpolatorIndex);
+	virtual void GetCaption(TArray<FString>& OutCaptions) const override;
+	virtual uint32 GetInputType(int32 InputIndex) override { return MCT_Float4; }
+	virtual bool HasCustomSourceOutput() override { return true; }
+#endif
+
+	virtual FExpressionInput* GetInput(int32 InputIndex) override;
+	virtual FString GetFunctionName() const override { return TEXT("VertexInterpolator"); }
+	virtual FString GetInputName(int32 InputIndex) const override { return TEXT("VS"); }
+
+	int32				InterpolatorIndex;
+	EMaterialValueType	InterpolatedType;
+	int32				InterpolatorOffset;
+};
+
+
+
diff --git a/Engine/Source/Runtime/Engine/Classes/Materials/MaterialExpressionWorldPosition.h b/Engine/Source/Runtime/Engine/Classes/Materials/MaterialExpressionWorldPosition.h
--- a/Engine/Source/Runtime/Engine/Classes/Materials/MaterialExpressionWorldPosition.h
+++ b/Engine/Source/Runtime/Engine/Classes/Materials/MaterialExpressionWorldPosition.h
@@ -39,6 +39,7 @@ class UMaterialExpressionWorldPosition : public UMaterialExpression
 #if WITH_EDITOR
 	virtual int32 Compile(class FMaterialCompiler* Compiler, int32 OutputIndex) override;
 	virtual void GetCaption(TArray<FString>& OutCaptions) const override;
+	virtual FText GetKeywords() const override {return FText::FromString(TEXT("position"));}
 #endif
 	//~ End UMaterialExpression Interface
 };
diff --git a/Engine/Source/Runtime/Engine/Private/DebugViewModeMaterialProxy.cpp b/Engine/Source/Runtime/Engine/Private/DebugViewModeMaterialProxy.cpp
--- a/Engine/Source/Runtime/Engine/Private/DebugViewModeMaterialProxy.cpp
+++ b/Engine/Source/Runtime/Engine/Private/DebugViewModeMaterialProxy.cpp
@@ -237,4 +237,12 @@ void FDebugViewModeMaterialProxy::GatherCustomOutputExpressions(TArray<class UMa
 	}
 }
 
-#endif
+void FDebugViewModeMaterialProxy::GatherExpressionsForCustomInterpolators(TArray<class UMaterialExpression*>& OutExpressions) const
+{
+	if (Material)
+	{
+		Material->GetAllExpressionsForCustomInterpolators(OutExpressions);
+	}
+}
+
+#endif
\ No newline at end of file
diff --git a/Engine/Source/Runtime/Engine/Private/DebugViewModeMaterialProxy.h b/Engine/Source/Runtime/Engine/Private/DebugViewModeMaterialProxy.h
--- a/Engine/Source/Runtime/Engine/Private/DebugViewModeMaterialProxy.h
+++ b/Engine/Source/Runtime/Engine/Private/DebugViewModeMaterialProxy.h
@@ -116,6 +116,7 @@ public:
 	virtual enum EMaterialShadingModel GetShadingModel() const override;
 	virtual float GetOpacityMaskClipValue() const override;
 	virtual void GatherCustomOutputExpressions(TArray<class UMaterialExpressionCustomOutput*>& OutCustomOutputs) const override;
+	virtual void GatherExpressionsForCustomInterpolators(TArray<class UMaterialExpression*>& OutExpressions) const override;
 
 	virtual EMaterialShaderMapUsage::Type GetMaterialShaderMapUsage() const { return Usage; }
 
diff --git a/Engine/Source/Runtime/Engine/Private/DistanceFieldAtlas.cpp b/Engine/Source/Runtime/Engine/Private/DistanceFieldAtlas.cpp
--- a/Engine/Source/Runtime/Engine/Private/DistanceFieldAtlas.cpp
+++ b/Engine/Source/Runtime/Engine/Private/DistanceFieldAtlas.cpp
@@ -38,6 +38,13 @@ static TAutoConsoleVariable<int32> CVarDistField(
 	TEXT("Enabling will increase mesh build times and memory usage.  Changing this value will cause a rebuild of all static meshes."),
 	ECVF_ReadOnly);
 
+static TAutoConsoleVariable<int32> CVarCompressDistField(
+	TEXT("r.CompressMeshDistanceFields"),
+	0,	
+	TEXT("Whether to store mesh distance fields compressed in memory, which reduces how much memory they take, but also causes serious hitches when making new levels visible.  Only enable if your project does not stream levels in-game.\n")
+	TEXT("Changing this regenerates all mesh distance fields."),
+	ECVF_ReadOnly);
+
 static TAutoConsoleVariable<int32> CVarDistFieldRes(
 	TEXT("r.DistanceFields.MaxPerMeshResolution"),
 	128,	
@@ -225,6 +232,9 @@ void FDistanceFieldVolumeTextureAtlas::UpdateAllocations()
 			UE_LOG(LogStaticMesh,Log,TEXT("%s"),*GetSizeString());
 		}
 
+		static const auto CVarCompress = IConsoleManager::Get().FindTConsoleVariableDataInt(TEXT("r.CompressMeshDistanceFields"));
+		const bool bDataIsCompressed = CVarCompress->GetValueOnAnyThread() != 0;
+
 		for (int32 AllocationIndex = 0; AllocationIndex < PendingAllocations.Num(); AllocationIndex++)
 		{
 			FDistanceFieldVolumeTexture* Texture = PendingAllocations[AllocationIndex];
@@ -244,14 +254,23 @@ void FDistanceFieldVolumeTextureAtlas::UpdateAllocations()
 			const int32 FormatSize = GPixelFormats[Format].BlockBytes;
 			const int32 UncompressedSize = Size.X * Size.Y * Size.Z * FormatSize;
 
-			TArray<FFloat16> UncompressedData;
-			UncompressedData.Empty(Size.X * Size.Y * Size.Z);
-			UncompressedData.AddUninitialized(Size.X * Size.Y * Size.Z);
+			if (bDataIsCompressed)
+			{
+				TArray<FFloat16> UncompressedData;
+				UncompressedData.Empty(Size.X * Size.Y * Size.Z);
+				UncompressedData.AddUninitialized(Size.X * Size.Y * Size.Z);
 
-			verify(FCompression::UncompressMemory((ECompressionFlags)COMPRESS_ZLIB, UncompressedData.GetData(), UncompressedSize, Texture->VolumeData.CompressedDistanceFieldVolume.GetData(), Texture->VolumeData.CompressedDistanceFieldVolume.Num()));
+				verify(FCompression::UncompressMemory((ECompressionFlags)COMPRESS_ZLIB, UncompressedData.GetData(), UncompressedSize, Texture->VolumeData.CompressedDistanceFieldVolume.GetData(), Texture->VolumeData.CompressedDistanceFieldVolume.Num()));
 
-			// Update the volume texture atlas
-			RHIUpdateTexture3D(VolumeTextureRHI, 0, UpdateRegion, Size.X * FormatSize, Size.X * Size.Y * FormatSize, (const uint8*)UncompressedData.GetData());
+				// Update the volume texture atlas
+				RHIUpdateTexture3D(VolumeTextureRHI, 0, UpdateRegion, Size.X * FormatSize, Size.X * Size.Y * FormatSize, (const uint8*)UncompressedData.GetData());
+			}
+			else
+			{
+				// Update the volume texture atlas
+				check(Texture->VolumeData.CompressedDistanceFieldVolume.Num() == Size.X * Size.Y * Size.Z * FormatSize);
+				RHIUpdateTexture3D(VolumeTextureRHI, 0, UpdateRegion, Size.X * FormatSize, Size.X * Size.Y * FormatSize, (const uint8*)Texture->VolumeData.CompressedDistanceFieldVolume.GetData());
+			}
 		}
 
 		CurrentAllocations.Append(PendingAllocations);
diff --git a/Engine/Source/Runtime/Engine/Private/KismetRenderingLibrary.cpp b/Engine/Source/Runtime/Engine/Private/KismetRenderingLibrary.cpp
--- a/Engine/Source/Runtime/Engine/Private/KismetRenderingLibrary.cpp
+++ b/Engine/Source/Runtime/Engine/Private/KismetRenderingLibrary.cpp
@@ -90,7 +90,7 @@ void UKismetRenderingLibrary::DrawMaterialToRenderTarget(UObject* WorldContextOb
 
 		FName RTName = TextureRenderTarget->GetFName();
 		ENQUEUE_RENDER_COMMAND(BeginDrawEventCommand)(
-			[&RTName ,DrawMaterialToTargetEvent](FRHICommandList& RHICmdList)
+			[RTName ,DrawMaterialToTargetEvent](FRHICommandList& RHICmdList)
 			{
 				BEGIN_DRAW_EVENTF(
 					RHICmdList, 
diff --git a/Engine/Source/Runtime/Engine/Private/Materials/HLSLMaterialTranslator.h b/Engine/Source/Runtime/Engine/Private/Materials/HLSLMaterialTranslator.h
--- a/Engine/Source/Runtime/Engine/Private/Materials/HLSLMaterialTranslator.h
+++ b/Engine/Source/Runtime/Engine/Private/Materials/HLSLMaterialTranslator.h
@@ -28,6 +28,7 @@
 #include "Materials/MaterialExpressionCustom.h"
 #include "Materials/MaterialExpressionCustomOutput.h"
 #include "Materials/MaterialExpressionVectorNoise.h"
+#include "Materials/MaterialExpressionVertexInterpolator.h"
 #include "Materials/MaterialUniformExpressions.h"
 #include "ParameterCollection.h"
 #include "Materials/MaterialParameterCollection.h"
@@ -178,6 +179,11 @@ protected:
 	/** Any custom output function implementations */
 	TArray<FString> CustomOutputImplementations;
 
+	/** Custom vertex interpolators */
+	TArray<UMaterialExpressionVertexInterpolator*> CustomVertexInterpolators;
+	/** Current float-width offset for custom vertex interpolators */
+	int32 CurrentCustomVertexInterpolatorOffset;
+
 	/** Whether the translation succeeded. */
 	uint32 bSuccess : 1;
 	/** Whether the compute shader material inputs were compiled. */
@@ -251,6 +257,7 @@ public:
 	,	FeatureLevel(InFeatureLevel)
 	,	MaterialTemplateLineNumber(INDEX_NONE)
 	,	NextSymbolIndex(INDEX_NONE)
+	,	CurrentCustomVertexInterpolatorOffset(0)
 	,	bSuccess(false)
 	,	bCompileForComputeShader(false)
 	,	bUsesSceneDepth(false)
@@ -305,6 +312,45 @@ public:
 		const FGuid& MissingAttribute = FMaterialAttributeDefinitionMap::GetID(MP_MAX);
 		MaterialAttributesStack.Add(MissingAttribute);
 	}
+
+	void GatherCustomVertexInterpolators(TArray<UMaterialExpression*> Expressions)
+	{
+		for (UMaterialExpression* Expression : Expressions)
+		{
+			if (UMaterialExpressionVertexInterpolator* Interpolator = Cast<UMaterialExpressionVertexInterpolator>(Expression))
+			{
+				TArray<FShaderCodeChunk> CustomExpressionChunks;
+				CurrentScopeChunks = &CustomExpressionChunks;
+
+				int32 Ret = Interpolator->CompileInput(this, CustomVertexInterpolators.Num());
+				if (Ret != INDEX_NONE)
+				{
+					CustomVertexInterpolators.Add(Interpolator);
+				}
+
+				// Each interpolator chain must be handled as an independent compile
+				for (FMaterialFunctionCompileState& FunctionStack : FunctionStacks[SF_Vertex])
+				{
+					FunctionStack.ExpressionStack.Empty();
+					FunctionStack.ExpressionCodeMap.Empty();
+				}
+			}
+			else if (UMaterialExpressionMaterialFunctionCall* FunctionCall = Cast<UMaterialExpressionMaterialFunctionCall>(Expression))
+			{
+				if (FunctionCall->MaterialFunction)
+				{
+					FunctionCall->MaterialFunction->LinkIntoCaller(FunctionCall->FunctionInputs);
+					PushFunction(FMaterialFunctionCompileState(FunctionCall));
+
+					GatherCustomVertexInterpolators(FunctionCall->MaterialFunction->FunctionExpressions);
+
+					const FMaterialFunctionCompileState CompileState = PopFunction();
+					check(CompileState.ExpressionStack.Num() == 0);
+					FunctionCall->MaterialFunction->UnlinkFromCaller();
+				}
+			}
+		}
+	}
  
 	bool Translate()
 	{
@@ -348,13 +394,25 @@ public:
 
 			memset(Chunk, -1, sizeof(Chunk));
 
+			// Translate all custom vertex interpolators before main attributes so type information is available
+			{
+				CustomVertexInterpolators.Empty();
+				CurrentCustomVertexInterpolatorOffset = 0;
+				MaterialProperty = MP_MAX;
+				ShaderFrequency = SF_Vertex;
+
+				TArray<UMaterialExpression*> Expressions;
+				Material->GatherExpressionsForCustomInterpolators(Expressions);
+				GatherCustomVertexInterpolators(Expressions);
+			}
+
 			const EShaderFrequency NormalShaderFrequency = FMaterialAttributeDefinitionMap::GetShaderFrequency(MP_Normal);
 
 			// Normal must always be compiled first; this will ensure its chunk calculations are the first to be added
 			{
 				// Verify that start chunk is 0
 				check(SharedPropertyCodeChunks[NormalShaderFrequency].Num() == 0);
-				Chunk[MP_Normal]						= Material->CompilePropertyAndSetMaterialProperty(MP_Normal, this);
+				Chunk[MP_Normal]					= Material->CompilePropertyAndSetMaterialProperty(MP_Normal, this);
 				NormalCodeChunkEnd = SharedPropertyCodeChunks[NormalShaderFrequency].Num();
 			}
 
@@ -539,6 +597,22 @@ public:
 			// Catch any modifications to NumUserTexCoords that will not seen by customized UVs
 			check(SavedNumUserTexCoords == NumUserTexCoords);
 
+			// Finished compilation, verify final interpolator count restrictions
+			if (CurrentCustomVertexInterpolatorOffset > 0)
+			{
+				const int32 MaxNumScalars = (FeatureLevel == ERHIFeatureLevel::ES2) ? 3 * 2 : 8 * 2;
+				const int32 TotalUsedScalars = CurrentCustomVertexInterpolatorOffset + NumUserTexCoords * 2;
+
+ 				if (TotalUsedScalars > MaxNumScalars)
+				{
+					Errorf(TEXT("Maximum number of custom vertex interpolators exceeded. (%i / %i scalar values) (TexCoord: %i scalars, Custom: %i scalars)"),
+						TotalUsedScalars, MaxNumScalars, NumUserTexCoords * 2, CurrentCustomVertexInterpolatorOffset);
+				}
+			}
+
+			MaterialCompilationOutput.NumUsedUVScalars = NumUserTexCoords * 2;
+			MaterialCompilationOutput.NumUsedCustomInterpolatorScalars = CurrentCustomVertexInterpolatorOffset;
+
 			ResourcesString = TEXT("");
 
 #if HANDLE_CUSTOM_OUTPUTS_AS_MATERIAL_ATTRIBUTES
@@ -614,6 +688,11 @@ public:
 
 				for (UMaterialExpressionCustomOutput* CustomOutput : CustomOutputExpressions)
 				{
+					if (CustomOutput->HasCustomSourceOutput())
+					{
+						continue;
+					}
+
 					if (SeenCustomOutputExpressionsClases.Contains(CustomOutput->GetClass()))
 					{
 						Errorf(TEXT("The material can contain only one %s node"), *CustomOutput->GetDescription());
@@ -932,8 +1011,13 @@ public:
 		// use "MaterialTemplate.usf" to create the functions to get data (e.g. material attributes) and code (e.g. material expressions to create specular color) from C++
 		FLazyPrintf LazyPrintf(*MaterialTemplate);
 
+		const uint32 NumCustomVectors = FMath::DivideAndRoundUp((uint32)CurrentCustomVertexInterpolatorOffset, 2u);
+		const uint32 NumTexCoordVectors = NumUserTexCoords + NumCustomVectors;
+
 		LazyPrintf.PushParam(*FString::Printf(TEXT("%u"),NumUserVertexTexCoords));
 		LazyPrintf.PushParam(*FString::Printf(TEXT("%u"),NumUserTexCoords));
+		LazyPrintf.PushParam(*FString::Printf(TEXT("%u"),NumCustomVectors));
+		LazyPrintf.PushParam(*FString::Printf(TEXT("%u"),NumTexCoordVectors));
 
 		// Stores the shared shader results member declarations
 		FString PixelMembersDeclaration;
@@ -983,28 +1067,66 @@ public:
 		LazyPrintf.PushParam(*GenerateFunctionCode(MP_CustomData0));
 		LazyPrintf.PushParam(*GenerateFunctionCode(MP_CustomData1));
 
+		// Print custom texture coordinate assignments
 		FString CustomUVAssignments;
 
-			int32 LastProperty = -1;
-			for (uint32 CustomUVIndex = 0; CustomUVIndex < NumUserTexCoords; CustomUVIndex++)
+		int32 LastProperty = -1;
+		for (uint32 CustomUVIndex = 0; CustomUVIndex < NumUserTexCoords; CustomUVIndex++)
+		{
+			if (CustomUVIndex == 0)
 			{
-				if (CustomUVIndex == 0)
+				CustomUVAssignments += TranslatedCodeChunkDefinitions[MP_CustomizedUVs0 + CustomUVIndex];
+			}
+
+			if (TranslatedCodeChunkDefinitions[MP_CustomizedUVs0 + CustomUVIndex].Len() > 0)
+			{
+				if (LastProperty >= 0)
 				{
-					CustomUVAssignments += TranslatedCodeChunkDefinitions[MP_CustomizedUVs0 + CustomUVIndex];
+					check(TranslatedCodeChunkDefinitions[LastProperty].Len() == TranslatedCodeChunkDefinitions[MP_CustomizedUVs0 + CustomUVIndex].Len());
 				}
+				LastProperty = MP_CustomizedUVs0 + CustomUVIndex;
+			}
+			CustomUVAssignments += FString::Printf(TEXT("\tOutTexCoords[%u] = %s;") LINE_TERMINATOR, CustomUVIndex, *TranslatedCodeChunks[MP_CustomizedUVs0 + CustomUVIndex]);
+		}
+
+		LazyPrintf.PushParam(*CustomUVAssignments);
+
+		// Print custom vertex shader interpolator assignments
+		FString CustomInterpolatorAssignments;
+
+		for (int32 Index = 0; Index < CustomVertexInterpolators.Num(); ++Index)
+		{
+			UMaterialExpressionVertexInterpolator* Interpolator = CustomVertexInterpolators[Index];
+			check(Interpolator && Interpolator->InterpolatorIndex != INDEX_NONE);
+			check(Interpolator->InterpolatedType & MCT_Float);
 
-				if (TranslatedCodeChunkDefinitions[MP_CustomizedUVs0 + CustomUVIndex].Len() > 0)
+			if (Interpolator->InterpolatorOffset != INDEX_NONE)
+			{
+				const EMaterialValueType Type = Interpolator->InterpolatedType == MCT_Float ? MCT_Float1 : Interpolator->InterpolatedType;
+				const TCHAR* Swizzle[2] = { TEXT("x"), TEXT("y") };
+				const int32 Offset = Interpolator->InterpolatorOffset;
+
+				// Note: We reference the UV define directly to avoid having to pre-accumulate UV counts before property translation
+				CustomInterpolatorAssignments += FString::Printf(TEXT("\tOutTexCoords[%i + NUM_MATERIAL_TEXCOORDS].%s = VertexInterpolator%i(Parameters).x;") LINE_TERMINATOR, Offset/2, Swizzle[Offset%2], Index);
+				
+				if (Type >= MCT_Float2)
 				{
-					if (LastProperty >= 0)
+					CustomInterpolatorAssignments += FString::Printf(TEXT("\tOutTexCoords[%i + NUM_MATERIAL_TEXCOORDS].%s = VertexInterpolator%i(Parameters).y;") LINE_TERMINATOR, (Offset+1)/2, Swizzle[(Offset+1)%2], Index);
+
+					if (Type >= MCT_Float3)
 					{
-						check(TranslatedCodeChunkDefinitions[LastProperty].Len() == TranslatedCodeChunkDefinitions[MP_CustomizedUVs0 + CustomUVIndex].Len());
+						CustomInterpolatorAssignments += FString::Printf(TEXT("\tOutTexCoords[%i + NUM_MATERIAL_TEXCOORDS].%s = VertexInterpolator%i(Parameters).z;") LINE_TERMINATOR, (Offset+2)/2, Swizzle[(Offset+2)%2], Index);
+
+						if (Type == MCT_Float4)
+						{
+							CustomInterpolatorAssignments += FString::Printf(TEXT("\tOutTexCoords[%i + NUM_MATERIAL_TEXCOORDS].%s = VertexInterpolator%i(Parameters).w;") LINE_TERMINATOR, (Offset+3)/2, Swizzle[(Offset+3)%2], Index);
+						}
 					}
-					LastProperty = MP_CustomizedUVs0 + CustomUVIndex;
 				}
-				CustomUVAssignments += FString::Printf(TEXT("\tOutTexCoords[%u] = %s;") LINE_TERMINATOR, CustomUVIndex, *TranslatedCodeChunks[MP_CustomizedUVs0 + CustomUVIndex]);
 			}
+		}
 
-		LazyPrintf.PushParam(*CustomUVAssignments);
+		LazyPrintf.PushParam(*CustomInterpolatorAssignments);
 
 		// Initializers required for Normal
 		LazyPrintf.PushParam(*TranslatedCodeChunkDefinitions[MP_Normal]);
@@ -3640,6 +3762,65 @@ protected:
 		return AddInlinedCodeChunk(MCT_Float3,TEXT("Parameters.PreSkinnedNormal"));
 	}
 
+	virtual int32 VertexInterpolator(uint32 InterpolatorIndex) override
+	{
+		if (ShaderFrequency != SF_Pixel)
+		{
+			return Errorf(TEXT("Custom interpolator outputs only available in pixel shaders."));
+		}
+		else if (InterpolatorIndex >= (uint32)CustomVertexInterpolators.Num())
+		{
+			return Errorf(TEXT("Invalid custom interpolator index."));
+		}
+
+		UMaterialExpressionVertexInterpolator* Interpolator = CustomVertexInterpolators[InterpolatorIndex];
+		check(Interpolator && Interpolator->InterpolatorIndex == InterpolatorIndex);
+		check(Interpolator->InterpolatedType & MCT_Float);
+		check(Interpolator->InterpolatorOffset == INDEX_NONE);
+
+		// Assign interpolator offset and accumulate size
+		int32 InterpolatorSize = 0;
+		switch (Interpolator->InterpolatedType)
+		{
+		case MCT_Float4:	InterpolatorSize = 4; break;
+		case MCT_Float3:	InterpolatorSize = 3; break;
+		case MCT_Float2:	InterpolatorSize = 2; break;
+		default:			InterpolatorSize = 1;
+		};
+
+		Interpolator->InterpolatorOffset = CurrentCustomVertexInterpolatorOffset;
+		CurrentCustomVertexInterpolatorOffset += InterpolatorSize;
+
+		// Copy interpolated data from pixel parameters to local
+		const EMaterialValueType Type = Interpolator->InterpolatedType == MCT_Float ? MCT_Float1 : Interpolator->InterpolatedType;
+		const TCHAR* TypeName = HLSLTypeString(Type);
+		const TCHAR* Swizzle[2] = { TEXT("x"), TEXT("y") };
+		const int32 Offset = Interpolator->InterpolatorOffset;
+	
+		// Note: We reference the UV define directly to avoid having to pre-accumulate UV counts before property translation
+		FString GetValueCode = FString::Printf(TEXT("%s(Parameters.TexCoords[%i + NUM_MATERIAL_TEXCOORDS].%s"), TypeName, (Offset/2), Swizzle[Offset%2]);
+
+		if (Type >= MCT_Float2)
+		{
+			GetValueCode += FString::Printf(TEXT(", Parameters.TexCoords[%i + NUM_MATERIAL_TEXCOORDS].%s"), (Offset+1)/2, Swizzle[(Offset+1)%2]);
+
+			if (Type >= MCT_Float3)
+			{
+				GetValueCode += FString::Printf(TEXT(", Parameters.TexCoords[%i + NUM_MATERIAL_TEXCOORDS].%s"), (Offset+2)/2, Swizzle[(Offset+2)%2]);
+
+				if (Type == MCT_Float4)
+				{
+					GetValueCode += FString::Printf(TEXT(", Parameters.TexCoords[%i + NUM_MATERIAL_TEXCOORDS].%s"), (Offset+3)/2, Swizzle[(Offset+3)%2]);
+				}
+			}
+		}
+
+		GetValueCode.Append(TEXT(")"));
+
+		int32 RetCode = AddCodeChunk(Type, *GetValueCode);
+		return RetCode;
+	}
+
 	virtual int32 Add(int32 A,int32 B) override
 	{
 		if(A == INDEX_NONE || B == INDEX_NONE)
diff --git a/Engine/Source/Runtime/Engine/Private/Materials/Material.cpp b/Engine/Source/Runtime/Engine/Private/Materials/Material.cpp
--- a/Engine/Source/Runtime/Engine/Private/Materials/Material.cpp
+++ b/Engine/Source/Runtime/Engine/Private/Materials/Material.cpp
@@ -32,6 +32,7 @@
 #include "Materials/MaterialExpressionTextureSample.h"
 #include "Materials/MaterialExpressionTextureSampleParameter.h"
 #include "Materials/MaterialExpressionVectorParameter.h"
+#include "Materials/MaterialExpressionVertexInterpolator.h"
 #include "Materials/MaterialExpressionSceneColor.h"
 #include "SceneManagement.h"
 #include "Materials/MaterialUniformExpressions.h"
@@ -154,6 +155,11 @@ void FMaterialResource::GatherCustomOutputExpressions(TArray<UMaterialExpression
 	Material->GetAllCustomOutputExpressions(OutCustomOutputs);
 }
 
+void FMaterialResource::GatherExpressionsForCustomInterpolators(TArray<UMaterialExpression*>& OutExpressions) const
+{
+	Material->GetAllExpressionsForCustomInterpolators(OutExpressions);
+}
+
 void FMaterialResource::GetShaderMapId(EShaderPlatform Platform, FMaterialShaderMapId& OutId) const
 {
 	FMaterial::GetShaderMapId(Platform, OutId);
@@ -3897,6 +3903,19 @@ void UMaterial::GetAllCustomOutputExpressions(TArray<class UMaterialExpressionCu
 	}
 }
 
+void UMaterial::GetAllExpressionsForCustomInterpolators(TArray<class UMaterialExpression*>& OutExpressions) const
+{
+	for (UMaterialExpression* Expression : Expressions)
+	{
+		if (Expression &&
+			(Expression->IsA(UMaterialExpressionVertexInterpolator::StaticClass()) ||
+			Expression->IsA(UMaterialExpressionMaterialFunctionCall::StaticClass())) )
+		{
+				OutExpressions.Add(Expression);
+		}
+	}
+}
+
 #if WITH_EDITOR
 bool UMaterial::GetAllReferencedExpressions(TArray<UMaterialExpression*>& OutExpressions, class FStaticParameterSet* InStaticParameterSet)
 {
diff --git a/Engine/Source/Runtime/Engine/Private/Materials/MaterialExpressions.cpp b/Engine/Source/Runtime/Engine/Private/Materials/MaterialExpressions.cpp
--- a/Engine/Source/Runtime/Engine/Private/Materials/MaterialExpressions.cpp
+++ b/Engine/Source/Runtime/Engine/Private/Materials/MaterialExpressions.cpp
@@ -176,6 +176,7 @@
 #include "Materials/MaterialExpressionTwoSidedSign.h"
 #include "Materials/MaterialExpressionVectorNoise.h"
 #include "Materials/MaterialExpressionVertexColor.h"
+#include "Materials/MaterialExpressionVertexInterpolator.h"
 #include "Materials/MaterialExpressionVertexNormalWS.h"
 #include "Materials/MaterialExpressionViewProperty.h"
 #include "Materials/MaterialExpressionViewSize.h"
@@ -4708,9 +4709,6 @@ int32 UMaterialExpressionBlendMaterialAttributes::Compile(class FMaterialCompile
 {
 	const FGuid AttributeID = Compiler->GetMaterialAttribute();
 
-	int32 ResultA = A.CompileWithDefault(Compiler, AttributeID);
-	int32 ResultB = B.CompileWithDefault(Compiler, AttributeID);
-
 	// Blending is optional, can skip on a per-node basis
 	EMaterialAttributeBlend::Type BlendType;
 	EShaderFrequency AttributeFrequency = FMaterialAttributeDefinitionMap::GetShaderFrequency(AttributeID);
@@ -4727,13 +4725,15 @@ int32 UMaterialExpressionBlendMaterialAttributes::Compile(class FMaterialCompile
 
 	switch (BlendType)
 	{
-	case EMaterialAttributeBlend::UseA: return ResultA;
-	case EMaterialAttributeBlend::UseB: return ResultB;
+	case EMaterialAttributeBlend::UseA: return A.CompileWithDefault(Compiler, AttributeID);
+	case EMaterialAttributeBlend::UseB: return B.CompileWithDefault(Compiler, AttributeID);
 	default:
 		check(BlendType == EMaterialAttributeBlend::Blend);
 	}
 
 	// Allow custom blends or fallback to standard interpolation
+	int32 ResultA = A.CompileWithDefault(Compiler, AttributeID);
+	int32 ResultB = B.CompileWithDefault(Compiler, AttributeID);
 	int32 ResultAlpha = Alpha.Compile(Compiler);
 
 	MaterialAttributeBlendFunction BlendFunction = FMaterialAttributeDefinitionMap::GetBlendFunction(AttributeID);
@@ -12032,6 +12032,87 @@ FExpressionInput* UMaterialExpressionClearCoatNormalCustomOutput::GetInput(int32
 	return &Input;
 }
 
+///////////////////////////////////////////////////////////////////////////////
+// Vertex to pixel interpolated data handler
+///////////////////////////////////////////////////////////////////////////////
+
+UMaterialExpressionVertexInterpolator::UMaterialExpressionVertexInterpolator(const FObjectInitializer& ObjectInitializer)
+: Super(ObjectInitializer)
+{
+	// Structure to hold one-time initialization
+	struct FConstructorStatics
+	{
+		FText NAME_Utility;
+		FConstructorStatics()
+			: NAME_Utility(LOCTEXT("Utility", "Utility"))
+		{
+		}
+	};
+	static FConstructorStatics ConstructorStatics;
+
+#if WITH_EDITORONLY_DATA
+	MenuCategories.Add(ConstructorStatics.NAME_Utility);
+#endif
+
+	Outputs.Reset();
+	Outputs.Add(FExpressionOutput(TEXT("PS"), 0, 0, 0, 0, 0));
+	bShowOutputNameOnPin = true;
+
+	InterpolatorIndex = INDEX_NONE;
+	InterpolatedType = MCT_Unknown;
+	InterpolatorOffset = INDEX_NONE;
+}
+
+#if WITH_EDITOR
+int32 UMaterialExpressionVertexInterpolator::Compile(class FMaterialCompiler* Compiler, int32 OutputIndex)
+{
+	if (Input.Expression)
+	{
+		if (InterpolatorIndex == INDEX_NONE || InterpolatorOffset != INDEX_NONE)
+		{
+			return Compiler->Errorf(TEXT("Failed interpolator pre-compile or recursively post-compiled."));
+		}
+		else
+		{
+			return Compiler->VertexInterpolator(InterpolatorIndex);
+		}
+	}
+	else
+	{
+		return CompilerError(Compiler, TEXT("Input missing"));
+	}
+}
+
+int32 UMaterialExpressionVertexInterpolator::CompileInput(class FMaterialCompiler* Compiler, int32 AssignedInterpolatorIndex)
+{
+	int32 Ret = INDEX_NONE;
+	InterpolatorIndex = INDEX_NONE;
+	InterpolatedType = MCT_Unknown;
+	InterpolatorOffset = INDEX_NONE;
+
+	if (Input.Expression)
+	{
+		int32 InternalCode = Input.Compile(Compiler);
+		Compiler->CustomOutput(this, AssignedInterpolatorIndex, InternalCode);
+		InterpolatorIndex = AssignedInterpolatorIndex;
+		InterpolatedType = Compiler->GetType(InternalCode);
+		Ret = InternalCode;
+	}
+
+	return Ret;
+}
+
+void UMaterialExpressionVertexInterpolator::GetCaption(TArray<FString>& OutCaptions) const
+{
+	OutCaptions.Add(FString(TEXT("VertexInterpolator")));
+}
+#endif // WITH_EDITOR
+
+FExpressionInput* UMaterialExpressionVertexInterpolator::GetInput(int32 InputIndex)
+{
+	return &Input;
+}
+
 ///////////////////////////////////////////////////////////////////////////////
 // UMaterialExpressionrAtmosphericLightVector
 ///////////////////////////////////////////////////////////////////////////////
diff --git a/Engine/Source/Runtime/Engine/Private/Materials/MaterialShared.cpp b/Engine/Source/Runtime/Engine/Private/Materials/MaterialShared.cpp
--- a/Engine/Source/Runtime/Engine/Private/Materials/MaterialShared.cpp
+++ b/Engine/Source/Runtime/Engine/Private/Materials/MaterialShared.cpp
@@ -2108,6 +2108,17 @@ int32 FMaterialResource::GetSamplerUsage() const
 	return -1;
 }
 
+void FMaterialResource::GetUserInterpolatorUsage(uint32& NumUsedUVScalars, uint32& NumUsedCustomInterpolatorScalars) const
+{
+	NumUsedUVScalars = NumUsedCustomInterpolatorScalars = 0;
+
+	if (const FMaterialShaderMap* ShaderMap = GetGameThreadShaderMap())
+	{
+		NumUsedUVScalars = ShaderMap->GetNumUsedUVScalars();
+		NumUsedCustomInterpolatorScalars = ShaderMap->GetNumUsedCustomInterpolatorScalars();
+	}
+}
+
 FString FMaterialResource::GetMaterialUsageDescription() const
 {
 	check(Material);
diff --git a/Engine/Source/Runtime/Engine/Private/StaticMesh.cpp b/Engine/Source/Runtime/Engine/Private/StaticMesh.cpp
--- a/Engine/Source/Runtime/Engine/Private/StaticMesh.cpp
+++ b/Engine/Source/Runtime/Engine/Private/StaticMesh.cpp
@@ -1092,9 +1092,13 @@ static FString BuildDistanceFieldDerivedDataKey(const FString& InMeshKey)
 	const float VoxelDensity = CVarDensity->GetValueOnAnyThread();
 	const FString VoxelDensityString = VoxelDensity == .1f ? TEXT("") : FString(TEXT("_%.3f"), VoxelDensity);
 
+	static const auto CVarCompress = IConsoleManager::Get().FindTConsoleVariableDataInt(TEXT("r.CompressMeshDistanceFields"));
+	const bool bCompress = CVarCompress->GetValueOnAnyThread() != 0;
+	const FString CompressString = bCompress ? TEXT("") : TEXT("_uc");
+
 	return FDerivedDataCacheInterface::BuildCacheKey(
 		TEXT("DIST"),
-		*FString::Printf(TEXT("%s_%s%s%s"), *InMeshKey, DISTANCEFIELD_DERIVEDDATA_VER, *PerMeshMaxString, *VoxelDensityString),
+		*FString::Printf(TEXT("%s_%s%s%s%s"), *InMeshKey, DISTANCEFIELD_DERIVEDDATA_VER, *PerMeshMaxString, *VoxelDensityString, *CompressString),
 		TEXT(""));
 }
 
diff --git a/Engine/Source/Runtime/Engine/Public/MaterialCompiler.h b/Engine/Source/Runtime/Engine/Public/MaterialCompiler.h
--- a/Engine/Source/Runtime/Engine/Public/MaterialCompiler.h
+++ b/Engine/Source/Runtime/Engine/Public/MaterialCompiler.h
@@ -186,6 +186,8 @@ public:
 	virtual int32 PreSkinnedPosition() = 0;
 	virtual int32 PreSkinnedNormal() = 0;
 
+	virtual int32 VertexInterpolator(uint32 InterpolatorIndex) = 0;
+
 #if WITH_EDITOR
 	virtual int32 MaterialBakingWorldPosition() = 0;
 #endif
diff --git a/Engine/Source/Runtime/Engine/Public/MaterialShared.h b/Engine/Source/Runtime/Engine/Public/MaterialShared.h
--- a/Engine/Source/Runtime/Engine/Public/MaterialShared.h
+++ b/Engine/Source/Runtime/Engine/Public/MaterialShared.h
@@ -306,6 +306,8 @@ class FMaterialCompilationOutput
 {
 public:
 	FMaterialCompilationOutput() :
+		NumUsedUVScalars(0),
+		NumUsedCustomInterpolatorScalars(0),
 		bRequiresSceneColorCopy(false),
 		bNeedsSceneTextures(false),
 		bUsesEyeAdaptation(false),
@@ -321,9 +323,13 @@ public:
 
 	FUniformExpressionSet UniformExpressionSet;
 
-	/** 
-	 * Indicates whether the material uses scene color. 
-	 */
+	/** Number of used custom UV scalars. */
+	uint8 NumUsedUVScalars;
+
+	/** Number of used custom vertex interpolation scalars. */
+	uint8 NumUsedCustomInterpolatorScalars;
+
+	/** Indicates whether the material uses scene color. */
 	bool bRequiresSceneColorCopy;
 
 	/** true if the material needs the scenetexture lookups. */
@@ -730,6 +736,8 @@ public:
 	bool ModifiesMeshPosition() const { return MaterialCompilationOutput.bModifiesMeshPosition; }
 	bool UsesPixelDepthOffset() const { return MaterialCompilationOutput.bUsesPixelDepthOffset; }
 	bool UsesSceneDepthLookup() const { return MaterialCompilationOutput.bUsesSceneDepthLookup; }
+	uint32 GetNumUsedUVScalars() const { return MaterialCompilationOutput.NumUsedUVScalars; }
+	uint32 GetNumUsedCustomInterpolatorScalars() const { return MaterialCompilationOutput.NumUsedCustomInterpolatorScalars; }
 
 	bool IsValidForRendering() const
 	{
@@ -1239,6 +1247,9 @@ protected:
 	/* Gather any UMaterialExpressionCustomOutput expressions they can be compiled in turn */
 	virtual void GatherCustomOutputExpressions(TArray<class UMaterialExpressionCustomOutput*>& OutCustomOutputs) const {}
 
+	/* Gather any UMaterialExpressionCustomOutput expressions in the material and referenced function calls */
+	virtual void GatherExpressionsForCustomInterpolators(TArray<class UMaterialExpression*>& OutExpressions) const {}
+
 	/** Returns the index to the Expression in the Expressions array, or -1 if not found. */
 	int32 FindExpression(const TArray<TRefCountPtr<FMaterialUniformExpressionTexture> >&Expressions, const FMaterialUniformExpressionTexture &Expression);
 
@@ -1606,6 +1617,7 @@ public:
 
 	/** Returns the number of samplers used in this material, or -1 if the material does not have a valid shader map (compile error or still compiling). */
 	ENGINE_API int32 GetSamplerUsage() const;
+	ENGINE_API void GetUserInterpolatorUsage(uint32& NumUsedUVScalars, uint32& NumUsedCustomInterpolatorScalars) const;
 
 	ENGINE_API virtual FString GetMaterialUsageDescription() const override;
 
@@ -1728,6 +1740,7 @@ protected:
 
 	/* Gives the material a chance to compile any custom output nodes it has added */
 	ENGINE_API virtual void GatherCustomOutputExpressions(TArray<class UMaterialExpressionCustomOutput*>& OutCustomOutputs) const override;
+	ENGINE_API virtual void GatherExpressionsForCustomInterpolators(TArray<class UMaterialExpression*>& OutExpressions) const override;
 	ENGINE_API virtual bool HasVertexPositionOffsetConnected() const override;
 	ENGINE_API virtual bool HasPixelDepthOffsetConnected() const override;
 	ENGINE_API virtual bool HasMaterialAttributesConnected() const override;
diff --git a/Engine/Source/Runtime/RHI/Public/RHIDefinitions.h b/Engine/Source/Runtime/RHI/Public/RHIDefinitions.h
--- a/Engine/Source/Runtime/RHI/Public/RHIDefinitions.h
+++ b/Engine/Source/Runtime/RHI/Public/RHIDefinitions.h
@@ -80,13 +80,23 @@ enum ERenderQueryType
 enum { MAX_TEXTURE_MIP_COUNT = 14 };
 
 /** The maximum number of vertex elements which can be used by a vertex declaration. */
-enum { MaxVertexElementCount = 16 };
+enum
+{
+	MaxVertexElementCount = 16,
+	MaxVertexElementCount_NumBits = 4,
+};
+static_assert(MaxVertexElementCount <= (1 << MaxVertexElementCount_NumBits), "MaxVertexElementCount will not fit on MaxVertexElementCount_NumBits");
 
 /** The alignment in bytes between elements of array shader parameters. */
 enum { ShaderArrayElementAlignBytes = 16 };
 
 /** The number of render-targets that may be simultaneously written to. */
-enum { MaxSimultaneousRenderTargets = 8 };
+enum
+{
+	MaxSimultaneousRenderTargets = 8,
+	MaxSimultaneousRenderTargets_NumBits = 3,
+};
+static_assert(MaxSimultaneousRenderTargets <= (1 << MaxSimultaneousRenderTargets_NumBits), "MaxSimultaneousRenderTargets will not fit on MaxSimultaneousRenderTargets_NumBits");
 
 /** The number of UAVs that may be simultaneously bound to a shader. */
 enum { MaxSimultaneousUAVs = 8 };
@@ -131,7 +141,11 @@ enum ESamplerFilter
 	SF_Trilinear,
 	SF_AnisotropicPoint,
 	SF_AnisotropicLinear,
+
+	ESamplerFilter_Num,
+	ESamplerFilter_NumBits = 3,
 };
+static_assert(ESamplerFilter_Num <= (1 << ESamplerFilter_NumBits), "ESamplerFilter_Num will not fit on ESamplerFilter_NumBits");
 
 enum ESamplerAddressMode
 {
@@ -139,8 +153,12 @@ enum ESamplerAddressMode
 	AM_Clamp,
 	AM_Mirror,
 	/** Not supported on all platforms */
-	AM_Border
+	AM_Border,
+
+	ESamplerAddressMode_Num,
+	ESamplerAddressMode_NumBits = 2,
 };
+static_assert(ESamplerAddressMode_Num <= (1 << ESamplerAddressMode_NumBits), "ESamplerAddressMode_Num will not fit on ESamplerAddressMode_NumBits");
 
 enum ESamplerCompareFunction
 {
@@ -152,15 +170,23 @@ enum ERasterizerFillMode
 {
 	FM_Point,
 	FM_Wireframe,
-	FM_Solid
+	FM_Solid,
+
+	ERasterizerFillMode_Num,
+	ERasterizerFillMode_NumBits = 2,
 };
+static_assert(ERasterizerFillMode_Num <= (1 << ERasterizerFillMode_NumBits), "ERasterizerFillMode_Num will not fit on ERasterizerFillMode_NumBits");
 
 enum ERasterizerCullMode
 {
 	CM_None,
 	CM_CW,
-	CM_CCW
+	CM_CCW,
+
+	ERasterizerCullMode_Num,
+	ERasterizerCullMode_NumBits = 2,
 };
+static_assert(ERasterizerCullMode_Num <= (1 << ERasterizerCullMode_NumBits), "ERasterizerCullMode_Num will not fit on ERasterizerCullMode_NumBits");
 
 enum EColorWriteMask
 {
@@ -174,6 +200,8 @@ enum EColorWriteMask
 	CW_RGBA  = CW_RED | CW_GREEN | CW_BLUE | CW_ALPHA,
 	CW_RG    = CW_RED | CW_GREEN,
 	CW_BA    = CW_BLUE | CW_ALPHA,
+
+	EColorWriteMask_NumBits = 4,
 };
 
 enum ECompareFunction
@@ -187,12 +215,16 @@ enum ECompareFunction
 	CF_Never,
 	CF_Always,
 
+	ECompareFunction_Num,
+	ECompareFunction_NumBits = 3,
+
 	// Utility enumerations
 	CF_DepthNearOrEqual		= (((int32)ERHIZBuffer::IsInverted != 0) ? CF_GreaterEqual : CF_LessEqual),
 	CF_DepthNear			= (((int32)ERHIZBuffer::IsInverted != 0) ? CF_Greater : CF_Less),
 	CF_DepthFartherOrEqual	= (((int32)ERHIZBuffer::IsInverted != 0) ? CF_LessEqual : CF_GreaterEqual),
 	CF_DepthFarther			= (((int32)ERHIZBuffer::IsInverted != 0) ? CF_Less : CF_Greater),
 };
+static_assert(ECompareFunction_Num <= (1 << ECompareFunction_NumBits), "ECompareFunction_Num will not fit on ECompareFunction_NumBits");
 
 enum EStencilOp
 {
@@ -203,8 +235,12 @@ enum EStencilOp
 	SO_SaturatedDecrement,
 	SO_Invert,
 	SO_Increment,
-	SO_Decrement
+	SO_Decrement,
+
+	EStencilOp_Num,
+	EStencilOp_NumBits = 3,
 };
+static_assert(EStencilOp_Num <= (1 << EStencilOp_NumBits), "EStencilOp_Num will not fit on EStencilOp_NumBits");
 
 enum EBlendOperation
 {
@@ -213,7 +249,11 @@ enum EBlendOperation
 	BO_Min,
 	BO_Max,
 	BO_ReverseSubtract,
+
+	EBlendOperation_Num,
+	EBlendOperation_NumBits = 3,
 };
+static_assert(EBlendOperation_Num <= (1 << EBlendOperation_NumBits), "EBlendOperation_Num will not fit on EBlendOperation_NumBits");
 
 enum EBlendFactor
 {
@@ -228,8 +268,12 @@ enum EBlendFactor
 	BF_DestColor,
 	BF_InverseDestColor,
 	BF_ConstantBlendFactor,
-	BF_InverseConstantBlendFactor
+	BF_InverseConstantBlendFactor,
+
+	EBlendFactor_Num,
+	EBlendFactor_NumBits = 4,
 };
+static_assert(EBlendFactor_Num <= (1 << EBlendFactor_NumBits), "EBlendFactor_Num will not fit on EBlendFactor_NumBits");
 
 enum EVertexElementType
 {
@@ -253,8 +297,11 @@ enum EVertexElementType
 	VET_UShort2N,		// 16 bit word normalized to (value/65535.0,value/65535.0,0,0,1)
 	VET_UShort4N,		// 4 X 16 bit word unsigned, normalized 
 	VET_URGB10A2N,		// 10 bit r, g, b and 2 bit a normalized to (value/1023.0f, value/1023.0f, value/1023.0f, value/3.0f)
-	VET_MAX
+	VET_MAX,
+
+	VET_NumBits = 5,
 };
+static_assert(VET_MAX <= (1 << VET_NumBits), "VET_MAX will not fit on VET_NumBits");
 
 enum ECubeFace
 {
@@ -289,8 +336,12 @@ enum EUniformBufferBaseType
 	UBMT_SRV,
 	UBMT_UAV,
 	UBMT_SAMPLER,
-	UBMT_TEXTURE
+	UBMT_TEXTURE,
+
+	EUniformBufferBaseType_Num,
+	EUniformBufferBaseType_NumBits = 4,
 };
+static_assert(EUniformBufferBaseType_Num <= (1 << EUniformBufferBaseType_NumBits), "EUniformBufferBaseType_Num will not fit on EUniformBufferBaseType_NumBits");
 
 struct FRHIResourceTableEntry
 {
@@ -366,10 +417,12 @@ enum class EPrimitiveTopologyType : uint8
 	Patch,
 	Line,
 	Point,
-	Quad,
+	//Quad,
+
 	Num,
-	NumBits = 3
+	NumBits = 2,
 };
+static_assert((uint32)EPrimitiveTopologyType::Num <= (1 << (uint32)EPrimitiveTopologyType::NumBits), "EPrimitiveTopologyType::Num will not fit on EPrimitiveTopologyType::NumBits");
 
 enum EPrimitiveType
 {
@@ -413,7 +466,6 @@ enum EPrimitiveType
 	PT_Num,
 	PT_NumBits = 6
 };
-
 static_assert(PT_Num <= (1 << 8), "EPrimitiveType doesn't fit in a byte");
 static_assert(PT_Num <= (1 << PT_NumBits), "PT_NumBits is too small");
 
@@ -584,7 +636,12 @@ enum class ERenderTargetLoadAction
 	ENoAction,
 	ELoad,
 	EClear,
+
+	Num,
+	NumBits = 2,
 };
+static_assert((uint32)ERenderTargetLoadAction::Num <= (1 << (uint32)ERenderTargetLoadAction::NumBits), "ERenderTargetLoadAction::Num will not fit on ERenderTargetLoadAction::NumBits");
+
 
 /**
  * Action to take when a rendertarget is unset or at the end of a pass. 
@@ -594,7 +651,11 @@ enum class ERenderTargetStoreAction
 	ENoAction,
 	EStore,
 	EMultisampleResolve,
+
+	Num,
+	NumBits = 2,
 };
+static_assert((uint32)ERenderTargetStoreAction::Num <= (1 << (uint32)ERenderTargetStoreAction::NumBits), "ERenderTargetStoreAction::Num will not fit on ERenderTargetStoreAction::NumBits");
 
 /**
  * Common render target use cases
diff --git a/Engine/Source/Runtime/Renderer/Private/CapsuleShadowRendering.cpp b/Engine/Source/Runtime/Renderer/Private/CapsuleShadowRendering.cpp
--- a/Engine/Source/Runtime/Renderer/Private/CapsuleShadowRendering.cpp
+++ b/Engine/Source/Runtime/Renderer/Private/CapsuleShadowRendering.cpp
@@ -1383,7 +1383,7 @@ void FDeferredShadingSceneRenderer::RenderCapsuleShadowsForMovableSkylight(FRHIC
 		if (bAnyViewsUseCapsuleShadows)
 		{
 			TRefCountPtr<IPooledRenderTarget> NewBentNormal;
-			AllocateOrReuseAORenderTarget(RHICmdList, NewBentNormal, TEXT("CapsuleBentNormal"), true);
+			AllocateOrReuseAORenderTarget(RHICmdList, NewBentNormal, TEXT("CapsuleBentNormal"), PF_FloatRGBA);
 
 			for (int32 ViewIndex = 0; ViewIndex < Views.Num(); ViewIndex++)
 			{
diff --git a/Engine/Source/Runtime/Renderer/Private/DistanceFieldLightingPost.cpp b/Engine/Source/Runtime/Renderer/Private/DistanceFieldLightingPost.cpp
--- a/Engine/Source/Runtime/Renderer/Private/DistanceFieldLightingPost.cpp
+++ b/Engine/Source/Runtime/Renderer/Private/DistanceFieldLightingPost.cpp
@@ -11,23 +11,6 @@
 #include "DistanceFieldAmbientOcclusion.h"
 #include "CompositionLighting/PostProcessAmbientOcclusion.h"
 
-int32 GAOFillGaps = 1;
-FAutoConsoleVariableRef CVarAOFillGaps(
-	TEXT("r.AOFillGaps"),
-	GAOFillGaps,
-	TEXT("Whether to fill in pixels using a screen space filter that had no valid world space interpolation weight from surface cache samples.\n")
-	TEXT("This is needed whenever r.AOMinLevel is not 0."),
-	ECVF_RenderThreadSafe
-	);
-
-int32 GAOFillGapsHighQuality = 1;
-FAutoConsoleVariableRef CVarAOFillGapsHighQuality(
-	TEXT("r.AOFillGapsHighQuality"),
-	GAOFillGapsHighQuality,
-	TEXT("Whether to use the higher quality gap filling method that does a 5x5 gather, or the cheaper method that just looks up the 4 grid samples."),
-	ECVF_RenderThreadSafe
-	);
-
 int32 GAOUseHistory = 1;
 FAutoConsoleVariableRef CVarAOUseHistory(
 	TEXT("r.AOUseHistory"),
@@ -52,6 +35,15 @@ FAutoConsoleVariableRef CVarAOHistoryWeight(
 	ECVF_RenderThreadSafe
 	);
 
+float GAOHistoryMinConfidenceScale = .8f;
+FAutoConsoleVariableRef CVarAOHistoryMinConfidenceScale(
+	TEXT("r.AOHistoryMinConfidenceScale"),
+	GAOHistoryMinConfidenceScale,
+	TEXT("Minimum amount that confidence can scale down the history weight. Pixels whose AO value was interpolated from foreground onto background incorrectly have a confidence of 0.\n")
+	TEXT("At a value of 1, confidence is effectively disabled.  Lower values increase the convergence speed of AO history for pixels with low confidence, but introduce jittering (history is thrown away)."),
+	ECVF_RenderThreadSafe
+	);
+
 float GAOHistoryDistanceThreshold = 30;
 FAutoConsoleVariableRef CVarAOHistoryDistanceThreshold(
 	TEXT("r.AOHistoryDistanceThreshold"),
@@ -69,218 +61,9 @@ FAutoConsoleVariableRef CVarAOViewFadeDistanceScale(
 	);
 
 template<bool bSupportIrradiance>
-class TDistanceFieldAOCombinePS : public FGlobalShader
-{
-	DECLARE_SHADER_TYPE(TDistanceFieldAOCombinePS, Global);
-public:
-
-	static bool ShouldCache(EShaderPlatform Platform)
-	{
-		return IsFeatureLevelSupported(Platform, ERHIFeatureLevel::SM5) && DoesPlatformSupportDistanceFieldAO(Platform);
-	}
-
-	static void ModifyCompilationEnvironment(EShaderPlatform Platform, FShaderCompilerEnvironment& OutEnvironment)
-	{
-		OutEnvironment.SetDefine(TEXT("DOWNSAMPLE_FACTOR"), GAODownsampleFactor);
-		OutEnvironment.SetDefine(TEXT("SUPPORT_IRRADIANCE"), bSupportIrradiance);
-	}
-
-	/** Default constructor. */
-	TDistanceFieldAOCombinePS() {}
-
-	/** Initialization constructor. */
-	TDistanceFieldAOCombinePS(const ShaderMetaType::CompiledShaderInitializerType& Initializer)
-		: FGlobalShader(Initializer)
-	{
-		DeferredParameters.Bind(Initializer.ParameterMap);
-		AOParameters.Bind(Initializer.ParameterMap);
-		BentNormalAOTexture.Bind(Initializer.ParameterMap,TEXT("BentNormalAOTexture"));
-		BentNormalAOSampler.Bind(Initializer.ParameterMap,TEXT("BentNormalAOSampler"));
-		IrradianceTexture.Bind(Initializer.ParameterMap,TEXT("IrradianceTexture"));
-		IrradianceSampler.Bind(Initializer.ParameterMap,TEXT("IrradianceSampler"));
-		DistanceFieldNormalTexture.Bind(Initializer.ParameterMap, TEXT("DistanceFieldNormalTexture"));
-		DistanceFieldNormalSampler.Bind(Initializer.ParameterMap, TEXT("DistanceFieldNormalSampler"));
-		DistanceFadeScale.Bind(Initializer.ParameterMap, TEXT("DistanceFadeScale"));
-		SelfOcclusionReplacement.Bind(Initializer.ParameterMap, TEXT("SelfOcclusionReplacement"));
-	}
-
-	void SetParameters(
-		FRHICommandList& RHICmdList, 
-		const FSceneView& View, 
-		FSceneRenderTargetItem& InBentNormalTexture, 
-		IPooledRenderTarget* InIrradianceTexture, 
-		FSceneRenderTargetItem& DistanceFieldNormal, 
-		const FDistanceFieldAOParameters& Parameters)
-	{
-		const FPixelShaderRHIParamRef ShaderRHI = GetPixelShader();
-
-		FGlobalShader::SetParameters(RHICmdList, ShaderRHI, View);
-		DeferredParameters.Set(RHICmdList, ShaderRHI, View);
-		AOParameters.Set(RHICmdList, ShaderRHI, Parameters);
-
-		SetTextureParameter(RHICmdList, ShaderRHI, BentNormalAOTexture, BentNormalAOSampler, TStaticSamplerState<SF_Point>::GetRHI(), InBentNormalTexture.ShaderResourceTexture);
-
-		if (IrradianceTexture.IsBound())
-		{
-			SetTextureParameter(RHICmdList, ShaderRHI, IrradianceTexture, IrradianceSampler, TStaticSamplerState<SF_Point>::GetRHI(), InIrradianceTexture->GetRenderTargetItem().ShaderResourceTexture);
-		}
-
-		SetTextureParameter(
-			RHICmdList,
-			ShaderRHI,
-			DistanceFieldNormalTexture,
-			DistanceFieldNormalSampler,
-			TStaticSamplerState<SF_Point,AM_Clamp,AM_Clamp,AM_Clamp>::GetRHI(),
-			DistanceFieldNormal.ShaderResourceTexture
-			);
-
-		const float DistanceFadeScaleValue = 1.0f / ((1.0f - GAOViewFadeDistanceScale) * GetMaxAOViewDistance());
-		SetShaderValue(RHICmdList, ShaderRHI, DistanceFadeScale, DistanceFadeScaleValue);
-
-		extern float GVPLSelfOcclusionReplacement;
-		SetShaderValue(RHICmdList, ShaderRHI, SelfOcclusionReplacement, GVPLSelfOcclusionReplacement);
-	}
-
-	// FShader interface.
-	virtual bool Serialize(FArchive& Ar) override
-	{
-		bool bShaderHasOutdatedParameters = FGlobalShader::Serialize(Ar);
-		Ar << DeferredParameters;
-		Ar << AOParameters;
-		Ar << BentNormalAOTexture;
-		Ar << BentNormalAOSampler;
-		Ar << IrradianceTexture;
-		Ar << IrradianceSampler;
-		Ar << DistanceFieldNormalTexture;
-		Ar << DistanceFieldNormalSampler;
-		Ar << DistanceFadeScale;
-		Ar << SelfOcclusionReplacement;
-		return bShaderHasOutdatedParameters;
-	}
-
-private:
-
-	FDeferredPixelShaderParameters DeferredParameters;
-	FAOParameters AOParameters;
-	FShaderResourceParameter BentNormalAOTexture;
-	FShaderResourceParameter BentNormalAOSampler;
-	FShaderResourceParameter IrradianceTexture;
-	FShaderResourceParameter IrradianceSampler;
-	FShaderResourceParameter DistanceFieldNormalTexture;
-	FShaderResourceParameter DistanceFieldNormalSampler;
-	FShaderParameter DistanceFadeScale;
-	FShaderParameter SelfOcclusionReplacement;
-};
-
-IMPLEMENT_SHADER_TYPE(template<>,TDistanceFieldAOCombinePS<true>,TEXT("DistanceFieldLightingPost"),TEXT("AOCombinePS"),SF_Pixel);
-IMPLEMENT_SHADER_TYPE(template<>,TDistanceFieldAOCombinePS<false>,TEXT("DistanceFieldLightingPost"),TEXT("AOCombinePS"),SF_Pixel);
-
-template<bool bSupportIrradiance, bool bHighQuality>
-class TFillGapsPS : public FGlobalShader
-{
-	DECLARE_SHADER_TYPE(TFillGapsPS, Global);
-public:
-
-	static bool ShouldCache(EShaderPlatform Platform)
-	{
-		return IsFeatureLevelSupported(Platform, ERHIFeatureLevel::SM5) && DoesPlatformSupportDistanceFieldAO(Platform);
-	}
-
-	static void ModifyCompilationEnvironment(EShaderPlatform Platform, FShaderCompilerEnvironment& OutEnvironment)
-	{
-		OutEnvironment.SetDefine(TEXT("DOWNSAMPLE_FACTOR"), GAODownsampleFactor);
-		OutEnvironment.SetDefine(TEXT("SUPPORT_IRRADIANCE"), bSupportIrradiance);
-		OutEnvironment.SetDefine(TEXT("HIGH_QUALITY_FILL_GAPS"), bHighQuality);
-	}
-
-	/** Default constructor. */
-	TFillGapsPS() {}
-
-	/** Initialization constructor. */
-	TFillGapsPS(const ShaderMetaType::CompiledShaderInitializerType& Initializer)
-		: FGlobalShader(Initializer)
-	{
-		BentNormalAOTexture.Bind(Initializer.ParameterMap, TEXT("BentNormalAOTexture"));
-		BentNormalAOSampler.Bind(Initializer.ParameterMap, TEXT("BentNormalAOSampler"));
-		IrradianceTexture.Bind(Initializer.ParameterMap, TEXT("IrradianceTexture"));
-		IrradianceSampler.Bind(Initializer.ParameterMap, TEXT("IrradianceSampler"));
-		BentNormalAOTexelSize.Bind(Initializer.ParameterMap, TEXT("BentNormalAOTexelSize"));
-		MinDownsampleFactorToBaseLevel.Bind(Initializer.ParameterMap, TEXT("MinDownsampleFactorToBaseLevel"));
-	}
-
-	void SetParameters(FRHICommandList& RHICmdList, const FSceneView& View, FSceneRenderTargetItem& DistanceFieldAOBentNormal2, IPooledRenderTarget* DistanceFieldIrradiance2)
-	{
-		const FPixelShaderRHIParamRef ShaderRHI = GetPixelShader();
-
-		FGlobalShader::SetParameters(RHICmdList, ShaderRHI, View);
-
-		SetTextureParameter(
-			RHICmdList,
-			ShaderRHI,
-			BentNormalAOTexture,
-			BentNormalAOSampler,
-			TStaticSamplerState<SF_Bilinear,AM_Clamp,AM_Clamp,AM_Clamp>::GetRHI(),
-			DistanceFieldAOBentNormal2.ShaderResourceTexture
-			);
-
-		if (IrradianceTexture.IsBound())
-		{
-			SetTextureParameter(
-				RHICmdList,
-				ShaderRHI,
-				IrradianceTexture,
-				IrradianceSampler,
-				TStaticSamplerState<SF_Bilinear,AM_Clamp,AM_Clamp,AM_Clamp>::GetRHI(),
-				DistanceFieldIrradiance2->GetRenderTargetItem().ShaderResourceTexture
-				);
-		}
-		FSceneRenderTargets& SceneContext = FSceneRenderTargets::Get(RHICmdList);
-
-		const FIntPoint DownsampledBufferSize(SceneContext.GetBufferSizeXY() / FIntPoint(GAODownsampleFactor, GAODownsampleFactor));
-		const FVector2D BaseLevelTexelSizeValue(1.0f / DownsampledBufferSize.X, 1.0f / DownsampledBufferSize.Y);
-		SetShaderValue(RHICmdList, ShaderRHI, BentNormalAOTexelSize, BaseLevelTexelSizeValue);
-
-		extern int32 GAOMinLevel;
-		extern int32 GAOPowerOfTwoBetweenLevels;
-		const float MinDownsampleFactor = 1 << (GAOMinLevel * GAOPowerOfTwoBetweenLevels); 
-		SetShaderValue(RHICmdList, ShaderRHI, MinDownsampleFactorToBaseLevel, MinDownsampleFactor);
-	}
-	// FShader interface.
-	virtual bool Serialize(FArchive& Ar) override
-	{
-		bool bShaderHasOutdatedParameters = FGlobalShader::Serialize(Ar);
-		Ar << BentNormalAOTexture;
-		Ar << BentNormalAOSampler;
-		Ar << IrradianceTexture;
-		Ar << IrradianceSampler;
-		Ar << BentNormalAOTexelSize;
-		Ar << MinDownsampleFactorToBaseLevel;
-		return bShaderHasOutdatedParameters;
-	}
-
-private:
-
-	FShaderResourceParameter BentNormalAOTexture;
-	FShaderResourceParameter BentNormalAOSampler;
-	FShaderResourceParameter IrradianceTexture;
-	FShaderResourceParameter IrradianceSampler;
-	FShaderParameter BentNormalAOTexelSize;
-	FShaderParameter MinDownsampleFactorToBaseLevel;
-};
-
-#define IMPLEMENT_FILLGAPS_PS_TYPE(bSupportIrradiance,bHighQuality) \
-	typedef TFillGapsPS<bSupportIrradiance, bHighQuality> TFillGapsPS##bSupportIrradiance##bHighQuality; \
-	IMPLEMENT_SHADER_TYPE(template<>,TFillGapsPS##bSupportIrradiance##bHighQuality,TEXT("DistanceFieldLightingPost"),TEXT("FillGapsPS"),SF_Pixel); 
-
-IMPLEMENT_FILLGAPS_PS_TYPE(true, true);
-IMPLEMENT_FILLGAPS_PS_TYPE(true, false);
-IMPLEMENT_FILLGAPS_PS_TYPE(false, true);
-IMPLEMENT_FILLGAPS_PS_TYPE(false, false);
-
-template<bool bSupportIrradiance>
-class TUpdateHistoryPS : public FGlobalShader
+class TUpdateHistoryDepthRejectionPS : public FGlobalShader
 {
-	DECLARE_SHADER_TYPE(TUpdateHistoryPS, Global);
+	DECLARE_SHADER_TYPE(TUpdateHistoryDepthRejectionPS, Global);
 public:
 
 	static bool ShouldCache(EShaderPlatform Platform)
@@ -295,22 +78,25 @@ public:
 	}
 
 	/** Default constructor. */
-	TUpdateHistoryPS() {}
+	TUpdateHistoryDepthRejectionPS() {}
 
 	/** Initialization constructor. */
-	TUpdateHistoryPS(const ShaderMetaType::CompiledShaderInitializerType& Initializer)
+	TUpdateHistoryDepthRejectionPS(const ShaderMetaType::CompiledShaderInitializerType& Initializer)
 		: FGlobalShader(Initializer)
 	{
 		DeferredParameters.Bind(Initializer.ParameterMap);
 		BentNormalAOTexture.Bind(Initializer.ParameterMap, TEXT("BentNormalAOTexture"));
+		ConfidenceTexture.Bind(Initializer.ParameterMap, TEXT("ConfidenceTexture"));
 		BentNormalAOSampler.Bind(Initializer.ParameterMap, TEXT("BentNormalAOSampler"));
 		BentNormalHistoryTexture.Bind(Initializer.ParameterMap, TEXT("BentNormalHistoryTexture"));
+		ConfidenceHistoryTexture.Bind(Initializer.ParameterMap, TEXT("ConfidenceHistoryTexture"));
 		BentNormalHistorySampler.Bind(Initializer.ParameterMap, TEXT("BentNormalHistorySampler"));
 		IrradianceTexture.Bind(Initializer.ParameterMap, TEXT("IrradianceTexture"));
 		IrradianceSampler.Bind(Initializer.ParameterMap, TEXT("IrradianceSampler"));
 		IrradianceHistoryTexture.Bind(Initializer.ParameterMap, TEXT("IrradianceHistoryTexture"));
 		IrradianceHistorySampler.Bind(Initializer.ParameterMap, TEXT("IrradianceHistorySampler"));
 		HistoryWeight.Bind(Initializer.ParameterMap, TEXT("HistoryWeight"));
+		AOHistoryMinConfidenceScale.Bind(Initializer.ParameterMap, TEXT("AOHistoryMinConfidenceScale"));
 		HistoryDistanceThreshold.Bind(Initializer.ParameterMap, TEXT("HistoryDistanceThreshold"));
 		UseHistoryFilter.Bind(Initializer.ParameterMap, TEXT("UseHistoryFilter"));
 		VelocityTexture.Bind(Initializer.ParameterMap, TEXT("VelocityTexture"));
@@ -324,8 +110,10 @@ public:
 		const FSceneView& View, 
 		FSceneRenderTargetItem& DistanceFieldNormal, 
 		FSceneRenderTargetItem& BentNormalHistoryTextureValue, 
+		FSceneRenderTargetItem& ConfidenceHistoryTextureValue,
 		TRefCountPtr<IPooledRenderTarget>* IrradianceHistoryRT, 
 		FSceneRenderTargetItem& DistanceFieldAOBentNormal, 
+		FSceneRenderTargetItem& DistanceFieldAOConfidence, 
 		IPooledRenderTarget* DistanceFieldIrradiance,
 		IPooledRenderTarget* VelocityTextureValue)
 	{
@@ -343,15 +131,33 @@ public:
 			DistanceFieldAOBentNormal.ShaderResourceTexture
 			);
 
+		SetTextureParameter(
+			RHICmdList,
+			ShaderRHI,
+			ConfidenceTexture,
+			BentNormalAOSampler,
+			TStaticSamplerState<SF_Bilinear,AM_Clamp,AM_Clamp,AM_Clamp>::GetRHI(),
+			DistanceFieldAOConfidence.ShaderResourceTexture
+			);
+
 		SetTextureParameter(
 			RHICmdList,
 			ShaderRHI,
 			BentNormalHistoryTexture,
 			BentNormalHistorySampler,
-			TStaticSamplerState<SF_Point,AM_Clamp,AM_Clamp,AM_Clamp>::GetRHI(),
+			TStaticSamplerState<SF_Bilinear,AM_Clamp,AM_Clamp,AM_Clamp>::GetRHI(),
 			BentNormalHistoryTextureValue.ShaderResourceTexture
 			);
 
+		SetTextureParameter(
+			RHICmdList,
+			ShaderRHI,
+			ConfidenceHistoryTexture,
+			BentNormalHistorySampler,
+			TStaticSamplerState<SF_Bilinear,AM_Clamp,AM_Clamp,AM_Clamp>::GetRHI(),
+			ConfidenceHistoryTextureValue.ShaderResourceTexture
+			);
+
 		SetTextureParameter(
 			RHICmdList,
 			ShaderRHI,
@@ -388,6 +194,7 @@ public:
 		SetUniformBufferParameter(RHICmdList, ShaderRHI, GetUniformBufferParameter<FCameraMotionParameters>(), CreateCameraMotionParametersUniformBuffer(View));
 
 		SetShaderValue(RHICmdList, ShaderRHI, HistoryWeight, GAOHistoryWeight);
+		SetShaderValue(RHICmdList, ShaderRHI, AOHistoryMinConfidenceScale, GAOHistoryMinConfidenceScale);
 		SetShaderValue(RHICmdList, ShaderRHI, HistoryDistanceThreshold, GAOHistoryDistanceThreshold);
 		SetShaderValue(RHICmdList, ShaderRHI, UseHistoryFilter, (GAOHistoryStabilityPass ? 1.0f : 0.0f));
 
@@ -406,14 +213,17 @@ public:
 		bool bShaderHasOutdatedParameters = FGlobalShader::Serialize(Ar);
 		Ar << DeferredParameters;
 		Ar << BentNormalAOTexture;
+		Ar << ConfidenceTexture;
 		Ar << BentNormalAOSampler;
 		Ar << BentNormalHistoryTexture;
+		Ar << ConfidenceHistoryTexture;
 		Ar << BentNormalHistorySampler;
 		Ar << IrradianceTexture;
 		Ar << IrradianceSampler;
 		Ar << IrradianceHistoryTexture;
 		Ar << IrradianceHistorySampler;
 		Ar << HistoryWeight;
+		Ar << AOHistoryMinConfidenceScale;
 		Ar << HistoryDistanceThreshold;
 		Ar << UseHistoryFilter;
 		Ar << VelocityTexture;
@@ -427,14 +237,17 @@ private:
 
 	FDeferredPixelShaderParameters DeferredParameters;
 	FShaderResourceParameter BentNormalAOTexture;
+	FShaderResourceParameter ConfidenceTexture;
 	FShaderResourceParameter BentNormalAOSampler;
 	FShaderResourceParameter BentNormalHistoryTexture;
+	FShaderResourceParameter ConfidenceHistoryTexture;
 	FShaderResourceParameter BentNormalHistorySampler;
 	FShaderResourceParameter IrradianceTexture;
 	FShaderResourceParameter IrradianceSampler;
 	FShaderResourceParameter IrradianceHistoryTexture;
 	FShaderResourceParameter IrradianceHistorySampler;
 	FShaderParameter HistoryWeight;
+	FShaderParameter AOHistoryMinConfidenceScale;
 	FShaderParameter HistoryDistanceThreshold;
 	FShaderParameter UseHistoryFilter;
 	FShaderResourceParameter VelocityTexture;
@@ -443,8 +256,8 @@ private:
 	FShaderResourceParameter DistanceFieldNormalSampler;
 };
 
-IMPLEMENT_SHADER_TYPE(template<>,TUpdateHistoryPS<true>,TEXT("DistanceFieldLightingPost"),TEXT("UpdateHistoryPS"),SF_Pixel);
-IMPLEMENT_SHADER_TYPE(template<>,TUpdateHistoryPS<false>,TEXT("DistanceFieldLightingPost"),TEXT("UpdateHistoryPS"),SF_Pixel);
+IMPLEMENT_SHADER_TYPE(template<>,TUpdateHistoryDepthRejectionPS<true>,TEXT("DistanceFieldLightingPost"),TEXT("UpdateHistoryDepthRejectionPS"),SF_Pixel);
+IMPLEMENT_SHADER_TYPE(template<>,TUpdateHistoryDepthRejectionPS<false>,TEXT("DistanceFieldLightingPost"),TEXT("UpdateHistoryDepthRejectionPS"),SF_Pixel);
 
 
 template<bool bSupportIrradiance>
@@ -472,6 +285,7 @@ public:
 		: FGlobalShader(Initializer)
 	{
 		BentNormalAOTexture.Bind(Initializer.ParameterMap, TEXT("BentNormalAOTexture"));
+		ConfidenceTexture.Bind(Initializer.ParameterMap, TEXT("ConfidenceTexture"));
 		BentNormalAOSampler.Bind(Initializer.ParameterMap, TEXT("BentNormalAOSampler"));
 		IrradianceTexture.Bind(Initializer.ParameterMap, TEXT("IrradianceTexture"));
 		IrradianceSampler.Bind(Initializer.ParameterMap, TEXT("IrradianceSampler"));
@@ -486,6 +300,7 @@ public:
 		const FSceneView& View, 
 		FSceneRenderTargetItem& DistanceFieldNormal, 
 		FSceneRenderTargetItem& BentNormalHistoryTextureValue, 
+		FSceneRenderTargetItem& ConfidenceHistoryTextureValue, 
 		IPooledRenderTarget* IrradianceHistoryRT)
 	{
 		const FPixelShaderRHIParamRef ShaderRHI = GetPixelShader();
@@ -501,6 +316,15 @@ public:
 			BentNormalHistoryTextureValue.ShaderResourceTexture
 			);
 
+		SetTextureParameter(
+			RHICmdList,
+			ShaderRHI,
+			ConfidenceTexture,
+			BentNormalAOSampler,
+			TStaticSamplerState<SF_Bilinear,AM_Clamp,AM_Clamp,AM_Clamp>::GetRHI(),
+			ConfidenceHistoryTextureValue.ShaderResourceTexture
+			);
+
 		SetTextureParameter(
 			RHICmdList,
 			ShaderRHI,
@@ -535,6 +359,7 @@ public:
 	{
 		bool bShaderHasOutdatedParameters = FGlobalShader::Serialize(Ar);
 		Ar << BentNormalAOTexture;
+		Ar << ConfidenceTexture;
 		Ar << BentNormalAOSampler;
 		Ar << IrradianceTexture;
 		Ar << IrradianceSampler;
@@ -548,6 +373,7 @@ public:
 private:
 
 	FShaderResourceParameter BentNormalAOTexture;
+	FShaderResourceParameter ConfidenceTexture;
 	FShaderResourceParameter BentNormalAOSampler;
 	FShaderResourceParameter IrradianceTexture;
 	FShaderResourceParameter IrradianceSampler;
@@ -561,13 +387,13 @@ IMPLEMENT_SHADER_TYPE(template<>,TFilterHistoryPS<true>,TEXT("DistanceFieldLight
 IMPLEMENT_SHADER_TYPE(template<>,TFilterHistoryPS<false>,TEXT("DistanceFieldLightingPost"),TEXT("FilterHistoryPS"),SF_Pixel);
 
 
-void AllocateOrReuseAORenderTarget(FRHICommandList& RHICmdList, TRefCountPtr<IPooledRenderTarget>& Target, const TCHAR* Name, bool bWithAlphaOrFP16Precision)
+void AllocateOrReuseAORenderTarget(FRHICommandList& RHICmdList, TRefCountPtr<IPooledRenderTarget>& Target, const TCHAR* Name, EPixelFormat Format)
 {
 	if (!Target)
 	{
 		FIntPoint BufferSize = GetBufferSizeForAO();
 
-		FPooledRenderTargetDesc Desc(FPooledRenderTargetDesc::Create2DDesc(BufferSize, bWithAlphaOrFP16Precision ? PF_FloatRGBA : PF_FloatRGB, FClearValueBinding::None, TexCreate_None, TexCreate_RenderTargetable | TexCreate_UAV, false));
+		FPooledRenderTargetDesc Desc(FPooledRenderTargetDesc::Create2DDesc(BufferSize, Format, FClearValueBinding::None, TexCreate_None, TexCreate_RenderTargetable | TexCreate_UAV, false));
 		Desc.AutoWritable = false;
 		GRenderTargetPool.FindFreeElement(RHICmdList, Desc, Target, Name);
 	}
@@ -577,14 +403,17 @@ void UpdateHistory(
 	FRHICommandList& RHICmdList,
 	const FViewInfo& View, 
 	const TCHAR* BentNormalHistoryRTName,
+	const TCHAR* ConfidenceHistoryRTName,
 	const TCHAR* IrradianceHistoryRTName,
 	IPooledRenderTarget* VelocityTexture,
 	FSceneRenderTargetItem& DistanceFieldNormal,
 	/** Contains last frame's history, if non-NULL.  This will be updated with the new frame's history. */
 	TRefCountPtr<IPooledRenderTarget>* BentNormalHistoryState,
+	TRefCountPtr<IPooledRenderTarget>* ConfidenceHistoryState,
 	TRefCountPtr<IPooledRenderTarget>* IrradianceHistoryState,
 	/** Source */
 	TRefCountPtr<IPooledRenderTarget>& BentNormalSource, 
+	TRefCountPtr<IPooledRenderTarget>& ConfidenceSource,
 	TRefCountPtr<IPooledRenderTarget>& IrradianceSource, 
 	/** Output of Temporal Reprojection for the next step in the pipeline. */
 	TRefCountPtr<IPooledRenderTarget>& BentNormalHistoryOutput,
@@ -603,21 +432,25 @@ void UpdateHistory(
 		{
 			// Reuse a render target from the pool with a consistent name, for vis purposes
 			TRefCountPtr<IPooledRenderTarget> NewBentNormalHistory;
-			AllocateOrReuseAORenderTarget(RHICmdList, NewBentNormalHistory, BentNormalHistoryRTName, true);
+			AllocateOrReuseAORenderTarget(RHICmdList, NewBentNormalHistory, BentNormalHistoryRTName, PF_FloatRGBA);
+
+			TRefCountPtr<IPooledRenderTarget> NewConfidenceHistory;
+			AllocateOrReuseAORenderTarget(RHICmdList, NewConfidenceHistory, ConfidenceHistoryRTName, PF_G8);
 
 			TRefCountPtr<IPooledRenderTarget> NewIrradianceHistory;
 
 			if (bUseDistanceFieldGI)
 			{
-				AllocateOrReuseAORenderTarget(RHICmdList, NewIrradianceHistory, IrradianceHistoryRTName, false);
+				AllocateOrReuseAORenderTarget(RHICmdList, NewIrradianceHistory, IrradianceHistoryRTName, PF_FloatRGB);
 			}
 
 			SCOPED_DRAW_EVENT(RHICmdList, UpdateHistory);
 
 			{
-				FTextureRHIParamRef RenderTargets[2] =
+				FTextureRHIParamRef RenderTargets[3] =
 				{
 					NewBentNormalHistory->GetRenderTargetItem().TargetableTexture,
+					NewConfidenceHistory->GetRenderTargetItem().TargetableTexture,
 					bUseDistanceFieldGI ? NewIrradianceHistory->GetRenderTargetItem().TargetableTexture : NULL,
 				};
 
@@ -632,19 +465,19 @@ void UpdateHistory(
 
 				if (bUseDistanceFieldGI)
 				{
-					TShaderMapRef<TUpdateHistoryPS<true> > PixelShader(View.ShaderMap);
+					TShaderMapRef<TUpdateHistoryDepthRejectionPS<true> > PixelShader(View.ShaderMap);
 
 					static FGlobalBoundShaderState BoundShaderState;
 					SetGlobalBoundShaderState(RHICmdList, View.GetFeatureLevel(), BoundShaderState, GFilterVertexDeclaration.VertexDeclarationRHI, *VertexShader, *PixelShader);
-					PixelShader->SetParameters(RHICmdList, View, DistanceFieldNormal, (*BentNormalHistoryState)->GetRenderTargetItem(), IrradianceHistoryState, BentNormalSource->GetRenderTargetItem(), IrradianceSource, VelocityTexture);
+					PixelShader->SetParameters(RHICmdList, View, DistanceFieldNormal, (*BentNormalHistoryState)->GetRenderTargetItem(), (*ConfidenceHistoryState)->GetRenderTargetItem(), IrradianceHistoryState, BentNormalSource->GetRenderTargetItem(), ConfidenceSource->GetRenderTargetItem(), IrradianceSource, VelocityTexture);
 				}
 				else
 				{
-					TShaderMapRef<TUpdateHistoryPS<false> > PixelShader(View.ShaderMap);
+					TShaderMapRef<TUpdateHistoryDepthRejectionPS<false> > PixelShader(View.ShaderMap);
 
 					static FGlobalBoundShaderState BoundShaderState;
 					SetGlobalBoundShaderState(RHICmdList, View.GetFeatureLevel(), BoundShaderState, GFilterVertexDeclaration.VertexDeclarationRHI, *VertexShader, *PixelShader);
-					PixelShader->SetParameters(RHICmdList, View, DistanceFieldNormal, (*BentNormalHistoryState)->GetRenderTargetItem(), IrradianceHistoryState, BentNormalSource->GetRenderTargetItem(), IrradianceSource, VelocityTexture);
+					PixelShader->SetParameters(RHICmdList, View, DistanceFieldNormal, (*BentNormalHistoryState)->GetRenderTargetItem(), (*ConfidenceHistoryState)->GetRenderTargetItem(), IrradianceHistoryState, BentNormalSource->GetRenderTargetItem(), ConfidenceSource->GetRenderTargetItem(), IrradianceSource, VelocityTexture);
 				}
 
 				VertexShader->SetParameters(RHICmdList, View);
@@ -660,6 +493,7 @@ void UpdateHistory(
 					*VertexShader);
 
 				RHICmdList.CopyToResolveTarget(NewBentNormalHistory->GetRenderTargetItem().TargetableTexture, NewBentNormalHistory->GetRenderTargetItem().ShaderResourceTexture, false, FResolveParams());
+				RHICmdList.CopyToResolveTarget(NewConfidenceHistory->GetRenderTargetItem().TargetableTexture, NewConfidenceHistory->GetRenderTargetItem().ShaderResourceTexture, false, FResolveParams());
 			
 				if (bUseDistanceFieldGI)
 				{
@@ -675,22 +509,26 @@ void UpdateHistory(
 				if (HistoryDesc.Extent != SceneContext.GetBufferSizeXY() / FIntPoint(GAODownsampleFactor, GAODownsampleFactor))
 				{
 					GRenderTargetPool.FreeUnusedResource(*BentNormalHistoryState);
+					GRenderTargetPool.FreeUnusedResource(*ConfidenceHistoryState);
 					*BentNormalHistoryState = NULL;
+					*ConfidenceHistoryState = NULL;
 					// Update the view state's render target reference with the new history
-					AllocateOrReuseAORenderTarget(RHICmdList, *BentNormalHistoryState, BentNormalHistoryRTName, true);
+					AllocateOrReuseAORenderTarget(RHICmdList, *BentNormalHistoryState, BentNormalHistoryRTName, PF_FloatRGBA);
+					AllocateOrReuseAORenderTarget(RHICmdList, *ConfidenceHistoryState, ConfidenceHistoryRTName, PF_G8);
 
 					if (bUseDistanceFieldGI)
 					{
 						GRenderTargetPool.FreeUnusedResource(*IrradianceHistoryState);
 						*IrradianceHistoryState = NULL;
-						AllocateOrReuseAORenderTarget(RHICmdList, *IrradianceHistoryState, IrradianceHistoryRTName, false);
+						AllocateOrReuseAORenderTarget(RHICmdList, *IrradianceHistoryState, IrradianceHistoryRTName, PF_FloatRGB);
 					}
 				}
 
 				{
-					FTextureRHIParamRef RenderTargets[2] =
+					FTextureRHIParamRef RenderTargets[3] =
 					{
 						(*BentNormalHistoryState)->GetRenderTargetItem().TargetableTexture,
+						(*ConfidenceHistoryState)->GetRenderTargetItem().TargetableTexture,
 						bUseDistanceFieldGI ? (*IrradianceHistoryState)->GetRenderTargetItem().TargetableTexture : NULL,
 					};
 
@@ -709,7 +547,7 @@ void UpdateHistory(
 
 						static FGlobalBoundShaderState BoundShaderState;
 						SetGlobalBoundShaderState(RHICmdList, View.GetFeatureLevel(), BoundShaderState, GFilterVertexDeclaration.VertexDeclarationRHI, *VertexShader, *PixelShader);
-						PixelShader->SetParameters(RHICmdList, View, DistanceFieldNormal, NewBentNormalHistory->GetRenderTargetItem(), NewIrradianceHistory);
+						PixelShader->SetParameters(RHICmdList, View, DistanceFieldNormal, NewBentNormalHistory->GetRenderTargetItem(), NewConfidenceHistory->GetRenderTargetItem(), NewIrradianceHistory);
 					}
 					else
 					{
@@ -717,7 +555,7 @@ void UpdateHistory(
 
 						static FGlobalBoundShaderState BoundShaderState;
 						SetGlobalBoundShaderState(RHICmdList, View.GetFeatureLevel(), BoundShaderState, GFilterVertexDeclaration.VertexDeclarationRHI, *VertexShader, *PixelShader);
-						PixelShader->SetParameters(RHICmdList, View, DistanceFieldNormal, NewBentNormalHistory->GetRenderTargetItem(), NewIrradianceHistory);
+						PixelShader->SetParameters(RHICmdList, View, DistanceFieldNormal, NewBentNormalHistory->GetRenderTargetItem(), NewConfidenceHistory->GetRenderTargetItem(), NewIrradianceHistory);
 					}
 
 					VertexShader->SetParameters(RHICmdList, View);
@@ -733,7 +571,8 @@ void UpdateHistory(
 						*VertexShader);
 
 					RHICmdList.CopyToResolveTarget((*BentNormalHistoryState)->GetRenderTargetItem().TargetableTexture, (*BentNormalHistoryState)->GetRenderTargetItem().ShaderResourceTexture, false, FResolveParams());
-			
+					RHICmdList.CopyToResolveTarget((*ConfidenceHistoryState)->GetRenderTargetItem().TargetableTexture, (*ConfidenceHistoryState)->GetRenderTargetItem().ShaderResourceTexture, false, FResolveParams());
+
 					if (bUseDistanceFieldGI)
 					{
 						RHICmdList.CopyToResolveTarget((*IrradianceHistoryState)->GetRenderTargetItem().TargetableTexture, (*IrradianceHistoryState)->GetRenderTargetItem().ShaderResourceTexture, false, FResolveParams());
@@ -749,6 +588,8 @@ void UpdateHistory(
 				*BentNormalHistoryState = NewBentNormalHistory;
 				BentNormalHistoryOutput = NewBentNormalHistory;
 
+				*ConfidenceHistoryState = NewConfidenceHistory;
+
 				*IrradianceHistoryState = NewIrradianceHistory;
 				IrradianceHistoryOutput = NewIrradianceHistory;
 			}
@@ -760,6 +601,9 @@ void UpdateHistory(
 			BentNormalHistoryOutput = BentNormalSource;
 			BentNormalSource = NULL;
 
+			*ConfidenceHistoryState = ConfidenceSource;
+			ConfidenceSource = NULL;
+
 			*IrradianceHistoryState = IrradianceSource;
 			IrradianceHistoryOutput = IrradianceSource;
 			IrradianceSource = NULL;
diff --git a/Engine/Source/Runtime/Renderer/Private/DistanceFieldLightingPost.h b/Engine/Source/Runtime/Renderer/Private/DistanceFieldLightingPost.h
--- a/Engine/Source/Runtime/Renderer/Private/DistanceFieldLightingPost.h
+++ b/Engine/Source/Runtime/Renderer/Private/DistanceFieldLightingPost.h
@@ -11,20 +11,23 @@
 
 class FDistanceFieldAOParameters;
 
-extern void AllocateOrReuseAORenderTarget(FRHICommandList& RHICmdList, TRefCountPtr<IPooledRenderTarget>& Target, const TCHAR* Name, bool bWithAlphaOrFP16Precision);
+extern void AllocateOrReuseAORenderTarget(FRHICommandList& RHICmdList, TRefCountPtr<IPooledRenderTarget>& Target, const TCHAR* Name, EPixelFormat Format);
 
 extern void UpdateHistory(
 	FRHICommandList& RHICmdList,
 	const FViewInfo& View, 
 	const TCHAR* BentNormalHistoryRTName,
+	const TCHAR* ConfidenceHistoryRTName,
 	const TCHAR* IrradianceHistoryRTName,
 	IPooledRenderTarget* VelocityTexture,
 	FSceneRenderTargetItem& DistanceFieldNormal,
 	/** Contains last frame's history, if non-NULL.  This will be updated with the new frame's history. */
 	TRefCountPtr<IPooledRenderTarget>* BentNormalHistoryState,
+	TRefCountPtr<IPooledRenderTarget>* ConfidenceHistoryState,
 	TRefCountPtr<IPooledRenderTarget>* IrradianceHistoryState,
 	/** Source */
 	TRefCountPtr<IPooledRenderTarget>& BentNormalSource, 
+	TRefCountPtr<IPooledRenderTarget>& ConfidenceSource,
 	TRefCountPtr<IPooledRenderTarget>& IrradianceSource, 
 	/** Output of Temporal Reprojection for the next step in the pipeline. */
 	TRefCountPtr<IPooledRenderTarget>& BentNormalHistoryOutput,
diff --git a/Engine/Source/Runtime/Renderer/Private/DistanceFieldLightingShared.h b/Engine/Source/Runtime/Renderer/Private/DistanceFieldLightingShared.h
--- a/Engine/Source/Runtime/Renderer/Private/DistanceFieldLightingShared.h
+++ b/Engine/Source/Runtime/Renderer/Private/DistanceFieldLightingShared.h
@@ -348,8 +348,8 @@ public:
 		DistanceFieldAtlasTexelSize.Bind(ParameterMap, TEXT("DistanceFieldAtlasTexelSize"));
 	}
 
-	template<typename TParamRef>
-	void Set(FRHICommandList& RHICmdList, const TParamRef& ShaderRHI, const FDistanceFieldCulledObjectBuffers& ObjectBuffers)
+	template<typename TParamRef, typename TRHICommandList>
+	void Set(TRHICommandList& RHICmdList, const TParamRef& ShaderRHI, const FDistanceFieldCulledObjectBuffers& ObjectBuffers)
 	{
 		ObjectIndirectArguments.SetBuffer(RHICmdList, ShaderRHI, ObjectBuffers.ObjectIndirectArguments);
 		CulledObjectBounds.SetBuffer(RHICmdList, ShaderRHI, ObjectBuffers.Bounds);
@@ -377,8 +377,8 @@ public:
 		SetShaderValue(RHICmdList, ShaderRHI, DistanceFieldAtlasTexelSize, InvTextureDim);
 	}
 
-	template<typename TParamRef>
-	void UnsetParameters(FRHICommandList& RHICmdList, const TParamRef& ShaderRHI)
+	template<typename TParamRef, typename TRHICommandList>
+	void UnsetParameters(TRHICommandList& RHICmdList, const TParamRef& ShaderRHI)
 	{
 		ObjectIndirectArguments.UnsetUAV(RHICmdList, ShaderRHI);
 		CulledObjectBounds.UnsetUAV(RHICmdList, ShaderRHI);
@@ -535,8 +535,8 @@ public:
 		return ShadowTileHeadDataUnpacked.IsBound() || ShadowTileArrayData.IsBound() || ShadowTileListGroupSize.IsBound() || ShadowMaxObjectsPerTile.IsBound();
 	}
 
-	template<typename TParamRef>
-	void Set(FRHICommandList& RHICmdList, const TParamRef& ShaderRHI, const FLightTileIntersectionResources& LightTileIntersectionResources)
+	template<typename TParamRef, typename TRHICommandList>
+	void Set(TRHICommandList& RHICmdList, const TParamRef& ShaderRHI, const FLightTileIntersectionResources& LightTileIntersectionResources)
 	{
 		ShadowTileHeadDataUnpacked.SetBuffer(RHICmdList, ShaderRHI, LightTileIntersectionResources.TileHeadDataUnpacked);
 
diff --git a/Engine/Source/Runtime/Renderer/Private/DistanceFieldScreenGridLighting.cpp b/Engine/Source/Runtime/Renderer/Private/DistanceFieldScreenGridLighting.cpp
--- a/Engine/Source/Runtime/Renderer/Private/DistanceFieldScreenGridLighting.cpp
+++ b/Engine/Source/Runtime/Renderer/Private/DistanceFieldScreenGridLighting.cpp
@@ -635,20 +635,23 @@ void PostProcessBentNormalAOScreenGrid(
 	const bool bUseDistanceFieldGI = IsDistanceFieldGIAllowed(View);
 
 	TRefCountPtr<IPooledRenderTarget> DistanceFieldAOBentNormal;
+	TRefCountPtr<IPooledRenderTarget> DistanceFieldAOConfidence;
 	TRefCountPtr<IPooledRenderTarget> DistanceFieldIrradiance;
-	AllocateOrReuseAORenderTarget(RHICmdList, DistanceFieldAOBentNormal, TEXT("DistanceFieldBentNormalAO"), true);
+	AllocateOrReuseAORenderTarget(RHICmdList, DistanceFieldAOBentNormal, TEXT("DistanceFieldBentNormalAO"), PF_FloatRGBA);
+	AllocateOrReuseAORenderTarget(RHICmdList, DistanceFieldAOConfidence, TEXT("DistanceFieldConfidence"), PF_G8);
 
 	if (bUseDistanceFieldGI)
 	{
-		AllocateOrReuseAORenderTarget(RHICmdList, DistanceFieldIrradiance, TEXT("DistanceFieldIrradiance"), false);
+		AllocateOrReuseAORenderTarget(RHICmdList, DistanceFieldIrradiance, TEXT("DistanceFieldIrradiance"), PF_FloatRGB);
 	}
 
 	{
 		SCOPED_DRAW_EVENT(RHICmdList, GeometryAwareUpsample);
 
-		FTextureRHIParamRef RenderTargets[2] =
+		FTextureRHIParamRef RenderTargets[3] =
 		{
 			DistanceFieldAOBentNormal->GetRenderTargetItem().TargetableTexture,
+			DistanceFieldAOConfidence->GetRenderTargetItem().TargetableTexture,
 			bUseDistanceFieldGI ? DistanceFieldIrradiance->GetRenderTargetItem().TargetableTexture : NULL,
 		};
 
@@ -693,6 +696,7 @@ void PostProcessBentNormalAOScreenGrid(
 		}
 
 		RHICmdList.CopyToResolveTarget(DistanceFieldAOBentNormal->GetRenderTargetItem().TargetableTexture, DistanceFieldAOBentNormal->GetRenderTargetItem().ShaderResourceTexture, false, FResolveParams());
+		RHICmdList.CopyToResolveTarget(DistanceFieldAOConfidence->GetRenderTargetItem().TargetableTexture, DistanceFieldAOConfidence->GetRenderTargetItem().ShaderResourceTexture, false, FResolveParams());
 			
 		if (bUseDistanceFieldGI)
 		{
@@ -702,6 +706,7 @@ void PostProcessBentNormalAOScreenGrid(
 
 	FSceneViewState* ViewState = (FSceneViewState*)View.State;
 	TRefCountPtr<IPooledRenderTarget>* BentNormalHistoryState = ViewState ? &ViewState->DistanceFieldAOHistoryRT : NULL;
+	TRefCountPtr<IPooledRenderTarget>* ConfidenceHistoryState = ViewState ? &ViewState->DistanceFieldAOConfidenceHistoryRT : NULL;
 	TRefCountPtr<IPooledRenderTarget>* IrradianceHistoryState = ViewState ? &ViewState->DistanceFieldIrradianceHistoryRT : NULL;
 	BentNormalOutput = DistanceFieldAOBentNormal;
 	IrradianceOutput = DistanceFieldIrradiance;
@@ -712,12 +717,15 @@ void PostProcessBentNormalAOScreenGrid(
 			RHICmdList,
 			View, 
 			TEXT("DistanceFieldAOHistory"),
+			TEXT("DistanceFieldAOConfidence"),
 			TEXT("DistanceFieldIrradianceHistory"),
 			VelocityTexture,
 			DistanceFieldNormal,
 			BentNormalHistoryState,
+			ConfidenceHistoryState,
 			IrradianceHistoryState,
 			DistanceFieldAOBentNormal, 
+			DistanceFieldAOConfidence,
 			DistanceFieldIrradiance,
 			BentNormalOutput,
 			IrradianceOutput);
diff --git a/Engine/Source/Runtime/Renderer/Private/DistanceFieldShadowing.cpp b/Engine/Source/Runtime/Renderer/Private/DistanceFieldShadowing.cpp
--- a/Engine/Source/Runtime/Renderer/Private/DistanceFieldShadowing.cpp
+++ b/Engine/Source/Runtime/Renderer/Private/DistanceFieldShadowing.cpp
@@ -42,6 +42,14 @@ FAutoConsoleVariableRef CVarFullResolutionDFShadowing(
 	ECVF_Scalability | ECVF_RenderThreadSafe
 	);
 
+int32 GAsyncComputeDFShadowing = 0;
+FAutoConsoleVariableRef CVarAsyncComputeDFShadowing(
+	TEXT("r.DFShadowAsyncCompute"),
+	GAsyncComputeDFShadowing,
+	TEXT("Whether to use async compute for ray marching distance fields."),
+	ECVF_Scalability | ECVF_RenderThreadSafe
+	);
+
 int32 GShadowScatterTileCulling = 1;
 FAutoConsoleVariableRef CVarShadowScatterTileCulling(
 	TEXT("r.DFShadowScatterTileCulling"),
@@ -425,8 +433,9 @@ public:
 		DownsampleFactor.Bind(Initializer.ParameterMap, TEXT("DownsampleFactor"));
 	}
 
+	template<typename TRHICommandList>
 	void SetParameters(
-		FRHICommandList& RHICmdList, 
+		TRHICommandList& RHICmdList, 
 		const FSceneView& View, 
 		const FProjectedShadowInfo* ProjectedShadowInfo,
 		FSceneRenderTargetItem& ShadowFactorsValue, 
@@ -498,7 +507,8 @@ public:
 		SetShaderValue(RHICmdList, ShaderRHI, DownsampleFactor, GetDFShadowDownsampleFactor());
 	}
 
-	void UnsetParameters(FRHICommandList& RHICmdList, FSceneRenderTargetItem& ShadowFactorsValue)
+	template<typename TRHICommandList>
+	void UnsetParameters(TRHICommandList& RHICmdList, FSceneRenderTargetItem& ShadowFactorsValue)
 	{
 		ShadowFactors.UnsetUAV(RHICmdList, GetComputeShader());
 		RHICmdList.TransitionResource(EResourceTransitionAccess::ERWBarrier, EResourceTransitionPipeline::EComputeToCompute, ShadowFactorsValue.UAV);
@@ -801,12 +811,52 @@ bool FDeferredShadingSceneRenderer::ShouldPrepareForDistanceFieldShadows() const
 		&& SupportsDistanceFieldShadows(Scene->GetFeatureLevel(), Scene->GetShaderPlatform());
 }
 
-void FProjectedShadowInfo::RenderRayTracedDistanceFieldProjection(FRHICommandListImmediate& RHICmdList, const FViewInfo& View, bool bProjectingForForwardShading) const
+template<typename TRHICommandList>
+void RayTraceShadows(TRHICommandList& RHICmdList, const FViewInfo& View, FProjectedShadowInfo* ProjectedShadowInfo, FLightTileIntersectionResources* TileIntersectionResources)
+{
+	FIntRect ScissorRect;
+
+	if (!ProjectedShadowInfo->GetLightSceneInfo().Proxy->GetScissorRect(ScissorRect, View))
+	{
+		ScissorRect = View.ViewRect;
+	}
+
+	uint32 GroupSizeX = FMath::DivideAndRoundUp(ScissorRect.Size().X / GetDFShadowDownsampleFactor(), GDistanceFieldAOTileSizeX);
+	uint32 GroupSizeY = FMath::DivideAndRoundUp(ScissorRect.Size().Y / GetDFShadowDownsampleFactor(), GDistanceFieldAOTileSizeY);
+
+	FSceneRenderTargetItem& RayTracedShadowsRTI = ProjectedShadowInfo->RayTracedShadowsRT->GetRenderTargetItem();
+	if (ProjectedShadowInfo->bDirectionalLight && GShadowScatterTileCulling)
+	{
+		TShaderMapRef<TDistanceFieldShadowingCS<DFS_DirectionalLightScatterTileCulling> > ComputeShader(View.ShaderMap);
+		RHICmdList.SetComputeShader(ComputeShader->GetComputeShader());
+		ComputeShader->SetParameters(RHICmdList, View, ProjectedShadowInfo, RayTracedShadowsRTI, FVector2D(GroupSizeX, GroupSizeY), ScissorRect, TileIntersectionResources);
+		DispatchComputeShader(RHICmdList, *ComputeShader, GroupSizeX, GroupSizeY, 1);
+		ComputeShader->UnsetParameters(RHICmdList, RayTracedShadowsRTI);
+	}
+	else if (ProjectedShadowInfo->bDirectionalLight)
+	{
+		TShaderMapRef<TDistanceFieldShadowingCS<DFS_DirectionalLightTiledCulling> > ComputeShader(View.ShaderMap);
+		RHICmdList.SetComputeShader(ComputeShader->GetComputeShader());
+		ComputeShader->SetParameters(RHICmdList, View, ProjectedShadowInfo, RayTracedShadowsRTI, FVector2D(GroupSizeX, GroupSizeY), ScissorRect, TileIntersectionResources);
+		DispatchComputeShader(RHICmdList, *ComputeShader, GroupSizeX, GroupSizeY, 1);
+		ComputeShader->UnsetParameters(RHICmdList, RayTracedShadowsRTI);
+	}
+	else
+	{
+		TShaderMapRef<TDistanceFieldShadowingCS<DFS_PointLightTiledCulling> > ComputeShader(View.ShaderMap);
+		RHICmdList.SetComputeShader(ComputeShader->GetComputeShader());
+		ComputeShader->SetParameters(RHICmdList, View, ProjectedShadowInfo, RayTracedShadowsRTI, FVector2D(GroupSizeX, GroupSizeY), ScissorRect, TileIntersectionResources);
+		DispatchComputeShader(RHICmdList, *ComputeShader, GroupSizeX, GroupSizeY, 1);
+		ComputeShader->UnsetParameters(RHICmdList, RayTracedShadowsRTI);
+	}
+}
+
+void FProjectedShadowInfo::BeginRenderRayTracedDistanceFieldProjection(FRHICommandListImmediate& RHICmdList, const FViewInfo& View)
 {
 	if (SupportsDistanceFieldShadows(View.GetFeatureLevel(), View.GetShaderPlatform()))
 	{
-		QUICK_SCOPE_CYCLE_COUNTER(STAT_RenderRayTracedDistanceFieldShadows);
-		SCOPED_DRAW_EVENT(RHICmdList, RayTracedDistanceFieldShadow);
+		QUICK_SCOPE_CYCLE_COUNTER(STAT_BeginRenderRayTracedDistanceFieldShadows);
+		SCOPED_DRAW_EVENT(RHICmdList, BeginRayTracedDistanceFieldShadow);
 
 		const FScene* Scene = (const FScene*)(View.Family->Scene);
 
@@ -841,7 +891,7 @@ void FProjectedShadowInfo::RenderRayTracedDistanceFieldProjection(FRHICommandLis
 			CullDistanceFieldObjectsForLight(
 				RHICmdList,
 				View,
-				LightSceneInfo->Proxy, 
+				LightSceneInfo->Proxy,
 				WorldToShadowValue,
 				NumPlanes,
 				PlaneData,
@@ -850,117 +900,121 @@ void FProjectedShadowInfo::RenderRayTracedDistanceFieldProjection(FRHICommandLis
 				LightSceneInfo->TileIntersectionResources
 				);
 
+			// Note: using the same TileIntersectionResources for multiple views, breaks splitscreen / stereo
 			FLightTileIntersectionResources* TileIntersectionResources = LightSceneInfo->TileIntersectionResources.Get();
 
 			View.HeightfieldLightingViewInfo.ComputeRayTracedShadowing(View, RHICmdList, this, TileIntersectionResources, GShadowCulledObjectBuffers);
 
-			TRefCountPtr<IPooledRenderTarget> RayTracedShadowsRT;
-
 			{
 				const FIntPoint BufferSize = GetBufferSizeForDFShadows();
 				FPooledRenderTargetDesc Desc(FPooledRenderTargetDesc::Create2DDesc(BufferSize, PF_G16R16F, FClearValueBinding::None, TexCreate_None, TexCreate_RenderTargetable | TexCreate_UAV, false));
 				GRenderTargetPool.FindFreeElement(RHICmdList, Desc, RayTracedShadowsRT, TEXT("RayTracedShadows"));
 			}
 
-			FIntRect ScissorRect;
+			SCOPED_DRAW_EVENT(RHICmdList, RayTraceShadows);
+			SetRenderTarget(RHICmdList, NULL, NULL);
 
+			if (GAsyncComputeDFShadowing)
 			{
-				if (!LightSceneInfo->Proxy->GetScissorRect(ScissorRect, View))
-				{
-					ScissorRect = View.ViewRect;
-				}
+				FRHIAsyncComputeCommandListImmediate& RHICmdListComputeImmediate = FRHICommandListExecutor::GetImmediateAsyncComputeCommandList();
+				static const FName RTShadowBeginComputeName(TEXT("RTShadowComputeBegin"));
+				static const FName RTShadowComputeEndName(TEXT("RTShadowComputeEnd"));
+				FComputeFenceRHIRef BeginFence = RHICmdList.CreateComputeFence(RTShadowBeginComputeName);
+				RayTracedShadowsEndFence = RHICmdList.CreateComputeFence(RTShadowComputeEndName);
 
-				uint32 GroupSizeX = FMath::DivideAndRoundUp(ScissorRect.Size().X / GetDFShadowDownsampleFactor(), GDistanceFieldAOTileSizeX);
-				uint32 GroupSizeY = FMath::DivideAndRoundUp(ScissorRect.Size().Y / GetDFShadowDownsampleFactor(), GDistanceFieldAOTileSizeY);
+				RHICmdList.TransitionResources(EResourceTransitionAccess::EReadable, EResourceTransitionPipeline::EGfxToCompute, nullptr, 0, BeginFence);
+				RHICmdListComputeImmediate.WaitComputeFence(BeginFence);
 
-				{
-					SCOPED_DRAW_EVENT(RHICmdList, RayTraceShadows);
-
-					SetRenderTarget(RHICmdList, NULL, NULL);
-
-					FSceneRenderTargetItem& RayTracedShadowsRTI = RayTracedShadowsRT->GetRenderTargetItem();
-					if (bDirectionalLight && GShadowScatterTileCulling)
-					{
-						TShaderMapRef<TDistanceFieldShadowingCS<DFS_DirectionalLightScatterTileCulling> > ComputeShader(View.ShaderMap);
-						RHICmdList.SetComputeShader(ComputeShader->GetComputeShader());
-						ComputeShader->SetParameters(RHICmdList, View, this, RayTracedShadowsRTI, FVector2D(GroupSizeX, GroupSizeY), ScissorRect, TileIntersectionResources);
-						DispatchComputeShader(RHICmdList, *ComputeShader, GroupSizeX, GroupSizeY, 1);
-						ComputeShader->UnsetParameters(RHICmdList, RayTracedShadowsRTI);
-					}
-					else if (bDirectionalLight)
-					{
-						TShaderMapRef<TDistanceFieldShadowingCS<DFS_DirectionalLightTiledCulling> > ComputeShader(View.ShaderMap);
-						RHICmdList.SetComputeShader(ComputeShader->GetComputeShader());
-						ComputeShader->SetParameters(RHICmdList, View, this, RayTracedShadowsRTI, FVector2D(GroupSizeX, GroupSizeY), ScissorRect, TileIntersectionResources);
-						DispatchComputeShader(RHICmdList, *ComputeShader, GroupSizeX, GroupSizeY, 1);
-						ComputeShader->UnsetParameters(RHICmdList, RayTracedShadowsRTI);
-					}
-					else
-					{
-						TShaderMapRef<TDistanceFieldShadowingCS<DFS_PointLightTiledCulling> > ComputeShader(View.ShaderMap);
-						RHICmdList.SetComputeShader(ComputeShader->GetComputeShader());
-						ComputeShader->SetParameters(RHICmdList, View, this, RayTracedShadowsRTI, FVector2D(GroupSizeX, GroupSizeY), ScissorRect, TileIntersectionResources);
-						DispatchComputeShader(RHICmdList, *ComputeShader, GroupSizeX, GroupSizeY, 1);
-						ComputeShader->UnsetParameters(RHICmdList, RayTracedShadowsRTI);
-					}
-				}
-			}
+				RayTraceShadows(RHICmdListComputeImmediate, View, this, TileIntersectionResources);
 
+				RHICmdListComputeImmediate.TransitionResources(EResourceTransitionAccess::EReadable, EResourceTransitionPipeline::EComputeToGfx, nullptr, 0, RayTracedShadowsEndFence);
+
+				FRHIAsyncComputeCommandListImmediate::ImmediateDispatch(RHICmdListComputeImmediate);
+			}
+			else
 			{
-				FSceneRenderTargets::Get(RHICmdList).BeginRenderingLightAttenuation(RHICmdList);
+				RayTraceShadows(RHICmdList, View, this, TileIntersectionResources);
+			}
+		}
+	}
+}
+
+void FProjectedShadowInfo::RenderRayTracedDistanceFieldProjection(FRHICommandListImmediate& RHICmdList, const FViewInfo& View, bool bProjectingForForwardShading) 
+{
+	if (RayTracedShadowsRT)
+	{
+		QUICK_SCOPE_CYCLE_COUNTER(STAT_RenderRayTracedDistanceFieldShadows);
+		SCOPED_DRAW_EVENT(RHICmdList, RayTracedDistanceFieldShadow);
+
+		if (GAsyncComputeDFShadowing)
+		{
+			RHICmdList.WaitComputeFence(RayTracedShadowsEndFence);
+		}
+
+		FIntRect ScissorRect;
+
+		if (!LightSceneInfo->Proxy->GetScissorRect(ScissorRect, View))
+		{
+			ScissorRect = View.ViewRect;
+		}
 
-				SCOPED_DRAW_EVENT(RHICmdList, Upsample);
+		{
+			FSceneRenderTargets::Get(RHICmdList).BeginRenderingLightAttenuation(RHICmdList);
 
-				RHICmdList.SetViewport(ScissorRect.Min.X, ScissorRect.Min.Y, 0.0f, ScissorRect.Max.X, ScissorRect.Max.Y, 1.0f);
-				RHICmdList.SetRasterizerState(TStaticRasterizerState<FM_Solid, CM_None>::GetRHI());
-				RHICmdList.SetDepthStencilState(TStaticDepthStencilState<false, CF_Always>::GetRHI());
+			SCOPED_DRAW_EVENT(RHICmdList, Upsample);
+
+			RHICmdList.SetViewport(ScissorRect.Min.X, ScissorRect.Min.Y, 0.0f, ScissorRect.Max.X, ScissorRect.Max.Y, 1.0f);
+			RHICmdList.SetRasterizerState(TStaticRasterizerState<FM_Solid, CM_None>::GetRHI());
+			RHICmdList.SetDepthStencilState(TStaticDepthStencilState<false, CF_Always>::GetRHI());
 				
-				SetBlendStateForProjection(RHICmdList, bProjectingForForwardShading, false);
+			SetBlendStateForProjection(RHICmdList, bProjectingForForwardShading, false);
 
-				//@todo - depth bounds test for local lights
-				if (bDirectionalLight)
-				{
-					EnableDepthBoundsTest(RHICmdList, CascadeSettings.SplitNear - CascadeSettings.SplitNearFadeRegion, CascadeSettings.SplitFar, View.ViewMatrices.GetProjectionMatrix());
-				}
+			//@todo - depth bounds test for local lights
+			if (bDirectionalLight)
+			{
+				EnableDepthBoundsTest(RHICmdList, CascadeSettings.SplitNear - CascadeSettings.SplitNearFadeRegion, CascadeSettings.SplitFar, View.ViewMatrices.GetProjectionMatrix());
+			}
 
-				TShaderMapRef<FPostProcessVS> VertexShader(View.ShaderMap);
+			TShaderMapRef<FPostProcessVS> VertexShader(View.ShaderMap);
 
-				if (GFullResolutionDFShadowing)
-				{
-					TShaderMapRef<TDistanceFieldShadowingUpsamplePS<false> > PixelShader(View.ShaderMap);
+			if (GFullResolutionDFShadowing)
+			{
+				TShaderMapRef<TDistanceFieldShadowingUpsamplePS<false> > PixelShader(View.ShaderMap);
 
-					static FGlobalBoundShaderState BoundShaderState;
-					SetGlobalBoundShaderState(RHICmdList, View.GetFeatureLevel(), BoundShaderState, GFilterVertexDeclaration.VertexDeclarationRHI, *VertexShader, *PixelShader);
+				static FGlobalBoundShaderState BoundShaderState;
+				SetGlobalBoundShaderState(RHICmdList, View.GetFeatureLevel(), BoundShaderState, GFilterVertexDeclaration.VertexDeclarationRHI, *VertexShader, *PixelShader);
 
-					VertexShader->SetParameters(RHICmdList, View);
-					PixelShader->SetParameters(RHICmdList, View, this, ScissorRect, RayTracedShadowsRT);
-				}
-				else
-				{
-					TShaderMapRef<TDistanceFieldShadowingUpsamplePS<true> > PixelShader(View.ShaderMap);
+				VertexShader->SetParameters(RHICmdList, View);
+				PixelShader->SetParameters(RHICmdList, View, this, ScissorRect, RayTracedShadowsRT);
+			}
+			else
+			{
+				TShaderMapRef<TDistanceFieldShadowingUpsamplePS<true> > PixelShader(View.ShaderMap);
 
-					static FGlobalBoundShaderState BoundShaderState;
-					SetGlobalBoundShaderState(RHICmdList, View.GetFeatureLevel(), BoundShaderState, GFilterVertexDeclaration.VertexDeclarationRHI, *VertexShader, *PixelShader);
+				static FGlobalBoundShaderState BoundShaderState;
+				SetGlobalBoundShaderState(RHICmdList, View.GetFeatureLevel(), BoundShaderState, GFilterVertexDeclaration.VertexDeclarationRHI, *VertexShader, *PixelShader);
 
-					VertexShader->SetParameters(RHICmdList, View);
-					PixelShader->SetParameters(RHICmdList, View, this, ScissorRect, RayTracedShadowsRT);
-				}
+				VertexShader->SetParameters(RHICmdList, View);
+				PixelShader->SetParameters(RHICmdList, View, this, ScissorRect, RayTracedShadowsRT);
+			}
 
-				DrawRectangle( 
-					RHICmdList,
-					0, 0, 
-					ScissorRect.Width(), ScissorRect.Height(),
-					ScissorRect.Min.X / GetDFShadowDownsampleFactor(), ScissorRect.Min.Y / GetDFShadowDownsampleFactor(), 
-					ScissorRect.Width() / GetDFShadowDownsampleFactor(), ScissorRect.Height() / GetDFShadowDownsampleFactor(),
-					FIntPoint(ScissorRect.Width(), ScissorRect.Height()),
-					GetBufferSizeForDFShadows(),
-					*VertexShader);
-
-				if (bDirectionalLight)
-				{
-					DisableDepthBoundsTest(RHICmdList);
-				}
+			DrawRectangle( 
+				RHICmdList,
+				0, 0, 
+				ScissorRect.Width(), ScissorRect.Height(),
+				ScissorRect.Min.X / GetDFShadowDownsampleFactor(), ScissorRect.Min.Y / GetDFShadowDownsampleFactor(), 
+				ScissorRect.Width() / GetDFShadowDownsampleFactor(), ScissorRect.Height() / GetDFShadowDownsampleFactor(),
+				FIntPoint(ScissorRect.Width(), ScissorRect.Height()),
+				GetBufferSizeForDFShadows(),
+				*VertexShader);
+
+			if (bDirectionalLight)
+			{
+				DisableDepthBoundsTest(RHICmdList);
 			}
 		}
+
+		RayTracedShadowsRT = NULL;
+		RayTracedShadowsEndFence = NULL;
 	}
 }
diff --git a/Engine/Source/Runtime/Renderer/Private/PostProcess/PostProcessEyeAdaptation.cpp b/Engine/Source/Runtime/Renderer/Private/PostProcess/PostProcessEyeAdaptation.cpp
--- a/Engine/Source/Runtime/Renderer/Private/PostProcess/PostProcessEyeAdaptation.cpp
+++ b/Engine/Source/Runtime/Renderer/Private/PostProcess/PostProcessEyeAdaptation.cpp
@@ -8,6 +8,7 @@
 #include "PostProcess/SceneFilterRendering.h"
 #include "PostProcess/PostProcessing.h"
 #include "ClearQuad.h"
+#include "UnrealMathUtility.h"
 
 
 /**
@@ -401,7 +402,7 @@ public:
 		PostprocessParameter.Bind(Initializer.ParameterMap);
 		EyeAdaptationTexture.Bind(Initializer.ParameterMap, TEXT("EyeAdaptationTexture"));
 		EyeAdaptationParams.Bind(Initializer.ParameterMap, TEXT("EyeAdaptationParams"));
-		EyeAdaptationExtent.Bind(Initializer.ParameterMap, TEXT("EyeAdaptionSrcExtent"));
+		EyeAdaptationSrcRect.Bind(Initializer.ParameterMap, TEXT("EyeAdaptionSrcRect"));
 	}
 
 public:
@@ -419,7 +420,7 @@ public:
 	}
 
 
-	void SetPS(const FRenderingCompositePassContext& Context, const FIntPoint SrcSize, IPooledRenderTarget* EyeAdaptationLastFrameRT)
+	void SetPS(const FRenderingCompositePassContext& Context, const FIntRect& SrcRect, IPooledRenderTarget* EyeAdaptationLastFrameRT)
 	{
 		const FPixelShaderRHIParamRef ShaderRHI = GetPixelShader();
 
@@ -453,7 +454,7 @@ public:
 		}
 
 		// Set the src extent for the shader
-		SetShaderValue(Context.RHICmdList, ShaderRHI, EyeAdaptationExtent, SrcSize);
+		SetShaderValue(Context.RHICmdList, ShaderRHI, EyeAdaptationSrcRect, SrcRect);
 	}
 
 	// FShader interface.
@@ -461,7 +462,7 @@ public:
 	{
 		bool bShaderHasOutdatedParameters = FGlobalShader::Serialize(Ar);
 		Ar << PostprocessParameter << EyeAdaptationTexture 
-			<< EyeAdaptationParams << EyeAdaptationExtent;
+			<< EyeAdaptationParams << EyeAdaptationSrcRect;
 		return bShaderHasOutdatedParameters;
 	}
 
@@ -469,7 +470,7 @@ private:
 	FPostProcessPassParameters PostprocessParameter;
 	FShaderResourceParameter EyeAdaptationTexture;
 	FShaderParameter EyeAdaptationParams;
-	FShaderParameter EyeAdaptationExtent;
+	FShaderParameter EyeAdaptationSrcRect;
 };
 
 IMPLEMENT_SHADER_TYPE(, FPostProcessLogLuminance2ExposureScalePS, TEXT("PostProcessEyeAdaptation"), TEXT("MainLogLuminance2ExposureScalePS"), SF_Pixel);
@@ -490,9 +491,11 @@ void FRCPassPostProcessBasicEyeAdaptation::Process(FRenderingCompositePassContex
 
 	// The input texture sample size.  Averaged in the pixel shader.
 	FIntPoint SrcSize = GetInputDesc(ePId_Input0)->Extent;
-	
-	// Clamp gather area to actual used scene viewrect
-	SrcSize = SrcSize.ComponentMin(DownsampledViewRect);
+
+	// Compute the region of interest in the source texture.
+	uint32 ScaleFactor = FMath::DivideAndRoundUp(FSceneRenderTargets::Get(Context.RHICmdList).GetBufferSizeXY().Y, SrcSize.Y);
+
+	FIntRect SrcRect = View.ViewRect / ScaleFactor;
 
 	SCOPED_DRAW_EVENTF(Context.RHICmdList, PostProcessBasicEyeAdaptation, TEXT("PostProcessBasicEyeAdaptation %dx%d"), SrcSize.X, SrcSize.Y);
 
@@ -521,7 +524,7 @@ void FRCPassPostProcessBasicEyeAdaptation::Process(FRenderingCompositePassContex
 
 	// Set the parameters used by the pixel shader.
 
-	PixelShader->SetPS(Context, SrcSize, EyeAdaptationLastFrameRT);
+	PixelShader->SetPS(Context, SrcRect, EyeAdaptationLastFrameRT);
 
 	// Draw a quad mapping scene color to the view's render target
 	DrawRectangle(
diff --git a/Engine/Source/Runtime/Renderer/Private/ScenePrivate.h b/Engine/Source/Runtime/Renderer/Private/ScenePrivate.h
--- a/Engine/Source/Runtime/Renderer/Private/ScenePrivate.h
+++ b/Engine/Source/Runtime/Renderer/Private/ScenePrivate.h
@@ -678,6 +678,7 @@ public:
 	// Temporal AA result for light shafts of last frame
 	TMap<const ULightComponent*, TRefCountPtr<IPooledRenderTarget> > LightShaftBloomHistoryRTs;
 	TRefCountPtr<IPooledRenderTarget> DistanceFieldAOHistoryRT;
+	TRefCountPtr<IPooledRenderTarget> DistanceFieldAOConfidenceHistoryRT;
 	TRefCountPtr<IPooledRenderTarget> DistanceFieldIrradianceHistoryRT;
 	// Mobile temporal AA surfaces.
 	TRefCountPtr<IPooledRenderTarget> MobileAaBloomSunVignette0;
@@ -984,6 +985,8 @@ public:
 		LightShaftOcclusionHistoryRT.SafeRelease();
 		LightShaftBloomHistoryRTs.Empty();
 		DistanceFieldAOHistoryRT.SafeRelease();
+		DistanceFieldAOConfidenceHistoryRT.SafeRelease();
+		DistanceFieldAOConfidenceHistoryRT.SafeRelease();
 		DistanceFieldIrradianceHistoryRT.SafeRelease();
 		MobileAaBloomSunVignette0.SafeRelease();
 		MobileAaBloomSunVignette1.SafeRelease();
diff --git a/Engine/Source/Runtime/Renderer/Private/SceneRendering.h b/Engine/Source/Runtime/Renderer/Private/SceneRendering.h
--- a/Engine/Source/Runtime/Renderer/Private/SceneRendering.h
+++ b/Engine/Source/Runtime/Renderer/Private/SceneRendering.h
@@ -1323,6 +1323,8 @@ protected:
 		const TArray<FProjectedShadowInfo*, SceneRenderingAllocator>& ViewDependentWholeSceneShadows,
 		bool bReflectionCaptureScene);
 
+	void BeginRenderRayTracedDistanceFieldProjections(FRHICommandListImmediate& RHICmdList);
+
 	void RenderShadowDepthMaps(FRHICommandListImmediate& RHICmdList);
 	void RenderShadowDepthMapAtlases(FRHICommandListImmediate& RHICmdList);
 
diff --git a/Engine/Source/Runtime/Renderer/Private/ShadowDepthRendering.cpp b/Engine/Source/Runtime/Renderer/Private/ShadowDepthRendering.cpp
--- a/Engine/Source/Runtime/Renderer/Private/ShadowDepthRendering.cpp
+++ b/Engine/Source/Runtime/Renderer/Private/ShadowDepthRendering.cpp
@@ -2207,6 +2207,31 @@ void FSceneRenderer::RenderShadowDepthMapAtlases(FRHICommandListImmediate& RHICm
 	}
 }
 
+void FSceneRenderer::BeginRenderRayTracedDistanceFieldProjections(FRHICommandListImmediate& RHICmdList)
+{
+	for (TSparseArray<FLightSceneInfoCompact>::TConstIterator LightIt(Scene->Lights); LightIt; ++LightIt)
+	{
+		const FLightSceneInfoCompact& LightSceneInfoCompact = *LightIt;
+		const FLightSceneInfo* const LightSceneInfo = LightSceneInfoCompact.LightSceneInfo;
+		FVisibleLightInfo& VisibleLightInfo = VisibleLightInfos[LightSceneInfo->Id];
+
+		for (int32 ViewIndex = 0; ViewIndex < Views.Num(); ViewIndex++)
+		{
+			const FViewInfo& View = Views[ViewIndex];
+
+			for (int32 ShadowIndex = 0; ShadowIndex < VisibleLightInfo.ShadowsToProject.Num(); ShadowIndex++)
+			{
+				FProjectedShadowInfo* ProjectedShadowInfo = VisibleLightInfo.ShadowsToProject[ShadowIndex];
+
+				if (ProjectedShadowInfo->bRayTracedDistanceField)
+				{
+					ProjectedShadowInfo->BeginRenderRayTracedDistanceFieldProjection(RHICmdList, View);
+				}
+			}
+		}
+	}
+}
+
 void FSceneRenderer::RenderShadowDepthMaps(FRHICommandListImmediate& RHICmdList)
 {
 	FSceneRenderTargets& SceneContext = FSceneRenderTargets::Get(RHICmdList);
@@ -2214,6 +2239,8 @@ void FSceneRenderer::RenderShadowDepthMaps(FRHICommandListImmediate& RHICmdList)
 	SCOPED_DRAW_EVENT(RHICmdList, ShadowDepths);
 	SCOPED_GPU_STAT(RHICmdList, Stat_GPU_ShadowDepths);
 
+	BeginRenderRayTracedDistanceFieldProjections(RHICmdList);
+
 	FSceneRenderer::RenderShadowDepthMapAtlases(RHICmdList);
 
 	for (int32 CubemapIndex = 0; CubemapIndex < SortedShadowsForShadowDepthPass.ShadowMapCubemaps.Num(); CubemapIndex++)
diff --git a/Engine/Source/Runtime/Renderer/Private/ShadowRendering.h b/Engine/Source/Runtime/Renderer/Private/ShadowRendering.h
--- a/Engine/Source/Runtime/Renderer/Private/ShadowRendering.h
+++ b/Engine/Source/Runtime/Renderer/Private/ShadowRendering.h
@@ -808,6 +808,10 @@ public:
 	/** Frustums for each cubemap face, used for object culling one pass point light shadows. */
 	TArray<FConvexVolume> OnePassShadowFrustums;
 
+	/** Data passed from async compute begin to end. */
+	FComputeFenceRHIRef RayTracedShadowsEndFence;
+	TRefCountPtr<IPooledRenderTarget> RayTracedShadowsRT;
+
 public:
 
 	// default constructor
@@ -871,8 +875,10 @@ public:
 	 */
 	void RenderProjection(FRHICommandListImmediate& RHICmdList, int32 ViewIndex, const class FViewInfo* View, bool bProjectingForForwardShading, bool bMobile) const;
 
+	void BeginRenderRayTracedDistanceFieldProjection(FRHICommandListImmediate& RHICmdList, const FViewInfo& View);
+
 	/** Renders ray traced distance field shadows. */
-	void RenderRayTracedDistanceFieldProjection(FRHICommandListImmediate& RHICmdList, const class FViewInfo& View, bool bProjectingForForwardShading) const;
+	void RenderRayTracedDistanceFieldProjection(FRHICommandListImmediate& RHICmdList, const class FViewInfo& View, bool bProjectingForForwardShading);
 
 	/** Render one pass point light shadow projections. */
 	void RenderOnePassPointLightProjection(FRHICommandListImmediate& RHICmdList, int32 ViewIndex, const FViewInfo& View, bool bProjectingForForwardShading) const;
diff --git a/Engine/Source/Runtime/ShaderCore/Public/Shader.h b/Engine/Source/Runtime/ShaderCore/Public/Shader.h
--- a/Engine/Source/Runtime/ShaderCore/Public/Shader.h
+++ b/Engine/Source/Runtime/ShaderCore/Public/Shader.h
@@ -218,7 +218,12 @@ public:
 
 	/** Returns true if and only if TargetPlatform is compatible for use with CurrentPlatform. */
 	SHADERCORE_API static bool ArePlatformsCompatible(EShaderPlatform CurrentPlatform, EShaderPlatform TargetPlatform);
-	
+
+	void GetShaderCode(TArray<uint8>& OutCode) const
+	{
+		UncompressCode(OutCode);
+	}
+
 private:
 	// compression functions
 	void UncompressCode(TArray<uint8>& UncompressedCode) const;
diff --git a/Engine/Source/Runtime/VulkanRHI/Private/VulkanCommandWrappers.h b/Engine/Source/Runtime/VulkanRHI/Private/VulkanCommandWrappers.h
--- a/Engine/Source/Runtime/VulkanRHI/Private/VulkanCommandWrappers.h
+++ b/Engine/Source/Runtime/VulkanRHI/Private/VulkanCommandWrappers.h
@@ -509,13 +509,16 @@ namespace VulkanRHI
 		VULKANAPINAMESPACE::vkDestroySemaphore(Device, Semaphore, Allocator);
 	}
 
-#if 0
-	static FORCEINLINE_DEBUGGABLE VkResult  vkCreateEvent(
-		VkDevice                                    Device,
-		const VkEventCreateInfo*                    CreateInfo,
-		const VkAllocationCallbacks*                Allocator,
-		VkEvent*                                    pEvent);
-#endif
+	static FORCEINLINE_DEBUGGABLE VkResult  vkCreateEvent(VkDevice Device, const VkEventCreateInfo* CreateInfo, const VkAllocationCallbacks* Allocator, VkEvent* pEvent)
+	{
+		DevicePrintfBegin(Device, FString::Printf(TEXT("vkCreateEvent(CreateInfo=0x%016llx%s, OutEvent=0x%016llx)"), CreateInfo, Event));
+
+		VkResult Result = VULKANAPINAMESPACE::vkCreateEvent(Device, CreateInfo, Allocator,pEvent);
+
+		PrintResultAndNamedHandle(Result, TEXT("Event"), *Fence);
+		return Result;
+	}
+
 	static FORCEINLINE_DEBUGGABLE void  vkDestroyEvent(VkDevice Device, VkEvent Event, const VkAllocationCallbacks* Allocator)
 	{
 		DevicePrintfBegin(Device, FString::Printf(TEXT("vkDestroyEvent(Event=%p)"), Event));
@@ -1157,30 +1160,33 @@ namespace VulkanRHI
 
 		VULKANAPINAMESPACE::vkCmdResolveImage(CommandBuffer, SrcImage, SrcImageLayout, DstImage, DstImageLayout, RegionCount, Regions);
 	}
-#if 0
-	static FORCEINLINE_DEBUGGABLE void  vkCmdSetEvent(
-		VkCommandBuffer                             commandBuffer,
-		VkEvent                                     event,
-		VkPipelineStageFlags                        stageMask);
 
-	static FORCEINLINE_DEBUGGABLE void  vkCmdResetEvent(
-		VkCommandBuffer                             commandBuffer,
-		VkEvent                                     event,
-		VkPipelineStageFlags                        stageMask);
+	static FORCEINLINE_DEBUGGABLE void  vkCmdSetEvent(VkCommandBuffer CommandBuffer, VkEvent Event, VkPipelineStageFlags StageMask)
+	{
+		CmdPrintfBegin(CommandBuffer, FString::Printf(TEXT("vkCmdSetEvent(Event=%p, StageMask=0x%x)"), Event, StageMask));
+
+		VULKANAPINAMESPACE::vkCmdSetEvent(CommandBuffer, Event, StageMask);
+	}
+
+	static FORCEINLINE_DEBUGGABLE void  vkCmdResetEvent(VkCommandBuffer CommandBuffer, VkEvent Event, VkPipelineStageFlags StageMask)
+	{
+		CmdPrintfBegin(CommandBuffer, FString::Printf(TEXT("vkCmdResetEvent(Event=%p, StageMask=0x%x)"), Event, StageMask));
+
+		VULKANAPINAMESPACE::vkCmdResetEvent(CommandBuffer, Event, StageMask);
+	}
+
+	static FORCEINLINE_DEBUGGABLE void  vkCmdWaitEvents(VkCommandBuffer CommandBuffer, uint32 EventCount, const VkEvent* pEvents,
+		VkPipelineStageFlags SrcStageMask, VkPipelineStageFlags DstStageMask,
+		uint32 MemoryBarrierCount, const VkMemoryBarrier* pMemoryBarriers,
+		uint32 BufferMemoryBarrierCount, const VkBufferMemoryBarrier* pBufferMemoryBarriers,
+		uint32 ImageMemoryBarrierCount, const VkImageMemoryBarrier* pImageMemoryBarriers)
+	{
+		CmdPrintfBegin(CommandBuffer, FString::Printf(TEXT("vkCmdWaitEvents()[...]"), Event, StageMask));
+
+		VULKANAPINAMESPACE::vkCmdWaitEvents(CommandBuffer, EventCount, pEvents, SrcStageMask, DstStageMask, MemoryBarrierCount, pMemoryBarriers,
+			BufferMemoryBarrierCount, pBufferMemoryBarriers, ImageMemoryBarrierCount, pImageMemoryBarriers);
+	}
 
-	static FORCEINLINE_DEBUGGABLE void  vkCmdWaitEvents(
-		VkCommandBuffer                             commandBuffer,
-		uint32                                    eventCount,
-		const VkEvent*                              pEvents,
-		VkPipelineStageFlags                        srcStageMask,
-		VkPipelineStageFlags                        dstStageMask,
-		uint32                                    memoryBarrierCount,
-		const VkMemoryBarrier*                      pMemoryBarriers,
-		uint32                                    bufferMemoryBarrierCount,
-		const VkBufferMemoryBarrier*                pBufferMemoryBarriers,
-		uint32                                    imageMemoryBarrierCount,
-		const VkImageMemoryBarrier*                 pImageMemoryBarriers);
-#endif
 	static FORCEINLINE_DEBUGGABLE void  vkCmdPipelineBarrier(
 		VkCommandBuffer CommandBuffer, VkPipelineStageFlags SrcStageMask, VkPipelineStageFlags DstStageMask, VkDependencyFlags DependencyFlags,
 		uint32 MemoryBarrierCount, const VkMemoryBarrier* MemoryBarriers,
diff --git a/Engine/Source/Runtime/VulkanRHI/Private/VulkanMemory.cpp b/Engine/Source/Runtime/VulkanRHI/Private/VulkanMemory.cpp
--- a/Engine/Source/Runtime/VulkanRHI/Private/VulkanMemory.cpp
+++ b/Engine/Source/Runtime/VulkanRHI/Private/VulkanMemory.cpp
@@ -1378,6 +1378,21 @@ namespace VulkanRHI
 	}
 
 
+	FGPUEvent::FGPUEvent(FVulkanDevice* InDevice)
+		: FDeviceChild(InDevice)
+	{
+		VkEventCreateInfo Info;
+		FMemory::Memzero(Info);
+		Info.sType = VK_STRUCTURE_TYPE_EVENT_CREATE_INFO;
+		VERIFYVULKANRESULT(VulkanRHI::vkCreateEvent(InDevice->GetInstanceHandle(), &Info, nullptr, &Handle));
+	}
+
+	FGPUEvent::~FGPUEvent()
+	{
+		VulkanRHI::vkDestroyEvent(Device->GetInstanceHandle(), Handle, nullptr);
+	}
+
+
 	FDeferredDeletionQueue::FDeferredDeletionQueue(FVulkanDevice* InDevice)
 		: FDeviceChild(InDevice)
 	{
diff --git a/Engine/Source/Runtime/VulkanRHI/Public/VulkanMemory.h b/Engine/Source/Runtime/VulkanRHI/Public/VulkanMemory.h
--- a/Engine/Source/Runtime/VulkanRHI/Public/VulkanMemory.h
+++ b/Engine/Source/Runtime/VulkanRHI/Public/VulkanMemory.h
@@ -993,6 +993,21 @@ namespace VulkanRHI
 		void DestroyFence(FFence* Fence);
 	};
 
+	class FGPUEvent : public FDeviceChild, public FRefCount
+	{
+	public:
+		FGPUEvent(FVulkanDevice* InDevice);
+		virtual ~FGPUEvent();
+
+		inline VkEvent GetHandle() const
+		{
+			return Handle;
+		}
+
+	protected:
+		VkEvent Handle;
+	};
+
 	class FDeferredDeletionQueue : public FDeviceChild
 	{
 		//typedef TPair<FRefCountedObject*, uint64> FFencedObject;
diff --git a/Engine/Source/Runtime/Windows/D3D11RHI/Private/D3D11Viewport.cpp b/Engine/Source/Runtime/Windows/D3D11RHI/Private/D3D11Viewport.cpp
--- a/Engine/Source/Runtime/Windows/D3D11RHI/Private/D3D11Viewport.cpp
+++ b/Engine/Source/Runtime/Windows/D3D11RHI/Private/D3D11Viewport.cpp
@@ -240,7 +240,7 @@ void FD3D11Viewport::Resize(uint32 InSizeX, uint32 InSizeY, bool bInIsFullscreen
 	}
 
 	// Float RGBA backbuffers are requested whenever HDR mode is desired
-	if (PixelFormat == PF_FloatRGBA && bIsFullscreen)
+	if (PixelFormat == GRHIHDRDisplayOutputFormat && bIsFullscreen)
 	{
 		D3DRHI->EnableHDR();
 	}
