Temporary solution to broken sky light contribution with mobile renderer

diff --git a/Engine/Source/Runtime/Renderer/Private/ReflectionEnvironmentCapture.cpp b/Engine/Source/Runtime/Renderer/Private/ReflectionEnvironmentCapture.cpp
--- a/Engine/Source/Runtime/Renderer/Private/ReflectionEnvironmentCapture.cpp
+++ b/Engine/Source/Runtime/Renderer/Private/ReflectionEnvironmentCapture.cpp
@@ -1508,7 +1508,7 @@ void CopyToSkyTexture(FRHICommandList& RHICmdList, FScene* Scene, FTexture* Proc
 // Warning: returns before writes to OutIrradianceEnvironmentMap have completed, as they are queued on the rendering thread
 void FScene::UpdateSkyCaptureContents(const USkyLightComponent* CaptureComponent, bool bCaptureEmissiveOnly, UTextureCube* SourceCubemap, FTexture* OutProcessedTexture, float& OutAverageBrightness, FSHVectorRGB3& OutIrradianceEnvironmentMap)
 {	
-	if (GSupportsRenderTargetFormat_PF_FloatRGBA || GetFeatureLevel() >= ERHIFeatureLevel::SM4)
+	if (GetFeatureLevel() >= ERHIFeatureLevel::SM4)
 	{
 		QUICK_SCOPE_CYCLE_COUNTER(STAT_UpdateSkyCaptureContents);
 		{
