Copy of 3348133 from Dev-VR to Dev-VREditor

diff --git a/Engine/Source/Runtime/HeadMountedDisplay/Private/HeadMountedDisplayModule.cpp b/Engine/Source/Runtime/HeadMountedDisplay/Private/HeadMountedDisplayModule.cpp
--- a/Engine/Source/Runtime/HeadMountedDisplay/Private/HeadMountedDisplayModule.cpp
+++ b/Engine/Source/Runtime/HeadMountedDisplay/Private/HeadMountedDisplayModule.cpp
@@ -52,22 +52,6 @@ static TAutoConsoleVariable<int32> CVarHiddenAreaMask(
 	TEXT("0 to disable hidden area mask, 1 to enable."),
 	ECVF_Scalability | ECVF_RenderThreadSafe);
 
-static TAutoConsoleVariable<int32> CVarMonoscopicFarField(
-	TEXT("vr.MonoscopicFarField"),
-	0,
-	TEXT("0 to disable (default), 1 to enable."),
-	ECVF_ReadOnly | ECVF_RenderThreadSafe);
-
-static TAutoConsoleVariable<int32> CVarMonoscopicFarFieldMode(
-	TEXT("vr.MonoscopicFarFieldMode"),
-	1,
-	TEXT("Experimental, mobile only")
-	TEXT(", 0 to disable, 1 to enable (default)")
-	TEXT(", 2 stereo near field only")
-	TEXT(", 3 stereo near field with far field pixel depth test disabled")
-	TEXT(", 4 mono far field only"),
-	ECVF_Scalability | ECVF_RenderThreadSafe);
-
 static void SetTrackingOrigin(const TArray<FString>& Args)
 {
 	int Origin = 0;
diff --git a/Engine/Source/Runtime/Renderer/Private/SceneRendering.cpp b/Engine/Source/Runtime/Renderer/Private/SceneRendering.cpp
--- a/Engine/Source/Runtime/Renderer/Private/SceneRendering.cpp
+++ b/Engine/Source/Runtime/Renderer/Private/SceneRendering.cpp
@@ -96,6 +96,29 @@ static TAutoConsoleVariable<int32> CVarMobileMultiView(
 	TEXT("0 to disable mobile multi-view, 1 to enable.\n"),
 	ECVF_ReadOnly | ECVF_RenderThreadSafe);
 
+static TAutoConsoleVariable<int32> CVarMobileMultiViewDirect(
+	TEXT("vr.MobileMultiView.Direct"),
+	0,
+	TEXT("0 to disable mobile multi-view direct, 1 to enable.\n")
+	TEXT("When enabled the scene color render target array is provided by the hmd plugin so we can skip the blit.\n"),
+	ECVF_ReadOnly | ECVF_RenderThreadSafe);
+
+static TAutoConsoleVariable<int32> CVarMonoscopicFarField(
+	TEXT("vr.MonoscopicFarField"),
+	0,
+	TEXT("0 to disable (default), 1 to enable."),
+	ECVF_ReadOnly | ECVF_RenderThreadSafe);
+
+static TAutoConsoleVariable<int32> CVarMonoscopicFarFieldMode(
+	TEXT("vr.MonoscopicFarFieldMode"),
+	1,
+	TEXT("Experimental, mobile only")
+	TEXT(", 0 to disable, 1 to enable (default)")
+	TEXT(", 2 stereo near field only")
+	TEXT(", 3 stereo near field with far field pixel depth test disabled")
+	TEXT(", 4 mono far field only"),
+	ECVF_Scalability | ECVF_RenderThreadSafe);
+
 #if !(UE_BUILD_SHIPPING || UE_BUILD_TEST)
 static TAutoConsoleVariable<float> CVarGeneralPurposeTweak(
 	TEXT("r.GeneralPurposeTweak"),
