Renaming r.AllowMotionBlur to vr.AllowMotionBlur

diff --git a/Engine/Source/Runtime/Renderer/Private/VelocityRendering.cpp b/Engine/Source/Runtime/Renderer/Private/VelocityRendering.cpp
--- a/Engine/Source/Runtime/Renderer/Private/VelocityRendering.cpp
+++ b/Engine/Source/Runtime/Renderer/Private/VelocityRendering.cpp
@@ -40,7 +40,7 @@ static TAutoConsoleVariable<int32> CVarRHICmdVelocityPassDeferredContexts(
 	TEXT("True to use deferred contexts to parallelize velocity pass command list execution."));
 
 RENDERER_API TAutoConsoleVariable<int32> CVarAllowMotionBlurInVR(
-	TEXT("r.AllowMotionBlurInVR"),
+	TEXT("vr.AllowMotionBlurInVR"),
 	0,
 	TEXT("For projects with motion blur enabled, this allows motion blur to be enabled even while in VR."));
 
