Reverting indirect lighting cache changes (3320581)

diff --git a/Engine/Source/Runtime/Renderer/Private/IndirectLightingCache.cpp b/Engine/Source/Runtime/Renderer/Private/IndirectLightingCache.cpp
--- a/Engine/Source/Runtime/Renderer/Private/IndirectLightingCache.cpp
+++ b/Engine/Source/Runtime/Renderer/Private/IndirectLightingCache.cpp
@@ -326,7 +326,7 @@ void FIndirectLightingCache::CalculateBlockScaleAndAdd(FIntVector InTexelMin, in
 
 FIndirectLightingCacheAllocation* FIndirectLightingCache::AllocatePrimitive(const FPrimitiveSceneInfo* PrimitiveSceneInfo, bool bUnbuiltPreview)
 {
-	const bool bPointSample = !bUnbuiltPreview && PrimitiveSceneInfo->Proxy->GetIndirectLightingCacheQuality() == ILCQ_Point;
+	const bool bPointSample = PrimitiveSceneInfo->Proxy->GetIndirectLightingCacheQuality() == ILCQ_Point || bUnbuiltPreview;
 	const int32 BlockSize = bPointSample ? 1 : GLightingCacheMovableObjectAllocationSize;
 	return PrimitiveAllocations.Add(PrimitiveSceneInfo->PrimitiveComponentId, CreateAllocation(BlockSize, PrimitiveSceneInfo->Proxy->GetBounds(), bPointSample, bUnbuiltPreview));
 }
@@ -686,7 +686,7 @@ void FIndirectLightingCache::UpdateCachePrimitive(
 		{
 			FIndirectLightingCacheAllocation* OriginalAllocation = PrimitiveAllocation;
 			const bool bUnbuiltPreview = bAllowUnbuiltPreview && !bIsMovable;
-			const bool bPointSample = !bUnbuiltPreview && ( PrimitiveSceneProxy->GetIndirectLightingCacheQuality() == ILCQ_Point || !bOpaqueRelevance );
+			const bool bPointSample = PrimitiveSceneProxy->GetIndirectLightingCacheQuality() == ILCQ_Point || bUnbuiltPreview || !bOpaqueRelevance;
 			const int32 BlockSize = bPointSample ? 1 : GLightingCacheMovableObjectAllocationSize;
 
 			// Light with the cumulative bounds of the entire attachment group
