Copying //UE4/Ocean-Staging to //UE4/Dev-Main (Source: //Ocean/Main @ 3355692)

==========================
MAJOR FEATURES + CHANGES
==========================

Change 3348195 on 2017/03/15 by [REDACTED_PERSON]

	Hotfix from CL 3347729 for variable shadow warning caused by CL 3347393

Change 3347935 on 2017/03/15 by [REDACTED_PERSON]

	Hotfix from CL 3347393 for a runtime asset crash

Change 3340537 on 2017/03/09 by [REDACTED_PERSON]

	Hotfix from CL 3338679 for arrow keys not working to navigate between elements in the details panel

Change 3335755 on 2017/03/07 by [REDACTED_PERSON]

	Pass bCreateStubIPA into PrepForUATPackageOrDeploy rather than setting it through a global configuration object (which is no longer accessible).

Change 3334652 on 2017/03/06 by [REDACTED_PERSON]

	Merging //UE4/Ocean-Staging to //Ocean/Main @ CL 3331740

Change 3332441 on 2017/03/03 by [REDACTED_PERSON]

	Hotfix from CL 3296668 for assert when creating memory reports

Change 3318933 on 2017/02/23 by [REDACTED_PERSON]

	Temp workaround for iOS metal sRGB capable render targets sRGB correcting twice.

Change 3314064 on 2017/02/21 by [REDACTED_PERSON]

	Implement randomized manifest filenames for Ocean.

Change 3313139 on 2017/02/20 by [REDACTED_PERSON]

	Hotfix from CL 3262052 for UE-40782 Put back the skinxx workflow for the staticmesh

Change 3292678 on 2017/02/08 by [REDACTED_PERSON]

	Back out changelist 3292455

Change 3292455 on 2017/02/08 by [REDACTED_PERSON]

	Back out changelist 3292268

Change 3292268 on 2017/02/08 by [REDACTED_PERSON]

	Restore DestStride output value when locking compressed textures.

Change 3288925 on 2017/02/06 by [REDACTED_PERSON]

	Change this error from a check to an error log: "Loading non-streamed mips from an external bulk file.  This is not desireable."

Change 3288441 on 2017/02/06 by [REDACTED_PERSON]

	Hotfix from CL 3255451 for UE-40384: Reference Viewer comment text is difficult to read

Change 3285501 on 2017/02/03 by [REDACTED_PERSON]

	Fix  for ring buffer not allocating sufficient space for large requests.

Change 3280062 on 2017/01/31 by [REDACTED_PERSON]

	Sequencer: Find object in any package. This fixes an issue where bound objects can't be found.

Change 3279375 on 2017/01/31 by [REDACTED_PERSON]

	Increase TaskGraph thread stack size for non-shipping builds to prevent stack overflow on iOS.

Change 3279366 on 2017/01/31 by [REDACTED_PERSON]

	Fix for iOS virtual keyboard.

Change 3275965 on 2017/01/27 by [REDACTED_PERSON]

	Sequencer: Only reset persistent evaluation data when the sequence has changed
	  - This ensures that we don't destroy persistent data that is assumed to still exist (i.e. it was created in ::Setup) from the same sequence

	Copy from Release-4.15

Change 3275963 on 2017/01/27 by [REDACTED_PERSON]

	Sequencer: Fixed crash caused by lingering persistent evaluation data

	Copy from Release-4.15

Change 3275962 on 2017/01/27 by [REDACTED_PERSON]

	Sequencer: Evaluation templates are now only fully rebuilt in PIE, and will not re-cycle track identifiers
	  - This addresses issues with newly compiled tracks recycling the persistent data of old stale tracks.
	  - This commit also ensures we don't fully rebuild templates in the editor when in Sequencer

	Copy from Release-4.15

Change 3275523 on 2017/01/27 by [REDACTED_PERSON]

	Fix issue with automation tool caused by a bad resolve

Change 3274292 on 2017/01/26 by [REDACTED_PERSON]

	Merging //UE4/Ocean-Staging to //Ocean/Main @ CL 3269883

Change 3252898 on 2017/01/10 by [REDACTED_PERSON]

	Hotfix from CL 3187276 for UE-39071: Crash on adding Actor Reference property to a Struct

Change 3248463 on 2017/01/05 by [REDACTED_PERSON]

	Implementing UE-39421 fix into Ocean

Change 3247933 on 2017/01/05 by [REDACTED_PERSON]

	Hotfix from CL 3246828 for UE-39249: Particle effects only render in one eye for VR

Change 3242274 on 2016/12/21 by [REDACTED_PERSON]

	Hotfix from CL 3209235 for UE-16688 - Selection and selection highlight doesn't work on child actors

Change 3242273 on 2016/12/21 by [REDACTED_PERSON]

	Hotfix from CL 3209232 for making the function GetComponents optionally include components in child actors
	Also fixed a circular dependency due to the hotfix (already applied in UE4 Main)

Change 3240035 on 2016/12/19 by [REDACTED_PERSON]

	MovieSceneEventTemplate - Renamed FEventData to FMovieSceneEventData to eliminate confusion and clear up a crash occuring in Ocean debug builds.

Change 3236490 on 2016/12/15 by [REDACTED_PERSON]

	Updating packaged build to not make a stub.  This was stopping all icons from making it into the .ipa and messing up the resolution settings.
	also cleaned up the local staging directories.

Change 3231487 on 2016/12/12 by [REDACTED_PERSON]

	Maintain aspect ratio when render target is clamped.

Change 3229050 on 2016/12/09 by [REDACTED_PERSON]

	#ocn - Hotfix from CL 3218104.  Fixes code signing errors related to iCloud support.

Change 3225632 on 2016/12/07 by [REDACTED_PERSON]

	New Chair Entertainment Group iOS Distribution certificate .p12.

Change 3224015 on 2016/12/06 by [REDACTED_PERSON]

	#ocn - Hotfix from CL 3223546 fixes iOS audio crash after last merge from main.

Change 3221665 on 2016/12/05 by [REDACTED_PERSON]

	Merging hotfixes from //UE4/Ocean-Staging to //Ocean/Main
	Hotfix from CL 3216832 for non-unity build failures
	Hotfix from CL 3217564 for fixing the FixDylib script being run on too many modules. Mac build farm is oversaturated right now.

Change 3217759 on 2016/12/01 by [REDACTED_PERSON]

	Merging //UE4/Ocean-Staging to //Ocean/Main @ 3215855

Change 3214085 on 2016/11/29 by [REDACTED_PERSON]

	Removed logic in K2Node_BaseAsyncTask that excluded properties defined on super classes. Check with engine team prior to making change.

diff --git a/Engine/Plugins/Runtime/OculusRift/Source/OculusRift/Private/OculusRiftCommon.cpp b/Engine/Plugins/Runtime/OculusRift/Source/OculusRift/Private/OculusRiftCommon.cpp
--- a/Engine/Plugins/Runtime/OculusRift/Source/OculusRift/Private/OculusRiftCommon.cpp
+++ b/Engine/Plugins/Runtime/OculusRift/Source/OculusRift/Private/OculusRiftCommon.cpp
@@ -31,7 +31,10 @@ void FOvrSessionShared::Destroy()
 	{
 		UE_LOG(LogHMD, Log, TEXT("Destroying ovr session"));
 		FSuspendRenderingThread SuspendRenderingThread(false);
-		check(LockCnt.GetValue() == 0);
+		if (LockCnt.GetValue() != 0)
+		{
+			UE_LOG(LogHMD, Warning, TEXT("OVR:  Lock count is not 0!  Check for multiple sessions having been started."));
+		}
 
 		DelegatesLock.Lock();
 		DestoryDelegate.Broadcast(Session);
diff --git a/Engine/Shaders/PostProcessTonemap.usf b/Engine/Shaders/PostProcessTonemap.usf
--- a/Engine/Shaders/PostProcessTonemap.usf
+++ b/Engine/Shaders/PostProcessTonemap.usf
@@ -599,6 +599,7 @@ half2 CocBlendScaleBiasFine()
 }
 
 float4 OverlayColor;
+float SRGBAwareTarget;
 
 void MainPS_ES2(
 	in float4 UVAndScreenPos : TEXCOORD0,
@@ -724,13 +725,23 @@ void MainPS_ES2(
 		#endif
 
 		half3 FilmColor = FilmPostProcess(LinearColor.rgb);
-		#if IOS
+		half3 TonemappedColor;
+		if (SRGBAwareTarget > 0.5)
+		{
+			// skip sRGB conversion if HW conversion will take place.
+			TonemappedColor = FilmColor;
+		}
+		else
+		{
+#if IOS
 			// Note, iOS native output is raw gamma 2.2 not sRGB!
-			half3 TonemappedColor = pow(FilmColor.rgb, 1.0/2.2);
-		#else
+			TonemappedColor = pow(FilmColor.rgb, 1.0 / 2.2);
+#else
 			// Apply conversion to sRGB (this must be an exact sRGB conversion else darks are bad).
-			half3 TonemappedColor = LinearToSrgbBranchless(FilmColor);
-		#endif
+			TonemappedColor = LinearToSrgbBranchless(FilmColor);
+#endif
+		}
+
 
 		// Blend with custom LDR color, used for Fade track in Matinee.
 		// This is the 101% wrong way to do this,
diff --git a/Engine/Source/Editor/BlueprintGraph/Private/K2Node_BaseAsyncTask.cpp b/Engine/Source/Editor/BlueprintGraph/Private/K2Node_BaseAsyncTask.cpp
--- a/Engine/Source/Editor/BlueprintGraph/Private/K2Node_BaseAsyncTask.cpp
+++ b/Engine/Source/Editor/BlueprintGraph/Private/K2Node_BaseAsyncTask.cpp
@@ -79,7 +79,7 @@ void UK2Node_BaseAsyncTask::AllocateDefaultPins()
 	}
 
 	UFunction* DelegateSignatureFunction = NULL;
-	for (TFieldIterator<UProperty> PropertyIt(ProxyClass, EFieldIteratorFlags::ExcludeSuper); PropertyIt; ++PropertyIt)
+	for (TFieldIterator<UProperty> PropertyIt(ProxyClass); PropertyIt; ++PropertyIt)
 	{
 		if (UMulticastDelegateProperty* Property = Cast<UMulticastDelegateProperty>(*PropertyIt))
 		{
@@ -340,7 +340,7 @@ void UK2Node_BaseAsyncTask::ExpandNode(class FKismetCompilerContext& CompilerCon
 	bIsErrorFree &= Schema->TryCreateConnection(LastThenPin, ValidateProxyNode->GetExecPin());
 	LastThenPin = ValidateProxyNode->GetThenPin();
 
-	for (TFieldIterator<UMulticastDelegateProperty> PropertyIt(ProxyClass, EFieldIteratorFlags::ExcludeSuper); PropertyIt && bIsErrorFree; ++PropertyIt)
+	for (TFieldIterator<UMulticastDelegateProperty> PropertyIt(ProxyClass); PropertyIt && bIsErrorFree; ++PropertyIt)
 	{
 		bIsErrorFree &= FBaseAsyncTaskHelper::HandleDelegateImplementation(*PropertyIt, VariableOutputs, ProxyObjectPin, LastThenPin, this, SourceGraph, CompilerContext);
 	}
diff --git a/Engine/Source/Programs/AutomationTool/IOS/IOSPlatform.Automation.cs b/Engine/Source/Programs/AutomationTool/IOS/IOSPlatform.Automation.cs
--- a/Engine/Source/Programs/AutomationTool/IOS/IOSPlatform.Automation.cs
+++ b/Engine/Source/Programs/AutomationTool/IOS/IOSPlatform.Automation.cs
@@ -330,8 +330,8 @@ public class IOSPlatform : Platform
 				CombinePaths(SC.LocalRoot, "Engine"),
 				Params.Distribution, 
 				"",
-				false,
-				true);
+				false, 
+				false);
 
 			// figure out where to pop in the staged files
 			string AppDirectory = string.Format("{0}/Payload/{1}.app",
diff --git a/Engine/Source/Runtime/Apple/MetalRHI/Private/MetalBufferPools.cpp b/Engine/Source/Runtime/Apple/MetalRHI/Private/MetalBufferPools.cpp
--- a/Engine/Source/Runtime/Apple/MetalRHI/Private/MetalBufferPools.cpp
+++ b/Engine/Source/Runtime/Apple/MetalRHI/Private/MetalBufferPools.cpp
@@ -279,10 +279,13 @@ uint32 FRingBuffer::Allocate(uint32 Size, uint32 Alignment)
 	}
 	else
 	{
-		uint32 BufferSize = Buffer.length;
-		UE_LOG(LogMetal, Warning, TEXT("Reallocating ring-buffer from %d to %d to avoid wrapping write at offset %d into outstanding buffer region %d at frame %lld]"), BufferSize, BufferSize * 2, Offset, LastRead, GFrameCounter);
+		const uint32 BufferSize = Buffer.length;
+		uint32 NewBufferSize;
+		for(NewBufferSize = BufferSize << 1; NewBufferSize < Size; NewBufferSize <<= 1) { }
+
+		UE_LOG(LogMetal, Warning, TEXT("Reallocating ring-buffer from %d to %d to avoid wrapping write at offset %d into outstanding buffer region %d at frame %lld]"), BufferSize, NewBufferSize, Offset, LastRead, GFrameCounter);
 		SafeReleaseMetalResource(Buffer);
-		Buffer = [GetMetalDeviceContext().GetDevice() newBufferWithLength:BufferSize * 2 options:BUFFER_CACHE_MODE];
+		Buffer = [GetMetalDeviceContext().GetDevice() newBufferWithLength:NewBufferSize options:BUFFER_CACHE_MODE];
 		TRACK_OBJECT(STAT_MetalBufferCount, Buffer);
 		Offset = 0;
 		LastRead = Size;
diff --git a/Engine/Source/Runtime/Apple/MetalRHI/Private/MetalTexture.cpp b/Engine/Source/Runtime/Apple/MetalRHI/Private/MetalTexture.cpp
--- a/Engine/Source/Runtime/Apple/MetalRHI/Private/MetalTexture.cpp
+++ b/Engine/Source/Runtime/Apple/MetalRHI/Private/MetalTexture.cpp
@@ -1175,7 +1175,7 @@ void* FMetalSurface::Lock(uint32 MipIndex, uint32 ArrayIndex, EResourceLockMode
 #else
 	bool const bAlignedStride = false; // Always Shared memory, so won't use the blit command encoder operations.
 #endif
-	const uint32 MipBytes = GetMipSize(MipIndex, &DestStride, false, bAlignedStride);
+	uint32 MipBytes = GetMipSize(MipIndex, &DestStride, false, bAlignedStride);
 	
 	// allocate some temporary memory
 	if(!LockedMemory[MipIndex])
@@ -1224,7 +1224,14 @@ void* FMetalSurface::Lock(uint32 MipIndex, uint32 ArrayIndex, EResourceLockMode
 				//kick the current command buffer.
 				GetMetalDeviceContext().SubmitCommandBufferAndWait();
 #endif
-				[Texture getBytes:[LockedMemory[MipIndex] contents] bytesPerRow:DestStride bytesPerImage:MipBytes fromRegion:Region mipmapLevel:MipIndex slice:ArrayIndex];
+				uint32 BytesPerRow = DestStride;
+				if (PixelFormat == PF_PVRTC2 || PixelFormat == PF_PVRTC4)
+				{
+					// for compressed textures metal debug RT expects 0 for rowBytes and imageBytes.
+					BytesPerRow = 0;
+					MipBytes = 0;
+				}
+				[Texture getBytes:[LockedMemory[MipIndex] contents] bytesPerRow:BytesPerRow bytesPerImage:MipBytes fromRegion:Region mipmapLevel:MipIndex slice:ArrayIndex];
 			}
 			
 #if PLATFORM_MAC
diff --git a/Engine/Source/Runtime/Core/Private/Async/TaskGraph.cpp b/Engine/Source/Runtime/Core/Private/Async/TaskGraph.cpp
--- a/Engine/Source/Runtime/Core/Private/Async/TaskGraph.cpp
+++ b/Engine/Source/Runtime/Core/Private/Async/TaskGraph.cpp
@@ -1490,7 +1490,12 @@ public:
 				Name = FString::Printf(TEXT("TaskGraphThreadNP %d"), ThreadIndex - (LastExternalThread + 1));
 				ThreadPri = TPri_BelowNormal; // we want normal tasks below normal threads like the game thread
 			}
+
+#if ( UE_BUILD_SHIPPING || UE_BUILD_TEST )
 			uint32 StackSize = 384 * 1024;
+#else
+			uint32 StackSize = 512 * 1024;
+#endif
 			WorkerThreads[ThreadIndex].RunnableThread = FRunnableThread::Create(&Thread(ThreadIndex), *Name, StackSize, ThreadPri, Affinity); // these are below normal threads so that they sleep when the named threads are active
 			WorkerThreads[ThreadIndex].bAttached = true;
 		}
diff --git a/Engine/Source/Runtime/Engine/Private/Rendering/SkinWeightVertexBuffer.cpp b/Engine/Source/Runtime/Engine/Private/Rendering/SkinWeightVertexBuffer.cpp
--- a/Engine/Source/Runtime/Engine/Private/Rendering/SkinWeightVertexBuffer.cpp
+++ b/Engine/Source/Runtime/Engine/Private/Rendering/SkinWeightVertexBuffer.cpp
@@ -109,7 +109,7 @@ FArchive& operator<<(FArchive& Ar, FSkinWeightVertexBuffer& VertexBuffer)
 			if (!Ar.IsCountingMemory())
 			{
 				// update cached buffer info
-				VertexBuffer.Data = (VertexBuffer.NumVertices > 0) ? VertexBuffer.WeightData->GetDataPointer() : nullptr;
+				VertexBuffer.Data = (VertexBuffer.NumVertices > 0 && VertexBuffer.WeightData->GetResourceArray()->GetResourceDataSize()) ? VertexBuffer.WeightData->GetDataPointer() : nullptr;
 				VertexBuffer.Stride = VertexBuffer.WeightData->GetStride();
 			}
 		}
diff --git a/Engine/Source/Runtime/Engine/Private/TextureDerivedData.cpp b/Engine/Source/Runtime/Engine/Private/TextureDerivedData.cpp
--- a/Engine/Source/Runtime/Engine/Private/TextureDerivedData.cpp
+++ b/Engine/Source/Runtime/Engine/Private/TextureDerivedData.cpp
@@ -1176,7 +1176,13 @@ bool FTexturePlatformData::TryLoadMips(int32 FirstMipToLoad, void** OutMipData)
 			if (OutMipData)
 			{
 				OutMipData[MipIndex - FirstMipToLoad] = FMemory::Malloc(Mip.BulkData.GetBulkDataSize());
+#if 0
 				checkSlow(!Mip.BulkData.GetFilename().EndsWith(TEXT(".ubulk"))); // We want to make sure that any non-streamed mips are coming from the texture asset file, and not from an external bulk file
+#else
+				UE_CLOG(Mip.BulkData.GetFilename().EndsWith(TEXT(".ubulk")), LogTexture, Error, TEXT("Loading non-streamed mips from an external bulk file.  This is not desireable.  File %s"), *(Mip.BulkData.GetFilename() ) );
+#endif
+				
+				
 				Mip.BulkData.GetCopy(&OutMipData[MipIndex - FirstMipToLoad]);
 			}
 			NumMipsCached++;
diff --git a/Engine/Source/Runtime/Engine/Private/TextureRenderTarget2D.cpp b/Engine/Source/Runtime/Engine/Private/TextureRenderTarget2D.cpp
--- a/Engine/Source/Runtime/Engine/Private/TextureRenderTarget2D.cpp
+++ b/Engine/Source/Runtime/Engine/Private/TextureRenderTarget2D.cpp
@@ -159,6 +159,9 @@ void UTextureRenderTarget2D::PostEditChangeProperty(FPropertyChangedEvent& Prope
 
 void UTextureRenderTarget2D::PostLoad()
 {
+	float OriginalSizeX = SizeX;
+	float OriginalSizeY = SizeY;
+	
 	if (!FPlatformProperties::SupportsWindowedMode())
 	{
 		// Clamp the render target size in order to avoid reallocating the scene render targets,
@@ -170,6 +173,22 @@ void UTextureRenderTarget2D::PostLoad()
 	SizeX = FMath::Min<int32>(SizeX, GTextureRenderTarget2DMaxSizeX);
 	SizeY = FMath::Min<int32>(SizeY, GTextureRenderTarget2DMaxSizeY);
 	
+	// Maintain aspect ratio if clamped
+	if( SizeX != OriginalSizeX || SizeY != OriginalSizeY )
+	{
+		float ScaleX = SizeX / OriginalSizeX;
+		float ScaleY = SizeY / OriginalSizeY;
+		
+		if( ScaleX < ScaleY )
+		{
+			SizeY = OriginalSizeY * ScaleX;
+		}
+		else
+		{
+			SizeX = OriginalSizeX * ScaleY;
+		}
+	}
+	
 	Super::PostLoad();
 }
 
diff --git a/Engine/Source/Runtime/MovieSceneTracks/Private/Evaluation/MovieSceneEventTemplate.cpp b/Engine/Source/Runtime/MovieSceneTracks/Private/Evaluation/MovieSceneEventTemplate.cpp
--- a/Engine/Source/Runtime/MovieSceneTracks/Private/Evaluation/MovieSceneEventTemplate.cpp
+++ b/Engine/Source/Runtime/MovieSceneTracks/Private/Evaluation/MovieSceneEventTemplate.cpp
@@ -12,10 +12,9 @@
 
 DECLARE_CYCLE_STAT(TEXT("Event Track Token Execute"), MovieSceneEval_EventTrack_TokenExecute, STATGROUP_MovieSceneEval);
 
-
-struct FEventData
+struct FMovieSceneEventData
 {
-	FEventData(const FEventPayload& InPayload, float InGlobalPosition) : Payload(InPayload), GlobalPosition(InGlobalPosition) {}
+	FMovieSceneEventData(const FEventPayload& InPayload, float InGlobalPosition) : Payload(InPayload), GlobalPosition(InGlobalPosition) {}
 
 	FEventPayload Payload;
 	float GlobalPosition;
@@ -25,7 +24,7 @@ struct FEventData
 struct FEventTrackExecutionToken
 	: IMovieSceneExecutionToken
 {
-	FEventTrackExecutionToken(TArray<FEventData> InEvents) : Events(MoveTemp(InEvents)) {}
+	FEventTrackExecutionToken(TArray<FMovieSceneEventData> InEvents) : Events(MoveTemp(InEvents)) {}
 
 	/** Execute this token, operating on all objects referenced by 'Operand' */
 	virtual void Execute(const FMovieSceneContext& Context, const FMovieSceneEvaluationOperand& Operand, FPersistentEvaluationData& PersistentData, IMovieScenePlayer& Player) override
@@ -41,7 +40,7 @@ struct FEventTrackExecutionToken
 				continue;
 			}
 
-			for (FEventData& Event : Events)
+			for (FMovieSceneEventData& Event : Events)
 			{
 #if !UE_BUILD_SHIPPING
 				if (Event.Payload.EventName == NAME_PerformanceCapture)
@@ -71,7 +70,7 @@ struct FEventTrackExecutionToken
 #endif	// UE_BUILD_SHIPPING
 	}
 
-	void TriggerEvent(FEventData& Event, UObject& EventContextObject, IMovieScenePlayer& Player)
+	void TriggerEvent(FMovieSceneEventData& Event, UObject& EventContextObject, IMovieScenePlayer& Player)
 	{
 		UFunction* EventFunction = EventContextObject.FindFunction(Event.Payload.EventName);
 
@@ -128,7 +127,7 @@ struct FEventTrackExecutionToken
 		}
 	}
 
-	TArray<FEventData> Events;
+	TArray<FMovieSceneEventData> Events;
 };
 
 FMovieSceneEventSectionTemplate::FMovieSceneEventSectionTemplate(const UMovieSceneEventSection& Section, const UMovieSceneEventTrack& Track)
@@ -156,7 +155,7 @@ void FMovieSceneEventSectionTemplate::EvaluateSwept(const FMovieSceneEvaluationO
 		return;
 	}
 
-	TArray<FEventData> Events;
+	TArray<FMovieSceneEventData> Events;
 
 	TRange<float> SweptRange = Context.GetRange();
 
@@ -177,7 +176,7 @@ void FMovieSceneEventSectionTemplate::EvaluateSwept(const FMovieSceneEvaluationO
 			float Time = KeyTimes[KeyIndex];
 			if (SweptRange.Contains(Time))
 			{
-				Events.Add(FEventData(KeyValues[KeyIndex], Position));
+				Events.Add(FMovieSceneEventData(KeyValues[KeyIndex], Position));
 			}
 		}
 	}
@@ -187,7 +186,7 @@ void FMovieSceneEventSectionTemplate::EvaluateSwept(const FMovieSceneEvaluationO
 		float Time = KeyTimes[KeyIndex];
 		if (SweptRange.Contains(Time))
 		{
-			Events.Add(FEventData(KeyValues[KeyIndex], Position));
+			Events.Add(FMovieSceneEventData(KeyValues[KeyIndex], Position));
 		}
 	}
 
diff --git a/Engine/Source/Runtime/Renderer/Private/PostProcess/PostProcessTonemap.cpp b/Engine/Source/Runtime/Renderer/Private/PostProcess/PostProcessTonemap.cpp
--- a/Engine/Source/Runtime/Renderer/Private/PostProcess/PostProcessTonemap.cpp
+++ b/Engine/Source/Runtime/Renderer/Private/PostProcess/PostProcessTonemap.cpp
@@ -1446,6 +1446,7 @@ public:
 
 	FShaderParameter OverlayColor;
 	FShaderParameter FringeIntensity;
+	FShaderParameter SRGBAwareTargetParam;
 
 	/** Initialization constructor. */
 	FPostProcessTonemapPS_ES2(const ShaderMetaType::CompiledShaderInitializerType& Initializer)
@@ -1470,6 +1471,9 @@ public:
 
 		OverlayColor.Bind(Initializer.ParameterMap, TEXT("OverlayColor"));
 		FringeIntensity.Bind(Initializer.ParameterMap, TEXT("FringeIntensity"));
+
+		
+		SRGBAwareTargetParam.Bind(Initializer.ParameterMap, TEXT("SRGBAwareTarget"));
 	}
 	
 	// FShader interface.
@@ -1480,12 +1484,13 @@ public:
 			<< TexScale << GrainScaleBiasJitter << TonemapperParams
 			<< ColorMatrixR_ColorCurveCd1 << ColorMatrixG_ColorCurveCd3Cm3 << ColorMatrixB_ColorCurveCm2 << ColorCurve_Cm0Cd0_Cd2_Ch0Cm1_Ch3 << ColorCurve_Ch1_Ch2 << ColorShadow_Luma << ColorShadow_Tint1 << ColorShadow_Tint2
 			<< OverlayColor
-			<< FringeIntensity;
+			<< FringeIntensity
+			<< SRGBAwareTargetParam;
 
 		return bShaderHasOutdatedParameters;
 	}
 
-	void SetPS(const FRenderingCompositePassContext& Context)
+	void SetPS(const FRenderingCompositePassContext& Context, bool bSRGBAwareTarget)
 	{
 		const FPostProcessSettings& Settings = Context.View.FinalPostProcessSettings;
 		const FSceneViewFamily& ViewFamily = *(Context.View.Family);
@@ -1561,6 +1566,8 @@ public:
 			SetShaderValue(Context.RHICmdList, ShaderRHI, ColorShadow_Tint1, Constants[6]);
 			SetShaderValue(Context.RHICmdList, ShaderRHI, ColorShadow_Tint2, Constants[7]);
 		}
+
+		SetShaderValue(Context.RHICmdList, ShaderRHI, SRGBAwareTargetParam, bSRGBAwareTarget ? 1.0f : 0.0f );
 	}
 	
 	static const TCHAR* GetSourceFilename()
@@ -1652,7 +1659,7 @@ namespace PostProcessTonemap_ES2Util
 {
 	// Template implementation supports unique static BoundShaderState for each permutation of Pixel Shaders 
 	template <uint32 ConfigIndex>
-	static inline void SetShaderTemplES2(const FRenderingCompositePassContext& Context, bool bUsedFramebufferFetch)
+	static inline void SetShaderTemplES2(const FRenderingCompositePassContext& Context, bool bUsedFramebufferFetch, bool bSRGBAwareTarget)
 	{
 		FGraphicsPipelineStateInitializer GraphicsPSOInit;
 		Context.RHICmdList.ApplyCachedRenderTargets(GraphicsPSOInit);
@@ -1673,14 +1680,15 @@ namespace PostProcessTonemap_ES2Util
 		SetGraphicsPipelineState(Context.RHICmdList, GraphicsPSOInit);
 
 		VertexShader->SetVS(Context);
-		PixelShader->SetPS(Context);
+		PixelShader->SetPS(Context, bSRGBAwareTarget);
 	}
 }
 
-FRCPassPostProcessTonemapES2::FRCPassPostProcessTonemapES2(const FViewInfo& InView, bool bInUsedFramebufferFetch)
+FRCPassPostProcessTonemapES2::FRCPassPostProcessTonemapES2(const FViewInfo& InView, bool bInUsedFramebufferFetch, bool bInSRGBAwareTarget)
 	: bDoScreenPercentageInTonemapper(false)
 	, View(InView)
 	, bUsedFramebufferFetch(bInUsedFramebufferFetch)
+	, bSRGBAwareTarget(bInSRGBAwareTarget)
 {
 	uint32 ConfigBitmask = TonemapperGenerateBitmaskMobile(&View, false);
 	ConfigIndexMobile = TonemapperFindLeastExpensive(TonemapperConfBitmaskMobile, sizeof(TonemapperConfBitmaskMobile)/4, TonemapperCostTab, ConfigBitmask);
@@ -1736,49 +1744,49 @@ void FRCPassPostProcessTonemapES2::Process(FRenderingCompositePassContext& Conte
 
 	const int32 ConfigOverride = CVarTonemapperOverride->GetInt();
 	const uint32 FinalConfigIndex = ConfigOverride == -1 ? ConfigIndexMobile : (int32)ConfigOverride;
-
+	
 	switch(FinalConfigIndex)
 	{
 		using namespace PostProcessTonemap_ES2Util;
-		case 0:	SetShaderTemplES2<0>(Context, bUsedFramebufferFetch); break;
-		case 1:	SetShaderTemplES2<1>(Context, bUsedFramebufferFetch); break;
-		case 2:	SetShaderTemplES2<2>(Context, bUsedFramebufferFetch); break;
-		case 3:	SetShaderTemplES2<3>(Context, bUsedFramebufferFetch); break;
-		case 4:	SetShaderTemplES2<4>(Context, bUsedFramebufferFetch); break;
-		case 5:	SetShaderTemplES2<5>(Context, bUsedFramebufferFetch); break;
-		case 6:	SetShaderTemplES2<6>(Context, bUsedFramebufferFetch); break;
-		case 7:	SetShaderTemplES2<7>(Context, bUsedFramebufferFetch); break;
-		case 8:	SetShaderTemplES2<8>(Context, bUsedFramebufferFetch); break;
-		case 9:	SetShaderTemplES2<9>(Context, bUsedFramebufferFetch); break;
-		case 10: SetShaderTemplES2<10>(Context, bUsedFramebufferFetch); break;
-		case 11: SetShaderTemplES2<11>(Context, bUsedFramebufferFetch); break;
-		case 12: SetShaderTemplES2<12>(Context, bUsedFramebufferFetch); break;
-		case 13: SetShaderTemplES2<13>(Context, bUsedFramebufferFetch); break;
-		case 14: SetShaderTemplES2<14>(Context, bUsedFramebufferFetch); break;
-		case 15: SetShaderTemplES2<15>(Context, bUsedFramebufferFetch); break;
-		case 16: SetShaderTemplES2<16>(Context, bUsedFramebufferFetch); break;
-		case 17: SetShaderTemplES2<17>(Context, bUsedFramebufferFetch); break;
-		case 18: SetShaderTemplES2<18>(Context, bUsedFramebufferFetch); break;
-		case 19: SetShaderTemplES2<19>(Context, bUsedFramebufferFetch); break;
-		case 20: SetShaderTemplES2<20>(Context, bUsedFramebufferFetch); break;
-		case 21: SetShaderTemplES2<21>(Context, bUsedFramebufferFetch); break;
-		case 22: SetShaderTemplES2<22>(Context, bUsedFramebufferFetch); break;
-		case 23: SetShaderTemplES2<23>(Context, bUsedFramebufferFetch); break;
-		case 24: SetShaderTemplES2<24>(Context, bUsedFramebufferFetch); break;
-		case 25: SetShaderTemplES2<25>(Context, bUsedFramebufferFetch); break;
-		case 26: SetShaderTemplES2<26>(Context, bUsedFramebufferFetch); break;
-		case 27: SetShaderTemplES2<27>(Context, bUsedFramebufferFetch); break;
-		case 28: SetShaderTemplES2<28>(Context, bUsedFramebufferFetch); break;
-		case 29: SetShaderTemplES2<29>(Context, bUsedFramebufferFetch); break;
-		case 30: SetShaderTemplES2<30>(Context, bUsedFramebufferFetch); break;
-		case 31: SetShaderTemplES2<31>(Context, bUsedFramebufferFetch); break;
-		case 32: SetShaderTemplES2<32>(Context, bUsedFramebufferFetch); break;
-		case 33: SetShaderTemplES2<33>(Context, bUsedFramebufferFetch); break;
-		case 34: SetShaderTemplES2<34>(Context, bUsedFramebufferFetch); break;
-		case 35: SetShaderTemplES2<35>(Context, bUsedFramebufferFetch); break;
-		case 36: SetShaderTemplES2<36>(Context, bUsedFramebufferFetch); break;
-		case 37: SetShaderTemplES2<37>(Context, bUsedFramebufferFetch); break;
-		case 38: SetShaderTemplES2<38>(Context, bUsedFramebufferFetch); break;
+		case 0:	SetShaderTemplES2<0>(Context, bUsedFramebufferFetch, bSRGBAwareTarget); break;
+		case 1:	SetShaderTemplES2<1>(Context, bUsedFramebufferFetch, bSRGBAwareTarget); break;
+		case 2:	SetShaderTemplES2<2>(Context, bUsedFramebufferFetch, bSRGBAwareTarget); break;
+		case 3:	SetShaderTemplES2<3>(Context, bUsedFramebufferFetch, bSRGBAwareTarget); break;
+		case 4:	SetShaderTemplES2<4>(Context, bUsedFramebufferFetch, bSRGBAwareTarget); break;
+		case 5:	SetShaderTemplES2<5>(Context, bUsedFramebufferFetch, bSRGBAwareTarget); break;
+		case 6:	SetShaderTemplES2<6>(Context, bUsedFramebufferFetch, bSRGBAwareTarget); break;
+		case 7:	SetShaderTemplES2<7>(Context, bUsedFramebufferFetch, bSRGBAwareTarget); break;
+		case 8:	SetShaderTemplES2<8>(Context, bUsedFramebufferFetch, bSRGBAwareTarget); break;
+		case 9:	SetShaderTemplES2<9>(Context, bUsedFramebufferFetch, bSRGBAwareTarget); break;
+		case 10: SetShaderTemplES2<10>(Context, bUsedFramebufferFetch, bSRGBAwareTarget); break;
+		case 11: SetShaderTemplES2<11>(Context, bUsedFramebufferFetch, bSRGBAwareTarget); break;
+		case 12: SetShaderTemplES2<12>(Context, bUsedFramebufferFetch, bSRGBAwareTarget); break;
+		case 13: SetShaderTemplES2<13>(Context, bUsedFramebufferFetch, bSRGBAwareTarget); break;
+		case 14: SetShaderTemplES2<14>(Context, bUsedFramebufferFetch, bSRGBAwareTarget); break;
+		case 15: SetShaderTemplES2<15>(Context, bUsedFramebufferFetch, bSRGBAwareTarget); break;
+		case 16: SetShaderTemplES2<16>(Context, bUsedFramebufferFetch, bSRGBAwareTarget); break;
+		case 17: SetShaderTemplES2<17>(Context, bUsedFramebufferFetch, bSRGBAwareTarget); break;
+		case 18: SetShaderTemplES2<18>(Context, bUsedFramebufferFetch, bSRGBAwareTarget); break;
+		case 19: SetShaderTemplES2<19>(Context, bUsedFramebufferFetch, bSRGBAwareTarget); break;
+		case 20: SetShaderTemplES2<20>(Context, bUsedFramebufferFetch, bSRGBAwareTarget); break;
+		case 21: SetShaderTemplES2<21>(Context, bUsedFramebufferFetch, bSRGBAwareTarget); break;
+		case 22: SetShaderTemplES2<22>(Context, bUsedFramebufferFetch, bSRGBAwareTarget); break;
+		case 23: SetShaderTemplES2<23>(Context, bUsedFramebufferFetch, bSRGBAwareTarget); break;
+		case 24: SetShaderTemplES2<24>(Context, bUsedFramebufferFetch, bSRGBAwareTarget); break;
+		case 25: SetShaderTemplES2<25>(Context, bUsedFramebufferFetch, bSRGBAwareTarget); break;
+		case 26: SetShaderTemplES2<26>(Context, bUsedFramebufferFetch, bSRGBAwareTarget); break;
+		case 27: SetShaderTemplES2<27>(Context, bUsedFramebufferFetch, bSRGBAwareTarget); break;
+		case 28: SetShaderTemplES2<28>(Context, bUsedFramebufferFetch, bSRGBAwareTarget); break;
+		case 29: SetShaderTemplES2<29>(Context, bUsedFramebufferFetch, bSRGBAwareTarget); break;
+		case 30: SetShaderTemplES2<30>(Context, bUsedFramebufferFetch, bSRGBAwareTarget); break;
+		case 31: SetShaderTemplES2<31>(Context, bUsedFramebufferFetch, bSRGBAwareTarget); break;
+		case 32: SetShaderTemplES2<32>(Context, bUsedFramebufferFetch, bSRGBAwareTarget); break;
+		case 33: SetShaderTemplES2<33>(Context, bUsedFramebufferFetch, bSRGBAwareTarget); break;
+		case 34: SetShaderTemplES2<34>(Context, bUsedFramebufferFetch, bSRGBAwareTarget); break;
+		case 35: SetShaderTemplES2<35>(Context, bUsedFramebufferFetch, bSRGBAwareTarget); break;
+		case 36: SetShaderTemplES2<36>(Context, bUsedFramebufferFetch, bSRGBAwareTarget); break;
+		case 37: SetShaderTemplES2<37>(Context, bUsedFramebufferFetch, bSRGBAwareTarget); break;
+		case 38: SetShaderTemplES2<38>(Context, bUsedFramebufferFetch, bSRGBAwareTarget); break;
 		default:
 			check(0);
 	}
diff --git a/Engine/Source/Runtime/Renderer/Private/PostProcess/PostProcessTonemap.h b/Engine/Source/Runtime/Renderer/Private/PostProcess/PostProcessTonemap.h
--- a/Engine/Source/Runtime/Renderer/Private/PostProcess/PostProcessTonemap.h
+++ b/Engine/Source/Runtime/Renderer/Private/PostProcess/PostProcessTonemap.h
@@ -78,7 +78,7 @@ private:
 class FRCPassPostProcessTonemapES2 : public TRenderingCompositePassBase<3, 1>
 {
 public:
-	FRCPassPostProcessTonemapES2(const FViewInfo& View, bool bInUsedFramebufferFetch);
+	FRCPassPostProcessTonemapES2(const FViewInfo& View, bool bInUsedFramebufferFetch, bool bInSRGBAwareTarget);
 
 	// interface FRenderingCompositePass ---------
 
@@ -91,6 +91,7 @@ private:
 	const FViewInfo& View;
 
 	bool bUsedFramebufferFetch;
+	bool bSRGBAwareTarget;
 	// set in constructor
 	uint32 ConfigIndexMobile;
 
diff --git a/Engine/Source/Runtime/Renderer/Private/PostProcess/PostProcessing.cpp b/Engine/Source/Runtime/Renderer/Private/PostProcess/PostProcessing.cpp
--- a/Engine/Source/Runtime/Renderer/Private/PostProcess/PostProcessing.cpp
+++ b/Engine/Source/Runtime/Renderer/Private/PostProcess/PostProcessing.cpp
@@ -2032,6 +2032,12 @@ void FPostProcessing::ProcessES2(FRHICommandListImmediate& RHICmdList, const FVi
 		FIntPoint SceneColorSize = FSceneRenderTargets::Get(RHICmdList).GetBufferSizeXY();
 		bool bViewRectSource = bUsedFramebufferFetch || SceneColorSize != PrePostSourceViewportSize;
 
+		// temporary solution for SP_METAL using HW sRGB flag during read vs all other mob platforms using
+		// incorrect UTexture::SRGB state. (UTexture::SRGB != HW texture state)
+		bool bSRGBAwareTarget = View.Family->RenderTarget->GetDisplayGamma() == 1.0f
+			&& View.bIsSceneCapture
+			&& View.GetShaderPlatform() == EShaderPlatform::SP_METAL;
+
 		// add the passes we want to add to the graph (commenting a line means the pass is not inserted into the graph) ---------
 		if( View.Family->EngineShowFlags.PostProcessing )
 		{
@@ -2321,7 +2327,7 @@ void FPostProcessing::ProcessES2(FRHICommandListImmediate& RHICmdList, const FVi
 		else
 		{
 			// Must run to blit to back buffer even if post processing is off.
-			FRCPassPostProcessTonemapES2* PostProcessTonemap = (FRCPassPostProcessTonemapES2*)Context.Graph.RegisterPass(new(FMemStack::Get()) FRCPassPostProcessTonemapES2(Context.View,  bViewRectSource));
+			FRCPassPostProcessTonemapES2* PostProcessTonemap = (FRCPassPostProcessTonemapES2*)Context.Graph.RegisterPass(new(FMemStack::Get()) FRCPassPostProcessTonemapES2(Context.View, bViewRectSource, bSRGBAwareTarget));
 			PostProcessTonemap->SetInput(ePId_Input0, Context.FinalOutput);
 			PostProcessTonemap->SetInput(ePId_Input1, BloomOutput);
 			PostProcessTonemap->SetInput(ePId_Input2, DofOutput);
diff --git a/Engine/Source/Runtime/RuntimeAssetCache/Private/RuntimeAssetCacheEntryMetadata.h b/Engine/Source/Runtime/RuntimeAssetCache/Private/RuntimeAssetCacheEntryMetadata.h
--- a/Engine/Source/Runtime/RuntimeAssetCache/Private/RuntimeAssetCacheEntryMetadata.h
+++ b/Engine/Source/Runtime/RuntimeAssetCache/Private/RuntimeAssetCacheEntryMetadata.h
@@ -59,17 +59,17 @@ public:
 		LastAccessTime = Value;
 	}
 
-	FDateTime GetLastAccessTime()
+	FDateTime GetLastAccessTime() const
 	{
 		return LastAccessTime;
 	}
 
-	FName GetName()
+	FName GetName() const
 	{
 		return Name;
 	}
 
-	bool IsBuilding()
+	bool IsBuilding() const
 	{
 		return bIsBuilding;
 	}
diff --git a/Engine/Source/Runtime/RuntimeAssetCache/Private/RuntimeAssetCacheFilesystemBackend.cpp b/Engine/Source/Runtime/RuntimeAssetCache/Private/RuntimeAssetCacheFilesystemBackend.cpp
--- a/Engine/Source/Runtime/RuntimeAssetCache/Private/RuntimeAssetCacheFilesystemBackend.cpp
+++ b/Engine/Source/Runtime/RuntimeAssetCache/Private/RuntimeAssetCacheFilesystemBackend.cpp
@@ -6,6 +6,7 @@
 #include "Misc/Paths.h"
 #include "Misc/ConfigCacheIni.h"
 #include "RuntimeAssetCacheBucket.h"
+#include "RuntimeAssetCacheEntryMetadata.h"
 
 FArchive* FRuntimeAssetCacheFilesystemBackend::CreateReadArchive(FName Bucket, const TCHAR* CacheKey)
 {
@@ -64,17 +65,51 @@ FRuntimeAssetCacheBucket* FRuntimeAssetCacheFilesystemBackend::PreLoadBucket(FNa
 			FString CacheKey = FPaths::GetBaseFilename(FilenameOrDirectory);
 			FArchive* Ar = Backend->CreateReadArchive(BucketName, *CacheKey);
 			FCacheEntryMetadata* Metadata = Backend->PreloadMetadata(Ar);
-			Bucket->AddMetadataEntry(*FPaths::GetBaseFilename(FilenameOrDirectory), Metadata, true);
+			Files.Add(TFilenameMetadataPair(*FPaths::GetBaseFilename(FilenameOrDirectory), Metadata));
 			delete Ar;
 			return true;
 		}
 
+		void Finalize()
+		{
+			Files.Sort([](const TFilenameMetadataPair& A, const TFilenameMetadataPair& B) { return A.Get<1>()->GetLastAccessTime() > B.Get<1>()->GetLastAccessTime(); });
+
+			int32 CurrentFileIndex = 0;
+			for (; CurrentFileIndex < Files.Num(); ++CurrentFileIndex)
+			{
+				FName Name = Files[CurrentFileIndex].Get<0>();
+				FCacheEntryMetadata* Metadata = Files[CurrentFileIndex].Get<1>();
+
+				if ((Bucket->GetCurrentSize() + Metadata->GetCachedAssetSize()) <= Bucket->GetSize())
+				{
+					Bucket->AddMetadataEntry(Name.ToString(), Metadata, true);
+				}
+				else
+				{
+					// No space left. Stop adding entries to the bucket, and all the rest of the older files
+					break;
+				}
+			}
+
+			for (; CurrentFileIndex < Files.Num(); ++CurrentFileIndex)
+			{
+				FName Name = Files[CurrentFileIndex].Get<0>();
+				FCacheEntryMetadata* Metadata = Files[CurrentFileIndex].Get<1>();
+				Backend->RemoveCacheEntry(BucketName, *Name.ToString());
+				delete Metadata;
+			}
+		}
+
 	private:
 		FRuntimeAssetCacheBucket* Bucket;
 		FName BucketName;
 		FRuntimeAssetCacheFilesystemBackend* Backend;
+
+		typedef TTuple<FName, FCacheEntryMetadata*> TFilenameMetadataPair;
+		TArray<TFilenameMetadataPair> Files;
 	} Visitor(Result, BucketName, this);
 	IFileManager::Get().IterateDirectory(*Path, Visitor);
+	Visitor.Finalize();
 
 	return Result;
 }
