Merge CL 3424117 from Release-4.16.

diff --git a/Engine/Source/Runtime/Renderer/Private/PostProcess/SceneRenderTargets.cpp b/Engine/Source/Runtime/Renderer/Private/PostProcess/SceneRenderTargets.cpp
--- a/Engine/Source/Runtime/Renderer/Private/PostProcess/SceneRenderTargets.cpp
+++ b/Engine/Source/Runtime/Renderer/Private/PostProcess/SceneRenderTargets.cpp
@@ -875,6 +875,7 @@ void FSceneRenderTargets::AllocLightAttenuation(FRHICommandList& RHICmdList)
 void FSceneRenderTargets::ReleaseGBufferTargets()
 {
 	GBufferResourcesUniformBuffer.SafeRelease();
+	GBufferDummyResourcesUniformBuffer.SafeRelease();
 	GBufferA.SafeRelease();
 	GBufferB.SafeRelease();
 	GBufferC.SafeRelease();
