Copying //UE4/Dev-Rendering to //UE4/Dev-Main (Source: //UE4/Dev-Rendering @ 3461187)

diff --git a/Engine/Source/Runtime/Renderer/Private/LightGridInjection.cpp b/Engine/Source/Runtime/Renderer/Private/LightGridInjection.cpp
--- a/Engine/Source/Runtime/Renderer/Private/LightGridInjection.cpp
+++ b/Engine/Source/Runtime/Renderer/Private/LightGridInjection.cpp
@@ -343,17 +343,21 @@ void FDeferredShadingSceneRenderer::ComputeLightGrid(FRHICommandListImmediate& R
 			GatherSimpleLights(ViewFamily, Views, SimpleLights);
 		}
 
+		TArray<FForwardGlobalLightData, TInlineAllocator<2>> GlobalLightDataForAllViews;
+		GlobalLightDataForAllViews.Empty(Views.Num());
+		GlobalLightDataForAllViews.AddDefaulted(Views.Num());
+
 		for (int32 ViewIndex = 0; ViewIndex < Views.Num(); ViewIndex++)
 		{
 			FViewInfo& View = Views[ViewIndex];
+			FForwardGlobalLightData& GlobalLightData = GlobalLightDataForAllViews[ViewIndex];
 
-			FForwardGlobalLightData GlobalLightData;
 			TArray<FForwardLocalLightData, SceneRenderingAllocator> ForwardLocalLightData;
 			float FurthestLight = 1000;
 
 			if (bCullLightsToGrid)
 			{
-				ForwardLocalLightData.Empty(Scene->Lights.Num());
+				ForwardLocalLightData.Empty(Scene->Lights.Num() + SimpleLights.InstanceData.Num());
 
 				for (TSparseArray<FLightSceneInfoCompact>::TConstIterator LightIt(Scene->Lights); LightIt; ++LightIt)
 				{
@@ -509,6 +513,10 @@ void FDeferredShadingSceneRenderer::ComputeLightGrid(FRHICommandListImmediate& R
 
 				// Pack both values into a single float to keep float4 alignment
 				const FFloat16 SimpleLightSourceLength16f = FFloat16(0);
+				FLightingChannels SimpleLightLightingChannels;
+				// Put simple lights in all lighting channels
+				SimpleLightLightingChannels.bChannel0 = SimpleLightLightingChannels.bChannel1 = SimpleLightLightingChannels.bChannel2 = true;
+				const uint32 SimpleLightLightingChannelMask = GetLightingChannelMaskForStruct(SimpleLightLightingChannels);
 
 				for (int32 SimpleLightIndex = 0; SimpleLightIndex < SimpleLights.InstanceData.Num(); SimpleLightIndex++)
 				{	
@@ -520,7 +528,10 @@ void FDeferredShadingSceneRenderer::ComputeLightGrid(FRHICommandListImmediate& R
 					LightData.LightPositionAndInvRadius = FVector4(SimpleLightPerViewData.Position, 1.0f / FMath::Max(SimpleLight.Radius, KINDA_SMALL_NUMBER));
 					LightData.LightColorAndFalloffExponent = FVector4(SimpleLight.Color, SimpleLight.Exponent);
 
+					// No shadowmap channels for simple lights
 					uint32 ShadowMapChannelMask = 0;
+					ShadowMapChannelMask |= SimpleLightLightingChannelMask << 8;
+
 					LightData.LightDirectionAndShadowMapChannelMask = FVector4(FVector(1, 0, 0), *((float*)&ShadowMapChannelMask));
 
 					// Pack both values into a single float to keep float4 alignment
@@ -592,6 +603,7 @@ void FDeferredShadingSceneRenderer::ComputeLightGrid(FRHICommandListImmediate& R
 		for (int32 ViewIndex = 0; ViewIndex < Views.Num(); ViewIndex++)
 		{
 			FViewInfo& View = Views[ViewIndex];
+			const FForwardGlobalLightData& GlobalLightData = GlobalLightDataForAllViews[ViewIndex];
 
 			const FIntPoint LightGridSizeXY = FIntPoint::DivideAndRoundUp(View.ViewRect.Size(), GLightGridPixelSize);
 			const int32 NumCells = LightGridSizeXY.X * LightGridSizeXY.Y * GLightGridSizeZ * NumCulledGridPrimitiveTypes;
@@ -613,7 +625,12 @@ void FDeferredShadingSceneRenderer::ComputeLightGrid(FRHICommandListImmediate& R
 			const FIntVector NumGroups = FIntVector::DivideAndRoundUp(FIntVector(LightGridSizeXY.X, LightGridSizeXY.Y, GLightGridSizeZ), LightGridInjectionGroupSize);
 
 			{
-				SCOPED_DRAW_EVENT(RHICmdList, CullLights);
+				SCOPED_DRAW_EVENTF(RHICmdList, CullLights, TEXT("CullLights %ux%ux%u NumLights %u NumCaptures %u"), 
+					GlobalLightData.CulledGridSize.X, 
+					GlobalLightData.CulledGridSize.Y,
+					GlobalLightData.CulledGridSize.Z,
+					GlobalLightData.NumLocalLights,
+					GlobalLightData.NumReflectionCaptures);
 
 				TArray<FUnorderedAccessViewRHIParamRef, TInlineAllocator<6>> OutUAVs;
 				OutUAVs.Add(View.ForwardLightingResources->NumCulledLightsGrid.UAV);
@@ -626,9 +643,9 @@ void FDeferredShadingSceneRenderer::ComputeLightGrid(FRHICommandListImmediate& R
 
 				if (GLightLinkedListCulling)
 				{
-					ClearUAV(RHICmdList, GMaxRHIFeatureLevel, View.ForwardLightingResources->StartOffsetGrid, 0xFFFFFFFF);
-					ClearUAV(RHICmdList, GMaxRHIFeatureLevel, View.ForwardLightingResources->NextCulledLightLink, 0);
-					ClearUAV(RHICmdList, GMaxRHIFeatureLevel, View.ForwardLightingResources->NextCulledLightData, 0);
+					ClearUAV(RHICmdList, View.ForwardLightingResources->StartOffsetGrid, 0xFFFFFFFF);
+					ClearUAV(RHICmdList, View.ForwardLightingResources->NextCulledLightLink, 0);
+					ClearUAV(RHICmdList, View.ForwardLightingResources->NextCulledLightData, 0);
 
 					TShaderMapRef<TLightGridInjectionCS<true> > ComputeShader(View.ShaderMap);
 					RHICmdList.SetComputeShader(ComputeShader->GetComputeShader());
@@ -638,7 +655,7 @@ void FDeferredShadingSceneRenderer::ComputeLightGrid(FRHICommandListImmediate& R
 				}
 				else
 				{
-					ClearUAV(RHICmdList, GMaxRHIFeatureLevel, View.ForwardLightingResources->NumCulledLightsGrid, 0);
+					ClearUAV(RHICmdList, View.ForwardLightingResources->NumCulledLightsGrid, 0);
 
 					TShaderMapRef<TLightGridInjectionCS<false> > ComputeShader(View.ShaderMap);
 					RHICmdList.SetComputeShader(ComputeShader->GetComputeShader());
diff --git a/Engine/Source/Runtime/Renderer/Private/LightRendering.cpp b/Engine/Source/Runtime/Renderer/Private/LightRendering.cpp
--- a/Engine/Source/Runtime/Renderer/Private/LightRendering.cpp
+++ b/Engine/Source/Runtime/Renderer/Private/LightRendering.cpp
@@ -614,7 +614,7 @@ void FDeferredShadingSceneRenderer::RenderLights(FRHICommandListImmediate& RHICm
 							}
 
 							RHICmdList.SetViewport(ScissorRect.Min.X, ScissorRect.Min.Y, 0.0f, ScissorRect.Max.X, ScissorRect.Max.Y, 1.0f);
-							DrawClearQuad(RHICmdList, FeatureLevel, true, FLinearColor(1, 1, 1, 1), false, 0, false, 0);
+							DrawClearQuad(RHICmdList, true, FLinearColor(1, 1, 1, 1), false, 0, false, 0);
 						}
 					}
 
@@ -717,7 +717,7 @@ void FDeferredShadingSceneRenderer::RenderStationaryLightOverlap(FRHICommandList
 		FSceneRenderTargets::Get(RHICmdList).BeginRenderingSceneColor(RHICmdList, ESimpleRenderTargetMode::EUninitializedColorExistingDepth, FExclusiveDepthStencil::DepthRead_StencilWrite);
 
 		// Clear to discard base pass values in scene color since we didn't skip that, to have valid scene depths
-		DrawClearQuad(RHICmdList, GMaxRHIFeatureLevel, FLinearColor::Black);
+		DrawClearQuad(RHICmdList, FLinearColor::Black);
 
 		RenderLightArrayForOverlapViewmode(RHICmdList, Scene->Lights);
 
@@ -1029,7 +1029,7 @@ void FDeferredShadingSceneRenderer::RenderLight(FRHICommandList& RHICmdList, con
 	if (bStencilDirty)
 	{
 		// Clear the stencil buffer to 0.
-		DrawClearQuad(RHICmdList, GMaxRHIFeatureLevel, false, FLinearColor::Transparent, false, 0, true, 1);
+		DrawClearQuad(RHICmdList, false, FLinearColor::Transparent, false, 0, true, 1);
 	}
 }
 
diff --git a/Engine/Source/Runtime/Renderer/Private/LightSceneInfo.cpp b/Engine/Source/Runtime/Renderer/Private/LightSceneInfo.cpp
--- a/Engine/Source/Runtime/Renderer/Private/LightSceneInfo.cpp
+++ b/Engine/Source/Runtime/Renderer/Private/LightSceneInfo.cpp
@@ -104,7 +104,7 @@ void FLightSceneInfo::AddToScene()
  */
 void FLightSceneInfo::CreateLightPrimitiveInteraction(const FLightSceneInfoCompact& LightSceneInfoCompact, const FPrimitiveSceneInfoCompact& PrimitiveSceneInfoCompact)
 {
-	if(	LightSceneInfoCompact.AffectsPrimitive(PrimitiveSceneInfoCompact))
+	if(	LightSceneInfoCompact.AffectsPrimitive(PrimitiveSceneInfoCompact.Bounds, PrimitiveSceneInfoCompact.Proxy))
 	{
 		// create light interaction and add to light/primitive lists
 		FLightPrimitiveInteraction::Create(this,PrimitiveSceneInfoCompact.PrimitiveSceneInfo);
@@ -218,14 +218,14 @@ FORCEINLINE bool AreSpheresNotIntersecting(
 * @param CompactPrimitiveSceneInfo - The primitive to test.
 * @return True if the light affects the primitive.
 */
-bool FLightSceneInfoCompact::AffectsPrimitive(const FPrimitiveSceneInfoCompact& CompactPrimitiveSceneInfo) const
+bool FLightSceneInfoCompact::AffectsPrimitive(const FBoxSphereBounds& PrimitiveBounds, const FPrimitiveSceneProxy* PrimitiveSceneProxy) const
 {
 	// Check if the light's bounds intersect the primitive's bounds.
 	if(AreSpheresNotIntersecting(
 		BoundingSphereVector,
 		VectorReplicate(BoundingSphereVector,3),
-		VectorLoadFloat3(&CompactPrimitiveSceneInfo.Bounds.Origin),
-		VectorLoadFloat1(&CompactPrimitiveSceneInfo.Bounds.SphereRadius)
+		VectorLoadFloat3(&PrimitiveBounds.Origin),
+		VectorLoadFloat1(&PrimitiveBounds.SphereRadius)
 		))
 	{
 		return false;
@@ -233,17 +233,17 @@ bool FLightSceneInfoCompact::AffectsPrimitive(const FPrimitiveSceneInfoCompact&
 
 	// Cull based on information in the full scene infos.
 
-	if(!LightSceneInfo->Proxy->AffectsBounds(CompactPrimitiveSceneInfo.Bounds))
+	if(!LightSceneInfo->Proxy->AffectsBounds(PrimitiveBounds))
 	{
 		return false;
 	}
 
-	if (LightSceneInfo->Proxy->CastsShadowsFromCinematicObjectsOnly() && !CompactPrimitiveSceneInfo.Proxy->CastsCinematicShadow())
+	if (LightSceneInfo->Proxy->CastsShadowsFromCinematicObjectsOnly() && !PrimitiveSceneProxy->CastsCinematicShadow())
 	{
 		return false;
 	}
 
-	if (!(LightSceneInfo->Proxy->GetLightingChannelMask() & CompactPrimitiveSceneInfo.Proxy->GetLightingChannelMask()))
+	if (!(LightSceneInfo->Proxy->GetLightingChannelMask() & PrimitiveSceneProxy->GetLightingChannelMask()))
 	{
 		return false;
 	}
diff --git a/Engine/Source/Runtime/Renderer/Private/LightSceneInfo.h b/Engine/Source/Runtime/Renderer/Private/LightSceneInfo.h
--- a/Engine/Source/Runtime/Renderer/Private/LightSceneInfo.h
+++ b/Engine/Source/Runtime/Renderer/Private/LightSceneInfo.h
@@ -58,7 +58,7 @@ public:
 	 * @param CompactPrimitiveSceneInfo - The primitive to test.
 	 * @return True if the light affects the primitive.
 	 */
-	bool AffectsPrimitive(const FPrimitiveSceneInfoCompact& CompactPrimitiveSceneInfo) const;
+	bool AffectsPrimitive(const FBoxSphereBounds& PrimitiveBounds, const FPrimitiveSceneProxy* PrimitiveSceneProxy) const;
 };
 
 /** Information for sorting lights. */
@@ -93,6 +93,12 @@ struct FSortedLightSceneInfo
 	}
 };
 
+template <>
+struct TUseBitwiseSwap<FSortedLightSceneInfo>
+{
+	enum { Value = false };
+};
+
 /** The type of the octree used by FScene to find lights. */
 typedef TOctree<FLightSceneInfoCompact,struct FLightOctreeSemantics> FSceneLightOctree;
 
diff --git a/Engine/Source/Runtime/Renderer/Private/MobileSceneCaptureRendering.cpp b/Engine/Source/Runtime/Renderer/Private/MobileSceneCaptureRendering.cpp
--- a/Engine/Source/Runtime/Renderer/Private/MobileSceneCaptureRendering.cpp
+++ b/Engine/Source/Runtime/Renderer/Private/MobileSceneCaptureRendering.cpp
@@ -393,7 +393,7 @@ void UpdateSceneCaptureContentMobile_RenderThread(
 		{
 			auto& RenderTargetRHI = Target->GetRenderTargetTexture();
 			SetRenderTarget(RHICmdList, RenderTargetRHI, NULL, true);
-			DrawClearQuad(RHICmdList, SceneRenderer->FeatureLevel, true, FLinearColor::Black, false, 0, false, 0, Target->GetSizeXY(), ViewRect);
+			DrawClearQuad(RHICmdList, true, FLinearColor::Black, false, 0, false, 0, Target->GetSizeXY(), ViewRect);
 		}
 
 		// Render the scene normally
diff --git a/Engine/Source/Runtime/Renderer/Private/MobileShadingRenderer.cpp b/Engine/Source/Runtime/Renderer/Private/MobileShadingRenderer.cpp
--- a/Engine/Source/Runtime/Renderer/Private/MobileShadingRenderer.cpp
+++ b/Engine/Source/Runtime/Renderer/Private/MobileShadingRenderer.cpp
@@ -144,7 +144,7 @@ void FMobileSceneRenderer::Render(FRHICommandListImmediate& RHICmdList)
 		}
 	}
 
-	if (GRHIThread)
+	if (IsRunningRHIInSeparateThread())
 	{
 		// we will probably stall on occlusion queries, so might as well have the RHI thread and GPU work while we wait.
 		// Also when doing RHI thread this is the only spot that will process pending deletes
@@ -208,7 +208,7 @@ void FMobileSceneRenderer::Render(FRHICommandListImmediate& RHICmdList)
 
 	if (GIsEditor && !View.bIsSceneCapture)
 	{
-		DrawClearQuad(RHICmdList, GMaxRHIFeatureLevel, Views[0].BackgroundColor);
+		DrawClearQuad(RHICmdList, Views[0].BackgroundColor);
 	}
 
 	RenderMobileBasePass(RHICmdList, ViewList);
diff --git a/Engine/Source/Runtime/Renderer/Private/MobileTranslucentRendering.cpp b/Engine/Source/Runtime/Renderer/Private/MobileTranslucentRendering.cpp
--- a/Engine/Source/Runtime/Renderer/Private/MobileTranslucentRendering.cpp
+++ b/Engine/Source/Runtime/Renderer/Private/MobileTranslucentRendering.cpp
@@ -276,9 +276,7 @@ template <class TDrawingPolicyFactory>
 void FTranslucentPrimSet::DrawPrimitivesForMobile(FRHICommandListImmediate& RHICmdList, const FViewInfo& View, const FDrawingPolicyRenderState& DrawRenderState,
 	typename TDrawingPolicyFactory::ContextType& DrawingContext) const
 {
-	ETranslucencyPass::Type TranslucencyPass = DrawingContext.ShouldRenderSeparateTranslucency() ? ETranslucencyPass::TPT_SeparateTranslucency : ETranslucencyPass::TPT_StandardTranslucency;
-
-	FInt32Range PassRange = SortedPrimsNum.GetPassRange(TranslucencyPass);
+	FInt32Range PassRange = SortedPrimsNum.GetPassRange(ETranslucencyPass::TPT_AllTranslucency);
 
 	// Draw sorted scene prims
 	for (int32 PrimIdx = PassRange.GetLowerBoundValue(); PrimIdx < PassRange.GetUpperBoundValue(); PrimIdx++)
@@ -327,7 +325,7 @@ void FTranslucentPrimSet::DrawPrimitivesForMobile(FRHICommandListImmediate& RHIC
 
 void FMobileSceneRenderer::RenderTranslucency(FRHICommandListImmediate& RHICmdList, const TArrayView<const FViewInfo*> PassViews)
 {
-	if (ShouldRenderTranslucency())
+	if (ShouldRenderTranslucency(ETranslucencyPass::TPT_AllTranslucency))
 	{
 		const bool bGammaSpace = !IsMobileHDR();
 		const bool bLinearHDR64 = !bGammaSpace && !IsMobileHDR32bpp();
@@ -670,7 +668,7 @@ bool FMobileSceneRenderer::RenderInverseOpacity(FRHICommandListImmediate& RHICmd
 {
 	bool bDirty = false;
 
-	if (ShouldRenderTranslucency())
+	if (ShouldRenderTranslucency(ETranslucencyPass::TPT_AllTranslucency))
 	{
 		const bool bGammaSpace = !IsMobileHDR();
 		const bool bLinearHDR64 = !bGammaSpace && !IsMobileHDR32bpp();
diff --git a/Engine/Source/Runtime/Renderer/Private/PlanarReflectionRendering.cpp b/Engine/Source/Runtime/Renderer/Private/PlanarReflectionRendering.cpp
--- a/Engine/Source/Runtime/Renderer/Private/PlanarReflectionRendering.cpp
+++ b/Engine/Source/Runtime/Renderer/Private/PlanarReflectionRendering.cpp
@@ -272,7 +272,7 @@ static void UpdatePlanarReflectionContents_RenderThread(
 				else
 				{
 					SetRenderTarget(RHICmdList, Target->GetRenderTargetTexture(), nullptr, true);
-					DrawClearQuad(RHICmdList, GMaxRHIFeatureLevel, FLinearColor::Black);
+					DrawClearQuad(RHICmdList, FLinearColor::Black);
 				}
 
 				// Reflection view late update
@@ -601,7 +601,7 @@ bool FDeferredShadingSceneRenderer::RenderDeferredPlanarReflections(FRHICommandL
 
 		if (!bSSRAsInput)
 		{
-			DrawClearQuad(RHICmdList, GMaxRHIFeatureLevel, FLinearColor(0, 0, 0, 0));
+			DrawClearQuad(RHICmdList, FLinearColor(0, 0, 0, 0));
 		}
 
 		{
diff --git a/Engine/Source/Runtime/Renderer/Private/PostProcess/PostProcessBloomSetup.cpp b/Engine/Source/Runtime/Renderer/Private/PostProcess/PostProcessBloomSetup.cpp
--- a/Engine/Source/Runtime/Renderer/Private/PostProcess/PostProcessBloomSetup.cpp
+++ b/Engine/Source/Runtime/Renderer/Private/PostProcess/PostProcessBloomSetup.cpp
@@ -287,7 +287,7 @@ void FRCPassPostProcessBloomSetup::Process(FRenderingCompositePassContext& Conte
 		SetRenderTarget(Context.RHICmdList, DestRenderTarget.TargetableTexture, FTextureRHIRef());
 
 		// is optimized away if possible (RT size=view size, )
-		DrawClearQuad(Context.RHICmdList, Context.GetFeatureLevel(), true, FLinearColor::Black, false, 0, false, 0, DestSize, DestRect);
+		DrawClearQuad(Context.RHICmdList, true, FLinearColor::Black, false, 0, false, 0, DestSize, DestRect);
 
 		Context.SetViewportAndCallRHI(0, 0, 0.0f, DestSize.X, DestSize.Y, 1.0f );
 
@@ -434,7 +434,7 @@ void FRCPassPostProcessVisualizeBloomSetup::Process(FRenderingCompositePassConte
 	SetRenderTarget(Context.RHICmdList, DestRenderTarget.TargetableTexture, FTextureRHIRef());
 
 	// is optimized away if possible (RT size=view size, )
-	DrawClearQuad(Context.RHICmdList, Context.GetFeatureLevel(), true, FLinearColor(0, 0, 0, 0), false, 0, false, 0, DestSize, DestRect);
+	DrawClearQuad(Context.RHICmdList, true, FLinearColor(0, 0, 0, 0), false, 0, false, 0, DestSize, DestRect);
 
 	Context.SetViewportAndCallRHI(0, 0, 0.0f, DestRect.Width(), DestRect.Height(), 1.0f );
 
@@ -578,7 +578,7 @@ void FRCPassPostProcessVisualizeBloomOverlay::Process(FRenderingCompositePassCon
 	SetRenderTarget(Context.RHICmdList, DestRenderTarget.TargetableTexture, FTextureRHIRef());
 
 	// is optimized away if possible (RT size=view size, )
-	DrawClearQuad(Context.RHICmdList, Context.GetFeatureLevel(), true, FLinearColor(0, 0, 0 ,0), false, 0, false, 0, DestSize, DestRect);
+	DrawClearQuad(Context.RHICmdList, true, FLinearColor(0, 0, 0 ,0), false, 0, false, 0, DestSize, DestRect);
 
 	Context.SetViewportAndCallRHI(0, 0, 0.0f, DestRect.Width(), DestRect.Height(), 1.0f );
 
diff --git a/Engine/Source/Runtime/Renderer/Private/PostProcess/PostProcessBokehDOF.cpp b/Engine/Source/Runtime/Renderer/Private/PostProcess/PostProcessBokehDOF.cpp
--- a/Engine/Source/Runtime/Renderer/Private/PostProcess/PostProcessBokehDOF.cpp
+++ b/Engine/Source/Runtime/Renderer/Private/PostProcess/PostProcessBokehDOF.cpp
@@ -18,6 +18,7 @@
 #include "PostProcess/PostProcessCircleDOF.h"
 #include "ClearQuad.h"
 #include "PipelineStateCache.h"
+#include "SpriteIndexBuffer.h"
 
 const int32 GBokehDOFSetupTileSizeX = 8;
 const int32 GBokehDOFSetupTileSizeY = 8;
@@ -50,7 +51,7 @@ public:
 };
 
 /** Global Bokeh index buffer. */
-TGlobalResource< FBokehIndexBuffer > GBokehIndexBuffer;
+TGlobalResource< FSpriteIndexBuffer<8> > GBokehIndexBuffer;
 
 /** Encapsulates the post processing depth of field setup pixel shader. */
 class PostProcessVisualizeDOFPS : public FGlobalShader
@@ -194,7 +195,7 @@ void FRCPassPostProcessVisualizeDOF::Process(FRenderingCompositePassContext& Con
 	// can be optimized (don't clear areas we overwrite, don't clear when full screen),
 	// needed when a camera (matinee) has black borders or with multiple viewports
 	// focal distance depth is stored in the alpha channel to avoid DOF artifacts
-	DrawClearQuad(Context.RHICmdList, Context.GetFeatureLevel(), true, FLinearColor(0, 0, 0, View.FinalPostProcessSettings.DepthOfFieldFocalDistance), false, 0, false, 0, PassOutputs[0].RenderTargetDesc.Extent, DestRect);
+	DrawClearQuad(Context.RHICmdList, true, FLinearColor(0, 0, 0, View.FinalPostProcessSettings.DepthOfFieldFocalDistance), false, 0, false, 0, PassOutputs[0].RenderTargetDesc.Extent, DestRect);
 
 	Context.SetViewportAndCallRHI(0, 0, 0.0f, DestSize.X, DestSize.Y, 1.0f );
 
@@ -545,7 +546,7 @@ void FRCPassPostProcessBokehDOFSetup::Process(FRenderingCompositePassContext& Co
 		// can be optimized (don't clear areas we overwrite, don't clear when full screen),
 		// needed when a camera (matinee) has black borders or with multiple viewports
 		// focal distance depth is stored in the alpha channel to avoid DOF artifacts
-		DrawClearQuad(Context.RHICmdList, Context.GetFeatureLevel(), true, FLinearColor(0, 0, 0, View.FinalPostProcessSettings.DepthOfFieldFocalDistance), false, 0, false, 0, DestSize, DestRect);
+		DrawClearQuad(Context.RHICmdList, true, FLinearColor(0, 0, 0, View.FinalPostProcessSettings.DepthOfFieldFocalDistance), false, 0, false, 0, DestSize, DestRect);
 
 		Context.SetViewportAndCallRHI(0, 0, 0.0f, DestSize.X, DestSize.Y, 1.0f );
 
@@ -886,7 +887,7 @@ void FRCPassPostProcessBokehDOF::Process(FRenderingCompositePassContext& Context
 	SetRenderTarget(Context.RHICmdList, DestRenderTarget.TargetableTexture, FTextureRHIRef());
 
 	// This clean is required to make the accumulation working
-	DrawClearQuad(Context.RHICmdList, Context.GetFeatureLevel(), true, FLinearColor(0, 0, 0, 0), false, 0, false, 0, GetOutput(ePId_Output0)->RenderTargetDesc.Extent, FIntRect());
+	DrawClearQuad(Context.RHICmdList, true, FLinearColor(0, 0, 0, 0), false, 0, false, 0, GetOutput(ePId_Output0)->RenderTargetDesc.Extent, FIntRect());
 
 	// we need to output to the whole rendertarget
 	Context.SetViewportAndCallRHI(0, 0, 0.0f, PassOutputs[0].RenderTargetDesc.Extent.X, PassOutputs[0].RenderTargetDesc.Extent.Y, 1.0f);
diff --git a/Engine/Source/Runtime/Renderer/Private/PostProcess/PostProcessBokehDOFRecombine.cpp b/Engine/Source/Runtime/Renderer/Private/PostProcess/PostProcessBokehDOFRecombine.cpp
--- a/Engine/Source/Runtime/Renderer/Private/PostProcess/PostProcessBokehDOFRecombine.cpp
+++ b/Engine/Source/Runtime/Renderer/Private/PostProcess/PostProcessBokehDOFRecombine.cpp
@@ -14,13 +14,7 @@
 #include "PostProcess/PostProcessBokehDOF.h"
 #include "ClearQuad.h"
 #include "PipelineStateCache.h"
-
-static TAutoConsoleVariable<int32> CVarSeparateTranslucencyUpsampleMode(
-	TEXT("r.SeparateTranslucencyUpsampleMode"),
-	1,
-	TEXT("Upsample method to use on separate translucency.  These are only used when r.SeparateTranslucencyScreenPercentage is less than 100.\n")
-	TEXT("0: bilinear 1: Nearest-Depth Neighbor (only when r.SeparateTranslucencyScreenPercentage is 50)"),
-	ECVF_Scalability | ECVF_Default);
+#include "TranslucentRendering.h"
 
 const int32 GBokehDOFRecombineComputeTileSizeX = 8;
 const int32 GBokehDOFRecombineComputeTileSizeY = 8;
@@ -121,20 +115,20 @@ public:
 
 		if (UseNearestDepthNeighborUpsample())
 		{
-			check(SceneContext.IsSeparateTranslucencyDepthValid());
-			FTextureRHIParamRef LowResDepth = SceneContext.GetSeparateTranslucencyDepthSurface();
+			check(SceneContext.IsDownsampledTranslucencyDepthValid());
+			FTextureRHIParamRef LowResDepth = SceneContext.GetDownsampledTranslucencyDepthSurface();
 			SetTextureParameter(Context.RHICmdList, ShaderRHI, LowResDepthTexture, LowResDepth);
 
 			const auto& BuiltinSamplersUBParameter = GetUniformBufferParameter<FBuiltinSamplersParameters>();
 			SetUniformBufferParameter(Context.RHICmdList, ShaderRHI, BuiltinSamplersUBParameter, GBuiltinSamplersUniformBuffer.GetUniformBufferRHI());
-
-			SetSamplerParameter(Context.RHICmdList, ShaderRHI, BilinearClampedSampler, TStaticSamplerState<SF_Bilinear,AM_Clamp,AM_Clamp,AM_Clamp>::GetRHI());
-			SetSamplerParameter(Context.RHICmdList, ShaderRHI, PointClampedSampler, TStaticSamplerState<SF_Point,AM_Clamp,AM_Clamp,AM_Clamp>::GetRHI());
 		}
 		else
 		{
 			checkSlow(!LowResDepthTexture.IsBound());
 		}
+
+		SetSamplerParameter(Context.RHICmdList, ShaderRHI, BilinearClampedSampler, TStaticSamplerState<SF_Bilinear, AM_Clamp, AM_Clamp, AM_Clamp>::GetRHI());
+		SetSamplerParameter(Context.RHICmdList, ShaderRHI, PointClampedSampler, TStaticSamplerState<SF_Point, AM_Clamp, AM_Clamp, AM_Clamp>::GetRHI());
 	}
 
 	static const TCHAR* GetSourceFilename()
@@ -264,8 +258,8 @@ public:
 
 		if (UseNearestDepthNeighborUpsample())
 		{
-			check(SceneContext.IsSeparateTranslucencyDepthValid());
-			FTextureRHIParamRef LowResDepth = SceneContext.GetSeparateTranslucencyDepthSurface();
+			check(SceneContext.IsDownsampledTranslucencyDepthValid());
+			FTextureRHIParamRef LowResDepth = SceneContext.GetDownsampledTranslucencyDepthSurface();
 			SetTextureParameter(RHICmdList, ShaderRHI, LowResDepthTexture, LowResDepth);
 
 			const auto& BuiltinSamplersUBParameter = GetUniformBufferParameter<FBuiltinSamplersParameters>();
@@ -365,14 +359,7 @@ void FRCPassPostProcessBokehDOFRecombine::Process(FRenderingCompositePassContext
 		check(GetInput(ePId_Input2)->GetPass());
 	}
 
-	FSceneRenderTargets& SceneContext = FSceneRenderTargets::Get(Context.RHICmdList);
-	FIntPoint OutScaledSize;
-	float OutScale;
-	SceneContext.GetSeparateTranslucencyDimensions(OutScaledSize, OutScale);
-
-	const bool bUseNearestDepthNeighborUpsample = 
-		CVarSeparateTranslucencyUpsampleMode.GetValueOnRenderThread() != 0
-		&& FMath::Abs(OutScale - .5f) < .001f;
+	const bool bUseNearestDepthNeighborUpsample = UseNearestDepthNeighborUpsampleForSeparateTranslucency(FSceneRenderTargets::Get(Context.RHICmdList));
 
 	if (Method != 1 && bUseNearestDepthNeighborUpsample)
 	{
@@ -441,7 +428,7 @@ void FRCPassPostProcessBokehDOFRecombine::Process(FRenderingCompositePassContext
 		SetRenderTarget(Context.RHICmdList, DestRenderTarget.TargetableTexture, FTextureRHIRef());
 
 		// is optimized away if possible (RT size=view size, )
-		DrawClearQuad(Context.RHICmdList, Context.GetFeatureLevel(), true, FLinearColor::Black, false, 1.0f, false, 0, PassOutputs[0].RenderTargetDesc.Extent, View.ViewRect);
+		DrawClearQuad(Context.RHICmdList, true, FLinearColor::Black, false, 1.0f, false, 0, PassOutputs[0].RenderTargetDesc.Extent, View.ViewRect);
 
 		Context.SetViewportAndCallRHI(View.ViewRect);
 
diff --git a/Engine/Source/Runtime/Renderer/Private/PostProcess/PostProcessCircleDOF.cpp b/Engine/Source/Runtime/Renderer/Private/PostProcess/PostProcessCircleDOF.cpp
--- a/Engine/Source/Runtime/Renderer/Private/PostProcess/PostProcessCircleDOF.cpp
+++ b/Engine/Source/Runtime/Renderer/Private/PostProcess/PostProcessCircleDOF.cpp
@@ -208,7 +208,7 @@ void FRCPassPostProcessCircleDOFSetup::Process(FRenderingCompositePassContext& C
 		FLinearColor(0, 0, 0, 0)
 	};
 	// is optimized away if possible (RT size=view size, )
-	DrawClearQuadMRT(Context.RHICmdList, Context.GetFeatureLevel(), true, NumRenderTargets, ClearColors, false, 0, false, 0, PassOutputs[0].RenderTargetDesc.Extent, DestRect);
+	DrawClearQuadMRT(Context.RHICmdList, true, NumRenderTargets, ClearColors, false, 0, false, 0, PassOutputs[0].RenderTargetDesc.Extent, DestRect);
 
 	Context.SetViewportAndCallRHI(0, 0, 0.0f, DestSize.X, DestSize.Y, 1.0f );
 
@@ -417,7 +417,7 @@ void FRCPassPostProcessCircleDOFDilate::Process(FRenderingCompositePassContext&
 		FLinearColor(0, 0, 0, 0)
 	};
 	// is optimized away if possible (RT size=view size, )
-	DrawClearQuadMRT(Context.RHICmdList, Context.GetFeatureLevel(), true, NumRenderTargets, ClearColors, false, 0, false, 0, PassOutputs[0].RenderTargetDesc.Extent, DestRect);
+	DrawClearQuadMRT(Context.RHICmdList, true, NumRenderTargets, ClearColors, false, 0, false, 0, PassOutputs[0].RenderTargetDesc.Extent, DestRect);
 
 	Context.SetViewportAndCallRHI(0, 0, 0.0f, DestSize.X, DestSize.Y, 1.0f );
 
@@ -698,7 +698,7 @@ void FRCPassPostProcessCircleDOF::Process(FRenderingCompositePassContext& Contex
 	};
 
 	// is optimized away if possible (RT size=view size, )
-	DrawClearQuadMRT(Context.RHICmdList, Context.GetFeatureLevel(), true, NumRenderTargets, ClearColors, false, 0, false, 0, PassOutputs[0].RenderTargetDesc.Extent, DestRect);
+	DrawClearQuadMRT(Context.RHICmdList, true, NumRenderTargets, ClearColors, false, 0, false, 0, PassOutputs[0].RenderTargetDesc.Extent, DestRect);
 
 	Context.SetViewportAndCallRHI(0, 0, 0.0f, DestSize.X, DestSize.Y, 1.0f );
 		
@@ -916,7 +916,7 @@ void FRCPassPostProcessCircleDOFRecombine::Process(FRenderingCompositePassContex
 	SetRenderTarget(Context.RHICmdList, DestRenderTarget.TargetableTexture, FTextureRHIRef());
 
 	// is optimized away if possible (RT size=view size, )
-	DrawClearQuad(Context.RHICmdList, Context.GetFeatureLevel(), true, FLinearColor::Black, false, 0, false, 0, PassOutputs[0].RenderTargetDesc.Extent, View.ViewRect);
+	DrawClearQuad(Context.RHICmdList, true, FLinearColor::Black, false, 0, false, 0, PassOutputs[0].RenderTargetDesc.Extent, View.ViewRect);
 
 	Context.SetViewportAndCallRHI(View.ViewRect);
 
diff --git a/Engine/Source/Runtime/Renderer/Private/PostProcess/PostProcessCompositeEditorPrimitives.cpp b/Engine/Source/Runtime/Renderer/Private/PostProcess/PostProcessCompositeEditorPrimitives.cpp
--- a/Engine/Source/Runtime/Renderer/Private/PostProcess/PostProcessCompositeEditorPrimitives.cpp
+++ b/Engine/Source/Runtime/Renderer/Private/PostProcess/PostProcessCompositeEditorPrimitives.cpp
@@ -243,7 +243,7 @@ void FRCPassPostProcessCompositeEditorPrimitives::Process(FRenderingCompositePas
 		if (bClearIsNeeded)
 		{
 			SCOPED_DRAW_EVENT(Context.RHICmdList, ClearViewEditorPrimitives);
-			DrawClearQuad(Context.RHICmdList, Context.GetFeatureLevel(), true, FLinearColor(0, 0, 0, 0), true, (float)ERHIZBuffer::FarPlane, false, 0, ColorTarget->GetSizeXY(), FIntRect());
+			DrawClearQuad(Context.RHICmdList, true, FLinearColor(0, 0, 0, 0), true, (float)ERHIZBuffer::FarPlane, false, 0, ColorTarget->GetSizeXY(), FIntRect());
 		}
 
 		SCOPED_DRAW_EVENT(Context.RHICmdList, RenderEditorPrimitives);
diff --git a/Engine/Source/Runtime/Renderer/Private/PostProcess/PostProcessDOF.cpp b/Engine/Source/Runtime/Renderer/Private/PostProcess/PostProcessDOF.cpp
--- a/Engine/Source/Runtime/Renderer/Private/PostProcess/PostProcessDOF.cpp
+++ b/Engine/Source/Runtime/Renderer/Private/PostProcess/PostProcessDOF.cpp
@@ -193,7 +193,7 @@ void FRCPassPostProcessDOFSetup::Process(FRenderingCompositePassContext& Context
 		FLinearColor(0, 0, 0, 0)
 	};
 	// is optimized away if possible (RT size=view size, )
-	DrawClearQuadMRT(Context.RHICmdList, Context.GetFeatureLevel(), true, NumRenderTargets, ClearColors, false, 0, false, 0, DestSize, DestRect);
+	DrawClearQuadMRT(Context.RHICmdList, true, NumRenderTargets, ClearColors, false, 0, false, 0, DestSize, DestRect);
 
 	Context.SetViewportAndCallRHI(DestRect.Min.X, DestRect.Min.Y, 0.0f, DestRect.Max.X + 1, DestRect.Max.Y + 1, 1.0f );
 	
@@ -447,7 +447,7 @@ void FRCPassPostProcessDOFRecombine::Process(FRenderingCompositePassContext& Con
 	{
 		SetRenderTarget(Context.RHICmdList, DestRenderTarget.TargetableTexture, FTextureRHIRef());
 		// is optimized away if possible (RT size=view size, )
-		DrawClearQuad(Context.RHICmdList, Context.GetFeatureLevel(), true, FLinearColor::Black, false, 0, false, 0, PassOutputs[0].RenderTargetDesc.Extent, View.ViewRect);
+		DrawClearQuad(Context.RHICmdList, true, FLinearColor::Black, false, 0, false, 0, PassOutputs[0].RenderTargetDesc.Extent, View.ViewRect);
 	}
 
 	Context.SetViewportAndCallRHI(View.ViewRect);
diff --git a/Engine/Source/Runtime/Renderer/Private/PostProcess/PostProcessDownsample.cpp b/Engine/Source/Runtime/Renderer/Private/PostProcess/PostProcessDownsample.cpp
--- a/Engine/Source/Runtime/Renderer/Private/PostProcess/PostProcessDownsample.cpp
+++ b/Engine/Source/Runtime/Renderer/Private/PostProcess/PostProcessDownsample.cpp
@@ -364,7 +364,7 @@ void FRCPassPostProcessDownsample::Process(FRenderingCompositePassContext& Conte
 			// Set the view family's render target/viewport.
 			SetRenderTarget(Context.RHICmdList, DestRenderTarget.TargetableTexture, FTextureRHIRef(), ESimpleRenderTargetMode::EExistingColorAndDepth);
 			Context.SetViewportAndCallRHI(0, 0, 0.0f, DestSize.X, DestSize.Y, 1.0f);
-			DrawClearQuad(Context.RHICmdList, Context.GetFeatureLevel(), true, FLinearColor(0, 0, 0, 0), false, 1.0f, false, 0, DestSize, DestRect);
+			DrawClearQuad(Context.RHICmdList, true, FLinearColor(0, 0, 0, 0), false, 1.0f, false, 0, DestSize, DestRect);
 		}
 
 		// InflateSize increases the size of the source/dest rectangle to compensate for bilinear reads and UIBlur pass requirements.
diff --git a/Engine/Source/Runtime/Renderer/Private/PostProcess/PostProcessEyeAdaptation.cpp b/Engine/Source/Runtime/Renderer/Private/PostProcess/PostProcessEyeAdaptation.cpp
--- a/Engine/Source/Runtime/Renderer/Private/PostProcess/PostProcessEyeAdaptation.cpp
+++ b/Engine/Source/Runtime/Renderer/Private/PostProcess/PostProcessEyeAdaptation.cpp
@@ -461,7 +461,7 @@ void FRCPassPostProcessBasicEyeAdaptationSetUp::Process(FRenderingCompositePassC
 	SetRenderTarget(Context.RHICmdList, DestRenderTarget.TargetableTexture, FTextureRHIRef());
 
 	// is optimized away if possible (RT size=view size, )
-	DrawClearQuad(Context.RHICmdList, Context.GetFeatureLevel(), true, FLinearColor::Black, false, 0, false, 0, PassOutputs[0].RenderTargetDesc.Extent, DestRect);
+	DrawClearQuad(Context.RHICmdList, true, FLinearColor::Black, false, 0, false, 0, PassOutputs[0].RenderTargetDesc.Extent, DestRect);
 
 	Context.SetViewportAndCallRHI(0, 0, 0.0f, DestSize.X, DestSize.Y, 1.0f);
 
diff --git a/Engine/Source/Runtime/Renderer/Private/PostProcess/PostProcessFFTBloom.cpp b/Engine/Source/Runtime/Renderer/Private/PostProcess/PostProcessFFTBloom.cpp
--- a/Engine/Source/Runtime/Renderer/Private/PostProcess/PostProcessFFTBloom.cpp
+++ b/Engine/Source/Runtime/Renderer/Private/PostProcess/PostProcessFFTBloom.cpp
@@ -777,6 +777,12 @@ FSceneRenderTargetItem* FRCPassFFTBloom::InitDomainAndGetKernel(FRenderingCompos
 	// The kernel parameters on the FinalPostProcess.  
 
 	UTexture2D* BloomConvolutionTexture = PPSettings.BloomConvolutionTexture;
+	
+	if (BloomConvolutionTexture == nullptr)
+	{
+		BloomConvolutionTexture = GEngine->DefaultBloomKernelTexture;
+	}
+
 	const float BloomConvolutionSize    = PPSettings.BloomConvolutionSize;
 	const FVector2D CenterUV            = PPSettings.BloomConvolutionCenterUV;
 	const float ClampedBloomConvolutionBufferScale = FMath::Clamp(PPSettings.BloomConvolutionBufferScale, 0.f, 1.f);
@@ -1042,8 +1048,15 @@ FPooledRenderTargetDesc FRCPassFFTBloom::ComputeOutputDesc(EPassOutputId InPassO
 bool FRCPassFFTBloom::HasValidPhysicalKernel(FPostprocessContext& Context)
 {
 	const FViewInfo& View = Context.View;
+
 	UTexture2D* BloomConvolutionTexture = View.FinalPostProcessSettings.BloomConvolutionTexture;
 
+	// Fall back to the default bloom texture if provided.
+	if (BloomConvolutionTexture == nullptr)
+	{
+		BloomConvolutionTexture = GEngine->DefaultBloomKernelTexture;
+	}
+
 	bool bValidSetup = (BloomConvolutionTexture != nullptr && BloomConvolutionTexture->Resource != nullptr);
 
 	if (bValidSetup && BloomConvolutionTexture->IsFullyStreamedIn() == false)
diff --git a/Engine/Source/Runtime/Renderer/Private/PostProcess/PostProcessHMD.cpp b/Engine/Source/Runtime/Renderer/Private/PostProcess/PostProcessHMD.cpp
--- a/Engine/Source/Runtime/Renderer/Private/PostProcess/PostProcessHMD.cpp
+++ b/Engine/Source/Runtime/Renderer/Private/PostProcess/PostProcessHMD.cpp
@@ -181,7 +181,7 @@ void FRCPassPostProcessHMD::Process(FRenderingCompositePassContext& Context)
 	{
 		SetRenderTarget(Context.RHICmdList, DestRenderTarget.TargetableTexture, FTextureRHIRef());
 		Context.SetViewportAndCallRHI(DestRect);
-		DrawClearQuad(Context.RHICmdList, Context.GetFeatureLevel(), FLinearColor::Black);
+		DrawClearQuad(Context.RHICmdList, FLinearColor::Black);
 	}
 
 	FGraphicsPipelineStateInitializer GraphicsPSOInit;
diff --git a/Engine/Source/Runtime/Renderer/Private/PostProcess/PostProcessLensFlares.cpp b/Engine/Source/Runtime/Renderer/Private/PostProcess/PostProcessLensFlares.cpp
--- a/Engine/Source/Runtime/Renderer/Private/PostProcess/PostProcessLensFlares.cpp
+++ b/Engine/Source/Runtime/Renderer/Private/PostProcess/PostProcessLensFlares.cpp
@@ -168,12 +168,12 @@ void FRCPassPostProcessLensFlares::Process(FRenderingCompositePassContext& Conte
 		
 	if (Context.HasHmdMesh() && View.StereoPass == eSSP_LEFT_EYE)
 	{
-		DrawClearQuad(Context.RHICmdList, Context.GetFeatureLevel(), true, FLinearColor::Black, false, 0, false, 0);
+		DrawClearQuad(Context.RHICmdList, true, FLinearColor::Black, false, 0, false, 0);
 	}
 	else
 	{
 		// is optimized away if possible (RT size=view size, )
-		DrawClearQuad(Context.RHICmdList, Context.GetFeatureLevel(), true, FLinearColor::Black, false, 0, false, 0, PassOutputs[0].RenderTargetDesc.Extent, ViewRect1);
+		DrawClearQuad(Context.RHICmdList, true, FLinearColor::Black, false, 0, false, 0, PassOutputs[0].RenderTargetDesc.Extent, ViewRect1);
 	}
 
 	Context.SetViewportAndCallRHI(ViewRect1);
diff --git a/Engine/Source/Runtime/Renderer/Private/PostProcess/PostProcessMaterial.cpp b/Engine/Source/Runtime/Renderer/Private/PostProcess/PostProcessMaterial.cpp
--- a/Engine/Source/Runtime/Renderer/Private/PostProcess/PostProcessMaterial.cpp
+++ b/Engine/Source/Runtime/Renderer/Private/PostProcess/PostProcessMaterial.cpp
@@ -239,11 +239,11 @@ void FRCPassPostProcessMaterial::Process(FRenderingCompositePassContext& Context
 	if (Context.HasHmdMesh() && View.StereoPass == eSSP_LEFT_EYE)
 	{
 		// needed when using an hmd mesh instead of a full screen quad because we don't touch all of the pixels in the render target
-		DrawClearQuad(Context.RHICmdList, Context.GetFeatureLevel(), FLinearColor::Black);
+		DrawClearQuad(Context.RHICmdList, FLinearColor::Black);
 	}
 	else if (ViewFamily.RenderTarget->GetRenderTargetTexture() != DestRenderTarget.TargetableTexture)
 	{
-		DrawClearQuad(Context.RHICmdList, Context.GetFeatureLevel(), true, FLinearColor::Black, false, 0, false, 0, PassOutputs[0].RenderTargetDesc.Extent, View.ViewRect);
+		DrawClearQuad(Context.RHICmdList, true, FLinearColor::Black, false, 0, false, 0, PassOutputs[0].RenderTargetDesc.Extent, View.ViewRect);
 	}
 
 	Context.SetViewportAndCallRHI(View.ViewRect);
diff --git a/Engine/Source/Runtime/Renderer/Private/PostProcess/PostProcessMobile.cpp b/Engine/Source/Runtime/Renderer/Private/PostProcess/PostProcessMobile.cpp
--- a/Engine/Source/Runtime/Renderer/Private/PostProcess/PostProcessMobile.cpp
+++ b/Engine/Source/Runtime/Renderer/Private/PostProcess/PostProcessMobile.cpp
@@ -239,7 +239,7 @@ void FRCPassPostProcessBloomSetupES2::Process(FRenderingCompositePassContext& Co
 	else
 	{
 		SetRenderTarget(Context.RHICmdList, DestRenderTarget.TargetableTexture, FTextureRHIRef());
-		DrawClearQuad(Context.RHICmdList, Context.GetFeatureLevel(), FLinearColor::Black);
+		DrawClearQuad(Context.RHICmdList, FLinearColor::Black);
 	}
 
 	Context.SetViewportAndCallRHI(0, 0, 0.0f, DstX, DstY, 1.0f );
@@ -454,7 +454,7 @@ void FRCPassPostProcessBloomSetupSmallES2::Process(FRenderingCompositePassContex
 	else
 	{
 		SetRenderTarget(Context.RHICmdList, DestRenderTarget.TargetableTexture, FTextureRHIRef());
-		DrawClearQuad(Context.RHICmdList, Context.GetFeatureLevel(), FLinearColor::Black);
+		DrawClearQuad(Context.RHICmdList, FLinearColor::Black);
 	}
 
 	Context.SetViewportAndCallRHI(0, 0, 0.0f, DstX, DstY, 1.0f );
@@ -613,7 +613,7 @@ void FRCPassPostProcessBloomDownES2::Process(FRenderingCompositePassContext& Con
 	else
 	{
 		SetRenderTarget(Context.RHICmdList, DestRenderTarget.TargetableTexture, FTextureRHIRef());
-		DrawClearQuad(Context.RHICmdList, Context.GetFeatureLevel(), FLinearColor::Black);
+		DrawClearQuad(Context.RHICmdList, FLinearColor::Black);
 	}
 
 	Context.SetViewportAndCallRHI(0, 0, 0.0f, DstX, DstY, 1.0f );
@@ -795,7 +795,7 @@ void FRCPassPostProcessBloomUpES2::Process(FRenderingCompositePassContext& Conte
 	else
 	{
 		SetRenderTarget(Context.RHICmdList, DestRenderTarget.TargetableTexture, FTextureRHIRef());
-		DrawClearQuad(Context.RHICmdList, Context.GetFeatureLevel(), FLinearColor::Black);
+		DrawClearQuad(Context.RHICmdList, FLinearColor::Black);
 	}
 
 	Context.SetViewportAndCallRHI(0, 0, 0.0f, DstX, DstY, 1.0f );
@@ -1102,7 +1102,7 @@ void FRCPassPostProcessSunMaskES2::Process(FRenderingCompositePassContext& Conte
 		else
 		{
 			SetRenderTarget(Context.RHICmdList, DestRenderTarget.TargetableTexture, FTextureRHIRef());
-			DrawClearQuad(Context.RHICmdList, Context.GetFeatureLevel(), FLinearColor::Black);
+			DrawClearQuad(Context.RHICmdList, FLinearColor::Black);
 		}
 
 		Context.SetViewportAndCallRHI(0, 0, 0.0f, DstX, DstY, 1.0f );
@@ -1305,7 +1305,7 @@ void FRCPassPostProcessSunAlphaES2::Process(FRenderingCompositePassContext& Cont
 	else
 	{
 		SetRenderTarget(Context.RHICmdList, DestRenderTarget.TargetableTexture, FTextureRHIRef());
-		DrawClearQuad(Context.RHICmdList, Context.GetFeatureLevel(), FLinearColor::Black);
+		DrawClearQuad(Context.RHICmdList, FLinearColor::Black);
 	}
 
 	Context.SetViewportAndCallRHI(0, 0, 0.0f, DstX, DstY, 1.0f );
@@ -1468,7 +1468,7 @@ void FRCPassPostProcessSunBlurES2::Process(FRenderingCompositePassContext& Conte
 	else
 	{
 		SetRenderTarget(Context.RHICmdList, DestRenderTarget.TargetableTexture, FTextureRHIRef());
-		DrawClearQuad(Context.RHICmdList, Context.GetFeatureLevel(), FLinearColor::Black);
+		DrawClearQuad(Context.RHICmdList, FLinearColor::Black);
 	}
 
 	Context.SetViewportAndCallRHI(0, 0, 0.0f, DstX, DstY, 1.0f );
@@ -1718,7 +1718,7 @@ void FRCPassPostProcessSunMergeES2::Process(FRenderingCompositePassContext& Cont
 	else
 	{
 		SetRenderTarget(Context.RHICmdList, DestRenderTarget.TargetableTexture, FTextureRHIRef());
-		DrawClearQuad(Context.RHICmdList, Context.GetFeatureLevel(), FLinearColor::Black);
+		DrawClearQuad(Context.RHICmdList, FLinearColor::Black);
 	}
 
 	Context.SetViewportAndCallRHI(0, 0, 0.0f, DstX, DstY, 1.0f );
@@ -1924,7 +1924,7 @@ void FRCPassPostProcessSunMergeSmallES2::Process(FRenderingCompositePassContext&
 	else
 	{
 		SetRenderTarget(Context.RHICmdList, DestRenderTarget.TargetableTexture, FTextureRHIRef());
-		DrawClearQuad(Context.RHICmdList, Context.GetFeatureLevel(), FLinearColor::Black);
+		DrawClearQuad(Context.RHICmdList, FLinearColor::Black);
 	}
 
 	Context.SetViewportAndCallRHI(0, 0, 0.0f, DstX, DstY, 1.0f );
@@ -2164,7 +2164,7 @@ void FRCPassPostProcessDofDownES2::Process(FRenderingCompositePassContext& Conte
 	else
 	{
 		SetRenderTarget(Context.RHICmdList, DestRenderTarget.TargetableTexture, FTextureRHIRef());
-		DrawClearQuad(Context.RHICmdList, Context.GetFeatureLevel(), FLinearColor::Black);
+		DrawClearQuad(Context.RHICmdList, FLinearColor::Black);
 	}
 
 	Context.SetViewportAndCallRHI(0, 0, 0.0f, DstX, DstY, 1.0f );
@@ -2367,7 +2367,7 @@ void FRCPassPostProcessDofNearES2::Process(FRenderingCompositePassContext& Conte
 	else
 	{
 		SetRenderTarget(Context.RHICmdList, DestRenderTarget.TargetableTexture, FTextureRHIRef());
-		DrawClearQuad(Context.RHICmdList, Context.GetFeatureLevel(), FLinearColor::Black);
+		DrawClearQuad(Context.RHICmdList, FLinearColor::Black);
 	}
 
 	Context.SetViewportAndCallRHI(0, 0, 0.0f, DstX, DstY, 1.0f );
@@ -2523,7 +2523,7 @@ void FRCPassPostProcessDofBlurES2::Process(FRenderingCompositePassContext& Conte
 	else
 	{
 		SetRenderTarget(Context.RHICmdList, DestRenderTarget.TargetableTexture, FTextureRHIRef());
-		DrawClearQuad(Context.RHICmdList, Context.GetFeatureLevel(), FLinearColor::Black);
+		DrawClearQuad(Context.RHICmdList, FLinearColor::Black);
 	}
 
 	Context.SetViewportAndCallRHI(0, 0, 0.0f, DstX, DstY, 1.0f );
@@ -2726,7 +2726,7 @@ void FRCPassPostProcessSunAvgES2::Process(FRenderingCompositePassContext& Contex
 	else
 	{
 		SetRenderTarget(Context.RHICmdList, DestRenderTarget.TargetableTexture, FTextureRHIRef());
-		DrawClearQuad(Context.RHICmdList, Context.GetFeatureLevel(), FLinearColor::Black);
+		DrawClearQuad(Context.RHICmdList, FLinearColor::Black);
 	}
 
 	Context.SetViewportAndCallRHI(0, 0, 0.0f, DstX, DstY, 1.0f );
@@ -2987,7 +2987,7 @@ void FRCPassPostProcessAaES2::Process(FRenderingCompositePassContext& Context)
 		else
 		{
 			SetRenderTarget(Context.RHICmdList, DestRenderTarget.TargetableTexture, FTextureRHIRef(), ESimpleRenderTargetMode::EUninitializedColorAndDepth);
-			DrawClearQuad(Context.RHICmdList, Context.GetFeatureLevel(), FLinearColor::Black);
+			DrawClearQuad(Context.RHICmdList, FLinearColor::Black);
 		}
 	}
 	else
diff --git a/Engine/Source/Runtime/Renderer/Private/PostProcess/PostProcessMotionBlur.cpp b/Engine/Source/Runtime/Renderer/Private/PostProcess/PostProcessMotionBlur.cpp
--- a/Engine/Source/Runtime/Renderer/Private/PostProcess/PostProcessMotionBlur.cpp
+++ b/Engine/Source/Runtime/Renderer/Private/PostProcess/PostProcessMotionBlur.cpp
@@ -18,6 +18,7 @@
 #include "DeferredShadingRenderer.h"
 #include "ClearQuad.h"
 #include "PipelineStateCache.h"
+#include "SpriteIndexBuffer.h"
 
 #if !(UE_BUILD_SHIPPING || UE_BUILD_TEST)
 static TAutoConsoleVariable<int32> CVarMotionBlurFiltering(
@@ -212,32 +213,7 @@ FPooledRenderTargetDesc FRCPassPostProcessVelocityFlatten::ComputeOutputDesc(EPa
 }
 
 
-class FScatterQuadIndexBuffer : public FIndexBuffer
-{
-public:
-	virtual void InitRHI() override
-	{
-		const uint32 Size = sizeof(uint16) * 6 * 8;
-		const uint32 Stride = sizeof(uint16);
-		FRHIResourceCreateInfo CreateInfo;
-		void* Buffer = nullptr;
-		IndexBufferRHI = RHICreateAndLockIndexBuffer( Stride, Size, BUF_Static, CreateInfo, Buffer );
-		uint16* Indices = (uint16*)Buffer;
-		for (uint32 SpriteIndex = 0; SpriteIndex < 8; ++SpriteIndex)
-		{
-			Indices[SpriteIndex*6 + 0] = SpriteIndex*4 + 0;
-			Indices[SpriteIndex*6 + 1] = SpriteIndex*4 + 3;
-			Indices[SpriteIndex*6 + 2] = SpriteIndex*4 + 2;
-			Indices[SpriteIndex*6 + 3] = SpriteIndex*4 + 0;
-			Indices[SpriteIndex*6 + 4] = SpriteIndex*4 + 1;
-			Indices[SpriteIndex*6 + 5] = SpriteIndex*4 + 3;
-		}
-		RHIUnlockIndexBuffer( IndexBufferRHI );
-	}
-};
-
-TGlobalResource< FScatterQuadIndexBuffer > GScatterQuadIndexBuffer;
-
+TGlobalResource< FSpriteIndexBuffer<8> > GScatterQuadIndexBuffer;
 
 class FPostProcessVelocityScatterVS : public FGlobalShader
 {
@@ -862,7 +838,7 @@ void FRCPassPostProcessMotionBlur::Process(FRenderingCompositePassContext& Conte
 		SetRenderTarget(Context.RHICmdList, DestRenderTarget.TargetableTexture, FTextureRHIRef());
 
 		// Is optimized away if possible (RT size=view size, )
-		DrawClearQuad(Context.RHICmdList, Context.GetFeatureLevel(), true, FLinearColor::Black, false, 0, false, 0, DestSize, SrcRect);
+		DrawClearQuad(Context.RHICmdList, true, FLinearColor::Black, false, 0, false, 0, DestSize, SrcRect);
 
 		Context.SetViewportAndCallRHI(SrcRect);
 	
@@ -1084,7 +1060,7 @@ void FRCPassPostProcessVisualizeMotionBlur::Process(FRenderingCompositePassConte
 	SetRenderTarget(Context.RHICmdList, DestRenderTarget.TargetableTexture, FTextureRHIRef());
 	
 	// is optimized away if possible (RT size=view size, )
-	DrawClearQuad(Context.RHICmdList, Context.GetFeatureLevel(), true, FLinearColor::Black, false, 0, false, 0, PassOutputs[0].RenderTargetDesc.Extent, SrcRect);
+	DrawClearQuad(Context.RHICmdList, true, FLinearColor::Black, false, 0, false, 0, PassOutputs[0].RenderTargetDesc.Extent, SrcRect);
 
 	Context.SetViewportAndCallRHI(SrcRect);
 
diff --git a/Engine/Source/Runtime/Renderer/Private/PostProcess/PostProcessNoiseBlur.cpp b/Engine/Source/Runtime/Renderer/Private/PostProcess/PostProcessNoiseBlur.cpp
--- a/Engine/Source/Runtime/Renderer/Private/PostProcess/PostProcessNoiseBlur.cpp
+++ b/Engine/Source/Runtime/Renderer/Private/PostProcess/PostProcessNoiseBlur.cpp
@@ -151,7 +151,7 @@ void FRCPassPostProcessNoiseBlur::Process(FRenderingCompositePassContext& Contex
 	SetRenderTarget(Context.RHICmdList, DestRenderTarget.TargetableTexture, FTextureRHIRef());
 
 	// is optimized away if possible (RT size=view size, )
-	DrawClearQuad(Context.RHICmdList, Context.GetFeatureLevel(), true, FLinearColor(0, 0, 0, 0), false, 0, false, 0, PassOutputs[0].RenderTargetDesc.Extent, DestRect);
+	DrawClearQuad(Context.RHICmdList, true, FLinearColor(0, 0, 0, 0), false, 0, false, 0, PassOutputs[0].RenderTargetDesc.Extent, DestRect);
 
 	Context.SetViewportAndCallRHI(0, 0, 0.0f, DestSize.X, DestSize.Y, 1.0f );
 	
diff --git a/Engine/Source/Runtime/Renderer/Private/PostProcess/PostProcessSelectionOutline.cpp b/Engine/Source/Runtime/Renderer/Private/PostProcess/PostProcessSelectionOutline.cpp
--- a/Engine/Source/Runtime/Renderer/Private/PostProcess/PostProcessSelectionOutline.cpp
+++ b/Engine/Source/Runtime/Renderer/Private/PostProcess/PostProcessSelectionOutline.cpp
@@ -111,16 +111,16 @@ void FRCPassPostProcessSelectionOutlineColor::Process(FRenderingCompositePassCon
 
 			// top
 			Context.RHICmdList.SetScissorRect(true, ViewRect.Min.X, ViewRect.Min.Y, ViewRect.Max.X, InnerRect.Min.Y);
-			DrawClearQuad(Context.RHICmdList, Context.GetFeatureLevel(), false, FLinearColor(), true, (float)ERHIZBuffer::FarPlane, true, 0, PassOutputs[0].RenderTargetDesc.Extent, FIntRect());
+			DrawClearQuad(Context.RHICmdList, false, FLinearColor(), true, (float)ERHIZBuffer::FarPlane, true, 0, PassOutputs[0].RenderTargetDesc.Extent, FIntRect());
 			// bottom
 			Context.RHICmdList.SetScissorRect(true, ViewRect.Min.X, InnerRect.Max.Y, ViewRect.Max.X, ViewRect.Max.Y);
-			DrawClearQuad(Context.RHICmdList, Context.GetFeatureLevel(), false, FLinearColor(), true, (float)ERHIZBuffer::FarPlane, true, 0, PassOutputs[0].RenderTargetDesc.Extent, FIntRect());
+			DrawClearQuad(Context.RHICmdList, false, FLinearColor(), true, (float)ERHIZBuffer::FarPlane, true, 0, PassOutputs[0].RenderTargetDesc.Extent, FIntRect());
 			// left
 			Context.RHICmdList.SetScissorRect(true, ViewRect.Min.X, ViewRect.Min.Y, InnerRect.Min.X, ViewRect.Max.Y);
-			DrawClearQuad(Context.RHICmdList, Context.GetFeatureLevel(), false, FLinearColor(), true, (float)ERHIZBuffer::FarPlane, true, 0, PassOutputs[0].RenderTargetDesc.Extent, FIntRect());
+			DrawClearQuad(Context.RHICmdList, false, FLinearColor(), true, (float)ERHIZBuffer::FarPlane, true, 0, PassOutputs[0].RenderTargetDesc.Extent, FIntRect());
 			// right
 			Context.RHICmdList.SetScissorRect(true, InnerRect.Max.X, ViewRect.Min.Y, ViewRect.Max.X, ViewRect.Max.Y);
-			DrawClearQuad(Context.RHICmdList, Context.GetFeatureLevel(), false, FLinearColor(), true, (float)ERHIZBuffer::FarPlane, true, 0, PassOutputs[0].RenderTargetDesc.Extent, FIntRect());
+			DrawClearQuad(Context.RHICmdList, false, FLinearColor(), true, (float)ERHIZBuffer::FarPlane, true, 0, PassOutputs[0].RenderTargetDesc.Extent, FIntRect());
 
 			Context.RHICmdList.SetScissorRect(false, 0, 0, 0, 0);
 		}
diff --git a/Engine/Source/Runtime/Renderer/Private/PostProcess/PostProcessSubsurface.cpp b/Engine/Source/Runtime/Renderer/Private/PostProcess/PostProcessSubsurface.cpp
--- a/Engine/Source/Runtime/Renderer/Private/PostProcess/PostProcessSubsurface.cpp
+++ b/Engine/Source/Runtime/Renderer/Private/PostProcess/PostProcessSubsurface.cpp
@@ -290,7 +290,7 @@ void FRCPassPostProcessSubsurfaceVisualize::Process(FRenderingCompositePassConte
 	SetRenderTarget(Context.RHICmdList, DestRenderTarget.TargetableTexture, FTextureRHIRef());
 
 	// is optimized away if possible (RT size=view size, )
-	DrawClearQuad(Context.RHICmdList, Context.GetFeatureLevel(), true, FLinearColor::Black, false, 0, false, 0, PassOutputs[0].RenderTargetDesc.Extent, DestRect);
+	DrawClearQuad(Context.RHICmdList, true, FLinearColor::Black, false, 0, false, 0, PassOutputs[0].RenderTargetDesc.Extent, DestRect);
 
 	Context.SetViewportAndCallRHI(0, 0, 0.0f, DestSize.X, DestSize.Y, 1.0f );
 
diff --git a/Engine/Source/Runtime/Renderer/Private/PostProcess/PostProcessTemporalAA.cpp b/Engine/Source/Runtime/Renderer/Private/PostProcess/PostProcessTemporalAA.cpp
--- a/Engine/Source/Runtime/Renderer/Private/PostProcess/PostProcessTemporalAA.cpp
+++ b/Engine/Source/Runtime/Renderer/Private/PostProcess/PostProcessTemporalAA.cpp
@@ -451,7 +451,7 @@ void FRCPassPostProcessSSRTemporalAA::Process(FRenderingCompositePassContext& Co
 	SetRenderTarget(Context.RHICmdList, DestRenderTarget.TargetableTexture, FTextureRHIRef());
 
 	// is optimized away if possible (RT size=view size, )
-	DrawClearQuad(Context.RHICmdList, Context.GetFeatureLevel(), true, FLinearColor::Black, false, 0, false, 0, PassOutputs[0].RenderTargetDesc.Extent, SrcRect);
+	DrawClearQuad(Context.RHICmdList, true, FLinearColor::Black, false, 0, false, 0, PassOutputs[0].RenderTargetDesc.Extent, SrcRect);
 
 	Context.SetViewportAndCallRHI(SrcRect);
 
@@ -584,7 +584,7 @@ void FRCPassPostProcessDOFTemporalAA::Process(FRenderingCompositePassContext& Co
 		SetRenderTarget(Context.RHICmdList, DestRenderTarget.TargetableTexture, FTextureRHIRef());
 
 		// is optimized away if possible (RT size=view size, )
-		DrawClearQuad(Context.RHICmdList, Context.GetFeatureLevel(), true, FLinearColor::Black, false, 0, false, 0, DestSize, SrcRect);
+		DrawClearQuad(Context.RHICmdList, true, FLinearColor::Black, false, 0, false, 0, DestSize, SrcRect);
 
 		Context.SetViewportAndCallRHI(SrcRect);
 
@@ -685,7 +685,7 @@ void FRCPassPostProcessDOFTemporalAANear::Process(FRenderingCompositePassContext
 	SetRenderTarget(Context.RHICmdList, DestRenderTarget.TargetableTexture, FTextureRHIRef());
 
 	// is optimized away if possible (RT size=view size, )
-	DrawClearQuad(Context.RHICmdList, Context.GetFeatureLevel(), true, FLinearColor::Black, false, 0, false, 0, DestSize, SrcRect);
+	DrawClearQuad(Context.RHICmdList, true, FLinearColor::Black, false, 0, false, 0, DestSize, SrcRect);
 
 	Context.SetViewportAndCallRHI(SrcRect);
 
@@ -773,7 +773,7 @@ void FRCPassPostProcessLightShaftTemporalAA::Process(FRenderingCompositePassCont
 	SetRenderTarget(Context.RHICmdList, DestRenderTarget.TargetableTexture, FTextureRHIRef());
 
 	// is optimized away if possible (RT size=view size, )
-	DrawClearQuad(Context.RHICmdList, Context.GetFeatureLevel(), true, FLinearColor::Black, false, 0, false, 0, DestSize, SrcRect);
+	DrawClearQuad(Context.RHICmdList, true, FLinearColor::Black, false, 0, false, 0, DestSize, SrcRect);
 
 	Context.SetViewportAndCallRHI(SrcRect);
 
@@ -915,7 +915,7 @@ void FRCPassPostProcessTemporalAA::Process(FRenderingCompositePassContext& Conte
 		SetRenderTarget(Context.RHICmdList, DestRenderTarget.TargetableTexture, SceneContext.GetSceneDepthTexture(), ESimpleRenderTargetMode::EUninitializedColorExistingDepth, FExclusiveDepthStencil::DepthRead_StencilWrite);
 
 		// is optimized away if possible (RT size=view size, )
-		DrawClearQuad(Context.RHICmdList, Context.GetFeatureLevel(), true, FLinearColor::Black, false, 0, false, 0, DestSize, SrcRect);
+		DrawClearQuad(Context.RHICmdList, true, FLinearColor::Black, false, 0, false, 0, DestSize, SrcRect);
 
 		Context.SetViewportAndCallRHI(SrcRect);
 
diff --git a/Engine/Source/Runtime/Renderer/Private/PostProcess/PostProcessTonemap.cpp b/Engine/Source/Runtime/Renderer/Private/PostProcess/PostProcessTonemap.cpp
--- a/Engine/Source/Runtime/Renderer/Private/PostProcess/PostProcessTonemap.cpp
+++ b/Engine/Source/Runtime/Renderer/Private/PostProcess/PostProcessTonemap.cpp
@@ -907,10 +907,12 @@ public:
 		OutputDevice.Bind(ParameterMap, TEXT("OutputDevice"));
 		OutputGamut.Bind(ParameterMap, TEXT("OutputGamut"));
 		EncodeHDROutput.Bind(ParameterMap, TEXT("EncodeHDROutput"));
+		
+		EyeAdaptation.Bind(ParameterMap, TEXT("EyeAdaptation"));
 	}
 	
 	template <typename TRHICmdList, typename TRHIShader>
-	void Set(TRHICmdList& RHICmdList, const TRHIShader ShaderRHI, const FRenderingCompositePassContext& Context, const TShaderUniformBufferParameter<FBloomDirtMaskParameters>& BloomDirtMaskParam)
+	void Set(TRHICmdList& RHICmdList, const TRHIShader ShaderRHI, const FRenderingCompositePassContext& Context, const TShaderUniformBufferParameter<FBloomDirtMaskParameters>& BloomDirtMaskParam, bool bDoEyeAdaptation = false)
 	{
 		const FPostProcessSettings& Settings = Context.View.FinalPostProcessSettings;
 		const FSceneViewFamily& ViewFamily = *(Context.View.Family);
@@ -1071,6 +1073,26 @@ public:
 			SetShaderValue(RHICmdList, ShaderRHI, ColorShadow_Tint1, Constants[6]);
 			SetShaderValue(RHICmdList, ShaderRHI, ColorShadow_Tint2, Constants[7]);
 		}
+		
+		if(bDoEyeAdaptation)
+		{
+			// Fix for eye adaptation vertex texture read failing in Metal macOS 10.11 - If Mac should be Metal but also check lauguage version (Clone of Vertex Shader version implementation).
+			if (Context.View.HasValidEyeAdaptation())
+			{
+				IPooledRenderTarget* EyeAdaptationRT = Context.View.GetEyeAdaptation(Context.RHICmdList);
+				FTextureRHIParamRef EyeAdaptationRTRef = EyeAdaptationRT->GetRenderTargetItem().TargetableTexture;
+				if (EyeAdaptationRTRef)
+				{
+					Context.RHICmdList.TransitionResources(EResourceTransitionAccess::EReadable, &EyeAdaptationRTRef, 1);
+				}
+				SetTextureParameter(RHICmdList, ShaderRHI, EyeAdaptation, EyeAdaptationRT->GetRenderTargetItem().TargetableTexture);
+			}
+			else
+			{
+				// some views don't have a state, thumbnail rendering?
+				SetTextureParameter(RHICmdList, ShaderRHI, EyeAdaptation, GWhiteTexture->TextureRHI);
+			}
+		}
 	}
 
 	friend FArchive& operator<<(FArchive& Ar,FPostProcessTonemapShaderParameters& P)
@@ -1081,6 +1103,7 @@ public:
 		Ar << P.ColorMatrixR_ColorCurveCd1 << P.ColorMatrixG_ColorCurveCd3Cm3 << P.ColorMatrixB_ColorCurveCm2 << P.ColorCurve_Cm0Cd0_Cd2_Ch0Cm1_Ch3 << P.ColorCurve_Ch1_Ch2 << P.ColorShadow_Luma << P.ColorShadow_Tint1 << P.ColorShadow_Tint2;
 		Ar << P.OverlayColor;
 		Ar << P.OutputDevice << P.OutputGamut << P.EncodeHDROutput;
+		Ar << P.EyeAdaptation;
 
 		return Ar;
 	}
@@ -1111,19 +1134,31 @@ public:
 	FShaderParameter OutputDevice;
 	FShaderParameter OutputGamut;
 	FShaderParameter EncodeHDROutput;
+	
+	//Fix for eye adaptation vertex texture read failing in Metal macOS 10.11
+	FShaderResourceParameter EyeAdaptation;
+};
+
+namespace PostProcessTonemapUtil
+{
+	// Function exists out side of the FPostProcessTonemapPS class - otherwise calling it would require template parameters
+	static inline bool PlatformRequiresEyeAdaptationPSSampling(EShaderPlatform ShaderPlatform)
+	{
+		return IsMetalPlatform(ShaderPlatform) && RHIGetShaderLanguageVersion(ShaderPlatform) < 2;
+	}
 };
 
 /**
  * Encapsulates the post processing tonemapper pixel shader.
  */
-template<uint32 ConfigIndex>
+template<uint32 ConfigIndex, bool bDoEyeAdaptation>
 class FPostProcessTonemapPS : public FGlobalShader
 {
 	DECLARE_SHADER_TYPE(FPostProcessTonemapPS, Global);
 
 	static bool ShouldCache(EShaderPlatform Platform)
 	{
-		return IsFeatureLevelSupported(Platform, ERHIFeatureLevel::ES2);
+		return (IsFeatureLevelSupported(Platform, ERHIFeatureLevel::ES2) && (!bDoEyeAdaptation || PostProcessTonemapUtil::PlatformRequiresEyeAdaptationPSSampling(Platform)));
 	}
 
 	static void ModifyCompilationEnvironment(EShaderPlatform Platform, FShaderCompilerEnvironment& OutEnvironment)
@@ -1144,6 +1179,8 @@ class FPostProcessTonemapPS : public FGlobalShader
 		OutEnvironment.SetDefine(TEXT("USE_COLOR_FRINGE"),		 TonemapperIsDefined(ConfigBitmask, TonemapperColorFringe));
 		OutEnvironment.SetDefine(TEXT("USE_SHARPEN"),	         TonemapperIsDefined(ConfigBitmask, TonemapperSharpen));
 		OutEnvironment.SetDefine(TEXT("USE_VOLUME_LUT"),		 UseVolumeTextureLUT(Platform));
+		
+		OutEnvironment.SetDefine(TEXT("EYEADAPTATION_EXPOSURE_FIX"), bDoEyeAdaptation ? 1 : 0);
 	}
 
 	/** Default constructor. */
@@ -1189,7 +1226,7 @@ public:
 			PostprocessParameter.SetPS(ShaderRHI, Context, 0, eFC_0000, Filters);
 		}
 
-		PostProcessTonemapShaderParameters.Set(Context.RHICmdList, ShaderRHI, Context, GetUniformBufferParameter<FBloomDirtMaskParameters>());
+		PostProcessTonemapShaderParameters.Set(Context.RHICmdList, ShaderRHI, Context, GetUniformBufferParameter<FBloomDirtMaskParameters>(), bDoEyeAdaptation);
 	}
 	
 	static const TCHAR* GetSourceFilename()
@@ -1204,8 +1241,11 @@ public:
 };
 
 // #define avoids a lot of code duplication
-#define VARIATION1(A) typedef FPostProcessTonemapPS<A> FPostProcessTonemapPS##A; \
-	IMPLEMENT_SHADER_TYPE2(FPostProcessTonemapPS##A, SF_Pixel);
+#define VARIATION1(A)																	\
+	typedef FPostProcessTonemapPS<A,true> FPostProcessTonemapPS_EyeAdaptation##A;		\
+	IMPLEMENT_SHADER_TYPE2(FPostProcessTonemapPS_EyeAdaptation##A, SF_Pixel);			\
+	typedef FPostProcessTonemapPS<A,false> FPostProcessTonemapPS_NoEyeAdaptation##A;	\
+	IMPLEMENT_SHADER_TYPE2(FPostProcessTonemapPS_NoEyeAdaptation##A, SF_Pixel);
 
 	VARIATION1(0)  VARIATION1(1)  VARIATION1(2)  VARIATION1(3)  VARIATION1(4)  VARIATION1(5) VARIATION1(6) VARIATION1(7) VARIATION1(8)
 	VARIATION1(9)  VARIATION1(10) VARIATION1(11) VARIATION1(12) VARIATION1(13) VARIATION1(14)
@@ -1424,7 +1464,7 @@ FRCPassPostProcessTonemap::FRCPassPostProcessTonemap(const FViewInfo& InView, bo
 namespace PostProcessTonemapUtil
 {
 	// Template implementation supports unique static BoundShaderState for each permutation of Vertex/Pixel Shaders 
-	template <uint32 ConfigIndex, bool bDoEyeAdaptation>
+	template <uint32 ConfigIndex, bool bVSDoEyeAdaptation, bool bPSDoEyeAdaptation>
 	static inline void SetShaderTempl(const FRenderingCompositePassContext& Context)
 	{
 		FGraphicsPipelineStateInitializer GraphicsPSOInit;
@@ -1433,8 +1473,8 @@ namespace PostProcessTonemapUtil
 		GraphicsPSOInit.RasterizerState = TStaticRasterizerState<>::GetRHI();
 		GraphicsPSOInit.DepthStencilState = TStaticDepthStencilState<false, CF_Always>::GetRHI();
 
-		typedef TPostProcessTonemapVS<bDoEyeAdaptation> VertexShaderType;
-		typedef FPostProcessTonemapPS<ConfigIndex>      PixelShaderType;
+		typedef TPostProcessTonemapVS<bVSDoEyeAdaptation>				VertexShaderType;
+		typedef FPostProcessTonemapPS<ConfigIndex, bPSDoEyeAdaptation>	PixelShaderType;
 
 		TShaderMapRef<PixelShaderType>  PixelShader(Context.GetShaderMap());
 		TShaderMapRef<VertexShaderType> VertexShader(Context.GetShaderMap());
@@ -1455,11 +1495,18 @@ namespace PostProcessTonemapUtil
 	{
 		if (bDoEyeAdaptation)
 		{
-			SetShaderTempl<ConfigIndex, true>(Context);
+			if (PlatformRequiresEyeAdaptationPSSampling(Context.GetShaderPlatform()))
+			{
+				SetShaderTempl<ConfigIndex, true, true>(Context);
+			}
+			else
+			{
+				SetShaderTempl<ConfigIndex, true, false>(Context);
+			}
 		}
 		else
 		{
-			SetShaderTempl<ConfigIndex, false>(Context);
+			SetShaderTempl<ConfigIndex, false, false>(Context);
 		}
 	}
 
@@ -1573,12 +1620,12 @@ void FRCPassPostProcessTonemap::Process(FRenderingCompositePassContext& Context)
 			if (Context.HasHmdMesh() && View.StereoPass == eSSP_LEFT_EYE)
 			{
 				// needed when using an hmd mesh instead of a full screen quad because we don't touch all of the pixels in the render target
-				DrawClearQuad(Context.RHICmdList, GMaxRHIFeatureLevel, FLinearColor::Black);
+				DrawClearQuad(Context.RHICmdList, FLinearColor::Black);
 			}
 			else if (ViewFamily.RenderTarget->GetRenderTargetTexture() != DestRenderTarget.TargetableTexture)
 			{
 				// needed to not have PostProcessAA leaking in content (e.g. Matinee black borders), is optimized away if possible (RT size=view size, )
-				DrawClearQuad(Context.RHICmdList, Context.GetFeatureLevel(), true, FLinearColor::Black, false, 0, false, 0, PassOutputs[0].RenderTargetDesc.Extent, DestRect);
+				DrawClearQuad(Context.RHICmdList, true, FLinearColor::Black, false, 0, false, 0, PassOutputs[0].RenderTargetDesc.Extent, DestRect);
 			}
 		}
 
@@ -2064,7 +2111,7 @@ void FRCPassPostProcessTonemapES2::Process(FRenderingCompositePassContext& Conte
 		// Full clear to avoid restore
 		if ((View.StereoPass == eSSP_FULL && bFirstView) || View.StereoPass == eSSP_LEFT_EYE)
 		{
-			DrawClearQuad(Context.RHICmdList, GMaxRHIFeatureLevel, FLinearColor::Black);
+			DrawClearQuad(Context.RHICmdList, FLinearColor::Black);
 		}
 	}
 
diff --git a/Engine/Source/Runtime/Renderer/Private/PostProcess/PostProcessUpscale.cpp b/Engine/Source/Runtime/Renderer/Private/PostProcess/PostProcessUpscale.cpp
--- a/Engine/Source/Runtime/Renderer/Private/PostProcess/PostProcessUpscale.cpp
+++ b/Engine/Source/Runtime/Renderer/Private/PostProcess/PostProcessUpscale.cpp
@@ -327,7 +327,7 @@ void FRCPassPostProcessUpscale::Process(FRenderingCompositePassContext& Context)
 	Context.SetViewportAndCallRHI(DestRect);
 	if (View.StereoPass == eSSP_FULL || View.StereoPass == eSSP_LEFT_EYE)
 	{
-		DrawClearQuad(Context.RHICmdList, Context.GetFeatureLevel(), true, FLinearColor::Black, false, 0, false, 0, PassOutputs[0].RenderTargetDesc.Extent, ExcludeRect);
+		DrawClearQuad(Context.RHICmdList, true, FLinearColor::Black, false, 0, false, 0, PassOutputs[0].RenderTargetDesc.Extent, ExcludeRect);
 	}
 
 	FShader* VertexShader = 0;
diff --git a/Engine/Source/Runtime/Renderer/Private/PostProcess/PostProcessWeightedSampleSum.cpp b/Engine/Source/Runtime/Renderer/Private/PostProcess/PostProcessWeightedSampleSum.cpp
--- a/Engine/Source/Runtime/Renderer/Private/PostProcess/PostProcessWeightedSampleSum.cpp
+++ b/Engine/Source/Runtime/Renderer/Private/PostProcess/PostProcessWeightedSampleSum.cpp
@@ -1129,7 +1129,7 @@ void FRCPassPostProcessWeightedSampleSum::DrawClear(FRHICommandListImmediate& RH
 		AdjustRectsForFastBlur(SrcRect, DestRect);
 	}
 
-	DrawClearQuad(RHICmdList, FeatureLevel, true, FLinearColor(0, 0, 0, 0), false, 1.0f, false, 0, DestSize, DestRect);	
+	DrawClearQuad(RHICmdList, true, FLinearColor(0, 0, 0, 0), false, 1.0f, false, 0, DestSize, DestRect);	
 }
 
 void FRCPassPostProcessWeightedSampleSum::DrawQuad(FRHICommandListImmediate& RHICmdList, ERHIFeatureLevel::Type FeatureLevel, bool bDoFastBlur, FIntRect SrcRect, FIntRect DestRect, FIntPoint DestSize, FIntPoint SrcSize, FShader* VertexShader) const
diff --git a/Engine/Source/Runtime/Renderer/Private/PostProcess/PostProcessing.cpp b/Engine/Source/Runtime/Renderer/Private/PostProcess/PostProcessing.cpp
--- a/Engine/Source/Runtime/Renderer/Private/PostProcess/PostProcessing.cpp
+++ b/Engine/Source/Runtime/Renderer/Private/PostProcess/PostProcessing.cpp
@@ -719,8 +719,7 @@ static FRenderingCompositeOutputRef AddBloom(FBloomDownSampleArray& BloomDownSam
 	// Extract the Context
 	FPostprocessContext& Context = BloomDownSampleArray.Context;
 
-	const bool bUseFFTBloom = (Context.View.FinalPostProcessSettings.BloomConvolutionTexture != nullptr  
-		&& Context.View.FinalPostProcessSettings.BloomMethod == EBloomMethod::BM_FFT
+	const bool bUseFFTBloom = (Context.View.FinalPostProcessSettings.BloomMethod == EBloomMethod::BM_FFT
 		&& Context.View.FeatureLevel >= ERHIFeatureLevel::SM5);
 
 	// Extract the downsample array.
diff --git a/Engine/Source/Runtime/Renderer/Private/PostProcess/RenderTargetPool.cpp b/Engine/Source/Runtime/Renderer/Private/PostProcess/RenderTargetPool.cpp
--- a/Engine/Source/Runtime/Renderer/Private/PostProcess/RenderTargetPool.cpp
+++ b/Engine/Source/Runtime/Renderer/Private/PostProcess/RenderTargetPool.cpp
@@ -190,7 +190,7 @@ void FRenderTargetPool::TransitionTargetsWritable(FRHICommandListImmediate& RHIC
 	if (TransitionTargets.Num() > 0)
 	{
 		RHICmdList.TransitionResourceArrayNoCopy(EResourceTransitionAccess::EWritable, TransitionTargets);
-		if (GRHIThread)
+		if (IsRunningRHIInSeparateThread())
 		{
 			TransitionFence = RHICmdList.RHIThreadFence(false);
 		}
@@ -494,17 +494,17 @@ Done:
 			if(Found->GetDesc().TargetableFlags & TexCreate_RenderTargetable)
 			{
 				SetRenderTarget(RHICmdList, Found->RenderTargetItem.TargetableTexture, FTextureRHIRef());
-				DrawClearQuad(RHICmdList, GMaxRHIFeatureLevel, FLinearColor(1000, 1000, 1000, 1000));
+				DrawClearQuad(RHICmdList, FLinearColor(1000, 1000, 1000, 1000));
 			}
 			else if(Found->GetDesc().TargetableFlags & TexCreate_UAV)
 			{
-				ClearUAV(RHICmdList, GMaxRHIFeatureLevel, Found->RenderTargetItem, FLinearColor(1000, 1000, 1000, 1000));
+				ClearUAV(RHICmdList, Found->RenderTargetItem, FLinearColor(1000, 1000, 1000, 1000));
 			}
 
 			if(Desc.TargetableFlags & TexCreate_DepthStencilTargetable)
 			{
 				SetRenderTarget(RHICmdList, FTextureRHIRef(), Found->RenderTargetItem.TargetableTexture);
-				DrawClearQuad(RHICmdList, GMaxRHIFeatureLevel, false, FLinearColor::Black, true, 0, false, 0);
+				DrawClearQuad(RHICmdList, false, FLinearColor::Black, true, 0.0f, true, 0);
 			}
 		}
 	}
diff --git a/Engine/Source/Runtime/Renderer/Private/PostProcess/SceneRenderTargets.cpp b/Engine/Source/Runtime/Renderer/Private/PostProcess/SceneRenderTargets.cpp
--- a/Engine/Source/Runtime/Renderer/Private/PostProcess/SceneRenderTargets.cpp
+++ b/Engine/Source/Runtime/Renderer/Private/PostProcess/SceneRenderTargets.cpp
@@ -230,7 +230,7 @@ FSceneRenderTargets::FSceneRenderTargets(const FViewInfo& View, const FSceneRend
 	, EditorPrimitivesColor(GRenderTargetPool.MakeSnapshot(SnapshotSource.EditorPrimitivesColor))
 	, EditorPrimitivesDepth(GRenderTargetPool.MakeSnapshot(SnapshotSource.EditorPrimitivesDepth))
 	, SeparateTranslucencyRT(SnapshotSource.SeparateTranslucencyRT)
-	, SeparateTranslucencyDepthRT(SnapshotSource.SeparateTranslucencyDepthRT)
+	, DownsampledTranslucencyDepthRT(SnapshotSource.DownsampledTranslucencyDepthRT)
 	, bScreenSpaceAOIsValid(SnapshotSource.bScreenSpaceAOIsValid)
 	, bCustomDepthIsValid(SnapshotSource.bCustomDepthIsValid)
 	, GBufferRefCount(SnapshotSource.GBufferRefCount)
@@ -591,7 +591,7 @@ void FSceneRenderTargets::SetQuadOverdrawUAV(FRHICommandList& RHICmdList, bool b
 			Info.UnorderedAccessView[Info.NumUAVs++] = QuadOverdrawBuffer->GetRenderTargetItem().UAV;
 
 			// Clear to default value
-			ClearUAV(RHICmdList, GMaxRHIFeatureLevel, QuadOverdrawBuffer->GetRenderTargetItem(), FLinearColor::Transparent);
+			ClearUAV(RHICmdList, QuadOverdrawBuffer->GetRenderTargetItem(), FLinearColor::Transparent);
 			RHICmdList.TransitionResource(EResourceTransitionAccess::ERWBarrier, EResourceTransitionPipeline::EGfxToGfx, QuadOverdrawBuffer->GetRenderTargetItem().UAV);
 		}
 	}
@@ -669,7 +669,7 @@ void FSceneRenderTargets::BeginRenderingGBuffer(FRHICommandList& RHICmdList, ERe
 				Textures[i] = RenderTargets[i].Texture;
 			}
 			//depth/stencil should have been handled by the fast clear.  only color for RT0 can get changed.
-			DrawClearQuadMRT(RHICmdList, GMaxRHIFeatureLevel, true, MRTCount, ClearColors, false, 0, false, 0);
+			DrawClearQuadMRT(RHICmdList, true, MRTCount, ClearColors, false, 0, false, 0);
 		}
 
 		//bind any clear data that won't be bound automatically by the preceding SetRenderTargetsAndClear
@@ -1318,7 +1318,7 @@ void FSceneRenderTargets::BeginRenderingLightAttenuation(FRHICommandList& RHICmd
 		TransitionSetRenderTargetsHelper(RHICmdList, GetLightAttenuationSurface(), GetSceneDepthSurface(), FExclusiveDepthStencil::DepthRead_StencilWrite);
 
 		FRHIRenderTargetView View = FRHIRenderTargetView(GetLightAttenuationSurface(), ERenderTargetLoadAction::EClear);
-		FRHISetRenderTargetsInfo Info(1, &View, FRHIDepthRenderTargetView(GetSceneDepthSurface(), ERenderTargetLoadAction::ELoad, ERenderTargetStoreAction::ENoAction, ERenderTargetLoadAction::ELoad, ERenderTargetStoreAction::ENoAction));
+		FRHISetRenderTargetsInfo Info(1, &View, FRHIDepthRenderTargetView(GetSceneDepthSurface(), ERenderTargetLoadAction::ELoad, ERenderTargetStoreAction::ENoAction, FExclusiveDepthStencil::DepthRead_StencilWrite));
 		RHICmdList.SetRenderTargetsAndClear(Info);
 	}
 	else
@@ -1352,16 +1352,16 @@ TRefCountPtr<IPooledRenderTarget>& FSceneRenderTargets::GetSeparateTranslucency(
 	return SeparateTranslucencyRT;
 }
 
-TRefCountPtr<IPooledRenderTarget>& FSceneRenderTargets::GetSeparateTranslucencyDepth(FRHICommandList& RHICmdList, FIntPoint Size)
+TRefCountPtr<IPooledRenderTarget>& FSceneRenderTargets::GetDownsampledTranslucencyDepth(FRHICommandList& RHICmdList, FIntPoint Size)
 {
-	if (!SeparateTranslucencyDepthRT || SeparateTranslucencyDepthRT->GetDesc().Extent != Size)
+	if (!DownsampledTranslucencyDepthRT || DownsampledTranslucencyDepthRT->GetDesc().Extent != Size)
 	{
 		// Create the SeparateTranslucency depth render target 
 		FPooledRenderTargetDesc Desc(FPooledRenderTargetDesc::Create2DDesc(Size, PF_DepthStencil, FClearValueBinding::None, TexCreate_None, TexCreate_DepthStencilTargetable, false));
 		Desc.NumSamples = GetNumSceneColorMSAASamples(CurrentFeatureLevel);
-		GRenderTargetPool.FindFreeElement(RHICmdList, Desc, SeparateTranslucencyDepthRT, TEXT("SeparateTranslucencyDepth"));
+		GRenderTargetPool.FindFreeElement(RHICmdList, Desc, DownsampledTranslucencyDepthRT, TEXT("SeparateTranslucencyDepth"));
 	}
-	return SeparateTranslucencyDepthRT;
+	return DownsampledTranslucencyDepthRT;
 }
 
 void FSceneRenderTargets::BeginRenderingTranslucency(FRHICommandList& RHICmdList, const FViewInfo& View, bool bFirstTimeThisFrame)
@@ -1373,7 +1373,7 @@ void FSceneRenderTargets::BeginRenderingTranslucency(FRHICommandList& RHICmdList
 	{
 		// Clear the stencil buffer for ResponsiveAA
 		const FTexture2DRHIRef& DepthSurface = GetSceneDepthSurface();
-		DrawClearQuad(RHICmdList, CurrentFeatureLevel, false, FLinearColor(), false, 0, true, 0, FIntPoint(DepthSurface->GetSizeX(), DepthSurface->GetSizeY()), View.ViewRect);
+		DrawClearQuad(RHICmdList, false, FLinearColor(), false, 0, true, 0, FIntPoint(DepthSurface->GetSizeX(), DepthSurface->GetSizeY()), View.ViewRect);
 	}
 		
 	// viewport to match view size
@@ -1382,14 +1382,8 @@ void FSceneRenderTargets::BeginRenderingTranslucency(FRHICommandList& RHICmdList
 
 void FSceneRenderTargets::BeginRenderingSeparateTranslucency(FRHICommandList& RHICmdList, const FViewInfo& View, bool bFirstTimeThisFrame)
 {
-	check(IsSeparateTranslucencyActive(View));
-
 	bSeparateTranslucencyPass = true;
 
-	FIntPoint ScaledSize;
-	float Scale = 1.0f;
-	GetSeparateTranslucencyDimensions(ScaledSize, Scale);
-
 	SCOPED_DRAW_EVENT(RHICmdList, BeginSeparateTranslucency);
 
 	TRefCountPtr<IPooledRenderTarget>* SeparateTranslucency;
@@ -1400,9 +1394,9 @@ void FSceneRenderTargets::BeginRenderingSeparateTranslucency(FRHICommandList& RH
 	}
 	else
 	{
-		SeparateTranslucency = &GetSeparateTranslucency(RHICmdList, ScaledSize);
+		SeparateTranslucency = &GetSeparateTranslucency(RHICmdList, SeparateTranslucencyBufferSize);
 	}
-	const FTexture2DRHIRef &SeparateTranslucencyDepth = Scale < 1.0f ? (const FTexture2DRHIRef&)GetSeparateTranslucencyDepth(RHICmdList, ScaledSize)->GetRenderTargetItem().TargetableTexture : GetSceneDepthSurface();
+	const FTexture2DRHIRef &SeparateTranslucencyDepth = SeparateTranslucencyScale < 1.0f ? (const FTexture2DRHIRef&)GetDownsampledTranslucencyDepth(RHICmdList, SeparateTranslucencyBufferSize)->GetRenderTargetItem().TargetableTexture : GetSceneDepthSurface();
 
 	check((*SeparateTranslucency)->GetRenderTargetItem().TargetableTexture->GetClearColor() == FLinearColor::Black);
 	// clear the render target the first time, re-use afterwards
@@ -1416,10 +1410,10 @@ void FSceneRenderTargets::BeginRenderingSeparateTranslucency(FRHICommandList& RH
 		RHICmdList.BindClearMRTValues(true, false, true);
 	}
 
-	RHICmdList.SetViewport(View.ViewRect.Min.X * Scale, View.ViewRect.Min.Y * Scale, 0.0f, View.ViewRect.Max.X * Scale, View.ViewRect.Max.Y * Scale, 1.0f);
+	RHICmdList.SetViewport(View.ViewRect.Min.X * SeparateTranslucencyScale, View.ViewRect.Min.Y * SeparateTranslucencyScale, 0.0f, View.ViewRect.Max.X * SeparateTranslucencyScale, View.ViewRect.Max.Y * SeparateTranslucencyScale, 1.0f);
 }
 
-void FSceneRenderTargets::FinishRenderingSeparateTranslucency(FRHICommandList& RHICmdList)
+void FSceneRenderTargets::FinishRenderingSeparateTranslucency(FRHICommandList& RHICmdList, const FViewInfo& View)
 {
 	SCOPED_DRAW_EVENT(RHICmdList, FinishSeparateTranslucency);
 
@@ -1429,19 +1423,23 @@ void FSceneRenderTargets::FinishRenderingSeparateTranslucency(FRHICommandList& R
 	{
 		check(SeparateTranslucencyRT.GetReference());
 		SeparateTranslucency = &SeparateTranslucencyRT;
-		SeparateTranslucencyDepth = &SeparateTranslucencyDepthRT;
+		SeparateTranslucencyDepth = SeparateTranslucencyScale < 1.f ? &DownsampledTranslucencyDepthRT : &SceneDepthZ;
 	}
 	else
 	{
-		FIntPoint ScaledSize;
-		float Scale = 1.0f;
-		GetSeparateTranslucencyDimensions(ScaledSize, Scale);
-		SeparateTranslucency = &GetSeparateTranslucency(RHICmdList, ScaledSize);
-		SeparateTranslucencyDepth = &GetSeparateTranslucencyDepth(RHICmdList, ScaledSize);
+		SeparateTranslucency = &GetSeparateTranslucency(RHICmdList, SeparateTranslucencyBufferSize);
+		SeparateTranslucencyDepth = SeparateTranslucencyScale < 1.f ? &GetDownsampledTranslucencyDepth(RHICmdList, SeparateTranslucencyBufferSize) : &SceneDepthZ;
 	}
 
-	RHICmdList.CopyToResolveTarget((*SeparateTranslucency)->GetRenderTargetItem().TargetableTexture, (*SeparateTranslucency)->GetRenderTargetItem().ShaderResourceTexture, true, FResolveParams());
-	RHICmdList.CopyToResolveTarget((*SeparateTranslucencyDepth)->GetRenderTargetItem().TargetableTexture, (*SeparateTranslucencyDepth)->GetRenderTargetItem().ShaderResourceTexture, true, FResolveParams());
+	const FResolveRect SeparateResolveRect(
+		View.ViewRect.Min.X * SeparateTranslucencyScale, 
+		View.ViewRect.Min.Y * SeparateTranslucencyScale, 
+		View.ViewRect.Max.X * SeparateTranslucencyScale, 
+		View.ViewRect.Max.Y * SeparateTranslucencyScale
+		);
+
+	RHICmdList.CopyToResolveTarget((*SeparateTranslucency)->GetRenderTargetItem().TargetableTexture, (*SeparateTranslucency)->GetRenderTargetItem().ShaderResourceTexture, true, SeparateResolveRect);
+	RHICmdList.CopyToResolveTarget((*SeparateTranslucencyDepth)->GetRenderTargetItem().TargetableTexture, (*SeparateTranslucencyDepth)->GetRenderTargetItem().ShaderResourceTexture, true, SeparateResolveRect);
 
 	bSeparateTranslucencyPass = false;
 }
@@ -1540,24 +1538,6 @@ const FTexture2DRHIRef& FSceneRenderTargets::GetEditorPrimitivesDepth(FRHIComman
 	return (const FTexture2DRHIRef&)EditorPrimitivesDepth->GetRenderTargetItem().TargetableTexture;
 }
 
-TAutoConsoleVariable<int32> FSceneRenderTargets::CVarSetSeperateTranslucencyEnabled(
-	TEXT("r.SeparateTranslucency"),
-	1,
-	TEXT("Allows to disable the separate translucency feature (all translucency is rendered in separate RT and composited\n")
-	TEXT("after DOF, if not specified otherwise in the material).\n")
-	TEXT(" 0: off (translucency is affected by depth of field)\n")
-	TEXT(" 1: on costs GPU performance and memory but keeps translucency unaffected by Depth of Field. (default)"),
-	ECVF_RenderThreadSafe);
-
-bool FSceneRenderTargets::IsSeparateTranslucencyActive(const FViewInfo& View) const
-{	
-	int32 Value = FSceneRenderTargets::CVarSetSeperateTranslucencyEnabled.GetValueOnRenderThread();
-
-	return (Value != 0) && CurrentFeatureLevel >= ERHIFeatureLevel::SM4
-		&& (View.Family->EngineShowFlags.PostProcessing || View.Family->EngineShowFlags.ShaderComplexity)
-		&& View.Family->EngineShowFlags.SeparateTranslucency;
-}
-
 void FSceneRenderTargets::InitEditorPrimitivesColor(FRHICommandList& RHICmdList)
 {
 	FPooledRenderTargetDesc Desc(FPooledRenderTargetDesc::Create2DDesc(BufferSize, 
@@ -2428,7 +2408,7 @@ void FSceneTextureShaderParameters::Set(
 		{
 			const FTexture2DRHIRef* DepthTexture = SceneContext.GetActualDepthTexture();
 
-			if (SceneContext.IsSeparateTranslucencyPass() && SceneContext.IsSeparateTranslucencyDepthValid())
+			if (SceneContext.IsSeparateTranslucencyPass() && SceneContext.IsDownsampledTranslucencyDepthValid())
 			{
 				FIntPoint OutScaledSize;
 				float OutScale;
@@ -2436,7 +2416,7 @@ void FSceneTextureShaderParameters::Set(
 
 				if (OutScale < 1.0f)
 				{
-					DepthTexture = &SceneContext.GetSeparateTranslucencyDepthSurface();
+					DepthTexture = &SceneContext.GetDownsampledTranslucencyDepthSurface();
 				}
 			}
 
@@ -2725,7 +2705,7 @@ void FDeferredPixelShaderParameters::Set(TRHICmdList& RHICmdList, const ShaderRH
 	if (IsDBufferEnabled())
 	{
 		IPooledRenderTarget* DBufferA = SceneContext.DBufferA ? SceneContext.DBufferA : GSystemTextures.BlackAlphaOneDummy;
-		IPooledRenderTarget* DBufferB = SceneContext.DBufferB ? SceneContext.DBufferB : GSystemTextures.MidGrayDummy;
+		IPooledRenderTarget* DBufferB = SceneContext.DBufferB ? SceneContext.DBufferB : GSystemTextures.DefaultNormal8Bit;
 		IPooledRenderTarget* DBufferC = SceneContext.DBufferC ? SceneContext.DBufferC : GSystemTextures.GreenDummy;
 
 		// todo: optimize out when not needed
diff --git a/Engine/Source/Runtime/Renderer/Private/PostProcess/SceneRenderTargets.h b/Engine/Source/Runtime/Renderer/Private/PostProcess/SceneRenderTargets.h
--- a/Engine/Source/Runtime/Renderer/Private/PostProcess/SceneRenderTargets.h
+++ b/Engine/Source/Runtime/Renderer/Private/PostProcess/SceneRenderTargets.h
@@ -173,8 +173,6 @@ public:
 	/** Destruct all snapshots */
 	void DestroyAllSnapshots();
 
-	static TAutoConsoleVariable<int32> CVarSetSeperateTranslucencyEnabled;
-
 protected:
 	/** Constructor */
 	FSceneRenderTargets(): 
@@ -246,21 +244,23 @@ public:
 	/** Binds the appropriate shadow depth cube map for rendering. */
 	void BeginRenderingCubeShadowDepth(FRHICommandList& RHICmdList, int32 ShadowResolution);
 
+	/** Begin rendering translucency in the scene color. */
 	void BeginRenderingTranslucency(FRHICommandList& RHICmdList, const class FViewInfo& View, bool bFirstTimeThisFrame = true);
-
+	/** Begin rendering translucency in a separate (offscreen) buffer. This can be any translucency pass. */
 	void BeginRenderingSeparateTranslucency(FRHICommandList& RHICmdList, const FViewInfo& View, bool bFirstTimeThisFrame);
-	void FinishRenderingSeparateTranslucency(FRHICommandList& RHICmdList);
+	void FinishRenderingSeparateTranslucency(FRHICommandList& RHICmdList, const FViewInfo& View);
+
 	void FreeSeparateTranslucency()
 	{
 		SeparateTranslucencyRT.SafeRelease();
 		check(!SeparateTranslucencyRT);
 	}
 
-	void FreeSeparateTranslucencyDepth()
+	void FreeDownsampledTranslucencyDepth()
 	{
-		if (SeparateTranslucencyDepthRT.GetReference())
+		if (DownsampledTranslucencyDepthRT.GetReference())
 		{
-			SeparateTranslucencyDepthRT.SafeRelease();
+			DownsampledTranslucencyDepthRT.SafeRelease();
 		}
 	}
 
@@ -286,24 +286,25 @@ public:
 		DefaultDepthClear = DepthClear;
 	}
 
-	void GetSeparateTranslucencyDimensions(FIntPoint& OutScaledSize, float& OutScale)
+	FORCEINLINE void GetSeparateTranslucencyDimensions(FIntPoint& OutScaledSize, float& OutScale) const
 	{
 		OutScaledSize = SeparateTranslucencyBufferSize;
 		OutScale = SeparateTranslucencyScale;
 	}
 
+	/** Separate translucency buffer can be downsampled or not (as it is used to store the AfterDOF translucency) */
 	TRefCountPtr<IPooledRenderTarget>& GetSeparateTranslucency(FRHICommandList& RHICmdList, FIntPoint Size);
 
-	bool IsSeparateTranslucencyDepthValid()
+	bool IsDownsampledTranslucencyDepthValid()
 	{
-		return SeparateTranslucencyDepthRT != nullptr;
+		return DownsampledTranslucencyDepthRT != nullptr;
 	}
 
-	TRefCountPtr<IPooledRenderTarget>& GetSeparateTranslucencyDepth(FRHICommandList& RHICmdList, FIntPoint Size);
+	TRefCountPtr<IPooledRenderTarget>& GetDownsampledTranslucencyDepth(FRHICommandList& RHICmdList, FIntPoint Size);
 
-	const FTexture2DRHIRef& GetSeparateTranslucencyDepthSurface()
+	const FTexture2DRHIRef& GetDownsampledTranslucencyDepthSurface()
 	{
-		return (const FTexture2DRHIRef&)SeparateTranslucencyDepthRT->GetRenderTargetItem().TargetableTexture;
+		return (const FTexture2DRHIRef&)DownsampledTranslucencyDepthRT->GetRenderTargetItem().TargetableTexture;
 	}
 
 	/**
@@ -514,13 +515,7 @@ public:
 
 	TRefCountPtr<IPooledRenderTarget>& GetReflectionBrightnessTarget();
 
-
-	bool IsSeparateTranslucencyActive(const FViewInfo& View) const;
-
-	bool IsSeparateTranslucencyPass()
-	{
-		return bSeparateTranslucencyPass;
-	}
+	FORCEINLINE bool IsSeparateTranslucencyPass() const { return bSeparateTranslucencyPass; }
 	
 	// Can be called when the Scene Color content is no longer needed. As we create SceneColor on demand we can make sure it is created with the right format.
 	// (as a call to SetSceneColor() can override it with a different format)
@@ -608,7 +603,8 @@ public:
 
 	/** ONLY for snapshots!!! this is a copy of the SeparateTranslucencyRT from the view state. */
 	TRefCountPtr<IPooledRenderTarget> SeparateTranslucencyRT;
-	TRefCountPtr<IPooledRenderTarget> SeparateTranslucencyDepthRT;
+	/** Downsampled depth used when rendering translucency in smaller resolution. */
+	TRefCountPtr<IPooledRenderTarget> DownsampledTranslucencyDepthRT;
 
 	// todo: free ScreenSpaceAO so pool can reuse
 	bool bScreenSpaceAOIsValid;
diff --git a/Engine/Source/Runtime/Renderer/Private/PostProcess/ScreenSpaceReflections.cpp b/Engine/Source/Runtime/Renderer/Private/PostProcess/ScreenSpaceReflections.cpp
--- a/Engine/Source/Runtime/Renderer/Private/PostProcess/ScreenSpaceReflections.cpp
+++ b/Engine/Source/Runtime/Renderer/Private/PostProcess/ScreenSpaceReflections.cpp
@@ -376,7 +376,7 @@ void FRCPassPostProcessScreenSpaceReflections::Process(FRenderingCompositePassCo
 		Context.SetViewportAndCallRHI(View.ViewRect);
 
 		// Clear stencil to 0
-		DrawClearQuad(RHICmdList, Context.GetFeatureLevel(), false, FLinearColor(), false, 0, true, 0, PassOutputs[0].RenderTargetDesc.Extent, View.ViewRect);
+		DrawClearQuad(RHICmdList, false, FLinearColor(), false, 0, true, 0, PassOutputs[0].RenderTargetDesc.Extent, View.ViewRect);
 	
 		FGraphicsPipelineStateInitializer GraphicsPSOInit;
 		RHICmdList.ApplyCachedRenderTargets(GraphicsPSOInit);
@@ -437,7 +437,7 @@ void FRCPassPostProcessScreenSpaceReflections::Process(FRenderingCompositePassCo
 		RHICmdList.ApplyCachedRenderTargets(GraphicsPSOInit);
 
 		// clear DestRenderTarget only outside of the view's rectangle
-		DrawClearQuad(RHICmdList, SceneContext.GetCurrentFeatureLevel(), true, FLinearColor::Black, false, 0, false, 0, PassOutputs[0].RenderTargetDesc.Extent, View.ViewRect);
+		DrawClearQuad(RHICmdList, true, FLinearColor::Black, false, 0, false, 0, PassOutputs[0].RenderTargetDesc.Extent, View.ViewRect);
 
 		// set the state
 		GraphicsPSOInit.BlendState = TStaticBlendState<>::GetRHI();
diff --git a/Engine/Source/Runtime/Renderer/Private/PrimitiveSceneInfo.cpp b/Engine/Source/Runtime/Renderer/Private/PrimitiveSceneInfo.cpp
--- a/Engine/Source/Runtime/Renderer/Private/PrimitiveSceneInfo.cpp
+++ b/Engine/Source/Runtime/Renderer/Private/PrimitiveSceneInfo.cpp
@@ -64,7 +64,14 @@ private:
 	TRefCountPtr<HHitProxy> CurrentHitProxy;
 };
 
-void FPrimitiveSceneInfoCompact::Init(FPrimitiveSceneInfo* InPrimitiveSceneInfo)
+FPrimitiveFlagsCompact::FPrimitiveFlagsCompact(const FPrimitiveSceneProxy* Proxy)
+	: bCastDynamicShadow(Proxy->CastsDynamicShadow())
+	, bStaticLighting(Proxy->HasStaticLighting())
+	, bCastStaticShadow(Proxy->CastsStaticShadow())
+{}
+
+FPrimitiveSceneInfoCompact::FPrimitiveSceneInfoCompact(FPrimitiveSceneInfo* InPrimitiveSceneInfo) :
+	PrimitiveFlagsCompact(InPrimitiveSceneInfo->Proxy)
 {
 	PrimitiveSceneInfo = InPrimitiveSceneInfo;
 	Proxy = PrimitiveSceneInfo->Proxy;
@@ -73,14 +80,6 @@ void FPrimitiveSceneInfoCompact::Init(FPrimitiveSceneInfo* InPrimitiveSceneInfo)
 	MaxDrawDistance = PrimitiveSceneInfo->Proxy->GetMaxDrawDistance();
 
 	VisibilityId = PrimitiveSceneInfo->Proxy->GetVisibilityId();
-
-	bHasViewDependentDPG = Proxy->HasViewDependentDPG();
-	bCastDynamicShadow = PrimitiveSceneInfo->Proxy->CastsDynamicShadow();
-
-	bAffectDynamicIndirectLighting = PrimitiveSceneInfo->Proxy->AffectsDynamicIndirectLighting();
-	LpvBiasMultiplier = PrimitiveSceneInfo->Proxy->GetLpvBiasMultiplier();
-	
-	StaticDepthPriorityGroup = bHasViewDependentDPG ? 0 : Proxy->GetStaticDepthPriorityGroup();
 }
 
 FPrimitiveSceneInfo::FPrimitiveSceneInfo(UPrimitiveComponent* InComponent,FScene* InScene):
@@ -158,6 +157,12 @@ void FPrimitiveSceneInfo::AddStaticMeshes(FRHICommandListImmediate& RHICmdList)
 		Scene->StaticMeshes[SceneArrayAllocation.Index] = &Mesh;
 		Mesh.Id = SceneArrayAllocation.Index;
 
+		if (Mesh.bRequiresPerElementVisibility)
+		{
+			// Use a separate index into StaticMeshBatchVisibility, since most meshes don't use it
+			Mesh.BatchVisibilityId = Scene->StaticMeshBatchVisibility.AddUninitialized().Index;
+		}
+
 		// By this point, the index buffer render resource must be initialized
 		// Add the static mesh to the appropriate draw lists.
 		Mesh.AddToDrawLists(RHICmdList, Scene);
@@ -196,16 +201,11 @@ void FPrimitiveSceneInfo::AddToScene(FRHICommandListImmediate& RHICmdList, bool
 	}
 
 	// create potential storage for our compact info
-	FPrimitiveSceneInfoCompact LocalCompactPrimitiveSceneInfo;
-	FPrimitiveSceneInfoCompact* CompactPrimitiveSceneInfo = &LocalCompactPrimitiveSceneInfo;
-
-	// if we are being added directly to the Octree, initialize the temp compact scene info,
-	// and let the Octree make a copy of it
-	LocalCompactPrimitiveSceneInfo.Init(this);
+	FPrimitiveSceneInfoCompact CompactPrimitiveSceneInfo(this);
 
 	// Add the primitive to the octree.
 	check(!OctreeId.IsValidId());
-	Scene->PrimitiveOctree.AddElement(LocalCompactPrimitiveSceneInfo);
+	Scene->PrimitiveOctree.AddElement(CompactPrimitiveSceneInfo);
 	check(OctreeId.IsValidId());
 
 	if (Proxy->CastsDynamicIndirectShadow())
@@ -213,15 +213,17 @@ void FPrimitiveSceneInfo::AddToScene(FRHICommandListImmediate& RHICmdList, bool
 		Scene->DynamicIndirectCasterPrimitives.Add(this);
 	}
 
+	Scene->PrimitiveSceneProxies[PackedIndex] = Proxy;
+
 	// Set bounds.
 	FPrimitiveBounds& PrimitiveBounds = Scene->PrimitiveBounds[PackedIndex];
 	FBoxSphereBounds BoxSphereBounds = Proxy->GetBounds();
-	PrimitiveBounds.Origin = BoxSphereBounds.Origin;
-	PrimitiveBounds.SphereRadius = BoxSphereBounds.SphereRadius;
-	PrimitiveBounds.BoxExtent = BoxSphereBounds.BoxExtent;
+	PrimitiveBounds.BoxSphereBounds = BoxSphereBounds;
 	PrimitiveBounds.MinDrawDistanceSq = FMath::Square(Proxy->GetMinDrawDistance());
 	PrimitiveBounds.MaxDrawDistance = Proxy->GetMaxDrawDistance();
 
+	Scene->PrimitiveFlagsCompact[PackedIndex] = FPrimitiveFlagsCompact(Proxy);
+
 	// Store precomputed visibility ID.
 	int32 VisibilityBitIndex = Proxy->GetVisibilityId();
 	FPrimitiveVisibilityId& VisibilityId = Scene->PrimitiveVisibilityIds[PackedIndex];
@@ -276,7 +278,7 @@ void FPrimitiveSceneInfo::AddToScene(FRHICommandListImmediate& RHICmdList, bool
 			LightIt.Advance())
 		{
 			const FLightSceneInfoCompact& LightSceneInfoCompact = LightIt.GetCurrentElement();
-			if (LightSceneInfoCompact.AffectsPrimitive(*CompactPrimitiveSceneInfo))
+			if (LightSceneInfoCompact.AffectsPrimitive(CompactPrimitiveSceneInfo.Bounds, CompactPrimitiveSceneInfo.Proxy))
 			{
 				FLightPrimitiveInteraction::Create(LightSceneInfoCompact.LightSceneInfo,this);
 			}
diff --git a/Engine/Source/Runtime/Renderer/Private/ReflectionEnvironment.cpp b/Engine/Source/Runtime/Renderer/Private/ReflectionEnvironment.cpp
--- a/Engine/Source/Runtime/Renderer/Private/ReflectionEnvironment.cpp
+++ b/Engine/Source/Runtime/Renderer/Private/ReflectionEnvironment.cpp
@@ -436,8 +436,9 @@ public:
 	{
 	}
 
+	template<typename TRHICommandList>
 	void SetParameters(
-		FRHIAsyncComputeCommandListImmediate& RHICmdList, 
+		TRHICommandList& RHICmdList,
 		const FViewInfo& View,
 		FTextureRHIParamRef SSRTexture,
 		FUnorderedAccessViewRHIParamRef OutSceneColorUAV, 
@@ -487,7 +488,8 @@ public:
 		ForwardLightingParameters.Set(RHICmdList, ShaderRHI, View);
 	}
 
-	void UnsetParameters(FRHIAsyncComputeCommandListImmediate& RHICmdList, FUnorderedAccessViewRHIParamRef OutSceneColorUAV)
+	template<typename TRHICommandList>
+	void UnsetParameters(TRHICommandList& RHICmdList, FUnorderedAccessViewRHIParamRef OutSceneColorUAV)
 	{
 		const FComputeShaderRHIParamRef ShaderRHI = GetComputeShader();		
 		OutSceneColor.UnsetUAV(RHICmdList, ShaderRHI);
@@ -1153,24 +1155,26 @@ void FDeferredShadingSceneRenderer::RenderTiledDeferredImageBasedReflections(FRH
 
 			static const FName TiledReflBeginComputeName(TEXT("ReflectionEnvBeginComputeFence"));
 			static const FName TiledReflEndComputeName(TEXT("ReflectionEnvEndComputeFence"));
-			FComputeFenceRHIRef ReflectionBeginFence = RHICmdList.CreateComputeFence(TiledReflBeginComputeName);
-			FComputeFenceRHIRef ReflectionEndFence = RHICmdList.CreateComputeFence(TiledReflEndComputeName);
+			//FComputeFenceRHIRef ReflectionBeginFence = RHICmdList.CreateComputeFence(TiledReflBeginComputeName);
+			//FComputeFenceRHIRef ReflectionEndFence = RHICmdList.CreateComputeFence(TiledReflEndComputeName);
 
 			//Grab the async compute commandlist.
-			FRHIAsyncComputeCommandListImmediate& RHICmdListComputeImmediate = FRHICommandListExecutor::GetImmediateAsyncComputeCommandList();
+			//FRHIAsyncComputeCommandListImmediate& RHICmdListComputeImmediate = FRHICommandListExecutor::GetImmediateAsyncComputeCommandList();
+			auto& RHICmdListComputeImmediate = RHICmdList;
 			{
-				SCOPED_COMPUTE_EVENTF(RHICmdListComputeImmediate, ReflectionEnvironment, TEXT("ReflectionEnvironment ComputeShader %dx%d Tile:%dx%d Box:%d Sphere:%d SkyLight:%d"),
+				/*SCOPED_COMPUTE_EVENTF(RHICmdListComputeImmediate, ReflectionEnvironment, TEXT("ReflectionEnvironment ComputeShader %dx%d Tile:%dx%d Box:%d Sphere:%d SkyLight:%d"),
 					View.ViewRect.Width(), View.ViewRect.Height(), GReflectionEnvironmentTileSizeX, GReflectionEnvironmentTileSizeY,
 					View.NumBoxReflectionCaptures, View.NumSphereReflectionCaptures, bHasSkyLight);
+					*/
 
 				ComputeShader = SelectReflectionEnvironmentTiledDeferredCS(View.ShaderMap, bUseLightmaps, bHasSkyLight, bHasBoxCaptures, bHasSphereCaptures, DynamicBentNormalAO != NULL);
 
 				//Really we should write this fence where we transition the final depedency for the reflections.  We may add an RHI command just for writing fences if this
 				//can't be done in the general case.  In the meantime, hack this command a bit to write the fence.
-				RHICmdList.TransitionResources(EResourceTransitionAccess::EReadable, EResourceTransitionPipeline::EGfxToCompute, nullptr, 0, ReflectionBeginFence);
+				//RHICmdList.TransitionResources(EResourceTransitionAccess::EReadable, EResourceTransitionPipeline::EGfxToCompute, nullptr, 0, ReflectionBeginFence);
 					
 				//we must wait on the fence written from the Gfx pipe to let us know all our dependencies are ready.
-				RHICmdListComputeImmediate.WaitComputeFence(ReflectionBeginFence);
+				//RHICmdListComputeImmediate.WaitComputeFence(ReflectionBeginFence);
 
 				//standard compute setup, but on the async commandlist.
 				RHICmdListComputeImmediate.SetComputeShader(ComputeShader->GetComputeShader());
@@ -1185,15 +1189,15 @@ void FDeferredShadingSceneRenderer::RenderTiledDeferredImageBasedReflections(FRH
 				ComputeShader->UnsetParameters(RHICmdListComputeImmediate, OutUAV);
 			
 				//transition the output to readable and write the fence to allow the Gfx pipe to carry on.
-				RHICmdListComputeImmediate.TransitionResources(EResourceTransitionAccess::EReadable, EResourceTransitionPipeline::EComputeToGfx, &OutUAV, 1, ReflectionEndFence);
+				RHICmdListComputeImmediate.TransitionResources(EResourceTransitionAccess::EReadable, EResourceTransitionPipeline::EComputeToGfx, &OutUAV, 1);// , ReflectionEndFence);
 			}
 
 			//immediately dispatch our async compute commands to the RHI thread to be submitted to the GPU as soon as possible.
 			//dispatch after the scope so the drawevent pop is inside the dispatch
-			FRHIAsyncComputeCommandListImmediate::ImmediateDispatch(RHICmdListComputeImmediate);			
+			//FRHIAsyncComputeCommandListImmediate::ImmediateDispatch(RHICmdListComputeImmediate);			
 			
 			//Gfx pipe must wait for the async compute reflection job to complete.
-			RHICmdList.WaitComputeFence(ReflectionEndFence);
+			//RHICmdList.WaitComputeFence(ReflectionEndFence);
 		}
 	}
 
diff --git a/Engine/Source/Runtime/Renderer/Private/ReflectionEnvironmentCapture.cpp b/Engine/Source/Runtime/Renderer/Private/ReflectionEnvironmentCapture.cpp
--- a/Engine/Source/Runtime/Renderer/Private/ReflectionEnvironmentCapture.cpp
+++ b/Engine/Source/Runtime/Renderer/Private/ReflectionEnvironmentCapture.cpp
@@ -187,6 +187,10 @@ void CreateCubeMips( FRHICommandListImmediate& RHICmdList, ERHIFeatureLevel::Typ
 			SetRenderTarget(RHICmdList, Cubemap.TargetableTexture, MipIndex, CubeFace, NULL, true);
 			RHICmdList.ApplyCachedRenderTargets(GraphicsPSOInit);
 
+			//this is a weired place but at the end of the scope it will not be good enough as SetRendertarget sets it to writable 
+			//Use ERWSubResBarrier since we don't transition individual subresources.  Basically treat the whole texture as R/W as we walk down the mip chain.
+			RHICmdList.TransitionResources(EResourceTransitionAccess::ERWSubResBarrier, &CubeRef, 1);
+
 			const FIntRect ViewRect(0, 0, MipSize, MipSize);
 			RHICmdList.SetViewport(0, 0, 0.0f, MipSize, MipSize, 1.0f);
 
@@ -222,11 +226,9 @@ void CreateCubeMips( FRHICommandListImmediate& RHICmdList, ERHIFeatureLevel::Typ
 				FIntPoint(ViewRect.Width(), ViewRect.Height()),
 				FIntPoint(MipSize, MipSize),
 				*VertexShader);
-
-			//Use ERWSubResBarrier since we don't transition individual subresources.  Basically treat the whole texture as R/W as we walk down the mip chain.
-			RHICmdList.TransitionResources(EResourceTransitionAccess::ERWSubResBarrier, &CubeRef, 1);
 		}
 	}
+	RHICmdList.TransitionResources(EResourceTransitionAccess::EReadable, &CubeRef, 1);
 }
 
 /** Computes the average brightness of the given reflection capture and stores it in the scene. */
diff --git a/Engine/Source/Runtime/Renderer/Private/Renderer.cpp b/Engine/Source/Runtime/Renderer/Private/Renderer.cpp
--- a/Engine/Source/Runtime/Renderer/Private/Renderer.cpp
+++ b/Engine/Source/Runtime/Renderer/Private/Renderer.cpp
@@ -101,6 +101,8 @@ void FRendererModule::DrawTileMesh(FRHICommandListImmediate& RHICmdList, FDrawin
 			GMinimalDummyForwardLightingResources = new TGlobalResource<FMinimalDummyForwardLightingResources>();
 		}
 
+		FMaterialRenderProxy::UpdateDeferredCachedUniformExpressions();
+
 		//Apply the minimal forward lighting resources
 		View.ForwardLightingResources = &GMinimalDummyForwardLightingResources->ForwardLightingResources;
 		
