Copying //UE4/Dev-Rendering to //UE4/Dev-Main (Source: //UE4/Dev-Rendering @ 3461187)

diff --git a/Engine/Source/Runtime/Renderer/Private/RendererScene.cpp b/Engine/Source/Runtime/Renderer/Private/RendererScene.cpp
--- a/Engine/Source/Runtime/Renderer/Private/RendererScene.cpp
+++ b/Engine/Source/Runtime/Renderer/Private/RendererScene.cpp
@@ -474,7 +474,9 @@ SIZE_T FScene::GetSizeBytes() const
 
 void FScene::CheckPrimitiveArrays()
 {
+	check(Primitives.Num() == PrimitiveSceneProxies.Num());
 	check(Primitives.Num() == PrimitiveBounds.Num());
+	check(Primitives.Num() == PrimitiveFlagsCompact.Num());
 	check(Primitives.Num() == PrimitiveVisibilityIds.Num());
 	check(Primitives.Num() == PrimitiveOcclusionFlags.Num());
 	check(Primitives.Num() == PrimitiveComponentIds.Num());
@@ -490,7 +492,9 @@ void FScene::AddPrimitiveSceneInfo_RenderThread(FRHICommandListImmediate& RHICmd
 	int32 PrimitiveIndex = Primitives.Add(PrimitiveSceneInfo);
 	PrimitiveSceneInfo->PackedIndex = PrimitiveIndex;
 
+	PrimitiveSceneProxies.AddUninitialized();
 	PrimitiveBounds.AddUninitialized();
+	PrimitiveFlagsCompact.AddUninitialized();
 	PrimitiveVisibilityIds.AddUninitialized();
 	PrimitiveOcclusionFlags.AddUninitialized();
 	PrimitiveComponentIds.AddUninitialized();
@@ -546,7 +550,7 @@ FReadOnlyCVARCache::FReadOnlyCVARCache()
 	static const auto CVarSupportPointLightWholeSceneShadows = IConsoleManager::Get().FindTConsoleVariableDataInt(TEXT("r.SupportPointLightWholeSceneShadows"));
 	static const auto CVarSupportAllShaderPermutations = IConsoleManager::Get().FindTConsoleVariableDataInt(TEXT("r.SupportAllShaderPermutations"));	
 	static const auto CVarVertexFoggingForOpaque = IConsoleManager::Get().FindTConsoleVariableDataInt(TEXT("r.VertexFoggingForOpaque"));	
-	static const auto CVarForwardShading = IConsoleManager::Get().FindTConsoleVariableDataInt(TEXT("r.ForwardShading"));
+	static const auto CVarForwardShading = IConsoleManager::Get().FindConsoleVariable(TEXT("r.ForwardShading"));
 	static const auto CVarAllowStaticLighting = IConsoleManager::Get().FindTConsoleVariableDataInt(TEXT("r.AllowStaticLighting"));
 
 	static const auto CVarMobileAllowMovableDirectionalLights = IConsoleManager::Get().FindTConsoleVariableDataInt(TEXT("r.Mobile.AllowMovableDirectionalLights"));
@@ -571,7 +575,7 @@ FReadOnlyCVARCache::FReadOnlyCVARCache()
 	NumMobileMovablePointLights = CVarMobileNumDynamicPointLights->GetValueOnAnyThread();
 
 	// Only enable VertexFoggingForOpaque if ForwardShading is enabled 
-	const bool bForwardShading = CVarForwardShading && CVarForwardShading->GetValueOnAnyThread() != 0;
+	const bool bForwardShading = CVarForwardShading && CVarForwardShading->GetInt() != 0;
 	bEnableVertexFoggingForOpaque = bForwardShading && ( !CVarVertexFoggingForOpaque || CVarVertexFoggingForOpaque->GetValueOnAnyThread() != 0 );
 
 	const bool bShowMissmatchedLowQualityLightmapsWarning = (!bEnableLowQualityLightmaps) && (GEngine->bShouldGenerateLowQualityLightmaps_DEPRECATED);
@@ -956,7 +960,9 @@ void FScene::RemovePrimitiveSceneInfo_RenderThread(FPrimitiveSceneInfo* Primitiv
 
 	int32 PrimitiveIndex = PrimitiveSceneInfo->PackedIndex;
 	Primitives.RemoveAtSwap(PrimitiveIndex);
+	PrimitiveSceneProxies.RemoveAtSwap(PrimitiveIndex);
 	PrimitiveBounds.RemoveAtSwap(PrimitiveIndex);
+	PrimitiveFlagsCompact.RemoveAtSwap(PrimitiveIndex);
 	PrimitiveVisibilityIds.RemoveAtSwap(PrimitiveIndex);
 	PrimitiveOcclusionFlags.RemoveAtSwap(PrimitiveIndex);
 	PrimitiveComponentIds.RemoveAtSwap(PrimitiveIndex);
@@ -1503,7 +1509,7 @@ const FReflectionCaptureProxy* FScene::FindClosestReflectionCapture(FVector Posi
 	return ClosestCaptureIndex != INDEX_NONE ? ReflectionSceneData.RegisteredReflectionCaptures[ClosestCaptureIndex] : NULL;
 }
 
-const FPlanarReflectionSceneProxy* FScene::FindClosestPlanarReflection(const FPrimitiveBounds& Bounds) const
+const FPlanarReflectionSceneProxy* FScene::FindClosestPlanarReflection(const FBoxSphereBounds& Bounds) const
 {
 	checkSlow(IsInParallelRenderingThread());
 	const FPlanarReflectionSceneProxy* ClosestPlanarReflection = NULL;
@@ -2711,7 +2717,7 @@ void FScene::ApplyWorldOffset_RenderThread(FVector InOffset)
 	// Primitive bounds
 	for (auto It = PrimitiveBounds.CreateIterator(); It; ++It)
 	{
-		(*It).Origin+= InOffset;
+		(*It).BoxSphereBounds.Origin+= InOffset;
 	}
 
 	// Primitive occlusion bounds
@@ -2785,7 +2791,7 @@ void FScene::OnLevelAddedToWorld(FName LevelAddedName, UWorld* InWorld, bool bIs
 {
 	if (bIsLightingScenario)
 	{
-		InWorld->PropagateLightingScenarioChange();
+		InWorld->PropagateLightingScenarioChange(true);
 	}
 
 	ENQUEUE_UNIQUE_RENDER_COMMAND_TWOPARAMETER(
@@ -2818,7 +2824,7 @@ void FScene::OnLevelRemovedFromWorld(UWorld* InWorld, bool bIsLightingScenario)
 {
 	if (bIsLightingScenario)
 	{
-		InWorld->PropagateLightingScenarioChange();
+		InWorld->PropagateLightingScenarioChange(false);
 	}
 }
 
@@ -3236,7 +3242,7 @@ bool FLatentGPUTimer::Tick(FRHICommandListImmediate& RHICmdList)
 
 	if (StartQueries[QueryIndex] && EndQueries[QueryIndex])
 	{
-		if (GRHIThread)
+		if (IsRunningRHIInSeparateThread())
 		{
 			// Block until the RHI thread has processed the previous query commands, if necessary
 			// Stat disabled since we buffer 2 frames minimum, it won't actually block
@@ -3308,7 +3314,7 @@ void FLatentGPUTimer::End(FRHICommandListImmediate& RHICmdList)
 	// for these query results on some platforms.
 	RHICmdList.SubmitCommandsHint();
 
-	if (GRHIThread)
+	if (IsRunningRHIInSeparateThread())
 	{
 		int32 NumFrames = NumBufferedFrames;
 		for (int32 Dest = 1; Dest < NumFrames; Dest++)
diff --git a/Engine/Source/Runtime/Renderer/Private/SceneCaptureRendering.cpp b/Engine/Source/Runtime/Renderer/Private/SceneCaptureRendering.cpp
--- a/Engine/Source/Runtime/Renderer/Private/SceneCaptureRendering.cpp
+++ b/Engine/Source/Runtime/Renderer/Private/SceneCaptureRendering.cpp
@@ -254,7 +254,7 @@ static void UpdateSceneCaptureContentDeferred_RenderThread(
 		FIntRect ViewRect = View.ViewRect;
 		FIntRect UnconstrainedViewRect = View.UnconstrainedViewRect;
 		SetRenderTarget(RHICmdList, Target->GetRenderTargetTexture(), nullptr, true);
-		DrawClearQuad(RHICmdList, SceneRenderer->FeatureLevel, true, FLinearColor::Black, false, 0, false, 0, Target->GetSizeXY(), ViewRect);
+		DrawClearQuad(RHICmdList, true, FLinearColor::Black, false, 0, false, 0, Target->GetSizeXY(), ViewRect);
 
 		// Render the scene normally
 		{
@@ -278,8 +278,10 @@ static void UpdateSceneCaptureContent_RenderThread(
 	const FName OwnerName,
 	const FResolveParams& ResolveParams)
 {
+	FMaterialRenderProxy::UpdateDeferredCachedUniformExpressions();
+
 	switch (SceneRenderer->Scene->GetShadingPath())
-			{
+	{
 		case EShadingPath::Mobile:
 		{
 			UpdateSceneCaptureContentMobile_RenderThread(
@@ -391,7 +393,7 @@ FSceneRenderer* CreateSceneRendererForSceneCapture(
 		Scene,
 		SceneCaptureComponent->ShowFlags)
 		.SetResolveScene(!bCaptureSceneColor)
-		.SetRealtimeUpdate(bIsPlanarReflection));
+		.SetRealtimeUpdate(bIsPlanarReflection || SceneCaptureComponent->bCaptureEveryFrame));
 
 	for (int32 ViewIndex = 0; ViewIndex < Views.Num(); ++ViewIndex)
 	{
diff --git a/Engine/Source/Runtime/Renderer/Private/SceneCore.cpp b/Engine/Source/Runtime/Renderer/Private/SceneCore.cpp
--- a/Engine/Source/Runtime/Renderer/Private/SceneCore.cpp
+++ b/Engine/Source/Runtime/Renderer/Private/SceneCore.cpp
@@ -381,7 +381,7 @@ void FStaticMesh::RemoveFromDrawLists()
 	// Remove the mesh from all draw lists.
 	while(DrawListLinks.Num())
 	{
-		FStaticMesh::FDrawListElementLink* Link = DrawListLinks[0];
+		TRefCountPtr<FStaticMesh::FDrawListElementLink> Link = DrawListLinks[0];
 		const int32 OriginalNumLinks = DrawListLinks.Num();
 		// This will call UnlinkDrawList.
 		Link->Remove(true);
diff --git a/Engine/Source/Runtime/Renderer/Private/SceneCore.h b/Engine/Source/Runtime/Renderer/Private/SceneCore.h
--- a/Engine/Source/Runtime/Renderer/Private/SceneCore.h
+++ b/Engine/Source/Runtime/Renderer/Private/SceneCore.h
@@ -147,6 +147,9 @@ public:
 	/** The index of the mesh in the scene's static meshes array. */
 	int32 Id;
 
+	/** Index of the mesh into the scene's StaticMeshBatchVisibility array. */
+	int32 BatchVisibilityId;
+
 	// Constructor/destructor.
 	FStaticMesh(
 		FPrimitiveSceneInfo* InPrimitiveSceneInfo,
@@ -157,7 +160,8 @@ public:
 		FMeshBatch(InMesh),
 		ScreenSize(InScreenSize),
 		PrimitiveSceneInfo(InPrimitiveSceneInfo),
-		Id(INDEX_NONE)
+		Id(INDEX_NONE),
+		BatchVisibilityId(INDEX_NONE)
 	{
 		BatchHitProxyId = InHitProxyId;
 	}
@@ -188,7 +192,8 @@ private:
 		FMeshBatch(InStaticMesh),
 		ScreenSize(InStaticMesh.ScreenSize),
 		PrimitiveSceneInfo(InStaticMesh.PrimitiveSceneInfo),
-		Id(InStaticMesh.Id)
+		Id(InStaticMesh.Id),
+		BatchVisibilityId(InStaticMesh.BatchVisibilityId)
 	{}
 };
 
diff --git a/Engine/Source/Runtime/Renderer/Private/SceneHitProxyRendering.cpp b/Engine/Source/Runtime/Renderer/Private/SceneHitProxyRendering.cpp
--- a/Engine/Source/Runtime/Renderer/Private/SceneHitProxyRendering.cpp
+++ b/Engine/Source/Runtime/Renderer/Private/SceneHitProxyRendering.cpp
@@ -444,7 +444,7 @@ void InitHitProxyRender(FRHICommandListImmediate& RHICmdList, const FSceneRender
 	{
 		const FViewInfo& View = Views[ViewIndex];
 		RHICmdList.SetViewport(View.ViewRect.Min.X, View.ViewRect.Min.Y, 0.0f, View.ViewRect.Max.X, View.ViewRect.Max.Y, 1.0f);
-		DrawClearQuad(RHICmdList, SceneRenderer->FeatureLevel, true, FLinearColor::White, false, 0, false, 0, OutHitProxyRT->GetDesc().Extent, FIntRect());
+		DrawClearQuad(RHICmdList, true, FLinearColor::White, false, 0, false, 0, OutHitProxyRT->GetDesc().Extent, FIntRect());
 	}
 }
 
@@ -472,7 +472,7 @@ static void DoRenderHitProxies(FRHICommandListImmediate& RHICmdList, const FScen
 		RHICmdList.SetViewport(View.ViewRect.Min.X, View.ViewRect.Min.Y, 0.0f, View.ViewRect.Max.X, View.ViewRect.Max.Y, 1.0f);
 
 		// Clear the depth buffer for each DPG.
-		DrawClearQuad(RHICmdList, SceneRenderer->FeatureLevel, false, FLinearColor(), true, (float)ERHIZBuffer::FarPlane, true, 0, HitProxyDepthRT->GetDesc().Extent, FIntRect());
+		DrawClearQuad(RHICmdList, false, FLinearColor(), true, (float)ERHIZBuffer::FarPlane, true, 0, HitProxyDepthRT->GetDesc().Extent, FIntRect());
 
 		// Depth tests + writes, no alpha blending.
 		DrawRenderState.SetDepthStencilState(TStaticDepthStencilState<true, CF_DepthNearOrEqual>::GetRHI());
diff --git a/Engine/Source/Runtime/Renderer/Private/ScenePrivate.h b/Engine/Source/Runtime/Renderer/Private/ScenePrivate.h
--- a/Engine/Source/Runtime/Renderer/Private/ScenePrivate.h
+++ b/Engine/Source/Runtime/Renderer/Private/ScenePrivate.h
@@ -1585,7 +1585,7 @@ private:
 	/** Internal helper to determine if indirect lighting is enabled at all */
 	bool IndirectLightingAllowed(FScene* Scene, FSceneRenderer& Renderer) const;
 
-	void ProcessPrimitiveUpdate(FScene* Scene, FViewInfo& View, int32 PrimitiveIndex, bool bAllowUnbuiltPreview, TMap<FIntVector, FBlockUpdateInfo>& OutBlocksToUpdate, TArray<FIndirectLightingCacheAllocation*>& OutTransitionsOverTimeToUpdate);
+	void ProcessPrimitiveUpdate(FScene* Scene, FViewInfo& View, int32 PrimitiveIndex, bool bAllowUnbuiltPreview, bool bAllowVolumeSample, TMap<FIntVector, FBlockUpdateInfo>& OutBlocksToUpdate, TArray<FIndirectLightingCacheAllocation*>& OutTransitionsOverTimeToUpdate);
 
 	/** Internal helper to perform the work of updating the cache primitives.  Can be done on any thread as a task */
 	void UpdateCachePrimitivesInternal(FScene* Scene, FSceneRenderer& Renderer, bool bAllowUnbuiltPreview, TMap<FIntVector, FBlockUpdateInfo>& OutBlocksToUpdate, TArray<FIndirectLightingCacheAllocation*>& OutTransitionsOverTimeToUpdate);
@@ -1611,7 +1611,7 @@ private:
 		const TMap<FPrimitiveComponentId, FAttachmentGroupSceneInfo>& AttachmentGroups,
 		FPrimitiveSceneInfo* PrimitiveSceneInfo,
 		bool bAllowUnbuiltPreview, 
-		bool bOpaqueRelevance, 
+		bool bAllowVolumeSample, 
 		TMap<FIntVector, FBlockUpdateInfo>& BlocksToUpdate, 
 		TArray<FIndirectLightingCacheAllocation*>& TransitionsOverTimeToUpdate);	
 
@@ -1703,12 +1703,7 @@ private:
  */
 struct FPrimitiveBounds
 {
-	/** Origin of the primitive. */
-	FVector Origin;
-	/** Radius of the bounding sphere. */
-	float SphereRadius;
-	/** Extents of the axis-aligned bounding box. */
-	FVector BoxExtent;
+	FBoxSphereBounds BoxSphereBounds;
 	/** Square of the minimum draw distance for the primitive. */
 	float MinDrawDistanceSq;
 	/** Maximum draw distance for the primitive. */
@@ -1926,8 +1921,12 @@ public:
 
 	/** Packed array of primitives in the scene. */
 	TArray<FPrimitiveSceneInfo*> Primitives;
+	/** Packed array of primitive scene proxies in the scene. */
+	TArray<FPrimitiveSceneProxy*> PrimitiveSceneProxies;
 	/** Packed array of primitive bounds. */
 	TArray<FPrimitiveBounds> PrimitiveBounds;
+	/** Packed array of primitive flags. */
+	TArray<FPrimitiveFlagsCompact> PrimitiveFlagsCompact;
 	/** Packed array of precomputed primitive visibility IDs. */
 	TArray<FPrimitiveVisibilityId> PrimitiveVisibilityIds;
 	/** Packed array of primitive occlusion flags. See EOcclusionFlags. */
@@ -2015,6 +2014,9 @@ public:
 	/** The static meshes in the scene. */
 	TSparseArray<FStaticMesh*> StaticMeshes;
 
+	/** This sparse array is used just to track free indices for FStaticMesh::BatchVisibilityId. */
+	TSparseArray<bool> StaticMeshBatchVisibility;
+
 	/** The exponential fog components in the scene. */
 	TArray<FExponentialHeightFogSceneInfo> ExponentialFogs;
 
@@ -2168,7 +2170,7 @@ public:
 	/** Finds the closest reflection capture to a point in space. */
 	const FReflectionCaptureProxy* FindClosestReflectionCapture(FVector Position) const;
 
-	const class FPlanarReflectionSceneProxy* FindClosestPlanarReflection(const FPrimitiveBounds& Bounds) const;
+	const class FPlanarReflectionSceneProxy* FindClosestPlanarReflection(const FBoxSphereBounds& Bounds) const;
 
 	void FindClosestReflectionCaptures(FVector Position, const FReflectionCaptureProxy* (&SortedByDistanceOUT)[FPrimitiveSceneInfo::MaxCachedReflectionCaptureProxies]) const;
 	
diff --git a/Engine/Source/Runtime/Renderer/Private/SceneRendering.cpp b/Engine/Source/Runtime/Renderer/Private/SceneRendering.cpp
--- a/Engine/Source/Runtime/Renderer/Private/SceneRendering.cpp
+++ b/Engine/Source/Runtime/Renderer/Private/SceneRendering.cpp
@@ -160,14 +160,6 @@ static TAutoConsoleVariable<int32> CVarMaxMobileShadowCascades(
 	ECVF_Scalability | ECVF_RenderThreadSafe
 );
 
-static TAutoConsoleVariable<int32> CVarForwardShading(
-	TEXT("r.ForwardShading"),
-	0,
-	TEXT("Whether to use forward shading on desktop platforms - requires Shader Model 5 hardware.\n")
-	TEXT("Forward shading has lower constant cost, but fewer features supported. 0:off, 1:on\n")
-	TEXT("This rendering path is a work in progress with many unimplemented features, notably only a single reflection capture is applied per object and no translucency dynamic shadow receiving."),
-	ECVF_RenderThreadSafe | ECVF_ReadOnly);
-
 static TAutoConsoleVariable<int32> CVarSupportSimpleForwardShading(
 	TEXT("r.SupportSimpleForwardShading"),
 	0,
@@ -598,6 +590,13 @@ void UpdateNoiseTextureParameters(FViewUniformShaderParameters& ViewUniformShade
 	}
 	check(ViewUniformShaderParameters.PerlinNoise3DTexture);
 	ViewUniformShaderParameters.PerlinNoise3DTextureSampler = TStaticSamplerState<SF_Bilinear, AM_Wrap, AM_Wrap, AM_Wrap>::GetRHI();
+
+	if (GSystemTextures.SobolSampling.GetReference())
+	{
+		ViewUniformShaderParameters.SobolSamplingTexture = (FTexture2DRHIRef&)GSystemTextures.SobolSampling->GetRenderTargetItem().ShaderResourceTexture;
+		SetBlack2DIfNull(ViewUniformShaderParameters.SobolSamplingTexture);
+	}
+	check(ViewUniformShaderParameters.SobolSamplingTexture);
 }
 
 /** Creates the view's uniform buffers given a set of view transforms. */
@@ -936,7 +935,7 @@ FViewInfo* FViewInfo::CreateSnapshot() const
 	}
 	else
 	{
-		Result = (FViewInfo*)FMemory::Malloc(sizeof(FViewInfo), ALIGNOF(FViewInfo));
+		Result = (FViewInfo*)FMemory::Malloc(sizeof(FViewInfo), alignof(FViewInfo));
 	}
 	FMemory::Memcpy(*Result, *this);
 
@@ -1931,6 +1930,12 @@ void FRendererModule::BeginRenderingViewFamily(FCanvas* Canvas, FSceneViewFamily
 		}
 	}
 
+	ENQUEUE_RENDER_COMMAND(UpdateDeferredCachedUniformExpressions)(
+		[](FRHICommandList& RHICmdList)
+		{
+			FMaterialRenderProxy::UpdateDeferredCachedUniformExpressions();
+		});
+
 	// Flush the canvas first.
 	Canvas->Flush_GameThread();
 
diff --git a/Engine/Source/Runtime/Renderer/Private/SceneRendering.h b/Engine/Source/Runtime/Renderer/Private/SceneRendering.h
--- a/Engine/Source/Runtime/Renderer/Private/SceneRendering.h
+++ b/Engine/Source/Runtime/Renderer/Private/SceneRendering.h
@@ -152,7 +152,7 @@ namespace ETranslucencyPass
 	enum Type
 	{
 		TPT_StandardTranslucency,
-		TPT_SeparateTranslucency,
+		TPT_TranslucencyAfterDOF,
 
 		/** Drawing all translucency, regardless of separate or standard.  Used when drawing translucency outside of the main renderer, eg FRendererModule::DrawTile. */
 		TPT_AllTranslucency,
@@ -165,6 +165,8 @@ struct FTranslucenyPrimCount
 {
 private:
 	uint32 Count[ETranslucencyPass::TPT_MAX];
+	bool UseSceneColorCopyPerPass[ETranslucencyPass::TPT_MAX];
+	bool DisableOffscreenRenderingPerPass[ETranslucencyPass::TPT_MAX];
 
 public:
 	// constructor
@@ -173,6 +175,8 @@ public:
 		for(uint32 i = 0; i < ETranslucencyPass::TPT_MAX; ++i)
 		{
 			Count[i] = 0;
+			UseSceneColorCopyPerPass[i] = false;
+			DisableOffscreenRenderingPerPass[i] = false;
 		}
 	}
 
@@ -182,13 +186,17 @@ public:
 		for(uint32 i = 0; i < ETranslucencyPass::TPT_MAX; ++i)
 		{
 			Count[i] += InSrc.Count[i];
+			UseSceneColorCopyPerPass[i] |= InSrc.UseSceneColorCopyPerPass[i];
+			DisableOffscreenRenderingPerPass[i] |= InSrc.DisableOffscreenRenderingPerPass[i];
 		}
 	}
 
 	// interface similar to TArray but here we only store the count of Prims per pass
-	void Add(ETranslucencyPass::Type InPass)
+	void Add(ETranslucencyPass::Type InPass, bool bUseSceneColorCopy, bool bDisableOffscreenRendering)
 	{
 		++Count[InPass];
+		UseSceneColorCopyPerPass[InPass] |= bUseSceneColorCopy;
+		DisableOffscreenRenderingPerPass[InPass] |= bDisableOffscreenRendering;
 	}
 
 	// @return range in SortedPrims[] after sorting
@@ -218,6 +226,16 @@ public:
 	{
 		return Count[InPass];
 	}
+
+	bool UseSceneColorCopy(ETranslucencyPass::Type InPass) const
+	{
+		return UseSceneColorCopyPerPass[InPass];
+	}
+
+	bool DisableOffscreenRendering(ETranslucencyPass::Type InPass) const
+	{
+		return DisableOffscreenRenderingPerPass[InPass];
+	}
 };
 
 
@@ -278,9 +296,9 @@ class FMeshDecalPrimSet : public FSortedPrimSet<uint32>
 public:
 	typedef FSortedPrimSet<uint32>::FSortedPrim KeyType;
 
-	static KeyType GenerateKey(FPrimitiveSceneInfo* PrimitiveSceneInfo)
+	static KeyType GenerateKey(FPrimitiveSceneInfo* PrimitiveSceneInfo, int16 InSortPriority)
 	{
-		return KeyType(PrimitiveSceneInfo, 0);
+		return KeyType(PrimitiveSceneInfo, (uint32)(InSortPriority - SHRT_MIN));
 	}
 };
 
@@ -365,7 +383,7 @@ public:
 	* Insert a primitive to the translucency rendering list[s]
 	*/
 	
-	static void PlaceScenePrimitive(FPrimitiveSceneInfo* PrimitiveSceneInfo, const FViewInfo& ViewInfo, bool bUseNormalTranslucency, bool bUseSeparateTranslucency, bool bUseMobileSeparateTranslucency,
+	static void PlaceScenePrimitive(FPrimitiveSceneInfo* PrimitiveSceneInfo, const FViewInfo& ViewInfo, const FPrimitiveViewRelevance& ViewRelevance,
 		FTranslucentSortedPrim* InArrayStart, int32& InOutArrayNum, FTranslucenyPrimCount& OutCount);
 
 	/**
@@ -1389,28 +1407,12 @@ protected:
 	*/
 	bool CheckForProjectedShadows(const FLightSceneInfo* LightSceneInfo) const;
 
-	/** Returns whether a per object shadow should be created due to the light being a stationary light. */
-	bool ShouldCreateObjectShadowForStationaryLight(const FLightSceneInfo* LightSceneInfo, const FPrimitiveSceneProxy* PrimitiveSceneProxy, bool bInteractionShadowMapped) const;
-
 	/** Gathers the list of primitives used to draw various shadow types */
 	void GatherShadowPrimitives(
 		const TArray<FProjectedShadowInfo*, SceneRenderingAllocator>& PreShadows,
 		const TArray<FProjectedShadowInfo*, SceneRenderingAllocator>& ViewDependentWholeSceneShadows,
 		bool bReflectionCaptureScene);
 
-	/**
-	* Checks to see if this primitive is affected by various shadow types
-	*
-	* @param PrimitiveSceneInfoCompact The primitive to check for shadow interaction
-	* @param PreShadows The list of pre-shadows to check against
-	*/
-	void GatherShadowsForPrimitiveInner(const FPrimitiveSceneInfoCompact& PrimitiveSceneInfoCompact,
-		const TArray<FProjectedShadowInfo*, SceneRenderingAllocator>& PreShadows,
-		const TArray<FProjectedShadowInfo*, SceneRenderingAllocator>& ViewDependentWholeSceneShadows,
-		bool bReflectionCaptureScene);
-
-	void BeginRenderRayTracedDistanceFieldProjections(FRHICommandListImmediate& RHICmdList);
-
 	void RenderShadowDepthMaps(FRHICommandListImmediate& RHICmdList);
 	void RenderShadowDepthMapAtlases(FRHICommandListImmediate& RHICmdList);
 
@@ -1456,7 +1458,7 @@ protected:
 	void InitFogConstants();
 
 	/** Returns whether there are translucent primitives to be rendered. */
-	bool ShouldRenderTranslucency() const;
+	bool ShouldRenderTranslucency(ETranslucencyPass::Type TranslucencyPass) const;
 
 	/** TODO: REMOVE if no longer needed: Copies scene color to the viewport's render target after applying gamma correction. */
 	void GammaCorrectToViewportRenderTarget(FRHICommandList& RHICmdList, const FViewInfo* View, float OverrideGamma);
diff --git a/Engine/Source/Runtime/Renderer/Private/SceneVisibility.cpp b/Engine/Source/Runtime/Renderer/Private/SceneVisibility.cpp
--- a/Engine/Source/Runtime/Renderer/Private/SceneVisibility.cpp
+++ b/Engine/Source/Runtime/Renderer/Private/SceneVisibility.cpp
@@ -339,7 +339,7 @@ static int32 FrustumCull(const FScene* Scene, FViewInfo& View)
 				{
 					int32 Index = WordIndex * NumBitsPerDWORD + BitSubIndex;
 					const FPrimitiveBounds& Bounds = Scene->PrimitiveBounds[Index];
-					float DistanceSquared = (Bounds.Origin - ViewOriginForDistanceCulling).SizeSquared();
+					float DistanceSquared = (Bounds.BoxSphereBounds.Origin - ViewOriginForDistanceCulling).SizeSquared();
 					float MaxDrawDistance = Bounds.MaxDrawDistance * MaxDrawDistanceScale;
 					int32 VisibilityId = INDEX_NONE;
 
@@ -358,9 +358,9 @@ static int32 FrustumCull(const FScene* Scene, FViewInfo& View)
 
 					if (DistanceSquared > FMath::Square(MaxDrawDistance + FadeRadius) ||
 						(DistanceSquared < Bounds.MinDrawDistanceSq) ||
-						(UseCustomCulling && !View.CustomVisibilityQuery->IsVisible(VisibilityId, FBoxSphereBounds(Bounds.Origin, Bounds.BoxExtent, Bounds.SphereRadius))) ||
-						(bAlsoUseSphereTest && View.ViewFrustum.IntersectSphere(Bounds.Origin, Bounds.SphereRadius) == false) ||
-						View.ViewFrustum.IntersectBox(Bounds.Origin, Bounds.BoxExtent) == false ||
+						(UseCustomCulling && !View.CustomVisibilityQuery->IsVisible(VisibilityId, FBoxSphereBounds(Bounds.BoxSphereBounds.Origin, Bounds.BoxSphereBounds.BoxExtent, Bounds.BoxSphereBounds.SphereRadius))) ||
+						(bAlsoUseSphereTest && View.ViewFrustum.IntersectSphere(Bounds.BoxSphereBounds.Origin, Bounds.BoxSphereBounds.SphereRadius) == false) ||
+						View.ViewFrustum.IntersectBox(Bounds.BoxSphereBounds.Origin, Bounds.BoxSphereBounds.BoxExtent) == false ||
 						(UseMonoCulling && Scene->Primitives[Index]->Proxy->RenderInMono()))
 					{
 						STAT(NumCulledPrimitives.Increment());
@@ -1512,7 +1512,7 @@ struct FRelevancePacket
 
 			if (ViewRelevance.bDecal)
 			{
-				MeshDecalPrimSet.AddPrim(FMeshDecalPrimSet::GenerateKey(PrimitiveSceneInfo));
+				MeshDecalPrimSet.AddPrim(FMeshDecalPrimSet::GenerateKey(PrimitiveSceneInfo, PrimitiveSceneInfo->Proxy->GetTranslucencySortPriority()));
 			}
 
 			if (bEditorRelevance)
@@ -1535,9 +1535,7 @@ struct FRelevancePacket
 			if (bTranslucentRelevance && !bEditorRelevance && ViewRelevance.bRenderInMainPass)
 			{
 				// Add to set of dynamic translucent primitives
-				FTranslucentPrimSet::PlaceScenePrimitive(PrimitiveSceneInfo, View, 
-					ViewRelevance.bNormalTranslucencyRelevance, ViewRelevance.bSeparateTranslucencyRelevance, ViewRelevance.bMobileSeparateTranslucencyRelevance, 
-					&TranslucencyPrims.Prims[0], TranslucencyPrims.NumPrims, TranslucencyPrimCount);
+				FTranslucentPrimSet::PlaceScenePrimitive(PrimitiveSceneInfo, View, ViewRelevance, &TranslucencyPrims.Prims[0], TranslucencyPrims.NumPrims, TranslucencyPrimCount);
 
 				if (ViewRelevance.bDistortionRelevance)
 				{
@@ -1590,13 +1588,13 @@ struct FRelevancePacket
 				// For mobile, the per-object reflection is used for everything
 				&& (Scene->GetShadingPath() == EShadingPath::Mobile || bTranslucentRelevance || IsForwardShadingEnabled(Scene->GetFeatureLevel())))
 			{
-				PrimitiveSceneInfo->CachedReflectionCaptureProxy = Scene->FindClosestReflectionCapture(Scene->PrimitiveBounds[BitIndex].Origin);
-				PrimitiveSceneInfo->CachedPlanarReflectionProxy = Scene->FindClosestPlanarReflection(Scene->PrimitiveBounds[BitIndex]);
+				PrimitiveSceneInfo->CachedReflectionCaptureProxy = Scene->FindClosestReflectionCapture(Scene->PrimitiveBounds[BitIndex].BoxSphereBounds.Origin);
+				PrimitiveSceneInfo->CachedPlanarReflectionProxy = Scene->FindClosestPlanarReflection(Scene->PrimitiveBounds[BitIndex].BoxSphereBounds);
 
 				if (Scene->GetShadingPath() == EShadingPath::Mobile)
 				{
 					// mobile HQ reflections
-					Scene->FindClosestReflectionCaptures(Scene->PrimitiveBounds[BitIndex].Origin, PrimitiveSceneInfo->CachedReflectionCaptureProxies);
+					Scene->FindClosestReflectionCaptures(Scene->PrimitiveBounds[BitIndex].BoxSphereBounds.Origin, PrimitiveSceneInfo->CachedReflectionCaptureProxies);
 				}
 
 				PrimitiveSceneInfo->bNeedsCachedReflectionCaptureUpdate = false;
@@ -1630,15 +1628,15 @@ struct FRelevancePacket
 			const FPrimitiveBounds& Bounds = Scene->PrimitiveBounds[PrimitiveIndex];
 			const FPrimitiveViewRelevance& ViewRelevance = View.PrimitiveViewRelevanceMap[PrimitiveIndex];
 
-			FLODMask LODToRender = ComputeLODForMeshes( PrimitiveSceneInfo->StaticMeshes, View, Bounds.Origin, Bounds.SphereRadius, ViewData.ForcedLODLevel, ViewData.LODScale);
+			FLODMask LODToRender = ComputeLODForMeshes( PrimitiveSceneInfo->StaticMeshes, View, Bounds.BoxSphereBounds.Origin, Bounds.BoxSphereBounds.SphereRadius, ViewData.ForcedLODLevel, ViewData.LODScale);
 			const bool bIsHLODFading = bHLODActive && ViewState && ViewState->HLODVisibilityState.IsNodeFading(PrimitiveIndex);
 			const bool bIsHLODFadingOut = bHLODActive && ViewState && ViewState->HLODVisibilityState.IsNodeFadingOut(PrimitiveIndex);
 			const bool bIsLODDithered = LODToRender.IsDithered();
 
-			float DistanceSquared = (Bounds.Origin - ViewData.ViewOrigin).SizeSquared();
+			float DistanceSquared = (Bounds.BoxSphereBounds.Origin - ViewData.ViewOrigin).SizeSquared();
 			const float LODFactorDistanceSquared = DistanceSquared * FMath::Square(View.LODDistanceFactor * ViewData.InvLODScale);
-			const bool bDrawShadowDepth = FMath::Square(Bounds.SphereRadius) > ViewData.MinScreenRadiusForCSMDepthSquared * LODFactorDistanceSquared;
-			const bool bDrawDepthOnly = ViewData.bFullEarlyZPass || FMath::Square(Bounds.SphereRadius) > GMinScreenRadiusForDepthPrepass * GMinScreenRadiusForDepthPrepass * LODFactorDistanceSquared;
+			const bool bDrawShadowDepth = FMath::Square(Bounds.BoxSphereBounds.SphereRadius) > ViewData.MinScreenRadiusForCSMDepthSquared * LODFactorDistanceSquared;
+			const bool bDrawDepthOnly = ViewData.bFullEarlyZPass || FMath::Square(Bounds.BoxSphereBounds.SphereRadius) > GMinScreenRadiusForDepthPrepass * GMinScreenRadiusForDepthPrepass * LODFactorDistanceSquared;
 
 			const int32 NumStaticMeshes = PrimitiveSceneInfo->StaticMeshes.Num();
 			for(int32 MeshIndex = 0;MeshIndex < NumStaticMeshes;MeshIndex++)
@@ -1729,7 +1727,7 @@ struct FRelevancePacket
 					// Static meshes which don't need per-element visibility always draw all elements
 					if (bNeedsBatchVisibility && StaticMesh.bRequiresPerElementVisibility)
 					{
-						WriteView.StaticMeshBatchVisibility[StaticMesh.Id] = StaticMesh.VertexFactory->GetStaticBatchElementVisibility(View, &StaticMesh);
+						WriteView.StaticMeshBatchVisibility[StaticMesh.BatchVisibilityId] = StaticMesh.VertexFactory->GetStaticBatchElementVisibility(View, &StaticMesh);
 					}
 				}
 			}
@@ -2403,7 +2401,7 @@ void FSceneRenderer::ComputeViewVisibility(FRHICommandListImmediate& RHICmdList)
 		View.StaticMeshFadeInDitheredLODMap.Init(false,Scene->StaticMeshes.GetMaxIndex());
 		View.StaticMeshVelocityMap.Init(false,Scene->StaticMeshes.GetMaxIndex());
 		View.StaticMeshShadowDepthMap.Init(false,Scene->StaticMeshes.GetMaxIndex());
-		View.StaticMeshBatchVisibility.AddZeroed(Scene->StaticMeshes.GetMaxIndex());
+		View.StaticMeshBatchVisibility.AddZeroed(Scene->StaticMeshBatchVisibility.GetMaxIndex());
 
 		View.VisibleLightInfos.Empty(Scene->Lights.GetMaxIndex());
 
@@ -2566,7 +2564,7 @@ void FSceneRenderer::ComputeViewVisibility(FRHICommandListImmediate& RHICmdList)
 			float ScreenSizeScale = FMath::Max(View.ViewMatrices.GetProjectionMatrix().M[0][0] * View.ViewRect.Width(), View.ViewMatrices.GetProjectionMatrix().M[1][1] * View.ViewRect.Height());
 			for (FSceneSetBitIterator BitIt(View.PrimitiveVisibilityMap); BitIt; ++BitIt)
 			{
-				if (ScreenSizeScale * Scene->PrimitiveBounds[BitIt.GetIndex()].SphereRadius <= GWireframeCullThreshold)
+				if (ScreenSizeScale * Scene->PrimitiveBounds[BitIt.GetIndex()].BoxSphereBounds.SphereRadius <= GWireframeCullThreshold)
 				{
 					View.PrimitiveVisibilityMap.AccessCorrespondingBit(BitIt) = false;
 				}
@@ -2807,8 +2805,7 @@ void FSceneRenderer::PostVisibilityFrameSetup(FILCUpdatePrimTaskData& OutILCTask
 				
 					FLinearColor Color( ColorAndFalloffExponent );
 
-					// Scale by visible area
-					Color /= PI * FMath::Square( SourceRadius );
+					Color /= PI * FMath::Square( SourceRadius ) + 0.5f * PI * SourceRadius * SourceLength;
 
 					if( Proxy->IsInverseSquared() )
 					{
@@ -2831,9 +2828,12 @@ void FSceneRenderer::PostVisibilityFrameSetup(FILCUpdatePrimTaskData& OutILCTask
 					FVector L = ToLight.GetSafeNormal();
 					Color.A *= FMath::Square( FMath::Clamp( ( (L | Direction) - SpotAngles.X ) * SpotAngles.Y, 0.0f, 1.0f ) );
 				
-					FViewElementPDI LightPDI( &View, NULL );
 					FMaterialRenderProxy* const ColoredMeshInstance = new(FMemStack::Get()) FColoredMaterialRenderProxy( GEngine->DebugMeshMaterial->GetRenderProxy(false), Color );
-					DrawSphere( &LightPDI, Origin, FVector( SourceRadius, SourceRadius, SourceRadius ), 36, 24, ColoredMeshInstance, SDPG_World );
+
+					FViewElementPDI LightPDI( &View, NULL );
+					// Scaled sphere to handle SourceLength
+					const float ZScale = FMath::Max(SourceRadius, SourceLength);
+					DrawSphere(&LightPDI, Origin, FRotationMatrix::MakeFromZ(Direction).Rotator(), FVector(SourceRadius, SourceRadius, ZScale), 36, 24, ColoredMeshInstance, SDPG_World);
 				}
 			}
 		}
@@ -3088,7 +3088,7 @@ void FLODSceneTree::UpdateAndApplyVisibilityStates(FViewInfo& View)
 				}
 			}
 
-			const float DistanceSquared = (Bounds.Origin - View.ViewMatrices.GetViewOrigin()).SizeSquared();
+			const float DistanceSquared = (Bounds.BoxSphereBounds.Origin - View.ViewMatrices.GetViewOrigin()).SizeSquared();
 			const bool bIsInDrawRange = DistanceSquared >= Bounds.MinDrawDistanceSq;
 
 			const bool bWasFadingPreUpdate = !!NodeVisibility.bIsFading;
diff --git a/Engine/Source/Runtime/Renderer/Private/ShadowDepthRendering.cpp b/Engine/Source/Runtime/Renderer/Private/ShadowDepthRendering.cpp
--- a/Engine/Source/Runtime/Renderer/Private/ShadowDepthRendering.cpp
+++ b/Engine/Source/Runtime/Renderer/Private/ShadowDepthRendering.cpp
@@ -1319,7 +1319,7 @@ void FProjectedShadowInfo::ClearDepth(FRHICommandList& RHICmdList, class FSceneR
 			ViewportMaxZ
 			);
 		
-		DrawClearQuadMRT(RHICmdList, SceneRenderer->FeatureLevel, bClearColor, NumClearColors, Colors, true, 1.0f, false, 0);
+		DrawClearQuadMRT(RHICmdList, bClearColor, NumClearColors, Colors, true, 1.0f, false, 0);
 	}
 	else
 	{
@@ -1360,7 +1360,7 @@ void DrawMeshElements(FRHICommandList& RHICmdList, FShadowDepthDrawingPolicy<bRe
 	SharedDrawingPolicy.ApplyDitheredLODTransitionState(DrawRenderStateLocal, View, *Mesh, false);
 
 	// Render only those batch elements that match the current LOD
-	uint64 BatchElementMask = Mesh->bRequiresPerElementVisibility ? View.StaticMeshBatchVisibility[Mesh->Id] : ((1ull << Mesh->Elements.Num()) - 1);
+	uint64 BatchElementMask = Mesh->bRequiresPerElementVisibility ? View.StaticMeshBatchVisibility[Mesh->BatchVisibilityId] : ((1ull << Mesh->Elements.Num()) - 1);
 	int32 BatchElementIndex = 0;
 	do
 	{
@@ -1580,7 +1580,7 @@ void FProjectedShadowInfo::SetStateForDepth(FRHICommandList& RHICmdList, EShadow
 
 static TAutoConsoleVariable<int32> CVarParallelShadows(
 	TEXT("r.ParallelShadows"),
-	0,
+	1,
 	TEXT("Toggles parallel shadow rendering. Parallel rendering must be enabled for this to have an effect."),
 	ECVF_RenderThreadSafe
 	);
@@ -2173,6 +2173,11 @@ void FSceneRenderer::RenderShadowDepthMapAtlases(FRHICommandListImmediate& RHICm
 			}
 			InRHICmdList.TransitionResource(EResourceTransitionAccess::EWritable, Info.DepthStencilRenderTarget.Texture);
 			InRHICmdList.SetRenderTargetsAndClear(Info);
+
+			if (!bPerformClear)
+			{
+				InRHICmdList.BindClearMRTValues(false, true, false);
+			}
 		};
 
 		{
@@ -2221,31 +2226,6 @@ void FSceneRenderer::RenderShadowDepthMapAtlases(FRHICommandListImmediate& RHICm
 	}
 }
 
-void FSceneRenderer::BeginRenderRayTracedDistanceFieldProjections(FRHICommandListImmediate& RHICmdList)
-{
-	for (TSparseArray<FLightSceneInfoCompact>::TConstIterator LightIt(Scene->Lights); LightIt; ++LightIt)
-	{
-		const FLightSceneInfoCompact& LightSceneInfoCompact = *LightIt;
-		const FLightSceneInfo* const LightSceneInfo = LightSceneInfoCompact.LightSceneInfo;
-		FVisibleLightInfo& VisibleLightInfo = VisibleLightInfos[LightSceneInfo->Id];
-
-		for (int32 ViewIndex = 0; ViewIndex < Views.Num(); ViewIndex++)
-		{
-			const FViewInfo& View = Views[ViewIndex];
-
-			for (int32 ShadowIndex = 0; ShadowIndex < VisibleLightInfo.ShadowsToProject.Num(); ShadowIndex++)
-			{
-				FProjectedShadowInfo* ProjectedShadowInfo = VisibleLightInfo.ShadowsToProject[ShadowIndex];
-
-				if (ProjectedShadowInfo->bRayTracedDistanceField)
-				{
-					ProjectedShadowInfo->BeginRenderRayTracedDistanceFieldProjection(RHICmdList, View);
-				}
-			}
-		}
-	}
-}
-
 void FSceneRenderer::RenderShadowDepthMaps(FRHICommandListImmediate& RHICmdList)
 {
 	FSceneRenderTargets& SceneContext = FSceneRenderTargets::Get(RHICmdList);
@@ -2253,8 +2233,6 @@ void FSceneRenderer::RenderShadowDepthMaps(FRHICommandListImmediate& RHICmdList)
 	SCOPED_DRAW_EVENT(RHICmdList, ShadowDepths);
 	SCOPED_GPU_STAT(RHICmdList, Stat_GPU_ShadowDepths);
 
-	BeginRenderRayTracedDistanceFieldProjections(RHICmdList);
-
 	FSceneRenderer::RenderShadowDepthMapAtlases(RHICmdList);
 
 	for (int32 CubemapIndex = 0; CubemapIndex < SortedShadowsForShadowDepthPass.ShadowMapCubemaps.Num(); CubemapIndex++)
diff --git a/Engine/Source/Runtime/Renderer/Private/ShadowRendering.cpp b/Engine/Source/Runtime/Renderer/Private/ShadowRendering.cpp
--- a/Engine/Source/Runtime/Renderer/Private/ShadowRendering.cpp
+++ b/Engine/Source/Runtime/Renderer/Private/ShadowRendering.cpp
@@ -78,6 +78,20 @@ static FAutoConsoleVariableRef CVarStencilOptimization(
 	ECVF_RenderThreadSafe
 	);
 
+static TAutoConsoleVariable<int32> CVarFilterMethod(
+	TEXT("r.Shadow.FilterMethod"),
+	0,
+	TEXT("Chooses the shadow filtering method.\n")
+	TEXT(" 0: Uniform PCF (default)\n")
+	TEXT(" 1: PCSS (experimental)\n"),
+	ECVF_RenderThreadSafe);
+
+static TAutoConsoleVariable<int32> CVarMaxSoftKernelSize(
+	TEXT("r.Shadow.MaxSoftKernelSize"),
+	40,
+	TEXT("Mazimum size of the softening kernels in pixels."),
+	ECVF_RenderThreadSafe);
+
 DECLARE_FLOAT_COUNTER_STAT(TEXT("ShadowProjection"), Stat_GPU_ShadowProjection, STATGROUP_GPU);
 
 // 0:off, 1:low, 2:med, 3:high, 4:very high, 5:max
@@ -144,15 +158,15 @@ TGlobalResource<StencilingGeometry::FStencilConeVertexBuffer> StencilingGeometry
 TGlobalResource<StencilingGeometry::FStencilConeIndexBuffer> StencilingGeometry::GStencilConeIndexBuffer;
 
 /*-----------------------------------------------------------------------------
-	FShadowProjectionVS
+	FShadowVolumeBoundProjectionVS
 -----------------------------------------------------------------------------*/
 
-bool FShadowProjectionVS::ShouldCache(EShaderPlatform Platform)
+bool FShadowVolumeBoundProjectionVS::ShouldCache(EShaderPlatform Platform)
 {
 	return true;
 }
 
-void FShadowProjectionVS::SetParameters(FRHICommandList& RHICmdList, const FSceneView& View, const FProjectedShadowInfo* ShadowInfo)
+void FShadowVolumeBoundProjectionVS::SetParameters(FRHICommandList& RHICmdList, const FSceneView& View, const FProjectedShadowInfo* ShadowInfo)
 {
 	FGlobalShader::SetParameters<FViewUniformShaderParameters>(RHICmdList, GetVertexShader(),View.ViewUniformBuffer);
 	
@@ -176,7 +190,7 @@ void FShadowProjectionVS::SetParameters(FRHICommandList& RHICmdList, const FScen
 
 IMPLEMENT_SHADER_TYPE(,FShadowProjectionNoTransformVS,TEXT("ShadowProjectionVertexShader"),TEXT("Main"),SF_Vertex);
 
-IMPLEMENT_SHADER_TYPE(,FShadowProjectionVS,TEXT("ShadowProjectionVertexShader"),TEXT("Main"),SF_Vertex);
+IMPLEMENT_SHADER_TYPE(,FShadowVolumeBoundProjectionVS,TEXT("ShadowProjectionVertexShader"),TEXT("Main"),SF_Vertex);
 
 /**
  * Implementations for TShadowProjectionPS.  
@@ -199,6 +213,8 @@ IMPLEMENT_SHADOW_PROJECTION_PIXEL_SHADER(2,true);
 IMPLEMENT_SHADOW_PROJECTION_PIXEL_SHADER(3,true);
 IMPLEMENT_SHADOW_PROJECTION_PIXEL_SHADER(4,true);
 IMPLEMENT_SHADOW_PROJECTION_PIXEL_SHADER(5,true);
+
+#undef IMPLEMENT_SHADOW_PROJECTION_PIXEL_SHADER
 #endif
 
 // Implement a pixel shader for rendering modulated shadow projections.
@@ -220,7 +236,15 @@ IMPLEMENT_SHADER_TYPE(template<>,TOnePassPointShadowProjectionPS<1>,TEXT("Shadow
 IMPLEMENT_SHADER_TYPE(template<>,TOnePassPointShadowProjectionPS<2>,TEXT("ShadowProjectionPixelShader"),TEXT("MainOnePassPointLightPS"),SF_Pixel);
 IMPLEMENT_SHADER_TYPE(template<>,TOnePassPointShadowProjectionPS<3>,TEXT("ShadowProjectionPixelShader"),TEXT("MainOnePassPointLightPS"),SF_Pixel);
 IMPLEMENT_SHADER_TYPE(template<>,TOnePassPointShadowProjectionPS<4>,TEXT("ShadowProjectionPixelShader"),TEXT("MainOnePassPointLightPS"),SF_Pixel);
-IMPLEMENT_SHADER_TYPE(template<>,TOnePassPointShadowProjectionPS<5>,TEXT("ShadowProjectionPixelShader"),TEXT("MainOnePassPointLightPS"),SF_Pixel);
+IMPLEMENT_SHADER_TYPE(template<>, TOnePassPointShadowProjectionPS<5>, TEXT("ShadowProjectionPixelShader"), TEXT("MainOnePassPointLightPS"), SF_Pixel);
+
+// Implements a pixel shader for directional light PCSS.
+#define IMPLEMENT_SHADOW_PROJECTION_PIXEL_SHADER(Quality,UseFadePlane) \
+	typedef TDirectionalPercentageCloserShadowProjectionPS<Quality, UseFadePlane> TDirectionalPercentageCloserShadowProjectionPS##Quality##UseFadePlane; \
+	IMPLEMENT_SHADER_TYPE(template<>,TDirectionalPercentageCloserShadowProjectionPS##Quality##UseFadePlane,TEXT("ShadowProjectionPixelShader"),TEXT("Main"),SF_Pixel);
+IMPLEMENT_SHADOW_PROJECTION_PIXEL_SHADER(5,false);
+IMPLEMENT_SHADOW_PROJECTION_PIXEL_SHADER(5,true);
+#undef IMPLEMENT_SHADOW_PROJECTION_PIXEL_SHADER
 
 void StencilingGeometry::DrawSphere(FRHICommandList& RHICmdList)
 {
@@ -247,88 +271,100 @@ void StencilingGeometry::DrawCone(FRHICommandList& RHICmdList)
 		FStencilConeIndexBuffer::NumVerts, 0, StencilingGeometry::GStencilConeIndexBuffer.GetIndexCount() / 3, 1);
 }
 
-template <uint32 Quality>
-static void SetShadowProjectionShaderTemplNew(FRHICommandList& RHICmdList, FGraphicsPipelineStateInitializer& GraphicsPSOInit, int32 ViewIndex, const FViewInfo& View, const FProjectedShadowInfo* ShadowInfo, bool bMobileModulatedProjections)
+static void GetShadowProjectionShaders(
+	int32 Quality, const FViewInfo& View, const FProjectedShadowInfo* ShadowInfo, bool bMobileModulatedProjections,
+	FShadowProjectionVertexShaderInterface** OutShadowProjVS, FShadowProjectionPixelShaderInterface** OutShadowProjPS)
 {
+	check(!*OutShadowProjVS);
+	check(!*OutShadowProjPS);
+
 	if (ShadowInfo->bTranslucentShadow)
 	{
-		// Get the Shadow Projection Vertex Shader (with transforms)
-		FShadowProjectionVS* ShadowProjVS = View.ShaderMap->GetShader<FShadowProjectionVS>();
-
-		// Get the translucency pixel shader
-		FShadowProjectionPixelShaderInterface* ShadowProjPS = View.ShaderMap->GetShader<TShadowProjectionFromTranslucencyPS<Quality> >();
-
-		GraphicsPSOInit.BoundShaderState.VertexDeclarationRHI = GetVertexDeclarationFVector4();
-		GraphicsPSOInit.BoundShaderState.VertexShaderRHI = GETSAFERHISHADER_VERTEX(ShadowProjVS);
-		GraphicsPSOInit.BoundShaderState.PixelShaderRHI = GETSAFERHISHADER_PIXEL(ShadowProjPS);
+		*OutShadowProjVS = View.ShaderMap->GetShader<FShadowVolumeBoundProjectionVS>();
 
-		SetGraphicsPipelineState(RHICmdList, GraphicsPSOInit);
-
-		// Set shader parameters
-		ShadowProjVS->SetParameters(RHICmdList, View, ShadowInfo);
-		ShadowProjPS->SetParameters(RHICmdList, ViewIndex, View, ShadowInfo);
+		switch (Quality)
+		{
+		case 1: *OutShadowProjPS = View.ShaderMap->GetShader<TShadowProjectionFromTranslucencyPS<1> >(); break;
+		case 2: *OutShadowProjPS = View.ShaderMap->GetShader<TShadowProjectionFromTranslucencyPS<2> >(); break;
+		case 3: *OutShadowProjPS = View.ShaderMap->GetShader<TShadowProjectionFromTranslucencyPS<3> >(); break;
+		case 4: *OutShadowProjPS = View.ShaderMap->GetShader<TShadowProjectionFromTranslucencyPS<4> >(); break;
+		case 5: *OutShadowProjPS = View.ShaderMap->GetShader<TShadowProjectionFromTranslucencyPS<5> >(); break;
+		default:
+			check(0);
+		}
 	}
 	else if (ShadowInfo->IsWholeSceneDirectionalShadow())
 	{
-		// Get the Shadow Projection Vertex Shader which does not use a transform
-		FShadowProjectionNoTransformVS* ShadowProjVS = View.ShaderMap->GetShader<FShadowProjectionNoTransformVS>();
+		*OutShadowProjVS = View.ShaderMap->GetShader<FShadowProjectionNoTransformVS>();
 
-		// Get the Shadow Projection Pixel Shader for PSSM
-		if (ShadowInfo->CascadeSettings.FadePlaneLength > 0)
+		if (CVarFilterMethod.GetValueOnRenderThread() == 1)
 		{
-			// This shader fades the shadow towards the end of the split subfrustum.
-			FShadowProjectionPixelShaderInterface* ShadowProjPS = View.ShaderMap->GetShader<TShadowProjectionPS<Quality, true> >();
-
-			GraphicsPSOInit.BoundShaderState.VertexDeclarationRHI = GetVertexDeclarationFVector4();
-			GraphicsPSOInit.BoundShaderState.VertexShaderRHI = GETSAFERHISHADER_VERTEX(ShadowProjVS);
-			GraphicsPSOInit.BoundShaderState.PixelShaderRHI = GETSAFERHISHADER_PIXEL(ShadowProjPS);
-
-			SetGraphicsPipelineState(RHICmdList, GraphicsPSOInit);
-
-			ShadowProjPS->SetParameters(RHICmdList, ViewIndex, View, ShadowInfo);
+			if (ShadowInfo->CascadeSettings.FadePlaneLength > 0)
+				*OutShadowProjPS = View.ShaderMap->GetShader<TDirectionalPercentageCloserShadowProjectionPS<5, true> >();
+			else
+				*OutShadowProjPS = View.ShaderMap->GetShader<TDirectionalPercentageCloserShadowProjectionPS<5, false> >();
+		}
+		else if (ShadowInfo->CascadeSettings.FadePlaneLength > 0)
+		{
+			switch (Quality)
+			{
+			case 1: *OutShadowProjPS = View.ShaderMap->GetShader<TShadowProjectionPS<1, true> >(); break;
+			case 2: *OutShadowProjPS = View.ShaderMap->GetShader<TShadowProjectionPS<2, true> >(); break;
+			case 3: *OutShadowProjPS = View.ShaderMap->GetShader<TShadowProjectionPS<3, true> >(); break;
+			case 4: *OutShadowProjPS = View.ShaderMap->GetShader<TShadowProjectionPS<4, true> >(); break;
+			case 5: *OutShadowProjPS = View.ShaderMap->GetShader<TShadowProjectionPS<5, true> >(); break;
+			default:
+				check(0);
+			}
 		}
 		else
 		{
-			// Do not use the fade plane shader if the fade plane region length is 0 (avoids divide by 0).
-			FShadowProjectionPixelShaderInterface* ShadowProjPS = View.ShaderMap->GetShader<TShadowProjectionPS<Quality, false> >();
-
-			GraphicsPSOInit.BoundShaderState.VertexDeclarationRHI = GetVertexDeclarationFVector4();
-			GraphicsPSOInit.BoundShaderState.VertexShaderRHI = GETSAFERHISHADER_VERTEX(ShadowProjVS);
-			GraphicsPSOInit.BoundShaderState.PixelShaderRHI = GETSAFERHISHADER_PIXEL(ShadowProjPS);
-
-			SetGraphicsPipelineState(RHICmdList, GraphicsPSOInit);
-
-			ShadowProjPS->SetParameters(RHICmdList, ViewIndex, View, ShadowInfo);
+			switch (Quality)
+			{
+			case 1: *OutShadowProjPS = View.ShaderMap->GetShader<TShadowProjectionPS<1, false> >(); break;
+			case 2: *OutShadowProjPS = View.ShaderMap->GetShader<TShadowProjectionPS<2, false> >(); break;
+			case 3: *OutShadowProjPS = View.ShaderMap->GetShader<TShadowProjectionPS<3, false> >(); break;
+			case 4: *OutShadowProjPS = View.ShaderMap->GetShader<TShadowProjectionPS<4, false> >(); break;
+			case 5: *OutShadowProjPS = View.ShaderMap->GetShader<TShadowProjectionPS<5, false> >(); break;
+			default:
+				check(0);
+			}
 		}
-
-		ShadowProjVS->SetParameters(RHICmdList, View.ViewUniformBuffer);
 	}
 	else
 	{
-		// Get the Shadow Projection Vertex Shader
-		FShadowProjectionVS* ShadowProjVS = View.ShaderMap->GetShader<FShadowProjectionVS>();
+		*OutShadowProjVS = View.ShaderMap->GetShader<FShadowVolumeBoundProjectionVS>();
 
-		// Get the Shadow Projection Pixel Shader
-		// This shader is the ordinary projection shader used by point/spot lights.		
-		FShadowProjectionPixelShaderInterface* ShadowProjPS;
 		if(bMobileModulatedProjections)
 		{
-			ShadowProjPS = View.ShaderMap->GetShader<TModulatedShadowProjection<Quality> >();
+			switch (Quality)
+			{
+			case 1: *OutShadowProjPS = View.ShaderMap->GetShader<TModulatedShadowProjection<1> >(); break;
+			case 2: *OutShadowProjPS = View.ShaderMap->GetShader<TModulatedShadowProjection<2> >(); break;
+			case 3: *OutShadowProjPS = View.ShaderMap->GetShader<TModulatedShadowProjection<3> >(); break;
+			case 4: *OutShadowProjPS = View.ShaderMap->GetShader<TModulatedShadowProjection<4> >(); break;
+			case 5: *OutShadowProjPS = View.ShaderMap->GetShader<TModulatedShadowProjection<5> >(); break;
+			default:
+				check(0);
+			}
 		}
 		else
 		{
-			ShadowProjPS = View.ShaderMap->GetShader<TShadowProjectionPS<Quality, false> >();
+			switch (Quality)
+			{
+			case 1: *OutShadowProjPS = View.ShaderMap->GetShader<TShadowProjectionPS<1, false> >(); break;
+			case 2: *OutShadowProjPS = View.ShaderMap->GetShader<TShadowProjectionPS<2, false> >(); break;
+			case 3: *OutShadowProjPS = View.ShaderMap->GetShader<TShadowProjectionPS<3, false> >(); break;
+			case 4: *OutShadowProjPS = View.ShaderMap->GetShader<TShadowProjectionPS<4, false> >(); break;
+			case 5: *OutShadowProjPS = View.ShaderMap->GetShader<TShadowProjectionPS<5, false> >(); break;
+			default:
+				check(0);
+			}
 		}
-
-		GraphicsPSOInit.BoundShaderState.VertexDeclarationRHI = GetVertexDeclarationFVector4();
-		GraphicsPSOInit.BoundShaderState.VertexShaderRHI = GETSAFERHISHADER_VERTEX(ShadowProjVS);
-		GraphicsPSOInit.BoundShaderState.PixelShaderRHI = GETSAFERHISHADER_PIXEL(ShadowProjPS);
-
-		SetGraphicsPipelineState(RHICmdList, GraphicsPSOInit);
-
-		ShadowProjVS->SetParameters(RHICmdList, View, ShadowInfo);
-		ShadowProjPS->SetParameters(RHICmdList, ViewIndex, View, ShadowInfo);
 	}
+
+	check(*OutShadowProjVS);
+	check(*OutShadowProjPS);
 }
 
 void FProjectedShadowInfo::SetBlendStateForProjection(
@@ -607,7 +643,7 @@ void FProjectedShadowInfo::SetupProjectionStencilMask(
 								*View,
 								FDepthDrawingPolicyFactory::ContextType(DDM_AllOccluders, false),
 								StaticMesh,
-								StaticMesh.bRequiresPerElementVisibility ? View->StaticMeshBatchVisibility[StaticMesh.Id] : ((1ull << StaticMesh.Elements.Num() )- 1),
+								StaticMesh.bRequiresPerElementVisibility ? View->StaticMeshBatchVisibility[StaticMesh.BatchVisibilityId] : ((1ull << StaticMesh.Elements.Num() )- 1),
 								true,
 								DrawRenderState,
 								ReceiverPrimitiveSceneInfo->Proxy,
@@ -631,7 +667,7 @@ void FProjectedShadowInfo::SetupProjectionStencilMask(
 					*View,
 					FDepthDrawingPolicyFactory::ContextType(DDM_AllOccluders, false),
 					StaticMesh,
-					StaticMesh.bRequiresPerElementVisibility ? View->StaticMeshBatchVisibility[StaticMesh.Id] : ((1ull << StaticMesh.Elements.Num() )- 1),
+					StaticMesh.bRequiresPerElementVisibility ? View->StaticMeshBatchVisibility[StaticMesh.BatchVisibilityId] : ((1ull << StaticMesh.Elements.Num() )- 1),
 					true,
 					DrawRenderState,
 					StaticMesh.PrimitiveSceneInfo->Proxy,
@@ -746,7 +782,7 @@ void FProjectedShadowInfo::SetupProjectionStencilMask(
 		GraphicsPSOInit.RasterizerState = TStaticRasterizerState<FM_Solid, CM_None>::GetRHI();
 
 		// Find the projection shaders.
-		TShaderMapRef<FShadowProjectionVS> VertexShader(View->ShaderMap);
+		TShaderMapRef<FShadowVolumeBoundProjectionVS> VertexShader(View->ShaderMap);
 
 		GraphicsPSOInit.BoundShaderState.VertexDeclarationRHI = GetVertexDeclarationFVector4();
 		GraphicsPSOInit.BoundShaderState.VertexShaderRHI = GETSAFERHISHADER_VERTEX(*VertexShader);
@@ -910,16 +946,19 @@ void FProjectedShadowInfo::RenderProjection(FRHICommandListImmediate& RHICmdList
 			}
 		}
 
-		switch(LocalQuality)
-		{
-			case 1: SetShadowProjectionShaderTemplNew<1>(RHICmdList, GraphicsPSOInit, ViewIndex, *View, this, bMobileModulatedProjections); break;
-			case 2: SetShadowProjectionShaderTemplNew<2>(RHICmdList, GraphicsPSOInit, ViewIndex, *View, this, bMobileModulatedProjections); break;
-			case 3: SetShadowProjectionShaderTemplNew<3>(RHICmdList, GraphicsPSOInit, ViewIndex, *View, this, bMobileModulatedProjections); break;
-			case 4: SetShadowProjectionShaderTemplNew<4>(RHICmdList, GraphicsPSOInit, ViewIndex, *View, this, bMobileModulatedProjections); break;
-			case 5: SetShadowProjectionShaderTemplNew<5>(RHICmdList, GraphicsPSOInit, ViewIndex, *View, this, bMobileModulatedProjections); break;
-			default:
-				check(0);
-		}
+		FShadowProjectionVertexShaderInterface* ShadowProjVS = nullptr;
+		FShadowProjectionPixelShaderInterface* ShadowProjPS = nullptr;
+
+		GetShadowProjectionShaders(LocalQuality, *View, this, bMobileModulatedProjections, &ShadowProjVS, &ShadowProjPS);
+
+		GraphicsPSOInit.BoundShaderState.VertexDeclarationRHI = GetVertexDeclarationFVector4();
+		GraphicsPSOInit.BoundShaderState.VertexShaderRHI = GETSAFERHISHADER_VERTEX(ShadowProjVS);
+		GraphicsPSOInit.BoundShaderState.PixelShaderRHI = GETSAFERHISHADER_PIXEL(ShadowProjPS);
+
+		SetGraphicsPipelineState(RHICmdList, GraphicsPSOInit);
+
+		ShadowProjVS->SetParameters(RHICmdList, *View, this);
+		ShadowProjPS->SetParameters(RHICmdList, ViewIndex, *View, this);
 	}
 
 	if (IsWholeSceneDirectionalShadow())
@@ -949,7 +988,7 @@ void FProjectedShadowInfo::RenderProjection(FRHICommandListImmediate& RHICmdList
 		// Clear the stencil buffer to 0.
 		if (!GStencilOptimization)
 		{
-			DrawClearQuad(RHICmdList, GMaxRHIFeatureLevel, false, FLinearColor::Transparent, false, 0, true, 1);
+			DrawClearQuad(RHICmdList, false, FLinearColor::Transparent, false, 0, true, 1);
 		}
 	}
 }
@@ -958,7 +997,7 @@ void FProjectedShadowInfo::RenderProjection(FRHICommandListImmediate& RHICmdList
 template <uint32 Quality>
 static void SetPointLightShaderTempl(FRHICommandList& RHICmdList, FGraphicsPipelineStateInitializer& GraphicsPSOInit, int32 ViewIndex, const FViewInfo& View, const FProjectedShadowInfo* ShadowInfo)
 {
-	TShaderMapRef<FShadowProjectionVS> VertexShader(View.ShaderMap);
+	TShaderMapRef<FShadowVolumeBoundProjectionVS> VertexShader(View.ShaderMap);
 	TShaderMapRef<TOnePassPointShadowProjectionPS<Quality> > PixelShader(View.ShaderMap);
 
 	GraphicsPSOInit.BoundShaderState.VertexDeclarationRHI = GetVertexDeclarationFVector4();
diff --git a/Engine/Source/Runtime/Renderer/Private/ShadowRendering.h b/Engine/Source/Runtime/Renderer/Private/ShadowRendering.h
--- a/Engine/Source/Runtime/Renderer/Private/ShadowRendering.h
+++ b/Engine/Source/Runtime/Renderer/Private/ShadowRendering.h
@@ -894,7 +894,7 @@ public:
 	/**
 	 * Adds a primitive to the shadow's subject list.
 	 */
-	void AddSubjectPrimitive(FPrimitiveSceneInfo* PrimitiveSceneInfo, TArray<FViewInfo>* ViewArray, bool bRecordShadowSubjectForMobileShading);
+	void AddSubjectPrimitive(FPrimitiveSceneInfo* PrimitiveSceneInfo, TArray<FViewInfo>* ViewArray, ERHIFeatureLevel::Type FeatureLevel, bool bRecordShadowSubjectForMobileShading);
 
 	/**
 	* @return TRUE if this shadow info has any casting subject prims to render
@@ -1068,7 +1068,7 @@ private:
 
 	friend class FShadowDepthVS;
 	template <bool bRenderingReflectiveShadowMaps> friend class TShadowDepthBasePS;
-	friend class FShadowProjectionVS;
+	friend class FShadowVolumeBoundProjectionVS;
 	friend class FShadowProjectionPS;
 	friend class FShadowDepthDrawingPolicyFactory;
 };
@@ -1190,16 +1190,31 @@ private:
 	FShaderParameter StencilPreViewTranslation;
 };
 
+/**
+* A generic vertex shader for projecting a shadow depth buffer onto the scene.
+*/
+class FShadowProjectionVertexShaderInterface : public FGlobalShader
+{
+public:
+	FShadowProjectionVertexShaderInterface() {}
+	FShadowProjectionVertexShaderInterface(const ShaderMetaType::CompiledShaderInitializerType& Initializer)
+		: FGlobalShader(Initializer)
+	{ }
+
+	virtual void SetParameters(FRHICommandList& RHICmdList, const FSceneView& View, const FProjectedShadowInfo* ShadowInfo) = 0;
+};
+
 /**
 * A vertex shader for projecting a shadow depth buffer onto the scene.
 */
-class FShadowProjectionVS : public FGlobalShader
+class FShadowVolumeBoundProjectionVS : public FShadowProjectionVertexShaderInterface
 {
-	DECLARE_SHADER_TYPE(FShadowProjectionVS,Global);
+	DECLARE_SHADER_TYPE(FShadowVolumeBoundProjectionVS,Global);
 public:
 
-	FShadowProjectionVS() {}
-	FShadowProjectionVS(const ShaderMetaType::CompiledShaderInitializerType& Initializer) : FGlobalShader(Initializer) 
+	FShadowVolumeBoundProjectionVS() {}
+	FShadowVolumeBoundProjectionVS(const ShaderMetaType::CompiledShaderInitializerType& Initializer)
+		: FShadowProjectionVertexShaderInterface(Initializer) 
 	{
 		StencilingGeometryParameters.Bind(Initializer.ParameterMap);
 	}
@@ -1208,11 +1223,11 @@ public:
 	
 	static void ModifyCompilationEnvironment(EShaderPlatform Platform, FShaderCompilerEnvironment& OutEnvironment)
 	{
-		FGlobalShader::ModifyCompilationEnvironment(Platform, OutEnvironment);
+		FShadowProjectionVertexShaderInterface::ModifyCompilationEnvironment(Platform, OutEnvironment);
 		OutEnvironment.SetDefine(TEXT("USE_TRANSFORM"), (uint32)1);
 	}
 
-	void SetParameters(FRHICommandList& RHICmdList, const FSceneView& View, const FProjectedShadowInfo* ShadowInfo);
+	void SetParameters(FRHICommandList& RHICmdList, const FSceneView& View, const FProjectedShadowInfo* ShadowInfo) override;
 
 	//~ Begin FShader Interface
 	virtual bool Serialize(FArchive& Ar) override
@@ -1227,13 +1242,13 @@ private:
 	FStencilingGeometryShaderParameters StencilingGeometryParameters;
 };
 
-class FShadowProjectionNoTransformVS : public FGlobalShader
+class FShadowProjectionNoTransformVS : public FShadowProjectionVertexShaderInterface
 {
 	DECLARE_SHADER_TYPE(FShadowProjectionNoTransformVS,Global);
 public:
 	FShadowProjectionNoTransformVS() {}
 	FShadowProjectionNoTransformVS(const ShaderMetaType::CompiledShaderInitializerType& Initializer)
-		: FGlobalShader(Initializer) 
+		: FShadowProjectionVertexShaderInterface(Initializer) 
 	{
 	}
 
@@ -1243,7 +1258,7 @@ public:
 	 */
 	static void ModifyCompilationEnvironment(EShaderPlatform Platform, FShaderCompilerEnvironment& OutEnvironment)
 	{
-		FGlobalShader::ModifyCompilationEnvironment(Platform, OutEnvironment);
+		FShadowProjectionVertexShaderInterface::ModifyCompilationEnvironment(Platform, OutEnvironment);
 		OutEnvironment.SetDefine(TEXT("USE_TRANSFORM"), (uint32)0);
 	}
 
@@ -1252,10 +1267,15 @@ public:
 		return true;
 	}
 
-	inline void SetParameters(FRHICommandList& RHICmdList, const FUniformBufferRHIParamRef ViewUniformBuffer)
+	void SetParameters(FRHICommandList& RHICmdList, const FUniformBufferRHIParamRef ViewUniformBuffer)
 	{
 		FGlobalShader::SetParameters<FViewUniformShaderParameters>(RHICmdList, GetVertexShader(), ViewUniformBuffer);
 	}
+
+	void SetParameters(FRHICommandList& RHICmdList, const FSceneView& View, const FProjectedShadowInfo*) override
+	{
+		FGlobalShader::SetParameters<FViewUniformShaderParameters>(RHICmdList, GetVertexShader(), View.ViewUniformBuffer);
+	}
 };
 
 /**
@@ -1327,22 +1347,24 @@ public:
 
 		const FIntPoint ShadowBufferResolution = ShadowInfo->GetShadowBufferResolution();
 
-		// Set the transform from screen coordinates to shadow depth texture coordinates.
-		if (bModulatedShadows)
+		if (ShadowTileOffsetAndSizeParam.IsBound())
 		{
-			// UE-29083 : work around precision issues with ScreenToShadowMatrix on low end devices.
-			const FMatrix ScreenToShadow = ShadowInfo->GetScreenToShadowMatrix(View, 0, 0, ShadowBufferResolution.X, ShadowBufferResolution.Y);
-			SetShaderValue(RHICmdList, ShaderRHI, ScreenToShadowMatrix, ScreenToShadow);
-
-			FVector2D InverseShadowBufferResolution(1.0f/ShadowBufferResolution.X, 1.0f/ShadowBufferResolution.Y);	
+			FVector2D InverseShadowBufferResolution(1.0f / ShadowBufferResolution.X, 1.0f / ShadowBufferResolution.Y);
 			FVector4 ShadowTileOffsetAndSize(
 				(ShadowInfo->BorderSize + ShadowInfo->X) * InverseShadowBufferResolution.X,
 				(ShadowInfo->BorderSize + ShadowInfo->Y) * InverseShadowBufferResolution.Y,
 				ShadowInfo->ResolutionX * InverseShadowBufferResolution.X,
 				ShadowInfo->ResolutionY * InverseShadowBufferResolution.Y);
-
 			SetShaderValue(RHICmdList, ShaderRHI, ShadowTileOffsetAndSizeParam, ShadowTileOffsetAndSize);
 		}
+
+		// Set the transform from screen coordinates to shadow depth texture coordinates.
+		if (bModulatedShadows)
+		{
+			// UE-29083 : work around precision issues with ScreenToShadowMatrix on low end devices.
+			const FMatrix ScreenToShadow = ShadowInfo->GetScreenToShadowMatrix(View, 0, 0, ShadowBufferResolution.X, ShadowBufferResolution.Y);
+			SetShaderValue(RHICmdList, ShaderRHI, ScreenToShadowMatrix, ScreenToShadow);
+		}
 		else
 		{
 			const FMatrix ScreenToShadow = ShadowInfo->GetScreenToShadowMatrix(View);
@@ -1886,6 +1908,72 @@ struct FShadowProjectionMatrix: FMatrix
 	{}
 };
 
+
+/** Pixel shader to project directional PCSS onto the scene. */
+template<uint32 Quality, bool bUseFadePlane>
+class TDirectionalPercentageCloserShadowProjectionPS : public TShadowProjectionPS<Quality, bUseFadePlane>
+{
+	DECLARE_SHADER_TYPE(TDirectionalPercentageCloserShadowProjectionPS, Global);
+public:
+
+	TDirectionalPercentageCloserShadowProjectionPS() {}
+	TDirectionalPercentageCloserShadowProjectionPS(const ShaderMetaType::CompiledShaderInitializerType& Initializer) :
+		TShadowProjectionPS<Quality, bUseFadePlane>(Initializer)
+	{
+		PCSSParameters.Bind(Initializer.ParameterMap, TEXT("PCSSParameters"));
+	}
+
+	static bool ShouldCache(EShaderPlatform Platform)
+	{
+		return IsFeatureLevelSupported(Platform, ERHIFeatureLevel::SM5) && Platform == SP_PCD3D_SM5;
+	}
+
+	static void ModifyCompilationEnvironment(EShaderPlatform Platform, FShaderCompilerEnvironment& OutEnvironment)
+	{
+		TShadowProjectionPS<Quality, bUseFadePlane>::ModifyCompilationEnvironment(Platform, OutEnvironment);
+		OutEnvironment.SetDefine(TEXT("USE_PCSS"), 1);
+	}
+
+	virtual void SetParameters(
+		FRHICommandList& RHICmdList,
+		int32 ViewIndex,
+		const FSceneView& View,
+		const FProjectedShadowInfo* ShadowInfo) override
+	{
+		TShadowProjectionPS<Quality, bUseFadePlane>::SetParameters(RHICmdList, ViewIndex, View, ShadowInfo);
+
+		const FPixelShaderRHIParamRef ShaderRHI = this->GetPixelShader();
+
+		// GetLightSourceAngle returns the full angle.
+		float TanLightSourceAngle = FMath::Tan(0.5 * FMath::DegreesToRadians(ShadowInfo->GetLightSceneInfo().Proxy->GetLightSourceAngle()));
+
+		static IConsoleVariable* CVarMaxSoftShadowKernelSize = IConsoleManager::Get().FindConsoleVariable(TEXT("r.Shadow.MaxSoftKernelSize"));
+		check(CVarMaxSoftShadowKernelSize);
+		int32 MaxKernelSize = CVarMaxSoftShadowKernelSize->GetInt();
+
+		float SW = 2.0 * ShadowInfo->ShadowBounds.W;
+		float SZ = ShadowInfo->MaxSubjectZ - ShadowInfo->MinSubjectZ;
+
+		FVector4 PCSSParameterValues = FVector4(TanLightSourceAngle * SZ / SW, MaxKernelSize / float(ShadowInfo->ResolutionX), 0, 0);
+		SetShaderValue(RHICmdList, ShaderRHI, PCSSParameters, PCSSParameterValues);
+	}
+
+	/**
+	* Serialize the parameters for this shader
+	* @param Ar - archive to serialize to
+	*/
+	virtual bool Serialize(FArchive& Ar) override
+	{
+		bool bShaderHasOutdatedParameters = TShadowProjectionPS<Quality, bUseFadePlane>::Serialize(Ar);
+		Ar << PCSSParameters;
+		return bShaderHasOutdatedParameters;
+	}
+
+protected:
+	FShaderParameter PCSSParameters;
+};
+
+
 // Sort by descending resolution
 struct FCompareFProjectedShadowInfoByResolution
 {
diff --git a/Engine/Source/Runtime/Renderer/Private/ShadowSetup.cpp b/Engine/Source/Runtime/Renderer/Private/ShadowSetup.cpp
--- a/Engine/Source/Runtime/Renderer/Private/ShadowSetup.cpp
+++ b/Engine/Source/Runtime/Renderer/Private/ShadowSetup.cpp
@@ -193,6 +193,28 @@ static TAutoConsoleVariable<int32> CVarEnableCsmShaderCulling(
 	TEXT(""),
 	ECVF_RenderThreadSafe);
 
+static TAutoConsoleVariable<int32> CVarParallelGatherShadowPrimitives(
+	TEXT("r.ParallelGatherShadowPrimitives"),
+	1,  
+	TEXT("Toggles parallel Gather shadow primitives. 0 = off; 1 = on"),
+	ECVF_RenderThreadSafe
+	);
+
+static TAutoConsoleVariable<int32> CVarParallelGatherNumPrimitivesPerPacket(
+	TEXT("r.ParallelGatherNumPrimitivesPerPacket"),
+	256,  
+	TEXT("Number of primitives per packet.  Only used when r.Shadow.UseOctreeForCulling is disabled."),
+	ECVF_RenderThreadSafe
+	);
+
+int32 GUseOctreeForShadowCulling = 1;
+FAutoConsoleVariableRef CVarUseOctreeForShadowCulling(
+	TEXT("r.Shadow.UseOctreeForCulling"),
+	GUseOctreeForShadowCulling,
+	TEXT("Whether to use the primitive octree for shadow subject culling.  The octree culls large groups of primitives at a time, but introduces cache misses walking the data structure."),
+	ECVF_Scalability | ECVF_RenderThreadSafe
+	);
+
 #if !UE_BUILD_SHIPPING
 // read and written on the render thread
 bool GDumpShadowSetup = false;
@@ -739,7 +761,7 @@ void FProjectedShadowInfo::SetupWholeSceneProjection(
 	UpdateShaderDepthBias();
 }
 
-void FProjectedShadowInfo::AddSubjectPrimitive(FPrimitiveSceneInfo* PrimitiveSceneInfo, TArray<FViewInfo>* ViewArray, bool bRecordShadowSubjectsForMobileShading)
+void FProjectedShadowInfo::AddSubjectPrimitive(FPrimitiveSceneInfo* PrimitiveSceneInfo, TArray<FViewInfo>* ViewArray, ERHIFeatureLevel::Type FeatureLevel, bool bRecordShadowSubjectsForMobileShading)
 {
 	QUICK_SCOPE_CYCLE_COUNTER(STAT_AddSubjectPrimitive);
 
@@ -774,7 +796,6 @@ void FProjectedShadowInfo::AddSubjectPrimitive(FPrimitiveSceneInfo* PrimitiveSce
 		bool bShadowRelevance = false;
 		uint32 ViewMask = 0;
 		int32 PrimitiveId = PrimitiveSceneInfo->GetIndex();
-		const auto FeatureLevel = PrimitiveSceneInfo->Scene->GetFeatureLevel();
 
 		for (int32 ViewIndex = 0, Num = Views.Num(); ViewIndex < Num; ViewIndex++)
 		{
@@ -883,9 +904,11 @@ void FProjectedShadowInfo::AddSubjectPrimitive(FPrimitiveSceneInfo* PrimitiveSce
 								if (bMeshIsVisible && bWholeSceneDirectionalShadow)
 								{
 									StaticMeshWholeSceneShadowDepthMap[StaticMesh.Id] = true;
-									StaticMeshWholeSceneShadowBatchVisibility[StaticMesh.Id] = StaticMesh.bRequiresPerElementVisibility
-										? StaticMesh.VertexFactory->GetStaticBatchElementShadowVisibility(CurrentView, LightSceneInfo->Proxy, &StaticMesh)
-										: ((1ull << StaticMesh.Elements.Num()) - 1);
+
+									if (StaticMesh.bRequiresPerElementVisibility)
+									{
+										StaticMeshWholeSceneShadowBatchVisibility[StaticMesh.BatchVisibilityId] = StaticMesh.VertexFactory->GetStaticBatchElementVisibility(CurrentView, &StaticMesh);
+									}
 								}
 							}
 						}
@@ -918,10 +941,7 @@ void FProjectedShadowInfo::AddSubjectPrimitive(FPrimitiveSceneInfo* PrimitiveSce
 							}
 							else
 							{
-								FPrimitiveBounds PrimitiveBounds;
-								PrimitiveBounds.Origin = Bounds.Origin;
-								PrimitiveBounds.SphereRadius = Bounds.SphereRadius;
-								LODToRender = ComputeLODForMeshes(PrimitiveSceneInfo->StaticMeshes, CurrentView, PrimitiveBounds.Origin, PrimitiveBounds.SphereRadius, ForcedLODLevel);
+								LODToRender = ComputeLODForMeshes(PrimitiveSceneInfo->StaticMeshes, CurrentView, Bounds.Origin, Bounds.SphereRadius, ForcedLODLevel);
 							}
 
 							if (bWholeSceneDirectionalShadow)
@@ -932,9 +952,11 @@ void FProjectedShadowInfo::AddSubjectPrimitive(FPrimitiveSceneInfo* PrimitiveSce
 									if (StaticMesh.CastShadow && LODToRender.ContainsLOD(StaticMesh.LODIndex))
 									{
 										StaticMeshWholeSceneShadowDepthMap[StaticMesh.Id] = true;
-										StaticMeshWholeSceneShadowBatchVisibility[StaticMesh.Id] = StaticMesh.bRequiresPerElementVisibility
-											? StaticMesh.VertexFactory->GetStaticBatchElementShadowVisibility(CurrentView, LightSceneInfo->Proxy, &StaticMesh)
-											: ((1ull << StaticMesh.Elements.Num()) - 1);
+
+										if (StaticMesh.bRequiresPerElementVisibility)
+										{
+											StaticMeshWholeSceneShadowBatchVisibility[StaticMesh.BatchVisibilityId] = StaticMesh.VertexFactory->GetStaticBatchElementVisibility(CurrentView, &StaticMesh);
+										}
 
 										bDrawingStaticMeshes = true;
 									}
@@ -948,9 +970,11 @@ void FProjectedShadowInfo::AddSubjectPrimitive(FPrimitiveSceneInfo* PrimitiveSce
 									if (StaticMesh.CastShadow && LODToRender.ContainsLOD(StaticMesh.LODIndex))
 									{
 										CurrentView.StaticMeshShadowDepthMap[StaticMesh.Id] = true;
-										CurrentView.StaticMeshBatchVisibility[StaticMesh.Id] = StaticMesh.bRequiresPerElementVisibility
-											? StaticMesh.VertexFactory->GetStaticBatchElementShadowVisibility(CurrentView, LightSceneInfo->Proxy, &StaticMesh)
-											: ((1ull << StaticMesh.Elements.Num()) - 1);
+
+										if (StaticMesh.bRequiresPerElementVisibility)
+										{
+											CurrentView.StaticMeshBatchVisibility[StaticMesh.BatchVisibilityId] = StaticMesh.VertexFactory->GetStaticBatchElementVisibility(CurrentView, &StaticMesh);
+										}
 
 										bDrawingStaticMeshes = true;
 									}
@@ -1315,7 +1339,7 @@ void FSceneRenderer::UpdatePreshadowCache(FSceneRenderTargets& SceneContext)
 	}
 }
 
-bool FSceneRenderer::ShouldCreateObjectShadowForStationaryLight(const FLightSceneInfo* LightSceneInfo, const FPrimitiveSceneProxy* PrimitiveSceneProxy, bool bInteractionShadowMapped) const
+bool ShouldCreateObjectShadowForStationaryLight(const FLightSceneInfo* LightSceneInfo, const FPrimitiveSceneProxy* PrimitiveSceneProxy, bool bInteractionShadowMapped) 
 {
 	const bool bCreateObjectShadowForStationaryLight = 
 		LightSceneInfo->bCreatePerObjectShadowsForDynamicObjects 
@@ -1626,7 +1650,7 @@ void FSceneRenderer::CreatePerObjectProjectedShadow(
 						for (int32 ChildIndex = 0, ChildCount = ShadowGroupPrimitives.Num(); ChildIndex < ChildCount; ChildIndex++)
 						{
 							FPrimitiveSceneInfo* ShadowChild = ShadowGroupPrimitives[ChildIndex];
-							ProjectedShadowInfo->AddSubjectPrimitive(ShadowChild, &Views, false);
+							ProjectedShadowInfo->AddSubjectPrimitive(ShadowChild, &Views, FeatureLevel, false);
 						}
 					}
 					else if (bShadowIsPotentiallyVisibleNextFrame)
@@ -1666,7 +1690,7 @@ void FSceneRenderer::CreatePerObjectProjectedShadow(
 						for (int32 ChildIndex = 0, ChildCount = ShadowGroupPrimitives.Num(); ChildIndex < ChildCount; ChildIndex++)
 						{
 							FPrimitiveSceneInfo* ShadowChild = ShadowGroupPrimitives[ChildIndex];
-							ProjectedShadowInfo->AddSubjectPrimitive(ShadowChild, &Views, false);
+							ProjectedShadowInfo->AddSubjectPrimitive(ShadowChild, &Views, FeatureLevel, false);
 						}
 					}
 					else if (bShadowIsPotentiallyVisibleNextFrame)
@@ -2067,7 +2091,7 @@ void FSceneRenderer::CreateWholeSceneProjectedShadow(FLightSceneInfo* LightScene
 									&& !Interaction->CastsSelfShadowOnly()
 									&& (!bStaticSceneOnly || Interaction->GetPrimitiveSceneInfo()->Proxy->HasStaticLighting()))
 								{
-									ProjectedShadowInfo->AddSubjectPrimitive(Interaction->GetPrimitiveSceneInfo(), &Views, false);
+									ProjectedShadowInfo->AddSubjectPrimitive(Interaction->GetPrimitiveSceneInfo(), &Views, FeatureLevel, false);
 								}
 							}
 						}
@@ -2084,7 +2108,7 @@ void FSceneRenderer::CreateWholeSceneProjectedShadow(FLightSceneInfo* LightScene
 									&& !Interaction->CastsSelfShadowOnly()
 									&& (!bStaticSceneOnly || Interaction->GetPrimitiveSceneInfo()->Proxy->HasStaticLighting()))
 								{
-									ProjectedShadowInfo->AddSubjectPrimitive(Interaction->GetPrimitiveSceneInfo(), &Views, false);
+									ProjectedShadowInfo->AddSubjectPrimitive(Interaction->GetPrimitiveSceneInfo(), &Views, FeatureLevel, false);
 								}
 							}
 						}
@@ -2366,124 +2390,193 @@ void FSceneRenderer::GatherShadowDynamicMeshElements()
 	}
 }
 
-inline void FSceneRenderer::GatherShadowsForPrimitiveInner(
-	const FPrimitiveSceneInfoCompact& PrimitiveSceneInfoCompact, 
-	const TArray<FProjectedShadowInfo*,SceneRenderingAllocator>& PreShadows,
-	const TArray<FProjectedShadowInfo*,SceneRenderingAllocator>& ViewDependentWholeSceneShadows,
-	bool bStaticSceneOnly)
+typedef TArray<FPrimitiveSceneInfo*> FShadowSubjectPrimitives;
+
+struct FGatherShadowPrimitivesPacket
 {
-	if(PrimitiveSceneInfoCompact.bCastDynamicShadow)
+	// Inputs
+	const FScene* Scene;
+	TArray<FViewInfo>& Views;
+	const FScenePrimitiveOctree::FNode* Node;
+	int32 StartPrimitiveIndex;
+	int32 NumPrimitives;
+	const TArray<FProjectedShadowInfo*, SceneRenderingAllocator>& PreShadows;
+	const TArray<FProjectedShadowInfo*, SceneRenderingAllocator>& ViewDependentWholeSceneShadows;
+	ERHIFeatureLevel::Type FeatureLevel;
+	bool bStaticSceneOnly;
+
+	// Outputs
+	TArray<FShadowSubjectPrimitives, SceneRenderingAllocator> PreShadowSubjectPrimitives;
+	TArray<FShadowSubjectPrimitives, SceneRenderingAllocator> ViewDependentWholeSceneShadowSubjectPrimitives;
+
+	FGatherShadowPrimitivesPacket(
+		const FScene* InScene,
+		TArray<FViewInfo>& InViews,
+		const FScenePrimitiveOctree::FNode* InNode,
+		int32 InStartPrimitiveIndex,
+		int32 InNumPrimitives,
+		const TArray<FProjectedShadowInfo*, SceneRenderingAllocator>& InPreShadows,
+		const TArray<FProjectedShadowInfo*, SceneRenderingAllocator>& InViewDependentWholeSceneShadows,
+		ERHIFeatureLevel::Type InFeatureLevel,
+		bool bInStaticSceneOnly) 
+		: Scene(InScene)
+		, Views(InViews)
+		, Node(InNode)
+		, StartPrimitiveIndex(InStartPrimitiveIndex)
+		, NumPrimitives(InNumPrimitives)
+		, PreShadows(InPreShadows)
+		, ViewDependentWholeSceneShadows(InViewDependentWholeSceneShadows)
+		, FeatureLevel(InFeatureLevel)
+		, bStaticSceneOnly(bInStaticSceneOnly)
+	{
+		PreShadowSubjectPrimitives.Empty(PreShadows.Num());
+		PreShadowSubjectPrimitives.AddDefaulted(PreShadows.Num());
+
+		ViewDependentWholeSceneShadowSubjectPrimitives.Empty(ViewDependentWholeSceneShadows.Num());
+		ViewDependentWholeSceneShadowSubjectPrimitives.AddDefaulted(ViewDependentWholeSceneShadows.Num());
+	}
+
+	void AnyThreadTask()
 	{
-		FPrimitiveSceneInfo* RESTRICT	PrimitiveSceneInfo	= PrimitiveSceneInfoCompact.PrimitiveSceneInfo;
-		FPrimitiveSceneProxy* RESTRICT	PrimitiveProxy		= PrimitiveSceneInfoCompact.Proxy;
-		const FBoxSphereBounds&			PrimitiveBounds		= PrimitiveSceneInfoCompact.Bounds;
+		if (Node)
+		{
+			// Check all the primitives in this octree node.
+			for (FScenePrimitiveOctree::ElementConstIt NodePrimitiveIt(Node->GetElementIt()); NodePrimitiveIt; ++NodePrimitiveIt)
+			{
+				if (NodePrimitiveIt->PrimitiveFlagsCompact.bCastDynamicShadow)
+				{
+					FilterPrimitiveForShadows(NodePrimitiveIt->Bounds, NodePrimitiveIt->PrimitiveFlagsCompact, NodePrimitiveIt->PrimitiveSceneInfo, NodePrimitiveIt->Proxy);
+				}
+			}
+		}
+		else
+		{
+			check(NumPrimitives > 0);
+
+			// Check primitives in this packet's range
+			for (int32 PrimitiveIndex = StartPrimitiveIndex; PrimitiveIndex < StartPrimitiveIndex + NumPrimitives; PrimitiveIndex++)
+			{
+				FPrimitiveFlagsCompact PrimitiveFlagsCompact = Scene->PrimitiveFlagsCompact[PrimitiveIndex];
+
+				if (PrimitiveFlagsCompact.bCastDynamicShadow)
+				{
+					FilterPrimitiveForShadows(Scene->PrimitiveBounds[PrimitiveIndex].BoxSphereBounds, PrimitiveFlagsCompact, Scene->Primitives[PrimitiveIndex], Scene->PrimitiveSceneProxies[PrimitiveIndex]); 
+				}
+			}
+		}
+	}
 
+	void FilterPrimitiveForShadows(const FBoxSphereBounds& PrimitiveBounds, FPrimitiveFlagsCompact PrimitiveFlagsCompact, FPrimitiveSceneInfo* PrimitiveSceneInfo, FPrimitiveSceneProxy* PrimitiveProxy)
+	{
 		// Check if the primitive is a subject for any of the preshadows.
 		// Only allow preshadows from lightmapped primitives that cast both dynamic and static shadows.
-		if (PreShadows.Num() && PrimitiveProxy->CastsStaticShadow() && PrimitiveProxy->HasStaticLighting())
+		if (PreShadows.Num() && PrimitiveFlagsCompact.bCastStaticShadow && PrimitiveFlagsCompact.bStaticLighting)
 		{
-			for( int32 ShadowIndex = 0, Num = PreShadows.Num(); ShadowIndex < PreShadows.Num(); ShadowIndex++ )
+			for (int32 ShadowIndex = 0, Num = PreShadows.Num(); ShadowIndex < Num; ShadowIndex++)
 			{
 				FProjectedShadowInfo* RESTRICT ProjectedShadowInfo = PreShadows[ShadowIndex];
 
+				// Note: Culling based on the primitive's bounds BEFORE dereferencing PrimitiveSceneInfo / PrimitiveProxy
 				// Check if this primitive is in the shadow's frustum.
-				bool bInFrustum = ProjectedShadowInfo->CasterFrustum.IntersectBox( PrimitiveBounds.Origin, ProjectedShadowInfo->PreShadowTranslation, PrimitiveBounds.BoxExtent );
+				bool bInFrustum = ProjectedShadowInfo->CasterFrustum.IntersectBox(PrimitiveBounds.Origin, ProjectedShadowInfo->PreShadowTranslation, PrimitiveBounds.BoxExtent);
 
-				if( bInFrustum && ProjectedShadowInfo->GetLightSceneInfoCompact().AffectsPrimitive(PrimitiveSceneInfoCompact) )
+				if (bInFrustum && ProjectedShadowInfo->GetLightSceneInfoCompact().AffectsPrimitive(PrimitiveBounds, PrimitiveProxy))
 				{
-					// Add this primitive to the shadow.
-					ProjectedShadowInfo->AddSubjectPrimitive(PrimitiveSceneInfo, &Views, false);
+					PreShadowSubjectPrimitives[ShadowIndex].Add(PrimitiveSceneInfo);
 				}
 			}
 		}
 
-		if(PrimitiveSceneInfoCompact.bCastDynamicShadow || PrimitiveSceneInfoCompact.bAffectDynamicIndirectLighting )
+		for (int32 ShadowIndex = 0, Num = ViewDependentWholeSceneShadows.Num();ShadowIndex < Num;ShadowIndex++)
 		{
-			for(int32 ShadowIndex = 0, Num = ViewDependentWholeSceneShadows.Num();ShadowIndex < Num;ShadowIndex++)
+			const FProjectedShadowInfo* RESTRICT ProjectedShadowInfo = ViewDependentWholeSceneShadows[ShadowIndex];
+			const FLightSceneInfo& RESTRICT LightSceneInfo = ProjectedShadowInfo->GetLightSceneInfo();
+			const FLightSceneProxy& RESTRICT LightProxy = *LightSceneInfo.Proxy;
+
+			const FVector LightDirection = LightProxy.GetDirection();
+			const FVector PrimitiveToShadowCenter = ProjectedShadowInfo->ShadowBounds.Center - PrimitiveBounds.Origin;
+			// Project the primitive's bounds origin onto the light vector
+			const float ProjectedDistanceFromShadowOriginAlongLightDir = PrimitiveToShadowCenter | LightDirection;
+			// Calculate the primitive's squared distance to the cylinder's axis
+			const float PrimitiveDistanceFromCylinderAxisSq = (-LightDirection * ProjectedDistanceFromShadowOriginAlongLightDir + PrimitiveToShadowCenter).SizeSquared();
+			const float CombinedRadiusSq = FMath::Square(ProjectedShadowInfo->ShadowBounds.W + PrimitiveBounds.SphereRadius);
+
+			// Note: Culling based on the primitive's bounds BEFORE dereferencing PrimitiveSceneInfo / PrimitiveProxy
+
+			// Check if this primitive is in the shadow's cylinder
+			if (PrimitiveDistanceFromCylinderAxisSq < CombinedRadiusSq
+				// If the primitive is further along the cone axis than the shadow bounds origin, 
+				// Check if the primitive is inside the spherical cap of the cascade's bounds
+				&& !(ProjectedDistanceFromShadowOriginAlongLightDir < 0 && PrimitiveToShadowCenter.SizeSquared() > CombinedRadiusSq)
+				// Test against the convex hull containing the extruded shadow bounds
+				&& ProjectedShadowInfo->CascadeSettings.ShadowBoundsAccurate.IntersectBox(PrimitiveBounds.Origin, PrimitiveBounds.BoxExtent))
 			{
-				FProjectedShadowInfo* RESTRICT ProjectedShadowInfo = ViewDependentWholeSceneShadows[ShadowIndex];
+				// Distance culling for RSMs
+				const float MinScreenRadiusForShadowCaster = ProjectedShadowInfo->bReflectiveShadowmap ? GMinScreenRadiusForShadowCasterRSM : GMinScreenRadiusForShadowCaster;
+
+				bool bScreenSpaceSizeCulled = false;
+				check(ProjectedShadowInfo->DependentView);
 
-				if ( ProjectedShadowInfo->bReflectiveShadowmap && !PrimitiveSceneInfoCompact.bAffectDynamicIndirectLighting )
 				{
-					continue;
+					const float DistanceSquared = (PrimitiveBounds.Origin - ProjectedShadowInfo->DependentView->ShadowViewMatrices.GetViewOrigin()).SizeSquared();
+					bScreenSpaceSizeCulled = FMath::Square(PrimitiveBounds.SphereRadius) < FMath::Square(MinScreenRadiusForShadowCaster) * DistanceSquared * ProjectedShadowInfo->DependentView->LODDistanceFactorSquared;
 				}
-				if ( !ProjectedShadowInfo->bReflectiveShadowmap && !PrimitiveSceneInfoCompact.bCastDynamicShadow )
+
+				if (!bScreenSpaceSizeCulled
+					&& ProjectedShadowInfo->GetLightSceneInfoCompact().AffectsPrimitive(PrimitiveBounds, PrimitiveProxy)
+					// Include all primitives for movable lights, but only statically shadowed primitives from a light with static shadowing,
+					// Since lights with static shadowing still create per-object shadows for primitives without static shadowing.
+					&& (!LightProxy.HasStaticLighting() || (!LightSceneInfo.IsPrecomputedLightingValid() || LightProxy.UseCSMForDynamicObjects()))
+					// Only render primitives into a reflective shadowmap that are supposed to affect indirect lighting
+					&& !(ProjectedShadowInfo->bReflectiveShadowmap && !PrimitiveProxy->AffectsDynamicIndirectLighting())
+					// Exclude primitives that will create their own per-object shadow, except when rendering RSMs
+					&& (!PrimitiveProxy->CastsInsetShadow() || ProjectedShadowInfo->bReflectiveShadowmap)
+					// Exclude primitives that will create a per-object shadow from a stationary light
+					&& !ShouldCreateObjectShadowForStationaryLight(&LightSceneInfo, PrimitiveProxy, true)
+					// Only render shadows from objects that use static lighting during a reflection capture, since the reflection capture doesn't update at runtime
+					&& (!bStaticSceneOnly || PrimitiveProxy->HasStaticLighting())
+					// Render dynamic lit objects if CSMForDynamicObjects is enabled.
+					&& (!LightProxy.UseCSMForDynamicObjects() || !PrimitiveProxy->HasStaticLighting()))
 				{
-					continue;
+					ViewDependentWholeSceneShadowSubjectPrimitives[ShadowIndex].Add(PrimitiveSceneInfo);
 				}
+			}
+		}
+	}
 
-				FLightSceneProxy* RESTRICT LightProxy = ProjectedShadowInfo->GetLightSceneInfo().Proxy;
-
-				const FVector LightDirection = LightProxy->GetDirection();
-				const FVector PrimitiveToShadowCenter = ProjectedShadowInfo->ShadowBounds.Center - PrimitiveBounds.Origin;
-				// Project the primitive's bounds origin onto the light vector
-				const float ProjectedDistanceFromShadowOriginAlongLightDir = PrimitiveToShadowCenter | LightDirection;
-				// Calculate the primitive's squared distance to the cylinder's axis
-				const float PrimitiveDistanceFromCylinderAxisSq = (-LightDirection * ProjectedDistanceFromShadowOriginAlongLightDir + PrimitiveToShadowCenter).SizeSquared();
-				const float CombinedRadiusSq = FMath::Square(ProjectedShadowInfo->ShadowBounds.W + PrimitiveBounds.SphereRadius);
-
-				// Include all primitives for movable lights, but only statically shadowed primitives from a light with static shadowing,
-				// Since lights with static shadowing still create per-object shadows for primitives without static shadowing.
-				if( (!LightProxy->HasStaticLighting() || (!ProjectedShadowInfo->GetLightSceneInfo().IsPrecomputedLightingValid() || LightProxy->UseCSMForDynamicObjects()))
-					// Check if this primitive is in the shadow's cylinder
-					&& PrimitiveDistanceFromCylinderAxisSq < CombinedRadiusSq
-					// Check if the primitive is closer than the cylinder cap toward the light
-					// next line is commented as it breaks large world shadows, if this was meant to be an optimization we should think about a better solution
-					//// && ProjectedDistanceFromShadowOriginAlongLightDir - PrimitiveBounds.SphereRadius < -ProjectedShadowInfo->MinPreSubjectZ
-					// If the primitive is further along the cone axis than the shadow bounds origin, 
-					// Check if the primitive is inside the spherical cap of the cascade's bounds
-					&& !(ProjectedDistanceFromShadowOriginAlongLightDir < 0 
-						&& PrimitiveToShadowCenter.SizeSquared() > CombinedRadiusSq))
-				{
-					const bool bInFrustum = ProjectedShadowInfo->CascadeSettings.ShadowBoundsAccurate.IntersectBox( PrimitiveBounds.Origin, PrimitiveBounds.BoxExtent );
+	void RenderThreadFinalize()
+	{
+		for (int32 ShadowIndex = 0; ShadowIndex < PreShadowSubjectPrimitives.Num(); ShadowIndex++)
+		{
+			FProjectedShadowInfo* ProjectedShadowInfo = PreShadows[ShadowIndex];
 
-					if( bInFrustum )
-					{
-						// Distance culling for RSMs
-						float MinScreenRadiusForShadowCaster = GMinScreenRadiusForShadowCaster;
-						if (ProjectedShadowInfo->bReflectiveShadowmap)
-						{
-							MinScreenRadiusForShadowCaster = GMinScreenRadiusForShadowCasterRSM;
-						}
+			for (int32 PrimitiveIndex = 0; PrimitiveIndex < PreShadowSubjectPrimitives[ShadowIndex].Num(); PrimitiveIndex++)
+			{
+				ProjectedShadowInfo->AddSubjectPrimitive(PreShadowSubjectPrimitives[ShadowIndex][PrimitiveIndex], &Views, FeatureLevel, false);
+			}
+		}
 
-						bool bScreenSpaceSizeCulled = false;
-						check( ProjectedShadowInfo->DependentView );
-						if ( ProjectedShadowInfo->DependentView ) 
-						{
-							const float DistanceSquared = ( PrimitiveBounds.Origin - ProjectedShadowInfo->DependentView->ShadowViewMatrices.GetViewOrigin() ).SizeSquared();
-							bScreenSpaceSizeCulled = FMath::Square( PrimitiveBounds.SphereRadius ) < FMath::Square( MinScreenRadiusForShadowCaster ) * DistanceSquared * ProjectedShadowInfo->DependentView->LODDistanceFactorSquared;
-						}
+		for (int32 ShadowIndex = 0; ShadowIndex < ViewDependentWholeSceneShadowSubjectPrimitives.Num(); ShadowIndex++)
+		{
+			FProjectedShadowInfo* ProjectedShadowInfo = ViewDependentWholeSceneShadows[ShadowIndex];
 
-						if (ProjectedShadowInfo->GetLightSceneInfoCompact().AffectsPrimitive(PrimitiveSceneInfoCompact)
-							// Exclude primitives that will create their own per-object shadow, except when rendering RSMs
-							&& ( !PrimitiveProxy->CastsInsetShadow() || ProjectedShadowInfo->bReflectiveShadowmap )
-							// Exclude primitives that will create a per-object shadow from a stationary light
-							&& !ShouldCreateObjectShadowForStationaryLight(&ProjectedShadowInfo->GetLightSceneInfo(), PrimitiveSceneInfo->Proxy, true)
-							// Only render shadows from objects that use static lighting during a reflection capture, since the reflection capture doesn't update at runtime
-							&& (!bStaticSceneOnly || PrimitiveProxy->HasStaticLighting())
-							// Render dynamic lit objects if CSMForDynamicObjects is enabled.
-							&& (!LightProxy->UseCSMForDynamicObjects() || !PrimitiveProxy->HasStaticLighting())
-							&& !bScreenSpaceSizeCulled )
-						{
-							bool bRecordShadowSubjectsForMobile = false;
+			bool bRecordShadowSubjectsForMobile = false;
 							
-							if (Scene->GetShadingPath() == EShadingPath::Mobile)
-							{
-								static auto* CVarMobileEnableStaticAndCSMShadowReceivers = IConsoleManager::Get().FindTConsoleVariableDataInt(TEXT("r.Mobile.EnableStaticAndCSMShadowReceivers"));
-								bRecordShadowSubjectsForMobile = CVarEnableCsmShaderCulling.GetValueOnRenderThread() 
-									&& CVarMobileEnableStaticAndCSMShadowReceivers->GetValueOnRenderThread()
-									&& LightProxy->UseCSMForDynamicObjects();
-							}
+			if (FSceneInterface::GetShadingPath(FeatureLevel) == EShadingPath::Mobile)
+			{
+				static auto* CVarMobileEnableStaticAndCSMShadowReceivers = IConsoleManager::Get().FindTConsoleVariableDataInt(TEXT("r.Mobile.EnableStaticAndCSMShadowReceivers"));
+				bRecordShadowSubjectsForMobile = CVarEnableCsmShaderCulling.GetValueOnRenderThread() 
+					&& CVarMobileEnableStaticAndCSMShadowReceivers->GetValueOnRenderThread()
+					&& ProjectedShadowInfo->GetLightSceneInfo().Proxy->UseCSMForDynamicObjects();
+			}
 
-							// Add this primitive to the shadow.
-							ProjectedShadowInfo->AddSubjectPrimitive(PrimitiveSceneInfo, NULL, bRecordShadowSubjectsForMobile);
-						}
-					}
-				}
+			for (int32 PrimitiveIndex = 0; PrimitiveIndex < ViewDependentWholeSceneShadowSubjectPrimitives[ShadowIndex].Num(); PrimitiveIndex++)
+			{
+				ProjectedShadowInfo->AddSubjectPrimitive(ViewDependentWholeSceneShadowSubjectPrimitives[ShadowIndex][PrimitiveIndex], NULL, FeatureLevel, bRecordShadowSubjectsForMobile);
 			}
 		}
 	}
-}
+};
 
 void FSceneRenderer::GatherShadowPrimitives(
 	const TArray<FProjectedShadowInfo*,SceneRenderingAllocator>& PreShadows,
@@ -2493,94 +2586,140 @@ void FSceneRenderer::GatherShadowPrimitives(
 {
 	SCOPE_CYCLE_COUNTER(STAT_GatherShadowPrimitivesTime);
 
-	if(PreShadows.Num() || ViewDependentWholeSceneShadows.Num())
+	if (PreShadows.Num() || ViewDependentWholeSceneShadows.Num())
 	{
-		for(int32 ShadowIndex = 0, Num = ViewDependentWholeSceneShadows.Num(); ShadowIndex < Num;ShadowIndex++)
+		for (int32 ShadowIndex = 0, Num = ViewDependentWholeSceneShadows.Num(); ShadowIndex < Num;ShadowIndex++)
 		{
 			FProjectedShadowInfo* ProjectedShadowInfo = ViewDependentWholeSceneShadows[ShadowIndex];
 			checkSlow(ProjectedShadowInfo->DependentView);
 			// Initialize the whole scene shadow's depth map with the shadow independent depth map from the view
 			ProjectedShadowInfo->StaticMeshWholeSceneShadowDepthMap.Init(false,Scene->StaticMeshes.GetMaxIndex());
-			ProjectedShadowInfo->StaticMeshWholeSceneShadowBatchVisibility.AddZeroed(Scene->StaticMeshes.GetMaxIndex());
+			ProjectedShadowInfo->StaticMeshWholeSceneShadowBatchVisibility.AddZeroed(Scene->StaticMeshBatchVisibility.GetMaxIndex());
 		}
 
-		QUICK_SCOPE_CYCLE_COUNTER(STAT_ShadowOctreeTraversal);
-		// Find primitives that are in a shadow frustum in the octree.
-		for(FScenePrimitiveOctree::TConstIterator<SceneRenderingAllocator> PrimitiveOctreeIt(Scene->PrimitiveOctree);
-			PrimitiveOctreeIt.HasPendingNodes();
-			PrimitiveOctreeIt.Advance())
+		TArray<FGatherShadowPrimitivesPacket*,SceneRenderingAllocator> Packets;
+
+		if (GUseOctreeForShadowCulling)
 		{
-			const FScenePrimitiveOctree::FNode& PrimitiveOctreeNode = PrimitiveOctreeIt.GetCurrentNode();
-			const FOctreeNodeContext& PrimitiveOctreeNodeContext = PrimitiveOctreeIt.GetCurrentContext();
+			QUICK_SCOPE_CYCLE_COUNTER(STAT_ShadowSceneOctreeTraversal);
 
+			Packets.Reserve(100);
+
+			// Find primitives that are in a shadow frustum in the octree.
+			for(FScenePrimitiveOctree::TConstIterator<SceneRenderingAllocator> PrimitiveOctreeIt(Scene->PrimitiveOctree);
+				PrimitiveOctreeIt.HasPendingNodes();
+				PrimitiveOctreeIt.Advance())
 			{
-				// Find children of this octree node that may contain relevant primitives.
-				FOREACH_OCTREE_CHILD_NODE(ChildRef)
+				const FScenePrimitiveOctree::FNode& PrimitiveOctreeNode = PrimitiveOctreeIt.GetCurrentNode();
+				const FOctreeNodeContext& PrimitiveOctreeNodeContext = PrimitiveOctreeIt.GetCurrentContext();
+
 				{
-					if(PrimitiveOctreeNode.HasChild(ChildRef))
+					// Find children of this octree node that may contain relevant primitives.
+					FOREACH_OCTREE_CHILD_NODE(ChildRef)
 					{
-						// Check that the child node is in the frustum for at least one shadow.
-						const FOctreeNodeContext ChildContext = PrimitiveOctreeNodeContext.GetChildContext(ChildRef);
-						bool bIsInFrustum = false;
-
-						// Check for subjects of preshadows.
-						if(!bIsInFrustum)
+						if(PrimitiveOctreeNode.HasChild(ChildRef))
 						{
-							for(int32 ShadowIndex = 0, Num = PreShadows.Num(); ShadowIndex < Num; ShadowIndex++)
+							// Check that the child node is in the frustum for at least one shadow.
+							const FOctreeNodeContext ChildContext = PrimitiveOctreeNodeContext.GetChildContext(ChildRef);
+							bool bIsInFrustum = false;
+
+							// Check for subjects of preshadows.
+							if(!bIsInFrustum)
 							{
-								FProjectedShadowInfo* ProjectedShadowInfo = PreShadows[ShadowIndex];
-
-								check(ProjectedShadowInfo->CasterFrustum.PermutedPlanes.Num());
-								// Check if this primitive is in the shadow's frustum.
-								if(ProjectedShadowInfo->CasterFrustum.IntersectBox(
-									ChildContext.Bounds.Center + ProjectedShadowInfo->PreShadowTranslation,
-									ChildContext.Bounds.Extent
-									))
+								for(int32 ShadowIndex = 0, Num = PreShadows.Num(); ShadowIndex < Num; ShadowIndex++)
 								{
-									bIsInFrustum = true;
-									break;
+									FProjectedShadowInfo* ProjectedShadowInfo = PreShadows[ShadowIndex];
+
+									check(ProjectedShadowInfo->CasterFrustum.PermutedPlanes.Num());
+									// Check if this primitive is in the shadow's frustum.
+									if(ProjectedShadowInfo->CasterFrustum.IntersectBox(
+										ChildContext.Bounds.Center + ProjectedShadowInfo->PreShadowTranslation,
+										ChildContext.Bounds.Extent
+										))
+									{
+										bIsInFrustum = true;
+										break;
+									}
 								}
 							}
-						}
 
-						if (!bIsInFrustum)
-						{
-							for(int32 ShadowIndex = 0, Num = ViewDependentWholeSceneShadows.Num(); ShadowIndex < Num; ShadowIndex++)
+							if (!bIsInFrustum)
 							{
-								FProjectedShadowInfo* ProjectedShadowInfo = ViewDependentWholeSceneShadows[ShadowIndex];
-
-								//check(ProjectedShadowInfo->CasterFrustum.PermutedPlanes.Num());
-								// Check if this primitive is in the shadow's frustum.
-								if(ProjectedShadowInfo->CasterFrustum.IntersectBox(
-									ChildContext.Bounds.Center + ProjectedShadowInfo->PreShadowTranslation,
-									ChildContext.Bounds.Extent
-									))
+								for(int32 ShadowIndex = 0, Num = ViewDependentWholeSceneShadows.Num(); ShadowIndex < Num; ShadowIndex++)
 								{
-									bIsInFrustum = true;
-									break;
+									FProjectedShadowInfo* ProjectedShadowInfo = ViewDependentWholeSceneShadows[ShadowIndex];
+
+									//check(ProjectedShadowInfo->CasterFrustum.PermutedPlanes.Num());
+									// Check if this primitive is in the shadow's frustum.
+									if(ProjectedShadowInfo->CasterFrustum.IntersectBox(
+										ChildContext.Bounds.Center + ProjectedShadowInfo->PreShadowTranslation,
+										ChildContext.Bounds.Extent
+										))
+									{
+										bIsInFrustum = true;
+										break;
+									}
 								}
 							}
-						}
 
-						if(bIsInFrustum)
-						{
-							// If the child node was in the frustum of at least one preshadow, push it on
-							// the iterator's pending node stack.
-							PrimitiveOctreeIt.PushChild(ChildRef);
+							if(bIsInFrustum)
+							{
+								// If the child node was in the frustum of at least one preshadow, push it on
+								// the iterator's pending node stack.
+								PrimitiveOctreeIt.PushChild(ChildRef);
+							}
 						}
 					}
 				}
+
+				if (PrimitiveOctreeNode.GetElementCount() > 0)
+				{
+					FGatherShadowPrimitivesPacket* Packet = new(FMemStack::Get()) FGatherShadowPrimitivesPacket(Scene, Views, &PrimitiveOctreeNode, 0, 0, PreShadows, ViewDependentWholeSceneShadows, FeatureLevel, bStaticSceneOnly);
+					Packets.Add(Packet);
+				}
+			}
+		}
+		else
+		{
+			const int32 PacketSize = CVarParallelGatherNumPrimitivesPerPacket.GetValueOnRenderThread();
+			const int32 NumPackets = FMath::DivideAndRoundUp(Scene->Primitives.Num(), PacketSize);
+			
+			Packets.Reserve(NumPackets);
+
+			for (int32 PacketIndex = 0; PacketIndex < NumPackets; PacketIndex++)
+			{
+				const int32 StartPrimitiveIndex = PacketIndex * PacketSize;
+				const int32 NumPrimitives = FMath::Min(PacketSize, Scene->Primitives.Num() - StartPrimitiveIndex);
+				FGatherShadowPrimitivesPacket* Packet = new(FMemStack::Get()) FGatherShadowPrimitivesPacket(Scene, Views, NULL, StartPrimitiveIndex, NumPrimitives, PreShadows, ViewDependentWholeSceneShadows, FeatureLevel, bStaticSceneOnly);
+				Packets.Add(Packet);
 			}
+		}
+			
+		{
+			QUICK_SCOPE_CYCLE_COUNTER(STAT_FilterPrimitivesForShadows);
 
-			// Check all the primitives in this octree node.
-			for(FScenePrimitiveOctree::ElementConstIt NodePrimitiveIt(PrimitiveOctreeNode.GetElementIt());NodePrimitiveIt;++NodePrimitiveIt)
+			ParallelFor(Packets.Num(), 
+				[&Packets](int32 Index)
+				{
+					Packets[Index]->AnyThreadTask();
+				},
+				!(FApp::ShouldUseThreadingForPerformance() && CVarParallelGatherShadowPrimitives.GetValueOnRenderThread() > 0)
+			);
+		}
+
+		{
+			QUICK_SCOPE_CYCLE_COUNTER(STAT_RenderThreadFinalize);
+
+			for (int32 PacketIndex = 0; PacketIndex < Packets.Num(); PacketIndex++)
 			{
-				// gather the shadows for this one primitive
-				GatherShadowsForPrimitiveInner(*NodePrimitiveIt, PreShadows, ViewDependentWholeSceneShadows, bStaticSceneOnly);
+				FGatherShadowPrimitivesPacket* Packet = Packets[PacketIndex];
+				Packet->RenderThreadFinalize();
+				// Class was allocated on the memstack which does not call destructors
+				Packet->~FGatherShadowPrimitivesPacket();
 			}
 		}
 
-		for(int32 ShadowIndex = 0, Num = PreShadows.Num(); ShadowIndex < Num; ShadowIndex++)
+		for (int32 ShadowIndex = 0, Num = PreShadows.Num(); ShadowIndex < Num; ShadowIndex++)
 		{
 			FProjectedShadowInfo* ProjectedShadowInfo = PreShadows[ShadowIndex];
 			//@todo - sort other shadow types' subject mesh elements?
diff --git a/Engine/Source/Runtime/Renderer/Private/ShadowSetupMobile.cpp b/Engine/Source/Runtime/Renderer/Private/ShadowSetupMobile.cpp
--- a/Engine/Source/Runtime/Renderer/Private/ShadowSetupMobile.cpp
+++ b/Engine/Source/Runtime/Renderer/Private/ShadowSetupMobile.cpp
@@ -120,7 +120,7 @@ static bool MobileDetermineStaticMeshesCSMVisibilityStateInner(
 	{
 		FProjectedShadowInfo* RESTRICT ProjectedShadowInfo = WholeSceneShadow;
 
-		if (ProjectedShadowInfo->bReflectiveShadowmap && !PrimitiveSceneInfoCompact.bAffectDynamicIndirectLighting)
+		if (ProjectedShadowInfo->bReflectiveShadowmap && !PrimitiveProxy->AffectsDynamicIndirectLighting())
 		{
 			return bFoundCSMReceiver;
 		}
diff --git a/Engine/Source/Runtime/Renderer/Private/SpriteIndexBuffer.h b/Engine/Source/Runtime/Renderer/Private/SpriteIndexBuffer.h
new file mode 100644
--- /dev/null
+++ b/Engine/Source/Runtime/Renderer/Private/SpriteIndexBuffer.h
@@ -0,0 +1,28 @@
+// Copyright 1998-2016 Epic Games, Inc. All Rights Reserved.
+
+#pragma once
+
+template< uint32 NumSprites >
+class FSpriteIndexBuffer : public FIndexBuffer
+{
+public:
+	virtual void InitRHI() override
+	{
+		const uint32 Size = sizeof(uint16) * 6 * NumSprites;
+		const uint32 Stride = sizeof(uint16);
+		FRHIResourceCreateInfo CreateInfo;
+		void* Buffer = nullptr;
+		IndexBufferRHI = RHICreateAndLockIndexBuffer( Stride, Size, BUF_Static, CreateInfo, Buffer );
+		uint16* Indices = (uint16*)Buffer;
+		for (uint32 SpriteIndex = 0; SpriteIndex < NumSprites; ++SpriteIndex)
+		{
+			Indices[SpriteIndex*6 + 0] = SpriteIndex*4 + 0;
+			Indices[SpriteIndex*6 + 1] = SpriteIndex*4 + 3;
+			Indices[SpriteIndex*6 + 2] = SpriteIndex*4 + 2;
+			Indices[SpriteIndex*6 + 3] = SpriteIndex*4 + 0;
+			Indices[SpriteIndex*6 + 4] = SpriteIndex*4 + 1;
+			Indices[SpriteIndex*6 + 5] = SpriteIndex*4 + 3;
+		}
+		RHIUnlockIndexBuffer( IndexBufferRHI );
+	}
+};
\ No newline at end of file
