Copying //UE4/Dev-Rendering to //UE4/Dev-Main (Source: //UE4/Dev-Rendering @ 3461187)

diff --git a/Engine/Source/Runtime/Renderer/Private/StaticMeshDrawList.inl b/Engine/Source/Runtime/Renderer/Private/StaticMeshDrawList.inl
--- a/Engine/Source/Runtime/Renderer/Private/StaticMeshDrawList.inl
+++ b/Engine/Source/Runtime/Renderer/Private/StaticMeshDrawList.inl
@@ -59,16 +59,17 @@ void TStaticMeshDrawList<DrawingPolicyType>::FElementHandle::Remove(const bool b
 
 	checkSlow(LocalDrawingPolicyLink->SetId == SetId);
 
-	check(LocalDrawingPolicyLink->Elements[ElementIndex].Mesh->MaterialRenderProxy);
-	LocalDrawingPolicyLink->Elements[ElementIndex].Mesh->MaterialRenderProxy->SetUnreferencedInDrawList();
+	check(LocalDrawingPolicyLink->Elements[LocalElementIndex].Mesh->MaterialRenderProxy);
+	LocalDrawingPolicyLink->Elements[LocalElementIndex].Mesh->MaterialRenderProxy->SetUnreferencedInDrawList();
 
 	// Unlink the mesh from this draw list. Not necessary if the mesh is being destroyed
 	if (bUnlinkMesh)
 	{
 		// Expensive (Order N). Spins through whole list
-		LocalDrawingPolicyLink->Elements[ElementIndex].Mesh->UnlinkDrawList(this);
+		LocalDrawingPolicyLink->Elements[LocalElementIndex].Mesh->UnlinkDrawList(this);
 	}
-	LocalDrawingPolicyLink->Elements[ElementIndex].Mesh = NULL;
+	//from this point on the memory the this pointer point to might be gone (e.g. if we unlink ourselves)
+	LocalDrawingPolicyLink->Elements[LocalElementIndex].Mesh = NULL;
 
 	checkSlow(LocalDrawingPolicyLink->Elements.Num() == LocalDrawingPolicyLink->CompactElements.Num());
 
@@ -355,7 +356,7 @@ bool TStaticMeshDrawList<DrawingPolicyType>::DrawVisibleInner(
 					STAT(StatInc += Element.Mesh->GetNumPrimitives();)
 					int32 SubCount = Element.Mesh->Elements.Num();
 					// Avoid the cache miss looking up batch visibility if there is only one element.
-					uint64 BatchElementMask = Element.Mesh->bRequiresPerElementVisibility ? (*BatchVisibilityArray)[Element.Mesh->Id] : ((1ull << SubCount) - 1);
+					uint64 BatchElementMask = Element.Mesh->bRequiresPerElementVisibility ? (*BatchVisibilityArray)[Element.Mesh->BatchVisibilityId] : ((1ull << SubCount) - 1);
 					Count += DrawElement<InstancedStereoPolicy::Disabled>(RHICmdList, View, PolicyContext, DrawRenderState, Element, BatchElementMask, DrawingPolicyLink, bDrawnShared);
 				}
 			}
@@ -363,24 +364,24 @@ bool TStaticMeshDrawList<DrawingPolicyType>::DrawVisibleInner(
 			// Stereo pair, we need to test both eyes
 			else
 			{
-				const TArray<uint64, SceneRenderingAllocator>* ResolvedVisiblityArray = nullptr;
+				const TArray<uint64, SceneRenderingAllocator>* ResolvedBatchVisiblityArray = nullptr;
 
 				if (StereoView->LeftViewVisibilityMap->AccessCorrespondingBit(FRelativeBitReference(CompactElementPtr->MeshId)))
 				{
-					ResolvedVisiblityArray = StereoView->LeftViewBatchVisibilityArray;
+					ResolvedBatchVisiblityArray = StereoView->LeftViewBatchVisibilityArray;
 				}
 				else if (StereoView->RightViewVisibilityMap->AccessCorrespondingBit(FRelativeBitReference(CompactElementPtr->MeshId)))
 				{
-					ResolvedVisiblityArray = StereoView->RightViewBatchVisibilityArray;
+					ResolvedBatchVisiblityArray = StereoView->RightViewBatchVisibilityArray;
 				}
 
-				if (ResolvedVisiblityArray != nullptr)
+				if (ResolvedBatchVisiblityArray != nullptr)
 				{
 					const FElement& Element = DrawingPolicyLink->Elements[ElementIndex];
 					STAT(StatInc += Element.Mesh->GetNumPrimitives();)
 						int32 SubCount = Element.Mesh->Elements.Num();
 					// Avoid the cache miss looking up batch visibility if there is only one element.
-					uint64 BatchElementMask = Element.Mesh->bRequiresPerElementVisibility ? (*ResolvedVisiblityArray)[Element.Mesh->Id] : ((1ull << SubCount) - 1);
+					uint64 BatchElementMask = Element.Mesh->bRequiresPerElementVisibility ? (*ResolvedBatchVisiblityArray)[Element.Mesh->BatchVisibilityId] : ((1ull << SubCount) - 1);
 					Count += DrawElement<InstancedStereo>(RHICmdList, View, PolicyContext, DrawRenderState, Element, BatchElementMask, DrawingPolicyLink, bDrawnShared);
 				}
 			}
@@ -603,19 +604,19 @@ void TStaticMeshDrawList<DrawingPolicyType>::DrawVisibleParallelInternal(
 								const FElement& Element = DrawingPolicyLink->Elements[ElementIndex];
 								const int32 SubCount = Element.Mesh->Elements.Num();
 
-								// Avoid the cache miss looking up batch visibility if there is only one element.
-								if (SubCount == 1)
+								// Avoid the cache miss looking up batch visibility if it's not needed
+								if (!Element.Mesh->bRequiresPerElementVisibility)
 								{
-									++Count;
+									Count += SubCount;
 								}
 								else if (!bIsInstancedStereo)
 								{
-									Count += CountBits((*BatchVisibilityArray)[Element.Mesh->Id]);
+									Count += CountBits((*BatchVisibilityArray)[Element.Mesh->BatchVisibilityId]);
 								}
 								else
 								{
-									const int32 LeftCount = CountBits((*StereoView->LeftViewBatchVisibilityArray)[Element.Mesh->Id]);
-									const int32 RightCount = CountBits((*StereoView->RightViewBatchVisibilityArray)[Element.Mesh->Id]);
+									const int32 LeftCount = CountBits((*StereoView->LeftViewBatchVisibilityArray)[Element.Mesh->BatchVisibilityId]);
+									const int32 RightCount = CountBits((*StereoView->RightViewBatchVisibilityArray)[Element.Mesh->BatchVisibilityId]);
 									Count += (LeftCount > RightCount) ? LeftCount : RightCount;
 								}
 							}
@@ -878,7 +879,7 @@ int32 TStaticMeshDrawList<DrawingPolicyType>::DrawVisibleFrontToBackInner(
 		const TArray<uint64, SceneRenderingAllocator>* const ResolvedVisiblityArray = (InstancedStereo == InstancedStereoPolicy::Disabled) ? BatchVisibilityArray : ElementVisibility[SortedIndex];
 
 		// Avoid the cache miss looking up batch visibility if there is only one element.
-		uint64 BatchElementMask = Element.Mesh->bRequiresPerElementVisibility ? (*ResolvedVisiblityArray)[Element.Mesh->Id] : ((1ull << Element.Mesh->Elements.Num()) - 1);
+		uint64 BatchElementMask = Element.Mesh->bRequiresPerElementVisibility ? (*ResolvedVisiblityArray)[Element.Mesh->BatchVisibilityId] : ((1ull << Element.Mesh->Elements.Num()) - 1);
 		DrawElement<InstancedStereoPolicy::Disabled>(RHICmdList, View, PolicyContext, DrawRenderState, Element, BatchElementMask, DrawingPolicyLink, bDrawnShared);
 		NumDraws++;
 	}
diff --git a/Engine/Source/Runtime/Renderer/Private/SystemTextures.cpp b/Engine/Source/Runtime/Renderer/Private/SystemTextures.cpp
--- a/Engine/Source/Runtime/Renderer/Private/SystemTextures.cpp
+++ b/Engine/Source/Runtime/Renderer/Private/SystemTextures.cpp
@@ -62,14 +62,14 @@ void FSystemTextures::InternalInitializeTextures(FRHICommandListImmediate& RHICm
 			RHICmdList.CopyToResolveTarget(GreenDummy->GetRenderTargetItem().TargetableTexture, GreenDummy->GetRenderTargetItem().ShaderResourceTexture, true, FResolveParams());
 		}
 
-		// Create a MidGrayDummy texture
+		// Create a DefaultNormal8Bit texture
 		{
-			FPooledRenderTargetDesc Desc(FPooledRenderTargetDesc::Create2DDesc(FIntPoint(1, 1), PF_B8G8R8A8, FClearValueBinding::MidGray, TexCreate_HideInVisualizeTexture, TexCreate_RenderTargetable | TexCreate_NoFastClear, false));
+			FPooledRenderTargetDesc Desc(FPooledRenderTargetDesc::Create2DDesc(FIntPoint(1, 1), PF_B8G8R8A8, FClearValueBinding::DefaultNormal8Bit, TexCreate_HideInVisualizeTexture, TexCreate_RenderTargetable | TexCreate_NoFastClear, false));
 			Desc.AutoWritable = false;
-			GRenderTargetPool.FindFreeElement(RHICmdList, Desc, MidGrayDummy, TEXT("MidGrayDummy"));
+			GRenderTargetPool.FindFreeElement(RHICmdList, Desc, DefaultNormal8Bit, TEXT("DefaultNormal8Bit"));
 
-			SetRenderTarget(RHICmdList, MidGrayDummy->GetRenderTargetItem().TargetableTexture, FTextureRHIRef(), ESimpleRenderTargetMode::EClearColorExistingDepth);
-			RHICmdList.CopyToResolveTarget(MidGrayDummy->GetRenderTargetItem().TargetableTexture, MidGrayDummy->GetRenderTargetItem().ShaderResourceTexture, true, FResolveParams());
+			SetRenderTarget(RHICmdList, DefaultNormal8Bit->GetRenderTargetItem().TargetableTexture, FTextureRHIRef(), ESimpleRenderTargetMode::EClearColorExistingDepth);
+			RHICmdList.CopyToResolveTarget(DefaultNormal8Bit->GetRenderTargetItem().TargetableTexture, DefaultNormal8Bit->GetRenderTargetItem().ShaderResourceTexture, true, FResolveParams());
 		}
 
 		// Create the PerlinNoiseGradient texture
@@ -104,6 +104,63 @@ void FSystemTextures::InternalInitializeTextures(FRHICommandListImmediate& RHICm
 			RHICmdList.UnlockTexture2D((FTexture2DRHIRef&)PerlinNoiseGradient->GetRenderTargetItem().ShaderResourceTexture, 0, false);
 		}
 
+		// Create the SobolSampling texture
+		{
+			const size_t SobolBits = 8;
+			static const uint16 SobolXCell[SobolBits][4] = {
+				{ 0x5880, 0x7780, 0x9400, 0xc400 },
+				{ 0x5400, 0xa400, 0x4a00, 0xc200 },
+				{ 0x3a80, 0x2980, 0xfb00, 0x5700 },
+				{ 0xe800, 0x8800, 0x0400, 0x5400 },
+				{ 0xea00, 0x3600, 0xa200, 0x8a00 },
+				{ 0x4c00, 0x1c00, 0x2600, 0x5e00 },
+				{ 0xa480, 0x9b80, 0xe600, 0x9e00 },
+				{ 0x6880, 0x0780, 0xae00, 0x7600 }
+			};
+			static const uint16 SobolYCell[SobolBits][4] = {
+				{ 0x8e80, 0xed80, 0xf600, 0x8e00 },
+				{ 0x6e00, 0x8200, 0x8e00, 0x5600 },
+				{ 0xf600, 0xba00, 0x1100, 0x3500 },
+				{ 0x6a80, 0xb980, 0x2200, 0x0a00 },
+				{ 0x2600, 0xaa00, 0x4400, 0x1400 },
+				{ 0xe880, 0x8780, 0x8800, 0x2800 },
+				{ 0xa480, 0x9b80, 0xe600, 0x9e00 },
+				{ 0x6880, 0x0780, 0xae00, 0x7600 }
+			};
+
+			FPooledRenderTargetDesc Desc(FPooledRenderTargetDesc::Create2DDesc(FIntPoint(1 << SobolBits, 1 << SobolBits), PF_R16G16B16A16_UINT, FClearValueBinding::None, TexCreate_HideInVisualizeTexture, TexCreate_NoFastClear, false));
+			Desc.AutoWritable = false;
+			GRenderTargetPool.FindFreeElement(RHICmdList, Desc, SobolSampling, TEXT("SobolSampling"));
+			// Write the contents of the texture.
+			uint32 DestStride;
+			uint8* DestBuffer = (uint8*)RHICmdList.LockTexture2D((FTexture2DRHIRef&)SobolSampling->GetRenderTargetItem().ShaderResourceTexture, 0, RLM_WriteOnly, DestStride, false);
+
+			for (int y = 0; y < Desc.Extent.Y; ++y)
+			{
+				for (int x = 0; x < Desc.Extent.X; ++x)
+				{
+					uint16 *Dest = (uint16*)(DestBuffer + x * 4 * sizeof(uint16) + y * DestStride);
+					Dest[0] = Dest[1] = Dest[2] = Dest[3] = 0;
+
+					for (int bit = 0; bit < SobolBits; ++bit)
+					{
+						Dest[0] ^= (x & (1 << bit)) ? SobolXCell[bit][0] : 0;
+						Dest[0] ^= (y & (1 << bit)) ? SobolYCell[bit][0] : 0;
+
+						Dest[1] ^= (x & (1 << bit)) ? SobolXCell[bit][1] : 0;
+						Dest[1] ^= (y & (1 << bit)) ? SobolYCell[bit][1] : 0;
+
+						Dest[2] ^= (x & (1 << bit)) ? SobolXCell[bit][2] : 0;
+						Dest[2] ^= (y & (1 << bit)) ? SobolYCell[bit][2] : 0;
+
+						Dest[3] ^= (x & (1 << bit)) ? SobolXCell[bit][3] : 0;
+						Dest[3] ^= (y & (1 << bit)) ? SobolYCell[bit][3] : 0;
+					}
+				}
+			}
+			RHICmdList.UnlockTexture2D((FTexture2DRHIRef&)SobolSampling->GetRenderTargetItem().ShaderResourceTexture, 0, false);
+		}
+
 		if (!GSupportsShaderFramebufferFetch && GPixelFormats[PF_FloatRGBA].Supported)
 		{
 			FPooledRenderTargetDesc Desc(FPooledRenderTargetDesc::Create2DDesc(FIntPoint(1, 1), PF_FloatRGBA, FClearValueBinding(FLinearColor(65000.0f, 65000.0f, 65000.0f, 65000.0f)), TexCreate_HideInVisualizeTexture, TexCreate_RenderTargetable | TexCreate_NoFastClear, false));
@@ -432,12 +489,13 @@ void FSystemTextures::ReleaseDynamicRHI()
 	BlackAlphaOneDummy.SafeRelease();
 	PerlinNoiseGradient.SafeRelease();
 	PerlinNoise3D.SafeRelease();
+	SobolSampling.SafeRelease();
 	SSAORandomization.SafeRelease();
 	PreintegratedGF.SafeRelease();
 	MaxFP16Depth.SafeRelease();
 	DepthDummy.SafeRelease();
 	GreenDummy.SafeRelease();
-	MidGrayDummy.SafeRelease();
+	DefaultNormal8Bit.SafeRelease();
 
 	GRenderTargetPool.FreeUnusedResources();
 
diff --git a/Engine/Source/Runtime/Renderer/Private/SystemTextures.h b/Engine/Source/Runtime/Renderer/Private/SystemTextures.h
--- a/Engine/Source/Runtime/Renderer/Private/SystemTextures.h
+++ b/Engine/Source/Runtime/Renderer/Private/SystemTextures.h
@@ -54,6 +54,8 @@ public:
 	TRefCountPtr<IPooledRenderTarget> PerlinNoiseGradient;
 	// used by the material expression Noise (faster version, should replace old version), todo: move out of SceneRenderTargets
 	TRefCountPtr<IPooledRenderTarget> PerlinNoise3D;
+	// Sobol sampling texture, the first sample points for four sobol dimensions in RGBA
+	TRefCountPtr<IPooledRenderTarget> SobolSampling;
 	/** SSAO randomization */
 	TRefCountPtr<IPooledRenderTarget> SSAORandomization;
 	/** Preintegrated GF for single sample IBL */
@@ -65,7 +67,7 @@ public:
 	// float4(0,1,0,1)
 	TRefCountPtr<IPooledRenderTarget> GreenDummy;
 	// float4(0.5,0.5,0.5,1)
-	TRefCountPtr<IPooledRenderTarget> MidGrayDummy;
+	TRefCountPtr<IPooledRenderTarget> DefaultNormal8Bit;
 
 protected:
 	/** Maximum feature level that the textures have been initialized up to */
diff --git a/Engine/Source/Runtime/Renderer/Private/TranslucentLighting.cpp b/Engine/Source/Runtime/Renderer/Private/TranslucentLighting.cpp
--- a/Engine/Source/Runtime/Renderer/Private/TranslucentLighting.cpp
+++ b/Engine/Source/Runtime/Renderer/Private/TranslucentLighting.cpp
@@ -549,7 +549,7 @@ void FProjectedShadowInfo::RenderTranslucencyDepths(FRHICommandList& RHICmdList,
 			);
 
 		FLinearColor ClearColors[2] = {FLinearColor(0,0,0,0), FLinearColor(0,0,0,0)};
-		DrawClearQuadMRT(RHICmdList, SceneRenderer->FeatureLevel, true, ARRAY_COUNT(ClearColors), ClearColors, false, 1.0f, false, 0);
+		DrawClearQuadMRT(RHICmdList, true, ARRAY_COUNT(ClearColors), ClearColors, false, 1.0f, false, 0);
 
 		// Set the viewport for the shadow.
 		RHICmdList.SetViewport(
@@ -903,7 +903,7 @@ void ClearVolumeTextures(FRHICommandList& RHICmdList, ERHIFeatureLevel::Type Fea
 	if (false)
 #endif
 	{
-		DrawClearQuadMRT(RHICmdList, GMaxRHIFeatureLevel, true, NumRenderTargets, ClearColors, false, 0, false, 0);
+		DrawClearQuadMRT(RHICmdList, true, NumRenderTargets, ClearColors, false, 0, false, 0);
 	}
 	else
 	{
diff --git a/Engine/Source/Runtime/Renderer/Private/TranslucentRendering.cpp b/Engine/Source/Runtime/Renderer/Private/TranslucentRendering.cpp
--- a/Engine/Source/Runtime/Renderer/Private/TranslucentRendering.cpp
+++ b/Engine/Source/Runtime/Renderer/Private/TranslucentRendering.cpp
@@ -49,11 +49,100 @@ static TAutoConsoleVariable<float> CVarSeparateTranslucencyMinDownsampleChangeTi
 	TEXT("Minimum time in seconds between changes to automatic downsampling state, used to prevent rapid swapping between half and full res."),
 	ECVF_Scalability | ECVF_Default);
 
-static TAutoConsoleVariable<int32> CVarCopySceneColorOncePerViewOnly(
-	TEXT("r.CopySceneColorOncePerViewOnly"),
-	0,
-	TEXT("Copy the scene color for translucent materials that sample from it only once per-view rather than prior to each draw call. Currently adds a fixed overhead of one copy per-view even when it isn't used."),
-	ECVF_RenderThreadSafe);
+static TAutoConsoleVariable<int32> CVarSeparateTranslucencyUpsampleMode(
+	TEXT("r.SeparateTranslucencyUpsampleMode"),
+	1,
+	TEXT("Upsample method to use on separate translucency.  These are only used when r.SeparateTranslucencyScreenPercentage is less than 100.\n")
+	TEXT("0: bilinear 1: Nearest-Depth Neighbor (only when r.SeparateTranslucencyScreenPercentage is 50)"),
+	ECVF_Scalability | ECVF_Default);
+
+int32 GAllowDownsampledStandardTranslucency = 0;
+
+static FAutoConsoleVariableRef CVarAllowDownsampledStandardTranslucency(
+	TEXT("r.AllowDownsampledStandardTranslucency"),
+	GAllowDownsampledStandardTranslucency,
+	TEXT("Allow standard translucency to be rendered in smaller resolution as an optimization\n")
+	TEXT("This is incompatible with materials using blend modulate. Use 2 to ignore those. \n")
+	TEXT(" <0: off\n")
+	TEXT(" 0: on unless a material using blend modulate is used (default)")
+	TEXT(" >0: on and ignores any material using blend modulate"),
+	ECVF_RenderThreadSafe
+	);
+
+/** Mostly used to know if debug rendering should be drawn in this pass */
+FORCEINLINE bool IsMainTranslucencyPass(ETranslucencyPass::Type TranslucencyPass)
+{
+	return TranslucencyPass == ETranslucencyPass::TPT_AllTranslucency || TranslucencyPass == ETranslucencyPass::TPT_StandardTranslucency;
+}
+
+static bool RenderInSeparateTranslucency(const FSceneRenderTargets& SceneContext, ETranslucencyPass::Type TranslucencyPass, bool bPrimitiveDisablesOffscreenBuffer)
+{
+	// Currently AfterDOF is rendered earlier in the frame and must be rendered in a separate (offscreen) buffer.
+	if (TranslucencyPass == ETranslucencyPass::TPT_TranslucencyAfterDOF)
+	{
+		// If bPrimitiveDisablesOffscreenBuffer, that will trigger an ensure call
+		return true;
+	}
+	
+	// Otherwise it only gets rendered in the separate buffer if it is downsampled
+	if (bPrimitiveDisablesOffscreenBuffer ? (GAllowDownsampledStandardTranslucency > 0) : (GAllowDownsampledStandardTranslucency >= 0))
+	{
+		FIntPoint ScaledSize;
+		float DownsamplingScale = 1.f;
+		SceneContext.GetSeparateTranslucencyDimensions(ScaledSize, DownsamplingScale);
+
+		if (DownsamplingScale < 1.f)
+		{
+			return true;
+		}
+	}
+
+	return false;
+}
+
+bool FTranslucencyDrawingPolicyFactory::ContextType::ShouldDraw(const FViewInfo& View, const FMaterial* Material, bool bIsSeparateTranslucency) const
+{
+	bool bShouldDraw = false;
+
+	if (Material)
+	{
+		// Only render translucent materials
+		const EBlendMode BlendMode = Material->GetBlendMode();
+		if (IsTranslucentBlendMode(BlendMode) && ShouldIncludeDomainInMeshPass(Material->GetMaterialDomain()))
+		{
+			if (TranslucencyPass == ETranslucencyPass::TPT_AllTranslucency)
+			{
+				bShouldDraw = true;
+			}
+			// Only draw meshes in the relevant pass
+			const ETranslucencyPass::Type MaterialPass = Material->IsTranslucencyAfterDOFEnabled() ? ETranslucencyPass::TPT_TranslucencyAfterDOF : ETranslucencyPass::TPT_StandardTranslucency;
+			if (TranslucencyPass == MaterialPass || (!View.Family->EngineShowFlags.PostProcessing && MaterialPass == ETranslucencyPass::TPT_TranslucencyAfterDOF && TranslucencyPass == ETranslucencyPass::TPT_StandardTranslucency))
+			{
+				bShouldDraw = true;
+			}
+		}
+
+		if (bShouldDraw && BlendMode == BLEND_Modulate && bIsSeparateTranslucency)
+		{
+			// < 0 : never downsample, = 0 downsample only if no blend modulate, > 0 ignore
+			ensure(GAllowDownsampledStandardTranslucency > 0);
+	#if !UE_BUILD_SHIPPING
+			if (GAllowDownsampledStandardTranslucency > 0)
+			{
+				static bool bOnce = false;
+				if (!bOnce)
+				{
+					UE_LOG(LogRenderer, Warning, TEXT("Blend modulate materials (%s) are not supported when r.AllowDownsampledStandardTranslucency > 0."), *Material->GetFriendlyName());
+					bOnce = true;
+				}
+			}
+	#endif
+		}
+
+	}
+
+	return bShouldDraw;
+}
 
 void FDeferredShadingSceneRenderer::UpdateTranslucencyTimersAndSeparateTranslucencyBufferSize(FRHICommandListImmediate& RHICmdList)
 {
@@ -138,9 +227,8 @@ void FDeferredShadingSceneRenderer::UpdateTranslucencyTimersAndSeparateTransluce
 }
 
 void FDeferredShadingSceneRenderer::BeginTimingSeparateTranslucencyPass(FRHICommandListImmediate& RHICmdList, const FViewInfo& View)
-{
-	if (View.ViewState 
-		&& GSupportsTimestampRenderQueries
+{ 
+	if (View.ViewState && GSupportsTimestampRenderQueries
 #if !STATS
 		&& (CVarSeparateTranslucencyAutoDownsample.GetValueOnRenderThread() != 0)
 #endif
@@ -152,8 +240,7 @@ void FDeferredShadingSceneRenderer::BeginTimingSeparateTranslucencyPass(FRHIComm
 
 void FDeferredShadingSceneRenderer::EndTimingSeparateTranslucencyPass(FRHICommandListImmediate& RHICmdList, const FViewInfo& View)
 {
-	if (View.ViewState 
-		&& GSupportsTimestampRenderQueries
+	if (View.ViewState && GSupportsTimestampRenderQueries
 #if !STATS
 		&& (CVarSeparateTranslucencyAutoDownsample.GetValueOnRenderThread() != 0)
 #endif
@@ -163,27 +250,6 @@ void FDeferredShadingSceneRenderer::EndTimingSeparateTranslucencyPass(FRHIComman
 	}
 }
 
-static void SetTranslucentRenderTarget(FRHICommandList& RHICmdList, const FViewInfo& View, ETranslucencyPass::Type TranslucencyPass, bool bFirstTimeThisFrame = false)
-{
-	FSceneRenderTargets& SceneContext = FSceneRenderTargets::Get(RHICmdList);
-	bool bNeedsClear = (&View == View.Family->Views[0]) && bFirstTimeThisFrame;
-
-	if ((TranslucencyPass == ETranslucencyPass::TPT_SeparateTranslucency) && SceneContext.IsSeparateTranslucencyActive(View))
-	{
-		SceneContext.BeginRenderingSeparateTranslucency(RHICmdList, View, bNeedsClear);
-	}
-	else if (TranslucencyPass == ETranslucencyPass::TPT_StandardTranslucency)
-	{
-		SceneContext.BeginRenderingTranslucency(RHICmdList, View, bNeedsClear);
-	}
-}
-
-static void SetTranslucentState(FRHICommandList& RHICmdList, FDrawingPolicyRenderState& DrawRenderState)
-{
-	// Enable depth test, disable depth writes.
-	DrawRenderState.SetDepthStencilState(TStaticDepthStencilState<false, CF_DepthNearOrEqual>::GetRHI());
-}
-
 const FProjectedShadowInfo* FDeferredShadingSceneRenderer::PrepareTranslucentShadowMap(FRHICommandList& RHICmdList, const FViewInfo& View, FPrimitiveSceneInfo* PrimitiveSceneInfo, ETranslucencyPass::Type TranslucencyPass)
 {
 	QUICK_SCOPE_CYCLE_COUNTER(STAT_FDeferredShadingSceneRenderer_PrepareTranslucentShadowMap);
@@ -259,15 +325,16 @@ private:
 
 IMPLEMENT_SHADER_TYPE(,FCopySceneColorPS,TEXT("TranslucentLightingShaders"),TEXT("CopySceneColorMain"),SF_Pixel);
 
-void FTranslucencyDrawingPolicyFactory::CopySceneColor(FRHICommandList& RHICmdList, const FViewInfo& View, const FPrimitiveSceneProxy* PrimitiveSceneProxy)
+void FTranslucencyDrawingPolicyFactory::CopySceneColor(FRHICommandList& RHICmdList, const FViewInfo& View)
 {
 	FSceneRenderTargets& SceneContext = FSceneRenderTargets::Get(RHICmdList);
 
-	SCOPED_DRAW_EVENTF(RHICmdList, EventCopy, TEXT("CopySceneColor from SceneColor node for %s %s"), PrimitiveSceneProxy ? *PrimitiveSceneProxy->GetOwnerName().ToString() : TEXT("Scene"), PrimitiveSceneProxy ? *PrimitiveSceneProxy->GetResourceName().ToString() : TEXT("Scene"));
+	SCOPED_DRAW_EVENTF(RHICmdList, EventCopy, TEXT("CopySceneColor from SceneColor for translucency"));
 
 	RHICmdList.CopyToResolveTarget(SceneContext.GetSceneColorSurface(), SceneContext.GetSceneColorTexture(), true, FResolveRect(View.ViewRect.Min.X, View.ViewRect.Min.Y, View.ViewRect.Max.X, View.ViewRect.Max.Y));
 
 	SceneContext.BeginRenderingLightAttenuation(RHICmdList);
+	RHICmdList.SetViewport(View.ViewRect.Min.X, View.ViewRect.Min.Y, 0.0f, View.ViewRect.Max.X, View.ViewRect.Max.Y, 1.0f);
 
 	FGraphicsPipelineStateInitializer GraphicsPSOInit;
 	RHICmdList.ApplyCachedRenderTargets(GraphicsPSOInit);
@@ -275,7 +342,6 @@ void FTranslucencyDrawingPolicyFactory::CopySceneColor(FRHICommandList& RHICmdLi
 	GraphicsPSOInit.DepthStencilState = TStaticDepthStencilState<false, CF_Always>::GetRHI();
 	GraphicsPSOInit.BlendState = TStaticBlendState<>::GetRHI();
 
-	RHICmdList.SetViewport(View.ViewRect.Min.X, View.ViewRect.Min.Y, 0.0f, View.ViewRect.Max.X, View.ViewRect.Max.Y, 1.0f);
 
 	TShaderMapRef<FScreenVS> ScreenVertexShader(View.ShaderMap);
 	TShaderMapRef<FCopySceneColorPS> PixelShader(View.ShaderMap);
@@ -419,52 +485,6 @@ public:
 	}
 };
 
-static void CopySceneColorAndRestore(FRHICommandList& RHICmdList, const FViewInfo& View, const FPrimitiveSceneProxy* PrimitiveSceneProxy)
-{
-	check(IsInRenderingThread());
-	FTranslucencyDrawingPolicyFactory::CopySceneColor(RHICmdList, View, PrimitiveSceneProxy);
-	// Restore state
-	SetTranslucentRenderTarget(RHICmdList, View, ETranslucencyPass::TPT_StandardTranslucency);
-}
-
-class FCopySceneColorAndRestoreRenderThreadTask
-{
-	FRHICommandList& RHICmdList;
-	const FViewInfo& View;
-	const FPrimitiveSceneProxy* PrimitiveSceneProxy;
-public:
-
-	FCopySceneColorAndRestoreRenderThreadTask(FRHICommandList& InRHICmdList, const FViewInfo& InView, const FPrimitiveSceneProxy* InPrimitiveSceneProxy)
-		: RHICmdList(InRHICmdList)
-		, View(InView)
-		, PrimitiveSceneProxy(InPrimitiveSceneProxy)
-	{
-	}
-
-	FORCEINLINE TStatId GetStatId() const
-	{
-		RETURN_QUICK_DECLARE_CYCLE_STAT(FCopySceneColorAndRestoreRenderThreadTask, STATGROUP_TaskGraphTasks);
-	}
-
-	ENamedThreads::Type GetDesiredThread()
-	{
-		return ENamedThreads::RenderThread_Local;
-	}
-
-	static ESubsequentsMode::Type GetSubsequentsMode() { return ESubsequentsMode::TrackSubsequents; }
-
-	void DoTask(ENamedThreads::Type CurrentThread, const FGraphEventRef& MyCompletionGraphEvent)
-	{
-		CopySceneColorAndRestore(RHICmdList, View, PrimitiveSceneProxy);
-		RHICmdList.HandleRTThreadTaskCompletion(MyCompletionGraphEvent);
-	}
-};
-
-bool IsSeparateTranslucencyPossible(const FViewInfo& View)
-{
-	return (FSceneRenderTargets::CVarSetSeperateTranslucencyEnabled.GetValueOnRenderThread() != 0) && View.Family->EngineShowFlags.SeparateTranslucency && View.Family->EngineShowFlags.PostProcessing;
-}
-
 /**
 * Render a dynamic or static mesh using a translucent draw policy
 * @return true if the mesh rendered
@@ -484,113 +504,78 @@ bool FTranslucencyDrawingPolicyFactory::DrawMesh(
 	bool bDirty = false;
 	const auto FeatureLevel = View.GetFeatureLevel();
 
+	const FSceneRenderTargets& SceneContext = FSceneRenderTargets::Get(RHICmdList);
+
 	// Determine the mesh's material and blend mode.
 	const FMaterial* Material = Mesh.MaterialRenderProxy->GetMaterial(FeatureLevel);
-	const EBlendMode BlendMode = Material->GetBlendMode();
 
-	// Only render translucent materials
-	if (IsTranslucentBlendMode(BlendMode) && ShouldIncludeDomainInMeshPass(Material->GetMaterialDomain()))
+	// Only render relevant materials
+	if (DrawingContext.ShouldDraw(View, Material, SceneContext.IsSeparateTranslucencyPass()))
 	{
-		// fix for materials on Canvas
-		// PrimitiveSceneProxy is NULL when rendering Canvas items
-		const bool bMeshUseSeparateTranslucency = Material->IsSeparateTranslucencyEnabled() && IsSeparateTranslucencyPossible(View) && PrimitiveSceneProxy != nullptr;
-
-		// Only draw meshes in the relevant pass
-		if (bMeshUseSeparateTranslucency == (DrawingContext.TranslucenyPassType == ETranslucencyPass::TPT_SeparateTranslucency)
-			|| DrawingContext.TranslucenyPassType == ETranslucencyPass::TPT_AllTranslucency)
-		{
-			FDrawingPolicyRenderState DrawRenderStateLocal(DrawRenderState);
-
-			if (Material->RequiresSceneColorCopy_RenderThread())
-			{
-				if (DrawingContext.bSceneColorCopyIsUpToDate == false)
-				{
-					if (CVarCopySceneColorOncePerViewOnly.GetValueOnRenderThread() == 0)
-					{
-						if (!RHICmdList.Bypass() && !IsInActualRenderingThread() && !IsInGameThread())
-						{
-							FRHICommandList* CmdList = new FRHICommandList;
-							CmdList->CopyRenderThreadContexts(RHICmdList);
-							FGraphEventRef RenderThreadCompletionEvent = TGraphTask<FCopySceneColorAndRestoreRenderThreadTask>::CreateTask().ConstructAndDispatchWhenReady(*CmdList, View, PrimitiveSceneProxy);
-							RHICmdList.QueueRenderThreadCommandListSubmit(RenderThreadCompletionEvent, CmdList);
-						}
-						else
-						{
-							// otherwise, just do it now. We don't want to defer in this case because that can interfere with render target visualization (a debugging tool).
-							CopySceneColorAndRestore(RHICmdList, View, PrimitiveSceneProxy);
-						}
-					}
-
-					// todo: this optimization is currently broken
-					DrawingContext.bSceneColorCopyIsUpToDate = (DrawingContext.TranslucenyPassType == ETranslucencyPass::TPT_SeparateTranslucency);
-					SetTranslucentState(RHICmdList, DrawRenderStateLocal);
-				}
-			}
+		FDrawingPolicyRenderState DrawRenderStateLocal(DrawRenderState);
 
-			const bool bDisableDepthTest = Material->ShouldDisableDepthTest();
-			const bool bEnableResponsiveAA = Material->ShouldEnableResponsiveAA();
-			// editor compositing not supported on translucent materials currently
-			const bool bEditorCompositeDepthTest = false;
+		const bool bDisableDepthTest = Material->ShouldDisableDepthTest();
+		const bool bEnableResponsiveAA = Material->ShouldEnableResponsiveAA();
+		// editor compositing not supported on translucent materials currently
+		const bool bEditorCompositeDepthTest = false;
 
-			// if this draw is coming postAA then there is probably no depth buffer (it's canvas) and bEnableResponsiveAA wont' do anything anyway.
-			if (bEnableResponsiveAA && !DrawingContext.bPostAA)
+		// if this draw is coming postAA then there is probably no depth buffer (it's canvas) and bEnableResponsiveAA wont' do anything anyway.
+		if (bEnableResponsiveAA && !DrawingContext.bPostAA)
+		{
+			if( bDisableDepthTest )
 			{
-				if( bDisableDepthTest )
-				{
-					DrawRenderStateLocal.SetDepthStencilState(TStaticDepthStencilState<
-						false, CF_Always,
-						true, CF_Always, SO_Keep, SO_Keep, SO_Replace,
-						false, CF_Always, SO_Keep, SO_Keep, SO_Keep,
-						STENCIL_TEMPORAL_RESPONSIVE_AA_MASK, STENCIL_TEMPORAL_RESPONSIVE_AA_MASK
-						>::GetRHI());
-					DrawRenderStateLocal.SetStencilRef(STENCIL_TEMPORAL_RESPONSIVE_AA_MASK);
-				}
-				else
-				{
-					DrawRenderStateLocal.SetDepthStencilState(TStaticDepthStencilState<
-						false, CF_DepthNearOrEqual,
-						true, CF_Always, SO_Keep, SO_Keep, SO_Replace,
-						false, CF_Always, SO_Keep, SO_Keep, SO_Keep,
-						STENCIL_TEMPORAL_RESPONSIVE_AA_MASK, STENCIL_TEMPORAL_RESPONSIVE_AA_MASK
-						>::GetRHI());
-					DrawRenderStateLocal.SetStencilRef(STENCIL_TEMPORAL_RESPONSIVE_AA_MASK);
-				}
+				DrawRenderStateLocal.SetDepthStencilState(TStaticDepthStencilState<
+					false, CF_Always,
+					true, CF_Always, SO_Keep, SO_Keep, SO_Replace,
+					false, CF_Always, SO_Keep, SO_Keep, SO_Keep,
+					STENCIL_TEMPORAL_RESPONSIVE_AA_MASK, STENCIL_TEMPORAL_RESPONSIVE_AA_MASK
+					>::GetRHI());
+				DrawRenderStateLocal.SetStencilRef(STENCIL_TEMPORAL_RESPONSIVE_AA_MASK);
 			}
-			else if( bDisableDepthTest )
+			else
 			{
-				DrawRenderStateLocal.SetDepthStencilState(TStaticDepthStencilState<false,CF_Always>::GetRHI());
+				DrawRenderStateLocal.SetDepthStencilState(TStaticDepthStencilState<
+					false, CF_DepthNearOrEqual,
+					true, CF_Always, SO_Keep, SO_Keep, SO_Replace,
+					false, CF_Always, SO_Keep, SO_Keep, SO_Keep,
+					STENCIL_TEMPORAL_RESPONSIVE_AA_MASK, STENCIL_TEMPORAL_RESPONSIVE_AA_MASK
+					>::GetRHI());
+				DrawRenderStateLocal.SetStencilRef(STENCIL_TEMPORAL_RESPONSIVE_AA_MASK);
 			}
-
-			FSceneRenderTargets& SceneContext = FSceneRenderTargets::Get(RHICmdList);
-			FIntPoint OutScaledSize;
-			float OutScale;
-			SceneContext.GetSeparateTranslucencyDimensions(OutScaledSize, OutScale);
-
-			ProcessBasePassMesh(
-				RHICmdList, 
-				FProcessBasePassMeshParameters(
-					Mesh,
-					BatchElementMask,
-					Material,
-					PrimitiveSceneProxy, 
-					!bPreFog,
-					bEditorCompositeDepthTest,
-					DrawingContext.TextureMode,
-					FeatureLevel
-				),
-				FDrawTranslucentMeshAction(
-					RHICmdList,
-					View,
-					DrawRenderStateLocal,
-					HitProxyId,
-					DrawingContext.TranslucentSelfShadow,
-					PrimitiveSceneProxy && PrimitiveSceneProxy->CastsVolumetricTranslucentShadow(),
-					DrawingContext.TranslucenyPassType == ETranslucencyPass::TPT_SeparateTranslucency && OutScale < 1.0f
-				)
-			);
-
-			bDirty = true;
 		}
+		else if( bDisableDepthTest )
+		{
+			DrawRenderStateLocal.SetDepthStencilState(TStaticDepthStencilState<false,CF_Always>::GetRHI());
+		}
+
+		FIntPoint OutScaledSize;
+		float OutScale;
+		SceneContext.GetSeparateTranslucencyDimensions(OutScaledSize, OutScale);
+
+		ProcessBasePassMesh(
+			RHICmdList, 
+			FProcessBasePassMeshParameters(
+				Mesh,
+				BatchElementMask,
+				Material,
+				PrimitiveSceneProxy, 
+				!bPreFog,
+				bEditorCompositeDepthTest,
+				DrawingContext.TextureMode,
+				FeatureLevel
+			),
+			FDrawTranslucentMeshAction(
+				RHICmdList,
+				View,
+				DrawRenderStateLocal,
+				HitProxyId,
+				DrawingContext.TranslucentSelfShadow,
+				PrimitiveSceneProxy && PrimitiveSceneProxy->CastsVolumetricTranslucentShadow(),
+				OutScale < 1.f && SceneContext.IsSeparateTranslucencyPass()
+			)
+		);
+
+		bDirty = true;
 	}
 	return bDirty;
 }
@@ -690,18 +675,18 @@ class FVolumetricTranslucentShadowRenderThreadTask
 	const FViewInfo& View;
 	FDrawingPolicyRenderState DrawRenderState;
 	FDeferredShadingSceneRenderer& Renderer;
-	ETranslucencyPass::Type TranslucenyPassType;
+	ETranslucencyPass::Type TranslucencyPass;
 	int32 Index;
 
 public:
 
-	FORCEINLINE_DEBUGGABLE FVolumetricTranslucentShadowRenderThreadTask(FRHICommandList& InRHICmdList, const FTranslucentPrimSet& InPrimSet, const FViewInfo& InView, const FDrawingPolicyRenderState& InDrawRenderState, FDeferredShadingSceneRenderer& InRenderer, ETranslucencyPass::Type InTranslucenyPassType, int32 InIndex)
+	FORCEINLINE_DEBUGGABLE FVolumetricTranslucentShadowRenderThreadTask(FRHICommandList& InRHICmdList, const FTranslucentPrimSet& InPrimSet, const FViewInfo& InView, const FDrawingPolicyRenderState& InDrawRenderState, FDeferredShadingSceneRenderer& InRenderer, ETranslucencyPass::Type InTranslucencyPass, int32 InIndex)
 		: RHICmdList(InRHICmdList)
 		, PrimSet(InPrimSet)
 		, View(InView)
 		, DrawRenderState(InDrawRenderState)
 		, Renderer(InRenderer)
-		, TranslucenyPassType(InTranslucenyPassType)
+		, TranslucencyPass(InTranslucencyPass)
 		, Index(InIndex)
 	{
 	}
@@ -720,7 +705,7 @@ public:
 
 	void DoTask(ENamedThreads::Type CurrentThread, const FGraphEventRef& MyCompletionGraphEvent)
 	{
-		PrimSet.DrawAPrimitive(RHICmdList, View, DrawRenderState, Renderer, TranslucenyPassType, Index);
+		PrimSet.DrawAPrimitive(RHICmdList, View, DrawRenderState, Renderer, TranslucencyPass, Index);
 		RHICmdList.HandleRTThreadTaskCompletion(MyCompletionGraphEvent);
 	}
 };
@@ -730,7 +715,7 @@ void FTranslucentPrimSet::DrawPrimitivesParallel(
 	const FViewInfo& View,
 	const FDrawingPolicyRenderState& DrawRenderState,
 	FDeferredShadingSceneRenderer& Renderer,
-	ETranslucencyPass::Type TranslucenyPassType,
+	ETranslucencyPass::Type TranslucencyPass,
 	int32 FirstPrimIdx, int32 LastPrimIdx
 	) const
 {
@@ -749,12 +734,12 @@ void FTranslucentPrimSet::DrawPrimitivesParallel(
 			// can't do this in parallel, defer
 			FRHICommandList* CmdList = new FRHICommandList;
 			CmdList->CopyRenderThreadContexts(RHICmdList);
-			FGraphEventRef RenderThreadCompletionEvent = TGraphTask<FVolumetricTranslucentShadowRenderThreadTask>::CreateTask().ConstructAndDispatchWhenReady(*CmdList, *this, View, DrawRenderState, Renderer, TranslucenyPassType, PrimIdx);
+			FGraphEventRef RenderThreadCompletionEvent = TGraphTask<FVolumetricTranslucentShadowRenderThreadTask>::CreateTask().ConstructAndDispatchWhenReady(*CmdList, *this, View, DrawRenderState, Renderer, TranslucencyPass, PrimIdx);
 			RHICmdList.QueueRenderThreadCommandListSubmit(RenderThreadCompletionEvent, CmdList);
 		}
 		else
 		{
-			RenderPrimitive(RHICmdList, View, DrawRenderState, PrimitiveSceneInfo, ViewRelevance, nullptr, TranslucenyPassType);
+			RenderPrimitive(RHICmdList, View, DrawRenderState, PrimitiveSceneInfo, ViewRelevance, nullptr, TranslucencyPass);
 		}
 	}
 }
@@ -764,10 +749,10 @@ void FTranslucentPrimSet::DrawPrimitives(
 	const FViewInfo& View,
 	const FDrawingPolicyRenderState& DrawRenderState,
 	FDeferredShadingSceneRenderer& Renderer,
-	ETranslucencyPass::Type TranslucenyPassType
+	ETranslucencyPass::Type TranslucencyPass
 	) const
 {
-	FInt32Range PassRange = SortedPrimsNum.GetPassRange(TranslucenyPassType);
+	FInt32Range PassRange = SortedPrimsNum.GetPassRange(TranslucencyPass);
 
 	// Draw sorted scene prims
 	for( int32 PrimIdx = PassRange.GetLowerBoundValue(); PrimIdx < PassRange.GetUpperBoundValue(); PrimIdx++ )
@@ -778,9 +763,9 @@ void FTranslucentPrimSet::DrawPrimitives(
 
 		checkSlow(ViewRelevance.HasTranslucency());
 			
-		const FProjectedShadowInfo* TranslucentSelfShadow = Renderer.PrepareTranslucentShadowMap(RHICmdList, View, PrimitiveSceneInfo, TranslucenyPassType);
+		const FProjectedShadowInfo* TranslucentSelfShadow = Renderer.PrepareTranslucentShadowMap(RHICmdList, View, PrimitiveSceneInfo, TranslucencyPass);
 
-		RenderPrimitive(RHICmdList, View, DrawRenderState, PrimitiveSceneInfo, ViewRelevance, TranslucentSelfShadow, TranslucenyPassType);
+		RenderPrimitive(RHICmdList, View, DrawRenderState, PrimitiveSceneInfo, ViewRelevance, TranslucentSelfShadow, TranslucencyPass);
 	}
 
 	View.SimpleElementCollector.DrawBatchedElements(RHICmdList, DrawRenderState, View, FTexture2DRHIRef(), EBlendModeFilter::Translucent);
@@ -793,16 +778,14 @@ void FTranslucentPrimSet::RenderPrimitive(
 	FPrimitiveSceneInfo* PrimitiveSceneInfo,
 	const FPrimitiveViewRelevance& ViewRelevance,
 	const FProjectedShadowInfo* TranslucentSelfShadow,
-	ETranslucencyPass::Type TranslucenyPassType) const
+	ETranslucencyPass::Type TranslucencyPass) const
 {
 	checkSlow(ViewRelevance.HasTranslucency());
 	auto FeatureLevel = View.GetFeatureLevel();
 
 	if (ViewRelevance.bDrawRelevance)
 	{
-		FTranslucencyDrawingPolicyFactory::ContextType Context(TranslucentSelfShadow, TranslucenyPassType);
-	
-		const bool bSeparateTranslucencyPossible = IsSeparateTranslucencyPossible(View);
+		FTranslucencyDrawingPolicyFactory::ContextType Context(TranslucentSelfShadow, TranslucencyPass);
 
 		// Render dynamic scene prim
 		{	
@@ -828,27 +811,20 @@ void FTranslucentPrimSet::RenderPrimitive(
 			{
 				FStaticMesh& StaticMesh = PrimitiveSceneInfo->StaticMeshes[StaticMeshIdx];
 
-				if (View.StaticMeshVisibilityMap[StaticMesh.Id]
-					// Only render static mesh elements using translucent materials
-					&& StaticMesh.IsTranslucent(FeatureLevel))
+				// Only render visible elements with relevant materials
+				if (View.StaticMeshVisibilityMap[StaticMesh.Id] && Context.ShouldDraw(View, StaticMesh.MaterialRenderProxy->GetMaterial(FeatureLevel), FSceneRenderTargets::Get(RHICmdList).IsSeparateTranslucencyPass()))
 				{
-					const bool bMeshUseSeparateTranslucency = StaticMesh.MaterialRenderProxy->GetMaterial(FeatureLevel)->IsSeparateTranslucencyEnabled() && bSeparateTranslucencyPossible;
-
-					if (bMeshUseSeparateTranslucency == (TranslucenyPassType == ETranslucencyPass::TPT_SeparateTranslucency)
-						|| TranslucenyPassType == ETranslucencyPass::TPT_AllTranslucency)
-					{
-						FTranslucencyDrawingPolicyFactory::DrawStaticMesh(
-							RHICmdList,
-							View,
-							FTranslucencyDrawingPolicyFactory::ContextType(TranslucentSelfShadow, TranslucenyPassType),
-							StaticMesh,
-							StaticMesh.bRequiresPerElementVisibility ? View.StaticMeshBatchVisibility[StaticMesh.Id] : ((1ull << StaticMesh.Elements.Num()) - 1),
-							false,
-							DrawRenderState,
-							PrimitiveSceneInfo->Proxy,
-							StaticMesh.BatchHitProxyId
-							);
-					}
+					FTranslucencyDrawingPolicyFactory::DrawStaticMesh(
+						RHICmdList,
+						View,
+						Context,
+						StaticMesh,
+						StaticMesh.bRequiresPerElementVisibility ? View.StaticMeshBatchVisibility[StaticMesh.BatchVisibilityId] : ((1ull << StaticMesh.Elements.Num()) - 1),
+						false,
+						DrawRenderState,
+						PrimitiveSceneInfo->Proxy,
+						StaticMesh.BatchHitProxyId
+						);
 				}
 			}
 		}
@@ -886,35 +862,31 @@ void FTranslucentPrimSet::AppendScenePrimitives(FTranslucentSortedPrim* Elements
 	SortedPrimsNum.Append(TranslucentPrimitiveCountPerPass);
 }
 
-void FTranslucentPrimSet::PlaceScenePrimitive(FPrimitiveSceneInfo* PrimitiveSceneInfo, const FViewInfo& ViewInfo, bool bUseNormalTranslucency, bool bUseSeparateTranslucency, bool bUseMobileSeparateTranslucency,
+void FTranslucentPrimSet::PlaceScenePrimitive(FPrimitiveSceneInfo* PrimitiveSceneInfo, const FViewInfo& ViewInfo, const FPrimitiveViewRelevance& ViewRelevance,
 	FTranslucentPrimSet::FTranslucentSortedPrim *InArrayStart, int32& InOutArrayNum, FTranslucenyPrimCount& OutCount)
 {
 	const float SortKey = CalculateTranslucentSortKey(PrimitiveSceneInfo, ViewInfo);
 	const auto FeatureLevel = ViewInfo.GetFeatureLevel();
-	int32 CVarEnabled = FSceneRenderTargets::CVarSetSeperateTranslucencyEnabled.GetValueOnRenderThread();
-
-	bool bCanBeSeparate = CVarEnabled
-		&& FeatureLevel >= ERHIFeatureLevel::SM4
-		&& ViewInfo.Family->EngineShowFlags.PostProcessing 
-		&& !ViewInfo.Family->EngineShowFlags.ShaderComplexity
-		&& ViewInfo.Family->EngineShowFlags.SeparateTranslucency;
 
-	bool bIsSeparateTranslucency = bUseSeparateTranslucency && bCanBeSeparate;
-	bool bIsNonSeparateTranslucency = bUseNormalTranslucency || !bCanBeSeparate;
-
-	if (bIsSeparateTranslucency)
+	if (ViewInfo.Family->AllowTranslucencyAfterDOF())
 	{
-		ETranslucencyPass::Type TranslucencyPass = ETranslucencyPass::TPT_SeparateTranslucency;
+		if (ViewRelevance.bNormalTranslucencyRelevance)
+		{
+			new(&InArrayStart[InOutArrayNum++]) FTranslucentSortedPrim(PrimitiveSceneInfo, ETranslucencyPass::TPT_StandardTranslucency, PrimitiveSceneInfo->Proxy->GetTranslucencySortPriority(), SortKey);
+			OutCount.Add(ETranslucencyPass::TPT_StandardTranslucency, ViewRelevance.bUsesSceneColorCopy, ViewRelevance.bDisableOffscreenRendering);
+		}
 
-		new(&InArrayStart[InOutArrayNum++]) FTranslucentSortedPrim(PrimitiveSceneInfo, TranslucencyPass, PrimitiveSceneInfo->Proxy->GetTranslucencySortPriority(), SortKey);
-		OutCount.Add(TranslucencyPass);
+		if (ViewRelevance.bSeparateTranslucencyRelevance)
+		{
+			new(&InArrayStart[InOutArrayNum++]) FTranslucentSortedPrim(PrimitiveSceneInfo, ETranslucencyPass::TPT_TranslucencyAfterDOF, PrimitiveSceneInfo->Proxy->GetTranslucencySortPriority(), SortKey);
+			OutCount.Add(ETranslucencyPass::TPT_TranslucencyAfterDOF, ViewRelevance.bUsesSceneColorCopy, ViewRelevance.bDisableOffscreenRendering);
+		}
 	}
-	if (bIsNonSeparateTranslucency)
+	else // Otherwise, everything is rendered in a single bucket. This is not related to whether DOF is currently enabled or not.
 	{
-		ETranslucencyPass::Type TranslucencyPass = ETranslucencyPass::TPT_StandardTranslucency;
-
-		new(&InArrayStart[InOutArrayNum++]) FTranslucentSortedPrim(PrimitiveSceneInfo, TranslucencyPass, PrimitiveSceneInfo->Proxy->GetTranslucencySortPriority(), SortKey);
-		OutCount.Add(TranslucencyPass);
+		// When using all translucency, Standard and AfterDOF are sorted together instead of being rendered like 2 buckets.
+		new(&InArrayStart[InOutArrayNum++]) FTranslucentSortedPrim(PrimitiveSceneInfo, ETranslucencyPass::TPT_AllTranslucency, PrimitiveSceneInfo->Proxy->GetTranslucencySortPriority(), SortKey);
+		OutCount.Add(ETranslucencyPass::TPT_AllTranslucency, ViewRelevance.bUsesSceneColorCopy, ViewRelevance.bDisableOffscreenRendering);
 	}
 }
 
@@ -924,22 +896,42 @@ void FTranslucentPrimSet::SortPrimitives()
 	SortedPrims.Sort( FCompareFTranslucentSortedPrim() );
 }
 
-bool FSceneRenderer::ShouldRenderTranslucency() const
+extern int32 GLightShaftRenderAfterDOF;
+
+bool FSceneRenderer::ShouldRenderTranslucency(ETranslucencyPass::Type TranslucencyPass) const
 {
-	bool bRender = false;
+	// Change this condition to control where simple elements should be rendered.
+	if (IsMainTranslucencyPass(TranslucencyPass))
+	{
+		if (ViewFamily.EngineShowFlags.VisualizeLPV)
+		{
+			return true;
+		}
 
-	for (int32 ViewIndex = 0; ViewIndex < Views.Num(); ViewIndex++)
+		for (const FViewInfo& View : Views)
+		{
+			if (View.bHasTranslucentViewMeshElements || View.SimpleElementCollector.BatchedElements.HasPrimsToDraw())
+			{
+				return true;
+			}
+		}
+	}
+
+	// If lightshafts are rendered in low res, we must reset the offscreen buffer in case is was also used in TPT_StandardTranslucency.
+	if (GLightShaftRenderAfterDOF && TranslucencyPass == ETranslucencyPass::TPT_TranslucencyAfterDOF)
 	{
-		const FViewInfo& View = Views[ViewIndex];
+		return true;
+	}
 
-		if (View.TranslucentPrimSet.NumPrims() || View.bHasTranslucentViewMeshElements)
+	for (const FViewInfo& View : Views)
+	{
+		if (View.TranslucentPrimSet.SortedPrimsNum.Num(TranslucencyPass) > 0)
 		{
-			bRender = true;
-			break;
+			return true;
 		}
 	}
 
-	return bRender;
+	return false;
 }
 
 class FDrawSortedTransAnyThreadTask : public FRenderTask
@@ -948,7 +940,7 @@ class FDrawSortedTransAnyThreadTask : public FRenderTask
 	FRHICommandList& RHICmdList;
 	const FViewInfo& View;
 	FDrawingPolicyRenderState DrawRenderState;
-	ETranslucencyPass::Type TranslucenyPassType;
+	ETranslucencyPass::Type TranslucencyPass;
 
 	const int32 FirstIndex;
 	const int32 LastIndex;
@@ -960,7 +952,7 @@ public:
 		FRHICommandList& InRHICmdList,
 		const FViewInfo& InView,
 		const FDrawingPolicyRenderState& InDrawRenderState,
-		ETranslucencyPass::Type InTranslucenyPassType,
+		ETranslucencyPass::Type InTranslucencyPass,
 		int32 InFirstIndex,
 		int32 InLastIndex
 		)
@@ -968,7 +960,7 @@ public:
 		, RHICmdList(InRHICmdList)
 		, View(InView)
 		, DrawRenderState(InDrawRenderState)
-		, TranslucenyPassType(InTranslucenyPassType)
+		, TranslucencyPass(InTranslucencyPass)
 		, FirstIndex(InFirstIndex)
 		, LastIndex(InLastIndex)
 	{
@@ -984,7 +976,7 @@ public:
 	void DoTask(ENamedThreads::Type CurrentThread, const FGraphEventRef& MyCompletionGraphEvent)
 	{
 		FScopeCycleCounter ScopeOuter(RHICmdList.ExecuteStat);
-		View.TranslucentPrimSet.DrawPrimitivesParallel(RHICmdList, View, DrawRenderState, Renderer, TranslucenyPassType, FirstIndex, LastIndex);
+		View.TranslucentPrimSet.DrawPrimitivesParallel(RHICmdList, View, DrawRenderState, Renderer, TranslucencyPass, FirstIndex, LastIndex);
 		RHICmdList.HandleRTThreadTaskCompletion(MyCompletionGraphEvent);
 	}
 };
@@ -994,13 +986,14 @@ DECLARE_CYCLE_STAT(TEXT("Translucency"), STAT_CLP_Translucency, STATGROUP_Parall
 
 class FTranslucencyPassParallelCommandListSet : public FParallelCommandListSet
 {
-	ETranslucencyPass::Type TranslucenyPassType;
-	bool bFirstTimeThisFrame;
+	ETranslucencyPass::Type TranslucencyPass;
+	bool bRenderInSeparateTranslucency;
+
 public:
-	FTranslucencyPassParallelCommandListSet(const FViewInfo& InView, FRHICommandListImmediate& InParentCmdList, bool bInParallelExecute, bool bInCreateSceneContext, ETranslucencyPass::Type InTranslucenyPassType)
+	FTranslucencyPassParallelCommandListSet(const FViewInfo& InView, FRHICommandListImmediate& InParentCmdList, bool bInParallelExecute, bool bInCreateSceneContext, ETranslucencyPass::Type InTranslucencyPass, bool InRenderInSeparateTranslucency)
 		: FParallelCommandListSet(GET_STATID(STAT_CLP_Translucency), InView, InParentCmdList, bInParallelExecute, bInCreateSceneContext)
-		, TranslucenyPassType(InTranslucenyPassType)
-		, bFirstTimeThisFrame(true)
+		, TranslucencyPass(InTranslucencyPass)
+		, bRenderInSeparateTranslucency(InRenderInSeparateTranslucency)
 	{
 		SetStateOnCommandList(ParentCmdList);
 	}
@@ -1012,10 +1005,18 @@ public:
 
 	virtual void SetStateOnCommandList(FRHICommandList& CmdList) override
 	{
+		// Never needs clear here as it is already done in RenderTranslucency.
 		FParallelCommandListSet::SetStateOnCommandList(CmdList);
-		SetTranslucentRenderTarget(CmdList, View, TranslucenyPassType, bFirstTimeThisFrame);
-		SetTranslucentState(CmdList, DrawRenderState);
-		bFirstTimeThisFrame = false;
+		FSceneRenderTargets& SceneContext = FSceneRenderTargets::Get(CmdList);
+		if (bRenderInSeparateTranslucency)
+		{
+			SceneContext.BeginRenderingSeparateTranslucency(CmdList, View, false);
+		}
+		else
+		{
+			SceneContext.BeginRenderingTranslucency(CmdList, View, false);
+		}
+		DrawRenderState.SetDepthStencilState(TStaticDepthStencilState<false, CF_DepthNearOrEqual>::GetRHI());
 	}
 };
 
@@ -1029,74 +1030,96 @@ static TAutoConsoleVariable<int32> CVarRHICmdFlushRenderThreadTasksTranslucentPa
 	0,
 	TEXT("Wait for completion of parallel render thread tasks at the end of the translucent pass.  A more granular version of r.RHICmdFlushRenderThreadTasks. If either r.RHICmdFlushRenderThreadTasks or r.RHICmdFlushRenderThreadTasksTranslucentPass is > 0 we will flush."));
 
-// this is a static because we let the async tasks beyond the function
-static FTranslucencyDrawingPolicyFactory::ContextType GParallelTranslucencyContext;
 
-void FDeferredShadingSceneRenderer::RenderTranslucencyParallel(FRHICommandListImmediate& RHICmdList)
-{
-	FSceneRenderTargets& SceneContext = FSceneRenderTargets::Get(RHICmdList);
-	SceneContext.AllocLightAttenuation(RHICmdList); // materials will attempt to get this texture before the deferred command to set it up executes
-	check(IsInRenderingThread());
+static TAutoConsoleVariable<int32> CVarParallelTranslucency(
+	TEXT("r.ParallelTranslucency"),
+	1,
+	TEXT("Toggles parallel translucency rendering. Parallel rendering must be enabled for this to have an effect."),
+	ECVF_RenderThreadSafe
+	);
 
-	GParallelTranslucencyContext.TranslucentSelfShadow = nullptr;
-	GParallelTranslucencyContext.TranslucenyPassType = ETranslucencyPass::TPT_StandardTranslucency;
-	GParallelTranslucencyContext.bSceneColorCopyIsUpToDate = false;
-	FScopedCommandListWaitForTasks Flusher(CVarRHICmdFlushRenderThreadTasksTranslucentPass.GetValueOnRenderThread() > 0 || CVarRHICmdFlushRenderThreadTasks.GetValueOnRenderThread() > 0, RHICmdList);
+// this is a static because we let the async tasks beyond the function. Using all translucency as we want all materials to render
+static FTranslucencyDrawingPolicyFactory::ContextType GParallelTranslucencyContext(nullptr, ETranslucencyPass::TPT_AllTranslucency);
 
-	bool bRequiresSeperateTranslucencyResolve = false;
 
-	for (int32 ViewIndex = 0; ViewIndex < Views.Num(); ViewIndex++)
+void FDeferredShadingSceneRenderer::RenderViewTranslucency(FRHICommandListImmediate& RHICmdList, const FViewInfo& View, const FDrawingPolicyRenderState& DrawRenderState, ETranslucencyPass::Type TranslucencyPass)
+{
+	// Draw translucent prims
+	View.TranslucentPrimSet.DrawPrimitives(RHICmdList, View, DrawRenderState, *this, TranslucencyPass);
+
+	if (IsMainTranslucencyPass(TranslucencyPass))
 	{
-		SCOPED_CONDITIONAL_DRAW_EVENTF(RHICmdList, EventView, Views.Num() > 1, TEXT("View%d"), ViewIndex);
+		View.SimpleElementCollector.DrawBatchedElements(RHICmdList, DrawRenderState, View, FTexture2DRHIRef(), EBlendModeFilter::Translucent);
 
-		FViewInfo& View = Views[ViewIndex];
+		// editor and debug rendering
+		if (View.bHasTranslucentViewMeshElements)
+		{
+			FTranslucencyDrawingPolicyFactory::ContextType Context(0, TranslucencyPass);
+			DrawViewElements<FTranslucencyDrawingPolicyFactory>(RHICmdList, View, DrawRenderState, Context, SDPG_World, false);
+			DrawViewElements<FTranslucencyDrawingPolicyFactory>(RHICmdList, View, DrawRenderState, Context, SDPG_Foreground, false);
+		}
 
+		const FSceneViewState* ViewState = (const FSceneViewState*)View.State;
+		if (ViewState && View.Family->EngineShowFlags.VisualizeLPV)
 		{
+			FLightPropagationVolume* LightPropagationVolume = ViewState->GetLightPropagationVolume(View.GetFeatureLevel());
 
-#if STATS
-			if (View.ViewState)
+			if (LightPropagationVolume)
 			{
-				View.ViewState->TranslucencyTimer.Begin(RHICmdList);
+				LightPropagationVolume->Visualise(RHICmdList, View);
 			}
-#endif
-			
-			ETranslucencyPass::Type TranslucenyPassType = ETranslucencyPass::TPT_StandardTranslucency;
-
-			FTranslucencyPassParallelCommandListSet ParallelCommandListSet(View, RHICmdList, 
-				CVarRHICmdTranslucencyPassDeferredContexts.GetValueOnRenderThread() > 0, 
-				CVarRHICmdFlushRenderThreadTasksTranslucentPass.GetValueOnRenderThread() == 0  && CVarRHICmdFlushRenderThreadTasks.GetValueOnRenderThread() == 0, 
-				TranslucenyPassType);
+		}
+	}
+}
+void FDeferredShadingSceneRenderer::RenderViewTranslucencyParallel(FRHICommandListImmediate& RHICmdList, const FViewInfo& View, const FDrawingPolicyRenderState& DrawRenderState, ETranslucencyPass::Type TranslucencyPass)
+{
+	FTranslucencyPassParallelCommandListSet ParallelCommandListSet(
+		View, 
+		RHICmdList, 
+		CVarRHICmdTranslucencyPassDeferredContexts.GetValueOnRenderThread() > 0, 
+		CVarRHICmdFlushRenderThreadTasksTranslucentPass.GetValueOnRenderThread() == 0  && CVarRHICmdFlushRenderThreadTasks.GetValueOnRenderThread() == 0, 
+		TranslucencyPass,
+		FSceneRenderTargets::Get(RHICmdList).IsSeparateTranslucencyPass()
+		);
 
-			{
-				QUICK_SCOPE_CYCLE_COUNTER(RenderTranslucencyParallel_Start_FDrawSortedTransAnyThreadTask);
+	{
+		QUICK_SCOPE_CYCLE_COUNTER(RenderTranslucencyParallel_Start_FDrawSortedTransAnyThreadTask);
 				
-				FInt32Range PassRange = View.TranslucentPrimSet.SortedPrimsNum.GetPassRange(TranslucenyPassType);
-				int32 NumPrims = PassRange.Size<int32>();
-				int32 EffectiveThreads = FMath::Min<int32>(FMath::DivideAndRoundUp(NumPrims, ParallelCommandListSet.MinDrawsPerCommandList), ParallelCommandListSet.Width);
+		FInt32Range PassRange = View.TranslucentPrimSet.SortedPrimsNum.GetPassRange(TranslucencyPass);
+		int32 NumPrims = PassRange.Size<int32>();
+		int32 EffectiveThreads = FMath::Min<int32>(FMath::DivideAndRoundUp(NumPrims, ParallelCommandListSet.MinDrawsPerCommandList), ParallelCommandListSet.Width);
 
-				int32 Start = PassRange.GetLowerBoundValue();
-				if (EffectiveThreads)
-				{
-					int32 NumPer = NumPrims / EffectiveThreads;
-					int32 Extra = NumPrims - NumPer * EffectiveThreads;
+		int32 Start = PassRange.GetLowerBoundValue();
+		if (EffectiveThreads)
+		{
+			int32 NumPer = NumPrims / EffectiveThreads;
+			int32 Extra = NumPrims - NumPer * EffectiveThreads;
 
-					for (int32 ThreadIndex = 0; ThreadIndex < EffectiveThreads; ThreadIndex++)
-					{
-						int32 Last = Start + (NumPer - 1) + (ThreadIndex < Extra);
-						check(Last >= Start);
+			for (int32 ThreadIndex = 0; ThreadIndex < EffectiveThreads; ThreadIndex++)
+			{
+				int32 Last = Start + (NumPer - 1) + (ThreadIndex < Extra);
+				check(Last >= Start);
 
-						{
-							FRHICommandList* CmdList = ParallelCommandListSet.NewParallelCommandList();
+				{
+					FRHICommandList* CmdList = ParallelCommandListSet.NewParallelCommandList();
 						
-							FGraphEventRef AnyThreadCompletionEvent = TGraphTask<FDrawSortedTransAnyThreadTask>::CreateTask(ParallelCommandListSet.GetPrereqs(), ENamedThreads::RenderThread)
-								.ConstructAndDispatchWhenReady(*this, *CmdList, View, ParallelCommandListSet.DrawRenderState, TranslucenyPassType, Start, Last);
+					FGraphEventRef AnyThreadCompletionEvent = TGraphTask<FDrawSortedTransAnyThreadTask>::CreateTask(ParallelCommandListSet.GetPrereqs(), ENamedThreads::RenderThread)
+						.ConstructAndDispatchWhenReady(*this, *CmdList, View, ParallelCommandListSet.DrawRenderState, TranslucencyPass, Start, Last);
 
-							ParallelCommandListSet.AddParallelCommandList(CmdList, AnyThreadCompletionEvent);
-						}
-						Start = Last + 1;
-					}
+					ParallelCommandListSet.AddParallelCommandList(CmdList, AnyThreadCompletionEvent);
 				}
+				Start = Last + 1;
 			}
+		}
+	}
+
+	if (IsMainTranslucencyPass(TranslucencyPass))
+	{
+		View.SimpleElementCollector.DrawBatchedElements(RHICmdList, DrawRenderState, View, FTexture2DRHIRef(), EBlendModeFilter::Translucent);
+
+		// editor and debug rendering
+		if (View.bHasTranslucentViewMeshElements)
+		{
 			// Draw the view's mesh elements with the translucent drawing policy.
 			{
 				QUICK_SCOPE_CYCLE_COUNTER(RenderTranslucencyParallel_SDPG_World);
@@ -1108,240 +1131,331 @@ void FDeferredShadingSceneRenderer::RenderTranslucencyParallel(FRHICommandListIm
 				DrawViewElementsParallel<FTranslucencyDrawingPolicyFactory>(GParallelTranslucencyContext, SDPG_Foreground, false, ParallelCommandListSet);
 			}
 		}
-#if STATS
-		if (View.ViewState)
-		{
-			View.ViewState->TranslucencyTimer.End(RHICmdList);
-		}
-#endif
-		if (SceneContext.IsSeparateTranslucencyActive(View))
+
+		const FSceneViewState* ViewState = (const FSceneViewState*)View.State;
+		if (ViewState && View.Family->EngineShowFlags.VisualizeLPV)
 		{
-			QUICK_SCOPE_CYCLE_COUNTER(RenderTranslucencyParallel_Downsample);
-			// we need to allocate this now so it ends up in the snapshot
-			FIntPoint ScaledSize;
-			float Scale = 1.0f;
-			SceneContext.GetSeparateTranslucencyDimensions(ScaledSize, Scale);
+			FLightPropagationVolume* LightPropagationVolume = ViewState->GetLightPropagationVolume(View.GetFeatureLevel());
 
-			if (Scale < 1.0f)
+			if (LightPropagationVolume)
 			{
-				SceneContext.GetSeparateTranslucencyDepth(RHICmdList, ScaledSize);
-				DownsampleDepthSurface(RHICmdList, SceneContext.GetSeparateTranslucencyDepthSurface(), View, Scale, false);
+				LightPropagationVolume->Visualise(RHICmdList, View);
+			}
+		}
+	}
+}
+
+void FDeferredShadingSceneRenderer::SetupDownsampledTranslucencyViewUniformBuffer(FRHICommandListImmediate& RHICmdList, FViewInfo& View)
+{
+	if (!View.DownsampledTranslucencyViewUniformBuffer)
+	{
+		FSceneRenderTargets& SceneContext = FSceneRenderTargets::Get(RHICmdList);
+		FIntPoint ScaledSize;
+		float DownsamplingScale = 1.f;
+		SceneContext.GetSeparateTranslucencyDimensions(ScaledSize, DownsamplingScale);
+		ensure(DownsamplingScale < 1.f);
+
+		SceneContext.GetDownsampledTranslucencyDepth(RHICmdList, ScaledSize);
+		DownsampleDepthSurface(RHICmdList, SceneContext.GetDownsampledTranslucencyDepthSurface(), View, DownsamplingScale, false);
+
+		FViewUniformShaderParameters DownsampledTranslucencyParameters = *View.CachedViewUniformShaderParameters;
+
+		// Update the parts of DownsampledTranslucencyParameters which are dependent on the buffer size and view rect
+		View.SetupViewRectUniformBufferParameters(
+			DownsampledTranslucencyParameters,
+			ScaledSize,
+			FIntRect(View.ViewRect.Min.X * DownsamplingScale, View.ViewRect.Min.Y * DownsamplingScale, View.ViewRect.Max.X * DownsamplingScale, View.ViewRect.Max.Y * DownsamplingScale),
+			View.ViewMatrices,
+			View.PrevViewMatrices
+		);
 
-				FViewUniformShaderParameters DownsampledTranslucencyParameters = *View.CachedViewUniformShaderParameters;
+		View.DownsampledTranslucencyViewUniformBuffer = TUniformBufferRef<FViewUniformShaderParameters>::CreateUniformBufferImmediate(DownsampledTranslucencyParameters, UniformBuffer_SingleFrame);
+	}
+}
 
-				// Update the parts of DownsampledTranslucencyParameters which are dependent on the buffer size and view rect
-				View.SetupViewRectUniformBufferParameters(
-					DownsampledTranslucencyParameters,
-					ScaledSize,
-					FIntRect(View.ViewRect.Min.X * Scale, View.ViewRect.Min.Y * Scale, View.ViewRect.Max.X * Scale, View.ViewRect.Max.Y * Scale),
-					View.ViewMatrices,
-					View.PrevViewMatrices
-				);
+void FDeferredShadingSceneRenderer::ConditionalResolveSceneColorForTranslucentMaterials(FRHICommandListImmediate& RHICmdList)
+{
+	for (int32 ViewIndex = 0; ViewIndex < Views.Num(); ViewIndex++)
+	{
+		FViewInfo& View = Views[ViewIndex];
 
-				View.DownsampledTranslucencyViewUniformBuffer = TUniformBufferRef<FViewUniformShaderParameters>::CreateUniformBufferImmediate(DownsampledTranslucencyParameters, UniformBuffer_SingleFrame);
-			}
+		bool bNeedsResolve = false;
+		for (int32 TranslucencyPass = 0; TranslucencyPass < ETranslucencyPass::TPT_MAX && !bNeedsResolve; ++TranslucencyPass)
+		{
+			bNeedsResolve |= View.TranslucentPrimSet.SortedPrimsNum.UseSceneColorCopy((ETranslucencyPass::Type)TranslucencyPass);
 		}
 
+		if (bNeedsResolve)
 		{
-			BeginTimingSeparateTranslucencyPass(RHICmdList, View);
+			FTranslucencyDrawingPolicyFactory::CopySceneColor(RHICmdList, View);
+		}
+	}
+}
 
-			{
-				ETranslucencyPass::Type TranslucencyPass = ETranslucencyPass::TPT_SeparateTranslucency;
+void FDeferredShadingSceneRenderer::RenderTranslucency(FRHICommandListImmediate& RHICmdList, ETranslucencyPass::Type TranslucencyPass)
+{
+	if (!ShouldRenderTranslucency(TranslucencyPass))
+	{
+		return; // Early exit if nothing needs to be done.
+	}
 
-				if (CVarCopySceneColorOncePerViewOnly.GetValueOnRenderThread() != 0)
-				{
-					TGraphTask<FCopySceneColorAndRestoreRenderThreadTask>::CreateTask().ConstructAndDispatchWhenReady(RHICmdList, View, nullptr);
-				}
-				
-				// always call BeginRenderingSeparateTranslucency() even if there are no primitives to we keep the RT allocated
-				FTranslucencyPassParallelCommandListSet ParallelCommandListSet(View, 
-					RHICmdList, 
-					CVarRHICmdTranslucencyPassDeferredContexts.GetValueOnRenderThread() > 0, 
-					CVarRHICmdFlushRenderThreadTasksTranslucentPass.GetValueOnRenderThread() == 0 && CVarRHICmdFlushRenderThreadTasks.GetValueOnRenderThread() == 0, 
-					TranslucencyPass);
+	SCOPED_DRAW_EVENT(RHICmdList, Translucency);
+	SCOPED_GPU_STAT(RHICmdList, Stat_GPU_Translucency);
 
-				// Draw only translucent prims that are in the SeparateTranslucency pass
-				if (View.TranslucentPrimSet.SortedPrimsNum.Num(TranslucencyPass) > 0)
-				{
-					bRequiresSeperateTranslucencyResolve = true;
+	FSceneRenderTargets& SceneContext = FSceneRenderTargets::Get(RHICmdList);
 
-					QUICK_SCOPE_CYCLE_COUNTER(RenderTranslucencyParallel_Start_FDrawSortedTransAnyThreadTask_SeparateTransluceny);
+	// Support for parallel rendering.
+	const bool bUseParallel = GRHICommandList.UseParallelAlgorithms() && CVarParallelTranslucency.GetValueOnRenderThread();
+	if (bUseParallel)
+	{
+		SceneContext.AllocLightAttenuation(RHICmdList); // materials will attempt to get this texture before the deferred command to set it up executes
+	}
+	FScopedCommandListWaitForTasks Flusher(bUseParallel && (CVarRHICmdFlushRenderThreadTasksTranslucentPass.GetValueOnRenderThread() > 0 || CVarRHICmdFlushRenderThreadTasks.GetValueOnRenderThread() > 0), RHICmdList);
 
-					FInt32Range PassRange = View.TranslucentPrimSet.SortedPrimsNum.GetPassRange(TranslucencyPass);
-					int32 NumPrims = PassRange.Size<int32>();
-					int32 EffectiveThreads = FMath::Min<int32>(FMath::DivideAndRoundUp(NumPrims, ParallelCommandListSet.MinDrawsPerCommandList), ParallelCommandListSet.Width);
+	for (int32 ViewIndex = 0; ViewIndex < Views.Num(); ViewIndex++)
+	{
+		SCOPED_CONDITIONAL_DRAW_EVENTF(RHICmdList, EventView, Views.Num() > 1, TEXT("View%d"), ViewIndex);
 
-					int32 Start = PassRange.GetLowerBoundValue();
-					check(EffectiveThreads);
-					{
-						int32 NumPer = NumPrims / EffectiveThreads;
-						int32 Extra = NumPrims - NumPer * EffectiveThreads;
+		FViewInfo& View = Views[ViewIndex];
 
-						for (int32 ThreadIndex = 0; ThreadIndex < EffectiveThreads; ThreadIndex++)
-						{
-							int32 Last = Start + (NumPer - 1) + (ThreadIndex < Extra);
-							check(Last >= Start);
+		// if (View.TranslucentPrimSet.SortedPrimsNum.UseSceneColorCopy(TranslucencyPass))
+		// {
+		// 		FTranslucencyDrawingPolicyFactory::CopySceneColor(RHICmdList, View);
+		// }
 
-							{
-								FRHICommandList* CmdList = ParallelCommandListSet.NewParallelCommandList();
+#if STATS
+		if (View.ViewState && IsMainTranslucencyPass(TranslucencyPass))
+		{
+			View.ViewState->TranslucencyTimer.Begin(RHICmdList);
+		}
+#endif
 
-								FGraphEventRef AnyThreadCompletionEvent = TGraphTask<FDrawSortedTransAnyThreadTask>::CreateTask(ParallelCommandListSet.GetPrereqs(), ENamedThreads::RenderThread)
-									.ConstructAndDispatchWhenReady(*this, *CmdList, View, ParallelCommandListSet.DrawRenderState, TranslucencyPass, Start, Last);
+		FDrawingPolicyRenderState DrawRenderState(View);
 
-								ParallelCommandListSet.AddParallelCommandList(CmdList, AnyThreadCompletionEvent);
-							}
-							Start = Last + 1;
-						}
-					}
-				}
+		// If downsampling we need to render in the separate buffer. Otherwise we also need to render offscreen to apply TPT_TranslucencyAfterDOF
+		if (RenderInSeparateTranslucency(SceneContext, TranslucencyPass, View.TranslucentPrimSet.SortedPrimsNum.DisableOffscreenRendering(TranslucencyPass)))
+		{
+			FIntPoint ScaledSize;
+			float DownsamplingScale = 1.f;
+			SceneContext.GetSeparateTranslucencyDimensions(ScaledSize, DownsamplingScale);
+			if (DownsamplingScale < 1.f)
+			{
+				SetupDownsampledTranslucencyViewUniformBuffer(RHICmdList, View);
+			}
+
+			BeginTimingSeparateTranslucencyPass(RHICmdList, View);
+			SceneContext.BeginRenderingSeparateTranslucency(RHICmdList, View, ViewIndex == 0);
+
+			// Draw only translucent prims that are in the SeparateTranslucency pass
+			DrawRenderState.SetDepthStencilState(TStaticDepthStencilState<false, CF_DepthNearOrEqual>::GetRHI());
+
+			if (bUseParallel)
+			{
+				RenderViewTranslucencyParallel(RHICmdList, View, DrawRenderState, TranslucencyPass);
+			}
+			else
+			{
+				RenderViewTranslucency(RHICmdList, View, DrawRenderState, TranslucencyPass);
 			}
 
+			SceneContext.FinishRenderingSeparateTranslucency(RHICmdList, View);
 			EndTimingSeparateTranslucencyPass(RHICmdList, View);
+
+			if (TranslucencyPass != ETranslucencyPass::TPT_TranslucencyAfterDOF)
+			{
+				FTranslucencyDrawingPolicyFactory::UpsampleTranslucency(RHICmdList, View, false);
+			}
+		}
+		else
+		{
+			SceneContext.BeginRenderingTranslucency(RHICmdList, View, ViewIndex == 0);
+			DrawRenderState.SetDepthStencilState(TStaticDepthStencilState<false, CF_DepthNearOrEqual>::GetRHI());
+
+			if (bUseParallel)
+			{
+				RenderViewTranslucencyParallel(RHICmdList, View, DrawRenderState, TranslucencyPass);
+			}
+			else
+			{
+				RenderViewTranslucency(RHICmdList, View, DrawRenderState, TranslucencyPass);
+			}
+
+			// SceneContext.FinishRenderingTranslucency(RHICmdList, View);
 		}
+
+#if STATS
+		if (View.ViewState && IsMainTranslucencyPass(TranslucencyPass))
+		{
+			STAT(View.ViewState->TranslucencyTimer.End(RHICmdList));
+		}
+#endif
 	}
-	
-	if (bRequiresSeperateTranslucencyResolve)
+}
+
+class FTranslucencyUpsamplingPS : public FGlobalShader
+{
+protected:
+
+	static bool ShouldCache(EShaderPlatform Platform)
 	{
-		SceneContext.FinishRenderingSeparateTranslucency(RHICmdList);
+		return IsFeatureLevelSupported(Platform, ERHIFeatureLevel::SM4);
 	}
-}
 
-static TAutoConsoleVariable<int32> CVarParallelTranslucency(
-	TEXT("r.ParallelTranslucency"),
-	1,
-	TEXT("Toggles parallel translucency rendering. Parallel rendering must be enabled for this to have an effect."),
-	ECVF_RenderThreadSafe
-	);
+	/** Default constructor. */
+	FTranslucencyUpsamplingPS(bool InbUseNearestDepthNeighborUpsample)
+		: bUseNearestDepthNeighborUpsample(InbUseNearestDepthNeighborUpsample)
+	{
+	}
 
-void FDeferredShadingSceneRenderer::DrawAllTranslucencyPasses(FRHICommandListImmediate& RHICmdList, const FViewInfo& View, const FDrawingPolicyRenderState& DrawRenderState, ETranslucencyPass::Type TranslucencyPass)
-{
-	// Draw translucent prims
-	View.TranslucentPrimSet.DrawPrimitives(RHICmdList, View, DrawRenderState, *this, TranslucencyPass);
+	FShaderParameter LowResColorTexelSize;
+	FShaderResourceParameter SceneDepthTexture;
+	FShaderResourceParameter LowResDepthTexture;
+	FShaderResourceParameter LowResColorTexture;
+	FShaderResourceParameter BilinearClampedSampler;
+	FShaderResourceParameter PointClampedSampler;
 
-	FTranslucencyDrawingPolicyFactory::ContextType Context(0, TranslucencyPass);
+public:
 
-	// editor and debug rendering
-	DrawViewElements<FTranslucencyDrawingPolicyFactory>(RHICmdList, View, DrawRenderState, Context, SDPG_World, false);
-	DrawViewElements<FTranslucencyDrawingPolicyFactory>(RHICmdList, View, DrawRenderState, Context, SDPG_Foreground, false);
-}
+	/** Initialization constructor. */
+	FTranslucencyUpsamplingPS(const ShaderMetaType::CompiledShaderInitializerType& Initializer, bool InbUseNearestDepthNeighborUpsample)
+		: FGlobalShader(Initializer)
+		, bUseNearestDepthNeighborUpsample(InbUseNearestDepthNeighborUpsample)
+	{
+		LowResColorTexelSize.Bind(Initializer.ParameterMap, TEXT("LowResColorTexelSize"));
+		SceneDepthTexture.Bind(Initializer.ParameterMap, TEXT("SceneDepthTexture"));
+		LowResDepthTexture.Bind(Initializer.ParameterMap, TEXT("LowResDepthTexture"));
+		LowResColorTexture.Bind(Initializer.ParameterMap, TEXT("LowResColorTexture"));
+		BilinearClampedSampler.Bind(Initializer.ParameterMap, TEXT("BilinearClampedSampler"));
+		PointClampedSampler.Bind(Initializer.ParameterMap, TEXT("PointClampedSampler"));
+	}
 
-void FDeferredShadingSceneRenderer::RenderTranslucency(FRHICommandListImmediate& RHICmdList)
-{
-	if (ShouldRenderTranslucency())
+
+
+
+	// FShader interface.
+	virtual bool Serialize(FArchive& Ar) override
 	{
-		SCOPED_DRAW_EVENT(RHICmdList, Translucency);
-		SCOPED_GPU_STAT(RHICmdList, Stat_GPU_Translucency);
+		bool bShaderHasOutdatedParameters = FGlobalShader::Serialize(Ar);
+		Ar << LowResColorTexelSize << SceneDepthTexture << LowResDepthTexture << LowResColorTexture << BilinearClampedSampler << PointClampedSampler;
+		return bShaderHasOutdatedParameters;
+	}
 
-		if (GRHICommandList.UseParallelAlgorithms() && CVarParallelTranslucency.GetValueOnRenderThread())
-		{
-			RenderTranslucencyParallel(RHICmdList);
-			return;
-		}
+	void SetParameters(FRHICommandList& RHICmdList, const FViewInfo& View)
+	{
+		const FPixelShaderRHIParamRef ShaderRHI = GetPixelShader();
+		FGlobalShader::SetParameters<FViewUniformShaderParameters>(RHICmdList, ShaderRHI, View.ViewUniformBuffer);
 
 		FSceneRenderTargets& SceneContext = FSceneRenderTargets::Get(RHICmdList);
-		bool bRequiresSeperateTranslucencyResolve = false;
+		TRefCountPtr<IPooledRenderTarget>& DownsampledTranslucency = SceneContext.SeparateTranslucencyRT;
 
-		for (int32 ViewIndex = 0; ViewIndex < Views.Num(); ViewIndex++)
-		{
-			SCOPED_CONDITIONAL_DRAW_EVENTF(RHICmdList, EventView, Views.Num() > 1, TEXT("View%d"), ViewIndex);
+		float Width = DownsampledTranslucency->GetDesc().Extent.X;
+		float Height = DownsampledTranslucency->GetDesc().Extent.Y;
+		SetShaderValue(RHICmdList, ShaderRHI, LowResColorTexelSize, FVector4(Width, Height, 1.0f / Width, 1.0f / Height));
 
-			FViewInfo& View = Views[ViewIndex];
-			FDrawingPolicyRenderState DrawRenderState(View);
+		SetTextureParameter(RHICmdList, ShaderRHI, LowResColorTexture, DownsampledTranslucency->GetRenderTargetItem().ShaderResourceTexture);
+		SetTextureParameter(RHICmdList, ShaderRHI, LowResDepthTexture, SceneContext.GetDownsampledTranslucencyDepthSurface());
+		SetTextureParameter(RHICmdList, ShaderRHI, SceneDepthTexture, SceneContext.GetSceneDepthSurface());
 
-			// non separate translucency
-			{
-#if STATS
-				if (View.ViewState)
-				{
-					View.ViewState->TranslucencyTimer.Begin(RHICmdList);
-				}
-#endif
-				bool bFirstTimeThisFrame = (ViewIndex == 0);
-				SetTranslucentRenderTarget(RHICmdList, View, ETranslucencyPass::TPT_StandardTranslucency, bFirstTimeThisFrame);
-				SetTranslucentState(RHICmdList, DrawRenderState);
+		SetSamplerParameter(RHICmdList, ShaderRHI, BilinearClampedSampler, TStaticSamplerState<SF_Bilinear,AM_Clamp,AM_Clamp,AM_Clamp>::GetRHI());
+		SetSamplerParameter(RHICmdList, ShaderRHI, PointClampedSampler, TStaticSamplerState<SF_Point,AM_Clamp,AM_Clamp,AM_Clamp>::GetRHI());
+	}
 
-				DrawAllTranslucencyPasses(RHICmdList, View, DrawRenderState, ETranslucencyPass::TPT_StandardTranslucency);
+	const bool bUseNearestDepthNeighborUpsample;
+};
 
-				const FSceneViewState* ViewState = (const FSceneViewState*)View.State;
+class FTranslucencySimpleUpsamplingPS : public FTranslucencyUpsamplingPS
+{
+protected:
+	DECLARE_SHADER_TYPE(FTranslucencySimpleUpsamplingPS, Global);
+	FTranslucencySimpleUpsamplingPS() : FTranslucencyUpsamplingPS(false) {}
+public:
+	FTranslucencySimpleUpsamplingPS(const ShaderMetaType::CompiledShaderInitializerType& Initializer) : FTranslucencyUpsamplingPS(Initializer, false) {}
+};
 
-				if (ViewState && View.Family->EngineShowFlags.VisualizeLPV)
-				{
-					FLightPropagationVolume* LightPropagationVolume = ViewState->GetLightPropagationVolume(View.GetFeatureLevel());
+IMPLEMENT_SHADER_TYPE(,FTranslucencySimpleUpsamplingPS,TEXT("TranslucencyUpsampling"),TEXT("SimpleUpsamplingPS"),SF_Pixel);
 
-					if (LightPropagationVolume)
-					{
-						LightPropagationVolume->Visualise(RHICmdList, View);
-					}
-				}
-#if STATS
-				if (View.ViewState)
-				{
-					View.ViewState->TranslucencyTimer.End(RHICmdList);
-				}
-#endif
-			}
-			
-			// separate translucency
-			{
-				if (SceneContext.IsSeparateTranslucencyActive(View))
-				{
-					bRequiresSeperateTranslucencyResolve = true;
+class FTranslucencyNearestDepthNeighborUpsamplingPS : public FTranslucencyUpsamplingPS
+{
+protected:
+	DECLARE_SHADER_TYPE(FTranslucencyNearestDepthNeighborUpsamplingPS, Global);
+	FTranslucencyNearestDepthNeighborUpsamplingPS() : FTranslucencyUpsamplingPS(true) {}
+public:
+	FTranslucencyNearestDepthNeighborUpsamplingPS(const ShaderMetaType::CompiledShaderInitializerType& Initializer) : FTranslucencyUpsamplingPS(Initializer, true) {}
+};
 
-					// always call BeginRenderingSeparateTranslucency() even if there are no primitives to we keep the RT allocated
-					FIntPoint ScaledSize;
-					float Scale = 1.0f;
-					SceneContext.GetSeparateTranslucencyDimensions(ScaledSize, Scale);
-					if (Scale < 1.0f)
-					{
-						SceneContext.GetSeparateTranslucencyDepth(RHICmdList, ScaledSize);
-						DownsampleDepthSurface(RHICmdList, SceneContext.GetSeparateTranslucencyDepthSurface(), View, Scale, false);
-
-						FViewUniformShaderParameters DownsampledTranslucencyParameters = *View.CachedViewUniformShaderParameters;
-
-						// Update the parts of DownsampledTranslucencyParameters which are dependent on the buffer size and view rect
-						View.SetupViewRectUniformBufferParameters(
-							DownsampledTranslucencyParameters,
-							ScaledSize,
-							FIntRect(View.ViewRect.Min.X * Scale, View.ViewRect.Min.Y * Scale, View.ViewRect.Max.X * Scale, View.ViewRect.Max.Y * Scale),
-							View.ViewMatrices,
-							View.PrevViewMatrices
-						);
+IMPLEMENT_SHADER_TYPE(,FTranslucencyNearestDepthNeighborUpsamplingPS,TEXT("TranslucencyUpsampling"),TEXT("NearestDepthNeighborUpsamplingPS"),SF_Pixel);
 
-						View.DownsampledTranslucencyViewUniformBuffer = TUniformBufferRef<FViewUniformShaderParameters>::CreateUniformBufferImmediate(DownsampledTranslucencyParameters, UniformBuffer_SingleFrame);
-					}
+bool UseNearestDepthNeighborUpsampleForSeparateTranslucency(const FSceneRenderTargets& SceneContext)
+{
+	FIntPoint OutScaledSize;
+	float OutScale;
+	SceneContext.GetSeparateTranslucencyDimensions(OutScaledSize, OutScale);
 
-					BeginTimingSeparateTranslucencyPass(RHICmdList, View);
+	return 	CVarSeparateTranslucencyUpsampleMode.GetValueOnRenderThread() != 0 && FMath::Abs(OutScale - .5f) < .001f;
+}
 
-					if (CVarCopySceneColorOncePerViewOnly.GetValueOnRenderThread() != 0)
-					{
-						FTranslucencyDrawingPolicyFactory::CopySceneColor(RHICmdList, View, nullptr);
-					}
-					
-					bool bFirstTimeThisFrame = (ViewIndex == 0);
-					SceneContext.BeginRenderingSeparateTranslucency(RHICmdList, View, bFirstTimeThisFrame);
+void FTranslucencyDrawingPolicyFactory::UpsampleTranslucency(FRHICommandList& RHICmdList, const FViewInfo& View, bool bOverwrite)
+{
+	SCOPED_DRAW_EVENTF(RHICmdList, EventUpsampleCopy, TEXT("Upsample translucency"));
 
-					const TIndirectArray<FMeshBatch>& WorldList = View.ViewMeshElements;
-					const TIndirectArray<FMeshBatch>& ForegroundList = View.TopViewMeshElements;
+	FSceneRenderTargets& SceneContext = FSceneRenderTargets::Get(RHICmdList);
+	SceneContext.BeginRenderingSceneColor(RHICmdList, ESimpleRenderTargetMode::EExistingColorAndDepth, FExclusiveDepthStencil::DepthRead_StencilWrite);
+	RHICmdList.SetViewport(View.ViewRect.Min.X, View.ViewRect.Min.Y, 0.0f, View.ViewRect.Max.X, View.ViewRect.Max.Y, 1.0f);
 
-					bool bRenderSeparateTranslucency = View.TranslucentPrimSet.SortedPrimsNum.Num(ETranslucencyPass::TPT_SeparateTranslucency) > 0 || WorldList.Num() || ForegroundList.Num();
+	FGraphicsPipelineStateInitializer GraphicsPSOInit;
+	RHICmdList.ApplyCachedRenderTargets(GraphicsPSOInit);
 
-					// Draw only translucent prims that are in the SeparateTranslucency pass
-					if (bRenderSeparateTranslucency)
-					{
-						DrawRenderState.SetDepthStencilState(TStaticDepthStencilState<false, CF_DepthNearOrEqual>::GetRHI());
+	GraphicsPSOInit.RasterizerState = TStaticRasterizerState<FM_Solid, CM_None>::GetRHI();
+	GraphicsPSOInit.DepthStencilState = TStaticDepthStencilState<false, CF_Always>::GetRHI();
+	if (bOverwrite) // When overwriting, we also need to set the alpha as other translucent primitive could accumulate into the buffer.
+	{
+		GraphicsPSOInit.BlendState = TStaticBlendState<>::GetRHI();
+	}
+	else
+	{
+		GraphicsPSOInit.BlendState = TStaticBlendState<CW_RGB, BO_Add, BF_One, BF_SourceAlpha>::GetRHI();
+	}
 
-						DrawAllTranslucencyPasses(RHICmdList, View, DrawRenderState, ETranslucencyPass::TPT_SeparateTranslucency);
-					}
+	TShaderMapRef<FScreenVS> ScreenVertexShader(View.ShaderMap);
+	FTranslucencyUpsamplingPS* UpsamplingPixelShader = nullptr;
+	if (UseNearestDepthNeighborUpsampleForSeparateTranslucency(SceneContext))
+	{
+		TShaderMapRef<FTranslucencyNearestDepthNeighborUpsamplingPS> PixelShader(View.ShaderMap);
+		UpsamplingPixelShader = *PixelShader;
+	}
+	else
+	{
+		TShaderMapRef<FTranslucencySimpleUpsamplingPS> PixelShader(View.ShaderMap);
+		UpsamplingPixelShader = *PixelShader;
+	}
 
+	GraphicsPSOInit.BoundShaderState.VertexDeclarationRHI = GFilterVertexDeclaration.VertexDeclarationRHI;
+	GraphicsPSOInit.BoundShaderState.VertexShaderRHI = GETSAFERHISHADER_VERTEX(*ScreenVertexShader);
+	GraphicsPSOInit.BoundShaderState.PixelShaderRHI = GETSAFERHISHADER_PIXEL(UpsamplingPixelShader);
+	GraphicsPSOInit.PrimitiveType = PT_TriangleList;
 
-					EndTimingSeparateTranslucencyPass(RHICmdList, View);
-				}
-			}
-		}
+	SetGraphicsPipelineState(RHICmdList, GraphicsPSOInit);
+	UpsamplingPixelShader->SetParameters(RHICmdList, View);
 
-		if (bRequiresSeperateTranslucencyResolve)
-		{
-			SceneContext.FinishRenderingSeparateTranslucency(RHICmdList);
-		}
-	}
+	FIntPoint OutScaledSize;
+	float OutScale;
+	SceneContext.GetSeparateTranslucencyDimensions(OutScaledSize, OutScale);
+
+	TRefCountPtr<IPooledRenderTarget>& DownsampledTranslucency = SceneContext.SeparateTranslucencyRT;
+	int32 TextureWidth = DownsampledTranslucency->GetDesc().Extent.X;
+	int32 TextureHeight = DownsampledTranslucency->GetDesc().Extent.Y;
+
+	DrawRectangle(
+		RHICmdList,
+		View.ViewRect.Min.X, View.ViewRect.Min.Y, 
+		View.ViewRect.Width(), View.ViewRect.Height(),
+		View.ViewRect.Min.X * OutScale, View.ViewRect.Min.Y * OutScale, 
+		View.ViewRect.Width() * OutScale, View.ViewRect.Height() * OutScale,
+		View.ViewRect.Size(),
+		FIntPoint(TextureWidth, TextureHeight),
+		*ScreenVertexShader,
+		EDRF_UseTriangleOptimization);
 }
+
+
+
diff --git a/Engine/Source/Runtime/Renderer/Private/TranslucentRendering.h b/Engine/Source/Runtime/Renderer/Private/TranslucentRendering.h
--- a/Engine/Source/Runtime/Renderer/Private/TranslucentRendering.h
+++ b/Engine/Source/Runtime/Renderer/Private/TranslucentRendering.h
@@ -14,6 +14,9 @@
 #include "SceneRendering.h"
 #include "VolumeRendering.h"
 
+
+bool UseNearestDepthNeighborUpsampleForSeparateTranslucency(const FSceneRenderTargets& SceneContext);
+
 /**
 * Translucent draw policy factory.
 * Creates the policies needed for rendering a mesh based on its material
@@ -25,18 +28,19 @@ public:
 	struct ContextType 
 	{
 		const FProjectedShadowInfo* TranslucentSelfShadow;
-		ETranslucencyPass::Type TranslucenyPassType;
+		ETranslucencyPass::Type TranslucencyPass;
 		ESceneRenderTargetsMode::Type TextureMode;
-		bool bSceneColorCopyIsUpToDate;
 		bool bPostAA;
 
-		ContextType(const FProjectedShadowInfo* InTranslucentSelfShadow = NULL, ETranslucencyPass::Type InTranslucenyPassType = ETranslucencyPass::TPT_StandardTranslucency, bool bPostAAIn = false, ESceneRenderTargetsMode::Type InTextureMode = ESceneRenderTargetsMode::SetTextures)
+		ContextType(const FProjectedShadowInfo* InTranslucentSelfShadow, ETranslucencyPass::Type InTranslucencyPass, bool bPostAAIn = false, ESceneRenderTargetsMode::Type InTextureMode = ESceneRenderTargetsMode::SetTextures)
 			: TranslucentSelfShadow(InTranslucentSelfShadow)
-			, TranslucenyPassType(InTranslucenyPassType)
+			, TranslucencyPass(InTranslucencyPass)
 			, TextureMode(InTextureMode)
-			, bSceneColorCopyIsUpToDate(false)
 			, bPostAA(bPostAAIn)			
 		{}
+
+		/** Whether this material should be processed now */
+		bool ShouldDraw(const FViewInfo& View, const FMaterial* Material, bool bIsSeparateTranslucency) const;
 	};
 
 	/**
@@ -73,7 +77,8 @@ public:
 	/**
 	* Resolves the scene color target and copies it for use as a source texture.
 	*/
-	static void CopySceneColor(FRHICommandList& RHICmdList, const FViewInfo& View, const FPrimitiveSceneProxy* PrimitiveSceneProxy);
+	static void CopySceneColor(FRHICommandList& RHICmdList, const FViewInfo& View);
+	static void UpsampleTranslucency(FRHICommandList& RHICmdList, const FViewInfo& View, bool bOverwrite);
 
 private:
 	/**
diff --git a/Engine/Source/Runtime/Renderer/Private/VT/TexturePagePool.cpp b/Engine/Source/Runtime/Renderer/Private/VT/TexturePagePool.cpp
new file mode 100644
--- /dev/null
+++ b/Engine/Source/Runtime/Renderer/Private/VT/TexturePagePool.cpp
@@ -0,0 +1,708 @@
+// Copyright 1998-2017 Epic Games, Inc. All Rights Reserved.
+
+#include "TexturePagePool.h"
+
+#include "VirtualTextureSpace.h"
+#include "VirtualTextureSystem.h"
+
+template< typename IndexType, typename CountType >
+void RadixSort32( IndexType* RESTRICT Dst, IndexType* RESTRICT Src, CountType Num )
+{
+	CountType Histograms[ 1024 + 2048 + 2048 ];
+	CountType* RESTRICT Histogram0 = Histograms + 0;
+	CountType* RESTRICT Histogram1 = Histogram0 + 1024;
+	CountType* RESTRICT Histogram2 = Histogram1 + 2048;
+
+	FMemory::Memzero( Histograms, sizeof( Histograms ) );
+
+	{
+		// Parallel histogram generation pass
+		const IndexType* RESTRICT s = (const IndexType* RESTRICT)Src;
+		for( CountType i = 0; i < Num; i++ )
+		{
+			IndexType j = s[i];
+			Histogram0[ ( j >>  0 ) & 1023 ]++;
+			Histogram1[ ( j >> 10 ) & 2047 ]++;
+			Histogram2[ ( j >> 21 ) & 2047 ]++;
+		}
+	}
+	{
+		// Set each histogram entry to the sum of entries preceding it
+		CountType Sum0 = 0;
+		CountType Sum1 = 0;
+		CountType Sum2 = 0;
+		for( CountType i = 0; i < 1024; i++ )
+		{
+			CountType t;
+			t = Histogram0[i] + Sum0; Histogram0[i] = Sum0 - 1; Sum0 = t;
+			t = Histogram1[i] + Sum1; Histogram1[i] = Sum1 - 1; Sum1 = t;
+			t = Histogram2[i] + Sum2; Histogram2[i] = Sum2 - 1; Sum2 = t;
+		}
+		for( CountType i = 1024; i < 2048; i++ )
+		{
+			CountType t;
+			t = Histogram1[i] + Sum1; Histogram1[i] = Sum1 - 1; Sum1 = t;
+			t = Histogram2[i] + Sum2; Histogram2[i] = Sum2 - 1; Sum2 = t;
+		}
+	}
+	{
+		// Sort pass 1
+		const IndexType* RESTRICT s = (const IndexType* RESTRICT)Src;
+		IndexType* RESTRICT d = Dst;
+		for( CountType i = 0; i < Num; i++ )
+		{
+			IndexType k = s[i];
+			d[ ++Histogram0[ ( (k >> 0) & 1023 ) ] ] = k;
+		}
+	}
+	{
+		// Sort pass 2
+		const IndexType* RESTRICT s = (const IndexType* RESTRICT)Dst;
+		IndexType* RESTRICT d = Src;
+		for( CountType i = 0; i < Num; i++ )
+		{
+			IndexType k = s[i];
+			d[ ++Histogram1[ ( (k >> 10) & 2047 ) ] ] = k;
+		}
+	}
+	{
+		// Sort pass 3
+		const IndexType* RESTRICT s = (const IndexType* RESTRICT)Src;
+		IndexType* RESTRICT d = Dst;
+		for( CountType i = 0; i < Num; i++ )
+		{
+			IndexType k = s[i];
+			d[ ++Histogram2[ ( (k >> 21) & 2047 ) ] ] = k;
+		}
+	}
+}
+
+
+FORCEINLINE uint64 EncodeSortKey( uint8 ID, uint8 vLevel, uint32 vAddress, uint16 pAddress )
+{
+#if VT_RADIX_SORT
+	uint64 Key;
+	Key  = (uint64)vAddress	<<  0;
+	Key |= (uint64)vLevel	<< 24;
+	Key |= (uint64)ID		<< 28;
+	Key |= (uint64)pAddress	<< 32;
+	return Key;
+#else
+	uint64 Key;
+	Key  = (uint64)pAddress	<<  0;
+	Key |= (uint64)vAddress	<< 16;
+	Key |= (uint64)vLevel	<< 48;
+	Key |= (uint64)ID		<< 56;
+	return Key;
+#endif
+}
+
+FORCEINLINE void DecodeSortKey( uint64 Key, uint8& ID, uint8& vLevel, uint32& vAddress, uint16& pAddress )
+{
+#if VT_RADIX_SORT
+	vAddress	= ( Key >>  0 ) & 0xffffff;
+	vLevel		= ( Key >> 24 ) & 0xf;
+	ID			= ( Key >> 28 ) & 0xf;
+	pAddress	= ( Key >> 32 );
+#else
+	pAddress	= ( Key >>  0 ) & 0xffff;
+	vAddress	= ( Key >> 16 ) & 0xffffffff;
+	vLevel		= ( Key >> 48 ) & 0xf;
+	ID			= ( Key >> 56 );
+#endif
+}
+
+
+FTexturePagePool::FTexturePagePool( uint32 Size, uint32 Dimensions )
+	: vDimensions( Dimensions )
+	, HashTable( 2048, Size )
+	, FreeHeap( Size, Size )
+{
+	Pages.SetNum( Size );
+	SortedKeys.Reserve( Size );
+
+	for( int32 i = 0; i < Pages.Num(); i++ )
+	{
+		FTexturePage& Page = Pages[i];
+
+		Page.pAddress = i;
+		Page.vAddress = 0;
+		Page.vLevel = 0;
+		Page.ID = 0xff;
+
+		FreeHeap.Add( 0, i );
+	}
+}
+
+FTexturePagePool::~FTexturePagePool()
+{}
+
+bool FTexturePagePool::AnyFreeAvailable( uint32 Frame ) const
+{
+	if( FreeHeap.Num() > 0 )
+	{
+		// Keys include vLevel to help prevent parent before child ordering
+		uint32 PageIndex = FreeHeap.Top();
+		uint32 PageFrame = FreeHeap.GetKey( PageIndex ) >> 4;
+		return PageFrame != Frame;
+	}
+
+	return false;
+}
+
+uint32 FTexturePagePool::Alloc( uint32 Frame )
+{
+	check( AnyFreeAvailable( Frame ) );
+
+	uint32 PageIndex = FreeHeap.Top();
+	FreeHeap.Pop();
+	return PageIndex;
+}
+
+void FTexturePagePool::Free( uint32 Frame, uint32 PageIndex )
+{
+	FreeHeap.Add( (Frame << 4) + (Pages[ PageIndex ].vLevel & 0xf), PageIndex );
+}
+
+void FTexturePagePool::UpdateUsage( uint32 Frame, uint32 PageIndex )
+{
+	FreeHeap.Update( (Frame << 4) + (Pages[ PageIndex ].vLevel & 0xf), PageIndex );
+}
+
+
+uint32 FTexturePagePool::FindPage( uint8 ID, uint8 vLevel, uint32 vAddress ) const
+{
+	uint16 Hash = HashPage( vLevel, vAddress, vDimensions );
+	for( uint32 PageIndex = HashTable.First( Hash ); HashTable.IsValid( PageIndex ); PageIndex = HashTable.Next( PageIndex ) )
+	{
+		if( ID			== Pages[ PageIndex ].ID &&
+			vLevel		== Pages[ PageIndex ].vLevel &&
+			vAddress	== Pages[ PageIndex ].vAddress )
+		{
+			return PageIndex;
+		}
+	}
+
+	return ~0u;
+}
+
+uint32 FTexturePagePool::FindNearestPage( uint8 ID, uint8 vLevel, uint32 vAddress ) const
+{
+	while( vLevel < 16 )
+	{
+		uint16 Hash = HashPage( vLevel, vAddress, vDimensions );
+		for( uint32 PageIndex = HashTable.First( Hash ); HashTable.IsValid( PageIndex ); PageIndex = HashTable.Next( PageIndex ) )
+		{
+			if( ID			== Pages[ PageIndex ].ID &&
+				vLevel		== Pages[ PageIndex ].vLevel &&
+				vAddress	== Pages[ PageIndex ].vAddress )
+			{
+				return PageIndex;
+			}
+		}
+
+		vLevel++;
+		vAddress &= 0xffffffff << ( vDimensions * vLevel );
+	}
+
+	return ~0u;
+}
+
+void FTexturePagePool::UnmapPage( uint16 pAddress )
+{
+#if VT_RADIX_SORT
+	const uint64 Mask = 0xffffffff;
+#else
+	const uint64 Mask = 0xffffffffffff0000ull;
+#endif
+	
+	FTexturePage& Page = Pages[ pAddress ];
+
+	if( Page.ID != 0xff )
+	{
+		// Unmap old page
+		HashTable.Remove( HashPage( Page.vLevel, Page.vAddress, vDimensions ), Page.pAddress );
+
+		uint32 Ancestor_pAddress = FindNearestPage( Page.ID, Page.vLevel, Page.vAddress );
+		uint8  Ancestor_vLevel = Ancestor_pAddress == ~0u ? 0xff : Pages[ Ancestor_pAddress ].vLevel;
+		GVirtualTextureSystem.GetSpace( Page.ID )->QueueUpdate( Page.vLevel, Page.vAddress, Ancestor_vLevel, Ancestor_pAddress );
+
+#if VT_MERGE_RESORT
+		uint64 OldKey = EncodeSortKey( Page.ID, Page.vLevel, Page.vAddress, 0 );
+		uint32 OldIndex = LowerBound( 0, SortedKeys.Num(), OldKey, Mask );
+		SortedSubIndexes.Add( ( OldIndex << 16 ) | Page.pAddress );
+#endif
+	}
+
+	Page.vLevel = 0;
+	Page.vAddress = 0;
+	Page.ID = 0xff;
+
+	SortedKeysDirty = true;
+}
+
+void FTexturePagePool::MapPage( uint8 ID, uint8 vLevel, uint32 vAddress, uint16 pAddress )
+{
+#if VT_RADIX_SORT
+	const uint64 Mask = 0xffffffff;
+#else
+	const uint64 Mask = 0xffffffffffff0000ull;
+#endif
+	
+	FTexturePage& Page = Pages[ pAddress ];
+
+	Page.vLevel = vLevel;
+	Page.vAddress = vAddress;
+	Page.ID = ID;
+
+	{
+#if VT_MERGE_RESORT
+		uint64 NewKey = EncodeSortKey( Page.ID, Page.vLevel, Page.vAddress, 0 );
+		uint32 NewIndex = UpperBound( 0, SortedKeys.Num(), NewKey, Mask );
+		SortedAddIndexes.Add( ( NewIndex << 16 ) | Page.pAddress );
+#endif
+
+		// Map new page
+		HashTable.Add( HashPage( Page.vLevel, Page.vAddress, vDimensions ), Page.pAddress );
+		GVirtualTextureSystem.GetSpace( Page.ID )->QueueUpdate( Page.vLevel, Page.vAddress, Page.vLevel, Page.pAddress );
+	}
+
+	SortedKeysDirty = true;
+}
+
+// Must call this before the below functions so that SortedKeys is up to date.
+inline void FTexturePagePool::BuildSortedKeys()
+{
+#if VT_MERGE_RESORT
+	checkSlow( SortedSubIndexes.Num() || SortedAddIndexes.Num() );
+
+	SortedSubIndexes.Sort();
+	SortedAddIndexes.Sort(
+		[this]( const uint32& A, const uint32& B )
+		{
+			const FTexturePage& PageA = Pages[ A & 0xffff ];
+			const FTexturePage& PageB = Pages[ B & 0xffff ];
+
+			uint64 KeyA = EncodeSortKey( PageA.ID, PageA.vLevel, PageA.vAddress, PageA.pAddress );
+			uint64 KeyB = EncodeSortKey( PageB.ID, PageB.vLevel, PageB.vAddress, PageB.pAddress );
+
+			return KeyA < KeyB;
+		} );
+	
+	// Copy version
+	Exchange( SortedKeys, UnsortedKeys );
+
+	uint32 NumUnsorted = UnsortedKeys.Num();
+	SortedKeys.SetNum( NumUnsorted + SortedAddIndexes.Num() - SortedSubIndexes.Num(), false );
+	
+	int32 SubI = 0;
+	int32 AddI = 0;
+	int32 UnsortedI = 0;
+	int32 SortedI = 0;
+
+	while( SortedI < SortedKeys.Num() )
+	{
+		const uint32 SubIndex = SubI < SortedSubIndexes.Num() ? ( SortedSubIndexes[ SubI ] >> 16 ) : NumUnsorted;
+		const uint32 AddIndex = AddI < SortedAddIndexes.Num() ? ( SortedAddIndexes[ AddI ] >> 16 ) : NumUnsorted;
+
+		const uint32 Interval = FMath::Min( SubIndex, AddIndex ) - UnsortedI;
+		if( Interval )
+		{
+			FMemory::Memcpy( &SortedKeys[ SortedI ], &UnsortedKeys[ UnsortedI ], Interval * sizeof( uint64 ) );
+
+			UnsortedI += Interval;
+			SortedI += Interval;
+
+			if( SortedI >= SortedKeys.Num() )
+				break;
+		}
+
+		if( SubIndex < AddIndex )
+		{
+			checkSlow( SubI < SortedSubIndexes.Num() );
+
+			// Skip hole
+			UnsortedI++;
+			SubI++;
+		}
+		else
+		{
+			checkSlow( AddI < SortedAddIndexes.Num() );
+
+			// Add new updated page
+			uint16 pAddress = SortedAddIndexes[ AddI ] & 0xffff;
+
+			FTexturePage& Page = Pages[ pAddress ];
+			SortedKeys[ SortedI ] = EncodeSortKey( Page.ID, Page.vLevel, Page.vAddress, Page.pAddress );
+
+			SortedI++;
+			AddI++;
+		}
+	}
+
+	SortedSubIndexes.Reset();
+	SortedAddIndexes.Reset();
+
+	SortedKeysDirty = false;
+#elif VT_RADIX_SORT
+	UnsortedKeys.Reset();
+
+	for( auto Page : Pages )
+	{
+		if( Page.ID != 0xff )
+		{
+			UnsortedKeys.Add( EncodeSortKey( Page.ID, Page.vLevel, Page.vAddress, Page.pAddress ) );
+		}
+	}
+
+	SortedKeys.SetNum( UnsortedKeys.Num(), false );
+
+	RadixSort32( SortedKeys.GetData(), UnsortedKeys.GetData(), (uint16)SortedKeys.Num() );
+	SortedKeysDirty = false;
+#else
+	SortedKeys.Reset();
+
+	for( auto Page : Pages )
+	{
+		if( Page.ID != 0xff )
+		{
+			SortedKeys.Add( EncodeSortKey( Page.ID, Page.vLevel, Page.vAddress, Page.pAddress ) );
+		}
+	}
+
+	Sort( SortedKeys.GetData(), SortedKeys.Num() );
+	SortedKeysDirty = false;
+#endif
+}
+
+// Binary search lower bound
+// Similar to std::lower_bound
+// Range [Min,Max)
+uint32 FTexturePagePool::LowerBound( uint32 Min, uint32 Max, uint64 SearchKey, uint64 Mask ) const
+{
+	while( Min != Max )
+	{
+		uint32 Mid = Min + (Max - Min) / 2;
+		uint64 Key = SortedKeys[ Mid ] & Mask;
+
+		if( SearchKey <= Key )
+			Max = Mid;
+		else
+			Min = Mid + 1;
+	}
+
+	return Min;
+}
+
+// Binary search upper bound
+// Similar to std::upper_bound
+// Range [Min,Max)
+uint32 FTexturePagePool::UpperBound( uint32 Min, uint32 Max, uint64 SearchKey, uint64 Mask ) const
+{
+	while( Min != Max )
+	{
+		uint32 Mid = Min + (Max - Min) / 2;
+		uint64 Key = SortedKeys[ Mid ] & Mask;
+
+		if( SearchKey < Key )
+			Max = Mid;
+		else
+			Min = Mid + 1;
+	}
+
+	return Min;
+}
+
+// Binary search equal range
+// Similar to std::equal_range
+// Range [Min,Max)
+uint32 FTexturePagePool::EqualRange( uint32 Min, uint32 Max, uint64 SearchKey, uint64 Mask ) const
+{
+	while( Min != Max )
+	{
+		uint32 Mid = Min + (Max - Min) / 2;
+		uint64 Key = SortedKeys[ Mid ] & Mask;
+
+		if( SearchKey < Key )
+		{
+			Max = Mid;
+		}
+		else if( SearchKey > Key )
+		{
+			Min = Mid + 1;
+		}
+		else
+		{	// Range straddles Mid. Search both sides and return.
+			Min = LowerBound( Min,     Mid, SearchKey, Mask );
+			Max = UpperBound( Mid + 1, Max, SearchKey, Mask );
+			return Min | ( Max << 16 );
+		}
+	}
+
+	return 0;
+}
+
+void FTexturePagePool::RefreshEntirePageTable( uint8 ID, TArray< FPageUpdate >* Output )
+{
+	if( SortedKeysDirty )
+	{
+		BuildSortedKeys();
+	}
+
+	// TODO match ID
+
+	for( int i = SortedKeys.Num() - 1; i >= 0; i-- )
+	{
+		FPageUpdate Update;
+		uint8 Update_ID;
+		DecodeSortKey( SortedKeys[i], Update_ID, Update.vLevel, Update.vAddress, Update.pAddress );
+		Update.vLogSize = Update.vLevel;
+
+		for( int Mip = Update.vLevel; Mip >= 0; Mip-- )
+		{
+			Output[ Mip ].Add( Update );
+		}
+	}
+}
+
+/*
+======================
+Update entry in page table for this page and entries for all of its unmapped descendants.
+
+If no mapped descendants then this is a single square per mip.
+If there are mapped descendants then draw those on top using painters algorithm.
+Outputs list of FPageTableUpdate which will be drawn on the GPU to the page table.
+======================
+*/
+void FTexturePagePool::ExpandPageTableUpdatePainters( uint8 ID, FPageUpdate Update, TArray< FPageUpdate >* Output )
+{
+	if( SortedKeysDirty )
+	{
+		BuildSortedKeys();
+	}
+
+	static TArray< FPageUpdate > LoopOutput;
+
+	LoopOutput.Reset();
+
+	uint8  vLogSize = Update.vLogSize;
+	uint32 vAddress = Update.vAddress;
+	
+	Output[ vLogSize ].Add( Update );
+
+	// Start with input quad
+	LoopOutput.Add( Update );
+
+	uint32 SearchRange = SortedKeys.Num();
+	
+	for( uint32 Mip = vLogSize; Mip > 0; )
+	{
+		Mip--;
+		uint64 SearchKey = EncodeSortKey( ID, Mip, vAddress, 0 );
+#if VT_RADIX_SORT
+		uint64 Mask = ( 0xffffffff << ( vDimensions * vLogSize ) ) & 0xffffffff;
+#else
+		uint64 Mask = 0xffffffffffff0000ull << ( vDimensions * vLogSize );
+#endif
+		
+		uint32 DescendantRange = EqualRange( 0, SearchRange, SearchKey, Mask );
+		if( DescendantRange != 0 )
+		{
+			uint32 DescendantMin = DescendantRange & 0xffff;
+			uint32 DescendantMax = DescendantRange >> 16;
+
+			// List is sorted by level so lower levels must be earlier in the list than what we found.
+			SearchRange = DescendantMin;
+
+			for( uint32 DescendantIndex = DescendantMin; DescendantIndex < DescendantMax; DescendantIndex++ )
+			{
+				checkSlow( SearchKey == ( SortedKeys[ DescendantIndex ] & Mask ) );
+
+				FPageUpdate Descendant;
+				uint8 Descendant_ID, Descendant_Level;
+				DecodeSortKey( SortedKeys[ DescendantIndex ], Descendant_ID, Descendant_Level, Descendant.vAddress, Descendant.pAddress );
+
+				Descendant.vLevel	= Mip;
+				Descendant.vLogSize	= Mip;
+
+				checkSlow( Descendant_ID == ID );
+				checkSlow( Descendant_Level == Mip );
+
+				// Mask out low bits
+				uint32 Ancestor_vAddress = Descendant.vAddress & ( 0xffffffff << ( vDimensions * vLogSize ) );
+				checkSlow( Ancestor_vAddress == vAddress );
+
+				LoopOutput.Add( Descendant );
+			}
+		}
+
+		Output[ Mip ].Append( LoopOutput );
+	}
+}
+
+/*
+======================
+Update entry in page table for this page and entries for all of its unmapped descendants.
+
+If no mapped descendants then this is a single square per mip.
+If there are mapped descendants then break it up into many squares in quadtree order with holes for any already mapped pages.
+Outputs list of FPageTableUpdate which will be drawn on the GPU to the page table.
+======================
+*/
+void FTexturePagePool::ExpandPageTableUpdateMasked( uint8 ID, FPageUpdate Update, TArray< FPageUpdate >* Output )
+{
+	if( SortedKeysDirty )
+	{
+		BuildSortedKeys();
+	}
+
+	static TArray< FPageUpdate > LoopInput;
+	static TArray< FPageUpdate > LoopOutput;
+	static TArray< FPageUpdate > Stack;
+
+	LoopInput.Reset();
+	LoopOutput.Reset();
+	checkSlow( Stack.Num() == 0 );
+
+	uint8  vLogSize = Update.vLogSize;
+	uint32 vAddress = Update.vAddress;
+	
+	Output[ vLogSize ].Add( Update );
+
+	// Start with input quad
+	LoopOutput.Add( Update );
+
+	uint32 SearchRange = SortedKeys.Num();
+	
+	for( uint32 Mip = vLogSize; Mip > 0; )
+	{
+		Mip--;
+		uint64 SearchKey = EncodeSortKey( ID, Mip, vAddress, 0 );
+#if VT_RADIX_SORT
+		uint64 Mask = ( 0xffffffff << ( vDimensions * vLogSize ) ) & 0xffffffff;
+#else
+		uint64 Mask = 0xffffffffffff0000ull << ( vDimensions * vLogSize );
+#endif
+		
+		uint32 DescendantRange = EqualRange( 0, SearchRange, SearchKey, Mask );
+		if( DescendantRange != 0 )
+		{
+			uint32 DescendantMin = DescendantRange & 0xffff;
+			uint32 DescendantMax = DescendantRange >> 16;
+
+			// List is sorted by level so lower levels must be earlier in the list than what we found.
+			SearchRange = DescendantMin;
+
+			// Ping pong input and output
+			Exchange( LoopInput, LoopOutput );
+			LoopOutput.Reset();
+			int32 InputIndex = 0;
+
+			Update = LoopInput[ InputIndex++ ];
+
+			for( uint32 DescendantIndex = DescendantMin; DescendantIndex < DescendantMax; )
+			{
+				checkSlow( SearchKey == ( SortedKeys[ DescendantIndex ] & Mask ) );
+
+				FPageUpdate Descendant;
+				uint8 Descendant_ID, Descendant_Level;
+				DecodeSortKey( SortedKeys[ DescendantIndex ], Descendant_ID, Descendant_Level, Descendant.vAddress, Descendant.pAddress );
+
+				Descendant.vLevel	= Mip;
+				Descendant.vLogSize	= Mip;
+
+				checkSlow( Descendant_ID == ID );
+				checkSlow( Descendant_Level == Mip );
+
+				// Mask out low bits
+				uint32 Ancestor_vAddress = Descendant.vAddress & ( 0xffffffff << ( vDimensions * vLogSize ) );
+				checkSlow( Ancestor_vAddress == vAddress );
+
+				uint32 UpdateSize		= 1 << ( vDimensions * Update.vLogSize );
+				uint32 DescendantSize	= 1 << ( vDimensions * Descendant.vLogSize );
+
+				checkSlow( Update.vLogSize >= Mip );
+
+				Update.Check( vDimensions );
+				Descendant.Check( vDimensions );
+
+				// Find if Update intersects with Descendant
+
+				// Is update past descendant?
+				if( Update.vAddress > Descendant.vAddress )
+				{
+					checkSlow( Update.vAddress >= Descendant.vAddress + DescendantSize );
+					// Move to next descendant
+					DescendantIndex++;
+					continue;
+				}
+				// Is update quad before descendant quad and doesn't intersect?
+				else if( Update.vAddress + UpdateSize <= Descendant.vAddress )
+				{
+					// Output this update and fetch next
+					LoopOutput.Add( Update );
+				}
+				// Does update quad equal descendant quad?
+				else if(	Update.vAddress	== Descendant.vAddress &&
+							Update.vLogSize	== Descendant.vLogSize )
+				{
+					// Move to next descendant
+					DescendantIndex++;
+					// Toss this update and fetch next
+				}
+				else
+				{
+					checkSlow( Update.vLogSize > Mip );
+
+					// Update intersects with Descendant but isn't the same size
+					// Split update into 4 for 2D, 8 for 3D
+					Update.vLogSize--;
+					for( uint32 Sibling = (1 << vDimensions) - 1; Sibling > 0; Sibling-- )
+					{
+						Stack.Push( FPageUpdate( Update, Sibling, vDimensions ) );
+					}
+					continue;
+				}
+
+				// Fetch next update
+				if( Stack.Num() )
+				{
+					Update = Stack.Pop( false );
+				}
+				else if( InputIndex < LoopInput.Num() )
+				{
+					Update = LoopInput[ InputIndex++ ];
+				}
+				else
+				{
+					// No more input
+					Update.vLogSize = 0xff;
+					break;
+				}
+			}
+
+			// If update was still being worked with add it
+			if( Update.vLogSize != 0xff )
+			{
+				LoopOutput.Add( Update );
+			}
+			// Add remaining stack to output
+			while( Stack.Num() )
+			{
+				LoopOutput.Add( Stack.Pop( false ) );
+			}
+			// Add remaining input to output
+			LoopOutput.Append( LoopInput.GetData() + InputIndex, LoopInput.Num() - InputIndex );
+		}
+		
+		if( LoopOutput.Num() == 0 )
+		{
+			// Completely masked out by descendants
+			break;
+		}
+		else
+		{
+			Output[ Mip ].Append( LoopOutput );
+		}
+	}
+}
\ No newline at end of file
