Copying //UE4/Dev-Rendering to //UE4/Dev-Main (Source: //UE4/Dev-Rendering @ 3461187)

diff --git a/Engine/Source/Runtime/Renderer/Private/VT/TexturePagePool.h b/Engine/Source/Runtime/Renderer/Private/VT/TexturePagePool.h
new file mode 100644
--- /dev/null
+++ b/Engine/Source/Runtime/Renderer/Private/VT/TexturePagePool.h
@@ -0,0 +1,67 @@
+// Copyright 1998-2017 Epic Games, Inc. All Rights Reserved.
+
+#pragma once
+
+#include "CoreMinimal.h"
+#include "VirtualTextureShared.h"
+#include "BinaryHeap.h"
+#include "HashTable.h"
+
+#define VT_MERGE_RESORT		1
+#define VT_RADIX_SORT		0
+
+// 64k x 64k virtual pages
+// 256 x 256 physical pages
+struct FTexturePage
+{
+	// Address is Morton order, relative to mip 0
+	uint32	vAddress;
+	uint16	pAddress;
+	uint8	vLevel;
+	uint8	ID;
+};
+
+class FTexturePagePool
+{
+public:
+				FTexturePagePool( uint32 Size, uint32 Dimensions );
+				~FTexturePagePool();
+
+	uint32				GetSize() const						{ return Pages.Num(); }
+	const FTexturePage&	GetPage( uint16 pAddress ) const	{ return Pages[ pAddress ]; }
+					
+	bool		AnyFreeAvailable( uint32 Frame ) const;
+	uint32		Alloc( uint32 Frame );
+	void		Free( uint32 Frame, uint32 PageIndex );
+	void		UpdateUsage( uint32 Frame, uint32 PageIndex );
+
+	uint32		FindPage( uint8 ID, uint8 vLevel, uint32 vAddress ) const;
+	uint32		FindNearestPage( uint8 ID, uint8 vLevel, uint32 vAddress ) const;
+
+	void		UnmapPage( uint16 pAddress );
+	void		MapPage( uint8 ID, uint8 vLevel, uint32 vAddress, uint16 pAddress );
+
+	void		RefreshEntirePageTable( uint8 ID, TArray< FPageUpdate >* Output );
+	void		ExpandPageTableUpdatePainters( uint8 ID, FPageUpdate Update, TArray< FPageUpdate >* Output );
+	void		ExpandPageTableUpdateMasked( uint8 ID, FPageUpdate Update, TArray< FPageUpdate >* Output );
+
+private:
+	void		BuildSortedKeys();
+	uint32		LowerBound( uint32 Min, uint32 Max, uint64 SearchKey, uint64 Mask ) const;
+	uint32		UpperBound( uint32 Min, uint32 Max, uint64 SearchKey, uint64 Mask ) const;
+	uint32		EqualRange( uint32 Min, uint32 Max, uint64 SearchKey, uint64 Mask ) const;
+
+	uint32		vDimensions;
+
+	TArray< FTexturePage >	Pages;
+	
+	FHashTable						HashTable;
+	FBinaryHeap< uint32, uint16 >	FreeHeap;
+
+	TArray< uint64 >	UnsortedKeys;
+	TArray< uint64 >	SortedKeys;
+	bool				SortedKeysDirty;
+	
+	TArray< uint32 >	SortedSubIndexes;
+	TArray< uint32 >	SortedAddIndexes;
+};
\ No newline at end of file
diff --git a/Engine/Source/Runtime/Renderer/Private/VT/UniquePageList.h b/Engine/Source/Runtime/Renderer/Private/VT/UniquePageList.h
new file mode 100644
--- /dev/null
+++ b/Engine/Source/Runtime/Renderer/Private/VT/UniquePageList.h
@@ -0,0 +1,126 @@
+// Copyright 1998-2017 Epic Games, Inc. All Rights Reserved.
+
+#pragma once
+
+#include "CoreMinimal.h"
+#include "VirtualTextureShared.h"
+#include "HashTable.h"
+
+// 4k x 4k virtual pages
+// 16 IDs
+FORCEINLINE uint32 EncodePage( uint32 ID, uint32 vLevel, uint32 vPosition )
+{
+	uint32 Page;
+	Page  = vPosition	<<  0;
+	Page |= vLevel		<< 24;
+	Page |= ID			<< 28;
+	return Page;
+}
+
+FORCEINLINE void DecodePage( uint32 Page, uint32& ID, uint32& vLevel, uint32& vPosition )
+{
+	vPosition	= ( Page >>  0 ) & 0xffffff;
+	vLevel		= ( Page >> 24 ) & 0xf;
+	ID			= ( Page >> 28 );
+}
+
+
+class FUniquePageList
+{
+public:
+			FUniquePageList();
+
+	void	Add( uint32 Page, uint32 Count );
+	void	ExpandByMips( uint32 NumMips );
+
+	uint32	GetNum() const					{ return NumPages; }
+	uint32	GetPage( uint32 Index ) const	{ return Pages[ Index ]; }
+	uint32	GetCount( uint32 Index ) const	{ return Counts[ Index ]; }
+	
+	uint8	NumLevels[16];
+	uint8	Dimensions[16];
+
+private:
+	enum
+	{
+		HashSize		= 1024,
+		MaxUniquePages	= 4096,
+	};
+
+	uint32	NumPages;
+	uint32	Pages[ MaxUniquePages ];
+	uint16	Counts[ MaxUniquePages ];
+	
+	TStaticHashTable< HashSize, MaxUniquePages > HashTable;
+};
+
+FUniquePageList::FUniquePageList()
+	: NumPages( 0 )
+{}
+
+void FUniquePageList::Add( uint32 Page, uint32 Count )
+{
+	uint32 ID, vLevel, vPosition;
+	DecodePage( Page, ID, vLevel, vPosition );
+	uint32 vDimensions = Dimensions[ ID ];
+
+	// Search hash table
+	uint16 Hash = HashPage( vLevel, vPosition, vDimensions );
+	uint16 Index;
+	for( Index = HashTable.First( Hash ); HashTable.IsValid( Index ); Index = HashTable.Next( Index ) )
+	{
+		if( Page == Pages[ Index ] )
+		{
+			break;
+		}
+	}
+	if( !HashTable.IsValid( Index ) )
+	{
+		if( NumPages == MaxUniquePages )
+		{
+			// Ran out of space
+			return;
+		}
+
+		// Add page
+		Index = NumPages++;
+		HashTable.Add( Hash, Index );
+
+		Pages[ Index ] = Page;
+		Counts[ Index ] = 0;
+	}
+	
+	// Add count
+	Counts[ Index ] += Count;
+}
+
+// Expanding the list by mips gives look ahead and faster time to first data when many high res pages are requested.
+void FUniquePageList::ExpandByMips( uint32 NumMips )
+{
+	// TODO cache parent page chains
+
+	for( uint32 i = 0, Num = NumPages; i < Num; i++ )
+	{
+		uint64 Page = Pages[i];
+
+		uint32 ID, vLevel, vPosition;
+		DecodePage( Page, ID, vLevel, vPosition );
+		uint32 vDimensions = Dimensions[ ID ];
+		uint32 Count = Counts[i];
+
+		for( uint32 Mip = 0; Mip < NumMips; Mip++ )
+		{
+			vLevel++;
+			
+			if( vLevel >= NumLevels[ ID ] )
+			{
+				break;
+			}
+
+			// Mask out low bits
+			vPosition &= 0xffffffff << ( vDimensions * vLevel );
+
+			Add( EncodePage( ID, vLevel, vPosition ), Count );
+		}
+	}
+}
\ No newline at end of file
diff --git a/Engine/Source/Runtime/Renderer/Private/VT/VirtualTexture.h b/Engine/Source/Runtime/Renderer/Private/VT/VirtualTexture.h
new file mode 100644
--- /dev/null
+++ b/Engine/Source/Runtime/Renderer/Private/VT/VirtualTexture.h
@@ -0,0 +1,29 @@
+// Copyright 1998-2017 Epic Games, Inc. All Rights Reserved.
+
+#pragma once
+
+#include "CoreMinimal.h"
+#include "RHI.h"
+
+class IVirtualTexture
+{
+public:
+					IVirtualTexture( uint32 InSizeX, uint32 InSizeY, uint32 InSizeZ )
+						: SizeX( InSizeX )
+						, SizeY( InSizeY )
+						, SizeZ( InSizeZ )
+					{}
+
+	virtual			~IVirtualTexture() = 0;
+
+	// Locates page and returns if page data can be provided at this moment.
+	virtual bool	LocatePageData( uint8 vLevel, uint32 vAddress, void* RESTRICT& Location ) const = 0;
+
+	// Produces and fills in texture data for the page in the physical texture(s).
+	virtual void	ProducePageData( FRHICommandList& RHICmdList, ERHIFeatureLevel::Type FeatureLevel, uint8 vLevel, uint32 vAddress, uint16 pAddress, void* RESTRICT Location ) const = 0;
+	
+	// Size in pages
+	uint32	SizeX;
+	uint32	SizeY;
+	uint32	SizeZ;
+};
diff --git a/Engine/Source/Runtime/Renderer/Private/VT/VirtualTextureAllocator.cpp b/Engine/Source/Runtime/Renderer/Private/VT/VirtualTextureAllocator.cpp
new file mode 100644
--- /dev/null
+++ b/Engine/Source/Runtime/Renderer/Private/VT/VirtualTextureAllocator.cpp
@@ -0,0 +1,167 @@
+// Copyright 1998-2017 Epic Games, Inc. All Rights Reserved.
+
+#include "VirtualTextureAllocator.h"
+
+#include "VirtualTexture.h"
+
+FVirtualTextureAllocator::FVirtualTextureAllocator( uint32 Size, uint32 Dimensions )
+	: vDimensions( Dimensions )
+{
+	uint8 LogSize = FMath::CeilLogTwo( Size );
+	
+	// Start with one empty block
+	AddressBlocks.Add( FAddressBlock( LogSize ) );
+	SortedBlocks.AddUninitialized();
+	SortedBlocks[0].vAddress = 0;
+	SortedBlocks[0].Index = 0;
+	
+	// Init free list
+	FreeList.AddUninitialized( LogSize + 1 );
+	for( uint8 i = 0; i < LogSize; i++ )
+	{
+		FreeList[i] = 0xffff;
+	}
+	FreeList[ LogSize ] = 0;
+}
+
+// returns SortedIndex
+uint32 FVirtualTextureAllocator::Find( uint32 vAddress ) const
+{
+	uint32 Min = 0;
+	uint32 Max = SortedBlocks.Num();
+	
+	// Binary search for upper bound
+	while( Min != Max )
+	{
+		uint32 Mid = Min + (Max - Min) / 2;
+		uint32 Key = SortedBlocks[ Mid ].vAddress;
+
+		if( vAddress < Key )
+			Max = Mid;
+		else
+			Min = Mid + 1;
+	}
+
+	return Min - 1;
+}
+
+IVirtualTexture* FVirtualTextureAllocator::Find( uint32 vAddress, uint32& Local_vAddress ) const
+{
+	uint32 SortedIndex = Find( vAddress );
+
+	const FSortedBlock&	SortedBlock = SortedBlocks[ SortedIndex ];
+	const FAddressBlock& AddressBlock = AddressBlocks[ SortedBlock.Index ];
+	checkSlow( SortedBlock.vAddress == AddressBlock.vAddress );
+
+	uint32 BlockSize = 1 << ( vDimensions * AddressBlock.vLogSize );
+	if( vAddress >= AddressBlock.vAddress &&
+		vAddress <  AddressBlock.vAddress + BlockSize )
+	{
+		Local_vAddress = vAddress - AddressBlock.vAddress;
+		// TODO mip bias
+		return AddressBlock.VT;
+	}
+
+	return nullptr;
+}
+
+uint32 FVirtualTextureAllocator::Alloc( IVirtualTexture* VT )
+{
+	uint32 BlockSize = FMath::Max( VT->SizeX, VT->SizeY );
+	uint8 vLogSize = FMath::CeilLogTwo( BlockSize );
+
+	// Find smallest free that fits
+	for( int i = vLogSize; i < FreeList.Num(); i++ )
+	{
+		uint16 FreeIndex = FreeList[i];
+		if( FreeIndex != 0xffff )
+		{
+			// Found free
+			FAddressBlock& AllocBlock = AddressBlocks[ FreeIndex ];
+			checkSlow( AllocBlock.VT == nullptr );
+			checkSlow( AllocBlock.PrevFree == 0xffff );
+
+			// Remove from free list
+			FreeList[i] = AllocBlock.NextFree;
+			if( AllocBlock.NextFree != 0xffff )
+			{
+				AddressBlocks[ AllocBlock.NextFree ].PrevFree = 0xffff;
+				AllocBlock.NextFree = 0xffff;
+			}
+
+			AllocBlock.VT = VT;
+
+			// Add to hash table
+			uint16 Key = reinterpret_cast< UPTRINT >( VT ) / 16;
+			HashTable.Add( Key, FreeIndex );
+
+			// Recursive subdivide until the right size
+			int NumNewBlocks = 0;
+			while( AllocBlock.vLogSize > vLogSize )
+			{
+				AllocBlock.vLogSize--;
+				const uint32 NumSiblings = (1 << vDimensions) - 1;
+				for( uint32 Sibling = NumSiblings; Sibling > 0; Sibling-- )
+				{
+					AddressBlocks.Add( FAddressBlock( AllocBlock, Sibling, vDimensions ) );
+				}
+				NumNewBlocks += NumSiblings;
+			}
+
+			uint32 SortedIndex = Find( AllocBlock.vAddress ) + 1;
+			checkSlow( AllocBlock.vAddress == SortedBlocks[ SortedIndex ].vAddress );
+
+			// Make room for newly added
+			SortedBlocks.InsertUninitialized( SortedIndex, NumNewBlocks );
+
+			for( int Block = 0; Block < NumNewBlocks; Block++ )
+			{
+				uint32 Index = AddressBlocks.Num() - NumNewBlocks + Block;
+				FAddressBlock AddressBlock = AddressBlocks[ Index ];
+
+				// Place on free list
+				AddressBlock.NextFree = FreeList[ AddressBlock.vLogSize ];
+				AddressBlocks[ AddressBlock.NextFree ].PrevFree = Index;
+				FreeList[ AddressBlock.vLogSize ] = Index;
+
+				// Add to sorted list
+				SortedBlocks[ SortedIndex + Block ].vAddress = AddressBlock.vAddress;
+				SortedBlocks[ SortedIndex + Block ].Index = Index;
+			}
+
+			return AllocBlock.vAddress;
+		}
+	}
+
+	return ~0u;
+}
+
+void FVirtualTextureAllocator::Free( IVirtualTexture* VT )
+{
+	// Find block index
+	uint16 Key = reinterpret_cast< UPTRINT >( VT ) / 16;
+	uint32 Index;
+	for( Index = HashTable.First( Key ); HashTable.IsValid( Index ); Index = HashTable.Next( Index ) )
+	{
+		if( AddressBlocks[ Index ].VT == VT )
+		{
+			break;
+		}
+	}
+	if( HashTable.IsValid( Index ) )
+	{
+		FAddressBlock& AddressBlock = AddressBlocks[ Index ];
+		checkSlow( AddressBlock.VT == VT );
+		checkSlow( AddressBlock.NextFree == 0xffff );
+		checkSlow( AddressBlock.PrevFree == 0xffff );
+
+		AddressBlock.VT = nullptr;
+
+		// TODO merge with sibling free blocks
+
+		// Place on free list
+		AddressBlock.NextFree = FreeList[ AddressBlock.vLogSize ];
+		AddressBlocks[ AddressBlock.NextFree ].PrevFree = Index;
+		FreeList[ AddressBlock.vLogSize ] = Index;
+	}
+}
\ No newline at end of file
diff --git a/Engine/Source/Runtime/Renderer/Private/VT/VirtualTextureAllocator.h b/Engine/Source/Runtime/Renderer/Private/VT/VirtualTextureAllocator.h
new file mode 100644
--- /dev/null
+++ b/Engine/Source/Runtime/Renderer/Private/VT/VirtualTextureAllocator.h
@@ -0,0 +1,69 @@
+// Copyright 1998-2017 Epic Games, Inc. All Rights Reserved.
+
+#pragma once
+
+#include "CoreMinimal.h"
+#include "VirtualTextureShared.h"
+#include "HashTable.h"
+
+class IVirtualTexture;
+
+// Allocates virtual memory address space
+class FVirtualTextureAllocator
+{
+public:
+						FVirtualTextureAllocator( uint32 Size, uint32 Dimensions );
+						~FVirtualTextureAllocator() {}
+
+	IVirtualTexture*	Find( uint32 vAddress, uint32& Local_vAddress ) const;
+	uint32				Alloc( IVirtualTexture* VT );
+	void				Free( IVirtualTexture* VT );
+	// TODO				Realloc
+
+private:
+	uint32				Find( uint32 vAddress ) const;
+
+	struct FAddressBlock
+	{
+		IVirtualTexture*	VT;
+		uint32				vAddress;
+		uint16				NextFree;
+		uint16				PrevFree;
+		uint8				vLogSize;
+		uint8				MipBias;
+
+		FAddressBlock()
+		{}
+
+		FAddressBlock( uint8 LogSize )
+			: VT( nullptr )
+			, vAddress( 0 )
+			, NextFree( 0xffff )
+			, PrevFree( 0xffff )
+			, vLogSize( LogSize )
+			, MipBias( 0 )
+		{}
+
+		FAddressBlock( const FAddressBlock& Block, uint32 Offset, uint32 Dimensions )
+			: VT( nullptr )
+			, vAddress( Block.vAddress + ( Offset << ( Dimensions * Block.vLogSize ) ) )
+			, NextFree( 0xffff )
+			, PrevFree( 0xffff )
+			, vLogSize( Block.vLogSize )
+			, MipBias( 0 )
+		{}
+	};
+
+	struct FSortedBlock
+	{
+		uint32	vAddress;
+		uint16	Index;
+	};
+
+	uint32	vDimensions;
+
+	TArray< FAddressBlock >	AddressBlocks;
+	TArray< uint16 >		FreeList;
+	TArray< FSortedBlock >	SortedBlocks;
+	FHashTable				HashTable;
+};
\ No newline at end of file
diff --git a/Engine/Source/Runtime/Renderer/Private/VT/VirtualTextureFeedback.cpp b/Engine/Source/Runtime/Renderer/Private/VT/VirtualTextureFeedback.cpp
new file mode 100644
--- /dev/null
+++ b/Engine/Source/Runtime/Renderer/Private/VT/VirtualTextureFeedback.cpp
@@ -0,0 +1,76 @@
+// Copyright 1998-2017 Epic Games, Inc. All Rights Reserved.
+
+#include "VirtualTextureFeedback.h"
+
+#include "ClearQuad.h"
+
+TGlobalResource< FVirtualTextureFeedback > GVirtualTextureFeedback;
+
+FVirtualTextureFeedback::FVirtualTextureFeedback()
+	: Size( 0, 0 )
+{}
+
+void FVirtualTextureFeedback::InitDynamicRHI()
+{}
+
+void FVirtualTextureFeedback::ReleaseDynamicRHI()
+{
+	GRenderTargetPool.FreeUnusedResource( FeedbackTextureGPU );
+	GRenderTargetPool.FreeUnusedResource( FeedbackTextureCPU );
+}
+
+void FVirtualTextureFeedback::CreateResourceGPU( FRHICommandListImmediate& RHICmdList, int32 SizeX, int32 SizeY )
+{
+	Size = FIntPoint( SizeX, SizeY );
+
+	FPooledRenderTargetDesc Desc( FPooledRenderTargetDesc::Create2DDesc( Size, PF_R32_UINT, FClearValueBinding::None, TexCreate_None, TexCreate_UAV, false ) );
+	GRenderTargetPool.FindFreeElement( RHICmdList, Desc, FeedbackTextureGPU, TEXT("VTFeedbackGPU") );
+
+	// Clear to default value
+	const uint32 ClearValue[4] = { ~0u, ~0u, ~0u, ~0u };
+	ClearUAV( RHICmdList, FeedbackTextureGPU->GetRenderTargetItem(), ClearValue );
+	RHICmdList.TransitionResource( EResourceTransitionAccess::ERWNoBarrier, EResourceTransitionPipeline::EGfxToGfx, FeedbackTextureGPU->GetRenderTargetItem().UAV );
+}
+
+void FVirtualTextureFeedback::TransferGPUToCPU( FRHICommandListImmediate& RHICmdList )
+{
+	RHICmdList.TransitionResource( EResourceTransitionAccess::EReadable, EResourceTransitionPipeline::EGfxToGfx, FeedbackTextureGPU->GetRenderTargetItem().UAV );
+
+	GRenderTargetPool.VisualizeTexture.SetCheckPoint( RHICmdList, FeedbackTextureGPU );
+	
+	FPooledRenderTargetDesc Desc( FPooledRenderTargetDesc::Create2DDesc( Size, PF_R32_UINT, FClearValueBinding::None, TexCreate_CPUReadback | TexCreate_HideInVisualizeTexture, TexCreate_None, false ) );
+	GRenderTargetPool.FindFreeElement( RHICmdList, Desc, FeedbackTextureCPU, TEXT("VTFeedbackCPU") );
+
+	// Transfer memory GPU -> CPU
+	RHICmdList.CopyToResolveTarget( FeedbackTextureGPU->GetRenderTargetItem().TargetableTexture, FeedbackTextureCPU->GetRenderTargetItem().ShaderResourceTexture, false, FResolveParams() );
+
+	GRenderTargetPool.FreeUnusedResource( FeedbackTextureGPU );
+}
+
+uint32* FVirtualTextureFeedback::Map( FRHICommandListImmediate& RHICmdList, int32& Pitch )
+{
+	uint32* FeedbackBuffer = nullptr;
+
+	if( Size.X > 0 && Size.Y > 0 )
+	{
+		uint32 IdleStart = FPlatformTime::Cycles();
+
+		int32 Temp;
+		RHICmdList.MapStagingSurface( FeedbackTextureCPU->GetRenderTargetItem().ShaderResourceTexture, *(void**)&FeedbackBuffer, Pitch, Temp );
+
+		// RHIMapStagingSurface will block until the results are ready (from the previous frame) so we need to consider this RT idle time
+		GRenderThreadIdle[ERenderThreadIdleTypes::WaitingForGPUQuery] += FPlatformTime::Cycles() - IdleStart;
+		GRenderThreadNumIdle[ERenderThreadIdleTypes::WaitingForGPUQuery]++;
+	}
+
+	return FeedbackBuffer;
+}
+
+void FVirtualTextureFeedback::Unmap( FRHICommandListImmediate& RHICmdList )
+{
+	check( Size.X > 0 && Size.Y > 0 );
+	RHICmdList.UnmapStagingSurface( FeedbackTextureCPU->GetRenderTargetItem().ShaderResourceTexture );
+
+	GRenderTargetPool.FreeUnusedResource( FeedbackTextureCPU );
+	Size = FIntPoint::ZeroValue;
+}
\ No newline at end of file
diff --git a/Engine/Source/Runtime/Renderer/Private/VT/VirtualTextureFeedback.h b/Engine/Source/Runtime/Renderer/Private/VT/VirtualTextureFeedback.h
new file mode 100644
--- /dev/null
+++ b/Engine/Source/Runtime/Renderer/Private/VT/VirtualTextureFeedback.h
@@ -0,0 +1,39 @@
+// Copyright 1998-2017 Epic Games, Inc. All Rights Reserved.
+
+#pragma once
+
+#include "CoreMinimal.h"
+#include "RenderTargetPool.h"
+
+/*
+====================================
+	Manages GPU and CPU buffers for VT feedback.
+	Shared for all views, not per view.
+
+	Should use append buffer but that requires RHI ability to copy
+	a GPU structured buffer to a CPU read only version.
+====================================
+*/
+class FVirtualTextureFeedback : public FRenderResource
+{
+public:
+					FVirtualTextureFeedback();
+					~FVirtualTextureFeedback() {}
+
+	// FRenderResource interface
+	virtual void	InitDynamicRHI() override;
+	virtual void	ReleaseDynamicRHI() override;
+
+	void			CreateResourceGPU( FRHICommandListImmediate& RHICmdList, int32 SizeX, int32 SizeY );
+	void			TransferGPUToCPU( FRHICommandListImmediate& RHICmdList );
+
+	uint32*			Map( FRHICommandListImmediate& RHICmdList, int32& Pitch );
+	void			Unmap( FRHICommandListImmediate& RHICmdList );
+
+	FIntPoint		Size;
+
+	TRefCountPtr< IPooledRenderTarget >	FeedbackTextureGPU;
+	TRefCountPtr< IPooledRenderTarget >	FeedbackTextureCPU;
+};
+
+extern TGlobalResource< FVirtualTextureFeedback > GVirtualTextureFeedback;
\ No newline at end of file
diff --git a/Engine/Source/Runtime/Renderer/Private/VT/VirtualTextureShared.h b/Engine/Source/Runtime/Renderer/Private/VT/VirtualTextureShared.h
new file mode 100644
--- /dev/null
+++ b/Engine/Source/Runtime/Renderer/Private/VT/VirtualTextureShared.h
@@ -0,0 +1,38 @@
+// Copyright 1998-2016 Epic Games, Inc. All Rights Reserved.
+
+#pragma once
+
+#include "CoreMinimal.h"
+
+FORCEINLINE uint16 HashPage( uint32 vLevel, uint32 vIndex, uint8 vDimensions )
+{
+	// Mix level into top 4 bits 
+	return ( vLevel << 6 ) ^ ( vIndex >> ( vDimensions * vLevel ) );
+}
+
+
+struct FPageUpdate
+{
+	uint32	vAddress;
+	uint16	pAddress;
+	uint8	vLevel;
+	uint8	vLogSize;
+
+	FPageUpdate()
+	{}
+
+	FPageUpdate( const FPageUpdate& Update, uint32 Offset, uint8 vDimensions )
+		: vAddress( Update.vAddress + ( Offset << ( vDimensions * Update.vLogSize ) ) )
+		, pAddress( Update.pAddress )
+		, vLevel( Update.vLevel )
+		, vLogSize( Update.vLogSize )
+	{}
+
+	inline void Check( uint8 vDimensions )
+	{
+		uint32 LowBitMask = ( 1 << ( vDimensions * vLogSize ) ) - 1;
+		checkSlow( (vAddress & LowBitMask) == 0 );
+
+		checkSlow( vLogSize <= vLevel );
+	}
+};
\ No newline at end of file
diff --git a/Engine/Source/Runtime/Renderer/Private/VT/VirtualTextureSpace.cpp b/Engine/Source/Runtime/Renderer/Private/VT/VirtualTextureSpace.cpp
new file mode 100644
--- /dev/null
+++ b/Engine/Source/Runtime/Renderer/Private/VT/VirtualTextureSpace.cpp
@@ -0,0 +1,281 @@
+// Copyright 1998-2017 Epic Games, Inc. All Rights Reserved.
+
+#include "VirtualTextureSpace.h"
+
+#include "VirtualTextureSystem.h"
+#include "SpriteIndexBuffer.h"
+#include "SceneFilterRendering.h"
+#include "RenderTargetPool.h"
+#include "GlobalShader.h"
+#include "PipelineStateCache.h"
+#include "HAL/IConsoleManager.h"
+#include "SceneUtils.h"
+
+static TAutoConsoleVariable<int32> CVarVTRefreshEntirePageTable(
+	TEXT("r.VT.RefreshEntirePageTable"),
+	0,
+	TEXT("Refreshes the entire page table texture every frame"),
+	ECVF_RenderThreadSafe
+	);
+
+static TAutoConsoleVariable<int32> CVarVTMaskedPageTableUpdates(
+	TEXT("r.VT.MaskedPageTableUpdates"),
+	1,
+	TEXT("Masks the page table update quads to reduce pixel fill costs"),
+	ECVF_RenderThreadSafe
+	);
+
+
+FVirtualTextureSpace::FVirtualTextureSpace( uint32 InSize, uint32 InDimensions, FTexturePagePool* InPool )
+	: ID( 0xff )
+	, Dimensions( InDimensions )
+	, Pool( InPool )
+	, Allocator( InSize, InDimensions )
+{
+	PageTableSize = InSize;
+	PageTableLevels = FMath::FloorLog2( PageTableSize ) + 1;
+
+	GVirtualTextureSystem.RegisterSpace( this );
+}
+
+FVirtualTextureSpace::~FVirtualTextureSpace()
+{
+	GVirtualTextureSystem.UnregisterSpace( this );
+}
+
+void FVirtualTextureSpace::InitDynamicRHI()
+{
+	{
+		// Page Table
+		FRHICommandListImmediate& RHICmdList = FRHICommandListExecutor::GetImmediateCommandList();
+		FPooledRenderTargetDesc Desc( FPooledRenderTargetDesc::Create2DDesc( FIntPoint( PageTableSize, PageTableSize ), PF_R8G8B8A8, FClearValueBinding::None, TexCreate_None, TexCreate_RenderTargetable | TexCreate_ShaderResource, false, PageTableLevels ) );
+		GRenderTargetPool.FindFreeElement( RHICmdList, Desc, PageTable, TEXT("PageTable") );
+	}
+
+	{
+		// Worst case scenario of max mip page update when all other pages are perfectly sparse at mip 0.
+		const uint32 MaxSparseRegions = Pool->GetSize();
+		const uint32 SparseRegionSize = PageTableSize / (uint32)FMath::Pow( (float)MaxSparseRegions, 1.0f / Dimensions );
+		const uint32 PerRegionMaxExpansion = ( (1 << Dimensions) - 1 ) * FMath::FloorLog2( SparseRegionSize );
+		const uint32 MaxExpansionMip0 = PerRegionMaxExpansion * MaxSparseRegions;
+		const uint32 Mip0Updates = 64;
+		const uint32 MaxUpdates = MaxExpansionMip0 + (PageTableLevels - 1) + Mip0Updates;
+
+		// Update Buffer
+		FRHIResourceCreateInfo CreateInfo;
+		UpdateBuffer = RHICreateStructuredBuffer( sizeof( FPageUpdate ), MaxUpdates * sizeof( FPageUpdate ), BUF_ShaderResource | BUF_Volatile, CreateInfo );
+		UpdateBufferSRV = RHICreateShaderResourceView( UpdateBuffer );
+	}
+}
+
+void FVirtualTextureSpace::ReleaseDynamicRHI()
+{
+	GRenderTargetPool.FreeUnusedResource( PageTable );
+
+	UpdateBuffer.SafeRelease();
+	UpdateBufferSRV.SafeRelease();
+}
+
+void FVirtualTextureSpace::QueueUpdate( uint8 vLogSize, uint32 vAddress, uint8 vLevel, uint16 pAddress )
+{
+	FPageUpdate Update;
+	Update.vAddress	= vAddress;
+	Update.pAddress	= pAddress;
+	Update.vLevel	= vLevel;
+	Update.vLogSize	= vLogSize;
+	Update.Check( Dimensions );
+
+	PageTableUpdates.Add( Update );
+}
+
+
+TGlobalResource< FSpriteIndexBuffer<8> > GQuadIndexBuffer;
+
+class FPageTableUpdateVS : public FGlobalShader
+{
+	DECLARE_SHADER_TYPE(FPageTableUpdateVS,Global);
+
+	FPageTableUpdateVS() {}
+	
+public:
+	FShaderParameter			PageTableSize;
+	FShaderParameter			FirstUpdate;
+	FShaderParameter			NumUpdates;
+	FShaderResourceParameter	UpdateBuffer;
+
+	FPageTableUpdateVS(const ShaderMetaType::CompiledShaderInitializerType& Initializer)
+		: FGlobalShader(Initializer)
+	{
+		PageTableSize.Bind( Initializer.ParameterMap, TEXT("PageTableSize") );
+		FirstUpdate.Bind( Initializer.ParameterMap, TEXT("FirstUpdate") );
+		NumUpdates.Bind( Initializer.ParameterMap, TEXT("NumUpdates") );
+		UpdateBuffer.Bind( Initializer.ParameterMap, TEXT("UpdateBuffer") );
+	}
+
+	static bool ShouldCache(EShaderPlatform Platform)
+	{
+		return IsFeatureLevelSupported(Platform, ERHIFeatureLevel::SM5) && !IsHlslccShaderPlatform(Platform);
+	}
+
+	virtual bool Serialize( FArchive& Ar ) override
+	{
+		bool bShaderHasOutdatedParameters = FGlobalShader::Serialize(Ar);
+		Ar << PageTableSize;
+		Ar << FirstUpdate;
+		Ar << NumUpdates;
+		Ar << UpdateBuffer;
+		return bShaderHasOutdatedParameters;
+	}
+};
+
+class FPageTableUpdatePS : public FGlobalShader
+{
+	DECLARE_SHADER_TYPE(FPageTableUpdatePS, Global);
+
+	static bool ShouldCache( EShaderPlatform Platform )
+	{
+		return IsFeatureLevelSupported( Platform, ERHIFeatureLevel::SM5 ) && !IsHlslccShaderPlatform(Platform);
+	}
+
+	static void ModifyCompilationEnvironment( EShaderPlatform Platform, FShaderCompilerEnvironment& OutEnvironment )
+	{
+		FGlobalShader::ModifyCompilationEnvironment( Platform, OutEnvironment );
+		//OutEnvironment.SetRenderTargetOutputFormat( 0, PF_R16_UINT );
+	}
+
+	FPageTableUpdatePS() {}
+
+public:
+	FPageTableUpdatePS(const ShaderMetaType::CompiledShaderInitializerType& Initializer)
+		: FGlobalShader(Initializer)
+	{}
+};
+
+IMPLEMENT_SHADER_TYPE(, FPageTableUpdateVS, TEXT("PageTableUpdate"), TEXT("PageTableUpdateVS"), SF_Vertex );
+IMPLEMENT_SHADER_TYPE(, FPageTableUpdatePS, TEXT("PageTableUpdate"), TEXT("PageTableUpdatePS"), SF_Pixel );
+
+
+void FVirtualTextureSpace::ApplyUpdates( FRHICommandList& RHICmdList )
+{
+	static TArray< FPageUpdate > ExpandedUpdates[16];
+
+	if( CVarVTRefreshEntirePageTable.GetValueOnRenderThread() )
+	{
+		Pool->RefreshEntirePageTable( ID, ExpandedUpdates );
+	}
+	else
+	{
+		if( PageTableUpdates.Num() == 0 )
+		{
+			GRenderTargetPool.VisualizeTexture.SetCheckPoint( RHICmdList, PageTable );
+			return;
+		}
+
+		for( auto& Update : PageTableUpdates )
+		{
+			if( CVarVTMaskedPageTableUpdates.GetValueOnRenderThread() )
+			{
+				Pool->ExpandPageTableUpdateMasked( ID, Update, ExpandedUpdates );
+			}
+			else
+			{
+				Pool->ExpandPageTableUpdatePainters( ID, Update, ExpandedUpdates );
+			}
+		}
+	}
+
+	PageTableUpdates.Reset();
+
+	// TODO Expand 3D updates for slices of volume texture
+
+	uint32 TotalNumUpdates = 0;
+	for( uint32 Mip = 0; Mip < PageTableLevels; Mip++ )
+	{
+		TotalNumUpdates += ExpandedUpdates[ Mip ].Num();
+	}
+
+	if( TotalNumUpdates * sizeof( FPageUpdate ) > UpdateBuffer->GetSize() )
+	{
+		// Resize Update Buffer
+		uint32 MaxUpdates = FMath::RoundUpToPowerOfTwo( TotalNumUpdates );
+
+		FRHIResourceCreateInfo CreateInfo;
+		UpdateBuffer = RHICreateStructuredBuffer( sizeof( FPageUpdate ), MaxUpdates * sizeof( FPageUpdate ), BUF_ShaderResource | BUF_Volatile, CreateInfo );
+		UpdateBufferSRV = RHICreateShaderResourceView( UpdateBuffer );
+	}
+	
+	// This flushes the RHI thread!
+	uint8* Buffer = (uint8*)RHILockStructuredBuffer( UpdateBuffer, 0, TotalNumUpdates * sizeof( FPageUpdate ), RLM_WriteOnly );
+
+	for( uint32 Mip = 0; Mip < PageTableLevels; Mip++ )
+	{
+		uint32 NumUpdates = ExpandedUpdates[ Mip ].Num();
+		if( NumUpdates )
+		{
+			size_t UploadSize = NumUpdates * sizeof( FPageUpdate );
+			FMemory::Memcpy( Buffer, ExpandedUpdates[ Mip ].GetData(), UploadSize );
+			Buffer += UploadSize;
+		}
+	}
+	
+	RHIUnlockStructuredBuffer( UpdateBuffer );
+
+	// Draw
+	SCOPED_DRAW_EVENT( RHICmdList, PageTableUpdate );
+
+	auto ShaderMap = GetGlobalShaderMap( GetFeatureLevel() );
+
+	FSceneRenderTargetItem& PageTableTarget = PageTable->GetRenderTargetItem();
+
+	uint32 FirstUpdate = 0;
+	uint32 MipSize = PageTableSize;
+	for( uint32 Mip = 0; Mip < PageTableLevels; Mip++ )
+	{
+		uint32 NumUpdates = ExpandedUpdates[ Mip ].Num();
+		if( NumUpdates )
+		{
+			SetRenderTarget( RHICmdList, PageTableTarget.TargetableTexture, Mip, NULL );
+			RHICmdList.SetViewport( 0, 0, 0.0f, MipSize, MipSize, 1.0f );
+
+			FGraphicsPipelineStateInitializer GraphicsPSOInit;
+			RHICmdList.ApplyCachedRenderTargets( GraphicsPSOInit );
+
+			GraphicsPSOInit.BlendState = TStaticBlendState<>::GetRHI();
+			GraphicsPSOInit.RasterizerState = TStaticRasterizerState<>::GetRHI();
+			GraphicsPSOInit.DepthStencilState = TStaticDepthStencilState<false, CF_Always>::GetRHI();
+			GraphicsPSOInit.PrimitiveType = PT_TriangleList;
+
+			TShaderMapRef< FPageTableUpdateVS > VertexShader( ShaderMap );
+			TShaderMapRef< FPageTableUpdatePS > PixelShader( ShaderMap );
+
+			GraphicsPSOInit.BoundShaderState.VertexDeclarationRHI = GEmptyVertexDeclaration.VertexDeclarationRHI;
+			GraphicsPSOInit.BoundShaderState.VertexShaderRHI = GETSAFERHISHADER_VERTEX(*VertexShader);
+			GraphicsPSOInit.BoundShaderState.PixelShaderRHI = GETSAFERHISHADER_PIXEL(*PixelShader);
+			SetGraphicsPipelineState( RHICmdList, GraphicsPSOInit );
+
+			{
+				const FVertexShaderRHIParamRef ShaderRHI = VertexShader->GetVertexShader();
+
+				SetShaderValue( RHICmdList, ShaderRHI, VertexShader->PageTableSize,	PageTableSize );
+				SetShaderValue( RHICmdList, ShaderRHI, VertexShader->FirstUpdate,	FirstUpdate );
+				SetShaderValue( RHICmdList, ShaderRHI, VertexShader->NumUpdates,	NumUpdates );
+				SetSRVParameter( RHICmdList, ShaderRHI, VertexShader->UpdateBuffer,	UpdateBufferSRV );
+			}
+
+			// needs to be the same on shader side (faster on NVIDIA and AMD)
+			uint32 QuadsPerInstance = 8;
+
+			RHICmdList.SetStreamSource( 0, NULL, 0, 0 );
+			RHICmdList.DrawIndexedPrimitive( GQuadIndexBuffer.IndexBufferRHI, PT_TriangleList, 0, 0, 32, 0, 2 * QuadsPerInstance, FMath::DivideAndRoundUp( NumUpdates, QuadsPerInstance ) );
+
+			ExpandedUpdates[ Mip ].Reset();
+		}
+
+		FirstUpdate += NumUpdates;
+		MipSize >>= 1;
+	}
+
+	RHICmdList.CopyToResolveTarget( PageTableTarget.TargetableTexture, PageTableTarget.ShaderResourceTexture, false, FResolveParams() );
+
+	GRenderTargetPool.VisualizeTexture.SetCheckPoint( RHICmdList, PageTable );
+}
\ No newline at end of file
diff --git a/Engine/Source/Runtime/Renderer/Private/VT/VirtualTextureSpace.h b/Engine/Source/Runtime/Renderer/Private/VT/VirtualTextureSpace.h
new file mode 100644
--- /dev/null
+++ b/Engine/Source/Runtime/Renderer/Private/VT/VirtualTextureSpace.h
@@ -0,0 +1,43 @@
+// Copyright 1998-2017 Epic Games, Inc. All Rights Reserved.
+
+#pragma once
+
+#include "CoreMinimal.h"
+#include "VirtualTextureShared.h"
+#include "TexturePagePool.h"
+#include "VirtualTextureAllocator.h"
+#include "RenderTargetPool.h"
+
+// Virtual memory address space mapped by a page table texture
+class FVirtualTextureSpace : public FRenderResource
+{
+public:
+						FVirtualTextureSpace( uint32 InSize, uint32 InDimensions, FTexturePagePool* InPool );
+						~FVirtualTextureSpace();
+
+	// FRenderResource interface
+	virtual void		InitDynamicRHI() override;
+	virtual void		ReleaseDynamicRHI() override;
+	
+	FRHITexture*		GetPageTableTexture() const		{ return PageTable->GetRenderTargetItem().ShaderResourceTexture; }
+	
+	void				QueueUpdate( uint8 vLogSize, uint32 vAddress, uint8 vLevel, uint16 pAddress );
+	void				ApplyUpdates( FRHICommandList& RHICmdList );
+
+	uint32				ID;
+	uint32				PageTableSize;
+	uint32				PageTableLevels;
+	uint32				Dimensions;
+	
+	FTexturePagePool*	Pool;
+
+	FVirtualTextureAllocator	Allocator;
+
+private:
+	TRefCountPtr< IPooledRenderTarget >	PageTable;
+	
+	TArray< FPageUpdate >		PageTableUpdates;
+
+	FStructuredBufferRHIRef		UpdateBuffer;
+	FShaderResourceViewRHIRef	UpdateBufferSRV;
+};
diff --git a/Engine/Source/Runtime/Renderer/Private/VT/VirtualTextureSystem.cpp b/Engine/Source/Runtime/Renderer/Private/VT/VirtualTextureSystem.cpp
new file mode 100644
--- /dev/null
+++ b/Engine/Source/Runtime/Renderer/Private/VT/VirtualTextureSystem.cpp
@@ -0,0 +1,293 @@
+// Copyright 1998-2017 Epic Games, Inc. All Rights Reserved.
+
+#include "VirtualTextureSystem.h"
+
+#include "TexturePagePool.h"
+#include "VirtualTextureSpace.h"
+#include "VirtualTextureFeedback.h"
+#include "VirtualTexture.h"
+#include "UniquePageList.h"
+#include "Stats/Stats.h"
+#include "SceneUtils.h"
+#include "HAL/IConsoleManager.h"
+
+DECLARE_STATS_GROUP( TEXT("Virtual Texturing"), STATGROUP_VirtualTexturing, STATCAT_Advanced );
+
+DECLARE_CYCLE_STAT(			TEXT("Feedback Analysis"),				STAT_FeedbackAnalysis,				STATGROUP_VirtualTexturing );
+DECLARE_CYCLE_STAT(			TEXT("VirtualTextureSystem Update"),	STAT_VirtualTextureSystem_Update,	STATGROUP_VirtualTexturing );
+DECLARE_CYCLE_STAT(			TEXT("Page Table Updates"),				STAT_PageTableUpdates,				STATGROUP_VirtualTexturing );
+DECLARE_CYCLE_STAT(			TEXT("UniquePageList ExpandByMips"),	STAT_UniquePageList_ExpandByMips,	STATGROUP_VirtualTexturing );
+DECLARE_CYCLE_STAT(			TEXT("UniquePageList Sort"),			STAT_UniquePageList_Sort,			STATGROUP_VirtualTexturing );
+
+DECLARE_DWORD_COUNTER_STAT( TEXT("Num pages visible"),				STAT_NumPagesVisible,				STATGROUP_VirtualTexturing );
+DECLARE_DWORD_COUNTER_STAT( TEXT("Num page requests"),				STAT_NumPageRequests,				STATGROUP_VirtualTexturing );
+DECLARE_DWORD_COUNTER_STAT( TEXT("Num page requests resident"),		STAT_NumPageRequestsResident,		STATGROUP_VirtualTexturing );
+DECLARE_DWORD_COUNTER_STAT( TEXT("Num page requests not resident"), STAT_NumPageRequestsNotResident,	STATGROUP_VirtualTexturing );
+DECLARE_DWORD_COUNTER_STAT( TEXT("Num page uploads"),				STAT_NumPageUploads,				STATGROUP_VirtualTexturing );
+
+DECLARE_FLOAT_COUNTER_STAT( TEXT("VT"), STAT_GPU_VT, STATGROUP_GPU );
+
+
+static TAutoConsoleVariable<int32> CVarVTMaxUploadsPerFrame(
+	TEXT("r.VT.MaxUploadsPerFrame"),
+	16,
+	TEXT("Max number of page uploads per frame"),
+	ECVF_RenderThreadSafe
+	);
+
+static TAutoConsoleVariable<int32> CVarVTNumMipsToExpandRequests(
+	TEXT("r.VT.NumMipsToExpandRequests"),
+	3,
+	TEXT("Number of mip levels to request in addition to the original request"),
+	ECVF_RenderThreadSafe
+	);
+
+
+IVirtualTexture::~IVirtualTexture()
+{}
+
+
+FVirtualTextureSystem GVirtualTextureSystem;
+
+FVirtualTextureSystem::FVirtualTextureSystem()
+	: Frame(0)
+{
+	for( int ID = 0; ID < 16; ID++ )
+	{
+		Spaces[ ID ] = nullptr;
+	}
+}
+
+FVirtualTextureSystem::~FVirtualTextureSystem()
+{}
+
+void FVirtualTextureSystem::RegisterSpace( FVirtualTextureSpace* Space )
+{
+	check( Space );
+
+	for( int ID = 0; ID < 16; ID++ )
+	{
+		if( Spaces[ ID ] == nullptr )
+		{
+			Spaces[ ID ] = Space;
+			Space->ID = ID;
+			return;
+		}
+	}
+
+	check(0);
+}
+
+void FVirtualTextureSystem::UnregisterSpace( FVirtualTextureSpace* Space )
+{
+	check( Space );
+	check( Spaces[ Space->ID ] == Space );
+
+	Spaces[ Space->ID ] = nullptr;
+	Space->ID = 0xff;
+}
+
+void FVirtualTextureSystem::FeedbackAnalysis( FUniquePageList* RESTRICT RequestedPageList, const uint32* RESTRICT Buffer, uint32 Width, uint32 Height, uint32 Pitch )
+{
+	SCOPE_CYCLE_COUNTER( STAT_FeedbackAnalysis );
+
+	// Combine simple runs of identical requests
+	uint32 LastPixel = 0xffffffff;
+	uint32 LastPage = 0xffffffff;
+	uint32 LastCount = 0;
+
+	for( uint32 y = 0; y < Height; y++ )
+	{
+		for( uint32 x = 0; x < Width; x++ )
+		{
+			const uint32 Pixel = Buffer[ x + y * Pitch ];
+
+			if( Pixel == 0xffffffff )
+			{
+				continue;
+			}
+		
+			if( Pixel == LastPixel )
+			{
+				LastCount++;
+				continue;
+			}
+
+			// Decode pixel encoding
+			const uint32 PageX	= ( Pixel >>  0 ) & 0xfff;
+			const uint32 PageY	= ( Pixel >> 12 ) & 0xfff;
+			const uint32 Level	= ( Pixel >> 24 ) & 0xf;
+			const uint32 ID		= ( Pixel >> 28 );
+
+			const uint32 MaxLevel = RequestedPageList->NumLevels[ ID ] - 1;
+
+			uint32 vAddress = FMath::MortonCode2( PageX ) | ( FMath::MortonCode2( PageY ) << 1 );
+			uint32 vLevel = FMath::Min( Level, MaxLevel );
+			uint32 vDimensions = RequestedPageList->Dimensions[ ID ];
+
+			// Mask out low bits
+			vAddress &= 0xffffffff << ( vDimensions * vLevel );
+
+			uint32 Page = EncodePage( ID, vLevel, vAddress );
+			if( Page == LastPage )
+			{
+				LastCount++;
+				continue;
+			}
+
+			if( LastPage != 0xffffffff )
+			{
+				RequestedPageList->Add( LastPage, LastCount );
+			}
+
+			LastPixel = Pixel;
+			LastPage = Page;
+			LastCount = 1;
+		}
+	}
+
+	if( LastPage != 0xffffffff )
+	{
+		RequestedPageList->Add( LastPage, LastCount );
+	}
+}
+
+void FVirtualTextureSystem::Update( FRHICommandListImmediate& RHICmdList, ERHIFeatureLevel::Type FeatureLevel )
+{
+	SCOPE_CYCLE_COUNTER( STAT_VirtualTextureSystem_Update );
+	SCOPED_GPU_STAT( RHICmdList, STAT_GPU_VT );
+
+	FMemMark Mark( FMemStack::Get() );
+	FUniquePageList* RESTRICT RequestedPageList = new(FMemStack::Get()) FUniquePageList;
+
+	// Cache to avoid indirection
+	for( uint32 ID = 0; ID < 16; ID++ )
+	{
+		FVirtualTextureSpace* RESTRICT Space = Spaces[ ID ];
+
+		RequestedPageList->NumLevels[ ID ] = Space ? Space->PageTableLevels : 16;
+		RequestedPageList->Dimensions[ ID ] = Space ? Space->Dimensions : 2;
+	}
+	
+	int32 Pitch = 0;
+	const uint32* RESTRICT Buffer = GVirtualTextureFeedback.Map( RHICmdList, Pitch );
+	if( Buffer )
+	{
+		FeedbackAnalysis( RequestedPageList, Buffer, GVirtualTextureFeedback.Size.X, GVirtualTextureFeedback.Size.Y, Pitch );
+		GVirtualTextureFeedback.Unmap( RHICmdList );
+	}
+
+	SET_DWORD_STAT( STAT_NumPagesVisible, RequestedPageList->GetNum() );
+
+	// Can add other sources of pages to RequestedPageList here
+	
+	{
+		SCOPE_CYCLE_COUNTER( STAT_UniquePageList_ExpandByMips );
+		RequestedPageList->ExpandByMips( CVarVTNumMipsToExpandRequests.GetValueOnRenderThread() );
+	}
+	
+	SET_DWORD_STAT( STAT_NumPageRequests, RequestedPageList->GetNum() );
+
+	// Static to avoid malloc cost
+	static FBinaryHeap< uint32, uint16 > RequestHeap;
+	RequestHeap.Clear();
+
+	{
+		// Find resident
+
+		for( uint32 i = 0, Num = RequestedPageList->GetNum(); i < Num; i++ )
+		{
+			const uint64 PageEncoded = RequestedPageList->GetPage(i);
+
+			// Decode page
+			uint32 ID, vLevel, vAddress;
+			DecodePage( (uint32)PageEncoded, ID, vLevel, vAddress );
+
+			FVirtualTextureSpace* RESTRICT	Space = Spaces[ ID ];	checkSlow( Space );
+			FTexturePagePool* RESTRICT		Pool = Space->Pool;
+
+			// Is this page already resident?
+			uint32 pAddress = Pool->FindPage( ID, vLevel, vAddress );
+			if( pAddress == ~0u )
+			{
+				// Page isn't resident. Start searching at parent level for nearest ancestor.
+				uint32 Parent_vLevel = vLevel + 1;
+				uint32 Parent_vAddress = vAddress & ( 0xffffffff << ( Space->Dimensions * Parent_vLevel ) );
+
+				uint32 Ancestor_pAddress = Pool->FindNearestPage( ID, Parent_vLevel, Parent_vAddress );
+				uint32 Ancestor_vLevel = Ancestor_pAddress != ~0u ? Pool->GetPage( Ancestor_pAddress ).vLevel : Space->PageTableLevels - 1;
+
+				uint32 Count = RequestedPageList->GetCount(i);
+				uint32 Priority = Count * ( 1 << ( Ancestor_vLevel - vLevel ) );
+
+				RequestHeap.Add( ~Priority, i );
+			}
+			else
+			{
+				Pool->UpdateUsage( Frame, pAddress );
+
+				//FileCache->Touch( VT->FileName, PageOffset, PageSize, priority );
+			}
+		}
+
+		SET_DWORD_STAT( STAT_NumPageRequestsResident, RequestedPageList->GetNum() - RequestHeap.Num() );
+		SET_DWORD_STAT( STAT_NumPageRequestsNotResident, RequestHeap.Num() );
+	}
+	
+	// Limit the number of uploads
+	// Are all pages equal? Should there be different limits on different types of pages?
+	int32 NumUploadsLeft = CVarVTMaxUploadsPerFrame.GetValueOnRenderThread();
+
+	while( RequestHeap.Num() > 0 && NumUploadsLeft > 0 )
+	{
+		uint32 PageIndex = RequestHeap.Top();
+		RequestHeap.Pop();
+
+		const uint64 PageEncoded = RequestedPageList->GetPage( PageIndex );
+
+		// Decode page
+		uint32 ID, vLevel, vAddress;
+		DecodePage( (uint32)PageEncoded, ID, vLevel, vAddress );
+
+		FVirtualTextureSpace* RESTRICT	Space = Spaces[ ID ];	checkSlow( Space );
+		FTexturePagePool* RESTRICT		Pool = Space->Pool;
+
+		// Find specific VT in Space
+		uint32 Local_vAddress = 0;
+		IVirtualTexture* RESTRICT VT = Space->Allocator.Find( vAddress, Local_vAddress );
+
+		void* RESTRICT Location = nullptr;
+		bool bPageDataAvailable = VT->LocatePageData( vLevel, Local_vAddress, Location );
+
+		// FIXME ExpandByMips might not provide a valid page for this to fall back on
+
+		if( bPageDataAvailable && Pool->AnyFreeAvailable( Frame ) )
+		{
+			uint32 pAddress = Pool->Alloc( Frame );
+			check( pAddress != ~0u );
+
+			Pool->UnmapPage( pAddress );
+			
+			VT->ProducePageData( RHICmdList, FeatureLevel, vLevel, Local_vAddress, pAddress, Location );
+			
+			Pool->MapPage( ID, vLevel, vAddress, pAddress );
+			Pool->Free( Frame, pAddress );
+
+			NumUploadsLeft--;
+			INC_DWORD_STAT( STAT_NumPageUploads );
+		}
+	}
+
+	SCOPE_CYCLE_COUNTER( STAT_PageTableUpdates );
+	
+	// Update page tables
+	for( uint32 ID = 0; ID < 16; ID++ )
+	{
+		if( Spaces[ ID ] )
+		{
+			Spaces[ ID ]->ApplyUpdates( RHICmdList );
+		}
+	}
+
+	Frame++;
+}
\ No newline at end of file
diff --git a/Engine/Source/Runtime/Renderer/Private/VT/VirtualTextureSystem.h b/Engine/Source/Runtime/Renderer/Private/VT/VirtualTextureSystem.h
new file mode 100644
--- /dev/null
+++ b/Engine/Source/Runtime/Renderer/Private/VT/VirtualTextureSystem.h
@@ -0,0 +1,32 @@
+// Copyright 1998-2017 Epic Games, Inc. All Rights Reserved.
+
+#pragma once
+
+#include "CoreMinimal.h"
+#include "RHI.h"
+
+class FVirtualTextureSpace;
+class FUniquePageList;
+
+class FVirtualTextureSystem
+{
+public:
+			FVirtualTextureSystem();
+			~FVirtualTextureSystem();
+
+	void	Update( FRHICommandListImmediate& RHICmdList, ERHIFeatureLevel::Type FeatureLevel );
+
+	void	RegisterSpace( FVirtualTextureSpace* Space );
+	void	UnregisterSpace( FVirtualTextureSpace* Space );
+
+	FVirtualTextureSpace*	GetSpace( uint8 ID )	{ return Spaces[ ID ]; }
+
+private:
+	void	FeedbackAnalysis( FUniquePageList* RESTRICT RequestedPageList, const uint32* RESTRICT Buffer, uint32 Width, uint32 Height, uint32 Pitch );
+	
+	uint32	Frame;
+	
+	FVirtualTextureSpace*	Spaces[16];
+};
+
+extern FVirtualTextureSystem GVirtualTextureSystem;
\ No newline at end of file
diff --git a/Engine/Source/Runtime/Renderer/Private/VelocityRendering.cpp b/Engine/Source/Runtime/Renderer/Private/VelocityRendering.cpp
--- a/Engine/Source/Runtime/Renderer/Private/VelocityRendering.cpp
+++ b/Engine/Source/Runtime/Renderer/Private/VelocityRendering.cpp
@@ -608,7 +608,6 @@ bool IsMotionBlurEnabled(const FViewInfo& View)
 		&& View.FinalPostProcessSettings.MotionBlurMax > 0.001f
 		&& View.Family->bRealtimeUpdate
 		&& MotionBlurQuality > 0
-		&& !View.bIsSceneCapture
 		&& (CVarAllowMotionBlurInVR->GetInt() != 0 || !(View.Family->Views.Num() > 1));
 }
 
diff --git a/Engine/Source/Runtime/Renderer/Private/VolumetricFog.cpp b/Engine/Source/Runtime/Renderer/Private/VolumetricFog.cpp
--- a/Engine/Source/Runtime/Renderer/Private/VolumetricFog.cpp
+++ b/Engine/Source/Runtime/Renderer/Private/VolumetricFog.cpp
@@ -70,7 +70,7 @@ int32 GVolumetricFogJitter = 1;
 FAutoConsoleVariableRef CVarVolumetricFogJitter(
 	TEXT("r.VolumetricFog.Jitter"),
 	GVolumetricFogJitter,
-	TEXT(""),
+	TEXT("Whether to apply jitter to each frame's volumetric fog computation, achieving temporal super sampling."),
 	ECVF_Scalability | ECVF_RenderThreadSafe
 	);
 
@@ -78,7 +78,16 @@ float GVolumetricFogHistoryWeight = .9f;
 FAutoConsoleVariableRef CVarVolumetricFogHistoryWeight(
 	TEXT("r.VolumetricFog.HistoryWeight"),
 	GVolumetricFogHistoryWeight,
-	TEXT(""),
+	TEXT("How much the history value should be weighted each frame.  This is a tradeoff between visible jittering and responsiveness."),
+	ECVF_Scalability | ECVF_RenderThreadSafe
+	);
+
+int32 GVolumetricFogHistoryMissSupersampleCount = 4;
+FAutoConsoleVariableRef CVarVolumetricFogHistoryMissSupersampleCount(
+	TEXT("r.VolumetricFog.HistoryMissSupersampleCount"),
+	GVolumetricFogHistoryMissSupersampleCount,
+	TEXT("Number of lighting samples to compute for voxels whose history value is not available.\n")
+	TEXT("This reduces noise when panning or on camera cuts, but introduces a variable cost to volumetric fog computation.  Valid range [1, 16]."),
 	ECVF_Scalability | ECVF_RenderThreadSafe
 	);
 
@@ -158,13 +167,12 @@ public:
 	void SetParameters(
 		FRHICommandList& RHICmdList, 
 		const FViewInfo& View, 
-		const FExponentialHeightFogSceneInfo& FogInfo,
-		IPooledRenderTarget* VBufferA,
-		IPooledRenderTarget* VBufferB)
+		const FVolumetricFogIntegrationParameterData& IntegrationData,
+		const FExponentialHeightFogSceneInfo& FogInfo)
 	{
 		FComputeShaderRHIParamRef ShaderRHI = GetComputeShader();
 		FGlobalShader::SetParameters<FViewUniformShaderParameters>(RHICmdList, ShaderRHI, View.ViewUniformBuffer);
-		VolumetricFogParameters.Set(RHICmdList, ShaderRHI, View, VBufferA, VBufferB, NULL);
+		VolumetricFogParameters.Set(RHICmdList, ShaderRHI, View, IntegrationData);
 		HeightFogParameters.Set(RHICmdList, ShaderRHI, &View);
 		SetShaderValue(RHICmdList, ShaderRHI, GlobalAlbedo, FogInfo.VolumetricFogAlbedo);
 		SetShaderValue(RHICmdList, ShaderRHI, GlobalEmissive, FogInfo.VolumetricFogEmissive);
@@ -216,6 +224,7 @@ public:
 	{
 		FGlobalShader::ModifyCompilationEnvironment(Platform, OutEnvironment);
 		OutEnvironment.CompilerFlags.Add(CFLAG_VertexToGeometryShader);
+		FVolumetricFogIntegrationParameters::ModifyCompilationEnvironment(Platform, OutEnvironment);
 	}
 
 	FWriteToBoundingSphereVS(const ShaderMetaType::CompiledShaderInitializerType& Initializer):
@@ -229,7 +238,7 @@ public:
 
 	FWriteToBoundingSphereVS() {}
 
-	void SetParameters(FRHICommandList& RHICmdList, const FViewInfo& View, const FSphere& BoundingSphere, int32 MinZValue)
+	void SetParameters(FRHICommandList& RHICmdList, const FViewInfo& View, const FVolumetricFogIntegrationParameterData& IntegrationData, const FSphere& BoundingSphere, int32 MinZValue)
 	{
 		SetShaderValue(RHICmdList, GetVertexShader(), MinZ, MinZValue);
 
@@ -239,7 +248,7 @@ public:
 		const FMatrix ProjectionMatrix = View.ViewMatrices.ComputeProjectionNoAAMatrix();
 		SetShaderValue(RHICmdList, GetVertexShader(), ViewToVolumeClip, ProjectionMatrix);
 
-		VolumetricFogParameters.Set(RHICmdList, GetVertexShader(), View, NULL, NULL, NULL);
+		VolumetricFogParameters.Set(RHICmdList, GetVertexShader(), View, IntegrationData);
 	}
 
 	virtual bool Serialize(FArchive& Ar) override
@@ -262,7 +271,7 @@ private:
 IMPLEMENT_SHADER_TYPE(,FWriteToBoundingSphereVS,TEXT("VolumetricFog"),TEXT("WriteToBoundingSphereVS"),SF_Vertex);
 
 /** Shader that adds direct lighting contribution from the given light to the current volume lighting cascade. */
-template<bool bDynamicallyShadowed, bool bInverseSquared>
+template<bool bDynamicallyShadowed, bool bInverseSquared, bool bTemporalReprojection>
 class TInjectShadowedLocalLightPS : public FGlobalShader
 {
 	DECLARE_SHADER_TYPE(TInjectShadowedLocalLightPS,Global);
@@ -273,6 +282,8 @@ public:
 		FGlobalShader::ModifyCompilationEnvironment(Platform, OutEnvironment);
 		OutEnvironment.SetDefine(TEXT("DYNAMICALLY_SHADOWED"), (uint32)bDynamicallyShadowed);
 		OutEnvironment.SetDefine(TEXT("INVERSE_SQUARED_FALLOFF"), (uint32)bInverseSquared);
+		OutEnvironment.SetDefine(TEXT("USE_TEMPORAL_REPROJECTION"), bTemporalReprojection);
+		FVolumetricFogIntegrationParameters::ModifyCompilationEnvironment(Platform, OutEnvironment);
 	}
 
 	static bool ShouldCache(EShaderPlatform Platform)
@@ -296,6 +307,7 @@ public:
 	void SetParameters(
 		FRHICommandList& RHICmdList, 
 		const FViewInfo& View, 
+		const FVolumetricFogIntegrationParameterData& IntegrationData,
 		const FLightSceneInfo* LightSceneInfo,
 		const FExponentialHeightFogSceneInfo& FogInfo,
 		const FProjectedShadowInfo* ShadowMap)
@@ -306,7 +318,7 @@ public:
 		
 		SetDeferredLightParameters(RHICmdList, ShaderRHI, GetUniformBufferParameter<FDeferredLightUniformStruct>(), LightSceneInfo, View);
 
-		VolumetricFogParameters.Set(RHICmdList, ShaderRHI, View, NULL, NULL, NULL);
+		VolumetricFogParameters.Set(RHICmdList, ShaderRHI, View, IntegrationData);
 
 		SetShaderValue(RHICmdList, ShaderRHI, PhaseG, FogInfo.VolumetricFogScatteringDistribution);
 		SetShaderValue(RHICmdList, ShaderRHI, InverseSquaredLightDistanceBiasScale, GInverseSquaredLightDistanceBiasScale);
@@ -331,14 +343,18 @@ private:
 	FVolumeShadowingParameters VolumeShadowingParameters;
 };
 
-#define IMPLEMENT_LOCAL_LIGHT_INJECTION_PIXELSHADER_TYPE(bDynamicallyShadowed,bInverseSquared) \
-	typedef TInjectShadowedLocalLightPS<bDynamicallyShadowed,bInverseSquared> TInjectShadowedLocalLightPS##bDynamicallyShadowed##bInverseSquared; \
-	IMPLEMENT_SHADER_TYPE(template<>,TInjectShadowedLocalLightPS##bDynamicallyShadowed##bInverseSquared,TEXT("VolumetricFog"),TEXT("InjectShadowedLocalLightPS"),SF_Pixel);
+#define IMPLEMENT_LOCAL_LIGHT_INJECTION_PIXELSHADER_TYPE(bDynamicallyShadowed,bInverseSquared,bTemporalReprojection) \
+	typedef TInjectShadowedLocalLightPS<bDynamicallyShadowed,bInverseSquared,bTemporalReprojection> TInjectShadowedLocalLightPS##bDynamicallyShadowed##bInverseSquared##bTemporalReprojection; \
+	IMPLEMENT_SHADER_TYPE(template<>,TInjectShadowedLocalLightPS##bDynamicallyShadowed##bInverseSquared##bTemporalReprojection,TEXT("VolumetricFog"),TEXT("InjectShadowedLocalLightPS"),SF_Pixel);
 
-IMPLEMENT_LOCAL_LIGHT_INJECTION_PIXELSHADER_TYPE(true, true);
-IMPLEMENT_LOCAL_LIGHT_INJECTION_PIXELSHADER_TYPE(true, false);
-IMPLEMENT_LOCAL_LIGHT_INJECTION_PIXELSHADER_TYPE(false, true);
-IMPLEMENT_LOCAL_LIGHT_INJECTION_PIXELSHADER_TYPE(false, false);
+IMPLEMENT_LOCAL_LIGHT_INJECTION_PIXELSHADER_TYPE(true, true, true);
+IMPLEMENT_LOCAL_LIGHT_INJECTION_PIXELSHADER_TYPE(true, false, true);
+IMPLEMENT_LOCAL_LIGHT_INJECTION_PIXELSHADER_TYPE(false, true, true);
+IMPLEMENT_LOCAL_LIGHT_INJECTION_PIXELSHADER_TYPE(false, false, true);
+IMPLEMENT_LOCAL_LIGHT_INJECTION_PIXELSHADER_TYPE(true, true, false);
+IMPLEMENT_LOCAL_LIGHT_INJECTION_PIXELSHADER_TYPE(true, false, false);
+IMPLEMENT_LOCAL_LIGHT_INJECTION_PIXELSHADER_TYPE(false, true, false);
+IMPLEMENT_LOCAL_LIGHT_INJECTION_PIXELSHADER_TYPE(false, false, false);
 
 FProjectedShadowInfo* GetShadowForInjectionIntoVolumetricFog(const FLightSceneProxy* LightProxy, FVisibleLightInfo& VisibleLightInfo)
 {
@@ -391,10 +407,11 @@ FIntPoint CalculateVolumetricFogBoundsForLight(const FSphere& LightBounds, const
 	return VolumeZBounds;
 }
 
-template<bool bDynamicallyShadowed, bool bInverseSquared>
+template<bool bDynamicallyShadowed, bool bInverseSquared, bool bUseTemporalReprojection>
 void SetInjectShadowedLocalLightShaders(
 	FRHICommandList& RHICmdList, 
 	const FViewInfo& View, 
+	const FVolumetricFogIntegrationParameterData& IntegrationData,
 	const FLightSceneInfo* LightSceneInfo,
 	const FSphere& LightBounds,
 	const FExponentialHeightFogSceneInfo& FogInfo,
@@ -404,7 +421,7 @@ void SetInjectShadowedLocalLightShaders(
 {
 	TShaderMapRef<FWriteToBoundingSphereVS> VertexShader(View.ShaderMap);
 	TOptionalShaderMapRef<FWriteToSliceGS> GeometryShader(View.ShaderMap);
-	TShaderMapRef<TInjectShadowedLocalLightPS<bDynamicallyShadowed, bInverseSquared>> PixelShader(View.ShaderMap);
+	TShaderMapRef<TInjectShadowedLocalLightPS<bDynamicallyShadowed, bInverseSquared, bUseTemporalReprojection>> PixelShader(View.ShaderMap);
 
 	FGraphicsPipelineStateInitializer GraphicsPSOInit;
 	RHICmdList.ApplyCachedRenderTargets(GraphicsPSOInit);
@@ -417,12 +434,12 @@ void SetInjectShadowedLocalLightShaders(
 	GraphicsPSOInit.BoundShaderState.VertexShaderRHI = GETSAFERHISHADER_VERTEX(*VertexShader);
 	GraphicsPSOInit.BoundShaderState.GeometryShaderRHI = GETSAFERHISHADER_GEOMETRY(*GeometryShader);
 	GraphicsPSOInit.BoundShaderState.PixelShaderRHI = GETSAFERHISHADER_PIXEL(*PixelShader);
-	GraphicsPSOInit.PrimitiveType = PT_TriangleStrip;
+	GraphicsPSOInit.PrimitiveType = PT_TriangleList;
 
 	SetGraphicsPipelineState(RHICmdList, GraphicsPSOInit);
 
-	PixelShader->SetParameters(RHICmdList, View, LightSceneInfo, FogInfo, ProjectedShadowInfo);
-	VertexShader->SetParameters(RHICmdList, View, LightBounds, MinZ);
+	PixelShader->SetParameters(RHICmdList, View, IntegrationData, LightSceneInfo, FogInfo, ProjectedShadowInfo);
+	VertexShader->SetParameters(RHICmdList, View, IntegrationData, LightBounds, MinZ);
 
 	if (GeometryShader.IsValid())
 	{
@@ -430,11 +447,81 @@ void SetInjectShadowedLocalLightShaders(
 	}
 }
 
+/**  */
+class FCircleRasterizeVertexBuffer : public FVertexBuffer
+{
+public:
+
+	virtual void InitRHI() override
+	{
+		const int32 NumTriangles = NumVertices - 2;
+		const uint32 Size = NumVertices * sizeof(FScreenVertex);
+		FRHIResourceCreateInfo CreateInfo;
+		void* Buffer = nullptr;
+		VertexBufferRHI = RHICreateAndLockVertexBuffer(Size, BUF_Static, CreateInfo, Buffer);		
+		FScreenVertex* DestVertex = (FScreenVertex*)Buffer;
+
+		const int32 NumRings = NumVertices;
+		const float RadiansPerRingSegment = PI / (float)NumRings;
+
+		// Boost the effective radius so that the edges of the circle approximation lie on the circle, instead of the vertices
+		const float RadiusScale = 1.0f / FMath::Cos(RadiansPerRingSegment);
+
+		for (int32 VertexIndex = 0; VertexIndex < NumVertices; VertexIndex++)
+		{
+			float Angle = VertexIndex / (float)(NumVertices - 1) * 2 * PI;
+			// WriteToBoundingSphereVS only uses UV
+			DestVertex[VertexIndex].Position = FVector2D(0, 0);
+			DestVertex[VertexIndex].UV = FVector2D(RadiusScale * FMath::Cos(Angle) * .5f + .5f, RadiusScale * FMath::Sin(Angle) * .5f + .5f);
+		}
+
+		RHIUnlockVertexBuffer(VertexBufferRHI);      
+	}
+
+	static int32 NumVertices;
+};
+
+int32 FCircleRasterizeVertexBuffer::NumVertices = 8;
+
+TGlobalResource<FCircleRasterizeVertexBuffer> GCircleRasterizeVertexBuffer;
+
+/**  */
+class FCircleRasterizeIndexBuffer : public FIndexBuffer
+{
+public:
+
+	virtual void InitRHI() override
+	{
+		const int32 NumTriangles = FCircleRasterizeVertexBuffer::NumVertices - 2;
+
+		TResourceArray<uint16, INDEXBUFFER_ALIGNMENT> Indices;
+		Indices.Empty(NumTriangles * 3);
+
+		for (int32 TriangleIndex = 0; TriangleIndex < NumTriangles; TriangleIndex++)
+		{
+			int32 LeadingVertexIndex = TriangleIndex + 2;
+			Indices.Add(0);
+			Indices.Add(LeadingVertexIndex - 1);
+			Indices.Add(LeadingVertexIndex);
+		}
+
+		const uint32 Size = Indices.GetResourceDataSize();
+		const uint32 Stride = sizeof(uint16);
+
+		// Create index buffer. Fill buffer with initial data upon creation
+		FRHIResourceCreateInfo CreateInfo(&Indices);
+		IndexBufferRHI = RHICreateIndexBuffer(Stride, Size, BUF_Static, CreateInfo);
+	}
+};
+
+TGlobalResource<FCircleRasterizeIndexBuffer> GCircleRasterizeIndexBuffer;
+
 void FDeferredShadingSceneRenderer::RenderLocalLightsForVolumetricFog(
 	FRHICommandListImmediate& RHICmdList,
 	FViewInfo& View,
+	bool bUseTemporalReprojection,
+	const FVolumetricFogIntegrationParameterData& IntegrationData,
 	const FExponentialHeightFogSceneInfo& FogInfo,
-	IPooledRenderTarget* VBufferA,
 	FIntVector VolumetricFogGridSize, 
 	FVector GridZParams,
 	const FPooledRenderTargetDesc& VolumeDesc,
@@ -483,33 +570,61 @@ void FDeferredShadingSceneRenderer::RenderLocalLightsForVolumetricFog(
 
 			if (VolumeZBounds.X < VolumeZBounds.Y)
 			{
-				if (bDynamicallyShadowed)
+				if (bUseTemporalReprojection)
 				{
-					if (bInverseSquared)
+					if (bDynamicallyShadowed)
 					{
-						SetInjectShadowedLocalLightShaders<true, true>(RHICmdList, View, LightSceneInfo, LightBounds, FogInfo, ProjectedShadowInfo, VolumetricFogGridSize, VolumeZBounds.X);
+						if (bInverseSquared)
+						{
+							SetInjectShadowedLocalLightShaders<true, true, true>(RHICmdList, View, IntegrationData, LightSceneInfo, LightBounds, FogInfo, ProjectedShadowInfo, VolumetricFogGridSize, VolumeZBounds.X);
+						}
+						else
+						{
+							SetInjectShadowedLocalLightShaders<true, false, true>(RHICmdList, View, IntegrationData, LightSceneInfo, LightBounds, FogInfo, ProjectedShadowInfo, VolumetricFogGridSize, VolumeZBounds.X);
+						}
 					}
 					else
 					{
-						SetInjectShadowedLocalLightShaders<true, false>(RHICmdList, View, LightSceneInfo, LightBounds, FogInfo, ProjectedShadowInfo, VolumetricFogGridSize, VolumeZBounds.X);
+						if (bInverseSquared)
+						{
+							SetInjectShadowedLocalLightShaders<false, true, true>(RHICmdList, View, IntegrationData, LightSceneInfo, LightBounds, FogInfo, ProjectedShadowInfo, VolumetricFogGridSize, VolumeZBounds.X);
+						}
+						else
+						{
+							SetInjectShadowedLocalLightShaders<false, false, true>(RHICmdList, View, IntegrationData, LightSceneInfo, LightBounds, FogInfo, ProjectedShadowInfo, VolumetricFogGridSize, VolumeZBounds.X);
+						}
 					}
 				}
 				else
 				{
-					if (bInverseSquared)
+					if (bDynamicallyShadowed)
 					{
-						SetInjectShadowedLocalLightShaders<false, true>(RHICmdList, View, LightSceneInfo, LightBounds, FogInfo, ProjectedShadowInfo, VolumetricFogGridSize, VolumeZBounds.X);
+						if (bInverseSquared)
+						{
+							SetInjectShadowedLocalLightShaders<true, true, false>(RHICmdList, View, IntegrationData, LightSceneInfo, LightBounds, FogInfo, ProjectedShadowInfo, VolumetricFogGridSize, VolumeZBounds.X);
+						}
+						else
+						{
+							SetInjectShadowedLocalLightShaders<true, false, false>(RHICmdList, View, IntegrationData, LightSceneInfo, LightBounds, FogInfo, ProjectedShadowInfo, VolumetricFogGridSize, VolumeZBounds.X);
+						}
 					}
 					else
 					{
-						SetInjectShadowedLocalLightShaders<false, false>(RHICmdList, View, LightSceneInfo, LightBounds, FogInfo, ProjectedShadowInfo, VolumetricFogGridSize, VolumeZBounds.X);
+						if (bInverseSquared)
+						{
+							SetInjectShadowedLocalLightShaders<false, true, false>(RHICmdList, View, IntegrationData, LightSceneInfo, LightBounds, FogInfo, ProjectedShadowInfo, VolumetricFogGridSize, VolumeZBounds.X);
+						}
+						else
+						{
+							SetInjectShadowedLocalLightShaders<false, false, false>(RHICmdList, View, IntegrationData, LightSceneInfo, LightBounds, FogInfo, ProjectedShadowInfo, VolumetricFogGridSize, VolumeZBounds.X);
+						}
 					}
 				}
 
-				RHICmdList.SetStreamSource(0, GVolumeRasterizeVertexBuffer.VertexBufferRHI, sizeof(FScreenVertex), 0);
+				RHICmdList.SetStreamSource(0, GCircleRasterizeVertexBuffer.VertexBufferRHI, sizeof(FScreenVertex), 0);
 				const int32 NumInstances = VolumeZBounds.Y - VolumeZBounds.X;
-				// Render a quad per slice affected by the given bounds
-				RHICmdList.DrawPrimitive(PT_TriangleStrip, 0, 2, NumInstances);
+				const int32 NumTriangles = FCircleRasterizeVertexBuffer::NumVertices - 2;
+				RHICmdList.DrawIndexedPrimitive(GCircleRasterizeIndexBuffer.IndexBufferRHI, PT_TriangleList, 0, 0, FCircleRasterizeVertexBuffer::NumVertices, 0, NumTriangles, NumInstances);
 			}
 		}
 
@@ -543,7 +658,6 @@ public:
 	TVolumetricFogLightScatteringCS(const ShaderMetaType::CompiledShaderInitializerType& Initializer)
 		: FGlobalShader(Initializer)
 	{
-		HistoryWeight.Bind(Initializer.ParameterMap, TEXT("HistoryWeight"));
 		LocalShadowedLightScattering.Bind(Initializer.ParameterMap, TEXT("LocalShadowedLightScattering"));
 		LightScatteringHistory.Bind(Initializer.ParameterMap, TEXT("LightScatteringHistory"));
 		LightScatteringHistorySampler.Bind(Initializer.ParameterMap, TEXT("LightScatteringHistorySampler"));
@@ -569,12 +683,10 @@ public:
 
 	void SetParameters(
 		FRHICommandList& RHICmdList, 
-		const FViewInfo& View, 
+		const FViewInfo& View,
+		const FVolumetricFogIntegrationParameterData& IntegrationData,
 		const FExponentialHeightFogSceneInfo& FogInfo,
-		IPooledRenderTarget* VBufferA,
-		IPooledRenderTarget* VBufferB,
 		IPooledRenderTarget* LocalShadowedLightScatteringTarget,
-		IPooledRenderTarget* LightScatteringRenderTarget,
 		FTextureRHIParamRef LightScatteringHistoryTexture,
 		bool bUseDirectionalLightShadowing,
 		const FMatrix& DirectionalLightFunctionWorldToShadowValue,
@@ -583,8 +695,6 @@ public:
 		FComputeShaderRHIParamRef ShaderRHI = GetComputeShader();
 		FGlobalShader::SetParameters<FViewUniformShaderParameters>(RHICmdList, ShaderRHI, View.ViewUniformBuffer);
 
-		SetShaderValue(RHICmdList, ShaderRHI, HistoryWeight, GVolumetricFogHistoryWeight);
-
 		FTextureRHIParamRef LocalShadowedLightScatteringTexture = GBlackVolumeTexture->TextureRHI;
 
 		if (LocalShadowedLightScatteringTarget)
@@ -607,7 +717,7 @@ public:
 			TStaticSamplerState<SF_Bilinear, AM_Clamp, AM_Clamp, AM_Clamp>::GetRHI(), 
 			LightScatteringHistoryTexture);
 
-		VolumetricFogParameters.Set(RHICmdList, ShaderRHI, View, VBufferA, VBufferB, LightScatteringRenderTarget);
+		VolumetricFogParameters.Set(RHICmdList, ShaderRHI, View, IntegrationData);
 		ForwardLightingParameters.Set(RHICmdList, ShaderRHI, View);
 
 		SetShaderValue(RHICmdList, ShaderRHI, DirectionalLightFunctionWorldToShadow, DirectionalLightFunctionWorldToShadowValue);
@@ -664,7 +774,6 @@ public:
 	virtual bool Serialize(FArchive& Ar)
 	{		
 		bool bShaderHasOutdatedParameters = FGlobalShader::Serialize(Ar);
-		Ar << HistoryWeight;
 		Ar << LocalShadowedLightScattering;
 		Ar << LightScatteringHistory;
 		Ar << LightScatteringHistorySampler;
@@ -687,7 +796,6 @@ public:
 
 private:
 
-	FShaderParameter HistoryWeight;
 	FShaderResourceParameter LocalShadowedLightScattering;
 	FShaderResourceParameter LightScatteringHistory;
 	FShaderResourceParameter LightScatteringHistorySampler;
@@ -745,11 +853,11 @@ public:
 	{
 	}
 
-	void SetParameters(FRHICommandList& RHICmdList, const FViewInfo& View, IPooledRenderTarget* LightScatteringRenderTarget)
+	void SetParameters(FRHICommandList& RHICmdList, const FViewInfo& View, const FVolumetricFogIntegrationParameterData& IntegrationData)
 	{
 		FComputeShaderRHIParamRef ShaderRHI = GetComputeShader();
 		FGlobalShader::SetParameters<FViewUniformShaderParameters>(RHICmdList, ShaderRHI, View.ViewUniformBuffer);
-		VolumetricFogParameters.Set(RHICmdList, ShaderRHI, View, NULL, NULL, LightScatteringRenderTarget);
+		VolumetricFogParameters.Set(RHICmdList, ShaderRHI, View, IntegrationData);
 	}
 
 	void UnsetParameters(FRHICommandList& RHICmdList, const FViewInfo& View)
@@ -927,6 +1035,30 @@ void FDeferredShadingSceneRenderer::ComputeVolumetricFog(FRHICommandListImmediat
 
 			SCOPED_DRAW_EVENT(RHICmdList, VolumetricFog);
 
+			const FVector FrameJitterOffsetValue = VolumetricFogTemporalRandom(View.Family->FrameNumber);
+
+			FVolumetricFogIntegrationParameterData IntegrationData;
+			IntegrationData.FrameJitterOffsetValues.Empty(16);
+			IntegrationData.FrameJitterOffsetValues.AddZeroed(16);
+			IntegrationData.FrameJitterOffsetValues[0] = VolumetricFogTemporalRandom(View.Family->FrameNumber);
+
+			for (int32 FrameOffsetIndex = 1; FrameOffsetIndex < GVolumetricFogHistoryMissSupersampleCount; FrameOffsetIndex++)
+			{
+				IntegrationData.FrameJitterOffsetValues[FrameOffsetIndex] = VolumetricFogTemporalRandom(View.Family->FrameNumber - FrameOffsetIndex);
+			}
+
+			const bool bUseTemporalReprojection = 
+				GVolumetricFogTemporalReprojection 
+				&& View.ViewState;
+
+			IntegrationData.bTemporalHistoryIsValid =
+				bUseTemporalReprojection
+				&& !View.bCameraCut
+				&& !View.bPrevTransformsReset
+				&& ViewFamily.bRealtimeUpdate
+				&& View.ViewState->LightScatteringHistory
+				&& View.ViewState->LightScatteringHistory->GetDesc().GetSize() == VolumetricFogGridSize;
+
 			FMatrix LightFunctionWorldToShadow;
 			TRefCountPtr<IPooledRenderTarget> LightFunctionTexture;
 			bool bUseDirectionalLightShadowing;
@@ -948,6 +1080,9 @@ void FDeferredShadingSceneRenderer::ComputeVolumetricFog(FRHICommandListImmediat
 			GRenderTargetPool.FindFreeElement(RHICmdList, VolumeDesc, VBufferA, TEXT("VBufferA"));
 			GRenderTargetPool.FindFreeElement(RHICmdList, VolumeDesc, VBufferB, TEXT("VBufferB"));
 
+			IntegrationData.VBufferARenderTarget = VBufferA.GetReference();
+			IntegrationData.VBufferBRenderTarget = VBufferB.GetReference();
+
 			// Unbind render targets, the shadow depth target may still be bound
 			SetRenderTarget(RHICmdList, NULL, NULL);
 
@@ -958,7 +1093,7 @@ void FDeferredShadingSceneRenderer::ComputeVolumetricFog(FRHICommandListImmediat
 					SCOPED_DRAW_EVENT(RHICmdList, InitializeVolumeAttributes);
 					TShaderMapRef<FVolumetricFogMaterialSetupCS> ComputeShader(View.ShaderMap);
 					RHICmdList.SetComputeShader(ComputeShader->GetComputeShader());
-					ComputeShader->SetParameters(RHICmdList, View, FogInfo, VBufferA.GetReference(), VBufferB.GetReference());
+					ComputeShader->SetParameters(RHICmdList, View, IntegrationData, FogInfo);
 					DispatchComputeShader(RHICmdList, *ComputeShader, NumGroups.X, NumGroups.Y, NumGroups.Z);
 					ComputeShader->UnsetParameters(RHICmdList, View, VBufferA.GetReference(), VBufferB.GetReference());
 				}
@@ -966,11 +1101,10 @@ void FDeferredShadingSceneRenderer::ComputeVolumetricFog(FRHICommandListImmediat
 				VoxelizeFogVolumePrimitives(
 					RHICmdList, 
 					View, 
+					IntegrationData,
 					VolumetricFogGridSize, 
 					GridZParams,
-					FogInfo.VolumetricFogDistance, 
-					VBufferA.GetReference(), 
-					VBufferB.GetReference());
+					FogInfo.VolumetricFogDistance);
 
 				FTextureRHIParamRef VoxelizeRenderTargets[2] =
 				{
@@ -985,22 +1119,18 @@ void FDeferredShadingSceneRenderer::ComputeVolumetricFog(FRHICommandListImmediat
 			}
 
 			TRefCountPtr<IPooledRenderTarget> LocalShadowedLightScattering = NULL;
-			RenderLocalLightsForVolumetricFog(RHICmdList, View, FogInfo, VBufferA.GetReference(), VolumetricFogGridSize, GridZParams, VolumeDesc, LocalShadowedLightScattering);
+			RenderLocalLightsForVolumetricFog(RHICmdList, View, bUseTemporalReprojection, IntegrationData, FogInfo, VolumetricFogGridSize, GridZParams, VolumeDesc, LocalShadowedLightScattering);
 
 			TRefCountPtr<IPooledRenderTarget> LightScattering;
 			GRenderTargetPool.FindFreeElement(RHICmdList, VolumeDesc, LightScattering, TEXT("LightScattering"));
 
+			IntegrationData.LightScatteringRenderTarget = LightScattering.GetReference();
+
 			SetRenderTarget(RHICmdList, NULL, NULL);
 
 			{
 				const FIntVector NumGroups = FIntVector::DivideAndRoundUp(VolumetricFogGridSize, VolumetricFogGridInjectionGroupSize);
 
-				const bool bUseTemporalReprojection = 
-					GVolumetricFogTemporalReprojection 
-					&& View.ViewState
-					&& !View.bCameraCut
-					&& !View.bPrevTransformsReset;
-
 				const bool bUseGlobalDistanceField = UseGlobalDistanceField() && Scene->DistanceFieldSceneData.NumObjectsInBuffer > 0;
 
 				const bool bUseDistanceFieldSkyOcclusion =
@@ -1021,13 +1151,17 @@ void FDeferredShadingSceneRenderer::ComputeVolumetricFog(FRHICommandListImmediat
 					bUseDistanceFieldSkyOcclusion ? TEXT("DFAO") : TEXT(""),
 					LightFunctionTexture ? TEXT("LF") : TEXT(""));
 
-				if (bUseTemporalReprojection && View.ViewState->LightScatteringHistory)
+				if (bUseTemporalReprojection)
 				{
+					FTextureRHIParamRef LightScatteringHistoryTexture = View.ViewState->LightScatteringHistory 
+						? View.ViewState->LightScatteringHistory->GetRenderTargetItem().ShaderResourceTexture
+						: GBlackVolumeTexture->TextureRHI;
+
 					if (bUseDistanceFieldSkyOcclusion)
 					{
 						TShaderMapRef<TVolumetricFogLightScatteringCS<true, true> > ComputeShader(View.ShaderMap);
 						RHICmdList.SetComputeShader(ComputeShader->GetComputeShader());
-						ComputeShader->SetParameters(RHICmdList, View, FogInfo, VBufferA.GetReference(), VBufferB.GetReference(), LocalShadowedLightScattering.GetReference(), LightScattering.GetReference(), View.ViewState->LightScatteringHistory->GetRenderTargetItem().ShaderResourceTexture, bUseDirectionalLightShadowing, LightFunctionWorldToShadow, LightFunctionTexture);
+						ComputeShader->SetParameters(RHICmdList, View, IntegrationData, FogInfo, LocalShadowedLightScattering.GetReference(), LightScatteringHistoryTexture, bUseDirectionalLightShadowing, LightFunctionWorldToShadow, LightFunctionTexture);
 						DispatchComputeShader(RHICmdList, *ComputeShader, NumGroups.X, NumGroups.Y, NumGroups.Z);
 						ComputeShader->UnsetParameters(RHICmdList, View, LightScattering.GetReference());
 					}
@@ -1035,7 +1169,7 @@ void FDeferredShadingSceneRenderer::ComputeVolumetricFog(FRHICommandListImmediat
 					{
 						TShaderMapRef<TVolumetricFogLightScatteringCS<true, false> > ComputeShader(View.ShaderMap);
 						RHICmdList.SetComputeShader(ComputeShader->GetComputeShader());
-						ComputeShader->SetParameters(RHICmdList, View, FogInfo, VBufferA.GetReference(), VBufferB.GetReference(), LocalShadowedLightScattering.GetReference(), LightScattering.GetReference(), View.ViewState->LightScatteringHistory->GetRenderTargetItem().ShaderResourceTexture, bUseDirectionalLightShadowing, LightFunctionWorldToShadow, LightFunctionTexture);
+						ComputeShader->SetParameters(RHICmdList, View, IntegrationData, FogInfo, LocalShadowedLightScattering.GetReference(), LightScatteringHistoryTexture, bUseDirectionalLightShadowing, LightFunctionWorldToShadow, LightFunctionTexture);
 						DispatchComputeShader(RHICmdList, *ComputeShader, NumGroups.X, NumGroups.Y, NumGroups.Z);
 						ComputeShader->UnsetParameters(RHICmdList, View, LightScattering.GetReference());
 					}
@@ -1046,7 +1180,7 @@ void FDeferredShadingSceneRenderer::ComputeVolumetricFog(FRHICommandListImmediat
 					{
 						TShaderMapRef<TVolumetricFogLightScatteringCS<false, true> > ComputeShader(View.ShaderMap);
 						RHICmdList.SetComputeShader(ComputeShader->GetComputeShader());
-						ComputeShader->SetParameters(RHICmdList, View, FogInfo, VBufferA.GetReference(), VBufferB.GetReference(), LocalShadowedLightScattering.GetReference(), LightScattering.GetReference(), NULL, bUseDirectionalLightShadowing, LightFunctionWorldToShadow, LightFunctionTexture);
+						ComputeShader->SetParameters(RHICmdList, View, IntegrationData, FogInfo, LocalShadowedLightScattering.GetReference(), NULL, bUseDirectionalLightShadowing, LightFunctionWorldToShadow, LightFunctionTexture);
 						DispatchComputeShader(RHICmdList, *ComputeShader, NumGroups.X, NumGroups.Y, NumGroups.Z);
 						ComputeShader->UnsetParameters(RHICmdList, View, LightScattering.GetReference());
 					}
@@ -1054,7 +1188,7 @@ void FDeferredShadingSceneRenderer::ComputeVolumetricFog(FRHICommandListImmediat
 					{
 						TShaderMapRef<TVolumetricFogLightScatteringCS<false, false> > ComputeShader(View.ShaderMap);
 						RHICmdList.SetComputeShader(ComputeShader->GetComputeShader());
-						ComputeShader->SetParameters(RHICmdList, View, FogInfo, VBufferA.GetReference(), VBufferB.GetReference(), LocalShadowedLightScattering.GetReference(), LightScattering.GetReference(), NULL, bUseDirectionalLightShadowing, LightFunctionWorldToShadow, LightFunctionTexture);
+						ComputeShader->SetParameters(RHICmdList, View, IntegrationData, FogInfo, LocalShadowedLightScattering.GetReference(), NULL, bUseDirectionalLightShadowing, LightFunctionWorldToShadow, LightFunctionTexture);
 						DispatchComputeShader(RHICmdList, *ComputeShader, NumGroups.X, NumGroups.Y, NumGroups.Z);
 						ComputeShader->UnsetParameters(RHICmdList, View, LightScattering.GetReference());
 					}
@@ -1083,7 +1217,7 @@ void FDeferredShadingSceneRenderer::ComputeVolumetricFog(FRHICommandListImmediat
 				const FIntVector NumGroups = FIntVector::DivideAndRoundUp(VolumetricFogGridSize, VolumetricFogIntegrationGroupSize);
 				TShaderMapRef<FVolumetricFogFinalIntegrationCS> ComputeShader(View.ShaderMap);
 				RHICmdList.SetComputeShader(ComputeShader->GetComputeShader());
-				ComputeShader->SetParameters(RHICmdList, View, LightScattering.GetReference());
+				ComputeShader->SetParameters(RHICmdList, View, IntegrationData);
 				DispatchComputeShader(RHICmdList, *ComputeShader, NumGroups.X, NumGroups.Y, 1);
 				ComputeShader->UnsetParameters(RHICmdList, View);
 			}
diff --git a/Engine/Source/Runtime/Renderer/Private/VolumetricFogShared.h b/Engine/Source/Runtime/Renderer/Private/VolumetricFogShared.h
--- a/Engine/Source/Runtime/Renderer/Private/VolumetricFogShared.h
+++ b/Engine/Source/Runtime/Renderer/Private/VolumetricFogShared.h
@@ -12,6 +12,22 @@
 
 extern FVector VolumetricFogTemporalRandom(uint32 FrameNumber);
 
+struct FVolumetricFogIntegrationParameterData
+{
+	FVolumetricFogIntegrationParameterData() :
+		bTemporalHistoryIsValid(false),
+		VBufferARenderTarget(NULL),
+		VBufferBRenderTarget(NULL),
+		LightScatteringRenderTarget(NULL)
+	{}
+
+	bool bTemporalHistoryIsValid;
+	TArray<FVector4, TInlineAllocator<16>> FrameJitterOffsetValues;
+	IPooledRenderTarget* VBufferARenderTarget;
+	IPooledRenderTarget* VBufferBRenderTarget;
+	IPooledRenderTarget* LightScatteringRenderTarget;
+};
+
 /**  */
 class FVolumetricFogIntegrationParameters
 {
@@ -19,7 +35,6 @@ public:
 
 	static void ModifyCompilationEnvironment(EShaderPlatform Platform, FShaderCompilerEnvironment& OutEnvironment)
 	{
-		
 	}
 
 	void Bind(const FShaderParameterMap& ParameterMap)
@@ -32,7 +47,9 @@ public:
 		VolumetricFogData.Bind(ParameterMap, TEXT("VolumetricFog"));
 		UnjitteredClipToTranslatedWorld.Bind(ParameterMap, TEXT("UnjitteredClipToTranslatedWorld")); 
 		UnjitteredPrevWorldToClip.Bind(ParameterMap, TEXT("UnjitteredPrevWorldToClip")); 
-		FrameJitterOffset.Bind(ParameterMap, TEXT("FrameJitterOffset")); 
+		FrameJitterOffsets.Bind(ParameterMap, TEXT("FrameJitterOffsets")); 
+		HistoryWeight.Bind(ParameterMap, TEXT("HistoryWeight")); 
+		HistoryMissSuperSampleCount.Bind(ParameterMap, TEXT("HistoryMissSuperSampleCount")); 
 	}
 
 	template<typename ShaderRHIParamRef>
@@ -40,25 +57,23 @@ public:
 		FRHICommandList& RHICmdList, 
 		const ShaderRHIParamRef& ShaderRHI, 
 		const FViewInfo& View, 
-		IPooledRenderTarget* VBufferARenderTarget, 
-		IPooledRenderTarget* VBufferBRenderTarget, 
-		IPooledRenderTarget* LightScatteringRenderTarget) const
+		const FVolumetricFogIntegrationParameterData& IntegrationData) const
 	{
 		if (VBufferA.IsBound())
 		{
-			const FSceneRenderTargetItem& VBufferAItem = VBufferARenderTarget->GetRenderTargetItem();
+			const FSceneRenderTargetItem& VBufferAItem = IntegrationData.VBufferARenderTarget->GetRenderTargetItem();
 			VBufferA.SetTexture(RHICmdList, ShaderRHI, VBufferAItem.ShaderResourceTexture, VBufferAItem.UAV);
 		}
 
 		if (VBufferB.IsBound())
 		{
-			const FSceneRenderTargetItem& VBufferBItem = VBufferBRenderTarget->GetRenderTargetItem();
+			const FSceneRenderTargetItem& VBufferBItem = IntegrationData.VBufferBRenderTarget->GetRenderTargetItem();
 			VBufferB.SetTexture(RHICmdList, ShaderRHI, VBufferBItem.ShaderResourceTexture, VBufferBItem.UAV);
 		}
 
 		if (LightScattering.IsBound())
 		{
-			const FSceneRenderTargetItem& LightScatteringItem = LightScatteringRenderTarget->GetRenderTargetItem();
+			const FSceneRenderTargetItem& LightScatteringItem = IntegrationData.LightScatteringRenderTarget->GetRenderTargetItem();
 			LightScattering.SetTexture(RHICmdList, ShaderRHI, LightScatteringItem.ShaderResourceTexture, LightScatteringItem.UAV);
 		}
 
@@ -88,11 +103,16 @@ public:
 			SetShaderValue(RHICmdList, ShaderRHI, UnjitteredPrevWorldToClip, UnjitteredViewProjectionMatrix);
 		}
 
-		if (FrameJitterOffset.IsBound())
+		if (FrameJitterOffsets.IsBound())
 		{
-			const FVector FrameJitterOffsetValue = VolumetricFogTemporalRandom(View.Family->FrameNumber);
-			SetShaderValue(RHICmdList, ShaderRHI, FrameJitterOffset, FrameJitterOffsetValue);
+			SetShaderValueArray(RHICmdList, ShaderRHI, FrameJitterOffsets, IntegrationData.FrameJitterOffsetValues.GetData(), IntegrationData.FrameJitterOffsetValues.Num());
 		}
+
+		extern float GVolumetricFogHistoryWeight;
+		SetShaderValue(RHICmdList, ShaderRHI, HistoryWeight, IntegrationData.bTemporalHistoryIsValid ? GVolumetricFogHistoryWeight : 0.0f);
+
+		extern int32 GVolumetricFogHistoryMissSupersampleCount;
+		SetShaderValue(RHICmdList, ShaderRHI, HistoryMissSuperSampleCount, FMath::Clamp(GVolumetricFogHistoryMissSupersampleCount, 1, 16));
 	}
 
 	template<typename ShaderRHIParamRef>
@@ -151,7 +171,9 @@ public:
 		Ar << P.VolumetricFogData;
 		Ar << P.UnjitteredClipToTranslatedWorld;
 		Ar << P.UnjitteredPrevWorldToClip;
-		Ar << P.FrameJitterOffset;
+		Ar << P.FrameJitterOffsets;
+		Ar << P.HistoryWeight;
+		Ar << P.HistoryMissSuperSampleCount;
 		return Ar;
 	}
 
@@ -165,7 +187,9 @@ private:
 	FShaderUniformBufferParameter VolumetricFogData;
 	FShaderParameter UnjitteredClipToTranslatedWorld;
 	FShaderParameter UnjitteredPrevWorldToClip;
-	FShaderParameter FrameJitterOffset;
+	FShaderParameter FrameJitterOffsets;
+	FShaderParameter HistoryWeight;
+	FShaderParameter HistoryMissSuperSampleCount;
 };
 
 inline int32 ComputeZSliceFromDepth(float SceneDepth, FVector GridZParams)
diff --git a/Engine/Source/Runtime/Renderer/Private/VolumetricFogVoxelization.cpp b/Engine/Source/Runtime/Renderer/Private/VolumetricFogVoxelization.cpp
--- a/Engine/Source/Runtime/Renderer/Private/VolumetricFogVoxelization.cpp
+++ b/Engine/Source/Runtime/Renderer/Private/VolumetricFogVoxelization.cpp
@@ -72,6 +72,7 @@ public:
 		const FVertexFactory* VertexFactory,
 		const FMaterialRenderProxy* MaterialRenderProxy,
 		const FViewInfo& View, 
+		const FVolumetricFogIntegrationParameterData& IntegrationData,
 		const TUniformBufferRef<FViewUniformShaderParameters>& VoxelizeViewUniformBuffer,
 		FVector2D Jitter)
 	{
@@ -79,7 +80,7 @@ public:
 
 		if (!RHISupportsGeometryShaders(View.GetShaderPlatform()))
 		{
-			VolumetricFogParameters.Set(RHICmdList, GetVertexShader(), View, NULL, NULL, NULL);
+			VolumetricFogParameters.Set(RHICmdList, GetVertexShader(), View, IntegrationData);
 
 			FMatrix ProjectionMatrix = View.ViewMatrices.ComputeProjectionNoAAMatrix();
 
@@ -157,11 +158,12 @@ public:
 		const FVertexFactory* VertexFactory,
 		const FMaterialRenderProxy* MaterialRenderProxy,
 		const FViewInfo& View, 
+		const FVolumetricFogIntegrationParameterData& IntegrationData,
 		const TUniformBufferRef<FViewUniformShaderParameters>& VoxelizeViewUniformBuffer,
 		FVector2D Jitter)
 	{
 		FMeshMaterialShader::SetParameters(RHICmdList, (FGeometryShaderRHIParamRef)GetGeometryShader(), MaterialRenderProxy, *MaterialRenderProxy->GetMaterial(View.GetFeatureLevel()), View, VoxelizeViewUniformBuffer, ESceneRenderTargetsMode::SetTextures);
-		VolumetricFogParameters.Set(RHICmdList, GetGeometryShader(), View, NULL, NULL, NULL);
+		VolumetricFogParameters.Set(RHICmdList, GetGeometryShader(), View, IntegrationData);
 
 		FMatrix ProjectionMatrix = View.ViewMatrices.ComputeProjectionNoAAMatrix();
 
@@ -221,10 +223,16 @@ protected:
 
 public:
 	
-	void SetParameters(FRHICommandList& RHICmdList, const FVertexFactory* VertexFactory,const FMaterialRenderProxy* MaterialRenderProxy,const FViewInfo& View, const TUniformBufferRef<FViewUniformShaderParameters>& VoxelizeViewUniformBuffer)
+	void SetParameters(
+		FRHICommandList& RHICmdList, 
+		const FVertexFactory* VertexFactory,
+		const FMaterialRenderProxy* MaterialRenderProxy,
+		const FViewInfo& View, 
+		const FVolumetricFogIntegrationParameterData& IntegrationData,
+		const TUniformBufferRef<FViewUniformShaderParameters>& VoxelizeViewUniformBuffer)
 	{
 		FMeshMaterialShader::SetParameters(RHICmdList, GetPixelShader(), MaterialRenderProxy, *MaterialRenderProxy->GetMaterial(View.GetFeatureLevel()), View, VoxelizeViewUniformBuffer, ESceneRenderTargetsMode::SetTextures);
-		VolumetricFogParameters.Set(RHICmdList, GetPixelShader(), View, NULL, NULL, NULL);
+		VolumetricFogParameters.Set(RHICmdList, GetPixelShader(), View, IntegrationData);
 	}
 
 	void SetMesh(FRHICommandList& RHICmdList, const FVertexFactory* VertexFactory,const FSceneView& View,const FPrimitiveSceneProxy* Proxy,const FMeshBatchElement& BatchElement,const FDrawingPolicyRenderState& DrawRenderState)
@@ -276,6 +284,7 @@ public:
 	void SetSharedState(
 		FRHICommandList& RHICmdList, 
 		const FViewInfo& View, 
+		const FVolumetricFogIntegrationParameterData& IntegrationData,
 		const TUniformBufferRef<FViewUniformShaderParameters>& VoxelizeViewUniformBuffer, 
 		FVector2D Jitter, 
 		const ContextDataType PolicyContext, 
@@ -334,6 +343,7 @@ FDrawingPolicyMatchResult FVoxelizeVolumeDrawingPolicy::Matches(
 void FVoxelizeVolumeDrawingPolicy::SetSharedState(
 	FRHICommandList& RHICmdList, 
 	const FViewInfo& View,
+	const FVolumetricFogIntegrationParameterData& IntegrationData,
 	const TUniformBufferRef<FViewUniformShaderParameters>& VoxelizeViewUniformBuffer,
 	FVector2D Jitter,
 	const ContextDataType PolicyContext,
@@ -343,12 +353,12 @@ void FVoxelizeVolumeDrawingPolicy::SetSharedState(
 	// Set shared mesh resources
 	FMeshDrawingPolicy::SetSharedState(RHICmdList, DrawRenderState, &View, PolicyContext);
 
-	VertexShader->SetParameters(RHICmdList, VertexFactory, MaterialRenderProxy, View, VoxelizeViewUniformBuffer, Jitter);
+	VertexShader->SetParameters(RHICmdList, VertexFactory, MaterialRenderProxy, View, IntegrationData, VoxelizeViewUniformBuffer, Jitter);
 	if (GeometryShader)
 	{
-		GeometryShader->SetParameters(RHICmdList, VertexFactory, MaterialRenderProxy, View, VoxelizeViewUniformBuffer, Jitter);
+		GeometryShader->SetParameters(RHICmdList, VertexFactory, MaterialRenderProxy, View, IntegrationData, VoxelizeViewUniformBuffer, Jitter);
 	}
-	PixelShader->SetParameters(RHICmdList, VertexFactory, MaterialRenderProxy, View, VoxelizeViewUniformBuffer);
+	PixelShader->SetParameters(RHICmdList, VertexFactory, MaterialRenderProxy, View, IntegrationData, VoxelizeViewUniformBuffer);
 }
 
 FBoundShaderStateInput FVoxelizeVolumeDrawingPolicy::GetBoundShaderStateInput(ERHIFeatureLevel::Type InFeatureLevel)
@@ -387,6 +397,7 @@ void FVoxelizeVolumeDrawingPolicy::SetMeshRenderState(
 void VoxelizeVolumePrimitive(
 	FRHICommandListImmediate& RHICmdList,
 	const FViewInfo& View,
+	const FVolumetricFogIntegrationParameterData& IntegrationData,
 	const TUniformBufferRef<FViewUniformShaderParameters>& VoxelizeViewUniformBuffer,
 	FVector2D Jitter,
 	FIntVector VolumetricFogGridSize,
@@ -408,7 +419,7 @@ void VoxelizeVolumePrimitive(
 		FDrawingPolicyRenderState DrawRenderState(View);
 		DrawingPolicy.SetupPipelineState(DrawRenderState, View);
 		CommitGraphicsPipelineState(RHICmdList, DrawingPolicy, DrawRenderState, DrawingPolicy.GetBoundShaderStateInput(View.GetFeatureLevel()));
-		DrawingPolicy.SetSharedState(RHICmdList, View, VoxelizeViewUniformBuffer, Jitter, FVoxelizeVolumeDrawingPolicy::ContextDataType(), DrawRenderState);
+		DrawingPolicy.SetSharedState(RHICmdList, View, IntegrationData, VoxelizeViewUniformBuffer, Jitter, FVoxelizeVolumeDrawingPolicy::ContextDataType(), DrawRenderState);
 
 		FBoxSphereBounds Bounds = PrimitiveSceneProxy->GetBounds();
 		//@todo - compute NumSlices based on the largest particle size.  Bounds is overly conservative in most cases.
@@ -442,11 +453,10 @@ void VoxelizeVolumePrimitive(
 void FDeferredShadingSceneRenderer::VoxelizeFogVolumePrimitives(
 	FRHICommandListImmediate& RHICmdList, 
 	const FViewInfo& View,
+	const FVolumetricFogIntegrationParameterData& IntegrationData,
 	FIntVector VolumetricFogGridSize,
 	FVector GridZParams,
-	float VolumetricFogDistance,
-	IPooledRenderTarget* VBufferA,
-	IPooledRenderTarget* VBufferB)
+	float VolumetricFogDistance)
 {
 	if (View.VolumetricPrimSet.NumPrims() > 0
 		&& DoesPlatformSupportVolumetricFogVoxelization(View.GetShaderPlatform()))
@@ -466,13 +476,12 @@ void FDeferredShadingSceneRenderer::VoxelizeFogVolumePrimitives(
 
 		TUniformBufferRef<FViewUniformShaderParameters> VoxelizeViewUniformBuffer = TUniformBufferRef<FViewUniformShaderParameters>::CreateUniformBufferImmediate(VoxelizeParameters, UniformBuffer_SingleFrame);
 
-		FVector JitterOffset = VolumetricFogTemporalRandom(View.Family->FrameNumber);
-		FVector2D Jitter(JitterOffset.X / VolumetricFogGridSize.X, JitterOffset.Y / VolumetricFogGridSize.Y);
+		FVector2D Jitter(IntegrationData.FrameJitterOffsetValues[0].X / VolumetricFogGridSize.X, IntegrationData.FrameJitterOffsetValues[0].Y / VolumetricFogGridSize.Y);
 
 		FTextureRHIParamRef RenderTargets[2] =
 		{
-			VBufferA->GetRenderTargetItem().TargetableTexture,
-			VBufferB->GetRenderTargetItem().TargetableTexture
+			IntegrationData.VBufferARenderTarget->GetRenderTargetItem().TargetableTexture,
+			IntegrationData.VBufferBRenderTarget->GetRenderTargetItem().TargetableTexture
 		};
 
 		SetRenderTargets(RHICmdList, ARRAY_COUNT(RenderTargets), RenderTargets, FTextureRHIParamRef(), 0, NULL);
@@ -499,7 +508,7 @@ void FDeferredShadingSceneRenderer::VoxelizeFogVolumePrimitives(
 						checkSlow(MeshBatchAndRelevance.PrimitiveSceneProxy == PrimitiveSceneInfo->Proxy);
 
 						const FMeshBatch& MeshBatch = *MeshBatchAndRelevance.Mesh;
-						VoxelizeVolumePrimitive(RHICmdList, View, VoxelizeViewUniformBuffer, Jitter, VolumetricFogGridSize, GridZParams, PrimitiveSceneProxy, MeshBatch);
+						VoxelizeVolumePrimitive(RHICmdList, View, IntegrationData, VoxelizeViewUniformBuffer, Jitter, VolumetricFogGridSize, GridZParams, PrimitiveSceneProxy, MeshBatch);
 					}
 				}
 
@@ -511,7 +520,7 @@ void FDeferredShadingSceneRenderer::VoxelizeFogVolumePrimitives(
 
 						if (View.StaticMeshVisibilityMap[StaticMesh.Id])
 						{
-							VoxelizeVolumePrimitive(RHICmdList, View, VoxelizeViewUniformBuffer, Jitter, VolumetricFogGridSize, GridZParams, PrimitiveSceneProxy, StaticMesh);
+							VoxelizeVolumePrimitive(RHICmdList, View, IntegrationData, VoxelizeViewUniformBuffer, Jitter, VolumetricFogGridSize, GridZParams, PrimitiveSceneProxy, StaticMesh);
 						}
 					}
 				}
diff --git a/Engine/Source/Runtime/Renderer/Public/PrimitiveSceneInfo.h b/Engine/Source/Runtime/Renderer/Public/PrimitiveSceneInfo.h
--- a/Engine/Source/Runtime/Renderer/Public/PrimitiveSceneInfo.h
+++ b/Engine/Source/Runtime/Renderer/Public/PrimitiveSceneInfo.h
@@ -138,6 +138,19 @@ public:
 	}
 };
 
+/** Flags needed for shadow culling.  These are pulled out of the FPrimitiveSceneProxy so we can do rough culling before dereferencing the proxy. */
+struct FPrimitiveFlagsCompact
+{
+	/** True if the primitive casts dynamic shadows. */
+	uint32 bCastDynamicShadow : 1;
+	/** True if the primitive will cache static lighting. */
+	uint32 bStaticLighting : 1;
+	/** True if the primitive casts static shadows. */
+	uint32 bCastStaticShadow : 1;
+
+	FPrimitiveFlagsCompact(const FPrimitiveSceneProxy* Proxy);
+};
+
 /** The information needed to determine whether a primitive is visible. */
 class FPrimitiveSceneInfoCompact
 {
@@ -147,28 +160,12 @@ public:
 	FBoxSphereBounds Bounds;
 	float MinDrawDistance;
 	float MaxDrawDistance;
-	float LpvBiasMultiplier;
 	/** Used for precomputed visibility */
 	int32 VisibilityId;
-	uint32 bHasViewDependentDPG : 1;
-	uint32 bCastDynamicShadow : 1;
-	uint32 bAffectDynamicIndirectLighting : 1;
-	uint32 StaticDepthPriorityGroup : SDPG_NumBits;
-
-	/** Initializes the compact scene info from the primitive's full scene info. */
-	void Init(FPrimitiveSceneInfo* InPrimitiveSceneInfo);
-
-	/** Default constructor. */
-	FPrimitiveSceneInfoCompact():
-		PrimitiveSceneInfo(NULL),
-		Proxy(NULL)
-	{}
+	FPrimitiveFlagsCompact PrimitiveFlagsCompact;
 
 	/** Initialization constructor. */
-	FPrimitiveSceneInfoCompact(FPrimitiveSceneInfo* InPrimitiveSceneInfo)
-	{
-		Init(InPrimitiveSceneInfo);
-	}
+	FPrimitiveSceneInfoCompact(FPrimitiveSceneInfo* InPrimitiveSceneInfo);
 };
 
 /** The type of the octree used by FScene to find primitives. */
@@ -447,7 +444,8 @@ private:
 /** Defines how the primitive is stored in the scene's primitive octree. */
 struct FPrimitiveOctreeSemantics
 {
-	enum { MaxElementsPerLeaf = 16 };
+	/** Note: this is coupled to shadow gather task granularity, see r.ParallelGatherShadowPrimitives. */
+	enum { MaxElementsPerLeaf = 256 };
 	enum { MinInclusiveElementsPerNode = 7 };
 	enum { MaxNodeDepth = 12 };
 
diff --git a/Engine/Source/Runtime/ShaderCore/Private/Shader.cpp b/Engine/Source/Runtime/ShaderCore/Private/Shader.cpp
--- a/Engine/Source/Runtime/ShaderCore/Private/Shader.cpp
+++ b/Engine/Source/Runtime/ShaderCore/Private/Shader.cpp
@@ -1864,8 +1864,8 @@ void ShaderMapAppendKeyString(EShaderPlatform Platform, FString& KeyString)
 	}
 
 	{
-		static const auto CVar = IConsoleManager::Get().FindTConsoleVariableDataInt(TEXT("r.ForwardShading"));
-		if (CVar && CVar->GetValueOnAnyThread() > 0)
+		static const auto CVar = IConsoleManager::Get().FindConsoleVariable(TEXT("r.ForwardShading"));
+		if (CVar && CVar->GetInt() > 0)
 		{
 			KeyString += TEXT("_FS");
 		}
diff --git a/Engine/Source/Runtime/ShaderCore/Private/ShaderCache.cpp b/Engine/Source/Runtime/ShaderCore/Private/ShaderCache.cpp
--- a/Engine/Source/Runtime/ShaderCore/Private/ShaderCache.cpp
+++ b/Engine/Source/Runtime/ShaderCore/Private/ShaderCache.cpp
@@ -15,6 +15,7 @@
 #include "Shader.h"
 #include "Misc/EngineVersion.h"
 #include "PipelineStateCache.h"
+#include "ScopeRWLock.h"
 
 DECLARE_STATS_GROUP(TEXT("Shader Cache"),STATGROUP_ShaderCache, STATCAT_Advanced);
 DECLARE_DWORD_ACCUMULATOR_STAT(TEXT("Num Shaders Cached"),STATGROUP_NumShadersCached,STATGROUP_ShaderCache);
@@ -944,10 +945,6 @@ FShaderCache::FShaderCache(uint32 InOptions)
 , OverridePrecompileTime(0)
 , OverridePredrawBatchTime(0)
 , bBatchingPaused(false)
-, ShaderCacheGlobalStateMutex(0)
-, PipelineStateMutex(1)
-, DrawLogMutex(2)
-, ContextCacheStatesMutex(3)
 , DefaultCacheState(nullptr)
 {
 	MaxTextureSamplers = FMath::Min<uint32>(GetMaxTextureSamplers(), FShaderDrawKey::MaxNumSamplers);
diff --git a/Engine/Source/Runtime/ShaderCore/Private/ShaderCacheTypes.cpp b/Engine/Source/Runtime/ShaderCore/Private/ShaderCacheTypes.cpp
--- a/Engine/Source/Runtime/ShaderCore/Private/ShaderCacheTypes.cpp
+++ b/Engine/Source/Runtime/ShaderCore/Private/ShaderCacheTypes.cpp
@@ -8,10 +8,6 @@
 #include "ShaderCache.h"
 #include "Serialization/MemoryWriter.h"
 
-#if RWLOCK_DEBUG_CHECKS
-thread_local TArray<uint32> FRWMutex::LockLevels;
-#endif
-
 bool FShaderDrawKey::bTrackDrawResources = true;
 
 FArchive& operator<<( FArchive& Ar, FShaderCaches& Info )
diff --git a/Engine/Source/Runtime/ShaderCore/Private/ShaderCodeLibrary.cpp b/Engine/Source/Runtime/ShaderCore/Private/ShaderCodeLibrary.cpp
--- a/Engine/Source/Runtime/ShaderCore/Private/ShaderCodeLibrary.cpp
+++ b/Engine/Source/Runtime/ShaderCore/Private/ShaderCodeLibrary.cpp
@@ -522,7 +522,11 @@ struct FEditorShaderCodeArchive
 	bool PackageNativeShaderLibrary(const FString& ShaderCodeDir, const FString& DebugShaderCodeDir)
 	{
 		bool bOK = false;
-		FString TempPath = GetShaderCodeFilename(FPaths::GameIntermediateDir(), FormatName) / TEXT("NativeLibrary");
+		
+		FString IntermediateFormatPath = GetShaderCodeFilename(FPaths::GameIntermediateDir(), FormatName);
+		FString IntermediateCookedByteCodePath = IntermediateFormatPath / TEXT("NativeCookedByteCode");
+		FString TempPath = IntermediateFormatPath / TEXT("NativeLibrary");
+		
 		EShaderPlatform Platform = ShaderFormatToLegacyShaderPlatform(FormatName);
 		IShaderFormatArchive* Archive = Format->CreateShaderArchive(FormatName, TempPath);
 		if (Archive)
@@ -531,8 +535,23 @@ struct FEditorShaderCodeArchive
 			FString DebugPath = GetShaderCodeFilename(DebugShaderCodeDir, FormatName);
 			bOK = true;
 			
+			//Collect previous native cooked bytecode files into this shader files processing directory - keep the rest of the code simpler, don't overwrite in case dest file is newer
+			{
+				TArray<FString> NativeShaderFiles;
+				IFileManager::Get().FindFiles(NativeShaderFiles, *IntermediateCookedByteCodePath, TEXT("*.ushaderbytecode"));
+				
+				for (FString const& FileName : NativeShaderFiles)
+				{
+					if (FileName.Len() > 2 && FileName[1] == TEXT('_'))
+					{
+						IFileManager::Get().Move(*(OutputPath / FileName), *(IntermediateCookedByteCodePath / FileName), false);
+					}
+				}
+			}
+			
 			TArray<FString> ShaderFiles;
 			IFileManager::Get().FindFiles(ShaderFiles, *OutputPath, TEXT("*.ushaderbytecode"));
+			
 			for (FString const& FileName : ShaderFiles)
 			{
 				if (FileName.Len() > 2 && FileName[1] == TEXT('_'))
@@ -580,9 +599,21 @@ struct FEditorShaderCodeArchive
 			{
 				bOK = Archive->Finalize(ShaderCodeDir, DebugPath, nullptr);
 				
+				//Always delete debug directory
+				IFileManager::Get().DeleteDirectory(*DebugShaderCodeDir, true, true);
+				
+				//Move files to intermediate dir for next iterative cook with overwwrite Move mode
 				if (bOK)
 				{
-					IFileManager::Get().DeleteDirectory(*DebugPath, true, true);
+					for (FString const& FileName : ShaderFiles)
+					{
+						if (FileName.Len() > 2 && FileName[1] == TEXT('_'))
+						{
+							IFileManager::Get().Move(*(IntermediateCookedByteCodePath / FileName), *(OutputPath / FileName), true);
+						}
+					}
+					
+					//We don't want to keep the shader code library shader files for native cooked content
 					IFileManager::Get().DeleteDirectory(*OutputPath, true, true);
 				}
 			}
@@ -655,6 +686,8 @@ public:
 		if(ShaderCodeArchive.IsValid())
 		{
 			bNativeFormat = true;
+			
+			UE_LOG(LogTemp, Display, TEXT("Cooked Context: Loaded Native Format Shared Shader Library"));
 		}
 		else
 		{
@@ -663,6 +696,8 @@ public:
 			{
 				ShaderCodeArchive = new FShaderCodeArchive(ShaderPlatform, Filename);
 				bSupportsPipelines = (ShaderCodeArchive != nullptr);
+				
+				UE_LOG(LogTemp, Display, TEXT("Cooked Context: Using Shared Shader Library"));
 			}
 		}
 		return IsValidRef(ShaderCodeArchive);
diff --git a/Engine/Source/Runtime/ShaderCore/Private/ShaderCore.cpp b/Engine/Source/Runtime/ShaderCore/Private/ShaderCore.cpp
--- a/Engine/Source/Runtime/ShaderCore/Private/ShaderCore.cpp
+++ b/Engine/Source/Runtime/ShaderCore/Private/ShaderCore.cpp
@@ -207,15 +207,15 @@ void FShaderCompilerOutput::GenerateOutputHash()
 */
 static void AddShaderSourceFileEntry( TArray<FString>& ShaderSourceFiles, const FString& ShaderFilename)
 {
-	FString ShaderFilenameBase( FPaths::GetBaseFilename(ShaderFilename) );
+	FString ShaderRelativeFilename = FindShaderRelativePath(ShaderFilename);
 
 	// get the filename for the the vertex factory type
-	if( !ShaderSourceFiles.Contains(ShaderFilenameBase) )
+	if( !ShaderSourceFiles.Contains(ShaderRelativeFilename) )
 	{
-		ShaderSourceFiles.Add(ShaderFilenameBase);
+		ShaderSourceFiles.Add(ShaderRelativeFilename);
 
 		TArray<FString> ShaderIncludes;
-		GetShaderIncludes(*ShaderFilenameBase,ShaderIncludes);
+		GetShaderIncludes(*ShaderRelativeFilename,ShaderIncludes);
 		for( int32 IncludeIdx=0; IncludeIdx < ShaderIncludes.Num(); IncludeIdx++ )
 		{
 			ShaderSourceFiles.AddUnique(ShaderIncludes[IncludeIdx]);
@@ -448,8 +448,10 @@ void GetShaderIncludes(const TCHAR* Filename, TArray<FString>& IncludeFilenames,
 					// ignore the header if it's optional and doesn't exist
 					if (bIsOptionalInclude)
 					{
-						FString ShaderFilename = FPaths::Combine(FPlatformProcess::BaseDir(), FPlatformProcess::ShaderDir(), *ExtractedIncludeFilename);
-						if (!FPaths::FileExists(ShaderFilename))
+						// Search in the default engine shader folder and in the same folder as the shader (in case its a project or plugin shader)
+						FString EngineShaderFilename = FPaths::Combine(FPlatformProcess::BaseDir(), FPlatformProcess::ShaderDir(), *ExtractedIncludeFilename);
+						FString LocalShaderFilename = FPaths::Combine(FPaths::GetPath(Filename), *ExtractedIncludeFilename);
+						if (!FPaths::FileExists(EngineShaderFilename) && !FPaths::FileExists(LocalShaderFilename))
 						{
 							bIgnoreInclude = true;
 						}
@@ -458,10 +460,9 @@ void GetShaderIncludes(const TCHAR* Filename, TArray<FString>& IncludeFilenames,
 					//vertex factories need to be handled separately
 					if (!bIgnoreInclude)
 					{
-						GetShaderIncludes(*ExtractedIncludeFilename, IncludeFilenames, DepthLimit - 1);
-						// maintain subdirectory info, but strip the extension
-						ExtractedIncludeFilename = FPaths::GetBaseFilename(ExtractedIncludeFilename, false);
-						IncludeFilenames.AddUnique(ExtractedIncludeFilename);
+						FString RelativeFilename = FindShaderRelativePath(ExtractedIncludeFilename);
+						GetShaderIncludes(*RelativeFilename, IncludeFilenames, DepthLimit - 1);
+						IncludeFilenames.AddUnique(RelativeFilename);
 					}
 				}
 			}
@@ -483,7 +484,8 @@ static void UpdateSingleShaderFilehash(FSHA1& InOutHashState, const TCHAR* Filen
 {
 	// Get the list of includes this file contains
 	TArray<FString> IncludeFilenames;
-	GetShaderIncludes(Filename, IncludeFilenames);
+	FString RelativeFilename = FindShaderRelativePath(Filename);
+	GetShaderIncludes(*RelativeFilename, IncludeFilenames);
 
 	for (int32 IncludeIndex = 0; IncludeIndex < IncludeFilenames.Num(); IncludeIndex++)
 	{
@@ -495,7 +497,7 @@ static void UpdateSingleShaderFilehash(FSHA1& InOutHashState, const TCHAR* Filen
 
 	// Load the source file and hash it
 	FString FileContents;
-	LoadShaderSourceFileChecked(Filename, FileContents);
+	LoadShaderSourceFileChecked(*RelativeFilename, FileContents);
 	InOutHashState.UpdateWithString(*FileContents, FileContents.Len());
 }
 
@@ -573,6 +575,49 @@ const FSHAHash& GetShaderFilesHash(const TArray<FString>& Filenames)
 	}
 }
 
+extern SHADERCORE_API FString FindShaderRelativePath(const FString Filename)
+{
+	bool bHadExtension = true;
+	FString ShaderFilename = Filename;
+	if (FPaths::GetExtension(ShaderFilename) == TEXT(""))
+	{
+		bHadExtension = false;
+		ShaderFilename += TEXT(".usf");
+	}
+
+	bool bFound = false;
+	FString ShaderPath;
+	for (FString ShaderDir : FPlatformProcess::AllShaderDirs())
+	{
+		FString Path = FPaths::Combine(FPlatformProcess::BaseDir(), ShaderDir, ShaderFilename);
+		if (FPaths::FileExists(Path))
+		{
+			ShaderPath = Path;
+			FPaths::MakePathRelativeTo(ShaderPath, *FPaths::Combine(FPlatformProcess::ShaderDir(), TEXT("/")));
+			bFound = true;
+			break;
+		}
+	}
+	if (!bFound)
+	{
+		UE_LOG(LogShaders, Fatal, TEXT("Couldn't find shader file \'%s\'"), *Filename);
+		UE_LOG(LogShaders, Fatal, TEXT("Known shader dirs:"));
+		for (FString ShaderDir : FPlatformProcess::AllShaderDirs())
+		{
+			UE_LOG(LogShaders, Fatal, TEXT("   %s"), *ShaderDir);
+		}
+	}
+
+	// If the filename didn't provide an extension, remove extension in returned path
+	// Extension are added in many places, don't want to end up with double extensions
+	if (!bHadExtension)
+	{
+		ShaderPath = FPaths::SetExtension(ShaderPath, TEXT(""));
+	}
+
+	return ShaderPath;
+}
+
 void BuildShaderFileToUniformBufferMap(TMap<FString, TArray<const TCHAR*> >& ShaderFileToUniformBufferVariables)
 {
 	if (!FPlatformProperties::RequiresCookedData())
@@ -685,8 +730,9 @@ void GenerateReferencedUniformBuffers(
 	TMap<const TCHAR*,FCachedUniformBufferDeclaration>& UniformBufferEntries)
 {
 	TArray<FString> FilesToSearch;
-	GetShaderIncludes(SourceFilename, FilesToSearch);
-	FilesToSearch.Add(SourceFilename);
+	FString ShaderRelativeFilename = FindShaderRelativePath(SourceFilename);
+	GetShaderIncludes(*ShaderRelativeFilename, FilesToSearch);
+	FilesToSearch.Add(ShaderRelativeFilename);
 
 	for (int32 FileIndex = 0; FileIndex < FilesToSearch.Num(); FileIndex++)
 	{
diff --git a/Engine/Source/Runtime/ShaderCore/Public/Shader.h b/Engine/Source/Runtime/ShaderCore/Public/Shader.h
--- a/Engine/Source/Runtime/ShaderCore/Public/Shader.h
+++ b/Engine/Source/Runtime/ShaderCore/Public/Shader.h
@@ -1417,13 +1417,18 @@ class TShaderMap
 	TArray<FShader*> SerializedShaders;
 	/** List of serialzied shader pipeline stages to be processed and registered on the game thread */
 	TArray<FSerializedShaderPipeline*> SerializedShaderPipelines;
+protected:
+	/** The platform this shader map was compiled with */
+	EShaderPlatform Platform;
+private:
 	/** Flag that makes sure this shader map isn't used until all shaders have been registerd */
 	bool bHasBeenRegistered;
 
 public:
 	/** Default constructor. */
-	TShaderMap()
-		: bHasBeenRegistered(true)
+	TShaderMap(EShaderPlatform InPlatform)
+		: Platform(InPlatform)
+		, bHasBeenRegistered(true)
 	{}
 
 	/** Destructor ensures pipelines cleared up. */
@@ -1432,13 +1437,15 @@ public:
 		EmptyShaderPipelines();
 	}
 
+	EShaderPlatform GetShaderPlatform() const { return Platform; }
+
 	/** Finds the shader with the given type.  Asserts on failure. */
 	template<typename ShaderType>
 	ShaderType* GetShader() const
 	{
 		check(bHasBeenRegistered);
 		const TRefCountPtr<FShader>* ShaderRef = Shaders.Find(&ShaderType::StaticType);
-		checkf(ShaderRef != NULL && *ShaderRef != nullptr, TEXT("Failed to find shader type %s"), ShaderType::StaticType.GetName());
+		checkf(ShaderRef != NULL && *ShaderRef != nullptr, TEXT("Failed to find shader type %s in Platform %s"), ShaderType::StaticType.GetName(), *LegacyShaderPlatformToShaderFormat(Platform).ToString());
 		return (ShaderType*)((*ShaderRef)->GetShaderChecked());
 	}
 
