Resolved depth target set merge issue.

diff --git a/Engine/Source/Runtime/Renderer/Private/PostProcess/SceneRenderTargets.cpp b/Engine/Source/Runtime/Renderer/Private/PostProcess/SceneRenderTargets.cpp
--- a/Engine/Source/Runtime/Renderer/Private/PostProcess/SceneRenderTargets.cpp
+++ b/Engine/Source/Runtime/Renderer/Private/PostProcess/SceneRenderTargets.cpp
@@ -2394,28 +2394,8 @@ void FSceneTextureShaderParameters::Set(
 
 		if (GSupportsDepthFetchDuringDepthTest)
 		{
-			const FTexture2DRHIRef* DepthTexture = SceneContext.GetActualDepthTexture();
-
-			if (SceneContext.IsSeparateTranslucencyPass() && SceneContext.IsDownsampledTranslucencyDepthValid())
-			{
-				FIntPoint OutScaledSize;
-				float OutScale;
-				SceneContext.GetSeparateTranslucencyDimensions(OutScaledSize, OutScale);
-
-				if (OutScale < 1.0f)
-				{
-					DepthTexture = &SceneContext.GetDownsampledTranslucencyDepthSurface();
-				}
-			}
-
-			SetTextureParameter(
-				RHICmdList, 
-				ShaderRHI,
-				SceneDepthTextureParameter,
-				SceneDepthTextureParameterSampler,
-				TStaticSamplerState<SF_Point,AM_Clamp,AM_Clamp,AM_Clamp>::GetRHI(),
-				*DepthTexture
-				);
+			DepthSurface = SceneContext.GetSceneDepthSurface();
+			DepthTextureNonMS = SceneContext.GetSceneDepthTexture();			
 		}
 		else
 		{
