Copying //UE4/Dev-Console to //UE4/Dev-Main (Source: //UE4/Dev-Console @ 3483086)

diff --git a/Engine/Source/Runtime/Renderer/Private/LightGridInjection.cpp b/Engine/Source/Runtime/Renderer/Private/LightGridInjection.cpp
--- a/Engine/Source/Runtime/Renderer/Private/LightGridInjection.cpp
+++ b/Engine/Source/Runtime/Renderer/Private/LightGridInjection.cpp
@@ -108,16 +108,16 @@ public:
 	}
 
 	template<typename ShaderRHIParamRef>
-	void Set(FRHICommandList& RHICmdList, const ShaderRHIParamRef& ShaderRHI, const FViewInfo& View)
+	void Set(FRHICommandList& RHICmdList, const ShaderRHIParamRef& ShaderRHI, const FForwardLightingCullingResources& ForwardLightingCullingResources)
 	{
-		NextCulledLightLink.SetBuffer(RHICmdList, ShaderRHI, View.ForwardLightingResources->NextCulledLightLink);
-		StartOffsetGrid.SetBuffer(RHICmdList, ShaderRHI, View.ForwardLightingResources->StartOffsetGrid);
-		CulledLightLinks.SetBuffer(RHICmdList, ShaderRHI, View.ForwardLightingResources->CulledLightLinks);
-		NextCulledLightData.SetBuffer(RHICmdList, ShaderRHI, View.ForwardLightingResources->NextCulledLightData);
+		NextCulledLightLink.SetBuffer(RHICmdList, ShaderRHI, ForwardLightingCullingResources.NextCulledLightLink);
+		StartOffsetGrid.SetBuffer(RHICmdList, ShaderRHI, ForwardLightingCullingResources.StartOffsetGrid);
+		CulledLightLinks.SetBuffer(RHICmdList, ShaderRHI, ForwardLightingCullingResources.CulledLightLinks);
+		NextCulledLightData.SetBuffer(RHICmdList, ShaderRHI, ForwardLightingCullingResources.NextCulledLightData);
 	}
 
 	template<typename ShaderRHIParamRef>
-	void UnsetParameters(FRHICommandList& RHICmdList, const ShaderRHIParamRef& ShaderRHI, const FViewInfo& View)
+	void UnsetParameters(FRHICommandList& RHICmdList, const ShaderRHIParamRef& ShaderRHI, const FForwardLightingCullingResources& ForwardLightingCullingResources)
 	{
 		NextCulledLightLink.UnsetUAV(RHICmdList, ShaderRHI);
 		StartOffsetGrid.UnsetUAV(RHICmdList, ShaderRHI);
@@ -128,22 +128,22 @@ public:
 
 		if (NextCulledLightLink.IsUAVBound())
 		{
-			OutUAVs.Add(View.ForwardLightingResources->NextCulledLightLink.UAV);
+			OutUAVs.Add(ForwardLightingCullingResources.NextCulledLightLink.UAV);
 		}
 
 		if (StartOffsetGrid.IsUAVBound())
 		{
-			OutUAVs.Add(View.ForwardLightingResources->StartOffsetGrid.UAV);
+			OutUAVs.Add(ForwardLightingCullingResources.StartOffsetGrid.UAV);
 		}
 
 		if (CulledLightLinks.IsUAVBound())
 		{
-			OutUAVs.Add(View.ForwardLightingResources->CulledLightLinks.UAV);
+			OutUAVs.Add(ForwardLightingCullingResources.CulledLightLinks.UAV);
 		}
 
 		if (NextCulledLightData.IsUAVBound())
 		{
-			OutUAVs.Add(View.ForwardLightingResources->NextCulledLightData.UAV);
+			OutUAVs.Add(ForwardLightingCullingResources.NextCulledLightData.UAV);
 		}
 
 		if (OutUAVs.Num() > 0)
@@ -204,18 +204,18 @@ public:
 	{
 	}
 
-	void SetParameters(FRHICommandList& RHICmdList, const FViewInfo& View)
+	void SetParameters(FRHICommandList& RHICmdList, const FViewInfo& View, const FForwardLightingCullingResources& ForwardLightingCullingResources)
 	{
 		FComputeShaderRHIParamRef ShaderRHI = GetComputeShader();
 		FGlobalShader::SetParameters<FViewUniformShaderParameters>(RHICmdList, ShaderRHI, View.ViewUniformBuffer);
 		ForwardLightingParameters.Set(RHICmdList, ShaderRHI, View);
-		ForwardCullingParameters.Set(RHICmdList, ShaderRHI, View);
+		ForwardCullingParameters.Set(RHICmdList, ShaderRHI, ForwardLightingCullingResources);
 	}
 
-	void UnsetParameters(FRHICommandList& RHICmdList, const FViewInfo& View)
+	void UnsetParameters(FRHICommandList& RHICmdList, const FViewInfo& View, const FForwardLightingCullingResources& ForwardLightingCullingResources)
 	{
 		ForwardLightingParameters.UnsetParameters(RHICmdList, GetComputeShader(), View);
-		ForwardCullingParameters.UnsetParameters(RHICmdList, GetComputeShader(), View);
+		ForwardCullingParameters.UnsetParameters(RHICmdList, GetComputeShader(), ForwardLightingCullingResources);
 	}
 
 	virtual bool Serialize(FArchive& Ar)
@@ -265,18 +265,18 @@ public:
 	{
 	}
 
-	void SetParameters(FRHICommandList& RHICmdList, const FViewInfo& View)
+	void SetParameters(FRHICommandList& RHICmdList, const FViewInfo& View, const FForwardLightingCullingResources& ForwardLightingCullingResources)
 	{
 		FComputeShaderRHIParamRef ShaderRHI = GetComputeShader();
 		FGlobalShader::SetParameters<FViewUniformShaderParameters>(RHICmdList, ShaderRHI, View.ViewUniformBuffer);
 		ForwardLightingParameters.Set(RHICmdList, ShaderRHI, View);
-		ForwardCullingParameters.Set(RHICmdList, ShaderRHI, View);
+		ForwardCullingParameters.Set(RHICmdList, ShaderRHI, ForwardLightingCullingResources);
 	}
 
-	void UnsetParameters(FRHICommandList& RHICmdList, const FViewInfo& View)
+	void UnsetParameters(FRHICommandList& RHICmdList, const FViewInfo& View, const FForwardLightingCullingResources& ForwardLightingCullingResources)
 	{
 		ForwardLightingParameters.UnsetParameters(RHICmdList, GetComputeShader(), View);
-		ForwardCullingParameters.UnsetParameters(RHICmdList, GetComputeShader(), View);
+		ForwardCullingParameters.UnsetParameters(RHICmdList, GetComputeShader(), ForwardLightingCullingResources);
 	}
 
 	virtual bool Serialize(FArchive& Ar)
@@ -611,15 +611,33 @@ void FDeferredShadingSceneRenderer::ComputeLightGrid(FRHICommandListImmediate& R
 			if (View.ForwardLightingResources->NumCulledLightsGrid.NumBytes != NumCells * NumCulledLightsGridStride * sizeof(uint32))
 			{
 				View.ForwardLightingResources->NumCulledLightsGrid.Initialize(sizeof(uint32), NumCells * NumCulledLightsGridStride, PF_R32_UINT);
-				View.ForwardLightingResources->NextCulledLightLink.Initialize(sizeof(uint32), 1, PF_R32_UINT);
-				View.ForwardLightingResources->StartOffsetGrid.Initialize(sizeof(uint32), NumCells, PF_R32_UINT);
-				View.ForwardLightingResources->NextCulledLightData.Initialize(sizeof(uint32), 1, PF_R32_UINT);
 			}
 
 			if (View.ForwardLightingResources->CulledLightDataGrid.NumBytes != NumCells * GMaxCulledLightsPerCell * sizeof(FLightIndexType))
 			{
 				View.ForwardLightingResources->CulledLightDataGrid.Initialize(sizeof(FLightIndexType), NumCells * GMaxCulledLightsPerCell, sizeof(FLightIndexType) == sizeof(uint16) ? PF_R16_UINT : PF_R32_UINT);
-				View.ForwardLightingResources->CulledLightLinks.Initialize(sizeof(uint32), NumCells * GMaxCulledLightsPerCell * LightLinkStride, PF_R32_UINT);
+			}
+
+			const bool bShouldCacheTemporaryBuffers = View.ViewState != nullptr;
+			FForwardLightingCullingResources LocalCullingResources;
+			FForwardLightingCullingResources& ForwardLightingCullingResources = bShouldCacheTemporaryBuffers ? View.ViewState->ForwardLightingCullingResources : LocalCullingResources;
+
+			if (ForwardLightingCullingResources.CulledLightLinks.NumBytes != NumCells * GMaxCulledLightsPerCell * LightLinkStride)
+			{
+				const uint32 FastVRamFlag = IsTransientResourceBufferAliasingEnabled() ? (BUF_FastVRAM | BUF_Transient) : BUF_None;
+				ForwardLightingCullingResources.CulledLightLinks.Initialize(sizeof(uint32), NumCells * GMaxCulledLightsPerCell * LightLinkStride, PF_R32_UINT, FastVRamFlag, TEXT("CulledLightLinks"));
+				ForwardLightingCullingResources.NextCulledLightLink.Initialize(sizeof(uint32), 1, PF_R32_UINT, FastVRamFlag, TEXT("NextCulledLightLink"));
+				ForwardLightingCullingResources.StartOffsetGrid.Initialize(sizeof(uint32), NumCells, PF_R32_UINT, FastVRamFlag, TEXT("StartOffsetGrid"));
+				ForwardLightingCullingResources.NextCulledLightData.Initialize(sizeof(uint32), 1, PF_R32_UINT, FastVRamFlag, TEXT("NextCulledLightData"));
+			}
+
+			if (IsTransientResourceBufferAliasingEnabled())
+			{
+				// Acquire resources
+				ForwardLightingCullingResources.CulledLightLinks.AcquireTransientResource();
+				ForwardLightingCullingResources.NextCulledLightLink.AcquireTransientResource();
+				ForwardLightingCullingResources.StartOffsetGrid.AcquireTransientResource();
+				ForwardLightingCullingResources.NextCulledLightData.AcquireTransientResource();
 			}
 
 			const FIntVector NumGroups = FIntVector::DivideAndRoundUp(FIntVector(LightGridSizeXY.X, LightGridSizeXY.Y, GLightGridSizeZ), LightGridInjectionGroupSize);
@@ -635,23 +653,23 @@ void FDeferredShadingSceneRenderer::ComputeLightGrid(FRHICommandListImmediate& R
 				TArray<FUnorderedAccessViewRHIParamRef, TInlineAllocator<6>> OutUAVs;
 				OutUAVs.Add(View.ForwardLightingResources->NumCulledLightsGrid.UAV);
 				OutUAVs.Add(View.ForwardLightingResources->CulledLightDataGrid.UAV);
-				OutUAVs.Add(View.ForwardLightingResources->NextCulledLightLink.UAV);
-				OutUAVs.Add(View.ForwardLightingResources->StartOffsetGrid.UAV);
-				OutUAVs.Add(View.ForwardLightingResources->CulledLightLinks.UAV);
-				OutUAVs.Add(View.ForwardLightingResources->NextCulledLightData.UAV);
+				OutUAVs.Add(ForwardLightingCullingResources.NextCulledLightLink.UAV);
+				OutUAVs.Add(ForwardLightingCullingResources.StartOffsetGrid.UAV);
+				OutUAVs.Add(ForwardLightingCullingResources.CulledLightLinks.UAV);
+				OutUAVs.Add(ForwardLightingCullingResources.NextCulledLightData.UAV);
 				RHICmdList.TransitionResources(EResourceTransitionAccess::EWritable, EResourceTransitionPipeline::EGfxToCompute, OutUAVs.GetData(), OutUAVs.Num());
 
 				if (GLightLinkedListCulling)
 				{
-					ClearUAV(RHICmdList, View.ForwardLightingResources->StartOffsetGrid, 0xFFFFFFFF);
-					ClearUAV(RHICmdList, View.ForwardLightingResources->NextCulledLightLink, 0);
-					ClearUAV(RHICmdList, View.ForwardLightingResources->NextCulledLightData, 0);
+					ClearUAV(RHICmdList, ForwardLightingCullingResources.StartOffsetGrid, 0xFFFFFFFF);
+					ClearUAV(RHICmdList, ForwardLightingCullingResources.NextCulledLightLink, 0);
+					ClearUAV(RHICmdList, ForwardLightingCullingResources.NextCulledLightData, 0);
 
 					TShaderMapRef<TLightGridInjectionCS<true> > ComputeShader(View.ShaderMap);
 					RHICmdList.SetComputeShader(ComputeShader->GetComputeShader());
-					ComputeShader->SetParameters(RHICmdList, View);
+					ComputeShader->SetParameters(RHICmdList, View, ForwardLightingCullingResources);
 					DispatchComputeShader(RHICmdList, *ComputeShader, NumGroups.X, NumGroups.Y, NumGroups.Z);
-					ComputeShader->UnsetParameters(RHICmdList, View);
+					ComputeShader->UnsetParameters(RHICmdList, View, ForwardLightingCullingResources);
 				}
 				else
 				{
@@ -659,9 +677,9 @@ void FDeferredShadingSceneRenderer::ComputeLightGrid(FRHICommandListImmediate& R
 
 					TShaderMapRef<TLightGridInjectionCS<false> > ComputeShader(View.ShaderMap);
 					RHICmdList.SetComputeShader(ComputeShader->GetComputeShader());
-					ComputeShader->SetParameters(RHICmdList, View);
+					ComputeShader->SetParameters(RHICmdList, View, ForwardLightingCullingResources);
 					DispatchComputeShader(RHICmdList, *ComputeShader, NumGroups.X, NumGroups.Y, NumGroups.Z);
-					ComputeShader->UnsetParameters(RHICmdList, View);
+					ComputeShader->UnsetParameters(RHICmdList, View, ForwardLightingCullingResources);
 				}
 			}
 
@@ -671,9 +689,16 @@ void FDeferredShadingSceneRenderer::ComputeLightGrid(FRHICommandListImmediate& R
 
 				TShaderMapRef<FLightGridCompactCS> ComputeShader(View.ShaderMap);
 				RHICmdList.SetComputeShader(ComputeShader->GetComputeShader());
-				ComputeShader->SetParameters(RHICmdList, View);
+				ComputeShader->SetParameters(RHICmdList, View, ForwardLightingCullingResources);
 				DispatchComputeShader(RHICmdList, *ComputeShader, NumGroups.X, NumGroups.Y, NumGroups.Z);
-				ComputeShader->UnsetParameters(RHICmdList, View);
+				ComputeShader->UnsetParameters(RHICmdList, View, ForwardLightingCullingResources);
+			}
+			if (IsTransientResourceBufferAliasingEnabled())
+			{
+				ForwardLightingCullingResources.CulledLightLinks.DiscardTransientResource();
+				ForwardLightingCullingResources.NextCulledLightLink.DiscardTransientResource();
+				ForwardLightingCullingResources.StartOffsetGrid.DiscardTransientResource();
+				ForwardLightingCullingResources.NextCulledLightData.DiscardTransientResource();
 			}
 		}
 	}
@@ -717,14 +742,14 @@ void FDeferredShadingSceneRenderer::RenderForwardShadingShadowProjections(FRHICo
 
 			if (VisibleLightInfo.ShadowsToProject.Num() > 0)
 			{
-				FSceneRenderer::RenderShadowProjections(RHICmdList, LightSceneInfo, true, false);
+				FSceneRenderer::RenderShadowProjections(RHICmdList, LightSceneInfo, SceneRenderTargets.GetLightAttenuation(), true, false);
 			}
 
-			RenderCapsuleDirectShadows(*LightSceneInfo, RHICmdList, VisibleLightInfo.CapsuleShadowsToProject, true);
+			RenderCapsuleDirectShadows(RHICmdList, *LightSceneInfo, SceneRenderTargets.GetLightAttenuation(), VisibleLightInfo.CapsuleShadowsToProject, true);
 
 			if (LightSceneInfo->GetDynamicShadowMapChannel() >= 0 && LightSceneInfo->GetDynamicShadowMapChannel() < 4)
 			{
-				RenderLightFunction(RHICmdList, LightSceneInfo, true, true);
+				RenderLightFunction(RHICmdList, LightSceneInfo, SceneRenderTargets.GetLightAttenuation(), true, true);
 			}
 		}
 
diff --git a/Engine/Source/Runtime/Renderer/Private/LightPropagationVolume.cpp b/Engine/Source/Runtime/Renderer/Private/LightPropagationVolume.cpp
--- a/Engine/Source/Runtime/Renderer/Private/LightPropagationVolume.cpp
+++ b/Engine/Source/Runtime/Renderer/Private/LightPropagationVolume.cpp
@@ -823,13 +823,13 @@ FLightPropagationVolume::FLightPropagationVolume() :
 	int32 RSMResolution = FSceneRenderTargets::Get_FrameConstantsOnly().GetReflectiveShadowMapResolution();
 	int32 GvListBufferSize = RSMResolution * RSMResolution * 16; // Allow 16 layers of depth per every pixel of the RSM (on average) 
 	int32 VplListBufferSize = RSMResolution * RSMResolution * 4; // Allow 4 layers of depth per pixel in the RSM (1 for the RSM injection + 3 for light injection)
-	mVplListBuffer->Initialize( sizeof( VplListEntry ), VplListBufferSize, 0, true, false );
+	mVplListBuffer->Initialize( sizeof( VplListEntry ), VplListBufferSize, 0, TEXT("mVplListBuffer"), true, false );
 	mVplListHeadBuffer = new FRWBufferByteAddress();
 	mVplListHeadBuffer->Initialize( LPV_GRIDRES*LPV_GRIDRES*LPV_GRIDRES*4, BUF_ByteAddressBuffer );
 
 	// Geometry volume buffers
 	GvListBuffer = new FRWBufferStructured();
-	GvListBuffer->Initialize( sizeof( VplListEntry ), GvListBufferSize, 0, true, false );
+	GvListBuffer->Initialize( sizeof( VplListEntry ), GvListBufferSize, 0, TEXT("GvListBuffer"), true, false );
 	GvListHeadBuffer = new FRWBufferByteAddress();
 	GvListHeadBuffer->Initialize( LPV_GRIDRES*LPV_GRIDRES*LPV_GRIDRES*4, BUF_ByteAddressBuffer );
 
diff --git a/Engine/Source/Runtime/Renderer/Private/LightRendering.cpp b/Engine/Source/Runtime/Renderer/Private/LightRendering.cpp
--- a/Engine/Source/Runtime/Renderer/Private/LightRendering.cpp
+++ b/Engine/Source/Runtime/Renderer/Private/LightRendering.cpp
@@ -79,17 +79,17 @@ public:
 	{
 	}
 
-	void SetParameters(FRHICommandList& RHICmdList, const FSceneView& View, const FLightSceneInfo* LightSceneInfo)
+	void SetParameters(FRHICommandList& RHICmdList, const FSceneView& View, const FLightSceneInfo* LightSceneInfo, IPooledRenderTarget* ScreenShadowMaskTexture)
 	{
 		const FPixelShaderRHIParamRef ShaderRHI = GetPixelShader();
-		SetParametersBase(RHICmdList, ShaderRHI, View, LightSceneInfo->Proxy->GetIESTextureResource());
+		SetParametersBase(RHICmdList, ShaderRHI, View, ScreenShadowMaskTexture, LightSceneInfo->Proxy->GetIESTextureResource());
 		SetDeferredLightParameters(RHICmdList, ShaderRHI, GetUniformBufferParameter<FDeferredLightUniformStruct>(), LightSceneInfo, View);
 	}
 
 	void SetParametersSimpleLight(FRHICommandList& RHICmdList, const FSceneView& View, const FSimpleLightEntry& SimpleLight, const FSimpleLightPerViewEntry& SimpleLightPerViewData)
 	{
 		const FPixelShaderRHIParamRef ShaderRHI = GetPixelShader();
-		SetParametersBase(RHICmdList, ShaderRHI, View, NULL);
+		SetParametersBase(RHICmdList, ShaderRHI, View, NULL, NULL);
 		SetSimpleDeferredLightParameters(RHICmdList, ShaderRHI, GetUniformBufferParameter<FDeferredLightUniformStruct>(), SimpleLight, SimpleLightPerViewData, View);
 	}
 
@@ -110,7 +110,7 @@ public:
 
 private:
 
-	void SetParametersBase(FRHICommandList& RHICmdList, const FPixelShaderRHIParamRef ShaderRHI, const FSceneView& View, FTexture* IESTextureResource)
+	void SetParametersBase(FRHICommandList& RHICmdList, const FPixelShaderRHIParamRef ShaderRHI, const FSceneView& View, IPooledRenderTarget* ScreenShadowMaskTexture, FTexture* IESTextureResource)
 	{
 		FGlobalShader::SetParameters<FViewUniformShaderParameters>(RHICmdList, ShaderRHI,View.ViewUniformBuffer);
 		DeferredParameters.Set(RHICmdList, ShaderRHI, View, MD_PostProcess);
@@ -125,7 +125,7 @@ private:
 				LightAttenuationTexture,
 				LightAttenuationTextureSampler,
 				TStaticSamplerState<SF_Point,AM_Wrap,AM_Wrap,AM_Wrap>::GetRHI(),
-				SceneRenderTargets.GetEffectiveLightAttenuationTexture(true)
+				ScreenShadowMaskTexture ? ScreenShadowMaskTexture->GetRenderTargetItem().ShaderResourceTexture : GWhiteTexture->TextureRHI
 				);
 		}
 
@@ -350,7 +350,7 @@ void FDeferredShadingSceneRenderer::RenderLights(FRHICommandListImmediate& RHICm
 			{
 				if (LightSceneInfo->ShouldRenderLight(Views[ViewIndex]))
 				{
-					FSortedLightSceneInfo* SortedLightInfo = new(SortedLights) FSortedLightSceneInfo(LightSceneInfoCompact);
+					FSortedLightSceneInfo* SortedLightInfo = new(SortedLights) FSortedLightSceneInfo(LightSceneInfo);
 
 					// Check for shadows and light functions.
 					SortedLightInfo->SortKey.Fields.LightType = LightSceneInfoCompact.LightType;
@@ -426,7 +426,6 @@ void FDeferredShadingSceneRenderer::RenderLights(FRHICommandListImmediate& RHICm
 		{
 			SCOPED_DRAW_EVENT(RHICmdList, NonShadowedLights);
 			INC_DWORD_STAT_BY(STAT_NumUnshadowedLights, AttenuationLightStart);
-			SceneContext.SetLightAttenuationMode(false);
 
 			int32 StandardDeferredStart = 0;
 
@@ -470,11 +469,10 @@ void FDeferredShadingSceneRenderer::RenderLights(FRHICommandListImmediate& RHICm
 				for (int32 LightIndex = StandardDeferredStart; LightIndex < AttenuationLightStart; LightIndex++)
 				{
 					const FSortedLightSceneInfo& SortedLightInfo = SortedLights[LightIndex];
-					const FLightSceneInfoCompact& LightSceneInfoCompact = SortedLightInfo.SceneInfo;
-					const FLightSceneInfo* const LightSceneInfo = LightSceneInfoCompact.LightSceneInfo;
+					const FLightSceneInfo* const LightSceneInfo = SortedLightInfo.LightSceneInfo;
 
 					// Render the light to the scene color buffer, using a 1x1 white texture as input 
-					RenderLight(RHICmdList, LightSceneInfo, false, false);
+					RenderLight(RHICmdList, LightSceneInfo, NULL, false, false);
 				}
 			}
 
@@ -506,8 +504,7 @@ void FDeferredShadingSceneRenderer::RenderLights(FRHICommandListImmediate& RHICm
 			for (int32 LightIndex = AttenuationLightStart; LightIndex < SortedLights.Num(); LightIndex++)
 			{
 				const FSortedLightSceneInfo& SortedLightInfo = SortedLights[LightIndex];
-				const FLightSceneInfoCompact& LightSceneInfoCompact = SortedLightInfo.SceneInfo;
-				const FLightSceneInfo& LightSceneInfo = *LightSceneInfoCompact.LightSceneInfo;
+				const FLightSceneInfo& LightSceneInfo = *SortedLightInfo.LightSceneInfo;
 				// Render any reflective shadow maps (if necessary)
 				if ( LightSceneInfo.Proxy && LightSceneInfo.Proxy->NeedsLPVInjection() )
 				{
@@ -526,8 +523,7 @@ void FDeferredShadingSceneRenderer::RenderLights(FRHICommandListImmediate& RHICm
 				for (int32 LightIndex = 0; LightIndex < SortedLights.Num(); LightIndex++)
 				{
 					const FSortedLightSceneInfo& SortedLightInfo = SortedLights[LightIndex];
-					const FLightSceneInfoCompact& LightSceneInfoCompact = SortedLightInfo.SceneInfo;
-					const FLightSceneInfo* const LightSceneInfo = LightSceneInfoCompact.LightSceneInfo;
+					const FLightSceneInfo* const LightSceneInfo = SortedLightInfo.LightSceneInfo;
 
 					// Render any reflective shadow maps (if necessary)
 					if ( LightSceneInfo && LightSceneInfo->Proxy && LightSceneInfo->Proxy->NeedsLPVInjection() )
@@ -566,18 +562,28 @@ void FDeferredShadingSceneRenderer::RenderLights(FRHICommandListImmediate& RHICm
 
 			bool bDirectLighting = ViewFamily.EngineShowFlags.DirectLighting;
 
+			TRefCountPtr<IPooledRenderTarget> ScreenShadowMaskTexture;
+
 			// Draw shadowed and light function lights
 			for (int32 LightIndex = AttenuationLightStart; LightIndex < SortedLights.Num(); LightIndex++)
 			{
 				const FSortedLightSceneInfo& SortedLightInfo = SortedLights[LightIndex];
-				const FLightSceneInfoCompact& LightSceneInfoCompact = SortedLightInfo.SceneInfo;
-				const FLightSceneInfo& LightSceneInfo = *LightSceneInfoCompact.LightSceneInfo;
+				const FLightSceneInfo& LightSceneInfo = *SortedLightInfo.LightSceneInfo;
 				bool bDrawShadows = SortedLightInfo.SortKey.Fields.bShadowed;
 				bool bDrawLightFunction = SortedLightInfo.SortKey.Fields.bLightFunction;
+				bool bDrawPreviewIndicator = ViewFamily.EngineShowFlags.PreviewShadowsIndicator && !LightSceneInfo.IsPrecomputedLightingValid() && LightSceneInfo.Proxy->HasStaticShadowing();
 				bool bInjectedTranslucentVolume = false;
-				bool bUsedLightAttenuation = false;
+				bool bUsedShadowMaskTexture = false;
 				FScopeCycleCounter Context(LightSceneInfo.Proxy->GetStatId());
 
+				if ((bDrawShadows || bDrawLightFunction || bDrawPreviewIndicator) && !ScreenShadowMaskTexture.IsValid())
+				{
+					FPooledRenderTargetDesc Desc(FPooledRenderTargetDesc::Create2DDesc(SceneContext.GetBufferSizeXY(), PF_B8G8R8A8, FClearValueBinding::White, TexCreate_None, TexCreate_RenderTargetable, false));
+					Desc.Flags |= GetTextureFastVRamFlag_DynamicLayout();
+					Desc.NumSamples = SceneContext.GetNumSceneColorMSAASamples(SceneContext.GetCurrentFeatureLevel());
+					GRenderTargetPool.FindFreeElement(RHICmdList, Desc, ScreenShadowMaskTexture, TEXT("ScreenShadowMaskTexture"));
+				}
+
 				FString LightNameWithLevel;
 				GetLightNameForDrawEvent(LightSceneInfo.Proxy, LightNameWithLevel);
 				SCOPED_DRAW_EVENTF(RHICmdList, EventLightPass, *LightNameWithLevel);
@@ -593,11 +599,11 @@ void FDeferredShadingSceneRenderer::RenderLights(FRHICommandListImmediate& RHICm
 					}
 
 					// Clear light attenuation for local lights with a quad covering their extents
-					const bool bClearLightScreenExtentsOnly = SortedLightInfo.SceneInfo.LightType != LightType_Directional;
-
+					const bool bClearLightScreenExtentsOnly = SortedLightInfo.SortKey.Fields.LightType != LightType_Directional;
 					// All shadows render with min blending
 					bool bClearToWhite = !bClearLightScreenExtentsOnly;
-					SceneContext.BeginRenderingLightAttenuation(RHICmdList, bClearToWhite);
+
+					SetRenderTarget(RHICmdList, ScreenShadowMaskTexture->GetRenderTargetItem().TargetableTexture, SceneContext.GetSceneDepthSurface(), bClearToWhite ? ESimpleRenderTargetMode::EClearColorExistingDepth : ESimpleRenderTargetMode::EExistingColorAndDepth, FExclusiveDepthStencil::DepthRead_StencilWrite, true);
 
 					if (bClearLightScreenExtentsOnly)
 					{
@@ -618,9 +624,9 @@ void FDeferredShadingSceneRenderer::RenderLights(FRHICommandListImmediate& RHICm
 						}
 					}
 
-					RenderShadowProjections(RHICmdList, &LightSceneInfo, bInjectedTranslucentVolume);
+					RenderShadowProjections(RHICmdList, &LightSceneInfo, ScreenShadowMaskTexture, bInjectedTranslucentVolume);
 
-					bUsedLightAttenuation = true;
+					bUsedShadowMaskTexture = true;
 				}
 
 				for (int32 ViewIndex = 0; ViewIndex < Views.Num(); ViewIndex++)
@@ -632,14 +638,15 @@ void FDeferredShadingSceneRenderer::RenderLights(FRHICommandListImmediate& RHICm
 				// Render light function to the attenuation buffer.
 				if (bDirectLighting)
 				{
-					const bool bLightFunctionRendered = RenderLightFunction(RHICmdList, &LightSceneInfo, bDrawShadows, false);
-					bUsedLightAttenuation |= bLightFunctionRendered;
+					if (bDrawLightFunction)
+					{
+						const bool bLightFunctionRendered = RenderLightFunction(RHICmdList, &LightSceneInfo, ScreenShadowMaskTexture, bDrawShadows, false);
+						bUsedShadowMaskTexture |= bLightFunctionRendered;
+					}
 
-					if (ViewFamily.EngineShowFlags.PreviewShadowsIndicator
-						&& !LightSceneInfo.IsPrecomputedLightingValid() 
-						&& LightSceneInfo.Proxy->HasStaticShadowing())
+					if (bDrawPreviewIndicator)
 					{
-						RenderPreviewShadowsIndicator(RHICmdList, &LightSceneInfo, bUsedLightAttenuation);
+						RenderPreviewShadowsIndicator(RHICmdList, &LightSceneInfo, ScreenShadowMaskTexture, bUsedShadowMaskTexture);
 					}
 
 					if (!bDrawShadows)
@@ -648,10 +655,9 @@ void FDeferredShadingSceneRenderer::RenderLights(FRHICommandListImmediate& RHICm
 					}
 				}
 			
-				if( bUsedLightAttenuation )
+				if (bUsedShadowMaskTexture)
 				{
-					// Resolve light attenuation buffer
-					SceneContext.FinishRenderingLightAttenuation(RHICmdList);
+					RHICmdList.CopyToResolveTarget(ScreenShadowMaskTexture->GetRenderTargetItem().TargetableTexture, ScreenShadowMaskTexture->GetRenderTargetItem().ShaderResourceTexture, false, FResolveParams(FResolveRect()));
 				}
 			
 				if(bDirectLighting && !bInjectedTranslucentVolume)
@@ -661,19 +667,14 @@ void FDeferredShadingSceneRenderer::RenderLights(FRHICommandListImmediate& RHICm
 					InjectTranslucentVolumeLighting(RHICmdList, LightSceneInfo, NULL);
 				}
 
-				SceneContext.SetLightAttenuationMode(bUsedLightAttenuation);
 				SceneContext.BeginRenderingSceneColor(RHICmdList, ESimpleRenderTargetMode::EExistingColorAndDepth, FExclusiveDepthStencil::DepthRead_StencilWrite);
 
 				// Render the light to the scene color buffer, conditionally using the attenuation buffer or a 1x1 white texture as input 
 				if(bDirectLighting)
 				{
-					RenderLight(RHICmdList, &LightSceneInfo, false, true);
+					RenderLight(RHICmdList, &LightSceneInfo, ScreenShadowMaskTexture, false, true);
 				}
 			}
-
-			// Restore the default mode
-			SceneContext.SetLightAttenuationMode(true);
-
 		}
 	}
 }
@@ -705,7 +706,7 @@ void FDeferredShadingSceneRenderer::RenderLightArrayForOverlapViewmode(FRHIComma
 			&& !LightSceneInfo->Proxy->HasStaticLighting()
 			&& LightSceneInfo->Proxy->CastsStaticShadow())
 		{
-			RenderLight(RHICmdList, LightSceneInfo, true, false);
+			RenderLight(RHICmdList, LightSceneInfo, NULL, true, false);
 		}
 	}
 }
@@ -764,7 +765,8 @@ static void SetShaderTemplLighting(
 	FGraphicsPipelineStateInitializer& GraphicsPSOInit,
 	const FViewInfo& View,
 	FShader* VertexShader,
-	const FLightSceneInfo* LightSceneInfo)
+	const FLightSceneInfo* LightSceneInfo, 
+	IPooledRenderTarget* ScreenShadowMaskTexture)
 {
 	if(View.Family->EngineShowFlags.VisualizeLightCulling)
 	{
@@ -773,7 +775,7 @@ static void SetShaderTemplLighting(
 		GraphicsPSOInit.BoundShaderState.VertexShaderRHI = GETSAFERHISHADER_VERTEX(VertexShader);
 		GraphicsPSOInit.BoundShaderState.PixelShaderRHI = GETSAFERHISHADER_PIXEL(*PixelShader);
 		SetGraphicsPipelineState(RHICmdList, GraphicsPSOInit);
-		PixelShader->SetParameters(RHICmdList, View, LightSceneInfo);
+		PixelShader->SetParameters(RHICmdList, View, LightSceneInfo, ScreenShadowMaskTexture);
 	}
 	else
 	{
@@ -784,7 +786,7 @@ static void SetShaderTemplLighting(
 			GraphicsPSOInit.BoundShaderState.VertexShaderRHI = GETSAFERHISHADER_VERTEX(VertexShader);
 			GraphicsPSOInit.BoundShaderState.PixelShaderRHI = GETSAFERHISHADER_PIXEL(*PixelShader);
 			SetGraphicsPipelineState(RHICmdList, GraphicsPSOInit);
-			PixelShader->SetParameters(RHICmdList, View, LightSceneInfo);
+			PixelShader->SetParameters(RHICmdList, View, LightSceneInfo, ScreenShadowMaskTexture);
 		}
 		else
 		{
@@ -793,7 +795,7 @@ static void SetShaderTemplLighting(
 			GraphicsPSOInit.BoundShaderState.VertexShaderRHI = GETSAFERHISHADER_VERTEX(VertexShader);
 			GraphicsPSOInit.BoundShaderState.PixelShaderRHI = GETSAFERHISHADER_PIXEL(*PixelShader);
 			SetGraphicsPipelineState(RHICmdList, GraphicsPSOInit);
-			PixelShader->SetParameters(RHICmdList, View, LightSceneInfo);
+			PixelShader->SetParameters(RHICmdList, View, LightSceneInfo, ScreenShadowMaskTexture);
 		}
 	}
 }
@@ -865,7 +867,7 @@ void CalculateLightNearFarDepthFromBounds(const FViewInfo& View, const FSphere &
  * @param LightIndex The light's index into FScene::Lights
  * @return true if anything got rendered
  */
-void FDeferredShadingSceneRenderer::RenderLight(FRHICommandList& RHICmdList, const FLightSceneInfo* LightSceneInfo, bool bRenderOverlap, bool bIssueDrawEvent)
+void FDeferredShadingSceneRenderer::RenderLight(FRHICommandList& RHICmdList, const FLightSceneInfo* LightSceneInfo, IPooledRenderTarget* ScreenShadowMaskTexture, bool bRenderOverlap, bool bIssueDrawEvent)
 {
 	SCOPE_CYCLE_COUNTER(STAT_DirectLightRenderingTime);
 	INC_DWORD_STAT(STAT_NumLightsUsingStandardDeferred);
@@ -922,11 +924,11 @@ void FDeferredShadingSceneRenderer::RenderLight(FRHICommandList& RHICmdList, con
 			{
 				if(bUseIESTexture)
 				{
-					SetShaderTemplLighting<true, false, false>(RHICmdList, GraphicsPSOInit, View, *VertexShader, LightSceneInfo);
+					SetShaderTemplLighting<true, false, false>(RHICmdList, GraphicsPSOInit, View, *VertexShader, LightSceneInfo, ScreenShadowMaskTexture);
 				}
 				else
 				{
-					SetShaderTemplLighting<false, false, false>(RHICmdList, GraphicsPSOInit, View, *VertexShader, LightSceneInfo);
+					SetShaderTemplLighting<false, false, false>(RHICmdList, GraphicsPSOInit, View, *VertexShader, LightSceneInfo, ScreenShadowMaskTexture);
 				}
 			}
 
@@ -966,22 +968,22 @@ void FDeferredShadingSceneRenderer::RenderLight(FRHICommandList& RHICmdList, con
 				{
 					if(bUseIESTexture)
 					{
-						SetShaderTemplLighting<true, true, true>(RHICmdList, GraphicsPSOInit, View, *VertexShader, LightSceneInfo);
+						SetShaderTemplLighting<true, true, true>(RHICmdList, GraphicsPSOInit, View, *VertexShader, LightSceneInfo, ScreenShadowMaskTexture);
 					}
 					else
 					{
-						SetShaderTemplLighting<false, true, true>(RHICmdList, GraphicsPSOInit, View, *VertexShader, LightSceneInfo);
+						SetShaderTemplLighting<false, true, true>(RHICmdList, GraphicsPSOInit, View, *VertexShader, LightSceneInfo, ScreenShadowMaskTexture);
 					}
 				}
 				else
 				{
 					if(bUseIESTexture)
 					{
-						SetShaderTemplLighting<true, true, false>(RHICmdList, GraphicsPSOInit, View, *VertexShader, LightSceneInfo);
+						SetShaderTemplLighting<true, true, false>(RHICmdList, GraphicsPSOInit, View, *VertexShader, LightSceneInfo, ScreenShadowMaskTexture);
 					}
 					else
 					{
-						SetShaderTemplLighting<false, true, false>(RHICmdList, GraphicsPSOInit, View, *VertexShader, LightSceneInfo);
+						SetShaderTemplLighting<false, true, false>(RHICmdList, GraphicsPSOInit, View, *VertexShader, LightSceneInfo, ScreenShadowMaskTexture);
 					}
 				}
 			}
diff --git a/Engine/Source/Runtime/Renderer/Private/LightSceneInfo.h b/Engine/Source/Runtime/Renderer/Private/LightSceneInfo.h
--- a/Engine/Source/Runtime/Renderer/Private/LightSceneInfo.h
+++ b/Engine/Source/Runtime/Renderer/Private/LightSceneInfo.h
@@ -82,12 +82,12 @@ struct FSortedLightSceneInfo
 		/** Sort key bits packed into an integer. */
 		int32 Packed;
 	} SortKey;
-	/** The compact light scene info. */
-	FLightSceneInfoCompact SceneInfo;
+
+	const FLightSceneInfo* LightSceneInfo;
 
 	/** Initialization constructor. */
-	explicit FSortedLightSceneInfo(const FLightSceneInfoCompact& InSceneInfo)
-		: SceneInfo(InSceneInfo)
+	explicit FSortedLightSceneInfo(const FLightSceneInfo* InLightSceneInfo)
+		: LightSceneInfo(InLightSceneInfo)
 	{
 		SortKey.Packed = 0;
 	}
diff --git a/Engine/Source/Runtime/Renderer/Private/PostProcess/PostProcessBokehDOF.cpp b/Engine/Source/Runtime/Renderer/Private/PostProcess/PostProcessBokehDOF.cpp
--- a/Engine/Source/Runtime/Renderer/Private/PostProcess/PostProcessBokehDOF.cpp
+++ b/Engine/Source/Runtime/Renderer/Private/PostProcess/PostProcessBokehDOF.cpp
@@ -620,6 +620,8 @@ FPooledRenderTargetDesc FRCPassPostProcessBokehDOFSetup::ComputeOutputDesc(EPass
 	Ret.DebugName = TEXT("BokehDOFSetup");
 	Ret.TargetableFlags &= ~(TexCreate_RenderTargetable | TexCreate_UAV);
 	Ret.TargetableFlags |= bIsComputePass ? TexCreate_UAV : TexCreate_RenderTargetable;
+	Ret.Flags |= GetTextureFastVRamFlag_DynamicLayout();
+
 	return Ret;
 }
 
@@ -952,6 +954,7 @@ FPooledRenderTargetDesc FRCPassPostProcessBokehDOF::ComputeOutputDesc(EPassOutpu
 	uint32 FullRes = FSceneRenderTargets::Get_FrameConstantsOnly().GetBufferSizeXY().Y;
 	uint32 HalfRes = FMath::DivideAndRoundUp(FullRes, (uint32)2);
 
+	Ret.Flags |= GetTextureFastVRamFlag_DynamicLayout();
 	// we need space for the front part and the back part
 	Ret.Extent.Y = HalfRes * 2 + SafetyBorder;
 	Ret.DebugName = TEXT("BokehDOF");
diff --git a/Engine/Source/Runtime/Renderer/Private/PostProcess/PostProcessBokehDOFRecombine.cpp b/Engine/Source/Runtime/Renderer/Private/PostProcess/PostProcessBokehDOFRecombine.cpp
--- a/Engine/Source/Runtime/Renderer/Private/PostProcess/PostProcessBokehDOFRecombine.cpp
+++ b/Engine/Source/Runtime/Renderer/Private/PostProcess/PostProcessBokehDOFRecombine.cpp
@@ -506,6 +506,7 @@ FPooledRenderTargetDesc FRCPassPostProcessBokehDOFRecombine::ComputeOutputDesc(E
 
 	Ret.AutoWritable = false;
 	Ret.DebugName = TEXT("BokehDOFRecombine");
+	Ret.Flags |= GetTextureFastVRamFlag_DynamicLayout();
 
 	return Ret;
 }
diff --git a/Engine/Source/Runtime/Renderer/Private/PostProcess/PostProcessCircleDOF.cpp b/Engine/Source/Runtime/Renderer/Private/PostProcess/PostProcessCircleDOF.cpp
--- a/Engine/Source/Runtime/Renderer/Private/PostProcess/PostProcessCircleDOF.cpp
+++ b/Engine/Source/Runtime/Renderer/Private/PostProcess/PostProcessCircleDOF.cpp
@@ -280,6 +280,8 @@ FPooledRenderTargetDesc FRCPassPostProcessCircleDOFSetup::ComputeOutputDesc(EPas
 	Ret.TargetableFlags &= ~(uint32)TexCreate_UAV;
 	Ret.TargetableFlags |= TexCreate_RenderTargetable;
 	Ret.AutoWritable = false;
+	Ret.Flags |= GetTextureFastVRamFlag_DynamicLayout();
+
 	if (FPostProcessing::HasAlphaChannelSupport())
 	{
 		if (InPassOutputId == ePId_Output0)
@@ -494,6 +496,7 @@ FPooledRenderTargetDesc FRCPassPostProcessCircleDOFDilate::ComputeOutputDesc(EPa
 //	Ret.Format = PF_FloatRGBA;
 	// we only use one channel, maybe using 4 channels would save memory as we reuse
 	Ret.Format = PF_R16F;
+	Ret.Flags |= GetTextureFastVRamFlag_DynamicLayout();
 
 	return Ret;
 }
@@ -748,6 +751,7 @@ FPooledRenderTargetDesc FRCPassPostProcessCircleDOF::ComputeOutputDesc(EPassOutp
 	Ret.Reset();
 	Ret.TargetableFlags &= ~(uint32)TexCreate_UAV;
 	Ret.TargetableFlags |= TexCreate_RenderTargetable;
+	Ret.Flags |= GetTextureFastVRamFlag_DynamicLayout();
 
 	if (FPostProcessing::HasAlphaChannelSupport())
 	{
diff --git a/Engine/Source/Runtime/Renderer/Private/PostProcess/PostProcessCombineLUTs.cpp b/Engine/Source/Runtime/Renderer/Private/PostProcess/PostProcessCombineLUTs.cpp
--- a/Engine/Source/Runtime/Renderer/Private/PostProcess/PostProcessCombineLUTs.cpp
+++ b/Engine/Source/Runtime/Renderer/Private/PostProcess/PostProcessCombineLUTs.cpp
@@ -1032,7 +1032,7 @@ FPooledRenderTargetDesc FRCPassPostProcessCombineLUTs::ComputeOutputDesc(EPassOu
 			Ret.Extent = FIntPoint(GLUTSize, GLUTSize);
 			Ret.Depth = GLUTSize;
 		}
-		
+		Ret.Flags |= GetTextureFastVRamFlag_DynamicLayout();
 		Ret.DebugName = TEXT("CombineLUTs");
 	}
 
diff --git a/Engine/Source/Runtime/Renderer/Private/PostProcess/PostProcessDownsample.cpp b/Engine/Source/Runtime/Renderer/Private/PostProcess/PostProcessDownsample.cpp
--- a/Engine/Source/Runtime/Renderer/Private/PostProcess/PostProcessDownsample.cpp
+++ b/Engine/Source/Runtime/Renderer/Private/PostProcess/PostProcessDownsample.cpp
@@ -443,6 +443,7 @@ FPooledRenderTargetDesc FRCPassPostProcessDownsample::ComputeOutputDesc(EPassOut
 
 	Ret.TargetableFlags &= ~(TexCreate_RenderTargetable | TexCreate_UAV);
 	Ret.TargetableFlags |= bIsComputePass ? TexCreate_UAV : TexCreate_RenderTargetable;
+	Ret.Flags |= GetTextureFastVRamFlag_DynamicLayout();
 	Ret.AutoWritable = false;
 	Ret.DebugName = DebugName;
 
diff --git a/Engine/Source/Runtime/Renderer/Private/PostProcess/PostProcessEyeAdaptation.cpp b/Engine/Source/Runtime/Renderer/Private/PostProcess/PostProcessEyeAdaptation.cpp
--- a/Engine/Source/Runtime/Renderer/Private/PostProcess/PostProcessEyeAdaptation.cpp
+++ b/Engine/Source/Runtime/Renderer/Private/PostProcess/PostProcessEyeAdaptation.cpp
@@ -507,6 +507,7 @@ FPooledRenderTargetDesc FRCPassPostProcessBasicEyeAdaptationSetUp::ComputeOutput
 	Ret.DebugName = TEXT("EyeAdaptationBasicSetup");
 	// Require alpha channel for log2 information.
 	Ret.Format = PF_FloatRGBA;
+	Ret.Flags |= GetTextureFastVRamFlag_DynamicLayout();
 	return Ret;
 }
 
@@ -685,6 +686,6 @@ FPooledRenderTargetDesc FRCPassPostProcessBasicEyeAdaptation::ComputeOutputDesc(
 	FPooledRenderTargetDesc Ret;
 
 	Ret.DebugName = TEXT("EyeAdaptationBasic");
-
+	Ret.Flags |= GetTextureFastVRamFlag_DynamicLayout();
 	return Ret;
 }
diff --git a/Engine/Source/Runtime/Renderer/Private/PostProcess/PostProcessHistogram.cpp b/Engine/Source/Runtime/Renderer/Private/PostProcess/PostProcessHistogram.cpp
--- a/Engine/Source/Runtime/Renderer/Private/PostProcess/PostProcessHistogram.cpp
+++ b/Engine/Source/Runtime/Renderer/Private/PostProcess/PostProcessHistogram.cpp
@@ -161,7 +161,7 @@ FPooledRenderTargetDesc FRCPassPostProcessHistogram::ComputeOutputDesc(EPassOutp
 
 	// format can be optimized later
 	FPooledRenderTargetDesc Ret(FPooledRenderTargetDesc::Create2DDesc(NewSize, PF_FloatRGBA, FClearValueBinding::None, TexCreate_None, TexCreate_RenderTargetable | TexCreate_UAV, false));
-
+	Ret.Flags |= GetTextureFastVRamFlag_DynamicLayout();
 	Ret.DebugName = TEXT("Histogram");
 
 	return Ret;
diff --git a/Engine/Source/Runtime/Renderer/Private/PostProcess/PostProcessMotionBlur.cpp b/Engine/Source/Runtime/Renderer/Private/PostProcess/PostProcessMotionBlur.cpp
--- a/Engine/Source/Runtime/Renderer/Private/PostProcess/PostProcessMotionBlur.cpp
+++ b/Engine/Source/Runtime/Renderer/Private/PostProcess/PostProcessMotionBlur.cpp
@@ -191,6 +191,7 @@ FPooledRenderTargetDesc FRCPassPostProcessVelocityFlatten::ComputeOutputDesc(EPa
 		Ret.Format = PF_FloatR11G11B10;
 		Ret.TargetableFlags |= TexCreate_UAV;
 		Ret.TargetableFlags |= TexCreate_RenderTargetable;
+		Ret.Flags |= GetTextureFastVRamFlag_DynamicLayout();
 		Ret.DebugName = TEXT("VelocityFlat");
 
 		return Ret;
@@ -205,7 +206,7 @@ FPooledRenderTargetDesc FRCPassPostProcessVelocityFlatten::ComputeOutputDesc(EPa
 		FIntPoint TileCount = GetNumTiles16x16(PixelExtent);
 
 		FPooledRenderTargetDesc Ret(FPooledRenderTargetDesc::Create2DDesc(TileCount, PF_FloatRGBA, FClearValueBinding::None, TexCreate_None, TexCreate_RenderTargetable | TexCreate_UAV, false));
-
+		Ret.Flags |= GetTextureFastVRamFlag_DynamicLayout();
 		Ret.DebugName = TEXT("MaxVelocity");
 
 		return Ret;
@@ -493,6 +494,7 @@ FPooledRenderTargetDesc FRCPassPostProcessVelocityGather::ComputeOutputDesc(EPas
 
 	Ret.TargetableFlags |= TexCreate_UAV;
 	Ret.TargetableFlags |= TexCreate_RenderTargetable;
+	Ret.Flags |= GetTextureFastVRamFlag_DynamicLayout();
 	Ret.DebugName = TEXT("ScatteredMaxVelocity");
 
 	return Ret;
@@ -929,6 +931,7 @@ FPooledRenderTargetDesc FRCPassPostProcessMotionBlur::ComputeOutputDesc(EPassOut
 	{
 		Ret.Format = PF_FloatRGB;
 	}
+	Ret.Flags |= GetTextureFastVRamFlag_DynamicLayout();
 	Ret.DebugName = TEXT("MotionBlur");
 	Ret.AutoWritable = false;
 
diff --git a/Engine/Source/Runtime/Renderer/Private/PostProcess/PostProcessTemporalAA.cpp b/Engine/Source/Runtime/Renderer/Private/PostProcess/PostProcessTemporalAA.cpp
--- a/Engine/Source/Runtime/Renderer/Private/PostProcess/PostProcessTemporalAA.cpp
+++ b/Engine/Source/Runtime/Renderer/Private/PostProcess/PostProcessTemporalAA.cpp
@@ -1121,6 +1121,7 @@ void FRCPassPostProcessTemporalAA::DispatchCS(TRHICmdList& RHICmdList, FRenderin
 FPooledRenderTargetDesc FRCPassPostProcessTemporalAA::ComputeOutputDesc(EPassOutputId InPassOutputId) const
 {
 	FPooledRenderTargetDesc Ret = GetInput(ePId_Input0)->GetOutput()->RenderTargetDesc;
+	Ret.Flags &= ~(TexCreate_FastVRAM | TexCreate_Transient);
 	Ret.Reset();
 	//regardless of input type, PF_FloatRGBA is required to properly accumulate between frames for a good result.
 	Ret.Format = PF_FloatRGBA;
diff --git a/Engine/Source/Runtime/Renderer/Private/PostProcess/PostProcessUpscale.cpp b/Engine/Source/Runtime/Renderer/Private/PostProcess/PostProcessUpscale.cpp
--- a/Engine/Source/Runtime/Renderer/Private/PostProcess/PostProcessUpscale.cpp
+++ b/Engine/Source/Runtime/Renderer/Private/PostProcess/PostProcessUpscale.cpp
@@ -385,6 +385,7 @@ FPooledRenderTargetDesc FRCPassPostProcessUpscale::ComputeOutputDesc(EPassOutput
 	Ret.Reset();
 	Ret.DebugName = TEXT("Upscale");
 	Ret.Extent = OutputExtent;
+	Ret.Flags |= GetTextureFastVRamFlag_DynamicLayout();
 
 	return Ret;
 }
@@ -400,4 +401,4 @@ FPooledRenderTargetDesc FRCPassPostProcessUpscaleES2::ComputeOutputDesc(EPassOut
 	FPooledRenderTargetDesc Ret = FRCPassPostProcessUpscale::ComputeOutputDesc(InPassOutputId);
 	Ret.Extent = View.UnscaledViewRect.Max;
 	return Ret;
-}
+}
\ No newline at end of file
diff --git a/Engine/Source/Runtime/Renderer/Private/PostProcess/PostProcessWeightedSampleSum.cpp b/Engine/Source/Runtime/Renderer/Private/PostProcess/PostProcessWeightedSampleSum.cpp
--- a/Engine/Source/Runtime/Renderer/Private/PostProcess/PostProcessWeightedSampleSum.cpp
+++ b/Engine/Source/Runtime/Renderer/Private/PostProcess/PostProcessWeightedSampleSum.cpp
@@ -791,7 +791,7 @@ void FRCPassPostProcessWeightedSampleSum::Process(FRenderingCompositePassContext
 	}
 
 	// in texel (input resolution), /2 as we use the diameter, 100 as we use percent
-	float EffectiveBlurRadius = SizeScale * SrcSizeForThisAxis / 2 / 100.0f;
+	float EffectiveBlurRadius = SizeScale * SrcSizeForThisAxis  / 2 / 100.0f;
 
 	// compute 1D filtered samples
 	FVector2D BlurOffsets[MAX_FILTER_SAMPLES];
@@ -815,7 +815,7 @@ void FRCPassPostProcessWeightedSampleSum::Process(FRenderingCompositePassContext
 	{
 		float YOffset = bDoFastBlur ? (InvSrcSize.Y * 0.5f) : 0.0f;
 		for (uint32 i = 0; i < NumSamples; ++i)
-		{
+	{
 			BlurOffsets[i] = FVector2D(InvSrcSize.X * OffsetAndWeight[i].X, YOffset);
 		}
 	}
@@ -825,7 +825,7 @@ void FRCPassPostProcessWeightedSampleSum::Process(FRenderingCompositePassContext
 		for (uint32 i = 0; i < NumSamples; ++i)
 		{
 			BlurOffsets[i] = FVector2D(0, InvSrcSize.Y * OffsetAndWeight[i].X + YOffset);
-		}
+	}
 	}
 
 	uint32 CombineMethodInt = (CombineMethod == EFCM_MaxMagnitude) ? 2 : 0;
@@ -862,7 +862,7 @@ void FRCPassPostProcessWeightedSampleSum::Process(FRenderingCompositePassContext
 		{
 			// Async path
 			FRHIAsyncComputeCommandListImmediate& RHICmdListComputeImmediate = FRHICommandListExecutor::GetImmediateAsyncComputeCommandList();
-			{
+	{
  				SCOPED_COMPUTE_EVENT(RHICmdListComputeImmediate, AsyncWeightedSampleSum);
 				WaitForInputPassComputeFences(RHICmdListComputeImmediate);
 				
@@ -899,7 +899,7 @@ void FRCPassPostProcessWeightedSampleSum::Process(FRenderingCompositePassContext
 		else
 		{
 			SetRenderTarget(Context.RHICmdList, DestRenderTarget.TargetableTexture, FTextureRHIRef(), ESimpleRenderTargetMode::EExistingColorAndDepth);	
-		}
+	}
 
 		Context.SetViewportAndCallRHI(0, 0, 0.0f, DestSize.X, DestSize.Y, 1.0f);
 
@@ -910,23 +910,23 @@ void FRCPassPostProcessWeightedSampleSum::Process(FRenderingCompositePassContext
 		}
 
 
-		FShader* VertexShader = nullptr;
-		SetFilterShaders(
-			Context.RHICmdList,
-			FeatureLevel,
-			TStaticSamplerState<SF_Bilinear,AM_Border,AM_Border,AM_Clamp>::GetRHI(),
-			FilterTexture,
-			AdditiveTexture,
-			CombineMethodInt,
-			BlurOffsets,
-			BlurWeights,
-			NumSamples,
-			&VertexShader
-			);	
+	FShader* VertexShader = nullptr;
+	SetFilterShaders(
+		Context.RHICmdList,
+		FeatureLevel,
+		TStaticSamplerState<SF_Bilinear,AM_Border,AM_Border,AM_Clamp>::GetRHI(),
+		FilterTexture,
+		AdditiveTexture,
+		CombineMethodInt,
+		BlurOffsets,
+		BlurWeights,
+		NumSamples,
+		&VertexShader
+		);	
 
-		DrawQuad(Context.RHICmdList, FeatureLevel, bDoFastBlur, SrcRect, DestRect, DestSize, SrcSize, VertexShader);
+	DrawQuad(Context.RHICmdList, FeatureLevel, bDoFastBlur, SrcRect, DestRect, DestSize, SrcSize, VertexShader);
 
-		Context.RHICmdList.CopyToResolveTarget(DestRenderTarget.TargetableTexture, DestRenderTarget.ShaderResourceTexture, false, FResolveParams());
+	Context.RHICmdList.CopyToResolveTarget(DestRenderTarget.TargetableTexture, DestRenderTarget.ShaderResourceTexture, false, FResolveParams());
 	}
 }
 
@@ -1043,7 +1043,7 @@ FPooledRenderTargetDesc FRCPassPostProcessWeightedSampleSum::ComputeOutputDesc(E
 	Ret.AutoWritable = false;
 	Ret.TargetableFlags &= ~(TexCreate_RenderTargetable | TexCreate_UAV);
 	Ret.TargetableFlags |= bIsComputePass ? TexCreate_UAV : TexCreate_RenderTargetable;
-
+	Ret.Flags &= ~TexCreate_FastVRAM;
 	Ret.ClearValue = FClearValueBinding(FLinearColor(0, 0, 0, 0));
 
 	return Ret;
diff --git a/Engine/Source/Runtime/Renderer/Private/PostProcess/RenderTargetPool.cpp b/Engine/Source/Runtime/Renderer/Private/PostProcess/RenderTargetPool.cpp
--- a/Engine/Source/Runtime/Renderer/Private/PostProcess/RenderTargetPool.cpp
+++ b/Engine/Source/Runtime/Renderer/Private/PostProcess/RenderTargetPool.cpp
@@ -71,6 +71,15 @@ static TAutoConsoleVariable<int32> CVarRenderTargetPoolTest(
 	ECVF_Cheat | ECVF_RenderThreadSafe);
 #endif
 
+static TAutoConsoleVariable<int32> CVarAllowMultipleAliasingDiscardsPerFrame(
+	TEXT("r.RenderTargetPool.AllowMultipleAliasingDiscardsPerFrame"),
+	0,
+	TEXT("If enabled, allows rendertargets to be discarded and reacquired in the same frame.\n")
+	TEXT("This should give better aliasing efficiency, but carries some RHIthread/GPU performance overhead\n")
+	TEXT("with some RHIs (due to additional commandlist flushes)\n")
+	TEXT(" 0:off (default), 1:on"),
+	ECVF_Cheat | ECVF_RenderThreadSafe);
+
 bool FRenderTargetPool::IsEventRecordingEnabled() const
 {
 #if !(UE_BUILD_SHIPPING || UE_BUILD_TEST)
@@ -127,45 +136,6 @@ FRenderTargetPool::FRenderTargetPool()
 {
 }
 
-static void LogVRamUsage(FPooledRenderTarget& Ref, FTextureRHIParamRef TexRef)
-{
-	if(FPlatformProperties::SupportsFastVRAMMemory() && TexRef)
-	{
-		FRHIResourceInfo Info;
-
-		RHIGetResourceInfo(TexRef, Info);
-
-		if(Info.VRamAllocation.AllocationSize)
-		{
-			// note we do KB for more readable numbers but this can cause quantization loss
-			UE_LOG(LogShaders, Log, TEXT(" VRamInKB(Start/Size):%d/%d %s '%s'"), 
-				Info.VRamAllocation.AllocationStart / 1024, 
-				(Info.VRamAllocation.AllocationSize + 1023) / 1024,
-				*Ref.GetDesc().GenerateInfoString(),
-				Ref.GetDesc().DebugName);
-		}
-		else
-		{
-			UE_LOG(LogShaders, Log, TEXT(" VRamInKB request failed %s '%s'"), 
-				*Ref.GetDesc().GenerateInfoString(),
-				Ref.GetDesc().DebugName);
-		}
-	}
-}
-
-static void LogVRamUsage(FPooledRenderTarget& Ref)
-{
-	if(Ref.GetDesc().Flags & TexCreate_FastVRAM)
-	{
-		LogVRamUsage(Ref, Ref.GetRenderTargetItem().TargetableTexture);
-
-		if(Ref.GetRenderTargetItem().TargetableTexture != Ref.GetRenderTargetItem().ShaderResourceTexture)
-		{
-			LogVRamUsage(Ref, Ref.GetRenderTargetItem().ShaderResourceTexture);
-		}
-	}
-}
-
 void FRenderTargetPool::TransitionTargetsWritable(FRHICommandListImmediate& RHICmdList)
 {
 	QUICK_SCOPE_CYCLE_COUNTER(STAT_RenderTargetPoolTransition);
@@ -211,16 +181,47 @@ void FRenderTargetPool::WaitForTransitionFence()
 	DeferredDeleteArray.Reset();
 }
 
-bool FRenderTargetPool::FindFreeElement(FRHICommandList& RHICmdList, const FPooledRenderTargetDesc& Desc, TRefCountPtr<IPooledRenderTarget> &Out, const TCHAR* InDebugName, bool bDoWritableBarrier)
+bool FRenderTargetPool::FindFreeElement(FRHICommandList& RHICmdList, const FPooledRenderTargetDesc& InputDesc, TRefCountPtr<IPooledRenderTarget> &Out, const TCHAR* InDebugName, bool bDoWritableBarrier)
 {
 	check(IsInRenderingThread());
 
-	if(!Desc.IsValid())
+	if(!InputDesc.IsValid())
 	{
 		// no need to do anything
 		return true;
 	}
 
+	// If we're doing aliasing, we may need to override Transient flags, depending on the input format and mode
+	int32 AliasingMode = CVarTransientResourceAliasing_RenderTargets.GetValueOnRenderThread();
+	FPooledRenderTargetDesc ModifiedDesc;
+	bool bModifyDesc = false;
+	if (GSupportsTransientResourceAliasing && AliasingMode > 0 )
+	{
+		// Only override if the flags aren't already transient (this avoids a copy)
+		if ((InputDesc.Flags & TexCreate_Transient) == 0)
+		{
+			if (AliasingMode == 1)
+			{
+				// Mode 1: Make FastVRAM rendertargets transient
+				if (InputDesc.Flags & TexCreate_FastVRAM)
+				{
+					ModifiedDesc = InputDesc;
+					ModifiedDesc.Flags |= TexCreate_Transient;
+					bModifyDesc = true;
+				}
+			}
+			else if (AliasingMode == 2)
+			{
+				// Mode 2: Make all rendertargets transient
+				ModifiedDesc = InputDesc;
+				ModifiedDesc.Flags |= TexCreate_Transient;
+				bModifyDesc = true;
+			}
+		}
+	}
+
+	const FPooledRenderTargetDesc& Desc = bModifyDesc ? ModifiedDesc : InputDesc;
+
 	// Make sure if requesting a depth format that the clear value is correct
 	ensure(!IsDepthOrStencilFormat(Desc.Format) || (Desc.ClearValue.ColorBinding == EClearBinding::ENoneBound || Desc.ClearValue.ColorBinding == EClearBinding::EDepthStencilBound));
 
@@ -268,30 +269,37 @@ bool FRenderTargetPool::FindFreeElement(FRHICommandList& RHICmdList, const FPool
 	// try to find a suitable element in the pool
 	{
 		//don't spend time doing 2 passes if the platform doesn't support fastvram
-		uint32 PassCount = ((Desc.Flags & TexCreate_FastVRAM) && FPlatformProperties::SupportsFastVRAMMemory()) ? 2 : 1;			
+		uint32 PassCount = 1;
+		if (AliasingMode == 0)
+		{
+			if ((Desc.Flags & TexCreate_FastVRAM) && FPlatformMemory::SupportsFastVRAMMemory() )
+		    {
+			    PassCount = 2;
+		    }
+		}
 
+		bool bAllowMultipleDiscards = ( CVarAllowMultipleAliasingDiscardsPerFrame.GetValueOnRenderThread() != 0 );
 		// first we try exact, if that fails we try without TexCreate_FastVRAM
 		// (easily we can run out of VRam, if this search becomes a performance problem we can optimize or we should use less TexCreate_FastVRAM)
 		for(uint32 Pass = 0; Pass < PassCount; ++Pass)
 		{
 			bool bExactMatch = (Pass == 0);
-
+    
 			for(uint32 i = 0, Num = (uint32)PooledRenderTargets.Num(); i < Num; ++i)
 			{
 				FPooledRenderTarget* Element = PooledRenderTargets[i];
-
+    
 				if(Element && Element->IsFree() && Element->GetDesc().Compare(Desc, bExactMatch))
 				{
+					if ( ( Desc.Flags & TexCreate_Transient ) && bAllowMultipleDiscards == false && Element->HasBeenDiscardedThisFrame() )
+					{
+						// We can't re-use transient resources if they've already been discarded this frame
+						continue;
+					}
 					check(!Element->IsSnapshot());
 					Found = Element;
 					FoundIndex = i;
 					bReusingExistingTarget = true;
-					if (Element->RenderTargetItem.TargetableTexture)
-					{
-						FRHICommandListImmediate& RHICmdListImmediate = FRHICommandListExecutor::GetImmediateCommandList();
-						FTextureRHIParamRef Texture = Element->RenderTargetItem.TargetableTexture;
-						RHICmdListImmediate.SetResourceAliasability(EResourceAliasability::EUnaliasable, &Texture, 1);
-					}
 					goto Done;
 				}
 			}
@@ -304,14 +312,17 @@ Done:
 		UE_LOG(LogRenderTargetPool, Display, TEXT("%d MB, NewRT %s %s"), (AllocationLevelInKB + 1023) / 1024, *Desc.GenerateInfoString(), InDebugName);
 
 		// not found in the pool, create a new element
-		Found = new FPooledRenderTarget(Desc);
+		Found = new FPooledRenderTarget(Desc, this);
 
 		PooledRenderTargets.Add(Found);
 		
 		// TexCreate_UAV should be used on Desc.TargetableFlags
 		check(!(Desc.Flags & TexCreate_UAV));
+		// TexCreate_FastVRAM should be used on Desc.Flags
+		ensure(!(Desc.TargetableFlags & TexCreate_FastVRAM));
 
 		FRHIResourceCreateInfo CreateInfo(Desc.ClearValue);
+		CreateInfo.DebugName = InDebugName;
 
 		if(Desc.TargetableFlags & (TexCreate_RenderTargetable | TexCreate_DepthStencilTargetable | TexCreate_UAV))
 		{
@@ -480,10 +491,7 @@ Done:
 
 		FoundIndex = PooledRenderTargets.Num() - 1;
 
-		// done twice but it doesn't hurt an LogVRamUsage gets the new name this way
 		Found->Desc.DebugName = InDebugName;
-
-		LogVRamUsage(*Found);
 	}
 
 #if !(UE_BUILD_SHIPPING || UE_BUILD_TEST)
@@ -518,15 +526,28 @@ Done:
 
 	AddAllocEvent(FoundIndex, Found);
 
+	uint32 OriginalNumRefs = Found->GetRefCount();
+
 	// assign to the reference counted variable
 	Out = Found;
 
 	check(!Found->IsFree());
 
-	if (bReusingExistingTarget && bDoWritableBarrier)
+	// Only referenced by the pool, map the physical pages
+	if (OriginalNumRefs == 1 && Found->GetRenderTargetItem().TargetableTexture != nullptr)
+	{
+		RHIAcquireTransientResource(Found->GetRenderTargetItem().TargetableTexture);
+	}
+
+	if (bReusingExistingTarget)
+	{
+		if (bDoWritableBarrier)
 	{
 		RHICmdList.TransitionResource(EResourceTransitionAccess::EWritable, Found->GetRenderTargetItem().TargetableTexture);
 	}
+
+	}
+
 	return false;
 }
 
@@ -537,7 +558,7 @@ void FRenderTargetPool::CreateUntrackedElement(const FPooledRenderTargetDesc& De
 	Out = 0;
 
 	// not found in the pool, create a new element
-	FPooledRenderTarget* Found = new FPooledRenderTarget(Desc);
+	FPooledRenderTarget* Found = new FPooledRenderTarget(Desc, NULL);
 
 	Found->RenderTargetItem = Item;
 	check(!Found->IsSnapshot());
@@ -1329,7 +1350,7 @@ uint32 FPooledRenderTarget::AddRef() const
 	return 1;
 }
 
-uint32 FPooledRenderTarget::Release() const
+uint32 FPooledRenderTarget::Release()
 {
 	if (!bSnapshot)
 	{
@@ -1337,20 +1358,14 @@ uint32 FPooledRenderTarget::Release() const
 		uint32 Refs = uint32(--NumRefs);
 		if(Refs == 0)
 		{
-			// better we remove const from Release()
-			FSceneRenderTargetItem& NonConstItem = (FSceneRenderTargetItem&)RenderTargetItem;
-
-			NonConstItem.SafeRelease();
+			RenderTargetItem.SafeRelease();
 			delete this;
 		}
-		else if (IsFree() && GIsRHIInitialized)
+		else if (Refs == 1 && RenderTargetPool && IsTransient() )
 		{
-			FRHICommandListImmediate& RHICmdList = FRHICommandListExecutor::GetImmediateCommandList();
-			if (RenderTargetItem.TargetableTexture)
-			{
-				FTextureRHIParamRef Texture = RenderTargetItem.TargetableTexture;
-				RHICmdList.SetResourceAliasability(EResourceAliasability::EAliasable, &Texture, 1);
-			}
+			// Discard the resource
+			RHIDiscardTransientResource( GetRenderTargetItem().TargetableTexture );
+			FrameNumberLastDiscard = GFrameNumberRenderThread;
 		}
 		return Refs;
 	}
@@ -1398,7 +1413,6 @@ void FRenderTargetPool::DestructSnapshots()
 	PooledRenderTargetSnapshots.Reset();
 }
 
-
 // for debugging purpose
 FPooledRenderTarget* FRenderTargetPool::GetElementById(uint32 Id) const
 {
diff --git a/Engine/Source/Runtime/Renderer/Private/PostProcess/RenderTargetPool.h b/Engine/Source/Runtime/Renderer/Private/PostProcess/RenderTargetPool.h
--- a/Engine/Source/Runtime/Renderer/Private/PostProcess/RenderTargetPool.h
+++ b/Engine/Source/Runtime/Renderer/Private/PostProcess/RenderTargetPool.h
@@ -17,11 +17,13 @@ class FViewInfo;
 /** The reference to a pooled render target, use like this: TRefCountPtr<IPooledRenderTarget> */
 struct FPooledRenderTarget : public IPooledRenderTarget
 {
-	FPooledRenderTarget(const FPooledRenderTargetDesc& InDesc) 
+	FPooledRenderTarget(const FPooledRenderTargetDesc& InDesc, class FRenderTargetPool* InRenderTargetPool) 
 		: NumRefs(0)
 		, UnusedForNFrames(0)
 		, Desc(InDesc)
 		, bSnapshot(false)
+		, RenderTargetPool(InRenderTargetPool)
+		, FrameNumberLastDiscard(-1)
 	{
 	}
 	/* Constructor that makes a snapshot */
@@ -30,6 +32,8 @@ struct FPooledRenderTarget : public IPooledRenderTarget
 		, UnusedForNFrames(0)
 		, Desc(SnaphotSource.Desc)
 		, bSnapshot(true)
+		, RenderTargetPool(SnaphotSource.RenderTargetPool)
+		, FrameNumberLastDiscard(-1)
 	{
 		check(IsInRenderingThread());
 		RenderTargetItem = SnaphotSource.RenderTargetItem;
@@ -55,9 +59,17 @@ struct FPooledRenderTarget : public IPooledRenderTarget
 	// interface IPooledRenderTarget --------------
 
 	virtual uint32 AddRef() const override final;
-	virtual uint32 Release() const override final;
+	virtual uint32 Release() override final;
 	virtual uint32 GetRefCount() const override final;
 	virtual bool IsFree() const override final;
+	virtual uint32 HasBeenDiscardedThisFrame() const
+	{
+		return GFrameNumberRenderThread == FrameNumberLastDiscard;
+	}
+	bool IsTransient() const
+	{
+		return !!(Desc.Flags & TexCreate_Transient);
+	}
 	virtual void SetDebugName(const TCHAR *InName);
 	virtual const FPooledRenderTargetDesc& GetDesc() const;
 	virtual uint32 ComputeMemorySize() const;
@@ -76,6 +88,12 @@ private:
 	/** Snapshots are sortof fake pooled render targets, they don't own anything and can outlive the things that created them. These are for threaded rendering. */
 	bool bSnapshot;
 
+	/** Pointer back to the pool for render targets which are actually pooled, otherwise NULL. */
+	FRenderTargetPool* RenderTargetPool;
+
+	/** Keeps track of the last frame we unmapped physical memory for this resource. We can't map again in the same frame if we did that */
+	uint32 FrameNumberLastDiscard;
+
 	/** @return true:release this one, false otherwise */
 	bool OnFrameStart();
 
@@ -213,7 +231,7 @@ public:
 	/** Destruct all snapshots, this must be done after all outstanding async tasks are done. It is important because they hold ref counted texture pointers etc **/
 	void DestructSnapshots();
 
-
+	void OnRenderTargetUnreferenced(IPooledRenderTarget* RenderTarget);
 
 	/** Only to get statistics on usage and free elements. Normally only called in renderthread or if FlushRenderingCommands was called() */
 	void GetStats(uint32& OutWholeCount, uint32& OutWholePoolInKB, uint32& OutUsedInKB) const;
diff --git a/Engine/Source/Runtime/Renderer/Private/PostProcess/SceneRenderTargets.cpp b/Engine/Source/Runtime/Renderer/Private/PostProcess/SceneRenderTargets.cpp
--- a/Engine/Source/Runtime/Renderer/Private/PostProcess/SceneRenderTargets.cpp
+++ b/Engine/Source/Runtime/Renderer/Private/PostProcess/SceneRenderTargets.cpp
@@ -92,31 +92,25 @@ static TAutoConsoleVariable<int32> CVarGBufferFormat(
 	ECVF_RenderThreadSafe);
 
 // Fast VRam Settings. These should be tuned based on title requirements and resolution
-// Note: Currently, changing these settings at runtime does not cause buffers to be reallocated 
-static TAutoConsoleVariable<int32> CVarFastVRamGBufferCount(
-	TEXT("r.FastVRamGBufferCount"),
-	4,
-	TEXT("Number of fast VRAM GBuffers, in order A-E. May not be optimal on platforms with limited fast VRAM. Adjust based on resolution"),
-	ECVF_Scalability | ECVF_RenderThreadSafe);
+
+static TAutoConsoleVariable<int32> CVarFastVRamFillWithGBuffers(
+	TEXT("r.FastVRamStaticLayoutFillWithGBuffers"),
+	1,
+	TEXT("Whether to fill ESRAM with GBuffers, all pass temporaries will fail to allocate.  Games with mostly static lighting will benefit more from filling ESRAM with GBuffers instead of pass temporaries.\n")
+	TEXT("NOTE: Only applies if r.FastVRamDynamicAllocation is disabled."),
+	ECVF_RenderThreadSafe);
 
 static TAutoConsoleVariable<int32> CVarFastVRamSceneColor(
 	TEXT("r.FastVRamSceneColor"),
 	1,
 	TEXT("Whether to store scene color in fast VRAM"),
-	ECVF_Scalability | ECVF_RenderThreadSafe);
+	ECVF_RenderThreadSafe);
 
 static TAutoConsoleVariable<int32> CVarFastVRamSceneDepth(
 	TEXT("r.FastVRamSceneDepth"),
 	1,
 	TEXT("Whether to store scene depth in fast VRAM"),
-	ECVF_Scalability | ECVF_RenderThreadSafe);
-
-static TAutoConsoleVariable<int32> CVarFastVRamLightAttenuation(
-	TEXT("r.FastVRamLightAttenuation"),
-	1,
-	TEXT("Whether to store light attenuation in fast VRAM"),
-	ECVF_Scalability | ECVF_RenderThreadSafe);
-
+	ECVF_RenderThreadSafe);
 
 /** The global render targets used for scene rendering. */
 static TGlobalResource<FSceneRenderTargets> SceneRenderTargetsSingleton;
@@ -784,7 +778,6 @@ void FSceneRenderTargets::AllocMobileMultiViewSceneColor(FRHICommandList& RHICmd
 		const FIntPoint MultiViewBufferSize(BufferSize.X / 2, BufferSize.Y);
 
 		FPooledRenderTargetDesc Desc(FPooledRenderTargetDesc::Create2DDesc(MultiViewBufferSize, SceneColorBufferFormat, DefaultColorClear, TexCreate_None, TexCreate_RenderTargetable, false));
-		Desc.Flags |= TexCreate_FastVRAM;
 		Desc.NumSamples = GetNumSceneColorMSAASamples(CurrentFeatureLevel);
 		Desc.ArraySize = 2;
 		Desc.bIsArray = true;
@@ -865,10 +858,6 @@ void FSceneRenderTargets::AllocLightAttenuation(FRHICommandList& RHICmdList)
 	// Create a texture to store the resolved light attenuation values, and a render-targetable surface to hold the unresolved light attenuation values.
 	{
 		FPooledRenderTargetDesc Desc(FPooledRenderTargetDesc::Create2DDesc(BufferSize, PF_B8G8R8A8, FClearValueBinding::White, TexCreate_None, TexCreate_RenderTargetable, false));
-		if (CVarFastVRamLightAttenuation.GetValueOnRenderThread() >= 1)
-		{
-			Desc.Flags |= TexCreate_FastVRAM;
-		}
 		Desc.NumSamples = GetNumSceneColorMSAASamples(CurrentFeatureLevel);
 		GRenderTargetPool.FindFreeElement(RHICmdList, Desc, LightAttenuation, TEXT("LightAttenuation"));
 	}
@@ -914,11 +903,6 @@ void FSceneRenderTargets::GetGBufferADesc(FPooledRenderTargetDesc& Desc) const
 			NormalGBufferFormat = PF_FloatRGBA;
 		}
 
-		if (CVarFastVRamGBufferCount.GetValueOnRenderThread() >= 1)
-		{
-			Desc.Flags |= TexCreate_FastVRAM;
-		}
-
 		Desc = FPooledRenderTargetDesc::Create2DDesc(BufferSize, NormalGBufferFormat, FClearValueBinding::Transparent, TexCreate_None, TexCreate_RenderTargetable, false);
 	}
 }
@@ -957,10 +941,7 @@ void FSceneRenderTargets::AllocGBufferTargets(FRHICommandList& RHICmdList)
 			const EPixelFormat SpecularGBufferFormat = bHighPrecisionGBuffers ? PF_FloatRGBA : PF_B8G8R8A8;
 
 			FPooledRenderTargetDesc Desc(FPooledRenderTargetDesc::Create2DDesc(BufferSize, SpecularGBufferFormat, FClearValueBinding::Transparent, TexCreate_None, TexCreate_RenderTargetable, false));
-			if (CVarFastVRamGBufferCount.GetValueOnRenderThread() >= 2)
-			{
 				Desc.Flags |= TexCreate_FastVRAM;
-			}
 			GRenderTargetPool.FindFreeElement(RHICmdList, Desc, GBufferB, TEXT("GBufferB"));
 		}
 
@@ -968,9 +949,9 @@ void FSceneRenderTargets::AllocGBufferTargets(FRHICommandList& RHICmdList)
 		{
 			const EPixelFormat DiffuseGBufferFormat = bHighPrecisionGBuffers ? PF_FloatRGBA : PF_B8G8R8A8;
 			FPooledRenderTargetDesc Desc(FPooledRenderTargetDesc::Create2DDesc(BufferSize, DiffuseGBufferFormat, FClearValueBinding::Transparent, TexCreate_SRGB, TexCreate_RenderTargetable, false));
-			if (CVarFastVRamGBufferCount.GetValueOnRenderThread() >= 3)
+			if (CVarFastVRamFillWithGBuffers.GetValueOnRenderThread() >= 1)
 			{
-				Desc.Flags |= TexCreate_FastVRAM;
+				Desc.Flags |= GetTextureFastVRamFlag_StaticLayout();
 			}
 			GRenderTargetPool.FindFreeElement(RHICmdList, Desc, GBufferC, TEXT("GBufferC"));
 		}
@@ -978,9 +959,9 @@ void FSceneRenderTargets::AllocGBufferTargets(FRHICommandList& RHICmdList)
 		// Create the mask g-buffer (e.g. SSAO, subsurface scattering, wet surface mask, skylight mask, ...).
 		{
 			FPooledRenderTargetDesc Desc(FPooledRenderTargetDesc::Create2DDesc(BufferSize, PF_B8G8R8A8, FClearValueBinding::Transparent, TexCreate_None, TexCreate_RenderTargetable, false));
-			if (CVarFastVRamGBufferCount.GetValueOnRenderThread() >= 4)
+			if (CVarFastVRamFillWithGBuffers.GetValueOnRenderThread() >= 1)
 			{
-				Desc.Flags |= TexCreate_FastVRAM;
+				Desc.Flags |= GetTextureFastVRamFlag_StaticLayout();
 			}
 			GRenderTargetPool.FindFreeElement(RHICmdList, Desc, GBufferD, TEXT("GBufferD"));
 		}
@@ -988,10 +969,6 @@ void FSceneRenderTargets::AllocGBufferTargets(FRHICommandList& RHICmdList)
 		if (bAllowStaticLighting)
 		{
 			FPooledRenderTargetDesc Desc(FPooledRenderTargetDesc::Create2DDesc(BufferSize, PF_B8G8R8A8, FClearValueBinding::Transparent, TexCreate_None, TexCreate_RenderTargetable, false));
-			if (CVarFastVRamGBufferCount.GetValueOnRenderThread() >= 5)
-			{
-				Desc.Flags |= TexCreate_FastVRAM;
-			}
 			GRenderTargetPool.FindFreeElement(RHICmdList, Desc, GBufferE, TEXT("GBufferE"));
 		}
 
diff --git a/Engine/Source/Runtime/Renderer/Private/SceneOcclusion.cpp b/Engine/Source/Runtime/Renderer/Private/SceneOcclusion.cpp
--- a/Engine/Source/Runtime/Renderer/Private/SceneOcclusion.cpp
+++ b/Engine/Source/Runtime/Renderer/Private/SceneOcclusion.cpp
@@ -44,12 +44,6 @@ static FAutoConsoleVariableRef CVarShowRelevantPrecomputedVisibilityCells(
 	ECVF_RenderThreadSafe
 	);
 
-static TAutoConsoleVariable<int32> CVarFastVRamHzb(
-	TEXT("r.FastVramHzb"),
-	0,
-	TEXT("Whether to store HZB in fast VRAM"),
-	ECVF_Scalability | ECVF_RenderThreadSafe);
-
 #define NUM_CUBE_VERTICES 36
 
 /** Random table for occlusion **/
@@ -1088,10 +1082,6 @@ void BuildHZB( FRHICommandListImmediate& RHICmdList, FViewInfo& View )
 	View.HZBMipmap0Size = HZBSize;
 
 	FPooledRenderTargetDesc Desc(FPooledRenderTargetDesc::Create2DDesc(HZBSize, PF_R16F, FClearValueBinding::None, TexCreate_None, TexCreate_RenderTargetable | TexCreate_ShaderResource | TexCreate_NoFastClear, false, NumMips));
-	if (CVarFastVRamHzb.GetValueOnRenderThread() >= 1)
-	{
-		Desc.Flags |= TexCreate_FastVRAM;
-	}
 	GRenderTargetPool.FindFreeElement(RHICmdList, Desc, View.HZB, TEXT("HZB") );
 	
 	FSceneRenderTargetItem& HZBRenderTarget = View.HZB->GetRenderTargetItem();
diff --git a/Engine/Source/Runtime/Renderer/Private/ScenePrivate.h b/Engine/Source/Runtime/Renderer/Private/ScenePrivate.h
--- a/Engine/Source/Runtime/Renderer/Private/ScenePrivate.h
+++ b/Engine/Source/Runtime/Renderer/Private/ScenePrivate.h
@@ -729,6 +729,8 @@ public:
 
 	FForwardLightingViewResources ForwardLightingResources;
 
+	FForwardLightingCullingResources ForwardLightingCullingResources;
+
 	TRefCountPtr<IPooledRenderTarget> LightScatteringHistory;
 
 	/** Distance field AO tile intersection GPU resources.  Last frame's state is not used, but they must be sized exactly to the view so stored here. */
@@ -1051,6 +1053,7 @@ public:
 		TranslucencyTimer.Release();
 		SeparateTranslucencyTimer.Release();
 		ForwardLightingResources.Release();
+		ForwardLightingCullingResources.Release();
 		LightScatteringHistory.SafeRelease();
 	}
 
diff --git a/Engine/Source/Runtime/Renderer/Private/SceneRendering.cpp b/Engine/Source/Runtime/Renderer/Private/SceneRendering.cpp
--- a/Engine/Source/Runtime/Renderer/Private/SceneRendering.cpp
+++ b/Engine/Source/Runtime/Renderer/Private/SceneRendering.cpp
@@ -233,6 +233,20 @@ static TAutoConsoleVariable<int32> CVarWideCustomResolve(
 	ECVF_RenderThreadSafe | ECVF_Scalability
 	);
 
+TAutoConsoleVariable<int32> CVarTransientResourceAliasing_RenderTargets(
+	TEXT("r.TransientResourceAliasing.RenderTargets"),
+	1,
+	TEXT("0 : Disabled\n")
+	TEXT("1 : enable transient resource aliasing for fastVRam rendertargets\n")
+	TEXT("2 : enable transient resource aliasing for ALL rendertargets (experimental!)\n"),
+	ECVF_ReadOnly);
+
+TAutoConsoleVariable<int32> CVarTransientResourceAliasing_Buffers(
+	TEXT("r.TransientResourceAliasing.Buffers"),
+	1,
+	TEXT("If true, enable transient resource aliasing for buffers"),
+	ECVF_ReadOnly);
+
 static FParallelCommandListSet* GOutstandingParallelCommandListSet = nullptr;
 
 DECLARE_CYCLE_STAT(TEXT("DeferredShadingSceneRenderer MotionBlurStartFrame"), STAT_FDeferredShadingSceneRenderer_MotionBlurStartFrame, STATGROUP_SceneRendering);
diff --git a/Engine/Source/Runtime/Renderer/Private/SceneRendering.h b/Engine/Source/Runtime/Renderer/Private/SceneRendering.h
--- a/Engine/Source/Runtime/Renderer/Private/SceneRendering.h
+++ b/Engine/Source/Runtime/Renderer/Private/SceneRendering.h
@@ -685,10 +685,6 @@ public:
 	FDynamicReadBuffer ForwardLocalLightBuffer;
 	FRWBuffer NumCulledLightsGrid;
 	FRWBuffer CulledLightDataGrid;
-	FRWBuffer NextCulledLightLink;
-	FRWBuffer StartOffsetGrid;
-	FRWBuffer CulledLightLinks;
-	FRWBuffer NextCulledLightData;
 
 	void Release()
 	{
@@ -696,6 +692,19 @@ public:
 		ForwardLocalLightBuffer.Release();
 		NumCulledLightsGrid.Release();
 		CulledLightDataGrid.Release();
+	}
+};
+
+class FForwardLightingCullingResources
+{
+public:
+	FRWBuffer NextCulledLightLink;
+	FRWBuffer StartOffsetGrid;
+	FRWBuffer CulledLightLinks;
+	FRWBuffer NextCulledLightData;
+
+	void Release()
+	{
 		NextCulledLightLink.Release();
 		StartOffsetGrid.Release();
 		CulledLightLinks.Release();
@@ -1350,7 +1359,7 @@ protected:
 
 	void InitDynamicShadows(FRHICommandListImmediate& RHICmdList);
 
-	bool RenderShadowProjections(FRHICommandListImmediate& RHICmdList, const FLightSceneInfo* LightSceneInfo, bool bProjectingForForwardShading, bool bMobileModulatedProjections);
+	bool RenderShadowProjections(FRHICommandListImmediate& RHICmdList, const FLightSceneInfo* LightSceneInfo, IPooledRenderTarget* ScreenShadowMaskTexture, bool bProjectingForForwardShading, bool bMobileModulatedProjections);
 
 	/** Finds a matching cached preshadow, if one exists. */
 	TRefCountPtr<FProjectedShadowInfo> GetCachedPreshadow(
@@ -1593,3 +1602,23 @@ inline void SetBlack3DIfNull(FTextureRHIParamRef& Tex)
 		SetBlack2DIfNull(Tex);
 	}
 }
+
+extern TAutoConsoleVariable<int32> CVarTransientResourceAliasing_RenderTargets;
+extern TAutoConsoleVariable<int32> CVarTransientResourceAliasing_Buffers;
+
+FORCEINLINE bool IsTransientResourceBufferAliasingEnabled()
+{
+	return (GSupportsTransientResourceAliasing && CVarTransientResourceAliasing_Buffers.GetValueOnRenderThread() != 0);
+}
+
+// Helper functions for fast vram to handle dynamic/static allocation
+FORCEINLINE uint32 GetTextureFastVRamFlag_DynamicLayout()
+{
+	return (GSupportsTransientResourceAliasing && CVarTransientResourceAliasing_RenderTargets.GetValueOnRenderThread() > 0) ? ( TexCreate_FastVRAM | TexCreate_Transient ) : 0;
+}
+
+FORCEINLINE uint32 GetTextureFastVRamFlag_StaticLayout()
+{
+	return (GSupportsTransientResourceAliasing == false && CVarTransientResourceAliasing_RenderTargets.GetValueOnRenderThread() == 0) ? TexCreate_FastVRAM : 0;
+}
+
diff --git a/Engine/Source/Runtime/Renderer/Private/ShadowRendering.cpp b/Engine/Source/Runtime/Renderer/Private/ShadowRendering.cpp
--- a/Engine/Source/Runtime/Renderer/Private/ShadowRendering.cpp
+++ b/Engine/Source/Runtime/Renderer/Private/ShadowRendering.cpp
@@ -1372,7 +1372,7 @@ bool FDeferredShadingSceneRenderer::InjectReflectiveShadowMaps(FRHICommandListIm
 	
 extern int32 GCapsuleShadows;
 
-bool FSceneRenderer::RenderShadowProjections(FRHICommandListImmediate& RHICmdList, const FLightSceneInfo* LightSceneInfo, bool bProjectingForForwardShading, bool bMobileModulatedProjections)
+bool FSceneRenderer::RenderShadowProjections(FRHICommandListImmediate& RHICmdList, const FLightSceneInfo* LightSceneInfo, IPooledRenderTarget* ScreenShadowMaskTexture, bool bProjectingForForwardShading, bool bMobileModulatedProjections)
 {
 	FVisibleLightInfo& VisibleLightInfo = VisibleLightInfos[LightSceneInfo->Id];
 	FSceneRenderTargets& SceneContext = FSceneRenderTargets::Get(RHICmdList);
@@ -1383,8 +1383,8 @@ bool FSceneRenderer::RenderShadowProjections(FRHICommandListImmediate& RHICmdLis
 	}
 	else
 	{
-		// Normal deferred shadows render to light attenuation
-		SceneContext.BeginRenderingLightAttenuation(RHICmdList);
+		// Normal deferred shadows render to the shadow mask
+		SetRenderTarget(RHICmdList, ScreenShadowMaskTexture->GetRenderTargetItem().TargetableTexture, SceneContext.GetSceneDepthSurface(), ESimpleRenderTargetMode::EExistingColorAndDepth, FExclusiveDepthStencil::DepthRead_StencilWrite, true);
 	}
 
 	for (int32 ViewIndex = 0; ViewIndex < Views.Num(); ViewIndex++)
@@ -1406,7 +1406,7 @@ bool FSceneRenderer::RenderShadowProjections(FRHICommandListImmediate& RHICmdLis
 
 			if (ProjectedShadowInfo->bRayTracedDistanceField)
 			{
-				ProjectedShadowInfo->RenderRayTracedDistanceFieldProjection(RHICmdList, View, bProjectingForForwardShading);
+				ProjectedShadowInfo->RenderRayTracedDistanceFieldProjection(RHICmdList, View, ScreenShadowMaskTexture, bProjectingForForwardShading);
 			}
 			else if (ProjectedShadowInfo->bAllocated)
 			{
@@ -1421,11 +1421,6 @@ bool FSceneRenderer::RenderShadowProjections(FRHICommandListImmediate& RHICmdLis
 					{
 						ProjectedShadowInfo->RenderProjection(RHICmdList, ViewIndex, &View, bProjectingForForwardShading, bMobileModulatedProjections);
 					}
-
-					if (!bMobileModulatedProjections)
-					{
-						GRenderTargetPool.VisualizeTexture.SetCheckPoint(RHICmdList, SceneContext.GetLightAttenuation());
-					}
 				}
 			}
 		}
@@ -1437,7 +1432,7 @@ bool FSceneRenderer::RenderShadowProjections(FRHICommandListImmediate& RHICmdLis
 	return true;
 }
 	
-bool FDeferredShadingSceneRenderer::RenderShadowProjections(FRHICommandListImmediate& RHICmdList, const FLightSceneInfo* LightSceneInfo, bool& bInjectedTranslucentVolume)
+bool FDeferredShadingSceneRenderer::RenderShadowProjections(FRHICommandListImmediate& RHICmdList, const FLightSceneInfo* LightSceneInfo, IPooledRenderTarget* ScreenShadowMaskTexture, bool& bInjectedTranslucentVolume)
 {
 	SCOPE_CYCLE_COUNTER(STAT_ProjectedShadowDrawTime);
 	SCOPED_DRAW_EVENT(RHICmdList, ShadowProjectionOnOpaque);
@@ -1445,7 +1440,7 @@ bool FDeferredShadingSceneRenderer::RenderShadowProjections(FRHICommandListImmed
 
 	FVisibleLightInfo& VisibleLightInfo = VisibleLightInfos[LightSceneInfo->Id];
 
-	FSceneRenderer::RenderShadowProjections(RHICmdList, LightSceneInfo, false, false);
+	FSceneRenderer::RenderShadowProjections(RHICmdList, LightSceneInfo, ScreenShadowMaskTexture, false, false);
 
 	for (int32 ShadowIndex = 0; ShadowIndex < VisibleLightInfo.ShadowsToProject.Num(); ShadowIndex++)
 	{
@@ -1467,7 +1462,7 @@ bool FDeferredShadingSceneRenderer::RenderShadowProjections(FRHICommandListImmed
 		}
 	}
 
-	RenderCapsuleDirectShadows(*LightSceneInfo, RHICmdList, VisibleLightInfo.CapsuleShadowsToProject, false);
+	RenderCapsuleDirectShadows(RHICmdList, *LightSceneInfo, ScreenShadowMaskTexture, VisibleLightInfo.CapsuleShadowsToProject, false);
 
 	for (int32 ViewIndex = 0; ViewIndex < Views.Num(); ViewIndex++)
 	{
@@ -1505,7 +1500,7 @@ void FMobileSceneRenderer::RenderModulatedShadowProjections(FRHICommandListImmed
 		{
 			TArray<FProjectedShadowInfo*, SceneRenderingAllocator> Shadows;
 			SCOPE_CYCLE_COUNTER(STAT_ProjectedShadowDrawTime);
-			FSceneRenderer::RenderShadowProjections(RHICmdList, LightSceneInfo, false, true);
+			FSceneRenderer::RenderShadowProjections(RHICmdList, LightSceneInfo, NULL, false, true);
 		}
 	}
 }
diff --git a/Engine/Source/Runtime/Renderer/Private/ShadowRendering.h b/Engine/Source/Runtime/Renderer/Private/ShadowRendering.h
--- a/Engine/Source/Runtime/Renderer/Private/ShadowRendering.h
+++ b/Engine/Source/Runtime/Renderer/Private/ShadowRendering.h
@@ -881,7 +881,7 @@ public:
 	void BeginRenderRayTracedDistanceFieldProjection(FRHICommandListImmediate& RHICmdList, const FViewInfo& View);
 
 	/** Renders ray traced distance field shadows. */
-	void RenderRayTracedDistanceFieldProjection(FRHICommandListImmediate& RHICmdList, const class FViewInfo& View, bool bProjectingForForwardShading);
+	void RenderRayTracedDistanceFieldProjection(FRHICommandListImmediate& RHICmdList, const class FViewInfo& View, IPooledRenderTarget* ScreenShadowMaskTexture, bool bProjectingForForwardShading);
 
 	/** Render one pass point light shadow projections. */
 	void RenderOnePassPointLightProjection(FRHICommandListImmediate& RHICmdList, int32 ViewIndex, const FViewInfo& View, bool bProjectingForForwardShading) const;
diff --git a/Engine/Source/Runtime/Renderer/Private/ShadowSetup.cpp b/Engine/Source/Runtime/Renderer/Private/ShadowSetup.cpp
--- a/Engine/Source/Runtime/Renderer/Private/ShadowSetup.cpp
+++ b/Engine/Source/Runtime/Renderer/Private/ShadowSetup.cpp
@@ -923,7 +923,18 @@ void FProjectedShadowInfo::AddSubjectPrimitive(FPrimitiveSceneInfo* PrimitiveSce
 						else
 						{
 							FLODMask LODToRender;
-							int32 ForcedLODLevel = (CurrentView.Family->EngineShowFlags.LOD) ? GetCVarForceLOD() : 0;
+							int32 ForcedLODLevel = 0;
+
+							if (CurrentView.Family->EngineShowFlags.LOD)
+							{
+								// Shadow specific setting is highest priority
+								ForcedLODLevel = GetCVarForceLODShadow();
+
+								if (ForcedLODLevel == -1)
+								{
+									ForcedLODLevel = GetCVarForceLOD();
+								}
+							}
 
 							// Add the primitive's static mesh elements to the draw lists.
 							if (bForceLowestDetailLevel) 
diff --git a/Engine/Source/Runtime/Renderer/Private/SystemTextures.cpp b/Engine/Source/Runtime/Renderer/Private/SystemTextures.cpp
--- a/Engine/Source/Runtime/Renderer/Private/SystemTextures.cpp
+++ b/Engine/Source/Runtime/Renderer/Private/SystemTextures.cpp
@@ -359,7 +359,7 @@ void FSystemTextures::InternalInitializeTextures(FRHICommandListImmediate& RHICm
 				Format = PF_G16R16;
 			}
 
-			FPooledRenderTargetDesc Desc(FPooledRenderTargetDesc::Create2DDesc(FIntPoint(128, 32), Format, FClearValueBinding::None, TexCreate_FastVRAM, TexCreate_None, false));
+			FPooledRenderTargetDesc Desc(FPooledRenderTargetDesc::Create2DDesc(FIntPoint(128, 32), Format, FClearValueBinding::None, 0, TexCreate_None, false));
 			Desc.AutoWritable = false;
 			if (bReference)
 			{
diff --git a/Engine/Source/Runtime/Renderer/Private/TiledDeferredLightRendering.cpp b/Engine/Source/Runtime/Renderer/Private/TiledDeferredLightRendering.cpp
--- a/Engine/Source/Runtime/Renderer/Private/TiledDeferredLightRendering.cpp
+++ b/Engine/Source/Runtime/Renderer/Private/TiledDeferredLightRendering.cpp
@@ -160,8 +160,7 @@ public:
 			if (StartIndex + LightIndex < NumLightsToRenderInSortedLights)
 			{
 				const FSortedLightSceneInfo& SortedLightInfo = SortedLights[StartIndex + LightIndex];
-				const FLightSceneInfoCompact& LightSceneInfoCompact = SortedLightInfo.SceneInfo;
-				const FLightSceneInfo* const LightSceneInfo = LightSceneInfoCompact.LightSceneInfo;
+				const FLightSceneInfo* const LightSceneInfo = SortedLightInfo.LightSceneInfo;
 
 				FVector NormalizedLightDirection;
 				FVector2D SpotAngles;
diff --git a/Engine/Source/Runtime/Renderer/Private/TranslucentLighting.cpp b/Engine/Source/Runtime/Renderer/Private/TranslucentLighting.cpp
--- a/Engine/Source/Runtime/Renderer/Private/TranslucentLighting.cpp
+++ b/Engine/Source/Runtime/Renderer/Private/TranslucentLighting.cpp
@@ -1606,8 +1606,7 @@ void FDeferredShadingSceneRenderer::InjectTranslucentVolumeLightingArray(FRHICom
 	for (int32 LightIndex = 0; LightIndex < NumLights; LightIndex++)
 	{
 		const FSortedLightSceneInfo& SortedLightInfo = SortedLights[LightIndex];
-		const FLightSceneInfoCompact& LightSceneInfoCompact = SortedLightInfo.SceneInfo;
-		const FLightSceneInfo* const LightSceneInfo = LightSceneInfoCompact.LightSceneInfo;
+		const FLightSceneInfo* const LightSceneInfo = SortedLightInfo.LightSceneInfo;
 
 		AddLightForInjection(*this, *LightSceneInfo, NULL, LightInjectionData);
 	}
diff --git a/Engine/Source/Runtime/Renderer/Private/VolumetricFog.cpp b/Engine/Source/Runtime/Renderer/Private/VolumetricFog.cpp
--- a/Engine/Source/Runtime/Renderer/Private/VolumetricFog.cpp
+++ b/Engine/Source/Runtime/Renderer/Private/VolumetricFog.cpp
@@ -597,25 +597,25 @@ void FDeferredShadingSceneRenderer::RenderLocalLightsForVolumetricFog(
 				}
 				else
 				{
-					if (bDynamicallyShadowed)
+				if (bDynamicallyShadowed)
+				{
+					if (bInverseSquared)
 					{
-						if (bInverseSquared)
-						{
 							SetInjectShadowedLocalLightShaders<true, true, false>(RHICmdList, View, IntegrationData, LightSceneInfo, LightBounds, FogInfo, ProjectedShadowInfo, VolumetricFogGridSize, VolumeZBounds.X);
-						}
-						else
-						{
-							SetInjectShadowedLocalLightShaders<true, false, false>(RHICmdList, View, IntegrationData, LightSceneInfo, LightBounds, FogInfo, ProjectedShadowInfo, VolumetricFogGridSize, VolumeZBounds.X);
-						}
 					}
 					else
 					{
-						if (bInverseSquared)
-						{
+							SetInjectShadowedLocalLightShaders<true, false, false>(RHICmdList, View, IntegrationData, LightSceneInfo, LightBounds, FogInfo, ProjectedShadowInfo, VolumetricFogGridSize, VolumeZBounds.X);
+					}
+				}
+				else
+				{
+					if (bInverseSquared)
+					{
 							SetInjectShadowedLocalLightShaders<false, true, false>(RHICmdList, View, IntegrationData, LightSceneInfo, LightBounds, FogInfo, ProjectedShadowInfo, VolumetricFogGridSize, VolumeZBounds.X);
-						}
-						else
-						{
+					}
+					else
+					{
 							SetInjectShadowedLocalLightShaders<false, false, false>(RHICmdList, View, IntegrationData, LightSceneInfo, LightBounds, FogInfo, ProjectedShadowInfo, VolumetricFogGridSize, VolumeZBounds.X);
 						}
 					}
@@ -1077,8 +1077,10 @@ void FDeferredShadingSceneRenderer::ComputeVolumetricFog(FRHICommandListImmediat
 
 			const uint32 Flags = TexCreate_ShaderResource | TexCreate_RenderTargetable | TexCreate_UAV | TexCreate_ReduceMemoryWithTilingMode;
 			FPooledRenderTargetDesc VolumeDesc(FPooledRenderTargetDesc::CreateVolumeDesc(VolumetricFogGridSize.X, VolumetricFogGridSize.Y, VolumetricFogGridSize.Z, PF_FloatRGBA, FClearValueBinding::Black, TexCreate_None, Flags, false));
-			GRenderTargetPool.FindFreeElement(RHICmdList, VolumeDesc, VBufferA, TEXT("VBufferA"));
-			GRenderTargetPool.FindFreeElement(RHICmdList, VolumeDesc, VBufferB, TEXT("VBufferB"));
+			FPooledRenderTargetDesc VolumeDescFastVRAM = VolumeDesc;
+			VolumeDescFastVRAM.Flags |= GetTextureFastVRamFlag_DynamicLayout();
+			GRenderTargetPool.FindFreeElement(RHICmdList, VolumeDescFastVRAM, VBufferA, TEXT("VBufferA"));
+			GRenderTargetPool.FindFreeElement(RHICmdList, VolumeDescFastVRAM, VBufferB, TEXT("VBufferB"));
 
 			IntegrationData.VBufferARenderTarget = VBufferA.GetReference();
 			IntegrationData.VBufferBRenderTarget = VBufferB.GetReference();
@@ -1119,7 +1121,7 @@ void FDeferredShadingSceneRenderer::ComputeVolumetricFog(FRHICommandListImmediat
 			}
 
 			TRefCountPtr<IPooledRenderTarget> LocalShadowedLightScattering = NULL;
-			RenderLocalLightsForVolumetricFog(RHICmdList, View, bUseTemporalReprojection, IntegrationData, FogInfo, VolumetricFogGridSize, GridZParams, VolumeDesc, LocalShadowedLightScattering);
+			RenderLocalLightsForVolumetricFog(RHICmdList, View, bUseTemporalReprojection, IntegrationData, FogInfo, VolumetricFogGridSize, GridZParams, VolumeDescFastVRAM, LocalShadowedLightScattering);
 
 			TRefCountPtr<IPooledRenderTarget> LightScattering;
 			GRenderTargetPool.FindFreeElement(RHICmdList, VolumeDesc, LightScattering, TEXT("LightScattering"));
@@ -1208,6 +1210,7 @@ void FDeferredShadingSceneRenderer::ComputeVolumetricFog(FRHICommandListImmediat
 
 			VBufferA = NULL;
 			VBufferB = NULL;
+			LightFunctionTexture = NULL;
 
 			GRenderTargetPool.FindFreeElement(RHICmdList, VolumeDesc, View.VolumetricFogResources.IntegratedLightScattering, TEXT("IntegratedLightScattering"));
 
diff --git a/Engine/Source/Runtime/Renderer/Private/VolumetricFogLightFunction.cpp b/Engine/Source/Runtime/Renderer/Private/VolumetricFogLightFunction.cpp
--- a/Engine/Source/Runtime/Renderer/Private/VolumetricFogLightFunction.cpp
+++ b/Engine/Source/Runtime/Renderer/Private/VolumetricFogLightFunction.cpp
@@ -187,6 +187,7 @@ void FDeferredShadingSceneRenderer::RenderLightFunctionForVolumetricFog(
 		if (MaterialProxy && MaterialProxy->GetMaterial(Scene->GetFeatureLevel())->IsLightFunction())
 		{
 			FPooledRenderTargetDesc LightFunctionTextureDesc(FPooledRenderTargetDesc::Create2DDesc(LightFunctionResolution, PF_G8, FClearValueBinding::None, TexCreate_None, TexCreate_ShaderResource | TexCreate_RenderTargetable, false));
+			LightFunctionTextureDesc.Flags |= GetTextureFastVRamFlag_DynamicLayout();
 			GRenderTargetPool.FindFreeElement(RHICmdList, LightFunctionTextureDesc, OutLightFunctionTexture, TEXT("VolumetricFogLightFunction"));
 			
 			const FMatrix WorldToShadowValue = FTranslationMatrix(ProjectedShadowInfo.PreShadowTranslation) * ProjectedShadowInfo.SubjectAndReceiverMatrix;
diff --git a/Engine/Source/Runtime/Windows/D3D11RHI/Private/D3D11StructuredBuffer.cpp b/Engine/Source/Runtime/Windows/D3D11RHI/Private/D3D11StructuredBuffer.cpp
--- a/Engine/Source/Runtime/Windows/D3D11RHI/Private/D3D11StructuredBuffer.cpp
+++ b/Engine/Source/Runtime/Windows/D3D11RHI/Private/D3D11StructuredBuffer.cpp
@@ -53,7 +53,7 @@ FStructuredBufferRHIRef FD3D11DynamicRHI::RHICreateStructuredBuffer(uint32 Strid
 
 	Desc.StructureByteStride = Stride;
 
-	if (FPlatformProperties::SupportsFastVRAMMemory())
+	if (FPlatformMemory::SupportsFastVRAMMemory())
 	{
 		if (InUsage & BUF_FastVRAM)
 		{
diff --git a/Engine/Source/Runtime/Windows/D3D11RHI/Private/D3D11Texture.cpp b/Engine/Source/Runtime/Windows/D3D11RHI/Private/D3D11Texture.cpp
--- a/Engine/Source/Runtime/Windows/D3D11RHI/Private/D3D11Texture.cpp
+++ b/Engine/Source/Runtime/Windows/D3D11RHI/Private/D3D11Texture.cpp
@@ -622,7 +622,7 @@ TD3D11Texture2D<BaseResourceType>* FD3D11DynamicRHI::CreateD3D11Texture2D(uint32
 
 	FVRamAllocation VRamAllocation;
 
-	if (FPlatformProperties::SupportsFastVRAMMemory())
+	if (FPlatformMemory::SupportsFastVRAMMemory())
 	{
 		if (Flags & TexCreate_FastVRAM)
 		{
@@ -939,7 +939,7 @@ FD3D11Texture3D* FD3D11DynamicRHI::CreateD3D11Texture3D(uint32 SizeX,uint32 Size
 
 	FVRamAllocation VRamAllocation;
 
-	if (FPlatformProperties::SupportsFastVRAMMemory())
+	if (FPlatformMemory::SupportsFastVRAMMemory())
 	{
 		if (Flags & TexCreate_FastVRAM)
 		{
diff --git a/Engine/Source/Runtime/Windows/D3D11RHI/Private/D3D11VertexBuffer.cpp b/Engine/Source/Runtime/Windows/D3D11RHI/Private/D3D11VertexBuffer.cpp
--- a/Engine/Source/Runtime/Windows/D3D11RHI/Private/D3D11VertexBuffer.cpp
+++ b/Engine/Source/Runtime/Windows/D3D11RHI/Private/D3D11VertexBuffer.cpp
@@ -51,7 +51,7 @@ FVertexBufferRHIRef FD3D11DynamicRHI::RHICreateVertexBuffer(uint32 Size,uint32 I
 		Desc.BindFlags |= D3D11_BIND_SHADER_RESOURCE;
 	}
 
-	if (FPlatformProperties::SupportsFastVRAMMemory())
+	if (FPlatformMemory::SupportsFastVRAMMemory())
 	{
 		if (InUsage & BUF_FastVRAM)
 		{
diff --git a/Engine/Source/Runtime/Windows/XAudio2/Private/XAudio2Device.cpp b/Engine/Source/Runtime/Windows/XAudio2/Private/XAudio2Device.cpp
--- a/Engine/Source/Runtime/Windows/XAudio2/Private/XAudio2Device.cpp
+++ b/Engine/Source/Runtime/Windows/XAudio2/Private/XAudio2Device.cpp
@@ -31,6 +31,10 @@ THIRD_PARTY_INCLUDES_END
 #include "XAudio2Support.h"
 #include "Runtime/HeadMountedDisplay/Public/IHeadMountedDisplayModule.h"
 
+#if WITH_XMA2
+#include "XMAAudioInfo.h"
+#endif
+
 DEFINE_LOG_CATEGORY(LogXAudio2);
 
 class FXAudio2DeviceModule : public IAudioDeviceModule
@@ -90,8 +94,10 @@ bool FXAudio2Device::InitializeHardware()
 	DeviceProperties->XAudio2 = nullptr;
 	DeviceProperties->MasteringVoice = nullptr;
 
+#if WITH_OGGVORBIS
 	// Load ogg and vorbis dlls if they haven't been loaded yet
 	LoadVorbisLibraries();
+#endif
 
 	SampleRate = UE4_XAUDIO2_SAMPLERATE;
 
@@ -121,11 +127,16 @@ bool FXAudio2Device::InitializeHardware()
 	uint32 Flags = 0;
 #endif
 
+#if WITH_XMA2
+	// We don't use all of the SHAPE processor, so this flag prevents wasted resources
+	Flags |= XAUDIO2_DO_NOT_USE_SHAPE;
+#endif
+
 	// Create a new XAudio2 device object instance
 	if (XAudio2Create(&DeviceProperties->XAudio2, Flags, (XAUDIO2_PROCESSOR)FPlatformAffinity::GetAudioThreadMask()) != S_OK)
 	{
-		UE_LOG(LogInit, Log, TEXT( "Failed to create XAudio2 interface" ) );
-		return( false );
+		UE_LOG(LogInit, Log, TEXT("Failed to create XAudio2 interface"));
+		return(false);
 	}
 
 	check(DeviceProperties->XAudio2 != nullptr);
@@ -257,6 +268,10 @@ bool FXAudio2Device::InitializeHardware()
 	// Now initialize the audio clock voice after xaudio2 is initialized
 	DeviceProperties->InitAudioClockVoice();
 
+#if WITH_XMA2
+	FXMAAudioInfo::Initialize();
+#endif
+
 	return true;
 }
 
@@ -310,7 +325,7 @@ FSoundSource* FXAudio2Device::CreateSoundSource()
 
 bool FXAudio2Device::HasCompressedAudioInfoClass(USoundWave* SoundWave)
 {
-#if WITH_OGGVORBIS
+#if WITH_OGGVORBIS || WITH_XMA2
 	return true;
 #else
 	return false;
@@ -338,14 +353,23 @@ class ICompressedAudioInfo* FXAudio2Device::CreateCompressedAudioInfo(USoundWave
 		}
 		return CompressedInfo;
 	}
-	else
+#endif
+
+#if WITH_XMA2
+	static const FName NAME_XMA(TEXT("XMA"));
+	if (FPlatformProperties::RequiresCookedData() ? SoundWave->HasCompressedData(NAME_XMA) : (SoundWave->GetCompressedData(NAME_XMA) != nullptr))
 	{
-		return nullptr;
+		ICompressedAudioInfo* CompressedInfo = new FXMAAudioInfo();
+		if (!CompressedInfo)
+		{
+			UE_LOG(LogAudio, Error, TEXT("Failed to create new FXMAAudioInfo for SoundWave %s: out of memory."), *SoundWave->GetName());
+			return nullptr;
+		}
+		return CompressedInfo;
 	}
-#else
-	return nullptr;
 #endif
-	
+
+	return nullptr;
 }
 
 /**  
@@ -532,6 +556,7 @@ bool FXAudio2Device::GetOutputMatrix( uint32 ChannelMask, uint32 NumChannels )
 /** Test decompress a vorbis file */
 void FXAudio2Device::TestDecompressOggVorbis( USoundWave* Wave )
 {
+#if WITH_OGGVORBIS
 	FVorbisAudioInfo	OggInfo;
 	FSoundQualityInfo	QualityInfo = { 0 };
 
@@ -551,6 +576,7 @@ void FXAudio2Device::TestDecompressOggVorbis( USoundWave* Wave )
 
 		FMemory::Free( Wave->RawPCMData );
 	}
+#endif
 }
 
 #if !UE_BUILD_SHIPPING
diff --git a/Engine/Source/Runtime/Windows/XAudio2/Public/XAudio2Device.h b/Engine/Source/Runtime/Windows/XAudio2/Public/XAudio2Device.h
--- a/Engine/Source/Runtime/Windows/XAudio2/Public/XAudio2Device.h
+++ b/Engine/Source/Runtime/Windows/XAudio2/Public/XAudio2Device.h
@@ -6,6 +6,10 @@
 
 #pragma once
 
+#ifndef WITH_XMA2
+#define WITH_XMA2 0
+#endif
+
 /*------------------------------------------------------------------------------------
 	XAudio2 system headers
 ------------------------------------------------------------------------------------*/
@@ -83,6 +87,15 @@ class FSpatializationHelper;
  */
 class FXAudio2Device : public FAudioDevice
 {
+public:
+	FXAudio2Device() : FAudioDevice()
+	{
+#if WITH_XMA2 
+		bDisableAudioCaching = true;	// Do not allow DTYPE_Native buffers, only DTYPE_RealTime or DTYPE_Streaming since on the fly decompression is so cheap, it saves memory, and requires fewer code paths
+#endif
+	}
+
+private:
 	virtual void GetAudioDeviceList(TArray<FString>& OutAudioDeviceNames) const override;
 
 	/** Starts up any platform specific hardware/APIs */
@@ -113,7 +126,22 @@ class FXAudio2Device : public FAudioDevice
 		}
 #if WITH_OGGVORBIS
 		static FName NAME_OGG(TEXT("OGG"));
+		static FName NAME_XMA(TEXT("XMA"));
+
+#if WITH_XMA2
+		if (SoundWave->NumChannels > 2)
+		{
+			// Use OGG for surround wave sources, until we can sort out the channel assignments properly
+			return NAME_OGG;
+		}
+		else
+		{
+			return NAME_XMA;
+		}
+#else
 		return NAME_OGG;
+#endif
+
 #else //WITH_OGGVORBIS
 		static FName NAME_XMA(TEXT("XMA"));
 		return NAME_XMA;
