Copying //UE4/Dev-Console to //UE4/Dev-Main (Source: //UE4/Dev-Console @ 3505701)

=====================================
 MAJOR FEATURES + CHANGES
=====================================

Change 3484422 by [REDACTED_PERSON]

	[PS4]

	Playstation Texture lookup code now same as other platform as the PS4 shader compiler now supports switch statements.

Change 3486871 by [REDACTED_PERSON]

	[PS4]

	Added support to visualize the CMask in the vis command in the console.
	(This also adds support for all other UINT8 textures)

Change 3486877 by [REDACTED_PERSON]

	[PS4]

	DBuffer decals optimisation is now fixed on PSVR.

Change 3486935 by [REDACTED_PERSON]

	[General]

	Fixed a small compilation bug in the visualize shader for CL 3486871.

Change 3487400 by [REDACTED_PERSON]

	- Updated the "merge to main CL description maker" script to split the CLs up into multiple files depending on changes being made to NDA platforms directories

Change 3488955 by [REDACTED_PERSON]

	[General]

	The symbol cache will now be generated for the appropriate platforms.

	#3539

Change 3491593 by [REDACTED_PERSON]

	[D3D12] Add some additional pixel formats to ComputeBytesPerPixel. In particular,
		case DXGI_FORMAT_R32_TYPELESS:
		case DXGI_FORMAT_R32_FLOAT:

Change 3496287 by [REDACTED_PERSON]

	- Fixed minor issue where Multi text file would always be opened

Change 3498054 by [REDACTED_PERSON]

	Make Crash Malloc work if Binned2 is previous Malloc. Fixes callstack issues for non-gamethread crashes.

Change 3498188 by [REDACTED_PERSON]

	[PS4]

	Fixed an ensure which fires when the -norhithread flag is used.

Change 3498358 by [REDACTED_PERSON]

	Xbox One - Added XboxOnePDBFileUtil.exe to the installed engine filters.

Change 3500276 by [REDACTED_PERSON]

	[PS4]

	Mesh particle systems now bind the previous transform of the particles which removes the asset of "Resource at slot x, was not bound..." when LCUE debug is on.

Change 3500531 by [REDACTED_PERSON]

	Xbox One - Removed some bad assumptions about GPU memory handing over to the RHI in ContainerAllocationPolicies. Due to how PLATFORM_HAS_UMA is defined, this only affects DX11.x.

Change 3503080 by [REDACTED_PERSON]

	Fix for slate postprocess issue reported by [REDACTED_PERSON] (copy/dest were the wrong way around in a TransitionResources call)

Change 3504554 by [REDACTED_PERSON]

	Xbox One - Add a TitleID and Service Config ID to QAGame (along with a backend config) so that save/load works in TM-Gameplay

Change 3505443 by [REDACTED_PERSON]

	[CONSOLE] [!] Fix missing include in LinuxPlatformProcess.cpp

diff --git a/Engine/Build/InstalledEngineFilters.xml b/Engine/Build/InstalledEngineFilters.xml
--- a/Engine/Build/InstalledEngineFilters.xml
+++ b/Engine/Build/InstalledEngineFilters.xml
@@ -417,6 +417,7 @@
 
 		<!-- Additional Binaries -->
 		Engine/Binaries/Win64/XboxOne/XboxOnePackageNameUtil.exe
+    Engine/Binaries/Win64/XboxOne/XboxOnePDBFileUtil.exe
 	</Property>
 
 	<Property Name="CopyXboxOneExceptions">
diff --git a/Engine/Shaders/MaterialTemplate.usf b/Engine/Shaders/MaterialTemplate.usf
--- a/Engine/Shaders/MaterialTemplate.usf
+++ b/Engine/Shaders/MaterialTemplate.usf
@@ -1187,40 +1187,6 @@ float4 SceneTextureLookup(float2 UV, int SceneTextureIndex, bool bFiltered)
 #endif
 
 	FScreenSpaceData ScreenSpaceData = GetScreenSpaceData(UV, false);
-
-#if PS4_PROFILE
-	// TODO: PS4 doesn't optimize out correctly the switch(), so it thinks it needs all the Samplers even if they get compiled out
-	//	This will get fixed after launch per Sony...
-		 if (SceneTextureIndex == 0)	return float4(CalcSceneColor(UV), 0);
-	else if (SceneTextureIndex == 1)	return ScreenSpaceData.GBuffer.Depth;
-	else if (SceneTextureIndex == 2)	return float4(ScreenSpaceData.GBuffer.DiffuseColor, 0);
-	else if (SceneTextureIndex == 3)	return float4(ScreenSpaceData.GBuffer.SpecularColor, 0);
-	else if (SceneTextureIndex == 4)	return float4(ScreenSpaceData.GBuffer.CustomData.rgb, 0);
-	else if (SceneTextureIndex == 5)	return float4(ScreenSpaceData.GBuffer.BaseColor, 0);
-	else if (SceneTextureIndex == 6)	return ScreenSpaceData.GBuffer.Specular;
-	else if (SceneTextureIndex == 7)	return ScreenSpaceData.GBuffer.Metallic;
-	else if (SceneTextureIndex == 8)	return float4(ScreenSpaceData.GBuffer.WorldNormal, 0);
-	else if (SceneTextureIndex == 9)	return 1; // todo
-	else if (SceneTextureIndex == 10)	return ScreenSpaceData.GBuffer.CustomData.a;
-	else if (SceneTextureIndex == 11)	return ScreenSpaceData.GBuffer.Roughness;
-	else if (SceneTextureIndex == 12)	return ScreenSpaceData.GBuffer.GBufferAO;
-	else if (SceneTextureIndex == 13)	return ScreenSpaceData.GBuffer.CustomDepth;
-#if POST_PROCESS_MATERIAL
-	else if (SceneTextureIndex == 14)	if (bFiltered) { return Texture2DSample(PostprocessInput0, BilinearTextureSampler0, UV);} else { return Texture2DSample(PostprocessInput0, PostprocessInput0Sampler, UV);} 
-	else if (SceneTextureIndex == 15)	if (bFiltered) { return Texture2DSample(PostprocessInput1, BilinearTextureSampler0, UV);} else { return Texture2DSample(PostprocessInput1, PostprocessInput1Sampler, UV);} 
-	else if (SceneTextureIndex == 16)	if (bFiltered) { return Texture2DSample(PostprocessInput2, BilinearTextureSampler0, UV);} else { return Texture2DSample(PostprocessInput2, PostprocessInput2Sampler, UV);} 
-	else if (SceneTextureIndex == 17)	if (bFiltered) { return Texture2DSample(PostprocessInput3, BilinearTextureSampler0, UV);} else { return Texture2DSample(PostprocessInput3, PostprocessInput3Sampler, UV);} 
-	else if (SceneTextureIndex == 18)	if (bFiltered) { return Texture2DSample(PostprocessInput4, BilinearTextureSampler0, UV);} else { return Texture2DSample(PostprocessInput4, PostprocessInput4Sampler, UV);} 
-	else if (SceneTextureIndex == 19)	if (bFiltered) { return Texture2DSample(PostprocessInput5, BilinearTextureSampler0, UV);} else { return Texture2DSample(PostprocessInput5, PostprocessInput5Sampler, UV);} 
-	else if (SceneTextureIndex == 20)	if (bFiltered) { return Texture2DSample(PostprocessInput6, BilinearTextureSampler0, UV);} else { return Texture2DSample(PostprocessInput6, PostprocessInput6Sampler, UV);} 
-#endif // POST_PROCESS_MATERIAL
-	else if (SceneTextureIndex == 21)	return 0; // material compiler will return an error
-	else if (SceneTextureIndex == 22)	return float4(GetShadingModelColor(ScreenSpaceData.GBuffer.ShadingModelID), 1);
-	else if (SceneTextureIndex == 23)	return ScreenSpaceData.AmbientOcclusion;
-	else if (SceneTextureIndex == 24)	return ScreenSpaceData.GBuffer.CustomStencil;
-	else if (SceneTextureIndex == 25)	return float4(ScreenSpaceData.GBuffer.StoredBaseColor, 0);
-	else if (SceneTextureIndex == 26)	return float4(ScreenSpaceData.GBuffer.StoredSpecular.rrr, 0);
-#else
 	switch(SceneTextureIndex)
 	{
 		// order needs to match to ESceneTextureId
@@ -1284,12 +1250,6 @@ float4 SceneTextureLookup(float2 UV, int SceneTextureIndex, bool bFiltered)
 		default:
 			return float4(0, 0, 0, 0);
 	}
-#endif
-
-#if PS4_PROFILE
-	// PS4 as of SDK 930 can't figure out the switch statement exits through all code paths...
-	return float4(0, 0, 0, 0);
-#endif
 }
 #endif // NEEDS_SCENE_TEXTURES
 
diff --git a/Engine/Shaders/VisualizeTexture.usf b/Engine/Shaders/VisualizeTexture.usf
--- a/Engine/Shaders/VisualizeTexture.usf
+++ b/Engine/Shaders/VisualizeTexture.usf
@@ -20,6 +20,8 @@ SamplerState VisualizeTextureCubeSampler;
 //default Texture2D type is <float4>, which won't read the X28_G8UINT format properly.
 Texture2D<uint4> VisualizeDepthStencilTexture; 
 
+Texture2D<uint> VisualizeUINT8Texture2D;
+
 #if FEATURE_LEVEL >= FEATURE_LEVEL_SM5
 Texture2DMS<float4> VisualizeTexture2DMS;
 TextureCubeArray VisualizeTextureCubeArray;
@@ -90,6 +92,9 @@ void VisualizeTexturePS(
 	#if FEATURE_LEVEL >= FEATURE_LEVEL_SM5
 		TexLookup = VisualizeTexture2DMS.Load(int2(UV * TextureExtent.xy), 0);
 	#endif
+#elif TEXTURE_TYPE == 7
+    float IntTexLookup = float(VisualizeUINT8Texture2D.Load(int3(UV * TextureExtent.xy, 0)));
+	float4 TexLookup = IntTexLookup.xxxx;
 #endif
 	
 	// Normal = 0, SceneDepth = 1 and ShadowDepth = 2
diff --git a/Engine/Source/Programs/AutomationTool/Scripts/StreamCopyDescription.Automation.cs b/Engine/Source/Programs/AutomationTool/Scripts/StreamCopyDescription.Automation.cs
--- a/Engine/Source/Programs/AutomationTool/Scripts/StreamCopyDescription.Automation.cs
+++ b/Engine/Source/Programs/AutomationTool/Scripts/StreamCopyDescription.Automation.cs
@@ -141,11 +141,6 @@ namespace AutomationTool
 			{
 				string Desc = DescriptionBuilders[PlatformIndex].ToString().Replace("\r\n", "\n");
 
-				if (Desc == "")
-				{
-					continue;
-				}
-
 				// Clean any workspace names that may reveal internal information
 				Desc = Regex.Replace(Desc, "(Change[^@]*)@.*", "$1", RegexOptions.Multiline);
 
@@ -167,6 +162,11 @@ namespace AutomationTool
 				// Remove merge-only changelists
 				Desc = Regex.Replace(Desc, "(?<=(^|\\n))Change .*\\s*Merging .* to .*\\s*\\n(?=(Change|$))", "");
 
+				if (string.IsNullOrEmpty(Desc))
+				{
+					continue;
+				}
+
 				// Figure out the target stream
 				IProcessResult StreamResult = P4.P4(String.Format("stream -o {0}", Stream), AllowSpew: false);
 				if (StreamResult.ExitCode != 0)
@@ -208,6 +208,9 @@ namespace AutomationTool
 					{
 						Writer.WriteLine(Line);
 					}
+
+					Writer.WriteLine("DONE!");
+
 				}
 				Log("Written {0}.", OutputFileName);
 
diff --git a/Engine/Source/Programs/UnrealBuildTool/System/LinkEnvironment.cs b/Engine/Source/Programs/UnrealBuildTool/System/LinkEnvironment.cs
--- a/Engine/Source/Programs/UnrealBuildTool/System/LinkEnvironment.cs
+++ b/Engine/Source/Programs/UnrealBuildTool/System/LinkEnvironment.cs
@@ -124,7 +124,7 @@ namespace UnrealBuildTool
 		/// <summary>
 		/// True if debug symbols that are cached for some platforms should not be created.
 		/// </summary>
-		public bool bDisableSymbolCache = true;
+		public bool bDisableSymbolCache = false;
 
 		/// <summary>
 		/// True if we're compiling .cpp files that will go into a library (.lib file)
diff --git a/Engine/Source/Runtime/Core/Private/GenericPlatform/GenericPlatformMallocCrash.cpp b/Engine/Source/Runtime/Core/Private/GenericPlatform/GenericPlatformMallocCrash.cpp
--- a/Engine/Source/Runtime/Core/Private/GenericPlatform/GenericPlatformMallocCrash.cpp
+++ b/Engine/Source/Runtime/Core/Private/GenericPlatform/GenericPlatformMallocCrash.cpp
@@ -297,6 +297,7 @@ void* FGenericPlatformMallocCrash::Realloc( void* Ptr, SIZE_T NewSize, uint32 Al
 			{
 				// We can safely get allocation size only from a few mallocs, this may change in future.
 				if( FCStringWide::Strcmp( PreviousMalloc->GetDescriptiveName(), TEXT("binned") ) == 0 ||
+					FCStringWide::Strcmp( PreviousMalloc->GetDescriptiveName(), TEXT("binned2")) == 0 ||
 					FCStringWide::Strcmp( PreviousMalloc->GetDescriptiveName(), TEXT("jemalloc") ) == 0 )
 				{
 					// Realloc from the previous allocator.
diff --git a/Engine/Source/Runtime/Core/Private/Linux/LinuxPlatformProcess.cpp b/Engine/Source/Runtime/Core/Private/Linux/LinuxPlatformProcess.cpp
--- a/Engine/Source/Runtime/Core/Private/Linux/LinuxPlatformProcess.cpp
+++ b/Engine/Source/Runtime/Core/Private/Linux/LinuxPlatformProcess.cpp
@@ -24,6 +24,7 @@
 #include "Linux/LinuxPlatformTLS.h"
 #include "Containers/CircularQueue.h"
 #include "Misc/FileHelper.h"
+#include "Misc/OutputDeviceRedirector.h"
 
 namespace PlatformProcessLimits
 {
diff --git a/Engine/Source/Runtime/Core/Public/Containers/ContainerAllocationPolicies.h b/Engine/Source/Runtime/Core/Public/Containers/ContainerAllocationPolicies.h
--- a/Engine/Source/Runtime/Core/Public/Containers/ContainerAllocationPolicies.h
+++ b/Engine/Source/Runtime/Core/Public/Containers/ContainerAllocationPolicies.h
@@ -228,11 +228,7 @@ public:
 
 			if (Data)
 			{
-#if PLATFORM_HAS_UMA
-				FMemory::GPUFree(Data);
-#else
 				FMemory::Free(Data);
-#endif
 			}
 
 			Data       = Other.Data;
@@ -244,11 +240,7 @@ public:
 		{
 			if(Data)
 			{
-#if PLATFORM_HAS_UMA
-				// The RHI should have taken ownership of this memory, so we don't free it here
-#else
 				FMemory::Free(Data);
-#endif
 			}
 		}
 
@@ -266,12 +258,8 @@ public:
 			// Avoid calling FMemory::Realloc( nullptr, 0 ) as ANSI C mandates returning a valid pointer which is not what we want.
 			if (Data || NumElements)
 			{
-#if PLATFORM_HAS_UMA
-				Data = (FScriptContainerElement*)FMemory::GPURealloc(Data, NumElements*NumBytesPerElement, Alignment);
-#else
 				//checkSlow(((uint64)NumElements*(uint64)ElementTypeInfo.GetSize() < (uint64)INT_MAX));
 				Data = (FScriptContainerElement*)FMemory::Realloc( Data, NumElements*NumBytesPerElement, Alignment );
-#endif
 			}
 		}
 		FORCEINLINE int32 CalculateSlackReserve(int32 NumElements, int32 NumBytesPerElement) const
diff --git a/Engine/Source/Runtime/D3D12RHI/Private/D3D12RenderTarget.cpp b/Engine/Source/Runtime/D3D12RHI/Private/D3D12RenderTarget.cpp
--- a/Engine/Source/Runtime/D3D12RHI/Private/D3D12RenderTarget.cpp
+++ b/Engine/Source/Runtime/D3D12RHI/Private/D3D12RenderTarget.cpp
@@ -606,6 +606,8 @@ static uint32 ComputeBytesPerPixel(DXGI_FORMAT Format)
 	case DXGI_FORMAT_R11G11B10_FLOAT:
 	case DXGI_FORMAT_R16G16_UNORM:
 	case DXGI_FORMAT_R32_UINT:
+	case DXGI_FORMAT_R32_TYPELESS:
+	case DXGI_FORMAT_R32_FLOAT:
 		BytesPerPixel = 4;
 		break;
 	case DXGI_FORMAT_R16G16B16A16_FLOAT:
diff --git a/Engine/Source/Runtime/Engine/Private/Particles/FXSystem.cpp b/Engine/Source/Runtime/Engine/Private/Particles/FXSystem.cpp
--- a/Engine/Source/Runtime/Engine/Private/Particles/FXSystem.cpp
+++ b/Engine/Source/Runtime/Engine/Private/Particles/FXSystem.cpp
@@ -372,11 +372,11 @@ void FFXSystem::PostRenderOpaque(FRHICommandListImmediate& RHICmdList, const FUn
 {
 	if (RHISupportsGPUParticles() && IsParticleCollisionModeSupported(GetShaderPlatform(), PCM_DepthBuffer))
 	{
-		PrepareGPUSimulation(RHICmdList);
-
+		PrepareGPUSimulation(RHICmdList, SceneDepthTexture);
+		
 		SimulateGPUParticles(RHICmdList, EParticleSimulatePhase::CollisionDepthBuffer, ViewUniformBuffer, NULL, SceneDepthTexture, GBufferATexture);
-
-		FinalizeGPUSimulation(RHICmdList);
+		
+		FinalizeGPUSimulation(RHICmdList, SceneDepthTexture);
 
 		SortGPUParticles(RHICmdList);		
 	}
diff --git a/Engine/Source/Runtime/Engine/Private/Particles/FXSystemPrivate.h b/Engine/Source/Runtime/Engine/Private/Particles/FXSystemPrivate.h
--- a/Engine/Source/Runtime/Engine/Private/Particles/FXSystemPrivate.h
+++ b/Engine/Source/Runtime/Engine/Private/Particles/FXSystemPrivate.h
@@ -188,8 +188,8 @@ public:
 	 */
 	int32 AddSortedGPUSimulation(FParticleSimulationGPU* Simulation, const FVector& ViewOrigin);
 
-	void PrepareGPUSimulation(FRHICommandListImmediate& RHICmdList);
-	void FinalizeGPUSimulation(FRHICommandListImmediate& RHICmdList);
+	void PrepareGPUSimulation(FRHICommandListImmediate& RHICmdList, FTexture2DRHIParamRef SceneDepthTexture = nullptr);
+	void FinalizeGPUSimulation(FRHICommandListImmediate& RHICmdList, FTexture2DRHIParamRef SceneDepthTexture = nullptr);
 
 private:
 
diff --git a/Engine/Source/Runtime/Engine/Private/Particles/ParticleGpuSimulation.cpp b/Engine/Source/Runtime/Engine/Private/Particles/ParticleGpuSimulation.cpp
--- a/Engine/Source/Runtime/Engine/Private/Particles/ParticleGpuSimulation.cpp
+++ b/Engine/Source/Runtime/Engine/Private/Particles/ParticleGpuSimulation.cpp
@@ -4452,7 +4452,7 @@ bool FFXSystem::UsesGlobalDistanceFieldInternal() const
 	return false;
 }
 
-void FFXSystem::PrepareGPUSimulation(FRHICommandListImmediate& RHICmdList)
+void FFXSystem::PrepareGPUSimulation(FRHICommandListImmediate& RHICmdList, FTexture2DRHIParamRef SceneDepthTexture)
 {
 	// Grab resources.
 	FParticleStateTextures& CurrentStateTextures = ParticleSimulationResources->GetCurrentStateTextures();
@@ -4465,9 +4465,13 @@ void FFXSystem::PrepareGPUSimulation(FRHICommandListImmediate& RHICmdList)
 	};
 
 	RHICmdList.TransitionResources(EResourceTransitionAccess::EWritable, RenderTargets, 2);
+	if (SceneDepthTexture)
+	{
+		RHICmdList.TransitionResource(EResourceTransitionAccess::EReadable, SceneDepthTexture);
+	}
 }
 
-void FFXSystem::FinalizeGPUSimulation(FRHICommandListImmediate& RHICmdList)
+void FFXSystem::FinalizeGPUSimulation(FRHICommandListImmediate& RHICmdList, FTexture2DRHIParamRef SceneDepthTexture)
 {
 	// Grab resources.
 	FParticleStateTextures& CurrentStateTextures = ParticleSimulationResources->GetVisualizeStateTextures();
@@ -4479,7 +4483,11 @@ void FFXSystem::FinalizeGPUSimulation(FRHICommandListImmediate& RHICmdList)
 		CurrentStateTextures.VelocityTextureTargetRHI,
 	};
 	
-	RHICmdList.TransitionResources(EResourceTransitionAccess::EReadable, RenderTargets, 2);	
+	RHICmdList.TransitionResources(EResourceTransitionAccess::EReadable, RenderTargets, 2);
+	if (SceneDepthTexture)
+	{
+		RHICmdList.TransitionResource(EResourceTransitionAccess::EWritable, SceneDepthTexture);
+	}
 }
 
 void FFXSystem::SimulateGPUParticles(
diff --git a/Engine/Source/Runtime/Engine/Private/Particles/ParticleSystemRender.cpp b/Engine/Source/Runtime/Engine/Private/Particles/ParticleSystemRender.cpp
--- a/Engine/Source/Runtime/Engine/Private/Particles/ParticleSystemRender.cpp
+++ b/Engine/Source/Runtime/Engine/Private/Particles/ParticleSystemRender.cpp
@@ -1591,7 +1591,8 @@ void FDynamicMeshEmitterData::GetDynamicMeshElementsEmitter(const FParticleSyste
 
 				// For OpenGL & Metal we can't assume that it is OK to leave the PrevTransformBuffer buffer unbound.
 				// Doing so can lead to undefined behaviour if the buffer is referenced in the shader even if protected by a branch that is not meant to be taken.
-				bool const bGeneratePrevTransformBuffer = (FeatureLevel >= ERHIFeatureLevel::SM4) && (Source.MeshMotionBlurOffset || IsOpenGLPlatform(ShaderPlatform) || IsMetalPlatform(ShaderPlatform));
+				bool const bGeneratePrevTransformBuffer = (FeatureLevel >= ERHIFeatureLevel::SM4) && 
+                                                          (Source.MeshMotionBlurOffset || IsOpenGLPlatform(ShaderPlatform) || IsMetalPlatform(ShaderPlatform) || IsPS4Platform(ShaderPlatform));
 
 
 				if (bInstanced)
diff --git a/Engine/Source/Runtime/RHI/Public/RHIDefinitions.h b/Engine/Source/Runtime/RHI/Public/RHIDefinitions.h
--- a/Engine/Source/Runtime/RHI/Public/RHIDefinitions.h
+++ b/Engine/Source/Runtime/RHI/Public/RHIDefinitions.h
@@ -766,6 +766,11 @@ inline bool IsSwitchPlatform(const EShaderPlatform Platform)
 	return Platform == SP_SWITCH || Platform == SP_SWITCH_FORWARD;
 }
 
+inline bool IsPS4Platform(const EShaderPlatform Platform)
+{
+	return Platform == SP_PS4;
+}
+
 inline bool IsVulkanPlatform(const EShaderPlatform Platform)
 {
 	return Platform == SP_VULKAN_SM5 || Platform == SP_VULKAN_SM4 || Platform == SP_VULKAN_PCES3_1 || Platform == SP_VULKAN_ES3_1_ANDROID;
diff --git a/Engine/Source/Runtime/Renderer/Private/CompositionLighting/PostProcessDeferredDecals.cpp b/Engine/Source/Runtime/Renderer/Private/CompositionLighting/PostProcessDeferredDecals.cpp
--- a/Engine/Source/Runtime/Renderer/Private/CompositionLighting/PostProcessDeferredDecals.cpp
+++ b/Engine/Source/Runtime/Renderer/Private/CompositionLighting/PostProcessDeferredDecals.cpp
@@ -860,19 +860,26 @@ void FRCPassPostProcessDeferredDecals::Process(FRenderingCompositePassContext& C
 				DrawClearQuad(RHICmdList, false, FLinearColor(), false, 0, true, 0, SceneContext.GetSceneDepthSurface()->GetSizeXY(), FIntRect());
 			}
 
+			// This stops the targets from being resolved and decoded until the last view is rendered.
+			// This is done so as to not run eliminate fast clear on the views before the end.
+			bool bLastView = Context.View.Family->Views.Last() == &Context.View;
 			if (CurrentStage == DRS_BeforeBasePass)
 			{
 				// combine DBuffer RTWriteMasks; will end up in one texture we can load from in the base pass PS and decide whether to do the actual work or not
 				RenderTargetManager.FlushMetaData();
 
-				if (GSupportsRenderTargetWriteMask)
+				if (GSupportsRenderTargetWriteMask && bLastView)
 				{
 					DecodeRTWriteMask(Context);
+					GRenderTargetPool.VisualizeTexture.SetCheckPoint(RHICmdList, SceneContext.DBufferMask);
 					bHasValidDBufferMask = true;
 				}
 			}
 
-			RenderTargetManager.ResolveTargets();
+			if (bLastView || !GSupportsRenderTargetWriteMask)
+			{
+				RenderTargetManager.ResolveTargets();
+			}
 		}
 
 		if (CurrentStage == DRS_BeforeBasePass && bNeedsDBufferTargets)
diff --git a/Engine/Source/Runtime/Renderer/Private/PostProcess/VisualizeTexture.cpp b/Engine/Source/Runtime/Renderer/Private/PostProcess/VisualizeTexture.cpp
--- a/Engine/Source/Runtime/Renderer/Private/PostProcess/VisualizeTexture.cpp
+++ b/Engine/Source/Runtime/Renderer/Private/PostProcess/VisualizeTexture.cpp
@@ -26,9 +26,21 @@
 #include "PipelineStateCache.h"
 
 
+enum class EVisualisePSType
+{
+	Cube = 0,
+	Texture1D = 1, //not supported
+	Texture2DNoMSAA = 2,
+	Texture3D = 3,
+	CubeArray = 4,
+	Texture2DMSAA = 5,
+	Texture2DDepthStencilNoMSAA = 6,
+	Texture2DUINT8 = 7,
+};
+
 /** A pixel shader which filters a texture. */
 // @param TextureType 0:Cube, 1:1D(not yet supported), 2:2D no MSAA, 3:3D, 4:Cube[], 5:2D MSAA, 6:2D DepthStencil no MSAA (needed to avoid D3DDebug error)
-template<uint32 TextureType>
+template<EVisualisePSType TextureType>
 class VisualizeTexturePS : public FGlobalShader
 {
 	DECLARE_SHADER_TYPE(VisualizeTexturePS,Global);
@@ -42,7 +54,7 @@ public:
 	static void ModifyCompilationEnvironment(EShaderPlatform Platform, FShaderCompilerEnvironment& OutEnvironment)
 	{
 		FGlobalShader::ModifyCompilationEnvironment(Platform, OutEnvironment);
-		OutEnvironment.SetDefine(TEXT("TEXTURE_TYPE"), TextureType);
+		OutEnvironment.SetDefine(TEXT("TEXTURE_TYPE"), (int32)TextureType);
 	}
 
 	/** Default constructor. */
@@ -65,6 +77,7 @@ public:
 		VisualizeTextureCubeArraySampler.Bind(Initializer.ParameterMap,TEXT("VisualizeTextureCubeArraySampler"));
 		VisualizeParam.Bind(Initializer.ParameterMap,TEXT("VisualizeParam"));
 		TextureExtent.Bind(Initializer.ParameterMap,TEXT("TextureExtent"));
+		VisualizeUINT8Texture2D.Bind(Initializer.ParameterMap, TEXT("VisualizeUINT8Texture2D"));
 	}
 
 	/** Serializer */
@@ -83,6 +96,7 @@ public:
 		Ar << VisualizeTextureCubeArraySampler;
 		Ar << VisualizeParam;
 		Ar << TextureExtent;
+		Ar << VisualizeUINT8Texture2D;
 		return bShaderHasOutdatedParameters;
 	}
 
@@ -122,6 +136,7 @@ public:
 		SetTextureParameter(RHICmdList, ShaderRHI, VisualizeTexture3D, VisualizeTexture3DSampler, TStaticSamplerState<SF_Point,AM_Clamp,AM_Clamp,AM_Clamp>::GetRHI(), (FTextureRHIRef&)Data.RenderTargetItem.ShaderResourceTexture);
 		SetTextureParameter(RHICmdList, ShaderRHI, VisualizeTextureCube, VisualizeTextureCubeSampler, TStaticSamplerState<SF_Point,AM_Clamp,AM_Clamp,AM_Clamp>::GetRHI(), (FTextureRHIRef&)Data.RenderTargetItem.ShaderResourceTexture);
 		SetTextureParameter(RHICmdList, ShaderRHI, VisualizeTextureCubeArray, VisualizeTextureCubeArraySampler, TStaticSamplerState<SF_Point,AM_Clamp,AM_Clamp,AM_Clamp>::GetRHI(), (FTextureRHIRef&)Data.RenderTargetItem.ShaderResourceTexture);
+		SetTextureParameter(RHICmdList, ShaderRHI, VisualizeUINT8Texture2D, Data.RenderTargetItem.TargetableTexture);
 	}
 
 	static const TCHAR* GetSourceFilename()
@@ -145,15 +160,24 @@ protected:
 	FShaderResourceParameter VisualizeTextureCubeSampler;
 	FShaderResourceParameter VisualizeTextureCubeArray;
 	FShaderResourceParameter VisualizeTextureCubeArraySampler;
+	FShaderResourceParameter VisualizeUINT8Texture2D;
 	FShaderParameter VisualizeParam;
 	FShaderParameter TextureExtent;
 };
 
 // #define avoids a lot of code duplication
-#define VARIATION1(A) typedef VisualizeTexturePS<A> VisualizeTexturePS##A; \
+#define VARIATION1(A) typedef VisualizeTexturePS<EVisualisePSType::A> VisualizeTexturePS##A; \
 	IMPLEMENT_SHADER_TYPE2(VisualizeTexturePS##A, SF_Pixel);
 
-VARIATION1(0)			VARIATION1(2)			VARIATION1(3)			VARIATION1(4)			VARIATION1(5)			VARIATION1(6)
+
+VARIATION1(Cube)
+VARIATION1(Texture2DNoMSAA)
+VARIATION1(Texture3D)
+VARIATION1(CubeArray)
+VARIATION1(Texture2DMSAA)
+VARIATION1(Texture2DDepthStencilNoMSAA)	
+VARIATION1(Texture2DUINT8)
+
 #undef VARIATION1
 
 
@@ -203,7 +227,7 @@ public:
 IMPLEMENT_SHADER_TYPE(,FVisualizeTexturePresentPS,TEXT("VisualizeTexture"),TEXT("PresentPS"),SF_Pixel);
 
 
-template<uint32 TextureType> void VisualizeTextureForTextureType(FRHICommandList& RHICmdList, ERHIFeatureLevel::Type FeatureLevel, const FVisualizeTextureData& Data)
+template<EVisualisePSType TextureType> void VisualizeTextureForTextureType(FRHICommandList& RHICmdList, ERHIFeatureLevel::Type FeatureLevel, const FVisualizeTextureData& Data)
 {
 	FGraphicsPipelineStateInitializer GraphicsPSOInit;
 	RHICmdList.ApplyCachedRenderTargets(GraphicsPSOInit);
@@ -251,37 +275,41 @@ void RenderVisualizeTexture(FRHICommandListImmediate& RHICmdList, ERHIFeatureLev
 		if(Data.Desc.NumSamples > 1)
 		{
 			// MSAA
-			VisualizeTextureForTextureType<5>(RHICmdList, FeatureLevel, Data);
+			VisualizeTextureForTextureType<EVisualisePSType::Texture2DMSAA>(RHICmdList, FeatureLevel, Data);
 		}
 		else
 		{
 			if(Data.Desc.Format == PF_DepthStencil)
 			{
 				// DepthStencil non MSAA (needed to avoid D3DDebug error)
-				VisualizeTextureForTextureType<6>(RHICmdList, FeatureLevel, Data);
+				VisualizeTextureForTextureType<EVisualisePSType::Texture2DDepthStencilNoMSAA>(RHICmdList, FeatureLevel, Data);
+			}
+			else if (Data.Desc.Format == PF_R8_UINT)
+			{
+				VisualizeTextureForTextureType<EVisualisePSType::Texture2DUINT8>(RHICmdList, FeatureLevel, Data);
 			}
 			else
 			{
 				// non MSAA
-				VisualizeTextureForTextureType<2>(RHICmdList, FeatureLevel, Data);
+				VisualizeTextureForTextureType<EVisualisePSType::Texture2DNoMSAA>(RHICmdList, FeatureLevel, Data);
 			}
 		}
 	}
 	else if(Data.Desc.Is3DTexture())
 	{		
-		VisualizeTextureForTextureType<3>(RHICmdList, FeatureLevel, Data);
+		VisualizeTextureForTextureType<EVisualisePSType::Texture3D>(RHICmdList, FeatureLevel, Data);
 	}
 	else if(Data.Desc.IsCubemap())
 	{		
 		if(Data.Desc.IsArray())
 		{
 			// Cube[]
-			VisualizeTextureForTextureType<4>(RHICmdList, FeatureLevel, Data);
+			VisualizeTextureForTextureType<EVisualisePSType::CubeArray>(RHICmdList, FeatureLevel, Data);
 		}
 		else
 		{
 			// Cube
-			VisualizeTextureForTextureType<0>(RHICmdList, FeatureLevel, Data);
+			VisualizeTextureForTextureType<EVisualisePSType::Cube>(RHICmdList, FeatureLevel, Data);
 		}
 	}
 }
diff --git a/Engine/Source/Runtime/SlateRHIRenderer/Private/SlatePostProcessor.cpp b/Engine/Source/Runtime/SlateRHIRenderer/Private/SlatePostProcessor.cpp
--- a/Engine/Source/Runtime/SlateRHIRenderer/Private/SlatePostProcessor.cpp
+++ b/Engine/Source/Runtime/SlateRHIRenderer/Private/SlatePostProcessor.cpp
@@ -309,8 +309,8 @@ void FSlatePostProcessor::UpsampleRect(FRHICommandListImmediate& RHICmdList, IRe
 
 	RHICmdList.SetViewport(0, 0, 0, DestTextureWidth, DestTextureHeight, 0.0f);
 
-	RHICmdList.TransitionResource(EResourceTransitionAccess::EWritable, Params.SourceTexture);
-	RHICmdList.TransitionResource(EResourceTransitionAccess::EReadable, DestTexture);
+	RHICmdList.TransitionResource(EResourceTransitionAccess::EReadable, Params.SourceTexture);
+	RHICmdList.TransitionResource(EResourceTransitionAccess::EWritable, DestTexture);
 
 	SetRenderTarget(RHICmdList, DestTexture, FTextureRHIRef());
 	RHICmdList.ApplyCachedRenderTargets(GraphicsPSOInit);
