UE4 - Copy 3518832
Copy & extend 3518077
- Fix for movable skylight shader missing on simple forward (low lighting quality mode)

diff --git a/Engine/Source/Runtime/Renderer/Private/BasePassRendering.h b/Engine/Source/Runtime/Renderer/Private/BasePassRendering.h
--- a/Engine/Source/Runtime/Renderer/Private/BasePassRendering.h
+++ b/Engine/Source/Runtime/Renderer/Private/BasePassRendering.h
@@ -784,7 +784,7 @@ public:
 		Ar << ReflectionParameters;
 		Ar << TranslucentLightingParameters;
 		Ar << HeightFogParameters;
- 		Ar << EditorCompositeParams;
+		Ar << EditorCompositeParams;
 		Ar << ForwardLightingParameters;
 		return bShaderHasOutdatedParameters;
 	}
@@ -831,7 +831,7 @@ class TBasePassPS : public TBasePassPixelShaderBaseType<LightMapPolicyType>
 {
 	DECLARE_SHADER_TYPE(TBasePassPS,MeshMaterial);
 public:
-	
+
 	static bool ShouldCache(EShaderPlatform Platform,const FMaterial* Material,const FVertexFactoryType* VertexFactoryType)
 	{
 		// Only compile skylight version for lit materials, and if the project allows them.
@@ -842,9 +842,12 @@ public:
 		const bool bForceAllPermutations = SupportAllShaderPermutations && SupportAllShaderPermutations->GetValueOnAnyThread() != 0;
 		const bool bProjectSupportsStationarySkylight = !SupportStationarySkylight || SupportStationarySkylight->GetValueOnAnyThread() != 0 || bForceAllPermutations;
 
-		//translucent materials need to compile skylight support to support MOVABLE skylights also.
-		const bool bCacheShaders = !bEnableSkyLight || (bProjectSupportsStationarySkylight && (Material->GetShadingModel() != MSM_Unlit)) || bTranslucent;
-
+		const bool bCacheShaders = !bEnableSkyLight
+			//translucent materials need to compile skylight support to support MOVABLE skylights also.
+			|| bTranslucent
+			// Some lightmap policies (eg Simple Forward) always require skylight support
+			|| LightMapPolicyType::RequiresSkylight()
+			|| (bProjectSupportsStationarySkylight && (Material->GetShadingModel() != MSM_Unlit));
 		return bCacheShaders
 			&& (IsFeatureLevelSupported(Platform, ERHIFeatureLevel::SM4))
 			&& TBasePassPixelShaderBaseType<LightMapPolicyType>::ShouldCache(Platform, Material, VertexFactoryType);
@@ -1451,110 +1454,110 @@ void ProcessBasePassMesh(
 	const bool bAllowStaticLighting = (!AllowStaticLightingVar || AllowStaticLightingVar->GetValueOnRenderThread() != 0);
 	
 	
-		const FLightMapInteraction LightMapInteraction = (bAllowStaticLighting && Parameters.Mesh.LCI && bIsLitMaterial) 
-			? Parameters.Mesh.LCI->GetLightMapInteraction(Parameters.FeatureLevel) 
-			: FLightMapInteraction();
+	const FLightMapInteraction LightMapInteraction = (bAllowStaticLighting && Parameters.Mesh.LCI && bIsLitMaterial) 
+		? Parameters.Mesh.LCI->GetLightMapInteraction(Parameters.FeatureLevel) 
+		: FLightMapInteraction();
 
-		// force LQ lightmaps based on system settings
-		const bool bPlatformAllowsHighQualityLightMaps = AllowHighQualityLightmaps(Parameters.FeatureLevel);
-		const bool bAllowHighQualityLightMaps = bPlatformAllowsHighQualityLightMaps && LightMapInteraction.AllowsHighQualityLightmaps();
+	// force LQ lightmaps based on system settings
+	const bool bPlatformAllowsHighQualityLightMaps = AllowHighQualityLightmaps(Parameters.FeatureLevel);
+	const bool bAllowHighQualityLightMaps = bPlatformAllowsHighQualityLightMaps && LightMapInteraction.AllowsHighQualityLightmaps();
 
-		if (IsSimpleForwardShadingEnabled(GetFeatureLevelShaderPlatform(Parameters.FeatureLevel)))
+	if (IsSimpleForwardShadingEnabled(GetFeatureLevelShaderPlatform(Parameters.FeatureLevel)))
+	{
+		// Only compiling simple lighting shaders for HQ lightmaps to save on permutations
+		check(bPlatformAllowsHighQualityLightMaps);
+		ProcessBasePassMeshForSimpleForwardShading(RHICmdList, Parameters, Action, LightMapInteraction, bIsLitMaterial, bAllowStaticLighting);
+	}
+	// Render self-shadowing only for >= SM4 and fallback to non-shadowed for lesser shader models
+	else if (bIsLitMaterial && Action.UseTranslucentSelfShadowing() && Parameters.FeatureLevel >= ERHIFeatureLevel::SM4)
+	{
+		if (IsIndirectLightingCacheAllowed(Parameters.FeatureLevel)
+				&& Action.AllowIndirectLightingCache()
+				&& Parameters.PrimitiveSceneProxy)
 		{
-			// Only compiling simple lighting shaders for HQ lightmaps to save on permutations
-			check(bPlatformAllowsHighQualityLightMaps);
-			ProcessBasePassMeshForSimpleForwardShading(RHICmdList, Parameters, Action, LightMapInteraction, bIsLitMaterial, bAllowStaticLighting);
-		}
-		// Render self-shadowing only for >= SM4 and fallback to non-shadowed for lesser shader models
-		else if (bIsLitMaterial && Action.UseTranslucentSelfShadowing() && Parameters.FeatureLevel >= ERHIFeatureLevel::SM4)
-		{
-			if (IsIndirectLightingCacheAllowed(Parameters.FeatureLevel)
-					&& Action.AllowIndirectLightingCache()
-					&& Parameters.PrimitiveSceneProxy)
-			{
-				// Apply cached point indirect lighting as well as self shadowing if needed
-				Action.template Process<FSelfShadowedCachedPointIndirectLightingPolicy>(RHICmdList, Parameters, FSelfShadowedCachedPointIndirectLightingPolicy(), FSelfShadowedTranslucencyPolicy::ElementDataType(Action.GetTranslucentSelfShadow()));
-			}
-			else
-			{
-				Action.template Process<FSelfShadowedTranslucencyPolicy>(RHICmdList, Parameters, FSelfShadowedTranslucencyPolicy(), FSelfShadowedTranslucencyPolicy::ElementDataType(Action.GetTranslucentSelfShadow()));
-			}
+			// Apply cached point indirect lighting as well as self shadowing if needed
+			Action.template Process<FSelfShadowedCachedPointIndirectLightingPolicy>(RHICmdList, Parameters, FSelfShadowedCachedPointIndirectLightingPolicy(), FSelfShadowedTranslucencyPolicy::ElementDataType(Action.GetTranslucentSelfShadow()));
 		}
 		else
 		{
-			static const auto CVarSupportLowQualityLightmap = IConsoleManager::Get().FindTConsoleVariableDataInt(TEXT("r.SupportLowQualityLightmaps"));
-			const bool bAllowLowQualityLightMaps = (!CVarSupportLowQualityLightmap) || (CVarSupportLowQualityLightmap->GetValueOnAnyThread() != 0);
+			Action.template Process<FSelfShadowedTranslucencyPolicy>(RHICmdList, Parameters, FSelfShadowedTranslucencyPolicy(), FSelfShadowedTranslucencyPolicy::ElementDataType(Action.GetTranslucentSelfShadow()));
+		}
+	}
+	else
+	{
+		static const auto CVarSupportLowQualityLightmap = IConsoleManager::Get().FindTConsoleVariableDataInt(TEXT("r.SupportLowQualityLightmaps"));
+		const bool bAllowLowQualityLightMaps = (!CVarSupportLowQualityLightmap) || (CVarSupportLowQualityLightmap->GetValueOnAnyThread() != 0);
 
-			switch (LightMapInteraction.GetType())
+		switch (LightMapInteraction.GetType())
+		{
+		case LMIT_Texture:
+			if (bAllowHighQualityLightMaps)
 			{
-			case LMIT_Texture:
-				if (bAllowHighQualityLightMaps)
-				{
-					const FShadowMapInteraction ShadowMapInteraction = (bAllowStaticLighting && Parameters.Mesh.LCI && bIsLitMaterial)
-						? Parameters.Mesh.LCI->GetShadowMapInteraction()
-						: FShadowMapInteraction();
+				const FShadowMapInteraction ShadowMapInteraction = (bAllowStaticLighting && Parameters.Mesh.LCI && bIsLitMaterial)
+					? Parameters.Mesh.LCI->GetShadowMapInteraction()
+					: FShadowMapInteraction();
 
-					if (ShadowMapInteraction.GetType() == SMIT_Texture)
-					{
-						Action.template Process< FUniformLightMapPolicy >(RHICmdList, Parameters, FUniformLightMapPolicy(LMP_DISTANCE_FIELD_SHADOWS_AND_HQ_LIGHTMAP), Parameters.Mesh.LCI);
-					}
-					else
-					{
-						Action.template Process< FUniformLightMapPolicy >(RHICmdList, Parameters, FUniformLightMapPolicy(LMP_HQ_LIGHTMAP), Parameters.Mesh.LCI);
-					}
-				}
-				else if (bAllowLowQualityLightMaps)
+				if (ShadowMapInteraction.GetType() == SMIT_Texture)
 				{
-					Action.template Process< FUniformLightMapPolicy >(RHICmdList, Parameters, FUniformLightMapPolicy(LMP_LQ_LIGHTMAP), Parameters.Mesh.LCI);
+					Action.template Process< FUniformLightMapPolicy >(RHICmdList, Parameters, FUniformLightMapPolicy(LMP_DISTANCE_FIELD_SHADOWS_AND_HQ_LIGHTMAP), Parameters.Mesh.LCI);
 				}
 				else
 				{
-					Action.template Process< FUniformLightMapPolicy >(RHICmdList, Parameters, FUniformLightMapPolicy(LMP_NO_LIGHTMAP), Parameters.Mesh.LCI);
+					Action.template Process< FUniformLightMapPolicy >(RHICmdList, Parameters, FUniformLightMapPolicy(LMP_HQ_LIGHTMAP), Parameters.Mesh.LCI);
 				}
-				break;
-			default:
-				if (bIsLitMaterial
-					&& IsIndirectLightingCacheAllowed(Parameters.FeatureLevel)
-					&& Action.AllowIndirectLightingCache()
-					&& Parameters.PrimitiveSceneProxy)
+			}
+			else if (bAllowLowQualityLightMaps)
+			{
+				Action.template Process< FUniformLightMapPolicy >(RHICmdList, Parameters, FUniformLightMapPolicy(LMP_LQ_LIGHTMAP), Parameters.Mesh.LCI);
+			}
+			else
+			{
+				Action.template Process< FUniformLightMapPolicy >(RHICmdList, Parameters, FUniformLightMapPolicy(LMP_NO_LIGHTMAP), Parameters.Mesh.LCI);
+			}
+			break;
+		default:
+			if (bIsLitMaterial
+				&& IsIndirectLightingCacheAllowed(Parameters.FeatureLevel)
+				&& Action.AllowIndirectLightingCache()
+				&& Parameters.PrimitiveSceneProxy)
+			{
+				const FIndirectLightingCacheAllocation* IndirectLightingCacheAllocation = Parameters.PrimitiveSceneProxy->GetPrimitiveSceneInfo()->IndirectLightingCacheAllocation;
+				const bool bPrimitiveIsMovable = Parameters.PrimitiveSceneProxy->IsMovable();
+				const bool bPrimitiveUsesILC = Parameters.PrimitiveSceneProxy->GetIndirectLightingCacheQuality() != ILCQ_Off;
+
+				// Use the indirect lighting cache shaders if the object has a cache allocation
+				// This happens for objects with unbuilt lighting
+				if (bPrimitiveUsesILC &&
+					((IndirectLightingCacheAllocation && IndirectLightingCacheAllocation->IsValid())
+						// Use the indirect lighting cache shaders if the object is movable, it may not have a cache allocation yet because that is done in InitViews
+						// And movable objects are sometimes rendered in the static draw lists
+						|| bPrimitiveIsMovable))
 				{
-					const FIndirectLightingCacheAllocation* IndirectLightingCacheAllocation = Parameters.PrimitiveSceneProxy->GetPrimitiveSceneInfo()->IndirectLightingCacheAllocation;
-					const bool bPrimitiveIsMovable = Parameters.PrimitiveSceneProxy->IsMovable();
-					const bool bPrimitiveUsesILC = Parameters.PrimitiveSceneProxy->GetIndirectLightingCacheQuality() != ILCQ_Off;
-
-					// Use the indirect lighting cache shaders if the object has a cache allocation
-					// This happens for objects with unbuilt lighting
-					if (bPrimitiveUsesILC &&
-						((IndirectLightingCacheAllocation && IndirectLightingCacheAllocation->IsValid())
-							// Use the indirect lighting cache shaders if the object is movable, it may not have a cache allocation yet because that is done in InitViews
-							// And movable objects are sometimes rendered in the static draw lists
-							|| bPrimitiveIsMovable))
+					if (CanIndirectLightingCacheUseVolumeTexture(Parameters.FeatureLevel)
+						// Translucency forces point sample for pixel performance
+						&& Action.AllowIndirectLightingCacheVolumeTexture()
+						&& ((IndirectLightingCacheAllocation && !IndirectLightingCacheAllocation->bPointSample)
+							|| (bPrimitiveIsMovable && Parameters.PrimitiveSceneProxy->GetIndirectLightingCacheQuality() == ILCQ_Volume)))
 					{
-						if (CanIndirectLightingCacheUseVolumeTexture(Parameters.FeatureLevel)
-							// Translucency forces point sample for pixel performance
-							&& Action.AllowIndirectLightingCacheVolumeTexture()
-							&& ((IndirectLightingCacheAllocation && !IndirectLightingCacheAllocation->bPointSample)
-								|| (bPrimitiveIsMovable && Parameters.PrimitiveSceneProxy->GetIndirectLightingCacheQuality() == ILCQ_Volume)))
-						{
-							// Use a lightmap policy that supports reading indirect lighting from a volume texture for dynamic objects
-							Action.template Process< FUniformLightMapPolicy >(RHICmdList, Parameters, FUniformLightMapPolicy(LMP_CACHED_VOLUME_INDIRECT_LIGHTING), Parameters.Mesh.LCI);
-						}
-						else
-						{
-							// Use a lightmap policy that supports reading indirect lighting from a single SH sample
-							Action.template Process< FUniformLightMapPolicy >(RHICmdList, Parameters, FUniformLightMapPolicy(LMP_CACHED_POINT_INDIRECT_LIGHTING), Parameters.Mesh.LCI);
-						}
+						// Use a lightmap policy that supports reading indirect lighting from a volume texture for dynamic objects
+						Action.template Process< FUniformLightMapPolicy >(RHICmdList, Parameters, FUniformLightMapPolicy(LMP_CACHED_VOLUME_INDIRECT_LIGHTING), Parameters.Mesh.LCI);
 					}
 					else
 					{
-						Action.template Process< FUniformLightMapPolicy >(RHICmdList, Parameters, FUniformLightMapPolicy(LMP_NO_LIGHTMAP), Parameters.Mesh.LCI);
+						// Use a lightmap policy that supports reading indirect lighting from a single SH sample
+						Action.template Process< FUniformLightMapPolicy >(RHICmdList, Parameters, FUniformLightMapPolicy(LMP_CACHED_POINT_INDIRECT_LIGHTING), Parameters.Mesh.LCI);
 					}
 				}
 				else
 				{
 					Action.template Process< FUniformLightMapPolicy >(RHICmdList, Parameters, FUniformLightMapPolicy(LMP_NO_LIGHTMAP), Parameters.Mesh.LCI);
 				}
-				break;
-			};
+			}
+			else
+			{
+				Action.template Process< FUniformLightMapPolicy >(RHICmdList, Parameters, FUniformLightMapPolicy(LMP_NO_LIGHTMAP), Parameters.Mesh.LCI);
+			}
+			break;
+		};
 	}	
 }
diff --git a/Engine/Source/Runtime/Renderer/Private/LightMapRendering.h b/Engine/Source/Runtime/Renderer/Private/LightMapRendering.h
--- a/Engine/Source/Runtime/Renderer/Private/LightMapRendering.h
+++ b/Engine/Source/Runtime/Renderer/Private/LightMapRendering.h
@@ -98,6 +98,11 @@ struct FNoLightMapPolicy
 
 	static void ModifyCompilationEnvironment(EShaderPlatform Platform, const FMaterial* Material, FShaderCompilerEnvironment& OutEnvironment)
 	{}
+
+	static bool RequiresSkylight()
+	{
+		return false;
+	}
 };
 
 enum ELightmapQuality
@@ -141,6 +146,11 @@ struct TLightMapPolicy
 			&& (!AllowStaticLightingVar || AllowStaticLightingVar->GetValueOnAnyThread() != 0)
 			&& (Material->IsUsedWithStaticLighting() || Material->IsSpecialEngineMaterial());		
 	}
+
+	static bool RequiresSkylight()
+	{
+		return false;
+	}
 };
 
 // A light map policy for computing up to 4 signed distance field shadow factors in the base pass.
@@ -231,6 +241,11 @@ public:
 		OutEnvironment.SetDefine(TEXT("TRANSLUCENT_SELF_SHADOWING"),TEXT("1"));
 	}
 
+	static bool RequiresSkylight()
+	{
+		return false;
+	}
+
 	/** Initialization constructor. */
 	FSelfShadowedTranslucencyPolicy() {}
 
@@ -323,6 +338,11 @@ struct FCachedVolumeIndirectLightingPolicy
 	{
 		OutEnvironment.SetDefine(TEXT("CACHED_VOLUME_INDIRECT_LIGHTING"),TEXT("1"));
 	}
+
+	static bool RequiresSkylight()
+	{
+		return false;
+	}
 };
 
 
@@ -340,6 +360,11 @@ struct FCachedPointIndirectLightingPolicy
 	}
 
 	static void ModifyCompilationEnvironment(EShaderPlatform Platform, const FMaterial* Material, FShaderCompilerEnvironment& OutEnvironment);
+
+	static bool RequiresSkylight()
+	{
+		return false;
+	}
 };
 
 /**
@@ -358,6 +383,11 @@ struct FSimpleNoLightmapLightingPolicy : public FNoLightMapPolicy
 		OutEnvironment.SetDefine(TEXT("SIMPLE_FORWARD_SHADING"),TEXT("1"));
 		FNoLightMapPolicy::ModifyCompilationEnvironment(Platform, Material, OutEnvironment);
 	}
+
+	static bool RequiresSkylight()
+	{
+		return true;
+	}
 };
 
 /**
@@ -379,6 +409,11 @@ struct FSimpleLightmapOnlyLightingPolicy : public TLightMapPolicy<HQ_LIGHTMAP>
 		OutEnvironment.SetDefine(TEXT("SIMPLE_FORWARD_SHADING"),TEXT("1"));
 		TLightMapPolicy<HQ_LIGHTMAP>::ModifyCompilationEnvironment(Platform, Material, OutEnvironment);
 	}
+
+	static bool RequiresSkylight()
+	{
+		return true;
+	}
 };
 
 /**
@@ -397,6 +432,11 @@ struct FSimpleDirectionalLightLightingPolicy
 		OutEnvironment.SetDefine(TEXT("SIMPLE_FORWARD_SHADING"),TEXT("1"));
 		OutEnvironment.SetDefine(TEXT("SIMPLE_FORWARD_DIRECTIONAL_LIGHT"),TEXT("1"));
 	}
+
+	static bool RequiresSkylight()
+	{
+		return true;
+	}
 };
 
 /**
@@ -419,6 +459,11 @@ struct FSimpleStationaryLightPrecomputedShadowsLightingPolicy : public TDistance
 		OutEnvironment.SetDefine(TEXT("SIMPLE_FORWARD_DIRECTIONAL_LIGHT"),TEXT("1"));
 		TDistanceFieldShadowsAndLightMapPolicy<HQ_LIGHTMAP>::ModifyCompilationEnvironment(Platform, Material, OutEnvironment);
 	}
+
+	static bool RequiresSkylight()
+	{
+		return true;
+	}
 };
 
 /**
@@ -441,6 +486,11 @@ struct FSimpleStationaryLightSingleSampleShadowsLightingPolicy : public FCachedP
 		OutEnvironment.SetDefine(TEXT("SIMPLE_FORWARD_DIRECTIONAL_LIGHT"),TEXT("1"));
 		FCachedPointIndirectLightingPolicy::ModifyCompilationEnvironment(Platform, Material, OutEnvironment);
 	}
+
+	static bool RequiresSkylight()
+	{
+		return true;
+	}
 };
 
 /** Mobile Specific: Combines a distance field shadow with LQ lightmaps. */
@@ -454,6 +504,11 @@ public:
 		const bool bMobileAllowDistanceFieldShadows = CVarMobileAllowDistanceFieldShadows->GetValueOnAnyThread() == 1;
 		return bMobileAllowDistanceFieldShadows && Super::ShouldCache(Platform, Material, VertexFactoryType);
 	}
+
+	static bool RequiresSkylight()
+	{
+		return false;
+	}
 };
 
 /** Mobile Specific: Combines an distance field shadow with LQ lightmaps and CSM. */
@@ -476,6 +531,11 @@ public:
 
 		Super::ModifyCompilationEnvironment(Platform, Material, OutEnvironment);
 	}
+
+	static bool RequiresSkylight()
+	{
+		return false;
+	}
 };
 
 /** Mobile Specific: Combines an unshadowed directional light with indirect lighting from a single SH sample. */
@@ -493,6 +553,11 @@ struct FMobileDirectionalLightAndSHIndirectPolicy
 	{
 		FCachedPointIndirectLightingPolicy::ModifyCompilationEnvironment(Platform, Material, OutEnvironment);
 	}
+
+	static bool RequiresSkylight()
+	{
+		return false;
+	}
 };
 
 /** Mobile Specific: Combines a movable directional light with indirect lighting from a single SH sample. */
@@ -512,6 +577,11 @@ struct FMobileMovableDirectionalLightAndSHIndirectPolicy : public FMobileDirecti
 		OutEnvironment.SetDefine(TEXT(PREPROCESSOR_TO_STRING(MAX_MOBILE_SHADOWCASCADES)), MAX_MOBILE_SHADOWCASCADES);
 		FMobileDirectionalLightAndSHIndirectPolicy::ModifyCompilationEnvironment(Platform, Material, OutEnvironment);
 	}
+
+	static bool RequiresSkylight()
+	{
+		return false;
+	}
 };
 
 /** Mobile Specific: Combines a movable directional light with CSM with indirect lighting from a single SH sample. */
@@ -523,6 +593,11 @@ struct FMobileMovableDirectionalLightCSMAndSHIndirectPolicy : public FMobileMova
 		OutEnvironment.SetDefine(TEXT(PREPROCESSOR_TO_STRING(MAX_MOBILE_SHADOWCASCADES)), MAX_MOBILE_SHADOWCASCADES);
 		FMobileMovableDirectionalLightAndSHIndirectPolicy::ModifyCompilationEnvironment(Platform, Material, OutEnvironment);
 	}
+
+	static bool RequiresSkylight()
+	{
+		return false;
+	}
 };
 
 /** Mobile Specific: Combines a directional light with CSM with indirect lighting from a single SH sample. */
@@ -535,6 +610,11 @@ public:
 		OutEnvironment.SetDefine(TEXT(PREPROCESSOR_TO_STRING(MAX_MOBILE_SHADOWCASCADES)), MAX_MOBILE_SHADOWCASCADES);
 		FMobileDirectionalLightAndSHIndirectPolicy::ModifyCompilationEnvironment(Platform, Material, OutEnvironment);
 	}
+
+	static bool RequiresSkylight()
+	{
+		return false;
+	}
 };
 
 /** Mobile Specific */
@@ -552,6 +632,11 @@ struct FMobileMovableDirectionalLightLightingPolicy
 	{
 		OutEnvironment.SetDefine(TEXT("MOVABLE_DIRECTIONAL_LIGHT"),TEXT("1"));
 	}
+
+	static bool RequiresSkylight()
+	{
+		return false;
+	}
 };
 
 /** Mobile Specific */
@@ -573,6 +658,11 @@ struct FMobileMovableDirectionalLightCSMLightingPolicy
 
 		FNoLightMapPolicy::ModifyCompilationEnvironment(Platform, Material, OutEnvironment);
 	}
+
+	static bool RequiresSkylight()
+	{
+		return false;
+	}
 };
 
 /** Mobile Specific */
@@ -598,6 +688,11 @@ struct FMobileMovableDirectionalLightWithLightmapPolicy : public TLightMapPolicy
 
 		Super::ModifyCompilationEnvironment(Platform, Material, OutEnvironment);
 	}
+
+	static bool RequiresSkylight()
+	{
+		return false;
+	}
 };
 
 /** Mobile Specific */
@@ -609,6 +704,11 @@ struct FMobileMovableDirectionalLightCSMWithLightmapPolicy : public FMobileMovab
 
 		FMobileMovableDirectionalLightWithLightmapPolicy::ModifyCompilationEnvironment(Platform, Material, OutEnvironment);
 	}
+
+	static bool RequiresSkylight()
+	{
+		return false;
+	}
 };
 
 enum ELightMapPolicyType
@@ -869,6 +969,23 @@ public:
 			break;
 		}
 	}
+
+	static bool RequiresSkylight()
+	{
+		CA_SUPPRESS(6326);
+		switch (Policy)
+		{
+		// Simple forward
+		case LMP_SIMPLE_NO_LIGHTMAP:
+		case LMP_SIMPLE_LIGHTMAP_ONLY_LIGHTING:
+		case LMP_SIMPLE_DIRECTIONAL_LIGHT_LIGHTING:
+		case LMP_SIMPLE_STATIONARY_PRECOMPUTED_SHADOW_LIGHTING:
+		case LMP_SIMPLE_STATIONARY_SINGLESAMPLE_SHADOW_LIGHTING:
+			return true;
+		default:
+			return false;
+		};
+	}
 };
 
 /**
diff --git a/Engine/Source/Runtime/Renderer/Private/ScenePrivate.h b/Engine/Source/Runtime/Renderer/Private/ScenePrivate.h
--- a/Engine/Source/Runtime/Renderer/Private/ScenePrivate.h
+++ b/Engine/Source/Runtime/Renderer/Private/ScenePrivate.h
@@ -2260,9 +2260,8 @@ public:
 	virtual ERHIFeatureLevel::Type GetFeatureLevel() const override { return FeatureLevel; }
 
 	bool ShouldRenderSkylightInBasePass(EBlendMode BlendMode) const
-	{		
-		const bool bStationarySkylight = SkyLight && SkyLight->bWantsStaticShadowing;
-		return ShouldRenderSkylightInBasePass_Internal(BlendMode) && (ReadOnlyCVARCache.bEnableStationarySkylight || !bStationarySkylight);
+	{
+		return ShouldRenderSkylightInBasePass_Internal(BlendMode) && (ReadOnlyCVARCache.bEnableStationarySkylight || IsSimpleForwardShadingEnabled(GetShaderPlatform()));
 	}
 
 	bool ShouldRenderSkylightInBasePass_Internal(EBlendMode BlendMode) const
