Fixes CIs compilation error caused by CL 3517032 that would try to compile PS4 specific shader file even on non PS4 shader platform.

diff --git a/Engine/Source/Runtime/Renderer/Private/PostProcess/PostProcessMorpheus.cpp b/Engine/Source/Runtime/Renderer/Private/PostProcess/PostProcessMorpheus.cpp
--- a/Engine/Source/Runtime/Renderer/Private/PostProcess/PostProcessMorpheus.cpp
+++ b/Engine/Source/Runtime/Renderer/Private/PostProcess/PostProcessMorpheus.cpp
@@ -30,7 +30,7 @@ class FPostProcessMorpheusPS : public FGlobalShader
 		// we must use a run time check for this because the builds the build machines create will have Morpheus defined,
 		// but a user will not necessarily have the Morpheus files
 		bool bEnableMorpheus = false;
-		if (GConfig->GetBool(TEXT("/Script/MorpheusEditor.MorpheusRuntimeSettings"), TEXT("bEnableMorpheus"), bEnableMorpheus, GEngineIni))
+		if (Platform == EShaderPlatform::SP_PS4 && GConfig->GetBool(TEXT("/Script/MorpheusEditor.MorpheusRuntimeSettings"), TEXT("bEnableMorpheus"), bEnableMorpheus, GEngineIni))
 		{
 			return bEnableMorpheus;
 		}
