Fixes compilation failure of PostProcessHMDMorpheus (again). The includer scanner was failing because does not have the logic for file ignoring based on whether it is in PS4 or XboxOne directories.

diff --git a/Engine/Shaders/Private/PostProcessHMDMorpheus.usf b/Engine/Shaders/Private/PostProcessHMDMorpheus.usf
new file mode 100644
--- /dev/null
+++ b/Engine/Shaders/Private/PostProcessHMDMorpheus.usf
@@ -0,0 +1,9 @@
+// Copyright 1998-2017 Epic Games, Inc. All Rights Reserved.
+
+// PS4 directory may not exists. As a result a shader compiler may not be able to
+// find file for include scanning, but the include scanner will silently ignore
+// PS4 directory by design.
+//
+// TODO(VR team): move shader file and C++ into a plugin.
+
+#include "PS4/PostProcessHMDMorpheus.usf"
diff --git a/Engine/Source/Runtime/Renderer/Private/PostProcess/PostProcessMorpheus.cpp b/Engine/Source/Runtime/Renderer/Private/PostProcess/PostProcessMorpheus.cpp
--- a/Engine/Source/Runtime/Renderer/Private/PostProcess/PostProcessMorpheus.cpp
+++ b/Engine/Source/Runtime/Renderer/Private/PostProcess/PostProcessMorpheus.cpp
@@ -185,8 +185,8 @@ public:
 	}
 };
 
-IMPLEMENT_SHADER_TYPE(, FPostProcessMorpheusVS, TEXT("/Engine/Private/PS4/PostProcessHMDMorpheus.usf"), TEXT("MainVS"), SF_Vertex);
-IMPLEMENT_SHADER_TYPE(, FPostProcessMorpheusPS,TEXT("/Engine/Private/PS4/PostProcessHMDMorpheus.usf"),TEXT("MainPS"),SF_Pixel);
+IMPLEMENT_SHADER_TYPE(, FPostProcessMorpheusVS, TEXT("/Engine/Private/PostProcessHMDMorpheus.usf"), TEXT("MainVS"), SF_Vertex);
+IMPLEMENT_SHADER_TYPE(, FPostProcessMorpheusPS, TEXT("/Engine/Private/PostProcessHMDMorpheus.usf"), TEXT("MainPS"), SF_Pixel);
 
 void FRCPassPostProcessMorpheus::Process(FRenderingCompositePassContext& Context)
 {
