Fixes an assertion failures when rendering with wire frame.

diff --git a/Engine/Source/Editor/UnrealEd/Private/EditorViewportClient.cpp b/Engine/Source/Editor/UnrealEd/Private/EditorViewportClient.cpp
--- a/Engine/Source/Editor/UnrealEd/Private/EditorViewportClient.cpp
+++ b/Engine/Source/Editor/UnrealEd/Private/EditorViewportClient.cpp
@@ -3354,7 +3354,8 @@ void FEditorViewportClient::Draw(FViewport* InViewport, FCanvas* Canvas)
 		ViewExt->SetupViewFamily(ViewFamily);
 	}
 
-	EngineShowFlagOverride(ESFIM_Editor, GetViewMode(), ViewFamily.EngineShowFlags, CurrentBufferVisualizationMode);
+	ViewFamily.ViewMode = GetViewMode();
+	EngineShowFlagOverride(ESFIM_Editor, ViewFamily.ViewMode, ViewFamily.EngineShowFlags, CurrentBufferVisualizationMode);
 	EngineShowFlagOrthographicOverride(IsPerspective(), ViewFamily.EngineShowFlags);
 
 	UpdateLightingShowFlags( ViewFamily.EngineShowFlags );
diff --git a/Engine/Source/Runtime/Engine/Private/GameViewportClient.cpp b/Engine/Source/Runtime/Engine/Private/GameViewportClient.cpp
--- a/Engine/Source/Runtime/Engine/Private/GameViewportClient.cpp
+++ b/Engine/Source/Runtime/Engine/Private/GameViewportClient.cpp
@@ -1030,7 +1030,8 @@ void UGameViewportClient::Draw(FViewport* InViewport, FCanvas* SceneCanvas)
 	}
 
 	ESplitScreenType::Type SplitScreenConfig = GetCurrentSplitscreenConfiguration();
-	EngineShowFlagOverride(ESFIM_Game, (EViewModeIndex)ViewModeIndex, ViewFamily.EngineShowFlags, NAME_None, SplitScreenConfig != ESplitScreenType::None);
+	ViewFamily.ViewMode = EViewModeIndex(ViewModeIndex);
+	EngineShowFlagOverride(ESFIM_Game, ViewFamily.ViewMode, ViewFamily.EngineShowFlags, NAME_None, SplitScreenConfig != ESplitScreenType::None);
 
 	if (ViewFamily.EngineShowFlags.VisualizeBuffer && AllowDebugViewmodes())
 	{
diff --git a/Engine/Source/Runtime/Engine/Private/SceneView.cpp b/Engine/Source/Runtime/Engine/Private/SceneView.cpp
--- a/Engine/Source/Runtime/Engine/Private/SceneView.cpp
+++ b/Engine/Source/Runtime/Engine/Private/SceneView.cpp
@@ -2300,6 +2300,7 @@ void FSceneView::SetupCommonViewUniformBufferParameters(
 
 FSceneViewFamily::FSceneViewFamily(const ConstructionValues& CVS)
 	:
+	ViewMode(VMI_Lit),
 	FamilySizeX(0),
 	FamilySizeY(0),
 	InstancedStereoWidth(0),
diff --git a/Engine/Source/Runtime/Engine/Public/SceneView.h b/Engine/Source/Runtime/Engine/Public/SceneView.h
--- a/Engine/Source/Runtime/Engine/Public/SceneView.h
+++ b/Engine/Source/Runtime/Engine/Public/SceneView.h
@@ -1300,6 +1300,9 @@ public:
 	/** The views which make up the family. */
 	TArray<const FSceneView*> Views;
 
+	/** View mode of the family. */
+	EViewModeIndex ViewMode;
+
 	/** The width in screen pixels of the view family being rendered (maximum x of all viewports). */
 	uint32 FamilySizeX;
 
diff --git a/Engine/Source/Runtime/Renderer/Private/DeferredShadingRenderer.cpp b/Engine/Source/Runtime/Renderer/Private/DeferredShadingRenderer.cpp
--- a/Engine/Source/Runtime/Renderer/Private/DeferredShadingRenderer.cpp
+++ b/Engine/Source/Runtime/Renderer/Private/DeferredShadingRenderer.cpp
@@ -636,13 +636,17 @@ void FDeferredShadingSceneRenderer::Render(FRHICommandListImmediate& RHICmdList)
 	bool bComputeLightGrid = false;
 	if (bUseGBuffer)
 	{
-		bComputeLightGrid = bRenderDeferredLighting || ShouldRenderVolumetricFog();
+		bComputeLightGrid = bRenderDeferredLighting;
 	}
 	else
 	{
-		bComputeLightGrid = ViewFamily.EngineShowFlags.Lighting || ShouldRenderVolumetricFog();
+		bComputeLightGrid = ViewFamily.EngineShowFlags.Lighting;
 	}
 
+	bComputeLightGrid |= (
+		ShouldRenderVolumetricFog() ||
+		ViewFamily.ViewMode != VMI_Lit);
+
 	if (ClearMethodCVar)
 	{
 		int32 ClearMethod = ClearMethodCVar->GetValueOnRenderThread();
