Cherrypick 3635836: Fixes crash when building lighting with skylight that draws at least one non emissive only translucent material.

diff --git a/Engine/Source/Runtime/Renderer/Private/DeferredShadingRenderer.cpp b/Engine/Source/Runtime/Renderer/Private/DeferredShadingRenderer.cpp
--- a/Engine/Source/Runtime/Renderer/Private/DeferredShadingRenderer.cpp
+++ b/Engine/Source/Runtime/Renderer/Private/DeferredShadingRenderer.cpp
@@ -144,6 +144,8 @@ DECLARE_FLOAT_COUNTER_STAT(TEXT("Postprocessing"), Stat_GPU_Postprocessing, STAT
 DECLARE_FLOAT_COUNTER_STAT(TEXT("HZB"), Stat_GPU_HZB, STATGROUP_GPU);
 DECLARE_FLOAT_COUNTER_STAT(TEXT("[unaccounted]"), Stat_GPU_Unaccounted, STATGROUP_GPU);
 
+FForwardLightingViewResources* GetMinimalDummyForwardLightingResources();
+
 bool ShouldForceFullDepthPass(ERHIFeatureLevel::Type FeatureLevel)
 {
 	static IConsoleVariable* CDBufferVar = IConsoleManager::Get().FindConsoleVariable(TEXT("r.DBuffer"));
@@ -807,7 +809,14 @@ void FDeferredShadingSceneRenderer::Render(FRHICommandListImmediate& RHICmdList)
     if (bComputeLightGrid)
     {
         ComputeLightGrid(RHICmdList);
-    }
+	}
+	else
+	{
+		for (auto& View : Views)
+		{
+			View.ForwardLightingResources = GetMinimalDummyForwardLightingResources();
+		}
+	}
 
 	if (bUseGBuffer || IsSimpleForwardShadingEnabled(GetFeatureLevelShaderPlatform(FeatureLevel)))
 	{
diff --git a/Engine/Source/Runtime/Renderer/Private/Renderer.cpp b/Engine/Source/Runtime/Renderer/Private/Renderer.cpp
--- a/Engine/Source/Runtime/Renderer/Private/Renderer.cpp
+++ b/Engine/Source/Runtime/Renderer/Private/Renderer.cpp
@@ -63,7 +63,17 @@ public:
 	}
 };
 
-static TGlobalResource<FMinimalDummyForwardLightingResources>* GMinimalDummyForwardLightingResources = NULL;
+FForwardLightingViewResources* GetMinimalDummyForwardLightingResources()
+{
+	static TGlobalResource<FMinimalDummyForwardLightingResources>* GMinimalDummyForwardLightingResources = nullptr;
+
+	if (!GMinimalDummyForwardLightingResources)
+	{
+		GMinimalDummyForwardLightingResources = new TGlobalResource<FMinimalDummyForwardLightingResources>();
+	}
+
+	return &GMinimalDummyForwardLightingResources->ForwardLightingResources;
+}
 
 DEFINE_LOG_CATEGORY(LogRenderer);
 
@@ -101,16 +111,11 @@ void FRendererModule::DrawTileMesh(FRHICommandListImmediate& RHICmdList, FDrawin
 		//@todo - reuse this view for multiple tiles, this is going to be slow for each tile
 		FViewInfo View(&SceneView);
 		
-		if (!GMinimalDummyForwardLightingResources)
-		{
-			GMinimalDummyForwardLightingResources = new TGlobalResource<FMinimalDummyForwardLightingResources>();
-		}
-
 		FMaterialRenderProxy::UpdateDeferredCachedUniformExpressions();
 
 		//Apply the minimal forward lighting resources
-		View.ForwardLightingResources = &GMinimalDummyForwardLightingResources->ForwardLightingResources;
-		
+		View.ForwardLightingResources = GetMinimalDummyForwardLightingResources();
+
 		View.InitRHIResources();
 
 		const auto FeatureLevel = View.GetFeatureLevel();
