Disabled expected warning affecting Mac automated tests.

diff --git a/Engine/Source/Runtime/Renderer/Private/GPUBenchmark.cpp b/Engine/Source/Runtime/Renderer/Private/GPUBenchmark.cpp
--- a/Engine/Source/Runtime/Renderer/Private/GPUBenchmark.cpp
+++ b/Engine/Source/Runtime/Renderer/Private/GPUBenchmark.cpp
@@ -504,10 +504,11 @@ void RendererGPUBenchmark(FRHICommandListImmediate& RHICmdList, FSynthBenchmarkR
 		const bool bSupportsTimerQueries = (TimerQueries[0] != NULL);
 		if(!bSupportsTimerQueries)
 		{
+#if !PLATFORM_MAC
 			UE_LOG(LogSynthBenchmark, Warning, TEXT("GPU driver does not support timer queries."));
 
+#else
 			// Workaround for Metal not having a timing API and some drivers not properly supporting command-buffer completion handler based implementation...
-#if PLATFORM_MAC
 			FTextureMemoryStats MemStats;
 			RHIGetTextureMemoryStats(MemStats);
 			
