Removing material parameter collection logging and fallback buffer that came back from another branch to help track a live crash issue.

diff --git a/Engine/Source/Runtime/Engine/Private/Materials/ParameterCollection.cpp b/Engine/Source/Runtime/Engine/Private/Materials/ParameterCollection.cpp
--- a/Engine/Source/Runtime/Engine/Private/Materials/ParameterCollection.cpp
+++ b/Engine/Source/Runtime/Engine/Private/Materials/ParameterCollection.cpp
@@ -24,11 +24,9 @@ void UMaterialParameterCollection::PostInitProperties()
 {
 	Super::PostInitProperties();
 
-	DefaultResource = new FMaterialParameterCollectionInstanceResource();
-
-	if (HasAnyFlags(RF_ClassDefaultObject))
+	if (!HasAnyFlags(RF_ClassDefaultObject))
 	{
-		UpdateDefaultResource();
+		DefaultResource = new FMaterialParameterCollectionInstanceResource();
 	}
 }
 
@@ -58,13 +56,10 @@ void UMaterialParameterCollection::BeginDestroy()
 	if (DefaultResource)
 	{
 		FGuid Id = StateId;
-		FName Name = DefaultResource->GetOwnerName();
 		ENQUEUE_RENDER_COMMAND(RemoveDefaultResourceCommand)(
-			[Id, Name](FRHICommandListImmediate& RHICmdList)
+			[Id](FRHICommandListImmediate& RHICmdList)
 			{	
 				GDefaultMaterialParameterCollectionInstances.Remove(Id);			
-				UE_LOG(LogMaterial, Log, TEXT("Removed default parameter collection resource, %s, ID:%s"),
-					*Name.ToString(), *Id.ToString());
 			}
 		);
 
@@ -436,14 +431,6 @@ void UMaterialParameterCollection::CreateBufferStruct()
 
 void UMaterialParameterCollection::GetDefaultParameterData(TArray<FVector4>& ParameterData) const
 {
-	if (HasAnyFlags(RF_ClassDefaultObject))
-	{
-		// Dummy data for our fallback parameter collection, 256 empty vectors
-		ParameterData.Empty(256);
-		ParameterData.AddZeroed(256);
-		return;
-	}
-
 	// The memory layout created here must match the index assignment in UMaterialParameterCollection::GetParameterIndex
 
 	ParameterData.Empty(FMath::DivideAndRoundUp(ScalarParameters.Num(), 4) + VectorParameters.Num());
@@ -483,8 +470,6 @@ void UMaterialParameterCollection::UpdateDefaultResource()
 		[Id, Resource](FRHICommandListImmediate& RHICmdList)
 		{	
 			GDefaultMaterialParameterCollectionInstances.Add(Id, Resource);
-			UE_LOG(LogMaterial, Log, TEXT("Added default parameter collection resource, %s, ID:%s"),
-				Resource ? *Resource->GetOwnerName().ToString() : TEXT("None"), *Id.ToString());
 		}
 	);
 }
@@ -704,7 +689,7 @@ void FMaterialParameterCollectionInstanceResource::UpdateContents(const FGuid& I
 	Id = InId;
 	OwnerName = InOwnerName;
 
-	if (Data.Num() > 0)
+	if (InId != FGuid() && Data.Num() > 0)
 	{
 		UniformBufferLayout.ConstantBufferSize = Data.GetTypeSize() * Data.Num();
 		UniformBufferLayout.ResourceOffset = 0;
diff --git a/Engine/Source/Runtime/Renderer/Private/ShaderBaseClasses.cpp b/Engine/Source/Runtime/Renderer/Private/ShaderBaseClasses.cpp
--- a/Engine/Source/Runtime/Renderer/Private/ShaderBaseClasses.cpp
+++ b/Engine/Source/Runtime/Renderer/Private/ShaderBaseClasses.cpp
@@ -10,10 +10,6 @@
 #include "ScenePrivate.h"
 #include "ParameterCollection.h"
 
-#if !WITH_EDITOR
-TQueue<FString> GFailedToFindParamCollectionBufferQueue;
-#endif
-
 /** If true, cached uniform expressions are allowed. */
 int32 FMaterialShader::bAllowCachedUniformExpressions = true;
 
@@ -285,36 +281,20 @@ void FMaterialShader::SetParameters(
 				// Dump the currently registered parameter collections and the ID we failed to find.
 				// In a cooked project these numbers are persistent so we can track back to the original
 				// parameter collection that was being referenced and no longer exists
-				static TArray<FGuid> LoggedErrorGuids;
-				FGuid ParameterCollectionGuid = ParameterCollections[CollectionIndex];
-
-				if (!LoggedErrorGuids.Find(ParameterCollectionGuid))
+				FString InstancesString;
+				TMap<FGuid, FMaterialParameterCollectionInstanceResource*>::TIterator Iter = GDefaultMaterialParameterCollectionInstances.CreateIterator();
+				while (Iter)
 				{
-					LoggedErrorGuids.Push(ParameterCollectionGuid);
-				
-					FString InstancesString;
-					TMap<FGuid, FMaterialParameterCollectionInstanceResource*>::TIterator Iter = GDefaultMaterialParameterCollectionInstances.CreateIterator();
-					while (Iter)
-					{
-						FMaterialParameterCollectionInstanceResource* Instance = Iter.Value();
-						InstancesString += FString::Printf(TEXT("\n0x%p: %s: %s"),
-							Instance, Instance ? *Instance->GetOwnerName().ToString() : TEXT("None"), *Iter.Key().ToString());
-						++Iter;
-					}
-
-					UE_LOG(LogRenderer, Error, TEXT("Failed to find parameter collection buffer with GUID '%s'.\n")
-						TEXT("Currently %i listed default instances: %s"),
-						*ParameterCollectionGuid.ToString(),
-						GDefaultMaterialParameterCollectionInstances.Num(), *InstancesString);
-
-#if !WITH_EDITOR
-					FString IdAndNum = ParameterCollectionGuid.ToString() + TEXT(" ") + FString::FromInt(GDefaultMaterialParameterCollectionInstances.Num());
-					GFailedToFindParamCollectionBufferQueue.Enqueue(IdAndNum);
-#endif
+					FMaterialParameterCollectionInstanceResource* Instance = Iter.Value();
+					InstancesString += FString::Printf(TEXT("\n0x%p: %s: %s"),
+						Instance, Instance ? *Instance->GetOwnerName().ToString() : TEXT("None"), *Iter.Key().ToString());
+					++Iter;
 				}
 
-				// Apply the fallback collection buffer. This will cause some frame corruption, hopefully only a flicker
-				UniformBuffer = GDefaultMaterialParameterCollectionInstances.FindChecked(FGuid())->GetUniformBuffer();
+				UE_LOG(LogRenderer, Fatal, TEXT("Failed to find parameter collection buffer with GUID '%s'.\n")
+					TEXT("Currently %i listed default instances: %s"),
+					*ParameterCollections[CollectionIndex].ToString(),
+					GDefaultMaterialParameterCollectionInstances.Num(), *InstancesString);
 			}
 
 			SetUniformBufferParameter(RHICmdList, ShaderRHI, ParameterCollectionUniformBuffers[CollectionIndex], UniformBuffer);			
