First batch of fixes for PVS Studio warnings (all benign, just disabling).

diff --git a/Engine/Plugins/Editor/SpeedTreeImporter/Source/SpeedTreeImporter/Private/SpeedTreeImportFactory.cpp b/Engine/Plugins/Editor/SpeedTreeImporter/Source/SpeedTreeImporter/Private/SpeedTreeImportFactory.cpp
--- a/Engine/Plugins/Editor/SpeedTreeImporter/Source/SpeedTreeImporter/Private/SpeedTreeImportFactory.cpp
+++ b/Engine/Plugins/Editor/SpeedTreeImporter/Source/SpeedTreeImporter/Private/SpeedTreeImportFactory.cpp
@@ -2109,14 +2109,14 @@ UObject* USpeedTreeImportFactory::FactoryCreateBinary8(UClass* InClass, UObject*
 
 			for (uint32 VertexIndex = 0; VertexIndex < LOD.Vertices().Count(); ++VertexIndex)
 			{
-				const SpeedTree8::SVertex& Vertex = LOD.Vertices()[VertexIndex];
+				const SpeedTree8::SVertex& Vertex = LOD.Vertices()[VertexIndex]; //-V758
 				FVector vPosition = FVector(Vertex.m_vAnchor.x, Vertex.m_vAnchor.y, Vertex.m_vAnchor.z) + FVector(Vertex.m_vOffset.x, Vertex.m_vOffset.y, Vertex.m_vOffset.z);
 				RawMesh.VertexPositions.Add(vPosition);
 			}
 
 			for (uint32 DrawCallIndex = 0; DrawCallIndex < LOD.DrawCalls().Count(); ++DrawCallIndex)
 			{
-				const SpeedTree8::SDrawCall& DrawCall = LOD.DrawCalls()[DrawCallIndex];
+				const SpeedTree8::SDrawCall& DrawCall = LOD.DrawCalls()[DrawCallIndex]; //-V758
 
 				// find correct material/geometry combo
 				FIntPoint MaterialKey(DrawCall.m_uiMaterialIndex, DrawCall.m_eWindGeometryType);
@@ -2159,7 +2159,7 @@ UObject* USpeedTreeImportFactory::FactoryCreateBinary8(UClass* InClass, UObject*
 						int32 VertexIndex = LOD.Indices()[DrawCall.m_uiIndexStart + TriangleIndex * 3 + (2 - Corner)];
 						RawMesh.WedgeIndices.Add(VertexIndex);
 
-						const SpeedTree8::SVertex& Vertex = LOD.Vertices()[VertexIndex];
+						const SpeedTree8::SVertex& Vertex = LOD.Vertices()[VertexIndex]; //-V758
 
 						// tangents
 						RawMesh.WedgeTangentX.Add(FVector(Vertex.m_vTangent.x, Vertex.m_vTangent.y, Vertex.m_vTangent.z));
diff --git a/Engine/Plugins/Experimental/BackChannel/Source/BackChannel/Private/Transport/BackChannelConnection.cpp b/Engine/Plugins/Experimental/BackChannel/Source/BackChannel/Private/Transport/BackChannelConnection.cpp
--- a/Engine/Plugins/Experimental/BackChannel/Source/BackChannel/Private/Transport/BackChannelConnection.cpp
+++ b/Engine/Plugins/Experimental/BackChannel/Source/BackChannel/Private/Transport/BackChannelConnection.cpp
@@ -189,12 +189,12 @@ bool FBackChannelConnection::Listen(const int16 Port)
                 int32 ReceiveBufferSize = 2 * 1024 * 1024;
                 int32 SendBufferSize = ReceiveBufferSize;
                 
-                if (ReceiveBufferSize > 0)
+                if (ReceiveBufferSize > 0) //-V547
                 {
                     NewSocket->SetReceiveBufferSize(ReceiveBufferSize, OutNewSize);
                 }
                 
-                if (SendBufferSize > 0)
+                if (SendBufferSize > 0) //-V547
                 {
                     NewSocket->SetSendBufferSize(SendBufferSize, OutNewSize);
                 }
diff --git a/Engine/Plugins/Online/OnlineSubsystemUtils/Source/OnlineSubsystemUtils/Private/PartyBeaconClient.cpp b/Engine/Plugins/Online/OnlineSubsystemUtils/Source/OnlineSubsystemUtils/Private/PartyBeaconClient.cpp
--- a/Engine/Plugins/Online/OnlineSubsystemUtils/Source/OnlineSubsystemUtils/Private/PartyBeaconClient.cpp
+++ b/Engine/Plugins/Online/OnlineSubsystemUtils/Source/OnlineSubsystemUtils/Private/PartyBeaconClient.cpp
@@ -209,7 +209,7 @@ bool APartyBeaconClient::RequestReservationUpdate(const FString& ConnectInfoStr,
 				RequestType = EClientRequestType::ReservationUpdate;
 			}
 		}
-		else if (MyConnectionState == EBeaconConnectionState::Open)
+		else
 		{
 			bWasStarted = RequestReservationUpdate(RequestingPartyLeader, PlayersToAdd);
 		}
@@ -237,7 +237,7 @@ bool APartyBeaconClient::RequestReservationUpdate(const FOnlineSessionSearchResu
 			RequestType = EClientRequestType::ReservationUpdate;
 		}
 	}
-	else if (MyConnectionState == EBeaconConnectionState::Open)
+	else
 	{
 		RequestReservationUpdate(RequestingPartyLeader, PlayersToAdd);
 	}
diff --git a/Engine/Plugins/Online/OnlineSubsystemUtils/Source/OnlineSubsystemUtils/Private/VoiceInterfaceImpl.cpp b/Engine/Plugins/Online/OnlineSubsystemUtils/Source/OnlineSubsystemUtils/Private/VoiceInterfaceImpl.cpp
--- a/Engine/Plugins/Online/OnlineSubsystemUtils/Source/OnlineSubsystemUtils/Private/VoiceInterfaceImpl.cpp
+++ b/Engine/Plugins/Online/OnlineSubsystemUtils/Source/OnlineSubsystemUtils/Private/VoiceInterfaceImpl.cpp
@@ -541,7 +541,7 @@ bool FOnlineVoiceImpl::UnmuteRemoteTalker(uint8 LocalUserNum, const FUniqueNetId
 		UE_LOG(LogVoice, Warning, TEXT("Invalid user specified in UnmuteRemoteTalker(%d)"), LocalUserNum);
 	}
 
-	return Return == S_OK;
+	return Return == S_OK; //-V547
 }
 
 void FOnlineVoiceImpl::ProcessMuteChangeNotification()
diff --git a/Engine/Plugins/Runtime/GameplayAbilities/Source/GameplayAbilities/Private/GameplayTagResponseTable.cpp b/Engine/Plugins/Runtime/GameplayAbilities/Source/GameplayAbilities/Private/GameplayTagResponseTable.cpp
--- a/Engine/Plugins/Runtime/GameplayAbilities/Source/GameplayAbilities/Private/GameplayTagResponseTable.cpp
+++ b/Engine/Plugins/Runtime/GameplayAbilities/Source/GameplayAbilities/Private/GameplayTagResponseTable.cpp
@@ -125,7 +125,7 @@ void UGameplayTagReponseTable::TagResponseEvent(const FGameplayTag Tag, int32 Ne
 		Remove(ASC, Info.NegativeHandles);
 		AddOrUpdate(ASC, Entry.Positive.ResponseGameplayEffects, TotalCount, Info.PositiveHandles);
 	}
-	else if (TotalCount == 0)
+	else
 	{
 		Remove(ASC, Info.PositiveHandles);
 		Remove(ASC, Info.NegativeHandles);
diff --git a/Engine/Plugins/Runtime/GoogleVR/GoogleVRHMD/Source/GoogleVRHMD/Private/GoogleVRHMD.cpp b/Engine/Plugins/Runtime/GoogleVR/GoogleVRHMD/Source/GoogleVRHMD/Private/GoogleVRHMD.cpp
--- a/Engine/Plugins/Runtime/GoogleVR/GoogleVRHMD/Source/GoogleVRHMD/Private/GoogleVRHMD.cpp
+++ b/Engine/Plugins/Runtime/GoogleVR/GoogleVRHMD/Source/GoogleVRHMD/Private/GoogleVRHMD.cpp
@@ -930,8 +930,7 @@ void FGoogleVRHMD::SetNumOfDistortionPoints(int32 XPoints, int32 YPoints)
 #if !GOOGLEVRHMD_SUPPORTED_PLATFORMS
 	XPoints = 40;
 	YPoints = 40;
-#endif
-
+#else
 	// clamp values
 	if (XPoints < 2)
 	{
@@ -950,6 +949,7 @@ void FGoogleVRHMD::SetNumOfDistortionPoints(int32 XPoints, int32 YPoints)
 	{
 		YPoints = 200;
 	}
+#endif
 
 	// calculate our values
 	DistortionPointsX = XPoints;
diff --git a/Engine/Source/Runtime/Core/Public/Windows/WindowsPlatformCodeAnalysis.h b/Engine/Source/Runtime/Core/Public/Windows/WindowsPlatformCodeAnalysis.h
--- a/Engine/Source/Runtime/Core/Public/Windows/WindowsPlatformCodeAnalysis.h
+++ b/Engine/Source/Runtime/Core/Public/Windows/WindowsPlatformCodeAnalysis.h
@@ -95,7 +95,7 @@
 	#pragma warning(disable : 6320) // warning C6320: Exception-filter expression is the constant EXCEPTION_EXECUTE_HANDLER. This might mask exceptions that were not intended to be handled.
 
 	//PVS-Studio settings:
-	//-V::505,542,581,601,668,677,686,688,690,703,704,711,719,720,728,730,735,751
+	//-V::505,542,581,601,668,677,686,688,690,703,704,711,719,720,728,730,735,751,1002
 	//-V:TRYCOMPRESSION:519,547
 	//-V:check(:501,547,560,605
 	//-V:checkf(:510
diff --git a/Engine/Source/Runtime/CoreUObject/Private/Serialization/AsyncLoading.cpp b/Engine/Source/Runtime/CoreUObject/Private/Serialization/AsyncLoading.cpp
--- a/Engine/Source/Runtime/CoreUObject/Private/Serialization/AsyncLoading.cpp
+++ b/Engine/Source/Runtime/CoreUObject/Private/Serialization/AsyncLoading.cpp
@@ -5542,7 +5542,7 @@ EAsyncPackageState::Type FAsyncPackage::TickAsyncPackage(bool InbUseTimeLimit, b
 		if (!GEventDrivenLoaderEnabled)
 		{
 			// Create raw linker. Needs to be async created via ticking before it can be used.
-			if (LoadingState == EAsyncPackageState::Complete)
+			if (LoadingState == EAsyncPackageState::Complete) //-V547
 			{
 				SCOPED_LOADTIMER(Package_CreateLinker);
 				LoadingState = CreateLinker();
diff --git a/Engine/Source/Runtime/CoreUObject/Private/UObject/Property.cpp b/Engine/Source/Runtime/CoreUObject/Private/UObject/Property.cpp
--- a/Engine/Source/Runtime/CoreUObject/Private/UObject/Property.cpp
+++ b/Engine/Source/Runtime/CoreUObject/Private/UObject/Property.cpp
@@ -1138,7 +1138,7 @@ const TCHAR* UProperty::ImportSingleProperty( const TCHAR* Str, void* DestData,
 						Str = Result;
 					}
 				}
-				else if(ArrayOp == ADO_Remove)
+				else
 				{
 					int32 Size = ArrayProperty->Inner->ElementSize;
 
diff --git a/Engine/Source/Runtime/MovieSceneCapture/Private/FrameGrabber.cpp b/Engine/Source/Runtime/MovieSceneCapture/Private/FrameGrabber.cpp
--- a/Engine/Source/Runtime/MovieSceneCapture/Private/FrameGrabber.cpp
+++ b/Engine/Source/Runtime/MovieSceneCapture/Private/FrameGrabber.cpp
@@ -90,8 +90,8 @@ void FViewportSurfaceReader::ResolveRenderTarget(const FViewportRHIRef& Viewport
 		// @todo: JIRA UE-41879 and UE-43829. If any of these ensures go off, something has overwritten the memory for this render command (buffer underflow/overflow?)
 		bool bMemoryTrample = !ensureMsgf(RendererModule, TEXT("RendererModule has become null. This indicates a memory trample.")) ||
 			!ensureMsgf(RendererModule == RendererModuleDebug, TEXT("RendererModule and RendererModuleDebug are not equal (0x%016x != 0x%016x). This indicates a memory trample."), (void*)RendererModule, (void*)RendererModuleDebug) ||
-			!ensureMsgf(MemoryGuard1 == 0xaffec7ed, TEXT("Memory guard 1 is now 0x%08x, expected 0xaffec7ed."), MemoryGuard1) ||
-			!ensureMsgf(MemoryGuard2 == 0xaffec7ed, TEXT("Memory guard 2 is now 0x%08x, expected 0xaffec7ed."), MemoryGuard2);
+			!ensureMsgf(MemoryGuard1 == 0xaffec7ed, TEXT("Memory guard 1 is now 0x%08x, expected 0xaffec7ed."), MemoryGuard1) || //-V547
+			!ensureMsgf(MemoryGuard2 == 0xaffec7ed, TEXT("Memory guard 2 is now 0x%08x, expected 0xaffec7ed."), MemoryGuard2); //-V547
 
 		if (bMemoryTrample)
 		{
diff --git a/Engine/Source/Runtime/OpenGLDrv/Private/OpenGLShaders.cpp b/Engine/Source/Runtime/OpenGLDrv/Private/OpenGLShaders.cpp
--- a/Engine/Source/Runtime/OpenGLDrv/Private/OpenGLShaders.cpp
+++ b/Engine/Source/Runtime/OpenGLDrv/Private/OpenGLShaders.cpp
@@ -602,7 +602,7 @@ ShaderType* CompileOpenGLShader(const TArray<uint8>& InShaderCode)
 			CompileStatus = CompileCurrentShader(Resource, GlslCode);
 		}
 
-		if ( CompileStatus == GL_TRUE )
+		if ( CompileStatus == GL_TRUE ) //-V547
 		{
 			if (Capabilities.bSupportsSeparateShaderObjects)
 			{
diff --git a/Engine/Source/Runtime/Renderer/Private/GPUFastFourierTransform.cpp b/Engine/Source/Runtime/Renderer/Private/GPUFastFourierTransform.cpp
--- a/Engine/Source/Runtime/Renderer/Private/GPUFastFourierTransform.cpp
+++ b/Engine/Source/Runtime/Renderer/Private/GPUFastFourierTransform.cpp
@@ -1293,8 +1293,8 @@ namespace
 		const uint32 YThreadCount = FCopyWindowCS::YThreadCount();
 
 		// Number of thread groups
-		const uint32 XGroups = ( DstExtent.X / XThreadCount ) + ( (DstExtent.X % XThreadCount == 0) ? 0 : 1 );
-		const uint32 YGroups = ( DstExtent.Y / YThreadCount ) + ( (DstExtent.Y % YThreadCount == 0) ? 0 : 1 );
+		const uint32 XGroups = ( DstExtent.X / XThreadCount ) + ( (DstExtent.X % XThreadCount == 0) ? 0 : 1 ); //-V547
+		const uint32 YGroups = ( DstExtent.Y / YThreadCount ) + ( (DstExtent.Y % YThreadCount == 0) ? 0 : 1 ); //-V547
 
 		const FGPUFFTShaderContext::ShaderMapType& ShaderMap = Context.GetShaderMap();
 		FRHICommandListImmediate& RHICmdList = Context.GetRHICmdList();
diff --git a/Engine/Source/Runtime/Renderer/Private/PostProcess/RenderTargetPool.cpp b/Engine/Source/Runtime/Renderer/Private/PostProcess/RenderTargetPool.cpp
--- a/Engine/Source/Runtime/Renderer/Private/PostProcess/RenderTargetPool.cpp
+++ b/Engine/Source/Runtime/Renderer/Private/PostProcess/RenderTargetPool.cpp
@@ -358,9 +358,9 @@ bool FRenderTargetPool::FindFreeElement(FRHICommandList& RHICmdList, const FPool
 		if (AliasingMode == 0)
 		{
 			if ((Desc.Flags & TexCreate_FastVRAM) && FPlatformMemory::SupportsFastVRAMMemory() )
-		{
-			PassCount = 2;
-		}
+			{
+				PassCount = 2;
+			}
 		}
 
 		bool bAllowMultipleDiscards = ( CVarAllowMultipleAliasingDiscardsPerFrame.GetValueOnRenderThread() != 0 );
@@ -368,7 +368,7 @@ bool FRenderTargetPool::FindFreeElement(FRHICommandList& RHICmdList, const FPool
 		// (easily we can run out of VRam, if this search becomes a performance problem we can optimize or we should use less TexCreate_FastVRAM)
 		for(uint32 Pass = 0; Pass < PassCount; ++Pass)
 		{
-			bool bExactMatch = (Pass == 0);
+			bool bExactMatch = (Pass == 0); //-V547
     
 			for(uint32 i = 0, Num = (uint32)PooledRenderTargets.Num(); i < Num; ++i)
 			{
diff --git a/Engine/Source/Runtime/Renderer/Private/ShadowSetup.cpp b/Engine/Source/Runtime/Renderer/Private/ShadowSetup.cpp
--- a/Engine/Source/Runtime/Renderer/Private/ShadowSetup.cpp
+++ b/Engine/Source/Runtime/Renderer/Private/ShadowSetup.cpp
@@ -1174,23 +1174,23 @@ void FProjectedShadowInfo::GatherDynamicMeshElements(FSceneRenderer& Renderer, F
 		if (IsWholeSceneDirectionalShadow())
 		{
 			ShadowDepthView->SetPreShadowTranslation(FVector(0, 0, 0));
-			ShadowDepthView->SetDynamicMeshElementsShadowCullFrustum((Disable & 1) ? &NoCull : &CascadeSettings.ShadowBoundsAccurate);
+			ShadowDepthView->SetDynamicMeshElementsShadowCullFrustum((Disable & 1) ? &NoCull : &CascadeSettings.ShadowBoundsAccurate); //-V547
 			GatherDynamicMeshElementsArray(ShadowDepthView, Renderer, DynamicSubjectPrimitives, DynamicSubjectMeshElements, ReusedViewsArray);
 			ShadowDepthView->SetPreShadowTranslation(PreShadowTranslation);
 		}
 		else
 		{
 			ShadowDepthView->SetPreShadowTranslation(PreShadowTranslation);
-			ShadowDepthView->SetDynamicMeshElementsShadowCullFrustum((Disable & 1) ? &NoCull : &CasterFrustum);
+			ShadowDepthView->SetDynamicMeshElementsShadowCullFrustum((Disable & 1) ? &NoCull : &CasterFrustum); //-V547
 			GatherDynamicMeshElementsArray(ShadowDepthView, Renderer, DynamicSubjectPrimitives, DynamicSubjectMeshElements, ReusedViewsArray);
 		}
 
 		ShadowDepthView->DrawDynamicFlags = EDrawDynamicFlags::None;
 
-		ShadowDepthView->SetDynamicMeshElementsShadowCullFrustum((Disable & 2) ? &NoCull : &ReceiverFrustum);
+		ShadowDepthView->SetDynamicMeshElementsShadowCullFrustum((Disable & 2) ? &NoCull : &ReceiverFrustum); //-V547
 		GatherDynamicMeshElementsArray(ShadowDepthView, Renderer, ReceiverPrimitives, DynamicReceiverMeshElements, ReusedViewsArray);
 
-		ShadowDepthView->SetDynamicMeshElementsShadowCullFrustum((Disable & 4) ? &NoCull : &CasterFrustum);
+		ShadowDepthView->SetDynamicMeshElementsShadowCullFrustum((Disable & 4) ? &NoCull : &CasterFrustum); //-V547
 		GatherDynamicMeshElementsArray(ShadowDepthView, Renderer, SubjectTranslucentPrimitives, DynamicSubjectTranslucentMeshElements, ReusedViewsArray);
 
 		Renderer.MeshCollector.ProcessTasks();
