Backed out a change from CL 3320581:

Rendering: Use 'Volume' indirect lighting cache instead of 'Point' for unbuilt objects (higher quality)
- This allows static objects that are moved around in the editor to look much nicer

This was breaking integration testing by causing placed BSP or static meshes to be lit incorrectly when Volumetric Lightmap was enabled. We can revisit this in the future if we wish to reproduce the 'nicer' look that this was supposed to implement.

diff --git a/Engine/Source/Runtime/Renderer/Private/IndirectLightingCache.cpp b/Engine/Source/Runtime/Renderer/Private/IndirectLightingCache.cpp
--- a/Engine/Source/Runtime/Renderer/Private/IndirectLightingCache.cpp
+++ b/Engine/Source/Runtime/Renderer/Private/IndirectLightingCache.cpp
@@ -326,7 +326,7 @@ void FIndirectLightingCache::CalculateBlockScaleAndAdd(FIntVector InTexelMin, in
 
 FIndirectLightingCacheAllocation* FIndirectLightingCache::AllocatePrimitive(const FPrimitiveSceneInfo* PrimitiveSceneInfo, bool bUnbuiltPreview)
 {
-	const bool bPointSample = !bUnbuiltPreview && PrimitiveSceneInfo->Proxy->GetIndirectLightingCacheQuality() == ILCQ_Point;
+	const bool bPointSample = PrimitiveSceneInfo->Proxy->GetIndirectLightingCacheQuality() == ILCQ_Point || bUnbuiltPreview;
 	const int32 BlockSize = bPointSample ? 1 : GLightingCacheMovableObjectAllocationSize;
 	return PrimitiveAllocations.Add(PrimitiveSceneInfo->PrimitiveComponentId, CreateAllocation(BlockSize, PrimitiveSceneInfo->Proxy->GetBounds(), bPointSample, bUnbuiltPreview));
 }
