duplicate 4087965 and 4088060.
Fixes for PS4 auto test failures

diff --git a/Engine/Source/Runtime/Renderer/Private/PostProcess/DiaphragmDOFPasses.cpp b/Engine/Source/Runtime/Renderer/Private/PostProcess/DiaphragmDOFPasses.cpp
--- a/Engine/Source/Runtime/Renderer/Private/PostProcess/DiaphragmDOFPasses.cpp
+++ b/Engine/Source/Runtime/Renderer/Private/PostProcess/DiaphragmDOFPasses.cpp
@@ -911,16 +911,13 @@ void FRCPassDiaphragmDOFReduce::Process(FRenderingCompositePassContext& Context)
 		EResourceTransitionPipeline::EComputeToCompute,
 		DrawIndirectParametersBuffer->UAV);
 
-	if (Params.bExtractForegroundHybridScattering)
+	if (bDoAnyHybridScatteringExtraction)
 	{
 		Context.RHICmdList.TransitionResource(
 			EResourceTransitionAccess::ERWNoBarrier,
 			EResourceTransitionPipeline::EComputeToCompute,
 			ScatterDrawListBuffer[0]->UAV);
-	}
 	
-	if (Params.bExtractBackgroundHybridScattering)
-	{
 		Context.RHICmdList.TransitionResource(
 			EResourceTransitionAccess::ERWNoBarrier,
 			EResourceTransitionPipeline::EComputeToCompute,
diff --git a/Engine/Source/Runtime/Renderer/Private/PostProcess/PostProcessTemporalAA.cpp b/Engine/Source/Runtime/Renderer/Private/PostProcess/PostProcessTemporalAA.cpp
--- a/Engine/Source/Runtime/Renderer/Private/PostProcess/PostProcessTemporalAA.cpp
+++ b/Engine/Source/Runtime/Renderer/Private/PostProcess/PostProcessTemporalAA.cpp
@@ -511,7 +511,7 @@ void DrawPixelPassTemplate(
 	GraphicsPSOInit.PrimitiveType = PT_TriangleList;
 
 	SetGraphicsPipelineState(Context.RHICmdList, GraphicsPSOInit);
-
+	VertexShader->TransitionResources(Context);
 	VertexShader->SetVS(Context);
 	PixelShader->SetParameters(Context.RHICmdList, Context, InputHistory, PassParameters, bUseDither, SrcSize);
 
