ML - Fix case where StaticshadowMap data is still enqueued for updating while running Lightmass but rendering is happening at the same time.

diff --git a/Engine/Source/Runtime/Renderer/Private/VolumeLighting.h b/Engine/Source/Runtime/Renderer/Private/VolumeLighting.h
--- a/Engine/Source/Runtime/Renderer/Private/VolumeLighting.h
+++ b/Engine/Source/Runtime/Renderer/Private/VolumeLighting.h
@@ -121,7 +121,7 @@ public:
 		OnePassShadowParameters.Set(RHICmdList, ShaderRHI, bDynamicallyShadowed && LightType == LightType_Point ? ShadowMap : NULL);
 
 		const FStaticShadowDepthMap* StaticShadowDepthMap = LightSceneInfo->Proxy->GetStaticShadowDepthMap();
-		const uint32 bStaticallyShadowedValue = LightSceneInfo->IsPrecomputedLightingValid() && StaticShadowDepthMap && StaticShadowDepthMap->TextureRHI ? 1 : 0;
+		const uint32 bStaticallyShadowedValue = LightSceneInfo->IsPrecomputedLightingValid() && StaticShadowDepthMap && StaticShadowDepthMap->Data && StaticShadowDepthMap->TextureRHI ? 1 : 0;
 		FTextureRHIParamRef StaticShadowDepthMapTexture = bStaticallyShadowedValue ? StaticShadowDepthMap->TextureRHI : GWhiteTexture->TextureRHI;
 		const FMatrix WorldToStaticShadow = bStaticallyShadowedValue ? StaticShadowDepthMap->Data->WorldToLight : FMatrix::Identity;
 		const FVector4 StaticShadowBufferSizeValue = bStaticallyShadowedValue ? FVector4(StaticShadowDepthMap->Data->ShadowMapSizeX, StaticShadowDepthMap->Data->ShadowMapSizeY, 1.0f / StaticShadowDepthMap->Data->ShadowMapSizeX, 1.0f / StaticShadowDepthMap->Data->ShadowMapSizeY) : FVector4(0, 0, 0, 0);
