Duplicate 4098665.  Fix d3ddebug error on launch

diff --git a/Engine/Source/Runtime/Renderer/Private/LightGridInjection.cpp b/Engine/Source/Runtime/Renderer/Private/LightGridInjection.cpp
--- a/Engine/Source/Runtime/Renderer/Private/LightGridInjection.cpp
+++ b/Engine/Source/Runtime/Renderer/Private/LightGridInjection.cpp
@@ -31,12 +31,6 @@ LightGridInjection.cpp
 // Workaround for platforms that don't support 16bit integers
 #define	CHANGE_LIGHTINDEXTYPE_SIZE	(PLATFORM_MAC || PLATFORM_IOS) 
 
-#if PLATFORM_MAC
-#define LOCAL_LIGHT_BUFFER_FORMAT PF_A32B32G32R32F
-#else
-#define LOCAL_LIGHT_BUFFER_FORMAT PF_R32G32B32A32_UINT
-#endif
-
 int32 GLightGridPixelSize = 64;
 FAutoConsoleVariableRef CVarLightGridPixelSize(
 	TEXT("r.Forward.LightGridPixelSize"),
@@ -651,7 +645,7 @@ void FDeferredShadingSceneRenderer::ComputeLightGrid(FRHICommandListImmediate& R
 				if (View.ForwardLightingResources->ForwardLocalLightBuffer.NumBytes < NumBytesRequired)
 				{
 					View.ForwardLightingResources->ForwardLocalLightBuffer.Release();
-					View.ForwardLightingResources->ForwardLocalLightBuffer.Initialize(sizeof(FVector4), NumBytesRequired / sizeof(FVector4), LOCAL_LIGHT_BUFFER_FORMAT, BUF_Volatile);
+					View.ForwardLightingResources->ForwardLocalLightBuffer.Initialize(sizeof(FVector4), NumBytesRequired / sizeof(FVector4), PF_A32B32G32R32F, BUF_Volatile);
 				}
 
 				ForwardLightData.ForwardLocalLightBuffer = View.ForwardLightingResources->ForwardLocalLightBuffer.SRV;
@@ -878,4 +872,3 @@ void FDeferredShadingSceneRenderer::RenderForwardShadingShadowProjections(FRHICo
 }
 
 #undef CHANGE_LIGHTINDEXTYPE_SIZE
-#undef LOCAL_LIGHT_BUFFER_FORMAT
