Clone of CL 4077664:  Fix for Mac crash from crash report: [CrashReport] UE4Editor-Renderer.dylib!void FForwardLightingParameters::Set<FRHICommandList, FRHIPixelShader*>().  Reflection uniform buffer is NULL when shader parameter is bound - fails checkSlow(!Parameter.IsBound() || UniformBufferRHI); in SetUniformBufferParameter().

diff --git a/Engine/Source/Runtime/Renderer/Private/Renderer.cpp b/Engine/Source/Runtime/Renderer/Private/Renderer.cpp
--- a/Engine/Source/Runtime/Renderer/Private/Renderer.cpp
+++ b/Engine/Source/Runtime/Renderer/Private/Renderer.cpp
@@ -37,6 +37,22 @@ IMPLEMENT_MODULE(FRendererModule, Renderer);
 	FSystemSettings* GSystemSettingsForVisualizers = &GSystemSettings;
 #endif
 
+// Dummy Reflection capture uniform buffer for translucent tile mesh rendering
+class FDummyReflectionCaptureUniformBuffer : public TUniformBuffer<FReflectionCaptureShaderData>
+{
+	typedef TUniformBuffer< FReflectionCaptureShaderData > Super;
+	
+public:
+	virtual void InitDynamicRHI() override
+	{
+		FReflectionCaptureShaderData DummyPositionsBuffer;
+		FMemory::Memzero(DummyPositionsBuffer);
+		SetContentsNoUpdate(DummyPositionsBuffer);
+		Super::InitDynamicRHI();
+	}
+};
+static TGlobalResource< FDummyReflectionCaptureUniformBuffer > GDummyReflectionCaptureUniformBuffer;
+
 void FRendererModule::ReallocateSceneRenderTargets()
 {
 	FLightPrimitiveInteraction::InitializeMemoryPool();
@@ -91,6 +107,12 @@ void FRendererModule::DrawTileMesh(FRHICommandListImmediate& RHICmdList, FDrawin
 		{
 			if (FeatureLevel >= ERHIFeatureLevel::SM4)
 			{
+				// Crash fix - reflection capture shader parameter is bound but we have no buffer during Build Texture Streaming
+				if(!View.ReflectionCaptureUniformBuffer.IsValid())
+				{
+					View.ReflectionCaptureUniformBuffer = GDummyReflectionCaptureUniformBuffer;
+				}
+			
 				TUniformBufferRef<FTranslucentBasePassUniformParameters> BasePassUniformBuffer;
 				CreateTranslucentBasePassUniformBuffer(RHICmdList, View, nullptr, ESceneTextureSetupMode::None, BasePassUniformBuffer);
 				DrawRenderState.SetPassUniformBuffer(BasePassUniformBuffer);
