Cherry-pick 4073044: PR #4681: UE-58051: Scene Capture 2D: additive mode not working in 4.19 (Contributed by [REDACTED_PERSON])

diff --git a/Engine/Source/Runtime/Renderer/Private/SceneCaptureRendering.cpp b/Engine/Source/Runtime/Renderer/Private/SceneCaptureRendering.cpp
--- a/Engine/Source/Runtime/Renderer/Private/SceneCaptureRendering.cpp
+++ b/Engine/Source/Runtime/Renderer/Private/SceneCaptureRendering.cpp
@@ -323,9 +323,8 @@ static void UpdateSceneCaptureContentDeferred_RenderThread(
 
 		// TODO: Could avoid the clear by replacing with dummy black system texture.
 		FViewInfo& View = SceneRenderer->Views[0];
-		FIntRect ViewRect = View.ViewRect;
 		SetRenderTarget(RHICmdList, Target->GetRenderTargetTexture(), nullptr, true);
-		DrawClearQuad(RHICmdList, true, FLinearColor::Black, false, 0, false, 0, Target->GetSizeXY(), ViewRect);
+		DrawClearQuad(RHICmdList, true, FLinearColor::Black, false, 0, false, 0, Target->GetSizeXY(), View.UnscaledViewRect);
 
 		// Render the scene normally
 		{
