Cherry-pick 4073050: Whitelists vulkan for DOF's R11G11B10 optimisation and increase to 4 gathering on consoles.

diff --git a/Engine/Source/Runtime/Renderer/Private/PostProcess/DiaphragmDOFPasses.h b/Engine/Source/Runtime/Renderer/Private/PostProcess/DiaphragmDOFPasses.h
--- a/Engine/Source/Runtime/Renderer/Private/PostProcess/DiaphragmDOFPasses.h
+++ b/Engine/Source/Runtime/Renderer/Private/PostProcess/DiaphragmDOFPasses.h
@@ -320,7 +320,7 @@ public:
 		{
 			return 5;
 		}
-		return 3;
+		return 4;
 	}
 
 	/** Returns whether the shader for bokeh simulation are compiled. */
@@ -340,7 +340,7 @@ public:
 		}
 
 		// There is high number of UAV to write in reduce pass.
-		return ShaderPlatform == SP_PS4 || ShaderPlatform == SP_XBOXONE_D3D12;
+		return ShaderPlatform == SP_PS4 || ShaderPlatform == SP_XBOXONE_D3D12 || ShaderPlatform == SP_VULKAN_SM5;
 	}
 
 	/** Quality configurations for gathering passes. */
