Cherry-pick 4102327: Fixes a bug in FTAAPassParameters::TopLeftCornerViewRects() that was mistakenly increasing input and output view sizes, causing DOF buffers misalignement for split screen & VR.

diff --git a/Engine/Source/Runtime/Renderer/Private/PostProcess/PostProcessTemporalAA.h b/Engine/Source/Runtime/Renderer/Private/PostProcess/PostProcessTemporalAA.h
--- a/Engine/Source/Runtime/Renderer/Private/PostProcess/PostProcessTemporalAA.h
+++ b/Engine/Source/Runtime/Renderer/Private/PostProcess/PostProcessTemporalAA.h
@@ -101,8 +101,10 @@ struct FTAAPassParameters
 	// Shifts input and output view rect to top left corner
 	FORCEINLINE void TopLeftCornerViewRects()
 	{
-		InputViewRect -= InputViewRect.Min;
-		OutputViewRect -= OutputViewRect.Min;
+		InputViewRect.Max -= InputViewRect.Min;
+		InputViewRect.Min = FIntPoint::ZeroValue;
+		OutputViewRect.Max -= OutputViewRect.Min;
+		OutputViewRect.Min = FIntPoint::ZeroValue;
 	}
 };
 
