Mirroring CL 4098493 from Dev-VR

Not forcing stereo scene capture target to resize (only intended for the main stereo scene buffer, as the depth buffer needs to match the color buffer on some platforms).

diff --git a/Engine/Source/Runtime/Renderer/Private/PostProcess/SceneRenderTargets.cpp b/Engine/Source/Runtime/Renderer/Private/PostProcess/SceneRenderTargets.cpp
--- a/Engine/Source/Runtime/Renderer/Private/PostProcess/SceneRenderTargets.cpp
+++ b/Engine/Source/Runtime/Renderer/Private/PostProcess/SceneRenderTargets.cpp
@@ -304,7 +304,7 @@ FIntPoint FSceneRenderTargets::ComputeDesiredSize(const FSceneViewFamily& ViewFa
 		bIsVRScene |= View->StereoPass != EStereoscopicPass::eSSP_FULL;
 	}
 
-	if(!FPlatformProperties::SupportsWindowedMode() || bIsVRScene)
+	if(!FPlatformProperties::SupportsWindowedMode() || (bIsVRScene && !bIsSceneCapture))
 	{
 		// Force ScreenRes on non windowed platforms.
 		SceneTargetsSizingMethod = RequestedSize;
