Mirroring CL 4103694 from Dev-VR

Making it so the HMD post-process pass properly fills out its RenderTargetDesc (so it doesn't rely on being the final output to have it set).

diff --git a/Engine/Source/Runtime/Renderer/Private/PostProcess/PostProcessHMD.cpp b/Engine/Source/Runtime/Renderer/Private/PostProcess/PostProcessHMD.cpp
--- a/Engine/Source/Runtime/Renderer/Private/PostProcess/PostProcessHMD.cpp
+++ b/Engine/Source/Runtime/Renderer/Private/PostProcess/PostProcessHMD.cpp
@@ -172,7 +172,7 @@ void FRCPassPostProcessHMD::Process(FRenderingCompositePassContext& Context)
 	const FSceneRenderTargetItem& DestRenderTarget = PassOutputs[0].RequestSurface(Context);
 
 	// Set the view family's render target/viewport
-	if (DestRenderTarget.TargetableTexture->GetClearColor() == FLinearColor::Black)
+	if (!ensure(DestRenderTarget.TargetableTexture.IsValid()) || DestRenderTarget.TargetableTexture->GetClearColor() == FLinearColor::Black)
 	{
 		FRHIRenderTargetView RtView = FRHIRenderTargetView(DestRenderTarget.TargetableTexture, ERenderTargetLoadAction::EClear);
 		FRHISetRenderTargetsInfo Info(1, &RtView, FRHIDepthRenderTargetView());
@@ -220,6 +220,14 @@ FPooledRenderTargetDesc FRCPassPostProcessHMD::ComputeOutputDesc(EPassOutputId I
 {
 	FPooledRenderTargetDesc Ret = PassOutputs[0].RenderTargetDesc;
 
+	if (const FRenderingCompositeOutputRef* InputPass0 = GetInput(ePId_Input0))
+	{
+		if (const FRenderingCompositeOutput* PassInput = InputPass0->GetOutput())
+		{
+			Ret = PassInput->RenderTargetDesc;
+		}
+	}
+
 	Ret.Reset();
 	Ret.DebugName = TEXT("HMD");
 
