Fixed: ARPG crashes on Vulkan PC and Android

diff --git a/Engine/Source/Runtime/Renderer/Private/MobileShadingRenderer.cpp b/Engine/Source/Runtime/Renderer/Private/MobileShadingRenderer.cpp
--- a/Engine/Source/Runtime/Renderer/Private/MobileShadingRenderer.cpp
+++ b/Engine/Source/Runtime/Renderer/Private/MobileShadingRenderer.cpp
@@ -516,13 +516,13 @@ void FMobileSceneRenderer::ConditionalResolveSceneDepth(FRHICommandListImmediate
 
 	if ((IsMobileHDR() || IsHTML5Platform())
 		&& IsMobilePlatform(ShaderPlatform) 
+		&& !IsVulkanPlatform(ShaderPlatform)
+		&& !IsMetalPlatform(ShaderPlatform)
 		&& !IsPCPlatform(ShaderPlatform) // exclude mobile emulation on PC
 		&& !View.bIsPlanarReflection)	// exclude depth resolve from planar reflection captures, can't do it reliably more than once per frame
 	{
 		bool bSceneDepthInAlpha = (SceneContext.GetSceneColor()->GetDesc().Format == PF_FloatRGBA);
-		bool bOnChipDepthFetch = 
-			(GSupportsShaderDepthStencilFetch || 
-			(GSupportsShaderFramebufferFetch && (bSceneDepthInAlpha || IsMetalPlatform(ShaderPlatform))));
+		bool bOnChipDepthFetch = (GSupportsShaderDepthStencilFetch || (GSupportsShaderFramebufferFetch && bSceneDepthInAlpha));
 		
 		const bool bAlwaysResolveDepth = CVarMobileAlwaysResolveDepth.GetValueOnRenderThread() == 1;
 
diff --git a/Engine/Source/Runtime/Renderer/Private/SceneOcclusion.cpp b/Engine/Source/Runtime/Renderer/Private/SceneOcclusion.cpp
--- a/Engine/Source/Runtime/Renderer/Private/SceneOcclusion.cpp
+++ b/Engine/Source/Runtime/Renderer/Private/SceneOcclusion.cpp
@@ -1447,15 +1447,19 @@ void FSceneRenderer::BeginOcclusionTests(FRHICommandListImmediate& RHICmdList, b
 				}
 			}
 
-			FRHIRenderPassInfo RPInfo(
-				bUseDownsampledDepth ? SceneContext.GetSmallDepthSurface() : SceneContext.GetSceneDepthSurface(),
-				NumQueriesForBatch,
-				EDepthStencilTargetActions::LoadDepthStencil_StoreStencilNotDepth,
-				nullptr,
-				FExclusiveDepthStencil::DepthRead_StencilWrite
-			);
+			// On mobile occlusion queries are done in base pass
+			if (FeatureLevel > ERHIFeatureLevel::ES3_1)
+			{
+				FRHIRenderPassInfo RPInfo(
+					bUseDownsampledDepth ? SceneContext.GetSmallDepthSurface() : SceneContext.GetSceneDepthSurface(),
+					NumQueriesForBatch,
+					EDepthStencilTargetActions::LoadDepthStencil_StoreStencilNotDepth,
+					nullptr,
+					FExclusiveDepthStencil::DepthRead_StencilWrite
+				);
 
-			RHICmdList.BeginRenderPass(RPInfo, TEXT("OcclusionQueries"));
+				RHICmdList.BeginRenderPass(RPInfo, TEXT("OcclusionQueries"));
+			}
 
 			FGraphicsPipelineStateInitializer GraphicsPSOInit;
 			RHICmdList.ApplyCachedRenderTargets(GraphicsPSOInit);
@@ -1544,7 +1548,11 @@ void FSceneRenderer::BeginOcclusionTests(FRHICommandListImmediate& RHICmdList, b
 				}
 			}
 			
-			RHICmdList.EndRenderPass();
+			// On mobile occlusion queries are done in base pass
+			if (FeatureLevel > ERHIFeatureLevel::ES3_1)
+			{
+				RHICmdList.EndRenderPass();
+			}
 			
 			if (bUseDownsampledDepth)
 			{
diff --git a/Engine/Source/Runtime/VulkanRHI/Private/VulkanMemory.cpp b/Engine/Source/Runtime/VulkanRHI/Private/VulkanMemory.cpp
--- a/Engine/Source/Runtime/VulkanRHI/Private/VulkanMemory.cpp
+++ b/Engine/Source/Runtime/VulkanRHI/Private/VulkanMemory.cpp
@@ -1613,6 +1613,14 @@ namespace VulkanRHI
 
 		{
 			FScopeLock ScopeLock(&CS);
+
+#if VULKAN_HAS_DEBUGGING_ENABLED
+			FEntry* ExistingEntry = Entries.FindByPredicate([&](const FEntry& InEntry) { 
+				return InEntry.Handle == Entry.Handle; 
+			});
+			checkf(ExistingEntry == nullptr, TEXT("Attempt to double-delete resource, Type: %d, Handle: %llu"), (int32)Type, Handle);
+#endif
+
 			Entries.Add(Entry);
 		}
 	}
diff --git a/Engine/Source/Runtime/VulkanRHI/Private/VulkanViewport.cpp b/Engine/Source/Runtime/VulkanRHI/Private/VulkanViewport.cpp
--- a/Engine/Source/Runtime/VulkanRHI/Private/VulkanViewport.cpp
+++ b/Engine/Source/Runtime/VulkanRHI/Private/VulkanViewport.cpp
@@ -339,7 +339,7 @@ void FVulkanFramebuffer::Destroy(FVulkanDevice& Device)
 	for (int32 Index = 0; Index < AttachmentViewsToDelete.Num(); ++Index)
 	{
 		DEC_DWORD_STAT(STAT_VulkanNumImageViews);
-		Queue.EnqueueResource(VulkanRHI::FDeferredDeletionQueue::EType::ImageView, AttachmentViews[Index]);
+		Queue.EnqueueResource(VulkanRHI::FDeferredDeletionQueue::EType::ImageView, AttachmentViewsToDelete[Index]);
 	}
 
 	Queue.EnqueueResource(VulkanRHI::FDeferredDeletionQueue::EType::Framebuffer, Framebuffer);
