Enables mobile multi-view for translucent objects.

diff --git a/Engine/Shaders/Private/MobileBasePassVertexShader.usf b/Engine/Shaders/Private/MobileBasePassVertexShader.usf
--- a/Engine/Shaders/Private/MobileBasePassVertexShader.usf
+++ b/Engine/Shaders/Private/MobileBasePassVertexShader.usf
@@ -21,10 +21,6 @@
 #include "/Engine/Generated/VertexFactory.ush"
 #include "HeightFogCommon.ush"
 
-#if MOBILE_MULTI_VIEW
-int MobileMultiViewMask;
-#endif
-
 struct FMobileShadingBasePassVSToPS
 {
 	FVertexFactoryInterpolantsVSToPS FactoryInterpolants;
@@ -46,17 +42,6 @@ void Main(
 {
 #if MOBILE_MULTI_VIEW
 	const int MultiViewId = int(ViewId);
-	// Early out with a degenerate vertex to be clipped if the view is masked
-	
-	#if MATERIALBLENDING_ANY_TRANSLUCENT
-		BRANCH
-		if (MultiViewId == MobileMultiViewMask)
-		{
-			Output.Position = vec4(0.0);
-			return;
-		}
-	#endif
-
 	ResolvedView = ResolveView(uint(MultiViewId));
 	Output.BasePassInterpolants.MultiViewId = float(MultiViewId);
 #else
diff --git a/Engine/Source/Runtime/Renderer/Private/BasePassRendering.h b/Engine/Source/Runtime/Renderer/Private/BasePassRendering.h
--- a/Engine/Source/Runtime/Renderer/Private/BasePassRendering.h
+++ b/Engine/Source/Runtime/Renderer/Private/BasePassRendering.h
@@ -1004,7 +1004,6 @@ public:
 	const bool bAllowFog;
 	ERHIFeatureLevel::Type FeatureLevel;
 	const bool bIsInstancedStereo;
-	const bool bUseMobileMultiViewMask;
 
 	/** Initialization constructor. */
 	FProcessBasePassMeshParameters(
@@ -1013,8 +1012,7 @@ public:
 		const FPrimitiveSceneProxy* InPrimitiveSceneProxy,
 		bool InbAllowFog,
 		ERHIFeatureLevel::Type InFeatureLevel,
-		const bool InbIsInstancedStereo = false,
-		const bool InbUseMobileMultiViewMask = false
+		const bool InbIsInstancedStereo = false
 		):
 		Mesh(InMesh),
 		BatchElementMask(Mesh.Elements.Num()==1 ? 1 : (1<<Mesh.Elements.Num())-1), // 1 bit set for each mesh element
@@ -1024,8 +1022,7 @@ public:
 		ShadingModel(InMaterial->GetShadingModel()),
 		bAllowFog(InbAllowFog),
 		FeatureLevel(InFeatureLevel), 
-		bIsInstancedStereo(InbIsInstancedStereo), 
-		bUseMobileMultiViewMask(InbUseMobileMultiViewMask)
+		bIsInstancedStereo(InbIsInstancedStereo)
 	{
 	}
 
@@ -1037,8 +1034,7 @@ public:
 		const FPrimitiveSceneProxy* InPrimitiveSceneProxy,
 		bool InbAllowFog,
 		ERHIFeatureLevel::Type InFeatureLevel, 
-		bool InbIsInstancedStereo = false, 
-		bool InbUseMobileMultiViewMask = false
+		bool InbIsInstancedStereo = false
 		) :
 		Mesh(InMesh),
 		BatchElementMask(InBatchElementMask),
@@ -1048,8 +1044,7 @@ public:
 		ShadingModel(InMaterial->GetShadingModel()),
 		bAllowFog(InbAllowFog),
 		FeatureLevel(InFeatureLevel),
-		bIsInstancedStereo(InbIsInstancedStereo), 
-		bUseMobileMultiViewMask(InbUseMobileMultiViewMask)
+		bIsInstancedStereo(InbIsInstancedStereo)
 	{
 	}
 };
diff --git a/Engine/Source/Runtime/Renderer/Private/MobileBasePassRendering.h b/Engine/Source/Runtime/Renderer/Private/MobileBasePassRendering.h
--- a/Engine/Source/Runtime/Renderer/Private/MobileBasePassRendering.h
+++ b/Engine/Source/Runtime/Renderer/Private/MobileBasePassRendering.h
@@ -84,7 +84,6 @@ protected:
 		FMeshMaterialShader(Initializer)
 	{
 		VertexParametersType::Bind(Initializer.ParameterMap);
-		MobileMultiViewMaskParameter.Bind(Initializer.ParameterMap, TEXT("MobileMultiViewMask"));
 		PassUniformBuffer.Bind(Initializer.ParameterMap, FMobileBasePassUniformParameters::StaticStruct.GetShaderVariableName());
 	}
 
@@ -101,7 +100,6 @@ public:
 	{
 		bool bShaderHasOutdatedParameters = FMeshMaterialShader::Serialize(Ar);
 		VertexParametersType::Serialize(Ar);
-		Ar << MobileMultiViewMaskParameter;
 		return bShaderHasOutdatedParameters;
 	}
 
@@ -113,20 +111,6 @@ public:
 
 		FMaterialShader::SetViewParameters(RHICmdList, GetVertexShader(), *View, DrawRenderState.GetViewUniformBuffer());
 		FMeshMaterialShader::SetPassUniformBuffer(RHICmdList, GetVertexShader(), DrawRenderState.GetPassUniformBuffer());
-
-		if (MobileMultiViewMaskParameter.IsBound())
-		{
-			// Default is no masking
-			SetShaderValue(RHICmdList, GetVertexShader(), MobileMultiViewMaskParameter, -1);
-		}
-	}
-
-	void SetMobileMultiViewMask(FRHICommandList& RHICmdList, const int32 EyeIndex)
-	{
-		if (EyeIndex >= 0 && MobileMultiViewMaskParameter.IsBound())
-		{
-			SetShaderValue(RHICmdList, GetVertexShader(), MobileMultiViewMaskParameter, EyeIndex);
-		}
 	}
 
 	// Set parameters specific to mesh
@@ -144,9 +128,6 @@ public:
 		uint32 DataFlags = 0;
 		FMeshMaterialShader::SetMesh(RHICmdList, GetVertexShader(), InVertexFactory, View, Proxy, BatchElement, DrawRenderState, DataFlags);
 	}
-
-private:
-	FShaderParameter MobileMultiViewMaskParameter;
 };
 
 template<typename LightMapPolicyType>
@@ -815,11 +796,6 @@ public:
 		return CompareDrawingPolicy(A.LightMapPolicy,B.LightMapPolicy);
 	}
 
-	void SetMobileMultiViewMask(FRHICommandList& RHICmdList, const int32 EyeIndex)
-	{
-		VertexShader->SetMobileMultiViewMask(RHICmdList, EyeIndex);
-	}
-
 	void SetupPipelineState(FDrawingPolicyRenderState& DrawRenderState, const FViewInfo& View) const
 	{
 		if (UseDebugViewPS())
@@ -1055,7 +1031,6 @@ public:
 	const bool bAllowFog;
 	ERHIFeatureLevel::Type FeatureLevel;
 	const bool bIsInstancedStereo;
-	const bool bUseMobileMultiViewMask;
 
 	/** Initialization constructor. */
 	FMobileProcessBasePassMeshParameters(
@@ -1064,8 +1039,7 @@ public:
 		const FPrimitiveSceneProxy* InPrimitiveSceneProxy,
 		bool InbAllowFog,
 		ERHIFeatureLevel::Type InFeatureLevel,
-		const bool InbIsInstancedStereo = false,
-		const bool InbUseMobileMultiViewMask = false
+		const bool InbIsInstancedStereo = false
 	) :
 		Mesh(InMesh),
 		BatchElementMask(Mesh.Elements.Num() == 1 ? 1 : (1 << Mesh.Elements.Num()) - 1), // 1 bit set for each mesh element
@@ -1076,8 +1050,7 @@ public:
 		ShadingModel(InMaterial->GetShadingModel()),
 		bAllowFog(InbAllowFog),
 		FeatureLevel(InFeatureLevel),
-		bIsInstancedStereo(InbIsInstancedStereo),
-		bUseMobileMultiViewMask(InbUseMobileMultiViewMask)
+		bIsInstancedStereo(InbIsInstancedStereo)
 	{
 	}
 
@@ -1089,8 +1062,7 @@ public:
 		const FPrimitiveSceneProxy* InPrimitiveSceneProxy,
 		bool InbAllowFog,
 		ERHIFeatureLevel::Type InFeatureLevel,
-		bool InbIsInstancedStereo = false,
-		bool InbUseMobileMultiViewMask = false
+		bool InbIsInstancedStereo = false
 	) :
 		Mesh(InMesh),
 		BatchElementMask(InBatchElementMask),
@@ -1101,8 +1073,7 @@ public:
 		ShadingModel(InMaterial->GetShadingModel()),
 		bAllowFog(InbAllowFog),
 		FeatureLevel(InFeatureLevel),
-		bIsInstancedStereo(InbIsInstancedStereo),
-		bUseMobileMultiViewMask(InbUseMobileMultiViewMask)
+		bIsInstancedStereo(InbIsInstancedStereo)
 	{
 	}
 
diff --git a/Engine/Source/Runtime/Renderer/Private/MobileTranslucentRendering.cpp b/Engine/Source/Runtime/Renderer/Private/MobileTranslucentRendering.cpp
--- a/Engine/Source/Runtime/Renderer/Private/MobileTranslucentRendering.cpp
+++ b/Engine/Source/Runtime/Renderer/Private/MobileTranslucentRendering.cpp
@@ -185,12 +185,6 @@ public:
 		CommitGraphicsPipelineState(RHICmdList, DrawingPolicy, DrawRenderState, DrawingPolicy.GetBoundShaderStateInput(View.GetFeatureLevel()));
 		DrawingPolicy.SetSharedState(RHICmdList, DrawRenderState, &View, typename TMobileBasePassDrawingPolicy<FUniformLightMapPolicy>::ContextDataType());
 
-		if (Parameters.bUseMobileMultiViewMask)
-		{
-			// Mask opposite view
-			DrawingPolicy.SetMobileMultiViewMask(RHICmdList, (View.StereoPass == EStereoscopicPass::eSSP_LEFT_EYE) ? 1 : 0);
-		}
-
 		for (int32 BatchElementIndex = 0; BatchElementIndex<Parameters.Mesh.Elements.Num(); BatchElementIndex++)
 		{
 			TDrawEvent<FRHICommandList> MeshEvent;
@@ -251,8 +245,7 @@ bool FMobileTranslucencyDrawingPolicyFactory::DrawDynamicMesh(
 				PrimitiveSceneProxy,
 				true,
 				FeatureLevel, 
-				false, // ISR disabled for mobile
-				View.bIsMobileMultiViewEnabled
+				false // ISR disabled for mobile
 				),
 			FDrawMobileTranslucentMeshAction(
 				RHICmdList,
@@ -334,6 +327,10 @@ void FMobileSceneRenderer::RenderTranslucency(FRHICommandListImmediate& RHICmdLi
 			SCOPED_CONDITIONAL_DRAW_EVENTF(RHICmdList, EventView, Views.Num() > 1, TEXT("View%d"), ViewIndex);
 			
 			const FViewInfo& View = *PassViews[ViewIndex];
+			if (!View.ShouldRenderView())
+			{
+				continue;
+			}
 
 			TUniformBufferRef<FMobileBasePassUniformParameters> BasePassUniformBuffer;
 			CreateMobileBasePassUniformBuffer(RHICmdList, View, true, BasePassUniformBuffer);
