Fix to leak in post process managment.

diff --git a/Engine/Source/Runtime/Core/Private/HAL/LowLevelMemTracker.cpp b/Engine/Source/Runtime/Core/Private/HAL/LowLevelMemTracker.cpp
--- a/Engine/Source/Runtime/Core/Private/HAL/LowLevelMemTracker.cpp
+++ b/Engine/Source/Runtime/Core/Private/HAL/LowLevelMemTracker.cpp
@@ -61,6 +61,7 @@ DECLARE_LLM_MEMORY_STAT(TEXT("Networking"), STAT_NetworkingLLM, STATGROUP_LLMFUL
 DECLARE_LLM_MEMORY_STAT(TEXT("Meshes"), STAT_MeshesLLM, STATGROUP_LLMFULL);
 DECLARE_LLM_MEMORY_STAT(TEXT("Stats"), STAT_StatsLLM, STATGROUP_LLMFULL);
 DECLARE_LLM_MEMORY_STAT(TEXT("Shaders"), STAT_ShadersLLM, STATGROUP_LLMFULL);
+DECLARE_LLM_MEMORY_STAT(TEXT("PSO"), STAT_PSOLLM, STATGROUP_LLMFULL);
 DECLARE_LLM_MEMORY_STAT(TEXT("Textures"), STAT_TexturesLLM, STATGROUP_LLMFULL);
 DECLARE_LLM_MEMORY_STAT(TEXT("Render Targets"), STAT_RenderTargetsLLM, STATGROUP_LLMFULL);
 DECLARE_LLM_MEMORY_STAT(TEXT("RHIMisc"), STAT_RHIMiscLLM, STATGROUP_LLMFULL);
diff --git a/Engine/Source/Runtime/Core/Public/HAL/LowLevelMemTracker.h b/Engine/Source/Runtime/Core/Public/HAL/LowLevelMemTracker.h
--- a/Engine/Source/Runtime/Core/Public/HAL/LowLevelMemTracker.h
+++ b/Engine/Source/Runtime/Core/Public/HAL/LowLevelMemTracker.h
@@ -138,6 +138,7 @@ enum class ELLMTagSet : uint8
 	macro(Meshes,								"Meshes",						GET_STATFNAME(STAT_MeshesLLM),								GET_STATFNAME(STAT_MeshesSummaryLLM))		\
 	macro(Stats,								"Stats",						GET_STATFNAME(STAT_StatsLLM),								GET_STATFNAME(STAT_EngineSummaryLLM))		\
 	macro(Shaders,								"Shaders",						GET_STATFNAME(STAT_ShadersLLM),								GET_STATFNAME(STAT_EngineSummaryLLM))		\
+	macro(PSO,									"PSO",							GET_STATFNAME(STAT_PSOLLM),									GET_STATFNAME(STAT_EngineSummaryLLM))		\
 	macro(Textures,								"Textures",						GET_STATFNAME(STAT_TexturesLLM),							GET_STATFNAME(STAT_TexturesSummaryLLM))		\
 	macro(RenderTargets,						"RenderTargets",				GET_STATFNAME(STAT_RenderTargetsLLM),						GET_STATFNAME(STAT_EngineSummaryLLM))		\
 	macro(RHIMisc,								"RHIMisc",						GET_STATFNAME(STAT_RHIMiscLLM),								GET_STATFNAME(STAT_EngineSummaryLLM))		\
diff --git a/Engine/Source/Runtime/RHI/Private/PipelineStateCache.cpp b/Engine/Source/Runtime/RHI/Private/PipelineStateCache.cpp
--- a/Engine/Source/Runtime/RHI/Private/PipelineStateCache.cpp
+++ b/Engine/Source/Runtime/RHI/Private/PipelineStateCache.cpp
@@ -483,6 +483,8 @@ public:
 
 	void DoTask(ENamedThreads::Type CurrentThread, const FGraphEventRef& MyCompletionGraphEvent)
 	{
+		LLM_SCOPE(ELLMTag::PSO);
+
 		if (Pipeline->IsCompute())
 		{
 			FComputePipelineState* ComputePipeline = static_cast<FComputePipelineState*>(Pipeline);
@@ -713,7 +715,7 @@ FRHIComputePipelineState* ExecuteSetComputePipelineState(FComputePipelineState*
 
 FGraphicsPipelineState* PipelineStateCache::GetAndOrCreateGraphicsPipelineState(FRHICommandList& RHICmdList, const FGraphicsPipelineStateInitializer& OriginalInitializer, EApplyRendertargetOption ApplyFlags)
 {
-
+	LLM_SCOPE(ELLMTag::PSO);
 	SCOPE_CYCLE_COUNTER(STAT_GetOrCreatePSO);
 
 	FGraphicsPipelineStateInitializer NewInitializer;
diff --git a/Engine/Source/Runtime/Renderer/Private/PostProcess/PostProcessing.cpp b/Engine/Source/Runtime/Renderer/Private/PostProcess/PostProcessing.cpp
--- a/Engine/Source/Runtime/Renderer/Private/PostProcess/PostProcessing.cpp
+++ b/Engine/Source/Runtime/Renderer/Private/PostProcess/PostProcessing.cpp
@@ -932,11 +932,11 @@ static void AddTemporalAA( FPostprocessContext& Context, FRenderingCompositeOutp
 
 	FSceneViewState* ViewState = Context.View.ViewState;
 
-	FRCPassPostProcessTemporalAA* TemporalAAPass = new(FMemStack::Get()) FRCPassPostProcessTemporalAA(
+	FRCPassPostProcessTemporalAA* TemporalAAPass = Context.Graph.RegisterPass(new(FMemStack::Get()) FRCPassPostProcessTemporalAA(
 		Context,
 		Parameters,
 		Context.View.PrevViewInfo.TemporalAAHistory,
-		&ViewState->PendingPrevFrameViewInfo.TemporalAAHistory);
+		&ViewState->PendingPrevFrameViewInfo.TemporalAAHistory));
 
 	TemporalAAPass->SetInput( ePId_Input0, Context.FinalOutput );
 	TemporalAAPass->SetInput( ePId_Input2, VelocityInput );
@@ -1726,7 +1726,6 @@ void FPostProcessing::Process(FRHICommandListImmediate& RHICmdList, const FViewI
 			// some views don't have a state (thumbnail rendering)
 			if(!GIsHighResScreenshot && View.State && IStereoRendering::IsAPrimaryView(StereoPass))
 			{
-				
 				const bool bUseBasicEyeAdaptation = (AutoExposure.MethodId == EAutoExposureMethod::AEM_Basic);
 
 				if (bUseBasicEyeAdaptation) // log average ps reduction ( non histogram ) 
