RL - Fix nullpointer access when StaticSM data is not available

diff --git a/Engine/Source/Runtime/Renderer/Private/LightGridInjection.cpp b/Engine/Source/Runtime/Renderer/Private/LightGridInjection.cpp
--- a/Engine/Source/Runtime/Renderer/Private/LightGridInjection.cpp
+++ b/Engine/Source/Runtime/Renderer/Private/LightGridInjection.cpp
@@ -589,7 +589,7 @@ void FDeferredShadingSceneRenderer::ComputeLightGrid(FRHICommandListImmediate& R
 							}
 
 							const FStaticShadowDepthMap* StaticShadowDepthMap = LightSceneInfo->Proxy->GetStaticShadowDepthMap();
-							const uint32 bStaticallyShadowedValue = LightSceneInfo->IsPrecomputedLightingValid() && StaticShadowDepthMap && StaticShadowDepthMap->TextureRHI ? 1 : 0;
+							const uint32 bStaticallyShadowedValue = LightSceneInfo->IsPrecomputedLightingValid() && StaticShadowDepthMap && StaticShadowDepthMap->Data && StaticShadowDepthMap->TextureRHI ? 1 : 0;
 
 							ForwardLightData.DirectionalLightUseStaticShadowing = bStaticallyShadowedValue;
 							ForwardLightData.DirectionalLightStaticShadowBufferSize = bStaticallyShadowedValue ? FVector4(StaticShadowDepthMap->Data->ShadowMapSizeX, StaticShadowDepthMap->Data->ShadowMapSizeY, 1.0f / StaticShadowDepthMap->Data->ShadowMapSizeX, 1.0f / StaticShadowDepthMap->Data->ShadowMapSizeY) : FVector4(0, 0, 0, 0);
diff --git a/Engine/Source/Runtime/Renderer/Private/VolumetricFog.cpp b/Engine/Source/Runtime/Renderer/Private/VolumetricFog.cpp
--- a/Engine/Source/Runtime/Renderer/Private/VolumetricFog.cpp
+++ b/Engine/Source/Runtime/Renderer/Private/VolumetricFog.cpp
@@ -384,7 +384,7 @@ bool LightNeedsSeparateInjectionIntoVolumetricFog(const FLightSceneInfo* LightSc
 		&& LightProxy->CastsVolumetricShadow())
 	{
 		const FStaticShadowDepthMap* StaticShadowDepthMap = LightProxy->GetStaticShadowDepthMap();
-		const bool bStaticallyShadowed = LightSceneInfo->IsPrecomputedLightingValid() && StaticShadowDepthMap && StaticShadowDepthMap->TextureRHI;
+		const bool bStaticallyShadowed = LightSceneInfo->IsPrecomputedLightingValid() && StaticShadowDepthMap && StaticShadowDepthMap->Data && StaticShadowDepthMap->TextureRHI;
 
 		return GetShadowForInjectionIntoVolumetricFog(LightProxy, VisibleLightInfo) != NULL || bStaticallyShadowed;
 	}
