Fixed: Mobile Temporal AA is broken

diff --git a/Engine/Source/Runtime/Renderer/Private/PostProcess/PostProcessing.cpp b/Engine/Source/Runtime/Renderer/Private/PostProcess/PostProcessing.cpp
--- a/Engine/Source/Runtime/Renderer/Private/PostProcess/PostProcessing.cpp
+++ b/Engine/Source/Runtime/Renderer/Private/PostProcess/PostProcessing.cpp
@@ -1300,14 +1300,22 @@ void OverrideRenderTarget(FRenderingCompositeOutputRef It, TRefCountPtr<IPooledR
 
 bool FPostProcessing::AllowFullPostProcessing(const FViewInfo& View, ERHIFeatureLevel::Type FeatureLevel)
 {
-	return View.Family->EngineShowFlags.PostProcessing 
-		&& FeatureLevel >= ERHIFeatureLevel::SM4 
-		&& !View.Family->EngineShowFlags.VisualizeDistanceFieldAO
-		&& !View.Family->EngineShowFlags.VisualizeDistanceFieldGI
-		&& !View.Family->EngineShowFlags.VisualizeShadingModels
-		&& !View.Family->EngineShowFlags.VisualizeMeshDistanceFields
-		&& !View.Family->EngineShowFlags.VisualizeGlobalDistanceField
-		&& !View.Family->EngineShowFlags.ShaderComplexity;
+	if(FeatureLevel >= ERHIFeatureLevel::SM4)
+	{
+		return View.Family->EngineShowFlags.PostProcessing
+			&& !View.Family->EngineShowFlags.VisualizeDistanceFieldAO
+			&& !View.Family->EngineShowFlags.VisualizeDistanceFieldGI
+			&& !View.Family->EngineShowFlags.VisualizeShadingModels
+			&& !View.Family->EngineShowFlags.VisualizeMeshDistanceFields
+			&& !View.Family->EngineShowFlags.VisualizeGlobalDistanceField
+			&& !View.Family->EngineShowFlags.ShaderComplexity;
+	}
+	else
+	{
+		// Mobile post processing
+		return View.Family->EngineShowFlags.PostProcessing
+			&& !View.Family->EngineShowFlags.ShaderComplexity;
+	}
 }
 
 void FPostProcessing::RegisterHMDPostprocessPass(FPostprocessContext& Context, const FEngineShowFlags& EngineShowFlags) const
