Fixed: Mobile tonemapper permutations need to be refactored

diff --git a/Engine/Source/Runtime/Renderer/Private/PostProcess/PostProcessTonemap.cpp b/Engine/Source/Runtime/Renderer/Private/PostProcess/PostProcessTonemap.cpp
--- a/Engine/Source/Runtime/Renderer/Private/PostProcess/PostProcessTonemap.cpp
+++ b/Engine/Source/Runtime/Renderer/Private/PostProcess/PostProcessTonemap.cpp
@@ -1194,7 +1194,7 @@ class FPostProcessTonemapPS_ES2 : public FGlobalShader
 	class FTonemapperColorMatrixDim : SHADER_PERMUTATION_BOOL("USE_COLOR_MATRIX");
 	class FTonemapperShadowTintDim  : SHADER_PERMUTATION_BOOL("USE_SHADOW_TINT");
 	class FTonemapperContrastDim    : SHADER_PERMUTATION_BOOL("USE_CONTRAST");
-
+	
 	using FPermutationDomain = TShaderPermutationDomain<
 		TonemapperPermutation::FCommonDomain,
 		FTonemapperMsaaDim,
@@ -1204,6 +1204,89 @@ class FPostProcessTonemapPS_ES2 : public FGlobalShader
 		FTonemapperColorMatrixDim,
 		FTonemapperShadowTintDim,
 		FTonemapperContrastDim>;
+	
+	template<class TPermDim, class TPermVector>
+	static void EnableIfSet(const TPermVector& SourceDomain, TPermVector& DestDomain)
+	{
+		if (SourceDomain.template Get<TPermDim>())
+		{
+			DestDomain.template Set<TPermDim>(true);
+		}
+	}
+
+	// Reduce the number of permutations by combining common states
+	static FPermutationDomain RemapPermutationVector(FPermutationDomain WantedPermutationVector)
+	{
+		TonemapperPermutation::FCommonDomain WantedCommonPermutationVector = WantedPermutationVector.Get<TonemapperPermutation::FCommonDomain>();
+		FPermutationDomain RemappedPermutationVector;
+		TonemapperPermutation::FCommonDomain RemappedCommonPermutationVector;
+
+		// Note: FTonemapperSharpenDim, FTonemapperGrainJitterDim are not supported.
+
+		// 32 bit hdr (encoding)
+		EnableIfSet<FTonemapper32BPPHDRDim>(WantedPermutationVector, RemappedPermutationVector);
+
+		// Gamma only
+		if (WantedCommonPermutationVector.Get<TonemapperPermutation::FTonemapperGammaOnlyDim>())
+		{
+			RemappedCommonPermutationVector.Set<TonemapperPermutation::FTonemapperGammaOnlyDim>(true);
+
+			// Mutually exclusive - clear the wanted vector. 
+			WantedPermutationVector = FPermutationDomain();
+			WantedCommonPermutationVector = WantedPermutationVector.Get<TonemapperPermutation::FCommonDomain>();
+		}
+		else
+		{
+			// Always enable contrast.
+			RemappedPermutationVector.Set<FTonemapperContrastDim>(true);
+		}
+
+ 		// Bloom permutation
+		EnableIfSet<TonemapperPermutation::FTonemapperBloomDim>(WantedCommonPermutationVector, RemappedCommonPermutationVector);
+ 		// Vignette permutation
+		EnableIfSet<TonemapperPermutation::FTonemapperVignetteDim>(WantedCommonPermutationVector, RemappedCommonPermutationVector);
+ 		// Grain intensity permutation
+		EnableIfSet<TonemapperPermutation::FTonemapperGrainIntensityDim>(WantedCommonPermutationVector, RemappedCommonPermutationVector);
+ 		// Color matrix
+		EnableIfSet<FTonemapperColorMatrixDim>(WantedPermutationVector, RemappedPermutationVector);
+ 		// msaa permutation.
+		EnableIfSet<FTonemapperMsaaDim>(WantedPermutationVector, RemappedPermutationVector);
+
+		// DoF
+		if (WantedPermutationVector.Get<FTonemapperDOFDim>())
+		{
+			RemappedPermutationVector.Set<FTonemapperDOFDim>(true);
+			RemappedPermutationVector.Set<FTonemapperLightShaftsDim>(true);
+			RemappedCommonPermutationVector.Set<TonemapperPermutation::FTonemapperVignetteDim>(true);
+			RemappedCommonPermutationVector.Set<TonemapperPermutation::FTonemapperBloomDim>(true);
+		}
+
+		// light shafts
+		if (WantedPermutationVector.Get<FTonemapperLightShaftsDim>())
+		{
+			RemappedPermutationVector.Set<FTonemapperLightShaftsDim>(true);
+			RemappedCommonPermutationVector.Set<TonemapperPermutation::FTonemapperVignetteDim>(true);
+			RemappedCommonPermutationVector.Set<TonemapperPermutation::FTonemapperBloomDim>(true);
+		}
+
+		// Shadow tint
+		if (WantedPermutationVector.Get<FTonemapperShadowTintDim>())
+		{
+			RemappedPermutationVector.Set<FTonemapperShadowTintDim>(true);
+			RemappedPermutationVector.Set<FTonemapperColorMatrixDim>(true);
+		}
+
+		if (RemappedPermutationVector.Get<FTonemapper32BPPHDRDim>())
+		{
+			// 32 bpp hdr does not support:
+			RemappedPermutationVector.Set<FTonemapperDOFDim>(false);
+			RemappedPermutationVector.Set<FTonemapperMsaaDim>(false);
+			RemappedPermutationVector.Set<FTonemapperLightShaftsDim>(false);
+		}
+
+		RemappedPermutationVector.Set<TonemapperPermutation::FCommonDomain>(RemappedCommonPermutationVector);
+		return RemappedPermutationVector;
+	}
 
 	static FPermutationDomain BuildPermutationVector(const FViewInfo& View)
 	{
@@ -1218,7 +1301,7 @@ class FPostProcessTonemapPS_ES2 : public FGlobalShader
 		if (CommonPermutationVector.Get<TonemapperPermutation::FTonemapperGammaOnlyDim>())
 		{
 			MobilePermutationVector.Set<FTonemapper32BPPHDRDim>(bUse32BPPHDR);
-			return MobilePermutationVector;
+			return RemapPermutationVector(MobilePermutationVector);
 		}
 
 		const FFinalPostProcessSettings& Settings = View.FinalPostProcessSettings;
@@ -1274,7 +1357,8 @@ class FPostProcessTonemapPS_ES2 : public FGlobalShader
 		CommonPermutationVector.Set<TonemapperPermutation::FTonemapperSharpenDim>(false);
 		MobilePermutationVector.Set<TonemapperPermutation::FCommonDomain>(CommonPermutationVector);
 
-		return MobilePermutationVector;
+		// We're not supporting every possible permutation, remap the permutation vector to combine common effects.
+		return RemapPermutationVector(MobilePermutationVector);
 	}
 
 	static bool ShouldCompilePermutation(const FGlobalShaderPermutationParameters& Parameters)
@@ -1286,15 +1370,11 @@ class FPostProcessTonemapPS_ES2 : public FGlobalShader
 		{
 			return false;
 		}
-		
-		if (CommonPermutationVector.Get<TonemapperPermutation::FTonemapperGammaOnlyDim>())
+
+		// If this permutation vector is remapped at runtime, we can avoid the compile.
+		if (RemapPermutationVector(PermutationVector) != PermutationVector)
 		{
-			if (!PermutationVector.Get<FTonemapperMsaaDim>() &&
-				!PermutationVector.Get<FTonemapperDOFDim>() &&
-				!PermutationVector.Get<FTonemapperLightShaftsDim>() &&
-				!PermutationVector.Get<FTonemapperColorMatrixDim>() &&
-				!PermutationVector.Get<FTonemapperShadowTintDim>() &&
-				!PermutationVector.Get<FTonemapperContrastDim>()) return false;
+			return false;
 		}
 
 		// Only cache for ES2/3.1 shader platforms, and only compile 32bpp shaders for Android or PC emulation
