Fixed: Software occlusion and hardware occlusion queries should be mutually exclusive

diff --git a/Engine/Source/Runtime/Renderer/Private/SceneSoftwareOcclusion.cpp b/Engine/Source/Runtime/Renderer/Private/SceneSoftwareOcclusion.cpp
--- a/Engine/Source/Runtime/Renderer/Private/SceneSoftwareOcclusion.cpp
+++ b/Engine/Source/Runtime/Renderer/Private/SceneSoftwareOcclusion.cpp
@@ -116,15 +116,34 @@ struct FOcclusionMeshData
 	FPrimitiveComponentId	PrimId;
 };
 
+struct FSortedIndexDepth
+{
+	int32 Index;
+	float Depth;
+};
+
 struct FOcclusionFrameData
 {
-	TArray<uint16>					SortedTriangles;
+	// binned tris
+	TArray<FSortedIndexDepth>		SortedTriangles[BIN_NUM];
 	
-	// frame data	
+	// tris data	
 	TArray<FScreenTriangle>			ScreenTriangles;
-	TArray<float>					ScreenTrianglesDepth;
 	TArray<FPrimitiveComponentId>	ScreenTrianglesPrimID;
 	TArray<uint8>					ScreenTrianglesFlags;
+
+	void ReserveBuffers(int32 NumTriangles)
+	{
+		const int32 NumTrianglesPerBin = NumTriangles/BIN_NUM + 1;
+		for (int32 BinIdx = 0; BinIdx < BIN_NUM; ++BinIdx)
+		{
+			SortedTriangles[BinIdx].Reserve(NumTrianglesPerBin);
+		}
+				
+		ScreenTriangles.Reserve(NumTriangles);
+		ScreenTrianglesPrimID.Reserve(NumTriangles);
+		ScreenTrianglesFlags.Reserve(NumTriangles);
+	}
 };
 
 struct FOcclusionSceneData
@@ -133,6 +152,7 @@ struct FOcclusionSceneData
 	TArray<FVector>					OccludeeBoxMinMax;
 	TArray<FPrimitiveComponentId>	OccludeeBoxPrimId;
 	TArray<FOcclusionMeshData>		OccluderData;
+	int32							NumOccluderTriangles;
 };
 
 inline uint64 ComputeBinRowMask(int32 BinMinX, float fX0, float fX1)
@@ -281,19 +301,26 @@ inline bool AddTriangle(FScreenTriangle& Tri, float TriDepth, FPrimitiveComponen
 		}
 	}
 
+	int32 TriangleID = InData.ScreenTriangles.Add(Tri);
 	InData.ScreenTrianglesPrimID.Add(PrimitiveId);
 	InData.ScreenTrianglesFlags.Add(MeshFlags);
-	InData.ScreenTrianglesDepth.Add(TriDepth);
-	int32 TriangleID = InData.ScreenTriangles.Add(Tri);
-
-	if (TriangleID <= MAX_uint16)
+	
+	// bin
+	int32 MinX = FMath::Min3(Tri.V[0].X, Tri.V[1].X, Tri.V[2].X) / BIN_WIDTH; 
+	int32 MaxX = FMath::Max3(Tri.V[0].X, Tri.V[1].X, Tri.V[2].X) / BIN_WIDTH;
+	int32 BinMin = FMath::Max(MinX, 0);
+	int32 BinMax = FMath::Min(MaxX, BIN_NUM-1);
+	
+	FSortedIndexDepth SortedIndexDepth;
+	SortedIndexDepth.Index = TriangleID;
+	SortedIndexDepth.Depth = TriDepth;
+			
+	for (int32 BinIdx = BinMin; BinIdx <= BinMax; ++BinIdx)
 	{
-		uint16 TriangleID16 = (uint16)TriangleID;
-		InData.SortedTriangles.Add(TriangleID16);
-		return true;
+		InData.SortedTriangles[BinIdx].Add(SortedIndexDepth);
 	}
 
-	return false;
+	return true;
 }
 
 static const VectorRegister vFramebufferBounds = MakeVectorRegister(FRAMEBUFFER_WIDTH-1, FRAMEBUFFER_HEIGHT-1, 1.0f, 1.0f);
@@ -309,7 +336,8 @@ static const uint32 sBBzInd[NUM_CUBE_VTX] = { 1, 1, 0, 0, 0, 1, 1, 0 };
 
 static void ProcessOccludeeGeomSIMD(const FMatrix& InMat, const FVector* InMinMax, int32 Num, int32* RESTRICT OutQuads, float* RESTRICT OutQuadDepth, int32* RESTRICT OutQuadClipped)
 {
-	VectorRegister vClippingW = VectorLoadFloat1(&GNearClippingPlane);
+	const float W_CLIP = InMat.M[3][2];
+	VectorRegister vClippingW = VectorLoadFloat1(&W_CLIP);
 	VectorRegister mRow0  = VectorLoadAligned(InMat.M[0]);
 	VectorRegister mRow1  = VectorLoadAligned(InMat.M[1]);
 	VectorRegister mRow2  = VectorLoadAligned(InMat.M[2]);
@@ -374,7 +402,7 @@ static void ProcessOccludeeGeomSIMD(const FMatrix& InMat, const FVector* InMinMa
 
 static void ProcessOccludeeGeomScalar(const FMatrix& InMat, const FVector* InMinMax, int32 Num, int32* RESTRICT OutQuads, float* RESTRICT OutQuadDepth, int32* RESTRICT OutQuadClipped)
 {
-	const float W_CLIP = GNearClippingPlane;
+	const float W_CLIP =  InMat.M[3][2];
 	FVector4 AX = FVector4(InMat.M[0][0], InMat.M[0][1], InMat.M[0][2], InMat.M[0][3]);
 	FVector4 AY = FVector4(InMat.M[1][0], InMat.M[1][1], InMat.M[1][2], InMat.M[1][3]);
 	FVector4 AZ = FVector4(InMat.M[2][0], InMat.M[2][1], InMat.M[2][2], InMat.M[2][3]);
@@ -560,10 +588,8 @@ static bool ClippedVertexToScreen(const FVector4& XFV, FScreenPosition& OutSP, f
 	return false;
 }
 
-static uint8 ProcessXFormVertex(const FVector4& XFV)
+static uint8 ProcessXFormVertex(const FVector4& XFV, float W_CLIP)
 {
-	const float W_CLIP = GNearClippingPlane;
-	
 	uint8 Flags = 0;
 	float W = XFV.W;
 
@@ -597,7 +623,7 @@ static uint8 ProcessXFormVertex(const FVector4& XFV)
 
 static void ProcessOccluderGeom(const FOcclusionSceneData& SceneData, FOcclusionFrameData& OutData)
 {
-	const float W_CLIP = GNearClippingPlane;
+	const float W_CLIP = SceneData.ViewProj.M[3][2];
 
 	const int32 NumMeshes = SceneData.OccluderData.Num();
 	const FOcclusionMeshData* MeshData = SceneData.OccluderData.GetData();
@@ -643,7 +669,7 @@ static void ProcessOccluderGeom(const FOcclusionSceneData& SceneData, FOcclusion
 				// Store
 				VectorStoreAligned(VTempX, &MeshClipVertices[i]);
 						
-				uint8 VertexFlags = ProcessXFormVertex(MeshClipVertices[i]);
+				uint8 VertexFlags = ProcessXFormVertex(MeshClipVertices[i], W_CLIP);
 				MeshClipVertexFlags[i] = VertexFlags;
 			}
 		}
@@ -651,12 +677,6 @@ static void ProcessOccluderGeom(const FOcclusionSceneData& SceneData, FOcclusion
 		const uint16* MeshIndices = Mesh.IndicesSP->GetData();
 		int32 NumTris = Mesh.IndicesSP->Num()/3;
 		int32 NumDataTris = OutData.ScreenTriangles.Num();
-		
-		// Reserve space for mesh triangles
-		OutData.ScreenTriangles.Reserve(NumDataTris + NumTris);
-		OutData.ScreenTrianglesDepth.Reserve(NumDataTris + NumTris);
-		OutData.ScreenTrianglesPrimID.Reserve(NumDataTris + NumTris);
-		OutData.ScreenTrianglesFlags.Reserve(NumDataTris + NumTris);
 
 		// Create triangles
 		for (int32 i = 0; i < NumTris; ++i)
@@ -759,6 +779,7 @@ public:
 	FSWOccluderElementsCollector(FOcclusionSceneData& InData)
 		: SceneData(InData)
 	{
+		SceneData.NumOccluderTriangles = 0;
 	}
 	
 	void SetPrimitiveID(FPrimitiveComponentId PrimitiveId)
@@ -775,6 +796,8 @@ public:
 		MeshData.LocalToWorld = LocalToWorld;
 		MeshData.VerticesSP = Vertices;
 		MeshData.IndicesSP = Indices;
+
+		SceneData.NumOccluderTriangles+= Indices->Num()/3;
 	}
 
 public:
@@ -785,7 +808,9 @@ public:
 static void ProcessOcclusionFrame(const FOcclusionSceneData& InSceneData, FOcclusionFrameResults& OutResults)
 {
 	FOcclusionFrameData FrameData;
-	
+	int32 NumExpectedTriangles = InSceneData.NumOccluderTriangles + InSceneData.OccludeeBoxPrimId.Num(); // one triangle for each occludee
+	FrameData.ReserveBuffers(NumExpectedTriangles);
+		
 	{
 		SCOPE_CYCLE_COUNTER(STAT_SoftwareOcclusionProcessOccluder)
 		ProcessOccluderGeom(InSceneData, FrameData);
@@ -797,63 +822,46 @@ static void ProcessOcclusionFrame(const FOcclusionSceneData& InSceneData, FOcclu
 		ProcessOccludeeGeom(InSceneData, FrameData, OutResults.VisibilityMap);
 	}
 
-	{
-		SCOPE_CYCLE_COUNTER(STAT_SoftwareOcclusionSort);
-		// Sort all triangles by depth
-		FrameData.SortedTriangles.Sort([&](uint16 A, uint16 B) { 
-			// biggerZ (closer) first 
-			return FrameData.ScreenTrianglesDepth[A] > FrameData.ScreenTrianglesDepth[B]; 
-		});
-	}
-
 	int32 NumRasterizedOccluderTris = 0;
 	int32 NumRasterizedOccludeeTris = 0;
-	const int32 NumTris = FrameData.SortedTriangles.Num();
 	{
 		SCOPE_CYCLE_COUNTER(STAT_SoftwareOcclusionRasterize);
 
 		const uint8* MeshFlags = FrameData.ScreenTrianglesFlags.GetData();
 		const FPrimitiveComponentId* PrimitiveIds = FrameData.ScreenTrianglesPrimID.GetData();
 		const FScreenTriangle* Tris = FrameData.ScreenTriangles.GetData();
-		const uint16* TrisID = FrameData.SortedTriangles.GetData();
-				
-		for (int32 i = 0; i < NumTris; ++i)
+						
+		for (int32 BinIdx = 0; BinIdx < BIN_NUM; ++BinIdx)
 		{
-			int32 TriID = TrisID[i];
-			const FScreenTriangle& Tri = Tris[TriID];
-			uint8 Flags = MeshFlags[TriID];
-			FPrimitiveComponentId PrimitiveId = PrimitiveIds[TriID];
+			// Sort triangles in the bin by depth
+			FrameData.SortedTriangles[BinIdx].Sort([](const FSortedIndexDepth& A, const FSortedIndexDepth& B) { 
+				// biggerZ (closer) first 
+				return A.Depth > B.Depth; 
+			});
+
+			const FSortedIndexDepth* SortedTriIndices = FrameData.SortedTriangles[BinIdx].GetData();
+			const int32 NumTris = FrameData.SortedTriangles[BinIdx].Num();
+			const int32 BinMinX = BinIdx*BIN_WIDTH;
+			FFramebufferBin& Bin = OutResults.Bins[BinIdx];
+			// TODO: add a way to check when bin is already fully rasterized, so we can skip this work
 						
-			// bin
-			int32 MinX = FMath::Min3(Tri.V[0].X, Tri.V[1].X, Tri.V[2].X) / BIN_WIDTH; 
-			int32 MaxX = FMath::Max3(Tri.V[0].X, Tri.V[1].X, Tri.V[2].X) / BIN_WIDTH;
-			int32 BinMin = FMath::Max(MinX, 0);
-			int32 BinMax = FMath::Min(MaxX, BIN_NUM-1);
-			
-			if (Flags != 0)
+			for (int32 TriIdx = 0; TriIdx < NumTris; ++TriIdx)
 			{
-				for (int32 j = BinMin; j <= BinMax; ++j)
+				int32 TriID = SortedTriIndices[TriIdx].Index;
+				uint8 Flags = MeshFlags[TriID];
+				FPrimitiveComponentId PrimitiveId = PrimitiveIds[TriID];
+				const FScreenTriangle& Tri = Tris[TriID];
+
+				if (Flags != 0)
 				{
-					const int32 BinMinX = j*BIN_WIDTH;
-					FFramebufferBin& Bin = OutResults.Bins[j];
-					//TODO: add a way to check when bin is fully rasterized, so we can skip this work
-					
 					// rasterize occluder
 					RasterizeOccluderTri(Tri, Bin.Data, BinMinX);
 					NumRasterizedOccluderTris++;
 				}
-			}
-			else
-			{
-				bool& VisBit = OutResults.VisibilityMap.FindOrAdd(PrimitiveId);
-
-				for (int32 j = BinMin; j <= BinMax; ++j)
+				else
 				{
-					const int32 BinMinX = j*BIN_WIDTH;
-					FFramebufferBin& Bin = OutResults.Bins[j];
-					//TODO: add a way to check when bin is fully rasterized, so we can skip this work
-					
 					// rasterize occludee
+					bool& VisBit = OutResults.VisibilityMap.FindOrAdd(PrimitiveId);
 					bool bVisible = RasterizeOccludeeQuad(Tri, Bin.Data, BinMinX);
 					VisBit|= bVisible;
 					NumRasterizedOccludeeTris++;
@@ -862,7 +870,8 @@ static void ProcessOcclusionFrame(const FOcclusionSceneData& InSceneData, FOcclu
 		}
 	}
 	
-	INC_DWORD_STAT_BY(STAT_SoftwareTriangles, NumTris);
+	int32 NumTotalTris = FrameData.ScreenTriangles.Num();
+	INC_DWORD_STAT_BY(STAT_SoftwareTriangles, NumTotalTris);
 	INC_DWORD_STAT_BY(STAT_SoftwareOccluderTris, NumRasterizedOccluderTris);
 	INC_DWORD_STAT_BY(STAT_SoftwareOccludeeTris, NumRasterizedOccludeeTris);
 }
@@ -887,10 +896,17 @@ static int32 ApplyResults(const FScene* Scene, FViewInfo& View, const FOcclusion
 		FPrimitiveComponentId PrimId = Scene->PrimitiveComponentIds[PrimitiveIndex];
 			
 		const bool* bVisiblePtr = Results.VisibilityMap.Find(PrimId);
-		if (bVisiblePtr && *bVisiblePtr == false)
+		if (bVisiblePtr)
 		{
-			View.PrimitiveVisibilityMap[PrimitiveIndex] = false;
-			NumOccluded++;
+			if (*bVisiblePtr == false)
+			{
+				View.PrimitiveVisibilityMap[PrimitiveIndex] = false;
+				NumOccluded++;
+			}
+			else
+			{
+				View.PrimitiveDefinitelyUnoccludedMap[PrimitiveIndex] = true;
+			}
 		}
 	}
 
diff --git a/Engine/Source/Runtime/Renderer/Private/SceneVisibility.cpp b/Engine/Source/Runtime/Renderer/Private/SceneVisibility.cpp
--- a/Engine/Source/Runtime/Renderer/Private/SceneVisibility.cpp
+++ b/Engine/Source/Runtime/Renderer/Private/SceneVisibility.cpp
@@ -1505,7 +1505,12 @@ static int32 OcclusionCull(FRHICommandListImmediate& RHICmdList, const FScene* S
 	float CurrentRealTime = View.Family->CurrentRealTime;
 	if (ViewState)
 	{
-		if (Scene->GetFeatureLevel() >= ERHIFeatureLevel::ES3_1)
+		if (ViewState->SceneSoftwareOcclusion)
+		{
+			SCOPE_CYCLE_COUNTER(STAT_SoftwareOcclusionCull)
+			NumOccludedPrimitives += ViewState->SceneSoftwareOcclusion->Process(RHICmdList, Scene, View);
+		}
+		else if (Scene->GetFeatureLevel() >= ERHIFeatureLevel::ES3_1)
 		{
 			bool bSubmitQueries = !View.bDisableQuerySubmissions;
 #if !(UE_BUILD_SHIPPING || UE_BUILD_TEST)
@@ -2994,13 +2999,6 @@ void FSceneRenderer::ComputeViewVisibility(FRHICommandListImmediate& RHICmdList)
 			STAT(NumOccludedPrimitives += NumOccludedPrimitivesInView);
 		}
 
-		if (ViewState && ViewState->SceneSoftwareOcclusion && !View.Family->EngineShowFlags.Wireframe)
-		{
-			SCOPE_CYCLE_COUNTER(STAT_SoftwareOcclusionCull)
-			int32 NumOccludedPrimitivesInView = ViewState->SceneSoftwareOcclusion->Process(RHICmdList, Scene, View);
-			STAT(NumOccludedPrimitives += NumOccludedPrimitivesInView);
-		}
-
 		MarkAllPrimitivesForReflectionProxyUpdate(Scene);
 		{
 			QUICK_SCOPE_CYCLE_COUNTER(STAT_ViewVisibilityTime_ConditionalMarkStaticMeshElementsForUpdate);
