[Copy] Fixed FForwardLightData not being reinitialized every frame

diff --git a/Engine/Source/Runtime/Renderer/Private/LightGridInjection.cpp b/Engine/Source/Runtime/Renderer/Private/LightGridInjection.cpp
--- a/Engine/Source/Runtime/Renderer/Private/LightGridInjection.cpp
+++ b/Engine/Source/Runtime/Renderer/Private/LightGridInjection.cpp
@@ -452,6 +452,7 @@ void FDeferredShadingSceneRenderer::ComputeLightGrid(FRHICommandListImmediate& R
 		{
 			FViewInfo& View = Views[ViewIndex];
 			FForwardLightData& ForwardLightData = View.ForwardLightingResources->ForwardLightData;
+			ForwardLightData = FForwardLightData();
 
 			TArray<FForwardLocalLightData, SceneRenderingAllocator> ForwardLocalLightData;
 			float FurthestLight = 1000;
