Fixed: Android Rendering issues on ARPG in Vulkan
Explicit base pass
3 frames occlusion latency on Android Vulkan + RHIT
Removed unused depth buffer copy

diff --git a/Engine/Source/Runtime/RHI/Public/RHIResources.h b/Engine/Source/Runtime/RHI/Public/RHIResources.h
--- a/Engine/Source/Runtime/RHI/Public/RHIResources.h
+++ b/Engine/Source/Runtime/RHI/Public/RHIResources.h
@@ -1947,8 +1947,8 @@ struct FRHIRenderPassInfo
 	FResolveParams ResolveParameters;
 
 	// Some RHIs require a hint that occlusion queries will be used in this render pass
-	const uint32 NumOcclusionQueries = 0;
-	const bool bOcclusionQueries = false;
+	uint32 NumOcclusionQueries = 0;
+	bool bOcclusionQueries = false;
 
 	// Some RHIs need to know if this render pass is going to be reading and writing to the same texture in the case of generating mip maps for partial resource transitions
 	bool bGeneratingMips = false;
diff --git a/Engine/Source/Runtime/Renderer/Private/MobileShadingRenderer.cpp b/Engine/Source/Runtime/Renderer/Private/MobileShadingRenderer.cpp
--- a/Engine/Source/Runtime/Renderer/Private/MobileShadingRenderer.cpp
+++ b/Engine/Source/Runtime/Renderer/Private/MobileShadingRenderer.cpp
@@ -250,30 +250,55 @@ void FMobileSceneRenderer::Render(FRHICommandListImmediate& RHICmdList)
 	}
 	else
 	{
-		// Begin rendering to scene color
-		SceneContext.BeginRenderingSceneColor(RHICmdList, ESimpleRenderTargetMode::EClearColorAndDepth);
-		SceneColor = SceneContext.GetSceneColorSurface();
+		if (IsVulkanPlatform(ViewFamily.GetShaderPlatform()))
+		{	
+			SceneColor = SceneContext.GetSceneColorSurface();
+			// decals/translucency need to read/test depth
+			// TODO: add sub-passes! this is expensive, especially with MSAA
+			EDepthStencilTargetActions DepthStoreAction = EDepthStencilTargetActions::ClearDepthStencil_StoreDepthStencil;
+						
+			FRHIRenderPassInfo RPInfo(
+				SceneColor,
+				ERenderTargetActions::Clear_Store,
+				SceneContext.GetSceneDepthSurface(),
+				DepthStoreAction, 
+				FExclusiveDepthStencil::DepthWrite_StencilWrite
+			);
+			RPInfo.NumOcclusionQueries = ComputeNumOcclusionQueriesToBatch();
+			RPInfo.bOcclusionQueries = RPInfo.NumOcclusionQueries != 0;
+			RHICmdList.BeginRenderPass(RPInfo, TEXT("BasePass"));
+		}
+		else
+		{
+			// Begin rendering to scene color
+			SceneContext.BeginRenderingSceneColor(RHICmdList, ESimpleRenderTargetMode::EClearColorAndDepth);
+			SceneColor = SceneContext.GetSceneColorSurface();
+		}
 	}
 
 	if (GIsEditor && !View.bIsSceneCapture)
 	{
 		DrawClearQuad(RHICmdList, Views[0].BackgroundColor);
 	}
-
-
+	
 	RHICmdList.SetCurrentStat(GET_STATID(STAT_CLMM_BasePass));
 
 	RenderMobileBasePass(RHICmdList, ViewList);
-
-
+	
 	RHICmdList.SetCurrentStat(GET_STATID(STAT_CLMM_Occlusion));
 
 	// Issue occlusion queries
 	RenderOcclusion(RHICmdList);
-
-
+	
 	RHICmdList.SetCurrentStat(GET_STATID(STAT_CLMM_Post));
 
+	if (RHICmdList.IsInsideRenderPass())
+	{
+		// Vulkan only right now
+		// END BasePass
+		RHICmdList.EndRenderPass();
+	}
+
 	for (int32 ViewExt = 0; ViewExt < ViewFamily.ViewExtensions.Num(); ++ViewExt)
 	{
 		for (int32 ViewIndex = 0; ViewIndex < ViewFamily.Views.Num(); ++ViewIndex)
@@ -505,6 +530,24 @@ void FMobileSceneRenderer::RenderOcclusion(FRHICommandListImmediate& RHICmdList)
 	}
 }
 
+int32 FMobileSceneRenderer::ComputeNumOcclusionQueriesToBatch() const
+{
+	int32 NumQueriesForBatch = 0;
+	for (int32 ViewIndex = 0; ViewIndex < Views.Num(); ViewIndex++)
+	{
+		const FViewInfo& View = Views[ViewIndex];
+		const FSceneViewState* ViewState = (FSceneViewState*)View.State;
+#if !(UE_BUILD_SHIPPING || UE_BUILD_TEST)
+		if (!ViewState || (!ViewState->HasViewParent() && !ViewState->bIsFrozen))
+#endif
+		{
+			NumQueriesForBatch += View.IndividualOcclusionQueries.GetNumBatchOcclusionQueries();
+			NumQueriesForBatch += View.GroupedOcclusionQueries.GetNumBatchOcclusionQueries();
+		}
+	}
+	
+	return NumQueriesForBatch;
+}
 
 void FMobileSceneRenderer::ConditionalResolveSceneDepth(FRHICommandListImmediate& RHICmdList, const FViewInfo& View)
 {
diff --git a/Engine/Source/Runtime/Renderer/Private/SceneOcclusion.cpp b/Engine/Source/Runtime/Renderer/Private/SceneOcclusion.cpp
--- a/Engine/Source/Runtime/Renderer/Private/SceneOcclusion.cpp
+++ b/Engine/Source/Runtime/Renderer/Private/SceneOcclusion.cpp
@@ -79,10 +79,11 @@ int32 FOcclusionQueryHelpers::GetNumBufferedFrames(ERHIFeatureLevel::Type Featur
 	if (FeatureLevel <= ERHIFeatureLevel::ES3_1)
 	{
 		NumExtraMobileFrames++; // the mobile renderer just doesn't do much after the basepass, and hence it will be asking for the query results almost immediately; the results can't possibly be ready in 1 frame.
-
-		if (IsOpenGLPlatform(GShaderPlatformForFeatureLevel[FeatureLevel]) && IsRunningRHIInSeparateThread())
+		
+		EShaderPlatform ShaderPlatform = GShaderPlatformForFeatureLevel[FeatureLevel];
+		if ((IsOpenGLPlatform(ShaderPlatform) || IsVulkanPlatform(ShaderPlatform)) && IsRunningRHIInSeparateThread())
 		{
-			// OpenGL, unfortunately, requires the RHIThread to mediate the readback of queries. Therefore we need an extra frame to avoid a stall in either thread. 
+			// Android, unfortunately, requires the RHIThread to mediate the readback of queries. Therefore we need an extra frame to avoid a stall in either thread. 
 			// The RHIT needs to do read back after the queries are ready and before the RT needs them to avoid stalls. The RHIT may be busy when the queries become ready, so this is all very complicated.
 			NumExtraMobileFrames++;
 		}
@@ -1318,7 +1319,8 @@ void FSceneRenderer::BeginOcclusionTests(FRHICommandListImmediate& RHICmdList, b
 					ViewState->OcclusionQueryPool.ReleaseQuery(QueryIt.Value());
 				}
 				ShadowOcclusionQueryMap.Reset();
-
+				
+				if (FeatureLevel > ERHIFeatureLevel::ES3_1)
 				{
 					SCOPED_DRAW_EVENT(RHICmdList, ShadowFrustumQueries);
 
@@ -1392,7 +1394,8 @@ void FSceneRenderer::BeginOcclusionTests(FRHICommandListImmediate& RHICmdList, b
 					}
 				}
 
-				if (!View.bIsPlanarReflection &&
+				if (FeatureLevel > ERHIFeatureLevel::ES3_1 &&
+					!View.bIsPlanarReflection &&
 					!View.bIsSceneCapture &&
 					!View.bIsReflectionCapture)
 				{
@@ -1504,6 +1507,7 @@ void FSceneRenderer::BeginOcclusionTests(FRHICommandListImmediate& RHICmdList, b
 				SetGraphicsPipelineState(RHICmdList, GraphicsPSOInit);
 				VertexShader->SetParameters(RHICmdList, View);
 
+				if (FeatureLevel > ERHIFeatureLevel::ES3_1)
 				{
 					SCOPED_DRAW_EVENT(RHICmdList, ShadowFrustumQueries);
 					for (TPair<FProjectedShadowInfo const*, FRenderQueryRHIRef> const& Query : ViewQuery.PointLightQueries)
@@ -1522,6 +1526,7 @@ void FSceneRenderer::BeginOcclusionTests(FRHICommandListImmediate& RHICmdList, b
 					}
 				}
 
+				if (FeatureLevel > ERHIFeatureLevel::ES3_1)
 				{
 					SCOPED_DRAW_EVENT(RHICmdList, PlanarReflectionQueries);
 					for (TPair<FPlanarReflectionSceneProxy const*, FRenderQueryRHIRef> const& Query : ViewQuery.ReflectionQueries)
diff --git a/Engine/Source/Runtime/Renderer/Private/SceneRendering.h b/Engine/Source/Runtime/Renderer/Private/SceneRendering.h
--- a/Engine/Source/Runtime/Renderer/Private/SceneRendering.h
+++ b/Engine/Source/Runtime/Renderer/Private/SceneRendering.h
@@ -1718,6 +1718,9 @@ protected:
 
 	/** Issues occlusion queries */
 	void RenderOcclusion(FRHICommandListImmediate& RHICmdList);
+	
+	/** Computes how many queries will be issued this frame */
+	int32 ComputeNumOcclusionQueriesToBatch() const;
 
 	/** Renders decals. */
 	void RenderDecals(FRHICommandListImmediate& RHICmdList);
diff --git a/Engine/Source/Runtime/VulkanRHI/Private/Android/VulkanAndroidPlatform.h b/Engine/Source/Runtime/VulkanRHI/Private/Android/VulkanAndroidPlatform.h
--- a/Engine/Source/Runtime/VulkanRHI/Private/Android/VulkanAndroidPlatform.h
+++ b/Engine/Source/Runtime/VulkanRHI/Private/Android/VulkanAndroidPlatform.h
@@ -45,7 +45,7 @@ public:
 	static bool SupportsStandardSwapchain();
 	static EPixelFormat GetPixelFormatForNonDefaultSwapchain();
 
-	static bool SupportsDepthFetchDuringDepthTest() { return false; }
+	static bool SupportsDepthFetchDuringDepthTest() { return true; }
 	static bool SupportsTimestampRenderQueries() { return false; }
 
 	static bool RequiresMobileRenderer() { return true; }
diff --git a/Engine/Source/Runtime/VulkanRHI/Private/VulkanRenderTarget.cpp b/Engine/Source/Runtime/VulkanRHI/Private/VulkanRenderTarget.cpp
--- a/Engine/Source/Runtime/VulkanRHI/Private/VulkanRenderTarget.cpp
+++ b/Engine/Source/Runtime/VulkanRHI/Private/VulkanRenderTarget.cpp
@@ -1825,20 +1825,16 @@ FVulkanRenderTargetLayout::FVulkanRenderTargetLayout(FVulkanDevice& InDevice, co
 				NumUsedClearValues = NumAttachmentDescriptions + 1;
 			}
 		}
-		if (CurrDesc.samples == VK_SAMPLE_COUNT_1_BIT)
-		{
-			CurrDesc.storeOp = RenderTargetStoreActionToVulkan(GetStoreAction(GetDepthActions(RPInfo.DepthStencilRenderTarget.Action)), true);
-			CurrDesc.stencilStoreOp = RenderTargetStoreActionToVulkan(GetStoreAction(GetStencilActions(RPInfo.DepthStencilRenderTarget.Action)), true);
-		}
-		else
+
+		if (CurrDesc.samples != VK_SAMPLE_COUNT_1_BIT)
 		{
-			ensure(GetStoreAction(GetDepthActions(RPInfo.DepthStencilRenderTarget.Action)) == ERenderTargetStoreAction::ENoAction);
-			ensure(GetStoreAction(GetStencilActions(RPInfo.DepthStencilRenderTarget.Action)) == ERenderTargetStoreAction::ENoAction);
-			// Never want to store MSAA depth/stencil
-			CurrDesc.storeOp = VK_ATTACHMENT_STORE_OP_DONT_CARE;
-			CurrDesc.stencilStoreOp = VK_ATTACHMENT_STORE_OP_DONT_CARE;
+			// Can't resolve MSAA depth/stencil
+			ensure(GetStoreAction(GetDepthActions(RPInfo.DepthStencilRenderTarget.Action)) != ERenderTargetStoreAction::EMultisampleResolve);
+			ensure(GetStoreAction(GetStencilActions(RPInfo.DepthStencilRenderTarget.Action)) != ERenderTargetStoreAction::EMultisampleResolve);
 		}
-
+		CurrDesc.storeOp = RenderTargetStoreActionToVulkan(GetStoreAction(GetDepthActions(RPInfo.DepthStencilRenderTarget.Action)), true);
+		CurrDesc.stencilStoreOp = RenderTargetStoreActionToVulkan(GetStoreAction(GetStencilActions(RPInfo.DepthStencilRenderTarget.Action)), true);
+		
 		DepthStencilLayout = VulkanRHI::GetDepthStencilLayout(RPInfo.DepthStencilRenderTarget.ExclusiveDepthStencil, InDevice);
 		CurrDesc.initialLayout = DepthStencilLayout;
 		CurrDesc.finalLayout = DepthStencilLayout;
