Vulkan: fix overbright rendering with some drivers (UE-58570).

- The underlying reason were multiple transitions from Undef->Writable between rendering into different cubemap faces.
- Also contains Rolando's CL 4127455.

diff --git a/Engine/Source/Runtime/Renderer/Private/ReflectionEnvironmentCapture.cpp b/Engine/Source/Runtime/Renderer/Private/ReflectionEnvironmentCapture.cpp
--- a/Engine/Source/Runtime/Renderer/Private/ReflectionEnvironmentCapture.cpp
+++ b/Engine/Source/Runtime/Renderer/Private/ReflectionEnvironmentCapture.cpp
@@ -308,11 +308,13 @@ void FilterReflectionEnvironment(FRHICommandListImmediate& RHICmdList, ERHIFeatu
 	GraphicsPSOInit.DepthStencilState = TStaticDepthStencilState<false, CF_Always>::GetRHI();
 	GraphicsPSOInit.BlendState = TStaticBlendState<CW_RGBA, BO_Add, BF_Zero, BF_DestAlpha, BO_Add, BF_Zero, BF_One>::GetRHI();
 
+	RHICmdList.TransitionResource(EResourceTransitionAccess::EWritable, EffectiveColorRT.TargetableTexture);
+
 	// Premultiply alpha in-place using alpha blending
 	for (uint32 CubeFace = 0; CubeFace < CubeFace_MAX; CubeFace++)
 	{
-		SetRenderTarget(RHICmdList, EffectiveColorRT.TargetableTexture, 0, CubeFace, NULL, true);
-		RHICmdList.ApplyCachedRenderTargets(GraphicsPSOInit);
+		FRHIRenderPassInfo RPInfo(EffectiveColorRT.TargetableTexture, ERenderTargetActions::Clear_Store, nullptr, 0, CubeFace);
+		RHICmdList.BeginRenderPass(RPInfo, TEXT("FilterReflectionEnvironmentRP"));
 
 		const FIntPoint SourceDimensions(CubmapSize, CubmapSize);
 		const FIntRect ViewRect(0, 0, EffectiveTopMipSize, EffectiveTopMipSize);
@@ -341,9 +343,11 @@ void FilterReflectionEnvironment(FRHICommandListImmediate& RHICmdList, ERHIFeatu
 			SourceDimensions,
 			*VertexShader);
 
-		RHICmdList.CopyToResolveTarget(EffectiveColorRT.TargetableTexture, EffectiveColorRT.ShaderResourceTexture, FResolveParams(FResolveRect(), (ECubeFace)CubeFace));
+		RHICmdList.EndRenderPass();
 	}
 
+	RHICmdList.TransitionResource(EResourceTransitionAccess::EReadable, EffectiveColorRT.TargetableTexture);
+
 	auto ShaderMap = GetGlobalShaderMap(FeatureLevel);
 	FSceneRenderTargets& SceneContext = FSceneRenderTargets::Get(RHICmdList);
 
@@ -369,6 +373,8 @@ void FilterReflectionEnvironment(FRHICommandListImmediate& RHICmdList, ERHIFeatu
 		GraphicsPSOInit.DepthStencilState = TStaticDepthStencilState<false, CF_Always>::GetRHI();
 		GraphicsPSOInit.BlendState = TStaticBlendState<>::GetRHI();
 
+		RHICmdList.TransitionResource(EResourceTransitionAccess::EWritable, FilteredCube.TargetableTexture);
+
 		// Filter all the mips
 		for (int32 MipIndex = 0; MipIndex < NumMips; MipIndex++)
 		{
@@ -376,8 +382,8 @@ void FilterReflectionEnvironment(FRHICommandListImmediate& RHICmdList, ERHIFeatu
 
 			for (int32 CubeFace = 0; CubeFace < CubeFace_MAX; CubeFace++)
 			{
-				SetRenderTarget(RHICmdList, FilteredCube.TargetableTexture, MipIndex, CubeFace, NULL, true);
-				RHICmdList.ApplyCachedRenderTargets(GraphicsPSOInit);
+				FRHIRenderPassInfo RPInfo(FilteredCube.TargetableTexture, ERenderTargetActions::DontLoad_Store, nullptr, MipIndex, CubeFace);
+				RHICmdList.BeginRenderPass(RPInfo, TEXT("FilterCubeMapRP"));
 
 				const FIntRect ViewRect(0, 0, MipSize, MipSize);
 				RHICmdList.SetViewport(0, 0, 0.0f, MipSize, MipSize, 1.0f);
@@ -425,7 +431,7 @@ void FilterReflectionEnvironment(FRHICommandListImmediate& RHICmdList, ERHIFeatu
 					FIntPoint(MipSize, MipSize),
 					*VertexShader);
 
-				RHICmdList.CopyToResolveTarget(FilteredCube.TargetableTexture, FilteredCube.ShaderResourceTexture, FResolveParams(FResolveRect(), (ECubeFace)CubeFace, MipIndex));
+				RHICmdList.EndRenderPass();
 			}
 		}
 	}
@@ -655,12 +661,12 @@ void ClearScratchCubemaps(FRHICommandList& RHICmdList, int32 TargetSize)
 	{
 		SCOPED_DRAW_EVENT(RHICmdList, ClearScratchCubemapsRT0);
 
+		TransitionSetRenderTargetsHelper(RHICmdList, RT0.TargetableTexture, FTextureRHIParamRef(), FExclusiveDepthStencil::DepthWrite_StencilWrite);
+
 		for (int32 MipIndex = 0; MipIndex < NumMips; MipIndex++)
 		{
 			for (int32 CubeFace = 0; CubeFace < CubeFace_MAX; CubeFace++)
 			{
-				TransitionSetRenderTargetsHelper(RHICmdList, RT0.TargetableTexture, FTextureRHIParamRef(), FExclusiveDepthStencil::DepthWrite_StencilWrite);
-
 				FRHIRenderTargetView RtView = FRHIRenderTargetView(RT0.TargetableTexture, ERenderTargetLoadAction::EClear, MipIndex, CubeFace);
 				FRHISetRenderTargetsInfo Info(1, &RtView, FRHIDepthRenderTargetView());
 				RHICmdList.SetRenderTargetsAndClear(Info);
@@ -674,12 +680,12 @@ void ClearScratchCubemaps(FRHICommandList& RHICmdList, int32 TargetSize)
 		FSceneRenderTargetItem& RT1 = SceneContext.ReflectionColorScratchCubemap[1]->GetRenderTargetItem();
 		NumMips = (int32)RT1.TargetableTexture->GetNumMips();
 
+		TransitionSetRenderTargetsHelper(RHICmdList, RT1.TargetableTexture, FTextureRHIParamRef(), FExclusiveDepthStencil::DepthWrite_StencilWrite);
+
 		for (int32 MipIndex = 0; MipIndex < NumMips; MipIndex++)
 		{
 			for (int32 CubeFace = 0; CubeFace < CubeFace_MAX; CubeFace++)
 			{
-				TransitionSetRenderTargetsHelper(RHICmdList, RT1.TargetableTexture, FTextureRHIParamRef(), FExclusiveDepthStencil::DepthWrite_StencilWrite);
-
 				FRHIRenderTargetView RtView = FRHIRenderTargetView(RT1.TargetableTexture, ERenderTargetLoadAction::EClear, MipIndex, CubeFace);
 				FRHISetRenderTargetsInfo Info(1, &RtView, FRHIDepthRenderTargetView());
 				RHICmdList.SetRenderTargetsAndClear(Info);
@@ -721,12 +727,14 @@ void CaptureSceneToScratchCubemap(FRHICommandListImmediate& RHICmdList, FSceneRe
 
 		const int32 EffectiveSize = CubemapSize;
 		FSceneRenderTargetItem& EffectiveColorRT =  SceneContext.ReflectionColorScratchCubemap[0]->GetRenderTargetItem();
+		RHICmdList.TransitionResource(EResourceTransitionAccess::EWritable, EffectiveColorRT.TargetableTexture);
 
 		{
 			SCOPED_DRAW_EVENT(RHICmdList, CubeMapCopyScene);
 			
 			// Copy the captured scene into the cubemap face
-			SetRenderTarget(RHICmdList, EffectiveColorRT.TargetableTexture, 0, CubeFace, NULL);
+			FRHIRenderPassInfo RPInfo(EffectiveColorRT.TargetableTexture, ERenderTargetActions::DontLoad_Store, nullptr, 0, CubeFace);
+			RHICmdList.BeginRenderPass(RPInfo, TEXT("CubeMapCopySceneRP"));
 
 			const FIntRect ViewRect(0, 0, EffectiveSize, EffectiveSize);
 			RHICmdList.SetViewport(0, 0, 0.0f, EffectiveSize, EffectiveSize, 1.0f);
@@ -760,7 +768,7 @@ void CaptureSceneToScratchCubemap(FRHICommandListImmediate& RHICmdList, FSceneRe
 				SceneContext.GetBufferSizeXY(),
 				*VertexShader);
 
-			RHICmdList.CopyToResolveTarget(EffectiveColorRT.TargetableTexture, EffectiveColorRT.ShaderResourceTexture, FResolveParams(FResolveRect(), CubeFace));
+			RHICmdList.EndRenderPass();
 		}
 	}
 
@@ -775,10 +783,13 @@ void CopyCubemapToScratchCubemap(FRHICommandList& RHICmdList, ERHIFeatureLevel::
 	const int32 EffectiveSize = CubemapSize;
 	FSceneRenderTargetItem& EffectiveColorRT =  FSceneRenderTargets::Get(RHICmdList).ReflectionColorScratchCubemap[0]->GetRenderTargetItem();
 
+	RHICmdList.TransitionResource(EResourceTransitionAccess::EWritable, EffectiveColorRT.TargetableTexture);
+
 	for (uint32 CubeFace = 0; CubeFace < CubeFace_MAX; CubeFace++)
 	{
 		// Copy the captured scene into the cubemap face
-		SetRenderTarget(RHICmdList, EffectiveColorRT.TargetableTexture, 0, CubeFace, NULL, true);
+		FRHIRenderPassInfo RPInfo(EffectiveColorRT.TargetableTexture, ERenderTargetActions::DontLoad_Store, nullptr, 0, CubeFace);
+		RHICmdList.BeginRenderPass(RPInfo, TEXT("CopyCubemapToScratchCubemapRP"));
 
 		const FTexture* SourceCubemapResource = SourceCubemap->Resource;
 		const FIntPoint SourceDimensions(SourceCubemapResource->GetSizeX(), SourceCubemapResource->GetSizeY());
@@ -812,7 +823,7 @@ void CopyCubemapToScratchCubemap(FRHICommandList& RHICmdList, ERHIFeatureLevel::
 			SourceDimensions,
 			*VertexShader);
 
-		RHICmdList.CopyToResolveTarget(EffectiveColorRT.TargetableTexture, EffectiveColorRT.ShaderResourceTexture, FResolveParams(FResolveRect(), (ECubeFace)CubeFace));
+		RHICmdList.EndRenderPass();
 	}
 }
 
diff --git a/Engine/Source/Runtime/Renderer/Private/ReflectionEnvironmentDiffuseIrradiance.cpp b/Engine/Source/Runtime/Renderer/Private/ReflectionEnvironmentDiffuseIrradiance.cpp
--- a/Engine/Source/Runtime/Renderer/Private/ReflectionEnvironmentDiffuseIrradiance.cpp
+++ b/Engine/Source/Runtime/Renderer/Private/ReflectionEnvironmentDiffuseIrradiance.cpp
@@ -242,11 +242,13 @@ void ComputeDiffuseIrradiance(FRHICommandListImmediate& RHICmdList, ERHIFeatureL
 			const int32 MipIndex = 0;
 			const int32 MipSize = GDiffuseIrradianceCubemapSize;
 			FSceneRenderTargetItem& EffectiveRT = GetEffectiveDiffuseIrradianceRenderTarget(SceneContext, MipIndex);			
-			
+
+			RHICmdList.TransitionResource(EResourceTransitionAccess::EWritable, EffectiveRT.TargetableTexture);
+
 			for (int32 CubeFace = 0; CubeFace < CubeFace_MAX; CubeFace++)
 			{
-				SetRenderTarget(RHICmdList, EffectiveRT.TargetableTexture, 0, CubeFace, NULL, true);
-				RHICmdList.ApplyCachedRenderTargets(GraphicsPSOInit);
+				FRHIRenderPassInfo RPInfo(EffectiveRT.TargetableTexture, ERenderTargetActions::DontLoad_Store, nullptr, 0, CubeFace);
+				RHICmdList.BeginRenderPass(RPInfo, TEXT("CopyDiffuseIrradianceRP"));
 
 				const FIntRect ViewRect(0, 0, MipSize, MipSize);
 				RHICmdList.SetViewport(0, 0, 0.0f, MipSize, MipSize, 1.0f);
@@ -273,8 +275,10 @@ void ComputeDiffuseIrradiance(FRHICommandListImmediate& RHICmdList, ERHIFeatureL
 					FIntPoint(MipSize, MipSize),
 					*VertexShader);
 					
-				RHICmdList.CopyToResolveTarget(EffectiveRT.TargetableTexture, EffectiveRT.ShaderResourceTexture, FResolveParams(FResolveRect(), (ECubeFace)CubeFace, MipIndex));
-			}			
+				RHICmdList.EndRenderPass();
+			}
+
+			RHICmdList.TransitionResource(EResourceTransitionAccess::EReadable, EffectiveRT.TargetableTexture);
 		}
 
 		const int32 NumMips = FMath::CeilLogTwo(GDiffuseIrradianceCubemapSize) + 1;
@@ -289,11 +293,13 @@ void ComputeDiffuseIrradiance(FRHICommandListImmediate& RHICmdList, ERHIFeatureL
 				FSceneRenderTargetItem& EffectiveRT = GetEffectiveDiffuseIrradianceRenderTarget(SceneContext, MipIndex);
 				FSceneRenderTargetItem& EffectiveSource = GetEffectiveDiffuseIrradianceSourceTexture(SceneContext, MipIndex);
 				check(EffectiveRT.TargetableTexture != EffectiveSource.ShaderResourceTexture);
-				
+
+				RHICmdList.TransitionResource(EResourceTransitionAccess::EWritable, EffectiveRT.TargetableTexture);
+
 				for (int32 CubeFace = 0; CubeFace < CubeFace_MAX; CubeFace++)
 				{
-					SetRenderTarget(RHICmdList, EffectiveRT.TargetableTexture, MipIndex, CubeFace, NULL, true);
-					RHICmdList.ApplyCachedRenderTargets(GraphicsPSOInit);
+					FRHIRenderPassInfo RPInfo(EffectiveRT.TargetableTexture, ERenderTargetActions::Load_Store, nullptr, MipIndex, CubeFace);
+					RHICmdList.BeginRenderPass(RPInfo, TEXT("AccumulateDiffuseIrradianceRP"));
 
 					const FIntRect ViewRect(0, 0, MipSize, MipSize);
 					RHICmdList.SetViewport(0, 0, 0.0f, MipSize, MipSize, 1.0f);
@@ -320,8 +326,10 @@ void ComputeDiffuseIrradiance(FRHICommandListImmediate& RHICmdList, ERHIFeatureL
 						FIntPoint(MipSize, MipSize),
 						*VertexShader);
 						
-					RHICmdList.CopyToResolveTarget(EffectiveRT.TargetableTexture, EffectiveRT.ShaderResourceTexture, FResolveParams(FResolveRect(), (ECubeFace)CubeFace, MipIndex));
+					RHICmdList.EndRenderPass();
 				}				
+
+				RHICmdList.TransitionResource(EResourceTransitionAccess::EReadable, EffectiveRT.TargetableTexture);
 			}
 		}
 
@@ -330,10 +338,9 @@ void ComputeDiffuseIrradiance(FRHICommandListImmediate& RHICmdList, ERHIFeatureL
 			FSceneRenderTargetItem& EffectiveRT = FSceneRenderTargets::Get(RHICmdList).SkySHIrradianceMap->GetRenderTargetItem();
 
 			//load/store actions so we don't lose results as we render one pixel at a time on tile renderers.
-			FRHIRenderTargetView RTV(EffectiveRT.TargetableTexture, 0, -1, ERenderTargetLoadAction::ELoad, ERenderTargetStoreAction::EStore);
 			RHICmdList.TransitionResource(EResourceTransitionAccess::EWritable, EffectiveRT.TargetableTexture);
-			RHICmdList.SetRenderTargets(1, &RTV, nullptr, 0, nullptr);
-			RHICmdList.ApplyCachedRenderTargets(GraphicsPSOInit);
+			FRHIRenderPassInfo RPInfo(EffectiveRT.TargetableTexture, ERenderTargetActions::Load_Store, nullptr);
+			RHICmdList.BeginRenderPass(RPInfo, TEXT("GatherCoeffRP"));
 
 			const FIntRect ViewRect(CoefficientIndex, 0, CoefficientIndex + 1, 1);
 			RHICmdList.SetViewport(0, 0, 0.0f, FSHVector3::MaxSHBasis, 1, 1.0f);
@@ -363,7 +370,7 @@ void ComputeDiffuseIrradiance(FRHICommandListImmediate& RHICmdList, ERHIFeatureL
 				FIntPoint(MipSize, MipSize),
 				*VertexShader);
 
-			RHICmdList.CopyToResolveTarget(EffectiveRT.TargetableTexture, EffectiveRT.ShaderResourceTexture, FResolveParams());
+			RHICmdList.EndRenderPass();
 		}
 	}
 
diff --git a/Engine/Source/Runtime/VulkanRHI/Private/VulkanRenderTarget.cpp b/Engine/Source/Runtime/VulkanRHI/Private/VulkanRenderTarget.cpp
--- a/Engine/Source/Runtime/VulkanRHI/Private/VulkanRenderTarget.cpp
+++ b/Engine/Source/Runtime/VulkanRHI/Private/VulkanRenderTarget.cpp
@@ -1238,22 +1238,28 @@ void FVulkanCommandListContext::TransitionResources(const FPendingTransition& Pe
 
 				if ((AspectMask & VK_IMAGE_ASPECT_COLOR_BIT) != 0)
 				{
-					// When transitioning a full 3d/cube texture just assume it's undefined
-					if (SrcLayout == VK_IMAGE_LAYOUT_UNDEFINED || (Surface.GetNumMips() > 1 && Surface.GetNumberOfArrayLevels() > 1))
+					if (SrcLayout != VK_IMAGE_LAYOUT_COLOR_ATTACHMENT_OPTIMAL)
 					{
-						VulkanRHI::ImagePipelineBarrier(InCmdBuffer, Surface.Image, EImageLayoutBarrier::Undefined, EImageLayoutBarrier::ColorAttachment, SubresourceRange);
-					}
-					else
-					{
-						VulkanSetImageLayout(InCmdBuffer, Surface.Image, SrcLayout, VK_IMAGE_LAYOUT_COLOR_ATTACHMENT_OPTIMAL, SubresourceRange);
+						// When transitioning a full 3d/cube texture just assume it's undefined
+						if (SrcLayout == VK_IMAGE_LAYOUT_UNDEFINED || (Surface.GetNumMips() > 1 && Surface.GetNumberOfArrayLevels() > 1))
+						{
+							VulkanRHI::ImagePipelineBarrier(InCmdBuffer, Surface.Image, EImageLayoutBarrier::Undefined, EImageLayoutBarrier::ColorAttachment, SubresourceRange);
+						}
+						else
+						{
+							VulkanSetImageLayout(InCmdBuffer, Surface.Image, SrcLayout, VK_IMAGE_LAYOUT_COLOR_ATTACHMENT_OPTIMAL, SubresourceRange);
+						}
+						SrcLayout = VK_IMAGE_LAYOUT_COLOR_ATTACHMENT_OPTIMAL;
 					}
-					SrcLayout = VK_IMAGE_LAYOUT_COLOR_ATTACHMENT_OPTIMAL;
 				}
 				else
 				{
-					check(Surface.IsDepthOrStencilAspect());
-					VulkanSetImageLayout(InCmdBuffer, Surface.Image, SrcLayout, VK_IMAGE_LAYOUT_DEPTH_STENCIL_ATTACHMENT_OPTIMAL, SubresourceRange);
-					SrcLayout = VK_IMAGE_LAYOUT_DEPTH_STENCIL_ATTACHMENT_OPTIMAL;
+					if (SrcLayout != VK_IMAGE_LAYOUT_DEPTH_STENCIL_ATTACHMENT_OPTIMAL)
+					{
+						check(Surface.IsDepthOrStencilAspect());
+						VulkanSetImageLayout(InCmdBuffer, Surface.Image, SrcLayout, VK_IMAGE_LAYOUT_DEPTH_STENCIL_ATTACHMENT_OPTIMAL, SubresourceRange);
+						SrcLayout = VK_IMAGE_LAYOUT_DEPTH_STENCIL_ATTACHMENT_OPTIMAL;
+					}
 				}
 			};
 
