Back out changelist 4133715 due to screenshot differences and other errors.

diff --git a/Engine/Source/Runtime/Renderer/Private/ReflectionEnvironmentCapture.cpp b/Engine/Source/Runtime/Renderer/Private/ReflectionEnvironmentCapture.cpp
--- a/Engine/Source/Runtime/Renderer/Private/ReflectionEnvironmentCapture.cpp
+++ b/Engine/Source/Runtime/Renderer/Private/ReflectionEnvironmentCapture.cpp
@@ -308,13 +308,11 @@ void FilterReflectionEnvironment(FRHICommandListImmediate& RHICmdList, ERHIFeatu
 	GraphicsPSOInit.DepthStencilState = TStaticDepthStencilState<false, CF_Always>::GetRHI();
 	GraphicsPSOInit.BlendState = TStaticBlendState<CW_RGBA, BO_Add, BF_Zero, BF_DestAlpha, BO_Add, BF_Zero, BF_One>::GetRHI();
 
-	RHICmdList.TransitionResource(EResourceTransitionAccess::EWritable, EffectiveColorRT.TargetableTexture);
-
 	// Premultiply alpha in-place using alpha blending
 	for (uint32 CubeFace = 0; CubeFace < CubeFace_MAX; CubeFace++)
 	{
-		FRHIRenderPassInfo RPInfo(EffectiveColorRT.TargetableTexture, ERenderTargetActions::Clear_Store, nullptr, 0, CubeFace);
-		RHICmdList.BeginRenderPass(RPInfo, TEXT("FilterReflectionEnvironmentRP"));
+		SetRenderTarget(RHICmdList, EffectiveColorRT.TargetableTexture, 0, CubeFace, NULL, true);
+		RHICmdList.ApplyCachedRenderTargets(GraphicsPSOInit);
 
 		const FIntPoint SourceDimensions(CubmapSize, CubmapSize);
 		const FIntRect ViewRect(0, 0, EffectiveTopMipSize, EffectiveTopMipSize);
@@ -343,11 +341,9 @@ void FilterReflectionEnvironment(FRHICommandListImmediate& RHICmdList, ERHIFeatu
 			SourceDimensions,
 			*VertexShader);
 
-		RHICmdList.EndRenderPass();
+		RHICmdList.CopyToResolveTarget(EffectiveColorRT.TargetableTexture, EffectiveColorRT.ShaderResourceTexture, FResolveParams(FResolveRect(), (ECubeFace)CubeFace));
 	}
 
-	RHICmdList.TransitionResource(EResourceTransitionAccess::EReadable, EffectiveColorRT.TargetableTexture);
-
 	auto ShaderMap = GetGlobalShaderMap(FeatureLevel);
 	FSceneRenderTargets& SceneContext = FSceneRenderTargets::Get(RHICmdList);
 
@@ -373,8 +369,6 @@ void FilterReflectionEnvironment(FRHICommandListImmediate& RHICmdList, ERHIFeatu
 		GraphicsPSOInit.DepthStencilState = TStaticDepthStencilState<false, CF_Always>::GetRHI();
 		GraphicsPSOInit.BlendState = TStaticBlendState<>::GetRHI();
 
-		RHICmdList.TransitionResource(EResourceTransitionAccess::EWritable, FilteredCube.TargetableTexture);
-
 		// Filter all the mips
 		for (int32 MipIndex = 0; MipIndex < NumMips; MipIndex++)
 		{
@@ -382,8 +376,8 @@ void FilterReflectionEnvironment(FRHICommandListImmediate& RHICmdList, ERHIFeatu
 
 			for (int32 CubeFace = 0; CubeFace < CubeFace_MAX; CubeFace++)
 			{
-				FRHIRenderPassInfo RPInfo(FilteredCube.TargetableTexture, ERenderTargetActions::DontLoad_Store, nullptr, MipIndex, CubeFace);
-				RHICmdList.BeginRenderPass(RPInfo, TEXT("FilterCubeMapRP"));
+				SetRenderTarget(RHICmdList, FilteredCube.TargetableTexture, MipIndex, CubeFace, NULL, true);
+				RHICmdList.ApplyCachedRenderTargets(GraphicsPSOInit);
 
 				const FIntRect ViewRect(0, 0, MipSize, MipSize);
 				RHICmdList.SetViewport(0, 0, 0.0f, MipSize, MipSize, 1.0f);
@@ -431,7 +425,7 @@ void FilterReflectionEnvironment(FRHICommandListImmediate& RHICmdList, ERHIFeatu
 					FIntPoint(MipSize, MipSize),
 					*VertexShader);
 
-				RHICmdList.EndRenderPass();
+				RHICmdList.CopyToResolveTarget(FilteredCube.TargetableTexture, FilteredCube.ShaderResourceTexture, FResolveParams(FResolveRect(), (ECubeFace)CubeFace, MipIndex));
 			}
 		}
 	}
@@ -661,12 +655,12 @@ void ClearScratchCubemaps(FRHICommandList& RHICmdList, int32 TargetSize)
 	{
 		SCOPED_DRAW_EVENT(RHICmdList, ClearScratchCubemapsRT0);
 
-		TransitionSetRenderTargetsHelper(RHICmdList, RT0.TargetableTexture, FTextureRHIParamRef(), FExclusiveDepthStencil::DepthWrite_StencilWrite);
-
 		for (int32 MipIndex = 0; MipIndex < NumMips; MipIndex++)
 		{
 			for (int32 CubeFace = 0; CubeFace < CubeFace_MAX; CubeFace++)
 			{
+				TransitionSetRenderTargetsHelper(RHICmdList, RT0.TargetableTexture, FTextureRHIParamRef(), FExclusiveDepthStencil::DepthWrite_StencilWrite);
+
 				FRHIRenderTargetView RtView = FRHIRenderTargetView(RT0.TargetableTexture, ERenderTargetLoadAction::EClear, MipIndex, CubeFace);
 				FRHISetRenderTargetsInfo Info(1, &RtView, FRHIDepthRenderTargetView());
 				RHICmdList.SetRenderTargetsAndClear(Info);
@@ -680,12 +674,12 @@ void ClearScratchCubemaps(FRHICommandList& RHICmdList, int32 TargetSize)
 		FSceneRenderTargetItem& RT1 = SceneContext.ReflectionColorScratchCubemap[1]->GetRenderTargetItem();
 		NumMips = (int32)RT1.TargetableTexture->GetNumMips();
 
-		TransitionSetRenderTargetsHelper(RHICmdList, RT1.TargetableTexture, FTextureRHIParamRef(), FExclusiveDepthStencil::DepthWrite_StencilWrite);
-
 		for (int32 MipIndex = 0; MipIndex < NumMips; MipIndex++)
 		{
 			for (int32 CubeFace = 0; CubeFace < CubeFace_MAX; CubeFace++)
 			{
+				TransitionSetRenderTargetsHelper(RHICmdList, RT1.TargetableTexture, FTextureRHIParamRef(), FExclusiveDepthStencil::DepthWrite_StencilWrite);
+
 				FRHIRenderTargetView RtView = FRHIRenderTargetView(RT1.TargetableTexture, ERenderTargetLoadAction::EClear, MipIndex, CubeFace);
 				FRHISetRenderTargetsInfo Info(1, &RtView, FRHIDepthRenderTargetView());
 				RHICmdList.SetRenderTargetsAndClear(Info);
@@ -727,14 +721,12 @@ void CaptureSceneToScratchCubemap(FRHICommandListImmediate& RHICmdList, FSceneRe
 
 		const int32 EffectiveSize = CubemapSize;
 		FSceneRenderTargetItem& EffectiveColorRT =  SceneContext.ReflectionColorScratchCubemap[0]->GetRenderTargetItem();
-		RHICmdList.TransitionResource(EResourceTransitionAccess::EWritable, EffectiveColorRT.TargetableTexture);
 
 		{
 			SCOPED_DRAW_EVENT(RHICmdList, CubeMapCopyScene);
 			
 			// Copy the captured scene into the cubemap face
-			FRHIRenderPassInfo RPInfo(EffectiveColorRT.TargetableTexture, ERenderTargetActions::DontLoad_Store, nullptr, 0, CubeFace);
-			RHICmdList.BeginRenderPass(RPInfo, TEXT("CubeMapCopySceneRP"));
+			SetRenderTarget(RHICmdList, EffectiveColorRT.TargetableTexture, 0, CubeFace, NULL);
 
 			const FIntRect ViewRect(0, 0, EffectiveSize, EffectiveSize);
 			RHICmdList.SetViewport(0, 0, 0.0f, EffectiveSize, EffectiveSize, 1.0f);
@@ -768,7 +760,7 @@ void CaptureSceneToScratchCubemap(FRHICommandListImmediate& RHICmdList, FSceneRe
 				SceneContext.GetBufferSizeXY(),
 				*VertexShader);
 
-			RHICmdList.EndRenderPass();
+			RHICmdList.CopyToResolveTarget(EffectiveColorRT.TargetableTexture, EffectiveColorRT.ShaderResourceTexture, FResolveParams(FResolveRect(), CubeFace));
 		}
 	}
 
@@ -783,13 +775,10 @@ void CopyCubemapToScratchCubemap(FRHICommandList& RHICmdList, ERHIFeatureLevel::
 	const int32 EffectiveSize = CubemapSize;
 	FSceneRenderTargetItem& EffectiveColorRT =  FSceneRenderTargets::Get(RHICmdList).ReflectionColorScratchCubemap[0]->GetRenderTargetItem();
 
-	RHICmdList.TransitionResource(EResourceTransitionAccess::EWritable, EffectiveColorRT.TargetableTexture);
-
 	for (uint32 CubeFace = 0; CubeFace < CubeFace_MAX; CubeFace++)
 	{
 		// Copy the captured scene into the cubemap face
-		FRHIRenderPassInfo RPInfo(EffectiveColorRT.TargetableTexture, ERenderTargetActions::DontLoad_Store, nullptr, 0, CubeFace);
-		RHICmdList.BeginRenderPass(RPInfo, TEXT("CopyCubemapToScratchCubemapRP"));
+		SetRenderTarget(RHICmdList, EffectiveColorRT.TargetableTexture, 0, CubeFace, NULL, true);
 
 		const FTexture* SourceCubemapResource = SourceCubemap->Resource;
 		const FIntPoint SourceDimensions(SourceCubemapResource->GetSizeX(), SourceCubemapResource->GetSizeY());
@@ -823,7 +812,7 @@ void CopyCubemapToScratchCubemap(FRHICommandList& RHICmdList, ERHIFeatureLevel::
 			SourceDimensions,
 			*VertexShader);
 
-		RHICmdList.EndRenderPass();
+		RHICmdList.CopyToResolveTarget(EffectiveColorRT.TargetableTexture, EffectiveColorRT.ShaderResourceTexture, FResolveParams(FResolveRect(), (ECubeFace)CubeFace));
 	}
 }
 
diff --git a/Engine/Source/Runtime/Renderer/Private/ReflectionEnvironmentDiffuseIrradiance.cpp b/Engine/Source/Runtime/Renderer/Private/ReflectionEnvironmentDiffuseIrradiance.cpp
--- a/Engine/Source/Runtime/Renderer/Private/ReflectionEnvironmentDiffuseIrradiance.cpp
+++ b/Engine/Source/Runtime/Renderer/Private/ReflectionEnvironmentDiffuseIrradiance.cpp
@@ -242,13 +242,11 @@ void ComputeDiffuseIrradiance(FRHICommandListImmediate& RHICmdList, ERHIFeatureL
 			const int32 MipIndex = 0;
 			const int32 MipSize = GDiffuseIrradianceCubemapSize;
 			FSceneRenderTargetItem& EffectiveRT = GetEffectiveDiffuseIrradianceRenderTarget(SceneContext, MipIndex);			
-
-			RHICmdList.TransitionResource(EResourceTransitionAccess::EWritable, EffectiveRT.TargetableTexture);
-
+			
 			for (int32 CubeFace = 0; CubeFace < CubeFace_MAX; CubeFace++)
 			{
-				FRHIRenderPassInfo RPInfo(EffectiveRT.TargetableTexture, ERenderTargetActions::DontLoad_Store, nullptr, 0, CubeFace);
-				RHICmdList.BeginRenderPass(RPInfo, TEXT("CopyDiffuseIrradianceRP"));
+				SetRenderTarget(RHICmdList, EffectiveRT.TargetableTexture, 0, CubeFace, NULL, true);
+				RHICmdList.ApplyCachedRenderTargets(GraphicsPSOInit);
 
 				const FIntRect ViewRect(0, 0, MipSize, MipSize);
 				RHICmdList.SetViewport(0, 0, 0.0f, MipSize, MipSize, 1.0f);
@@ -275,10 +273,8 @@ void ComputeDiffuseIrradiance(FRHICommandListImmediate& RHICmdList, ERHIFeatureL
 					FIntPoint(MipSize, MipSize),
 					*VertexShader);
 					
-				RHICmdList.EndRenderPass();
-			}
-
-			RHICmdList.TransitionResource(EResourceTransitionAccess::EReadable, EffectiveRT.TargetableTexture);
+				RHICmdList.CopyToResolveTarget(EffectiveRT.TargetableTexture, EffectiveRT.ShaderResourceTexture, FResolveParams(FResolveRect(), (ECubeFace)CubeFace, MipIndex));
+			}			
 		}
 
 		const int32 NumMips = FMath::CeilLogTwo(GDiffuseIrradianceCubemapSize) + 1;
@@ -293,13 +289,11 @@ void ComputeDiffuseIrradiance(FRHICommandListImmediate& RHICmdList, ERHIFeatureL
 				FSceneRenderTargetItem& EffectiveRT = GetEffectiveDiffuseIrradianceRenderTarget(SceneContext, MipIndex);
 				FSceneRenderTargetItem& EffectiveSource = GetEffectiveDiffuseIrradianceSourceTexture(SceneContext, MipIndex);
 				check(EffectiveRT.TargetableTexture != EffectiveSource.ShaderResourceTexture);
-
-				RHICmdList.TransitionResource(EResourceTransitionAccess::EWritable, EffectiveRT.TargetableTexture);
-
+				
 				for (int32 CubeFace = 0; CubeFace < CubeFace_MAX; CubeFace++)
 				{
-					FRHIRenderPassInfo RPInfo(EffectiveRT.TargetableTexture, ERenderTargetActions::Load_Store, nullptr, MipIndex, CubeFace);
-					RHICmdList.BeginRenderPass(RPInfo, TEXT("AccumulateDiffuseIrradianceRP"));
+					SetRenderTarget(RHICmdList, EffectiveRT.TargetableTexture, MipIndex, CubeFace, NULL, true);
+					RHICmdList.ApplyCachedRenderTargets(GraphicsPSOInit);
 
 					const FIntRect ViewRect(0, 0, MipSize, MipSize);
 					RHICmdList.SetViewport(0, 0, 0.0f, MipSize, MipSize, 1.0f);
@@ -326,10 +320,8 @@ void ComputeDiffuseIrradiance(FRHICommandListImmediate& RHICmdList, ERHIFeatureL
 						FIntPoint(MipSize, MipSize),
 						*VertexShader);
 						
-					RHICmdList.EndRenderPass();
+					RHICmdList.CopyToResolveTarget(EffectiveRT.TargetableTexture, EffectiveRT.ShaderResourceTexture, FResolveParams(FResolveRect(), (ECubeFace)CubeFace, MipIndex));
 				}				
-
-				RHICmdList.TransitionResource(EResourceTransitionAccess::EReadable, EffectiveRT.TargetableTexture);
 			}
 		}
 
@@ -338,9 +330,10 @@ void ComputeDiffuseIrradiance(FRHICommandListImmediate& RHICmdList, ERHIFeatureL
 			FSceneRenderTargetItem& EffectiveRT = FSceneRenderTargets::Get(RHICmdList).SkySHIrradianceMap->GetRenderTargetItem();
 
 			//load/store actions so we don't lose results as we render one pixel at a time on tile renderers.
+			FRHIRenderTargetView RTV(EffectiveRT.TargetableTexture, 0, -1, ERenderTargetLoadAction::ELoad, ERenderTargetStoreAction::EStore);
 			RHICmdList.TransitionResource(EResourceTransitionAccess::EWritable, EffectiveRT.TargetableTexture);
-			FRHIRenderPassInfo RPInfo(EffectiveRT.TargetableTexture, ERenderTargetActions::Load_Store, nullptr);
-			RHICmdList.BeginRenderPass(RPInfo, TEXT("GatherCoeffRP"));
+			RHICmdList.SetRenderTargets(1, &RTV, nullptr, 0, nullptr);
+			RHICmdList.ApplyCachedRenderTargets(GraphicsPSOInit);
 
 			const FIntRect ViewRect(CoefficientIndex, 0, CoefficientIndex + 1, 1);
 			RHICmdList.SetViewport(0, 0, 0.0f, FSHVector3::MaxSHBasis, 1, 1.0f);
@@ -370,7 +363,7 @@ void ComputeDiffuseIrradiance(FRHICommandListImmediate& RHICmdList, ERHIFeatureL
 				FIntPoint(MipSize, MipSize),
 				*VertexShader);
 
-			RHICmdList.EndRenderPass();
+			RHICmdList.CopyToResolveTarget(EffectiveRT.TargetableTexture, EffectiveRT.ShaderResourceTexture, FResolveParams());
 		}
 	}
 
diff --git a/Engine/Source/Runtime/VulkanRHI/Private/VulkanRenderTarget.cpp b/Engine/Source/Runtime/VulkanRHI/Private/VulkanRenderTarget.cpp
--- a/Engine/Source/Runtime/VulkanRHI/Private/VulkanRenderTarget.cpp
+++ b/Engine/Source/Runtime/VulkanRHI/Private/VulkanRenderTarget.cpp
@@ -1238,28 +1238,22 @@ void FVulkanCommandListContext::TransitionResources(const FPendingTransition& Pe
 
 				if ((AspectMask & VK_IMAGE_ASPECT_COLOR_BIT) != 0)
 				{
-					if (SrcLayout != VK_IMAGE_LAYOUT_COLOR_ATTACHMENT_OPTIMAL)
+					// When transitioning a full 3d/cube texture just assume it's undefined
+					if (SrcLayout == VK_IMAGE_LAYOUT_UNDEFINED || (Surface.GetNumMips() > 1 && Surface.GetNumberOfArrayLevels() > 1))
 					{
-						// When transitioning a full 3d/cube texture just assume it's undefined
-						if (SrcLayout == VK_IMAGE_LAYOUT_UNDEFINED || (Surface.GetNumMips() > 1 && Surface.GetNumberOfArrayLevels() > 1))
-						{
-							VulkanRHI::ImagePipelineBarrier(InCmdBuffer, Surface.Image, EImageLayoutBarrier::Undefined, EImageLayoutBarrier::ColorAttachment, SubresourceRange);
-						}
-						else
-						{
-							VulkanSetImageLayout(InCmdBuffer, Surface.Image, SrcLayout, VK_IMAGE_LAYOUT_COLOR_ATTACHMENT_OPTIMAL, SubresourceRange);
-						}
-						SrcLayout = VK_IMAGE_LAYOUT_COLOR_ATTACHMENT_OPTIMAL;
+						VulkanRHI::ImagePipelineBarrier(InCmdBuffer, Surface.Image, EImageLayoutBarrier::Undefined, EImageLayoutBarrier::ColorAttachment, SubresourceRange);
 					}
+					else
+					{
+						VulkanSetImageLayout(InCmdBuffer, Surface.Image, SrcLayout, VK_IMAGE_LAYOUT_COLOR_ATTACHMENT_OPTIMAL, SubresourceRange);
+					}
+					SrcLayout = VK_IMAGE_LAYOUT_COLOR_ATTACHMENT_OPTIMAL;
 				}
 				else
 				{
-					if (SrcLayout != VK_IMAGE_LAYOUT_DEPTH_STENCIL_ATTACHMENT_OPTIMAL)
-					{
-						check(Surface.IsDepthOrStencilAspect());
-						VulkanSetImageLayout(InCmdBuffer, Surface.Image, SrcLayout, VK_IMAGE_LAYOUT_DEPTH_STENCIL_ATTACHMENT_OPTIMAL, SubresourceRange);
-						SrcLayout = VK_IMAGE_LAYOUT_DEPTH_STENCIL_ATTACHMENT_OPTIMAL;
-					}
+					check(Surface.IsDepthOrStencilAspect());
+					VulkanSetImageLayout(InCmdBuffer, Surface.Image, SrcLayout, VK_IMAGE_LAYOUT_DEPTH_STENCIL_ATTACHMENT_OPTIMAL, SubresourceRange);
+					SrcLayout = VK_IMAGE_LAYOUT_DEPTH_STENCIL_ATTACHMENT_OPTIMAL;
 				}
 			};
 
