LightTangentAndSoftSourceRadius in local light struct not being initialized for simple lights.

diff --git a/Engine/Source/Runtime/Renderer/Private/LightGridInjection.cpp b/Engine/Source/Runtime/Renderer/Private/LightGridInjection.cpp
--- a/Engine/Source/Runtime/Renderer/Private/LightGridInjection.cpp
+++ b/Engine/Source/Runtime/Renderer/Private/LightGridInjection.cpp
@@ -628,6 +628,7 @@ void FDeferredShadingSceneRenderer::ComputeLightGrid(FRHICommandListImmediate& R
 					const uint32 PackedWInt = ((uint32)SimpleLightSourceLength16f.Encoded) | ((uint32)VolumetricScatteringIntensity16f.Encoded << 16);
 
 					LightData.SpotAnglesAndSourceRadiusPacked = FVector4(-2, 1, 0, *(float*)&PackedWInt);
+					LightData.LightTangentAndSoftSourceRadius = FVector4(1.0f, 0.0f, 0.0f, 0.0f);
 				}
 			}
 
