Vulkan: fix overbright rendering with some drivers (UE-58570), second version.

- Fixing:
 - unnecessary transitions Undef->Writable between rendering into different cubemap faces.
 - errorneous transition Undef->Writable that discarded the RT which we were multiplying in-place.

diff --git a/Engine/Source/Runtime/Renderer/Private/ReflectionEnvironmentCapture.cpp b/Engine/Source/Runtime/Renderer/Private/ReflectionEnvironmentCapture.cpp
--- a/Engine/Source/Runtime/Renderer/Private/ReflectionEnvironmentCapture.cpp
+++ b/Engine/Source/Runtime/Renderer/Private/ReflectionEnvironmentCapture.cpp
@@ -308,10 +308,13 @@ void FilterReflectionEnvironment(FRHICommandListImmediate& RHICmdList, ERHIFeatu
 	GraphicsPSOInit.DepthStencilState = TStaticDepthStencilState<false, CF_Always>::GetRHI();
 	GraphicsPSOInit.BlendState = TStaticBlendState<CW_RGBA, BO_Add, BF_Zero, BF_DestAlpha, BO_Add, BF_Zero, BF_One>::GetRHI();
 
+	RHICmdList.TransitionResource(EResourceTransitionAccess::EWritable, EffectiveColorRT.TargetableTexture);
+
 	// Premultiply alpha in-place using alpha blending
 	for (uint32 CubeFace = 0; CubeFace < CubeFace_MAX; CubeFace++)
 	{
-		SetRenderTarget(RHICmdList, EffectiveColorRT.TargetableTexture, 0, CubeFace, NULL, true);
+		FRHIRenderPassInfo RPInfo(EffectiveColorRT.TargetableTexture, ERenderTargetActions::Load_Store, nullptr, 0, CubeFace);
+		RHICmdList.BeginRenderPass(RPInfo, TEXT("FilterReflectionEnvironmentRP"));
 		RHICmdList.ApplyCachedRenderTargets(GraphicsPSOInit);
 
 		const FIntPoint SourceDimensions(CubmapSize, CubmapSize);
@@ -341,9 +344,11 @@ void FilterReflectionEnvironment(FRHICommandListImmediate& RHICmdList, ERHIFeatu
 			SourceDimensions,
 			*VertexShader);
 
-		RHICmdList.CopyToResolveTarget(EffectiveColorRT.TargetableTexture, EffectiveColorRT.ShaderResourceTexture, FResolveParams(FResolveRect(), (ECubeFace)CubeFace));
+		RHICmdList.EndRenderPass();
 	}
 
+	RHICmdList.TransitionResource(EResourceTransitionAccess::EReadable, EffectiveColorRT.TargetableTexture);
+
 	auto ShaderMap = GetGlobalShaderMap(FeatureLevel);
 	FSceneRenderTargets& SceneContext = FSceneRenderTargets::Get(RHICmdList);
 
@@ -369,6 +374,8 @@ void FilterReflectionEnvironment(FRHICommandListImmediate& RHICmdList, ERHIFeatu
 		GraphicsPSOInit.DepthStencilState = TStaticDepthStencilState<false, CF_Always>::GetRHI();
 		GraphicsPSOInit.BlendState = TStaticBlendState<>::GetRHI();
 
+		RHICmdList.TransitionResource(EResourceTransitionAccess::EWritable, FilteredCube.TargetableTexture);
+
 		// Filter all the mips
 		for (int32 MipIndex = 0; MipIndex < NumMips; MipIndex++)
 		{
@@ -376,7 +383,8 @@ void FilterReflectionEnvironment(FRHICommandListImmediate& RHICmdList, ERHIFeatu
 
 			for (int32 CubeFace = 0; CubeFace < CubeFace_MAX; CubeFace++)
 			{
-				SetRenderTarget(RHICmdList, FilteredCube.TargetableTexture, MipIndex, CubeFace, NULL, true);
+				FRHIRenderPassInfo RPInfo(FilteredCube.TargetableTexture, ERenderTargetActions::DontLoad_Store, nullptr, MipIndex, CubeFace);
+				RHICmdList.BeginRenderPass(RPInfo, TEXT("FilterCubeMapRP"));
 				RHICmdList.ApplyCachedRenderTargets(GraphicsPSOInit);
 
 				const FIntRect ViewRect(0, 0, MipSize, MipSize);
@@ -425,9 +433,11 @@ void FilterReflectionEnvironment(FRHICommandListImmediate& RHICmdList, ERHIFeatu
 					FIntPoint(MipSize, MipSize),
 					*VertexShader);
 
-				RHICmdList.CopyToResolveTarget(FilteredCube.TargetableTexture, FilteredCube.ShaderResourceTexture, FResolveParams(FResolveRect(), (ECubeFace)CubeFace, MipIndex));
+				RHICmdList.EndRenderPass();
 			}
 		}
+
+		RHICmdList.TransitionResource(EResourceTransitionAccess::EReadable, FilteredCube.TargetableTexture);
 	}
 }
 
@@ -655,12 +665,12 @@ void ClearScratchCubemaps(FRHICommandList& RHICmdList, int32 TargetSize)
 	{
 		SCOPED_DRAW_EVENT(RHICmdList, ClearScratchCubemapsRT0);
 
+		TransitionSetRenderTargetsHelper(RHICmdList, RT0.TargetableTexture, FTextureRHIParamRef(), FExclusiveDepthStencil::DepthWrite_StencilWrite);
+
 		for (int32 MipIndex = 0; MipIndex < NumMips; MipIndex++)
 		{
 			for (int32 CubeFace = 0; CubeFace < CubeFace_MAX; CubeFace++)
 			{
-				TransitionSetRenderTargetsHelper(RHICmdList, RT0.TargetableTexture, FTextureRHIParamRef(), FExclusiveDepthStencil::DepthWrite_StencilWrite);
-
 				FRHIRenderTargetView RtView = FRHIRenderTargetView(RT0.TargetableTexture, ERenderTargetLoadAction::EClear, MipIndex, CubeFace);
 				FRHISetRenderTargetsInfo Info(1, &RtView, FRHIDepthRenderTargetView());
 				RHICmdList.SetRenderTargetsAndClear(Info);
@@ -674,12 +684,12 @@ void ClearScratchCubemaps(FRHICommandList& RHICmdList, int32 TargetSize)
 		FSceneRenderTargetItem& RT1 = SceneContext.ReflectionColorScratchCubemap[1]->GetRenderTargetItem();
 		NumMips = (int32)RT1.TargetableTexture->GetNumMips();
 
+		TransitionSetRenderTargetsHelper(RHICmdList, RT1.TargetableTexture, FTextureRHIParamRef(), FExclusiveDepthStencil::DepthWrite_StencilWrite);
+
 		for (int32 MipIndex = 0; MipIndex < NumMips; MipIndex++)
 		{
 			for (int32 CubeFace = 0; CubeFace < CubeFace_MAX; CubeFace++)
 			{
-				TransitionSetRenderTargetsHelper(RHICmdList, RT1.TargetableTexture, FTextureRHIParamRef(), FExclusiveDepthStencil::DepthWrite_StencilWrite);
-
 				FRHIRenderTargetView RtView = FRHIRenderTargetView(RT1.TargetableTexture, ERenderTargetLoadAction::EClear, MipIndex, CubeFace);
 				FRHISetRenderTargetsInfo Info(1, &RtView, FRHIDepthRenderTargetView());
 				RHICmdList.SetRenderTargetsAndClear(Info);
@@ -721,12 +731,14 @@ void CaptureSceneToScratchCubemap(FRHICommandListImmediate& RHICmdList, FSceneRe
 
 		const int32 EffectiveSize = CubemapSize;
 		FSceneRenderTargetItem& EffectiveColorRT =  SceneContext.ReflectionColorScratchCubemap[0]->GetRenderTargetItem();
+		RHICmdList.TransitionResource(EResourceTransitionAccess::EWritable, EffectiveColorRT.TargetableTexture);
 
 		{
 			SCOPED_DRAW_EVENT(RHICmdList, CubeMapCopyScene);
 			
 			// Copy the captured scene into the cubemap face
-			SetRenderTarget(RHICmdList, EffectiveColorRT.TargetableTexture, 0, CubeFace, NULL);
+			FRHIRenderPassInfo RPInfo(EffectiveColorRT.TargetableTexture, ERenderTargetActions::DontLoad_Store, nullptr, 0, CubeFace);
+			RHICmdList.BeginRenderPass(RPInfo, TEXT("CubeMapCopySceneRP"));
 
 			const FIntRect ViewRect(0, 0, EffectiveSize, EffectiveSize);
 			RHICmdList.SetViewport(0, 0, 0.0f, EffectiveSize, EffectiveSize, 1.0f);
@@ -760,7 +772,7 @@ void CaptureSceneToScratchCubemap(FRHICommandListImmediate& RHICmdList, FSceneRe
 				SceneContext.GetBufferSizeXY(),
 				*VertexShader);
 
-			RHICmdList.CopyToResolveTarget(EffectiveColorRT.TargetableTexture, EffectiveColorRT.ShaderResourceTexture, FResolveParams(FResolveRect(), CubeFace));
+			RHICmdList.EndRenderPass();
 		}
 	}
 
@@ -775,10 +787,13 @@ void CopyCubemapToScratchCubemap(FRHICommandList& RHICmdList, ERHIFeatureLevel::
 	const int32 EffectiveSize = CubemapSize;
 	FSceneRenderTargetItem& EffectiveColorRT =  FSceneRenderTargets::Get(RHICmdList).ReflectionColorScratchCubemap[0]->GetRenderTargetItem();
 
+	RHICmdList.TransitionResource(EResourceTransitionAccess::EWritable, EffectiveColorRT.TargetableTexture);
+
 	for (uint32 CubeFace = 0; CubeFace < CubeFace_MAX; CubeFace++)
 	{
 		// Copy the captured scene into the cubemap face
-		SetRenderTarget(RHICmdList, EffectiveColorRT.TargetableTexture, 0, CubeFace, NULL, true);
+		FRHIRenderPassInfo RPInfo(EffectiveColorRT.TargetableTexture, ERenderTargetActions::DontLoad_Store, nullptr, 0, CubeFace);
+		RHICmdList.BeginRenderPass(RPInfo, TEXT("CopyCubemapToScratchCubemapRP"));
 
 		const FTexture* SourceCubemapResource = SourceCubemap->Resource;
 		const FIntPoint SourceDimensions(SourceCubemapResource->GetSizeX(), SourceCubemapResource->GetSizeY());
@@ -812,7 +827,7 @@ void CopyCubemapToScratchCubemap(FRHICommandList& RHICmdList, ERHIFeatureLevel::
 			SourceDimensions,
 			*VertexShader);
 
-		RHICmdList.CopyToResolveTarget(EffectiveColorRT.TargetableTexture, EffectiveColorRT.ShaderResourceTexture, FResolveParams(FResolveRect(), (ECubeFace)CubeFace));
+		RHICmdList.EndRenderPass();
 	}
 }
 
diff --git a/Engine/Source/Runtime/Renderer/Private/ReflectionEnvironmentDiffuseIrradiance.cpp b/Engine/Source/Runtime/Renderer/Private/ReflectionEnvironmentDiffuseIrradiance.cpp
--- a/Engine/Source/Runtime/Renderer/Private/ReflectionEnvironmentDiffuseIrradiance.cpp
+++ b/Engine/Source/Runtime/Renderer/Private/ReflectionEnvironmentDiffuseIrradiance.cpp
@@ -242,10 +242,13 @@ void ComputeDiffuseIrradiance(FRHICommandListImmediate& RHICmdList, ERHIFeatureL
 			const int32 MipIndex = 0;
 			const int32 MipSize = GDiffuseIrradianceCubemapSize;
 			FSceneRenderTargetItem& EffectiveRT = GetEffectiveDiffuseIrradianceRenderTarget(SceneContext, MipIndex);			
-			
+
+			RHICmdList.TransitionResource(EResourceTransitionAccess::EWritable, EffectiveRT.TargetableTexture);
+
 			for (int32 CubeFace = 0; CubeFace < CubeFace_MAX; CubeFace++)
 			{
-				SetRenderTarget(RHICmdList, EffectiveRT.TargetableTexture, 0, CubeFace, NULL, true);
+				FRHIRenderPassInfo RPInfo(EffectiveRT.TargetableTexture, ERenderTargetActions::DontLoad_Store, nullptr, 0, CubeFace);
+				RHICmdList.BeginRenderPass(RPInfo, TEXT("CopyDiffuseIrradianceRP"));
 				RHICmdList.ApplyCachedRenderTargets(GraphicsPSOInit);
 
 				const FIntRect ViewRect(0, 0, MipSize, MipSize);
@@ -253,7 +256,7 @@ void ComputeDiffuseIrradiance(FRHICommandListImmediate& RHICmdList, ERHIFeatureL
 				TShaderMapRef<FCopyDiffuseIrradiancePS> PixelShader(ShaderMap);
 					
 				TShaderMapRef<FScreenVS> VertexShader(GetGlobalShaderMap(FeatureLevel));
-					
+
 				GraphicsPSOInit.BoundShaderState.VertexDeclarationRHI = GFilterVertexDeclaration.VertexDeclarationRHI;
 				GraphicsPSOInit.BoundShaderState.VertexShaderRHI = GETSAFERHISHADER_VERTEX(*VertexShader);
 				GraphicsPSOInit.BoundShaderState.PixelShaderRHI = GETSAFERHISHADER_PIXEL(*PixelShader);
@@ -273,8 +276,10 @@ void ComputeDiffuseIrradiance(FRHICommandListImmediate& RHICmdList, ERHIFeatureL
 					FIntPoint(MipSize, MipSize),
 					*VertexShader);
 					
-				RHICmdList.CopyToResolveTarget(EffectiveRT.TargetableTexture, EffectiveRT.ShaderResourceTexture, FResolveParams(FResolveRect(), (ECubeFace)CubeFace, MipIndex));
-			}			
+				RHICmdList.EndRenderPass();
+			}
+
+			RHICmdList.TransitionResource(EResourceTransitionAccess::EReadable, EffectiveRT.TargetableTexture);
 		}
 
 		const int32 NumMips = FMath::CeilLogTwo(GDiffuseIrradianceCubemapSize) + 1;
@@ -289,10 +294,13 @@ void ComputeDiffuseIrradiance(FRHICommandListImmediate& RHICmdList, ERHIFeatureL
 				FSceneRenderTargetItem& EffectiveRT = GetEffectiveDiffuseIrradianceRenderTarget(SceneContext, MipIndex);
 				FSceneRenderTargetItem& EffectiveSource = GetEffectiveDiffuseIrradianceSourceTexture(SceneContext, MipIndex);
 				check(EffectiveRT.TargetableTexture != EffectiveSource.ShaderResourceTexture);
-				
+
+				RHICmdList.TransitionResource(EResourceTransitionAccess::EWritable, EffectiveRT.TargetableTexture);
+
 				for (int32 CubeFace = 0; CubeFace < CubeFace_MAX; CubeFace++)
 				{
-					SetRenderTarget(RHICmdList, EffectiveRT.TargetableTexture, MipIndex, CubeFace, NULL, true);
+					FRHIRenderPassInfo RPInfo(EffectiveRT.TargetableTexture, ERenderTargetActions::Load_Store, nullptr, MipIndex, CubeFace);
+					RHICmdList.BeginRenderPass(RPInfo, TEXT("AccumulateDiffuseIrradianceRP"));
 					RHICmdList.ApplyCachedRenderTargets(GraphicsPSOInit);
 
 					const FIntRect ViewRect(0, 0, MipSize, MipSize);
@@ -320,8 +328,10 @@ void ComputeDiffuseIrradiance(FRHICommandListImmediate& RHICmdList, ERHIFeatureL
 						FIntPoint(MipSize, MipSize),
 						*VertexShader);
 						
-					RHICmdList.CopyToResolveTarget(EffectiveRT.TargetableTexture, EffectiveRT.ShaderResourceTexture, FResolveParams(FResolveRect(), (ECubeFace)CubeFace, MipIndex));
+					RHICmdList.EndRenderPass();
 				}				
+
+				RHICmdList.TransitionResource(EResourceTransitionAccess::EReadable, EffectiveRT.TargetableTexture);
 			}
 		}
 
@@ -330,9 +340,9 @@ void ComputeDiffuseIrradiance(FRHICommandListImmediate& RHICmdList, ERHIFeatureL
 			FSceneRenderTargetItem& EffectiveRT = FSceneRenderTargets::Get(RHICmdList).SkySHIrradianceMap->GetRenderTargetItem();
 
 			//load/store actions so we don't lose results as we render one pixel at a time on tile renderers.
-			FRHIRenderTargetView RTV(EffectiveRT.TargetableTexture, 0, -1, ERenderTargetLoadAction::ELoad, ERenderTargetStoreAction::EStore);
 			RHICmdList.TransitionResource(EResourceTransitionAccess::EWritable, EffectiveRT.TargetableTexture);
-			RHICmdList.SetRenderTargets(1, &RTV, nullptr, 0, nullptr);
+			FRHIRenderPassInfo RPInfo(EffectiveRT.TargetableTexture, ERenderTargetActions::Load_Store, nullptr);
+			RHICmdList.BeginRenderPass(RPInfo, TEXT("GatherCoeffRP"));
 			RHICmdList.ApplyCachedRenderTargets(GraphicsPSOInit);
 
 			const FIntRect ViewRect(CoefficientIndex, 0, CoefficientIndex + 1, 1);
@@ -363,7 +373,8 @@ void ComputeDiffuseIrradiance(FRHICommandListImmediate& RHICmdList, ERHIFeatureL
 				FIntPoint(MipSize, MipSize),
 				*VertexShader);
 
-			RHICmdList.CopyToResolveTarget(EffectiveRT.TargetableTexture, EffectiveRT.ShaderResourceTexture, FResolveParams());
+			RHICmdList.EndRenderPass();
+			RHICmdList.TransitionResource(EResourceTransitionAccess::EReadable, EffectiveRT.TargetableTexture);
 		}
 	}
 
diff --git a/Engine/Source/Runtime/VulkanRHI/Private/VulkanRenderTarget.cpp b/Engine/Source/Runtime/VulkanRHI/Private/VulkanRenderTarget.cpp
--- a/Engine/Source/Runtime/VulkanRHI/Private/VulkanRenderTarget.cpp
+++ b/Engine/Source/Runtime/VulkanRHI/Private/VulkanRenderTarget.cpp
@@ -1238,22 +1238,20 @@ void FVulkanCommandListContext::TransitionResources(const FPendingTransition& Pe
 
 				if ((AspectMask & VK_IMAGE_ASPECT_COLOR_BIT) != 0)
 				{
-					// When transitioning a full 3d/cube texture just assume it's undefined
-					if (SrcLayout == VK_IMAGE_LAYOUT_UNDEFINED || (Surface.GetNumMips() > 1 && Surface.GetNumberOfArrayLevels() > 1))
-					{
-						VulkanRHI::ImagePipelineBarrier(InCmdBuffer, Surface.Image, EImageLayoutBarrier::Undefined, EImageLayoutBarrier::ColorAttachment, SubresourceRange);
-					}
-					else
+					if (SrcLayout != VK_IMAGE_LAYOUT_COLOR_ATTACHMENT_OPTIMAL)
 					{
 						VulkanSetImageLayout(InCmdBuffer, Surface.Image, SrcLayout, VK_IMAGE_LAYOUT_COLOR_ATTACHMENT_OPTIMAL, SubresourceRange);
+						SrcLayout = VK_IMAGE_LAYOUT_COLOR_ATTACHMENT_OPTIMAL;
 					}
-					SrcLayout = VK_IMAGE_LAYOUT_COLOR_ATTACHMENT_OPTIMAL;
 				}
 				else
 				{
-					check(Surface.IsDepthOrStencilAspect());
-					VulkanSetImageLayout(InCmdBuffer, Surface.Image, SrcLayout, VK_IMAGE_LAYOUT_DEPTH_STENCIL_ATTACHMENT_OPTIMAL, SubresourceRange);
-					SrcLayout = VK_IMAGE_LAYOUT_DEPTH_STENCIL_ATTACHMENT_OPTIMAL;
+					if (SrcLayout != VK_IMAGE_LAYOUT_DEPTH_STENCIL_ATTACHMENT_OPTIMAL)
+					{
+						check(Surface.IsDepthOrStencilAspect());
+						VulkanSetImageLayout(InCmdBuffer, Surface.Image, SrcLayout, VK_IMAGE_LAYOUT_DEPTH_STENCIL_ATTACHMENT_OPTIMAL, SubresourceRange);
+						SrcLayout = VK_IMAGE_LAYOUT_DEPTH_STENCIL_ATTACHMENT_OPTIMAL;
+					}
 				}
 			};
 
