Fixed crash inside SetResourcesFromTables - SceneTexturesStruct was required by the shader, but wasn't set.

diff --git a/Engine/Source/Runtime/Renderer/Private/TranslucentRendering.cpp b/Engine/Source/Runtime/Renderer/Private/TranslucentRendering.cpp
--- a/Engine/Source/Runtime/Renderer/Private/TranslucentRendering.cpp
+++ b/Engine/Source/Runtime/Renderer/Private/TranslucentRendering.cpp
@@ -1456,6 +1456,7 @@ protected:
 	{
 	}
 
+	FSceneTextureShaderParameters SceneTextureParameters;
 	FShaderParameter LowResColorTexelSize;
 	FShaderResourceParameter LowResDepthTexture;
 	FShaderResourceParameter LowResColorTexture;
@@ -1469,6 +1470,7 @@ public:
 		: FGlobalShader(Initializer)
 		, bUseNearestDepthNeighborUpsample(InbUseNearestDepthNeighborUpsample)
 	{
+		SceneTextureParameters.Bind(Initializer);
 		LowResColorTexelSize.Bind(Initializer.ParameterMap, TEXT("LowResColorTexelSize"));
 		LowResDepthTexture.Bind(Initializer.ParameterMap, TEXT("LowResDepthTexture"));
 		LowResColorTexture.Bind(Initializer.ParameterMap, TEXT("LowResColorTexture"));
@@ -1480,7 +1482,7 @@ public:
 	virtual bool Serialize(FArchive& Ar) override
 	{
 		bool bShaderHasOutdatedParameters = FGlobalShader::Serialize(Ar);
-		Ar << LowResColorTexelSize << LowResDepthTexture << LowResColorTexture << BilinearClampedSampler << PointClampedSampler;
+		Ar << SceneTextureParameters << LowResColorTexelSize << LowResDepthTexture << LowResColorTexture << BilinearClampedSampler << PointClampedSampler;
 		return bShaderHasOutdatedParameters;
 	}
 			
@@ -1501,6 +1503,8 @@ public:
 
 		SetSamplerParameter(RHICmdList, ShaderRHI, BilinearClampedSampler, TStaticSamplerState<SF_Bilinear,AM_Clamp,AM_Clamp,AM_Clamp>::GetRHI());
 		SetSamplerParameter(RHICmdList, ShaderRHI, PointClampedSampler, TStaticSamplerState<SF_Point,AM_Clamp,AM_Clamp,AM_Clamp>::GetRHI());
+
+		SceneTextureParameters.Set(RHICmdList, GetPixelShader(), View.FeatureLevel, ESceneTextureSetupMode::All);
 	}
 
 	const bool bUseNearestDepthNeighborUpsample;
