Fix extremely low framerates in ContentExamples on Linux Vulkan (UE-60421).

- Fallout of CL 3985258 that fixed UE-55756 (crashes after iOS orientation changes). Perhaps just exempting scene captures would be fine, but better to keep the whole workaround platform-specific.
- The performance hit likely manifested on other platforms/RHIs, but Linux Vulkan was affected particularly badly (due to long RT alloc times of ~20-30 ms).

diff --git a/Engine/Source/Runtime/Renderer/Private/PostProcess/SceneRenderTargets.cpp b/Engine/Source/Runtime/Renderer/Private/PostProcess/SceneRenderTargets.cpp
--- a/Engine/Source/Runtime/Renderer/Private/PostProcess/SceneRenderTargets.cpp
+++ b/Engine/Source/Runtime/Renderer/Private/PostProcess/SceneRenderTargets.cpp
@@ -341,7 +341,9 @@ FIntPoint FSceneRenderTargets::ComputeDesiredSize(const FSceneViewFamily& ViewFa
 		default:
 			checkNoEntry();
 	}
-	
+
+	// This is specific to iOS and should not matter elsewhere.
+#if PLATFORM_IOS
 	// Don't consider the history buffer when the aspect ratio changes, the existing buffers won't make much sense at all.
 	// This prevents problems when orientation changes on mobile in particular.
 	float DesiredAspectRatio = (float)DesiredBufferSize.X / (float)DesiredBufferSize.Y;
@@ -355,11 +357,16 @@ FIntPoint FSceneRenderTargets::ComputeDesiredSize(const FSceneViewFamily& ViewFa
 			bAspectRatioChanged = !FMath::IsNearlyEqual(DesiredAspectRatio, LargestAspectRatio);
 		}
 	}
+#endif // PLATFORM_IOS
 
 	// we want to shrink the buffer but as we can have multiple scenecaptures per frame we have to delay that a frame to get all size requests
 	// Don't save buffer size in history while making high-res screenshot.
 	// We have to use the requested size when allocating an hmd depth target to ensure it matches the hmd allocated render target size.
-	if(!GIsHighResScreenshot && !bHMDAllocatedDepthTarget && !bAspectRatioChanged)
+	if(!GIsHighResScreenshot && !bHMDAllocatedDepthTarget 
+#if PLATFORM_IOS
+		&& !bAspectRatioChanged
+#endif // PLATFORM_IOS
+		)
 	{
 		// this allows The BufferSize to not grow below the SceneCapture requests (happen before scene rendering, in the same frame with a Grow request)
 		LargestDesiredSizes[CurrentDesiredSizeIndex] = LargestDesiredSizes[CurrentDesiredSizeIndex].ComponentMax(DesiredBufferSize);
