Cherry-pick 4154962: Fixes a bug in DOF where a permutation of Reduce shader can be missing on Win64.

diff --git a/Engine/Source/Runtime/Renderer/Private/PostProcess/DiaphragmDOFGraph.cpp b/Engine/Source/Runtime/Renderer/Private/PostProcess/DiaphragmDOFGraph.cpp
--- a/Engine/Source/Runtime/Renderer/Private/PostProcess/DiaphragmDOFGraph.cpp
+++ b/Engine/Source/Runtime/Renderer/Private/PostProcess/DiaphragmDOFGraph.cpp
@@ -221,7 +221,8 @@ bool DiaphragmDOF::WireSceneColorPasses(FPostprocessContext& Context, const FRen
 
 		// This is just to get the number of shader permutation down.
 		RecombineQuality == 0 &&
-		bUseLowAccumulatorQuality);
+		bUseLowAccumulatorQuality &&
+		FRCPassDiaphragmDOFGather::SupportRGBColorBuffer(ShaderPlatform));
 
 
 	// Derives everything needed from the view.
