Add some extra data output to the check to track down UE-60169

diff --git a/Engine/Source/Runtime/Renderer/Private/ShaderBaseClasses.cpp b/Engine/Source/Runtime/Renderer/Private/ShaderBaseClasses.cpp
--- a/Engine/Source/Runtime/Renderer/Private/ShaderBaseClasses.cpp
+++ b/Engine/Source/Runtime/Renderer/Private/ShaderBaseClasses.cpp
@@ -176,7 +176,8 @@ void FMaterialShader::SetParametersInner(
 		(MaterialRenderProxy->IsSelected() || MaterialRenderProxy->IsHovered());
 
 	ERHIFeatureLevel::Type FeatureLevel = View.GetFeatureLevel();
-	check(Material.GetRenderingThreadShaderMap() && Material.GetRenderingThreadShaderMap()->IsValidForRendering() && Material.GetFeatureLevel() == FeatureLevel);
+	checkf(Material.GetRenderingThreadShaderMap(), TEXT("RenderingThreadShaderMap: %i"), Material.GetRenderingThreadShaderMap() ? 1 : 0);
+	checkf(Material.GetRenderingThreadShaderMap()->IsValidForRendering() && Material.GetFeatureLevel() == FeatureLevel, TEXT("IsValid:%i, MaterialFeatureLevel:%i, FeatureLevel:%i"), Material.GetRenderingThreadShaderMap()->IsValidForRendering() ? 1 : 0, Material.GetFeatureLevel(), FeatureLevel);
 
 	FUniformExpressionCache* UniformExpressionCache = &MaterialRenderProxy->UniformExpressionCache[FeatureLevel];
 	bool bUniformExpressionCacheNeedsDelete = false;
