Cherry-pick 4157036: Tightens coc tile dilatation as much as possible.

diff --git a/Engine/Source/Runtime/Renderer/Private/PostProcess/DiaphragmDOFGraph.cpp b/Engine/Source/Runtime/Renderer/Private/PostProcess/DiaphragmDOFGraph.cpp
--- a/Engine/Source/Runtime/Renderer/Private/PostProcess/DiaphragmDOFGraph.cpp
+++ b/Engine/Source/Runtime/Renderer/Private/PostProcess/DiaphragmDOFGraph.cpp
@@ -346,25 +346,43 @@ bool DiaphragmDOF::WireSceneColorPasses(FPostprocessContext& Context, const FRen
 		CocTileOutput = CocFlatten;
 
 		// Parameters for the dilate Coc passes.
-		FRCPassDiaphragmDOFDilateCoc::FParameters DilateParams[2];
+		FRCPassDiaphragmDOFDilateCoc::FParameters DilateParams[3];
 		{
+			const int32 MaxSampleRadiusCount = FRCPassDiaphragmDOFDilateCoc::MaxSampleRadiusCount;
+
 			// Compute the maximum tile dilation.
 			int32 MaximumTileDilation = FMath::CeilToInt(
 				MaxBluringRadius / FRCPassDiaphragmDOFFlattenCoc::CocTileResolutionDivisor);
 
 			// There is always at least one dilate pass so that even small Coc radius conservatively dilate on next neighboor.
 			DilateParams[0].SampleRadiusCount = FMath::Min(
-				MaximumTileDilation, FRCPassDiaphragmDOFDilateCoc::MaxSampleRadiusCount);
+				MaximumTileDilation, MaxSampleRadiusCount);
 			
-			// If the theoric radius is too big, setup second dilate pass.
-			if (MaximumTileDilation > DilateParams[0].SampleRadiusCount)
+			int32 CurrentConvolutionRadius = DilateParams[0].SampleRadiusCount;
+
+			// If the theoric radius is too big, setup more dilate passes.
+			for (int32 i = 1; i < ARRAY_COUNT(DilateParams); i++)
 			{
-				DilateParams[1].SampleDistanceMultiplier = DilateParams[0].SampleRadiusCount + 1;
-				DilateParams[1].SampleRadiusCount = FMath::Min(
-					FMath::DivideAndRoundUp(
-						MaximumTileDilation - DilateParams[0].SampleRadiusCount,
-						DilateParams[1].SampleDistanceMultiplier),
-					FRCPassDiaphragmDOFDilateCoc::MaxSampleRadiusCount);
+				if (MaximumTileDilation <= CurrentConvolutionRadius)
+				{
+					break;
+				}
+
+				// Highest upper bound possible for SampleDistanceMultiplier to not step over any tile.
+				int32 HighestPossibleMultiplierUpperBound = CurrentConvolutionRadius + 1;
+
+				// Find out how many step we need to do on dilate radius.
+				DilateParams[i].SampleRadiusCount = FMath::Min(
+					MaximumTileDilation / HighestPossibleMultiplierUpperBound,
+					MaxSampleRadiusCount);
+
+				// Find out ideal multiplier to not dilate an area too large.
+				// TODO: Could add control over the radius of the last.
+				int32 IdealMultiplier = FMath::DivideAndRoundUp(MaximumTileDilation - CurrentConvolutionRadius, DilateParams[1].SampleRadiusCount);
+
+				DilateParams[i].SampleDistanceMultiplier = FMath::Min(IdealMultiplier, HighestPossibleMultiplierUpperBound);
+
+				CurrentConvolutionRadius += DilateParams[i].SampleRadiusCount * DilateParams[i].SampleDistanceMultiplier;
 			}
 		}
 
