Cherry-pick 4157118: Exposes r.DOF.Scatter.NeighborCompareMaxColor

diff --git a/Engine/Shaders/Private/DiaphragmDOF/DOFHybridScatterCompilation.ush b/Engine/Shaders/Private/DiaphragmDOF/DOFHybridScatterCompilation.ush
--- a/Engine/Shaders/Private/DiaphragmDOF/DOFHybridScatterCompilation.ush
+++ b/Engine/Shaders/Private/DiaphragmDOF/DOFHybridScatterCompilation.ush
@@ -10,6 +10,8 @@
 
 //------------------------------------------------------- PARAMETERS
 
+float NeighborCompareMaxColor;
+
 // Parameters for ComputeScatterFactor.
 struct FHybridScatterInputs
 {
@@ -59,7 +61,7 @@ float ComputeNeighborBasedScatterFactor(in const FHybridScatterInputs Parameters
 		const float Threshold = 0.7;
 
 		float LinearScale = rcp(Parameters.FrameExposureScale);
-		float LinearColorMaximum = 3 * LinearScale;
+		float LinearColorMaximum = NeighborCompareMaxColor * LinearScale;
 
 		float4 SaturatedColor = min(Parameters.Color, LinearColorMaximum);
 
diff --git a/Engine/Source/Runtime/Renderer/Private/PostProcess/DiaphragmDOFPasses.cpp b/Engine/Source/Runtime/Renderer/Private/PostProcess/DiaphragmDOFPasses.cpp
--- a/Engine/Source/Runtime/Renderer/Private/PostProcess/DiaphragmDOFPasses.cpp
+++ b/Engine/Source/Runtime/Renderer/Private/PostProcess/DiaphragmDOFPasses.cpp
@@ -34,6 +34,14 @@ TAutoConsoleVariable<int32> CVarDebugScatterPerf(
 
 #endif
 
+TAutoConsoleVariable<float> CVarScatterNeighborCompareMaxColor(
+	TEXT("r.DOF.Scatter.NeighborCompareMaxColor"),
+	10,
+	TEXT("Controles the linear color clamping upperbound applied before color of pixel and neighbors are compared.")
+	TEXT(" To low, and you may not scatter enough; to high you may scatter unnecessarily too much in highlights")
+	TEXT(" (Default: 10)."),
+	ECVF_RenderThreadSafe);
+
 } // namespace
 
 
@@ -750,6 +758,7 @@ FPooledRenderTargetDesc FRCPassDiaphragmDOFSetup::ComputeOutputDesc(EPassOutputI
 	PARAMETER(FShaderParameter, MaxScatteringGroupCount) \
 	PARAMETER(FShaderParameter, PreProcessingToProcessingCocRadiusFactor) \
 	PARAMETER(FShaderParameter, MinScatteringCocRadius) \
+	PARAMETER(FShaderParameter, NeighborCompareMaxColor) \
 	PARAMETER(FShaderResourceParameter, OutScatterDrawIndirectParameters) \
 	PARAMETER(FShaderResourceParameter, OutForegroundScatterDrawList) \
 	PARAMETER(FShaderResourceParameter, OutBackgroundScatterDrawList) \
@@ -895,6 +904,9 @@ void FRCPassDiaphragmDOFReduce::Process(FRenderingCompositePassContext& Context)
 				DispatchCtx->OutBackgroundScatterDrawList, ScatterDrawListBuffer[1]->UAV);
 
 			SetShaderValue(Context.RHICmdList, DispatchCtx.ShaderRHI, DispatchCtx->MinScatteringCocRadius, Params.MinScatteringCocRadius);
+
+			float NeighborCompareMaxColor = CVarScatterNeighborCompareMaxColor.GetValueOnRenderThread();
+			SetShaderValue(Context.RHICmdList, DispatchCtx.ShaderRHI, DispatchCtx->NeighborCompareMaxColor, NeighborCompareMaxColor);
 		}
 
 		DispatchCtx.Dispatch();
