Cherry-pick 4152104: Enables Diaphragm DOF on Linux's Vulkan SM5, Mac's Metal SM5, and iOS SM5

diff --git a/Engine/Source/Runtime/Renderer/Private/PostProcess/DiaphragmDOF.h b/Engine/Source/Runtime/Renderer/Private/PostProcess/DiaphragmDOF.h
--- a/Engine/Source/Runtime/Renderer/Private/PostProcess/DiaphragmDOF.h
+++ b/Engine/Source/Runtime/Renderer/Private/PostProcess/DiaphragmDOF.h
@@ -12,7 +12,7 @@
 
 
 // Whitelist diaphragm DOF for platforms that actually have been tested.
-#define WITH_DIAPHRAGM_DOF (PLATFORM_WINDOWS || PLATFORM_XBOXONE || PLATFORM_PS4)
+#define WITH_DIAPHRAGM_DOF (PLATFORM_WINDOWS || PLATFORM_XBOXONE || PLATFORM_PS4 || PLATFORM_MAC || PLATFORM_LINUX || PLATFORM_IOS)
 
 
 namespace DiaphragmDOF
@@ -117,7 +117,15 @@ inline bool IsSupported(EShaderPlatform ShaderPlatform)
 	#endif
 
 	// Only compile diaphragm DOF on platform it has been tested to ensure this is not blocking anyone else.
-	return ShaderPlatform == SP_PCD3D_SM5 || ShaderPlatform == SP_XBOXONE_D3D12 || ShaderPlatform == SP_PS4 || IsVulkanSM5Platform(ShaderPlatform);
+	return 
+		ShaderPlatform == SP_PCD3D_SM5 ||
+		ShaderPlatform == SP_XBOXONE_D3D12 ||
+		ShaderPlatform == SP_PS4 ||
+		IsVulkanSM5Platform(ShaderPlatform) ||
+		ShaderPlatform == SP_METAL_SM5 ||
+		ShaderPlatform == SP_METAL_SM5_NOTESS ||
+		ShaderPlatform == SP_METAL_MRT ||
+		ShaderPlatform == SP_METAL_MRT_MAC;
 }
 
 
