Integrate as edit CL#4161511:
Fix mobile renderer's CopyToSkyTexture, fixes corrupt sky reflection.

diff --git a/Engine/Source/Runtime/Renderer/Private/MobileBasePassRendering.h b/Engine/Source/Runtime/Renderer/Private/MobileBasePassRendering.h
--- a/Engine/Source/Runtime/Renderer/Private/MobileBasePassRendering.h
+++ b/Engine/Source/Runtime/Renderer/Private/MobileBasePassRendering.h
@@ -245,7 +245,7 @@ public:
 			// if >0 this will disable shader's RGBM decoding and enable sky light tinting of this envmap.
 			FTexture* ReflectionTexture;
 			float AverageBrightness = 1.0f;
-			FVector4 MobileSkyReflectionValues;
+			FVector4 MobileSkyReflectionValues(ForceInit);
 			GetSkyTextureParams(RenderScene, AverageBrightness, ReflectionTexture, MobileSkyReflectionValues.X);
 			FRHIPixelShader* PixelShader = GetPixelShader();
 			// Set the reflection cubemap
@@ -322,7 +322,7 @@ public:
 		{
 			FTexture* ReflectionTexture = GBlackTextureCube;
 			float AverageBrightness = 1.0f;
-			FVector4 MobileSkyReflectionValues;
+			FVector4 MobileSkyReflectionValues(ForceInit);
 
 			if (PrimitiveSceneInfo 
 				&& PrimitiveSceneInfo->CachedReflectionCaptureProxy
diff --git a/Engine/Source/Runtime/Renderer/Private/MobileReflectionEnvironmentCapture.cpp b/Engine/Source/Runtime/Renderer/Private/MobileReflectionEnvironmentCapture.cpp
--- a/Engine/Source/Runtime/Renderer/Private/MobileReflectionEnvironmentCapture.cpp
+++ b/Engine/Source/Runtime/Renderer/Private/MobileReflectionEnvironmentCapture.cpp
@@ -185,6 +185,9 @@ namespace MobileReflectionEnvironmentCapture
 
 			RHICmdList.TransitionResource(EResourceTransitionAccess::EWritable, ProcessedTexture->TextureRHI);
 
+			FRHICopyTextureInfo CopyInfo(ProcessedTexture->GetSizeX(), ProcessedTexture->GetSizeY());
+			CopyInfo.NumArraySlices = 6;
+
 			// GPU copy back to the skylight's texture, which is not a render target
 			for (int32 MipIndex = 0; MipIndex < NumMips; MipIndex++)
 			{
@@ -192,8 +195,6 @@ namespace MobileReflectionEnvironmentCapture
 				// In the HQ case, the full image is contained in GetEffectiveRenderTarget(.., false,0).
 				FSceneRenderTargetItem& EffectiveSource = GetEffectiveRenderTarget(SceneContext, false, bUseHQFiltering ? 0 : MipIndex);
 				RHICmdList.TransitionResource(EResourceTransitionAccess::EReadable, EffectiveSource.ShaderResourceTexture);
-				FRHICopyTextureInfo CopyInfo(EffectiveSource.ShaderResourceTexture->GetSizeXYZ());
-				CopyInfo.NumArraySlices = 6;
 				RHICmdList.CopyTexture(EffectiveSource.ShaderResourceTexture, ProcessedTexture->TextureRHI, CopyInfo);
 				CopyInfo.AdvanceMip();
 			}
