Fix for hidden dynamic meshes with bCastHiddenShadows. bDrawRelevance is not required here to test if mesh shadow should be rendered.
Merged #4829 GitHub pull request.

diff --git a/Engine/Source/Runtime/Renderer/Private/ShadowSetup.cpp b/Engine/Source/Runtime/Renderer/Private/ShadowSetup.cpp
--- a/Engine/Source/Runtime/Renderer/Private/ShadowSetup.cpp
+++ b/Engine/Source/Runtime/Renderer/Private/ShadowSetup.cpp
@@ -1243,7 +1243,7 @@ void FProjectedShadowInfo::GatherDynamicMeshElementsArray(
 		}
 
 		// Only draw if the subject primitive is shadow relevant.
-		if (ViewRelevance.bShadowRelevance && ViewRelevance.bDynamicRelevance && ViewRelevance.bDrawRelevance)
+		if (ViewRelevance.bShadowRelevance && ViewRelevance.bDynamicRelevance)
 		{
 			Renderer.MeshCollector.SetPrimitive(PrimitiveSceneInfo->Proxy, PrimitiveSceneInfo->DefaultDynamicHitProxyId);
 
