https://jira.it.epicgames.net/browse/UE-61383

Turning off FFT Bloom on Metal until this bug is fixed.

diff --git a/Engine/Source/Runtime/Renderer/Private/PostProcess/PostProcessing.cpp b/Engine/Source/Runtime/Renderer/Private/PostProcess/PostProcessing.cpp
--- a/Engine/Source/Runtime/Renderer/Private/PostProcess/PostProcessing.cpp
+++ b/Engine/Source/Runtime/Renderer/Private/PostProcess/PostProcessing.cpp
@@ -742,14 +742,14 @@ static FRenderingCompositeOutputRef AddBloom(FBloomDownSampleArray& BloomDownSam
 	// Extract the Context
 	FPostprocessContext& Context = BloomDownSampleArray.Context;
 
-	const bool bOldMetalNoFFT = (IsMetalPlatform(Context.View.GetShaderPlatform()) && RHIGetShaderLanguageVersion(Context.View.GetShaderPlatform()) < 2);
+	const bool bOldMetalNoFFT = IsMetalPlatform(Context.View.GetShaderPlatform());
 	const bool bUseFFTBloom = (Context.View.FinalPostProcessSettings.BloomMethod == EBloomMethod::BM_FFT
 		&& Context.View.FeatureLevel >= ERHIFeatureLevel::SM5);
 		
 	static bool bWarnAboutOldMetalFFTOnce = false;
 	if (bOldMetalNoFFT && bUseFFTBloom && !bWarnAboutOldMetalFFTOnce)
 	{
-		UE_LOG(LogRenderer, Error, TEXT("Metal v1.2 and above is required to enable FFT Bloom. Set Max. Shader Standard to target to Metal v1.2 in Project Settings > Mac/iOS and recook."));
+		UE_LOG(LogRenderer, Error, TEXT("FFT Bloom is unsupported in Metal."));
 		bWarnAboutOldMetalFFTOnce = true;
 	}
 
@@ -2809,4 +2809,4 @@ void FPostProcessing::ProcessPlanarReflection(FRHICommandListImmediate& RHICmdLi
 bool FPostProcessing::HasAlphaChannelSupport()
 {
 	return CVarAlphaChannel.GetValueOnAnyThread() != 0;
-}
\ No newline at end of file
+}
