Temporarily removing a block dealing with the VR shared depth buffers that was saving us on a theoretically uneeded depth copy. For the moment, the added depth copy ensures that the depth buffer is set as a input param when needed, and that the swapchain behavior that we do under the hood doesn't interfere with the state caching system. We need to better understand what's going wrong here before we can add this block back in.

diff --git a/Engine/Source/Runtime/Renderer/Private/PostProcess/SceneRenderTargets.cpp b/Engine/Source/Runtime/Renderer/Private/PostProcess/SceneRenderTargets.cpp
--- a/Engine/Source/Runtime/Renderer/Private/PostProcess/SceneRenderTargets.cpp
+++ b/Engine/Source/Runtime/Renderer/Private/PostProcess/SceneRenderTargets.cpp
@@ -1853,13 +1853,19 @@ void FSceneRenderTargets::AllocateCommonDepthTargets(FRHICommandList& RHICmdList
 		{
 			const uint32 OldElementSize = SceneDepthZ->ComputeMemorySize();
 			bHMDAllocatedDepthTarget = true;
+		
+			/*
+			@TODO: UE-61597 (4.20) - Right now there appears to be an issue with the shared depth buffer and the shader param caching system.
+			The guess is that since SceneDepthZ's resource is changed under the hood and needs to be reset as a param for certain shaders, but
+			the caching system isn't aware of that. For the time being, we've commented out this block, incurring added expense on the deferred
+			renderer (adding a depth copy that theoretically is unneeded)
 
 			// If SRT and texture are different (MSAA), only modify the resolve render target, to avoid creating a swapchain of MSAA textures
 			if (SceneDepthZ->GetRenderTargetItem().ShaderResourceTexture == SceneDepthZ->GetRenderTargetItem().TargetableTexture)
 			{
 				SceneDepthZ->GetRenderTargetItem().ShaderResourceTexture = SceneDepthZ->GetRenderTargetItem().TargetableTexture = SRTex;
 			}
-			else
+			else*/
 			{
 				SceneDepthZ->GetRenderTargetItem().ShaderResourceTexture = SRTex;
 			}
