Fixed: bRenderInMainPass on SkeletalMeshes doesn't work on Mobile

diff --git a/Engine/Source/Runtime/Renderer/Private/MobileBasePassRendering.cpp b/Engine/Source/Runtime/Renderer/Private/MobileBasePassRendering.cpp
--- a/Engine/Source/Runtime/Renderer/Private/MobileBasePassRendering.cpp
+++ b/Engine/Source/Runtime/Renderer/Private/MobileBasePassRendering.cpp
@@ -730,8 +730,9 @@ void FMobileSceneRenderer::RenderMobileBasePassDynamicData(FRHICommandList& RHIC
 	{
 		const FMeshBatchAndRelevance& MeshBatchAndRelevance = View.DynamicMeshElements[Index];
 
-		if ((BlendMode == BLEND_Opaque && MeshBatchAndRelevance.GetHasOpaqueMaterial()) || 
-			(BlendMode == BLEND_Masked && MeshBatchAndRelevance.GetHasMaskedMaterial()) || bWireFrame)
+		if (MeshBatchAndRelevance.GetRenderInMainPass() &&
+			((BlendMode == BLEND_Opaque && MeshBatchAndRelevance.GetHasOpaqueMaterial()) || 
+			(BlendMode == BLEND_Masked && MeshBatchAndRelevance.GetHasMaskedMaterial()) || bWireFrame))
 		{
 			const FMeshBatch& MeshBatch = *MeshBatchAndRelevance.Mesh;
 			FMobileBasePassOpaqueDrawingPolicyFactory::DrawDynamicMesh(RHICmdList, View, Context, MeshBatch, true, DrawRenderState, MeshBatchAndRelevance.PrimitiveSceneProxy, MeshBatch.BatchHitProxyId);
