Fix CIS.

#rnx

diff --git a/Engine/Source/Runtime/Renderer/Private/RendererScene.cpp b/Engine/Source/Runtime/Renderer/Private/RendererScene.cpp
--- a/Engine/Source/Runtime/Renderer/Private/RendererScene.cpp
+++ b/Engine/Source/Runtime/Renderer/Private/RendererScene.cpp
@@ -3351,18 +3351,6 @@ TStaticMeshDrawList<TBasePassDrawingPolicy<FUniformLightMapPolicy> >& FScene::Ge
 	return BasePassUniformLightMapPolicyDrawList[DrawType];
 }
 
-template<>
-TStaticMeshDrawList<TMobileBasePassDrawingPolicy<FUniformLightMapPolicy>>& FScene::GetMobileBasePassDrawList<FUniformLightMapPolicy>(EBasePassDrawListType DrawType)
-{
-	return MobileBasePassUniformLightMapPolicyDrawList[DrawType];
-}
-
-template<>
-TStaticMeshDrawList<TMobileBasePassDrawingPolicy<FUniformLightMapPolicy>>& FScene::GetMobileBasePassCSMDrawList<FUniformLightMapPolicy>(EBasePassDrawListType DrawType)
-{
-	return MobileBasePassUniformLightMapPolicyDrawListWithCSM[DrawType];
-}
-
 /*-----------------------------------------------------------------------------
 	MotionBlurInfoData
 -----------------------------------------------------------------------------*/
diff --git a/Engine/Source/Runtime/Renderer/Private/ScenePrivate.h b/Engine/Source/Runtime/Renderer/Private/ScenePrivate.h
--- a/Engine/Source/Runtime/Renderer/Private/ScenePrivate.h
+++ b/Engine/Source/Runtime/Renderer/Private/ScenePrivate.h
@@ -2116,13 +2116,6 @@ public:
 	template<typename LightMapPolicyType>
 	TStaticMeshDrawList<TMobileBasePassDrawingPolicy<LightMapPolicyType> >& GetMobileBasePassCSMDrawList(EBasePassDrawListType DrawType);
 
-	/** Prototype explicit specializations to prevent instantiation */
-	template<>
-	TStaticMeshDrawList<TMobileBasePassDrawingPolicy<FUniformLightMapPolicy>>& GetMobileBasePassDrawList<FUniformLightMapPolicy>(EBasePassDrawListType DrawType);
-
-	template<>
-	TStaticMeshDrawList<TMobileBasePassDrawingPolicy<FUniformLightMapPolicy>>& GetMobileBasePassCSMDrawList<FUniformLightMapPolicy>(EBasePassDrawListType DrawType);
-
 #if WITH_EDITOR
 	/** Draw list to use for selected static meshes in the editor only */
 	TStaticMeshDrawList<FEditorSelectionDrawingPolicy> EditorSelectionDrawList;
@@ -2650,5 +2643,17 @@ inline bool ShouldIncludeDomainInMeshPass(EMaterialDomain Domain)
 	return Domain != MD_Volume;
 }
 
+template<>
+inline TStaticMeshDrawList<TMobileBasePassDrawingPolicy<FUniformLightMapPolicy>>& FScene::GetMobileBasePassDrawList<FUniformLightMapPolicy>(EBasePassDrawListType DrawType)
+{
+	return MobileBasePassUniformLightMapPolicyDrawList[DrawType];
+}
+
+template<>
+inline TStaticMeshDrawList<TMobileBasePassDrawingPolicy<FUniformLightMapPolicy>>& FScene::GetMobileBasePassCSMDrawList<FUniformLightMapPolicy>(EBasePassDrawListType DrawType)
+{
+	return MobileBasePassUniformLightMapPolicyDrawListWithCSM[DrawType];
+}
+
 #include "BasePassRendering.inl"
 
