Fix stress test unity compile error.

#rnx

diff --git a/Engine/Source/Runtime/Renderer/Public/MeshMaterialShader.h b/Engine/Source/Runtime/Renderer/Public/MeshMaterialShader.h
--- a/Engine/Source/Runtime/Renderer/Public/MeshMaterialShader.h
+++ b/Engine/Source/Runtime/Renderer/Public/MeshMaterialShader.h
@@ -37,11 +37,11 @@ public:
 		return true;
 	}
 
-	// Clang fails with a linker error when force-inlining this function. Not sure why.
+	// Clang treats FORCEINLINE as adivsory, and will not inline it on debug builds. Since Engine does not depend on the Renderer module, it fails to link against it.
 #if PLATFORM_WINDOWS && defined(__clang__)
 	void ValidateAfterBind();
 #else
-	void ValidateAfterBind()
+	FORCEINLINE void ValidateAfterBind()
 	{
 		checkfSlow(PassUniformBuffer.IsInitialized(), TEXT("FMeshMaterialShader must bind a pass uniform buffer, even if it is just FSceneTexturesUniformParameters: %s"), GetType()->GetName());
 	}
