Prevent division by zero caused by zero-area scissor rect

diff --git a/Engine/Source/Runtime/Engine/Private/LocalLightSceneProxy.cpp b/Engine/Source/Runtime/Engine/Private/LocalLightSceneProxy.cpp
--- a/Engine/Source/Runtime/Engine/Private/LocalLightSceneProxy.cpp
+++ b/Engine/Source/Runtime/Engine/Private/LocalLightSceneProxy.cpp
@@ -60,17 +60,23 @@ bool FLocalLightSceneProxy::GetScissorRect(FIntRect& ScissorRect, const FSceneVi
 	return FMath::ComputeProjectedSphereScissorRect(ScissorRect, GetLightToWorld().GetOrigin(), Radius, View.ViewMatrices.GetViewOrigin(), View.ViewMatrices.GetViewMatrix(), View.ViewMatrices.GetProjectionMatrix()) == 1;
 }
 
-void FLocalLightSceneProxy::SetScissorRect(FRHICommandList& RHICmdList, const FSceneView& View, const FIntRect& ViewRect) const
+bool FLocalLightSceneProxy::SetScissorRect(FRHICommandList& RHICmdList, const FSceneView& View, const FIntRect& ViewRect, FIntRect* OutScissorRect) const
 {
 	FIntRect ScissorRect;
 
 	if (GetScissorRect(ScissorRect, View, ViewRect))
 	{
 		RHICmdList.SetScissorRect(true, ScissorRect.Min.X, ScissorRect.Min.Y, ScissorRect.Max.X, ScissorRect.Max.Y);
+		if (OutScissorRect)
+		{
+			*OutScissorRect = ScissorRect;
+		}
+		return true;
 	}
 	else
 	{
 		RHICmdList.SetScissorRect(false, 0, 0, 0, 0);
+		return false;
 	}
 }
 
diff --git a/Engine/Source/Runtime/Engine/Public/SceneManagement.h b/Engine/Source/Runtime/Engine/Public/SceneManagement.h
--- a/Engine/Source/Runtime/Engine/Public/SceneManagement.h
+++ b/Engine/Source/Runtime/Engine/Public/SceneManagement.h
@@ -1194,8 +1194,11 @@ public:
 		return false;
 	}
 
-	virtual void SetScissorRect(FRHICommandList& RHICmdList, const FSceneView& View, const FIntRect& ViewRect) const
+	// @param OutScissorRect the scissor rect used if one is set
+	// @return whether a scissor rect is set
+	virtual bool SetScissorRect(FRHICommandList& RHICmdList, const FSceneView& View, const FIntRect& ViewRect, FIntRect* OutScissorRect = nullptr) const
 	{
+		return false;
 	}
 
 	virtual bool ShouldCreateRayTracedCascade(ERHIFeatureLevel::Type Type, bool bPrecomputedLightingIsValid, int32 MaxNearCascades) const { return false; }
diff --git a/Engine/Source/Runtime/Renderer/Private/DistanceFieldShadowing.cpp b/Engine/Source/Runtime/Renderer/Private/DistanceFieldShadowing.cpp
--- a/Engine/Source/Runtime/Renderer/Private/DistanceFieldShadowing.cpp
+++ b/Engine/Source/Runtime/Renderer/Private/DistanceFieldShadowing.cpp
@@ -1010,8 +1010,10 @@ void FProjectedShadowInfo::BeginRenderRayTracedDistanceFieldProjection(FRHIComma
 	}
 }
 
-void FProjectedShadowInfo::RenderRayTracedDistanceFieldProjection(FRHICommandListImmediate& RHICmdList, const FViewInfo& View, IPooledRenderTarget* ScreenShadowMaskTexture, bool bProjectingForForwardShading) 
+void FProjectedShadowInfo::RenderRayTracedDistanceFieldProjection(FRHICommandListImmediate& RHICmdList, const FViewInfo& View, const FIntRect& ScissorRect, IPooledRenderTarget* ScreenShadowMaskTexture, bool bProjectingForForwardShading)
 {
+	check(ScissorRect.Area() > 0);
+
 	BeginRenderRayTracedDistanceFieldProjection(RHICmdList, View);
 
 	if (RayTracedShadowsRT)
@@ -1019,13 +1021,6 @@ void FProjectedShadowInfo::RenderRayTracedDistanceFieldProjection(FRHICommandLis
 		QUICK_SCOPE_CYCLE_COUNTER(STAT_RenderRayTracedDistanceFieldShadows);
 		SCOPED_DRAW_EVENT(RHICmdList, RayTracedDistanceFieldShadow);
 
-		FIntRect ScissorRect;
-
-		if (!LightSceneInfo->Proxy->GetScissorRect(ScissorRect, View, View.ViewRect))
-		{
-			ScissorRect = View.ViewRect;
-		}
-
 		if ( IsTransientResourceBufferAliasingEnabled() )
 		{
 			GShadowCulledObjectBuffers.Buffers.DiscardTransientResource();
diff --git a/Engine/Source/Runtime/Renderer/Private/ShadowRendering.cpp b/Engine/Source/Runtime/Renderer/Private/ShadowRendering.cpp
--- a/Engine/Source/Runtime/Renderer/Private/ShadowRendering.cpp
+++ b/Engine/Source/Runtime/Renderer/Private/ShadowRendering.cpp
@@ -1346,13 +1346,20 @@ bool FSceneRenderer::RenderShadowProjections(FRHICommandListImmediate& RHICmdLis
 			RHICmdList.SetViewport(View.ViewRect.Min.X, View.ViewRect.Min.Y, 0.0f, View.ViewRect.Max.X, View.ViewRect.Max.Y, 1.0f);
 
 			// Set the light's scissor rectangle.
-			LightSceneInfo->Proxy->SetScissorRect(RHICmdList, View, View.ViewRect);
+			FIntRect ScissorRect;
+			if (!LightSceneInfo->Proxy->SetScissorRect(RHICmdList, View, View.ViewRect, &ScissorRect))
+			{
+				ScissorRect = View.ViewRect;
+			}
 
-			// Project the shadow depth buffers onto the scene.
-			for (int32 ShadowIndex = 0; ShadowIndex < DistanceFieldShadows.Num(); ShadowIndex++)
+			if (ScissorRect.Area() > 0)
 			{
-				FProjectedShadowInfo* ProjectedShadowInfo = DistanceFieldShadows[ShadowIndex];
-				ProjectedShadowInfo->RenderRayTracedDistanceFieldProjection(RHICmdList, View, ScreenShadowMaskTexture, bProjectingForForwardShading);
+				// Project the shadow depth buffers onto the scene.
+				for (int32 ShadowIndex = 0; ShadowIndex < DistanceFieldShadows.Num(); ShadowIndex++)
+				{
+					FProjectedShadowInfo* ProjectedShadowInfo = DistanceFieldShadows[ShadowIndex];
+					ProjectedShadowInfo->RenderRayTracedDistanceFieldProjection(RHICmdList, View, ScissorRect, ScreenShadowMaskTexture, bProjectingForForwardShading);
+				}
 			}
 
 			// Reset the scissor rectangle.
diff --git a/Engine/Source/Runtime/Renderer/Private/ShadowRendering.h b/Engine/Source/Runtime/Renderer/Private/ShadowRendering.h
--- a/Engine/Source/Runtime/Renderer/Private/ShadowRendering.h
+++ b/Engine/Source/Runtime/Renderer/Private/ShadowRendering.h
@@ -397,7 +397,7 @@ public:
 	void BeginRenderRayTracedDistanceFieldProjection(FRHICommandListImmediate& RHICmdList, const FViewInfo& View);
 
 	/** Renders ray traced distance field shadows. */
-	void RenderRayTracedDistanceFieldProjection(FRHICommandListImmediate& RHICmdList, const class FViewInfo& View, IPooledRenderTarget* ScreenShadowMaskTexture, bool bProjectingForForwardShading);
+	void RenderRayTracedDistanceFieldProjection(FRHICommandListImmediate& RHICmdList, const class FViewInfo& View, const FIntRect& ScissorRect, IPooledRenderTarget* ScreenShadowMaskTexture, bool bProjectingForForwardShading);
 
 	/** Render one pass point light shadow projections. */
 	void RenderOnePassPointLightProjection(FRHICommandListImmediate& RHICmdList, int32 ViewIndex, const FViewInfo& View, bool bProjectingForForwardShading) const;
