Trying to get main compiling again, changing LightColor to >GetEffectiveLightColor, please review this to make sure it is correct.

Marcus.Wassmer

diff --git a/Engine/Source/Runtime/Renderer/Private/PathTracing.cpp b/Engine/Source/Runtime/Renderer/Private/PathTracing.cpp
--- a/Engine/Source/Runtime/Renderer/Private/PathTracing.cpp
+++ b/Engine/Source/Runtime/Renderer/Private/PathTracing.cpp
@@ -212,7 +212,7 @@ public:
 			FSkyLightData SkyLightData;
 			if (SkyLightTextureData)
 			{
-				SkyLightData.Color = FVector(Scene->SkyLight->LightColor);
+				SkyLightData.Color = FVector(Scene->SkyLight->GetEffectiveLightColor());
 				SkyLightData.Texture = Scene->SkyLight->ProcessedTexture->TextureRHI;
 				SkyLightData.TextureSampler = Scene->SkyLight->ProcessedTexture->SamplerStateRHI;
 
@@ -560,7 +560,7 @@ void FDeferredShadingSceneRenderer::RenderPathTracing(FRHICommandListImmediate&
 	uint32 SkyLightColumnCount = 0;
 	if (Scene->SkyLight)
 	{
-		SkyLightColor = Scene->SkyLight->LightColor;
+		SkyLightColor = Scene->SkyLight->GetEffectiveLightColor();
 		SkyLightTexture = Scene->SkyLight->ProcessedTexture;
 		SkyLightColumnCount = SkyLightTexture->GetSizeX();
 		SkyLightRowCount = SkyLightTexture->GetSizeY();
diff --git a/Engine/Source/Runtime/Renderer/Private/RayTracing/RayTracingGlobalIllumination.cpp b/Engine/Source/Runtime/Renderer/Private/RayTracing/RayTracingGlobalIllumination.cpp
--- a/Engine/Source/Runtime/Renderer/Private/RayTracing/RayTracingGlobalIllumination.cpp
+++ b/Engine/Source/Runtime/Renderer/Private/RayTracing/RayTracingGlobalIllumination.cpp
@@ -143,7 +143,7 @@ void SetupSkyLightParameters(
 
 	if (Scene.SkyLight)
 	{
-		SkyLight->Color = FVector(Scene.SkyLight->LightColor);
+		SkyLight->Color = FVector(Scene.SkyLight->GetEffectiveLightColor());
 		SkyLight->Texture = Scene.SkyLight->ProcessedTexture->TextureRHI;
 		SkyLight->TextureSampler = Scene.SkyLight->ProcessedTexture->SamplerStateRHI;
 		SkyLight->MipDimensions = Scene.SkyLight->SkyLightMipDimensions;
diff --git a/Engine/Source/Runtime/Renderer/Private/RayTracing/RaytracingSkylight.cpp b/Engine/Source/Runtime/Renderer/Private/RayTracing/RaytracingSkylight.cpp
--- a/Engine/Source/Runtime/Renderer/Private/RayTracing/RaytracingSkylight.cpp
+++ b/Engine/Source/Runtime/Renderer/Private/RayTracing/RaytracingSkylight.cpp
@@ -759,7 +759,7 @@ void FDeferredShadingSceneRenderer::RenderRayTracingSkyLight(
 	FSkyLightData SkyLightData;
 	SkyLightData.SamplesPerPixel = GRayTracingSkyLightSamplesPerPixel;
 	SkyLightData.SamplingStopLevel = GRayTracingSkyLightSamplingStopLevel;
-	SkyLightData.Color = FVector(Scene->SkyLight->LightColor);
+	SkyLightData.Color = FVector(Scene->SkyLight->GetEffectiveLightColor());
 	SkyLightData.Texture = Scene->SkyLight->ProcessedTexture->TextureRHI;
 	SkyLightData.TextureSampler = Scene->SkyLight->ProcessedTexture->SamplerStateRHI;
 	SkyLightData.MipDimensions = Scene->SkyLight->SkyLightMipDimensions;
