Currently the wait for the Parallel task was done during the destruction of the FViewInfo, which mean that if the task were using some stuff from the View, they could have been destroyed, resulting in using gargage. We now wait in the FViewInfo for the async task before doing any destruction work.

diff --git a/Engine/Source/Runtime/Renderer/Private/MeshDrawCommands.cpp b/Engine/Source/Runtime/Renderer/Private/MeshDrawCommands.cpp
--- a/Engine/Source/Runtime/Renderer/Private/MeshDrawCommands.cpp
+++ b/Engine/Source/Runtime/Renderer/Private/MeshDrawCommands.cpp
@@ -1012,7 +1012,7 @@ void FParallelMeshDrawCommandPass::WaitForMeshPassSetupTask() const
 	}
 }
 
-void FParallelMeshDrawCommandPass::Empty()
+void FParallelMeshDrawCommandPass::WaitForTasksAndEmpty()
 {
 	// Need to wait in case if someone dispatched sort and draw merge task, but didn't draw it.
 	WaitForMeshPassSetupTask();
@@ -1052,7 +1052,7 @@ void FParallelMeshDrawCommandPass::Empty()
 
 FParallelMeshDrawCommandPass::~FParallelMeshDrawCommandPass()
 {
-	Empty();
+	check(TaskEventRef == nullptr);
 }
 
 class FDrawVisibleMeshCommandsAnyThreadTask : public FRenderTask
diff --git a/Engine/Source/Runtime/Renderer/Private/MeshDrawCommands.h b/Engine/Source/Runtime/Renderer/Private/MeshDrawCommands.h
--- a/Engine/Source/Runtime/Renderer/Private/MeshDrawCommands.h
+++ b/Engine/Source/Runtime/Renderer/Private/MeshDrawCommands.h
@@ -155,7 +155,7 @@ public:
 	 */
 	void DispatchDraw(FParallelCommandListSet* ParallelCommandListSet, FRHICommandList& RHICmdList) const;
 
-	void Empty();
+	void WaitForTasksAndEmpty();
 	void SetDumpInstancingStats(const FString& InPassName);
 	bool HasAnyDraw() const { return MaxNumDraws > 0; }
 
diff --git a/Engine/Source/Runtime/Renderer/Private/SceneRendering.cpp b/Engine/Source/Runtime/Renderer/Private/SceneRendering.cpp
--- a/Engine/Source/Runtime/Renderer/Private/SceneRendering.cpp
+++ b/Engine/Source/Runtime/Renderer/Private/SceneRendering.cpp
@@ -902,6 +902,11 @@ FViewInfo::~FViewInfo()
 		CustomVisibilityQuery->Release();
 	}
 
+	for (int32 MeshDrawIndex = 0; MeshDrawIndex < EMeshPass::Num; MeshDrawIndex++)
+	{
+		ParallelMeshDrawCommandPasses[MeshDrawIndex].WaitForTasksAndEmpty();
+	}
+
 	//this uses memstack allocation for strongrefs, so we need to manually empty to get the destructor called to not leak the uniformbuffers stored here.
 	TranslucentSelfShadowUniformBufferMap.Empty();
 }
@@ -1566,7 +1571,7 @@ void FViewInfo::DestroyAllSnapshots()
 
 		for (int32 Index = 0; Index < Snapshot->ParallelMeshDrawCommandPasses.Num(); ++Index)
 		{
-			Snapshot->ParallelMeshDrawCommandPasses[Index].Empty();
+			Snapshot->ParallelMeshDrawCommandPasses[Index].WaitForTasksAndEmpty();
 		}
 
 		FreeViewInfoSnapshots.Add(Snapshot);
@@ -2320,6 +2325,7 @@ FSceneRenderer::~FSceneRenderer()
 {
 	// To prevent keeping persistent references to single frame buffers, clear any such reference at this point.
 	ClearPrimitiveSingleFrameIndirectLightingCacheBuffers();
+	
 
 	if(Scene)
 	{
@@ -2333,6 +2339,7 @@ FSceneRenderer::~FSceneRenderer()
 				{
 					// FProjectedShadowInfo's in MemStackProjectedShadows were allocated on the rendering thread mem stack, 
 					// Their memory will be freed when the stack is freed with no destructor call, so invoke the destructor explicitly
+					VisibleLightInfo.MemStackProjectedShadows[ShadowIndex]->WaitForMeshDrawTasksAndEmpty();
 					VisibleLightInfo.MemStackProjectedShadows[ShadowIndex]->~FProjectedShadowInfo();
 				}
 			}
diff --git a/Engine/Source/Runtime/Renderer/Private/ShadowRendering.h b/Engine/Source/Runtime/Renderer/Private/ShadowRendering.h
--- a/Engine/Source/Runtime/Renderer/Private/ShadowRendering.h
+++ b/Engine/Source/Runtime/Renderer/Private/ShadowRendering.h
@@ -500,6 +500,8 @@ public:
 		return FShadowDepthType(bDirectionalLight, bOnePassPointLightShadow, bReflectiveShadowmap);
 	}
 
+	void WaitForMeshDrawTasksAndEmpty() { ShadowDepthPass.WaitForTasksAndEmpty(); }
+
 private:
 	// 0 if Setup...() wasn't called yet
 	const FLightSceneInfo* LightSceneInfo;
diff --git a/Engine/Source/Runtime/Renderer/Private/ShadowSetup.cpp b/Engine/Source/Runtime/Renderer/Private/ShadowSetup.cpp
--- a/Engine/Source/Runtime/Renderer/Private/ShadowSetup.cpp
+++ b/Engine/Source/Runtime/Renderer/Private/ShadowSetup.cpp
@@ -1530,7 +1530,7 @@ void FProjectedShadowInfo::ClearTransientArrays()
 	DynamicSubjectTranslucentMeshElements.Empty();
 
 	ShadowDepthPassVisibleCommands.Empty();
-	ShadowDepthPass.Empty();
+	ShadowDepthPass.WaitForTasksAndEmpty();
 
 	SubjectMeshCommandBuildRequests.Empty();
 
