Fixed editor startup crash with ray tracing enabled

diff --git a/Engine/Source/Runtime/Renderer/Private/MeshPassProcessor.cpp b/Engine/Source/Runtime/Renderer/Private/MeshPassProcessor.cpp
--- a/Engine/Source/Runtime/Renderer/Private/MeshPassProcessor.cpp
+++ b/Engine/Source/Runtime/Renderer/Private/MeshPassProcessor.cpp
@@ -323,16 +323,19 @@ FGraphicsMinimalPipelineStateId FGraphicsMinimalPipelineStateId::GetPersistentId
 
 void FGraphicsMinimalPipelineStateId::RemovePersistentId(FGraphicsMinimalPipelineStateId Id)
 {
-	check(Id.bValid && !Id.bOneFrameId);
+	check(!Id.bOneFrameId);
 
-	const FSetElementId SetElementId = FSetElementId::FromInteger(Id.SetElementIndex);
-	FRefCountedGraphicsMinimalPipelineStateInitializer& RefCountedStateInitializer = PersistentIdTable[SetElementId];
-
-	check(RefCountedStateInitializer.RefNum > 0);
-	--RefCountedStateInitializer.RefNum;
-	if (RefCountedStateInitializer.RefNum <= 0)
+	if (Id.bValid)
 	{
-		PersistentIdTable.Remove(SetElementId);
+		const FSetElementId SetElementId = FSetElementId::FromInteger(Id.SetElementIndex);
+		FRefCountedGraphicsMinimalPipelineStateInitializer& RefCountedStateInitializer = PersistentIdTable[SetElementId];
+
+		check(RefCountedStateInitializer.RefNum > 0);
+		--RefCountedStateInitializer.RefNum;
+		if (RefCountedStateInitializer.RefNum <= 0)
+		{
+			PersistentIdTable.Remove(SetElementId);
+		}
 	}
 }
 
