REL - Fix Decal with Emissive

diff --git a/Engine/Source/Runtime/Renderer/Private/CompositionLighting/PostProcessMeshDecals.cpp b/Engine/Source/Runtime/Renderer/Private/CompositionLighting/PostProcessMeshDecals.cpp
--- a/Engine/Source/Runtime/Renderer/Private/CompositionLighting/PostProcessMeshDecals.cpp
+++ b/Engine/Source/Runtime/Renderer/Private/CompositionLighting/PostProcessMeshDecals.cpp
@@ -253,6 +253,11 @@ void FMeshDecalMeshProcessor::AddMeshBatch(const FMeshBatch& RESTRICT MeshBatch,
 					bShouldRender = bShouldRender && RenderTargetMode != FDecalRenderingCommon::RTM_GBufferNormal;
 				}
 
+				if (PassDecalStage == DRS_Emissive)
+				{
+					bShouldRender = bShouldRender && Material->HasEmissiveColorConnected();
+				}
+
 				if (bShouldRender)
 				{
 					const bool bHasNormal = Material->HasNormalConnected();
