Wait for all the dispatched shadow mesh draw tasks, before we destroy FSceneRenderer. This fixes assert inside FParallelMeshDrawCommandPass destructor, indicating that pass was destroyed too late and it's a potential race condition.
#rnx

diff --git a/Engine/Source/Runtime/Renderer/Private/MeshDrawCommands.cpp b/Engine/Source/Runtime/Renderer/Private/MeshDrawCommands.cpp
--- a/Engine/Source/Runtime/Renderer/Private/MeshDrawCommands.cpp
+++ b/Engine/Source/Runtime/Renderer/Private/MeshDrawCommands.cpp
@@ -1043,6 +1043,7 @@ void FParallelMeshDrawCommandPass::WaitForTasksAndEmpty()
 	TaskContext.DynamicMeshElements = nullptr;
 	TaskContext.DynamicMeshElementsPassRelevance = nullptr;
 	TaskContext.MeshDrawCommands.Empty();
+	TaskContext.MeshDrawCommandStorage.MeshDrawCommands.Empty();
 	TaskContext.MobileBasePassCSMMeshDrawCommands.Empty();
 	TaskContext.DynamicMeshCommandBuildRequests.Empty();
 	TaskContext.TempVisibleMeshDrawCommands.Empty();
diff --git a/Engine/Source/Runtime/Renderer/Private/MeshDrawCommands.h b/Engine/Source/Runtime/Renderer/Private/MeshDrawCommands.h
--- a/Engine/Source/Runtime/Renderer/Private/MeshDrawCommands.h
+++ b/Engine/Source/Runtime/Renderer/Private/MeshDrawCommands.h
@@ -123,7 +123,8 @@ class FParallelMeshDrawCommandPass
 {
 public:
 	FParallelMeshDrawCommandPass()
-		: bPrimitiveIdBufferDataOwnedByRHIThread(false)
+		: PrimitiveIdVertexBufferRHI(nullptr)
+		, bPrimitiveIdBufferDataOwnedByRHIThread(false)
 		, MaxNumDraws(0)
 	{
 	}
diff --git a/Engine/Source/Runtime/Renderer/Private/SceneRendering.cpp b/Engine/Source/Runtime/Renderer/Private/SceneRendering.cpp
--- a/Engine/Source/Runtime/Renderer/Private/SceneRendering.cpp
+++ b/Engine/Source/Runtime/Renderer/Private/SceneRendering.cpp
@@ -2326,7 +2326,6 @@ FSceneRenderer::~FSceneRenderer()
 	// To prevent keeping persistent references to single frame buffers, clear any such reference at this point.
 	ClearPrimitiveSingleFrameIndirectLightingCacheBuffers();
 	
-
 	if(Scene)
 	{
 		// Destruct the projected shadow infos.
@@ -2339,7 +2338,6 @@ FSceneRenderer::~FSceneRenderer()
 				{
 					// FProjectedShadowInfo's in MemStackProjectedShadows were allocated on the rendering thread mem stack, 
 					// Their memory will be freed when the stack is freed with no destructor call, so invoke the destructor explicitly
-					VisibleLightInfo.MemStackProjectedShadows[ShadowIndex]->WaitForMeshDrawTasksAndEmpty();
 					VisibleLightInfo.MemStackProjectedShadows[ShadowIndex]->~FProjectedShadowInfo();
 				}
 			}
@@ -2900,11 +2898,10 @@ void FSceneRenderer::WaitForTasksClearSnapshotsAndDeleteSceneRenderer(FRHIComman
 		RHICmdList.ImmediateFlush(EImmediateFlushType::WaitForOutstandingTasksOnly);
 	}
 
-	// Destroy cached preshadow transient arrays (allocated with SceneRenderingAllocator).
-	TArray<TRefCountPtr<FProjectedShadowInfo>>& CachedPreshadows = SceneRenderer->Scene->CachedPreshadows;
-	for (int32 CachedShadowIndex = 0; CachedShadowIndex < CachedPreshadows.Num(); ++CachedShadowIndex)
+	// Wait for all dispatched shadow mesh draw tasks.
+	for (int32 PassIndex = 0; PassIndex < SceneRenderer->DispatchedShadowDepthPasses.Num(); ++PassIndex)
 	{
-		CachedPreshadows[CachedShadowIndex]->ClearTransientArrays();
+		SceneRenderer->DispatchedShadowDepthPasses[PassIndex]->WaitForTasksAndEmpty();
 	}
 
 	FViewInfo::DestroyAllSnapshots(); // this destroys viewinfo snapshots
diff --git a/Engine/Source/Runtime/Renderer/Private/SceneRendering.h b/Engine/Source/Runtime/Renderer/Private/SceneRendering.h
--- a/Engine/Source/Runtime/Renderer/Private/SceneRendering.h
+++ b/Engine/Source/Runtime/Renderer/Private/SceneRendering.h
@@ -1286,6 +1286,9 @@ public:
 	/** Information about the visible lights. */
 	TArray<FVisibleLightInfo,SceneRenderingAllocator> VisibleLightInfos;
 
+	/** Array of dispatched parallel shadow depth passes. */
+	TArray<FParallelMeshDrawCommandPass*, SceneRenderingAllocator> DispatchedShadowDepthPasses;
+
 	FSortedShadowMaps SortedShadowsForShadowDepthPass;
 
 	/** If a freeze request has been made */
diff --git a/Engine/Source/Runtime/Renderer/Private/ShadowRendering.h b/Engine/Source/Runtime/Renderer/Private/ShadowRendering.h
--- a/Engine/Source/Runtime/Renderer/Private/ShadowRendering.h
+++ b/Engine/Source/Runtime/Renderer/Private/ShadowRendering.h
@@ -500,8 +500,6 @@ public:
 		return FShadowDepthType(bDirectionalLight, bOnePassPointLightShadow, bReflectiveShadowmap);
 	}
 
-	void WaitForMeshDrawTasksAndEmpty() { ShadowDepthPass.WaitForTasksAndEmpty(); }
-
 private:
 	// 0 if Setup...() wasn't called yet
 	const FLightSceneInfo* LightSceneInfo;
diff --git a/Engine/Source/Runtime/Renderer/Private/ShadowSetup.cpp b/Engine/Source/Runtime/Renderer/Private/ShadowSetup.cpp
--- a/Engine/Source/Runtime/Renderer/Private/ShadowSetup.cpp
+++ b/Engine/Source/Runtime/Renderer/Private/ShadowSetup.cpp
@@ -1191,6 +1191,8 @@ void FProjectedShadowInfo::SetupMeshDrawCommandsForShadowDepth(FSceneRenderer& R
 		SubjectMeshCommandBuildRequests,
 		NumSubjectMeshCommandBuildRequestElements * InstanceFactor,
 		ShadowDepthPassVisibleCommands);
+
+	Renderer.DispatchedShadowDepthPasses.Add(&ShadowDepthPass);
 }
 
 void FProjectedShadowInfo::SetupMeshDrawCommandsForProjectionStenciling(FSceneRenderer& Renderer)
