Fixed memory leak in shadow mesh draw commands. DynamicMeshDrawCommandStorage was never cleared for cached preshadows which are reused between frames.
#rnx

diff --git a/Engine/Source/Runtime/Renderer/Private/ShadowSetup.cpp b/Engine/Source/Runtime/Renderer/Private/ShadowSetup.cpp
--- a/Engine/Source/Runtime/Renderer/Private/ShadowSetup.cpp
+++ b/Engine/Source/Runtime/Renderer/Private/ShadowSetup.cpp
@@ -1537,6 +1537,8 @@ void FProjectedShadowInfo::ClearTransientArrays()
 	SubjectMeshCommandBuildRequests.Empty();
 
 	ProjectionStencilingPasses.Reset();
+
+	DynamicMeshDrawCommandStorage.MeshDrawCommands.Empty();
 }
 
 /** Returns a cached preshadow matching the input criteria if one exists. */
