Support ISR in distortion rendering

diff --git a/Engine/Shaders/Private/DistortAccumulatePS.usf b/Engine/Shaders/Private/DistortAccumulatePS.usf
--- a/Engine/Shaders/Private/DistortAccumulatePS.usf
+++ b/Engine/Shaders/Private/DistortAccumulatePS.usf
@@ -33,7 +33,11 @@ void Main(
 	out float4 OutColor		: SV_Target0
 	)
 {
+#if INSTANCED_STEREO
+	ResolvedView = ResolveView(Interpolants.EyeIndex);
+#else
 	ResolvedView = ResolveView();
+#endif
 
 	// material parameter inputs
 	FMaterialPixelParameters MaterialParameters = GetMaterialPixelParameters(Interpolants, SvPosition);
@@ -81,9 +85,16 @@ void Main(
 	float2 NDC = (MaterialParameters.ScreenPosition.xy / MaterialParameters.ScreenPosition.w);
 	float2 ScreenUV = NDC * ResolvedView.ScreenPositionScaleBias.xy + ResolvedView.ScreenPositionScaleBias.wz;
 
-	//Fix for Fov and aspect.
+	// When ISR is enabled we store two FOVs in the DistortionParams and compute the aspect ratio here
+#if INSTANCED_STEREO
+	float InvTanHalfFov = DistortionParams[Interpolants.EyeIndex];
+    float Ratio = DistortionParams.z / DistortionParams.w;
+#else
 	float InvTanHalfFov = DistortionParams.x;
     float Ratio = DistortionParams.y;
+#endif
+
+	// Fix for Fov and aspect.
  	float2 FovFix = float2(InvTanHalfFov,Ratio*InvTanHalfFov);
 	BufferUVDistortion *= DistortionParams.zw * float2(OffsetFudgeFactor,-OffsetFudgeFactor) * FovFix;
 
diff --git a/Engine/Shaders/Private/DistortAccumulateVS.usf b/Engine/Shaders/Private/DistortAccumulateVS.usf
--- a/Engine/Shaders/Private/DistortAccumulateVS.usf
+++ b/Engine/Shaders/Private/DistortAccumulateVS.usf
@@ -61,7 +61,7 @@ struct FDistortAccumulateVSToPS
 		O.FactoryInterpolants = VertexFactoryAssignInterpolants(Interpolants.FactoryInterpolants);
 		
 		// Finally, transform position to clip-space
-		O.Position = mul(WorldPosition, View.TranslatedWorldToClip);
+		O.Position = mul(WorldPosition, ResolvedView.TranslatedWorldToClip);
 			
 		O.PixelPosition = WorldPosition;
 		
@@ -80,10 +80,29 @@ void Main(
 	out VS_OUTPUT_TYPE Output
 #if USE_GLOBAL_CLIP_PLANE && !USING_TESSELLATION
 	, out float OutGlobalClipPlaneDistance : SV_ClipDistance
+#endif
+#if INSTANCED_STEREO
+	, uint InstanceId : SV_InstanceID
+	#if !MULTI_VIEW
+		, out float OutClipDistance : SV_ClipDistance1
+	#else
+		, out uint ViewportIndex : SV_ViewPortArrayIndex
+	#endif
 #endif
 	)
 {
+#if INSTANCED_STEREO
+	const uint EyeIndex = GetEyeIndex(InstanceId);
+	ResolvedView = ResolveView(EyeIndex);
+	#if !MULTI_VIEW
+		OutClipDistance = 0.0;
+	#else
+		ViewportIndex = EyeIndex;
+	#endif
+#else
+	uint EyeIndex = 0;
 	ResolvedView = ResolveView();
+#endif
 
 	FVertexFactoryIntermediates VFIntermediates = GetVertexFactoryIntermediates(Input);
 	float4	WorldPosition = VertexFactoryGetWorldPosition(Input, VFIntermediates);
@@ -104,9 +123,24 @@ void Main(
 	
 #else // !USING_TESSELLATION
 
-	Output.Position = mul(WorldPosition, View.TranslatedWorldToClip);
+	Output.Position = mul(WorldPosition, ResolvedView.TranslatedWorldToClip);
 	Output.FactoryInterpolants = VertexFactoryGetInterpolantsVSToPS(Input, VFIntermediates, VertexParameters);
 
+	#if INSTANCED_STEREO && !MULTI_VIEW
+	BRANCH
+	if (IsInstancedStereo())
+	{
+		// Clip at the center of the screen
+		OutClipDistance = dot(Output.Position, EyeClipEdge[EyeIndex]);
+
+		// Scale to the width of a single eye viewport
+		Output.Position.x *= 0.5 * ResolvedView.HMDEyePaddingOffset;
+
+		// Shift to the eye viewport
+		Output.Position.x += (EyeOffsetScale[EyeIndex] * Output.Position.w) * (1.0f - 0.5 * ResolvedView.HMDEyePaddingOffset);
+	}
+	#endif
+
 #if USE_GLOBAL_CLIP_PLANE
 	OutGlobalClipPlaneDistance = dot(ResolvedView.GlobalClippingPlane, float4(WorldPosition.xyz - ResolvedView.PreViewTranslation.xyz, 1));
 #endif
@@ -115,6 +149,14 @@ void Main(
 
 #endif // USING_TESSELLATION
 
+#if INSTANCED_STEREO
+	#if USING_TESSELLATION	
+		Output.FactoryInterpolants.InterpolantsVSToPS.EyeIndex = EyeIndex;
+	#else
+		Output.FactoryInterpolants.EyeIndex = EyeIndex;
+	#endif
+#endif
+
 	OutputVertexID( Output );
 }
 
diff --git a/Engine/Source/Runtime/Renderer/Private/DistortionRendering.cpp b/Engine/Source/Runtime/Renderer/Private/DistortionRendering.cpp
--- a/Engine/Source/Runtime/Renderer/Private/DistortionRendering.cpp
+++ b/Engine/Source/Runtime/Renderer/Private/DistortionRendering.cpp
@@ -51,6 +51,16 @@ void SetupDistortionPassUniformBuffer(FRHICommandListImmediate& RHICmdList, cons
 	DistortionPassParameters.DistortionParams.Y = Ratio;
 	DistortionPassParameters.DistortionParams.Z = (float)View.UnscaledViewRect.Width();
 	DistortionPassParameters.DistortionParams.W = (float)View.UnscaledViewRect.Height();
+
+	// When ISR is enabled we store two FOVs in the distortion parameters and compute the aspect ratio in the shader instead.
+	if ((View.IsInstancedStereoPass() || View.bIsMobileMultiViewEnabled) && View.Family->Views.Num() > 0)
+	{
+		// When drawing the left eye in a stereo scene, copy the right eye view values into the instanced view uniform buffer.
+		const EStereoscopicPass StereoPassIndex = (View.StereoPass != eSSP_FULL) ? eSSP_RIGHT_EYE : eSSP_FULL;
+
+		const FViewInfo& InstancedView = static_cast<const FViewInfo&>(View.Family->GetStereoEyeView(StereoPassIndex));
+		DistortionPassParameters.DistortionParams.Y = InstancedView.ViewMatrices.GetProjectionMatrix().M[0][0];
+	}
 }
 
 void SetupMobileDistortionPassUniformBuffer(FRHICommandListImmediate& RHICmdList, const FViewInfo& View, FMobileDistortionPassUniformParameters& DistortionPassParameters)
@@ -63,6 +73,16 @@ void SetupMobileDistortionPassUniformBuffer(FRHICommandListImmediate& RHICmdList
 	DistortionPassParameters.DistortionParams.Y = Ratio;
 	DistortionPassParameters.DistortionParams.Z = (float)View.UnscaledViewRect.Width();
 	DistortionPassParameters.DistortionParams.W = (float)View.UnscaledViewRect.Height();
+
+	// When ISR is enabled we store two FOVs in the distortion parameters and compute the aspect ratio in the shader instead.
+	if ((View.IsInstancedStereoPass() || View.bIsMobileMultiViewEnabled) && View.Family->Views.Num() > 0)
+	{
+		// When drawing the left eye in a stereo scene, copy the right eye view values into the instanced view uniform buffer.
+		const EStereoscopicPass StereoPassIndex = (View.StereoPass != eSSP_FULL) ? eSSP_RIGHT_EYE : eSSP_FULL;
+
+		const FViewInfo& InstancedView = static_cast<const FViewInfo&>(View.Family->GetStereoEyeView(StereoPassIndex));
+		DistortionPassParameters.DistortionParams.Y = InstancedView.ViewMatrices.GetProjectionMatrix().M[0][0];
+	}
 }
 
 /**
@@ -494,6 +514,33 @@ bool SubmitDistortionMeshDrawCommands(FRHICommandListImmediate& RHICmdList, cons
 	return bDirty;
 }
 
+static void SetupDistortionPassView(FRHICommandList& RHICmdList, const FViewInfo& View, const FSceneRenderer* SceneRenderer)
+{
+	if (!View.IsInstancedStereoPass())
+	{
+		RHICmdList.SetViewport(View.ViewRect.Min.X, View.ViewRect.Min.Y, 0.0f, View.ViewRect.Max.X, View.ViewRect.Max.Y, 1.0f);
+	}
+	else
+	{
+		if (View.bIsMultiViewEnabled)
+		{
+			const uint32 LeftMinX = SceneRenderer->Views[0].ViewRect.Min.X;
+			const uint32 LeftMaxX = SceneRenderer->Views[0].ViewRect.Max.X;
+			const uint32 RightMinX = SceneRenderer->Views[1].ViewRect.Min.X;
+			const uint32 RightMaxX = SceneRenderer->Views[1].ViewRect.Max.X;
+
+			const uint32 LeftMaxY = SceneRenderer->Views[0].ViewRect.Max.Y;
+			const uint32 RightMaxY = SceneRenderer->Views[1].ViewRect.Max.Y;
+
+			RHICmdList.SetStereoViewport(LeftMinX, RightMinX, 0, 0, 0.0f, LeftMaxX, RightMaxX, LeftMaxY, RightMaxY, 1.0f);
+		}
+		else
+		{
+			RHICmdList.SetViewport(0, 0, 0, SceneRenderer->InstancedStereoWidth, View.ViewRect.Max.Y, 1);
+		}
+	}
+}
+
 /**
  * Renders the scene's distortion
  */
@@ -508,7 +555,7 @@ void FSceneRenderer::RenderDistortion(FRHICommandListImmediate& RHICmdList)
 	for(int32 ViewIndex = 0;ViewIndex < Views.Num();ViewIndex++)
 	{
 		const FViewInfo& View = Views[ViewIndex];
-		if (View.bHasDistortionPrimitives)
+		if (View.bHasDistortionPrimitives && View.ShouldRenderView())
 		{
 			bRender=true;
 			break;
@@ -560,8 +607,13 @@ void FSceneRenderer::RenderDistortion(FRHICommandListImmediate& RHICmdList)
 					SCOPED_CONDITIONAL_DRAW_EVENTF(RHICmdList, EventView, Views.Num() > 1, TEXT("View%d"), ViewIndex);
 
 					FViewInfo& View = Views[ViewIndex];
+					if (!View.ShouldRenderView())
+					{
+						continue;
+					}
+
 					// viewport to match view size
-					RHICmdList.SetViewport(View.ViewRect.Min.X, View.ViewRect.Min.Y, 0.0f, View.ViewRect.Max.X, View.ViewRect.Max.Y, 1.0f);
+					SetupDistortionPassView(RHICmdList, View, this);
 
 					Scene->UniformBuffers.UpdateViewUniformBuffer(View);
 
@@ -570,6 +622,7 @@ void FSceneRenderer::RenderDistortion(FRHICommandListImmediate& RHICmdList)
 					Scene->UniformBuffers.DistortionPassUniformBuffer.UpdateUniformBufferImmediate(DistortionPassParameters);
 
 					FMeshPassProcessorRenderState DrawRenderState(View, Scene->UniformBuffers.DistortionPassUniformBuffer);
+					DrawRenderState.SetInstancedViewUniformBuffer(Scene->UniformBuffers.InstancedViewUniformBuffer);
 
 					// test against depth and write stencil mask
 					DrawRenderState.SetDepthStencilState(TStaticDepthStencilState<
