[Switch] Add Switch to list of platforms PREVENT_RENDERTARGET_SIZE_THRASHING. Prevents rendering slowdown and log spam.

diff --git a/Engine/Source/Runtime/Renderer/Private/PostProcess/SceneRenderTargets.h b/Engine/Source/Runtime/Renderer/Private/PostProcess/SceneRenderTargets.h
--- a/Engine/Source/Runtime/Renderer/Private/PostProcess/SceneRenderTargets.h
+++ b/Engine/Source/Runtime/Renderer/Private/PostProcess/SceneRenderTargets.h
@@ -116,7 +116,7 @@ static const int32 NumTranslucencyShadowSurfaces = 2;
 
 #define STENCIL_LIGHTING_CHANNELS_MASK(Value) uint8((Value & 0x7) << STENCIL_LIGHTING_CHANNELS_BIT_ID)
 
-#define PREVENT_RENDERTARGET_SIZE_THRASHING (PLATFORM_DESKTOP || PLATFORM_XBOXONE || PLATFORM_PS4 || PLATFORM_ANDROID || PLATFORM_IOS)
+#define PREVENT_RENDERTARGET_SIZE_THRASHING (PLATFORM_DESKTOP || PLATFORM_XBOXONE || PLATFORM_PS4 || PLATFORM_ANDROID || PLATFORM_IOS || PLATFORM_SWITCH)
 
 enum class ESceneColorFormatType
 {
