Properly disabled gpu particle collision simulation for mobile renderer

diff --git a/Engine/Source/Runtime/Engine/Private/Particles/FXSystemPrivate.h b/Engine/Source/Runtime/Engine/Private/Particles/FXSystemPrivate.h
--- a/Engine/Source/Runtime/Engine/Private/Particles/FXSystemPrivate.h
+++ b/Engine/Source/Runtime/Engine/Private/Particles/FXSystemPrivate.h
@@ -64,7 +64,7 @@ inline bool IsParticleCollisionModeSupported(EShaderPlatform InPlatform, EPartic
 	case PCM_DepthBuffer:
 		// we only need to check for simple forward if we're NOT currently attempting to cache the shader
 		// since SF is a runtime change, we need to compile the shader regardless, because we could be switching to deferred at any time
-		return (IsFeatureLevelSupported(InPlatform, ERHIFeatureLevel::SM4) || (IsFeatureLevelSupported(InPlatform, ERHIFeatureLevel::ES3_1) && IsVulkanPlatform(InPlatform)))
+		return (IsFeatureLevelSupported(InPlatform, ERHIFeatureLevel::SM4))
 			&& (bForCaching || !IsSimpleForwardShadingEnabled(InPlatform));
 	case PCM_DistanceField:
 		return IsFeatureLevelSupported(InPlatform, ERHIFeatureLevel::SM5);
diff --git a/Engine/Source/Runtime/Renderer/Private/MobileShadingRenderer.cpp b/Engine/Source/Runtime/Renderer/Private/MobileShadingRenderer.cpp
--- a/Engine/Source/Runtime/Renderer/Private/MobileShadingRenderer.cpp
+++ b/Engine/Source/Runtime/Renderer/Private/MobileShadingRenderer.cpp
@@ -69,13 +69,6 @@ static TAutoConsoleVariable<int32> CVarMobileMoveSubmissionHintAfterTranslucency
 	TEXT("1: Submission hint occurs after translucency. (Default)"),
 	ECVF_Scalability | ECVF_RenderThreadSafe);
 
-static TAutoConsoleVariable<int32> CVarMobileDisableGPUParticleCollision(
-	TEXT("r.Mobile.DisableGPUParticleCollision"),
-	1,
-	TEXT("0: Allow GPU particle collision simulation if supported (Default).\n")
-	TEXT("1: Disable GPU particle collision simulation"),
-	ECVF_RenderThreadSafe);
-
 DECLARE_CYCLE_STAT(TEXT("SceneStart"), STAT_CLMM_SceneStart, STATGROUP_CommandListMarkers);
 DECLARE_CYCLE_STAT(TEXT("SceneEnd"), STAT_CLMM_SceneEnd, STATGROUP_CommandListMarkers);
 DECLARE_CYCLE_STAT(TEXT("InitViews"), STAT_CLMM_InitViews, STATGROUP_CommandListMarkers);
@@ -457,18 +450,7 @@ void FMobileSceneRenderer::Render(FRHICommandListImmediate& RHICmdList)
 	{
 		RHICmdList.EndRenderPass();
 	}
-
-	// Notify the FX system that opaque primitives have been rendered.
-	if (Scene->FXSystem && IsGPUParticleCollisionEnabled(Views[0]))
-	{
-		FMobileSceneTextureUniformParameters MobileSceneTextureParameters;
-		SetupMobileSceneTextureUniformParameters(SceneContext, FeatureLevel, true, MobileSceneTextureParameters);
-		TUniformBufferRef<FMobileSceneTextureUniformParameters> MobileSceneTextureUniformBuffer = TUniformBufferRef<FMobileSceneTextureUniformParameters>::CreateUniformBufferImmediate(MobileSceneTextureParameters, UniformBuffer_SingleFrame);
-
-		// This is switching to another RT!
-		Scene->FXSystem->PostRenderOpaque(RHICmdList, View.ViewUniformBuffer, &FMobileSceneTextureUniformParameters::StaticStructMetadata, MobileSceneTextureUniformBuffer.GetReference());
-	}
-
+	
 	RHICmdList.SetCurrentStat(GET_STATID(STAT_CLMM_Translucency));
 
 	// Restart trancluceny render pass if needed
@@ -741,7 +723,7 @@ void FMobileSceneRenderer::ConditionalResolveSceneDepth(FRHICommandListImmediate
 		// Only these features require depth texture
 		bool bDecals = ViewFamily.EngineShowFlags.Decals && Scene->Decals.Num();
 		bool bModulatedShadows = ViewFamily.EngineShowFlags.DynamicShadows && bModulatedShadowsInUse;
-		
+
 		if (bDecals || bModulatedShadows || bAlwaysResolveDepth || View.bUsesSceneDepth)
 		{
 			SCOPED_DRAW_EVENT(RHICmdList, ConditionalResolveSceneDepth);
@@ -1024,12 +1006,3 @@ void FMobileSceneRenderer::PreTonemapMSAA(FRHICommandListImmediate& RHICmdList)
 		*VertexShader,
 		EDRF_UseTriangleOptimization);
 }
-
-bool FMobileSceneRenderer::IsGPUParticleCollisionEnabled(const FViewInfo& View)
-{
-	if (!View.bIsPlanarReflection && ViewFamily.EngineShowFlags.Particles)
-	{
-		return CVarMobileDisableGPUParticleCollision.GetValueOnRenderThread() == 0;
-	}
-	return false;
-}
diff --git a/Engine/Source/Runtime/Renderer/Private/SceneRendering.h b/Engine/Source/Runtime/Renderer/Private/SceneRendering.h
--- a/Engine/Source/Runtime/Renderer/Private/SceneRendering.h
+++ b/Engine/Source/Runtime/Renderer/Private/SceneRendering.h
@@ -1614,8 +1614,6 @@ protected:
 	/** On chip pre-tonemap before scene color MSAA resolve (iOS only) */
 	void PreTonemapMSAA(FRHICommandListImmediate& RHICmdList);
 
-	/** Whether GPU particle collisions simulation is allowed. */
-	bool IsGPUParticleCollisionEnabled(const FViewInfo& View);
 	void SortMobileBasePassAfterShadowInit(FExclusiveDepthStencil::Type BasePassDepthStencilAccess, FViewVisibleCommandsPerView& ViewCommandsPerView);
 	void SetupMobileBasePassAfterShadowInit(FExclusiveDepthStencil::Type BasePassDepthStencilAccess, FViewVisibleCommandsPerView& ViewCommandsPerView);
 
