Fixed r.raytracing.translucency crash

The JIRA crash mention that raytracing was 1 but I could not repro that crash. When raytracing is 0 the crash was happening and that is the fix.

diff --git a/Engine/Source/Runtime/Renderer/Private/DeferredShadingRenderer.cpp b/Engine/Source/Runtime/Renderer/Private/DeferredShadingRenderer.cpp
--- a/Engine/Source/Runtime/Renderer/Private/DeferredShadingRenderer.cpp
+++ b/Engine/Source/Runtime/Renderer/Private/DeferredShadingRenderer.cpp
@@ -1714,7 +1714,8 @@ void FDeferredShadingSceneRenderer::Render(FRHICommandListImmediate& RHICmdList)
 	TRefCountPtr<IPooledRenderTarget> SkyLightRT;
 	TRefCountPtr<IPooledRenderTarget> GlobalIlluminationRT;
 	TRefCountPtr<IPooledRenderTarget> HitDistanceRT;
-	if (IsRayTracingEnabled())
+	const bool bRayTracingEnabled = IsRayTracingEnabled();
+	if (bRayTracingEnabled)
 	{
 		for (int32 ViewIndex = 0; ViewIndex < Views.Num(); ViewIndex++)
 		{
@@ -2068,7 +2069,7 @@ void FDeferredShadingSceneRenderer::Render(FRHICommandListImmediate& RHICmdList)
 
 #if RHI_RAYTRACING
 		const bool bRaytracedTranslucency = CVarRayTracingTranslucency.GetValueOnRenderThread() != 0;
-		if (bRaytracedTranslucency)
+		if (bRayTracingEnabled && bRaytracedTranslucency)
 		{
 			ResolveSceneColor(RHICmdList);
 			RayTraceTranslucency(RHICmdList);
