Removed Intensity units for skylights when using a scene capture cubemap, also renaming the property as "Intensity Scale".
Added an "Exposure Compensation Curve" to the post process settings that add a bias based on the scene luminance in EV100.
Enabled the extended range settings for new projects.
guillaume.abadie, brian.karis

diff --git a/Engine/Plugins/Compositing/Composure/Source/Composure/Private/ComposureBlueprintLibrary.cpp b/Engine/Plugins/Compositing/Composure/Source/Composure/Private/ComposureBlueprintLibrary.cpp
--- a/Engine/Plugins/Compositing/Composure/Source/Composure/Private/ComposureBlueprintLibrary.cpp
+++ b/Engine/Plugins/Compositing/Composure/Source/Composure/Private/ComposureBlueprintLibrary.cpp
@@ -122,6 +122,7 @@ void UComposureBlueprintLibrary::CopyCameraSettingsToSceneCapture(UCameraCompone
 		COPY_PPSETTING(DepthOfFieldMinFstop, bOverride_DepthOfFieldMinFstop);
 		COPY_PPSETTING(DepthOfFieldBladeCount, bOverride_DepthOfFieldBladeCount);
 		COPY_PPSETTING(AutoExposureBias, bOverride_AutoExposureBias);
+		COPY_PPSETTING(AutoExposureBiasCurve, bOverride_AutoExposureBiasCurve);
 			
 		COPY_PPSETTING(DepthOfFieldSensorWidth, bOverride_DepthOfFieldSensorWidth);
 		COPY_PPSETTING(DepthOfFieldFocalDistance, bOverride_DepthOfFieldFocalDistance);
diff --git a/Engine/Shaders/Private/PostProcessEyeAdaptation.usf b/Engine/Shaders/Private/PostProcessEyeAdaptation.usf
--- a/Engine/Shaders/Private/PostProcessEyeAdaptation.usf
+++ b/Engine/Shaders/Private/PostProcessEyeAdaptation.usf
@@ -65,13 +65,14 @@ float ComputeEyeAdaptation(float OldExposure, float TargetExposure, float FrameT
 	return clamp(OldExposure + Diff * Factor, EyeAdaptationParams[0].z, EyeAdaptationParams[0].w);
 }
 
-float2 MainEyeAdaptationCommon()
+float4 MainEyeAdaptationCommon()
 {
-	float2 OutColor = 0;
+	float4 OutColor = 0;
+
+	const float ExposureOffsetMultipler = EyeAdaptationParams[1].x;
+	const float AverageSceneLuminance = ComputeEyeAdaptationExposure(PostprocessInput0);
+	const float TargetExposure = clamp(AverageSceneLuminance, EyeAdaptationParams[0].z, EyeAdaptationParams[0].w);
 
-	float ExposureOffsetMultipler = EyeAdaptationParams[1].x;
-	float TargetExposure = ComputeEyeAdaptationExposure(PostprocessInput0);
-	
 	float OldExposureScale = PostprocessInput0.Load(int3(0, 1, 0)).x;
 	float OldExposure = ExposureOffsetMultipler / ( OldExposureScale != 0 ? OldExposureScale : 1.0f );
 	float FrameTime = EyeAdaptationParams[1].y;
@@ -84,24 +85,25 @@ float2 MainEyeAdaptationCommon()
 
 	OutColor.x = ExposureOffsetMultipler * SmoothedExposureScale;
 	OutColor.y = ExposureOffsetMultipler * TargetExposureScale;
-
+	OutColor.z = AverageSceneLuminance;
+	OutColor.w = 0;
 
 	return OutColor;
 }
 
 void MainPS(float4 UVAndScreenPos : TEXCOORD0, 
-	out float2 OutColor : SV_Target0)
+	out float4 OutColor : SV_Target0)
 {
 	OutColor = MainEyeAdaptationCommon();
 }
 
 #if COMPUTESHADER
-RWTexture2D<float2> RWOutComputeTex;
+RWTexture2D<float4> RWOutComputeTex;
 
 [numthreads(1, 1, 1)]
 void MainCS(uint2 DispatchThreadId : SV_DispatchThreadID)
 {
-	float2 OutColor = MainEyeAdaptationCommon();
+	float4 OutColor = MainEyeAdaptationCommon();
 
 	uint2 PixelPos = DispatchThreadId + (uint2)View.ViewRectMin.xy;
 	RWOutComputeTex[PixelPos] = OutColor;
@@ -149,7 +151,7 @@ void MainReductionPS(
 
 void MainLogLuminance2ExposureScalePS(
 	in float4 UVAndScreenPos : TEXCOORD0,
-	out float2 OutColor : SV_Target0
+	out float4 OutColor : SV_Target0
 	)
 {
 	float2 UV = UVAndScreenPos.xy;
@@ -174,12 +176,13 @@ void MainLogLuminance2ExposureScalePS(
 	LogLumAve = (LogLumAve - EyeAdaptationParams[2].y) / EyeAdaptationParams[2].x;
 	
 	// Convert LogLuminanceAverage to Average Intensity
-	float LumAve = OneOverPreExposure * exp2(LogLumAve) * AverageLuminanceScale;
+	const float AverageSceneLuminance = OneOverPreExposure * exp2(LogLumAve);
+	const float LumAve = AverageSceneLuminance * AverageLuminanceScale;
 
-	float ClampedLumAve = clamp(LumAve, EyeAdaptationParams[0].z, EyeAdaptationParams[0].w );
+	const float ClampedLumAve = clamp(LumAve, EyeAdaptationParams[0].z, EyeAdaptationParams[0].w );
 	
 	// NB: this param value differs by a factor from the value used in the histogram-based alogrithm 
-	float ExposureOffsetMultipler = EyeAdaptationParams[1].x;
+	const float ExposureOffsetMultipler = EyeAdaptationParams[1].x;
 
 	// The Exposure Scale (and thus intensity) used in the previous frame
 	float ExposureScaleOld = EyeAdaptationTexture.Load(int3(0, 0, 0)).x;
@@ -197,4 +200,6 @@ void MainLogLuminance2ExposureScalePS(
 	OutColor.x = ExposureOffsetMultipler * SmoothedExposureScale;
 	// Output the target value
 	OutColor.y = ExposureOffsetMultipler * TargetExposureScale;
+	OutColor.z = AverageSceneLuminance;
+	OutColor.w = 0;
 }
\ No newline at end of file
diff --git a/Engine/Source/Editor/DetailCustomizations/Private/SkyLightComponentDetails.cpp b/Engine/Source/Editor/DetailCustomizations/Private/SkyLightComponentDetails.cpp
--- a/Engine/Source/Editor/DetailCustomizations/Private/SkyLightComponentDetails.cpp
+++ b/Engine/Source/Editor/DetailCustomizations/Private/SkyLightComponentDetails.cpp
@@ -3,6 +3,7 @@
 #include "SkyLightComponentDetails.h"
 #include "Components/SceneComponent.h"
 #include "Engine/SkyLight.h"
+#include "Components/SkyLightComponent.h"
 #include "Widgets/DeclarativeSyntaxSupport.h"
 #include "Widgets/Text/STextBlock.h"
 #include "Widgets/Input/SButton.h"
@@ -26,6 +27,23 @@ TSharedRef<IDetailCustomization> FSkyLightComponentDetails::MakeInstance()
 
 void FSkyLightComponentDetails::CustomizeDetails( IDetailLayoutBuilder& DetailLayout )
 {
+	const TArray< TWeakObjectPtr<UObject> >& SelectedObjects = DetailLayout.GetSelectedObjects();
+	for (int32 ObjectIndex = 0; ObjectIndex < SelectedObjects.Num(); ++ObjectIndex)
+	{
+		const TWeakObjectPtr<UObject>& CurrentObject = SelectedObjects[ObjectIndex];
+		if (CurrentObject.IsValid())
+		{
+			ASkyLight* CurrentCaptureActor = Cast<ASkyLight>(CurrentObject.Get());
+			if (CurrentCaptureActor != NULL)
+			{
+				SkyLight = CurrentCaptureActor;
+				break;
+			}
+		}
+	}
+
+	USkyLightComponent* SkyLightComponent = SkyLight.IsValid() ? SkyLight->GetLightComponent() : nullptr;
+
 	// Mobility property is on the scene component base class not the light component and that is why we have to use USceneComponent::StaticClass
 	TSharedRef<IPropertyHandle> MobilityHandle = DetailLayout.GetProperty(GET_MEMBER_NAME_CHECKED(USkyLightComponent, Mobility), USceneComponent::StaticClass());
 	// Set a mobility tooltip specific to lights
@@ -39,22 +57,14 @@ void FSkyLightComponentDetails::CustomizeDetails( IDetailLayoutBuilder& DetailLa
 		LightIntensityProperty->SetInstanceMetaData("UIMin", TEXT("0.0f"));
 		LightIntensityProperty->SetInstanceMetaData("UIMax", TEXT("50000.0f"));
 		LightIntensityProperty->SetInstanceMetaData("SliderExponent", TEXT("10.0f"));
-		LightIntensityProperty->SetInstanceMetaData("Units", TEXT("CandelaPerMeter2"));
-	}
-
-	const TArray< TWeakObjectPtr<UObject> >& SelectedObjects = DetailLayout.GetSelectedObjects();
 
-	for( int32 ObjectIndex = 0; ObjectIndex < SelectedObjects.Num(); ++ObjectIndex )
-	{
-		const TWeakObjectPtr<UObject>& CurrentObject = SelectedObjects[ObjectIndex];
-		if ( CurrentObject.IsValid() )
+		if (!SkyLightComponent || SkyLightComponent->SourceType != SLS_CapturedScene)
 		{
-			ASkyLight* CurrentCaptureActor = Cast<ASkyLight>(CurrentObject.Get());
-			if (CurrentCaptureActor != NULL)
-			{
-				SkyLight = CurrentCaptureActor;
-				break;
-			}
+			LightIntensityProperty->SetInstanceMetaData("Units", TEXT("CandelaPerMeter2"));
+		}
+		else
+		{
+			LightIntensityProperty->SetPropertyDisplayName(LOCTEXT("LightIntensityScaleDisplayName", "Intensity Scale"));
 		}
 	}
 
@@ -64,6 +74,10 @@ void FSkyLightComponentDetails::CustomizeDetails( IDetailLayoutBuilder& DetailLa
 	// Editing the rendering category and giving it TypeSpecific priority will place it just under the Light category
 	DetailLayout.EditCategory("Rendering", FText::GetEmpty(), ECategoryPriority::TypeSpecific);
 
+	TSharedPtr<IPropertyHandle> SourceTypeProperty = DetailLayout.GetProperty(GET_MEMBER_NAME_CHECKED(USkyLightComponent, SourceType), USkyLightComponent::StaticClass());
+	SourceTypeProperty->SetOnPropertyValueChanged(FSimpleDelegate::CreateSP(this, &FSkyLightComponentDetails::OnSourceTypeChanged));
+
+
 	DetailLayout.EditCategory( "SkyLight" )
 	.AddCustomRow( NSLOCTEXT("SkyLightDetails", "UpdateSkyLight", "Recapture Scene") )
 		.NameContent()
@@ -89,6 +103,12 @@ void FSkyLightComponentDetails::CustomizeDetails( IDetailLayoutBuilder& DetailLa
 		];
 }
 
+void FSkyLightComponentDetails::CustomizeDetails(const TSharedPtr<IDetailLayoutBuilder>& DetailBuilder)
+{
+	CachedDetailBuilder = DetailBuilder;
+	CustomizeDetails(*DetailBuilder);
+}
+
 FReply FSkyLightComponentDetails::OnUpdateSkyCapture()
 {
 	if (SkyLight.IsValid())
@@ -102,4 +122,13 @@ FReply FSkyLightComponentDetails::OnUpdateSkyCapture()
 	return FReply::Handled();
 }
 
+void FSkyLightComponentDetails::OnSourceTypeChanged()
+{
+	IDetailLayoutBuilder* DetailBuilder = CachedDetailBuilder.Pin().Get();
+	if (DetailBuilder)
+	{
+		DetailBuilder->ForceRefreshDetails();
+	}
+}
+
 #undef LOCTEXT_NAMESPACE
diff --git a/Engine/Source/Editor/DetailCustomizations/Private/SkyLightComponentDetails.h b/Engine/Source/Editor/DetailCustomizations/Private/SkyLightComponentDetails.h
--- a/Engine/Source/Editor/DetailCustomizations/Private/SkyLightComponentDetails.h
+++ b/Engine/Source/Editor/DetailCustomizations/Private/SkyLightComponentDetails.h
@@ -16,11 +16,20 @@ public:
 	/** Makes a new instance of this detail layout class for a specific detail view requesting it */
 	static TSharedRef<IDetailCustomization> MakeInstance();
 private:
+
 	/** IDetailCustomization interface */
 	virtual void CustomizeDetails( IDetailLayoutBuilder& DetailLayout ) override;
+	virtual void CustomizeDetails(const TSharedPtr<IDetailLayoutBuilder>& DetailBuilder) override;
 
 	FReply OnUpdateSkyCapture();
+
+	void OnSourceTypeChanged();
+
 private:
+
+	// The detail builder for this cusomtomisation
+	TWeakPtr<IDetailLayoutBuilder> CachedDetailBuilder;
+
 	/** The selected sky light */
 	TWeakObjectPtr<ASkyLight> SkyLight;
 };
diff --git a/Engine/Source/Editor/GameProjectGeneration/Private/GameProjectUtils.cpp b/Engine/Source/Editor/GameProjectGeneration/Private/GameProjectUtils.cpp
--- a/Engine/Source/Editor/GameProjectGeneration/Private/GameProjectUtils.cpp
+++ b/Engine/Source/Editor/GameProjectGeneration/Private/GameProjectUtils.cpp
@@ -1830,7 +1830,14 @@ bool GameProjectUtils::GenerateConfigFiles(const FProjectInformation& InProjectI
 			FileContents += TEXT("UseAudioMixer=True") LINE_TERMINATOR;
 			FileContents += LINE_TERMINATOR;
 		}
-		
+
+		if (InProjectInfo.bForceExtendedLuminanceRange)
+		{
+			FileContents += TEXT("[/Script/Engine.RendererSettings]") LINE_TERMINATOR;
+			FileContents += TEXT("r.DefaultFeature.AutoExposure.ExtendDefaultLuminanceRange=True") LINE_TERMINATOR;
+			FileContents += LINE_TERMINATOR;
+		}
+
 		if (InProjectInfo.bCopyStarterContent)
 		{
 			FString SpecificEditorStartupMap;
diff --git a/Engine/Source/Editor/GameProjectGeneration/Private/SNewProjectWizard.cpp b/Engine/Source/Editor/GameProjectGeneration/Private/SNewProjectWizard.cpp
--- a/Engine/Source/Editor/GameProjectGeneration/Private/SNewProjectWizard.cpp
+++ b/Engine/Source/Editor/GameProjectGeneration/Private/SNewProjectWizard.cpp
@@ -1406,6 +1406,7 @@ bool SNewProjectWizard::CreateProject( const FString& ProjectFile )
 	ProjectInfo.TargetedHardware = SelectedHardwareClassTarget;
 	ProjectInfo.DefaultGraphicsPerformance = SelectedGraphicsPreset;
 	ProjectInfo.bIsEnterpriseProject = (SelectedTemplate->Type == FTemplateCategory::EnterpriseCategoryName) || FApp::IsEnterpriseInstalled();
+	ProjectInfo.bForceExtendedLuminanceRange = SelectedTemplate->ProjectFile.IsEmpty();
 
 	if (!GameProjectUtils::CreateProject(ProjectInfo, FailReason, FailLog))
 	{
diff --git a/Engine/Source/Editor/GameProjectGeneration/Public/GameProjectUtils.h b/Engine/Source/Editor/GameProjectGeneration/Public/GameProjectUtils.h
--- a/Engine/Source/Editor/GameProjectGeneration/Public/GameProjectUtils.h
+++ b/Engine/Source/Editor/GameProjectGeneration/Public/GameProjectUtils.h
@@ -20,6 +20,7 @@ struct FProjectInformation
 		, bShouldGenerateCode(bInGenerateCode)
 		, bCopyStarterContent(bInCopyStarterContent)
 		, bIsEnterpriseProject(false)
+		, bForceExtendedLuminanceRange(false)
 		, TargetedHardware(EHardwareClass::Desktop)
 		, DefaultGraphicsPerformance(EGraphicsPreset::Maximum)
 	{
@@ -31,6 +32,7 @@ struct FProjectInformation
 	bool bShouldGenerateCode;
 	bool bCopyStarterContent;
 	bool bIsEnterpriseProject;
+	bool bForceExtendedLuminanceRange; // See "r.DefaultFeature.AutoExposure.ExtendDefaultLuminanceRange"
 
 	EHardwareClass::Type TargetedHardware;
 	EGraphicsPreset::Type DefaultGraphicsPerformance;
diff --git a/Engine/Source/Runtime/Engine/Classes/Engine/Scene.h b/Engine/Source/Runtime/Engine/Classes/Engine/Scene.h
--- a/Engine/Source/Runtime/Engine/Classes/Engine/Scene.h
+++ b/Engine/Source/Runtime/Engine/Classes/Engine/Scene.h
@@ -491,6 +491,14 @@ struct FCameraExposureSettings
 	UPROPERTY(Interp, BlueprintReadWrite, Category = "Exposure", meta = (UIMin = "-8.0", UIMax = "8.0", DisplayName = "Exposure Bias"))
 	float Bias;
 
+	/**
+	 * Exposure compensation based on the scene EV100.
+	 * Used to calibrate the final exposure differently depending on the average scene luminance.
+	 * 0: no adjustment, -1:2x darker, -2:4x darker, 1:2x brighter, 2:4x brighter, ...
+	 */
+	UPROPERTY(EditAnywhere, BlueprintReadWrite, Category = "Exposure", meta = (DisplayName = "Exposure Bias Curve"))
+	class UCurveFloat* BiasCurve = nullptr;
+
 	/** temporary exposed until we found good values, -8: 1/256, -10: 1/1024 */
 	UPROPERTY(Interp, BlueprintReadWrite, Category="Exposure", AdvancedDisplay, meta=(UIMin = "-16", UIMax = "0.0"))
 	float HistogramLogMin;
@@ -797,6 +805,9 @@ struct FPostProcessSettings
 	UPROPERTY(EditAnywhere, BlueprintReadWrite, Category=Overrides, meta=(PinHiddenByDefault, InlineEditConditionToggle))
 	uint8 bOverride_AutoExposureBias:1;
 
+	UPROPERTY(EditAnywhere, BlueprintReadWrite, Category=Overrides, meta=(PinHiddenByDefault, InlineEditConditionToggle))
+	uint8 bOverride_AutoExposureBiasCurve:1;
+
 	UPROPERTY(EditAnywhere, BlueprintReadWrite, Category=Overrides, meta=(PinHiddenByDefault, InlineEditConditionToggle))
 	uint8 bOverride_HistogramLogMin:1;
 
@@ -1320,6 +1331,14 @@ struct FPostProcessSettings
 	UPROPERTY(interp, BlueprintReadWrite, Category = "Lens|Exposure", meta = (UIMin = "-8.0", UIMax = "8.0", editcondition = "bOverride_AutoExposureBias", DisplayName = "Exposure Compensation "))
 	float AutoExposureBias;
 
+	/**
+	 * Exposure compensation based on the scene EV100.
+	 * Used to calibrate the final exposure differently depending on the average scene luminance.
+	 * 0: no adjustment, -1:2x darker, -2:4x darker, 1:2x brighter, 2:4x brighter, ...
+	 */
+	UPROPERTY(EditAnywhere, BlueprintReadWrite, Category = "Lens|Exposure", meta = (editcondition = "bOverride_AutoExposureBiasCurve", DisplayName = "Exposure Compensation Curve"))
+	class UCurveFloat* AutoExposureBiasCurve = nullptr;
+
 	/**
 	 * The eye adaptation will adapt to a value extracted from the luminance histogram of the scene color.
 	 * The value is defined as having x percent below this brightness. Higher values give bright spots on the screen more priority
diff --git a/Engine/Source/Runtime/Engine/Private/Scene.cpp b/Engine/Source/Runtime/Engine/Private/Scene.cpp
--- a/Engine/Source/Runtime/Engine/Private/Scene.cpp
+++ b/Engine/Source/Runtime/Engine/Private/Scene.cpp
@@ -200,6 +200,7 @@ void FCameraExposureSettings::ExportToPostProcessSettings(FPostProcessSettings*
 	OutPostProcessSettings->bOverride_AutoExposureSpeedUp = true;
 	OutPostProcessSettings->bOverride_AutoExposureSpeedDown = true;
 	OutPostProcessSettings->bOverride_AutoExposureBias = true;
+	OutPostProcessSettings->bOverride_AutoExposureBiasCurve = true;
 	OutPostProcessSettings->bOverride_HistogramLogMin = true;
 	OutPostProcessSettings->bOverride_HistogramLogMax = true;
 
@@ -210,6 +211,7 @@ void FCameraExposureSettings::ExportToPostProcessSettings(FPostProcessSettings*
 	OutPostProcessSettings->AutoExposureSpeedUp = SpeedUp;
 	OutPostProcessSettings->AutoExposureSpeedDown = SpeedDown;
 	OutPostProcessSettings->AutoExposureBias = Bias;
+	OutPostProcessSettings->AutoExposureBiasCurve = BiasCurve;
 	OutPostProcessSettings->HistogramLogMin = HistogramLogMin;
 	OutPostProcessSettings->HistogramLogMax = HistogramLogMax;
 }
@@ -644,6 +646,7 @@ FPostProcessSettings::FPostProcessSettings(const FPostProcessSettings& Settings)
 	, bOverride_AutoExposureSpeedUp(Settings.bOverride_AutoExposureSpeedUp)
 	, bOverride_AutoExposureSpeedDown(Settings.bOverride_AutoExposureSpeedDown)
 	, bOverride_AutoExposureBias(Settings.bOverride_AutoExposureBias)
+	, bOverride_AutoExposureBiasCurve(Settings.bOverride_AutoExposureBiasCurve)
 	, bOverride_HistogramLogMin(Settings.bOverride_HistogramLogMin)
 	, bOverride_HistogramLogMax(Settings.bOverride_HistogramLogMax)
 	, bOverride_LensFlareIntensity(Settings.bOverride_LensFlareIntensity)
@@ -807,6 +810,7 @@ FPostProcessSettings::FPostProcessSettings(const FPostProcessSettings& Settings)
 	, DepthOfFieldMinFstop(Settings.DepthOfFieldMinFstop)
 	, DepthOfFieldBladeCount(Settings.DepthOfFieldBladeCount)
 	, AutoExposureBias(Settings.AutoExposureBias)
+	, AutoExposureBiasCurve(Settings.AutoExposureBiasCurve)
 	, AutoExposureLowPercent(Settings.AutoExposureLowPercent)
 	, AutoExposureHighPercent(Settings.AutoExposureHighPercent)
 	, AutoExposureMinBrightness(Settings.AutoExposureMinBrightness)
diff --git a/Engine/Source/Runtime/Engine/Private/SceneView.cpp b/Engine/Source/Runtime/Engine/Private/SceneView.cpp
--- a/Engine/Source/Runtime/Engine/Private/SceneView.cpp
+++ b/Engine/Source/Runtime/Engine/Private/SceneView.cpp
@@ -1408,6 +1408,12 @@ void FSceneView::OverridePostProcessSettings(const FPostProcessSettings& Src, fl
 		}
 		PRAGMA_ENABLE_DEPRECATION_WARNINGS
 
+		// Curve assets can not be blended.
+		IF_PP(AutoExposureBiasCurve)
+		{
+			Dest.AutoExposureBiasCurve = Src.AutoExposureBiasCurve;
+		}
+
 		// actual texture cannot be blended but the intensity can be blended
 		IF_PP(LensFlareBokehShape)
 		{
diff --git a/Engine/Source/Runtime/Renderer/Private/PostProcess/PostProcessEyeAdaptation.cpp b/Engine/Source/Runtime/Renderer/Private/PostProcess/PostProcessEyeAdaptation.cpp
--- a/Engine/Source/Runtime/Renderer/Private/PostProcess/PostProcessEyeAdaptation.cpp
+++ b/Engine/Source/Runtime/Renderer/Private/PostProcess/PostProcessEyeAdaptation.cpp
@@ -11,6 +11,8 @@
 #include "PipelineStateCache.h"
 #include "Math/UnrealMathUtility.h"
 #include "ScenePrivate.h"
+#include "Curves/CurveFloat.h"
+
 
 RENDERCORE_API bool UsePreExposure(EShaderPlatform Platform);
 
@@ -44,16 +46,6 @@ TAutoConsoleVariable<float> CVarEyeAdaptationFocus(
 	TEXT(">0: Center focus, 1 is a good number (default)"),
 	ECVF_Scalability | ECVF_RenderThreadSafe);
 
-FORCEINLINE float EV100ToLuminance(float EV100)
-{
-	return 1.2 * FMath::Pow(2.0f, EV100);
-}
-
-FORCEINLINE float EV100ToLog2(float EV100)
-{
-	return EV100 + 0.263f; // Where .263 is log2(1.2)
-}
-
 /**
  *   Shared functionality used in computing the eye-adaptation parameters
  *   Compute the parameters used for eye-adaptation.  These will default to values
@@ -82,7 +74,17 @@ inline static void ComputeEyeAdaptationValues(const ERHIFeatureLevel::Type MinFe
 	float MinAverageLuminance = 1;
 	float MaxAverageLuminance = 1;
 	// This scales the average luminance AFTER it gets clamped, affecting the exposure value directly.
-	float LocalExposureMultipler = FMath::Pow(2.0f, Settings.AutoExposureBias);
+	float AutoExposureBias = Settings.AutoExposureBias;
+	if (Settings.AutoExposureBiasCurve)
+	{
+		float AverageSceneLuminance = View.GetLastAverageSceneLuminance();
+		if (AverageSceneLuminance > 0)
+		{
+			AutoExposureBias += Settings.AutoExposureBiasCurve->GetFloatValue(LuminanceToEV100(AverageSceneLuminance));
+		}
+	}
+
+	float LocalExposureMultipler = FMath::Pow(2.0f, AutoExposureBias);
 	// This scales the average luminance BEFORE it gets clamped, used to implement the calibration constant for AEM_Basic.
 	float AverageLuminanceScale = 1.f;
 
@@ -432,6 +434,8 @@ float FRCPassPostProcessEyeAdaptation::GetFixedExposure(const FViewInfo& View)
 
 void FSceneViewState::UpdatePreExposure(FViewInfo& View)
 {
+	const FSceneViewFamily& ViewFamily = *View.Family;
+
 	PreExposure = 1.f;
 	bUpdateLastExposure = false;
 
@@ -443,11 +447,9 @@ void FSceneViewState::UpdatePreExposure(FViewInfo& View)
 			PreExposure = FRCPassPostProcessEyeAdaptation::GetFixedExposure(View);
 		}
 	}
-	else if (UsePreExposure(View.GetShaderPlatform()))
+	else if (!IsRichView(ViewFamily) /*&& !ViewFamily.EngineShowFlags.VisualizeHDR */&& !ViewFamily.EngineShowFlags.VisualizeBloom && ViewFamily.EngineShowFlags.PostProcessing && ViewFamily.bResolveScene)
 	{
-		const FSceneViewFamily& ViewFamily = *View.Family;
-
-		if (!IsRichView(ViewFamily) /*&& !ViewFamily.EngineShowFlags.VisualizeHDR */&& !ViewFamily.EngineShowFlags.VisualizeBloom && ViewFamily.EngineShowFlags.PostProcessing && ViewFamily.bResolveScene)
+		if (UsePreExposure(View.GetShaderPlatform()))
 		{
 			const float PreExposureOverride = CVarEyeAdaptationPreExposureOverride.GetValueOnRenderThread();
 			const float LastExposure = View.GetLastEyeAdaptationExposure();
@@ -462,10 +464,12 @@ void FSceneViewState::UpdatePreExposure(FViewInfo& View)
 
 			bUpdateLastExposure = true;
 		}
+		// The exposure compensation curves require the scene average luminance
+		else if (View.FinalPostProcessSettings.AutoExposureBiasCurve)
+		{
+			bUpdateLastExposure = true;
+		}
 	}
-
-	// Set up view's preexposure.
-	View.PreExposure = PreExposure > 0 ? PreExposure : 1.0f;
 }
 
 FPooledRenderTargetDesc FRCPassPostProcessEyeAdaptation::ComputeOutputDesc(EPassOutputId InPassOutputId) const
@@ -820,7 +824,7 @@ void FSceneViewState::FEyeAdaptationRTManager::SwapRTs(bool bInUpdateLastExposur
 	{
 		if (!StagingBuffers[CurrentStagingBuffer].IsValid())
 		{
-			FPooledRenderTargetDesc Desc(FPooledRenderTargetDesc::Create2DDesc(FIntPoint(1, 1), PF_G32R32F, FClearValueBinding::None, TexCreate_CPUReadback | TexCreate_HideInVisualizeTexture, TexCreate_None, false));
+			FPooledRenderTargetDesc Desc(FPooledRenderTargetDesc::Create2DDesc(FIntPoint(1, 1), PF_A32B32G32R32F, FClearValueBinding::None, TexCreate_CPUReadback | TexCreate_HideInVisualizeTexture, TexCreate_None, false));
 			GRenderTargetPool.FindFreeElement(RHICmdList, Desc, StagingBuffers[CurrentStagingBuffer], TEXT("EyeAdaptationCPUReadBack"), true, ERenderTargetTransience::NonTransient);
 		}
 
@@ -837,7 +841,8 @@ void FSceneViewState::FEyeAdaptationRTManager::SwapRTs(bool bInUpdateLastExposur
 
 			if (ResultsBuffer)
 			{
-				LastExposure = *ResultsBuffer;
+				LastExposure = ResultsBuffer[0];
+				LastAverageSceneLuminance = ResultsBuffer[2];
 			}
 			RHICmdList.UnmapStagingSurface(StagingBuffers[NextStagingBuffer]->GetRenderTargetItem().ShaderResourceTexture);
 		}
@@ -857,7 +862,7 @@ TRefCountPtr<IPooledRenderTarget>&  FSceneViewState::FEyeAdaptationRTManager::Ge
 	if (!PooledRenderTarget[BufferNumber].IsValid() && RHICmdList)
 	{
 		// Create the texture needed for EyeAdaptation
-		FPooledRenderTargetDesc Desc(FPooledRenderTargetDesc::Create2DDesc(FIntPoint(1, 1), PF_G32R32F, FClearValueBinding::None, TexCreate_None, TexCreate_RenderTargetable, false));
+		FPooledRenderTargetDesc Desc(FPooledRenderTargetDesc::Create2DDesc(FIntPoint(1, 1), PF_A32B32G32R32F, FClearValueBinding::None, TexCreate_None, TexCreate_RenderTargetable, false));
 		if (GMaxRHIFeatureLevel >= ERHIFeatureLevel::SM5)
 		{
 			Desc.TargetableFlags |= TexCreate_UAV;
diff --git a/Engine/Source/Runtime/Renderer/Private/PostProcess/PostProcessEyeAdaptation.h b/Engine/Source/Runtime/Renderer/Private/PostProcess/PostProcessEyeAdaptation.h
--- a/Engine/Source/Runtime/Renderer/Private/PostProcess/PostProcessEyeAdaptation.h
+++ b/Engine/Source/Runtime/Renderer/Private/PostProcess/PostProcessEyeAdaptation.h
@@ -14,6 +14,26 @@
 
 #define EYE_ADAPTATION_PARAMS_SIZE 4
 
+FORCEINLINE float EV100ToLuminance(float EV100)
+{
+	return 1.2 * FMath::Pow(2.0f, EV100);
+}
+
+FORCEINLINE float EV100ToLog2(float EV100)
+{
+	return EV100 + 0.263f; // Where .263 is log2(1.2)
+}
+
+FORCEINLINE float LuminanceToEV100(float Luminance)
+{
+	return FMath::Log2(Luminance / 1.2f);
+}
+
+FORCEINLINE float Log2ToEV100(float Log2)
+{
+	return Log2 - 0.263f; // Where .263 is log2(1.2)
+}
+
 // Computes the eye-adaptation from HDRHistogram.
 // ePId_Input0: HDRHistogram or nothing
 // derives from TRenderingCompositePassBase<InputCount, OutputCount> 
diff --git a/Engine/Source/Runtime/Renderer/Private/PostProcess/PostProcessVisualizeHDR.cpp b/Engine/Source/Runtime/Renderer/Private/PostProcess/PostProcessVisualizeHDR.cpp
--- a/Engine/Source/Runtime/Renderer/Private/PostProcess/PostProcessVisualizeHDR.cpp
+++ b/Engine/Source/Runtime/Renderer/Private/PostProcess/PostProcessVisualizeHDR.cpp
@@ -197,16 +197,6 @@ FString LogToString(float LogValue)
 	}
 }
 
-FORCEINLINE float LuminanceToEV100(float Luminance)
-{
-	return FMath::Log2(Luminance / 1.2f);
-}
-
-FORCEINLINE float Log2ToEV100(float Log2)
-{
-	return Log2 - 0.263f; // Where .263 is log2(1.2)
-}
-
 void FRCPassPostProcessVisualizeHDR::Process(FRenderingCompositePassContext& Context)
 {
 	SCOPED_DRAW_EVENT(Context.RHICmdList, PostProcessVisualizeHDR);
diff --git a/Engine/Source/Runtime/Renderer/Private/ScenePrivate.h b/Engine/Source/Runtime/Renderer/Private/ScenePrivate.h
--- a/Engine/Source/Runtime/Renderer/Private/ScenePrivate.h
+++ b/Engine/Source/Runtime/Renderer/Private/ScenePrivate.h
@@ -786,6 +786,9 @@ private:
 		/** Get the last frame exposure value (used to compute pre-exposure) */
 		float GetLastExposure() const { return LastExposure; }
 
+		/** Get the last frame average scene luminance (used for exposure compensation curve) */
+		float GetLastAverageSceneLuminance() const { return LastAverageSceneLuminance; }
+
 	private:
 
 		/** Return one of two two render targets */
@@ -796,6 +799,8 @@ private:
 		int32 CurrentBuffer;
 
 		float LastExposure;
+		float LastAverageSceneLuminance = 0; // 0 means invalid. Used for Exposure Compensation Curve.
+
 		int32 CurrentStagingBuffer;
 		static const int32 NUM_STAGING_BUFFERS = 3;
 
@@ -1171,6 +1176,11 @@ public:
 		return EyeAdaptationRTManager.GetLastExposure();
 	}
 
+	float GetLastAverageSceneLuminance() const
+	{
+		return EyeAdaptationRTManager.GetLastAverageSceneLuminance();
+	}
+
 	bool HasValidTonemappingLUT() const
 	{
 		return bValidTonemappingLUT;
diff --git a/Engine/Source/Runtime/Renderer/Private/SceneRendering.cpp b/Engine/Source/Runtime/Renderer/Private/SceneRendering.cpp
--- a/Engine/Source/Runtime/Renderer/Private/SceneRendering.cpp
+++ b/Engine/Source/Runtime/Renderer/Private/SceneRendering.cpp
@@ -1684,6 +1684,16 @@ float FViewInfo::GetLastEyeAdaptationExposure() const
 	return 0.f; // Invalid exposure
 }
 
+float FViewInfo::GetLastAverageSceneLuminance() const
+{
+	const FSceneViewState* EffectiveViewState = GetEffectiveViewState();
+	if (EffectiveViewState)
+	{
+		return EffectiveViewState->GetLastAverageSceneLuminance();
+	}
+	return 0.f; // Invalid scene luminance
+}
+
 void FViewInfo::SetValidTonemappingLUT() const
 {
 	FSceneViewState* EffectiveViewState = GetEffectiveViewState();
diff --git a/Engine/Source/Runtime/Renderer/Private/SceneRendering.h b/Engine/Source/Runtime/Renderer/Private/SceneRendering.h
--- a/Engine/Source/Runtime/Renderer/Private/SceneRendering.h
+++ b/Engine/Source/Runtime/Renderer/Private/SceneRendering.h
@@ -1062,6 +1062,9 @@ public:
 	/** Get the last valid exposure value for eye adapation. */
 	float GetLastEyeAdaptationExposure() const;
 
+	/** Get the last valid average scene luminange for eye adapation (exposure compensation curve). */
+	float GetLastAverageSceneLuminance() const;
+
 	/** Informs sceneinfo that tonemapping LUT has queued commands to compute it at least once */
 	void SetValidTonemappingLUT() const;
 
