Smooth ssao result when TAA isn't used

diff --git a/Engine/Shaders/Private/PostProcessAmbientOcclusion.usf b/Engine/Shaders/Private/PostProcessAmbientOcclusion.usf
--- a/Engine/Shaders/Private/PostProcessAmbientOcclusion.usf
+++ b/Engine/Shaders/Private/PostProcessAmbientOcclusion.usf
@@ -866,6 +866,36 @@ void MainCS(
 	// http://on-demand.gputechconf.com/gtc/2010/presentations/S12312-DirectCompute-Pre-Conference-Tutorial.pdf
 	OutTexture[PixelPos] = OutColor.r;
 }
+
+// xy - output rect min
+// zw - output rect width and height
+float4 SSAOSmoothParams;
+RWTexture2D<float> SSAOSmoothResult;
+
+[numthreads(THREADGROUP_SIZEX, THREADGROUP_SIZEY, 1)]
+void MainSSAOSmoothCS(
+	uint2 GroupId : SV_GroupID,
+	uint2 DispatchThreadId : SV_DispatchThreadID,
+	uint2 GroupThreadId : SV_GroupThreadID)
+{
+	BRANCH
+	if (any(DispatchThreadId >= SSAOSmoothParams.zw))
+	{
+		return;
+	}
+
+	uint2 DestPixelPos = SSAOSmoothParams.xy + DispatchThreadId;
+	float2 SampleUV = DestPixelPos * PostprocessInput0Size.zw;
+
+	// Use a 4x4 box filter because the random texture is tiled 4x4
+	float Result;
+	Result = PostprocessInput0.SampleLevel(PostprocessInput0Sampler, SampleUV, 0).r;
+	Result += PostprocessInput0.SampleLevel(PostprocessInput0Sampler, SampleUV, 0, int2(2, 0)).r;
+	Result += PostprocessInput0.SampleLevel(PostprocessInput0Sampler, SampleUV, 0, int2(0, 2)).r;
+	Result += PostprocessInput0.SampleLevel(PostprocessInput0Sampler, SampleUV, 0, int2(2, 2)).r;
+
+	SSAOSmoothResult[DestPixelPos] = Result * 0.25;
+}
 #endif
 
 
diff --git a/Engine/Source/Runtime/Renderer/Private/CompositionLighting/CompositionLighting.cpp b/Engine/Source/Runtime/Renderer/Private/CompositionLighting/CompositionLighting.cpp
--- a/Engine/Source/Runtime/Renderer/Private/CompositionLighting/CompositionLighting.cpp
+++ b/Engine/Source/Runtime/Renderer/Private/CompositionLighting/CompositionLighting.cpp
@@ -55,6 +55,12 @@ static TAutoConsoleVariable<int32> CVarSubsurfaceScattering(
 	TEXT(" 1: enabled (default)"),
 	ECVF_RenderThreadSafe | ECVF_Scalability);
 
+static TAutoConsoleVariable<int32> CVarSSAOSmoothPass(
+	TEXT("r.AmbientOcclusion.Compute.Smooth"),
+	1,
+	TEXT("Whether to smooth SSAO output when TAA is disabled"),
+	ECVF_RenderThreadSafe | ECVF_Scalability);
+
 bool IsAmbientCubemapPassRequired(const FSceneView& View)
 {
 	FScene* Scene = (FScene*)View.Family->Scene;
@@ -209,11 +215,25 @@ static FRenderingCompositeOutputRef AddPostProcessingAmbientOcclusion(FRHIComman
 		GBufferA = Context.Graph.RegisterPass(new(FMemStack::Get()) FRCPassPostProcessInput(SceneContext.GBufferA));
 	}
 
-	FRenderingCompositePass* AmbientOcclusionPassMip0 = Context.Graph.RegisterPass(new(FMemStack::Get()) FRCPassPostProcessAmbientOcclusion(Context.View, FullResAOType, false));
+	// If there is no temporal upsampling, we need a smooth pass to get rid of the grid pattern.
+	// PS version has relatively smooth result so no need to do extra work
+	const bool bNeedSmoothingPass = FullResAOType != ESSAOType::EPS && Context.View.AntiAliasingMethod != AAM_TemporalAA && CVarSSAOSmoothPass.GetValueOnRenderThread();
+	const EPixelFormat SmoothingPassInputFormat = bNeedSmoothingPass ? PF_G8 : PF_Unknown;
+
+	FRenderingCompositePass* AmbientOcclusionPassMip0 = Context.Graph.RegisterPass(new(FMemStack::Get()) FRCPassPostProcessAmbientOcclusion(Context.View, FullResAOType, false, bNeedSmoothingPass, SmoothingPassInputFormat));
 	AmbientOcclusionPassMip0->SetInput(ePId_Input0, GBufferA);
 	AmbientOcclusionPassMip0->SetInput(ePId_Input1, AmbientOcclusionInMip1);
 	AmbientOcclusionPassMip0->SetInput(ePId_Input2, AmbientOcclusionPassMip1);
 	AmbientOcclusionPassMip0->SetInput(ePId_Input3, HZBInput);
+	FRenderingCompositePass* FinalOutputPass = AmbientOcclusionPassMip0;
+
+	if (bNeedSmoothingPass)
+	{
+		FRenderingCompositePass* SSAOSmoothPass;
+		SSAOSmoothPass = Context.Graph.RegisterPass(new (FMemStack::Get()) FRCPassPostProcessAmbientOcclusionSmooth(FullResAOType, true));
+		SSAOSmoothPass->SetInput(ePId_Input0, AmbientOcclusionPassMip0);
+		FinalOutputPass = SSAOSmoothPass;
+	}
 
 	// to make sure this pass is processed as well (before), needed to make process decals before computing AO
 	if(AmbientOcclusionInMip1)
@@ -225,11 +245,11 @@ static FRenderingCompositeOutputRef AddPostProcessingAmbientOcclusion(FRHIComman
 		AmbientOcclusionPassMip0->AddDependency(Context.FinalOutput);
 	}
 
-	Context.FinalOutput = FRenderingCompositeOutputRef(AmbientOcclusionPassMip0);
+	Context.FinalOutput = FRenderingCompositeOutputRef(FinalOutputPass);
 
 	SceneContext.bScreenSpaceAOIsValid = true;
 
-	return FRenderingCompositeOutputRef(AmbientOcclusionPassMip0);
+	return FRenderingCompositeOutputRef(FinalOutputPass);
 }
 
 void FCompositionLighting::ProcessBeforeBasePass(FRHICommandListImmediate& RHICmdList, FViewInfo& View, bool bDBuffer, uint32 SSAOLevels)
diff --git a/Engine/Source/Runtime/Renderer/Private/CompositionLighting/PostProcessAmbientOcclusion.cpp b/Engine/Source/Runtime/Renderer/Private/CompositionLighting/PostProcessAmbientOcclusion.cpp
--- a/Engine/Source/Runtime/Renderer/Private/CompositionLighting/PostProcessAmbientOcclusion.cpp
+++ b/Engine/Source/Runtime/Renderer/Private/CompositionLighting/PostProcessAmbientOcclusion.cpp
@@ -19,6 +19,7 @@
 DECLARE_GPU_STAT_NAMED(SSAOSetup, TEXT("ScreenSpace AO Setup") );
 DECLARE_GPU_STAT_NAMED(SSAO, TEXT("ScreenSpace AO") );
 DECLARE_GPU_STAT_NAMED(BasePassAO, TEXT("BasePass AO") );
+DECLARE_GPU_STAT_NAMED(SSAOSmooth, TEXT("SSAO smooth"));
 
 // Tile size for the AmbientOcclusion compute shader, tweaked for 680 GTX. */
 // see GCN Performance Tip 21 http://developer.amd.com/wordpress/media/2013/05/GCNPerformanceTweets.pdf 
@@ -480,6 +481,167 @@ bool FRCPassPostProcessAmbientOcclusionSetup::IsInitialPass() const
 
 // --------------------------------------------------------
 
+class FPostProcessAmbientOcclusionSmoothCS : public FGlobalShader
+{
+	DECLARE_SHADER_TYPE(FPostProcessAmbientOcclusionSmoothCS, Global);
+
+	static bool ShouldCompilePermutation(const FGlobalShaderPermutationParameters& Parameters)
+	{
+		return IsFeatureLevelSupported(Parameters.Platform, ERHIFeatureLevel::SM5);
+	}
+
+	static void ModifyCompilationEnvironment(const FGlobalShaderPermutationParameters& Parameters, FShaderCompilerEnvironment& OutEnvironment)
+	{
+		constexpr int32 ThreadGroupSize1D = FRCPassPostProcessAmbientOcclusionSmooth::ThreadGroupSize1D;
+
+		FGlobalShader::ModifyCompilationEnvironment(Parameters, OutEnvironment);
+		OutEnvironment.SetDefine(TEXT("COMPUTE_SHADER"), 1);
+		OutEnvironment.SetDefine(TEXT("THREADGROUP_SIZEX"), ThreadGroupSize1D);
+		OutEnvironment.SetDefine(TEXT("THREADGROUP_SIZEY"), ThreadGroupSize1D);
+	}
+
+	/** Default constructor. */
+	FPostProcessAmbientOcclusionSmoothCS() {}
+
+public:
+	FPostProcessPassParameters PostprocessParameter;
+	FShaderParameter SSAOSmoothParams;
+	FShaderParameter SSAOSmoothResult;
+
+	/** Initialization constructor. */
+	FPostProcessAmbientOcclusionSmoothCS(const ShaderMetaType::CompiledShaderInitializerType& Initializer)
+		: FGlobalShader(Initializer)
+	{
+		PostprocessParameter.Bind(Initializer.ParameterMap);
+		SSAOSmoothParams.Bind(Initializer.ParameterMap, TEXT("SSAOSmoothParams"));
+		SSAOSmoothResult.Bind(Initializer.ParameterMap, TEXT("SSAOSmoothResult"));
+	}
+
+	template <typename TRHICmdList>
+	void SetParameters(
+		TRHICmdList& RHICmdList,
+		const FRenderingCompositePassContext& Context,
+		const FIntRect& OutputRect,
+		FUnorderedAccessViewRHIParamRef OutUAV)
+	{
+		const FComputeShaderRHIParamRef ShaderRHI = GetComputeShader();
+
+		FGlobalShader::SetParameters<FViewUniformShaderParameters>(RHICmdList, ShaderRHI, Context.View.ViewUniformBuffer);
+
+		PostprocessParameter.SetCS(ShaderRHI, Context, RHICmdList, TStaticSamplerState<SF_Bilinear, AM_Clamp, AM_Clamp, AM_Clamp>::GetRHI());
+
+		FVector4 SSAOSmoothParamsValue(OutputRect.Min.X, OutputRect.Min.Y, OutputRect.Width(), OutputRect.Height());
+		SetShaderValue(RHICmdList, ShaderRHI, SSAOSmoothParams, SSAOSmoothParamsValue);
+
+		RHICmdList.SetUAVParameter(ShaderRHI, SSAOSmoothResult.GetBaseIndex(), OutUAV);
+	}
+
+	template <typename TRHICmdList>
+	void UnsetParameters(TRHICmdList& RHICmdList)
+	{
+		const FComputeShaderRHIParamRef ShaderRHI = GetComputeShader();
+		RHICmdList.SetUAVParameter(ShaderRHI, SSAOSmoothResult.GetBaseIndex(), nullptr);
+	}
+
+	// FShader interface.
+	virtual bool Serialize(FArchive& Ar) override
+	{
+		bool bShaderHasOutdatedParameters = FGlobalShader::Serialize(Ar);
+		Ar << PostprocessParameter << SSAOSmoothParams << SSAOSmoothResult;
+		return bShaderHasOutdatedParameters;
+	}
+
+	static const TCHAR* GetSourceFilename()
+	{
+		return TEXT("/Engine/Private/PostProcessAmbientOcclusion.usf");
+	}
+
+	static const TCHAR* GetFunctionName()
+	{
+		return TEXT("MainSSAOSmoothCS");
+	}
+};
+
+IMPLEMENT_SHADER_TYPE3(FPostProcessAmbientOcclusionSmoothCS, SF_Compute);
+
+FRCPassPostProcessAmbientOcclusionSmooth::FRCPassPostProcessAmbientOcclusionSmooth(ESSAOType InAOType, bool bInDirectOutput)
+	: AOType(InAOType)
+	, bDirectOutput(bInDirectOutput)
+{}
+
+template <typename TRHICmdList>
+void FRCPassPostProcessAmbientOcclusionSmooth::DispatchCS(
+	TRHICmdList& RHICmdList,
+	const FRenderingCompositePassContext& Context,
+	const FIntRect& OutputRect,
+	FUnorderedAccessViewRHIParamRef OutUAV) const
+{
+	TShaderMapRef<FPostProcessAmbientOcclusionSmoothCS> ComputeShader(Context.GetShaderMap());
+	RHICmdList.SetComputeShader(ComputeShader->GetComputeShader());
+	ComputeShader->SetParameters(RHICmdList, Context, OutputRect, OutUAV);
+	const uint32 NumGroupsX = FMath::DivideAndRoundUp(OutputRect.Width(), ThreadGroupSize1D);
+	const uint32 NumGroupsY = FMath::DivideAndRoundUp(OutputRect.Height(), ThreadGroupSize1D);
+	DispatchComputeShader(RHICmdList, *ComputeShader, NumGroupsX, NumGroupsY, 1);
+	ComputeShader->UnsetParameters(RHICmdList);
+}
+
+void FRCPassPostProcessAmbientOcclusionSmooth::Process(FRenderingCompositePassContext& Context)
+{
+	SCOPED_GPU_STAT(Context.RHICmdList, SSAOSmooth);
+
+	SetRenderTarget(Context.RHICmdList, nullptr, nullptr);
+	Context.SetViewportAndCallRHI(Context.View.ViewRect);
+
+	FSceneRenderTargets& SceneContext = FSceneRenderTargets::Get(Context.RHICmdList);
+	const FSceneRenderTargetItem& DestRenderTarget = bDirectOutput ? SceneContext.ScreenSpaceAO->GetRenderTargetItem() : PassOutputs[0].RequestSurface(Context);
+	const FIntPoint OutputExtent = bDirectOutput ? SceneContext.GetBufferSizeXY() : PassOutputs[0].RenderTargetDesc.Extent;
+	const int32 DownSampleFactor = FMath::DivideAndRoundUp(Context.ReferenceBufferSize.X, OutputExtent.X);
+	const FIntRect OutputRect = Context.GetViewport() / DownSampleFactor;
+
+	if (AOType == ESSAOType::EAsyncCS)
+	{
+		FRHIAsyncComputeCommandListImmediate& AsyncComputeCmdList = FRHICommandListExecutor::GetImmediateAsyncComputeCommandList();
+		FComputeFenceRHIRef AsyncStartFence = Context.RHICmdList.CreateComputeFence(TEXT("AsyncStartFence"));
+
+		SCOPED_COMPUTE_EVENTF(AsyncComputeCmdList, SSAOSmooth, TEXT("SSAO smooth %dx%d"), OutputRect.Width(), OutputRect.Height());
+
+		Context.RHICmdList.TransitionResource(EResourceTransitionAccess::ERWBarrier, EResourceTransitionPipeline::EGfxToCompute, DestRenderTarget.UAV, AsyncStartFence);
+		AsyncComputeCmdList.WaitComputeFence(AsyncStartFence);
+		DispatchCS(AsyncComputeCmdList, Context, OutputRect, DestRenderTarget.UAV);
+	}
+	else
+	{
+		check(AOType == ESSAOType::ECS);
+		SCOPED_DRAW_EVENTF(Context.RHICmdList, SSAOSmooth, TEXT("SSAO smooth %dx%d"), OutputRect.Width(), OutputRect.Height());
+
+		Context.RHICmdList.TransitionResource(EResourceTransitionAccess::ERWBarrier, EResourceTransitionPipeline::EGfxToCompute, DestRenderTarget.UAV);
+		DispatchCS(Context.RHICmdList, Context, OutputRect, DestRenderTarget.UAV);
+	}
+}
+
+FPooledRenderTargetDesc FRCPassPostProcessAmbientOcclusionSmooth::ComputeOutputDesc(EPassOutputId InPassOutputId) const
+{
+	if (bDirectOutput)
+	{
+		FPooledRenderTargetDesc Ret;
+		Ret.DebugName = TEXT("AmbientOcclusionDirect");
+		return Ret;
+	}
+
+	const FPooledRenderTargetDesc* Input0Desc = GetInputDesc(ePId_Input0);
+	check(Input0Desc);
+	FPooledRenderTargetDesc Ret = *Input0Desc;
+	Ret.Reset();
+	Ret.Format = PF_G8;
+	Ret.ClearValue = FClearValueBinding::None;
+	Ret.TargetableFlags &= ~TexCreate_DepthStencilTargetable;
+	Ret.TargetableFlags |= TexCreate_UAV;
+	Ret.DebugName = TEXT("SSAOSmoothResult");
+	return Ret;
+}
+
+// --------------------------------------------------------
+
 FArchive& operator<<(FArchive& Ar, FScreenSpaceAOParameters& This)
 {
 	Ar << This.ScreenSpaceAOParams;
@@ -706,9 +868,11 @@ void FRCPassPostProcessAmbientOcclusion::DispatchCS(TRHICmdList& RHICmdList, con
 
 // --------------------------------------------------------
 
-FRCPassPostProcessAmbientOcclusion::FRCPassPostProcessAmbientOcclusion(const FSceneView& View, ESSAOType InAOType, bool bInAOSetupAsInput)
+FRCPassPostProcessAmbientOcclusion::FRCPassPostProcessAmbientOcclusion(const FSceneView& View, ESSAOType InAOType, bool bInAOSetupAsInput, bool bInForcecIntermediateOutput, EPixelFormat InIntermediateFormatOverride)
 	: AOType(InAOType)
-	, bAOSetupAsInput(bInAOSetupAsInput)	
+	, IntermediateFormatOverride(InIntermediateFormatOverride)
+	, bAOSetupAsInput(bInAOSetupAsInput)
+	, bForceIntermediateOutput(bInForcecIntermediateOutput)
 {
 }
 
@@ -909,7 +1073,7 @@ void FRCPassPostProcessAmbientOcclusion::Process(FRenderingCompositePassContext&
 	const FSceneRenderTargetItem* DestRenderTarget = 0;
 	FSceneRenderTargets& SceneContext = FSceneRenderTargets::Get(Context.RHICmdList);
 
-	if(bAOSetupAsInput)
+	if(bAOSetupAsInput || bForceIntermediateOutput)
 	{
 		DestRenderTarget = &PassOutputs[0].RequestSurface(Context);
 	}
@@ -950,7 +1114,7 @@ void FRCPassPostProcessAmbientOcclusion::Process(FRenderingCompositePassContext&
 
 FPooledRenderTargetDesc FRCPassPostProcessAmbientOcclusion::ComputeOutputDesc(EPassOutputId InPassOutputId) const
 {
-	if(!bAOSetupAsInput)
+	if(!bAOSetupAsInput && !bForceIntermediateOutput)
 	{
 		FPooledRenderTargetDesc Ret;
 
@@ -978,6 +1142,11 @@ FPooledRenderTargetDesc FRCPassPostProcessAmbientOcclusion::ComputeOutputDesc(EP
 	}
 	Ret.DebugName = TEXT("AmbientOcclusion");
 
+	if (IntermediateFormatOverride != PF_Unknown)
+	{
+		Ret.Format = IntermediateFormatOverride;
+	}
+
 	return Ret;
 }
 
diff --git a/Engine/Source/Runtime/Renderer/Private/CompositionLighting/PostProcessAmbientOcclusion.h b/Engine/Source/Runtime/Renderer/Private/CompositionLighting/PostProcessAmbientOcclusion.h
--- a/Engine/Source/Runtime/Renderer/Private/CompositionLighting/PostProcessAmbientOcclusion.h
+++ b/Engine/Source/Runtime/Renderer/Private/CompositionLighting/PostProcessAmbientOcclusion.h
@@ -66,6 +66,30 @@ private:
 	bool IsInitialPass() const;
 };
 
+// ePId_Input0: Lower resolution AO result buffer
+class FRCPassPostProcessAmbientOcclusionSmooth : public TRenderingCompositePassBase<1, 1>
+{
+public:
+	static constexpr int32 ThreadGroupSize1D = 8;
+
+	FRCPassPostProcessAmbientOcclusionSmooth(ESSAOType InAOType, bool bInDirectOutput = false);
+
+	// interface FRenderingCompositePass ---------
+	virtual void Process(FRenderingCompositePassContext& Context) final override;
+	virtual void Release() final override { delete this; }
+	virtual FPooledRenderTargetDesc ComputeOutputDesc(EPassOutputId InPassOutputId) const final override;
+
+private:
+	template <typename TRHICmdList>
+	void DispatchCS(
+		TRHICmdList& RHICmdList,
+		const FRenderingCompositePassContext& Context,
+		const FIntRect& OutputRect,
+		FUnorderedAccessViewRHIParamRef OutUAV) const;
+
+	const ESSAOType AOType;
+	const bool bDirectOutput;
+};
 
 // ePId_Input0: defines the resolution we compute AO and provides the normal (only needed if bInAOSetupAsInput)
 // ePId_Input1: setup in same resolution as ePId_Input1 for depth expect when running in full resolution, then it's half (only needed if bInAOSetupAsInput)
@@ -76,7 +100,7 @@ class FRCPassPostProcessAmbientOcclusion : public TRenderingCompositePassBase<4,
 {
 public:
 	// @param bInAOSetupAsInput true:use AO setup as input, false: use GBuffer normal and native z depth
-	FRCPassPostProcessAmbientOcclusion(const FSceneView& View, ESSAOType InAOType, bool bInAOSetupAsInput = true);
+	FRCPassPostProcessAmbientOcclusion(const FSceneView& View, ESSAOType InAOType, bool bInAOSetupAsInput = true, bool bInForcecIntermediateOutput = false, EPixelFormat InIntermediateFormatOverride = PF_Unknown);
 
 	// interface FRenderingCompositePass ---------
 	virtual void Process(FRenderingCompositePassContext& Context) override;
@@ -95,7 +119,9 @@ private:
 	void DispatchCS(TRHICmdList& RHICmdList, const FRenderingCompositePassContext& Context, const FIntPoint& TexSize, FUnorderedAccessViewRHIParamRef OutTextureUAV);
 	
 	const ESSAOType AOType;
+	const EPixelFormat IntermediateFormatOverride;
 	const bool bAOSetupAsInput;
+	const bool bForceIntermediateOutput;
 };
 
 // apply the AO to the SceneColor (lightmapped object), extra pas that is not always needed
