Enable stencil testing for post process materials.
  - Write directly to post process input 0 if able.
  - Copy post process input 0 to target if unable.
  - Allow SM4 to make copy DS if necessay.
  - Allow PPM to use blend modes to composite onto output.

diff --git a/Engine/Source/Runtime/Engine/Classes/Engine/EngineTypes.h b/Engine/Source/Runtime/Engine/Classes/Engine/EngineTypes.h
--- a/Engine/Source/Runtime/Engine/Classes/Engine/EngineTypes.h
+++ b/Engine/Source/Runtime/Engine/Classes/Engine/EngineTypes.h
@@ -502,6 +502,19 @@ enum EMaterialSamplerType
 	SAMPLERTYPE_MAX,
 };
 
+UENUM()
+enum EMaterialStencilCompare
+{
+	MSC_Less			UMETA(DisplayName = "Less Than"),
+	MSC_LessEqual		UMETA(DisplayName = "Less Than or Equal"),
+	MSC_Greater			UMETA(DisplayName = "Greater Than"),
+	MSC_GreaterEqual	UMETA(DisplayName = "Greater Than or Equal"),
+	MSC_Equal			UMETA(DisplayName = "Equal"),
+	MSC_NotEqual		UMETA(DisplayName = "Not Equal"),
+	MSC_Never			UMETA(DisplayName = "Never"),
+	MSC_Always			UMETA(DisplayName = "Always"),
+	MSC_Count			UMETA(Hidden),
+};
 /**	Lighting build quality enumeration */
 UENUM()
 enum ELightingBuildQuality
diff --git a/Engine/Source/Runtime/Engine/Classes/Materials/Material.h b/Engine/Source/Runtime/Engine/Classes/Materials/Material.h
--- a/Engine/Source/Runtime/Engine/Classes/Materials/Material.h
+++ b/Engine/Source/Runtime/Engine/Classes/Materials/Material.h
@@ -5,6 +5,7 @@
 #include "CoreMinimal.h"
 #include "UObject/ObjectMacros.h"
 #include "Misc/Guid.h"
+#include "Engine/EngineTypes.h"
 #include "RenderCommandFence.h"
 #include "Templates/ScopedPointer.h"
 #include "Materials/MaterialInterface.h"
@@ -852,6 +853,19 @@ public:
 	UPROPERTY(EditAnywhere, BlueprintReadWrite, Category = PostProcessMaterial, meta = (DisplayName = "Output Alpha"))
 	bool BlendableOutputAlpha;
 
+	/** 
+	 * Selectively execute post process material only for pixels that pass the stencil test against the Custom Depth/Stencil buffer. 
+	 * Pixels that fail the stencil test are filled with the previous post process material output or scene color.
+	 */
+	UPROPERTY(EditAnywhere, Category = PostProcessMaterial, AdvancedDisplay)
+	uint8 bEnableStencilTest : 1;
+
+	UPROPERTY(EditAnywhere, Category = PostProcessMaterial, AdvancedDisplay, meta = (EditCondition = "bEnableStencilTest"))
+	TEnumAsByte<EMaterialStencilCompare> StencilCompare;
+
+	UPROPERTY(EditAnywhere, Category = PostProcessMaterial, AdvancedDisplay, meta = (EditCondition = "bEnableStencilTest"))
+	uint8 StencilRefValue = 0;
+
 	/** Controls how the Refraction input is interpreted and how the refraction offset into scene color is computed for this material. */
 	UPROPERTY(EditAnywhere, Category=Refraction)
 	TEnumAsByte<enum ERefractionMode> RefractionMode;
diff --git a/Engine/Source/Runtime/Engine/Private/Materials/HLSLMaterialTranslator.h b/Engine/Source/Runtime/Engine/Private/Materials/HLSLMaterialTranslator.h
--- a/Engine/Source/Runtime/Engine/Private/Materials/HLSLMaterialTranslator.h
+++ b/Engine/Source/Runtime/Engine/Private/Materials/HLSLMaterialTranslator.h
@@ -3833,6 +3833,7 @@ protected:
 	void UseSceneTextureId(ESceneTextureId SceneTextureId, bool bTextureLookup)
 	{
 		MaterialCompilationOutput.bNeedsSceneTextures = true;
+		MaterialCompilationOutput.UsedSceneTextures |= (1ull << SceneTextureId);
 
 		if(Material->GetMaterialDomain() == MD_DeferredDecal)
 		{
diff --git a/Engine/Source/Runtime/Engine/Private/Materials/Material.cpp b/Engine/Source/Runtime/Engine/Private/Materials/Material.cpp
--- a/Engine/Source/Runtime/Engine/Private/Materials/Material.cpp
+++ b/Engine/Source/Runtime/Engine/Private/Materials/Material.cpp
@@ -879,6 +879,7 @@ UMaterial::UMaterial(const FObjectInitializer& ObjectInitializer)
 	BlendableLocation = BL_AfterTonemapping;
 	BlendablePriority = 0;
 	BlendableOutputAlpha = false;
+	bEnableStencilTest = false;
 
 	bUseEmissiveForDynamicAreaLighting = false;
 	bBlockGI = false;
@@ -4211,14 +4212,18 @@ bool UMaterial::CanEditChange(const UProperty* InProperty) const
 
 		if (PropertyName == GET_MEMBER_NAME_STRING_CHECKED(UMaterial, BlendableLocation) ||
 			PropertyName == GET_MEMBER_NAME_STRING_CHECKED(UMaterial, BlendablePriority) || 
-			PropertyName == GET_MEMBER_NAME_STRING_CHECKED(UMaterial, BlendableOutputAlpha)	)
+			PropertyName == GET_MEMBER_NAME_STRING_CHECKED(UMaterial, BlendableOutputAlpha) ||
+			PropertyName == GET_MEMBER_NAME_STRING_CHECKED(UMaterial, bEnableStencilTest) ||
+			PropertyName == GET_MEMBER_NAME_STRING_CHECKED(UMaterial, StencilCompare) ||
+			PropertyName == GET_MEMBER_NAME_STRING_CHECKED(UMaterial, StencilRefValue)
+			)
 		{
 			return MaterialDomain == MD_PostProcess;
 		}
 
 		if (PropertyName == GET_MEMBER_NAME_STRING_CHECKED(UMaterial, BlendMode))
 		{
-			return MaterialDomain == MD_DeferredDecal || MaterialDomain == MD_Surface || MaterialDomain == MD_Volume || MaterialDomain == MD_UI;
+			return MaterialDomain == MD_DeferredDecal || MaterialDomain == MD_Surface || MaterialDomain == MD_Volume || MaterialDomain == MD_UI || (MaterialDomain == MD_PostProcess && BlendableOutputAlpha);
 		}
 	
 		if (PropertyName == GET_MEMBER_NAME_STRING_CHECKED(UMaterial, ShadingModel))
diff --git a/Engine/Source/Runtime/Engine/Private/Materials/MaterialShared.cpp b/Engine/Source/Runtime/Engine/Private/Materials/MaterialShared.cpp
--- a/Engine/Source/Runtime/Engine/Private/Materials/MaterialShared.cpp
+++ b/Engine/Source/Runtime/Engine/Private/Materials/MaterialShared.cpp
@@ -448,6 +448,7 @@ void FMaterialCompilationOutput::Serialize(FArchive& Ar)
 {
 	UniformExpressionSet.Serialize(Ar);
 
+	Ar << UsedSceneTextures;
 	Ar << NumUsedUVScalars;
 	Ar << NumUsedCustomInterpolatorScalars;
 	Ar << EstimatedNumTextureSamplesVS;
@@ -1281,6 +1282,21 @@ bool FMaterialResource::GetBlendableOutputAlpha() const
 	return Material->BlendableOutputAlpha;
 }
 
+bool FMaterialResource::IsStencilTestEnabled() const
+{
+	return GetMaterialDomain() == MD_PostProcess && Material->bEnableStencilTest;
+}
+
+uint32 FMaterialResource::GetStencilRefValue() const
+{
+	return GetMaterialDomain() == MD_PostProcess ? Material->StencilRefValue : 0;
+}
+
+uint32 FMaterialResource::GetStencilCompare() const
+{
+	return GetMaterialDomain() == MD_PostProcess ? uint32(Material->StencilCompare.GetValue()) : 0;
+}
+
 UMaterialInterface* FMaterialResource::GetMaterialInterface() const 
 { 
 	return MaterialInstance ? (UMaterialInterface*)MaterialInstance : (UMaterialInterface*)Material;
@@ -2530,7 +2546,8 @@ bool FMaterial::WritesEveryPixel(bool bShadowPass) const
 		// Render dithered material as masked if a stencil prepass is not used (UE-50064, UE-49537)
 		&& !((bShadowPass || !bStencilDitheredLOD) && IsDitheredLODTransition())
 		&& !IsWireframe()
-		&& !(bStencilDitheredLOD && IsDitheredLODTransition() && IsUsedWithInstancedStaticMeshes());
+		&& !(bStencilDitheredLOD && IsDitheredLODTransition() && IsUsedWithInstancedStaticMeshes())
+		&& !IsStencilTestEnabled();
 }
 
 #if WITH_EDITOR
diff --git a/Engine/Source/Runtime/Engine/Private/ShaderDerivedDataVersion.h b/Engine/Source/Runtime/Engine/Private/ShaderDerivedDataVersion.h
--- a/Engine/Source/Runtime/Engine/Private/ShaderDerivedDataVersion.h
+++ b/Engine/Source/Runtime/Engine/Private/ShaderDerivedDataVersion.h
@@ -11,5 +11,5 @@
 // In case of merge conflicts with DDC versions, you MUST generate a new GUID and set this new
 // guid as version
 
-#define GLOBALSHADERMAP_DERIVEDDATA_VER			TEXT("A97D25A9914A43C193801914248E0F66")
-#define MATERIALSHADERMAP_DERIVEDDATA_VER		TEXT("AD9DE195370D4C7DB112DF9967877D21")
\ No newline at end of file
+#define GLOBALSHADERMAP_DERIVEDDATA_VER			TEXT("4FFBFDF601234B80893B59E7779F89E7")
+#define MATERIALSHADERMAP_DERIVEDDATA_VER		TEXT("FE2EEAF7F8BC456195582BBB000D17CD")
diff --git a/Engine/Source/Runtime/Engine/Public/MaterialShared.h b/Engine/Source/Runtime/Engine/Public/MaterialShared.h
--- a/Engine/Source/Runtime/Engine/Public/MaterialShared.h
+++ b/Engine/Source/Runtime/Engine/Public/MaterialShared.h
@@ -417,6 +417,7 @@ class FMaterialCompilationOutput
 {
 public:
 	FMaterialCompilationOutput() :
+		UsedSceneTextures(0),
 		NumUsedUVScalars(0),
 		NumUsedCustomInterpolatorScalars(0),
 		EstimatedNumTextureSamplesVS(0),
@@ -438,6 +439,9 @@ public:
 
 	FUniformExpressionSet UniformExpressionSet;
 
+	/** Bitfield of the ESceneTextures used */
+	uint64 UsedSceneTextures;
+
 	/** Number of used custom UV scalars. */
 	uint8 NumUsedUVScalars;
 
@@ -896,6 +900,7 @@ public:
 	uint32 GetNumUsedUVScalars() const { return MaterialCompilationOutput.NumUsedUVScalars; }
 	uint32 GetNumUsedCustomInterpolatorScalars() const { return MaterialCompilationOutput.NumUsedCustomInterpolatorScalars; }
 	void GetEstimatedNumTextureSamples(uint32& VSSamples, uint32& PSSamples) const { VSSamples = MaterialCompilationOutput.EstimatedNumTextureSamplesVS; PSSamples = MaterialCompilationOutput.EstimatedNumTextureSamplesPS; }
+	bool UsesSceneTexture(uint32 TexId) const { return MaterialCompilationOutput.UsedSceneTextures & (1ull << TexId); }
 
 	bool IsValidForRendering(bool bFailOnInvalid = false) const
 	{
@@ -1283,6 +1288,9 @@ public:
 	virtual int32 GetNumCustomizedUVs() const { return 0; }
 	virtual int32 GetBlendableLocation() const { return 0; }
 	virtual bool GetBlendableOutputAlpha() const { return false; }
+	virtual bool IsStencilTestEnabled() const { return false; }
+	virtual uint32 GetStencilRefValue() const { return 0; }
+	virtual uint32 GetStencilCompare() const { return 0; }
 	/**
 	 * Should shaders compiled for this material be saved to disk?
 	 */
@@ -1953,6 +1961,9 @@ public:
 	ENGINE_API virtual int32 GetNumCustomizedUVs() const override;
 	ENGINE_API virtual int32 GetBlendableLocation() const override;
 	ENGINE_API virtual bool GetBlendableOutputAlpha() const override;
+	ENGINE_API virtual bool IsStencilTestEnabled() const override;
+	ENGINE_API virtual uint32 GetStencilRefValue() const override;
+	ENGINE_API virtual uint32 GetStencilCompare() const override;
 	ENGINE_API virtual float GetRefractionDepthBiasValue() const override;
 	ENGINE_API virtual float GetMaxDisplacement() const override;
 	ENGINE_API virtual bool ShouldApplyFogging() const override;
diff --git a/Engine/Source/Runtime/Renderer/Private/PostProcess/PostProcessMaterial.cpp b/Engine/Source/Runtime/Renderer/Private/PostProcess/PostProcessMaterial.cpp
--- a/Engine/Source/Runtime/Renderer/Private/PostProcess/PostProcessMaterial.cpp
+++ b/Engine/Source/Runtime/Renderer/Private/PostProcess/PostProcessMaterial.cpp
@@ -13,8 +13,26 @@
 #include "PostProcess/SceneFilterRendering.h"
 #include "SceneRendering.h"
 #include "ClearQuad.h"
+#include "Materials/MaterialExpressionSceneTexture.h"
 #include "RHI/Public/PipelineStateCache.h"
 
+static TAutoConsoleVariable<int32> CVarPostProcessAllowStencilTest(
+	TEXT("r.PostProcessAllowStencilTest"),
+	1,
+	TEXT("Enables stencil testing in post process materials.\n")
+	TEXT("0: disable stencil testing\n")
+	TEXT("1: allow stencil testing\n")
+	TEXT("2: allow stencil testing and, if necessary, making a copy of custom depth/stencil buffer\n")
+	);
+
+static TAutoConsoleVariable<int32> CVarPostProcessAllowBlendModes(
+	TEXT("r.PostProcessAllowBlendModes"),
+	1,
+	TEXT("Enables blend modes in post process materials.\n")
+	TEXT("0: disable blend modes. Uses replace\n")
+	TEXT("1: allow blend modes\n")
+	);
+
 enum class EPostProcessMaterialTarget
 {
 	HighEnd,
@@ -205,45 +223,187 @@ TGlobalResource<FPostProcessMaterialVertexDeclaration> GPostProcessMaterialVerte
 void FRCPassPostProcessMaterial::Process(FRenderingCompositePassContext& Context)
 {
 	FMaterialRenderProxy* Proxy = MaterialInterface->GetRenderProxy();
-
 	check(Proxy);
 
 	ERHIFeatureLevel::Type FeatureLevel = Context.View.GetFeatureLevel();
 
+	FSceneRenderTargets& SceneContext = FSceneRenderTargets::Get(Context.RHICmdList);
+	
 	const FMaterial* Material = Proxy->GetMaterial(FeatureLevel);
-
 	check(Material);
 
-	const FViewInfo& View = Context.View;
-	const FSceneViewFamily& ViewFamily = *(View.Family);
+	const FMaterialShaderMap* MaterialShaderMap = Material->GetRenderingThreadShaderMap();
+
+	int32 AllowStencilTest = CVarPostProcessAllowStencilTest.GetValueOnRenderThread();
+	bool bReadsCustomDepthStencil = MaterialShaderMap->UsesSceneTexture(PPI_CustomDepth) || MaterialShaderMap->UsesSceneTexture(PPI_CustomStencil);
+	bool bDoStencilTest = false;
+	if (Material->IsStencilTestEnabled() && AllowStencilTest > 0)
+	{
+		static const IConsoleVariable* CVarCustomDepth = IConsoleManager::Get().FindConsoleVariable(TEXT("r.CustomDepth"));
+
+		int32 CustomDepthSetting = CVarCustomDepth->GetInt();
+		if (CustomDepthSetting == 3)
+		{
+			bool AllowStencilTestWithCopy = AllowStencilTest == 2;
+
+			// do the stencil test if 
+			// >= SM5 
+			//   OR 
+			// reads DS but allowed make DS copy 
+			//   OR 
+			// DS not read at all.
+			bDoStencilTest = (FeatureLevel >= ERHIFeatureLevel::SM5) || 
+				((bReadsCustomDepthStencil == AllowStencilTestWithCopy) || AllowStencilTestWithCopy);
+		}
+		else
+		{
+			UE_LOG(LogRenderer, Warning, TEXT("PostProcessMaterial uses stencil test, but custom stencil not allocated. Set r.CustomDepth to 3 to allocate custom stencil."));
+		}
+	}
 
 	const FPooledRenderTargetDesc* InputDesc = GetInputDesc(ePId_Input0);
 
-	const FSceneRenderTargetItem& DestRenderTarget = PassOutputs[0].RequestSurface(Context);
+	FIntVector RectSize = GetOutput(ePId_Output0)->RenderTargetDesc.GetSize();
+	SCOPED_DRAW_EVENTF(Context.RHICmdList, PostProcessMaterial, TEXT("PostProcessMaterial %dx%d Material=%s"), RectSize.X, RectSize.Y, *Material->GetFriendlyName());
 
-	FIntRect SrcRect = Context.SceneColorViewRect;
-	FIntRect DestRect = Context.GetSceneColorDestRect(DestRenderTarget);
-	checkf(DestRect.Size() == SrcRect.Size(), TEXT("Post process material should not be used as upscaling pass."));
+	// Copy of custom depth/stencil buffer if HW does not support simultaneously a texture bound as DepthRead_StencilRead and SRV
+	TRefCountPtr<IPooledRenderTarget> CustomDepthStencilCopy;
+
+	// The PP target - either from the render target pool or the ePId_Input0
+	const FSceneRenderTargetItem* DestRenderTarget = nullptr;
+	const FSceneRenderTargetItem* CustomDepthStencilTarget = nullptr;
+
+	FDepthStencilStateRHIParamRef DepthStencilState;
+	uint32 StencilRefValue = 0;
+
+	if (bDoStencilTest)
+	{
+		CustomDepthStencilTarget = &SceneContext.CustomDepth->GetRenderTargetItem();
 
-	FIntPoint SrcSize = InputDesc->Extent;
+		// SM4 HW lacks support for texture bound as DepthRead_StencilRead and SRV simultaneously thus make a copy of DS buffer
+		if (FeatureLevel < ERHIFeatureLevel::SM5 && bReadsCustomDepthStencil)
+		{
+			// Dest param of CopyResource() call can only be an SRV (No render target flags) on DX10.0 (SM4)
+			FPooledRenderTargetDesc DSCopyDesc = SceneContext.CustomDepth->GetDesc();
+			DSCopyDesc.Flags = TexCreate_None;
+			DSCopyDesc.TargetableFlags = TexCreate_ShaderResource;
+
+			GRenderTargetPool.FindFreeElement(Context.RHICmdList, DSCopyDesc, CustomDepthStencilCopy, TEXT("CustomDepthStencilCopy"));
+
+			Context.RHICmdList.CopyTexture(
+				SceneContext.CustomDepth->GetRenderTargetItem().ShaderResourceTexture,
+				CustomDepthStencilCopy->GetRenderTargetItem().ShaderResourceTexture,
+				FRHICopyTextureInfo()
+				);
+
+			// The copy DS buffer was created only as a SRV (no render target flags), thus swap with real CustomDepth buffer
+			// with copy to allow DS buffer with render target flag to be set as DepthRead_StencilRead and DS buffer copy to
+			// be set as SRV for post process material.
+			Swap(SceneContext.CustomDepth, CustomDepthStencilCopy);
+		}
 
-	SCOPED_DRAW_EVENTF(Context.RHICmdList, PostProcessMaterial, TEXT("PostProcessMaterial %dx%d Material=%s"), DestRect.Width(), DestRect.Height(), *Material->GetFriendlyName());
+		// essentially only care about format and dimensions
+		FPooledRenderTargetDesc TargetDesc = PassOutputs[0].RenderTargetDesc;
+		TargetDesc.AutoWritable = InputDesc->AutoWritable;
 
-	ERenderTargetLoadAction LoadAction = Context.GetLoadActionForRenderTarget(DestRenderTarget);
+		// write directly into PPI0 if material does not read from it, otherwise make a copy of PPI0.
+		if (InputDesc->Compare(TargetDesc, false) && !MaterialShaderMap->UsesSceneTexture(PPI_PostProcessInput0))
+		{
+			PassOutputs[0].PooledRenderTarget = GetInput(ePId_Input0)->GetOutput()->RequestInput();
+			DestRenderTarget = &PassOutputs[0].RequestSurface(Context);
+		}
+		else
+		{
+			DestRenderTarget = &PassOutputs[0].RequestSurface(Context);
+
+			Context.RHICmdList.CopyTexture(
+				GetInput(ePId_Input0)->GetOutput()->RequestSurface(Context).ShaderResourceTexture,
+				DestRenderTarget->TargetableTexture,
+				FRHICopyTextureInfo()
+				);
+		}
+
+		static const FDepthStencilStateRHIParamRef StencilStates[] =
+		{
+			TStaticDepthStencilState<false, CF_Always, true, CF_Less>::GetRHI(),
+			TStaticDepthStencilState<false, CF_Always, true, CF_LessEqual>::GetRHI(),
+			TStaticDepthStencilState<false, CF_Always, true, CF_Greater>::GetRHI(),
+			TStaticDepthStencilState<false, CF_Always, true, CF_GreaterEqual>::GetRHI(),
+			TStaticDepthStencilState<false, CF_Always, true, CF_Equal>::GetRHI(),
+			TStaticDepthStencilState<false, CF_Always, true, CF_NotEqual>::GetRHI(),
+			TStaticDepthStencilState<false, CF_Always, true, CF_Never>::GetRHI(),
+			TStaticDepthStencilState<false, CF_Always, true, CF_Always>::GetRHI(),
+		};
+		static_assert(EMaterialStencilCompare::MSC_Count == ARRAY_COUNT(StencilStates), "Ensure that all EMaterialStencilCompare values are accounted for.");
+
+		DepthStencilState = StencilStates[Material->GetStencilCompare()];
+		StencilRefValue = Material->GetStencilRefValue();
+	}
+	else
+	{
+		DestRenderTarget = &PassOutputs[0].RequestSurface(Context);
+		DepthStencilState = TStaticDepthStencilState<false, CF_Always>::GetRHI();
+	}
+
+	FBlendStateRHIParamRef BlendState = TStaticBlendState<>::GetRHI();
+	if (Material->GetBlendableOutputAlpha() && CVarPostProcessAllowBlendModes.GetValueOnRenderThread() != 0)
+	{
+		static const FBlendStateRHIParamRef BlendStates[] =
+		{
+			TStaticBlendState<>::GetRHI(),
+			TStaticBlendState<>::GetRHI(),
+			TStaticBlendState<CW_RGB, BO_Add, BF_SourceAlpha, BF_InverseSourceAlpha, BO_Add, BF_Zero, BF_One>::GetRHI(),
+			TStaticBlendState<CW_RGB, BO_Add, BF_One, BF_One>::GetRHI(),
+			TStaticBlendState<CW_RGB, BO_Add, BF_DestColor, BF_Zero>::GetRHI(),
+			TStaticBlendState<CW_RGBA, BO_Add, BF_One, BF_InverseSourceAlpha, BO_Add, BF_One, BF_InverseSourceAlpha>::GetRHI(),
+		};
+		static_assert(EBlendMode::BLEND_MAX == ARRAY_COUNT(BlendStates), "Ensure that all EBlendMode values are accounted for.");
+
+		BlendState = BlendStates[Material->GetBlendMode()];
+	}
+
+	FIntRect SrcRect = Context.SceneColorViewRect;
+	FIntRect DestRect = Context.GetSceneColorDestRect(*DestRenderTarget);
+	checkf(DestRect.Size() == SrcRect.Size(), TEXT("Post process material should not be used as upscaling pass."));
 	
-	FRHIRenderPassInfo RPInfo(DestRenderTarget.TargetableTexture, MakeRenderTargetActions(LoadAction, ERenderTargetStoreAction::EStore));
+	FRHIRenderPassInfo RPInfo;
+	if (CustomDepthStencilTarget)
+	{
+		RPInfo = FRHIRenderPassInfo(
+			DestRenderTarget->TargetableTexture,
+			MakeRenderTargetActions(Context.GetLoadActionForRenderTarget(*DestRenderTarget), ERenderTargetStoreAction::EStore),
+			CustomDepthStencilTarget->TargetableTexture,
+			MakeDepthStencilTargetActions(
+				MakeRenderTargetActions(ERenderTargetLoadAction::ENoAction, ERenderTargetStoreAction::ENoAction),
+				MakeRenderTargetActions(ERenderTargetLoadAction::ELoad, ERenderTargetStoreAction::ENoAction)
+				),
+			FExclusiveDepthStencil::DepthRead_StencilRead
+			);
+	}
+	else
+	{
+		RPInfo = FRHIRenderPassInfo(
+			DestRenderTarget->TargetableTexture,
+			MakeRenderTargetActions(Context.GetLoadActionForRenderTarget(*DestRenderTarget), ERenderTargetStoreAction::EStore)
+			);
+	}
+
 	Context.RHICmdList.BeginRenderPass(RPInfo, TEXT("PostProcessMaterial"));
 	{
 		Context.SetViewportAndCallRHI(DestRect);
 
 		FGraphicsPipelineStateInitializer GraphicsPSOInit;
 		Context.RHICmdList.ApplyCachedRenderTargets(GraphicsPSOInit);
-		GraphicsPSOInit.BlendState = TStaticBlendState<>::GetRHI();
+		GraphicsPSOInit.BlendState = BlendState;
 		GraphicsPSOInit.RasterizerState = TStaticRasterizerState<>::GetRHI();
-		GraphicsPSOInit.DepthStencilState = TStaticDepthStencilState<false, CF_Always>::GetRHI();
+		GraphicsPSOInit.DepthStencilState = DepthStencilState;
 		GraphicsPSOInit.PrimitiveType = PT_TriangleList;
 
-		const FMaterialShaderMap* MaterialShaderMap = Material->GetRenderingThreadShaderMap();
+		const FViewInfo& View = Context.View;
+		const FSceneViewFamily& ViewFamily = *(View.Family);
+
+		FIntPoint SrcSize = InputDesc->Extent;
+
 		FShader* VertexShader = nullptr;
 
 		// uses mobile's post process material.
@@ -257,6 +417,7 @@ void FRCPassPostProcessMaterial::Process(FRenderingCompositePassContext& Context
 			GraphicsPSOInit.BoundShaderState.PixelShaderRHI = GETSAFERHISHADER_PIXEL(PixelShader_Mobile);
 
 			SetGraphicsPipelineState(Context.RHICmdList, GraphicsPSOInit);
+			Context.RHICmdList.SetStencilRef(StencilRefValue);
 
 			VertexShader_Mobile->SetParameters(Context.RHICmdList, Context, Proxy);
 			PixelShader_Mobile->SetParameters(Context.RHICmdList, Context, Proxy);
@@ -273,6 +434,7 @@ void FRCPassPostProcessMaterial::Process(FRenderingCompositePassContext& Context
 			GraphicsPSOInit.BoundShaderState.PixelShaderRHI = GETSAFERHISHADER_PIXEL(PixelShader_HighEnd);
 
 			SetGraphicsPipelineState(Context.RHICmdList, GraphicsPSOInit);
+			Context.RHICmdList.SetStencilRef(StencilRefValue);
 
 			VertexShader_HighEnd->SetParameters(Context.RHICmdList, Context, Proxy);
 			PixelShader_HighEnd->SetParameters(Context.RHICmdList, Context, Proxy);
@@ -289,6 +451,7 @@ void FRCPassPostProcessMaterial::Process(FRenderingCompositePassContext& Context
 			GraphicsPSOInit.BoundShaderState.PixelShaderRHI = GETSAFERHISHADER_PIXEL(PixelShader_HighEnd);
 
 			SetGraphicsPipelineState(Context.RHICmdList, GraphicsPSOInit);
+			Context.RHICmdList.SetStencilRef(StencilRefValue);
 
 			VertexShader_HighEnd->SetParameters(Context.RHICmdList, Context, Proxy);
 			PixelShader_HighEnd->SetParameters(Context.RHICmdList, Context, Proxy);
@@ -309,7 +472,12 @@ void FRCPassPostProcessMaterial::Process(FRenderingCompositePassContext& Context
 			EDRF_UseTriangleOptimization);
 	}
 	Context.RHICmdList.EndRenderPass();
-	Context.RHICmdList.CopyToResolveTarget(DestRenderTarget.TargetableTexture, DestRenderTarget.ShaderResourceTexture, FResolveParams());
+	Context.RHICmdList.CopyToResolveTarget(DestRenderTarget->TargetableTexture, DestRenderTarget->ShaderResourceTexture, FResolveParams());
+
+	if (CustomDepthStencilCopy.IsValid() && CustomDepthStencilCopy->GetRenderTargetItem().IsValid())
+	{
+		Swap(SceneContext.CustomDepth, CustomDepthStencilCopy);
+	}
 
 	if(Material->NeedsGBuffer())
 	{
