Remove redefinition of log category that is included via PCH to fix build break.

diff --git a/Engine/Source/Runtime/Renderer/Private/PostProcess/PostProcessMaterial.cpp b/Engine/Source/Runtime/Renderer/Private/PostProcess/PostProcessMaterial.cpp
--- a/Engine/Source/Runtime/Renderer/Private/PostProcess/PostProcessMaterial.cpp
+++ b/Engine/Source/Runtime/Renderer/Private/PostProcess/PostProcessMaterial.cpp
@@ -5,6 +5,7 @@
 =============================================================================*/
 
 #include "PostProcess/PostProcessMaterial.h"
+#include "RendererModule.h"
 #include "Materials/Material.h"
 #include "MaterialShaderType.h"
 #include "MaterialShader.h"
