Fixed stale proxy pointers inside VelocityData. By mistake non-movable primitives were added to the VelocityData inside FScene::UpdatePrimitiveTransform_RenderThread. Only movable primitives can be removed from the VelocityData. Thus leaving a stale pointer, which later crashes application.

diff --git a/Engine/Source/Runtime/Renderer/Private/RendererScene.cpp b/Engine/Source/Runtime/Renderer/Private/RendererScene.cpp
--- a/Engine/Source/Runtime/Renderer/Private/RendererScene.cpp
+++ b/Engine/Source/Runtime/Renderer/Private/RendererScene.cpp
@@ -1188,7 +1188,10 @@ void FScene::UpdatePrimitiveTransform_RenderThread(FRHICommandListImmediate& RHI
 	// (note that the octree update relies on the bounds not being modified yet).
 	PrimitiveSceneInfo->RemoveFromScene(bUpdateStaticDrawLists);
 
-	VelocityData.UpdateTransform(PrimitiveSceneInfo, LocalToWorld, PrimitiveSceneProxy->GetLocalToWorld());
+	if (PrimitiveSceneInfo->Proxy->IsMovable())
+	{
+		VelocityData.UpdateTransform(PrimitiveSceneInfo, LocalToWorld, PrimitiveSceneProxy->GetLocalToWorld());
+	}
 
 	if (GWarningOnRedundantTransformUpdate && PrimitiveSceneProxy->WouldSetTransformBeRedundant(LocalToWorld, WorldBounds, LocalBounds, AttachmentRootPosition))
 	{
@@ -2279,10 +2282,8 @@ void FSceneVelocityData::StartFrame(FScene* Scene)
 			// Recreate PrimitiveUniformBuffer on the frame after the primitive moved, since it contains PreviousLocalToWorld
 			VelocityData.PrimitiveSceneInfo->SetNeedsUniformBufferUpdate(true);
 
-			if (VelocityData.PrimitiveSceneInfo->IsIndexValid())
-			{
-				AddPrimitiveToUpdateGPU(*Scene, VelocityData.PrimitiveSceneInfo->GetIndex());
-			}
+			check(VelocityData.PrimitiveSceneInfo->IsIndexValid());
+			AddPrimitiveToUpdateGPU(*Scene, VelocityData.PrimitiveSceneInfo->GetIndex());
 		}
 
 		if (bTrimOld && (InternalFrameIndex - VelocityData.LastFrameUsed) > 10)
diff --git a/Engine/Source/Runtime/Renderer/Private/ScenePrivate.h b/Engine/Source/Runtime/Renderer/Private/ScenePrivate.h
--- a/Engine/Source/Runtime/Renderer/Private/ScenePrivate.h
+++ b/Engine/Source/Runtime/Renderer/Private/ScenePrivate.h
@@ -2193,6 +2193,8 @@ public:
 	 */
 	void UpdateTransform(FPrimitiveSceneInfo* PrimitiveSceneInfo, FMatrix LocalToWorld, FMatrix PreviousLocalToWorld)
 	{
+		check(PrimitiveSceneInfo->Proxy->IsMovable());
+
 		FComponentVelocityData& VelocityData = ComponentData.FindOrAdd(PrimitiveSceneInfo->PrimitiveComponentId);
 		VelocityData.LocalToWorld = LocalToWorld;
 		VelocityData.LastFrameUsed = InternalFrameIndex;
