fix validation error in d3d when injecting lighting from a rect.

diff --git a/Engine/Source/Runtime/Renderer/Private/VolumeLighting.h b/Engine/Source/Runtime/Renderer/Private/VolumeLighting.h
--- a/Engine/Source/Runtime/Renderer/Private/VolumeLighting.h
+++ b/Engine/Source/Runtime/Renderer/Private/VolumeLighting.h
@@ -96,7 +96,7 @@ public:
 			SetShaderValue(RHICmdList, ShaderRHI, DepthBiasParameters, FVector2D(ShadowMap->GetShaderDepthBias(), 1.0f / (ShadowMap->MaxSubjectZ - ShadowMap->MinSubjectZ)));
 
 			FTextureRHIParamRef ShadowDepthTextureResource = nullptr;
-			if (LightType == LightType_Point)
+			if (LightType == LightType_Point || LightType == LightType_Rect)
 			{
 				if (GBlackTexture && GBlackTexture->TextureRHI)
 				{
