fix for lighttype rect shadowmap not being set

diff --git a/Engine/Source/Runtime/Renderer/Private/VolumeLighting.h b/Engine/Source/Runtime/Renderer/Private/VolumeLighting.h
--- a/Engine/Source/Runtime/Renderer/Private/VolumeLighting.h
+++ b/Engine/Source/Runtime/Renderer/Private/VolumeLighting.h
@@ -118,7 +118,7 @@ public:
 				);
 		}
 
-		OnePassShadowParameters.Set(RHICmdList, ShaderRHI, bDynamicallyShadowed && LightType == LightType_Point ? ShadowMap : NULL);
+		OnePassShadowParameters.Set(RHICmdList, ShaderRHI, bDynamicallyShadowed && (LightType == LightType_Point || LightType == LightType_Rect)? ShadowMap : NULL);
 
 		const FStaticShadowDepthMap* StaticShadowDepthMap = LightSceneInfo->Proxy->GetStaticShadowDepthMap();
 		const uint32 bStaticallyShadowedValue = LightSceneInfo->IsPrecomputedLightingValid() && StaticShadowDepthMap && StaticShadowDepthMap->Data && StaticShadowDepthMap->TextureRHI ? 1 : 0;
