Assign shadow map channel to directional light in deferred shading path per Daniel's comment because translucency pass may need it. Tested on a Win64 client

diff --git a/Engine/Source/Runtime/Renderer/Private/RendererScene.cpp b/Engine/Source/Runtime/Renderer/Private/RendererScene.cpp
--- a/Engine/Source/Runtime/Renderer/Private/RendererScene.cpp
+++ b/Engine/Source/Runtime/Renderer/Private/RendererScene.cpp
@@ -1595,8 +1595,9 @@ void FScene::AddLightSceneInfo_RenderThread(FLightSceneInfo* LightSceneInfo)
 	// Add the light to the light list.
 	LightSceneInfo->Id = Lights.Add(FLightSceneInfoCompact(LightSceneInfo));
 	const FLightSceneInfoCompact& LightSceneInfoCompact = Lights[LightSceneInfo->Id];
+	const bool bDirectionalLight = LightSceneInfo->Proxy->GetLightType() == LightType_Directional;
 
-	if (LightSceneInfo->Proxy->GetLightType() == LightType_Directional &&
+	if (bDirectionalLight &&
 		// Only use a stationary or movable light
 		!LightSceneInfo->Proxy->HasStaticLighting())
 	{
@@ -1636,8 +1637,10 @@ void FScene::AddLightSceneInfo_RenderThread(FLightSceneInfo* LightSceneInfo)
 	}
 
 	const bool bForwardShading = IsForwardShadingEnabled(GetShaderPlatform());
-
-	if (bForwardShading && (LightSceneInfo->Proxy->CastsDynamicShadow() || LightSceneInfo->Proxy->GetLightFunctionMaterial()))
+	// Need to set shadow map channel for directional light in deferred shading path also.
+	// In translucency pass, TLM_SurfacePerPixelLighting uses forward shading and requires light data set up correctly.
+	// Only done for directional light in deferred path because translucent objects only receive dynamic shadow from directional light
+	if ((bForwardShading || bDirectionalLight) && (LightSceneInfo->Proxy->CastsDynamicShadow() || LightSceneInfo->Proxy->GetLightFunctionMaterial()))
 	{
 		AssignAvailableShadowMapChannelForLight(LightSceneInfo);
 	}
