Workaround Apple A11 & A12 GPUs not clamping F32->F16 conversions to F16_MIN/F16_MAX like the A10 and earlier which meant reflection brightness of 0 would cause chaos when inverted.

diff --git a/Engine/Source/Runtime/Renderer/Private/MobileBasePass.cpp b/Engine/Source/Runtime/Renderer/Private/MobileBasePass.cpp
--- a/Engine/Source/Runtime/Renderer/Private/MobileBasePass.cpp
+++ b/Engine/Source/Runtime/Renderer/Private/MobileBasePass.cpp
@@ -473,7 +473,7 @@ void TMobileBasePassPSPolicyParamType<FUniformLightMapPolicy>::GetShaderBindings
 			{
 				float AverageBrightness = 1.0f;
 				GetSkyTextureParams(Scene, AverageBrightness, ReflectionCubemapTextures[0], ReflectionParams.W);
-				ReflectionParams.X = 1.0f / AverageBrightness;
+				ReflectionParams.X = FMath::Max(FMath::Min(1.0f / AverageBrightness, 65504.f), -65504.f);
 			}
 		}
 		else
@@ -493,7 +493,7 @@ void TMobileBasePassPSPolicyParamType<FUniformLightMapPolicy>::GetShaderBindings
 						{
 							ReflectionCubemapTextures[i] = PrimitiveSceneInfo->CachedReflectionCaptureProxies[i]->EncodedHDRCubemap;
 						}
-						ReflectionParams[i] = 1.0f / ReflectionProxy->EncodedHDRAverageBrightness;
+						ReflectionParams.X = FMath::Max(FMath::Min(1.0f / ReflectionProxy->EncodedHDRAverageBrightness, 65504.f), -65504.f);
 					}
 				}
 			}
@@ -504,7 +504,7 @@ void TMobileBasePassPSPolicyParamType<FUniformLightMapPolicy>::GetShaderBindings
 					&& PrimitiveSceneInfo->CachedReflectionCaptureProxy->EncodedHDRCubemap
 					&& PrimitiveSceneInfo->CachedReflectionCaptureProxy->EncodedHDRCubemap->IsInitialized())
 				{
-					ReflectionParams.X = 1.0f / PrimitiveSceneInfo->CachedReflectionCaptureProxy->EncodedHDRAverageBrightness;
+					ReflectionParams.X = FMath::Max(FMath::Min(1.0f / PrimitiveSceneInfo->CachedReflectionCaptureProxy->EncodedHDRAverageBrightness, 65504.f), -65504.f);
 					ReflectionCubemapTextures[0] = PrimitiveSceneInfo->CachedReflectionCaptureProxy->EncodedHDRCubemap;
 				}
 			}
