FFT Bloom:  Apparent clipping of bloom shape when screen has high aspect ratio.
This bug was visible only when a very small (few pixels), but very bright light source would produce a very large convolution bloom that covered much of the screen.

diff --git a/Engine/Source/Runtime/Renderer/Private/PostProcess/PostProcessFFTBloom.cpp b/Engine/Source/Runtime/Renderer/Private/PostProcess/PostProcessFFTBloom.cpp
--- a/Engine/Source/Runtime/Renderer/Private/PostProcess/PostProcessFFTBloom.cpp
+++ b/Engine/Source/Runtime/Renderer/Private/PostProcess/PostProcessFFTBloom.cpp
@@ -614,7 +614,7 @@ void FRCPassFFTBloom::InitializeDomainParameters(FRenderingCompositePassContext&
 	// so if the kernel is being applied on the edge of the image it will see padding and not periodicity
 	// NB:  If the kernel padding would force a transform buffer that is too big for group shared memory (> 4096)
 	//      we clamp it.  This could result in a wrap-around in the bloom (from one side of the screen to the other),
-	//      but since the amplitude of the bloom kernel tails is usually very small, this shouldnt be too bad.
+	//      but since the amplitude of the bloom kernel tails is usually very small, this shouldn't be too bad.
 	auto KernelRadiusSupportFunctor = [KernelSupportScale, KernelSupportScaleClamp](const FIntPoint& Size) ->int32 {
 
 		float ClampedKernelSupportScale = (KernelSupportScaleClamp > 0) ? FMath::Min(KernelSupportScale, KernelSupportScaleClamp) : KernelSupportScale;
@@ -663,9 +663,14 @@ void FRCPassFFTBloom::InitializeDomainParameters(FRenderingCompositePassContext&
 	// Capture the region of interest
 	ImageRect = InputRect;
 	const FIntPoint ImageSize = ImageRect.Size();
-	
+
+	// The length of the a side of the square kernel image in pixels
+
+	int32 KernelSize = FMath::CeilToInt(KernelSupportScale * FMath::Max(ImageSize.X, ImageSize.Y));
+
 	int32 SpectralPadding = KernelRadiusSupportFunctor(ImageSize);
 
+
 	// The following are mathematically equivalent
 	// 1) Horizontal FFT / Vertical FFT / Filter / Vertical InvFFT / Horizontal InvFFT
 	// 2) Vertical FFT / Horizontal FFT / Filter / Horizontal InvFFT / Vertical InvFFT
@@ -675,9 +680,13 @@ void FRCPassFFTBloom::InitializeDomainParameters(FRenderingCompositePassContext&
 	// the width of the kernel.  The ImageRect is virtually padded
 	// with black to account for the gather action of the convolution.
 	FIntPoint PaddedImageSize = ImageSize + FIntPoint(SpectralPadding, SpectralPadding);
+
+	PaddedImageSize.X = FMath::Max(PaddedImageSize.X, KernelSize);
+	PaddedImageSize.Y = FMath::Max(PaddedImageSize.Y, KernelSize);
+
 	FrequencySize = FIntPoint(FMath::RoundUpToPowerOfTwo(PaddedImageSize.X), FMath::RoundUpToPowerOfTwo(PaddedImageSize.Y));
 
-	// Chose to do to transform in the direction that results in writting the least amount of data to main memory.
+	// Choose to do to transform in the direction that results in writing the least amount of data to main memory.
 
 	bDoHorizontalFirst = ((FrequencySize.Y * PaddedImageSize.X) > (FrequencySize.X * PaddedImageSize.Y));
 
