Added post processing option to motion blur to set the target frame rate for shutter angle.

The default is 30hz, which matches the current behavior. Assigning a value of zero tracks the current frame rate as the target using an exponential moving average.

diff --git a/Engine/Source/Runtime/Engine/Classes/Engine/Scene.h b/Engine/Source/Runtime/Engine/Classes/Engine/Scene.h
--- a/Engine/Source/Runtime/Engine/Classes/Engine/Scene.h
+++ b/Engine/Source/Runtime/Engine/Classes/Engine/Scene.h
@@ -1014,6 +1014,9 @@ struct FPostProcessSettings
 	UPROPERTY(EditAnywhere, BlueprintReadWrite, Category=Overrides, meta=(PinHiddenByDefault, InlineEditConditionToggle))
 	uint8 bOverride_MotionBlurMax:1;
 
+	UPROPERTY(EditAnywhere, BlueprintReadWrite, Category=Overrides, meta=(PinHiddenByDefault, InlineEditConditionToggle))
+	uint8 bOverride_MotionBlurTargetFPS : 1;
+
 	UPROPERTY(EditAnywhere, BlueprintReadWrite, Category=Overrides, meta=(PinHiddenByDefault, InlineEditConditionToggle))
 	uint8 bOverride_MotionBlurPerObjectSize:1;
 
@@ -1580,6 +1583,15 @@ struct FPostProcessSettings
 	/** max distortion caused by motion blur, in percent of the screen width, 0:off */
 	UPROPERTY(interp, BlueprintReadWrite, Category="Rendering Features|Motion Blur", meta=(ClampMin = "0.0", ClampMax = "100.0", editcondition = "bOverride_MotionBlurMax", DisplayName = "Max"))
 	float MotionBlurMax;
+	/**
+	 * Defines the target FPS for motion blur. Makes motion blur independent of actual frame rate and relative
+	 * to the specified target FPS instead. Higher target FPS results in shorter frames, which means shorter
+	 * shutter times and less motion blur. Lower FPS means more motion blur. A value of zero makes the motion
+	 * blur dependent on the actual frame rate.
+	 */
+	UPROPERTY(EditAnywhere, BlueprintReadWrite, Category = "Rendering Features|Motion Blur", meta=(ClampMin = "0", ClampMax = "120", editcondition = "bOverride_MotionBlurTargetFPS", DisplayName = "Target FPS"))
+	int32 MotionBlurTargetFPS;
+
 	/** The minimum projected screen radius for a primitive to be drawn in the velocity pass, percentage of screen width. smaller numbers cause more draw calls, default: 4% */
 	UPROPERTY(interp, BlueprintReadWrite, Category="Rendering Features|Motion Blur", meta=(ClampMin = "0.0", UIMax = "100.0", editcondition = "bOverride_MotionBlurPerObjectSize", DisplayName = "Per Object Size"))
 	float MotionBlurPerObjectSize;
diff --git a/Engine/Source/Runtime/Engine/Private/Scene.cpp b/Engine/Source/Runtime/Engine/Private/Scene.cpp
--- a/Engine/Source/Runtime/Engine/Private/Scene.cpp
+++ b/Engine/Source/Runtime/Engine/Private/Scene.cpp
@@ -542,6 +542,7 @@ FPostProcessSettings::FPostProcessSettings()
 	// next value might get overwritten by r.DefaultFeature.MotionBlur
 	MotionBlurAmount = 0.5f;
 	MotionBlurMax = 5.0f;
+	MotionBlurTargetFPS = 30;
 	MotionBlurPerObjectSize = 0.5f;
 	ScreenPercentage = 100.0f;
 	ScreenSpaceReflectionIntensity = 100.0f;
@@ -866,6 +867,7 @@ FPostProcessSettings::FPostProcessSettings(const FPostProcessSettings& Settings)
 	PRAGMA_ENABLE_DEPRECATION_WARNINGS
 	, MotionBlurAmount(Settings.MotionBlurAmount)
 	, MotionBlurMax(Settings.MotionBlurMax)
+	, MotionBlurTargetFPS(Settings.MotionBlurTargetFPS)
 	, MotionBlurPerObjectSize(Settings.MotionBlurPerObjectSize)
 	, LPVIntensity(Settings.LPVIntensity)
 	, LPVVplInjectionBias(Settings.LPVVplInjectionBias)
diff --git a/Engine/Source/Runtime/Engine/Private/SceneView.cpp b/Engine/Source/Runtime/Engine/Private/SceneView.cpp
--- a/Engine/Source/Runtime/Engine/Private/SceneView.cpp
+++ b/Engine/Source/Runtime/Engine/Private/SceneView.cpp
@@ -211,6 +211,15 @@ static TAutoConsoleVariable<float> CVarMotionBlurMax(
 	TEXT("-1: override (default)"),
 	ECVF_Scalability | ECVF_RenderThreadSafe);
 
+static TAutoConsoleVariable<int32> CVarMotionBlurTargetFPS(
+	TEXT("r.MotionBlur.TargetFPS"),
+	-1,
+	TEXT("Allows to override the postprocess setting (target FPS for motion blur velocity length scaling).\n")
+	TEXT("-1: override (default)")
+	TEXT(" 0: target current frame rate with moving average\n")
+	TEXT("[1,120]: target FPS for motion blur velocity scaling"),
+	ECVF_Scalability | ECVF_RenderThreadSafe);
+
 static TAutoConsoleVariable<float> CVarSceneColorFringeMax(
 	TEXT("r.SceneColorFringe.Max"),
 	-1.0f,
@@ -1449,6 +1458,11 @@ void FSceneView::OverridePostProcessSettings(const FPostProcessSettings& Src, fl
 		{
 			Dest.AmbientOcclusionRadiusInWS = Src.AmbientOcclusionRadiusInWS;
 		}
+
+		if (Src.bOverride_MotionBlurTargetFPS)
+		{
+			Dest.MotionBlurTargetFPS = Src.MotionBlurTargetFPS;
+		}
 	}
 
 	// will be deprecated soon, use the new asset LightPropagationVolumeBlendable instead
@@ -1877,6 +1891,15 @@ void FSceneView::EndFinalPostprocessSettings(const FSceneViewInitOptions& ViewIn
 		}
 	}
 
+	{
+		int32 TargetFPS = CVarMotionBlurTargetFPS.GetValueOnGameThread();
+
+		if (TargetFPS >= 0)
+		{
+			FinalPostProcessSettings.MotionBlurTargetFPS = TargetFPS;
+		}
+	}
+
 	{
 		float Value = CVarSceneColorFringeMax.GetValueOnGameThread();
 
diff --git a/Engine/Source/Runtime/Renderer/Private/RendererScene.cpp b/Engine/Source/Runtime/Renderer/Private/RendererScene.cpp
--- a/Engine/Source/Runtime/Renderer/Private/RendererScene.cpp
+++ b/Engine/Source/Runtime/Renderer/Private/RendererScene.cpp
@@ -107,6 +107,7 @@ FSceneViewState::FSceneViewState()
 	LastRenderTime = -FLT_MAX;
 	LastRenderTimeDelta = 0.0f;
 	MotionBlurTimeScale = 1.0f;
+	MotionBlurTargetDeltaTime = 1.0f / 60.0f; // Start with a reasonable default of 60hz.
 	PrevViewMatrixForOcclusionQuery.SetIdentity();
 	PrevViewOriginForOcclusionQuery = FVector::ZeroVector;
 #if !(UE_BUILD_SHIPPING || UE_BUILD_TEST)
diff --git a/Engine/Source/Runtime/Renderer/Private/ScenePrivate.h b/Engine/Source/Runtime/Renderer/Private/ScenePrivate.h
--- a/Engine/Source/Runtime/Renderer/Private/ScenePrivate.h
+++ b/Engine/Source/Runtime/Renderer/Private/ScenePrivate.h
@@ -694,6 +694,7 @@ public:
 	float		LastRenderTime;
 	float		LastRenderTimeDelta;
 	float		MotionBlurTimeScale;
+	float		MotionBlurTargetDeltaTime;
 	FMatrix		PrevViewMatrixForOcclusionQuery;
 	FVector		PrevViewOriginForOcclusionQuery;
 
@@ -1083,6 +1084,12 @@ public:
 
 	void TrimHistoryRenderTargets(const FScene* Scene);
 
+	/**
+	 * Calculates and stores the scale factor to apply to motion vectors based on the current game
+	 * time and view post process settings.
+	 */
+	void UpdateMotionBlurTimeScale(const FViewInfo& View);
+
 	/** 
 	 * Called every frame after UpdateLastRenderTime, sets up the information for the lagged temporal LOD transition
 	 */
diff --git a/Engine/Source/Runtime/Renderer/Private/SceneVisibility.cpp b/Engine/Source/Runtime/Renderer/Private/SceneVisibility.cpp
--- a/Engine/Source/Runtime/Renderer/Private/SceneVisibility.cpp
+++ b/Engine/Source/Runtime/Renderer/Private/SceneVisibility.cpp
@@ -3158,10 +3158,7 @@ void FSceneRenderer::PreVisibilityFrameSetup(FRHICommandListImmediate& RHICmdLis
 				// Clamp DeltaWorldTime to reasonable values for the purposes of motion blur, things like TimeDilation can make it very small
 				if (!ViewFamily.bWorldIsPaused)
 				{
-					const bool bEnableTimeScale = !ViewState->bSequencerIsPaused;
-					const float FixedBlurTimeScale = 2.0f;// 1 / (30 * 1 / 60)
-
-					ViewState->MotionBlurTimeScale = bEnableTimeScale ? (1.0f / (FMath::Max(View.Family->DeltaWorldTime, .00833f) * 30.0f)) : FixedBlurTimeScale;
+					ViewState->UpdateMotionBlurTimeScale(View);
 				}
 			}
 
@@ -3194,6 +3191,43 @@ void FSceneRenderer::PreVisibilityFrameSetup(FRHICommandListImmediate& RHICmdLis
 	}
 }
 
+void FSceneViewState::UpdateMotionBlurTimeScale(const FViewInfo& View)
+{
+	const int32 MotionBlurTargetFPS = View.FinalPostProcessSettings.MotionBlurTargetFPS;
+
+	// Frame rates over 120 FPS are clamped to avoid creating huge motion vectors.
+	float DeltaWorldTime = FMath::Max(View.Family->DeltaWorldTime, 1.0f / 120.0f);
+
+	// Track the current FPS by using an exponential moving average of the current delta time.
+	if (MotionBlurTargetFPS <= 0)
+	{
+		// Keep motion vector lengths stable for paused sequencer frames.
+		if (bSequencerIsPaused)
+		{
+			// Reset the moving average to the current delta time.
+			MotionBlurTargetDeltaTime = DeltaWorldTime;
+		}
+		else
+		{
+			// Smooth the target delta time using a moving average.
+			MotionBlurTargetDeltaTime = FMath::Lerp(MotionBlurTargetDeltaTime, DeltaWorldTime, 0.1f);
+		}
+	}
+	else // Track a fixed target FPS.
+	{
+		// Keep motion vector lengths stable for paused sequencer frames. Assumes a 60 FPS tick.
+		// Tuned for content compatibility with existing content when target is the default 30 FPS.
+		if (bSequencerIsPaused)
+		{
+			DeltaWorldTime = 1.0f / 60.0f;
+		}
+
+		MotionBlurTargetDeltaTime = 1.0f / static_cast<float>(MotionBlurTargetFPS);
+	}
+
+	MotionBlurTimeScale = MotionBlurTargetDeltaTime / DeltaWorldTime;
+}
+
 static TAutoConsoleVariable<int32> CVarAlsoUseSphereForFrustumCull(
 	TEXT("r.AlsoUseSphereForFrustumCull"),
 	0,  
