Fixed visualize HDR to take into account the new "Exposure Compensation Curve"

diff --git a/Engine/Source/Runtime/Renderer/Private/PostProcess/PostProcessVisualizeHDR.cpp b/Engine/Source/Runtime/Renderer/Private/PostProcess/PostProcessVisualizeHDR.cpp
--- a/Engine/Source/Runtime/Renderer/Private/PostProcess/PostProcessVisualizeHDR.cpp
+++ b/Engine/Source/Runtime/Renderer/Private/PostProcess/PostProcessVisualizeHDR.cpp
@@ -17,6 +17,7 @@
 #include "PostProcess/PostProcessEyeAdaptation.h"
 #include "PostProcess/PostProcessTonemap.h"
 #include "PipelineStateCache.h"
+#include "Curves/CurveFloat.h"
 
 /** Encapsulates the post processing eye adaptation pixel shader. */
 class FPostProcessVisualizeHDRPS : public FGlobalShader
@@ -324,7 +325,17 @@ void FRCPassPostProcessVisualizeHDR::Process(FRenderingCompositePassContext& Con
 	Canvas.DrawShadowedString(X, Y += YStep, TEXT("Speed Up/Down:"), GetStatsFont(), FLinearColor(1, 1, 1));
 	Canvas.DrawShadowedString(X + ColumnWidth, Y, *Line, GetStatsFont(), FLinearColor(1, 1, 1));
 
-	Line = FString::Printf(TEXT("%.2g"), View.FinalPostProcessSettings.AutoExposureBias);
+	float AutoExposureBias = View.FinalPostProcessSettings.AutoExposureBias;
+	if (View.FinalPostProcessSettings.AutoExposureBiasCurve)
+	{
+		float AverageSceneLuminance = View.GetLastAverageSceneLuminance();
+		if (AverageSceneLuminance > 0)
+		{
+			AutoExposureBias += View.FinalPostProcessSettings.AutoExposureBiasCurve->GetFloatValue(LuminanceToEV100(AverageSceneLuminance));
+		}
+	}
+
+	Line = FString::Printf(TEXT("%.2g"), AutoExposureBias);
 	Canvas.DrawShadowedString(X, Y += YStep, TEXT("Exposure Compensation: "), GetStatsFont(), FLinearColor(1, 1, 1));
 	Canvas.DrawShadowedString(X + ColumnWidth, Y, *Line, GetStatsFont(), FLinearColor(1, 0.3f, 0.3f));
 
