Fix shadow variable warnings when running with StressTestUnity.

diff --git a/Engine/Source/Runtime/Renderer/Private/LightMapDensityRendering.cpp b/Engine/Source/Runtime/Renderer/Private/LightMapDensityRendering.cpp
--- a/Engine/Source/Runtime/Renderer/Private/LightMapDensityRendering.cpp
+++ b/Engine/Source/Runtime/Renderer/Private/LightMapDensityRendering.cpp
@@ -242,8 +242,8 @@ void FLightmapDensityMeshProcessor::AddMeshBatch(const FMeshBatch& RESTRICT Mesh
 		// Force simple lightmaps based on system settings.
 		bool bAllowHighQualityLightMaps = AllowHighQualityLightmaps(FeatureLevel) && LightMapInteraction.AllowsHighQualityLightmaps();
 
-		static const auto CVarSupportLowQualityLightmaps = IConsoleManager::Get().FindTConsoleVariableDataInt(TEXT("r.SupportLowQualityLightmaps"));
-		const bool bAllowLowQualityLightMaps = (!CVarSupportLowQualityLightmaps) || (CVarSupportLowQualityLightmaps->GetValueOnAnyThread() != 0);
+		static const auto SupportLowQualityLightmapsVar = IConsoleManager::Get().FindTConsoleVariableDataInt(TEXT("r.SupportLowQualityLightmaps"));
+		const bool bAllowLowQualityLightMaps = (!SupportLowQualityLightmapsVar) || (SupportLowQualityLightmapsVar->GetValueOnAnyThread() != 0);
 
 		if (!bTranslucentBlendMode || ViewIfDynamicMeshCommand->Family->EngineShowFlags.Wireframe)
 		{
diff --git a/Engine/Source/Runtime/Renderer/Private/LightShaftRendering.cpp b/Engine/Source/Runtime/Renderer/Private/LightShaftRendering.cpp
--- a/Engine/Source/Runtime/Renderer/Private/LightShaftRendering.cpp
+++ b/Engine/Source/Runtime/Renderer/Private/LightShaftRendering.cpp
@@ -31,8 +31,8 @@
 #include "PipelineStateCache.h"
 
 /** Tweaked values from UE3 implementation **/
-const float PointLightFadeDistanceIncrease = 200;
-const float PointLightRadiusFadeFactor = 5;
+extern const float PointLightFadeDistanceIncrease = 200;
+extern const float PointLightRadiusFadeFactor = 5;
 
 // 0 is off, any other value is on, later we can expose more quality settings e.g. sample count
 int32 GLightShafts = 1;
diff --git a/Engine/Source/Runtime/Renderer/Private/SceneCore.cpp b/Engine/Source/Runtime/Renderer/Private/SceneCore.cpp
--- a/Engine/Source/Runtime/Renderer/Private/SceneCore.cpp
+++ b/Engine/Source/Runtime/Renderer/Private/SceneCore.cpp
@@ -229,12 +229,12 @@ FLightPrimitiveInteraction::FLightPrimitiveInteraction(
 		if (PrimitiveSceneInfo->Scene->GetShadingPath() == EShadingPath::Mobile && LightSceneInfo->Proxy->IsMovable())
 		{
 			const uint8 LightType = LightSceneInfo->Proxy->GetLightType();
-			static const auto CVarMobileEnableMovableSpotLights = IConsoleManager::Get().FindTConsoleVariableDataInt(TEXT("r.Mobile.EnableMovableSpotLights"));
+			static const auto MobileEnableMovableSpotLightsVar = IConsoleManager::Get().FindTConsoleVariableDataInt(TEXT("r.Mobile.EnableMovableSpotLights"));
 
 			const bool bIsValidLightType = 
 				   LightType == LightType_Rect
 				|| LightType == LightType_Point
-				|| (LightType == LightType_Spot && CVarMobileEnableMovableSpotLights->GetValueOnRenderThread());
+				|| (LightType == LightType_Spot && MobileEnableMovableSpotLightsVar->GetValueOnRenderThread());
 
 			if( bIsValidLightType )
 			{
diff --git a/Engine/Source/Runtime/Renderer/Private/SceneRendering.cpp b/Engine/Source/Runtime/Renderer/Private/SceneRendering.cpp
--- a/Engine/Source/Runtime/Renderer/Private/SceneRendering.cpp
+++ b/Engine/Source/Runtime/Renderer/Private/SceneRendering.cpp
@@ -3233,11 +3233,11 @@ FRendererModule::FRendererModule()
 	CVarSimpleForwardShading_PreviousValue = CVarSimpleForwardShading.AsVariable()->GetInt();
 	CVarSimpleForwardShading.AsVariable()->SetOnChangedCallback(FConsoleVariableDelegate::CreateStatic(&OnChangeSimpleForwardShading));
 
-	static auto CVarEarlyZPass = IConsoleManager::Get().FindConsoleVariable(TEXT("r.EarlyZPass"));
-	CVarEarlyZPass->SetOnChangedCallback(FConsoleVariableDelegate::CreateStatic(&OnChangeCVarRequiringRecreateRenderState));
+	static auto EarlyZPassVar = IConsoleManager::Get().FindConsoleVariable(TEXT("r.EarlyZPass"));
+	EarlyZPassVar->SetOnChangedCallback(FConsoleVariableDelegate::CreateStatic(&OnChangeCVarRequiringRecreateRenderState));
 
-	static auto CVarEarlyZPassMovable = IConsoleManager::Get().FindConsoleVariable(TEXT("r.EarlyZPassMovable"));
-	CVarEarlyZPassMovable->SetOnChangedCallback(FConsoleVariableDelegate::CreateStatic(&OnChangeCVarRequiringRecreateRenderState));
+	static auto EarlyZPassMovableVar = IConsoleManager::Get().FindConsoleVariable(TEXT("r.EarlyZPassMovable"));
+	EarlyZPassMovableVar->SetOnChangedCallback(FConsoleVariableDelegate::CreateStatic(&OnChangeCVarRequiringRecreateRenderState));
 
 #if !(UE_BUILD_SHIPPING || UE_BUILD_TEST)
 	void InitDebugViewModeInterfaces();
diff --git a/Engine/Source/Runtime/Renderer/Private/SceneVisibility.cpp b/Engine/Source/Runtime/Renderer/Private/SceneVisibility.cpp
--- a/Engine/Source/Runtime/Renderer/Private/SceneVisibility.cpp
+++ b/Engine/Source/Runtime/Renderer/Private/SceneVisibility.cpp
@@ -3796,16 +3796,16 @@ void FSceneRenderer::PostVisibilityFrameSetup(FILCUpdatePrimTaskData& OutILCTask
 			{
 				VisibleLightViewInfo.bInViewFrustum = true;
 
-				static const auto CVarMobileMSAA = IConsoleManager::Get().FindTConsoleVariableDataInt(TEXT("r.MobileMSAA"));
-				bool bNotMobileMSAA = !(CVarMobileMSAA ? CVarMobileMSAA->GetValueOnRenderThread() > 1 : false);
+				static const auto MobileMSAAVar = IConsoleManager::Get().FindTConsoleVariableDataInt(TEXT("r.MobileMSAA"));
+				bool bNotMobileMSAA = !(MobileMSAAVar ? MobileMSAAVar->GetValueOnRenderThread() > 1 : false);
 
 				// Setup single sun-shaft from direction lights for mobile.
 				if(bCheckLightShafts && LightSceneInfo->bEnableLightShaftBloom)
 				{
 					// Find directional light for sun shafts.
 					// Tweaked values from UE3 implementation.
-					const float PointLightFadeDistanceIncrease = 200.0f;
-					const float PointLightRadiusFadeFactor = 5.0f;
+					extern const float PointLightFadeDistanceIncrease;
+					extern const float PointLightRadiusFadeFactor;
 
 					const FVector WorldSpaceBlurOrigin = LightSceneInfo->Proxy->GetPosition();
 					// Transform into post projection space
diff --git a/Engine/Source/Runtime/Renderer/Private/ShadowDepthRendering.cpp b/Engine/Source/Runtime/Renderer/Private/ShadowDepthRendering.cpp
--- a/Engine/Source/Runtime/Renderer/Private/ShadowDepthRendering.cpp
+++ b/Engine/Source/Runtime/Renderer/Private/ShadowDepthRendering.cpp
@@ -203,8 +203,8 @@ public:
 
 	static bool ShouldCompilePermutation(EShaderPlatform Platform, const FMaterial* Material, const FVertexFactoryType* VertexFactoryType)
 	{
-		static const auto CVarSupportAllShaderPermutations = IConsoleManager::Get().FindTConsoleVariableDataInt(TEXT("r.SupportAllShaderPermutations"));
-		const bool bForceAllPermutations = CVarSupportAllShaderPermutations && CVarSupportAllShaderPermutations->GetValueOnAnyThread() != 0;
+		static const auto SupportAllShaderPermutationsVar = IConsoleManager::Get().FindTConsoleVariableDataInt(TEXT("r.SupportAllShaderPermutations"));
+		const bool bForceAllPermutations = SupportAllShaderPermutationsVar && SupportAllShaderPermutationsVar->GetValueOnAnyThread() != 0;
 		const bool bSupportPointLightWholeSceneShadows = CVarSupportPointLightWholeSceneShadows.GetValueOnAnyThread() != 0 || bForceAllPermutations;
 		const bool bRHISupportsShadowCastingPointLights = RHISupportsGeometryShaders(Platform) || RHISupportsVertexShaderLayer(Platform);
 
