Moving CL 5312080 from UE4-Release-4.22:

Fixed distance field interaction with far shadow cascades. Distance field shadows need to be renderer last as they blend over far shadow cascades.

diff --git a/Engine/Source/Runtime/Renderer/Private/ShadowRendering.cpp b/Engine/Source/Runtime/Renderer/Private/ShadowRendering.cpp
--- a/Engine/Source/Runtime/Renderer/Private/ShadowRendering.cpp
+++ b/Engine/Source/Runtime/Renderer/Private/ShadowRendering.cpp
@@ -1333,40 +1333,6 @@ bool FSceneRenderer::RenderShadowProjections(FRHICommandListImmediate& RHICmdLis
 		}
 	}
 
-	if (DistanceFieldShadows.Num() > 0)
-	{
-		// Dispatch distance field shadows first
-		for (int32 ViewIndex = 0; ViewIndex < Views.Num(); ViewIndex++)
-		{
-			SCOPED_CONDITIONAL_DRAW_EVENTF(RHICmdList, EventView, Views.Num() > 1, TEXT("DistanceFieldShadows_View%d"), ViewIndex);
-
-			const FViewInfo& View = Views[ViewIndex];
-
-			// Set the device viewport for the view.
-			RHICmdList.SetViewport(View.ViewRect.Min.X, View.ViewRect.Min.Y, 0.0f, View.ViewRect.Max.X, View.ViewRect.Max.Y, 1.0f);
-
-			// Set the light's scissor rectangle.
-			FIntRect ScissorRect;
-			if (!LightSceneInfo->Proxy->SetScissorRect(RHICmdList, View, View.ViewRect, &ScissorRect))
-			{
-				ScissorRect = View.ViewRect;
-			}
-
-			if (ScissorRect.Area() > 0)
-			{
-				// Project the shadow depth buffers onto the scene.
-				for (int32 ShadowIndex = 0; ShadowIndex < DistanceFieldShadows.Num(); ShadowIndex++)
-				{
-					FProjectedShadowInfo* ProjectedShadowInfo = DistanceFieldShadows[ShadowIndex];
-					ProjectedShadowInfo->RenderRayTracedDistanceFieldProjection(RHICmdList, View, ScissorRect, ScreenShadowMaskTexture, bProjectingForForwardShading);
-				}
-			}
-
-			// Reset the scissor rectangle.
-			RHICmdList.SetScissorRect(false, 0, 0, 0, 0);
-		}
-	}
-
 	if (NormalShadows.Num() > 0)
 	{
 		// Render normal shadows
@@ -1439,6 +1405,40 @@ bool FSceneRenderer::RenderShadowProjections(FRHICommandListImmediate& RHICmdLis
 		}
 	}
 
+	if (DistanceFieldShadows.Num() > 0)
+	{
+		// Distance field shadows need to be renderer last as they blend over far shadow cascades.
+		for (int32 ViewIndex = 0; ViewIndex < Views.Num(); ViewIndex++)
+		{
+			SCOPED_CONDITIONAL_DRAW_EVENTF(RHICmdList, EventView, Views.Num() > 1, TEXT("DistanceFieldShadows_View%d"), ViewIndex);
+
+			const FViewInfo& View = Views[ViewIndex];
+
+			// Set the device viewport for the view.
+			RHICmdList.SetViewport(View.ViewRect.Min.X, View.ViewRect.Min.Y, 0.0f, View.ViewRect.Max.X, View.ViewRect.Max.Y, 1.0f);
+
+			// Set the light's scissor rectangle.
+			FIntRect ScissorRect;
+			if (!LightSceneInfo->Proxy->SetScissorRect(RHICmdList, View, View.ViewRect, &ScissorRect))
+			{
+				ScissorRect = View.ViewRect;
+			}
+
+			if (ScissorRect.Area() > 0)
+			{
+				// Project the shadow depth buffers onto the scene.
+				for (int32 ShadowIndex = 0; ShadowIndex < DistanceFieldShadows.Num(); ShadowIndex++)
+				{
+					FProjectedShadowInfo* ProjectedShadowInfo = DistanceFieldShadows[ShadowIndex];
+					ProjectedShadowInfo->RenderRayTracedDistanceFieldProjection(RHICmdList, View, ScissorRect, ScreenShadowMaskTexture, bProjectingForForwardShading);
+				}
+			}
+
+			// Reset the scissor rectangle.
+			RHICmdList.SetScissorRect(false, 0, 0, 0, 0);
+		}
+	}
+
 	return true;
 }
 	
