Copying from Main to fix integration history.

#rnx

diff --git a/Engine/Shaders/Private/DataStructures/SceneViewFamilyBlackboard.ush b/Engine/Shaders/Private/DataStructures/SceneViewFamilyBlackboard.ush
deleted file mode 100644
diff --git a/Engine/Shaders/Private/RayTracing/RayTracingRectLightRGS.usf b/Engine/Shaders/Private/RayTracing/RayTracingRectLightRGS.usf
new file mode 100644
--- /dev/null
+++ b/Engine/Shaders/Private/RayTracing/RayTracingRectLightRGS.usf
@@ -0,0 +1,345 @@
+// Copyright 1998-2019 Epic Games, Inc. All Rights Reserved.
+
+#include "../Common.ush"
+#include "../DeferredShadingCommon.ush"
+#include "../PathTracing/Utilities/PathTracingRandomSequence.ush"
+#include "../RectLight.ush"
+#include "../ShadingModels.ush"
+#include "RayTracingCommon.ush"
+#include "RayTracingRectLight.ush"
+#include "MipTreeCommon.ush"
+
+RaytracingAccelerationStructure TLAS;
+RWTexture2D<float4> RWLuminanceUAV;
+RWTexture2D<float> RWRayDistanceUAV;
+
+uint CalcLinearIndex(uint2 PixelCoord)
+{
+	return PixelCoord.y * View.BufferSizeAndInvSize.x + PixelCoord.x;
+}
+
+// #patrick todo: move to RayTracingCommon.ush
+float3 ReconstructWorldPositionFromDepth(float2 UV, float Depth)
+{
+	float2 ScreenPosition = (UV - View.ScreenPositionScaleBias.wz) / View.ScreenPositionScaleBias.xy;
+	float4 HomogeneousWorldPosition = mul(float4(ScreenPosition * Depth, Depth, 1), View.ScreenToWorld);
+	float3 WorldPosition = HomogeneousWorldPosition.xyz / HomogeneousWorldPosition.w;
+
+	return WorldPosition;
+}
+
+bool GenerateRectLightImportanceSampledRay(
+	float3 WorldPosition,
+	float3 WorldNormal,
+	float3 LightPosition,
+	float3 LightNormal,
+	float3 LightdPdu,
+	float3 LightdPdv,
+	float LightWidth,
+	float LightHeight,
+	inout RandomSequence RandSequence,
+	out float3 RayOrigin,
+	out float3 RayDirection,
+	out float RayTMin,
+	out float RayTMax,
+	out float RayPdf
+)
+{
+	uint MipCount = log2(RectLight.MipTreeDimensions.x);
+	uint2 MipPixel = 0;
+	float MipPdf = 1.0;
+
+	uint StopLevel = 0;
+	for (uint Index = 0; Index < MipCount - StopLevel; ++Index)
+	{
+		uint MipIndex = MipCount - Index - 1;
+		uint BufferOffset = BufferOffsetAtPixel(MipPixel, MipIndex, RectLight.MipTreeDimensions.xy);
+		uint4 BufferOffsets = BufferOffset + uint4(0, 1, 2, 3);
+		float4 PowerTerm = float4(RectLight.MipTree[BufferOffsets.x], RectLight.MipTree[BufferOffsets.y], RectLight.MipTree[BufferOffsets.z], RectLight.MipTree[BufferOffsets.w]);
+		PowerTerm = max(PowerTerm, 0.0);
+
+		// Determine sample positions
+		uint DummyVariable;
+		float2 Jitter = RandomSequence_GenerateSample2D(RandSequence, DummyVariable);
+		float2 DimInv = 1.0 / DimensionsAtLevel(RectLight.MipTreeDimensions.xy, MipIndex);
+		float2 UV00 = (MipPixel + Jitter) * DimInv - 0.5;
+		float2 UV10 = (MipPixel + Jitter + uint2(1, 0)) * DimInv - 0.5;
+		float2 UV01 = (MipPixel + Jitter + uint2(0, 1)) * DimInv - 0.5;
+		float2 UV11 = (MipPixel + Jitter + uint2(1, 1)) * DimInv - 0.5;
+
+		float3 Position00 = LightPosition + LightdPdu * LightWidth * UV00.x + LightdPdv * LightHeight * UV00.y;
+		float3 Position10 = LightPosition + LightdPdu * LightWidth * UV10.x + LightdPdv * LightHeight * UV10.y;
+		float3 Position01 = LightPosition + LightdPdu * LightWidth * UV01.x + LightdPdv * LightHeight * UV01.y;
+		float3 Position11 = LightPosition + LightdPdu * LightWidth * UV11.x + LightdPdv * LightHeight * UV11.y;
+
+		// Evaluate irradiance coefficients
+		float3 LightDirection00 = WorldPosition - Position00;
+		float3 LightDirection10 = WorldPosition - Position10;
+		float3 LightDirection01 = WorldPosition - Position01;
+		float3 LightDirection11 = WorldPosition - Position11;
+
+		float4 DistanceSquaredTerm = float4(
+			dot(LightDirection00, LightDirection00),
+			dot(LightDirection10, LightDirection10),
+			dot(LightDirection01, LightDirection01),
+			dot(LightDirection11, LightDirection11)
+		);
+
+		float4 CosineTerm = float4(
+			saturate(dot(normalize(LightDirection00), LightNormal)),
+			saturate(dot(normalize(LightDirection10), LightNormal)),
+			saturate(dot(normalize(LightDirection01), LightNormal)),
+			saturate(dot(normalize(LightDirection11), LightNormal))
+		);
+
+		// Build CDF on approximate irradiance
+		float4 Cdf = PowerTerm * CosineTerm / DistanceSquaredTerm;
+		Cdf.y += Cdf.x;
+		Cdf.z += Cdf.y;
+		Cdf.w += Cdf.z;
+		Cdf /= Cdf.w;
+
+		// Sample CDF
+		float RandSample = RandomSequence_GenerateSample1D(RandSequence, DummyVariable);
+		if (RandSample < Cdf.x)
+		{
+			MipPdf *= Cdf.x;
+		}
+		else if (RandSample < Cdf.y)
+		{
+			MipPdf *= Cdf.y - Cdf.x;
+			MipPixel.x += 1;
+		}
+		else if (RandSample < Cdf.z)
+		{
+			MipPdf *= Cdf.z - Cdf.y;
+			MipPixel.y += 1;
+		}
+		else // if (RandSample < Cdf.w)
+		{
+			MipPdf *= Cdf.w - Cdf.z;
+			MipPixel += 1;
+		}
+
+		MipPixel = uint2(MipPixel.x << 1, MipPixel.y << 1);
+	}
+	MipPixel = uint2(MipPixel.x >> 1, MipPixel.y >> 1);
+
+	// Sample within stop-level texel
+	uint DummyVariable;
+	float2 PixelOffset = RandomSequence_GenerateSample2D(RandSequence, DummyVariable);
+	float2 DimensionsAtStopLevel = DimensionsAtLevel(RectLight.MipTreeDimensions.xy, StopLevel);
+	float2 MipPixelUV = (MipPixel + PixelOffset) / DimensionsAtStopLevel;
+
+	// Map sample point to quad
+	MipPixelUV -= 0.5;
+	float2 RectLightOffset = MipPixel / DimensionsAtStopLevel;
+	float3 LightOrigin = LightPosition + LightdPdu * LightWidth * RectLightOffset.x + LightdPdv * LightHeight * RectLightOffset.y;
+	LightPosition += LightdPdu * LightWidth * MipPixelUV.x + LightdPdv * LightHeight * MipPixelUV.y;
+	float3 LightDirection = normalize(LightPosition - WorldPosition);
+
+	// Light-normal culling
+	if (dot(-LightDirection, LightNormal) <= 0.0)
+	{
+		RayPdf = 0.0;
+		return false;
+	}
+
+	// #dxr_todo: Performing direct lighting here could save some recalculation
+
+	// Define ray
+	RayOrigin = WorldPosition;
+	RayDirection = LightDirection;
+	RayTMin = 0.0;
+	RayTMax = length(LightPosition - WorldPosition);
+
+	// Construct solid angle pdf
+	FRect Rect;
+	Rect.Origin = LightOrigin - WorldPosition;
+	Rect.Axis[0] = LightdPdu;
+	Rect.Axis[1] = LightdPdv;
+	Rect.Axis[2] = cross(LightdPdu, LightdPdv);
+	Rect.Extent = float2(LightWidth, LightHeight) / DimensionsAtStopLevel;
+	FSphericalRect SphericalRect = BuildSphericalRect(Rect);
+	float SolidAnglePdf = 1.0 / SphericalRect.SolidAngle;
+
+	RayPdf = MipPdf * SolidAnglePdf;
+	return RayTMax > RayTMin;
+}
+
+[shader("raygeneration")]
+void RectLightRGS()
+{
+	uint2 PixelCoord = DispatchRaysIndex().xy;
+
+	RandomSequence RandSequence;
+	uint LinearIndex = CalcLinearIndex(PixelCoord);
+	RandomSequence_Initialize(RandSequence, LinearIndex, View.FrameNumber);
+
+	// Get G-Buffer surface data
+	float2 InvBufferSize = View.BufferSizeAndInvSize.zw;
+	float2 UV = (float2(PixelCoord) + 0.5) * InvBufferSize;
+	FScreenSpaceData ScreenSpaceData = GetScreenSpaceData(UV);
+	FGBufferData GBuffer = ScreenSpaceData.GBuffer;
+	float Depth = ScreenSpaceData.GBuffer.Depth;
+	float3 WorldPosition = ReconstructWorldPositionFromDepth(UV, Depth);
+	float3 CameraOrigin = ReconstructWorldPositionFromDepth(UV, 0.0f);
+	float3 CameraDirection = normalize(WorldPosition - CameraOrigin);
+	float3 WorldNormal = ScreenSpaceData.GBuffer.WorldNormal;
+	float3 BaseColor = ScreenSpaceData.GBuffer.BaseColor;
+	float3 DiffuseColor = ScreenSpaceData.GBuffer.DiffuseColor;
+	float3 SpecularColor = ScreenSpaceData.GBuffer.SpecularColor;
+	float Roughness = ScreenSpaceData.GBuffer.Roughness;
+
+	float RayDistance = 0.0;
+	float HitCount = 0.0;
+	uint SamplesPerPixel = RectLight.SamplesPerPixel;
+
+	// Mask out depth values that are infinitely far away
+	float DeviceZ = SceneTexturesStruct.SceneDepthTexture.Load(int3(PixelCoord, 0)).r;
+	bool IsFiniteDepth = DeviceZ > 0.0;
+	bool bTraceRay = (
+		IsFiniteDepth &&
+		GBuffer.ShadingModelID != SHADINGMODELID_UNLIT);
+	if (!bTraceRay)
+	{
+		SamplesPerPixel = 0.0;
+	}
+
+	float3 LightPosition = RectLight.Position;
+	float3 LightNormal = RectLight.Normal;
+	float3 LightdPdu = RectLight.dPdu;
+	float3 LightdPdv = RectLight.dPdv;
+	float LightWidth = RectLight.Width;
+	float LightHeight = RectLight.Height;
+	float LightBarnCosAngle = RectLight.BarnCosAngle;
+	float LightBarnLength = RectLight.BarnLength;
+
+	// Compute the visible portion of the light from the current shading point
+	{
+		float3 ToLight = LightPosition - WorldPosition;
+		FRect VisibleRect = GetRect(
+				ToLight,
+				-LightNormal,
+				LightdPdv,
+				LightWidth*0.5f,
+				LightHeight*0.5f,
+				RectLight.BarnCosAngle,
+				RectLight.BarnLength,
+				true);
+
+		LightPosition = LightPosition - LightdPdu * VisibleRect.Offset.x - LightdPdv * VisibleRect.Offset.y;	
+		LightWidth	  = VisibleRect.Extent.x * 2.0f;
+		LightHeight   = VisibleRect.Extent.y * 2.0f;
+	}
+
+	// Encode the rect light  parameter into the light parameters struct
+	// todo: send direction the light shader parameter data instead of the custon RectLight data.
+	FLightShaderParameters LightParameters;
+	LightParameters.Position = LightPosition;
+	LightParameters.Direction = -LightNormal;
+	LightParameters.Tangent = LightdPdu;
+	LightParameters.SourceRadius = 0.5f * LightWidth;
+	LightParameters.SourceLength = 0.5f * LightHeight;
+	LightParameters.RectLightBarnCosAngle =  LightBarnCosAngle;
+	LightParameters.RectLightBarnLength =  LightBarnLength;
+
+	if (LightWidth == 0 || LightHeight == 0) 
+		SamplesPerPixel = 0;
+
+	float3 ExitantRadiance = 0.0;
+	float3 DiffuseExitantRadiance = 0.0;
+	for (uint SampleIndex = 0; SampleIndex < SamplesPerPixel; ++SampleIndex)
+	{
+		RayDesc Ray;
+
+		uint DummyVariable;
+		float2 RandSample = RandomSequence_GenerateSample2D(RandSequence, DummyVariable);
+		float RayPdf = 0.0;
+
+		bool bIsValid = false;
+#if TEXTURE_IMPORTANCE_SAMPLING
+		if (RectLight.bIsTextureImportanceSampling)
+		{
+			bIsValid = GenerateRectLightImportanceSampledRay(WorldPosition, WorldNormal, LightPosition, LightNormal, LightdPdu, LightdPdv, LightWidth, LightHeight, RandSequence, Ray.Origin, Ray.Direction, Ray.TMin, Ray.TMax, RayPdf);
+		}
+		else
+#endif
+		{
+			bIsValid = GenerateRectLightOcclusionRay(LightParameters, WorldPosition, WorldNormal, RandSample, Ray.Origin, Ray.Direction, Ray.TMin, Ray.TMax, RayPdf);
+		}
+
+		if (!bIsValid)
+		{
+			continue;
+		}
+		ApplyPositionBias(Ray, WorldNormal, RectLight.MaxNormalBias);
+
+		float NoL = dot(WorldNormal, Ray.Direction);
+		if (NoL <= 0.0)
+		{
+			continue;
+		}
+
+		uint RayFlags = 0
+			| RAY_FLAG_CULL_BACK_FACING_TRIANGLES
+			| RAY_FLAG_FORCE_OPAQUE; // #dxr_todo: generate a special shader permutation for occlusion ray tracing that includes any-hit shaders for masked geometry
+
+		FDefaultPayload Payload = (FDefaultPayload)0;
+
+		TraceRay(
+			TLAS,      // AccelerationStructure
+			RayFlags,
+			RAY_TRACING_MASK_SHADOW,          // InstanceInclusionMask
+			RAY_TRACING_SHADER_SLOT_MATERIAL, // RayContributionToHitGroupIndex
+			RAY_TRACING_NUM_SHADER_SLOTS,     // MultiplierForGeometryContributionToShaderIndex
+			0,         // MissShaderIndex
+			Ray,       // RayDesc
+			Payload    // Payload
+		);
+
+		if (Payload.HitT >= 0.0)
+		{
+			RayDistance += Payload.HitT;
+			HitCount += 1.0;
+		}
+		else
+		{
+			// Sample light texture
+			float3 SampleLightPosition = Ray.Origin + Ray.Direction * Ray.TMax;
+			float3 RectLightCorner = LightPosition - 0.5 * (LightdPdu * LightWidth + LightdPdv * LightHeight);
+			float3 DeltaPosition = SampleLightPosition - RectLightCorner;
+			float2 SampleUV = float2(dot(DeltaPosition, LightdPdu) / LightWidth, dot(DeltaPosition, LightdPdv) / LightHeight);
+			SampleUV = 1.0 - SampleUV;
+			float3 TextureColor = Texture2DSample(RectLight.Texture, RectLight.TextureSampler, SampleUV).xyz;
+			float3 IncomingRadiance = RectLight.Color * TextureColor;
+
+			// Evaluate material
+			half3 N = WorldNormal;
+			half3 V = -CameraDirection;
+			half3 L = Ray.Direction;
+			FShadowTerms ShadowTerms = { 0.0, 0.0, 0.0 };
+			FDirectLighting LightingSample = EvaluateBxDF(ScreenSpaceData.GBuffer, N, V, L, NoL, ShadowTerms);
+			float3 Brdf = LightingSample.Diffuse + LightingSample.Transmission + LightingSample.Specular;
+
+			// #dxr_todo: Determine best way to incorporate integrals:
+			// 1) Return combined contribution (current method)
+			// 2) Split the integral (defer specular contribution to another pass)
+			// 3) Store diffuse and specular contribution separately
+			float3 Irradiance = IncomingRadiance * NoL / RayPdf;
+			ExitantRadiance += Irradiance * Brdf;
+			DiffuseExitantRadiance += Irradiance * (LightingSample.Diffuse + LightingSample.Transmission);
+		}
+	}
+
+	float3 OutputRadiance = (SamplesPerPixel > 0) ? ExitantRadiance / SamplesPerPixel : ExitantRadiance;
+
+#if USE_PREEXPOSURE
+	OutputRadiance *= View.PreExposure;
+#endif
+
+	OutputRadiance = ClampToHalfFloatRange(OutputRadiance);
+	RWLuminanceUAV[PixelCoord].rgb = OutputRadiance;
+	RWLuminanceUAV[PixelCoord].a = 1.0;
+	RWRayDistanceUAV[PixelCoord] = (HitCount > 0.0) ? RayDistance / HitCount : 1.0e27;
+}
diff --git a/Engine/Shaders/Private/SceneViewFamilyBlackboard.ush b/Engine/Shaders/Private/SceneViewFamilyBlackboard.ush
new file mode 100644
--- /dev/null
+++ b/Engine/Shaders/Private/SceneViewFamilyBlackboard.ush
@@ -0,0 +1,53 @@
+// Copyright 1998-2019 Epic Games, Inc. All Rights Reserved.
+
+/*=============================================================================
+	SceneViewFamilyBlackboard.ush: 
+=============================================================================*/
+
+#pragma once
+
+#include "SceneTexturesCommon.ush"
+#include "DeferredShadingCommon.ush"
+
+uint bIsSceneLightingChannelsValid;
+
+// Matches FSceneViewFamilyBlackboard
+Texture2D SceneDepthBuffer;
+Texture2D SceneVelocityBuffer;
+Texture2D SceneGBufferA;
+Texture2D SceneGBufferB;
+Texture2D SceneGBufferC;
+Texture2D SceneGBufferD;
+Texture2D SceneGBufferE;
+Texture2D<uint> SceneLightingChannels;
+
+
+// @param UV - UV space in the GBuffer textures (BufferSize resolution)
+FGBufferData GetGBufferDataFromSceneBlackboard(float2 UV, bool bGetNormalizedNormal = true)
+{
+	float4 GBufferA = SceneGBufferA.SampleLevel(GlobalPointClampedSampler, UV, 0);
+	float4 GBufferB = SceneGBufferB.SampleLevel(GlobalPointClampedSampler, UV, 0);
+	float4 GBufferC = SceneGBufferC.SampleLevel(GlobalPointClampedSampler, UV, 0);
+	float4 GBufferD = SceneGBufferD.SampleLevel(GlobalPointClampedSampler, UV, 0);
+	float4 GBufferE = SceneGBufferE.SampleLevel(GlobalPointClampedSampler, UV, 0);
+	float4 GBufferVelocity = SceneVelocityBuffer.SampleLevel(GlobalPointClampedSampler, UV, 0);
+
+	uint CustomStencil = 0;
+	float CustomNativeDepth = 0;
+
+	float DeviceZ = SceneDepthBuffer.SampleLevel(GlobalPointClampedSampler, UV, 0).r;
+	float SceneDepth = ConvertFromDeviceZ(DeviceZ);
+
+	return DecodeGBufferData(GBufferA, GBufferB, GBufferC, GBufferD, GBufferE, GBufferVelocity, CustomNativeDepth, CustomStencil, SceneDepth, bGetNormalizedNormal, CheckerFromSceneColorUV(UV));
+}
+
+/** Returns the light channel mask that should be executed for this pixel. */
+uint GetSceneLightingChannel(uint2 PixelCoord)
+{
+	BRANCH
+	if (bIsSceneLightingChannelsValid)
+	{
+		return SceneLightingChannels.Load(uint3(PixelCoord, 0)).x;
+	}
+	return ~0;
+}
diff --git a/Engine/Source/Runtime/Renderer/Private/RayTracing/RayTracingRectLight.cpp b/Engine/Source/Runtime/Renderer/Private/RayTracing/RayTracingRectLight.cpp
new file mode 100644
--- /dev/null
+++ b/Engine/Source/Runtime/Renderer/Private/RayTracing/RayTracingRectLight.cpp
@@ -0,0 +1,440 @@
+// Copyright 1998-2019 Epic Games, Inc. All Rights Reserved.
+
+#include "DeferredShadingRenderer.h"
+
+#if RHI_RAYTRACING
+
+#include "ClearQuad.h"
+#include "SceneRendering.h"
+#include "SceneRenderTargets.h"
+#include "RenderTargetPool.h"
+#include "RHIResources.h"
+#include "UniformBuffer.h"
+#include "VisualizeTexture.h"
+
+#include "PostProcess/PostProcessing.h"
+#include "PostProcess/SceneFilterRendering.h"
+#include "RectLightSceneProxy.h"
+#include "RaytracingOptions.h"
+
+#include "RHI/Public/PipelineStateCache.h"
+
+static int32 GRayTracingStochasticRectLight = 0;
+static FAutoConsoleVariableRef CVarRayTracingStochasticRectLight(
+	TEXT("r.RayTracing.StochasticRectLight"),
+	GRayTracingStochasticRectLight,
+	TEXT("0: use analytical evaluation (default)\n")
+	TEXT("1: use stochastic evaluation\n")
+);
+
+static int32 GRayTracingStochasticRectLightSamplesPerPixel = 1;
+static FAutoConsoleVariableRef CVarRayTracingRecLightStochasticSamplesPerPixel(
+	TEXT("r.RayTracing.StochasticRectLight.SamplesPerPixel"),
+	GRayTracingStochasticRectLightSamplesPerPixel,
+	TEXT("Sets the samples-per-pixel for rect light evaluation (default = 1)")
+);
+
+static int32 GRayTracingStochasticRectLightIsTextureImportanceSampling = 1;
+static FAutoConsoleVariableRef CVarRayTracingStochasticRecLightIsTextureImportanceSampling(
+	TEXT("r.RayTracing.StochasticRectLight.IsTextureImportanceSampling"),
+	GRayTracingStochasticRectLightIsTextureImportanceSampling,
+	TEXT("Enable importance sampling for rect light evaluation (default = 1)")
+);
+
+bool ShouldRenderRayTracingStochasticRectLight(const FLightSceneInfo& LightSceneInfo)
+{
+	return IsRayTracingEnabled() && GRayTracingStochasticRectLight == 1
+		&& LightSceneInfo.Proxy->CastsRaytracedShadow()
+		&& LightSceneInfo.Proxy->GetLightType() == LightType_Rect;
+}
+
+BEGIN_GLOBAL_SHADER_PARAMETER_STRUCT(FRectLightData, )
+// Pass settings
+SHADER_PARAMETER(int, SamplesPerPixel)
+SHADER_PARAMETER(int, bIsTextureImportanceSampling)
+// Light data
+SHADER_PARAMETER(FVector, Position)
+SHADER_PARAMETER(FVector, Normal)
+SHADER_PARAMETER(FVector, dPdu)
+SHADER_PARAMETER(FVector, dPdv)
+SHADER_PARAMETER(FVector, Color)
+SHADER_PARAMETER(float, Width)
+SHADER_PARAMETER(float, Height)
+SHADER_PARAMETER(FIntVector, MipTreeDimensions)
+SHADER_PARAMETER(float, MaxNormalBias)
+SHADER_PARAMETER(float, BarnCosAngle)
+SHADER_PARAMETER(float, BarnLength)
+SHADER_PARAMETER_TEXTURE(Texture2D, Texture)
+SHADER_PARAMETER_SAMPLER(SamplerState, TextureSampler)
+// Sampling data
+SHADER_PARAMETER_SRV(Buffer<float>, MipTree)
+END_GLOBAL_SHADER_PARAMETER_STRUCT()
+
+DECLARE_GPU_STAT_NAMED(RayTracingRectLight, TEXT("Ray Tracing RectLight"));
+
+IMPLEMENT_GLOBAL_SHADER_PARAMETER_STRUCT(FRectLightData, "RectLight");
+
+template <int TextureImportanceSampling>
+class FRectLightRGS : public FGlobalShader
+{
+	DECLARE_SHADER_TYPE(FRectLightRGS, Global)
+
+	static void ModifyCompilationEnvironment(const FGlobalShaderPermutationParameters& Parameters, FShaderCompilerEnvironment& OutEnvironment)
+	{
+		OutEnvironment.SetDefine(TEXT("TEXTURE_IMPORTANCE_SAMPLING"), TextureImportanceSampling);
+	}
+
+public:
+	static bool ShouldCompilePermutation(const FGlobalShaderPermutationParameters& Parameters)
+	{
+		return ShouldCompileRayTracingShadersForProject(Parameters.Platform);
+	}
+
+	FRectLightRGS() {}
+	virtual ~FRectLightRGS() {}
+
+	FRectLightRGS(const ShaderMetaType::CompiledShaderInitializerType& Initializer)
+		: FGlobalShader(Initializer)
+	{
+		ViewParameter.Bind(Initializer.ParameterMap, TEXT("View"));
+		SceneTexturesParameter.Bind(Initializer.ParameterMap, TEXT("SceneTexturesStruct"));
+		RectLightParameter.Bind(Initializer.ParameterMap, TEXT("RectLight"));
+		TLASParameter.Bind(Initializer.ParameterMap, TEXT("TLAS"));
+		TransmissionProfilesTextureParameter.Bind(Initializer.ParameterMap, TEXT("SSProfilesTexture"));
+		TransmissionProfilesLinearSamplerParameter.Bind(Initializer.ParameterMap, TEXT("TransmissionProfilesLinearSampler"));
+
+		LuminanceUAVParameter.Bind(Initializer.ParameterMap, TEXT("RWLuminanceUAV"));
+		RayDistanceUAVParameter.Bind(Initializer.ParameterMap, TEXT("RWRayDistanceUAV"));
+	}
+
+	bool Serialize(FArchive& Ar)
+	{
+		bool bShaderHasOutdatedParameters = FGlobalShader::Serialize(Ar);
+		Ar << ViewParameter;
+		Ar << SceneTexturesParameter;
+		Ar << RectLightParameter;
+		Ar << TLASParameter;
+		Ar << TransmissionProfilesTextureParameter;
+		Ar << TransmissionProfilesLinearSamplerParameter;
+		Ar << LuminanceUAVParameter;
+		Ar << RayDistanceUAVParameter;
+		return bShaderHasOutdatedParameters;
+	}
+
+	void Dispatch(
+		FRHICommandListImmediate& RHICmdList,
+		const FRayTracingScene& RayTracingScene,
+		FUniformBufferRHIParamRef ViewUniformBuffer,
+		FUniformBufferRHIParamRef SceneTexturesUniformBuffer,
+		FUniformBufferRHIParamRef RectLightUniformBuffer,
+		FUnorderedAccessViewRHIParamRef LuminanceUAV,
+		FUnorderedAccessViewRHIParamRef RayDistanceUAV,
+		uint32 Width, uint32 Height
+	)
+	{
+		FRayTracingPipelineStateInitializer Initializer;
+
+		FRayTracingShaderRHIParamRef RayGenShaderTable[] = { GetRayTracingShader() };
+		Initializer.SetRayGenShaderTable(RayGenShaderTable);
+
+		FRHIRayTracingPipelineState* Pipeline = PipelineStateCache::GetAndOrCreateRayTracingPipelineState(Initializer); // #dxr_todo: this should be done once at load-time and cached
+
+		FRayTracingShaderBindingsWriter GlobalResources;
+		GlobalResources.Set(TLASParameter, RayTracingScene.RayTracingSceneRHI->GetShaderResourceView());
+		GlobalResources.Set(ViewParameter, ViewUniformBuffer);
+		GlobalResources.Set(SceneTexturesParameter, SceneTexturesUniformBuffer);
+		GlobalResources.Set(RectLightParameter, RectLightUniformBuffer);
+		GlobalResources.Set(LuminanceUAVParameter, LuminanceUAV);
+		GlobalResources.Set(RayDistanceUAVParameter, RayDistanceUAV);
+
+		if (TransmissionProfilesTextureParameter.IsBound())
+		{
+			FSceneRenderTargets& SceneContext = FSceneRenderTargets::Get(RHICmdList);
+			const IPooledRenderTarget* PooledRT = GetSubsufaceProfileTexture_RT((FRHICommandListImmediate&)RHICmdList);
+
+			if (!PooledRT)
+			{
+				// no subsurface profile was used yet
+				PooledRT = GSystemTextures.BlackDummy;
+			}
+			const FSceneRenderTargetItem& Item = PooledRT->GetRenderTargetItem();
+
+			GlobalResources.SetTexture(TransmissionProfilesTextureParameter.GetBaseIndex(), Item.ShaderResourceTexture);
+			GlobalResources.SetSampler(TransmissionProfilesLinearSamplerParameter.GetBaseIndex(), TStaticSamplerState<SF_Bilinear, AM_Clamp, AM_Clamp, AM_Clamp>::GetRHI());
+		}
+
+		RHICmdList.RayTraceDispatch(Pipeline, GetRayTracingShader(), RayTracingScene.RayTracingSceneRHI, GlobalResources, Width, Height);
+	}
+
+private:
+	// Input
+	FShaderResourceParameter TLASParameter;
+	FShaderUniformBufferParameter ViewParameter;
+	FShaderUniformBufferParameter SceneTexturesParameter;
+	FShaderUniformBufferParameter RectLightParameter;
+
+	// SSS Profile
+	FShaderResourceParameter TransmissionProfilesTextureParameter;
+	FShaderResourceParameter TransmissionProfilesLinearSamplerParameter;
+
+	// Output
+	FShaderResourceParameter LuminanceUAVParameter;
+	FShaderResourceParameter RayDistanceUAVParameter;
+};
+
+#define IMPLEMENT_RECT_LIGHT_TYPE(TextureImportanceSampling) \
+	typedef FRectLightRGS<TextureImportanceSampling> FRectLightRGS##TextureImportanceSampling; \
+	IMPLEMENT_SHADER_TYPE(template<>, FRectLightRGS##TextureImportanceSampling, TEXT("/Engine/Private/RayTracing/RayTracingRectLightRGS.usf"), TEXT("RectLightRGS"), SF_RayGen);
+
+IMPLEMENT_RECT_LIGHT_TYPE(0);
+IMPLEMENT_RECT_LIGHT_TYPE(1);
+
+class FVisualizeRectLightMipTreePS : public FGlobalShader
+{
+	DECLARE_SHADER_TYPE(FVisualizeRectLightMipTreePS, Global);
+
+public:
+	static bool ShouldCache(EShaderPlatform Platform)
+	{
+		return ShouldCompileRayTracingShadersForProject(Platform);
+	}
+
+	static bool ShouldCompilePermutation(const FGlobalShaderPermutationParameters& Parameters)
+	{
+		return ShouldCompileRayTracingShadersForProject(Parameters.Platform);
+	}
+
+	static void ModifyCompilationEnvironment(const FGlobalShaderPermutationParameters& Parameters, FShaderCompilerEnvironment& OutEnvironment)
+	{
+		FGlobalShader::ModifyCompilationEnvironment(Parameters, OutEnvironment);
+	}
+
+	FVisualizeRectLightMipTreePS(const ShaderMetaType::CompiledShaderInitializerType& Initializer)
+		: FGlobalShader(Initializer)
+	{
+		DimensionsParameter.Bind(Initializer.ParameterMap, TEXT("Dimensions"));
+		MipTreeParameter.Bind(Initializer.ParameterMap, TEXT("MipTree"));
+	}
+
+	FVisualizeRectLightMipTreePS() {}
+
+	template<typename TRHICommandList>
+	void SetParameters(
+		TRHICommandList& RHICmdList,
+		const FViewInfo& View,
+		const FRWBuffer& MipTree,
+		const FIntVector Dimensions)
+	{
+		const FPixelShaderRHIParamRef ShaderRHI = GetPixelShader();
+		FGlobalShader::SetParameters<FViewUniformShaderParameters>(RHICmdList, ShaderRHI, View.ViewUniformBuffer);
+
+		SetShaderValue(RHICmdList, ShaderRHI, DimensionsParameter, Dimensions);
+		SetSRVParameter(RHICmdList, ShaderRHI, MipTreeParameter, MipTree.SRV);
+	}
+
+	virtual bool Serialize(FArchive& Ar) override
+	{
+		bool bShaderHasOutdatedParameters = FGlobalShader::Serialize(Ar);
+		Ar << DimensionsParameter;
+		Ar << MipTreeParameter;
+		return bShaderHasOutdatedParameters;
+	}
+
+private:
+	FShaderParameter DimensionsParameter;
+	FShaderResourceParameter MipTreeParameter;
+};
+
+IMPLEMENT_SHADER_TYPE(, FVisualizeRectLightMipTreePS, TEXT("/Engine/Private/PathTracing/VisualizeMipTreePixelShader.usf"), TEXT("VisualizeMipTreePS"), SF_Pixel)
+
+void FDeferredShadingSceneRenderer::VisualizeRectLightMipTree(
+	FRHICommandListImmediate& RHICmdList,
+	const FViewInfo& View,
+	const FRWBuffer& RectLightMipTree,
+	const FIntVector& RectLightMipTreeDimensions
+)
+{
+	// Allocate render target
+	FSceneRenderTargets& SceneContext = FSceneRenderTargets::Get(RHICmdList);
+	FPooledRenderTargetDesc Desc = SceneContext.GetSceneColor()->GetDesc();
+	Desc.Flags &= ~(TexCreate_FastVRAM | TexCreate_Transient);
+	TRefCountPtr<IPooledRenderTarget> RectLightMipTreeRT;
+	GRenderTargetPool.FindFreeElement(RHICmdList, Desc, RectLightMipTreeRT, TEXT("RectLightMipTreeRT"));
+
+	// Define shaders
+	const auto ShaderMap = GetGlobalShaderMap(View.FeatureLevel);
+	TShaderMapRef<FPostProcessVS> VertexShader(ShaderMap);
+	TShaderMapRef<FVisualizeRectLightMipTreePS> PixelShader(ShaderMap);
+	FTextureRHIParamRef RenderTargets[2] =
+	{
+		SceneContext.GetSceneColor()->GetRenderTargetItem().TargetableTexture,
+		RectLightMipTreeRT->GetRenderTargetItem().TargetableTexture
+	};
+	PRAGMA_DISABLE_DEPRECATION_WARNINGS
+	SetRenderTargets(RHICmdList, 2, RenderTargets, SceneContext.GetSceneDepthSurface(), ESimpleRenderTargetMode::EExistingColorAndDepth, FExclusiveDepthStencil::DepthRead_StencilNop);
+	PRAGMA_ENABLE_DEPRECATION_WARNINGS
+
+	// PSO definition
+	FGraphicsPipelineStateInitializer GraphicsPSOInit;
+	SceneContext.BeginRenderingSceneColor(RHICmdList, ESimpleRenderTargetMode::EExistingColorAndDepth, FExclusiveDepthStencil::DepthRead_StencilWrite, true);
+	RHICmdList.ApplyCachedRenderTargets(GraphicsPSOInit);
+	GraphicsPSOInit.BlendState = TStaticBlendState<CW_RGB, BO_Add, BF_One, BF_One>::GetRHI();
+	GraphicsPSOInit.RasterizerState = TStaticRasterizerState<FM_Solid, CM_None>::GetRHI();
+	GraphicsPSOInit.DepthStencilState = TStaticDepthStencilState<false, CF_Always>::GetRHI();
+	GraphicsPSOInit.BoundShaderState.VertexDeclarationRHI = GFilterVertexDeclaration.VertexDeclarationRHI;
+	GraphicsPSOInit.BoundShaderState.VertexShaderRHI = GETSAFERHISHADER_VERTEX(*VertexShader);
+	GraphicsPSOInit.BoundShaderState.PixelShaderRHI = GETSAFERHISHADER_PIXEL(*PixelShader);
+	GraphicsPSOInit.PrimitiveType = PT_TriangleList;
+	SetGraphicsPipelineState(RHICmdList, GraphicsPSOInit);
+
+	// Transition to graphics
+	RHICmdList.TransitionResource(EResourceTransitionAccess::EReadable, EResourceTransitionPipeline::EComputeToGfx, RectLightMipTree.UAV);
+
+	// Draw
+	RHICmdList.SetViewport(View.ViewRect.Min.X, View.ViewRect.Min.Y, 0.0f, View.ViewRect.Max.X, View.ViewRect.Max.Y, 1.0f);
+	PixelShader->SetParameters(RHICmdList, View, RectLightMipTree, RectLightMipTreeDimensions);
+	DrawRectangle(
+		RHICmdList,
+		0, 0,
+		View.ViewRect.Width(), View.ViewRect.Height(),
+		View.ViewRect.Min.X, View.ViewRect.Min.Y,
+		View.ViewRect.Width(), View.ViewRect.Height(),
+		FIntPoint(View.ViewRect.Width(), View.ViewRect.Height()),
+		SceneContext.GetBufferSizeXY(),
+		*VertexShader);
+	ResolveSceneColor(RHICmdList);
+	GVisualizeTexture.SetCheckPoint(RHICmdList, RectLightMipTreeRT);
+
+	// Transition to compute
+	RHICmdList.TransitionResource(EResourceTransitionAccess::ERWBarrier, EResourceTransitionPipeline::EGfxToCompute, RectLightMipTree.UAV);
+
+	ResolveSceneColor(RHICmdList);
+	SceneContext.FinishRenderingSceneColor(RHICmdList);
+}
+
+void FDeferredShadingSceneRenderer::PrepareRayTracingRectLight(const FViewInfo& View, TArray<FRayTracingShaderRHIParamRef>& OutRayGenShaders)
+{
+	// Declare all RayGen shaders that require material closest hit shaders to be bound
+
+	TShaderMapRef<FRectLightRGS<0>> Shader0(GetGlobalShaderMap(View.FeatureLevel));
+	TShaderMapRef<FRectLightRGS<1>> Shader1(GetGlobalShaderMap(View.FeatureLevel));
+
+	OutRayGenShaders.Add(Shader0->GetRayTracingShader());
+	OutRayGenShaders.Add(Shader1->GetRayTracingShader());
+}
+
+template <int TextureImportanceSampling>
+void RenderRayTracingRectLightInternal(
+	FRHICommandListImmediate& RHICmdList,
+	const TArray<FViewInfo>& Views,
+	const FLightSceneInfo& RectLightSceneInfo,
+	TRefCountPtr<IPooledRenderTarget>& ScreenShadowMaskTexture,
+	TRefCountPtr<IPooledRenderTarget>& RayDistanceTexture
+)
+{
+	check(RectLightSceneInfo.Proxy);
+	check(RectLightSceneInfo.Proxy->IsRectLight());
+	FRectLightSceneProxy* RectLightSceneProxy = (FRectLightSceneProxy*)RectLightSceneInfo.Proxy;
+
+	FLightShaderParameters LightShaderParameters;
+	RectLightSceneProxy->GetLightShaderParameters(LightShaderParameters);
+
+	FRectLightData RectLightData;
+	RectLightData.SamplesPerPixel = GRayTracingStochasticRectLightSamplesPerPixel;
+	RectLightData.bIsTextureImportanceSampling = GRayTracingStochasticRectLightIsTextureImportanceSampling;
+	RectLightData.Position = RectLightSceneInfo.Proxy->GetOrigin();
+	RectLightData.Normal = RectLightSceneInfo.Proxy->GetDirection();
+	const FMatrix& WorldToLight = RectLightSceneInfo.Proxy->GetWorldToLight();
+	RectLightData.dPdu = FVector(WorldToLight.M[0][1], WorldToLight.M[1][1], WorldToLight.M[2][1]);
+	RectLightData.dPdv = FVector(WorldToLight.M[0][2], WorldToLight.M[1][2], WorldToLight.M[2][2]);
+	RectLightData.Color = LightShaderParameters.Color / 2.0;
+
+	// #dxr_todo: Ray traced textured area lights are 1.5X brighter than those in lit mode.
+	if (RectLightSceneProxy->HasSourceTexture())
+	{
+		RectLightData.Color *= 2.0 / 3.0;
+	}
+
+	RectLightData.Width = 2.0f * LightShaderParameters.SourceRadius;
+	RectLightData.Height = 2.0f * LightShaderParameters.SourceLength;
+	RectLightData.Texture = LightShaderParameters.SourceTexture;
+	RectLightData.TextureSampler = RHICreateSamplerState(FSamplerStateInitializerRHI(SF_Bilinear, AM_Border, AM_Border, AM_Border));
+	RectLightData.MipTree = RectLightSceneProxy->RayTracingData->RectLightMipTree.SRV;
+	RectLightData.MipTreeDimensions = RectLightSceneProxy->RayTracingData->RectLightMipTreeDimensions;
+	RectLightData.MaxNormalBias = GetRaytracingMaxNormalBias();
+	RectLightData.BarnCosAngle = FMath::Cos(FMath::DegreesToRadians(RectLightSceneProxy->BarnDoorAngle));
+	RectLightData.BarnLength = RectLightSceneProxy->BarnDoorLength;
+	FUniformBufferRHIRef RectLightUniformBuffer = RHICreateUniformBuffer(&RectLightData, FRectLightData::StaticStructMetadata.GetLayout(), EUniformBufferUsage::UniformBuffer_SingleDraw);
+
+	for (int32 ViewIndex = 0; ViewIndex < Views.Num(); ++ViewIndex)
+	{
+		const FViewInfo& View = Views[ViewIndex];
+		const FIntPoint ViewSize = View.ViewRect.Size();
+
+		TShaderMapRef<FRectLightRGS<TextureImportanceSampling>> RectLightRayGenerationShader(GetGlobalShaderMap(View.FeatureLevel));
+
+		FSceneRenderTargets& SceneContext = FSceneRenderTargets::Get(RHICmdList);
+		FSceneTexturesUniformParameters SceneTextures;
+		SetupSceneTextureUniformParameters(SceneContext, View.FeatureLevel, ESceneTextureSetupMode::All, SceneTextures);
+		FUniformBufferRHIRef SceneTexturesUniformBuffer = RHICreateUniformBuffer(&SceneTextures, FSceneTexturesUniformParameters::StaticStructMetadata.GetLayout(), EUniformBufferUsage::UniformBuffer_SingleDraw);
+
+		// Dispatch
+		RectLightRayGenerationShader->Dispatch(
+			RHICmdList,
+			View.RayTracingScene,
+			View.ViewUniformBuffer,
+			SceneTexturesUniformBuffer,
+			RectLightUniformBuffer,
+			ScreenShadowMaskTexture->GetRenderTargetItem().UAV,
+			RayDistanceTexture->GetRenderTargetItem().UAV,
+			ViewSize.X, ViewSize.Y
+		);
+	}
+
+	// Transition out to graphics pipeline
+	FComputeFenceRHIRef Fence = RHICmdList.CreateComputeFence(TEXT("RayTracingRectLight"));
+	RHICmdList.TransitionResource(EResourceTransitionAccess::ERWBarrier, EResourceTransitionPipeline::EComputeToGfx, ScreenShadowMaskTexture->GetRenderTargetItem().UAV, Fence);
+	GVisualizeTexture.SetCheckPoint(RHICmdList, ScreenShadowMaskTexture);
+
+	RHICmdList.TransitionResource(EResourceTransitionAccess::ERWBarrier, EResourceTransitionPipeline::EComputeToGfx, RayDistanceTexture->GetRenderTargetItem().UAV);
+	GVisualizeTexture.SetCheckPoint(RHICmdList, RayDistanceTexture);
+}
+
+#endif // RHI_RAYTRACING
+
+void FDeferredShadingSceneRenderer::RenderRayTracingStochasticRectLight(
+	FRHICommandListImmediate& RHICmdList,
+	const FLightSceneInfo& RectLightSceneInfo,
+	TRefCountPtr<IPooledRenderTarget>& RectLightRT,
+	TRefCountPtr<IPooledRenderTarget>& HitDistanceRT
+)
+#if RHI_RAYTRACING
+{
+	SCOPED_DRAW_EVENT(RHICmdList, RayTracingRectLight);
+	SCOPED_GPU_STAT(RHICmdList, RayTracingRectLight);
+
+	FSceneRenderTargets& SceneContext = FSceneRenderTargets::Get(RHICmdList);
+	FPooledRenderTargetDesc Desc = SceneContext.GetSceneColor()->GetDesc();
+	Desc.Format = PF_FloatRGBA;
+	Desc.Flags &= ~(TexCreate_FastVRAM | TexCreate_Transient);
+	GRenderTargetPool.FindFreeElement(RHICmdList, Desc, RectLightRT, TEXT("RayTracingRectLight"));
+	//ClearUAV(RHICmdList, RectLightRT->GetRenderTargetItem(), FLinearColor::Black);
+
+	Desc.Format = PF_R16F;
+	GRenderTargetPool.FindFreeElement(RHICmdList, Desc, HitDistanceRT, TEXT("RayTracingRectLightDistance"));
+	//ClearUAV(RHICmdList, HitDistanceRT->GetRenderTargetItem(), FLinearColor::Black);
+
+	if (RectLightSceneInfo.Proxy->HasSourceTexture())
+	{
+		RenderRayTracingRectLightInternal<1>(RHICmdList, Views, RectLightSceneInfo, RectLightRT, HitDistanceRT);
+	}
+	else
+	{
+		RenderRayTracingRectLightInternal<0>(RHICmdList, Views, RectLightSceneInfo, RectLightRT, HitDistanceRT);
+	}
+}
+#else
+{
+	unimplemented();
+}
+#endif
diff --git a/Engine/Source/Runtime/Renderer/Public/RayTracingDynamicGeometryCollection.h b/Engine/Source/Runtime/Renderer/Public/RayTracingDynamicGeometryCollection.h
new file mode 100644
--- /dev/null
+++ b/Engine/Source/Runtime/Renderer/Public/RayTracingDynamicGeometryCollection.h
@@ -0,0 +1,20 @@
+// Copyright 1998-2019 Epic Games, Inc. All Rights Reserved.
+
+#pragma once
+
+#if RHI_RAYTRACING
+
+class RENDERER_API FRayTracingDynamicGeometryCollection
+{
+public:
+	FRayTracingDynamicGeometryCollection();
+
+	void AddDynamicMeshBatchForGeometryUpdate(const FScene* Scene, const FSceneView* View, const FPrimitiveSceneProxy* PrimitiveSceneProxy, const FMeshBatch& MeshBatch, FRayTracingGeometry& Geometry, uint32 NumMaxVertices, FRWBuffer& Buffer);
+	void DispatchUpdates(FRHICommandListImmediate& RHICmdList);
+	void Clear();
+
+private:
+	TUniquePtr<TArray<struct FMeshComputeDispatchCommand>> DispatchCommands;
+};
+
+#endif // RHI_RAYTRACING
