Fixed random crash in FMeshDrawCommand::SubmitDraw():
* First cached mesh draw commands (MDCs) were inserted into visible MDC arrays
* Then during init shadows ConditionalUpdateStaticMeshes could be called on MDCs which weren't visible in main view, but are visible in shadow maps
* ConditionalUpdateStaticMeshes calls CacheMeshDrawCommands which potentially can reallocate CachedDrawLists and CachedMeshDrawCommandStateBuckets, thus invalidating already issued MDC pointers inside visible MDC arrays
* Finally visible MDC arrays are submitted for draw and it can crash, as those pointers now point to random memory

Solution is to defer ConditionalUpdateStaticMeshes inside init shadow, by drawing such primitives with dynamic MDCs and marking them to be recached during the next InitViews.

diff --git a/Engine/Source/Runtime/Renderer/Private/PrimitiveSceneInfo.cpp b/Engine/Source/Runtime/Renderer/Private/PrimitiveSceneInfo.cpp
--- a/Engine/Source/Runtime/Renderer/Private/PrimitiveSceneInfo.cpp
+++ b/Engine/Source/Runtime/Renderer/Private/PrimitiveSceneInfo.cpp
@@ -126,6 +126,7 @@ FPrimitiveSceneInfo::FPrimitiveSceneInfo(UPrimitiveComponent* InComponent,FScene
 	PackedIndex(INDEX_NONE),
 	ComponentForDebuggingOnly(InComponent),
 	bNeedsStaticMeshUpdate(false),
+	bNeedsStaticMeshUpdateWithoutVisibilityCheck(false),
 	bNeedsUniformBufferUpdate(false),
 	bIndirectLightingCacheBufferDirty(false),
 	LightmapDataOffset(INDEX_NONE),
@@ -635,6 +636,13 @@ void FPrimitiveSceneInfo::RemoveFromScene(bool bUpdateStaticDrawLists)
 		bNeedsStaticMeshUpdate = false;
 	}
 
+	if (bNeedsStaticMeshUpdateWithoutVisibilityCheck)
+	{
+		Scene->PrimitivesNeedingStaticMeshUpdateWithoutVisibilityCheck.Remove(this);
+
+		bNeedsStaticMeshUpdateWithoutVisibilityCheck = false;
+	}
+
 	if (bUpdateStaticDrawLists)
 	{
 		// IndirectLightingCacheUniformBuffer may be cached inside cached mesh draw commands, so we 
@@ -663,6 +671,13 @@ void FPrimitiveSceneInfo::UpdateStaticMeshes(FRHICommandListImmediate& RHICmdLis
 		}
 	}
 
+	if (!bNeedsStaticMeshUpdate && bNeedsStaticMeshUpdateWithoutVisibilityCheck)
+	{
+		Scene->PrimitivesNeedingStaticMeshUpdateWithoutVisibilityCheck.Remove(this);
+
+		bNeedsStaticMeshUpdateWithoutVisibilityCheck = false;
+	}
+
 	RemoveCachedMeshDrawCommands();
 	if (bReAddToDrawLists)
 	{
@@ -670,6 +685,16 @@ void FPrimitiveSceneInfo::UpdateStaticMeshes(FRHICommandListImmediate& RHICmdLis
 	}
 }
 
+void FPrimitiveSceneInfo::ConditionalUpdateStaticMeshesWithoutVisibilityCheckDuringNextInitViews()
+{
+	if (bNeedsStaticMeshUpdate && !bNeedsStaticMeshUpdateWithoutVisibilityCheck)
+	{
+		bNeedsStaticMeshUpdateWithoutVisibilityCheck = true;
+
+		Scene->PrimitivesNeedingStaticMeshUpdateWithoutVisibilityCheck.Add(this);
+	}
+}
+
 void FPrimitiveSceneInfo::UpdateUniformBuffer(FRHICommandListImmediate& RHICmdList)
 {
 	checkSlow(bNeedsUniformBufferUpdate);
diff --git a/Engine/Source/Runtime/Renderer/Private/ScenePrivate.h b/Engine/Source/Runtime/Renderer/Private/ScenePrivate.h
--- a/Engine/Source/Runtime/Renderer/Private/ScenePrivate.h
+++ b/Engine/Source/Runtime/Renderer/Private/ScenePrivate.h
@@ -2498,6 +2498,7 @@ public:
 	TArray<FPrimitiveComponentId> PrimitiveComponentIds;
 
 	TSet<FPrimitiveSceneInfo*> PrimitivesNeedingStaticMeshUpdate;
+	TSet<FPrimitiveSceneInfo*> PrimitivesNeedingStaticMeshUpdateWithoutVisibilityCheck;
 
 	struct FTypeOffsetTableEntry
 	{
diff --git a/Engine/Source/Runtime/Renderer/Private/SceneVisibility.cpp b/Engine/Source/Runtime/Renderer/Private/SceneVisibility.cpp
--- a/Engine/Source/Runtime/Renderer/Private/SceneVisibility.cpp
+++ b/Engine/Source/Runtime/Renderer/Private/SceneVisibility.cpp
@@ -3415,6 +3415,19 @@ void FSceneRenderer::ComputeViewVisibility(FRHICommandListImmediate& RHICmdList,
 	const bool bIsInstancedStereo = (Views.Num() > 0) ? (Views[0].IsInstancedStereoPass() || Views[0].bIsMobileMultiViewEnabled) : false;
 	UpdateReflectionSceneData(Scene);
 
+	{
+		QUICK_SCOPE_CYCLE_COUNTER(STAT_ViewVisibilityTime_ConditionalUpdateStaticMeshesWithoutVisibilityCheck);
+
+		Scene->ConditionalMarkStaticMeshElementsForUpdate();
+
+		for (TSet<FPrimitiveSceneInfo*>::TIterator It(Scene->PrimitivesNeedingStaticMeshUpdateWithoutVisibilityCheck); It; ++It)
+		{
+			FPrimitiveSceneInfo* Primitive = *It;
+			Primitive->ConditionalUpdateStaticMeshes(RHICmdList);
+		}
+		Scene->PrimitivesNeedingStaticMeshUpdateWithoutVisibilityCheck.Reset();
+	}
+
 	uint8 ViewBit = 0x1;
 	for (int32 ViewIndex = 0; ViewIndex < Views.Num(); ++ViewIndex, ViewBit <<= 1)
 	{
@@ -3632,7 +3645,6 @@ void FSceneRenderer::ComputeViewVisibility(FRHICommandListImmediate& RHICmdList,
 
 		{
 			QUICK_SCOPE_CYCLE_COUNTER(STAT_ViewVisibilityTime_ConditionalUpdateStaticMeshes);
-			Scene->ConditionalMarkStaticMeshElementsForUpdate();
 
 			for (TSet<FPrimitiveSceneInfo*>::TIterator It(Scene->PrimitivesNeedingStaticMeshUpdate); It; ++It)
 			{
diff --git a/Engine/Source/Runtime/Renderer/Private/ShadowSetup.cpp b/Engine/Source/Runtime/Renderer/Private/ShadowSetup.cpp
--- a/Engine/Source/Runtime/Renderer/Private/ShadowSetup.cpp
+++ b/Engine/Source/Runtime/Renderer/Private/ShadowSetup.cpp
@@ -917,7 +917,7 @@ bool FProjectedShadowInfo::ShouldDrawStaticMeshes(FViewInfo& InCurrentView, bool
 
 			// Don't cache if it requires per view per mesh state for distance cull fade.
 			const bool bIsPrimitiveDistanceCullFading = InCurrentView.PotentiallyFadingPrimitiveMap[InPrimitiveSceneInfo->GetIndex()];
-			const bool bCanCache = !bIsPrimitiveDistanceCullFading;
+			const bool bCanCache = !bIsPrimitiveDistanceCullFading && !InPrimitiveSceneInfo->NeedsUpdateStaticMeshes();
 
 			for (int32 MeshIndex = 0; MeshIndex < InPrimitiveSceneInfo->StaticMeshRelevances.Num(); MeshIndex++)
 			{
@@ -1079,17 +1079,24 @@ void FProjectedShadowInfo::AddSubjectPrimitive(FPrimitiveSceneInfo* PrimitiveSce
 			// Update the primitive component's last render time. Allows the component to update when using bCastWhenHidden.
 			const float CurrentWorldTime = Views[0]->Family->CurrentWorldTime;
 			*(PrimitiveSceneInfo->ComponentLastRenderTime) = CurrentWorldTime;
-			if (PrimitiveSceneInfo->NeedsUniformBufferUpdate() || PrimitiveSceneInfo->NeedsUpdateStaticMeshes())
+
+			if (PrimitiveSceneInfo->NeedsUniformBufferUpdate())
 			{
 				for (int32 ViewIndex = 0; ViewIndex < Views.Num(); ++ViewIndex)
 				{
-					// Main view visible primitives are processed on parallel tasks, updating them here will cause a race condition.
+					// Main view visible primitives are processed on parallel tasks, updating uniform buffer them here will cause a race condition.
 					check(!Views[ViewIndex]->PrimitiveVisibilityMap[PrimitiveSceneInfo->GetIndex()]);
 				}
 
-				PrimitiveSceneInfo->ConditionalUpdateStaticMeshes(FRHICommandListExecutor::GetImmediateCommandList());
 				PrimitiveSceneInfo->ConditionalUpdateUniformBuffer(FRHICommandListExecutor::GetImmediateCommandList());
 			}
+
+			if (PrimitiveSceneInfo->NeedsUpdateStaticMeshes())
+			{
+				// Need to defer to next InitViews, as main view visible primitives are processed on parallel tasks and calling 
+				// CacheMeshDrawCommands may resize CachedDrawLists/CachedMeshDrawCommandStateBuckets causing a crash.
+				PrimitiveSceneInfo->ConditionalUpdateStaticMeshesWithoutVisibilityCheckDuringNextInitViews();
+			}
 		}
 
 		if (bOpaqueRelevance && bShadowRelevance)
diff --git a/Engine/Source/Runtime/Renderer/Public/PrimitiveSceneInfo.h b/Engine/Source/Runtime/Renderer/Public/PrimitiveSceneInfo.h
--- a/Engine/Source/Runtime/Renderer/Public/PrimitiveSceneInfo.h
+++ b/Engine/Source/Runtime/Renderer/Public/PrimitiveSceneInfo.h
@@ -341,6 +341,9 @@ public:
 		}
 	}
 
+	/** Will update static meshes during next InitViews, even if it's not visible. */
+	void ConditionalUpdateStaticMeshesWithoutVisibilityCheckDuringNextInitViews();
+
 	/** Updates the primitive's uniform buffer. */
 	void UpdateUniformBuffer(FRHICommandListImmediate& RHICmdList);
 
@@ -461,13 +464,16 @@ private:
 	const UPrimitiveComponent* ComponentForDebuggingOnly;
 
 	/** If this is TRUE, this primitive's static meshes needs to be updated before it can be rendered. */
-	bool bNeedsStaticMeshUpdate;
+	bool bNeedsStaticMeshUpdate : 1;
+
+	/** If this is TRUE, this primitive's static meshes will be update even if it's not visible. */
+	bool bNeedsStaticMeshUpdateWithoutVisibilityCheck : 1;
 
 	/** If this is TRUE, this primitive's uniform buffer needs to be updated before it can be rendered. */
-	bool bNeedsUniformBufferUpdate;
+	bool bNeedsUniformBufferUpdate : 1;
 
 	/** If this is TRUE, this primitive's indirect lighting cache buffer needs to be updated before it can be rendered. */
-	bool bIndirectLightingCacheBufferDirty;
+	bool bIndirectLightingCacheBufferDirty : 1;
 
 	/** Offset into the scene's lightmap data buffer, when GPUScene is enabled. */
 	int32 LightmapDataOffset;
