Fixed crash on sky light access in mobile renderer when trying to draw a lit mesh without a proxy (DrawTileMesh rendering path).

diff --git a/Engine/Source/Runtime/Renderer/Private/MobileBasePass.cpp b/Engine/Source/Runtime/Renderer/Private/MobileBasePass.cpp
--- a/Engine/Source/Runtime/Renderer/Private/MobileBasePass.cpp
+++ b/Engine/Source/Runtime/Renderer/Private/MobileBasePass.cpp
@@ -618,7 +618,7 @@ void FMobileBasePassMeshProcessor::Process(
 		FBaseDS,
 		TMobileBasePassPSPolicyParamType<FUniformLightMapPolicy>> BasePassShaders;
 	
-	bool bEnableSkyLight = ShadingModel != MSM_Unlit && Scene->ShouldRenderSkylightInBasePass(BlendMode);
+	bool bEnableSkyLight = ShadingModel != MSM_Unlit && Scene && Scene->ShouldRenderSkylightInBasePass(BlendMode);
 	int32 NumMovablePointLights = MobileBasePass::CalcNumMovablePointLights(MaterialResource, PrimitiveSceneProxy);
 
 	MobileBasePass::GetShaders(
