Fixed random crash in FMeshDrawCommand::SubmitDraw():
* First cached mesh draw commands (MDCs) were inserted into visible MDC arrays
* Then during init shadows ConditionalUpdateStaticMeshes could be called on MDCs which weren't visible in main view, but are visible in shadow maps
* ConditionalUpdateStaticMeshes calls CacheMeshDrawCommands which potentially can reallocate CachedDrawLists and CachedMeshDrawCommandStateBuckets, thus invalidating already issued MDC pointers inside visible MDC arrays
* Finally visible MDC arrays are submitted for draw and it can crash, as those pointers now point to random memory

Solution is to defer ConditionalUpdateStaticMeshes inside init shadow, by drawing such primitives with dynamic MDCs and marking them to be recached during the next InitViews.

diff --git a/Engine/Source/Runtime/Renderer/Private/PrimitiveSceneInfo.cpp b/Engine/Source/Runtime/Renderer/Private/PrimitiveSceneInfo.cpp
--- a/Engine/Source/Runtime/Renderer/Private/PrimitiveSceneInfo.cpp
+++ b/Engine/Source/Runtime/Renderer/Private/PrimitiveSceneInfo.cpp
@@ -683,6 +683,13 @@ void FPrimitiveSceneInfo::UpdateStaticMeshes(FRHICommandListImmediate& RHICmdLis
 	{
 		CacheMeshDrawCommands(RHICmdList);
 	}
+
+#if RHI_RAYTRACING
+	if (IsRayTracingEnabled())
+	{
+		UpdateRayTracingLodIndexToMeshDrawCommandIndicies();
+	}
+#endif // RHI_RAYTRACING
 }
 
 void FPrimitiveSceneInfo::ConditionalUpdateStaticMeshesWithoutVisibilityCheckDuringNextInitViews()
