Fixed shadow flickering happening in some case due to floating point precision issue.
Now always use 2 vertices from far clip plane to compute frustum edges plane equations.
The near clip plane could also go wrong I imagine but I could not produce the bug visually.

diff --git a/Engine/Source/Runtime/Renderer/Private/ShadowRendering.cpp b/Engine/Source/Runtime/Renderer/Private/ShadowRendering.cpp
--- a/Engine/Source/Runtime/Renderer/Private/ShadowRendering.cpp
+++ b/Engine/Source/Runtime/Renderer/Private/ShadowRendering.cpp
@@ -565,7 +565,7 @@ void FProjectedShadowInfo::SetupFrustumForProjection(const FViewInfo* View, TArr
 		const FPlane Top(BackTopRight, BackTopLeft, FrontTopLeft);
 		const float TopDistance = Top.PlaneDot(ShadowViewOrigin);
 
-		const FPlane Bottom(FrontBottomRight, FrontBottomLeft, BackBottomLeft);
+		const FPlane Bottom(BackBottomLeft, BackBottomRight, FrontBottomLeft);
 		const float BottomDistance = Bottom.PlaneDot(ShadowViewOrigin);
 
 		// Use a distance threshold to treat the case where the near plane is intersecting the frustum as the camera being inside
