Fix pre-exposure never be set.

Pre-exposure on the sceneInfo was never this. This got removed in CL 5238694 as a typo?.

diff --git a/Engine/Source/Runtime/Renderer/Private/PostProcess/PostProcessEyeAdaptation.cpp b/Engine/Source/Runtime/Renderer/Private/PostProcess/PostProcessEyeAdaptation.cpp
--- a/Engine/Source/Runtime/Renderer/Private/PostProcess/PostProcessEyeAdaptation.cpp
+++ b/Engine/Source/Runtime/Renderer/Private/PostProcess/PostProcessEyeAdaptation.cpp
@@ -470,6 +470,9 @@ void FSceneViewState::UpdatePreExposure(FViewInfo& View)
 			bUpdateLastExposure = true;
 		}
 	}
+
+	// Update the pre-exposure value on the actual view
+	View.PreExposure = PreExposure;
 }
 
 FPooledRenderTargetDesc FRCPassPostProcessEyeAdaptation::ComputeOutputDesc(EPassOutputId InPassOutputId) const
