Fix UpdateRayTracingLodIndexToMeshDrawCommandIndicies identifier not found

diff --git a/Engine/Source/Runtime/Renderer/Private/PrimitiveSceneInfo.cpp b/Engine/Source/Runtime/Renderer/Private/PrimitiveSceneInfo.cpp
--- a/Engine/Source/Runtime/Renderer/Private/PrimitiveSceneInfo.cpp
+++ b/Engine/Source/Runtime/Renderer/Private/PrimitiveSceneInfo.cpp
@@ -683,13 +683,6 @@ void FPrimitiveSceneInfo::UpdateStaticMeshes(FRHICommandListImmediate& RHICmdLis
 	{
 		CacheMeshDrawCommands(RHICmdList);
 	}
-
-#if RHI_RAYTRACING
-	if (IsRayTracingEnabled())
-	{
-		UpdateRayTracingLodIndexToMeshDrawCommandIndicies();
-	}
-#endif // RHI_RAYTRACING
 }
 
 void FPrimitiveSceneInfo::ConditionalUpdateStaticMeshesWithoutVisibilityCheckDuringNextInitViews()
