Fixed static mesh updates. Update primitive transform was calling RemoveFromScene, which was also disabling pending static mesh updates. This was causing artifacts where some meshes were rendered with stale data. Solution is to keep static mesh updates, unless RemoveFromScene is called with bUpdateStaticDrawLists, which forces meshes to be removed too.

diff --git a/Engine/Source/Runtime/Renderer/Private/PrimitiveSceneInfo.cpp b/Engine/Source/Runtime/Renderer/Private/PrimitiveSceneInfo.cpp
--- a/Engine/Source/Runtime/Renderer/Private/PrimitiveSceneInfo.cpp
+++ b/Engine/Source/Runtime/Renderer/Private/PrimitiveSceneInfo.cpp
@@ -629,22 +629,22 @@ void FPrimitiveSceneInfo::RemoveFromScene(bool bUpdateStaticDrawLists)
 
 	DEC_MEMORY_STAT_BY(STAT_PrimitiveInfoMemory, sizeof(*this) + StaticMeshes.GetAllocatedSize() + StaticMeshRelevances.GetAllocatedSize() + Proxy->GetMemoryFootprint());
 
-	if (bNeedsStaticMeshUpdate)
+	if (bUpdateStaticDrawLists)
 	{
-		Scene->PrimitivesNeedingStaticMeshUpdate.Remove(this);
+		if (bNeedsStaticMeshUpdate)
+		{
+			Scene->PrimitivesNeedingStaticMeshUpdate.Remove(this);
 
-		bNeedsStaticMeshUpdate = false;
-	}
+			bNeedsStaticMeshUpdate = false;
+		}
 
-	if (bNeedsStaticMeshUpdateWithoutVisibilityCheck)
-	{
-		Scene->PrimitivesNeedingStaticMeshUpdateWithoutVisibilityCheck.Remove(this);
+		if (bNeedsStaticMeshUpdateWithoutVisibilityCheck)
+		{
+			Scene->PrimitivesNeedingStaticMeshUpdateWithoutVisibilityCheck.Remove(this);
 
-		bNeedsStaticMeshUpdateWithoutVisibilityCheck = false;
-	}
+			bNeedsStaticMeshUpdateWithoutVisibilityCheck = false;
+		}
 
-	if (bUpdateStaticDrawLists)
-	{
 		// IndirectLightingCacheUniformBuffer may be cached inside cached mesh draw commands, so we 
 		// can't delete it unless we also update cached mesh command.
 		IndirectLightingCacheUniformBuffer.SafeRelease();
@@ -685,16 +685,6 @@ void FPrimitiveSceneInfo::UpdateStaticMeshes(FRHICommandListImmediate& RHICmdLis
 	}
 }
 
-void FPrimitiveSceneInfo::ConditionalUpdateStaticMeshesWithoutVisibilityCheckDuringNextInitViews()
-{
-	if (bNeedsStaticMeshUpdate && !bNeedsStaticMeshUpdateWithoutVisibilityCheck)
-	{
-		bNeedsStaticMeshUpdateWithoutVisibilityCheck = true;
-
-		Scene->PrimitivesNeedingStaticMeshUpdateWithoutVisibilityCheck.Add(this);
-	}
-}
-
 void FPrimitiveSceneInfo::UpdateUniformBuffer(FRHICommandListImmediate& RHICmdList)
 {
 	checkSlow(bNeedsUniformBufferUpdate);
@@ -713,6 +703,16 @@ void FPrimitiveSceneInfo::BeginDeferredUpdateStaticMeshes()
 	}
 }
 
+void FPrimitiveSceneInfo::BeginDeferredUpdateStaticMeshesWithoutVisibilityCheck()
+{
+	if (bNeedsStaticMeshUpdate && !bNeedsStaticMeshUpdateWithoutVisibilityCheck)
+	{
+		bNeedsStaticMeshUpdateWithoutVisibilityCheck = true;
+
+		Scene->PrimitivesNeedingStaticMeshUpdateWithoutVisibilityCheck.Add(this);
+	}
+}
+
 void FPrimitiveSceneInfo::LinkLODParentComponent()
 {
 	if (LODParentComponentId.IsValid())
diff --git a/Engine/Source/Runtime/Renderer/Private/ShadowSetup.cpp b/Engine/Source/Runtime/Renderer/Private/ShadowSetup.cpp
--- a/Engine/Source/Runtime/Renderer/Private/ShadowSetup.cpp
+++ b/Engine/Source/Runtime/Renderer/Private/ShadowSetup.cpp
@@ -1095,7 +1095,7 @@ void FProjectedShadowInfo::AddSubjectPrimitive(FPrimitiveSceneInfo* PrimitiveSce
 			{
 				// Need to defer to next InitViews, as main view visible primitives are processed on parallel tasks and calling 
 				// CacheMeshDrawCommands may resize CachedDrawLists/CachedMeshDrawCommandStateBuckets causing a crash.
-				PrimitiveSceneInfo->ConditionalUpdateStaticMeshesWithoutVisibilityCheckDuringNextInitViews();
+				PrimitiveSceneInfo->BeginDeferredUpdateStaticMeshesWithoutVisibilityCheck();
 			}
 		}
 
diff --git a/Engine/Source/Runtime/Renderer/Public/PrimitiveSceneInfo.h b/Engine/Source/Runtime/Renderer/Public/PrimitiveSceneInfo.h
--- a/Engine/Source/Runtime/Renderer/Public/PrimitiveSceneInfo.h
+++ b/Engine/Source/Runtime/Renderer/Public/PrimitiveSceneInfo.h
@@ -341,9 +341,6 @@ public:
 		}
 	}
 
-	/** Will update static meshes during next InitViews, even if it's not visible. */
-	void ConditionalUpdateStaticMeshesWithoutVisibilityCheckDuringNextInitViews();
-
 	/** Updates the primitive's uniform buffer. */
 	void UpdateUniformBuffer(FRHICommandListImmediate& RHICmdList);
 
@@ -359,6 +356,9 @@ public:
 	/** Sets a flag to update the primitive's static meshes before it is next rendered. */
 	void BeginDeferredUpdateStaticMeshes();
 
+	/** Will update static meshes during next InitViews, even if it's not visible. */
+	void BeginDeferredUpdateStaticMeshesWithoutVisibilityCheck();
+
 	/** Adds the primitive's static meshes to the scene. */
 	void AddStaticMeshes(FRHICommandListImmediate& RHICmdList, bool bUpdateStaticDrawLists = true);
 
