Cannot fallback to cached shadow maps if they exceed max resolution

diff --git a/Engine/Source/Runtime/Renderer/Private/ShadowSetup.cpp b/Engine/Source/Runtime/Renderer/Private/ShadowSetup.cpp
--- a/Engine/Source/Runtime/Renderer/Private/ShadowSetup.cpp
+++ b/Engine/Source/Runtime/Renderer/Private/ShadowSetup.cpp
@@ -2151,11 +2151,19 @@ void FSceneRenderer::CreatePerObjectProjectedShadow(
 	}
 }
 
+static bool CanFallbackToOldShadowMapCache(const FShadowMapRenderTargetsRefCounted& CachedShadowMap, const FIntPoint& MaxShadowResolution)
+{
+	return CachedShadowMap.IsValid()
+		&& CachedShadowMap.GetSize().X <= MaxShadowResolution.X
+		&& CachedShadowMap.GetSize().Y <= MaxShadowResolution.Y;
+}
+
 void ComputeWholeSceneShadowCacheModes(
 	const FLightSceneInfo* LightSceneInfo,
 	bool bCubeShadowMap,
 	float RealTime,
 	float ActualDesiredResolution,
+	const FIntPoint& MaxShadowResolution,
 	FScene* Scene,
 	FWholeSceneProjectedShadowInitializer& InOutProjectedShadowInitializer,
 	FIntPoint& InOutShadowMapSize,
@@ -2214,7 +2222,7 @@ void ComputeWholeSceneShadowCacheModes(
 						++*NumCachesUpdatedThisFrame;
 						
 						// Check if update is caused by resolution change
-						if (CachedShadowMapData->ShadowMap.IsValid())
+						if (CanFallbackToOldShadowMapCache(CachedShadowMapData->ShadowMap, MaxShadowResolution))
 						{
 							FIntPoint ExistingShadowMapSize = CachedShadowMapData->ShadowMap.GetSize();
 							bool bOverBudget = *NumCachesUpdatedThisFrame > MaxCacheUpdatesAllowed;
@@ -2556,6 +2564,7 @@ void FSceneRenderer::CreateWholeSceneProjectedShadow(
 						ProjectedShadowInitializer.bOnePassPointLightShadow,
 						ViewFamily.CurrentRealTime,
 						MaxDesiredResolution,
+						FIntPoint(MaxShadowResolution, MaxShadowResolutionY),
 						Scene,
 						// Below are in-out or out parameters. They can change
 						ProjectedShadowInitializer,
