Backing out fix for UE-69918, needs more work.

diff --git a/Engine/Source/Runtime/Engine/Private/Materials/HLSLMaterialTranslator.h b/Engine/Source/Runtime/Engine/Private/Materials/HLSLMaterialTranslator.h
--- a/Engine/Source/Runtime/Engine/Private/Materials/HLSLMaterialTranslator.h
+++ b/Engine/Source/Runtime/Engine/Private/Materials/HLSLMaterialTranslator.h
@@ -3799,6 +3799,8 @@ protected:
 		else // mobile
 		{
 			int32 UV = BufferUV;
+
+			// On mobile in post process material, there is no need to do ViewportUV->BufferUV conversion because ViewSize == BufferSize.
 			if (Material->GetMaterialDomain() == MD_PostProcess)
 			{
 				int32 BlendableLocation = Material->GetBlendableLocation();
@@ -3813,6 +3815,10 @@ protected:
 				{
 					UV = TextureCoordinate(0, false, false);
 				}
+				else
+				{
+					UV = ViewportUV;
+				}
 			}
 			
 			return AddCodeChunk(MCT_Float4,	TEXT("MobileSceneTextureLookup(Parameters, %d, %s)"), (int32)SceneTextureId, *CoerceParameter(UV, MCT_Float2));
diff --git a/Engine/Source/Runtime/Renderer/Private/PostProcess/PostProcessMaterial.cpp b/Engine/Source/Runtime/Renderer/Private/PostProcess/PostProcessMaterial.cpp
--- a/Engine/Source/Runtime/Renderer/Private/PostProcess/PostProcessMaterial.cpp
+++ b/Engine/Source/Runtime/Renderer/Private/PostProcess/PostProcessMaterial.cpp
@@ -177,13 +177,11 @@ private:
 
 typedef FPostProcessMaterialPS<EPostProcessMaterialTarget::HighEnd, 0> FFPostProcessMaterialPS_HighEnd0;
 typedef FPostProcessMaterialPS<EPostProcessMaterialTarget::HighEnd, 1> FFPostProcessMaterialPS_HighEnd1;
-typedef FPostProcessMaterialPS<EPostProcessMaterialTarget::Mobile, 0> FPostProcessMaterialPS_Mobile0;
-typedef FPostProcessMaterialPS<EPostProcessMaterialTarget::Mobile, 1> FPostProcessMaterialPS_Mobile1;
+typedef FPostProcessMaterialPS<EPostProcessMaterialTarget::Mobile, 0> FPostProcessMaterialPS_Mobile;
 
 IMPLEMENT_MATERIAL_SHADER_TYPE(template<>, FFPostProcessMaterialPS_HighEnd0, TEXT("/Engine/Private/PostProcessMaterialShaders.usf"), TEXT("MainPS"), SF_Pixel);
 IMPLEMENT_MATERIAL_SHADER_TYPE(template<>, FFPostProcessMaterialPS_HighEnd1, TEXT("/Engine/Private/PostProcessMaterialShaders.usf"), TEXT("MainPS"), SF_Pixel);
-IMPLEMENT_MATERIAL_SHADER_TYPE(template<>, FPostProcessMaterialPS_Mobile0, TEXT("/Engine/Private/PostProcessMaterialShaders.usf"), TEXT("MainPS_ES2"), SF_Pixel);
-IMPLEMENT_MATERIAL_SHADER_TYPE(template<>, FPostProcessMaterialPS_Mobile1, TEXT("/Engine/Private/PostProcessMaterialShaders.usf"), TEXT("MainPS_ES2"), SF_Pixel);
+IMPLEMENT_MATERIAL_SHADER_TYPE(template<>,FPostProcessMaterialPS_Mobile,TEXT("/Engine/Private/PostProcessMaterialShaders.usf"),TEXT("MainPS_ES2"),SF_Pixel);
 
 FRCPassPostProcessMaterial::FRCPassPostProcessMaterial(UMaterialInterface* InMaterialInterface, ERHIFeatureLevel::Type InFeatureLevel, EPixelFormat OutputFormatIN)
 : MaterialInterface(InMaterialInterface), OutputFormat(OutputFormatIN)
@@ -230,24 +228,6 @@ public:
 };
 TGlobalResource<FPostProcessMaterialVertexDeclaration> GPostProcessMaterialVertexDeclaration;
 
-template<typename TPixelShader>
-FShader* SetMobileShaders(const FMaterialShaderMap* MaterialShaderMap, FGraphicsPipelineStateInitializer &GraphicsPSOInit, FRenderingCompositePassContext &Context, FMaterialRenderProxy* Proxy, uint32 StencilRefValue)
-{
-	TPixelShader* PixelShader_Mobile = MaterialShaderMap->GetShader<TPixelShader>();
-	FPostProcessMaterialVS_Mobile* VertexShader_Mobile = MaterialShaderMap->GetShader<FPostProcessMaterialVS_Mobile>();
-
-	GraphicsPSOInit.BoundShaderState.VertexDeclarationRHI = GFilterVertexDeclaration.VertexDeclarationRHI;
-	GraphicsPSOInit.BoundShaderState.VertexShaderRHI = GETSAFERHISHADER_VERTEX(VertexShader_Mobile);
-	GraphicsPSOInit.BoundShaderState.PixelShaderRHI = GETSAFERHISHADER_PIXEL(PixelShader_Mobile);
-
-	SetGraphicsPipelineState(Context.RHICmdList, GraphicsPSOInit);
-	Context.RHICmdList.SetStencilRef(StencilRefValue);
-
-	VertexShader_Mobile->SetParameters(Context.RHICmdList, Context, Proxy);
-	PixelShader_Mobile->SetParameters(Context.RHICmdList, Context, Proxy);
-	return VertexShader_Mobile;
-}
-
 void FRCPassPostProcessMaterial::Process(FRenderingCompositePassContext& Context)
 {
 	FMaterialRenderProxy* Proxy = MaterialInterface->GetRenderProxy();
@@ -441,21 +421,25 @@ void FRCPassPostProcessMaterial::Process(FRenderingCompositePassContext& Context
 
 		FShader* VertexShader = nullptr;
 
-		const bool bViewSizeMatchesBufferSize = (View.ViewRect == Context.SceneColorViewRect && View.ViewRect.Size() == SrcSize && View.ViewRect.Min == FIntPoint::ZeroValue);
+		// uses mobile's post process material.
 		if (FeatureLevel <= ERHIFeatureLevel::ES3_1)
 		{
-			// use mobile's post process material.
-			if (bViewSizeMatchesBufferSize)
-			{
-				VertexShader = SetMobileShaders< FPostProcessMaterialPS_Mobile0>(MaterialShaderMap, GraphicsPSOInit, Context, Proxy, StencilRefValue);
-			}
-			else
-			{
-				VertexShader = SetMobileShaders< FPostProcessMaterialPS_Mobile1>(MaterialShaderMap, GraphicsPSOInit, Context, Proxy, StencilRefValue);
-			}
+			FPostProcessMaterialPS_Mobile* PixelShader_Mobile = MaterialShaderMap->GetShader<FPostProcessMaterialPS_Mobile>();
+			FPostProcessMaterialVS_Mobile* VertexShader_Mobile = MaterialShaderMap->GetShader<FPostProcessMaterialVS_Mobile>();
+
+			GraphicsPSOInit.BoundShaderState.VertexDeclarationRHI = GFilterVertexDeclaration.VertexDeclarationRHI;
+			GraphicsPSOInit.BoundShaderState.VertexShaderRHI = GETSAFERHISHADER_VERTEX(VertexShader_Mobile);
+			GraphicsPSOInit.BoundShaderState.PixelShaderRHI = GETSAFERHISHADER_PIXEL(PixelShader_Mobile);
+
+			SetGraphicsPipelineState(Context.RHICmdList, GraphicsPSOInit);
+			Context.RHICmdList.SetStencilRef(StencilRefValue);
+
+			VertexShader_Mobile->SetParameters(Context.RHICmdList, Context, Proxy);
+			PixelShader_Mobile->SetParameters(Context.RHICmdList, Context, Proxy);
+			VertexShader = VertexShader_Mobile;
 		}
 		// Uses highend post process material that assumed ViewSize == BufferSize.
-		else if (bViewSizeMatchesBufferSize)
+		else if (View.ViewRect == Context.SceneColorViewRect && View.ViewRect.Size() == SrcSize && View.ViewRect.Min == FIntPoint::ZeroValue)
 		{
 			FFPostProcessMaterialPS_HighEnd0* PixelShader_HighEnd = MaterialShaderMap->GetShader<FFPostProcessMaterialPS_HighEnd0>();
 			FPostProcessMaterialVS_HighEnd* VertexShader_HighEnd = MaterialShaderMap->GetShader<FPostProcessMaterialVS_HighEnd>();
