Add permutation to mobile's post process materials to avoid viewport to colour buffer calculation when viewsize == buffer size.
Remove assumption from SceneTextureLookup that mobile only deals with viewsize == buffersize case.
Corrected default UV values for SceneTextureLookup.
Corrected material UVs.

diff --git a/Engine/Shaders/Private/PostProcessMaterialShaders.usf b/Engine/Shaders/Private/PostProcessMaterialShaders.usf
--- a/Engine/Shaders/Private/PostProcessMaterialShaders.usf
+++ b/Engine/Shaders/Private/PostProcessMaterialShaders.usf
@@ -215,7 +215,7 @@ void MainVS_ES2(
 	UNROLL
 	for (int CoordinateIndex = 0; CoordinateIndex < NUM_MATERIAL_TEXCOORDS; CoordinateIndex++)
 	{
-		SetUV(Output, CoordinateIndex, OutUV);
+		SetUV(Output, CoordinateIndex, InPosition.xy);
 	}
 	UNROLL
 	for (int CoordinateIndex = NUM_MATERIAL_TEXCOORDS; CoordinateIndex < NUM_TEX_COORD_INTERPOLATORS; CoordinateIndex++)
diff --git a/Engine/Shaders/Public/ShaderVersion.ush b/Engine/Shaders/Public/ShaderVersion.ush
--- a/Engine/Shaders/Public/ShaderVersion.ush
+++ b/Engine/Shaders/Public/ShaderVersion.ush
@@ -3,5 +3,4 @@
 // in Platform.ush (which should be included in any shader) it allows to invalidate the shader DDC.
 // 
 // If you are merging streams and there is a conflict with this GUID you should make a new GUID rather than taking one or the other.
-// GUID = AB45A550-00EC-43F7-AB17-E6C58BF99AA8
-
+// GUID = 228AA72549B318E4175C78AB57A4F8AC
\ No newline at end of file
diff --git a/Engine/Source/Runtime/Engine/Private/Materials/HLSLMaterialTranslator.h b/Engine/Source/Runtime/Engine/Private/Materials/HLSLMaterialTranslator.h
--- a/Engine/Source/Runtime/Engine/Private/Materials/HLSLMaterialTranslator.h
+++ b/Engine/Source/Runtime/Engine/Private/Materials/HLSLMaterialTranslator.h
@@ -3799,8 +3799,6 @@ protected:
 		else // mobile
 		{
 			int32 UV = BufferUV;
-
-			// On mobile in post process material, there is no need to do ViewportUV->BufferUV conversion because ViewSize == BufferSize.
 			if (Material->GetMaterialDomain() == MD_PostProcess)
 			{
 				int32 BlendableLocation = Material->GetBlendableLocation();
@@ -3810,15 +3808,6 @@ protected:
 					// SceneColor.A channel holds depth till BeforeTonemapping location, then it's gets overwritten
 					return Errorf(TEXT("SceneDepth lookups are only available when BlendableLocation is BeforeTranslucency or BeforeTonemapping"));
 				}
-				
-				if (ViewportUV == INDEX_NONE)
-				{
-					UV = TextureCoordinate(0, false, false);
-				}
-				else
-				{
-					UV = ViewportUV;
-				}
 			}
 			
 			return AddCodeChunk(MCT_Float4,	TEXT("MobileSceneTextureLookup(Parameters, %d, %s)"), (int32)SceneTextureId, *CoerceParameter(UV, MCT_Float2));
diff --git a/Engine/Source/Runtime/Renderer/Private/PostProcess/PostProcessMaterial.cpp b/Engine/Source/Runtime/Renderer/Private/PostProcess/PostProcessMaterial.cpp
--- a/Engine/Source/Runtime/Renderer/Private/PostProcess/PostProcessMaterial.cpp
+++ b/Engine/Source/Runtime/Renderer/Private/PostProcess/PostProcessMaterial.cpp
@@ -177,11 +177,13 @@ private:
 
 typedef FPostProcessMaterialPS<EPostProcessMaterialTarget::HighEnd, 0> FFPostProcessMaterialPS_HighEnd0;
 typedef FPostProcessMaterialPS<EPostProcessMaterialTarget::HighEnd, 1> FFPostProcessMaterialPS_HighEnd1;
-typedef FPostProcessMaterialPS<EPostProcessMaterialTarget::Mobile, 0> FPostProcessMaterialPS_Mobile;
+typedef FPostProcessMaterialPS<EPostProcessMaterialTarget::Mobile, 0> FPostProcessMaterialPS_Mobile0;
+typedef FPostProcessMaterialPS<EPostProcessMaterialTarget::Mobile, 1> FPostProcessMaterialPS_Mobile1;
 
 IMPLEMENT_MATERIAL_SHADER_TYPE(template<>, FFPostProcessMaterialPS_HighEnd0, TEXT("/Engine/Private/PostProcessMaterialShaders.usf"), TEXT("MainPS"), SF_Pixel);
 IMPLEMENT_MATERIAL_SHADER_TYPE(template<>, FFPostProcessMaterialPS_HighEnd1, TEXT("/Engine/Private/PostProcessMaterialShaders.usf"), TEXT("MainPS"), SF_Pixel);
-IMPLEMENT_MATERIAL_SHADER_TYPE(template<>,FPostProcessMaterialPS_Mobile,TEXT("/Engine/Private/PostProcessMaterialShaders.usf"),TEXT("MainPS_ES2"),SF_Pixel);
+IMPLEMENT_MATERIAL_SHADER_TYPE(template<>, FPostProcessMaterialPS_Mobile0, TEXT("/Engine/Private/PostProcessMaterialShaders.usf"), TEXT("MainPS_ES2"), SF_Pixel);
+IMPLEMENT_MATERIAL_SHADER_TYPE(template<>, FPostProcessMaterialPS_Mobile1, TEXT("/Engine/Private/PostProcessMaterialShaders.usf"), TEXT("MainPS_ES2"), SF_Pixel);
 
 FRCPassPostProcessMaterial::FRCPassPostProcessMaterial(UMaterialInterface* InMaterialInterface, ERHIFeatureLevel::Type InFeatureLevel, EPixelFormat OutputFormatIN)
 : MaterialInterface(InMaterialInterface), OutputFormat(OutputFormatIN)
@@ -228,6 +230,23 @@ public:
 };
 TGlobalResource<FPostProcessMaterialVertexDeclaration> GPostProcessMaterialVertexDeclaration;
 
+template<typename TPixelShader>
+FShader* SetMobileShaders(const FMaterialShaderMap* MaterialShaderMap, FGraphicsPipelineStateInitializer &GraphicsPSOInit, FRenderingCompositePassContext &Context, FMaterialRenderProxy* Proxy)
+{
+	TPixelShader* PixelShader_Mobile = MaterialShaderMap->GetShader<TPixelShader>();
+	FPostProcessMaterialVS_Mobile* VertexShader_Mobile = MaterialShaderMap->GetShader<FPostProcessMaterialVS_Mobile>();
+
+	GraphicsPSOInit.BoundShaderState.VertexDeclarationRHI = GFilterVertexDeclaration.VertexDeclarationRHI;
+	GraphicsPSOInit.BoundShaderState.VertexShaderRHI = GETSAFERHISHADER_VERTEX(VertexShader_Mobile);
+	GraphicsPSOInit.BoundShaderState.PixelShaderRHI = GETSAFERHISHADER_PIXEL(PixelShader_Mobile);
+
+	SetGraphicsPipelineState(Context.RHICmdList, GraphicsPSOInit);
+
+	VertexShader_Mobile->SetParameters(Context.RHICmdList, Context, Proxy);
+	PixelShader_Mobile->SetParameters(Context.RHICmdList, Context, Proxy);
+	return VertexShader_Mobile;
+}
+
 void FRCPassPostProcessMaterial::Process(FRenderingCompositePassContext& Context)
 {
 	FMaterialRenderProxy* Proxy = MaterialInterface->GetRenderProxy();
@@ -421,25 +440,22 @@ void FRCPassPostProcessMaterial::Process(FRenderingCompositePassContext& Context
 
 		FShader* VertexShader = nullptr;
 
-		// uses mobile's post process material.
+		const bool bViewSizeMatchesBufferSize = (View.ViewRect == Context.SceneColorViewRect && View.ViewRect.Size() == SrcSize && View.ViewRect.Min == FIntPoint::ZeroValue);
 		if (FeatureLevel <= ERHIFeatureLevel::ES3_1)
 		{
-			FPostProcessMaterialPS_Mobile* PixelShader_Mobile = MaterialShaderMap->GetShader<FPostProcessMaterialPS_Mobile>();
-			FPostProcessMaterialVS_Mobile* VertexShader_Mobile = MaterialShaderMap->GetShader<FPostProcessMaterialVS_Mobile>();
-
-			GraphicsPSOInit.BoundShaderState.VertexDeclarationRHI = GFilterVertexDeclaration.VertexDeclarationRHI;
-			GraphicsPSOInit.BoundShaderState.VertexShaderRHI = GETSAFERHISHADER_VERTEX(VertexShader_Mobile);
-			GraphicsPSOInit.BoundShaderState.PixelShaderRHI = GETSAFERHISHADER_PIXEL(PixelShader_Mobile);
-
-			SetGraphicsPipelineState(Context.RHICmdList, GraphicsPSOInit);
+			// use mobile's post process material.
+			if (bViewSizeMatchesBufferSize)
+			{
+				VertexShader = SetMobileShaders< FPostProcessMaterialPS_Mobile0>(MaterialShaderMap, GraphicsPSOInit, Context, Proxy);
+			}
+			else
+			{
+				VertexShader = SetMobileShaders< FPostProcessMaterialPS_Mobile1>(MaterialShaderMap, GraphicsPSOInit, Context, Proxy);
+			}
 			Context.RHICmdList.SetStencilRef(StencilRefValue);
-
-			VertexShader_Mobile->SetParameters(Context.RHICmdList, Context, Proxy);
-			PixelShader_Mobile->SetParameters(Context.RHICmdList, Context, Proxy);
-			VertexShader = VertexShader_Mobile;
 		}
 		// Uses highend post process material that assumed ViewSize == BufferSize.
-		else if (View.ViewRect == Context.SceneColorViewRect && View.ViewRect.Size() == SrcSize && View.ViewRect.Min == FIntPoint::ZeroValue)
+		else if (bViewSizeMatchesBufferSize)
 		{
 			FFPostProcessMaterialPS_HighEnd0* PixelShader_HighEnd = MaterialShaderMap->GetShader<FFPostProcessMaterialPS_HighEnd0>();
 			FPostProcessMaterialVS_HighEnd* VertexShader_HighEnd = MaterialShaderMap->GetShader<FPostProcessMaterialVS_HighEnd>();
