RL - Fix missing SafeRelease on SkyLightHistory

diff --git a/Engine/Source/Runtime/Renderer/Private/SceneRendering.h b/Engine/Source/Runtime/Renderer/Private/SceneRendering.h
--- a/Engine/Source/Runtime/Renderer/Private/SceneRendering.h
+++ b/Engine/Source/Runtime/Renderer/Private/SceneRendering.h
@@ -685,6 +685,7 @@ struct FPreviousViewInfo
 		ReflectionsHistory.SafeRelease();
 		AmbientOcclusionHistory.SafeRelease();
 		GlobalIlluminationHistory.SafeRelease();
+		SkyLightHistory.SafeRelease();
 		ShadowHistories.Reset();
 	}
 };
