Fixed ray tracing mesh instance gathering, which was broken due to ComputeLODForMeshes() signature change in CL 4916117.

An extra int8 parameter was added to ComputeLODForMeshes, however due to C4244 being explicitly silenced and presence of default arguments old usage of the function compiled without errors or warnings. Invalid LOD index was returned as the result and none of the mesh instances were added to the RT scene.

diff --git a/Engine/Source/Runtime/Renderer/Private/DeferredShadingRenderer.cpp b/Engine/Source/Runtime/Renderer/Private/DeferredShadingRenderer.cpp
--- a/Engine/Source/Runtime/Renderer/Private/DeferredShadingRenderer.cpp
+++ b/Engine/Source/Runtime/Renderer/Private/DeferredShadingRenderer.cpp
@@ -624,18 +624,24 @@ bool FDeferredShadingSceneRenderer::GatherRayTracingWorldInstances(FRHICommandLi
 						float LODScale = ICVarStaticMeshLODDistanceScale->GetFloat() * View.LODDistanceFactor;
 
 						const FPrimitiveBounds& Bounds = Scene->PrimitiveBounds[PrimitiveIndex];
+						const FPrimitiveSceneInfo* RESTRICT PrimitiveSceneInfo = Scene->Primitives[PrimitiveIndex];
 
 						FLODMask LODToRender;
+
+						const int8 CurFirstLODIdx = PrimitiveSceneInfo->Proxy->GetCurrentFirstLODIdx_RenderThread();
+						check(CurFirstLODIdx >= 0);
+
 						float MeshScreenSizeSquared = 0;
 						int32 ForcedLODLevel = GetCVarForceLOD();
 						if (SceneInfo->bIsUsingCustomLODRules)
 						{
 							FPrimitiveSceneProxy* SceneProxy = Scene->PrimitiveSceneProxies[PrimitiveIndex];
 							LODToRender = SceneProxy->GetCustomLOD(View, View.LODDistanceFactor, ForcedLODLevel, MeshScreenSizeSquared);
+							LODToRender.ClampToFirstLOD(CurFirstLODIdx);
 						}
 						else
 						{
-							LODToRender = ComputeLODForMeshes(SceneInfo->StaticMeshRelevances, View, Bounds.BoxSphereBounds.Origin, Bounds.BoxSphereBounds.SphereRadius, ForcedLODLevel, MeshScreenSizeSquared, LODScale, false);
+							LODToRender = ComputeLODForMeshes(SceneInfo->StaticMeshRelevances, View, Bounds.BoxSphereBounds.Origin, Bounds.BoxSphereBounds.SphereRadius, ForcedLODLevel, MeshScreenSizeSquared, CurFirstLODIdx, LODScale, false);
 						}
 
 						FRayTracingGeometryRHIRef RayTracingGeometryInstance = SceneInfo->GetStaticRayTracingGeometryInstance(LODToRender.GetRayTracedLOD());
