Fixed a UpdateEarlyZPassMode race condition by moving all UpdateEarlyZPassMode calls to main thread. Render thread could be changing early z pass mode, while main thread was reading this value inside FDeferredShadingSceneRenderer constructor. Thanks Richard Wallis for finding it!

diff --git a/Engine/Source/Runtime/Renderer/Private/RendererScene.cpp b/Engine/Source/Runtime/Renderer/Private/RendererScene.cpp
--- a/Engine/Source/Runtime/Renderer/Private/RendererScene.cpp
+++ b/Engine/Source/Runtime/Renderer/Private/RendererScene.cpp
@@ -577,6 +577,34 @@ static void DoLazyStaticMeshUpdateCVarSinkFunction()
 
 static FAutoConsoleVariableSink CVarDoLazyStaticMeshUpdateSink(FConsoleCommandDelegate::CreateStatic(&DoLazyStaticMeshUpdateCVarSinkFunction));
 
+static void UpdateEarlyZPassModeCVarSinkFunction()
+{
+	static int32 CachedEarlyZPass = CVarEarlyZPass.GetValueOnGameThread();
+	static int32 CachedBasePassWriteDepthEvenWithFullPrepass = CVarBasePassWriteDepthEvenWithFullPrepass.GetValueOnGameThread();
+
+	const int32 EarlyZPass = CVarEarlyZPass.GetValueOnGameThread();
+	const int32 BasePassWriteDepthEvenWithFullPrepass = CVarBasePassWriteDepthEvenWithFullPrepass.GetValueOnGameThread();
+
+	if (EarlyZPass != CachedEarlyZPass
+		|| BasePassWriteDepthEvenWithFullPrepass != CachedBasePassWriteDepthEvenWithFullPrepass)
+	{
+		for (TObjectIterator<UWorld> It; It; ++It)
+		{
+			UWorld* World = *It;
+			if (World && World->Scene)
+			{
+				FScene* Scene = (FScene*)(World->Scene);
+				Scene->UpdateEarlyZPassMode();
+			}
+		}
+
+		CachedEarlyZPass = EarlyZPass;
+		CachedBasePassWriteDepthEvenWithFullPrepass = BasePassWriteDepthEvenWithFullPrepass;
+	}
+}
+
+static FAutoConsoleVariableSink CVarUpdateEarlyZPassModeSink(FConsoleCommandDelegate::CreateStatic(&UpdateEarlyZPassModeCVarSinkFunction));
+
 void FScene::UpdateDoLazyStaticMeshUpdate(FRHICommandListImmediate& CmdList)
 {
 	bool DoLazyStaticMeshUpdate = CVarDoLazyStaticMeshUpdate.GetValueOnRenderThread() && !GIsEditor && FApp::CanEverRender();
@@ -3027,6 +3055,8 @@ bool ShouldForceFullDepthPass(EShaderPlatform ShaderPlatform)
 
 void FScene::UpdateEarlyZPassMode()
 {
+	checkSlow(IsInGameThread());
+
 	DefaultBasePassDepthStencilAccess = FExclusiveDepthStencil::DepthWrite_StencilWrite;
 	EarlyZPassMode = DDM_NonMaskedOnly;
 	bEarlyZPassMovable = false;
@@ -3064,8 +3094,6 @@ void FScene::UpdateEarlyZPassMode()
 
 void FScene::ConditionalMarkStaticMeshElementsForUpdate()
 {
-	UpdateEarlyZPassMode();
-
 	if (bScenesPrimitivesNeedStaticMeshElementUpdate
 		|| CachedDefaultBasePassDepthStencilAccess != DefaultBasePassDepthStencilAccess)
 	{
