Canvas rendering now uses MPC uniform buffers correctly.

diff --git a/Engine/Source/Runtime/Engine/Private/TileRendering.cpp b/Engine/Source/Runtime/Engine/Private/TileRendering.cpp
--- a/Engine/Source/Runtime/Engine/Private/TileRendering.cpp
+++ b/Engine/Source/Runtime/Engine/Private/TileRendering.cpp
@@ -368,7 +368,7 @@ bool FCanvasTileRendererItem::Render_RenderThread(FRHICommandListImmediate& RHIC
 	FSceneViewFamily* ViewFamily = new FSceneViewFamily(FSceneViewFamily::ConstructionValues(
 		CanvasRenderTarget,
 		nullptr,
-		FEngineShowFlags(ESFIM_Game))
+		FEngineShowFlags(ESFIM_Canvas))
 		.SetWorldTimes(CurrentWorldTime, DeltaWorldTime, CurrentRealTime)
 		.SetGammaCorrection(CanvasRenderTarget->GetDisplayGamma()));
 
@@ -438,7 +438,7 @@ bool FCanvasTileRendererItem::Render_GameThread(const FCanvas* Canvas, FRenderTh
 	FSceneViewFamily* ViewFamily = new FSceneViewFamily(FSceneViewFamily::ConstructionValues(
 		CanvasRenderTarget,
 		Canvas->GetScene(),
-		FEngineShowFlags(ESFIM_Game))
+		FEngineShowFlags(ESFIM_Canvas))
 		.SetWorldTimes(CurrentWorldTime, DeltaWorldTime, CurrentRealTime)
 		.SetGammaCorrection(CanvasRenderTarget->GetDisplayGamma()));
 
diff --git a/Engine/Source/Runtime/Engine/Public/ShowFlags.h b/Engine/Source/Runtime/Engine/Public/ShowFlags.h
--- a/Engine/Source/Runtime/Engine/Public/ShowFlags.h
+++ b/Engine/Source/Runtime/Engine/Public/ShowFlags.h
@@ -31,6 +31,7 @@ enum EShowFlagInitMode
 	ESFIM_Game,
 	ESFIM_Editor,
 	ESFIM_VREditing,
+	ESFIM_Canvas,
 	ESFIM_All0
 };
 
@@ -396,6 +397,7 @@ private:
 		SetOcclusionMeshes(false);
 		SetPathTracing(false);
 		SetRayTracingDebug(false);
+		SetDecals(InitMode != ESFIM_Canvas);
 	}
 
 
diff --git a/Engine/Source/Runtime/Renderer/Private/BasePassRendering.cpp b/Engine/Source/Runtime/Renderer/Private/BasePassRendering.cpp
--- a/Engine/Source/Runtime/Renderer/Private/BasePassRendering.cpp
+++ b/Engine/Source/Runtime/Renderer/Private/BasePassRendering.cpp
@@ -944,7 +944,6 @@ void FBasePassMeshProcessor::Process(
 
 	FMeshPassProcessorRenderState DrawRenderState(PassDrawRenderState);
 
-	const bool bEnableReceiveDecalOutput = Scene != nullptr;
 	SetDepthStencilStateForBasePass(DrawRenderState, FeatureLevel, MeshBatch, PrimitiveSceneProxy, bEnableReceiveDecalOutput, false, nullptr);
 
 	if (bTranslucentBasePass)
@@ -1509,6 +1508,8 @@ FBasePassMeshProcessor::FBasePassMeshProcessor(const FScene* Scene, ERHIFeatureL
 	, PassDrawRenderState(InDrawRenderState)
 	, TranslucencyPassType(InTranslucencyPassType)
 	, bTranslucentBasePass(InTranslucencyPassType != ETranslucencyPass::TPT_MAX)
+	  // Defaults to true for cached mesh draw commands, which can't use view and view ptr is nullptr.
+	, bEnableReceiveDecalOutput(InViewIfDynamicMeshCommand ? (InViewIfDynamicMeshCommand->Family->EngineShowFlags.Decals == 1) : true) 
 	, EarlyZPassMode(Scene ? Scene->EarlyZPassMode : DDM_None)
 {
 }
diff --git a/Engine/Source/Runtime/Renderer/Private/BasePassRendering.h b/Engine/Source/Runtime/Renderer/Private/BasePassRendering.h
--- a/Engine/Source/Runtime/Renderer/Private/BasePassRendering.h
+++ b/Engine/Source/Runtime/Renderer/Private/BasePassRendering.h
@@ -637,6 +637,7 @@ private:
 
 	const ETranslucencyPass::Type TranslucencyPassType;
 	const bool bTranslucentBasePass;
+	const bool bEnableReceiveDecalOutput;
 	EDepthDrawingMode EarlyZPassMode;
 };
 
diff --git a/Engine/Source/Runtime/Renderer/Private/Renderer.cpp b/Engine/Source/Runtime/Renderer/Private/Renderer.cpp
--- a/Engine/Source/Runtime/Renderer/Private/Renderer.cpp
+++ b/Engine/Source/Runtime/Renderer/Private/Renderer.cpp
@@ -96,9 +96,11 @@ void FRendererModule::DrawTileMesh(FRHICommandListImmediate& RHICmdList, FMeshPa
 		//@todo - reuse this view for multiple tiles, this is going to be slow for each tile
 		FViewInfo View(&SceneView);
 		View.ViewRect = View.UnscaledViewRect;
-		
+
 		const auto FeatureLevel = View.GetFeatureLevel();
 		const EShadingPath ShadingPath = FSceneInterface::GetShadingPath(FeatureLevel);
+		const FSceneViewFamily* ViewFamily = View.Family;
+		const FScene* Scene = ViewFamily->Scene->GetRenderScene();
 
 		Mesh.MaterialRenderProxy->UpdateUniformExpressionCacheIfNeeded(FeatureLevel);
 		FMaterialRenderProxy::UpdateDeferredCachedUniformExpressions();
@@ -157,7 +159,7 @@ void FRendererModule::DrawTileMesh(FRHICommandListImmediate& RHICmdList, FMeshPa
 		FMemMark Mark(FMemStack::Get());
 
 		// handle translucent material blend modes, not relevant in MaterialTexCoordScalesAnalysis since it outputs the scales.
-		if (View.Family->GetDebugViewShaderMode() == DVSM_OutputMaterialTextureScales)
+		if (ViewFamily->GetDebugViewShaderMode() == DVSM_OutputMaterialTextureScales)
 		{
 #if !(UE_BUILD_SHIPPING || UE_BUILD_TEST)
 			// make sure we are doing opaque drawing
@@ -175,9 +177,15 @@ void FRendererModule::DrawTileMesh(FRHICommandListImmediate& RHICmdList, FMeshPa
 				);
 
 				DrawDynamicMeshPass(View, RHICmdList,
-					[&View, &DrawRenderState, &DebugViewModePassUniformBuffer, &Mesh](FMeshPassDrawListContext* InDrawListContext)
+					[Scene, &View, &DrawRenderState, &DebugViewModePassUniformBuffer, &Mesh](FMeshPassDrawListContext* InDrawListContext)
 				{
-					FDebugViewModeMeshProcessor PassMeshProcessor(nullptr, View.GetFeatureLevel(), &View, DebugViewModePassUniformBuffer, false, InDrawListContext);
+					FDebugViewModeMeshProcessor PassMeshProcessor(
+						Scene,
+						View.GetFeatureLevel(),
+						&View,
+						DebugViewModePassUniformBuffer,
+						false,
+						InDrawListContext);
 					const uint64 DefaultBatchElementMask = ~0ull;
 					PassMeshProcessor.AddMeshBatch(Mesh, DefaultBatchElementMask, nullptr);
 				});
@@ -203,15 +211,15 @@ void FRendererModule::DrawTileMesh(FRHICommandListImmediate& RHICmdList, FMeshPa
 				DrawRenderState.SetPassUniformBuffer(BasePassUniformBuffer);
 				
 				DrawDynamicMeshPass(View, RHICmdList,
-					[&View, &DrawRenderState, &Mesh](FDynamicPassMeshDrawListContext* DynamicMeshPassContext)
+					[Scene, &View, &DrawRenderState, &Mesh](FDynamicPassMeshDrawListContext* DynamicMeshPassContext)
 				{
 					FBasePassMeshProcessor PassMeshProcessor(
-															 nullptr,
-															 View.GetFeatureLevel(),
-															 &View,
-															 DrawRenderState,
-															 DynamicMeshPassContext,
-															 ETranslucencyPass::TPT_AllTranslucency);
+						Scene,
+						View.GetFeatureLevel(),
+						&View,
+						DrawRenderState,
+						DynamicMeshPassContext,
+						ETranslucencyPass::TPT_AllTranslucency);
 					
 					const uint64 DefaultBatchElementMask = ~0ull;
 					PassMeshProcessor.AddMeshBatch(Mesh, DefaultBatchElementMask, nullptr);
@@ -226,23 +234,21 @@ void FRendererModule::DrawTileMesh(FRHICommandListImmediate& RHICmdList, FMeshPa
 				DrawRenderState.SetPassUniformBuffer(BasePassUniformBuffer);
 				
 				DrawDynamicMeshPass(View, RHICmdList,
-					[&View, &DrawRenderState, &Mesh](FDynamicPassMeshDrawListContext* DynamicMeshPassContext)
+					[Scene, &View, &DrawRenderState, &Mesh](FDynamicPassMeshDrawListContext* DynamicMeshPassContext)
 				{
 					FMobileBasePassMeshProcessor PassMeshProcessor(
-															 nullptr,
-															 View.GetFeatureLevel(),
-															 &View,
-															 DrawRenderState,
-															 DynamicMeshPassContext,
-															 false,
-															 ETranslucencyPass::TPT_AllTranslucency); // No CSM?
+						Scene,
+						View.GetFeatureLevel(),
+						&View,
+						DrawRenderState,
+						DynamicMeshPassContext,
+						false,
+						ETranslucencyPass::TPT_AllTranslucency); // No CSM?
 					
 					const uint64 DefaultBatchElementMask = ~0ull;
 					PassMeshProcessor.AddMeshBatch(Mesh, DefaultBatchElementMask, nullptr);
 				});
 			}
-
-			
 		}
 		// handle opaque materials
 		else
@@ -257,10 +263,10 @@ void FRendererModule::DrawTileMesh(FRHICommandListImmediate& RHICmdList, FMeshPa
 
 #if WITH_EDITOR
 				DrawDynamicMeshPass(View, RHICmdList,
-					[&View, &DrawRenderState, &Mesh](FDynamicPassMeshDrawListContext* DynamicMeshPassContext)
+					[Scene, &View, &DrawRenderState, &Mesh](FDynamicPassMeshDrawListContext* DynamicMeshPassContext)
 				{
 					FHitProxyMeshProcessor PassMeshProcessor(
-						nullptr,
+						Scene,
 						&View,
 						false,
 						DrawRenderState,
@@ -284,14 +290,14 @@ void FRendererModule::DrawTileMesh(FRHICommandListImmediate& RHICmdList, FMeshPa
 					DrawRenderState.SetPassUniformBuffer(BasePassUniformBuffer);
 					
 					DrawDynamicMeshPass(View, RHICmdList,
-						[&View, &DrawRenderState, &Mesh](FDynamicPassMeshDrawListContext* DynamicMeshPassContext)
+						[Scene, &View, &DrawRenderState, &Mesh](FDynamicPassMeshDrawListContext* DynamicMeshPassContext)
 					{
 						FBasePassMeshProcessor PassMeshProcessor(
-																 nullptr,
-																 View.GetFeatureLevel(),
-																 &View,
-																 DrawRenderState,
-																 DynamicMeshPassContext);
+							Scene,
+							View.GetFeatureLevel(),
+							&View,
+							DrawRenderState,
+							DynamicMeshPassContext);
 						
 						const uint64 DefaultBatchElementMask = ~0ull;
 						PassMeshProcessor.AddMeshBatch(Mesh, DefaultBatchElementMask, nullptr);
@@ -306,21 +312,19 @@ void FRendererModule::DrawTileMesh(FRHICommandListImmediate& RHICmdList, FMeshPa
 					DrawRenderState.SetPassUniformBuffer(BasePassUniformBuffer);
 					
 					DrawDynamicMeshPass(View, RHICmdList,
-						[&View, &DrawRenderState, &Mesh](FDynamicPassMeshDrawListContext* DynamicMeshPassContext)
+						[Scene, &View, &DrawRenderState, &Mesh](FDynamicPassMeshDrawListContext* DynamicMeshPassContext)
 					{
 						FMobileBasePassMeshProcessor PassMeshProcessor(
-																 nullptr,
-																 View.GetFeatureLevel(),
-																 &View,
-																 DrawRenderState,
-																 DynamicMeshPassContext, true);
+							Scene,
+							View.GetFeatureLevel(),
+							&View,
+							DrawRenderState,
+							DynamicMeshPassContext, true);
 						
 						const uint64 DefaultBatchElementMask = ~0ull;
 						PassMeshProcessor.AddMeshBatch(Mesh, DefaultBatchElementMask, nullptr);
 					});
 				}
-
-				
 			}
 		}
 
