Fix crash in renderer

diff --git a/Engine/Source/Runtime/Renderer/Private/Renderer.cpp b/Engine/Source/Runtime/Renderer/Private/Renderer.cpp
--- a/Engine/Source/Runtime/Renderer/Private/Renderer.cpp
+++ b/Engine/Source/Runtime/Renderer/Private/Renderer.cpp
@@ -100,7 +100,7 @@ void FRendererModule::DrawTileMesh(FRHICommandListImmediate& RHICmdList, FMeshPa
 		const auto FeatureLevel = View.GetFeatureLevel();
 		const EShadingPath ShadingPath = FSceneInterface::GetShadingPath(FeatureLevel);
 		const FSceneViewFamily* ViewFamily = View.Family;
-		const FScene* Scene = ViewFamily->Scene->GetRenderScene();
+		const FScene* Scene = ViewFamily->Scene ? ViewFamily->Scene->GetRenderScene() : nullptr;
 
 		Mesh.MaterialRenderProxy->UpdateUniformExpressionCacheIfNeeded(FeatureLevel);
 		FMaterialRenderProxy::UpdateDeferredCachedUniformExpressions();
