Fixed assert with depth stencil targets in canvas rendering, as well as null pointer dereference when merging actors.

#jyi dmitriy.dyomin

diff --git a/Engine/Source/Runtime/Engine/Private/TileRendering.cpp b/Engine/Source/Runtime/Engine/Private/TileRendering.cpp
--- a/Engine/Source/Runtime/Engine/Private/TileRendering.cpp
+++ b/Engine/Source/Runtime/Engine/Private/TileRendering.cpp
@@ -368,7 +368,7 @@ bool FCanvasTileRendererItem::Render_RenderThread(FRHICommandListImmediate& RHIC
 	FSceneViewFamily* ViewFamily = new FSceneViewFamily(FSceneViewFamily::ConstructionValues(
 		CanvasRenderTarget,
 		nullptr,
-		FEngineShowFlags(ESFIM_Canvas))
+		FEngineShowFlags(ESFIM_Game))
 		.SetWorldTimes(CurrentWorldTime, DeltaWorldTime, CurrentRealTime)
 		.SetGammaCorrection(CanvasRenderTarget->GetDisplayGamma()));
 
@@ -438,7 +438,7 @@ bool FCanvasTileRendererItem::Render_GameThread(const FCanvas* Canvas, FRenderTh
 	FSceneViewFamily* ViewFamily = new FSceneViewFamily(FSceneViewFamily::ConstructionValues(
 		CanvasRenderTarget,
 		Canvas->GetScene(),
-		FEngineShowFlags(ESFIM_Canvas))
+		FEngineShowFlags(ESFIM_Game))
 		.SetWorldTimes(CurrentWorldTime, DeltaWorldTime, CurrentRealTime)
 		.SetGammaCorrection(CanvasRenderTarget->GetDisplayGamma()));
 
diff --git a/Engine/Source/Runtime/Engine/Public/ShowFlags.h b/Engine/Source/Runtime/Engine/Public/ShowFlags.h
--- a/Engine/Source/Runtime/Engine/Public/ShowFlags.h
+++ b/Engine/Source/Runtime/Engine/Public/ShowFlags.h
@@ -31,7 +31,6 @@ enum EShowFlagInitMode
 	ESFIM_Game,
 	ESFIM_Editor,
 	ESFIM_VREditing,
-	ESFIM_Canvas,
 	ESFIM_All0
 };
 
@@ -397,7 +396,6 @@ private:
 		SetOcclusionMeshes(false);
 		SetPathTracing(false);
 		SetRayTracingDebug(false);
-		SetDecals(InitMode != ESFIM_Canvas);
 	}
 
 
diff --git a/Engine/Source/Runtime/Renderer/Private/BasePassRendering.cpp b/Engine/Source/Runtime/Renderer/Private/BasePassRendering.cpp
--- a/Engine/Source/Runtime/Renderer/Private/BasePassRendering.cpp
+++ b/Engine/Source/Runtime/Renderer/Private/BasePassRendering.cpp
@@ -1503,13 +1503,19 @@ void FBasePassMeshProcessor::AddMeshBatch(const FMeshBatch& RESTRICT MeshBatch,
 	}
 }
 
-FBasePassMeshProcessor::FBasePassMeshProcessor(const FScene* Scene, ERHIFeatureLevel::Type InFeatureLevel, const FSceneView* InViewIfDynamicMeshCommand, const FMeshPassProcessorRenderState& InDrawRenderState, FMeshPassDrawListContext* InDrawListContext, ETranslucencyPass::Type InTranslucencyPassType)
+FBasePassMeshProcessor::FBasePassMeshProcessor(
+	const FScene* Scene,
+	ERHIFeatureLevel::Type InFeatureLevel,
+	const FSceneView* InViewIfDynamicMeshCommand,
+	const FMeshPassProcessorRenderState& InDrawRenderState,
+	FMeshPassDrawListContext* InDrawListContext,
+	EFlags Flags,
+	ETranslucencyPass::Type InTranslucencyPassType)
 	: FMeshPassProcessor(Scene, InFeatureLevel, InViewIfDynamicMeshCommand, InDrawListContext)
 	, PassDrawRenderState(InDrawRenderState)
 	, TranslucencyPassType(InTranslucencyPassType)
 	, bTranslucentBasePass(InTranslucencyPassType != ETranslucencyPass::TPT_MAX)
-	  // Defaults to true for cached mesh draw commands, which can't use view and view ptr is nullptr.
-	, bEnableReceiveDecalOutput(InViewIfDynamicMeshCommand ? (InViewIfDynamicMeshCommand->Family->EngineShowFlags.Decals == 1) : true) 
+	, bEnableReceiveDecalOutput((Flags & EFlags::CanUseDepthStencil) == EFlags::CanUseDepthStencil)
 	, EarlyZPassMode(Scene ? Scene->EarlyZPassMode : DDM_None)
 {
 }
@@ -1520,7 +1526,9 @@ FMeshPassProcessor* CreateBasePassProcessor(const FScene* Scene, const FSceneVie
 	PassDrawRenderState.SetInstancedViewUniformBuffer(Scene->UniformBuffers.InstancedViewUniformBuffer);
 	SetupBasePassState(Scene->DefaultBasePassDepthStencilAccess, false, PassDrawRenderState);
 
-	return new(FMemStack::Get()) FBasePassMeshProcessor(Scene, Scene->GetFeatureLevel(), InViewIfDynamicMeshCommand, PassDrawRenderState, InDrawListContext);
+	const FBasePassMeshProcessor::EFlags Flags = FBasePassMeshProcessor::EFlags::CanUseDepthStencil;
+
+	return new(FMemStack::Get()) FBasePassMeshProcessor(Scene, Scene->GetFeatureLevel(), InViewIfDynamicMeshCommand, PassDrawRenderState, InDrawListContext, Flags);
 }
 
 FMeshPassProcessor* CreateTranslucencyStandardPassProcessor(const FScene* Scene, const FSceneView* InViewIfDynamicMeshCommand, FMeshPassDrawListContext* InDrawListContext)
@@ -1529,7 +1537,9 @@ FMeshPassProcessor* CreateTranslucencyStandardPassProcessor(const FScene* Scene,
 	PassDrawRenderState.SetInstancedViewUniformBuffer(Scene->UniformBuffers.InstancedViewUniformBuffer);
 	PassDrawRenderState.SetDepthStencilState(TStaticDepthStencilState<false, CF_DepthNearOrEqual>::GetRHI());
 
-	return new(FMemStack::Get()) FBasePassMeshProcessor(Scene, Scene->GetFeatureLevel(), InViewIfDynamicMeshCommand, PassDrawRenderState, InDrawListContext, ETranslucencyPass::TPT_StandardTranslucency);
+	const FBasePassMeshProcessor::EFlags Flags = FBasePassMeshProcessor::EFlags::CanUseDepthStencil;
+
+	return new(FMemStack::Get()) FBasePassMeshProcessor(Scene, Scene->GetFeatureLevel(), InViewIfDynamicMeshCommand, PassDrawRenderState, InDrawListContext, Flags, ETranslucencyPass::TPT_StandardTranslucency);
 }
 
 FMeshPassProcessor* CreateTranslucencyAfterDOFProcessor(const FScene* Scene, const FSceneView* InViewIfDynamicMeshCommand, FMeshPassDrawListContext* InDrawListContext)
@@ -1538,7 +1548,9 @@ FMeshPassProcessor* CreateTranslucencyAfterDOFProcessor(const FScene* Scene, con
 	PassDrawRenderState.SetInstancedViewUniformBuffer(Scene->UniformBuffers.InstancedViewUniformBuffer);
 	PassDrawRenderState.SetDepthStencilState(TStaticDepthStencilState<false, CF_DepthNearOrEqual>::GetRHI());
 
-	return new(FMemStack::Get()) FBasePassMeshProcessor(Scene, Scene->GetFeatureLevel(), InViewIfDynamicMeshCommand, PassDrawRenderState, InDrawListContext, ETranslucencyPass::TPT_TranslucencyAfterDOF);
+	const FBasePassMeshProcessor::EFlags Flags = FBasePassMeshProcessor::EFlags::CanUseDepthStencil;
+
+	return new(FMemStack::Get()) FBasePassMeshProcessor(Scene, Scene->GetFeatureLevel(), InViewIfDynamicMeshCommand, PassDrawRenderState, InDrawListContext, Flags, ETranslucencyPass::TPT_TranslucencyAfterDOF);
 }
 
 FMeshPassProcessor* CreateTranslucencyAllPassProcessor(const FScene* Scene, const FSceneView* InViewIfDynamicMeshCommand, FMeshPassDrawListContext* InDrawListContext)
@@ -1547,7 +1559,9 @@ FMeshPassProcessor* CreateTranslucencyAllPassProcessor(const FScene* Scene, cons
 	PassDrawRenderState.SetInstancedViewUniformBuffer(Scene->UniformBuffers.InstancedViewUniformBuffer);
 	PassDrawRenderState.SetDepthStencilState(TStaticDepthStencilState<false, CF_DepthNearOrEqual>::GetRHI());
 
-	return new(FMemStack::Get()) FBasePassMeshProcessor(Scene, Scene->GetFeatureLevel(), InViewIfDynamicMeshCommand, PassDrawRenderState, InDrawListContext, ETranslucencyPass::TPT_AllTranslucency);
+	const FBasePassMeshProcessor::EFlags Flags = FBasePassMeshProcessor::EFlags::CanUseDepthStencil;
+
+	return new(FMemStack::Get()) FBasePassMeshProcessor(Scene, Scene->GetFeatureLevel(), InViewIfDynamicMeshCommand, PassDrawRenderState, InDrawListContext, Flags, ETranslucencyPass::TPT_AllTranslucency);
 }
 
 FRegisterPassProcessorCreateFunction RegisterBasePass(&CreateBasePassProcessor, EShadingPath::Deferred, EMeshPass::BasePass, EMeshPassFlags::CachedMeshCommands | EMeshPassFlags::MainView);
diff --git a/Engine/Source/Runtime/Renderer/Private/BasePassRendering.h b/Engine/Source/Runtime/Renderer/Private/BasePassRendering.h
--- a/Engine/Source/Runtime/Renderer/Private/BasePassRendering.h
+++ b/Engine/Source/Runtime/Renderer/Private/BasePassRendering.h
@@ -596,8 +596,22 @@ void GetBasePassShaders<FUniformLightMapPolicy>(
 class FBasePassMeshProcessor : public FMeshPassProcessor
 {
 public:
+	enum class EFlags
+	{
+		None = 0,
+
+		// Informs the processor whether a depth-stencil target is bound when processed draw commands are issued.
+		CanUseDepthStencil = (1 << 0)
+	};
 
-	FBasePassMeshProcessor(const FScene* InScene, ERHIFeatureLevel::Type InFeatureLevel, const FSceneView* InViewIfDynamicMeshCommand, const FMeshPassProcessorRenderState& InDrawRenderState, FMeshPassDrawListContext* InDrawListContext, ETranslucencyPass::Type InTranslucencyPassType = ETranslucencyPass::TPT_MAX);
+	FBasePassMeshProcessor(
+		const FScene* InScene,
+		ERHIFeatureLevel::Type InFeatureLevel,
+		const FSceneView* InViewIfDynamicMeshCommand,
+		const FMeshPassProcessorRenderState& InDrawRenderState,
+		FMeshPassDrawListContext* InDrawListContext,
+		EFlags Flags,
+		ETranslucencyPass::Type InTranslucencyPassType = ETranslucencyPass::TPT_MAX);
 
 	virtual void AddMeshBatch(const FMeshBatch& RESTRICT MeshBatch, uint64 BatchElementMask, const FPrimitiveSceneProxy* RESTRICT PrimitiveSceneProxy, int32 StaticMeshId = -1) override final;
 
@@ -641,6 +655,8 @@ private:
 	EDepthDrawingMode EarlyZPassMode;
 };
 
+ENUM_CLASS_FLAGS(FBasePassMeshProcessor::EFlags);
+
 extern void SetDepthStencilStateForBasePass(FMeshPassProcessorRenderState& DrawRenderState, ERHIFeatureLevel::Type FeatureLevel, const FMeshBatch& Mesh, const FPrimitiveSceneProxy* PrimitiveSceneProxy, bool bEnableReceiveDecalOutput, bool bUseDebugViewPS, FDepthStencilStateRHIParamRef LodFadeOverrideDepthStencilState);
 extern void SetupBasePassState(FExclusiveDepthStencil::Type BasePassDepthStencilAccess, const bool bShaderComplexity, FMeshPassProcessorRenderState& DrawRenderState);
 extern FMeshDrawCommandSortKey CalculateTranslucentMeshStaticSortKey(const FPrimitiveSceneProxy* RESTRICT PrimitiveSceneProxy, uint16 MeshIdInPrimitive);
\ No newline at end of file
diff --git a/Engine/Source/Runtime/Renderer/Private/MobileBasePass.cpp b/Engine/Source/Runtime/Renderer/Private/MobileBasePass.cpp
--- a/Engine/Source/Runtime/Renderer/Private/MobileBasePass.cpp
+++ b/Engine/Source/Runtime/Renderer/Private/MobileBasePass.cpp
@@ -550,12 +550,20 @@ void TMobileBasePassPSPolicyParamType<FUniformLightMapPolicy>::GetShaderBindings
 	}
 }
 
-FMobileBasePassMeshProcessor::FMobileBasePassMeshProcessor(const FScene* Scene, ERHIFeatureLevel::Type InFeatureLevel, const FSceneView* InViewIfDynamicMeshCommand, const FMeshPassProcessorRenderState& InDrawRenderState, FMeshPassDrawListContext* InDrawListContext, bool bInCanReceiveCSM, ETranslucencyPass::Type InTranslucencyPassType)
+FMobileBasePassMeshProcessor::FMobileBasePassMeshProcessor(
+	const FScene* Scene,
+	ERHIFeatureLevel::Type InFeatureLevel,
+	const FSceneView* InViewIfDynamicMeshCommand,
+	const FMeshPassProcessorRenderState& InDrawRenderState,
+	FMeshPassDrawListContext* InDrawListContext,
+	EFlags Flags,
+	ETranslucencyPass::Type InTranslucencyPassType)
 	: FMeshPassProcessor(Scene, InFeatureLevel, InViewIfDynamicMeshCommand, InDrawListContext)
 	, PassDrawRenderState(InDrawRenderState)
 	, TranslucencyPassType(InTranslucencyPassType)
 	, bTranslucentBasePass(InTranslucencyPassType != ETranslucencyPass::TPT_MAX)
-	, bCanReceiveCSM(bInCanReceiveCSM)
+	, bCanReceiveCSM((Flags & EFlags::CanReceiveCSM) == EFlags::CanReceiveCSM)
+	, bEnableReceiveDecalOutput((Flags & EFlags::CanUseDepthStencil) == EFlags::CanUseDepthStencil)
 {}
 
 void FMobileBasePassMeshProcessor::AddMeshBatch(const FMeshBatch& RESTRICT MeshBatch, uint64 BatchElementMask, const FPrimitiveSceneProxy* RESTRICT PrimitiveSceneProxy, int32 StaticMeshId)
@@ -637,8 +645,7 @@ void FMobileBasePassMeshProcessor::Process(
 	}
 	else
 	{
-		bool bEnableReceiveDecalOutput = IsMobileHDR();
-		MobileBasePass::SetOpaqueRenderState(DrawRenderState, PrimitiveSceneProxy, MaterialResource, bEnableReceiveDecalOutput);
+		MobileBasePass::SetOpaqueRenderState(DrawRenderState, PrimitiveSceneProxy, MaterialResource, bEnableReceiveDecalOutput && IsMobileHDR());
 	}
 
 	FMeshDrawCommandSortKey SortKey; 
@@ -684,7 +691,9 @@ FMeshPassProcessor* CreateMobileBasePassProcessor(const FScene* Scene, const FSc
 	PassDrawRenderState.SetDepthStencilAccess(Scene->DefaultBasePassDepthStencilAccess);
 	PassDrawRenderState.SetDepthStencilState(TStaticDepthStencilState<true, CF_DepthNearOrEqual>::GetRHI());
 
-	return new(FMemStack::Get()) FMobileBasePassMeshProcessor(Scene, Scene->GetFeatureLevel(), InViewIfDynamicMeshCommand, PassDrawRenderState, InDrawListContext, false);
+	const FMobileBasePassMeshProcessor::EFlags Flags = FMobileBasePassMeshProcessor::EFlags::CanUseDepthStencil;
+
+	return new(FMemStack::Get()) FMobileBasePassMeshProcessor(Scene, Scene->GetFeatureLevel(), InViewIfDynamicMeshCommand, PassDrawRenderState, InDrawListContext, Flags);
 }
 
 FMeshPassProcessor* CreateMobileBasePassCSMProcessor(const FScene* Scene, const FSceneView* InViewIfDynamicMeshCommand, FMeshPassDrawListContext* InDrawListContext)
@@ -695,7 +704,9 @@ FMeshPassProcessor* CreateMobileBasePassCSMProcessor(const FScene* Scene, const
 	PassDrawRenderState.SetDepthStencilAccess(Scene->DefaultBasePassDepthStencilAccess);
 	PassDrawRenderState.SetDepthStencilState(TStaticDepthStencilState<true, CF_DepthNearOrEqual>::GetRHI());
 
-	return new(FMemStack::Get()) FMobileBasePassMeshProcessor(Scene, Scene->GetFeatureLevel(), InViewIfDynamicMeshCommand, PassDrawRenderState, InDrawListContext, true);
+	const FMobileBasePassMeshProcessor::EFlags Flags = FMobileBasePassMeshProcessor::EFlags::CanReceiveCSM | FMobileBasePassMeshProcessor::EFlags::CanUseDepthStencil;
+
+	return new(FMemStack::Get()) FMobileBasePassMeshProcessor(Scene, Scene->GetFeatureLevel(), InViewIfDynamicMeshCommand, PassDrawRenderState, InDrawListContext, Flags);
 }
 
 FMeshPassProcessor* CreateMobileTranslucencyStandardPassProcessor(const FScene* Scene, const FSceneView* InViewIfDynamicMeshCommand, FMeshPassDrawListContext* InDrawListContext)
@@ -704,8 +715,10 @@ FMeshPassProcessor* CreateMobileTranslucencyStandardPassProcessor(const FScene*
 	PassDrawRenderState.SetInstancedViewUniformBuffer(Scene->UniformBuffers.InstancedViewUniformBuffer);
 	PassDrawRenderState.SetDepthStencilState(TStaticDepthStencilState<false, CF_DepthNearOrEqual>::GetRHI());
 	PassDrawRenderState.SetDepthStencilAccess(FExclusiveDepthStencil::DepthRead_StencilRead);
+	
+	const FMobileBasePassMeshProcessor::EFlags Flags = FMobileBasePassMeshProcessor::EFlags::CanUseDepthStencil;
 
-	return new(FMemStack::Get()) FMobileBasePassMeshProcessor(Scene, Scene->GetFeatureLevel(), InViewIfDynamicMeshCommand, PassDrawRenderState, InDrawListContext, false, ETranslucencyPass::TPT_StandardTranslucency);
+	return new(FMemStack::Get()) FMobileBasePassMeshProcessor(Scene, Scene->GetFeatureLevel(), InViewIfDynamicMeshCommand, PassDrawRenderState, InDrawListContext, Flags, ETranslucencyPass::TPT_StandardTranslucency);
 }
 
 FMeshPassProcessor* CreateMobileTranslucencyAfterDOFProcessor(const FScene* Scene, const FSceneView* InViewIfDynamicMeshCommand, FMeshPassDrawListContext* InDrawListContext)
@@ -715,7 +728,9 @@ FMeshPassProcessor* CreateMobileTranslucencyAfterDOFProcessor(const FScene* Scen
 	PassDrawRenderState.SetDepthStencilState(TStaticDepthStencilState<false, CF_DepthNearOrEqual>::GetRHI());
 	PassDrawRenderState.SetDepthStencilAccess(FExclusiveDepthStencil::DepthRead_StencilRead);
 
-	return new(FMemStack::Get()) FMobileBasePassMeshProcessor(Scene, Scene->GetFeatureLevel(), InViewIfDynamicMeshCommand, PassDrawRenderState, InDrawListContext, false, ETranslucencyPass::TPT_TranslucencyAfterDOF);
+	const FMobileBasePassMeshProcessor::EFlags Flags = FMobileBasePassMeshProcessor::EFlags::CanUseDepthStencil;
+
+	return new(FMemStack::Get()) FMobileBasePassMeshProcessor(Scene, Scene->GetFeatureLevel(), InViewIfDynamicMeshCommand, PassDrawRenderState, InDrawListContext, Flags, ETranslucencyPass::TPT_TranslucencyAfterDOF);
 }
 
 FMeshPassProcessor* CreateMobileTranslucencyAllPassProcessor(const FScene* Scene, const FSceneView* InViewIfDynamicMeshCommand, FMeshPassDrawListContext* InDrawListContext)
@@ -725,7 +740,9 @@ FMeshPassProcessor* CreateMobileTranslucencyAllPassProcessor(const FScene* Scene
 	PassDrawRenderState.SetDepthStencilState(TStaticDepthStencilState<false, CF_DepthNearOrEqual>::GetRHI());
 	PassDrawRenderState.SetDepthStencilAccess(FExclusiveDepthStencil::DepthRead_StencilRead);
 
-	return new(FMemStack::Get()) FMobileBasePassMeshProcessor(Scene, Scene->GetFeatureLevel(), InViewIfDynamicMeshCommand, PassDrawRenderState, InDrawListContext, false, ETranslucencyPass::TPT_AllTranslucency);
+	const FMobileBasePassMeshProcessor::EFlags Flags = FMobileBasePassMeshProcessor::EFlags::CanUseDepthStencil;
+
+	return new(FMemStack::Get()) FMobileBasePassMeshProcessor(Scene, Scene->GetFeatureLevel(), InViewIfDynamicMeshCommand, PassDrawRenderState, InDrawListContext, Flags, ETranslucencyPass::TPT_AllTranslucency);
 }
 
 FRegisterPassProcessorCreateFunction RegisterMobileBasePass(&CreateMobileBasePassProcessor, EShadingPath::Mobile, EMeshPass::BasePass, EMeshPassFlags::CachedMeshCommands | EMeshPassFlags::MainView);
diff --git a/Engine/Source/Runtime/Renderer/Private/MobileBasePassRendering.h b/Engine/Source/Runtime/Renderer/Private/MobileBasePassRendering.h
--- a/Engine/Source/Runtime/Renderer/Private/MobileBasePassRendering.h
+++ b/Engine/Source/Runtime/Renderer/Private/MobileBasePassRendering.h
@@ -358,7 +358,7 @@ namespace MobileBasePass
 
 	const FLightSceneInfo* GetDirectionalLightInfo(const FScene* Scene, const FPrimitiveSceneProxy* PrimitiveSceneProxy);
 	int32 CalcNumMovablePointLights(const FMaterial& InMaterial, const FPrimitiveSceneProxy* InPrimitiveSceneProxy);
-		
+
 	bool StaticCanReceiveCSM(const FLightSceneInfo* LightSceneInfo, const FPrimitiveSceneProxy* PrimitiveSceneProxy);
 
 	void SetOpaqueRenderState(FMeshPassProcessorRenderState& DrawRenderState, const FPrimitiveSceneProxy* PrimitiveSceneProxy, const FMaterial& Material, bool bEnableReceiveDecalOutput);
@@ -446,13 +446,24 @@ public:
 class FMobileBasePassMeshProcessor : public FMeshPassProcessor
 {
 public:
+	enum class EFlags
+	{
+		None = 0,
+
+		// Informs the processor whether a depth-stencil target is bound when processed draw commands are issued.
+		CanUseDepthStencil = (1 << 0),
+
+		// Informs the processor whether primitives can receive shadows from cascade shadow maps.
+		CanReceiveCSM =  (1 << 1)
+	};
+
 	FMobileBasePassMeshProcessor(
 		const FScene* InScene, 
 		ERHIFeatureLevel::Type InFeatureLevel, 
 		const FSceneView* InViewIfDynamicMeshCommand, 
 		const FMeshPassProcessorRenderState& InDrawRenderState, 
 		FMeshPassDrawListContext* InDrawListContext,
-		bool bInCanReceiveCSM,
+		EFlags Flags,
 		ETranslucencyPass::Type InTranslucencyPassType = ETranslucencyPass::TPT_MAX);
 
 	virtual void AddMeshBatch(const FMeshBatch& RESTRICT MeshBatch, uint64 BatchElementMask, const FPrimitiveSceneProxy* RESTRICT PrimitiveSceneProxy, int32 StaticMeshId = -1) override final;
@@ -475,4 +486,7 @@ private:
 	const ETranslucencyPass::Type TranslucencyPassType;
 	const bool bTranslucentBasePass;
 	const bool bCanReceiveCSM;
+	const bool bEnableReceiveDecalOutput;
 };
+
+ENUM_CLASS_FLAGS(FMobileBasePassMeshProcessor::EFlags);
diff --git a/Engine/Source/Runtime/Renderer/Private/Renderer.cpp b/Engine/Source/Runtime/Renderer/Private/Renderer.cpp
--- a/Engine/Source/Runtime/Renderer/Private/Renderer.cpp
+++ b/Engine/Source/Runtime/Renderer/Private/Renderer.cpp
@@ -100,7 +100,12 @@ void FRendererModule::DrawTileMesh(FRHICommandListImmediate& RHICmdList, FMeshPa
 		const auto FeatureLevel = View.GetFeatureLevel();
 		const EShadingPath ShadingPath = FSceneInterface::GetShadingPath(FeatureLevel);
 		const FSceneViewFamily* ViewFamily = View.Family;
-		const FScene* Scene = ViewFamily->Scene ? ViewFamily->Scene->GetRenderScene() : nullptr;
+		const FScene* Scene = nullptr;
+
+		if (ViewFamily->Scene)
+		{
+			Scene = ViewFamily->Scene->GetRenderScene();
+		}
 
 		Mesh.MaterialRenderProxy->UpdateUniformExpressionCacheIfNeeded(FeatureLevel);
 		FMaterialRenderProxy::UpdateDeferredCachedUniformExpressions();
@@ -219,6 +224,7 @@ void FRendererModule::DrawTileMesh(FRHICommandListImmediate& RHICmdList, FMeshPa
 						&View,
 						DrawRenderState,
 						DynamicMeshPassContext,
+						FBasePassMeshProcessor::EFlags::None,
 						ETranslucencyPass::TPT_AllTranslucency);
 					
 					const uint64 DefaultBatchElementMask = ~0ull;
@@ -242,8 +248,8 @@ void FRendererModule::DrawTileMesh(FRHICommandListImmediate& RHICmdList, FMeshPa
 						&View,
 						DrawRenderState,
 						DynamicMeshPassContext,
-						false,
-						ETranslucencyPass::TPT_AllTranslucency); // No CSM?
+						FMobileBasePassMeshProcessor::EFlags::None,
+						ETranslucencyPass::TPT_AllTranslucency);
 					
 					const uint64 DefaultBatchElementMask = ~0ull;
 					PassMeshProcessor.AddMeshBatch(Mesh, DefaultBatchElementMask, nullptr);
@@ -297,7 +303,8 @@ void FRendererModule::DrawTileMesh(FRHICommandListImmediate& RHICmdList, FMeshPa
 							View.GetFeatureLevel(),
 							&View,
 							DrawRenderState,
-							DynamicMeshPassContext);
+							DynamicMeshPassContext,
+							FBasePassMeshProcessor::EFlags::None);
 						
 						const uint64 DefaultBatchElementMask = ~0ull;
 						PassMeshProcessor.AddMeshBatch(Mesh, DefaultBatchElementMask, nullptr);
@@ -319,7 +326,8 @@ void FRendererModule::DrawTileMesh(FRHICommandListImmediate& RHICmdList, FMeshPa
 							View.GetFeatureLevel(),
 							&View,
 							DrawRenderState,
-							DynamicMeshPassContext, true);
+							DynamicMeshPassContext,
+							FMobileBasePassMeshProcessor::EFlags::CanReceiveCSM);
 						
 						const uint64 DefaultBatchElementMask = ~0ull;
 						PassMeshProcessor.AddMeshBatch(Mesh, DefaultBatchElementMask, nullptr);
diff --git a/Engine/Source/Runtime/Renderer/Private/TranslucentRendering.cpp b/Engine/Source/Runtime/Renderer/Private/TranslucentRendering.cpp
--- a/Engine/Source/Runtime/Renderer/Private/TranslucentRendering.cpp
+++ b/Engine/Source/Runtime/Renderer/Private/TranslucentRendering.cpp
@@ -388,6 +388,7 @@ void RenderViewTranslucencyInner(FRHICommandListImmediate& RHICmdList, const FVi
 						&View,
 						DrawRenderState,
 						DynamicMeshPassContext,
+						FBasePassMeshProcessor::EFlags::CanUseDepthStencil,
 						ETranslucencyPass::TPT_StandardTranslucency);
 
 					const uint64 DefaultBatchElementMask = ~0ull;
@@ -412,6 +413,7 @@ void RenderViewTranslucencyInner(FRHICommandListImmediate& RHICmdList, const FVi
 						&View,
 						DrawRenderState,
 						DynamicMeshPassContext,
+						FBasePassMeshProcessor::EFlags::CanUseDepthStencil,
 						ETranslucencyPass::TPT_StandardTranslucency);
 
 					const uint64 DefaultBatchElementMask = ~0ull;
