UpdateGPUScene optimizations.
* Disabled LCI creation and lightmap cluster upload if bAllowStaticLighting is false - halves UpdateGPUScene time
* Replaced "OutBounds=GetPreSkinnedLocalBounds" with "GetPreSkinnedLocalBounds(OutBounds)" - ~5% faster

diff --git a/Engine/Source/Editor/UnrealEd/Private/Animation/DebugSkelMeshComponent.cpp b/Engine/Source/Editor/UnrealEd/Private/Animation/DebugSkelMeshComponent.cpp
--- a/Engine/Source/Editor/UnrealEd/Private/Animation/DebugSkelMeshComponent.cpp
+++ b/Engine/Source/Editor/UnrealEd/Private/Animation/DebugSkelMeshComponent.cpp
@@ -50,7 +50,9 @@ FBoxSphereBounds UDebugSkelMeshComponent::CalcBounds(const FTransform& LocalToWo
 	// Override bounds with pre-skinned bounds if asking for them
 	if (IsUsingPreSkinnedBounds())
 	{
-		return GetPreSkinnedLocalBounds();
+		FBoxSphereBounds PreSkinnedLocalBounds;
+		GetPreSkinnedLocalBounds(PreSkinnedLocalBounds);
+		return PreSkinnedLocalBounds;
 	}
 
 	FBoxSphereBounds Result = Super::CalcBounds(LocalToWorld);
diff --git a/Engine/Source/Runtime/Engine/Classes/Components/SkinnedMeshComponent.h b/Engine/Source/Runtime/Engine/Classes/Components/SkinnedMeshComponent.h
--- a/Engine/Source/Runtime/Engine/Classes/Components/SkinnedMeshComponent.h
+++ b/Engine/Source/Runtime/Engine/Classes/Components/SkinnedMeshComponent.h
@@ -795,7 +795,7 @@ public:
 	//~ End UPrimitiveComponent Interface
 
 	/** Get the pre-skinning local space bounds for this component. */
-	FBoxSphereBounds GetPreSkinnedLocalBounds() const;
+	void GetPreSkinnedLocalBounds(FBoxSphereBounds& OutBounds) const;
 
 	/**
 	 *	Sets the value of the bForceWireframe flag and reattaches the component as necessary.
diff --git a/Engine/Source/Runtime/Engine/Private/Components/SkinnedMeshComponent.cpp b/Engine/Source/Runtime/Engine/Private/Components/SkinnedMeshComponent.cpp
--- a/Engine/Source/Runtime/Engine/Private/Components/SkinnedMeshComponent.cpp
+++ b/Engine/Source/Runtime/Engine/Private/Components/SkinnedMeshComponent.cpp
@@ -1058,24 +1058,24 @@ FBoxSphereBounds USkinnedMeshComponent::CalcMeshBound(const FVector& RootOffset,
 	return NewBounds;
 }
 
-FBoxSphereBounds USkinnedMeshComponent::GetPreSkinnedLocalBounds() const 
+void USkinnedMeshComponent::GetPreSkinnedLocalBounds(FBoxSphereBounds& OutBounds) const
 {
 	const USkinnedMeshComponent* const MasterPoseComponentInst = MasterPoseComponent.Get();
 
 	if (SkeletalMesh)
 	{
 		// Get the Pre-skinned bounds from the skeletal mesh. Note that these bounds are the "ExtendedBounds", so they can be tweaked on the SkeletalMesh   
-		return SkeletalMesh->GetBounds();
+		OutBounds = SkeletalMesh->GetBounds();
 	}
 	else if(MasterPoseComponentInst && MasterPoseComponentInst->SkeletalMesh)
 	{
 		// Get the bounds from the master pose if ther is no skeletal mesh
-		return MasterPoseComponentInst->SkeletalMesh->GetBounds();
+		OutBounds = MasterPoseComponentInst->SkeletalMesh->GetBounds();
 	}
 	else
 	{
 		// Fall back
-		return FBoxSphereBounds(ForceInitToZero);
+		OutBounds = FBoxSphereBounds(ForceInitToZero);
 	}
 }
 
diff --git a/Engine/Source/Runtime/Engine/Private/PrimitiveSceneProxy.cpp b/Engine/Source/Runtime/Engine/Private/PrimitiveSceneProxy.cpp
--- a/Engine/Source/Runtime/Engine/Private/PrimitiveSceneProxy.cpp
+++ b/Engine/Source/Runtime/Engine/Private/PrimitiveSceneProxy.cpp
@@ -269,6 +269,9 @@ void FPrimitiveSceneProxy::UpdateUniformBuffer()
 
 		Scene->GetPrimitiveUniformShaderParameters_RenderThread(PrimitiveSceneInfo, bHasPrecomputedVolumetricLightmap, PreviousLocalToWorld, SingleCaptureIndex);
 
+		FBoxSphereBounds PreSkinnedLocalBounds;
+		GetPreSkinnedLocalBounds(PreSkinnedLocalBounds);
+
 		// Update the uniform shader parameters.
 		const FPrimitiveUniformShaderParameters PrimitiveUniformShaderParameters = 
 			GetPrimitiveUniformShaderParameters(
@@ -277,7 +280,7 @@ void FPrimitiveSceneProxy::UpdateUniformBuffer()
 				ActorPosition, 
 				Bounds, 
 				LocalBounds, 
-				GetPreSkinnedLocalBounds(),
+				PreSkinnedLocalBounds,
 				bReceivesDecals, 
 				HasDistanceFieldRepresentation(), 
 				HasDynamicIndirectShadowCasterRepresentation(), 
diff --git a/Engine/Source/Runtime/Engine/Private/PrimitiveUniformShaderParameters.cpp b/Engine/Source/Runtime/Engine/Private/PrimitiveUniformShaderParameters.cpp
--- a/Engine/Source/Runtime/Engine/Private/PrimitiveUniformShaderParameters.cpp
+++ b/Engine/Source/Runtime/Engine/Private/PrimitiveUniformShaderParameters.cpp
@@ -43,13 +43,16 @@ FPrimitiveSceneShaderData::FPrimitiveSceneShaderData(const FPrimitiveSceneProxy*
 
 	Proxy->GetScene().GetPrimitiveUniformShaderParameters_RenderThread(Proxy->GetPrimitiveSceneInfo(), bHasPrecomputedVolumetricLightmap, PreviousLocalToWorld, SingleCaptureIndex);
 
+	FBoxSphereBounds PreSkinnedLocalBounds;
+	Proxy->GetPreSkinnedLocalBounds(PreSkinnedLocalBounds);
+
 	Setup(GetPrimitiveUniformShaderParameters(
 		Proxy->GetLocalToWorld(),
 		PreviousLocalToWorld,
 		Proxy->GetActorPosition(), 
 		Proxy->GetBounds(), 
 		Proxy->GetLocalBounds(),
-		Proxy->GetPreSkinnedLocalBounds(),
+		PreSkinnedLocalBounds,
 		Proxy->ReceivesDecals(), 
 		Proxy->HasDistanceFieldRepresentation(), 
 		Proxy->HasDynamicIndirectShadowCasterRepresentation(), 
diff --git a/Engine/Source/Runtime/Engine/Private/SkeletalMesh.cpp b/Engine/Source/Runtime/Engine/Private/SkeletalMesh.cpp
--- a/Engine/Source/Runtime/Engine/Private/SkeletalMesh.cpp
+++ b/Engine/Source/Runtime/Engine/Private/SkeletalMesh.cpp
@@ -3410,7 +3410,7 @@ FSkeletalMeshSceneProxy::FSkeletalMeshSceneProxy(const USkinnedMeshComponent* Co
 	bCastInsetShadow = bCastInsetShadow || bCastCapsuleDirectShadow;
 
 	// Get the pre-skinned local bounds
-	PreSkinnedLocalBounds = Component->GetPreSkinnedLocalBounds();
+	Component->GetPreSkinnedLocalBounds(PreSkinnedLocalBounds);
 
 	const USkinnedMeshComponent* SkinnedMeshComponent = Cast<const USkinnedMeshComponent>(Component);
 	if(SkinnedMeshComponent && SkinnedMeshComponent->bPerBoneMotionBlur)
diff --git a/Engine/Source/Runtime/Engine/Public/PrimitiveSceneProxy.h b/Engine/Source/Runtime/Engine/Public/PrimitiveSceneProxy.h
--- a/Engine/Source/Runtime/Engine/Public/PrimitiveSceneProxy.h
+++ b/Engine/Source/Runtime/Engine/Public/PrimitiveSceneProxy.h
@@ -425,7 +425,7 @@ public:
 	inline bool IsLocalToWorldDeterminantNegative() const { return bIsLocalToWorldDeterminantNegative; }
 	inline const FBoxSphereBounds& GetBounds() const { return Bounds; }
 	inline const FBoxSphereBounds& GetLocalBounds() const { return LocalBounds; }
-	virtual FBoxSphereBounds GetPreSkinnedLocalBounds() const { return LocalBounds; }
+	virtual void GetPreSkinnedLocalBounds(FBoxSphereBounds& OutBounds) const { OutBounds = LocalBounds; }
 	inline FName GetOwnerName() const { return OwnerName; }
 	inline FName GetResourceName() const { return ResourceName; }
 	inline FName GetLevelName() const { return LevelName; }
diff --git a/Engine/Source/Runtime/Engine/Public/SkeletalMeshTypes.h b/Engine/Source/Runtime/Engine/Public/SkeletalMeshTypes.h
--- a/Engine/Source/Runtime/Engine/Public/SkeletalMeshTypes.h
+++ b/Engine/Source/Runtime/Engine/Public/SkeletalMeshTypes.h
@@ -242,7 +242,7 @@ public:
 	virtual void GetShadowShapes(TArray<FCapsuleShape>& CapsuleShapes) const override;
 
 	/** Return the bounds for the pre-skinned primitive in local space */
-	virtual FBoxSphereBounds GetPreSkinnedLocalBounds() const override { return PreSkinnedLocalBounds; }
+	virtual void GetPreSkinnedLocalBounds(FBoxSphereBounds& OutBounds) const override { OutBounds = PreSkinnedLocalBounds; }
 
 	/** Returns a pre-sorted list of shadow capsules's bone indicies */
 	const TArray<uint16>& GetSortedShadowBoneIndices() const
diff --git a/Engine/Source/Runtime/Renderer/Private/PrimitiveSceneInfo.cpp b/Engine/Source/Runtime/Renderer/Private/PrimitiveSceneInfo.cpp
--- a/Engine/Source/Runtime/Renderer/Private/PrimitiveSceneInfo.cpp
+++ b/Engine/Source/Runtime/Renderer/Private/PrimitiveSceneInfo.cpp
@@ -475,23 +475,27 @@ void FPrimitiveSceneInfo::AddToScene(FRHICommandListImmediate& RHICmdList, bool
 	}
 	MarkIndirectLightingCacheBufferDirty();
 
-	FPrimitiveSceneProxy::FLCIArray LCIs;
-	Proxy->GetLCIs(LCIs);
-	for (int32 i = 0; i < LCIs.Num(); ++i)
+	const bool bAllowStaticLighting = FReadOnlyCVARCache::Get().bAllowStaticLighting;
+	if (bAllowStaticLighting)
 	{
-		FLightCacheInterface* LCI = LCIs[i];
-
-		if (LCI) 
+		FPrimitiveSceneProxy::FLCIArray LCIs;
+		Proxy->GetLCIs(LCIs);
+		for (int32 i = 0; i < LCIs.Num(); ++i)
 		{
-			LCI->CreatePrecomputedLightingUniformBuffer_RenderingThread(Scene->GetFeatureLevel());
+			FLightCacheInterface* LCI = LCIs[i];
+
+			if (LCI)
+			{
+				LCI->CreatePrecomputedLightingUniformBuffer_RenderingThread(Scene->GetFeatureLevel());
+			}
 		}
-	}
 
-	NumLightmapDataEntries = LCIs.Num();
+		NumLightmapDataEntries = LCIs.Num();
 
-	if (NumLightmapDataEntries > 0 && UseGPUScene(GMaxRHIShaderPlatform, Scene->GetFeatureLevel()))
-	{
-		LightmapDataOffset = Scene->GPUScene.LightmapDataAllocator.Allocate(NumLightmapDataEntries);
+		if (NumLightmapDataEntries > 0 && UseGPUScene(GMaxRHIShaderPlatform, Scene->GetFeatureLevel()))
+		{
+			LightmapDataOffset = Scene->GPUScene.LightmapDataAllocator.Allocate(NumLightmapDataEntries);
+		}
 	}
 
 	// Cache the nearest reflection proxy if needed
