Fixed issues with translucent materials in shader complexity.
Reverted shader complexity to previous logic by adding fixed deffered shading cost to instruction count.
Fixed quad complexity translucency not showing up.

diff --git a/Engine/Source/Runtime/Renderer/Private/DebugViewModeRendering.cpp b/Engine/Source/Runtime/Renderer/Private/DebugViewModeRendering.cpp
--- a/Engine/Source/Runtime/Renderer/Private/DebugViewModeRendering.cpp
+++ b/Engine/Source/Runtime/Renderer/Private/DebugViewModeRendering.cpp
@@ -382,23 +382,15 @@ void FDebugViewModeMeshProcessor::UpdateInstructionCount(FDebugViewModeShaderEle
 	OutShaderElementData.NumVSInstructions = InBatchMaterial->GetShader<TBasePassVS<TUniformLightMapPolicy<LMP_NO_LIGHTMAP>, false>>(InVertexFactoryType)->GetNumInstructions();
 	OutShaderElementData.NumPSInstructions = InBatchMaterial->GetShader<TBasePassPS<TUniformLightMapPolicy<LMP_NO_LIGHTMAP>, false>>(InVertexFactoryType)->GetNumInstructions();
 
-	// Those numbers are taken from a simple material where common inputs are bound to vector parameters (to prevent constant optimizations).
-	if (bDeferred)
+	if (!bDeferred)
 	{
-		OutShaderElementData.NumVSInstructions -= 41;
-		OutShaderElementData.NumPSInstructions -= bLit ? 111 : 33;
+		// Those numbers are taken from a simple material where common inputs are bound to vector parameters (to prevent constant optimizations).
+		OutShaderElementData.NumVSInstructions -= GShaderComplexityBaselineForwardVS - GShaderComplexityBaselineDeferredVS;
+		OutShaderElementData.NumPSInstructions -= bLit ? (GShaderComplexityBaselineForwardPS - GShaderComplexityBaselineDeferredPS) : (GShaderComplexityBaselineForwardUnlitPS - GShaderComplexityBaselineDeferredUnlitPS);
 	}
-	else
-	{
-		OutShaderElementData.NumVSInstructions -= 134;
-		OutShaderElementData.NumPSInstructions -= bLit ? 635 : 47;
-	}
-
-	// Ensure a minimum instruction count so that overdraw accumulates.
-	const int32 MinInstructionCount = IsTranslucentBlendMode(InBatchMaterial->GetBlendMode()) ? 20 : 0;
 	
-	OutShaderElementData.NumVSInstructions = FMath::Max<int32>(MinInstructionCount, OutShaderElementData.NumVSInstructions);
-	OutShaderElementData.NumPSInstructions = FMath::Max<int32>(MinInstructionCount, OutShaderElementData.NumPSInstructions);
+	OutShaderElementData.NumVSInstructions = FMath::Max<int32>(0, OutShaderElementData.NumVSInstructions);
+	OutShaderElementData.NumPSInstructions = FMath::Max<int32>(0, OutShaderElementData.NumPSInstructions);
 }
 
 FMeshPassProcessor* CreateDebugViewModePassProcessor(const FScene* Scene, const FSceneView* InViewIfDynamicMeshCommand, FMeshPassDrawListContext* InDrawListContext)
diff --git a/Engine/Source/Runtime/Renderer/Private/ShaderComplexityRendering.cpp b/Engine/Source/Runtime/Renderer/Private/ShaderComplexityRendering.cpp
--- a/Engine/Source/Runtime/Renderer/Private/ShaderComplexityRendering.cpp
+++ b/Engine/Source/Runtime/Renderer/Private/ShaderComplexityRendering.cpp
@@ -10,6 +10,54 @@ ShaderComplexityRendering.cpp: Contains definitions for rendering the shader com
 
 #if !(UE_BUILD_SHIPPING || UE_BUILD_TEST)
 
+int32 GShaderComplexityBaselineForwardVS = 134;
+static FAutoConsoleVariableRef CVarShaderComplexityBaselineForwardVS(
+	TEXT("r.ShaderComplexity.Baseline.Forward.VS"),
+	GShaderComplexityBaselineForwardVS,
+	TEXT("Minimum number of instructions for vertex shaders in forward shading (default=134)"),
+	ECVF_Default
+);
+
+int32 GShaderComplexityBaselineForwardPS = 635;
+static FAutoConsoleVariableRef CVarShaderComplexityBaselineForwardPS(
+	TEXT("r.ShaderComplexity.Baseline.Forward.PS"),
+	GShaderComplexityBaselineForwardPS,
+	TEXT("Minimum number of instructions for pixel shaders in forward shading (default=635)"),
+	ECVF_Default
+);
+
+int32 GShaderComplexityBaselineForwardUnlitPS = 47;
+static FAutoConsoleVariableRef CVarShaderComplexityBaselineForwardUnlitPS(
+	TEXT("r.ShaderComplexity.Baseline.Forward.UnlitPS"),
+	GShaderComplexityBaselineForwardUnlitPS,
+	TEXT("Minimum number of instructions for unlit material pixel shaders in forward shading (default=47)"),
+	ECVF_Default
+);
+
+int32 GShaderComplexityBaselineDeferredVS = 41;
+static FAutoConsoleVariableRef CVarShaderComplexityBaselineDeferredVS(
+	TEXT("r.ShaderComplexity.Baseline.Deferred.VS"),
+	GShaderComplexityBaselineDeferredVS,
+	TEXT("Minimum number of instructions for vertex shaders in deferred shading (default=41)"),
+	ECVF_Default
+);
+
+int32 GShaderComplexityBaselineDeferredPS = 111;
+static FAutoConsoleVariableRef CVarShaderComplexityBaselineDeferredPS(
+	TEXT("r.ShaderComplexity.Baseline.Deferred.PS"),
+	GShaderComplexityBaselineDeferredPS,
+	TEXT("Minimum number of instructions for pixel shaders in deferred shading (default=111)"),
+	ECVF_Default
+);
+
+int32 GShaderComplexityBaselineDeferredUnlitPS = 33;
+static FAutoConsoleVariableRef CVarShaderComplexityBaselineDeferredUnlitPS(
+	TEXT("r.ShaderComplexity.Baseline.Deferred.UnlitPS"),
+	GShaderComplexityBaselineDeferredUnlitPS,
+	TEXT("Minimum number of instructions for unlit material pixel shaders in deferred shading (default=33)"),
+	ECVF_Default
+);
+
 IMPLEMENT_SHADER_TYPE(template<>,TComplexityAccumulatePS<false>,TEXT("/Engine/Private/ShaderComplexityAccumulatePixelShader.usf"),TEXT("Main"),SF_Pixel);
 IMPLEMENT_SHADER_TYPE(template<>,TComplexityAccumulatePS<true>,TEXT("/Engine/Private/QuadComplexityAccumulatePixelShader.usf"),TEXT("Main"),SF_Pixel);
 
@@ -57,29 +105,8 @@ FDebugViewModePS* FComplexityAccumulateInterface::GetPixelShader(const FMaterial
 }
 void FComplexityAccumulateInterface::SetDrawRenderState(EBlendMode BlendMode, FRenderState& DrawRenderState) const
 {
-	if (bShowShaderComplexity)
-	{
-		// When rendering masked materials in the shader complexity viewmode, 
-		// We want to overwrite complexity for the pixels which get depths written,
-		// And accumulate complexity for pixels which get killed due to the opacity mask being below the clip value.
-		// This is accomplished by forcing the masked materials to render depths in the depth only pass, 
-		// Then rendering in the base pass with additive complexity blending, depth tests on, and depth writes off.
-		DrawRenderState.DepthStencilState = TStaticDepthStencilState<false, CF_DepthNearOrEqual>::GetRHI();
-
-		if (IsTranslucentBlendMode(BlendMode))
-		{
-			DrawRenderState.BlendState = TStaticBlendState<>::GetRHI();
-		}
-		else
-		{
-			// If we are in the translucent pass then override the blend mode, otherwise maintain additive blending.
-			DrawRenderState.BlendState = TStaticBlendState<CW_RGBA, BO_Add, BF_One, BF_One, BO_Add, BF_Zero, BF_One>::GetRHI();
-		}
-	}
-	else
-	{
-		FDebugViewModeInterface::SetDrawRenderState(BlendMode, DrawRenderState);
-	}
+	DrawRenderState.DepthStencilState = TStaticDepthStencilState<false, CF_DepthNearOrEqual>::GetRHI();
+	DrawRenderState.BlendState = TStaticBlendState<CW_RGBA, BO_Add, BF_One, BF_One, BO_Add, BF_Zero, BF_One>::GetRHI();
 }
 
 // Instantiate the template 
diff --git a/Engine/Source/Runtime/Renderer/Private/ShaderComplexityRendering.h b/Engine/Source/Runtime/Renderer/Private/ShaderComplexityRendering.h
--- a/Engine/Source/Runtime/Renderer/Private/ShaderComplexityRendering.h
+++ b/Engine/Source/Runtime/Renderer/Private/ShaderComplexityRendering.h
@@ -19,6 +19,13 @@ struct FMeshDrawingRenderState;
 
 #if !(UE_BUILD_SHIPPING || UE_BUILD_TEST)
 
+extern int32 GShaderComplexityBaselineForwardVS;
+extern int32 GShaderComplexityBaselineForwardPS;
+extern int32 GShaderComplexityBaselineForwardUnlitPS;
+extern int32 GShaderComplexityBaselineDeferredVS;
+extern int32 GShaderComplexityBaselineDeferredPS;
+extern int32 GShaderComplexityBaselineDeferredUnlitPS;
+
 template <bool bQuadComplexity>
 class TComplexityAccumulatePS : public FDebugViewModePS
 {
