Fixes for VS2019 compatibility.

diff --git a/Engine/Plugins/FX/Niagara/Source/NiagaraEditor/Private/NiagaraCompiler.cpp b/Engine/Plugins/FX/Niagara/Source/NiagaraEditor/Private/NiagaraCompiler.cpp
--- a/Engine/Plugins/FX/Niagara/Source/NiagaraEditor/Private/NiagaraCompiler.cpp
+++ b/Engine/Plugins/FX/Niagara/Source/NiagaraEditor/Private/NiagaraCompiler.cpp
@@ -617,7 +617,7 @@ void FNiagaraEditorModule::TestCompileScriptFromConsole(const TArray<FString>& A
 			FShaderCompilerOutput Output;
 			FVectorVMCompilationOutput CompilationOutput;
 			double StartTime = FPlatformTime::Seconds();
-			bool bSucceeded = CompileShader_VectorVM(Input, Output, FString(FPlatformProcess::ShaderDir()), 0, CompilationOutput, GNiagaraSkipVectorVMBackendOptimizations);
+			bool bSucceeded = CompileShader_VectorVM(Input, Output, FString(FPlatformProcess::ShaderDir()), 0, CompilationOutput, GNiagaraSkipVectorVMBackendOptimizations != 0);
 			float DeltaTime = (float)(FPlatformTime::Seconds() - StartTime);
 
 			if (bSucceeded)
@@ -768,7 +768,7 @@ FNiagaraCompileResults FHlslNiagaraCompiler::CompileScript(const FNiagaraCompile
 			
 			CritSec.Lock();
 			double StartTime = FPlatformTime::Seconds();
-			CompileResults.bVMSucceeded = CompileShader_VectorVM(Input, Output, FString(FPlatformProcess::ShaderDir()), 0, CompilationOutput, GNiagaraSkipVectorVMBackendOptimizations);
+			CompileResults.bVMSucceeded = CompileShader_VectorVM(Input, Output, FString(FPlatformProcess::ShaderDir()), 0, CompilationOutput, GNiagaraSkipVectorVMBackendOptimizations != 0);
 			CompileResults.CompileTime = (float)(FPlatformTime::Seconds() - StartTime);
 			CritSec.Unlock();
 		}
diff --git a/Engine/Source/Runtime/Renderer/Private/RayTracing/RayTracingMaterialHitShaders.cpp b/Engine/Source/Runtime/Renderer/Private/RayTracing/RayTracingMaterialHitShaders.cpp
--- a/Engine/Source/Runtime/Renderer/Private/RayTracing/RayTracingMaterialHitShaders.cpp
+++ b/Engine/Source/Runtime/Renderer/Private/RayTracing/RayTracingMaterialHitShaders.cpp
@@ -260,19 +260,19 @@ void FRayTracingMeshProcessor::Process(
 	switch (LightMapPolicy.GetIndirectPolicy())
 	{
 	case LMP_PRECOMPUTED_IRRADIANCE_VOLUME_INDIRECT_LIGHTING:
-		RayTracingShaders.RayHitGroupShader = GetMaterialHitShader<TUniformLightMapPolicy<LMP_PRECOMPUTED_IRRADIANCE_VOLUME_INDIRECT_LIGHTING>>(MaterialResource, VertexFactory, CVarRayTracingTextureLodCvar.GetValueOnRenderThread());
+		RayTracingShaders.RayHitGroupShader = GetMaterialHitShader<TUniformLightMapPolicy<LMP_PRECOMPUTED_IRRADIANCE_VOLUME_INDIRECT_LIGHTING>>(MaterialResource, VertexFactory, CVarRayTracingTextureLodCvar.GetValueOnRenderThread() != 0);
 		break;
 	case LMP_LQ_LIGHTMAP:
-		RayTracingShaders.RayHitGroupShader = GetMaterialHitShader<TUniformLightMapPolicy<LMP_LQ_LIGHTMAP>>(MaterialResource, VertexFactory, CVarRayTracingTextureLodCvar.GetValueOnRenderThread());
+		RayTracingShaders.RayHitGroupShader = GetMaterialHitShader<TUniformLightMapPolicy<LMP_LQ_LIGHTMAP>>(MaterialResource, VertexFactory, CVarRayTracingTextureLodCvar.GetValueOnRenderThread() != 0);
 		break;
 	case LMP_HQ_LIGHTMAP:
-		RayTracingShaders.RayHitGroupShader = GetMaterialHitShader<TUniformLightMapPolicy<LMP_HQ_LIGHTMAP>>(MaterialResource, VertexFactory, CVarRayTracingTextureLodCvar.GetValueOnRenderThread());
+		RayTracingShaders.RayHitGroupShader = GetMaterialHitShader<TUniformLightMapPolicy<LMP_HQ_LIGHTMAP>>(MaterialResource, VertexFactory, CVarRayTracingTextureLodCvar.GetValueOnRenderThread() != 0);
 		break;
 	case LMP_DISTANCE_FIELD_SHADOWS_AND_HQ_LIGHTMAP:
-		RayTracingShaders.RayHitGroupShader = GetMaterialHitShader<TUniformLightMapPolicy<LMP_DISTANCE_FIELD_SHADOWS_AND_HQ_LIGHTMAP>>(MaterialResource, VertexFactory, CVarRayTracingTextureLodCvar.GetValueOnRenderThread());
+		RayTracingShaders.RayHitGroupShader = GetMaterialHitShader<TUniformLightMapPolicy<LMP_DISTANCE_FIELD_SHADOWS_AND_HQ_LIGHTMAP>>(MaterialResource, VertexFactory, CVarRayTracingTextureLodCvar.GetValueOnRenderThread() != 0);
 		break;
 	case LMP_NO_LIGHTMAP:
-		RayTracingShaders.RayHitGroupShader = GetMaterialHitShader<TUniformLightMapPolicy<LMP_NO_LIGHTMAP>>(MaterialResource, VertexFactory, CVarRayTracingTextureLodCvar.GetValueOnRenderThread());
+		RayTracingShaders.RayHitGroupShader = GetMaterialHitShader<TUniformLightMapPolicy<LMP_NO_LIGHTMAP>>(MaterialResource, VertexFactory, CVarRayTracingTextureLodCvar.GetValueOnRenderThread() != 0);
 		break;
 	default:
 		check(false);
@@ -456,7 +456,7 @@ FRHIRayTracingPipelineState* FDeferredShadingSceneRenderer::BindRayTracingMateri
 	FRayTracingShaderRHIParamRef MissShaderTable[] = { MissShader };
 	Initializer.SetMissShaderTable(MissShaderTable);
 
-	const bool bEnableMaterials = GEnableRayTracingMaterials;
+	const bool bEnableMaterials = GEnableRayTracingMaterials != 0;
 
 	TArray<FRayTracingShaderRHIParamRef> RayTracingMaterialLibrary;
 
diff --git a/Engine/Source/Runtime/Renderer/Private/RayTracing/RayTracingReflections.cpp b/Engine/Source/Runtime/Renderer/Private/RayTracing/RayTracingReflections.cpp
--- a/Engine/Source/Runtime/Renderer/Private/RayTracing/RayTracingReflections.cpp
+++ b/Engine/Source/Runtime/Renderer/Private/RayTracing/RayTracingReflections.cpp
@@ -188,7 +188,7 @@ void FDeferredShadingSceneRenderer::PrepareRayTracingReflections(const FViewInfo
 {
 	// Declare all RayGen shaders that require material closest hit shaders to be bound
 
-	const bool bSortMaterials = CVarRayTracingReflectionsSortMaterials.GetValueOnRenderThread();
+	const bool bSortMaterials = CVarRayTracingReflectionsSortMaterials.GetValueOnRenderThread() != 0;
 
 	const EDeferredMaterialMode DeferredMaterialMode = bSortMaterials ? EDeferredMaterialMode::Shade : EDeferredMaterialMode::None;
 
@@ -214,7 +214,7 @@ void FDeferredShadingSceneRenderer::RenderRayTracingReflections(
 #if RHI_RAYTRACING
 {
 	const uint32 SortTileSize = CVarRayTracingReflectionsSortTileSize.GetValueOnRenderThread();
-	const bool bSortMaterials = CVarRayTracingReflectionsSortMaterials.GetValueOnRenderThread();
+	const bool bSortMaterials = CVarRayTracingReflectionsSortMaterials.GetValueOnRenderThread() != 0;
 
 	FSceneRenderTargets& SceneContext = FSceneRenderTargets::Get(GraphBuilder.RHICmdList);
 
