Fix for:  Ensure: ReflectionCaptureBuffer.IsBound occurs when building texture streaming for the first time.

- Remove FDummyReflectionCaptureUniformBuffer - was put in for similar bug and turns out this just masked the issue for a while.
- DVSM_OutputMaterialTextureScales was unsuported on Metal platforms due to returningn false from AllowDebugViewShaderMode() in PlatformSupportsDebugViewShaders() for DVSM_OutputMaterialTextureScales.
- Enable debug view mode shaders except for the following modes: DVSM_ShaderComplexityContainedQuadOverhead, DVSM_ShaderComplexityBleedingQuadOverhead, DVSM_QuadComplexity as these don't compile correctly currently on Metal.

diff --git a/Engine/Source/Runtime/Engine/Private/DebugViewModeHelpers.cpp b/Engine/Source/Runtime/Engine/Private/DebugViewModeHelpers.cpp
--- a/Engine/Source/Runtime/Engine/Private/DebugViewModeHelpers.cpp
+++ b/Engine/Source/Runtime/Engine/Private/DebugViewModeHelpers.cpp
@@ -19,7 +19,7 @@
 static bool PlatformSupportsDebugViewShaders(EShaderPlatform Platform)
 {
 	// List of platforms that have been tested and proved functional.
-	return Platform == SP_PCD3D_SM4 || Platform == SP_PCD3D_SM5 || Platform == SP_OPENGL_SM4;
+	return Platform == SP_PCD3D_SM4 || Platform == SP_PCD3D_SM5 || Platform == SP_OPENGL_SM4 || Platform == SP_METAL_SM5_NOTESS || Platform == SP_METAL_SM5;
 }
 
 bool AllowDebugViewVSDSHS(EShaderPlatform Platform)
@@ -48,7 +48,7 @@ bool AllowDebugViewShaderMode(EDebugViewShaderMode ShaderMode, EShaderPlatform P
 	case DVSM_ShaderComplexityContainedQuadOverhead:
 	case DVSM_ShaderComplexityBleedingQuadOverhead:
 	case DVSM_QuadComplexity:
-		return FeatureLevel >= ERHIFeatureLevel::SM5 && (bForceQuadOverdraw || PlatformSupportsDebugViewShaders(Platform));
+		return FeatureLevel >= ERHIFeatureLevel::SM5 && (bForceQuadOverdraw || (PlatformSupportsDebugViewShaders(Platform) && !IsMetalPlatform(Platform))); // Last one to fix for Metal then remove this Metal check.
 	case DVSM_PrimitiveDistanceAccuracy:
 	case DVSM_MeshUVDensityAccuracy:
 		return FeatureLevel >= ERHIFeatureLevel::SM5 && (bForceStreamingAccuracy || PlatformSupportsDebugViewShaders(Platform));
diff --git a/Engine/Source/Runtime/Renderer/Private/Renderer.cpp b/Engine/Source/Runtime/Renderer/Private/Renderer.cpp
--- a/Engine/Source/Runtime/Renderer/Private/Renderer.cpp
+++ b/Engine/Source/Runtime/Renderer/Private/Renderer.cpp
@@ -44,22 +44,6 @@ IMPLEMENT_MODULE(FRendererModule, Renderer);
 	FSystemSettings* GSystemSettingsForVisualizers = &GSystemSettings;
 #endif
 
-// Dummy Reflection capture uniform buffer for translucent tile mesh rendering
-class FDummyReflectionCaptureUniformBuffer : public TUniformBuffer<FReflectionCaptureShaderData>
-{
-	typedef TUniformBuffer< FReflectionCaptureShaderData > Super;
-	
-public:
-	virtual void InitDynamicRHI() override
-	{
-		FReflectionCaptureShaderData DummyPositionsBuffer;
-		FMemory::Memzero(DummyPositionsBuffer);
-		SetContentsNoUpdate(DummyPositionsBuffer);
-		Super::InitDynamicRHI();
-	}
-};
-static TGlobalResource< FDummyReflectionCaptureUniformBuffer > GDummyReflectionCaptureUniformBuffer;
-
 void FRendererModule::StartupModule()
 {
 	GScreenSpaceDenoiser = IScreenSpaceDenoiser::GetDefaultDenoiser();
@@ -203,12 +187,6 @@ void FRendererModule::DrawTileMesh(FRHICommandListImmediate& RHICmdList, FMeshPa
 
 			if (ShadingPath == EShadingPath::Deferred)
 			{
-				// Crash fix - reflection capture shader parameter is bound but we have no buffer during Build Texture Streaming
-				if(!View.ReflectionCaptureUniformBuffer.IsValid())
-				{
-					View.ReflectionCaptureUniformBuffer = GDummyReflectionCaptureUniformBuffer.GetUniformBufferRef();
-				}
-			
 				TUniformBufferRef<FTranslucentBasePassUniformParameters> TranslucentBasePassUniformBuffer;
 				CreateTranslucentBasePassUniformBuffer(RHICmdList, View, nullptr, ESceneTextureSetupMode::None, TranslucentBasePassUniformBuffer, 0);
 				BasePassUniformBuffer = TranslucentBasePassUniformBuffer;
