Part 2 fix for "Crash on Mac while packaging TM-Decals with Forward Shading enabled" (Note: also occurs on windws).     Protect against nullptr HZB buffer assert by disabling SSAO (on it's own this would hide this problem but also protects against future crashes with zprepass off in forward and SSAO having no valid HZB buffer).

Part 1 of this fix has now already been submiited to 4.22 by [REDACTED_PERSON] 5629897 so this CL now doesn't cover that part; it was a race condition between the game thread setting EarlyZPassMode in FDeferredShadingSceneRenderer::FScene::UpdateEarlyZPassMode() constructor and the render thread is updating this in FScene::UpdateEarlyZPassMode().

diff --git a/Engine/Source/Runtime/Renderer/Private/DeferredShadingRenderer.cpp b/Engine/Source/Runtime/Renderer/Private/DeferredShadingRenderer.cpp
--- a/Engine/Source/Runtime/Renderer/Private/DeferredShadingRenderer.cpp
+++ b/Engine/Source/Runtime/Renderer/Private/DeferredShadingRenderer.cpp
@@ -1232,9 +1232,9 @@ void FDeferredShadingSceneRenderer::Render(FRHICommandListImmediate& RHICmdList)
 			Scene->UniformBuffers.UpdateViewUniformBuffer(View);
 
 			uint32 SSAOLevels = FSSAOHelper::ComputeAmbientOcclusionPassCount(View);
-			// In deferred shader, the SSAO uses the GBuffer and must be executed after base pass. 
-			// Otherwise, async compute runs the shader in RenderHzb()
-			if (!IsForwardShadingEnabled(ShaderPlatform) || FSSAOHelper::IsAmbientOcclusionAsyncCompute(View, SSAOLevels))
+			// In deferred shader, the SSAO uses the GBuffer and must be executed after base pass. Otherwise, async compute runs the shader in RenderHzb()
+			// In forward, if zprepass is off - as SSAO here requires a valid HZB buffer - disable SSAO
+			if (!IsForwardShadingEnabled(ShaderPlatform) || !View.HZB.IsValid() || FSSAOHelper::IsAmbientOcclusionAsyncCompute(View, SSAOLevels))
 			{
 				SSAOLevels = 0;
 			}
