Fixed visible translucents and mesh decals, when Visible is unsed and HiddenShadow is set. Fixed visible translucents and mesh decals, when Visible is unsed and HiddenShadow is set. If ViewRelevance.bDrawRelevance is false, then we shouln't be drawing those.

diff --git a/Engine/Source/Runtime/Renderer/Private/SceneVisibility.cpp b/Engine/Source/Runtime/Renderer/Private/SceneVisibility.cpp
--- a/Engine/Source/Runtime/Renderer/Private/SceneVisibility.cpp
+++ b/Engine/Source/Runtime/Renderer/Private/SceneVisibility.cpp
@@ -2129,126 +2129,131 @@ struct FRelevancePacket
 						bNeedsBatchVisibility = true;
 					}
 
-					if(ViewRelevance.bDrawRelevance && (StaticMeshRelevance.bUseForMaterial || StaticMeshRelevance.bUseAsOccluder) && (ViewRelevance.bRenderInMainPass || ViewRelevance.bRenderCustomDepth) && !bHiddenByHLODFade)
+					if (ViewRelevance.bDrawRelevance)
 					{
-						if (StaticMeshRelevance.bUseForDepthPass && bDrawDepthOnly)
+						if ((StaticMeshRelevance.bUseForMaterial || StaticMeshRelevance.bUseAsOccluder) 
+							&& (ViewRelevance.bRenderInMainPass || ViewRelevance.bRenderCustomDepth) 
+							&& !bHiddenByHLODFade)
 						{
-							DrawCommandPacket.AddCommandsForMesh(PrimitiveIndex, PrimitiveSceneInfo, StaticMeshRelevance, StaticMesh, Scene, bCanCache, EMeshPass::DepthPass);
-						}
-
-						// Mark static mesh as visible for rendering
-						if (StaticMeshRelevance.bUseForMaterial)
-						{
-							DrawCommandPacket.AddCommandsForMesh(PrimitiveIndex, PrimitiveSceneInfo, StaticMeshRelevance, StaticMesh, Scene, bCanCache, EMeshPass::BasePass);
-							MarkMask |= EMarkMaskBits::StaticMeshVisibilityMapMask;
-
-							if (ShadingPath == EShadingPath::Mobile)
+							if (StaticMeshRelevance.bUseForDepthPass && bDrawDepthOnly)
 							{
-								DrawCommandPacket.AddCommandsForMesh(PrimitiveIndex, PrimitiveSceneInfo, StaticMeshRelevance, StaticMesh, Scene, bCanCache, EMeshPass::MobileBasePassCSM);
+								DrawCommandPacket.AddCommandsForMesh(PrimitiveIndex, PrimitiveSceneInfo, StaticMeshRelevance, StaticMesh, Scene, bCanCache, EMeshPass::DepthPass);
 							}
 
-							if (ViewRelevance.bRenderCustomDepth)
+							// Mark static mesh as visible for rendering
+							if (StaticMeshRelevance.bUseForMaterial)
 							{
-								DrawCommandPacket.AddCommandsForMesh(PrimitiveIndex, PrimitiveSceneInfo, StaticMeshRelevance, StaticMesh, Scene, bCanCache, EMeshPass::CustomDepth);
-							}
+								DrawCommandPacket.AddCommandsForMesh(PrimitiveIndex, PrimitiveSceneInfo, StaticMeshRelevance, StaticMesh, Scene, bCanCache, EMeshPass::BasePass);
+								MarkMask |= EMarkMaskBits::StaticMeshVisibilityMapMask;
 
-							if (bAddLightmapDensityCommands)
-							{
-								DrawCommandPacket.AddCommandsForMesh(PrimitiveIndex, PrimitiveSceneInfo, StaticMeshRelevance, StaticMesh, Scene, bCanCache, EMeshPass::LightmapDensity);
-							}
+								if (ShadingPath == EShadingPath::Mobile)
+								{
+									DrawCommandPacket.AddCommandsForMesh(PrimitiveIndex, PrimitiveSceneInfo, StaticMeshRelevance, StaticMesh, Scene, bCanCache, EMeshPass::MobileBasePassCSM);
+								}
+
+								if (ViewRelevance.bRenderCustomDepth)
+								{
+									DrawCommandPacket.AddCommandsForMesh(PrimitiveIndex, PrimitiveSceneInfo, StaticMeshRelevance, StaticMesh, Scene, bCanCache, EMeshPass::CustomDepth);
+								}
+
+								if (bAddLightmapDensityCommands)
+								{
+									DrawCommandPacket.AddCommandsForMesh(PrimitiveIndex, PrimitiveSceneInfo, StaticMeshRelevance, StaticMesh, Scene, bCanCache, EMeshPass::LightmapDensity);
+								}
 #if !(UE_BUILD_SHIPPING || UE_BUILD_TEST)
-							else if (View.Family->UseDebugViewPS())
-							{
-								DrawCommandPacket.AddCommandsForMesh(PrimitiveIndex, PrimitiveSceneInfo, StaticMeshRelevance, StaticMesh, Scene, bCanCache, EMeshPass::DebugViewMode);
-							}
+								else if (View.Family->UseDebugViewPS())
+								{
+									DrawCommandPacket.AddCommandsForMesh(PrimitiveIndex, PrimitiveSceneInfo, StaticMeshRelevance, StaticMesh, Scene, bCanCache, EMeshPass::DebugViewMode);
+								}
 #endif
 
 #if WITH_EDITOR
-							if (StaticMeshRelevance.bSelectable)
-							{
-								if (View.bAllowTranslucentPrimitivesInHitProxy)
+								if (StaticMeshRelevance.bSelectable)
 								{
-									DrawCommandPacket.AddCommandsForMesh(PrimitiveIndex, PrimitiveSceneInfo, StaticMeshRelevance, StaticMesh, Scene, bCanCache, EMeshPass::HitProxy);
+									if (View.bAllowTranslucentPrimitivesInHitProxy)
+									{
+										DrawCommandPacket.AddCommandsForMesh(PrimitiveIndex, PrimitiveSceneInfo, StaticMeshRelevance, StaticMesh, Scene, bCanCache, EMeshPass::HitProxy);
+									}
+									else
+									{
+										DrawCommandPacket.AddCommandsForMesh(PrimitiveIndex, PrimitiveSceneInfo, StaticMeshRelevance, StaticMesh, Scene, bCanCache, EMeshPass::HitProxyOpaqueOnly);
+									}
 								}
-								else
+#endif
+
+								if (ViewRelevance.bVelocityRelevance
+									&& FVelocityRendering::PrimitiveHasVelocity(View.GetFeatureLevel(), PrimitiveSceneInfo)
+									&& FVelocityRendering::PrimitiveHasVelocityForView(View, Bounds.BoxSphereBounds, PrimitiveSceneInfo))
 								{
-									DrawCommandPacket.AddCommandsForMesh(PrimitiveIndex, PrimitiveSceneInfo, StaticMeshRelevance, StaticMesh, Scene, bCanCache, EMeshPass::HitProxyOpaqueOnly);
+									DrawCommandPacket.AddCommandsForMesh(PrimitiveIndex, PrimitiveSceneInfo, StaticMeshRelevance, StaticMesh, Scene, bCanCache, EMeshPass::Velocity);
 								}
-							}
-#endif
-							 
-							if (ViewRelevance.bVelocityRelevance 
-								&& FVelocityRendering::PrimitiveHasVelocity(View.GetFeatureLevel(), PrimitiveSceneInfo)
-								&& FVelocityRendering::PrimitiveHasVelocityForView(View, Bounds.BoxSphereBounds, PrimitiveSceneInfo))
-							{
-								DrawCommandPacket.AddCommandsForMesh(PrimitiveIndex, PrimitiveSceneInfo, StaticMeshRelevance, StaticMesh, Scene, bCanCache, EMeshPass::Velocity);
+
+								++NumVisibleStaticMeshElements;
 							}
 
-							++NumVisibleStaticMeshElements;
+							bNeedsBatchVisibility = true;
 						}
 
-						bNeedsBatchVisibility = true;
-					}
-
-					if (StaticMeshRelevance.bUseForMaterial
-						&& ViewRelevance.HasTranslucency() 
-						&& !ViewRelevance.bEditorPrimitiveRelevance
-						&& ViewRelevance.bRenderInMainPass)
-					{
-						if (View.Family->AllowTranslucencyAfterDOF())
+						if (StaticMeshRelevance.bUseForMaterial
+							&& ViewRelevance.HasTranslucency()
+							&& !ViewRelevance.bEditorPrimitiveRelevance
+							&& ViewRelevance.bRenderInMainPass)
 						{
-							if (ViewRelevance.bNormalTranslucencyRelevance)
+							if (View.Family->AllowTranslucencyAfterDOF())
+							{
+								if (ViewRelevance.bNormalTranslucencyRelevance)
+								{
+									DrawCommandPacket.AddCommandsForMesh(PrimitiveIndex, PrimitiveSceneInfo, StaticMeshRelevance, StaticMesh, Scene, bCanCache, EMeshPass::TranslucencyStandard);
+								}
+
+								if (ViewRelevance.bSeparateTranslucencyRelevance)
+								{
+									DrawCommandPacket.AddCommandsForMesh(PrimitiveIndex, PrimitiveSceneInfo, StaticMeshRelevance, StaticMesh, Scene, bCanCache, EMeshPass::TranslucencyAfterDOF);
+								}
+							}
+							else
+							{
+								// Otherwise, everything is rendered in a single bucket. This is not related to whether DOF is currently enabled or not.
+								// When using all translucency, Standard and AfterDOF are sorted together instead of being rendered like 2 buckets.
+								DrawCommandPacket.AddCommandsForMesh(PrimitiveIndex, PrimitiveSceneInfo, StaticMeshRelevance, StaticMesh, Scene, bCanCache, EMeshPass::TranslucencyAll);
+							}
+
+							if (ViewRelevance.bDistortionRelevance)
 							{
-								DrawCommandPacket.AddCommandsForMesh(PrimitiveIndex, PrimitiveSceneInfo, StaticMeshRelevance, StaticMesh, Scene, bCanCache, EMeshPass::TranslucencyStandard);
+								DrawCommandPacket.AddCommandsForMesh(PrimitiveIndex, PrimitiveSceneInfo, StaticMeshRelevance, StaticMesh, Scene, bCanCache, EMeshPass::Distortion);
 							}
 
-							if (ViewRelevance.bSeparateTranslucencyRelevance)
+							if (ShadingPath == EShadingPath::Mobile && View.bIsSceneCapture)
 							{
-								DrawCommandPacket.AddCommandsForMesh(PrimitiveIndex, PrimitiveSceneInfo, StaticMeshRelevance, StaticMesh, Scene, bCanCache, EMeshPass::TranslucencyAfterDOF);
+								DrawCommandPacket.AddCommandsForMesh(PrimitiveIndex, PrimitiveSceneInfo, StaticMeshRelevance, StaticMesh, Scene, bCanCache, EMeshPass::MobileInverseOpacity);
 							}
 						}
-						else
+
+#if WITH_EDITOR
+						if (ViewRelevance.bEditorStaticSelectionRelevance)
 						{
-							// Otherwise, everything is rendered in a single bucket. This is not related to whether DOF is currently enabled or not.
-							// When using all translucency, Standard and AfterDOF are sorted together instead of being rendered like 2 buckets.
-							DrawCommandPacket.AddCommandsForMesh(PrimitiveIndex, PrimitiveSceneInfo, StaticMeshRelevance, StaticMesh, Scene, bCanCache, EMeshPass::TranslucencyAll);
+							DrawCommandPacket.AddCommandsForMesh(PrimitiveIndex, PrimitiveSceneInfo, StaticMeshRelevance, StaticMesh, Scene, bCanCache, EMeshPass::EditorSelection);
 						}
+#endif
 
-						if (ViewRelevance.bDistortionRelevance)
+						if (ViewRelevance.bHasVolumeMaterialDomain)
 						{
-							DrawCommandPacket.AddCommandsForMesh(PrimitiveIndex, PrimitiveSceneInfo, StaticMeshRelevance, StaticMesh, Scene, bCanCache, EMeshPass::Distortion);
+							VolumetricMeshBatches.AddUninitialized(1);
+							FVolumetricMeshBatch& BatchAndProxy = VolumetricMeshBatches.Last();
+							BatchAndProxy.Mesh = &StaticMesh;
+							BatchAndProxy.Proxy = PrimitiveSceneInfo->Proxy;
 						}
 
-						if (ShadingPath == EShadingPath::Mobile && View.bIsSceneCapture)
+						if (ViewRelevance.bRenderInMainPass && ViewRelevance.bDecal)
 						{
-							DrawCommandPacket.AddCommandsForMesh(PrimitiveIndex, PrimitiveSceneInfo, StaticMeshRelevance, StaticMesh, Scene, bCanCache, EMeshPass::MobileInverseOpacity);
+							MeshDecalBatches.AddUninitialized(1);
+							FMeshDecalBatch& BatchAndProxy = MeshDecalBatches.Last();
+							BatchAndProxy.Mesh = &StaticMesh;
+							BatchAndProxy.Proxy = PrimitiveSceneInfo->Proxy;
+							BatchAndProxy.SortKey = PrimitiveSceneInfo->Proxy->GetTranslucencySortPriority();
 						}
 					}
 
-#if WITH_EDITOR
-					if(ViewRelevance.bDrawRelevance && ViewRelevance.bEditorStaticSelectionRelevance)
-					{
-						DrawCommandPacket.AddCommandsForMesh(PrimitiveIndex, PrimitiveSceneInfo, StaticMeshRelevance, StaticMesh, Scene, bCanCache, EMeshPass::EditorSelection);
-					}
-#endif
-
-					if (ViewRelevance.bHasVolumeMaterialDomain)
-					{
-						VolumetricMeshBatches.AddUninitialized(1);
-						FVolumetricMeshBatch& BatchAndProxy = VolumetricMeshBatches.Last();
-						BatchAndProxy.Mesh = &StaticMesh;
-						BatchAndProxy.Proxy = PrimitiveSceneInfo->Proxy;
-					}
-
-					if (ViewRelevance.bRenderInMainPass && ViewRelevance.bDecal)
-					{
-						MeshDecalBatches.AddUninitialized(1);
-						FMeshDecalBatch& BatchAndProxy = MeshDecalBatches.Last();
-						BatchAndProxy.Mesh = &StaticMesh;
-						BatchAndProxy.Proxy = PrimitiveSceneInfo->Proxy;
-						BatchAndProxy.SortKey = PrimitiveSceneInfo->Proxy->GetTranslucencySortPriority();
-					}
-
 					if (MarkMask)
 					{
 						MarkMasks[StaticMeshRelevance.Id] = MarkMask;
