Output correct size value for PP draw event.

diff --git a/Engine/Source/Runtime/Renderer/Private/PostProcess/PostProcessMaterial.cpp b/Engine/Source/Runtime/Renderer/Private/PostProcess/PostProcessMaterial.cpp
--- a/Engine/Source/Runtime/Renderer/Private/PostProcess/PostProcessMaterial.cpp
+++ b/Engine/Source/Runtime/Renderer/Private/PostProcess/PostProcessMaterial.cpp
@@ -288,9 +288,7 @@ void FRCPassPostProcessMaterial::Process(FRenderingCompositePassContext& Context
 		}
 	}
 
-	const FPooledRenderTargetDesc* InputDesc = GetInputDesc(ePId_Input0);
-
-	FIntVector RectSize = GetOutput(ePId_Output0)->RenderTargetDesc.GetSize();
+	FIntPoint RectSize = Context.SceneColorViewRect.Size();
 	SCOPED_DRAW_EVENTF(Context.RHICmdList, PostProcessMaterial, TEXT("PostProcessMaterial %dx%d Material=%s"), RectSize.X, RectSize.Y, *Material->GetFriendlyName());
 
 	// Copy of custom depth/stencil buffer if HW does not support simultaneously a texture bound as DepthRead_StencilRead and SRV
@@ -436,7 +434,7 @@ void FRCPassPostProcessMaterial::Process(FRenderingCompositePassContext& Context
 		const FViewInfo& View = Context.View;
 		const FSceneViewFamily& ViewFamily = *(View.Family);
 
-		FIntPoint SrcSize = InputDesc->Extent;
+		FIntPoint SrcSize = GetInputDesc(ePId_Input0)->Extent;
 
 		FShader* VertexShader = nullptr;
 
