Don't attempt to recycle PP render target if it has already been allocated. This can happen when PPM is the last effect in the PP chain as being output directly to the View->RT

diff --git a/Engine/Source/Runtime/Renderer/Private/PostProcess/PostProcessMaterial.cpp b/Engine/Source/Runtime/Renderer/Private/PostProcess/PostProcessMaterial.cpp
--- a/Engine/Source/Runtime/Renderer/Private/PostProcess/PostProcessMaterial.cpp
+++ b/Engine/Source/Runtime/Renderer/Private/PostProcess/PostProcessMaterial.cpp
@@ -370,7 +370,7 @@ void FRCPassPostProcessMaterial::Process(FRenderingCompositePassContext& Context
 		
 	if (bDoStencilTest || bDoOutputBlend)
 	{
-		if (!MaterialShaderMap->UsesSceneTexture(PPI_PostProcessInput0))
+		if (!MaterialShaderMap->UsesSceneTexture(PPI_PostProcessInput0) && !PassOutputs[0].PooledRenderTarget)
 		{
 			PassOutputs[0].PooledRenderTarget = GetInput(ePId_Input0)->GetOutput()->RequestInput();
 			DestRenderTarget = &PassOutputs[0].RequestSurface(Context);
