Fix thread safety of dynamic PSO lookups.  This fixes a race condition in the renderer that can cause crashes and potentially GPU hangs on almost all platforms.

diff --git a/Engine/Source/Runtime/Renderer/Private/MeshDrawCommands.cpp b/Engine/Source/Runtime/Renderer/Private/MeshDrawCommands.cpp
--- a/Engine/Source/Runtime/Renderer/Private/MeshDrawCommands.cpp
+++ b/Engine/Source/Runtime/Renderer/Private/MeshDrawCommands.cpp
@@ -668,7 +668,12 @@ void ApplyViewOverridesToMeshDrawCommands(
 				NewMeshCommand = MeshCommand;
 
 				const ERasterizerCullMode LocalCullMode = bRenderSceneTwoSided ? CM_None : bReverseCulling ? FMeshPassProcessor::InverseCullMode(VisibleMeshDrawCommand.MeshCullMode) : VisibleMeshDrawCommand.MeshCullMode;
-				FGraphicsMinimalPipelineStateInitializer PipelineState = MeshCommand.CachedPipelineId.GetPipelineState();
+
+				FGraphicsMinimalPipelineStateInitializer PipelineState;
+				{
+					FGraphicsMinimalPipelineStateId::FPipelineStateIdLookupScope SafeLookupScope(MeshCommand.CachedPipelineId);
+					PipelineState = SafeLookupScope.GetPipelineState();
+				}
 				PipelineState.RasterizerState = GetStaticRasterizerState<true>(VisibleMeshDrawCommand.MeshFillMode, LocalCullMode);
 
 				if (BasePassDepthStencilAccess != DefaultBasePassDepthStencilAccess && PassType == EMeshPass::BasePass)
diff --git a/Engine/Source/Runtime/Renderer/Private/MeshPassProcessor.cpp b/Engine/Source/Runtime/Renderer/Private/MeshPassProcessor.cpp
--- a/Engine/Source/Runtime/Renderer/Private/MeshPassProcessor.cpp
+++ b/Engine/Source/Runtime/Renderer/Private/MeshPassProcessor.cpp
@@ -16,7 +16,7 @@
 
 TSet<FRefCountedGraphicsMinimalPipelineStateInitializer, RefCountedGraphicsMinimalPipelineStateInitializerKeyFuncs> FGraphicsMinimalPipelineStateId::PersistentIdTable;
 TSet<FGraphicsMinimalPipelineStateInitializer> FGraphicsMinimalPipelineStateId::OneFrameIdTable;
-FCriticalSection FGraphicsMinimalPipelineStateId::OneFrameIdTableCriticalSection;
+FRWLock FGraphicsMinimalPipelineStateId::OneFrameIdTableCriticalSection;
 
 const FMeshDrawCommandSortKey FMeshDrawCommandSortKey::Default = { {0} };
 
@@ -29,6 +29,12 @@ static FAutoConsoleVariableRef CVarEmitMeshDrawEvent(
 	ECVF_RenderThreadSafe
 );
 
+static TAutoConsoleVariable<int32> CVarSafeStateLookup(
+	TEXT("r.SafeStateLookup"),
+	1,
+	TEXT("Forces new-style safe state lookup for easy runtime perf comparison\n"),
+	ECVF_Scalability | ECVF_RenderThreadSafe);
+
 enum { MAX_SRVs_PER_SHADER_STAGE = 128 };
 enum { MAX_UNIFORM_BUFFERS_PER_SHADER_STAGE = 14 };
 enum { MAX_SAMPLERS_PER_SHADER_STAGE = 32 };
@@ -344,9 +350,7 @@ void FGraphicsMinimalPipelineStateId::RemovePersistentId(FGraphicsMinimalPipelin
 
 FGraphicsMinimalPipelineStateId FGraphicsMinimalPipelineStateId::GetOneFrameId(const FGraphicsMinimalPipelineStateInitializer& InPipelineState)
 {
-	// Need to lock as this is called from multiple parallel tasks during mesh draw command generation or patching.
-	FScopeLock Lock(&OneFrameIdTableCriticalSection);
-
+	// Need to lock as this is called from multiple parallel tasks during mesh draw command generation or patching.	
 	FGraphicsMinimalPipelineStateId Ret;
 	Ret.bValid = 1;
 	Ret.bOneFrameId = 0;
@@ -354,6 +358,7 @@ FGraphicsMinimalPipelineStateId FGraphicsMinimalPipelineStateId::GetOneFrameId(c
 	FSetElementId TableId = PersistentIdTable.FindId(InPipelineState);
 	if (!TableId.IsValidId())
 	{
+		FRWScopeLock ScopeLock(OneFrameIdTableCriticalSection, SLT_Write);
 		Ret.bOneFrameId = 1;
 
 		TableId = OneFrameIdTable.FindId(InPipelineState);
@@ -369,6 +374,33 @@ FGraphicsMinimalPipelineStateId FGraphicsMinimalPipelineStateId::GetOneFrameId(c
 	return Ret;
 }
 
+FGraphicsMinimalPipelineStateId::FPipelineStateIdLookupScope::FPipelineStateIdLookupScope(const FGraphicsMinimalPipelineStateId& LookupID)
+{
+	bLocked = false;
+	if (LookupID.bOneFrameId)
+	{
+		if (CVarSafeStateLookup.GetValueOnRenderThread())
+		{
+			OneFrameIdTableCriticalSection.ReadLock();
+			bLocked = true;
+		}
+	}
+	SafeStateRef = &LookupID.GetPipelineState();
+}
+
+FGraphicsMinimalPipelineStateId::FPipelineStateIdLookupScope::~FPipelineStateIdLookupScope()
+{
+	if (bLocked)
+	{
+		OneFrameIdTableCriticalSection.ReadUnlock();
+	}
+}
+
+const FGraphicsMinimalPipelineStateInitializer& FGraphicsMinimalPipelineStateId::FPipelineStateIdLookupScope::GetPipelineState()
+{
+	return *SafeStateRef;
+}
+
 void FGraphicsMinimalPipelineStateId::ResetOneFrameIdTable()
 {
 	OneFrameIdTable.Reset();
@@ -862,39 +894,42 @@ void FMeshDrawCommand::SubmitDraw(
 	}
 #endif
 
-	const FGraphicsMinimalPipelineStateInitializer& MeshPipelineState = MeshDrawCommand.CachedPipelineId.GetPipelineState();
-
-	if (MeshDrawCommand.CachedPipelineId.GetId() != StateCache.PipelineId)
 	{
-		FGraphicsPipelineStateInitializer GraphicsPSOInit = MeshPipelineState;
-		RHICmdList.ApplyCachedRenderTargets(GraphicsPSOInit);
-		SetGraphicsPipelineState(RHICmdList, GraphicsPSOInit);
-		StateCache.SetPipelineState(MeshDrawCommand.CachedPipelineId.GetId());
-	}
+		FGraphicsMinimalPipelineStateId::FPipelineStateIdLookupScope SafeLookupScope(MeshDrawCommand.CachedPipelineId);
+		const FGraphicsMinimalPipelineStateInitializer& MeshPipelineState = SafeLookupScope.GetPipelineState();
 
-	if (MeshDrawCommand.StencilRef != StateCache.StencilRef)
-	{
-		RHICmdList.SetStencilRef(MeshDrawCommand.StencilRef);
-		StateCache.StencilRef = MeshDrawCommand.StencilRef;
-	}
-
-	for (int32 VertexBindingIndex = 0; VertexBindingIndex < MeshDrawCommand.VertexStreams.Num(); VertexBindingIndex++)
-	{
-		const FVertexInputStream& Stream = MeshDrawCommand.VertexStreams[VertexBindingIndex];
+		if (MeshDrawCommand.CachedPipelineId.GetId() != StateCache.PipelineId)
+		{
+			FGraphicsPipelineStateInitializer GraphicsPSOInit = MeshPipelineState;
+			RHICmdList.ApplyCachedRenderTargets(GraphicsPSOInit);
+			SetGraphicsPipelineState(RHICmdList, GraphicsPSOInit);
+			StateCache.SetPipelineState(MeshDrawCommand.CachedPipelineId.GetId());
+		}
 
-		if (MeshDrawCommand.PrimitiveIdStreamIndex != -1 && Stream.StreamIndex == MeshDrawCommand.PrimitiveIdStreamIndex)
+		if (MeshDrawCommand.StencilRef != StateCache.StencilRef)
 		{
-			RHICmdList.SetStreamSource(Stream.StreamIndex, ScenePrimitiveIdsBuffer, PrimitiveIdOffset);
-			StateCache.VertexStreams[Stream.StreamIndex] = Stream;
+			RHICmdList.SetStencilRef(MeshDrawCommand.StencilRef);
+			StateCache.StencilRef = MeshDrawCommand.StencilRef;
 		}
-		else if (StateCache.VertexStreams[Stream.StreamIndex] != Stream)
+
+		for (int32 VertexBindingIndex = 0; VertexBindingIndex < MeshDrawCommand.VertexStreams.Num(); VertexBindingIndex++)
 		{
-			RHICmdList.SetStreamSource(Stream.StreamIndex, Stream.VertexBuffer, Stream.Offset);
-			StateCache.VertexStreams[Stream.StreamIndex] = Stream;
+			const FVertexInputStream& Stream = MeshDrawCommand.VertexStreams[VertexBindingIndex];
+
+			if (MeshDrawCommand.PrimitiveIdStreamIndex != -1 && Stream.StreamIndex == MeshDrawCommand.PrimitiveIdStreamIndex)
+			{
+				RHICmdList.SetStreamSource(Stream.StreamIndex, ScenePrimitiveIdsBuffer, PrimitiveIdOffset);
+				StateCache.VertexStreams[Stream.StreamIndex] = Stream;
+			}
+			else if (StateCache.VertexStreams[Stream.StreamIndex] != Stream)
+			{
+				RHICmdList.SetStreamSource(Stream.StreamIndex, Stream.VertexBuffer, Stream.Offset);
+				StateCache.VertexStreams[Stream.StreamIndex] = Stream;
+			}
 		}
-	}
 
-	MeshDrawCommand.ShaderBindings.SetOnCommandList(RHICmdList, MeshPipelineState.BoundShaderState, StateCache.ShaderBindings);
+		MeshDrawCommand.ShaderBindings.SetOnCommandList(RHICmdList, MeshPipelineState.BoundShaderState, StateCache.ShaderBindings);
+	}
 
 	if (MeshDrawCommand.IndexBuffer)
 	{
@@ -913,10 +948,10 @@ void FMeshDrawCommand::SubmitDraw(
 		else
 		{
 			RHICmdList.DrawIndexedPrimitiveIndirect(
-				MeshDrawCommand.IndexBuffer, 
-				MeshDrawCommand.IndirectArgsBuffer, 
+				MeshDrawCommand.IndexBuffer,
+				MeshDrawCommand.IndirectArgsBuffer,
 				0
-				);
+			);
 		}
 	}
 	else
@@ -927,7 +962,6 @@ void FMeshDrawCommand::SubmitDraw(
 			MeshDrawCommand.NumInstances * InstanceFactor
 		);
 	}
-
 }
 
 void SubmitMeshDrawCommands(
diff --git a/Engine/Source/Runtime/Renderer/Private/ShadowSetup.cpp b/Engine/Source/Runtime/Renderer/Private/ShadowSetup.cpp
--- a/Engine/Source/Runtime/Renderer/Private/ShadowSetup.cpp
+++ b/Engine/Source/Runtime/Renderer/Private/ShadowSetup.cpp
@@ -1345,7 +1345,12 @@ void FProjectedShadowInfo::ApplyViewOverridesToMeshDrawCommands(const FViewInfo&
 			NewMeshCommand = MeshCommand;
 
 			const ERasterizerCullMode LocalCullMode = View.bRenderSceneTwoSided ? CM_None : View.bReverseCulling ? FMeshPassProcessor::InverseCullMode(VisibleMeshDrawCommand.MeshCullMode) : VisibleMeshDrawCommand.MeshCullMode;
-			FGraphicsMinimalPipelineStateInitializer PipelineState = MeshCommand.CachedPipelineId.GetPipelineState();
+
+			FGraphicsMinimalPipelineStateInitializer PipelineState;
+			{
+				FGraphicsMinimalPipelineStateId::FPipelineStateIdLookupScope SafeLookupScope(MeshCommand.CachedPipelineId);
+				PipelineState = SafeLookupScope.GetPipelineState();
+			}
 			PipelineState.RasterizerState = GetStaticRasterizerState<true>(VisibleMeshDrawCommand.MeshFillMode, LocalCullMode);
 
 			const FGraphicsMinimalPipelineStateId PipelineId = FGraphicsMinimalPipelineStateId::GetOneFrameId(PipelineState);
diff --git a/Engine/Source/Runtime/Renderer/Public/MeshPassProcessor.h b/Engine/Source/Runtime/Renderer/Public/MeshPassProcessor.h
--- a/Engine/Source/Runtime/Renderer/Public/MeshPassProcessor.h
+++ b/Engine/Source/Runtime/Renderer/Public/MeshPassProcessor.h
@@ -156,18 +156,6 @@ public:
 		return !(*this == rhs);
 	}
 	
-	inline const FGraphicsMinimalPipelineStateInitializer& GetPipelineState() const
-	{
-		const FSetElementId SetElementId = FSetElementId::FromInteger(SetElementIndex);
-
-		if (bOneFrameId)
-		{
-			return OneFrameIdTable[SetElementId];
-		}
-
-		return PersistentIdTable[SetElementId].StateInitializer;
-	}
-
 	/**
 	 * Get a ref counted persistent pipeline id, which needs to manually released.
 	 */
@@ -188,7 +176,36 @@ public:
 	static int32 GetPersistentIdNum() { return PersistentIdTable.Num(); }
 	static SIZE_T GetOneFrameIdTableSize() { return OneFrameIdTable.GetAllocatedSize(); }
 
+	class FPipelineStateIdLookupScope
+	{
+	public:
+		FPipelineStateIdLookupScope(const FGraphicsMinimalPipelineStateId& LookupID);
+		~FPipelineStateIdLookupScope();
+		
+		const FGraphicsMinimalPipelineStateInitializer& GetPipelineState();
+	private:
+		FPipelineStateIdLookupScope() {}
+		const FGraphicsMinimalPipelineStateInitializer* SafeStateRef;
+		bool bLocked;
+	};
+
 private:
+
+	friend class FPipelineStateIdLookupScope;
+
+	//this reference isn't safe unless it's wrapped in a scope with an FRWLock around OneFrameIDTable.
+	inline const FGraphicsMinimalPipelineStateInitializer& GetPipelineState() const
+	{
+		const FSetElementId SetElementId = FSetElementId::FromInteger(SetElementIndex);
+
+		if (bOneFrameId)
+		{
+			return OneFrameIdTable[SetElementId];
+		}
+
+		return PersistentIdTable[SetElementId].StateInitializer;
+	}
+
 	union
 	{
 		uint32 PackedId = 0;
@@ -203,7 +220,7 @@ private:
 
 	static TSet<FRefCountedGraphicsMinimalPipelineStateInitializer, RefCountedGraphicsMinimalPipelineStateInitializerKeyFuncs> PersistentIdTable;
 	static TSet<FGraphicsMinimalPipelineStateInitializer> OneFrameIdTable;
-	static FCriticalSection OneFrameIdTableCriticalSection;
+	static FRWLock OneFrameIdTableCriticalSection;
 };
 
 struct FMeshProcessorShaders
