More efficent iteration over scene primitives that need their commands re-cached (saves 0.3ms on PC, 0.5ms on Mobile)

diff --git a/Engine/Source/Runtime/Renderer/Private/PrimitiveSceneInfo.cpp b/Engine/Source/Runtime/Renderer/Private/PrimitiveSceneInfo.cpp
--- a/Engine/Source/Runtime/Renderer/Private/PrimitiveSceneInfo.cpp
+++ b/Engine/Source/Runtime/Renderer/Private/PrimitiveSceneInfo.cpp
@@ -126,7 +126,6 @@ FPrimitiveSceneInfo::FPrimitiveSceneInfo(UPrimitiveComponent* InComponent,FScene
 #endif
 	PackedIndex(INDEX_NONE),
 	ComponentForDebuggingOnly(InComponent),
-	bNeedsStaticMeshUpdate(false),
 	bNeedsStaticMeshUpdateWithoutVisibilityCheck(false),
 	bNeedsUniformBufferUpdate(false),
 	bIndirectLightingCacheBufferDirty(false),
@@ -638,11 +637,9 @@ void FPrimitiveSceneInfo::RemoveFromScene(bool bUpdateStaticDrawLists)
 
 	if (bUpdateStaticDrawLists)
 	{
-		if (bNeedsStaticMeshUpdate)
+		if (IsIndexValid()) // PackedIndex
 		{
-			Scene->PrimitivesNeedingStaticMeshUpdate.Remove(this);
-
-			bNeedsStaticMeshUpdate = false;
+			Scene->PrimitivesNeedingStaticMeshUpdate[PackedIndex] = false;
 		}
 
 		if (bNeedsStaticMeshUpdateWithoutVisibilityCheck)
@@ -660,23 +657,16 @@ void FPrimitiveSceneInfo::RemoveFromScene(bool bUpdateStaticDrawLists)
 	}
 }
 
+bool FPrimitiveSceneInfo::NeedsUpdateStaticMeshes()
+{
+	return Scene->PrimitivesNeedingStaticMeshUpdate[PackedIndex];
+}
+
 void FPrimitiveSceneInfo::UpdateStaticMeshes(FRHICommandListImmediate& RHICmdList, bool bReAddToDrawLists)
 {
 	QUICK_SCOPE_CYCLE_COUNTER(STAT_FPrimitiveSceneInfo_UpdateStaticMeshes);
-	const bool bOriginalNeedsStaticMeshUpdate = bNeedsStaticMeshUpdate;
-	bNeedsStaticMeshUpdate = !bReAddToDrawLists;
-
-	if (bOriginalNeedsStaticMeshUpdate != bNeedsStaticMeshUpdate)
-	{
-		if (bNeedsStaticMeshUpdate)
-		{
-			Scene->PrimitivesNeedingStaticMeshUpdate.Add(this);
-		}
-		else
-		{
-			Scene->PrimitivesNeedingStaticMeshUpdate.Remove(this);
-		}
-	}
+	const bool bNeedsStaticMeshUpdate = !bReAddToDrawLists;
+	Scene->PrimitivesNeedingStaticMeshUpdate[PackedIndex] = bNeedsStaticMeshUpdate;
 
 	if (!bNeedsStaticMeshUpdate && bNeedsStaticMeshUpdateWithoutVisibilityCheck)
 	{
@@ -701,18 +691,16 @@ void FPrimitiveSceneInfo::UpdateUniformBuffer(FRHICommandListImmediate& RHICmdLi
 
 void FPrimitiveSceneInfo::BeginDeferredUpdateStaticMeshes()
 {
-	if (!bNeedsStaticMeshUpdate)
+	// Set a flag which causes InitViews to update the static meshes the next time the primitive is visible.
+	if (IsIndexValid()) // PackedIndex
 	{
-		// Set a flag which causes InitViews to update the static meshes the next time the primitive is visible.
-		bNeedsStaticMeshUpdate = true;
-
-		Scene->PrimitivesNeedingStaticMeshUpdate.Add(this);
+		Scene->PrimitivesNeedingStaticMeshUpdate[PackedIndex] = true;
 	}
 }
 
 void FPrimitiveSceneInfo::BeginDeferredUpdateStaticMeshesWithoutVisibilityCheck()
 {
-	if (bNeedsStaticMeshUpdate && !bNeedsStaticMeshUpdateWithoutVisibilityCheck)
+	if (NeedsUpdateStaticMeshes() && !bNeedsStaticMeshUpdateWithoutVisibilityCheck)
 	{
 		bNeedsStaticMeshUpdateWithoutVisibilityCheck = true;
 
diff --git a/Engine/Source/Runtime/Renderer/Private/RendererScene.cpp b/Engine/Source/Runtime/Renderer/Private/RendererScene.cpp
--- a/Engine/Source/Runtime/Renderer/Private/RendererScene.cpp
+++ b/Engine/Source/Runtime/Renderer/Private/RendererScene.cpp
@@ -543,6 +543,7 @@ void FScene::CheckPrimitiveArrays()
 	check(Primitives.Num() == PrimitiveOcclusionFlags.Num());
 	check(Primitives.Num() == PrimitiveComponentIds.Num());
 	check(Primitives.Num() == PrimitiveOcclusionBounds.Num());
+	check(Primitives.Num() == PrimitivesNeedingStaticMeshUpdate.Num());
 }
 
 
@@ -735,6 +736,16 @@ static void TArraySwapElements(TArray<T>& Array, int i1, int i2)
 	Array[i2] = tmp;
 }
 
+static void TBitArraySwapElements(TBitArray<>& Array, int32 i1, int32 i2)
+{
+	FBitReference BitRef1 = Array[i1];
+	FBitReference BitRef2 = Array[i2];
+	bool Bit1 = BitRef1;
+	bool Bit2 = BitRef2;
+	BitRef1 = Bit2;
+	BitRef2 = Bit1;
+}
+
 void FScene::AddPrimitiveSceneInfo_RenderThread(FRHICommandListImmediate& RHICmdList, FPrimitiveSceneInfo* PrimitiveSceneInfo)
 {
 	SCOPE_CYCLE_COUNTER(STAT_AddScenePrimitiveRenderThreadTime);
@@ -751,6 +762,7 @@ void FScene::AddPrimitiveSceneInfo_RenderThread(FRHICommandListImmediate& RHICmd
 	PrimitiveOcclusionFlags.AddUninitialized();
 	PrimitiveComponentIds.AddUninitialized();
 	PrimitiveOcclusionBounds.AddUninitialized();
+	PrimitivesNeedingStaticMeshUpdate.Add(false);
 	
 	const int SourceIndex = PrimitiveSceneProxies.Num() - 1;
 	PrimitiveSceneInfo->PackedIndex = SourceIndex;
@@ -818,6 +830,7 @@ void FScene::AddPrimitiveSceneInfo_RenderThread(FRHICommandListImmediate& RHICmd
 				TArraySwapElements(PrimitiveOcclusionFlags, DestIndex, SourceIndex);
 				TArraySwapElements(PrimitiveComponentIds, DestIndex, SourceIndex);
 				TArraySwapElements(PrimitiveOcclusionBounds, DestIndex, SourceIndex);
+				TBitArraySwapElements(PrimitivesNeedingStaticMeshUpdate, DestIndex, SourceIndex);
 
 				AddPrimitiveToUpdateGPU(*this, DestIndex);
 			}
@@ -1482,6 +1495,7 @@ void FScene::RemovePrimitiveSceneInfo_RenderThread(FPrimitiveSceneInfo* Primitiv
 				TArraySwapElements(PrimitiveOcclusionFlags, DestIndex, SourceIndex);
 				TArraySwapElements(PrimitiveComponentIds, DestIndex, SourceIndex);
 				TArraySwapElements(PrimitiveOcclusionBounds, DestIndex, SourceIndex);
+				TBitArraySwapElements(PrimitivesNeedingStaticMeshUpdate, DestIndex, SourceIndex);
 				SourceIndex = DestIndex;
 
 				AddPrimitiveToUpdateGPU(*this, DestIndex);
@@ -1512,6 +1526,7 @@ void FScene::RemovePrimitiveSceneInfo_RenderThread(FPrimitiveSceneInfo* Primitiv
 		PrimitiveOcclusionFlags.Pop();
 		PrimitiveComponentIds.Pop();
 		PrimitiveOcclusionBounds.Pop();
+		PrimitivesNeedingStaticMeshUpdate.RemoveAt(PrimitivesNeedingStaticMeshUpdate.Num()-1);
 	}
 
 	PrimitiveSceneInfo->PackedIndex = MAX_int32;
diff --git a/Engine/Source/Runtime/Renderer/Private/ScenePrivate.h b/Engine/Source/Runtime/Renderer/Private/ScenePrivate.h
--- a/Engine/Source/Runtime/Renderer/Private/ScenePrivate.h
+++ b/Engine/Source/Runtime/Renderer/Private/ScenePrivate.h
@@ -2497,7 +2497,7 @@ public:
 	/** Packed array of primitive components associated with the primitive. */
 	TArray<FPrimitiveComponentId> PrimitiveComponentIds;
 
-	TSet<FPrimitiveSceneInfo*> PrimitivesNeedingStaticMeshUpdate;
+	TBitArray<> PrimitivesNeedingStaticMeshUpdate;
 	TSet<FPrimitiveSceneInfo*> PrimitivesNeedingStaticMeshUpdateWithoutVisibilityCheck;
 
 	struct FTypeOffsetTableEntry
diff --git a/Engine/Source/Runtime/Renderer/Private/SceneVisibility.cpp b/Engine/Source/Runtime/Renderer/Private/SceneVisibility.cpp
--- a/Engine/Source/Runtime/Renderer/Private/SceneVisibility.cpp
+++ b/Engine/Source/Runtime/Renderer/Private/SceneVisibility.cpp
@@ -3650,15 +3650,10 @@ void FSceneRenderer::ComputeViewVisibility(FRHICommandListImmediate& RHICmdList,
 
 		{
 			QUICK_SCOPE_CYCLE_COUNTER(STAT_ViewVisibilityTime_ConditionalUpdateStaticMeshes);
-
-			for (TSet<FPrimitiveSceneInfo*>::TIterator It(Scene->PrimitivesNeedingStaticMeshUpdate); It; ++It)
+			for (TConstDualSetBitIterator<SceneRenderingBitArrayAllocator, FDefaultBitArrayAllocator> BitIt(View.PrimitiveVisibilityMap, Scene->PrimitivesNeedingStaticMeshUpdate); BitIt; ++BitIt)
 			{
-				FPrimitiveSceneInfo* Primitive = *It;
-
-				if (View.PrimitiveVisibilityMap[Primitive->GetIndex()])
-				{
-					Primitive->ConditionalUpdateStaticMeshes(RHICmdList);
-				}
+				int32 PrimitiveIndex = BitIt.GetIndex();
+				Scene->Primitives[PrimitiveIndex]->UpdateStaticMeshes(RHICmdList);
 			}
 		}
 
diff --git a/Engine/Source/Runtime/Renderer/Public/PrimitiveSceneInfo.h b/Engine/Source/Runtime/Renderer/Public/PrimitiveSceneInfo.h
--- a/Engine/Source/Runtime/Renderer/Public/PrimitiveSceneInfo.h
+++ b/Engine/Source/Runtime/Renderer/Public/PrimitiveSceneInfo.h
@@ -312,10 +312,7 @@ public:
 	void RemoveFromScene(bool bUpdateStaticDrawLists);
 
 	/** return true if we need to call ConditionalUpdateStaticMeshes */
-	FORCEINLINE bool NeedsUpdateStaticMeshes()
-	{
-		return bNeedsStaticMeshUpdate;
-	}
+	bool NeedsUpdateStaticMeshes();
 
 	/** return true if we need to call LazyUpdateForRendering */
 	FORCEINLINE bool NeedsUniformBufferUpdate() const
@@ -463,9 +460,6 @@ private:
 	 */
 	const UPrimitiveComponent* ComponentForDebuggingOnly;
 
-	/** If this is TRUE, this primitive's static meshes needs to be updated before it can be rendered. */
-	bool bNeedsStaticMeshUpdate : 1;
-
 	/** If this is TRUE, this primitive's static meshes will be update even if it's not visible. */
 	bool bNeedsStaticMeshUpdateWithoutVisibilityCheck : 1;
 
