Fixed shadow stenciling (used for pre-shadows for dynamic shadows on baked lighting) from dynamic mesh elements. Dynamic mesh elements were gathered for every shadow map view, while renderer later with the main view, using mismatching GPUScene dynamic data. Solution was to reuse already gathered dynamic mesh elements from main view.

diff --git a/Engine/Source/Runtime/Renderer/Private/SceneRendering.cpp b/Engine/Source/Runtime/Renderer/Private/SceneRendering.cpp
--- a/Engine/Source/Runtime/Renderer/Private/SceneRendering.cpp
+++ b/Engine/Source/Runtime/Renderer/Private/SceneRendering.cpp
@@ -1581,6 +1581,25 @@ void FViewInfo::DestroyAllSnapshots()
 	ViewInfoSnapshots.Reset();
 }
 
+FInt32Range FViewInfo::GetDynamicMeshElementRange(uint32 PrimitiveIndex) const
+{
+	int32 Start = 0;	// inclusive
+	int32 AfterEnd = 0;	// exclusive
+
+	// DynamicMeshEndIndices contains valid values only for visible primitives with bDynamicRelevance.
+	if (PrimitiveVisibilityMap[PrimitiveIndex])
+	{
+		const FPrimitiveViewRelevance& ViewRelevance = PrimitiveViewRelevanceMap[PrimitiveIndex];
+		if (ViewRelevance.bDynamicRelevance)
+		{
+			Start = (PrimitiveIndex == 0) ? 0 : DynamicMeshEndIndices[PrimitiveIndex - 1];
+			AfterEnd = DynamicMeshEndIndices[PrimitiveIndex];
+		}
+	}
+
+	return FInt32Range(Start, AfterEnd);
+}
+
 FSceneViewState* FViewInfo::GetEffectiveViewState() const
 {
 	FSceneViewState* EffectiveViewState = ViewState;
diff --git a/Engine/Source/Runtime/Renderer/Private/SceneRendering.h b/Engine/Source/Runtime/Renderer/Private/SceneRendering.h
--- a/Engine/Source/Runtime/Renderer/Private/SceneRendering.h
+++ b/Engine/Source/Runtime/Renderer/Private/SceneRendering.h
@@ -819,6 +819,9 @@ public:
 	/** Gathered in initviews from all the primitives with dynamic view relevance, used in each mesh pass. */
 	TArray<FMeshBatchAndRelevance,SceneRenderingAllocator> DynamicMeshElements;
 
+	/* [PrimitiveIndex] = end index index in DynamicMeshElements[], to support GetDynamicMeshElementRange(). Contains valid values only for visible primitives with bDynamicRelevance. */
+	TArray<uint32, SceneRenderingAllocator> DynamicMeshEndIndices;
+
 	/* Mesh pass relevance for gathered dynamic mesh elements. */
 	TArray<FMeshPassMask, SceneRenderingAllocator> DynamicMeshElementsPassRelevance;
 
@@ -1130,6 +1133,10 @@ public:
 	/** Destroy all snapshots before we wipe the scene allocator. */
 	static void DestroyAllSnapshots();
 
+	// Get the range in DynamicMeshElements[] for a given PrimitiveIndex
+	// @return range (start is inclusive, end is exclusive)
+	FInt32Range GetDynamicMeshElementRange(uint32 PrimitiveIndex) const;
+
 	/** Set the custom data associated with a primitive scene info.	*/
 	void SetCustomData(const FPrimitiveSceneInfo* InPrimitiveSceneInfo, void* InCustomData);
 
diff --git a/Engine/Source/Runtime/Renderer/Private/SceneVisibility.cpp b/Engine/Source/Runtime/Renderer/Private/SceneVisibility.cpp
--- a/Engine/Source/Runtime/Renderer/Private/SceneVisibility.cpp
+++ b/Engine/Source/Runtime/Renderer/Private/SceneVisibility.cpp
@@ -2716,6 +2716,14 @@ void FSceneRenderer::GatherDynamicMeshElements(
 
 				PrimitiveSceneInfo->Proxy->GetDynamicMeshElements(InViewFamily.Views, InViewFamily, ViewMaskFinal, Collector);
 
+				if (PrimitiveIndex > 0)
+				{
+					for (int32 ViewIndex = 0; ViewIndex < ViewCount; ViewIndex++)
+					{
+						InViews[ViewIndex].DynamicMeshEndIndices[PrimitiveIndex - 1] = Collector.GetMeshBatchCount(ViewIndex);
+					}
+				}
+
 				// Compute DynamicMeshElementsMeshPassRelevance for this primitive.
 				for (int32 ViewIndex = 0; ViewIndex < ViewCount; ViewIndex++)
 				{
@@ -2737,6 +2745,11 @@ void FSceneRenderer::GatherDynamicMeshElements(
 						}
 					}
 				}
+
+				for (int32 ViewIndex = 0; ViewIndex < ViewCount; ViewIndex++)
+				{
+					InViews[ViewIndex].DynamicMeshEndIndices[PrimitiveIndex] = Collector.GetMeshBatchCount(ViewIndex);
+				}
 			}
 		}
 	}
@@ -3445,6 +3458,7 @@ void FSceneRenderer::ComputeViewVisibility(FRHICommandListImmediate& RHICmdList,
 		// Allocate the view's visibility maps.
 		View.PrimitiveVisibilityMap.Init(false,Scene->Primitives.Num());
 		// we don't initialized as we overwrite the whole array (in GatherDynamicMeshElements)
+		View.DynamicMeshEndIndices.SetNumUninitialized(Scene->Primitives.Num());
 		View.PrimitiveDefinitelyUnoccludedMap.Init(false,Scene->Primitives.Num());
 		View.PotentiallyFadingPrimitiveMap.Init(false,Scene->Primitives.Num());
 		View.PrimitiveFadeUniformBuffers.AddZeroed(Scene->Primitives.Num());
diff --git a/Engine/Source/Runtime/Renderer/Private/ShadowRendering.h b/Engine/Source/Runtime/Renderer/Private/ShadowRendering.h
--- a/Engine/Source/Runtime/Renderer/Private/ShadowRendering.h
+++ b/Engine/Source/Runtime/Renderer/Private/ShadowRendering.h
@@ -520,8 +520,6 @@ private:
 
 	/** Dynamic mesh elements for subject primitives. */
 	TArray<FMeshBatchAndRelevance,SceneRenderingAllocator> DynamicSubjectMeshElements;
-	/** Dynamic mesh elements for receiver primitives. */
-	TArray<FMeshBatchAndRelevance,SceneRenderingAllocator> DynamicReceiverMeshElements;
 	/** Dynamic mesh elements for translucent subject primitives. */
 	TArray<FMeshBatchAndRelevance,SceneRenderingAllocator> DynamicSubjectTranslucentMeshElements;
 
diff --git a/Engine/Source/Runtime/Renderer/Private/ShadowSetup.cpp b/Engine/Source/Runtime/Renderer/Private/ShadowSetup.cpp
--- a/Engine/Source/Runtime/Renderer/Private/ShadowSetup.cpp
+++ b/Engine/Source/Runtime/Renderer/Private/ShadowSetup.cpp
@@ -1267,19 +1267,6 @@ void FProjectedShadowInfo::SetupMeshDrawCommandsForProjectionStenciling(FSceneRe
 
 			// Pre-shadows mask by receiver elements, self-shadow mask by subject elements.
 			// Note that self-shadow pre-shadows still mask by receiver elements.
-			const TArray<FMeshBatchAndRelevance, SceneRenderingAllocator>& ProjectionStencilingDynamicMeshElements = bPreShadow ? DynamicReceiverMeshElements : DynamicSubjectMeshElements;
-
-			const int32 NumDynamicMeshBatches = ProjectionStencilingDynamicMeshElements.Num();
-
-			for (int32 MeshIndex = 0; MeshIndex < NumDynamicMeshBatches; MeshIndex++)
-			{
-				const FMeshBatchAndRelevance& MeshAndRelevance = ProjectionStencilingDynamicMeshElements[MeshIndex];
-				check(!MeshAndRelevance.Mesh->bRequiresPerElementVisibility);
-				const uint64 BatchElementMask = ~0ull;
-
-				DepthPassMeshProcessor.AddMeshBatch(*MeshAndRelevance.Mesh, BatchElementMask, MeshAndRelevance.PrimitiveSceneProxy);
-			}
-
 			const PrimitiveArrayType& MaskPrimitives = bPreShadow ? ReceiverPrimitives : DynamicSubjectPrimitives;
 
 			for (int32 PrimitiveIndex = 0, PrimitiveCount = MaskPrimitives.Num(); PrimitiveIndex < PrimitiveCount; PrimitiveIndex++)
@@ -1303,6 +1290,20 @@ void FProjectedShadowInfo::SetupMeshDrawCommandsForProjectionStenciling(FSceneRe
 							}
 						}
 					}
+
+					if (ViewRelevance.bRenderInMainPass && ViewRelevance.bDynamicRelevance)
+					{
+						const FInt32Range MeshBatchRange = View.GetDynamicMeshElementRange(ReceiverPrimitiveSceneInfo->GetIndex());
+
+						for (int32 MeshBatchIndex = MeshBatchRange.GetLowerBoundValue(); MeshBatchIndex < MeshBatchRange.GetUpperBoundValue(); ++MeshBatchIndex)
+						{
+							const FMeshBatchAndRelevance& MeshAndRelevance = View.DynamicMeshElements[MeshBatchIndex];
+							check(!MeshAndRelevance.Mesh->bRequiresPerElementVisibility);
+							const uint64 BatchElementMask = ~0ull;
+
+							DepthPassMeshProcessor.AddMeshBatch(*MeshAndRelevance.Mesh, BatchElementMask, MeshAndRelevance.PrimitiveSceneProxy);
+						}
+					}
 				}
 			}
 
@@ -1414,11 +1415,6 @@ void FProjectedShadowInfo::GatherDynamicMeshElements(FSceneRenderer& Renderer, F
 
 		ShadowDepthView->DrawDynamicFlags = EDrawDynamicFlags::None;
 
-		int32 NumDynamicReceiverMeshElements = 0;
-		ShadowDepthView->SetDynamicMeshElementsShadowCullFrustum(&ReceiverFrustum);
-		GatherDynamicMeshElementsArray(ShadowDepthView, Renderer, DynamicIndexBuffer, DynamicVertexBuffer, DynamicReadBuffer, 
-			ReceiverPrimitives, ReusedViewsArray, DynamicReceiverMeshElements, NumDynamicReceiverMeshElements);
-
 		int32 NumDynamicSubjectTranslucentMeshElements = 0;
 		ShadowDepthView->SetDynamicMeshElementsShadowCullFrustum(&CasterFrustum);
 		GatherDynamicMeshElementsArray(ShadowDepthView, Renderer, DynamicIndexBuffer, DynamicVertexBuffer, DynamicReadBuffer, 
@@ -1545,7 +1541,6 @@ void FProjectedShadowInfo::ClearTransientArrays()
 	DynamicSubjectPrimitives.Empty();
 	ReceiverPrimitives.Empty();
 	DynamicSubjectMeshElements.Empty();
-	DynamicReceiverMeshElements.Empty();
 	DynamicSubjectTranslucentMeshElements.Empty();
 
 	ShadowDepthPassVisibleCommands.Empty();
