Base pass can output velocity data for moveable object with no static lighting, in forward and deferred rendering. This saves cpu draw call submission.
In base vertex and pixel shader pass, uniform variable OutputVelocityBuffer is tested to ouput or not velocity data
Configurable per project and/or per platform (r.BasePassOutputsVelocity)

diff --git a/Engine/Plugins/Experimental/GeometryCache/Source/GeometryCache/Private/GeometryCacheSceneProxy.cpp b/Engine/Plugins/Experimental/GeometryCache/Source/GeometryCache/Private/GeometryCacheSceneProxy.cpp
--- a/Engine/Plugins/Experimental/GeometryCache/Source/GeometryCache/Private/GeometryCacheSceneProxy.cpp
+++ b/Engine/Plugins/Experimental/GeometryCache/Source/GeometryCache/Private/GeometryCacheSceneProxy.cpp
@@ -308,67 +308,67 @@ void FGeometryCacheSceneProxy::CreateMeshBatch(
 	UserData.MeshExtension = FVector::OneVector;
 	UserData.MeshOrigin = FVector::ZeroVector;
 
-	const bool bHasMotionVectors = (
-		TrackProxy->MeshData->VertexInfo.bHasMotionVectors &&
-		TrackProxy->NextFrameMeshData->VertexInfo.bHasMotionVectors &&
-		TrackProxy->MeshData->Positions.Num() == TrackProxy->MeshData->MotionVectors.Num())
-		&& (TrackProxy->NextFrameMeshData->Positions.Num() == TrackProxy->NextFrameMeshData->MotionVectors.Num());
+			const bool bHasMotionVectors = (
+				TrackProxy->MeshData->VertexInfo.bHasMotionVectors &&
+				TrackProxy->NextFrameMeshData->VertexInfo.bHasMotionVectors &&
+				TrackProxy->MeshData->Positions.Num() == TrackProxy->MeshData->MotionVectors.Num())
+				&& (TrackProxy->NextFrameMeshData->Positions.Num() == TrackProxy->NextFrameMeshData->MotionVectors.Num());
 
-	if (!bHasMotionVectors)
-	{
-		UserData.MotionBlurDataExtension = FVector::OneVector;
-		UserData.MotionBlurDataOrigin = FVector::ZeroVector;
-		UserData.MotionBlurPositionScale = 0.0f;
-	}
-	else
-	{
-		UserData.MotionBlurDataExtension = FVector::OneVector * PlaybackSpeed;
-		UserData.MotionBlurDataOrigin = FVector::ZeroVector;
-		UserData.MotionBlurPositionScale = 1.0f;
-	}
+				if (!bHasMotionVectors)
+				{
+					UserData.MotionBlurDataExtension = FVector::OneVector;
+					UserData.MotionBlurDataOrigin = FVector::ZeroVector;
+					UserData.MotionBlurPositionScale = 0.0f;
+				}
+				else
+				{
+					UserData.MotionBlurDataExtension = FVector::OneVector * PlaybackSpeed;
+					UserData.MotionBlurDataOrigin = FVector::ZeroVector;
+					UserData.MotionBlurPositionScale = 1.0f;
+				}
 
-	if (IsRayTracingEnabled())
-	{
-		// No vertex manipulation is allowed in the vertex shader
-		// Otherwise we need an additional compute shader pass to execute the vertex shader and dump to a staging buffer
-		check(UserData.MeshExtension == FVector::OneVector);
-		check(UserData.MeshOrigin == FVector::ZeroVector);
-	}
+				if (IsRayTracingEnabled())
+				{
+					// No vertex manipulation is allowed in the vertex shader
+					// Otherwise we need an additional compute shader pass to execute the vertex shader and dump to a staging buffer
+					check(UserData.MeshExtension == FVector::OneVector);
+					check(UserData.MeshOrigin == FVector::ZeroVector);
+				}
 
-	UserData.PositionBuffer = &TrackProxy->PositionBuffers[TrackProxy->CurrentPositionBufferIndex % 2];
-	UserData.MotionBlurDataBuffer = &TrackProxy->PositionBuffers[(TrackProxy->CurrentPositionBufferIndex+1) % 2];
+				UserData.PositionBuffer = &TrackProxy->PositionBuffers[TrackProxy->CurrentPositionBufferIndex % 2];
+				UserData.MotionBlurDataBuffer = &TrackProxy->PositionBuffers[(TrackProxy->CurrentPositionBufferIndex+1) % 2];
 
-	FGeometryCacheVertexFactoryUniformBufferParameters UniformBufferParameters;
+				FGeometryCacheVertexFactoryUniformBufferParameters UniformBufferParameters;
 
-	UniformBufferParameters.MeshOrigin = UserData.MeshOrigin;
-	UniformBufferParameters.MeshExtension = UserData.MeshExtension;
-	UniformBufferParameters.MotionBlurDataOrigin = UserData.MotionBlurDataOrigin;
-	UniformBufferParameters.MotionBlurDataExtension = UserData.MotionBlurDataExtension;
-	UniformBufferParameters.MotionBlurPositionScale = UserData.MotionBlurPositionScale;
+				UniformBufferParameters.MeshOrigin = UserData.MeshOrigin;
+				UniformBufferParameters.MeshExtension = UserData.MeshExtension;
+				UniformBufferParameters.MotionBlurDataOrigin = UserData.MotionBlurDataOrigin;
+				UniformBufferParameters.MotionBlurDataExtension = UserData.MotionBlurDataExtension;
+				UniformBufferParameters.MotionBlurPositionScale = UserData.MotionBlurPositionScale;
 
-	UserData.UniformBuffer = FGeometryCacheVertexFactoryUniformBufferParametersRef::CreateUniformBufferImmediate(UniformBufferParameters, UniformBuffer_SingleFrame);
-	TrackProxy->VertexFactory.CreateManualVertexFetchUniformBuffer(UserData.PositionBuffer, UserData.MotionBlurDataBuffer, UserData);
+				UserData.UniformBuffer = FGeometryCacheVertexFactoryUniformBufferParametersRef::CreateUniformBufferImmediate(UniformBufferParameters, UniformBuffer_SingleFrame);
+				TrackProxy->VertexFactory.CreateManualVertexFetchUniformBuffer(UserData.PositionBuffer, UserData.MotionBlurDataBuffer, UserData);
 
-	// Draw the mesh.
-	FMeshBatchElement& BatchElement = Mesh.Elements[0];
-	BatchElement.IndexBuffer = &TrackProxy->IndexBuffer;
-	Mesh.VertexFactory = &TrackProxy->VertexFactory;
-	Mesh.SegmentIndex = 0;
+				// Draw the mesh.
+				FMeshBatchElement& BatchElement = Mesh.Elements[0];
+				BatchElement.IndexBuffer = &TrackProxy->IndexBuffer;
+				Mesh.VertexFactory = &TrackProxy->VertexFactory;
+				Mesh.SegmentIndex = 0;
 
-	const FMatrix& LocalToWorldTransform = TrackProxy->WorldMatrix * GetLocalToWorld();
+				const FMatrix& LocalToWorldTransform = TrackProxy->WorldMatrix * GetLocalToWorld();
 
-	DynamicPrimitiveUniformBuffer.Set(LocalToWorldTransform, LocalToWorldTransform, GetBounds(), GetLocalBounds(), true, false, UseEditorDepthTest());
-	BatchElement.PrimitiveUniformBuffer = DynamicPrimitiveUniformBuffer.UniformBuffer.GetUniformBufferRHI();
+				DynamicPrimitiveUniformBuffer.Set(LocalToWorldTransform, LocalToWorldTransform, GetBounds(), GetLocalBounds(), true, false, UseEditorDepthTest(), false);
+				BatchElement.PrimitiveUniformBuffer = DynamicPrimitiveUniformBuffer.UniformBuffer.GetUniformBufferRHI();
 
-	BatchElement.FirstIndex = BatchInfo.StartIndex;
-	BatchElement.NumPrimitives = BatchInfo.NumTriangles;
-	BatchElement.MinVertexIndex = 0;
-	BatchElement.MaxVertexIndex = TrackProxy->MeshData->Positions.Num() - 1;
-	BatchElement.VertexFactoryUserData = &UserDataWrapper.Data;
-	Mesh.ReverseCulling = IsLocalToWorldDeterminantNegative();
-	Mesh.Type = PT_TriangleList;
-	Mesh.DepthPriorityGroup = SDPG_World;
-	Mesh.bCanApplyViewModeOverrides = false;
+				BatchElement.FirstIndex = BatchInfo.StartIndex;
+				BatchElement.NumPrimitives = BatchInfo.NumTriangles;
+				BatchElement.MinVertexIndex = 0;
+				BatchElement.MaxVertexIndex = TrackProxy->MeshData->Positions.Num() - 1;
+				BatchElement.VertexFactoryUserData = &UserDataWrapper.Data;
+				Mesh.ReverseCulling = IsLocalToWorldDeterminantNegative();
+				Mesh.Type = PT_TriangleList;
+				Mesh.DepthPriorityGroup = SDPG_World;
+				Mesh.bCanApplyViewModeOverrides = false;
 }
 
 void FGeometryCacheSceneProxy::GetDynamicMeshElements(const TArray<const FSceneView*>& Views, const FSceneViewFamily& ViewFamily, uint32 VisibilityMap, FMeshElementCollector& Collector) const
diff --git a/Engine/Plugins/FX/Niagara/Source/Niagara/Private/NiagaraRendererMeshes.cpp b/Engine/Plugins/FX/Niagara/Source/Niagara/Private/NiagaraRendererMeshes.cpp
--- a/Engine/Plugins/FX/Niagara/Source/Niagara/Private/NiagaraRendererMeshes.cpp
+++ b/Engine/Plugins/FX/Niagara/Source/Niagara/Private/NiagaraRendererMeshes.cpp
@@ -169,7 +169,8 @@ void NiagaraRendererMeshes::GetDynamicMeshElements(const TArray<const FSceneView
 				SceneProxy->GetLightingChannelMask(),
 				0,
 				INDEX_NONE,
-				INDEX_NONE
+				INDEX_NONE,
+				false
 				);
 			WorldSpacePrimitiveUniformBuffer.SetContents(PrimitiveUniformShaderParameters);
 			WorldSpacePrimitiveUniformBuffer.InitResource();
diff --git a/Engine/Plugins/FX/Niagara/Source/Niagara/Private/NiagaraRendererRibbons.cpp b/Engine/Plugins/FX/Niagara/Source/Niagara/Private/NiagaraRendererRibbons.cpp
--- a/Engine/Plugins/FX/Niagara/Source/Niagara/Private/NiagaraRendererRibbons.cpp
+++ b/Engine/Plugins/FX/Niagara/Source/Niagara/Private/NiagaraRendererRibbons.cpp
@@ -136,7 +136,8 @@ void NiagaraRendererRibbons::GetDynamicMeshElements(const TArray<const FSceneVie
 				SceneProxy->GetLightingChannelMask(),
 				0,
 				INDEX_NONE,
-				INDEX_NONE
+				INDEX_NONE,
+				false
 				);
 			WorldSpacePrimitiveUniformBuffer.SetContents(PrimitiveUniformShaderParameters);
 			WorldSpacePrimitiveUniformBuffer.InitResource();
diff --git a/Engine/Plugins/FX/Niagara/Source/Niagara/Private/NiagaraRendererSprites.cpp b/Engine/Plugins/FX/Niagara/Source/Niagara/Private/NiagaraRendererSprites.cpp
--- a/Engine/Plugins/FX/Niagara/Source/Niagara/Private/NiagaraRendererSprites.cpp
+++ b/Engine/Plugins/FX/Niagara/Source/Niagara/Private/NiagaraRendererSprites.cpp
@@ -134,7 +134,8 @@ void NiagaraRendererSprites::ConditionalInitPrimitiveUniformBuffer(const FNiagar
 			SceneProxy->GetLightingChannelMask(),
 			0,
 			INDEX_NONE,
-			INDEX_NONE
+			INDEX_NONE,
+			false
 		);
 		WorldSpacePrimitiveUniformBuffer.SetContents(PrimitiveUniformShaderParameters);
 		WorldSpacePrimitiveUniformBuffer.InitResource();
diff --git a/Engine/Plugins/Runtime/CableComponent/Source/CableComponent/Private/CableComponent.cpp b/Engine/Plugins/Runtime/CableComponent/Source/CableComponent/Private/CableComponent.cpp
--- a/Engine/Plugins/Runtime/CableComponent/Source/CableComponent/Private/CableComponent.cpp
+++ b/Engine/Plugins/Runtime/CableComponent/Source/CableComponent/Private/CableComponent.cpp
@@ -285,10 +285,11 @@ public:
 				bool bHasPrecomputedVolumetricLightmap;
 				FMatrix PreviousLocalToWorld;
 				int32 SingleCaptureIndex;
-				GetScene().GetPrimitiveUniformShaderParameters_RenderThread(GetPrimitiveSceneInfo(), bHasPrecomputedVolumetricLightmap, PreviousLocalToWorld, SingleCaptureIndex);
+				bool bOutputVelocity;
+				GetScene().GetPrimitiveUniformShaderParameters_RenderThread(GetPrimitiveSceneInfo(), bHasPrecomputedVolumetricLightmap, PreviousLocalToWorld, SingleCaptureIndex, bOutputVelocity);
 
 				FDynamicPrimitiveUniformBuffer& DynamicPrimitiveUniformBuffer = Collector.AllocateOneFrameResource<FDynamicPrimitiveUniformBuffer>();
-				DynamicPrimitiveUniformBuffer.Set(GetLocalToWorld(), PreviousLocalToWorld, GetBounds(), GetLocalBounds(), true, bHasPrecomputedVolumetricLightmap, UseEditorDepthTest());
+				DynamicPrimitiveUniformBuffer.Set(GetLocalToWorld(), PreviousLocalToWorld, GetBounds(), GetLocalBounds(), true, bHasPrecomputedVolumetricLightmap, UseEditorDepthTest(), bOutputVelocity);
 				BatchElement.PrimitiveUniformBufferResource = &DynamicPrimitiveUniformBuffer.UniformBuffer;
 
 				BatchElement.FirstIndex = 0;
diff --git a/Engine/Plugins/Runtime/CustomMeshComponent/Source/CustomMeshComponent/Private/CustomMeshComponent.cpp b/Engine/Plugins/Runtime/CustomMeshComponent/Source/CustomMeshComponent/Private/CustomMeshComponent.cpp
--- a/Engine/Plugins/Runtime/CustomMeshComponent/Source/CustomMeshComponent/Private/CustomMeshComponent.cpp
+++ b/Engine/Plugins/Runtime/CustomMeshComponent/Source/CustomMeshComponent/Private/CustomMeshComponent.cpp
@@ -132,7 +132,8 @@ public:
 				bool bHasPrecomputedVolumetricLightmap;
 				FMatrix PreviousLocalToWorld;
 				int32 SingleCaptureIndex;
-				GetScene().GetPrimitiveUniformShaderParameters_RenderThread(GetPrimitiveSceneInfo(), bHasPrecomputedVolumetricLightmap, PreviousLocalToWorld, SingleCaptureIndex);
+				bool bOutputVelocity;
+				GetScene().GetPrimitiveUniformShaderParameters_RenderThread(GetPrimitiveSceneInfo(), bHasPrecomputedVolumetricLightmap, PreviousLocalToWorld, SingleCaptureIndex, bOutputVelocity);
 
 				FDynamicPrimitiveUniformBuffer& DynamicPrimitiveUniformBuffer = Collector.AllocateOneFrameResource<FDynamicPrimitiveUniformBuffer>();
 				DynamicPrimitiveUniformBuffer.Set(GetLocalToWorld(), PreviousLocalToWorld, GetBounds(), GetLocalBounds(), true, bHasPrecomputedVolumetricLightmap, UseEditorDepthTest());
diff --git a/Engine/Plugins/Runtime/ProceduralMeshComponent/Source/ProceduralMeshComponent/Private/ProceduralMeshComponent.cpp b/Engine/Plugins/Runtime/ProceduralMeshComponent/Source/ProceduralMeshComponent/Private/ProceduralMeshComponent.cpp
--- a/Engine/Plugins/Runtime/ProceduralMeshComponent/Source/ProceduralMeshComponent/Private/ProceduralMeshComponent.cpp
+++ b/Engine/Plugins/Runtime/ProceduralMeshComponent/Source/ProceduralMeshComponent/Private/ProceduralMeshComponent.cpp
@@ -304,10 +304,11 @@ public:
 						bool bHasPrecomputedVolumetricLightmap;
 						FMatrix PreviousLocalToWorld;
 						int32 SingleCaptureIndex;
-						GetScene().GetPrimitiveUniformShaderParameters_RenderThread(GetPrimitiveSceneInfo(), bHasPrecomputedVolumetricLightmap, PreviousLocalToWorld, SingleCaptureIndex);
+						bool bOutputVelocity;
+						GetScene().GetPrimitiveUniformShaderParameters_RenderThread(GetPrimitiveSceneInfo(), bHasPrecomputedVolumetricLightmap, PreviousLocalToWorld, SingleCaptureIndex, bOutputVelocity);
 
 						FDynamicPrimitiveUniformBuffer& DynamicPrimitiveUniformBuffer = Collector.AllocateOneFrameResource<FDynamicPrimitiveUniformBuffer>();
-						DynamicPrimitiveUniformBuffer.Set(GetLocalToWorld(), PreviousLocalToWorld, GetBounds(), GetLocalBounds(), true, bHasPrecomputedVolumetricLightmap, UseEditorDepthTest());
+						DynamicPrimitiveUniformBuffer.Set(GetLocalToWorld(), PreviousLocalToWorld, GetBounds(), GetLocalBounds(), true, bHasPrecomputedVolumetricLightmap, UseEditorDepthTest(), bOutputVelocity);
 						BatchElement.PrimitiveUniformBufferResource = &DynamicPrimitiveUniformBuffer.UniformBuffer;
 
 						BatchElement.FirstIndex = 0;
@@ -329,7 +330,7 @@ public:
 		for (int32 ViewIndex = 0; ViewIndex < Views.Num(); ViewIndex++)
 		{
 			if (VisibilityMap & (1 << ViewIndex))
-			{
+			{	 
 				// Draw simple collision as wireframe if 'show collision', and collision is enabled, and we are not using the complex as the simple
 				if (ViewFamily.EngineShowFlags.Collision && IsCollisionEnabled() && BodySetup->GetCollisionTraceFlag() != ECollisionTraceFlag::CTF_UseComplexAsSimple)
 				{
diff --git a/Engine/Shaders/Private/BasePassCommon.ush b/Engine/Shaders/Private/BasePassCommon.ush
--- a/Engine/Shaders/Private/BasePassCommon.ush
+++ b/Engine/Shaders/Private/BasePassCommon.ush
@@ -50,7 +50,9 @@
 #define WRITES_PRECSHADOWFACTOR_TO_GBUFFER	(GBUFFER_HAS_PRECSHADOWFACTOR && !WRITES_PRECSHADOWFACTOR_ZERO)
 
 // If a primitive has static lighting, we assume it is not moving. If it is, it will be rerendered in an extra renderpass.
-#define WRITES_VELOCITY_TO_GBUFFER			(USES_GBUFFER && GBUFFER_HAS_VELOCITY && (!SELECTIVE_BASEPASS_OUTPUTS || !(STATICLIGHTING_TEXTUREMASK || STATICLIGHTING_SIGNEDDISTANCEFIELD || HQ_TEXTURE_LIGHTMAP || LQ_TEXTURE_LIGHTMAP)))
+#define SUPPORTS_WRITING_VELOCITY_TO_BASE_PASS	(FEATURE_LEVEL >= FEATURE_LEVEL_SM4 && (MATERIALBLENDING_SOLID || MATERIALBLENDING_MASKED) && !SIMPLE_FORWARD_SHADING) 
+#define WRITES_VELOCITY_TO_GBUFFER				((SUPPORTS_WRITING_VELOCITY_TO_BASE_PASS || USES_GBUFFER) && GBUFFER_HAS_VELOCITY && (!SELECTIVE_BASEPASS_OUTPUTS || !(STATICLIGHTING_TEXTUREMASK || STATICLIGHTING_SIGNEDDISTANCEFIELD || HQ_TEXTURE_LIGHTMAP || LQ_TEXTURE_LIGHTMAP)))
+
 
 // Use an extra interpolator for position
 #define WRITES_VELOCITY_TO_GBUFFER_USE_POS_INTERPOLATOR		(WRITES_VELOCITY_TO_GBUFFER && USING_TESSELLATION)
diff --git a/Engine/Shaders/Private/BasePassPixelShader.usf b/Engine/Shaders/Private/BasePassPixelShader.usf
--- a/Engine/Shaders/Private/BasePassPixelShader.usf
+++ b/Engine/Shaders/Private/BasePassPixelShader.usf
@@ -872,8 +872,12 @@ void FPixelShaderInOut_MainPS(
 
 #if USES_GBUFFER
 	GBuffer.SelectiveOutputMask = GetSelectiveOutputMask();
+	GBuffer.Velocity = 0;
+#endif
 
-	#if WRITES_VELOCITY_TO_GBUFFER
+#if WRITES_VELOCITY_TO_GBUFFER
+	BRANCH
+	if (GetPrimitiveData(MaterialParameters.PrimitiveId).OutputVelocity>0)
 	{
 		// 2d velocity, includes camera an object motion
 	#if WRITES_VELOCITY_TO_GBUFFER_USE_POS_INTERPOLATOR
@@ -883,13 +887,15 @@ void FPixelShaderInOut_MainPS(
 	#endif
 
 		// Make sure not to touch 0,0 which is clear color
-		GBuffer.Velocity = float4(EncodeVelocityToTexture(Velocity), 0, 0);
-	}
+		float4 EncodedVelocity = float4(EncodeVelocityToTexture(Velocity), 0, 0);
+	#if USES_GBUFFER
+		GBuffer.Velocity = EncodedVelocity;
 	#else
-		GBuffer.Velocity = 0;
+		OutVelocity = EncodedVelocity;
 	#endif
+	}
 #endif
-
+	
 	// So that the following code can still use DiffuseColor and SpecularColor.
 	GBuffer.SpecularColor = ComputeF0(Specular, BaseColor, Metallic);
 
@@ -1198,6 +1204,11 @@ void FPixelShaderInOut_MainPS(
 	#if GBUFFER_HAS_PRECSHADOWFACTOR
 		Out.MRT[GBUFFER_HAS_VELOCITY ? 6 : 5] = OutGBufferE;
 	#endif
+#else
+	// If not using the full gbuffer (forward shading) the velocity buffer can still be written to in the basepass.
+	#if GBUFFER_HAS_VELOCITY
+		Out.MRT[1] = OutVelocity;
+	#endif
 #endif
 
 #if !MATERIALBLENDING_MODULATE && USE_PREEXPOSURE
@@ -1208,22 +1219,25 @@ void FPixelShaderInOut_MainPS(
 
 // the following needs to match to the code in FSceneRenderTargets::GetGBufferRenderTargets()
 #define PIXELSHADEROUTPUT_BASEPASS 1
-#define PIXELSHADEROUTPUT_MRT0 (!USES_GBUFFER || !SELECTIVE_BASEPASS_OUTPUTS || NEEDS_BASEPASS_VERTEX_FOGGING || USES_EMISSIVE_COLOR || ALLOW_STATIC_LIGHTING)
-#define PIXELSHADEROUTPUT_MRT1 (USES_GBUFFER && (!SELECTIVE_BASEPASS_OUTPUTS || !MATERIAL_SHADINGMODEL_UNLIT))
-#define PIXELSHADEROUTPUT_MRT2 (USES_GBUFFER && (!SELECTIVE_BASEPASS_OUTPUTS || !MATERIAL_SHADINGMODEL_UNLIT))
-#define PIXELSHADEROUTPUT_MRT3 (USES_GBUFFER && (!SELECTIVE_BASEPASS_OUTPUTS || !MATERIAL_SHADINGMODEL_UNLIT))
+#if USES_GBUFFER
+#define PIXELSHADEROUTPUT_MRT0 (!SELECTIVE_BASEPASS_OUTPUTS || NEEDS_BASEPASS_VERTEX_FOGGING || USES_EMISSIVE_COLOR || ALLOW_STATIC_LIGHTING)
+#define PIXELSHADEROUTPUT_MRT1 ((!SELECTIVE_BASEPASS_OUTPUTS || !MATERIAL_SHADINGMODEL_UNLIT))
+#define PIXELSHADEROUTPUT_MRT2 ((!SELECTIVE_BASEPASS_OUTPUTS || !MATERIAL_SHADINGMODEL_UNLIT))
+#define PIXELSHADEROUTPUT_MRT3 ((!SELECTIVE_BASEPASS_OUTPUTS || !MATERIAL_SHADINGMODEL_UNLIT))
+	#if GBUFFER_HAS_VELOCITY
+		#define PIXELSHADEROUTPUT_MRT4 WRITES_VELOCITY_TO_GBUFFER
+		#define PIXELSHADEROUTPUT_MRT5 (!SELECTIVE_BASEPASS_OUTPUTS || WRITES_CUSTOMDATA_TO_GBUFFER)
+		#define PIXELSHADEROUTPUT_MRT6 (GBUFFER_HAS_PRECSHADOWFACTOR && (!SELECTIVE_BASEPASS_OUTPUTS || WRITES_PRECSHADOWFACTOR_TO_GBUFFER && !MATERIAL_SHADINGMODEL_UNLIT))
+	#else //GBUFFER_HAS_VELOCITY
+		#define PIXELSHADEROUTPUT_MRT4 (!SELECTIVE_BASEPASS_OUTPUTS || WRITES_CUSTOMDATA_TO_GBUFFER)
+		#define PIXELSHADEROUTPUT_MRT5 (GBUFFER_HAS_PRECSHADOWFACTOR && (!SELECTIVE_BASEPASS_OUTPUTS || WRITES_PRECSHADOWFACTOR_TO_GBUFFER && !MATERIAL_SHADINGMODEL_UNLIT))
+	#endif //GBUFFER_HAS_VELOCITY
+#else //USES_GBUFFER
+	#define PIXELSHADEROUTPUT_MRT0 1
+	#define PIXELSHADEROUTPUT_MRT1 WRITES_VELOCITY_TO_GBUFFER
+#endif //USES_GBUFFER
 #define PIXELSHADEROUTPUT_A2C ((EDITOR_ALPHA2COVERAGE) != 0)
 #define PIXELSHADEROUTPUT_COVERAGE (MATERIALBLENDING_MASKED_USING_COVERAGE && !EARLY_Z_PASS_ONLY_MATERIAL_MASKING)
-#if GBUFFER_HAS_VELOCITY
-#define PIXELSHADEROUTPUT_MRT4 WRITES_VELOCITY_TO_GBUFFER
-#define PIXELSHADEROUTPUT_MRT5 (USES_GBUFFER && (!SELECTIVE_BASEPASS_OUTPUTS || WRITES_CUSTOMDATA_TO_GBUFFER))
-#define PIXELSHADEROUTPUT_MRT6 (GBUFFER_HAS_PRECSHADOWFACTOR && (!SELECTIVE_BASEPASS_OUTPUTS || WRITES_PRECSHADOWFACTOR_TO_GBUFFER && !MATERIAL_SHADINGMODEL_UNLIT))
-
-#else //WRITES_VELOCITY_TO_GBUFFER
-#define PIXELSHADEROUTPUT_MRT4 (USES_GBUFFER && (!SELECTIVE_BASEPASS_OUTPUTS || WRITES_CUSTOMDATA_TO_GBUFFER))
-#define PIXELSHADEROUTPUT_MRT5 (GBUFFER_HAS_PRECSHADOWFACTOR && (!SELECTIVE_BASEPASS_OUTPUTS || WRITES_PRECSHADOWFACTOR_TO_GBUFFER && !MATERIAL_SHADINGMODEL_UNLIT))
-
-#endif //WRITES_VELOCITY_TO_GBUFFER
 
 // all PIXELSHADEROUTPUT_ and "void FPixelShaderInOut_MainPS()" need to be setup before this include
 // this include generates the wrapper code to call MainPS(inout FPixelShaderOutput PixelShaderOutput)
diff --git a/Engine/Shaders/Private/BasePassVertexShader.usf b/Engine/Shaders/Private/BasePassVertexShader.usf
--- a/Engine/Shaders/Private/BasePassVertexShader.usf
+++ b/Engine/Shaders/Private/BasePassVertexShader.usf
@@ -180,9 +180,18 @@ void Main(
 
 #if WRITES_VELOCITY_TO_GBUFFER
 	{
-		float4 PrevTranslatedWorldPosition = VertexFactoryGetPreviousWorldPosition( Input, VFIntermediates );	
-		VertexParameters = GetMaterialVertexParameters(Input, VFIntermediates, PrevTranslatedWorldPosition.xyz, TangentToLocal);
-		PrevTranslatedWorldPosition.xyz += GetMaterialPreviousWorldPositionOffset(VertexParameters);
+		float4 PrevTranslatedWorldPosition = float4(0, 0, 0, 1);
+		BRANCH
+		if (GetPrimitiveData(VFIntermediates.PrimitiveId).OutputVelocity>0)
+		{
+			PrevTranslatedWorldPosition = VertexFactoryGetPreviousWorldPosition( Input, VFIntermediates );	
+			VertexParameters = GetMaterialVertexParameters(Input, VFIntermediates, PrevTranslatedWorldPosition.xyz, TangentToLocal);
+			PrevTranslatedWorldPosition.xyz += GetMaterialPreviousWorldPositionOffset(VertexParameters);
+	
+			#if !USING_TESSELATION
+				PrevTranslatedWorldPosition = mul(float4(PrevTranslatedWorldPosition.xyz, 1), ResolvedView.PrevTranslatedWorldToClip);
+			#endif
+		}
 
 		#if USING_TESSELLATION
 			// We let the Domain Shader convert to post projection when tessellating
@@ -192,7 +201,7 @@ void Main(
 			#endif
 		#else
 			// compute the old screen pos with the old world position and the old camera matrix
-			Output.BasePassInterpolants.VelocityPrevScreenPosition = mul(float4(PrevTranslatedWorldPosition.xyz, 1), ResolvedView.PrevTranslatedWorldToClip);		
+			Output.BasePassInterpolants.VelocityPrevScreenPosition = PrevTranslatedWorldPosition; 
 			#if WRITES_VELOCITY_TO_GBUFFER_USE_POS_INTERPOLATOR
 			Output.BasePassInterpolants.VelocityScreenPosition = ClipSpacePosition;
 			#endif
diff --git a/Engine/Shaders/Private/SceneData.ush b/Engine/Shaders/Private/SceneData.ush
--- a/Engine/Shaders/Private/SceneData.ush
+++ b/Engine/Shaders/Private/SceneData.ush
@@ -34,6 +34,7 @@ struct FPrimitiveSceneData
 	uint LightingChannelMask;
 	uint LightmapDataIndex;
 	int SingleCaptureIndex;
+    uint OutputVelocity;
 };
 
 // Stride of a single primitive's data in float4's, must match C++
@@ -89,6 +90,7 @@ FPrimitiveSceneData GetPrimitiveData(uint PrimitiveId)
 	PrimitiveData.LightmapDataIndex = asuint(View.PrimitiveSceneData[PrimitiveBaseOffset + 24].w);
 
 	PrimitiveData.SingleCaptureIndex = asint(View.PrimitiveSceneData[PrimitiveBaseOffset + 25].x);
+	PrimitiveData.OutputVelocity = asuint(View.PrimitiveSceneData[PrimitiveBaseOffset + 25].y);
 
 	PrimitiveData.PreSkinnedLocalBounds = View.PrimitiveSceneData[PrimitiveBaseOffset + 26].xyz;
 
diff --git a/Engine/Source/Developer/TargetPlatform/Private/TargetPlatformBase.cpp b/Engine/Source/Developer/TargetPlatform/Private/TargetPlatformBase.cpp
--- a/Engine/Source/Developer/TargetPlatform/Private/TargetPlatformBase.cpp
+++ b/Engine/Source/Developer/TargetPlatform/Private/TargetPlatformBase.cpp
@@ -16,6 +16,18 @@ bool FTargetPlatformBase::UsesDBuffer() const
 	return CVar ? (CVar->GetInt() != 0) : false;
 }
 
+bool FTargetPlatformBase::UsesBasePassVelocity() const
+{
+	static IConsoleVariable* CVar = IConsoleManager::Get().FindConsoleVariable(TEXT("r.BasePassOutputsVelocity"));
+	return CVar ? (CVar->GetInt()) : false;
+}
+
+bool FTargetPlatformBase::UsesSelectiveBasePassOutputs() const
+{
+	static IConsoleVariable* CVar = IConsoleManager::Get().FindConsoleVariable(TEXT("r.SelectiveBasePassOutputs"));
+	return CVar ? (CVar->GetInt()) : false;
+}
+
 TSharedPtr<IDeviceManagerCustomPlatformWidgetCreator> FTargetPlatformBase::GetCustomWidgetCreator() const
 {
 	static TSharedPtr<FDeviceBrowserDefaultPlatformWidgetCreator> DefaultWidgetCreator = MakeShared<FDeviceBrowserDefaultPlatformWidgetCreator>();
diff --git a/Engine/Source/Developer/TargetPlatform/Public/Common/TargetPlatformBase.h b/Engine/Source/Developer/TargetPlatform/Public/Common/TargetPlatformBase.h
--- a/Engine/Source/Developer/TargetPlatform/Public/Common/TargetPlatformBase.h
+++ b/Engine/Source/Developer/TargetPlatform/Public/Common/TargetPlatformBase.h
@@ -35,6 +35,10 @@ public:
 
 	TARGETPLATFORM_API virtual bool UsesDBuffer() const override;
 
+	TARGETPLATFORM_API virtual bool UsesBasePassVelocity() const override;
+
+	TARGETPLATFORM_API virtual bool UsesSelectiveBasePassOutputs() const override;
+
 #if WITH_ENGINE
 	virtual void GetReflectionCaptureFormats( TArray<FName>& OutFormats ) const override
 	{
diff --git a/Engine/Source/Developer/TargetPlatform/Public/Interfaces/ITargetPlatform.h b/Engine/Source/Developer/TargetPlatform/Public/Interfaces/ITargetPlatform.h
--- a/Engine/Source/Developer/TargetPlatform/Public/Interfaces/ITargetPlatform.h
+++ b/Engine/Source/Developer/TargetPlatform/Public/Interfaces/ITargetPlatform.h
@@ -343,6 +343,16 @@ public:
 	*/
 	virtual bool UsesDBuffer() const = 0;
 	
+	/**
+	* Gets whether the platform should output velocity in the base pass.
+	*/
+	virtual bool UsesBasePassVelocity() const = 0;
+
+	/**
+	* Gets whether the platform will use selective outputs in the base pass shaders.
+	*/
+	virtual bool UsesSelectiveBasePassOutputs() const = 0; 
+
 #if WITH_ENGINE
 	/**
 	 * Gets the format to use for a particular body setup.
diff --git a/Engine/Source/Runtime/Engine/Private/Components/ArrowComponent.cpp b/Engine/Source/Runtime/Engine/Private/Components/ArrowComponent.cpp
--- a/Engine/Source/Runtime/Engine/Private/Components/ArrowComponent.cpp
+++ b/Engine/Source/Runtime/Engine/Private/Components/ArrowComponent.cpp
@@ -148,7 +148,7 @@ public:
 				Mesh.MaterialRenderProxy = ArrowMaterialRenderProxy;
 
 				FDynamicPrimitiveUniformBuffer& DynamicPrimitiveUniformBuffer = Collector.AllocateOneFrameResource<FDynamicPrimitiveUniformBuffer>();
-				DynamicPrimitiveUniformBuffer.Set(FScaleMatrix(ViewScale) * EffectiveLocalToWorld, FScaleMatrix(ViewScale) * EffectiveLocalToWorld, GetBounds(), GetLocalBounds(), true, false, UseEditorDepthTest());
+				DynamicPrimitiveUniformBuffer.Set(FScaleMatrix(ViewScale) * EffectiveLocalToWorld, FScaleMatrix(ViewScale) * EffectiveLocalToWorld, GetBounds(), GetLocalBounds(), true, false, UseEditorDepthTest(), false);
 				BatchElement.PrimitiveUniformBufferResource = &DynamicPrimitiveUniformBuffer.UniformBuffer;
 
 				BatchElement.FirstIndex = 0;
diff --git a/Engine/Source/Runtime/Engine/Private/DynamicMeshBuilder.cpp b/Engine/Source/Runtime/Engine/Private/DynamicMeshBuilder.cpp
--- a/Engine/Source/Runtime/Engine/Private/DynamicMeshBuilder.cpp
+++ b/Engine/Source/Runtime/Engine/Private/DynamicMeshBuilder.cpp
@@ -736,7 +736,8 @@ void FDynamicMeshBuilder::GetMesh(const FMatrix& LocalToWorld,const FMaterialRen
 				GetDefaultLightingChannelMask(),
 				0,
 				INDEX_NONE,
-				INDEX_NONE
+				INDEX_NONE,
+				false
 			);
 
 
@@ -826,7 +827,8 @@ void FDynamicMeshBuilder::GetMeshElement(const FMatrix& LocalToWorld, const FMat
 				GetDefaultLightingChannelMask(),
 				0,
 				INDEX_NONE,
-				INDEX_NONE
+				INDEX_NONE,
+				false
 			);
 
 
@@ -903,7 +905,8 @@ void FDynamicMeshBuilder::Draw(FPrimitiveDrawInterface* PDI,const FMatrix& Local
 			GetDefaultLightingChannelMask(),
 			0,
 			INDEX_NONE,
-			INDEX_NONE
+			INDEX_NONE,
+			false
 			);
 
 
diff --git a/Engine/Source/Runtime/Engine/Private/Particles/ParticleSystemRender.cpp b/Engine/Source/Runtime/Engine/Private/Particles/ParticleSystemRender.cpp
--- a/Engine/Source/Runtime/Engine/Private/Particles/ParticleSystemRender.cpp
+++ b/Engine/Source/Runtime/Engine/Private/Particles/ParticleSystemRender.cpp
@@ -7275,7 +7275,8 @@ void FParticleSystemSceneProxy::UpdateWorldSpacePrimitiveUniformBuffer() const
 			GetLightingChannelMask(),
 			0,
 			INDEX_NONE,
-			INDEX_NONE
+			INDEX_NONE,
+			false
 			);
 		WorldSpacePrimitiveUniformBuffer.SetContents(PrimitiveUniformShaderParameters);
 		WorldSpacePrimitiveUniformBuffer.InitResource();
diff --git a/Engine/Source/Runtime/Engine/Private/PhysicsEngine/PhysDrawing.cpp b/Engine/Source/Runtime/Engine/Private/PhysicsEngine/PhysDrawing.cpp
--- a/Engine/Source/Runtime/Engine/Private/PhysicsEngine/PhysDrawing.cpp
+++ b/Engine/Source/Runtime/Engine/Private/PhysicsEngine/PhysDrawing.cpp
@@ -894,7 +894,7 @@ void FKAggregateGeom::GetAggGeom(const FTransform& Transform, const FColor Color
 				CalcBoxSphereBounds(LocalBounds, FTransform::Identity);
 
 				FDynamicPrimitiveUniformBuffer& DynamicPrimitiveUniformBuffer = Collector.AllocateOneFrameResource<FDynamicPrimitiveUniformBuffer>();
-				DynamicPrimitiveUniformBuffer.Set(LocalToWorld.ToMatrixWithScale(), LocalToWorld.ToMatrixWithScale(), WorldBounds, LocalBounds, true, false, bUseEditorDepthTest);
+				DynamicPrimitiveUniformBuffer.Set(LocalToWorld.ToMatrixWithScale(), LocalToWorld.ToMatrixWithScale(), WorldBounds, LocalBounds, true, false, bUseEditorDepthTest, false);
 				BatchElement.PrimitiveUniformBufferResource = &DynamicPrimitiveUniformBuffer.UniformBuffer;
 
 			 	// previous l2w not used so treat as static
diff --git a/Engine/Source/Runtime/Engine/Private/PrimitiveSceneProxy.cpp b/Engine/Source/Runtime/Engine/Private/PrimitiveSceneProxy.cpp
--- a/Engine/Source/Runtime/Engine/Private/PrimitiveSceneProxy.cpp
+++ b/Engine/Source/Runtime/Engine/Private/PrimitiveSceneProxy.cpp
@@ -266,8 +266,9 @@ void FPrimitiveSceneProxy::UpdateUniformBuffer()
 		bool bHasPrecomputedVolumetricLightmap;
 		FMatrix PreviousLocalToWorld;
 		int32 SingleCaptureIndex;
+		bool bOutputVelocity;
 
-		Scene->GetPrimitiveUniformShaderParameters_RenderThread(PrimitiveSceneInfo, bHasPrecomputedVolumetricLightmap, PreviousLocalToWorld, SingleCaptureIndex);
+		Scene->GetPrimitiveUniformShaderParameters_RenderThread(PrimitiveSceneInfo, bHasPrecomputedVolumetricLightmap, PreviousLocalToWorld, SingleCaptureIndex, bOutputVelocity);
 
 		FBoxSphereBounds PreSkinnedLocalBounds;
 		GetPreSkinnedLocalBounds(PreSkinnedLocalBounds);
@@ -290,7 +291,8 @@ void FPrimitiveSceneProxy::UpdateUniformBuffer()
 				GetLightingChannelMask(),
 				LpvBiasMultiplier,
 				PrimitiveSceneInfo ? PrimitiveSceneInfo->GetLightmapDataOffset() : 0,
-				SingleCaptureIndex);
+				SingleCaptureIndex, 
+				bOutputVelocity);
 
 		if (UniformBuffer.GetReference())
 		{
diff --git a/Engine/Source/Runtime/Engine/Private/PrimitiveUniformShaderParameters.cpp b/Engine/Source/Runtime/Engine/Private/PrimitiveUniformShaderParameters.cpp
--- a/Engine/Source/Runtime/Engine/Private/PrimitiveUniformShaderParameters.cpp
+++ b/Engine/Source/Runtime/Engine/Private/PrimitiveUniformShaderParameters.cpp
@@ -40,8 +40,9 @@ FPrimitiveSceneShaderData::FPrimitiveSceneShaderData(const FPrimitiveSceneProxy*
 	bool bHasPrecomputedVolumetricLightmap;
 	FMatrix PreviousLocalToWorld;
 	int32 SingleCaptureIndex;
+	bool bOutputVelocity;
 
-	Proxy->GetScene().GetPrimitiveUniformShaderParameters_RenderThread(Proxy->GetPrimitiveSceneInfo(), bHasPrecomputedVolumetricLightmap, PreviousLocalToWorld, SingleCaptureIndex);
+	Proxy->GetScene().GetPrimitiveUniformShaderParameters_RenderThread(Proxy->GetPrimitiveSceneInfo(), bHasPrecomputedVolumetricLightmap, PreviousLocalToWorld, SingleCaptureIndex, bOutputVelocity);
 
 	FBoxSphereBounds PreSkinnedLocalBounds;
 	Proxy->GetPreSkinnedLocalBounds(PreSkinnedLocalBounds);
@@ -62,7 +63,8 @@ FPrimitiveSceneShaderData::FPrimitiveSceneShaderData(const FPrimitiveSceneProxy*
 		Proxy->GetLightingChannelMask(),
 		Proxy->GetLpvBiasMultiplier(),
 		Proxy->GetPrimitiveSceneInfo()->GetLightmapDataOffset(),
-		SingleCaptureIndex));
+		SingleCaptureIndex,
+        bOutputVelocity));
 }
 
 void FPrimitiveSceneShaderData::Setup(const FPrimitiveUniformShaderParameters& PrimitiveUniformShaderParameters)
@@ -111,6 +113,7 @@ void FPrimitiveSceneShaderData::Setup(const FPrimitiveUniformShaderParameters& P
 
 	Data[25] = FVector4(0.0f, 0.0f, 0.0f, 0.0f);
 	Data[25].X = *(const float*)&PrimitiveUniformShaderParameters.SingleCaptureIndex;
+	Data[25].Y = *(const float*)&PrimitiveUniformShaderParameters.OutputVelocity;
 
 	Data[26] = FVector4(PrimitiveUniformShaderParameters.PreSkinnedLocalBounds, 0.0f); // .w unused
 }
\ No newline at end of file
diff --git a/Engine/Source/Runtime/Engine/Private/SceneManagement.cpp b/Engine/Source/Runtime/Engine/Private/SceneManagement.cpp
--- a/Engine/Source/Runtime/Engine/Private/SceneManagement.cpp
+++ b/Engine/Source/Runtime/Engine/Private/SceneManagement.cpp
@@ -308,11 +308,12 @@ void FDynamicPrimitiveUniformBuffer::Set(
 	const FBoxSphereBounds& PreSkinnedLocalBounds,
 	bool bReceivesDecals,
 	bool bHasPrecomputedVolumetricLightmap,
-	bool bUseEditorDepthTest)
+	bool bUseEditorDepthTest,
+	bool bOutputVelocity)
 {
 	check(IsInRenderingThread());
 	UniformBuffer.SetContents(
-		GetPrimitiveUniformShaderParameters(LocalToWorld, PreviousLocalToWorld, WorldBounds.Origin, WorldBounds, LocalBounds, PreSkinnedLocalBounds, bReceivesDecals, false, false, false, bHasPrecomputedVolumetricLightmap, bUseEditorDepthTest, GetDefaultLightingChannelMask(), 1.0f, INDEX_NONE, INDEX_NONE));
+		GetPrimitiveUniformShaderParameters(LocalToWorld, PreviousLocalToWorld, WorldBounds.Origin, WorldBounds, LocalBounds, PreSkinnedLocalBounds, bReceivesDecals, false, false, false, bHasPrecomputedVolumetricLightmap, bUseEditorDepthTest, GetDefaultLightingChannelMask(), 1.0f, INDEX_NONE, INDEX_NONE, bOutputVelocity));
 	UniformBuffer.InitResource();
 }
 
@@ -323,9 +324,10 @@ void FDynamicPrimitiveUniformBuffer::Set(
 	const FBoxSphereBounds& LocalBounds,
 	bool bReceivesDecals,
 	bool bHasPrecomputedVolumetricLightmap,
-	bool bUseEditorDepthTest)
+	bool bUseEditorDepthTest,
+	bool bOutputVelocity)
 {
-	Set(LocalToWorld, PreviousLocalToWorld, WorldBounds, LocalBounds, LocalBounds, bReceivesDecals, bHasPrecomputedVolumetricLightmap, bUseEditorDepthTest);
+	Set(LocalToWorld, PreviousLocalToWorld, WorldBounds, LocalBounds, LocalBounds, bReceivesDecals, bHasPrecomputedVolumetricLightmap, bUseEditorDepthTest, bOutputVelocity);
 }
 
 FLightMapInteraction FLightMapInteraction::Texture(
diff --git a/Engine/Source/Runtime/Engine/Private/ShaderCompiler/ShaderCompiler.cpp b/Engine/Source/Runtime/Engine/Private/ShaderCompiler/ShaderCompiler.cpp
--- a/Engine/Source/Runtime/Engine/Private/ShaderCompiler/ShaderCompiler.cpp
+++ b/Engine/Source/Runtime/Engine/Private/ShaderCompiler/ShaderCompiler.cpp
@@ -3245,13 +3245,11 @@ void GlobalBeginCompileShader(
 	}
 
 	{
-		static const auto CVar = IConsoleManager::Get().FindTConsoleVariableDataInt(TEXT("r.BasePassOutputsVelocity"));
-		Input.Environment.SetDefine(TEXT("GBUFFER_HAS_VELOCITY"), CVar ? (CVar->GetValueOnGameThread() != 0) : 0);
+		Input.Environment.SetDefine(TEXT("GBUFFER_HAS_VELOCITY"), IsUsingBasePassVelocity((EShaderPlatform)Target.Platform) ? 1 : 0);
 	}
 
 	{
-		static const auto CVar = IConsoleManager::Get().FindTConsoleVariableDataInt(TEXT("r.SelectiveBasePassOutputs"));
-		Input.Environment.SetDefine(TEXT("SELECTIVE_BASEPASS_OUTPUTS"), CVar ? (CVar->GetValueOnGameThread() != 0) : 0);
+		Input.Environment.SetDefine(TEXT("SELECTIVE_BASEPASS_OUTPUTS"), IsUsingSelectiveBasePassOutputs((EShaderPlatform)Target.Platform) ? 1 : 0);
 	}
 
 	{
diff --git a/Engine/Source/Runtime/Engine/Public/PrimitiveUniformShaderParameters.h b/Engine/Source/Runtime/Engine/Public/PrimitiveUniformShaderParameters.h
--- a/Engine/Source/Runtime/Engine/Public/PrimitiveUniformShaderParameters.h
+++ b/Engine/Source/Runtime/Engine/Public/PrimitiveUniformShaderParameters.h
@@ -37,6 +37,8 @@ BEGIN_GLOBAL_SHADER_PARAMETER_STRUCT(FPrimitiveUniformShaderParameters,ENGINE_AP
 	SHADER_PARAMETER(uint32,LightingChannelMask)
 	SHADER_PARAMETER(uint32,LightmapDataIndex)
 	SHADER_PARAMETER(int32, SingleCaptureIndex)
+    SHADER_PARAMETER(uint32, OutputVelocity)
+
 END_GLOBAL_SHADER_PARAMETER_STRUCT()
 
 /** Initializes the primitive uniform shader parameters. */
@@ -56,7 +58,8 @@ inline FPrimitiveUniformShaderParameters GetPrimitiveUniformShaderParameters(
 	uint32 LightingChannelMask,
 	float LpvBiasMultiplier,
 	uint32 LightmapDataIndex,
-	int32 SingleCaptureIndex
+	int32 SingleCaptureIndex,
+	bool bOutputVelocity
 )
 {
 	FPrimitiveUniformShaderParameters Result;
@@ -97,6 +100,7 @@ inline FPrimitiveUniformShaderParameters GetPrimitiveUniformShaderParameters(
 	Result.UseEditorDepthTest = bUseEditorDepthTest ? 1 : 0;
 	Result.LightmapDataIndex = LightmapDataIndex;
 	Result.SingleCaptureIndex = SingleCaptureIndex;
+	Result.OutputVelocity = (bOutputVelocity) ? 1 : 0;
 	return Result;
 }
 
@@ -116,12 +120,13 @@ inline FPrimitiveUniformShaderParameters GetPrimitiveUniformShaderParameters(
 	uint32 LightingChannelMask,
 	float LpvBiasMultiplier,
 	uint32 LightmapDataIndex,
-	int32 SingleCaptureIndex
+	int32 SingleCaptureIndex,
+    bool bOutputVelocity
 )
 {
 	// Pass through call
 	return GetPrimitiveUniformShaderParameters(LocalToWorld, PreviousLocalToWorld, ActorPosition, WorldBounds, LocalBounds, LocalBounds, bReceivesDecals, bHasDistanceFieldRepresentation, bHasCapsuleRepresentation, 
-		bUseSingleSampleShadowFromStationaryLights, bUseVolumetricLightmap, bUseEditorDepthTest, LightingChannelMask, LpvBiasMultiplier, LightmapDataIndex, SingleCaptureIndex);
+		bUseSingleSampleShadowFromStationaryLights, bUseVolumetricLightmap, bUseEditorDepthTest, LightingChannelMask, LpvBiasMultiplier, LightmapDataIndex, SingleCaptureIndex, bOutputVelocity);
 }
 
 inline TUniformBufferRef<FPrimitiveUniformShaderParameters> CreatePrimitiveUniformBufferImmediate(
@@ -136,7 +141,7 @@ inline TUniformBufferRef<FPrimitiveUniformShaderParameters> CreatePrimitiveUnifo
 {
 	check(IsInRenderingThread());
 	return TUniformBufferRef<FPrimitiveUniformShaderParameters>::CreateUniformBufferImmediate(
-		GetPrimitiveUniformShaderParameters(LocalToWorld, LocalToWorld, WorldBounds.Origin, WorldBounds, LocalBounds, PreSkinnedLocalBounds, bReceivesDecals, false, false, false, false, bUseEditorDepthTest, GetDefaultLightingChannelMask(), LpvBiasMultiplier, INDEX_NONE, INDEX_NONE),
+		GetPrimitiveUniformShaderParameters(LocalToWorld, LocalToWorld, WorldBounds.Origin, WorldBounds, LocalBounds, PreSkinnedLocalBounds, bReceivesDecals, false, false, false, false, bUseEditorDepthTest, GetDefaultLightingChannelMask(), LpvBiasMultiplier, INDEX_NONE, INDEX_NONE, false),
 		UniformBuffer_MultiFrame
 		);
 }
@@ -160,7 +165,8 @@ inline FPrimitiveUniformShaderParameters GetIdentityPrimitiveParameters()
 		GetDefaultLightingChannelMask(),
 		1.0f,		// LPV bias
 		INDEX_NONE,
-		INDEX_NONE
+		INDEX_NONE,
+		false
 	);
 }
 
diff --git a/Engine/Source/Runtime/Engine/Public/SceneInterface.h b/Engine/Source/Runtime/Engine/Public/SceneInterface.h
--- a/Engine/Source/Runtime/Engine/Public/SceneInterface.h
+++ b/Engine/Source/Runtime/Engine/Public/SceneInterface.h
@@ -154,7 +154,7 @@ public:
 	virtual void UpdateSceneCaptureContents(class USceneCaptureComponent2D* CaptureComponent) {}
 	virtual void UpdateSceneCaptureContents(class USceneCaptureComponentCube* CaptureComponent) {}
 	virtual void UpdatePlanarReflectionContents(class UPlanarReflectionComponent* CaptureComponent, class FSceneRenderer& MainSceneRenderer) {}
-
+	
 	virtual void AddPrecomputedLightVolume(const class FPrecomputedLightVolume* Volume) {}
 	virtual void RemovePrecomputedLightVolume(const class FPrecomputedLightVolume* Volume) {}
 
@@ -165,7 +165,7 @@ public:
 	/** 
 	 * Retrieves primitive uniform shader parameters that are internal to the renderer.
 	 */
-	virtual void GetPrimitiveUniformShaderParameters_RenderThread(const FPrimitiveSceneInfo* PrimitiveSceneInfo, bool& bHasPrecomputedVolumetricLightmap, FMatrix& PreviousLocalToWorld, int32& SingleCaptureIndex) const {}
+	virtual void GetPrimitiveUniformShaderParameters_RenderThread(const FPrimitiveSceneInfo* PrimitiveSceneInfo, bool& bHasPrecomputedVolumetricLightmap, FMatrix& PreviousLocalToWorld, int32& SingleCaptureIndex, bool& OutputVelocity) const {}
 	 
 	/** 
 	 * Updates the transform of a light which has already been added to the scene. 
diff --git a/Engine/Source/Runtime/Engine/Public/SceneManagement.h b/Engine/Source/Runtime/Engine/Public/SceneManagement.h
--- a/Engine/Source/Runtime/Engine/Public/SceneManagement.h
+++ b/Engine/Source/Runtime/Engine/Public/SceneManagement.h
@@ -2208,7 +2208,8 @@ public:
 		const FBoxSphereBounds& PreSkinnedLocalBounds,
 		bool bReceivesDecals,
 		bool bHasPrecomputedVolumetricLightmap,
-		bool bUseEditorDepthTest);
+		bool bUseEditorDepthTest,
+		bool bOutputVelocity);
 
 	/** Pass-through implementation which calls the overloaded Set function with LocalBounds for PreSkinnedLocalBounds. */
 	ENGINE_API void Set(
@@ -2218,7 +2219,8 @@ public:
 		const FBoxSphereBounds& LocalBounds,
 		bool bReceivesDecals,
 		bool bHasPrecomputedVolumetricLightmap,
-		bool bUseEditorDepthTest);
+		bool bUseEditorDepthTest,
+		bool bOutputVelocity);
 };
 
 //
diff --git a/Engine/Source/Runtime/MRMesh/Private/MRMeshComponent.cpp b/Engine/Source/Runtime/MRMesh/Private/MRMeshComponent.cpp
--- a/Engine/Source/Runtime/MRMesh/Private/MRMeshComponent.cpp
+++ b/Engine/Source/Runtime/MRMesh/Private/MRMeshComponent.cpp
@@ -325,7 +325,7 @@ private:
 						Mesh.MaterialRenderProxy = MaterialProxy;
 
 						FDynamicPrimitiveUniformBuffer& DynamicPrimitiveUniformBuffer = Collector.AllocateOneFrameResource<FDynamicPrimitiveUniformBuffer>();
-						DynamicPrimitiveUniformBuffer.Set(GetLocalToWorld(), GetLocalToWorld(), InfiniteBounds, InfiniteBounds, true, false, UseEditorDepthTest());
+						DynamicPrimitiveUniformBuffer.Set(GetLocalToWorld(), GetLocalToWorld(), InfiniteBounds, InfiniteBounds, true, false, UseEditorDepthTest(), false);
 						BatchElement.PrimitiveUniformBufferResource = &DynamicPrimitiveUniformBuffer.UniformBuffer;
 
 						BatchElement.FirstIndex = 0;
diff --git a/Engine/Source/Runtime/NavigationSystem/Private/NavMesh/NavMeshRenderingComponent.cpp b/Engine/Source/Runtime/NavigationSystem/Private/NavMesh/NavMeshRenderingComponent.cpp
--- a/Engine/Source/Runtime/NavigationSystem/Private/NavMesh/NavMeshRenderingComponent.cpp
+++ b/Engine/Source/Runtime/NavigationSystem/Private/NavMesh/NavMeshRenderingComponent.cpp
@@ -952,7 +952,7 @@ void FNavMeshSceneProxy::GetDynamicMeshElements(const TArray<const FSceneView*>&
 					BatchElement = MeshBatchElements[Index];
 
 					FDynamicPrimitiveUniformBuffer& DynamicPrimitiveUniformBuffer = Collector.AllocateOneFrameResource<FDynamicPrimitiveUniformBuffer>();
-					DynamicPrimitiveUniformBuffer.Set(FMatrix::Identity, FMatrix::Identity, GetBounds(), GetLocalBounds(), false, false, UseEditorDepthTest());
+					DynamicPrimitiveUniformBuffer.Set(FMatrix::Identity, FMatrix::Identity, GetBounds(), GetLocalBounds(), false, false, UseEditorDepthTest(), false);
 					BatchElement.PrimitiveUniformBufferResource = &DynamicPrimitiveUniformBuffer.UniformBuffer;
 
 					Mesh.bWireframe = false;
diff --git a/Engine/Source/Runtime/RenderCore/Private/RenderUtils.cpp b/Engine/Source/Runtime/RenderCore/Private/RenderUtils.cpp
--- a/Engine/Source/Runtime/RenderCore/Private/RenderUtils.cpp
+++ b/Engine/Source/Runtime/RenderCore/Private/RenderUtils.cpp
@@ -914,6 +914,12 @@ static_assert(SP_NumPlatforms <= sizeof(GForwardShadingPlatformMask) * 8, "GForw
 RENDERCORE_API uint32 GDBufferPlatformMask = 0;
 static_assert(SP_NumPlatforms <= sizeof(GDBufferPlatformMask) * 8, "GDBufferPlatformMask must be large enough to support all shader platforms");
 
+RENDERCORE_API uint32 GBasePassVelocityPlatformMask = 0;
+static_assert(SP_NumPlatforms <= sizeof(GBasePassVelocityPlatformMask) * 8, "GBasePassVelocityPlatformMask must be large enough to support all shader platforms");
+
+RENDERCORE_API uint32 GSelectiveBasePassOutputsPlatformMask = 0;
+static_assert(SP_NumPlatforms <= sizeof(GSelectiveBasePassOutputsPlatformMask) * 8, "GSelectiveBasePassOutputsPlatformMask must be large enough to support all shader platforms");
+
 RENDERCORE_API void RenderUtilsInit()
 {
 	if (GUseForwardShading)
@@ -921,12 +927,24 @@ RENDERCORE_API void RenderUtilsInit()
 		GForwardShadingPlatformMask = ~0u;
 	}
 
-	static IConsoleVariable* CDBufferVar = IConsoleManager::Get().FindConsoleVariable(TEXT("r.DBuffer"));
-	if (CDBufferVar && CDBufferVar->GetInt())
+	static IConsoleVariable* DBufferVar = IConsoleManager::Get().FindConsoleVariable(TEXT("r.DBuffer"));
+	if (DBufferVar && DBufferVar->GetInt())
 	{
 		GDBufferPlatformMask = ~0u;
 	}
 
+	static IConsoleVariable* BasePassVelocityCVar = IConsoleManager::Get().FindConsoleVariable(TEXT("r.BasePassOutputsVelocity"));
+	if (BasePassVelocityCVar && BasePassVelocityCVar->GetInt())
+	{
+		GBasePassVelocityPlatformMask = ~0u;
+	}
+
+	static IConsoleVariable* SelectiveBasePassOutputsCVar = IConsoleManager::Get().FindConsoleVariable(TEXT("r.SelectiveBasePassOutputs"));
+	if (SelectiveBasePassOutputsCVar && SelectiveBasePassOutputsCVar->GetInt())
+	{
+		GSelectiveBasePassOutputsPlatformMask = ~0u;
+	}
+
 #if WITH_EDITOR
 	ITargetPlatformManagerModule* TargetPlatformManager = GetTargetPlatformManager();
 	if (TargetPlatformManager)
@@ -957,6 +975,24 @@ RENDERCORE_API void RenderUtilsInit()
 				{
 					GDBufferPlatformMask &= ~Mask;
 				}
+
+				if (TargetPlatform->UsesBasePassVelocity())
+				{
+					GBasePassVelocityPlatformMask |= Mask;
+				}
+				else
+				{
+					GBasePassVelocityPlatformMask &= ~Mask;
+				}
+
+				if (TargetPlatform->UsesSelectiveBasePassOutputs())
+				{
+					GSelectiveBasePassOutputsPlatformMask |= Mask;
+				}
+				else
+				{
+					GSelectiveBasePassOutputsPlatformMask &= ~Mask;
+				}
 			}
 		}
 	}
diff --git a/Engine/Source/Runtime/RenderCore/Private/Shader.cpp b/Engine/Source/Runtime/RenderCore/Private/Shader.cpp
--- a/Engine/Source/Runtime/RenderCore/Private/Shader.cpp
+++ b/Engine/Source/Runtime/RenderCore/Private/Shader.cpp
@@ -16,6 +16,7 @@
 #include "ShaderCore.h"
 #include "RenderUtils.h"
 #include "Misc/ConfigCacheIni.h"
+#include "Misc/ScopeLock.h"
 #include "UObject/RenderingObjectVersion.h"
 #include "UObject/FortniteMainBranchObjectVersion.h"
 #include "Misc/ScopeLock.h"
@@ -2067,11 +2068,7 @@ void ShaderMapAppendKeyString(EShaderPlatform Platform, FString& KeyString)
 	}
 
 	{
-		static const auto CVar = IConsoleManager::Get().FindTConsoleVariableDataInt(TEXT("r.BasePassOutputsVelocity"));
-		if (CVar && CVar->GetValueOnGameThread() != 0)
-		{
-			KeyString += TEXT("_GV");
-		}
+		KeyString += IsUsingBasePassVelocity(Platform) ? TEXT("_GV") : TEXT("");
 	}
 
 	{
@@ -2110,11 +2107,7 @@ void ShaderMapAppendKeyString(EShaderPlatform Platform, FString& KeyString)
 	}
 
 	{
-		static const auto CVar = IConsoleManager::Get().FindTConsoleVariableDataInt(TEXT("r.SelectiveBasePassOutputs"));
-		if (CVar && CVar->GetValueOnGameThread() != 0)
-		{
-			KeyString += TEXT("_SO");
-		}
+		KeyString += IsUsingSelectiveBasePassOutputs(Platform) ? TEXT("_SO") : TEXT("");
 	}
 
 	{
diff --git a/Engine/Source/Runtime/RenderCore/Public/RenderUtils.h b/Engine/Source/Runtime/RenderCore/Public/RenderUtils.h
--- a/Engine/Source/Runtime/RenderCore/Public/RenderUtils.h
+++ b/Engine/Source/Runtime/RenderCore/Public/RenderUtils.h
@@ -456,6 +456,20 @@ inline bool IsUsingDBuffers(EShaderPlatform Platform)
 	return !!(GDBufferPlatformMask & (1u << Platform));
 }
 
+/** Returns whether the base pass should output to the velocity buffer is enabled for a given shader platform */
+inline bool IsUsingBasePassVelocity(EShaderPlatform Platform)
+{
+	extern RENDERCORE_API uint32 GBasePassVelocityPlatformMask;
+	return !!(GBasePassVelocityPlatformMask & (1u << Platform));
+}
+
+/** Returns whether the base pass should use selective outputs for a given shader platform */
+inline bool IsUsingSelectiveBasePassOutputs(EShaderPlatform Platform)
+{
+	extern RENDERCORE_API uint32 GSelectiveBasePassOutputsPlatformMask;
+	return !!(GSelectiveBasePassOutputsPlatformMask & (1u << Platform));
+}
+
 inline bool IsUsingPerPixelDBufferMask(EShaderPlatform Platform)
 {
 	switch (Platform)
diff --git a/Engine/Source/Runtime/Renderer/Private/BasePassRendering.cpp b/Engine/Source/Runtime/Renderer/Private/BasePassRendering.cpp
--- a/Engine/Source/Runtime/Renderer/Private/BasePassRendering.cpp
+++ b/Engine/Source/Runtime/Renderer/Private/BasePassRendering.cpp
@@ -45,11 +45,6 @@ static TAutoConsoleVariable<int32> CVarRHICmdFlushRenderThreadTasksBasePass(
 	0,
 	TEXT("Wait for completion of parallel render thread tasks at the end of the base pass. A more granular version of r.RHICmdFlushRenderThreadTasks. If either r.RHICmdFlushRenderThreadTasks or r.RHICmdFlushRenderThreadTasksBasePass is > 0 we will flush."));
 
-bool UseSelectiveBasePassOutputs()
-{
-	return CVarSelectiveBasePassOutputs.GetValueOnAnyThread() == 1;
-}
-
 static TAutoConsoleVariable<int32> CVarSupportStationarySkylight(
 	TEXT("r.SupportStationarySkylight"),
 	1,
diff --git a/Engine/Source/Runtime/Renderer/Private/BasePassRendering.h b/Engine/Source/Runtime/Renderer/Private/BasePassRendering.h
--- a/Engine/Source/Runtime/Renderer/Private/BasePassRendering.h
+++ b/Engine/Source/Runtime/Renderer/Private/BasePassRendering.h
@@ -35,9 +35,6 @@ class FViewInfo;
 /** Whether to allow the indirect lighting cache to be applied to dynamic objects. */
 extern int32 GIndirectLightingCache;
 
-/** Whether some GBuffer targets are optional. */
-extern bool UseSelectiveBasePassOutputs();
-
 class FForwardLocalLightData
 {
 public:
@@ -393,7 +390,7 @@ public:
 		const bool bOutputVelocity = FVelocityRendering::BasePassCanOutputVelocity(Platform);
 		if (bOutputVelocity)
 		{
-			const int32 VelocityIndex = 4; // As defined in BasePassPixelShader.usf
+			const int32 VelocityIndex = IsForwardShadingEnabled(Platform) ? 1 : 4; // As defined in BasePassPixelShader.usf
 			OutEnvironment.SetRenderTargetOutputFormat(VelocityIndex, PF_G16R16);
 		}
 
diff --git a/Engine/Source/Runtime/Renderer/Private/DeferredShadingRenderer.cpp b/Engine/Source/Runtime/Renderer/Private/DeferredShadingRenderer.cpp
--- a/Engine/Source/Runtime/Renderer/Private/DeferredShadingRenderer.cpp
+++ b/Engine/Source/Runtime/Renderer/Private/DeferredShadingRenderer.cpp
@@ -698,14 +698,14 @@ bool FDeferredShadingSceneRenderer::GatherRayTracingWorldInstances(FRHICommandLi
 
 			if (RayTracedMeshElementsMask != 0)
 			{
-				FPrimitiveSceneProxy* SceneProxy = Scene->PrimitiveSceneProxies[PrimitiveIndex];
+					FPrimitiveSceneProxy* SceneProxy = Scene->PrimitiveSceneProxies[PrimitiveIndex];
 				TArray<FRayTracingInstance> RayTracingInstances;
 				SceneProxy->GetDynamicRayTracingInstances(MaterialGatheringContext, RayTracingInstances);
 
 				if (RayTracingInstances.Num() > 0)
-				{
+						{
 					for (FRayTracingInstance& Instance : RayTracingInstances)
-					{
+							{
 						FRayTracingGeometryInstance RayTracingInstance = { Instance.Geometry->RayTracingGeometryRHI };
 						RayTracingInstance.Transform = Instance.InstanceTransforms[0];
 						ensureMsgf(Instance.InstanceTransforms.Num() == 1, TEXT("Multi-instancing hasn't been supported"));
@@ -718,12 +718,12 @@ bool FDeferredShadingSceneRenderer::GatherRayTracingWorldInstances(FRHICommandLi
 						uint32 InstanceIndex = ReferenceView.RayTracingGeometryInstances.Add(RayTracingInstance);
 
 						for (int32 ViewIndex = 1; ViewIndex < Views.Num(); ViewIndex++)
-						{
+							{
 							Views[ViewIndex].RayTracingGeometryInstances.Add(RayTracingInstance);
-						}
+				}
 
 						for (int32 SegmentIndex = 0; SegmentIndex < Instance.Materials.Num(); SegmentIndex++)
-						{
+				{
 							FMeshBatch& MeshBatch = Instance.Materials[SegmentIndex];
 							FDynamicRayTracingMeshCommandContext CommandContext(ReferenceView.DynamicRayTracingMeshCommandStorage, ReferenceView.VisibleRayTracingMeshCommands, SegmentIndex, InstanceIndex);
 							FRayTracingMeshProcessor RayTracingMeshProcessor(&CommandContext, Scene, &ReferenceView);
@@ -750,12 +750,12 @@ bool FDeferredShadingSceneRenderer::DispatchRayTracingWorldUpdates(FRHICommandLi
 
 	Scene->GetRayTracingDynamicGeometryCollection()->DispatchUpdates(RHICmdList);
 
-	for (int32 ViewIndex = 0; ViewIndex < Views.Num(); ++ViewIndex)
-	{
-		FViewInfo& View = Views[ViewIndex];
-		SET_DWORD_STAT(STAT_RayTracingInstances, View.RayTracingGeometryInstances.Num());
-		FRayTracingSceneInitializer Initializer;
-		Initializer.Instances = View.RayTracingGeometryInstances;
+		for (int32 ViewIndex = 0; ViewIndex < Views.Num(); ++ViewIndex)
+		{
+			FViewInfo& View = Views[ViewIndex];
+			SET_DWORD_STAT(STAT_RayTracingInstances, View.RayTracingGeometryInstances.Num());
+			FRayTracingSceneInitializer Initializer;
+			Initializer.Instances = View.RayTracingGeometryInstances;
 		Initializer.ShaderSlotsPerGeometrySegment = RAY_TRACING_NUM_SHADER_SLOTS;
 		View.RayTracingScene.RayTracingSceneRHI = RHICreateRayTracingScene(Initializer);
 		RHICmdList.BuildAccelerationStructure(View.RayTracingScene.RayTracingSceneRHI);
@@ -771,7 +771,7 @@ bool FDeferredShadingSceneRenderer::DispatchRayTracingWorldUpdates(FRHICommandLi
 		PreparePathTracing(View, RayGenShaders);
 
 		if (RayGenShaders.Num())
-		{
+	{
 			auto DefaultHitShader = View.ShaderMap->GetShader<FOpaqueShadowHitGroup>()->GetRayTracingShader();
 			auto DefaultMissShader = View.ShaderMap->GetShader<FDefaultMaterialMS>()->GetRayTracingShader();
 
@@ -780,7 +780,7 @@ bool FDeferredShadingSceneRenderer::DispatchRayTracingWorldUpdates(FRHICommandLi
 				DefaultMissShader,
 				DefaultHitShader
 			);
-		}
+	}
 	}
 
 	return true;
@@ -1123,14 +1123,13 @@ void FDeferredShadingSceneRenderer::Render(FRHICommandListImmediate& RHICmdList)
 
 	const bool bShouldRenderVelocities = ShouldRenderVelocities();
 	const bool bBasePassCanOutputVelocity = FVelocityRendering::BasePassCanOutputVelocity(FeatureLevel);
-	const bool bUseSelectiveBasePassOutputs = bUseGBuffer && UseSelectiveBasePassOutputs();
+	const bool bUseSelectiveBasePassOutputs = IsUsingSelectiveBasePassOutputs(ShaderPlatform);
 
 	SceneContext.ResolveSceneDepthTexture(RHICmdList, FResolveRect(0, 0, FamilySize.X, FamilySize.Y));
 	SceneContext.ResolveSceneDepthToAuxiliaryTexture(RHICmdList);
 
 	ComputeLightGrid(RHICmdList, bComputeLightGrid);
 
-	if (bUseGBuffer || IsSimpleForwardShadingEnabled(ShaderPlatform))
 	{
 		SCOPE_CYCLE_COUNTER(STAT_FDeferredShadingSceneRenderer_AllocGBufferTargets);
 		SceneContext.PreallocGBufferTargets(); // Even if !bShouldRenderVelocities, the velocity buffer must be bound because it's a compile time option for the shader.
@@ -1389,16 +1388,11 @@ void FDeferredShadingSceneRenderer::Render(FRHICommandListImmediate& RHICmdList)
 		bRequiresFarZQuadClear = false;
 	}
 
-	if (bUseGBuffer)
 	{
 		SCOPE_CYCLE_COUNTER(STAT_FDeferredShadingSceneRenderer_Resolve_After_Basepass);
+		// Will early return if simple forward
 		SceneContext.FinishGBufferPassAndResolve(RHICmdList);
 	}
-	else
-	{
-		// #todo-renderpasses which paths can lead us here?
-		RHICmdList.EndRenderPass();
-	}
 
 	if (!bAllowReadonlyDepthBasePass)
 	{
diff --git a/Engine/Source/Runtime/Renderer/Private/PostProcess/SceneRenderTargets.cpp b/Engine/Source/Runtime/Renderer/Private/PostProcess/SceneRenderTargets.cpp
--- a/Engine/Source/Runtime/Renderer/Private/PostProcess/SceneRenderTargets.cpp
+++ b/Engine/Source/Runtime/Renderer/Private/PostProcess/SceneRenderTargets.cpp
@@ -680,31 +680,37 @@ int32 FSceneRenderTargets::FillGBufferRenderPassInfo(ERenderTargetLoadAction Col
 	OutRenderPassInfo.ColorRenderTargets[MRTCount].MipIndex = 0;
 	MRTCount++;
 
-	OutRenderPassInfo.ColorRenderTargets[MRTCount].Action = MakeRenderTargetActions(ColorLoadAction, ERenderTargetStoreAction::EStore);
-	OutRenderPassInfo.ColorRenderTargets[MRTCount].RenderTarget = GBufferA->GetRenderTargetItem().TargetableTexture;
-	OutRenderPassInfo.ColorRenderTargets[MRTCount].ArraySlice = -1;
-	OutRenderPassInfo.ColorRenderTargets[MRTCount].MipIndex = 0;
-	MRTCount++;
+	const EShaderPlatform ShaderPlatform = GetFeatureLevelShaderPlatform(CurrentFeatureLevel);
+	const bool bUseGBuffer = IsUsingGBuffers(ShaderPlatform);
 
-	OutRenderPassInfo.ColorRenderTargets[MRTCount].Action = MakeRenderTargetActions(ColorLoadAction, ERenderTargetStoreAction::EStore);
-	OutRenderPassInfo.ColorRenderTargets[MRTCount].RenderTarget = GBufferB->GetRenderTargetItem().TargetableTexture;
-	OutRenderPassInfo.ColorRenderTargets[MRTCount].ArraySlice = -1;
-	OutRenderPassInfo.ColorRenderTargets[MRTCount].MipIndex = 0;
-	MRTCount++;
+	if (bUseGBuffer)
+	{
+		OutRenderPassInfo.ColorRenderTargets[MRTCount].Action = MakeRenderTargetActions(ColorLoadAction, ERenderTargetStoreAction::EStore);
+		OutRenderPassInfo.ColorRenderTargets[MRTCount].RenderTarget = GBufferA->GetRenderTargetItem().TargetableTexture;
+		OutRenderPassInfo.ColorRenderTargets[MRTCount].ArraySlice = -1;
+		OutRenderPassInfo.ColorRenderTargets[MRTCount].MipIndex = 0;
+		MRTCount++;
 
-	OutRenderPassInfo.ColorRenderTargets[MRTCount].Action = MakeRenderTargetActions(ColorLoadAction, ERenderTargetStoreAction::EStore);
-	OutRenderPassInfo.ColorRenderTargets[MRTCount].RenderTarget = GBufferC->GetRenderTargetItem().TargetableTexture;
-	OutRenderPassInfo.ColorRenderTargets[MRTCount].ArraySlice = -1;
-	OutRenderPassInfo.ColorRenderTargets[MRTCount].MipIndex = 0;
-	MRTCount++;
+		OutRenderPassInfo.ColorRenderTargets[MRTCount].Action = MakeRenderTargetActions(ColorLoadAction, ERenderTargetStoreAction::EStore);
+		OutRenderPassInfo.ColorRenderTargets[MRTCount].RenderTarget = GBufferB->GetRenderTargetItem().TargetableTexture;
+		OutRenderPassInfo.ColorRenderTargets[MRTCount].ArraySlice = -1;
+		OutRenderPassInfo.ColorRenderTargets[MRTCount].MipIndex = 0;
+		MRTCount++;
+
+		OutRenderPassInfo.ColorRenderTargets[MRTCount].Action = MakeRenderTargetActions(ColorLoadAction, ERenderTargetStoreAction::EStore);
+		OutRenderPassInfo.ColorRenderTargets[MRTCount].RenderTarget = GBufferC->GetRenderTargetItem().TargetableTexture;
+		OutRenderPassInfo.ColorRenderTargets[MRTCount].ArraySlice = -1;
+		OutRenderPassInfo.ColorRenderTargets[MRTCount].MipIndex = 0;
+		MRTCount++;
+	}
 
 	// The velocity buffer needs to be bound before other optionnal rendertargets (when UseSelectiveBasePassOutputs() is true).
 	// Otherwise there is an issue on some AMD hardware where the target does not get updated. Seems to be related to the velocity buffer format as it works fine with other targets.
 	if (bAllocateVelocityGBuffer)
 	{
 		OutVelocityRTIndex = MRTCount;
-		check(OutVelocityRTIndex == 4); // As defined in BasePassPixelShader.usf
-
+		check(OutVelocityRTIndex == 4 || (!bUseGBuffer && OutVelocityRTIndex == 1)); // As defined in BasePassPixelShader.usf
+		
 		OutRenderPassInfo.ColorRenderTargets[MRTCount].Action = MakeRenderTargetActions(ColorLoadAction, ERenderTargetStoreAction::EStore);
 		OutRenderPassInfo.ColorRenderTargets[MRTCount].RenderTarget = SceneVelocity->GetRenderTargetItem().TargetableTexture;
 		OutRenderPassInfo.ColorRenderTargets[MRTCount].ArraySlice = -1;
@@ -716,20 +722,23 @@ int32 FSceneRenderTargets::FillGBufferRenderPassInfo(ERenderTargetLoadAction Col
 		OutVelocityRTIndex = -1;
 	}
 
-	OutRenderPassInfo.ColorRenderTargets[MRTCount].Action = MakeRenderTargetActions(ColorLoadAction, ERenderTargetStoreAction::EStore);
-	OutRenderPassInfo.ColorRenderTargets[MRTCount].RenderTarget = GBufferD->GetRenderTargetItem().TargetableTexture;
-	OutRenderPassInfo.ColorRenderTargets[MRTCount].ArraySlice = -1;
-	OutRenderPassInfo.ColorRenderTargets[MRTCount].MipIndex = 0;
-	MRTCount++;
-
-	if (bAllowStaticLighting)
+	if (bUseGBuffer)
 	{
-		check(MRTCount == (bAllocateVelocityGBuffer ? 6 : 5)); // As defined in BasePassPixelShader.usf
 		OutRenderPassInfo.ColorRenderTargets[MRTCount].Action = MakeRenderTargetActions(ColorLoadAction, ERenderTargetStoreAction::EStore);
-		OutRenderPassInfo.ColorRenderTargets[MRTCount].RenderTarget = GBufferE->GetRenderTargetItem().TargetableTexture;
+		OutRenderPassInfo.ColorRenderTargets[MRTCount].RenderTarget = GBufferD->GetRenderTargetItem().TargetableTexture;
 		OutRenderPassInfo.ColorRenderTargets[MRTCount].ArraySlice = -1;
 		OutRenderPassInfo.ColorRenderTargets[MRTCount].MipIndex = 0;
 		MRTCount++;
+
+		if (bAllowStaticLighting)
+		{
+			check(MRTCount == (bAllocateVelocityGBuffer ? 6 : 5)); // As defined in BasePassPixelShader.usf
+			OutRenderPassInfo.ColorRenderTargets[MRTCount].Action = MakeRenderTargetActions(ColorLoadAction, ERenderTargetStoreAction::EStore);
+			OutRenderPassInfo.ColorRenderTargets[MRTCount].RenderTarget = GBufferE->GetRenderTargetItem().TargetableTexture;
+			OutRenderPassInfo.ColorRenderTargets[MRTCount].ArraySlice = -1;
+			OutRenderPassInfo.ColorRenderTargets[MRTCount].MipIndex = 0;
+			MRTCount++;
+		}
 	}
 
 	check(MRTCount <= MaxSimultaneousRenderTargets);
@@ -740,16 +749,22 @@ int32 FSceneRenderTargets::GetGBufferRenderTargets(ERenderTargetLoadAction Color
 {
 	int32 MRTCount = 0;
 	OutRenderTargets[MRTCount++] = FRHIRenderTargetView(GetSceneColorSurface(), 0, -1, ColorLoadAction, ERenderTargetStoreAction::EStore);
-	OutRenderTargets[MRTCount++] = FRHIRenderTargetView(GBufferA->GetRenderTargetItem().TargetableTexture, 0, -1, ColorLoadAction, ERenderTargetStoreAction::EStore);
-	OutRenderTargets[MRTCount++] = FRHIRenderTargetView(GBufferB->GetRenderTargetItem().TargetableTexture, 0, -1, ColorLoadAction, ERenderTargetStoreAction::EStore);
-	OutRenderTargets[MRTCount++] = FRHIRenderTargetView(GBufferC->GetRenderTargetItem().TargetableTexture, 0, -1, ColorLoadAction, ERenderTargetStoreAction::EStore);
+
+	const EShaderPlatform ShaderPlatform = GetFeatureLevelShaderPlatform(CurrentFeatureLevel);
+	const bool bUseGBuffer = IsUsingGBuffers(ShaderPlatform);
+	if (bUseGBuffer)
+	{
+		OutRenderTargets[MRTCount++] = FRHIRenderTargetView(GBufferA->GetRenderTargetItem().TargetableTexture, 0, -1, ColorLoadAction, ERenderTargetStoreAction::EStore);
+		OutRenderTargets[MRTCount++] = FRHIRenderTargetView(GBufferB->GetRenderTargetItem().TargetableTexture, 0, -1, ColorLoadAction, ERenderTargetStoreAction::EStore);
+		OutRenderTargets[MRTCount++] = FRHIRenderTargetView(GBufferC->GetRenderTargetItem().TargetableTexture, 0, -1, ColorLoadAction, ERenderTargetStoreAction::EStore);
+	}
 
 	// The velocity buffer needs to be bound before other optionnal rendertargets (when UseSelectiveBasePassOutputs() is true).
 	// Otherwise there is an issue on some AMD hardware where the target does not get updated. Seems to be related to the velocity buffer format as it works fine with other targets.
-	if (bAllocateVelocityGBuffer)
+	if (bAllocateVelocityGBuffer && !IsSimpleForwardShadingEnabled(ShaderPlatform))
 	{
 		OutVelocityRTIndex = MRTCount;
-		check(OutVelocityRTIndex == 4); // As defined in BasePassPixelShader.usf
+		check(OutVelocityRTIndex == 4 || (!bUseGBuffer && OutVelocityRTIndex == 1)); // As defined in BasePassPixelShader.usf
 		OutRenderTargets[MRTCount++] = FRHIRenderTargetView(SceneVelocity->GetRenderTargetItem().TargetableTexture, 0, -1, ColorLoadAction, ERenderTargetStoreAction::EStore);
 	}
 	else
@@ -757,12 +772,15 @@ int32 FSceneRenderTargets::GetGBufferRenderTargets(ERenderTargetLoadAction Color
 		OutVelocityRTIndex = -1;
 	}
 
-	OutRenderTargets[MRTCount++] = FRHIRenderTargetView(GBufferD->GetRenderTargetItem().TargetableTexture, 0, -1, ColorLoadAction, ERenderTargetStoreAction::EStore);
-
-	if (bAllowStaticLighting)
+	if (bUseGBuffer)
 	{
-		check(MRTCount == (bAllocateVelocityGBuffer ? 6 : 5)); // As defined in BasePassPixelShader.usf
-		OutRenderTargets[MRTCount++] = FRHIRenderTargetView(GBufferE->GetRenderTargetItem().TargetableTexture, 0, -1, ColorLoadAction, ERenderTargetStoreAction::EStore);
+		OutRenderTargets[MRTCount++] = FRHIRenderTargetView(GBufferD->GetRenderTargetItem().TargetableTexture, 0, -1, ColorLoadAction, ERenderTargetStoreAction::EStore);
+
+		if (bAllowStaticLighting)
+		{
+			check(MRTCount == (bAllocateVelocityGBuffer ? 6 : 5)); // As defined in BasePassPixelShader.usf
+			OutRenderTargets[MRTCount++] = FRHIRenderTargetView(GBufferE->GetRenderTargetItem().TargetableTexture, 0, -1, ColorLoadAction, ERenderTargetStoreAction::EStore);
+		}
 	}
 
 	check(MRTCount <= MaxSimultaneousRenderTargets);
@@ -824,22 +842,8 @@ void FSceneRenderTargets::BeginRenderingGBuffer(FRHICommandList& RHICmdList, ERe
 	}
 
 	int32 VelocityRTIndex = -1;
-	int32 MRTCount;
-
 	FRHIRenderPassInfo RPInfo;
-
-	if (IsAnyForwardShadingEnabled(GetFeatureLevelShaderPlatform(CurrentFeatureLevel)))
-	{
-		MRTCount = 1;
-		RPInfo.ColorRenderTargets[0].Action = MakeRenderTargetActions(ColorLoadAction, ERenderTargetStoreAction::EStore);
-		RPInfo.ColorRenderTargets[0].RenderTarget = SceneColorTex;
-		RPInfo.ColorRenderTargets[0].ArraySlice = -1;
-		RPInfo.ColorRenderTargets[0].MipIndex = 0;
-	}
-	else
-	{
-		MRTCount = FillGBufferRenderPassInfo(ColorLoadAction, RPInfo, VelocityRTIndex);
-	}
+	int32 MRTCount = FillGBufferRenderPassInfo(ColorLoadAction, RPInfo, VelocityRTIndex);
 
 	//make sure our conditions for shader clear fallback are valid.
 	check(RPInfo.ColorRenderTargets[0].RenderTarget == SceneColorTex);
@@ -912,15 +916,26 @@ void FSceneRenderTargets::FinishGBufferPassAndResolve(FRHICommandListImmediate&
 {
 	RHICmdList.EndRenderPass();
 
+	const EShaderPlatform ShaderPlatform = GetFeatureLevelShaderPlatform(CurrentFeatureLevel);
+	if (IsSimpleForwardShadingEnabled(ShaderPlatform))
+	{
+		// Currently nothing to do here for simple forward shading.
+		return;
+	}
+
 	int32 VelocityRTIndex;
 	FRHIRenderTargetView RenderTargets[MaxSimultaneousRenderTargets];
 	int32 NumMRTs = GetGBufferRenderTargets(ERenderTargetLoadAction::ELoad, RenderTargets, VelocityRTIndex);
 
 	FResolveParams ResolveParams;
-	for (int32 i = 0; i < NumMRTs; ++i)
+
+	// Skip the SceneColor resolve when using forward shading as this is done after this function is called.
+	int32 i = IsForwardShadingEnabled(ShaderPlatform) ? 1 : 0;
+
+	for ( ; i < NumMRTs; ++i)
 	{
 		 // When the basepass outputs to the velocity buffer, don't resolve it yet if selective outputs are enabled, as it will be resolved after the velocity pass.
-		if (i != VelocityRTIndex || !UseSelectiveBasePassOutputs())
+		if (i != VelocityRTIndex || !IsUsingSelectiveBasePassOutputs(ShaderPlatform))
 		{
 			RHICmdList.CopyToResolveTarget(RenderTargets[i].Texture, RenderTargets[i].Texture, ResolveParams);
 		}
@@ -933,10 +948,14 @@ int32 FSceneRenderTargets::GetNumGBufferTargets() const
 {
 	int32 NumGBufferTargets = 1;
 
-	if (CurrentFeatureLevel >= ERHIFeatureLevel::SM4 && !IsAnyForwardShadingEnabled(GetFeatureLevelShaderPlatform(CurrentFeatureLevel)))
+	if (CurrentFeatureLevel >= ERHIFeatureLevel::SM4)
 	{
-		// This needs to match TBasePassPixelShaderBaseType::ModifyCompilationEnvironment()
-		NumGBufferTargets = bAllowStaticLighting ? 6 : 5;
+		const EShaderPlatform ShaderPlatform = GetFeatureLevelShaderPlatform(CurrentFeatureLevel);
+		if (IsUsingGBuffers(ShaderPlatform))
+		{
+			// This needs to match TBasePassPixelShaderBaseType::ModifyCompilationEnvironment()
+			NumGBufferTargets = bAllowStaticLighting ? 6 : 5;
+		}
 
 		if (bAllocateVelocityGBuffer)
 		{
@@ -1164,17 +1183,17 @@ void FSceneRenderTargets::AllocGBufferTargets(FRHICommandList& RHICmdList)
 			GRenderTargetPool.FindFreeElement(RHICmdList, Desc, GBufferE, TEXT("GBufferE"));
 		}
 
-		if (bAllocateVelocityGBuffer)
-		{
-			FPooledRenderTargetDesc VelocityRTDesc = FVelocityRendering::GetRenderTargetDesc();
-			VelocityRTDesc.Flags |= GFastVRamConfig.GBufferVelocity; 
-			GRenderTargetPool.FindFreeElement(RHICmdList, VelocityRTDesc, SceneVelocity, TEXT("GBufferVelocity"));
-		}
-
 		// otherwise we have a severe problem
 		check(GBufferA);
 	}
 
+	if (bAllocateVelocityGBuffer)
+	{
+		FPooledRenderTargetDesc VelocityRTDesc = FVelocityRendering::GetRenderTargetDesc();
+		VelocityRTDesc.Flags |= GFastVRamConfig.GBufferVelocity;
+		GRenderTargetPool.FindFreeElement(RHICmdList, VelocityRTDesc, SceneVelocity, TEXT("GBufferVelocity"));
+	}
+
 	GBufferRefCount = 1;
 }
 
diff --git a/Engine/Source/Runtime/Renderer/Private/RendererScene.cpp b/Engine/Source/Runtime/Renderer/Private/RendererScene.cpp
--- a/Engine/Source/Runtime/Renderer/Private/RendererScene.cpp
+++ b/Engine/Source/Runtime/Renderer/Private/RendererScene.cpp
@@ -2361,13 +2361,13 @@ void FSceneVelocityData::StartFrame(FScene* Scene)
 	}
 }
 
-void FScene::GetPrimitiveUniformShaderParameters_RenderThread(const FPrimitiveSceneInfo* PrimitiveSceneInfo, bool& bHasPrecomputedVolumetricLightmap, FMatrix& PreviousLocalToWorld, int32& SingleCaptureIndex) const 
+void FScene::GetPrimitiveUniformShaderParameters_RenderThread(const FPrimitiveSceneInfo* PrimitiveSceneInfo, bool& bHasPrecomputedVolumetricLightmap, FMatrix& PreviousLocalToWorld, int32& SingleCaptureIndex, bool& bOutputVelocity) const 
 {
 	bHasPrecomputedVolumetricLightmap = VolumetricLightmapSceneData.HasData();
 	
-	const bool bVelocityDataFound = VelocityData.GetComponentPreviousLocalToWorld(PrimitiveSceneInfo->PrimitiveComponentId, PreviousLocalToWorld);
+	bOutputVelocity = VelocityData.GetComponentPreviousLocalToWorld(PrimitiveSceneInfo->PrimitiveComponentId, PreviousLocalToWorld);
 
-	if (!bVelocityDataFound)
+	if (!bOutputVelocity)
 	{
 		PreviousLocalToWorld = PrimitiveSceneInfo->Proxy->GetLocalToWorld();
 	}
diff --git a/Engine/Source/Runtime/Renderer/Private/ScenePrivate.h b/Engine/Source/Runtime/Renderer/Private/ScenePrivate.h
--- a/Engine/Source/Runtime/Renderer/Private/ScenePrivate.h
+++ b/Engine/Source/Runtime/Renderer/Private/ScenePrivate.h
@@ -2712,7 +2712,7 @@ public:
 	virtual bool HasPrecomputedVolumetricLightmap_RenderThread() const override;
 	virtual void AddPrecomputedVolumetricLightmap(const class FPrecomputedVolumetricLightmap* Volume) override;
 	virtual void RemovePrecomputedVolumetricLightmap(const class FPrecomputedVolumetricLightmap* Volume) override;
-	virtual void GetPrimitiveUniformShaderParameters_RenderThread(const FPrimitiveSceneInfo* PrimitiveSceneInfo, bool& bHasPrecomputedVolumetricLightmap, FMatrix& PreviousLocalToWorld, int32& SingleCaptureIndex) const override;
+	virtual void GetPrimitiveUniformShaderParameters_RenderThread(const FPrimitiveSceneInfo* PrimitiveSceneInfo, bool& bHasPrecomputedVolumetricLightmap, FMatrix& PreviousLocalToWorld, int32& SingleCaptureIndex, bool& bOutputVelocity) const override;
 	virtual void UpdateLightTransform(ULightComponent* Light) override;
 	virtual void UpdateLightColorAndBrightness(ULightComponent* Light) override;
 	virtual void AddExponentialHeightFog(UExponentialHeightFogComponent* FogComponent) override;
diff --git a/Engine/Source/Runtime/Renderer/Private/VelocityRendering.cpp b/Engine/Source/Runtime/Renderer/Private/VelocityRendering.cpp
--- a/Engine/Source/Runtime/Renderer/Private/VelocityRendering.cpp
+++ b/Engine/Source/Runtime/Renderer/Private/VelocityRendering.cpp
@@ -493,7 +493,7 @@ FPooledRenderTargetDesc FVelocityRendering::GetRenderTargetDesc()
 
 bool FVelocityRendering::BasePassCanOutputVelocity(EShaderPlatform ShaderPlatform)
 {
-	return !IsForwardShadingEnabled(ShaderPlatform) && CVarBasePassOutputsVelocity.GetValueOnAnyThread() == 1;
+	return IsUsingBasePassVelocity(ShaderPlatform);
 }
 
 bool FVelocityRendering::BasePassCanOutputVelocity(ERHIFeatureLevel::Type FeatureLevel)
@@ -504,7 +504,7 @@ bool FVelocityRendering::BasePassCanOutputVelocity(ERHIFeatureLevel::Type Featur
 
 bool FVelocityRendering::VertexFactoryOnlyOutputsVelocityInBasePass(EShaderPlatform ShaderPlatform, bool bVertexFactorySupportsStaticLighting)
 {
-	return BasePassCanOutputVelocity(ShaderPlatform) && !(UseSelectiveBasePassOutputs() && bVertexFactorySupportsStaticLighting);;
+	return BasePassCanOutputVelocity(ShaderPlatform) && !(IsUsingSelectiveBasePassOutputs(ShaderPlatform) && bVertexFactorySupportsStaticLighting);;
 }
 
 bool FVelocityRendering::PrimitiveHasVelocity(ERHIFeatureLevel::Type FeatureLevel, const FPrimitiveSceneInfo* PrimitiveSceneInfo)
@@ -516,10 +516,7 @@ bool FVelocityRendering::PrimitiveHasVelocity(ERHIFeatureLevel::Type FeatureLeve
 	}
 
 	// If the base pass is allowed to render velocity in the GBuffer, only mesh with static lighting need the velocity pass.
-	const bool bVelocityInGBuffer = FVelocityRendering::BasePassCanOutputVelocity(FeatureLevel) 
-		&& !(UseSelectiveBasePassOutputs() && (PrimitiveSceneInfo->Proxy->HasStaticLighting()));
-
-	if (bVelocityInGBuffer)
+	if (FVelocityRendering::BasePassCanOutputVelocity(FeatureLevel) && !(IsUsingSelectiveBasePassOutputs(GetFeatureLevelShaderPlatform(FeatureLevel)) && PrimitiveSceneInfo->Proxy->HasStaticLighting()))
 	{
 		return false;
 	}
