GPUScene - optimized UploadDynamicPrimitiveShaderDataForView. If there are no dynamic primitives for current view, then just skip entire function and use instead Scene.GPUScene.PrimitiveBuffer. This skips most GPUScene per shadow map view GPU data copies.

diff --git a/Engine/Source/Runtime/Renderer/Private/GPUScene.cpp b/Engine/Source/Runtime/Renderer/Private/GPUScene.cpp
--- a/Engine/Source/Runtime/Renderer/Private/GPUScene.cpp
+++ b/Engine/Source/Runtime/Renderer/Private/GPUScene.cpp
@@ -305,48 +305,56 @@ void UploadDynamicPrimitiveShaderDataForView(FRHICommandList& RHICmdList, FScene
 	{
 		QUICK_SCOPE_CYCLE_COUNTER(STAT_UploadDynamicPrimitiveShaderDataForView);
 
-		FRWBufferStructured& ViewPrimitiveShaderDataBuffer = View.ViewState ? View.ViewState->PrimitiveShaderDataBuffer : View.OneFramePrimitiveShaderDataBuffer;
-
-		const int32 NumPrimitiveEntries = Scene.Primitives.Num() + View.DynamicPrimitiveShaderData.Num();
-		const uint32 PrimitiveSceneNumFloat4s = NumPrimitiveEntries * FPrimitiveSceneShaderData::PrimitiveDataStrideInFloat4s;
+		const int32 NumPrimitiveDataUploads = View.DynamicPrimitiveShaderData.Num();
+		if (NumPrimitiveDataUploads > 0)
+		{
+			FRWBufferStructured& ViewPrimitiveShaderDataBuffer = View.ViewState ? View.ViewState->PrimitiveShaderDataBuffer : View.OneFramePrimitiveShaderDataBuffer;
 
-		uint32 ViewPrimitiveSceneNumFloat4s = FMath::RoundUpToPowerOfTwo(PrimitiveSceneNumFloat4s);
-		uint32 BytesPerElement = GPixelFormats[PF_A32B32G32R32F].BlockBytes;
+			const int32 NumPrimitiveEntries = Scene.Primitives.Num() + View.DynamicPrimitiveShaderData.Num();
+			const uint32 PrimitiveSceneNumFloat4s = NumPrimitiveEntries * FPrimitiveSceneShaderData::PrimitiveDataStrideInFloat4s;
 
-		// Reserve enough space
-		if (ViewPrimitiveSceneNumFloat4s * BytesPerElement != ViewPrimitiveShaderDataBuffer.NumBytes)
-		{
-			ViewPrimitiveShaderDataBuffer.Release();
-			ViewPrimitiveShaderDataBuffer.Initialize(BytesPerElement, ViewPrimitiveSceneNumFloat4s, 0, TEXT("ViewPrimitiveShaderDataBuffer"));
-		}
+			uint32 ViewPrimitiveSceneNumFloat4s = FMath::RoundUpToPowerOfTwo(PrimitiveSceneNumFloat4s);
+			uint32 BytesPerElement = GPixelFormats[PF_A32B32G32R32F].BlockBytes;
 
-		// Copy scene primitive data into view primitive data
-		MemcpyBuffer(RHICmdList, Scene.GPUScene.PrimitiveBuffer, ViewPrimitiveShaderDataBuffer, Scene.Primitives.Num() * FPrimitiveSceneShaderData::PrimitiveDataStrideInFloat4s);
-		
-		const int32 NumPrimitiveDataUploads = View.DynamicPrimitiveShaderData.Num();
+			// Reserve enough space
+			if (ViewPrimitiveSceneNumFloat4s * BytesPerElement != ViewPrimitiveShaderDataBuffer.NumBytes)
+			{
+				ViewPrimitiveShaderDataBuffer.Release();
+				ViewPrimitiveShaderDataBuffer.Initialize(BytesPerElement, ViewPrimitiveSceneNumFloat4s, 0, TEXT("ViewPrimitiveShaderDataBuffer"));
+			}
 
-		// Append View.DynamicPrimitiveShaderData to the end of the view primitive data buffer
-		if (NumPrimitiveDataUploads > 0)
-		{
-			FScatterUploadBuilder PrimitivesUploadBuilder(NumPrimitiveDataUploads, FPrimitiveSceneShaderData::PrimitiveDataStrideInFloat4s, Scene.GPUScene.PrimitivesUploadScatterBuffer, Scene.GPUScene.PrimitivesUploadDataBuffer);
+			// Copy scene primitive data into view primitive data
+			MemcpyBuffer(RHICmdList, Scene.GPUScene.PrimitiveBuffer, ViewPrimitiveShaderDataBuffer, Scene.Primitives.Num() * FPrimitiveSceneShaderData::PrimitiveDataStrideInFloat4s);
 
-			for (int32 DynamicUploadIndex = 0; DynamicUploadIndex < View.DynamicPrimitiveShaderData.Num(); DynamicUploadIndex++)
+			// Append View.DynamicPrimitiveShaderData to the end of the view primitive data buffer
+			if (NumPrimitiveDataUploads > 0)
 			{
-				FPrimitiveSceneShaderData PrimitiveSceneData(View.DynamicPrimitiveShaderData[DynamicUploadIndex]);
-				// Place dynamic primitive shader data just after scene primitive data
-				PrimitivesUploadBuilder.Add(Scene.Primitives.Num() + DynamicUploadIndex, &PrimitiveSceneData.Data[0]);
+				FScatterUploadBuilder PrimitivesUploadBuilder(NumPrimitiveDataUploads, FPrimitiveSceneShaderData::PrimitiveDataStrideInFloat4s, Scene.GPUScene.PrimitivesUploadScatterBuffer, Scene.GPUScene.PrimitivesUploadDataBuffer);
+
+				for (int32 DynamicUploadIndex = 0; DynamicUploadIndex < View.DynamicPrimitiveShaderData.Num(); DynamicUploadIndex++)
+				{
+					FPrimitiveSceneShaderData PrimitiveSceneData(View.DynamicPrimitiveShaderData[DynamicUploadIndex]);
+					// Place dynamic primitive shader data just after scene primitive data
+					PrimitivesUploadBuilder.Add(Scene.Primitives.Num() + DynamicUploadIndex, &PrimitiveSceneData.Data[0]);
+				}
+
+				RHICmdList.TransitionResource(EResourceTransitionAccess::ERWBarrier, EResourceTransitionPipeline::EComputeToCompute, ViewPrimitiveShaderDataBuffer.UAV);
+
+				PrimitivesUploadBuilder.UploadTo(RHICmdList, ViewPrimitiveShaderDataBuffer);
 			}
 
-			RHICmdList.TransitionResource(EResourceTransitionAccess::ERWBarrier, EResourceTransitionPipeline::EComputeToCompute, ViewPrimitiveShaderDataBuffer.UAV);
+			RHICmdList.TransitionResource(EResourceTransitionAccess::EReadable, EResourceTransitionPipeline::EComputeToGfx, ViewPrimitiveShaderDataBuffer.UAV);
 
-			PrimitivesUploadBuilder.UploadTo(RHICmdList, ViewPrimitiveShaderDataBuffer);
+			View.CachedViewUniformShaderParameters->PrimitiveSceneData = ViewPrimitiveShaderDataBuffer.SRV;
+		}
+		else
+		{
+			// No dynamic primitives for this view, we just use Scene.GPUScene.PrimitiveBuffer.
+			View.CachedViewUniformShaderParameters->PrimitiveSceneData = Scene.GPUScene.PrimitiveBuffer.SRV;
 		}
-
-		RHICmdList.TransitionResource(EResourceTransitionAccess::EReadable, EResourceTransitionPipeline::EComputeToGfx, ViewPrimitiveShaderDataBuffer.UAV);
 
 		// Update view uniform buffer
 		View.CachedViewUniformShaderParameters->LightmapSceneData = Scene.GPUScene.LightmapDataBuffer.SRV;
-		View.CachedViewUniformShaderParameters->PrimitiveSceneData = ViewPrimitiveShaderDataBuffer.SRV;
 		View.ViewUniformBuffer.UpdateUniformBufferImmediate(*View.CachedViewUniformShaderParameters);
 	}
 }
