ComputeDynamicMeshRelevance tests if objects should be or not in velocity pass. There were automatically added if bVelictyRelevance is true without testing if velocity is computed in gbuffer/base pass

#rnx

diff --git a/Engine/Source/Runtime/Renderer/Private/SceneVisibility.cpp b/Engine/Source/Runtime/Renderer/Private/SceneVisibility.cpp
--- a/Engine/Source/Runtime/Renderer/Private/SceneVisibility.cpp
+++ b/Engine/Source/Runtime/Renderer/Private/SceneVisibility.cpp
@@ -2543,7 +2543,7 @@ static void SetDynamicMeshElementViewCustomData(TArray<FViewInfo>& InViews, cons
 	}
 }
 
-void ComputeDynamicMeshRelevance(EShadingPath ShadingPath, bool bAddLightmapDensityCommands, const FPrimitiveViewRelevance& ViewRelevance, const FMeshBatchAndRelevance& MeshBatch, FViewInfo& View, FMeshPassMask& PassMask)
+void ComputeDynamicMeshRelevance(EShadingPath ShadingPath, bool bAddLightmapDensityCommands, const FPrimitiveViewRelevance& ViewRelevance, const FMeshBatchAndRelevance& MeshBatch, FViewInfo& View, FMeshPassMask& PassMask, FPrimitiveSceneInfo* PrimitiveSceneInfo, const FPrimitiveBounds& Bounds)
 {
 	const int32 NumElements = MeshBatch.Mesh->Elements.Num();
 
@@ -2593,7 +2593,9 @@ void ComputeDynamicMeshRelevance(EShadingPath ShadingPath, bool bAddLightmapDens
 		}
 #endif
 
-		if (ViewRelevance.bVelocityRelevance)
+		if (ViewRelevance.bVelocityRelevance
+				&& FVelocityRendering::PrimitiveHasVelocity(View.GetFeatureLevel(), PrimitiveSceneInfo)
+				&& FVelocityRendering::PrimitiveHasVelocityForView(View, Bounds.BoxSphereBounds, PrimitiveSceneInfo))
 		{
 			PassMask.Set(EMeshPass::Velocity);
 			View.NumVisibleDynamicMeshElements[EMeshPass::Velocity] += NumElements;
@@ -2710,6 +2712,7 @@ void FSceneRenderer::GatherDynamicMeshElements(
 				const uint8 ViewMaskFinal = (bIsInstancedStereo) ? ViewMask | 0x3 : ViewMask;
 
 				FPrimitiveSceneInfo* PrimitiveSceneInfo = InScene->Primitives[PrimitiveIndex];
+				const FPrimitiveBounds& Bounds = InScene->PrimitiveBounds[PrimitiveIndex];
 				Collector.SetPrimitive(PrimitiveSceneInfo->Proxy, PrimitiveSceneInfo->DefaultDynamicHitProxyId);
 
 				SetDynamicMeshElementViewCustomData(InViews, HasViewCustomDataMasks, PrimitiveSceneInfo);
@@ -2741,7 +2744,7 @@ void FSceneRenderer::GatherDynamicMeshElements(
 							const FMeshBatchAndRelevance& MeshBatch = View.DynamicMeshElements[ElementIndex];
 							FMeshPassMask& PassRelevance = View.DynamicMeshElementsPassRelevance[ElementIndex];
 
-							ComputeDynamicMeshRelevance(ShadingPath, bAddLightmapDensityCommands, ViewRelevance, MeshBatch, View, PassRelevance);
+							ComputeDynamicMeshRelevance(ShadingPath, bAddLightmapDensityCommands, ViewRelevance, MeshBatch, View, PassRelevance, PrimitiveSceneInfo, Bounds);
 						}
 					}
 				}
