Fixed failing shadow test due to incorrectly computed DynamicMeshEndIndices in CL 5810490.

diff --git a/Engine/Source/Runtime/Renderer/Private/SceneVisibility.cpp b/Engine/Source/Runtime/Renderer/Private/SceneVisibility.cpp
--- a/Engine/Source/Runtime/Renderer/Private/SceneVisibility.cpp
+++ b/Engine/Source/Runtime/Renderer/Private/SceneVisibility.cpp
@@ -2717,8 +2717,7 @@ void FSceneRenderer::GatherDynamicMeshElements(
 
 				SetDynamicMeshElementViewCustomData(InViews, HasViewCustomDataMasks, PrimitiveSceneInfo);
 
-				PrimitiveSceneInfo->Proxy->GetDynamicMeshElements(InViewFamily.Views, InViewFamily, ViewMaskFinal, Collector);
-
+				// Mark DynamicMeshEndIndices start.
 				if (PrimitiveIndex > 0)
 				{
 					for (int32 ViewIndex = 0; ViewIndex < ViewCount; ViewIndex++)
@@ -2727,6 +2726,14 @@ void FSceneRenderer::GatherDynamicMeshElements(
 					}
 				}
 
+				PrimitiveSceneInfo->Proxy->GetDynamicMeshElements(InViewFamily.Views, InViewFamily, ViewMaskFinal, Collector);
+
+				// Mark DynamicMeshEndIndices end.
+				for (int32 ViewIndex = 0; ViewIndex < ViewCount; ViewIndex++)
+				{
+					InViews[ViewIndex].DynamicMeshEndIndices[PrimitiveIndex] = Collector.GetMeshBatchCount(ViewIndex);
+				}
+
 				// Compute DynamicMeshElementsMeshPassRelevance for this primitive.
 				for (int32 ViewIndex = 0; ViewIndex < ViewCount; ViewIndex++)
 				{
@@ -2748,11 +2755,6 @@ void FSceneRenderer::GatherDynamicMeshElements(
 						}
 					}
 				}
-
-				for (int32 ViewIndex = 0; ViewIndex < ViewCount; ViewIndex++)
-				{
-					InViews[ViewIndex].DynamicMeshEndIndices[PrimitiveIndex] = Collector.GetMeshBatchCount(ViewIndex);
-				}
 			}
 		}
 	}
