Cherry-pick 5763172: Invalidate path tracer when the projection change (e.g., FOV).

diff --git a/Engine/Source/Runtime/Renderer/Private/SceneVisibility.cpp b/Engine/Source/Runtime/Renderer/Private/SceneVisibility.cpp
--- a/Engine/Source/Runtime/Renderer/Private/SceneVisibility.cpp
+++ b/Engine/Source/Runtime/Renderer/Private/SceneVisibility.cpp
@@ -3119,9 +3119,10 @@ void FSceneRenderer::PreVisibilityFrameSetup(FRHICommandListImmediate& RHICmdLis
 				bResetCamera ||
 				Scene->bPathTracingNeedsInvalidation ||
 				View.ViewRect != ViewState->PathTracingRect ||
+				View.ViewMatrices.GetProjectionNoAAMatrix() != View.ViewState->PrevFrameViewInfo.ViewMatrices.GetProjectionNoAAMatrix() ||
 				IsLargeCameraMovement(
-					View, ViewState->PrevFrameViewInfo.ViewMatrices.GetViewMatrix(), 
-					ViewState->PrevFrameViewInfo.ViewMatrices.GetViewOrigin(), 
+					View, ViewState->PrevFrameViewInfo.ViewMatrices.GetViewMatrix(),
+					ViewState->PrevFrameViewInfo.ViewMatrices.GetViewOrigin(),
 					0.1f, 0.1f)
 			);
 
