Cherry-pick 5764555: Fix path tracer invalidation caused by precision issue.

diff --git a/Engine/Source/Runtime/Renderer/Private/SceneVisibility.cpp b/Engine/Source/Runtime/Renderer/Private/SceneVisibility.cpp
--- a/Engine/Source/Runtime/Renderer/Private/SceneVisibility.cpp
+++ b/Engine/Source/Runtime/Renderer/Private/SceneVisibility.cpp
@@ -2799,7 +2799,7 @@ void FSceneRenderer::GatherDynamicMeshElements(
  */
 static bool IsLargeCameraMovement(FSceneView& View, const FMatrix& PrevViewMatrix, const FVector& PrevViewOrigin, float CameraRotationThreshold, float CameraTranslationThreshold)
 {
-	float RotationThreshold = FMath::Cos(CameraRotationThreshold * PI / 180.0f);
+	float RotationThreshold = FMath::Cos(FMath::DegreesToRadians(CameraRotationThreshold));
 	float ViewRightAngle = View.ViewMatrices.GetViewMatrix().GetColumn(0) | PrevViewMatrix.GetColumn(0);
 	float ViewUpAngle = View.ViewMatrices.GetViewMatrix().GetColumn(1) | PrevViewMatrix.GetColumn(1);
 	float ViewDirectionAngle = View.ViewMatrices.GetViewMatrix().GetColumn(2) | PrevViewMatrix.GetColumn(2);
@@ -3114,16 +3114,19 @@ void FSceneRenderer::PreVisibilityFrameSetup(FRHICommandListImmediate& RHICmdLis
 			const bool bResetCamera = (bFirstFrameOrTimeWasReset || View.bCameraCut || bIsLargeCameraMovement);
 			
 #if RHI_RAYTRACING
+			// Note: 0.18 deg is the minimum angle for avoiding numerical precision issue (which would cause constant invalidation)
+			const bool bIsThereALargeMomvement= IsLargeCameraMovement(
+				View, ViewState->PrevFrameViewInfo.ViewMatrices.GetViewMatrix(),
+				ViewState->PrevFrameViewInfo.ViewMatrices.GetViewOrigin(),
+				0.18f /*degree*/, 0.1f /*cm*/);
+			const bool bIsProjMatrixDifferent = View.ViewMatrices.GetProjectionNoAAMatrix() != View.ViewState->PrevFrameViewInfo.ViewMatrices.GetProjectionNoAAMatrix();
 			const bool bInvalidatePathTracer = View.RayTracingRenderMode == ERayTracingRenderMode::PathTracing &&
 			(
 				bResetCamera ||
 				Scene->bPathTracingNeedsInvalidation ||
 				View.ViewRect != ViewState->PathTracingRect ||
-				View.ViewMatrices.GetProjectionNoAAMatrix() != View.ViewState->PrevFrameViewInfo.ViewMatrices.GetProjectionNoAAMatrix() ||
-				IsLargeCameraMovement(
-					View, ViewState->PrevFrameViewInfo.ViewMatrices.GetViewMatrix(),
-					ViewState->PrevFrameViewInfo.ViewMatrices.GetViewOrigin(),
-					0.1f, 0.1f)
+				bIsProjMatrixDifferent ||
+				bIsThereALargeMomvement
 			);
 
 			if (bInvalidatePathTracer)
