Reflection shader parameters use unifrom buffer instead of loose data
This reduces mem usage of cached mesh draw commands on mobile (~2MB savings)

diff --git a/Engine/Shaders/Private/MobileBasePassPixelShader.usf b/Engine/Shaders/Private/MobileBasePassPixelShader.usf
--- a/Engine/Shaders/Private/MobileBasePassPixelShader.usf
+++ b/Engine/Shaders/Private/MobileBasePassPixelShader.usf
@@ -65,17 +65,16 @@ float4 SpotLightDirection[MAX_DYNAMIC_POINT_LIGHTS];
 
 /** Prenormalized capture of the scene that's closest to the object being rendered. */
 #if !FULLY_ROUGH
-	TextureCube ReflectionCubemap;
-	SamplerState ReflectionCubemapSampler;
-	// x,y,z - inverted average brightness, w - max sky cube mip
-	float4 MobileReflectionParams;
 	#if HQ_REFLECTIONS
 	#define MAX_HQ_REFLECTIONS 3
-
+	TextureCube ReflectionCubemap0;
+	SamplerState ReflectionCubemapSampler0;
 	TextureCube ReflectionCubemap1;
 	SamplerState ReflectionCubemapSampler1;
 	TextureCube ReflectionCubemap2;
 	SamplerState ReflectionCubemapSampler2;
+	// x,y,z - inverted average brightness for 0, 1, 2
+	float4 ReflectionInvAverageBrigtness;
 	float4 ReflectionPositionsAndRadii[MAX_HQ_REFLECTIONS];
 		#if ALLOW_CUBE_REFLECTIONS
 		float4x4 CaptureBoxTransformArray[MAX_HQ_REFLECTIONS];
@@ -227,9 +226,9 @@ half3 BlendReflectionCaptures(FMaterialPixelParameters MaterialParameters, half
 	half3 SpecularIBL = half3(0, 0, 0);
 	half BlendFactor = 1;
 
-	SpecularIBL += BlendReflectionCapture(MaterialParameters, Roughness, ReflectionCubemap, ReflectionCubemapSampler, AbsoluteSpecularMip, 0, MobileReflectionParams.x, BlendFactor);
-	SpecularIBL += BlendReflectionCapture(MaterialParameters, Roughness, ReflectionCubemap1, ReflectionCubemapSampler1, AbsoluteSpecularMip, 1, MobileReflectionParams.y, BlendFactor);
-	SpecularIBL += BlendReflectionCapture(MaterialParameters, Roughness, ReflectionCubemap2, ReflectionCubemapSampler2, AbsoluteSpecularMip, 2, MobileReflectionParams.z, BlendFactor);
+	SpecularIBL += BlendReflectionCapture(MaterialParameters, Roughness, ReflectionCubemap0, ReflectionCubemapSampler0, AbsoluteSpecularMip, 0, ReflectionInvAverageBrigtness.x, BlendFactor);
+	SpecularIBL += BlendReflectionCapture(MaterialParameters, Roughness, ReflectionCubemap1, ReflectionCubemapSampler1, AbsoluteSpecularMip, 1, ReflectionInvAverageBrigtness.y, BlendFactor);
+	SpecularIBL += BlendReflectionCapture(MaterialParameters, Roughness, ReflectionCubemap2, ReflectionCubemapSampler2, AbsoluteSpecularMip, 2, ReflectionInvAverageBrigtness.z, BlendFactor);
 
 	return SpecularIBL;
 }
@@ -257,14 +256,14 @@ half3 GetImageBasedReflectionLighting(FMaterialPixelParameters MaterialParameter
 
 	half3 ProjectedCaptureVector = MaterialParameters.ReflectionVector;	
 
-	half UsingSkyReflection = MobileReflectionParams.w > 0;
-	half CubemapMaxMip = UsingSkyReflection ? MobileReflectionParams.w : ResolvedView.ReflectionCubemapMaxMip;
+	half UsingSkyReflection = MobileReflectionCapture.Params.y > 0;
+	half CubemapMaxMip = UsingSkyReflection ? MobileReflectionCapture.Params.y : ResolvedView.ReflectionCubemapMaxMip;
 
 	// Compute fractional mip from roughness
 	half AbsoluteSpecularMip = ComputeReflectionCaptureMipFromRoughness(Roughness, CubemapMaxMip);
 	// Fetch from cubemap and convert to linear HDR
 	half3 SpecularIBL;
-	half4 SpecularIBLSample = ReflectionCubemap.SampleLevel(ReflectionCubemapSampler, ProjectedCaptureVector, AbsoluteSpecularMip);
+	half4 SpecularIBLSample = MobileReflectionCapture.Texture.SampleLevel(MobileReflectionCapture.TextureSampler, ProjectedCaptureVector, AbsoluteSpecularMip);
 	if (UsingSkyReflection)
 	{
 		SpecularIBL = SpecularIBLSample.rgb;
@@ -277,7 +276,7 @@ half3 GetImageBasedReflectionLighting(FMaterialPixelParameters MaterialParameter
 		SpecularIBL = SpecularIBL * SpecularIBL;
 		#if LQ_TEXTURE_LIGHTMAP || CACHED_POINT_INDIRECT_LIGHTING
 			// divide by average brightness for lightmap use cases only.
-			SpecularIBL *= MobileReflectionParams.x;
+			SpecularIBL *= MobileReflectionCapture.Params.x;
 		#endif
 	}
 
diff --git a/Engine/Source/Runtime/Engine/Private/Components/ReflectionCaptureComponent.cpp b/Engine/Source/Runtime/Engine/Private/Components/ReflectionCaptureComponent.cpp
--- a/Engine/Source/Runtime/Engine/Private/Components/ReflectionCaptureComponent.cpp
+++ b/Engine/Source/Runtime/Engine/Private/Components/ReflectionCaptureComponent.cpp
@@ -1254,3 +1254,28 @@ void FReflectionCaptureProxy::SetTransform(const FMatrix& InTransform)
 	ReflectionXAxisAndYScale = ReflectionXAxis.GetSafeNormal() * ScaleVector.Y;
 	ReflectionXAxisAndYScale.W = ScaleVector.Y / ScaleVector.Z;
 }
+
+void FReflectionCaptureProxy::UpdateMobileUniformBuffer()
+{
+	float InvBrightness = FMath::Max(FMath::Min(1.0f / EncodedHDRAverageBrightness, 65504.f), -65504.f);
+	FTexture* CaptureTexture = GBlackTextureCube;
+	if (EncodedHDRCubemap)
+	{
+		check(EncodedHDRCubemap->IsInitialized());
+		CaptureTexture = EncodedHDRCubemap;
+	}
+		
+	FMobileReflectionCaptureShaderParameters Parameters;
+	Parameters.Params = FVector4(InvBrightness, 0.f, 0.f, 0.f);
+	Parameters.Texture = CaptureTexture->TextureRHI;
+	Parameters.TextureSampler = CaptureTexture->SamplerStateRHI;
+
+	if (MobileUniformBuffer.GetReference())
+	{
+		MobileUniformBuffer.UpdateUniformBufferImmediate(Parameters);
+	}
+	else
+	{
+		MobileUniformBuffer = TUniformBufferRef<FMobileReflectionCaptureShaderParameters>::CreateUniformBufferImmediate(Parameters, UniformBuffer_MultiFrame);
+	}
+}
diff --git a/Engine/Source/Runtime/Engine/Private/Components/SkyLightComponent.cpp b/Engine/Source/Runtime/Engine/Private/Components/SkyLightComponent.cpp
--- a/Engine/Source/Runtime/Engine/Private/Components/SkyLightComponent.cpp
+++ b/Engine/Source/Runtime/Engine/Private/Components/SkyLightComponent.cpp
@@ -158,6 +158,33 @@ FLinearColor FSkyLightSceneProxy::GetEffectiveLightColor() const
 	return LightColor * GSkylightIntensityMultiplier;
 }
 
+void FSkyLightSceneProxy::UpdateMobileUniformBuffer()
+{
+	float InvBrightness = FMath::Max(FMath::Min(1.0f / AverageBrightness, 65504.f), -65504.f);
+	float SkyMaxMipIndex = 0.f;
+	FTexture* CaptureTexture = GBlackTextureCube;
+	if (ProcessedTexture)
+	{
+		check(ProcessedTexture->IsInitialized());
+		CaptureTexture = ProcessedTexture;
+		SkyMaxMipIndex = FMath::Log2(ProcessedTexture->GetSizeX());
+	}
+		
+	FMobileReflectionCaptureShaderParameters Parameters;
+	Parameters.Params = FVector4(InvBrightness, SkyMaxMipIndex, 0.f, 0.f);
+	Parameters.Texture = CaptureTexture->TextureRHI;
+	Parameters.TextureSampler = CaptureTexture->SamplerStateRHI;
+
+	if (MobileUniformBuffer.GetReference())
+	{
+		MobileUniformBuffer.UpdateUniformBufferImmediate(Parameters);
+	}
+	else
+	{
+		MobileUniformBuffer = TUniformBufferRef<FMobileReflectionCaptureShaderParameters>::CreateUniformBufferImmediate(Parameters, UniformBuffer_MultiFrame);
+	}
+}
+
 FSkyLightSceneProxy::FSkyLightSceneProxy(const USkyLightComponent* InLightComponent)
 	: LightComponent(InLightComponent)
 	, ProcessedTexture(InLightComponent->ProcessedSkyTexture)
diff --git a/Engine/Source/Runtime/Engine/Private/SceneManagement.cpp b/Engine/Source/Runtime/Engine/Private/SceneManagement.cpp
--- a/Engine/Source/Runtime/Engine/Private/SceneManagement.cpp
+++ b/Engine/Source/Runtime/Engine/Private/SceneManagement.cpp
@@ -954,4 +954,19 @@ void FMeshBatch::PreparePrimitiveUniformBuffer(const FPrimitiveSceneProxy* Primi
 		checkf(bPrimitiveShaderDataComesFromSceneBuffer || Elements[ElementIndex].PrimitiveUniformBuffer || Elements[ElementIndex].PrimitiveUniformBufferResource != NULL,
 			TEXT("FMeshBatch was not properly setup.  The primitive uniform buffer must be specified."));
 	}
-}
\ No newline at end of file
+}
+
+IMPLEMENT_GLOBAL_SHADER_PARAMETER_STRUCT(FMobileReflectionCaptureShaderParameters, "MobileReflectionCapture");
+
+void FDefaultMobileReflectionCaptureUniformBuffer::InitDynamicRHI()
+{
+	FMobileReflectionCaptureShaderParameters Parameters;
+	Parameters.Params = FVector4(1.f, 0.f, 0.f, 0.f);
+	Parameters.Texture = GBlackTextureCube->TextureRHI;
+	Parameters.TextureSampler = GBlackTextureCube->SamplerStateRHI;
+	SetContents(Parameters);
+	Super::InitDynamicRHI();
+}
+
+/** Global uniform buffer containing the default reflection data used in mobile renderer. */
+TGlobalResource<FDefaultMobileReflectionCaptureUniformBuffer> GDefaultMobileReflectionCaptureUniformBuffer;
\ No newline at end of file
diff --git a/Engine/Source/Runtime/Engine/Public/SceneManagement.h b/Engine/Source/Runtime/Engine/Public/SceneManagement.h
--- a/Engine/Source/Runtime/Engine/Public/SceneManagement.h
+++ b/Engine/Source/Runtime/Engine/Public/SceneManagement.h
@@ -956,6 +956,21 @@ inline bool DoesPlatformSupportDistanceFieldShadowing(EShaderPlatform Platform)
 		|| IsVulkanSM5Platform(Platform);
 }
 
+BEGIN_GLOBAL_SHADER_PARAMETER_STRUCT(FMobileReflectionCaptureShaderParameters,ENGINE_API)
+	SHADER_PARAMETER(FVector4, Params) // x - inv average brightness, y - sky cubemap max mip, zw - unused
+	SHADER_PARAMETER_TEXTURE(TextureCube, Texture)
+	SHADER_PARAMETER_SAMPLER(SamplerState, TextureSampler)
+END_GLOBAL_SHADER_PARAMETER_STRUCT()
+
+class FDefaultMobileReflectionCaptureUniformBuffer : public TUniformBuffer<FMobileReflectionCaptureShaderParameters>
+{
+	typedef TUniformBuffer<FMobileReflectionCaptureShaderParameters> Super;
+public:
+	virtual void InitDynamicRHI() override;
+};
+
+ENGINE_API extern TGlobalResource<FDefaultMobileReflectionCaptureUniformBuffer> GDefaultMobileReflectionCaptureUniformBuffer;
+
 /** Represents a USkyLightComponent to the rendering thread. */
 class ENGINE_API FSkyLightSceneProxy
 {
@@ -992,6 +1007,9 @@ public:
 	float MinOcclusion;
 	FLinearColor OcclusionTint;
 	int32 SamplesPerPixel;
+	/** Used with mobile renderer */
+	TUniformBufferRef<FMobileReflectionCaptureShaderParameters> MobileUniformBuffer;
+
 
 #if RHI_RAYTRACING
 	bool IsDirtyImportanceSamplingData;
@@ -1024,6 +1042,7 @@ public:
 	}
 	FLinearColor GetEffectiveLightColor() const;
 
+	void UpdateMobileUniformBuffer();
 private:
 	FLinearColor LightColor;
 };
@@ -1550,7 +1569,8 @@ public:
 
 	int32 PackedIndex;
 
-	/** Used in Feature level ES2 */
+	/** Used with mobile renderer */
+	TUniformBufferRef<FMobileReflectionCaptureShaderParameters> MobileUniformBuffer;
 	FTexture* EncodedHDRCubemap;
 	float EncodedHDRAverageBrightness;
 
@@ -1578,6 +1598,7 @@ public:
 	FReflectionCaptureProxy(const class UReflectionCaptureComponent* InComponent);
 
 	void SetTransform(const FMatrix& InTransform);
+	void UpdateMobileUniformBuffer();
 };
 
 /** Calculated wind data with support for accumulating other weighted wind data */
diff --git a/Engine/Source/Runtime/Renderer/Private/MobileBasePass.cpp b/Engine/Source/Runtime/Renderer/Private/MobileBasePass.cpp
--- a/Engine/Source/Runtime/Renderer/Private/MobileBasePass.cpp
+++ b/Engine/Source/Runtime/Renderer/Private/MobileBasePass.cpp
@@ -192,16 +192,6 @@ static bool UseSkyReflectionCapture(const FScene* RenderScene)
 		&& RenderScene->SkyLight->ProcessedTexture->TextureRHI;
 }
 
-static void GetSkyTextureParams(const FScene* Scene, float& AverageBrightnessOUT, FTexture*& ReflectionTextureOUT, float& OutSkyMaxMipIndex)
-{
-	if (Scene && Scene->SkyLight && Scene->SkyLight->ProcessedTexture->TextureRHI)
-	{
-		AverageBrightnessOUT = Scene->SkyLight->AverageBrightness;
-		ReflectionTextureOUT = Scene->SkyLight->ProcessedTexture;
-		OutSkyMaxMipIndex = FMath::Log2(ReflectionTextureOUT->GetSizeX());
-	}
-}
-
 const FLightSceneInfo* MobileBasePass::GetDirectionalLightInfo(const FScene* Scene, const FPrimitiveSceneProxy* PrimitiveSceneProxy)
 {
 	const FLightSceneInfo* MobileDirectionalLight = nullptr;
@@ -456,33 +446,17 @@ void TMobileBasePassPSPolicyParamType<FUniformLightMapPolicy>::GetShaderBindings
 
 	if (Scene)
 	{
-		static const int32 MaxNumReflections = FPrimitiveSceneInfo::MaxCachedReflectionCaptureProxies;
-		static_assert(MaxNumReflections == 3, "Update reflection array initializations to match MaxCachedReflectionCaptureProxies");
-		// set reflection parameters
-		const FShaderResourceParameter* ReflectionTextureParameters[MaxNumReflections] = { &ReflectionCubemap, &ReflectionCubemap1, &ReflectionCubemap2 };
-		const FShaderResourceParameter* ReflectionSamplerParameters[MaxNumReflections] = { &ReflectionSampler, &ReflectionSampler1, &ReflectionSampler2 };
-		FTexture* ReflectionCubemapTextures[MaxNumReflections] = { GBlackTextureCube, GBlackTextureCube, GBlackTextureCube };
-		FVector4 CapturePositions[MaxNumReflections] = { FVector4(0, 0, 0, 0), FVector4(0, 0, 0, 0), FVector4(0, 0, 0, 0) };
-		FVector4 ReflectionParams(1.0f, 1.0f, 1.0f, 0.0f);
-		
-		// If no reflection captures are available then attempt to use sky light's texture.
-		if (UseSkyReflectionCapture(Scene))
-		{
-			// if > 0 this will disable shader's RGBM decoding and enable sky light tinting of this envmap.
-			// ReflectionParams.X == inv average brightness
-			// ReflectionParams.W == max sky cube mip
-			if (FeatureLevel > ERHIFeatureLevel::ES2) // not-supported on ES2 at the moment
-			{
-				float AverageBrightness = 1.0f;
-				GetSkyTextureParams(Scene, AverageBrightness, ReflectionCubemapTextures[0], ReflectionParams.W);
-				ReflectionParams.X = FMath::Max(FMath::Min(1.0f / AverageBrightness, 65504.f), -65504.f);
-			}
-		}
-		else
+		// test for HQ reflection parameter existence
+		if (HQReflectionCubemaps[0].IsBound() || HQReflectionCubemaps[1].IsBound() || HQReflectionCubemaps[2].IsBound())
 		{
+			static const int32 MaxNumReflections = FPrimitiveSceneInfo::MaxCachedReflectionCaptureProxies;
+			static_assert(MaxNumReflections == 3, "Update reflection array initializations to match MaxCachedReflectionCaptureProxies");
+			// set reflection parameters
+			FTexture* ReflectionCubemapTextures[MaxNumReflections] = { GBlackTextureCube, GBlackTextureCube, GBlackTextureCube };
+			FVector4 CapturePositions[MaxNumReflections] = { FVector4(0, 0, 0, 0), FVector4(0, 0, 0, 0), FVector4(0, 0, 0, 0) };
+			FVector4 ReflectionParams(1.0f, 1.0f, 1.0f, 0.0f);
 			FPrimitiveSceneInfo* PrimitiveSceneInfo = PrimitiveSceneProxy ? PrimitiveSceneProxy->GetPrimitiveSceneInfo() : nullptr;
-			// test for HQ reflection parameter existence
-			if (PrimitiveSceneInfo && (ReflectionCubemap1.IsBound() || ReflectionCubemap2.IsBound() || ReflectionPositionsAndRadii.IsBound()))
+			if (PrimitiveSceneInfo)
 			{
 				for (int32 i = 0; i < MaxNumReflections; i++)
 				{
@@ -499,26 +473,32 @@ void TMobileBasePassPSPolicyParamType<FUniformLightMapPolicy>::GetShaderBindings
 					}
 				}
 			}
-			else if (ReflectionCubemap.IsBound())
+			for (int32 i = 0; i < MaxNumReflections; i++)
 			{
-				if (PrimitiveSceneInfo 
-					&& PrimitiveSceneInfo->CachedReflectionCaptureProxy
-					&& PrimitiveSceneInfo->CachedReflectionCaptureProxy->EncodedHDRCubemap
-					&& PrimitiveSceneInfo->CachedReflectionCaptureProxy->EncodedHDRCubemap->IsInitialized())
-				{
-					ReflectionParams.X = FMath::Max(FMath::Min(1.0f / PrimitiveSceneInfo->CachedReflectionCaptureProxy->EncodedHDRAverageBrightness, 65504.f), -65504.f);
-					ReflectionCubemapTextures[0] = PrimitiveSceneInfo->CachedReflectionCaptureProxy->EncodedHDRCubemap;
-				}
+				ShaderBindings.AddTexture(HQReflectionCubemaps[i], HQReflectionSamplers[i], ReflectionCubemapTextures[i]->SamplerStateRHI, ReflectionCubemapTextures[i]->TextureRHI);
 			}
+			ShaderBindings.Add(HQReflectionInvAverageBrigtnessParams, ReflectionParams);
+			ShaderBindings.Add(HQReflectionPositionsAndRadii, CapturePositions);
 		}
-
-		for (int32 i = 0; i < MaxNumReflections; i++)
+		else if (ReflectionParameter.IsBound())
 		{
-			ShaderBindings.AddTexture(*ReflectionTextureParameters[i], *ReflectionSamplerParameters[i], ReflectionCubemapTextures[i]->SamplerStateRHI, ReflectionCubemapTextures[i]->TextureRHI);
+			FUniformBufferRHIParamRef RelfectionUB = GDefaultMobileReflectionCaptureUniformBuffer.GetUniformBufferRHI();
+			// If no reflection captures are available then attempt to use sky light's texture.
+			if (UseSkyReflectionCapture(Scene) && FeatureLevel > ERHIFeatureLevel::ES2) // not-supported on ES2 at the moment
+			{
+				RelfectionUB = Scene->SkyLight->MobileUniformBuffer;
+			}
+			else
+			{
+				FPrimitiveSceneInfo* PrimitiveSceneInfo = PrimitiveSceneProxy ? PrimitiveSceneProxy->GetPrimitiveSceneInfo() : nullptr;
+				if (PrimitiveSceneInfo && PrimitiveSceneInfo->CachedReflectionCaptureProxy)
+				{
+					RelfectionUB = PrimitiveSceneInfo->CachedReflectionCaptureProxy->MobileUniformBuffer;
+				}
+			}
+			ShaderBindings.Add(ReflectionParameter, RelfectionUB);
 		}
-		ShaderBindings.Add(ReflectionPositionsAndRadii, CapturePositions);
-		ShaderBindings.Add(MobileReflectionParams, ReflectionParams);
-
+		
 		if (LightPositionAndInvRadiusParameter.IsBound() || SpotLightDirectionParameter.IsBound())
 		{
 			// Set dynamic point lights
@@ -532,7 +512,7 @@ void TMobileBasePassPSPolicyParamType<FUniformLightMapPolicy>::GetShaderBindings
 	}
 	else
 	{
-		ensure(!ReflectionCubemap.IsBound());
+		ensure(!ReflectionParameter.IsBound());
 	}
 
 	// Set directional light UB
diff --git a/Engine/Source/Runtime/Renderer/Private/MobileBasePassRendering.h b/Engine/Source/Runtime/Renderer/Private/MobileBasePassRendering.h
--- a/Engine/Source/Runtime/Renderer/Private/MobileBasePassRendering.h
+++ b/Engine/Source/Runtime/Renderer/Private/MobileBasePassRendering.h
@@ -223,22 +223,23 @@ public:
 	}
 
 	/** Initialization constructor. */
-	TMobileBasePassPSPolicyParamType(const FMeshMaterialShaderType::CompiledShaderInitializerType& Initializer):
-		FMeshMaterialShader(Initializer)
+	TMobileBasePassPSPolicyParamType(const FMeshMaterialShaderType::CompiledShaderInitializerType& Initializer)
+		: FMeshMaterialShader(Initializer)
 	{
 		LightMapPolicyType::PixelParametersType::Bind(Initializer.ParameterMap);
 		PassUniformBuffer.Bind(Initializer.ParameterMap, FMobileBasePassUniformParameters::StaticStructMetadata.GetShaderVariableName());
 		
 		MobileDirectionLightBufferParam.Bind(Initializer.ParameterMap, FMobileDirectionalLightShaderParameters::StaticStructMetadata.GetShaderVariableName());
+		ReflectionParameter.Bind(Initializer.ParameterMap, FMobileReflectionCaptureShaderParameters::StaticStructMetadata.GetShaderVariableName());
 
-		ReflectionCubemap.Bind(Initializer.ParameterMap, TEXT("ReflectionCubemap"));
-		ReflectionSampler.Bind(Initializer.ParameterMap, TEXT("ReflectionCubemapSampler"));
-		ReflectionCubemap1.Bind(Initializer.ParameterMap, TEXT("ReflectionCubemap1"));
-		ReflectionSampler1.Bind(Initializer.ParameterMap, TEXT("ReflectionCubemapSampler1"));
-		ReflectionCubemap2.Bind(Initializer.ParameterMap, TEXT("ReflectionCubemap2"));
-		ReflectionSampler2.Bind(Initializer.ParameterMap, TEXT("ReflectionCubemapSampler2"));
-		MobileReflectionParams.Bind(Initializer.ParameterMap, TEXT("MobileReflectionParams"));
-		ReflectionPositionsAndRadii.Bind(Initializer.ParameterMap, TEXT("ReflectionPositionsAndRadii"));
+		HQReflectionCubemaps[0].Bind(Initializer.ParameterMap, TEXT("ReflectionCubemap0"));
+		HQReflectionSamplers[0].Bind(Initializer.ParameterMap, TEXT("ReflectionCubemapSampler0"));
+		HQReflectionCubemaps[1].Bind(Initializer.ParameterMap, TEXT("ReflectionCubemap1"));
+		HQReflectionSamplers[1].Bind(Initializer.ParameterMap, TEXT("ReflectionCubemapSampler1"));
+		HQReflectionCubemaps[2].Bind(Initializer.ParameterMap, TEXT("ReflectionCubemap2"));
+		HQReflectionSamplers[2].Bind(Initializer.ParameterMap, TEXT("ReflectionCubemapSampler2"));
+		HQReflectionInvAverageBrigtnessParams.Bind(Initializer.ParameterMap, TEXT("ReflectionInvAverageBrigtness"));
+		HQReflectionPositionsAndRadii.Bind(Initializer.ParameterMap, TEXT("ReflectionPositionsAndRadii"));
 
 		LightPositionAndInvRadiusParameter.Bind(Initializer.ParameterMap, TEXT("LightPositionAndInvRadius"));
 		LightColorAndFalloffExponentParameter.Bind(Initializer.ParameterMap, TEXT("LightColorAndFalloffExponent"));
@@ -248,6 +249,7 @@ public:
 						
 		CSMDebugHintParams.Bind(Initializer.ParameterMap, TEXT("CSMDebugHint"));
 	}
+
 	TMobileBasePassPSPolicyParamType() {}
 
 	virtual bool Serialize(FArchive& Ar) override
@@ -256,15 +258,17 @@ public:
 		LightMapPolicyType::PixelParametersType::Serialize(Ar);
 
 		Ar << MobileDirectionLightBufferParam;
-		Ar << ReflectionCubemap;
-		Ar << ReflectionSampler;
-		Ar << ReflectionCubemap1;
-		Ar << ReflectionCubemap2;
-		Ar << ReflectionSampler1;
-		Ar << ReflectionSampler2;
-		Ar << MobileReflectionParams;
-		Ar << ReflectionPositionsAndRadii;
-
+		Ar << ReflectionParameter;
+
+		Ar << HQReflectionCubemaps[0];
+		Ar << HQReflectionCubemaps[1];
+		Ar << HQReflectionCubemaps[2];
+		Ar << HQReflectionSamplers[0];
+		Ar << HQReflectionSamplers[1];
+		Ar << HQReflectionSamplers[2];
+		Ar << HQReflectionInvAverageBrigtnessParams;
+		Ar << HQReflectionPositionsAndRadii;
+		
 		Ar << LightPositionAndInvRadiusParameter;
 		Ar << LightColorAndFalloffExponentParameter;
 		Ar << NumDynamicPointLightsParameter;
@@ -278,15 +282,13 @@ public:
 
 private:
 	FShaderUniformBufferParameter MobileDirectionLightBufferParam;
+	FShaderUniformBufferParameter ReflectionParameter;
 
-	FShaderResourceParameter ReflectionCubemap;
-	FShaderResourceParameter ReflectionSampler;
-	FShaderResourceParameter ReflectionCubemap1;
-	FShaderResourceParameter ReflectionSampler1;
-	FShaderResourceParameter ReflectionCubemap2;
-	FShaderResourceParameter ReflectionSampler2;
-	FShaderParameter MobileReflectionParams;
-	FShaderParameter ReflectionPositionsAndRadii;
+	// HQ reflection bound as loose params
+	FShaderResourceParameter HQReflectionCubemaps[3];
+	FShaderResourceParameter HQReflectionSamplers[3];
+	FShaderParameter HQReflectionInvAverageBrigtnessParams;
+	FShaderParameter HQReflectionPositionsAndRadii;
 	
 	FShaderParameter LightPositionAndInvRadiusParameter;
 	FShaderParameter LightColorAndFalloffExponentParameter;
diff --git a/Engine/Source/Runtime/Renderer/Private/RendererScene.cpp b/Engine/Source/Runtime/Renderer/Private/RendererScene.cpp
--- a/Engine/Source/Runtime/Renderer/Private/RendererScene.cpp
+++ b/Engine/Source/Runtime/Renderer/Private/RendererScene.cpp
@@ -1807,6 +1807,11 @@ void FScene::SetSkyLight(FSkyLightSceneProxy* LightProxy)
 				// The base pass chooses shaders based on whether there's a skylight in the scene, and that is cached in static draw lists
 				Scene->bScenesPrimitivesNeedStaticMeshElementUpdate = true;
 			}
+
+			if (Scene->GetFeatureLevel() <= ERHIFeatureLevel::ES3_1)
+			{
+				Scene->SkyLight->UpdateMobileUniformBuffer();
+			}
 		});
 }
 
@@ -1998,6 +2003,11 @@ void FScene::AddReflectionCapture(UReflectionCaptureComponent* Component)
 
 			Proxy->PackedIndex = PackedIndex;
 			Scene->ReflectionSceneData.RegisteredReflectionCapturePositions.Add(Proxy->Position);
+			
+			if (Scene->GetFeatureLevel() <= ERHIFeatureLevel::ES3_1)
+			{
+				Proxy->UpdateMobileUniformBuffer();
+			}
 
 			checkSlow(Scene->ReflectionSceneData.RegisteredReflectionCaptures.Num() == Scene->ReflectionSceneData.RegisteredReflectionCapturePositions.Num());
 		});
@@ -2069,6 +2079,11 @@ void FScene::UpdateReflectionCaptureTransform(UReflectionCaptureComponent* Compo
 
 			Scene->ReflectionSceneData.bRegisteredReflectionCapturesHasChanged = true;
 			Proxy->SetTransform(Transform);
+
+			if (Scene->GetFeatureLevel() <= ERHIFeatureLevel::ES3_1)
+			{
+				Proxy->UpdateMobileUniformBuffer();
+			}
 		});
 	}
 }
