Add RenderThread CSV exclusive stats

diff --git a/Engine/Source/Runtime/RenderCore/Private/RenderGraphBuilder.cpp b/Engine/Source/Runtime/RenderCore/Private/RenderGraphBuilder.cpp
--- a/Engine/Source/Runtime/RenderCore/Private/RenderGraphBuilder.cpp
+++ b/Engine/Source/Runtime/RenderCore/Private/RenderGraphBuilder.cpp
@@ -5,7 +5,7 @@
 #include "RenderTargetPool.h"
 #include "RenderGraphResourcePool.h"
 #include "VisualizeTexture.h"
-
+#include "ProfilingDebugging/CsvProfiler.h"
 
 #if RENDER_GRAPH_DEBUGGING
 
@@ -122,6 +122,7 @@ FRDGEventName::FRDGEventName(const TCHAR* EventFormat, ...)
 
 void FRDGBuilder::Execute()
 {
+	CSV_SCOPED_TIMING_STAT_EXCLUSIVE(FRDGBuilder_Execute);
 	#if RENDER_GRAPH_DEBUGGING
 	{
 		/** The usage need RDG_EVENT_SCOPE() needs to happen in inner scope of the one containing FRDGBuilder because of
diff --git a/Engine/Source/Runtime/Renderer/Private/BasePassRendering.cpp b/Engine/Source/Runtime/Renderer/Private/BasePassRendering.cpp
--- a/Engine/Source/Runtime/Renderer/Private/BasePassRendering.cpp
+++ b/Engine/Source/Runtime/Renderer/Private/BasePassRendering.cpp
@@ -603,6 +603,7 @@ void CreateOpaqueBasePassUniformBuffer(
  */
 bool FDeferredShadingSceneRenderer::RenderBasePass(FRHICommandListImmediate& RHICmdList, FExclusiveDepthStencil::Type BasePassDepthStencilAccess, IPooledRenderTarget* ForwardScreenSpaceShadowMask, bool bParallelBasePass, bool bRenderLightmapDensity)
 {
+	CSV_SCOPED_TIMING_STAT_EXCLUSIVE(RenderBasePass);
 	SCOPED_NAMED_EVENT(FDeferredShadingSceneRenderer_RenderBasePass, FColor::Emerald);
 
 	bool bDirty = false;
diff --git a/Engine/Source/Runtime/Renderer/Private/DeferredShadingRenderer.cpp b/Engine/Source/Runtime/Renderer/Private/DeferredShadingRenderer.cpp
--- a/Engine/Source/Runtime/Renderer/Private/DeferredShadingRenderer.cpp
+++ b/Engine/Source/Runtime/Renderer/Private/DeferredShadingRenderer.cpp
@@ -491,6 +491,7 @@ void FDeferredShadingSceneRenderer::PrepareDistanceFieldScene(FRHICommandListImm
 {
 	if (ShouldPrepareDistanceFieldScene())
 	{
+		CSV_SCOPED_TIMING_STAT_EXCLUSIVE(RenderDFAO);
 		SCOPE_CYCLE_COUNTER(STAT_FDeferredShadingSceneRenderer_DistanceFieldAO_Init);
 		GDistanceFieldVolumeTextureAtlas.UpdateAllocations();
 		UpdateGlobalDistanceFieldObjectBuffers(RHICmdList);
@@ -799,6 +800,8 @@ void FDeferredShadingSceneRenderer::Render(FRHICommandListImmediate& RHICmdList)
 {
 	check(RHICmdList.IsOutsideRenderPass());
 
+	CSV_SCOPED_TIMING_STAT_EXCLUSIVE(RenderOther);
+
 	PrepareViewRectsForRendering();
 
 	SCOPED_NAMED_EVENT(FDeferredShadingSceneRenderer_Render, FColor::Emerald);
@@ -1486,6 +1489,7 @@ void FDeferredShadingSceneRenderer::Render(FRHICommandListImmediate& RHICmdList)
 
 	if(GetCustomDepthPassLocation() == 1)
 	{
+		CSV_SCOPED_TIMING_STAT_EXCLUSIVE(CustomDepthPass);
 		QUICK_SCOPE_CYCLE_COUNTER(STAT_FDeferredShadingSceneRenderer_CustomDepthPass1);
 		RenderCustomDepthPassAtLocation(RHICmdList, 1);
 	}
@@ -1634,6 +1638,7 @@ void FDeferredShadingSceneRenderer::Render(FRHICommandListImmediate& RHICmdList)
 	// e.g. deferred decals, SSAO
 	if (FeatureLevel >= ERHIFeatureLevel::SM4)
 	{
+		CSV_SCOPED_TIMING_STAT_EXCLUSIVE(AfterBasePass);
 		SCOPE_CYCLE_COUNTER(STAT_FDeferredShadingSceneRenderer_AfterBasePass);
 
 		GRenderTargetPool.AddPhaseEvent(TEXT("AfterBasePass"));
@@ -1673,6 +1678,7 @@ void FDeferredShadingSceneRenderer::Render(FRHICommandListImmediate& RHICmdList)
 	// Render lighting.
 	if (bRenderDeferredLighting)
 	{
+		CSV_SCOPED_TIMING_STAT_EXCLUSIVE(RenderLighting);
 		SCOPE_CYCLE_COUNTER(STAT_FDeferredShadingSceneRenderer_Lighting);
 
 		GRenderTargetPool.AddPhaseEvent(TEXT("Lighting"));
@@ -1819,6 +1825,7 @@ void FDeferredShadingSceneRenderer::Render(FRHICommandListImmediate& RHICmdList)
 	// Draw fog.
 	if (bCanOverlayRayTracingOutput && ShouldRenderFog(ViewFamily))
 	{
+		CSV_SCOPED_TIMING_STAT_EXCLUSIVE(RenderFog);
 		SCOPE_CYCLE_COUNTER(STAT_FDeferredShadingSceneRenderer_RenderFog);
 		RenderFog(RHICmdList, LightShaftOutput);
 		ServiceLocalQueue();
@@ -1847,6 +1854,7 @@ void FDeferredShadingSceneRenderer::Render(FRHICommandListImmediate& RHICmdList)
 	// Notify the FX system that opaque primitives have been rendered and we now have a valid depth buffer.
 	if (Scene->FXSystem && Views.IsValidIndex(0) && bAllowGPUParticleSceneUpdate)
 	{
+		CSV_SCOPED_TIMING_STAT_EXCLUSIVE(RenderOpaqueFX);
 		SCOPE_CYCLE_COUNTER(STAT_FDeferredShadingSceneRenderer_FXSystem_PostRenderOpaque);
 
 		FSceneTexturesUniformParameters SceneTextureParameters;
@@ -1872,6 +1880,7 @@ void FDeferredShadingSceneRenderer::Render(FRHICommandListImmediate& RHICmdList)
 	// Draw translucency.
 	if (bCanOverlayRayTracingOutput && ViewFamily.EngineShowFlags.Translucency)
 	{
+		CSV_SCOPED_TIMING_STAT_EXCLUSIVE(RenderTranslucency);
 		SCOPE_CYCLE_COUNTER(STAT_TranslucencyDrawTime);
 
 		RHICmdList.SetCurrentStat(GET_STATID(STAT_CLM_Translucency));
diff --git a/Engine/Source/Runtime/Renderer/Private/DepthRendering.cpp b/Engine/Source/Runtime/Renderer/Private/DepthRendering.cpp
--- a/Engine/Source/Runtime/Renderer/Private/DepthRendering.cpp
+++ b/Engine/Source/Runtime/Renderer/Private/DepthRendering.cpp
@@ -526,6 +526,7 @@ bool FDeferredShadingSceneRenderer::RenderPrePass(FRHICommandListImmediate& RHIC
 	SCOPED_DRAW_EVENTF(RHICmdList, PrePass, TEXT("PrePass %s %s"), GetDepthDrawingModeString(EarlyZPassMode), GetDepthPassReason(bDitheredLODTransitionsUseStencil, ShaderPlatform));
 
 	SCOPE_CYCLE_COUNTER(STAT_DepthDrawTime);
+	CSV_SCOPED_TIMING_STAT_EXCLUSIVE(RenderPrePass);
 	SCOPED_GPU_STAT(RHICmdList, Prepass);
 
 	bool bDidPrePre = false;
diff --git a/Engine/Source/Runtime/Renderer/Private/LightGridInjection.cpp b/Engine/Source/Runtime/Renderer/Private/LightGridInjection.cpp
--- a/Engine/Source/Runtime/Renderer/Private/LightGridInjection.cpp
+++ b/Engine/Source/Runtime/Renderer/Private/LightGridInjection.cpp
@@ -424,6 +424,7 @@ void FDeferredShadingSceneRenderer::ComputeLightGrid(FRHICommandListImmediate& R
 	}
 
 	{
+		CSV_SCOPED_TIMING_STAT_EXCLUSIVE(ComputeLightGrid);
 		QUICK_SCOPE_CYCLE_COUNTER(STAT_ComputeLightGrid);
 		SCOPED_DRAW_EVENT(RHICmdList, ComputeLightGrid);
 
diff --git a/Engine/Source/Runtime/Renderer/Private/MobileBasePassRendering.cpp b/Engine/Source/Runtime/Renderer/Private/MobileBasePassRendering.cpp
--- a/Engine/Source/Runtime/Renderer/Private/MobileBasePassRendering.cpp
+++ b/Engine/Source/Runtime/Renderer/Private/MobileBasePassRendering.cpp
@@ -223,6 +223,7 @@ void SetupMobileDirectionalLightUniformParameters(
 
 void FMobileSceneRenderer::RenderMobileBasePass(FRHICommandListImmediate& RHICmdList, const TArrayView<const FViewInfo*> PassViews)
 {
+	CSV_SCOPED_TIMING_STAT_EXCLUSIVE(RenderBasePass);
 	SCOPED_DRAW_EVENT(RHICmdList, MobileBasePass);
 	SCOPE_CYCLE_COUNTER(STAT_BasePassDrawTime);
 
diff --git a/Engine/Source/Runtime/Renderer/Private/MobileShadingRenderer.cpp b/Engine/Source/Runtime/Renderer/Private/MobileShadingRenderer.cpp
--- a/Engine/Source/Runtime/Renderer/Private/MobileShadingRenderer.cpp
+++ b/Engine/Source/Runtime/Renderer/Private/MobileShadingRenderer.cpp
@@ -187,6 +187,7 @@ void FMobileSceneRenderer::InitViews(FRHICommandListImmediate& RHICmdList)
 	SCOPED_DRAW_EVENT(RHICmdList, InitViews);
 
 	SCOPE_CYCLE_COUNTER(STAT_InitViewsTime);
+	CSV_SCOPED_TIMING_STAT_EXCLUSIVE(InitViews_Scene);
 
 	FILCUpdatePrimTaskData ILCTaskData;
 	FViewVisibleCommandsPerView ViewCommandsPerView;
@@ -271,6 +272,7 @@ void FMobileSceneRenderer::Render(FRHICommandListImmediate& RHICmdList)
 
 	PrepareViewRectsForRendering();
 
+	CSV_SCOPED_TIMING_STAT_EXCLUSIVE(RenderOther);
 	QUICK_SCOPE_CYCLE_COUNTER(STAT_FMobileSceneRenderer_Render);
 	//FRHICommandListExecutor::GetImmediateCommandList().ImmediateFlush(EImmediateFlushType::DispatchToRHIThread);
 
@@ -435,6 +437,7 @@ void FMobileSceneRenderer::Render(FRHICommandListImmediate& RHICmdList)
 	RHICmdList.ImmediateFlush(EImmediateFlushType::DispatchToRHIThread);
 
 	{
+		CSV_SCOPED_TIMING_STAT_EXCLUSIVE(ViewExtensionPostRenderBasePass);
 		QUICK_SCOPE_CYCLE_COUNTER(STAT_FMobileSceneRenderer_ViewExtensionPostRenderBasePass);
 		for (int32 ViewExt = 0; ViewExt < ViewFamily.ViewExtensions.Num(); ++ViewExt)
 		{
@@ -490,11 +493,13 @@ void FMobileSceneRenderer::Render(FRHICommandListImmediate& RHICmdList)
 	{
 		if (ViewFamily.EngineShowFlags.Decals)
 		{
+			CSV_SCOPED_TIMING_STAT_EXCLUSIVE(RenderDecals);
 			RenderDecals(RHICmdList);
 		}
 
 		if (ViewFamily.EngineShowFlags.DynamicShadows)
 		{
+			CSV_SCOPED_TIMING_STAT_EXCLUSIVE(RenderShadowProjections);
 			RenderModulatedShadowProjections(RHICmdList);
 		}
 	}
@@ -502,6 +507,7 @@ void FMobileSceneRenderer::Render(FRHICommandListImmediate& RHICmdList)
 	// Draw translucency.
 	if (ViewFamily.EngineShowFlags.Translucency)
 	{
+		CSV_SCOPED_TIMING_STAT_EXCLUSIVE(RenderTranslucency);
 		SCOPE_CYCLE_COUNTER(STAT_TranslucencyDrawTime);
 		RenderTranslucency(RHICmdList, ViewList, !bGammaSpace || bRenderToSceneColor);
 		FRHICommandListExecutor::GetImmediateCommandList().PollOcclusionQueries();
diff --git a/Engine/Source/Runtime/Renderer/Private/PostProcess/PostProcessing.cpp b/Engine/Source/Runtime/Renderer/Private/PostProcess/PostProcessing.cpp
--- a/Engine/Source/Runtime/Renderer/Private/PostProcess/PostProcessing.cpp
+++ b/Engine/Source/Runtime/Renderer/Private/PostProcess/PostProcessing.cpp
@@ -1383,7 +1383,8 @@ void FPostProcessing::RegisterHMDPostprocessPass(FPostprocessContext& Context, c
 
 void FPostProcessing::Process(FRHICommandListImmediate& RHICmdList, const FViewInfo& View, TRefCountPtr<IPooledRenderTarget>& VelocityRT)
 {
-	QUICK_SCOPE_CYCLE_COUNTER( STAT_PostProcessing_Process );
+	CSV_SCOPED_TIMING_STAT_EXCLUSIVE(RenderPostProcessing);
+	QUICK_SCOPE_CYCLE_COUNTER(STAT_PostProcessing_Process);
 
 	check(IsInRenderingThread());
 	check(View.VerifyMembersChecks());
diff --git a/Engine/Source/Runtime/Renderer/Private/SceneRendering.cpp b/Engine/Source/Runtime/Renderer/Private/SceneRendering.cpp
--- a/Engine/Source/Runtime/Renderer/Private/SceneRendering.cpp
+++ b/Engine/Source/Runtime/Renderer/Private/SceneRendering.cpp
@@ -2767,6 +2767,8 @@ void FSceneRenderer::RenderCustomDepthPassAtLocation(FRHICommandListImmediate& R
 
 void FSceneRenderer::RenderCustomDepthPass(FRHICommandListImmediate& RHICmdList)
 {
+	CSV_SCOPED_TIMING_STAT_EXCLUSIVE(RenderCustomDepthPass);
+
 	// do we have primitives in this pass?
 	bool bPrimitives = false;
 
diff --git a/Engine/Source/Runtime/Renderer/Private/SceneVisibility.cpp b/Engine/Source/Runtime/Renderer/Private/SceneVisibility.cpp
--- a/Engine/Source/Runtime/Renderer/Private/SceneVisibility.cpp
+++ b/Engine/Source/Runtime/Renderer/Private/SceneVisibility.cpp
@@ -1178,6 +1178,7 @@ public:
 
 static int32 FetchVisibilityForPrimitives(const FScene* Scene, FViewInfo& View, const bool bSubmitQueries, const bool bHZBOcclusion, FGlobalDynamicVertexBuffer& DynamicVertexBuffer)
 {
+	CSV_SCOPED_TIMING_STAT_EXCLUSIVE(FetchVisibilityForPrimitives);
 	QUICK_SCOPE_CYCLE_COUNTER(STAT_FetchVisibilityForPrimitives);
 	FSceneViewState* ViewState = (FSceneViewState*)View.State;
 	
@@ -4014,6 +4015,7 @@ bool FDeferredShadingSceneRenderer::InitViews(FRHICommandListImmediate& RHICmdLi
 {
 	SCOPED_NAMED_EVENT(FDeferredShadingSceneRenderer_InitViews, FColor::Emerald);
 	SCOPE_CYCLE_COUNTER(STAT_InitViewsTime);
+	CSV_SCOPED_TIMING_STAT_EXCLUSIVE(InitViews_Scene);
 	check(RHICmdList.IsOutsideRenderPass());
 
 	PreVisibilityFrameSetup(RHICmdList);
diff --git a/Engine/Source/Runtime/Renderer/Private/ShadowDepthRendering.cpp b/Engine/Source/Runtime/Renderer/Private/ShadowDepthRendering.cpp
--- a/Engine/Source/Runtime/Renderer/Private/ShadowDepthRendering.cpp
+++ b/Engine/Source/Runtime/Renderer/Private/ShadowDepthRendering.cpp
@@ -1444,6 +1444,7 @@ void FSceneRenderer::RenderShadowDepthMaps(FRHICommandListImmediate& RHICmdList)
 {
 	check(RHICmdList.IsOutsideRenderPass());
 
+	CSV_SCOPED_TIMING_STAT_EXCLUSIVE(RenderShadows);
 	SCOPED_NAMED_EVENT(FSceneRenderer_RenderShadowDepthMaps, FColor::Emerald);
 	FSceneRenderTargets& SceneContext = FSceneRenderTargets::Get(RHICmdList);
 
diff --git a/Engine/Source/Runtime/Renderer/Private/ShadowSetup.cpp b/Engine/Source/Runtime/Renderer/Private/ShadowSetup.cpp
--- a/Engine/Source/Runtime/Renderer/Private/ShadowSetup.cpp
+++ b/Engine/Source/Runtime/Renderer/Private/ShadowSetup.cpp
@@ -4119,6 +4119,7 @@ void FSceneRenderer::InitDynamicShadows(FRHICommandListImmediate& RHICmdList, FG
 	TArray<FProjectedShadowInfo*,SceneRenderingAllocator> ViewDependentWholeSceneShadows;
 	TArray<FProjectedShadowInfo*,SceneRenderingAllocator> ViewDependentWholeSceneShadowsThatNeedCulling;
 	{
+		CSV_SCOPED_TIMING_STAT_EXCLUSIVE(InitViews_Shadows);
 		SCOPE_CYCLE_COUNTER(STAT_InitDynamicShadowsTime);
 
 		for (TSparseArray<FLightSceneInfoCompact>::TConstIterator LightIt(Scene->Lights); LightIt; ++LightIt)
diff --git a/Engine/Source/Runtime/Renderer/Private/VelocityRendering.cpp b/Engine/Source/Runtime/Renderer/Private/VelocityRendering.cpp
--- a/Engine/Source/Runtime/Renderer/Private/VelocityRendering.cpp
+++ b/Engine/Source/Runtime/Renderer/Private/VelocityRendering.cpp
@@ -435,6 +435,7 @@ void FDeferredShadingSceneRenderer::RenderVelocities(FRHICommandListImmediate& R
 	SCOPED_NAMED_EVENT(FDeferredShadingSceneRenderer_RenderVelocities, FColor::Emerald);
 
 	check(FeatureLevel >= ERHIFeatureLevel::SM4);
+	CSV_SCOPED_TIMING_STAT_EXCLUSIVE(RenderVelocities);
 	SCOPE_CYCLE_COUNTER(STAT_RenderVelocities);
 
 	if (!ShouldRenderVelocities())
diff --git a/Engine/Source/Runtime/Renderer/Private/VolumetricFog.cpp b/Engine/Source/Runtime/Renderer/Private/VolumetricFog.cpp
--- a/Engine/Source/Runtime/Renderer/Private/VolumetricFog.cpp
+++ b/Engine/Source/Runtime/Renderer/Private/VolumetricFog.cpp
@@ -973,6 +973,7 @@ void FDeferredShadingSceneRenderer::ComputeVolumetricFog(FRHICommandListImmediat
 
 	if (ShouldRenderVolumetricFog())
 	{		
+		CSV_SCOPED_TIMING_STAT_EXCLUSIVE(VolumetricFog);
 		QUICK_SCOPE_CYCLE_COUNTER(STAT_VolumetricFog);
 		SCOPED_GPU_STAT(RHICmdListImmediate, VolumetricFog);
 
