Fix out-of-bounds read when detecting a main view for ISR

diff --git a/Engine/Source/Runtime/Renderer/Private/MeshDrawCommands.cpp b/Engine/Source/Runtime/Renderer/Private/MeshDrawCommands.cpp
--- a/Engine/Source/Runtime/Renderer/Private/MeshDrawCommands.cpp
+++ b/Engine/Source/Runtime/Renderer/Private/MeshDrawCommands.cpp
@@ -949,7 +949,7 @@ void FParallelMeshDrawCommandPass::DispatchPassSetup(
 	TaskContext.NumDynamicMeshCommandBuildRequestElements = NumDynamicMeshCommandBuildRequestElements;
 
 	// Only apply instancing for ISR to main view passes
-	const bool bIsMainViewPass = (FPassProcessorManager::GetPassFlags(TaskContext.ShadingPath, TaskContext.PassType) & EMeshPassFlags::MainView) != EMeshPassFlags::None;
+	const bool bIsMainViewPass = PassType != EMeshPass::Num && (FPassProcessorManager::GetPassFlags(TaskContext.ShadingPath, TaskContext.PassType) & EMeshPassFlags::MainView) != EMeshPassFlags::None;
 	TaskContext.InstanceFactor = (bIsMainViewPass && View.IsInstancedStereoPass()) ? 2 : 1;
 
 	// Setup translucency sort key update pass based on view.
