Add bounds checks in FPassProcessorManager

diff --git a/Engine/Source/Runtime/Renderer/Public/MeshPassProcessor.h b/Engine/Source/Runtime/Renderer/Public/MeshPassProcessor.h
--- a/Engine/Source/Runtime/Renderer/Public/MeshPassProcessor.h
+++ b/Engine/Source/Runtime/Renderer/Public/MeshPassProcessor.h
@@ -1067,6 +1067,7 @@ class FPassProcessorManager
 public:
 	static PassProcessorCreateFunction GetCreateFunction(EShadingPath ShadingPath, EMeshPass::Type PassType)
 	{
+		check(ShadingPath < EShadingPath::Num && PassType < EMeshPass::Num);
 		uint32 ShadingPathIdx = (uint32)ShadingPath;
 		checkf(JumpTable[ShadingPathIdx][PassType], TEXT("Pass type %u create function was never registered for shading path %u.  Use a FRegisterPassProcessorCreateFunction to register a create function for this enum value."), (uint32)PassType, ShadingPathIdx);
 		return JumpTable[ShadingPathIdx][PassType];
@@ -1074,6 +1075,7 @@ public:
 
 	static EMeshPassFlags GetPassFlags(EShadingPath ShadingPath, EMeshPass::Type PassType)
 	{
+		check(ShadingPath < EShadingPath::Num && PassType < EMeshPass::Num);
 		uint32 ShadingPathIdx = (uint32)ShadingPath;
 		return Flags[ShadingPathIdx][PassType];
 	}
