Cherry-pick 5838739: Changed default value of r.RayTracing.Shadows.EnableTwoSidedGeometry from 0 to 1, to make ray traced shadows from single-sided geometry behave more similar to shadowmaps.

diff --git a/Engine/Source/Runtime/Renderer/Private/RayTracing/RayTracingShadows.cpp b/Engine/Source/Runtime/Renderer/Private/RayTracing/RayTracingShadows.cpp
--- a/Engine/Source/Runtime/Renderer/Private/RayTracing/RayTracingShadows.cpp
+++ b/Engine/Source/Runtime/Renderer/Private/RayTracing/RayTracingShadows.cpp
@@ -36,8 +36,8 @@ static FAutoConsoleVariableRef CVarRayTracingShadowsEnableMaterials(
 
 static TAutoConsoleVariable<int32> CVarRayTracingShadowsEnableTwoSidedGeometry(
 	TEXT("r.RayTracing.Shadows.EnableTwoSidedGeometry"),
-	0,
-	TEXT("Enables two-sided geometry when tracing shadow rays (default = 0)"),
+	1,
+	TEXT("Enables two-sided geometry when tracing shadow rays (default = 1)"),
 	ECVF_RenderThreadSafe
 );
 
