Cherry-pick 5899556: Fixed PrepareRayTracingShadows is missing some denoiser requirements

diff --git a/Engine/Source/Runtime/Renderer/Private/RayTracing/RayTracingShadows.cpp b/Engine/Source/Runtime/Renderer/Private/RayTracing/RayTracingShadows.cpp
--- a/Engine/Source/Runtime/Renderer/Private/RayTracing/RayTracingShadows.cpp
+++ b/Engine/Source/Runtime/Renderer/Private/RayTracing/RayTracingShadows.cpp
@@ -88,6 +88,7 @@ void FDeferredShadingSceneRenderer::PrepareRayTracingShadows(const FViewInfo& Vi
 	{
 		IScreenSpaceDenoiser::EShadowRequirements::Bailout,
 		IScreenSpaceDenoiser::EShadowRequirements::ClosestOccluder,
+		IScreenSpaceDenoiser::EShadowRequirements::PenumbraAndAvgOccluder,
 		IScreenSpaceDenoiser::EShadowRequirements::PenumbraAndClosestOccluder,
 	};
 
