Fixes assertion failure when screen percentage < 50% with TAAU.

diff --git a/Engine/Source/Runtime/Renderer/Private/PostProcess/PostProcessTemporalAA.cpp b/Engine/Source/Runtime/Renderer/Private/PostProcess/PostProcessTemporalAA.cpp
--- a/Engine/Source/Runtime/Renderer/Private/PostProcess/PostProcessTemporalAA.cpp
+++ b/Engine/Source/Runtime/Renderer/Private/PostProcess/PostProcessTemporalAA.cpp
@@ -816,9 +816,9 @@ void FRCPassPostProcessTemporalAA::Process(FRenderingCompositePassContext& Conte
 			}
 			// If screen percentage < 50% on X and Y axes, then use screen percentage range = 3 shader permutation.
 			else if (SrcRect.Width() * 100 < 50 * DestRect.Width() &&
-				SrcRect.Height() * 100 < 50 * DestRect.Height())
+				SrcRect.Height() * 100 < 50 * DestRect.Height() && 
+				Parameters.Pass == ETAAPassConfig::MainSuperSampling)
 			{
-				check(Parameters.Pass == ETAAPassConfig::MainSuperSampling);
 				PermutationVector.Set<FTAAScreenPercentageDim>(3);
 			}
 			// If screen percentage < 71% on X and Y axes, then use screen percentage range = 1 shader permutation to have smaller LDS caching.
