Cherry-pick 5916652: Fixed StochasticRectLight does not respect Light "samples per pixel" control

diff --git a/Engine/Source/Runtime/Renderer/Private/RayTracing/RayTracingRectLight.cpp b/Engine/Source/Runtime/Renderer/Private/RayTracing/RayTracingRectLight.cpp
--- a/Engine/Source/Runtime/Renderer/Private/RayTracing/RayTracingRectLight.cpp
+++ b/Engine/Source/Runtime/Renderer/Private/RayTracing/RayTracingRectLight.cpp
@@ -27,7 +27,7 @@ static FAutoConsoleVariableRef CVarRayTracingStochasticRectLight(
 	TEXT("1: use stochastic evaluation\n")
 );
 
-static int32 GRayTracingStochasticRectLightSamplesPerPixel = 1;
+static int32 GRayTracingStochasticRectLightSamplesPerPixel = -1;
 static FAutoConsoleVariableRef CVarRayTracingRecLightStochasticSamplesPerPixel(
 	TEXT("r.RayTracing.StochasticRectLight.SamplesPerPixel"),
 	GRayTracingStochasticRectLightSamplesPerPixel,
@@ -488,7 +488,7 @@ void RenderRayTracingRectLightInternal(
 	RectLightSceneProxy->GetLightShaderParameters(LightShaderParameters);
 
 	FRectLightData RectLightData;
-	RectLightData.SamplesPerPixel = GRayTracingStochasticRectLightSamplesPerPixel;
+	RectLightData.SamplesPerPixel = GRayTracingStochasticRectLightSamplesPerPixel >= 0 ? GRayTracingStochasticRectLightSamplesPerPixel : RectLightSceneProxy->GetSamplesPerPixel();
 	RectLightData.bIsTextureImportanceSampling = GRayTracingStochasticRectLightIsTextureImportanceSampling;
 	RectLightData.Position = RectLightSceneInfo.Proxy->GetOrigin();
 	RectLightData.Normal = RectLightSceneInfo.Proxy->GetDirection();
