Disabling a spammy profile counter in MeshPassProcessor.

diff --git a/Engine/Source/Runtime/Renderer/Private/MeshPassProcessor.cpp b/Engine/Source/Runtime/Renderer/Private/MeshPassProcessor.cpp
--- a/Engine/Source/Runtime/Renderer/Private/MeshPassProcessor.cpp
+++ b/Engine/Source/Runtime/Renderer/Private/MeshPassProcessor.cpp
@@ -1148,7 +1148,8 @@ void FCachedPassMeshDrawListContext::FinalizeCommand(
 	const FMeshProcessorShaders* ShadersForDebugging,
 	FMeshDrawCommand& MeshDrawCommand)
 {
-	QUICK_SCOPE_CYCLE_COUNTER(STAT_FinalizeCachedMeshDrawCommand);
+	// disabling this by default as it incurs a high cost in perf captures due to sheer volume.  Recommendation is to re-enable locally if you need to profile this particular code.
+	// QUICK_SCOPE_CYCLE_COUNTER(STAT_FinalizeCachedMeshDrawCommand);
 
 	FGraphicsMinimalPipelineStateId PipelineId;
 	PipelineId = FGraphicsMinimalPipelineStateId::GetPersistentId(PipelineState);
