Cherry-pick 5928643: Fixed IES profiles in ray tracing reflections and translucency passes

diff --git a/Engine/Shaders/Private/IESLightProfilesCommon.ush b/Engine/Shaders/Private/IESLightProfilesCommon.ush
--- a/Engine/Shaders/Private/IESLightProfilesCommon.ush
+++ b/Engine/Shaders/Private/IESLightProfilesCommon.ush
@@ -36,10 +36,8 @@ float ComputeLightProfileMultiplier(float3 WorldPosition, float3 LightPosition,
 #endif
 }
 
-#if RAYHITGROUPSHADER
-
 // Apply IES light profile texture. Each row of the texture stores one profile.
-float ComputeLightProfileMultiplier(
+float ComputeRayTracingLightProfileMultiplier(
 	float3 WorldPosition,
 	float3 LightPosition,
 	float3 LightDirection,
@@ -56,11 +54,10 @@ float ComputeLightProfileMultiplier(
 		float Angle = asin(DotProd);
 		// 0..1
 		const float NormAngle = Angle / PI + 0.5f;
-
-		const float2 UV = float2(NormAngle, (float(LightProfileTextureIndex) + 0.5) * IESLightProfileData.IESLightProfileInvCount);
-
-		Result = Texture2DSampleLevel(IESLightProfileData.IESLightProfileTexture, IESLightProfileData.IESLightProfileTextureSampler, UV, 0).r;
+		
+		const float2 UV = float2(NormAngle, (float(LightProfileTextureIndex) + 0.5) * RaytracingLightsDataPacked.IESLightProfileInvCount);
+		Result = Texture2DSampleLevel(RaytracingLightsDataPacked.IESLightProfileTexture, RaytracingLightsDataPacked.IESLightProfileTextureSampler, UV, 0).r;
 	}
+
 	return Result;
 }
-#endif // RAYHITGROUPSHADER
\ No newline at end of file
diff --git a/Engine/Shaders/Private/RayTracing/RayTracingLightingCommon.ush b/Engine/Shaders/Private/RayTracing/RayTracingLightingCommon.ush
--- a/Engine/Shaders/Private/RayTracing/RayTracingLightingCommon.ush
+++ b/Engine/Shaders/Private/RayTracing/RayTracingLightingCommon.ush
@@ -420,7 +420,7 @@ float3 ComputeDirectLightingCulled(
 
 		if (LightProfileIndex >= 0)
 		{
-			LightProfileMultiplier = ComputeLightProfileMultiplier(WorldPosition, LightData.Position, LightData.Direction, LightProfileIndex);
+			LightProfileMultiplier = ComputeRayTracingLightProfileMultiplier(WorldPosition, LightData.Position, LightData.Direction, LightProfileIndex);
 		}
 
 		float3 ShadowRayDirection;
diff --git a/Engine/Source/Runtime/Renderer/Private/RayTracing/RayTracingIESLightProfiles.cpp b/Engine/Source/Runtime/Renderer/Private/RayTracing/RayTracingIESLightProfiles.cpp
--- a/Engine/Source/Runtime/Renderer/Private/RayTracing/RayTracingIESLightProfiles.cpp
+++ b/Engine/Source/Runtime/Renderer/Private/RayTracing/RayTracingIESLightProfiles.cpp
@@ -5,29 +5,6 @@
 
 #if RHI_RAYTRACING
 
-IMPLEMENT_GLOBAL_SHADER_PARAMETER_STRUCT(FIESLightProfileParameters, "IESLightProfileData");
-
-void SetupIESLightProfilesUniformParameters(const FViewInfo& View, FIESLightProfileParameters& OutParameters)
-{
-	FTextureRHIParamRef TextureRHI = nullptr;
-	float IESInvProfileCount = 0.0f;
-
-	if (View.IESLightProfileResource && View.IESLightProfileResource->GetIESLightProfilesCount())
-	{
-		OutParameters.IESLightProfileTexture = View.IESLightProfileResource->GetTexture();
-
-		uint32 ProfileCount = View.IESLightProfileResource->GetIESLightProfilesCount();
-		IESInvProfileCount = ProfileCount ? 1.f / static_cast<float>(ProfileCount) : 0.f;
-	}
-	else
-	{
-		OutParameters.IESLightProfileTexture = GWhiteTexture->TextureRHI;
-	}
-	   
-	OutParameters.IESLightProfileInvCount = IESInvProfileCount;
-	OutParameters.IESLightProfileTextureSampler = TStaticSamplerState<SF_Bilinear, AM_Clamp, AM_Clamp, AM_Clamp>::GetRHI();
-}
-
 void FIESLightProfileResource::BuildIESLightProfilesTexture(const TArray<UTextureLightProfile*, SceneRenderingAllocator>& NewIESProfilesArray)
 {
 	// Rebuild 2D texture that contains one IES light profile per row
@@ -98,11 +75,4 @@ void FIESLightProfileResource::BuildIESLightProfilesTexture(const TArray<UTextur
 	RHIUnlockTexture2D(TextureRHI, 0, false);
 }
 
-TUniformBufferRef<FIESLightProfileParameters> CreateIESLightProfilesUniformBuffer(const FViewInfo& View, EUniformBufferUsage Usage)
-{
-	FIESLightProfileParameters IESLightProfileStruct;
-	SetupIESLightProfilesUniformParameters(View, IESLightProfileStruct);
-	return CreateUniformBufferImmediate(IESLightProfileStruct, Usage);
-}
-
 #endif
diff --git a/Engine/Source/Runtime/Renderer/Private/RayTracing/RayTracingIESLightProfiles.h b/Engine/Source/Runtime/Renderer/Private/RayTracing/RayTracingIESLightProfiles.h
--- a/Engine/Source/Runtime/Renderer/Private/RayTracing/RayTracingIESLightProfiles.h
+++ b/Engine/Source/Runtime/Renderer/Private/RayTracing/RayTracingIESLightProfiles.h
@@ -9,15 +9,6 @@
 
 #if RHI_RAYTRACING
 
-BEGIN_GLOBAL_SHADER_PARAMETER_STRUCT(FIESLightProfileParameters, )
-SHADER_PARAMETER(float, IESLightProfileInvCount)
-SHADER_PARAMETER_TEXTURE(Texture2D, IESLightProfileTexture)
-SHADER_PARAMETER_SAMPLER(SamplerState, IESLightProfileTextureSampler)
-END_GLOBAL_SHADER_PARAMETER_STRUCT()
-
-void SetupIESLightProfilesUniformParameters(const FViewInfo& View, FIESLightProfileParameters& OutParameters);
-TUniformBufferRef<FIESLightProfileParameters> CreateIESLightProfilesUniformBuffer(const class FViewInfo& View, EUniformBufferUsage Usage);
-
 class FIESLightProfileResource
 {
 public:
diff --git a/Engine/Source/Runtime/Renderer/Private/RayTracing/RayTracingLighting.cpp b/Engine/Source/Runtime/Renderer/Private/RayTracing/RayTracingLighting.cpp
--- a/Engine/Source/Runtime/Renderer/Private/RayTracing/RayTracingLighting.cpp
+++ b/Engine/Source/Runtime/Renderer/Private/RayTracing/RayTracingLighting.cpp
@@ -34,6 +34,27 @@ void SetupRaytracingLightDataPacked(
 	static constexpr uint32 MaxRectLightTextureSlos = 8;
 	static constexpr uint32 InvalidTextureIndex = 99; // #dxr_todo: share this definition with ray tracing shaders
 
+	{
+		// IES profiles
+		FTextureRHIParamRef IESTextureRHI = nullptr;
+		float IESInvProfileCount = 1.0f;
+
+		if (View.IESLightProfileResource && View.IESLightProfileResource->GetIESLightProfilesCount())
+		{
+			LightData->IESLightProfileTexture = View.IESLightProfileResource->GetTexture();
+
+			uint32 ProfileCount = View.IESLightProfileResource->GetIESLightProfilesCount();
+			IESInvProfileCount = ProfileCount ? 1.f / static_cast<float>(ProfileCount) : 0.f;
+		}
+		else
+		{
+			LightData->IESLightProfileTexture = GWhiteTexture->TextureRHI;
+		}
+
+		LightData->IESLightProfileInvCount = IESInvProfileCount;
+		LightData->IESLightProfileTextureSampler = TStaticSamplerState<SF_Bilinear, AM_Clamp, AM_Clamp, AM_Clamp>::GetRHI();
+	}
+
 	for (auto Light : Lights)
 	{
 		const bool bHasStaticLighting = Light.LightSceneInfo->Proxy->HasStaticLighting() && Light.LightSceneInfo->IsPrecomputedLightingValid();
diff --git a/Engine/Source/Runtime/Renderer/Private/RayTracing/RayTracingLighting.h b/Engine/Source/Runtime/Renderer/Private/RayTracing/RayTracingLighting.h
--- a/Engine/Source/Runtime/Renderer/Private/RayTracing/RayTracingLighting.h
+++ b/Engine/Source/Runtime/Renderer/Private/RayTracing/RayTracingLighting.h
@@ -14,6 +14,7 @@ const static uint32 GRaytracingLightCountMaximum = 64;
 
 BEGIN_GLOBAL_SHADER_PARAMETER_STRUCT(FRaytracingLightDataPacked, )
 	SHADER_PARAMETER(uint32, Count)
+	SHADER_PARAMETER(float, IESLightProfileInvCount)
 	SHADER_PARAMETER_ARRAY(FIntVector, Type_LightProfileIndex_RectLightTextureIndex, [GRaytracingLightCountMaximum])
 	SHADER_PARAMETER_ARRAY(FVector4, LightPosition_InvRadius, [GRaytracingLightCountMaximum])
 	SHADER_PARAMETER_ARRAY(FVector4, LightColor_SpecularScale, [GRaytracingLightCountMaximum])
@@ -33,6 +34,8 @@ BEGIN_GLOBAL_SHADER_PARAMETER_STRUCT(FRaytracingLightDataPacked, )
 	SHADER_PARAMETER_TEXTURE(Texture2D, RectLightTexture5)
 	SHADER_PARAMETER_TEXTURE(Texture2D, RectLightTexture6)
 	SHADER_PARAMETER_TEXTURE(Texture2D, RectLightTexture7)
+	SHADER_PARAMETER_SAMPLER(SamplerState, IESLightProfileTextureSampler)
+	SHADER_PARAMETER_TEXTURE(Texture2D, IESLightProfileTexture)
 END_GLOBAL_SHADER_PARAMETER_STRUCT()
 
 void SetupRaytracingLightDataPacked(
diff --git a/Engine/Source/Runtime/Renderer/Private/RayTracing/RayTracingReflections.cpp b/Engine/Source/Runtime/Renderer/Private/RayTracing/RayTracingReflections.cpp
--- a/Engine/Source/Runtime/Renderer/Private/RayTracing/RayTracingReflections.cpp
+++ b/Engine/Source/Runtime/Renderer/Private/RayTracing/RayTracingReflections.cpp
@@ -17,7 +17,6 @@
 #include "SystemTextures.h"
 #include "RayTracing/RaytracingOptions.h"
 #include "RayTracing/RayTracingDeferredMaterials.h"
-#include "RayTracing/RayTracingIESLightProfiles.h"
 #include "RayTracing/RayTracingLighting.h"
 #include "RenderGraph.h"
 #include "RayTracing/RayTracingLighting.h"
@@ -138,7 +137,6 @@ class FRayTracingReflectionsRGS : public FGlobalShader
 		SHADER_PARAMETER_STRUCT_REF(FRaytracingLightDataPacked, LightDataPacked)
 		SHADER_PARAMETER_STRUCT_REF(FReflectionUniformParameters, ReflectionStruct)
 		SHADER_PARAMETER_STRUCT_REF(FFogUniformParameters, FogUniformParameters)
-		SHADER_PARAMETER_STRUCT_REF(FIESLightProfileParameters, IESLightProfileParameters)
 
 		// Optional indirection buffer used for sorted materials
 		SHADER_PARAMETER_RDG_BUFFER_UAV(StructuredBuffer<FDeferredMaterialPayload>, MaterialBuffer)
@@ -280,7 +278,6 @@ void FDeferredShadingSceneRenderer::RenderRayTracingReflections(
 	CommonParameters.SceneTexturesStruct = CreateSceneTextureUniformBuffer( SceneContext, FeatureLevel, ESceneTextureSetupMode::All, EUniformBufferUsage::UniformBuffer_SingleFrame);
 	CommonParameters.ReflectionStruct = CreateReflectionUniformBuffer(View, EUniformBufferUsage::UniformBuffer_SingleFrame);
 	CommonParameters.FogUniformParameters = CreateFogUniformBuffer(View, EUniformBufferUsage::UniformBuffer_SingleFrame);
-	CommonParameters.IESLightProfileParameters = CreateIESLightProfilesUniformBuffer(View, EUniformBufferUsage::UniformBuffer_SingleFrame);
 	CommonParameters.ColorOutput = GraphBuilder.CreateUAV(*OutColorTexture);
 	CommonParameters.RayHitDistanceOutput = GraphBuilder.CreateUAV(*OutRayHitDistanceTexture);
 	CommonParameters.RayImaginaryDepthOutput = GraphBuilder.CreateUAV(*OutRayImaginaryDepthTexture);
diff --git a/Engine/Source/Runtime/Renderer/Private/RayTracing/RayTracingTranslucency.cpp b/Engine/Source/Runtime/Renderer/Private/RayTracing/RayTracingTranslucency.cpp
--- a/Engine/Source/Runtime/Renderer/Private/RayTracing/RayTracingTranslucency.cpp
+++ b/Engine/Source/Runtime/Renderer/Private/RayTracing/RayTracingTranslucency.cpp
@@ -124,7 +124,6 @@ class FRayTracingTranslucencyRGS : public FGlobalShader
 		SHADER_PARAMETER_STRUCT_REF(FRaytracingLightDataPacked, LightDataPacked)
 		SHADER_PARAMETER_STRUCT_REF(FReflectionUniformParameters, ReflectionStruct)
 		SHADER_PARAMETER_STRUCT_REF(FFogUniformParameters, FogUniformParameters)
-		SHADER_PARAMETER_STRUCT_REF(FIESLightProfileParameters, IESLightProfileParameters)
 
 		SHADER_PARAMETER_RDG_TEXTURE_UAV(RWTexture2D<float4>, ColorOutput)
 		SHADER_PARAMETER_RDG_TEXTURE_UAV(RWTexture2D<float>, RayHitDistanceOutput)
@@ -373,7 +372,6 @@ void FDeferredShadingSceneRenderer::RenderRayTracingTranslucencyView(
 	PassParameters->SceneTexturesStruct = CreateSceneTextureUniformBuffer(SceneContext, FeatureLevel, ESceneTextureSetupMode::All, EUniformBufferUsage::UniformBuffer_SingleFrame);	
 	PassParameters->ReflectionStruct = CreateReflectionUniformBuffer(View, EUniformBufferUsage::UniformBuffer_SingleFrame);
 	PassParameters->FogUniformParameters = CreateFogUniformBuffer(View, EUniformBufferUsage::UniformBuffer_SingleFrame);
-	PassParameters->IESLightProfileParameters = CreateIESLightProfilesUniformBuffer(View, EUniformBufferUsage::UniformBuffer_SingleFrame);
 
 	PassParameters->ColorOutput = GraphBuilder.CreateUAV(*OutColorTexture);
 	PassParameters->RayHitDistanceOutput = GraphBuilder.CreateUAV(*OutRayHitDistanceTexture);
