Removing unused code

#rnx

diff --git a/Engine/Source/Runtime/Renderer/Private/ShadowDepthRendering.cpp b/Engine/Source/Runtime/Renderer/Private/ShadowDepthRendering.cpp
--- a/Engine/Source/Runtime/Renderer/Private/ShadowDepthRendering.cpp
+++ b/Engine/Source/Runtime/Renderer/Private/ShadowDepthRendering.cpp
@@ -1794,7 +1794,6 @@ void FShadowDepthPassMeshProcessor::AddMeshBatch(const FMeshBatch& RESTRICT Mesh
 
 		const FMaterialRenderProxy& MaterialRenderProxy = FallbackMaterialRenderProxyPtr ? *FallbackMaterialRenderProxyPtr : *MeshBatch.MaterialRenderProxy;
 		const EBlendMode BlendMode = Material.GetBlendMode();
-		const EMaterialShadingModel ShadingModel = Material.GetShadingModel();
 		const bool bReflectiveShadowmap = ShadowDepthType.bReflectiveShadowmap && !ShadowDepthType.bOnePassPointLightShadow;
 		const bool bShouldCastShadow = Material.ShouldCastDynamicShadows();
 
