Cherry-pick 6024707: Fixes bug where the velocity buffer would not end up in the scene view family blackboard with separate velocity buffer pass.

diff --git a/Engine/Source/Runtime/Renderer/Private/DeferredShadingRenderer.cpp b/Engine/Source/Runtime/Renderer/Private/DeferredShadingRenderer.cpp
--- a/Engine/Source/Runtime/Renderer/Private/DeferredShadingRenderer.cpp
+++ b/Engine/Source/Runtime/Renderer/Private/DeferredShadingRenderer.cpp
@@ -1496,19 +1496,12 @@ void FDeferredShadingSceneRenderer::Render(FRHICommandListImmediate& RHICmdList)
 
 	ServiceLocalQueue();
 
-	TRefCountPtr<IPooledRenderTarget> VelocityRT;
-
-	if (bBasePassCanOutputVelocity)
-	{
-		VelocityRT = SceneContext.SceneVelocity;
-	}
-	
 	// If bBasePassCanOutputVelocity is set, basepass fully writes the velocity buffer unless bUseSelectiveBasePassOutputs is enabled.
 	if (bShouldRenderVelocities && (!bBasePassCanOutputVelocity || bUseSelectiveBasePassOutputs))
 	{
 		// Render the velocities of movable objects
 		RHICmdList.SetCurrentStat(GET_STATID(STAT_CLM_Velocity));
-		RenderVelocities(RHICmdList, VelocityRT);
+		RenderVelocities(RHICmdList, SceneContext.SceneVelocity);
 		RHICmdList.SetCurrentStat(GET_STATID(STAT_CLM_AfterVelocity));
 		ServiceLocalQueue();
 	}
@@ -1516,7 +1509,7 @@ void FDeferredShadingSceneRenderer::Render(FRHICommandListImmediate& RHICmdList)
 #if !UE_BUILD_SHIPPING
 	if (CVarForceBlackVelocityBuffer.GetValueOnRenderThread())
 	{
-		VelocityRT = SceneContext.SceneVelocity = GSystemTextures.BlackDummy;
+		SceneContext.SceneVelocity = GSystemTextures.BlackDummy;
 	}
 #endif
 	checkSlow(RHICmdList.IsOutsideRenderPass());
@@ -1691,7 +1684,7 @@ void FDeferredShadingSceneRenderer::Render(FRHICommandListImmediate& RHICmdList)
 
 		TRefCountPtr<IPooledRenderTarget> DynamicBentNormalAO;
 		// These modulate the scenecolor output from the basepass, which is assumed to be indirect lighting
-		RenderDFAOAsIndirectShadowing(RHICmdList, VelocityRT, DynamicBentNormalAO);
+		RenderDFAOAsIndirectShadowing(RHICmdList, SceneContext.SceneVelocity, DynamicBentNormalAO);
 
 		// Clear the translucent lighting volumes before we accumulate
 		if ((GbEnableAsyncComputeTranslucencyLightingVolumeClear && GSupportsEfficientAsyncCompute) == false)
@@ -1745,7 +1738,7 @@ void FDeferredShadingSceneRenderer::Render(FRHICommandListImmediate& RHICmdList)
 		checkSlow(RHICmdList.IsOutsideRenderPass());
 
 		// Render diffuse sky lighting and reflections that only operate on opaque pixels
-		RenderDeferredReflectionsAndSkyLighting(RHICmdList, DynamicBentNormalAO, VelocityRT);
+		RenderDeferredReflectionsAndSkyLighting(RHICmdList, DynamicBentNormalAO, SceneContext.SceneVelocity);
 
 		DynamicBentNormalAO = NULL;
 
@@ -1961,7 +1954,7 @@ void FDeferredShadingSceneRenderer::Render(FRHICommandListImmediate& RHICmdList)
 		const float OcclusionMaxDistance = Scene->SkyLight && !Scene->SkyLight->bWantsStaticShadowing ? Scene->SkyLight->OcclusionMaxDistance : Scene->DefaultMaxDistanceFieldOcclusionDistance;
 		TRefCountPtr<IPooledRenderTarget> DummyOutput;
 		RHICmdList.SetCurrentStat(GET_STATID(STAT_CLM_RenderDistanceFieldLighting));
-		RenderDistanceFieldLighting(RHICmdList, FDistanceFieldAOParameters(OcclusionMaxDistance), VelocityRT, DummyOutput, false, ViewFamily.EngineShowFlags.VisualizeDistanceFieldAO); 
+		RenderDistanceFieldLighting(RHICmdList, FDistanceFieldAOParameters(OcclusionMaxDistance), SceneContext.SceneVelocity, DummyOutput, false, ViewFamily.EngineShowFlags.VisualizeDistanceFieldAO); 
 		ServiceLocalQueue();
 	}
 
@@ -2004,11 +1997,11 @@ void FDeferredShadingSceneRenderer::Render(FRHICommandListImmediate& RHICmdList)
 
 			if (ViewFamily.UseDebugViewPS())
 			{
-				DoDebugViewModePostProcessing(RHICmdList, Views[ViewIndex], VelocityRT);
+				DoDebugViewModePostProcessing(RHICmdList, Views[ViewIndex], SceneContext.SceneVelocity);
 			}
 			else
 			{
-			GPostProcessing.Process(RHICmdList, Views[ ViewIndex ], VelocityRT);
+			GPostProcessing.Process(RHICmdList, Views[ ViewIndex ], SceneContext.SceneVelocity);
 		}
 		}
 
