Optimized PostProcessEyeAdaptationBasic.
Use a computer shader to divide the work up across 256 threads and then use LDS to merge back at the end.
Measured difference of 0.385ms to 0.01ms

diff --git a/Engine/Shaders/Private/PostProcessEyeAdaptation.usf b/Engine/Shaders/Private/PostProcessEyeAdaptation.usf
--- a/Engine/Shaders/Private/PostProcessEyeAdaptation.usf
+++ b/Engine/Shaders/Private/PostProcessEyeAdaptation.usf
@@ -48,6 +48,36 @@ float ComputeWeightedTextureAverageAlpha(Texture2D InBuffer, in uint4 ROI, in fl
 	return Average;
 }
 
+float2 ComputeWeightedTextureAverageAlphaSubRegion(Texture2D InBuffer, in uint4 SubROI, in uint4 MainROI, in float Slope)
+{
+	// The inverse of the Region of Interest size.
+	const float InvXMax = 1.f / float(MainROI.z - MainROI.x);
+	const float InvYMax = 1.f / float(MainROI.w - MainROI.y);
+
+	// use product of linear weight in x and y.  
+
+	float Value = 0.0f;
+	float WeightTotal = 0.0000000001f;	
+
+	for (uint i = SubROI.x; i < SubROI.z; ++i)
+	{
+		float IWeight = AdaptationWeight( i * InvXMax, Slope);  
+		for (uint j = SubROI.y; j < SubROI.w; ++j)
+		{
+			float JWeight = AdaptationWeight( j * InvYMax, Slope);
+			float Weight = max(IWeight * JWeight, 0.05f);
+			
+			WeightTotal += Weight;
+
+			//accumulate values from alpha channel
+			float Sample = InBuffer.Load(int3(i, j, 0)).w;
+			Value += Weight * Sample;
+		}
+	}
+
+	return  float2(Value, WeightTotal);
+}
+
 
 // @param FrameTime in seconds
 // @return smoothed exposure
@@ -148,13 +178,28 @@ void MainReductionPS(
 	OutColor = Texture2DSample(PostprocessInput0, PostprocessInput0Sampler, InUV);
 }
 
+#if COMPUTESHADER
+
+#define TGSIZE 16
+groupshared float2 SubRectValueWeight[TGSIZE*TGSIZE];
+RWTexture2D<float4> RWOutputComputeTex;
 
-void MainLogLuminance2ExposureScalePS(
-	in float4 UVAndScreenPos : TEXCOORD0,
+[numthreads(TGSIZE, TGSIZE, 1)]
+void MainLogLuminance2ExposureScaleCS(uint GIndex : SV_GroupIndex, 
+	uint2 GTId : SV_GroupThreadID
+	)
+{
+	float4 OutColor=0;
+
+#else
+
+void MainLogLuminance2ExposureScalePS( 
+	in float4 UVAndScreenPos : TEXCOORD0, 
 	out float4 OutColor : SV_Target0
 	)
 {
-	float2 UV = UVAndScreenPos.xy;
+
+#endif
 
 #if USE_PREEXPOSURE
 	const float OneOverPreExposure = View.OneOverPreExposure;
@@ -169,8 +214,49 @@ void MainLogLuminance2ExposureScalePS(
 	//              >1 focuses more in the center of the screen.
 	float WeightSlope = EyeAdaptationParams[2].w;
 	
-	// Compute scaled Log Luminance Average 
+	// Compute scaled Log Luminance Average
+#if COMPUTESHADER
+	// There are TGSIZE*TGSIZE threads. Each thread will calculate the luminance for its own subregions from a TGSIZE*TGSIZE screen grid
+	uint4 SubRect = uint4(	((EyeAdaptionSrcRect.z-EyeAdaptionSrcRect.x) * GTId.x)		/ TGSIZE,
+							((EyeAdaptionSrcRect.w-EyeAdaptionSrcRect.y) * GTId.y)		/ TGSIZE, 
+							((EyeAdaptionSrcRect.z-EyeAdaptionSrcRect.x) * (GTId.x+1))	/ TGSIZE  ,
+							((EyeAdaptionSrcRect.w-EyeAdaptionSrcRect.y) * (GTId.y+1))	/ TGSIZE  );
+
+	float2 LogLumAveWithWeight = ComputeWeightedTextureAverageAlphaSubRegion(PostprocessInput0, SubRect,  EyeAdaptionSrcRect, WeightSlope);
+
+	// Store in LDS
+	SubRectValueWeight[GIndex] = LogLumAveWithWeight;
+	GroupMemoryBarrierWithGroupSync();
+
+	// Merge the ValueWeight from all threads
+	SubRectValueWeight[GIndex] =  SubRectValueWeight[GIndex] + SubRectValueWeight[GIndex ^ 1];
+	GroupMemoryBarrierWithGroupSync();
+
+	SubRectValueWeight[GIndex] =  SubRectValueWeight[GIndex] + SubRectValueWeight[GIndex ^ 2];
+	GroupMemoryBarrierWithGroupSync();
+
+	SubRectValueWeight[GIndex] =  SubRectValueWeight[GIndex] + SubRectValueWeight[GIndex ^ 4];
+	GroupMemoryBarrierWithGroupSync();
+
+	SubRectValueWeight[GIndex] =  SubRectValueWeight[GIndex] + SubRectValueWeight[GIndex ^ 8];
+	GroupMemoryBarrierWithGroupSync();
+
+	SubRectValueWeight[GIndex] =  SubRectValueWeight[GIndex] + SubRectValueWeight[GIndex ^ 16];
+	GroupMemoryBarrierWithGroupSync();
+
+	SubRectValueWeight[GIndex] =  SubRectValueWeight[GIndex] + SubRectValueWeight[GIndex ^ 32];
+	GroupMemoryBarrierWithGroupSync();
+
+	SubRectValueWeight[GIndex] =  SubRectValueWeight[GIndex] + SubRectValueWeight[GIndex ^ 64];
+	GroupMemoryBarrierWithGroupSync();
+
+	SubRectValueWeight[GIndex] =  SubRectValueWeight[GIndex] + SubRectValueWeight[GIndex ^ 128];
+	GroupMemoryBarrierWithGroupSync();
+
+	float LogLumAve = SubRectValueWeight[0].x /  SubRectValueWeight[0].y;
+#else
 	float LogLumAve = ComputeWeightedTextureAverageAlpha(PostprocessInput0,  EyeAdaptionSrcRect, WeightSlope);
+#endif
 
 	// Correct for [0,1] scaling
 	LogLumAve = (LogLumAve - EyeAdaptationParams[2].y) / EyeAdaptationParams[2].x;
@@ -202,4 +288,11 @@ void MainLogLuminance2ExposureScalePS(
 	OutColor.y = ExposureOffsetMultipler * TargetExposureScale;
 	OutColor.z = AverageSceneLuminance;
 	OutColor.w = 0;
+		
+#if COMPUTESHADER
+	if(GIndex==0)
+	{
+		RWOutputComputeTex[uint2(0,0)] = OutColor;
+	}
+#endif
 }
\ No newline at end of file
diff --git a/Engine/Source/Runtime/Renderer/Private/PostProcess/PostProcessEyeAdaptation.cpp b/Engine/Source/Runtime/Renderer/Private/PostProcess/PostProcessEyeAdaptation.cpp
--- a/Engine/Source/Runtime/Renderer/Private/PostProcess/PostProcessEyeAdaptation.cpp
+++ b/Engine/Source/Runtime/Renderer/Private/PostProcess/PostProcessEyeAdaptation.cpp
@@ -46,6 +46,14 @@ TAutoConsoleVariable<float> CVarEyeAdaptationFocus(
 	TEXT(">0: Center focus, 1 is a good number (default)"),
 	ECVF_Scalability | ECVF_RenderThreadSafe);
 
+TAutoConsoleVariable<int32> CVarEyeAdaptationBasicCompute(
+	TEXT("r.EyeAdaptation.Basic.Compute"),
+	1,
+	TEXT("Use Pixel or Compute Shader to compute the basic eye adaptation. \n")
+	TEXT("= 0 : Pixel Shader\n")
+	TEXT("> 0 : Compute Shader (default) \n"),
+	ECVF_Scalability | ECVF_RenderThreadSafe);
+
 /**
  *   Shared functionality used in computing the eye-adaptation parameters
  *   Compute the parameters used for eye-adaptation.  These will default to values
@@ -625,19 +633,13 @@ FPooledRenderTargetDesc FRCPassPostProcessBasicEyeAdaptationSetUp::ComputeOutput
 
 
 
-/** Encapsulates the post process computation of the exposure scale  pixel shader. */
-class FPostProcessLogLuminance2ExposureScalePS : public FGlobalShader
+class FPostProcessLogLuminance2ExposureScaleBase: public FGlobalShader
 {
-	DECLARE_SHADER_TYPE(FPostProcessLogLuminance2ExposureScalePS, Global);
-
-
-
 public:
-	/** Default constructor. */
-	FPostProcessLogLuminance2ExposureScalePS() {}
 
-	/** Initialization constructor. */
-	FPostProcessLogLuminance2ExposureScalePS(const ShaderMetaType::CompiledShaderInitializerType& Initializer)
+	FPostProcessLogLuminance2ExposureScaleBase() {}
+
+	FPostProcessLogLuminance2ExposureScaleBase(const ShaderMetaType::CompiledShaderInitializerType& Initializer)
 		: FGlobalShader(Initializer)
 	{
 		PostprocessParameter.Bind(Initializer.ParameterMap);
@@ -646,19 +648,54 @@ public:
 		EyeAdaptationSrcRect.Bind(Initializer.ParameterMap, TEXT("EyeAdaptionSrcRect"));
 	}
 
+	/** Static Shader boilerplate */
+	static void ModifyCompilationEnvironment(const FGlobalShaderPermutationParameters& Parameters, FShaderCompilerEnvironment& OutEnvironment)
+	{
+		FGlobalShader::ModifyCompilationEnvironment(Parameters, OutEnvironment);
+		OutEnvironment.SetRenderTargetOutputFormat(0, PF_A32B32G32R32F);
+		OutEnvironment.SetDefine(TEXT("EYE_ADAPTATION_PARAMS_SIZE"), (uint32)EYE_ADAPTATION_PARAMS_SIZE);
+	}
+
+	// FShader interface.
+	virtual bool Serialize(FArchive& Ar) override
+	{
+		bool bShaderHasOutdatedParameters = FGlobalShader::Serialize(Ar);
+		Ar << PostprocessParameter << EyeAdaptationTexture
+			<< EyeAdaptationParams << EyeAdaptationSrcRect;
+		return bShaderHasOutdatedParameters;
+	}
+
+protected:
+	FPostProcessPassParameters	PostprocessParameter;
+	FShaderResourceParameter	EyeAdaptationTexture;
+	FShaderParameter			EyeAdaptationParams;
+	FShaderParameter			EyeAdaptationSrcRect;
+};
+
+
+
+/** Encapsulates the post process computation of the exposure scale pixel shader. */
+class FPostProcessLogLuminance2ExposureScalePS : public FPostProcessLogLuminance2ExposureScaleBase
+{
+	DECLARE_SHADER_TYPE(FPostProcessLogLuminance2ExposureScalePS, Global);
+
 public:
 
-	/** Static Shader boilerplate */
 	static bool ShouldCompilePermutation(const FGlobalShaderPermutationParameters& Parameters)
 	{
 		return IsFeatureLevelSupported(Parameters.Platform, BasicEyeAdaptationMinFeatureLevel);
 	}
 
-	static void ModifyCompilationEnvironment(const FGlobalShaderPermutationParameters& Parameters, FShaderCompilerEnvironment& OutEnvironment)
+
+	/** Default constructor. */
+	FPostProcessLogLuminance2ExposureScalePS()
+		: FPostProcessLogLuminance2ExposureScaleBase()
+	{}
+
+	/** Initialization constructor. */
+	FPostProcessLogLuminance2ExposureScalePS(const ShaderMetaType::CompiledShaderInitializerType& Initializer)
+		: FPostProcessLogLuminance2ExposureScaleBase(Initializer)
 	{
-		FGlobalShader::ModifyCompilationEnvironment(Parameters, OutEnvironment);
-		OutEnvironment.SetRenderTargetOutputFormat(0, PF_A32B32G32R32F);
-		OutEnvironment.SetDefine(TEXT("EYE_ADAPTATION_PARAMS_SIZE"), (uint32)EYE_ADAPTATION_PARAMS_SIZE);
 	}
 
 
@@ -697,24 +734,91 @@ public:
 		// Set the src extent for the shader
 		SetShaderValue(Context.RHICmdList, ShaderRHI, EyeAdaptationSrcRect, SrcRect);
 	}
+};
+IMPLEMENT_SHADER_TYPE(, FPostProcessLogLuminance2ExposureScalePS, TEXT("/Engine/Private/PostProcessEyeAdaptation.usf"), TEXT("MainLogLuminance2ExposureScalePS"), SF_Pixel);
+
+
+class FPostProcessLogLuminance2ExposureScaleCS : public FPostProcessLogLuminance2ExposureScaleBase
+{
+	DECLARE_SHADER_TYPE(FPostProcessLogLuminance2ExposureScaleCS, Global);
+
+
+public:
+
+	/** Default constructor. */
+	FPostProcessLogLuminance2ExposureScaleCS()
+		: FPostProcessLogLuminance2ExposureScaleBase()
+	{}
+
+	/** Initialization constructor. */
+	FPostProcessLogLuminance2ExposureScaleCS(const ShaderMetaType::CompiledShaderInitializerType& Initializer)
+		: FPostProcessLogLuminance2ExposureScaleBase(Initializer)
+	{
+		OutComputeTex.Bind(Initializer.ParameterMap, TEXT("OutputComputeTex"));
+	}
+
+	/** Static Shader boilerplate */
+	static bool ShouldCompilePermutation(const FGlobalShaderPermutationParameters& Parameters)
+	{
+		return IsFeatureLevelSupported(Parameters.Platform, ERHIFeatureLevel::SM5);
+	}
 
 	// FShader interface.
 	virtual bool Serialize(FArchive& Ar) override
 	{
 		bool bShaderHasOutdatedParameters = FGlobalShader::Serialize(Ar);
-		Ar << PostprocessParameter << EyeAdaptationTexture 
-			<< EyeAdaptationParams << EyeAdaptationSrcRect;
+		Ar << PostprocessParameter << EyeAdaptationTexture
+			<< EyeAdaptationParams << EyeAdaptationSrcRect << OutComputeTex;
 		return bShaderHasOutdatedParameters;
 	}
 
-private:
-	FPostProcessPassParameters PostprocessParameter;
-	FShaderResourceParameter EyeAdaptationTexture;
-	FShaderParameter EyeAdaptationParams;
-	FShaderParameter EyeAdaptationSrcRect;
+
+	void UnsetParameters(const FRenderingCompositePassContext& Context)
+	{
+		const FComputeShaderRHIParamRef ShaderRHI = GetComputeShader();
+		OutComputeTex.UnsetUAV(Context.RHICmdList, ShaderRHI);
+	}
+
+	void SetParameters(const FRenderingCompositePassContext& Context, const FIntRect& SrcRect, IPooledRenderTarget* LastEyeAdaptation, FUnorderedAccessViewRHIParamRef DestUAV)
+	{
+		const FComputeShaderRHIParamRef ShaderRHI = GetComputeShader();
+
+		FGlobalShader::SetParameters<FViewUniformShaderParameters>(Context.RHICmdList, ShaderRHI, Context.View.ViewUniformBuffer);
+		PostprocessParameter.SetCS(ShaderRHI, Context, Context.RHICmdList, TStaticSamplerState<SF_Bilinear, AM_Clamp, AM_Clamp, AM_Clamp>::GetRHI());
+
+		// Associate the eye adaptation buffer from the previous frame with a texture to be read in this frame. 
+		if (Context.View.HasValidEyeAdaptation())
+		{
+			SetTextureParameter(Context.RHICmdList, ShaderRHI, EyeAdaptationTexture, LastEyeAdaptation->GetRenderTargetItem().TargetableTexture);
+		}
+		else // some views don't have a state, thumbnail rendering?
+		{
+			SetTextureParameter(Context.RHICmdList, ShaderRHI, EyeAdaptationTexture, GWhiteTexture->TextureRHI);
+		}
+
+		// Pack the eye adaptation parameters for the shader
+		{
+			FVector4 Temp[EYE_ADAPTATION_PARAMS_SIZE];
+			// static computation function
+			ComputeEyeAdaptationValues(BasicEyeAdaptationMinFeatureLevel, Context.View, Temp);
+			// Log-based computation of the exposure scale has a built in scaling.
+			//Temp[1].X *= 0.16;  
+			//Encode the eye-focus slope
+			// Get the focus value for the eye-focus weighting
+			Temp[2].W = GetBasicAutoExposureFocus();
+			SetShaderValueArray(Context.RHICmdList, ShaderRHI, EyeAdaptationParams, Temp, EYE_ADAPTATION_PARAMS_SIZE);
+		}
+
+		// Set the src extent for the shader
+		SetShaderValue(Context.RHICmdList, ShaderRHI, EyeAdaptationSrcRect, SrcRect);
+		OutComputeTex.SetTexture(Context.RHICmdList, ShaderRHI, nullptr, DestUAV);
+	}
+	private:
+
+	FRWShaderParameter			OutComputeTex;
 };
+IMPLEMENT_SHADER_TYPE(, FPostProcessLogLuminance2ExposureScaleCS, TEXT("/Engine/Private/PostProcessEyeAdaptation.usf"), TEXT("MainLogLuminance2ExposureScaleCS"), SF_Compute);
 
-IMPLEMENT_SHADER_TYPE(, FPostProcessLogLuminance2ExposureScalePS, TEXT("/Engine/Private/PostProcessEyeAdaptation.usf"), TEXT("MainLogLuminance2ExposureScalePS"), SF_Pixel);
 
 void FRCPassPostProcessBasicEyeAdaptation::Process(FRenderingCompositePassContext& Context)
 {
@@ -750,44 +854,63 @@ void FRCPassPostProcessBasicEyeAdaptation::Process(FRenderingCompositePassContex
 
 	FRHIRenderPassInfo RPInfo(DestRenderTarget.TargetableTexture, ERenderTargetActions::Load_Store);
 	TransitionRenderPassTargets(Context.RHICmdList, RPInfo);
-	Context.RHICmdList.BeginRenderPass(RPInfo, TEXT("BasicEyeAdaptation"));
+		
+	bool IsCompute = (CVarEyeAdaptationBasicCompute.GetValueOnRenderThread()>0) && (Context.View.GetFeatureLevel() >= ERHIFeatureLevel::SM5);
+	if (IsCompute)
 	{
-		Context.SetViewportAndCallRHI(0, 0, 0.0f, DestSize.X, DestSize.Y, 1.0f);
+		WaitForInputPassComputeFences(Context.RHICmdList);
 
-		FGraphicsPipelineStateInitializer GraphicsPSOInit;
-		Context.RHICmdList.ApplyCachedRenderTargets(GraphicsPSOInit);
-		GraphicsPSOInit.BlendState = TStaticBlendState<>::GetRHI();
-		GraphicsPSOInit.RasterizerState = TStaticRasterizerState<>::GetRHI();
-		GraphicsPSOInit.DepthStencilState = TStaticDepthStencilState<false, CF_Always>::GetRHI();
+		SetRenderTarget(Context.RHICmdList, FTextureRHIRef(), FTextureRHIRef(), false);
+		TShaderMapRef<FPostProcessLogLuminance2ExposureScaleCS> ComputeShader(Context.GetShaderMap());
+		Context.RHICmdList.SetComputeShader(ComputeShader->GetComputeShader());
 
-		TShaderMapRef<FPostProcessVS> VertexShader(Context.GetShaderMap());
-		TShaderMapRef<FPostProcessLogLuminance2ExposureScalePS> PixelShader(Context.GetShaderMap());
+		Context.RHICmdList.TransitionResource(EResourceTransitionAccess::ERWBarrier, EResourceTransitionPipeline::EGfxToCompute, DestRenderTarget.UAV);
+		ComputeShader->SetParameters(Context, SrcRect, EyeAdaptationLastFrameRT, DestRenderTarget.UAV);
+		DispatchComputeShader(Context.RHICmdList, *ComputeShader, 1, 1, 1);
+		ComputeShader->UnsetParameters(Context);
+		Context.RHICmdList.TransitionResource(EResourceTransitionAccess::EReadable, EResourceTransitionPipeline::EComputeToGfx, DestRenderTarget.UAV);
+	}
+	else
+	{
+		Context.RHICmdList.BeginRenderPass(RPInfo, TEXT("BasicEyeAdaptation"));
+		{
+			Context.SetViewportAndCallRHI(0, 0, 0.0f, DestSize.X, DestSize.Y, 1.0f);
 
-		GraphicsPSOInit.BoundShaderState.VertexDeclarationRHI = GFilterVertexDeclaration.VertexDeclarationRHI;
-		GraphicsPSOInit.BoundShaderState.VertexShaderRHI = GETSAFERHISHADER_VERTEX(*VertexShader);
-		GraphicsPSOInit.BoundShaderState.PixelShaderRHI = GETSAFERHISHADER_PIXEL(*PixelShader);
-		GraphicsPSOInit.PrimitiveType = PT_TriangleList;
+			FGraphicsPipelineStateInitializer GraphicsPSOInit;
+			Context.RHICmdList.ApplyCachedRenderTargets(GraphicsPSOInit);
+			GraphicsPSOInit.BlendState = TStaticBlendState<>::GetRHI();
+			GraphicsPSOInit.RasterizerState = TStaticRasterizerState<>::GetRHI();
+			GraphicsPSOInit.DepthStencilState = TStaticDepthStencilState<false, CF_Always>::GetRHI();
 
-		SetGraphicsPipelineState(Context.RHICmdList, GraphicsPSOInit);
+			TShaderMapRef<FPostProcessVS> VertexShader(Context.GetShaderMap());
+			TShaderMapRef<FPostProcessLogLuminance2ExposureScalePS> PixelShader(Context.GetShaderMap());
+
+			GraphicsPSOInit.BoundShaderState.VertexDeclarationRHI = GFilterVertexDeclaration.VertexDeclarationRHI;
+			GraphicsPSOInit.BoundShaderState.VertexShaderRHI = GETSAFERHISHADER_VERTEX(*VertexShader);
+			GraphicsPSOInit.BoundShaderState.PixelShaderRHI = GETSAFERHISHADER_PIXEL(*PixelShader);
+			GraphicsPSOInit.PrimitiveType = PT_TriangleList;
 
-		// Set the parameters used by the pixel shader.
+			SetGraphicsPipelineState(Context.RHICmdList, GraphicsPSOInit);
 
-		PixelShader->SetPS(Context, SrcRect, EyeAdaptationLastFrameRT);
+			// Set the parameters used by the pixel shader.
 
-		// Draw a quad mapping scene color to the view's render target
-		DrawRectangle(
-			Context.RHICmdList,
-			0, 0,
-			DestSize.X, DestSize.Y,
-			0, 0,
-			DestSize.X, DestSize.Y,
-			DestSize,
-			DestSize,
-			*VertexShader,
-			EDRF_UseTriangleOptimization);
+			PixelShader->SetPS(Context, SrcRect, EyeAdaptationLastFrameRT);
+
+			// Draw a quad mapping scene color to the view's render target
+			DrawRectangle(
+				Context.RHICmdList,
+				0, 0,
+				DestSize.X, DestSize.Y,
+				0, 0,
+				DestSize.X, DestSize.Y,
+				DestSize,
+				DestSize,
+				*VertexShader,
+				EDRF_UseTriangleOptimization);
+		}
+		Context.RHICmdList.EndRenderPass();
+		Context.RHICmdList.CopyToResolveTarget(DestRenderTarget.TargetableTexture, DestRenderTarget.ShaderResourceTexture, FResolveParams());
 	}
-	Context.RHICmdList.EndRenderPass();
-	Context.RHICmdList.CopyToResolveTarget(DestRenderTarget.TargetableTexture, DestRenderTarget.ShaderResourceTexture, FResolveParams());
 
 	// Inform MultiGPU systems that we've finished with this texture for this frame
 	Context.RHICmdList.EndUpdateMultiFrameResource(DestRenderTarget.ShaderResourceTexture);
@@ -804,6 +927,11 @@ FPooledRenderTargetDesc FRCPassPostProcessBasicEyeAdaptation::ComputeOutputDesc(
 
 	Ret.DebugName = TEXT("EyeAdaptationBasic");
 	Ret.Flags |= GFastVRamConfig.EyeAdaptation;
+
+	if (GMaxRHIFeatureLevel >= ERHIFeatureLevel::SM5)
+	{
+		Ret.Flags |= TexCreate_UAV;
+	}
 	return Ret;
 }
 
