Fix deprecation warning in eye adaptation.

diff --git a/Engine/Source/Runtime/Renderer/Private/PostProcess/PostProcessEyeAdaptation.cpp b/Engine/Source/Runtime/Renderer/Private/PostProcess/PostProcessEyeAdaptation.cpp
--- a/Engine/Source/Runtime/Renderer/Private/PostProcess/PostProcessEyeAdaptation.cpp
+++ b/Engine/Source/Runtime/Renderer/Private/PostProcess/PostProcessEyeAdaptation.cpp
@@ -860,7 +860,8 @@ void FRCPassPostProcessBasicEyeAdaptation::Process(FRenderingCompositePassContex
 	{
 		WaitForInputPassComputeFences(Context.RHICmdList);
 
-		SetRenderTarget(Context.RHICmdList, FTextureRHIRef(), FTextureRHIRef(), false);
+		// #todo-renderpasses remove once everything is renderpasses
+		UnbindRenderTargets(Context.RHICmdList);
 		TShaderMapRef<FPostProcessLogLuminance2ExposureScaleCS> ComputeShader(Context.GetShaderMap());
 		Context.RHICmdList.SetComputeShader(ComputeShader->GetComputeShader());
 
