Added Distance field shadow / AO support on Switch
Don't compile distance field shader permutations if bUseDistanceFields is unset or false
Added UAV output to pixel shader
Clear Tiny UAV uses command buffer ClearBuffer command instead of allocating a temp buffer and copying it to UAV

anthony.bills
#rnx

diff --git a/Engine/Shaders/Private/GlobalDistanceFieldShared.ush b/Engine/Shaders/Private/GlobalDistanceFieldShared.ush
--- a/Engine/Shaders/Private/GlobalDistanceFieldShared.ush
+++ b/Engine/Shaders/Private/GlobalDistanceFieldShared.ush
@@ -13,7 +13,7 @@ float MaxGlobalDistance;
 #endif
 
 //@todo - get the clipmap branches to compile under gl
-#if SM5_PROFILE || PS4_PROFILE || METAL_SM5_PROFILE || METAL_SM5_NOTESS_PROFILE || METAL_MRT_PROFILE || GL4_PROFILE || VULKAN_PROFILE_SM5
+#if SM5_PROFILE || PS4_PROFILE || METAL_SM5_PROFILE || METAL_SM5_NOTESS_PROFILE || METAL_MRT_PROFILE || GL4_PROFILE || VULKAN_PROFILE_SM5 || SWITCH_PROFILE
 
 // Must match C++
 #define MAX_GLOBAL_DF_CLIPMAPS 4
diff --git a/Engine/Source/Developer/TargetPlatform/Private/TargetPlatformBase.cpp b/Engine/Source/Developer/TargetPlatform/Private/TargetPlatformBase.cpp
--- a/Engine/Source/Developer/TargetPlatform/Private/TargetPlatformBase.cpp
+++ b/Engine/Source/Developer/TargetPlatform/Private/TargetPlatformBase.cpp
@@ -28,6 +28,11 @@ bool FTargetPlatformBase::UsesSelectiveBasePassOutputs() const
 	return CVar ? (CVar->GetInt() != 0) : false;
 }
 
+bool FTargetPlatformBase::UsesDistanceFields() const
+{
+	return true;
+}
+
 TSharedPtr<IDeviceManagerCustomPlatformWidgetCreator> FTargetPlatformBase::GetCustomWidgetCreator() const
 {
 	static TSharedPtr<FDeviceBrowserDefaultPlatformWidgetCreator> DefaultWidgetCreator = MakeShared<FDeviceBrowserDefaultPlatformWidgetCreator>();
diff --git a/Engine/Source/Developer/TargetPlatform/Public/Common/TargetPlatformBase.h b/Engine/Source/Developer/TargetPlatform/Public/Common/TargetPlatformBase.h
--- a/Engine/Source/Developer/TargetPlatform/Public/Common/TargetPlatformBase.h
+++ b/Engine/Source/Developer/TargetPlatform/Public/Common/TargetPlatformBase.h
@@ -38,6 +38,8 @@ public:
 	TARGETPLATFORM_API virtual bool UsesBasePassVelocity() const override;
 
 	TARGETPLATFORM_API virtual bool UsesSelectiveBasePassOutputs() const override;
+	
+	TARGETPLATFORM_API virtual bool UsesDistanceFields() const override;
 
 #if WITH_ENGINE
 	virtual void GetReflectionCaptureFormats( TArray<FName>& OutFormats ) const override
diff --git a/Engine/Source/Developer/TargetPlatform/Public/Interfaces/ITargetPlatform.h b/Engine/Source/Developer/TargetPlatform/Public/Interfaces/ITargetPlatform.h
--- a/Engine/Source/Developer/TargetPlatform/Public/Interfaces/ITargetPlatform.h
+++ b/Engine/Source/Developer/TargetPlatform/Public/Interfaces/ITargetPlatform.h
@@ -353,6 +353,11 @@ public:
 	*/
 	virtual bool UsesSelectiveBasePassOutputs() const = 0; 
 
+	/**
+	* Gets whether the platform will use distance fields.
+	*/
+	virtual bool UsesDistanceFields() const = 0;
+
 #if WITH_ENGINE
 	/**
 	 * Gets the format to use for a particular body setup.
diff --git a/Engine/Source/Runtime/Engine/Public/SceneManagement.h b/Engine/Source/Runtime/Engine/Public/SceneManagement.h
--- a/Engine/Source/Runtime/Engine/Public/SceneManagement.h
+++ b/Engine/Source/Runtime/Engine/Public/SceneManagement.h
@@ -953,7 +953,8 @@ inline bool DoesPlatformSupportDistanceFieldShadowing(EShaderPlatform Platform)
 	return Platform == SP_PCD3D_SM5 || Platform == SP_PS4
 		|| IsMetalSM5Platform(Platform)
 		|| Platform == SP_XBOXONE_D3D12
-		|| IsVulkanSM5Platform(Platform);
+		|| IsVulkanSM5Platform(Platform)
+	    || Platform == SP_SWITCH || Platform == SP_SWITCH_FORWARD;
 }
 
 BEGIN_GLOBAL_SHADER_PARAMETER_STRUCT(FMobileReflectionCaptureShaderParameters,ENGINE_API)
diff --git a/Engine/Source/Runtime/RenderCore/Private/RenderUtils.cpp b/Engine/Source/Runtime/RenderCore/Private/RenderUtils.cpp
--- a/Engine/Source/Runtime/RenderCore/Private/RenderUtils.cpp
+++ b/Engine/Source/Runtime/RenderCore/Private/RenderUtils.cpp
@@ -908,6 +908,16 @@ static FAutoConsoleVariableRef CVarForwardShading(
 	ECVF_RenderThreadSafe | ECVF_ReadOnly
 	); 
 
+static TAutoConsoleVariable<int32> CVarDistanceFields(
+	TEXT("r.DistanceFields"),
+	1,
+	TEXT("Enables distance fields rendering.\n") \
+	TEXT(" 0: Disabled.\n") \
+	TEXT(" 1: Enabled."),
+	ECVF_RenderThreadSafe | ECVF_ReadOnly
+	); 
+
+
 RENDERCORE_API uint32 GForwardShadingPlatformMask = 0;
 static_assert(SP_NumPlatforms <= sizeof(GForwardShadingPlatformMask) * 8, "GForwardShadingPlatformMask must be large enough to support all shader platforms");
 
@@ -920,6 +930,9 @@ static_assert(SP_NumPlatforms <= sizeof(GBasePassVelocityPlatformMask) * 8, "GBa
 RENDERCORE_API uint32 GSelectiveBasePassOutputsPlatformMask = 0;
 static_assert(SP_NumPlatforms <= sizeof(GSelectiveBasePassOutputsPlatformMask) * 8, "GSelectiveBasePassOutputsPlatformMask must be large enough to support all shader platforms");
 
+RENDERCORE_API uint32 GDistanceFieldsPlatformMask = 0;
+static_assert(SP_NumPlatforms <= sizeof(GDistanceFieldsPlatformMask) * 8, "GDistanceFieldsPlatformMask must be large enough to support all shader platforms");
+
 RENDERCORE_API void RenderUtilsInit()
 {
 	if (GUseForwardShading)
@@ -945,6 +958,12 @@ RENDERCORE_API void RenderUtilsInit()
 		GSelectiveBasePassOutputsPlatformMask = ~0u;
 	}
 
+	static IConsoleVariable* DistanceFieldsCVar = IConsoleManager::Get().FindConsoleVariable(TEXT("r.DistanceFields")); 
+	if (DistanceFieldsCVar && DistanceFieldsCVar->GetInt())
+	{
+		GDistanceFieldsPlatformMask = ~0u;
+	}
+
 #if WITH_EDITOR
 	ITargetPlatformManagerModule* TargetPlatformManager = GetTargetPlatformManager();
 	if (TargetPlatformManager)
@@ -993,6 +1012,15 @@ RENDERCORE_API void RenderUtilsInit()
 				{
 					GSelectiveBasePassOutputsPlatformMask &= ~Mask;
 				}
+
+				if (TargetPlatform->UsesDistanceFields())
+				{
+					GDistanceFieldsPlatformMask |= Mask;
+				}
+				else
+				{
+					GDistanceFieldsPlatformMask &= ~Mask;
+				}
 			}
 		}
 	}
diff --git a/Engine/Source/Runtime/RenderCore/Public/RenderUtils.h b/Engine/Source/Runtime/RenderCore/Public/RenderUtils.h
--- a/Engine/Source/Runtime/RenderCore/Public/RenderUtils.h
+++ b/Engine/Source/Runtime/RenderCore/Public/RenderUtils.h
@@ -470,6 +470,13 @@ inline bool IsUsingSelectiveBasePassOutputs(EShaderPlatform Platform)
 	return !!(GSelectiveBasePassOutputsPlatformMask & (1u << Platform));
 }
 
+/** Returns whether distance fields are enabled for a given shader platform */
+inline bool IsUsingDistanceFields(EShaderPlatform Platform)
+{
+	extern RENDERCORE_API uint32 GDistanceFieldsPlatformMask;
+	return !!(GDistanceFieldsPlatformMask & (1u << Platform));
+}
+
 inline bool IsUsingPerPixelDBufferMask(EShaderPlatform Platform)
 {
 	switch (Platform)
diff --git a/Engine/Source/Runtime/Renderer/Private/DistanceFieldAmbientOcclusion.cpp b/Engine/Source/Runtime/Renderer/Private/DistanceFieldAmbientOcclusion.cpp
--- a/Engine/Source/Runtime/Renderer/Private/DistanceFieldAmbientOcclusion.cpp
+++ b/Engine/Source/Runtime/Renderer/Private/DistanceFieldAmbientOcclusion.cpp
@@ -279,7 +279,7 @@ public:
 
 	static bool ShouldCompilePermutation(const FGlobalShaderPermutationParameters& Parameters)
 	{
-		return IsFeatureLevelSupported(Parameters.Platform, ERHIFeatureLevel::SM5) && DoesPlatformSupportDistanceFieldAO(Parameters.Platform);
+		return IsFeatureLevelSupported(Parameters.Platform, ERHIFeatureLevel::SM5) && DoesPlatformSupportDistanceFieldAO(Parameters.Platform) && IsUsingDistanceFields(Parameters.Platform);
 	}
 
 	static void ModifyCompilationEnvironment(const FGlobalShaderPermutationParameters& Parameters, FShaderCompilerEnvironment& OutEnvironment)
@@ -333,7 +333,7 @@ public:
 
 	static bool ShouldCompilePermutation(const FGlobalShaderPermutationParameters& Parameters)
 	{
-		return IsFeatureLevelSupported(Parameters.Platform, ERHIFeatureLevel::SM5) && DoesPlatformSupportDistanceFieldAO(Parameters.Platform);
+		return IsFeatureLevelSupported(Parameters.Platform, ERHIFeatureLevel::SM5) && DoesPlatformSupportDistanceFieldAO(Parameters.Platform) && IsUsingDistanceFields(Parameters.Platform);
 	}
 
 	static void ModifyCompilationEnvironment(const FGlobalShaderPermutationParameters& Parameters, FShaderCompilerEnvironment& OutEnvironment)
@@ -715,7 +715,8 @@ bool SupportsDistanceFieldAO(ERHIFeatureLevel::Type FeatureLevel, EShaderPlatfor
 		// Intel HD 4000 hangs in the RHICreateTexture3D call to allocate the large distance field atlas, and virtually no Intel cards can afford it anyway
 		&& !IsRHIDeviceIntel()
 		&& FeatureLevel >= ERHIFeatureLevel::SM5
-		&& DoesPlatformSupportDistanceFieldAO(ShaderPlatform);
+		&& DoesPlatformSupportDistanceFieldAO(ShaderPlatform)
+		&& IsUsingDistanceFields(ShaderPlatform);
 }
 
 bool ShouldRenderDeferredDynamicSkyLight(const FScene* Scene, const FSceneViewFamily& ViewFamily)
diff --git a/Engine/Source/Runtime/Renderer/Private/DistanceFieldAmbientOcclusion.h b/Engine/Source/Runtime/Renderer/Private/DistanceFieldAmbientOcclusion.h
--- a/Engine/Source/Runtime/Renderer/Private/DistanceFieldAmbientOcclusion.h
+++ b/Engine/Source/Runtime/Renderer/Private/DistanceFieldAmbientOcclusion.h
@@ -30,7 +30,8 @@ inline bool DoesPlatformSupportDistanceFieldAO(EShaderPlatform Platform)
 {
 	return Platform == SP_PCD3D_SM5 || Platform == SP_PS4 || Platform == SP_XBOXONE_D3D12
 		|| IsMetalSM5Platform(Platform)
-		|| IsVulkanSM5Platform(Platform);
+		|| IsVulkanSM5Platform(Platform)
+		|| Platform == SP_SWITCH || Platform == SP_SWITCH_FORWARD;
 }
 
 extern FIntPoint GetBufferSizeForAO();
diff --git a/Engine/Source/Runtime/Renderer/Private/DistanceFieldLightingPost.cpp b/Engine/Source/Runtime/Renderer/Private/DistanceFieldLightingPost.cpp
--- a/Engine/Source/Runtime/Renderer/Private/DistanceFieldLightingPost.cpp
+++ b/Engine/Source/Runtime/Renderer/Private/DistanceFieldLightingPost.cpp
@@ -165,7 +165,7 @@ public:
 
 	static bool ShouldCompilePermutation(const FGlobalShaderPermutationParameters& Parameters)
 	{
-		return IsFeatureLevelSupported(Parameters.Platform, ERHIFeatureLevel::SM5) && DoesPlatformSupportDistanceFieldAO(Parameters.Platform);
+		return IsFeatureLevelSupported(Parameters.Platform, ERHIFeatureLevel::SM5) && DoesPlatformSupportDistanceFieldAO(Parameters.Platform) && IsUsingDistanceFields(Parameters.Platform);
 	}
 
 	static void ModifyCompilationEnvironment(const FGlobalShaderPermutationParameters& Parameters, FShaderCompilerEnvironment& OutEnvironment)
@@ -302,7 +302,7 @@ public:
 
 	static bool ShouldCompilePermutation(const FGlobalShaderPermutationParameters& Parameters)
 	{
-		return IsFeatureLevelSupported(Parameters.Platform, ERHIFeatureLevel::SM5) && DoesPlatformSupportDistanceFieldAO(Parameters.Platform);
+		return IsFeatureLevelSupported(Parameters.Platform, ERHIFeatureLevel::SM5) && DoesPlatformSupportDistanceFieldAO(Parameters.Platform) && IsUsingDistanceFields(Parameters.Platform);
 	}
 
 	static void ModifyCompilationEnvironment(const FGlobalShaderPermutationParameters& Parameters, FShaderCompilerEnvironment& OutEnvironment)
@@ -413,7 +413,7 @@ public:
 
 	static bool ShouldCompilePermutation(const FGlobalShaderPermutationParameters& Parameters)
 	{
-		return IsFeatureLevelSupported(Parameters.Platform, ERHIFeatureLevel::SM5) && DoesPlatformSupportDistanceFieldAO(Parameters.Platform);
+		return IsFeatureLevelSupported(Parameters.Platform, ERHIFeatureLevel::SM5) && DoesPlatformSupportDistanceFieldAO(Parameters.Platform) && IsUsingDistanceFields(Parameters.Platform);
 	}
 
 	static void ModifyCompilationEnvironment(const FGlobalShaderPermutationParameters& Parameters, FShaderCompilerEnvironment& OutEnvironment)
@@ -703,7 +703,7 @@ public:
 
 	static bool ShouldCompilePermutation(const FGlobalShaderPermutationParameters& Parameters)
 	{
-		return IsFeatureLevelSupported(Parameters.Platform, ERHIFeatureLevel::SM5) && DoesPlatformSupportDistanceFieldAO(Parameters.Platform);
+		return IsFeatureLevelSupported(Parameters.Platform, ERHIFeatureLevel::SM5) && DoesPlatformSupportDistanceFieldAO(Parameters.Platform) && IsUsingDistanceFields(Parameters.Platform);
 	}
 
 	static void ModifyCompilationEnvironment(const FGlobalShaderPermutationParameters& Parameters, FShaderCompilerEnvironment& OutEnvironment)
diff --git a/Engine/Source/Runtime/Renderer/Private/DistanceFieldLightingShared.h b/Engine/Source/Runtime/Renderer/Private/DistanceFieldLightingShared.h
--- a/Engine/Source/Runtime/Renderer/Private/DistanceFieldLightingShared.h
+++ b/Engine/Source/Runtime/Renderer/Private/DistanceFieldLightingShared.h
@@ -287,8 +287,8 @@ public:
 		{
 			const uint32 FastVRamFlag = GFastVRamConfig.DistanceFieldCulledObjectBuffers | ( IsTransientResourceBufferAliasingEnabled() ? BUF_Transient : BUF_None );
 
-			ObjectIndirectArguments.Initialize(sizeof(uint32), 5, PF_R32_UINT, BUF_Static | BUF_DrawIndirect);
-			ObjectIndirectDispatch.Initialize(sizeof(uint32), 3, PF_R32_UINT, BUF_Static | BUF_DrawIndirect);
+			ObjectIndirectArguments.Initialize(sizeof(uint32), 5, PF_R32_UINT, BUF_Static | BUF_DrawIndirect, TEXT("FDistanceFieldCulledObjectBuffers::ObjectIndirectArguments"));
+			ObjectIndirectDispatch.Initialize(sizeof(uint32), 3, PF_R32_UINT, BUF_Static | BUF_DrawIndirect, TEXT("FDistanceFieldCulledObjectBuffers::ObjectIndirectDispatch"));
 			Bounds.Initialize(sizeof(FVector4), MaxObjects, BUF_Static | FastVRamFlag, TEXT("FDistanceFieldCulledObjectBuffers::Bounds"));
 			Data.Initialize(sizeof(FVector4), MaxObjects * ObjectDataStride, BUF_Static | FastVRamFlag, TEXT("FDistanceFieldCulledObjectBuffers::Data"));
 
diff --git a/Engine/Source/Runtime/Renderer/Private/DistanceFieldObjectCulling.cpp b/Engine/Source/Runtime/Renderer/Private/DistanceFieldObjectCulling.cpp
--- a/Engine/Source/Runtime/Renderer/Private/DistanceFieldObjectCulling.cpp
+++ b/Engine/Source/Runtime/Renderer/Private/DistanceFieldObjectCulling.cpp
@@ -110,7 +110,7 @@ public:
 
 	static bool ShouldCompilePermutation(const FGlobalShaderPermutationParameters& Parameters)
 	{
-		return IsFeatureLevelSupported(Parameters.Platform, ERHIFeatureLevel::SM5) && DoesPlatformSupportDistanceFieldAO(Parameters.Platform);
+		return IsFeatureLevelSupported(Parameters.Platform, ERHIFeatureLevel::SM5) && DoesPlatformSupportDistanceFieldAO(Parameters.Platform) && IsUsingDistanceFields(Parameters.Platform);
 	}
 
 	static void ModifyCompilationEnvironment(const FGlobalShaderPermutationParameters& Parameters, FShaderCompilerEnvironment& OutEnvironment)
@@ -231,7 +231,7 @@ public:
 
 	static bool ShouldCompilePermutation(const FGlobalShaderPermutationParameters& Parameters)
 	{
-		return IsFeatureLevelSupported(Parameters.Platform, ERHIFeatureLevel::SM5) && DoesPlatformSupportDistanceFieldAO(Parameters.Platform);
+		return IsFeatureLevelSupported(Parameters.Platform, ERHIFeatureLevel::SM5) && DoesPlatformSupportDistanceFieldAO(Parameters.Platform) && IsUsingDistanceFields(Parameters.Platform);
 	}
 
 	static void ModifyCompilationEnvironment(const FGlobalShaderPermutationParameters& Parameters, FShaderCompilerEnvironment& OutEnvironment)
@@ -342,7 +342,7 @@ public:
 
 	static bool ShouldCompilePermutation(const FGlobalShaderPermutationParameters& Parameters) 
 	{ 
-		return IsFeatureLevelSupported(Parameters.Platform, ERHIFeatureLevel::SM5) && DoesPlatformSupportDistanceFieldAO(Parameters.Platform); 
+		return IsFeatureLevelSupported(Parameters.Platform, ERHIFeatureLevel::SM5) && DoesPlatformSupportDistanceFieldAO(Parameters.Platform) && IsUsingDistanceFields(Parameters.Platform);
 	}
 
 	FObjectCullVS(const ShaderMetaType::CompiledShaderInitializerType& Initializer):
@@ -397,7 +397,7 @@ public:
 
 	static bool ShouldCompilePermutation(const FGlobalShaderPermutationParameters& Parameters)
 	{
-		return IsFeatureLevelSupported(Parameters.Platform, ERHIFeatureLevel::SM5) && DoesPlatformSupportDistanceFieldAO(Parameters.Platform) && RHISupportsPixelShaderUAVs(Parameters.Platform);
+		return IsFeatureLevelSupported(Parameters.Platform, ERHIFeatureLevel::SM5) && DoesPlatformSupportDistanceFieldAO(Parameters.Platform) && RHISupportsPixelShaderUAVs(Parameters.Platform) && IsUsingDistanceFields(Parameters.Platform);
 	}
 
 	static void ModifyCompilationEnvironment(const FGlobalShaderPermutationParameters& Parameters, FShaderCompilerEnvironment& OutEnvironment)
@@ -482,7 +482,7 @@ public:
 
 	static bool ShouldCompilePermutation(const FGlobalShaderPermutationParameters& Parameters)
 	{
-		return IsFeatureLevelSupported(Parameters.Platform, ERHIFeatureLevel::SM5) && DoesPlatformSupportDistanceFieldAO(Parameters.Platform);
+		return IsFeatureLevelSupported(Parameters.Platform, ERHIFeatureLevel::SM5) && DoesPlatformSupportDistanceFieldAO(Parameters.Platform) && IsUsingDistanceFields(Parameters.Platform);
 	}
 
 	static void ModifyCompilationEnvironment(const FGlobalShaderPermutationParameters& Parameters, FShaderCompilerEnvironment& OutEnvironment)
diff --git a/Engine/Source/Runtime/Renderer/Private/DistanceFieldObjectManagement.cpp b/Engine/Source/Runtime/Renderer/Private/DistanceFieldObjectManagement.cpp
--- a/Engine/Source/Runtime/Renderer/Private/DistanceFieldObjectManagement.cpp
+++ b/Engine/Source/Runtime/Renderer/Private/DistanceFieldObjectManagement.cpp
@@ -154,7 +154,7 @@ public:
 
 	static bool ShouldCompilePermutation(const FGlobalShaderPermutationParameters& Parameters)
 	{
-		return IsFeatureLevelSupported(Parameters.Platform, ERHIFeatureLevel::SM5) && DoesPlatformSupportDistanceFieldAO(Parameters.Platform);
+		return IsFeatureLevelSupported(Parameters.Platform, ERHIFeatureLevel::SM5) && DoesPlatformSupportDistanceFieldAO(Parameters.Platform) && IsUsingDistanceFields(Parameters.Platform);
 	}
 
 	static void ModifyCompilationEnvironment(const FGlobalShaderPermutationParameters& Parameters, FShaderCompilerEnvironment& OutEnvironment)
@@ -227,7 +227,7 @@ public:
 
 	static bool ShouldCompilePermutation(const FGlobalShaderPermutationParameters& Parameters)
 	{
-		return IsFeatureLevelSupported(Parameters.Platform, ERHIFeatureLevel::SM5) && DoesPlatformSupportDistanceFieldAO(Parameters.Platform);
+		return IsFeatureLevelSupported(Parameters.Platform, ERHIFeatureLevel::SM5) && DoesPlatformSupportDistanceFieldAO(Parameters.Platform) && IsUsingDistanceFields(Parameters.Platform);
 	}
 
 	static void ModifyCompilationEnvironment(const FGlobalShaderPermutationParameters& Parameters, FShaderCompilerEnvironment& OutEnvironment)
@@ -377,7 +377,7 @@ public:
 
 	static bool ShouldCompilePermutation(const FGlobalShaderPermutationParameters& Parameters)
 	{
-		return IsFeatureLevelSupported(Parameters.Platform, ERHIFeatureLevel::SM5) && DoesPlatformSupportDistanceFieldAO(Parameters.Platform);
+		return IsFeatureLevelSupported(Parameters.Platform, ERHIFeatureLevel::SM5) && DoesPlatformSupportDistanceFieldAO(Parameters.Platform) && IsUsingDistanceFields(Parameters.Platform);
 	}
 
 	static void ModifyCompilationEnvironment(const FGlobalShaderPermutationParameters& Parameters, FShaderCompilerEnvironment& OutEnvironment)
@@ -441,7 +441,7 @@ public:
 
 	static bool ShouldCompilePermutation(const FGlobalShaderPermutationParameters& Parameters)
 	{
-		return IsFeatureLevelSupported(Parameters.Platform, ERHIFeatureLevel::SM5) && DoesPlatformSupportDistanceFieldAO(Parameters.Platform);
+		return IsFeatureLevelSupported(Parameters.Platform, ERHIFeatureLevel::SM5) && DoesPlatformSupportDistanceFieldAO(Parameters.Platform) && IsUsingDistanceFields(Parameters.Platform);
 	}
 
 	static void ModifyCompilationEnvironment(const FGlobalShaderPermutationParameters& Parameters, FShaderCompilerEnvironment& OutEnvironment)
diff --git a/Engine/Source/Runtime/Renderer/Private/DistanceFieldScreenGridLighting.cpp b/Engine/Source/Runtime/Renderer/Private/DistanceFieldScreenGridLighting.cpp
--- a/Engine/Source/Runtime/Renderer/Private/DistanceFieldScreenGridLighting.cpp
+++ b/Engine/Source/Runtime/Renderer/Private/DistanceFieldScreenGridLighting.cpp
@@ -87,7 +87,7 @@ public:
 
 	static bool ShouldCompilePermutation(const FGlobalShaderPermutationParameters& Parameters)
 	{
-		return IsFeatureLevelSupported(Parameters.Platform, ERHIFeatureLevel::SM5) && DoesPlatformSupportDistanceFieldAO(Parameters.Platform);
+		return IsFeatureLevelSupported(Parameters.Platform, ERHIFeatureLevel::SM5) && DoesPlatformSupportDistanceFieldAO(Parameters.Platform) && IsUsingDistanceFields(Parameters.Platform);
 	}
 
 	static void ModifyCompilationEnvironment(const FGlobalShaderPermutationParameters& Parameters, FShaderCompilerEnvironment& OutEnvironment)
@@ -236,7 +236,7 @@ public:
 
 	static bool ShouldCompilePermutation(const FGlobalShaderPermutationParameters& Parameters)
 	{
-		return IsFeatureLevelSupported(Parameters.Platform, ERHIFeatureLevel::SM5) && DoesPlatformSupportDistanceFieldAO(Parameters.Platform);
+		return IsFeatureLevelSupported(Parameters.Platform, ERHIFeatureLevel::SM5) && DoesPlatformSupportDistanceFieldAO(Parameters.Platform) && IsUsingDistanceFields(Parameters.Platform);
 	}
 
 	static void ModifyCompilationEnvironment(const FGlobalShaderPermutationParameters& Parameters, FShaderCompilerEnvironment& OutEnvironment)
@@ -385,7 +385,7 @@ public:
 
 	static bool ShouldCompilePermutation(const FGlobalShaderPermutationParameters& Parameters)
 	{
-		return IsFeatureLevelSupported(Parameters.Platform, ERHIFeatureLevel::SM5) && DoesPlatformSupportDistanceFieldAO(Parameters.Platform);
+		return IsFeatureLevelSupported(Parameters.Platform, ERHIFeatureLevel::SM5) && DoesPlatformSupportDistanceFieldAO(Parameters.Platform) && IsUsingDistanceFields(Parameters.Platform);
 	}
 
 	static void ModifyCompilationEnvironment(const FGlobalShaderPermutationParameters& Parameters, FShaderCompilerEnvironment& OutEnvironment)
diff --git a/Engine/Source/Runtime/Renderer/Private/DistanceFieldVisualization.cpp b/Engine/Source/Runtime/Renderer/Private/DistanceFieldVisualization.cpp
--- a/Engine/Source/Runtime/Renderer/Private/DistanceFieldVisualization.cpp
+++ b/Engine/Source/Runtime/Renderer/Private/DistanceFieldVisualization.cpp
@@ -25,7 +25,7 @@ public:
 
 	static bool ShouldCompilePermutation(const FGlobalShaderPermutationParameters& Parameters)
 	{
-		return IsFeatureLevelSupported(Parameters.Platform, ERHIFeatureLevel::SM5) && DoesPlatformSupportDistanceFieldAO(Parameters.Platform);
+		return IsFeatureLevelSupported(Parameters.Platform, ERHIFeatureLevel::SM5) && DoesPlatformSupportDistanceFieldAO(Parameters.Platform) && IsUsingDistanceFields(Parameters.Platform);
 	}
 
 	static void ModifyCompilationEnvironment(const FGlobalShaderPermutationParameters& Parameters, FShaderCompilerEnvironment& OutEnvironment)
@@ -117,7 +117,7 @@ public:
 
 	static bool ShouldCompilePermutation(const FGlobalShaderPermutationParameters& Parameters)
 	{
-		return IsFeatureLevelSupported(Parameters.Platform, ERHIFeatureLevel::SM5) && DoesPlatformSupportDistanceFieldAO(Parameters.Platform);
+		return IsFeatureLevelSupported(Parameters.Platform, ERHIFeatureLevel::SM5) && DoesPlatformSupportDistanceFieldAO(Parameters.Platform) && IsUsingDistanceFields(Parameters.Platform);
 	}
 
 	static void ModifyCompilationEnvironment(const FGlobalShaderPermutationParameters& Parameters, FShaderCompilerEnvironment& OutEnvironment)
diff --git a/Engine/Source/Runtime/Renderer/Private/GlobalDistanceField.cpp b/Engine/Source/Runtime/Renderer/Private/GlobalDistanceField.cpp
--- a/Engine/Source/Runtime/Renderer/Private/GlobalDistanceField.cpp
+++ b/Engine/Source/Runtime/Renderer/Private/GlobalDistanceField.cpp
@@ -157,7 +157,7 @@ public:
 
 	static bool ShouldCompilePermutation(const FGlobalShaderPermutationParameters& Parameters)
 	{
-		return IsFeatureLevelSupported(Parameters.Platform, ERHIFeatureLevel::SM5) && DoesPlatformSupportDistanceFieldAO(Parameters.Platform);
+		return IsFeatureLevelSupported(Parameters.Platform, ERHIFeatureLevel::SM5) && DoesPlatformSupportDistanceFieldAO(Parameters.Platform) && IsUsingDistanceFields(Parameters.Platform);
 	}
 
 	static void ModifyCompilationEnvironment(const FGlobalShaderPermutationParameters& Parameters, FShaderCompilerEnvironment& OutEnvironment)
@@ -321,7 +321,7 @@ public:
 
 	static bool ShouldCompilePermutation(const FGlobalShaderPermutationParameters& Parameters)
 	{
-		return IsFeatureLevelSupported(Parameters.Platform, ERHIFeatureLevel::SM5) && DoesPlatformSupportDistanceFieldAO(Parameters.Platform);
+		return IsFeatureLevelSupported(Parameters.Platform, ERHIFeatureLevel::SM5) && DoesPlatformSupportDistanceFieldAO(Parameters.Platform) && IsUsingDistanceFields(Parameters.Platform);
 	}
 
 	static void ModifyCompilationEnvironment(const FGlobalShaderPermutationParameters& Parameters, FShaderCompilerEnvironment& OutEnvironment)
@@ -447,7 +447,7 @@ public:
 
 	static bool ShouldCompilePermutation(const FGlobalShaderPermutationParameters& Parameters)
 	{
-		return IsFeatureLevelSupported(Parameters.Platform, ERHIFeatureLevel::SM5) && DoesPlatformSupportDistanceFieldAO(Parameters.Platform);
+		return IsFeatureLevelSupported(Parameters.Platform, ERHIFeatureLevel::SM5) && DoesPlatformSupportDistanceFieldAO(Parameters.Platform) && IsUsingDistanceFields(Parameters.Platform);
 	}
 
 	static void ModifyCompilationEnvironment(const FGlobalShaderPermutationParameters& Parameters, FShaderCompilerEnvironment& OutEnvironment)
@@ -596,7 +596,7 @@ public:
 
 	static bool ShouldCompilePermutation(const FGlobalShaderPermutationParameters& Parameters)
 	{
-		return IsFeatureLevelSupported(Parameters.Platform, ERHIFeatureLevel::SM5) && DoesPlatformSupportDistanceFieldAO(Parameters.Platform) && !IsMetalPlatform(Parameters.Platform);
+		return IsFeatureLevelSupported(Parameters.Platform, ERHIFeatureLevel::SM5) && DoesPlatformSupportDistanceFieldAO(Parameters.Platform) && !IsMetalPlatform(Parameters.Platform) && IsUsingDistanceFields(Parameters.Platform);
 	}
 
 	static void ModifyCompilationEnvironment(const FGlobalShaderPermutationParameters& Parameters, FShaderCompilerEnvironment& OutEnvironment)
diff --git a/Engine/Source/Runtime/Renderer/Private/RendererScene.cpp b/Engine/Source/Runtime/Renderer/Private/RendererScene.cpp
--- a/Engine/Source/Runtime/Renderer/Private/RendererScene.cpp
+++ b/Engine/Source/Runtime/Renderer/Private/RendererScene.cpp
@@ -361,7 +361,7 @@ FDistanceFieldSceneData::FDistanceFieldSceneData(EShaderPlatform ShaderPlatform)
 {
 	static const auto CVar = IConsoleManager::Get().FindTConsoleVariableDataInt(TEXT("r.GenerateMeshDistanceFields"));
 
-	bTrackAllPrimitives = (DoesPlatformSupportDistanceFieldAO(ShaderPlatform) || DoesPlatformSupportDistanceFieldShadowing(ShaderPlatform)) && CVar->GetValueOnGameThread() != 0;
+	bTrackAllPrimitives = (DoesPlatformSupportDistanceFieldAO(ShaderPlatform) || DoesPlatformSupportDistanceFieldShadowing(ShaderPlatform)) && CVar->GetValueOnGameThread() != 0 && IsUsingDistanceFields(ShaderPlatform);
 
 	bCanUse16BitObjectIndices = RHISupportsBufferLoadTypeConversion(ShaderPlatform);
 }
