Force FFTBloom convolution texture to be resident so that Bloom functions properly under memory texture streaming pressure.

diff --git a/Engine/Source/Runtime/Renderer/Private/PostProcess/PostProcessFFTBloom.cpp b/Engine/Source/Runtime/Renderer/Private/PostProcess/PostProcessFFTBloom.cpp
--- a/Engine/Source/Runtime/Renderer/Private/PostProcess/PostProcessFFTBloom.cpp
+++ b/Engine/Source/Runtime/Renderer/Private/PostProcess/PostProcessFFTBloom.cpp
@@ -790,7 +790,7 @@ FSceneRenderTargetItem* FRCPassFFTBloom::InitDomainAndGetKernel(FRenderingCompos
 	// The kernel parameters on the FinalPostProcess.  
 
 	UTexture2D* BloomConvolutionTexture = PPSettings.BloomConvolutionTexture;
-	
+
 	if (BloomConvolutionTexture == nullptr)
 	{
 		BloomConvolutionTexture = GEngine->DefaultBloomKernelTexture;
@@ -1071,18 +1071,33 @@ bool FRCPassFFTBloom::HasValidPhysicalKernel(FPostprocessContext& Context)
 		BloomConvolutionTexture = GEngine->DefaultBloomKernelTexture;
 	}
 
+	// The bloom convolution kernel needs to be resident to avoid visual artifacts.
+	{
+		const int32 CinematicTextureGroups = 0;
+		const float Seconds = 5.0f;
+		BloomConvolutionTexture->SetForceMipLevelsToBeResident(Seconds, CinematicTextureGroups);
+	}
+
 	bool bValidSetup = (BloomConvolutionTexture != nullptr && BloomConvolutionTexture->Resource != nullptr);
 
+	const uint32 FramesPerWarning = 15;
+
 	if (bValidSetup && BloomConvolutionTexture->IsFullyStreamedIn() == false)
 	{
-		UE_LOG(LogRenderer, Warning, TEXT("The Physical Kernel Texture not fully streamed in."));
+		if ((View.Family->FrameNumber % FramesPerWarning) == 0)
+		{
+			UE_LOG(LogRenderer, Warning, TEXT("The Physical Kernel Texture not fully streamed in."));
+		}
 	}
 
 	bValidSetup = bValidSetup && (BloomConvolutionTexture->IsFullyStreamedIn() == true);
 
 	if (bValidSetup && BloomConvolutionTexture->bHasStreamingUpdatePending == true)
 	{
-		UE_LOG(LogRenderer, Warning, TEXT("The Physical Kernel Texture has pending update."));
+		if ((View.Family->FrameNumber % FramesPerWarning) == 0)
+		{
+			UE_LOG(LogRenderer, Warning, TEXT("The Physical Kernel Texture has pending update."));
+		}
 	}
 
 	bValidSetup = bValidSetup && (BloomConvolutionTexture->bHasStreamingUpdatePending == false);
