Fixed static analysis warning about utilizing pointer without checking for nullptr.

diff --git a/Engine/Source/Runtime/Renderer/Private/PostProcess/PostProcessFFTBloom.cpp b/Engine/Source/Runtime/Renderer/Private/PostProcess/PostProcessFFTBloom.cpp
--- a/Engine/Source/Runtime/Renderer/Private/PostProcess/PostProcessFFTBloom.cpp
+++ b/Engine/Source/Runtime/Renderer/Private/PostProcess/PostProcessFFTBloom.cpp
@@ -1071,15 +1071,16 @@ bool FRCPassFFTBloom::HasValidPhysicalKernel(FPostprocessContext& Context)
 		BloomConvolutionTexture = GEngine->DefaultBloomKernelTexture;
 	}
 
+	bool bValidSetup = (BloomConvolutionTexture != nullptr && BloomConvolutionTexture->Resource != nullptr);
+
 	// The bloom convolution kernel needs to be resident to avoid visual artifacts.
+	if (bValidSetup)
 	{
 		const int32 CinematicTextureGroups = 0;
 		const float Seconds = 5.0f;
 		BloomConvolutionTexture->SetForceMipLevelsToBeResident(Seconds, CinematicTextureGroups);
 	}
 
-	bool bValidSetup = (BloomConvolutionTexture != nullptr && BloomConvolutionTexture->Resource != nullptr);
-
 	const uint32 FramesPerWarning = 15;
 
 	if (bValidSetup && BloomConvolutionTexture->IsFullyStreamedIn() == false)
