Removing mutual exclusivity between ambient occlusion and RTGI.

This bypasses an RTGI artifact that produces poor ambient values.

diff --git a/Engine/Source/Runtime/Renderer/Private/DeferredShadingRenderer.cpp b/Engine/Source/Runtime/Renderer/Private/DeferredShadingRenderer.cpp
--- a/Engine/Source/Runtime/Renderer/Private/DeferredShadingRenderer.cpp
+++ b/Engine/Source/Runtime/Renderer/Private/DeferredShadingRenderer.cpp
@@ -1550,7 +1550,7 @@ void FDeferredShadingSceneRenderer::Render(FRHICommandListImmediate& RHICmdList)
 					RenderRayTracingGlobalIllumination(RHICmdList, Views[ViewIndex], GlobalIlluminationRT, SceneContext.ScreenSpaceAO);
 				}
 			}
-			else if (ShouldRenderRayTracingAmbientOcclusion())
+			if (ShouldRenderRayTracingAmbientOcclusion())
 			{
 				checkSlow(RHICmdList.IsOutsideRenderPass());
 				checkSlow(RHICmdList.IsOutsideRenderPass());
