SceneCapture2D - Adding a new property to allow you control the near clipping plane of the scene capture component.

diff --git a/Engine/Source/Runtime/Engine/Classes/Components/SceneCaptureComponent2D.h b/Engine/Source/Runtime/Engine/Classes/Components/SceneCaptureComponent2D.h
--- a/Engine/Source/Runtime/Engine/Classes/Components/SceneCaptureComponent2D.h
+++ b/Engine/Source/Runtime/Engine/Classes/Components/SceneCaptureComponent2D.h
@@ -20,6 +20,8 @@ UCLASS(hidecategories=(Collision, Object, Physics, SceneComponent), ClassGroup=R
 class ENGINE_API USceneCaptureComponent2D : public USceneCaptureComponent
 {
 	GENERATED_UCLASS_BODY()
+		
+public:
 
 	UPROPERTY(EditAnywhere, BlueprintReadWrite, Category=Projection, meta=(DisplayName = "Projection Type"))
 	TEnumAsByte<ECameraProjectionMode::Type> ProjectionType;
@@ -50,6 +52,15 @@ class ENGINE_API USceneCaptureComponent2D : public USceneCaptureComponent
 	UPROPERTY(interp, Category=PostProcessVolume, BlueprintReadWrite, meta=(UIMin = "0.0", UIMax = "1.0"))
 	float PostProcessBlendWeight;
 
+	UPROPERTY(EditAnywhere, Category = Projection, meta = (InlineEditConditionToggle))
+	uint32 bOverride_CustomNearClippingPlane : 1;
+
+	/** 
+	 * Set bOverride_CustomNearClippingPlane to true if you want to use a custom clipping plane instead of GNearClippingPlane.
+	 */
+	UPROPERTY(EditAnywhere, BlueprintReadWrite, AdvancedDisplay, Category = Projection, meta = (editcondition = "bOverride_CustomNearClippingPlane"))
+	float CustomNearClippingPlane = 0;
+
 	/** Whether a custom projection matrix will be used during rendering. Use with caution. Does not currently affect culling */
 	UPROPERTY(BlueprintReadWrite, AdvancedDisplay, Category = Projection)
 	bool bUseCustomProjectionMatrix;
diff --git a/Engine/Source/Runtime/Engine/Private/Components/SceneCaptureComponent.cpp b/Engine/Source/Runtime/Engine/Private/Components/SceneCaptureComponent.cpp
--- a/Engine/Source/Runtime/Engine/Private/Components/SceneCaptureComponent.cpp
+++ b/Engine/Source/Runtime/Engine/Private/Components/SceneCaptureComponent.cpp
@@ -592,10 +592,9 @@ void USceneCaptureComponent2D::UpdateDrawFrustum()
 			DrawFrustum->FrustumAngle = -OrthoWidth;
 		}
 
-		DrawFrustum->FrustumStartDist = GNearClippingPlane;
+		DrawFrustum->FrustumStartDist = (bOverride_CustomNearClippingPlane) ? CustomNearClippingPlane : GNearClippingPlane;
 		// 1000 is the default frustum distance, ideally this would be infinite but that might cause rendering issues
-		DrawFrustum->FrustumEndDist = (MaxViewDistanceOverride > DrawFrustum->FrustumStartDist)
-			? MaxViewDistanceOverride : 1000.0f;
+		DrawFrustum->FrustumEndDist = (MaxViewDistanceOverride > DrawFrustum->FrustumStartDist) ? MaxViewDistanceOverride : 1000.0f;
 		DrawFrustum->MarkRenderStateDirty();
 	}
 }
diff --git a/Engine/Source/Runtime/Renderer/Private/PlanarReflectionRendering.cpp b/Engine/Source/Runtime/Renderer/Private/PlanarReflectionRendering.cpp
--- a/Engine/Source/Runtime/Renderer/Private/PlanarReflectionRendering.cpp
+++ b/Engine/Source/Runtime/Renderer/Private/PlanarReflectionRendering.cpp
@@ -417,7 +417,7 @@ static void UpdatePlanarReflectionContents_RenderThread(
 	FSceneRenderer::WaitForTasksClearSnapshotsAndDeleteSceneRenderer(RHICmdList, SceneRenderer);
 }
 
-extern void BuildProjectionMatrix(FIntPoint RenderTargetSize, ECameraProjectionMode::Type ProjectionType, float FOV, float OrthoWidth, FMatrix& ProjectionMatrix);
+extern void BuildProjectionMatrix(FIntPoint RenderTargetSize, ECameraProjectionMode::Type ProjectionType, float FOV, float OrthoWidth, float InNearClippingPlane, FMatrix& ProjectionMatrix);
 
 extern void SetupViewVamilyForSceneCapture(
 	FSceneViewFamily& ViewFamily,
@@ -514,7 +514,7 @@ void FScene::UpdatePlanarReflectionContents(UPlanarReflectionComponent* CaptureC
 			const float HalfFOV = FMath::Atan(1.0f / View.ViewMatrices.GetProjectionMatrix().M[0][0]);
 
 			FMatrix ProjectionMatrix;
-			BuildProjectionMatrix(View.UnscaledViewRect.Size(), ECameraProjectionMode::Perspective, HalfFOV + FMath::DegreesToRadians(CaptureComponent->ExtraFOV), 1.0f, ProjectionMatrix);
+			BuildProjectionMatrix(View.UnscaledViewRect.Size(), ECameraProjectionMode::Perspective, HalfFOV + FMath::DegreesToRadians(CaptureComponent->ExtraFOV), 1.0f, GNearClippingPlane, ProjectionMatrix);
 
 			NewView.ViewLocation = ViewLocation;
 			NewView.ViewRotationMatrix = ViewRotationMatrix;
diff --git a/Engine/Source/Runtime/Renderer/Private/SceneCaptureRendering.cpp b/Engine/Source/Runtime/Renderer/Private/SceneCaptureRendering.cpp
--- a/Engine/Source/Runtime/Renderer/Private/SceneCaptureRendering.cpp
+++ b/Engine/Source/Runtime/Renderer/Private/SceneCaptureRendering.cpp
@@ -447,7 +447,7 @@ static void UpdateSceneCaptureContent_RenderThread(
 		}
 }
 
-void BuildProjectionMatrix(FIntPoint RenderTargetSize, ECameraProjectionMode::Type ProjectionType, float FOV, float InOrthoWidth, FMatrix& ProjectionMatrix)
+void BuildProjectionMatrix(FIntPoint RenderTargetSize, ECameraProjectionMode::Type ProjectionType, float FOV, float InOrthoWidth, float InNearClippingPlane, FMatrix& ProjectionMatrix)
 {
 	float const XAxisMultiplier = 1.0f;
 	float const YAxisMultiplier = RenderTargetSize.X / (float)RenderTargetSize.Y;
@@ -480,8 +480,8 @@ void BuildProjectionMatrix(FIntPoint RenderTargetSize, ECameraProjectionMode::Ty
 				FOV,
 				XAxisMultiplier,
 				YAxisMultiplier,
-				GNearClippingPlane,
-				GNearClippingPlane
+				InNearClippingPlane,
+				InNearClippingPlane
 				);
 		}
 		else
@@ -491,8 +491,8 @@ void BuildProjectionMatrix(FIntPoint RenderTargetSize, ECameraProjectionMode::Ty
 				FOV,
 				XAxisMultiplier,
 				YAxisMultiplier,
-				GNearClippingPlane,
-				GNearClippingPlane
+				InNearClippingPlane,
+				InNearClippingPlane
 				);
 		}
 	}
@@ -706,7 +706,8 @@ void FScene::UpdateSceneCaptureContents(USceneCaptureComponent2D* CaptureCompone
 		}
 		else
 		{
-			BuildProjectionMatrix(CaptureSize, CaptureComponent->ProjectionType, FOV, CaptureComponent->OrthoWidth, ProjectionMatrix);
+			const float ClippingPlane = (CaptureComponent->bOverride_CustomNearClippingPlane) ? CaptureComponent->CustomNearClippingPlane : GNearClippingPlane;
+			BuildProjectionMatrix(CaptureSize, CaptureComponent->ProjectionType, FOV, CaptureComponent->OrthoWidth, ClippingPlane, ProjectionMatrix);
 		}
 
 		const bool bUseSceneColorTexture = CaptureComponent->CaptureSource != SCS_FinalColorLDR;
@@ -870,7 +871,7 @@ void FScene::UpdateSceneCaptureContents(USceneCaptureComponentCube* CaptureCompo
 				const FMatrix ViewRotationMatrix = FLocal::CalcCubeFaceTransform(TargetFace);
 				FIntPoint CaptureSize(TextureTarget->GetSurfaceWidth(), TextureTarget->GetSurfaceHeight());
 				FMatrix ProjectionMatrix;
-				BuildProjectionMatrix(CaptureSize, ECameraProjectionMode::Perspective, FOV, 1.0f, ProjectionMatrix);
+				BuildProjectionMatrix(CaptureSize, ECameraProjectionMode::Perspective, FOV, 1.0f, GNearClippingPlane, ProjectionMatrix);
 				FPostProcessSettings PostProcessSettings;
 
 				float StereoIPD = 0.0f;
