Merge CL 6316011 from Dev-VR - Disable volumetric LUT on Lumin SM5 shader platform to work around a driver bug.
Also reverts CL 6298294 which turned out not to be an optimization issue, but rather a caching issue.

diff --git a/Engine/Shaders/Private/PostProcessCommon.ush b/Engine/Shaders/Private/PostProcessCommon.ush
--- a/Engine/Shaders/Private/PostProcessCommon.ush
+++ b/Engine/Shaders/Private/PostProcessCommon.ush
@@ -222,8 +222,7 @@ half4 UnwrappedTexture3DSample( Texture2D Texture, SamplerState Sampler, float3
 	float V = UVW.y;
 
 	half4 RG0 = Texture2DSample( Texture, Sampler, float2(U, V) );
-    U += 1.0f / Size;
-	half4 RG1 = Texture2DSample( Texture, Sampler, float2(U, V) );
+	half4 RG1 = Texture2DSample( Texture, Sampler, float2(U + 1.0f / Size, V) );
 
 	return lerp(RG0, RG1, FracW);
 }
diff --git a/Engine/Shaders/Public/ShaderVersion.ush b/Engine/Shaders/Public/ShaderVersion.ush
--- a/Engine/Shaders/Public/ShaderVersion.ush
+++ b/Engine/Shaders/Public/ShaderVersion.ush
@@ -3,4 +3,4 @@
 // in Platform.ush (which should be included in any shader) it allows to invalidate the shader DDC.
 // 
 // If you are merging streams and there is a conflict with this GUID you should make a new GUID rather than taking one or the other.
-// GUID = 0D98D1EF7B834118A251910A92A95B92
\ No newline at end of file
+// GUID = 8B5F808E3B5D477D92B82A44549A51E4
\ No newline at end of file
diff --git a/Engine/Source/Runtime/Renderer/Private/PostProcess/PostProcessCombineLUTs.cpp b/Engine/Source/Runtime/Renderer/Private/PostProcess/PostProcessCombineLUTs.cpp
--- a/Engine/Source/Runtime/Renderer/Private/PostProcess/PostProcessCombineLUTs.cpp
+++ b/Engine/Source/Runtime/Renderer/Private/PostProcess/PostProcessCombineLUTs.cpp
@@ -61,7 +61,7 @@ static TAutoConsoleVariable<int32> CVarMobileTonemapperFilm(
 // USE_VOLUME_LUT: needs to be the same for C++ and HLSL
 bool UseVolumeTextureLUT(EShaderPlatform Platform) 
 {
-	return (IsFeatureLevelSupported(Platform,ERHIFeatureLevel::SM4) && GSupportsVolumeTextureRendering && (Platform != SP_METAL_MRT && Platform != SP_METAL_MRT_MAC) && (RHISupportsGeometryShaders(Platform) || RHISupportsVertexShaderLayer(Platform)));
+	return (IsFeatureLevelSupported(Platform,ERHIFeatureLevel::SM4) && GSupportsVolumeTextureRendering && (Platform != SP_METAL_MRT && Platform != SP_METAL_MRT_MAC && Platform != SP_VULKAN_SM5_LUMIN) && (RHISupportsGeometryShaders(Platform) || RHISupportsVertexShaderLayer(Platform)));
 }
 
 // including the neutral one at index 0
