UE-74201: Unable to load maps when DX12 and Raytracing is enabled

There is an out-of-range index query when attempting to access ray tracing LOD of a particular set of static meshes. This fix prevents the crash, but an additional fix is necessary to understand why the ray tracing LOD selection is outside the static mesh LOD bounds. Captured in UE-74213.

diff --git a/Engine/Source/Runtime/Renderer/Private/DeferredShadingRenderer.cpp b/Engine/Source/Runtime/Renderer/Private/DeferredShadingRenderer.cpp
--- a/Engine/Source/Runtime/Renderer/Private/DeferredShadingRenderer.cpp
+++ b/Engine/Source/Runtime/Renderer/Private/DeferredShadingRenderer.cpp
@@ -656,7 +656,9 @@ bool FDeferredShadingSceneRenderer::GatherRayTracingWorldInstances(FRHICommandLi
 						bool bAllSegmentsOpaque = true;
 						bool bAnySegmentsCastShadow = false;
 
-						const auto& CachedRayTracingMeshCommandIndices = SceneInfo->CachedRayTracingMeshCommandIndicesPerLOD[LODToRender.GetRayTracedLOD()];
+						// UE-74201: prevents assigning an LOD index that is outside of the current range
+						uint32 LODIndex = FMath::Min((int)LODToRender.GetRayTracedLOD(), SceneInfo->CachedRayTracingMeshCommandIndicesPerLOD.Num() - 1);
+						const auto& CachedRayTracingMeshCommandIndices = SceneInfo->CachedRayTracingMeshCommandIndicesPerLOD[LODIndex];
 						for (auto CommandIndex : CachedRayTracingMeshCommandIndices)
 						{
 							if (CommandIndex >= 0)
