Fix CIS build error with FMeshDrawCommand::SetDebugData() not being exported

diff --git a/Engine/Source/Runtime/Renderer/Public/MeshPassProcessor.h b/Engine/Source/Runtime/Renderer/Public/MeshPassProcessor.h
--- a/Engine/Source/Runtime/Renderer/Public/MeshPassProcessor.h
+++ b/Engine/Source/Runtime/Renderer/Public/MeshPassProcessor.h
@@ -538,7 +538,7 @@ public:
 		return Other.CachedPipelineId.GetId();
 	}
 #if MESH_DRAW_COMMAND_DEBUG_DATA
-	void SetDebugData(const FPrimitiveSceneProxy* PrimitiveSceneProxy, const FMaterial* Material, const FMaterialRenderProxy* MaterialRenderProxy, const FMeshProcessorShaders& UntypedShaders);
+	RENDERER_API void SetDebugData(const FPrimitiveSceneProxy* PrimitiveSceneProxy, const FMaterial* Material, const FMaterialRenderProxy* MaterialRenderProxy, const FMeshProcessorShaders& UntypedShaders);
 #else
 	void SetDebugData(const FPrimitiveSceneProxy* PrimitiveSceneProxy, const FMaterial* Material, const FMaterialRenderProxy* MaterialRenderProxy, const FMeshProcessorShaders& UntypedShaders){}
 #endif
