Fixed rendering of static light sources inside reflection captures. Translucent view elements were always rendered with ETranslucencyPass::TPT_StandardTranslucency, causing TPT_TranslucencyAfterDOF elements to be filtered out.

diff --git a/Engine/Source/Runtime/Renderer/Private/TranslucentRendering.cpp b/Engine/Source/Runtime/Renderer/Private/TranslucentRendering.cpp
--- a/Engine/Source/Runtime/Renderer/Private/TranslucentRendering.cpp
+++ b/Engine/Source/Runtime/Renderer/Private/TranslucentRendering.cpp
@@ -380,7 +380,7 @@ void RenderViewTranslucencyInner(FRHICommandListImmediate& RHICmdList, const FVi
 				QUICK_SCOPE_CYCLE_COUNTER(RenderTranslucencyParallel_SDPG_World);
 
 				DrawDynamicMeshPass(View, RHICmdList,
-					[&View, &DrawRenderState](FDynamicPassMeshDrawListContext* DynamicMeshPassContext)
+					[&View, &DrawRenderState, TranslucencyPass](FDynamicPassMeshDrawListContext* DynamicMeshPassContext)
 				{
 					FBasePassMeshProcessor PassMeshProcessor(
 						View.Family->Scene->GetRenderScene(),
@@ -389,7 +389,7 @@ void RenderViewTranslucencyInner(FRHICommandListImmediate& RHICmdList, const FVi
 						DrawRenderState,
 						DynamicMeshPassContext,
 						FBasePassMeshProcessor::EFlags::CanUseDepthStencil,
-						ETranslucencyPass::TPT_StandardTranslucency);
+						TranslucencyPass);
 
 					const uint64 DefaultBatchElementMask = ~0ull;
 
@@ -405,7 +405,7 @@ void RenderViewTranslucencyInner(FRHICommandListImmediate& RHICmdList, const FVi
 				QUICK_SCOPE_CYCLE_COUNTER(RenderTranslucencyParallel_SDPG_Foreground);
 
 				DrawDynamicMeshPass(View, RHICmdList,
-					[&View, &DrawRenderState](FDynamicPassMeshDrawListContext* DynamicMeshPassContext)
+					[&View, &DrawRenderState, TranslucencyPass](FDynamicPassMeshDrawListContext* DynamicMeshPassContext)
 				{
 					FBasePassMeshProcessor PassMeshProcessor(
 						View.Family->Scene->GetRenderScene(),
@@ -414,7 +414,7 @@ void RenderViewTranslucencyInner(FRHICommandListImmediate& RHICmdList, const FVi
 						DrawRenderState,
 						DynamicMeshPassContext,
 						FBasePassMeshProcessor::EFlags::CanUseDepthStencil,
-						ETranslucencyPass::TPT_StandardTranslucency);
+						TranslucencyPass);
 
 					const uint64 DefaultBatchElementMask = ~0ull;
 
