Implement better hashing for MeshDrawCommandKeyFuncs to reduce hash collisions.
Fix map loading time regression in 4.22 (>30s to <1s in some cases).

#rnx

diff --git a/Engine/Source/Runtime/Core/Public/Misc/Crc.h b/Engine/Source/Runtime/Core/Public/Misc/Crc.h
--- a/Engine/Source/Runtime/Core/Public/Misc/Crc.h
+++ b/Engine/Source/Runtime/Core/Public/Misc/Crc.h
@@ -23,6 +23,13 @@ struct CORE_API FCrc
 	/** generates CRC hash of the memory area */
 	static uint32 MemCrc32( const void* Data, int32 Length, uint32 CRC=0 );
 
+	/** generates CRC hash of the element */
+	template <typename T>
+	static uint32 TypeCrc32( const T& Data, uint32 CRC=0 )
+	{
+		return MemCrc32(&Data, sizeof(T), CRC);
+	}
+
 	/** String CRC. */
 	template <typename CharType>
 	static typename TEnableIf<sizeof(CharType) != 1, uint32>::Type StrCrc32(const CharType* Data, uint32 CRC = 0)
diff --git a/Engine/Source/Runtime/RenderCore/Public/Shader.h b/Engine/Source/Runtime/RenderCore/Public/Shader.h
--- a/Engine/Source/Runtime/RenderCore/Public/Shader.h
+++ b/Engine/Source/Runtime/RenderCore/Public/Shader.h
@@ -225,6 +225,13 @@ public:
 		checkf(BufferIndex == InBufferIndex, TEXT("Tweak FShaderLooseParameterBufferInfo type sizes"));
 	}
 
+	inline uint32 GetHash() const
+	{
+		uint32 Hash = FCrc::TypeCrc32(BufferIndex, 0);
+		Hash = FCrc::TypeCrc32(BufferSize, Hash);
+		return Hash;
+	}
+
 	friend FArchive& operator<<(FArchive& Ar,FShaderLooseParameterBufferInfo& Info)
 	{
 		Ar << Info.BufferIndex;
@@ -258,6 +265,16 @@ public:
 		return Ar;
 	}
 
+	inline uint32 GetHash() const
+	{
+		uint32 Hash = 0;
+		for (const FShaderLooseParameterBufferInfo& ShaderLooseParameterBufferInfo: LooseParameterBuffers)
+		{
+			Hash = HashCombine(ShaderLooseParameterBufferInfo.GetHash(), Hash);
+		}
+		return Hash;
+	}
+
 	inline bool operator==(const FShaderParameterMapInfo& Rhs) const
 	{
 		return UniformBuffers == Rhs.UniformBuffers
diff --git a/Engine/Source/Runtime/Renderer/Private/MeshPassProcessor.cpp b/Engine/Source/Runtime/Renderer/Private/MeshPassProcessor.cpp
--- a/Engine/Source/Runtime/Renderer/Private/MeshPassProcessor.cpp
+++ b/Engine/Source/Runtime/Renderer/Private/MeshPassProcessor.cpp
@@ -818,7 +818,7 @@ bool FMeshDrawShaderBindings::MatchesForDynamicInstancing(const FMeshDrawShaderB
 
 		if (SingleShaderBindings.ParameterMapInfo.SRVs.Num() > 0 || SingleShaderBindings.ParameterMapInfo.LooseParameterBuffers.Num() > 0 || SingleShaderBindings.ParameterMapInfo.TextureSamplers.Num() > 0)
 		{
-			// Not implemented
+			// Not implemented. Note: this must match with GetDynamicInstancingHash.
 			return false;
 		}
 
@@ -841,7 +841,42 @@ bool FMeshDrawShaderBindings::MatchesForDynamicInstancing(const FMeshDrawShaderB
 	}
 
 	return true;
-} 
+}
+
+uint32 FMeshDrawShaderBindings::GetDynamicInstancingHash() const
+{
+	uint32 Hash = FCrc::TypeCrc32(Size, 0);
+
+	for (const FMeshDrawShaderBindingsLayout& MeshDrawShaderBindingsLayout: ShaderLayouts)
+	{
+		Hash = HashCombine(MeshDrawShaderBindingsLayout.GetHash(), Hash);
+	}
+
+	const uint8* ShaderBindingDataPtr = GetData();
+	for (int32 ShaderBindingsIndex = 0; ShaderBindingsIndex < ShaderLayouts.Num(); ShaderBindingsIndex++)
+	{
+		FReadOnlyMeshDrawSingleShaderBindings SingleShaderBindings(ShaderLayouts[ShaderBindingsIndex], ShaderBindingDataPtr);
+
+		if (SingleShaderBindings.ParameterMapInfo.SRVs.Num() > 0 || SingleShaderBindings.ParameterMapInfo.LooseParameterBuffers.Num() > 0 || SingleShaderBindings.ParameterMapInfo.TextureSamplers.Num() > 0)
+		{
+			// Since this is not implemented, we must return a unique hash to minimize hash collisions.
+			// Note: this must match with MatchesForDynamicInstancing.
+			Hash = PointerHash(this, Hash);
+			return Hash;
+		}
+
+		const FUniformBufferRHIParamRef* UniformBufferBindings = SingleShaderBindings.GetUniformBufferStart();
+		for (int32 UniformBufferIndex = 0; UniformBufferIndex < SingleShaderBindings.ParameterMapInfo.UniformBuffers.Num(); UniformBufferIndex++)
+		{
+			FUniformBufferRHIParamRef UniformBuffer = UniformBufferBindings[UniformBufferIndex];
+			Hash = PointerHash(UniformBuffer, Hash);
+		}
+
+		ShaderBindingDataPtr += ShaderLayouts[ShaderBindingsIndex].GetDataSizeBytes();
+	}
+
+	return Hash;
+}
 
 void FMeshDrawCommand::SubmitDraw(
 	const FMeshDrawCommand& RESTRICT MeshDrawCommand, 
diff --git a/Engine/Source/Runtime/Renderer/Private/ScenePrivate.h b/Engine/Source/Runtime/Renderer/Private/ScenePrivate.h
--- a/Engine/Source/Runtime/Renderer/Private/ScenePrivate.h
+++ b/Engine/Source/Runtime/Renderer/Private/ScenePrivate.h
@@ -2429,7 +2429,7 @@ struct MeshDrawCommandKeyFuncs : DefaultKeyFuncs<FMeshDrawCommandStateBucket,fal
 	/** Calculates a hash index for a key. */
 	static FORCEINLINE uint32 GetKeyHash(KeyInitType Key)
 	{
-		return PointerHash(Key.IndexBuffer, GetTypeHash(Key.CachedPipelineId.GetId()));
+		return Key.GetDynamicInstancingHash();
 	}
 };
 
diff --git a/Engine/Source/Runtime/Renderer/Public/MeshDrawShaderBindings.h b/Engine/Source/Runtime/Renderer/Public/MeshDrawShaderBindings.h
--- a/Engine/Source/Runtime/Renderer/Public/MeshDrawShaderBindings.h
+++ b/Engine/Source/Runtime/Renderer/Public/MeshDrawShaderBindings.h
@@ -24,6 +24,13 @@ public:
 		Frequency = (EShaderFrequency)Shader->GetTarget().Frequency;
 		checkSlow((EShaderFrequency)Frequency == (EShaderFrequency)Shader->GetTarget().Frequency);
 	}
+	
+	inline uint32 GetHash() const
+	{
+		uint32 LocalFrequency = Frequency;
+		uint32 Hash = FCrc::TypeCrc32(LocalFrequency, ParameterMapInfo.GetHash());
+		return Hash;
+	}
 
 	bool operator==(const FMeshDrawShaderBindingsLayout& Rhs) const
 	{
diff --git a/Engine/Source/Runtime/Renderer/Public/MeshPassProcessor.h b/Engine/Source/Runtime/Renderer/Public/MeshPassProcessor.h
--- a/Engine/Source/Runtime/Renderer/Public/MeshPassProcessor.h
+++ b/Engine/Source/Runtime/Renderer/Public/MeshPassProcessor.h
@@ -352,6 +352,8 @@ public:
 	/** Returns whether this set of shader bindings can be merged into an instanced draw call with another. */
 	bool MatchesForDynamicInstancing(const FMeshDrawShaderBindings& Rhs) const;
 
+	uint32 GetDynamicInstancingHash() const;
+
 	SIZE_T GetAllocatedSize() const
 	{
 		SIZE_T Bytes = ShaderLayouts.GetAllocatedSize();
@@ -501,6 +503,41 @@ public:
 				|| (NumPrimitives == 0 && IndirectArgsBuffer == Rhs.IndirectArgsBuffer));
 	}
 
+	uint32 GetDynamicInstancingHash() const
+	{
+		uint32 Hash = FCrc::TypeCrc32(CachedPipelineId.GetId(), 0);
+		Hash = FCrc::TypeCrc32(StencilRef, Hash);
+		Hash = HashCombine(ShaderBindings.GetDynamicInstancingHash(), Hash);
+
+		for (const FVertexInputStream& VertexInputStream: VertexStreams)
+		{
+			const uint32 StreamIndex = VertexInputStream.StreamIndex;
+			const uint32 Offset = VertexInputStream.Offset;
+
+			Hash = FCrc::TypeCrc32(StreamIndex, Hash);
+			Hash = FCrc::TypeCrc32(Offset, Hash);
+			Hash = PointerHash(VertexInputStream.VertexBuffer, Hash);
+		}
+
+		Hash = FCrc::TypeCrc32(PrimitiveIdStreamIndex, Hash);
+		Hash = PointerHash(IndexBuffer, Hash);
+		Hash = FCrc::TypeCrc32(FirstIndex, Hash);
+		Hash = FCrc::TypeCrc32(NumPrimitives, Hash);
+		Hash = FCrc::TypeCrc32(NumInstances, Hash);
+
+		if (NumPrimitives > 0)
+		{
+			Hash = FCrc::TypeCrc32(VertexParams.BaseVertexIndex, Hash);
+			Hash = FCrc::TypeCrc32(VertexParams.NumVertices, Hash);
+		}
+		else
+		{
+			Hash = PointerHash(IndirectArgsBuffer, Hash);
+		}		
+
+		return Hash;
+	}
+
 	/** Sets shaders on the mesh draw command and allocates room for the shader bindings. */
 	RENDERER_API void SetShaders(FVertexDeclarationRHIParamRef VertexDeclaration, const FMeshProcessorShaders& Shaders, FGraphicsMinimalPipelineStateInitializer& PipelineState);
 
