Copy from 6063061
Fixes crash where readback resources are not cleaned up on the render thread.
#rnx

diff --git a/Engine/Source/Runtime/Renderer/Private/RendererScene.cpp b/Engine/Source/Runtime/Renderer/Private/RendererScene.cpp
--- a/Engine/Source/Runtime/Renderer/Private/RendererScene.cpp
+++ b/Engine/Source/Runtime/Renderer/Private/RendererScene.cpp
@@ -259,7 +259,11 @@ FSceneViewState::~FSceneViewState()
 	DestroyRWBuffer(VarianceMipTree);
 	DestroyRWBuffer(TotalRayCountBuffer);
 
-	delete RayCountGPUReadback;
+	ENQUEUE_RENDER_COMMAND(FDeleteGpuReadback)(
+		[DeleteMe = RayCountGPUReadback](FRHICommandList&)
+	{
+		delete DeleteMe;
+	});
 #endif // RHI_RAYTRACING
 }
 
