Cherry-pick 6411466: Proper fix for UE-74213: GatherRayTracingWorldInstances() accesses LOD indexes outside the valid range

diff --git a/Engine/Source/Runtime/Renderer/Private/DeferredShadingRenderer.cpp b/Engine/Source/Runtime/Renderer/Private/DeferredShadingRenderer.cpp
--- a/Engine/Source/Runtime/Renderer/Private/DeferredShadingRenderer.cpp
+++ b/Engine/Source/Runtime/Renderer/Private/DeferredShadingRenderer.cpp
@@ -656,41 +656,45 @@ bool FDeferredShadingSceneRenderer::GatherRayTracingWorldInstances(FRHICommandLi
 						bool bAllSegmentsOpaque = true;
 						bool bAnySegmentsCastShadow = false;
 
-						// UE-74201: prevents assigning an LOD index that is outside of the current range
-						uint32 LODIndex = FMath::Min((int)LODToRender.GetRayTracedLOD(), SceneInfo->CachedRayTracingMeshCommandIndicesPerLOD.Num() - 1);
-						const auto& CachedRayTracingMeshCommandIndices = SceneInfo->CachedRayTracingMeshCommandIndicesPerLOD[LODIndex];
-						for (auto CommandIndex : CachedRayTracingMeshCommandIndices)
+						uint32 LODIndex = LODToRender.GetRayTracedLOD();
+						// Sometimes LODIndex is out of range because it is clamped by ClampToFirstLOD, like the requested LOD is being streamed in and hasn't been available
+						// According to InitViews, we should hide the static mesh instance
+						if (SceneInfo->CachedRayTracingMeshCommandIndicesPerLOD.IsValidIndex(LODIndex))
 						{
-							if (CommandIndex >= 0)
+							const auto& CachedRayTracingMeshCommandIndices = SceneInfo->CachedRayTracingMeshCommandIndicesPerLOD[LODIndex];
+							for (auto CommandIndex : CachedRayTracingMeshCommandIndices)
 							{
-								FVisibleRayTracingMeshCommand NewVisibleMeshCommand;
-
-								NewVisibleMeshCommand.RayTracingMeshCommand = &Scene->CachedRayTracingMeshCommands.RayTracingMeshCommands[CommandIndex];
-								NewVisibleMeshCommand.InstanceIndex = NewInstanceIndex;
-
-								View.VisibleRayTracingMeshCommands.Add(NewVisibleMeshCommand);
-								VisibleDrawCommandStartOffset[ViewIndex]++;
-
-								NewInstanceMask |= NewVisibleMeshCommand.RayTracingMeshCommand->InstanceMask;
-								bAllSegmentsOpaque &= NewVisibleMeshCommand.RayTracingMeshCommand->bOpaque;
-								bAnySegmentsCastShadow |= NewVisibleMeshCommand.RayTracingMeshCommand->bCastRayTracedShadows;
-							}
-							else
-							{
-								// CommandIndex == -1 indicates that the mesh batch has been filtered by FRayTracingMeshProcessor (like the shadow depth pass batch)
-								// Do nothing in this case
+								if (CommandIndex >= 0)
+								{
+									FVisibleRayTracingMeshCommand NewVisibleMeshCommand;
+
+									NewVisibleMeshCommand.RayTracingMeshCommand = &Scene->CachedRayTracingMeshCommands.RayTracingMeshCommands[CommandIndex];
+									NewVisibleMeshCommand.InstanceIndex = NewInstanceIndex;
+
+									View.VisibleRayTracingMeshCommands.Add(NewVisibleMeshCommand);
+									VisibleDrawCommandStartOffset[ViewIndex]++;
+
+									NewInstanceMask |= NewVisibleMeshCommand.RayTracingMeshCommand->InstanceMask;
+									bAllSegmentsOpaque &= NewVisibleMeshCommand.RayTracingMeshCommand->bOpaque;
+									bAnySegmentsCastShadow |= NewVisibleMeshCommand.RayTracingMeshCommand->bCastRayTracedShadows;
+								}
+								else
+								{
+									// CommandIndex == -1 indicates that the mesh batch has been filtered by FRayTracingMeshProcessor (like the shadow depth pass batch)
+									// Do nothing in this case
+								}
 							}
-						}
 
-						NewInstanceMask |= bAnySegmentsCastShadow ? RAY_TRACING_MASK_SHADOW : 0;
+							NewInstanceMask |= bAnySegmentsCastShadow ? RAY_TRACING_MASK_SHADOW : 0;
 
-						// When no cached command is found, NewInstanceMask == 0 and the instance is effectively filtered out
-						FRayTracingGeometryInstance RayTracingInstance = { RayTracingGeometryInstance };
-						RayTracingInstance.Transform = Scene->PrimitiveTransforms[PrimitiveIndex];
-						RayTracingInstance.UserData = (uint32)PrimitiveIndex;
-						RayTracingInstance.Mask = NewInstanceMask;
-						RayTracingInstance.bForceOpaque = bAllSegmentsOpaque;
-						View.RayTracingGeometryInstances.Add(RayTracingInstance);
+							// When no cached command is found, NewInstanceMask == 0 and the instance is effectively filtered out
+							FRayTracingGeometryInstance RayTracingInstance = { RayTracingGeometryInstance };
+							RayTracingInstance.Transform = Scene->PrimitiveTransforms[PrimitiveIndex];
+							RayTracingInstance.UserData = (uint32)PrimitiveIndex;
+							RayTracingInstance.Mask = NewInstanceMask;
+							RayTracingInstance.bForceOpaque = bAllSegmentsOpaque;
+							View.RayTracingGeometryInstances.Add(RayTracingInstance);
+						}
 					}
 					else if (View.Family->EngineShowFlags.SkeletalMeshes)
 					{
