Disable writing alpha from tonemapper if AA is disabled. This will leave the alpha channel with a value of 1.

diff --git a/Engine/Source/Runtime/Renderer/Private/PostProcess/PostProcessTonemap.cpp b/Engine/Source/Runtime/Renderer/Private/PostProcess/PostProcessTonemap.cpp
--- a/Engine/Source/Runtime/Renderer/Private/PostProcess/PostProcessTonemap.cpp
+++ b/Engine/Source/Runtime/Renderer/Private/PostProcess/PostProcessTonemap.cpp
@@ -1107,7 +1107,15 @@ void FRCPassPostProcessTonemap::Process(FRenderingCompositePassContext& Context)
 			{
 				FGraphicsPipelineStateInitializer GraphicsPSOInit;
 				Context.RHICmdList.ApplyCachedRenderTargets(GraphicsPSOInit);
-				GraphicsPSOInit.BlendState = TStaticBlendState<>::GetRHI();
+				if (Context.View.AntiAliasingMethod == AAM_None)
+				{
+					// Skip writing to alpha when AA is disabled. Luminance in the alpha channel is unused in this case. 
+					GraphicsPSOInit.BlendState = TStaticBlendStateWriteMask<CW_RGB>::GetRHI();
+				}
+				else
+				{
+					GraphicsPSOInit.BlendState = TStaticBlendState<>::GetRHI();
+				}
 				GraphicsPSOInit.RasterizerState = TStaticRasterizerState<>::GetRHI();
 				GraphicsPSOInit.DepthStencilState = TStaticDepthStencilState<false, CF_Always>::GetRHI();
 
