Disabled player outlines effect on mobile to save mem

diff --git a/Engine/Source/Runtime/Renderer/Private/MobileShadingRenderer.cpp b/Engine/Source/Runtime/Renderer/Private/MobileShadingRenderer.cpp
--- a/Engine/Source/Runtime/Renderer/Private/MobileShadingRenderer.cpp
+++ b/Engine/Source/Runtime/Renderer/Private/MobileShadingRenderer.cpp
@@ -43,6 +43,7 @@
 #include "VisualizeTexturePresent.h"
 #include "RendererModule.h"
 #include "EngineModule.h"
+#include "MaterialSceneTextureId.h"
 
 #include "VisualizeTexture.h"
 
@@ -82,6 +83,35 @@ FGlobalDynamicIndexBuffer FMobileSceneRenderer::DynamicIndexBuffer;
 FGlobalDynamicVertexBuffer FMobileSceneRenderer::DynamicVertexBuffer;
 TGlobalResource<FGlobalDynamicReadBuffer> FMobileSceneRenderer::DynamicReadBuffer;
 
+static bool UsesCustomDepthStencilLookup(const FViewInfo& View)
+{
+	// Find out whether post-process materials use CustomDepth/Stencil lookups
+	bool bPPUsesCustomDepth = false;
+	bool bPPUsesCustomStencil = false;
+	const FBlendableManager& BlendableManager = View.FinalPostProcessSettings.BlendableManager;
+	FBlendableEntry* BlendableIt = nullptr;
+
+	while (FPostProcessMaterialNode* DataPtr = BlendableManager.IterateBlendables<FPostProcessMaterialNode>(BlendableIt))
+	{
+		if (DataPtr->IsValid())
+		{
+			FMaterialRenderProxy* Proxy = DataPtr->GetMaterialInterface()->GetRenderProxy();
+			check(Proxy);
+
+			const FMaterial* Material = Proxy->GetMaterial(View.GetFeatureLevel());
+			check(Material);
+			const FMaterialShaderMap* MaterialShaderMap = Material->GetRenderingThreadShaderMap();
+
+			bPPUsesCustomDepth|= MaterialShaderMap->UsesSceneTexture(PPI_CustomDepth);
+			bPPUsesCustomStencil|= MaterialShaderMap->UsesSceneTexture(PPI_CustomStencil);
+		}
+	}
+
+	//TODO: check if translucency uses CustomDepth
+	return bPPUsesCustomDepth || bPPUsesCustomStencil;
+}
+
+
 FMobileSceneRenderer::FMobileSceneRenderer(const FSceneViewFamily* InViewFamily,FHitProxyConsumer* HitProxyConsumer)
 	:	FSceneRenderer(InViewFamily, HitProxyConsumer)
 {
@@ -337,7 +367,17 @@ void FMobileSceneRenderer::Render(FRHICommandListImmediate& RHICmdList)
 	// Custom depth
 	if (!bGammaSpace)
 	{
-		RenderCustomDepthPass(RHICmdList);
+		// see if anything uses CustomDepth, and skip custom depth rendering if possible
+		bool bUsesCustomDepthStencil = false;
+		for (int32 ViewIndex = 0; ViewIndex < Views.Num(); ViewIndex++) 
+		{
+			bUsesCustomDepthStencil = UsesCustomDepthStencilLookup(Views[ViewIndex]);
+		}
+
+		if (bUsesCustomDepthStencil)
+		{
+			RenderCustomDepthPass(RHICmdList);
+		}
 	}
 		
 	const FIntPoint RenderTargetSize = (ViewFamily.RenderTarget->GetRenderTargetTexture().IsValid()) ? ViewFamily.RenderTarget->GetRenderTargetTexture()->GetSizeXY() : ViewFamily.RenderTarget->GetSizeXY();
