Alpha channel writing in tonemapper is now only enabled for FXAA

#rnx

diff --git a/Engine/Source/Runtime/Renderer/Private/PostProcess/PostProcessTonemap.cpp b/Engine/Source/Runtime/Renderer/Private/PostProcess/PostProcessTonemap.cpp
--- a/Engine/Source/Runtime/Renderer/Private/PostProcess/PostProcessTonemap.cpp
+++ b/Engine/Source/Runtime/Renderer/Private/PostProcess/PostProcessTonemap.cpp
@@ -1107,14 +1107,15 @@ void FRCPassPostProcessTonemap::Process(FRenderingCompositePassContext& Context)
 			{
 				FGraphicsPipelineStateInitializer GraphicsPSOInit;
 				Context.RHICmdList.ApplyCachedRenderTargets(GraphicsPSOInit);
-				if (Context.View.AntiAliasingMethod == AAM_None)
+				if (Context.View.AntiAliasingMethod == AAM_FXAA)
 				{
-					// Skip writing to alpha when AA is disabled. Luminance in the alpha channel is unused in this case. 
-					GraphicsPSOInit.BlendState = TStaticBlendStateWriteMask<CW_RGB>::GetRHI();
+					// Alpha is used to store luminance for post process AA
+					GraphicsPSOInit.BlendState = TStaticBlendState<>::GetRHI();
 				}
 				else
 				{
-					GraphicsPSOInit.BlendState = TStaticBlendState<>::GetRHI();
+					// Skip writing to alpha, luminance in the alpha channel is unused in this case. 
+					GraphicsPSOInit.BlendState = TStaticBlendStateWriteMask<CW_RGB>::GetRHI();
 				}
 				GraphicsPSOInit.RasterizerState = TStaticRasterizerState<>::GetRHI();
 				GraphicsPSOInit.DepthStencilState = TStaticDepthStencilState<false, CF_Always>::GetRHI();
