Fix inverted scope rendering on android GLES.
Do not flip  Y axis on any pass that can be used as input for post process material.
flip Y axis only on final tonemapper or after tonemapper post process material stage.

diff --git a/Engine/Shaders/Private/MaterialTemplate.ush b/Engine/Shaders/Private/MaterialTemplate.ush
--- a/Engine/Shaders/Private/MaterialTemplate.ush
+++ b/Engine/Shaders/Private/MaterialTemplate.ush
@@ -1412,11 +1412,6 @@ MaterialFloat4 MobileSceneTextureLookup(inout FMaterialPixelParameters Parameter
 	else if (SceneTextureId == 14)
 	{
 #if POST_PROCESS_MATERIAL	
-		// SceneColor was transformed to OpenGL coordinates during tonemapping
-		#if !POST_PROCESS_MATERIAL_BEFORE_TONEMAP && (COMPILER_GLSL_ES2 || COMPILER_GLSL_ES3_1)
-		UV.y = 1.0 - UV.y; 
-		#endif
-
 		MaterialFloat4 Input0 = Texture2DSample(PostprocessInput0, PostprocessInput0Sampler, UV);
 		#if POST_PROCESS_MATERIAL_BEFORE_TONEMAP
 		Input0 = Decode32BPPHDR(Input0);
diff --git a/Engine/Shaders/Private/PostProcessMaterialShaders.usf b/Engine/Shaders/Private/PostProcessMaterialShaders.usf
--- a/Engine/Shaders/Private/PostProcessMaterialShaders.usf
+++ b/Engine/Shaders/Private/PostProcessMaterialShaders.usf
@@ -184,6 +184,11 @@ void MainPS_VideoOverlay(
 #endif
 
 #else // FEATURE_LEVEL_ES3_1
+
+#ifndef NEEDTOSWITCHVERTICLEAXIS
+#define NEEDTOSWITCHVERTICLEAXIS 0
+#endif
+
 //
 // Mobile version
 //
@@ -194,16 +199,11 @@ void MainVS_ES2(
 	out float2 OutViewportUV : TEXCOORD8
 	)
 {
-	// flip UV.y if we rendering to backbuffer in OpenGL space
-#if !POST_PROCESS_MATERIAL_BEFORE_TONEMAP && (COMPILER_GLSL_ES2 || COMPILER_GLSL_ES3_1)
-	InTexCoord.y = 1.0 - InTexCoord.y;
-#endif
-	
 	Output = (FPostProcessMaterialVSToPS)0;
 	float2 OutUV;
 	DrawRectangle(InPosition, InTexCoord, Output.Position, OutUV);
 
-	OutViewportUV = OutUV;
+	OutViewportUV = InPosition.xy;
 	
 #if NUM_TEX_COORD_INTERPOLATORS
 	FMaterialVertexParameters VertexParameters = GetPostProcessMaterialVSParameters(OutUV);
@@ -250,14 +250,13 @@ void MainPS_ES2(
 	Parameters.VertexColor = 1;
 	
 	float4 SvPosition = Input.Position;
-	float4 SvPositionMP = SvPosition;
-	#if !POST_PROCESS_MATERIAL_BEFORE_TONEMAP && (COMPILER_GLSL_ES2 || COMPILER_GLSL_ES3_1)
-		SvPositionMP.y = View.ViewSizeAndInvSize.y - SvPositionMP.y; 
-	#endif
-	SvPositionMP.z = LookupDeviceZ(InViewportUV);
-	SvPositionMP.z = max(SvPositionMP.z, 1e-18);
+#if NEEDTOSWITCHVERTICLEAXIS
+		SvPosition.y = View.ViewSizeAndInvSize.y - SvPosition.y; 
+#endif
+	SvPosition.z = LookupDeviceZ(InViewportUV);
+	SvPosition.z = max(SvPosition.z, 1e-18);
 	// fill out other related material parameters
-	CalcMaterialParametersPost(Parameters, PixelMaterialInputs, SvPositionMP, true);
+	CalcMaterialParametersPost(Parameters, PixelMaterialInputs, SvPosition, true);
 
 	// Grab emissive colour as output
 	half3 EmissiveColor = GetMaterialEmissive(PixelMaterialInputs);
diff --git a/Engine/Shaders/Private/PostProcessTonemap.usf b/Engine/Shaders/Private/PostProcessTonemap.usf
--- a/Engine/Shaders/Private/PostProcessTonemap.usf
+++ b/Engine/Shaders/Private/PostProcessTonemap.usf
@@ -106,7 +106,9 @@ float PseudoRandomTex(float2 xy)
 
 float SwitchVerticalAxis;
 
-#define NEEDTOSWITCHVERTICLEAXIS (ES2_PROFILE && COMPILER_GLSL_ES2) || (ES3_1_PROFILE && COMPILER_GLSL_ES3_1)
+#ifndef NEEDTOSWITCHVERTICLEAXIS
+#define NEEDTOSWITCHVERTICLEAXIS 0
+#endif
 
 // can be optimized
 float2 ScreenPosToUV(float2 ScreenPos)
diff --git a/Engine/Source/Runtime/RHI/Public/RHIDefinitions.h b/Engine/Source/Runtime/RHI/Public/RHIDefinitions.h
--- a/Engine/Source/Runtime/RHI/Public/RHIDefinitions.h
+++ b/Engine/Source/Runtime/RHI/Public/RHIDefinitions.h
@@ -1138,7 +1138,8 @@ inline bool RHINeedsToSwitchVerticalAxis(EShaderPlatform Platform)
 	static const auto CVar = IConsoleManager::Get().FindTConsoleVariableDataInt(TEXT("r.Mobile.ForceRHISwitchVerticalAxis"));
 	if (CVar->GetValueOnAnyThread())
 	{
-		return true;
+		// only allow this for mobile preview.
+		return IsMobilePlatform(Platform);
 	}
 #endif
 
diff --git a/Engine/Source/Runtime/Renderer/Private/PostProcess/PostProcessMaterial.cpp b/Engine/Source/Runtime/Renderer/Private/PostProcess/PostProcessMaterial.cpp
--- a/Engine/Source/Runtime/Renderer/Private/PostProcess/PostProcessMaterial.cpp
+++ b/Engine/Source/Runtime/Renderer/Private/PostProcess/PostProcessMaterial.cpp
@@ -16,6 +16,7 @@
 #include "ClearQuad.h"
 #include "Materials/MaterialExpressionSceneTexture.h"
 #include "RHI/Public/PipelineStateCache.h"
+#include "PostProcessMobile.h"
 
 static TAutoConsoleVariable<int32> CVarPostProcessAllowStencilTest(
 	TEXT("r.PostProcessAllowStencilTest"),
@@ -56,7 +57,7 @@ static bool ShouldCachePostProcessMaterial(EPostProcessMaterialTarget MaterialTa
 	return false;
 }
 
-template<EPostProcessMaterialTarget MaterialTarget>
+template<EPostProcessMaterialTarget MaterialTarget, bool bSwitchVerticalAxis>
 class FPostProcessMaterialVS : public FMaterialShader
 {
 	DECLARE_SHADER_TYPE(FPostProcessMaterialVS, Material);
@@ -66,8 +67,10 @@ public:
 	  * Only compile these shaders for post processing domain materials
 	  */
 	static bool ShouldCompilePermutation(EShaderPlatform Platform, const FMaterial* Material)
-	{
-		return ShouldCachePostProcessMaterial(MaterialTarget, Platform, Material);
+	{		
+		return ShouldCachePostProcessMaterial(MaterialTarget, Platform, Material)
+			// compile mobile axis switched versions only for after tonemapper passes
+			&& (bSwitchVerticalAxis == false || (RHINeedsToSwitchVerticalAxis(Platform) && Material->GetBlendableLocation() == BL_AfterTonemapping));
 	}
 
 	static void ModifyCompilationEnvironment(EShaderPlatform Platform, const class FMaterial* Material, FShaderCompilerEnvironment& OutEnvironment)
@@ -75,6 +78,7 @@ public:
 		FMaterialShader::ModifyCompilationEnvironment(Platform, OutEnvironment);
 
 		OutEnvironment.SetDefine(TEXT("POST_PROCESS_MATERIAL"), 1);
+		OutEnvironment.SetDefine(TEXT("NEEDTOSWITCHVERTICLEAXIS"), bSwitchVerticalAxis);
 
 		if (MaterialTarget == EPostProcessMaterialTarget::Mobile)
 		{
@@ -108,16 +112,18 @@ private:
 	FPostProcessPassParameters PostprocessParameter;
 };
 
-typedef FPostProcessMaterialVS<EPostProcessMaterialTarget::HighEnd> FPostProcessMaterialVS_HighEnd;
-typedef FPostProcessMaterialVS<EPostProcessMaterialTarget::Mobile> FPostProcessMaterialVS_Mobile;
+typedef FPostProcessMaterialVS<EPostProcessMaterialTarget::HighEnd, false> FPostProcessMaterialVS_HighEnd;
+typedef FPostProcessMaterialVS<EPostProcessMaterialTarget::Mobile, false> FPostProcessMaterialVS_Mobile;
+typedef FPostProcessMaterialVS<EPostProcessMaterialTarget::Mobile, true> FPostProcessMaterialVS_Mobile_AxisSwitch;
 
 IMPLEMENT_MATERIAL_SHADER_TYPE(template<>,FPostProcessMaterialVS_HighEnd,TEXT("/Engine/Private/PostProcessMaterialShaders.usf"),TEXT("MainVS"),SF_Vertex);
 IMPLEMENT_MATERIAL_SHADER_TYPE(template<>,FPostProcessMaterialVS_Mobile,TEXT("/Engine/Private/PostProcessMaterialShaders.usf"),TEXT("MainVS_ES2"),SF_Vertex);
+IMPLEMENT_MATERIAL_SHADER_TYPE(template<>,FPostProcessMaterialVS_Mobile_AxisSwitch,TEXT("/Engine/Private/PostProcessMaterialShaders.usf"),TEXT("MainVS_ES2"),SF_Vertex);
 
 /**
  * A pixel shader for rendering a post process material
  */
-template<EPostProcessMaterialTarget MaterialTarget, uint32 UVPolicy>
+template<EPostProcessMaterialTarget MaterialTarget, uint32 UVPolicy, bool bSwitchVerticalAxis>
 class FPostProcessMaterialPS : public FMaterialShader
 {
 	DECLARE_SHADER_TYPE(FPostProcessMaterialPS,Material);
@@ -128,7 +134,9 @@ public:
 	  */
 	static bool ShouldCompilePermutation(EShaderPlatform Platform, const FMaterial* Material)
 	{
-		return ShouldCachePostProcessMaterial(MaterialTarget, Platform, Material);
+		return ShouldCachePostProcessMaterial(MaterialTarget, Platform, Material)
+			// compile mobile axis switched versions only for after tonemapper passes
+			&& (bSwitchVerticalAxis == false || (RHINeedsToSwitchVerticalAxis(Platform) && Material->GetBlendableLocation() == BL_AfterTonemapping));
 	}
 
 	static void ModifyCompilationEnvironment(EShaderPlatform Platform, const class FMaterial* Material, FShaderCompilerEnvironment& OutEnvironment)
@@ -141,6 +149,8 @@ public:
 		EBlendableLocation Location = EBlendableLocation(Material->GetBlendableLocation());
 		OutEnvironment.SetDefine(TEXT("POST_PROCESS_MATERIAL_AFTER_TAA_UPSAMPLE"), (Location == BL_AfterTonemapping || Location == BL_ReplacingTonemapper) ? 1 : 0);
 
+		OutEnvironment.SetDefine(TEXT("NEEDTOSWITCHVERTICLEAXIS"), bSwitchVerticalAxis);
+
 		if (MaterialTarget == EPostProcessMaterialTarget::Mobile)
 		{
 			OutEnvironment.SetDefine(TEXT("MOBILE_FORCE_DEPTH_TEXTURE_READS"), 1); // Ensure post process materials will not attempt depth buffer fetch operations.
@@ -175,15 +185,19 @@ private:
 	FPostProcessPassParameters PostprocessParameter;
 };
 
-typedef FPostProcessMaterialPS<EPostProcessMaterialTarget::HighEnd, 0> FFPostProcessMaterialPS_HighEnd0;
-typedef FPostProcessMaterialPS<EPostProcessMaterialTarget::HighEnd, 1> FFPostProcessMaterialPS_HighEnd1;
-typedef FPostProcessMaterialPS<EPostProcessMaterialTarget::Mobile, 0> FPostProcessMaterialPS_Mobile0;
-typedef FPostProcessMaterialPS<EPostProcessMaterialTarget::Mobile, 1> FPostProcessMaterialPS_Mobile1;
+typedef FPostProcessMaterialPS<EPostProcessMaterialTarget::HighEnd, 0, false> FFPostProcessMaterialPS_HighEnd0;
+typedef FPostProcessMaterialPS<EPostProcessMaterialTarget::HighEnd, 1, false> FFPostProcessMaterialPS_HighEnd1;
+typedef FPostProcessMaterialPS<EPostProcessMaterialTarget::Mobile, 0, false> FPostProcessMaterialPS_Mobile0;
+typedef FPostProcessMaterialPS<EPostProcessMaterialTarget::Mobile, 1, false> FPostProcessMaterialPS_Mobile1;
+typedef FPostProcessMaterialPS<EPostProcessMaterialTarget::Mobile, 0, true> FPostProcessMaterialPS_Mobile0_AxisSwitch;
+typedef FPostProcessMaterialPS<EPostProcessMaterialTarget::Mobile, 1,true> FPostProcessMaterialPS_Mobile1_AxisSwitch;
 
 IMPLEMENT_MATERIAL_SHADER_TYPE(template<>, FFPostProcessMaterialPS_HighEnd0, TEXT("/Engine/Private/PostProcessMaterialShaders.usf"), TEXT("MainPS"), SF_Pixel);
 IMPLEMENT_MATERIAL_SHADER_TYPE(template<>, FFPostProcessMaterialPS_HighEnd1, TEXT("/Engine/Private/PostProcessMaterialShaders.usf"), TEXT("MainPS"), SF_Pixel);
 IMPLEMENT_MATERIAL_SHADER_TYPE(template<>, FPostProcessMaterialPS_Mobile0, TEXT("/Engine/Private/PostProcessMaterialShaders.usf"), TEXT("MainPS_ES2"), SF_Pixel);
 IMPLEMENT_MATERIAL_SHADER_TYPE(template<>, FPostProcessMaterialPS_Mobile1, TEXT("/Engine/Private/PostProcessMaterialShaders.usf"), TEXT("MainPS_ES2"), SF_Pixel);
+IMPLEMENT_MATERIAL_SHADER_TYPE(template<>, FPostProcessMaterialPS_Mobile0_AxisSwitch, TEXT("/Engine/Private/PostProcessMaterialShaders.usf"), TEXT("MainPS_ES2"), SF_Pixel);
+IMPLEMENT_MATERIAL_SHADER_TYPE(template<>, FPostProcessMaterialPS_Mobile1_AxisSwitch, TEXT("/Engine/Private/PostProcessMaterialShaders.usf"), TEXT("MainPS_ES2"), SF_Pixel);
 
 FRCPassPostProcessMaterial::FRCPassPostProcessMaterial(UMaterialInterface* InMaterialInterface, ERHIFeatureLevel::Type InFeatureLevel, EPixelFormat OutputFormatIN)
 : MaterialInterface(InMaterialInterface), OutputFormat(OutputFormatIN)
@@ -230,11 +244,11 @@ public:
 };
 TGlobalResource<FPostProcessMaterialVertexDeclaration> GPostProcessMaterialVertexDeclaration;
 
-template<typename TPixelShader>
+template<typename TVertexShader, typename TPixelShader>
 FShader* SetMobileShaders(const FMaterialShaderMap* MaterialShaderMap, FGraphicsPipelineStateInitializer &GraphicsPSOInit, FRenderingCompositePassContext &Context, FMaterialRenderProxy* Proxy)
 {
 	TPixelShader* PixelShader_Mobile = MaterialShaderMap->GetShader<TPixelShader>();
-	FPostProcessMaterialVS_Mobile* VertexShader_Mobile = MaterialShaderMap->GetShader<FPostProcessMaterialVS_Mobile>();
+	TVertexShader* VertexShader_Mobile = MaterialShaderMap->GetShader<TVertexShader>();
 
 	GraphicsPSOInit.BoundShaderState.VertexDeclarationRHI = GFilterVertexDeclaration.VertexDeclarationRHI;
 	GraphicsPSOInit.BoundShaderState.VertexShaderRHI = GETSAFERHISHADER_VERTEX(VertexShader_Mobile);
@@ -396,6 +410,10 @@ void FRCPassPostProcessMaterial::Process(FRenderingCompositePassContext& Context
 
 	FIntRect SrcRect = Context.SceneColorViewRect;
 	FIntRect DestRect = Context.GetSceneColorDestRect(*DestRenderTarget);
+	FIntPoint DestPos(0, 0);
+	FIntPoint DestSize = DestRect.Size();
+	EDrawRectangleFlags DrawRectangleFlags = EDRF_UseTriangleOptimization;
+
 	checkf(DestRect.Size() == SrcRect.Size(), TEXT("Post process material should not be used as upscaling pass."));
 	
 	FRHIRenderPassInfo RPInfo;
@@ -441,14 +459,42 @@ void FRCPassPostProcessMaterial::Process(FRenderingCompositePassContext& Context
 		const bool bViewSizeMatchesBufferSize = (View.ViewRect == Context.SceneColorViewRect && View.ViewRect.Size() == SrcSize && View.ViewRect.Min == FIntPoint::ZeroValue);
 		if (FeatureLevel <= ERHIFeatureLevel::ES3_1)
 		{
+			bool bNeedsVerticalAxisFlip = ShouldMobilePassFlipVerticalAxis(this) && RHINeedsToSwitchVerticalAxis(GShaderPlatformForFeatureLevel[Context.GetFeatureLevel()]);
+			if (bNeedsVerticalAxisFlip)
+			{
+				if (EBlendableLocation(Material->GetBlendableLocation()) == EBlendableLocation::BL_AfterTonemapping)
+				{
+					// flip dest rect y axis.
+					GraphicsPSOInit.RasterizerState = TStaticRasterizerState<FM_Solid, CM_None>::GetRHI();
+					DestPos.Y = DestRect.Max.Y;
+					DestSize.Y = -DestRect.Max.Y;
+					// triangle optimization currently doesn't work when flipped.
+					DrawRectangleFlags = EDRF_Default;
+				}
+			}
+
 			// use mobile's post process material.
 			if (bViewSizeMatchesBufferSize)
 			{
-				VertexShader = SetMobileShaders< FPostProcessMaterialPS_Mobile0>(MaterialShaderMap, GraphicsPSOInit, Context, Proxy);
+				if(bNeedsVerticalAxisFlip)
+				{
+					VertexShader = SetMobileShaders<FPostProcessMaterialVS_Mobile_AxisSwitch, FPostProcessMaterialPS_Mobile0_AxisSwitch>(MaterialShaderMap, GraphicsPSOInit, Context, Proxy);
+				}
+				else
+				{
+					VertexShader = SetMobileShaders<FPostProcessMaterialVS_Mobile, FPostProcessMaterialPS_Mobile0>(MaterialShaderMap, GraphicsPSOInit, Context, Proxy);
+				}
 			}
 			else
 			{
-				VertexShader = SetMobileShaders< FPostProcessMaterialPS_Mobile1>(MaterialShaderMap, GraphicsPSOInit, Context, Proxy);
+				if (bNeedsVerticalAxisFlip)
+				{
+					VertexShader = SetMobileShaders<FPostProcessMaterialVS_Mobile_AxisSwitch, FPostProcessMaterialPS_Mobile1_AxisSwitch>(MaterialShaderMap, GraphicsPSOInit, Context, Proxy);
+				}
+				else
+				{
+					VertexShader = SetMobileShaders<FPostProcessMaterialVS_Mobile, FPostProcessMaterialPS_Mobile1>(MaterialShaderMap, GraphicsPSOInit, Context, Proxy);
+				}
 			}
 			Context.RHICmdList.SetStencilRef(StencilRefValue);
 		}
@@ -489,8 +535,8 @@ void FRCPassPostProcessMaterial::Process(FRenderingCompositePassContext& Context
 
 		DrawPostProcessPass(
 			Context.RHICmdList,
-			0, 0,
-			DestRect.Width(), DestRect.Height(),
+			DestPos.X, DestPos.Y,
+			DestSize.X, DestSize.Y,
 			SrcRect.Min.X, SrcRect.Min.Y,
 			SrcRect.Width(), SrcRect.Height(),
 			DestRect.Size(),
@@ -498,7 +544,7 @@ void FRCPassPostProcessMaterial::Process(FRenderingCompositePassContext& Context
 			VertexShader,
 			View.StereoPass,
 			Context.HasHmdMesh(),
-			EDRF_UseTriangleOptimization);
+			DrawRectangleFlags);
 	}
 	Context.RHICmdList.EndRenderPass();
 	Context.RHICmdList.CopyToResolveTarget(DestRenderTarget->TargetableTexture, DestRenderTarget->ShaderResourceTexture, FResolveParams());
diff --git a/Engine/Source/Runtime/Renderer/Private/PostProcess/PostProcessMobile.cpp b/Engine/Source/Runtime/Renderer/Private/PostProcess/PostProcessMobile.cpp
--- a/Engine/Source/Runtime/Renderer/Private/PostProcess/PostProcessMobile.cpp
+++ b/Engine/Source/Runtime/Renderer/Private/PostProcess/PostProcessMobile.cpp
@@ -29,6 +29,17 @@ float GetMobileDepthOfFieldScale(const FViewInfo& View)
 	PRAGMA_ENABLE_DEPRECATION_WARNINGS
 }
 
+static const FRenderingCompositePass* GMobilePassShouldFlipVerticalAxis = nullptr;
+void SetMobilePassFlipVerticalAxis(const FRenderingCompositePass* FlipPass)
+{
+	GMobilePassShouldFlipVerticalAxis = FlipPass;
+}
+
+bool ShouldMobilePassFlipVerticalAxis(const FRenderingCompositePass* ShouldFlipPass)
+{
+	return GMobilePassShouldFlipVerticalAxis == ShouldFlipPass;
+}
+
 //
 // BLOOM SETUP
 //
diff --git a/Engine/Source/Runtime/Renderer/Private/PostProcess/PostProcessMobile.h b/Engine/Source/Runtime/Renderer/Private/PostProcess/PostProcessMobile.h
--- a/Engine/Source/Runtime/Renderer/Private/PostProcess/PostProcessMobile.h
+++ b/Engine/Source/Runtime/Renderer/Private/PostProcess/PostProcessMobile.h
@@ -15,6 +15,10 @@ class FViewInfo;
 // return Depth of Field Scale if Gaussian DoF mode is active. 0.0f otherwise.
 float GetMobileDepthOfFieldScale(const FViewInfo& View);
 
+// Used to indicate the final PP stage which needs to be flipped on platforms that 'RHINeedsToSwitchVerticalAxis'
+void SetMobilePassFlipVerticalAxis(const FRenderingCompositePass* FlipPass);
+bool ShouldMobilePassFlipVerticalAxis(const FRenderingCompositePass* ShouldFlipPass);
+
 class FRCPassPostProcessBloomSetupES2 : public TRenderingCompositePassBase<1, 1>
 {
 public:
diff --git a/Engine/Source/Runtime/Renderer/Private/PostProcess/PostProcessTemporalAA.cpp b/Engine/Source/Runtime/Renderer/Private/PostProcess/PostProcessTemporalAA.cpp
--- a/Engine/Source/Runtime/Renderer/Private/PostProcess/PostProcessTemporalAA.cpp
+++ b/Engine/Source/Runtime/Renderer/Private/PostProcess/PostProcessTemporalAA.cpp
@@ -612,9 +612,14 @@ class FPostProcessTemporalAACS : public FGlobalShader
 
 IMPLEMENT_GLOBAL_SHADER(FPostProcessTemporalAACS, "/Engine/Private/PostProcessTemporalAA.usf", "MainCS", SF_Compute);
 
+static FPostProcessTonemapVS::FPermutationDomain GetVertexShaderPermutationVector()
+{
+	return FPostProcessTonemapVS::BuildPermutationVector(true/*bDoEyeAdaptation*/, false/*bNeedsToSwitchVerticalAxis*/);
+}
+
 static inline void TransitionPixelPassResources(FRenderingCompositePassContext& Context)
 {
-	TShaderMapRef< FPostProcessTonemapVS > VertexShader(Context.GetShaderMap());
+	TShaderMapRef< FPostProcessTonemapVS > VertexShader(Context.GetShaderMap(), GetVertexShaderPermutationVector());
 	VertexShader->TransitionResources(Context);
 }
 
@@ -634,7 +639,8 @@ void DrawPixelPassTemplate(
 	GraphicsPSOInit.RasterizerState = TStaticRasterizerState<>::GetRHI();
 	GraphicsPSOInit.DepthStencilState = DepthStencilState;
 
-	TShaderMapRef< FPostProcessTonemapVS > VertexShader(Context.GetShaderMap());
+	auto VertexShaderPermutationVector = GetVertexShaderPermutationVector();
+	TShaderMapRef< FPostProcessTonemapVS > VertexShader(Context.GetShaderMap(), VertexShaderPermutationVector);
 	TShaderMapRef< FPostProcessTemporalAAPS > PixelShader(Context.GetShaderMap(), PermutationVector);
 
 	GraphicsPSOInit.BoundShaderState.VertexDeclarationRHI = GFilterVertexDeclaration.VertexDeclarationRHI;
@@ -643,7 +649,7 @@ void DrawPixelPassTemplate(
 	GraphicsPSOInit.PrimitiveType = PT_TriangleList;
 
 	SetGraphicsPipelineState(Context.RHICmdList, GraphicsPSOInit);
-	VertexShader->SetVS(Context);
+	VertexShader->SetVS(Context, VertexShaderPermutationVector);
 	PixelShader->SetParameters(Context.RHICmdList, Context, InputHistory, PassParameters, bUseDither, SrcSize);
 
 	DrawRectangle(
diff --git a/Engine/Source/Runtime/Renderer/Private/PostProcess/PostProcessTonemap.cpp b/Engine/Source/Runtime/Renderer/Private/PostProcess/PostProcessTonemap.cpp
--- a/Engine/Source/Runtime/Renderer/Private/PostProcess/PostProcessTonemap.cpp
+++ b/Engine/Source/Runtime/Renderer/Private/PostProcess/PostProcessTonemap.cpp
@@ -122,6 +122,7 @@ class FTonemapperGrainIntensityDim : SHADER_PERMUTATION_BOOL("USE_GRAIN_INTENSIT
 class FTonemapperVignetteDim       : SHADER_PERMUTATION_BOOL("USE_VIGNETTE");
 class FTonemapperSharpenDim        : SHADER_PERMUTATION_BOOL("USE_SHARPEN");
 class FTonemapperGrainJitterDim    : SHADER_PERMUTATION_BOOL("USE_GRAIN_JITTER");
+class FTonemapperSwitchAxis        : SHADER_PERMUTATION_BOOL("NEEDTOSWITCHVERTICLEAXIS");
 
 using FCommonDomain = TShaderPermutationDomain<
 	FTonemapperBloomDim,
@@ -129,10 +130,17 @@ using FCommonDomain = TShaderPermutationDomain<
 	FTonemapperGrainIntensityDim,
 	FTonemapperVignetteDim,
 	FTonemapperSharpenDim,
-	FTonemapperGrainJitterDim>;
+	FTonemapperGrainJitterDim,
+	FTonemapperSwitchAxis>;
 
-FORCEINLINE bool ShouldCompileCommonPermutation(const FCommonDomain& PermutationVector)
+FORCEINLINE_DEBUGGABLE bool ShouldCompileCommonPermutation(const FGlobalShaderPermutationParameters& Parameters, const FCommonDomain& PermutationVector)
 {
+	// Prevent switch axis permutation on platforms that dont require it.
+	if (PermutationVector.Get<FTonemapperSwitchAxis>() && !RHINeedsToSwitchVerticalAxis(Parameters.Platform))
+	{
+		return false;
+	}
+
 	// If GammaOnly, don't compile any other dimmension == true.
 	if (PermutationVector.Get<FTonemapperGammaOnlyDim>())
 	{
@@ -146,7 +154,7 @@ FORCEINLINE bool ShouldCompileCommonPermutation(const FCommonDomain& Permutation
 }
 
 // Common conversion of engine settings into.
-FCommonDomain BuildCommonPermutationDomain(const FViewInfo& View, bool bGammaOnly)
+FCommonDomain BuildCommonPermutationDomain(const FViewInfo& View, bool bGammaOnly, bool bSwitchVerticalAxis)
 {
 	const FSceneViewFamily* Family = View.Family;
 
@@ -166,7 +174,8 @@ FCommonDomain BuildCommonPermutationDomain(const FViewInfo& View, bool bGammaOnl
 	PermutationVector.Set<FTonemapperVignetteDim>(Settings.VignetteIntensity > 0.0f);
 	PermutationVector.Set<FTonemapperBloomDim>(Settings.BloomIntensity > 0.0);
 	PermutationVector.Set<FTonemapperGrainJitterDim>(Settings.GrainJitter > 0.0f);
-	PermutationVector.Set<FTonemapperSharpenDim>(CVarTonemapperSharpen.GetValueOnRenderThread() > 0.0f);
+	PermutationVector.Set<FTonemapperSharpenDim>(CVarTonemapperSharpen.GetValueOnRenderThread() > 0.0f);	
+	PermutationVector.Set<FTonemapperSwitchAxis>(bSwitchVerticalAxis);
 
 	return PermutationVector;
 }
@@ -217,7 +226,7 @@ FDesktopDomain RemapPermutation(FDesktopDomain PermutationVector)
 	return PermutationVector;
 }
 
-bool ShouldCompileDesktopPermutation(FDesktopDomain PermutationVector)
+bool ShouldCompileDesktopPermutation(const FGlobalShaderPermutationParameters& Parameters, FDesktopDomain PermutationVector)
 {
 	auto CommonPermutationVector = PermutationVector.Get<FCommonDomain>();
 
@@ -226,7 +235,7 @@ bool ShouldCompileDesktopPermutation(FDesktopDomain PermutationVector)
 		return false;
 	}
 
-	if (!ShouldCompileCommonPermutation(CommonPermutationVector))
+	if (!ShouldCompileCommonPermutation(Parameters, CommonPermutationVector))
 	{
 		return false;
 	}
@@ -693,9 +702,8 @@ public:
 };
 
 
-// Vertex Shader permutations based on bool AutoExposure.
-IMPLEMENT_SHADER_TYPE(template<>, TPostProcessTonemapVS<true>, TEXT("/Engine/Private/PostProcessTonemap.usf"), TEXT("MainVS"), SF_Vertex);
-IMPLEMENT_SHADER_TYPE(template<>, TPostProcessTonemapVS<false>, TEXT("/Engine/Private/PostProcessTonemap.usf"), TEXT("MainVS"), SF_Vertex);
+// Vertex Shader 
+IMPLEMENT_GLOBAL_SHADER(FPostProcessTonemapVS, "/Engine/Private/PostProcessTonemap.usf", "MainVS", SF_Vertex);
 
 
 class FPostProcessTonemapPS : public FGlobalShader
@@ -710,7 +718,7 @@ class FPostProcessTonemapPS : public FGlobalShader
 		{
 			return false;
 		}
-		return TonemapperPermutation::ShouldCompileDesktopPermutation(FPermutationDomain(Parameters.PermutationId));
+		return TonemapperPermutation::ShouldCompileDesktopPermutation(Parameters, FPermutationDomain(Parameters.PermutationId));
 	}
 
 	static void ModifyCompilationEnvironment(const FGlobalShaderPermutationParameters& Parameters, FShaderCompilerEnvironment& OutEnvironment)
@@ -785,7 +793,7 @@ class FPostProcessTonemapCS : public FGlobalShader
 
 		FPermutationDomain PermutationVector(Parameters.PermutationId);
 
-		return TonemapperPermutation::ShouldCompileDesktopPermutation(PermutationVector.Get<TonemapperPermutation::FDesktopDomain>());
+		return TonemapperPermutation::ShouldCompileDesktopPermutation(Parameters, PermutationVector.Get<TonemapperPermutation::FDesktopDomain>());
 	}	
 	
 	static void ModifyCompilationEnvironment(const FGlobalShaderPermutationParameters& Parameters, FShaderCompilerEnvironment& OutEnvironment)
@@ -924,11 +932,9 @@ namespace
 namespace PostProcessTonemapUtil
 {
 
-template <bool bVSDoEyeAdaptation>
-static inline void ShaderTransitionResources(const FRenderingCompositePassContext& Context)
+static inline void ShaderTransitionResources(const FRenderingCompositePassContext& Context, const FPostProcessTonemapVS::FPermutationDomain& PermutationVector)
 {
-	typedef TPostProcessTonemapVS<bVSDoEyeAdaptation>				VertexShaderType;
-	TShaderMapRef<VertexShaderType> VertexShader(Context.GetShaderMap());
+	TShaderMapRef<FPostProcessTonemapVS> VertexShader(Context.GetShaderMap(), PermutationVector);
 	VertexShader->TransitionResources(Context);
 }
 
@@ -987,10 +993,11 @@ void FRCPassPostProcessTonemap::Process(FRenderingCompositePassContext& Context)
 
 	FSceneRenderTargets& SceneContext = FSceneRenderTargets::Get(Context.RHICmdList);
 
+	bool bNeedsToSwitchVerticalAxis = RHINeedsToSwitchVerticalAxis(GShaderPlatformForFeatureLevel[Context.GetFeatureLevel()]) && ShouldMobilePassFlipVerticalAxis(this);
 	// Generate permutation vector for the desktop tonemapper.
 	TonemapperPermutation::FDesktopDomain DesktopPermutationVector;
 	{
-		TonemapperPermutation::FCommonDomain CommonDomain = TonemapperPermutation::BuildCommonPermutationDomain(View, bDoGammaOnly);
+		TonemapperPermutation::FCommonDomain CommonDomain = TonemapperPermutation::BuildCommonPermutationDomain(View, bDoGammaOnly, bNeedsToSwitchVerticalAxis);
 		DesktopPermutationVector.Set<TonemapperPermutation::FCommonDomain>(CommonDomain);
 
 		if (!CommonDomain.Get<TonemapperPermutation::FTonemapperGammaOnlyDim>())
@@ -1062,17 +1069,12 @@ void FRCPassPostProcessTonemap::Process(FRenderingCompositePassContext& Context)
 	}
 	else
 	{
+		auto VertexShaderPermutationVector = FPostProcessTonemapVS::BuildPermutationVector(bDoEyeAdaptation, bNeedsToSwitchVerticalAxis);
+
 		WaitForInputPassComputeFences(Context.RHICmdList);
 
 		const EShaderPlatform ShaderPlatform = GShaderPlatformForFeatureLevel[Context.GetFeatureLevel()];
-		if (bDoEyeAdaptation)
-		{
-			PostProcessTonemapUtil::ShaderTransitionResources<true>(Context);
-		}
-		else
-		{
-			PostProcessTonemapUtil::ShaderTransitionResources<false>(Context);
-		}
+		PostProcessTonemapUtil::ShaderTransitionResources(Context, VertexShaderPermutationVector);
 
 		ERenderTargetLoadAction LoadAction = ERenderTargetLoadAction::ELoad;
 
@@ -1102,8 +1104,8 @@ void FRCPassPostProcessTonemap::Process(FRenderingCompositePassContext& Context)
 		Context.RHICmdList.BeginRenderPass(RPInfo, TEXT("Tonemap"));
 		{
 			Context.SetViewportAndCallRHI(DestRect, 0.0f, 1.0f);
-
-			FShader* VertexShader;
+			TShaderMapRef<FPostProcessTonemapVS> VertexShader(Context.GetShaderMap(), VertexShaderPermutationVector);
+			TShaderMapRef<FPostProcessTonemapPS> PixelShader(Context.GetShaderMap(), DesktopPermutationVector);
 			{
 				FGraphicsPipelineStateInitializer GraphicsPSOInit;
 				Context.RHICmdList.ApplyCachedRenderTargets(GraphicsPSOInit);
@@ -1119,36 +1121,14 @@ void FRCPassPostProcessTonemap::Process(FRenderingCompositePassContext& Context)
 				}
 				GraphicsPSOInit.RasterizerState = TStaticRasterizerState<>::GetRHI();
 				GraphicsPSOInit.DepthStencilState = TStaticDepthStencilState<false, CF_Always>::GetRHI();
-
-				if (bDoEyeAdaptation)
-				{
-					VertexShader = Context.GetShaderMap()->GetShader<TPostProcessTonemapVS<true>>();
-				}
-				else
-				{
-					VertexShader = Context.GetShaderMap()->GetShader<TPostProcessTonemapVS<false>>();
-				}
-
-				TShaderMapRef<FPostProcessTonemapPS> PixelShader(Context.GetShaderMap(), DesktopPermutationVector);
-
+				
 				GraphicsPSOInit.BoundShaderState.VertexDeclarationRHI = GFilterVertexDeclaration.VertexDeclarationRHI;
-				GraphicsPSOInit.BoundShaderState.VertexShaderRHI = VertexShader->GetVertexShader();
+				GraphicsPSOInit.BoundShaderState.VertexShaderRHI = GETSAFERHISHADER_VERTEX(*VertexShader);
 				GraphicsPSOInit.BoundShaderState.PixelShaderRHI = GETSAFERHISHADER_PIXEL(*PixelShader);
 				GraphicsPSOInit.PrimitiveType = PT_TriangleList;
 
 				SetGraphicsPipelineState(Context.RHICmdList, GraphicsPSOInit);
-
-				if (bDoEyeAdaptation)
-				{
-					TShaderMapRef<TPostProcessTonemapVS<true>> VertexShaderMapRef(Context.GetShaderMap());
-					VertexShaderMapRef->SetVS(Context);
-				}
-				else
-				{
-					TShaderMapRef<TPostProcessTonemapVS<false>> VertexShaderMapRef(Context.GetShaderMap());
-					VertexShaderMapRef->SetVS(Context);
-				}
-
+				VertexShader->SetVS(Context, VertexShaderPermutationVector);
 				PixelShader->SetPS(Context);
 			}
 
@@ -1160,7 +1140,7 @@ void FRCPassPostProcessTonemap::Process(FRenderingCompositePassContext& Context)
 				SrcRect.Width(), SrcRect.Height(),
 				DestRect.Size(),
 				SrcSize,
-				VertexShader,
+				*VertexShader,
 				View.StereoPass,
 				Context.HasHmdMesh(),
 				EDRF_UseTriangleOptimization);
@@ -1279,6 +1259,9 @@ class FPostProcessTonemapPS_ES2 : public FGlobalShader
 		EnableIfSet<TonemapperPermutation::FTonemapperVignetteDim>(WantedCommonPermutationVector, RemappedCommonPermutationVector);
  		// Grain intensity permutation
 		EnableIfSet<TonemapperPermutation::FTonemapperGrainIntensityDim>(WantedCommonPermutationVector, RemappedCommonPermutationVector);
+		// Switch Y axis
+		EnableIfSet<TonemapperPermutation::FTonemapperSwitchAxis>(WantedCommonPermutationVector, RemappedCommonPermutationVector);
+
  		// Color matrix
 		EnableIfSet<FTonemapperColorMatrixDim>(WantedPermutationVector, RemappedPermutationVector);
  		// msaa permutation.
@@ -1320,9 +1303,9 @@ class FPostProcessTonemapPS_ES2 : public FGlobalShader
 		return RemappedPermutationVector;
 	}
 
-	static FPermutationDomain BuildPermutationVector(const FViewInfo& View)
+	static FPermutationDomain BuildPermutationVector(const FViewInfo& View, bool bNeedsToSwitchVerticalAxis)
 	{
-		TonemapperPermutation::FCommonDomain CommonPermutationVector = TonemapperPermutation::BuildCommonPermutationDomain(View, /* bGammaOnly = */ false);
+		TonemapperPermutation::FCommonDomain CommonPermutationVector = TonemapperPermutation::BuildCommonPermutationDomain(View, /* bGammaOnly = */ false, bNeedsToSwitchVerticalAxis);
 
 		FPostProcessTonemapPS_ES2::FPermutationDomain MobilePermutationVector;
 		MobilePermutationVector.Set<TonemapperPermutation::FCommonDomain>(CommonPermutationVector);
@@ -1403,7 +1386,7 @@ class FPostProcessTonemapPS_ES2 : public FGlobalShader
 		FPermutationDomain PermutationVector(Parameters.PermutationId);
 
 		auto CommonPermutationVector = PermutationVector.Get<TonemapperPermutation::FCommonDomain>();
-		if (!TonemapperPermutation::ShouldCompileCommonPermutation(CommonPermutationVector))
+		if (!TonemapperPermutation::ShouldCompileCommonPermutation(Parameters, CommonPermutationVector))
 		{
 			return false;
 		}
@@ -1589,13 +1572,37 @@ class FPostProcessTonemapVS_ES2 : public FGlobalShader
 {
 	DECLARE_GLOBAL_SHADER(FPostProcessTonemapVS_ES2);
 
-	using FPermutationDomain = TShaderPermutationDomain<>;
+	class FTonemapperSwitchAxis : SHADER_PERMUTATION_BOOL("NEEDTOSWITCHVERTICLEAXIS");
+
+	using FPermutationDomain = TShaderPermutationDomain<
+		FTonemapperSwitchAxis
+	>;
 
 	static bool ShouldCompilePermutation(const FGlobalShaderPermutationParameters& Parameters)
 	{
+		FPermutationDomain PermutationVector(Parameters.PermutationId);
+
+		// Prevent switch axis permutation on platforms that dont require it.
+		if (PermutationVector.Get<FTonemapperSwitchAxis>() && !RHINeedsToSwitchVerticalAxis(Parameters.Platform))
+		{
+			return false;
+		}
+
 		return !IsConsolePlatform(Parameters.Platform);
 	}
 
+	static void ModifyCompilationEnvironment(const FGlobalShaderPermutationParameters& Parameters, FShaderCompilerEnvironment& OutEnvironment)
+	{
+		FGlobalShader::ModifyCompilationEnvironment(Parameters, OutEnvironment);
+	}
+
+	static FPermutationDomain BuildPermutationVector(bool bNeedsToSwitchVerticalAxis)
+	{
+		FPermutationDomain PermutationVector;
+		PermutationVector.Set<FTonemapperSwitchAxis>(bNeedsToSwitchVerticalAxis);
+		return PermutationVector;
+	}
+
 	FPostProcessTonemapVS_ES2() { }
 
 public:
@@ -1706,10 +1713,13 @@ void FRCPassPostProcessTonemapES2::Process(FRenderingCompositePassContext& Conte
 	{
 		Context.SetViewportAndCallRHI(DestRect);
 
-		auto PermutationVector = FPostProcessTonemapPS_ES2::BuildPermutationVector(View);
+		bool bNeedsToSwitchVerticalAxis = RHINeedsToSwitchVerticalAxis(GShaderPlatformForFeatureLevel[Context.GetFeatureLevel()]) && ShouldMobilePassFlipVerticalAxis(this);
+
+		auto VertexShaderPermutationVector = FPostProcessTonemapVS_ES2::BuildPermutationVector(bNeedsToSwitchVerticalAxis);
+		auto PixelShaderPermutationVector = FPostProcessTonemapPS_ES2::BuildPermutationVector(View, bNeedsToSwitchVerticalAxis);
 
-		TShaderMapRef<FPostProcessTonemapVS_ES2> VertexShader(Context.GetShaderMap());
-		TShaderMapRef<FPostProcessTonemapPS_ES2> PixelShader(Context.GetShaderMap(), PermutationVector);
+		TShaderMapRef<FPostProcessTonemapVS_ES2> VertexShader(Context.GetShaderMap(), VertexShaderPermutationVector);
+		TShaderMapRef<FPostProcessTonemapPS_ES2> PixelShader(Context.GetShaderMap(), PixelShaderPermutationVector);
 
 		{
 			FGraphicsPipelineStateInitializer GraphicsPSOInit;
@@ -1728,7 +1738,7 @@ void FRCPassPostProcessTonemapES2::Process(FRenderingCompositePassContext& Conte
 			SetGraphicsPipelineState(Context.RHICmdList, GraphicsPSOInit);
 
 			VertexShader->SetVS(Context);
-			PixelShader->SetPS(Context.RHICmdList, Context, PermutationVector, bSRGBAwareTarget);
+			PixelShader->SetPS(Context.RHICmdList, Context, PixelShaderPermutationVector, bSRGBAwareTarget);
 		}
 
 		DrawRectangle(
diff --git a/Engine/Source/Runtime/Renderer/Private/PostProcess/PostProcessTonemap.h b/Engine/Source/Runtime/Renderer/Private/PostProcess/PostProcessTonemap.h
--- a/Engine/Source/Runtime/Renderer/Private/PostProcess/PostProcessTonemap.h
+++ b/Engine/Source/Runtime/Renderer/Private/PostProcess/PostProcessTonemap.h
@@ -98,19 +98,31 @@ private:
 
 
 /** Encapsulates the post processing tone map vertex shader. */
-template< bool bUseAutoExposure>
-class TPostProcessTonemapVS : public FGlobalShader
+class FPostProcessTonemapVS : public FGlobalShader
 {
 	// This class is in the header so that Temporal AA can share this vertex shader.
-	DECLARE_SHADER_TYPE(TPostProcessTonemapVS,Global);
+	DECLARE_GLOBAL_SHADER(FPostProcessTonemapVS);
+
+	class FTonemapperVSSwitchAxis : SHADER_PERMUTATION_BOOL("NEEDTOSWITCHVERTICLEAXIS");
+	class FTonemapperVSUseAutoExposure : SHADER_PERMUTATION_BOOL("EYEADAPTATION_EXPOSURE_FIX");
+	using FPermutationDomain = TShaderPermutationDomain<
+		FTonemapperVSSwitchAxis,
+		FTonemapperVSUseAutoExposure
+	>;
 
 	static bool ShouldCompilePermutation(const FGlobalShaderPermutationParameters& Parameters)
 	{
+		FPermutationDomain PermutationVector(Parameters.PermutationId);
+		// Prevent switch axis permutation on platforms that dont require it.
+		if (PermutationVector.Get<FTonemapperVSSwitchAxis>() && !RHINeedsToSwitchVerticalAxis(Parameters.Platform))
+		{
+			return false;
+		}
 		return true;
 	}
 
 	/** Default constructor. */
-	TPostProcessTonemapVS(){}
+	FPostProcessTonemapVS(){}
 
 public:
 	FPostProcessPassParameters PostprocessParameter;
@@ -120,7 +132,7 @@ public:
 	FShaderParameter ScreenPosToScenePixel;
 
 	/** Initialization constructor. */
-	TPostProcessTonemapVS(const ShaderMetaType::CompiledShaderInitializerType& Initializer)
+	FPostProcessTonemapVS(const ShaderMetaType::CompiledShaderInitializerType& Initializer)
 		: FGlobalShader(Initializer)
 	{
 		PostprocessParameter.Bind(Initializer.ParameterMap);
@@ -132,8 +144,15 @@ public:
 
 	static void ModifyCompilationEnvironment(const FGlobalShaderPermutationParameters& Parameters, FShaderCompilerEnvironment& OutEnvironment)
 	{
-		// Compile time template-based conditional
-		OutEnvironment.SetDefine(TEXT("EYEADAPTATION_EXPOSURE_FIX"), (uint32)bUseAutoExposure);
+	}
+
+	static FPermutationDomain BuildPermutationVector(bool bDoEyeAdaptation, bool bNeedsToSwitchVerticalAxis)
+	{
+		FPermutationDomain PermutationVector;
+		PermutationVector.Set<FTonemapperVSSwitchAxis>(bNeedsToSwitchVerticalAxis);
+		PermutationVector.Set<FTonemapperVSUseAutoExposure>(bDoEyeAdaptation);
+
+		return PermutationVector;
 	}
 
 	void TransitionResources(const FRenderingCompositePassContext& Context)
@@ -149,7 +168,7 @@ public:
 		}
 	}
 
-	void SetVS(const FRenderingCompositePassContext& Context)
+	void SetVS(const FRenderingCompositePassContext& Context, const FPermutationDomain& PermutationVector )
 	{
 		const FVertexShaderRHIParamRef ShaderRHI = GetVertexShader();
 
@@ -187,7 +206,7 @@ public:
 		}
 
 		// Compile time template-based conditional
-		if (!bUseAutoExposure)
+		if (!PermutationVector.Get<FTonemapperVSUseAutoExposure>())
 		{
 			// Compute a CPU-based default.  NB: reverts to "1" if SM5 feature level is not supported
 			float FixedExposure = FRCPassPostProcessEyeAdaptation::GetFixedExposure(Context.View);
@@ -216,6 +235,3 @@ public:
 		return bShaderHasOutdatedParameters;
 	}
 };
-
-// Default uses eye adaptation
-typedef TPostProcessTonemapVS<true/*bUseEyeAdaptation*/> FPostProcessTonemapVS;
diff --git a/Engine/Source/Runtime/Renderer/Private/PostProcess/PostProcessUpscale.cpp b/Engine/Source/Runtime/Renderer/Private/PostProcess/PostProcessUpscale.cpp
--- a/Engine/Source/Runtime/Renderer/Private/PostProcess/PostProcessUpscale.cpp
+++ b/Engine/Source/Runtime/Renderer/Private/PostProcess/PostProcessUpscale.cpp
@@ -417,10 +417,11 @@ void FRCPassPostProcessUpscale::Process(FRenderingCompositePassContext& Context)
 
 		if (bIsMobileRenderer && RHINeedsToSwitchVerticalAxis(GShaderPlatformForFeatureLevel[Context.GetFeatureLevel()]))
 		{
-			// Move dest rect on y axis as we're rendering directly onto the ogl backbuffer.
-			const int32 OutputTargetHeight = PassOutputs[0].RenderTargetDesc.Extent.Y;
-			int32 NewMaxY = OutputTargetHeight - DestRect.Min.Y;
-			DestRect.Min.Y = OutputTargetHeight - DestRect.Max.Y;
+			// Get get this from the pooled targets as this may have been overridden.
+			FIntPoint DestSize = PassOutputs[0].PooledRenderTarget->GetDesc().Extent;
+			// Move dest rect on y axis as we're rendering directly onto the ogl backbuffer
+			int32 NewMaxY = DestSize.Y - DestRect.Min.Y;
+			DestRect.Min.Y = DestSize.Y - DestRect.Max.Y;
 			DestRect.Max.Y = NewMaxY;
 		}
 
diff --git a/Engine/Source/Runtime/Renderer/Private/PostProcess/PostProcessing.cpp b/Engine/Source/Runtime/Renderer/Private/PostProcess/PostProcessing.cpp
--- a/Engine/Source/Runtime/Renderer/Private/PostProcess/PostProcessing.cpp
+++ b/Engine/Source/Runtime/Renderer/Private/PostProcess/PostProcessing.cpp
@@ -2737,6 +2737,7 @@ void FPostProcessing::ProcessES2(FRHICommandListImmediate& RHICmdList, FScene* S
 				PostProcessTonemap->bDoScreenPercentageInTonemapper = false;
 				DoScreenPercentageInTonemapperPtr = &PostProcessTonemap->bDoScreenPercentageInTonemapper;
 			}
+			SetMobilePassFlipVerticalAxis(TonemapperPass);
 		}
 
 		// if Context.FinalOutput was the clipped result of sunmask stage then this stage also restores Context.FinalOutput back original target size.
@@ -2748,6 +2749,7 @@ void FPostProcessing::ProcessES2(FRHICommandListImmediate& RHICmdList, FScene* S
 			{
 				Context.FinalOutput = AddPostProcessMaterialChain(Context, BL_AfterTonemapping, nullptr);
 			}
+			SetMobilePassFlipVerticalAxis(Context.FinalOutput.GetPass());
 
 			if (bUseAa)
 			{
@@ -2869,6 +2871,7 @@ void FPostProcessing::ProcessES2(FRHICommandListImmediate& RHICmdList, FScene* S
 			CompositeContext.Process(Context.FinalOutput.GetPass(), TEXT("PostProcessingES2"));
 		}
 	}
+	SetMobilePassFlipVerticalAxis(nullptr);
 }
 
 void FPostProcessing::ProcessPlanarReflection(FRHICommandListImmediate& RHICmdList, FViewInfo& View, TRefCountPtr<IPooledRenderTarget>& VelocityRT, TRefCountPtr<IPooledRenderTarget>& OutFilteredSceneColor)
