Added checks in static draw interface to get a better callstack.

diff --git a/Engine/Source/Runtime/Renderer/Private/PrimitiveSceneInfo.cpp b/Engine/Source/Runtime/Renderer/Private/PrimitiveSceneInfo.cpp
--- a/Engine/Source/Runtime/Renderer/Private/PrimitiveSceneInfo.cpp
+++ b/Engine/Source/Runtime/Renderer/Private/PrimitiveSceneInfo.cpp
@@ -55,6 +55,11 @@ public:
 			check(Mesh.VertexFactory->IsInitialized());
 			checkSlow(IsInRenderingThread());
 
+			for (const FMeshBatchElement& BatchElement : Mesh.Elements)
+			{
+				check(!BatchElement.IndexBuffer || (BatchElement.IndexBuffer && BatchElement.IndexBuffer->IsInitialized() && BatchElement.IndexBuffer->IndexBufferRHI));
+			}
+
 			FPrimitiveSceneProxy* PrimitiveSceneProxy = PrimitiveSceneInfo->Proxy;
 			PrimitiveSceneProxy->VerifyUsedMaterial(Mesh.MaterialRenderProxy);
 
