UE4 - Compile fix
#5831

diff --git a/Engine/Source/Runtime/Renderer/Private/DeferredShadingRenderer.cpp b/Engine/Source/Runtime/Renderer/Private/DeferredShadingRenderer.cpp
--- a/Engine/Source/Runtime/Renderer/Private/DeferredShadingRenderer.cpp
+++ b/Engine/Source/Runtime/Renderer/Private/DeferredShadingRenderer.cpp
@@ -2007,7 +2007,7 @@ void FDeferredShadingSceneRenderer::Render(FRHICommandListImmediate& RHICmdList)
 			}
 			else
 			{
-				GPostProcessing.Process(RHICmdList, Views[ ViewIndex ], VelocityRT);
+				GPostProcessing.Process(RHICmdList, Views[ ViewIndex ], SceneContext.SceneVelocity);
 			}
 		}
 
