The render thread now updates Actor's LastRenderTime at the same time that it updates the PrimitiveComponent's LastRenderTime. This avoids WasRecentlyRendered and GetLastRenderTime from having to search the component list for the information. LastRenderTime and LastRenderTimeOnScreen are now private and have accessors for interacting with them.

diff --git a/Engine/Source/Runtime/Engine/Classes/Components/PrimitiveComponent.h b/Engine/Source/Runtime/Engine/Classes/Components/PrimitiveComponent.h
--- a/Engine/Source/Runtime/Engine/Classes/Components/PrimitiveComponent.h
+++ b/Engine/Source/Runtime/Engine/Classes/Components/PrimitiveComponent.h
@@ -609,15 +609,18 @@ public:
 	/** Last time the component was submitted for rendering (called FScene::AddPrimitive). */
 	float LastSubmitTime;
 
+private:
 	/**
 	 * The value of WorldSettings->TimeSeconds for the frame when this component was last rendered.  This is written
 	 * from the render thread, which is up to a frame behind the game thread, so you should allow this time to
 	 * be at least a frame behind the game thread's world time before you consider the actor non-visible.
 	 */
-	float LastRenderTime;
+	mutable float LastRenderTime;
 
 	/** Same as LastRenderTimeOnScreen but only updated if the component is on screen. Used by the texture streamer. */
-	float LastRenderTimeOnScreen;
+	mutable float LastRenderTimeOnScreen;
+
+	friend class FPrimitiveSceneInfo;
 
 public:
 
diff --git a/Engine/Source/Runtime/Engine/Classes/GameFramework/Actor.h b/Engine/Source/Runtime/Engine/Classes/GameFramework/Actor.h
--- a/Engine/Source/Runtime/Engine/Classes/GameFramework/Actor.h
+++ b/Engine/Source/Runtime/Engine/Classes/GameFramework/Actor.h
@@ -494,8 +494,14 @@ public:
 	float NetPriority;
 
 private:
-	/** Caches the most recent last render time we've looked at for this actor */
-	mutable float CachedLastRenderTime;
+	/**
+	 * The value of WorldSettings->TimeSeconds for the frame when one of this actor's components was last rendered.  This is written
+	 * from the render thread, which is up to a frame behind the game thread, so you should allow this time to
+	 * be at least a frame behind the game thread's world time before you consider the actor non-visible.
+	 */
+	float LastRenderTime;
+
+	friend struct FActorLastRenderTime;
 
 public:
 	/**
@@ -3070,6 +3076,24 @@ private:
 	friend UWorld;
 };
 
+/** Helper struct that allows UPrimitiveComponent and FPrimitiveSceneInfo write to the Actor's LastRenderTime member */
+struct FActorLastRenderTime
+{
+private:
+	static void Set(AActor* InActor, float LastRenderTime)
+	{
+		InActor->LastRenderTime = LastRenderTime;
+	}
+
+	static float* GetPtr(AActor* InActor)
+	{
+		return (InActor ? &InActor->LastRenderTime : nullptr);
+	}
+
+	friend class UPrimitiveComponent;
+	friend class FPrimitiveSceneInfo;
+};
+
 #if WITH_EDITOR
 struct FSetActorHiddenInSceneOutliner
 {
diff --git a/Engine/Source/Runtime/Engine/Private/Actor.cpp b/Engine/Source/Runtime/Engine/Private/Actor.cpp
--- a/Engine/Source/Runtime/Engine/Private/Actor.cpp
+++ b/Engine/Source/Runtime/Engine/Private/Actor.cpp
@@ -1535,31 +1535,14 @@ bool AActor::WasRecentlyRendered(float Tolerance) const
 		const float RenderTimeThreshold = FMath::Max(Tolerance, World->DeltaTimeSeconds + KINDA_SMALL_NUMBER);
 
 		// If the current cached value is less than the tolerance then we don't need to go look at the components
-		if (World->TimeSince(CachedLastRenderTime) <= RenderTimeThreshold)
-		{
-			return true;
-		}
-
-		CachedLastRenderTime = GetLastRenderTime(); // Store this off as an overriden version of GetLastRenderTime will not have set it directly internally
-
-		return (World->TimeSince(CachedLastRenderTime) <= RenderTimeThreshold);
+		return World->TimeSince(GetLastRenderTime()) <= RenderTimeThreshold;
 	}
 	return false;
 }
 
 float AActor::GetLastRenderTime() const
 {
-	// return most recent of Components' LastRenderTime
-	for (const UActorComponent* ActorComponent : GetComponents())
-	{
-		const UPrimitiveComponent* PrimComp = Cast<const UPrimitiveComponent>(ActorComponent);
-		if (PrimComp && PrimComp->IsRegistered())
-		{
-			CachedLastRenderTime = FMath::Max(CachedLastRenderTime, PrimComp->LastRenderTime);
-		}
-	}
-
-	return CachedLastRenderTime;
+	return LastRenderTime;
 }
 
 void AActor::SetOwner(AActor* NewOwner)
diff --git a/Engine/Source/Runtime/Engine/Private/Components/PrimitiveComponent.cpp b/Engine/Source/Runtime/Engine/Private/Components/PrimitiveComponent.cpp
--- a/Engine/Source/Runtime/Engine/Private/Components/PrimitiveComponent.cpp
+++ b/Engine/Source/Runtime/Engine/Private/Components/PrimitiveComponent.cpp
@@ -512,6 +512,9 @@ void UPrimitiveComponent::OnRegister()
 	{
 		bNavigationRelevant = false;
 	}
+
+	// Update our Owner's LastRenderTime
+	SetLastRenderTime(LastRenderTime);
 }
 
 
@@ -3467,6 +3470,13 @@ void UPrimitiveComponent::SetCustomNavigableGeometry(const EHasCustomNavigableGe
 void UPrimitiveComponent::SetLastRenderTime(float InLastRenderTime)
 {
 	LastRenderTime = InLastRenderTime;
+	if (AActor* Owner = GetOwner())
+	{
+		if (LastRenderTime > Owner->GetLastRenderTime())
+		{
+			FActorLastRenderTime::Set(Owner, LastRenderTime);
+		}
+	}
 }
 
 #if WITH_EDITOR
diff --git a/Engine/Source/Runtime/Renderer/Private/PrimitiveSceneInfo.cpp b/Engine/Source/Runtime/Renderer/Private/PrimitiveSceneInfo.cpp
--- a/Engine/Source/Runtime/Renderer/Private/PrimitiveSceneInfo.cpp
+++ b/Engine/Source/Runtime/Renderer/Private/PrimitiveSceneInfo.cpp
@@ -108,8 +108,7 @@ FPrimitiveSceneInfoCompact::FPrimitiveSceneInfoCompact(FPrimitiveSceneInfo* InPr
 FPrimitiveSceneInfo::FPrimitiveSceneInfo(UPrimitiveComponent* InComponent,FScene* InScene):
 	Proxy(InComponent->SceneProxy),
 	PrimitiveComponentId(InComponent->ComponentId),
-	ComponentLastRenderTime(&InComponent->LastRenderTime),
-	ComponentLastRenderTimeOnScreen(&InComponent->LastRenderTimeOnScreen),
+	OwnerLastRenderTime(FActorLastRenderTime::GetPtr(InComponent->GetOwner())),
 	IndirectLightingCacheAllocation(NULL),
 	CachedPlanarReflectionProxy(NULL),
 	CachedReflectionCaptureProxy(NULL),
@@ -1010,3 +1009,16 @@ void FPrimitiveSceneInfo::CacheReflectionCaptures()
 	
 	bNeedsCachedReflectionCaptureUpdate = false;
 }
+
+void FPrimitiveSceneInfo::UpdateComponentLastRenderTime(float CurrentWorldTime, bool bUpdateLastRenderTimeOnScreen) const
+{
+	ComponentForDebuggingOnly->LastRenderTime = CurrentWorldTime;
+	if (bUpdateLastRenderTimeOnScreen)
+	{
+		ComponentForDebuggingOnly->LastRenderTimeOnScreen = CurrentWorldTime;
+	}
+	if (OwnerLastRenderTime)
+	{
+		*OwnerLastRenderTime = CurrentWorldTime; // Sets OwningActor->LastRenderTime
+	}
+}
diff --git a/Engine/Source/Runtime/Renderer/Private/SceneVisibility.cpp b/Engine/Source/Runtime/Renderer/Private/SceneVisibility.cpp
--- a/Engine/Source/Runtime/Renderer/Private/SceneVisibility.cpp
+++ b/Engine/Source/Runtime/Renderer/Private/SceneVisibility.cpp
@@ -1979,9 +1979,7 @@ struct FRelevancePacket
 			// on the game thread. This signals that the primitive is visible.
 			if (View.PrimitiveDefinitelyUnoccludedMap[BitIndex] || (View.Family->EngineShowFlags.Wireframe && View.PrimitiveVisibilityMap[BitIndex]))
 			{
-				// Update the PrimitiveComponent's LastRenderTime.
-				*(PrimitiveSceneInfo->ComponentLastRenderTime) = CurrentWorldTime;
-				*(PrimitiveSceneInfo->ComponentLastRenderTimeOnScreen) = CurrentWorldTime;
+				PrimitiveSceneInfo->UpdateComponentLastRenderTime(CurrentWorldTime, /*bUpdateLastRenderTimeOnScreen=*/true);
 			}
 
 			// Cache the nearest reflection proxy if needed
diff --git a/Engine/Source/Runtime/Renderer/Private/ShadowSetup.cpp b/Engine/Source/Runtime/Renderer/Private/ShadowSetup.cpp
--- a/Engine/Source/Runtime/Renderer/Private/ShadowSetup.cpp
+++ b/Engine/Source/Runtime/Renderer/Private/ShadowSetup.cpp
@@ -1078,7 +1078,7 @@ void FProjectedShadowInfo::AddSubjectPrimitive(FPrimitiveSceneInfo* PrimitiveSce
 		{
 			// Update the primitive component's last render time. Allows the component to update when using bCastWhenHidden.
 			const float CurrentWorldTime = Views[0]->Family->CurrentWorldTime;
-			*(PrimitiveSceneInfo->ComponentLastRenderTime) = CurrentWorldTime;
+			PrimitiveSceneInfo->UpdateComponentLastRenderTime(CurrentWorldTime, /*bUpdateLastRenderTimeOnScreen=*/false);
 
 			if (PrimitiveSceneInfo->NeedsUniformBufferUpdate())
 			{
diff --git a/Engine/Source/Runtime/Renderer/Public/PrimitiveSceneInfo.h b/Engine/Source/Runtime/Renderer/Public/PrimitiveSceneInfo.h
--- a/Engine/Source/Runtime/Renderer/Public/PrimitiveSceneInfo.h
+++ b/Engine/Source/Runtime/Renderer/Public/PrimitiveSceneInfo.h
@@ -190,15 +190,12 @@ public:
 	FPrimitiveComponentId PrimitiveComponentId;
 
 	/** 
-	 * Pointer to the primitive's last render time variable, which is written to by the RT and read by the GT.
+	 * Pointer to the last render time variable on the primitive's owning actor (if owned), which is written to by the RT and read by the GT.
 	 * The value of LastRenderTime will therefore not be deterministic due to race conditions, but the GT uses it in a way that allows this.
-	 * Storing a pointer to the UObject member variable only works because UPrimitiveComponent has a mechanism to ensure it does not get deleted before the proxy (DetachFence).
+	 * Storing a pointer to the UObject member variable only works because UPrimitiveComponent and AActor has a mechanism to ensure it does not get deleted before the proxy (DetachFence).
 	 * In general feedback from the renderer to the game thread like this should be avoided.
 	 */
-	float* ComponentLastRenderTime;
-
-	/** Same as ComponentLastRenderTime but only updated if the component is on screen. Used by the texture streamer. */
-	float* ComponentLastRenderTimeOnScreen;
+	float* OwnerLastRenderTime;
 
 	/** 
 	 * The root attachment component id for use with lighting, if valid.
@@ -438,6 +435,9 @@ public:
 	/** Cache per-primitive reflection captures used for mobile/forward rendering */
 	void CacheReflectionCaptures();
 
+	/** Helper function for writing out to the last render times to the game thread */
+	void UpdateComponentLastRenderTime(float CurrentWorldTime, bool bUpdateLastRenderTimeOnScreen) const;
+
 #if RHI_RAYTRACING
 	RENDERER_API FRayTracingGeometryRHIRef GetStaticRayTracingGeometryInstance(int LodLevel);
 #endif
