Replaced global frame minimal graphic pipeline state TSet to a TSet per DynamicPassMeshDrawListContext. This remove RW lock contention when global TSet was accessed by multiple threads

#rnx

diff --git a/Engine/Source/Runtime/Renderer/Private/MeshDrawCommands.cpp b/Engine/Source/Runtime/Renderer/Private/MeshDrawCommands.cpp
--- a/Engine/Source/Runtime/Renderer/Private/MeshDrawCommands.cpp
+++ b/Engine/Source/Runtime/Renderer/Private/MeshDrawCommands.cpp
@@ -494,7 +494,8 @@ void GenerateDynamicMeshDrawCommands(
 	const TArray<const FStaticMeshBatch*, SceneRenderingAllocator>& DynamicMeshCommandBuildRequests,
 	int32 MaxNumBuildRequestElements,
 	FMeshCommandOneFrameArray& VisibleCommands,
-	FDynamicMeshDrawCommandStorage& MeshDrawCommandStorage
+	FDynamicMeshDrawCommandStorage& MeshDrawCommandStorage,
+	FGraphicsMinimalPipelineStateSet& MinimalPipelineStatePassSet
 )
 {
 	QUICK_SCOPE_CYCLE_COUNTER(STAT_GenerateDynamicMeshDrawCommands);
@@ -503,7 +504,8 @@ void GenerateDynamicMeshDrawCommands(
 
 	FDynamicPassMeshDrawListContext DynamicPassMeshDrawListContext(
 		MeshDrawCommandStorage,
-		VisibleCommands
+		VisibleCommands,
+		MinimalPipelineStatePassSet
 	);
 	PassMeshProcessor->SetDrawListContext(&DynamicPassMeshDrawListContext);
 
@@ -562,7 +564,8 @@ void GenerateMobileBasePassDynamicMeshDrawCommands(
 	const TArray<const FStaticMeshBatch*, SceneRenderingAllocator>& DynamicMeshCommandBuildRequests,
 	int32 MaxNumBuildRequestElements,
 	FMeshCommandOneFrameArray& VisibleCommands,
-	FDynamicMeshDrawCommandStorage& MeshDrawCommandStorage
+	FDynamicMeshDrawCommandStorage& MeshDrawCommandStorage,
+	FGraphicsMinimalPipelineStateSet& GraphicsMinimalPipelineStateSet
 )
 {
 	QUICK_SCOPE_CYCLE_COUNTER(STAT_GenerateMobileBasePassDynamicMeshDrawCommands);
@@ -571,7 +574,8 @@ void GenerateMobileBasePassDynamicMeshDrawCommands(
 
 	FDynamicPassMeshDrawListContext DynamicPassMeshDrawListContext(
 		MeshDrawCommandStorage,
-		VisibleCommands
+		VisibleCommands,
+		GraphicsMinimalPipelineStateSet
 	);
 	PassMeshProcessor->SetDrawListContext(&DynamicPassMeshDrawListContext);
 	MobilePassCSMPassMeshProcessor->SetDrawListContext(&DynamicPassMeshDrawListContext);
@@ -648,6 +652,7 @@ void ApplyViewOverridesToMeshDrawCommands(
 	FExclusiveDepthStencil::Type DefaultBasePassDepthStencilAccess,
 	FMeshCommandOneFrameArray& VisibleMeshDrawCommands,
 	FDynamicMeshDrawCommandStorage& MeshDrawCommandStorage,
+	FGraphicsMinimalPipelineStateSet& MinimalPipelineStatePassSet,
 	FMeshCommandOneFrameArray& TempVisibleMeshDrawCommands
 	)
 {
@@ -669,11 +674,7 @@ void ApplyViewOverridesToMeshDrawCommands(
 
 				const ERasterizerCullMode LocalCullMode = bRenderSceneTwoSided ? CM_None : bReverseCulling ? FMeshPassProcessor::InverseCullMode(VisibleMeshDrawCommand.MeshCullMode) : VisibleMeshDrawCommand.MeshCullMode;
 
-				FGraphicsMinimalPipelineStateInitializer PipelineState;
-				{
-					FGraphicsMinimalPipelineStateId::FPipelineStateIdLookupScope SafeLookupScope(MeshCommand.CachedPipelineId);
-					PipelineState = SafeLookupScope.GetPipelineState();
-				}
+				FGraphicsMinimalPipelineStateInitializer PipelineState = MeshCommand.CachedPipelineId.GetPipelineState(MinimalPipelineStatePassSet);
 				PipelineState.RasterizerState = GetStaticRasterizerState<true>(VisibleMeshDrawCommand.MeshFillMode, LocalCullMode);
 
 				if (BasePassDepthStencilAccess != DefaultBasePassDepthStencilAccess && PassType == EMeshPass::BasePass)
@@ -683,7 +684,7 @@ void ApplyViewOverridesToMeshDrawCommands(
 					PipelineState.DepthStencilState = PassDrawRenderState.GetDepthStencilState();
 				}
 
-				const FGraphicsMinimalPipelineStateId PipelineId = FGraphicsMinimalPipelineStateId::GetOneFrameId(PipelineState);
+				const FGraphicsMinimalPipelineStateId PipelineId = FGraphicsMinimalPipelineStateId::GetPipelineStateId(PipelineState, MinimalPipelineStatePassSet);
 				NewMeshCommand.Finalize(PipelineId, nullptr);
 
 				FVisibleMeshDrawCommand NewVisibleMeshDrawCommand;
@@ -765,7 +766,8 @@ public:
 				Context.DynamicMeshCommandBuildRequests,
 				Context.NumDynamicMeshCommandBuildRequestElements,
 				Context.MeshDrawCommands,
-				Context.MeshDrawCommandStorage
+				Context.MeshDrawCommandStorage,
+				Context.MinimalPipelineStatePassSet
 			);
 		}
 		else
@@ -781,7 +783,8 @@ public:
 				Context.DynamicMeshCommandBuildRequests,
 				Context.NumDynamicMeshCommandBuildRequestElements,
 				Context.MeshDrawCommands,
-				Context.MeshDrawCommandStorage
+				Context.MeshDrawCommandStorage,
+				Context.MinimalPipelineStatePassSet
 			);
 		}
 
@@ -798,6 +801,7 @@ public:
 					Context.DefaultBasePassDepthStencilAccess,
 					Context.MeshDrawCommands,
 					Context.MeshDrawCommandStorage,
+					Context.MinimalPipelineStatePassSet,
 					Context.TempVisibleMeshDrawCommands
 				);
 			}
@@ -1050,6 +1054,8 @@ void FParallelMeshDrawCommandPass::WaitForTasksAndEmpty()
 	TaskContext.DynamicMeshElementsPassRelevance = nullptr;
 	TaskContext.MeshDrawCommands.Empty();
 	TaskContext.MeshDrawCommandStorage.MeshDrawCommands.Empty();
+	FGraphicsMinimalPipelineStateId::AddSizeToLocalPipelineIdTableSize(TaskContext.MinimalPipelineStatePassSet.GetAllocatedSize());
+	TaskContext.MinimalPipelineStatePassSet.Empty();
 	TaskContext.MobileBasePassCSMMeshDrawCommands.Empty();
 	TaskContext.DynamicMeshCommandBuildRequests.Empty();
 	TaskContext.TempVisibleMeshDrawCommands.Empty();
@@ -1066,6 +1072,7 @@ class FDrawVisibleMeshCommandsAnyThreadTask : public FRenderTask
 {
 	FRHICommandList& RHICmdList;
 	const FMeshCommandOneFrameArray& VisibleMeshDrawCommands;
+	const FGraphicsMinimalPipelineStateSet& GraphicsMinimalPipelineStateSet;
 	FVertexBufferRHIParamRef PrimitiveIdsBuffer;
 	int32 BasePrimitiveIdsOffset;
 	bool bDynamicInstancing;
@@ -1078,6 +1085,7 @@ public:
 	FDrawVisibleMeshCommandsAnyThreadTask(
 		FRHICommandList& InRHICmdList,
 		const FMeshCommandOneFrameArray& InVisibleMeshDrawCommands,
+		const FGraphicsMinimalPipelineStateSet& InGraphicsMinimalPipelineStateSet,
 		FVertexBufferRHIParamRef InPrimitiveIdsBuffer,
 		int32 InBasePrimitiveIdsOffset,
 		bool bInDynamicInstancing,
@@ -1087,6 +1095,7 @@ public:
 	)
 		: RHICmdList(InRHICmdList)
 		, VisibleMeshDrawCommands(InVisibleMeshDrawCommands)
+		, GraphicsMinimalPipelineStateSet(InGraphicsMinimalPipelineStateSet)
 		, PrimitiveIdsBuffer(InPrimitiveIdsBuffer)
 		, BasePrimitiveIdsOffset(InBasePrimitiveIdsOffset)
 		, bDynamicInstancing(bInDynamicInstancing)
@@ -1112,7 +1121,7 @@ public:
 		const int32 StartIndex = TaskIndex * NumDrawsPerTask;
 		const int32 NumDraws = FMath::Min(NumDrawsPerTask, DrawNum - StartIndex);
 
-		SubmitMeshDrawCommandsRange(VisibleMeshDrawCommands, PrimitiveIdsBuffer, BasePrimitiveIdsOffset, bDynamicInstancing, StartIndex, NumDraws, InstanceFactor, RHICmdList);
+		SubmitMeshDrawCommandsRange(VisibleMeshDrawCommands, GraphicsMinimalPipelineStateSet, PrimitiveIdsBuffer, BasePrimitiveIdsOffset, bDynamicInstancing, StartIndex, NumDraws, InstanceFactor, RHICmdList);
 
 		RHICmdList.EndRenderPass();
 		RHICmdList.HandleRTThreadTaskCompletion(MyCompletionGraphEvent);
@@ -1211,7 +1220,7 @@ void FParallelMeshDrawCommandPass::DispatchDraw(FParallelCommandListSet* Paralle
 			FRHICommandList* CmdList = ParallelCommandListSet->NewParallelCommandList();
 
 			FGraphEventRef AnyThreadCompletionEvent = TGraphTask<FDrawVisibleMeshCommandsAnyThreadTask>::CreateTask(&Prereqs, RenderThread)
-				.ConstructAndDispatchWhenReady(*CmdList, TaskContext.MeshDrawCommands, PrimitiveIdsBuffer, BasePrimitiveIdsOffset, TaskContext.bDynamicInstancing, TaskContext.InstanceFactor, TaskIndex, NumTasks);
+				.ConstructAndDispatchWhenReady(*CmdList, TaskContext.MeshDrawCommands, TaskContext.MinimalPipelineStatePassSet, PrimitiveIdsBuffer, BasePrimitiveIdsOffset, TaskContext.bDynamicInstancing, TaskContext.InstanceFactor, TaskIndex, NumTasks);
 			ParallelCommandListSet->AddParallelCommandList(CmdList, AnyThreadCompletionEvent, NumDraws);
 		}
 	}
@@ -1229,7 +1238,7 @@ void FParallelMeshDrawCommandPass::DispatchDraw(FParallelCommandListSet* Paralle
 			RHIUnlockVertexBuffer(PrimitiveIdVertexBufferRHI);
 		}
 
-		SubmitMeshDrawCommandsRange(TaskContext.MeshDrawCommands, PrimitiveIdsBuffer, BasePrimitiveIdsOffset, TaskContext.bDynamicInstancing, 0, TaskContext.MeshDrawCommands.Num(), TaskContext.InstanceFactor, RHICmdList);
+		SubmitMeshDrawCommandsRange(TaskContext.MeshDrawCommands, TaskContext.MinimalPipelineStatePassSet, PrimitiveIdsBuffer, BasePrimitiveIdsOffset, TaskContext.bDynamicInstancing, 0, TaskContext.MeshDrawCommands.Num(), TaskContext.InstanceFactor, RHICmdList);
 	}
 }
 
diff --git a/Engine/Source/Runtime/Renderer/Private/MeshDrawCommands.h b/Engine/Source/Runtime/Renderer/Private/MeshDrawCommands.h
--- a/Engine/Source/Runtime/Renderer/Private/MeshDrawCommands.h
+++ b/Engine/Source/Runtime/Renderer/Private/MeshDrawCommands.h
@@ -96,6 +96,7 @@ public:
 	TArray<const FStaticMeshBatch*, SceneRenderingAllocator> DynamicMeshCommandBuildRequests;
 	TArray<const FStaticMeshBatch*, SceneRenderingAllocator> MobileBasePassCSMDynamicMeshCommandBuildRequests;
 	FDynamicMeshDrawCommandStorage MeshDrawCommandStorage;
+	FGraphicsMinimalPipelineStateSet MinimalPipelineStatePassSet;
 
 	// Resources preallocated on rendering thread.
 	void* PrimitiveIdBufferData;
diff --git a/Engine/Source/Runtime/Renderer/Private/MeshPassProcessor.cpp b/Engine/Source/Runtime/Renderer/Private/MeshPassProcessor.cpp
--- a/Engine/Source/Runtime/Renderer/Private/MeshPassProcessor.cpp
+++ b/Engine/Source/Runtime/Renderer/Private/MeshPassProcessor.cpp
@@ -15,8 +15,8 @@
 #include "PipelineStateCache.h"
 
 TSet<FRefCountedGraphicsMinimalPipelineStateInitializer, RefCountedGraphicsMinimalPipelineStateInitializerKeyFuncs> FGraphicsMinimalPipelineStateId::PersistentIdTable;
-TSet<FGraphicsMinimalPipelineStateInitializer> FGraphicsMinimalPipelineStateId::OneFrameIdTable;
-FRWLock FGraphicsMinimalPipelineStateId::OneFrameIdTableCriticalSection;
+int32 FGraphicsMinimalPipelineStateId::LocalPipelineIdTableSize = 0;
+int32 FGraphicsMinimalPipelineStateId::CurrentLocalPipelineIdTableSize = 0;
 
 const FMeshDrawCommandSortKey FMeshDrawCommandSortKey::Default = { {0} };
 
@@ -328,14 +328,14 @@ FGraphicsMinimalPipelineStateId FGraphicsMinimalPipelineStateId::GetPersistentId
 
 	FGraphicsMinimalPipelineStateId Ret;
 	Ret.bValid = 1;
-	Ret.bOneFrameId = 0;
+	Ret.bComesFromLocalPipelineStateSet = 0;
 	Ret.SetElementIndex = TableId.AsInteger();
 	return Ret;
 }
 
 void FGraphicsMinimalPipelineStateId::RemovePersistentId(FGraphicsMinimalPipelineStateId Id)
 {
-	check(!Id.bOneFrameId && Id.bValid);
+	check(!Id.bComesFromLocalPipelineStateSet && Id.bValid);
 
 	const FSetElementId SetElementId = FSetElementId::FromInteger(Id.SetElementIndex);
 	FRefCountedGraphicsMinimalPipelineStateInitializer& RefCountedStateInitializer = PersistentIdTable[SetElementId];
@@ -348,23 +348,20 @@ void FGraphicsMinimalPipelineStateId::RemovePersistentId(FGraphicsMinimalPipelin
 	}
 }
 
-FGraphicsMinimalPipelineStateId FGraphicsMinimalPipelineStateId::GetOneFrameId(const FGraphicsMinimalPipelineStateInitializer& InPipelineState)
+FGraphicsMinimalPipelineStateId FGraphicsMinimalPipelineStateId::GetPipelineStateId(const FGraphicsMinimalPipelineStateInitializer& InPipelineState, FGraphicsMinimalPipelineStateSet& InOutPassSet)
 {
-	// Need to lock as this is called from multiple parallel tasks during mesh draw command generation or patching.	
 	FGraphicsMinimalPipelineStateId Ret;
 	Ret.bValid = 1;
-	Ret.bOneFrameId = 0;
+	Ret.bComesFromLocalPipelineStateSet = 0;
 
 	FSetElementId TableId = PersistentIdTable.FindId(InPipelineState);
 	if (!TableId.IsValidId())
 	{
-		FRWScopeLock ScopeLock(OneFrameIdTableCriticalSection, SLT_Write);
-		Ret.bOneFrameId = 1;
-
-		TableId = OneFrameIdTable.FindId(InPipelineState);
+		Ret.bComesFromLocalPipelineStateSet = 1;
+		TableId = InOutPassSet.FindId(InPipelineState);
 		if (!TableId.IsValidId())
 		{
-			TableId = OneFrameIdTable.Add(InPipelineState);
+			TableId = InOutPassSet.Add(InPipelineState);
 		}
 	}
 
@@ -374,36 +371,15 @@ FGraphicsMinimalPipelineStateId FGraphicsMinimalPipelineStateId::GetOneFrameId(c
 	return Ret;
 }
 
-FGraphicsMinimalPipelineStateId::FPipelineStateIdLookupScope::FPipelineStateIdLookupScope(const FGraphicsMinimalPipelineStateId& LookupID)
-{
-	bLocked = false;
-	if (LookupID.bOneFrameId)
-	{
-		if (CVarSafeStateLookup.GetValueOnRenderThread())
-		{
-			OneFrameIdTableCriticalSection.ReadLock();
-			bLocked = true;
-		}
-	}
-	SafeStateRef = &LookupID.GetPipelineState();
-}
-
-FGraphicsMinimalPipelineStateId::FPipelineStateIdLookupScope::~FPipelineStateIdLookupScope()
-{
-	if (bLocked)
-	{
-		OneFrameIdTableCriticalSection.ReadUnlock();
-	}
-}
-
-const FGraphicsMinimalPipelineStateInitializer& FGraphicsMinimalPipelineStateId::FPipelineStateIdLookupScope::GetPipelineState()
+void FGraphicsMinimalPipelineStateId::ResetLocalPipelineIdTableSize()
 {
-	return *SafeStateRef;
+	LocalPipelineIdTableSize = CurrentLocalPipelineIdTableSize;
+	CurrentLocalPipelineIdTableSize = 0;
 }
 
-void FGraphicsMinimalPipelineStateId::ResetOneFrameIdTable()
+void FGraphicsMinimalPipelineStateId::AddSizeToLocalPipelineIdTableSize(SIZE_T Size)
 {
-	OneFrameIdTable.Reset();
+	CurrentLocalPipelineIdTableSize += int32(Size);
 }
 
 class FMeshDrawCommandStateCache
@@ -880,6 +856,7 @@ uint32 FMeshDrawShaderBindings::GetDynamicInstancingHash() const
 
 void FMeshDrawCommand::SubmitDraw(
 	const FMeshDrawCommand& RESTRICT MeshDrawCommand, 
+	const FGraphicsMinimalPipelineStateSet& GraphicsMinimalPipelineStateSet,
 	FVertexBufferRHIParamRef ScenePrimitiveIdsBuffer,
 	int32 PrimitiveIdOffset,
 	uint32 InstanceFactor,
@@ -921,8 +898,7 @@ void FMeshDrawCommand::SubmitDraw(
 #endif
 
 	{
-		FGraphicsMinimalPipelineStateId::FPipelineStateIdLookupScope SafeLookupScope(MeshDrawCommand.CachedPipelineId);
-		const FGraphicsMinimalPipelineStateInitializer& MeshPipelineState = SafeLookupScope.GetPipelineState();
+		const FGraphicsMinimalPipelineStateInitializer& MeshPipelineState = MeshDrawCommand.CachedPipelineId.GetPipelineState(GraphicsMinimalPipelineStateSet);
 
 		if (MeshDrawCommand.CachedPipelineId.GetId() != StateCache.PipelineId)
 		{
@@ -1003,17 +979,19 @@ void FMeshDrawCommand::SetDebugData(const FPrimitiveSceneProxy* PrimitiveScenePr
 
 void SubmitMeshDrawCommands(
 	const FMeshCommandOneFrameArray& VisibleMeshDrawCommands,
+	const FGraphicsMinimalPipelineStateSet& GraphicsMinimalPipelineStateSet, 
 	FVertexBufferRHIParamRef PrimitiveIdsBuffer,
 	int32 BasePrimitiveIdsOffset,
 	bool bDynamicInstancing,
 	uint32 InstanceFactor,
 	FRHICommandList& RHICmdList)
 {
-	SubmitMeshDrawCommandsRange(VisibleMeshDrawCommands, PrimitiveIdsBuffer, BasePrimitiveIdsOffset, bDynamicInstancing, 0, VisibleMeshDrawCommands.Num(), InstanceFactor, RHICmdList);
+	SubmitMeshDrawCommandsRange(VisibleMeshDrawCommands, GraphicsMinimalPipelineStateSet, PrimitiveIdsBuffer, BasePrimitiveIdsOffset, bDynamicInstancing, 0, VisibleMeshDrawCommands.Num(), InstanceFactor, RHICmdList);
 }
 
 void SubmitMeshDrawCommandsRange(
 	const FMeshCommandOneFrameArray& VisibleMeshDrawCommands,
+	const FGraphicsMinimalPipelineStateSet& GraphicsMinimalPipelineStateSet,
 	FVertexBufferRHIParamRef PrimitiveIdsBuffer,
 	int32 BasePrimitiveIdsOffset,
 	bool bDynamicInstancing,
@@ -1032,7 +1010,7 @@ void SubmitMeshDrawCommandsRange(
 		const FVisibleMeshDrawCommand& VisibleMeshDrawCommand = VisibleMeshDrawCommands[DrawCommandIndex];
 		const int32 PrimitiveIdBufferOffset = BasePrimitiveIdsOffset + (bDynamicInstancing ? VisibleMeshDrawCommand.PrimitiveIdBufferOffset : DrawCommandIndex) * sizeof(int32);
 		checkSlow(!bDynamicInstancing || VisibleMeshDrawCommand.PrimitiveIdBufferOffset >= 0);
-		FMeshDrawCommand::SubmitDraw(*VisibleMeshDrawCommand.MeshDrawCommand, PrimitiveIdsBuffer, PrimitiveIdBufferOffset, InstanceFactor, RHICmdList, StateCache);
+		FMeshDrawCommand::SubmitDraw(*VisibleMeshDrawCommand.MeshDrawCommand, GraphicsMinimalPipelineStateSet, PrimitiveIdsBuffer, PrimitiveIdBufferOffset, InstanceFactor, RHICmdList, StateCache);
 	}
 }
 
@@ -1041,6 +1019,7 @@ void DrawDynamicMeshPassPrivate(
 	FRHICommandList& RHICmdList,
 	FMeshCommandOneFrameArray& VisibleMeshDrawCommands,
 	FDynamicMeshDrawCommandStorage& DynamicMeshDrawCommandStorage,
+	FGraphicsMinimalPipelineStateSet& GraphicsMinimalPipelineStateSet,
 	uint32 InstanceFactor)
 {
 	if (VisibleMeshDrawCommands.Num() > 0)
@@ -1051,7 +1030,7 @@ void DrawDynamicMeshPassPrivate(
 
 		SortAndMergeDynamicPassMeshDrawCommands(View.GetFeatureLevel(), VisibleMeshDrawCommands, DynamicMeshDrawCommandStorage, PrimitiveIdVertexBuffer, InstanceFactor);
 
-		SubmitMeshDrawCommandsRange(VisibleMeshDrawCommands, PrimitiveIdVertexBuffer, 0, bDynamicInstancing, 0, VisibleMeshDrawCommands.Num(), InstanceFactor, RHICmdList);
+		SubmitMeshDrawCommandsRange(VisibleMeshDrawCommands, GraphicsMinimalPipelineStateSet, PrimitiveIdVertexBuffer, 0, bDynamicInstancing, 0, VisibleMeshDrawCommands.Num(), InstanceFactor, RHICmdList);
 	}
 }
 
diff --git a/Engine/Source/Runtime/Renderer/Private/RendererScene.cpp b/Engine/Source/Runtime/Renderer/Private/RendererScene.cpp
--- a/Engine/Source/Runtime/Renderer/Private/RendererScene.cpp
+++ b/Engine/Source/Runtime/Renderer/Private/RendererScene.cpp
@@ -730,7 +730,7 @@ void FScene::DumpMeshDrawCommandMemoryStats()
 	UE_LOG(LogRenderer, Log, TEXT("Primitive StaticMeshCommandInfos %.1fKb"), TotalStaticMeshCommandInfos / 1024.0f);
 	UE_LOG(LogRenderer, Log, TEXT("GPUScene CPU structures %.1fKb"), GPUScene.PrimitivesToUpdate.GetAllocatedSize() / 1024.0f);
 	UE_LOG(LogRenderer, Log, TEXT("PSO persistent Id table %.1fKb %d elements"), FGraphicsMinimalPipelineStateId::GetPersistentIdTableSize() / 1024.0f, FGraphicsMinimalPipelineStateId::GetPersistentIdNum());
-	UE_LOG(LogRenderer, Log, TEXT("PSO one frame Id %.1fKb"), FGraphicsMinimalPipelineStateId::GetOneFrameIdTableSize() / 1024.0f);
+	UE_LOG(LogRenderer, Log, TEXT("PSO one frame Id %.1fKb"), FGraphicsMinimalPipelineStateId::GetLocalPipelineIdTableSize() / 1024.0f);
 }
 
 template<typename T>
diff --git a/Engine/Source/Runtime/Renderer/Private/SceneRendering.cpp b/Engine/Source/Runtime/Renderer/Private/SceneRendering.cpp
--- a/Engine/Source/Runtime/Renderer/Private/SceneRendering.cpp
+++ b/Engine/Source/Runtime/Renderer/Private/SceneRendering.cpp
@@ -2957,7 +2957,7 @@ void FSceneRenderer::WaitForTasksClearSnapshotsAndDeleteSceneRenderer(FRHIComman
 
 	// Can relase only after all mesh pass tasks are finished.
 	GPrimitiveIdVertexBufferPool.DiscardAll();
-	FGraphicsMinimalPipelineStateId::ResetOneFrameIdTable();
+	FGraphicsMinimalPipelineStateId::ResetLocalPipelineIdTableSize();
 
 	delete LocalRootMark;
 }
diff --git a/Engine/Source/Runtime/Renderer/Private/ShadowRendering.cpp b/Engine/Source/Runtime/Renderer/Private/ShadowRendering.cpp
--- a/Engine/Source/Runtime/Renderer/Private/ShadowRendering.cpp
+++ b/Engine/Source/Runtime/Renderer/Private/ShadowRendering.cpp
@@ -616,7 +616,7 @@ void FProjectedShadowInfo::SetupProjectionStencilMask(
 		const FShadowMeshDrawCommandPass& ProjectionStencilingPass = ProjectionStencilingPasses[ViewIndex];
 		if (ProjectionStencilingPass.VisibleMeshDrawCommands.Num() > 0)
 		{
-			SubmitMeshDrawCommands(ProjectionStencilingPass.VisibleMeshDrawCommands, ProjectionStencilingPass.PrimitiveIdVertexBuffer, 0, bDynamicInstancing, bIsInstancedStereoEmulated ? 2 : 1, RHICmdList);
+			SubmitMeshDrawCommands(ProjectionStencilingPass.VisibleMeshDrawCommands, GraphicsMinimalPipelineStateSet, ProjectionStencilingPass.PrimitiveIdVertexBuffer, 0, bDynamicInstancing, bIsInstancedStereoEmulated ? 2 : 1, RHICmdList);
 		}
 
 		// Restore viewport
@@ -757,7 +757,7 @@ void FProjectedShadowInfo::SetupProjectionStencilMask(
 			const FShadowMeshDrawCommandPass& ProjectionStencilingPass = ProjectionStencilingPasses[ViewIndex];
 			if (ProjectionStencilingPass.VisibleMeshDrawCommands.Num() > 0)
 			{
-				SubmitMeshDrawCommands(ProjectionStencilingPass.VisibleMeshDrawCommands, ProjectionStencilingPass.PrimitiveIdVertexBuffer, 0, bDynamicInstancing, 1, RHICmdList);
+				SubmitMeshDrawCommands(ProjectionStencilingPass.VisibleMeshDrawCommands, GraphicsMinimalPipelineStateSet, ProjectionStencilingPass.PrimitiveIdVertexBuffer, 0, bDynamicInstancing, 1, RHICmdList);
 			}
 		}
 	}
diff --git a/Engine/Source/Runtime/Renderer/Private/ShadowRendering.h b/Engine/Source/Runtime/Renderer/Private/ShadowRendering.h
--- a/Engine/Source/Runtime/Renderer/Private/ShadowRendering.h
+++ b/Engine/Source/Runtime/Renderer/Private/ShadowRendering.h
@@ -537,6 +537,7 @@ private:
 	TArray<FShadowMeshDrawCommandPass, TInlineAllocator<2>> ProjectionStencilingPasses;
 
 	FDynamicMeshDrawCommandStorage DynamicMeshDrawCommandStorage;
+	FGraphicsMinimalPipelineStateSet GraphicsMinimalPipelineStateSet;
 
 	/**
 	 * Bias during in shadowmap rendering, stored redundantly for better performance 
diff --git a/Engine/Source/Runtime/Renderer/Private/ShadowSetup.cpp b/Engine/Source/Runtime/Renderer/Private/ShadowSetup.cpp
--- a/Engine/Source/Runtime/Renderer/Private/ShadowSetup.cpp
+++ b/Engine/Source/Runtime/Renderer/Private/ShadowSetup.cpp
@@ -1224,7 +1224,7 @@ void FProjectedShadowInfo::SetupMeshDrawCommandsForProjectionStenciling(FSceneRe
 			ProjectionStencilingPasses.Add(FShadowMeshDrawCommandPass());
 			FShadowMeshDrawCommandPass& ProjectionStencilingPass = ProjectionStencilingPasses[ViewIndex];
 
-			FDynamicPassMeshDrawListContext ProjectionStencilingContext(DynamicMeshDrawCommandStorage, ProjectionStencilingPass.VisibleMeshDrawCommands);
+			FDynamicPassMeshDrawListContext ProjectionStencilingContext(DynamicMeshDrawCommandStorage, ProjectionStencilingPass.VisibleMeshDrawCommands, GraphicsMinimalPipelineStateSet);
 
 			FMeshPassProcessorRenderState DrawRenderState;
 			DrawRenderState.SetBlendState(TStaticBlendState<CW_NONE>::GetRHI());
@@ -1347,14 +1347,10 @@ void FProjectedShadowInfo::ApplyViewOverridesToMeshDrawCommands(const FViewInfo&
 
 			const ERasterizerCullMode LocalCullMode = View.bRenderSceneTwoSided ? CM_None : View.bReverseCulling ? FMeshPassProcessor::InverseCullMode(VisibleMeshDrawCommand.MeshCullMode) : VisibleMeshDrawCommand.MeshCullMode;
 
-			FGraphicsMinimalPipelineStateInitializer PipelineState;
-			{
-				FGraphicsMinimalPipelineStateId::FPipelineStateIdLookupScope SafeLookupScope(MeshCommand.CachedPipelineId);
-				PipelineState = SafeLookupScope.GetPipelineState();
-			}
+			FGraphicsMinimalPipelineStateInitializer PipelineState = MeshCommand.CachedPipelineId.GetPipelineState(GraphicsMinimalPipelineStateSet);
 			PipelineState.RasterizerState = GetStaticRasterizerState<true>(VisibleMeshDrawCommand.MeshFillMode, LocalCullMode);
 
-			const FGraphicsMinimalPipelineStateId PipelineId = FGraphicsMinimalPipelineStateId::GetOneFrameId(PipelineState);
+			const FGraphicsMinimalPipelineStateId PipelineId = FGraphicsMinimalPipelineStateId::GetPipelineStateId(PipelineState, GraphicsMinimalPipelineStateSet);
 			NewMeshCommand.Finalize(PipelineId, nullptr);
 
 			FVisibleMeshDrawCommand NewVisibleMeshDrawCommand;
@@ -1551,6 +1547,7 @@ void FProjectedShadowInfo::ClearTransientArrays()
 	ProjectionStencilingPasses.Reset();
 
 	DynamicMeshDrawCommandStorage.MeshDrawCommands.Empty();
+	GraphicsMinimalPipelineStateSet.Empty();
 }
 
 /** Returns a cached preshadow matching the input criteria if one exists. */
diff --git a/Engine/Source/Runtime/Renderer/Private/TranslucentLighting.cpp b/Engine/Source/Runtime/Renderer/Private/TranslucentLighting.cpp
--- a/Engine/Source/Runtime/Renderer/Private/TranslucentLighting.cpp
+++ b/Engine/Source/Runtime/Renderer/Private/TranslucentLighting.cpp
@@ -575,7 +575,7 @@ void FProjectedShadowInfo::RenderTranslucencyDepths(FRHICommandList& RHICmdList,
 
 
 		FMeshCommandOneFrameArray VisibleMeshDrawCommands;
-		FDynamicPassMeshDrawListContext TranslucencyDepthContext(DynamicMeshDrawCommandStorage, VisibleMeshDrawCommands);
+		FDynamicPassMeshDrawListContext TranslucencyDepthContext(DynamicMeshDrawCommandStorage, VisibleMeshDrawCommands, GraphicsMinimalPipelineStateSet);
 
 		FTranslucencyDepthPassMeshProcessor TranslucencyDepthPassMeshProcessor(
 			SceneRenderer->Scene,
@@ -622,7 +622,7 @@ void FProjectedShadowInfo::RenderTranslucencyDepths(FRHICommandList& RHICmdList,
 			FVertexBufferRHIParamRef PrimitiveIdVertexBuffer = nullptr;
 			ApplyViewOverridesToMeshDrawCommands(*ShadowDepthView, VisibleMeshDrawCommands);
 			SortAndMergeDynamicPassMeshDrawCommands(SceneRenderer->FeatureLevel, VisibleMeshDrawCommands, DynamicMeshDrawCommandStorage, PrimitiveIdVertexBuffer, 1);
-			SubmitMeshDrawCommands(VisibleMeshDrawCommands, PrimitiveIdVertexBuffer, 0, bDynamicInstancing, 1, RHICmdList);
+			SubmitMeshDrawCommands(VisibleMeshDrawCommands, GraphicsMinimalPipelineStateSet, PrimitiveIdVertexBuffer, 0, bDynamicInstancing, 1, RHICmdList);
 		}
 	}
 }
diff --git a/Engine/Source/Runtime/Renderer/Public/MeshPassProcessor.h b/Engine/Source/Runtime/Renderer/Public/MeshPassProcessor.h
--- a/Engine/Source/Runtime/Renderer/Public/MeshPassProcessor.h
+++ b/Engine/Source/Runtime/Renderer/Public/MeshPassProcessor.h
@@ -131,6 +131,9 @@ struct RefCountedGraphicsMinimalPipelineStateInitializerKeyFuncs : DefaultKeyFun
 	}
 };
 
+/** Set of FGraphicsMinimalPipelineStateInitializer unique per MeshDrawCommandsPassContext */
+typedef TSet< FGraphicsMinimalPipelineStateInitializer > FGraphicsMinimalPipelineStateSet;
+
 /** Uniquely represents a FGraphicsMinimalPipelineStateInitializer for fast compares. */
 class FGraphicsMinimalPipelineStateId
 {
@@ -156,6 +159,18 @@ public:
 		return !(*this == rhs);
 	}
 	
+	inline const FGraphicsMinimalPipelineStateInitializer& GetPipelineState(const FGraphicsMinimalPipelineStateSet& InPipelineSet) const
+	{
+		const FSetElementId SetElementId = FSetElementId::FromInteger(SetElementIndex);
+
+		if (bComesFromLocalPipelineStateSet)
+		{
+			return InPipelineSet[SetElementId];
+		}
+
+		return PersistentIdTable[SetElementId].StateInitializer;
+	}
+
 	/**
 	 * Get a ref counted persistent pipeline id, which needs to manually released.
 	 */
@@ -165,62 +180,36 @@ public:
 	 * Removes a persistent pipeline Id from the global persistent Id table.
 	 */
 	static void RemovePersistentId(FGraphicsMinimalPipelineStateId Id);
-
+	
 	/**
-	 * Get a pipeline id, which is valid only for a single frame and doesn't need to be released manually.
+	 * Get a pipeline state id in this order: global persistent Id table. If not found, will lookup in PassSet argument. If not found in PassSet argument, create a blank pipeline set id and add it PassSet argument
 	 */
-	RENDERER_API static FGraphicsMinimalPipelineStateId GetOneFrameId(const FGraphicsMinimalPipelineStateInitializer& InPipelineState);
+	RENDERER_API static FGraphicsMinimalPipelineStateId GetPipelineStateId(const FGraphicsMinimalPipelineStateInitializer& InPipelineState, FGraphicsMinimalPipelineStateSet& InOutPassSet);
+
+	static int32 GetLocalPipelineIdTableSize() { return LocalPipelineIdTableSize; }
+	static void ResetLocalPipelineIdTableSize();
+	static void AddSizeToLocalPipelineIdTableSize(SIZE_T Size);
 
-	static void ResetOneFrameIdTable();
 	static SIZE_T GetPersistentIdTableSize() { return PersistentIdTable.GetAllocatedSize(); }
 	static int32 GetPersistentIdNum() { return PersistentIdTable.Num(); }
-	static SIZE_T GetOneFrameIdTableSize() { return OneFrameIdTable.GetAllocatedSize(); }
-
-	class FPipelineStateIdLookupScope
-	{
-	public:
-		FPipelineStateIdLookupScope(const FGraphicsMinimalPipelineStateId& LookupID);
-		~FPipelineStateIdLookupScope();
-		
-		const FGraphicsMinimalPipelineStateInitializer& GetPipelineState();
-	private:
-		FPipelineStateIdLookupScope() {}
-		const FGraphicsMinimalPipelineStateInitializer* SafeStateRef;
-		bool bLocked;
-	};
 
 private:
-
-	friend class FPipelineStateIdLookupScope;
-
-	//this reference isn't safe unless it's wrapped in a scope with an FRWLock around OneFrameIDTable.
-	inline const FGraphicsMinimalPipelineStateInitializer& GetPipelineState() const
-	{
-		const FSetElementId SetElementId = FSetElementId::FromInteger(SetElementIndex);
-
-		if (bOneFrameId)
-		{
-			return OneFrameIdTable[SetElementId];
-		}
-
-		return PersistentIdTable[SetElementId].StateInitializer;
-	}
-
 	union
 	{
 		uint32 PackedId = 0;
 
 		struct
 		{
-			uint32 SetElementIndex	: 30;
-			uint32 bOneFrameId		: 1;
-			uint32 bValid			: 1;
+			uint32 SetElementIndex				   : 30;
+			uint32 bComesFromLocalPipelineStateSet : 1;
+			uint32 bValid						   : 1;
 		};
 	};
 
 	static TSet<FRefCountedGraphicsMinimalPipelineStateInitializer, RefCountedGraphicsMinimalPipelineStateInitializerKeyFuncs> PersistentIdTable;
-	static TSet<FGraphicsMinimalPipelineStateInitializer> OneFrameIdTable;
-	static FRWLock OneFrameIdTableCriticalSection;
+	
+	static int32 LocalPipelineIdTableSize;
+	static int32 CurrentLocalPipelineIdTableSize;
 };
 
 struct FMeshProcessorShaders
@@ -564,7 +553,8 @@ public:
 	/** Submits commands to the RHI Commandlist to draw the MeshDrawCommand. */
 	static void SubmitDraw(
 		const FMeshDrawCommand& RESTRICT MeshDrawCommand, 
-		FVertexBufferRHIParamRef ScenePrimitiveIdsBuffer, 
+		const FGraphicsMinimalPipelineStateSet& GraphicsMinimalPipelineStateSet,
+		FVertexBufferRHIParamRef ScenePrimitiveIdsBuffer,
 		int32 PrimitiveIdOffset,
 		uint32 InstanceFactor,
 		FRHICommandList& CommandList, 
@@ -735,10 +725,12 @@ public:
 	FDynamicPassMeshDrawListContext
 	(
 		FDynamicMeshDrawCommandStorage& InDrawListStorage, 
-		FMeshCommandOneFrameArray& InDrawList
+		FMeshCommandOneFrameArray& InDrawList,
+		FGraphicsMinimalPipelineStateSet& InPipelineStateSet
 	) :
 		DrawListStorage(InDrawListStorage),
-		DrawList(InDrawList)
+		DrawList(InDrawList),
+		GraphicsMinimalPipelineStateSet(InPipelineStateSet)
 	{}
 
 	virtual FMeshDrawCommand& AddCommand(const FMeshDrawCommand& Initializer) override final
@@ -759,8 +751,7 @@ public:
 		const FMeshProcessorShaders* ShadersForDebugging,
 		FMeshDrawCommand& MeshDrawCommand) override final
 	{
-		FGraphicsMinimalPipelineStateId PipelineId;
-		PipelineId = FGraphicsMinimalPipelineStateId::GetOneFrameId(PipelineState);
+		FGraphicsMinimalPipelineStateId PipelineId = FGraphicsMinimalPipelineStateId::GetPipelineStateId(PipelineState, GraphicsMinimalPipelineStateSet);
 
 		MeshDrawCommand.SetDrawParametersAndFinalize(MeshBatch, BatchElementIndex, PipelineId, ShadersForDebugging);
 
@@ -774,6 +765,7 @@ public:
 private:
 	FDynamicMeshDrawCommandStorage& DrawListStorage;
 	FMeshCommandOneFrameArray& DrawList;
+	FGraphicsMinimalPipelineStateSet& GraphicsMinimalPipelineStateSet;
 };
 
 #if PLATFORM_SUPPORTS_PRAGMA_PACK
@@ -1148,6 +1140,7 @@ private:
 
 extern void SubmitMeshDrawCommands(
 	const FMeshCommandOneFrameArray& VisibleMeshDrawCommands,
+	const FGraphicsMinimalPipelineStateSet& GraphicsMinimalPipelineStateSet, 
 	FVertexBufferRHIParamRef PrimitiveIdsBuffer,
 	int32 BasePrimitiveIdsOffset,
 	bool bDynamicInstancing,
@@ -1156,6 +1149,7 @@ extern void SubmitMeshDrawCommands(
 
 extern void SubmitMeshDrawCommandsRange(
 	const FMeshCommandOneFrameArray& VisibleMeshDrawCommands,
+	const FGraphicsMinimalPipelineStateSet& GraphicsMinimalPipelineStateSet,
 	FVertexBufferRHIParamRef PrimitiveIdsBuffer,
 	int32 BasePrimitiveIdsOffset,
 	bool bDynamicInstancing,
@@ -1169,6 +1163,7 @@ RENDERER_API extern void DrawDynamicMeshPassPrivate(
 	FRHICommandList& RHICmdList,
 	FMeshCommandOneFrameArray& VisibleMeshDrawCommands,
 	FDynamicMeshDrawCommandStorage& DynamicMeshDrawCommandStorage,
+	FGraphicsMinimalPipelineStateSet& GraphicsMinimalPipelineStateSet,
 	uint32 InstanceFactor);
 
 RENDERER_API extern FMeshDrawCommandSortKey CalculateMeshStaticSortKey(const FMeshMaterialShader* VertexShader, const FMeshMaterialShader* PixelShader);
diff --git a/Engine/Source/Runtime/Renderer/Public/MeshPassProcessor.inl b/Engine/Source/Runtime/Renderer/Public/MeshPassProcessor.inl
--- a/Engine/Source/Runtime/Renderer/Public/MeshPassProcessor.inl
+++ b/Engine/Source/Runtime/Renderer/Public/MeshPassProcessor.inl
@@ -145,12 +145,13 @@ void DrawDynamicMeshPass(const FSceneView& View, FRHICommandList& RHICmdList, co
 {
 	FDynamicMeshDrawCommandStorage DynamicMeshDrawCommandStorage;
 	FMeshCommandOneFrameArray VisibleMeshDrawCommands;
+	FGraphicsMinimalPipelineStateSet GraphicsMinimalPipelineStateSet;
 
-	FDynamicPassMeshDrawListContext DynamicMeshPassContext(DynamicMeshDrawCommandStorage, VisibleMeshDrawCommands);
+	FDynamicPassMeshDrawListContext DynamicMeshPassContext(DynamicMeshDrawCommandStorage, VisibleMeshDrawCommands, GraphicsMinimalPipelineStateSet);
 
 	BuildPassProcessorLambda(&DynamicMeshPassContext);
 
 	// We assume all dynamic passes are in stereo if it is enabled in the view, so we apply ISR to them
 	const uint32 InstanceFactor = View.IsInstancedStereoPass() ? 2 : 1;
-	DrawDynamicMeshPassPrivate(View, RHICmdList, VisibleMeshDrawCommands, DynamicMeshDrawCommandStorage, InstanceFactor);
+	DrawDynamicMeshPassPrivate(View, RHICmdList, VisibleMeshDrawCommands, DynamicMeshDrawCommandStorage, GraphicsMinimalPipelineStateSet, InstanceFactor);
 }
