Conflicts from //UE4/Main merge.

#rnx

diff --git a/Engine/Source/Runtime/Renderer/Private/ByteBuffer.cpp b/Engine/Source/Runtime/Renderer/Private/ByteBuffer.cpp
--- a/Engine/Source/Runtime/Renderer/Private/ByteBuffer.cpp
+++ b/Engine/Source/Runtime/Renderer/Private/ByteBuffer.cpp
@@ -16,7 +16,7 @@ class FMemsetBufferCS : public FGlobalShader
 	
 	static bool ShouldCompilePermutation( const FGlobalShaderPermutationParameters& Parameters )
 	{
-		return IsFeatureLevelSupported( Parameters.Platform, ERHIFeatureLevel::SM5 )
+		return RHISupportsComputeShaders(Parameters.Platform)
 			// @todo platplug can we remove ths property now that we are just checking SM5 instead of platforms (and all uses below)?
 			|| FDataDrivenShaderPlatformInfo::GetInfo(Parameters.Platform).bSupportsByteBufferComputeShaders;
 	}
@@ -87,7 +87,7 @@ class FMemcpyBufferCS : public FGlobalShader
 	
 	static bool ShouldCompilePermutation( const FGlobalShaderPermutationParameters& Parameters )
 	{
-		return IsFeatureLevelSupported( Parameters.Platform, ERHIFeatureLevel::SM5 )
+		return RHISupportsComputeShaders(Parameters.Platform)
 			|| FDataDrivenShaderPlatformInfo::GetInfo(Parameters.Platform).bSupportsByteBufferComputeShaders;
 	}
 
@@ -187,7 +187,7 @@ class FScatterCopyCS : public FGlobalShader
 	
 	static bool ShouldCompilePermutation( const FGlobalShaderPermutationParameters& Parameters )
 	{
-		return IsFeatureLevelSupported( Parameters.Platform, ERHIFeatureLevel::SM5 )
+		return RHISupportsComputeShaders(Parameters.Platform)
 			|| FDataDrivenShaderPlatformInfo::GetInfo(Parameters.Platform).bSupportsByteBufferComputeShaders;
 	}
 
