Optimize frustum cull

diff --git a/Engine/Source/Runtime/Renderer/Private/SceneVisibility.cpp b/Engine/Source/Runtime/Renderer/Private/SceneVisibility.cpp
--- a/Engine/Source/Runtime/Renderer/Private/SceneVisibility.cpp
+++ b/Engine/Source/Runtime/Renderer/Private/SceneVisibility.cpp
@@ -348,6 +348,67 @@ bool FViewInfo::IsDistanceCulled( float DistanceSquared, float MinDrawDistance,
 	return ( bDistanceCulled && !bStillFading );
 }
 
+FORCEINLINE bool IntersectBox8Plane(const FVector& InOrigin, const FVector& InExtent, const FPlane*PermutedPlanePtr)
+{
+	// this removes a lot of the branches as we know there's 8 planes
+	// copied directly out of ConvexVolume.cpp
+	const VectorRegister Origin = VectorLoadFloat3(&InOrigin);
+	const VectorRegister Extent = VectorLoadFloat3(&InExtent);
+
+	const VectorRegister PlanesX_0 = VectorLoadAligned(&PermutedPlanePtr[0]);
+	const VectorRegister PlanesY_0 = VectorLoadAligned(&PermutedPlanePtr[1]);
+	const VectorRegister PlanesZ_0 = VectorLoadAligned(&PermutedPlanePtr[2]);
+	const VectorRegister PlanesW_0 = VectorLoadAligned(&PermutedPlanePtr[3]);
+
+	const VectorRegister PlanesX_1 = VectorLoadAligned(&PermutedPlanePtr[4]);
+	const VectorRegister PlanesY_1 = VectorLoadAligned(&PermutedPlanePtr[5]);
+	const VectorRegister PlanesZ_1 = VectorLoadAligned(&PermutedPlanePtr[6]);
+	const VectorRegister PlanesW_1 = VectorLoadAligned(&PermutedPlanePtr[7]);
+
+	// Splat origin into 3 vectors
+	VectorRegister OrigX = VectorReplicate(Origin, 0);
+	VectorRegister OrigY = VectorReplicate(Origin, 1);
+	VectorRegister OrigZ = VectorReplicate(Origin, 2);
+	// Splat the already abs Extent for the push out calculation
+	VectorRegister AbsExtentX = VectorReplicate(Extent, 0);
+	VectorRegister AbsExtentY = VectorReplicate(Extent, 1);
+	VectorRegister AbsExtentZ = VectorReplicate(Extent, 2);
+
+	// Calculate the distance (x * x) + (y * y) + (z * z) - w
+	VectorRegister DistX_0 = VectorMultiply(OrigX, PlanesX_0);
+	VectorRegister DistY_0 = VectorMultiplyAdd(OrigY, PlanesY_0, DistX_0);
+	VectorRegister DistZ_0 = VectorMultiplyAdd(OrigZ, PlanesZ_0, DistY_0);
+	VectorRegister Distance_0 = VectorSubtract(DistZ_0, PlanesW_0);
+	// Now do the push out FMath::Abs(x * x) + FMath::Abs(y * y) + FMath::Abs(z * z)
+	VectorRegister PushX_0 = VectorMultiply(AbsExtentX, VectorAbs(PlanesX_0));
+	VectorRegister PushY_0 = VectorMultiplyAdd(AbsExtentY, VectorAbs(PlanesY_0), PushX_0);
+	VectorRegister PushOut_0 = VectorMultiplyAdd(AbsExtentZ, VectorAbs(PlanesZ_0), PushY_0);
+
+	// Check for completely outside
+	if (VectorAnyGreaterThan(Distance_0, PushOut_0))
+	{
+		return false;
+	}
+
+	// Calculate the distance (x * x) + (y * y) + (z * z) - w
+	VectorRegister DistX_1 = VectorMultiply(OrigX, PlanesX_1);
+	VectorRegister DistY_1 = VectorMultiplyAdd(OrigY, PlanesY_1, DistX_1);
+	VectorRegister DistZ_1 = VectorMultiplyAdd(OrigZ, PlanesZ_1, DistY_1);
+	VectorRegister Distance_1 = VectorSubtract(DistZ_1, PlanesW_1);
+	// Now do the push out FMath::Abs(x * x) + FMath::Abs(y * y) + FMath::Abs(z * z)
+	VectorRegister PushX_1 = VectorMultiply(AbsExtentX, VectorAbs(PlanesX_1));
+	VectorRegister PushY_1 = VectorMultiplyAdd(AbsExtentY, VectorAbs(PlanesY_1), PushX_1);
+	VectorRegister PushOut_1 = VectorMultiplyAdd(AbsExtentZ, VectorAbs(PlanesZ_1), PushY_1);
+
+	// Check for completely outside
+	if (VectorAnyGreaterThan(Distance_1, PushOut_1))
+	{
+		return false;
+	}
+	return true;
+}
+
+
 static int32 FrustumCullNumWordsPerTask = 128;
 static FAutoConsoleVariableRef CVarFrustumCullNumWordsPerTask(
 	TEXT("r.FrustumCullNumWordsPerTask"),
@@ -357,7 +418,7 @@ static FAutoConsoleVariableRef CVarFrustumCullNumWordsPerTask(
 	);
 
 
-template<bool UseCustomCulling, bool bAlsoUseSphereTest>
+template<bool UseCustomCulling, bool bAlsoUseSphereTest, bool bUseFastIntersect>
 static int32 FrustumCull(const FScene* Scene, FViewInfo& View)
 {
 	SCOPE_CYCLE_COUNTER(STAT_FrustumCull);
@@ -383,6 +444,7 @@ static int32 FrustumCull(const FScene* Scene, FViewInfo& View)
 		[&NumCulledPrimitives, Scene, &View, MaxDrawDistanceScale, HLODState](int32 TaskIndex)
 		{
 			QUICK_SCOPE_CYCLE_COUNTER(STAT_FrustumCull_Loop);
+			const FPlane* PermutedPlanePtr = View.ViewFrustum.PermutedPlanes.GetData();
 			const int32 BitArrayNumInner = View.PrimitiveVisibilityMap.Num();
 			FVector ViewOriginForDistanceCulling = View.ViewMatrices.GetViewOrigin();
 			float FadeRadius = GDisableLODFade ? 0.0f : GDistanceFadeMaxTravel;
@@ -435,7 +497,7 @@ static int32 FrustumCull(const FScene* Scene, FViewInfo& View)
 						(DistanceSquared < MinDrawDistanceSq) ||
 						(UseCustomCulling && !View.CustomVisibilityQuery->IsVisible(VisibilityId, FBoxSphereBounds(Bounds.BoxSphereBounds.Origin, Bounds.BoxSphereBounds.BoxExtent, Bounds.BoxSphereBounds.SphereRadius))) ||
 						(bAlsoUseSphereTest && View.ViewFrustum.IntersectSphere(Bounds.BoxSphereBounds.Origin, Bounds.BoxSphereBounds.SphereRadius) == false) ||
-						View.ViewFrustum.IntersectBox(Bounds.BoxSphereBounds.Origin, Bounds.BoxSphereBounds.BoxExtent) == false)
+						(bUseFastIntersect ? IntersectBox8Plane(Bounds.BoxSphereBounds.Origin, Bounds.BoxSphereBounds.BoxExtent, PermutedPlanePtr) : View.ViewFrustum.IntersectBox(Bounds.BoxSphereBounds.Origin, Bounds.BoxSphereBounds.BoxExtent)) == false)
 					{
 						STAT(NumCulledPrimitives.Increment());
 					}
@@ -3282,6 +3344,13 @@ static TAutoConsoleVariable<int32> CVarAlsoUseSphereForFrustumCull(
 	ECVF_RenderThreadSafe
 	);
 
+static TAutoConsoleVariable<int32> CVarUseFastIntersect(
+	TEXT("r.UseFastIntersect"),
+	1,
+	TEXT("Use optimized 8 plane fast intersection code if we have 8 permuted planes."),
+	ECVF_RenderThreadSafe
+);
+
 
 void UpdateReflectionSceneData(FScene* Scene)
 {
@@ -3586,26 +3655,27 @@ void FSceneRenderer::ComputeViewVisibility(FRHICommandListImmediate& RHICmdList,
 			}
 
 			int32 NumCulledPrimitivesForView;
+			const bool bUseFastIntersect = (View.ViewFrustum.PermutedPlanes.Num() == 8) && CVarUseFastIntersect.GetValueOnRenderThread();
 			if (View.CustomVisibilityQuery && View.CustomVisibilityQuery->Prepare())
 			{
 				if (CVarAlsoUseSphereForFrustumCull.GetValueOnRenderThread())
 				{
-					NumCulledPrimitivesForView = FrustumCull<true, true>(Scene, View);
+					NumCulledPrimitivesForView = bUseFastIntersect ? FrustumCull<true, true, true>(Scene, View) : FrustumCull<true, true, false>(Scene, View);
 				}
 				else
 				{
-					NumCulledPrimitivesForView = FrustumCull<true, false>(Scene, View);
+					NumCulledPrimitivesForView = bUseFastIntersect ? FrustumCull<true, false, true>(Scene, View) : FrustumCull<true, false, false>(Scene, View);
 				}
 			}
 			else
 			{
 				if (CVarAlsoUseSphereForFrustumCull.GetValueOnRenderThread())
 				{
-					NumCulledPrimitivesForView = FrustumCull<false, true>(Scene, View);
+					NumCulledPrimitivesForView = bUseFastIntersect ? FrustumCull<false, true, true>(Scene, View) : FrustumCull<false, true, false>(Scene, View);
 				}
 				else
 				{
-					NumCulledPrimitivesForView = FrustumCull<false, false>(Scene, View);
+					NumCulledPrimitivesForView = bUseFastIntersect ? FrustumCull<false, false, true>(Scene, View) : FrustumCull<false, false, false>(Scene, View);
 				}
 			}
 			STAT(NumCulledPrimitives += NumCulledPrimitivesForView);
