Fix to remove stall introduced in Adreno occlusion mode

diff --git a/Engine/Source/Runtime/Renderer/Private/MobileShadingRenderer.cpp b/Engine/Source/Runtime/Renderer/Private/MobileShadingRenderer.cpp
--- a/Engine/Source/Runtime/Renderer/Private/MobileShadingRenderer.cpp
+++ b/Engine/Source/Runtime/Renderer/Private/MobileShadingRenderer.cpp
@@ -572,7 +572,7 @@ void FMobileSceneRenderer::Render(FRHICommandListImmediate& RHICmdList)
 	{
 	    RHICmdList.SetCurrentStat(GET_STATID(STAT_CLMM_Occlusion));
 		// flush
-		RHICmdList.SubmitCommandsAndFlushGPU();
+		RHICmdList.SubmitCommandsHint();
 		// Issue occlusion queries
 	    RenderOcclusion(RHICmdList);
 	    RHICmdList.ImmediateFlush(EImmediateFlushType::DispatchToRHIThread);
