Add a missing barrier for SceneColor writeable->readable in mobile renderer

diff --git a/Engine/Source/Runtime/Renderer/Private/MobileShadingRenderer.cpp b/Engine/Source/Runtime/Renderer/Private/MobileShadingRenderer.cpp
--- a/Engine/Source/Runtime/Renderer/Private/MobileShadingRenderer.cpp
+++ b/Engine/Source/Runtime/Renderer/Private/MobileShadingRenderer.cpp
@@ -587,6 +587,12 @@ void FMobileSceneRenderer::Render(FRHICommandListImmediate& RHICmdList)
 	// End of scene color rendering
 	RHICmdList.EndRenderPass();
 
+	if (!bGammaSpace || bRenderToSceneColor)
+	{
+		// transition scene color to Readable for post-processing
+		RHICmdList.TransitionResource(EResourceTransitionAccess::EReadable, SceneColor);
+	}
+
 	RHICmdList.SetCurrentStat(GET_STATID(STAT_CLMM_Post));
 
 	if (!View.bIsMobileMultiViewDirectEnabled)
