Merging using DevVRtoRelease422
Fix for custom depth with instanced stereo enabled.

diff --git a/Engine/Source/Runtime/Renderer/Private/CustomDepthRendering.cpp b/Engine/Source/Runtime/Renderer/Private/CustomDepthRendering.cpp
--- a/Engine/Source/Runtime/Renderer/Private/CustomDepthRendering.cpp
+++ b/Engine/Source/Runtime/Renderer/Private/CustomDepthRendering.cpp
@@ -39,6 +39,7 @@ FCustomDepthPassMeshProcessor::FCustomDepthPassMeshProcessor(const FScene* Scene
 	: FMeshPassProcessor(Scene, Scene->GetFeatureLevel(), InViewIfDynamicMeshCommand, InDrawListContext)
 {
 	PassDrawRenderState.SetViewUniformBuffer(Scene->UniformBuffers.CustomDepthViewUniformBuffer);
+	PassDrawRenderState.SetInstancedViewUniformBuffer(Scene->UniformBuffers.InstancedCustomDepthViewUniformBuffer);
 
 	if (FSceneInterface::GetShadingPath(FeatureLevel) == EShadingPath::Mobile)
 	{
diff --git a/Engine/Source/Runtime/Renderer/Private/RendererScene.cpp b/Engine/Source/Runtime/Renderer/Private/RendererScene.cpp
--- a/Engine/Source/Runtime/Renderer/Private/RendererScene.cpp
+++ b/Engine/Source/Runtime/Renderer/Private/RendererScene.cpp
@@ -970,6 +970,7 @@ void FPersistentUniformBuffers::Initialize()
 	MobileCustomDepthPassUniformBuffer = TUniformBufferRef<FMobileSceneTextureUniformParameters>::CreateUniformBufferImmediate(MobileCustomDepthPassUniformBufferParameters, UniformBuffer_MultiFrame, EUniformBufferValidation::None);
 
 	CustomDepthViewUniformBuffer = TUniformBufferRef<FViewUniformShaderParameters>::CreateUniformBufferImmediate(ViewUniformBufferParameters, UniformBuffer_MultiFrame, EUniformBufferValidation::None);
+	InstancedCustomDepthViewUniformBuffer = TUniformBufferRef<FInstancedViewUniformShaderParameters>::CreateUniformBufferImmediate(InstancedViewUniformBufferParameters, UniformBuffer_MultiFrame, EUniformBufferValidation::None);
 
 	FMobileShadowDepthPassUniformParameters MobileCSMShadowDepthPassParameters;
 	MobileCSMShadowDepthPassUniformBuffer = TUniformBufferRef<FMobileShadowDepthPassUniformParameters>::CreateUniformBufferImmediate(MobileCSMShadowDepthPassParameters, UniformBuffer_MultiFrame, EUniformBufferValidation::None);
@@ -1005,10 +1006,7 @@ bool FPersistentUniformBuffers::UpdateViewUniformBuffer(const FViewInfo& View)
 
 		if ((View.IsInstancedStereoPass() || View.bIsMobileMultiViewEnabled) && View.Family->Views.Num() > 0)
 		{
-			// When drawing the left eye in a stereo scene, copy the right eye view values into the instanced view uniform buffer.
-			const EStereoscopicPass StereoPassIndex = (View.StereoPass != eSSP_FULL) ? eSSP_RIGHT_EYE : eSSP_FULL;
-
-			const FViewInfo& InstancedView = static_cast<const FViewInfo&>(View.Family->GetStereoEyeView(StereoPassIndex));
+			const FViewInfo& InstancedView = GetInstancedView(View);
 			InstancedViewUniformBuffer.UpdateUniformBufferImmediate(reinterpret_cast<FInstancedViewUniformShaderParameters&>(*InstancedView.CachedViewUniformShaderParameters));
 		}
 		else
diff --git a/Engine/Source/Runtime/Renderer/Private/ScenePrivate.h b/Engine/Source/Runtime/Renderer/Private/ScenePrivate.h
--- a/Engine/Source/Runtime/Renderer/Private/ScenePrivate.h
+++ b/Engine/Source/Runtime/Renderer/Private/ScenePrivate.h
@@ -2355,6 +2355,14 @@ public:
 	// @return Whether uniform buffer was updated
 	bool UpdateViewUniformBuffer(const FViewInfo& View);
 
+	const FViewInfo& GetInstancedView(const FViewInfo& View)
+	{
+		// When drawing the left eye in a stereo scene, copy the right eye view values into the instanced view uniform buffer.
+		const EStereoscopicPass StereoPassIndex = (View.StereoPass != eSSP_FULL) ? eSSP_RIGHT_EYE : eSSP_FULL;
+
+		return static_cast<const FViewInfo&>(View.Family->GetStereoEyeView(StereoPassIndex));
+	}
+
 	TUniformBufferRef<FViewUniformShaderParameters> ViewUniformBuffer;
 	TUniformBufferRef<FInstancedViewUniformShaderParameters> InstancedViewUniformBuffer;
 	TUniformBufferRef<FSceneTexturesUniformParameters> DepthPassUniformBuffer;
@@ -2375,6 +2383,7 @@ public:
 	TUniformBufferRef<FSceneTexturesUniformParameters> CustomDepthPassUniformBuffer;
 	TUniformBufferRef<FMobileSceneTextureUniformParameters> MobileCustomDepthPassUniformBuffer;
 	TUniformBufferRef<FViewUniformShaderParameters> CustomDepthViewUniformBuffer;
+	TUniformBufferRef<FInstancedViewUniformShaderParameters> InstancedCustomDepthViewUniformBuffer;
 
 	TUniformBufferRef<FMobileBasePassUniformParameters> MobileOpaqueBasePassUniformBuffer;
 	TUniformBufferRef<FMobileBasePassUniformParameters> MobileTranslucentBasePassUniformBuffer;
diff --git a/Engine/Source/Runtime/Renderer/Private/SceneRendering.cpp b/Engine/Source/Runtime/Renderer/Private/SceneRendering.cpp
--- a/Engine/Source/Runtime/Renderer/Private/SceneRendering.cpp
+++ b/Engine/Source/Runtime/Renderer/Private/SceneRendering.cpp
@@ -2803,8 +2803,6 @@ void FSceneRenderer::RenderCustomDepthPass(FRHICommandListImmediate& RHICmdList)
 
 			if (View.ShouldRenderView())
 			{
-				Scene->UniformBuffers.UpdateViewUniformBuffer(View);
-
 				FUniformBufferRHIParamRef PassUniformBuffer = nullptr;
 
 				if (FSceneInterface::GetShadingPath(View.FeatureLevel) == EShadingPath::Mobile)
@@ -2844,26 +2842,63 @@ void FSceneRenderer::RenderCustomDepthPass(FRHICommandListImmediate& RHICmdList)
 					}
 				}
 
-
-				FViewUniformShaderParameters CustomDepthViewUniformBufferParameters = *View.CachedViewUniformShaderParameters;
-
 				if (CVarCustomDepthTemporalAAJitter.GetValueOnRenderThread() == 0 && View.AntiAliasingMethod == AAM_TemporalAA)
 				{
+					// Handle the "current" view the same, always.
+					FViewUniformShaderParameters CustomDepthViewUniformBufferParameters = *View.CachedViewUniformShaderParameters;
+
 					FBox VolumeBounds[TVC_MAX];
 
-					FViewMatrices ModifiedViewMatricies = View.ViewMatrices;
-					ModifiedViewMatricies.HackRemoveTemporalAAProjectionJitter();
+					FViewMatrices ModifiedViewMatrices = View.ViewMatrices;
+					ModifiedViewMatrices.HackRemoveTemporalAAProjectionJitter();
 
 					View.SetupUniformBufferParameters(
 						SceneContext,
-						ModifiedViewMatricies,
-						ModifiedViewMatricies,
+						ModifiedViewMatrices,
+						ModifiedViewMatrices,
 						VolumeBounds,
 						TVC_MAX,
 						CustomDepthViewUniformBufferParameters);
-				}
 
-				Scene->UniformBuffers.CustomDepthViewUniformBuffer.UpdateUniformBufferImmediate(CustomDepthViewUniformBufferParameters);
+					Scene->UniformBuffers.CustomDepthViewUniformBuffer.UpdateUniformBufferImmediate(CustomDepthViewUniformBufferParameters);
+
+					if ((View.IsInstancedStereoPass() || View.bIsMobileMultiViewEnabled) && View.Family->Views.Num() > 0)
+					{
+						// When drawing the left eye in a stereo scene, set up the instanced custom depth uniform buffer with the right-eye data,
+						// with the TAA jitter removed.
+						const EStereoscopicPass StereoPassIndex = (View.StereoPass != eSSP_FULL) ? eSSP_RIGHT_EYE : eSSP_FULL;
+
+						const FViewInfo& InstancedView = static_cast<const FViewInfo&>(View.Family->GetStereoEyeView(StereoPassIndex));
+
+						CustomDepthViewUniformBufferParameters = *InstancedView.CachedViewUniformShaderParameters;
+						ModifiedViewMatrices = InstancedView.ViewMatrices;
+						ModifiedViewMatrices.HackRemoveTemporalAAProjectionJitter();
+
+						InstancedView.SetupUniformBufferParameters(
+							SceneContext,
+							ModifiedViewMatrices,
+							ModifiedViewMatrices,
+							VolumeBounds,
+							TVC_MAX,
+							CustomDepthViewUniformBufferParameters);
+
+						Scene->UniformBuffers.InstancedCustomDepthViewUniformBuffer.UpdateUniformBufferImmediate(reinterpret_cast<FInstancedViewUniformShaderParameters&>(CustomDepthViewUniformBufferParameters));
+					}
+				}
+				else
+				{
+					Scene->UniformBuffers.CustomDepthViewUniformBuffer.UpdateUniformBufferImmediate(*View.CachedViewUniformShaderParameters);
+					if ((View.IsInstancedStereoPass() || View.bIsMobileMultiViewEnabled) && View.Family->Views.Num() > 0)
+					{
+						const FViewInfo& InstancedView = Scene->UniformBuffers.GetInstancedView(View);
+						Scene->UniformBuffers.InstancedCustomDepthViewUniformBuffer.UpdateUniformBufferImmediate(reinterpret_cast<FInstancedViewUniformShaderParameters&>(*InstancedView.CachedViewUniformShaderParameters));
+					}
+					else
+					{
+						// If we don't render this pass in stereo we simply update the buffer with the same view uniform parameters.
+						Scene->UniformBuffers.InstancedCustomDepthViewUniformBuffer.UpdateUniformBufferImmediate(reinterpret_cast<FInstancedViewUniformShaderParameters&>(*View.CachedViewUniformShaderParameters));
+					}
+				}
 	
 				View.ParallelMeshDrawCommandPasses[EMeshPass::CustomDepth].DispatchDraw(nullptr, RHICmdList);
 			}
