Keep DistanceFieldData.ObjectBuffers persistent instead of recreating a new one when needed in UpdateGlobalDistanceFieldObjectRemoves

#rnx

diff --git a/Engine/Source/Runtime/Renderer/Private/CapsuleShadowRendering.cpp b/Engine/Source/Runtime/Renderer/Private/CapsuleShadowRendering.cpp
--- a/Engine/Source/Runtime/Renderer/Private/CapsuleShadowRendering.cpp
+++ b/Engine/Source/Runtime/Renderer/Private/CapsuleShadowRendering.cpp
@@ -424,9 +424,9 @@ public:
 		float IndirectCapsuleSelfShadowingIntensityValue = Scene->DynamicIndirectShadowsSelfShadowingIntensity;
 		SetShaderValue(RHICmdList, ShaderRHI, IndirectCapsuleSelfShadowingIntensity, IndirectCapsuleSelfShadowingIntensityValue);
 
-		if (Scene->DistanceFieldSceneData.ObjectBuffers)
+		if (Scene->DistanceFieldSceneData.GetCurrentObjectBuffers())
 		{
-			DistanceFieldObjectParameters.Set(RHICmdList, ShaderRHI, *Scene->DistanceFieldSceneData.ObjectBuffers, Scene->DistanceFieldSceneData.NumObjectsInBuffer);
+			DistanceFieldObjectParameters.Set(RHICmdList, ShaderRHI, *Scene->DistanceFieldSceneData.GetCurrentObjectBuffers(), Scene->DistanceFieldSceneData.NumObjectsInBuffer);
 		}
 		else
 		{
diff --git a/Engine/Source/Runtime/Renderer/Private/DistanceFieldAmbientOcclusion.cpp b/Engine/Source/Runtime/Renderer/Private/DistanceFieldAmbientOcclusion.cpp
--- a/Engine/Source/Runtime/Renderer/Private/DistanceFieldAmbientOcclusion.cpp
+++ b/Engine/Source/Runtime/Renderer/Private/DistanceFieldAmbientOcclusion.cpp
@@ -633,7 +633,7 @@ void ListDistanceFieldLightingMemory(const FViewInfo& View, FSceneRenderer& Scen
 
 	if (Scene->DistanceFieldSceneData.NumObjectsInBuffer > 0)
 	{
-		UE_LOG(LogRenderer, Log, TEXT("   Scene Object data %.3fMb"), Scene->DistanceFieldSceneData.ObjectBuffers->GetSizeBytes() / 1024.0f / 1024.0f);
+		UE_LOG(LogRenderer, Log, TEXT("   Scene Object data %.3fMb"), Scene->DistanceFieldSceneData.GetCurrentObjectBuffers()->GetSizeBytes() / 1024.0f / 1024.0f);
 	}
 	
 	UE_LOG(LogRenderer, Log, TEXT("   %s"), *GDistanceFieldVolumeTextureAtlas.GetSizeString());
diff --git a/Engine/Source/Runtime/Renderer/Private/DistanceFieldObjectCulling.cpp b/Engine/Source/Runtime/Renderer/Private/DistanceFieldObjectCulling.cpp
--- a/Engine/Source/Runtime/Renderer/Private/DistanceFieldObjectCulling.cpp
+++ b/Engine/Source/Runtime/Renderer/Private/DistanceFieldObjectCulling.cpp
@@ -141,13 +141,13 @@ public:
 		OutUAVs[1] = GAOCulledObjectBuffers.Buffers.Bounds.UAV;
 		OutUAVs[2] = GAOCulledObjectBuffers.Buffers.Data.UAV;
 		OutUAVs[3] = GAOCulledObjectBuffers.Buffers.BoxBounds.UAV;
-		OutUAVs[4] = Scene->DistanceFieldSceneData.ObjectBuffers->Data.UAV;
-		OutUAVs[5] = Scene->DistanceFieldSceneData.ObjectBuffers->Bounds.UAV;
+		OutUAVs[4] = Scene->DistanceFieldSceneData.GetCurrentObjectBuffers()->Data.UAV;
+		OutUAVs[5] = Scene->DistanceFieldSceneData.GetCurrentObjectBuffers()->Bounds.UAV;
 		RHICmdList.TransitionResources(EResourceTransitionAccess::ERWBarrier, EResourceTransitionPipeline::EComputeToCompute, OutUAVs, ARRAY_COUNT(OutUAVs));
 
 		FRHIComputeShader* ShaderRHI = GetComputeShader();
 		FGlobalShader::SetParameters<FViewUniformShaderParameters>(RHICmdList, ShaderRHI, View.ViewUniformBuffer);
-		ObjectBufferParameters.Set(RHICmdList, ShaderRHI, *(Scene->DistanceFieldSceneData.ObjectBuffers), Scene->DistanceFieldSceneData.NumObjectsInBuffer);
+		ObjectBufferParameters.Set(RHICmdList, ShaderRHI, *(Scene->DistanceFieldSceneData.GetCurrentObjectBuffers()), Scene->DistanceFieldSceneData.NumObjectsInBuffer);
 		CulledObjectParameters.Set(RHICmdList, ShaderRHI, GAOCulledObjectBuffers.Buffers);
 		AOParameters.Set(RHICmdList, ShaderRHI, Parameters);
 
@@ -160,7 +160,7 @@ public:
 
 	void UnsetParameters(FRHICommandList& RHICmdList, const FScene* Scene)
 	{
-		ObjectBufferParameters.UnsetParameters(RHICmdList, GetComputeShader(), *(Scene->DistanceFieldSceneData.ObjectBuffers));
+		ObjectBufferParameters.UnsetParameters(RHICmdList, GetComputeShader(), *(Scene->DistanceFieldSceneData.GetCurrentObjectBuffers()));
 		CulledObjectParameters.UnsetParameters(RHICmdList, GetComputeShader());
 
 		FRHIUnorderedAccessView* OutUAVs[6];
@@ -168,8 +168,8 @@ public:
 		OutUAVs[1] = GAOCulledObjectBuffers.Buffers.Bounds.UAV;
 		OutUAVs[2] = GAOCulledObjectBuffers.Buffers.Data.UAV;
 		OutUAVs[3] = GAOCulledObjectBuffers.Buffers.BoxBounds.UAV;
-		OutUAVs[4] = Scene->DistanceFieldSceneData.ObjectBuffers->Data.UAV;
-		OutUAVs[5] = Scene->DistanceFieldSceneData.ObjectBuffers->Bounds.UAV;		
+		OutUAVs[4] = Scene->DistanceFieldSceneData.GetCurrentObjectBuffers()->Data.UAV;
+		OutUAVs[5] = Scene->DistanceFieldSceneData.GetCurrentObjectBuffers()->Bounds.UAV;		
 		RHICmdList.TransitionResources(EResourceTransitionAccess::ERWBarrier, EResourceTransitionPipeline::EComputeToCompute, OutUAVs, ARRAY_COUNT(OutUAVs));
 	}
 
diff --git a/Engine/Source/Runtime/Renderer/Private/DistanceFieldObjectManagement.cpp b/Engine/Source/Runtime/Renderer/Private/DistanceFieldObjectManagement.cpp
--- a/Engine/Source/Runtime/Renderer/Private/DistanceFieldObjectManagement.cpp
+++ b/Engine/Source/Runtime/Renderer/Private/DistanceFieldObjectManagement.cpp
@@ -186,14 +186,15 @@ public:
 		SetSRVParameter(RHICmdList, ShaderRHI, UploadOperationData, InUploadOperationData);
 
 
-		ObjectBufferParameters.Set(RHICmdList, ShaderRHI, *(Scene->DistanceFieldSceneData.ObjectBuffers), Scene->DistanceFieldSceneData.NumObjectsInBuffer, true);
+		ObjectBufferParameters.Set(RHICmdList, ShaderRHI, *(Scene->DistanceFieldSceneData.GetCurrentObjectBuffers()), Scene->DistanceFieldSceneData.NumObjectsInBuffer, true);
 	}
 
 	void UnsetParameters(FRHICommandList& RHICmdList, const FScene* Scene)
 	{
-		ObjectBufferParameters.UnsetParameters(RHICmdList, GetComputeShader(), *(Scene->DistanceFieldSceneData.ObjectBuffers), true);
+		const FDistanceFieldObjectBuffers& ObjectBuffers = *(Scene->DistanceFieldSceneData.GetCurrentObjectBuffers());
+
+		ObjectBufferParameters.UnsetParameters(RHICmdList, GetComputeShader(), ObjectBuffers, true);
 
-		const FDistanceFieldObjectBuffers& ObjectBuffers = *(Scene->DistanceFieldSceneData.ObjectBuffers);
 		FRHIUnorderedAccessView* OutUAVs[2];
 		OutUAVs[0] = ObjectBuffers.Bounds.UAV;
 		OutUAVs[1] = ObjectBuffers.Data.UAV;
@@ -477,14 +478,14 @@ public:
 
 		SetShaderValue(RHICmdList, ShaderRHI, NumRemoveOperations, NumRemoveOperationsValue);
 		SetSRVParameter(RHICmdList, ShaderRHI, RemoveOperationIndices, InRemoveOperationIndices);
-		ObjectBufferParameters.Set(RHICmdList, ShaderRHI, *(Scene->DistanceFieldSceneData.ObjectBuffers), Scene->DistanceFieldSceneData.NumObjectsInBuffer, true);
+		ObjectBufferParameters.Set(RHICmdList, ShaderRHI, *(Scene->DistanceFieldSceneData.GetCurrentObjectBuffers()), Scene->DistanceFieldSceneData.NumObjectsInBuffer, true);
 		SetSRVParameter(RHICmdList, ShaderRHI, ObjectBounds2, InObjectBounds2);
 		SetSRVParameter(RHICmdList, ShaderRHI, ObjectData2, InObjectData2);
 	}
 
 	void UnsetParameters(FRHICommandList& RHICmdList, const FScene* Scene)
 	{
-		ObjectBufferParameters.UnsetParameters(RHICmdList, GetComputeShader(), *(Scene->DistanceFieldSceneData.ObjectBuffers), true);
+		ObjectBufferParameters.UnsetParameters(RHICmdList, GetComputeShader(), *(Scene->DistanceFieldSceneData.GetCurrentObjectBuffers()), true);
 	}
 
 	virtual bool Serialize(FArchive& Ar) override
@@ -592,7 +593,6 @@ void UpdateGlobalDistanceFieldObjectRemoves(FRHICommandListImmediate& RHICmdList
 	FDistanceFieldSceneData& DistanceFieldSceneData = Scene->DistanceFieldSceneData;
 
 	TArray<FIntRect> RemoveObjectIndices;
-	FDistanceFieldObjectBuffers* TemporaryCopySourceBuffers = NULL;
 
 	if (DistanceFieldSceneData.PendingRemoveOperations.Num() > 0)
 	{
@@ -632,6 +632,8 @@ void UpdateGlobalDistanceFieldObjectRemoves(FRHICommandListImmediate& RHICmdList
 			// We have multiple remove requests enqueued in PendingRemoveOperations, can only use the RemoveAtSwap version when there won't be collisions
 			const bool bUseRemoveAtSwap = PendingRemoveOperations.Last() < DistanceFieldSceneData.NumObjectsInBuffer - PendingRemoveOperations.Num();
 
+			FDistanceFieldObjectBuffers*& CurrentObjectBuffers = DistanceFieldSceneData.ObjectBuffers[DistanceFieldSceneData.ObjectBufferIndex];
+
 			if (bUseRemoveAtSwap)
 			{
 				// Remove everything in parallel in the same buffer with a RemoveAtSwap algorithm
@@ -657,13 +659,31 @@ void UpdateGlobalDistanceFieldObjectRemoves(FRHICommandListImmediate& RHICmdList
 			{
 				const double StartTime = FPlatformTime::Seconds();
 
+				TArray<FPrimitiveAndInstance> OriginalPrimitiveInstanceMapping;
+
 				// Have to copy the object data to allow parallel removing
-				TemporaryCopySourceBuffers = DistanceFieldSceneData.ObjectBuffers;
-				DistanceFieldSceneData.ObjectBuffers = new FDistanceFieldObjectBuffers();
-				DistanceFieldSceneData.ObjectBuffers->MaxObjects = TemporaryCopySourceBuffers->MaxObjects;
-				DistanceFieldSceneData.ObjectBuffers->Initialize();
+				int32 NextObjectBufferIndex = (DistanceFieldSceneData.ObjectBufferIndex + 1) & 1;
+				FDistanceFieldObjectBuffers*& NextObjectBuffers = DistanceFieldSceneData.ObjectBuffers[NextObjectBufferIndex];
+				
+				check(CurrentObjectBuffers != nullptr);
+
+				DistanceFieldSceneData.ObjectBufferIndex = NextObjectBufferIndex;
+
+				if (NextObjectBuffers != nullptr && NextObjectBuffers->MaxObjects < CurrentObjectBuffers->MaxObjects)
+				{
+					NextObjectBuffers->Release();
+					delete NextObjectBuffers;
+					NextObjectBuffers = nullptr;
+				}
 
-				TArray<FPrimitiveAndInstance> OriginalPrimitiveInstanceMapping = DistanceFieldSceneData.PrimitiveInstanceMapping;
+				if (NextObjectBuffers == nullptr)
+				{
+					NextObjectBuffers = new FDistanceFieldObjectBuffers();
+					NextObjectBuffers->MaxObjects = CurrentObjectBuffers->MaxObjects;
+					NextObjectBuffers->Initialize();
+				}
+
+				OriginalPrimitiveInstanceMapping = DistanceFieldSceneData.PrimitiveInstanceMapping;
 				DistanceFieldSceneData.PrimitiveInstanceMapping.Reset();
 
 				const int32 NumDestObjects = DistanceFieldSceneData.NumObjectsInBuffer - PendingRemoveOperations.Num();
@@ -742,7 +762,6 @@ void UpdateGlobalDistanceFieldObjectRemoves(FRHICommandListImmediate& RHICmdList
 
 				if (bUseRemoveAtSwap)
 				{
-					check(!TemporaryCopySourceBuffers);
 					TShaderMapRef<TRemoveObjectsFromBufferCS<true> > ComputeShader(GetGlobalShaderMap(Scene->GetFeatureLevel()));
 					RHICmdList.SetComputeShader(ComputeShader->GetComputeShader());
 					ComputeShader->SetParameters(RHICmdList, Scene, RemoveObjectIndices.Num(), GDistanceFieldRemoveIndices.RemoveIndices.BufferSRV, NULL, NULL);
@@ -752,24 +771,14 @@ void UpdateGlobalDistanceFieldObjectRemoves(FRHICommandListImmediate& RHICmdList
 				}
 				else
 				{
-					check(TemporaryCopySourceBuffers);
 					TShaderMapRef<TRemoveObjectsFromBufferCS<false> > ComputeShader(GetGlobalShaderMap(Scene->GetFeatureLevel()));
 					RHICmdList.SetComputeShader(ComputeShader->GetComputeShader());
-					ComputeShader->SetParameters(RHICmdList, Scene, RemoveObjectIndices.Num(), GDistanceFieldRemoveIndices.RemoveIndices.BufferSRV, TemporaryCopySourceBuffers->Bounds.SRV, TemporaryCopySourceBuffers->Data.SRV);
+					ComputeShader->SetParameters(RHICmdList, Scene, RemoveObjectIndices.Num(), GDistanceFieldRemoveIndices.RemoveIndices.BufferSRV, CurrentObjectBuffers->Bounds.SRV, CurrentObjectBuffers->Data.SRV);
 
 					DispatchComputeShader(RHICmdList, *ComputeShader, FMath::DivideAndRoundUp<uint32>(RemoveObjectIndices.Num(), UpdateObjectsGroupSize), 1, 1);
 					ComputeShader->UnsetParameters(RHICmdList, Scene);
 				}
 			}
-			
-			// make sure to delete the temporary buffer (even if RemoveObjectIndices is empty)
-			if (TemporaryCopySourceBuffers)
-			{
-				check(bUseRemoveAtSwap == false);
-				TemporaryCopySourceBuffers->Release();
-				delete TemporaryCopySourceBuffers;
-				TemporaryCopySourceBuffers = nullptr;
-			}
 		}
 	}
 }
@@ -1026,9 +1035,9 @@ void FDeferredShadingSceneRenderer::UpdateGlobalDistanceFieldObjectBuffers(FRHIC
 		QUICK_SCOPE_CYCLE_COUNTER(STAT_UpdateObjectData);
 		SCOPED_DRAW_EVENT(RHICmdList, UpdateSceneObjectData);
 
-		if (!DistanceFieldSceneData.ObjectBuffers)
+		if (DistanceFieldSceneData.ObjectBuffers[DistanceFieldSceneData.ObjectBufferIndex]==nullptr)
 		{
-			DistanceFieldSceneData.ObjectBuffers = new FDistanceFieldObjectBuffers();
+			DistanceFieldSceneData.ObjectBuffers[DistanceFieldSceneData.ObjectBufferIndex] = new FDistanceFieldObjectBuffers();
 		}
 
 		if (!DistanceFieldSceneData.SurfelBuffers)
@@ -1241,9 +1250,11 @@ void FDeferredShadingSceneRenderer::UpdateGlobalDistanceFieldObjectBuffers(FRHIC
 				DistanceFieldSceneData.PendingThrottledOperations.Empty();
 			}
 
-			if (DistanceFieldSceneData.ObjectBuffers->MaxObjects < DistanceFieldSceneData.NumObjectsInBuffer)
+			FDistanceFieldObjectBuffers*& ObjectBuffers = DistanceFieldSceneData.ObjectBuffers[DistanceFieldSceneData.ObjectBufferIndex];
+
+			if (ObjectBuffers->MaxObjects < DistanceFieldSceneData.NumObjectsInBuffer)
 			{
-				if (DistanceFieldSceneData.ObjectBuffers->MaxObjects > 0)
+				if (ObjectBuffers->MaxObjects > 0)
 				{
 					// Realloc
 					FDistanceFieldObjectBuffers* NewObjectBuffers = new FDistanceFieldObjectBuffers();
@@ -1253,21 +1264,21 @@ void FDeferredShadingSceneRenderer::UpdateGlobalDistanceFieldObjectBuffers(FRHIC
 					{
 						TShaderMapRef<FCopyObjectBufferCS> ComputeShader(GetGlobalShaderMap(Scene->GetFeatureLevel()));
 						RHICmdList.SetComputeShader(ComputeShader->GetComputeShader());
-						ComputeShader->SetParameters(RHICmdList, *(DistanceFieldSceneData.ObjectBuffers), *NewObjectBuffers, OriginalNumObjects);
+						ComputeShader->SetParameters(RHICmdList, *ObjectBuffers, *NewObjectBuffers, OriginalNumObjects);
 
 						DispatchComputeShader(RHICmdList, *ComputeShader, FMath::DivideAndRoundUp<uint32>(OriginalNumObjects, UpdateObjectsGroupSize), 1, 1);
 						ComputeShader->UnsetParameters(RHICmdList, *NewObjectBuffers);
 					}
 
-					DistanceFieldSceneData.ObjectBuffers->Release();
-					delete DistanceFieldSceneData.ObjectBuffers;
-					DistanceFieldSceneData.ObjectBuffers = NewObjectBuffers;
+					ObjectBuffers->Release();
+					delete ObjectBuffers;
+					ObjectBuffers = NewObjectBuffers;
 				}
 				else
 				{
 					// First time allocate
-					DistanceFieldSceneData.ObjectBuffers->MaxObjects = DistanceFieldSceneData.NumObjectsInBuffer * 5 / 4;
-					DistanceFieldSceneData.ObjectBuffers->Initialize();
+					ObjectBuffers->MaxObjects = DistanceFieldSceneData.NumObjectsInBuffer * 5 / 4;
+					ObjectBuffers->Initialize();
 				}
 			}
 		}
diff --git a/Engine/Source/Runtime/Renderer/Private/DistanceFieldShadowing.cpp b/Engine/Source/Runtime/Renderer/Private/DistanceFieldShadowing.cpp
--- a/Engine/Source/Runtime/Renderer/Private/DistanceFieldShadowing.cpp
+++ b/Engine/Source/Runtime/Renderer/Private/DistanceFieldShadowing.cpp
@@ -136,7 +136,7 @@ public:
 	{
 		FRHIComputeShader* ShaderRHI = GetComputeShader();
 		FGlobalShader::SetParameters<FViewUniformShaderParameters>(RHICmdList, ShaderRHI, View.ViewUniformBuffer);
-		ObjectBufferParameters.Set(RHICmdList, ShaderRHI, *(Scene->DistanceFieldSceneData.ObjectBuffers), Scene->DistanceFieldSceneData.NumObjectsInBuffer);
+		ObjectBufferParameters.Set(RHICmdList, ShaderRHI, *(Scene->DistanceFieldSceneData.GetCurrentObjectBuffers()), Scene->DistanceFieldSceneData.NumObjectsInBuffer);
 
 		FRHIUnorderedAccessView* OutUAVs[4];
 		OutUAVs[0] = GShadowCulledObjectBuffers.Buffers.ObjectIndirectArguments.UAV;
@@ -164,7 +164,7 @@ public:
 
 	void UnsetParameters(FRHICommandList& RHICmdList, const FScene* Scene)
 	{
-		ObjectBufferParameters.UnsetParameters(RHICmdList, GetComputeShader(), *(Scene->DistanceFieldSceneData.ObjectBuffers));
+		ObjectBufferParameters.UnsetParameters(RHICmdList, GetComputeShader(), *(Scene->DistanceFieldSceneData.GetCurrentObjectBuffers()));
 		CulledObjectParameters.UnsetParameters(RHICmdList, GetComputeShader());
 
 		FRHIUnorderedAccessView* OutUAVs[4];
diff --git a/Engine/Source/Runtime/Renderer/Private/GlobalDistanceField.cpp b/Engine/Source/Runtime/Renderer/Private/GlobalDistanceField.cpp
--- a/Engine/Source/Runtime/Renderer/Private/GlobalDistanceField.cpp
+++ b/Engine/Source/Runtime/Renderer/Private/GlobalDistanceField.cpp
@@ -184,7 +184,7 @@ public:
 	{
 		FRHIComputeShader* ShaderRHI = GetComputeShader();
 		FGlobalShader::SetParameters<FViewUniformShaderParameters>(RHICmdList, ShaderRHI, View.ViewUniformBuffer);
-		ObjectBufferParameters.Set(RHICmdList, ShaderRHI, *(Scene->DistanceFieldSceneData.ObjectBuffers), Scene->DistanceFieldSceneData.NumObjectsInBuffer);
+		ObjectBufferParameters.Set(RHICmdList, ShaderRHI, *(Scene->DistanceFieldSceneData.GetCurrentObjectBuffers()), Scene->DistanceFieldSceneData.NumObjectsInBuffer);
 
 		FRHIUnorderedAccessView* OutUAVs[4];
 		OutUAVs[0] = GGlobalDistanceFieldCulledObjectBuffers.Buffers.ObjectIndirectArguments.UAV;
@@ -218,7 +218,7 @@ public:
 
 	void UnsetParameters(FRHICommandList& RHICmdList, const FScene* Scene)
 	{
-		ObjectBufferParameters.UnsetParameters(RHICmdList, GetComputeShader(), *(Scene->DistanceFieldSceneData.ObjectBuffers));
+		ObjectBufferParameters.UnsetParameters(RHICmdList, GetComputeShader(), *(Scene->DistanceFieldSceneData.GetCurrentObjectBuffers()));
 		CulledObjectParameters.UnsetParameters(RHICmdList, GetComputeShader());
 
 		TArray<FRHIUnorderedAccessView*> UAVs;
diff --git a/Engine/Source/Runtime/Renderer/Private/RendererScene.cpp b/Engine/Source/Runtime/Renderer/Private/RendererScene.cpp
--- a/Engine/Source/Runtime/Renderer/Private/RendererScene.cpp
+++ b/Engine/Source/Runtime/Renderer/Private/RendererScene.cpp
@@ -365,11 +365,14 @@ bool FPixelInspectorData::AddPixelInspectorRequest(FPixelInspectorRequest *Pixel
 
 FDistanceFieldSceneData::FDistanceFieldSceneData(EShaderPlatform ShaderPlatform) 
 	: NumObjectsInBuffer(0)
-	, ObjectBuffers(NULL)
+	, ObjectBufferIndex(0)
 	, SurfelBuffers(NULL)
 	, InstancedSurfelBuffers(NULL)
 	, AtlasGeneration(0)
 {
+	ObjectBuffers[0] = nullptr;
+	ObjectBuffers[1] = nullptr;
+
 	static const auto CVar = IConsoleManager::Get().FindTConsoleVariableDataInt(TEXT("r.GenerateMeshDistanceFields"));
 
 	bTrackAllPrimitives = (DoesPlatformSupportDistanceFieldAO(ShaderPlatform) || DoesPlatformSupportDistanceFieldShadowing(ShaderPlatform)) && CVar->GetValueOnGameThread() != 0 && IsUsingDistanceFields(ShaderPlatform);
@@ -379,7 +382,8 @@ FDistanceFieldSceneData::FDistanceFieldSceneData(EShaderPlatform ShaderPlatform)
 
 FDistanceFieldSceneData::~FDistanceFieldSceneData() 
 {
-	delete ObjectBuffers;
+	delete ObjectBuffers[0];
+	delete ObjectBuffers[1];
 }
 
 void FDistanceFieldSceneData::AddPrimitive(FPrimitiveSceneInfo* InPrimitive)
@@ -459,9 +463,13 @@ void FDistanceFieldSceneData::RemovePrimitive(FPrimitiveSceneInfo* InPrimitive)
 
 void FDistanceFieldSceneData::Release()
 {
-	if (ObjectBuffers)
+	if (ObjectBuffers[0]!=nullptr)
+	{
+		ObjectBuffers[0]->Release();
+	}
+	if (ObjectBuffers[1] != nullptr)
 	{
-		ObjectBuffers->Release();
+		ObjectBuffers[1]->Release();
 	}
 }
 
diff --git a/Engine/Source/Runtime/Renderer/Private/ScenePrivate.h b/Engine/Source/Runtime/Renderer/Private/ScenePrivate.h
--- a/Engine/Source/Runtime/Renderer/Private/ScenePrivate.h
+++ b/Engine/Source/Runtime/Renderer/Private/ScenePrivate.h
@@ -1777,8 +1777,14 @@ public:
 		return bCanUse16BitObjectIndices && (NumObjectsInBuffer < (1 << 16));
 	}
 
+	const class FDistanceFieldObjectBuffers* GetCurrentObjectBuffers() const
+	{
+		return ObjectBuffers[ObjectBufferIndex];
+	}
+
 	int32 NumObjectsInBuffer;
-	class FDistanceFieldObjectBuffers* ObjectBuffers;
+	class FDistanceFieldObjectBuffers* ObjectBuffers[2];
+	int ObjectBufferIndex;
 
 	/** Stores the primitive and instance index of every entry in the object buffer. */
 	TArray<FPrimitiveAndInstance> PrimitiveInstanceMapping;
