Merge ENQUEUE_RENDER_COMMAND to lower the overhead on task in AddPrimitive.

#rnx

diff --git a/Engine/Source/Runtime/Renderer/Private/RendererScene.cpp b/Engine/Source/Runtime/Renderer/Private/RendererScene.cpp
--- a/Engine/Source/Runtime/Renderer/Private/RendererScene.cpp
+++ b/Engine/Source/Runtime/Renderer/Private/RendererScene.cpp
@@ -1238,18 +1238,6 @@ void FScene::AddPrimitive(UPrimitiveComponent* Primitive)
 	ensureMsgf(!Primitive->Bounds.BoxExtent.ContainsNaN() && !Primitive->Bounds.Origin.ContainsNaN() && !FMath::IsNaN(Primitive->Bounds.SphereRadius) && FMath::IsFinite(Primitive->Bounds.SphereRadius),
 			TEXT("Nans found on Bounds for Primitive %s: Origin %s, BoxExtent %s, SphereRadius %f"), *Primitive->GetName(), *Primitive->Bounds.Origin.ToString(), *Primitive->Bounds.BoxExtent.ToString(), Primitive->Bounds.SphereRadius);
 
-	// Create any RenderThreadResources required.
-	ENQUEUE_RENDER_COMMAND(CreateRenderThreadResourcesCommand)(
-		[Params](FRHICommandListImmediate& RHICmdList)
-	{
-		FPrimitiveSceneProxy* SceneProxy = Params.PrimitiveSceneProxy;
-		FScopeCycleCounter Context(SceneProxy->GetStatId());
-		SceneProxy->SetTransform(Params.RenderMatrix, Params.WorldBounds, Params.LocalBounds, Params.AttachmentRootPosition);
-
-		// Create any RenderThreadResources required.
-		SceneProxy->CreateRenderThreadResources();
-	});
-
 	INC_DWORD_STAT_BY( STAT_GameToRendererMallocTotal, PrimitiveSceneProxy->GetMemoryFootprint() + PrimitiveSceneInfo->GetMemoryFootprint() );
 
 	// Verify the primitive is valid (this will compile away to a nop without CHECK_FOR_PIE_PRIMITIVE_ATTACH_SCENE_MISMATCH)
@@ -1258,16 +1246,22 @@ void FScene::AddPrimitive(UPrimitiveComponent* Primitive)
 	// Increment the attachment counter, the primitive is about to be attached to the scene.
 	Primitive->AttachmentCounter.Increment();
 
-	// Send a command to the rendering thread to add the primitive to the scene.
+	// Create any RenderThreadResources required and send a command to the rendering thread to add the primitive to the scene.
 	FScene* Scene = this;
 
 	// If this primitive has a simulated previous transform, ensure that the velocity data for the scene representation is correct
 	TOptional<FTransform> PreviousTransform = FMotionVectorSimulation::Get().GetPreviousTransform(Primitive);
 
 	ENQUEUE_RENDER_COMMAND(AddPrimitiveCommand)(
-		[Scene, PrimitiveSceneInfo, PreviousTransform](FRHICommandListImmediate& RHICmdList)
+		[Params, Scene, PrimitiveSceneInfo, PreviousTransform](FRHICommandListImmediate& RHICmdList)
 		{
-			FScopeCycleCounter Context(PrimitiveSceneInfo->Proxy->GetStatId());
+			FPrimitiveSceneProxy* SceneProxy = Params.PrimitiveSceneProxy;
+			FScopeCycleCounter Context(SceneProxy->GetStatId());
+			SceneProxy->SetTransform(Params.RenderMatrix, Params.WorldBounds, Params.LocalBounds, Params.AttachmentRootPosition);
+
+			// Create any RenderThreadResources required.
+			SceneProxy->CreateRenderThreadResources();
+
 			Scene->AddPrimitiveSceneInfo_RenderThread(RHICmdList, PrimitiveSceneInfo);
 
 			if (PreviousTransform.IsSet())
