Mark VT prefetch pages as recently used, otherwise they may be evicted later in frame, and break code that assumes they're resident
Copied from Dev-Rendering CL 6983865

diff --git a/Engine/Source/Runtime/Renderer/Private/VT/VirtualTextureSystem.cpp b/Engine/Source/Runtime/Renderer/Private/VT/VirtualTextureSystem.cpp
--- a/Engine/Source/Runtime/Renderer/Private/VT/VirtualTextureSystem.cpp
+++ b/Engine/Source/Runtime/Renderer/Private/VT/VirtualTextureSystem.cpp
@@ -1399,6 +1399,11 @@ void FVirtualTextureSystem::GatherRequestsTask(const FGatherRequestsParameters&
 								LocalLayerMaskToPrefetch |= (1u << LocalLayerIndex);
 								++NumPrefetchPages;
 							}
+							else
+							{
+								// Need to mark the page as recently used, otherwise it may be evicted later this frame
+								AddPageUpdate(PageUpdateBuffers, PageUpdateFlushCount, PhysicalSpace->GetID(), pAddress);
+							}
 						}
 					}
 
