Adding mip map generation support for scene capture 2d

diff --git a/Engine/Source/Runtime/Engine/Private/TextureRenderTarget2D.cpp b/Engine/Source/Runtime/Engine/Private/TextureRenderTarget2D.cpp
--- a/Engine/Source/Runtime/Engine/Private/TextureRenderTarget2D.cpp
+++ b/Engine/Source/Runtime/Engine/Private/TextureRenderTarget2D.cpp
@@ -524,10 +524,10 @@ void FTextureRenderTarget2DResource::UpdateDeferredResource( FRHICommandListImme
 	{
 		/**Convert the input values from the editor to a compatible format for FSamplerStateInitializerRHI. 
 			Ensure default sampler is Bilinear clamp*/
-		FGenerateMips::Execute(RHICmdList, RenderTargetTextureRHI,
+		FGenerateMips::Execute(RHICmdList, RenderTargetTextureRHI, FGenerateMipsParams{
 			Owner->MipsSamplerFilter == TF_Nearest ? SF_Point : (Owner->MipsSamplerFilter == TF_Trilinear ? SF_Trilinear : SF_Bilinear),
 			Owner->MipsAddressU == TA_Wrap ? AM_Wrap : (Owner->MipsAddressU == TA_Mirror ? AM_Mirror : AM_Clamp),
-			Owner->MipsAddressV == TA_Wrap ? AM_Wrap : (Owner->MipsAddressV == TA_Mirror ? AM_Mirror : AM_Clamp));
+			Owner->MipsAddressV == TA_Wrap ? AM_Wrap : (Owner->MipsAddressV == TA_Mirror ? AM_Mirror : AM_Clamp)});
 	}
 
  	// copy surface to the texture for use
diff --git a/Engine/Source/Runtime/RenderCore/Private/GenerateMips.cpp b/Engine/Source/Runtime/RenderCore/Private/GenerateMips.cpp
--- a/Engine/Source/Runtime/RenderCore/Private/GenerateMips.cpp
+++ b/Engine/Source/Runtime/RenderCore/Private/GenerateMips.cpp
@@ -36,10 +36,7 @@ IMPLEMENT_GLOBAL_SHADER(FGenerateMipsCS, "/Engine/Private/ComputeGenerateMips.us
 
 //Initialise the texture for usage with RenderGraph and ComputeGenerateMips shader
 FGenerateMipsStruct* FGenerateMips::SetupTexture(FTextureRHIParamRef InTexture,
-	ESamplerFilter InFilter,
-	ESamplerAddressMode InAddressU,
-	ESamplerAddressMode InAddressV,
-	ESamplerAddressMode InAddressW)
+	const FGenerateMipsParams& InParams)
 {
 	//Currently only 2D textures supported
 	check(InTexture->GetTexture2D());
@@ -79,10 +76,10 @@ FGenerateMipsStruct* FGenerateMips::SetupTexture(FTextureRHIParamRef InTexture,
 		GRenderTargetPool.CreateUntrackedElement(Desc, InTexture->GenMipsStruct->RenderTarget, RenderTexture);
 
 		//Specify the Sampler details based on the input.
-		InTexture->GenMipsStruct->Sampler.Filter = InFilter;
-		InTexture->GenMipsStruct->Sampler.AddressU = InAddressU;
-		InTexture->GenMipsStruct->Sampler.AddressV = InAddressV;
-		InTexture->GenMipsStruct->Sampler.AddressW = InAddressW;
+		InTexture->GenMipsStruct->Sampler.Filter = InParams.Filter;
+		InTexture->GenMipsStruct->Sampler.AddressU = InParams.AddressU;
+		InTexture->GenMipsStruct->Sampler.AddressV = InParams.AddressV;
+		InTexture->GenMipsStruct->Sampler.AddressW = InParams.AddressW;
 	}
 
 	//Return the raw pointer
@@ -142,10 +139,7 @@ void FGenerateMips::Compute(FRHICommandListImmediate& RHIImmCmdList, FTextureRHI
 
 //Public execute function for calling the generate mips compute shader. Handles everything per platform.
 void FGenerateMips::Execute(FRHICommandListImmediate& RHICmdList, FTextureRHIParamRef InTexture,
-	ESamplerFilter InFilter,
-	ESamplerAddressMode InAddressU,
-	ESamplerAddressMode InAddressV,
-	ESamplerAddressMode InAddressW)
+	const FGenerateMipsParams& InParams)
 {
 	//Only executes if mips are required.
 	if (InTexture->GetNumMips() > 1)
@@ -156,7 +150,7 @@ void FGenerateMips::Execute(FRHICommandListImmediate& RHICmdList, FTextureRHIPar
 			//Generate the RenderGraph texture if required.
 			if (!InTexture->GenMipsStruct)
 			{
-				SetupTexture(InTexture, InFilter, InAddressU, InAddressV, InAddressW);
+				SetupTexture(InTexture, InParams);
 			}
 			Compute(RHICmdList, InTexture);
 		}
diff --git a/Engine/Source/Runtime/RenderCore/Public/GenerateMips.h b/Engine/Source/Runtime/RenderCore/Public/GenerateMips.h
--- a/Engine/Source/Runtime/RenderCore/Public/GenerateMips.h
+++ b/Engine/Source/Runtime/RenderCore/Public/GenerateMips.h
@@ -5,22 +5,25 @@
 
 struct FGenerateMipsStruct;
 
+/** Parameters for generating mip maps */
+struct FGenerateMipsParams
+{
+	ESamplerFilter Filter = SF_Bilinear;
+	ESamplerAddressMode AddressU = AM_Clamp;
+	ESamplerAddressMode AddressV = AM_Clamp;
+	ESamplerAddressMode AddressW = AM_Clamp;
+};
+
 class RENDERCORE_API FGenerateMips
 {
 public:
 	//Public function for executing the generate mips compute shader 
 	//Default sampler is always bilinear clamp
 	static void Execute(FRHICommandListImmediate& RHICmdList, FTextureRHIParamRef InTexture,
-		ESamplerFilter InFilter = SF_Bilinear,
-		ESamplerAddressMode InAddressU = AM_Clamp,
-		ESamplerAddressMode InAddressV = AM_Clamp,
-		ESamplerAddressMode InAddressW = AM_Clamp);
+		const FGenerateMipsParams& InParams = FGenerateMipsParams());
 
 private:
 	static void Compute(FRHICommandListImmediate& RHIImmCmdList, FTextureRHIParamRef InTexture);
 	static FGenerateMipsStruct* SetupTexture(FTextureRHIParamRef InTexture,
-		ESamplerFilter InFilter = SF_Bilinear,
-		ESamplerAddressMode InAddressU = AM_Clamp,
-		ESamplerAddressMode InAddressV = AM_Clamp,
-		ESamplerAddressMode InAddressW = AM_Clamp);
+		const FGenerateMipsParams& InParams = FGenerateMipsParams());
 };
\ No newline at end of file
diff --git a/Engine/Source/Runtime/Renderer/Private/MobileSceneCaptureRendering.cpp b/Engine/Source/Runtime/Renderer/Private/MobileSceneCaptureRendering.cpp
--- a/Engine/Source/Runtime/Renderer/Private/MobileSceneCaptureRendering.cpp
+++ b/Engine/Source/Runtime/Renderer/Private/MobileSceneCaptureRendering.cpp
@@ -27,6 +27,7 @@ MobileSceneCaptureRendering.cpp - Mobile specific scene capture code.
 #include "ClearQuad.h"
 #include "PipelineStateCache.h"
 #include "CommonRenderResources.h"
+#include "GenerateMips.h"
 
 /**
 * Shader set for the copy of scene color to capture target, decoding mosaic or RGBE encoded HDR image as part of a
@@ -346,7 +347,9 @@ void UpdateSceneCaptureContentMobile_RenderThread(
 	FRenderTarget* RenderTarget,
 	FTexture* RenderTargetTexture,
 	const FString& EventName,
-	const FResolveParams& ResolveParams)
+	const FResolveParams& ResolveParams,
+	bool bGenerateMips,
+	const FGenerateMipsParams& GenerateMipsParams)
 {
 	FMemMark MemStackMark(FMemStack::Get());
 
@@ -458,6 +461,11 @@ void UpdateSceneCaptureContentMobile_RenderThread(
 			CopyCaptureToTarget(RHICmdList, Target, TargetSize, View, ViewRect, FSceneRenderTargets::Get(RHICmdList).GetSceneColorTexture()->GetTexture2D(), bNeedsFlippedCopy, SceneRenderer);
 		}
 
+		if (bGenerateMips)
+		{
+			FGenerateMips::Execute(RHICmdList, RenderTarget->GetRenderTargetTexture(), GenerateMipsParams);
+		}
+
 		RHICmdList.CopyToResolveTarget(RenderTarget->GetRenderTargetTexture(), RenderTargetTexture->TextureRHI, ResolveParams);
 	}
 	FSceneRenderer::WaitForTasksClearSnapshotsAndDeleteSceneRenderer(RHICmdList, SceneRenderer);
diff --git a/Engine/Source/Runtime/Renderer/Private/MobileSceneCaptureRendering.h b/Engine/Source/Runtime/Renderer/Private/MobileSceneCaptureRendering.h
--- a/Engine/Source/Runtime/Renderer/Private/MobileSceneCaptureRendering.h
+++ b/Engine/Source/Runtime/Renderer/Private/MobileSceneCaptureRendering.h
@@ -9,6 +9,7 @@ class FRHICommandListImmediate;
 class FSceneRenderer;
 class FTexture;
 struct FResolveParams;
+struct FGenerateMipsParams;
 
 void UpdateSceneCaptureContentMobile_RenderThread(
 	FRHICommandListImmediate& RHICmdList,
@@ -16,4 +17,6 @@ void UpdateSceneCaptureContentMobile_RenderThread(
 	FRenderTarget* RenderTarget,
 	FTexture* RenderTargetTexture,
 	const FString& OwnerName,
-	const FResolveParams& ResolveParams);
+	const FResolveParams& ResolveParams,
+	bool bGenerateMips,
+	const FGenerateMipsParams& GenerateMipsParams);
diff --git a/Engine/Source/Runtime/Renderer/Private/SceneCaptureRendering.cpp b/Engine/Source/Runtime/Renderer/Private/SceneCaptureRendering.cpp
--- a/Engine/Source/Runtime/Renderer/Private/SceneCaptureRendering.cpp
+++ b/Engine/Source/Runtime/Renderer/Private/SceneCaptureRendering.cpp
@@ -42,6 +42,7 @@
 #include "RendererModule.h"
 #include "Rendering/MotionVectorSimulation.h"
 #include "SceneViewExtension.h"
+#include "GenerateMips.h"
 
 const TCHAR* GShaderSourceModeDefineName[] =
 {
@@ -320,7 +321,10 @@ static void UpdateSceneCaptureContentDeferred_RenderThread(
 	FRenderTarget* RenderTarget, 
 	FTexture* RenderTargetTexture, 
 	const FString& EventName, 
-	const FResolveParams& ResolveParams)
+	const FResolveParams& ResolveParams,
+	bool bGenerateMips,
+	const FGenerateMipsParams& GenerateMipsParams
+	)
 {
 	FMemMark MemStackMark(FMemStack::Get());
 
@@ -352,6 +356,11 @@ static void UpdateSceneCaptureContentDeferred_RenderThread(
 			SceneRenderer->Render(RHICmdList);
 		}
 
+		if (bGenerateMips)
+		{
+			FGenerateMips::Execute(RHICmdList, RenderTarget->GetRenderTargetTexture(), GenerateMipsParams);
+		}
+
 		// Note: When the ViewFamily.SceneCaptureSource requires scene textures (i.e. SceneCaptureSource != SCS_FinalColorLDR), the copy to RenderTarget 
 		// will be done in CopySceneCaptureComponentToTarget while the GBuffers are still alive for the frame.
 		RHICmdList.CopyToResolveTarget(RenderTarget->GetRenderTargetTexture(), RenderTargetTexture->TextureRHI, ResolveParams);		
@@ -415,7 +424,9 @@ static void UpdateSceneCaptureContent_RenderThread(
 	FRenderTarget* RenderTarget,
 	FTexture* RenderTargetTexture,
 	const FString& EventName,
-	const FResolveParams& ResolveParams)
+	const FResolveParams& ResolveParams,
+	bool bGenerateMips,
+	const FGenerateMipsParams& GenerateMipsParams)
 {
 	FMaterialRenderProxy::UpdateDeferredCachedUniformExpressions();
 
@@ -429,7 +440,9 @@ static void UpdateSceneCaptureContent_RenderThread(
 				RenderTarget,
 				RenderTargetTexture,
 				EventName,
-				ResolveParams);
+				ResolveParams,
+				bGenerateMips,
+				GenerateMipsParams);
 			break;
 		}
 		case EShadingPath::Deferred:
@@ -440,7 +453,9 @@ static void UpdateSceneCaptureContent_RenderThread(
 				RenderTarget,
 				RenderTargetTexture,
 				EventName,
-				ResolveParams);
+				ResolveParams,
+				bGenerateMips,
+				GenerateMipsParams);
 			break;
 		}
 		default:
@@ -682,7 +697,7 @@ void FScene::UpdateSceneCaptureContents(USceneCaptureComponent2D* CaptureCompone
 {
 	check(CaptureComponent);
 
-	if (CaptureComponent->TextureTarget)
+	if (UTextureRenderTarget2D* TextureRenderTarget = CaptureComponent->TextureTarget)
 	{
 		FTransform Transform = CaptureComponent->GetComponentToWorld();
 		FVector ViewLocation = Transform.GetTranslation();
@@ -699,7 +714,7 @@ void FScene::UpdateSceneCaptureContents(USceneCaptureComponent2D* CaptureCompone
 			FPlane(0, 1, 0, 0),
 			FPlane(0, 0, 0, 1));
 		const float FOV = CaptureComponent->FOVAngle * (float)PI / 360.0f;
-		FIntPoint CaptureSize(CaptureComponent->TextureTarget->GetSurfaceWidth(), CaptureComponent->TextureTarget->GetSurfaceHeight());
+		FIntPoint CaptureSize(TextureRenderTarget->GetSurfaceWidth(), TextureRenderTarget->GetSurfaceHeight());
 
 		FMatrix ProjectionMatrix;
 		if (CaptureComponent->bUseCustomProjectionMatrix)
@@ -718,7 +733,7 @@ void FScene::UpdateSceneCaptureContents(USceneCaptureComponent2D* CaptureCompone
 		FSceneRenderer* SceneRenderer = CreateSceneRendererForSceneCapture(
 			this, 
 			CaptureComponent, 
-			CaptureComponent->TextureTarget->GameThread_GetRenderTargetResource(), 
+			TextureRenderTarget->GameThread_GetRenderTargetResource(), 
 			CaptureSize, 
 			ViewRotationMatrix, 
 			ViewLocation, 
@@ -791,7 +806,7 @@ void FScene::UpdateSceneCaptureContents(USceneCaptureComponent2D* CaptureCompone
 		// Reset scene capture's camera cut.
 		CaptureComponent->bCameraCutThisFrame = false;
 
-		FTextureRenderTargetResource* TextureRenderTarget = CaptureComponent->TextureTarget->GameThread_GetRenderTargetResource();
+		FTextureRenderTargetResource* TextureRenderTargetResource = TextureRenderTarget->GameThread_GetRenderTargetResource();
 
 		FString EventName;
 		if (!CaptureComponent->ProfilingEventName.IsEmpty())
@@ -803,10 +818,15 @@ void FScene::UpdateSceneCaptureContents(USceneCaptureComponent2D* CaptureCompone
 			CaptureComponent->GetOwner()->GetFName().ToString(EventName);
 		}
 
+		const bool bGenerateMips = TextureRenderTarget->bAutoGenerateMips;
+		FGenerateMipsParams GenerateMipsParams{TextureRenderTarget->MipsSamplerFilter == TF_Nearest ? SF_Point : (TextureRenderTarget->MipsSamplerFilter == TF_Trilinear ? SF_Trilinear : SF_Bilinear),
+			TextureRenderTarget->MipsAddressU == TA_Wrap ? AM_Wrap : (TextureRenderTarget->MipsAddressU == TA_Mirror ? AM_Mirror : AM_Clamp),
+			TextureRenderTarget->MipsAddressV == TA_Wrap ? AM_Wrap : (TextureRenderTarget->MipsAddressV == TA_Mirror ? AM_Mirror : AM_Clamp)};
+
 		ENQUEUE_RENDER_COMMAND(CaptureCommand)(
-			[SceneRenderer, TextureRenderTarget, EventName](FRHICommandListImmediate& RHICmdList)
+			[SceneRenderer, TextureRenderTargetResource, EventName, bGenerateMips, GenerateMipsParams](FRHICommandListImmediate& RHICmdList)
 			{
-				UpdateSceneCaptureContent_RenderThread(RHICmdList, SceneRenderer, TextureRenderTarget, TextureRenderTarget, EventName, FResolveParams());
+				UpdateSceneCaptureContent_RenderThread(RHICmdList, SceneRenderer, TextureRenderTargetResource, TextureRenderTargetResource, EventName, FResolveParams(), bGenerateMips, GenerateMipsParams);
 			}
 		);
 	}
@@ -927,7 +947,7 @@ void FScene::UpdateSceneCaptureContents(USceneCaptureComponentCube* CaptureCompo
 				ENQUEUE_RENDER_COMMAND(CaptureCommand)(
 					[SceneRenderer, TextureRenderTarget, EventName, TargetFace](FRHICommandListImmediate& RHICmdList)
 				{
-					UpdateSceneCaptureContent_RenderThread(RHICmdList, SceneRenderer, TextureRenderTarget, TextureRenderTarget, EventName, FResolveParams(FResolveRect(), TargetFace));
+					UpdateSceneCaptureContent_RenderThread(RHICmdList, SceneRenderer, TextureRenderTarget, TextureRenderTarget, EventName, FResolveParams(FResolveRect(), TargetFace), false, FGenerateMipsParams());
 				}
 				);
 			}
