Removing MIXEDREALITY_CHANGE sections in engine-wide source files.

All changes were avoiding asserts/ensures/other checks in the PLATFORM_HOLOLENS case, but HoloLensQA project runs for me on device with all of the changes removed.  If we run into issues in the future, we'll have to solve them in ways more contained to the HoloLens-specific modules.

diff --git a/Engine/Source/Runtime/Core/Private/Math/UnrealMath.cpp b/Engine/Source/Runtime/Core/Private/Math/UnrealMath.cpp
--- a/Engine/Source/Runtime/Core/Private/Math/UnrealMath.cpp
+++ b/Engine/Source/Runtime/Core/Private/Math/UnrealMath.cpp
@@ -452,9 +452,7 @@ FQuat FRotator::Quaternion() const
 	RotationQuat.W =  CR*CP*CY + SR*SP*SY;
 #endif // PLATFORM_ENABLE_VECTORINTRINSICS
 
-	// @MIXEDREALITY_CHANGE : BEGIN - This check is blocking HoloLens apps from launching in Development mode.
-#if (ENABLE_NAN_DIAGNOSTIC || DO_CHECK) && !PLATFORM_HOLOLENS
-	// @MIXEDREALITY_CHANGE : END
+#if ENABLE_NAN_DIAGNOSTIC || DO_CHECK
 	// Very large inputs can cause NaN's. Want to catch this here
 	if (RotationQuat.ContainsNaN())
 	{
diff --git a/Engine/Source/Runtime/Engine/Private/Components/SceneComponent.cpp b/Engine/Source/Runtime/Engine/Private/Components/SceneComponent.cpp
--- a/Engine/Source/Runtime/Engine/Private/Components/SceneComponent.cpp
+++ b/Engine/Source/Runtime/Engine/Private/Components/SceneComponent.cpp
@@ -619,9 +619,7 @@ void USceneComponent::UpdateComponentToWorldWithParent(USceneComponent* Parent,F
 		NewTransform = CalcNewComponentToWorld(RelativeTransform, Parent, SocketName);
 	}
 
-	// @MIXEDREALITY_CHANGE : BEGIN - This check is blocking HoloLens apps from launching in Development mode.
-#if DO_CHECK && !PLATFORM_HOLOLENS
-	// @MIXEDREALITY_CHANGE : END
+#if DO_CHECK
 	ensure(NewTransform.IsValid());
 #endif
 
diff --git a/Engine/Source/Runtime/Engine/Private/Rendering/SeparableSSS.cpp b/Engine/Source/Runtime/Engine/Private/Rendering/SeparableSSS.cpp
--- a/Engine/Source/Runtime/Engine/Private/Rendering/SeparableSSS.cpp
+++ b/Engine/Source/Runtime/Engine/Private/Rendering/SeparableSSS.cpp
@@ -166,11 +166,7 @@ void ComputeMirroredSSSKernel(FLinearColor* TargetBuffer, uint32 TargetBufferSiz
 
 	// generate output (remove negative samples)
 	{
-		// @MIXEDREALITY_CHANGE : BEGIN - TODO: this is crashing on HoloLens.
-		#if !PLATFORM_HOLOLENS
 		check(kernel[0].A == 0.0f);
-		#endif
-		// @MIXEDREALITY_CHANGE : END
 
 		// center sample
 		TargetBuffer[0] = kernel[0];
diff --git a/Engine/Source/Runtime/Renderer/Private/RendererScene.cpp b/Engine/Source/Runtime/Renderer/Private/RendererScene.cpp
--- a/Engine/Source/Runtime/Renderer/Private/RendererScene.cpp
+++ b/Engine/Source/Runtime/Renderer/Private/RendererScene.cpp
@@ -1382,12 +1382,8 @@ void FScene::UpdatePrimitiveTransform(UPrimitiveComponent* Primitive)
 			UpdateParams.PreviousTransform = FMotionVectorSimulation::Get().GetPreviousTransform(Primitive);
 
 			// Help track down primitive with bad bounds way before the it gets to the Renderer
-			// @MIXEDREALITY_CHANGE : BEGIN - This check is blocking HoloLens apps from launching in Development mode.
-#if !PLATFORM_HOLOLENS
-			// @MIXEDREALITY_CHANGE : END
 			ensureMsgf(!Primitive->Bounds.BoxExtent.ContainsNaN() && !Primitive->Bounds.Origin.ContainsNaN() && !FMath::IsNaN(Primitive->Bounds.SphereRadius) && FMath::IsFinite(Primitive->Bounds.SphereRadius),
 				TEXT("Nans found on Bounds for Primitive %s: Origin %s, BoxExtent %s, SphereRadius %f"), *Primitive->GetName(), *Primitive->Bounds.Origin.ToString(), *Primitive->Bounds.BoxExtent.ToString(), Primitive->Bounds.SphereRadius);
-#endif
 
 			ENQUEUE_RENDER_COMMAND(UpdateTransformCommand)(
 				[UpdateParams](FRHICommandListImmediate& RHICmdList)
diff --git a/Engine/Source/Runtime/Renderer/Private/ShadowSetup.cpp b/Engine/Source/Runtime/Renderer/Private/ShadowSetup.cpp
--- a/Engine/Source/Runtime/Renderer/Private/ShadowSetup.cpp
+++ b/Engine/Source/Runtime/Renderer/Private/ShadowSetup.cpp
@@ -1830,15 +1830,11 @@ void FSceneRenderer::CreatePerObjectProjectedShadow(
 	// Compute the composite bounds of this group of shadow primitives.
 	FBoxSphereBounds OriginalBounds = ShadowGroupPrimitives[0]->Proxy->GetBounds();
 
-	// @MIXEDREALITY_CHANGE : BEGIN - This check is blocking HoloLens apps from launching in Development mode.
-#if !PLATFORM_HOLOLENS
-	// @MIXEDREALITY_CHANGE : END
 	if (!ensureMsgf(OriginalBounds.ContainsNaN() == false, TEXT("OriginalBound contains NaN : %s"), *OriginalBounds.ToString()))
 	{
 		// fix up OriginalBounds. This is going to cause flickers
 		OriginalBounds = FBoxSphereBounds(FVector::ZeroVector, FVector(1.f), 1.f);
 	}
-#endif
 
 	for (int32 ChildIndex = 1; ChildIndex < ShadowGroupPrimitives.Num(); ChildIndex++)
 	{
@@ -1888,11 +1884,7 @@ void FSceneRenderer::CreatePerObjectProjectedShadow(
 			OriginalBounds.SphereRadius /
 			FMath::Max(ShadowViewDistFromBounds, 1.0f);
 		// Early catch for invalid CalculateShadowFadeAlpha()
-		// @MIXEDREALITY_CHANGE : BEGIN - This check is blocking HoloLens apps from launching in Development mode.
-#if !PLATFORM_HOLOLENS
-		// @MIXEDREALITY_CHANGE : END
 		ensureMsgf(ScreenRadius >= 0.0f, TEXT("View.ShadowViewMatrices.ScreenScale %f, OriginalBounds.SphereRadius %f, ShadowViewDistFromBounds %f"), View.ShadowViewMatrices.GetScreenScale(), OriginalBounds.SphereRadius, ShadowViewDistFromBounds);
-#endif
 
 		const float ScreenPercent = FMath::Max(
 			1.0f / 2.0f * View.ShadowViewMatrices.GetProjectionScale().X,
