UE4.23 - Deprecate FUniformBufferRHIParamRef, FIndexBufferRHIParamRef, FVertexBufferRHIParamRef, FStructuredBufferRHIParamRef, FTextureRHIParamRef, FTexture2DRHIParamRef, FTexture2DArrayRHIParamRef, F

diff --git a/Engine/Source/Runtime/RHI/Public/RHIValidation.h b/Engine/Source/Runtime/RHI/Public/RHIValidation.h
--- a/Engine/Source/Runtime/RHI/Public/RHIValidation.h
+++ b/Engine/Source/Runtime/RHI/Public/RHIValidation.h
@@ -676,7 +676,7 @@ public:
 	* Creates a shader resource view for a texture
 	*/
 	// FlushType: Wait RHI Thread
-	virtual FShaderResourceViewRHIRef RHICreateShaderResourceView(FTextureRHIParamRef TextureRHI, const FRHITextureSRVCreateInfo& CreateInfo) override final
+	virtual FShaderResourceViewRHIRef RHICreateShaderResourceView(FRHITexture* TextureRHI, const FRHITextureSRVCreateInfo& CreateInfo) override final
 	{
 		return RHI->RHICreateShaderResourceView(TextureRHI, CreateInfo);
 	}
@@ -1048,7 +1048,7 @@ public:
 	/** Sets stream output targets, for use with a geometry shader created with RHICreateGeometryShaderWithStreamOutput. */
 	//@todo this should be a CMDLIST method
 	// FlushType: Flush Immediate (seems wrong)
-	virtual void RHISetStreamOutTargets(uint32 NumTargets, const FVertexBufferRHIParamRef* VertexBuffers, const uint32* Offsets) override final
+	virtual void RHISetStreamOutTargets(uint32 NumTargets, FRHIVertexBuffer* const* VertexBuffers, const uint32* Offsets) override final
 	{
 		RHI->RHISetStreamOutTargets(NumTargets, VertexBuffers, Offsets);
 	}
@@ -1441,7 +1441,7 @@ public:
 		return RHI->RHICreateUnorderedAccessView_RenderThread(RHICmdList, IndexBuffer, Format);
 	}
 
-	virtual FShaderResourceViewRHIRef RHICreateShaderResourceView_RenderThread(class FRHICommandListImmediate& RHICmdList, FTextureRHIParamRef Texture, const FRHITextureSRVCreateInfo& CreateInfo) override final
+	virtual FShaderResourceViewRHIRef RHICreateShaderResourceView_RenderThread(class FRHICommandListImmediate& RHICmdList, FRHITexture* Texture, const FRHITextureSRVCreateInfo& CreateInfo) override final
 	{
 		return RHI->RHICreateShaderResourceView_RenderThread(RHICmdList, Texture, CreateInfo);
 	}
diff --git a/Engine/Source/Runtime/RHI/Public/RHIValidationCommon.h b/Engine/Source/Runtime/RHI/Public/RHIValidationCommon.h
--- a/Engine/Source/Runtime/RHI/Public/RHIValidationCommon.h
+++ b/Engine/Source/Runtime/RHI/Public/RHIValidationCommon.h
@@ -19,8 +19,8 @@ class FValidationRHIUtils
 {
 public:
 	static void ValidateCopyTexture(
-		FTextureRHIParamRef	SourceTexture,
-		FTextureRHIParamRef	DestTexture,
+		FRHITexture*	SourceTexture,
+		FRHITexture*	DestTexture,
 		FIntVector			CopySize		= FIntVector::ZeroValue,
 		const FIntVector&	SourcePosition	= FIntVector::ZeroValue,
 		const FIntVector&	DestPosition	= FIntVector::ZeroValue)
diff --git a/Engine/Source/Runtime/RHI/Public/RHIValidationContext.h b/Engine/Source/Runtime/RHI/Public/RHIValidationContext.h
--- a/Engine/Source/Runtime/RHI/Public/RHIValidationContext.h
+++ b/Engine/Source/Runtime/RHI/Public/RHIValidationContext.h
@@ -44,7 +44,7 @@ public:
 		RHIContext->RHIDispatchComputeShader(ThreadGroupCountX, ThreadGroupCountY, ThreadGroupCountZ);
 	}
 
-	virtual void RHIDispatchIndirectComputeShader(FVertexBufferRHIParamRef ArgumentBuffer, uint32 ArgumentOffset) override final
+	virtual void RHIDispatchIndirectComputeShader(FRHIVertexBuffer* ArgumentBuffer, uint32 ArgumentOffset) override final
 	{
 		RHIContext->RHIDispatchIndirectComputeShader(ArgumentBuffer, ArgumentOffset);
 	}
@@ -76,7 +76,7 @@ public:
 	}
 
 	/** Set the shader resource view of a surface.  This is used for binding TextureMS parameter types that need a multi sampled view. */
-	virtual void RHISetShaderTexture(FRHIComputeShader* Shader, uint32 TextureIndex, FTextureRHIParamRef NewTexture) override final
+	virtual void RHISetShaderTexture(FRHIComputeShader* Shader, uint32 TextureIndex, FRHITexture* NewTexture) override final
 	{
 		checkf(State.bComputeShaderSet, TEXT("A Compute shader has to be set to set resources into a shader!"));
 		RHIContext->RHISetShaderTexture(Shader, TextureIndex, NewTexture);
@@ -125,7 +125,7 @@ public:
 		RHIContext->RHISetShaderResourceViewParameter(Shader, SamplerIndex, SRV);
 	}
 
-	virtual void RHISetShaderUniformBuffer(FRHIComputeShader* Shader, uint32 BufferIndex, FUniformBufferRHIParamRef Buffer) override final
+	virtual void RHISetShaderUniformBuffer(FRHIComputeShader* Shader, uint32 BufferIndex, FRHIUniformBuffer* Buffer) override final
 	{
 		checkf(State.bComputeShaderSet, TEXT("A Compute shader has to be set to set resources into a shader!"));
 		RHIContext->RHISetShaderUniformBuffer(Shader, BufferIndex, Buffer);
@@ -152,7 +152,7 @@ public:
 		RHIContext->RHIWriteGPUFence(FenceRHI);
 	}
 
-	virtual void RHICopyToStagingBuffer(FVertexBufferRHIParamRef SourceBufferRHI, FRHIStagingBuffer* DestinationStagingBufferRHI, uint32 InOffset, uint32 InNumBytes) override final
+	virtual void RHICopyToStagingBuffer(FRHIVertexBuffer* SourceBufferRHI, FRHIStagingBuffer* DestinationStagingBufferRHI, uint32 InOffset, uint32 InNumBytes) override final
 	{
 		RHIContext->RHICopyToStagingBuffer(SourceBufferRHI, DestinationStagingBufferRHI, InOffset, InNumBytes);
 	}
@@ -203,7 +203,7 @@ public:
 		RHIContext->RHIDispatchComputeShader(ThreadGroupCountX, ThreadGroupCountY, ThreadGroupCountZ);
 	}
 
-	virtual void RHIDispatchIndirectComputeShader(FVertexBufferRHIParamRef ArgumentBuffer, uint32 ArgumentOffset) override final
+	virtual void RHIDispatchIndirectComputeShader(FRHIVertexBuffer* ArgumentBuffer, uint32 ArgumentOffset) override final
 	{
 		RHIContext->RHIDispatchIndirectComputeShader(ArgumentBuffer, ArgumentOffset);
 	}
@@ -243,7 +243,7 @@ public:
 	* @param DestTexture - texture to resolve to, 0 is silenty ignored
 	* @param ResolveParams - optional resolve params
 	*/
-	virtual void RHICopyToResolveTarget(FTextureRHIParamRef SourceTexture, FTextureRHIParamRef DestTexture, const FResolveParams& ResolveParams) override final
+	virtual void RHICopyToResolveTarget(FRHITexture* SourceTexture, FRHITexture* DestTexture, const FResolveParams& ResolveParams) override final
 	{
 		RHIContext->RHICopyToResolveTarget(SourceTexture, DestTexture, ResolveParams);
 	}
@@ -257,7 +257,7 @@ public:
 	* @param InTextures - array of texture objects to transition
 	* @param NumTextures - number of textures to transition
 	*/
-	virtual void RHITransitionResources(EResourceTransitionAccess TransitionType, FTextureRHIParamRef* InTextures, int32 NumTextures) override final
+	virtual void RHITransitionResources(EResourceTransitionAccess TransitionType, FRHITexture** InTextures, int32 NumTextures) override final
 	{
 		RHIContext->RHITransitionResources(TransitionType, InTextures, NumTextures);
 	}
@@ -307,7 +307,7 @@ public:
 	}
 
 	// This method is queued with an RHIThread, otherwise it will flush after it is queued; without an RHI thread there is no benefit to queuing this frame advance commands
-	virtual void RHIBeginDrawingViewport(FRHIViewport* Viewport, FTextureRHIParamRef RenderTargetRHI) override final
+	virtual void RHIBeginDrawingViewport(FRHIViewport* Viewport, FRHITexture* RenderTargetRHI) override final
 	{
 		RHIContext->RHIBeginDrawingViewport(Viewport, RenderTargetRHI);
 	}
@@ -354,12 +354,12 @@ public:
 	/**
 	* Signals the beginning and ending of rendering to a resource to be used in the next frame on a multiGPU system
 	*/
-	virtual void RHIBeginUpdateMultiFrameResource(FTextureRHIParamRef Texture) override final
+	virtual void RHIBeginUpdateMultiFrameResource(FRHITexture* Texture) override final
 	{
 		RHIContext->RHIBeginUpdateMultiFrameResource(Texture);
 	}
 
-	virtual void RHIEndUpdateMultiFrameResource(FTextureRHIParamRef Texture) override final
+	virtual void RHIEndUpdateMultiFrameResource(FRHITexture* Texture) override final
 	{
 		RHIContext->RHIEndUpdateMultiFrameResource(Texture);
 	}
@@ -374,7 +374,7 @@ public:
 		RHIContext->RHIEndUpdateMultiFrameResource(UAV);
 	}
 
-	virtual void RHISetStreamSource(uint32 StreamIndex, FVertexBufferRHIParamRef VertexBuffer, uint32 Offset) override final
+	virtual void RHISetStreamSource(uint32 StreamIndex, FRHIVertexBuffer* VertexBuffer, uint32 Offset) override final
 	{
 		//#todo-rco: Decide if this is needed or not...
 		//checkf(State.bGfxPSOSet, TEXT("A Graphics PSO has to be set to set-up the vertex streams!"));
@@ -415,42 +415,42 @@ public:
 	}
 
 	/** Set the shader resource view of a surface.  This is used for binding TextureMS parameter types that need a multi sampled view. */
-	virtual void RHISetShaderTexture(FRHIVertexShader* Shader, uint32 TextureIndex, FTextureRHIParamRef NewTexture) override final
+	virtual void RHISetShaderTexture(FRHIVertexShader* Shader, uint32 TextureIndex, FRHITexture* NewTexture) override final
 	{
 		checkf(State.bGfxPSOSet, TEXT("A Graphics PSO has to be set to set resources into a shader!"));
 		RHIContext->RHISetShaderTexture(Shader, TextureIndex, NewTexture);
 	}
 
 	/** Set the shader resource view of a surface.  This is used for binding TextureMS parameter types that need a multi sampled view. */
-	virtual void RHISetShaderTexture(FRHIHullShader* Shader, uint32 TextureIndex, FTextureRHIParamRef NewTexture) override final
+	virtual void RHISetShaderTexture(FRHIHullShader* Shader, uint32 TextureIndex, FRHITexture* NewTexture) override final
 	{
 		checkf(State.bGfxPSOSet, TEXT("A Graphics PSO has to be set to set resources into a shader!"));
 		RHIContext->RHISetShaderTexture(Shader, TextureIndex, NewTexture);
 	}
 
 	/** Set the shader resource view of a surface.  This is used for binding TextureMS parameter types that need a multi sampled view. */
-	virtual void RHISetShaderTexture(FRHIDomainShader* Shader, uint32 TextureIndex, FTextureRHIParamRef NewTexture) override final
+	virtual void RHISetShaderTexture(FRHIDomainShader* Shader, uint32 TextureIndex, FRHITexture* NewTexture) override final
 	{
 		checkf(State.bGfxPSOSet, TEXT("A Graphics PSO has to be set to set resources into a shader!"));
 		RHIContext->RHISetShaderTexture(Shader, TextureIndex, NewTexture);
 	}
 
 	/** Set the shader resource view of a surface.  This is used for binding TextureMS parameter types that need a multi sampled view. */
-	virtual void RHISetShaderTexture(FRHIGeometryShader* Shader, uint32 TextureIndex, FTextureRHIParamRef NewTexture) override final
+	virtual void RHISetShaderTexture(FRHIGeometryShader* Shader, uint32 TextureIndex, FRHITexture* NewTexture) override final
 	{
 		checkf(State.bGfxPSOSet, TEXT("A Graphics PSO has to be set to set resources into a shader!"));
 		RHIContext->RHISetShaderTexture(Shader, TextureIndex, NewTexture);
 	}
 
 	/** Set the shader resource view of a surface.  This is used for binding TextureMS parameter types that need a multi sampled view. */
-	virtual void RHISetShaderTexture(FRHIPixelShader* Shader, uint32 TextureIndex, FTextureRHIParamRef NewTexture) override final
+	virtual void RHISetShaderTexture(FRHIPixelShader* Shader, uint32 TextureIndex, FRHITexture* NewTexture) override final
 	{
 		checkf(State.bGfxPSOSet, TEXT("A Graphics PSO has to be set to set resources into a shader!"));
 		RHIContext->RHISetShaderTexture(Shader, TextureIndex, NewTexture);
 	}
 
 	/** Set the shader resource view of a surface.  This is used for binding TextureMS parameter types that need a multi sampled view. */
-	virtual void RHISetShaderTexture(FRHIComputeShader* Shader, uint32 TextureIndex, FTextureRHIParamRef NewTexture) override final
+	virtual void RHISetShaderTexture(FRHIComputeShader* Shader, uint32 TextureIndex, FRHITexture* NewTexture) override final
 	{
 		checkf(State.bComputeShaderSet, TEXT("A Compute shader has to be set to set resources into a shader!"));
 		RHIContext->RHISetShaderTexture(Shader, TextureIndex, NewTexture);
@@ -590,37 +590,37 @@ public:
 		RHIContext->RHISetShaderResourceViewParameter(Shader, SamplerIndex, SRV);
 	}
 
-	virtual void RHISetShaderUniformBuffer(FRHIVertexShader* Shader, uint32 BufferIndex, FUniformBufferRHIParamRef Buffer) override final
+	virtual void RHISetShaderUniformBuffer(FRHIVertexShader* Shader, uint32 BufferIndex, FRHIUniformBuffer* Buffer) override final
 	{
 		checkf(State.bGfxPSOSet, TEXT("A Graphics PSO has to be set to set resources into a shader!"));
 		RHIContext->RHISetShaderUniformBuffer(Shader, BufferIndex, Buffer);
 	}
 
-	virtual void RHISetShaderUniformBuffer(FRHIHullShader* Shader, uint32 BufferIndex, FUniformBufferRHIParamRef Buffer) override final
+	virtual void RHISetShaderUniformBuffer(FRHIHullShader* Shader, uint32 BufferIndex, FRHIUniformBuffer* Buffer) override final
 	{
 		checkf(State.bGfxPSOSet, TEXT("A Graphics PSO has to be set to set resources into a shader!"));
 		RHIContext->RHISetShaderUniformBuffer(Shader, BufferIndex, Buffer);
 	}
 
-	virtual void RHISetShaderUniformBuffer(FRHIDomainShader* Shader, uint32 BufferIndex, FUniformBufferRHIParamRef Buffer) override final
+	virtual void RHISetShaderUniformBuffer(FRHIDomainShader* Shader, uint32 BufferIndex, FRHIUniformBuffer* Buffer) override final
 	{
 		checkf(State.bGfxPSOSet, TEXT("A Graphics PSO has to be set to set resources into a shader!"));
 		RHIContext->RHISetShaderUniformBuffer(Shader, BufferIndex, Buffer);
 	}
 
-	virtual void RHISetShaderUniformBuffer(FRHIGeometryShader* Shader, uint32 BufferIndex, FUniformBufferRHIParamRef Buffer) override final
+	virtual void RHISetShaderUniformBuffer(FRHIGeometryShader* Shader, uint32 BufferIndex, FRHIUniformBuffer* Buffer) override final
 	{
 		checkf(State.bGfxPSOSet, TEXT("A Graphics PSO has to be set to set resources into a shader!"));
 		RHIContext->RHISetShaderUniformBuffer(Shader, BufferIndex, Buffer);
 	}
 
-	virtual void RHISetShaderUniformBuffer(FRHIPixelShader* Shader, uint32 BufferIndex, FUniformBufferRHIParamRef Buffer) override final
+	virtual void RHISetShaderUniformBuffer(FRHIPixelShader* Shader, uint32 BufferIndex, FRHIUniformBuffer* Buffer) override final
 	{
 		checkf(State.bGfxPSOSet, TEXT("A Graphics PSO has to be set to set resources into a shader!"));
 		RHIContext->RHISetShaderUniformBuffer(Shader, BufferIndex, Buffer);
 	}
 
-	virtual void RHISetShaderUniformBuffer(FRHIComputeShader* Shader, uint32 BufferIndex, FUniformBufferRHIParamRef Buffer) override final
+	virtual void RHISetShaderUniformBuffer(FRHIComputeShader* Shader, uint32 BufferIndex, FRHIUniformBuffer* Buffer) override final
 	{
 		checkf(State.bComputeShaderSet, TEXT("A Compute shader has to be set to set resources into a shader!"));
 		RHIContext->RHISetShaderUniformBuffer(Shader, BufferIndex, Buffer);
@@ -698,26 +698,26 @@ public:
 		RHIContext->RHIDrawPrimitive(BaseVertexIndex, NumPrimitives, NumInstances);
 	}
 
-	virtual void RHIDrawPrimitiveIndirect(FVertexBufferRHIParamRef ArgumentBuffer, uint32 ArgumentOffset) override final
+	virtual void RHIDrawPrimitiveIndirect(FRHIVertexBuffer* ArgumentBuffer, uint32 ArgumentOffset) override final
 	{
 		checkf(State.bGfxPSOSet, TEXT("A Graphics PSO has to be set to draw!"));
 		RHIContext->RHIDrawPrimitiveIndirect(ArgumentBuffer, ArgumentOffset);
 	}
 
-	virtual void RHIDrawIndexedIndirect(FIndexBufferRHIParamRef IndexBufferRHI, FStructuredBufferRHIParamRef ArgumentsBufferRHI, int32 DrawArgumentsIndex, uint32 NumInstances) override final
+	virtual void RHIDrawIndexedIndirect(FRHIIndexBuffer* IndexBufferRHI, FRHIStructuredBuffer* ArgumentsBufferRHI, int32 DrawArgumentsIndex, uint32 NumInstances) override final
 	{
 		checkf(State.bGfxPSOSet, TEXT("A Graphics PSO has to be set to draw!"));
 		RHIContext->RHIDrawIndexedIndirect(IndexBufferRHI, ArgumentsBufferRHI, DrawArgumentsIndex, NumInstances);
 	}
 
 	// @param NumPrimitives need to be >0 
-	virtual void RHIDrawIndexedPrimitive(FIndexBufferRHIParamRef IndexBuffer, int32 BaseVertexIndex, uint32 FirstInstance, uint32 NumVertices, uint32 StartIndex, uint32 NumPrimitives, uint32 NumInstances) override final
+	virtual void RHIDrawIndexedPrimitive(FRHIIndexBuffer* IndexBuffer, int32 BaseVertexIndex, uint32 FirstInstance, uint32 NumVertices, uint32 StartIndex, uint32 NumPrimitives, uint32 NumInstances) override final
 	{
 		checkf(State.bGfxPSOSet, TEXT("A Graphics PSO has to be set to draw!"));
 		RHIContext->RHIDrawIndexedPrimitive(IndexBuffer, BaseVertexIndex, FirstInstance, NumVertices, StartIndex, NumPrimitives, NumInstances);
 	}
 
-	virtual void RHIDrawIndexedPrimitiveIndirect(FIndexBufferRHIParamRef IndexBuffer, FVertexBufferRHIParamRef ArgumentBuffer, uint32 ArgumentOffset) override final
+	virtual void RHIDrawIndexedPrimitiveIndirect(FRHIIndexBuffer* IndexBuffer, FRHIVertexBuffer* ArgumentBuffer, uint32 ArgumentOffset) override final
 	{
 		checkf(State.bGfxPSOSet, TEXT("A Graphics PSO has to be set to draw!"));
 		RHIContext->RHIDrawIndexedPrimitiveIndirect(IndexBuffer, ArgumentBuffer, ArgumentOffset);
@@ -745,7 +745,7 @@ public:
 		RHIContext->RHIPopEvent();
 	}
 
-	virtual void RHIUpdateTextureReference(FRHITextureReference* TextureRef, FTextureRHIParamRef NewTexture) override final
+	virtual void RHIUpdateTextureReference(FRHITextureReference* TextureRef, FRHITexture* NewTexture) override final
 	{
 		RHIContext->RHIUpdateTextureReference(TextureRef, NewTexture);
 	}
@@ -792,12 +792,12 @@ public:
 		RHIContext->RHIWriteGPUFence(FenceRHI);
 	}
 
-	virtual void RHICopyToStagingBuffer(FVertexBufferRHIParamRef SourceBufferRHI, FRHIStagingBuffer* DestinationStagingBufferRHI, uint32 InOffset, uint32 InNumBytes) override final
+	virtual void RHICopyToStagingBuffer(FRHIVertexBuffer* SourceBufferRHI, FRHIStagingBuffer* DestinationStagingBufferRHI, uint32 InOffset, uint32 InNumBytes) override final
 	{
 		RHIContext->RHICopyToStagingBuffer(SourceBufferRHI, DestinationStagingBufferRHI, InOffset, InNumBytes);
 	}
 
-	virtual void RHICopyTexture(FTextureRHIParamRef SourceTexture, FTextureRHIParamRef DestTexture, const FRHICopyTextureInfo& CopyInfo) override final
+	virtual void RHICopyTexture(FRHITexture* SourceTexture, FRHITexture* DestTexture, const FRHICopyTextureInfo& CopyInfo) override final
 	{
 		ensureMsgf(!State.bInsideBeginRenderPass, TEXT("Copying inside a RenderPass is not efficient!"));
 		FValidationRHIUtils::ValidateCopyTexture(SourceTexture, DestTexture, CopyInfo.Size, CopyInfo.SourcePosition, CopyInfo.DestPosition);
@@ -851,7 +851,7 @@ public:
 	virtual void RHISetRayTracingHitGroup(
 		FRHIRayTracingScene* Scene, uint32 InstanceIndex, uint32 SegmentIndex, uint32 ShaderSlot,
 		FRHIRayTracingPipelineState* Pipeline, uint32 HitGroupIndex,
-		uint32 NumUniformBuffers, const FUniformBufferRHIParamRef* UniformBuffers,
+		uint32 NumUniformBuffers, FRHIUniformBuffer* const* UniformBuffers,
 		uint32 LooseParameterDataSize, const void* LooseParameterData,
 		uint32 UserData) override final
 	{
diff --git a/Engine/Source/Runtime/RenderCore/Private/DynamicBufferAllocator.cpp b/Engine/Source/Runtime/RenderCore/Private/DynamicBufferAllocator.cpp
--- a/Engine/Source/Runtime/RenderCore/Private/DynamicBufferAllocator.cpp
+++ b/Engine/Source/Runtime/RenderCore/Private/DynamicBufferAllocator.cpp
@@ -78,12 +78,12 @@ void FGlobalDynamicReadBuffer::Cleanup()
 }
 void FGlobalDynamicReadBuffer::InitRHI()
 {
-	UE_LOG(LogRendererCore, Log, TEXT("FGlobalReadBuffer::InitRHI"));
+	UE_LOG(LogRendererCore, Verbose, TEXT("FGlobalReadBuffer::InitRHI"));
 }
 
 void FGlobalDynamicReadBuffer::ReleaseRHI()
 {
-	UE_LOG(LogRendererCore, Log, TEXT("FGlobalReadBuffer::ReleaseRHI"));
+	UE_LOG(LogRendererCore, Verbose, TEXT("FGlobalReadBuffer::ReleaseRHI"));
 	Cleanup();
 }
 
diff --git a/Engine/Source/Runtime/RenderCore/Private/GenerateMips.cpp b/Engine/Source/Runtime/RenderCore/Private/GenerateMips.cpp
--- a/Engine/Source/Runtime/RenderCore/Private/GenerateMips.cpp
+++ b/Engine/Source/Runtime/RenderCore/Private/GenerateMips.cpp
@@ -35,7 +35,7 @@ class FGenerateMipsCS : public FGlobalShader
 IMPLEMENT_GLOBAL_SHADER(FGenerateMipsCS, "/Engine/Private/ComputeGenerateMips.usf", "MainCS", SF_Compute);
 
 //Initialise the texture for usage with RenderGraph and ComputeGenerateMips shader
-FGenerateMipsStruct* FGenerateMips::SetupTexture(FTextureRHIParamRef InTexture,
+FGenerateMipsStruct* FGenerateMips::SetupTexture(FRHITexture* InTexture,
 	const FGenerateMipsParams& InParams)
 {
 	//Currently only 2D textures supported
@@ -87,7 +87,7 @@ FGenerateMipsStruct* FGenerateMips::SetupTexture(FTextureRHIParamRef InTexture,
 }
 
 //Compute shader execution function for generating mips in real time.
-void FGenerateMips::Compute(FRHICommandListImmediate& RHIImmCmdList, FTextureRHIParamRef InTexture)
+void FGenerateMips::Compute(FRHICommandListImmediate& RHIImmCmdList, FRHITexture* InTexture)
 {
 	check(IsInRenderingThread());	
 	//Currently only 2D textures supported
@@ -138,7 +138,7 @@ void FGenerateMips::Compute(FRHICommandListImmediate& RHIImmCmdList, FTextureRHI
 }
 
 //Public execute function for calling the generate mips compute shader. Handles everything per platform.
-void FGenerateMips::Execute(FRHICommandListImmediate& RHICmdList, FTextureRHIParamRef InTexture,
+void FGenerateMips::Execute(FRHICommandListImmediate& RHICmdList, FRHITexture* InTexture,
 	const FGenerateMipsParams& InParams)
 {
 	//Only executes if mips are required.
diff --git a/Engine/Source/Runtime/RenderCore/Private/RenderGraphBuilder.cpp b/Engine/Source/Runtime/RenderCore/Private/RenderGraphBuilder.cpp
--- a/Engine/Source/Runtime/RenderCore/Private/RenderGraphBuilder.cpp
+++ b/Engine/Source/Runtime/RenderCore/Private/RenderGraphBuilder.cpp
@@ -775,7 +775,7 @@ void FRDGBuilder::TransitionTexture(FRDGTexture* Texture, ERDGPassAccess PassAcc
 
 	if (PassAccessBefore != PassAccessAfter || Texture->PassPipeline != PassPipelineAfter)
 	{
-		FTextureRHIParamRef RHITexture = Texture->GetRHIUnchecked();
+		FRHITexture* RHITexture = Texture->GetRHIUnchecked();
 
 		const bool bIsMultiFrameResource = (Texture->Flags & ERDGResourceFlags::MultiFrame) == ERDGResourceFlags::MultiFrame;
 
diff --git a/Engine/Source/Runtime/RenderCore/Private/RenderTargetPool.cpp b/Engine/Source/Runtime/RenderCore/Private/RenderTargetPool.cpp
--- a/Engine/Source/Runtime/RenderCore/Private/RenderTargetPool.cpp
+++ b/Engine/Source/Runtime/RenderCore/Private/RenderTargetPool.cpp
@@ -197,7 +197,7 @@ void FRenderTargetPool::TransitionTargetsWritable(FRHICommandListImmediate& RHIC
 		FPooledRenderTarget* PooledRT = PooledRenderTargets[i];
 		if (PooledRT && PooledRT->GetDesc().AutoWritable)
 		{
-			FTextureRHIParamRef RenderTarget = PooledRT->GetRenderTargetItem().TargetableTexture;
+			FRHITexture* RenderTarget = PooledRT->GetRenderTargetItem().TargetableTexture;
 			if (RenderTarget)
 			{				
 				TransitionTargets.Add(RenderTarget);
diff --git a/Engine/Source/Runtime/RenderCore/Private/Shader.cpp b/Engine/Source/Runtime/RenderCore/Private/Shader.cpp
--- a/Engine/Source/Runtime/RenderCore/Private/Shader.cpp
+++ b/Engine/Source/Runtime/RenderCore/Private/Shader.cpp
@@ -2037,7 +2037,7 @@ void DispatchComputeShader(
 void DispatchIndirectComputeShader(
 	FRHICommandList& RHICmdList,
 	FShader* Shader,
-	FVertexBufferRHIParamRef ArgumentBuffer,
+	FRHIVertexBuffer* ArgumentBuffer,
 	uint32 ArgumentOffset)
 {
 	RHICmdList.DispatchIndirectComputeShader(ArgumentBuffer, ArgumentOffset);
diff --git a/Engine/Source/Runtime/RenderCore/Private/ShaderParameterStruct.cpp b/Engine/Source/Runtime/RenderCore/Private/ShaderParameterStruct.cpp
--- a/Engine/Source/Runtime/RenderCore/Private/ShaderParameterStruct.cpp
+++ b/Engine/Source/Runtime/RenderCore/Private/ShaderParameterStruct.cpp
@@ -426,7 +426,7 @@ void ValidateShaderParameters(const FShader* Shader, const FShaderParametersMeta
 	// Textures
 	for (const FShaderParameterBindings::FResourceParameter& ParameterBinding : Bindings.Textures)
 	{
-		auto ShaderParameterRef = *reinterpret_cast<const FTextureRHIParamRef*>(Base + ParameterBinding.ByteOffset);
+		FRHITexture* ShaderParameterRef = *(FRHITexture**)(Base + ParameterBinding.ByteOffset);
 		if (!ShaderParameterRef)
 		{
 			EmitNullShaderParameterFatalError(Shader, ParametersMetadata, ParameterBinding.ByteOffset);
diff --git a/Engine/Source/Runtime/RenderCore/Public/GenerateMips.h b/Engine/Source/Runtime/RenderCore/Public/GenerateMips.h
--- a/Engine/Source/Runtime/RenderCore/Public/GenerateMips.h
+++ b/Engine/Source/Runtime/RenderCore/Public/GenerateMips.h
@@ -19,11 +19,11 @@ class RENDERCORE_API FGenerateMips
 public:
 	//Public function for executing the generate mips compute shader 
 	//Default sampler is always bilinear clamp
-	static void Execute(FRHICommandListImmediate& RHICmdList, FTextureRHIParamRef InTexture,
+	static void Execute(FRHICommandListImmediate& RHICmdList, FRHITexture* InTexture,
 		const FGenerateMipsParams& InParams = FGenerateMipsParams());
 
 private:
-	static void Compute(FRHICommandListImmediate& RHIImmCmdList, FTextureRHIParamRef InTexture);
-	static FGenerateMipsStruct* SetupTexture(FTextureRHIParamRef InTexture,
+	static void Compute(FRHICommandListImmediate& RHIImmCmdList, FRHITexture* InTexture);
+	static FGenerateMipsStruct* SetupTexture(FRHITexture* InTexture,
 		const FGenerateMipsParams& InParams = FGenerateMipsParams());
 };
\ No newline at end of file
diff --git a/Engine/Source/Runtime/RenderCore/Public/GlobalShader.h b/Engine/Source/Runtime/RenderCore/Public/GlobalShader.h
--- a/Engine/Source/Runtime/RenderCore/Public/GlobalShader.h
+++ b/Engine/Source/Runtime/RenderCore/Public/GlobalShader.h
@@ -152,7 +152,7 @@ public:
 	RENDERCORE_API FGlobalShader(const ShaderMetaType::CompiledShaderInitializerType& Initializer);
 	
 	template<typename TViewUniformShaderParameters, typename ShaderRHIParamRef, typename TRHICmdList>
-	inline void SetParameters(TRHICmdList& RHICmdList, const ShaderRHIParamRef ShaderRHI, const FUniformBufferRHIParamRef ViewUniformBuffer)
+	inline void SetParameters(TRHICmdList& RHICmdList, const ShaderRHIParamRef ShaderRHI, FRHIUniformBuffer* ViewUniformBuffer)
 	{
 		const auto& ViewUniformBufferParameter = static_cast<const FShaderUniformBufferParameter&>(GetUniformBufferParameter<TViewUniformShaderParameters>());
 		SetUniformBufferParameter(RHICmdList, ShaderRHI, ViewUniformBufferParameter, ViewUniformBuffer);
diff --git a/Engine/Source/Runtime/RenderCore/Public/RenderGraphResources.h b/Engine/Source/Runtime/RenderCore/Public/RenderGraphResources.h
--- a/Engine/Source/Runtime/RenderCore/Public/RenderGraphResources.h
+++ b/Engine/Source/Runtime/RenderCore/Public/RenderGraphResources.h
@@ -204,9 +204,9 @@ public:
 	}
 
 	/** Returns the allocated RHI texture. */
-	FTextureRHIParamRef GetRHI() const
+	FRHITexture* GetRHI() const
 	{
-		return static_cast<FTextureRHIParamRef>(FRDGResource::GetRHI());
+		return static_cast<FRHITexture*>(FRDGResource::GetRHI());
 	}
 
 	//////////////////////////////////////////////////////////////////////////
@@ -221,9 +221,9 @@ private:
 	{}
 
 	/** Returns the allocated RHI texture without access checks. */
-	FTextureRHIParamRef GetRHIUnchecked() const
+	FRHITexture* GetRHIUnchecked() const
 	{
-		return static_cast<FTextureRHIParamRef>(FRDGResource::GetRHIUnchecked());
+		return static_cast<FRHITexture*>(FRDGResource::GetRHIUnchecked());
 	}
 
 	/** This is not a TRefCountPtr<> because FRDGTexture is allocated on the FMemStack
@@ -585,7 +585,7 @@ public:
 	//! The following methods may only be called during pass execution.
 
 	/** Returns the buffer to use for indirect RHI calls. */
-	FVertexBufferRHIParamRef GetIndirectRHICallBuffer() const
+	FRHIVertexBuffer* GetIndirectRHICallBuffer() const
 	{
 		ValidateRHIAccess();
 		check(PooledBuffer);
diff --git a/Engine/Source/Runtime/RenderCore/Public/RenderGraphUtils.h b/Engine/Source/Runtime/RenderCore/Public/RenderGraphUtils.h
--- a/Engine/Source/Runtime/RenderCore/Public/RenderGraphUtils.h
+++ b/Engine/Source/Runtime/RenderCore/Public/RenderGraphUtils.h
@@ -123,7 +123,7 @@ struct RENDERCORE_API FComputeShaderUtils
 		FRHICommandList& RHICmdList,
 		const TShaderClass* ComputeShader,
 		const typename TShaderClass::FParameters& Parameters,
-		FVertexBufferRHIParamRef IndirectArgsBuffer,
+		FRHIVertexBuffer* IndirectArgsBuffer,
 		uint32 IndirectArgOffset)
 	{
 		FRHIComputeShader* ShaderRHI = ComputeShader->GetComputeShader();
diff --git a/Engine/Source/Runtime/RenderCore/Public/RenderTargetPool.h b/Engine/Source/Runtime/RenderCore/Public/RenderTargetPool.h
--- a/Engine/Source/Runtime/RenderCore/Public/RenderTargetPool.h
+++ b/Engine/Source/Runtime/RenderCore/Public/RenderTargetPool.h
@@ -286,7 +286,7 @@ private:
 	/** Elements can be 0, we compact the buffer later. */
 	TArray< TRefCountPtr<FPooledRenderTarget> > PooledRenderTargets;
 	TArray< TRefCountPtr<FPooledRenderTarget> > DeferredDeleteArray;
-	TArray< FTextureRHIParamRef > TransitionTargets;	
+	TArray< FRHITexture* > TransitionTargets;
 
 	/** These are snapshots, have odd life times, live in the scene allocator, and don't contribute to any accounting or other management. */
 	TArray<FPooledRenderTarget*> PooledRenderTargetSnapshots;
diff --git a/Engine/Source/Runtime/RenderCore/Public/RendererInterface.h b/Engine/Source/Runtime/RenderCore/Public/RendererInterface.h
--- a/Engine/Source/Runtime/RenderCore/Public/RendererInterface.h
+++ b/Engine/Source/Runtime/RenderCore/Public/RendererInterface.h
@@ -381,7 +381,7 @@ struct FSceneRenderTargetItem
 	FSceneRenderTargetItem() {}
 
 	/** constructor */
-	FSceneRenderTargetItem(FTextureRHIParamRef InTargetableTexture, FTextureRHIParamRef InShaderResourceTexture, FUnorderedAccessViewRHIRef InUAV)
+	FSceneRenderTargetItem(FRHITexture* InTargetableTexture, FRHITexture* InShaderResourceTexture, FUnorderedAccessViewRHIRef InUAV)
 		:	TargetableTexture(InTargetableTexture)
 		,	ShaderResourceTexture(InShaderResourceTexture)
 		,	UAV(InUAV)
@@ -488,7 +488,7 @@ public:
 	FRHITexture2D* NormalTexture;
 	FRHITexture2D* SmallDepthTexture;
 	FRHICommandListImmediate* RHICmdList;
-	FUniformBufferRHIParamRef ViewUniformBuffer;
+	FRHIUniformBuffer* ViewUniformBuffer;
 	TUniformBufferRef<FSceneTexturesUniformParameters> SceneTexturesUniformParams;
 	const FGlobalDistanceFieldParameterData* GlobalDistanceFieldParams;
 	void* Uid; // A unique identifier for the view.
diff --git a/Engine/Source/Runtime/RenderCore/Public/Shader.h b/Engine/Source/Runtime/RenderCore/Public/Shader.h
--- a/Engine/Source/Runtime/RenderCore/Public/Shader.h
+++ b/Engine/Source/Runtime/RenderCore/Public/Shader.h
@@ -2705,7 +2705,7 @@ extern RENDERCORE_API void DispatchComputeShader(
 extern RENDERCORE_API void DispatchIndirectComputeShader(
 	FRHICommandList& RHICmdList,
 	FShader* Shader,
-	FVertexBufferRHIParamRef ArgumentBuffer,
+	FRHIVertexBuffer* ArgumentBuffer,
 	uint32 ArgumentOffset);
 
 /** Appends to KeyString for all shaders. */
diff --git a/Engine/Source/Runtime/RenderCore/Public/ShaderParameterMacros.h b/Engine/Source/Runtime/RenderCore/Public/ShaderParameterMacros.h
--- a/Engine/Source/Runtime/RenderCore/Public/ShaderParameterMacros.h
+++ b/Engine/Source/Runtime/RenderCore/Public/ShaderParameterMacros.h
@@ -128,7 +128,7 @@ public:
 private:
 
 	/** A private constructor used to coerce an arbitrary RHI uniform buffer reference to a structured reference. */
-	TUniformBufferRef(FUniformBufferRHIParamRef InRHIRef)
+	TUniformBufferRef(FRHIUniformBuffer* InRHIRef)
 	: FUniformBufferRHIRef(InRHIRef)
 	{}
 
@@ -829,10 +829,10 @@ extern RENDERCORE_API FShaderParametersMetadata* FindUniformBufferStructByFName(
  *	SHADER_PARAMETER_TEXTURE_ARRAY(Texture2D, MyArrayOfTextures, [8])
  */
 #define SHADER_PARAMETER_TEXTURE(ShaderType,MemberName) \
-	INTERNAL_SHADER_PARAMETER_EXPLICIT(UBMT_TEXTURE, TShaderResourceParameterTypeInfo<FTextureRHIParamRef>, FTextureRHIParamRef,MemberName,, = nullptr,EShaderPrecisionModifier::Float,TEXT(#ShaderType),false)
+	INTERNAL_SHADER_PARAMETER_EXPLICIT(UBMT_TEXTURE, TShaderResourceParameterTypeInfo<FRHITexture*>, FRHITexture*,MemberName,, = nullptr,EShaderPrecisionModifier::Float,TEXT(#ShaderType),false)
 
 #define SHADER_PARAMETER_TEXTURE_ARRAY(ShaderType,MemberName, ArrayDecl) \
-	INTERNAL_SHADER_PARAMETER_EXPLICIT(UBMT_TEXTURE, TShaderResourceParameterTypeInfo<FTextureRHIParamRef ArrayDecl>, FTextureRHIParamRef,MemberName,ArrayDecl,,EShaderPrecisionModifier::Float,TEXT(#ShaderType),false)
+	INTERNAL_SHADER_PARAMETER_EXPLICIT(UBMT_TEXTURE, TShaderResourceParameterTypeInfo<FRHITexture* ArrayDecl>, FRHITexture*,MemberName,ArrayDecl,,EShaderPrecisionModifier::Float,TEXT(#ShaderType),false)
 
 /** Adds a shader resource view.
  *
diff --git a/Engine/Source/Runtime/RenderCore/Public/ShaderParameterStruct.h b/Engine/Source/Runtime/RenderCore/Public/ShaderParameterStruct.h
--- a/Engine/Source/Runtime/RenderCore/Public/ShaderParameterStruct.h
+++ b/Engine/Source/Runtime/RenderCore/Public/ShaderParameterStruct.h
@@ -179,7 +179,7 @@ inline void SetShaderParameters(TRHICmdList& RHICmdList, const TShaderClass* Sha
 	// Textures
 	for (const FShaderParameterBindings::FResourceParameter& ParameterBinding : Bindings.Textures)
 	{
-		auto ShaderParameterRef = *reinterpret_cast<const FTextureRHIParamRef*>(Base + ParameterBinding.ByteOffset);
+		auto ShaderParameterRef = *(FRHITexture**)(Base + ParameterBinding.ByteOffset);
 
 		if (DO_CHECK && !ShaderParameterRef)
 		{
@@ -291,7 +291,7 @@ void SetShaderParameters(FRayTracingShaderBindingsWriter& RTBindingsWriter, cons
 	// Textures
 	for (const FShaderParameterBindings::FResourceParameter& ParameterBinding : Bindings.Textures)
 	{
-		auto ShaderParameterRef = *reinterpret_cast<const FTextureRHIParamRef*>(Base + ParameterBinding.ByteOffset);
+		auto ShaderParameterRef = *(FRHITexture**)(Base + ParameterBinding.ByteOffset);
 
 		if (DO_CHECK && !ShaderParameterRef)
 		{
@@ -422,7 +422,7 @@ void SetShaderParameters(FRayTracingShaderBindingsWriter& RTBindingsWriter, cons
 		// Do not do any validation at some resources may have been removed from the structure because known to not be used by the shader.
 		EUniformBufferValidation Validation = EUniformBufferValidation::None;
 
-		FUniformBufferRHIParamRef RootUniformBuffer = CreateUniformBufferImmediate(Parameters, UniformBuffer_SingleDraw, Validation);
+		FRHIUniformBuffer* RootUniformBuffer = CreateUniformBufferImmediate(Parameters, UniformBuffer_SingleDraw, Validation);
 		RTBindingsWriter.SetUniformBuffer(Bindings.RootParameterBufferIndex, RootUniformBuffer);
 	}
 }
diff --git a/Engine/Source/Runtime/RenderCore/Public/ShaderParameterUtils.h b/Engine/Source/Runtime/RenderCore/Public/ShaderParameterUtils.h
--- a/Engine/Source/Runtime/RenderCore/Public/ShaderParameterUtils.h
+++ b/Engine/Source/Runtime/RenderCore/Public/ShaderParameterUtils.h
@@ -235,7 +235,7 @@ FORCEINLINE void SetTextureParameter(
 	const FShaderResourceParameter& TextureParameter,
 	const FShaderResourceParameter& SamplerParameter,
 	FRHISamplerState* SamplerStateRHI,
-	FTextureRHIParamRef TextureRHI,
+	FRHITexture* TextureRHI,
 	uint32 ElementIndex = 0
 	)
 {
@@ -273,7 +273,7 @@ FORCEINLINE void SetTextureParameter(
 	TRHICmdList& RHICmdList,
 	TRHIShader* Shader,
 	const FShaderResourceParameter& Parameter,
-	FTextureRHIParamRef NewTextureRHI
+	FRHITexture* NewTextureRHI
 	)
 {
 	if(Parameter.IsBound())
@@ -425,7 +425,7 @@ inline void FRWShaderParameter::SetBuffer(TRHICmdList& RHICmdList, TShaderRHIRef
 }
 
 template<typename TShaderRHIRef, typename TRHICmdList>
-inline void FRWShaderParameter::SetTexture(TRHICmdList& RHICmdList, TShaderRHIRef Shader, const FTextureRHIParamRef Texture, FRHIUnorderedAccessView* UAV) const
+inline void FRWShaderParameter::SetTexture(TRHICmdList& RHICmdList, TShaderRHIRef Shader, FRHITexture* Texture, FRHIUnorderedAccessView* UAV) const
 {
 	if (!SetUAVParameterIfCS(RHICmdList, Shader, UAVParameter, UAV))
 	{
@@ -463,7 +463,7 @@ inline void SetUniformBufferParameter(
 	TRHICmdList& RHICmdList,
 	TShaderRHIRef Shader,
 	const FShaderUniformBufferParameter& Parameter,
-	FUniformBufferRHIParamRef UniformBufferRHI
+	FRHIUniformBuffer* UniformBufferRHI
 	)
 {
 	// This will trigger if the parameter was not serialized
diff --git a/Engine/Source/Runtime/RenderCore/Public/ShaderParameters.h b/Engine/Source/Runtime/RenderCore/Public/ShaderParameters.h
--- a/Engine/Source/Runtime/RenderCore/Public/ShaderParameters.h
+++ b/Engine/Source/Runtime/RenderCore/Public/ShaderParameters.h
@@ -145,7 +145,7 @@ public:
 	inline void SetBuffer(TRHICmdList& RHICmdList, TShaderRHIRef Shader, const FRWBufferStructured& RWBuffer) const;
 
 	template<typename TShaderRHIRef, typename TRHICmdList>
-	inline void SetTexture(TRHICmdList& RHICmdList, TShaderRHIRef Shader, const FTextureRHIParamRef Texture, FRHIUnorderedAccessView* UAV) const;
+	inline void SetTexture(TRHICmdList& RHICmdList, TShaderRHIRef Shader, FRHITexture* Texture, FRHIUnorderedAccessView* UAV) const;
 
 	template<typename TRHICmdList>
 	inline void UnsetUAV(TRHICmdList& RHICmdList, FRHIComputeShader* ComputeShader) const;
@@ -240,7 +240,7 @@ public:
 #if RHI_RAYTRACING
 struct FRayTracingShaderBindingsWriter : FRayTracingShaderBindings
 {
-	void Set(const FShaderResourceParameter& Param, FTextureRHIParamRef Value)
+	void Set(const FShaderResourceParameter& Param, FRHITexture* Value)
 	{
 		if (Param.IsBound())
 		{
@@ -258,7 +258,7 @@ struct FRayTracingShaderBindingsWriter : FRayTracingShaderBindings
 		}
 	}
 
-	void Set(const FShaderUniformBufferParameter& Param, FUniformBufferRHIParamRef Value)
+	void Set(const FShaderUniformBufferParameter& Param, FRHIUniformBuffer* Value)
 	{
 		if (Param.IsBound())
 		{
@@ -284,7 +284,7 @@ struct FRayTracingShaderBindingsWriter : FRayTracingShaderBindings
 		}
 	}
 
-	void SetTexture(uint16 BaseIndex, FTextureRHIParamRef Value)
+	void SetTexture(uint16 BaseIndex, FRHITexture* Value)
 	{
 		checkSlow(BaseIndex < ARRAY_COUNT(Textures));
 		Textures[BaseIndex] = Value;
@@ -308,7 +308,7 @@ struct FRayTracingShaderBindingsWriter : FRayTracingShaderBindings
 		UAVs[BaseIndex] = Value;
 	}
 
-	void SetUniformBuffer(uint16 BaseIndex, FUniformBufferRHIParamRef Value)
+	void SetUniformBuffer(uint16 BaseIndex, FRHIUniformBuffer* Value)
 	{
 		checkSlow(BaseIndex < ARRAY_COUNT(UniformBuffers));
 		UniformBuffers[BaseIndex] = Value;
diff --git a/Engine/Source/Runtime/RenderCore/Public/UniformBuffer.h b/Engine/Source/Runtime/RenderCore/Public/UniformBuffer.h
--- a/Engine/Source/Runtime/RenderCore/Public/UniformBuffer.h
+++ b/Engine/Source/Runtime/RenderCore/Public/UniformBuffer.h
@@ -81,7 +81,7 @@ public:
 	}
 
 	// Accessors.
-	FUniformBufferRHIParamRef GetUniformBufferRHI() const 
+	FRHIUniformBuffer* GetUniformBufferRHI() const
 	{ 
 		checkSlow(IsInRenderingThread() || IsInParallelRenderingThread());
 		checkf(UniformBufferRHI.GetReference(), TEXT("Attempted to access UniformBufferRHI on a TUniformBuffer that was never filled in with anything")); 
diff --git a/Engine/Source/Runtime/RenderCore/Public/VertexFactory.h b/Engine/Source/Runtime/RenderCore/Public/VertexFactory.h
--- a/Engine/Source/Runtime/RenderCore/Public/VertexFactory.h
+++ b/Engine/Source/Runtime/RenderCore/Public/VertexFactory.h
@@ -26,7 +26,7 @@ struct FVertexInputStream
 {
 	uint32 StreamIndex : 4;
 	uint32 Offset : 28;
-	FVertexBufferRHIParamRef VertexBuffer;
+	FRHIVertexBuffer* VertexBuffer;
 
 	FVertexInputStream() :
 		StreamIndex(0),
@@ -34,7 +34,7 @@ struct FVertexInputStream
 		VertexBuffer(nullptr)
 	{}
 
-	FVertexInputStream(uint32 InStreamIndex, uint32 InOffset, FVertexBufferRHIParamRef InVertexBuffer)
+	FVertexInputStream(uint32 InStreamIndex, uint32 InOffset, FRHIVertexBuffer* InVertexBuffer)
 		: StreamIndex(InStreamIndex), Offset(InOffset), VertexBuffer(InVertexBuffer)
 	{
 		// Verify no overflow
diff --git a/Engine/Source/Runtime/Renderer/Private/AtmosphereRendering.cpp b/Engine/Source/Runtime/Renderer/Private/AtmosphereRendering.cpp
--- a/Engine/Source/Runtime/Renderer/Private/AtmosphereRendering.cpp
+++ b/Engine/Source/Runtime/Renderer/Private/AtmosphereRendering.cpp
@@ -1257,7 +1257,7 @@ void FAtmosphericFogSceneInfo::RenderAtmosphereShaders(FRHICommandList& RHICmdLi
 		{
 			int32 Layer = Atmosphere3DTextureIndex;
 			{
-				FTextureRHIParamRef RenderTargets[] =
+				FRHITexture* RenderTargets[] =
 				{
 					AtmosphereTextures->AtmosphereDeltaSR->GetRenderTargetItem().TargetableTexture,
 					AtmosphereTextures->AtmosphereDeltaSM->GetRenderTargetItem().TargetableTexture,
diff --git a/Engine/Source/Runtime/Renderer/Private/BasePassRendering.cpp b/Engine/Source/Runtime/Renderer/Private/BasePassRendering.cpp
--- a/Engine/Source/Runtime/Renderer/Private/BasePassRendering.cpp
+++ b/Engine/Source/Runtime/Renderer/Private/BasePassRendering.cpp
@@ -551,7 +551,7 @@ void CreateOpaqueBasePassUniformBuffer(
 
 		BasePassParameters.IndirectOcclusionTexture = IndirectOcclusion->GetRenderTargetItem().ShaderResourceTexture;
 
-		FTextureRHIParamRef ResolvedSceneDepthTextureValue = GSystemTextures.WhiteDummy->GetRenderTargetItem().ShaderResourceTexture;
+		FRHITexture* ResolvedSceneDepthTextureValue = GSystemTextures.WhiteDummy->GetRenderTargetItem().ShaderResourceTexture;
 
 		if (SceneRenderTargets.GetMSAACount() > 1)
 		{
diff --git a/Engine/Source/Runtime/Renderer/Private/BasePassRendering.inl b/Engine/Source/Runtime/Renderer/Private/BasePassRendering.inl
--- a/Engine/Source/Runtime/Renderer/Private/BasePassRendering.inl
+++ b/Engine/Source/Runtime/Renderer/Private/BasePassRendering.inl
@@ -34,7 +34,7 @@ void TBasePassVertexShaderPolicyParamType<LightMapPolicyType>::GetShaderBindings
 
 	if (Scene)
 	{
-		FUniformBufferRHIParamRef ReflectionCaptureUniformBuffer = Scene->UniformBuffers.ReflectionCaptureUniformBuffer.GetReference();
+		FRHIUniformBuffer* ReflectionCaptureUniformBuffer = Scene->UniformBuffers.ReflectionCaptureUniformBuffer.GetReference();
 		ShaderBindings.Add(ReflectionCaptureBuffer, ReflectionCaptureUniformBuffer);
 	}
 	else
@@ -81,7 +81,7 @@ void TBasePassPixelShaderPolicyParamType<LightMapPolicyType>::GetShaderBindings(
 
 	if (Scene)
 	{
-		FUniformBufferRHIParamRef ReflectionCaptureUniformBuffer = Scene->UniformBuffers.ReflectionCaptureUniformBuffer.GetReference();
+		FRHIUniformBuffer* ReflectionCaptureUniformBuffer = Scene->UniformBuffers.ReflectionCaptureUniformBuffer.GetReference();
 		ShaderBindings.Add(ReflectionCaptureBuffer, ReflectionCaptureUniformBuffer);
 	}
 	else
diff --git a/Engine/Source/Runtime/Renderer/Private/CapsuleShadowRendering.cpp b/Engine/Source/Runtime/Renderer/Private/CapsuleShadowRendering.cpp
--- a/Engine/Source/Runtime/Renderer/Private/CapsuleShadowRendering.cpp
+++ b/Engine/Source/Runtime/Renderer/Private/CapsuleShadowRendering.cpp
@@ -327,7 +327,7 @@ public:
 		int32 NumMeshDistanceFieldCastersValue,
 		FRHIShaderResourceView* MeshDistanceFieldCasterIndicesSRV,
 		FRHIShaderResourceView* LightDirectionDataSRV,
-		FTextureRHIParamRef ReceiverBentNormalTextureValue)
+		FRHITexture* ReceiverBentNormalTextureValue)
 	{
 		FRHIComputeShader* ShaderRHI = GetComputeShader();
 
@@ -1293,8 +1293,8 @@ void FDeferredShadingSceneRenderer::SetupIndirectCapsuleShadows(
 
 void FDeferredShadingSceneRenderer::RenderIndirectCapsuleShadows(
 	FRHICommandListImmediate& RHICmdList, 
-	FTextureRHIParamRef IndirectLightingTexture, 
-	FTextureRHIParamRef ExistingIndirectOcclusionTexture) const
+	FRHITexture* IndirectLightingTexture,
+	FRHITexture* ExistingIndirectOcclusionTexture) const
 {
 	check(RHICmdList.IsOutsideRenderPass());
 
@@ -1331,7 +1331,7 @@ void FDeferredShadingSceneRenderer::RenderIndirectCapsuleShadows(
 			}
 
 			FSceneRenderTargets& SceneContext = FSceneRenderTargets::Get(RHICmdList);
-			TArray<FTextureRHIParamRef, TInlineAllocator<2>> RenderTargets;
+			TArray<FRHITexture*, TInlineAllocator<2>> RenderTargets;
 
 			if (IndirectLightingTexture)
 			{
@@ -1358,7 +1358,7 @@ void FDeferredShadingSceneRenderer::RenderIndirectCapsuleShadows(
 
 				// #todo is this transition really needed?
 				{
-					FTextureRHIParamRef Transitions[MaxSimultaneousRenderTargets + 1];
+					FRHITexture* Transitions[MaxSimultaneousRenderTargets + 1];
 					int32 TransitionIndex = 0;
 					for (int32 Index = 0; Index < RenderTargets.Num(); Index++)
 					{
diff --git a/Engine/Source/Runtime/Renderer/Private/CompositionLighting/PostProcessDeferredDecals.cpp b/Engine/Source/Runtime/Renderer/Private/CompositionLighting/PostProcessDeferredDecals.cpp
--- a/Engine/Source/Runtime/Renderer/Private/CompositionLighting/PostProcessDeferredDecals.cpp
+++ b/Engine/Source/Runtime/Renderer/Private/CompositionLighting/PostProcessDeferredDecals.cpp
@@ -612,7 +612,7 @@ void FRCPassPostProcessDeferredDecals::Process(FRenderingCompositePassContext& C
 			{
 				SCOPED_DRAW_EVENT(RHICmdList, DBufferClear);
 
-				FTextureRHIParamRef RenderTargets[4];
+				FRHITexture* RenderTargets[4];
 
 				RenderTargets[0] = SceneContext.DBufferA->GetRenderTargetItem().TargetableTexture;
 				RenderTargets[1] = SceneContext.DBufferB->GetRenderTargetItem().TargetableTexture;
@@ -846,7 +846,7 @@ void FRCPassPostProcessDeferredDecals::Process(FRenderingCompositePassContext& C
 			if (CurrentStage == DRS_BeforeBasePass)
 			{
 				// combine DBuffer RTWriteMasks; will end up in one texture we can load from in the base pass PS and decide whether to do the actual work or not
-				FTextureRHIParamRef Textures[3] =
+				FRHITexture* Textures[3] =
 				{
 					SceneContext.DBufferA->GetRenderTargetItem().TargetableTexture,
 					SceneContext.DBufferB->GetRenderTargetItem().TargetableTexture,
@@ -1011,8 +1011,8 @@ void FDecalRenderTargetManager::SetRenderTargetMode(FDecalRenderingCommon::ERend
 	FExclusiveDepthStencil DepthStencilAccess = FExclusiveDepthStencil::DepthRead_StencilWrite;
 	ERenderTargetActions DepthTargetActions = ERenderTargetActions::Load_DontStore;
 	uint32 NumColorTargets = 1;
-	FTextureRHIParamRef* TargetsToBind = &TargetsToResolve[0];
-	FTextureRHIParamRef DepthTarget = SceneContext.GetSceneDepthSurface();
+	FRHITexture** TargetsToBind = &TargetsToResolve[0];
+	FRHITexture* DepthTarget = SceneContext.GetSceneDepthSurface();
 
 	// The SceneColorAndGBuffer modes do not actually need GBufferA bound when there's no normal.
 	// The apis based on renderpasses fill fail to actually bind anything past a null entry in their RT list so we have to bind it anyway.
@@ -1117,7 +1117,7 @@ void FDecalRenderTargetManager::SetRenderTargetMode(FDecalRenderingCommon::ERend
 
 
 
-void FDecalRenderTargetManager::FlushMetaData(FTextureRHIParamRef* Textures, uint32 NumTextures)
+void FDecalRenderTargetManager::FlushMetaData(FRHITexture** Textures, uint32 NumTextures)
 {
 	RHICmdList.TransitionResources(EResourceTransitionAccess::EMetaData, Textures, NumTextures);
 }
diff --git a/Engine/Source/Runtime/Renderer/Private/CompositionLighting/PostProcessDeferredDecals.h b/Engine/Source/Runtime/Renderer/Private/CompositionLighting/PostProcessDeferredDecals.h
--- a/Engine/Source/Runtime/Renderer/Private/CompositionLighting/PostProcessDeferredDecals.h
+++ b/Engine/Source/Runtime/Renderer/Private/CompositionLighting/PostProcessDeferredDecals.h
@@ -60,7 +60,7 @@ struct FDecalRenderTargetManager
 	//
 	bool TargetsToTransitionWritable[ResolveBufferMax];
 	//
-	FTextureRHIParamRef TargetsToResolve[ResolveBufferMax];
+	FRHITexture* TargetsToResolve[ResolveBufferMax];
 	//
 	bool bGufferADirty;
 	bool bGufferBCDirty;
@@ -77,7 +77,7 @@ struct FDecalRenderTargetManager
 
 	void ResolveTargets();
 
-	void FlushMetaData(FTextureRHIParamRef* Textures, uint32 NumTextures);
+	void FlushMetaData(FRHITexture** Textures, uint32 NumTextures);
 
 	void SetRenderTargetMode(FDecalRenderingCommon::ERenderTargetMode CurrentRenderTargetMode, bool bHasNormal, bool bPerPixelDBufferMask);
 };
diff --git a/Engine/Source/Runtime/Renderer/Private/CompositionLighting/PostProcessLpvIndirect.cpp b/Engine/Source/Runtime/Renderer/Private/CompositionLighting/PostProcessLpvIndirect.cpp
--- a/Engine/Source/Runtime/Renderer/Private/CompositionLighting/PostProcessLpvIndirect.cpp
+++ b/Engine/Source/Runtime/Renderer/Private/CompositionLighting/PostProcessLpvIndirect.cpp
@@ -78,8 +78,8 @@ public:
 	template <typename TRHICmdList>
 	void SetParameters(	
 		TRHICmdList& RHICmdList,
-		FTextureRHIParamRef* LpvBufferSRVsIn, 
-		FTextureRHIParamRef AOVolumeTextureSRVIn, 
+		FRHITexture** LpvBufferSRVsIn,
+		FRHITexture* AOVolumeTextureSRVIn,
 		FLpvReadUniformBufferRef LpvUniformBuffer, 
 		const FRenderingCompositePassContext& Context )
 	{
@@ -177,7 +177,7 @@ public:
 	}
 
 	void SetParameters(	
-		FTextureRHIParamRef AOVolumeTextureSRVIn, 
+		FRHITexture* AOVolumeTextureSRVIn,
 		FLpvReadUniformBufferRef LpvUniformBuffer, 
 		const FRenderingCompositePassContext& Context )
 	{
@@ -269,7 +269,7 @@ void FRCPassPostProcessLpvIndirect::Process(FRenderingCompositePassContext& Cont
 		DoDirectionalOcclusionPass(Context);
 	}
 
-	FTextureRHIParamRef RenderTargets[2];
+	FRHITexture* RenderTargets[2];
 	RenderTargets[0] = DestColorRenderTarget.TargetableTexture;
 	RenderTargets[1] = DestSpecularRenderTarget.TargetableTexture;
 
@@ -321,7 +321,7 @@ void FRCPassPostProcessLpvIndirect::Process(FRenderingCompositePassContext& Cont
 		LpvReadUniformBufferParams = Lpv->GetReadUniformBufferParams();
 		LpvReadUniformBuffer = FLpvReadUniformBufferRef::CreateUniformBufferImmediate(LpvReadUniformBufferParams, UniformBuffer_SingleDraw);
 
-		FTextureRHIParamRef LpvBufferSrvs[7];
+		FRHITexture* LpvBufferSrvs[7];
 		for (int i = 0; i < 7; i++)
 		{
 			LpvBufferSrvs[i] = Lpv->GetLpvBufferSrv(i);
@@ -383,7 +383,7 @@ void FRCPassPostProcessLpvIndirect::DoDirectionalOcclusionPass(FRenderingComposi
 		return;
 	}
 
-	FTextureRHIParamRef RenderTarget = DestDirectionalOcclusionRenderTarget.TargetableTexture;
+	FRHITexture* RenderTarget = DestDirectionalOcclusionRenderTarget.TargetableTexture;
 
 	FRHIRenderPassInfo RPInfo(RenderTarget, ERenderTargetActions::Clear_Store);
 	Context.RHICmdList.BeginRenderPass(RPInfo, TEXT("DoDirectionalOcclusionPass"));
diff --git a/Engine/Source/Runtime/Renderer/Private/ConvertToUniformMesh.cpp b/Engine/Source/Runtime/Renderer/Private/ConvertToUniformMesh.cpp
--- a/Engine/Source/Runtime/Renderer/Private/ConvertToUniformMesh.cpp
+++ b/Engine/Source/Runtime/Renderer/Private/ConvertToUniformMesh.cpp
@@ -233,7 +233,7 @@ int32 FUniformMeshConverter::Convert(
 	int32 LODIndex,
 	FUniformMeshBuffers*& OutUniformMeshBuffers,
 	const FMaterialRenderProxy*& OutMaterialRenderProxy,
-	FUniformBufferRHIParamRef& OutPrimitiveUniformBuffer)
+	FRHIUniformBuffer*& OutPrimitiveUniformBuffer)
 {
 	const FPrimitiveSceneProxy* PrimitiveSceneProxy = PrimitiveSceneInfo->Proxy;
 	const auto FeatureLevel = View.GetFeatureLevel();
@@ -263,7 +263,7 @@ int32 FUniformMeshConverter::Convert(
 		UnbindRenderTargets(RHICmdList);
 
 		uint32 Offsets[1] = {0};
-		const FVertexBufferRHIParamRef StreamOutTargets[1] = {GUniformMeshTemporaryBuffers.TriangleData.GetReference()};
+		FRHIVertexBuffer* const StreamOutTargets[1] = {GUniformMeshTemporaryBuffers.TriangleData.GetReference()};
 		RHICmdList.SetStreamOutTargets(1, StreamOutTargets, Offsets);
 
 		for (int32 MeshIndex = 0; MeshIndex < MeshElements.Num(); MeshIndex++)
@@ -290,7 +290,7 @@ int32 FUniformMeshConverter::Convert(
 			}
 		}
 
-		RHICmdList.SetStreamOutTargets(1, (const FVertexBufferRHIParamRef*)NULL, Offsets);
+		RHICmdList.SetStreamOutTargets(1, nullptr, Offsets);
 	}
 
 	OutUniformMeshBuffers = &GUniformMeshTemporaryBuffers;
@@ -336,7 +336,7 @@ public:
 		int32 SurfelStartIndexValue,
 		int32 NumSurfelsToGenerateValue,
 		const FMaterialRenderProxy* MaterialProxy,
-		FUniformBufferRHIParamRef PrimitiveUniformBuffer,
+		FRHIUniformBuffer* PrimitiveUniformBuffer,
 		const FMatrix& Instance0Transform
 		)
 	{
@@ -395,7 +395,7 @@ void FUniformMeshConverter::GenerateSurfels(
 	FViewInfo& View, 
 	const FPrimitiveSceneInfo* PrimitiveSceneInfo, 
 	const FMaterialRenderProxy* MaterialProxy,
-	FUniformBufferRHIParamRef PrimitiveUniformBuffer, 
+	FRHIUniformBuffer* PrimitiveUniformBuffer,
 	const FMatrix& Instance0Transform,
 	int32 SurfelOffset,
 	int32 NumSurfels)
diff --git a/Engine/Source/Runtime/Renderer/Private/DebugViewModeRendering.cpp b/Engine/Source/Runtime/Renderer/Private/DebugViewModeRendering.cpp
--- a/Engine/Source/Runtime/Renderer/Private/DebugViewModeRendering.cpp
+++ b/Engine/Source/Runtime/Renderer/Private/DebugViewModeRendering.cpp
@@ -286,7 +286,7 @@ FDebugViewModeMeshProcessor::FDebugViewModeMeshProcessor(
 	const FScene* InScene, 
 	ERHIFeatureLevel::Type InFeatureLevel,
 	const FSceneView* InViewIfDynamicMeshCommand, 
-	FUniformBufferRHIParamRef InPassUniformBuffer, 
+	FRHIUniformBuffer* InPassUniformBuffer,
 	bool bTranslucentBasePass,
 	FMeshPassDrawListContext* InDrawListContext
 )
diff --git a/Engine/Source/Runtime/Renderer/Private/DebugViewModeRendering.h b/Engine/Source/Runtime/Renderer/Private/DebugViewModeRendering.h
--- a/Engine/Source/Runtime/Renderer/Private/DebugViewModeRendering.h
+++ b/Engine/Source/Runtime/Renderer/Private/DebugViewModeRendering.h
@@ -236,7 +236,7 @@ public:
 class FDebugViewModeMeshProcessor : public FMeshPassProcessor
 {
 public:
-	FDebugViewModeMeshProcessor(const FScene* InScene, ERHIFeatureLevel::Type InFeatureLevel, const FSceneView* InViewIfDynamicMeshCommand, FUniformBufferRHIParamRef InPassUniformBuffer, bool bTranslucentBasePass, FMeshPassDrawListContext* InDrawListContext);
+	FDebugViewModeMeshProcessor(const FScene* InScene, ERHIFeatureLevel::Type InFeatureLevel, const FSceneView* InViewIfDynamicMeshCommand, FRHIUniformBuffer* InPassUniformBuffer, bool bTranslucentBasePass, FMeshPassDrawListContext* InDrawListContext);
 	virtual void AddMeshBatch(const FMeshBatch& RESTRICT MeshBatch, uint64 BatchElementMask, const FPrimitiveSceneProxy* RESTRICT PrimitiveSceneProxy, int32 StaticMeshId = -1) override final;
 
 private:
diff --git a/Engine/Source/Runtime/Renderer/Private/DecalRenderingShared.cpp b/Engine/Source/Runtime/Renderer/Private/DecalRenderingShared.cpp
--- a/Engine/Source/Runtime/Renderer/Private/DecalRenderingShared.cpp
+++ b/Engine/Source/Runtime/Renderer/Private/DecalRenderingShared.cpp
@@ -72,7 +72,7 @@ public:
 		FrustumComponentToClip.Bind(Initializer.ParameterMap, TEXT("FrustumComponentToClip"));
 	}
 
-	void SetParameters(FRHICommandList& RHICmdList, const FUniformBufferRHIParamRef ViewUniformBuffer, const FMatrix& InFrustumComponentToClip)
+	void SetParameters(FRHICommandList& RHICmdList, FRHIUniformBuffer* ViewUniformBuffer, const FMatrix& InFrustumComponentToClip)
 	{
 		FRHIVertexShader* ShaderRHI = GetVertexShader();
 
diff --git a/Engine/Source/Runtime/Renderer/Private/DeferredShadingRenderer.h b/Engine/Source/Runtime/Renderer/Private/DeferredShadingRenderer.h
--- a/Engine/Source/Runtime/Renderer/Private/DeferredShadingRenderer.h
+++ b/Engine/Source/Runtime/Renderer/Private/DeferredShadingRenderer.h
@@ -331,8 +331,8 @@ private:
 	/** Renders indirect shadows from capsules modulated onto scene color. */
 	void RenderIndirectCapsuleShadows(
 		FRHICommandListImmediate& RHICmdList, 
-		FTextureRHIParamRef IndirectLightingTexture, 
-		FTextureRHIParamRef ExistingIndirectOcclusionTexture) const;
+		FRHITexture* IndirectLightingTexture,
+		FRHITexture* ExistingIndirectOcclusionTexture) const;
 
 	/** Renders capsule shadows for movable skylights, using the cone of visibility (bent normal) from DFAO. */
 	void RenderCapsuleShadowsForMovableSkylight(FRHICommandListImmediate& RHICmdList, TRefCountPtr<IPooledRenderTarget>& BentNormalOutput) const;
@@ -550,13 +550,13 @@ private:
 
 	void VisualizeVarianceMipTree(FRHICommandListImmediate& RHICmdList, const FViewInfo& View, const FRWBuffer& VarianceMipTree, FIntVector VarianceMipTreeDimensions);
 
-	void ComputePathCompaction(FRHICommandListImmediate& RHICmdList, const FViewInfo& View, FTextureRHIParamRef RadianceTexture, FTextureRHIParamRef SampleCountTexture, FTextureRHIParamRef PixelPositionTexture,
+	void ComputePathCompaction(FRHICommandListImmediate& RHICmdList, const FViewInfo& View, FRHITexture* RadianceTexture, FRHITexture* SampleCountTexture, FRHITexture* PixelPositionTexture,
 		FRHIUnorderedAccessView* RadianceSortedRedUAV, FRHIUnorderedAccessView* RadianceSortedGreenUAV, FRHIUnorderedAccessView* RadianceSortedBlueUAV, FRHIUnorderedAccessView* RadianceSortedAlphaUAV, FRHIUnorderedAccessView* SampleCountSortedUAV);
 
 	void BuildSkyLightCdf(FRHICommandListImmediate& RHICmdList, const FViewInfo& View, const FTexture& SkyLightTextureCube, FRWBuffer& RowCdf, FRWBuffer& ColumnCdf, FRWBuffer& CubeFaceCdf);
 	void VisualizeSkyLightCdf(FRHICommandListImmediate& RHICmdList, const FViewInfo& View, FIntVector Dimensions, const FRWBuffer& RowCdf, const FRWBuffer& ColumnCdf, const FRWBuffer& CubeFaceCdf);
 
-	void ComputeRayCount(FRHICommandListImmediate& RHICmdList, const FViewInfo& View, FTextureRHIParamRef RayCountPerPixelTexture);
+	void ComputeRayCount(FRHICommandListImmediate& RHICmdList, const FViewInfo& View, FRHITexture* RayCountPerPixelTexture);
 
 	/** Debug ray tracing functions. */
 	void RenderRayTracingDebug(FRHICommandListImmediate& RHICmdList, const FViewInfo& View);
diff --git a/Engine/Source/Runtime/Renderer/Private/DistanceFieldAmbientOcclusion.h b/Engine/Source/Runtime/Renderer/Private/DistanceFieldAmbientOcclusion.h
--- a/Engine/Source/Runtime/Renderer/Private/DistanceFieldAmbientOcclusion.h
+++ b/Engine/Source/Runtime/Renderer/Private/DistanceFieldAmbientOcclusion.h
@@ -350,7 +350,7 @@ public:
 
 	void Set(FRHICommandList& RHICmdList, FRHIPixelShader* ShaderRHI, const FViewInfo& View, const TRefCountPtr<IPooledRenderTarget>& DistanceFieldAOBentNormal)
 	{
-		FTextureRHIParamRef BentNormalAO = DistanceFieldAOBentNormal ? DistanceFieldAOBentNormal->GetRenderTargetItem().ShaderResourceTexture : GWhiteTexture->TextureRHI;
+		FRHITexture* BentNormalAO = DistanceFieldAOBentNormal ? DistanceFieldAOBentNormal->GetRenderTargetItem().ShaderResourceTexture : GWhiteTexture->TextureRHI;
 		SetTextureParameter(RHICmdList, ShaderRHI, BentNormalAOTexture, BentNormalAOSampler, TStaticSamplerState<SF_Bilinear>::GetRHI(), BentNormalAO);
 
 		FIntPoint const AOBufferSize = GetBufferSizeForAO();
diff --git a/Engine/Source/Runtime/Renderer/Private/DistanceFieldLightingShared.h b/Engine/Source/Runtime/Renderer/Private/DistanceFieldLightingShared.h
--- a/Engine/Source/Runtime/Renderer/Private/DistanceFieldLightingShared.h
+++ b/Engine/Source/Runtime/Renderer/Private/DistanceFieldLightingShared.h
@@ -686,14 +686,14 @@ public:
 		int32 LODIndex,
 		FUniformMeshBuffers*& OutUniformMeshBuffers,
 		const FMaterialRenderProxy*& OutMaterialRenderProxy,
-		FUniformBufferRHIParamRef& OutPrimitiveUniformBuffer);
+		FRHIUniformBuffer*& OutPrimitiveUniformBuffer);
 
 	static void GenerateSurfels(
 		FRHICommandListImmediate& RHICmdList, 
 		FViewInfo& View, 
 		const FPrimitiveSceneInfo* PrimitiveSceneInfo, 
 		const FMaterialRenderProxy* MaterialProxy,
-		FUniformBufferRHIParamRef PrimitiveUniformBuffer,
+		FRHIUniformBuffer* PrimitiveUniformBuffer,
 		const FMatrix& Instance0Transform,
 		int32 SurfelOffset,
 		int32 NumSurfels);
diff --git a/Engine/Source/Runtime/Renderer/Private/DistortionRendering.cpp b/Engine/Source/Runtime/Renderer/Private/DistortionRendering.cpp
--- a/Engine/Source/Runtime/Renderer/Private/DistortionRendering.cpp
+++ b/Engine/Source/Runtime/Renderer/Private/DistortionRendering.cpp
@@ -124,8 +124,8 @@ public:
 
 		FGlobalShader::SetParameters<FViewUniformShaderParameters>(Context.RHICmdList, ShaderRHI, View.ViewUniformBuffer);
 
-		FTextureRHIParamRef DistortionTextureValue = DistortionRT.GetRenderTargetItem().TargetableTexture;
-		FTextureRHIParamRef SceneColorTextureValue = SceneContext.GetSceneColor()->GetRenderTargetItem().TargetableTexture;
+		FRHITexture* DistortionTextureValue = DistortionRT.GetRenderTargetItem().TargetableTexture;
+		FRHITexture* SceneColorTextureValue = SceneContext.GetSceneColor()->GetRenderTargetItem().TargetableTexture;
 
 		// Here we use SF_Point as in fullscreen the pixels are 1:1 mapped.
 		SetTextureParameter(
@@ -214,7 +214,7 @@ public:
 	}
 	TDistortionMergePS() {}
 
-	void SetParameters(const FRenderingCompositePassContext& Context, const FViewInfo& View, const FTextureRHIParamRef& PassTexture)
+	void SetParameters(const FRenderingCompositePassContext& Context, const FViewInfo& View, FRHITexture* PassTexture)
 	{
 		FSceneRenderTargets& SceneContext = FSceneRenderTargets::Get(Context.RHICmdList);
 		FRHIPixelShader* ShaderRHI = GetPixelShader();
@@ -467,7 +467,7 @@ static void DrawDistortionApplyScreenPass(FRHICommandListImmediate& RHICmdList,
 }
 
 template <bool UseMSAA>
-static void DrawDistortionMergePass(FRHICommandListImmediate& RHICmdList, FSceneRenderTargets& SceneContext, FViewInfo& View, const FTextureRHIParamRef& PassTexture) {
+static void DrawDistortionMergePass(FRHICommandListImmediate& RHICmdList, FSceneRenderTargets& SceneContext, FViewInfo& View, FRHITexture* PassTexture) {
 	TShaderMapRef<FPostProcessVS> VertexShader(View.ShaderMap);
 	TShaderMapRef<TDistortionMergePS<UseMSAA>> PixelShader(View.ShaderMap);
 
diff --git a/Engine/Source/Runtime/Renderer/Private/FogRendering.cpp b/Engine/Source/Runtime/Renderer/Private/FogRendering.cpp
--- a/Engine/Source/Runtime/Renderer/Private/FogRendering.cpp
+++ b/Engine/Source/Runtime/Renderer/Private/FogRendering.cpp
@@ -69,7 +69,7 @@ void SetupFogUniformParameters(const FViewInfo& View, FFogUniformParameters& Out
 
 	// Volumetric Fog
 	{
-		FTextureRHIParamRef IntegratedLightScatteringTexture = nullptr;
+		FRHITexture* IntegratedLightScatteringTexture = nullptr;
 
 		if (View.VolumetricFogResources.IntegratedLightScattering)
 		{
diff --git a/Engine/Source/Runtime/Renderer/Private/GPUBenchmark.cpp b/Engine/Source/Runtime/Renderer/Private/GPUBenchmark.cpp
--- a/Engine/Source/Runtime/Renderer/Private/GPUBenchmark.cpp
+++ b/Engine/Source/Runtime/Renderer/Private/GPUBenchmark.cpp
@@ -191,7 +191,7 @@ struct FVertexThroughputDeclaration : public FRenderResource
 TGlobalResource<FVertexThroughputDeclaration> GVertexThroughputDeclaration;
 
 template <uint32 VsMethod, uint32 PsMethod>
-void RunBenchmarkShader(FRHICommandList& RHICmdList, FVertexBufferRHIParamRef VertexThroughputBuffer, const FSceneView& View, TRefCountPtr<IPooledRenderTarget>& Src, float WorkScale)
+void RunBenchmarkShader(FRHICommandList& RHICmdList, FRHIVertexBuffer* VertexThroughputBuffer, const FSceneView& View, TRefCountPtr<IPooledRenderTarget>& Src, float WorkScale)
 {
 	auto ShaderMap = GetGlobalShaderMap(View.GetFeatureLevel());
 
@@ -268,7 +268,7 @@ void RunBenchmarkShader(FRHICommandList& RHICmdList, FVertexBufferRHIParamRef Ve
 	}
 }
 
-void RunBenchmarkShader(FRHICommandListImmediate& RHICmdList, FVertexBufferRHIParamRef VertexThroughputBuffer, const FSceneView& View, uint32 MethodId, TRefCountPtr<IPooledRenderTarget>& Src, float WorkScale)
+void RunBenchmarkShader(FRHICommandListImmediate& RHICmdList, FRHIVertexBuffer* VertexThroughputBuffer, const FSceneView& View, uint32 MethodId, TRefCountPtr<IPooledRenderTarget>& Src, float WorkScale)
 {
 	SCOPED_DRAW_EVENTF(RHICmdList, Benchmark, TEXT("Benchmark Method:%d"), MethodId);
 
diff --git a/Engine/Source/Runtime/Renderer/Private/GPUFastFourierTransform.h b/Engine/Source/Runtime/Renderer/Private/GPUFastFourierTransform.h
--- a/Engine/Source/Runtime/Renderer/Private/GPUFastFourierTransform.h
+++ b/Engine/Source/Runtime/Renderer/Private/GPUFastFourierTransform.h
@@ -520,7 +520,7 @@ namespace GPUFFTComputeShaderUtils
 			return *this;
 		}
 
-		inline FComputeParamterValueSetter& operator()(FShaderResourceParameter& TextureParameter, const FTextureRHIParamRef& TextureRHI)
+		inline FComputeParamterValueSetter& operator()(FShaderResourceParameter& TextureParameter, FRHITexture* TextureRHI)
 		{
 			if (TextureParameter.IsBound())
 			{
@@ -531,7 +531,7 @@ namespace GPUFFTComputeShaderUtils
 
 
 		template <ESamplerFilter Filter, ESamplerAddressMode AddressMode>
-		inline FComputeParamterValueSetter& Set(FShaderResourceParameter& TextureParameter, FShaderResourceParameter& SamplerParameter, const FTextureRHIParamRef& TextureRHI)
+		inline FComputeParamterValueSetter& Set(FShaderResourceParameter& TextureParameter, FShaderResourceParameter& SamplerParameter, FRHITexture* TextureRHI)
 		{
 			SetTextureParameter(*RHICmdList, ShaderRHI, TextureParameter, SamplerParameter,
 				TStaticSamplerState<Filter, AddressMode, AddressMode, AddressMode>::GetRHI(), TextureRHI);
diff --git a/Engine/Source/Runtime/Renderer/Private/GlobalDistanceField.cpp b/Engine/Source/Runtime/Renderer/Private/GlobalDistanceField.cpp
--- a/Engine/Source/Runtime/Renderer/Private/GlobalDistanceField.cpp
+++ b/Engine/Source/Runtime/Renderer/Private/GlobalDistanceField.cpp
@@ -108,7 +108,7 @@ void FGlobalDistanceFieldInfo::UpdateParameterData(float MaxOcclusionDistance)
 	{
 		for (int32 ClipmapIndex = 0; ClipmapIndex < GMaxGlobalDistanceFieldClipmaps; ClipmapIndex++)
 		{
-			FTextureRHIParamRef TextureValue = ClipmapIndex < Clipmaps.Num() 
+			FRHITexture* TextureValue = ClipmapIndex < Clipmaps.Num()
 				? Clipmaps[ClipmapIndex].RenderTarget->GetRenderTargetItem().ShaderResourceTexture 
 				: NULL;
 
diff --git a/Engine/Source/Runtime/Renderer/Private/HeightfieldLighting.cpp b/Engine/Source/Runtime/Renderer/Private/HeightfieldLighting.cpp
--- a/Engine/Source/Runtime/Renderer/Private/HeightfieldLighting.cpp
+++ b/Engine/Source/Runtime/Renderer/Private/HeightfieldLighting.cpp
@@ -466,7 +466,7 @@ void FHeightfieldLightingViewInfo::SetupVisibleHeightfields(const FViewInfo& Vie
 					const FIntPoint LightingAtlasSize = ExistingAtlas->GetAtlasSize();
 					const FVector2D InvLightingAtlasSize(1.0f / LightingAtlasSize.X, 1.0f / LightingAtlasSize.Y);
 
-					FTextureRHIParamRef RenderTargets[3] =
+					FRHITexture* RenderTargets[3] =
 					{
 						ExistingAtlas->Height->GetRenderTargetItem().TargetableTexture,
 						ExistingAtlas->Normal->GetRenderTargetItem().TargetableTexture,
@@ -783,7 +783,7 @@ public:
 		SetShaderValue(RHICmdList, ShaderRHI, ShadowDepthBias, ProjectedShadowInfo->GetShaderDepthBias());
 		SetShaderValue(RHICmdList, ShaderRHI, CascadeDepthMinMax, FVector2D(ProjectedShadowInfo->CascadeSettings.SplitNear, ProjectedShadowInfo->CascadeSettings.SplitFar));
 
-		FTextureRHIParamRef ShadowDepthTextureValue = ProjectedShadowInfo->RenderTargets.DepthTarget->GetRenderTargetItem().ShaderResourceTexture.GetReference();
+		FRHITexture* ShadowDepthTextureValue = ProjectedShadowInfo->RenderTargets.DepthTarget->GetRenderTargetItem().ShaderResourceTexture.GetReference();
 		FRHISamplerState* DepthSamplerState = TStaticSamplerState<SF_Point, AM_Clamp, AM_Clamp, AM_Clamp>::GetRHI();
 
 		SetTextureParameter(RHICmdList, ShaderRHI, ShadowDepthTexture, ShadowDepthTextureSampler, DepthSamplerState, ShadowDepthTextureValue);
diff --git a/Engine/Source/Runtime/Renderer/Private/LightMapRendering.cpp b/Engine/Source/Runtime/Renderer/Private/LightMapRendering.cpp
--- a/Engine/Source/Runtime/Renderer/Private/LightMapRendering.cpp
+++ b/Engine/Source/Runtime/Renderer/Private/LightMapRendering.cpp
@@ -37,7 +37,7 @@ void FSelfShadowedCachedPointIndirectLightingPolicy::ModifyCompilationEnvironmen
 	FSelfShadowedTranslucencyPolicy::ModifyCompilationEnvironment(Parameters, OutEnvironment);
 }
 
-void SetupLCIUniformBuffers(const FPrimitiveSceneProxy* PrimitiveSceneProxy, const FLightCacheInterface* LCI, FUniformBufferRHIParamRef& PrecomputedLightingBuffer, FUniformBufferRHIParamRef& LightmapResourceClusterBuffer, FUniformBufferRHIParamRef& IndirectLightingCacheBuffer)
+void SetupLCIUniformBuffers(const FPrimitiveSceneProxy* PrimitiveSceneProxy, const FLightCacheInterface* LCI, FRHIUniformBuffer*& PrecomputedLightingBuffer, FRHIUniformBuffer*& LightmapResourceClusterBuffer, FRHIUniformBuffer*& IndirectLightingCacheBuffer)
 {
 	if (LCI)
 	{
@@ -77,9 +77,9 @@ void FSelfShadowedCachedPointIndirectLightingPolicy::GetPixelShaderBindings(
 	const PixelParametersType* PixelShaderParameters,
 	FMeshDrawSingleShaderBindings& ShaderBindings)
 {
-	FUniformBufferRHIParamRef PrecomputedLightingBuffer = nullptr;
-	FUniformBufferRHIParamRef LightmapResourceClusterBuffer = nullptr;
-	FUniformBufferRHIParamRef IndirectLightingCacheBuffer = nullptr;
+	FRHIUniformBuffer* PrecomputedLightingBuffer = nullptr;
+	FRHIUniformBuffer* LightmapResourceClusterBuffer = nullptr;
+	FRHIUniformBuffer* IndirectLightingCacheBuffer = nullptr;
 
 	SetupLCIUniformBuffers(PrimitiveSceneProxy, ShaderElementData.LCI, PrecomputedLightingBuffer, LightmapResourceClusterBuffer, IndirectLightingCacheBuffer);
 
@@ -96,9 +96,9 @@ void FSelfShadowedVolumetricLightmapPolicy::GetPixelShaderBindings(
 	const PixelParametersType* PixelShaderParameters,
 	FMeshDrawSingleShaderBindings& ShaderBindings)
 {
-	FUniformBufferRHIParamRef PrecomputedLightingBuffer = nullptr;
-	FUniformBufferRHIParamRef LightmapResourceClusterBuffer = nullptr;
-	FUniformBufferRHIParamRef IndirectLightingCacheBuffer = nullptr;
+	FRHIUniformBuffer* PrecomputedLightingBuffer = nullptr;
+	FRHIUniformBuffer* LightmapResourceClusterBuffer = nullptr;
+	FRHIUniformBuffer* IndirectLightingCacheBuffer = nullptr;
 
 	SetupLCIUniformBuffers(PrimitiveSceneProxy, ShaderElementData.LCI, PrecomputedLightingBuffer, LightmapResourceClusterBuffer, IndirectLightingCacheBuffer);
 
@@ -115,9 +115,9 @@ void FUniformLightMapPolicy::GetVertexShaderBindings(
 	const VertexParametersType* VertexShaderParameters,
 	FMeshDrawSingleShaderBindings& ShaderBindings) 
 {
-	FUniformBufferRHIParamRef PrecomputedLightingBuffer = nullptr;
-	FUniformBufferRHIParamRef LightmapResourceClusterBuffer = nullptr;
-	FUniformBufferRHIParamRef IndirectLightingCacheBuffer = nullptr;
+	FRHIUniformBuffer* PrecomputedLightingBuffer = nullptr;
+	FRHIUniformBuffer* LightmapResourceClusterBuffer = nullptr;
+	FRHIUniformBuffer* IndirectLightingCacheBuffer = nullptr;
 
 	SetupLCIUniformBuffers(PrimitiveSceneProxy, ShaderElementData, PrecomputedLightingBuffer, LightmapResourceClusterBuffer, IndirectLightingCacheBuffer);
 
@@ -132,9 +132,9 @@ void FUniformLightMapPolicy::GetPixelShaderBindings(
 	const PixelParametersType* PixelShaderParameters,
 	FMeshDrawSingleShaderBindings& ShaderBindings)
 {
-	FUniformBufferRHIParamRef PrecomputedLightingBuffer = nullptr;
-	FUniformBufferRHIParamRef LightmapResourceClusterBuffer = nullptr;
-	FUniformBufferRHIParamRef IndirectLightingCacheBuffer = nullptr;
+	FRHIUniformBuffer* PrecomputedLightingBuffer = nullptr;
+	FRHIUniformBuffer* LightmapResourceClusterBuffer = nullptr;
+	FRHIUniformBuffer* IndirectLightingCacheBuffer = nullptr;
 
 	SetupLCIUniformBuffers(PrimitiveSceneProxy, ShaderElementData, PrecomputedLightingBuffer, LightmapResourceClusterBuffer, IndirectLightingCacheBuffer);
 
@@ -151,9 +151,9 @@ void FUniformLightMapPolicy::GetRayHitGroupShaderBindings(
 	FMeshDrawSingleShaderBindings& RayHitGroupBindings
 ) const
 {
-	FUniformBufferRHIParamRef PrecomputedLightingBuffer = nullptr;
-	FUniformBufferRHIParamRef LightmapResourceClusterBuffer = nullptr;
-	FUniformBufferRHIParamRef IndirectLightingCacheBuffer = nullptr;
+	FRHIUniformBuffer* PrecomputedLightingBuffer = nullptr;
+	FRHIUniformBuffer* LightmapResourceClusterBuffer = nullptr;
+	FRHIUniformBuffer* IndirectLightingCacheBuffer = nullptr;
 
 	SetupLCIUniformBuffers(PrimitiveSceneProxy, LCI, PrecomputedLightingBuffer, LightmapResourceClusterBuffer, IndirectLightingCacheBuffer);
 
diff --git a/Engine/Source/Runtime/Renderer/Private/LightMapRendering.h b/Engine/Source/Runtime/Renderer/Private/LightMapRendering.h
--- a/Engine/Source/Runtime/Renderer/Private/LightMapRendering.h
+++ b/Engine/Source/Runtime/Renderer/Private/LightMapRendering.h
@@ -178,7 +178,7 @@ class FSelfShadowedTranslucencyPolicy
 {
 public:
 
-	typedef const FUniformBufferRHIParamRef ElementDataType;
+	typedef FRHIUniformBuffer* ElementDataType;
 
 	class VertexParametersType
 	{
@@ -972,7 +972,7 @@ public:
 struct FSelfShadowLightCacheElementData
 {
 	const FLightCacheInterface* LCI;
-	FUniformBufferRHIParamRef SelfShadowTranslucencyUniformBuffer;
+	FRHIUniformBuffer* SelfShadowTranslucencyUniformBuffer;
 };
 
 /**
diff --git a/Engine/Source/Runtime/Renderer/Private/LightPropagationVolume.cpp b/Engine/Source/Runtime/Renderer/Private/LightPropagationVolume.cpp
--- a/Engine/Source/Runtime/Renderer/Private/LightPropagationVolume.cpp
+++ b/Engine/Source/Runtime/Renderer/Private/LightPropagationVolume.cpp
@@ -302,7 +302,7 @@ public:
 		{
 			if ( LpvBufferSRVParameters[i].IsBound() )
 		    {
-				RHICmdList.SetShaderTexture(ShaderRHI, LpvBufferSRVParameters[i].GetBaseIndex(), FTextureRHIParamRef());
+				RHICmdList.SetShaderTexture(ShaderRHI, LpvBufferSRVParameters[i].GetBaseIndex(), nullptr);
 		    }
 			if ( LpvBufferUAVs[i].IsBound() )
 		    {
@@ -332,7 +332,7 @@ public:
 		{
 			if ( GvBufferSRVParameters[i].IsBound() )
 			{
-					RHICmdList.SetShaderTexture(ShaderRHI, GvBufferSRVParameters[i].GetBaseIndex(), FTextureRHIParamRef());
+					RHICmdList.SetShaderTexture(ShaderRHI, GvBufferSRVParameters[i].GetBaseIndex(), nullptr);
 			}
 			if ( GvBufferUAVs[i].IsBound() )
 			{
@@ -479,9 +479,9 @@ public:
 	void SetParameters(
 		FRHICommandList& RHICmdList, 
 		FLpvBaseWriteShaderParams& BaseParams,
-		FTextureRHIParamRef RsmDiffuseTextureRHI,
-		FTextureRHIParamRef RsmNormalTextureRHI,
-		FTextureRHIParamRef RsmDepthTextureRHI )
+		FRHITexture* RsmDiffuseTextureRHI,
+		FRHITexture* RsmNormalTextureRHI,
+		FRHITexture* RsmDepthTextureRHI )
 	{
 		FRHIComputeShader* ShaderRHI = GetComputeShader();
 		FLpvWriteShaderCSBase::SetParameters(RHICmdList, BaseParams );
diff --git a/Engine/Source/Runtime/Renderer/Private/LightPropagationVolume.h b/Engine/Source/Runtime/Renderer/Private/LightPropagationVolume.h
--- a/Engine/Source/Runtime/Renderer/Private/LightPropagationVolume.h
+++ b/Engine/Source/Runtime/Renderer/Private/LightPropagationVolume.h
@@ -127,7 +127,7 @@ END_GLOBAL_SHADER_PARAMETER_STRUCT()
 struct FLpvBaseWriteShaderParams
 {
 	FLpvWriteUniformBufferRef		UniformBuffer;
-	FTextureRHIParamRef				LpvBufferSRVs[7];
+	FRHITexture*				LpvBufferSRVs[7];
 	FRHIUnorderedAccessView* LpvBufferUAVs[7];
 
 	FRHIShaderResourceView*	VplListHeadBufferSRV;
@@ -135,7 +135,7 @@ struct FLpvBaseWriteShaderParams
 	FRHIShaderResourceView*	VplListBufferSRV;
 	FRHIUnorderedAccessView* VplListBufferUAV;
 
-	FTextureRHIParamRef				GvBufferSRVs[3];
+	FRHITexture*				GvBufferSRVs[3];
 	FRHIUnorderedAccessView* GvBufferUAVs[3];
 
 	FRHIShaderResourceView*	GvListHeadBufferSRV;
@@ -144,7 +144,7 @@ struct FLpvBaseWriteShaderParams
 	FRHIUnorderedAccessView* GvListBufferUAV;
 
 	FRHIUnorderedAccessView* AOVolumeTextureUAV;
-	FTextureRHIParamRef				AOVolumeTextureSRV;
+	FRHITexture*				AOVolumeTextureSRV;
 };
 
 class FLightPropagationVolume : public FRefCountedObject
@@ -188,7 +188,7 @@ public:
 	FLpvWriteUniformBufferRef GetWriteUniformBuffer() const				{ return LpvWriteUniformBuffer.GetUniformBufferRef(); }
 	FLpvWriteUniformBufferRef GetRsmUniformBuffer() const				{ return RsmRenderUniformBuffer.GetUniformBufferRef(); }
 
-	FTextureRHIParamRef GetLpvBufferSrv( int i )						{ return LpvVolumeTextures[ 1-mWriteBufferIndex ][i]->GetRenderTargetItem().ShaderResourceTexture; }
+	FRHITexture* GetLpvBufferSrv( int i )						{ return LpvVolumeTextures[ 1-mWriteBufferIndex ][i]->GetRenderTargetItem().ShaderResourceTexture; }
 
 	FRHIUnorderedAccessView* GetVplListBufferUav()				{ return mVplListBuffer->UAV; }
 	FRHIUnorderedAccessView* GetVplListHeadBufferUav()			{ return mVplListHeadBuffer->UAV; }
@@ -209,7 +209,7 @@ public:
 	void GetShadowInfo( const FProjectedShadowInfo& ProjectedShadowInfo, FRsmInfo& RsmInfoOut );
 
 	void ComputeDirectionalOcclusion(FRHICommandListImmediate& RHICmdList, FViewInfo& View);
-	FTextureRHIParamRef GetAOVolumeTextureSRV() { return AOVolumeTexture->GetRenderTargetItem().ShaderResourceTexture; }
+	FRHITexture* GetAOVolumeTextureSRV() { return AOVolumeTexture->GetRenderTargetItem().ShaderResourceTexture; }
 
 	TRefCountPtr<IPooledRenderTarget>	LpvVolumeTextures[2][7];		// double buffered
 	FRWByteAddressBuffer*				mVplListHeadBuffer;
diff --git a/Engine/Source/Runtime/Renderer/Private/LightPropagationVolumeVisualisation.cpp b/Engine/Source/Runtime/Renderer/Private/LightPropagationVolumeVisualisation.cpp
--- a/Engine/Source/Runtime/Renderer/Private/LightPropagationVolumeVisualisation.cpp
+++ b/Engine/Source/Runtime/Renderer/Private/LightPropagationVolumeVisualisation.cpp
@@ -121,7 +121,7 @@ public:
 		
 		for ( int i = 0; i < 7; i++ )
 		{
-			FTextureRHIParamRef LpvBufferSrv = LPV->LpvVolumeTextures[ 1-LPV->mWriteBufferIndex ][i]->GetRenderTargetItem().ShaderResourceTexture;
+			FRHITexture* LpvBufferSrv = LPV->LpvVolumeTextures[ 1-LPV->mWriteBufferIndex ][i]->GetRenderTargetItem().ShaderResourceTexture;
 			if ( LpvBufferSRVParameters[i].IsBound() )
 			{
 				RHICmdList.SetShaderTexture(ShaderRHI, LpvBufferSRVParameters[i].GetBaseIndex(), LpvBufferSrv);
@@ -131,7 +131,7 @@ public:
 		
 		for ( int i = 0; i < NUM_GV_TEXTURES; i++ ) 
 		{
-			FTextureRHIParamRef GvBufferSrv = LPV->GvVolumeTextures[i]->GetRenderTargetItem().ShaderResourceTexture;
+			FRHITexture* GvBufferSrv = LPV->GvVolumeTextures[i]->GetRenderTargetItem().ShaderResourceTexture;
 			if ( GvBufferSRVParameters[i].IsBound() )
 			{
 				RHICmdList.SetShaderTexture(ShaderRHI, GvBufferSRVParameters[i].GetBaseIndex(), GvBufferSrv);
@@ -174,7 +174,7 @@ public:
 		{
 			if ( LpvBufferSRVParameters[i].IsBound() )
 			{
-				RHICmdList.SetShaderTexture(ShaderRHI, LpvBufferSRVParameters[i].GetBaseIndex(), FTextureRHIParamRef());
+				RHICmdList.SetShaderTexture(ShaderRHI, LpvBufferSRVParameters[i].GetBaseIndex(), nullptr);
 			}
 		}
 		
@@ -182,7 +182,7 @@ public:
 		{
 			if ( GvBufferSRVParameters[i].IsBound() )
 			{
-				RHICmdList.SetShaderTexture(ShaderRHI, GvBufferSRVParameters[i].GetBaseIndex(), FTextureRHIParamRef());
+				RHICmdList.SetShaderTexture(ShaderRHI, GvBufferSRVParameters[i].GetBaseIndex(), nullptr);
 			}
 		}
 	}
diff --git a/Engine/Source/Runtime/Renderer/Private/LightRendering.cpp b/Engine/Source/Runtime/Renderer/Private/LightRendering.cpp
--- a/Engine/Source/Runtime/Renderer/Private/LightRendering.cpp
+++ b/Engine/Source/Runtime/Renderer/Private/LightRendering.cpp
@@ -352,7 +352,7 @@ private:
 			);
 
 		{
-			FTextureRHIParamRef TextureRHI = IESTextureResource ? IESTextureResource->TextureRHI : GSystemTextures.WhiteDummy->GetRenderTargetItem().TargetableTexture;
+			FRHITexture* TextureRHI = IESTextureResource ? IESTextureResource->TextureRHI : GSystemTextures.WhiteDummy->GetRenderTargetItem().TargetableTexture;
 
 			SetTextureParameter(
 				RHICmdList, 
@@ -366,7 +366,7 @@ private:
 
 		if( LightingChannelsTexture.IsBound() )
 		{
-			FTextureRHIParamRef LightingChannelsTextureRHI = SceneRenderTargets.LightingChannels ? SceneRenderTargets.LightingChannels->GetRenderTargetItem().ShaderResourceTexture : GSystemTextures.WhiteDummy->GetRenderTargetItem().TargetableTexture;
+			FRHITexture* LightingChannelsTextureRHI = SceneRenderTargets.LightingChannels ? SceneRenderTargets.LightingChannels->GetRenderTargetItem().ShaderResourceTexture : GSystemTextures.WhiteDummy->GetRenderTargetItem().TargetableTexture;
 
 			SetTextureParameter(
 				RHICmdList, 
diff --git a/Engine/Source/Runtime/Renderer/Private/MeshDrawCommands.cpp b/Engine/Source/Runtime/Renderer/Private/MeshDrawCommands.cpp
--- a/Engine/Source/Runtime/Renderer/Private/MeshDrawCommands.cpp
+++ b/Engine/Source/Runtime/Renderer/Private/MeshDrawCommands.cpp
@@ -35,7 +35,7 @@ FPrimitiveIdVertexBufferPool::~FPrimitiveIdVertexBufferPool()
 	check(!Entries.Num());
 }
 
-FVertexBufferRHIParamRef FPrimitiveIdVertexBufferPool::Allocate(int32 BufferSize)
+FRHIVertexBuffer* FPrimitiveIdVertexBufferPool::Allocate(int32 BufferSize)
 {
 	BufferSize = Align(BufferSize, 1024);
 
@@ -868,7 +868,7 @@ void SortAndMergeDynamicPassMeshDrawCommands(
 	ERHIFeatureLevel::Type FeatureLevel,
 	FMeshCommandOneFrameArray& VisibleMeshDrawCommands,
 	FDynamicMeshDrawCommandStorage& MeshDrawCommandStorage,
-	FVertexBufferRHIParamRef& OutPrimitiveIdVertexBuffer,
+	FRHIVertexBuffer*& OutPrimitiveIdVertexBuffer,
 	uint32 InstanceFactor)
 {
 	const bool bUseGPUScene = UseGPUScene(GMaxRHIShaderPlatform, FeatureLevel);
@@ -1073,7 +1073,7 @@ class FDrawVisibleMeshCommandsAnyThreadTask : public FRenderTask
 	FRHICommandList& RHICmdList;
 	const FMeshCommandOneFrameArray& VisibleMeshDrawCommands;
 	const FGraphicsMinimalPipelineStateSet& GraphicsMinimalPipelineStateSet;
-	FVertexBufferRHIParamRef PrimitiveIdsBuffer;
+	FRHIVertexBuffer* PrimitiveIdsBuffer;
 	int32 BasePrimitiveIdsOffset;
 	bool bDynamicInstancing;
 	uint32 InstanceFactor;
@@ -1086,7 +1086,7 @@ public:
 		FRHICommandList& InRHICmdList,
 		const FMeshCommandOneFrameArray& InVisibleMeshDrawCommands,
 		const FGraphicsMinimalPipelineStateSet& InGraphicsMinimalPipelineStateSet,
-		FVertexBufferRHIParamRef InPrimitiveIdsBuffer,
+		FRHIVertexBuffer* InPrimitiveIdsBuffer,
 		int32 InBasePrimitiveIdsOffset,
 		bool bInDynamicInstancing,
 		uint32 InInstanceFactor,
@@ -1133,14 +1133,14 @@ public:
  */
 struct FRHICommandUpdatePrimitiveIdBuffer : public FRHICommand<FRHICommandUpdatePrimitiveIdBuffer>
 {
-	FVertexBufferRHIParamRef VertexBuffer;
+	FRHIVertexBuffer* VertexBuffer;
 	void* VertexBufferData;
 	int32 VertexBufferDataSize;
 
 	virtual ~FRHICommandUpdatePrimitiveIdBuffer() {}
 
 	FORCEINLINE_DEBUGGABLE FRHICommandUpdatePrimitiveIdBuffer(
-		FVertexBufferRHIParamRef InVertexBuffer,
+		FRHIVertexBuffer* InVertexBuffer,
 		void* InVertexBufferData,
 		int32 InVertexBufferDataSize)
 		: VertexBuffer(InVertexBuffer)
@@ -1167,7 +1167,7 @@ void FParallelMeshDrawCommandPass::DispatchDraw(FParallelCommandListSet* Paralle
 		return;
 	}
 
-	FVertexBufferRHIParamRef PrimitiveIdsBuffer = PrimitiveIdVertexBufferRHI;
+	FRHIVertexBuffer* PrimitiveIdsBuffer = PrimitiveIdVertexBufferRHI;
 	const int32 BasePrimitiveIdsOffset = 0;
 
 	if (ParallelCommandListSet)
diff --git a/Engine/Source/Runtime/Renderer/Private/MeshDrawCommands.h b/Engine/Source/Runtime/Renderer/Private/MeshDrawCommands.h
--- a/Engine/Source/Runtime/Renderer/Private/MeshDrawCommands.h
+++ b/Engine/Source/Runtime/Renderer/Private/MeshDrawCommands.h
@@ -22,7 +22,7 @@ public:
 	FPrimitiveIdVertexBufferPool();
 	~FPrimitiveIdVertexBufferPool();
 
-	FVertexBufferRHIParamRef Allocate(int32 BufferSize);
+	FRHIVertexBuffer* Allocate(int32 BufferSize);
 	void DiscardAll();
 
 	virtual void ReleaseDynamicRHI() override;
@@ -164,7 +164,7 @@ public:
 
 
 private:
-	FVertexBufferRHIParamRef PrimitiveIdVertexBufferRHI;
+	FRHIVertexBuffer* PrimitiveIdVertexBufferRHI;
 	FMeshDrawCommandPassSetupTaskContext TaskContext;
 	FGraphEventRef TaskEventRef;
 	FString PassNameForStats;
@@ -184,5 +184,5 @@ extern void SortAndMergeDynamicPassMeshDrawCommands(
 	ERHIFeatureLevel::Type FeatureLevel,
 	FMeshCommandOneFrameArray& VisibleMeshDrawCommands,
 	FDynamicMeshDrawCommandStorage& MeshDrawCommandStorage,
-	FVertexBufferRHIParamRef& OutPrimitiveIdVertexBuffer,
+	FRHIVertexBuffer*& OutPrimitiveIdVertexBuffer,
 	uint32 InstanceFactor);
\ No newline at end of file
diff --git a/Engine/Source/Runtime/Renderer/Private/MeshPassProcessor.cpp b/Engine/Source/Runtime/Renderer/Private/MeshPassProcessor.cpp
--- a/Engine/Source/Runtime/Renderer/Private/MeshPassProcessor.cpp
+++ b/Engine/Source/Runtime/Renderer/Private/MeshPassProcessor.cpp
@@ -45,9 +45,9 @@ public:
 	int32 MaxSRVUsed = -1;
 	FRHIShaderResourceView* SRVs[MAX_SRVs_PER_SHADER_STAGE] = {};
 	int32 MaxUniformBufferUsed = -1;
-	FUniformBufferRHIParamRef UniformBuffers[MAX_UNIFORM_BUFFERS_PER_SHADER_STAGE] = {};
+	FRHIUniformBuffer* UniformBuffers[MAX_UNIFORM_BUFFERS_PER_SHADER_STAGE] = {};
 	int32 MaxTextureUsed = -1;
-	FTextureRHIParamRef Textures[MAX_SRVs_PER_SHADER_STAGE] = {};
+	FRHITexture* Textures[MAX_SRVs_PER_SHADER_STAGE] = {};
 	int32 MaxSamplerUsed = -1;
 	FRHISamplerState* Samplers[MAX_SAMPLERS_PER_SHADER_STAGE] = {};
 };
@@ -61,21 +61,21 @@ public:
 		Data = InData;
 	}
 
-	inline const FUniformBufferRHIParamRef* GetUniformBufferStart() const
+	inline FRHIUniformBuffer*const* GetUniformBufferStart() const
 	{
-		return (const FUniformBufferRHIParamRef*)(Data + GetUniformBufferOffset());
+		return (FRHIUniformBuffer**)(Data + GetUniformBufferOffset());
 	}
 
-	inline const FRHISamplerState** GetSamplerStart() const
+	inline FRHISamplerState** GetSamplerStart() const
 	{
 		const uint8* SamplerDataStart = Data + GetSamplerOffset();
-		return (const FRHISamplerState**)SamplerDataStart;
+		return (FRHISamplerState**)SamplerDataStart;
 	}
 
-	inline const FRHIResource** GetSRVStart() const
+	inline FRHIResource** GetSRVStart() const
 	{
 		const uint8* SRVDataStart = Data + GetSRVOffset();
-		return (const FRHIResource**)SRVDataStart;
+		return (FRHIResource**)SRVDataStart;
 	}
 
 	inline const uint8* GetSRVTypeStart() const
@@ -101,7 +101,7 @@ void FMeshDrawShaderBindings::SetShaderBindings(
 	const FReadOnlyMeshDrawSingleShaderBindings& RESTRICT SingleShaderBindings,
 	FShaderBindingState& RESTRICT ShaderBindingState)
 {
-	const FUniformBufferRHIParamRef* RESTRICT UniformBufferBindings = SingleShaderBindings.GetUniformBufferStart();
+	FRHIUniformBuffer* const* RESTRICT UniformBufferBindings = SingleShaderBindings.GetUniformBufferStart();
 	const FShaderParameterInfo* RESTRICT UniformBufferParameters = SingleShaderBindings.ParameterMapInfo.UniformBuffers.GetData();
 	const int32 NumUniformBuffers = SingleShaderBindings.ParameterMapInfo.UniformBuffers.Num();
 
@@ -109,7 +109,7 @@ void FMeshDrawShaderBindings::SetShaderBindings(
 	{
 		FShaderParameterInfo Parameter = UniformBufferParameters[UniformBufferIndex];
 		checkSlow(Parameter.BaseIndex < ARRAY_COUNT(ShaderBindingState.UniformBuffers));
-		FUniformBufferRHIParamRef UniformBuffer = UniformBufferBindings[UniformBufferIndex];
+		FRHIUniformBuffer* UniformBuffer = UniformBufferBindings[UniformBufferIndex];
 
 		if (UniformBuffer != ShaderBindingState.UniformBuffers[Parameter.BaseIndex])
 		{
@@ -119,7 +119,7 @@ void FMeshDrawShaderBindings::SetShaderBindings(
 		}
 	}
 
-	const FRHISamplerState** RESTRICT SamplerBindings = SingleShaderBindings.GetSamplerStart();
+	FRHISamplerState* const* RESTRICT SamplerBindings = SingleShaderBindings.GetSamplerStart();
 	const FShaderParameterInfo* RESTRICT TextureSamplerParameters = SingleShaderBindings.ParameterMapInfo.TextureSamplers.GetData();
 	const int32 NumTextureSamplers = SingleShaderBindings.ParameterMapInfo.TextureSamplers.Num();
 
@@ -138,7 +138,7 @@ void FMeshDrawShaderBindings::SetShaderBindings(
 	}
 
 	const uint8* RESTRICT SRVType = SingleShaderBindings.GetSRVTypeStart();
-	const FRHIResource** RESTRICT SRVBindings = SingleShaderBindings.GetSRVStart();
+	FRHIResource* const* RESTRICT SRVBindings = SingleShaderBindings.GetSRVStart();
 	const FShaderParameterInfo* RESTRICT SRVParameters = SingleShaderBindings.ParameterMapInfo.SRVs.GetData();
 	const uint32 NumSRVs = SingleShaderBindings.ParameterMapInfo.SRVs.Num();
 
@@ -163,7 +163,7 @@ void FMeshDrawShaderBindings::SetShaderBindings(
 		}
 		else
 		{
-			FTextureRHIParamRef Texture = (FTextureRHIParamRef)SRVBindings[SRVIndex];
+			FRHITexture* Texture = (FRHITexture*)SRVBindings[SRVIndex];
 
 			if (Texture != ShaderBindingState.Textures[Parameter.BaseIndex])
 			{
@@ -199,19 +199,19 @@ void FMeshDrawShaderBindings::SetShaderBindings(
 	RHIShaderType Shader,
 	const FReadOnlyMeshDrawSingleShaderBindings& RESTRICT SingleShaderBindings)
 {
-	const FUniformBufferRHIParamRef* RESTRICT UniformBufferBindings = SingleShaderBindings.GetUniformBufferStart();
+	FRHIUniformBuffer* const* RESTRICT UniformBufferBindings = SingleShaderBindings.GetUniformBufferStart();
 	const FShaderParameterInfo* RESTRICT UniformBufferParameters = SingleShaderBindings.ParameterMapInfo.UniformBuffers.GetData();
 	const int32 NumUniformBuffers = SingleShaderBindings.ParameterMapInfo.UniformBuffers.Num();
 
 	for (int32 UniformBufferIndex = 0; UniformBufferIndex < NumUniformBuffers; UniformBufferIndex++)
 	{
 		FShaderParameterInfo Parameter = UniformBufferParameters[UniformBufferIndex];
-		FUniformBufferRHIParamRef UniformBuffer = UniformBufferBindings[UniformBufferIndex];
+		FRHIUniformBuffer* UniformBuffer = UniformBufferBindings[UniformBufferIndex];
 
 		RHICmdList.SetShaderUniformBuffer(Shader, Parameter.BaseIndex, UniformBuffer);
 	}
 
-	const FRHISamplerState** RESTRICT SamplerBindings = SingleShaderBindings.GetSamplerStart();
+	FRHISamplerState* const* RESTRICT SamplerBindings = SingleShaderBindings.GetSamplerStart();
 	const FShaderParameterInfo* RESTRICT TextureSamplerParameters = SingleShaderBindings.ParameterMapInfo.TextureSamplers.GetData();
 	const int32 NumTextureSamplers = SingleShaderBindings.ParameterMapInfo.TextureSamplers.Num();
 
@@ -224,7 +224,7 @@ void FMeshDrawShaderBindings::SetShaderBindings(
 	}
 
 	const uint8* RESTRICT SRVType = SingleShaderBindings.GetSRVTypeStart();
-	const FRHIResource** RESTRICT SRVBindings = SingleShaderBindings.GetSRVStart();
+	FRHIResource* const* RESTRICT SRVBindings = SingleShaderBindings.GetSRVStart();
 	const FShaderParameterInfo* RESTRICT SRVParameters = SingleShaderBindings.ParameterMapInfo.SRVs.GetData();
 	const uint32 NumSRVs = SingleShaderBindings.ParameterMapInfo.SRVs.Num();
 
@@ -242,7 +242,7 @@ void FMeshDrawShaderBindings::SetShaderBindings(
 		}
 		else
 		{
-			FTextureRHIParamRef Texture = (FTextureRHIParamRef)SRVBindings[SRVIndex];
+			FRHITexture* Texture = (FRHITexture*)SRVBindings[SRVIndex];
 			RHICmdList.SetShaderTexture(Shader, Parameter.BaseIndex, Texture);
 		}
 	}
@@ -280,9 +280,9 @@ void FMeshDrawShaderBindings::SetRayTracingShaderBindingsForHitGroup(
 
 	FReadOnlyMeshDrawSingleShaderBindings SingleShaderBindings(ShaderLayouts[0], GetData());
 
-	FUniformBufferRHIParamRef LocalUniformBuffers[MAX_UNIFORM_BUFFERS_PER_SHADER_STAGE] = {};
+	const FRHIUniformBuffer* LocalUniformBuffers[MAX_UNIFORM_BUFFERS_PER_SHADER_STAGE] = {};
 
-	const FUniformBufferRHIParamRef* RESTRICT UniformBufferBindings = SingleShaderBindings.GetUniformBufferStart();
+	FRHIUniformBuffer* const* RESTRICT UniformBufferBindings = SingleShaderBindings.GetUniformBufferStart();
 	const FShaderParameterInfo* RESTRICT UniformBufferParameters = SingleShaderBindings.ParameterMapInfo.UniformBuffers.GetData();
 	const int32 NumUniformBuffers = SingleShaderBindings.ParameterMapInfo.UniformBuffers.Num();
 
@@ -291,7 +291,7 @@ void FMeshDrawShaderBindings::SetRayTracingShaderBindingsForHitGroup(
 	{
 		FShaderParameterInfo Parameter = UniformBufferParameters[UniformBufferIndex];
 		checkSlow(Parameter.BaseIndex < ARRAY_COUNT(LocalUniformBuffers));
-		FUniformBufferRHIParamRef UniformBuffer = UniformBufferBindings[UniformBufferIndex];
+		const FRHIUniformBuffer* UniformBuffer = UniformBufferBindings[UniformBufferIndex];
 		if (Parameter.BaseIndex < ARRAY_COUNT(LocalUniformBuffers))
 		{
 			LocalUniformBuffers[Parameter.BaseIndex] = UniformBuffer;
@@ -339,7 +339,7 @@ void FMeshDrawShaderBindings::SetRayTracingShaderBindingsForHitGroup(
 	uint32 NumUniformBuffersToSet = MaxUniformBufferUsed + 1;
 	const uint32 UserData = 0; // UserData could be used to store material ID or any other kind of per-material constant. This can be retrieved in hit shaders via GetHitGroupUserData().
 	RHICmdList.SetRayTracingHitGroup(Scene, InstanceIndex, SegmentIndex, ShaderSlot, PipelineState, HitGroupIndex, 
-		NumUniformBuffersToSet, LocalUniformBuffers,
+		NumUniformBuffersToSet, (FRHIUniformBuffer**)LocalUniformBuffers,
 		ReorderedLooseParameters.Num(), ReorderedLooseParameters.GetData(),
 		UserData);
 }
@@ -489,12 +489,12 @@ FMeshDrawShaderBindings::~FMeshDrawShaderBindings()
 	{
 		FMeshDrawSingleShaderBindings SingleShaderBindings(ShaderLayouts[ShaderBindingsIndex], ShaderBindingDataPtr);
 
-		const FUniformBufferRHIParamRef* RESTRICT UniformBufferBindings = SingleShaderBindings.GetUniformBufferStart();
+		FRHIUniformBuffer** RESTRICT UniformBufferBindings = SingleShaderBindings.GetUniformBufferStart();
 		const int32 NumUniformBuffers = SingleShaderBindings.ParameterMapInfo.UniformBuffers.Num();
 
 		for (int32 UniformBufferIndex = 0; UniformBufferIndex < NumUniformBuffers; UniformBufferIndex++)
 		{
-			FUniformBufferRHIParamRef UniformBuffer = UniformBufferBindings[UniformBufferIndex];
+			FRHIUniformBuffer* UniformBuffer = UniformBufferBindings[UniformBufferIndex];
 
 			if (UniformBuffer)
 			{
@@ -595,13 +595,13 @@ void FMeshDrawShaderBindings::Finalize(const FMeshProcessorShaders* ShadersForDe
 
 		FReadOnlyMeshDrawSingleShaderBindings SingleShaderBindings(ShaderLayout, ShaderBindingDataPtr);
 
-		const FUniformBufferRHIParamRef* UniformBufferBindings = SingleShaderBindings.GetUniformBufferStart();
+		FRHIUniformBuffer* const* UniformBufferBindings = SingleShaderBindings.GetUniformBufferStart();
 
 		for (int32 BindingIndex = 0; BindingIndex < ShaderLayout.ParameterMapInfo.UniformBuffers.Num(); BindingIndex++)
 		{
 			FShaderParameterInfo ParameterInfo = ShaderLayout.ParameterMapInfo.UniformBuffers[BindingIndex];
 
-			FUniformBufferRHIParamRef UniformBufferValue = UniformBufferBindings[BindingIndex];
+			FRHIUniformBuffer* UniformBufferValue = UniformBufferBindings[BindingIndex];
 
 			if (!UniformBufferValue)
 			{
@@ -626,7 +626,7 @@ void FMeshDrawShaderBindings::Finalize(const FMeshProcessorShaders* ShadersForDe
 			}
 		}
 
-		const FRHISamplerState** SamplerBindings = SingleShaderBindings.GetSamplerStart();
+		FRHISamplerState* const* SamplerBindings = SingleShaderBindings.GetSamplerStart();
 
 		for (int32 BindingIndex = 0; BindingIndex < ShaderLayout.ParameterMapInfo.TextureSamplers.Num(); BindingIndex++)
 		{
@@ -639,7 +639,7 @@ void FMeshDrawShaderBindings::Finalize(const FMeshProcessorShaders* ShadersForDe
 		}
 
 		const uint8* RESTRICT SRVType = SingleShaderBindings.GetSRVTypeStart();
-		const FRHIResource** RESTRICT SRVBindings = SingleShaderBindings.GetSRVStart();
+		FRHIResource* const* RESTRICT SRVBindings = SingleShaderBindings.GetSRVStart();
 		const FShaderParameterInfo* RESTRICT SRVParameters = SingleShaderBindings.ParameterMapInfo.SRVs.GetData();
 		const uint32 NumSRVs = SingleShaderBindings.ParameterMapInfo.SRVs.Num();
 
@@ -661,7 +661,7 @@ void FMeshDrawShaderBindings::Finalize(const FMeshProcessorShaders* ShadersForDe
 			}
 			else
 			{
-				FTextureRHIParamRef Texture = (FTextureRHIParamRef)SRVBindings[SRVIndex];
+				FRHITexture* Texture = (FRHITexture*)SRVBindings[SRVIndex];
 
 				ensureMsgf(Texture, TEXT("Shader %s with vertex factory %s never set texture at BaseIndex %u.  This can cause GPU hangs, depending on how the shader uses it."), 
 					Shader->GetType()->GetName(), 
@@ -688,12 +688,12 @@ void FMeshDrawShaderBindings::CopyFrom(const FMeshDrawShaderBindings& Other)
 	for (int32 ShaderBindingsIndex = 0; ShaderBindingsIndex < ShaderLayouts.Num(); ShaderBindingsIndex++)
 	{
 		FMeshDrawSingleShaderBindings SingleShaderBindings(ShaderLayouts[ShaderBindingsIndex], ShaderBindingDataPtr);
-		const FUniformBufferRHIParamRef* RESTRICT UniformBufferBindings = SingleShaderBindings.GetUniformBufferStart();
+		FRHIUniformBuffer** RESTRICT UniformBufferBindings = SingleShaderBindings.GetUniformBufferStart();
 		const int32 NumUniformBuffers = SingleShaderBindings.ParameterMapInfo.UniformBuffers.Num();
 
 		for (int32 UniformBufferIndex = 0; UniformBufferIndex < NumUniformBuffers; UniformBufferIndex++)
 		{
-			FUniformBufferRHIParamRef UniformBuffer = UniformBufferBindings[UniformBufferIndex];
+			FRHIUniformBuffer* UniformBuffer = UniformBufferBindings[UniformBufferIndex];
 
 			if (UniformBuffer)
 			{
@@ -834,13 +834,13 @@ bool FMeshDrawShaderBindings::MatchesForDynamicInstancing(const FMeshDrawShaderB
 			return false;
 		}
 
-		const FUniformBufferRHIParamRef* UniformBufferBindings = SingleShaderBindings.GetUniformBufferStart();
-		const FUniformBufferRHIParamRef* OtherUniformBufferBindings = OtherSingleShaderBindings.GetUniformBufferStart();
+		FRHIUniformBuffer* const* UniformBufferBindings = SingleShaderBindings.GetUniformBufferStart();
+		FRHIUniformBuffer* const* OtherUniformBufferBindings = OtherSingleShaderBindings.GetUniformBufferStart();
 
 		for (int32 UniformBufferIndex = 0; UniformBufferIndex < SingleShaderBindings.ParameterMapInfo.UniformBuffers.Num(); UniformBufferIndex++)
 		{
-			FUniformBufferRHIParamRef UniformBuffer = UniformBufferBindings[UniformBufferIndex];
-			FUniformBufferRHIParamRef OtherUniformBuffer = OtherUniformBufferBindings[UniformBufferIndex];
+			const FRHIUniformBuffer* UniformBuffer = UniformBufferBindings[UniformBufferIndex];
+			const FRHIUniformBuffer* OtherUniformBuffer = OtherUniformBufferBindings[UniformBufferIndex];
 			
 			if (UniformBuffer != OtherUniformBuffer)
 			{
@@ -877,10 +877,10 @@ uint32 FMeshDrawShaderBindings::GetDynamicInstancingHash() const
 			return Hash;
 		}
 
-		const FUniformBufferRHIParamRef* UniformBufferBindings = SingleShaderBindings.GetUniformBufferStart();
+		FRHIUniformBuffer* const* UniformBufferBindings = SingleShaderBindings.GetUniformBufferStart();
 		for (int32 UniformBufferIndex = 0; UniformBufferIndex < SingleShaderBindings.ParameterMapInfo.UniformBuffers.Num(); UniformBufferIndex++)
 		{
-			FUniformBufferRHIParamRef UniformBuffer = UniformBufferBindings[UniformBufferIndex];
+			const FRHIUniformBuffer* UniformBuffer = UniformBufferBindings[UniformBufferIndex];
 			Hash = PointerHash(UniformBuffer, Hash);
 		}
 
@@ -893,7 +893,7 @@ uint32 FMeshDrawShaderBindings::GetDynamicInstancingHash() const
 void FMeshDrawCommand::SubmitDraw(
 	const FMeshDrawCommand& RESTRICT MeshDrawCommand, 
 	const FGraphicsMinimalPipelineStateSet& GraphicsMinimalPipelineStateSet,
-	FVertexBufferRHIParamRef ScenePrimitiveIdsBuffer,
+	FRHIVertexBuffer* ScenePrimitiveIdsBuffer,
 	int32 PrimitiveIdOffset,
 	uint32 InstanceFactor,
 	FRHICommandList& RHICmdList,
@@ -1016,7 +1016,7 @@ void FMeshDrawCommand::SetDebugData(const FPrimitiveSceneProxy* PrimitiveScenePr
 void SubmitMeshDrawCommands(
 	const FMeshCommandOneFrameArray& VisibleMeshDrawCommands,
 	const FGraphicsMinimalPipelineStateSet& GraphicsMinimalPipelineStateSet, 
-	FVertexBufferRHIParamRef PrimitiveIdsBuffer,
+	FRHIVertexBuffer* PrimitiveIdsBuffer,
 	int32 BasePrimitiveIdsOffset,
 	bool bDynamicInstancing,
 	uint32 InstanceFactor,
@@ -1028,7 +1028,7 @@ void SubmitMeshDrawCommands(
 void SubmitMeshDrawCommandsRange(
 	const FMeshCommandOneFrameArray& VisibleMeshDrawCommands,
 	const FGraphicsMinimalPipelineStateSet& GraphicsMinimalPipelineStateSet,
-	FVertexBufferRHIParamRef PrimitiveIdsBuffer,
+	FRHIVertexBuffer* PrimitiveIdsBuffer,
 	int32 BasePrimitiveIdsOffset,
 	bool bDynamicInstancing,
 	int32 StartIndex,
@@ -1062,7 +1062,7 @@ void DrawDynamicMeshPassPrivate(
 	{
 		const bool bDynamicInstancing = IsDynamicInstancingEnabled(View.GetFeatureLevel());
 
-		FVertexBufferRHIParamRef PrimitiveIdVertexBuffer = nullptr;
+		FRHIVertexBuffer* PrimitiveIdVertexBuffer = nullptr;
 
 		SortAndMergeDynamicPassMeshDrawCommands(View.GetFeatureLevel(), VisibleMeshDrawCommands, DynamicMeshDrawCommandStorage, PrimitiveIdVertexBuffer, InstanceFactor);
 
