UE4.23 - Deprecate FUniformBufferRHIParamRef, FIndexBufferRHIParamRef, FVertexBufferRHIParamRef, FStructuredBufferRHIParamRef, FTextureRHIParamRef, FTexture2DRHIParamRef, FTexture2DArrayRHIParamRef, F

diff --git a/Engine/Source/Runtime/Renderer/Private/MobileBasePass.cpp b/Engine/Source/Runtime/Renderer/Private/MobileBasePass.cpp
--- a/Engine/Source/Runtime/Renderer/Private/MobileBasePass.cpp
+++ b/Engine/Source/Runtime/Renderer/Private/MobileBasePass.cpp
@@ -486,21 +486,21 @@ void TMobileBasePassPSPolicyParamType<FUniformLightMapPolicy>::GetShaderBindings
 		}
 		else if (ReflectionParameter.IsBound())
 		{
-			FUniformBufferRHIParamRef RelfectionUB = GDefaultMobileReflectionCaptureUniformBuffer.GetUniformBufferRHI();
+			FRHIUniformBuffer* ReflectionUB = GDefaultMobileReflectionCaptureUniformBuffer.GetUniformBufferRHI();
 			// If no reflection captures are available then attempt to use sky light's texture.
 			if (UseSkyReflectionCapture(Scene) && FeatureLevel > ERHIFeatureLevel::ES2) // not-supported on ES2 at the moment
 			{
-				RelfectionUB = Scene->UniformBuffers.MobileSkyReflectionUniformBuffer;
+				ReflectionUB = Scene->UniformBuffers.MobileSkyReflectionUniformBuffer;
 			}
 			else
 			{
 				FPrimitiveSceneInfo* PrimitiveSceneInfo = PrimitiveSceneProxy ? PrimitiveSceneProxy->GetPrimitiveSceneInfo() : nullptr;
 				if (PrimitiveSceneInfo && PrimitiveSceneInfo->CachedReflectionCaptureProxy)
 				{
-					RelfectionUB = PrimitiveSceneInfo->CachedReflectionCaptureProxy->MobileUniformBuffer;
+					ReflectionUB = PrimitiveSceneInfo->CachedReflectionCaptureProxy->MobileUniformBuffer;
 				}
 			}
-			ShaderBindings.Add(ReflectionParameter, RelfectionUB);
+			ShaderBindings.Add(ReflectionParameter, ReflectionUB);
 		}
 		
 		if (LightPositionAndInvRadiusParameter.IsBound() || SpotLightDirectionParameter.IsBound())
diff --git a/Engine/Source/Runtime/Renderer/Private/MobileSceneCaptureRendering.cpp b/Engine/Source/Runtime/Renderer/Private/MobileSceneCaptureRendering.cpp
--- a/Engine/Source/Runtime/Renderer/Private/MobileSceneCaptureRendering.cpp
+++ b/Engine/Source/Runtime/Renderer/Private/MobileSceneCaptureRendering.cpp
@@ -65,7 +65,7 @@ public:
 		}
 	}
 
-	void SetParameters(FRHICommandList& RHICmdList, const FViewInfo& View, FRHISamplerState* SamplerStateRHI, FTextureRHIParamRef TextureRHI)
+	void SetParameters(FRHICommandList& RHICmdList, const FViewInfo& View, FRHISamplerState* SamplerStateRHI, FRHITexture* TextureRHI)
 	{
 		FGlobalShader::SetParameters<FViewUniformShaderParameters>(RHICmdList, GetPixelShader(), View.ViewUniformBuffer);
 		SetTextureParameter(RHICmdList, GetPixelShader(), InTexture, InTextureSampler, SamplerStateRHI, TextureRHI);
@@ -147,7 +147,7 @@ IMPLEMENT_SHADER_TYPE(template<>, FMobileSceneCaptureCopyVS<true>, TEXT("/Engine
  
 
 template <bool bDemosaic, ESceneCaptureSource CaptureSource>
-static FShader* SetCaptureToTargetShaders(FRHICommandListImmediate& RHICmdList, FGraphicsPipelineStateInitializer& GraphicsPSOInit, FViewInfo& View, const FIntPoint& SourceTexSize, FTextureRHIParamRef SourceTextureRHI)
+static FShader* SetCaptureToTargetShaders(FRHICommandListImmediate& RHICmdList, FGraphicsPipelineStateInitializer& GraphicsPSOInit, FViewInfo& View, const FIntPoint& SourceTexSize, FRHITexture* SourceTextureRHI)
 {
 	TShaderMapRef<FMobileSceneCaptureCopyVS<bDemosaic>> VertexShader(View.ShaderMap);
 	TShaderMapRef<FMobileSceneCaptureCopyPS<bDemosaic, CaptureSource>> PixelShader(View.ShaderMap);
@@ -164,7 +164,7 @@ static FShader* SetCaptureToTargetShaders(FRHICommandListImmediate& RHICmdList,
 }
 
 template <bool bDemosaic>
-static FShader* SetCaptureToTargetShaders(FRHICommandListImmediate& RHICmdList, FGraphicsPipelineStateInitializer& GraphicsPSOInit, ESceneCaptureSource CaptureSource, FViewInfo& View, const FIntPoint& SourceTexSize, FTextureRHIParamRef SourceTextureRHI)
+static FShader* SetCaptureToTargetShaders(FRHICommandListImmediate& RHICmdList, FGraphicsPipelineStateInitializer& GraphicsPSOInit, ESceneCaptureSource CaptureSource, FViewInfo& View, const FIntPoint& SourceTexSize, FRHITexture* SourceTextureRHI)
 {
 	switch (CaptureSource)
 	{
@@ -192,7 +192,7 @@ static void CopyCaptureToTarget(
 	const FIntPoint& TargetSize, 
 	FViewInfo& View, 
 	const FIntRect& ViewRect, 
-	FTexture2DRHIParamRef SourceTextureRHI, 
+	FRHITexture2D* SourceTextureRHI,
 	bool bNeedsFlippedRenderTarget,
 	FSceneRenderer* SceneRenderer)
 {
diff --git a/Engine/Source/Runtime/Renderer/Private/MobileShadingRenderer.cpp b/Engine/Source/Runtime/Renderer/Private/MobileShadingRenderer.cpp
--- a/Engine/Source/Runtime/Renderer/Private/MobileShadingRenderer.cpp
+++ b/Engine/Source/Runtime/Renderer/Private/MobileShadingRenderer.cpp
@@ -416,9 +416,9 @@ void FMobileSceneRenderer::Render(FRHICommandListImmediate& RHICmdList)
 		(bForceDepthResolve || bSeparateTranslucencyActive || (View.bIsSceneCapture && (ViewFamily.SceneCaptureSource == ESceneCaptureSource::SCS_SceneColorHDR || ViewFamily.SceneCaptureSource == ESceneCaptureSource::SCS_SceneColorSceneDepth)));
 
 	//
-	FTextureRHIParamRef SceneColor = nullptr;
-	FTextureRHIParamRef SceneColorResolve = nullptr;
-	FTextureRHIParamRef SceneDepth = nullptr;
+	FRHITexture* SceneColor = nullptr;
+	FRHITexture* SceneColorResolve = nullptr;
+	FRHITexture* SceneDepth = nullptr;
 	ERenderTargetActions ColorTargetAction = ERenderTargetActions::Clear_Store;
 	EDepthStencilTargetActions DepthTargetAction = EDepthStencilTargetActions::ClearDepthStencil_DontStoreDepthStencil;
 	bool bMobileMSAA = SceneContext.GetSceneColorSurface()->GetNumSamples() > 1;
@@ -950,7 +950,7 @@ public:
 
 	FCopyMobileMultiViewSceneColorPS() {}
 
-	void SetParameters(FRHICommandList& RHICmdList, const FUniformBufferRHIParamRef ViewUniformBuffer, FTextureRHIRef InMobileMultiViewSceneColorTexture)
+	void SetParameters(FRHICommandList& RHICmdList, FRHIUniformBuffer* ViewUniformBuffer, FTextureRHIRef InMobileMultiViewSceneColorTexture)
 	{
 		FRHIPixelShader* ShaderRHI = GetPixelShader();
 		FGlobalShader::SetParameters<FViewUniformShaderParameters>(RHICmdList, ShaderRHI, ViewUniformBuffer);
diff --git a/Engine/Source/Runtime/Renderer/Private/PathTracing.cpp b/Engine/Source/Runtime/Renderer/Private/PathTracing.cpp
--- a/Engine/Source/Runtime/Renderer/Private/PathTracing.cpp
+++ b/Engine/Source/Runtime/Renderer/Private/PathTracing.cpp
@@ -113,8 +113,8 @@ public:
 		const FViewInfo& View,
 		FRayTracingShaderBindingsWriter& GlobalResources,
 		const FRayTracingScene& RayTracingScene,
-		FUniformBufferRHIParamRef ViewUniformBuffer,
-		FUniformBufferRHIParamRef SceneTexturesUniformBuffer,
+		FRHIUniformBuffer* ViewUniformBuffer,
+		FRHIUniformBuffer* SceneTexturesUniformBuffer,
 		// Light buffer
 		const TSparseArray<FLightSceneInfoCompact>& Lights,
 		// Adaptive sampling
@@ -348,13 +348,13 @@ class FPathTracingCompositorPS : public FGlobalShader
 	void SetParameters(
 		TRHICommandList& RHICmdList,
 		const FViewInfo& View,
-		FTextureRHIParamRef RadianceRedRT,
-		FTextureRHIParamRef RadianceGreenRT,
-		FTextureRHIParamRef RadianceBlueRT,
-		FTextureRHIParamRef RadianceAlphaRT,
-		FTextureRHIParamRef SampleCountRT,
-		FTextureRHIParamRef CumulativeIrradianceRT,
-		FTextureRHIParamRef CumulativeSampleCountRT)
+		FRHITexture* RadianceRedRT,
+		FRHITexture* RadianceGreenRT,
+		FRHITexture* RadianceBlueRT,
+		FRHITexture* RadianceAlphaRT,
+		FRHITexture* SampleCountRT,
+		FRHITexture* CumulativeIrradianceRT,
+		FRHITexture* CumulativeSampleCountRT)
 	{
 		FRHIPixelShader* ShaderRHI = GetPixelShader();
 		FGlobalShader::SetParameters<FViewUniformShaderParameters>(RHICmdList, ShaderRHI, View.ViewUniformBuffer);
@@ -448,7 +448,7 @@ void FDeferredShadingSceneRenderer::RenderPathTracing(FRHICommandListImmediate&
 
 	FSceneTexturesUniformParameters SceneTextures;
 	SetupSceneTextureUniformParameters(SceneContext, FeatureLevel, ESceneTextureSetupMode::All, SceneTextures);
-	FUniformBufferRHIParamRef SceneTexturesUniformBuffer = RHICreateUniformBuffer(&SceneTextures, FSceneTexturesUniformParameters::StaticStructMetadata.GetLayout(), EUniformBufferUsage::UniformBuffer_SingleDraw);
+	FRHIUniformBuffer* SceneTexturesUniformBuffer = RHICreateUniformBuffer(&SceneTextures, FSceneTexturesUniformParameters::StaticStructMetadata.GetLayout(), EUniformBufferUsage::UniformBuffer_SingleDraw);
 
 	RayGenShader->SetParameters(
 		Scene,
@@ -538,7 +538,7 @@ void FDeferredShadingSceneRenderer::RenderPathTracing(FRHICommandListImmediate&
 
 	TShaderMapRef<FPostProcessVS> VertexShader(ShaderMap);
 	TShaderMapRef<FPathTracingCompositorPS> PixelShader(ShaderMap);
-	FTextureRHIParamRef RenderTargets[3] =
+	FRHITexture* RenderTargets[3] =
 	{
 		SceneContext.GetSceneColor()->GetRenderTargetItem().TargetableTexture,
 		OutputRadianceRT->GetRenderTargetItem().TargetableTexture,
diff --git a/Engine/Source/Runtime/Renderer/Private/PathTracing/PathCompactionCompute.cpp b/Engine/Source/Runtime/Renderer/Private/PathTracing/PathCompactionCompute.cpp
--- a/Engine/Source/Runtime/Renderer/Private/PathTracing/PathCompactionCompute.cpp
+++ b/Engine/Source/Runtime/Renderer/Private/PathTracing/PathCompactionCompute.cpp
@@ -55,9 +55,9 @@ public:
 	void SetParameters(
 		FRHICommandListImmediate& RHICmdList,
 		const FViewInfo& View,
-		FTextureRHIParamRef RadianceTexture,
-		FTextureRHIParamRef SampleCountTexture,
-		FTextureRHIParamRef PixelPositionTexture,
+		FRHITexture* RadianceTexture,
+		FRHITexture* SampleCountTexture,
+		FRHITexture* PixelPositionTexture,
 		FRHIUnorderedAccessView* RadianceSortedRedUAV,
 		FRHIUnorderedAccessView* RadianceSortedGreenUAV,
 		FRHIUnorderedAccessView* RadianceSortedBlueUAV,
@@ -143,9 +143,9 @@ IMPLEMENT_SHADER_TYPE(, FPathCompactionCS, TEXT("/Engine/Private/PathTracing/Pat
 void FDeferredShadingSceneRenderer::ComputePathCompaction(
 	FRHICommandListImmediate& RHICmdList,
 	const FViewInfo& View,
-	FTextureRHIParamRef RadianceTexture,
-	FTextureRHIParamRef SampleCountTexture,
-	FTextureRHIParamRef PixelPositionTexture,
+	FRHITexture* RadianceTexture,
+	FRHITexture* SampleCountTexture,
+	FRHITexture* PixelPositionTexture,
 	FRHIUnorderedAccessView* RadianceSortedRedUAV,
 	FRHIUnorderedAccessView* RadianceSortedGreenUAV,
 	FRHIUnorderedAccessView* RadianceSortedBlueUAV,
diff --git a/Engine/Source/Runtime/Renderer/Private/PathTracing/RayCounterCompute.cpp b/Engine/Source/Runtime/Renderer/Private/PathTracing/RayCounterCompute.cpp
--- a/Engine/Source/Runtime/Renderer/Private/PathTracing/RayCounterCompute.cpp
+++ b/Engine/Source/Runtime/Renderer/Private/PathTracing/RayCounterCompute.cpp
@@ -49,7 +49,7 @@ public:
 
 	void SetParameters(
 		FRHICommandList& RHICmdList,
-		FTextureRHIParamRef RayCountPerPixelBuffer,
+		FRHITexture* RayCountPerPixelBuffer,
 		const FIntPoint& ViewSize,
 		FRHIUnorderedAccessView* TotalRayCountBuffer)
 	{
@@ -93,7 +93,7 @@ private:
 IMPLEMENT_SHADER_TYPE(, FRayCounterCS, TEXT("/Engine/Private/PathTracing/PathTracingRayCounterComputeShader.usf"), TEXT("RayCounterCS"), SF_Compute)
 
 
-void FDeferredShadingSceneRenderer::ComputeRayCount(FRHICommandListImmediate& RHICmdList, const FViewInfo& View, FTextureRHIParamRef RayCountPerPixelTexture)
+void FDeferredShadingSceneRenderer::ComputeRayCount(FRHICommandListImmediate& RHICmdList, const FViewInfo& View, FRHITexture* RayCountPerPixelTexture)
 {
 	FSceneViewState* ViewState = (FSceneViewState*)View.State;
 	ClearUAV(RHICmdList, *ViewState->TotalRayCountBuffer, 0);
diff --git a/Engine/Source/Runtime/Renderer/Private/PathTracing/SkyLightCdfBuild.cpp b/Engine/Source/Runtime/Renderer/Private/PathTracing/SkyLightCdfBuild.cpp
--- a/Engine/Source/Runtime/Renderer/Private/PathTracing/SkyLightCdfBuild.cpp
+++ b/Engine/Source/Runtime/Renderer/Private/PathTracing/SkyLightCdfBuild.cpp
@@ -433,7 +433,7 @@ void FDeferredShadingSceneRenderer::VisualizeSkyLightCdf(FRHICommandListImmediat
 	const auto ShaderMap = GetGlobalShaderMap(FeatureLevel);
 	TShaderMapRef<FPostProcessVS> VertexShader(ShaderMap);
 	TShaderMapRef<FVisualizeCdfPS> PixelShader(ShaderMap);
-	FTextureRHIParamRef RenderTargets[1] =
+	FRHITexture* RenderTargets[1] =
 	{
 		OutputRT->GetRenderTargetItem().TargetableTexture
 	};
diff --git a/Engine/Source/Runtime/Renderer/Private/PathTracing/VarianceMipTreeBuild.cpp b/Engine/Source/Runtime/Renderer/Private/PathTracing/VarianceMipTreeBuild.cpp
--- a/Engine/Source/Runtime/Renderer/Private/PathTracing/VarianceMipTreeBuild.cpp
+++ b/Engine/Source/Runtime/Renderer/Private/PathTracing/VarianceMipTreeBuild.cpp
@@ -208,7 +208,7 @@ void FDeferredShadingSceneRenderer::VisualizeVarianceMipTree(FRHICommandListImme
 	const auto ShaderMap = GetGlobalShaderMap(FeatureLevel);
 	TShaderMapRef<FPostProcessVS> VertexShader(ShaderMap);
 	TShaderMapRef<FVisualizeMipTreePS> PixelShader(ShaderMap);
-	FTextureRHIParamRef RenderTargets[2] =
+	FRHITexture* RenderTargets[2] =
 	{
 		SceneContext.GetSceneColor()->GetRenderTargetItem().TargetableTexture, 
 		VarianceMipTreeRT->GetRenderTargetItem().TargetableTexture
diff --git a/Engine/Source/Runtime/Renderer/Private/PlanarReflectionRendering.cpp b/Engine/Source/Runtime/Renderer/Private/PlanarReflectionRendering.cpp
--- a/Engine/Source/Runtime/Renderer/Private/PlanarReflectionRendering.cpp
+++ b/Engine/Source/Runtime/Renderer/Private/PlanarReflectionRendering.cpp
@@ -171,7 +171,7 @@ public:
 		SceneTextureParameters.Bind(Initializer);
 	}
 
-	void SetParameters(FRHICommandList& RHICmdList, const FSceneView& View, const FPlanarReflectionSceneProxy* ReflectionSceneProxy, FTextureRHIParamRef SceneColorInput, int32 FilterWidth)
+	void SetParameters(FRHICommandList& RHICmdList, const FSceneView& View, const FPlanarReflectionSceneProxy* ReflectionSceneProxy, FRHITexture* SceneColorInput, int32 FilterWidth)
 	{
 		FRHIPixelShader* ShaderRHI = GetPixelShader();
 		FGlobalShader::SetParameters<FViewUniformShaderParameters>(RHICmdList, ShaderRHI, View.ViewUniformBuffer);
@@ -216,7 +216,7 @@ IMPLEMENT_SHADER_TYPE(template<>, FPrefilterPlanarReflectionPS<true>, TEXT("/Eng
 template<bool bEnablePlanarReflectionPrefilter>
 void PrefilterPlanarReflection(FRHICommandListImmediate& RHICmdList, FViewInfo& View, const FPlanarReflectionSceneProxy* ReflectionSceneProxy, const FRenderTarget* Target)
 {
-	FTextureRHIParamRef SceneColorInput = FSceneRenderTargets::Get(RHICmdList).GetSceneColorTexture();
+	FRHITexture* SceneColorInput = FSceneRenderTargets::Get(RHICmdList).GetSceneColorTexture();
 
 	if(View.FeatureLevel >= ERHIFeatureLevel::SM4)
 	{
diff --git a/Engine/Source/Runtime/Renderer/Private/PostProcess/PostProcessBloomSetup.cpp b/Engine/Source/Runtime/Renderer/Private/PostProcess/PostProcessBloomSetup.cpp
--- a/Engine/Source/Runtime/Renderer/Private/PostProcess/PostProcessBloomSetup.cpp
+++ b/Engine/Source/Runtime/Renderer/Private/PostProcess/PostProcessBloomSetup.cpp
@@ -178,7 +178,7 @@ public:
 	}
 
 	template <typename TRHICmdList>
-	void SetParameters(TRHICmdList& RHICmdList, const FRenderingCompositePassContext& Context, const FIntPoint& DestSize, FRHIUnorderedAccessView* DestUAV, FTextureRHIParamRef EyeAdaptationTex)
+	void SetParameters(TRHICmdList& RHICmdList, const FRenderingCompositePassContext& Context, const FIntPoint& DestSize, FRHIUnorderedAccessView* DestUAV, FRHITexture* EyeAdaptationTex)
 	{
 		FRHIComputeShader* ShaderRHI = GetComputeShader();
 		const FPostProcessSettings& Settings = Context.View.FinalPostProcessSettings;
@@ -329,7 +329,7 @@ void FRCPassPostProcessBloomSetup::Process(FRenderingCompositePassContext& Conte
 }
 
 template <typename TRHICmdList>
-void FRCPassPostProcessBloomSetup::DispatchCS(TRHICmdList& RHICmdList, FRenderingCompositePassContext& Context, const FIntRect& DestRect, FRHIUnorderedAccessView* DestUAV, FTextureRHIParamRef EyeAdaptationTex)
+void FRCPassPostProcessBloomSetup::DispatchCS(TRHICmdList& RHICmdList, FRenderingCompositePassContext& Context, const FIntRect& DestRect, FRHIUnorderedAccessView* DestUAV, FRHITexture* EyeAdaptationTex)
 {
 	auto ShaderMap = Context.GetShaderMap();
 	TShaderMapRef<FPostProcessBloomSetupCS> ComputeShader(ShaderMap);
diff --git a/Engine/Source/Runtime/Renderer/Private/PostProcess/PostProcessBloomSetup.h b/Engine/Source/Runtime/Renderer/Private/PostProcess/PostProcessBloomSetup.h
--- a/Engine/Source/Runtime/Renderer/Private/PostProcess/PostProcessBloomSetup.h
+++ b/Engine/Source/Runtime/Renderer/Private/PostProcess/PostProcessBloomSetup.h
@@ -31,7 +31,7 @@ public:
 
 private:
 	template <typename TRHICmdList>
-	void DispatchCS(TRHICmdList& RHICmdList, FRenderingCompositePassContext& Context, const FIntRect& DestRect, FRHIUnorderedAccessView* DestUAV, FTextureRHIParamRef EyeAdaptationTex);
+	void DispatchCS(TRHICmdList& RHICmdList, FRenderingCompositePassContext& Context, const FIntRect& DestRect, FRHIUnorderedAccessView* DestUAV, FRHITexture* EyeAdaptationTex);
 
 	FComputeFenceRHIRef AsyncEndFence;
 };
diff --git a/Engine/Source/Runtime/Renderer/Private/PostProcess/PostProcessCompositeEditorPrimitives.cpp b/Engine/Source/Runtime/Renderer/Private/PostProcess/PostProcessCompositeEditorPrimitives.cpp
--- a/Engine/Source/Runtime/Renderer/Private/PostProcess/PostProcessCompositeEditorPrimitives.cpp
+++ b/Engine/Source/Runtime/Renderer/Private/PostProcess/PostProcessCompositeEditorPrimitives.cpp
@@ -561,7 +561,7 @@ void FRCPassPostProcessCompositeEditorPrimitives::Process(FRenderingCompositePas
 
 			TUniformBufferRef<FOpaqueBasePassUniformParameters> OpaqueBasePassUniformBuffer;
 			TUniformBufferRef<FMobileBasePassUniformParameters> MobileBasePassUniformBuffer;
-			FUniformBufferRHIParamRef BasePassUniformBuffer = nullptr;
+			FRHIUniformBuffer* BasePassUniformBuffer = nullptr;
 
 			if (bDeferredBasePass)
 			{
@@ -600,7 +600,7 @@ void FRCPassPostProcessCompositeEditorPrimitives::Process(FRenderingCompositePas
 
 		GVisualizeTexture.SetCheckPoint(Context.RHICmdList, SceneContext.EditorPrimitivesColor);
 
-		FTextureRHIParamRef EditorRenderTargets[2];
+		FRHITexture* EditorRenderTargets[2];
 		EditorRenderTargets[0] = EditorColorTarget;
 		EditorRenderTargets[1] = EditorDepthTarget;
 
diff --git a/Engine/Source/Runtime/Renderer/Private/PostProcess/PostProcessDOF.cpp b/Engine/Source/Runtime/Renderer/Private/PostProcess/PostProcessDOF.cpp
--- a/Engine/Source/Runtime/Renderer/Private/PostProcess/PostProcessDOF.cpp
+++ b/Engine/Source/Runtime/Renderer/Private/PostProcess/PostProcessDOF.cpp
@@ -172,7 +172,7 @@ void FRCPassPostProcessDOFSetup::Process(FRenderingCompositePassContext& Context
 	const FSceneRenderTargetItem& DestRenderTarget1 = (NumRenderTargets == 2) ? PassOutputs[1].RequestSurface(Context) : FSceneRenderTargetItem();
 
 	// Set the view family's render target/viewport.
-	FTextureRHIParamRef RenderTargets[2] =
+	FRHITexture* RenderTargets[2] =
 	{
 		DestRenderTarget0.TargetableTexture,
 		DestRenderTarget1.TargetableTexture
diff --git a/Engine/Source/Runtime/Renderer/Private/PostProcess/PostProcessLensBlur.cpp b/Engine/Source/Runtime/Renderer/Private/PostProcess/PostProcessLensBlur.cpp
--- a/Engine/Source/Runtime/Renderer/Private/PostProcess/PostProcessLensBlur.cpp
+++ b/Engine/Source/Runtime/Renderer/Private/PostProcess/PostProcessLensBlur.cpp
@@ -128,7 +128,7 @@ public:
 		PostprocessParameter.SetPS(RHICmdList, ShaderRHI, Context, TStaticSamplerState<SF_Bilinear, AM_Clamp, AM_Clamp, AM_Clamp>::GetRHI());
 				
 		{
-			FTextureRHIParamRef TextureRHI = GWhiteTexture->TextureRHI;
+			FRHITexture* TextureRHI = GWhiteTexture->TextureRHI;
 
 			if(GEngine->DefaultBokehTexture)
 			{
diff --git a/Engine/Source/Runtime/Renderer/Private/PostProcess/PostProcessMaterial.cpp b/Engine/Source/Runtime/Renderer/Private/PostProcess/PostProcessMaterial.cpp
--- a/Engine/Source/Runtime/Renderer/Private/PostProcess/PostProcessMaterial.cpp
+++ b/Engine/Source/Runtime/Renderer/Private/PostProcess/PostProcessMaterial.cpp
@@ -395,8 +395,8 @@ void FRCPassPostProcessMaterial::Process(FRenderingCompositePassContext& Context
 		{
 			DestRenderTarget = &PassOutputs[0].RequestSurface(Context);
 
-			FTextureRHIParamRef DstTexture = DestRenderTarget->TargetableTexture;
-			FTextureRHIParamRef SrcTexture = GetInput(ePId_Input0)->GetOutput()->RequestSurface(Context).ShaderResourceTexture;
+			FRHITexture* DstTexture = DestRenderTarget->TargetableTexture;
+			FRHITexture* SrcTexture = GetInput(ePId_Input0)->GetOutput()->RequestSurface(Context).ShaderResourceTexture;
 
 			// CopyResource() can only be called when format and size match. Otherwise must use shader to do stretch & format conversion.
 			if (DstTexture->GetFormat() == SrcTexture->GetFormat() && DstTexture->GetSizeXYZ() == SrcTexture->GetSizeXYZ())
diff --git a/Engine/Source/Runtime/Renderer/Private/PostProcess/PostProcessMorpheus.cpp b/Engine/Source/Runtime/Renderer/Private/PostProcess/PostProcessMorpheus.cpp
--- a/Engine/Source/Runtime/Renderer/Private/PostProcess/PostProcessMorpheus.cpp
+++ b/Engine/Source/Runtime/Renderer/Private/PostProcess/PostProcessMorpheus.cpp
@@ -175,7 +175,7 @@ class FPostProcessMorpheusVS : public FGlobalShader
 		FGlobalShader::SetParameters<FViewUniformShaderParameters>(Context.RHICmdList, GetVertexShader(), Context.View.ViewUniformBuffer);
 	}
 
-	void SetParameters(FRHICommandList& RHICmdList, const FUniformBufferRHIParamRef ViewUniformBuffer)
+	void SetParameters(FRHICommandList& RHICmdList, FRHIUniformBuffer* ViewUniformBuffer)
 	{
 		FGlobalShader::SetParameters<FViewUniformShaderParameters>(RHICmdList, GetVertexShader(), ViewUniformBuffer);
 	}
diff --git a/Engine/Source/Runtime/Renderer/Private/PostProcess/PostProcessMotionBlur.cpp b/Engine/Source/Runtime/Renderer/Private/PostProcess/PostProcessMotionBlur.cpp
--- a/Engine/Source/Runtime/Renderer/Private/PostProcess/PostProcessMotionBlur.cpp
+++ b/Engine/Source/Runtime/Renderer/Private/PostProcess/PostProcessMotionBlur.cpp
@@ -849,7 +849,7 @@ FRDGTextureRef VisualizeMotionBlur(
 		{
 			const FViewInfo& LocalView = ScreenPassView.View;
 			const FSceneViewFamily& ViewFamily = *LocalView.Family;
-			FRenderTargetTemp TempRenderTarget(static_cast<FTexture2DRHIParamRef>(ColorTextureOutput->GetRHI()), ColorViewport.Extent);
+			FRenderTargetTemp TempRenderTarget(static_cast<FRHITexture2D*>(ColorTextureOutput->GetRHI()), ColorViewport.Extent);
 			FCanvas Canvas(&TempRenderTarget, nullptr, ViewFamily.CurrentRealTime, ViewFamily.CurrentWorldTime, ViewFamily.DeltaWorldTime, LocalView.GetFeatureLevel());
 
 			float X = 20;
diff --git a/Engine/Source/Runtime/Renderer/Private/PostProcess/PostProcessSubsurface.cpp b/Engine/Source/Runtime/Renderer/Private/PostProcess/PostProcessSubsurface.cpp
--- a/Engine/Source/Runtime/Renderer/Private/PostProcess/PostProcessSubsurface.cpp
+++ b/Engine/Source/Runtime/Renderer/Private/PostProcess/PostProcessSubsurface.cpp
@@ -715,7 +715,7 @@ void VisualizeSubsurface(
 		{
 			const FViewInfo& LocalView = ScreenPassView.View;
 			const FSceneViewFamily& ViewFamily = *LocalView.Family;
-			FRenderTargetTemp TempRenderTarget(static_cast<FTexture2DRHIParamRef>(SceneTextureOutput->GetRHI()), SceneTextureOutput->Desc.Extent);
+			FRenderTargetTemp TempRenderTarget(static_cast<FRHITexture2D*>(SceneTextureOutput->GetRHI()), SceneTextureOutput->Desc.Extent);
 			FCanvas Canvas(&TempRenderTarget, nullptr, ViewFamily.CurrentRealTime, ViewFamily.CurrentWorldTime, ViewFamily.DeltaWorldTime, LocalView.GetFeatureLevel());
 
 			float X = 30;
diff --git a/Engine/Source/Runtime/Renderer/Private/PostProcess/PostProcessTonemap.cpp b/Engine/Source/Runtime/Renderer/Private/PostProcess/PostProcessTonemap.cpp
--- a/Engine/Source/Runtime/Renderer/Private/PostProcess/PostProcessTonemap.cpp
+++ b/Engine/Source/Runtime/Renderer/Private/PostProcess/PostProcessTonemap.cpp
@@ -830,7 +830,7 @@ public:
 	}
 
 	template <typename TRHICmdList>
-	void SetParameters(TRHICmdList& RHICmdList, const FRenderingCompositePassContext& Context, const FIntPoint& DestSize, FRHIUnorderedAccessView* DestUAV, FTextureRHIParamRef EyeAdaptationTex)
+	void SetParameters(TRHICmdList& RHICmdList, const FRenderingCompositePassContext& Context, const FIntPoint& DestSize, FRHIUnorderedAccessView* DestUAV, FRHITexture* EyeAdaptationTex)
 	{
 		FRHIComputeShader* ShaderRHI = GetComputeShader();
 		const FPostProcessSettings& Settings = Context.View.FinalPostProcessSettings;
@@ -934,7 +934,7 @@ static inline void ShaderTransitionResources(const FRenderingCompositePassContex
 } // PostProcessTonemapUtil
 
 template <typename TRHICmdList>
-inline void DispatchComputeShader(TRHICmdList& RHICmdList, FRenderingCompositePassContext& Context, const FIntRect& DestRect, FRHIUnorderedAccessView* DestUAV, const TonemapperPermutation::FDesktopDomain& DesktopPermutationVector, FTextureRHIParamRef EyeAdaptationTex, bool bDoEyeAdaptation)
+inline void DispatchComputeShader(TRHICmdList& RHICmdList, FRenderingCompositePassContext& Context, const FIntRect& DestRect, FRHIUnorderedAccessView* DestUAV, const TonemapperPermutation::FDesktopDomain& DesktopPermutationVector, FRHITexture* EyeAdaptationTex, bool bDoEyeAdaptation)
 {
 	FPostProcessTonemapCS::FPermutationDomain PermutationVector;
 	PermutationVector.Set<TonemapperPermutation::FDesktopDomain>(DesktopPermutationVector);
diff --git a/Engine/Source/Runtime/Renderer/Private/PostProcess/PostProcessTonemap.h b/Engine/Source/Runtime/Renderer/Private/PostProcess/PostProcessTonemap.h
--- a/Engine/Source/Runtime/Renderer/Private/PostProcess/PostProcessTonemap.h
+++ b/Engine/Source/Runtime/Renderer/Private/PostProcess/PostProcessTonemap.h
@@ -160,7 +160,7 @@ public:
 		if (Context.View.HasValidEyeAdaptation())
 		{
 			IPooledRenderTarget* EyeAdaptationRT = Context.View.GetEyeAdaptation(Context.RHICmdList);
-			FTextureRHIParamRef EyeAdaptationRTRef = EyeAdaptationRT->GetRenderTargetItem().TargetableTexture;
+			FRHITexture* EyeAdaptationRTRef = EyeAdaptationRT->GetRenderTargetItem().TargetableTexture;
 			if (EyeAdaptationRTRef)
 			{
 				Context.RHICmdList.TransitionResources(EResourceTransitionAccess::EReadable, &EyeAdaptationRTRef, 1);
@@ -192,7 +192,7 @@ public:
 		if (Context.View.HasValidEyeAdaptation())
 		{
 			IPooledRenderTarget* EyeAdaptationRT = Context.View.GetEyeAdaptation(Context.RHICmdList);
-			FTextureRHIParamRef EyeAdaptationRTRef = EyeAdaptationRT->GetRenderTargetItem().TargetableTexture;
+			FRHITexture* EyeAdaptationRTRef = EyeAdaptationRT->GetRenderTargetItem().TargetableTexture;
 			if (EyeAdaptationRTRef)
 			{
 				//Context.RHICmdList.TransitionResources(EResourceTransitionAccess::EReadable, &EyeAdaptationRTRef, 1);
diff --git a/Engine/Source/Runtime/Renderer/Private/PostProcess/PostProcessWeightedSampleSum.cpp b/Engine/Source/Runtime/Renderer/Private/PostProcess/PostProcessWeightedSampleSum.cpp
--- a/Engine/Source/Runtime/Renderer/Private/PostProcess/PostProcessWeightedSampleSum.cpp
+++ b/Engine/Source/Runtime/Renderer/Private/PostProcess/PostProcessWeightedSampleSum.cpp
@@ -137,7 +137,7 @@ public:
 
 	/** Sets shader parameter values */
 	void SetParameters(
-		FRHICommandList& RHICmdList, FRHISamplerState* SamplerStateRHI, FTextureRHIParamRef FilterTextureRHI, FTextureRHIParamRef AdditiveTextureRHI,
+		FRHICommandList& RHICmdList, FRHISamplerState* SamplerStateRHI, FRHITexture* FilterTextureRHI, FRHITexture* AdditiveTextureRHI,
 		const FLinearColor* SampleWeightValues, const FVector2D* SampleOffsetValues, uint32 NumSamples, const FVector4& FilteredBufferUVMinMaxValue, const FVector4& AdditiveBufferUVMinMaxValue )
 	{
 		check(CompileTimeNumSamples == 0 && NumSamples > 0 && NumSamples <= MAX_FILTER_SAMPLES || CompileTimeNumSamples == NumSamples);
@@ -400,7 +400,7 @@ public:
 
 	template <typename TRHICmdList>
 	void SetParameters(TRHICmdList& RHICmdList, const FRenderingCompositePassContext& Context, const FIntPoint& DestSize, FRHIUnorderedAccessView* DestUAV,
-		FTextureRHIParamRef FilterTextureRHI, FTextureRHIParamRef AdditiveTextureRHI, FLinearColor* SampleWeightValues, FVector2D* SampleOffsetValues, uint32 NumSamples)
+		FRHITexture* FilterTextureRHI, FRHITexture* AdditiveTextureRHI, FLinearColor* SampleWeightValues, FVector2D* SampleOffsetValues, uint32 NumSamples)
 	{
 		FRHIComputeShader* ShaderRHI = GetComputeShader();
 		const FPostProcessSettings& Settings = Context.View.FinalPostProcessSettings;
@@ -502,8 +502,8 @@ void SetFilterShaders(
 	FRHICommandListImmediate& RHICmdList,
 	ERHIFeatureLevel::Type FeatureLevel,
 	FRHISamplerState* SamplerStateRHI,
-	FTextureRHIParamRef FilterTextureRHI,
-	FTextureRHIParamRef AdditiveTextureRHI,
+	FRHITexture* FilterTextureRHI,
+	FRHITexture* AdditiveTextureRHI,
 	uint32 CombineMethodInt,
 	FVector2D* SampleOffsets,
 	FLinearColor* SampleWeights,
@@ -1083,7 +1083,7 @@ void FRCPassPostProcessWeightedSampleSum::Process(FRenderingCompositePassContext
 
 template <uint32 CompileTimeNumSamples, uint32 CombineMethodInt, typename TRHICmdList>
 void DispatchCSTemplate(TRHICmdList& RHICmdList, FRenderingCompositePassContext& Context, const FIntRect& DestRect, FRHIUnorderedAccessView* DestUAV,
-	FTextureRHIParamRef FilterTextureRHI, FTextureRHIParamRef AdditiveTextureRHI, FLinearColor* SampleWeightValues, FVector2D* SampleOffsetValues, uint32 NumSamples)
+	FRHITexture* FilterTextureRHI, FRHITexture* AdditiveTextureRHI, FLinearColor* SampleWeightValues, FVector2D* SampleOffsetValues, uint32 NumSamples)
 {
 	auto ShaderMap = Context.GetShaderMap();
 	TShaderMapRef<TFilterCS<CompileTimeNumSamples, CombineMethodInt, false>> ComputeShader(ShaderMap);
@@ -1103,7 +1103,7 @@ void DispatchCSTemplate(TRHICmdList& RHICmdList, FRenderingCompositePassContext&
 
 template <typename TRHICmdList>
 void FRCPassPostProcessWeightedSampleSum::DispatchCS(TRHICmdList& RHICmdList, FRenderingCompositePassContext& Context, const FIntRect& DestRect, FRHIUnorderedAccessView* DestUAV,
-	FTextureRHIParamRef FilterTextureRHI, FTextureRHIParamRef AdditiveTextureRHI, uint32 CombineMethodInt, FLinearColor* SampleWeightValues, FVector2D* SampleOffsetValues, uint32 NumSamples)
+	FRHITexture* FilterTextureRHI, FRHITexture* AdditiveTextureRHI, uint32 CombineMethodInt, FLinearColor* SampleWeightValues, FVector2D* SampleOffsetValues, uint32 NumSamples)
 {
 	// Helper macros
 #define DISPATCH_SHADER(_NumCompileSamples, _Method)												\
diff --git a/Engine/Source/Runtime/Renderer/Private/PostProcess/PostProcessWeightedSampleSum.h b/Engine/Source/Runtime/Renderer/Private/PostProcess/PostProcessWeightedSampleSum.h
--- a/Engine/Source/Runtime/Renderer/Private/PostProcess/PostProcessWeightedSampleSum.h
+++ b/Engine/Source/Runtime/Renderer/Private/PostProcess/PostProcessWeightedSampleSum.h
@@ -75,7 +75,7 @@ private:
 
 	template <typename TRHICmdList>
 	void DispatchCS(TRHICmdList& RHICmdList, FRenderingCompositePassContext& Context, const FIntRect& DestRect, FRHIUnorderedAccessView* DestUAV,
-		FTextureRHIParamRef FilterTextureRHI, FTextureRHIParamRef AdditiveTextureRHI, uint32 CombineMethodInt, FLinearColor* SampleWeightValues, FVector2D* SampleOffsetValues, uint32 NumSamples);
+		FRHITexture* FilterTextureRHI, FRHITexture* AdditiveTextureRHI, uint32 CombineMethodInt, FLinearColor* SampleWeightValues, FVector2D* SampleOffsetValues, uint32 NumSamples);
 
 	FComputeFenceRHIRef AsyncEndFence;
 };
diff --git a/Engine/Source/Runtime/Renderer/Private/PostProcess/PostProcessing.h b/Engine/Source/Runtime/Renderer/Private/PostProcess/PostProcessing.h
--- a/Engine/Source/Runtime/Renderer/Private/PostProcess/PostProcessing.h
+++ b/Engine/Source/Runtime/Renderer/Private/PostProcess/PostProcessing.h
@@ -21,7 +21,7 @@ public:
 	{}
 
 	void SetParameters(const FRenderingCompositePassContext&) {}
-	void SetParameters(FRHICommandList&, const FUniformBufferRHIParamRef) {}
+	void SetParameters(FRHICommandList&, FRHIUniformBuffer*) {}
 };
 
 /** The context used to setup a post-process pass. */
diff --git a/Engine/Source/Runtime/Renderer/Private/PostProcess/SceneRenderTargets.cpp b/Engine/Source/Runtime/Renderer/Private/PostProcess/SceneRenderTargets.cpp
--- a/Engine/Source/Runtime/Renderer/Private/PostProcess/SceneRenderTargets.cpp
+++ b/Engine/Source/Runtime/Renderer/Private/PostProcess/SceneRenderTargets.cpp
@@ -911,7 +911,7 @@ void FSceneRenderTargets::BeginRenderingGBuffer(FRHICommandList& RHICmdList, ERe
 	if (bShaderClear)
 	{
 		FLinearColor ClearColors[MaxSimultaneousRenderTargets];
-		FTextureRHIParamRef Textures[MaxSimultaneousRenderTargets];
+		FRHITexture* Textures[MaxSimultaneousRenderTargets];
 		ClearColors[0] = ClearColor;
 		Textures[0] = RPInfo.ColorRenderTargets[0].RenderTarget;
 		for (int32 i = 1; i < MRTCount; ++i)
@@ -2482,7 +2482,7 @@ void FSceneRenderTargets::ClearTranslucentVolumeLighting(FRHICommandListImmediat
 			IConsoleManager::Get().FindConsoleVariable(TEXT("r.TranslucencyVolumeBlur"));
 		static constexpr int32 Num3DTextures = NumTranslucentVolumeRenderTargetSets << 1;
 
-		FTextureRHIParamRef RenderTargets[Num3DTextures];
+		FRHITexture* RenderTargets[Num3DTextures];
 		bool bUseTransLightingVolBlur = CVarTranslucencyVolumeBlur->GetInt() > 0;
 		const int32 NumIterations = bUseTransLightingVolBlur ?
 			NumTranslucentVolumeRenderTargetSets : NumTranslucentVolumeRenderTargetSets - 1;
@@ -2508,7 +2508,7 @@ void FSceneRenderTargets::ClearTranslucentVolumeLighting(FRHICommandListImmediat
 
 /** Helper function that clears the given volume texture render targets. */
 template<int32 NumRenderTargets>
-void FSceneRenderTargets::ClearVolumeTextures(FRHICommandList& RHICmdList, ERHIFeatureLevel::Type FeatureLevel, FTextureRHIParamRef* RenderTargets, const FLinearColor* ClearColors)
+void FSceneRenderTargets::ClearVolumeTextures(FRHICommandList& RHICmdList, ERHIFeatureLevel::Type FeatureLevel, FRHITexture** RenderTargets, const FLinearColor* ClearColors)
 {
 	check(!RHICmdList.IsInsideRenderPass());
 
@@ -2552,16 +2552,16 @@ void FSceneRenderTargets::ClearVolumeTextures(FRHICommandList& RHICmdList, ERHIF
 	}
 	RHICmdList.EndRenderPass();
 
-	RHICmdList.TransitionResources(EResourceTransitionAccess::EReadable, (FTextureRHIParamRef*)RenderTargets, NumRenderTargets);
+	RHICmdList.TransitionResources(EResourceTransitionAccess::EReadable, (FRHITexture**)RenderTargets, NumRenderTargets);
 }
-template void FSceneRenderTargets::ClearVolumeTextures<1>(FRHICommandList& RHICmdList, ERHIFeatureLevel::Type FeatureLevel, FTextureRHIParamRef* RenderTargets, const FLinearColor* ClearColors);
-template void FSceneRenderTargets::ClearVolumeTextures<2>(FRHICommandList& RHICmdList, ERHIFeatureLevel::Type FeatureLevel, FTextureRHIParamRef* RenderTargets, const FLinearColor* ClearColors);
-template void FSceneRenderTargets::ClearVolumeTextures<3>(FRHICommandList& RHICmdList, ERHIFeatureLevel::Type FeatureLevel, FTextureRHIParamRef* RenderTargets, const FLinearColor* ClearColors);
-template void FSceneRenderTargets::ClearVolumeTextures<4>(FRHICommandList& RHICmdList, ERHIFeatureLevel::Type FeatureLevel, FTextureRHIParamRef* RenderTargets, const FLinearColor* ClearColors);
-template void FSceneRenderTargets::ClearVolumeTextures<5>(FRHICommandList& RHICmdList, ERHIFeatureLevel::Type FeatureLevel, FTextureRHIParamRef* RenderTargets, const FLinearColor* ClearColors);
-template void FSceneRenderTargets::ClearVolumeTextures<6>(FRHICommandList& RHICmdList, ERHIFeatureLevel::Type FeatureLevel, FTextureRHIParamRef* RenderTargets, const FLinearColor* ClearColors);
-template void FSceneRenderTargets::ClearVolumeTextures<7>(FRHICommandList& RHICmdList, ERHIFeatureLevel::Type FeatureLevel, FTextureRHIParamRef* RenderTargets, const FLinearColor* ClearColors);
-template void FSceneRenderTargets::ClearVolumeTextures<8>(FRHICommandList& RHICmdList, ERHIFeatureLevel::Type FeatureLevel, FTextureRHIParamRef* RenderTargets, const FLinearColor* ClearColors);
+template void FSceneRenderTargets::ClearVolumeTextures<1>(FRHICommandList& RHICmdList, ERHIFeatureLevel::Type FeatureLevel, FRHITexture** RenderTargets, const FLinearColor* ClearColors);
+template void FSceneRenderTargets::ClearVolumeTextures<2>(FRHICommandList& RHICmdList, ERHIFeatureLevel::Type FeatureLevel, FRHITexture** RenderTargets, const FLinearColor* ClearColors);
+template void FSceneRenderTargets::ClearVolumeTextures<3>(FRHICommandList& RHICmdList, ERHIFeatureLevel::Type FeatureLevel, FRHITexture** RenderTargets, const FLinearColor* ClearColors);
+template void FSceneRenderTargets::ClearVolumeTextures<4>(FRHICommandList& RHICmdList, ERHIFeatureLevel::Type FeatureLevel, FRHITexture** RenderTargets, const FLinearColor* ClearColors);
+template void FSceneRenderTargets::ClearVolumeTextures<5>(FRHICommandList& RHICmdList, ERHIFeatureLevel::Type FeatureLevel, FRHITexture** RenderTargets, const FLinearColor* ClearColors);
+template void FSceneRenderTargets::ClearVolumeTextures<6>(FRHICommandList& RHICmdList, ERHIFeatureLevel::Type FeatureLevel, FRHITexture** RenderTargets, const FLinearColor* ClearColors);
+template void FSceneRenderTargets::ClearVolumeTextures<7>(FRHICommandList& RHICmdList, ERHIFeatureLevel::Type FeatureLevel, FRHITexture** RenderTargets, const FLinearColor* ClearColors);
+template void FSceneRenderTargets::ClearVolumeTextures<8>(FRHICommandList& RHICmdList, ERHIFeatureLevel::Type FeatureLevel, FRHITexture** RenderTargets, const FLinearColor* ClearColors);
 
 EPixelFormat FSceneRenderTargets::GetSceneColorFormat() const
 {
@@ -2969,9 +2969,9 @@ void SetupSceneTextureUniformParameters(
 	ESceneTextureSetupMode SetupMode,
 	FSceneTexturesUniformParameters& SceneTextureParameters)
 {
-	FTextureRHIParamRef WhiteDefault2D = GSystemTextures.WhiteDummy->GetRenderTargetItem().ShaderResourceTexture;
-	FTextureRHIParamRef BlackDefault2D = GSystemTextures.BlackDummy->GetRenderTargetItem().ShaderResourceTexture;
-	FTextureRHIParamRef DepthDefault = GSystemTextures.DepthDummy->GetRenderTargetItem().ShaderResourceTexture;
+	FRHITexture* WhiteDefault2D = GSystemTextures.WhiteDummy->GetRenderTargetItem().ShaderResourceTexture;
+	FRHITexture* BlackDefault2D = GSystemTextures.BlackDummy->GetRenderTargetItem().ShaderResourceTexture;
+	FRHITexture* DepthDefault = GSystemTextures.DepthDummy->GetRenderTargetItem().ShaderResourceTexture;
 
 	// Scene Color / Depth
 	{
@@ -3054,7 +3054,7 @@ void SetupSceneTextureUniformParameters(
 	{
 		const bool bSetupCustomDepth = (SetupMode & ESceneTextureSetupMode::CustomDepth) != ESceneTextureSetupMode::None;
 
-		FTextureRHIParamRef CustomDepth = DepthDefault;
+		FRHITexture* CustomDepth = DepthDefault;
 		FRHIShaderResourceView* CustomStencilSRV = GSystemTextures.WhiteDummySRV;
 
 		// if there is no custom depth it's better to have the far distance there
@@ -3123,8 +3123,8 @@ void SetupMobileSceneTextureUniformParameters(
 	bool bSceneTexturesValid,
 	FMobileSceneTextureUniformParameters& SceneTextureParameters)
 {
-	FTextureRHIParamRef BlackDefault2D = GSystemTextures.BlackDummy->GetRenderTargetItem().ShaderResourceTexture;
-	FTextureRHIParamRef DepthDefault = GSystemTextures.DepthDummy->GetRenderTargetItem().ShaderResourceTexture;
+	FRHITexture* BlackDefault2D = GSystemTextures.BlackDummy->GetRenderTargetItem().ShaderResourceTexture;
+	FRHITexture* DepthDefault = GSystemTextures.DepthDummy->GetRenderTargetItem().ShaderResourceTexture;
 
 	SceneTextureParameters.SceneColorTexture = bSceneTexturesValid ? SceneContext.GetSceneColorTexture().GetReference() : BlackDefault2D;
 	SceneTextureParameters.SceneColorTextureSampler = TStaticSamplerState<SF_Point, AM_Clamp, AM_Clamp, AM_Clamp>::GetRHI();
@@ -3136,7 +3136,7 @@ void SetupMobileSceneTextureUniformParameters(
 	SceneTextureParameters.SceneAlphaCopyTexture = bSceneTexturesValid && SceneContext.HasSceneAlphaCopyTexture() ? SceneContext.GetSceneAlphaCopyTexture() : BlackDefault2D;
 	SceneTextureParameters.SceneAlphaCopyTextureSampler = TStaticSamplerState<SF_Point,AM_Clamp,AM_Clamp,AM_Clamp>::GetRHI();
 
-	FTextureRHIParamRef CustomDepth = DepthDefault;
+	FRHITexture* CustomDepth = DepthDefault;
 
 	// if there is no custom depth it's better to have the far distance there
 	IPooledRenderTarget* CustomDepthTarget = SceneContext.bCustomDepthIsValid ? SceneContext.CustomDepth.GetReference() : 0;
@@ -3148,7 +3148,7 @@ void SetupMobileSceneTextureUniformParameters(
 	SceneTextureParameters.CustomDepthTexture = CustomDepth;
 	SceneTextureParameters.CustomDepthTextureSampler = TStaticSamplerState<>::GetRHI();
 
-	FTextureRHIParamRef MobileCustomStencil = BlackDefault2D;
+	FRHITexture* MobileCustomStencil = BlackDefault2D;
 
 	if (SceneContext.MobileCustomStencil.IsValid())
 	{
diff --git a/Engine/Source/Runtime/Renderer/Private/PostProcess/SceneRenderTargets.h b/Engine/Source/Runtime/Renderer/Private/PostProcess/SceneRenderTargets.h
--- a/Engine/Source/Runtime/Renderer/Private/PostProcess/SceneRenderTargets.h
+++ b/Engine/Source/Runtime/Renderer/Private/PostProcess/SceneRenderTargets.h
@@ -403,7 +403,7 @@ public:
 	// ---
 
 	template<int32 NumRenderTargets>
-	static void ClearVolumeTextures(FRHICommandList& RHICmdList, ERHIFeatureLevel::Type FeatureLevel, FTextureRHIParamRef* RenderTargets, const FLinearColor* ClearColors);	
+	static void ClearVolumeTextures(FRHICommandList& RHICmdList, ERHIFeatureLevel::Type FeatureLevel, FRHITexture** RenderTargets, const FLinearColor* ClearColors);
 
 	void ClearTranslucentVolumeLighting(FRHICommandListImmediate& RHICmdList, int32 ViewIndex);
 
@@ -761,11 +761,11 @@ private:
 
 };
 
-extern template void FSceneRenderTargets::ClearVolumeTextures<0>(FRHICommandList& RHICmdList, ERHIFeatureLevel::Type FeatureLevel, FTextureRHIParamRef* RenderTargets, const FLinearColor* ClearColors);
-extern template void FSceneRenderTargets::ClearVolumeTextures<1>(FRHICommandList& RHICmdList, ERHIFeatureLevel::Type FeatureLevel, FTextureRHIParamRef* RenderTargets, const FLinearColor* ClearColors);
-extern template void FSceneRenderTargets::ClearVolumeTextures<2>(FRHICommandList& RHICmdList, ERHIFeatureLevel::Type FeatureLevel, FTextureRHIParamRef* RenderTargets, const FLinearColor* ClearColors);
-extern template void FSceneRenderTargets::ClearVolumeTextures<3>(FRHICommandList& RHICmdList, ERHIFeatureLevel::Type FeatureLevel, FTextureRHIParamRef* RenderTargets, const FLinearColor* ClearColors);
-extern template void FSceneRenderTargets::ClearVolumeTextures<4>(FRHICommandList& RHICmdList, ERHIFeatureLevel::Type FeatureLevel, FTextureRHIParamRef* RenderTargets, const FLinearColor* ClearColors);
-extern template void FSceneRenderTargets::ClearVolumeTextures<5>(FRHICommandList& RHICmdList, ERHIFeatureLevel::Type FeatureLevel, FTextureRHIParamRef* RenderTargets, const FLinearColor* ClearColors);
-extern template void FSceneRenderTargets::ClearVolumeTextures<6>(FRHICommandList& RHICmdList, ERHIFeatureLevel::Type FeatureLevel, FTextureRHIParamRef* RenderTargets, const FLinearColor* ClearColors);
-extern template void FSceneRenderTargets::ClearVolumeTextures<7>(FRHICommandList& RHICmdList, ERHIFeatureLevel::Type FeatureLevel, FTextureRHIParamRef* RenderTargets, const FLinearColor* ClearColors);
+extern template void FSceneRenderTargets::ClearVolumeTextures<0>(FRHICommandList& RHICmdList, ERHIFeatureLevel::Type FeatureLevel, FRHITexture** RenderTargets, const FLinearColor* ClearColors);
+extern template void FSceneRenderTargets::ClearVolumeTextures<1>(FRHICommandList& RHICmdList, ERHIFeatureLevel::Type FeatureLevel, FRHITexture** RenderTargets, const FLinearColor* ClearColors);
+extern template void FSceneRenderTargets::ClearVolumeTextures<2>(FRHICommandList& RHICmdList, ERHIFeatureLevel::Type FeatureLevel, FRHITexture** RenderTargets, const FLinearColor* ClearColors);
+extern template void FSceneRenderTargets::ClearVolumeTextures<3>(FRHICommandList& RHICmdList, ERHIFeatureLevel::Type FeatureLevel, FRHITexture** RenderTargets, const FLinearColor* ClearColors);
+extern template void FSceneRenderTargets::ClearVolumeTextures<4>(FRHICommandList& RHICmdList, ERHIFeatureLevel::Type FeatureLevel, FRHITexture** RenderTargets, const FLinearColor* ClearColors);
+extern template void FSceneRenderTargets::ClearVolumeTextures<5>(FRHICommandList& RHICmdList, ERHIFeatureLevel::Type FeatureLevel, FRHITexture** RenderTargets, const FLinearColor* ClearColors);
+extern template void FSceneRenderTargets::ClearVolumeTextures<6>(FRHICommandList& RHICmdList, ERHIFeatureLevel::Type FeatureLevel, FRHITexture** RenderTargets, const FLinearColor* ClearColors);
+extern template void FSceneRenderTargets::ClearVolumeTextures<7>(FRHICommandList& RHICmdList, ERHIFeatureLevel::Type FeatureLevel, FRHITexture** RenderTargets, const FLinearColor* ClearColors);
diff --git a/Engine/Source/Runtime/Renderer/Private/RayTracing/RayTracingLighting.cpp b/Engine/Source/Runtime/Renderer/Private/RayTracing/RayTracingLighting.cpp
--- a/Engine/Source/Runtime/Renderer/Private/RayTracing/RayTracingLighting.cpp
+++ b/Engine/Source/Runtime/Renderer/Private/RayTracing/RayTracingLighting.cpp
@@ -15,7 +15,7 @@ void SetupRaytracingLightDataPacked(
 	TResourceArray<FRTLightingData>& LightDataArray)
 {
 	TMap<UTextureLightProfile*, int32> IESLightProfilesMap;
-	TMap<FTextureRHIParamRef, uint32> RectTextureMap;
+	TMap<FRHITexture*, uint32> RectTextureMap;
 
 	LightData->Count = 0;
 	LightData->LTCMatTexture = GSystemTextures.LTCMat->GetRenderTargetItem().ShaderResourceTexture;
@@ -37,7 +37,7 @@ void SetupRaytracingLightDataPacked(
 
 	{
 		// IES profiles
-		FTextureRHIParamRef IESTextureRHI = nullptr;
+		FRHITexture* IESTextureRHI = nullptr;
 		float IESInvProfileCount = 1.0f;
 
 		if (View.IESLightProfileResource && View.IESLightProfileResource->GetIESLightProfilesCount())
diff --git a/Engine/Source/Runtime/Renderer/Private/RayTracing/RayTracingRectLight.cpp b/Engine/Source/Runtime/Renderer/Private/RayTracing/RayTracingRectLight.cpp
--- a/Engine/Source/Runtime/Renderer/Private/RayTracing/RayTracingRectLight.cpp
+++ b/Engine/Source/Runtime/Renderer/Private/RayTracing/RayTracingRectLight.cpp
@@ -259,9 +259,9 @@ public:
 	void Dispatch(
 		FRHICommandListImmediate& RHICmdList,
 		const FRayTracingScene& RayTracingScene,
-		FUniformBufferRHIParamRef ViewUniformBuffer,
-		FUniformBufferRHIParamRef SceneTexturesUniformBuffer,
-		FUniformBufferRHIParamRef RectLightUniformBuffer,
+		FRHIUniformBuffer* ViewUniformBuffer,
+		FRHIUniformBuffer* SceneTexturesUniformBuffer,
+		FRHIUniformBuffer* RectLightUniformBuffer,
 		FRHIUnorderedAccessView* LuminanceUAV,
 		FRHIUnorderedAccessView* RayDistanceUAV,
 		uint32 Width, uint32 Height
@@ -400,7 +400,7 @@ void FDeferredShadingSceneRenderer::VisualizeRectLightMipTree(
 	const auto ShaderMap = GetGlobalShaderMap(View.FeatureLevel);
 	TShaderMapRef<FPostProcessVS> VertexShader(ShaderMap);
 	TShaderMapRef<FVisualizeRectLightMipTreePS> PixelShader(ShaderMap);
-	FTextureRHIParamRef RenderTargets[2] =
+	FRHITexture* RenderTargets[2] =
 	{
 		SceneContext.GetSceneColor()->GetRenderTargetItem().TargetableTexture,
 		RectLightMipTreeRT->GetRenderTargetItem().TargetableTexture
diff --git a/Engine/Source/Runtime/Renderer/Private/RayTracing/RayTracingTranslucency.cpp b/Engine/Source/Runtime/Renderer/Private/RayTracing/RayTracingTranslucency.cpp
--- a/Engine/Source/Runtime/Renderer/Private/RayTracing/RayTracingTranslucency.cpp
+++ b/Engine/Source/Runtime/Renderer/Private/RayTracing/RayTracingTranslucency.cpp
@@ -218,8 +218,8 @@ public:
 	void SetParameters(
 		TRHICommandList& RHICmdList,
 		const FViewInfo& View,
-		FTextureRHIParamRef TranslucencyTexture,
-		FTextureRHIParamRef HitDistanceTexture)
+		FRHITexture* TranslucencyTexture,
+		FRHITexture* HitDistanceTexture)
 	{
 		FRHIPixelShader* ShaderRHI = GetPixelShader();
 		FGlobalShader::SetParameters<FViewUniformShaderParameters>(RHICmdList, ShaderRHI, View.ViewUniformBuffer);
diff --git a/Engine/Source/Runtime/Renderer/Private/ReflectionEnvironment.cpp b/Engine/Source/Runtime/Renderer/Private/ReflectionEnvironment.cpp
--- a/Engine/Source/Runtime/Renderer/Private/ReflectionEnvironment.cpp
+++ b/Engine/Source/Runtime/Renderer/Private/ReflectionEnvironment.cpp
@@ -262,7 +262,7 @@ void SetupReflectionUniformParameters(const FViewInfo& View, FReflectionUniformP
 	OutParameters.SkyLightCubemapBrightness = SkyAverageBrightness;
 
 	// Note: GBlackCubeArrayTexture has an alpha of 0, which is needed to represent invalid data so the sky cubemap can still be applied
-	FTextureRHIParamRef CubeArrayTexture = View.FeatureLevel >= ERHIFeatureLevel::SM5 ? GBlackCubeArrayTexture->TextureRHI : GBlackTextureCube->TextureRHI;
+	FRHITexture* CubeArrayTexture = View.FeatureLevel >= ERHIFeatureLevel::SM5 ? GBlackCubeArrayTexture->TextureRHI : GBlackTextureCube->TextureRHI;
 
 	if (View.Family->EngineShowFlags.ReflectionEnvironment 
 		&& View.FeatureLevel >= ERHIFeatureLevel::SM5
diff --git a/Engine/Source/Runtime/Renderer/Private/ReflectionEnvironmentCapture.cpp b/Engine/Source/Runtime/Renderer/Private/ReflectionEnvironmentCapture.cpp
--- a/Engine/Source/Runtime/Renderer/Private/ReflectionEnvironmentCapture.cpp
+++ b/Engine/Source/Runtime/Renderer/Private/ReflectionEnvironmentCapture.cpp
@@ -181,7 +181,7 @@ void CreateCubeMips( FRHICommandListImmediate& RHICmdList, ERHIFeatureLevel::Typ
 {	
 	SCOPED_DRAW_EVENT(RHICmdList, CreateCubeMips);
 
-	FTextureRHIParamRef CubeRef = Cubemap.TargetableTexture.GetReference();
+	FRHITexture* CubeRef = Cubemap.TargetableTexture.GetReference();
 
 	auto* ShaderMap = GetGlobalShaderMap(FeatureLevel);
 
diff --git a/Engine/Source/Runtime/Renderer/Private/RendererScene.cpp b/Engine/Source/Runtime/Renderer/Private/RendererScene.cpp
--- a/Engine/Source/Runtime/Renderer/Private/RendererScene.cpp
+++ b/Engine/Source/Runtime/Renderer/Private/RendererScene.cpp
@@ -3151,7 +3151,7 @@ void FScene::UpdateSpeedTreeWind(double CurrentTime)
 	#undef SET_SPEEDTREE_TABLE_FLOAT4V
 }
 
-FUniformBufferRHIParamRef FScene::GetSpeedTreeUniformBuffer(const FVertexFactory* VertexFactory) const
+FRHIUniformBuffer* FScene::GetSpeedTreeUniformBuffer(const FVertexFactory* VertexFactory) const
 {
 	if (VertexFactory != NULL)
 	{
@@ -3166,7 +3166,7 @@ FUniformBufferRHIParamRef FScene::GetSpeedTreeUniformBuffer(const FVertexFactory
 		}
 	}
 
-	return FUniformBufferRHIParamRef();
+	return nullptr;
 }
 
 /**
@@ -3666,7 +3666,7 @@ public:
 	virtual void AddSpeedTreeWind(class FVertexFactory* VertexFactory, const class UStaticMesh* StaticMesh) override {}
 	virtual void RemoveSpeedTreeWind_RenderThread(class FVertexFactory* VertexFactory, const class UStaticMesh* StaticMesh) override {}
 	virtual void UpdateSpeedTreeWind(double CurrentTime) override {}
-	virtual FUniformBufferRHIParamRef GetSpeedTreeUniformBuffer(const FVertexFactory* VertexFactory) const override { return FUniformBufferRHIParamRef(); }
+	virtual FRHIUniformBuffer* GetSpeedTreeUniformBuffer(const FVertexFactory* VertexFactory) const override { return nullptr; }
 
 	virtual void Release() override {}
 
diff --git a/Engine/Source/Runtime/Renderer/Private/SceneOcclusion.cpp b/Engine/Source/Runtime/Renderer/Private/SceneOcclusion.cpp
--- a/Engine/Source/Runtime/Renderer/Private/SceneOcclusion.cpp
+++ b/Engine/Source/Runtime/Renderer/Private/SceneOcclusion.cpp
@@ -352,14 +352,14 @@ void FOcclusionQueryBatcher::Flush(FRHICommandList& RHICmdList)
 		FMemMark MemStackMark(FMemStack::Get());
 
 		// Create the indices for MaxBatchedPrimitives boxes.
-		FIndexBufferRHIParamRef IndexBufferRHI = GOcclusionQueryIndexBuffer.IndexBufferRHI;
+		FRHIIndexBuffer* IndexBufferRHI = GOcclusionQueryIndexBuffer.IndexBufferRHI;
 
 		// Draw the batches.
 		for(int32 BatchIndex = 0, NumBatches = BatchOcclusionQueries.Num();BatchIndex < NumBatches;BatchIndex++)
 		{
 			FOcclusionBatch& Batch = BatchOcclusionQueries[BatchIndex];
 			FRHIRenderQuery* BatchOcclusionQuery = Batch.Query.GetQuery();
-			FVertexBufferRHIParamRef VertexBufferRHI = Batch.VertexAllocation.VertexBuffer->VertexBufferRHI;
+			FRHIVertexBuffer* VertexBufferRHI = Batch.VertexAllocation.VertexBuffer->VertexBufferRHI;
 			uint32 VertexBufferOffset = Batch.VertexAllocation.VertexOffset;
 			const int32 NumPrimitivesThisBatch = (BatchIndex != (NumBatches-1)) ? MaxBatchedPrimitives : NumBatchedPrimitives;
 				
@@ -818,7 +818,7 @@ public:
 		BoundsExtentSampler.Bind( Initializer.ParameterMap, TEXT("BoundsExtentSampler") );
 	}
 
-	void SetParameters(FRHICommandList& RHICmdList, const FViewInfo& View, FTextureRHIParamRef BoundsCenter, FTextureRHIParamRef BoundsExtent )
+	void SetParameters(FRHICommandList& RHICmdList, const FViewInfo& View, FRHITexture* BoundsCenter, FRHITexture* BoundsExtent )
 	{
 		FRHIPixelShader* ShaderRHI = GetPixelShader();
 
diff --git a/Engine/Source/Runtime/Renderer/Private/ScenePrivate.h b/Engine/Source/Runtime/Renderer/Private/ScenePrivate.h
--- a/Engine/Source/Runtime/Renderer/Private/ScenePrivate.h
+++ b/Engine/Source/Runtime/Renderer/Private/ScenePrivate.h
@@ -2689,7 +2689,7 @@ public:
 	virtual void AddSpeedTreeWind(FVertexFactory* VertexFactory, const UStaticMesh* StaticMesh) override;
 	virtual void RemoveSpeedTreeWind_RenderThread(FVertexFactory* VertexFactory, const UStaticMesh* StaticMesh) override;
 	virtual void UpdateSpeedTreeWind(double CurrentTime) override;
-	virtual FUniformBufferRHIParamRef GetSpeedTreeUniformBuffer(const FVertexFactory* VertexFactory) const override;
+	virtual FRHIUniformBuffer* GetSpeedTreeUniformBuffer(const FVertexFactory* VertexFactory) const override;
 	virtual void DumpUnbuiltLightInteractions( FOutputDevice& Ar ) const override;
 	virtual void UpdateParameterCollections(const TArray<FMaterialParameterCollectionInstanceResource*>& InParameterCollections) override;
 
@@ -2780,7 +2780,7 @@ public:
 	 **/
 	virtual void Export( FArchive& Ar ) const override;
 
-	FUniformBufferRHIParamRef GetParameterCollectionBuffer(const FGuid& InId) const
+	FRHIUniformBuffer* GetParameterCollectionBuffer(const FGuid& InId) const
 	{
 		const FUniformBufferRHIRef* ExistingUniformBuffer = ParameterCollections.Find(InId);
 
@@ -2789,7 +2789,7 @@ public:
 			return *ExistingUniformBuffer;
 		}
 
-		return FUniformBufferRHIParamRef();
+		return nullptr;
 	}
 
 	virtual void ApplyWorldOffset(FVector InOffset) override;
diff --git a/Engine/Source/Runtime/Renderer/Private/SceneRendering.cpp b/Engine/Source/Runtime/Renderer/Private/SceneRendering.cpp
--- a/Engine/Source/Runtime/Renderer/Private/SceneRendering.cpp
+++ b/Engine/Source/Runtime/Renderer/Private/SceneRendering.cpp
@@ -2833,7 +2833,7 @@ void FSceneRenderer::RenderCustomDepthPass(FRHICommandListImmediate& RHICmdList)
 
 			if (View.ShouldRenderView())
 			{
-				FUniformBufferRHIParamRef PassUniformBuffer = nullptr;
+				FRHIUniformBuffer* PassUniformBuffer = nullptr;
 
 				if (FSceneInterface::GetShadingPath(View.FeatureLevel) == EShadingPath::Mobile)
 				{
@@ -3896,7 +3896,7 @@ void FSceneRenderer::ResolveSceneColor(FRHICommandList& RHICmdList)
 	}
 }
 
-FTextureRHIParamRef FSceneRenderer::GetMultiViewSceneColor(const FSceneRenderTargets& SceneContext) const
+FRHITexture* FSceneRenderer::GetMultiViewSceneColor(const FSceneRenderTargets& SceneContext) const
 {
 	const FViewInfo& View = Views[0];
 
diff --git a/Engine/Source/Runtime/Renderer/Private/SceneRendering.h b/Engine/Source/Runtime/Renderer/Private/SceneRendering.h
--- a/Engine/Source/Runtime/Renderer/Private/SceneRendering.h
+++ b/Engine/Source/Runtime/Renderer/Private/SceneRendering.h
@@ -829,7 +829,7 @@ public:
 	FSceneBitArray PotentiallyFadingPrimitiveMap;
 
 	/** Primitive fade uniform buffers, indexed by packed primitive index. */
-	TArray<FUniformBufferRHIParamRef,SceneRenderingAllocator> PrimitiveFadeUniformBuffers;
+	TArray<FRHIUniformBuffer*,SceneRenderingAllocator> PrimitiveFadeUniformBuffers;
 
 	/**  Bit set when a primitive has a valid fade uniform buffer. */
 	FSceneBitArray PrimitiveFadeUniformBufferMap;
@@ -1662,7 +1662,7 @@ protected:
 	void RenderDistortion(FRHICommandListImmediate& RHICmdList);
 
 	/** Returns the scene color texture multi-view is targeting. */	
-	FTextureRHIParamRef GetMultiViewSceneColor(const FSceneRenderTargets& SceneContext) const;
+	FRHITexture* GetMultiViewSceneColor(const FSceneRenderTargets& SceneContext) const;
 
 	void UpdatePrimitiveIndirectLightingCacheBuffers();
 	void ClearPrimitiveSingleFrameIndirectLightingCacheBuffers();
@@ -1762,26 +1762,26 @@ private:
 // The noise textures need to be set in Slate too.
 RENDERER_API void UpdateNoiseTextureParameters(FViewUniformShaderParameters& ViewUniformShaderParameters);
 
-inline FTextureRHIParamRef OrBlack2DIfNull(FTextureRHIParamRef Tex)
+inline FRHITexture* OrBlack2DIfNull(FRHITexture* Tex)
 {
-	FTextureRHIParamRef Result = Tex ? Tex : GBlackTexture->TextureRHI.GetReference();
+	FRHITexture* Result = Tex ? Tex : GBlackTexture->TextureRHI.GetReference();
 	check(Result);
 	return Result;
 }
 
-inline FTextureRHIParamRef OrBlack3DIfNull(FTextureRHIParamRef Tex)
+inline FRHITexture* OrBlack3DIfNull(FRHITexture* Tex)
 {
 	// we fall back to 2D which are unbound es2 parameters
 	return OrBlack2DIfNull(Tex ? Tex : GBlackVolumeTexture->TextureRHI.GetReference());
 }
 
-inline FTextureRHIParamRef OrBlack3DUintIfNull(FTextureRHIParamRef Tex)
+inline FRHITexture* OrBlack3DUintIfNull(FRHITexture* Tex)
 {
 	// we fall back to 2D which are unbound es2 parameters
 	return OrBlack2DIfNull(Tex ? Tex : GBlackUintVolumeTexture->TextureRHI.GetReference());
 }
 
-inline void SetBlack2DIfNull(FTextureRHIParamRef& Tex)
+inline void SetBlack2DIfNull(FRHITexture*& Tex)
 {
 	if (!Tex)
 	{
@@ -1790,7 +1790,7 @@ inline void SetBlack2DIfNull(FTextureRHIParamRef& Tex)
 	}
 }
 
-inline void SetBlack3DIfNull(FTextureRHIParamRef& Tex)
+inline void SetBlack3DIfNull(FRHITexture*& Tex)
 {
 	if (!Tex)
 	{
diff --git a/Engine/Source/Runtime/Renderer/Private/SceneVisibility.cpp b/Engine/Source/Runtime/Renderer/Private/SceneVisibility.cpp
--- a/Engine/Source/Runtime/Renderer/Private/SceneVisibility.cpp
+++ b/Engine/Source/Runtime/Renderer/Private/SceneVisibility.cpp
@@ -351,7 +351,7 @@ bool FViewInfo::IsDistanceCulled( float DistanceSquared, float MinDrawDistance,
 		{
 			FPrimitiveFadingState& FadingState = ((FSceneViewState*)State)->PrimitiveFadingStates.FindOrAdd(PrimitiveSceneInfo->PrimitiveComponentId);
 			UpdatePrimitiveFadingState(FadingState, *this, !bDistanceCulled);
-			FUniformBufferRHIParamRef UniformBuffer = FadingState.UniformBuffer;
+			FRHIUniformBuffer* UniformBuffer = FadingState.UniformBuffer;
 			bStillFading = (UniformBuffer != NULL);
 			PrimitiveFadeUniformBuffers[PrimitiveIndex] = UniformBuffer;
 			PrimitiveFadeUniformBufferMap[PrimitiveIndex] = UniformBuffer != nullptr;
@@ -593,7 +593,7 @@ static void UpdatePrimitiveFading(const FScene* Scene, FViewInfo& View)
 				bool bVisible = View.PrimitiveVisibilityMap.AccessCorrespondingBit(BitIt);
 				FPrimitiveFadingState& FadingState = ViewState->PrimitiveFadingStates.FindOrAdd(Scene->PrimitiveComponentIds[BitIt.GetIndex()]);
 				UpdatePrimitiveFadingState(FadingState, View, bVisible);
-				FUniformBufferRHIParamRef UniformBuffer = FadingState.UniformBuffer;
+				FRHIUniformBuffer* UniformBuffer = FadingState.UniformBuffer;
 				if (UniformBuffer && !bVisible)
 				{
 					// If the primitive is fading out make sure it remains visible.
diff --git a/Engine/Source/Runtime/Renderer/Private/ShaderBaseClasses.cpp b/Engine/Source/Runtime/Renderer/Private/ShaderBaseClasses.cpp
--- a/Engine/Source/Runtime/Renderer/Private/ShaderBaseClasses.cpp
+++ b/Engine/Source/Runtime/Renderer/Private/ShaderBaseClasses.cpp
@@ -87,10 +87,10 @@ FMaterialShader::FMaterialShader(const FMaterialShaderType::CompiledShaderInitia
 	VTFeedbackBuffer.Bind(Initializer.ParameterMap, TEXT("VTFeedbackBuffer"));
 }
 
-FUniformBufferRHIParamRef FMaterialShader::GetParameterCollectionBuffer(const FGuid& Id, const FSceneInterface* SceneInterface) const
+FRHIUniformBuffer* FMaterialShader::GetParameterCollectionBuffer(const FGuid& Id, const FSceneInterface* SceneInterface) const
 {
 	const FScene* Scene = (const FScene*)SceneInterface;
-	FUniformBufferRHIParamRef UniformBuffer = Scene ? Scene->GetParameterCollectionBuffer(Id) : FUniformBufferRHIParamRef();
+	FRHIUniformBuffer* UniformBuffer = Scene ? Scene->GetParameterCollectionBuffer(Id) : nullptr;
 
 	if (!UniformBuffer)
 	{
@@ -277,7 +277,7 @@ void FMaterialShader::SetParametersInner(
 		// Find each referenced parameter collection's uniform buffer in the scene and set the parameter
 		for (int32 CollectionIndex = 0; CollectionIndex < NumToSet; CollectionIndex++)
 		{			
-			FUniformBufferRHIParamRef UniformBuffer = GetParameterCollectionBuffer(ParameterCollections[CollectionIndex], View.Family->Scene);
+			FRHIUniformBuffer* UniformBuffer = GetParameterCollectionBuffer(ParameterCollections[CollectionIndex], View.Family->Scene);
 
 			if (!UniformBuffer)
 			{
@@ -412,7 +412,7 @@ void FMaterialShader::GetShaderBindings(
 		// Find each referenced parameter collection's uniform buffer in the scene and set the parameter
 		for (int32 CollectionIndex = 0; CollectionIndex < NumToSet; CollectionIndex++)
 		{			
-			FUniformBufferRHIParamRef UniformBuffer = GetParameterCollectionBuffer(ParameterCollections[CollectionIndex], Scene);
+			FRHIUniformBuffer* UniformBuffer = GetParameterCollectionBuffer(ParameterCollections[CollectionIndex], Scene);
 
 			if (!UniformBuffer)
 			{
diff --git a/Engine/Source/Runtime/Renderer/Private/ShadowDepthRendering.cpp b/Engine/Source/Runtime/Renderer/Private/ShadowDepthRendering.cpp
--- a/Engine/Source/Runtime/Renderer/Private/ShadowDepthRendering.cpp
+++ b/Engine/Source/Runtime/Renderer/Private/ShadowDepthRendering.cpp
@@ -712,7 +712,7 @@ static void CheckShadowDepthMaterials(const FMaterialRenderProxy* InRenderProxy,
 	check(Material == InMaterial);
 }
 
-void FProjectedShadowInfo::ClearDepth(FRHICommandList& RHICmdList, class FSceneRenderer* SceneRenderer, int32 NumColorTextures, FTextureRHIParamRef* ColorTextures, FTextureRHIParamRef DepthTexture, bool bPerformClear)
+void FProjectedShadowInfo::ClearDepth(FRHICommandList& RHICmdList, class FSceneRenderer* SceneRenderer, int32 NumColorTextures, FRHITexture** ColorTextures, FRHITexture* DepthTexture, bool bPerformClear)
 {
 	check(RHICmdList.IsInsideRenderPass());
 
@@ -1096,7 +1096,7 @@ void FProjectedShadowInfo::SetupShadowUniformBuffers(FRHICommandListImmediate& R
 void FProjectedShadowInfo::RenderDepthInner(FRHICommandListImmediate& RHICmdList, FSceneRenderer* SceneRenderer, FBeginShadowRenderPassFunction BeginShadowRenderPass, bool bDoParallelDispatch)
 {
 	const ERHIFeatureLevel::Type FeatureLevel = ShadowDepthView->FeatureLevel;
-	FUniformBufferRHIParamRef PassUniformBuffer = ShadowDepthPassUniformBuffer;
+	FRHIUniformBuffer* PassUniformBuffer = ShadowDepthPassUniformBuffer;
 
 	const bool bIsWholeSceneDirectionalShadow = IsWholeSceneDirectionalShadow();
 
@@ -1490,7 +1490,7 @@ void FSceneRenderer::RenderShadowDepthMaps(FRHICommandListImmediate& RHICmdList)
 
 		auto BeginShadowRenderPass = [this, &RenderTarget, &SceneContext](FRHICommandList& InRHICmdList, bool bPerformClear)
 		{
-			FTextureRHIParamRef DepthTarget = RenderTarget.TargetableTexture;
+			FRHITexture* DepthTarget = RenderTarget.TargetableTexture;
 			ERenderTargetLoadAction DepthLoadAction = bPerformClear ? ERenderTargetLoadAction::EClear : ERenderTargetLoadAction::ELoad;
 
 			check(DepthTarget->GetDepthClearValue() == 1.0f);
@@ -1563,7 +1563,7 @@ void FSceneRenderer::RenderShadowDepthMaps(FRHICommandListImmediate& RHICmdList)
 
 				auto BeginShadowRenderPass = [this, ProjectedShadowInfo](FRHICommandList& InRHICmdList, bool bPerformClear)
 				{
-					FTextureRHIParamRef PreShadowCacheDepthZ = Scene->PreShadowCacheDepthZ->GetRenderTargetItem().TargetableTexture.GetReference();
+					FRHITexture* PreShadowCacheDepthZ = Scene->PreShadowCacheDepthZ->GetRenderTargetItem().TargetableTexture.GetReference();
 					InRHICmdList.TransitionResources(EResourceTransitionAccess::EWritable, &PreShadowCacheDepthZ, 1);
 
 					FRHIRenderPassInfo RPInfo(PreShadowCacheDepthZ, EDepthStencilTargetActions::LoadDepthStencil_StoreDepthStencil, nullptr, FExclusiveDepthStencil::DepthWrite_StencilWrite);
@@ -1605,7 +1605,7 @@ void FSceneRenderer::RenderShadowDepthMaps(FRHICommandListImmediate& RHICmdList)
 		FSceneRenderTargetItem ColorTarget0 = ShadowMapAtlas.RenderTargets.ColorTargets[0]->GetRenderTargetItem();
 		FSceneRenderTargetItem ColorTarget1 = ShadowMapAtlas.RenderTargets.ColorTargets[1]->GetRenderTargetItem();
 
-		FTextureRHIParamRef RenderTargetArray[2] =
+		FRHITexture* RenderTargetArray[2] =
 		{
 			ColorTarget0.TargetableTexture,
 			ColorTarget1.TargetableTexture
@@ -1673,7 +1673,7 @@ void FSceneRenderer::RenderShadowDepthMaps(FRHICommandListImmediate& RHICmdList)
 
 			auto BeginShadowRenderPass = [this, LightPropagationVolume, ProjectedShadowInfo, &ColorTarget0, &ColorTarget1, &DepthTarget](FRHICommandList& InRHICmdList, bool bPerformClear)
 			{
-				FTextureRHIParamRef RenderTargets[2];
+				FRHITexture* RenderTargets[2];
 				RenderTargets[0] = ColorTarget0.TargetableTexture;
 				RenderTargets[1] = ColorTarget1.TargetableTexture;
 
@@ -1862,7 +1862,7 @@ FShadowDepthPassMeshProcessor::FShadowDepthPassMeshProcessor(
 	const FScene* Scene,
 	const FSceneView* InViewIfDynamicMeshCommand,
 	const TUniformBufferRef<FViewUniformShaderParameters>& InViewUniformBuffer,
-	FUniformBufferRHIParamRef InPassUniformBuffer,
+	FRHIUniformBuffer* InPassUniformBuffer,
 	FShadowDepthType InShadowDepthType,
 	FMeshPassDrawListContext* InDrawListContext)
 	: FMeshPassProcessor(Scene, Scene->GetFeatureLevel(), InViewIfDynamicMeshCommand, InDrawListContext)
@@ -1876,7 +1876,7 @@ FShadowDepthType CSMShadowDepthType(true, false, false);
 
 FMeshPassProcessor* CreateCSMShadowDepthPassProcessor(const FScene* Scene, const FSceneView* InViewIfDynamicMeshCommand, FMeshPassDrawListContext* InDrawListContext)
 {
-	FUniformBufferRHIParamRef PassUniformBuffer = nullptr;
+	FRHIUniformBuffer* PassUniformBuffer = nullptr;
 
 	EShadingPath ShadingPath = Scene->GetShadingPath();
 	if (ShadingPath == EShadingPath::Mobile)
diff --git a/Engine/Source/Runtime/Renderer/Private/ShadowRendering.h b/Engine/Source/Runtime/Renderer/Private/ShadowRendering.h
--- a/Engine/Source/Runtime/Renderer/Private/ShadowRendering.h
+++ b/Engine/Source/Runtime/Renderer/Private/ShadowRendering.h
@@ -103,7 +103,7 @@ public:
 		const FScene* Scene, 
 		const FSceneView* InViewIfDynamicMeshCommand, 
 		const TUniformBufferRef<FViewUniformShaderParameters>& InViewUniformBuffer,
-		FUniformBufferRHIParamRef InPassUniformBuffer,
+		FRHIUniformBuffer* InPassUniformBuffer,
 		FShadowDepthType InShadowDepthType,
 		FMeshPassDrawListContext* InDrawListContext);
 
@@ -192,7 +192,7 @@ class FShadowMeshDrawCommandPass
 {
 public:
 	FMeshCommandOneFrameArray VisibleMeshDrawCommands;
-	FVertexBufferRHIParamRef PrimitiveIdVertexBuffer;
+	FRHIVertexBuffer* PrimitiveIdVertexBuffer;
 };
 
 /**
@@ -383,7 +383,7 @@ public:
 	/** Set state for depth rendering */
 	void SetStateForDepth(FMeshPassProcessorRenderState& DrawRenderState) const;
 
-	void ClearDepth(FRHICommandList& RHICmdList, class FSceneRenderer* SceneRenderer, int32 NumColorTextures, FTextureRHIParamRef* ColorTextures, FTextureRHIParamRef DepthTexture, bool bPerformClear);
+	void ClearDepth(FRHICommandList& RHICmdList, class FSceneRenderer* SceneRenderer, int32 NumColorTextures, FRHITexture** ColorTextures, FRHITexture* DepthTexture, bool bPerformClear);
 
 	/** Renders shadow maps for translucent primitives. */
 	void RenderTranslucencyDepths(FRHICommandList& RHICmdList, class FSceneRenderer* SceneRenderer);
@@ -435,7 +435,7 @@ public:
 	void GatherDynamicMeshElements(FSceneRenderer& Renderer, class FVisibleLightInfo& VisibleLightInfo, TArray<const FSceneView*>& ReusedViewsArray, 
 		FGlobalDynamicIndexBuffer& DynamicIndexBuffer, FGlobalDynamicVertexBuffer& DynamicVertexBuffer, FGlobalDynamicReadBuffer& DynamicReadBuffer);
 
-	void SetupMeshDrawCommandsForShadowDepth(FSceneRenderer& Renderer, FUniformBufferRHIParamRef PassUniformBuffer);
+	void SetupMeshDrawCommandsForShadowDepth(FSceneRenderer& Renderer, FRHIUniformBuffer* PassUniformBuffer);
 
 	void SetupMeshDrawCommandsForProjectionStenciling(FSceneRenderer& Renderer);
 
@@ -708,7 +708,7 @@ public:
 		return true;
 	}
 
-	void SetParameters(FRHICommandList& RHICmdList, const FUniformBufferRHIParamRef ViewUniformBuffer)
+	void SetParameters(FRHICommandList& RHICmdList, FRHIUniformBuffer* ViewUniformBuffer)
 	{
 		FGlobalShader::SetParameters<FViewUniformShaderParameters>(RHICmdList, GetVertexShader(), ViewUniformBuffer);
 	}
@@ -825,7 +825,7 @@ public:
 				FVector4(ShadowBufferSizeValue.X, ShadowBufferSizeValue.Y, 1.0f / ShadowBufferSizeValue.X, 1.0f / ShadowBufferSizeValue.Y));
 		}
 
-		FTextureRHIParamRef ShadowDepthTextureValue;
+		FRHITexture* ShadowDepthTextureValue;
 
 		// Translucency shadow projection has no depth target
 		if (ShadowInfo->RenderTargets.DepthTarget)
@@ -1173,7 +1173,7 @@ public:
 	template<typename ShaderRHIParamRef>
 	void Set(FRHICommandList& RHICmdList, const ShaderRHIParamRef ShaderRHI, const FProjectedShadowInfo* ShadowInfo) const
 	{
-		FTextureRHIParamRef ShadowDepthTextureValue = ShadowInfo 
+		FRHITexture* ShadowDepthTextureValue = ShadowInfo
 			? ShadowInfo->RenderTargets.DepthTarget->GetRenderTargetItem().ShaderResourceTexture->GetTextureCube()
 			: GBlackTextureDepthCube->TextureRHI.GetReference();
 		if (!ShadowDepthTextureValue)
diff --git a/Engine/Source/Runtime/Renderer/Private/ShadowSetup.cpp b/Engine/Source/Runtime/Renderer/Private/ShadowSetup.cpp
--- a/Engine/Source/Runtime/Renderer/Private/ShadowSetup.cpp
+++ b/Engine/Source/Runtime/Renderer/Private/ShadowSetup.cpp
@@ -1172,7 +1172,7 @@ void FProjectedShadowInfo::AddReceiverPrimitive(FPrimitiveSceneInfo* PrimitiveSc
 	ReceiverPrimitives.Add(PrimitiveSceneInfo);
 }
 
-void FProjectedShadowInfo::SetupMeshDrawCommandsForShadowDepth(FSceneRenderer& Renderer, FUniformBufferRHIParamRef PassUniformBuffer)
+void FProjectedShadowInfo::SetupMeshDrawCommandsForShadowDepth(FSceneRenderer& Renderer, FRHIUniformBuffer* PassUniformBuffer)
 {
 	QUICK_SCOPE_CYCLE_COUNTER(STAT_SetupMeshDrawCommandsForShadowDepth);
 
@@ -1425,7 +1425,7 @@ void FProjectedShadowInfo::GatherDynamicMeshElements(FSceneRenderer& Renderer, F
 	
 	// Create a pass uniform buffer so we can build mesh commands now in InitDynamicShadows.  This will be updated with the correct contents just before the actual pass.
 	const EShadingPath ShadingPath = FSceneInterface::GetShadingPath(Renderer.FeatureLevel);
-	FUniformBufferRHIParamRef PassUniformBuffer = nullptr;
+	FRHIUniformBuffer* PassUniformBuffer = nullptr;
 	if (ShadingPath == EShadingPath::Deferred)
 	{
 		FShadowDepthPassUniformParameters ShadowDepthParameters;
diff --git a/Engine/Source/Runtime/Renderer/Private/SurfelTree.cpp b/Engine/Source/Runtime/Renderer/Private/SurfelTree.cpp
--- a/Engine/Source/Runtime/Renderer/Private/SurfelTree.cpp
+++ b/Engine/Source/Runtime/Renderer/Private/SurfelTree.cpp
@@ -297,7 +297,7 @@ void GenerateSurfelRepresentation(FRHICommandListImmediate& RHICmdList, FSceneRe
 	{
 		FUniformMeshBuffers* UniformMeshBuffers = NULL;
 		const FMaterialRenderProxy* MaterialRenderProxy = NULL;
-		FUniformBufferRHIParamRef PrimitiveUniformBuffer = NULL;
+		FRHIUniformBuffer* PrimitiveUniformBuffer = NULL;
 		const int32 NumUniformTriangles = FUniformMeshConverter::Convert(RHICmdList, Renderer, View, PrimitiveSceneInfo, 0, UniformMeshBuffers, MaterialRenderProxy, PrimitiveUniformBuffer);
 
 		if (NumUniformTriangles > 0 && MaterialRenderProxy && PrimitiveUniformBuffer)
diff --git a/Engine/Source/Runtime/Renderer/Private/TranslucentLighting.cpp b/Engine/Source/Runtime/Renderer/Private/TranslucentLighting.cpp
--- a/Engine/Source/Runtime/Renderer/Private/TranslucentLighting.cpp
+++ b/Engine/Source/Runtime/Renderer/Private/TranslucentLighting.cpp
@@ -619,7 +619,7 @@ void FProjectedShadowInfo::RenderTranslucencyDepths(FRHICommandList& RHICmdList,
 		{
 			const bool bDynamicInstancing = IsDynamicInstancingEnabled(ShadowDepthView->FeatureLevel);
 
-			FVertexBufferRHIParamRef PrimitiveIdVertexBuffer = nullptr;
+			FRHIVertexBuffer* PrimitiveIdVertexBuffer = nullptr;
 			ApplyViewOverridesToMeshDrawCommands(*ShadowDepthView, VisibleMeshDrawCommands);
 			SortAndMergeDynamicPassMeshDrawCommands(SceneRenderer->FeatureLevel, VisibleMeshDrawCommands, DynamicMeshDrawCommandStorage, PrimitiveIdVertexBuffer, 1);
 			SubmitMeshDrawCommands(VisibleMeshDrawCommands, GraphicsMinimalPipelineStateSet, PrimitiveIdVertexBuffer, 0, bDynamicInstancing, 1, RHICmdList);
@@ -1175,7 +1175,7 @@ void FDeferredShadingSceneRenderer::ClearTranslucentVolumePerObjectShadowing(FRH
 		SCOPED_GPU_STAT(RHICmdList, TranslucentLighting);
 
 		static_assert(TVC_MAX == 2, "Only expecting two translucency lighting cascades.");
-		FTextureRHIParamRef RenderTargets[2];
+		FRHITexture* RenderTargets[2];
 		RenderTargets[0] = SceneContext.GetTranslucencyVolumeAmbient(TVC_Inner, ViewIndex)->GetRenderTargetItem().TargetableTexture;
 		RenderTargets[1] = SceneContext.GetTranslucencyVolumeDirectional(TVC_Inner, ViewIndex)->GetRenderTargetItem().TargetableTexture;
 
@@ -1241,7 +1241,7 @@ void FDeferredShadingSceneRenderer::AccumulateTranslucentVolumeObjectShadowing(F
 
 			if (VolumeBounds.IsValid())
 			{
-				FTextureRHIParamRef RenderTarget;
+				FRHITexture* RenderTarget;
 
 				if (VolumeCascadeIndex == 0)
 				{
@@ -1471,7 +1471,7 @@ static void InjectTranslucentLightArray(FRHICommandListImmediate& RHICmdList, co
 		GVisualizeTexture.SetCheckPoint(RHICmdList, RT0);
 		GVisualizeTexture.SetCheckPoint(RHICmdList, RT1);
 
-		FTextureRHIParamRef RenderTargets[2];
+		FRHITexture* RenderTargets[2];
 		RenderTargets[0] = RT0->GetRenderTargetItem().TargetableTexture;
 		RenderTargets[1] = RT1->GetRenderTargetItem().TargetableTexture;
 
@@ -1689,7 +1689,7 @@ void FDeferredShadingSceneRenderer::InjectSimpleTranslucentVolumeLightingArray(F
 			GVisualizeTexture.SetCheckPoint(RHICmdList, RT0);
 			GVisualizeTexture.SetCheckPoint(RHICmdList, RT1);
 
-			FTextureRHIParamRef RenderTargets[2];
+			FRHITexture* RenderTargets[2];
 			RenderTargets[0] = RT0->GetRenderTargetItem().TargetableTexture;
 			RenderTargets[1] = RT1->GetRenderTargetItem().TargetableTexture;
 
@@ -1801,11 +1801,11 @@ void FDeferredShadingSceneRenderer::FilterTranslucentVolumeLighting(FRHICommandL
 				GVisualizeTexture.SetCheckPoint(RHICmdList, RT0);
 				GVisualizeTexture.SetCheckPoint(RHICmdList, RT1);
 
-				FTextureRHIParamRef RenderTargets[2];
+				FRHITexture* RenderTargets[2];
 				RenderTargets[0] = RT0->GetRenderTargetItem().TargetableTexture;
 				RenderTargets[1] = RT1->GetRenderTargetItem().TargetableTexture;
 
-				FTextureRHIParamRef Inputs[2];
+				FRHITexture* Inputs[2];
 				Inputs[0] = Input0->GetRenderTargetItem().TargetableTexture;
 				Inputs[1] = Input1->GetRenderTargetItem().TargetableTexture;
 
diff --git a/Engine/Source/Runtime/Renderer/Private/VT/RuntimeVirtualTextureRender.cpp b/Engine/Source/Runtime/Renderer/Private/VT/RuntimeVirtualTextureRender.cpp
--- a/Engine/Source/Runtime/Renderer/Private/VT/RuntimeVirtualTextureRender.cpp
+++ b/Engine/Source/Runtime/Renderer/Private/VT/RuntimeVirtualTextureRender.cpp
@@ -350,7 +350,7 @@ namespace RuntimeVirtualTexture
 		// Sort and submit
 		if (AllocatedCommands.Num() > 0)
 		{
-			FVertexBufferRHIParamRef PrimitiveIdsBuffer;
+			FRHIVertexBuffer* PrimitiveIdsBuffer;
 			const bool bDynamicInstancing = IsDynamicInstancingEnabled(View->FeatureLevel);
 			const uint32 InstanceFactor = 1;
 
diff --git a/Engine/Source/Runtime/Renderer/Private/VisualizeVolumetricLightmap.cpp b/Engine/Source/Runtime/Renderer/Private/VisualizeVolumetricLightmap.cpp
--- a/Engine/Source/Runtime/Renderer/Private/VisualizeVolumetricLightmap.cpp
+++ b/Engine/Source/Runtime/Renderer/Private/VisualizeVolumetricLightmap.cpp
@@ -162,10 +162,10 @@ void FDeferredShadingSceneRenderer::VisualizeVolumetricLightmap(FRHICommandListI
 
 		int32 NumRenderTargets = 1;
 
-		FTextureRHIParamRef RenderTargets[2] =
+		FRHITexture* RenderTargets[2] =
 		{
 			SceneContext.GetSceneColorSurface(),
-			NULL,
+			nullptr,
 		};
 
 		if (SceneContext.GBufferB)
diff --git a/Engine/Source/Runtime/Renderer/Private/VolumeLighting.h b/Engine/Source/Runtime/Renderer/Private/VolumeLighting.h
--- a/Engine/Source/Runtime/Renderer/Private/VolumeLighting.h
+++ b/Engine/Source/Runtime/Renderer/Private/VolumeLighting.h
@@ -95,7 +95,7 @@ public:
 		{
 			SetShaderValue(RHICmdList, ShaderRHI, DepthBiasParameters, FVector4(ShadowMap->GetShaderDepthBias(), ShadowMap->GetShaderSlopeDepthBias(), ShadowMap->GetShaderMaxSlopeDepthBias(), 1.0f / (ShadowMap->MaxSubjectZ - ShadowMap->MinSubjectZ)));
 
-			FTextureRHIParamRef ShadowDepthTextureResource = nullptr;
+			FRHITexture* ShadowDepthTextureResource = nullptr;
 			if (LightType == LightType_Point || LightType == LightType_Rect)
 			{
 				if (GBlackTexture && GBlackTexture->TextureRHI)
@@ -122,7 +122,7 @@ public:
 
 		const FStaticShadowDepthMap* StaticShadowDepthMap = LightSceneInfo->Proxy->GetStaticShadowDepthMap();
 		const uint32 bStaticallyShadowedValue = LightSceneInfo->IsPrecomputedLightingValid() && StaticShadowDepthMap && StaticShadowDepthMap->Data && StaticShadowDepthMap->TextureRHI ? 1 : 0;
-		FTextureRHIParamRef StaticShadowDepthMapTexture = bStaticallyShadowedValue ? StaticShadowDepthMap->TextureRHI : GWhiteTexture->TextureRHI;
+		FRHITexture* StaticShadowDepthMapTexture = bStaticallyShadowedValue ? StaticShadowDepthMap->TextureRHI : GWhiteTexture->TextureRHI;
 		const FMatrix WorldToStaticShadow = bStaticallyShadowedValue ? StaticShadowDepthMap->Data->WorldToLight : FMatrix::Identity;
 		const FVector4 StaticShadowBufferSizeValue = bStaticallyShadowedValue ? FVector4(StaticShadowDepthMap->Data->ShadowMapSizeX, StaticShadowDepthMap->Data->ShadowMapSizeY, 1.0f / StaticShadowDepthMap->Data->ShadowMapSizeX, 1.0f / StaticShadowDepthMap->Data->ShadowMapSizeY) : FVector4(0, 0, 0, 0);
 
diff --git a/Engine/Source/Runtime/Renderer/Private/VolumetricFog.cpp b/Engine/Source/Runtime/Renderer/Private/VolumetricFog.cpp
--- a/Engine/Source/Runtime/Renderer/Private/VolumetricFog.cpp
+++ b/Engine/Source/Runtime/Renderer/Private/VolumetricFog.cpp
@@ -631,7 +631,7 @@ public:
 		const FViewInfo& View,
 		const FVolumetricFogIntegrationParameterData& IntegrationData,
 		const FExponentialHeightFogSceneInfo& FogInfo,
-		FTextureRHIParamRef LightScatteringHistoryTexture,
+		FRHITexture* LightScatteringHistoryTexture,
 		bool bUseDirectionalLightShadowing,
 		const FMatrix& DirectionalLightFunctionWorldToShadowValue)
 	{
@@ -1136,7 +1136,7 @@ void FDeferredShadingSceneRenderer::ComputeVolumetricFog(FRHICommandListImmediat
 
 					RHICmdList.SetComputeShader(ComputeShader->GetComputeShader());
 
-					FTextureRHIParamRef LightScatteringHistoryTexture = bUseTemporalReprojection && View.ViewState->LightScatteringHistory.IsValid()
+					FRHITexture* LightScatteringHistoryTexture = bUseTemporalReprojection && View.ViewState->LightScatteringHistory.IsValid()
 						? View.ViewState->LightScatteringHistory->GetRenderTargetItem().ShaderResourceTexture
 						: GBlackVolumeTexture->TextureRHI;
 
diff --git a/Engine/Source/Runtime/Renderer/Private/WideCustomResolveShaders.h b/Engine/Source/Runtime/Renderer/Private/WideCustomResolveShaders.h
--- a/Engine/Source/Runtime/Renderer/Private/WideCustomResolveShaders.h
+++ b/Engine/Source/Runtime/Renderer/Private/WideCustomResolveShaders.h
@@ -63,7 +63,7 @@ public:
 		return IsFeatureLevelSupported(Parameters.Platform, ERHIFeatureLevel::SM5);
 	}
 
-	void SetParameters(FRHICommandList& RHICmdList, FTextureRHIParamRef Texture2DMS, FTextureRHIParamRef FMaskTexture2D, FIntPoint Origin)
+	void SetParameters(FRHICommandList& RHICmdList, FRHITexture* Texture2DMS, FRHITexture* FMaskTexture2D, FIntPoint Origin)
 	{
 		FRHIPixelShader* PixelShaderRHI = GetPixelShader();
 		SetTextureParameter(RHICmdList, PixelShaderRHI, Tex, Texture2DMS);
diff --git a/Engine/Source/Runtime/Renderer/Public/GlobalDistanceFieldParameters.h b/Engine/Source/Runtime/Renderer/Public/GlobalDistanceFieldParameters.h
--- a/Engine/Source/Runtime/Renderer/Public/GlobalDistanceFieldParameters.h
+++ b/Engine/Source/Runtime/Renderer/Public/GlobalDistanceFieldParameters.h
@@ -27,7 +27,7 @@ public:
 
 	FVector4 CenterAndExtent[GMaxGlobalDistanceFieldClipmaps];
 	FVector4 WorldToUVAddAndMul[GMaxGlobalDistanceFieldClipmaps];
-	FTextureRHIParamRef Textures[GMaxGlobalDistanceFieldClipmaps];
+	FRHITexture* Textures[GMaxGlobalDistanceFieldClipmaps];
 	float GlobalDFResolution;
 	float MaxDistance;
 };
diff --git a/Engine/Source/Runtime/Renderer/Public/HdrCustomResolveShaders.h b/Engine/Source/Runtime/Renderer/Public/HdrCustomResolveShaders.h
--- a/Engine/Source/Runtime/Renderer/Public/HdrCustomResolveShaders.h
+++ b/Engine/Source/Runtime/Renderer/Public/HdrCustomResolveShaders.h
@@ -54,7 +54,7 @@ public:
 		return IsFeatureLevelSupported(Parameters.Platform, ERHIFeatureLevel::ES2);
 	}
 
-	void SetParameters(FRHICommandList& RHICmdList, FTextureRHIParamRef Texture2DMS)
+	void SetParameters(FRHICommandList& RHICmdList, FRHITexture* Texture2DMS)
 	{
 		FRHIPixelShader* PixelShaderRHI = GetPixelShader();
 		SetTextureParameter(RHICmdList, PixelShaderRHI, Tex, Texture2DMS);
@@ -93,7 +93,7 @@ public:
 		return IsFeatureLevelSupported(Parameters.Platform, ERHIFeatureLevel::ES2);
 	}
 
-	void SetParameters(FRHICommandList& RHICmdList, FTextureRHIParamRef Texture2DMS)
+	void SetParameters(FRHICommandList& RHICmdList, FRHITexture* Texture2DMS)
 	{
 		FRHIPixelShader* PixelShaderRHI = GetPixelShader();
 		SetTextureParameter(RHICmdList, PixelShaderRHI, Tex, Texture2DMS);
@@ -133,7 +133,7 @@ public:
 		return IsFeatureLevelSupported(Parameters.Platform, ERHIFeatureLevel::ES2);
 	}
 
-	void SetParameters(FRHICommandList& RHICmdList, FTextureRHIParamRef Texture2DMS)
+	void SetParameters(FRHICommandList& RHICmdList, FRHITexture* Texture2DMS)
 	{
 		FRHIPixelShader* PixelShaderRHI = GetPixelShader();
 		SetTextureParameter(RHICmdList, PixelShaderRHI, Tex, Texture2DMS);
@@ -174,7 +174,7 @@ public:
 		return IsFeatureLevelSupported(Parameters.Platform, ERHIFeatureLevel::ES2);
 	}	
 
-	void SetParameters(FRHICommandList& RHICmdList, FTextureRHIParamRef Texture2DMS, FTextureRHIParamRef FMaskTexure2D)
+	void SetParameters(FRHICommandList& RHICmdList, FRHITexture* Texture2DMS, FRHITexture* FMaskTexure2D)
 	{
 		FRHIPixelShader* PixelShaderRHI = GetPixelShader();
 		SetTextureParameter(RHICmdList, PixelShaderRHI, Tex, Texture2DMS);
@@ -218,7 +218,7 @@ public:
 		return IsFeatureLevelSupported(Parameters.Platform, ERHIFeatureLevel::ES2);
 	}
 
-	void SetParameters(FRHICommandList& RHICmdList, FTextureRHIParamRef Texture2DMS, FTextureRHIParamRef FMaskTexure2D)
+	void SetParameters(FRHICommandList& RHICmdList, FRHITexture* Texture2DMS, FRHITexture* FMaskTexure2D)
 	{
 		FRHIPixelShader* PixelShaderRHI = GetPixelShader();
 		SetTextureParameter(RHICmdList, PixelShaderRHI, Tex, Texture2DMS);
@@ -263,7 +263,7 @@ public:
 		return IsFeatureLevelSupported(Parameters.Platform, ERHIFeatureLevel::ES2);
 	}
 
-	void SetParameters(FRHICommandList& RHICmdList, FTextureRHIParamRef Texture2DMS, FTextureRHIParamRef FMaskTexure2D)
+	void SetParameters(FRHICommandList& RHICmdList, FRHITexture* Texture2DMS, FRHITexture* FMaskTexure2D)
 	{
 		FRHIPixelShader* PixelShaderRHI = GetPixelShader();
 		SetTextureParameter(RHICmdList, PixelShaderRHI, Tex, Texture2DMS);
diff --git a/Engine/Source/Runtime/Renderer/Public/MaterialShader.h b/Engine/Source/Runtime/Renderer/Public/MaterialShader.h
--- a/Engine/Source/Runtime/Renderer/Public/MaterialShader.h
+++ b/Engine/Source/Runtime/Renderer/Public/MaterialShader.h
@@ -124,7 +124,7 @@ public:
 		return true;
 	}
 
-	FUniformBufferRHIParamRef GetParameterCollectionBuffer(const FGuid& Id, const FSceneInterface* SceneInterface) const;
+	FRHIUniformBuffer* GetParameterCollectionBuffer(const FGuid& Id, const FSceneInterface* SceneInterface) const;
 
 	template<typename ShaderRHIParamRef>
 	FORCEINLINE_DEBUGGABLE void SetViewParameters(FRHICommandList& RHICmdList, const ShaderRHIParamRef ShaderRHI, const FSceneView& View, const TUniformBufferRef<FViewUniformShaderParameters>& ViewUniformBuffer)
diff --git a/Engine/Source/Runtime/Renderer/Public/MeshDrawShaderBindings.h b/Engine/Source/Runtime/Renderer/Public/MeshDrawShaderBindings.h
--- a/Engine/Source/Runtime/Renderer/Public/MeshDrawShaderBindings.h
+++ b/Engine/Source/Runtime/Renderer/Public/MeshDrawShaderBindings.h
@@ -44,7 +44,7 @@ public:
 			+ ParameterMapInfo.TextureSamplers.Num() 
 			+ ParameterMapInfo.SRVs.Num());
 
-		// Allocate a bit for each SRV tracking whether it is a FTextureRHIParamRef or FRHIShaderResourceView*
+		// Allocate a bit for each SRV tracking whether it is a FRHITexture* or FRHIShaderResourceView*
 		DataSize += FMath::DivideAndRoundUp(ParameterMapInfo.SRVs.Num(), 8);
 
 		for (int32 LooseBufferIndex = 0; LooseBufferIndex < ParameterMapInfo.LooseParameterBuffers.Num(); LooseBufferIndex++)
@@ -66,7 +66,7 @@ protected:
 
 	inline uint32 GetSamplerOffset() const
 	{
-		return ParameterMapInfo.UniformBuffers.Num() * sizeof(FUniformBufferRHIParamRef);
+		return ParameterMapInfo.UniformBuffers.Num() * sizeof(FRHIUniformBuffer*);
 	}
 
 	inline uint32 GetSRVOffset() const
@@ -125,7 +125,7 @@ public:
 		}
 	}
 
-	void Add(FShaderUniformBufferParameter Parameter, FUniformBufferRHIParamRef Value)
+	void Add(FShaderUniformBufferParameter Parameter, FRHIUniformBuffer* Value)
 	{
 		checkfSlow(Parameter.IsInitialized(), TEXT("Parameter was not serialized"));
 
@@ -153,7 +153,7 @@ public:
 		FShaderResourceParameter TextureParameter,
 		FShaderResourceParameter SamplerParameter,
 		FRHISamplerState* SamplerStateRHI,
-		FTextureRHIParamRef TextureRHI)
+		FRHITexture* TextureRHI)
 	{
 		checkfSlow(TextureParameter.IsInitialized(), TEXT("Parameter was not serialized"));
 		checkfSlow(SamplerParameter.IsInitialized(), TEXT("Parameter was not serialized"));
@@ -217,9 +217,9 @@ public:
 private:
 	uint8* Data;
 
-	inline FUniformBufferRHIParamRef* GetUniformBufferStart() const
+	inline FRHIUniformBuffer** GetUniformBufferStart() const
 	{
-		return (FUniformBufferRHIParamRef*)(Data + GetUniformBufferOffset());
+		return (FRHIUniformBuffer**)(Data + GetUniformBufferOffset());
 	}
 
 	inline FRHISamplerState** GetSamplerStart() const
@@ -247,7 +247,7 @@ private:
 		return LooseDataStart;
 	}
 
-	inline void WriteBindingUniformBuffer(FUniformBufferRHIParamRef Value, uint32 BaseIndex)
+	inline void WriteBindingUniformBuffer(FRHIUniformBuffer* Value, uint32 BaseIndex)
 	{
 		int32 FoundIndex = -1;
 
@@ -328,7 +328,7 @@ private:
 		checkfSlow(FoundIndex >= 0, TEXT("Attempted to set SRV at BaseIndex %u which was never in the shader's parameter map."), BaseIndex);
 	}
 
-	inline void WriteBindingTexture(FTextureRHIParamRef Value, uint32 BaseIndex)
+	inline void WriteBindingTexture(FRHITexture* Value, uint32 BaseIndex)
 	{
 		int32 FoundIndex = -1;
 
diff --git a/Engine/Source/Runtime/Renderer/Public/MeshMaterialShader.h b/Engine/Source/Runtime/Renderer/Public/MeshMaterialShader.h
--- a/Engine/Source/Runtime/Renderer/Public/MeshMaterialShader.h
+++ b/Engine/Source/Runtime/Renderer/Public/MeshMaterialShader.h
@@ -23,8 +23,8 @@ template<typename TBufferStruct> class TUniformBufferRef;
 class FMeshMaterialShaderElementData
 {
 public:
-	FUniformBufferRHIParamRef FadeUniformBuffer = nullptr;
-	FUniformBufferRHIParamRef DitherUniformBuffer = nullptr;
+	FRHIUniformBuffer* FadeUniformBuffer = nullptr;
+	FRHIUniformBuffer* DitherUniformBuffer = nullptr;
 
 	RENDERER_API void InitializeMeshMaterialData(const FSceneView* SceneView, const FPrimitiveSceneProxy* RESTRICT PrimitiveSceneProxy, const FMeshBatch& RESTRICT MeshBatch, int32 StaticMeshId, bool bAllowStencilDither);
 };
diff --git a/Engine/Source/Runtime/Renderer/Public/MeshPassProcessor.h b/Engine/Source/Runtime/Renderer/Public/MeshPassProcessor.h
--- a/Engine/Source/Runtime/Renderer/Public/MeshPassProcessor.h
+++ b/Engine/Source/Runtime/Renderer/Public/MeshPassProcessor.h
@@ -446,7 +446,7 @@ public:
 	 */
 	FMeshDrawShaderBindings ShaderBindings;
 	FVertexInputStreamArray VertexStreams;
-	FIndexBufferRHIParamRef IndexBuffer;
+	FRHIIndexBuffer* IndexBuffer;
 
 	/**
 	 * PSO
@@ -470,7 +470,7 @@ public:
 		
 		struct  
 		{
-			FVertexBufferRHIParamRef Buffer;
+			FRHIVertexBuffer* Buffer;
 			uint32 Offset;
 		} IndirectArgs;
 	};
@@ -561,7 +561,7 @@ public:
 	static void SubmitDraw(
 		const FMeshDrawCommand& RESTRICT MeshDrawCommand, 
 		const FGraphicsMinimalPipelineStateSet& GraphicsMinimalPipelineStateSet,
-		FVertexBufferRHIParamRef ScenePrimitiveIdsBuffer,
+		FRHIVertexBuffer* ScenePrimitiveIdsBuffer,
 		int32 PrimitiveIdOffset,
 		uint32 InstanceFactor,
 		FRHICommandList& CommandList, 
@@ -914,7 +914,7 @@ ENUM_CLASS_FLAGS(EMeshPassFeatures);
  */
 struct FMeshPassProcessorRenderState
 {
-	FMeshPassProcessorRenderState(const FSceneView& SceneView, FUniformBufferRHIParamRef InPassUniformBuffer = nullptr) : 
+	FMeshPassProcessorRenderState(const FSceneView& SceneView, FRHIUniformBuffer* InPassUniformBuffer = nullptr) :
 		  BlendState(nullptr)
 		, DepthStencilState(nullptr)
 		, DepthStencilAccess(FExclusiveDepthStencil::DepthRead_StencilRead)
@@ -926,7 +926,7 @@ struct FMeshPassProcessorRenderState
 	{
 	}
 
-	FMeshPassProcessorRenderState(const TUniformBufferRef<FViewUniformShaderParameters>& InViewUniformBuffer, FUniformBufferRHIParamRef InPassUniformBuffer) : 
+	FMeshPassProcessorRenderState(const TUniformBufferRef<FViewUniformShaderParameters>& InViewUniformBuffer, FRHIUniformBuffer* InPassUniformBuffer) :
 		  BlendState(nullptr)
 		, DepthStencilState(nullptr)
 		, DepthStencilAccess(FExclusiveDepthStencil::DepthRead_StencilRead)
@@ -1022,7 +1022,7 @@ public:
 		return InstancedViewUniformBuffer.IsValid() ? InstancedViewUniformBuffer : reinterpret_cast<const TUniformBufferRef<FInstancedViewUniformShaderParameters>&>(ViewUniformBuffer);
 	}
 
-	FORCEINLINE_DEBUGGABLE void SetReflectionCaptureUniformBuffer(FUniformBufferRHIParamRef InUniformBuffer)
+	FORCEINLINE_DEBUGGABLE void SetReflectionCaptureUniformBuffer(FRHIUniformBuffer* InUniformBuffer)
 	{
 		ReflectionCaptureUniformBuffer = InUniformBuffer;
 	}
@@ -1037,7 +1037,7 @@ public:
 		PassUniformBuffer = InPassUniformBuffer;
 	}
 
-	FORCEINLINE_DEBUGGABLE FUniformBufferRHIParamRef GetPassUniformBuffer() const
+	FORCEINLINE_DEBUGGABLE FRHIUniformBuffer* GetPassUniformBuffer() const
 	{
 		return PassUniformBuffer;
 	}
@@ -1064,7 +1064,7 @@ private:
 	/** Will be bound as reflection capture uniform buffer in case where scene is not available, typically set to dummy/empty buffer to avoid null binding */
 	FUniformBufferRHIRef			ReflectionCaptureUniformBuffer;
 
-	FUniformBufferRHIParamRef		PassUniformBuffer;
+	FRHIUniformBuffer*				PassUniformBuffer;
 	uint32							StencilRef;
 };
 
@@ -1185,7 +1185,7 @@ private:
 extern void SubmitMeshDrawCommands(
 	const FMeshCommandOneFrameArray& VisibleMeshDrawCommands,
 	const FGraphicsMinimalPipelineStateSet& GraphicsMinimalPipelineStateSet, 
-	FVertexBufferRHIParamRef PrimitiveIdsBuffer,
+	FRHIVertexBuffer* PrimitiveIdsBuffer,
 	int32 BasePrimitiveIdsOffset,
 	bool bDynamicInstancing,
 	uint32 InstanceFactor,
@@ -1194,7 +1194,7 @@ extern void SubmitMeshDrawCommands(
 extern void SubmitMeshDrawCommandsRange(
 	const FMeshCommandOneFrameArray& VisibleMeshDrawCommands,
 	const FGraphicsMinimalPipelineStateSet& GraphicsMinimalPipelineStateSet,
-	FVertexBufferRHIParamRef PrimitiveIdsBuffer,
+	FRHIVertexBuffer* PrimitiveIdsBuffer,
 	int32 BasePrimitiveIdsOffset,
 	bool bDynamicInstancing,
 	int32 StartIndex,
diff --git a/Engine/Source/Runtime/SlateCore/Public/Rendering/SlateRenderer.h b/Engine/Source/Runtime/SlateCore/Public/Rendering/SlateRenderer.h
--- a/Engine/Source/Runtime/SlateCore/Public/Rendering/SlateRenderer.h
+++ b/Engine/Source/Runtime/SlateCore/Public/Rendering/SlateRenderer.h
@@ -19,7 +19,6 @@ struct Rect;
 class FSceneInterface;
 struct FSlateBrush;
 
-typedef FRHITexture2D* FTexture2DRHIParamRef;
 typedef TRefCountPtr<FRHITexture2D> FTexture2DRHIRef;
 
 /**
diff --git a/Engine/Source/Runtime/SlateRHIRenderer/Private/SlateMaterialShader.cpp b/Engine/Source/Runtime/SlateRHIRenderer/Private/SlateMaterialShader.cpp
--- a/Engine/Source/Runtime/SlateRHIRenderer/Private/SlateMaterialShader.cpp
+++ b/Engine/Source/Runtime/SlateRHIRenderer/Private/SlateMaterialShader.cpp
@@ -127,7 +127,7 @@ void FSlateMaterialShaderPS::SetParameters(FRHICommandList& RHICmdList, const FS
 	FMaterialShader::SetParameters<FRHIPixelShader>(RHICmdList, ShaderRHI, MaterialRenderProxy, *Material, View, View.ViewUniformBuffer, SceneTextures);
 }
 
-void FSlateMaterialShaderPS::SetAdditionalTexture( FRHICommandList& RHICmdList, const FTextureRHIParamRef InTexture, const FSamplerStateRHIRef SamplerState )
+void FSlateMaterialShaderPS::SetAdditionalTexture( FRHICommandList& RHICmdList, FRHITexture* InTexture, const FSamplerStateRHIRef SamplerState )
 {
 	SetTextureParameter(RHICmdList, GetPixelShader(), AdditionalTextureParameter, TextureParameterSampler, SamplerState, InTexture );
 }
diff --git a/Engine/Source/Runtime/SlateRHIRenderer/Private/SlateMaterialShader.h b/Engine/Source/Runtime/SlateRHIRenderer/Private/SlateMaterialShader.h
--- a/Engine/Source/Runtime/SlateRHIRenderer/Private/SlateMaterialShader.h
+++ b/Engine/Source/Runtime/SlateRHIRenderer/Private/SlateMaterialShader.h
@@ -61,7 +61,7 @@ public:
 
 	void SetDisplayGammaAndContrast(FRHICommandList& RHICmdList, float InDisplayGamma, float InContrast);
 
-	void SetAdditionalTexture( FRHICommandList& RHICmdList, const FTextureRHIParamRef InTexture, const FSamplerStateRHIRef SamplerState );
+	void SetAdditionalTexture( FRHICommandList& RHICmdList, FRHITexture* InTexture, const FSamplerStateRHIRef SamplerState );
 
 	virtual bool Serialize(FArchive& Ar) override;
 
diff --git a/Engine/Source/Runtime/SlateRHIRenderer/Private/SlateRHIRenderer.cpp b/Engine/Source/Runtime/SlateRHIRenderer/Private/SlateRHIRenderer.cpp
--- a/Engine/Source/Runtime/SlateRHIRenderer/Private/SlateRHIRenderer.cpp
+++ b/Engine/Source/Runtime/SlateRHIRenderer/Private/SlateRHIRenderer.cpp
@@ -593,7 +593,7 @@ public:
 	}
 	FCompositePS() {}
 
-	void SetParameters(FRHICommandList& RHICmdList, FTextureRHIParamRef UITextureRHI, FTextureRHIParamRef UITextureWriteMaskRHI, FTextureRHIParamRef SceneTextureRHI, FTextureRHIParamRef ColorSpaceLUTRHI)
+	void SetParameters(FRHICommandList& RHICmdList, FRHITexture* UITextureRHI, FRHITexture* UITextureWriteMaskRHI, FRHITexture* SceneTextureRHI, FRHITexture* ColorSpaceLUTRHI)
 	{
 		static const auto CVarOutputDevice = IConsoleManager::Get().FindTConsoleVariableDataInt(TEXT("r.HDR.Display.OutputDevice"));
 
@@ -931,7 +931,7 @@ void FSlateRHIRenderer::DrawWindow_RenderThread(FRHICommandListImmediate& RHICmd
 
 						TShaderMapRef<FScreenVS> VertexShader(ShaderMap);
 
-						FTextureRHIParamRef UITargetRTMaskTexture = GSupportsRenderTargetWriteMask ? ViewportInfo.UITargetRTMask->GetRenderTargetItem().TargetableTexture : nullptr;
+						FRHITexture* UITargetRTMaskTexture = GSupportsRenderTargetWriteMask ? ViewportInfo.UITargetRTMask->GetRenderTargetItem().TargetableTexture : nullptr;
 						if (HDROutputDevice == 5 || HDROutputDevice == 6)
 						{
 							// ScRGB encoding
diff --git a/Engine/Source/Runtime/SlateRHIRenderer/Private/SlateRHIRenderingPolicy.cpp b/Engine/Source/Runtime/SlateRHIRenderer/Private/SlateRHIRenderingPolicy.cpp
--- a/Engine/Source/Runtime/SlateRHIRenderer/Private/SlateRHIRenderingPolicy.cpp
+++ b/Engine/Source/Runtime/SlateRHIRenderer/Private/SlateRHIRenderingPolicy.cpp
@@ -907,7 +907,7 @@ void FSlateRHIRenderingPolicy::DrawElements(
 #endif
 
 				FRHISamplerState* SamplerState = BilinearClamp;
-				FTextureRHIParamRef TextureRHI = GWhiteTexture->TextureRHI;
+				FRHITexture* TextureRHI = GWhiteTexture->TextureRHI;
 				if (ShaderResource)
 				{
 					ETextureSamplerFilter Filter = ETextureSamplerFilter::Bilinear;
@@ -926,9 +926,9 @@ void FSlateRHIRenderingPolicy::DrawElements(
 					}
 					else
 					{
-						FTextureRHIParamRef NativeTextureRHI = ((TSlateTexture<FTexture2DRHIRef>*)ShaderResource)->GetTypedResource();
+						FRHITexture* NativeTextureRHI = ((TSlateTexture<FTexture2DRHIRef>*)ShaderResource)->GetTypedResource();
 						// Atlas textures that have no content are never initialized but null textures are invalid on many platforms.
-						TextureRHI = NativeTextureRHI ? NativeTextureRHI : (FTextureRHIParamRef)GWhiteTexture->TextureRHI;
+						TextureRHI = NativeTextureRHI ? NativeTextureRHI : (FRHITexture*)GWhiteTexture->TextureRHI;
 					}
 
 					if (EnumHasAllFlags(DrawFlags, (ESlateBatchDrawFlag::TileU | ESlateBatchDrawFlag::TileV)))
diff --git a/Engine/Source/Runtime/SlateRHIRenderer/Private/SlateShaders.h b/Engine/Source/Runtime/SlateRHIRenderer/Private/SlateShaders.h
--- a/Engine/Source/Runtime/SlateRHIRenderer/Private/SlateShaders.h
+++ b/Engine/Source/Runtime/SlateRHIRenderer/Private/SlateShaders.h
@@ -137,7 +137,7 @@ public:
 	 * @param Texture	Texture resource to use when this pixel shader is bound
 	 * @param SamplerState	Sampler state to use when sampling this texture
 	 */
-	void SetTexture(FRHICommandList& RHICmdList, const FTextureRHIParamRef InTexture, const FSamplerStateRHIRef SamplerState )
+	void SetTexture(FRHICommandList& RHICmdList, FRHITexture* InTexture, const FSamplerStateRHIRef SamplerState )
 	{
 		SetTextureParameter(RHICmdList, GetPixelShader(), TextureParameter, TextureParameterSampler, SamplerState, InTexture );
 	}
diff --git a/Engine/Source/Runtime/UMG/Private/Blueprint/AsyncTaskDownloadImage.cpp b/Engine/Source/Runtime/UMG/Private/Blueprint/AsyncTaskDownloadImage.cpp
--- a/Engine/Source/Runtime/UMG/Private/Blueprint/AsyncTaskDownloadImage.cpp
+++ b/Engine/Source/Runtime/UMG/Private/Blueprint/AsyncTaskDownloadImage.cpp
@@ -20,7 +20,7 @@ static void WriteRawToTexture_RenderThread(FTexture2DDynamicResource* TextureRes
 {
 	check(IsInRenderingThread());
 
-	FTexture2DRHIParamRef TextureRHI = TextureResource->GetTexture2DRHI();
+	FRHITexture2D* TextureRHI = TextureResource->GetTexture2DRHI();
 
 	int32 Width = TextureRHI->GetSizeX();
 	int32 Height = TextureRHI->GetSizeY();
diff --git a/Engine/Source/Runtime/UtilityShaders/Public/OculusShaders.h b/Engine/Source/Runtime/UtilityShaders/Public/OculusShaders.h
--- a/Engine/Source/Runtime/UtilityShaders/Public/OculusShaders.h
+++ b/Engine/Source/Runtime/UtilityShaders/Public/OculusShaders.h
@@ -72,7 +72,7 @@ public:
 		SetTextureParameter(RHICmdList, GetPixelShader(), InTexture, InTextureSampler, Texture);
 	}
 
-	UTILITYSHADERS_API void SetParameters(FRHICommandList& RHICmdList, FRHISamplerState* SamplerStateRHI, FTextureRHIParamRef TextureRHI)
+	UTILITYSHADERS_API void SetParameters(FRHICommandList& RHICmdList, FRHISamplerState* SamplerStateRHI, FRHITexture* TextureRHI)
 	{
 		SetTextureParameter(RHICmdList, GetPixelShader(), InTexture, InTextureSampler, SamplerStateRHI, TextureRHI);
 	}
@@ -116,7 +116,7 @@ public:
 		SetShaderValue(RHICmdList, GetPixelShader(), InFaceIndexParameter, FaceIndex);
 	}
 
-	void SetParameters(FRHICommandList& RHICmdList, FRHISamplerState* SamplerStateRHI, FTextureRHIParamRef TextureRHI, int FaceIndex)
+	void SetParameters(FRHICommandList& RHICmdList, FRHISamplerState* SamplerStateRHI, FRHITexture* TextureRHI, int FaceIndex)
 	{
 		SetTextureParameter(RHICmdList, GetPixelShader(), InTexture, InTextureSampler, SamplerStateRHI, TextureRHI);
 		SetShaderValue(RHICmdList, GetPixelShader(), InFaceIndexParameter, FaceIndex);
