Fixed RDG texture extraction from composition graph to transition to readable.

diff --git a/Engine/Source/Runtime/Renderer/Private/PostProcess/RenderingCompositionGraph.cpp b/Engine/Source/Runtime/Renderer/Private/PostProcess/RenderingCompositionGraph.cpp
--- a/Engine/Source/Runtime/Renderer/Private/PostProcess/RenderingCompositionGraph.cpp
+++ b/Engine/Source/Runtime/Renderer/Private/PostProcess/RenderingCompositionGraph.cpp
@@ -1189,7 +1189,7 @@ void FRenderingCompositePass::ExtractRDGTextureForOutput(FRDGBuilder& GraphBuild
 	if (FRenderingCompositeOutput* Output = GetOutput(OutputId))
 	{
 		Output->RenderTargetDesc = Texture->Desc;
-		GraphBuilder.QueueTextureExtraction(Texture, &Output->PooledRenderTarget, false);
+		GraphBuilder.QueueTextureExtraction(Texture, &Output->PooledRenderTarget);
 	}
 }
 
