Cherry-pick 6956319: Add world position offset support for skeletal meshes

diff --git a/Engine/Plugins/FX/Niagara/Source/Niagara/Private/NiagaraRendererSprites.cpp b/Engine/Plugins/FX/Niagara/Source/Niagara/Private/NiagaraRendererSprites.cpp
--- a/Engine/Plugins/FX/Niagara/Source/Niagara/Private/NiagaraRendererSprites.cpp
+++ b/Engine/Plugins/FX/Niagara/Source/Niagara/Private/NiagaraRendererSprites.cpp
@@ -641,7 +641,7 @@ void FNiagaraRendererSprites::GetDynamicRayTracingInstances(FRayTracingMaterialG
 		Context.DynamicRayTracingGeometriesToUpdate.Add(
 			FRayTracingDynamicGeometryUpdateParams
 			{
-				MeshBatch,
+				RayTracingInstance.Materials,
 				MeshBatch.Elements[0].NumPrimitives == 0,
 				6 *  SourceParticleData->GetNumInstances(),
 				6 *  SourceParticleData->GetNumInstances() * (uint32)sizeof(FVector),
diff --git a/Engine/Shaders/Private/LocalVertexFactory.ush b/Engine/Shaders/Private/LocalVertexFactory.ush
--- a/Engine/Shaders/Private/LocalVertexFactory.ush
+++ b/Engine/Shaders/Private/LocalVertexFactory.ush
@@ -183,19 +183,39 @@ struct FVertexFactoryInput
 #endif
 };
 
-#if RAYHITGROUPSHADER
+#if RAYHITGROUPSHADER || COMPUTESHADER
+#if GPUSKIN_PASS_THROUGH
+Buffer<float> GPUSkinCachePositionBuffer;
+#endif
+
 FVertexFactoryInput LoadVertexFactoryInput(uint TriangleIndex, int VertexIndex)
 {
+	FVertexFactoryInput Input = (FVertexFactoryInput)0;
+
+#if RAYHITGROUPSHADER
 	FTriangleBaseAttributes Triangle = LoadTriangleBaseAttributes(TriangleIndex);
 
-	FVertexFactoryInput Input;
-	Input.Position = float4(Triangle.LocalPositions[VertexIndex], 1.0f);
 	Input.VertexId = Triangle.Indices[VertexIndex];
-#if VF_USE_PRIMITIVE_SCENE_DATA
-	Input.PrimitiveId = GetInstanceUserData();
+	Input.Position = float4(Triangle.LocalPositions[VertexIndex], 1.0f);
+	#if VF_USE_PRIMITIVE_SCENE_DATA
+		Input.PrimitiveId = GetInstanceUserData();
+	#endif
+#elif COMPUTESHADER	
+	#if GPUSKIN_PASS_THROUGH
+		Input.VertexId = TriangleIndex * 3 + VertexIndex;
+		uint VertexOffset = LocalVF.VertexFetch_Parameters[VF_VertexOffset] + Input.VertexId;
+		Input.Position.x = GPUSkinCachePositionBuffer[VertexOffset * 3 + 0];
+		Input.Position.y = GPUSkinCachePositionBuffer[VertexOffset * 3 + 1];
+		Input.Position.z = GPUSkinCachePositionBuffer[VertexOffset * 3 + 2];
+	#endif
 #endif
 	return Input;
 }
+
+uint GetNumRayTracingDynamicMeshVerticesIndirect()
+{
+	return 0;
+}
 #endif
 
 /** 
diff --git a/Engine/Shaders/Private/RayTracing/RayTracingDynamicMesh.usf b/Engine/Shaders/Private/RayTracing/RayTracingDynamicMesh.usf
--- a/Engine/Shaders/Private/RayTracing/RayTracingDynamicMesh.usf
+++ b/Engine/Shaders/Private/RayTracing/RayTracingDynamicMesh.usf
@@ -45,8 +45,9 @@ void RayTracingDynamicMeshGS(triangle FRayTracingDynamicMeshVSToGS Inputs[3], in
 }
 
 #if COMPUTESHADER
-uint NumMaxVertices;
-uint NumCPUVertices;
+uint VertexBufferSize;
+uint NumVertices;
+uint MinVertexIndex;
 RWBuffer<float> RWVertexPositions;
 
 [numthreads(256, 1, 1)]
@@ -54,12 +55,15 @@ void RayTracingDynamicGeometryConverterCS(uint3 DispatchThreadId : SV_DispatchTh
 { 
 	ResolvedView = ResolveView();
 
-	uint VertexIndex = DispatchThreadId.x;
-	
-	if (VertexIndex >= NumMaxVertices) return;
+	if (DispatchThreadId.x >= VertexBufferSize) return;
 
-	uint NumActualVertices = NumCPUVertices > 0 ? NumCPUVertices : GetNumRayTracingDynamicMeshVerticesIndirect();
-	if (VertexIndex < NumActualVertices)
+	bool bUsingIndirectDraw = GetNumRayTracingDynamicMeshVerticesIndirect() > 0;
+	
+	uint NumActualVertices = NumVertices > 0 ? NumVertices : GetNumRayTracingDynamicMeshVerticesIndirect();
+	
+	uint VertexIndex = MinVertexIndex + DispatchThreadId.x;
+	
+	if (DispatchThreadId.x < NumActualVertices)
 	{
 		FVertexFactoryInput Input = LoadVertexFactoryInput(VertexIndex / 3, VertexIndex % 3);
 		FVertexFactoryIntermediates VFIntermediates = GetVertexFactoryIntermediates(Input);
@@ -74,7 +78,7 @@ void RayTracingDynamicGeometryConverterCS(uint3 DispatchThreadId : SV_DispatchTh
 		RWVertexPositions[VertexIndex * 3 + 1] = WorldPosition.y;
 		RWVertexPositions[VertexIndex * 3 + 2] = WorldPosition.z;
 	}
-	else
+	else if (bUsingIndirectDraw)
 	{
 		RWVertexPositions[VertexIndex * 3 + 0] = asfloat(0xFFFFFFFF);
 		RWVertexPositions[VertexIndex * 3 + 1] = 0;
diff --git a/Engine/Source/Runtime/Engine/Private/GPUSkinCache.cpp b/Engine/Source/Runtime/Engine/Private/GPUSkinCache.cpp
--- a/Engine/Source/Runtime/Engine/Private/GPUSkinCache.cpp
+++ b/Engine/Source/Runtime/Engine/Private/GPUSkinCache.cpp
@@ -1207,6 +1207,7 @@ void FGPUSkinCache::GetShaderBindings(
 	const FShader* Shader, 
 	const FGPUSkinPassthroughVertexFactory* VertexFactory,
 	uint32 BaseVertexIndex, 
+	FShaderResourceParameter GPUSkinCachePositionBuffer,
 	FShaderResourceParameter GPUSkinCachePreviousPositionBuffer,
 	class FMeshDrawSingleShaderBindings& ShaderBindings,
 	FVertexInputStreamArray& VertexStreams)
@@ -1226,6 +1227,7 @@ void FGPUSkinCache::GetShaderBindings(
 		VertexStreams.Add(FVertexInputStream(VertexFactory->GetTangentStreamIndex(), 0, DispatchData.GetTangentRWBuffer()->Buffer));
 	}
 
+	ShaderBindings.Add(GPUSkinCachePositionBuffer, DispatchData.GetPositionRWBuffer()->SRV);
 	ShaderBindings.Add(GPUSkinCachePreviousPositionBuffer, DispatchData.GetPreviousPositionRWBuffer()->SRV);
 }
 
diff --git a/Engine/Source/Runtime/Engine/Private/GPUSkinVertexFactory.cpp b/Engine/Source/Runtime/Engine/Private/GPUSkinVertexFactory.cpp
--- a/Engine/Source/Runtime/Engine/Private/GPUSkinVertexFactory.cpp
+++ b/Engine/Source/Runtime/Engine/Private/GPUSkinVertexFactory.cpp
@@ -595,6 +595,7 @@ public:
 	virtual void Bind(const FShaderParameterMap& ParameterMap) override
 	{
 		FLocalVertexFactoryShaderParametersBase::Bind(ParameterMap);
+		GPUSkinCachePositionBuffer.Bind(ParameterMap,TEXT("GPUSkinCachePositionBuffer"));
 		GPUSkinCachePreviousPositionBuffer.Bind(ParameterMap,TEXT("GPUSkinCachePreviousPositionBuffer"));
 	}
 	/**
@@ -604,6 +605,7 @@ public:
 	virtual void Serialize(FArchive& Ar) override
 	{
 		FLocalVertexFactoryShaderParametersBase::Serialize(Ar);
+		Ar << GPUSkinCachePositionBuffer;
 		Ar << GPUSkinCachePreviousPositionBuffer;
 	}
 	
@@ -628,12 +630,13 @@ public:
 
 		// #dxr_todo do we need this call to the base?
 		FLocalVertexFactoryShaderParametersBase::GetElementShaderBindingsBase(Scene, View, Shader, InputStreamType, FeatureLevel, VertexFactory, BatchElement, VertexFactoryUniformBuffer, ShaderBindings, VertexStreams);
-		FGPUSkinCache::GetShaderBindings(BatchUserData->Entry, BatchUserData->Section, Shader, (const FGPUSkinPassthroughVertexFactory*)VertexFactory, BatchElement.MinVertexIndex, GPUSkinCachePreviousPositionBuffer, ShaderBindings, VertexStreams);
+		FGPUSkinCache::GetShaderBindings(BatchUserData->Entry, BatchUserData->Section, Shader, (const FGPUSkinPassthroughVertexFactory*)VertexFactory, BatchElement.MinVertexIndex, GPUSkinCachePositionBuffer, GPUSkinCachePreviousPositionBuffer, ShaderBindings, VertexStreams);
 	}
 
 	virtual uint32 GetSize() const override { return sizeof(*this); }
 
 private:
+	FShaderResourceParameter GPUSkinCachePositionBuffer;
 	FShaderResourceParameter GPUSkinCachePreviousPositionBuffer;
 };
 
@@ -730,7 +733,7 @@ void FGPUSkinPassthroughVertexFactory::InternalUpdateVertexDeclaration(FGPUBaseS
 FVertexFactoryShaderParameters* FGPUSkinPassthroughVertexFactory::ConstructShaderParameters(EShaderFrequency ShaderFrequency)
 {
 #if RHI_RAYTRACING
-	return (ShaderFrequency == SF_Vertex || ShaderFrequency == SF_RayHitGroup) ? new FGPUSkinVertexPassthroughFactoryShaderParameters() : nullptr;
+	return (ShaderFrequency == SF_Vertex || ShaderFrequency == SF_RayHitGroup || ShaderFrequency == SF_Compute) ? new FGPUSkinVertexPassthroughFactoryShaderParameters() : nullptr;
 #else // RHI_RAYTRACING
 	return (ShaderFrequency == SF_Vertex) ? new FGPUSkinVertexPassthroughFactoryShaderParameters() : nullptr;
 #endif // RHI_RAYTRACING
diff --git a/Engine/Source/Runtime/Engine/Private/SkeletalMesh.cpp b/Engine/Source/Runtime/Engine/Private/SkeletalMesh.cpp
--- a/Engine/Source/Runtime/Engine/Private/SkeletalMesh.cpp
+++ b/Engine/Source/Runtime/Engine/Private/SkeletalMesh.cpp
@@ -3469,6 +3469,9 @@ const FQuat SphylBasis(FVector(1.0f / FMath::Sqrt(2.0f), 0.0f, 1.0f / FMath::Sqr
  */
 FSkeletalMeshSceneProxy::FSkeletalMeshSceneProxy(const USkinnedMeshComponent* Component, FSkeletalMeshRenderData* InSkelMeshRenderData)
 		:	FPrimitiveSceneProxy(Component, Component->SkeletalMesh->GetFName())
+#if RHI_RAYTRACING
+		,	bAnySegmentUsesWorldPositionOffset(false)
+#endif
 		,	Owner(Component->GetOwner())
 		,	MeshObject(Component->MeshObject)
 		,	SkeletalMeshRenderData(InSkelMeshRenderData)
@@ -3582,6 +3585,11 @@ FSkeletalMeshSceneProxy::FSkeletalMeshSceneProxy(const USkinnedMeshComponent* Co
 				(Component->SkeletalMesh->Materials.IsValidIndex(UseMaterialIndex) == false || Section.bCastShadow);
 
 			bAnySectionCastsShadow |= bSectionCastsShadow;
+
+#if RHI_RAYTRACING
+			bAnySegmentUsesWorldPositionOffset |= MaterialRelevance.bUsesWorldPositionOffset;
+#endif
+
 			LODSection.SectionElements.Add(
 				FSectionElementInfo(
 					Material,
@@ -3994,7 +4002,8 @@ void FSkeletalMeshSceneProxy::CreateBaseMeshBatch(const FSceneView* View, const
 	FMeshBatchElement& BatchElement = Mesh.Elements[0];
 	BatchElement.FirstIndex = LODData.RenderSections[SectionIndex].BaseIndex;
 	BatchElement.IndexBuffer = LODData.MultiSizeIndexContainer.GetIndexBuffer();
-	BatchElement.MaxVertexIndex = LODData.GetNumVertices() - 1;
+	BatchElement.MinVertexIndex = LODData.RenderSections[SectionIndex].GetVertexBufferIndex();
+	BatchElement.MaxVertexIndex = LODData.RenderSections[SectionIndex].GetVertexBufferIndex() + LODData.RenderSections[SectionIndex].GetNumVertices();
 	BatchElement.VertexFactoryUserData = FGPUSkinCache::GetFactoryUserData(MeshObject->SkinCacheEntry, SectionIndex);
 	BatchElement.PrimitiveUniformBuffer = GetUniformBuffer();
 	BatchElement.NumPrimitives = LODData.RenderSections[SectionIndex].NumTriangles;
@@ -4151,7 +4160,6 @@ void FSkeletalMeshSceneProxy::GetDynamicRayTracingInstances(FRayTracingMaterialG
 			
 			FRayTracingInstance RayTracingInstance;
 			RayTracingInstance.Geometry = MeshObject->GetRayTracingGeometry();
-			RayTracingInstance.InstanceTransforms.Add(GetLocalToWorld());
 
 			{
 				// Setup materials for each segment
@@ -4174,6 +4182,42 @@ void FSkeletalMeshSceneProxy::GetDynamicRayTracingInstances(FRayTracingMaterialG
 
 					RayTracingInstance.Materials.Add(MeshBatch);
 				}
+
+				if (bAnySegmentUsesWorldPositionOffset)
+				{
+					RayTracingInstance.InstanceTransforms.Add(FMatrix::Identity);
+				}
+				else
+				{
+					RayTracingInstance.InstanceTransforms.Add(GetLocalToWorld());
+				}
+
+				if (bAnySegmentUsesWorldPositionOffset)
+				{
+					TArray<FRayTracingGeometrySegment> GeometrySections;
+					GeometrySections.Reserve(LODData.RenderSections.Num());
+					for (const FSkelMeshRenderSection& Section : LODData.RenderSections)
+					{
+						FRayTracingGeometrySegment Segment;
+						Segment.FirstPrimitive = Section.BaseIndex / 3;
+						Segment.NumPrimitives = Section.NumTriangles;
+						GeometrySections.Add(Segment);
+					}
+					MeshObject->GetRayTracingGeometry()->Initializer.Segments = GeometrySections;
+
+					Context.DynamicRayTracingGeometriesToUpdate.Add(
+						FRayTracingDynamicGeometryUpdateParams
+						{
+							RayTracingInstance.Materials,
+							false,
+							LODData.GetNumVertices(),
+							LODData.GetNumVertices() * (uint32)sizeof(FVector),
+							MeshObject->GetRayTracingGeometry()->Initializer.TotalPrimitiveCount,
+							MeshObject->GetRayTracingGeometry(),
+							MeshObject->GetRayTracingDynamicVertexBuffer()
+						}
+					);
+				}
 			}
 
 			RayTracingInstance.BuildInstanceMaskAndFlags();
diff --git a/Engine/Source/Runtime/Engine/Private/SkeletalRenderGPUSkin.cpp b/Engine/Source/Runtime/Engine/Private/SkeletalRenderGPUSkin.cpp
--- a/Engine/Source/Runtime/Engine/Private/SkeletalRenderGPUSkin.cpp
+++ b/Engine/Source/Runtime/Engine/Private/SkeletalRenderGPUSkin.cpp
@@ -362,10 +362,6 @@ void FSkeletalMeshObjectGPUSkin::UpdateDynamicData_RenderThread(FGPUSkinCache* G
 		{
 			if (bRequireRecreatingRayTracingGeometry)
 			{
-				// #dxr: Warning: this path invalidates all the instances referencing RayTracingGeometry
-				// which is expected to be detected inside USkinnedMeshComponent::SendRenderDynamicData_Concurrent()
-				// and instance updates are sent there
-
 				FSkeletalMeshLODRenderData& LODModel = this->SkeletalMeshRenderData->LODRenderData[DynamicData->LODIndex];
 				FIndexBufferRHIRef IndexBufferRHI = LODModel.MultiSizeIndexContainer.GetIndexBuffer()->IndexBufferRHI;
 				uint32 VertexBufferStride = LODModel.StaticVertexBuffers.PositionVertexBuffer.GetStride();
@@ -409,9 +405,17 @@ void FSkeletalMeshObjectGPUSkin::UpdateDynamicData_RenderThread(FGPUSkinCache* G
 			}
 			else
 			{
-				// Refit BLAS with new vertex buffer data
-				FGPUSkinCache::CreateMergedPositionVertexBuffer(RHICmdList, SkinCacheEntry, RayTracingGeometry.Initializer.PositionVertexBuffer);
-				GPUSkinCache->AddRayTracingGeometryToUpdate(&RayTracingGeometry);
+				// If we are not using world position offset in material, handle BLAS refit here
+				if (!DynamicData->bAnySegmentUsesWorldPositionOffset)
+				{
+					// Refit BLAS with new vertex buffer data
+					FGPUSkinCache::CreateMergedPositionVertexBuffer(RHICmdList, SkinCacheEntry, RayTracingGeometry.Initializer.PositionVertexBuffer);
+					GPUSkinCache->AddRayTracingGeometryToUpdate(&RayTracingGeometry);
+				}
+				else
+				{
+					// Otherwise, we will run the dynamic ray tracing geometry path, i.e. runnning VSinCS and refit geometry there, so do nothing here
+				}
 			}
 		}
 	}
@@ -1720,6 +1724,9 @@ void FDynamicSkelMeshObjectDataGPUSkin::Clear()
 	NumWeightedActiveMorphTargets = 0;
 	ClothingSimData.Reset();
 	ClothBlendWeight = 0.0f;
+#if RHI_RAYTRACING
+	bAnySegmentUsesWorldPositionOffset = false;
+#endif
 }
 
 #define SKELETON_POOL_GPUSKINS 1
@@ -1873,6 +1880,13 @@ void FDynamicSkelMeshObjectDataGPUSkin::InitDynamicSkelMeshObjectDataGPUSkin(
 	}
 	// Update the clothing simulation mesh positions and normals
 	UpdateClothSimulationData(InMeshComponent);
+
+#if RHI_RAYTRACING
+	if (SkeletalMeshProxy != nullptr)
+	{
+		bAnySegmentUsesWorldPositionOffset = SkeletalMeshProxy->bAnySegmentUsesWorldPositionOffset;
+	}
+#endif
 }
 
 bool FDynamicSkelMeshObjectDataGPUSkin::ActiveMorphTargetsEqual( const TArray<FActiveMorphTarget>& CompareActiveMorphTargets, const TArray<float>& CompareMorphTargetWeights)
diff --git a/Engine/Source/Runtime/Engine/Private/SkeletalRenderGPUSkin.h b/Engine/Source/Runtime/Engine/Private/SkeletalRenderGPUSkin.h
--- a/Engine/Source/Runtime/Engine/Private/SkeletalRenderGPUSkin.h
+++ b/Engine/Source/Runtime/Engine/Private/SkeletalRenderGPUSkin.h
@@ -113,6 +113,10 @@ public:
 
 	/** Update Simulated Positions & Normals from APEX Clothing actor */
 	bool UpdateClothSimulationData(USkinnedMeshComponent* InMeshComponent);
+
+#if RHI_RAYTRACING
+	bool bAnySegmentUsesWorldPositionOffset;
+#endif
 };
 
 /** morph target mesh data for a single vertex delta */
@@ -278,8 +282,11 @@ public:
 #if RHI_RAYTRACING
 	/** Geometry for ray tracing. */
 	FRayTracingGeometry RayTracingGeometry;
+	FRWBuffer RayTracingDynamicVertexBuffer;
 
+	virtual FRayTracingGeometry* GetRayTracingGeometry() { return &RayTracingGeometry; }
 	virtual const FRayTracingGeometry* GetRayTracingGeometry() const { return &RayTracingGeometry; }
+	virtual FRWBuffer* GetRayTracingDynamicVertexBuffer() { return &RayTracingDynamicVertexBuffer; }
 #endif // RHI_RAYTRACING
 
 	virtual int32 GetLOD() const override
diff --git a/Engine/Source/Runtime/Engine/Public/GPUSkinCache.h b/Engine/Source/Runtime/Engine/Public/GPUSkinCache.h
--- a/Engine/Source/Runtime/Engine/Public/GPUSkinCache.h
+++ b/Engine/Source/Runtime/Engine/Public/GPUSkinCache.h
@@ -118,6 +118,7 @@ public:
 		const FShader* Shader,
 		const FGPUSkinPassthroughVertexFactory* VertexFactory,
 		uint32 BaseVertexIndex,
+		FShaderResourceParameter GPUSkinCachePositionBuffer,
 		FShaderResourceParameter GPUSkinCachePreviousPositionBuffer,
 		class FMeshDrawSingleShaderBindings& ShaderBindings,
 		FVertexInputStreamArray& VertexStreams);
diff --git a/Engine/Source/Runtime/Engine/Public/SceneManagement.h b/Engine/Source/Runtime/Engine/Public/SceneManagement.h
--- a/Engine/Source/Runtime/Engine/Public/SceneManagement.h
+++ b/Engine/Source/Runtime/Engine/Public/SceneManagement.h
@@ -2258,7 +2258,7 @@ public:
 
 struct FRayTracingDynamicGeometryUpdateParams
 {
-	FMeshBatch MeshBatch;
+	TArray<FMeshBatch> MeshBatches;
 
 	bool bUsingIndirectDraw;
 	// When bUsingIndirectDraw == false, NumVertices == the actual number of vertices to process
diff --git a/Engine/Source/Runtime/Engine/Public/SkeletalMeshTypes.h b/Engine/Source/Runtime/Engine/Public/SkeletalMeshTypes.h
--- a/Engine/Source/Runtime/Engine/Public/SkeletalMeshTypes.h
+++ b/Engine/Source/Runtime/Engine/Public/SkeletalMeshTypes.h
@@ -296,6 +296,10 @@ public:
 	const TArray<FMatrix>& GetMeshObjectReferenceToLocalMatrices() const;
 	const TIndirectArray<FSkeletalMeshLODRenderData>& GetSkeletalMeshRenderDataLOD() const;
 
+#if RHI_RAYTRACING
+	bool bAnySegmentUsesWorldPositionOffset : 1;
+#endif
+
 protected:
 	AActor* Owner;
 	class FSkeletalMeshObject* MeshObject;
diff --git a/Engine/Source/Runtime/Engine/Public/SkeletalRenderPublic.h b/Engine/Source/Runtime/Engine/Public/SkeletalRenderPublic.h
--- a/Engine/Source/Runtime/Engine/Public/SkeletalRenderPublic.h
+++ b/Engine/Source/Runtime/Engine/Public/SkeletalRenderPublic.h
@@ -199,7 +199,9 @@ public:
 
 #if RHI_RAYTRACING
 	/** Retrieve ray tracing geometry from the underlying mesh object */
+	virtual FRayTracingGeometry* GetRayTracingGeometry() { return nullptr; }
 	virtual const FRayTracingGeometry* GetRayTracingGeometry() const { return nullptr; }
+	virtual FRWBuffer* GetRayTracingDynamicVertexBuffer() { return nullptr; }
 #endif // RHI_RAYTRACING
 
 	/** Called to notify clothing data that component transform has changed */
diff --git a/Engine/Source/Runtime/Landscape/Private/LandscapeRender.cpp b/Engine/Source/Runtime/Landscape/Private/LandscapeRender.cpp
--- a/Engine/Source/Runtime/Landscape/Private/LandscapeRender.cpp
+++ b/Engine/Source/Runtime/Landscape/Private/LandscapeRender.cpp
@@ -3130,6 +3130,8 @@ void FLandscapeComponentSceneProxy::GetDynamicRayTracingInstances(FRayTracingMat
 			BatchElement.IndexBuffer = SharedBuffers->ZeroOffsetIndexBuffers[CurrentLOD];
 			BatchElement.FirstIndex = 0;
 			BatchElement.NumPrimitives = NumPrimitives;
+			BatchElement.MinVertexIndex = 0;
+			BatchElement.MaxVertexIndex = 0;
 
 			MeshBatch.Elements.Add(BatchElement);
 
@@ -3207,12 +3209,19 @@ void FLandscapeComponentSceneProxy::GetDynamicRayTracingInstances(FRayTracingMat
 				}
 			}
 
+			FRayTracingInstance RayTracingInstance;
+			RayTracingInstance.Geometry = &SectionRayTracingStates[SubSectionIdx].Geometry;
+			RayTracingInstance.InstanceTransforms.Add(FMatrix::Identity);
+			RayTracingInstance.Materials.Add(MeshBatch);
+			RayTracingInstance.BuildInstanceMaskAndFlags();
+			OutRayTracingInstances.Add(RayTracingInstance);
+
 			if (bNeedsRayTracingGeometryUpdate)
 			{
 				Context.DynamicRayTracingGeometriesToUpdate.Add(
 					FRayTracingDynamicGeometryUpdateParams
 					{
-						MeshBatch,
+						RayTracingInstance.Materials,
 						false,
 						(uint32)FMath::Square(LodSubsectionSizeVerts),
 						FMath::Square(LodSubsectionSizeVerts) * (uint32)sizeof(FVector),
@@ -3222,13 +3231,6 @@ void FLandscapeComponentSceneProxy::GetDynamicRayTracingInstances(FRayTracingMat
 					}
 				);
 			}
-
-			FRayTracingInstance RayTracingInstance;
-			RayTracingInstance.Geometry = &SectionRayTracingStates[SubSectionIdx].Geometry;
-			RayTracingInstance.InstanceTransforms.Add(FMatrix::Identity);
-			RayTracingInstance.Materials.Add(MeshBatch);
-			RayTracingInstance.BuildInstanceMaskAndFlags();
-			OutRayTracingInstances.Add(RayTracingInstance);
 		}
 	}
 }
diff --git a/Engine/Source/Runtime/RHI/Public/DynamicRHI.h b/Engine/Source/Runtime/RHI/Public/DynamicRHI.h
--- a/Engine/Source/Runtime/RHI/Public/DynamicRHI.h
+++ b/Engine/Source/Runtime/RHI/Public/DynamicRHI.h
@@ -129,7 +129,7 @@ struct FRayTracingGeometryInitializer
 	uint32 BaseVertexIndex = 0;
 	ERayTracingGeometryType GeometryType = RTGT_Triangles;
 	uint32 TotalPrimitiveCount = 0;
-	TArrayView<FRayTracingGeometrySegment> Segments;
+	TArray<FRayTracingGeometrySegment> Segments;
 	bool bFastBuild = false;
 	bool bAllowUpdate = false;
 };
diff --git a/Engine/Source/Runtime/Renderer/Private/RayTracing/RayTracingDynamicGeometry.cpp b/Engine/Source/Runtime/Renderer/Private/RayTracing/RayTracingDynamicGeometry.cpp
--- a/Engine/Source/Runtime/Renderer/Private/RayTracing/RayTracingDynamicGeometry.cpp
+++ b/Engine/Source/Runtime/Renderer/Private/RayTracing/RayTracingDynamicGeometry.cpp
@@ -12,7 +12,8 @@ static bool IsSupportedDynamicVertexFactoryType(const FVertexFactoryType* Vertex
 {
 	return VertexFactoryType == FindVertexFactoryType(FName(TEXT("FNiagaraSpriteVertexFactory"), FNAME_Find))
 		|| VertexFactoryType == FindVertexFactoryType(FName(TEXT("FLandscapeVertexFactory"), FNAME_Find))
-		|| VertexFactoryType == FindVertexFactoryType(FName(TEXT("FLandscapeXYOffsetVertexFactory"), FNAME_Find));
+		|| VertexFactoryType == FindVertexFactoryType(FName(TEXT("FLandscapeXYOffsetVertexFactory"), FNAME_Find))
+		|| VertexFactoryType == FindVertexFactoryType(FName(TEXT("FGPUSkinPassthroughVertexFactory"), FNAME_Find));
 }
 
 class FRayTracingDynamicGeometryConverterCS : public FMeshMaterialShader
@@ -25,8 +26,9 @@ public:
 		PassUniformBuffer.Bind(Initializer.ParameterMap, FSceneTexturesUniformParameters::StaticStructMetadata.GetShaderVariableName());
 
 		RWVertexPositions.Bind(Initializer.ParameterMap, TEXT("VertexPositions"));
-		NumMaxVertices.Bind(Initializer.ParameterMap, TEXT("NumMaxVertices"));
-		NumCPUVertices.Bind(Initializer.ParameterMap, TEXT("NumCPUVertices"));
+		VertexBufferSize.Bind(Initializer.ParameterMap, TEXT("VertexBufferSize"));
+		NumVertices.Bind(Initializer.ParameterMap, TEXT("NumVertices"));
+		MinVertexIndex.Bind(Initializer.ParameterMap, TEXT("MinVertexIndex"));
 	}
 
 	FRayTracingDynamicGeometryConverterCS() = default;
@@ -40,8 +42,9 @@ public:
 	{
 		bool bShaderHasOutdatedParameters = FMeshMaterialShader::Serialize(Ar);
 		Ar << RWVertexPositions;
-		Ar << NumMaxVertices;
-		Ar << NumCPUVertices;
+		Ar << VertexBufferSize;
+		Ar << NumVertices;
+		Ar << MinVertexIndex;
 		return bShaderHasOutdatedParameters;
 	}
 
@@ -75,8 +78,9 @@ public:
 	}
 
 	FRWShaderParameter RWVertexPositions;
-	FShaderParameter NumMaxVertices;
-	FShaderParameter NumCPUVertices;
+	FShaderParameter VertexBufferSize;
+	FShaderParameter NumVertices;
+	FShaderParameter MinVertexIndex;
 };
 
 IMPLEMENT_MATERIAL_SHADER_TYPE(, FRayTracingDynamicGeometryConverterCS, TEXT("/Engine/Private/RayTracing/RayTracingDynamicMesh.usf"), TEXT("RayTracingDynamicGeometryConverterCS"), SF_Compute);
@@ -93,49 +97,62 @@ void FRayTracingDynamicGeometryCollection::AddDynamicMeshBatchForGeometryUpdate(
 	FRayTracingDynamicGeometryUpdateParams UpdateParams
 )
 {
-	const FMeshBatch& MeshBatch = UpdateParams.MeshBatch;
 	FRayTracingGeometry& Geometry = *UpdateParams.Geometry;
 	bool bUsingIndirectDraw = UpdateParams.bUsingIndirectDraw;
 	uint32 NumMaxVertices = UpdateParams.NumVertices;
 	FRWBuffer& Buffer = *UpdateParams.Buffer;
 
-	const FMaterialRenderProxy* FallbackMaterialRenderProxyPtr = nullptr;
-	const FMaterial& Material = MeshBatch.MaterialRenderProxy->GetMaterialWithFallback(Scene->GetFeatureLevel(), FallbackMaterialRenderProxyPtr);
+	for (const FMeshBatch& MeshBatch : UpdateParams.MeshBatches)
+	{
+		const FMaterialRenderProxy* FallbackMaterialRenderProxyPtr = nullptr;
+		const FMaterial& Material = MeshBatch.MaterialRenderProxy->GetMaterialWithFallback(Scene->GetFeatureLevel(), FallbackMaterialRenderProxyPtr);
+
+		const FMaterialRenderProxy& MaterialRenderProxy = FallbackMaterialRenderProxyPtr ? *FallbackMaterialRenderProxyPtr : *MeshBatch.MaterialRenderProxy;
+
+		TMeshProcessorShaders<
+			FMeshMaterialShader,
+			FMeshMaterialShader,
+			FMeshMaterialShader,
+			FMeshMaterialShader,
+			FMeshMaterialShader,
+			FMeshMaterialShader,
+			FRayTracingDynamicGeometryConverterCS> Shaders;
 
-	const FMaterialRenderProxy& MaterialRenderProxy = FallbackMaterialRenderProxyPtr ? *FallbackMaterialRenderProxyPtr : *MeshBatch.MaterialRenderProxy;
+		FMeshComputeDispatchCommand DispatchCmd;
 
-	TMeshProcessorShaders<
-		FMeshMaterialShader,
-		FMeshMaterialShader,
-		FMeshMaterialShader,
-		FMeshMaterialShader,
-		FMeshMaterialShader,
-		FMeshMaterialShader,
-		FRayTracingDynamicGeometryConverterCS> Shaders;
+		FRayTracingDynamicGeometryConverterCS* Shader = Material.GetShader<FRayTracingDynamicGeometryConverterCS>(MeshBatch.VertexFactory->GetType());
+		DispatchCmd.MaterialShader = Shader;
+		FMeshDrawShaderBindings& ShaderBindings = DispatchCmd.ShaderBindings;
 
-	FMeshComputeDispatchCommand DispatchCmd;
+		Shaders.ComputeShader = Shader;
+		ShaderBindings.Initialize(Shaders.GetUntypedShaders());
 
-	FRayTracingDynamicGeometryConverterCS* Shader = Material.GetShader<FRayTracingDynamicGeometryConverterCS>(MeshBatch.VertexFactory->GetType());
-	DispatchCmd.MaterialShader = Shader;
-	FMeshDrawShaderBindings& ShaderBindings = DispatchCmd.ShaderBindings;
+		FMeshMaterialShaderElementData ShaderElementData;
+		ShaderElementData.InitializeMeshMaterialData(View, PrimitiveSceneProxy, MeshBatch, -1, false);
 
-	Shaders.ComputeShader = Shader;
-	ShaderBindings.Initialize(Shaders.GetUntypedShaders());
+		FMeshDrawSingleShaderBindings SingleShaderBindings = ShaderBindings.GetSingleShaderBindings(SF_Compute);
+		FMeshPassProcessorRenderState DrawRenderState(Scene->UniformBuffers.ViewUniformBuffer, Scene->UniformBuffers.OpaqueBasePassUniformBuffer);
+		Shader->GetShaderBindings(Scene, Scene->GetFeatureLevel(), PrimitiveSceneProxy, MaterialRenderProxy, Material, DrawRenderState, ShaderElementData, SingleShaderBindings);
 
-	FMeshMaterialShaderElementData ShaderElementData;
-	ShaderElementData.InitializeMeshMaterialData(View, PrimitiveSceneProxy, MeshBatch, -1, false);
+		FVertexInputStreamArray DummyArray;
+		Shader->GetElementShaderBindings(Scene, View, MeshBatch.VertexFactory, EVertexInputStreamType::Default, Scene->GetFeatureLevel(), PrimitiveSceneProxy, MeshBatch, MeshBatch.Elements[0], ShaderElementData, SingleShaderBindings, DummyArray);
 
-	FMeshDrawSingleShaderBindings SingleShaderBindings = ShaderBindings.GetSingleShaderBindings(SF_Compute);
-	FMeshPassProcessorRenderState DrawRenderState(Scene->UniformBuffers.ViewUniformBuffer, Scene->UniformBuffers.OpaqueBasePassUniformBuffer);
-	Shader->GetShaderBindings(Scene, Scene->GetFeatureLevel(), PrimitiveSceneProxy, MaterialRenderProxy, Material, DrawRenderState, ShaderElementData, SingleShaderBindings);
+		DispatchCmd.TargetBuffer = &Buffer;
+		DispatchCmd.NumMaxVertices = UpdateParams.NumVertices;
+		DispatchCmd.NumCPUVertices = !bUsingIndirectDraw ? UpdateParams.NumVertices : 0;
+		if (MeshBatch.Elements[0].MinVertexIndex < MeshBatch.Elements[0].MaxVertexIndex)
+		{
+			DispatchCmd.NumCPUVertices = MeshBatch.Elements[0].MaxVertexIndex - MeshBatch.Elements[0].MinVertexIndex;
+		}
+		DispatchCmd.MinVertexIndex = MeshBatch.Elements[0].MinVertexIndex;
 
-	FVertexInputStreamArray DummyArray;
-	Shader->GetElementShaderBindings(Scene, View, MeshBatch.VertexFactory, EVertexInputStreamType::Default, Scene->GetFeatureLevel(), PrimitiveSceneProxy, MeshBatch, MeshBatch.Elements[0], ShaderElementData, SingleShaderBindings, DummyArray);
+#if MESH_DRAW_COMMAND_DEBUG_DATA
+		FMeshProcessorShaders ShadersForDebug = Shaders.GetUntypedShaders();
+		ShaderBindings.Finalize(&ShadersForDebug);
+#endif
 
-	DispatchCmd.TargetBuffer = &Buffer;
-	DispatchCmd.TargetGeometry = &Geometry;
-	DispatchCmd.NumMaxVertices = NumMaxVertices;
-	DispatchCmd.NumCPUVertices = !bUsingIndirectDraw ? UpdateParams.NumVertices : 0;
+		DispatchCommands->Add(DispatchCmd);
+	}
 
 	bool bRefit = true;
 
@@ -157,17 +174,6 @@ void FRayTracingDynamicGeometryCollection::AddDynamicMeshBatchForGeometryUpdate(
 		bRefit = false;
 	}
 
-	DispatchCmd.bRefit = bRefit;
-
-	check(DispatchCmd.TargetBuffer->NumBytes >= NumMaxVertices * sizeof(FVector));
-
-#if MESH_DRAW_COMMAND_DEBUG_DATA
-	FMeshProcessorShaders ShadersForDebug = Shaders.GetUntypedShaders();
-	ShaderBindings.Finalize(&ShadersForDebug);
-#endif
-
-	DispatchCommands->Add(DispatchCmd);
-
 	check(Geometry.IsInitialized());
 	Geometry.Initializer.PositionVertexBuffer = Buffer.Buffer;
 	Geometry.Initializer.TotalPrimitiveCount = UpdateParams.NumTriangles;
@@ -175,6 +181,11 @@ void FRayTracingDynamicGeometryCollection::AddDynamicMeshBatchForGeometryUpdate(
 	if (!bRefit)
 	{
 		Geometry.RayTracingGeometryRHI = RHICreateRayTracingGeometry(Geometry.Initializer);
+		BuildParams.Add(FAccelerationStructureUpdateParams{ Geometry.RayTracingGeometryRHI, Buffer.Buffer });
+	}
+	else
+	{
+		RefitParams.Add(FAccelerationStructureUpdateParams{ Geometry.RayTracingGeometryRHI, Buffer.Buffer });
 	}
 }
 
@@ -196,8 +207,9 @@ void FRayTracingDynamicGeometryCollection::DispatchUpdates(FRHICommandListImmedi
 
 					Cmd.ShaderBindings.SetOnCommandListForCompute(RHICmdList, Shader->GetComputeShader());
 					Shader->RWVertexPositions.SetBuffer(RHICmdList, Shader->GetComputeShader(), *Cmd.TargetBuffer);
-					SetShaderValue(RHICmdList, Shader->GetComputeShader(), Shader->NumMaxVertices, Cmd.TargetBuffer->NumBytes / sizeof(FVector));
-					SetShaderValue(RHICmdList, Shader->GetComputeShader(), Shader->NumCPUVertices, Cmd.NumCPUVertices);
+					SetShaderValue(RHICmdList, Shader->GetComputeShader(), Shader->VertexBufferSize, Cmd.TargetBuffer->NumBytes / sizeof(FVector));
+					SetShaderValue(RHICmdList, Shader->GetComputeShader(), Shader->NumVertices, Cmd.NumCPUVertices);
+					SetShaderValue(RHICmdList, Shader->GetComputeShader(), Shader->MinVertexIndex, Cmd.MinVertexIndex);
 
 					RHICmdList.DispatchComputeShader(FMath::DivideAndRoundUp<uint32>(Cmd.NumMaxVertices, 256), 1, 1);
 
@@ -206,17 +218,6 @@ void FRayTracingDynamicGeometryCollection::DispatchUpdates(FRHICommandListImmedi
 			}
 		}
 
-		TArray<FAccelerationStructureUpdateParams> BuildParams;
-		TArray<FAccelerationStructureUpdateParams> RefitParams;
-
-		for (auto& Cmd : *DispatchCommands)
-		{
-			if (Cmd.bRefit)
-				RefitParams.Add(FAccelerationStructureUpdateParams{ Cmd.TargetGeometry->RayTracingGeometryRHI, Cmd.TargetBuffer->Buffer });
-			else
-				BuildParams.Add(FAccelerationStructureUpdateParams { Cmd.TargetGeometry->RayTracingGeometryRHI, Cmd.TargetBuffer->Buffer });
-		}
-
 		{
 			SCOPED_DRAW_EVENT(RHICmdList, Build);
 			RHICmdList.BuildAccelerationStructures(BuildParams);
@@ -234,6 +235,8 @@ void FRayTracingDynamicGeometryCollection::DispatchUpdates(FRHICommandListImmedi
 void FRayTracingDynamicGeometryCollection::Clear()
 {
 	DispatchCommands->Empty();
+	BuildParams.Empty();
+	RefitParams.Empty();
 }
 
 #endif // RHI_RAYTRACING
diff --git a/Engine/Source/Runtime/Renderer/Private/ScenePrivate.h b/Engine/Source/Runtime/Renderer/Private/ScenePrivate.h
--- a/Engine/Source/Runtime/Renderer/Private/ScenePrivate.h
+++ b/Engine/Source/Runtime/Renderer/Private/ScenePrivate.h
@@ -2388,10 +2388,8 @@ struct FMeshComputeDispatchCommand
 
 	uint32 NumMaxVertices;
 	uint32 NumCPUVertices;
-	uint32 BaseVertexIndex;
+	uint32 MinVertexIndex;
 	FRWBuffer* TargetBuffer;
-	FRayTracingGeometry* TargetGeometry;
-	bool bRefit;
 };
 #endif
 
diff --git a/Engine/Source/Runtime/Renderer/Public/RayTracingDynamicGeometryCollection.h b/Engine/Source/Runtime/Renderer/Public/RayTracingDynamicGeometryCollection.h
--- a/Engine/Source/Runtime/Renderer/Public/RayTracingDynamicGeometryCollection.h
+++ b/Engine/Source/Runtime/Renderer/Public/RayTracingDynamicGeometryCollection.h
@@ -20,6 +20,8 @@ public:
 
 private:
 	TUniquePtr<TArray<struct FMeshComputeDispatchCommand>> DispatchCommands;
+	TArray<FAccelerationStructureUpdateParams> BuildParams;
+	TArray<FAccelerationStructureUpdateParams> RefitParams;
 };
 
 #endif // RHI_RAYTRACING
