Cherry-pick 6989000: Use TResourceArray to initialize Halton StructuredBuffer data

diff --git a/Engine/Source/Runtime/Renderer/Private/Halton.cpp b/Engine/Source/Runtime/Renderer/Private/Halton.cpp
--- a/Engine/Source/Runtime/Renderer/Private/Halton.cpp
+++ b/Engine/Source/Runtime/Renderer/Private/Halton.cpp
@@ -51,7 +51,7 @@ void FHaltonSequenceIteration::InitRHI()
 	SCOPED_GPU_STAT(FRHICommandListExecutor::GetImmediateCommandList(), HaltonSequence);
 	InitializeSequence();
 
-	TArray<FSequenceIterationData> RandomSamples;
+	TResourceArray<FSequenceIterationData> RandomSamples;
 	uint32 ElementCount = FMath::DivideAndRoundUp(DimensionCount, 3u);
 	RandomSamples.SetNum(SequenceCount * IterationCount * ElementCount);
 	for (uint32 SequenceIndex = 0; SequenceIndex < SequenceCount; ++SequenceIndex)
@@ -74,11 +74,8 @@ void FHaltonSequenceIteration::InitRHI()
 	FRHIResourceCreateInfo CreateInfo;
 	{
 		CreateInfo.DebugName = TEXT("HaltonSequenceIteration");
+		CreateInfo.ResourceArray = &RandomSamples;
 		SequenceIteration = RHICreateStructuredBuffer(sizeof(FSequenceIterationData), RandomSamples.Num() * sizeof(FSequenceIterationData), BUF_Transient | BUF_FastVRAM | BUF_ShaderResource | BUF_UnorderedAccess, CreateInfo);
-		uint32 Offset = 0;
-		void* BasePtr = RHILockStructuredBuffer(SequenceIteration, Offset, RandomSamples.Num() * sizeof(FSequenceIterationData), RLM_WriteOnly);
-		FPlatformMemory::Memcpy(BasePtr, RandomSamples.GetData(), RandomSamples.Num() * sizeof(FSequenceIterationData));
-		RHIUnlockStructuredBuffer(SequenceIteration);
 	}
 }
 
