Cherry-pick 6989371: Adding missing include

diff --git a/Engine/Source/Runtime/Renderer/Private/Halton.cpp b/Engine/Source/Runtime/Renderer/Private/Halton.cpp
--- a/Engine/Source/Runtime/Renderer/Private/Halton.cpp
+++ b/Engine/Source/Runtime/Renderer/Private/Halton.cpp
@@ -8,6 +8,7 @@
 #include "Halton.h"
 #include "Math/PackedVector.h"
 #include "SceneUtils.h"
+#include "Containers/DynamicRHIResourceArray.h"
 
 IMPLEMENT_GLOBAL_SHADER_PARAMETER_STRUCT(FHaltonIteration, "HaltonIteration");
 
