Cherry-pick 7049954: Re-enabling r.RayTracing.SkyLight CVar (default -1)

diff --git a/Engine/Source/Runtime/Renderer/Private/RayTracing/RaytracingSkylight.cpp b/Engine/Source/Runtime/Renderer/Private/RayTracing/RaytracingSkylight.cpp
--- a/Engine/Source/Runtime/Renderer/Private/RayTracing/RaytracingSkylight.cpp
+++ b/Engine/Source/Runtime/Renderer/Private/RayTracing/RaytracingSkylight.cpp
@@ -7,7 +7,7 @@
 #include "DeferredShadingRenderer.h"
 #include "SceneTextureParameters.h"
 
-static int32 GRayTracingSkyLight = 0;
+static int32 GRayTracingSkyLight = -1;
 
 #if RHI_RAYTRACING
 
@@ -83,7 +83,16 @@ bool ShouldRenderRayTracingSkyLight(const FSkyLightSceneProxy* SkyLightSceneProx
 	if (SkyLightSceneProxy != nullptr)
 	{
 		const int32 ForceAllRayTracingEffects = GetForceRayTracingEffectsCVarValue();
-		const bool bRayTracingSkyEnabled = (ForceAllRayTracingEffects > 0 || (SkyLightSceneProxy->bCastRayTracedShadow > 0 && ForceAllRayTracingEffects < 0));
+		bool bRayTracingSkyEnabled;
+		if (ForceAllRayTracingEffects >= 0)
+		{
+			bRayTracingSkyEnabled = ForceAllRayTracingEffects > 0;
+		}
+		else
+		{
+
+			bRayTracingSkyEnabled = GRayTracingSkyLight >= 0 ? GRayTracingSkyLight : SkyLightSceneProxy->bCastRayTracedShadow;
+		}
 
 		return IsRayTracingEnabled() && bRayTracingSkyEnabled;
 	}
@@ -842,7 +851,8 @@ void FDeferredShadingSceneRenderer::RenderRayTracingSkyLight(
 	// Fill SkyLight parameters
 	FSkyLightData SkyLightData;
 	SetupSkyLightParameters(*Scene, &SkyLightData);
-	SkyLightData.SamplesPerPixel = GRayTracingSkyLightSamplesPerPixel >= 0 ? GRayTracingSkyLightSamplesPerPixel : Scene->SkyLight->SamplesPerPixel;
+	uint32 SamplesPerPixel = GRayTracingSkyLightSamplesPerPixel >= 0 ? GRayTracingSkyLightSamplesPerPixel : Scene->SkyLight->SamplesPerPixel;
+	SkyLightData.SamplesPerPixel = SamplesPerPixel;
 	SkyLightData.MaxRayDistance = GRayTracingSkyLightMaxRayDistance;
 	SkyLightData.SamplingStopLevel = GRayTracingSkyLightSamplingStopLevel;
 
@@ -872,7 +882,7 @@ void FDeferredShadingSceneRenderer::RenderRayTracingSkyLight(
 
 
 		// Halton iteration
-		uint32 IterationCount = GRayTracingSkyLightSamplesPerPixel >= 0 ? GRayTracingSkyLightSamplesPerPixel : Scene->SkyLight->SamplesPerPixel;
+		uint32 IterationCount = FMath::Max(SamplesPerPixel, 1u);
 		uint32 SequenceCount = 1;
 		uint32 DimensionCount = 3;
 		FHaltonSequenceIteration HaltonSequenceIteration(Scene->HaltonSequence, IterationCount, SequenceCount, DimensionCount, View.ViewState ? (View.ViewState->FrameIndex % 1024) : 0);
@@ -922,7 +932,7 @@ void FDeferredShadingSceneRenderer::RenderRayTracingSkyLight(
 		});
 
 		// Denoising
-		if (GRayTracingSkyLightDenoiser != 0)
+		if (GRayTracingSkyLightDenoiser != 0 && SamplesPerPixel > 0)
 		{
 			const IScreenSpaceDenoiser* DefaultDenoiser = IScreenSpaceDenoiser::GetDefaultDenoiser();
 			const IScreenSpaceDenoiser* DenoiserToUse = DefaultDenoiser;// GRayTracingGlobalIlluminationDenoiser == 1 ? DefaultDenoiser : GScreenSpaceDenoiser;
@@ -934,7 +944,7 @@ void FDeferredShadingSceneRenderer::RenderRayTracingSkyLight(
 			{
 				IScreenSpaceDenoiser::FAmbientOcclusionRayTracingConfig RayTracingConfig;
 				RayTracingConfig.ResolutionFraction = 1.0;
-				RayTracingConfig.RayCountPerPixel = GRayTracingSkyLightSamplesPerPixel >= 0 ? GRayTracingSkyLightSamplesPerPixel : Scene->SkyLight->SamplesPerPixel;
+				RayTracingConfig.RayCountPerPixel = SamplesPerPixel;
 
 				RDG_EVENT_SCOPE(GraphBuilder, "%s%s(SkyLight) %dx%d",
 					DenoiserToUse != DefaultDenoiser ? TEXT("ThirdParty ") : TEXT(""),
