Fixed tonemapped ray tracing debug visualization modes should use preexposure

diff --git a/Engine/Shaders/Private/RayTracing/RayTracingDebug.usf b/Engine/Shaders/Private/RayTracing/RayTracingDebug.usf
--- a/Engine/Shaders/Private/RayTracing/RayTracingDebug.usf
+++ b/Engine/Shaders/Private/RayTracing/RayTracingDebug.usf
@@ -11,6 +11,7 @@ RWTexture2D<float4> Output;
 RaytracingAccelerationStructure TLAS;
 
 uint VisualizationMode;
+uint ShouldUsePreExposure;
 
 [shader("raygeneration")]
 void RayTracingDebugMainRGS()
@@ -106,6 +107,12 @@ void RayTracingDebugMainRGS()
 		}
 	}
 
+	if (ShouldUsePreExposure)
+	{
+		// Only  when the output is tonemapped
+		Result *= View.PreExposure;
+	}
+
 	Output[PixelCoord] = Result;
 }
 
diff --git a/Engine/Source/Editor/UnrealEd/Private/RayTracingDebugVisualizationMenuCommands.cpp b/Engine/Source/Editor/UnrealEd/Private/RayTracingDebugVisualizationMenuCommands.cpp
--- a/Engine/Source/Editor/UnrealEd/Private/RayTracingDebugVisualizationMenuCommands.cpp
+++ b/Engine/Source/Editor/UnrealEd/Private/RayTracingDebugVisualizationMenuCommands.cpp
@@ -72,6 +72,7 @@ bool FRayTracingDebugVisualizationMenuCommands::DebugModeShouldBeTonemapped(cons
 	static TArray<FName> TonemappedRayTracingDebugVisualizationModes;
 	if (TonemappedRayTracingDebugVisualizationModes.Num() == 0)
 	{
+		// World normal is tonemapped to match raster buffer visualization mode
 		TonemappedRayTracingDebugVisualizationModes.Add(*LOCTEXT("World Normal", "World Normal").ToString());
 		TonemappedRayTracingDebugVisualizationModes.Add(*LOCTEXT("BaseColor", "BaseColor").ToString());
 		TonemappedRayTracingDebugVisualizationModes.Add(*LOCTEXT("DiffuseColor", "DiffuseColor").ToString());
diff --git a/Engine/Source/Runtime/Renderer/Private/RayTracing/RayTracingDebug.cpp b/Engine/Source/Runtime/Renderer/Private/RayTracing/RayTracingDebug.cpp
--- a/Engine/Source/Runtime/Renderer/Private/RayTracing/RayTracingDebug.cpp
+++ b/Engine/Source/Runtime/Renderer/Private/RayTracing/RayTracingDebug.cpp
@@ -22,6 +22,7 @@ class FRayTracingDebugRGS : public FGlobalShader
 
 	BEGIN_SHADER_PARAMETER_STRUCT(FParameters, )
 		SHADER_PARAMETER(uint32, VisualizationMode)
+		SHADER_PARAMETER(int32, ShouldUsePreExposure)
 		SHADER_PARAMETER_SRV(RaytracingAccelerationStructure, TLAS)
 		SHADER_PARAMETER_RDG_TEXTURE_UAV(RWTexture2D<float4>, Output)
 		SHADER_PARAMETER_STRUCT_REF(FViewUniformShaderParameters, ViewUniformBuffer)
@@ -111,6 +112,7 @@ void FDeferredShadingSceneRenderer::RenderRayTracingDebug(FRHICommandListImmedia
 	FRayTracingDebugRGS::FParameters* RayGenParameters = GraphBuilder.AllocParameters<FRayTracingDebugRGS::FParameters>();
 
 	RayGenParameters->VisualizationMode = DebugVisualizationMode;
+	RayGenParameters->ShouldUsePreExposure = View.Family->EngineShowFlags.Tonemapper;
 	RayGenParameters->TLAS = RayTracingSceneRHI->GetShaderResourceView();
 	RayGenParameters->ViewUniformBuffer = View.ViewUniformBuffer;
 	RayGenParameters->Output = GraphBuilder.CreateUAV(GraphBuilder.RegisterExternalTexture(SceneContext.GetSceneColor()));
