Fix compile error on VS2019.

#rnx

diff --git a/Engine/Source/Runtime/Renderer/Private/RayTracing/RaytracingSkylight.cpp b/Engine/Source/Runtime/Renderer/Private/RayTracing/RaytracingSkylight.cpp
--- a/Engine/Source/Runtime/Renderer/Private/RayTracing/RaytracingSkylight.cpp
+++ b/Engine/Source/Runtime/Renderer/Private/RayTracing/RaytracingSkylight.cpp
@@ -91,7 +91,7 @@ bool ShouldRenderRayTracingSkyLight(const FSkyLightSceneProxy* SkyLightSceneProx
 		else
 		{
 
-			bRayTracingSkyEnabled = GRayTracingSkyLight >= 0 ? GRayTracingSkyLight : SkyLightSceneProxy->bCastRayTracedShadow;
+			bRayTracingSkyEnabled = GRayTracingSkyLight >= 0 ? (GRayTracingSkyLight != 0) : SkyLightSceneProxy->bCastRayTracedShadow;
 		}
 
 		return IsRayTracingEnabled() && bRayTracingSkyEnabled;
