Fixed RHI resource transitions causing validation failures.

diff --git a/Engine/Plugins/Compositing/LensDistortion/Source/LensDistortion/Private/LensDistortionRendering.cpp b/Engine/Plugins/Compositing/LensDistortion/Source/LensDistortion/Private/LensDistortionRendering.cpp
--- a/Engine/Plugins/Compositing/LensDistortion/Source/LensDistortion/Private/LensDistortionRendering.cpp
+++ b/Engine/Plugins/Compositing/LensDistortion/Source/LensDistortion/Private/LensDistortionRendering.cpp
@@ -177,7 +177,11 @@ static void DrawUVDisplacementToRenderTarget_RenderThread(
 	SCOPED_DRAW_EVENT(RHICmdList, DrawUVDisplacementToRenderTarget_RenderThread);
 #endif
 
-	FRHIRenderPassInfo RPInfo(OutTextureRenderTargetResource->GetRenderTargetTexture(), ERenderTargetActions::DontLoad_Store, OutTextureRenderTargetResource->TextureRHI);
+	FRHITexture2D* RenderTargetTexture = OutTextureRenderTargetResource->GetRenderTargetTexture();
+
+	RHICmdList.TransitionResource(EResourceTransitionAccess::EWritable, RenderTargetTexture);
+
+	FRHIRenderPassInfo RPInfo(RenderTargetTexture, ERenderTargetActions::DontLoad_Store, OutTextureRenderTargetResource->TextureRHI);
 	RHICmdList.BeginRenderPass(RPInfo, TEXT("DrawUVDisplacement"));
 	{
 		FIntPoint DisplacementMapResolution(OutTextureRenderTargetResource->GetSizeX(), OutTextureRenderTargetResource->GetSizeY());
diff --git a/Engine/Plugins/Compositing/OpenColorIO/Source/OpenColorIO/Private/OpenColorIORendering.cpp b/Engine/Plugins/Compositing/OpenColorIO/Source/OpenColorIO/Private/OpenColorIORendering.cpp
--- a/Engine/Plugins/Compositing/OpenColorIO/Source/OpenColorIO/Private/OpenColorIORendering.cpp
+++ b/Engine/Plugins/Compositing/OpenColorIO/Source/OpenColorIO/Private/OpenColorIORendering.cpp
@@ -71,6 +71,8 @@ static void ProcessOCIOColorSpaceTransform_RenderThread(
 
 	SCOPED_DRAW_EVENT(InRHICmdList, ProcessOCIOColorSpaceTransform);
 
+	InRHICmdList.TransitionResource(EResourceTransitionAccess::EWritable, OutputSpaceColorResource->TextureRHI);
+
 	FRHIRenderPassInfo RPInfo(OutputSpaceColorResource->TextureRHI, ERenderTargetActions::DontLoad_Store);
 	InRHICmdList.BeginRenderPass(RPInfo, TEXT("ProcessOCIOColorSpaceXfrm"));
 
diff --git a/Engine/Source/Runtime/Renderer/Private/PostProcess/PostProcessTemporalAA.cpp b/Engine/Source/Runtime/Renderer/Private/PostProcess/PostProcessTemporalAA.cpp
--- a/Engine/Source/Runtime/Renderer/Private/PostProcess/PostProcessTemporalAA.cpp
+++ b/Engine/Source/Runtime/Renderer/Private/PostProcess/PostProcessTemporalAA.cpp
@@ -716,7 +716,7 @@ FTAAOutputs FTAAPassParameters::AddTemporalAAPass(
 				SceneTextures.SceneDepthBuffer,
 				ERenderTargetLoadAction::ENoAction, ERenderTargetStoreAction::ENoAction,
 				ERenderTargetLoadAction::ELoad, ERenderTargetStoreAction::EStore,
-				FExclusiveDepthStencil::DepthRead_StencilWrite);
+				FExclusiveDepthStencil::DepthRead_StencilRead);
 		}
 		
 		FTemporalAAPS::FPermutationDomain BasePermutationVector;
diff --git a/Engine/Source/Runtime/Renderer/Private/ReflectionEnvironmentCapture.cpp b/Engine/Source/Runtime/Renderer/Private/ReflectionEnvironmentCapture.cpp
--- a/Engine/Source/Runtime/Renderer/Private/ReflectionEnvironmentCapture.cpp
+++ b/Engine/Source/Runtime/Renderer/Private/ReflectionEnvironmentCapture.cpp
@@ -460,11 +460,10 @@ void FilterReflectionEnvironment(FRHICommandListImmediate& RHICmdList, ERHIFeatu
 					*VertexShader);
 
 				RHICmdList.EndRenderPass();
-				RHICmdList.CopyToResolveTarget(FilteredCube.TargetableTexture, FilteredCube.ShaderResourceTexture, FResolveParams());
 			}
 		}
 
-		RHICmdList.TransitionResource(EResourceTransitionAccess::EReadable, FilteredCube.TargetableTexture);
+		RHICmdList.CopyToResolveTarget(FilteredCube.TargetableTexture, FilteredCube.ShaderResourceTexture, FResolveParams());
 	}
 }
 
@@ -1637,6 +1636,8 @@ void CopyToSkyTexture(FRHICommandList& RHICmdList, FScene* Scene, FTexture* Proc
 		CopyInfo.NumSlices = 6;
 		CopyInfo.NumMips = NumMips;
 
+		RHICmdList.TransitionResource(EResourceTransitionAccess::EWritable, ProcessedTexture->TextureRHI);
+
 		RHICmdList.CopyTexture(FilteredCube.ShaderResourceTexture, ProcessedTexture->TextureRHI, CopyInfo);
 
 		RHICmdList.TransitionResource(EResourceTransitionAccess::EReadable, ProcessedTexture->TextureRHI);
