Crash fix when visualizing render targets

diff --git a/Engine/Source/Runtime/Renderer/Private/VisualizeTexturePresent.cpp b/Engine/Source/Runtime/Renderer/Private/VisualizeTexturePresent.cpp
--- a/Engine/Source/Runtime/Renderer/Private/VisualizeTexturePresent.cpp
+++ b/Engine/Source/Runtime/Renderer/Private/VisualizeTexturePresent.cpp
@@ -381,7 +381,7 @@ void FVisualizeTexturePresent::PresentContent(FRHICommandListImmediate& RHICmdLi
 		{
 			FRDGBuilder GraphBuilder(RHICmdList);
 
-			TRefCountPtr<IPooledRenderTarget> ElementRefCount;
+			TRefCountPtr<IPooledRenderTarget> ElementRefCount = Element;
 
 			GVisualizeTexture.CreateContentCapturePass(GraphBuilder, GraphBuilder.RegisterExternalTexture(ElementRefCount));
 			GraphBuilder.Execute();
