Use RenderPass instead of set rendertarget when using DisplayInternals.

diff --git a/Engine/Source/Runtime/Renderer/Private/SceneRendering.cpp b/Engine/Source/Runtime/Renderer/Private/SceneRendering.cpp
--- a/Engine/Source/Runtime/Renderer/Private/SceneRendering.cpp
+++ b/Engine/Source/Runtime/Renderer/Private/SceneRendering.cpp
@@ -3619,10 +3619,9 @@ static void DisplayInternals(FRHICommandListImmediate& RHICmdList, FViewInfo& In
 		FCanvas Canvas((FRenderTarget*)Family->RenderTarget, NULL, Family->CurrentRealTime, Family->CurrentWorldTime, Family->DeltaWorldTime, InView.GetFeatureLevel());
 		Canvas.SetRenderTargetRect(FIntRect(0, 0, Family->RenderTarget->GetSizeXY().X, Family->RenderTarget->GetSizeXY().Y));
 
-		PRAGMA_DISABLE_DEPRECATION_WARNINGS
-		FRHIRenderTargetView RTV(Family->RenderTarget->GetRenderTargetTexture(), ERenderTargetLoadAction::ELoad);
-		RHICmdList.SetRenderTargets(1, &RTV, nullptr, 0, nullptr);
-		PRAGMA_ENABLE_DEPRECATION_WARNINGS
+
+		FRHIRenderPassInfo RenderPassInfo(Family->RenderTarget->GetRenderTargetTexture(), ERenderTargetActions::Load_Store);
+		RHICmdList.BeginRenderPass(RenderPassInfo, TEXT("DisplayInternalsRenderPass"));
 
 		// further down to not intersect with "LIGHTING NEEDS TO BE REBUILT"
 		FVector2D Pos(30, 140);
@@ -3709,9 +3708,11 @@ static void DisplayInternals(FRHICommandListImmediate& RHICmdList, FViewInfo& In
 #undef CANVAS_LINE
 #undef CANVAS_HEADER
 
+		RHICmdList.EndRenderPass();
 		Canvas.Flush_RenderThread(RHICmdList);
 	}
 #endif
+
 }
 
 TSharedRef<ISceneViewExtension, ESPMode::ThreadSafe> GetRendererViewExtension()
