Don't crash on VT feedback size of 0
Copied from Dev-Rendering CL 7097298

diff --git a/Engine/Source/Runtime/Renderer/Private/PostProcess/SceneRenderTargets.cpp b/Engine/Source/Runtime/Renderer/Private/PostProcess/SceneRenderTargets.cpp
--- a/Engine/Source/Runtime/Renderer/Private/PostProcess/SceneRenderTargets.cpp
+++ b/Engine/Source/Runtime/Renderer/Private/PostProcess/SceneRenderTargets.cpp
@@ -2439,7 +2439,7 @@ void FSceneRenderTargets::AllocateDeferredShadingPathRenderTargets(FRHICommandLi
 
 	if (UseVirtualTexturing(CurrentFeatureLevel))
 	{
-		FIntPoint FeedbackSize = FIntPoint::DivideAndRoundUp(BufferSize, GVirtualTextureFeedbackFactor);
+		FIntPoint FeedbackSize = FIntPoint::DivideAndRoundUp(BufferSize, FMath::Max(GVirtualTextureFeedbackFactor, 1));
 		VirtualTextureFeedback.CreateResourceGPU(RHICmdList, FeedbackSize);
 	}
 }
