Fix missing memory stat decrements in VT system
Integrating from //UE4/Dev-Rendering
#rnx

diff --git a/Engine/Source/Runtime/Renderer/Private/VT/VirtualTextureSystem.cpp b/Engine/Source/Runtime/Renderer/Private/VT/VirtualTextureSystem.cpp
--- a/Engine/Source/Runtime/Renderer/Private/VT/VirtualTextureSystem.cpp
+++ b/Engine/Source/Runtime/Renderer/Private/VT/VirtualTextureSystem.cpp
@@ -213,6 +213,7 @@ FVirtualTextureSystem::~FVirtualTextureSystem()
 		if (Space)
 		{
 			check(Space->GetRefCount() == 0u);
+			DEC_MEMORY_STAT_BY(STAT_TotalPagetableMemory, Space->GetSizeInBytes());
 			BeginReleaseResource(Space);
 		}
 	}
@@ -220,6 +221,7 @@ FVirtualTextureSystem::~FVirtualTextureSystem()
 	{
 		FVirtualTexturePhysicalSpace* PhysicalSpace = PhysicalSpaces[i].Get();
 		check(PhysicalSpace->GetRefCount() == 0u);
+		DEC_MEMORY_STAT_BY(STAT_TotalPhysicalMemory, PhysicalSpace->GetSizeInBytes());
 		BeginReleaseResource(PhysicalSpace);
 	}
 }
@@ -533,6 +535,7 @@ void FVirtualTextureSystem::ReleaseSpace(FVirtualTextureSpace* Space)
 	{
 		// Private spaces are destroyed when ref count reaches 0
 		// This can only happen on render thread, so we can call ReleaseResource() directly and then delete the pointer immediately
+		DEC_MEMORY_STAT_BY(STAT_TotalPagetableMemory, Space->GetSizeInBytes());
 		Space->ReleaseResource();
 		Spaces[Space->GetID()].Release();
 	}
