Integrated from //UE4/Dev-Rendering
#rnx

diff --git a/Engine/Source/Runtime/Renderer/Private/VT/VirtualTextureProducer.h b/Engine/Source/Runtime/Renderer/Private/VT/VirtualTextureProducer.h
--- a/Engine/Source/Runtime/Renderer/Private/VT/VirtualTextureProducer.h
+++ b/Engine/Source/Runtime/Renderer/Private/VT/VirtualTextureProducer.h
@@ -19,6 +19,7 @@ public:
 
 	inline const FVTProducerDescription& GetDescription() const { return Description; }
 	inline IVirtualTexture* GetVirtualTexture() const { return VirtualTexture; }
+	inline const FName& GetName() const { return Description.Name; }
 	inline uint32 GetNumLayers() const { return Description.NumLayers; }
 	inline uint32 GetWidthInTiles() const { return Description.BlockWidthInTiles * Description.WidthInBlocks; }
 	inline uint32 GetHeightInTiles() const { return Description.BlockHeightInTiles * Description.HeightInBlocks; }
diff --git a/Engine/Source/Runtime/Renderer/Private/VT/VirtualTextureSystem.cpp b/Engine/Source/Runtime/Renderer/Private/VT/VirtualTextureSystem.cpp
--- a/Engine/Source/Runtime/Renderer/Private/VT/VirtualTextureSystem.cpp
+++ b/Engine/Source/Runtime/Renderer/Private/VT/VirtualTextureSystem.cpp
@@ -1654,10 +1654,12 @@ void FVirtualTextureSystem::SubmitRequests(FRHICommandListImmediate& RHICmdList,
 						else
 						{
 							// Fill in pAddress for layers that are already resident
-							FVirtualTexturePhysicalSpace* RESTRICT PhysicalSpace = Producer.GetPhysicalSpace(LocalLayerIndex);
-							FTexturePagePool& RESTRICT PagePool = PhysicalSpace->GetPagePool();
+							const FVirtualTexturePhysicalSpace* RESTRICT PhysicalSpace = Producer.GetPhysicalSpace(LocalLayerIndex);
+							const FTexturePagePool& RESTRICT PagePool = PhysicalSpace->GetPagePool();
 							const uint32 pAddress = PagePool.FindPageAddress(ProducerHandle, LocalLayerIndex, TileToLoad.Local_vAddress, TileToLoad.Local_vLevel);
-							check(pAddress != ~0u);
+							checkf(pAddress != ~0u,
+								TEXT("%s missing tile: LayerMask: %X, Layer %d, vAddress %06X, vLevel %d"),
+								*Producer.GetName().ToString(), LocalLayerMask, LocalLayerIndex, TileToLoad.Local_vAddress, TileToLoad.Local_vLevel);
 							ProduceTarget[LocalLayerIndex].TextureRHI = PhysicalSpace->GetPhysicalTexture();
 							ProduceTarget[LocalLayerIndex].pPageLocation = PhysicalSpace->GetPhysicalLocation(pAddress);
 						}
