Fix MSAA-count discrepancy when running with HMD-allocated depth targets (should allow render-target pool to allocate the MSAA depth in this case).
Fix was to only defer texture creation in the render target pool when the HMD supplies depth, and there's no MSAA.

diff --git a/Engine/Source/Runtime/Renderer/Private/PostProcess/SceneRenderTargets.cpp b/Engine/Source/Runtime/Renderer/Private/PostProcess/SceneRenderTargets.cpp
--- a/Engine/Source/Runtime/Renderer/Private/PostProcess/SceneRenderTargets.cpp
+++ b/Engine/Source/Runtime/Renderer/Private/PostProcess/SceneRenderTargets.cpp
@@ -2210,7 +2210,10 @@ void FSceneRenderTargets::AllocateCommonDepthTargets(FRHICommandList& RHICmdList
 		FPooledRenderTargetDesc Desc(FPooledRenderTargetDesc::Create2DDesc(BufferSize, PF_DepthStencil, DefaultDepthClear, TexCreate_None, TexCreate_DepthStencilTargetable | TexCreate_ShaderResource | TexCreate_InputAttachmentRead, false));
 		Desc.NumSamples = GetNumSceneColorMSAASamples(CurrentFeatureLevel);
 		Desc.Flags |= GFastVRamConfig.SceneDepth;
-		GRenderTargetPool.FindFreeElement(RHICmdList, Desc, SceneDepthZ, TEXT("SceneDepthZ"), true, ERenderTargetTransience::Transient, bHMDAllocatedDepthTarget);
+
+		// Only defer texture allocation if we're an HMD-allocated target, and we're not MSAA.
+		const bool bDeferTextureAllocation = bHMDAllocatedDepthTarget && Desc.NumSamples == 1;
+		GRenderTargetPool.FindFreeElement(RHICmdList, Desc, SceneDepthZ, TEXT("SceneDepthZ"), true, ERenderTargetTransience::Transient, bDeferTextureAllocation);
 
 		if (SceneDepthZ && bHMDAllocatedDepthTarget)
 		{
