Cherry-pick 7112410: Adding collateral profile scopes back into ray tracing ambient occlusion and global illumination passes.

diff --git a/Engine/Source/Runtime/Renderer/Private/RayTracing/RayTracingAmbientOcclusion.cpp b/Engine/Source/Runtime/Renderer/Private/RayTracing/RayTracingAmbientOcclusion.cpp
--- a/Engine/Source/Runtime/Renderer/Private/RayTracing/RayTracingAmbientOcclusion.cpp
+++ b/Engine/Source/Runtime/Renderer/Private/RayTracing/RayTracingAmbientOcclusion.cpp
@@ -175,6 +175,9 @@ void FDeferredShadingSceneRenderer::RenderRayTracingAmbientOcclusion(
 	TRefCountPtr<IPooledRenderTarget>& AmbientOcclusionMaskRT
 )
 {
+	RDG_GPU_STAT_SCOPE(GraphBuilder, RayTracingAmbientOcclusion);
+	RDG_EVENT_SCOPE(GraphBuilder, "Ray Tracing Ambient Occlusion");
+
 	FSceneRenderTargets& SceneContext = FSceneRenderTargets::Get(RHICmdList);
 
 	// Build RTAO parameters
diff --git a/Engine/Source/Runtime/Renderer/Private/RayTracing/RayTracingGlobalIllumination.cpp b/Engine/Source/Runtime/Renderer/Private/RayTracing/RayTracingGlobalIllumination.cpp
--- a/Engine/Source/Runtime/Renderer/Private/RayTracing/RayTracingGlobalIllumination.cpp
+++ b/Engine/Source/Runtime/Renderer/Private/RayTracing/RayTracingGlobalIllumination.cpp
@@ -352,6 +352,9 @@ void FDeferredShadingSceneRenderer::RenderRayTracingGlobalIllumination(
 		GRenderTargetPool.FindFreeElement(RHICmdList, Desc, GlobalIlluminationRT, TEXT("RayTracingGlobalIllumination"));
 	}
 
+	RDG_GPU_STAT_SCOPE(GraphBuilder, RayTracingGlobalIllumination);
+	RDG_EVENT_SCOPE(GraphBuilder, "Ray Tracing Global Illumination");
+
 	FRDGBuilder GraphBuilder(RHICmdList);
 
 	FRDGTextureRef GlobalIlluminationTexture;
