Release VT physical pools when they are no longer referenced
Integrated from //UE4/Dev-Rendering
#rnx

diff --git a/Engine/Source/Runtime/Renderer/Private/VT/VirtualTextureSystem.cpp b/Engine/Source/Runtime/Renderer/Private/VT/VirtualTextureSystem.cpp
--- a/Engine/Source/Runtime/Renderer/Private/VT/VirtualTextureSystem.cpp
+++ b/Engine/Source/Runtime/Renderer/Private/VT/VirtualTextureSystem.cpp
@@ -220,9 +220,12 @@ FVirtualTextureSystem::~FVirtualTextureSystem()
 	for(int i = 0; i < PhysicalSpaces.Num(); ++i)
 	{
 		FVirtualTexturePhysicalSpace* PhysicalSpace = PhysicalSpaces[i].Get();
-		check(PhysicalSpace->GetRefCount() == 0u);
-		DEC_MEMORY_STAT_BY(STAT_TotalPhysicalMemory, PhysicalSpace->GetSizeInBytes());
-		BeginReleaseResource(PhysicalSpace);
+		if (PhysicalSpace)
+		{
+			check(PhysicalSpace->GetRefCount() == 0u);
+			DEC_MEMORY_STAT_BY(STAT_TotalPhysicalMemory, PhysicalSpace->GetSizeInBytes());
+			BeginReleaseResource(PhysicalSpace);
+		}
 	}
 }
 
@@ -254,24 +257,26 @@ void FVirtualTextureSystem::ListPhysicalPoolsFromConsole()
 {
 	for(int i = 0; i < PhysicalSpaces.Num(); ++i)
 	{
-		const FVirtualTexturePhysicalSpace& PhysicalSpace = *PhysicalSpaces[i];
-		const FVTPhysicalSpaceDescription& Desc = PhysicalSpace.GetDescription();
-		const FTexturePagePool& PagePool = PhysicalSpace.GetPagePool();
-		const uint32 TotalSizeInBytes = PhysicalSpace.GetSizeInBytes();
-
-		UE_LOG(LogConsoleResponse, Display, TEXT("PhysicaPool: [%i] PF_%s %ix%i:"), i, GPixelFormats[Desc.Format].Name, Desc.TileSize, Desc.TileSize);
-		UE_LOG(LogConsoleResponse, Display, TEXT("  SizeInMegabyte= %f"), (float)TotalSizeInBytes / 1024.0f / 1024.0f);
-		UE_LOG(LogConsoleResponse, Display, TEXT("  Dimensions= %ix%i"), PhysicalSpace.GetTextureSize(), PhysicalSpace.GetTextureSize());
-		UE_LOG(LogConsoleResponse, Display, TEXT("  Tiles= %i"), PhysicalSpace.GetNumTiles());
-		UE_LOG(LogConsoleResponse, Display, TEXT("  Tiles Mapped= %i"), PagePool.GetNumMappedPages());
-
-		const int32 LockedTiles = PagePool.GetNumLockedPages();
-		const float LockedLoad = (float)LockedTiles / (float)PhysicalSpace.GetNumTiles();
-		const float LockedMemory = LockedLoad * TotalSizeInBytes / 1024.0f / 1024.0f;
-		UE_LOG(LogConsoleResponse, Display, TEXT("  Tiles Locked= %i (%fMB)"), LockedTiles, LockedMemory);
+		if (PhysicalSpaces[i].IsValid())
+		{
+			const FVirtualTexturePhysicalSpace& PhysicalSpace = *PhysicalSpaces[i];
+			const FVTPhysicalSpaceDescription& Desc = PhysicalSpace.GetDescription();
+			const FTexturePagePool& PagePool = PhysicalSpace.GetPagePool();
+			const uint32 TotalSizeInBytes = PhysicalSpace.GetSizeInBytes();
+
+			UE_LOG(LogConsoleResponse, Display, TEXT("PhysicaPool: [%i] PF_%s %ix%i:"), i, GPixelFormats[Desc.Format].Name, Desc.TileSize, Desc.TileSize);
+			UE_LOG(LogConsoleResponse, Display, TEXT("  SizeInMegabyte= %f"), (float)TotalSizeInBytes / 1024.0f / 1024.0f);
+			UE_LOG(LogConsoleResponse, Display, TEXT("  Dimensions= %ix%i"), PhysicalSpace.GetTextureSize(), PhysicalSpace.GetTextureSize());
+			UE_LOG(LogConsoleResponse, Display, TEXT("  Tiles= %i"), PhysicalSpace.GetNumTiles());
+			UE_LOG(LogConsoleResponse, Display, TEXT("  Tiles Mapped= %i"), PagePool.GetNumMappedPages());
+
+			const int32 LockedTiles = PagePool.GetNumLockedPages();
+			const float LockedLoad = (float)LockedTiles / (float)PhysicalSpace.GetNumTiles();
+			const float LockedMemory = LockedLoad * TotalSizeInBytes / 1024.0f / 1024.0f;
+			UE_LOG(LogConsoleResponse, Display, TEXT("  Tiles Locked= %i (%fMB)"), LockedTiles, LockedMemory);
+		}
 	}
 
-
 	for (int ID = 0; ID < 16; ID++)
 	{
 		const FVirtualTextureSpace* Space = Spaces[ID].Get();
@@ -546,17 +551,31 @@ FVirtualTexturePhysicalSpace* FVirtualTextureSystem::AcquirePhysicalSpace(const
 	for (int i = 0; i < PhysicalSpaces.Num(); ++i)
 	{
 		FVirtualTexturePhysicalSpace* PhysicalSpace = PhysicalSpaces[i].Get();
-		if (PhysicalSpace->GetDescription() == InDesc)
+		if (PhysicalSpace && PhysicalSpace->GetDescription() == InDesc)
 		{
 			PhysicalSpace->AddRef();
 			return PhysicalSpace;
 		}
 	}
 
-	const uint32 ID = PhysicalSpaces.Num();
+	uint32 ID = PhysicalSpaces.Num();
 	check(ID <= 0x0fff);
 
-	TUniquePtr<FVirtualTexturePhysicalSpace>& PhysicalSpace = PhysicalSpaces.AddDefaulted_GetRef();
+	for (int i = 0; i < PhysicalSpaces.Num(); ++i)
+	{
+		if (!PhysicalSpaces[i].IsValid())
+		{
+			ID = i;
+			break;
+		}
+	}
+
+	if (ID == PhysicalSpaces.Num())
+	{
+		PhysicalSpaces.AddDefaulted_GetRef();
+	}
+
+	TUniquePtr<FVirtualTexturePhysicalSpace>& PhysicalSpace = PhysicalSpaces[ID];
 	PhysicalSpace.Reset(new FVirtualTexturePhysicalSpace(InDesc, ID));
 	INC_MEMORY_STAT_BY(STAT_TotalPhysicalMemory, PhysicalSpace->GetSizeInBytes());
 	BeginInitResource(PhysicalSpace.Get());
@@ -567,8 +586,14 @@ FVirtualTexturePhysicalSpace* FVirtualTextureSystem::AcquirePhysicalSpace(const
 void FVirtualTextureSystem::ReleasePhysicalSpace(FVirtualTexturePhysicalSpace* Space)
 {
 	const uint32 NumRefs = Space->Release();
-	// Don't delete physical space when ref count hits 0, as they are likely to be reused/recreated in future
-	// Might need to have some mechanism to explicitly delete unreferenced spaces, or delete unreferenced spaces after some fixed number of frames
+	// Physical space is released when ref count hits 0
+	// Might need to have some mechanism to hold an extra reference if we know we will be recycling very soon (such when doing level reload)
+	if (NumRefs == 0)
+	{
+		DEC_MEMORY_STAT_BY(STAT_TotalPhysicalMemory, Space->GetSizeInBytes());
+		Space->ReleaseResource();
+		PhysicalSpaces[Space->GetID()] = nullptr;
+	}
 }
 
 void FVirtualTextureSystem::LockTile(const FVirtualTextureLocalTile& Tile)
@@ -773,10 +798,13 @@ void FVirtualTextureSystem::Update(FRHICommandListImmediate& RHICmdList, ERHIFea
 		for (int i = 0; i < PhysicalSpaces.Num(); ++i)
 		{
 			FVirtualTexturePhysicalSpace* PhysicalSpace = PhysicalSpaces[i].Get();
-			// Collect locked pages to be produced again
-			PhysicalSpace->GetPagePool().GetAllLockedPages(this, MappedTilesToProduce);
-			// Flush unlocked pages
-			PhysicalSpace->GetPagePool().EvictAllPages(this);
+			if (PhysicalSpace)
+			{
+				// Collect locked pages to be produced again
+				PhysicalSpace->GetPagePool().GetAllLockedPages(this, MappedTilesToProduce);
+				// Flush unlocked pages
+				PhysicalSpace->GetPagePool().EvictAllPages(this);
+			}
 		}
 
 		bFlushCaches = false;
@@ -1457,6 +1485,11 @@ void FVirtualTextureSystem::GatherRequestsTask(const FGatherRequestsParameters&
 
 	for (uint32 PhysicalSpaceID = 0u; PhysicalSpaceID < (uint32)PhysicalSpaces.Num(); ++PhysicalSpaceID)
 	{
+		if (PhysicalSpaces[PhysicalSpaceID].Get() == nullptr)
+		{
+			continue;
+		}
+
 		FVirtualTexturePhysicalSpace* RESTRICT PhysicalSpace = GetPhysicalSpace(PhysicalSpaceID);
 		FPageUpdateBuffer& RESTRICT Buffer = PageUpdateBuffers[PhysicalSpaceID];
 
diff --git a/Engine/Source/Runtime/Renderer/Private/VT/VirtualTextureSystem.h b/Engine/Source/Runtime/Renderer/Private/VT/VirtualTextureSystem.h
--- a/Engine/Source/Runtime/Renderer/Private/VT/VirtualTextureSystem.h
+++ b/Engine/Source/Runtime/Renderer/Private/VT/VirtualTextureSystem.h
@@ -54,7 +54,7 @@ public:
 	void ReleasePhysicalSpace(FVirtualTexturePhysicalSpace* Space);
 
 	FVirtualTextureSpace* GetSpace(uint8 ID) const { check(ID < MaxSpaces); return Spaces[ID].Get(); }
-	FVirtualTexturePhysicalSpace* GetPhysicalSpace(uint16 ID) const { return PhysicalSpaces[ID].Get(); }
+	FVirtualTexturePhysicalSpace* GetPhysicalSpace(uint16 ID) const { check(PhysicalSpaces[ID].IsValid());  return PhysicalSpaces[ID].Get(); }
 
 	void LockTile(const FVirtualTextureLocalTile& Tile);
 	void UnlockTile(const FVirtualTextureLocalTile& Tile, const FVirtualTextureProducer* Producer);
