Fixed crash in RTGI

diff --git a/Engine/Source/Runtime/Renderer/Private/RayTracing/RayTracingAmbientOcclusion.cpp b/Engine/Source/Runtime/Renderer/Private/RayTracing/RayTracingAmbientOcclusion.cpp
--- a/Engine/Source/Runtime/Renderer/Private/RayTracing/RayTracingAmbientOcclusion.cpp
+++ b/Engine/Source/Runtime/Renderer/Private/RayTracing/RayTracingAmbientOcclusion.cpp
@@ -116,9 +116,6 @@ void FDeferredShadingSceneRenderer::PrepareRayTracingAmbientOcclusion(const FVie
 
 void FDeferredShadingSceneRenderer::RenderRayTracingAmbientOcclusion(FRHICommandListImmediate& RHICmdList, TRefCountPtr<IPooledRenderTarget>& AmbientOcclusionRT)
 {
-	SCOPED_DRAW_EVENT(RHICmdList, RayTracingAmbientOcclusion);
-	SCOPED_GPU_STAT(RHICmdList, RayTracingAmbientOcclusion);
-
 	bool bAnyViewWithRTAO = false;
 	for (int32 ViewIndex = 0; ViewIndex < Views.Num(); ++ViewIndex)
 	{
diff --git a/Engine/Source/Runtime/Renderer/Private/RayTracing/RayTracingGlobalIllumination.cpp b/Engine/Source/Runtime/Renderer/Private/RayTracing/RayTracingGlobalIllumination.cpp
--- a/Engine/Source/Runtime/Renderer/Private/RayTracing/RayTracingGlobalIllumination.cpp
+++ b/Engine/Source/Runtime/Renderer/Private/RayTracing/RayTracingGlobalIllumination.cpp
@@ -354,9 +354,6 @@ void FDeferredShadingSceneRenderer::RenderRayTracingGlobalIllumination(
 
 	FRDGBuilder GraphBuilder(RHICmdList);
 
-	RDG_GPU_STAT_SCOPE(GraphBuilder, RayTracingGlobalIllumination);
-	RDG_EVENT_SCOPE(GraphBuilder, "Ray Tracing Global Illumination");
-
 	FRDGTextureRef GlobalIlluminationTexture;
 	{
 		FRDGTextureDesc Desc = SceneContext.GetSceneColor()->GetDesc();
