PR #5968: Ray Traced Shadow optimizations (Contributed by [REDACTED_PERSON])

diff --git a/Engine/Shaders/Private/RayTracing/RayTracingOcclusionRGS.usf b/Engine/Shaders/Private/RayTracing/RayTracingOcclusionRGS.usf
--- a/Engine/Shaders/Private/RayTracing/RayTracingOcclusionRGS.usf
+++ b/Engine/Shaders/Private/RayTracing/RayTracingOcclusionRGS.usf
@@ -28,6 +28,7 @@ uint LightingChannelMask;
 uint SamplesPerPixel;
 float NormalBias;
 int4 LightScissor;
+int2 PixelOffset;
 
 bool GenerateOcclusionRay(
 	FLightShaderParameters LightParameters,
@@ -113,8 +114,8 @@ bool GenerateOcclusionRay(
 [shader("raygeneration")]
 void OcclusionRGS()
 {
-	uint2 PixelCoord = DispatchRaysIndex().xy + View.ViewRectMin;
-	
+	uint2 PixelCoord = DispatchRaysIndex().xy + View.ViewRectMin + PixelOffset;
+
 	float Visibility = 0.0;
 	float SumRayDistance = 0.0;
 	float ClosestRayDistance = DENOISER_INVALID_HIT_DISTANCE;
@@ -219,6 +220,12 @@ void OcclusionRGS()
 	#if !ENABLE_TWO_SIDED_GEOMETRY
 			RayFlags |= RAY_FLAG_CULL_BACK_FACING_TRIANGLES;
 	#endif
+			
+			if (DIM_DENOISER_OUTPUT == 0)
+			{
+				// Denoiser mode 0 doesn't use depth, so take first hit
+				RayFlags |= RAY_FLAG_ACCEPT_FIRST_HIT_AND_END_SEARCH;
+			}
 
 			FPackedMaterialClosestHitPayload Payload = (FPackedMaterialClosestHitPayload)0;
 			Payload.SetMinimalPayloadMode(); // Only HitT will be filled by closest hit shader, skipping full material evaluation.
diff --git a/Engine/Source/Runtime/Renderer/Private/LightRendering.cpp b/Engine/Source/Runtime/Renderer/Private/LightRendering.cpp
--- a/Engine/Source/Runtime/Renderer/Private/LightRendering.cpp
+++ b/Engine/Source/Runtime/Renderer/Private/LightRendering.cpp
@@ -76,6 +76,11 @@ static TAutoConsoleVariable<int32> CVarMaxShadowDenoisingBatchSize(
 	TEXT("Maximum number of shadow to denoise at the same time."),
 	ECVF_RenderThreadSafe);
 
+	static TAutoConsoleVariable<int32> CVarMaxShadowRayTracingBatchSize(
+		TEXT("r.RayTracing.Shadow.MaxBatchSize"), 8,
+		TEXT("Maximum number of shadows to trace at the same time."),
+		ECVF_RenderThreadSafe);
+
 #if ENABLE_DEBUG_DISCARD_PROP
 static float GDebugLightDiscardProp = 0.0f;
 static FAutoConsoleVariableRef CVarDebugLightDiscardProp(
@@ -146,10 +151,10 @@ void StencilingGeometry::DrawSphere(FRHICommandList& RHICmdList)
 {
 	RHICmdList.SetStreamSource(0, StencilingGeometry::GStencilSphereVertexBuffer.VertexBufferRHI, 0);
 	RHICmdList.DrawIndexedPrimitive(StencilingGeometry::GStencilSphereIndexBuffer.IndexBufferRHI, 0, 0,
-		StencilingGeometry::GStencilSphereVertexBuffer.GetVertexCount(), 0, 
+		StencilingGeometry::GStencilSphereVertexBuffer.GetVertexCount(), 0,
 		StencilingGeometry::GStencilSphereIndexBuffer.GetIndexCount() / 3, 1);
 }
-		
+
 void StencilingGeometry::DrawVectorSphere(FRHICommandList& RHICmdList)
 {
 	RHICmdList.SetStreamSource(0, StencilingGeometry::GStencilSphereVectorBuffer.VertexBufferRHI, 0);
@@ -272,7 +277,7 @@ class FDeferredLightPS : public FGlobalShader
 		LightingChannelsTexture.Bind(Initializer.ParameterMap, TEXT("LightingChannelsTexture"));
 		LightingChannelsSampler.Bind(Initializer.ParameterMap, TEXT("LightingChannelsSampler"));
 		TransmissionProfilesTexture.Bind(Initializer.ParameterMap, TEXT("SSProfilesTexture"));
-		TransmissionProfilesLinearSampler.Bind(Initializer.ParameterMap, TEXT("TransmissionProfilesLinearSampler"));		
+		TransmissionProfilesLinearSampler.Bind(Initializer.ParameterMap, TEXT("TransmissionProfilesLinearSampler"));
 	}
 
 	FDeferredLightPS()
@@ -294,7 +299,7 @@ public:
 	}
 
 	virtual bool Serialize(FArchive& Ar) override
-	{		
+	{
 		bool bShaderHasOutdatedParameters = FGlobalShader::Serialize(Ar);
 		Ar << SceneTextureParameters;
 		Ar << LightAttenuationTexture;
@@ -324,7 +329,7 @@ private:
 		if(LightAttenuationTexture.IsBound())
 		{
 			SetTextureParameter(
-				RHICmdList, 
+				RHICmdList,
 				ShaderRHI,
 				LightAttenuationTexture,
 				LightAttenuationTextureSampler,
@@ -334,7 +339,7 @@ private:
 		}
 
 		SetTextureParameter(
-			RHICmdList, 
+			RHICmdList,
 			ShaderRHI,
 			LTCMatTexture,
 			LTCMatSampler,
@@ -343,7 +348,7 @@ private:
 			);
 
 		SetTextureParameter(
-			RHICmdList, 
+			RHICmdList,
 			ShaderRHI,
 			LTCAmpTexture,
 			LTCAmpSampler,
@@ -355,7 +360,7 @@ private:
 			FRHITexture* TextureRHI = IESTextureResource ? IESTextureResource->TextureRHI : GSystemTextures.WhiteDummy->GetRenderTargetItem().TargetableTexture;
 
 			SetTextureParameter(
-				RHICmdList, 
+				RHICmdList,
 				ShaderRHI,
 				IESTexture,
 				IESTextureSampler,
@@ -369,7 +374,7 @@ private:
 			FRHITexture* LightingChannelsTextureRHI = SceneRenderTargets.LightingChannels ? SceneRenderTargets.LightingChannels->GetRenderTargetItem().ShaderResourceTexture : GSystemTextures.WhiteDummy->GetRenderTargetItem().TargetableTexture;
 
 			SetTextureParameter(
-				RHICmdList, 
+				RHICmdList,
 				ShaderRHI,
 				LightingChannelsTexture,
 				LightingChannelsSampler,
@@ -458,7 +463,7 @@ public:
 	}
 
 	virtual bool Serialize(FArchive& Ar) override
-	{		
+	{
 		bool bShaderHasOutdatedParameters = FGlobalShader::Serialize(Ar);
 		Ar << HasValidChannel;
 		Ar << SceneTextureParameters;
@@ -625,7 +630,7 @@ void FDeferredShadingSceneRenderer::GatherAndSortLights(FSortedLightSetSceneInfo
 					// In the forward case one directional light gets put into its own variables, and in the deferred case it gets a full-screen pass.
 					// Usually it'll have shadows and stuff anyway.
 					// Rect lights are not supported as the performance impact is significant even if not used, for now, left for trad. deferred.
-					const bool bTiledOrClusteredDeferredSupported = 
+					const bool bTiledOrClusteredDeferredSupported =
 						!SortedLightInfo->SortKey.Fields.bTextureProfile &&
 						!SortedLightInfo->SortKey.Fields.bShadowed &&
 						!SortedLightInfo->SortKey.Fields.bLightFunction &&
@@ -754,7 +759,7 @@ void FDeferredShadingSceneRenderer::RenderLights(FRHICommandListImmediate& RHICm
 		SCOPED_DRAW_EVENT(RHICmdList, DirectLighting);
 
 		FSceneRenderTargets& SceneContext = FSceneRenderTargets::Get(RHICmdList);
-		
+
 		if (GbEnableAsyncComputeTranslucencyLightingVolumeClear && GSupportsEfficientAsyncCompute)
 		{
 			//Gfx pipe must wait for the async compute clear of the translucency volume clear.
@@ -770,7 +775,7 @@ void FDeferredShadingSceneRenderer::RenderLights(FRHICommandListImmediate& RHICm
 			int32 StandardDeferredStart = SortedLightSet.SimpleLightsEnd;
 
 			bool bRenderSimpleLightsStandardDeferred = SortedLightSet.SimpleLights.InstanceData.Num() > 0;
-			
+
 			UE_CLOG(ShouldUseClusteredDeferredShading() && !AreClusteredLightsInLightGrid(), LogRenderer, Warning,
 				TEXT("Clustered deferred shading is enabled, but lights were not injected in grid, falling back to other methods (hint 'r.LightCulling.Quality' may cause this)."));
 
@@ -825,7 +830,7 @@ void FDeferredShadingSceneRenderer::RenderLights(FRHICommandListImmediate& RHICm
 					const FSortedLightSceneInfo& SortedLightInfo = SortedLights[LightIndex];
 					const FLightSceneInfo* const LightSceneInfo = SortedLightInfo.LightSceneInfo;
 
-					// Render the light to the scene color buffer, using a 1x1 white texture as input 
+					// Render the light to the scene color buffer, using a 1x1 white texture as input
 					RenderLight(RHICmdList, LightSceneInfo, NULL, false, false);
 				}
 
@@ -861,7 +866,7 @@ void FDeferredShadingSceneRenderer::RenderLights(FRHICommandListImmediate& RHICm
 			}
 		}
 
-		EShaderPlatform ShaderPlatformForFeatureLevel = GShaderPlatformForFeatureLevel[FeatureLevel]; 
+		EShaderPlatform ShaderPlatformForFeatureLevel = GShaderPlatformForFeatureLevel[FeatureLevel];
 
 		if ( IsFeatureLevelSupported(ShaderPlatformForFeatureLevel, ERHIFeatureLevel::SM5) )
 		{
@@ -915,7 +920,7 @@ void FDeferredShadingSceneRenderer::RenderLights(FRHICommandListImmediate& RHICm
 										}
 									}
 								}
-							}					
+							}
 						}
 					}
 				}
@@ -936,209 +941,233 @@ void FDeferredShadingSceneRenderer::RenderLights(FRHICommandListImmediate& RHICm
 			TArray<TRefCountPtr<IPooledRenderTarget>> PreprocessedShadowMaskTextures;
 
 			const int32 MaxDenoisingBatchSize = FMath::Clamp(CVarMaxShadowDenoisingBatchSize.GetValueOnRenderThread(), 1, IScreenSpaceDenoiser::kMaxBatchSize);
+			const int32 MaxRTShadowBatchSize = CVarMaxShadowRayTracingBatchSize.GetValueOnRenderThread();
 			const bool bDoShadowDenoisingBatching = DenoiserMode != 0 && MaxDenoisingBatchSize > 1;
+			const bool dDoShadowBatching = bDoShadowDenoisingBatching || MaxRTShadowBatchSize > 1;
 
 			// Optimisations: batches all shadow ray tracing denoising. Definitely could be smarter to avoid high VGPR pressure if this entire
 			// function was converted to render graph, and want least intrusive change as possible. So right not it trades render target memory pressure
 			// for denoising perf.
-			if (RHI_RAYTRACING && bDoShadowDenoisingBatching)
+			if (RHI_RAYTRACING && dDoShadowBatching)
 			{
-				TStaticArray<IScreenSpaceDenoiser::FShadowParameters, IScreenSpaceDenoiser::kMaxBatchSize> DenoisingQueue;
-				TStaticArray<int32, IScreenSpaceDenoiser::kMaxBatchSize> LightIndices;
-
-				FRDGBuilder GraphBuilder(RHICmdList);
-
-				FSceneTextureParameters SceneTextures;
-				SetupSceneTextureParameters(GraphBuilder, &SceneTextures);
-
 				FViewInfo& View = Views[0];
 
-					// Allocate PreprocessedShadowMaskTextures once so QueueTextureExtraction can deferred write.
+				// Allocate PreprocessedShadowMaskTextures once so QueueTextureExtraction can deferred write.
+				{
+					if (!View.bViewStateIsReadOnly)
 					{
-						if (!View.bViewStateIsReadOnly)
-						{
-							View.ViewState->PrevFrameViewInfo.ShadowHistories.Empty();
-							View.ViewState->PrevFrameViewInfo.ShadowHistories.Reserve(SortedLights.Num());
-						}
-						PreprocessedShadowMaskTextures.Reserve(SortedLights.Num());
-						for (int32 LightIndex = AttenuationLightStart; LightIndex < SortedLights.Num(); LightIndex++)
-						{
-							const FSortedLightSceneInfo& SortedLightInfo = SortedLights[LightIndex];
-							const FLightSceneInfo& LightSceneInfo = *SortedLightInfo.LightSceneInfo;
-
-							PreprocessedShadowMaskTextures.Add(TRefCountPtr<IPooledRenderTarget>());
-
-							if (!View.bViewStateIsReadOnly)
-							{
-								View.ViewState->PrevFrameViewInfo.ShadowHistories.Emplace(LightSceneInfo.Proxy->GetLightComponent());
-							}
-						}
+						View.ViewState->PrevFrameViewInfo.ShadowHistories.Empty();
+						View.ViewState->PrevFrameViewInfo.ShadowHistories.Reserve(SortedLights.Num());
 					}
+				}
+			} // if (RHI_RAYTRACING)
 
-					TMap<const FLightSceneInfo*, FRDGTextureRef> ShadowMaskTextures;
+			bool bDirectLighting = ViewFamily.EngineShowFlags.DirectLighting;
 
-					// Lambda to share the code quicking of the shadow denoiser.
-					auto QuickOffDenoisingBatch = [&](){
-						int32 InputParameterCount = 0;
-						for (int32 i = 0; i < IScreenSpaceDenoiser::kMaxBatchSize; i++)
-						{
-							InputParameterCount += DenoisingQueue[i].LightSceneInfo != nullptr ? 1 : 0;
-						}
+			TRefCountPtr<IPooledRenderTarget> ScreenShadowMaskTexture;
 
-						check(InputParameterCount >= 1);
+			// Draw shadowed and light function lights
+			for (int32 LightIndex = AttenuationLightStart; LightIndex < SortedLights.Num(); LightIndex++)
+			{
+				const FSortedLightSceneInfo& SortedLightInfo = SortedLights[LightIndex];
+				const FLightSceneInfo& LightSceneInfo = *SortedLightInfo.LightSceneInfo;
 
-						TStaticArray<IScreenSpaceDenoiser::FShadowPenumbraOutputs, IScreenSpaceDenoiser::kMaxBatchSize> Outputs;
+				// Note: Skip shadow mask generation for rect light if direct illumination is computed
+				//		 stochastically (rather than analytically + shadow mask)
+				const bool bDrawShadows = SortedLightInfo.SortKey.Fields.bShadowed && !ShouldRenderRayTracingStochasticRectLight(LightSceneInfo);
+				bool bDrawLightFunction = SortedLightInfo.SortKey.Fields.bLightFunction;
+				bool bDrawPreviewIndicator = ViewFamily.EngineShowFlags.PreviewShadowsIndicator && !LightSceneInfo.IsPrecomputedLightingValid() && LightSceneInfo.Proxy->HasStaticShadowing();
+				bool bInjectedTranslucentVolume = false;
+				bool bUsedShadowMaskTexture = false;
 
-						RDG_EVENT_SCOPE(GraphBuilder, "%s%s(Shadow BatchSize=%d) %dx%d",
-							DenoiserToUse != DefaultDenoiser ? TEXT("ThirdParty ") : TEXT(""),
-							DenoiserToUse->GetDebugName(),
-							InputParameterCount,
-							View.ViewRect.Width(), View.ViewRect.Height());
+				FScopeCycleCounter Context(LightSceneInfo.Proxy->GetStatId());
 
-						DenoiserToUse->DenoiseShadows(
-							GraphBuilder,
-							View,
-							&View.PrevViewInfo,
-							SceneTextures,
-							DenoisingQueue,
-							InputParameterCount,
-							Outputs);
+				if ((bDrawShadows || bDrawLightFunction || bDrawPreviewIndicator) && !ScreenShadowMaskTexture.IsValid())
+				{
+					SceneContext.AllocateScreenShadowMask(RHICmdList, ScreenShadowMaskTexture);
+				}
 
-						for (int32 i = 0; i < InputParameterCount; i++)
-						{
-							const FLightSceneInfo* LightSceneInfo = DenoisingQueue[i].LightSceneInfo;
+				FString LightNameWithLevel;
+				GetLightNameForDrawEvent(LightSceneInfo.Proxy, LightNameWithLevel);
+				SCOPED_DRAW_EVENTF(RHICmdList, EventLightPass, *LightNameWithLevel);
 
-							int32 LightIndex = LightIndices[i];
-							TRefCountPtr<IPooledRenderTarget>* RefDestination = &PreprocessedShadowMaskTextures[LightIndex - AttenuationLightStart];
-							check(*RefDestination == nullptr);
+				if (bDrawShadows)
+				{
+					INC_DWORD_STAT(STAT_NumShadowedLights);
 
-							GraphBuilder.QueueTextureExtraction(Outputs[i].DiffusePenumbra, RefDestination);
-							DenoisingQueue[i].LightSceneInfo = nullptr;
-						}
-					}; // QuickOffDenoisingBatch
+					const FLightOcclusionType OcclusionType = GetLightOcclusionType(*LightSceneInfo.Proxy);
 
-					// Ray trace shadows of light that needs, and quick off denoising batch.
-					for (int32 LightIndex = AttenuationLightStart; LightIndex < SortedLights.Num(); LightIndex++)
+					// Inline ray traced shadow batching, launches shadow batches when needed
+					// reduces memory overhead while keeping shadows batched to optimize costs
 					{
-						const FSortedLightSceneInfo& SortedLightInfo = SortedLights[LightIndex];
-						const FLightSceneInfo& LightSceneInfo = *SortedLightInfo.LightSceneInfo;
+						FViewInfo& View = Views[0];
 
-						// Denoiser do not support texture rect light important sampling.
-						const bool bDrawShadows = SortedLightInfo.SortKey.Fields.bShadowed && !ShouldRenderRayTracingStochasticRectLight(LightSceneInfo);
+						IScreenSpaceDenoiser::FShadowRayTracingConfig RayTracingConfig;
+						RayTracingConfig.RayCountPerPixel = LightSceneInfo.Proxy->GetSamplesPerPixel();
 
-						if (!bDrawShadows)
-							continue;
+						const bool bDenoiserCompatible = !LightRequiresDenosier(LightSceneInfo) || IScreenSpaceDenoiser::EShadowRequirements::PenumbraAndClosestOccluder == DenoiserToUse->GetShadowRequirements(View, LightSceneInfo, RayTracingConfig);
 
-						INC_DWORD_STAT(STAT_NumShadowedLights);
+						const bool bWantsBatchedShadow = OcclusionType == FLightOcclusionType::Raytraced && bDenoiserCompatible &&
+							SortedLightInfo.SortKey.Fields.bShadowed && !ShouldRenderRayTracingStochasticRectLight(LightSceneInfo);
 
-						const FLightOcclusionType OcclusionType = GetLightOcclusionType(*LightSceneInfo.Proxy);
-						if (OcclusionType != FLightOcclusionType::Raytraced)
-							continue;
+						// determine if this light doesn't yet have a precomuted shadow and execute a batch to amortize costs if one is needed
+						if (RHI_RAYTRACING && PreprocessedShadowMaskTextures.Num() > 0 && bWantsBatchedShadow && !PreprocessedShadowMaskTextures[LightIndex - AttenuationLightStart])
+						{
+							SCOPED_DRAW_EVENT(RHICmdList, ShadowBatch);
+							TStaticArray<IScreenSpaceDenoiser::FShadowParameters, IScreenSpaceDenoiser::kMaxBatchSize> DenoisingQueue;
+							TStaticArray<int32, IScreenSpaceDenoiser::kMaxBatchSize> LightIndices;
 
-						if (!LightRequiresDenosier(LightSceneInfo))
-							continue;
+							FRDGBuilder GraphBuilder(RHICmdList);
 
-						IScreenSpaceDenoiser::FShadowRayTracingConfig RayTracingConfig;
-						RayTracingConfig.RayCountPerPixel = LightSceneInfo.Proxy->GetSamplesPerPixel();
+							FSceneTextureParameters SceneTextures;
+							SetupSceneTextureParameters(GraphBuilder, &SceneTextures);
 
-						IScreenSpaceDenoiser::EShadowRequirements DenoiserRequirements = DenoiserToUse->GetShadowRequirements(
-							View, LightSceneInfo, RayTracingConfig);
+							int32 ProcessShadows = 0;
 
-						// Not worth batching and increase memory pressure if the denoiser do not support this ray tracing config.
-						// TODO: add suport for batch with multiple SPP.
-						if (DenoiserRequirements != IScreenSpaceDenoiser::EShadowRequirements::PenumbraAndClosestOccluder)
-						{
-							continue;
-						}
+							// Lambda to share the code quicking of the shadow denoiser.
+							auto QuickOffDenoisingBatch = [&]() {
+								int32 InputParameterCount = 0;
+								for (int32 i = 0; i < IScreenSpaceDenoiser::kMaxBatchSize; i++)
+								{
+									InputParameterCount += DenoisingQueue[i].LightSceneInfo != nullptr ? 1 : 0;
+								}
 
-					// Ray trace the shadow.
-					FRDGTextureRef ShadowMask;
-					FRDGTextureRef RayHitDistance;
-						{
-							FString LightNameWithLevel;
-							GetLightNameForDrawEvent(LightSceneInfo.Proxy, LightNameWithLevel);
-							RDG_EVENT_SCOPE(GraphBuilder, "%s", *LightNameWithLevel);
+								check(InputParameterCount >= 1);
 
-							RenderRayTracingShadows(
-								GraphBuilder,
-								SceneTextures,
-								View,
-								LightSceneInfo,
-								RayTracingConfig,
-								DenoiserRequirements,
-							&ShadowMask,
-							&RayHitDistance);
-						}
+								TStaticArray<IScreenSpaceDenoiser::FShadowPenumbraOutputs, IScreenSpaceDenoiser::kMaxBatchSize> Outputs;
 
-						// Queue the ray tracing output for shadow denoising.
-						for (int32 i = 0; i < IScreenSpaceDenoiser::kMaxBatchSize; i++)
-						{
-							if (DenoisingQueue[i].LightSceneInfo == nullptr)
-							{
-								DenoisingQueue[i].LightSceneInfo = &LightSceneInfo;
-								DenoisingQueue[i].RayTracingConfig = RayTracingConfig;
-								DenoisingQueue[i].InputTextures.Penumbra = ShadowMask;
-								DenoisingQueue[i].InputTextures.ClosestOccluder = RayHitDistance;
-								LightIndices[i] = LightIndex;
-
-								// If queue for this light type is full, quick of the batch.
-								if ((i + 1) == MaxDenoisingBatchSize)
+								RDG_EVENT_SCOPE(GraphBuilder, "%s%s(Shadow BatchSize=%d) %dx%d",
+									DenoiserToUse != DefaultDenoiser ? TEXT("ThirdParty ") : TEXT(""),
+									DenoiserToUse->GetDebugName(),
+									InputParameterCount,
+									View.ViewRect.Width(), View.ViewRect.Height());
+
+								DenoiserToUse->DenoiseShadows(
+									GraphBuilder,
+									View,
+									&View.PrevViewInfo,
+									SceneTextures,
+									DenoisingQueue,
+									InputParameterCount,
+									Outputs);
+
+								for (int32 i = 0; i < InputParameterCount; i++)
 								{
-									QuickOffDenoisingBatch();
+									const FLightSceneInfo* LightSceneInfo = DenoisingQueue[i].LightSceneInfo;
+
+									int32 LightIndex = LightIndices[i];
+									TRefCountPtr<IPooledRenderTarget>* RefDestination = &PreprocessedShadowMaskTextures[LightIndex - AttenuationLightStart];
+									check(*RefDestination == nullptr);
+
+									GraphBuilder.QueueTextureExtraction(Outputs[i].DiffusePenumbra, RefDestination);
+									DenoisingQueue[i].LightSceneInfo = nullptr;
 								}
-								break;
-							}
-							else
+							}; // QuickOffDenoisingBatch
+
+							// Ray trace shadows of light that needs, and quick off denoising batch.
+							for (int32 LightBatchIndex = LightIndex; LightBatchIndex < SortedLights.Num(); LightBatchIndex++)
 							{
-								check((i - 1) < IScreenSpaceDenoiser::kMaxBatchSize);
-							}
-						}
+								const FSortedLightSceneInfo& BatchSortedLightInfo = SortedLights[LightBatchIndex];
+								const FLightSceneInfo& BatchLightSceneInfo = *BatchSortedLightInfo.LightSceneInfo;
 
-					} // for (int32 LightIndex = AttenuationLightStart; LightIndex < SortedLights.Num(); LightIndex++)
+								// Denoiser do not support texture rect light important sampling.
+								const bool bBatchDrawShadows = BatchSortedLightInfo.SortKey.Fields.bShadowed && !ShouldRenderRayTracingStochasticRectLight(BatchLightSceneInfo);
 
-					// Ensures all denoising queues are processed.
-					if (DenoisingQueue[0].LightSceneInfo)
-					{
-						QuickOffDenoisingBatch();
-					}
+								if (!bBatchDrawShadows)
+									continue;
 
-				GraphBuilder.Execute();
-			} // if (RHI_RAYTRACING)
+								const FLightOcclusionType BatchOcclusionType = GetLightOcclusionType(*BatchLightSceneInfo.Proxy);
+								if (BatchOcclusionType != FLightOcclusionType::Raytraced)
+									continue;
 
-			bool bDirectLighting = ViewFamily.EngineShowFlags.DirectLighting;
+								const bool bRequiresDenoiser = LightRequiresDenosier(BatchLightSceneInfo);
 
-			TRefCountPtr<IPooledRenderTarget> ScreenShadowMaskTexture;
+								IScreenSpaceDenoiser::FShadowRayTracingConfig BatchRayTracingConfig;
+								BatchRayTracingConfig.RayCountPerPixel = BatchLightSceneInfo.Proxy->GetSamplesPerPixel();
 
-			// Draw shadowed and light function lights
-			for (int32 LightIndex = AttenuationLightStart; LightIndex < SortedLights.Num(); LightIndex++)
-			{
-				const FSortedLightSceneInfo& SortedLightInfo = SortedLights[LightIndex];
-				const FLightSceneInfo& LightSceneInfo = *SortedLightInfo.LightSceneInfo;
+								IScreenSpaceDenoiser::EShadowRequirements DenoiserRequirements = bRequiresDenoiser ?
+									DenoiserToUse->GetShadowRequirements(View, BatchLightSceneInfo, BatchRayTracingConfig) :
+									IScreenSpaceDenoiser::EShadowRequirements::Bailout;
 
-				// Note: Skip shadow mask generation for rect light if direct illumination is computed 
-				//		 stochastically (rather than analytically + shadow mask)
-				const bool bDrawShadows = SortedLightInfo.SortKey.Fields.bShadowed && !ShouldRenderRayTracingStochasticRectLight(LightSceneInfo);
-				bool bDrawLightFunction = SortedLightInfo.SortKey.Fields.bLightFunction;
-				bool bDrawPreviewIndicator = ViewFamily.EngineShowFlags.PreviewShadowsIndicator && !LightSceneInfo.IsPrecomputedLightingValid() && LightSceneInfo.Proxy->HasStaticShadowing();
-				bool bInjectedTranslucentVolume = false;
-				bool bUsedShadowMaskTexture = false;
-				
-				FScopeCycleCounter Context(LightSceneInfo.Proxy->GetStatId());
+								// Not worth batching and increase memory pressure if the denoiser do not support this ray tracing config.
+								// TODO: add suport for batch with multiple SPP.
+								if (bRequiresDenoiser && DenoiserRequirements != IScreenSpaceDenoiser::EShadowRequirements::PenumbraAndClosestOccluder)
+								{
+									continue;
+								}
 
-				if ((bDrawShadows || bDrawLightFunction || bDrawPreviewIndicator) && !ScreenShadowMaskTexture.IsValid())
-				{
-					SceneContext.AllocateScreenShadowMask(RHICmdList, ScreenShadowMaskTexture);
-				}
+								// Ray trace the shadow.
+								FRDGTextureRef ShadowMask;
+								FRDGTextureRef RayHitDistance;
+								{
+									FString BatchLightNameWithLevel;
+									GetLightNameForDrawEvent(BatchLightSceneInfo.Proxy, BatchLightNameWithLevel);
+									RDG_EVENT_SCOPE(GraphBuilder, "%s", *BatchLightNameWithLevel);
+
+									RenderRayTracingShadows(
+										GraphBuilder,
+										SceneTextures,
+										View,
+										BatchLightSceneInfo,
+										BatchRayTracingConfig,
+										DenoiserRequirements,
+										&ShadowMask,
+										&RayHitDistance);
+								}
 
-				FString LightNameWithLevel;
-				GetLightNameForDrawEvent(LightSceneInfo.Proxy, LightNameWithLevel);
-				SCOPED_DRAW_EVENTF(RHICmdList, EventLightPass, *LightNameWithLevel);
+								bool bBatchFull = false;
 
-				if (bDrawShadows)
-				{
-					INC_DWORD_STAT(STAT_NumShadowedLights);
+								if (bRequiresDenoiser)
+								{
+									// Queue the ray tracing output for shadow denoising.
+									for (int32 i = 0; i < IScreenSpaceDenoiser::kMaxBatchSize; i++)
+									{
+										if (DenoisingQueue[i].LightSceneInfo == nullptr)
+										{
+											DenoisingQueue[i].LightSceneInfo = &BatchLightSceneInfo;
+											DenoisingQueue[i].RayTracingConfig = RayTracingConfig;
+											DenoisingQueue[i].InputTextures.Penumbra = ShadowMask;
+											DenoisingQueue[i].InputTextures.ClosestOccluder = RayHitDistance;
+											LightIndices[i] = LightBatchIndex;
+
+											// If queue for this light type is full, quick of the batch.
+											if ((i + 1) == MaxDenoisingBatchSize)
+											{
+												QuickOffDenoisingBatch();
+												bBatchFull = true;
+											}
+											break;
+										}
+										else
+										{
+											check((i - 1) < IScreenSpaceDenoiser::kMaxBatchSize);
+										}
+									}
+								}
+								else
+								{
+									GraphBuilder.QueueTextureExtraction(ShadowMask, &PreprocessedShadowMaskTextures[LightBatchIndex - AttenuationLightStart]);
+								}
+
+								// terminate batch if we filled a denoiser batch or hit our max light batch
+								ProcessShadows++;
+								if (bBatchFull || ProcessShadows == MaxRTShadowBatchSize)
+								{
+									break;
+								}
+
+							} // for (int32 LightBatchIndex = LightIndex; LightIndex < SortedLights.Num(); LightIndex++)
+
+							// Ensures all denoising queues are processed.
+							if (DenoisingQueue[0].LightSceneInfo)
+							{
+								QuickOffDenoisingBatch();
+							}
+
+							GraphBuilder.Execute();
+						}
+					} // end inline batched raytraced shadow
 
-					const FLightOcclusionType OcclusionType = GetLightOcclusionType(*LightSceneInfo.Proxy);
 					if (RHI_RAYTRACING && PreprocessedShadowMaskTextures.Num() > 0 && PreprocessedShadowMaskTextures[LightIndex - AttenuationLightStart])
 					{
 						ScreenShadowMaskTexture = PreprocessedShadowMaskTextures[LightIndex - AttenuationLightStart];
@@ -1265,7 +1294,7 @@ void FDeferredShadingSceneRenderer::RenderLights(FRHICommandListImmediate& RHICm
 						RenderShadowProjections(RHICmdList, &LightSceneInfo, ScreenShadowMaskTexture, bInjectedTranslucentVolume);
 					}
 
-					bUsedShadowMaskTexture = true;					
+					bUsedShadowMaskTexture = true;
 				}
 
 				for (int32 ViewIndex = 0; ViewIndex < Views.Num(); ViewIndex++)
@@ -1273,7 +1302,7 @@ void FDeferredShadingSceneRenderer::RenderLights(FRHICommandListImmediate& RHICm
 					const FViewInfo& View = Views[ViewIndex];
 					View.HeightfieldLightingViewInfo.ComputeLighting(View, RHICmdList, LightSceneInfo);
 				}
-			
+
 				// Render light function to the attenuation buffer.
 				if (bDirectLighting)
 				{
@@ -1298,7 +1327,7 @@ void FDeferredShadingSceneRenderer::RenderLights(FRHICommandListImmediate& RHICm
 				{
 					RHICmdList.CopyToResolveTarget(ScreenShadowMaskTexture->GetRenderTargetItem().TargetableTexture, ScreenShadowMaskTexture->GetRenderTargetItem().ShaderResourceTexture, FResolveParams(FResolveRect()));
 				}
-			
+
 				if(bDirectLighting && !bInjectedTranslucentVolume)
 				{
 					for (int32 ViewIndex = 0; ViewIndex < Views.Num(); ViewIndex++)
@@ -1323,7 +1352,7 @@ void FDeferredShadingSceneRenderer::RenderLights(FRHICommandListImmediate& RHICm
 				{
 					SceneContext.BeginRenderingSceneColor(RHICmdList, ESimpleRenderTargetMode::EExistingColorAndDepth, FExclusiveDepthStencil::DepthRead_StencilWrite);
 
-					// Render the light to the scene color buffer, conditionally using the attenuation buffer or a 1x1 white texture as input 
+					// Render the light to the scene color buffer, conditionally using the attenuation buffer or a 1x1 white texture as input
 					if (bDirectLighting)
 					{
 						RenderLight(RHICmdList, &LightSceneInfo, ScreenShadowMaskTexture, false, true);
@@ -1357,9 +1386,9 @@ void FDeferredShadingSceneRenderer::RenderLightArrayForOverlapViewmode(FRHIComma
 			bShouldRender |= LightSceneInfo->ShouldRenderLight(Views[ViewIndex]);
 		}
 
-		if (bShouldRender 
+		if (bShouldRender
 			// Only render shadow casting stationary lights
-			&& LightSceneInfo->Proxy->HasStaticShadowing() 
+			&& LightSceneInfo->Proxy->HasStaticShadowing()
 			&& !LightSceneInfo->Proxy->HasStaticLighting()
 			&& LightSceneInfo->Proxy->CastsStaticShadow())
 		{
@@ -1443,7 +1472,7 @@ void CalculateLightNearFarDepthFromBounds(const FViewInfo& View, const FSphere &
 {
 	const FMatrix ViewProjection = View.ViewMatrices.GetViewProjectionMatrix();
 	const FVector ViewDirection = View.GetViewDirection();
-	
+
 	// push camera relative bounds center along view vec by its radius
 	const FVector FarPoint = LightBounds.Center + LightBounds.W * ViewDirection;
 	const FVector4 FarPoint4 = FVector4(FarPoint, 1.f);
@@ -1456,7 +1485,7 @@ void CalculateLightNearFarDepthFromBounds(const FViewInfo& View, const FSphere &
 	const FVector4 NearPoint4Clip = ViewProjection.TransformFVector4(NearPoint4);
 	NearDepth = NearPoint4Clip.Z / NearPoint4Clip.W;
 
-	// negative means behind view, but we use a NearClipPlane==1.f depth	
+	// negative means behind view, but we use a NearClipPlane==1.f depth
 
 	if (NearPoint4Clip.W < 0)
 		NearDepth = 1;
@@ -1552,11 +1581,11 @@ void FDeferredShadingSceneRenderer::RenderLight(FRHICommandList& RHICmdList, con
 			VertexShader->SetParameters(RHICmdList, View, LightSceneInfo);
 
 			// Apply the directional light as a full screen quad
-			DrawRectangle( 
+			DrawRectangle(
 				RHICmdList,
 				0, 0,
 				View.ViewRect.Width(), View.ViewRect.Height(),
-				View.ViewRect.Min.X, View.ViewRect.Min.Y, 
+				View.ViewRect.Min.X, View.ViewRect.Min.Y,
 				View.ViewRect.Width(), View.ViewRect.Height(),
 				View.ViewRect.Size(),
 				FSceneRenderTargets::Get(RHICmdList).GetBufferSizeXY(),
@@ -1625,7 +1654,7 @@ void FDeferredShadingSceneRenderer::RenderLight(FRHICommandList& RHICmdList, con
 			if( LightSceneInfo->Proxy->GetLightType() == LightType_Point ||
 				LightSceneInfo->Proxy->GetLightType() == LightType_Rect )
 			{
-				// Apply the point or spot light with some approximate bounding geometry, 
+				// Apply the point or spot light with some approximate bounding geometry,
 				// So we can get speedups from depth testing and not processing pixels outside of the light's influence.
 				StencilingGeometry::DrawSphere(RHICmdList);
 			}
@@ -1642,7 +1671,7 @@ void FDeferredShadingSceneRenderer::RenderSimpleLightsStandardDeferred(FRHIComma
 	SCOPE_CYCLE_COUNTER(STAT_DirectLightRenderingTime);
 	INC_DWORD_STAT_BY(STAT_NumLightsUsingStandardDeferred, SimpleLights.InstanceData.Num());
 	SCOPED_DRAW_EVENT(RHICmdList, StandardDeferredSimpleLights);
-	
+
 	FGraphicsPipelineStateInitializer GraphicsPSOInit;
 	RHICmdList.ApplyCachedRenderTargets(GraphicsPSOInit);
 
@@ -1682,7 +1711,7 @@ void FDeferredShadingSceneRenderer::RenderSimpleLightsStandardDeferred(FRHIComma
 
 			VertexShader->SetSimpleLightParameters(RHICmdList, View, LightBounds);
 
-			// Apply the point or spot light with some approximately bounding geometry, 
+			// Apply the point or spot light with some approximately bounding geometry,
 			// So we can get speedups from depth testing and not processing pixels outside of the light's influence.
 			StencilingGeometry::DrawSphere(RHICmdList);
 		}
diff --git a/Engine/Source/Runtime/Renderer/Private/RayTracing/RayTracingShadows.cpp b/Engine/Source/Runtime/Renderer/Private/RayTracing/RayTracingShadows.cpp
--- a/Engine/Source/Runtime/Renderer/Private/RayTracing/RayTracingShadows.cpp
+++ b/Engine/Source/Runtime/Renderer/Private/RayTracing/RayTracingShadows.cpp
@@ -62,6 +62,7 @@ class FOcclusionRGS : public FGlobalShader
 		SHADER_PARAMETER(float, NormalBias)
 		SHADER_PARAMETER(uint32, LightingChannelMask)
 		SHADER_PARAMETER(FIntRect, LightScissor)
+		SHADER_PARAMETER(FIntPoint, PixelOffset)
 
 		SHADER_PARAMETER_STRUCT(FLightShaderParameters, Light)
 		SHADER_PARAMETER_STRUCT_INCLUDE(FSceneTextureParameters, SceneTextures)
@@ -148,17 +149,26 @@ void FDeferredShadingSceneRenderer::RenderRayTracingShadows(
 		RayDistanceTexture = GraphBuilder.CreateTexture(Desc, TEXT("RayTracingOcclusionDistance"));
 	}
 
-	FIntRect ScissorRect = { {0,0}, View.ViewRect.Size() };
+	FIntRect ScissorRect = View.ViewRect;
+	FIntPoint PixelOffset = { 0, 0 };
+
+	// whether to clip the dispatch to a subrect, requires that denoising doesn't need the whole buffer
+	const bool bClipDispatch = DenoiserRequirements == IScreenSpaceDenoiser::EShadowRequirements::Bailout;
 
 	if (LightSceneProxy->GetScissorRect(ScissorRect, View, View.ViewRect))
 	{
 		// Account for scissor being defined on the whole frame viewport while the trace is only on the view subrect
-		ScissorRect.Min = ScissorRect.Min - View.ViewRect.Min;
-		ScissorRect.Max = ScissorRect.Max - View.ViewRect.Min;
+		ScissorRect.Min = ScissorRect.Min;
+		ScissorRect.Max = ScissorRect.Max;
 	}
 	else
 	{
-		ScissorRect = { {0,0}, View.ViewRect.Size() };
+		ScissorRect = View.ViewRect;
+	}
+
+	if (bClipDispatch)
+	{
+		PixelOffset = ScissorRect.Min;
 	}
 
 	// Ray generation pass for shadow occlusion.
@@ -174,6 +184,7 @@ void FDeferredShadingSceneRenderer::RenderRayTracingShadows(
 		PassParameters->ViewUniformBuffer = View.ViewUniformBuffer;
 		PassParameters->SceneTextures = SceneTextures;
 		PassParameters->LightScissor = ScissorRect;
+		PassParameters->PixelOffset = PixelOffset;
 		
 		FOcclusionRGS::FPermutationDomain PermutationVector;
 		PermutationVector.Set<FOcclusionRGS::FLightTypeDim>(LightSceneProxy->GetLightType());
@@ -202,8 +213,13 @@ void FDeferredShadingSceneRenderer::RenderRayTracingShadows(
 
 		FIntPoint Resolution(View.ViewRect.Width(), View.ViewRect.Height());
 
+		if (bClipDispatch)
+		{
+			Resolution = ScissorRect.Size();
+		}
+
 		GraphBuilder.AddPass(
-			RDG_EVENT_NAME("RayTracedShadow (spp=%d) %dx%d", RayTracingConfig.RayCountPerPixel, View.ViewRect.Width(), View.ViewRect.Height()),
+			RDG_EVENT_NAME("RayTracedShadow (spp=%d) %dx%d", RayTracingConfig.RayCountPerPixel, Resolution.X, Resolution.Y),
 			PassParameters,
 			ERenderGraphPassFlags::Compute,
 			[this, &View, RayGenerationShader, PassParameters, Resolution](FRHICommandList& RHICmdList)
