Better lod selection for RVT
Fixes ugly low lod splines
Integrating from //UE4/Dev-Rendering
#rnx

diff --git a/Engine/Source/Editor/LandscapeEditor/Private/LandscapeEdModeSplineTools.cpp b/Engine/Source/Editor/LandscapeEditor/Private/LandscapeEdModeSplineTools.cpp
--- a/Engine/Source/Editor/LandscapeEditor/Private/LandscapeEdModeSplineTools.cpp
+++ b/Engine/Source/Editor/LandscapeEditor/Private/LandscapeEdModeSplineTools.cpp
@@ -328,6 +328,8 @@ public:
 			NewSegment->bCastShadow = CopyFromSegment->bCastShadow;
 			NewSegment->TranslucencySortPriority = CopyFromSegment->TranslucencySortPriority;
 			NewSegment->RuntimeVirtualTextures = CopyFromSegment->RuntimeVirtualTextures;
+			NewSegment->VirtualTextureLodBias = CopyFromSegment->VirtualTextureLodBias;
+			NewSegment->VirtualTextureCullMips = CopyFromSegment->VirtualTextureCullMips;
 			NewSegment->VirtualTextureRenderPassType = CopyFromSegment->VirtualTextureRenderPassType;
 		}
 
@@ -686,6 +688,8 @@ public:
 		NewSegment->bCastShadow = Segment->bCastShadow;
 		NewSegment->TranslucencySortPriority = Segment->TranslucencySortPriority;
 		NewSegment->RuntimeVirtualTextures = Segment->RuntimeVirtualTextures;
+		NewSegment->VirtualTextureLodBias = Segment->VirtualTextureLodBias;
+		NewSegment->VirtualTextureCullMips = Segment->VirtualTextureCullMips;
 		NewSegment->VirtualTextureRenderPassType = Segment->VirtualTextureRenderPassType;
 
 		Segment->Connections[0].TangentLen *= t;
diff --git a/Engine/Source/Runtime/Engine/Classes/Components/PrimitiveComponent.h b/Engine/Source/Runtime/Engine/Classes/Components/PrimitiveComponent.h
--- a/Engine/Source/Runtime/Engine/Classes/Components/PrimitiveComponent.h
+++ b/Engine/Source/Runtime/Engine/Classes/Components/PrimitiveComponent.h
@@ -566,6 +566,25 @@ public:
 	UPROPERTY(EditAnywhere, BlueprintReadOnly, Category = VirtualTexture, meta = (DisplayName = "Render to Virtual Textures"))
 	TArray<URuntimeVirtualTexture*> RuntimeVirtualTextures;
 
+	/** Bias to the LOD selected for rendering to runtime virtual textures. */
+	UPROPERTY(EditAnywhere, AdvancedDisplay, BlueprintReadOnly, Category = VirtualTexture, meta = (DisplayName = "Virtual Texture LOD Bias", UIMin = "0", UIMax = "7"))
+	int32 VirtualTextureLodBias = 0;
+
+	/**
+	 * Number of lower mips in the runtime virtual texture to skip for rendering this primitive.
+	 * Larger values reduce the effective draw distance in the runtime virtual texture.
+	 * This culling method doesn't take into account primitive size or virtual texture size.
+	 */
+	UPROPERTY(EditAnywhere, AdvancedDisplay, BlueprintReadOnly, Category = VirtualTexture, meta = (DisplayName = "Virtual Texture Skip Mips", UIMin = "0", UIMax = "15"))
+	int32 VirtualTextureCullMips = 0;
+
+	/**
+	 * Set the minimum pixel coverage before culling from the runtime virtual texture.
+	 * Larger values reduce the effective draw distance in the runtime virtual texture.
+	 */
+	UPROPERTY(EditAnywhere, AdvancedDisplay, BlueprintReadOnly, Category = VirtualTexture, meta = (UIMin = "0", UIMax = "7"))
+	int32 VirtualTextureMinCoverage = 0;
+
 	/** Render to the main pass based on the virtual texture settings. */
 	UPROPERTY(EditAnywhere, BlueprintReadOnly, Category = VirtualTexture, meta = (DisplayName = "Virtual Texture Pass Type"))
 	ERuntimeVirtualTextureMainPassType VirtualTextureRenderPassType = ERuntimeVirtualTextureMainPassType::Exclusive;
diff --git a/Engine/Source/Runtime/Engine/Private/PrimitiveSceneProxy.cpp b/Engine/Source/Runtime/Engine/Private/PrimitiveSceneProxy.cpp
--- a/Engine/Source/Runtime/Engine/Private/PrimitiveSceneProxy.cpp
+++ b/Engine/Source/Runtime/Engine/Private/PrimitiveSceneProxy.cpp
@@ -141,6 +141,9 @@ FPrimitiveSceneProxy::FPrimitiveSceneProxy(const UPrimitiveComponent* InComponen
 ,	IndirectLightingCacheQuality(InComponent->IndirectLightingCacheQuality)
 ,	LpvBiasMultiplier(InComponent->LpvBiasMultiplier)
 ,	DynamicIndirectShadowMinVisibility(0)
+,	VirtualTextureLodBias(InComponent->VirtualTextureLodBias)
+,	VirtualTextureCullMips(InComponent->VirtualTextureCullMips)
+,	VirtualTextureMinCoverage(InComponent->VirtualTextureMinCoverage)
 ,	PrimitiveComponentId(InComponent->ComponentId)
 ,	Scene(InComponent->GetScene())
 ,	PrimitiveSceneInfo(NULL)
diff --git a/Engine/Source/Runtime/Engine/Private/StaticMeshRender.cpp b/Engine/Source/Runtime/Engine/Private/StaticMeshRender.cpp
--- a/Engine/Source/Runtime/Engine/Private/StaticMeshRender.cpp
+++ b/Engine/Source/Runtime/Engine/Private/StaticMeshRender.cpp
@@ -1031,61 +1031,46 @@ void FStaticMeshSceneProxy::DrawStaticElements(FStaticPrimitiveDrawInterface* PD
 #endif // WITH_EDITOR
 
 					const int32 NumBatches = GetNumMeshBatches();
+					const int32 NumRuntimeVirtualTextureTypes = RuntimeVirtualTextureMaterialTypes.Num();
 
-					PDI->ReserveMemoryForMeshes(NumBatches);
+					PDI->ReserveMemoryForMeshes(NumBatches * (1 + NumRuntimeVirtualTextureTypes));
 
 					for (int32 BatchIndex = 0; BatchIndex < NumBatches; BatchIndex++)
 					{
-						FMeshBatch MeshBatch;
-
-						if (GetMeshElement(LODIndex, BatchIndex, SectionIndex, PrimitiveDPG, bIsMeshElementSelected, true, MeshBatch))
+						FMeshBatch BaseMeshBatch;
+						if (GetMeshElement(LODIndex, BatchIndex, SectionIndex, PrimitiveDPG, bIsMeshElementSelected, true, BaseMeshBatch))
 						{
-							// If we have submitted an optimized shadow-only mesh, remaining mesh elements must not cast shadows.
-							MeshBatch.CastShadow &= !bUseUnifiedMeshForShadow;
-							MeshBatch.bUseAsOccluder &= !bUseUnifiedMeshForDepth;
-							MeshBatch.bUseForDepthPass &= !bUseUnifiedMeshForDepth;
-
-							PDI->DrawMesh(MeshBatch, ScreenSize);
+							{
+								// Standard mesh elements.
+								// If we have submitted an optimized shadow-only mesh, remaining mesh elements must not cast shadows.
+								FMeshBatch MeshBatch(BaseMeshBatch);
+								MeshBatch.CastShadow &= !bUseUnifiedMeshForShadow;
+								MeshBatch.bUseAsOccluder &= !bUseUnifiedMeshForDepth;
+								MeshBatch.bUseForDepthPass &= !bUseUnifiedMeshForDepth;
+								PDI->DrawMesh(MeshBatch, ScreenSize);
+							}
+							if (NumRuntimeVirtualTextureTypes > 0)
+							{
+								// Runtime virtual texture mesh elements.
+								FMeshBatch MeshBatch(BaseMeshBatch);
+								MeshBatch.CastShadow = 0;
+								MeshBatch.bUseAsOccluder = 0;
+								MeshBatch.bUseForDepthPass = 0;
+								MeshBatch.bUseForMaterial = 0;
+								MeshBatch.bDitheredLODTransition = 0;
+								MeshBatch.bRenderToVirtualTexture = 1;
+
+								for (ERuntimeVirtualTextureMaterialType MaterialType : RuntimeVirtualTextureMaterialTypes)
+								{
+									MeshBatch.RuntimeVirtualTextureMaterialType = (uint32)MaterialType;
+									PDI->DrawMesh(MeshBatch, FLT_MAX);
+								}
+							}
 						}
 					}
 				}
 			}
 		}
-
-		// Draw meshes for runtime virtual texture pass.
-		for (ERuntimeVirtualTextureMaterialType MaterialType : RuntimeVirtualTextureMaterialTypes)
-		{
-			// Use the minimal LOD which is usually (always?) what we would want.
-			//todo[vt]: Add user control to allow use of a specific LOD in the virtual texture.
-			// Note that ComputeStaticMeshLOD() checks *all* mesh batches for LOD calculation (including these VT ones) and assumes that they are stored in LOD order.
-			// Adding another set of last LODs here should be OK but if we want to use a different LOD they would need adding in the mesh batches in the above loop to preserve LOD order.
-			const int32 LODIndex = NumLODs - 1;
-			const FStaticMeshLODResources& LODModel = RenderData->LODResources[LODIndex];
-
-			for (int32 SectionIndex = 0; SectionIndex < LODModel.Sections.Num(); SectionIndex++)
-			{
-				const int32 NumBatches = GetNumMeshBatches();
-
-				PDI->ReserveMemoryForMeshes(NumBatches);
-
-				for (int32 BatchIndex = 0; BatchIndex < NumBatches; BatchIndex++)
-				{
-					FMeshBatch MeshBatch;
-
-					if (GetMeshElement(LODIndex, BatchIndex, SectionIndex, PrimitiveDPG, bIsMeshElementSelected, true, MeshBatch))
-					{
-						MeshBatch.bUseForDepthPass = 0;
-						MeshBatch.bUseForMaterial = 0;
-						MeshBatch.CastShadow = 0;
-						MeshBatch.bDitheredLODTransition = 0;
-						MeshBatch.bRenderToVirtualTexture = 1;
-						MeshBatch.RuntimeVirtualTextureMaterialType = (uint32)MaterialType;
-
-						PDI->DrawMesh(MeshBatch, FLT_MAX);
-					}
-				}
-			}
-		}
 	}
 }
 
diff --git a/Engine/Source/Runtime/Engine/Public/PrimitiveSceneProxy.h b/Engine/Source/Runtime/Engine/Public/PrimitiveSceneProxy.h
--- a/Engine/Source/Runtime/Engine/Public/PrimitiveSceneProxy.h
+++ b/Engine/Source/Runtime/Engine/Public/PrimitiveSceneProxy.h
@@ -456,6 +456,10 @@ public:
 	inline int16 GetTranslucencySortPriority() const { return TranslucencySortPriority; }
 	inline bool HasMotionBlurVelocityMeshes() const { return bHasMotionBlurVelocityMeshes; }
 
+	inline int32 GetVirtualTextureLodBias() const { return VirtualTextureLodBias; }
+	inline int32 GetVirtualTextureCullMips() const { return VirtualTextureCullMips; }
+	inline int32 GetVirtualTextureMinCoverage() const {	return VirtualTextureMinCoverage; }
+
 	inline bool IsMovable() const 
 	{ 
 		// Note: primitives with EComponentMobility::Stationary can still move (as opposed to lights with EComponentMobility::Stationary)
@@ -982,10 +986,18 @@ protected:
 
 	float DistanceFieldSelfShadowBias;
 
-	/** Array of virtual textures that this proxy should render to. */
+	/** Array of runtime virtual textures that this proxy should render to. */
 	TArray<URuntimeVirtualTexture*> RuntimeVirtualTextures;
+	/** Set of unique runtime virtual texture material types referenced by RuntimeVirtualTextures. */
 	TSet<ERuntimeVirtualTextureMaterialType> RuntimeVirtualTextureMaterialTypes;
 
+	/** Geometry Lod bias when rendering to runtime virtual texture. */
+	int8 VirtualTextureLodBias;
+	/** Number of low mips to skip when rendering to runtime virtual texture. */
+	int8 VirtualTextureCullMips;
+	/** Log2 of minimum estimated pixel coverage before culling from runtime virtual texture. */
+	int8 VirtualTextureMinCoverage;
+
 private:
 	/** The hierarchy of owners of this primitive.  These must not be dereferenced on the rendering thread, but the pointer values can be used for identification.  */
 	TArray<const AActor*> Owners;
diff --git a/Engine/Source/Runtime/Foliage/Private/InstancedFoliage.cpp b/Engine/Source/Runtime/Foliage/Private/InstancedFoliage.cpp
--- a/Engine/Source/Runtime/Foliage/Private/InstancedFoliage.cpp
+++ b/Engine/Source/Runtime/Foliage/Private/InstancedFoliage.cpp
@@ -1332,6 +1332,11 @@ void FFoliageStaticMesh::UpdateComponentSettings(const UFoliageType_InstancedSta
 			Component->VirtualTextureRenderPassType = FoliageType->VirtualTextureRenderPassType;
 			bNeedsMarkRenderStateDirty = true;
 		}
+		if (Component->VirtualTextureCullMips != FoliageType->VirtualTextureCullMips)
+		{
+			Component->VirtualTextureCullMips = FoliageType->VirtualTextureCullMips;
+			bNeedsMarkRenderStateDirty = true;
+		}
 		if (Component->TranslucencySortPriority != FoliageType->TranslucencySortPriority)
 		{
 			Component->TranslucencySortPriority = FoliageType->TranslucencySortPriority;
diff --git a/Engine/Source/Runtime/Foliage/Public/FoliageType.h b/Engine/Source/Runtime/Foliage/Public/FoliageType.h
--- a/Engine/Source/Runtime/Foliage/Public/FoliageType.h
+++ b/Engine/Source/Runtime/Foliage/Public/FoliageType.h
@@ -515,6 +515,14 @@ public:
 	UPROPERTY(EditAnywhere, BlueprintReadOnly, Category=VirtualTexture, meta = (DisplayName = "Render to Virtual Textures"))
 	TArray<URuntimeVirtualTexture*> RuntimeVirtualTextures;
 
+	/**
+	 * Number of lower mips in the runtime virtual texture to skip for rendering this primitive.
+	 * Larger values reduce the effective draw distance in the runtime virtual texture.
+	 * This culling method doesn't take into account primitive size or virtual texture size.
+	 */
+	UPROPERTY(EditAnywhere, AdvancedDisplay, BlueprintReadOnly, Category = VirtualTexture, meta = (DisplayName = "Virtual Texture Skip Mips", UIMin = "0", UIMax = "15"))
+	int32 VirtualTextureCullMips = 0;
+
 	/** Render to the main pass based on the virtual texture settings. */
 	UPROPERTY(EditAnywhere, BlueprintReadOnly, Category = VirtualTexture, meta = (DisplayName = "Virtual Texture Pass Type"))
 	ERuntimeVirtualTextureMainPassType VirtualTextureRenderPassType = ERuntimeVirtualTextureMainPassType::Exclusive;
diff --git a/Engine/Source/Runtime/Landscape/Classes/LandscapeProxy.h b/Engine/Source/Runtime/Landscape/Classes/LandscapeProxy.h
--- a/Engine/Source/Runtime/Landscape/Classes/LandscapeProxy.h
+++ b/Engine/Source/Runtime/Landscape/Classes/LandscapeProxy.h
@@ -511,6 +511,18 @@ public:
 	UPROPERTY(EditAnywhere, BlueprintReadOnly, Category = VirtualTexture, meta = (DisplayName = "Render to Virtual Textures"))
 	TArray<URuntimeVirtualTexture*> RuntimeVirtualTextures;
 
+	/** 
+	 * Number of mesh levels to use when rendering landscape into runtime virtual texture.
+	 * Set this only if the material used to render the virtual texture requires interpolated vertex data such as height.
+	 * Higher values use more tessellated meshes and are expensive when rendering the runtime virtual texture.
+	 */
+	UPROPERTY(EditAnywhere, AdvancedDisplay, BlueprintReadOnly, Category = VirtualTexture, meta = (DisplayName = "Virtual Texture Num LODs", UIMin = "0", UIMax = "7"))
+	int32 VirtualTextureNumLods = 0;
+
+	/** Bias to the LOD selected for rendering to runtime virtual textures. */
+	UPROPERTY(EditAnywhere, AdvancedDisplay, BlueprintReadOnly, Category = VirtualTexture, meta = (DisplayName = "Virtual Texture LOD Bias", UIMin = "0", UIMax = "7"))
+	int32 VirtualTextureLodBias = 0;
+
 	/** Render to the main pass based on the virtual texture settings. */
 	UPROPERTY(EditAnywhere, BlueprintReadOnly, Category = VirtualTexture, meta = (DisplayName = "Virtual Texture Pass Type"))
 	ERuntimeVirtualTextureMainPassType VirtualTextureRenderPassType = ERuntimeVirtualTextureMainPassType::Always;
diff --git a/Engine/Source/Runtime/Landscape/Classes/LandscapeSplineSegment.h b/Engine/Source/Runtime/Landscape/Classes/LandscapeSplineSegment.h
--- a/Engine/Source/Runtime/Landscape/Classes/LandscapeSplineSegment.h
+++ b/Engine/Source/Runtime/Landscape/Classes/LandscapeSplineSegment.h
@@ -238,6 +238,18 @@ class ULandscapeSplineSegment : public UObject
 	UPROPERTY(EditAnywhere, BlueprintReadOnly, Category = VirtualTexture, meta = (DisplayName = "Render to Virtual Textures"))
 	TArray<URuntimeVirtualTexture*> RuntimeVirtualTextures;
 
+	/** Lod bias for rendering to runtime virtual texture. */
+	UPROPERTY(EditAnywhere, AdvancedDisplay, BlueprintReadOnly, Category = VirtualTexture, meta = (DisplayName = "Virtual Texture LOD Bias", UIMin = "0", UIMax = "7"))
+	int32 VirtualTextureLodBias = 0;
+
+	/**
+	 * Number of lower mips in the runtime virtual texture to skip for rendering this primitive.
+	 * Larger values reduce the effective draw distance in the runtime virtual texture.
+	 * This culling method doesn't take into account primitive size or virtual texture size.
+	 */
+	UPROPERTY(EditAnywhere, AdvancedDisplay, BlueprintReadOnly, Category = VirtualTexture, meta = (DisplayName = "Virtual Texture Skip Mips", UIMin = "0", UIMax = "15"))
+	int32 VirtualTextureCullMips = 0;
+
 	/** Render to the main pass based on the virtual texture settings. */
 	UPROPERTY(EditAnywhere, BlueprintReadOnly, Category = VirtualTexture, meta = (DisplayName = "Virtual Texture Pass Type"))
 	ERuntimeVirtualTextureMainPassType VirtualTextureRenderPassType = ERuntimeVirtualTextureMainPassType::Exclusive;
diff --git a/Engine/Source/Runtime/Landscape/Private/Landscape.cpp b/Engine/Source/Runtime/Landscape/Private/Landscape.cpp
--- a/Engine/Source/Runtime/Landscape/Private/Landscape.cpp
+++ b/Engine/Source/Runtime/Landscape/Private/Landscape.cpp
@@ -184,6 +184,9 @@ ULandscapeComponent::ULandscapeComponent(const FObjectInitializer& ObjectInitial
 
 	// We don't want to load this on the server, this component is for graphical purposes only
 	AlwaysLoadOnServer = false;
+
+	// Default sort priority of landscape to -1 so that it will default to the first thing rendered in any runtime virtual texture
+	TranslucencySortPriority = -1;
 }
 
 int32 ULandscapeComponent::GetMaterialInstanceCount(bool InDynamic) const
diff --git a/Engine/Source/Runtime/Landscape/Private/LandscapeEdit.cpp b/Engine/Source/Runtime/Landscape/Private/LandscapeEdit.cpp
--- a/Engine/Source/Runtime/Landscape/Private/LandscapeEdit.cpp
+++ b/Engine/Source/Runtime/Landscape/Private/LandscapeEdit.cpp
@@ -4055,7 +4055,14 @@ void ALandscapeProxy::PostEditChangeProperty(FPropertyChangedEvent& PropertyChan
 	{
 		MarkComponentsRenderStateDirty();
 	}
-	
+	else if (PropertyName == GET_MEMBER_NAME_CHECKED(ALandscapeProxy, RuntimeVirtualTextures)
+		|| PropertyName == GET_MEMBER_NAME_CHECKED(ALandscapeProxy, VirtualTextureRenderPassType)
+		|| PropertyName == GET_MEMBER_NAME_CHECKED(ALandscapeProxy, VirtualTextureNumLods)
+		|| PropertyName == GET_MEMBER_NAME_CHECKED(ALandscapeProxy, VirtualTextureLodBias))
+	{
+		MarkComponentsRenderStateDirty();
+	}
+
 	// Remove null layer infos
 	EditorLayerSettings.RemoveAll([](const FLandscapeEditorLayerSettings& Entry) { return Entry.LayerInfoObj == nullptr; });
 
diff --git a/Engine/Source/Runtime/Landscape/Private/LandscapeRender.cpp b/Engine/Source/Runtime/Landscape/Private/LandscapeRender.cpp
--- a/Engine/Source/Runtime/Landscape/Private/LandscapeRender.cpp
+++ b/Engine/Source/Runtime/Landscape/Private/LandscapeRender.cpp
@@ -765,6 +765,10 @@ FLandscapeComponentSceneProxy::FLandscapeComponentSceneProxy(ULandscapeComponent
 	MinValidLOD = FMath::Clamp<int8>(LocalLODBias, -MaxLOD, MaxLOD);
 	MaxValidLOD = FMath::Min<int32>(MaxLOD, MaxLOD + LocalLODBias);
 
+	LastVirtualTextureLOD = MaxLOD;
+	FirstVirtualTextureLOD = FMath::Max(MaxLOD - InComponent->GetLandscapeProxy()->VirtualTextureNumLods, 0);
+	VirtualTextureLodBias = InComponent->GetLandscapeProxy()->VirtualTextureLodBias;
+
 	ComponentMaxExtend = SubsectionSizeQuads * FMath::Max(InComponent->GetComponentTransform().GetScale3D().X, InComponent->GetComponentTransform().GetScale3D().Y);
 
 	if (NumSubsections > 1)
@@ -1764,7 +1768,7 @@ void FLandscapeComponentSceneProxy::DrawStaticElements(FStaticPrimitiveDrawInter
 	}		
 
 	int32 TotalBatchCount = (1 + LastLOD - FirstLOD) * NumBatchesPerLOD * BatchCount;
-	TotalBatchCount += RuntimeVirtualTextureMaterialTypes.Num();
+	TotalBatchCount += (1 + LastVirtualTextureLOD - FirstVirtualTextureLOD) * RuntimeVirtualTextureMaterialTypes.Num();
 
 	StaticBatchParamArray.Empty(TotalBatchCount);
 	PDI->ReserveMemoryForMeshes(TotalBatchCount);
@@ -1838,20 +1842,16 @@ void FLandscapeComponentSceneProxy::DrawStaticElements(FStaticPrimitiveDrawInter
 	// Add fixed grid mesh batches for runtime virtual texture usage
  	for (ERuntimeVirtualTextureMaterialType MaterialType : RuntimeVirtualTextureMaterialTypes)
  	{
- 		// Use lowest detail geometry Lod and highest detail material Lod
- 		//todo[vt]:
-		// Add user control for to allow use of different Lods.
-		// Specifically we would want to use a higher geometry Lod in the case where the landscape interpolator outputs need to be at a high resolution for correct virtual texture rendering.
-		// This happens if we read world height in the landscape shading (although we could read height directly from the height texture instead of taking it from the vertex interpolator).
-		// Also it would happen if we are relying on modified vertex positions (either from position offsets in material, or from using the XYOffsetmapTexture).
- 		const int32 LODIndex = LastLOD;
- 		const int32 MaterialIndex = LODIndexToMaterialIndex[FirstLOD];
- 
- 		FMeshBatch RuntimeVirtualTextureMeshBatch;
- 		if (GetMeshElementForVirtualTexture(LODIndex, MaterialType, AvailableMaterials[MaterialIndex], RuntimeVirtualTextureMeshBatch, StaticBatchParamArray))
- 		{
- 			PDI->DrawMesh(RuntimeVirtualTextureMeshBatch, FLT_MAX);
- 		}
+		const int32 MaterialIndex = LODIndexToMaterialIndex[FirstLOD];
+
+		for (int32 LODIndex = FirstVirtualTextureLOD; LODIndex <= LastVirtualTextureLOD; ++LODIndex)
+		{
+			FMeshBatch RuntimeVirtualTextureMeshBatch;
+			if (GetMeshElementForVirtualTexture(LODIndex, MaterialType, AvailableMaterials[MaterialIndex], RuntimeVirtualTextureMeshBatch, StaticBatchParamArray))
+			{
+				PDI->DrawMesh(RuntimeVirtualTextureMeshBatch, FLT_MAX);
+			}
+		}
  	}
 
 	check(StaticBatchParamArray.Num() <= TotalBatchCount);
diff --git a/Engine/Source/Runtime/Landscape/Private/LandscapeSplines.cpp b/Engine/Source/Runtime/Landscape/Private/LandscapeSplines.cpp
--- a/Engine/Source/Runtime/Landscape/Private/LandscapeSplines.cpp
+++ b/Engine/Source/Runtime/Landscape/Private/LandscapeSplines.cpp
@@ -2612,6 +2612,8 @@ void ULandscapeSplineSegment::UpdateSplinePoints(bool bUpdateCollision)
 			MeshComponent->TranslucencySortPriority = TranslucencySortPriority;
 
 			MeshComponent->RuntimeVirtualTextures = RuntimeVirtualTextures;
+			MeshComponent->VirtualTextureLodBias = VirtualTextureLodBias;
+			MeshComponent->VirtualTextureCullMips = VirtualTextureCullMips;
 			MeshComponent->VirtualTextureRenderPassType = VirtualTextureRenderPassType;
 
 			MeshComponent->SetCastShadow(bCastShadow);
diff --git a/Engine/Source/Runtime/Landscape/Public/LandscapeRender.h b/Engine/Source/Runtime/Landscape/Public/LandscapeRender.h
--- a/Engine/Source/Runtime/Landscape/Public/LandscapeRender.h
+++ b/Engine/Source/Runtime/Landscape/Public/LandscapeRender.h
@@ -581,6 +581,8 @@ protected:
 	TArray<float>				LODScreenRatioSquared;		// Table of valid screen size -> LOD index
 	int32						FirstLOD;	// First LOD we have batch elements for
 	int32						LastLOD;	// Last LOD we have batch elements for
+	int32						FirstVirtualTextureLOD;
+	int32						LastVirtualTextureLOD;
 	float						ComponentMaxExtend; 		// The max extend value in any axis
 	float						ComponentSquaredScreenSizeToUseSubSections; // Size at which we start to draw in sub lod if LOD are different per sub section
 	float						MinValidLOD;							// Min LOD Taking into account LODBias
diff --git a/Engine/Source/Runtime/Renderer/Private/PrimitiveSceneInfo.cpp b/Engine/Source/Runtime/Renderer/Private/PrimitiveSceneInfo.cpp
--- a/Engine/Source/Runtime/Renderer/Private/PrimitiveSceneInfo.cpp
+++ b/Engine/Source/Runtime/Renderer/Private/PrimitiveSceneInfo.cpp
@@ -444,6 +444,24 @@ static void OnVirtualTextureDestroyed(const FVirtualTextureProducerHandle& InHan
 	AddPrimitiveToUpdateGPU(*PrimitiveSceneInfo->Scene, PrimitiveSceneInfo->GetIndex());
 }
 
+static void GetRuntimeVirtualTextureLODRange(TArray<class FStaticMeshBatchRelevance> const& MeshRelevances, int8& OutMinLOD, int8& OutMaxLOD)
+{
+	OutMinLOD = MAX_int8;
+	OutMaxLOD = 0;
+
+	for (int32 MeshIndex = 0; MeshIndex < MeshRelevances.Num(); ++MeshIndex)
+	{
+		const FStaticMeshBatchRelevance& MeshRelevance = MeshRelevances[MeshIndex];
+		if (MeshRelevance.bRenderToVirtualTexture)
+		{
+			OutMinLOD = FMath::Min(OutMinLOD, MeshRelevance.LODIndex);
+			OutMaxLOD = FMath::Max(OutMaxLOD, MeshRelevance.LODIndex);
+		}
+	}
+
+	check(OutMinLOD <= OutMaxLOD);
+}
+
 int32 FPrimitiveSceneInfo::UpdateStaticLightingBuffer()
 {
 	checkSlow(IsInRenderingThread());
@@ -617,12 +635,22 @@ void FPrimitiveSceneInfo::AddToScene(FRHICommandListImmediate& RHICmdList, bool
 	// Store the component.
 	Scene->PrimitiveComponentIds[PackedIndex] = PrimitiveComponentId;
 
-	// Store the runtime virtual texture flags.
+	// Store the runtime virtual texture flags and lod info.
 	const bool bRenderToVirtualTexture = Proxy->WritesVirtualTexture();
 	Scene->PrimitiveVirtualTextureFlags[PackedIndex].bRenderToVirtualTexture = bRenderToVirtualTexture;
 	if (bRenderToVirtualTexture)
 	{
 		Scene->PrimitiveVirtualTextureFlags[PackedIndex].RuntimeVirtualTextureMask = Scene->GetRuntimeVirtualTextureMask(Proxy);
+
+		int8 MinLod, MaxLod;
+		GetRuntimeVirtualTextureLODRange(StaticMeshRelevances, MinLod, MaxLod);
+
+		FPrimitiveVirtualTextureLodInfo& LodInfo = Scene->PrimitiveVirtualTextureLod[PackedIndex];
+		LodInfo.MinLod = MinLod;
+		LodInfo.MaxLod = MaxLod;
+		LodInfo.LodBias = Proxy->GetVirtualTextureLodBias();
+		LodInfo.CullMethod = Proxy->GetVirtualTextureMinCoverage() == 0 ? 0 : 1;
+		LodInfo.CullValue = LodInfo.CullMethod == 0 ? Proxy->GetVirtualTextureCullMips() : Proxy->GetVirtualTextureMinCoverage();
 	}
 
 	{
diff --git a/Engine/Source/Runtime/Renderer/Private/RendererScene.cpp b/Engine/Source/Runtime/Renderer/Private/RendererScene.cpp
--- a/Engine/Source/Runtime/Renderer/Private/RendererScene.cpp
+++ b/Engine/Source/Runtime/Renderer/Private/RendererScene.cpp
@@ -563,6 +563,7 @@ void FScene::CheckPrimitiveArrays()
 	check(Primitives.Num() == PrimitiveOcclusionFlags.Num());
 	check(Primitives.Num() == PrimitiveComponentIds.Num());
 	check(Primitives.Num() == PrimitiveVirtualTextureFlags.Num());
+	check(Primitives.Num() == PrimitiveVirtualTextureLod.Num());
 	check(Primitives.Num() == PrimitiveOcclusionBounds.Num());
 	check(Primitives.Num() == PrimitivesNeedingStaticMeshUpdate.Num());
 }
@@ -788,6 +789,7 @@ void FScene::AddPrimitiveSceneInfo_RenderThread(FRHICommandListImmediate& RHICmd
 	PrimitiveOcclusionFlags.AddUninitialized();
 	PrimitiveComponentIds.AddUninitialized();
 	PrimitiveVirtualTextureFlags.AddUninitialized();
+	PrimitiveVirtualTextureLod.AddUninitialized();
 	PrimitiveOcclusionBounds.AddUninitialized();
 	PrimitivesNeedingStaticMeshUpdate.Add(false);
 	
@@ -857,6 +859,7 @@ void FScene::AddPrimitiveSceneInfo_RenderThread(FRHICommandListImmediate& RHICmd
 				TArraySwapElements(PrimitiveOcclusionFlags, DestIndex, SourceIndex);
 				TArraySwapElements(PrimitiveComponentIds, DestIndex, SourceIndex);
 				TArraySwapElements(PrimitiveVirtualTextureFlags, DestIndex, SourceIndex);
+				TArraySwapElements(PrimitiveVirtualTextureLod, DestIndex, SourceIndex);
 				TArraySwapElements(PrimitiveOcclusionBounds, DestIndex, SourceIndex);
 				TBitArraySwapElements(PrimitivesNeedingStaticMeshUpdate, DestIndex, SourceIndex);
 
@@ -1584,6 +1587,7 @@ void FScene::RemovePrimitiveSceneInfo_RenderThread(FPrimitiveSceneInfo* Primitiv
 				TArraySwapElements(PrimitiveOcclusionFlags, DestIndex, SourceIndex);
 				TArraySwapElements(PrimitiveComponentIds, DestIndex, SourceIndex);
 				TArraySwapElements(PrimitiveVirtualTextureFlags, DestIndex, SourceIndex);
+				TArraySwapElements(PrimitiveVirtualTextureLod, DestIndex, SourceIndex);
 				TArraySwapElements(PrimitiveOcclusionBounds, DestIndex, SourceIndex);
 				TBitArraySwapElements(PrimitivesNeedingStaticMeshUpdate, DestIndex, SourceIndex);
 				SourceIndex = DestIndex;
@@ -1616,6 +1620,7 @@ void FScene::RemovePrimitiveSceneInfo_RenderThread(FPrimitiveSceneInfo* Primitiv
 		PrimitiveOcclusionFlags.Pop();
 		PrimitiveComponentIds.Pop();
 		PrimitiveVirtualTextureFlags.Pop();
+		PrimitiveVirtualTextureLod.Pop();
 		PrimitiveOcclusionBounds.Pop();
 		PrimitivesNeedingStaticMeshUpdate.RemoveAt(PrimitivesNeedingStaticMeshUpdate.Num()-1);
 	}
diff --git a/Engine/Source/Runtime/Renderer/Private/ScenePrivate.h b/Engine/Source/Runtime/Renderer/Private/ScenePrivate.h
--- a/Engine/Source/Runtime/Renderer/Private/ScenePrivate.h
+++ b/Engine/Source/Runtime/Renderer/Private/ScenePrivate.h
@@ -2445,6 +2445,8 @@ public:
 	TArray<FPrimitiveComponentId> PrimitiveComponentIds;
 	/** Packed array of runtime virtual texture flags. */
 	TArray<FPrimitiveVirtualTextureFlags> PrimitiveVirtualTextureFlags;
+	/** Packed array of runtime virtual texture lod info. */
+	TArray<FPrimitiveVirtualTextureLodInfo> PrimitiveVirtualTextureLod;
 
 	TBitArray<> PrimitivesNeedingStaticMeshUpdate;
 	TSet<FPrimitiveSceneInfo*> PrimitivesNeedingStaticMeshUpdateWithoutVisibilityCheck;
diff --git a/Engine/Source/Runtime/Renderer/Private/VT/RuntimeVirtualTextureProducer.cpp b/Engine/Source/Runtime/Renderer/Private/VT/RuntimeVirtualTextureProducer.cpp
--- a/Engine/Source/Runtime/Renderer/Private/VT/RuntimeVirtualTextureProducer.cpp
+++ b/Engine/Source/Runtime/Renderer/Private/VT/RuntimeVirtualTextureProducer.cpp
@@ -77,7 +77,8 @@ void FRuntimeVirtualTextureFinalizer::Finalize(FRHICommandListImmediate& RHICmdL
 			DestinationBox1,
 			UVToWorld,
 			UVRange, 
-			Entry.vLevel);
+			Entry.vLevel,
+			Desc.MaxLevel);
 	}
 
 	Tiles.SetNumUnsafeInternal(0);
diff --git a/Engine/Source/Runtime/Renderer/Private/VT/RuntimeVirtualTextureRender.cpp b/Engine/Source/Runtime/Renderer/Private/VT/RuntimeVirtualTextureRender.cpp
--- a/Engine/Source/Runtime/Renderer/Private/VT/RuntimeVirtualTextureRender.cpp
+++ b/Engine/Source/Runtime/Renderer/Private/VT/RuntimeVirtualTextureRender.cpp
@@ -269,7 +269,7 @@ namespace RuntimeVirtualTexture
 
 
 	/** Collect meshes and draw. */
-	void DrawMeshes(FRHICommandListImmediate& RHICmdList, FScene const* Scene, FViewInfo const* View, ERuntimeVirtualTextureMaterialType MaterialType, uint32 RuntimeVirtualTextureMask)
+	void DrawMeshes(FRHICommandListImmediate& RHICmdList, FScene const* Scene, FViewInfo const* View, ERuntimeVirtualTextureMaterialType MaterialType, uint32 RuntimeVirtualTextureMask, uint8 vLevel, uint8 MaxLevel)
 	{
 		// Cached draw command collectors
 		const FCachedPassMeshDrawList& SceneDrawList = Scene->CachedDrawLists[EMeshPass::VirtualTexture];
@@ -282,58 +282,95 @@ namespace RuntimeVirtualTexture
 		FDynamicPassMeshDrawListContext DynamicMeshPassContext(MeshDrawCommandStorage, AllocatedCommands, GraphicsMinimalPipelineStateSet);
 		FRuntimeVirtualTextureMeshProcessor MeshProcessor(Scene, View, &DynamicMeshPassContext);
 
+		// Pre-calculate view factors used for culling
+		const float RcpWorldSize = 1.f / (View->ViewMatrices.GetInvProjectionMatrix().M[0][0]);
+		const float WorldToPixel = View->ViewRect.Width() * RcpWorldSize;
+
 		// Iterate over scene and collect visible virtual texture draw commands for this view
 		//todo: Consider a broad phase (quad tree etc?) here. (But only if running over PrimitiveVirtualTextureFlags shows up as a bottleneck.)
 		for (int32 PrimitiveIndex = 0; PrimitiveIndex < Scene->Primitives.Num(); ++PrimitiveIndex)
 		{
-			if (Scene->PrimitiveVirtualTextureFlags[PrimitiveIndex].RuntimeVirtualTextureMask & RuntimeVirtualTextureMask)
+			const FPrimitiveVirtualTextureFlags Flags = Scene->PrimitiveVirtualTextureFlags[PrimitiveIndex];
+			if (!Flags.bRenderToVirtualTexture)
+			{
+				continue;
+			}
+			if ((Flags.RuntimeVirtualTextureMask & RuntimeVirtualTextureMask) == 0)
+			{
+				continue;
+			}
+
+			//todo[vt]: In our case we know that frustum is an oriented box so investigate cheaper test for intersecting that
+			const FSphere SphereBounds = Scene->PrimitiveBounds[PrimitiveIndex].BoxSphereBounds.GetSphere();
+			if (!View->ViewFrustum.IntersectSphere(SphereBounds.Center, SphereBounds.W))
+			{
+				continue;
+			}
+
+			// Cull primitives according to mip level or pixel coverage
+			const FPrimitiveVirtualTextureLodInfo LodInfo = Scene->PrimitiveVirtualTextureLod[PrimitiveIndex];
+			if (LodInfo.CullMethod == 0)
+			{
+				if (MaxLevel - vLevel < LodInfo.CullValue)
+				{
+					continue;
+				}
+			}
+			else
 			{
-				//todo[vt]: Use quicker/more accurate 2d test here since we can pre-calculate 2d bounds in VT space.
-				FBoxSphereBounds const& Bounds = Scene->PrimitiveBounds[PrimitiveIndex].BoxSphereBounds;
-				if (View->ViewFrustum.IntersectSphere(Bounds.GetSphere().Center, Bounds.GetSphere().W))
+				// Note that we use 2^MinPixelCoverage as that scales linearly with mip extents
+				int32 PixelCoverage = FMath::FloorToInt(2.f * SphereBounds.W * WorldToPixel);
+				if (PixelCoverage < (1 << LodInfo.CullValue))
 				{
-					FPrimitiveSceneInfo* PrimitiveSceneInfo = Scene->Primitives[PrimitiveIndex];
-					for (int32 MeshIndex = 0; MeshIndex < PrimitiveSceneInfo->StaticMeshes.Num(); ++MeshIndex)
+					continue;
+				}
+			}
+
+			// Calculate Lod for current mip
+			const float AreaRatio = 2.f * SphereBounds.W * RcpWorldSize;
+			const int32 LodIndex = FMath::Clamp<int32>((int32)LodInfo.LodBias - FMath::FloorToInt(FMath::Log2(AreaRatio)), LodInfo.MinLod, LodInfo.MaxLod);
+
+			// Process meshes
+			FPrimitiveSceneInfo* PrimitiveSceneInfo = Scene->Primitives[PrimitiveIndex];
+			for (int32 MeshIndex = 0; MeshIndex < PrimitiveSceneInfo->StaticMeshes.Num(); ++MeshIndex)
+			{
+				FStaticMeshBatchRelevance const& StaticMeshRelevance = PrimitiveSceneInfo->StaticMeshRelevances[MeshIndex];
+				if (StaticMeshRelevance.bRenderToVirtualTexture && StaticMeshRelevance.LODIndex == LodIndex && StaticMeshRelevance.RuntimeVirtualTextureMaterialType == (uint32)MaterialType)
+				{
+					bool bCachedDraw = false;
+					if (StaticMeshRelevance.bSupportsCachingMeshDrawCommands)
 					{
-						FStaticMeshBatchRelevance const& StaticMeshRelevance = PrimitiveSceneInfo->StaticMeshRelevances[MeshIndex];
-						if (StaticMeshRelevance.bRenderToVirtualTexture && StaticMeshRelevance.RuntimeVirtualTextureMaterialType == (uint32)MaterialType)
+						// Use cached draw command
+						const int32 StaticMeshCommandInfoIndex = StaticMeshRelevance.GetStaticMeshCommandInfoIndex(EMeshPass::VirtualTexture);
+						if (StaticMeshCommandInfoIndex >= 0)
 						{
-							bool bCachedDraw = false;
-							if (StaticMeshRelevance.bSupportsCachingMeshDrawCommands)
-							{
-								// Use cached draw command
-								const int32 StaticMeshCommandInfoIndex = StaticMeshRelevance.GetStaticMeshCommandInfoIndex(EMeshPass::VirtualTexture);
-								if (StaticMeshCommandInfoIndex >= 0)
-								{
-									FCachedMeshDrawCommandInfo& CachedMeshDrawCommand = PrimitiveSceneInfo->StaticMeshCommandInfos[StaticMeshCommandInfoIndex];
-
-									const FMeshDrawCommand* MeshDrawCommand = CachedMeshDrawCommand.StateBucketId >= 0
-										? &Scene->CachedMeshDrawCommandStateBuckets[FSetElementId::FromInteger(CachedMeshDrawCommand.StateBucketId)].MeshDrawCommand
-										: &SceneDrawList.MeshDrawCommands[CachedMeshDrawCommand.CommandIndex];
-
-									FVisibleMeshDrawCommand NewVisibleMeshDrawCommand;
-									NewVisibleMeshDrawCommand.Setup(
-										MeshDrawCommand,
-										PrimitiveIndex,
-										PrimitiveIndex,
-										CachedMeshDrawCommand.StateBucketId,
-										CachedMeshDrawCommand.MeshFillMode,
-										CachedMeshDrawCommand.MeshCullMode,
-										CachedMeshDrawCommand.SortKey);
-
-									CachedDrawCommands.Add(NewVisibleMeshDrawCommand);
-									bCachedDraw = true;
-								}
-							}
-
-							if (!bCachedDraw)
-							{
-								// No cached draw command was available. Process the mesh batch.
-								uint64 BatchElementMask = ~0ull;
-								MeshProcessor.AddMeshBatch(PrimitiveSceneInfo->StaticMeshes[MeshIndex], BatchElementMask, Scene->PrimitiveSceneProxies[PrimitiveIndex]);
-							}
+							FCachedMeshDrawCommandInfo& CachedMeshDrawCommand = PrimitiveSceneInfo->StaticMeshCommandInfos[StaticMeshCommandInfoIndex];
+
+							const FMeshDrawCommand* MeshDrawCommand = CachedMeshDrawCommand.StateBucketId >= 0
+								? &Scene->CachedMeshDrawCommandStateBuckets[FSetElementId::FromInteger(CachedMeshDrawCommand.StateBucketId)].MeshDrawCommand
+								: &SceneDrawList.MeshDrawCommands[CachedMeshDrawCommand.CommandIndex];
+
+							FVisibleMeshDrawCommand NewVisibleMeshDrawCommand;
+							NewVisibleMeshDrawCommand.Setup(
+								MeshDrawCommand,
+								PrimitiveIndex,
+								PrimitiveIndex,
+								CachedMeshDrawCommand.StateBucketId,
+								CachedMeshDrawCommand.MeshFillMode,
+								CachedMeshDrawCommand.MeshCullMode,
+								CachedMeshDrawCommand.SortKey);
+
+							CachedDrawCommands.Add(NewVisibleMeshDrawCommand);
+							bCachedDraw = true;
 						}
 					}
+
+					if (!bCachedDraw)
+					{
+						// No cached draw command was available. Process the mesh batch.
+						uint64 BatchElementMask = ~0ull;
+						MeshProcessor.AddMeshBatch(PrimitiveSceneInfo->StaticMeshes[MeshIndex], BatchElementMask, Scene->PrimitiveSceneProxies[PrimitiveIndex]);
+					}
 				}
 			}
 		}
@@ -550,7 +587,8 @@ namespace RuntimeVirtualTexture
 		FBox2D const& DestBox1,
 		FTransform const& UVToWorld,
 		FBox2D const& UVRange,
-		uint8 vLevel)
+		uint8 vLevel,
+		uint8 MaxLevel)
 	{
 		SCOPED_DRAW_EVENT(RHICmdList, VirtualTextureDynamicCache);
 
@@ -600,7 +638,7 @@ namespace RuntimeVirtualTexture
 		View->CachedViewUniformShaderParameters = MakeUnique<FViewUniformShaderParameters>();
 		View->SetupUniformBufferParameters(SceneContext, nullptr, 0, *View->CachedViewUniformShaderParameters);
 		View->CachedViewUniformShaderParameters->WorldToVirtualTexture = WorldToUVRotate.ToMatrixNoScale();
-		View->CachedViewUniformShaderParameters->VirtualTextureParams = FVector4((float)vLevel, 0.f, OrthoWidth/(float)TextureSize.X, OrthoHeight / (float)TextureSize.Y);
+		View->CachedViewUniformShaderParameters->VirtualTextureParams = FVector4((float)vLevel, (float)MaxLevel, OrthoWidth/(float)TextureSize.X, OrthoHeight / (float)TextureSize.Y);
 		View->ViewUniformBuffer = TUniformBufferRef<FViewUniformShaderParameters>::CreateUniformBufferImmediate(*View->CachedViewUniformShaderParameters, UniformBuffer_SingleFrame);
 		UploadDynamicPrimitiveShaderDataForView(RHICmdList, *(const_cast<FScene*>(Scene)), *View);
 		Scene->UniformBuffers.VirtualTextureViewUniformBuffer.UpdateUniformBufferImmediate(*View->CachedViewUniformShaderParameters);
@@ -621,9 +659,9 @@ namespace RuntimeVirtualTexture
 				RDG_EVENT_NAME("VirtualTextureDraw"),
 				PassParameters,
 				ERenderGraphPassFlags::None,
-				[Scene, View, MaterialType, RuntimeVirtualTextureMask](FRHICommandListImmediate& RHICmdListImmediate)
+				[Scene, View, MaterialType, RuntimeVirtualTextureMask, vLevel, MaxLevel](FRHICommandListImmediate& RHICmdListImmediate)
 			{
-				DrawMeshes(RHICmdListImmediate, Scene, View, MaterialType, RuntimeVirtualTextureMask);
+				DrawMeshes(RHICmdListImmediate, Scene, View, MaterialType, RuntimeVirtualTextureMask, vLevel, MaxLevel);
 			});
 		}
 
diff --git a/Engine/Source/Runtime/Renderer/Private/VT/RuntimeVirtualTextureRender.h b/Engine/Source/Runtime/Renderer/Private/VT/RuntimeVirtualTextureRender.h
--- a/Engine/Source/Runtime/Renderer/Private/VT/RuntimeVirtualTextureRender.h
+++ b/Engine/Source/Runtime/Renderer/Private/VT/RuntimeVirtualTextureRender.h
@@ -22,5 +22,6 @@ namespace RuntimeVirtualTexture
 		FRHITexture2D* Texture1, FBox2D const& DestBox1, 
 		FTransform const& UVToWorld,
 		FBox2D const& UVRange,
-		uint8 vLevel);
+		uint8 vLevel, 
+		uint8 MaxLevel);
 }
diff --git a/Engine/Source/Runtime/Renderer/Public/PrimitiveSceneInfo.h b/Engine/Source/Runtime/Renderer/Public/PrimitiveSceneInfo.h
--- a/Engine/Source/Runtime/Renderer/Public/PrimitiveSceneInfo.h
+++ b/Engine/Source/Runtime/Renderer/Public/PrimitiveSceneInfo.h
@@ -182,6 +182,25 @@ struct FPrimitiveVirtualTextureFlags
 	uint8 RuntimeVirtualTextureMask : RuntimeVirtualTexture_BitCount;
 };
 
+/** Lod data used for runtime virtual texture page rendering. Packed to reduce memory overhead since one of these is allocated per primitive. */
+struct FPrimitiveVirtualTextureLodInfo
+{
+	/** Minimum Lod for primitive in the runtime virtual texture. */
+	uint16 MinLod : 4;
+	/** Maximum Lod for primitive in the runtime virtual texture. */
+	uint16 MaxLod : 4;
+	/** Bias to use for Lod calculation in the runtime virtual texture. */
+	uint16 LodBias : 3;
+	/** 
+	 * Culling method used to remove the primitive from low mips of the runtime virtual texture.
+	 * 0: CullValue is the number of low mips for which we cull the primitive from the runtime virtual texture.
+	 * 1: CullValue is the pixel coverage threshold at which we cull the primitive from the runtime virtual texture. 
+	 */
+	uint16 CullMethod : 1;
+	/** Value used according to the CullMethod. */
+	uint16 CullValue : 4;
+};
+
 /** The type of the octree used by FScene to find primitives. */
 typedef TOctree<FPrimitiveSceneInfoCompact,struct FPrimitiveOctreeSemantics> FScenePrimitiveOctree;
 
