- Fix GPU hang with Niagara GPU systems & planar reflections
- Change FXSystem to always call PreRender but notify if particles should be udpated or not, this allows us to write back our indirect arguments for rendering particles into planar reflections (or other scene captures)

diff --git a/Engine/Plugins/FX/Niagara/Source/Niagara/Classes/NiagaraEmitterInstanceBatcher.h b/Engine/Plugins/FX/Niagara/Source/Niagara/Classes/NiagaraEmitterInstanceBatcher.h
--- a/Engine/Plugins/FX/Niagara/Source/Niagara/Classes/NiagaraEmitterInstanceBatcher.h
+++ b/Engine/Plugins/FX/Niagara/Source/Niagara/Classes/NiagaraEmitterInstanceBatcher.h
@@ -80,7 +80,7 @@ public:
 	virtual void UpdateVectorField(UVectorFieldComponent* VectorFieldComponent) override {}
 	virtual void PreInitViews(FRHICommandListImmediate& RHICmdList) override;
 	virtual bool UsesGlobalDistanceField() const override;
-	virtual void PreRender(FRHICommandListImmediate& RHICmdList, const class FGlobalDistanceFieldParameterData* GlobalDistanceFieldParameterData) override;
+	virtual void PreRender(FRHICommandListImmediate& RHICmdList, const class FGlobalDistanceFieldParameterData* GlobalDistanceFieldParameterData, bool bAllowGPUParticleSceneUpdate) override;
 
 	virtual void Tick(float DeltaTime) override
 	{
diff --git a/Engine/Plugins/FX/Niagara/Source/Niagara/Private/NiagaraEmitterInstanceBatcher.cpp b/Engine/Plugins/FX/Niagara/Source/Niagara/Private/NiagaraEmitterInstanceBatcher.cpp
--- a/Engine/Plugins/FX/Niagara/Source/Niagara/Private/NiagaraEmitterInstanceBatcher.cpp
+++ b/Engine/Plugins/FX/Niagara/Source/Niagara/Private/NiagaraEmitterInstanceBatcher.cpp
@@ -576,7 +576,7 @@ bool NiagaraEmitterInstanceBatcher::UsesGlobalDistanceField() const
 	return false;
 }
 
-void NiagaraEmitterInstanceBatcher::PreRender(FRHICommandListImmediate& RHICmdList, const class FGlobalDistanceFieldParameterData* GlobalDistanceFieldParameterData)
+void NiagaraEmitterInstanceBatcher::PreRender(FRHICommandListImmediate& RHICmdList, const class FGlobalDistanceFieldParameterData* GlobalDistanceFieldParameterData, bool bAllowGPUParticleSceneUpdate)
 {
 	GlobalDistanceFieldParams = GlobalDistanceFieldParameterData ? *GlobalDistanceFieldParameterData : FGlobalDistanceFieldParameterData();
 
diff --git a/Engine/Source/Runtime/Engine/Private/Particles/FXSystem.cpp b/Engine/Source/Runtime/Engine/Private/Particles/FXSystem.cpp
--- a/Engine/Source/Runtime/Engine/Private/Particles/FXSystem.cpp
+++ b/Engine/Source/Runtime/Engine/Private/Particles/FXSystem.cpp
@@ -387,9 +387,9 @@ DECLARE_CYCLE_STAT(TEXT("FXPreRender_SimulateCDF"), STAT_CLM_FXPreRender_Simulat
 DECLARE_CYCLE_STAT(TEXT("FXPreRender_FinalizeCDF"), STAT_CLM_FXPreRender_FinalizeCDF, STATGROUP_CommandListMarkers);
 
 
-void FFXSystem::PreRender(FRHICommandListImmediate& RHICmdList, const FGlobalDistanceFieldParameterData* GlobalDistanceFieldParameterData)
+void FFXSystem::PreRender(FRHICommandListImmediate& RHICmdList, const FGlobalDistanceFieldParameterData* GlobalDistanceFieldParameterData, bool bAllowGPUParticleSceneUpdate)
 {
-	if (RHISupportsGPUParticles())
+	if (RHISupportsGPUParticles() && bAllowGPUParticleSceneUpdate)
 	{
 		SCOPED_DRAW_EVENT(RHICmdList, GPUParticles_PreRender);
 		UpdateMultiGPUResources(RHICmdList);
diff --git a/Engine/Source/Runtime/Engine/Private/Particles/FXSystemPrivate.h b/Engine/Source/Runtime/Engine/Private/Particles/FXSystemPrivate.h
--- a/Engine/Source/Runtime/Engine/Private/Particles/FXSystemPrivate.h
+++ b/Engine/Source/Runtime/Engine/Private/Particles/FXSystemPrivate.h
@@ -151,7 +151,7 @@ public:
 	FParticleEmitterInstance* CreateGPUSpriteEmitterInstance(FGPUSpriteEmitterInfo& EmitterInfo);
 	virtual void PreInitViews(FRHICommandListImmediate& RHICmdList) override;
 	virtual bool UsesGlobalDistanceField() const override;
-	virtual void PreRender(FRHICommandListImmediate& RHICmdList, const FGlobalDistanceFieldParameterData* GlobalDistanceFieldParameterData) override;
+	virtual void PreRender(FRHICommandListImmediate& RHICmdList, const FGlobalDistanceFieldParameterData* GlobalDistanceFieldParameterData, bool bAllowGPUParticleSceneUpdate) override;
 	virtual void PostRenderOpaque(
 		FRHICommandListImmediate& RHICmdList, 
 		FRHIUniformBuffer* ViewUniformBuffer,
diff --git a/Engine/Source/Runtime/Engine/Private/Particles/FXSystemSet.cpp b/Engine/Source/Runtime/Engine/Private/Particles/FXSystemSet.cpp
--- a/Engine/Source/Runtime/Engine/Private/Particles/FXSystemSet.cpp
+++ b/Engine/Source/Runtime/Engine/Private/Particles/FXSystemSet.cpp
@@ -109,12 +109,12 @@ bool FFXSystemSet::UsesGlobalDistanceField() const
 	return false;
 }
 
-void FFXSystemSet::PreRender(FRHICommandListImmediate& RHICmdList, const class FGlobalDistanceFieldParameterData* GlobalDistanceFieldParameterData)
+void FFXSystemSet::PreRender(FRHICommandListImmediate& RHICmdList, const class FGlobalDistanceFieldParameterData* GlobalDistanceFieldParameterData, bool bAllowGPUParticleSceneUpdate)
 {
 	for (FFXSystemInterface* FXSystem : FXSystems)
 	{
 		check(FXSystem);
-		FXSystem->PreRender(RHICmdList, GlobalDistanceFieldParameterData);
+		FXSystem->PreRender(RHICmdList, GlobalDistanceFieldParameterData, bAllowGPUParticleSceneUpdate);
 	}
 }
 
diff --git a/Engine/Source/Runtime/Engine/Private/Particles/FXSystemSet.h b/Engine/Source/Runtime/Engine/Private/Particles/FXSystemSet.h
--- a/Engine/Source/Runtime/Engine/Private/Particles/FXSystemSet.h
+++ b/Engine/Source/Runtime/Engine/Private/Particles/FXSystemSet.h
@@ -32,7 +32,7 @@ public:
 	virtual void UpdateVectorField(UVectorFieldComponent* VectorFieldComponent) override;
 	virtual void PreInitViews(FRHICommandListImmediate& RHICmdList) override;
 	virtual bool UsesGlobalDistanceField() const override;
-	virtual void PreRender(FRHICommandListImmediate& RHICmdList, const class FGlobalDistanceFieldParameterData* GlobalDistanceFieldParameterData) override;
+	virtual void PreRender(FRHICommandListImmediate& RHICmdList, const class FGlobalDistanceFieldParameterData* GlobalDistanceFieldParameterData, bool bAllowGPUParticleSceneUpdate) override;
 	virtual void PostRenderOpaque(
 		FRHICommandListImmediate& RHICmdList, 
 		FRHIUniformBuffer* ViewUniformBuffer,
diff --git a/Engine/Source/Runtime/Engine/Public/FXSystem.h b/Engine/Source/Runtime/Engine/Public/FXSystem.h
--- a/Engine/Source/Runtime/Engine/Public/FXSystem.h
+++ b/Engine/Source/Runtime/Engine/Public/FXSystem.h
@@ -205,7 +205,7 @@ public:
 	 * Notification from the renderer that it is about to draw FX belonging to
 	 * this system.
 	 */
-	virtual void PreRender(FRHICommandListImmediate& RHICmdList, const class FGlobalDistanceFieldParameterData* GlobalDistanceFieldParameterData) = 0;
+	virtual void PreRender(FRHICommandListImmediate& RHICmdList, const class FGlobalDistanceFieldParameterData* GlobalDistanceFieldParameterData, bool bAllowGPUParticleSceneUpdate) = 0;
 
 	/**
 	 * Notification from the renderer that opaque primitives have rendered.
diff --git a/Engine/Source/Runtime/Renderer/Private/DeferredShadingRenderer.cpp b/Engine/Source/Runtime/Renderer/Private/DeferredShadingRenderer.cpp
--- a/Engine/Source/Runtime/Renderer/Private/DeferredShadingRenderer.cpp
+++ b/Engine/Source/Runtime/Renderer/Private/DeferredShadingRenderer.cpp
@@ -1069,13 +1069,13 @@ void FDeferredShadingSceneRenderer::Render(FRHICommandListImmediate& RHICmdList)
 	const bool bAllowGPUParticleSceneUpdate = !Views[0].bIsPlanarReflection && !Views[0].bIsSceneCapture && !Views[0].bIsReflectionCapture;
 
 	// Notify the FX system that the scene is about to be rendered.
-	bool bDoFXPrerender = Scene->FXSystem && Views.IsValidIndex(0) && bAllowGPUParticleSceneUpdate;
+	bool bDoFXPrerender = Scene->FXSystem && Views.IsValidIndex(0);
 
 	if (bDoFXPrerender)
 	{
 		SCOPE_CYCLE_COUNTER(STAT_FDeferredShadingSceneRenderer_FXSystem_PreRender);
 		RHICmdList.SetCurrentStat(GET_STATID(STAT_CLM_FXPreRender));
-		Scene->FXSystem->PreRender(RHICmdList, &Views[0].GlobalDistanceFieldInfo.ParameterData);
+		Scene->FXSystem->PreRender(RHICmdList, &Views[0].GlobalDistanceFieldInfo.ParameterData, bAllowGPUParticleSceneUpdate);
 	}
 
 	bool bDidAfterTaskWork = false;
diff --git a/Engine/Source/Runtime/Renderer/Private/MobileShadingRenderer.cpp b/Engine/Source/Runtime/Renderer/Private/MobileShadingRenderer.cpp
--- a/Engine/Source/Runtime/Renderer/Private/MobileShadingRenderer.cpp
+++ b/Engine/Source/Runtime/Renderer/Private/MobileShadingRenderer.cpp
@@ -354,9 +354,9 @@ void FMobileSceneRenderer::Render(FRHICommandListImmediate& RHICmdList)
 	RHICmdList.ImmediateFlush(EImmediateFlushType::DispatchToRHIThread);
 
 	// Notify the FX system that the scene is about to be rendered.
-	if (Scene->FXSystem && !Views[0].bIsPlanarReflection && ViewFamily.EngineShowFlags.Particles)
+	if (Scene->FXSystem && ViewFamily.EngineShowFlags.Particles)
 	{
-		Scene->FXSystem->PreRender(RHICmdList, NULL);
+		Scene->FXSystem->PreRender(RHICmdList, NULL, !Views[0].bIsPlanarReflection);
 	}
 	FRHICommandListExecutor::GetImmediateCommandList().PollOcclusionQueries();
 	RHICmdList.ImmediateFlush(EImmediateFlushType::DispatchToRHIThread);
