Fix static analysis warning due to redundant assignment of variable to itself.

#rnx

diff --git a/Engine/Source/Runtime/Renderer/Private/RayTracing/RayTracingShadows.cpp b/Engine/Source/Runtime/Renderer/Private/RayTracing/RayTracingShadows.cpp
--- a/Engine/Source/Runtime/Renderer/Private/RayTracing/RayTracingShadows.cpp
+++ b/Engine/Source/Runtime/Renderer/Private/RayTracing/RayTracingShadows.cpp
@@ -158,8 +158,8 @@ void FDeferredShadingSceneRenderer::RenderRayTracingShadows(
 	if (LightSceneProxy->GetScissorRect(ScissorRect, View, View.ViewRect))
 	{
 		// Account for scissor being defined on the whole frame viewport while the trace is only on the view subrect
-		ScissorRect.Min = ScissorRect.Min;
-		ScissorRect.Max = ScissorRect.Max;
+		// ScissorRect.Min = ScissorRect.Min;
+		// ScissorRect.Max = ScissorRect.Max;
 	}
 	else
 	{
