Fix shadow variable warnings.

#rnx

diff --git a/Engine/Source/Runtime/Renderer/Private/LightRendering.cpp b/Engine/Source/Runtime/Renderer/Private/LightRendering.cpp
--- a/Engine/Source/Runtime/Renderer/Private/LightRendering.cpp
+++ b/Engine/Source/Runtime/Renderer/Private/LightRendering.cpp
@@ -1053,10 +1053,10 @@ void FDeferredShadingSceneRenderer::RenderLights(FRHICommandListImmediate& RHICm
 
 								for (int32 i = 0; i < InputParameterCount; i++)
 								{
-									const FLightSceneInfo* LightSceneInfo = DenoisingQueue[i].LightSceneInfo;
+									const FLightSceneInfo* LocalLightSceneInfo = DenoisingQueue[i].LightSceneInfo;
 
-									int32 LightIndex = LightIndices[i];
-									TRefCountPtr<IPooledRenderTarget>* RefDestination = &PreprocessedShadowMaskTextures[LightIndex - AttenuationLightStart];
+									int32 LocalLightIndex = LightIndices[i];
+									TRefCountPtr<IPooledRenderTarget>* RefDestination = &PreprocessedShadowMaskTextures[LocalLightIndex - AttenuationLightStart];
 									check(*RefDestination == nullptr);
 
 									GraphBuilder.QueueTextureExtraction(Outputs[i].DiffusePenumbra, RefDestination);
