ShaderComplexity improvements:
- Added support for mobile preview in the editor.
- Translucent material now show their native instruction count in forward avoid 0 cost error.
- Fixed scene prepass not being taken into account in the complexity result.
- In deferred shading, disabled all renderpasses but depth, editor hitproxy and debugviewmode.

diff --git a/Engine/Source/Runtime/Engine/Private/DebugViewModeInterface.cpp b/Engine/Source/Runtime/Engine/Private/DebugViewModeInterface.cpp
--- a/Engine/Source/Runtime/Engine/Private/DebugViewModeInterface.cpp
+++ b/Engine/Source/Runtime/Engine/Private/DebugViewModeInterface.cpp
@@ -12,7 +12,7 @@ DebugViewModeInterface.cpp: Contains definitions for rendering debug viewmodes.
 
 FDebugViewModeInterface* FDebugViewModeInterface::Singletons[DVSM_MAX] = {}; // Init to null.
 
-void FDebugViewModeInterface::SetDrawRenderState(EBlendMode BlendMode, FRenderState& DrawRenderState) const
+void FDebugViewModeInterface::SetDrawRenderState(EBlendMode BlendMode, FRenderState& DrawRenderState, bool bHasDepthPrepassForMaskedMaterial) const
 {
 	if (IsTranslucentBlendMode(BlendMode))
 	{
diff --git a/Engine/Source/Runtime/Engine/Private/DebugViewModeMaterialProxy.cpp b/Engine/Source/Runtime/Engine/Private/DebugViewModeMaterialProxy.cpp
--- a/Engine/Source/Runtime/Engine/Private/DebugViewModeMaterialProxy.cpp
+++ b/Engine/Source/Runtime/Engine/Private/DebugViewModeMaterialProxy.cpp
@@ -31,6 +31,8 @@ FDebugViewModeMaterialProxy::FDebugViewModeMaterialProxy(
 	, bSynchronousCompilation(InSynchronousCompilation)
 {
 	SetQualityLevelProperties(QualityLevel, false, FeatureLevel);
+	const EShaderPlatform ShaderPlatform = GetFeatureLevelShaderPlatform(FeatureLevel);
+
 	Material = InMaterialInterface->GetMaterial();
 	MaterialInterface->AppendReferencedTextures(ReferencedTextures);
 
@@ -62,22 +64,22 @@ FDebugViewModeMaterialProxy::FDebugViewModeMaterialProxy(
 		}
 
 		FMaterialShaderMapId ResourceId;
-		Resource->GetShaderMapId(GMaxRHIShaderPlatform, ResourceId);
+		Resource->GetShaderMapId(ShaderPlatform, ResourceId);
 
 		{
 			TArray<FShaderType*> ShaderTypes;
 			TArray<FVertexFactoryType*> VFTypes;
 			TArray<const FShaderPipelineType*> ShaderPipelineTypes;
-			GetDependentShaderAndVFTypes(GMaxRHIShaderPlatform, ShaderTypes, ShaderPipelineTypes, VFTypes);
+			GetDependentShaderAndVFTypes(ShaderPlatform, ShaderTypes, ShaderPipelineTypes, VFTypes);
 
 			// Overwrite the shader map Id's dependencies with ones that came from the FMaterial actually being compiled (this)
 			// This is necessary as we change FMaterial attributes like GetShadingModels(), which factor into the ShouldCache functions that determine dependent shader types
-			ResourceId.SetShaderDependencies(ShaderTypes, ShaderPipelineTypes, VFTypes, GMaxRHIShaderPlatform);
+			ResourceId.SetShaderDependencies(ShaderTypes, ShaderPipelineTypes, VFTypes, ShaderPlatform);
 		}
 
 		ResourceId.Usage = Usage;
 
-		CacheShaders(ResourceId, GMaxRHIShaderPlatform);
+		CacheShaders(ResourceId, ShaderPlatform);
 	}
 	else
 	{
diff --git a/Engine/Source/Runtime/Engine/Public/DebugViewModeInterface.h b/Engine/Source/Runtime/Engine/Public/DebugViewModeInterface.h
--- a/Engine/Source/Runtime/Engine/Public/DebugViewModeInterface.h
+++ b/Engine/Source/Runtime/Engine/Public/DebugViewModeInterface.h
@@ -48,7 +48,7 @@ public:
 	virtual ~FDebugViewModeInterface() {}
 
 	virtual FDebugViewModePS* GetPixelShader(const FMaterial* InMaterial, FVertexFactoryType* VertexFactoryType) const = 0;
-	virtual void SetDrawRenderState(EBlendMode BlendMode, FRenderState& DrawRenderState) const;
+	virtual void SetDrawRenderState(EBlendMode BlendMode, FRenderState& DrawRenderState, bool bHasDepthPrepassForMaskedMaterial) const;
 
 	/** The shader class name, used to filter out shaders that need to be compiled. */
 	const TCHAR* PixelShaderName;
diff --git a/Engine/Source/Runtime/Renderer/Private/DebugViewModeRendering.cpp b/Engine/Source/Runtime/Renderer/Private/DebugViewModeRendering.cpp
--- a/Engine/Source/Runtime/Renderer/Private/DebugViewModeRendering.cpp
+++ b/Engine/Source/Runtime/Renderer/Private/DebugViewModeRendering.cpp
@@ -364,7 +364,7 @@ void FDebugViewModeMeshProcessor::AddMeshBatch(const FMeshBatch& RESTRICT MeshBa
 	DrawRenderState.SetPassUniformBuffer(PassUniformBuffer);
 
 	FDebugViewModeInterface::FRenderState InterfaceRenderState;
-	DebugViewModeInterface->SetDrawRenderState(Material->GetBlendMode(), InterfaceRenderState);
+	DebugViewModeInterface->SetDrawRenderState(BatchMaterial->GetBlendMode(), InterfaceRenderState, Scene ? (Scene->GetShadingPath() == EShadingPath::Deferred && Scene->EarlyZPassMode != DDM_NonMaskedOnly) : false);
 	DrawRenderState.SetBlendState(InterfaceRenderState.BlendState);
 	DrawRenderState.SetDepthStencilState(InterfaceRenderState.DepthStencilState);
 
@@ -406,21 +406,37 @@ void FDebugViewModeMeshProcessor::UpdateInstructionCount(FDebugViewModeShaderEle
 {
 	check(InBatchMaterial && InVertexFactoryType);
 
-	const bool bDeferred = !IsAnyForwardShadingEnabled(GetFeatureLevelShaderPlatform(InBatchMaterial->GetFeatureLevel()));
-	const bool bLit = InBatchMaterial->GetShadingModels().IsLit();
+	if (Scene)
+	{
+		if (Scene->GetShadingPath() == EShadingPath::Deferred)
+		{
+			const bool bDeferred = !IsAnyForwardShadingEnabled(GetFeatureLevelShaderPlatform(InBatchMaterial->GetFeatureLevel()));
+			const bool bLit = InBatchMaterial->GetShadingModels().IsLit();
 
-	OutShaderElementData.NumVSInstructions = InBatchMaterial->GetShader<TBasePassVS<TUniformLightMapPolicy<LMP_NO_LIGHTMAP>, false>>(InVertexFactoryType)->GetNumInstructions();
-	OutShaderElementData.NumPSInstructions = InBatchMaterial->GetShader<TBasePassPS<TUniformLightMapPolicy<LMP_NO_LIGHTMAP>, false>>(InVertexFactoryType)->GetNumInstructions();
+			OutShaderElementData.NumVSInstructions = InBatchMaterial->GetShader<TBasePassVS<TUniformLightMapPolicy<LMP_NO_LIGHTMAP>, false>>(InVertexFactoryType)->GetNumInstructions();
+			OutShaderElementData.NumPSInstructions = InBatchMaterial->GetShader<TBasePassPS<TUniformLightMapPolicy<LMP_NO_LIGHTMAP>, false>>(InVertexFactoryType)->GetNumInstructions();
 
-	if (!bDeferred)
-	{
-		// Those numbers are taken from a simple material where common inputs are bound to vector parameters (to prevent constant optimizations).
-		OutShaderElementData.NumVSInstructions -= GShaderComplexityBaselineForwardVS - GShaderComplexityBaselineDeferredVS;
-		OutShaderElementData.NumPSInstructions -= bLit ? (GShaderComplexityBaselineForwardPS - GShaderComplexityBaselineDeferredPS) : (GShaderComplexityBaselineForwardUnlitPS - GShaderComplexityBaselineDeferredUnlitPS);
+			if (!bDeferred && !IsTranslucentBlendMode(InBatchMaterial->GetBlendMode()))
+			{
+				// Those numbers are taken from a simple material where common inputs are bound to vector parameters (to prevent constant optimizations).
+				OutShaderElementData.NumVSInstructions -= GShaderComplexityBaselineForwardVS - GShaderComplexityBaselineDeferredVS;
+				OutShaderElementData.NumPSInstructions -= bLit ? (GShaderComplexityBaselineForwardPS - GShaderComplexityBaselineDeferredPS) : (GShaderComplexityBaselineForwardUnlitPS - GShaderComplexityBaselineDeferredUnlitPS);
+			}
+
+			OutShaderElementData.NumVSInstructions = FMath::Max<int32>(0, OutShaderElementData.NumVSInstructions);
+			OutShaderElementData.NumPSInstructions = FMath::Max<int32>(0, OutShaderElementData.NumPSInstructions);
+		}
+		else // EShadingPath::Mobile
+		{
+			TMobileBasePassVSPolicyParamType<FUniformLightMapPolicy>* MobileVS = nullptr;
+			TMobileBasePassPSPolicyParamType<FUniformLightMapPolicy>* MobilePS = nullptr;
+
+			MobileBasePass::GetShaders(LMP_NO_LIGHTMAP, 0, *InBatchMaterial, InVertexFactoryType, false, MobileVS, MobilePS);
+
+			OutShaderElementData.NumVSInstructions = MobileVS ? MobileVS->GetNumInstructions() : 0;
+			OutShaderElementData.NumPSInstructions = MobilePS ? MobilePS->GetNumInstructions() : 0;
+		}
 	}
-	
-	OutShaderElementData.NumVSInstructions = FMath::Max<int32>(0, OutShaderElementData.NumVSInstructions);
-	OutShaderElementData.NumPSInstructions = FMath::Max<int32>(0, OutShaderElementData.NumPSInstructions);
 }
 
 FMeshPassProcessor* CreateDebugViewModePassProcessor(const FScene* Scene, const FSceneView* InViewIfDynamicMeshCommand, FMeshPassDrawListContext* InDrawListContext)
@@ -429,9 +445,9 @@ FMeshPassProcessor* CreateDebugViewModePassProcessor(const FScene* Scene, const
 	return new(FMemStack::Get()) FDebugViewModeMeshProcessor(Scene, FeatureLevel, InViewIfDynamicMeshCommand, nullptr, false, InDrawListContext);
 }
 
+FRegisterPassProcessorCreateFunction RegisterDebugViewModeMobilePass(&CreateDebugViewModePassProcessor, EShadingPath::Mobile, EMeshPass::DebugViewMode, EMeshPassFlags::MainView);
 FRegisterPassProcessorCreateFunction RegisterDebugViewModePass(&CreateDebugViewModePassProcessor, EShadingPath::Deferred, EMeshPass::DebugViewMode, EMeshPassFlags::MainView);
 
-
 void InitDebugViewModeInterfaces()
 {
 	FDebugViewModeInterface::SetInterface(DVSM_ShaderComplexity, new FComplexityAccumulateInterface(true, false));
diff --git a/Engine/Source/Runtime/Renderer/Private/DeferredShadingRenderer.cpp b/Engine/Source/Runtime/Renderer/Private/DeferredShadingRenderer.cpp
--- a/Engine/Source/Runtime/Renderer/Private/DeferredShadingRenderer.cpp
+++ b/Engine/Source/Runtime/Renderer/Private/DeferredShadingRenderer.cpp
@@ -1819,7 +1819,7 @@ void FDeferredShadingSceneRenderer::Render(FRHICommandListImmediate& RHICmdList)
 	GRenderTargetPool.AddPhaseEvent(TEXT("Translucency"));
 
 	// Draw translucency.
-	if (bCanOverlayRayTracingOutput && ViewFamily.EngineShowFlags.Translucency && !ViewFamily.EngineShowFlags.VisualizeLightCulling)
+	if (bCanOverlayRayTracingOutput && ViewFamily.EngineShowFlags.Translucency && !ViewFamily.EngineShowFlags.VisualizeLightCulling && !ViewFamily.UseDebugViewPS())
 	{
 		CSV_SCOPED_TIMING_STAT_EXCLUSIVE(RenderTranslucency);
 		SCOPE_CYCLE_COUNTER(STAT_TranslucencyDrawTime);
diff --git a/Engine/Source/Runtime/Renderer/Private/MaterialTexCoordScalesRendering.cpp b/Engine/Source/Runtime/Renderer/Private/MaterialTexCoordScalesRendering.cpp
--- a/Engine/Source/Runtime/Renderer/Private/MaterialTexCoordScalesRendering.cpp
+++ b/Engine/Source/Runtime/Renderer/Private/MaterialTexCoordScalesRendering.cpp
@@ -48,7 +48,7 @@ void FMaterialTexCoordScalePS::GetDebugViewModeShaderBindings(
 	ShaderBindings.Add(PrimitiveAlphaParameter,  (!PrimitiveSceneProxy || PrimitiveSceneProxy->IsSelected()) ? 1.f : .2f);
 }
 
-void FOutputMaterialTexCoordScaleInterface::SetDrawRenderState(EBlendMode BlendMode, FRenderState& DrawRenderState) const
+void FOutputMaterialTexCoordScaleInterface::SetDrawRenderState(EBlendMode BlendMode, FRenderState& DrawRenderState, bool bHasDepthPrepassForMaskedMaterial) const
 {
 	DrawRenderState.BlendState = TStaticBlendState<>::GetRHI();
 	DrawRenderState.DepthStencilState = TStaticDepthStencilState<false, CF_DepthNearOrEqual>::GetRHI();
diff --git a/Engine/Source/Runtime/Renderer/Private/MaterialTexCoordScalesRendering.h b/Engine/Source/Runtime/Renderer/Private/MaterialTexCoordScalesRendering.h
--- a/Engine/Source/Runtime/Renderer/Private/MaterialTexCoordScalesRendering.h
+++ b/Engine/Source/Runtime/Renderer/Private/MaterialTexCoordScalesRendering.h
@@ -34,7 +34,7 @@ public:
 	static bool ShouldCompilePermutation(const FMeshMaterialShaderPermutationParameters& Parameters)
 	{
 		// See FDebugViewModeMaterialProxy::GetFriendlyName()
-		return AllowDebugViewShaderMode(DVSM_OutputMaterialTextureScales, Parameters.Platform, GetMaxSupportedFeatureLevel(Parameters.Platform)) && Parameters.Material->GetFriendlyName().Contains(TEXT("MaterialTexCoordScale"));
+		return AllowDebugViewShaderMode(DVSM_OutputMaterialTextureScales, Parameters.Platform, Parameters.Material->GetFeatureLevel()) && Parameters.Material->GetFriendlyName().Contains(TEXT("MaterialTexCoordScale"));
 	}
 
 	FMaterialTexCoordScalePS(const ShaderMetaType::CompiledShaderInitializerType& Initializer):
@@ -111,7 +111,7 @@ public:
 	{ 
 		return InMaterial->GetShader<FMaterialTexCoordScalePS>(VertexFactoryType); 
 	}
-	virtual void SetDrawRenderState(EBlendMode BlendMode, FRenderState& DrawRenderState) const override;
+	virtual void SetDrawRenderState(EBlendMode BlendMode, FRenderState& DrawRenderState, bool bHasDepthPrepassForMaskedMaterial) const override;
 };
 
 #endif // !(UE_BUILD_SHIPPING || UE_BUILD_TEST)
diff --git a/Engine/Source/Runtime/Renderer/Private/MeshTexCoordSizeAccuracyRendering.h b/Engine/Source/Runtime/Renderer/Private/MeshTexCoordSizeAccuracyRendering.h
--- a/Engine/Source/Runtime/Renderer/Private/MeshTexCoordSizeAccuracyRendering.h
+++ b/Engine/Source/Runtime/Renderer/Private/MeshTexCoordSizeAccuracyRendering.h
@@ -31,7 +31,7 @@ public:
 	static bool ShouldCompilePermutation(const FMeshMaterialShaderPermutationParameters& Parameters)
 	{
 		// See FDebugViewModeMaterialProxy::GetFriendlyName()
-		return AllowDebugViewShaderMode(DVSM_MeshUVDensityAccuracy, Parameters.Platform, GetMaxSupportedFeatureLevel(Parameters.Platform)) && Parameters.Material->GetFriendlyName().Contains(TEXT("MeshTexCoordSizeAccuracy"));
+		return AllowDebugViewShaderMode(DVSM_MeshUVDensityAccuracy, Parameters.Platform, Parameters.Material->GetFeatureLevel()) && Parameters.Material->GetFriendlyName().Contains(TEXT("MeshTexCoordSizeAccuracy"));
 	}
 
 	FMeshTexCoordSizeAccuracyPS(const ShaderMetaType::CompiledShaderInitializerType& Initializer):
diff --git a/Engine/Source/Runtime/Renderer/Private/MobileShadingRenderer.cpp b/Engine/Source/Runtime/Renderer/Private/MobileShadingRenderer.cpp
--- a/Engine/Source/Runtime/Renderer/Private/MobileShadingRenderer.cpp
+++ b/Engine/Source/Runtime/Renderer/Private/MobileShadingRenderer.cpp
@@ -45,7 +45,7 @@
 #include "EngineModule.h"
 #include "GPUScene.h"
 #include "MaterialSceneTextureId.h"
-
+#include "DebugViewModeRendering.h"
 #include "VisualizeTexture.h"
 
 uint32 GetShadowQuality();
@@ -487,6 +487,17 @@ void FMobileSceneRenderer::Render(FRHICommandListImmediate& RHICmdList)
 	RenderMobileBasePass(RHICmdList, ViewList);
 	RHICmdList.ImmediateFlush(EImmediateFlushType::DispatchToRHIThread);
 
+#if !(UE_BUILD_SHIPPING || UE_BUILD_TEST)
+	if (ViewFamily.UseDebugViewPS())
+	{
+		// Here we use the base pass depth result to get z culling for opaque and masque.
+		// The color needs to be cleared at this point since shader complexity renders in additive.
+		DrawClearQuad(RHICmdList, FLinearColor::Black);
+		RenderMobileDebugView(RHICmdList, ViewList);
+		RHICmdList.ImmediateFlush(EImmediateFlushType::DispatchToRHIThread);
+	}
+#endif // !(UE_BUILD_SHIPPING || UE_BUILD_TEST)
+
 	const bool bAdrenoOcclusionMode = CVarMobileAdrenoOcclusionMode.GetValueOnRenderThread() != 0;
 	if (!bAdrenoOcclusionMode)
 	{
@@ -709,6 +720,33 @@ void FMobileSceneRenderer::BasicPostProcess(FRHICommandListImmediate& RHICmdList
 	CompositeContext.Process(Context.FinalOutput.GetPass(), TEXT("ES2BasicPostProcess"));
 }
 
+void FMobileSceneRenderer::RenderMobileDebugView(FRHICommandListImmediate& RHICmdList, const TArrayView<const FViewInfo*> PassViews)
+{
+#if !(UE_BUILD_SHIPPING || UE_BUILD_TEST)
+	CSV_SCOPED_TIMING_STAT_EXCLUSIVE(RenderDebugView);
+	SCOPED_DRAW_EVENT(RHICmdList, MobileDebugView);
+	SCOPE_CYCLE_COUNTER(STAT_BasePassDrawTime);
+
+	for (int32 ViewIndex = 0; ViewIndex < PassViews.Num(); ViewIndex++)
+	{
+		SCOPED_CONDITIONAL_DRAW_EVENTF(RHICmdList, EventView, Views.Num() > 1, TEXT("View%d"), ViewIndex);
+		const FViewInfo& View = *PassViews[ViewIndex];
+		if (!View.ShouldRenderView())
+		{
+			continue;
+		}
+
+		FSceneRenderTargets& SceneContext = FSceneRenderTargets::Get(RHICmdList);
+		FDebugViewModePassPassUniformParameters PassParameters;
+		SetupDebugViewModePassUniformBuffer(SceneContext, View.GetFeatureLevel(), PassParameters);
+		Scene->UniformBuffers.DebugViewModePassUniformBuffer.UpdateUniformBufferImmediate(PassParameters);
+		
+		RHICmdList.SetViewport(View.ViewRect.Min.X, View.ViewRect.Min.Y, 0, View.ViewRect.Max.X, View.ViewRect.Max.Y, 1);
+		View.ParallelMeshDrawCommandPasses[EMeshPass::DebugViewMode].DispatchDraw(nullptr, RHICmdList);
+	}
+#endif // !(UE_BUILD_SHIPPING || UE_BUILD_TEST)
+}
+
 void FMobileSceneRenderer::RenderOcclusion(FRHICommandListImmediate& RHICmdList)
 {
 	if (!DoOcclusionQueries(FeatureLevel))
diff --git a/Engine/Source/Runtime/Renderer/Private/PrimitiveDistanceAccuracyRendering.h b/Engine/Source/Runtime/Renderer/Private/PrimitiveDistanceAccuracyRendering.h
--- a/Engine/Source/Runtime/Renderer/Private/PrimitiveDistanceAccuracyRendering.h
+++ b/Engine/Source/Runtime/Renderer/Private/PrimitiveDistanceAccuracyRendering.h
@@ -31,7 +31,7 @@ public:
 	static bool ShouldCompilePermutation(const FMeshMaterialShaderPermutationParameters& Parameters)
 	{
 		// See FDebugViewModeMaterialProxy::GetFriendlyName()
-		return AllowDebugViewShaderMode(DVSM_PrimitiveDistanceAccuracy, Parameters.Platform, GetMaxSupportedFeatureLevel(Parameters.Platform)) && Parameters.Material->GetFriendlyName().Contains(TEXT("PrimitiveDistanceAccuracy"));
+		return AllowDebugViewShaderMode(DVSM_PrimitiveDistanceAccuracy, Parameters.Platform, Parameters.Material->GetFeatureLevel()) && Parameters.Material->GetFriendlyName().Contains(TEXT("PrimitiveDistanceAccuracy"));
 	}
 
 	FPrimitiveDistanceAccuracyPS(const ShaderMetaType::CompiledShaderInitializerType& Initializer):
diff --git a/Engine/Source/Runtime/Renderer/Private/RequiredTextureResolutionRendering.h b/Engine/Source/Runtime/Renderer/Private/RequiredTextureResolutionRendering.h
--- a/Engine/Source/Runtime/Renderer/Private/RequiredTextureResolutionRendering.h
+++ b/Engine/Source/Runtime/Renderer/Private/RequiredTextureResolutionRendering.h
@@ -27,7 +27,7 @@ public:
 	static bool ShouldCompilePermutation(const FMeshMaterialShaderPermutationParameters& Parameters)
 	{
 		// See FDebugViewModeMaterialProxy::GetFriendlyName()
-		return AllowDebugViewShaderMode(DVSM_RequiredTextureResolution, Parameters.Platform, GetMaxSupportedFeatureLevel(Parameters.Platform)) && Parameters.Material->GetFriendlyName().Contains(TEXT("RequiredTextureResolution"));
+		return AllowDebugViewShaderMode(DVSM_RequiredTextureResolution, Parameters.Platform, Parameters.Material->GetFeatureLevel()) && Parameters.Material->GetFriendlyName().Contains(TEXT("RequiredTextureResolution"));
 	}
 
 	FRequiredTextureResolutionPS(const ShaderMetaType::CompiledShaderInitializerType& Initializer):
diff --git a/Engine/Source/Runtime/Renderer/Private/SceneRendering.cpp b/Engine/Source/Runtime/Renderer/Private/SceneRendering.cpp
--- a/Engine/Source/Runtime/Renderer/Private/SceneRendering.cpp
+++ b/Engine/Source/Runtime/Renderer/Private/SceneRendering.cpp
@@ -2739,6 +2739,24 @@ void FSceneRenderer::SetupMeshPass(FViewInfo& View, FExclusiveDepthStencil::Type
 				continue;
 			}
 
+			if (ViewFamily.UseDebugViewPS() && ShadingPath == EShadingPath::Deferred)
+			{
+				switch (PassType)
+				{
+					case EMeshPass::DepthPass:
+					case EMeshPass::CustomDepth:
+					case EMeshPass::DebugViewMode:
+#if WITH_EDITOR
+					case EMeshPass::HitProxy:
+					case EMeshPass::HitProxyOpaqueOnly:
+					case EMeshPass::EditorSelection:
+#endif
+						break;
+					default:
+						continue;
+				}
+			}
+
 			PassProcessorCreateFunction CreateFunction = FPassProcessorManager::GetCreateFunction(ShadingPath, PassType);
 			FMeshPassProcessor* MeshPassProcessor = CreateFunction(Scene, &View, nullptr);
 
diff --git a/Engine/Source/Runtime/Renderer/Private/SceneRendering.h b/Engine/Source/Runtime/Renderer/Private/SceneRendering.h
--- a/Engine/Source/Runtime/Renderer/Private/SceneRendering.h
+++ b/Engine/Source/Runtime/Renderer/Private/SceneRendering.h
@@ -1708,6 +1708,9 @@ protected:
 
 	void RenderMobileEditorPrimitives(FRHICommandList& RHICmdList, const FViewInfo& View, const FMeshPassProcessorRenderState& DrawRenderState);
 
+	/** Renders the debug view pass for mobile. */
+	void RenderMobileDebugView(FRHICommandListImmediate& RHICmdList, const TArrayView<const FViewInfo*> PassViews);
+
 	/** Render modulated shadow projections in to the scene, loops over any unrendered shadows until all are processed.*/
 	void RenderModulatedShadowProjections(FRHICommandListImmediate& RHICmdList);
 
diff --git a/Engine/Source/Runtime/Renderer/Private/ShaderComplexityRendering.cpp b/Engine/Source/Runtime/Renderer/Private/ShaderComplexityRendering.cpp
--- a/Engine/Source/Runtime/Renderer/Private/ShaderComplexityRendering.cpp
+++ b/Engine/Source/Runtime/Renderer/Private/ShaderComplexityRendering.cpp
@@ -103,9 +103,27 @@ FDebugViewModePS* FComplexityAccumulateInterface::GetPixelShader(const FMaterial
 	}
 
 }
-void FComplexityAccumulateInterface::SetDrawRenderState(EBlendMode BlendMode, FRenderState& DrawRenderState) const
+void FComplexityAccumulateInterface::SetDrawRenderState(EBlendMode BlendMode, FRenderState& DrawRenderState, bool bHasDepthPrepassForMaskedMaterial) const
 {
-	DrawRenderState.DepthStencilState = TStaticDepthStencilState<false, CF_DepthNearOrEqual>::GetRHI();
+	if (BlendMode == BLEND_Opaque)
+	{
+		DrawRenderState.DepthStencilState = TStaticDepthStencilState<true, CF_DepthNearOrEqual>::GetRHI();
+	}
+	else if (BlendMode == BLEND_Masked)
+	{
+		if (bHasDepthPrepassForMaskedMaterial)
+		{
+			DrawRenderState.DepthStencilState = TStaticDepthStencilState<false, CF_Equal>::GetRHI();
+		}
+		else
+		{
+			DrawRenderState.DepthStencilState = TStaticDepthStencilState<false, CF_DepthNearOrEqual>::GetRHI();
+		}
+	}
+	else // Translucent
+	{
+		DrawRenderState.DepthStencilState = TStaticDepthStencilState<false, CF_DepthNearOrEqual>::GetRHI();
+	}
 	DrawRenderState.BlendState = TStaticBlendState<CW_RGBA, BO_Add, BF_One, BF_One, BO_Add, BF_Zero, BF_One>::GetRHI();
 }
 
diff --git a/Engine/Source/Runtime/Renderer/Private/ShaderComplexityRendering.h b/Engine/Source/Runtime/Renderer/Private/ShaderComplexityRendering.h
--- a/Engine/Source/Runtime/Renderer/Private/ShaderComplexityRendering.h
+++ b/Engine/Source/Runtime/Renderer/Private/ShaderComplexityRendering.h
@@ -36,7 +36,7 @@ public:
 	static bool ShouldCompilePermutation(const FMeshMaterialShaderPermutationParameters& Parameters)
 	{
 		// See FDebugViewModeMaterialProxy::GetFriendlyName()
-		return AllowDebugViewShaderMode(bQuadComplexity ? DVSM_QuadComplexity : DVSM_ShaderComplexity, Parameters.Platform, GetMaxSupportedFeatureLevel(Parameters.Platform)) && Parameters.Material->GetFriendlyName().Contains(TEXT("ComplexityAccumulate"));
+		return AllowDebugViewShaderMode(bQuadComplexity ? DVSM_QuadComplexity : DVSM_ShaderComplexity, Parameters.Platform, Parameters.Material->GetFeatureLevel()) && Parameters.Material->GetFriendlyName().Contains(TEXT("ComplexityAccumulate"));
 	}
 
 	TComplexityAccumulatePS(const ShaderMetaType::CompiledShaderInitializerType& Initializer):
@@ -60,7 +60,7 @@ public:
 
 	static void ModifyCompilationEnvironment(const FMaterialShaderPermutationParameters& Parameters, FShaderCompilerEnvironment& OutEnvironment)
 	{
-		OutEnvironment.SetDefine(TEXT("OUTPUT_QUAD_OVERDRAW"), AllowDebugViewShaderMode(DVSM_QuadComplexity, Parameters.Platform, GetMaxSupportedFeatureLevel(Parameters.Platform)));
+		OutEnvironment.SetDefine(TEXT("OUTPUT_QUAD_OVERDRAW"), AllowDebugViewShaderMode(DVSM_QuadComplexity, Parameters.Platform, Parameters.Material->GetFeatureLevel()));
 		TCHAR BufferRegister[] = { 'u', '0', 0 };
 		BufferRegister[1] += FSceneRenderTargets::GetQuadOverdrawUAVIndex(Parameters.Platform, Parameters.Material->GetFeatureLevel());
 		OutEnvironment.SetDefine(TEXT("QUAD_BUFFER_REGISTER"), BufferRegister);
@@ -102,7 +102,7 @@ public:
 	{}
 
 	virtual FDebugViewModePS* GetPixelShader(const FMaterial* InMaterial, FVertexFactoryType* VertexFactoryType) const override;
-	virtual void SetDrawRenderState(EBlendMode BlendMode, FRenderState& DrawRenderState) const;
+	virtual void SetDrawRenderState(EBlendMode BlendMode, FRenderState& DrawRenderState, bool bHasDepthPrepassForMaskedMaterial) const;
 };
 
 #endif //!(UE_BUILD_SHIPPING || UE_BUILD_TEST)
