Fixed: Refraction no longer refracts on Mobile devices

diff --git a/Engine/Source/Runtime/Renderer/Private/DistortionRendering.cpp b/Engine/Source/Runtime/Renderer/Private/DistortionRendering.cpp
--- a/Engine/Source/Runtime/Renderer/Private/DistortionRendering.cpp
+++ b/Engine/Source/Runtime/Renderer/Private/DistortionRendering.cpp
@@ -386,10 +386,11 @@ public:
 	{
 		FMeshMaterialShader::ModifyCompilationEnvironment(Parameters, OutEnvironment);
 
-		if (GetMaxSupportedFeatureLevel(Parameters.Platform) <= ERHIFeatureLevel::ES3_1)
+		if (IsVulkanMobilePlatform(Parameters.Platform))
 		{
-			// do not allow to sample scene textures on mobile during accumulation path
-			OutEnvironment.SetDefine(TEXT("SCENE_TEXTURES_DISABLED"), 1);
+			// depth fetch only available during base pass rendering
+			// TODO: find better place to enable frame buffer fetch feature only for base pass
+			OutEnvironment.SetDefine(TEXT("VULKAN_SUBPASS_DEPTHFETCH"), 0);
 		}
 	}
 	
