Fixed: The screen will be darker than other platforms(Android,PC/Mac preview) on IOS if both MSAA and MobileTonemapperFilm enabled

diff --git a/Engine/Shaders/Private/PostProcessTonemap.usf b/Engine/Shaders/Private/PostProcessTonemap.usf
--- a/Engine/Shaders/Private/PostProcessTonemap.usf
+++ b/Engine/Shaders/Private/PostProcessTonemap.usf
@@ -464,6 +464,13 @@ float4 TonemapCommonPS(
 
 	#else 
 
+		#if USE_MSAA
+			// look for PreTonemapMSAA 
+			SceneColor.rgb *= rcp(SceneColor.r*(-0.299) + SceneColor.g*(-0.587) + SceneColor.b*(-0.114) + 1.0);
+			// Try to kill negatives and NaNs here
+			SceneColor.rgb = max(SceneColor.rgb, 0); 
+		#endif
+	
 		half3 LinearColor = SceneColor.rgb * ColorScale0.rgb;
 
 		#if USE_BLOOM
@@ -758,10 +765,7 @@ void MainPS_ES2(
 	}
 
 	#if USE_GAMMA_ONLY
-		#if USE_MSAA
-			SceneColor.rgb *= rcp(SceneColor.r*(-0.299) + SceneColor.g*(-0.587) + SceneColor.b*(-0.114) + 1.0);
-		#endif
-
+		
 		#if FEATURE_LEVEL == FEATURE_LEVEL_ES3_1
 			OutColor.rgb = pow(SceneColor.rgb, InverseGamma.x);
 		#else
diff --git a/Engine/Source/Runtime/Renderer/Private/PostProcess/PostProcessTonemap.cpp b/Engine/Source/Runtime/Renderer/Private/PostProcess/PostProcessTonemap.cpp
--- a/Engine/Source/Runtime/Renderer/Private/PostProcess/PostProcessTonemap.cpp
+++ b/Engine/Source/Runtime/Renderer/Private/PostProcess/PostProcessTonemap.cpp
@@ -127,6 +127,7 @@ class FTonemapperVignetteDim       : SHADER_PERMUTATION_BOOL("USE_VIGNETTE");
 class FTonemapperSharpenDim        : SHADER_PERMUTATION_BOOL("USE_SHARPEN");
 class FTonemapperGrainJitterDim    : SHADER_PERMUTATION_BOOL("USE_GRAIN_JITTER");
 class FTonemapperSwitchAxis        : SHADER_PERMUTATION_BOOL("NEEDTOSWITCHVERTICLEAXIS");
+class FTonemapperMsaaDim		   : SHADER_PERMUTATION_BOOL("USE_MSAA");
 
 using FCommonDomain = TShaderPermutationDomain<
 	FTonemapperBloomDim,
@@ -135,7 +136,8 @@ using FCommonDomain = TShaderPermutationDomain<
 	FTonemapperVignetteDim,
 	FTonemapperSharpenDim,
 	FTonemapperGrainJitterDim,
-	FTonemapperSwitchAxis>;
+	FTonemapperSwitchAxis,
+	FTonemapperMsaaDim>;
 
 FORCEINLINE_DEBUGGABLE bool ShouldCompileCommonPermutation(const FGlobalShaderPermutationParameters& Parameters, const FCommonDomain& PermutationVector)
 {
@@ -145,6 +147,12 @@ FORCEINLINE_DEBUGGABLE bool ShouldCompileCommonPermutation(const FGlobalShaderPe
 		return false;
 	}
 
+	// MSAA pre-resolve step only used on iOS atm
+	if (PermutationVector.Get<FTonemapperMsaaDim>() && !IsMetalMobilePlatform(Parameters.Platform))
+	{
+		return false;
+	}
+
 	// If GammaOnly, don't compile any other dimmension == true.
 	if (PermutationVector.Get<FTonemapperGammaOnlyDim>())
 	{
@@ -152,7 +160,8 @@ FORCEINLINE_DEBUGGABLE bool ShouldCompileCommonPermutation(const FGlobalShaderPe
 			!PermutationVector.Get<FTonemapperGrainIntensityDim>() &&
 			!PermutationVector.Get<FTonemapperVignetteDim>() &&
 			!PermutationVector.Get<FTonemapperSharpenDim>() &&
-			!PermutationVector.Get<FTonemapperGrainJitterDim>();
+			!PermutationVector.Get<FTonemapperGrainJitterDim>() &&
+			!PermutationVector.Get<FTonemapperMsaaDim>();
 	}
 	return true;
 }
@@ -181,6 +190,14 @@ FCommonDomain BuildCommonPermutationDomain(const FViewInfo& View, bool bGammaOnl
 	PermutationVector.Set<FTonemapperSharpenDim>(CVarTonemapperSharpen.GetValueOnRenderThread() > 0.0f);	
 	PermutationVector.Set<FTonemapperSwitchAxis>(bSwitchVerticalAxis);
 
+	if (IsMetalMobilePlatform(View.GetShaderPlatform()))
+	{
+		static const auto CVarMobileMSAA = IConsoleManager::Get().FindTConsoleVariableDataInt(TEXT("r.MobileMSAA"));
+		bool bMSAA = (CVarMobileMSAA ? CVarMobileMSAA->GetValueOnAnyThread() > 1 : false);
+
+		PermutationVector.Set<FTonemapperMsaaDim>(bMSAA);
+	}
+
 	return PermutationVector;
 }
 
@@ -196,7 +213,7 @@ using FDesktopDomain = TShaderPermutationDomain<
 	FTonemapperGrainQuantizationDim,
 	FTonemapperOutputDeviceDim>;
 
-FDesktopDomain RemapPermutation(FDesktopDomain PermutationVector)
+FDesktopDomain RemapPermutation(FDesktopDomain PermutationVector, ERHIFeatureLevel::Type FeatureLevel)
 {
 	FCommonDomain CommonPermutationVector = PermutationVector.Get<FCommonDomain>();
 
@@ -228,7 +245,14 @@ FDesktopDomain RemapPermutation(FDesktopDomain PermutationVector)
 		PermutationVector.Set<FTonemapperGrainQuantizationDim>(false);
 	else
 		PermutationVector.Set<FTonemapperGrainQuantizationDim>(true);
-
+	
+	// Mobile supports only sRGB and LinearNoToneCurve output
+	if (FeatureLevel <= ERHIFeatureLevel::ES3_1 && 
+		PermutationVector.Get<FTonemapperOutputDeviceDim>() != FTonemapperOutputDevice::LinearNoToneCurve)
+	{
+		PermutationVector.Set<FTonemapperOutputDeviceDim>(FTonemapperOutputDevice::sRGB);
+	}
+	
 	PermutationVector.Set<FCommonDomain>(CommonPermutationVector);
 	return PermutationVector;
 }
@@ -237,7 +261,7 @@ bool ShouldCompileDesktopPermutation(const FGlobalShaderPermutationParameters& P
 {
 	auto CommonPermutationVector = PermutationVector.Get<FCommonDomain>();
 
-	if (RemapPermutation(PermutationVector) != PermutationVector)
+	if (RemapPermutation(PermutationVector, GetMaxSupportedFeatureLevel(Parameters.Platform)) != PermutationVector)
 	{
 		return false;
 	}
@@ -704,7 +728,7 @@ class FPostProcessTonemapPS : public FGlobalShader
 
 	static bool ShouldCompilePermutation(const FGlobalShaderPermutationParameters& Parameters)
 	{
-		if (!IsFeatureLevelSupported(Parameters.Platform, ERHIFeatureLevel::ES2))
+		if (!IsFeatureLevelSupported(Parameters.Platform, ERHIFeatureLevel::ES3_1))
 		{
 			return false;
 		}
@@ -1013,7 +1037,7 @@ void FRCPassPostProcessTonemap::Process(FRenderingCompositePassContext& Context)
 		{
 			DesktopPermutationVector.Set<TonemapperPermutation::FTonemapperOutputDeviceDim>(GetOutputDeviceValue());
 		}
-		DesktopPermutationVector = TonemapperPermutation::RemapPermutation(DesktopPermutationVector);
+		DesktopPermutationVector = TonemapperPermutation::RemapPermutation(DesktopPermutationVector, Context.GetFeatureLevel());
 	}
 
 	if (bIsComputePass)
@@ -1194,7 +1218,6 @@ class FPostProcessTonemapPS_ES2 : public FGlobalShader
 	DECLARE_GLOBAL_SHADER(FPostProcessTonemapPS_ES2);
 
 	// Mobile renderer specific permutation dimensions.
-	class FTonemapperMsaaDim        : SHADER_PERMUTATION_BOOL("USE_MSAA");
 	class FTonemapperDOFDim         : SHADER_PERMUTATION_BOOL("USE_DOF");
 	class FTonemapperLightShaftsDim : SHADER_PERMUTATION_BOOL("USE_LIGHT_SHAFTS");
 	class FTonemapper32BPPHDRDim    : SHADER_PERMUTATION_BOOL("USE_32BPP_HDR");
@@ -1204,7 +1227,6 @@ class FPostProcessTonemapPS_ES2 : public FGlobalShader
 	
 	using FPermutationDomain = TShaderPermutationDomain<
 		TonemapperPermutation::FCommonDomain,
-		FTonemapperMsaaDim,
 		FTonemapperDOFDim,
 		FTonemapperLightShaftsDim,
 		FTonemapper32BPPHDRDim,
@@ -1256,11 +1278,11 @@ class FPostProcessTonemapPS_ES2 : public FGlobalShader
 		EnableIfSet<TonemapperPermutation::FTonemapperGrainIntensityDim>(WantedCommonPermutationVector, RemappedCommonPermutationVector);
 		// Switch Y axis
 		EnableIfSet<TonemapperPermutation::FTonemapperSwitchAxis>(WantedCommonPermutationVector, RemappedCommonPermutationVector);
+ 		// msaa permutation.
+		EnableIfSet<TonemapperPermutation::FTonemapperMsaaDim>(WantedCommonPermutationVector, RemappedCommonPermutationVector);
 
  		// Color matrix
 		EnableIfSet<FTonemapperColorMatrixDim>(WantedPermutationVector, RemappedPermutationVector);
- 		// msaa permutation.
-		EnableIfSet<FTonemapperMsaaDim>(WantedPermutationVector, RemappedPermutationVector);
 
 		// DoF
 		if (WantedPermutationVector.Get<FTonemapperDOFDim>())
@@ -1290,8 +1312,8 @@ class FPostProcessTonemapPS_ES2 : public FGlobalShader
 		{
 			// 32 bpp hdr does not support:
 			RemappedPermutationVector.Set<FTonemapperDOFDim>(false);
-			RemappedPermutationVector.Set<FTonemapperMsaaDim>(false);
 			RemappedPermutationVector.Set<FTonemapperLightShaftsDim>(false);
+			RemappedCommonPermutationVector.Set<TonemapperPermutation::FTonemapperMsaaDim>(false);
 		}
 
 		RemappedPermutationVector.Set<TonemapperPermutation::FCommonDomain>(RemappedCommonPermutationVector);
@@ -1338,10 +1360,11 @@ class FPostProcessTonemapPS_ES2 : public FGlobalShader
 		}
 
 		static const auto CVarMobileMSAA = IConsoleManager::Get().FindTConsoleVariableDataInt(TEXT("r.MobileMSAA"));
-		const EShaderPlatform ShaderPlatform = GShaderPlatformForFeatureLevel[View.GetFeatureLevel()];
-		if ((GSupportsShaderFramebufferFetch && (IsMetalMobilePlatform(ShaderPlatform) || IsVulkanMobilePlatform(ShaderPlatform))) && (CVarMobileMSAA ? CVarMobileMSAA->GetValueOnAnyThread() > 1 : false))
+		const EShaderPlatform ShaderPlatform = View.GetShaderPlatform();
+		if (IsMetalMobilePlatform(ShaderPlatform))
 		{
-			MobilePermutationVector.Set<FTonemapperMsaaDim>(true);
+			bool bMSAA = (CVarMobileMSAA ? CVarMobileMSAA->GetValueOnAnyThread() > 1 : false);
+			CommonPermutationVector.Set<TonemapperPermutation::FTonemapperMsaaDim>(bMSAA);
 		}
 
 		if (bUse32BPPHDR)
