HLR Fixes:

 - Fixed unreported issue where HLR would cached mesh commands crashes if scene render proxy emits a mesh batch with more than one mesh batch element.
 - Fixed issue where meshes with a large number of elements could unnecessarily reserve a large amount of memory.
 - Fixed issue where IndirectArgsBuffer static meshes would be skipped.

diff --git a/Engine/Source/Runtime/Engine/Private/PrimitiveSceneProxy.cpp b/Engine/Source/Runtime/Engine/Private/PrimitiveSceneProxy.cpp
--- a/Engine/Source/Runtime/Engine/Private/PrimitiveSceneProxy.cpp
+++ b/Engine/Source/Runtime/Engine/Private/PrimitiveSceneProxy.cpp
@@ -36,20 +36,26 @@ bool CacheShadowDepthsFromPrimitivesUsingWPO()
 	return CVarCacheWPOPrimitives.GetValueOnAnyThread(true) != 0;
 }
 
-bool SupportsCachingMeshDrawCommands(const FVertexFactory* RESTRICT VertexFactory, const FPrimitiveSceneProxy* RESTRICT PrimitiveSceneProxy)
+bool SupportsCachingMeshDrawCommands(const FPrimitiveSceneProxy* RESTRICT PrimitiveSceneProxy, const FMeshBatch& MeshBatch)
 {
-	// Volumetric self shadow mesh commands need to be generated every frame, as they depend on single frame uniform buffers with self shadow data.
-	return VertexFactory->GetType()->SupportsCachingMeshDrawCommands() 
-		&& !PrimitiveSceneProxy->CastsVolumetricTranslucentShadow();
+	return
+		// Cached mesh commands only allow for a single mesh element per batch.
+		(MeshBatch.Elements.Num() == 1) &&
+
+		// Vertex factory needs to support caching.
+		MeshBatch.VertexFactory->GetType()->SupportsCachingMeshDrawCommands() &&
+
+		// Volumetric self shadow mesh commands need to be generated every frame, as they depend on single frame uniform buffers with self shadow data.
+		!PrimitiveSceneProxy->CastsVolumetricTranslucentShadow();
 }
 
-bool SupportsCachingMeshDrawCommands(const FVertexFactory* RESTRICT VertexFactory, const FPrimitiveSceneProxy* RESTRICT PrimitiveSceneProxy, const FMaterialRenderProxy* MaterialRenderProxy, ERHIFeatureLevel::Type FeatureLevel)
+bool SupportsCachingMeshDrawCommands(const FPrimitiveSceneProxy* RESTRICT PrimitiveSceneProxy, const FMeshBatch& MeshBatch, ERHIFeatureLevel::Type FeatureLevel)
 {
-	if (SupportsCachingMeshDrawCommands(VertexFactory, PrimitiveSceneProxy))
+	if (SupportsCachingMeshDrawCommands(PrimitiveSceneProxy, MeshBatch))
 	{
 		// External textures get mapped to immutable samplers (which are part of the PSO); the mesh must go through the dynamic path, as the media player might not have
 		// valid textures/samplers the first few calls; once they're available the PSO needs to get invalidated and recreated with the immutable samplers.
-		const FMaterial* Material = MaterialRenderProxy->GetMaterial(FeatureLevel);
+		const FMaterial* Material = MeshBatch.MaterialRenderProxy->GetMaterial(FeatureLevel);
 		const FMaterialShaderMap* ShaderMap = Material->GetRenderingThreadShaderMap();
 		if (ShaderMap)
 		{
diff --git a/Engine/Source/Runtime/Engine/Public/MeshBatch.h b/Engine/Source/Runtime/Engine/Public/MeshBatch.h
--- a/Engine/Source/Runtime/Engine/Public/MeshBatch.h
+++ b/Engine/Source/Runtime/Engine/Public/MeshBatch.h
@@ -99,6 +99,23 @@ struct FMeshBatchElement
 	FRHIVertexBuffer* IndirectArgsBuffer;
 	uint32 IndirectArgsOffset;
 
+	FORCEINLINE int32 GetNumPrimitives() const
+	{
+		if (bIsInstanceRuns && InstanceRuns)
+		{
+			int32 Count = 0;
+			for (uint32 Run = 0; Run < NumInstances; Run++)
+			{
+				Count += NumPrimitives * (InstanceRuns[Run * 2 + 1] - InstanceRuns[Run * 2] + 1);
+			}
+			return Count;
+		}
+		else
+		{
+			return NumPrimitives * NumInstances;
+		}
+	}
+
 	FMeshBatchElement()
 	:	PrimitiveUniformBuffer(nullptr)
 	,	PrimitiveUniformBufferResource(nullptr)
@@ -251,22 +268,24 @@ struct FMeshBatch
 
 	FORCEINLINE int32 GetNumPrimitives() const
 	{
-		int32 Count=0;
-		for( int32 ElementIdx=0;ElementIdx<Elements.Num();ElementIdx++ )
+		int32 Count = 0;
+		for (int32 ElementIdx = 0; ElementIdx < Elements.Num(); ++ElementIdx)
 		{
-			if (Elements[ElementIdx].bIsInstanceRuns && Elements[ElementIdx].InstanceRuns)
-			{
-				for (uint32 Run = 0; Run < Elements[ElementIdx].NumInstances; Run++)
-				{
-					Count += Elements[ElementIdx].NumPrimitives * (Elements[ElementIdx].InstanceRuns[Run * 2 + 1] - Elements[ElementIdx].InstanceRuns[Run * 2] + 1);
-				}
-			}
-			else
+			Count += Elements[ElementIdx].GetNumPrimitives();
+		}
+		return Count;
+	}
+
+	FORCEINLINE bool HasAnyDrawCalls() const
+	{
+		for (int32 ElementIdx = 0; ElementIdx < Elements.Num(); ++ElementIdx)
+		{
+			if (Elements[ElementIdx].GetNumPrimitives() > 0 || Elements[ElementIdx].IndirectArgsBuffer)
 			{
-				Count += Elements[ElementIdx].NumPrimitives * Elements[ElementIdx].NumInstances;
+				return true;
 			}
 		}
-		return Count;
+		return false;
 	}
 
 	ENGINE_API void PreparePrimitiveUniformBuffer(const FPrimitiveSceneProxy* PrimitiveSceneProxy, ERHIFeatureLevel::Type FeatureLevel);
diff --git a/Engine/Source/Runtime/Engine/Public/PrimitiveSceneProxy.h b/Engine/Source/Runtime/Engine/Public/PrimitiveSceneProxy.h
--- a/Engine/Source/Runtime/Engine/Public/PrimitiveSceneProxy.h
+++ b/Engine/Source/Runtime/Engine/Public/PrimitiveSceneProxy.h
@@ -1096,10 +1096,10 @@ protected:
 /**
  * Returns if specified mesh command can be cached, or needs to be recreated every frame.
  */
-ENGINE_API extern bool SupportsCachingMeshDrawCommands(const FVertexFactory* RESTRICT VertexFactory, const FPrimitiveSceneProxy* RESTRICT PrimitiveSceneProxy);
+ENGINE_API extern bool SupportsCachingMeshDrawCommands(const FPrimitiveSceneProxy* RESTRICT PrimitiveSceneProxy, const FMeshBatch& MeshBatch);
 
 /**
  * Returns if specified mesh command can be cached, or needs to be recreated every frame; this is a slightly slower version
  * used for materials with external textures that need invalidating their PSOs.
  */
-ENGINE_API extern bool SupportsCachingMeshDrawCommands(const FVertexFactory* RESTRICT VertexFactory, const FPrimitiveSceneProxy* RESTRICT PrimitiveSceneProxy, const class FMaterialRenderProxy* MaterialRenderProxy, ERHIFeatureLevel::Type FeatureLevel);
+ENGINE_API extern bool SupportsCachingMeshDrawCommands(const FPrimitiveSceneProxy* RESTRICT PrimitiveSceneProxy, const FMeshBatch& MeshBatch, ERHIFeatureLevel::Type FeatureLevel);
\ No newline at end of file
diff --git a/Engine/Source/Runtime/Renderer/Private/MeshPassProcessor.cpp b/Engine/Source/Runtime/Renderer/Private/MeshPassProcessor.cpp
--- a/Engine/Source/Runtime/Renderer/Private/MeshPassProcessor.cpp
+++ b/Engine/Source/Runtime/Renderer/Private/MeshPassProcessor.cpp
@@ -482,35 +482,7 @@ public:
 
 FMeshDrawShaderBindings::~FMeshDrawShaderBindings()
 {
-#if VALIDATE_UNIFORM_BUFFER_LIFETIME
-	uint8* ShaderBindingDataPtr = GetData();
-
-	for (int32 ShaderBindingsIndex = 0; ShaderBindingsIndex < ShaderLayouts.Num(); ShaderBindingsIndex++)
-	{
-		FMeshDrawSingleShaderBindings SingleShaderBindings(ShaderLayouts[ShaderBindingsIndex], ShaderBindingDataPtr);
-
-		FRHIUniformBuffer** RESTRICT UniformBufferBindings = SingleShaderBindings.GetUniformBufferStart();
-		const int32 NumUniformBuffers = SingleShaderBindings.ParameterMapInfo.UniformBuffers.Num();
-
-		for (int32 UniformBufferIndex = 0; UniformBufferIndex < NumUniformBuffers; UniformBufferIndex++)
-		{
-			FRHIUniformBuffer* UniformBuffer = UniformBufferBindings[UniformBufferIndex];
-
-			if (UniformBuffer)
-			{
-				UniformBuffer->NumMeshCommandReferencesForDebugging--;
-				check(UniformBuffer->NumMeshCommandReferencesForDebugging >= 0);
-			}
-		}
-
-		ShaderBindingDataPtr += ShaderLayouts[ShaderBindingsIndex].GetDataSizeBytes();
-	}
-#endif
-
-	if (Size > ARRAY_COUNT(InlineStorage))
-	{
-		delete[] HeapData;
-	}
+	Release();
 }
 
 void FMeshDrawShaderBindings::Initialize(FMeshProcessorShaders Shaders)
@@ -677,6 +649,8 @@ void FMeshDrawShaderBindings::Finalize(const FMeshProcessorShaders* ShadersForDe
 
 void FMeshDrawShaderBindings::CopyFrom(const FMeshDrawShaderBindings& Other)
 {
+	Release();
+
 	ShaderLayouts = Other.ShaderLayouts;
 
 	Allocate(Other.Size);
@@ -706,6 +680,40 @@ void FMeshDrawShaderBindings::CopyFrom(const FMeshDrawShaderBindings& Other)
 #endif
 }
 
+void FMeshDrawShaderBindings::Release()
+{
+#if VALIDATE_UNIFORM_BUFFER_LIFETIME
+	uint8* ShaderBindingDataPtr = GetData();
+
+	for (int32 ShaderBindingsIndex = 0; ShaderBindingsIndex < ShaderLayouts.Num(); ShaderBindingsIndex++)
+	{
+		FMeshDrawSingleShaderBindings SingleShaderBindings(ShaderLayouts[ShaderBindingsIndex], ShaderBindingDataPtr);
+		FRHIUniformBuffer** RESTRICT UniformBufferBindings = SingleShaderBindings.GetUniformBufferStart();
+		const int32 NumUniformBuffers = SingleShaderBindings.ParameterMapInfo.UniformBuffers.Num();
+
+		for (int32 UniformBufferIndex = 0; UniformBufferIndex < NumUniformBuffers; UniformBufferIndex++)
+		{
+			FRHIUniformBuffer* UniformBuffer = UniformBufferBindings[UniformBufferIndex];
+
+			if (UniformBuffer)
+			{
+				UniformBuffer->NumMeshCommandReferencesForDebugging--;
+				check(UniformBuffer->NumMeshCommandReferencesForDebugging >= 0);
+			}
+		}
+
+		ShaderBindingDataPtr += ShaderLayouts[ShaderBindingsIndex].GetDataSizeBytes();
+	}
+#endif
+
+	if (Size > ARRAY_COUNT(InlineStorage))
+	{
+		delete[] HeapData;
+	}
+	Size = 0;
+	HeapData = nullptr;
+}
+
 void FMeshDrawCommand::SetShaders(FRHIVertexDeclaration* VertexDeclaration, const FMeshProcessorShaders& Shaders, FGraphicsMinimalPipelineStateInitializer& PipelineState)
 {
 	PipelineState.BoundShaderState = FBoundShaderStateInput(
diff --git a/Engine/Source/Runtime/Renderer/Private/PrimitiveSceneInfo.cpp b/Engine/Source/Runtime/Renderer/Private/PrimitiveSceneInfo.cpp
--- a/Engine/Source/Runtime/Renderer/Private/PrimitiveSceneInfo.cpp
+++ b/Engine/Source/Runtime/Renderer/Private/PrimitiveSceneInfo.cpp
@@ -45,23 +45,24 @@ public:
 
 	virtual void ReserveMemoryForMeshes(int32 MeshNum)
 	{
-		PrimitiveSceneInfo->StaticMeshRelevances.Reserve(PrimitiveSceneInfo->StaticMeshRelevances.Max() + MeshNum);
-		PrimitiveSceneInfo->StaticMeshes.Reserve(PrimitiveSceneInfo->StaticMeshes.Max() + MeshNum);
+		if (PrimitiveSceneInfo->StaticMeshRelevances.GetSlack() < MeshNum)
+		{
+			PrimitiveSceneInfo->StaticMeshRelevances.Reserve(PrimitiveSceneInfo->StaticMeshRelevances.Max() + MeshNum);
+		}
+		if (PrimitiveSceneInfo->StaticMeshes.GetSlack() < MeshNum)
+		{
+			PrimitiveSceneInfo->StaticMeshes.Reserve(PrimitiveSceneInfo->StaticMeshes.Max() + MeshNum);
+		}
 	}
 
 	virtual void DrawMesh(const FMeshBatch& Mesh, float ScreenSize) final override
 	{
-		if (Mesh.GetNumPrimitives() > 0)
+		if (Mesh.HasAnyDrawCalls())
 		{
 			check(Mesh.VertexFactory);
 			check(Mesh.VertexFactory->IsInitialized());
 			checkSlow(IsInRenderingThread());
 
-			for (const FMeshBatchElement& BatchElement : Mesh.Elements)
-			{
-				check(!BatchElement.IndexBuffer || (BatchElement.IndexBuffer && BatchElement.IndexBuffer->IsInitialized() && BatchElement.IndexBuffer->IndexBufferRHI));
-			}
-
 			FPrimitiveSceneProxy* PrimitiveSceneProxy = PrimitiveSceneInfo->Proxy;
 			PrimitiveSceneProxy->VerifyUsedMaterial(Mesh.MaterialRenderProxy);
 
@@ -74,7 +75,7 @@ public:
 			const ERHIFeatureLevel::Type FeatureLevel = PrimitiveSceneInfo->Scene->GetFeatureLevel();
 			StaticMesh->PreparePrimitiveUniformBuffer(PrimitiveSceneProxy, FeatureLevel);
 
-			const bool bSupportsCachingMeshDrawCommands = SupportsCachingMeshDrawCommands(StaticMesh->VertexFactory, PrimitiveSceneProxy, Mesh.MaterialRenderProxy, FeatureLevel);
+			const bool bSupportsCachingMeshDrawCommands = SupportsCachingMeshDrawCommands(PrimitiveSceneProxy, *StaticMesh, FeatureLevel);
 
 			FStaticMeshBatchRelevance* StaticMeshRelevance = new(PrimitiveSceneInfo->StaticMeshRelevances) FStaticMeshBatchRelevance(
 				*StaticMesh, 
@@ -203,7 +204,7 @@ void FPrimitiveSceneInfo::CacheMeshDrawCommands(FRHICommandListImmediate& RHICmd
 	for (int32 MeshIndex = 0; MeshIndex < StaticMeshes.Num(); MeshIndex++)
 	{
 		const FStaticMeshBatch& Mesh = StaticMeshes[MeshIndex];
-		if (SupportsCachingMeshDrawCommands(Mesh.VertexFactory, Proxy))
+		if (SupportsCachingMeshDrawCommands(Proxy, Mesh))
 		{
 			++MeshWithCachedCommandsNum;
 		}
@@ -250,7 +251,7 @@ void FPrimitiveSceneInfo::CacheMeshDrawCommands(FRHICommandListImmediate& RHICmd
 			check(MeshRelevance.CommandInfosMask.IsEmpty());
 			MeshRelevance.CommandInfosBase = StaticMeshCommandInfos.Num();
 
-			if (SupportsCachingMeshDrawCommands(Mesh.VertexFactory, Proxy))
+			if (SupportsCachingMeshDrawCommands(Proxy, Mesh))
 			{
 				for (int32 PassIndex = 0; PassIndex < EMeshPass::Num; PassIndex++)
 				{
diff --git a/Engine/Source/Runtime/Renderer/Public/MeshPassProcessor.h b/Engine/Source/Runtime/Renderer/Public/MeshPassProcessor.h
--- a/Engine/Source/Runtime/Renderer/Public/MeshPassProcessor.h
+++ b/Engine/Source/Runtime/Renderer/Public/MeshPassProcessor.h
@@ -414,6 +414,8 @@ private:
 
 	RENDERER_API void CopyFrom(const FMeshDrawShaderBindings& Other);
 
+	RENDERER_API void Release();
+
 	template<class RHIShaderType>
 	static void SetShaderBindings(
 		FRHICommandList& RHICmdList,
