Defer caching of mesh draw commands until the main render tick. This ensures that component render states queued for recreation (i.e. from MarkRenderStateDirty) are processed prior to caching mesh draw commands.

diff --git a/Engine/Source/Runtime/Renderer/Private/RendererScene.cpp b/Engine/Source/Runtime/Renderer/Private/RendererScene.cpp
--- a/Engine/Source/Runtime/Renderer/Private/RendererScene.cpp
+++ b/Engine/Source/Runtime/Renderer/Private/RendererScene.cpp
@@ -3804,7 +3804,10 @@ void UpdateStaticMeshesForMaterials(const TArray<const FMaterial*>& MaterialReso
 					ENQUEUE_RENDER_COMMAND(FUpdateStaticMeshesForMaterials)(
 						[SceneProxy](FRHICommandListImmediate& RHICmdList)
 						{
-							SceneProxy->GetPrimitiveSceneInfo()->UpdateStaticMeshes(RHICmdList);
+							// Defer the caching until the next render tick, to make sure that all render components queued
+							// for re-creation are processed. Otherwise, we may end up caching mesh commands from stale data.
+							bool bReAddToDrawLists = false;
+							SceneProxy->GetPrimitiveSceneInfo()->UpdateStaticMeshes(RHICmdList, bReAddToDrawLists);
 						});
 				}
 			}
