Clear cached reflection capture proxies on primitives when a reflection capture proxy is released to avoid dangling pointer access.

diff --git a/Engine/Source/Runtime/Renderer/Private/PrimitiveSceneInfo.cpp b/Engine/Source/Runtime/Renderer/Private/PrimitiveSceneInfo.cpp
--- a/Engine/Source/Runtime/Renderer/Private/PrimitiveSceneInfo.cpp
+++ b/Engine/Source/Runtime/Renderer/Private/PrimitiveSceneInfo.cpp
@@ -1089,6 +1089,14 @@ void FPrimitiveSceneInfo::CacheReflectionCaptures()
 	bNeedsCachedReflectionCaptureUpdate = false;
 }
 
+void FPrimitiveSceneInfo::RemoveCachedReflectionCaptures()
+{
+	CachedReflectionCaptureProxy = nullptr;
+	CachedPlanarReflectionProxy = nullptr;
+	FMemory::Memzero(CachedReflectionCaptureProxies);
+	bNeedsCachedReflectionCaptureUpdate = true;
+}
+
 void FPrimitiveSceneInfo::UpdateComponentLastRenderTime(float CurrentWorldTime, bool bUpdateLastRenderTimeOnScreen) const
 {
 	ComponentForDebuggingOnly->LastRenderTime = CurrentWorldTime;
diff --git a/Engine/Source/Runtime/Renderer/Private/RendererScene.cpp b/Engine/Source/Runtime/Renderer/Private/RendererScene.cpp
--- a/Engine/Source/Runtime/Renderer/Private/RendererScene.cpp
+++ b/Engine/Source/Runtime/Renderer/Private/RendererScene.cpp
@@ -2124,6 +2124,12 @@ void FScene::RemoveReflectionCapture(UReflectionCaptureComponent* Component)
 
 			Scene->ReflectionSceneData.bRegisteredReflectionCapturesHasChanged = true;
 
+			// Need to clear out all reflection captures on removal to avoid dangling pointers.
+			for (int32 PrimitiveIndex = 0; PrimitiveIndex < Scene->Primitives.Num(); ++PrimitiveIndex)
+			{
+				Scene->Primitives[PrimitiveIndex]->RemoveCachedReflectionCaptures();
+			}
+
 			int32 CaptureIndex = Proxy->PackedIndex;
 			Scene->ReflectionSceneData.RegisteredReflectionCaptures.RemoveAtSwap(CaptureIndex);
 			Scene->ReflectionSceneData.RegisteredReflectionCapturePositions.RemoveAtSwap(CaptureIndex);
diff --git a/Engine/Source/Runtime/Renderer/Private/SceneVisibility.cpp b/Engine/Source/Runtime/Renderer/Private/SceneVisibility.cpp
--- a/Engine/Source/Runtime/Renderer/Private/SceneVisibility.cpp
+++ b/Engine/Source/Runtime/Renderer/Private/SceneVisibility.cpp
@@ -3504,11 +3504,8 @@ void UpdateReflectionSceneData(FScene* Scene)
 			for (int32 PrimitiveIndex = 0; PrimitiveIndex < Scene->Primitives.Num(); PrimitiveIndex++)
 			{
 				FPrimitiveSceneInfo* Primitive = Scene->Primitives[PrimitiveIndex];
-				Primitive->CachedReflectionCaptureProxy = nullptr;
-				Primitive->CachedPlanarReflectionProxy = nullptr;
-				FMemory::Memzero(Primitive->CachedReflectionCaptureProxies);
-				Primitive->bNeedsCachedReflectionCaptureUpdate = true;
-				
+				Primitive->RemoveCachedReflectionCaptures();
+
 				if (bNeedsStaticMeshUpdate)
 				{
 					Primitive->CacheReflectionCaptures();
diff --git a/Engine/Source/Runtime/Renderer/Public/PrimitiveSceneInfo.h b/Engine/Source/Runtime/Renderer/Public/PrimitiveSceneInfo.h
--- a/Engine/Source/Runtime/Renderer/Public/PrimitiveSceneInfo.h
+++ b/Engine/Source/Runtime/Renderer/Public/PrimitiveSceneInfo.h
@@ -18,6 +18,7 @@
 class FPrimitiveSceneInfo;
 class FPrimitiveSceneProxy;
 class FReflectionCaptureProxy;
+class FPlanarReflectionSceneProxy;
 class FScene;
 class FViewInfo;
 class UPrimitiveComponent;
@@ -267,7 +268,7 @@ public:
 	/** 
 	 * Planar reflection that was closest to this primitive, used for forward reflections.
 	 */
-	const class FPlanarReflectionSceneProxy* CachedPlanarReflectionProxy;
+	const FPlanarReflectionSceneProxy* CachedPlanarReflectionProxy;
 
 	/** 
 	 * Reflection capture proxy that was closest to this primitive, used for the forward shading rendering path. 
@@ -462,10 +463,15 @@ public:
 	int32 GetLightmapDataOffset() const { return LightmapDataOffset; }
 	int32 GetNumLightmapDataEntries() const { return NumLightmapDataEntries; }
 
+	/** Returns whether the primitive needs to call CacheReflectionCaptures. */
 	bool NeedsReflectionCaptureUpdate() const;
-	/** Cache per-primitive reflection captures used for mobile/forward rendering */
+
+	/** Cache per-primitive reflection captures used for mobile/forward rendering. */
 	void CacheReflectionCaptures();
 
+	/** Nulls out the cached per-primitive reflection captures. */
+	void RemoveCachedReflectionCaptures();
+
 	/** Helper function for writing out to the last render times to the game thread */
 	void UpdateComponentLastRenderTime(float CurrentWorldTime, bool bUpdateLastRenderTimeOnScreen) const;
 
