Fixed shader complexity showing no instructions when running in simple forward shading.

diff --git a/Engine/Source/Runtime/Renderer/Private/DebugViewModeRendering.cpp b/Engine/Source/Runtime/Renderer/Private/DebugViewModeRendering.cpp
--- a/Engine/Source/Runtime/Renderer/Private/DebugViewModeRendering.cpp
+++ b/Engine/Source/Runtime/Renderer/Private/DebugViewModeRendering.cpp
@@ -410,17 +410,26 @@ void FDebugViewModeMeshProcessor::UpdateInstructionCount(FDebugViewModeShaderEle
 	{
 		if (Scene->GetShadingPath() == EShadingPath::Deferred)
 		{
-			const bool bDeferred = !IsAnyForwardShadingEnabled(GetFeatureLevelShaderPlatform(InBatchMaterial->GetFeatureLevel()));
-			const bool bLit = InBatchMaterial->GetShadingModels().IsLit();
+			const EShaderPlatform ShaderPlatform = GetFeatureLevelShaderPlatform(InBatchMaterial->GetFeatureLevel());
 
-			OutShaderElementData.NumVSInstructions = InBatchMaterial->GetShader<TBasePassVS<TUniformLightMapPolicy<LMP_NO_LIGHTMAP>, false>>(InVertexFactoryType)->GetNumInstructions();
-			OutShaderElementData.NumPSInstructions = InBatchMaterial->GetShader<TBasePassPS<TUniformLightMapPolicy<LMP_NO_LIGHTMAP>, false>>(InVertexFactoryType)->GetNumInstructions();
-
-			if (!bDeferred && !IsTranslucentBlendMode(InBatchMaterial->GetBlendMode()))
+			if (IsSimpleForwardShadingEnabled(ShaderPlatform))
+			{
+				OutShaderElementData.NumVSInstructions = InBatchMaterial->GetShader<TBasePassVS<TUniformLightMapPolicy<LMP_SIMPLE_NO_LIGHTMAP>, false>>(InVertexFactoryType)->GetNumInstructions();
+				OutShaderElementData.NumPSInstructions = InBatchMaterial->GetShader<TBasePassPS<TUniformLightMapPolicy<LMP_SIMPLE_NO_LIGHTMAP>, false>>(InVertexFactoryType)->GetNumInstructions();
+			}
+			else
 			{
-				// Those numbers are taken from a simple material where common inputs are bound to vector parameters (to prevent constant optimizations).
-				OutShaderElementData.NumVSInstructions -= GShaderComplexityBaselineForwardVS - GShaderComplexityBaselineDeferredVS;
-				OutShaderElementData.NumPSInstructions -= bLit ? (GShaderComplexityBaselineForwardPS - GShaderComplexityBaselineDeferredPS) : (GShaderComplexityBaselineForwardUnlitPS - GShaderComplexityBaselineDeferredUnlitPS);
+				OutShaderElementData.NumVSInstructions = InBatchMaterial->GetShader<TBasePassVS<TUniformLightMapPolicy<LMP_NO_LIGHTMAP>, false>>(InVertexFactoryType)->GetNumInstructions();
+				OutShaderElementData.NumPSInstructions = InBatchMaterial->GetShader<TBasePassPS<TUniformLightMapPolicy<LMP_NO_LIGHTMAP>, false>>(InVertexFactoryType)->GetNumInstructions();
+
+				if (IsForwardShadingEnabled(ShaderPlatform) && !IsTranslucentBlendMode(InBatchMaterial->GetBlendMode()))
+				{
+					const bool bLit = InBatchMaterial->GetShadingModels().IsLit();
+
+					// Those numbers are taken from a simple material where common inputs are bound to vector parameters (to prevent constant optimizations).
+					OutShaderElementData.NumVSInstructions -= GShaderComplexityBaselineForwardVS - GShaderComplexityBaselineDeferredVS;
+					OutShaderElementData.NumPSInstructions -= bLit ? (GShaderComplexityBaselineForwardPS - GShaderComplexityBaselineDeferredPS) : (GShaderComplexityBaselineForwardUnlitPS - GShaderComplexityBaselineDeferredUnlitPS);
+				}
 			}
 
 			OutShaderElementData.NumVSInstructions = FMath::Max<int32>(0, OutShaderElementData.NumVSInstructions);
