Fix for rare crash where directional light for volumetric fog contains NaNs in the direction vector.

 - Added validation for attempting to invert a non-invertible matrix in InverseFast().
 - Band-aid check in volumetric fog code to skip light if it contains NaNs.

diff --git a/Engine/Source/Runtime/Core/Public/Math/Matrix.inl b/Engine/Source/Runtime/Core/Public/Math/Matrix.inl
--- a/Engine/Source/Runtime/Core/Public/Math/Matrix.inl
+++ b/Engine/Source/Runtime/Core/Public/Math/Matrix.inl
@@ -267,6 +267,15 @@ inline FMatrix FMatrix::InverseFast() const
 	{
 		ErrorEnsure(TEXT("FMatrix::InverseFast(), trying to invert a NIL matrix, this results in NaNs! Use Inverse() instead."));
 	}
+	else
+	{
+		const float	Det = Determinant();
+
+		if (Det == 0.0f || !FMath::IsFinite(Det))
+		{
+			ErrorEnsure(TEXT("FMatrix::InverseFast(), trying to invert a non-invertible matrix, this results in NaNs! Use Inverse() instead."));
+		}
+	}
 #endif
 	FMatrix Result;
 	VectorMatrixInverse( &Result, this );
diff --git a/Engine/Source/Runtime/Renderer/Private/VolumetricFogLightFunction.cpp b/Engine/Source/Runtime/Renderer/Private/VolumetricFogLightFunction.cpp
--- a/Engine/Source/Runtime/Renderer/Private/VolumetricFogLightFunction.cpp
+++ b/Engine/Source/Runtime/Renderer/Private/VolumetricFogLightFunction.cpp
@@ -118,6 +118,8 @@ void FDeferredShadingSceneRenderer::RenderLightFunctionForVolumetricFog(
 
 		if (ViewFamily.EngineShowFlags.LightFunctions
 			&& LightSceneInfo->Proxy->GetLightType() == LightType_Directional
+			// Band-aid fix for extremely rare case that light scene proxy contains NaNs.
+			&& !LightSceneInfo->Proxy->GetDirection().ContainsNaN()
 			&& LightSceneInfo->ShouldRenderLightViewIndependent()
 			&& LightSceneInfo->ShouldRenderLight(View))
 		{
