Fixed: High quality reflection intensity bug in mobile render (contributed by [REDACTED_PERSON])
#5978

diff --git a/Engine/Source/Runtime/Renderer/Private/MobileBasePass.cpp b/Engine/Source/Runtime/Renderer/Private/MobileBasePass.cpp
--- a/Engine/Source/Runtime/Renderer/Private/MobileBasePass.cpp
+++ b/Engine/Source/Runtime/Renderer/Private/MobileBasePass.cpp
@@ -473,7 +473,7 @@ void TMobileBasePassPSPolicyParamType<FUniformLightMapPolicy>::GetShaderBindings
 						{
 							ReflectionCubemapTextures[i] = PrimitiveSceneInfo->CachedReflectionCaptureProxies[i]->EncodedHDRCubemap;
 						}
-						ReflectionParams.X = FMath::Max(FMath::Min(1.0f / ReflectionProxy->EncodedHDRAverageBrightness, 65504.f), -65504.f);
+						ReflectionParams[i] = FMath::Max(FMath::Min(1.0f / ReflectionProxy->EncodedHDRAverageBrightness, 65504.f), -65504.f);
 					}
 				}
 			}
