Dynamic mesh commands now take the view's cull mode overrides in to account.

diff --git a/Engine/Source/Runtime/Renderer/Private/MeshPassProcessor.cpp b/Engine/Source/Runtime/Renderer/Private/MeshPassProcessor.cpp
--- a/Engine/Source/Runtime/Renderer/Private/MeshPassProcessor.cpp
+++ b/Engine/Source/Runtime/Renderer/Private/MeshPassProcessor.cpp
@@ -1137,10 +1137,14 @@ ERasterizerCullMode FMeshPassProcessor::ComputeMeshCullMode(const FMeshBatch& Me
 	const FMeshDrawingPolicyOverrideSettings InOverrideSettings = ComputeMeshOverrideSettings(Mesh);
 	const bool bMaterialResourceIsTwoSided = InMaterialResource.IsTwoSided();
 	const bool bInTwoSidedOverride = !!(InOverrideSettings.MeshOverrideFlags & EDrawingPolicyOverrideFlags::TwoSided);
+	const bool bViewTwoSidedOverride = ViewIfDynamicMeshCommand ? ViewIfDynamicMeshCommand->bRenderSceneTwoSided : false;
+	const bool bViewReverseCullingOverride = ViewIfDynamicMeshCommand ? ViewIfDynamicMeshCommand->bReverseCulling : false;
 	const bool bInReverseCullModeOverride = !!(InOverrideSettings.MeshOverrideFlags & EDrawingPolicyOverrideFlags::ReverseCullMode);
+	const bool bReverseCullMode = bViewReverseCullingOverride ? !bInReverseCullModeOverride : bInReverseCullModeOverride;
 	const bool bIsTwoSided = (bMaterialResourceIsTwoSided || bInTwoSidedOverride);
-	const bool bMeshRenderTwoSided = bIsTwoSided || bInTwoSidedOverride;
-	return bMeshRenderTwoSided ? CM_None : (bInReverseCullModeOverride ? CM_CCW : CM_CW);
+	const bool bMeshRenderTwoSided = bIsTwoSided || bInTwoSidedOverride || bViewTwoSidedOverride;
+
+	return bMeshRenderTwoSided ? CM_None : (bReverseCullMode ? CM_CCW : CM_CW);
 }
 
 void FMeshPassProcessor::GetDrawCommandPrimitiveId(
