Fixed vulkan validation error with SceneDepthZ.

diff --git a/Engine/Source/Runtime/Renderer/Private/ReflectionEnvironmentCapture.cpp b/Engine/Source/Runtime/Renderer/Private/ReflectionEnvironmentCapture.cpp
--- a/Engine/Source/Runtime/Renderer/Private/ReflectionEnvironmentCapture.cpp
+++ b/Engine/Source/Runtime/Renderer/Private/ReflectionEnvironmentCapture.cpp
@@ -771,6 +771,7 @@ void CaptureSceneToScratchCubemap(FRHICommandListImmediate& RHICmdList, FSceneRe
 		const int32 EffectiveSize = CubemapSize;
 		FSceneRenderTargetItem& EffectiveColorRT =  SceneContext.ReflectionColorScratchCubemap[0]->GetRenderTargetItem();
 		RHICmdList.TransitionResource(EResourceTransitionAccess::EWritable, EffectiveColorRT.TargetableTexture);
+		RHICmdList.TransitionResource(EResourceTransitionAccess::EReadable, SceneContext.GetSceneDepthSurface());
 
 		{
 			SCOPED_DRAW_EVENT(RHICmdList, CubeMapCopyScene);
