Fixed scene capture crash.

Rendering tiles using DrawTileMesh in a separate process without an accompanied main render can cause a crash in GetSceneTexturesUniformBuffer. This happens because translucency volume textures are not allocated until the first render in which they are used. However, the scene textures uniform buffer call attempts to dereference an array that is also not allocated until the textures are. This is incorrect. Instead, the call should return null which is then handled correctly by the uniform buffer call. Tile mesh rendering always uses the unlit mode, so it is never required to actually have these textures allocated.

diff --git a/Engine/Source/Runtime/Renderer/Private/PostProcess/SceneRenderTargets.h b/Engine/Source/Runtime/Renderer/Private/PostProcess/SceneRenderTargets.h
--- a/Engine/Source/Runtime/Renderer/Private/PostProcess/SceneRenderTargets.h
+++ b/Engine/Source/Runtime/Renderer/Private/PostProcess/SceneRenderTargets.h
@@ -407,10 +407,10 @@ public:
 	void ClearTranslucentVolumeLighting(FRHICommandListImmediate& RHICmdList, int32 ViewIndex);
 
 	/** Get the current translucent ambient lighting volume texture. Can vary depending on whether volume filtering is enabled */
-	IPooledRenderTarget* GetTranslucencyVolumeAmbient(ETranslucencyVolumeCascade Cascade, int32 ViewIndex = 0) { return TranslucencyLightingVolumeAmbient[SelectTranslucencyVolumeTarget(Cascade) + ViewIndex * NumTranslucentVolumeRenderTargetSets].GetReference(); }
+	IPooledRenderTarget* GetTranslucencyVolumeAmbient(ETranslucencyVolumeCascade Cascade, int32 ViewIndex = 0) { return TranslucencyLightingVolumeAmbient.Num() ? TranslucencyLightingVolumeAmbient[SelectTranslucencyVolumeTarget(Cascade) + ViewIndex * NumTranslucentVolumeRenderTargetSets].GetReference() : nullptr;  }
 
 	/** Get the current translucent directional lighting volume texture. Can vary depending on whether volume filtering is enabled */
-	IPooledRenderTarget* GetTranslucencyVolumeDirectional(ETranslucencyVolumeCascade Cascade, int32 ViewIndex = 0) { return TranslucencyLightingVolumeDirectional[SelectTranslucencyVolumeTarget(Cascade) + ViewIndex * NumTranslucentVolumeRenderTargetSets].GetReference(); }
+	IPooledRenderTarget* GetTranslucencyVolumeDirectional(ETranslucencyVolumeCascade Cascade, int32 ViewIndex = 0) { return TranslucencyLightingVolumeDirectional.Num() ? TranslucencyLightingVolumeDirectional[SelectTranslucencyVolumeTarget(Cascade) + ViewIndex * NumTranslucentVolumeRenderTargetSets].GetReference() : nullptr; }
 
 	/** Returns the size of most screen space render targets e.g. SceneColor, SceneDepth, GBuffer, ... might be different from final RT or output Size because of ScreenPercentage use. */
 	FIntPoint GetBufferSizeXY() const { return BufferSize; }
