Merging CL 4379292 from Dev-VR - Fixes for mobile static draw list using incorrect material render proxy

diff --git a/Engine/Source/Runtime/Engine/Private/Materials/MaterialUniformExpressions.cpp b/Engine/Source/Runtime/Engine/Private/Materials/MaterialUniformExpressions.cpp
--- a/Engine/Source/Runtime/Engine/Private/Materials/MaterialUniformExpressions.cpp
+++ b/Engine/Source/Runtime/Engine/Private/Materials/MaterialUniformExpressions.cpp
@@ -1,4 +1,4 @@
-﻿// Copyright 1998-2018 Epic Games, Inc. All Rights Reserved.
+// Copyright 1998-2018 Epic Games, Inc. All Rights Reserved.
 
 /*=============================================================================
 	MaterialShared.cpp: Shared material implementation.
diff --git a/Engine/Source/Runtime/Landscape/Private/LandscapeGrass.cpp b/Engine/Source/Runtime/Landscape/Private/LandscapeGrass.cpp
--- a/Engine/Source/Runtime/Landscape/Private/LandscapeGrass.cpp
+++ b/Engine/Source/Runtime/Landscape/Private/LandscapeGrass.cpp
@@ -462,7 +462,7 @@ public:
 
 				const int32 ShaderPass = (PassIdx >= FirstHeightMipsPassIndex) ? 0 : PassIdx;
 				DrawingPolicy.SetupPipelineState(DrawRenderState, *View);
-				CommitGraphicsPipelineState(RHICmdList, DrawingPolicy, DrawRenderState, DrawingPolicy.GetBoundShaderStateInput(View->GetFeatureLevel()));
+				CommitGraphicsPipelineState(RHICmdList, DrawingPolicy, DrawRenderState, DrawingPolicy.GetBoundShaderStateInput(View->GetFeatureLevel()), DrawingPolicy.GetMaterialRenderProxy());
 				DrawingPolicy.SetSharedState(RHICmdList, DrawRenderState, View, FLandscapeGrassWeightDrawingPolicy::ContextDataType(), ShaderPass, ComponentInfo.ViewOffset + FVector2D(PassOffsetX * PassIdx, 0));
 
 				// The first batch element contains the grass batch for the entire component
diff --git a/Engine/Source/Runtime/Renderer/Private/BasePassRendering.cpp b/Engine/Source/Runtime/Renderer/Private/BasePassRendering.cpp
--- a/Engine/Source/Runtime/Renderer/Private/BasePassRendering.cpp
+++ b/Engine/Source/Runtime/Renderer/Private/BasePassRendering.cpp
@@ -479,7 +479,7 @@ public:
 
 		SetDepthStencilStateForBasePass(DrawRenderState, View, Parameters.Mesh, Parameters.PrimitiveSceneProxy, bEnableReceiveDecalOutput, DrawingPolicy.UseDebugViewPS(), nullptr);
 		DrawingPolicy.SetupPipelineState(DrawRenderState, View);
-		CommitGraphicsPipelineState(RHICmdList, DrawingPolicy, DrawRenderState, DrawingPolicy.GetBoundShaderStateInput(View.GetFeatureLevel()));
+		CommitGraphicsPipelineState(RHICmdList, DrawingPolicy, DrawRenderState, DrawingPolicy.GetBoundShaderStateInput(View.GetFeatureLevel()), DrawingPolicy.GetMaterialRenderProxy());
 		DrawingPolicy.SetSharedState(RHICmdList, DrawRenderState, &View, typename TBasePassDrawingPolicy<LightMapPolicyType>::ContextDataType(Parameters.bIsInstancedStereo));
 
 		for( int32 BatchElementIndex = 0, Num = Parameters.Mesh.Elements.Num(); BatchElementIndex < Num; BatchElementIndex++ )
diff --git a/Engine/Source/Runtime/Renderer/Private/CompositionLighting/PostProcessMeshDecals.cpp b/Engine/Source/Runtime/Renderer/Private/CompositionLighting/PostProcessMeshDecals.cpp
--- a/Engine/Source/Runtime/Renderer/Private/CompositionLighting/PostProcessMeshDecals.cpp
+++ b/Engine/Source/Runtime/Renderer/Private/CompositionLighting/PostProcessMeshDecals.cpp
@@ -599,7 +599,7 @@ private:
 						View.Family->GetDebugViewShaderMode(),
 						DrawingContext.CurrentDecalStage);
 					DrawingPolicy.SetupPipelineState(DrawRenderStateLocal, View);
-					CommitGraphicsPipelineState(RHICmdList, DrawingPolicy, DrawRenderStateLocal, DrawingPolicy.GetBoundShaderStateInput(View.GetFeatureLevel()));
+					CommitGraphicsPipelineState(RHICmdList, DrawingPolicy, DrawRenderStateLocal, DrawingPolicy.GetBoundShaderStateInput(View.GetFeatureLevel()), DrawingPolicy.GetMaterialRenderProxy());
 					DrawingPolicy.SetSharedState(RHICmdList, DrawRenderStateLocal, &View, FDepthDrawingPolicy::ContextDataType(bIsInstancedStereo, bNeedsInstancedStereoBias));
 
 					int32 BatchElementIndex = 0;
diff --git a/Engine/Source/Runtime/Renderer/Private/ConvertToUniformMesh.cpp b/Engine/Source/Runtime/Renderer/Private/ConvertToUniformMesh.cpp
--- a/Engine/Source/Runtime/Renderer/Private/ConvertToUniformMesh.cpp
+++ b/Engine/Source/Runtime/Renderer/Private/ConvertToUniformMesh.cpp
@@ -376,7 +376,7 @@ int32 FUniformMeshConverter::Convert(
 				FDrawingPolicyRenderState DrawRenderState(View);
 
 				DrawingPolicy.SetupPipelineState(DrawRenderState, View);
-				CommitGraphicsPipelineState(RHICmdList, DrawingPolicy, DrawRenderState, DrawingPolicy.GetBoundShaderStateInput(View.GetFeatureLevel()));
+				CommitGraphicsPipelineState(RHICmdList, DrawingPolicy, DrawRenderState, DrawingPolicy.GetBoundShaderStateInput(View.GetFeatureLevel()), DrawingPolicy.GetMaterialRenderProxy());
 				DrawingPolicy.SetSharedState(RHICmdList, DrawRenderState, &View, FConvertToUniformMeshDrawingPolicy::ContextDataType());
 
 				for (int32 BatchElementIndex = 0; BatchElementIndex < Mesh.Elements.Num(); BatchElementIndex++)
diff --git a/Engine/Source/Runtime/Renderer/Private/DepthRendering.cpp b/Engine/Source/Runtime/Renderer/Private/DepthRendering.cpp
--- a/Engine/Source/Runtime/Renderer/Private/DepthRendering.cpp
+++ b/Engine/Source/Runtime/Renderer/Private/DepthRendering.cpp
@@ -675,7 +675,7 @@ bool FDepthDrawingPolicyFactory::DrawMesh(
 
 			FDrawingPolicyRenderState DrawRenderStateLocal(DrawRenderState);
 			DrawingPolicy.SetupPipelineState(DrawRenderStateLocal, View);
-			CommitGraphicsPipelineState(RHICmdList, DrawingPolicy, DrawRenderStateLocal, DrawingPolicy.GetBoundShaderStateInput(View.GetFeatureLevel()));
+			CommitGraphicsPipelineState(RHICmdList, DrawingPolicy, DrawRenderStateLocal, DrawingPolicy.GetBoundShaderStateInput(View.GetFeatureLevel()), DrawingPolicy.GetMaterialRenderProxy());
 			DrawingPolicy.SetSharedState(RHICmdList, DrawRenderStateLocal, &View, FPositionOnlyDepthDrawingPolicy::ContextDataType(bIsInstancedStereo, bIsInstancedStereoEmulated));
 
 			int32 BatchElementIndex = 0;
@@ -742,7 +742,7 @@ bool FDepthDrawingPolicyFactory::DrawMesh(
 
 				FDrawingPolicyRenderState DrawRenderStateLocal(DrawRenderState);
 				DrawingPolicy.SetupPipelineState(DrawRenderStateLocal, View);
-				CommitGraphicsPipelineState(RHICmdList, DrawingPolicy, DrawRenderStateLocal, DrawingPolicy.GetBoundShaderStateInput(View.GetFeatureLevel()));
+				CommitGraphicsPipelineState(RHICmdList, DrawingPolicy, DrawRenderStateLocal, DrawingPolicy.GetBoundShaderStateInput(View.GetFeatureLevel()), DrawingPolicy.GetMaterialRenderProxy());
 				DrawingPolicy.SetSharedState(RHICmdList, DrawRenderStateLocal, &View, FDepthDrawingPolicy::ContextDataType(bIsInstancedStereo, bIsInstancedStereoEmulated));
 
 				int32 BatchElementIndex = 0;
diff --git a/Engine/Source/Runtime/Renderer/Private/DistortionRendering.cpp b/Engine/Source/Runtime/Renderer/Private/DistortionRendering.cpp
--- a/Engine/Source/Runtime/Renderer/Private/DistortionRendering.cpp
+++ b/Engine/Source/Runtime/Renderer/Private/DistortionRendering.cpp
@@ -767,7 +767,7 @@ bool TDistortionMeshDrawingPolicyFactory<DistortMeshPolicy>::DrawDynamicMesh(
 		FDrawingPolicyRenderState DrawRenderStateLocal(DrawRenderState);
 		DrawRenderStateLocal.SetDitheredLODTransitionAlpha(Mesh.DitheredLODTransitionAlpha);
 		DrawingPolicy.SetupPipelineState(DrawRenderStateLocal, View);
-		CommitGraphicsPipelineState(RHICmdList, DrawingPolicy, DrawRenderStateLocal, DrawingPolicy.GetBoundShaderStateInput(View.GetFeatureLevel()));
+		CommitGraphicsPipelineState(RHICmdList, DrawingPolicy, DrawRenderStateLocal, DrawingPolicy.GetBoundShaderStateInput(View.GetFeatureLevel()), DrawingPolicy.GetMaterialRenderProxy());
 		DrawingPolicy.SetSharedState(RHICmdList, DrawRenderStateLocal, &View, typename TDistortionMeshDrawingPolicy<DistortMeshPolicy>::ContextDataType());
 
 		for (int32 BatchElementIndex = 0; BatchElementIndex < Mesh.Elements.Num(); BatchElementIndex++)
@@ -823,7 +823,7 @@ bool TDistortionMeshDrawingPolicyFactory<DistortMeshPolicy>::DrawStaticMesh(
 		FDrawingPolicyRenderState DrawRenderStateLocal(DrawRenderState);
 		DrawingPolicy.ApplyDitheredLODTransitionState(DrawRenderStateLocal, *View, StaticMesh, false);
 		DrawingPolicy.SetupPipelineState(DrawRenderStateLocal, *View);
-		CommitGraphicsPipelineState(RHICmdList, DrawingPolicy, DrawRenderStateLocal, DrawingPolicy.GetBoundShaderStateInput(View->GetFeatureLevel()));
+		CommitGraphicsPipelineState(RHICmdList, DrawingPolicy, DrawRenderStateLocal, DrawingPolicy.GetBoundShaderStateInput(View->GetFeatureLevel()), DrawingPolicy.GetMaterialRenderProxy());
 		DrawingPolicy.SetSharedState(RHICmdList, DrawRenderStateLocal, View, typename TDistortionMeshDrawingPolicy<DistortMeshPolicy>::ContextDataType());
 		int32 BatchElementIndex = 0;
 		do
diff --git a/Engine/Source/Runtime/Renderer/Private/LightMapDensityRendering.cpp b/Engine/Source/Runtime/Renderer/Private/LightMapDensityRendering.cpp
--- a/Engine/Source/Runtime/Renderer/Private/LightMapDensityRendering.cpp
+++ b/Engine/Source/Runtime/Renderer/Private/LightMapDensityRendering.cpp
@@ -135,7 +135,7 @@ bool FLightMapDensityDrawingPolicyFactory::DrawDynamicMesh(
 			{
 				TLightMapDensityDrawingPolicy< TUniformLightMapPolicy<LMP_HQ_LIGHTMAP> > DrawingPolicy(View, Mesh.VertexFactory, MaterialRenderProxy, TUniformLightMapPolicy<LMP_HQ_LIGHTMAP>(), BlendMode, ComputeMeshOverrideSettings(Mesh));
 				DrawingPolicy.SetupPipelineState(DrawRenderStateLocal, View);
-				CommitGraphicsPipelineState(RHICmdList, DrawingPolicy, DrawRenderStateLocal, DrawingPolicy.GetBoundShaderStateInput(View.GetFeatureLevel()));
+				CommitGraphicsPipelineState(RHICmdList, DrawingPolicy, DrawRenderStateLocal, DrawingPolicy.GetBoundShaderStateInput(View.GetFeatureLevel()), DrawingPolicy.GetMaterialRenderProxy());
 				DrawingPolicy.SetSharedState(RHICmdList, DrawRenderStateLocal, &View, TLightMapDensityDrawingPolicy< FUniformLightMapPolicy >::ContextDataType());
 				for (int32 BatchElementIndex = 0; BatchElementIndex < Mesh.Elements.Num(); BatchElementIndex++)
 				{
@@ -151,7 +151,7 @@ bool FLightMapDensityDrawingPolicyFactory::DrawDynamicMesh(
 			{
 				TLightMapDensityDrawingPolicy< TUniformLightMapPolicy<LMP_LQ_LIGHTMAP> > DrawingPolicy(View, Mesh.VertexFactory, MaterialRenderProxy, TUniformLightMapPolicy<LMP_LQ_LIGHTMAP>(), BlendMode, ComputeMeshOverrideSettings(Mesh));
 				DrawingPolicy.SetupPipelineState(DrawRenderStateLocal, View);
-				CommitGraphicsPipelineState(RHICmdList, DrawingPolicy, DrawRenderStateLocal, DrawingPolicy.GetBoundShaderStateInput(View.GetFeatureLevel()));
+				CommitGraphicsPipelineState(RHICmdList, DrawingPolicy, DrawRenderStateLocal, DrawingPolicy.GetBoundShaderStateInput(View.GetFeatureLevel()), DrawingPolicy.GetMaterialRenderProxy());
 				DrawingPolicy.SetSharedState(RHICmdList, DrawRenderStateLocal, &View, TLightMapDensityDrawingPolicy< FUniformLightMapPolicy >::ContextDataType());
 				for (int32 BatchElementIndex = 0; BatchElementIndex < Mesh.Elements.Num(); BatchElementIndex++)
 				{
@@ -168,7 +168,7 @@ bool FLightMapDensityDrawingPolicyFactory::DrawDynamicMesh(
 		{
 			TLightMapDensityDrawingPolicy<TUniformLightMapPolicy<LMP_DUMMY> > DrawingPolicy(View, Mesh.VertexFactory, MaterialRenderProxy, TUniformLightMapPolicy<LMP_DUMMY>(), BlendMode, ComputeMeshOverrideSettings(Mesh));
 			DrawingPolicy.SetupPipelineState(DrawRenderStateLocal, View);
-			CommitGraphicsPipelineState(RHICmdList, DrawingPolicy, DrawRenderStateLocal, DrawingPolicy.GetBoundShaderStateInput(View.GetFeatureLevel()));
+			CommitGraphicsPipelineState(RHICmdList, DrawingPolicy, DrawRenderStateLocal, DrawingPolicy.GetBoundShaderStateInput(View.GetFeatureLevel()), DrawingPolicy.GetMaterialRenderProxy());
 			DrawingPolicy.SetSharedState(RHICmdList, DrawRenderStateLocal, &View, TLightMapDensityDrawingPolicy<FUniformLightMapPolicy >::ContextDataType());
 			for (int32 BatchElementIndex = 0; BatchElementIndex < Mesh.Elements.Num(); BatchElementIndex++)
 			{
@@ -185,7 +185,7 @@ bool FLightMapDensityDrawingPolicyFactory::DrawDynamicMesh(
 	{
 		TLightMapDensityDrawingPolicy<TUniformLightMapPolicy<LMP_NO_LIGHTMAP> > DrawingPolicy(View, Mesh.VertexFactory, MaterialRenderProxy, TUniformLightMapPolicy<LMP_NO_LIGHTMAP>(), BlendMode, ComputeMeshOverrideSettings(Mesh));
 		DrawingPolicy.SetupPipelineState(DrawRenderStateLocal, View);
-		CommitGraphicsPipelineState(RHICmdList, DrawingPolicy, DrawRenderStateLocal, DrawingPolicy.GetBoundShaderStateInput(View.GetFeatureLevel()));
+		CommitGraphicsPipelineState(RHICmdList, DrawingPolicy, DrawRenderStateLocal, DrawingPolicy.GetBoundShaderStateInput(View.GetFeatureLevel()), DrawingPolicy.GetMaterialRenderProxy());
 		DrawingPolicy.SetSharedState(RHICmdList, DrawRenderStateLocal, &View, TLightMapDensityDrawingPolicy<TUniformLightMapPolicy<LMP_NO_LIGHTMAP> >::ContextDataType());
 		for (int32 BatchElementIndex = 0; BatchElementIndex < Mesh.Elements.Num(); BatchElementIndex++)
 		{
diff --git a/Engine/Source/Runtime/Renderer/Private/MobileBasePassRendering.cpp b/Engine/Source/Runtime/Renderer/Private/MobileBasePassRendering.cpp
--- a/Engine/Source/Runtime/Renderer/Private/MobileBasePassRendering.cpp
+++ b/Engine/Source/Runtime/Renderer/Private/MobileBasePassRendering.cpp
@@ -465,7 +465,7 @@ public:
 			);
 
 		DrawingPolicy.SetupPipelineState(DrawRenderState, View);
-		CommitGraphicsPipelineState(RHICmdList, DrawingPolicy, DrawRenderState, DrawingPolicy.GetBoundShaderStateInput(View.GetFeatureLevel()));
+		CommitGraphicsPipelineState(RHICmdList, DrawingPolicy, DrawRenderState, DrawingPolicy.GetBoundShaderStateInput(View.GetFeatureLevel()), Parameters.Mesh.MaterialRenderProxy);
 		DrawingPolicy.SetSharedState(RHICmdList, DrawRenderState, &View, typename TMobileBasePassDrawingPolicy<FUniformLightMapPolicy>::ContextDataType());
 
 		for( int32 BatchElementIndex=0;BatchElementIndex<Parameters.Mesh.Elements.Num();BatchElementIndex++ )
diff --git a/Engine/Source/Runtime/Renderer/Private/MobileBasePassRendering.h b/Engine/Source/Runtime/Renderer/Private/MobileBasePassRendering.h
--- a/Engine/Source/Runtime/Renderer/Private/MobileBasePassRendering.h
+++ b/Engine/Source/Runtime/Renderer/Private/MobileBasePassRendering.h
@@ -708,7 +708,7 @@ public:
 		ERHIFeatureLevel::Type FeatureLevel,
 		bool bInEnableReceiveDecalOutput = false
 		):
-		FMeshDrawingPolicy(nullptr, InMaterialRenderProxy, InMaterialResource, InOverrideSettings, InDebugViewShaderMode),
+		FMeshDrawingPolicy(nullptr, nullptr, InMaterialResource, InOverrideSettings, InDebugViewShaderMode),
 		VertexDeclaration(InVertexFactory->GetDeclaration()),
 		LightMapPolicy(InLightMapPolicy),
 		NumMovablePointLights(InNumMovablePointLights),
@@ -824,6 +824,11 @@ public:
 		return PointerHash(VertexDeclaration, PointerHash(MaterialResource));
 	}
 
+	const FMaterialRenderProxy* GetPipelineMaterialRenderProxy(const FMaterialRenderProxy* ElementMaterialRenderProxy)
+	{
+		return ElementMaterialRenderProxy;
+	}
+
 	friend int32 CompareDrawingPolicy(const TMobileBasePassDrawingPolicy& A,const TMobileBasePassDrawingPolicy& B)
 	{
 		COMPAREDRAWINGPOLICYMEMBERS(MaterialResource);
diff --git a/Engine/Source/Runtime/Renderer/Private/MobileTranslucentRendering.cpp b/Engine/Source/Runtime/Renderer/Private/MobileTranslucentRendering.cpp
--- a/Engine/Source/Runtime/Renderer/Private/MobileTranslucentRendering.cpp
+++ b/Engine/Source/Runtime/Renderer/Private/MobileTranslucentRendering.cpp
@@ -182,7 +182,7 @@ public:
 			);
 
 		DrawingPolicy.SetupPipelineState(DrawRenderState, View);
-		CommitGraphicsPipelineState(RHICmdList, DrawingPolicy, DrawRenderState, DrawingPolicy.GetBoundShaderStateInput(View.GetFeatureLevel()));
+		CommitGraphicsPipelineState(RHICmdList, DrawingPolicy, DrawRenderState, DrawingPolicy.GetBoundShaderStateInput(View.GetFeatureLevel()), Parameters.Mesh.MaterialRenderProxy);
 		DrawingPolicy.SetSharedState(RHICmdList, DrawRenderState, &View, typename TMobileBasePassDrawingPolicy<FUniformLightMapPolicy>::ContextDataType());
 
 		for (int32 BatchElementIndex = 0; BatchElementIndex<Parameters.Mesh.Elements.Num(); BatchElementIndex++)
@@ -648,7 +648,7 @@ private:
 
 			FDrawingPolicyRenderState DrawRenderStateLocal(DrawRenderState);
 			DrawingPolicy.SetupPipelineState(DrawRenderStateLocal, View);
-			CommitGraphicsPipelineState(RHICmdList, DrawingPolicy, DrawRenderStateLocal, DrawingPolicy.GetBoundShaderStateInput(View.GetFeatureLevel()));
+			CommitGraphicsPipelineState(RHICmdList, DrawingPolicy, DrawRenderStateLocal, DrawingPolicy.GetBoundShaderStateInput(View.GetFeatureLevel()), Mesh.MaterialRenderProxy);
 			DrawingPolicy.SetSharedState(RHICmdList, DrawRenderStateLocal, &View, FMobileOpacityDrawingPolicy::ContextDataType());
 
 			int32 BatchElementIndex = 0;
diff --git a/Engine/Source/Runtime/Renderer/Private/SceneHitProxyRendering.cpp b/Engine/Source/Runtime/Renderer/Private/SceneHitProxyRendering.cpp
--- a/Engine/Source/Runtime/Renderer/Private/SceneHitProxyRendering.cpp
+++ b/Engine/Source/Runtime/Renderer/Private/SceneHitProxyRendering.cpp
@@ -390,7 +390,7 @@ bool FHitProxyDrawingPolicyFactory::DrawDynamicMesh(
 				FDrawingPolicyRenderState DrawRenderStateLocal(DrawRenderState);
 				DrawRenderStateLocal.SetDitheredLODTransitionAlpha(Mesh.DitheredLODTransitionAlpha);
 				DrawingPolicy.SetupPipelineState(DrawRenderStateLocal, View);
-				CommitGraphicsPipelineState(RHICmdList, DrawingPolicy, DrawRenderStateLocal, DrawingPolicy.GetBoundShaderStateInput(View.GetFeatureLevel()));
+				CommitGraphicsPipelineState(RHICmdList, DrawingPolicy, DrawRenderStateLocal, DrawingPolicy.GetBoundShaderStateInput(View.GetFeatureLevel()), DrawingPolicy.GetMaterialRenderProxy());
 				DrawingPolicy.SetSharedState(RHICmdList, DrawRenderStateLocal, &View, FHitProxyDrawingPolicy::ContextDataType());
 
 			    for (int32 BatchElementIndex = 0; BatchElementIndex < Mesh.Elements.Num(); BatchElementIndex++)
diff --git a/Engine/Source/Runtime/Renderer/Private/ShadowDepthRendering.cpp b/Engine/Source/Runtime/Renderer/Private/ShadowDepthRendering.cpp
--- a/Engine/Source/Runtime/Renderer/Private/ShadowDepthRendering.cpp
+++ b/Engine/Source/Runtime/Renderer/Private/ShadowDepthRendering.cpp
@@ -1101,7 +1101,7 @@ bool FShadowDepthDrawingPolicyFactory::DrawDynamicMesh(
 				DrawRenderStateLocal.SetDitheredLODTransitionAlpha(Mesh.DitheredLODTransitionAlpha);
 				SetViewFlagsForShadowPass(DrawRenderStateLocal, View, View.GetFeatureLevel(), DrawingPolicy.IsTwoSided(), true, bLocalOnePassPointLightShadow);
 				DrawingPolicy.SetupPipelineState(DrawRenderStateLocal, View);
-				CommitGraphicsPipelineState(RHICmdList, DrawingPolicy, DrawRenderStateLocal, DrawingPolicy.GetBoundShaderStateInput(View.GetFeatureLevel()));
+				CommitGraphicsPipelineState(RHICmdList, DrawingPolicy, DrawRenderStateLocal, DrawingPolicy.GetBoundShaderStateInput(View.GetFeatureLevel()), DrawingPolicy.GetMaterialRenderProxy());
 				DrawingPolicy.SetSharedState(RHICmdList, DrawRenderStateLocal, &View, PolicyContext);
 
 				for (int32 BatchElementIndex = 0, Num = Mesh.Elements.Num(); BatchElementIndex < Num; BatchElementIndex++)
@@ -1132,7 +1132,7 @@ bool FShadowDepthDrawingPolicyFactory::DrawDynamicMesh(
 				DrawRenderStateLocal.SetDitheredLODTransitionAlpha(Mesh.DitheredLODTransitionAlpha);
 				SetViewFlagsForShadowPass(DrawRenderStateLocal, View, View.GetFeatureLevel(), DrawingPolicy.IsTwoSided(), false, bLocalOnePassPointLightShadow);
 				DrawingPolicy.SetupPipelineState(DrawRenderStateLocal, View);
-				CommitGraphicsPipelineState(RHICmdList, DrawingPolicy, DrawRenderStateLocal, DrawingPolicy.GetBoundShaderStateInput(View.GetFeatureLevel()));
+				CommitGraphicsPipelineState(RHICmdList, DrawingPolicy, DrawRenderStateLocal, DrawingPolicy.GetBoundShaderStateInput(View.GetFeatureLevel()), DrawingPolicy.GetMaterialRenderProxy());
 				DrawingPolicy.SetSharedState(RHICmdList, DrawRenderStateLocal, &View, PolicyContext);
 
 				for (int32 BatchElementIndex = 0; BatchElementIndex < Mesh.Elements.Num(); BatchElementIndex++)
@@ -1312,7 +1312,7 @@ void DrawShadowMeshElements(FRHICommandList& RHICmdList, const FViewInfo& View,
 			DrawRenderStateLocal.SetDitheredLODTransitionAlpha(ShadowMesh.Mesh->DitheredLODTransitionAlpha);
 			SetViewFlagsForShadowPass(DrawRenderStateLocal, View, View.GetFeatureLevel(), SharedDrawingPolicy.IsTwoSided(), bReflectiveShadowmap, ShadowInfo.bOnePassPointLightShadow);
 			SharedDrawingPolicy.SetupPipelineState(DrawRenderStateLocal, View);
-			CommitGraphicsPipelineState(RHICmdList, SharedDrawingPolicy, DrawRenderStateLocal, SharedDrawingPolicy.GetBoundShaderStateInput(View.GetFeatureLevel()));
+			CommitGraphicsPipelineState(RHICmdList, SharedDrawingPolicy, DrawRenderStateLocal, SharedDrawingPolicy.GetBoundShaderStateInput(View.GetFeatureLevel()), SharedDrawingPolicy.GetMaterialRenderProxy());
 			SharedDrawingPolicy.SetSharedState(RHICmdList, DrawRenderStateLocal, &View, PolicyContext);
 		}
 
diff --git a/Engine/Source/Runtime/Renderer/Private/StaticMeshDrawList.inl b/Engine/Source/Runtime/Renderer/Private/StaticMeshDrawList.inl
--- a/Engine/Source/Runtime/Renderer/Private/StaticMeshDrawList.inl
+++ b/Engine/Source/Runtime/Renderer/Private/StaticMeshDrawList.inl
@@ -180,7 +180,9 @@ int32 TStaticMeshDrawList<DrawingPolicyType>::DrawElement(
 			BoundShaderStateInput = DrawingPolicyLink->DrawingPolicy.GetBoundShaderStateInput(View.GetFeatureLevel());
 		}
 
-		CommitGraphicsPipelineState(RHICmdList, DrawingPolicyLink->DrawingPolicy, DrawRenderState, BoundShaderStateInput);
+		// This is needed as the mobile renderer ignores the MRP inside the DrawingPolicy, so we need the Mesh's one
+		const FMaterialRenderProxy* PipelineMaterialRenderProxy = DrawingPolicyLink->DrawingPolicy.GetPipelineMaterialRenderProxy(Element.Mesh->MaterialRenderProxy);
+		CommitGraphicsPipelineState(RHICmdList, DrawingPolicyLink->DrawingPolicy, DrawRenderState, BoundShaderStateInput, PipelineMaterialRenderProxy);
 		DrawingPolicyLink->DrawingPolicy.SetSharedState(RHICmdList, DrawRenderState, &View, PolicyContext);
 
 		bDrawnShared = true;
diff --git a/Engine/Source/Runtime/Renderer/Private/TranslucentLighting.cpp b/Engine/Source/Runtime/Renderer/Private/TranslucentLighting.cpp
--- a/Engine/Source/Runtime/Renderer/Private/TranslucentLighting.cpp
+++ b/Engine/Source/Runtime/Renderer/Private/TranslucentLighting.cpp
@@ -478,7 +478,7 @@ public:
 				FDrawingPolicyRenderState DrawRenderStateLocal(DrawRenderState);
 				DrawRenderStateLocal.SetDitheredLODTransitionAlpha(Mesh.DitheredLODTransitionAlpha);
 				DrawingPolicy.SetupPipelineState(DrawRenderStateLocal, View);
-				CommitGraphicsPipelineState(RHICmdList, DrawingPolicy, DrawRenderStateLocal, DrawingPolicy.GetBoundShaderStateInput(View.GetFeatureLevel()));
+				CommitGraphicsPipelineState(RHICmdList, DrawingPolicy, DrawRenderStateLocal, DrawingPolicy.GetBoundShaderStateInput(View.GetFeatureLevel()), DrawingPolicy.GetMaterialRenderProxy());
 				DrawingPolicy.SetSharedState(RHICmdList, DrawRenderStateLocal, &View, FTranslucencyShadowDepthDrawingPolicy::ContextDataType(DrawingContext.ShadowInfo));
 
 				for (int32 BatchElementIndex = 0; BatchElementIndex < Mesh.Elements.Num(); BatchElementIndex++)
diff --git a/Engine/Source/Runtime/Renderer/Private/TranslucentRendering.cpp b/Engine/Source/Runtime/Renderer/Private/TranslucentRendering.cpp
--- a/Engine/Source/Runtime/Renderer/Private/TranslucentRendering.cpp
+++ b/Engine/Source/Runtime/Renderer/Private/TranslucentRendering.cpp
@@ -417,7 +417,7 @@ public:
 			false);
 
 		DrawingPolicy.SetupPipelineState(DrawRenderState, View);
-		CommitGraphicsPipelineState(RHICmdList, DrawingPolicy, DrawRenderState, DrawingPolicy.GetBoundShaderStateInput(View.GetFeatureLevel()));
+		CommitGraphicsPipelineState(RHICmdList, DrawingPolicy, DrawRenderState, DrawingPolicy.GetBoundShaderStateInput(View.GetFeatureLevel()), DrawingPolicy.GetMaterialRenderProxy());
 		DrawingPolicy.SetSharedState(RHICmdList, DrawRenderState, &View, typename TBasePassDrawingPolicy<LightMapPolicyType>::ContextDataType(Parameters.bIsInstancedStereo));
 
 		int32 BatchElementIndex = 0;
diff --git a/Engine/Source/Runtime/Renderer/Private/VelocityRendering.cpp b/Engine/Source/Runtime/Renderer/Private/VelocityRendering.cpp
--- a/Engine/Source/Runtime/Renderer/Private/VelocityRendering.cpp
+++ b/Engine/Source/Runtime/Renderer/Private/VelocityRendering.cpp
@@ -566,7 +566,7 @@ bool FVelocityDrawingPolicyFactory::DrawDynamicMesh(
 			FDrawingPolicyRenderState DrawRenderStateLocal(DrawRenderState);
 			DrawRenderStateLocal.SetDitheredLODTransitionAlpha(Mesh.DitheredLODTransitionAlpha);
 			DrawingPolicy.SetupPipelineState(DrawRenderStateLocal, View);
-			CommitGraphicsPipelineState(RHICmdList, DrawingPolicy, DrawRenderStateLocal, DrawingPolicy.GetBoundShaderStateInput(View.GetFeatureLevel()));
+			CommitGraphicsPipelineState(RHICmdList, DrawingPolicy, DrawRenderStateLocal, DrawingPolicy.GetBoundShaderStateInput(View.GetFeatureLevel()), DrawingPolicy.GetMaterialRenderProxy());
 			DrawingPolicy.SetSharedState(RHICmdList, DrawRenderStateLocal, &View, FVelocityDrawingPolicy::ContextDataType(bIsInstancedStereo));
 			for (int32 BatchElementIndex = 0, BatchElementCount = Mesh.Elements.Num(); BatchElementIndex < BatchElementCount; ++BatchElementIndex)
 			{
diff --git a/Engine/Source/Runtime/Renderer/Private/VolumetricFogVoxelization.cpp b/Engine/Source/Runtime/Renderer/Private/VolumetricFogVoxelization.cpp
--- a/Engine/Source/Runtime/Renderer/Private/VolumetricFogVoxelization.cpp
+++ b/Engine/Source/Runtime/Renderer/Private/VolumetricFogVoxelization.cpp
@@ -648,7 +648,7 @@ void VoxelizeVolumePrimitive(
 
 		FDrawingPolicyRenderState DrawRenderState(PassDrawRenderState);
 		DrawingPolicy.SetupPipelineState(DrawRenderState, View);
-		CommitGraphicsPipelineState(RHICmdList, DrawingPolicy, DrawRenderState, DrawingPolicy.GetBoundShaderStateInput(View.GetFeatureLevel()));
+		CommitGraphicsPipelineState(RHICmdList, DrawingPolicy, DrawRenderState, DrawingPolicy.GetBoundShaderStateInput(View.GetFeatureLevel()), DrawingPolicy.GetMaterialRenderProxy());
 		DrawingPolicy.SetSharedState(RHICmdList, View, IntegrationData, Jitter, FVoxelizeVolumeDrawingPolicy::ContextDataType(), DrawRenderState);
 
 		FBoxSphereBounds Bounds = PrimitiveSceneProxy->GetBounds();
diff --git a/Engine/Source/Runtime/Renderer/Public/DrawingPolicy.h b/Engine/Source/Runtime/Renderer/Public/DrawingPolicy.h
--- a/Engine/Source/Runtime/Renderer/Public/DrawingPolicy.h
+++ b/Engine/Source/Runtime/Renderer/Public/DrawingPolicy.h
@@ -1,4 +1,4 @@
-﻿// Copyright 1998-2018 Epic Games, Inc. All Rights Reserved.
+// Copyright 1998-2018 Epic Games, Inc. All Rights Reserved.
 
 /*=============================================================================
 	DrawingPolicy.h: Drawing policy definitions.
@@ -270,7 +270,7 @@ FORCEINLINE_DEBUGGABLE FMeshDrawingPolicyOverrideSettings ComputeMeshOverrideSet
 * Creates and sets the base PSO so that resources can be set. Generally best to call during SetSharedState.
 */
 template<class DrawingPolicyType>
-void CommitGraphicsPipelineState(FRHICommandList& RHICmdList, const DrawingPolicyType& DrawingPolicy, const FDrawingPolicyRenderState& DrawRenderState, const FBoundShaderStateInput& BoundShaderStateInput)
+void CommitGraphicsPipelineState(FRHICommandList& RHICmdList, const DrawingPolicyType& DrawingPolicy, const FDrawingPolicyRenderState& DrawRenderState, const FBoundShaderStateInput& BoundShaderStateInput, const FMaterialRenderProxy* InMaterialRenderProxy)
 {
 	FGraphicsPipelineStateInitializer GraphicsPSOInit;
 
@@ -278,10 +278,9 @@ void CommitGraphicsPipelineState(FRHICommandList& RHICmdList, const DrawingPolic
 	GraphicsPSOInit.BoundShaderState = BoundShaderStateInput;
 	GraphicsPSOInit.RasterizerState = DrawingPolicy.ComputeRasterizerState(DrawRenderState.GetViewOverrideFlags());
 
-	const FMaterialRenderProxy* const MaterialRenderProxy = DrawingPolicy.GetMaterialRenderProxy();
-	if (MaterialRenderProxy != nullptr)
+	if (InMaterialRenderProxy != nullptr)
 	{
-		GraphicsPSOInit.ImmutableSamplerState = MaterialRenderProxy->ImmutableSamplerState;
+		GraphicsPSOInit.ImmutableSamplerState = InMaterialRenderProxy->ImmutableSamplerState;
 	}
 
 	check(DrawRenderState.GetDepthStencilState());
@@ -350,6 +349,11 @@ public:
 		OnlyApplyDitheredLODTransitionState(DrawRenderState, ViewInfo, Mesh, InAllowStencilDither);
 	}
 
+	const FMaterialRenderProxy* GetPipelineMaterialRenderProxy(const FMaterialRenderProxy* ElementMaterialRenderProxy)
+	{
+		return MaterialRenderProxy;
+	}
+
 	FDrawingPolicyMatchResult Matches(const FMeshDrawingPolicy& OtherDrawer, bool bForReals = false) const
 	{
 		DRAWING_POLICY_MATCH_BEGIN
