Merging CL 4323400 from Dev-VR - Remove StereoPair in favor of visibility map merging introduced in CL 4036044

diff --git a/Engine/Source/Runtime/Renderer/Private/BasePassRendering.cpp b/Engine/Source/Runtime/Renderer/Private/BasePassRendering.cpp
--- a/Engine/Source/Runtime/Renderer/Private/BasePassRendering.cpp
+++ b/Engine/Source/Runtime/Renderer/Private/BasePassRendering.cpp
@@ -960,32 +960,12 @@ bool FDeferredShadingSceneRenderer::RenderBasePassStaticDataType(FRHICommandList
 {
 	SCOPED_DRAW_EVENTF(RHICmdList, StaticType, TEXT("Static EBasePassDrawListType=%d"), DrawType);
 
-	bool bDirty = false;
-
-	if (!View.IsInstancedStereoPass())
-	{
-		bDirty |= Scene->BasePassUniformLightMapPolicyDrawList[DrawType].DrawVisible(RHICmdList, View, DrawRenderState, View.StaticMeshVisibilityMap, View.StaticMeshBatchVisibility);
-	}
-	else
-	{
-		const StereoPair StereoView(Views[0], Views[1], Views[0].StaticMeshVisibilityMap, Views[1].StaticMeshVisibilityMap, Views[0].StaticMeshBatchVisibility, Views[1].StaticMeshBatchVisibility);
-		bDirty |= Scene->BasePassUniformLightMapPolicyDrawList[DrawType].DrawVisibleInstancedStereo(RHICmdList, StereoView, DrawRenderState);
-	}
-
-	return bDirty;
+	return Scene->BasePassUniformLightMapPolicyDrawList[DrawType].DrawVisible(RHICmdList, View, DrawRenderState, View.StaticMeshVisibilityMap, View.StaticMeshBatchVisibility);
 }
 
 void FDeferredShadingSceneRenderer::RenderBasePassStaticDataTypeParallel(FParallelCommandListSet& ParallelCommandListSet, const EBasePassDrawListType DrawType)
 {
-	if (!ParallelCommandListSet.View.IsInstancedStereoPass())
-	{
-		Scene->BasePassUniformLightMapPolicyDrawList[DrawType].DrawVisibleParallel(ParallelCommandListSet.View.StaticMeshVisibilityMap, ParallelCommandListSet.View.StaticMeshBatchVisibility, ParallelCommandListSet);
-	}
-	else
-	{
-		const StereoPair StereoView(Views[0], Views[1], Views[0].StaticMeshVisibilityMap, Views[1].StaticMeshVisibilityMap, Views[0].StaticMeshBatchVisibility, Views[1].StaticMeshBatchVisibility);
-		Scene->BasePassUniformLightMapPolicyDrawList[DrawType].DrawVisibleParallelInstancedStereo(StereoView, ParallelCommandListSet);
-	}
+	Scene->BasePassUniformLightMapPolicyDrawList[DrawType].DrawVisibleParallel(ParallelCommandListSet.View.StaticMeshVisibilityMap, ParallelCommandListSet.View.StaticMeshBatchVisibility, ParallelCommandListSet);
 }
 
 template<typename StaticMeshDrawList>
diff --git a/Engine/Source/Runtime/Renderer/Private/DepthRendering.cpp b/Engine/Source/Runtime/Renderer/Private/DepthRendering.cpp
--- a/Engine/Source/Runtime/Renderer/Private/DepthRendering.cpp
+++ b/Engine/Source/Runtime/Renderer/Private/DepthRendering.cpp
@@ -933,46 +933,24 @@ bool FDeferredShadingSceneRenderer::RenderPrePassView(FRHICommandList& RHICmdLis
 	SetupPrePassView(RHICmdList, View, this);
 
 	// Draw the static occluder primitives using a depth drawing policy.
-
-	if (!View.IsInstancedStereoPass())
 	{
-		{
-			// Draw opaque occluders which support a separate position-only
-			// vertex buffer to minimize vertex fetch bandwidth, which is
-			// often the bottleneck during the depth only pass.
-			SCOPED_DRAW_EVENT(RHICmdList, PosOnlyOpaque);
-			bDirty |= Scene->PositionOnlyDepthDrawList.DrawVisible(RHICmdList, View, DrawRenderState, View.StaticMeshOccluderMap, View.StaticMeshBatchVisibility);
-		}
-		{
-			// Draw opaque occluders, using double speed z where supported.
-			SCOPED_DRAW_EVENT(RHICmdList, Opaque);
-			bDirty |= Scene->DepthDrawList.DrawVisible(RHICmdList, View, DrawRenderState, View.StaticMeshOccluderMap, View.StaticMeshBatchVisibility);
-		}
-
-		if (EarlyZPassMode >= DDM_AllOccluders)
-		{
-			// Draw opaque occluders with masked materials
-			SCOPED_DRAW_EVENT(RHICmdList, Masked);
-			bDirty |= Scene->MaskedDepthDrawList.DrawVisible(RHICmdList, View, DrawRenderState, View.StaticMeshOccluderMap, View.StaticMeshBatchVisibility);
-		}
+		// Draw opaque occluders which support a separate position-only
+		// vertex buffer to minimize vertex fetch bandwidth, which is
+		// often the bottleneck during the depth only pass.
+		SCOPED_DRAW_EVENT(RHICmdList, PosOnlyOpaque);
+		bDirty |= Scene->PositionOnlyDepthDrawList.DrawVisible(RHICmdList, View, DrawRenderState, View.StaticMeshOccluderMap, View.StaticMeshBatchVisibility);
 	}
-	else
 	{
-		const StereoPair StereoView(Views[0], Views[1], Views[0].StaticMeshOccluderMap, Views[1].StaticMeshOccluderMap, Views[0].StaticMeshBatchVisibility, Views[1].StaticMeshBatchVisibility);
-		{
-			SCOPED_DRAW_EVENT(RHICmdList, PosOnlyOpaque);
-			bDirty |= Scene->PositionOnlyDepthDrawList.DrawVisibleInstancedStereo(RHICmdList, StereoView, DrawRenderState);
-		}
-		{
-			SCOPED_DRAW_EVENT(RHICmdList, Opaque);
-			bDirty |= Scene->DepthDrawList.DrawVisibleInstancedStereo(RHICmdList, StereoView, DrawRenderState);
-		}
+		// Draw opaque occluders, using double speed z where supported.
+		SCOPED_DRAW_EVENT(RHICmdList, Opaque);
+		bDirty |= Scene->DepthDrawList.DrawVisible(RHICmdList, View, DrawRenderState, View.StaticMeshOccluderMap, View.StaticMeshBatchVisibility);
+	}
 
-		if (EarlyZPassMode >= DDM_AllOccluders)
-		{
-			SCOPED_DRAW_EVENT(RHICmdList, Masked);
-			bDirty |= Scene->MaskedDepthDrawList.DrawVisibleInstancedStereo(RHICmdList, StereoView, DrawRenderState);
-		}
+	if (EarlyZPassMode >= DDM_AllOccluders)
+	{
+		// Draw opaque occluders with masked materials
+		SCOPED_DRAW_EVENT(RHICmdList, Masked);
+		bDirty |= Scene->MaskedDepthDrawList.DrawVisible(RHICmdList, View, DrawRenderState, View.StaticMeshOccluderMap, View.StaticMeshBatchVisibility);
 	}
 
 	{
@@ -1063,34 +1041,19 @@ bool FDeferredShadingSceneRenderer::RenderPrePassViewParallel(const FViewInfo& V
 			CVarRHICmdFlushRenderThreadTasksPrePass.GetValueOnRenderThread() == 0 && CVarRHICmdFlushRenderThreadTasks.GetValueOnRenderThread() == 0,
 			DrawRenderState);
 
-		if (!View.IsInstancedStereoPass())
-		{
-			// Draw the static occluder primitives using a depth drawing policy.
-			// Draw opaque occluders which support a separate position-only
-			// vertex buffer to minimize vertex fetch bandwidth, which is
-			// often the bottleneck during the depth only pass.
-			Scene->PositionOnlyDepthDrawList.DrawVisibleParallel(View.StaticMeshOccluderMap, View.StaticMeshBatchVisibility, ParallelCommandListSet);
-
-			// Draw opaque occluders, using double speed z where supported.
-			Scene->DepthDrawList.DrawVisibleParallel(View.StaticMeshOccluderMap, View.StaticMeshBatchVisibility, ParallelCommandListSet);
-
-			// Draw opaque occluders with masked materials
-			if (EarlyZPassMode >= DDM_AllOccluders)
-			{			
-				Scene->MaskedDepthDrawList.DrawVisibleParallel(View.StaticMeshOccluderMap, View.StaticMeshBatchVisibility, ParallelCommandListSet);
-			}
-		}
-		else
-		{
-			const StereoPair StereoView(Views[0], Views[1], Views[0].StaticMeshOccluderMap, Views[1].StaticMeshOccluderMap, Views[0].StaticMeshBatchVisibility, Views[1].StaticMeshBatchVisibility);
+		// Draw the static occluder primitives using a depth drawing policy.
+		// Draw opaque occluders which support a separate position-only
+		// vertex buffer to minimize vertex fetch bandwidth, which is
+		// often the bottleneck during the depth only pass.
+		Scene->PositionOnlyDepthDrawList.DrawVisibleParallel(View.StaticMeshOccluderMap, View.StaticMeshBatchVisibility, ParallelCommandListSet);
 
-			Scene->PositionOnlyDepthDrawList.DrawVisibleParallelInstancedStereo(StereoView, ParallelCommandListSet);
-			Scene->DepthDrawList.DrawVisibleParallelInstancedStereo(StereoView, ParallelCommandListSet);
+		// Draw opaque occluders, using double speed z where supported.
+		Scene->DepthDrawList.DrawVisibleParallel(View.StaticMeshOccluderMap, View.StaticMeshBatchVisibility, ParallelCommandListSet);
 
-			if (EarlyZPassMode >= DDM_AllOccluders)
-			{
-				Scene->MaskedDepthDrawList.DrawVisibleParallelInstancedStereo(StereoView, ParallelCommandListSet);
-			}
+		// Draw opaque occluders with masked materials
+		if (EarlyZPassMode >= DDM_AllOccluders)
+		{			
+			Scene->MaskedDepthDrawList.DrawVisibleParallel(View.StaticMeshOccluderMap, View.StaticMeshBatchVisibility, ParallelCommandListSet);
 		}
 
 		if (!GStartPrepassParallelTranslatesImmediately)
diff --git a/Engine/Source/Runtime/Renderer/Private/MobileBasePassRendering.cpp b/Engine/Source/Runtime/Renderer/Private/MobileBasePassRendering.cpp
--- a/Engine/Source/Runtime/Renderer/Private/MobileBasePassRendering.cpp
+++ b/Engine/Source/Runtime/Renderer/Private/MobileBasePassRendering.cpp
@@ -584,13 +584,10 @@ EBasePassSort::Type GetSortMode()
 static void DrawVisibleFrontToBack(
 	FRHICommandListImmediate& RHICmdList,
 	FScene* const Scene,
-	const EBasePassDrawListType DrawListType, 
-	const FViewInfo& View, 
-	const FDrawingPolicyRenderState& DrawRenderState,
-	const FMobileCSMVisibilityInfo* const MobileCSMVisibilityInfo, 
-	const StereoPair& StereoView, 
-	const StereoPair& StereoViewCSM, 
-	const StereoPair& StereoViewNonCSM, 
+	const EBasePassDrawListType DrawListType,
+	const FViewInfo& View,
+	FDrawingPolicyRenderState& DrawRenderState,
+	const FMobileCSMVisibilityInfo* const MobileCSMVisibilityInfo,
 	int32& MaxDraws)
 {
 #if FRAMEPRO_ENABLED
@@ -601,29 +598,14 @@ static void DrawVisibleFrontToBack(
 
 	const bool bIsCSM = MobileCSMVisibilityInfo != nullptr;
 	int32 NumDraws = 0;
-	if (View.bIsMobileMultiViewEnabled)
+	if (bIsCSM)
 	{
-		if (bIsCSM)
-		{
-			NumDraws += Scene->MobileBasePassUniformLightMapPolicyDrawListWithCSM[DrawListType].DrawVisibleFrontToBackMobileMultiView(RHICmdList, StereoViewCSM, DrawRenderState, MaxDraws);
-			NumDraws += Scene->MobileBasePassUniformLightMapPolicyDrawList[DrawListType].DrawVisibleFrontToBackMobileMultiView(RHICmdList, StereoViewNonCSM, DrawRenderState, MaxDraws);
-		}
-		else
-		{
-			NumDraws += Scene->MobileBasePassUniformLightMapPolicyDrawList[DrawListType].DrawVisibleFrontToBackMobileMultiView(RHICmdList, StereoView, DrawRenderState, MaxDraws);
-		}
+		NumDraws += Scene->MobileBasePassUniformLightMapPolicyDrawListWithCSM[DrawListType].DrawVisibleFrontToBack(RHICmdList, View, DrawRenderState, MobileCSMVisibilityInfo->MobileCSMStaticMeshVisibilityMap, MobileCSMVisibilityInfo->MobileCSMStaticBatchVisibility, MaxDraws);
+		NumDraws += Scene->MobileBasePassUniformLightMapPolicyDrawList[DrawListType].DrawVisibleFrontToBack(RHICmdList, View, DrawRenderState, MobileCSMVisibilityInfo->MobileNonCSMStaticMeshVisibilityMap, MobileCSMVisibilityInfo->MobileNonCSMStaticBatchVisibility, MaxDraws);
 	}
 	else
 	{
-		if (bIsCSM)
-		{
-			NumDraws += Scene->MobileBasePassUniformLightMapPolicyDrawListWithCSM[DrawListType].DrawVisibleFrontToBack(RHICmdList, View, DrawRenderState, MobileCSMVisibilityInfo->MobileCSMStaticMeshVisibilityMap, MobileCSMVisibilityInfo->MobileCSMStaticBatchVisibility, MaxDraws);
-			NumDraws += Scene->MobileBasePassUniformLightMapPolicyDrawList[DrawListType].DrawVisibleFrontToBack(RHICmdList, View, DrawRenderState, MobileCSMVisibilityInfo->MobileNonCSMStaticMeshVisibilityMap, MobileCSMVisibilityInfo->MobileNonCSMStaticBatchVisibility, MaxDraws);
-		}
-		else
-		{
-			NumDraws += Scene->MobileBasePassUniformLightMapPolicyDrawList[DrawListType].DrawVisibleFrontToBack(RHICmdList, View, DrawRenderState, View.StaticMeshVisibilityMap, View.StaticMeshBatchVisibility, MaxDraws);
-		}
+		NumDraws += Scene->MobileBasePassUniformLightMapPolicyDrawList[DrawListType].DrawVisibleFrontToBack(RHICmdList, View, DrawRenderState, View.StaticMeshVisibilityMap, View.StaticMeshBatchVisibility, MaxDraws);
 	}
 
 	MaxDraws -= NumDraws;
@@ -640,36 +622,18 @@ static void DrawVisible(
 	const EBasePassDrawListType DrawListType,
 	const FViewInfo& View,
 	const FDrawingPolicyRenderState& DrawRenderState,
-	const FMobileCSMVisibilityInfo* const MobileCSMVisibilityInfo,
-	const StereoPair& StereoView,
-	const StereoPair& StereoViewCSM,
-	const StereoPair& StereoViewNonCSM)
+	const FMobileCSMVisibilityInfo* const MobileCSMVisibilityInfo)
 {
 	SCOPE_CYCLE_COUNTER(STAT_StaticDrawListDrawTime);
 	const bool bIsCSM = MobileCSMVisibilityInfo != nullptr;
-	if (View.bIsMobileMultiViewEnabled)
+	if (bIsCSM)
 	{
-		if (bIsCSM)
-		{
-			Scene->MobileBasePassUniformLightMapPolicyDrawListWithCSM[DrawListType].DrawVisibleMobileMultiView(RHICmdList, StereoViewCSM, DrawRenderState);
-			Scene->MobileBasePassUniformLightMapPolicyDrawList[DrawListType].DrawVisibleMobileMultiView(RHICmdList, StereoViewNonCSM, DrawRenderState);
-		}
-		else
-		{
-			Scene->MobileBasePassUniformLightMapPolicyDrawList[DrawListType].DrawVisibleMobileMultiView(RHICmdList, StereoView, DrawRenderState);
-		}
+		Scene->MobileBasePassUniformLightMapPolicyDrawListWithCSM[DrawListType].DrawVisible(RHICmdList, View, DrawRenderState, MobileCSMVisibilityInfo->MobileCSMStaticMeshVisibilityMap, MobileCSMVisibilityInfo->MobileCSMStaticBatchVisibility);
+		Scene->MobileBasePassUniformLightMapPolicyDrawList[DrawListType].DrawVisible(RHICmdList, View, DrawRenderState, MobileCSMVisibilityInfo->MobileNonCSMStaticMeshVisibilityMap, MobileCSMVisibilityInfo->MobileNonCSMStaticBatchVisibility);
 	}
 	else
 	{
-		if (bIsCSM)
-		{
-			Scene->MobileBasePassUniformLightMapPolicyDrawListWithCSM[DrawListType].DrawVisible(RHICmdList, View, DrawRenderState, MobileCSMVisibilityInfo->MobileCSMStaticMeshVisibilityMap, MobileCSMVisibilityInfo->MobileCSMStaticBatchVisibility);
-			Scene->MobileBasePassUniformLightMapPolicyDrawList[DrawListType].DrawVisible(RHICmdList, View, DrawRenderState, MobileCSMVisibilityInfo->MobileNonCSMStaticMeshVisibilityMap, MobileCSMVisibilityInfo->MobileNonCSMStaticBatchVisibility);
-		}
-		else
-		{
-			Scene->MobileBasePassUniformLightMapPolicyDrawList[DrawListType].DrawVisible(RHICmdList, View, DrawRenderState, View.StaticMeshVisibilityMap, View.StaticMeshBatchVisibility);
-		}
+		Scene->MobileBasePassUniformLightMapPolicyDrawList[DrawListType].DrawVisible(RHICmdList, View, DrawRenderState, View.StaticMeshVisibilityMap, View.StaticMeshBatchVisibility);
 	}
 }
 
@@ -740,52 +704,6 @@ void FMobileSceneRenderer::RenderMobileBasePassDynamicData(FRHICommandList& RHIC
 	}
 }
 
-struct FMobileBasePassViewInfo
-{
-	const FMobileCSMVisibilityInfo* MobileCSMVisibilityInfo;
-	const FMobileCSMVisibilityInfo* MobileCSMVisibilityInfoStereo;
-	StereoPair StereoView;
-	StereoPair StereoViewCSM;
-	StereoPair StereoViewNonCSM;
-
-
-	FMobileBasePassViewInfo(const FViewInfo& View, TArray<FViewInfo>& Views)
-	{
-		MobileCSMVisibilityInfo = View.MobileCSMVisibilityInfo.bMobileDynamicCSMInUse ? &View.MobileCSMVisibilityInfo : nullptr;
-		MobileCSMVisibilityInfoStereo = nullptr;
-
-		if (View.bIsMobileMultiViewEnabled)
-		{
-			checkSlow(Views.Num() > 1);
-			StereoView.LeftView = &Views[0];
-			StereoView.RightView = &Views[1];
-			StereoView.LeftViewVisibilityMap = &Views[0].StaticMeshVisibilityMap;
-			StereoView.LeftViewBatchVisibilityArray = &Views[0].StaticMeshBatchVisibility;
-			StereoView.RightViewVisibilityMap = &Views[1].StaticMeshVisibilityMap;
-			StereoView.RightViewBatchVisibilityArray = &Views[1].StaticMeshBatchVisibility;
-
-			if (MobileCSMVisibilityInfo)
-			{
-				MobileCSMVisibilityInfoStereo = &Views[1].MobileCSMVisibilityInfo;
-
-				StereoViewCSM.LeftView = &Views[0];
-				StereoViewCSM.RightView = &Views[1];
-				StereoViewCSM.LeftViewVisibilityMap = &MobileCSMVisibilityInfo->MobileCSMStaticMeshVisibilityMap;
-				StereoViewCSM.LeftViewBatchVisibilityArray = &MobileCSMVisibilityInfo->MobileCSMStaticBatchVisibility;
-				StereoViewCSM.RightViewVisibilityMap = &MobileCSMVisibilityInfoStereo->MobileCSMStaticMeshVisibilityMap;
-				StereoViewCSM.RightViewBatchVisibilityArray = &MobileCSMVisibilityInfoStereo->MobileCSMStaticBatchVisibility;
-
-				StereoViewNonCSM.LeftView = &Views[0];
-				StereoViewNonCSM.RightView = &Views[1];
-				StereoViewNonCSM.LeftViewVisibilityMap = &MobileCSMVisibilityInfo->MobileNonCSMStaticMeshVisibilityMap;
-				StereoViewNonCSM.LeftViewBatchVisibilityArray = &MobileCSMVisibilityInfo->MobileNonCSMStaticBatchVisibility;
-				StereoViewNonCSM.RightViewVisibilityMap = &MobileCSMVisibilityInfoStereo->MobileNonCSMStaticMeshVisibilityMap;
-				StereoViewNonCSM.RightViewBatchVisibilityArray = &MobileCSMVisibilityInfoStereo->MobileNonCSMStaticBatchVisibility;
-			}
-		}
-	}
-};
-
 static void SetupMobileBasePassView(FRHICommandList& RHICmdList, const FViewInfo& View, FDrawingPolicyRenderState& DrawRenderState)
 {
 	// Opaque blending
@@ -917,10 +835,9 @@ static void RenderMobileBasePassDynamicDataParallel(FMobileSceneRenderer& ThisRe
 void FMobileSceneRenderer::RenderMobileBasePassViewParallel(const FViewInfo& View, FRHICommandListImmediate& ParentCmdList, TArray<FViewInfo>& InViews, const FDrawingPolicyRenderState& DrawRenderState)
 {
 	QUICK_SCOPE_CYCLE_COUNTER(STAT_RenderMobileBasePassViewParallel);
-	FMobileBasePassViewInfo VI(View, InViews);
 
 	bool bCreateContexts = false; // CVarRHICmdFlushRenderThreadTasksBasePass.GetValueOnRenderThread() == 0 && CVarRHICmdFlushRenderThreadTasks.GetValueOnRenderThread() == 0;
-	const bool bIsCSM = VI.MobileCSMVisibilityInfo != nullptr;
+	const bool bIsCSM = View.MobileCSMVisibilityInfo.bMobileDynamicCSMInUse;
 	check(!View.bIsMobileMultiViewEnabled && GetSortMode() != EBasePassSort::SortPerMesh); // easy to support, just isn't supported yet
 
 	{
@@ -933,8 +850,8 @@ void FMobileSceneRenderer::RenderMobileBasePassViewParallel(const FViewInfo& Vie
 
 		if (bIsCSM)
 		{
-			Scene->MobileBasePassUniformLightMapPolicyDrawListWithCSM[EBasePass_Default].DrawVisibleParallel(VI.MobileCSMVisibilityInfo->MobileCSMStaticMeshVisibilityMap, VI.MobileCSMVisibilityInfo->MobileCSMStaticBatchVisibility, ParallelSet);
-			Scene->MobileBasePassUniformLightMapPolicyDrawList[EBasePass_Default].DrawVisibleParallel(VI.MobileCSMVisibilityInfo->MobileNonCSMStaticMeshVisibilityMap, VI.MobileCSMVisibilityInfo->MobileNonCSMStaticBatchVisibility, ParallelSet);
+			Scene->MobileBasePassUniformLightMapPolicyDrawListWithCSM[EBasePass_Default].DrawVisibleParallel(View.MobileCSMVisibilityInfo.MobileCSMStaticMeshVisibilityMap, View.MobileCSMVisibilityInfo.MobileCSMStaticBatchVisibility, ParallelSet);
+			Scene->MobileBasePassUniformLightMapPolicyDrawList[EBasePass_Default].DrawVisibleParallel(View.MobileCSMVisibilityInfo.MobileNonCSMStaticMeshVisibilityMap, View.MobileCSMVisibilityInfo.MobileNonCSMStaticBatchVisibility, ParallelSet);
 		}
 		else
 		{
@@ -957,8 +874,8 @@ void FMobileSceneRenderer::RenderMobileBasePassViewParallel(const FViewInfo& Vie
 
 		if (bIsCSM)
 		{
-			Scene->MobileBasePassUniformLightMapPolicyDrawListWithCSM[EBasePass_Masked].DrawVisibleParallel(VI.MobileCSMVisibilityInfo->MobileCSMStaticMeshVisibilityMap, VI.MobileCSMVisibilityInfo->MobileCSMStaticBatchVisibility, ParallelSet);
-			Scene->MobileBasePassUniformLightMapPolicyDrawList[EBasePass_Masked].DrawVisibleParallel(VI.MobileCSMVisibilityInfo->MobileNonCSMStaticMeshVisibilityMap, VI.MobileCSMVisibilityInfo->MobileNonCSMStaticBatchVisibility, ParallelSet);
+			Scene->MobileBasePassUniformLightMapPolicyDrawListWithCSM[EBasePass_Masked].DrawVisibleParallel(View.MobileCSMVisibilityInfo.MobileCSMStaticMeshVisibilityMap, View.MobileCSMVisibilityInfo.MobileCSMStaticBatchVisibility, ParallelSet);
+			Scene->MobileBasePassUniformLightMapPolicyDrawList[EBasePass_Masked].DrawVisibleParallel(View.MobileCSMVisibilityInfo.MobileNonCSMStaticMeshVisibilityMap, View.MobileCSMVisibilityInfo.MobileNonCSMStaticBatchVisibility, ParallelSet);
 		}
 		else
 		{
@@ -1059,16 +976,14 @@ void FMobileSceneRenderer::RenderMobileBasePass(FRHICommandListImmediate& RHICmd
 
 			SetupMobileBasePassView(RHICmdList, View, DrawRenderState);
 
-			FMobileBasePassViewInfo VI(View, Views);
-
 			// Render the base pass static data
 			if (SortMode == EBasePassSort::SortPerMesh)
 			{
-				DrawVisibleFrontToBack(RHICmdList, Scene, EBasePass_Default, View, DrawRenderState, VI.MobileCSMVisibilityInfo, VI.StereoView, VI.StereoViewNonCSM, VI.StereoViewCSM, MaxDraws);
+				DrawVisibleFrontToBack(RHICmdList, Scene, EBasePass_Default, View, DrawRenderState, &View.MobileCSMVisibilityInfo, MaxDraws);
 			}
 			else
 			{
-				DrawVisible(RHICmdList, Scene, EBasePass_Default, View, DrawRenderState, VI.MobileCSMVisibilityInfo, VI.StereoView, VI.StereoViewNonCSM, VI.StereoViewCSM);
+				DrawVisible(RHICmdList, Scene, EBasePass_Default, View, DrawRenderState, &View.MobileCSMVisibilityInfo);
 			}
 			FRHICommandListExecutor::GetImmediateCommandList().PollOcclusionQueries();
 			FRHICommandListExecutor::GetImmediateCommandList().ImmediateFlush(EImmediateFlushType::DispatchToRHIThread);
@@ -1085,11 +1000,11 @@ void FMobileSceneRenderer::RenderMobileBasePass(FRHICommandListImmediate& RHICmd
 			// Issue static draw list masked draw calls last, as PVR wants it
 			if (SortMode == EBasePassSort::SortPerMesh)
 			{
-				DrawVisibleFrontToBack(RHICmdList, Scene, EBasePass_Masked, View, DrawRenderState, VI.MobileCSMVisibilityInfo, VI.StereoView, VI.StereoViewNonCSM, VI.StereoViewCSM, MaxDraws);
+				DrawVisibleFrontToBack(RHICmdList, Scene, EBasePass_Masked, View, DrawRenderState, &View.MobileCSMVisibilityInfo, MaxDraws);
 			}
 			else
 			{
-				DrawVisible(RHICmdList, Scene, EBasePass_Masked, View, DrawRenderState, VI.MobileCSMVisibilityInfo, VI.StereoView, VI.StereoViewNonCSM, VI.StereoViewCSM);
+				DrawVisible(RHICmdList, Scene, EBasePass_Masked, View, DrawRenderState, &View.MobileCSMVisibilityInfo);
 			}
 
 			FRHICommandListExecutor::GetImmediateCommandList().PollOcclusionQueries();
diff --git a/Engine/Source/Runtime/Renderer/Private/StaticMeshDrawList.h b/Engine/Source/Runtime/Renderer/Private/StaticMeshDrawList.h
--- a/Engine/Source/Runtime/Renderer/Private/StaticMeshDrawList.h
+++ b/Engine/Source/Runtime/Renderer/Private/StaticMeshDrawList.h
@@ -16,48 +16,6 @@
 #include "SceneRendering.h"
 
 extern ENGINE_API bool GDrawListsLocked;
-/** View state for instanced stereo rendering. */
-struct StereoPair
-{
-	StereoPair() :
-		LeftView(nullptr), 
-		RightView(nullptr),
-		LeftViewVisibilityMap(nullptr),
-		RightViewVisibilityMap(nullptr),
-		LeftViewBatchVisibilityArray(nullptr),
-		RightViewBatchVisibilityArray(nullptr)
-	{}
-
-	StereoPair(
-		const FViewInfo& InLeftView,
-		const FViewInfo& InRightView,
-		const TBitArray<SceneRenderingBitArrayAllocator>& InLeftViewVisibilityMap,
-		const TBitArray<SceneRenderingBitArrayAllocator>& InRightViewVisibilityMap,
-		const TArray<uint64, SceneRenderingAllocator>& InLeftViewBatchVisibilityArray,
-		const TArray<uint64, SceneRenderingAllocator>& InRightViewBatchVisibilityArray
-	) :
-		LeftView(&InLeftView), 
-		RightView(&InRightView),
-		LeftViewVisibilityMap(&InLeftViewVisibilityMap),
-		RightViewVisibilityMap(&InRightViewVisibilityMap),
-		LeftViewBatchVisibilityArray(&InLeftViewBatchVisibilityArray),
-		RightViewBatchVisibilityArray(&InRightViewBatchVisibilityArray)
-	{}
-
-	const FViewInfo* LeftView;
-	const FViewInfo* RightView;
-	const TBitArray<SceneRenderingBitArrayAllocator>* LeftViewVisibilityMap;
-	const TBitArray<SceneRenderingBitArrayAllocator>* RightViewVisibilityMap;
-	const TArray<uint64, SceneRenderingAllocator>* LeftViewBatchVisibilityArray;
-	const TArray<uint64, SceneRenderingAllocator>* RightViewBatchVisibilityArray;
-};
-
-/** Used to statically branch between regular rendering and instanced stereo rendering. */
-enum class InstancedStereoPolicy {
-	Enabled, 
-	MobileMultiView,
-	Disabled
-};
 
 /** Base class of the static draw list, used when comparing draw lists and the drawing policy type is not necessary. */
 class FStaticMeshDrawListBase
@@ -358,7 +316,6 @@ private:
 	* @param DrawingPolicyLink - the drawing policy link
 	* @param bDrawnShared - determines whether to draw shared
 	*/
-	template<InstancedStereoPolicy InstancedStereo>
 	int32 DrawElement(FRHICommandList& RHICmdList, const FViewInfo& View, const typename DrawingPolicyType::ContextDataType PolicyContext, FDrawingPolicyRenderState& DrawRenderState, const FElement& Element, uint64 BatchElementMask, FDrawingPolicyLink* DrawingPolicyLink, bool &bDrawnShared);
 
 public:
@@ -379,60 +336,23 @@ public:
 
 	/**
 	* Draws only the static meshes which are in the visibility map, limited to a range of policies
-	* Both StaticMeshVisibilityMap and BatchVisibilityArray should be non-null for regular rendering or StereoView should be non-null if rendering with instanced stereo
-	* @param View - The view of the meshes to render (use the left view of a stereo pair when rendering with instanced stereo)
+	* @param View - The view of the meshes to render
 	* @param StaticMeshVisibilityMap - An map from FStaticMesh::Id to visibility state.
 	* @param BatchVisibilityArray - An array of batch element visibility bitmasks.
-	* @param StereoView - Stereo pair to render
 	* @param FirstPolicy - First policy to render
 	* @param LastPolicy - Last policy to render
 	* @return True if any static meshes were drawn.
 	*/
-	template<InstancedStereoPolicy InstancedStereo>
 	bool DrawVisibleInner(FRHICommandList& RHICmdList,
 		const FViewInfo& View,
 		const typename DrawingPolicyType::ContextDataType PolicyContext,
 		FDrawingPolicyRenderState& DrawRenderState,
-		const TBitArray<SceneRenderingBitArrayAllocator>* const StaticMeshVisibilityMap,
-		const TArray<uint64, SceneRenderingAllocator>* const BatchVisibilityArray,
-		const StereoPair* const StereoView,
+		const TBitArray<SceneRenderingBitArrayAllocator>& StaticMeshVisibilityMap,
+		const TArray<uint64, SceneRenderingAllocator>& BatchVisibilityArray,
 		int32 FirstPolicy,
 		int32 LastPolicy,
 		bool bUpdateCounts);
 
-	/**
-	* Draws only the static meshes which are in the visibility map of the stereo pair
-	* @param StereoView - The stereo pair to render.
-	* @return True if any static meshes were drawn.
-	*/
-	inline bool DrawVisibleInstancedStereo(
-		FRHICommandList& RHICmdList,
-		const StereoPair& StereoView,
-		const FDrawingPolicyRenderState& DrawRenderState)
-	{
-		//moved out of the inner loop and only modified if bDrawnShared
-		FDrawingPolicyRenderState DrawRenderStateLocal(DrawRenderState);
-
-		return DrawVisibleInner<InstancedStereoPolicy::Enabled>(RHICmdList, *StereoView.LeftView, typename DrawingPolicyType::ContextDataType(true), DrawRenderStateLocal, nullptr, nullptr, &StereoView, 0, OrderedDrawingPolicies.Num() - 1, false);
-	}
-
-	/**
-	* Draws only the static meshes which are in the visibility map of the stereo pair
-	* Stereo instancing is not enabled, the driver handles this for mobile multi-view
-	* @param StereoView - The stereo pair to render.
-	* @return True if any static meshes were drawn.
-	*/
-	inline bool DrawVisibleMobileMultiView(
-		FRHICommandList& RHICmdList,
-		const StereoPair& StereoView,
-		const FDrawingPolicyRenderState& DrawRenderState)
-	{
-		//moved out of the inner loop and only modified if bDrawnShared
-		FDrawingPolicyRenderState DrawRenderStateLocal(DrawRenderState);
-
-		return DrawVisibleInner<InstancedStereoPolicy::MobileMultiView>(RHICmdList, *StereoView.LeftView, typename DrawingPolicyType::ContextDataType(false), DrawRenderStateLocal, nullptr, nullptr, &StereoView, 0, OrderedDrawingPolicies.Num() - 1, false);
-	}
-
 	/**
 	 * Draws only the static meshes which are in the visibility map.
 	 * @param View - The view of the meshes to render.
@@ -441,123 +361,57 @@ public:
 	 * @return True if any static meshes were drawn.
 	 */
 	bool DrawVisible(FRHICommandList& RHICmdList, const FViewInfo& View, const typename DrawingPolicyType::ContextDataType PolicyContext, const FDrawingPolicyRenderState& DrawRenderState, const TBitArray<SceneRenderingBitArrayAllocator>& StaticMeshVisibilityMap, const TArray<uint64, SceneRenderingAllocator>& BatchVisibilityArray);
-
-private:
-
-	void DrawVisibleParallelInternal(
-		const typename DrawingPolicyType::ContextDataType PolicyContext,
-		const TBitArray<SceneRenderingBitArrayAllocator>* StaticMeshVisibilityMap,
-		const TArray<uint64, SceneRenderingAllocator>* BatchVisibilityArray,
-		const StereoPair* const StereoView,
-		FParallelCommandListSet& ParallelCommandListSet);
-
-public:
-
+	
 	/**
 	* Draws only the static meshes which are in the visibility map.
 	* @param StaticMeshVisibilityMap - An map from FStaticMesh::Id to visibility state.
 	* @param BatchVisibilityArray - An array of batch element visibility bitmasks.
 	* @param ParallelCommandListSet - holds information on how to get a fresh command list and deal with submits, etc
 	*/
-	inline void DrawVisibleParallel(
+	void DrawVisibleParallel(
 		const typename DrawingPolicyType::ContextDataType PolicyContext,
 		const TBitArray<SceneRenderingBitArrayAllocator>& StaticMeshVisibilityMap,
 		const TArray<uint64, SceneRenderingAllocator>& BatchVisibilityArray,
-		FParallelCommandListSet& ParallelCommandListSet)
-	{
-		DrawVisibleParallelInternal(PolicyContext, &StaticMeshVisibilityMap, &BatchVisibilityArray, nullptr, ParallelCommandListSet);
-	}
-
-	/**
-	* Draws only the static meshes which are in the visibility map of either the left or right view using instanced stereo
-	* @param StereoView - Instanced stereo view state
-	* @param ParallelCommandListSet - holds information on how to get a fresh command list and deal with submits, etc
-	*/
-	inline void DrawVisibleParallelInstancedStereo(const StereoPair& StereoView, FParallelCommandListSet& ParallelCommandListSet)
-	{
-		DrawVisibleParallelInternal(typename DrawingPolicyType::ContextDataType(true), nullptr, nullptr, &StereoView, ParallelCommandListSet);
-	}
+		FParallelCommandListSet& ParallelCommandListSet);
 
 	/**
 	* Draws only the static meshes which are in the visibility map, sorted front-to-back.
-	* Both StaticMeshVisibilityMap and BatchVisibilityArray should be non-null for regular rendering or StereoView should be non-null if rendering with mobile multi-view
-	* @param View - The view of the meshes to render (use the left view of a stereo pair when rendering with mobile multi-view)
+	* @param View - The view of the meshes to render
 	* @param StaticMeshVisibilityMap - An map from FStaticMesh::Id to visibility state.
 	* @param BatchVisibilityArray - An array of batch element visibility bitmasks.
-	* @param StereoView - Stereo pair to render
 	* @param MaxToDraw - The maximum number of meshes to be drawn.
 	 * @return The number of static meshes drawn.
 	*/
-	template<InstancedStereoPolicy InstancedStereo>
-	int32 DrawVisibleFrontToBackInner(
+	int32 DrawVisibleFrontToBack(
 		FRHICommandList& RHICmdList,
 		const FViewInfo& View,
 		FDrawingPolicyRenderState& DrawRenderState,
 		const typename DrawingPolicyType::ContextDataType PolicyContext,
-		const TBitArray<SceneRenderingBitArrayAllocator>* const StaticMeshVisibilityMap,
-		const TArray<uint64, SceneRenderingAllocator>* const BatchVisibilityArray,
-		const StereoPair* const StereoView,
+		const TBitArray<SceneRenderingBitArrayAllocator>& StaticMeshVisibilityMap,
+		const TArray<uint64, SceneRenderingAllocator>& BatchVisibilityArray,
 		int32 MaxToDraw);
 
-	/**
-	 * Draws only the static meshes which are in the visibility map, sorted front-to-back.
-	 * @param View - The view of the meshes to render.
-	 * @param StaticMeshVisibilityMap - An map from FStaticMesh::Id to visibility state.
-	 * @param BatchVisibilityArray - An array of batch element visibility bitmasks.
-	 * @param MaxToDraw - The maximum number of meshes to be drawn.
-	 * @return The number of static meshes drawn.
-	 */
-	inline int32 DrawVisibleFrontToBack(
-		FRHICommandList& RHICmdList, 
-		const FViewInfo& View, 
-		const FDrawingPolicyRenderState& DrawRenderState,
-		const typename DrawingPolicyType::ContextDataType PolicyContext, 
-		const TBitArray<SceneRenderingBitArrayAllocator>& StaticMeshVisibilityMap, 
-		const TArray<uint64, SceneRenderingAllocator>& BatchVisibilityArray, 
-		int32 MaxToDraw)
-	{
-		//moved out of the inner loop and only modified if bDrawnShared
-		FDrawingPolicyRenderState DrawRenderStateLocal(DrawRenderState);
-
-		return DrawVisibleFrontToBackInner<InstancedStereoPolicy::Disabled>(RHICmdList, View, DrawRenderStateLocal, PolicyContext, &StaticMeshVisibilityMap, &BatchVisibilityArray, nullptr, MaxToDraw);
-	}
-
-	/**
-	* Draws only the static meshes which are in the visibility map of the stereo pair
-	* Stereo instancing is not enabled, the driver handles this for mobile multi-view
-	* @param StereoView - The stereo pair to render.
-	* @param MaxToDraw - The maximum number of meshes to be drawn.
-	* @return The number of static meshes drawn.
-	*/
-	inline int32 DrawVisibleFrontToBackMobileMultiView(FRHICommandList& RHICmdList, const StereoPair &StereoView, const FDrawingPolicyRenderState& DrawRenderState, const int32 MaxToDraw)
-	{
-		//moved out of the inner loop and only modified if bDrawnShared
-		FDrawingPolicyRenderState DrawRenderStateLocal(DrawRenderState);
-
-		return DrawVisibleFrontToBackInner<InstancedStereoPolicy::MobileMultiView>(RHICmdList, *StereoView.LeftView, DrawRenderStateLocal, typename DrawingPolicyType::ContextDataType(false), nullptr, nullptr, &StereoView, MaxToDraw);
-	}
-
 	/**
 	 * Helper functions when policy context is not needed.
 	 */
 	inline bool DrawVisible(const FViewInfo& View, const FDrawingPolicyRenderState& DrawRenderState, const TBitArray<SceneRenderingBitArrayAllocator>& StaticMeshVisibilityMap)
 	{
-		return DrawVisible(View, typename DrawingPolicyType::ContextDataType(), DrawRenderState, StaticMeshVisibilityMap);
+		return DrawVisible(View, typename DrawingPolicyType::ContextDataType(View.IsInstancedStereoPass()), DrawRenderState, StaticMeshVisibilityMap);
 	}
 
 	inline bool DrawVisible(FRHICommandList& RHICmdList, const FViewInfo& View, const FDrawingPolicyRenderState& DrawRenderState, const TBitArray<SceneRenderingBitArrayAllocator>& StaticMeshVisibilityMap, const TArray<uint64,SceneRenderingAllocator>& BatchVisibilityArray)
 	{
-		return DrawVisible(RHICmdList, View, typename DrawingPolicyType::ContextDataType(), DrawRenderState, StaticMeshVisibilityMap, BatchVisibilityArray);
+		return DrawVisible(RHICmdList, View, typename DrawingPolicyType::ContextDataType(View.IsInstancedStereoPass()), DrawRenderState, StaticMeshVisibilityMap, BatchVisibilityArray);
 	}
 
 	inline void DrawVisibleParallel(const TBitArray<SceneRenderingBitArrayAllocator>& StaticMeshVisibilityMap, const TArray<uint64, SceneRenderingAllocator>& BatchVisibilityArray, FParallelCommandListSet& ParallelCommandListSet)
 	{
-		DrawVisibleParallel(typename DrawingPolicyType::ContextDataType(), StaticMeshVisibilityMap, BatchVisibilityArray, ParallelCommandListSet);
+		DrawVisibleParallel(typename DrawingPolicyType::ContextDataType(ParallelCommandListSet.View.IsInstancedStereoPass()), StaticMeshVisibilityMap, BatchVisibilityArray, ParallelCommandListSet);
 	}
 
-	inline int32 DrawVisibleFrontToBack(FRHICommandList& RHICmdList, const FViewInfo& View, const FDrawingPolicyRenderState& DrawRenderState, const TBitArray<SceneRenderingBitArrayAllocator>& StaticMeshVisibilityMap, const TArray<uint64,SceneRenderingAllocator>& BatchVisibilityArray, int32 MaxToDraw)
+	inline int32 DrawVisibleFrontToBack(FRHICommandList& RHICmdList, const FViewInfo& View, FDrawingPolicyRenderState& DrawRenderState, const TBitArray<SceneRenderingBitArrayAllocator>& StaticMeshVisibilityMap, const TArray<uint64,SceneRenderingAllocator>& BatchVisibilityArray, int32 MaxToDraw)
 	{
-		return DrawVisibleFrontToBack(RHICmdList, View, DrawRenderState, typename DrawingPolicyType::ContextDataType(), StaticMeshVisibilityMap, BatchVisibilityArray, MaxToDraw);
+		return DrawVisibleFrontToBack(RHICmdList, View, DrawRenderState, typename DrawingPolicyType::ContextDataType(View.IsInstancedStereoPass()), StaticMeshVisibilityMap, BatchVisibilityArray, MaxToDraw);
 	}
 
 	/** Sorts OrderedDrawingPolicies front to back. */
diff --git a/Engine/Source/Runtime/Renderer/Private/StaticMeshDrawList.inl b/Engine/Source/Runtime/Renderer/Private/StaticMeshDrawList.inl
--- a/Engine/Source/Runtime/Renderer/Private/StaticMeshDrawList.inl
+++ b/Engine/Source/Runtime/Renderer/Private/StaticMeshDrawList.inl
@@ -40,7 +40,6 @@ class FViewInfo;
 class SceneRenderingBitArrayAllocator;
 struct FDrawListStats;
 struct FMeshDrawingRenderState;
-struct StereoPair;
 enum class InstancedStereoPolicy;
 template<typename DrawingPolicyType> class TStaticMeshDrawList;
 template<typename DrawingPolicyType> struct TCompareStaticMeshDrawList;
@@ -141,7 +140,6 @@ void TStaticMeshDrawList<DrawingPolicyType>::FElementHandle::Remove(const bool b
 }
 
 template<typename DrawingPolicyType>
-template<InstancedStereoPolicy InstancedStereo>
 int32 TStaticMeshDrawList<DrawingPolicyType>::DrawElement(
 	FRHICommandList& RHICmdList,
 	const FViewInfo& View,
@@ -201,7 +199,7 @@ int32 TStaticMeshDrawList<DrawingPolicyType>::DrawElement(
 		if(BatchElementMask & 1)
 		{
 			const FPrimitiveSceneProxy* Proxy = Element.Mesh->PrimitiveSceneInfo->Proxy;			
-			if (InstancedStereo == InstancedStereoPolicy::Enabled)
+			if (View.IsInstancedStereoPass())
 			{
 			    // We draw instanced static meshes twice when rendering with instanced stereo. Once for each eye.
 			    const bool bIsInstancedMesh = Element.Mesh->Elements[BatchElementIndex].bIsInstancedMesh;
@@ -358,23 +356,17 @@ void TStaticMeshDrawList<DrawingPolicyType>::ReleaseRHI()
 }
 
 template<typename DrawingPolicyType>
-template<InstancedStereoPolicy InstancedStereo>
 bool TStaticMeshDrawList<DrawingPolicyType>::DrawVisibleInner(
 	FRHICommandList& RHICmdList,
 	const FViewInfo& View,
 	const typename DrawingPolicyType::ContextDataType PolicyContext,
 	FDrawingPolicyRenderState& DrawRenderState,
-	const TBitArray<SceneRenderingBitArrayAllocator>* const StaticMeshVisibilityMap,
-	const TArray<uint64, SceneRenderingAllocator>* const BatchVisibilityArray,
-	const StereoPair* const StereoView,
+	const TBitArray<SceneRenderingBitArrayAllocator>& StaticMeshVisibilityMap,
+	const TArray<uint64, SceneRenderingAllocator>& BatchVisibilityArray,
 	int32 FirstPolicy, int32 LastPolicy,
 	bool bUpdateCounts
 	)
 {
-	// We should have a single view's visibility data, or a stereo pair
-	check((StaticMeshVisibilityMap != nullptr && BatchVisibilityArray != nullptr) || StereoView != nullptr);
-	check((InstancedStereo != InstancedStereoPolicy::Disabled) == (StereoView != nullptr));
-
 	bool bDirty = false;
 	STAT(int32 StatInc = 0;)
 	uint32 NumDraws = 0;
@@ -389,53 +381,18 @@ bool TStaticMeshDrawList<DrawingPolicyType>::DrawVisibleInner(
 		uint32 Count = 0;
 		for (int32 ElementIndex = 0; ElementIndex < NumElements; ElementIndex++, CompactElementPtr++)
 		{
-			// Single view
-			if (InstancedStereo == InstancedStereoPolicy::Disabled)
+			if (StaticMeshVisibilityMap.AccessCorrespondingBit(FRelativeBitReference(CompactElementPtr->MeshId)))
 			{
-				if (StaticMeshVisibilityMap->AccessCorrespondingBit(FRelativeBitReference(CompactElementPtr->MeshId)))
+				const FElement& Element = DrawingPolicyLink->Elements[ElementIndex];
+				STAT(StatInc += Element.Mesh->GetNumPrimitives();)
+				int32 SubCount = Element.Mesh->Elements.Num();
+				// Avoid the cache miss looking up batch visibility if there is only one element.
+				uint64 BatchElementMask = Element.Mesh->bRequiresPerElementVisibility ? BatchVisibilityArray[Element.Mesh->BatchVisibilityId] : ((1ull << SubCount) - 1);
+				Count += DrawElement(RHICmdList, View, PolicyContext, DrawRenderState, Element, BatchElementMask, DrawingPolicyLink, bDrawnShared);
+				if (++NumDraws % 16 == 0)
 				{
-					const FElement& Element = DrawingPolicyLink->Elements[ElementIndex];
-					STAT(StatInc += Element.Mesh->GetNumPrimitives();)
-					int32 SubCount = Element.Mesh->Elements.Num();
-					// Avoid the cache miss looking up batch visibility if there is only one element.
-					uint64 BatchElementMask = Element.Mesh->bRequiresPerElementVisibility ? (*BatchVisibilityArray)[Element.Mesh->BatchVisibilityId] : ((1ull << SubCount) - 1);
-					Count += DrawElement<InstancedStereoPolicy::Disabled>(RHICmdList, View, PolicyContext, DrawRenderState, Element, BatchElementMask, DrawingPolicyLink, bDrawnShared);
-					if (++NumDraws % 16 == 0)
-					{
-						// Periodically kick the RHI thread when some work has accumulated.
-						RHICmdList.MaybeDispatchToRHIThread();
-					}
-
-				}
-			}
-
-			// Stereo pair, we need to test both eyes
-			else
-			{
-				const TArray<uint64, SceneRenderingAllocator>* ResolvedBatchVisiblityArray = nullptr;
-
-				if (StereoView->LeftViewVisibilityMap->AccessCorrespondingBit(FRelativeBitReference(CompactElementPtr->MeshId)))
-				{
-					ResolvedBatchVisiblityArray = StereoView->LeftViewBatchVisibilityArray;
-				}
-				else if (StereoView->RightViewVisibilityMap->AccessCorrespondingBit(FRelativeBitReference(CompactElementPtr->MeshId)))
-				{
-					ResolvedBatchVisiblityArray = StereoView->RightViewBatchVisibilityArray;
-				}
-
-				if (ResolvedBatchVisiblityArray != nullptr)
-				{
-					const FElement& Element = DrawingPolicyLink->Elements[ElementIndex];
-					STAT(StatInc += Element.Mesh->GetNumPrimitives();)
-						int32 SubCount = Element.Mesh->Elements.Num();
-					// Avoid the cache miss looking up batch visibility if there is only one element.
-					uint64 BatchElementMask = Element.Mesh->bRequiresPerElementVisibility ? (*ResolvedBatchVisiblityArray)[Element.Mesh->BatchVisibilityId] : ((1ull << SubCount) - 1);
-					Count += DrawElement<InstancedStereo>(RHICmdList, View, PolicyContext, DrawRenderState, Element, BatchElementMask, DrawingPolicyLink, bDrawnShared);
-					if (++NumDraws % 16 == 0)
-					{
-						// Periodically kick the RHI thread when some work has accumulated.
-						RHICmdList.MaybeDispatchToRHIThread();
-					}
+					// Periodically kick the RHI thread when some work has accumulated.
+					RHICmdList.MaybeDispatchToRHIThread();
 				}
 			}
 		}
@@ -462,7 +419,7 @@ bool TStaticMeshDrawList<DrawingPolicyType>::DrawVisible(
 	//moved out of the inner loop and only modified if bDrawnShared
 	FDrawingPolicyRenderState DrawRenderStateLocal(DrawRenderState);
 
-	return DrawVisibleInner<InstancedStereoPolicy::Disabled>(RHICmdList, View, PolicyContext, DrawRenderStateLocal, &StaticMeshVisibilityMap, &BatchVisibilityArray, nullptr, 0, OrderedDrawingPolicies.Num() - 1, false);
+	return DrawVisibleInner(RHICmdList, View, PolicyContext, DrawRenderStateLocal, StaticMeshVisibilityMap, BatchVisibilityArray, 0, OrderedDrawingPolicies.Num() - 1, false);
 }
 
 template<typename DrawingPolicyType>
@@ -473,12 +430,11 @@ class FDrawVisibleAnyThreadTask : public FRenderTask
 	const FViewInfo& View;
 	FDrawingPolicyRenderState DrawRenderState;
 	const typename DrawingPolicyType::ContextDataType PolicyContext;
-	const TBitArray<SceneRenderingBitArrayAllocator>* StaticMeshVisibilityMap;
-	const TArray<uint64, SceneRenderingAllocator>* BatchVisibilityArray;
+	const TBitArray<SceneRenderingBitArrayAllocator>& StaticMeshVisibilityMap;
+	const TArray<uint64, SceneRenderingAllocator>& BatchVisibilityArray;
 	const int32 FirstPolicy;
 	const int32 LastPolicy;
 	const TArray<uint16, SceneRenderingAllocator>* PerDrawingPolicyCounts;
-	StereoPair StereoView;
 
 public:
 
@@ -488,9 +444,8 @@ public:
 		const FViewInfo& InView,
 		const FDrawingPolicyRenderState& InDrawRenderState,
 		const typename DrawingPolicyType::ContextDataType& InPolicyContext,
-		const TBitArray<SceneRenderingBitArrayAllocator>* InStaticMeshVisibilityMap,
-		const TArray<uint64, SceneRenderingAllocator>* InBatchVisibilityArray,
-		const StereoPair* const InStereoView,
+		const TBitArray<SceneRenderingBitArrayAllocator>& InStaticMeshVisibilityMap,
+		const TArray<uint64, SceneRenderingAllocator>& InBatchVisibilityArray,
 		int32 InFirstPolicy,
 		int32 InLastPolicy,
 		const TArray<uint16, SceneRenderingAllocator>* InPerDrawingPolicyCounts
@@ -506,10 +461,6 @@ public:
 		, LastPolicy(InLastPolicy)
 		, PerDrawingPolicyCounts(InPerDrawingPolicyCounts)
 	{
-		if (InStereoView != nullptr)
-		{
-			StereoView = *InStereoView;
-		}
 	}
 
 	FORCEINLINE TStatId GetStatId() const
@@ -541,14 +492,7 @@ public:
 					{
 						BatchEnd++;
 					}
-					if (!this->PolicyContext.bIsInstancedStereo)
-					{
-						this->Caller.template DrawVisibleInner<InstancedStereoPolicy::Disabled>(this->RHICmdList, this->View, this->PolicyContext, this->DrawRenderState, this->StaticMeshVisibilityMap, this->BatchVisibilityArray, nullptr, Start, BatchEnd, true);
-					}
-					else
-					{
-						this->Caller.template DrawVisibleInner<InstancedStereoPolicy::Enabled>(this->RHICmdList, *this->StereoView.LeftView, this->PolicyContext, this->DrawRenderState, nullptr, nullptr, &StereoView, Start, BatchEnd, true);
-					}
+					this->Caller.DrawVisibleInner(this->RHICmdList, this->View, this->PolicyContext, this->DrawRenderState, this->StaticMeshVisibilityMap, this->BatchVisibilityArray, Start, BatchEnd, true);
 					
 					Start = BatchEnd + 1;
 				}
@@ -556,31 +500,20 @@ public:
 		}
 		else
 		{
-			if (!this->PolicyContext.bIsInstancedStereo)
-			{
-				this->Caller.template DrawVisibleInner<InstancedStereoPolicy::Disabled>(this->RHICmdList, this->View, this->PolicyContext, this->DrawRenderState, this->StaticMeshVisibilityMap, this->BatchVisibilityArray, nullptr, this->FirstPolicy, this->LastPolicy, true);
-			}
-			else
-			{
-				this->Caller.template DrawVisibleInner<InstancedStereoPolicy::Enabled>(this->RHICmdList, *this->StereoView.LeftView, this->PolicyContext, this->DrawRenderState, nullptr, nullptr, &StereoView, this->FirstPolicy, this->LastPolicy, true);
-			}
+			this->Caller.DrawVisibleInner(this->RHICmdList, this->View, this->PolicyContext, this->DrawRenderState, this->StaticMeshVisibilityMap, this->BatchVisibilityArray, this->FirstPolicy, this->LastPolicy, true);
 		}
 		this->RHICmdList.HandleRTThreadTaskCompletion(MyCompletionGraphEvent);
 	}
 };
 
 template<typename DrawingPolicyType>
-void TStaticMeshDrawList<DrawingPolicyType>::DrawVisibleParallelInternal(
+void TStaticMeshDrawList<DrawingPolicyType>::DrawVisibleParallel(
 	const typename DrawingPolicyType::ContextDataType PolicyContext,
-	const TBitArray<SceneRenderingBitArrayAllocator>* StaticMeshVisibilityMap,
-	const TArray<uint64, SceneRenderingAllocator>* BatchVisibilityArray,
-	const StereoPair* const StereoView,
+	const TBitArray<SceneRenderingBitArrayAllocator>& StaticMeshVisibilityMap,
+	const TArray<uint64, SceneRenderingAllocator>& BatchVisibilityArray,
 	FParallelCommandListSet& ParallelCommandListSet
 	)
 {
-	// We should have a single view's visibility data, or a stereo pair
-	check((StaticMeshVisibilityMap != nullptr && BatchVisibilityArray != nullptr) || StereoView != nullptr);
-
 	int32 NumPolicies = OrderedDrawingPolicies.Num();
 	int32 EffectiveThreads = FMath::Min<int32>(NumPolicies, ParallelCommandListSet.Width);
 	if (EffectiveThreads)
@@ -615,7 +548,7 @@ void TStaticMeshDrawList<DrawingPolicyType>::DrawVisibleParallelInternal(
 			{
 				QUICK_SCOPE_CYCLE_COUNTER(STAT_FDeferredShadingSceneRenderer_DrawVisibleParallel_FullVisibilityScan);
 				ParallelForWithPreWork(NumPolicies, 
-					[this, &PerDrawingPolicyCounts, &StaticMeshVisibilityMap, &BatchVisibilityArray, StereoView](int32 Index)
+					[this, &PerDrawingPolicyCounts, &StaticMeshVisibilityMap, &BatchVisibilityArray](int32 Index)
 					{
 						int32 Count = 0;
 
@@ -624,27 +557,9 @@ void TStaticMeshDrawList<DrawingPolicyType>::DrawVisibleParallelInternal(
 						FPlatformMisc::Prefetch(CompactElementPtr);
 						const int32 NumElements = DrawingPolicyLink->CompactElements.Num();
 						FPlatformMisc::Prefetch(&DrawingPolicyLink->CompactElements.GetData()->MeshId);
-						const bool bIsInstancedStereo = StereoView != nullptr;
 						for (int32 ElementIndex = 0; ElementIndex < NumElements; ElementIndex++, CompactElementPtr++)
 						{
-							bool bIsVisible = false;
-
-							// Single view
-							if (!bIsInstancedStereo && StaticMeshVisibilityMap->AccessCorrespondingBit(FRelativeBitReference(CompactElementPtr->MeshId)))
-							{
-								bIsVisible = true;
-							}
-
-							// Stereo pair, we need to test both eyes
-							else if (
-								bIsInstancedStereo &&
-								(StereoView->LeftViewVisibilityMap->AccessCorrespondingBit(FRelativeBitReference(CompactElementPtr->MeshId)) ||
-								StereoView->RightViewVisibilityMap->AccessCorrespondingBit(FRelativeBitReference(CompactElementPtr->MeshId))))
-							{
-								bIsVisible = true;
-							}
-
-							if (bIsVisible)
+							if (StaticMeshVisibilityMap.AccessCorrespondingBit(FRelativeBitReference(CompactElementPtr->MeshId)))
 							{
 								const FElement& Element = DrawingPolicyLink->Elements[ElementIndex];
 								const int32 SubCount = Element.Mesh->Elements.Num();
@@ -654,15 +569,9 @@ void TStaticMeshDrawList<DrawingPolicyType>::DrawVisibleParallelInternal(
 								{
 									Count += SubCount;
 								}
-								else if (!bIsInstancedStereo)
-								{
-									Count += FPlatformMath::CountBits((*BatchVisibilityArray)[Element.Mesh->BatchVisibilityId]);
-								}
 								else
 								{
-									const int32 LeftCount = FPlatformMath::CountBits((*StereoView->LeftViewBatchVisibilityArray)[Element.Mesh->BatchVisibilityId]);
-									const int32 RightCount = FPlatformMath::CountBits((*StereoView->RightViewBatchVisibilityArray)[Element.Mesh->BatchVisibilityId]);
-									Count += (LeftCount > RightCount) ? LeftCount : RightCount;
+									Count += FPlatformMath::CountBits(BatchVisibilityArray[Element.Mesh->BatchVisibilityId]);
 								}
 							}
 						}
@@ -752,7 +661,7 @@ void TStaticMeshDrawList<DrawingPolicyType>::DrawVisibleParallelInternal(
 #endif
 								UE_CLOG(ParallelCommandListSet.bSpewBalance, LogTemp, Display, TEXT("    Index %d  BatchCount %d    (last merge)"), ParallelCommandListSet.NumParallelCommandLists(), PreviousBatchDraws + BatchCount);
 								FGraphEventRef AnyThreadCompletionEvent = TGraphTask<FDrawVisibleAnyThreadTask<DrawingPolicyType> >::CreateTask(ParallelCommandListSet.GetPrereqs(), RenderThread)
-									.ConstructAndDispatchWhenReady(*this, *CmdList, ParallelCommandListSet.View, ParallelCommandListSet.DrawRenderState, PolicyContext, StaticMeshVisibilityMap, BatchVisibilityArray, StereoView, PreviousBatchStart, LastNonZeroPolicy, bCountsAreAccurate ? &PerDrawingPolicyCounts : nullptr);
+									.ConstructAndDispatchWhenReady(*this, *CmdList, ParallelCommandListSet.View, ParallelCommandListSet.DrawRenderState, PolicyContext, StaticMeshVisibilityMap, BatchVisibilityArray, PreviousBatchStart, LastNonZeroPolicy, bCountsAreAccurate ? &PerDrawingPolicyCounts : nullptr);
 								ParallelCommandListSet.AddParallelCommandList(CmdList, AnyThreadCompletionEvent, FMath::Max<int32>(1, PreviousBatchDraws + BatchCount));
 								PreviousBatchStart = -1;
 								PreviousBatchEnd = -2;
@@ -766,7 +675,7 @@ void TStaticMeshDrawList<DrawingPolicyType>::DrawVisibleParallelInternal(
 								// this is a decent sized batch, emit the previous batch and save this one for possible merging
 								UE_CLOG(ParallelCommandListSet.bSpewBalance, LogTemp, Display, TEXT("    Index %d  BatchCount %d    "), ParallelCommandListSet.NumParallelCommandLists(), PreviousBatchDraws);
 								FGraphEventRef AnyThreadCompletionEvent = TGraphTask<FDrawVisibleAnyThreadTask<DrawingPolicyType> >::CreateTask(ParallelCommandListSet.GetPrereqs(), RenderThread)
-									.ConstructAndDispatchWhenReady(*this, *CmdList, ParallelCommandListSet.View, ParallelCommandListSet.DrawRenderState, PolicyContext, StaticMeshVisibilityMap, BatchVisibilityArray, StereoView, PreviousBatchStart, PreviousBatchEnd, bCountsAreAccurate ? &PerDrawingPolicyCounts : nullptr);
+									.ConstructAndDispatchWhenReady(*this, *CmdList, ParallelCommandListSet.View, ParallelCommandListSet.DrawRenderState, PolicyContext, StaticMeshVisibilityMap, BatchVisibilityArray, PreviousBatchStart, PreviousBatchEnd, bCountsAreAccurate ? &PerDrawingPolicyCounts : nullptr);
 								ParallelCommandListSet.AddParallelCommandList(CmdList, AnyThreadCompletionEvent, FMath::Max<int32>(1, PreviousBatchDraws));
 								PreviousBatchStart = Start;
 								PreviousBatchEnd = LastNonZeroPolicy;
@@ -793,7 +702,7 @@ void TStaticMeshDrawList<DrawingPolicyType>::DrawVisibleParallelInternal(
 				UE_CLOG(ParallelCommandListSet.bSpewBalance, LogTemp, Display, TEXT("    Index %d  BatchCount %d    (last)"), ParallelCommandListSet.NumParallelCommandLists(), PreviousBatchDraws);
 				FRHICommandList* CmdList = ParallelCommandListSet.NewParallelCommandList();
 				FGraphEventRef AnyThreadCompletionEvent = TGraphTask<FDrawVisibleAnyThreadTask<DrawingPolicyType> >::CreateTask(ParallelCommandListSet.GetPrereqs(), RenderThread)
-					.ConstructAndDispatchWhenReady(*this, *CmdList, ParallelCommandListSet.View, ParallelCommandListSet.DrawRenderState, PolicyContext, StaticMeshVisibilityMap, BatchVisibilityArray, StereoView, PreviousBatchStart, PreviousBatchEnd, bCountsAreAccurate ? &PerDrawingPolicyCounts : nullptr);
+					.ConstructAndDispatchWhenReady(*this, *CmdList, ParallelCommandListSet.View, ParallelCommandListSet.DrawRenderState, PolicyContext, StaticMeshVisibilityMap, BatchVisibilityArray, PreviousBatchStart, PreviousBatchEnd, bCountsAreAccurate ? &PerDrawingPolicyCounts : nullptr);
 				ParallelCommandListSet.AddParallelCommandList(CmdList, AnyThreadCompletionEvent, FMath::Max<int32>(1, PreviousBatchDraws));
 			}
 #if DO_CHECK
@@ -819,7 +728,7 @@ void TStaticMeshDrawList<DrawingPolicyType>::DrawVisibleParallelInternal(
 					FRHICommandList* CmdList = ParallelCommandListSet.NewParallelCommandList();
 
 						FGraphEventRef AnyThreadCompletionEvent = TGraphTask<FDrawVisibleAnyThreadTask<DrawingPolicyType> >::CreateTask(ParallelCommandListSet.GetPrereqs(), RenderThread)
-							.ConstructAndDispatchWhenReady(*this, *CmdList, ParallelCommandListSet.View, ParallelCommandListSet.DrawRenderState, PolicyContext, StaticMeshVisibilityMap, BatchVisibilityArray, StereoView, Start, Last, nullptr);
+							.ConstructAndDispatchWhenReady(*this, *CmdList, ParallelCommandListSet.View, ParallelCommandListSet.DrawRenderState, PolicyContext, StaticMeshVisibilityMap, BatchVisibilityArray, Start, Last, nullptr);
 
 					ParallelCommandListSet.AddParallelCommandList(CmdList, AnyThreadCompletionEvent);
 				}
@@ -832,42 +741,23 @@ void TStaticMeshDrawList<DrawingPolicyType>::DrawVisibleParallelInternal(
 }
 
 template<typename DrawingPolicyType>
-template<InstancedStereoPolicy InstancedStereo>
-int32 TStaticMeshDrawList<DrawingPolicyType>::DrawVisibleFrontToBackInner(
+int32 TStaticMeshDrawList<DrawingPolicyType>::DrawVisibleFrontToBack(
 	FRHICommandList& RHICmdList,
 	const FViewInfo& View,
 	FDrawingPolicyRenderState& DrawRenderState,
 	const typename DrawingPolicyType::ContextDataType PolicyContext,
-	const TBitArray<SceneRenderingBitArrayAllocator>* const StaticMeshVisibilityMap,
-	const TArray<uint64, SceneRenderingAllocator>* const BatchVisibilityArray,
-	const StereoPair* const StereoView,
+	const TBitArray<SceneRenderingBitArrayAllocator>& StaticMeshVisibilityMap,
+	const TArray<uint64, SceneRenderingAllocator>& BatchVisibilityArray,
 	int32 MaxToDraw
 	)
 {
-	// We should have a single view's visibility data, or a stereo pair
-	check((StaticMeshVisibilityMap != nullptr && BatchVisibilityArray != nullptr) || StereoView != nullptr);
-	check((InstancedStereo != InstancedStereoPolicy::Disabled) == (StereoView != nullptr));
-	
-	if (InstancedStereo == InstancedStereoPolicy::Disabled)
-	{
-		ComputeVisiblePoliciesBounds(*StaticMeshVisibilityMap);
-	}
-	else
-	{
-		ComputeVisiblePoliciesBounds(*StereoView->LeftViewVisibilityMap);
-	}
+	ComputeVisiblePoliciesBounds(StaticMeshVisibilityMap);
 
 	int32 NumDraws = 0;
 	TArray<FDrawListSortKey,SceneRenderingAllocator> SortKeys;
 	const FVector ViewLocation = View.ViewMatrices.GetViewOrigin();
 	SortKeys.Reserve(128);
 
-	TArray<const TArray<uint64, SceneRenderingAllocator>*, SceneRenderingAllocator> ElementVisibility;
-	if (InstancedStereo != InstancedStereoPolicy::Disabled)
-	{
-		ElementVisibility.Reserve(64);
-	}
-
 	for (typename TDrawingPolicySet::TConstIterator PolicyIt(DrawingPolicySet); PolicyIt; ++PolicyIt)
 	{
 		const FDrawingPolicyLink& DrawingPolicyLink = *PolicyIt;
@@ -878,28 +768,7 @@ int32 TStaticMeshDrawList<DrawingPolicyType>::DrawVisibleFrontToBackInner(
 		const FElementCompact* CompactElementPtr = DrawingPolicyLink.CompactElements.GetData();
 		for(int32 ElementIndex = 0; ElementIndex < NumElements; ElementIndex++, CompactElementPtr++)
 		{
-			bool bIsVisible = false;
-			if (InstancedStereo == InstancedStereoPolicy::Disabled)
-			{
-				// Single view
-				bIsVisible = StaticMeshVisibilityMap->AccessCorrespondingBit(FRelativeBitReference(CompactElementPtr->MeshId));
-			}
-			else
-			{
-				// Stereo pair, we need to test both eyes
-				if (StereoView->LeftViewVisibilityMap->AccessCorrespondingBit(FRelativeBitReference(CompactElementPtr->MeshId)))
-				{
-					bIsVisible = true;
-					ElementVisibility.Add(StereoView->LeftViewBatchVisibilityArray);
-				}
-				else if (StereoView->RightViewVisibilityMap->AccessCorrespondingBit(FRelativeBitReference(CompactElementPtr->MeshId)))
-				{
-					bIsVisible = true;
-					ElementVisibility.Add(StereoView->RightViewBatchVisibilityArray);
-				}
-			}
-
-			if (bIsVisible)
+			if (StaticMeshVisibilityMap.AccessCorrespondingBit(FRelativeBitReference(CompactElementPtr->MeshId)))
 			{
 				const FElement& Element = DrawingPolicyLink.Elements[ElementIndex];
 				const FBoxSphereBounds& Bounds = Element.Bounds;
@@ -934,11 +803,10 @@ int32 TStaticMeshDrawList<DrawingPolicyType>::DrawVisibleFrontToBackInner(
 		CA_SUPPRESS(6011);
 		const FElement& Element = DrawingPolicyLink->Elements[ElementIndex];
 		STAT(StatInc +=  Element.Mesh->GetNumPrimitives();)
-		const TArray<uint64, SceneRenderingAllocator>* const ResolvedVisiblityArray = (InstancedStereo == InstancedStereoPolicy::Disabled) ? BatchVisibilityArray : ElementVisibility[SortedIndex];
 
 		// Avoid the cache miss looking up batch visibility if there is only one element.
-		uint64 BatchElementMask = Element.Mesh->bRequiresPerElementVisibility ? (*ResolvedVisiblityArray)[Element.Mesh->BatchVisibilityId] : ((1ull << Element.Mesh->Elements.Num()) - 1);
-		DrawElement<InstancedStereoPolicy::Disabled>(RHICmdList, View, PolicyContext, DrawRenderState, Element, BatchElementMask, DrawingPolicyLink, bDrawnShared);
+		uint64 BatchElementMask = Element.Mesh->bRequiresPerElementVisibility ? BatchVisibilityArray[Element.Mesh->BatchVisibilityId] : ((1ull << Element.Mesh->Elements.Num()) - 1);
+		DrawElement(RHICmdList, View, PolicyContext, DrawRenderState, Element, BatchElementMask, DrawingPolicyLink, bDrawnShared);
 		NumDraws++;
 
 		if (((uint32)NumDraws % 16) == 0)
diff --git a/Engine/Source/Runtime/Renderer/Private/VelocityRendering.cpp b/Engine/Source/Runtime/Renderer/Private/VelocityRendering.cpp
--- a/Engine/Source/Runtime/Renderer/Private/VelocityRendering.cpp
+++ b/Engine/Source/Runtime/Renderer/Private/VelocityRendering.cpp
@@ -813,15 +813,7 @@ void FDeferredShadingSceneRenderer::RenderVelocitiesInnerParallel(FRHICommandLis
 				DrawRenderState,
 				VelocityRT);
 
-			if (!View.IsInstancedStereoPass())
-			{
-				Scene->VelocityDrawList.DrawVisibleParallel(View.StaticMeshVelocityMap, View.StaticMeshBatchVisibility, ParallelCommandListSet);
-			}
-			else
-			{
-				const StereoPair StereoView(Views[0], Views[1], Views[0].StaticMeshVelocityMap, Views[1].StaticMeshVelocityMap, Views[0].StaticMeshBatchVisibility, Views[1].StaticMeshBatchVisibility);
-				Scene->VelocityDrawList.DrawVisibleParallelInstancedStereo(StereoView, ParallelCommandListSet);
-			}
+			Scene->VelocityDrawList.DrawVisibleParallel(View.StaticMeshVelocityMap, View.StaticMeshBatchVisibility, ParallelCommandListSet);
 			
 			int32 NumPrims = View.DynamicMeshElements.Num();
 			int32 EffectiveThreads = FMath::Min<int32>(FMath::DivideAndRoundUp(NumPrims, ParallelCommandListSet.MinDrawsPerCommandList), ParallelCommandListSet.Width);
@@ -874,15 +866,7 @@ void FDeferredShadingSceneRenderer::RenderVelocitiesInner(FRHICommandListImmedia
 			SetVelocitiesState(RHICmdList, View, this, DrawRenderState, VelocityRT);
 
 			// Draw velocities for movable static meshes.
-			if (!View.IsInstancedStereoPass())
-			{
-				Scene->VelocityDrawList.DrawVisible(RHICmdList, View, DrawRenderState, View.StaticMeshVelocityMap, View.StaticMeshBatchVisibility);
-			}
-			else
-			{
-				const StereoPair StereoView(Views[0], Views[1], Views[0].StaticMeshVelocityMap, Views[1].StaticMeshVelocityMap, Views[0].StaticMeshBatchVisibility, Views[1].StaticMeshBatchVisibility);
-				Scene->VelocityDrawList.DrawVisibleInstancedStereo(RHICmdList, StereoView, DrawRenderState);
-			}
+			Scene->VelocityDrawList.DrawVisible(RHICmdList, View, DrawRenderState, View.StaticMeshVelocityMap, View.StaticMeshBatchVisibility);
 
 			RenderDynamicVelocitiesMeshElementsInner(RHICmdList, View, DrawRenderState, 0, View.DynamicMeshElements.Num() - 1);
 		}
