Merging CL 4323685 from Dev-VR - Fix GLES occlusion queries when HDR is disabled.
Add occlusion query support for Multi-View.
Add ShowFlag to visualize occlusion query bounding meshes.

diff --git a/Engine/Plugins/Runtime/Oculus/OculusVR/Source/OculusHMD/Private/OculusHMD.cpp b/Engine/Plugins/Runtime/Oculus/OculusVR/Source/OculusHMD/Private/OculusHMD.cpp
--- a/Engine/Plugins/Runtime/Oculus/OculusVR/Source/OculusHMD/Private/OculusHMD.cpp
+++ b/Engine/Plugins/Runtime/Oculus/OculusVR/Source/OculusHMD/Private/OculusHMD.cpp
@@ -1012,7 +1012,7 @@ namespace OculusHMD
 	{
 		if (Settings.IsValid())
 		{
-			const int32 ViewIndex = ViewIndexFromStereoPass(StereoPass);
+			const int32 ViewIndex = GetViewIndexForPass(StereoPass);
 			X = Settings->EyeUnscaledRenderViewport[ViewIndex].Min.X;
 			Y = Settings->EyeUnscaledRenderViewport[ViewIndex].Min.Y;
 			SizeX = Settings->EyeUnscaledRenderViewport[ViewIndex].Size().X;
@@ -1032,7 +1032,7 @@ namespace OculusHMD
 	{
 		CheckInRenderThread();
 
-		const int32 ViewIndex = ViewIndexFromStereoPass(StereoPass);
+		const int32 ViewIndex = GetViewIndexForPass(StereoPass);
 
 		if (Settings_RenderThread.IsValid())
 		{
@@ -1076,7 +1076,7 @@ namespace OculusHMD
 
 		check(IsStereoEnabled());
 
-		const int32 ViewIndex = ViewIndexFromStereoPass(StereoPassType);
+		const int32 ViewIndex = GetViewIndexForPass(StereoPassType);
 
 		FMatrix proj = ToFMatrix(Settings->EyeProjectionMatrices[ViewIndex]);
 
@@ -1134,7 +1134,7 @@ namespace OculusHMD
 		check(IsStereoEnabled());
 
 		// Don't use GetStereoProjectionMatrix because it is game thread only on oculus, we also don't need the zplane adjustments for this.
-		const int32 ViewIndex = ViewIndexFromStereoPass(StereoPassType);
+		const int32 ViewIndex = GetViewIndexForPass(StereoPassType);
 		const FMatrix StereoProjectionMatrix = ToFMatrix(Settings_RenderThread->EyeProjectionMatrices[ViewIndex]);
 
 		//0,0,1 is the straight ahead point, wherever it maps to is the center of the projection plane in -1..1 coordinates.  -1,-1 is bottom left.
diff --git a/Engine/Plugins/Runtime/Oculus/OculusVR/Source/OculusHMD/Private/OculusHMDPrivate.h b/Engine/Plugins/Runtime/Oculus/OculusVR/Source/OculusHMD/Private/OculusHMDPrivate.h
--- a/Engine/Plugins/Runtime/Oculus/OculusVR/Source/OculusHMD/Private/OculusHMDPrivate.h
+++ b/Engine/Plugins/Runtime/Oculus/OculusVR/Source/OculusHMD/Private/OculusHMDPrivate.h
@@ -116,25 +116,6 @@ namespace OculusHMD
 		return ovrpRecti { { rect.Min.X, rect.Min.Y }, { rect.Size().X, rect.Size().Y } };
 	}
 
-	FORCEINLINE int32 ViewIndexFromStereoPass(const EStereoscopicPass StereoPassType) {
-		switch (StereoPassType)
-		{
-		case eSSP_LEFT_EYE:
-		case eSSP_FULL:
-			return 0;
-
-		case eSSP_RIGHT_EYE:
-			return 1;
-
-		case eSSP_MONOSCOPIC_EYE:
-			return 2;
-
-		default:
-			check(0);
-			return -1;
-		}
-	}
-
 	/** Helper that converts ovrTrackedDeviceType to ETrackedDeviceType */
 	FORCEINLINE ETrackedDeviceType ToETrackedDeviceType(ovrpNode Source)
 	{
diff --git a/Engine/Shaders/Private/OcclusionQueryPixelShader.usf b/Engine/Shaders/Private/OcclusionQueryPixelShader.usf
new file mode 100644
--- /dev/null
+++ b/Engine/Shaders/Private/OcclusionQueryPixelShader.usf
@@ -0,0 +1,12 @@
+// Copyright 1998-2018 Epic Games, Inc. All Rights Reserved.
+
+/*=============================================================================
+	OcclusionQueryPixelShader.usf: Pixel shader for visualizing occlusion queries.
+=============================================================================*/
+
+#include "Common.ush"
+
+void Main(out float4 OutColor : SV_Target0)
+{
+	OutColor = float4(1, 0, 0, 1);
+}
diff --git a/Engine/Shaders/Private/OcclusionQueryPixelShader.ush b/Engine/Shaders/Private/OcclusionQueryPixelShader.ush
deleted file mode 100644
diff --git a/Engine/Shaders/Private/OcclusionQueryVertexShader.usf b/Engine/Shaders/Private/OcclusionQueryVertexShader.usf
--- a/Engine/Shaders/Private/OcclusionQueryVertexShader.usf
+++ b/Engine/Shaders/Private/OcclusionQueryVertexShader.usf
@@ -1,18 +1,44 @@
 // Copyright 1998-2018 Epic Games, Inc. All Rights Reserved.
 
 /*=============================================================================
-	OcclusionQueryVertexShader.hlsl: Vertex shader for drawing occlusion queries.
+	OcclusionQueryVertexShader.usf: Vertex shader for drawing occlusion queries.
 =============================================================================*/
 
 #include "Common.ush"
 
 float4 StencilingGeometryPosAndScale;
 
+#if MOBILE_MULTI_VIEW
+	uint ViewId;
+#endif
+
 void Main(
 	in float4 InPosition : ATTRIBUTE0,
 	out float4 OutPosition : SV_POSITION
+#if MOBILE_MULTI_VIEW
+	, in uint InViewId : SV_ViewID
+	, out uint OutViewId : SV_ViewID
+#endif
 	)
 {
+#if MOBILE_MULTI_VIEW
+	OutViewId = InViewId;
+
+	// Make the triangles degenerate if we're not targeting this view.
+	// We don't need to resolve the view info, since we already set the
+	// primary view with the data for the view we're interested in.
+	// TODO: Support targetting a single view properly.
+	if (InViewId != ViewId)
+	{
+		OutPosition = float4(0, 0, 0, 1);
+		return;
+	}
+#endif
+
 	float3 TransformedPosition = InPosition.xyz * StencilingGeometryPosAndScale.w + StencilingGeometryPosAndScale.xyz;
-	OutPosition = mul(float4(TransformedPosition,1), View.TranslatedWorldToClip);
+	OutPosition = mul(float4(TransformedPosition, 1), View.TranslatedWorldToClip);
+
+#if OUTPUT_GAMMA_SPACE && (COMPILER_GLSL_ES2 || COMPILER_GLSL_ES3_1 || COMPILER_GLSL_ES3_1_EXT)
+	OutPosition.y *= -1;
+#endif
 }
diff --git a/Engine/Source/Runtime/Engine/Public/ShowFlags.h b/Engine/Source/Runtime/Engine/Public/ShowFlags.h
--- a/Engine/Source/Runtime/Engine/Public/ShowFlags.h
+++ b/Engine/Source/Runtime/Engine/Public/ShowFlags.h
@@ -393,6 +393,7 @@ private:
 		SetBones(false);
 		SetScreenPercentage(InitMode != ESFIM_Editor && InitMode != ESFIM_VREditing);
 		SetVREditing(InitMode == ESFIM_VREditing);
+		SetOcclusionMeshes(false);
 	}
 
 
diff --git a/Engine/Source/Runtime/Engine/Public/ShowFlagsValues.inl b/Engine/Source/Runtime/Engine/Public/ShowFlagsValues.inl
--- a/Engine/Source/Runtime/Engine/Public/ShowFlagsValues.inl
+++ b/Engine/Source/Runtime/Engine/Public/ShowFlagsValues.inl
@@ -336,5 +336,7 @@ SHOWFLAG_FIXED_IN_SHIPPING(0, Bones, SFG_Developer, NSLOCTEXT("UnrealEd", "BoneS
 SHOWFLAG_ALWAYS_ACCESSIBLE(MediaPlanes, SFG_Normal, NSLOCTEXT("UnrealEd", "MediaPlanesSF", "Media Planes"))
 /** if this is a vr editing viewport, needed? */
 SHOWFLAG_FIXED_IN_SHIPPING(0, VREditing, SFG_Hidden, NSLOCTEXT("UnrealEd", "VREditSF", "VR Editing"))
+/** Visualize Occlusion Query bounding meshes */
+SHOWFLAG_FIXED_IN_SHIPPING(0, OcclusionMeshes, SFG_Visualize, NSLOCTEXT("UnrealEd", "VisualizeOcclusionQueries", "Visualize Occlusion Queries"))
 #undef SHOWFLAG_ALWAYS_ACCESSIBLE
 #undef SHOWFLAG_FIXED_IN_SHIPPING
diff --git a/Engine/Source/Runtime/Engine/Public/StereoRendering.h b/Engine/Source/Runtime/Engine/Public/StereoRendering.h
--- a/Engine/Source/Runtime/Engine/Public/StereoRendering.h
+++ b/Engine/Source/Runtime/Engine/Public/StereoRendering.h
@@ -64,6 +64,29 @@ public:
 		else
 			return EStereoscopicPass::eSSP_MONOSCOPIC_EYE;
 	}
+
+	/**
+	* For the specified stereoscopic pass type, assign a view index based on the extension's sorting
+	*/
+	virtual uint32 GetViewIndexForPass(EStereoscopicPass StereoPassType) const
+	{
+		switch (StereoPassType)
+		{
+		case eSSP_LEFT_EYE:
+		case eSSP_FULL:
+			return 0;
+
+		case eSSP_RIGHT_EYE:
+			return 1;
+
+		case eSSP_MONOSCOPIC_EYE:
+			return 2;
+
+		default:
+			check(0);
+			return -1;
+		}
+	}
 	
 	/**
 	* Return true if this pass is for a stereo eye view
diff --git a/Engine/Source/Runtime/Renderer/Private/SceneOcclusion.cpp b/Engine/Source/Runtime/Renderer/Private/SceneOcclusion.cpp
--- a/Engine/Source/Runtime/Renderer/Private/SceneOcclusion.cpp
+++ b/Engine/Source/Runtime/Renderer/Private/SceneOcclusion.cpp
@@ -95,6 +95,7 @@ int32 FOcclusionQueryHelpers::GetNumBufferedFrames(ERHIFeatureLevel::Type Featur
 
 // default, non-instanced shader implementation
 IMPLEMENT_SHADER_TYPE(,FOcclusionQueryVS,TEXT("/Engine/Private/OcclusionQueryVertexShader.usf"),TEXT("Main"),SF_Vertex);
+IMPLEMENT_SHADER_TYPE(,FOcclusionQueryPS,TEXT("/Engine/Private/OcclusionQueryPixelShader.usf"),TEXT("Main"),SF_Pixel);
 
 static FGlobalBoundShaderState GOcclusionTestBoundShaderState;
 
@@ -1507,9 +1508,15 @@ void FSceneRenderer::BeginOcclusionTests(FRHICommandListImmediate& RHICmdList, b
 
 				// Lookup the vertex shader.
 				TShaderMapRef<FOcclusionQueryVS> VertexShader(View.ShaderMap);
-
 				GraphicsPSOInit.BoundShaderState.VertexShaderRHI = GETSAFERHISHADER_VERTEX(*VertexShader);
 
+				if (View.Family->EngineShowFlags.OcclusionMeshes)
+				{
+					TShaderMapRef<FOcclusionQueryPS> PixelShader(View.ShaderMap);
+					GraphicsPSOInit.BoundShaderState.PixelShaderRHI = GETSAFERHISHADER_PIXEL(*PixelShader);
+					GraphicsPSOInit.BlendState = TStaticBlendState<CW_RGBA>::GetRHI();
+				}
+
 				SetGraphicsPipelineState(RHICmdList, GraphicsPSOInit);
 
 				if (FeatureLevel > ERHIFeatureLevel::ES3_1)
diff --git a/Engine/Source/Runtime/Renderer/Private/SceneOcclusion.h b/Engine/Source/Runtime/Renderer/Private/SceneOcclusion.h
--- a/Engine/Source/Runtime/Renderer/Private/SceneOcclusion.h
+++ b/Engine/Source/Runtime/Renderer/Private/SceneOcclusion.h
@@ -6,6 +6,7 @@
 #include "Shader.h"
 #include "GlobalShader.h"
 #include "ShadowRendering.h"
+#include "Engine/Engine.h"
 
 /*=============================================================================
 	SceneOcclusion.h
@@ -20,29 +21,46 @@ class FOcclusionQueryVS : public FGlobalShader
 public:
 	static bool ShouldCompilePermutation(const FGlobalShaderPermutationParameters& Parameters) { return IsFeatureLevelSupported(Parameters.Platform, ERHIFeatureLevel::ES3_1); }
 
+	static void ModifyCompilationEnvironment(const FGlobalShaderPermutationParameters& Parameters, FShaderCompilerEnvironment& OutEnvironment)
+	{
+		FGlobalShader::ModifyCompilationEnvironment(Parameters, OutEnvironment);
+		OutEnvironment.SetDefine(TEXT("OUTPUT_GAMMA_SPACE"), IsMobileHDR() == false);
+	}
+
 	FOcclusionQueryVS(const ShaderMetaType::CompiledShaderInitializerType& Initializer):
 		FGlobalShader(Initializer)
 	{
 			StencilingGeometryParameters.Bind(Initializer.ParameterMap);
+			ViewId.Bind(Initializer.ParameterMap, TEXT("ViewId"));
 	}
 
 	FOcclusionQueryVS() {}
 
-	void SetParametersWithBoundingSphere(FRHICommandList& RHICmdList, const FSceneView& View, const FSphere& BoundingSphere)
+	void SetParametersWithBoundingSphere(FRHICommandList& RHICmdList, const FViewInfo& View, const FSphere& BoundingSphere)
 	{
 		FGlobalShader::SetParameters<FViewUniformShaderParameters>(RHICmdList, GetVertexShader(), View.ViewUniformBuffer);
 
 		FVector4 StencilingSpherePosAndScale;
 		StencilingGeometry::GStencilSphereVertexBuffer.CalcTransform(StencilingSpherePosAndScale, BoundingSphere, View.ViewMatrices.GetPreViewTranslation());
 		StencilingGeometryParameters.Set(RHICmdList, this, StencilingSpherePosAndScale);
+
+		if (GEngine && GEngine->StereoRenderingDevice)
+		{
+			SetShaderValue(RHICmdList, GetVertexShader(), ViewId, GEngine->StereoRenderingDevice->GetViewIndexForPass(View.StereoPass));
+		}
 	}
 
-	void SetParameters(FRHICommandList& RHICmdList, const FSceneView& View)
+	void SetParameters(FRHICommandList& RHICmdList, const FViewInfo& View)
 	{
 		FGlobalShader::SetParameters<FViewUniformShaderParameters>(RHICmdList, GetVertexShader(),View.ViewUniformBuffer);
 
 		// Don't transform if rendering frustum
 		StencilingGeometryParameters.Set(RHICmdList, this, FVector4(0,0,0,1));
+
+		if (GEngine && GEngine->StereoRenderingDevice)
+		{
+			SetShaderValue(RHICmdList, GetVertexShader(), ViewId, GEngine->StereoRenderingDevice->GetViewIndexForPass(View.StereoPass));
+		}
 	}
 
 	//~ Begin FShader Interface
@@ -50,11 +68,27 @@ public:
 	{
 		bool bShaderHasOutdatedParameters = FGlobalShader::Serialize(Ar);
 		Ar << StencilingGeometryParameters;
+		Ar << ViewId;
 		return bShaderHasOutdatedParameters;
 	}
 	//~ Begin  End FShader Interface 
 
 private:
 	FStencilingGeometryShaderParameters StencilingGeometryParameters;
+	FShaderParameter ViewId;
+};
+
+/**
+ * A pixel shader for rendering a texture on a simple element.
+ */
+class FOcclusionQueryPS : public FGlobalShader
+{
+	DECLARE_SHADER_TYPE(FOcclusionQueryPS, Global);
+public:
+	static bool ShouldCompilePermutation(const FGlobalShaderPermutationParameters& Parameters) { return IsFeatureLevelSupported(Parameters.Platform, ERHIFeatureLevel::ES3_1); }
+
+	FOcclusionQueryPS(const ShaderMetaType::CompiledShaderInitializerType& Initializer) :
+		FGlobalShader(Initializer) {}
+	FOcclusionQueryPS() {}
 };
 
