From 4.20.3 hotfix: 4333358
Test render target texture size if available to ensure we don't do an extra upscale when not needed.
The render target size is not updated when the backing render target texture is reallocated.

diff --git a/Engine/Source/Runtime/Renderer/Private/MobileShadingRenderer.cpp b/Engine/Source/Runtime/Renderer/Private/MobileShadingRenderer.cpp
--- a/Engine/Source/Runtime/Renderer/Private/MobileShadingRenderer.cpp
+++ b/Engine/Source/Runtime/Renderer/Private/MobileShadingRenderer.cpp
@@ -235,7 +235,10 @@ void FMobileSceneRenderer::Render(FRHICommandListImmediate& RHICmdList)
 	}
 
 	const bool bGammaSpace = !IsMobileHDR();
-	const bool bRequiresUpscale = ((int32)ViewFamily.RenderTarget->GetSizeXY().X > FamilySize.X || (int32)ViewFamily.RenderTarget->GetSizeXY().Y > FamilySize.Y);
+	
+	const FIntPoint RenderTargetSize = (ViewFamily.RenderTarget->GetRenderTargetTexture().IsValid()) ? ViewFamily.RenderTarget->GetRenderTargetTexture()->GetSizeXY() : ViewFamily.RenderTarget->GetSizeXY();
+	const bool bRequiresUpscale = ((int32)RenderTargetSize.X > FamilySize.X || (int32)RenderTargetSize.Y > FamilySize.Y);
+
 	// ES2 requires that the back buffer and depth match dimensions.
 	// For the most part this is not the case when using scene captures. Thus scene captures always render to scene color target.
 	const bool bStereoRenderingAndHMD = View.Family->EngineShowFlags.StereoRendering && View.Family->EngineShowFlags.HMDDistortion;
