[Duplicating 4432262] - HLOD visibility state is now cleared when not in use.

diff --git a/Engine/Source/Runtime/Renderer/Private/ScenePrivate.h b/Engine/Source/Runtime/Renderer/Private/ScenePrivate.h
--- a/Engine/Source/Runtime/Renderer/Private/ScenePrivate.h
+++ b/Engine/Source/Runtime/Renderer/Private/ScenePrivate.h
@@ -1993,6 +1993,8 @@ public:
 	void UpdateNodeSceneInfo(FPrimitiveComponentId NodeId, FPrimitiveSceneInfo* SceneInfo);
 	void UpdateVisibilityStates(FViewInfo& View);
 
+	void ClearVisibilityState(FViewInfo& View);
+
 	bool IsActive() const { return (SceneNodes.Num() > 0); }
 
 private:
diff --git a/Engine/Source/Runtime/Renderer/Private/SceneVisibility.cpp b/Engine/Source/Runtime/Renderer/Private/SceneVisibility.cpp
--- a/Engine/Source/Runtime/Renderer/Private/SceneVisibility.cpp
+++ b/Engine/Source/Runtime/Renderer/Private/SceneVisibility.cpp
@@ -2963,6 +2963,10 @@ void FSceneRenderer::ComputeViewVisibility(FRHICommandListImmediate& RHICmdList)
 				QUICK_SCOPE_CYCLE_COUNTER(STAT_ViewVisibilityTime_HLODUpdate);
 				HLODTree.UpdateVisibilityStates(View);
 			}
+			else
+			{
+				HLODTree.ClearVisibilityState(View);
+			}
 
 			int32 NumCulledPrimitivesForView;
 			if (View.CustomVisibilityQuery && View.CustomVisibilityQuery->Prepare())
@@ -3643,6 +3647,36 @@ void FLODSceneTree::UpdateNodeSceneInfo(FPrimitiveComponentId NodeId, FPrimitive
 	}
 }
 
+void FLODSceneTree::ClearVisibilityState(FViewInfo& View)
+{
+	if (FSceneViewState* ViewState = (FSceneViewState*)View.State)
+	{
+#if !(UE_BUILD_SHIPPING || UE_BUILD_TEST)
+		// Skip update logic when frozen
+		if (ViewState->bIsFrozen)
+		{
+			return;
+		}
+#endif
+		FHLODVisibilityState& HLODState = ViewState->HLODVisibilityState;
+
+		if(HLODState.IsValidPrimitiveIndex(0))
+		{
+			HLODState.PrimitiveFadingLODMap.Empty(0);
+			HLODState.PrimitiveFadingOutLODMap.Empty(0);
+			HLODState.ForcedVisiblePrimitiveMap.Empty(0);
+			HLODState.ForcedHiddenPrimitiveMap.Empty(0);
+		}
+
+		TMap<FPrimitiveComponentId, FHLODSceneNodeVisibilityState>& VisibilityStates = ViewState->HLODSceneNodeVisibilityStates;
+
+		if(VisibilityStates.Num() > 0)
+		{
+			VisibilityStates.Empty(0);
+		}
+	}
+}
+
 void FLODSceneTree::UpdateVisibilityStates(FViewInfo& View)
 {
 	if (FSceneViewState* ViewState = (FSceneViewState*)View.State)
