UE-63963 - Fix FSystemTextures.InitializeTextures initialization ordering

diff --git a/Engine/Source/Runtime/Renderer/Private/SystemTextures.cpp b/Engine/Source/Runtime/Renderer/Private/SystemTextures.cpp
--- a/Engine/Source/Runtime/Renderer/Private/SystemTextures.cpp
+++ b/Engine/Source/Runtime/Renderer/Private/SystemTextures.cpp
@@ -17,502 +17,512 @@ SystemTextures
 /** The global render targets used for scene rendering. */
 TGlobalResource<FSystemTextures> GSystemTextures;
 
-void FSystemTextures::InternalInitializeTextures(FRHICommandListImmediate& RHICmdList, ERHIFeatureLevel::Type InFeatureLevel)
+void FSystemTextures::InitializeCommonTextures(FRHICommandListImmediate& RHICmdList)
 {
 	// First initialize textures that are common to all feature levels. This is always done the first time we come into this function, as doesn't care about the
 	// requested feature level
-	if (!bTexturesInitialized)
+
+	// Create a WhiteDummy texture
 	{
-		// Create a WhiteDummy texture
-		{
-			FPooledRenderTargetDesc Desc(FPooledRenderTargetDesc::Create2DDesc(FIntPoint(1, 1), PF_B8G8R8A8, FClearValueBinding::White, TexCreate_HideInVisualizeTexture, TexCreate_RenderTargetable | TexCreate_NoFastClear, false));
-			Desc.AutoWritable = false;
-			GRenderTargetPool.FindFreeElement(RHICmdList, Desc, WhiteDummy, TEXT("WhiteDummy"), true, ERenderTargetTransience::NonTransient);
+		FPooledRenderTargetDesc Desc(FPooledRenderTargetDesc::Create2DDesc(FIntPoint(1, 1), PF_B8G8R8A8, FClearValueBinding::White, TexCreate_HideInVisualizeTexture, TexCreate_RenderTargetable | TexCreate_NoFastClear, false));
+		Desc.AutoWritable = false;
+		GRenderTargetPool.FindFreeElement(RHICmdList, Desc, WhiteDummy, TEXT("WhiteDummy"), true, ERenderTargetTransience::NonTransient);
 
-			SetRenderTarget(RHICmdList, WhiteDummy->GetRenderTargetItem().TargetableTexture, FTextureRHIRef(), ESimpleRenderTargetMode::EClearColorExistingDepth);
-			RHICmdList.CopyToResolveTarget(WhiteDummy->GetRenderTargetItem().TargetableTexture, WhiteDummy->GetRenderTargetItem().ShaderResourceTexture, FResolveParams());
-		}
+		SetRenderTarget(RHICmdList, WhiteDummy->GetRenderTargetItem().TargetableTexture, FTextureRHIRef(), ESimpleRenderTargetMode::EClearColorExistingDepth);
+		RHICmdList.CopyToResolveTarget(WhiteDummy->GetRenderTargetItem().TargetableTexture, WhiteDummy->GetRenderTargetItem().ShaderResourceTexture, FResolveParams());
+	}
 
-		// Create a BlackDummy texture
-		{
-			FPooledRenderTargetDesc Desc(FPooledRenderTargetDesc::Create2DDesc(FIntPoint(1, 1), PF_B8G8R8A8, FClearValueBinding::Transparent, TexCreate_HideInVisualizeTexture, TexCreate_RenderTargetable | TexCreate_NoFastClear, false));
-			Desc.AutoWritable = false;
-			GRenderTargetPool.FindFreeElement(RHICmdList, Desc, BlackDummy, TEXT("BlackDummy"), true, ERenderTargetTransience::NonTransient);
+	// Create a BlackDummy texture
+	{
+		FPooledRenderTargetDesc Desc(FPooledRenderTargetDesc::Create2DDesc(FIntPoint(1, 1), PF_B8G8R8A8, FClearValueBinding::Transparent, TexCreate_HideInVisualizeTexture, TexCreate_RenderTargetable | TexCreate_NoFastClear, false));
+		Desc.AutoWritable = false;
+		GRenderTargetPool.FindFreeElement(RHICmdList, Desc, BlackDummy, TEXT("BlackDummy"), true, ERenderTargetTransience::NonTransient);
 
-			SetRenderTarget(RHICmdList, BlackDummy->GetRenderTargetItem().TargetableTexture, FTextureRHIRef(), ESimpleRenderTargetMode::EClearColorExistingDepth);
-			RHICmdList.CopyToResolveTarget(BlackDummy->GetRenderTargetItem().TargetableTexture, BlackDummy->GetRenderTargetItem().ShaderResourceTexture, FResolveParams());
-		}
+		SetRenderTarget(RHICmdList, BlackDummy->GetRenderTargetItem().TargetableTexture, FTextureRHIRef(), ESimpleRenderTargetMode::EClearColorExistingDepth);
+		RHICmdList.CopyToResolveTarget(BlackDummy->GetRenderTargetItem().TargetableTexture, BlackDummy->GetRenderTargetItem().ShaderResourceTexture, FResolveParams());
+	}
 
-		// Create a BlackAlphaOneDummy texture
-		{
-			FPooledRenderTargetDesc Desc(FPooledRenderTargetDesc::Create2DDesc(FIntPoint(1, 1), PF_B8G8R8A8, FClearValueBinding::Black, TexCreate_HideInVisualizeTexture, TexCreate_RenderTargetable | TexCreate_NoFastClear, false));
-			Desc.AutoWritable = false;
-			GRenderTargetPool.FindFreeElement(RHICmdList, Desc, BlackAlphaOneDummy, TEXT("BlackAlphaOneDummy"), true, ERenderTargetTransience::NonTransient);
+	// Create a BlackAlphaOneDummy texture
+	{
+		FPooledRenderTargetDesc Desc(FPooledRenderTargetDesc::Create2DDesc(FIntPoint(1, 1), PF_B8G8R8A8, FClearValueBinding::Black, TexCreate_HideInVisualizeTexture, TexCreate_RenderTargetable | TexCreate_NoFastClear, false));
+		Desc.AutoWritable = false;
+		GRenderTargetPool.FindFreeElement(RHICmdList, Desc, BlackAlphaOneDummy, TEXT("BlackAlphaOneDummy"), true, ERenderTargetTransience::NonTransient);
 
-			SetRenderTarget(RHICmdList, BlackAlphaOneDummy->GetRenderTargetItem().TargetableTexture, FTextureRHIRef(), ESimpleRenderTargetMode::EClearColorExistingDepth);
-			RHICmdList.CopyToResolveTarget(BlackAlphaOneDummy->GetRenderTargetItem().TargetableTexture, BlackAlphaOneDummy->GetRenderTargetItem().ShaderResourceTexture, FResolveParams());
-		}
+		SetRenderTarget(RHICmdList, BlackAlphaOneDummy->GetRenderTargetItem().TargetableTexture, FTextureRHIRef(), ESimpleRenderTargetMode::EClearColorExistingDepth);
+		RHICmdList.CopyToResolveTarget(BlackAlphaOneDummy->GetRenderTargetItem().TargetableTexture, BlackAlphaOneDummy->GetRenderTargetItem().ShaderResourceTexture, FResolveParams());
+	}
 
-		// Create a GreenDummy texture
-		{
-			FPooledRenderTargetDesc Desc(FPooledRenderTargetDesc::Create2DDesc(FIntPoint(1, 1), PF_B8G8R8A8, FClearValueBinding::Green, TexCreate_HideInVisualizeTexture, TexCreate_RenderTargetable | TexCreate_NoFastClear, false));
-			Desc.AutoWritable = false;
-			GRenderTargetPool.FindFreeElement(RHICmdList, Desc, GreenDummy, TEXT("GreenDummy"), true, ERenderTargetTransience::NonTransient);
+	// Create a GreenDummy texture
+	{
+		FPooledRenderTargetDesc Desc(FPooledRenderTargetDesc::Create2DDesc(FIntPoint(1, 1), PF_B8G8R8A8, FClearValueBinding::Green, TexCreate_HideInVisualizeTexture, TexCreate_RenderTargetable | TexCreate_NoFastClear, false));
+		Desc.AutoWritable = false;
+		GRenderTargetPool.FindFreeElement(RHICmdList, Desc, GreenDummy, TEXT("GreenDummy"), true, ERenderTargetTransience::NonTransient);
 
-			SetRenderTarget(RHICmdList, GreenDummy->GetRenderTargetItem().TargetableTexture, FTextureRHIRef(), ESimpleRenderTargetMode::EClearColorExistingDepth);
-			RHICmdList.CopyToResolveTarget(GreenDummy->GetRenderTargetItem().TargetableTexture, GreenDummy->GetRenderTargetItem().ShaderResourceTexture, FResolveParams());
-		}
+		SetRenderTarget(RHICmdList, GreenDummy->GetRenderTargetItem().TargetableTexture, FTextureRHIRef(), ESimpleRenderTargetMode::EClearColorExistingDepth);
+		RHICmdList.CopyToResolveTarget(GreenDummy->GetRenderTargetItem().TargetableTexture, GreenDummy->GetRenderTargetItem().ShaderResourceTexture, FResolveParams());
+	}
 
-		// Create a DefaultNormal8Bit texture
-		{
-			FPooledRenderTargetDesc Desc(FPooledRenderTargetDesc::Create2DDesc(FIntPoint(1, 1), PF_B8G8R8A8, FClearValueBinding::DefaultNormal8Bit, TexCreate_HideInVisualizeTexture, TexCreate_RenderTargetable | TexCreate_NoFastClear, false));
-			Desc.AutoWritable = false;
-			GRenderTargetPool.FindFreeElement(RHICmdList, Desc, DefaultNormal8Bit, TEXT("DefaultNormal8Bit"), true, ERenderTargetTransience::NonTransient);
+	// Create a DefaultNormal8Bit texture
+	{
+		FPooledRenderTargetDesc Desc(FPooledRenderTargetDesc::Create2DDesc(FIntPoint(1, 1), PF_B8G8R8A8, FClearValueBinding::DefaultNormal8Bit, TexCreate_HideInVisualizeTexture, TexCreate_RenderTargetable | TexCreate_NoFastClear, false));
+		Desc.AutoWritable = false;
+		GRenderTargetPool.FindFreeElement(RHICmdList, Desc, DefaultNormal8Bit, TEXT("DefaultNormal8Bit"), true, ERenderTargetTransience::NonTransient);
 
-			SetRenderTarget(RHICmdList, DefaultNormal8Bit->GetRenderTargetItem().TargetableTexture, FTextureRHIRef(), ESimpleRenderTargetMode::EClearColorExistingDepth);
-			RHICmdList.CopyToResolveTarget(DefaultNormal8Bit->GetRenderTargetItem().TargetableTexture, DefaultNormal8Bit->GetRenderTargetItem().ShaderResourceTexture, FResolveParams());
-		}
+		SetRenderTarget(RHICmdList, DefaultNormal8Bit->GetRenderTargetItem().TargetableTexture, FTextureRHIRef(), ESimpleRenderTargetMode::EClearColorExistingDepth);
+		RHICmdList.CopyToResolveTarget(DefaultNormal8Bit->GetRenderTargetItem().TargetableTexture, DefaultNormal8Bit->GetRenderTargetItem().ShaderResourceTexture, FResolveParams());
+	}
 
-		// Create the PerlinNoiseGradient texture
+	// Create the PerlinNoiseGradient texture
+	{
+		FPooledRenderTargetDesc Desc(FPooledRenderTargetDesc::Create2DDesc(FIntPoint(128, 128), PF_B8G8R8A8, FClearValueBinding::None, TexCreate_HideInVisualizeTexture, TexCreate_None | TexCreate_NoFastClear, false));
+		Desc.AutoWritable = false;
+		GRenderTargetPool.FindFreeElement(RHICmdList, Desc, PerlinNoiseGradient, TEXT("PerlinNoiseGradient"), true, ERenderTargetTransience::NonTransient);
+		// Write the contents of the texture.
+		uint32 DestStride;
+		uint8* DestBuffer = (uint8*)RHICmdList.LockTexture2D((FTexture2DRHIRef&)PerlinNoiseGradient->GetRenderTargetItem().ShaderResourceTexture, 0, RLM_WriteOnly, DestStride, false);
+		// seed the pseudo random stream with a good value
+		FRandomStream RandomStream(12345);
+		// Values represent float3 values in the -1..1 range.
+		// The vectors are the edge mid point of a cube from -1 .. 1
+		static uint32 gradtable[] =
 		{
-			FPooledRenderTargetDesc Desc(FPooledRenderTargetDesc::Create2DDesc(FIntPoint(128, 128), PF_B8G8R8A8, FClearValueBinding::None, TexCreate_HideInVisualizeTexture, TexCreate_None | TexCreate_NoFastClear, false));
-			Desc.AutoWritable = false;
-			GRenderTargetPool.FindFreeElement(RHICmdList, Desc, PerlinNoiseGradient, TEXT("PerlinNoiseGradient"), true, ERenderTargetTransience::NonTransient);
-			// Write the contents of the texture.
-			uint32 DestStride;
-			uint8* DestBuffer = (uint8*)RHICmdList.LockTexture2D((FTexture2DRHIRef&)PerlinNoiseGradient->GetRenderTargetItem().ShaderResourceTexture, 0, RLM_WriteOnly, DestStride, false);
-			// seed the pseudo random stream with a good value
-			FRandomStream RandomStream(12345);
-			// Values represent float3 values in the -1..1 range.
-			// The vectors are the edge mid point of a cube from -1 .. 1
-			static uint32 gradtable[] =
-			{
-				0x88ffff, 0xff88ff, 0xffff88,
-				0x88ff00, 0xff8800, 0xff0088,
-				0x8800ff, 0x0088ff, 0x00ff88,
-				0x880000, 0x008800, 0x000088,
-			};
-			for (int32 y = 0; y < Desc.Extent.Y; ++y)
+			0x88ffff, 0xff88ff, 0xffff88,
+			0x88ff00, 0xff8800, 0xff0088,
+			0x8800ff, 0x0088ff, 0x00ff88,
+			0x880000, 0x008800, 0x000088,
+		};
+		for (int32 y = 0; y < Desc.Extent.Y; ++y)
+		{
+			for (int32 x = 0; x < Desc.Extent.X; ++x)
 			{
-				for (int32 x = 0; x < Desc.Extent.X; ++x)
-				{
-					uint32* Dest = (uint32*)(DestBuffer + x * sizeof(uint32)+y * DestStride);
+				uint32* Dest = (uint32*)(DestBuffer + x * sizeof(uint32) + y * DestStride);
 
-					// pick a random direction (hacky way to overcome the quality issues FRandomStream has)
-					*Dest = gradtable[(uint32)(RandomStream.GetFraction() * 11.9999999f)];
-				}
+				// pick a random direction (hacky way to overcome the quality issues FRandomStream has)
+				*Dest = gradtable[(uint32)(RandomStream.GetFraction() * 11.9999999f)];
 			}
-			RHICmdList.UnlockTexture2D((FTexture2DRHIRef&)PerlinNoiseGradient->GetRenderTargetItem().ShaderResourceTexture, 0, false);
 		}
+		RHICmdList.UnlockTexture2D((FTexture2DRHIRef&)PerlinNoiseGradient->GetRenderTargetItem().ShaderResourceTexture, 0, false);
+	}
 
-		// Create the SobolSampling texture
-		if (InFeatureLevel >= ERHIFeatureLevel::ES3_1 && GPixelFormats[PF_R16_UINT].Supported)
-		{
-			FPooledRenderTargetDesc Desc(FPooledRenderTargetDesc::Create2DDesc(FIntPoint(32, 16), PF_R16_UINT, FClearValueBinding::None, TexCreate_HideInVisualizeTexture, TexCreate_NoFastClear, false));
-			Desc.AutoWritable = false;
-			GRenderTargetPool.FindFreeElement(RHICmdList, Desc, SobolSampling, TEXT("SobolSampling"));
-			// Write the contents of the texture.
-			uint32 DestStride;
-			uint8* DestBuffer = (uint8*)RHICmdList.LockTexture2D((FTexture2DRHIRef&)SobolSampling->GetRenderTargetItem().ShaderResourceTexture, 0, RLM_WriteOnly, DestStride, false);
-
-			uint16 Result, *Dest;
-			for (int y = 0; y < 16; ++y)
-			{
-				Dest = (uint16*)(DestBuffer + y * DestStride);
+	if (!GSupportsShaderFramebufferFetch && GPixelFormats[PF_FloatRGBA].Supported)
+	{
+		FPooledRenderTargetDesc Desc(FPooledRenderTargetDesc::Create2DDesc(FIntPoint(1, 1), PF_FloatRGBA, FClearValueBinding(FLinearColor(65000.0f, 65000.0f, 65000.0f, 65000.0f)), TexCreate_HideInVisualizeTexture, TexCreate_RenderTargetable | TexCreate_NoFastClear, false));
+		Desc.AutoWritable = false;
+		GRenderTargetPool.FindFreeElement(RHICmdList, Desc, MaxFP16Depth, TEXT("MaxFP16Depth"), true, ERenderTargetTransience::NonTransient);
 
-				// 16x16 block starting at 0,0 = Sobol X,Y from bottom 4 bits of cell X,Y
-				for (int x = 0; x < 16; ++x, ++Dest)
-				{
-					Result  = (x & 0x001) ? 0xf68e : 0;
-					Result ^= (x & 0x002) ? 0x8e56 : 0;
-					Result ^= (x & 0x004) ? 0x1135 : 0;
-					Result ^= (x & 0x008) ? 0x220a : 0;
-					Result ^= (y & 0x001) ? 0x94c4 : 0;
-					Result ^= (y & 0x002) ? 0x4ac2 : 0;
-					Result ^= (y & 0x004) ? 0xfb57 : 0;
-					Result ^= (y & 0x008) ? 0x0454 : 0;
-					*Dest = Result;
-				}
+		FRHIRenderTargetView RtView = FRHIRenderTargetView(MaxFP16Depth->GetRenderTargetItem().TargetableTexture, ERenderTargetLoadAction::EClear);
+		FRHISetRenderTargetsInfo Info(1, &RtView, FRHIDepthRenderTargetView());
+		RHICmdList.SetRenderTargetsAndClear(Info);
+		RHICmdList.CopyToResolveTarget(MaxFP16Depth->GetRenderTargetItem().TargetableTexture, MaxFP16Depth->GetRenderTargetItem().ShaderResourceTexture, FResolveParams());
+	}
 
-				// 16x16 block starting at 16,0 = Sobol X,Y from 2nd 4 bits of cell X,Y
-				for (int x = 0; x < 16; ++x, ++Dest)
-				{
-					Result  = (x & 0x010) ? 0x4414 : 0;
-					Result ^= (x & 0x020) ? 0x8828 : 0;
-					Result ^= (x & 0x040) ? 0xe69e : 0;
-					Result ^= (x & 0x080) ? 0xae76 : 0;
-					Result ^= (y & 0x010) ? 0xa28a : 0;
-					Result ^= (y & 0x020) ? 0x265e : 0;
-					Result ^= (y & 0x040) ? 0xe69e : 0;
-					Result ^= (y & 0x080) ? 0xae76 : 0;
-					*Dest = Result;
-				}
-			}
-			RHICmdList.UnlockTexture2D((FTexture2DRHIRef&)SobolSampling->GetRenderTargetItem().ShaderResourceTexture, 0, false);
-		}
+	// Create dummy 1x1 depth texture		
+	{
+		FPooledRenderTargetDesc Desc(FPooledRenderTargetDesc::Create2DDesc(FIntPoint(1, 1), PF_DepthStencil, FClearValueBinding::DepthFar, TexCreate_None, TexCreate_DepthStencilTargetable, false));
+		Desc.AutoWritable = false;
+		GRenderTargetPool.FindFreeElement(RHICmdList, Desc, DepthDummy, TEXT("DepthDummy"), true, ERenderTargetTransience::NonTransient);
 
-		if (!GSupportsShaderFramebufferFetch && GPixelFormats[PF_FloatRGBA].Supported)
-		{
-			FPooledRenderTargetDesc Desc(FPooledRenderTargetDesc::Create2DDesc(FIntPoint(1, 1), PF_FloatRGBA, FClearValueBinding(FLinearColor(65000.0f, 65000.0f, 65000.0f, 65000.0f)), TexCreate_HideInVisualizeTexture, TexCreate_RenderTargetable | TexCreate_NoFastClear, false));
-			Desc.AutoWritable = false;
-			GRenderTargetPool.FindFreeElement(RHICmdList, Desc, MaxFP16Depth, TEXT("MaxFP16Depth"), true, ERenderTargetTransience::NonTransient);
+		FRHISetRenderTargetsInfo Info(0, nullptr, FRHIDepthRenderTargetView(DepthDummy->GetRenderTargetItem().TargetableTexture, ERenderTargetLoadAction::EClear, ERenderTargetStoreAction::EStore));
+		RHICmdList.SetRenderTargetsAndClear(Info);
+		RHICmdList.CopyToResolveTarget(DepthDummy->GetRenderTargetItem().TargetableTexture, DepthDummy->GetRenderTargetItem().ShaderResourceTexture, FResolveParams());
+	}
+}
 
-			FRHIRenderTargetView RtView = FRHIRenderTargetView(MaxFP16Depth->GetRenderTargetItem().TargetableTexture, ERenderTargetLoadAction::EClear);
-			FRHISetRenderTargetsInfo Info(1, &RtView, FRHIDepthRenderTargetView());
-			RHICmdList.SetRenderTargetsAndClear(Info);
-			RHICmdList.CopyToResolveTarget(MaxFP16Depth->GetRenderTargetItem().TargetableTexture, MaxFP16Depth->GetRenderTargetItem().ShaderResourceTexture, FResolveParams());
-		}
+void FSystemTextures::InitializeFeatureLevelDependentTextures(FRHICommandListImmediate& RHICmdList, const ERHIFeatureLevel::Type InFeatureLevel)
+{
+	// this function will be called every time the feature level will be updated and some textures require a minimum feature level to exist
+	// the below declared variable (CurrentFeatureLevel) will guard against reinitialization of those textures already created in a previous call
+	// if FeatureLevelInitializedTo has its default value (ERHIFeatureLevel::Num) it means that setup was never performed and all textures are invalid
+	// thus CurrentFeatureLevel will be set to ERHIFeatureLevel::ES2 to validate all 'is valid' branching conditions below
+	// this should change if we add any ERHIFeatureLevel::ES2 dependent textures here!
+	ERHIFeatureLevel::Type CurrentFeatureLevel = FeatureLevelInitializedTo == ERHIFeatureLevel::Num ? ERHIFeatureLevel::ES2 : FeatureLevelInitializedTo;
+
+	// Create the SobolSampling texture
+	if (CurrentFeatureLevel < ERHIFeatureLevel::ES3_1 && InFeatureLevel >= ERHIFeatureLevel::ES3_1 && GPixelFormats[PF_R16_UINT].Supported)
+	{
+		FPooledRenderTargetDesc Desc(FPooledRenderTargetDesc::Create2DDesc(FIntPoint(32, 16), PF_R16_UINT, FClearValueBinding::None, TexCreate_HideInVisualizeTexture, TexCreate_NoFastClear, false));
+		Desc.AutoWritable = false;
+		GRenderTargetPool.FindFreeElement(RHICmdList, Desc, SobolSampling, TEXT("SobolSampling"));
+		// Write the contents of the texture.
+		uint32 DestStride;
+		uint8* DestBuffer = (uint8*)RHICmdList.LockTexture2D((FTexture2DRHIRef&)SobolSampling->GetRenderTargetItem().ShaderResourceTexture, 0, RLM_WriteOnly, DestStride, false);
 
-		// Create dummy 1x1 depth texture		
+		uint16 Result, *Dest;
+		for (int y = 0; y < 16; ++y)
 		{
-			FPooledRenderTargetDesc Desc(FPooledRenderTargetDesc::Create2DDesc(FIntPoint(1, 1), PF_DepthStencil, FClearValueBinding::DepthFar, TexCreate_None, TexCreate_DepthStencilTargetable, false));
-			Desc.AutoWritable = false;
-			GRenderTargetPool.FindFreeElement(RHICmdList, Desc, DepthDummy, TEXT("DepthDummy"), true, ERenderTargetTransience::NonTransient);
+			Dest = (uint16*)(DestBuffer + y * DestStride);
 
-			FRHISetRenderTargetsInfo Info(0, nullptr, FRHIDepthRenderTargetView(DepthDummy->GetRenderTargetItem().TargetableTexture, ERenderTargetLoadAction::EClear, ERenderTargetStoreAction::EStore));
-			RHICmdList.SetRenderTargetsAndClear(Info);
-			RHICmdList.CopyToResolveTarget(DepthDummy->GetRenderTargetItem().TargetableTexture, DepthDummy->GetRenderTargetItem().ShaderResourceTexture, FResolveParams());
+			// 16x16 block starting at 0,0 = Sobol X,Y from bottom 4 bits of cell X,Y
+			for (int x = 0; x < 16; ++x, ++Dest)
+			{
+				Result = (x & 0x001) ? 0xf68e : 0;
+				Result ^= (x & 0x002) ? 0x8e56 : 0;
+				Result ^= (x & 0x004) ? 0x1135 : 0;
+				Result ^= (x & 0x008) ? 0x220a : 0;
+				Result ^= (y & 0x001) ? 0x94c4 : 0;
+				Result ^= (y & 0x002) ? 0x4ac2 : 0;
+				Result ^= (y & 0x004) ? 0xfb57 : 0;
+				Result ^= (y & 0x008) ? 0x0454 : 0;
+				*Dest = Result;
+			}
+
+			// 16x16 block starting at 16,0 = Sobol X,Y from 2nd 4 bits of cell X,Y
+			for (int x = 0; x < 16; ++x, ++Dest)
+			{
+				Result = (x & 0x010) ? 0x4414 : 0;
+				Result ^= (x & 0x020) ? 0x8828 : 0;
+				Result ^= (x & 0x040) ? 0xe69e : 0;
+				Result ^= (x & 0x080) ? 0xae76 : 0;
+				Result ^= (y & 0x010) ? 0xa28a : 0;
+				Result ^= (y & 0x020) ? 0x265e : 0;
+				Result ^= (y & 0x040) ? 0xe69e : 0;
+				Result ^= (y & 0x080) ? 0xae76 : 0;
+				*Dest = Result;
+			}
 		}
+		RHICmdList.UnlockTexture2D((FTexture2DRHIRef&)SobolSampling->GetRenderTargetItem().ShaderResourceTexture, 0, false);
 	}
 
-	// Create the PerlinNoise3D texture (similar to http://prettyprocs.wordpress.com/2012/10/20/fast-perlin-noise/)
-	if (InFeatureLevel >= ERHIFeatureLevel::SM4)
+	if (CurrentFeatureLevel < ERHIFeatureLevel::SM4 && InFeatureLevel >= ERHIFeatureLevel::SM4)
 	{
-		uint32 Extent = 16;
+		// Create the PerlinNoise3D texture (similar to http://prettyprocs.wordpress.com/2012/10/20/fast-perlin-noise/)
+		{
+			uint32 Extent = 16;
 
-		const uint32 Square = Extent * Extent;
+			const uint32 Square = Extent * Extent;
 
-		FPooledRenderTargetDesc Desc(FPooledRenderTargetDesc::CreateVolumeDesc(Extent, Extent, Extent, PF_B8G8R8A8, FClearValueBinding::None, TexCreate_ShaderResource | TexCreate_HideInVisualizeTexture | TexCreate_NoTiling, TexCreate_None, false));
-		Desc.AutoWritable = false;
-		GRenderTargetPool.FindFreeElement(RHICmdList, Desc, PerlinNoise3D, TEXT("PerlinNoise3D"), true, ERenderTargetTransience::NonTransient);
-		// Write the contents of the texture.
-		TArray<uint32> DestBuffer;
+			FPooledRenderTargetDesc Desc(FPooledRenderTargetDesc::CreateVolumeDesc(Extent, Extent, Extent, PF_B8G8R8A8, FClearValueBinding::None, TexCreate_ShaderResource | TexCreate_HideInVisualizeTexture | TexCreate_NoTiling, TexCreate_None, false));
+			Desc.AutoWritable = false;
+			GRenderTargetPool.FindFreeElement(RHICmdList, Desc, PerlinNoise3D, TEXT("PerlinNoise3D"), true, ERenderTargetTransience::NonTransient);
+			// Write the contents of the texture.
+			TArray<uint32> DestBuffer;
 
-		DestBuffer.AddZeroed(Extent * Extent * Extent);
-		// seed the pseudo random stream with a good value
-		FRandomStream RandomStream(0x1234);
-		// Values represent float3 values in the -1..1 range.
-		// The vectors are the edge mid point of a cube from -1 .. 1
-		// -1:0 0:7f 1:fe, can be reconstructed with * 512/254 - 1
-		// * 2 - 1 cannot be used because 0 would not be mapped
-		static uint32 gradtable[] = 
-		{
-			0x7ffefe, 0xfe7ffe, 0xfefe7f,
-			0x7ffe00, 0xfe7f00, 0xfe007f,
-			0x7f00fe, 0x007ffe, 0x00fe7f,
-			0x7f0000, 0x007f00, 0x00007f,
-		};
-		// set random directions
-		{
-			for(uint32 z = 0; z < Extent - 1; ++z)
+			DestBuffer.AddZeroed(Extent * Extent * Extent);
+			// seed the pseudo random stream with a good value
+			FRandomStream RandomStream(0x1234);
+			// Values represent float3 values in the -1..1 range.
+			// The vectors are the edge mid point of a cube from -1 .. 1
+			// -1:0 0:7f 1:fe, can be reconstructed with * 512/254 - 1
+			// * 2 - 1 cannot be used because 0 would not be mapped
+			static uint32 gradtable[] =
 			{
-				for(uint32 y = 0; y < Extent - 1; ++y)
+				0x7ffefe, 0xfe7ffe, 0xfefe7f,
+				0x7ffe00, 0xfe7f00, 0xfe007f,
+				0x7f00fe, 0x007ffe, 0x00fe7f,
+				0x7f0000, 0x007f00, 0x00007f,
+			};
+			// set random directions
+			{
+				for (uint32 z = 0; z < Extent - 1; ++z)
 				{
-					for(uint32 x = 0; x < Extent - 1; ++x)
+					for (uint32 y = 0; y < Extent - 1; ++y)
 					{
-						uint32& Value = DestBuffer[x + y * Extent + z * Square];
+						for (uint32 x = 0; x < Extent - 1; ++x)
+						{
+							uint32& Value = DestBuffer[x + y * Extent + z * Square];
 
-						// pick a random direction (hacky way to overcome the quality issues FRandomStream has)
-						Value = gradtable[(uint32)(RandomStream.GetFraction() * 11.9999999f)];
+							// pick a random direction (hacky way to overcome the quality issues FRandomStream has)
+							Value = gradtable[(uint32)(RandomStream.GetFraction() * 11.9999999f)];
+						}
 					}
 				}
 			}
-		}
-		// replicate a border for filtering
-		{
-			uint32 Last = Extent - 1;
-
-			for(uint32 z = 0; z < Extent; ++z)
+			// replicate a border for filtering
 			{
-				for(uint32 y = 0; y < Extent; ++y)
+				uint32 Last = Extent - 1;
+
+				for (uint32 z = 0; z < Extent; ++z)
 				{
-					DestBuffer[Last + y * Extent + z * Square] = DestBuffer[0 + y * Extent + z * Square];
+					for (uint32 y = 0; y < Extent; ++y)
+					{
+						DestBuffer[Last + y * Extent + z * Square] = DestBuffer[0 + y * Extent + z * Square];
+					}
 				}
-			}
-			for(uint32 z = 0; z < Extent; ++z)
-			{
-				for(uint32 x = 0; x < Extent; ++x)
+				for (uint32 z = 0; z < Extent; ++z)
 				{
-					DestBuffer[x + Last * Extent + z * Square] = DestBuffer[x + 0 * Extent + z * Square];
+					for (uint32 x = 0; x < Extent; ++x)
+					{
+						DestBuffer[x + Last * Extent + z * Square] = DestBuffer[x + 0 * Extent + z * Square];
+					}
 				}
-			}
-			for(uint32 y = 0; y < Extent; ++y)
-			{
-				for(uint32 x = 0; x < Extent; ++x)
+				for (uint32 y = 0; y < Extent; ++y)
 				{
-					DestBuffer[x + y * Extent + Last * Square] = DestBuffer[x + y * Extent + 0 * Square];
+					for (uint32 x = 0; x < Extent; ++x)
+					{
+						DestBuffer[x + y * Extent + Last * Square] = DestBuffer[x + y * Extent + 0 * Square];
+					}
 				}
 			}
-		}
-		// precompute gradients
-		{
-			uint32* Dest = DestBuffer.GetData();
-
-			for(uint32 z = 0; z < Desc.Depth; ++z)
+			// precompute gradients
 			{
-				for(uint32 y = 0; y < (uint32)Desc.Extent.Y; ++y)
+				uint32* Dest = DestBuffer.GetData();
+
+				for (uint32 z = 0; z < Desc.Depth; ++z)
 				{
-					for(uint32 x = 0; x < (uint32)Desc.Extent.X; ++x)
+					for (uint32 y = 0; y < (uint32)Desc.Extent.Y; ++y)
 					{
-						uint32 Value = *Dest;
+						for (uint32 x = 0; x < (uint32)Desc.Extent.X; ++x)
+						{
+							uint32 Value = *Dest;
 
-						// todo: check if rgb order is correct
-						int32 r = Value >> 16;
-						int32 g = (Value >> 8) & 0xff;
-						int32 b = Value & 0xff;
+							// todo: check if rgb order is correct
+							int32 r = Value >> 16;
+							int32 g = (Value >> 8) & 0xff;
+							int32 b = Value & 0xff;
 
-						int nx = (r / 0x7f) - 1;
-						int ny = (g / 0x7f) - 1;
-						int nz = (b / 0x7f) - 1;
+							int nx = (r / 0x7f) - 1;
+							int ny = (g / 0x7f) - 1;
+							int nz = (b / 0x7f) - 1;
 
-						int32 d = nx * x + ny * y + nz * z;
+							int32 d = nx * x + ny * y + nz * z;
 
-						// compress in 8bit
-						uint32 a = d + 127;
+							// compress in 8bit
+							uint32 a = d + 127;
 
-						*Dest++ = Value | (a << 24);
+							*Dest++ = Value | (a << 24);
+						}
 					}
 				}
 			}
-		}
-
-		FUpdateTextureRegion3D Region(0, 0, 0, 0, 0, 0, Desc.Extent.X, Desc.Extent.Y, Desc.Depth);
 
-		RHICmdList.UpdateTexture3D(
-			(FTexture3DRHIRef&)PerlinNoise3D->GetRenderTargetItem().ShaderResourceTexture,
-			0,
-			Region,
-			Desc.Extent.X * sizeof(uint32),
-			Desc.Extent.X * Desc.Extent.Y * sizeof(uint32),
-			(const uint8*)DestBuffer.GetData());
-	}
-
-	// Create the SSAO randomization texture
-	if (InFeatureLevel >= ERHIFeatureLevel::SM4)
-	{
-		float g_AngleOff1 = 127;
-		float g_AngleOff2 = 198;
-		float g_AngleOff3 = 23;
+			FUpdateTextureRegion3D Region(0, 0, 0, 0, 0, 0, Desc.Extent.X, Desc.Extent.Y, Desc.Depth);
 
-		FColor Bases[16];
+			RHICmdList.UpdateTexture3D(
+				(FTexture3DRHIRef&)PerlinNoise3D->GetRenderTargetItem().ShaderResourceTexture,
+				0,
+				Region,
+				Desc.Extent.X * sizeof(uint32),
+				Desc.Extent.X * Desc.Extent.Y * sizeof(uint32),
+				(const uint8*)DestBuffer.GetData());
+		} // end Create the PerlinNoise3D texture
 
-		for(int32 Pos = 0; Pos < 16; ++Pos)
+		// Create the SSAO randomization texture
 		{
-			// distribute rotations over 4x4 pattern
-//			int32 Reorder[16] = { 0, 8, 2, 10, 12, 6, 14, 4, 3, 11, 1, 9, 15, 5, 13, 7 };
-			int32 Reorder[16] = { 0, 11, 7, 3, 10, 4, 15, 12, 6, 8, 1, 14, 13, 2, 9, 5 };
-			int32 w = Reorder[Pos];
+			{
+				float g_AngleOff1 = 127;
+				float g_AngleOff2 = 198;
+				float g_AngleOff3 = 23;
 
-			// ordered sampling of the rotation basis (*2 is missing as we use mirrored samples)
-			float ww = w / 16.0f * PI;
+				FColor Bases[16];
 
-			// randomize base scale
-			float lenm = 1.0f - (FMath::Sin(g_AngleOff2 * w * 0.01f) * 0.5f + 0.5f) * g_AngleOff3 * 0.01f;
-			float s = FMath::Sin(ww) * lenm;
-			float c = FMath::Cos(ww) * lenm;
+				for (int32 Pos = 0; Pos < 16; ++Pos)
+				{
+					// distribute rotations over 4x4 pattern
+					//			int32 Reorder[16] = { 0, 8, 2, 10, 12, 6, 14, 4, 3, 11, 1, 9, 15, 5, 13, 7 };
+					int32 Reorder[16] = { 0, 11, 7, 3, 10, 4, 15, 12, 6, 8, 1, 14, 13, 2, 9, 5 };
+					int32 w = Reorder[Pos];
 
-			Bases[Pos] = FColor(FMath::Quantize8SignedByte(c), FMath::Quantize8SignedByte(s), 0, 0);
-		}
+					// ordered sampling of the rotation basis (*2 is missing as we use mirrored samples)
+					float ww = w / 16.0f * PI;
 
-		{
-			// could be PF_V8U8 to save shader instructions but that doesn't work on all hardware
-			FPooledRenderTargetDesc Desc(FPooledRenderTargetDesc::Create2DDesc(FIntPoint(64, 64), PF_R8G8, FClearValueBinding::None, TexCreate_HideInVisualizeTexture, TexCreate_None | TexCreate_NoFastClear, false));
-			Desc.AutoWritable = false;
-			GRenderTargetPool.FindFreeElement(RHICmdList, Desc, SSAORandomization, TEXT("SSAORandomization"), true, ERenderTargetTransience::NonTransient);
-			// Write the contents of the texture.
-			uint32 DestStride;
-			uint8* DestBuffer = (uint8*)RHICmdList.LockTexture2D((FTexture2DRHIRef&)SSAORandomization->GetRenderTargetItem().ShaderResourceTexture, 0, RLM_WriteOnly, DestStride, false);
+					// randomize base scale
+					float lenm = 1.0f - (FMath::Sin(g_AngleOff2 * w * 0.01f) * 0.5f + 0.5f) * g_AngleOff3 * 0.01f;
+					float s = FMath::Sin(ww) * lenm;
+					float c = FMath::Cos(ww) * lenm;
+
+					Bases[Pos] = FColor(FMath::Quantize8SignedByte(c), FMath::Quantize8SignedByte(s), 0, 0);
+				}
 
-			for(int32 y = 0; y < Desc.Extent.Y; ++y)
-			{
-				for(int32 x = 0; x < Desc.Extent.X; ++x)
 				{
-					uint8* Dest = (uint8*)(DestBuffer + x * sizeof(uint16) + y * DestStride);
+					// could be PF_V8U8 to save shader instructions but that doesn't work on all hardware
+					FPooledRenderTargetDesc Desc(FPooledRenderTargetDesc::Create2DDesc(FIntPoint(64, 64), PF_R8G8, FClearValueBinding::None, TexCreate_HideInVisualizeTexture, TexCreate_None | TexCreate_NoFastClear, false));
+					Desc.AutoWritable = false;
+					GRenderTargetPool.FindFreeElement(RHICmdList, Desc, SSAORandomization, TEXT("SSAORandomization"), true, ERenderTargetTransience::NonTransient);
+					// Write the contents of the texture.
+					uint32 DestStride;
+					uint8* DestBuffer = (uint8*)RHICmdList.LockTexture2D((FTexture2DRHIRef&)SSAORandomization->GetRenderTargetItem().ShaderResourceTexture, 0, RLM_WriteOnly, DestStride, false);
+
+					for (int32 y = 0; y < Desc.Extent.Y; ++y)
+					{
+						for (int32 x = 0; x < Desc.Extent.X; ++x)
+						{
+							uint8* Dest = (uint8*)(DestBuffer + x * sizeof(uint16) + y * DestStride);
 
-					uint32 Index = (x % 4) + (y % 4) * 4; 
+							uint32 Index = (x % 4) + (y % 4) * 4;
 
-					Dest[0] = Bases[Index].R;
-					Dest[1] = Bases[Index].G;
+							Dest[0] = Bases[Index].R;
+							Dest[1] = Bases[Index].G;
+						}
+					}
 				}
+				RHICmdList.UnlockTexture2D((FTexture2DRHIRef&)SSAORandomization->GetRenderTargetItem().ShaderResourceTexture, 0, false);
 			}
-		}
-		RHICmdList.UnlockTexture2D((FTexture2DRHIRef&)SSAORandomization->GetRenderTargetItem().ShaderResourceTexture, 0, false);
 
-		{
-			// for testing, with 128x128 R8G8 we are very close to the reference (if lower res is needed we might have to add an offset to counter the 0.5f texel shift)
-			const bool bReference = false;
-
-			EPixelFormat Format = PF_R8G8;
-			// for low roughness we would get banding with PF_R8G8 but for low spec it could be used, for now we don't do this optimization
-			if (GPixelFormats[PF_G16R16].Supported)
 			{
-				Format = PF_G16R16;
-			}
+				// for testing, with 128x128 R8G8 we are very close to the reference (if lower res is needed we might have to add an offset to counter the 0.5f texel shift)
+				const bool bReference = false;
 
-			FPooledRenderTargetDesc Desc(FPooledRenderTargetDesc::Create2DDesc(FIntPoint(128, 32), Format, FClearValueBinding::None, 0, TexCreate_None, false));
-			Desc.AutoWritable = false;
-			if (bReference)
-			{
-				Desc.Extent.X = 128;
-				Desc.Extent.Y = 128;
-			}
+				EPixelFormat Format = PF_R8G8;
+				// for low roughness we would get banding with PF_R8G8 but for low spec it could be used, for now we don't do this optimization
+				if (GPixelFormats[PF_G16R16].Supported)
+				{
+					Format = PF_G16R16;
+				}
 
-			GRenderTargetPool.FindFreeElement(RHICmdList, Desc, PreintegratedGF, TEXT("PreintegratedGF"), true, ERenderTargetTransience::NonTransient);
-			// Write the contents of the texture.
-			uint32 DestStride;
-			uint8* DestBuffer = (uint8*)RHICmdList.LockTexture2D((FTexture2DRHIRef&)PreintegratedGF->GetRenderTargetItem().ShaderResourceTexture, 0, RLM_WriteOnly, DestStride, false);
+				FPooledRenderTargetDesc Desc(FPooledRenderTargetDesc::Create2DDesc(FIntPoint(128, 32), Format, FClearValueBinding::None, 0, TexCreate_None, false));
+				Desc.AutoWritable = false;
+				if (bReference)
+				{
+					Desc.Extent.X = 128;
+					Desc.Extent.Y = 128;
+				}
 
-			// x is NoV, y is roughness
-			for (int32 y = 0; y < Desc.Extent.Y; y++)
-			{
-				float Roughness = (float)(y + 0.5f) / Desc.Extent.Y;
-				float m = Roughness * Roughness;
-				float m2 = m*m;
+				GRenderTargetPool.FindFreeElement(RHICmdList, Desc, PreintegratedGF, TEXT("PreintegratedGF"), true, ERenderTargetTransience::NonTransient);
+				// Write the contents of the texture.
+				uint32 DestStride;
+				uint8* DestBuffer = (uint8*)RHICmdList.LockTexture2D((FTexture2DRHIRef&)PreintegratedGF->GetRenderTargetItem().ShaderResourceTexture, 0, RLM_WriteOnly, DestStride, false);
 
-				for (int32 x = 0; x < Desc.Extent.X; x++)
+				// x is NoV, y is roughness
+				for (int32 y = 0; y < Desc.Extent.Y; y++)
 				{
-					float NoV = (float)(x + 0.5f) / Desc.Extent.X;
+					float Roughness = (float)(y + 0.5f) / Desc.Extent.Y;
+					float m = Roughness * Roughness;
+					float m2 = m * m;
 
-					FVector V;
-					V.X = FMath::Sqrt(1.0f - NoV * NoV);	// sin
-					V.Y = 0.0f;
-					V.Z = NoV;								// cos
+					for (int32 x = 0; x < Desc.Extent.X; x++)
+					{
+						float NoV = (float)(x + 0.5f) / Desc.Extent.X;
 
-					float A = 0.0f;
-					float B = 0.0f;
-					float C = 0.0f;
+						FVector V;
+						V.X = FMath::Sqrt(1.0f - NoV * NoV);	// sin
+						V.Y = 0.0f;
+						V.Z = NoV;								// cos
 
-					const uint32 NumSamples = 128;
-					for (uint32 i = 0; i < NumSamples; i++)
-					{
-						float E1 = (float)i / NumSamples;
-						float E2 = (double)ReverseBits(i) / (double)0x100000000LL;
+						float A = 0.0f;
+						float B = 0.0f;
+						float C = 0.0f;
 
+						const uint32 NumSamples = 128;
+						for (uint32 i = 0; i < NumSamples; i++)
 						{
-							float Phi = 2.0f * PI * E1;
-							float CosPhi = FMath::Cos(Phi);
-							float SinPhi = FMath::Sin(Phi);
-							float CosTheta = FMath::Sqrt((1.0f - E2) / (1.0f + (m2 - 1.0f) * E2));
-							float SinTheta = FMath::Sqrt(1.0f - CosTheta * CosTheta);
-
-							FVector H(SinTheta * FMath::Cos(Phi), SinTheta * FMath::Sin(Phi), CosTheta);
-							FVector L = 2.0f * (V | H) * H - V;
+							float E1 = (float)i / NumSamples;
+							float E2 = (double)ReverseBits(i) / (double)0x100000000LL;
 
-							float NoL = FMath::Max(L.Z, 0.0f);
-							float NoH = FMath::Max(H.Z, 0.0f);
-							float VoH = FMath::Max(V | H, 0.0f);
+							{
+								float Phi = 2.0f * PI * E1;
+								float CosPhi = FMath::Cos(Phi);
+								float SinPhi = FMath::Sin(Phi);
+								float CosTheta = FMath::Sqrt((1.0f - E2) / (1.0f + (m2 - 1.0f) * E2));
+								float SinTheta = FMath::Sqrt(1.0f - CosTheta * CosTheta);
+
+								FVector H(SinTheta * FMath::Cos(Phi), SinTheta * FMath::Sin(Phi), CosTheta);
+								FVector L = 2.0f * (V | H) * H - V;
+
+								float NoL = FMath::Max(L.Z, 0.0f);
+								float NoH = FMath::Max(H.Z, 0.0f);
+								float VoH = FMath::Max(V | H, 0.0f);
+
+								if (NoL > 0.0f)
+								{
+									float Vis_SmithV = NoL * (NoV * (1 - m) + m);
+									float Vis_SmithL = NoV * (NoL * (1 - m) + m);
+									float Vis = 0.5f / (Vis_SmithV + Vis_SmithL);
+
+									float NoL_Vis_PDF = NoL * Vis * (4.0f * VoH / NoH);
+									float Fc = 1.0f - VoH;
+									Fc *= FMath::Square(Fc*Fc);
+									A += NoL_Vis_PDF * (1.0f - Fc);
+									B += NoL_Vis_PDF * Fc;
+								}
+							}
 
-							if (NoL > 0.0f)
 							{
-								float Vis_SmithV = NoL * ( NoV * ( 1 - m ) + m );
-								float Vis_SmithL = NoV * ( NoL * ( 1 - m ) + m );
-								float Vis = 0.5f / ( Vis_SmithV + Vis_SmithL );
-
-								float NoL_Vis_PDF = NoL * Vis * (4.0f * VoH / NoH);
-								float Fc = 1.0f - VoH;
-								Fc *= FMath::Square(Fc*Fc);
-								A += NoL_Vis_PDF * (1.0f - Fc);
-								B += NoL_Vis_PDF * Fc;
+								float Phi = 2.0f * PI * E1;
+								float CosPhi = FMath::Cos(Phi);
+								float SinPhi = FMath::Sin(Phi);
+								float CosTheta = FMath::Sqrt(E2);
+								float SinTheta = FMath::Sqrt(1.0f - CosTheta * CosTheta);
+
+								FVector L(SinTheta * FMath::Cos(Phi), SinTheta * FMath::Sin(Phi), CosTheta);
+								FVector H = (V + L).GetUnsafeNormal();
+
+								float NoL = FMath::Max(L.Z, 0.0f);
+								float NoH = FMath::Max(H.Z, 0.0f);
+								float VoH = FMath::Max(V | H, 0.0f);
+
+								float FD90 = 0.5f + 2.0f * VoH * VoH * Roughness;
+								float FdV = 1.0f + (FD90 - 1.0f) * pow(1.0f - NoV, 5);
+								float FdL = 1.0f + (FD90 - 1.0f) * pow(1.0f - NoL, 5);
+								C += FdV * FdL;// * ( 1.0f - 0.3333f * Roughness );
 							}
 						}
+						A /= NumSamples;
+						B /= NumSamples;
+						C /= NumSamples;
 
+						if (Desc.Format == PF_A16B16G16R16)
 						{
-							float Phi = 2.0f * PI * E1;
-							float CosPhi = FMath::Cos(Phi);
-							float SinPhi = FMath::Sin(Phi);
-							float CosTheta = FMath::Sqrt(E2);
-							float SinTheta = FMath::Sqrt(1.0f - CosTheta * CosTheta);
-
-							FVector L(SinTheta * FMath::Cos(Phi), SinTheta * FMath::Sin(Phi), CosTheta);
-							FVector H = (V + L).GetUnsafeNormal();
-
-							float NoL = FMath::Max(L.Z, 0.0f);
-							float NoH = FMath::Max(H.Z, 0.0f);
-							float VoH = FMath::Max(V | H, 0.0f);
-
-							float FD90 = 0.5f + 2.0f * VoH * VoH * Roughness;
-							float FdV = 1.0f + (FD90 - 1.0f) * pow( 1.0f - NoV, 5 );
-							float FdL = 1.0f + (FD90 - 1.0f) * pow( 1.0f - NoL, 5 );
-							C += FdV * FdL;// * ( 1.0f - 0.3333f * Roughness );
+							uint16* Dest = (uint16*)(DestBuffer + x * 8 + y * DestStride);
+							Dest[0] = (int32)(FMath::Clamp(A, 0.0f, 1.0f) * 65535.0f + 0.5f);
+							Dest[1] = (int32)(FMath::Clamp(B, 0.0f, 1.0f) * 65535.0f + 0.5f);
+							Dest[2] = (int32)(FMath::Clamp(C, 0.0f, 1.0f) * 65535.0f + 0.5f);
 						}
-					}
-					A /= NumSamples;
-					B /= NumSamples;
-					C /= NumSamples;
-
-					if( Desc.Format == PF_A16B16G16R16 )
-					{
-						uint16* Dest = (uint16*)(DestBuffer + x * 8 + y * DestStride);
-						Dest[0] = (int32)(FMath::Clamp(A, 0.0f, 1.0f) * 65535.0f + 0.5f);
-						Dest[1] = (int32)(FMath::Clamp(B, 0.0f, 1.0f) * 65535.0f + 0.5f);
-						Dest[2] = (int32)(FMath::Clamp(C, 0.0f, 1.0f) * 65535.0f + 0.5f);
-					}
-					else if (Desc.Format == PF_G16R16)
-					{
-						uint16* Dest = (uint16*)(DestBuffer + x * 4 + y * DestStride);
- 						Dest[0] = (int32)(FMath::Clamp(A, 0.0f, 1.0f) * 65535.0f + 0.5f);
- 						Dest[1] = (int32)(FMath::Clamp(B, 0.0f, 1.0f) * 65535.0f + 0.5f);
-					}
-					else
-					{
-						check(Desc.Format == PF_R8G8);
+						else if (Desc.Format == PF_G16R16)
+						{
+							uint16* Dest = (uint16*)(DestBuffer + x * 4 + y * DestStride);
+							Dest[0] = (int32)(FMath::Clamp(A, 0.0f, 1.0f) * 65535.0f + 0.5f);
+							Dest[1] = (int32)(FMath::Clamp(B, 0.0f, 1.0f) * 65535.0f + 0.5f);
+						}
+						else
+						{
+							check(Desc.Format == PF_R8G8);
 
-						uint8* Dest = (uint8*)(DestBuffer + x * 2 + y * DestStride);
-						Dest[0] = (int32)(FMath::Clamp(A, 0.0f, 1.0f) * 255.9999f);
-						Dest[1] = (int32)(FMath::Clamp(B, 0.0f, 1.0f) * 255.9999f);
+							uint8* Dest = (uint8*)(DestBuffer + x * 2 + y * DestStride);
+							Dest[0] = (int32)(FMath::Clamp(A, 0.0f, 1.0f) * 255.9999f);
+							Dest[1] = (int32)(FMath::Clamp(B, 0.0f, 1.0f) * 255.9999f);
+						}
 					}
 				}
+				RHICmdList.UnlockTexture2D((FTexture2DRHIRef&)PreintegratedGF->GetRenderTargetItem().ShaderResourceTexture, 0, false);
 			}
-			RHICmdList.UnlockTexture2D((FTexture2DRHIRef&)PreintegratedGF->GetRenderTargetItem().ShaderResourceTexture, 0, false);
-		}
 
-		{
-			FPooledRenderTargetDesc Desc(FPooledRenderTargetDesc::Create2DDesc(FIntPoint( LTC_Size, LTC_Size ), PF_FloatRGBA, FClearValueBinding::None, TexCreate_FastVRAM, TexCreate_None, false));
-			Desc.AutoWritable = false;
+			{
+				FPooledRenderTargetDesc Desc(FPooledRenderTargetDesc::Create2DDesc(FIntPoint(LTC_Size, LTC_Size), PF_FloatRGBA, FClearValueBinding::None, TexCreate_FastVRAM, TexCreate_None, false));
+				Desc.AutoWritable = false;
 
-			GRenderTargetPool.FindFreeElement(RHICmdList, Desc, LTCMat, TEXT("LTCMat"));
-			// Write the contents of the texture.
-			uint32 DestStride;
-			uint8* DestBuffer = (uint8*)RHICmdList.LockTexture2D((FTexture2DRHIRef&)LTCMat->GetRenderTargetItem().ShaderResourceTexture, 0, RLM_WriteOnly, DestStride, false);
+				GRenderTargetPool.FindFreeElement(RHICmdList, Desc, LTCMat, TEXT("LTCMat"));
+				// Write the contents of the texture.
+				uint32 DestStride;
+				uint8* DestBuffer = (uint8*)RHICmdList.LockTexture2D((FTexture2DRHIRef&)LTCMat->GetRenderTargetItem().ShaderResourceTexture, 0, RLM_WriteOnly, DestStride, false);
 
-			for(int32 y = 0; y < Desc.Extent.Y; ++y)
-			{
-				for(int32 x = 0; x < Desc.Extent.X; ++x)
+				for (int32 y = 0; y < Desc.Extent.Y; ++y)
 				{
-					uint16* Dest = (uint16*)(DestBuffer + x * 4 * sizeof(uint16) + y * DestStride);
+					for (int32 x = 0; x < Desc.Extent.X; ++x)
+					{
+						uint16* Dest = (uint16*)(DestBuffer + x * 4 * sizeof(uint16) + y * DestStride);
 
-					for( int k = 0; k < 4; k++ )
-						Dest[k] = FFloat16( LTC_Mat[ 4 * (x + y * LTC_Size) + k ] ).Encoded;
+						for (int k = 0; k < 4; k++)
+							Dest[k] = FFloat16(LTC_Mat[4 * (x + y * LTC_Size) + k]).Encoded;
+					}
 				}
+				RHICmdList.UnlockTexture2D((FTexture2DRHIRef&)LTCMat->GetRenderTargetItem().ShaderResourceTexture, 0, false);
 			}
-			RHICmdList.UnlockTexture2D((FTexture2DRHIRef&)LTCMat->GetRenderTargetItem().ShaderResourceTexture, 0, false);
-		}
 
-		{
-			FPooledRenderTargetDesc Desc(FPooledRenderTargetDesc::Create2DDesc(FIntPoint( LTC_Size, LTC_Size ), PF_G16R16F, FClearValueBinding::None, TexCreate_FastVRAM, TexCreate_None, false));
-			Desc.AutoWritable = false;
+			{
+				FPooledRenderTargetDesc Desc(FPooledRenderTargetDesc::Create2DDesc(FIntPoint(LTC_Size, LTC_Size), PF_G16R16F, FClearValueBinding::None, TexCreate_FastVRAM, TexCreate_None, false));
+				Desc.AutoWritable = false;
 
-			GRenderTargetPool.FindFreeElement(RHICmdList, Desc, LTCAmp, TEXT("LTCAmp"));
-			// Write the contents of the texture.
-			uint32 DestStride;
-			uint8* DestBuffer = (uint8*)RHICmdList.LockTexture2D((FTexture2DRHIRef&)LTCAmp->GetRenderTargetItem().ShaderResourceTexture, 0, RLM_WriteOnly, DestStride, false);
+				GRenderTargetPool.FindFreeElement(RHICmdList, Desc, LTCAmp, TEXT("LTCAmp"));
+				// Write the contents of the texture.
+				uint32 DestStride;
+				uint8* DestBuffer = (uint8*)RHICmdList.LockTexture2D((FTexture2DRHIRef&)LTCAmp->GetRenderTargetItem().ShaderResourceTexture, 0, RLM_WriteOnly, DestStride, false);
 
-			for(int32 y = 0; y < Desc.Extent.Y; ++y)
-			{
-				for(int32 x = 0; x < Desc.Extent.X; ++x)
+				for (int32 y = 0; y < Desc.Extent.Y; ++y)
 				{
-					uint16* Dest = (uint16*)(DestBuffer + x * 2 * sizeof(uint16) + y * DestStride);
+					for (int32 x = 0; x < Desc.Extent.X; ++x)
+					{
+						uint16* Dest = (uint16*)(DestBuffer + x * 2 * sizeof(uint16) + y * DestStride);
 
-					for( int k = 0; k < 2; k++ )
-						Dest[k] = FFloat16( LTC_Amp[ 4 * (x + y * LTC_Size) + k ] ).Encoded;
+						for (int k = 0; k < 2; k++)
+							Dest[k] = FFloat16(LTC_Amp[4 * (x + y * LTC_Size) + k]).Encoded;
+					}
 				}
+				RHICmdList.UnlockTexture2D((FTexture2DRHIRef&)LTCAmp->GetRenderTargetItem().ShaderResourceTexture, 0, false);
 			}
-			RHICmdList.UnlockTexture2D((FTexture2DRHIRef&)LTCAmp->GetRenderTargetItem().ShaderResourceTexture, 0, false);
-		}
-	}
+		} // end Create the SSAO randomization texture
+	} // end if (FeatureLevelInitializedTo < ERHIFeatureLevel::SM4 && InFeatureLevel >= ERHIFeatureLevel::SM4)
 
 	// Initialize textures only once.
-	bTexturesInitialized = true;
 	FeatureLevelInitializedTo = InFeatureLevel;
 }
 
@@ -536,5 +546,5 @@ void FSystemTextures::ReleaseDynamicRHI()
 	GRenderTargetPool.FreeUnusedResources();
 
 	// Indicate that textures will need to be reinitialized.
-	bTexturesInitialized = false;
+	FeatureLevelInitializedTo = ERHIFeatureLevel::Num;
 }
diff --git a/Engine/Source/Runtime/Renderer/Private/SystemTextures.h b/Engine/Source/Runtime/Renderer/Private/SystemTextures.h
--- a/Engine/Source/Runtime/Renderer/Private/SystemTextures.h
+++ b/Engine/Source/Runtime/Renderer/Private/SystemTextures.h
@@ -17,24 +17,27 @@ class FSystemTextures : public FRenderResource
 {
 public:
 	FSystemTextures()
-	: FRenderResource()
-	, FeatureLevelInitializedTo(ERHIFeatureLevel::ES2)
-	, bTexturesInitialized(false)
+		: FRenderResource()
+		, FeatureLevelInitializedTo(ERHIFeatureLevel::Num)
 	{}
 
 	/**
 	 * Initialize/allocate textures if not already.
 	 */
-	inline void InitializeTextures(FRHICommandListImmediate& RHICmdList, ERHIFeatureLevel::Type InFeatureLevel)
+	inline void InitializeTextures(FRHICommandListImmediate& RHICmdList, const ERHIFeatureLevel::Type InFeatureLevel)
 	{
-		if (bTexturesInitialized && FeatureLevelInitializedTo >= InFeatureLevel)
+		// if this is the first call initialize everything
+		if (FeatureLevelInitializedTo == ERHIFeatureLevel::Num)
 		{
-			// Already initialized up to at least the feature level we need, so do nothing
-			return;
+			InitializeCommonTextures(RHICmdList);
+			InitializeFeatureLevelDependentTextures(RHICmdList, InFeatureLevel);
 		}
-
-		ReleaseDynamicRHI();
-		InternalInitializeTextures(RHICmdList, InFeatureLevel);
+		// otherwise, if we request a higher feature level, we might need to initialize those textures that depend on the feature level
+		else if (InFeatureLevel > FeatureLevelInitializedTo)
+		{
+			InitializeFeatureLevelDependentTextures(RHICmdList, InFeatureLevel);
+		}
+		// there's no needed setup for those feature levels lower or identical to the current one
 	}
 
 	// FRenderResource interface.
@@ -81,9 +84,9 @@ public:
 protected:
 	/** Maximum feature level that the textures have been initialized up to */
 	ERHIFeatureLevel::Type FeatureLevelInitializedTo;
-	bool bTexturesInitialized;
 
-	void InternalInitializeTextures(FRHICommandListImmediate& RHICmdList, ERHIFeatureLevel::Type InFeatureLevel);
+	void InitializeCommonTextures(FRHICommandListImmediate& RHICmdList);
+	void InitializeFeatureLevelDependentTextures(FRHICommandListImmediate& RHICmdList, const ERHIFeatureLevel::Type InFeatureLevel);
 };
 
 /** The global system textures used for scene rendering. */
