UE-64680

The new depth bounds optimization in PostAO doesn't work with MSAA. The AO target and depth buffer are different sample counts. Forcing this feature off for now.
In the future if we want to turn it back on we will either have to resolve the depth buffer or increase the AO target sample count.

diff --git a/Engine/Source/Runtime/Renderer/Private/CompositionLighting/PostProcessAmbientOcclusion.cpp b/Engine/Source/Runtime/Renderer/Private/CompositionLighting/PostProcessAmbientOcclusion.cpp
--- a/Engine/Source/Runtime/Renderer/Private/CompositionLighting/PostProcessAmbientOcclusion.cpp
+++ b/Engine/Source/Runtime/Renderer/Private/CompositionLighting/PostProcessAmbientOcclusion.cpp
@@ -778,7 +778,8 @@ void FRCPassPostProcessAmbientOcclusion::ProcessPS(FRenderingCompositePassContex
 	const FSceneRenderTargetItem* DestRenderTarget, const FSceneRenderTargetItem* SceneDepthBuffer,
 	const FIntRect& ViewRect, const FIntPoint& TexSize, int32 ShaderQuality, bool bDoUpsample)
 {
-	const bool bDepthBoundsTestEnabled = GSupportsDepthBoundsTest && SceneDepthBuffer && CVarAmbientOcclusionDepthBoundsTest.GetValueOnRenderThread();
+	// We do not support the depth bounds optimization if we are in MSAA. To do so we would have to resolve the depth buffer here OR use a multisample texture for our AO target.
+	const bool bDepthBoundsTestEnabled = GSupportsDepthBoundsTest && SceneDepthBuffer && CVarAmbientOcclusionDepthBoundsTest.GetValueOnRenderThread() && SceneDepthBuffer->TargetableTexture->GetNumSamples() == 1;
 
 	// Set the view family's render target/viewport.
 
