Re-worked fix from CL 4442729 to handle !SUPPORTS_INDEPENDENT_SAMPLERS

diff --git a/Engine/Shaders/Private/BasePassCommon.ush b/Engine/Shaders/Private/BasePassCommon.ush
--- a/Engine/Shaders/Private/BasePassCommon.ush
+++ b/Engine/Shaders/Private/BasePassCommon.ush
@@ -1,7 +1,7 @@
 // Copyright 1998-2018 Epic Games, Inc. All Rights Reserved.
 
 /*=============================================================================
-	BasePassCommon.usf: Base pass definitions used by both vertex and pixel shader
+	BasePassCommon.ush: Base pass definitions used by both vertex and pixel shader
 =============================================================================*/
 
 // Defining SSProfilesTexture here overrides the default definition in SubsurfaceProfileCommon.ush
diff --git a/Engine/Shaders/Private/ForwardLightingCommon.ush b/Engine/Shaders/Private/ForwardLightingCommon.ush
--- a/Engine/Shaders/Private/ForwardLightingCommon.ush
+++ b/Engine/Shaders/Private/ForwardLightingCommon.ush
@@ -69,7 +69,14 @@ float2 CalculateNearestResolvedDepthScreenUV(float2 ScreenUV, float SceneDepth)
 
 float4 GetForwardDynamicShadowFactors(float2 ScreenUV)
 {
-	float4 Value = Texture2DSampleLevel(OpaqueBasePass.ForwardScreenSpaceShadowMaskTexture, GlobalPointClampedSampler, ScreenUV, 0);
+	int2 IntScreenUV = int2(trunc(ScreenUV * View_BufferSizeAndInvSize.xy));
+	float4 Value = 1.0f;
+	// Avoid sampling the white texture fallback because Load returns 0 when out-of-bound.
+	BRANCH
+	if (OpaqueBasePass.UseForwardScreenSpaceShadowMask)
+	{
+		Value = OpaqueBasePass.ForwardScreenSpaceShadowMaskTexture.Load(int3(IntScreenUV, 0));
+	}
 	return Square(Value);
 }
 
diff --git a/Engine/Source/Runtime/Renderer/Private/BasePassRendering.cpp b/Engine/Source/Runtime/Renderer/Private/BasePassRendering.cpp
--- a/Engine/Source/Runtime/Renderer/Private/BasePassRendering.cpp
+++ b/Engine/Source/Runtime/Renderer/Private/BasePassRendering.cpp
@@ -825,11 +825,16 @@ void CreateOpaqueBasePassUniformBuffer(
 
 	// Forward shading
 	{
-		if (!ForwardScreenSpaceShadowMask)
+		if (ForwardScreenSpaceShadowMask)
 		{
-			ForwardScreenSpaceShadowMask = GSystemTextures.WhiteDummy.GetReference();
+			BasePassParameters.UseForwardScreenSpaceShadowMask = 1;
+			BasePassParameters.ForwardScreenSpaceShadowMaskTexture = ForwardScreenSpaceShadowMask->GetRenderTargetItem().ShaderResourceTexture;
+		}
+		else
+		{
+			BasePassParameters.UseForwardScreenSpaceShadowMask = 0;
+			BasePassParameters.ForwardScreenSpaceShadowMaskTexture = GSystemTextures.WhiteDummy.GetReference()->GetRenderTargetItem().ShaderResourceTexture;
 		}
-		BasePassParameters.ForwardScreenSpaceShadowMaskTexture = ForwardScreenSpaceShadowMask->GetRenderTargetItem().ShaderResourceTexture;
 
 		IPooledRenderTarget* IndirectOcclusion = SceneRenderTargets.ScreenSpaceAO;
 
diff --git a/Engine/Source/Runtime/Renderer/Private/BasePassRendering.h b/Engine/Source/Runtime/Renderer/Private/BasePassRendering.h
--- a/Engine/Source/Runtime/Renderer/Private/BasePassRendering.h
+++ b/Engine/Source/Runtime/Renderer/Private/BasePassRendering.h
@@ -61,6 +61,7 @@ END_UNIFORM_BUFFER_STRUCT(FSharedBasePassUniformParameters)
 BEGIN_UNIFORM_BUFFER_STRUCT(FOpaqueBasePassUniformParameters,)
 	UNIFORM_MEMBER_STRUCT(FSharedBasePassUniformParameters, Shared)
 	// Forward shading 
+	UNIFORM_MEMBER(uint32, UseForwardScreenSpaceShadowMask)
 	UNIFORM_MEMBER_TEXTURE(Texture2D, ForwardScreenSpaceShadowMaskTexture)
 	UNIFORM_MEMBER_TEXTURE(Texture2D, IndirectOcclusionTexture)
 	UNIFORM_MEMBER_TEXTURE(Texture2D, ResolvedSceneDepthTexture)
