Clone of CL 4452841 Dev-Rendering:

Fix for high dpi and low screen percentage editor crash.  Updated with code review comments - move check location to test PrimaryResolutionFraction that should have already been clamped rather than after multiplied in high dpi.

diff --git a/Engine/Source/Runtime/Renderer/Private/SceneRendering.cpp b/Engine/Source/Runtime/Renderer/Private/SceneRendering.cpp
--- a/Engine/Source/Runtime/Renderer/Private/SceneRendering.cpp
+++ b/Engine/Source/Runtime/Renderer/Private/SceneRendering.cpp
@@ -1806,8 +1806,6 @@ FSceneRenderer::FSceneRenderer(const FSceneViewFamily* InViewFamily,FHitProxyCon
 // static
 FIntPoint FSceneRenderer::ApplyResolutionFraction(const FSceneViewFamily& ViewFamily, const FIntPoint& UnscaledViewSize, float ResolutionFraction)
 {
-	check(FSceneViewScreenPercentageConfig::IsValidResolutionFraction(ResolutionFraction));
-
 	FIntPoint ViewSize;
 
 	// CeilToInt so tha view size is at least 1x1 if ResolutionFraction == FSceneViewScreenPercentageConfig::kMinResolutionFraction.
@@ -1952,6 +1950,8 @@ void FSceneRenderer::PrepareViewRectsForRendering()
 		// Make sure the screen percentage interface has not lied to the renderer about the upper bound.
 		checkf(PrimaryResolutionFraction <= ViewFamily.GetPrimaryResolutionFractionUpperBound(),
 			TEXT("ISceneViewFamilyScreenPercentage::GetPrimaryResolutionFractionUpperBound() should not lie to the renderer."));
+		
+		check(FSceneViewScreenPercentageConfig::IsValidResolutionFraction(PrimaryResolutionFraction));
 
 		// Compute final resolution fraction.
 		float ResolutionFraction = PrimaryResolutionFraction * ViewFamily.SecondaryViewFraction;
