Fixed: Refraction causes rendering artifacts when MSAA is enabled on Mobile
Refraction now applied as post-process effect, correctly handling multiple overlapping distortions. This chnage has side effect that foreground translucency will be distorted as well

diff --git a/Engine/Shaders/Private/DistortAccumulatePS.usf b/Engine/Shaders/Private/DistortAccumulatePS.usf
--- a/Engine/Shaders/Private/DistortAccumulatePS.usf
+++ b/Engine/Shaders/Private/DistortAccumulatePS.usf
@@ -30,6 +30,15 @@ void Main(
 	FMaterialPixelParameters MaterialParameters = GetMaterialPixelParameters(Interpolants, SvPosition);
 	FPixelMaterialInputs PixelMaterialInputs;
 	CalcMaterialParameters(MaterialParameters, PixelMaterialInputs, SvPosition, bIsFrontFace);
+
+#if (ES2_PROFILE || ES3_1_PROFILE) 
+	// Mobile specific: 
+	// distortion primtitives rendered without depth testing
+	// discard pixels that are hidden by comparing them to scene depth stored in SceneColor.A
+	float2 BufferUV = SvPositionToBufferUV(SvPosition);
+	float ClipDepth = CalcFullSceneColor(BufferUV).a - SvPosition.w;
+	clip(ClipDepth);
+#endif
 	
 	// material distortion offset
 	half3 Normal = GetMaterialNormal(MaterialParameters, PixelMaterialInputs);
@@ -73,7 +82,8 @@ void Main(
 	float2 DistortBufferUV = ScreenUV + BufferUVDistortion;
 
 #if (ES2_PROFILE || ES3_1_PROFILE) && (!WEBGL)
-	float DistortSceneDepth = ConvertFromDeviceZ(Texture2DSampleLevel(MobileSceneTextures.SceneDepthTexture, MobileSceneTextures.SceneDepthTextureSampler, DistortBufferUV, 0).r); 
+	// sample depth from SceneColor.A on Mobile
+	float DistortSceneDepth = CalcFullSceneColor(DistortBufferUV).a; 
 #else
 	float DistortSceneDepth = CalcSceneDepth(DistortBufferUV);
 #endif // (ES2_PROFILE || ES3_1_PROFILE) 
@@ -86,11 +96,6 @@ void Main(
 	float InvWidth = 1.0f / max(1.0f, Range );
 	BufferUVDistortion *= saturate((Z - ZCompare) * InvWidth + Bias);
 
-#if (ES2_PROFILE || ES3_1_PROFILE)
-	// mobile path direclty outputs distorted colors
-	DistortBufferUV = ScreenUV + BufferUVDistortion;
-	OutColor = CalcFullSceneColor(DistortBufferUV);
-#else
 	//Scale up for better precision in low/subtle refractions at the expense of artefacts at higher refraction.
 	BufferUVDistortion*= DistortionScaleBias;
 
@@ -100,7 +105,6 @@ void Main(
 
 	// output positives in R|G channels and negatives in B|A channels
 	OutColor = float4(PosOffset.x,PosOffset.y,NegOffset.x,NegOffset.y);
-#endif // (ES2_PROFILE || ES3_1_PROFILE) 
 }
 
 
diff --git a/Engine/Shaders/Private/DistortApplyScreenPS.usf b/Engine/Shaders/Private/DistortApplyScreenPS.usf
--- a/Engine/Shaders/Private/DistortApplyScreenPS.usf
+++ b/Engine/Shaders/Private/DistortApplyScreenPS.usf
@@ -4,6 +4,7 @@
 =============================================================================*/
 
 #include "Common.ush"
+#include "PostProcessCommon.ush"
 
 static const half InvDistortionScaleBias = 1 / 4.0f;
 
@@ -82,4 +83,25 @@ void Merge(
 #endif
 }
 
+void Merge_ES2(
+	in noperspective float4 TexCoord: TEXCOORD0,
+	out half4 OutColor : SV_Target0
+)
+{
+	half4 AccumDist = PostprocessInput1.Sample(PostprocessInput1Sampler, TexCoord.xy);
+	// apply inverse scale
+	// offset = [R-B,G-A] 
+	half2 DistBufferUVOffset = (AccumDist.rg - AccumDist.ba);
+	//Scale by the screen size and a fudge factor to come close to the offset values we would have had under normal circumstances before my changes. Also flip Y and invert the precision bias scale.
+	DistBufferUVOffset *= InvDistortionScaleBias;
+	float2 NewBufferUV = TexCoord.xy + DistBufferUVOffset;
+	// If we're about to sample outside the valid BufferUVMinMax, set to 0 distortion.
+	FLATTEN if ( NewBufferUV.x < View.BufferBilinearUVMinMax.x || NewBufferUV.x > View.BufferBilinearUVMinMax.z ||
+				 NewBufferUV.y < View.BufferBilinearUVMinMax.y || NewBufferUV.y > View.BufferBilinearUVMinMax.w )
+	{
+		NewBufferUV = TexCoord.xy;
+	}
+
+	OutColor = PostprocessInput0.Sample(PostprocessInput0Sampler, NewBufferUV.xy);
+}
 
diff --git a/Engine/Source/Runtime/Renderer/Private/DistortionRendering.cpp b/Engine/Source/Runtime/Renderer/Private/DistortionRendering.cpp
--- a/Engine/Source/Runtime/Renderer/Private/DistortionRendering.cpp
+++ b/Engine/Source/Runtime/Renderer/Private/DistortionRendering.cpp
@@ -1192,118 +1192,3 @@ void FSceneRenderer::RenderDistortion(FRHICommandListImmediate& RHICmdList)
 		}
 	}
 }
-
-void FSceneRenderer::RenderDistortionES2(FRHICommandListImmediate& RHICmdList)
-{
-	// We need access to HDR scene color
-	if ((!IsMobileHDR() || IsMobileHDRMosaic()) && !IsHTML5Platform())
-	{
-		return;
-	}
-
-	// do we need to render the distortion pass?
-	bool bRender=false;
-	for(int32 ViewIndex = 0;ViewIndex < Views.Num();ViewIndex++)
-	{
-		const FViewInfo& View = Views[ViewIndex];
-		if( View.DistortionPrimSet.NumPrims() > 0)
-		{
-			bRender=true;
-			break;
-		}
-	}
-
-	static IConsoleVariable* CVar = IConsoleManager::Get().FindConsoleVariable(TEXT("r.DisableDistortion"));
-	int32 DisableDistortion = CVar->GetInt();
-
-	if (bRender && !DisableDistortion)
-	{
-		// Apply distortion
-		SCOPED_DRAW_EVENT(RHICmdList, Distortion);
-
-		FSceneRenderTargets& SceneContext = FSceneRenderTargets::Get(RHICmdList);
-
-		RHICmdList.CopyToResolveTarget(SceneContext.GetSceneColorSurface(), SceneContext.GetSceneColorTexture(), FResolveRect(0, 0, FamilySize.X, FamilySize.Y));
-
-		TRefCountPtr<IPooledRenderTarget> SceneColorDistorted;
-		FPooledRenderTargetDesc Desc = SceneContext.GetSceneColor()->GetDesc();
-		Desc.Flags &= ~(TexCreate_FastVRAM | TexCreate_Transient);
-		GRenderTargetPool.FindFreeElement(RHICmdList, Desc, SceneColorDistorted, TEXT("SceneColorDistorted"));
-		const FSceneRenderTargetItem& DistortedRenderTarget = SceneColorDistorted->GetRenderTargetItem();
-
-		FGraphicsPipelineStateInitializer GraphicsPSOInit;
-		SetRenderTarget(RHICmdList, DistortedRenderTarget.TargetableTexture, SceneContext.GetSceneDepthSurface(), ESimpleRenderTargetMode::EClearColorExistingDepth, FExclusiveDepthStencil::DepthRead_StencilNop);
-		RHICmdList.ApplyCachedRenderTargets(GraphicsPSOInit);
-
-		// Copy scene color to a new render target
-		for(int32 ViewIndex = 0, Num = Views.Num(); ViewIndex < Num; ++ViewIndex)
-		{
-			SCOPED_CONDITIONAL_DRAW_EVENTF(RHICmdList, EventView, Views.Num() > 1, TEXT("View%d"), ViewIndex);
-
-			FViewInfo& View = Views[ViewIndex];
-
-			// useful when we move this into the compositing graph
-			FRenderingCompositePassContext Context(RHICmdList, View);
-
-			// Set the view family's render target/viewport.
-			Context.SetViewportAndCallRHI(View.ViewRect);
-
-			GraphicsPSOInit.BlendState = TStaticBlendState<>::GetRHI();
-			GraphicsPSOInit.RasterizerState = TStaticRasterizerState<>::GetRHI();
-			GraphicsPSOInit.DepthStencilState = TStaticDepthStencilState<false, CF_Always>::GetRHI();
-
-			TShaderMapRef<FPostProcessVS> VertexShader(View.ShaderMap);
-			TShaderMapRef<TDistortionMergePS<false>> PixelShader(View.ShaderMap);
-
-			GraphicsPSOInit.BoundShaderState.VertexDeclarationRHI = GFilterVertexDeclaration.VertexDeclarationRHI;
-			GraphicsPSOInit.BoundShaderState.VertexShaderRHI = GETSAFERHISHADER_VERTEX(*VertexShader);
-			GraphicsPSOInit.BoundShaderState.PixelShaderRHI = GETSAFERHISHADER_PIXEL(*PixelShader);
-			GraphicsPSOInit.PrimitiveType = PT_TriangleList;
-
-			SetGraphicsPipelineState(RHICmdList, GraphicsPSOInit);
-
-			VertexShader->SetParameters(Context);
-			PixelShader->SetParameters(Context, View, SceneContext.GetSceneColor()->GetRenderTargetItem().ShaderResourceTexture);
-
-			DrawRectangle(
-				RHICmdList,
-				0, 0,
-				View.ViewRect.Width(), View.ViewRect.Height(),
-				View.ViewRect.Min.X, View.ViewRect.Min.Y,
-				View.ViewRect.Width(), View.ViewRect.Height(),
-				View.ViewRect.Size(),
-				SceneContext.GetBufferSizeXY(),
-				*VertexShader,
-				EDRF_UseTriangleOptimization);
-		}
-
-		// Distort scene color in place
-		for(int32 ViewIndex = 0, Num = Views.Num(); ViewIndex < Num; ++ViewIndex)
-		{
-			SCOPED_CONDITIONAL_DRAW_EVENTF(RHICmdList, EventView, Views.Num() > 1, TEXT("View%d"), ViewIndex);
-
-			FViewInfo& View = Views[ViewIndex];
-
-			// useful when we move this into the compositing graph
-			FRenderingCompositePassContext Context(RHICmdList, View);
-
-			// Set the view family's render target/viewport.
-			Context.SetViewportAndCallRHI(View.ViewRect);
-
-			FMobileSceneTextureUniformParameters SceneTextureParameters;
-			SetupMobileSceneTextureUniformParameters(SceneContext, View.FeatureLevel, true, SceneTextureParameters);
-			TUniformBufferRef<FMobileSceneTextureUniformParameters> MobilePassUniformBuffer = TUniformBufferRef<FMobileSceneTextureUniformParameters>::CreateUniformBufferImmediate(SceneTextureParameters, UniformBuffer_SingleFrame);
-
-			FDrawingPolicyRenderState DrawRenderState(View, MobilePassUniformBuffer);
-			// test against depth
-			DrawRenderState.SetBlendState(TStaticBlendState<>::GetRHI());
-			DrawRenderState.SetDepthStencilState(TStaticDepthStencilState<false, CF_DepthNearOrEqual>::GetRHI());
-
-			// draw only distortion meshes
-			View.DistortionPrimSet.DrawAccumulatedOffsets(RHICmdList, View, DrawRenderState, false);
-		}
-
-		// Set distorted scene color as main
-		SceneContext.SetSceneColor(SceneColorDistorted);
-	}
-}
diff --git a/Engine/Source/Runtime/Renderer/Private/MobileDistortionPass.cpp b/Engine/Source/Runtime/Renderer/Private/MobileDistortionPass.cpp
new file mode 100644
--- /dev/null
+++ b/Engine/Source/Runtime/Renderer/Private/MobileDistortionPass.cpp
@@ -0,0 +1,179 @@
+// Copyright 1998-2018 Epic Games, Inc. All Rights Reserved.
+
+/*=============================================================================
+MobileDistortionPass.cpp - Mobile specific rendering of primtives with refraction
+=============================================================================*/
+
+#include "MobileDistortionPass.h"
+#include "TranslucentRendering.h"
+#include "DynamicPrimitiveDrawing.h"
+#include "PostProcess/PostProcessing.h"
+#include "PostProcess/SceneFilterRendering.h"
+
+bool IsMobileDistortionActive(const FViewInfo& View)
+{
+	static IConsoleVariable* CVar = IConsoleManager::Get().FindConsoleVariable(TEXT("r.DisableDistortion"));
+	int32 DisableDistortion = CVar->GetInt();
+
+	// Distortion on mobile requires SceneDepth information in SceneColor.A channel
+	const EMobileHDRMode HDRMode = GetMobileHDRMode();
+
+	return 
+		HDRMode == EMobileHDRMode::EnabledFloat16 &&
+		View.Family->EngineShowFlags.Translucency &&
+		FSceneRenderer::GetRefractionQuality(*View.Family) > 0 &&
+		DisableDistortion == 0 && 
+		View.DistortionPrimSet.NumPrims() > 0;
+}
+
+void FRCDistortionAccumulatePassES2::Process(FRenderingCompositePassContext& Context)
+{
+	SCOPED_DRAW_EVENT(Context.RHICmdList, DistortionAccumulatePass);
+	
+	FViewInfo& View = const_cast<FViewInfo &>(Context.View);
+	FSceneRenderTargets& SceneTargets = FSceneRenderTargets::Get(Context.RHICmdList);
+	const FSceneRenderTargetItem& DestRenderTarget = PassOutputs[0].RequestSurface(Context);
+
+	SetRenderTarget(Context.RHICmdList, DestRenderTarget.TargetableTexture, FTextureRHIParamRef(), ESimpleRenderTargetMode::EClearColorAndDepth);
+	Context.SetViewportAndCallRHI(View.ViewRect);
+
+	FMobileSceneTextureUniformParameters SceneTextureParameters;
+	SetupMobileSceneTextureUniformParameters(SceneTargets, View.FeatureLevel, true, SceneTextureParameters);
+	TUniformBufferRef<FMobileSceneTextureUniformParameters> MobilePassUniformBuffer = TUniformBufferRef<FMobileSceneTextureUniformParameters>::CreateUniformBufferImmediate(SceneTextureParameters, UniformBuffer_SingleFrame);
+
+	FDrawingPolicyRenderState DrawRenderState(View, MobilePassUniformBuffer);
+	// We don't have depth, render all pixels, pixel shader will sample SceneDepth from SceneColor.A and discard if occluded
+	DrawRenderState.SetDepthStencilState(TStaticDepthStencilState<false, CF_Always>::GetRHI());
+	// additive blending of offsets
+	DrawRenderState.SetBlendState(TStaticBlendState<CW_RGBA, BO_Add, BF_One, BF_One, BO_Add, BF_One, BF_One>::GetRHI());
+
+	// draw only distortion meshes to accumulate their offsets
+	View.DistortionPrimSet.DrawAccumulatedOffsets(Context.RHICmdList, View, DrawRenderState, false);
+		
+	Context.RHICmdList.CopyToResolveTarget(DestRenderTarget.TargetableTexture, DestRenderTarget.ShaderResourceTexture, FResolveParams());
+}
+
+FPooledRenderTargetDesc FRCDistortionAccumulatePassES2::ComputeOutputDesc(EPassOutputId InPassOutputId) const
+{
+	FPooledRenderTargetDesc Ret;
+	Ret.Depth = 0;
+	Ret.ArraySize = 1;
+	Ret.bIsArray = false;
+	Ret.NumMips = 1;
+	Ret.TargetableFlags = TexCreate_RenderTargetable;
+	Ret.bForceSeparateTargetAndShaderResource = false;
+	Ret.Format = PF_B8G8R8A8;
+	Ret.NumSamples = 1;
+	Ret.Extent.X = FMath::Max(1, PrePostSourceViewportSize.X);
+	Ret.Extent.Y = FMath::Max(1, PrePostSourceViewportSize.Y);
+	Ret.DebugName = TEXT("DistortionAccumulatePass");
+	Ret.ClearValue = FClearValueBinding(FLinearColor::Transparent);
+	return Ret;
+}
+
+class FDistortionMergePS_ES2 : public FGlobalShader
+{
+	DECLARE_SHADER_TYPE(FDistortionMergePS_ES2, Global);
+
+	static bool ShouldCompilePermutation(const FGlobalShaderPermutationParameters& Parameters)
+	{
+		return !IsConsolePlatform(Parameters.Platform);
+	}
+
+	FDistortionMergePS_ES2() {}
+
+public:
+	FPostProcessPassParameters PostprocessParameter;
+	
+	/** Initialization constructor. */
+	FDistortionMergePS_ES2(const ShaderMetaType::CompiledShaderInitializerType& Initializer)
+		: FGlobalShader(Initializer)
+	{
+		PostprocessParameter.Bind(Initializer.ParameterMap);
+	}
+
+	void SetParameters(const FRenderingCompositePassContext& Context)
+	{
+		const FPixelShaderRHIParamRef ShaderRHI = GetPixelShader();
+		FGlobalShader::SetParameters<FViewUniformShaderParameters>(Context.RHICmdList, ShaderRHI, Context.View.ViewUniformBuffer);
+		PostprocessParameter.SetPS(Context.RHICmdList, ShaderRHI, Context, TStaticSamplerState<SF_Bilinear, AM_Clamp, AM_Clamp, AM_Clamp>::GetRHI());
+	}
+
+	virtual bool Serialize(FArchive& Ar) override
+	{
+		bool bShaderHasOutdatedParameters = FGlobalShader::Serialize(Ar);
+		Ar << PostprocessParameter;
+		return bShaderHasOutdatedParameters;
+	}
+};
+
+IMPLEMENT_SHADER_TYPE(,FDistortionMergePS_ES2,TEXT("/Engine/Private/DistortApplyScreenPS.usf"),TEXT("Merge_ES2"),SF_Pixel);
+
+void FRCDistortionMergePassES2::Process(FRenderingCompositePassContext& Context)
+{
+	SCOPED_DRAW_EVENT(Context.RHICmdList, DistortionMergePass);
+	
+	const FViewInfo& View = Context.View;
+	const FPooledRenderTargetDesc* InputDesc = GetInputDesc(ePId_Input0);
+	const FPooledRenderTargetDesc& OutputDesc = PassOutputs[0].RenderTargetDesc;
+	const FSceneRenderTargetItem& DestRenderTarget = PassOutputs[0].RequestSurface(Context);
+
+	SetRenderTarget(Context.RHICmdList, DestRenderTarget.TargetableTexture, FTextureRHIParamRef(), ESimpleRenderTargetMode::EClearColorAndDepth);
+
+	FIntRect SrcRect = View.ViewRect;
+	FIntRect DestRect = View.ViewRect;
+	FIntPoint SrcSize = InputDesc->Extent;
+	FIntPoint DstSize = OutputDesc.Extent;
+	
+	Context.SetViewportAndCallRHI(View.ViewRect);
+
+	FGraphicsPipelineStateInitializer GraphicsPSOInit;
+	Context.RHICmdList.ApplyCachedRenderTargets(GraphicsPSOInit);
+	GraphicsPSOInit.BlendState = TStaticBlendState<>::GetRHI();
+	GraphicsPSOInit.RasterizerState = TStaticRasterizerState<>::GetRHI();
+	GraphicsPSOInit.DepthStencilState = TStaticDepthStencilState<false, CF_Always>::GetRHI();
+
+	TShaderMapRef<FPostProcessVS> VertexShader(View.ShaderMap);
+	TShaderMapRef<FDistortionMergePS_ES2> PixelShader(View.ShaderMap);
+
+	GraphicsPSOInit.BoundShaderState.VertexDeclarationRHI = GFilterVertexDeclaration.VertexDeclarationRHI;
+	GraphicsPSOInit.BoundShaderState.VertexShaderRHI = GETSAFERHISHADER_VERTEX(*VertexShader);
+	GraphicsPSOInit.BoundShaderState.PixelShaderRHI = GETSAFERHISHADER_PIXEL(*PixelShader);
+	GraphicsPSOInit.PrimitiveType = PT_TriangleList;
+
+	SetGraphicsPipelineState(Context.RHICmdList, GraphicsPSOInit);
+
+	VertexShader->SetParameters(Context);
+	PixelShader->SetParameters(Context);
+
+	DrawRectangle(
+		Context.RHICmdList,
+		0, 0,
+		DstSize.X, DstSize.Y,
+		SrcRect.Min.X, SrcRect.Min.Y,
+		SrcRect.Width(), SrcRect.Height(),
+		DstSize,
+		SrcSize,
+		*VertexShader,
+		EDRF_UseTriangleOptimization);
+
+	Context.RHICmdList.CopyToResolveTarget(DestRenderTarget.TargetableTexture, DestRenderTarget.ShaderResourceTexture, FResolveParams());
+}
+
+FPooledRenderTargetDesc FRCDistortionMergePassES2::ComputeOutputDesc(EPassOutputId InPassOutputId) const
+{
+	FPooledRenderTargetDesc Ret;
+	Ret.Depth = 0;
+	Ret.ArraySize = 1;
+	Ret.bIsArray = false;
+	Ret.NumMips = 1;
+	Ret.TargetableFlags = TexCreate_RenderTargetable;
+	Ret.bForceSeparateTargetAndShaderResource = false;
+	Ret.Format = PF_FloatRGBA;
+	Ret.NumSamples = 1;
+	Ret.Extent.X = FMath::Max(1, PrePostSourceViewportSize.X);
+	Ret.Extent.Y = FMath::Max(1, PrePostSourceViewportSize.Y);
+	Ret.DebugName = TEXT("DistortionMergePass");
+	Ret.ClearValue = FClearValueBinding(FLinearColor::Black);
+	return Ret;
+}
diff --git a/Engine/Source/Runtime/Renderer/Private/MobileDistortionPass.h b/Engine/Source/Runtime/Renderer/Private/MobileDistortionPass.h
new file mode 100644
--- /dev/null
+++ b/Engine/Source/Runtime/Renderer/Private/MobileDistortionPass.h
@@ -0,0 +1,33 @@
+// Copyright 1998-2018 Epic Games, Inc.All Rights Reserved.
+
+#pragma once
+
+#include "CoreMinimal.h"
+#include "PostProcess/RenderingCompositionGraph.h"
+
+class FRCDistortionAccumulatePassES2 : public TRenderingCompositePassBase<1, 1>
+{
+public:
+	FRCDistortionAccumulatePassES2(FIntPoint InPrePostSourceViewportSize) : PrePostSourceViewportSize(InPrePostSourceViewportSize) { }
+	virtual void Process(FRenderingCompositePassContext& Context) override;
+	virtual FPooledRenderTargetDesc ComputeOutputDesc(EPassOutputId InPassOutputId) const override;
+	virtual void Release() override { delete this; }
+	virtual const TCHAR* GetDebugName() { return TEXT("FRCDistortionAccumulatePassES2"); }
+private:
+	FIntPoint PrePostSourceViewportSize;
+};
+
+class FRCDistortionMergePassES2 : public TRenderingCompositePassBase<2, 1>
+{
+public:
+	FRCDistortionMergePassES2(FIntPoint InPrePostSourceViewportSize) : PrePostSourceViewportSize(InPrePostSourceViewportSize) { }
+	virtual void Process(FRenderingCompositePassContext& Context) override;
+	virtual FPooledRenderTargetDesc ComputeOutputDesc(EPassOutputId InPassOutputId) const override;
+	virtual void Release() override { delete this; }
+	virtual const TCHAR* GetDebugName() { return TEXT("FRCDistortionMergePassES2"); }
+private:
+	FIntPoint PrePostSourceViewportSize;
+};
+
+// Returns whether distortion is enabled and there primitives to draw
+bool IsMobileDistortionActive(const FViewInfo& View);
diff --git a/Engine/Source/Runtime/Renderer/Private/MobileShadingRenderer.cpp b/Engine/Source/Runtime/Renderer/Private/MobileShadingRenderer.cpp
--- a/Engine/Source/Runtime/Renderer/Private/MobileShadingRenderer.cpp
+++ b/Engine/Source/Runtime/Renderer/Private/MobileShadingRenderer.cpp
@@ -39,6 +39,7 @@
 #include "PipelineStateCache.h"
 #include "ClearQuad.h"
 #include "MobileSeparateTranslucencyPass.h"
+#include "MobileDistortionPass.h"
 
 
 uint32 GetShadowQuality();
@@ -352,16 +353,6 @@ void FMobileSceneRenderer::Render(FRHICommandListImmediate& RHICmdList)
 	if (ViewFamily.EngineShowFlags.Translucency)
 	{
 		SCOPE_CYCLE_COUNTER(STAT_TranslucencyDrawTime);
-
-		// Note: Forward pass has no SeparateTranslucency, so refraction effect order with Translucency is different.
-		// Having the distortion applied between two different translucency passes would make it consistent with the deferred pass.
-		// This is not done yet.
-
-		if (GetRefractionQuality(ViewFamily) > 0)
-		{
-			// to apply refraction effect by distorting the scene color
-			RenderDistortionES2(RHICmdList);
-		}
 		RenderTranslucency(RHICmdList, ViewList, !bGammaSpace || bRenderToSceneColor);
 		FRHICommandListExecutor::GetImmediateCommandList().PollOcclusionQueries();
 		RHICmdList.ImmediateFlush(EImmediateFlushType::DispatchToRHIThread);
@@ -384,15 +375,14 @@ void FMobileSceneRenderer::Render(FRHICommandListImmediate& RHICmdList)
 	}
 
 	static const auto CVarMobileMSAA = IConsoleManager::Get().FindTConsoleVariableDataInt(TEXT("r.MobileMSAA"));
-	bool bOnChipSunMask =
+	bool bOnChipPP =
 		GSupportsRenderTargetFormat_PF_FloatRGBA &&
 		GSupportsShaderFramebufferFetch &&
-		ViewFamily.EngineShowFlags.PostProcessing &&
-		((View.bLightShaftUse) || GetMobileDepthOfFieldScale(View) > 0.0 ||
-		((ViewFamily.GetShaderPlatform() == SP_METAL) && (CVarMobileMSAA ? CVarMobileMSAA->GetValueOnAnyThread() > 1 : false))
-		);
+		ViewFamily.EngineShowFlags.PostProcessing;
+	bool bOnChipSunMask = bOnChipPP && (View.bLightShaftUse || GetMobileDepthOfFieldScale(View) > 0.0) && !IsMobileDistortionActive(View);
+	bool bOnChipPreResolveMSAA = bOnChipPP && (ViewFamily.GetShaderPlatform() == SP_METAL) && (CVarMobileMSAA ? CVarMobileMSAA->GetValueOnAnyThread() > 1 : false);
 
-	if (!bGammaSpace && bOnChipSunMask)
+	if (!bGammaSpace && (bOnChipSunMask || bOnChipPreResolveMSAA))
 	{
 		// Convert alpha from depth to circle of confusion with sunshaft intensity.
 		// This is done before resolve on hardware with framebuffer fetch.
diff --git a/Engine/Source/Runtime/Renderer/Private/PostProcess/PostProcessing.cpp b/Engine/Source/Runtime/Renderer/Private/PostProcess/PostProcessing.cpp
--- a/Engine/Source/Runtime/Renderer/Private/PostProcess/PostProcessing.cpp
+++ b/Engine/Source/Runtime/Renderer/Private/PostProcess/PostProcessing.cpp
@@ -54,6 +54,7 @@
 #include "DeferredShadingRenderer.h"
 #include "PostProcess/PostProcessFFTBloom.h"
 #include "MobileSeparateTranslucencyPass.h"
+#include "MobileDistortionPass.h"
 
 /** The global center for all post processing activities. */
 FPostProcessing GPostProcessing;
@@ -2365,6 +2366,18 @@ void FPostProcessing::ProcessES2(FRHICommandListImmediate& RHICmdList, const FVi
 
 			// Post is not supported on ES2 devices using mosaic.
 			bool bUsePost = bHDRModeAllowsPost && IsMobileHDR();
+			
+			if (bUsePost && IsMobileDistortionActive(View))
+			{
+				FRenderingCompositePass* AccumulatedDistortion = Context.Graph.RegisterPass(new(FMemStack::Get()) FRCDistortionAccumulatePassES2(SceneColorSize));
+				AccumulatedDistortion->SetInput(ePId_Input0, Context.FinalOutput); // unused atm
+				FRenderingCompositeOutputRef AccumulatedDistortionRef(AccumulatedDistortion);
+				
+				FRenderingCompositePass* PostProcessDistorsion = Context.Graph.RegisterPass(new(FMemStack::Get()) FRCDistortionMergePassES2(SceneColorSize));
+				PostProcessDistorsion->SetInput(ePId_Input0, Context.FinalOutput);
+				PostProcessDistorsion->SetInput(ePId_Input1, AccumulatedDistortionRef);
+				Context.FinalOutput = FRenderingCompositeOutputRef(PostProcessDistorsion);
+			}
 
 			// Always evaluate custom post processes
 			if (bUsePost)
diff --git a/Engine/Source/Runtime/Renderer/Private/SceneRendering.h b/Engine/Source/Runtime/Renderer/Private/SceneRendering.h
--- a/Engine/Source/Runtime/Renderer/Private/SceneRendering.h
+++ b/Engine/Source/Runtime/Renderer/Private/SceneRendering.h
@@ -1523,6 +1523,8 @@ public:
 		return ScreenPercentageInterface->Fork_GameThread(ForkedViewFamily);
 	}
 
+	static int32 GetRefractionQuality(const FSceneViewFamily& ViewFamily);
+	
 protected:
 
 	/** Size of the family. */
@@ -1659,7 +1661,6 @@ protected:
 
 	/** Renders the scene's distortion */
 	void RenderDistortion(FRHICommandListImmediate& RHICmdList);
-	void RenderDistortionES2(FRHICommandListImmediate& RHICmdList);
 
 	/** Returns the scene color texture multi-view is targeting. */	
 	FTextureRHIParamRef GetMultiViewSceneColor(const FSceneRenderTargets& SceneContext) const;
@@ -1670,8 +1671,6 @@ protected:
 	/** Renders a depth mask into the monoscopic far field view to ensure we only render visible pixels. */
 	void RenderMonoscopicFarFieldMask(FRHICommandListImmediate& RHICmdList);
 
-	static int32 GetRefractionQuality(const FSceneViewFamily& ViewFamily);
-
 	void UpdatePrimitivePrecomputedLightingBuffers();
 	void ClearPrimitiveSingleFramePrecomputedLightingBuffers();
 
