RL - The decal Stencil mask is cleared for every pass in the post decal rendering. This also cleared the maks for emissive (which probably was added later as the clear is quite old). Moving this to the very last pass (which is the emissive) fixed the problems.
The passes are always added even if the decal list is empty. Renderdoc confirmed that the clear does happen.

diff --git a/Engine/Source/Runtime/Renderer/Private/CompositionLighting/PostProcessDeferredDecals.cpp b/Engine/Source/Runtime/Renderer/Private/CompositionLighting/PostProcessDeferredDecals.cpp
--- a/Engine/Source/Runtime/Renderer/Private/CompositionLighting/PostProcessDeferredDecals.cpp
+++ b/Engine/Source/Runtime/Renderer/Private/CompositionLighting/PostProcessDeferredDecals.cpp
@@ -934,7 +934,10 @@ void FRCPassPostProcessDeferredDecals::Process(FRenderingCompositePassContext& C
 					RHICmdList.DrawIndexedPrimitive(GetUnitCubeIndexBuffer(), PT_TriangleList, 0, 0, 8, 0, ARRAY_COUNT(GCubeIndices) / 3, 1);
 					RenderTargetManager.bGufferADirty |= (RenderTargetManager.TargetsToResolve[FDecalRenderTargetManager::GBufferAIndex] != nullptr);
 				}
+			}
 
+			if (CurrentStage == DRS_Emissive)
+			{
 				// we don't modify stencil but if out input was having stencil for us (after base pass - we need to clear)
 				// Clear stencil to 0, which is the assumed default by other passes
 				DrawClearQuad(RHICmdList, false, FLinearColor(), false, 0, true, 0, SceneContext.GetSceneDepthSurface()->GetSizeXY(), FIntRect());
