Fixed: PostProcessing materails that sample SceneDepth fail when MSAA is enabled on Mobile
SceneDepth can only be sampled in PP materails that have blendable location set to BeforeTranslucency or BeforeTonemapping

diff --git a/Engine/Shaders/Private/MaterialTemplate.ush b/Engine/Shaders/Private/MaterialTemplate.ush
--- a/Engine/Shaders/Private/MaterialTemplate.ush
+++ b/Engine/Shaders/Private/MaterialTemplate.ush
@@ -1362,7 +1362,7 @@ MaterialFloat4 MobileSceneTextureLookup(inout FMaterialPixelParameters Parameter
 	//PPI_SceneDepth
 	if (SceneTextureId == 1)
 	{
-		MaterialFloat Depth = ConvertFromDeviceZ(Texture2DSample(MobileSceneTextures.SceneDepthTexture, MobileSceneTextures.SceneDepthTextureSampler, UV).r);
+		MaterialFloat Depth = Texture2DSample(MobileSceneTextures.SceneColorTexture, MobileSceneTextures.SceneColorTextureSampler, UV).a;
 		return MaterialFloat4(Depth.rrr, 0.f);
 	}
 	//PPI_CustomDepth
diff --git a/Engine/Shaders/Private/MobileBasePassPixelShader.usf b/Engine/Shaders/Private/MobileBasePassPixelShader.usf
--- a/Engine/Shaders/Private/MobileBasePassPixelShader.usf
+++ b/Engine/Shaders/Private/MobileBasePassPixelShader.usf
@@ -640,7 +640,7 @@ void Main(
 			if (GetHDR32bppEncodeMode() == HDR_ENCODE_NONE)
 			{
 				// FP16 HDR stores Z in framebuffer alpha
-				OutColor.a = BasePassInterpolants.PixelPosition.w; 
+				OutColor.a = min(BasePassInterpolants.PixelPosition.w, 65500); 
 			}
 		#endif
 	#endif
diff --git a/Engine/Source/Runtime/Engine/Private/Materials/HLSLMaterialTranslator.h b/Engine/Source/Runtime/Engine/Private/Materials/HLSLMaterialTranslator.h
--- a/Engine/Source/Runtime/Engine/Private/Materials/HLSLMaterialTranslator.h
+++ b/Engine/Source/Runtime/Engine/Private/Materials/HLSLMaterialTranslator.h
@@ -3787,6 +3787,14 @@ protected:
 			// On mobile in post process material, there is no need to do ViewportUV->BufferUV conversion because ViewSize == BufferSize.
 			if (Material->GetMaterialDomain() == MD_PostProcess)
 			{
+				int32 BlendableLocation = Material->GetBlendableLocation();
+				if (SceneTextureId == PPI_SceneDepth && !(BlendableLocation == BL_BeforeTranslucency || BlendableLocation == BL_BeforeTonemapping))
+				{
+					// SceneDepth lookups are not available when using MSAA, but we can access depth stored in SceneColor.A channel
+					// SceneColor.A channel holds depth till BeforeTonemapping location, then it's gets overwritten
+					return Errorf(TEXT("SceneDepth lookups are only available when BlendableLocation is BeforeTranslucency or BeforeTonemapping"));
+				}
+				
 				if (ViewportUV == INDEX_NONE)
 				{
 					UV = TextureCoordinate(0, false, false);
diff --git a/Engine/Source/Runtime/Renderer/Private/MobileShadingRenderer.cpp b/Engine/Source/Runtime/Renderer/Private/MobileShadingRenderer.cpp
--- a/Engine/Source/Runtime/Renderer/Private/MobileShadingRenderer.cpp
+++ b/Engine/Source/Runtime/Renderer/Private/MobileShadingRenderer.cpp
@@ -78,7 +78,6 @@ FMobileSceneRenderer::FMobileSceneRenderer(const FSceneViewFamily* InViewFamily,
 	:	FSceneRenderer(InViewFamily, HitProxyConsumer)
 {
 	bModulatedShadowsInUse = false;
-	bPostProcessUsesDepthTexture = false;
 }
 
 TUniformBufferRef<FMobileDirectionalLightShaderParameters>& GetNullMobileDirectionalLightShaderParameters()
@@ -131,9 +130,6 @@ void FMobileSceneRenderer::InitViews(FRHICommandListImmediate& RHICmdList)
 	// Now that the indirect lighting cache is updated, we can update the primitive precomputed lighting buffers.
 	UpdatePrimitivePrecomputedLightingBuffers();
 
-	UpdatePostProcessUsageFlags();
-
-	
 	PostInitViewCustomData();
 	
 	OnStartFrame(RHICmdList);
@@ -408,7 +404,7 @@ void FMobileSceneRenderer::Render(FRHICommandListImmediate& RHICmdList)
 		const bool bForceDepthResolve = CVarMobileForceDepthResolve.GetValueOnRenderThread() == 1;
 		const bool bSeparateTranslucencyActive = IsMobileSeparateTranslucencyActive(View);
 
-		bKeepDepthContent = bForceDepthResolve || bPostProcessUsesDepthTexture || bSeparateTranslucencyActive ||
+		bKeepDepthContent = bForceDepthResolve || bSeparateTranslucencyActive ||
 			(View.bIsSceneCapture && (ViewFamily.SceneCaptureSource == ESceneCaptureSource::SCS_SceneColorHDR || ViewFamily.SceneCaptureSource == ESceneCaptureSource::SCS_SceneColorSceneDepth));
 	}
 
@@ -808,32 +804,3 @@ void FMobileSceneRenderer::CopyMobileMultiViewSceneColor(FRHICommandListImmediat
 			EDRF_UseTriangleOptimization);
 	}
 }
-
-void FMobileSceneRenderer::UpdatePostProcessUsageFlags()
-{
-	bPostProcessUsesDepthTexture = false;
-	// Find out whether post-process materials require SceneDepth lookups, otherwise renderer can discard depth buffer before starting post-processing pass
-	for (int32 ViewIndex = 0; ViewIndex < Views.Num(); ViewIndex++)
-	{
-		const FBlendableManager& BlendableManager = Views[ViewIndex].FinalPostProcessSettings.BlendableManager;
-		FBlendableEntry* BlendableIt = nullptr;
-
-		while (FPostProcessMaterialNode* DataPtr = BlendableManager.IterateBlendables<FPostProcessMaterialNode>(BlendableIt))
-		{
-			if (DataPtr->IsValid())
-			{
-				FMaterialRenderProxy* Proxy = DataPtr->GetMaterialInterface()->GetRenderProxy(false);
-				check(Proxy);
-
-				const FMaterial* Material = Proxy->GetMaterial(Views[ViewIndex].GetFeatureLevel());
-				check(Material);
-
-				if (Material->MaterialUsesSceneDepthLookup_RenderThread())
-				{
-					bPostProcessUsesDepthTexture = true;
-					break;
-				}
-			}
-		}
-	}
-}
diff --git a/Engine/Source/Runtime/Renderer/Private/SceneRendering.h b/Engine/Source/Runtime/Renderer/Private/SceneRendering.h
--- a/Engine/Source/Runtime/Renderer/Private/SceneRendering.h
+++ b/Engine/Source/Runtime/Renderer/Private/SceneRendering.h
@@ -1746,9 +1746,6 @@ protected:
 	/** Copy scene color from the mobile multi-view render targat array to side by side stereo scene color */
 	void CopyMobileMultiViewSceneColor(FRHICommandListImmediate& RHICmdList);
 
-	/** Gather information about post-processing pass, which can be used by render for optimizations. Called by InitViews */
-	void UpdatePostProcessUsageFlags();
-
 	/** Render inverse opacity for the dynamic meshes. */
 	bool RenderInverseOpacityDynamic(FRHICommandListImmediate& RHICmdList, const FViewInfo& View, const FDrawingPolicyRenderState& DrawRenderState);
 
@@ -1756,9 +1753,7 @@ protected:
 	void PostInitViewCustomData();
 	
 private:
-
 	bool bModulatedShadowsInUse;
-	bool bPostProcessUsesDepthTexture;
 };
 
 // The noise textures need to be set in Slate too.
