Exposed 'Support Software Occlusion Culling' option to project settings

diff --git a/Engine/Source/Developer/Android/AndroidTargetPlatform/Private/AndroidTargetPlatform.cpp b/Engine/Source/Developer/Android/AndroidTargetPlatform/Private/AndroidTargetPlatform.cpp
--- a/Engine/Source/Developer/Android/AndroidTargetPlatform/Private/AndroidTargetPlatform.cpp
+++ b/Engine/Source/Developer/Android/AndroidTargetPlatform/Private/AndroidTargetPlatform.cpp
@@ -250,14 +250,8 @@ bool FAndroidTargetPlatform::SupportsVulkan() const
 
 bool FAndroidTargetPlatform::SupportsSoftwareOcclusion() const
 {
-	// default to not supporting
-	bool bSupportsSoftwareOcclusion = false;
-#if WITH_ENGINE
-	int32 IntValue = 0;
-	GConfig->GetInt(TEXT("ConsoleVariables"), TEXT("r.Mobile.AllowSoftwareOcclusion"), IntValue, GEngineIni);
-	bSupportsSoftwareOcclusion = (IntValue != 0);
-#endif
-	return bSupportsSoftwareOcclusion;
+	static auto* CVarMobileAllowSoftwareOcclusion = IConsoleManager::Get().FindTConsoleVariableDataInt(TEXT("r.Mobile.AllowSoftwareOcclusion"));
+	return CVarMobileAllowSoftwareOcclusion->GetValueOnAnyThread() != 0;
 }
 
 /* ITargetPlatform overrides
diff --git a/Engine/Source/Developer/IOS/IOSTargetPlatform/Private/IOSTargetPlatform.cpp b/Engine/Source/Developer/IOS/IOSTargetPlatform/Private/IOSTargetPlatform.cpp
--- a/Engine/Source/Developer/IOS/IOSTargetPlatform/Private/IOSTargetPlatform.cpp
+++ b/Engine/Source/Developer/IOS/IOSTargetPlatform/Private/IOSTargetPlatform.cpp
@@ -484,10 +484,8 @@ static bool CookASTC()
 
 static bool SupportsSoftwareOcclusion()
 {
-	// default to not support
-	int32 IntValue = 0;
-	GConfig->GetInt(TEXT("ConsoleVariables"), TEXT("r.Mobile.AllowSoftwareOcclusion"), IntValue, GEngineIni);
-	return (IntValue != 0);
+	static auto* CVarMobileAllowSoftwareOcclusion = IConsoleManager::Get().FindTConsoleVariableDataInt(TEXT("r.Mobile.AllowSoftwareOcclusion"));
+	return CVarMobileAllowSoftwareOcclusion->GetValueOnAnyThread() != 0;
 }
 
 bool FIOSTargetPlatform::CanSupportXGEShaderCompile() const
diff --git a/Engine/Source/Runtime/Engine/Classes/Engine/RendererSettings.h b/Engine/Source/Runtime/Engine/Classes/Engine/RendererSettings.h
--- a/Engine/Source/Runtime/Engine/Classes/Engine/RendererSettings.h
+++ b/Engine/Source/Runtime/Engine/Classes/Engine/RendererSettings.h
@@ -202,6 +202,12 @@ class ENGINE_API URendererSettings : public UDeveloperSettings
 		ToolTip="Whether to support 'Dithered LOD Transition' material option on mobile platforms. Enabling this may degrade performance as rendering will not benefit from Early-Z optimization.",
 		ConfigRestartRequired=true))
 	uint32 bMobileAllowDitheredLODTransition:1;
+
+	UPROPERTY(config, EditAnywhere, Category = Mobile, meta=(
+		ConsoleVariable="r.Mobile.AllowSoftwareOcclusion", DisplayName="Support Software Occlusion Culling",
+		ToolTip="Whether to support 'Software Occlusion Culling' on mobile platforms. This this package occluder information and enable Software Occlusion Culling.",
+		ConfigRestartRequired=false))
+	uint32 bMobileAllowSoftwareOcclusionCulling:1;
 	
 	UPROPERTY(config, EditAnywhere, Category = Materials, meta = (
 		ConsoleVariable = "r.DiscardUnusedQuality", DisplayName = "Game Discards Unused Material Quality Levels",
diff --git a/Engine/Source/Runtime/Renderer/Private/SceneOcclusion.cpp b/Engine/Source/Runtime/Renderer/Private/SceneOcclusion.cpp
--- a/Engine/Source/Runtime/Renderer/Private/SceneOcclusion.cpp
+++ b/Engine/Source/Runtime/Renderer/Private/SceneOcclusion.cpp
@@ -52,10 +52,9 @@ FAutoConsoleVariableRef CVarOcclusionCullCascadedShadowMaps(
 	ECVF_Scalability | ECVF_RenderThreadSafe
 	);
 
-int32 GMobileAllowSoftwareOcclusion = 0;
-FAutoConsoleVariableRef CVarMobileAllowSoftwareOcclusion(
+static TAutoConsoleVariable<int32> CVarMobileAllowSoftwareOcclusion(
 	TEXT("r.Mobile.AllowSoftwareOcclusion"),
-	GMobileAllowSoftwareOcclusion,
+	0,
 	TEXT("Whether to allow rasterizing scene on CPU for primitive occlusion.\n"),
 	ECVF_RenderThreadSafe
 	);
@@ -263,7 +262,8 @@ bool FSceneViewState::IsShadowOccluded(FRHICommandListImmediate& RHICmdList, FSc
 
 void FSceneViewState::ConditionallyAllocateSceneSoftwareOcclusion(ERHIFeatureLevel::Type InFeatureLevel)
 {
-	bool bShouldBeEnabled = InFeatureLevel <= ERHIFeatureLevel::ES3_1 && GMobileAllowSoftwareOcclusion != 0;
+	bool bMobileAllowSoftwareOcclusion = CVarMobileAllowSoftwareOcclusion.GetValueOnAnyThread() != 0;
+	bool bShouldBeEnabled = InFeatureLevel <= ERHIFeatureLevel::ES3_1 && bMobileAllowSoftwareOcclusion;
 
 	if (bShouldBeEnabled && !SceneSoftwareOcclusion)
 	{
