REL - Clear decals using fast-clear as RT are not always bound

diff --git a/Engine/Source/Runtime/Renderer/Private/CompositionLighting/PostProcessDeferredDecals.cpp b/Engine/Source/Runtime/Renderer/Private/CompositionLighting/PostProcessDeferredDecals.cpp
--- a/Engine/Source/Runtime/Renderer/Private/CompositionLighting/PostProcessDeferredDecals.cpp
+++ b/Engine/Source/Runtime/Renderer/Private/CompositionLighting/PostProcessDeferredDecals.cpp
@@ -936,11 +936,17 @@ void FRCPassPostProcessDeferredDecals::Process(FRenderingCompositePassContext& C
 				}
 			}
 
-			if (CurrentStage == DRS_Emissive)
+			if ((Scene.Decals.Num() > 0) && CurrentStage == DRS_Emissive)
 			{
 				// we don't modify stencil but if out input was having stencil for us (after base pass - we need to clear)
 				// Clear stencil to 0, which is the assumed default by other passes
-				DrawClearQuad(RHICmdList, false, FLinearColor(), false, 0, true, 0, SceneContext.GetSceneDepthSurface()->GetSizeXY(), FIntRect());
+				FRHIDepthRenderTargetView DepthView(SceneContext.GetSceneDepthTexture(), 
+													ERenderTargetLoadAction::ENoAction, ERenderTargetStoreAction::ENoAction, 
+													ERenderTargetLoadAction::EClear,	ERenderTargetStoreAction::EStore, 
+													FExclusiveDepthStencil(FExclusiveDepthStencil::DepthNop_StencilWrite));
+				FRHISetRenderTargetsInfo Info(0, nullptr, DepthView);
+				RHICmdList.SetRenderTargetsAndClear(Info);
+				//DrawClearQuad(RHICmdList, false, FLinearColor(), false, 0, true, 0, SceneContext.GetSceneDepthSurface()->GetSizeXY(), FIntRect());
 			}
 
 			// This stops the targets from being resolved and decoded until the last view is rendered.
