PR #5166: Exposure bias does not work on mobile anymore. (Contributed by [REDACTED_PERSON])

diff --git a/Engine/Source/Runtime/Renderer/Private/PostProcess/PostProcessEyeAdaptation.cpp b/Engine/Source/Runtime/Renderer/Private/PostProcess/PostProcessEyeAdaptation.cpp
--- a/Engine/Source/Runtime/Renderer/Private/PostProcess/PostProcessEyeAdaptation.cpp
+++ b/Engine/Source/Runtime/Renderer/Private/PostProcess/PostProcessEyeAdaptation.cpp
@@ -74,7 +74,7 @@ inline static void ComputeEyeAdaptationValues(const ERHIFeatureLevel::Type MinFe
 	float MaxAverageLuminance = 1;
 
 	// Disable exposure for buffer visualization. so it doesn't influence buffer visualization brightness.
-	const bool bDisableExposure = !(EngineShowFlags.Lighting && EngineShowFlags.EyeAdaptation && View.GetFeatureLevel() >= MinFeatureLevel)
+	const bool bDisableExposure = !(EngineShowFlags.Lighting && EngineShowFlags.EyeAdaptation)
 		|| (EngineShowFlags.VisualizeBuffer && View.CurrentBufferVisualizationMode != NAME_None)
 		|| View.Family->UseDebugViewPS();
 
