Fix a crash that can occur when mobile shadows are disabled via showflags.

diff --git a/Engine/Source/Runtime/Renderer/Private/MobileBasePassRendering.cpp b/Engine/Source/Runtime/Renderer/Private/MobileBasePassRendering.cpp
--- a/Engine/Source/Runtime/Renderer/Private/MobileBasePassRendering.cpp
+++ b/Engine/Source/Runtime/Renderer/Private/MobileBasePassRendering.cpp
@@ -587,7 +587,7 @@ static void DrawVisibleFrontToBack(
 	const EBasePassDrawListType DrawListType,
 	const FViewInfo& View,
 	FDrawingPolicyRenderState& DrawRenderState,
-	const FMobileCSMVisibilityInfo* const MobileCSMVisibilityInfo,
+	const FMobileCSMVisibilityInfo& MobileCSMVisibilityInfo,
 	int32& MaxDraws)
 {
 #if FRAMEPRO_ENABLED
@@ -596,12 +596,12 @@ static void DrawVisibleFrontToBack(
 	QUICK_SCOPE_CYCLE_COUNTER(STAT_StaticDrawListDrawTimeFrontToBack);
 #endif // FRAMEPRO_ENABLED
 
-	const bool bIsCSM = MobileCSMVisibilityInfo != nullptr;
+	const bool bIsCSM = MobileCSMVisibilityInfo.bMobileDynamicCSMInUse;
 	int32 NumDraws = 0;
 	if (bIsCSM)
 	{
-		NumDraws += Scene->MobileBasePassUniformLightMapPolicyDrawListWithCSM[DrawListType].DrawVisibleFrontToBack(RHICmdList, View, DrawRenderState, MobileCSMVisibilityInfo->MobileCSMStaticMeshVisibilityMap, MobileCSMVisibilityInfo->MobileCSMStaticBatchVisibility, MaxDraws);
-		NumDraws += Scene->MobileBasePassUniformLightMapPolicyDrawList[DrawListType].DrawVisibleFrontToBack(RHICmdList, View, DrawRenderState, MobileCSMVisibilityInfo->MobileNonCSMStaticMeshVisibilityMap, MobileCSMVisibilityInfo->MobileNonCSMStaticBatchVisibility, MaxDraws);
+		NumDraws += Scene->MobileBasePassUniformLightMapPolicyDrawListWithCSM[DrawListType].DrawVisibleFrontToBack(RHICmdList, View, DrawRenderState, MobileCSMVisibilityInfo.MobileCSMStaticMeshVisibilityMap, MobileCSMVisibilityInfo.MobileCSMStaticBatchVisibility, MaxDraws);
+		NumDraws += Scene->MobileBasePassUniformLightMapPolicyDrawList[DrawListType].DrawVisibleFrontToBack(RHICmdList, View, DrawRenderState, MobileCSMVisibilityInfo.MobileNonCSMStaticMeshVisibilityMap, MobileCSMVisibilityInfo.MobileNonCSMStaticBatchVisibility, MaxDraws);
 	}
 	else
 	{
@@ -622,14 +622,14 @@ static void DrawVisible(
 	const EBasePassDrawListType DrawListType,
 	const FViewInfo& View,
 	const FDrawingPolicyRenderState& DrawRenderState,
-	const FMobileCSMVisibilityInfo* const MobileCSMVisibilityInfo)
+	const FMobileCSMVisibilityInfo& MobileCSMVisibilityInfo)
 {
 	SCOPE_CYCLE_COUNTER(STAT_StaticDrawListDrawTime);
-	const bool bIsCSM = MobileCSMVisibilityInfo != nullptr;
+	const bool bIsCSM = MobileCSMVisibilityInfo.bMobileDynamicCSMInUse;
 	if (bIsCSM)
 	{
-		Scene->MobileBasePassUniformLightMapPolicyDrawListWithCSM[DrawListType].DrawVisible(RHICmdList, View, DrawRenderState, MobileCSMVisibilityInfo->MobileCSMStaticMeshVisibilityMap, MobileCSMVisibilityInfo->MobileCSMStaticBatchVisibility);
-		Scene->MobileBasePassUniformLightMapPolicyDrawList[DrawListType].DrawVisible(RHICmdList, View, DrawRenderState, MobileCSMVisibilityInfo->MobileNonCSMStaticMeshVisibilityMap, MobileCSMVisibilityInfo->MobileNonCSMStaticBatchVisibility);
+		Scene->MobileBasePassUniformLightMapPolicyDrawListWithCSM[DrawListType].DrawVisible(RHICmdList, View, DrawRenderState, MobileCSMVisibilityInfo.MobileCSMStaticMeshVisibilityMap, MobileCSMVisibilityInfo.MobileCSMStaticBatchVisibility);
+		Scene->MobileBasePassUniformLightMapPolicyDrawList[DrawListType].DrawVisible(RHICmdList, View, DrawRenderState, MobileCSMVisibilityInfo.MobileNonCSMStaticMeshVisibilityMap, MobileCSMVisibilityInfo.MobileNonCSMStaticBatchVisibility);
 	}
 	else
 	{
@@ -979,11 +979,11 @@ void FMobileSceneRenderer::RenderMobileBasePass(FRHICommandListImmediate& RHICmd
 			// Render the base pass static data
 			if (SortMode == EBasePassSort::SortPerMesh)
 			{
-				DrawVisibleFrontToBack(RHICmdList, Scene, EBasePass_Default, View, DrawRenderState, &View.MobileCSMVisibilityInfo, MaxDraws);
+				DrawVisibleFrontToBack(RHICmdList, Scene, EBasePass_Default, View, DrawRenderState, View.MobileCSMVisibilityInfo, MaxDraws);
 			}
 			else
 			{
-				DrawVisible(RHICmdList, Scene, EBasePass_Default, View, DrawRenderState, &View.MobileCSMVisibilityInfo);
+				DrawVisible(RHICmdList, Scene, EBasePass_Default, View, DrawRenderState, View.MobileCSMVisibilityInfo);
 			}
 			FRHICommandListExecutor::GetImmediateCommandList().PollOcclusionQueries();
 			FRHICommandListExecutor::GetImmediateCommandList().ImmediateFlush(EImmediateFlushType::DispatchToRHIThread);
@@ -1000,11 +1000,11 @@ void FMobileSceneRenderer::RenderMobileBasePass(FRHICommandListImmediate& RHICmd
 			// Issue static draw list masked draw calls last, as PVR wants it
 			if (SortMode == EBasePassSort::SortPerMesh)
 			{
-				DrawVisibleFrontToBack(RHICmdList, Scene, EBasePass_Masked, View, DrawRenderState, &View.MobileCSMVisibilityInfo, MaxDraws);
+				DrawVisibleFrontToBack(RHICmdList, Scene, EBasePass_Masked, View, DrawRenderState, View.MobileCSMVisibilityInfo, MaxDraws);
 			}
 			else
 			{
-				DrawVisible(RHICmdList, Scene, EBasePass_Masked, View, DrawRenderState, &View.MobileCSMVisibilityInfo);
+				DrawVisible(RHICmdList, Scene, EBasePass_Masked, View, DrawRenderState, View.MobileCSMVisibilityInfo);
 			}
 
 			FRHICommandListExecutor::GetImmediateCommandList().PollOcclusionQueries();
