Enables Diaphragm DOF for switch

diff --git a/Engine/Source/Runtime/Renderer/Private/PostProcess/DiaphragmDOF.h b/Engine/Source/Runtime/Renderer/Private/PostProcess/DiaphragmDOF.h
--- a/Engine/Source/Runtime/Renderer/Private/PostProcess/DiaphragmDOF.h
+++ b/Engine/Source/Runtime/Renderer/Private/PostProcess/DiaphragmDOF.h
@@ -122,7 +122,9 @@ inline bool IsSupported(EShaderPlatform ShaderPlatform)
 		ShaderPlatform == SP_XBOXONE_D3D12 ||
 		ShaderPlatform == SP_PS4 ||
 		IsVulkanSM5Platform(ShaderPlatform) ||
-		IsMetalSM5Platform(ShaderPlatform);
+		IsMetalSM5Platform(ShaderPlatform) ||
+		ShaderPlatform == SP_SWITCH ||
+		ShaderPlatform == SP_SWITCH_FORWARD;
 }
 
 
diff --git a/Engine/Source/Runtime/Renderer/Private/PostProcess/DiaphragmDOFGraph.cpp b/Engine/Source/Runtime/Renderer/Private/PostProcess/DiaphragmDOFGraph.cpp
--- a/Engine/Source/Runtime/Renderer/Private/PostProcess/DiaphragmDOFGraph.cpp
+++ b/Engine/Source/Runtime/Renderer/Private/PostProcess/DiaphragmDOFGraph.cpp
@@ -272,6 +272,12 @@ bool DiaphragmDOF::WireSceneColorPasses(FPostprocessContext& Context, const FRen
 	bool bForegroundHybridScattering = FgdHybridScatteringMode != EHybridScatterMode::Disabled && AbsMaxForegroundCocRadius > MinScatteringCocRadius && MaxScatteringRatio > 0.0f;
 	bool bBackgroundHybridScattering = BgdHybridScatteringMode != EHybridScatterMode::Disabled && MaxBackgroundCocRadius > MinScatteringCocRadius && MaxScatteringRatio > 0.0f;
 
+	if (!FRCPassDiaphragmDOFHybridScatter::IsSupported(ShaderPlatform))
+	{
+		bForegroundHybridScattering = false;
+		bBackgroundHybridScattering = false;
+	}
+
 	// Compute the reference buffer size for PrefilteringResolutionDivisor.
 	FSceneRenderTargets& SceneContext = FSceneRenderTargets::Get(Context.RHICmdList);
 	FIntPoint RefBufferSize = FIntPoint::DivideAndRoundUp(SceneContext.GetBufferSizeXY(), PrefilteringResolutionDivisor);
diff --git a/Engine/Source/Runtime/Renderer/Private/PostProcess/DiaphragmDOFPasses.cpp b/Engine/Source/Runtime/Renderer/Private/PostProcess/DiaphragmDOFPasses.cpp
--- a/Engine/Source/Runtime/Renderer/Private/PostProcess/DiaphragmDOFPasses.cpp
+++ b/Engine/Source/Runtime/Renderer/Private/PostProcess/DiaphragmDOFPasses.cpp
@@ -832,6 +832,14 @@ class FPostProcessDiaphragmDOFReduceCS : public FPostProcessDiaphragmDOFShader
 			return false;
 		}
 
+		if (!FRCPassDiaphragmDOFHybridScatter::IsSupported(Parameters.Platform))
+		{
+			if (PermutationVector.Get<FHybridScatterForeground>() || PermutationVector.Get<FHybridScatterBackground>())
+			{
+				return false;
+			}
+		}
+
 		return FPostProcessDiaphragmDOFShader::ShouldCompilePermutation(Parameters);
 	}
 };
@@ -840,6 +848,11 @@ class FPostProcessDiaphragmDOFScatterGroupPackCS : public FPostProcessDiaphragmD
 {
 	DECLARE_GLOBAL_SHADER(FPostProcessDiaphragmDOFScatterGroupPackCS);
 	SHADER_TYPE_PARAMETERS(FPostProcessDiaphragmDOFScatterGroupPackCS, FPostProcessDiaphragmDOFShader, REDUCE_SHADER_PARAMS);
+
+	static bool ShouldCompilePermutation(const FGlobalShaderPermutationParameters& Parameters)
+	{
+		return FRCPassDiaphragmDOFHybridScatter::IsSupported(Parameters.Platform);
+	}
 };
 
 IMPLEMENT_GLOBAL_SHADER(FPostProcessDiaphragmDOFReduceCS, "/Engine/Private/DiaphragmDOF/DOFReduce.usf", "ReduceCS", SF_Compute);
diff --git a/Engine/Source/Runtime/Renderer/Private/PostProcess/DiaphragmDOFPasses.h b/Engine/Source/Runtime/Renderer/Private/PostProcess/DiaphragmDOFPasses.h
--- a/Engine/Source/Runtime/Renderer/Private/PostProcess/DiaphragmDOFPasses.h
+++ b/Engine/Source/Runtime/Renderer/Private/PostProcess/DiaphragmDOFPasses.h
@@ -520,7 +520,7 @@ public:
 	/** Returns whether hybrid scattering is supported. */
 	static FORCEINLINE bool IsSupported(EShaderPlatform ShaderPlatform)
 	{
-		return true;
+		return !IsSwitchPlatform(ShaderPlatform);
 	}
 
 	// TODO: pixel shader post filtering to have hierarchical RT layout.
