Fixed "r.EyeAdaptationQuality 0" disabling the effect of AutoExposureBias.

diff --git a/Engine/Source/Runtime/Renderer/Private/PostProcess/PostProcessEyeAdaptation.cpp b/Engine/Source/Runtime/Renderer/Private/PostProcess/PostProcessEyeAdaptation.cpp
--- a/Engine/Source/Runtime/Renderer/Private/PostProcess/PostProcessEyeAdaptation.cpp
+++ b/Engine/Source/Runtime/Renderer/Private/PostProcess/PostProcessEyeAdaptation.cpp
@@ -67,53 +67,60 @@ inline static void ComputeEyeAdaptationValues(const ERHIFeatureLevel::Type MinFe
 	const FPostProcessSettings& Settings = View.FinalPostProcessSettings;
 	const FEngineShowFlags& EngineShowFlags = View.Family->EngineShowFlags;
 
+	// ----------
+
+	// Histogram related values.
 	float LowPercent = FMath::Clamp(Settings.AutoExposureLowPercent, 1.0f, 99.0f) * 0.01f;
 	float HighPercent = FMath::Clamp(Settings.AutoExposureHighPercent, 1.0f, 99.0f) * 0.01f;
-
-	float MinAverageLuminance = 1;
-	float MaxAverageLuminance = 1;
-
-	// Disable exposure for buffer visualization. so it doesn't influence buffer visualization brightness.
-	const bool bDisableExposure = !(EngineShowFlags.Lighting && EngineShowFlags.EyeAdaptation)
-		|| (EngineShowFlags.VisualizeBuffer && View.CurrentBufferVisualizationMode != NAME_None)
-		|| View.Family->UseDebugViewPS();
-
-	// This scales the average luminance after it gets clamped, affecting the exposure value directly.
-	const float LocalExposureMultipler = (bDisableExposure || View.Family->ExposureSettings.bFixed) ? 1.f : FMath::Pow(2.0f, Settings.AutoExposureBias);
-
-	// This is used in the basic mode as a premultiplier before clamping between [MinLuminance, MaxLuminance]
-	// Because this happens before clamping, the calibration constant has no impact on fixed EV100 modes.
-	const float CalibrationConstant = FMath::Clamp<float>(Settings.AutoExposureCalibrationConstant, 1.f, 100.f) * 0.01f;
-	float AverageLuminanceScale = Settings.AutoExposureMethod == EAutoExposureMethod::AEM_Basic ? (1.f / CalibrationConstant) : 1.f;
-
-	// example min/max: -8 .. 4   means a range from 1/256 to 4  pow(2,-8) .. pow(2,4)
 	float HistogramLogMin = bExtendedLuminanceRange ? EV100ToLog2(Settings.HistogramLogMin) : Settings.HistogramLogMin;
 	float HistogramLogMax = bExtendedLuminanceRange ? EV100ToLog2(Settings.HistogramLogMax) : Settings.HistogramLogMax;
 	HistogramLogMin = FMath::Min<float>(HistogramLogMin, HistogramLogMax - 1);
 
-	// Eye adaptation is disabled except for highend right now because the histogram is not computed.
-	if (!bDisableExposure && EngineShowFlags.Lighting && EngineShowFlags.EyeAdaptation && View.GetFeatureLevel() >= MinFeatureLevel)
+	// ----------
+
+	// Those clamps the average luminance computed from the scene color.
+	float MinAverageLuminance = 1;
+	float MaxAverageLuminance = 1;
+	// This scales the average luminance AFTER it gets clamped, affecting the exposure value directly.
+	float LocalExposureMultipler = FMath::Pow(2.0f, Settings.AutoExposureBias);
+	// This scales the average luminance BEFORE it gets clamped, used to implement the calibration constant for AEM_Basic.
+	float AverageLuminanceScale = 1.f;
+
+	// When any of those flags are set, make sure the tonemapper uses an exposure of 1.
+	if (!EngineShowFlags.Lighting || (EngineShowFlags.VisualizeBuffer && View.CurrentBufferVisualizationMode != NAME_None) || View.Family->UseDebugViewPS())
 	{
-		if (View.Family->ExposureSettings.bFixed || Settings.AutoExposureMethod == EAutoExposureMethod::AEM_Manual)
+		LocalExposureMultipler = 1.f;
+	}
+	// Otherwise handle the viewport override settings.
+	else if (View.Family->ExposureSettings.bFixed)
+	{
+		LocalExposureMultipler = 1.f;
+		MinAverageLuminance = MaxAverageLuminance = EV100ToLuminance(View.Family->ExposureSettings.FixedEV100);
+	}
+	// When !EngineShowFlags.EyeAdaptation (from "r.EyeAdaptationQuality 0") or the feature level doesn't support eye adaptation, only Settings.AutoExposureBias controls exposure.
+	else if (EngineShowFlags.EyeAdaptation && View.GetFeatureLevel() >= MinFeatureLevel)
+	{
+		if (Settings.AutoExposureMethod == EAutoExposureMethod::AEM_Manual)
 		{
-			float FixedEV100 = View.Family->ExposureSettings.bFixed ? 
-									View.Family->ExposureSettings.FixedEV100 : 
-									FMath::Log2(FMath::Square(Settings.DepthOfFieldFstop) * Settings.CameraShutterSpeed * 100 / FMath::Max(1.f, Settings.CameraISO));
-
+			const float FixedEV100 = FMath::Log2(FMath::Square(Settings.DepthOfFieldFstop) * Settings.CameraShutterSpeed * 100 / FMath::Max(1.f, Settings.CameraISO));
 			MinAverageLuminance = MaxAverageLuminance = EV100ToLuminance(FixedEV100);
 		}
+		else if (bExtendedLuminanceRange)
+		{
+			MinAverageLuminance = EV100ToLuminance(Settings.AutoExposureMinBrightness);
+			MaxAverageLuminance = EV100ToLuminance(Settings.AutoExposureMaxBrightness);
+		}
 		else
 		{
-			if (bExtendedLuminanceRange)
-			{
-				MinAverageLuminance = EV100ToLuminance(Settings.AutoExposureMinBrightness);
-				MaxAverageLuminance = EV100ToLuminance(Settings.AutoExposureMaxBrightness);
-			}
-			else
-			{
-				MinAverageLuminance = Settings.AutoExposureMinBrightness;
-				MaxAverageLuminance = Settings.AutoExposureMaxBrightness;
-			}
+			MinAverageLuminance = Settings.AutoExposureMinBrightness;
+			MaxAverageLuminance = Settings.AutoExposureMaxBrightness;
+		}
+
+		// Note that AEM_Histogram implements the calibration constant through LowPercent and HiPercent.
+		if (Settings.AutoExposureMethod == EAutoExposureMethod::AEM_Basic)
+		{
+			const float CalibrationConstant = FMath::Clamp<float>(Settings.AutoExposureCalibrationConstant, 1.f, 100.f) * 0.01f;
+			AverageLuminanceScale = 1.f / CalibrationConstant;
 		}
 	}
 
