UE4.21 - Fix SM4 cooks
#rnx

diff --git a/Engine/Source/Runtime/Renderer/Private/PostProcess/DiaphragmDOFPasses.cpp b/Engine/Source/Runtime/Renderer/Private/PostProcess/DiaphragmDOFPasses.cpp
--- a/Engine/Source/Runtime/Renderer/Private/PostProcess/DiaphragmDOFPasses.cpp
+++ b/Engine/Source/Runtime/Renderer/Private/PostProcess/DiaphragmDOFPasses.cpp
@@ -241,6 +241,11 @@ class FPostProcessDiaphragmDOFShader : public FGlobalShader
 public:
 	static bool ShouldCompilePermutation(const FGlobalShaderPermutationParameters& Parameters)
 	{
+		if (!IsFeatureLevelSupported(Parameters.Platform, ERHIFeatureLevel::SM5))
+		{
+			return false;
+		}
+
 		return DiaphragmDOF::IsSupported(Parameters.Platform);
 	}
 
