Fixes a bug where a DOF compute shader was being compiled for SM4

diff --git a/Engine/Shaders/Public/ShaderVersion.ush b/Engine/Shaders/Public/ShaderVersion.ush
--- a/Engine/Shaders/Public/ShaderVersion.ush
+++ b/Engine/Shaders/Public/ShaderVersion.ush
@@ -2,4 +2,4 @@
 // Each time the console command is executed it generates a new GUID. As this file is included
 // in Platform.ush (which should be included in any shader) it allows to invalidate the shader DDC.
 // 
-// GUID = F7C9D61CA78A4854B1D6EE5D3F987298
+// GUID = 7416D80B4444417FE033F3884948C9B5
\ No newline at end of file
diff --git a/Engine/Source/Runtime/Renderer/Private/PostProcess/DiaphragmDOF.h b/Engine/Source/Runtime/Renderer/Private/PostProcess/DiaphragmDOF.h
--- a/Engine/Source/Runtime/Renderer/Private/PostProcess/DiaphragmDOF.h
+++ b/Engine/Source/Runtime/Renderer/Private/PostProcess/DiaphragmDOF.h
@@ -12,7 +12,7 @@
 
 
 // Whitelist diaphragm DOF for platforms that actually have been tested.
-#define WITH_DIAPHRAGM_DOF (PLATFORM_WINDOWS || PLATFORM_XBOXONE || PLATFORM_PS4 || PLATFORM_MAC || PLATFORM_LINUX || PLATFORM_IOS)
+#define WITH_DIAPHRAGM_DOF (PLATFORM_WINDOWS || PLATFORM_XBOXONE || PLATFORM_PS4 || PLATFORM_MAC || PLATFORM_LINUX || PLATFORM_IOS || PLATFORM_SWITCH)
 
 
 namespace DiaphragmDOF
diff --git a/Engine/Source/Runtime/Renderer/Private/PostProcess/DiaphragmDOFPasses.cpp b/Engine/Source/Runtime/Renderer/Private/PostProcess/DiaphragmDOFPasses.cpp
--- a/Engine/Source/Runtime/Renderer/Private/PostProcess/DiaphragmDOFPasses.cpp
+++ b/Engine/Source/Runtime/Renderer/Private/PostProcess/DiaphragmDOFPasses.cpp
@@ -241,11 +241,6 @@ class FPostProcessDiaphragmDOFShader : public FGlobalShader
 public:
 	static bool ShouldCompilePermutation(const FGlobalShaderPermutationParameters& Parameters)
 	{
-		if (!IsFeatureLevelSupported(Parameters.Platform, ERHIFeatureLevel::SM5))
-		{
-			return false;
-		}
-
 		return DiaphragmDOF::IsSupported(Parameters.Platform);
 	}
 
@@ -856,7 +851,11 @@ class FPostProcessDiaphragmDOFScatterGroupPackCS : public FPostProcessDiaphragmD
 
 	static bool ShouldCompilePermutation(const FGlobalShaderPermutationParameters& Parameters)
 	{
-		return FRCPassDiaphragmDOFHybridScatter::IsSupported(Parameters.Platform);
+		if (!FRCPassDiaphragmDOFHybridScatter::IsSupported(Parameters.Platform))
+		{
+			return false;
+		}
+		return FPostProcessDiaphragmDOFShader::ShouldCompilePermutation(Parameters);
 	}
 };
 
