Fixed crash when selecting "Material Texture Scales" and "Required Texture Resolution" viewmodes.
Fixed "ShaderComplexity & Quads" and "Quad Overdraw" not showing quads.

diff --git a/Engine/Source/Runtime/Engine/Private/Materials/MaterialShared.cpp b/Engine/Source/Runtime/Engine/Private/Materials/MaterialShared.cpp
--- a/Engine/Source/Runtime/Engine/Private/Materials/MaterialShared.cpp
+++ b/Engine/Source/Runtime/Engine/Private/Materials/MaterialShared.cpp
@@ -558,7 +558,6 @@ const FMaterialShaderMap* FMaterial::GetShaderMapToUse() const
 	}
 	else 
 	{
-		check(IsInRenderingThread());
 		ShaderMapToUse = GetRenderingThreadShaderMap();
 	}
 
diff --git a/Engine/Source/Runtime/Renderer/Private/MaterialTexCoordScalesRendering.cpp b/Engine/Source/Runtime/Renderer/Private/MaterialTexCoordScalesRendering.cpp
--- a/Engine/Source/Runtime/Renderer/Private/MaterialTexCoordScalesRendering.cpp
+++ b/Engine/Source/Runtime/Renderer/Private/MaterialTexCoordScalesRendering.cpp
@@ -33,8 +33,7 @@ void FMaterialTexCoordScalePS::SetParameters(
 		SetShaderValue(RHICmdList, FMeshMaterialShader::GetPixelShader(), AccuracyColorsParameter, FLinearColor::Black, ColorIndex);
 	}
 
-	TUniformBufferRef<FSceneTexturesUniformParameters> PassUniformBufferValue = CreateSceneTextureUniformBufferSingleDraw(RHICmdList, ESceneTextureSetupMode::None, View.FeatureLevel);
-	FMeshMaterialShader::SetParameters(RHICmdList, FMeshMaterialShader::GetPixelShader(), MaterialRenderProxy, Material, View, DrawRenderState.GetViewUniformBuffer(), PassUniformBufferValue);
+	FMeshMaterialShader::SetParameters(RHICmdList, FMeshMaterialShader::GetPixelShader(), MaterialRenderProxy, Material, View, DrawRenderState.GetViewUniformBuffer(), DrawRenderState.GetPassUniformBuffer());
 }
 
 void FMaterialTexCoordScalePS::SetMesh(
diff --git a/Engine/Source/Runtime/Renderer/Private/PostProcess/SceneRenderTargets.cpp b/Engine/Source/Runtime/Renderer/Private/PostProcess/SceneRenderTargets.cpp
--- a/Engine/Source/Runtime/Renderer/Private/PostProcess/SceneRenderTargets.cpp
+++ b/Engine/Source/Runtime/Renderer/Private/PostProcess/SceneRenderTargets.cpp
@@ -671,10 +671,13 @@ void FSceneRenderTargets::SetQuadOverdrawUAV(FRHICommandList& RHICmdList, bool b
 			Info.NumColorRenderTargets = QuadOverdrawIndex;
 			Info.UnorderedAccessView[Info.NumUAVs++] = QuadOverdrawBuffer->GetRenderTargetItem().UAV;
 
-			// Clear to default value
-			const uint32 ClearValue[4] = { 0, 0, 0, 0 };
-			ClearUAV(RHICmdList, QuadOverdrawBuffer->GetRenderTargetItem(), ClearValue);
-			RHICmdList.TransitionResource(EResourceTransitionAccess::ERWBarrier, EResourceTransitionPipeline::EGfxToGfx, QuadOverdrawBuffer->GetRenderTargetItem().UAV);
+			if (Info.bClearColor)
+			{
+				// Clear to default value
+				const uint32 ClearValue[4] = { 0, 0, 0, 0 };
+				ClearUAV(RHICmdList, QuadOverdrawBuffer->GetRenderTargetItem(), ClearValue);
+				RHICmdList.TransitionResource(EResourceTransitionAccess::ERWBarrier, EResourceTransitionPipeline::EGfxToGfx, QuadOverdrawBuffer->GetRenderTargetItem().UAV);
+			}
 		}
 	}
 }
diff --git a/Engine/Source/Runtime/Renderer/Private/RequiredTextureResolutionRendering.cpp b/Engine/Source/Runtime/Renderer/Private/RequiredTextureResolutionRendering.cpp
--- a/Engine/Source/Runtime/Renderer/Private/RequiredTextureResolutionRendering.cpp
+++ b/Engine/Source/Runtime/Renderer/Private/RequiredTextureResolutionRendering.cpp
@@ -94,8 +94,7 @@ void FRequiredTextureResolutionPS::SetParameters(
 
 	SetShaderValue(RHICmdList, FMeshMaterialShader::GetPixelShader(), AnalysisParamsParameter, FIntVector4(AnalysisIndex, TextureResolution, 0, 0));
 
-	TUniformBufferRef<FSceneTexturesUniformParameters> PassUniformBufferValue = CreateSceneTextureUniformBufferSingleDraw(RHICmdList, ESceneTextureSetupMode::None, View.FeatureLevel);
-	FMeshMaterialShader::SetParameters(RHICmdList, FMeshMaterialShader::GetPixelShader(), MaterialRenderProxy, Material, View, View.ViewUniformBuffer, PassUniformBufferValue);
+	FMeshMaterialShader::SetParameters(RHICmdList, FMeshMaterialShader::GetPixelShader(), MaterialRenderProxy, Material, View, DrawRenderState.GetViewUniformBuffer(), DrawRenderState.GetPassUniformBuffer());
 }
 
 void FRequiredTextureResolutionPS::SetMesh(
