Merging CL 4494719 from //UE4/Dev-VR
Don't clear the decal stencil until the end of the last view.
This fixes decals only showing up in the left eye.

diff --git a/Engine/Source/Runtime/Renderer/Private/CompositionLighting/PostProcessDeferredDecals.cpp b/Engine/Source/Runtime/Renderer/Private/CompositionLighting/PostProcessDeferredDecals.cpp
--- a/Engine/Source/Runtime/Renderer/Private/CompositionLighting/PostProcessDeferredDecals.cpp
+++ b/Engine/Source/Runtime/Renderer/Private/CompositionLighting/PostProcessDeferredDecals.cpp
@@ -936,7 +936,10 @@ void FRCPassPostProcessDeferredDecals::Process(FRenderingCompositePassContext& C
 				}
 			}
 
-			if ((Scene.Decals.Num() > 0) && CurrentStage == DRS_Emissive)
+			// This stops the targets from being resolved and decoded until the last view is rendered.
+			// This is done so as to not run eliminate fast clear on the views before the end.
+			bool bLastView = Context.View.Family->Views.Last() == &Context.View;
+			if (Scene.Decals.Num() > 0 && bLastView && CurrentStage == DRS_Emissive)
 			{
 				// we don't modify stencil but if out input was having stencil for us (after base pass - we need to clear)
 				// Clear stencil to 0, which is the assumed default by other passes
@@ -949,9 +952,6 @@ void FRCPassPostProcessDeferredDecals::Process(FRenderingCompositePassContext& C
 				//DrawClearQuad(RHICmdList, false, FLinearColor(), false, 0, true, 0, SceneContext.GetSceneDepthSurface()->GetSizeXY(), FIntRect());
 			}
 
-			// This stops the targets from being resolved and decoded until the last view is rendered.
-			// This is done so as to not run eliminate fast clear on the views before the end.
-			bool bLastView = Context.View.Family->Views.Last() == &Context.View;
 			if (CurrentStage == DRS_BeforeBasePass)
 			{
 				// combine DBuffer RTWriteMasks; will end up in one texture we can load from in the base pass PS and decide whether to do the actual work or not
