Workaround for UE-65694 from Emil

diff --git a/Engine/Source/Runtime/Renderer/Private/CompositionLighting/PostProcessAmbientOcclusion.cpp b/Engine/Source/Runtime/Renderer/Private/CompositionLighting/PostProcessAmbientOcclusion.cpp
--- a/Engine/Source/Runtime/Renderer/Private/CompositionLighting/PostProcessAmbientOcclusion.cpp
+++ b/Engine/Source/Runtime/Renderer/Private/CompositionLighting/PostProcessAmbientOcclusion.cpp
@@ -802,22 +802,25 @@ void FRCPassPostProcessAmbientOcclusion::ProcessPS(FRenderingCompositePassContex
 		static_assert(bool(ERHIZBuffer::IsInverted), "Inverted depth buffer is assumed when setting depth bounds test for AO.");
 
 		// We must clear all pixels that won't be touched by AO shader.
-		FClearQuadCallbacks Callbacks;
-		Callbacks.PSOModifier = [](FGraphicsPipelineStateInitializer& PSOInitializer)
+		FClearQuadCallbacks* Callbacks = new FClearQuadCallbacks();
+
+		Callbacks->PSOModifier = [](FGraphicsPipelineStateInitializer& PSOInitializer)
 		{
 			PSOInitializer.bDepthBounds = true;
 		};
-		Callbacks.PreClear = [DepthFar](FRHICommandList& InRHICmdList)
+		Callbacks->PreClear = [DepthFar](FRHICommandList& InRHICmdList)
 		{
 			// This is done by rendering a clear quad over a depth range from AmbientOcclusionFadeDistance to far plane.
 			InRHICmdList.SetDepthBounds(0, DepthFar);	// NOTE: Inverted depth
 		};
-		Callbacks.PostClear = [DepthFar](FRHICommandList& InRHICmdList)
+		Callbacks->PostClear = [DepthFar](FRHICommandList& InRHICmdList)
 		{
 			// Set depth bounds test to cover everything from near plane to AmbientOcclusionFadeDistance and run AO pixel shader.
 			InRHICmdList.SetDepthBounds(DepthFar, 1.0f);
 		};
-		DrawClearQuad(Context.RHICmdList, FLinearColor::White, Callbacks);
+		DrawClearQuad(Context.RHICmdList, FLinearColor::White, *Callbacks);
+
+		delete Callbacks;
 	}
 
 	FGraphicsPipelineStateInitializer GraphicsPSOInit;
