Disables Diaphragm DOF for switch forward.

diff --git a/Engine/Source/Runtime/Renderer/Private/PostProcess/DiaphragmDOF.h b/Engine/Source/Runtime/Renderer/Private/PostProcess/DiaphragmDOF.h
--- a/Engine/Source/Runtime/Renderer/Private/PostProcess/DiaphragmDOF.h
+++ b/Engine/Source/Runtime/Renderer/Private/PostProcess/DiaphragmDOF.h
@@ -123,8 +123,7 @@ inline bool IsSupported(EShaderPlatform ShaderPlatform)
 		ShaderPlatform == SP_PS4 ||
 		IsVulkanSM5Platform(ShaderPlatform) ||
 		IsMetalSM5Platform(ShaderPlatform) ||
-		ShaderPlatform == SP_SWITCH ||
-		ShaderPlatform == SP_SWITCH_FORWARD;
+		ShaderPlatform == SP_SWITCH;
 }
 
 
