UE4.21 - Fix for mismatched primitive
#rnx

diff --git a/Engine/Source/Runtime/Renderer/Private/ShadowRendering.cpp b/Engine/Source/Runtime/Renderer/Private/ShadowRendering.cpp
--- a/Engine/Source/Runtime/Renderer/Private/ShadowRendering.cpp
+++ b/Engine/Source/Runtime/Renderer/Private/ShadowRendering.cpp
@@ -728,7 +728,7 @@ void FProjectedShadowInfo::SetupProjectionStencilMask(
 
 		GraphicsPSOInit.BoundShaderState.VertexDeclarationRHI = GetVertexDeclarationFVector4();
 		GraphicsPSOInit.BoundShaderState.VertexShaderRHI = GETSAFERHISHADER_VERTEX(*VertexShaderNoTransform);
-		GraphicsPSOInit.PrimitiveType = PT_TriangleList;
+		GraphicsPSOInit.PrimitiveType = PT_TriangleStrip;
 
 		SetGraphicsPipelineState(RHICmdList, GraphicsPSOInit);
 
