UE4.21 - Fix for fix of mismatched primitive in CL 4535486.
#rnx

diff --git a/Engine/Source/Runtime/Renderer/Private/ShadowRendering.cpp b/Engine/Source/Runtime/Renderer/Private/ShadowRendering.cpp
--- a/Engine/Source/Runtime/Renderer/Private/ShadowRendering.cpp
+++ b/Engine/Source/Runtime/Renderer/Private/ShadowRendering.cpp
@@ -728,7 +728,7 @@ void FProjectedShadowInfo::SetupProjectionStencilMask(
 
 		GraphicsPSOInit.BoundShaderState.VertexDeclarationRHI = GetVertexDeclarationFVector4();
 		GraphicsPSOInit.BoundShaderState.VertexShaderRHI = GETSAFERHISHADER_VERTEX(*VertexShaderNoTransform);
-		GraphicsPSOInit.PrimitiveType = PT_TriangleStrip;
+		GraphicsPSOInit.PrimitiveType = PT_TriangleList;
 
 		SetGraphicsPipelineState(RHICmdList, GraphicsPSOInit);
 
@@ -762,7 +762,7 @@ void FProjectedShadowInfo::SetupProjectionStencilMask(
 
 		RHIUnlockVertexBuffer(VertexBufferRHI);
 		RHICmdList.SetStreamSource(0, VertexBufferRHI, 0);
-		RHICmdList.DrawPrimitive(PT_TriangleStrip, 0, (CascadeSettings.ShadowSplitIndex > 0) ? 4 : 2, 1);
+		RHICmdList.DrawPrimitive(PT_TriangleList, 0, (CascadeSettings.ShadowSplitIndex > 0) ? 4 : 2, 1);
 	}
 	// Not a preshadow, mask the projection to any pixels inside the frustum.
 	else
