DR - Fix some mismatch of Draw*Primitive*(PrimType) not matching PSO.PrimitiveType
#rnx

diff --git a/Engine/Source/Runtime/Renderer/Private/PostProcess/PostProcessCombineLUTs.cpp b/Engine/Source/Runtime/Renderer/Private/PostProcess/PostProcessCombineLUTs.cpp
--- a/Engine/Source/Runtime/Renderer/Private/PostProcess/PostProcessCombineLUTs.cpp
+++ b/Engine/Source/Runtime/Renderer/Private/PostProcess/PostProcessCombineLUTs.cpp
@@ -679,14 +679,14 @@ IMPLEMENT_SHADER_TYPE(template<>,FLUTBlenderCS<4>,TEXT("/Engine/Private/PostProc
 IMPLEMENT_SHADER_TYPE(template<>,FLUTBlenderCS<5>,TEXT("/Engine/Private/PostProcessCombineLUTs.usf"),TEXT("MainCS"),SF_Compute);
 
 template <typename TRHICommandList>
-static void SetLUTBlenderShader(FRenderingCompositePassContext& Context, TRHICommandList& RHICmdList, uint32 BlendCount, FTexture* Texture[], float Weights[], const FVolumeBounds& VolumeBounds)
+static void SetLUTBlenderShader(FRenderingCompositePassContext& Context, TRHICommandList& RHICmdList, uint32 BlendCount, FTexture* Texture[], float Weights[], const FVolumeBounds& VolumeBounds, bool bUseTriangleStrip)
 {
 	FGraphicsPipelineStateInitializer GraphicsPSOInit;
 	RHICmdList.ApplyCachedRenderTargets(GraphicsPSOInit);
 	GraphicsPSOInit.BlendState = TStaticBlendState<>::GetRHI();
 	GraphicsPSOInit.RasterizerState = TStaticRasterizerState<>::GetRHI();
 	GraphicsPSOInit.DepthStencilState = TStaticDepthStencilState<false, CF_Always>::GetRHI();
-	GraphicsPSOInit.PrimitiveType = PT_TriangleList;
+	GraphicsPSOInit.PrimitiveType = bUseTriangleStrip ? PT_TriangleStrip : PT_TriangleList;
 
 	check(BlendCount > 0);
 
@@ -960,7 +960,7 @@ void FRCPassPostProcessCombineLUTs::Process(FRenderingCompositePassContext& Cont
 
 		const FVolumeBounds VolumeBounds(GLUTSize);
 
-		SetLUTBlenderShader(Context, Context.RHICmdList, LocalCount, LocalTextures, LocalWeights, VolumeBounds);
+		SetLUTBlenderShader(Context, Context.RHICmdList, LocalCount, LocalTextures, LocalWeights, VolumeBounds, bUseVolumeTextureLUT);
 
 		if (bUseVolumeTextureLUT)
 		{
diff --git a/Engine/Source/Runtime/Renderer/Private/PostProcess/SceneFilterRendering.h b/Engine/Source/Runtime/Renderer/Private/PostProcess/SceneFilterRendering.h
--- a/Engine/Source/Runtime/Renderer/Private/PostProcess/SceneFilterRendering.h
+++ b/Engine/Source/Runtime/Renderer/Private/PostProcess/SceneFilterRendering.h
@@ -26,6 +26,7 @@ END_UNIFORM_BUFFER_STRUCT( FDrawRectangleParameters )
  * The platform-dependent mapping from pixels to texels is done automatically.
  * Note that the positions are affected by the current viewport.
  * NOTE: DrawRectangle should be used in the vertex shader to calculate the correct position and uv for vertices.
+ * NOTE2: Assumes previously set PSO has PrimitiveType = PT_TriangleList
  *
  * X, Y							Position in screen pixels of the top left corner of the quad
  * SizeX, SizeY					Size in screen pixels of the quad
@@ -54,6 +55,7 @@ extern void DrawRectangle(
 	uint32 InstanceCount = 1
 	);
 
+// NOTE: Assumes previously set PSO has PrimitiveType = PT_TriangleList
 extern void DrawTransformedRectangle(
 	FRHICommandListImmediate& RHICmdList,
 	float X,
@@ -70,6 +72,7 @@ extern void DrawTransformedRectangle(
 	FIntPoint TextureSize
 	);
 
+// NOTE: Assumes previously set PSO has PrimitiveType = PT_TriangleList
 extern void DrawHmdMesh(
 	FRHICommandList& RHICmdList,
 	float X,
@@ -86,6 +89,7 @@ extern void DrawHmdMesh(
 	FShader* VertexShader
 	);
 
+// NOTE: Assumes previously set PSO has PrimitiveType = PT_TriangleList
 extern void DrawPostProcessPass(
 	FRHICommandList& RHICmdList,
 	float X,
diff --git a/Engine/Source/Runtime/Renderer/Public/VolumeRendering.h b/Engine/Source/Runtime/Renderer/Public/VolumeRendering.h
--- a/Engine/Source/Runtime/Renderer/Public/VolumeRendering.h
+++ b/Engine/Source/Runtime/Renderer/Public/VolumeRendering.h
@@ -129,6 +129,7 @@ private:
 	FShaderParameter MinZ;
 };
 
+// This function assumes the PSO had a PrimitiveType of PT_TriangleStrip
 extern ENGINE_API void RasterizeToVolumeTexture(FRHICommandList& RHICmdList, FVolumeBounds VolumeBounds);
 
 /** Vertex buffer used for rendering into a volume texture. */
