Fix for pixel inspector rounding errors causing source sizes to truncate to zero on integer conversion.

diff --git a/Engine/Source/Runtime/Renderer/Private/PostProcess/PostProcessBufferInspector.cpp b/Engine/Source/Runtime/Renderer/Private/PostProcess/PostProcessBufferInspector.cpp
--- a/Engine/Source/Runtime/Renderer/Private/PostProcess/PostProcessBufferInspector.cpp
+++ b/Engine/Source/Runtime/Renderer/Private/PostProcess/PostProcessBufferInspector.cpp
@@ -94,7 +94,8 @@ void FRCPassPostProcessBufferInspector::Process(FRenderingCompositePassContext&
 			// Pixel Depth
 			if (Scene->PixelInspectorData.RenderTargetBufferDepth[PixelInspectorRequest->BufferIndex] != nullptr)
 			{
-				FVector2D SourcePoint = SourceViewportUV * FVector2D(View.ViewRect.Size());
+				FVector2D SourcePoint(FMath::FloorToInt(SourceViewportUV.X * View.ViewRect.Width()),
+									  FMath::FloorToInt(SourceViewportUV.Y * View.ViewRect.Height()));
 				FBox2D SourceBox(SourcePoint, SourcePoint + ExtendSize);
 				
 				const FTexture2DRHIRef &DestinationBufferDepth = Scene->PixelInspectorData.RenderTargetBufferDepth[PixelInspectorRequest->BufferIndex]->GetRenderTargetTexture();
@@ -168,7 +169,8 @@ void FRCPassPostProcessBufferInspector::Process(FRenderingCompositePassContext&
 			const FTexture2DRHIRef &DestinationBufferHDR = Scene->PixelInspectorData.RenderTargetBufferHDR[PixelInspectorRequest->BufferIndex]->GetRenderTargetTexture();
 			if (InputDescHDR != nullptr && DestinationBufferHDR.IsValid())
 			{
-				FVector2D SourcePoint = SourceViewportUV * FVector2D(Context.SceneColorViewRect.Size());
+				FVector2D SourcePoint(FMath::FloorToInt(SourceViewportUV.X * Context.SceneColorViewRect.Width()),
+									  FMath::FloorToInt(SourceViewportUV.Y * Context.SceneColorViewRect.Height()));
 				FBox2D SourceBox(SourcePoint, SourcePoint + ExtendSize);
 
 				const FRenderingCompositeOutputRef* OutputRef1 = GetInput(ePId_Input1);
@@ -191,7 +193,8 @@ void FRCPassPostProcessBufferInspector::Process(FRenderingCompositePassContext&
 			// GBuffer A
 			if (Scene->PixelInspectorData.RenderTargetBufferA[PixelInspectorRequest->BufferIndex] != nullptr)
 			{
-				FVector2D SourcePoint = SourceViewportUV * FVector2D(View.ViewRect.Size());
+				FVector2D SourcePoint(FMath::FloorToInt(SourceViewportUV.X * View.ViewRect.Width()),
+									  FMath::FloorToInt(SourceViewportUV.Y * View.ViewRect.Height()));
 				FBox2D SourceBox(SourcePoint, SourcePoint + ExtendSize);
 
 				const FTexture2DRHIRef &DestinationBufferA = Scene->PixelInspectorData.RenderTargetBufferA[PixelInspectorRequest->BufferIndex]->GetRenderTargetTexture();
@@ -211,7 +214,8 @@ void FRCPassPostProcessBufferInspector::Process(FRenderingCompositePassContext&
 			const FTexture2DRHIRef &DestinationBufferBCDE = Scene->PixelInspectorData.RenderTargetBufferBCDE[PixelInspectorRequest->BufferIndex]->GetRenderTargetTexture();
 			if (DestinationBufferBCDE.IsValid())
 			{
-				FVector2D SourcePoint = SourceViewportUV * FVector2D(View.ViewRect.Size());
+				FVector2D SourcePoint(FMath::FloorToInt(SourceViewportUV.X * View.ViewRect.Width()),
+									  FMath::FloorToInt(SourceViewportUV.Y * View.ViewRect.Height()));
 				FBox2D SourceBox(SourcePoint, SourcePoint + ExtendSize);
 
 				if (SceneContext.GBufferB.IsValid() && SceneContext.GBufferB->GetRenderTargetItem().ShaderResourceTexture.IsValid())
