DR - Copy from 4369920
- Fix crash on split screen

diff --git a/Engine/Source/Runtime/Renderer/Private/PlanarReflectionRendering.cpp b/Engine/Source/Runtime/Renderer/Private/PlanarReflectionRendering.cpp
--- a/Engine/Source/Runtime/Renderer/Private/PlanarReflectionRendering.cpp
+++ b/Engine/Source/Runtime/Renderer/Private/PlanarReflectionRendering.cpp
@@ -728,7 +728,7 @@ void FPlanarReflectionParameters::SetParameters(FRHICommandList& RHICmdList, FPi
 			FVector4(0, 0, 0, 0),
 			FVector4(0, 0, 0, 0),
 		};
-		for (int32 ViewIndex = 0; ViewIndex < View.Family->Views.Num(); ViewIndex++)
+		for (int32 ViewIndex = 0; ViewIndex < View.Family->Views.Num() && ViewIndex < GMaxPlanarReflectionViews; ViewIndex++)
 		{
 			FIntRect ViewRect = ReflectionSceneProxy->ViewRect[ViewIndex];
 			ScreenScaleBiasValue[ViewIndex] = FVector4(
