Fix rendering of shadows without depth-bounds-test (e.g. Metal).

diff --git a/Engine/Source/Runtime/Renderer/Private/ShadowRendering.cpp b/Engine/Source/Runtime/Renderer/Private/ShadowRendering.cpp
--- a/Engine/Source/Runtime/Renderer/Private/ShadowRendering.cpp
+++ b/Engine/Source/Runtime/Renderer/Private/ShadowRendering.cpp
@@ -728,7 +728,7 @@ void FProjectedShadowInfo::SetupProjectionStencilMask(
 
 		GraphicsPSOInit.BoundShaderState.VertexDeclarationRHI = GetVertexDeclarationFVector4();
 		GraphicsPSOInit.BoundShaderState.VertexShaderRHI = GETSAFERHISHADER_VERTEX(*VertexShaderNoTransform);
-		GraphicsPSOInit.PrimitiveType = PT_TriangleStrip;
+		GraphicsPSOInit.PrimitiveType = PT_TriangleList;
 
 		SetGraphicsPipelineState(RHICmdList, GraphicsPSOInit);
 
