Licensee fix: Lights with a light function but no shadows in forward shading now render correctly

diff --git a/Engine/Source/Runtime/Renderer/Private/LightGridInjection.cpp b/Engine/Source/Runtime/Renderer/Private/LightGridInjection.cpp
--- a/Engine/Source/Runtime/Renderer/Private/LightGridInjection.cpp
+++ b/Engine/Source/Runtime/Renderer/Private/LightGridInjection.cpp
@@ -818,7 +818,7 @@ void FDeferredShadingSceneRenderer::RenderForwardShadingShadowProjections(FRHICo
 		const FLightSceneInfo* const LightSceneInfo = LightSceneInfoCompact.LightSceneInfo;
 		const FVisibleLightInfo& VisibleLightInfo = VisibleLightInfos[LightSceneInfo->Id];
 
-		bScreenShadowMaskNeeded = bScreenShadowMaskNeeded || VisibleLightInfo.ShadowsToProject.Num() > 0 || VisibleLightInfo.CapsuleShadowsToProject.Num() > 0;
+		bScreenShadowMaskNeeded |= VisibleLightInfo.ShadowsToProject.Num() > 0 || VisibleLightInfo.CapsuleShadowsToProject.Num() > 0 || LightSceneInfo->Proxy->GetLightFunctionMaterial() != nullptr;
 	}
 
 	if (bScreenShadowMaskNeeded)
diff --git a/Engine/Source/Runtime/Renderer/Private/RendererScene.cpp b/Engine/Source/Runtime/Renderer/Private/RendererScene.cpp
--- a/Engine/Source/Runtime/Renderer/Private/RendererScene.cpp
+++ b/Engine/Source/Runtime/Renderer/Private/RendererScene.cpp
@@ -1328,7 +1328,7 @@ void FScene::AddLightSceneInfo_RenderThread(FLightSceneInfo* LightSceneInfo)
 
 	const bool bForwardShading = IsForwardShadingEnabled(FeatureLevel);
 
-	if (bForwardShading && LightSceneInfo->Proxy->CastsDynamicShadow())
+	if (bForwardShading && (LightSceneInfo->Proxy->CastsDynamicShadow() || LightSceneInfo->Proxy->GetLightFunctionMaterial()))
 	{
 		if (LightSceneInfo->Proxy->HasStaticShadowing())
 		{
