Fixed translucent masked shadows from dynamic meshes.

diff --git a/Engine/Source/Runtime/Renderer/Private/ShadowDepthRendering.cpp b/Engine/Source/Runtime/Renderer/Private/ShadowDepthRendering.cpp
--- a/Engine/Source/Runtime/Renderer/Private/ShadowDepthRendering.cpp
+++ b/Engine/Source/Runtime/Renderer/Private/ShadowDepthRendering.cpp
@@ -1063,7 +1063,7 @@ bool FShadowDepthDrawingPolicyFactory::DrawDynamicMesh(
 		const bool bLocalOnePassPointLightShadow = Context.ShadowInfo->bOnePassPointLightShadow;
 		const bool bReflectiveShadowmap = Context.ShadowInfo->bReflectiveShadowmap && !bLocalOnePassPointLightShadow;
 
-		bool bProcess = !IsTranslucentBlendMode(BlendMode) && ShadingModel != MSM_Unlit && ShouldIncludeDomainInMeshPass(Material->GetMaterialDomain());
+		bool bProcess = Material->ShouldCastDynamicShadows() && ShouldIncludeDomainInMeshPass(Material->GetMaterialDomain());
 
 		if (bReflectiveShadowmap && Material->ShouldInjectEmissiveIntoLPV())
 		{
