Don't use the Volume Texture LUT on Metal MRT (desktop renderers for iOS) until I find a better way to conditionally enable volumetric rendering for A12 GPUs. Without this METAL_MRT was broken on earlier devices and/or the forward renderer.

diff --git a/Engine/Source/Runtime/Renderer/Private/PostProcess/PostProcessCombineLUTs.cpp b/Engine/Source/Runtime/Renderer/Private/PostProcess/PostProcessCombineLUTs.cpp
--- a/Engine/Source/Runtime/Renderer/Private/PostProcess/PostProcessCombineLUTs.cpp
+++ b/Engine/Source/Runtime/Renderer/Private/PostProcess/PostProcessCombineLUTs.cpp
@@ -61,7 +61,7 @@ static TAutoConsoleVariable<int32> CVarMobileTonemapperFilm(
 // USE_VOLUME_LUT: needs to be the same for C++ and HLSL
 bool UseVolumeTextureLUT(EShaderPlatform Platform) 
 {
-	return (IsFeatureLevelSupported(Platform,ERHIFeatureLevel::SM4) && GSupportsVolumeTextureRendering && (RHISupportsGeometryShaders(Platform) || RHISupportsVertexShaderLayer(Platform)));
+	return (IsFeatureLevelSupported(Platform,ERHIFeatureLevel::SM4) && GSupportsVolumeTextureRendering && (Platform != SP_METAL_MRT && Platform != SP_METAL_MRT_MAC) && (RHISupportsGeometryShaders(Platform) || RHISupportsVertexShaderLayer(Platform)));
 }
 
 // including the neutral one at index 0
