DR - Rolling back 4446567 (GPU timing query pools and Vulkan implementation)

diff --git a/Engine/Source/Runtime/Engine/Private/SceneUtils.cpp b/Engine/Source/Runtime/Engine/Private/SceneUtils.cpp
--- a/Engine/Source/Runtime/Engine/Private/SceneUtils.cpp
+++ b/Engine/Source/Runtime/Engine/Private/SceneUtils.cpp
@@ -175,7 +175,7 @@ public:
 	static const uint64 InvalidQueryResult = 0xFFFFFFFFFFFFFFFFull;
 
 public:
-	FRealtimeGPUProfilerEvent(const FName& InName, const FName& InStatName, FRenderQueryPoolRHIParamRef RenderQueryPool, uint32& QueryCount)
+	FRealtimeGPUProfilerEvent(const FName& InName, const FName& InStatName, FRenderQueryPool* RenderQueryPool)
 		: StartResultMicroseconds(InvalidQueryResult)
 		, EndResultMicroseconds(InvalidQueryResult)
 		, FrameNumber(-1)
@@ -189,11 +189,10 @@ public:
 		Name = InName;
 
 		const int MaxGPUQueries = CVarGPUStatsMaxQueriesPerFrame.GetValueOnRenderThread();
-		if ( MaxGPUQueries == -1 || QueryCount < uint32(MaxGPUQueries) )
+		if ( MaxGPUQueries == -1 || RenderQueryPool->GetAllocatedQueryCount() < MaxGPUQueries )
 		{
 			StartQuery = RenderQueryPool->AllocateQuery();
 			EndQuery = RenderQueryPool->AllocateQuery();
-			QueryCount += 2;
 		}
 	}
 
@@ -202,7 +201,7 @@ public:
 		return IsValidRef(StartQuery); 
 	}
 
-	void ReleaseQueries(FRenderQueryPoolRHIParamRef RenderQueryPool, FRHICommandListImmediate* RHICmdListPtr, uint32& QueryCount)
+	void ReleaseQueries(FRenderQueryPool* RenderQueryPool, FRHICommandListImmediate* RHICmdListPtr)
 	{
 		if ( HasQueriesAllocated() )
 		{
@@ -222,7 +221,6 @@ public:
 			}
 			RenderQueryPool->ReleaseQuery(StartQuery);
 			RenderQueryPool->ReleaseQuery(EndQuery);
-			QueryCount -= 2;
 		}
 	}
 
@@ -350,11 +348,9 @@ Container for a single frame's GPU stats
 -----------------------------------------------------------------------------*/
 class FRealtimeGPUProfilerFrame
 {
-	uint32& QueryCount;
 public:
-	FRealtimeGPUProfilerFrame(FRenderQueryPoolRHIRef InRenderQueryPool, uint32& InQueryCount)
-		: QueryCount(InQueryCount)
-		, FrameNumber(-1)
+	FRealtimeGPUProfilerFrame(FRenderQueryPool* InRenderQueryPool)
+		: FrameNumber(-1)
 		, RenderQueryPool(InRenderQueryPool)
 	{}
 
@@ -366,7 +362,7 @@ public:
 	void PushEvent(FRHICommandListImmediate& RHICmdList, const FName& Name, const FName& StatName)
 	{
 		// TODO: this should really use a pool / free list
-		FRealtimeGPUProfilerEvent* Event = new FRealtimeGPUProfilerEvent(Name, StatName, RenderQueryPool, QueryCount);
+		FRealtimeGPUProfilerEvent* Event = new FRealtimeGPUProfilerEvent(Name, StatName, RenderQueryPool);
 		const int32 EventIndex = GpuProfilerEvents.Num();
 
 		GpuProfilerEvents.Add(Event);
@@ -401,7 +397,7 @@ public:
 		{
 			if (GpuProfilerEvents[Index])
 			{
-				GpuProfilerEvents[Index]->ReleaseQueries(RenderQueryPool, RHICommandListPtr, QueryCount);
+				GpuProfilerEvents[Index]->ReleaseQueries(RenderQueryPool, RHICommandListPtr);
 				delete GpuProfilerEvents[Index];
 			}
 		}
@@ -568,7 +564,7 @@ private:
 	TArray<FGPUEventTimeAggregate> EventAggregates;
 
 	uint32 FrameNumber;
-	FRenderQueryPoolRHIRef RenderQueryPool;
+	FRenderQueryPool* RenderQueryPool;
 };
 
 /*-----------------------------------------------------------------------------
@@ -592,11 +588,10 @@ FRealtimeGPUProfiler::FRealtimeGPUProfiler()
 	, bStatGatheringPaused(false)
 	, bInBeginEndBlock(false)
 {
-	const int MaxGPUQueries = CVarGPUStatsMaxQueriesPerFrame.GetValueOnRenderThread();
-	RenderQueryPool = RHICreateRenderQueryPool(RQT_AbsoluteTime, (MaxGPUQueries > 0) ? MaxGPUQueries * 2 : 10240);
+	RenderQueryPool = new FRenderQueryPool(RQT_AbsoluteTime);
 	for (int Index = 0; Index < NumGPUProfilerBufferedFrames; Index++)
 	{
-		Frames.Add(new FRealtimeGPUProfilerFrame(RenderQueryPool, QueryCount));
+		Frames.Add(new FRealtimeGPUProfilerFrame(RenderQueryPool));
 	}
 }
 
@@ -607,7 +602,8 @@ void FRealtimeGPUProfiler::Release()
 		delete Frames[Index];
 	}
 	Frames.Empty();
-	RenderQueryPool.SafeRelease();
+	delete RenderQueryPool;
+	RenderQueryPool = nullptr;
 }
 
 void FRealtimeGPUProfiler::BeginFrame(FRHICommandListImmediate& RHICmdList)
diff --git a/Engine/Source/Runtime/Engine/Public/SceneUtils.h b/Engine/Source/Runtime/Engine/Public/SceneUtils.h
--- a/Engine/Source/Runtime/Engine/Public/SceneUtils.h
+++ b/Engine/Source/Runtime/Engine/Public/SceneUtils.h
@@ -227,8 +227,7 @@ private:
 	int32 WriteBufferIndex;
 	int32 ReadBufferIndex;
 	uint32 WriteFrameNumber;
-	uint32 QueryCount = 0;
-	FRenderQueryPoolRHIRef RenderQueryPool;
+	FRenderQueryPool* RenderQueryPool;
 	bool bStatGatheringPaused;
 	bool bInBeginEndBlock;
 };
diff --git a/Engine/Source/Runtime/RHI/Private/DynamicRHI.cpp b/Engine/Source/Runtime/RHI/Private/DynamicRHI.cpp
--- a/Engine/Source/Runtime/RHI/Private/DynamicRHI.cpp
+++ b/Engine/Source/Runtime/RHI/Private/DynamicRHI.cpp
@@ -314,26 +314,3 @@ void FDynamicRHI::EnableIdealGPUCaptureOptions(bool bEnabled)
 		RHICmdBypassVar->Set(bShouldRHICmdBypass ? 1 : 0, ECVF_SetByConsole);		
 	}	
 }
-
-TRefCountPtr<FRHIRenderQuery> FDefaultRHIRenderQueryPool::AllocateQuery()
-{
-	if (Queries.Num() > 0)
-	{
-		return Queries.Pop();
-	}
-	else
-	{
-		ensure(++AllocatedQueries <= NumQueries);
-		return DynamicRHI->RHICreateRenderQuery(QueryType);
-	}
-}
-
-void FDefaultRHIRenderQueryPool::ReleaseQuery(TRefCountPtr<FRHIRenderQuery> &Query)
-{
-	Queries.Push(Query);
-}
-
-FRenderQueryPoolRHIRef RHICreateRenderQueryPool(ERenderQueryType QueryType, uint32 NumQueries)
-{
-	return GDynamicRHI->RHICreateRenderQueryPool(QueryType, NumQueries);
-}
\ No newline at end of file
diff --git a/Engine/Source/Runtime/RHI/Public/DynamicRHI.h b/Engine/Source/Runtime/RHI/Public/DynamicRHI.h
--- a/Engine/Source/Runtime/RHI/Public/DynamicRHI.h
+++ b/Engine/Source/Runtime/RHI/Public/DynamicRHI.h
@@ -75,28 +75,6 @@ struct FRHIFlipDetails
 	{}
 };
 
-class FDynamicRHI;
-
-class RHI_API FDefaultRHIRenderQueryPool final : public FRHIRenderQueryPool
-{
-public:
-	FDefaultRHIRenderQueryPool(ERenderQueryType InQueryType, FDynamicRHI* InDynamicRHI, uint32 InNumQueries)
-		: DynamicRHI(InDynamicRHI)
-		, QueryType(InQueryType)
-		, NumQueries(InNumQueries)
-	{ }
-
-private:
-	virtual TRefCountPtr<FRHIRenderQuery> AllocateQuery() override;
-	virtual void ReleaseQuery(TRefCountPtr<FRHIRenderQuery> &Query) override;
-
-	FDynamicRHI* DynamicRHI = nullptr;
-	ERenderQueryType QueryType;
-	uint32 NumQueries = 0;
-	uint32 AllocatedQueries = 0;
-	TArray<TRefCountPtr<FRHIRenderQuery>> Queries;
-};
-
 /** The interface which is implemented by the dynamically bound RHI. */
 class RHI_API FDynamicRHI
 {
@@ -220,17 +198,6 @@ public:
 		return nullptr;
 	}
 	
-	/**
-	* Creates a pool for querys like timers or occlusion queries.
-	* @param QueryType The ype of the queries provided by this pool like RQT_Occlusion or RQT_AbsoluteTime.
-	* @return the Querypool.
-	*/
-	// FlushType: Must be Thread-Safe.
-	virtual FRenderQueryPoolRHIRef RHICreateRenderQueryPool(ERenderQueryType QueryType, uint32 NumQueries)
-	{
-		return new FDefaultRHIRenderQueryPool(QueryType, this, NumQueries);
-	}
-
 	/**
 	* Creates a compute fence.  Compute fences are named GPU fences which can be written to once before resetting.
 	* A command to write the fence must be enqueued before any commands to wait on them.  This is enforced on the CPU to avoid GPU hangs.
@@ -1269,8 +1236,6 @@ FORCEINLINE class IRHICommandContextContainer* RHIGetCommandContextContainer(int
 	return GDynamicRHI->RHIGetCommandContextContainer(Index, Num, GPUMask);
 }
 
-RHI_API FRenderQueryPoolRHIRef RHICreateRenderQueryPool(ERenderQueryType QueryType, uint32 NumQueries);
-
 /**
 * Defragment the texture pool.
 */
diff --git a/Engine/Source/Runtime/RHI/Public/RHIResources.h b/Engine/Source/Runtime/RHI/Public/RHIResources.h
--- a/Engine/Source/Runtime/RHI/Public/RHIResources.h
+++ b/Engine/Source/Runtime/RHI/Public/RHIResources.h
@@ -831,13 +831,6 @@ private:
 
 class FRHIRenderQuery : public FRHIResource {};
 
-class FRHIRenderQueryPool : public FRHIResource 
-{
-public:
-	virtual TRefCountPtr<FRHIRenderQuery> AllocateQuery() = 0;
-	virtual void ReleaseQuery(TRefCountPtr<FRHIRenderQuery> &Query) = 0;
-};
-
 class FRHIComputeFence : public FRHIResource
 {
 public:
@@ -1009,9 +1002,6 @@ typedef TRefCountPtr<FRHITextureReference> FTextureReferenceRHIRef;
 typedef FRHIRenderQuery*              FRenderQueryRHIParamRef;
 typedef TRefCountPtr<FRHIRenderQuery> FRenderQueryRHIRef;
 
-typedef FRHIRenderQueryPool*              FRenderQueryPoolRHIParamRef;
-typedef TRefCountPtr<FRHIRenderQueryPool> FRenderQueryPoolRHIRef;
-
 typedef FRHIGPUFence*				FGPUFenceRHIParamRef;
 typedef TRefCountPtr<FRHIGPUFence>	FGPUFenceRHIRef;
 
diff --git a/Engine/Source/Runtime/Renderer/Private/GPUBenchmark.cpp b/Engine/Source/Runtime/Renderer/Private/GPUBenchmark.cpp
--- a/Engine/Source/Runtime/Renderer/Private/GPUBenchmark.cpp
+++ b/Engine/Source/Runtime/Renderer/Private/GPUBenchmark.cpp
@@ -411,6 +411,8 @@ void RendererGPUBenchmark(FRHICommandListImmediate& RHICmdList, FSynthBenchmarkR
 {
 	check(IsInRenderingThread());
 
+	FRenderQueryPool TimerQueryPool(RQT_AbsoluteTime);
+
 	bool bValidGPUTimer = (FGPUTiming::GetTimingFrequency() / (1000 * 1000)) != 0;
 
 	if(!bValidGPUTimer)
@@ -495,14 +497,13 @@ void RendererGPUBenchmark(FRHICommandListImmediate& RHICmdList, FSynthBenchmarkR
 		uint32 DestRTIndex = 0;
 
 		const uint32 TimerSampleCount = IterationCount * MethodCount + 1;
-		FRenderQueryPoolRHIRef TimerQueryPool = RHICreateRenderQueryPool(RQT_AbsoluteTime, TimerSampleCount * 2);
 
 		static FRenderQueryRHIRef TimerQueries[TimerSampleCount];
 		static float LocalWorkScale[IterationCount];
 
 		for(uint32  i = 0; i < TimerSampleCount; ++i)
 		{
-			TimerQueries[i] = TimerQueryPool->AllocateQuery();
+			TimerQueries[i] = TimerQueryPool.AllocateQuery();
 		}
 
 		const bool bSupportsTimerQueries = (TimerQueries[0] != NULL);
@@ -616,7 +617,7 @@ void RendererGPUBenchmark(FRHICommandListImmediate& RHICmdList, FSynthBenchmarkR
 			// flushes the RHI thread to make sure all RHICmdList.EndRenderQuery() commands got executed.
 			RHICmdList.ImmediateFlush(EImmediateFlushType::FlushRHIThread);
 			RHICmdList.GetRenderQueryResult(TimerQueries[0], OldAbsTime, true);
-			TimerQueryPool->ReleaseQuery(TimerQueries[0]);
+			TimerQueryPool.ReleaseQuery(TimerQueries[0]);
 
 			for(uint32 Iteration = 0; Iteration < IterationCount; ++Iteration)
 			{
@@ -628,7 +629,7 @@ void RendererGPUBenchmark(FRHICommandListImmediate& RHICmdList, FSynthBenchmarkR
 
 					uint64 AbsTime;
 					RHICmdList.GetRenderQueryResult(TimerQueries[QueryIndex], AbsTime, true);
-					TimerQueryPool->ReleaseQuery(TimerQueries[QueryIndex]);
+					TimerQueryPool.ReleaseQuery(TimerQueries[QueryIndex]);
 
 					uint64 RelTime = FMath::Max(AbsTime - OldAbsTime, 1ull);
 
diff --git a/Engine/Source/Runtime/Renderer/Private/RendererScene.cpp b/Engine/Source/Runtime/Renderer/Private/RendererScene.cpp
--- a/Engine/Source/Runtime/Renderer/Private/RendererScene.cpp
+++ b/Engine/Source/Runtime/Renderer/Private/RendererScene.cpp
@@ -94,9 +94,6 @@ struct FSpeedTreeWindComputation
 /** Default constructor. */
 FSceneViewState::FSceneViewState()
 	: OcclusionQueryPool(RQT_Occlusion)
-	, TimerQueryPool(RHICreateRenderQueryPool(RQT_AbsoluteTime, FLatentGPUTimer::NumBufferedFrames * 2 * 2 * 2))
-	, TranslucencyTimer(TimerQueryPool)
-	, SeparateTranslucencyTimer(TimerQueryPool)
 {
 	UniqueID = FSceneViewState_UniqueID.Increment();
 	OcclusionFrameCounter = 0;
@@ -3555,18 +3552,12 @@ bool FMotionBlurInfoData::GetPrimitiveMotionBlurInfo(const FPrimitiveSceneInfo*
 
 //////////////////////////////////////////////////////////////////////////
 
-FLatentGPUTimer::FLatentGPUTimer(FRenderQueryPoolRHIRef& InTimerQueryPool, int32 InAvgSamples)
-: TimerQueryPool(InTimerQueryPool)
-, AvgSamples(InAvgSamples)
-, TotalTime(0.0f)
-, SampleIndex(0)
-, QueryIndex(0)
+FLatentGPUTimer::FLatentGPUTimer(int32 InAvgSamples)
+	: AvgSamples(InAvgSamples)
+	, TotalTime(0.0f)
+	, SampleIndex(0)
+	, QueryIndex(0)
 {
-	for (int i = 0; i < NumBufferedFrames; i++)
-	{
-		QueriesInFlight[i] = false;
-	}
-
 	TimeSamples.AddZeroed(AvgSamples);
 }
 
@@ -3579,57 +3570,43 @@ bool FLatentGPUTimer::Tick(FRHICommandListImmediate& RHICmdList)
 
 	QueryIndex = (QueryIndex + 1) % NumBufferedFrames;
 
-	for (int IndexOffset = 0; IndexOffset < NumBufferedFrames - 1; IndexOffset++)
+	if (StartQueries[QueryIndex] && EndQueries[QueryIndex])
 	{
-		int CollectedIndex = (QueryIndex + IndexOffset) % NumBufferedFrames;
-		if (!QueriesInFlight[CollectedIndex])
+		if (IsRunningRHIInSeparateThread())
 		{
-			continue;
+			// Block until the RHI thread has processed the previous query commands, if necessary
+			// Stat disabled since we buffer 2 frames minimum, it won't actually block
+			//SCOPE_CYCLE_COUNTER(STAT_TranslucencyTimestampQueryFence_Wait);
+			int32 BlockFrame = NumBufferedFrames - 1;
+			FRHICommandListExecutor::WaitOnRHIThreadFence(QuerySubmittedFences[BlockFrame]);
+			QuerySubmittedFences[BlockFrame] = nullptr;
 		}
 
-		if (StartQueries[CollectedIndex] && EndQueries[CollectedIndex])
-		{
-			if (IsRunningRHIInSeparateThread())
-			{
-				// Block until the RHI thread has processed the previous query commands, if necessary
-				// Stat disabled since we buffer 2 frames minimum, it won't actually block
-				//SCOPE_CYCLE_COUNTER(STAT_TranslucencyTimestampQueryFence_Wait);
-				int32 BlockFrame = NumBufferedFrames - 1;
-				FRHICommandListExecutor::WaitOnRHIThreadFence(QuerySubmittedFences[BlockFrame]);
-				QuerySubmittedFences[BlockFrame] = nullptr;
-			}
-
-			uint64 StartMicroseconds;
-			uint64 EndMicroseconds;
-			bool bStartSuccess;
-			bool bEndSuccess;
+		uint64 StartMicroseconds;
+		uint64 EndMicroseconds;
+		bool bStartSuccess;
+		bool bEndSuccess;
 
-			{
-				// Block on the GPU until we have the timestamp query results, if necessary
-				// Stat disabled since we buffer 2 frames minimum, it won't actually block
-				//SCOPE_CYCLE_COUNTER(STAT_TranslucencyTimestampQuery_Wait);
-				bStartSuccess = RHICmdList.GetRenderQueryResult(StartQueries[CollectedIndex], StartMicroseconds, false);
-				bEndSuccess = RHICmdList.GetRenderQueryResult(EndQueries[CollectedIndex], EndMicroseconds, false);
-				if (!bStartSuccess || !bEndSuccess)
-				{
-					return false;
-				}
-				QueriesInFlight[CollectedIndex] = false;
-			}
+		{
+			// Block on the GPU until we have the timestamp query results, if necessary
+			// Stat disabled since we buffer 2 frames minimum, it won't actually block
+			//SCOPE_CYCLE_COUNTER(STAT_TranslucencyTimestampQuery_Wait);
+			bStartSuccess = RHICmdList.GetRenderQueryResult(StartQueries[QueryIndex], StartMicroseconds, true);
+			bEndSuccess = RHICmdList.GetRenderQueryResult(EndQueries[QueryIndex], EndMicroseconds, true);
+		}
 
-			TotalTime -= TimeSamples[SampleIndex];
-			float LastFrameTranslucencyDurationMS = TimeSamples[SampleIndex];
-			if (bStartSuccess && bEndSuccess)
-			{
-				LastFrameTranslucencyDurationMS = (EndMicroseconds - StartMicroseconds) / 1000.0f;
-			}
+		TotalTime -= TimeSamples[SampleIndex];
+		float LastFrameTranslucencyDurationMS = TimeSamples[SampleIndex];
+		if (bStartSuccess && bEndSuccess)
+		{
+			LastFrameTranslucencyDurationMS = (EndMicroseconds - StartMicroseconds) / 1000.0f;
+		}
 
-			TimeSamples[SampleIndex] = LastFrameTranslucencyDurationMS;
-			TotalTime += LastFrameTranslucencyDurationMS;
-			SampleIndex = (SampleIndex + 1) % AvgSamples;
+		TimeSamples[SampleIndex] = LastFrameTranslucencyDurationMS;
+		TotalTime += LastFrameTranslucencyDurationMS;
+		SampleIndex = (SampleIndex + 1) % AvgSamples;
 
-			return bStartSuccess && bEndSuccess;
-		}
+		return bStartSuccess && bEndSuccess;
 	}
 
 	return false;
@@ -3644,10 +3621,9 @@ void FLatentGPUTimer::Begin(FRHICommandListImmediate& RHICmdList)
 	
 	if (!StartQueries[QueryIndex])
 	{		
-		StartQueries[QueryIndex] = TimerQueryPool->AllocateQuery();
+		StartQueries[QueryIndex] = RHICmdList.CreateRenderQuery(RQT_AbsoluteTime);
 	}
 
-	ensure(QueriesInFlight[QueryIndex] == false);
 	RHICmdList.EndRenderQuery(StartQueries[QueryIndex]);
 }
 
@@ -3660,12 +3636,10 @@ void FLatentGPUTimer::End(FRHICommandListImmediate& RHICmdList)
 	
 	if (!EndQueries[QueryIndex])
 	{
-		EndQueries[QueryIndex] = TimerQueryPool->AllocateQuery();
+		EndQueries[QueryIndex] = RHICmdList.CreateRenderQuery(RQT_AbsoluteTime);
 	}
 
-	ensure(QueriesInFlight[QueryIndex] == false);
 	RHICmdList.EndRenderQuery(EndQueries[QueryIndex]);
-	QueriesInFlight[QueryIndex] = true;
 	// Hint to the RHI to submit commands up to this point to the GPU if possible.  Can help avoid CPU stalls next frame waiting
 	// for these query results on some platforms.
 	RHICmdList.SubmitCommandsHint();
@@ -3687,8 +3661,8 @@ void FLatentGPUTimer::Release()
 {
 	for (int32 i = 0; i < NumBufferedFrames; ++i)
 	{
-		if(StartQueries[i].IsValid()) StartQueries[i].SafeRelease();
-		if(EndQueries[i].IsValid()) EndQueries[i].SafeRelease();
+		StartQueries[i].SafeRelease();
+		EndQueries[i].SafeRelease();
 	}
 }
 
diff --git a/Engine/Source/Runtime/Renderer/Private/ScenePrivate.h b/Engine/Source/Runtime/Renderer/Private/ScenePrivate.h
--- a/Engine/Source/Runtime/Renderer/Private/ScenePrivate.h
+++ b/Engine/Source/Runtime/Renderer/Private/ScenePrivate.h
@@ -495,15 +495,10 @@ Buffers multiple frames to avoid waiting on the GPU so times are a little lagged
 */
 class FLatentGPUTimer
 {
-	FRenderQueryPoolRHIRef& TimerQueryPool;
-public:
 	static const int32 NumBufferedFrames = FOcclusionQueryHelpers::MaxBufferedOcclusionFrames + 1;
+public:
 
-	FLatentGPUTimer(FRenderQueryPoolRHIRef& InTimerQueryPool, int32 InAvgSamples = 30);
-	~FLatentGPUTimer()
-	{
-		Release();
-	}
+	FLatentGPUTimer(int32 InAvgSamples = 30);
 
 	void Release();
 
@@ -530,7 +525,6 @@ private:
 	int32 SampleIndex;
 
 	int32 QueryIndex;
-	bool QueriesInFlight[NumBufferedFrames];
 	FRenderQueryRHIRef StartQueries[NumBufferedFrames];
 	FRenderQueryRHIRef EndQueries[NumBufferedFrames];
 	FGraphEventRef QuerySubmittedFences[NumBufferedFrames];
@@ -903,7 +897,6 @@ public:
 	FRWBuffer CapsuleTileIntersectionCountsBuffer;
 
 	/** Timestamp queries around separate translucency, used for auto-downsampling. */
-	FRenderQueryPoolRHIRef TimerQueryPool;
 	FLatentGPUTimer TranslucencyTimer;
 	FLatentGPUTimer SeparateTranslucencyTimer;
 
diff --git a/Engine/Source/Runtime/VulkanRHI/Private/VulkanCommandBuffer.h b/Engine/Source/Runtime/VulkanRHI/Private/VulkanCommandBuffer.h
--- a/Engine/Source/Runtime/VulkanRHI/Private/VulkanCommandBuffer.h
+++ b/Engine/Source/Runtime/VulkanRHI/Private/VulkanCommandBuffer.h
@@ -167,14 +167,14 @@ public:
 
 	bool AcquirePoolSetAndDescriptorsIfNeeded(const class FVulkanDescriptorSetsLayout& Layout, bool bNeedDescriptors, VkDescriptorSet* OutDescriptors);
 
-	TRefCountPtr<FVulkanTimestampQueryPool> PrepareTimestampQueryPool()
+	/*TRefCountPtr<*/FVulkanTimestampQueryPool* /*>*/ PrepareTimestampQueryPool()
 	{
 		return TimestampQueryPool;
 	}
 
 private:
 	FVulkanDevice* Device;
-	TRefCountPtr<FVulkanTimestampQueryPool> TimestampQueryPool = nullptr;
+	/*TRefCountPtr<*/FVulkanTimestampQueryPool* /*>*/ TimestampQueryPool = nullptr;
 	VkCommandBuffer CommandBufferHandle;
 	EState State;
 	double SubmittedTime = 0.0f;
diff --git a/Engine/Source/Runtime/VulkanRHI/Private/VulkanQuery.cpp b/Engine/Source/Runtime/VulkanRHI/Private/VulkanQuery.cpp
--- a/Engine/Source/Runtime/VulkanRHI/Private/VulkanQuery.cpp
+++ b/Engine/Source/Runtime/VulkanRHI/Private/VulkanQuery.cpp
@@ -840,6 +840,7 @@ void FVulkanCommandListContext::RHIEndRenderQuery(FRenderQueryRHIParamRef QueryR
 	}
 }
 
+/*
 FRenderQueryRHIRef FVulkanTimestampQueryPool::AllocateQuery()
 {
 	FVulkanTimingQuery* Query = new FVulkanTimingQuery();
@@ -853,6 +854,7 @@ void FVulkanTimestampQueryPool::ReleaseQuery(TRefCountPtr<FRHIRenderQuery> &Quer
 {
 	Query.SafeRelease();
 }
+*/
 
 void FVulkanCommandListContext::WriteBeginTimestamp(FVulkanCmdBuffer* CmdBuffer)
 {
diff --git a/Engine/Source/Runtime/VulkanRHI/Private/VulkanUAV.cpp b/Engine/Source/Runtime/VulkanRHI/Private/VulkanUAV.cpp
--- a/Engine/Source/Runtime/VulkanRHI/Private/VulkanUAV.cpp
+++ b/Engine/Source/Runtime/VulkanRHI/Private/VulkanUAV.cpp
@@ -343,6 +343,7 @@ FComputeFenceRHIRef FVulkanDynamicRHI::RHICreateComputeFence(const FName& Name)
 	return new FVulkanComputeFence(Device, Name);
 }
 
+/*
 FRenderQueryPoolRHIRef FVulkanDynamicRHI::RHICreateRenderQueryPool(ERenderQueryType QueryType, uint32 NumQueries)
 {
 	if (QueryType == RQT_AbsoluteTime)
@@ -354,6 +355,7 @@ FRenderQueryPoolRHIRef FVulkanDynamicRHI::RHICreateRenderQueryPool(ERenderQueryT
 		return new FDefaultRHIRenderQueryPool(QueryType, this, NumQueries);
 	}
 }
+*/
 
 
 bool FVulkanGPUFence::Poll() const
diff --git a/Engine/Source/Runtime/VulkanRHI/Public/VulkanDynamicRHI.h b/Engine/Source/Runtime/VulkanRHI/Public/VulkanDynamicRHI.h
--- a/Engine/Source/Runtime/VulkanRHI/Public/VulkanDynamicRHI.h
+++ b/Engine/Source/Runtime/VulkanRHI/Public/VulkanDynamicRHI.h
@@ -50,7 +50,6 @@ public:
 	virtual FGeometryShaderRHIRef RHICreateGeometryShaderWithStreamOutput(const TArray<uint8>& Code, const FStreamOutElementList& ElementList, uint32 NumStrides, const uint32* Strides, int32 RasterizedStream) final override;
 	virtual FComputeShaderRHIRef RHICreateComputeShader(const TArray<uint8>& Code) final override;
 	virtual FComputeFenceRHIRef RHICreateComputeFence(const FName& Name) final override;
-	virtual FRenderQueryPoolRHIRef RHICreateRenderQueryPool(ERenderQueryType QueryType, uint32 NumQueries) final override;
 	virtual FGPUFenceRHIRef RHICreateGPUFence(const FName &Name) final override;
 	virtual FStagingBufferRHIRef RHICreateStagingBuffer(FVertexBufferRHIParamRef BackingBuffer) final override;
 	virtual FBoundShaderStateRHIRef RHICreateBoundShaderState(FVertexDeclarationRHIParamRef VertexDeclaration, FVertexShaderRHIParamRef VertexShader, FHullShaderRHIParamRef HullShader, FDomainShaderRHIParamRef DomainShader, FPixelShaderRHIParamRef PixelShader, FGeometryShaderRHIParamRef GeometryShader) final override;
diff --git a/Engine/Source/Runtime/VulkanRHI/Public/VulkanResources.h b/Engine/Source/Runtime/VulkanRHI/Public/VulkanResources.h
--- a/Engine/Source/Runtime/VulkanRHI/Public/VulkanResources.h
+++ b/Engine/Source/Runtime/VulkanRHI/Public/VulkanResources.h
@@ -813,10 +813,13 @@ protected:
 	uint32 FrameNumber = UINT32_MAX;
 };
 
-class FVulkanTimestampQueryPool final : public FRHIRenderQueryPool, public FVulkanQueryPool
+
+class FVulkanTimestampQueryPool final : public FVulkanQueryPool
 {
+/*
 	FRenderQueryRHIRef AllocateQuery() override;
 	virtual void ReleaseQuery(TRefCountPtr<FRHIRenderQuery> &Query) override;
+*/
 
 	uint32 QueueHead = 0;
 	uint32 QueueTail = 0;
@@ -924,7 +927,7 @@ public:
 #endif
 
 protected:
-	TRefCountPtr<FVulkanTimestampQueryPool> Pool;
+	/*TRefCountPtr<*/FVulkanTimestampQueryPool* /*>*/ Pool;
 	EState State = EState::Unused;
 	bool IsPooledAtHighLevel = false;
 
