DR - Clear needs to happen if any decal (even in previous passes) were drawn

diff --git a/Engine/Source/Runtime/Renderer/Private/CompositionLighting/PostProcessDeferredDecals.cpp b/Engine/Source/Runtime/Renderer/Private/CompositionLighting/PostProcessDeferredDecals.cpp
--- a/Engine/Source/Runtime/Renderer/Private/CompositionLighting/PostProcessDeferredDecals.cpp
+++ b/Engine/Source/Runtime/Renderer/Private/CompositionLighting/PostProcessDeferredDecals.cpp
@@ -936,7 +936,7 @@ void FRCPassPostProcessDeferredDecals::Process(FRenderingCompositePassContext& C
 				}
 			}
 
-			if ((SortedDecals.Num() > 0) && CurrentStage == DRS_Emissive)
+			if ((Scene.Decals.Num() > 0) && CurrentStage == DRS_Emissive)
 			{
 				// we don't modify stencil but if out input was having stencil for us (after base pass - we need to clear)
 				// Clear stencil to 0, which is the assumed default by other passes
