Merging //UE4/Dev-Rendering-Graph@4492585 to Dev-Rendering (//UE4/Dev-Rendering)

diff --git a/Engine/Source/Runtime/Renderer/Private/DistortionRendering.cpp b/Engine/Source/Runtime/Renderer/Private/DistortionRendering.cpp
--- a/Engine/Source/Runtime/Renderer/Private/DistortionRendering.cpp
+++ b/Engine/Source/Runtime/Renderer/Private/DistortionRendering.cpp
@@ -11,7 +11,7 @@
 #include "Shader.h"
 #include "StaticBoundShaderState.h"
 #include "SceneUtils.h"
-#include "PostProcess/RenderTargetPool.h"
+#include "RenderTargetPool.h"
 #include "PostProcess/SceneRenderTargets.h"
 #include "GlobalShader.h"
 #include "MaterialShaderType.h"
@@ -1133,7 +1133,7 @@ void FSceneRenderer::RenderDistortion(FRHICommandListImmediate& RHICmdList)
 				// Ideally we skip the EliminateFastClear since we don't need pixels with no stencil set to be cleared
 				RHICmdList.TransitionResource( EResourceTransitionAccess::EReadable, DistortionRT->GetRenderTargetItem().TargetableTexture );
 				// to be able to observe results with VisualizeTexture
-				GRenderTargetPool.VisualizeTexture.SetCheckPoint(RHICmdList, DistortionRT);
+				GVisualizeTexture.SetCheckPoint(RHICmdList, DistortionRT);
 			}
 		}
 	}
diff --git a/Engine/Source/Runtime/Renderer/Private/FogRendering.cpp b/Engine/Source/Runtime/Renderer/Private/FogRendering.cpp
--- a/Engine/Source/Runtime/Renderer/Private/FogRendering.cpp
+++ b/Engine/Source/Runtime/Renderer/Private/FogRendering.cpp
@@ -41,7 +41,7 @@ static TAutoConsoleVariable<int32> CVarFog(
 	ECVF_RenderThreadSafe | ECVF_Scalability);
 
 
-IMPLEMENT_UNIFORM_BUFFER_STRUCT(FFogUniformParameters, TEXT("FogStruct"));
+IMPLEMENT_GLOBAL_SHADER_PARAMETER_STRUCT(FFogUniformParameters, "FogStruct");
 
 void SetupFogUniformParameters(const FViewInfo& View, FFogUniformParameters& OutParameters)
 {
diff --git a/Engine/Source/Runtime/Renderer/Private/FogRendering.h b/Engine/Source/Runtime/Renderer/Private/FogRendering.h
--- a/Engine/Source/Runtime/Renderer/Private/FogRendering.h
+++ b/Engine/Source/Runtime/Renderer/Private/FogRendering.h
@@ -11,20 +11,20 @@
 #include "SceneRendering.h"
 #include "VolumetricFog.h"
 
-BEGIN_UNIFORM_BUFFER_STRUCT(FFogUniformParameters,)
-	UNIFORM_MEMBER(FVector4, ExponentialFogParameters)
-	UNIFORM_MEMBER(FVector4, ExponentialFogColorParameter)
-	UNIFORM_MEMBER(FVector4, ExponentialFogParameters3)
-	UNIFORM_MEMBER(FVector4, InscatteringLightDirection) // non negative DirectionalInscatteringStartDistance in .W
-	UNIFORM_MEMBER(FVector4, DirectionalInscatteringColor)
-	UNIFORM_MEMBER(FVector2D, SinCosInscatteringColorCubemapRotation)
-	UNIFORM_MEMBER(FVector, FogInscatteringTextureParameters)
-	UNIFORM_MEMBER(float, ApplyVolumetricFog)
-	UNIFORM_MEMBER_TEXTURE(TextureCube, FogInscatteringColorCubemap)
-	UNIFORM_MEMBER_SAMPLER(SamplerState, FogInscatteringColorSampler)
-	UNIFORM_MEMBER_TEXTURE(Texture3D, IntegratedLightScattering)
-	UNIFORM_MEMBER_SAMPLER(SamplerState, IntegratedLightScatteringSampler)
-END_UNIFORM_BUFFER_STRUCT(FFogUniformParameters)
+BEGIN_GLOBAL_SHADER_PARAMETER_STRUCT(FFogUniformParameters,)
+	SHADER_PARAMETER(FVector4, ExponentialFogParameters)
+	SHADER_PARAMETER(FVector4, ExponentialFogColorParameter)
+	SHADER_PARAMETER(FVector4, ExponentialFogParameters3)
+	SHADER_PARAMETER(FVector4, InscatteringLightDirection) // non negative DirectionalInscatteringStartDistance in .W
+	SHADER_PARAMETER(FVector4, DirectionalInscatteringColor)
+	SHADER_PARAMETER(FVector2D, SinCosInscatteringColorCubemapRotation)
+	SHADER_PARAMETER(FVector, FogInscatteringTextureParameters)
+	SHADER_PARAMETER(float, ApplyVolumetricFog)
+	SHADER_PARAMETER_TEXTURE(TextureCube, FogInscatteringColorCubemap)
+	SHADER_PARAMETER_SAMPLER(SamplerState, FogInscatteringColorSampler)
+	SHADER_PARAMETER_TEXTURE(Texture3D, IntegratedLightScattering)
+	SHADER_PARAMETER_SAMPLER(SamplerState, IntegratedLightScatteringSampler)
+END_GLOBAL_SHADER_PARAMETER_STRUCT()
 
 extern void SetupFogUniformParameters(const class FViewInfo& View, FFogUniformParameters& OutParameters);
 
diff --git a/Engine/Source/Runtime/Renderer/Private/GPUBenchmark.cpp b/Engine/Source/Runtime/Renderer/Private/GPUBenchmark.cpp
--- a/Engine/Source/Runtime/Renderer/Private/GPUBenchmark.cpp
+++ b/Engine/Source/Runtime/Renderer/Private/GPUBenchmark.cpp
@@ -16,11 +16,12 @@
 #include "RHIStaticStates.h"
 #include "Containers/DynamicRHIResourceArray.h"
 #include "GlobalShader.h"
-#include "PostProcess/RenderTargetPool.h"
+#include "RenderTargetPool.h"
 #include "PostProcess/SceneFilterRendering.h"
 #include "GPUProfiler.h"
 #include "PipelineStateCache.h"
 #include "LongGPUTask.h"
+#include "PostProcess/VisualizeTexture.h"
 
 static const uint32 GBenchmarkResolution = 512;
 static const uint32 GBenchmarkPrimitives = 200000;
@@ -573,7 +574,7 @@ void RendererGPUBenchmark(FRHICommandListImmediate& RHICmdList, FSynthBenchmarkR
 				// 0 / 1
 				const uint32 SrcRTIndex = 1 - DestRTIndex;
 
-				GRenderTargetPool.VisualizeTexture.SetCheckPoint(RHICmdList, RTItems[DestRTIndex]);
+				GVisualizeTexture.SetCheckPoint(RHICmdList, RTItems[DestRTIndex]);
 
 				SetRenderTarget(RHICmdList, RTItems[DestRTIndex]->GetRenderTargetItem().TargetableTexture, FTextureRHIRef(), true);	
 
diff --git a/Engine/Source/Runtime/Renderer/Private/GPUFastFourierTransform.cpp b/Engine/Source/Runtime/Renderer/Private/GPUFastFourierTransform.cpp
--- a/Engine/Source/Runtime/Renderer/Private/GPUFastFourierTransform.cpp
+++ b/Engine/Source/Runtime/Renderer/Private/GPUFastFourierTransform.cpp
@@ -3,7 +3,7 @@
 #include "GPUFastFourierTransform.h"
 #include "SceneUtils.h"
 #include "GlobalShader.h"
-#include "PostProcess/RenderTargetPool.h" // used for on-the-fly acquisition of 
+#include "RenderTargetPool.h" // used for on-the-fly acquisition of 
 
 
 uint32 GPUFFT::MaxScanLineLength() 
diff --git a/Engine/Source/Runtime/Renderer/Private/LightGridInjection.cpp b/Engine/Source/Runtime/Renderer/Private/LightGridInjection.cpp
--- a/Engine/Source/Runtime/Renderer/Private/LightGridInjection.cpp
+++ b/Engine/Source/Runtime/Renderer/Private/LightGridInjection.cpp
@@ -136,7 +136,7 @@ FForwardLightingViewResources* GetMinimalDummyForwardLightingResources()
 	return &GMinimalDummyForwardLightingResources->ForwardLightingResources;
 }
 
-IMPLEMENT_UNIFORM_BUFFER_STRUCT(FForwardLightData, TEXT("ForwardLightData"));
+IMPLEMENT_GLOBAL_SHADER_PARAMETER_STRUCT(FForwardLightData, "ForwardLightData");
 
 FForwardLightData::FForwardLightData()
 {
diff --git a/Engine/Source/Runtime/Renderer/Private/LightMapDensityRendering.h b/Engine/Source/Runtime/Renderer/Private/LightMapDensityRendering.h
--- a/Engine/Source/Runtime/Renderer/Private/LightMapDensityRendering.h
+++ b/Engine/Source/Runtime/Renderer/Private/LightMapDensityRendering.h
@@ -53,7 +53,7 @@ public:
 		FMeshMaterialShader(Initializer)
 	{
 		LightMapPolicyType::VertexParametersType::Bind(Initializer.ParameterMap);
-		PassUniformBuffer.Bind(Initializer.ParameterMap, FSceneTexturesUniformParameters::StaticStruct.GetShaderVariableName());
+		PassUniformBuffer.Bind(Initializer.ParameterMap, FSceneTexturesUniformParameters::StaticStructMetadata.GetShaderVariableName());
 	}
 	TLightMapDensityVS() {}
 
@@ -180,7 +180,7 @@ public:
 		VertexMappedColor.Bind(Initializer.ParameterMap,TEXT("VertexMappedColor"));
 		GridTexture.Bind(Initializer.ParameterMap,TEXT("GridTexture"));
 		GridTextureSampler.Bind(Initializer.ParameterMap,TEXT("GridTextureSampler"));
-		PassUniformBuffer.Bind(Initializer.ParameterMap, FSceneTexturesUniformParameters::StaticStruct.GetShaderVariableName());
+		PassUniformBuffer.Bind(Initializer.ParameterMap, FSceneTexturesUniformParameters::StaticStructMetadata.GetShaderVariableName());
 	}
 	TLightMapDensityPS() {}
 
diff --git a/Engine/Source/Runtime/Renderer/Private/LightMapRendering.cpp b/Engine/Source/Runtime/Renderer/Private/LightMapRendering.cpp
--- a/Engine/Source/Runtime/Renderer/Private/LightMapRendering.cpp
+++ b/Engine/Source/Runtime/Renderer/Private/LightMapRendering.cpp
@@ -14,7 +14,7 @@ LightMapRendering.cpp: Light map rendering implementations.
 #include "Runtime/Engine/Classes/VT/VirtualTextureSpace.h"
 #include "VT/VirtualTextureSpace.h"
 
-IMPLEMENT_UNIFORM_BUFFER_STRUCT(FPrecomputedLightingParameters, TEXT("PrecomputedLightingBuffer"));
+IMPLEMENT_GLOBAL_SHADER_PARAMETER_STRUCT(FPrecomputedLightingParameters, "PrecomputedLightingBuffer");
 
 const TCHAR* GLightmapDefineName[2] =
 {
@@ -459,7 +459,7 @@ FUniformBufferRHIRef CreatePrecomputedLightingUniformBuffer(
 {
 	FPrecomputedLightingParameters Parameters;
 	GetPrecomputedLightingParameters(FeatureLevel, Parameters, LightingCache, LightingAllocation, VolumetricLightmapLookupPosition, SceneFrameNumber, VolumetricLightmapSceneData, LCI);
-	return FPrecomputedLightingParameters::CreateUniformBuffer(Parameters, BufferUsage);
+	return CreateUniformBufferImmediate(Parameters, BufferUsage);
 }
 
 void FEmptyPrecomputedLightingUniformBuffer::InitDynamicRHI()
diff --git a/Engine/Source/Runtime/Renderer/Private/LightMapRendering.h b/Engine/Source/Runtime/Renderer/Private/LightMapRendering.h
--- a/Engine/Source/Runtime/Renderer/Private/LightMapRendering.h
+++ b/Engine/Source/Runtime/Renderer/Private/LightMapRendering.h
@@ -21,46 +21,46 @@ extern ENGINE_API bool GShowDebugSelectedLightmap;
 extern ENGINE_API class FLightMap2D* GDebugSelectedLightmap;
 extern bool GVisualizeMipLevels;
 
-BEGIN_UNIFORM_BUFFER_STRUCT(FPrecomputedLightingParameters, )
-	UNIFORM_MEMBER(FVector, IndirectLightingCachePrimitiveAdd) // FCachedVolumeIndirectLightingPolicy
-	UNIFORM_MEMBER(FVector, IndirectLightingCachePrimitiveScale) // FCachedVolumeIndirectLightingPolicy
-	UNIFORM_MEMBER(FVector, IndirectLightingCacheMinUV) // FCachedVolumeIndirectLightingPolicy
-	UNIFORM_MEMBER(FVector, IndirectLightingCacheMaxUV) // FCachedVolumeIndirectLightingPolicy
-	UNIFORM_MEMBER(FVector4, PointSkyBentNormal) // FCachedPointIndirectLightingPolicy
-	UNIFORM_MEMBER_EX(float, DirectionalLightShadowing, EShaderPrecisionModifier::Half) // FCachedPointIndirectLightingPolicy
-	UNIFORM_MEMBER(FVector4, StaticShadowMapMasks) // TDistanceFieldShadowsAndLightMapPolicy
-	UNIFORM_MEMBER(FVector4, InvUniformPenumbraSizes) // TDistanceFieldShadowsAndLightMapPolicy
-	UNIFORM_MEMBER_ARRAY(FVector4, IndirectLightingSHCoefficients0, [3]) // FCachedPointIndirectLightingPolicy
-	UNIFORM_MEMBER_ARRAY(FVector4, IndirectLightingSHCoefficients1, [3]) // FCachedPointIndirectLightingPolicy
-	UNIFORM_MEMBER(FVector4,	IndirectLightingSHCoefficients2) // FCachedPointIndirectLightingPolicy
-	UNIFORM_MEMBER_EX(FVector4, IndirectLightingSHSingleCoefficient, EShaderPrecisionModifier::Half) // FCachedPointIndirectLightingPolicy used in forward Translucent
-	UNIFORM_MEMBER(FVector4, LightMapCoordinateScaleBias) // TLightMapPolicy
-	UNIFORM_MEMBER(FVector4, ShadowMapCoordinateScaleBias) // TDistanceFieldShadowsAndLightMapPolicy
-	UNIFORM_MEMBER_ARRAY_EX(FVector4, LightMapScale, [MAX_NUM_LIGHTMAP_COEF], EShaderPrecisionModifier::Half) // TLightMapPolicy
-	UNIFORM_MEMBER_ARRAY_EX(FVector4, LightMapAdd, [MAX_NUM_LIGHTMAP_COEF], EShaderPrecisionModifier::Half) // TLightMapPolicy
-	UNIFORM_MEMBER(FMatrix, LightmapVirtualTextureUniformData) // VT
-	UNIFORM_MEMBER_TEXTURE(Texture2D, LightMapTexture) // TLightMapPolicy
-	UNIFORM_MEMBER_TEXTURE(Texture2D, LightMapTexture_1) // VT
-	UNIFORM_MEMBER_TEXTURE(Texture2D, SkyOcclusionTexture) // TLightMapPolicy
-	UNIFORM_MEMBER_TEXTURE(Texture2D, AOMaterialMaskTexture) // TLightMapPolicy
-	UNIFORM_MEMBER_TEXTURE(Texture3D, IndirectLightingCacheTexture0) // FCachedVolumeIndirectLightingPolicy
-	UNIFORM_MEMBER_TEXTURE(Texture3D, IndirectLightingCacheTexture1) // FCachedVolumeIndirectLightingPolicy
-	UNIFORM_MEMBER_TEXTURE(Texture3D, IndirectLightingCacheTexture2) // FCachedVolumeIndirectLightingPolicy
-	UNIFORM_MEMBER_TEXTURE(Texture2D, StaticShadowTexture) // 
-	UNIFORM_MEMBER_SAMPLER(SamplerState, LightMapSampler) // TLightMapPolicy
-	UNIFORM_MEMBER_SAMPLER(SamplerState, SkyOcclusionSampler) // TLightMapPolicy
-	UNIFORM_MEMBER_SAMPLER(SamplerState, AOMaterialMaskSampler) // TLightMapPolicy
-	UNIFORM_MEMBER_SAMPLER(SamplerState, IndirectLightingCacheTextureSampler0) // FCachedVolumeIndirectLightingPolicy
-	UNIFORM_MEMBER_SAMPLER(SamplerState, IndirectLightingCacheTextureSampler1) // FCachedVolumeIndirectLightingPolicy
-	UNIFORM_MEMBER_SAMPLER(SamplerState, IndirectLightingCacheTextureSampler2) // FCachedVolumeIndirectLightingPolicy
-	UNIFORM_MEMBER_SAMPLER(SamplerState, StaticShadowTextureSampler) // TDistanceFieldShadowsAndLightMapPolicy
+BEGIN_GLOBAL_SHADER_PARAMETER_STRUCT(FPrecomputedLightingParameters, )
+	SHADER_PARAMETER(FVector, IndirectLightingCachePrimitiveAdd) // FCachedVolumeIndirectLightingPolicy
+	SHADER_PARAMETER(FVector, IndirectLightingCachePrimitiveScale) // FCachedVolumeIndirectLightingPolicy
+	SHADER_PARAMETER(FVector, IndirectLightingCacheMinUV) // FCachedVolumeIndirectLightingPolicy
+	SHADER_PARAMETER(FVector, IndirectLightingCacheMaxUV) // FCachedVolumeIndirectLightingPolicy
+	SHADER_PARAMETER(FVector4, PointSkyBentNormal) // FCachedPointIndirectLightingPolicy
+	SHADER_PARAMETER_EX(float, DirectionalLightShadowing, EShaderPrecisionModifier::Half) // FCachedPointIndirectLightingPolicy
+	SHADER_PARAMETER(FVector4, StaticShadowMapMasks) // TDistanceFieldShadowsAndLightMapPolicy
+	SHADER_PARAMETER(FVector4, InvUniformPenumbraSizes) // TDistanceFieldShadowsAndLightMapPolicy
+	SHADER_PARAMETER_ARRAY(FVector4, IndirectLightingSHCoefficients0, [3]) // FCachedPointIndirectLightingPolicy
+	SHADER_PARAMETER_ARRAY(FVector4, IndirectLightingSHCoefficients1, [3]) // FCachedPointIndirectLightingPolicy
+	SHADER_PARAMETER(FVector4,	IndirectLightingSHCoefficients2) // FCachedPointIndirectLightingPolicy
+	SHADER_PARAMETER_EX(FVector4, IndirectLightingSHSingleCoefficient, EShaderPrecisionModifier::Half) // FCachedPointIndirectLightingPolicy used in forward Translucent
+	SHADER_PARAMETER(FVector4, LightMapCoordinateScaleBias) // TLightMapPolicy
+	SHADER_PARAMETER(FVector4, ShadowMapCoordinateScaleBias) // TDistanceFieldShadowsAndLightMapPolicy
+	SHADER_PARAMETER_ARRAY_EX(FVector4, LightMapScale, [MAX_NUM_LIGHTMAP_COEF], EShaderPrecisionModifier::Half) // TLightMapPolicy
+	SHADER_PARAMETER_ARRAY_EX(FVector4, LightMapAdd, [MAX_NUM_LIGHTMAP_COEF], EShaderPrecisionModifier::Half) // TLightMapPolicy
+	SHADER_PARAMETER(FMatrix, LightmapVirtualTextureUniformData) // VT
+	SHADER_PARAMETER_TEXTURE(Texture2D, LightMapTexture) // TLightMapPolicy
+	SHADER_PARAMETER_TEXTURE(Texture2D, LightMapTexture_1) // VT
+	SHADER_PARAMETER_TEXTURE(Texture2D, SkyOcclusionTexture) // TLightMapPolicy
+	SHADER_PARAMETER_TEXTURE(Texture2D, AOMaterialMaskTexture) // TLightMapPolicy
+	SHADER_PARAMETER_TEXTURE(Texture3D, IndirectLightingCacheTexture0) // FCachedVolumeIndirectLightingPolicy
+	SHADER_PARAMETER_TEXTURE(Texture3D, IndirectLightingCacheTexture1) // FCachedVolumeIndirectLightingPolicy
+	SHADER_PARAMETER_TEXTURE(Texture3D, IndirectLightingCacheTexture2) // FCachedVolumeIndirectLightingPolicy
+	SHADER_PARAMETER_TEXTURE(Texture2D, StaticShadowTexture) // 
+	SHADER_PARAMETER_SAMPLER(SamplerState, LightMapSampler) // TLightMapPolicy
+	SHADER_PARAMETER_SAMPLER(SamplerState, SkyOcclusionSampler) // TLightMapPolicy
+	SHADER_PARAMETER_SAMPLER(SamplerState, AOMaterialMaskSampler) // TLightMapPolicy
+	SHADER_PARAMETER_SAMPLER(SamplerState, IndirectLightingCacheTextureSampler0) // FCachedVolumeIndirectLightingPolicy
+	SHADER_PARAMETER_SAMPLER(SamplerState, IndirectLightingCacheTextureSampler1) // FCachedVolumeIndirectLightingPolicy
+	SHADER_PARAMETER_SAMPLER(SamplerState, IndirectLightingCacheTextureSampler2) // FCachedVolumeIndirectLightingPolicy
+	SHADER_PARAMETER_SAMPLER(SamplerState, StaticShadowTextureSampler) // TDistanceFieldShadowsAndLightMapPolicy
 #if LIGHTMAP_VT_16BIT
-	UNIFORM_MEMBER_TEXTURE(Texture2D<uint>, LightmapVirtualTexturePageTable) // VT
+	SHADER_PARAMETER_TEXTURE(Texture2D<uint>, LightmapVirtualTexturePageTable) // VT
 #else
-	UNIFORM_MEMBER_TEXTURE(Texture2D, LightmapVirtualTexturePageTable) // VT
-	UNIFORM_MEMBER_SAMPLER(SamplerState, LightmapVirtualTexturePageTableSampler) // VT
+	SHADER_PARAMETER_TEXTURE(Texture2D, LightmapVirtualTexturePageTable) // VT
+	SHADER_PARAMETER_SAMPLER(SamplerState, LightmapVirtualTexturePageTableSampler) // VT
 #endif
-END_UNIFORM_BUFFER_STRUCT( FPrecomputedLightingParameters )
+END_GLOBAL_SHADER_PARAMETER_STRUCT()
 
 
 uint32 GetPrecompuledLightingVersionID(const FLightMapInteraction& LightMapInteraction, const FShadowMapInteraction& ShadowMapInteraction, ERHIFeatureLevel::Type FeatureLevel);
diff --git a/Engine/Source/Runtime/Renderer/Private/LightPropagationVolume.cpp b/Engine/Source/Runtime/Renderer/Private/LightPropagationVolume.cpp
--- a/Engine/Source/Runtime/Renderer/Private/LightPropagationVolume.cpp
+++ b/Engine/Source/Runtime/Renderer/Private/LightPropagationVolume.cpp
@@ -106,22 +106,22 @@ static TAutoConsoleVariable<float> CVarLPVDirectionalOcclusionDefaultSpecular(
 /**
  * Uniform buffer parameters for LPV direct injection shaders
  */
-BEGIN_UNIFORM_BUFFER_STRUCT( FLpvDirectLightInjectParameters, )
-	UNIFORM_MEMBER( float, LightRadius )
-	UNIFORM_MEMBER( FVector4, LightPosition )
-	UNIFORM_MEMBER( FVector4, LightColor )
-	UNIFORM_MEMBER( float, LightFalloffExponent )
-	UNIFORM_MEMBER( float, LightSourceLength )
-	UNIFORM_MEMBER( FVector4, LightDirection )
-	UNIFORM_MEMBER( FVector2D, LightSpotAngles )
-	UNIFORM_MEMBER( float, bLightInverseSquaredAttenuation )
-END_UNIFORM_BUFFER_STRUCT( FLpvDirectLightInjectParameters )
-
-IMPLEMENT_UNIFORM_BUFFER_STRUCT(FLpvDirectLightInjectParameters,TEXT("LpvInject"));
+BEGIN_GLOBAL_SHADER_PARAMETER_STRUCT( FLpvDirectLightInjectParameters, )
+	SHADER_PARAMETER( float, LightRadius )
+	SHADER_PARAMETER( FVector4, LightPosition )
+	SHADER_PARAMETER( FVector4, LightColor )
+	SHADER_PARAMETER( float, LightFalloffExponent )
+	SHADER_PARAMETER( float, LightSourceLength )
+	SHADER_PARAMETER( FVector4, LightDirection )
+	SHADER_PARAMETER( FVector2D, LightSpotAngles )
+	SHADER_PARAMETER( float, bLightInverseSquaredAttenuation )
+END_GLOBAL_SHADER_PARAMETER_STRUCT()
+
+IMPLEMENT_GLOBAL_SHADER_PARAMETER_STRUCT(FLpvDirectLightInjectParameters,"LpvInject");
 
 typedef TUniformBufferRef<FLpvDirectLightInjectParameters> FDirectLightInjectBufferRef;
 
-IMPLEMENT_UNIFORM_BUFFER_STRUCT(FLpvWriteUniformBufferParameters,TEXT("LpvWrite"));
+IMPLEMENT_GLOBAL_SHADER_PARAMETER_STRUCT(FLpvWriteUniformBufferParameters,"LpvWrite");
 
 // ----------------------------------------------------------------------------
 // Base LPV Write Compute shader
@@ -1339,7 +1339,7 @@ void FLightPropagationVolume::GetShaderParams( FLpvBaseWriteShaderParams& OutPar
 	OutParams.GvListHeadBufferUAV = GvListHeadBuffer->UAV;
 	OutParams.AOVolumeTextureUAV = AOVolumeTexture->GetRenderTargetItem().UAV;
 	OutParams.AOVolumeTextureSRV = AOVolumeTexture->GetRenderTargetItem().ShaderResourceTexture;
-	OutParams.UniformBuffer = LpvWriteUniformBuffer;
+	OutParams.UniformBuffer = LpvWriteUniformBuffer.GetUniformBufferRef();
 }
 
 void FLightPropagationVolume::InjectLightDirect(FRHICommandListImmediate& RHICmdList, const FLightSceneProxy& Light, const FViewInfo& View)
diff --git a/Engine/Source/Runtime/Renderer/Private/LightPropagationVolume.h b/Engine/Source/Runtime/Renderer/Private/LightPropagationVolume.h
--- a/Engine/Source/Runtime/Renderer/Private/LightPropagationVolume.h
+++ b/Engine/Source/Runtime/Renderer/Private/LightPropagationVolume.h
@@ -75,51 +75,51 @@ static const TCHAR * LpvGvVolumeTextureUAVNames[NUM_GV_TEXTURES] = {
 //
 // LPV Read constant buffer
 //
-BEGIN_UNIFORM_BUFFER_STRUCT( FLpvReadUniformBufferParameters, )
-	UNIFORM_MEMBER( FIntVector, mLpvGridOffset )
-	UNIFORM_MEMBER( float, LpvScale )
-	UNIFORM_MEMBER( float, OneOverLpvScale )
-	UNIFORM_MEMBER( float, SpecularIntensity )
-	UNIFORM_MEMBER( float, DiffuseIntensity )
-
-	UNIFORM_MEMBER( float, DirectionalOcclusionIntensity )
-	UNIFORM_MEMBER( float, DiffuseOcclusionExponent )
-	UNIFORM_MEMBER( float, SpecularOcclusionExponent )
-	UNIFORM_MEMBER( float, SpecularOcclusionIntensity )
-	UNIFORM_MEMBER( float, DiffuseOcclusionIntensity )
-	UNIFORM_MEMBER( float, PostprocessSpecularIntensityThreshold )
-
-	UNIFORM_MEMBER( FVector, LpvGridOffsetSmooth )
-	UNIFORM_MEMBER( FVector, DirectionalOcclusionDefaultValue )
-	UNIFORM_MEMBER( float, DirectionalOcclusionFadeRange )
-	UNIFORM_MEMBER( float, FadeRange )
-END_UNIFORM_BUFFER_STRUCT( FLpvReadUniformBufferParameters )
+BEGIN_GLOBAL_SHADER_PARAMETER_STRUCT( FLpvReadUniformBufferParameters, )
+	SHADER_PARAMETER( FIntVector, mLpvGridOffset )
+	SHADER_PARAMETER( float, LpvScale )
+	SHADER_PARAMETER( float, OneOverLpvScale )
+	SHADER_PARAMETER( float, SpecularIntensity )
+	SHADER_PARAMETER( float, DiffuseIntensity )
+
+	SHADER_PARAMETER( float, DirectionalOcclusionIntensity )
+	SHADER_PARAMETER( float, DiffuseOcclusionExponent )
+	SHADER_PARAMETER( float, SpecularOcclusionExponent )
+	SHADER_PARAMETER( float, SpecularOcclusionIntensity )
+	SHADER_PARAMETER( float, DiffuseOcclusionIntensity )
+	SHADER_PARAMETER( float, PostprocessSpecularIntensityThreshold )
+
+	SHADER_PARAMETER( FVector, LpvGridOffsetSmooth )
+	SHADER_PARAMETER( FVector, DirectionalOcclusionDefaultValue )
+	SHADER_PARAMETER( float, DirectionalOcclusionFadeRange )
+	SHADER_PARAMETER( float, FadeRange )
+END_GLOBAL_SHADER_PARAMETER_STRUCT()
 
 
 /**
  * Uniform buffer parameters for LPV write shaders
  */
-BEGIN_UNIFORM_BUFFER_STRUCT( FLpvWriteUniformBufferParameters, )
-	UNIFORM_MEMBER( FMatrix, mRsmToWorld )
-	UNIFORM_MEMBER( FVector4, mLightColour )
-	UNIFORM_MEMBER( FVector4, GeometryVolumeCaptureLightDirection )
-	UNIFORM_MEMBER( FVector4, mEyePos )
-	UNIFORM_MEMBER( FIntVector, mOldGridOffset ) 
-	UNIFORM_MEMBER( FIntVector, mLpvGridOffset )
-	UNIFORM_MEMBER( float, ClearMultiplier )
-	UNIFORM_MEMBER( float, LpvScale )
-	UNIFORM_MEMBER( float, OneOverLpvScale )
-	UNIFORM_MEMBER( float, DirectionalOcclusionIntensity )
-	UNIFORM_MEMBER( float, DirectionalOcclusionRadius )
-	UNIFORM_MEMBER( float, RsmAreaIntensityMultiplier )
-	UNIFORM_MEMBER( float, RsmPixelToTexcoordMultiplier )
-	UNIFORM_MEMBER( float, SecondaryOcclusionStrength )
-	UNIFORM_MEMBER( float, SecondaryBounceStrength )
-	UNIFORM_MEMBER( float, VplInjectionBias )
-	UNIFORM_MEMBER( float, GeometryVolumeInjectionBias )
-	UNIFORM_MEMBER( float, EmissiveInjectionMultiplier )
-	UNIFORM_MEMBER( int,	 PropagationIndex )
-END_UNIFORM_BUFFER_STRUCT( FLpvWriteUniformBufferParameters )
+BEGIN_GLOBAL_SHADER_PARAMETER_STRUCT( FLpvWriteUniformBufferParameters, )
+	SHADER_PARAMETER( FMatrix, mRsmToWorld )
+	SHADER_PARAMETER( FVector4, mLightColour )
+	SHADER_PARAMETER( FVector4, GeometryVolumeCaptureLightDirection )
+	SHADER_PARAMETER( FVector4, mEyePos )
+	SHADER_PARAMETER( FIntVector, mOldGridOffset ) 
+	SHADER_PARAMETER( FIntVector, mLpvGridOffset )
+	SHADER_PARAMETER( float, ClearMultiplier )
+	SHADER_PARAMETER( float, LpvScale )
+	SHADER_PARAMETER( float, OneOverLpvScale )
+	SHADER_PARAMETER( float, DirectionalOcclusionIntensity )
+	SHADER_PARAMETER( float, DirectionalOcclusionRadius )
+	SHADER_PARAMETER( float, RsmAreaIntensityMultiplier )
+	SHADER_PARAMETER( float, RsmPixelToTexcoordMultiplier )
+	SHADER_PARAMETER( float, SecondaryOcclusionStrength )
+	SHADER_PARAMETER( float, SecondaryBounceStrength )
+	SHADER_PARAMETER( float, VplInjectionBias )
+	SHADER_PARAMETER( float, GeometryVolumeInjectionBias )
+	SHADER_PARAMETER( float, EmissiveInjectionMultiplier )
+	SHADER_PARAMETER( int,	 PropagationIndex )
+END_GLOBAL_SHADER_PARAMETER_STRUCT()
 
 // ----------------------------------------------------------------------------
 // Shader params for base LPV write shaders
@@ -185,8 +185,8 @@ public:
 	const FLpvReadUniformBufferParameters& GetReadUniformBufferParams()		{ return LpvReadUniformBufferParams; }
 	const FLpvWriteUniformBufferParameters& GetWriteUniformBufferParams()	{ return *LpvWriteUniformBufferParams; }
 
-	FLpvWriteUniformBufferRef GetWriteUniformBuffer() const				{ return (FLpvWriteUniformBufferRef)LpvWriteUniformBuffer; }
-	FLpvWriteUniformBufferRef GetRsmUniformBuffer() const				{ return (FLpvWriteUniformBufferRef)RsmRenderUniformBuffer; }
+	FLpvWriteUniformBufferRef GetWriteUniformBuffer() const				{ return LpvWriteUniformBuffer.GetUniformBufferRef(); }
+	FLpvWriteUniformBufferRef GetRsmUniformBuffer() const				{ return RsmRenderUniformBuffer.GetUniformBufferRef(); }
 
 	FTextureRHIParamRef GetLpvBufferSrv( int i )						{ return LpvVolumeTextures[ 1-mWriteBufferIndex ][i]->GetRenderTargetItem().ShaderResourceTexture; }
 
diff --git a/Engine/Source/Runtime/Renderer/Private/LightRendering.cpp b/Engine/Source/Runtime/Renderer/Private/LightRendering.cpp
--- a/Engine/Source/Runtime/Renderer/Private/LightRendering.cpp
+++ b/Engine/Source/Runtime/Renderer/Private/LightRendering.cpp
@@ -16,7 +16,7 @@
 
 DECLARE_GPU_STAT(Lights);
 
-IMPLEMENT_UNIFORM_BUFFER_STRUCT(FDeferredLightUniformStruct,TEXT("DeferredLightUniforms"));
+IMPLEMENT_GLOBAL_SHADER_PARAMETER_STRUCT(FDeferredLightUniformStruct, "DeferredLightUniforms");
 
 extern int32 GUseTranslucentLightingVolumes;
 ENGINE_API const IPooledRenderTarget* GetSubsufaceProfileTexture_RT(FRHICommandListImmediate& RHICmdList);
@@ -832,7 +832,7 @@ void FDeferredShadingSceneRenderer::RenderLights(FRHICommandListImmediate& RHICm
 					}
 				}
 
-				GRenderTargetPool.VisualizeTexture.SetCheckPoint(RHICmdList, ScreenShadowMaskTexture);
+				GVisualizeTexture.SetCheckPoint(RHICmdList, ScreenShadowMaskTexture);
 
 				SceneContext.BeginRenderingSceneColor(RHICmdList, ESimpleRenderTargetMode::EExistingColorAndDepth, FExclusiveDepthStencil::DepthRead_StencilWrite);
 
diff --git a/Engine/Source/Runtime/Renderer/Private/LightRendering.h b/Engine/Source/Runtime/Renderer/Private/LightRendering.h
--- a/Engine/Source/Runtime/Renderer/Private/LightRendering.h
+++ b/Engine/Source/Runtime/Renderer/Private/LightRendering.h
@@ -19,26 +19,26 @@
 #include "LightSceneInfo.h"
 
 /** Uniform buffer for rendering deferred lights. */
-BEGIN_UNIFORM_BUFFER_STRUCT(FDeferredLightUniformStruct,)
-	UNIFORM_MEMBER(FVector,LightPosition)
-	UNIFORM_MEMBER(float,LightInvRadius)
-	UNIFORM_MEMBER(FVector,LightColor)
-	UNIFORM_MEMBER(float,LightFalloffExponent)
-	UNIFORM_MEMBER(FVector,NormalizedLightDirection)
-	UNIFORM_MEMBER(FVector,NormalizedLightTangent)
-	UNIFORM_MEMBER(FVector2D,SpotAngles)
-	UNIFORM_MEMBER(float,SpecularScale)
-	UNIFORM_MEMBER(float,SourceRadius)
-	UNIFORM_MEMBER(float,SoftSourceRadius)
-	UNIFORM_MEMBER(float,SourceLength)
-	UNIFORM_MEMBER(float,ContactShadowLength)
-	UNIFORM_MEMBER(FVector2D,DistanceFadeMAD)
-	UNIFORM_MEMBER(FVector4,ShadowMapChannelMask)
-	UNIFORM_MEMBER(uint32,ShadowedBits)
-	UNIFORM_MEMBER(uint32,LightingChannelMask)
-	UNIFORM_MEMBER(float,VolumetricScatteringIntensity)
-	UNIFORM_MEMBER_TEXTURE(Texture2D, SourceTexture)
-END_UNIFORM_BUFFER_STRUCT(FDeferredLightUniformStruct)
+BEGIN_GLOBAL_SHADER_PARAMETER_STRUCT(FDeferredLightUniformStruct,)
+	SHADER_PARAMETER(FVector,LightPosition)
+	SHADER_PARAMETER(float,LightInvRadius)
+	SHADER_PARAMETER(FVector,LightColor)
+	SHADER_PARAMETER(float,LightFalloffExponent)
+	SHADER_PARAMETER(FVector,NormalizedLightDirection)
+	SHADER_PARAMETER(FVector,NormalizedLightTangent)
+	SHADER_PARAMETER(FVector2D,SpotAngles)
+	SHADER_PARAMETER(float,SpecularScale)
+	SHADER_PARAMETER(float,SourceRadius)
+	SHADER_PARAMETER(float,SoftSourceRadius)
+	SHADER_PARAMETER(float,SourceLength)
+	SHADER_PARAMETER(float,ContactShadowLength)
+	SHADER_PARAMETER(FVector2D,DistanceFadeMAD)
+	SHADER_PARAMETER(FVector4,ShadowMapChannelMask)
+	SHADER_PARAMETER(uint32,ShadowedBits)
+	SHADER_PARAMETER(uint32,LightingChannelMask)
+	SHADER_PARAMETER(float,VolumetricScatteringIntensity)
+	SHADER_PARAMETER_TEXTURE(Texture2D, SourceTexture)
+END_GLOBAL_SHADER_PARAMETER_STRUCT()
 
 extern uint32 GetShadowQuality();
 
diff --git a/Engine/Source/Runtime/Renderer/Private/MaterialTexCoordScalesRendering.h b/Engine/Source/Runtime/Renderer/Private/MaterialTexCoordScalesRendering.h
--- a/Engine/Source/Runtime/Renderer/Private/MaterialTexCoordScalesRendering.h
+++ b/Engine/Source/Runtime/Renderer/Private/MaterialTexCoordScalesRendering.h
@@ -44,7 +44,7 @@ public:
 		OneOverCPUTexCoordScalesParameter.Bind(Initializer.ParameterMap,TEXT("OneOverCPUTexCoordScales"));
 		TexCoordIndicesParameter.Bind(Initializer.ParameterMap, TEXT("TexCoordIndices"));
 		PrimitiveAlphaParameter.Bind(Initializer.ParameterMap, TEXT("PrimitiveAlpha"));
-		PassUniformBuffer.Bind(Initializer.ParameterMap, FSceneTexturesUniformParameters::StaticStruct.GetShaderVariableName());
+		PassUniformBuffer.Bind(Initializer.ParameterMap, FSceneTexturesUniformParameters::StaticStructMetadata.GetShaderVariableName());
 	}
 
 	FMaterialTexCoordScalePS() {}
diff --git a/Engine/Source/Runtime/Renderer/Private/MobileBasePassRendering.cpp b/Engine/Source/Runtime/Renderer/Private/MobileBasePassRendering.cpp
--- a/Engine/Source/Runtime/Renderer/Private/MobileBasePassRendering.cpp
+++ b/Engine/Source/Runtime/Renderer/Private/MobileBasePassRendering.cpp
@@ -35,7 +35,7 @@ static TAutoConsoleVariable<int32> CVarMobileSeparateMaskedPass(
 	1,
 	TEXT("Draw masked primitives in separate pass after all opaque (default)"),
 	ECVF_RenderThreadSafe);
-IMPLEMENT_UNIFORM_BUFFER_STRUCT(FMobileBasePassUniformParameters, TEXT("MobileBasePass"));
+IMPLEMENT_GLOBAL_SHADER_PARAMETER_STRUCT(FMobileBasePassUniformParameters, "MobileBasePass");
 
 static TAutoConsoleVariable<int32> CVarMobileParallelBasePass(
 	TEXT("r.Mobile.ParallelBasePass"),
diff --git a/Engine/Source/Runtime/Renderer/Private/MobileBasePassRendering.h b/Engine/Source/Runtime/Renderer/Private/MobileBasePassRendering.h
--- a/Engine/Source/Runtime/Renderer/Private/MobileBasePassRendering.h
+++ b/Engine/Source/Runtime/Renderer/Private/MobileBasePassRendering.h
@@ -25,10 +25,10 @@
 #include "PlanarReflectionRendering.h"
 #include "BasePassRendering.h"
 
-BEGIN_UNIFORM_BUFFER_STRUCT(FMobileBasePassUniformParameters, )
-	UNIFORM_MEMBER_STRUCT(FFogUniformParameters, Fog)
-	UNIFORM_MEMBER_STRUCT(FMobileSceneTextureUniformParameters, SceneTextures)
-END_UNIFORM_BUFFER_STRUCT(FMobileBasePassUniformParameters)
+BEGIN_GLOBAL_SHADER_PARAMETER_STRUCT(FMobileBasePassUniformParameters, )
+	SHADER_PARAMETER_STRUCT(FFogUniformParameters, Fog)
+	SHADER_PARAMETER_STRUCT(FMobileSceneTextureUniformParameters, SceneTextures)
+END_GLOBAL_SHADER_PARAMETER_STRUCT()
 
 extern void CreateMobileBasePassUniformBuffer(
 	FRHICommandListImmediate& RHICmdList,
@@ -84,7 +84,7 @@ protected:
 		FMeshMaterialShader(Initializer)
 	{
 		VertexParametersType::Bind(Initializer.ParameterMap);
-		PassUniformBuffer.Bind(Initializer.ParameterMap, FMobileBasePassUniformParameters::StaticStruct.GetShaderVariableName());
+		PassUniformBuffer.Bind(Initializer.ParameterMap, FMobileBasePassUniformParameters::StaticStructMetadata.GetShaderVariableName());
 	}
 
 public:
@@ -210,7 +210,7 @@ public:
 		FMeshMaterialShader(Initializer)
 	{
 		PixelParametersType::Bind(Initializer.ParameterMap);
-		PassUniformBuffer.Bind(Initializer.ParameterMap, FMobileBasePassUniformParameters::StaticStruct.GetShaderVariableName());
+		PassUniformBuffer.Bind(Initializer.ParameterMap, FMobileBasePassUniformParameters::StaticStructMetadata.GetShaderVariableName());
 		ReflectionCubemap.Bind(Initializer.ParameterMap, TEXT("ReflectionCubemap"));
 		ReflectionSampler.Bind(Initializer.ParameterMap, TEXT("ReflectionCubemapSampler"));
 		InvReflectionCubemapAverageBrightness.Bind(Initializer.ParameterMap, TEXT("InvReflectionCubemapAverageBrightness"));
diff --git a/Engine/Source/Runtime/Renderer/Private/MobileSceneCaptureRendering.cpp b/Engine/Source/Runtime/Renderer/Private/MobileSceneCaptureRendering.cpp
--- a/Engine/Source/Runtime/Renderer/Private/MobileSceneCaptureRendering.cpp
+++ b/Engine/Source/Runtime/Renderer/Private/MobileSceneCaptureRendering.cpp
@@ -21,7 +21,7 @@ MobileSceneCaptureRendering.cpp - Mobile specific scene capture code.
 #include "SceneRenderTargetParameters.h"
 #include "PostProcess/SceneRenderTargets.h"
 #include "SceneRendering.h"
-#include "PostProcess/RenderTargetPool.h"
+#include "RenderTargetPool.h"
 #include "PostProcess/SceneFilterRendering.h"
 #include "ScreenRendering.h"
 #include "ClearQuad.h"
diff --git a/Engine/Source/Runtime/Renderer/Private/MobileShadingRenderer.cpp b/Engine/Source/Runtime/Renderer/Private/MobileShadingRenderer.cpp
--- a/Engine/Source/Runtime/Renderer/Private/MobileShadingRenderer.cpp
+++ b/Engine/Source/Runtime/Renderer/Private/MobileShadingRenderer.cpp
@@ -209,7 +209,7 @@ void FMobileSceneRenderer::Render(FRHICommandListImmediate& RHICmdList)
 	FRHICommandListExecutor::GetImmediateCommandList().PollOcclusionQueries();
 	RHICmdList.ImmediateFlush(EImmediateFlushType::DispatchToRHIThread);
 
-	GRenderTargetPool.VisualizeTexture.OnStartFrame(Views[0]);
+	GVisualizeTexture.OnStartFrame(Views[0]);
 
 	RHICmdList.SetCurrentStat(GET_STATID(STAT_CLMM_Shadows));
 
@@ -337,7 +337,7 @@ void FMobileSceneRenderer::Render(FRHICommandListImmediate& RHICmdList)
 		TUniformBufferRef<FMobileSceneTextureUniformParameters> MobileSceneTextureUniformBuffer = TUniformBufferRef<FMobileSceneTextureUniformParameters>::CreateUniformBufferImmediate(MobileSceneTextureParameters, UniformBuffer_SingleFrame);
 
 		// This is switching to another RT!
-		Scene->FXSystem->PostRenderOpaque(RHICmdList, View.ViewUniformBuffer, &FMobileSceneTextureUniformParameters::StaticStruct, MobileSceneTextureUniformBuffer.GetReference());
+		Scene->FXSystem->PostRenderOpaque(RHICmdList, View.ViewUniformBuffer, &FMobileSceneTextureUniformParameters::StaticStructMetadata, MobileSceneTextureUniformBuffer.GetReference());
 	}
 
 	if (!View.bIsPlanarReflection)
diff --git a/Engine/Source/Runtime/Renderer/Private/MobileTranslucentRendering.cpp b/Engine/Source/Runtime/Renderer/Private/MobileTranslucentRendering.cpp
--- a/Engine/Source/Runtime/Renderer/Private/MobileTranslucentRendering.cpp
+++ b/Engine/Source/Runtime/Renderer/Private/MobileTranslucentRendering.cpp
@@ -386,7 +386,7 @@ protected:
 	FOpacityOnlyVS(const FMeshMaterialShaderType::CompiledShaderInitializerType& Initializer) :
 		FMeshMaterialShader(Initializer)
 	{
-		PassUniformBuffer.Bind(Initializer.ParameterMap, FMobileSceneTextureUniformParameters::StaticStruct.GetShaderVariableName());
+		PassUniformBuffer.Bind(Initializer.ParameterMap, FMobileSceneTextureUniformParameters::StaticStructMetadata.GetShaderVariableName());
 	}
 
 public:
@@ -449,7 +449,7 @@ public:
 	FOpacityOnlyPS(const ShaderMetaType::CompiledShaderInitializerType& Initializer) :
 		FMeshMaterialShader(Initializer)
 	{
-		PassUniformBuffer.Bind(Initializer.ParameterMap, FMobileSceneTextureUniformParameters::StaticStruct.GetShaderVariableName());
+		PassUniformBuffer.Bind(Initializer.ParameterMap, FMobileSceneTextureUniformParameters::StaticStructMetadata.GetShaderVariableName());
 	}
 
 	FOpacityOnlyPS() {}
diff --git a/Engine/Source/Runtime/Renderer/Private/PostProcess/PostProcessCompositeEditorPrimitives.cpp b/Engine/Source/Runtime/Renderer/Private/PostProcess/PostProcessCompositeEditorPrimitives.cpp
--- a/Engine/Source/Runtime/Renderer/Private/PostProcess/PostProcessCompositeEditorPrimitives.cpp
+++ b/Engine/Source/Runtime/Renderer/Private/PostProcess/PostProcessCompositeEditorPrimitives.cpp
@@ -8,7 +8,7 @@
 #include "SceneRenderTargetParameters.h"
 #include "BasePassRendering.h"
 #include "MobileBasePassRendering.h"
-#include "PostProcess/RenderTargetPool.h"
+#include "RenderTargetPool.h"
 #include "DynamicPrimitiveDrawing.h"
 #include "ClearQuad.h"
 #include "PipelineStateCache.h"
@@ -526,7 +526,7 @@ void FRCPassPostProcessCompositeEditorPrimitives::Process(FRenderingCompositePas
 			}
 		}
 
-		GRenderTargetPool.VisualizeTexture.SetCheckPoint(Context.RHICmdList, SceneContext.EditorPrimitivesColor);
+		GVisualizeTexture.SetCheckPoint(Context.RHICmdList, SceneContext.EditorPrimitivesColor);
 
 		FTextureRHIParamRef EditorRenderTargets[2];
 		EditorRenderTargets[0] = EditorColorTarget;
diff --git a/Engine/Source/Runtime/Renderer/Private/PostProcess/PostProcessEyeAdaptation.cpp b/Engine/Source/Runtime/Renderer/Private/PostProcess/PostProcessEyeAdaptation.cpp
--- a/Engine/Source/Runtime/Renderer/Private/PostProcess/PostProcessEyeAdaptation.cpp
+++ b/Engine/Source/Runtime/Renderer/Private/PostProcess/PostProcessEyeAdaptation.cpp
@@ -12,7 +12,7 @@
 #include "Math/UnrealMathUtility.h"
 #include "ScenePrivate.h"
 
-SHADERCORE_API bool UsePreExposure(EShaderPlatform Platform);
+RENDERCORE_API bool UsePreExposure(EShaderPlatform Platform);
 
 // Initialize the static CVar
 TAutoConsoleVariable<float> CVarEyeAdaptationPreExposureOverride(
diff --git a/Engine/Source/Runtime/Renderer/Private/PostProcess/PostProcessSubsurface.cpp b/Engine/Source/Runtime/Renderer/Private/PostProcess/PostProcessSubsurface.cpp
--- a/Engine/Source/Runtime/Renderer/Private/PostProcess/PostProcessSubsurface.cpp
+++ b/Engine/Source/Runtime/Renderer/Private/PostProcess/PostProcessSubsurface.cpp
@@ -11,7 +11,7 @@
 #include "CanvasTypes.h"
 #include "UnrealEngine.h"
 #include "SceneUtils.h"
-#include "PostProcess/RenderTargetPool.h"
+#include "RenderTargetPool.h"
 #include "RenderTargetTemp.h"
 #include "PostProcess/SceneRenderTargets.h"
 #include "PostProcess/SceneFilterRendering.h"
@@ -730,7 +730,7 @@ void FRCPassPostProcessSubsurface::Process(FRenderingCompositePassContext& Conte
 		check(PooledRT);
 
 		// for debugging
-		GRenderTargetPool.VisualizeTexture.SetCheckPoint(Context.RHICmdList, PooledRT);
+		GVisualizeTexture.SetCheckPoint(Context.RHICmdList, PooledRT);
 	}
 
 	const FViewInfo& View = Context.View;
diff --git a/Engine/Source/Runtime/Renderer/Private/PostProcess/PostProcessTonemap.cpp b/Engine/Source/Runtime/Renderer/Private/PostProcess/PostProcessTonemap.cpp
--- a/Engine/Source/Runtime/Renderer/Private/PostProcess/PostProcessTonemap.cpp
+++ b/Engine/Source/Runtime/Renderer/Private/PostProcess/PostProcessTonemap.cpp
@@ -445,12 +445,12 @@ void FilmPostSetConstants(FVector4* RESTRICT const Constants, const FPostProcess
 }
 
 
-BEGIN_UNIFORM_BUFFER_STRUCT(FBloomDirtMaskParameters,)
-	UNIFORM_MEMBER(FVector4,Tint)
-	UNIFORM_MEMBER_TEXTURE(Texture2D,Mask)
-	UNIFORM_MEMBER_SAMPLER(SamplerState,MaskSampler)
-END_UNIFORM_BUFFER_STRUCT(FBloomDirtMaskParameters)
-IMPLEMENT_UNIFORM_BUFFER_STRUCT(FBloomDirtMaskParameters,TEXT("BloomDirtMask"));
+BEGIN_GLOBAL_SHADER_PARAMETER_STRUCT(FBloomDirtMaskParameters,)
+	SHADER_PARAMETER(FVector4,Tint)
+	SHADER_PARAMETER_TEXTURE(Texture2D,Mask)
+	SHADER_PARAMETER_SAMPLER(SamplerState,MaskSampler)
+END_GLOBAL_SHADER_PARAMETER_STRUCT()
+IMPLEMENT_GLOBAL_SHADER_PARAMETER_STRUCT(FBloomDirtMaskParameters,"BloomDirtMask");
 
 
 // ---------------------------------------------------- Shared parameters for desktop's PS and CS
diff --git a/Engine/Source/Runtime/Renderer/Private/PostProcess/RenderingCompositionGraph.cpp b/Engine/Source/Runtime/Renderer/Private/PostProcess/RenderingCompositionGraph.cpp
--- a/Engine/Source/Runtime/Renderer/Private/PostProcess/RenderingCompositionGraph.cpp
+++ b/Engine/Source/Runtime/Renderer/Private/PostProcess/RenderingCompositionGraph.cpp
@@ -10,7 +10,7 @@
 #include "Async/Async.h"
 #include "EngineGlobals.h"
 #include "Engine/Engine.h"
-#include "PostProcess/RenderTargetPool.h"
+#include "RenderTargetPool.h"
 #include "RendererModule.h"
 #include "HighResScreenshot.h"
 #include "IHeadMountedDisplay.h"
@@ -691,7 +691,7 @@ void FRenderingCompositionGraph::RecursivelyProcess(const FRenderingCompositeOut
 			// use intermediate texture unless it's the last one where we render to the final output
 			if(PassOutput->PooledRenderTarget)
 			{
-				GRenderTargetPool.VisualizeTexture.SetCheckPoint(Context.RHICmdList, PassOutput->PooledRenderTarget);
+				GVisualizeTexture.SetCheckPoint(Context.RHICmdList, PassOutput->PooledRenderTarget);
 
 				// If this buffer was given a dump filename, write it out
 				const FString& Filename = Pass->GetOutputDumpFilename((EPassOutputId)OutputId);
diff --git a/Engine/Source/Runtime/Renderer/Private/PostProcess/SceneFilterRendering.cpp b/Engine/Source/Runtime/Renderer/Private/PostProcess/SceneFilterRendering.cpp
--- a/Engine/Source/Runtime/Renderer/Private/PostProcess/SceneFilterRendering.cpp
+++ b/Engine/Source/Runtime/Renderer/Private/PostProcess/SceneFilterRendering.cpp
@@ -127,7 +127,7 @@ TGlobalResource<FFilterVertexDeclaration> GFilterVertexDeclaration;
 /** Vertex declaration for vertex shaders that don't require any inputs (eg. generated via vertex ID). */
 TGlobalResource<FEmptyVertexDeclaration> GEmptyVertexDeclaration;
 
-IMPLEMENT_UNIFORM_BUFFER_STRUCT(FDrawRectangleParameters,TEXT("DrawRectangleParameters"));
+IMPLEMENT_GLOBAL_SHADER_PARAMETER_STRUCT(FDrawRectangleParameters, "DrawRectangleParameters");
 
 typedef TUniformBufferRef<FDrawRectangleParameters> FDrawRectangleBufferRef;
 
diff --git a/Engine/Source/Runtime/Renderer/Private/PostProcess/SceneFilterRendering.h b/Engine/Source/Runtime/Renderer/Private/PostProcess/SceneFilterRendering.h
--- a/Engine/Source/Runtime/Renderer/Private/PostProcess/SceneFilterRendering.h
+++ b/Engine/Source/Runtime/Renderer/Private/PostProcess/SceneFilterRendering.h
@@ -15,11 +15,11 @@
 #include "ShaderParameterUtils.h"
 
 /** Uniform buffer for computing the vertex positional and UV adjustments in the vertex shader. */
-BEGIN_UNIFORM_BUFFER_STRUCT( FDrawRectangleParameters,)
-	UNIFORM_MEMBER( FVector4, PosScaleBias )
-	UNIFORM_MEMBER( FVector4, UVScaleBias )
-	UNIFORM_MEMBER( FVector4, InvTargetSizeAndTextureSize )
-END_UNIFORM_BUFFER_STRUCT( FDrawRectangleParameters )
+BEGIN_GLOBAL_SHADER_PARAMETER_STRUCT( FDrawRectangleParameters,)
+	SHADER_PARAMETER( FVector4, PosScaleBias )
+	SHADER_PARAMETER( FVector4, UVScaleBias )
+	SHADER_PARAMETER( FVector4, InvTargetSizeAndTextureSize )
+END_GLOBAL_SHADER_PARAMETER_STRUCT()
 
 /**
  * Draws a quad with the given vertex positions and UVs in denormalized pixel/texel coordinates.
diff --git a/Engine/Source/Runtime/Renderer/Private/PostProcess/SceneRenderTargets.cpp b/Engine/Source/Runtime/Renderer/Private/PostProcess/SceneRenderTargets.cpp
--- a/Engine/Source/Runtime/Renderer/Private/PostProcess/SceneRenderTargets.cpp
+++ b/Engine/Source/Runtime/Renderer/Private/PostProcess/SceneRenderTargets.cpp
@@ -1262,13 +1262,13 @@ void FSceneRenderTargets::BeginRenderingPrePass(FRHICommandList& RHICmdList, boo
 void FSceneRenderTargets::FinishRenderingPrePass(FRHICommandListImmediate& RHICmdList)
 {
 	SCOPED_DRAW_EVENT(RHICmdList, FinishRenderingPrePass);
-	GRenderTargetPool.VisualizeTexture.SetCheckPoint(RHICmdList, SceneDepthZ);
+	GVisualizeTexture.SetCheckPoint(RHICmdList, SceneDepthZ);
 }
 
 void FSceneRenderTargets::BeginRenderingSceneAlphaCopy(FRHICommandListImmediate& RHICmdList)
 {
 	SCOPED_DRAW_EVENT(RHICmdList, BeginRenderingSceneAlphaCopy);
-	GRenderTargetPool.VisualizeTexture.SetCheckPoint(RHICmdList, SceneAlphaCopy);
+	GVisualizeTexture.SetCheckPoint(RHICmdList, SceneAlphaCopy);
 	SetRenderTarget(RHICmdList, GetSceneAlphaCopySurface(), 0);
 }
 
@@ -1276,7 +1276,7 @@ void FSceneRenderTargets::FinishRenderingSceneAlphaCopy(FRHICommandListImmediate
 {
 	SCOPED_DRAW_EVENT(RHICmdList, FinishRenderingSceneAlphaCopy);
 	RHICmdList.CopyToResolveTarget(GetSceneAlphaCopySurface(), SceneAlphaCopy->GetRenderTargetItem().ShaderResourceTexture, FResolveParams(FResolveRect()));
-	GRenderTargetPool.VisualizeTexture.SetCheckPoint(RHICmdList, SceneAlphaCopy);
+	GVisualizeTexture.SetCheckPoint(RHICmdList, SceneAlphaCopy);
 }
 
 
@@ -1287,7 +1287,7 @@ void FSceneRenderTargets::BeginRenderingLightAttenuation(FRHICommandList& RHICmd
 
 	AllocLightAttenuation(RHICmdList);
 
-	GRenderTargetPool.VisualizeTexture.SetCheckPoint(RHICmdList, GetLightAttenuation());
+	GVisualizeTexture.SetCheckPoint(RHICmdList, GetLightAttenuation());
 
 	// Set the light attenuation surface as the render target, and the scene depth buffer as the depth-stencil surface.
 	if (bClearToWhite)
@@ -1311,7 +1311,7 @@ void FSceneRenderTargets::FinishRenderingLightAttenuation(FRHICommandList& RHICm
 	// Resolve the light attenuation surface.
 	RHICmdList.CopyToResolveTarget(GetLightAttenuationSurface(), LightAttenuation->GetRenderTargetItem().ShaderResourceTexture, FResolveParams(FResolveRect()));
 	
-	GRenderTargetPool.VisualizeTexture.SetCheckPoint(RHICmdList, GetLightAttenuation());
+	GVisualizeTexture.SetCheckPoint(RHICmdList, GetLightAttenuation());
 }
 
 TRefCountPtr<IPooledRenderTarget>& FSceneRenderTargets::GetSeparateTranslucency(FRHICommandList& RHICmdList, FIntPoint Size)
@@ -2657,7 +2657,7 @@ bool FSceneRenderTargets::AreShadingPathRenderTargetsAllocated(ESceneColorFormat
 FSceneTextureShaderParameters
 -----------------------------------------------------------------------------*/
 
-IMPLEMENT_UNIFORM_BUFFER_STRUCT(FSceneTexturesUniformParameters, TEXT("SceneTexturesStruct"));
+IMPLEMENT_GLOBAL_SHADER_PARAMETER_STRUCT(FSceneTexturesUniformParameters, "SceneTexturesStruct");
 
 void SetupSceneTextureUniformParameters(
 	FSceneRenderTargets& SceneContext,
@@ -2800,7 +2800,7 @@ IMPLEMENT_CreateSceneTextureUniformBuffer( FRHICommandList );
 IMPLEMENT_CreateSceneTextureUniformBuffer( FRHICommandListImmediate );
 IMPLEMENT_CreateSceneTextureUniformBuffer( FRHIAsyncComputeCommandListImmediate );
 
-IMPLEMENT_UNIFORM_BUFFER_STRUCT(FMobileSceneTextureUniformParameters, TEXT("MobileSceneTextures"));
+IMPLEMENT_GLOBAL_SHADER_PARAMETER_STRUCT(FMobileSceneTextureUniformParameters, "MobileSceneTextures");
 
 void SetupMobileSceneTextureUniformParameters(
 	FSceneRenderTargets& SceneContext,
@@ -2872,13 +2872,13 @@ void BindSceneTextureUniformBufferDependentOnShadingPath(
 		
 	if (FSceneInterface::GetShadingPath(FeatureLevel) == EShadingPath::Deferred)
 	{
-		SceneTexturesUniformBuffer.Bind(Initializer.ParameterMap, FSceneTexturesUniformParameters::StaticStruct.GetShaderVariableName());
-		checkfSlow(!Initializer.ParameterMap.ContainsParameterAllocation(FMobileSceneTextureUniformParameters::StaticStruct.GetShaderVariableName()), TEXT("Shader for Deferred shading path tried to bind FMobileSceneTextureUniformParameters which is only available in the mobile shading path: %s"), Initializer.Type->GetName());
+		SceneTexturesUniformBuffer.Bind(Initializer.ParameterMap, FSceneTexturesUniformParameters::StaticStructMetadata.GetShaderVariableName());
+		checkfSlow(!Initializer.ParameterMap.ContainsParameterAllocation(FMobileSceneTextureUniformParameters::StaticStructMetadata.GetShaderVariableName()), TEXT("Shader for Deferred shading path tried to bind FMobileSceneTextureUniformParameters which is only available in the mobile shading path: %s"), Initializer.Type->GetName());
 	}
 		
 	if (FSceneInterface::GetShadingPath(FeatureLevel) == EShadingPath::Mobile)
 	{
-		MobileSceneTexturesUniformBuffer.Bind(Initializer.ParameterMap, FMobileSceneTextureUniformParameters::StaticStruct.GetShaderVariableName());
-		checkfSlow(!Initializer.ParameterMap.ContainsParameterAllocation(FSceneTexturesUniformParameters::StaticStruct.GetShaderVariableName()), TEXT("Shader for Mobile shading path tried to bind FSceneTexturesUniformParameters which is only available in the deferred shading path: %s"), Initializer.Type->GetName());
+		MobileSceneTexturesUniformBuffer.Bind(Initializer.ParameterMap, FMobileSceneTextureUniformParameters::StaticStructMetadata.GetShaderVariableName());
+		checkfSlow(!Initializer.ParameterMap.ContainsParameterAllocation(FSceneTexturesUniformParameters::StaticStructMetadata.GetShaderVariableName()), TEXT("Shader for Mobile shading path tried to bind FSceneTexturesUniformParameters which is only available in the deferred shading path: %s"), Initializer.Type->GetName());
 	}
 }
diff --git a/Engine/Source/Runtime/Renderer/Private/PostProcess/VisualizeTexture.cpp b/Engine/Source/Runtime/Renderer/Private/PostProcess/VisualizeTexture.cpp
--- a/Engine/Source/Runtime/Renderer/Private/PostProcess/VisualizeTexture.cpp
+++ b/Engine/Source/Runtime/Renderer/Private/PostProcess/VisualizeTexture.cpp
@@ -15,8 +15,9 @@
 #include "Misc/Paths.h"
 #include "Misc/App.h"
 #include "CanvasTypes.h"
+#include "Engine/Canvas.h"
 #include "UnrealEngine.h"
-#include "PostProcess/RenderTargetPool.h"
+#include "RenderTargetPool.h"
 #include "PostProcess/SceneRenderTargets.h"
 #include "GlobalShader.h"
 #include "RenderTargetTemp.h"
@@ -38,6 +39,8 @@ enum class EVisualisePSType
 	Texture2DUINT8 = 7,
 };
 
+TGlobalResource<FVisualizeTexture> GVisualizeTexture;
+
 /** A pixel shader which filters a texture. */
 // @param TextureType 0:Cube, 1:1D(not yet supported), 2:2D no MSAA, 3:3D, 4:Cube[], 5:2D MSAA, 6:2D DepthStencil no MSAA (needed to avoid D3DDebug error)
 template<EVisualisePSType TextureType>
@@ -333,7 +336,7 @@ FVisualizeTexture::FVisualizeTexture()
 	bEnabled = true;
 }
 
-void FVisualizeTexture::Destroy()
+void FVisualizeTexture::ReleaseDynamicRHI()
 {
 	VisualizeTextureContent.SafeRelease();
 	StencilSRV.SafeRelease();
@@ -544,8 +547,295 @@ void FVisualizeTexture::GenerateContent(FRHICommandListImmediate& RHICmdList, co
 	}
 }
 
+// draw a single pixel sized rectangle using 4 sub elements
+static void DrawBorder(FCanvas& Canvas, const FIntRect Rect, FLinearColor Color)
+{
+	// top
+	Canvas.DrawTile(Rect.Min.X, Rect.Min.Y, Rect.Max.X - Rect.Min.X, 1, 0, 0, 1, 1, Color);
+	// bottom
+	Canvas.DrawTile(Rect.Min.X, Rect.Max.Y - 1, Rect.Max.X - Rect.Min.X, 1, 0, 0, 1, 1, Color);
+	// left
+	Canvas.DrawTile(Rect.Min.X, Rect.Min.Y + 1, 1, Rect.Max.Y - Rect.Min.Y - 2, 0, 0, 1, 1, Color);
+	// right
+	Canvas.DrawTile(Rect.Max.X - 1, Rect.Min.Y + 1, 1, Rect.Max.Y - Rect.Min.Y - 2, 0, 0, 1, 1, Color);
+}
+
+// helper class to get a consistent layout in multiple functions
+// MaxX and Y are the output value that can be requested during or after iteration
+// Examples usages:
+//    FRenderTargetPoolEventIterator It(RenderTargetPoolEvents, OptionalStartIndex);
+//    while(FRenderTargetPoolEvent* Event = It.Iterate()) {}
+struct FRenderTargetPoolEventIterator
+{
+	int32 Index;
+	TArray<FRenderTargetPoolEvent>& RenderTargetPoolEvents;
+	bool bLineContent;
+	uint32 TotalWidth;
+	int32 Y;
+
+	// constructor
+	FRenderTargetPoolEventIterator(TArray<FRenderTargetPoolEvent>& InRenderTargetPoolEvents, int32 InIndex = 0)
+		: Index(InIndex)
+		, RenderTargetPoolEvents(InRenderTargetPoolEvents)
+		, bLineContent(false)
+		, TotalWidth(1)
+		, Y(0)
+	{
+		Touch();
+	}
+
+	FRenderTargetPoolEvent* operator*()
+	{
+		if (Index < RenderTargetPoolEvents.Num())
+		{
+			return &RenderTargetPoolEvents[Index];
+		}
+
+		return 0;
+	}
+
+	// @return 0 if end was reached
+	FRenderTargetPoolEventIterator& operator++()
+	{
+		if (Index < RenderTargetPoolEvents.Num())
+		{
+			++Index;
+		}
+
+		Touch();
+
+		return *this;
+	}
+
+	int32 FindClosingEventY() const
+	{
+		FRenderTargetPoolEventIterator It = *this;
+
+		const ERenderTargetPoolEventType StartType = (*It)->GetEventType();
+
+		if (StartType == ERTPE_Alloc)
+		{
+			int32 PoolEntryId = RenderTargetPoolEvents[Index].GetPoolEntryId();
+
+			++It;
+
+			// search for next Dealloc of the same PoolEntryId
+			for (; *It; ++It)
+			{
+				FRenderTargetPoolEvent* Event = *It;
+
+				if (Event->GetEventType() == ERTPE_Dealloc && Event->GetPoolEntryId() == PoolEntryId)
+				{
+					break;
+				}
+			}
+		}
+		else if (StartType == ERTPE_Phase)
+		{
+			++It;
+
+			// search for next Phase
+			for (; *It; ++It)
+			{
+				FRenderTargetPoolEvent* Event = *It;
+
+				if (Event->GetEventType() == ERTPE_Phase)
+				{
+					break;
+				}
+			}
+		}
+		else
+		{
+			check(0);
+		}
+
+		return It.Y;
+	}
+
+private:
+
+	void Touch()
+	{
+		if (Index < RenderTargetPoolEvents.Num())
+		{
+			const FRenderTargetPoolEvent& Event = RenderTargetPoolEvents[Index];
+
+			const ERenderTargetPoolEventType Type = Event.GetEventType();
+
+			if (Type == ERTPE_Alloc)
+			{
+				// for now they are all equal width
+				TotalWidth = FMath::Max(TotalWidth, Event.GetColumnX() + Event.GetColumnSize());
+			}
+			Y = Event.GetTimeStep();
+		}
+	}
+};
+uint32 FVisualizeTexture::ComputeEventDisplayHeight()
+{
+	FRenderTargetPoolEventIterator It(GRenderTargetPool.RenderTargetPoolEvents);
+
+	for (; *It; ++It)
+	{
+	}
+
+	return It.Y;
+}
+
 void FVisualizeTexture::PresentContent(FRHICommandListImmediate& RHICmdList, const FViewInfo& View)
 {
+	if (GRenderTargetPool.RenderTargetPoolEvents.Num())
+	{
+		GRenderTargetPool.AddPhaseEvent(TEXT("FrameEnd"));
+
+		FIntPoint DisplayLeftTop(20, 50);
+		// on the right we leave more space to make the mouse tooltip readable
+		FIntPoint DisplayExtent(View.ViewRect.Width() - DisplayLeftTop.X * 2 - 140, View.ViewRect.Height() - DisplayLeftTop.Y * 2);
+
+		// if the area is not too small
+		if (DisplayExtent.X > 50 && DisplayExtent.Y > 50)
+		{
+			FRenderTargetPool::SMemoryStats MemoryStats = GRenderTargetPool.ComputeView();
+
+			SetRenderTarget(RHICmdList, View.Family->RenderTarget->GetRenderTargetTexture(), FTextureRHIRef());
+			RHICmdList.SetViewport(0, 0, 0.0f, FSceneRenderTargets::Get(RHICmdList).GetBufferSizeXY().X, FSceneRenderTargets::Get(RHICmdList).GetBufferSizeXY().Y, 1.0f);
+
+
+			FRenderTargetTemp TempRenderTarget(View, View.UnscaledViewRect.Size());
+			FCanvas Canvas(&TempRenderTarget, NULL, View.Family->CurrentRealTime, View.Family->CurrentWorldTime, View.Family->DeltaWorldTime, View.GetFeatureLevel());
+
+			// TinyFont property
+			const int32 FontHeight = 12;
+
+			FIntPoint MousePos = View.CursorPos;
+
+			FLinearColor BackgroundColor = FLinearColor(0.0f, 0.0f, 0.0f, 0.7f);
+			FLinearColor PhaseColor = FLinearColor(0.2f, 0.1f, 0.05f, 0.8f);
+			FLinearColor ElementColor = FLinearColor(0.3f, 0.3f, 0.3f, 0.9f);
+			FLinearColor ElementColorVRam = FLinearColor(0.4f, 0.25f, 0.25f, 0.9f);
+
+			UTexture2D* GradientTexture = UCanvas::StaticClass()->GetDefaultObject<UCanvas>()->GradientTexture0;
+
+			// background rectangle
+			Canvas.DrawTile(DisplayLeftTop.X, DisplayLeftTop.Y - 1 * FontHeight - 1, DisplayExtent.X, DisplayExtent.Y + FontHeight, 0, 0, 1, 1, BackgroundColor);
+
+			{
+				uint32 MB = 1024 * 1024;
+				uint32 MBm1 = MB - 1;
+
+				FString Headline = *FString::Printf(TEXT("RenderTargetPool elements(x) over time(y) >= %dKB, Displayed/Total:%d/%dMB"),
+					GRenderTargetPool.EventRecordingSizeThreshold,
+					(uint32)((MemoryStats.DisplayedUsageInBytes + MBm1) / MB),
+					(uint32)((MemoryStats.TotalUsageInBytes + MBm1) / MB));
+				Canvas.DrawShadowedString(DisplayLeftTop.X, DisplayLeftTop.Y - 1 * FontHeight - 1, *Headline, GEngine->GetTinyFont(), FLinearColor(1, 1, 1));
+			}
+
+			uint32 EventDisplayHeight = FVisualizeTexture::ComputeEventDisplayHeight();
+
+			float ScaleX = DisplayExtent.X / (float)MemoryStats.TotalColumnSize;
+			float ScaleY = DisplayExtent.Y / (float)EventDisplayHeight;
+
+			// 0 if none
+			FRenderTargetPoolEvent* HighlightedEvent = 0;
+			FIntRect HighlightedRect;
+
+			// Phase events
+			for (FRenderTargetPoolEventIterator It(GRenderTargetPool.RenderTargetPoolEvents); *It; ++It)
+			{
+				FRenderTargetPoolEvent* Event = *It;
+
+				if (Event->GetEventType() == ERTPE_Phase)
+				{
+					int32 Y0 = It.Y;
+					int32 Y1 = It.FindClosingEventY();
+
+					FIntPoint PixelLeftTop((int32)(DisplayLeftTop.X), (int32)(DisplayLeftTop.Y + ScaleY * Y0));
+					FIntPoint PixelRightBottom((int32)(DisplayLeftTop.X + DisplayExtent.X), (int32)(DisplayLeftTop.Y + ScaleY * Y1));
+
+					bool bHighlight = MousePos.X >= PixelLeftTop.X && MousePos.X < PixelRightBottom.X && MousePos.Y >= PixelLeftTop.Y && MousePos.Y <= PixelRightBottom.Y;
+
+					if (bHighlight)
+					{
+						HighlightedEvent = Event;
+						HighlightedRect = FIntRect(PixelLeftTop, PixelRightBottom);
+					}
+
+					// UMax is 0.9f to avoid getting some wrap texture leaking in at the bottom
+					Canvas.DrawTile(PixelLeftTop.X, PixelLeftTop.Y, PixelRightBottom.X - PixelLeftTop.X, PixelRightBottom.Y - PixelLeftTop.Y, 0, 0, 1, 0.9f, PhaseColor, GradientTexture->Resource);
+				}
+			}
+
+			// Alloc / Dealloc events
+			for (FRenderTargetPoolEventIterator It(GRenderTargetPool.RenderTargetPoolEvents); *It; ++It)
+			{
+				FRenderTargetPoolEvent* Event = *It;
+
+				if (Event->GetEventType() == ERTPE_Alloc && Event->GetColumnSize())
+				{
+					int32 Y0 = It.Y;
+					int32 Y1 = It.FindClosingEventY();
+
+					int32 X0 = Event->GetColumnX();
+					// for now they are all equal width
+					int32 X1 = X0 + Event->GetColumnSize();
+
+					FIntPoint PixelLeftTop((int32)(DisplayLeftTop.X + ScaleX * X0), (int32)(DisplayLeftTop.Y + ScaleY * Y0));
+					FIntPoint PixelRightBottom((int32)(DisplayLeftTop.X + ScaleX * X1), (int32)(DisplayLeftTop.Y + ScaleY * Y1));
+
+					bool bHighlight = MousePos.X >= PixelLeftTop.X && MousePos.X < PixelRightBottom.X && MousePos.Y >= PixelLeftTop.Y && MousePos.Y <= PixelRightBottom.Y;
+
+					if (bHighlight)
+					{
+						HighlightedEvent = Event;
+						HighlightedRect = FIntRect(PixelLeftTop, PixelRightBottom);
+					}
+
+					FLinearColor Color = ElementColor;
+
+					// Highlight EDRAM/FastVRAM usage
+					if (Event->GetDesc().Flags & TexCreate_FastVRAM)
+					{
+						Color = ElementColorVRam;
+					}
+
+					Canvas.DrawTile(
+						PixelLeftTop.X, PixelLeftTop.Y,
+						PixelRightBottom.X - PixelLeftTop.X - 1, PixelRightBottom.Y - PixelLeftTop.Y - 1,
+						0, 0, 1, 1, Color);
+				}
+			}
+
+			if (HighlightedEvent)
+			{
+				DrawBorder(Canvas, HighlightedRect, FLinearColor(0.8f, 0, 0, 0.5f));
+
+				// Offset to not intersect with crosshair (in editor) or arrow (in game).
+				FIntPoint Pos = MousePos + FIntPoint(12, 4);
+
+				if (HighlightedEvent->GetEventType() == ERTPE_Phase)
+				{
+					FString PhaseText = *FString::Printf(TEXT("Phase: %s"), *HighlightedEvent->GetPhaseName());
+
+					Canvas.DrawShadowedString(Pos.X, Pos.Y + 0 * FontHeight, *PhaseText, GEngine->GetTinyFont(), FLinearColor(0.5f, 0.5f, 1));
+				}
+				else
+				{
+					FString SizeString = FString::Printf(TEXT("%d KB"), (HighlightedEvent->GetSizeInBytes() + 1024) / 1024);
+
+					Canvas.DrawShadowedString(Pos.X, Pos.Y + 0 * FontHeight, HighlightedEvent->GetDesc().DebugName, GEngine->GetTinyFont(), FLinearColor(1, 1, 0));
+					Canvas.DrawShadowedString(Pos.X, Pos.Y + 1 * FontHeight, *HighlightedEvent->GetDesc().GenerateInfoString(), GEngine->GetTinyFont(), FLinearColor(1, 1, 0));
+					Canvas.DrawShadowedString(Pos.X, Pos.Y + 2 * FontHeight, *SizeString, GEngine->GetTinyFont(), FLinearColor(1, 1, 0));
+				}
+			}
+
+			Canvas.Flush_RenderThread(RHICmdList);
+
+			GRenderTargetPool.CurrentEventRecordingTime = 0;
+			GRenderTargetPool.RenderTargetPoolEvents.Empty();
+		}
+	}
+
 	if(Mode != 0)
 	{
 		// old mode is used, lets copy the specified texture to do it similar to the new system
@@ -913,7 +1203,7 @@ void FVisualizeTexture::DebugLog(bool bExtended)
 		TArray<FString> Entries;
 				
 		// sorted by pointer for efficiency, now we want to print sorted alphabetically
-		for (TMap<const TCHAR*, uint32>:: TIterator It(GRenderTargetPool.VisualizeTexture.VisualizeTextureCheckpoints); It; ++It)
+		for (TMap<const TCHAR*, uint32>:: TIterator It(GVisualizeTexture.VisualizeTextureCheckpoints); It; ++It)
 		{
 			const TCHAR* Key = It.Key();
 			uint32 Value = It.Value();
diff --git a/Engine/Source/Runtime/Renderer/Private/PostProcess/VisualizeTexture.h b/Engine/Source/Runtime/Renderer/Private/PostProcess/VisualizeTexture.h
--- a/Engine/Source/Runtime/Renderer/Private/PostProcess/VisualizeTexture.h
+++ b/Engine/Source/Runtime/Renderer/Private/PostProcess/VisualizeTexture.h
@@ -8,15 +8,18 @@
 
 #include "CoreMinimal.h"
 #include "RendererInterface.h"
+#include "RenderResource.h"
 
 class FSceneView;
 class FViewInfo;
 
-class FVisualizeTexture
+
+// TODO(RDG): Texture visualization needs to be totally rewritten cleanly with render graph in RenderCore module.
+class FVisualizeTexture : public FRenderResource
 {
 public:
 	FVisualizeTexture();
-	void Destroy();
+	virtual void ReleaseDynamicRHI() override;
 
 	/** renders the VisualizeTextureContent to the current render target */
 	void PresentContent(FRHICommandListImmediate& RHICmdList, const FViewInfo& View);
@@ -109,6 +112,8 @@ private:
 	void GenerateContent(FRHICommandListImmediate& RHICmdList, const FSceneRenderTargetItem& RenderTargetItem, const FPooledRenderTargetDesc& Desc);
 
 	FIntRect ComputeVisualizeTextureRect(FIntPoint InputTextureSize) const;
+
+	static uint32 ComputeEventDisplayHeight();
 };
 
 
@@ -133,3 +138,6 @@ struct FVisualizeTextureData
 	int32 ArrayIndex;
 	int32 CustomMip;
 };
+
+/** The global render targets for easy shading. */
+extern TGlobalResource<FVisualizeTexture> GVisualizeTexture;
diff --git a/Engine/Source/Runtime/Renderer/Private/ReflectionEnvironment.cpp b/Engine/Source/Runtime/Renderer/Private/ReflectionEnvironment.cpp
--- a/Engine/Source/Runtime/Renderer/Private/ReflectionEnvironment.cpp
+++ b/Engine/Source/Runtime/Renderer/Private/ReflectionEnvironment.cpp
@@ -148,7 +148,7 @@ bool IsReflectionCaptureAvailable()
 	return (!AllowStaticLightingVar || AllowStaticLightingVar->GetInt() != 0);
 }
 
-IMPLEMENT_UNIFORM_BUFFER_STRUCT(FReflectionUniformParameters, TEXT("ReflectionStruct"));
+IMPLEMENT_GLOBAL_SHADER_PARAMETER_STRUCT(FReflectionUniformParameters, "ReflectionStruct");
 
 void SetupReflectionUniformParameters(const FViewInfo& View, FReflectionUniformParameters& OutParameters)
 {
@@ -477,7 +477,7 @@ struct FReflectionCaptureSortData
 	}
 };
 
-IMPLEMENT_UNIFORM_BUFFER_STRUCT(FReflectionCaptureShaderData,TEXT("ReflectionCapture"));
+IMPLEMENT_GLOBAL_SHADER_PARAMETER_STRUCT(FReflectionCaptureShaderData, "ReflectionCapture");
 
 /** Pixel shader that does tiled deferred culling of reflection captures, then sorts and composites them. */
 class FReflectionEnvironmentSkyLightingPS : public FGlobalShader
diff --git a/Engine/Source/Runtime/Renderer/Private/ReflectionEnvironment.h b/Engine/Source/Runtime/Renderer/Private/ReflectionEnvironment.h
--- a/Engine/Source/Runtime/Renderer/Private/ReflectionEnvironment.h
+++ b/Engine/Source/Runtime/Renderer/Private/ReflectionEnvironment.h
@@ -14,15 +14,15 @@
 extern bool IsReflectionEnvironmentAvailable(ERHIFeatureLevel::Type InFeatureLevel);
 extern bool IsReflectionCaptureAvailable();
 
-BEGIN_UNIFORM_BUFFER_STRUCT(FReflectionUniformParameters,)
-	UNIFORM_MEMBER(FVector4, SkyLightParameters)
-	UNIFORM_MEMBER(float, SkyLightCubemapBrightness)
-	UNIFORM_MEMBER_TEXTURE(TextureCube, SkyLightCubemap)
-	UNIFORM_MEMBER_SAMPLER(SamplerState, SkyLightCubemapSampler)
-	UNIFORM_MEMBER_TEXTURE(TextureCube, SkyLightBlendDestinationCubemap)
-	UNIFORM_MEMBER_SAMPLER(SamplerState, SkyLightBlendDestinationCubemapSampler)
-	UNIFORM_MEMBER_TEXTURE(TextureCubeArray, ReflectionCubemap)
-	UNIFORM_MEMBER_SAMPLER(SamplerState, ReflectionCubemapSampler)
-END_UNIFORM_BUFFER_STRUCT(FReflectionUniformParameters)
+BEGIN_GLOBAL_SHADER_PARAMETER_STRUCT(FReflectionUniformParameters,)
+	SHADER_PARAMETER(FVector4, SkyLightParameters)
+	SHADER_PARAMETER(float, SkyLightCubemapBrightness)
+	SHADER_PARAMETER_TEXTURE(TextureCube, SkyLightCubemap)
+	SHADER_PARAMETER_SAMPLER(SamplerState, SkyLightCubemapSampler)
+	SHADER_PARAMETER_TEXTURE(TextureCube, SkyLightBlendDestinationCubemap)
+	SHADER_PARAMETER_SAMPLER(SamplerState, SkyLightBlendDestinationCubemapSampler)
+	SHADER_PARAMETER_TEXTURE(TextureCubeArray, ReflectionCubemap)
+	SHADER_PARAMETER_SAMPLER(SamplerState, ReflectionCubemapSampler)
+END_GLOBAL_SHADER_PARAMETER_STRUCT()
 
 extern void SetupReflectionUniformParameters(const class FViewInfo& View, FReflectionUniformParameters& OutParameters);
\ No newline at end of file
diff --git a/Engine/Source/Runtime/Renderer/Private/Renderer.cpp b/Engine/Source/Runtime/Renderer/Private/Renderer.cpp
--- a/Engine/Source/Runtime/Renderer/Private/Renderer.cpp
+++ b/Engine/Source/Runtime/Renderer/Private/Renderer.cpp
@@ -14,8 +14,9 @@
 #include "RenderingThread.h"
 #include "RHIStaticStates.h"
 #include "SceneView.h"
-#include "PostProcess/RenderTargetPool.h"
+#include "RenderTargetPool.h"
 #include "PostProcess/SceneRenderTargets.h"
+#include "PostProcess/VisualizeTexture.h"
 #include "SceneCore.h"
 #include "SceneHitProxyRendering.h"
 #include "SceneRendering.h"
@@ -110,7 +111,7 @@ void FRendererModule::DrawTileMesh(FRHICommandListImmediate& RHICmdList, FDrawin
 				// Crash fix - reflection capture shader parameter is bound but we have no buffer during Build Texture Streaming
 				if(!View.ReflectionCaptureUniformBuffer.IsValid())
 				{
-					View.ReflectionCaptureUniformBuffer = GDummyReflectionCaptureUniformBuffer;
+					View.ReflectionCaptureUniformBuffer = GDummyReflectionCaptureUniformBuffer.GetUniformBufferRef();
 				}
 			
 				TUniformBufferRef<FTranslucentBasePassUniformParameters> BasePassUniformBuffer;
@@ -170,9 +171,9 @@ void FRendererModule::TickRenderTargetPool()
 
 void FRendererModule::DebugLogOnCrash()
 {
-	GRenderTargetPool.VisualizeTexture.SortOrder = 1;
-	GRenderTargetPool.VisualizeTexture.bFullList = true;
-	GRenderTargetPool.VisualizeTexture.DebugLog(false);
+	GVisualizeTexture.SortOrder = 1;
+	GVisualizeTexture.bFullList = true;
+	GVisualizeTexture.DebugLog(false);
 	
 	GEngine->Exec(NULL, TEXT("rhi.DumpMemory"), *GLog);
 
@@ -243,7 +244,7 @@ void FRendererModule::QueryVisualizeTexture(FQueryVisualizeTexureInfo& Out)
 	check(IsInGameThread());
 	FlushRenderingCommands();
 
-	GRenderTargetPool.VisualizeTexture.QueryInfo(Out);
+	GVisualizeTexture.QueryInfo(Out);
 }
 
 static void VisualizeTextureExec( const TCHAR* Cmd, FOutputDevice &Ar )
@@ -267,44 +268,44 @@ static void VisualizeTextureExec( const TCHAR* Cmd, FOutputDevice &Ar )
 		// FULL flag
 		if(Parameter == TEXT("fulllist") || Parameter == TEXT("full"))
 		{
-			GRenderTargetPool.VisualizeTexture.bFullList = true;
+			GVisualizeTexture.bFullList = true;
 			// this one doesn't count as parameter so we can do "vis full"
 			continue;
 		}
 		// SORT0 flag
 		else if(Parameter == TEXT("sort0"))
 		{
-			GRenderTargetPool.VisualizeTexture.SortOrder = 0;
+			GVisualizeTexture.SortOrder = 0;
 			// this one doesn't count as parameter so we can do "vis full"
 			continue;
 		}
 		// SORT1 flag
 		else if(Parameter == TEXT("sort1"))
 		{
-			GRenderTargetPool.VisualizeTexture.SortOrder = 1;
+			GVisualizeTexture.SortOrder = 1;
 			// this one doesn't count as parameter so we can do "vis full"
 			continue;
 		}
 		else if(ParameterCount == 0)
 		{
 			// Init
-			GRenderTargetPool.VisualizeTexture.RGBMul = 1;
-			GRenderTargetPool.VisualizeTexture.SingleChannelMul = 0.0f;
-			GRenderTargetPool.VisualizeTexture.SingleChannel = -1;
-			GRenderTargetPool.VisualizeTexture.AMul = 0;
-			GRenderTargetPool.VisualizeTexture.UVInputMapping = 3;
-			GRenderTargetPool.VisualizeTexture.Flags = 0;
-			GRenderTargetPool.VisualizeTexture.Mode = 0;
-			GRenderTargetPool.VisualizeTexture.CustomMip = 0;
-			GRenderTargetPool.VisualizeTexture.ArrayIndex = 0;
-			GRenderTargetPool.VisualizeTexture.bOutputStencil = false;
+			GVisualizeTexture.RGBMul = 1;
+			GVisualizeTexture.SingleChannelMul = 0.0f;
+			GVisualizeTexture.SingleChannel = -1;
+			GVisualizeTexture.AMul = 0;
+			GVisualizeTexture.UVInputMapping = 3;
+			GVisualizeTexture.Flags = 0;
+			GVisualizeTexture.Mode = 0;
+			GVisualizeTexture.CustomMip = 0;
+			GVisualizeTexture.ArrayIndex = 0;
+			GVisualizeTexture.bOutputStencil = false;
 
 			// e.g. "VisualizeTexture Name" or "VisualizeTexture 5"
 			bool bIsDigit = FChar::IsDigit(**Parameter);
 
 			if (bIsDigit)
 			{
-				GRenderTargetPool.VisualizeTexture.Mode = FCString::Atoi(*Parameter);
+				GVisualizeTexture.Mode = FCString::Atoi(*Parameter);
 			}
 
 			if(!bIsDigit)
@@ -321,45 +322,45 @@ static void VisualizeTextureExec( const TCHAR* Cmd, FOutputDevice &Ar )
 				{
 					// user specified a reuse goal
 					FString NameWithoutAt = Parameter.Left(AfterAt - *Parameter);
-					GRenderTargetPool.VisualizeTexture.SetObserveTarget(*NameWithoutAt, FCString::Atoi(AfterAt + 1));
+					GVisualizeTexture.SetObserveTarget(*NameWithoutAt, FCString::Atoi(AfterAt + 1));
 				}
 				else
 				{
 					// we take the last one
-					GRenderTargetPool.VisualizeTexture.SetObserveTarget(*Parameter);
+					GVisualizeTexture.SetObserveTarget(*Parameter);
 				}
 			}
 			else
 			{
 				// the index was used
-				GRenderTargetPool.VisualizeTexture.SetObserveTarget(TEXT(""));
+				GVisualizeTexture.SetObserveTarget(TEXT(""));
 			}
 		}
 		// GRenderTargetPoolInputMapping mode
 		else if(Parameter == TEXT("uv0"))
 		{
-			GRenderTargetPool.VisualizeTexture.UVInputMapping = 0;
+			GVisualizeTexture.UVInputMapping = 0;
 		}
 		else if(Parameter == TEXT("uv1"))
 		{
-			GRenderTargetPool.VisualizeTexture.UVInputMapping = 1;
+			GVisualizeTexture.UVInputMapping = 1;
 		}
 		else if(Parameter == TEXT("uv2"))
 		{
-			GRenderTargetPool.VisualizeTexture.UVInputMapping = 2;
+			GVisualizeTexture.UVInputMapping = 2;
 		}
 		else if(Parameter == TEXT("pip"))
 		{
-			GRenderTargetPool.VisualizeTexture.UVInputMapping = 3;
+			GVisualizeTexture.UVInputMapping = 3;
 		}
 		// BMP flag
 		else if(Parameter == TEXT("bmp"))
 		{
-			GRenderTargetPool.VisualizeTexture.bSaveBitmap = true;
+			GVisualizeTexture.bSaveBitmap = true;
 		}
 		else if (Parameter == TEXT("stencil"))
 		{
-			GRenderTargetPool.VisualizeTexture.bOutputStencil = true;
+			GVisualizeTexture.bOutputStencil = true;
 		}
 		// saturate flag
 		else if(Parameter == TEXT("frac"))
@@ -369,19 +370,19 @@ static void VisualizeTextureExec( const TCHAR* Cmd, FOutputDevice &Ar )
 		// saturate flag
 		else if(Parameter == TEXT("sat"))
 		{
-			GRenderTargetPool.VisualizeTexture.Flags |= 0x1;
+			GVisualizeTexture.Flags |= 0x1;
 		}
 		// e.g. mip2 or mip0
 		else if(Parameter.Left(3) == TEXT("mip"))
 		{
 			Parameter = Parameter.Right(Parameter.Len() - 3);
-			GRenderTargetPool.VisualizeTexture.CustomMip = FCString::Atoi(*Parameter);
+			GVisualizeTexture.CustomMip = FCString::Atoi(*Parameter);
 		}
 		// e.g. [0] or [2]
 		else if(Parameter.Left(5) == TEXT("index"))
 		{
 			Parameter = Parameter.Right(Parameter.Len() - 5);
-			GRenderTargetPool.VisualizeTexture.ArrayIndex = FCString::Atoi(*Parameter);
+			GVisualizeTexture.ArrayIndex = FCString::Atoi(*Parameter);
 		}
 		// e.g. RGB*6, A, *22, /2.7, A*7
 		else if(Parameter.Left(3) == TEXT("rgb")
@@ -405,9 +406,9 @@ static void VisualizeTextureExec( const TCHAR* Cmd, FOutputDevice &Ar )
 			if ( SingleChannel >= 0 )
 			{
 				Parameter = Parameter.Right(Parameter.Len() - 1);
-				GRenderTargetPool.VisualizeTexture.SingleChannel = SingleChannel;
-				GRenderTargetPool.VisualizeTexture.SingleChannelMul = 1;
-				GRenderTargetPool.VisualizeTexture.RGBMul = 0;
+				GVisualizeTexture.SingleChannel = SingleChannel;
+				GVisualizeTexture.SingleChannelMul = 1;
+				GVisualizeTexture.RGBMul = 0;
 			}
 
 			float Mul = 1.0f;
@@ -423,9 +424,9 @@ static void VisualizeTextureExec( const TCHAR* Cmd, FOutputDevice &Ar )
 				Parameter = Parameter.Right(Parameter.Len() - 1);
 				Mul = 1.0f / FCString::Atof(*Parameter);
 			}
-			GRenderTargetPool.VisualizeTexture.RGBMul *= Mul;
-			GRenderTargetPool.VisualizeTexture.SingleChannelMul *= Mul;
-			GRenderTargetPool.VisualizeTexture.AMul *= Mul;
+			GVisualizeTexture.RGBMul *= Mul;
+			GVisualizeTexture.SingleChannelMul *= Mul;
+			GVisualizeTexture.AMul *= Mul;
 		}
 		else
 		{
@@ -468,9 +469,9 @@ static void VisualizeTextureExec( const TCHAR* Cmd, FOutputDevice &Ar )
 		Ar.Logf(TEXT("TextureId:"));
 		Ar.Logf(TEXT("  0        = <off>"));
 
-		GRenderTargetPool.VisualizeTexture.DebugLog(true);
+		GVisualizeTexture.DebugLog(true);
 	}
-	//		Ar.Logf(TEXT("VisualizeTexture %d"), GRenderTargetPool.VisualizeTexture.Mode);
+	//		Ar.Logf(TEXT("VisualizeTexture %d"), GVisualizeTexture.Mode);
 }
 
 static bool RendererExec( UWorld* InWorld, const TCHAR* Cmd, FOutputDevice& Ar )
diff --git a/Engine/Source/Runtime/Renderer/Private/RendererScene.cpp b/Engine/Source/Runtime/Renderer/Private/RendererScene.cpp
--- a/Engine/Source/Runtime/Renderer/Private/RendererScene.cpp
+++ b/Engine/Source/Runtime/Renderer/Private/RendererScene.cpp
@@ -62,7 +62,7 @@
 
 DECLARE_CYCLE_STAT(TEXT("DeferredShadingSceneRenderer MotionBlurStartFrame"), STAT_FDeferredShadingSceneRenderer_MotionBlurStartFrame, STATGROUP_SceneRendering);
 
-IMPLEMENT_UNIFORM_BUFFER_STRUCT(FDistanceCullFadeUniformShaderParameters,TEXT("PrimitiveFade"));
+IMPLEMENT_GLOBAL_SHADER_PARAMETER_STRUCT(FDistanceCullFadeUniformShaderParameters, "PrimitiveFade");
 
 /** Global primitive uniform buffer resource containing faded in */
 TGlobalResource< FGlobalDistanceCullFadeUniformBuffer > GDistanceCullFadedInUniformBuffer;
diff --git a/Engine/Source/Runtime/Renderer/Private/SceneHitProxyRendering.cpp b/Engine/Source/Runtime/Renderer/Private/SceneHitProxyRendering.cpp
--- a/Engine/Source/Runtime/Renderer/Private/SceneHitProxyRendering.cpp
+++ b/Engine/Source/Runtime/Renderer/Private/SceneHitProxyRendering.cpp
@@ -56,7 +56,7 @@ protected:
 	FHitProxyVS(const ShaderMetaType::CompiledShaderInitializerType& Initializer) :
 		FMeshMaterialShader(Initializer)
 	{
-		PassUniformBuffer.Bind(Initializer.ParameterMap, FSceneTexturesUniformParameters::StaticStruct.GetShaderVariableName());
+		PassUniformBuffer.Bind(Initializer.ParameterMap, FSceneTexturesUniformParameters::StaticStructMetadata.GetShaderVariableName());
 	}
 	FHitProxyVS() {}
 };
@@ -129,7 +129,7 @@ public:
 		FMeshMaterialShader(Initializer)
 	{
 		HitProxyId.Bind(Initializer.ParameterMap,TEXT("HitProxyId"), SPF_Mandatory);
-		PassUniformBuffer.Bind(Initializer.ParameterMap, FSceneTexturesUniformParameters::StaticStruct.GetShaderVariableName());
+		PassUniformBuffer.Bind(Initializer.ParameterMap, FSceneTexturesUniformParameters::StaticStructMetadata.GetShaderVariableName());
 	}
 
 	FHitProxyPS() {}
@@ -561,7 +561,7 @@ static void DoRenderHitProxies(FRHICommandListImmediate& RHICmdList, const FScen
 	RHICmdList.CopyToResolveTarget(SceneContext.GetSceneDepthSurface(), SceneContext.GetSceneDepthSurface(), FResolveParams());
 
 	// to be able to observe results with VisualizeTexture
-	GRenderTargetPool.VisualizeTexture.SetCheckPoint(RHICmdList, HitProxyRT);
+	GVisualizeTexture.SetCheckPoint(RHICmdList, HitProxyRT);
 
 	//
 	// Copy the hit proxy buffer into the view family's render target.
diff --git a/Engine/Source/Runtime/Renderer/Private/SceneOcclusion.cpp b/Engine/Source/Runtime/Renderer/Private/SceneOcclusion.cpp
--- a/Engine/Source/Runtime/Renderer/Private/SceneOcclusion.cpp
+++ b/Engine/Source/Runtime/Renderer/Private/SceneOcclusion.cpp
@@ -15,6 +15,8 @@
 #include "PostProcess/PostProcessing.h"
 #include "PlanarReflectionSceneProxy.h"
 #include "PipelineStateCache.h"
+#include "ShaderParameterStruct.h"
+
 
 /*-----------------------------------------------------------------------------
 	Globals
@@ -1043,16 +1045,35 @@ void FHZBOcclusionTester::Submit(FRHICommandListImmediate& RHICmdList, const FVi
 			EDRF_UseTriangleOptimization);
 	}
 
-	GRenderTargetPool.VisualizeTexture.SetCheckPoint(RHICmdList, ResultsTextureGPU);
+	GVisualizeTexture.SetCheckPoint(RHICmdList, ResultsTextureGPU);
 
 	// Transfer memory GPU -> CPU
 	RHICmdList.CopyToResolveTarget(ResultsTextureGPU->GetRenderTargetItem().TargetableTexture, ResultsTextureCPU->GetRenderTargetItem().ShaderResourceTexture, FResolveParams());
 }
 
-template< uint32 Stage >
-class THZBBuildPS : public FGlobalShader
+BEGIN_SHADER_PARAMETER_STRUCT(FHZBBuildPassParameters, )
+	RENDER_TARGET_BINDING_SLOTS()
+	SHADER_PARAMETER_GRAPH_SRV(Texture2D, Texture)
+	SHADER_PARAMETER_SAMPLER(SamplerState, TextureSampler)
+END_SHADER_PARAMETER_STRUCT()
+
+class FHZBBuildPS : public FGlobalShader
 {
-	DECLARE_SHADER_TYPE(THZBBuildPS, Global);
+	DECLARE_GLOBAL_SHADER(FHZBBuildPS);
+	SHADER_USE_PARAMETER_STRUCT(FHZBBuildPS, FGlobalShader)
+
+	class FStageDim : SHADER_PERMUTATION_BOOL("STAGE");
+	using FPermutationDomain = TShaderPermutationDomain<FStageDim>;
+	
+	BEGIN_SHADER_PARAMETER_STRUCT( FParameters, )
+		SHADER_PARAMETER( FVector2D,	InvSize )
+		SHADER_PARAMETER( FVector4,		InputUvFactorAndOffset )
+		SHADER_PARAMETER( FVector2D,	InputViewportMaxBound )
+
+		SHADER_PARAMETER_STRUCT_INCLUDE(FHZBBuildPassParameters, Pass)
+		SHADER_PARAMETER_STRUCT_REF(	FViewUniformShaderParameters,	View )
+		SHADER_PARAMETER_STRUCT_REF(	FSceneTexturesUniformParameters,SceneTextures )
+	END_SHADER_PARAMETER_STRUCT()
 
 	static bool ShouldCompilePermutation(const FGlobalShaderPermutationParameters& Parameters)
 	{
@@ -1061,93 +1082,19 @@ class THZBBuildPS : public FGlobalShader
 
 	static void ModifyCompilationEnvironment(const FGlobalShaderPermutationParameters& Parameters, FShaderCompilerEnvironment& OutEnvironment)
 	{
-		FGlobalShader::ModifyCompilationEnvironment(Parameters, OutEnvironment);
-		OutEnvironment.SetDefine( TEXT("STAGE"), Stage );
 		OutEnvironment.SetRenderTargetOutputFormat(0, PF_R32_FLOAT);
 	}
-
-	THZBBuildPS() {}
-
-public:
-	FShaderParameter				InvSizeParameter;
-	FShaderParameter				InputUvFactorAndOffsetParameter;
-	FShaderParameter				InputViewportMaxBoundParameter;
-	FSceneTextureShaderParameters	SceneTextureParameters;
-	FShaderResourceParameter		TextureParameter;
-	FShaderResourceParameter		TextureParameterSampler;
-
-	THZBBuildPS(const ShaderMetaType::CompiledShaderInitializerType& Initializer)
-		: FGlobalShader(Initializer)
-	{
-		InvSizeParameter.Bind( Initializer.ParameterMap, TEXT("InvSize") );
-		InputUvFactorAndOffsetParameter.Bind( Initializer.ParameterMap, TEXT("InputUvFactorAndOffset") );
-		InputViewportMaxBoundParameter.Bind( Initializer.ParameterMap, TEXT("InputViewportMaxBound") );
-		SceneTextureParameters.Bind( Initializer );
-		TextureParameter.Bind( Initializer.ParameterMap, TEXT("Texture") );
-		TextureParameterSampler.Bind( Initializer.ParameterMap, TEXT("TextureSampler") );
-	}
-
-	void SetParameters(FRHICommandList& RHICmdList, const FViewInfo& View)
-	{
-		const FPixelShaderRHIParamRef ShaderRHI = GetPixelShader();
-
-		FGlobalShader::SetParameters<FViewUniformShaderParameters>(RHICmdList, ShaderRHI, View.ViewUniformBuffer);
-		FSceneRenderTargets& SceneContext = FSceneRenderTargets::Get(RHICmdList);
-		
-		const FIntPoint GBufferSize = SceneContext.GetBufferSizeXY();
-		const FVector2D InvSize( 1.0f / float(GBufferSize.X), 1.0f / float(GBufferSize.Y) );
-		const FVector4 InputUvFactorAndOffset (
-			float(2 * View.HZBMipmap0Size.X) / float(GBufferSize.X),
-			float(2 * View.HZBMipmap0Size.Y) / float(GBufferSize.Y),
-			float(View.ViewRect.Min.X) / float(GBufferSize.X),
-			float(View.ViewRect.Min.Y) / float(GBufferSize.Y)
-			);
-		const FVector2D InputViewportMaxBound (
-			float(View.ViewRect.Max.X) / float(GBufferSize.X) - 0.5f * InvSize.X,
-			float(View.ViewRect.Max.Y) / float(GBufferSize.Y) - 0.5f * InvSize.Y
-			);
-		SetShaderValue(RHICmdList, ShaderRHI, InvSizeParameter, InvSize );
-		SetShaderValue(RHICmdList, ShaderRHI, InputUvFactorAndOffsetParameter, InputUvFactorAndOffset );
-		SetShaderValue(RHICmdList, ShaderRHI, InputViewportMaxBoundParameter, InputViewportMaxBound );
-		
-		SceneTextureParameters.Set(RHICmdList, ShaderRHI, View.FeatureLevel, ESceneTextureSetupMode::SceneDepth);
-	}
-
-	void SetParameters(FRHICommandList& RHICmdList, const FSceneView& View, const FIntPoint& Size, FShaderResourceViewRHIParamRef ShaderResourceView )
-	{
-		const FPixelShaderRHIParamRef ShaderRHI = GetPixelShader();
-
-		FGlobalShader::SetParameters<FViewUniformShaderParameters>(RHICmdList, ShaderRHI, View.ViewUniformBuffer);
-
-		const FVector2D InvSize( 1.0f / Size.X, 1.0f / Size.Y );
-		SetShaderValue(RHICmdList, ShaderRHI, InvSizeParameter, InvSize );
-
-		//SetTextureParameter( ShaderRHI, TextureParameter, TextureParameterSampler, TStaticSamplerState<SF_Point,AM_Clamp,AM_Clamp,AM_Clamp>::GetRHI(), Texture );
-
-		SetSRVParameter(RHICmdList, ShaderRHI, TextureParameter, ShaderResourceView );
-		SetSamplerParameter(RHICmdList, ShaderRHI, TextureParameterSampler, TStaticSamplerState<SF_Point,AM_Clamp,AM_Clamp,AM_Clamp>::GetRHI() );
-	}
-
-	virtual bool Serialize(FArchive& Ar) override
-	{
-		bool bShaderHasOutdatedParameters = FGlobalShader::Serialize(Ar);
-		Ar << InvSizeParameter;
-		Ar << InputUvFactorAndOffsetParameter;
-		Ar << InputViewportMaxBoundParameter;
-		Ar << SceneTextureParameters;
-		Ar << TextureParameter;
-		Ar << TextureParameterSampler;
-		return bShaderHasOutdatedParameters;
-	}
 };
 
-IMPLEMENT_SHADER_TYPE(template<>,THZBBuildPS<0>,TEXT("/Engine/Private/HZBOcclusion.usf"),TEXT("HZBBuildPS"),SF_Pixel);
-IMPLEMENT_SHADER_TYPE(template<>,THZBBuildPS<1>,TEXT("/Engine/Private/HZBOcclusion.usf"),TEXT("HZBBuildPS"),SF_Pixel);
+IMPLEMENT_GLOBAL_SHADER(FHZBBuildPS, "/Engine/Private/HZBOcclusion.usf", "HZBBuildPS", SF_Pixel);
+
 
-void BuildHZB( FRHICommandListImmediate& RHICmdList, FViewInfo& View )
+void BuildHZB(FRHICommandListImmediate& RHICmdList, FViewInfo& View)
 {
 	QUICK_SCOPE_CYCLE_COUNTER(STAT_BuildHZB);
 	
+	FRenderGraphBuilder GraphBuilder(RHICmdList);
+
 	// View.ViewRect.{Width,Height}() are most likely to be < 2^24, so the float
 	// conversion won't loss any precision (assuming float have 23bits for mantissa)
 	const int32 NumMipsX = FMath::Max(FPlatformMath::CeilToInt(FMath::Log2(float(View.ViewRect.Width()))) - 1, 1);
@@ -1158,108 +1105,158 @@ void BuildHZB( FRHICommandListImmediate& RHICmdList, FViewInfo& View )
 	const FIntPoint HZBSize( 1 << NumMipsX, 1 << NumMipsY );
 	View.HZBMipmap0Size = HZBSize;
 
-	FPooledRenderTargetDesc Desc(FPooledRenderTargetDesc::Create2DDesc(HZBSize, PF_R16F, FClearValueBinding::None, TexCreate_None, TexCreate_RenderTargetable | TexCreate_ShaderResource | TexCreate_NoFastClear, false, NumMips));
-	Desc.Flags |= GFastVRamConfig.HZB;
-	GRenderTargetPool.FindFreeElement(RHICmdList, Desc, View.HZB, TEXT("HZB") );
-	
-	FSceneRenderTargetItem& HZBRenderTarget = View.HZB->GetRenderTargetItem();
-	
-	FGraphicsPipelineStateInitializer GraphicsPSOInit;
-	GraphicsPSOInit.BlendState = TStaticBlendState<>::GetRHI();
-	GraphicsPSOInit.RasterizerState = TStaticRasterizerState<>::GetRHI();
-	GraphicsPSOInit.DepthStencilState = TStaticDepthStencilState<false, CF_Always>::GetRHI();
+	//@DW: Configure texture creation
+	FPooledRenderTargetDesc HZBDesc = FPooledRenderTargetDesc::Create2DDesc(HZBSize, PF_R16F, FClearValueBinding::None, TexCreate_None, TexCreate_RenderTargetable | TexCreate_ShaderResource | TexCreate_NoFastClear, false, NumMips);
+	HZBDesc.Flags |= GFastVRamConfig.HZB;
 
-	FTextureRHIParamRef HZBRenderTargetRef = HZBRenderTarget.TargetableTexture.GetReference();
-	// Mip 0
-	{
-		SCOPED_DRAW_EVENTF(RHICmdList, BuildHZB, TEXT("HZB SetupMip 0 %dx%d"), HZBSize.X, HZBSize.Y);
+	//@DW: Explicit creation of graph resource handles - full support for everything the RHI supports
+	//@DW: Now that we've created a resource handle, it will have to be passed around to other passes for manual wiring or put into a Blackboard structure for automatic wiring
+	const FGraphTexture* HZBTexture = GraphBuilder.CreateTexture(HZBDesc, TEXT("HZB"));
 
-		FRHIRenderPassInfo RPInfo(HZBRenderTarget.TargetableTexture, ERenderTargetActions::Load_Store, 0);
-		RHICmdList.BeginRenderPass(RPInfo, TEXT("HZB_SetupMip0"));
+	{
+		FHZBBuildPassParameters* PassParameters;
+		GraphBuilder.CreateParameters(&PassParameters);
+		PassParameters->RenderTargets[0] = FRenderTargetBinding( HZBTexture, ERenderTargetLoadAction::ENoAction, ERenderTargetStoreAction::EStore);
+
+		//@DW - this pass only reads external textures, we don't have any graph inputs
+		GraphBuilder.AddPass(TEXT("HZB_SetupMip0"),
+			PassParameters,
+			ERenderGraphPassFlags::None,
+			[PassParameters, &View, HZBSize](FRHICommandListImmediate& RHICmdList)
+			{
+				FGraphicsPipelineStateInitializer GraphicsPSOInit;
+				GraphicsPSOInit.BlendState = TStaticBlendState<>::GetRHI();
+				GraphicsPSOInit.RasterizerState = TStaticRasterizerState<>::GetRHI();
+				GraphicsPSOInit.DepthStencilState = TStaticDepthStencilState<false, CF_Always>::GetRHI();
 
-		RHICmdList.ApplyCachedRenderTargets(GraphicsPSOInit);
-		GraphicsPSOInit.PrimitiveType = PT_TriangleList;
+				// Mip 0
+				SCOPED_DRAW_EVENTF(RHICmdList, BuildHZB, TEXT("HZB SetupMip 0 %dx%d"), HZBSize.X, HZBSize.Y);
 
-		TShaderMapRef< FPostProcessVS >	VertexShader(View.ShaderMap);
-		TShaderMapRef< THZBBuildPS<0> >	PixelShader(View.ShaderMap);
+				RHICmdList.ApplyCachedRenderTargets(GraphicsPSOInit);
+				GraphicsPSOInit.PrimitiveType = PT_TriangleList;
 
-		GraphicsPSOInit.BoundShaderState.VertexDeclarationRHI = GFilterVertexDeclaration.VertexDeclarationRHI;
-		GraphicsPSOInit.BoundShaderState.VertexShaderRHI = GETSAFERHISHADER_VERTEX(*VertexShader);
-		GraphicsPSOInit.BoundShaderState.PixelShaderRHI = GETSAFERHISHADER_PIXEL(*PixelShader);
-		GraphicsPSOInit.PrimitiveType = PT_TriangleList;
+				FHZBBuildPS::FPermutationDomain PermutationVector;
+				PermutationVector.Set<FHZBBuildPS::FStageDim>(false);
 
-		SetGraphicsPipelineState(RHICmdList, GraphicsPSOInit);
+				auto VertexShader = View.ShaderMap->GetShader< FPostProcessVS >();
+				auto PixelShader  = View.ShaderMap->GetShader< FHZBBuildPS >( PermutationVector );
 
-		// Imperfect sampling, doesn't matter too much
-		PixelShader->SetParameters( RHICmdList, View );
+				GraphicsPSOInit.BoundShaderState.VertexDeclarationRHI = GFilterVertexDeclaration.VertexDeclarationRHI;
+				GraphicsPSOInit.BoundShaderState.VertexShaderRHI = GETSAFERHISHADER_VERTEX(VertexShader);
+				GraphicsPSOInit.BoundShaderState.PixelShaderRHI = GETSAFERHISHADER_PIXEL(PixelShader);
+				GraphicsPSOInit.PrimitiveType = PT_TriangleList;
 
-		RHICmdList.SetViewport(0, 0, 0.0f, HZBSize.X, HZBSize.Y, 1.0f);
+				SetGraphicsPipelineState(RHICmdList, GraphicsPSOInit);
 
-		DrawRectangle(
-			RHICmdList,
-			0, 0,
-			HZBSize.X, HZBSize.Y,
-			View.ViewRect.Min.X, View.ViewRect.Min.Y,
-			View.ViewRect.Width(), View.ViewRect.Height(),
-			HZBSize,
-			FSceneRenderTargets::Get(RHICmdList).GetBufferSizeXY(),
-			*VertexShader,
-			EDRF_UseTriangleOptimization);
-		RHICmdList.EndRenderPass();
+				// Imperfect sampling, doesn't matter too much
+				FSceneRenderTargets& SceneContext = FSceneRenderTargets::Get(RHICmdList);
+				FIntPoint Size = SceneContext.GetBufferSizeXY();
+
+				FHZBBuildPS::FParameters Parameters;
+				Parameters.InvSize = FVector2D(1.0f / Size.X, 1.0f / Size.Y);
+				Parameters.InputUvFactorAndOffset = FVector4(
+					float(2 * HZBSize.X) / float(Size.X),
+					float(2 * HZBSize.Y) / float(Size.Y),
+					float(View.ViewRect.Min.X) / float(Size.X),
+					float(View.ViewRect.Min.Y) / float(Size.Y));
+				Parameters.InputViewportMaxBound = FVector2D(
+					float(View.ViewRect.Max.X) / float(Size.X) - 0.5f * Parameters.InvSize.X,
+					float(View.ViewRect.Max.Y) / float(Size.Y) - 0.5f * Parameters.InvSize.Y);
+		
+				Parameters.Pass = *PassParameters;
+				Parameters.View = View.ViewUniformBuffer;
+				Parameters.SceneTextures = CreateSceneTextureUniformBufferSingleDraw(RHICmdList, ESceneTextureSetupMode::SceneDepth, View.FeatureLevel);
+
+				SetShaderParameters(RHICmdList, PixelShader, PixelShader->GetPixelShader(), Parameters);
+
+				RHICmdList.SetViewport(0, 0, 0.0f, HZBSize.X, HZBSize.Y, 1.0f);
+
+				DrawRectangle(
+					RHICmdList,
+					0, 0,
+					HZBSize.X, HZBSize.Y,
+					View.ViewRect.Min.X, View.ViewRect.Min.Y,
+					View.ViewRect.Width(), View.ViewRect.Height(),
+					HZBSize,
+					FSceneRenderTargets::Get(RHICmdList).GetBufferSizeXY(),
+					VertexShader,
+					EDRF_UseTriangleOptimization);
+			});
 	}
 
 	FIntPoint SrcSize = HZBSize;
 	FIntPoint DstSize = SrcSize / 2;
 
-	SCOPED_DRAW_EVENTF(RHICmdList, BuildHZB, TEXT("HZB SetupMips Mips:1..%d %dx%d"), NumMips - 1, DstSize.X, DstSize.Y);
-
-	TShaderMapRef< FPostProcessVS >	VertexShader(View.ShaderMap);
-	TShaderMapRef< THZBBuildPS<1> >	PixelShader(View.ShaderMap);
-
 	// Downsampling...
 	for (uint8 MipIndex = 1; MipIndex < NumMips; MipIndex++)
 	{
 		DstSize.X = FMath::Max(DstSize.X, 1);
 		DstSize.Y = FMath::Max(DstSize.Y, 1);
 
-		FRHIRenderPassInfo RPInfo(HZBRenderTarget.TargetableTexture, ERenderTargetActions::Load_Store, nullptr, MipIndex);
-		RPInfo.bGeneratingMips = true;
-		RHICmdList.BeginRenderPass(RPInfo, TEXT("HZB_SetupSubMip"));
+		//@DW: Explicit creation of SRV, full configuration of SRV supported
+		FGraphSRVDesc Desc(HZBTexture, MipIndex - 1);
+		const FGraphSRV* ParentMipSRV = GraphBuilder.CreateSRV(Desc);
 
-		RHICmdList.ApplyCachedRenderTargets(GraphicsPSOInit);
+		FHZBBuildPassParameters* PassParameters;
+		GraphBuilder.CreateParameters(&PassParameters);
+		PassParameters->RenderTargets[0] = FRenderTargetBinding( HZBTexture, ERenderTargetLoadAction::ENoAction, ERenderTargetStoreAction::EStore, MipIndex );
+		PassParameters->Texture = ParentMipSRV;
+		PassParameters->TextureSampler = TStaticSamplerState<SF_Point, AM_Clamp, AM_Clamp, AM_Clamp>::GetRHI();
 
-		GraphicsPSOInit.BoundShaderState.VertexDeclarationRHI = GFilterVertexDeclaration.VertexDeclarationRHI;
-		GraphicsPSOInit.BoundShaderState.VertexShaderRHI = GETSAFERHISHADER_VERTEX(*VertexShader);
-		GraphicsPSOInit.BoundShaderState.PixelShaderRHI = GETSAFERHISHADER_PIXEL(*PixelShader);
-		GraphicsPSOInit.PrimitiveType = PT_TriangleList;
+		GraphBuilder.AddPass(TEXT("HZB_Mip"),
+			PassParameters,
+			ERenderGraphPassFlags::None,
+			[PassParameters, SrcSize, DstSize, &View](FRHICommandListImmediate& RHICmdList)
+			{
+				SCOPED_DRAW_EVENTF(RHICmdList, BuildHZB, TEXT("HZB MipLevel %dx%d"), DstSize.X, DstSize.Y);
 
-		SetGraphicsPipelineState(RHICmdList, GraphicsPSOInit);
+				FHZBBuildPS::FPermutationDomain PermutationVector;
+				PermutationVector.Set<FHZBBuildPS::FStageDim>(true);
 
-		PixelShader->SetParameters(RHICmdList, View, SrcSize, HZBRenderTarget.MipSRVs[MipIndex - 1]);
+				auto VertexShader = View.ShaderMap->GetShader< FPostProcessVS >();
+				auto PixelShader  = View.ShaderMap->GetShader< FHZBBuildPS >( PermutationVector );
 
-		RHICmdList.SetViewport(0, 0, 0.0f, DstSize.X, DstSize.Y, 1.0f);
+				FGraphicsPipelineStateInitializer GraphicsPSOInit;
+				GraphicsPSOInit.BlendState = TStaticBlendState<>::GetRHI();
+				GraphicsPSOInit.RasterizerState = TStaticRasterizerState<>::GetRHI();
+				GraphicsPSOInit.DepthStencilState = TStaticDepthStencilState<false, CF_Always>::GetRHI();
 
-		DrawRectangle(
-			RHICmdList,
-			0, 0,
-			DstSize.X, DstSize.Y,
-			0, 0,
-			SrcSize.X, SrcSize.Y,
-			DstSize,
-			SrcSize,
-			*VertexShader,
-			EDRF_UseTriangleOptimization);
+				RHICmdList.ApplyCachedRenderTargets(GraphicsPSOInit);
+
+				GraphicsPSOInit.BoundShaderState.VertexDeclarationRHI = GFilterVertexDeclaration.VertexDeclarationRHI;
+				GraphicsPSOInit.BoundShaderState.VertexShaderRHI = GETSAFERHISHADER_VERTEX(VertexShader);
+				GraphicsPSOInit.BoundShaderState.PixelShaderRHI = GETSAFERHISHADER_PIXEL(PixelShader);
+				GraphicsPSOInit.PrimitiveType = PT_TriangleList;
+
+				SetGraphicsPipelineState(RHICmdList, GraphicsPSOInit);
+
+				FHZBBuildPS::FParameters Parameters;
+				Parameters.InvSize = FVector2D(1.0f / SrcSize.X, 1.0f / SrcSize.Y);
+				Parameters.Pass = *PassParameters;
+				Parameters.View = View.ViewUniformBuffer;
+
+				SetShaderParameters(RHICmdList, PixelShader, PixelShader->GetPixelShader(), Parameters);
+
+				RHICmdList.SetViewport(0, 0, 0.0f, DstSize.X, DstSize.Y, 1.0f);
+
+				DrawRectangle(
+					RHICmdList,
+					0, 0,
+					DstSize.X, DstSize.Y,
+					0, 0,
+					SrcSize.X, SrcSize.Y,
+					DstSize,
+					SrcSize,
+					VertexShader,
+					EDRF_UseTriangleOptimization);
+			});
 
 		SrcSize /= 2;
 		DstSize /= 2;
-		RHICmdList.EndRenderPass();
 	}
 
-	// TODO: there is a lot of inconsistency in here. The manual ERWSubResBarrier as now be moved to within
-	// BeginRenderPass, as oposed to usual explicit resource transition.
-	RHICmdList.TransitionResources(EResourceTransitionAccess::EReadable, &HZBRenderTargetRef, 1);
-
-	GRenderTargetPool.VisualizeTexture.SetCheckPoint( RHICmdList, View.HZB );
+	GraphBuilder.GetInternalTexture(HZBTexture, &View.HZB);
+	GraphBuilder.Execute();
 }
 
 struct FViewOcclusionQueries
diff --git a/Engine/Source/Runtime/Renderer/Private/ScenePrivate.h b/Engine/Source/Runtime/Renderer/Private/ScenePrivate.h
--- a/Engine/Source/Runtime/Renderer/Private/ScenePrivate.h
+++ b/Engine/Source/Runtime/Renderer/Private/ScenePrivate.h
@@ -25,7 +25,7 @@
 #include "Materials/MaterialInstanceDynamic.h"
 #include "TextureLayout3d.h"
 #include "ScenePrivateBase.h"
-#include "PostProcess/RenderTargetPool.h"
+#include "RenderTargetPool.h"
 #include "SceneCore.h"
 #include "PrimitiveSceneInfo.h"
 #include "LightSceneInfo.h"
@@ -41,6 +41,7 @@
 #include "MobileBasePassRendering.h"
 #include "VolumeRendering.h"
 #include "SceneSoftwareOcclusion.h"
+#include "PostProcess/VisualizeTexture.h"
 
 /** Factor by which to grow occlusion tests **/
 #define OCCLUSION_SLOP (1.0f)
diff --git a/Engine/Source/Runtime/Renderer/Private/ScenePrivateBase.h b/Engine/Source/Runtime/Renderer/Private/ScenePrivateBase.h
--- a/Engine/Source/Runtime/Renderer/Private/ScenePrivateBase.h
+++ b/Engine/Source/Runtime/Renderer/Private/ScenePrivateBase.h
@@ -9,21 +9,6 @@
 #include "CoreMinimal.h"
 #include "RendererInterface.h"
 
-class SceneRenderingBitArrayAllocator
-	: public TInlineAllocator<4,SceneRenderingAllocator>
-{
-};
-
-class SceneRenderingSparseArrayAllocator
-	: public TSparseArrayAllocator<SceneRenderingAllocator,SceneRenderingBitArrayAllocator>
-{
-};
-
-class SceneRenderingSetAllocator
-	: public TSetAllocator<SceneRenderingSparseArrayAllocator,TInlineAllocator<1,SceneRenderingAllocator> >
-{
-};
-
 typedef TBitArray<SceneRenderingBitArrayAllocator> FSceneBitArray;
 typedef TConstSetBitIterator<SceneRenderingBitArrayAllocator> FSceneSetBitIterator;
 typedef TConstDualSetBitIterator<SceneRenderingBitArrayAllocator,SceneRenderingBitArrayAllocator> FSceneDualSetBitIterator;
diff --git a/Engine/Source/Runtime/Renderer/Private/SceneRendering.cpp b/Engine/Source/Runtime/Renderer/Private/SceneRendering.cpp
--- a/Engine/Source/Runtime/Renderer/Private/SceneRendering.cpp
+++ b/Engine/Source/Runtime/Renderer/Private/SceneRendering.cpp
@@ -2515,7 +2515,7 @@ void FSceneRenderer::RenderFinish(FRHICommandListImmediate& RHICmdList)
 				continue;
 			}
 
-			GRenderTargetPool.PresentContent(RHICmdList, View);
+			GVisualizeTexture.PresentContent(RHICmdList, View);
 		}
 	}
 #endif
@@ -2688,7 +2688,7 @@ void FSceneRenderer::OnStartFrame(FRHICommandListImmediate& RHICmdList)
 {
 	FSceneRenderTargets& SceneContext = FSceneRenderTargets::Get(RHICmdList);
 
-	GRenderTargetPool.VisualizeTexture.OnStartFrame(Views[0]);
+	GVisualizeTexture.OnStartFrame(Views[0]);
 	CompositionGraph_OnStartFrame();
 	SceneContext.bScreenSpaceAOIsValid = false;
 	SceneContext.bCustomDepthIsValid = false;
