Merging //UE4/Dev-Rendering-Graph@4492585 to Dev-Rendering (//UE4/Dev-Rendering)

diff --git a/Engine/Source/Runtime/Renderer/Private/SceneRendering.h b/Engine/Source/Runtime/Renderer/Private/SceneRendering.h
--- a/Engine/Source/Runtime/Renderer/Private/SceneRendering.h
+++ b/Engine/Source/Runtime/Renderer/Private/SceneRendering.h
@@ -29,6 +29,7 @@
 #include "HeightfieldLighting.h"
 #include "GlobalDistanceFieldParameters.h"
 #include "Templates/UniquePtr.h"
+#include "RenderGraph.h"
 
 class FScene;
 class FSceneViewState;
@@ -659,39 +660,39 @@ public:
 const int32 GMaxForwardShadowCascades = 4;
 
 #define FORWARD_GLOBAL_LIGHT_DATA_UNIFORM_BUFFER_MEMBER_TABLE \
-	UNIFORM_MEMBER(uint32,NumLocalLights) \
-	UNIFORM_MEMBER(uint32, NumReflectionCaptures) \
-	UNIFORM_MEMBER(uint32, HasDirectionalLight) \
-	UNIFORM_MEMBER(uint32, NumGridCells) \
-	UNIFORM_MEMBER(FIntVector, CulledGridSize) \
-	UNIFORM_MEMBER(uint32, MaxCulledLightsPerCell) \
-	UNIFORM_MEMBER(uint32, LightGridPixelSizeShift) \
-	UNIFORM_MEMBER(FVector, LightGridZParams) \
-	UNIFORM_MEMBER(FVector, DirectionalLightDirection) \
-	UNIFORM_MEMBER(FVector, DirectionalLightColor) \
-	UNIFORM_MEMBER(float, DirectionalLightVolumetricScatteringIntensity) \
-	UNIFORM_MEMBER(uint32, DirectionalLightShadowMapChannelMask) \
-	UNIFORM_MEMBER(FVector2D, DirectionalLightDistanceFadeMAD) \
-	UNIFORM_MEMBER(uint32, NumDirectionalLightCascades) \
-	UNIFORM_MEMBER(FVector4, CascadeEndDepths) \
-	UNIFORM_MEMBER_ARRAY(FMatrix, DirectionalLightWorldToShadowMatrix, [GMaxForwardShadowCascades]) \
-	UNIFORM_MEMBER_ARRAY(FVector4, DirectionalLightShadowmapMinMax, [GMaxForwardShadowCascades]) \
-	UNIFORM_MEMBER(FVector4, DirectionalLightShadowmapAtlasBufferSize) \
-	UNIFORM_MEMBER(float, DirectionalLightDepthBias) \
-	UNIFORM_MEMBER(uint32, DirectionalLightUseStaticShadowing) \
-	UNIFORM_MEMBER(FVector4, DirectionalLightStaticShadowBufferSize) \
-	UNIFORM_MEMBER(FMatrix, DirectionalLightWorldToStaticShadow) \
-	UNIFORM_MEMBER_TEXTURE(Texture2D, DirectionalLightShadowmapAtlas) \
-	UNIFORM_MEMBER_SAMPLER(SamplerState, ShadowmapSampler) \
-	UNIFORM_MEMBER_TEXTURE(Texture2D, DirectionalLightStaticShadowmap) \
-	UNIFORM_MEMBER_SAMPLER(SamplerState, StaticShadowmapSampler) \
-	UNIFORM_MEMBER_SRV(StrongTypedBuffer<float4>, ForwardLocalLightBuffer) \
-	UNIFORM_MEMBER_SRV(StrongTypedBuffer<uint>, NumCulledLightsGrid) \
-	UNIFORM_MEMBER_SRV(StrongTypedBuffer<uint>, CulledLightDataGrid) 
-
-BEGIN_UNIFORM_BUFFER_STRUCT_WITH_CONSTRUCTOR(FForwardLightData,)
+	SHADER_PARAMETER(uint32,NumLocalLights) \
+	SHADER_PARAMETER(uint32, NumReflectionCaptures) \
+	SHADER_PARAMETER(uint32, HasDirectionalLight) \
+	SHADER_PARAMETER(uint32, NumGridCells) \
+	SHADER_PARAMETER(FIntVector, CulledGridSize) \
+	SHADER_PARAMETER(uint32, MaxCulledLightsPerCell) \
+	SHADER_PARAMETER(uint32, LightGridPixelSizeShift) \
+	SHADER_PARAMETER(FVector, LightGridZParams) \
+	SHADER_PARAMETER(FVector, DirectionalLightDirection) \
+	SHADER_PARAMETER(FVector, DirectionalLightColor) \
+	SHADER_PARAMETER(float, DirectionalLightVolumetricScatteringIntensity) \
+	SHADER_PARAMETER(uint32, DirectionalLightShadowMapChannelMask) \
+	SHADER_PARAMETER(FVector2D, DirectionalLightDistanceFadeMAD) \
+	SHADER_PARAMETER(uint32, NumDirectionalLightCascades) \
+	SHADER_PARAMETER(FVector4, CascadeEndDepths) \
+	SHADER_PARAMETER_ARRAY(FMatrix, DirectionalLightWorldToShadowMatrix, [GMaxForwardShadowCascades]) \
+	SHADER_PARAMETER_ARRAY(FVector4, DirectionalLightShadowmapMinMax, [GMaxForwardShadowCascades]) \
+	SHADER_PARAMETER(FVector4, DirectionalLightShadowmapAtlasBufferSize) \
+	SHADER_PARAMETER(float, DirectionalLightDepthBias) \
+	SHADER_PARAMETER(uint32, DirectionalLightUseStaticShadowing) \
+	SHADER_PARAMETER(FVector4, DirectionalLightStaticShadowBufferSize) \
+	SHADER_PARAMETER(FMatrix, DirectionalLightWorldToStaticShadow) \
+	SHADER_PARAMETER_TEXTURE(Texture2D, DirectionalLightShadowmapAtlas) \
+	SHADER_PARAMETER_SAMPLER(SamplerState, ShadowmapSampler) \
+	SHADER_PARAMETER_TEXTURE(Texture2D, DirectionalLightStaticShadowmap) \
+	SHADER_PARAMETER_SAMPLER(SamplerState, StaticShadowmapSampler) \
+	SHADER_PARAMETER_SRV(StrongTypedBuffer<float4>, ForwardLocalLightBuffer) \
+	SHADER_PARAMETER_SRV(StrongTypedBuffer<uint>, NumCulledLightsGrid) \
+	SHADER_PARAMETER_SRV(StrongTypedBuffer<uint>, CulledLightDataGrid) 
+
+BEGIN_GLOBAL_SHADER_PARAMETER_STRUCT_WITH_CONSTRUCTOR(FForwardLightData,)
 	FORWARD_GLOBAL_LIGHT_DATA_UNIFORM_BUFFER_MEMBER_TABLE
-END_UNIFORM_BUFFER_STRUCT(FForwardLightData)
+END_GLOBAL_SHADER_PARAMETER_STRUCT()
 
 class FForwardLightingViewResources
 {
@@ -728,22 +729,23 @@ public:
 	}
 };
 
-BEGIN_UNIFORM_BUFFER_STRUCT_WITH_CONSTRUCTOR(FVolumetricFogGlobalData,) 
-	UNIFORM_MEMBER(FIntVector, GridSizeInt)
-	UNIFORM_MEMBER(FVector, GridSize)
-	UNIFORM_MEMBER(uint32, GridPixelSizeShift)
-	UNIFORM_MEMBER(FVector, GridZParams)
-	UNIFORM_MEMBER(FVector2D, SVPosToVolumeUV)
-	UNIFORM_MEMBER(FIntPoint, FogGridToPixelXY)
-	UNIFORM_MEMBER(float, MaxDistance)
-	UNIFORM_MEMBER(FVector, HeightFogInscatteringColor)
-	UNIFORM_MEMBER(FVector, HeightFogDirectionalLightInscatteringColor)
-END_UNIFORM_BUFFER_STRUCT(FVolumetricFogGlobalData)
+BEGIN_GLOBAL_SHADER_PARAMETER_STRUCT_WITH_CONSTRUCTOR(FVolumetricFogGlobalData,)
+	SHADER_PARAMETER(FIntVector, GridSizeInt)
+	SHADER_PARAMETER(FVector, GridSize)
+	SHADER_PARAMETER(uint32, GridPixelSizeShift)
+	SHADER_PARAMETER(FVector, GridZParams)
+	SHADER_PARAMETER(FVector2D, SVPosToVolumeUV)
+	SHADER_PARAMETER(FIntPoint, FogGridToPixelXY)
+	SHADER_PARAMETER(float, MaxDistance)
+	SHADER_PARAMETER(FVector, HeightFogInscatteringColor)
+	SHADER_PARAMETER(FVector, HeightFogDirectionalLightInscatteringColor)
+END_GLOBAL_SHADER_PARAMETER_STRUCT()
 
 class FVolumetricFogViewResources
 {
 public:
 	TUniformBufferRef<FVolumetricFogGlobalData> VolumetricFogGlobalData;
+
 	TRefCountPtr<IPooledRenderTarget> IntegratedLightScattering;
 
 	FVolumetricFogViewResources()
@@ -793,15 +795,15 @@ private:
 static const int32 GMaxNumReflectionCaptures = 341;
 
 /** Per-reflection capture data needed by the shader. */
-BEGIN_UNIFORM_BUFFER_STRUCT(FReflectionCaptureShaderData,)
-	UNIFORM_MEMBER_ARRAY(FVector4,PositionAndRadius,[GMaxNumReflectionCaptures])
+BEGIN_GLOBAL_SHADER_PARAMETER_STRUCT(FReflectionCaptureShaderData,)
+	SHADER_PARAMETER_ARRAY(FVector4,PositionAndRadius,[GMaxNumReflectionCaptures])
 	// R is brightness, G is array index, B is shape
-	UNIFORM_MEMBER_ARRAY(FVector4,CaptureProperties,[GMaxNumReflectionCaptures])
-	UNIFORM_MEMBER_ARRAY(FVector4,CaptureOffsetAndAverageBrightness,[GMaxNumReflectionCaptures])
+	SHADER_PARAMETER_ARRAY(FVector4,CaptureProperties,[GMaxNumReflectionCaptures])
+	SHADER_PARAMETER_ARRAY(FVector4,CaptureOffsetAndAverageBrightness,[GMaxNumReflectionCaptures])
 	// Stores the box transform for a box shape, other data is packed for other shapes
-	UNIFORM_MEMBER_ARRAY(FMatrix,BoxTransform,[GMaxNumReflectionCaptures])
-	UNIFORM_MEMBER_ARRAY(FVector4,BoxScales,[GMaxNumReflectionCaptures])
-END_UNIFORM_BUFFER_STRUCT(FReflectionCaptureShaderData)
+	SHADER_PARAMETER_ARRAY(FMatrix,BoxTransform,[GMaxNumReflectionCaptures])
+	SHADER_PARAMETER_ARRAY(FVector4,BoxScales,[GMaxNumReflectionCaptures])
+END_GLOBAL_SHADER_PARAMETER_STRUCT()
 
 // Structure in charge of storing all information about TAA's history.
 struct FTemporalAAHistory
diff --git a/Engine/Source/Runtime/Renderer/Private/ShaderBaseClasses.h b/Engine/Source/Runtime/Renderer/Private/ShaderBaseClasses.h
--- a/Engine/Source/Runtime/Renderer/Private/ShaderBaseClasses.h
+++ b/Engine/Source/Runtime/Renderer/Private/ShaderBaseClasses.h
@@ -20,9 +20,9 @@ struct FMeshDrawingRenderState;
 // This was moved out of ScenePrivate.h to workaround MSVC vs clang template issue (it's used in this header file, so needs to be declared earlier)
 // Z is the dither fade value (-1 = just fading in, 0 no fade, 1 = just faded out)
 // W is unused and zero
-BEGIN_UNIFORM_BUFFER_STRUCT(FDistanceCullFadeUniformShaderParameters,)
-	UNIFORM_MEMBER_EX(FVector2D,FadeTimeScaleBias, EShaderPrecisionModifier::Half)
-END_UNIFORM_BUFFER_STRUCT(FDistanceCullFadeUniformShaderParameters)
+BEGIN_GLOBAL_SHADER_PARAMETER_STRUCT(FDistanceCullFadeUniformShaderParameters,)
+	SHADER_PARAMETER_EX(FVector2D,FadeTimeScaleBias, EShaderPrecisionModifier::Half)
+END_GLOBAL_SHADER_PARAMETER_STRUCT()
 
 typedef TUniformBufferRef< FDistanceCullFadeUniformShaderParameters > FDistanceCullFadeUniformBufferRef;
 
@@ -57,7 +57,7 @@ public:
 	FBaseHS(const ShaderMetaType::CompiledShaderInitializerType& Initializer):
 		FMeshMaterialShader(Initializer)
 	{
-		PassUniformBuffer.Bind(Initializer.ParameterMap, FSceneTexturesUniformParameters::StaticStruct.GetShaderVariableName());
+		PassUniformBuffer.Bind(Initializer.ParameterMap, FSceneTexturesUniformParameters::StaticStructMetadata.GetShaderVariableName());
 	}
 
 	FBaseHS() {}
@@ -109,7 +109,7 @@ public:
 	FBaseDS(const ShaderMetaType::CompiledShaderInitializerType& Initializer):
 		FMeshMaterialShader(Initializer)
 	{
-		PassUniformBuffer.Bind(Initializer.ParameterMap, FSceneTexturesUniformParameters::StaticStruct.GetShaderVariableName());
+		PassUniformBuffer.Bind(Initializer.ParameterMap, FSceneTexturesUniformParameters::StaticStructMetadata.GetShaderVariableName());
 	}
 
 	FBaseDS() {}
diff --git a/Engine/Source/Runtime/Renderer/Private/ShadowDepthRendering.cpp b/Engine/Source/Runtime/Renderer/Private/ShadowDepthRendering.cpp
--- a/Engine/Source/Runtime/Renderer/Private/ShadowDepthRendering.cpp
+++ b/Engine/Source/Runtime/Renderer/Private/ShadowDepthRendering.cpp
@@ -2068,7 +2068,7 @@ void FSceneRenderer::RenderShadowDepthMapAtlases(FRHICommandListImmediate& RHICm
 		FSceneRenderTargetItem& RenderTarget = ShadowMapAtlas.RenderTargets.DepthTarget->GetRenderTargetItem();
 		FIntPoint AtlasSize = ShadowMapAtlas.RenderTargets.DepthTarget->GetDesc().Extent;
 
-		GRenderTargetPool.VisualizeTexture.SetCheckPoint(RHICmdList, ShadowMapAtlas.RenderTargets.DepthTarget.GetReference());
+		GVisualizeTexture.SetCheckPoint(RHICmdList, ShadowMapAtlas.RenderTargets.DepthTarget.GetReference());
 
 		SCOPED_DRAW_EVENTF(RHICmdList, EventShadowDepths, TEXT("Atlas%u %ux%u"), AtlasIndex, AtlasSize.X, AtlasSize.Y);
 
@@ -2163,7 +2163,7 @@ void FSceneRenderer::RenderShadowDepthMaps(FRHICommandListImmediate& RHICmdList)
 		check(ShadowMap.Shadows.Num() == 1);
 		FProjectedShadowInfo* ProjectedShadowInfo = ShadowMap.Shadows[0];
 
-		GRenderTargetPool.VisualizeTexture.SetCheckPoint(RHICmdList, ShadowMap.RenderTargets.DepthTarget.GetReference());
+		GVisualizeTexture.SetCheckPoint(RHICmdList, ShadowMap.RenderTargets.DepthTarget.GetReference());
 
 		FString LightNameWithLevel;
 		GetLightNameForDrawEvent(ProjectedShadowInfo->GetLightSceneInfo().Proxy, LightNameWithLevel);
@@ -2213,7 +2213,7 @@ void FSceneRenderer::RenderShadowDepthMaps(FRHICommandListImmediate& RHICmdList)
 	{
 		FSceneRenderTargetItem& RenderTarget = SortedShadowsForShadowDepthPass.PreshadowCache.RenderTargets.DepthTarget->GetRenderTargetItem();
 
-		GRenderTargetPool.VisualizeTexture.SetCheckPoint(RHICmdList, SortedShadowsForShadowDepthPass.PreshadowCache.RenderTargets.DepthTarget.GetReference());
+		GVisualizeTexture.SetCheckPoint(RHICmdList, SortedShadowsForShadowDepthPass.PreshadowCache.RenderTargets.DepthTarget.GetReference());
 
 		SCOPED_DRAW_EVENT(RHICmdList, PreshadowCache);
 
diff --git a/Engine/Source/Runtime/Renderer/Private/TiledDeferredLightRendering.cpp b/Engine/Source/Runtime/Renderer/Private/TiledDeferredLightRendering.cpp
--- a/Engine/Source/Runtime/Renderer/Private/TiledDeferredLightRendering.cpp
+++ b/Engine/Source/Runtime/Renderer/Private/TiledDeferredLightRendering.cpp
@@ -59,21 +59,21 @@ static FAutoConsoleVariableRef CVarNumLightsBeforeUsingTiledDeferred(
  * First constant buffer of light data for tiled deferred. 
  * Light data is split into two constant buffers to allow more lights per pass before hitting the d3d11 max constant buffer size of 4096 float4's
  */
-BEGIN_UNIFORM_BUFFER_STRUCT(FTiledDeferredLightData,)
-	UNIFORM_MEMBER_ARRAY(FVector4,LightPositionAndInvRadius,[GMaxNumTiledDeferredLights])
-	UNIFORM_MEMBER_ARRAY(FVector4,LightColorAndFalloffExponent,[GMaxNumTiledDeferredLights])
-END_UNIFORM_BUFFER_STRUCT(FTiledDeferredLightData)
+BEGIN_GLOBAL_SHADER_PARAMETER_STRUCT(FTiledDeferredLightData,)
+	SHADER_PARAMETER_ARRAY(FVector4,LightPositionAndInvRadius,[GMaxNumTiledDeferredLights])
+	SHADER_PARAMETER_ARRAY(FVector4,LightColorAndFalloffExponent,[GMaxNumTiledDeferredLights])
+END_GLOBAL_SHADER_PARAMETER_STRUCT()
 
-IMPLEMENT_UNIFORM_BUFFER_STRUCT(FTiledDeferredLightData,TEXT("TiledDeferred"));
+IMPLEMENT_GLOBAL_SHADER_PARAMETER_STRUCT(FTiledDeferredLightData,"TiledDeferred");
 
 /** Second constant buffer of light data for tiled deferred. */
-BEGIN_UNIFORM_BUFFER_STRUCT(FTiledDeferredLightData2,)
-	UNIFORM_MEMBER_ARRAY(FVector4,LightDirectionAndSpotlightMaskAndSpecularScale,[GMaxNumTiledDeferredLights])
-	UNIFORM_MEMBER_ARRAY(FVector4,SpotAnglesAndSourceRadiusAndSimpleLighting,[GMaxNumTiledDeferredLights])
-	UNIFORM_MEMBER_ARRAY(FVector4,ShadowMapChannelMask,[GMaxNumTiledDeferredLights])
-END_UNIFORM_BUFFER_STRUCT(FTiledDeferredLightData2)
+BEGIN_GLOBAL_SHADER_PARAMETER_STRUCT(FTiledDeferredLightData2,)
+	SHADER_PARAMETER_ARRAY(FVector4,LightDirectionAndSpotlightMaskAndSpecularScale,[GMaxNumTiledDeferredLights])
+	SHADER_PARAMETER_ARRAY(FVector4,SpotAnglesAndSourceRadiusAndSimpleLighting,[GMaxNumTiledDeferredLights])
+	SHADER_PARAMETER_ARRAY(FVector4,ShadowMapChannelMask,[GMaxNumTiledDeferredLights])
+END_GLOBAL_SHADER_PARAMETER_STRUCT()
 
-IMPLEMENT_UNIFORM_BUFFER_STRUCT(FTiledDeferredLightData2,TEXT("TiledDeferred2"));
+IMPLEMENT_GLOBAL_SHADER_PARAMETER_STRUCT(FTiledDeferredLightData2,"TiledDeferred2");
 
 /** Compute shader used to implement tiled deferred lighting. */
 template <bool bVisualizeLightCulling>
diff --git a/Engine/Source/Runtime/Renderer/Private/TranslucentLighting.cpp b/Engine/Source/Runtime/Renderer/Private/TranslucentLighting.cpp
--- a/Engine/Source/Runtime/Renderer/Private/TranslucentLighting.cpp
+++ b/Engine/Source/Runtime/Renderer/Private/TranslucentLighting.cpp
@@ -200,7 +200,7 @@ public:
 		FMeshMaterialShader(Initializer)
 	{
 		ShadowParameters.Bind(Initializer.ParameterMap);
-		PassUniformBuffer.Bind(Initializer.ParameterMap, FSceneTexturesUniformParameters::StaticStruct.GetShaderVariableName());
+		PassUniformBuffer.Bind(Initializer.ParameterMap, FSceneTexturesUniformParameters::StaticStructMetadata.GetShaderVariableName());
 	}
 
 	virtual bool Serialize(FArchive& Ar) override
@@ -279,7 +279,7 @@ public:
 		TranslInvMaxSubjectDepth.Bind(Initializer.ParameterMap,TEXT("TranslInvMaxSubjectDepth"));
 		TranslucentShadowStartOffset.Bind(Initializer.ParameterMap,TEXT("TranslucentShadowStartOffset"));
 		TranslucencyProjectionParameters.Bind(Initializer.ParameterMap);
-		PassUniformBuffer.Bind(Initializer.ParameterMap, FSceneTexturesUniformParameters::StaticStruct.GetShaderVariableName());
+		PassUniformBuffer.Bind(Initializer.ParameterMap, FSceneTexturesUniformParameters::StaticStructMetadata.GetShaderVariableName());
 	}
 
 	FTranslucencyShadowDepthPS() {}
@@ -1434,8 +1434,8 @@ static void InjectTranslucentLightArray(FRHICommandListImmediate& RHICmdList, co
 		const IPooledRenderTarget* RT0 = SceneContext.TranslucencyLightingVolumeAmbient[VolumeCascadeIndex + NumTranslucentVolumeRenderTargetSets * ViewIndex];
 		const IPooledRenderTarget* RT1 = SceneContext.TranslucencyLightingVolumeDirectional[VolumeCascadeIndex + NumTranslucentVolumeRenderTargetSets * ViewIndex];
 
-		GRenderTargetPool.VisualizeTexture.SetCheckPoint(RHICmdList, RT0);
-		GRenderTargetPool.VisualizeTexture.SetCheckPoint(RHICmdList, RT1);
+		GVisualizeTexture.SetCheckPoint(RHICmdList, RT0);
+		GVisualizeTexture.SetCheckPoint(RHICmdList, RT1);
 
 		FTextureRHIParamRef RenderTargets[2];
 		RenderTargets[0] = RT0->GetRenderTargetItem().TargetableTexture;
@@ -1645,8 +1645,8 @@ void FDeferredShadingSceneRenderer::InjectSimpleTranslucentVolumeLightingArray(F
 			const IPooledRenderTarget* RT0 = SceneContext.TranslucencyLightingVolumeAmbient[VolumeCascadeIndex + NumTranslucentVolumeRenderTargetSets * ViewIndex];
 			const IPooledRenderTarget* RT1 = SceneContext.TranslucencyLightingVolumeDirectional[VolumeCascadeIndex + NumTranslucentVolumeRenderTargetSets * ViewIndex];
 
-			GRenderTargetPool.VisualizeTexture.SetCheckPoint(RHICmdList, RT0);
-			GRenderTargetPool.VisualizeTexture.SetCheckPoint(RHICmdList, RT1);
+			GVisualizeTexture.SetCheckPoint(RHICmdList, RT0);
+			GVisualizeTexture.SetCheckPoint(RHICmdList, RT1);
 
 			FTextureRHIParamRef RenderTargets[2];
 			RenderTargets[0] = RT0->GetRenderTargetItem().TargetableTexture;
@@ -1749,8 +1749,8 @@ void FDeferredShadingSceneRenderer::FilterTranslucentVolumeLighting(FRHICommandL
 				const IPooledRenderTarget* Input0 = SceneContext.TranslucencyLightingVolumeAmbient[VolumeCascadeIndex + NumTranslucentVolumeRenderTargetSets * ViewIndex];
 				const IPooledRenderTarget* Input1 = SceneContext.TranslucencyLightingVolumeDirectional[VolumeCascadeIndex + NumTranslucentVolumeRenderTargetSets * ViewIndex];
 
-				GRenderTargetPool.VisualizeTexture.SetCheckPoint(RHICmdList, RT0);
-				GRenderTargetPool.VisualizeTexture.SetCheckPoint(RHICmdList, RT1);
+				GVisualizeTexture.SetCheckPoint(RHICmdList, RT0);
+				GVisualizeTexture.SetCheckPoint(RHICmdList, RT1);
 
 				FTextureRHIParamRef RenderTargets[2];
 				RenderTargets[0] = RT0->GetRenderTargetItem().TargetableTexture;
diff --git a/Engine/Source/Runtime/Renderer/Private/VT/VirtualTextureFeedback.cpp b/Engine/Source/Runtime/Renderer/Private/VT/VirtualTextureFeedback.cpp
--- a/Engine/Source/Runtime/Renderer/Private/VT/VirtualTextureFeedback.cpp
+++ b/Engine/Source/Runtime/Renderer/Private/VT/VirtualTextureFeedback.cpp
@@ -1,6 +1,7 @@
 // Copyright 1998-2018 Epic Games, Inc. All Rights Reserved.
 
 #include "VirtualTextureFeedback.h"
+#include "PostProcess/VisualizeTexture.h"
 
 #include "ClearQuad.h"
 
@@ -36,7 +37,7 @@ void FVirtualTextureFeedback::TransferGPUToCPU( FRHICommandListImmediate& RHICmd
 {
 	RHICmdList.TransitionResource( EResourceTransitionAccess::EReadable, EResourceTransitionPipeline::EGfxToGfx, FeedbackTextureGPU->GetRenderTargetItem().UAV );
 
-	GRenderTargetPool.VisualizeTexture.SetCheckPoint( RHICmdList, FeedbackTextureGPU );
+	GVisualizeTexture.SetCheckPoint( RHICmdList, FeedbackTextureGPU );
 	
 	FPooledRenderTargetDesc Desc( FPooledRenderTargetDesc::Create2DDesc( Size, PF_R32_UINT, FClearValueBinding::None, TexCreate_CPUReadback | TexCreate_HideInVisualizeTexture, TexCreate_None, false ) );
 	GRenderTargetPool.FindFreeElement( RHICmdList, Desc, FeedbackTextureCPU, TEXT("VTFeedbackCPU") );
diff --git a/Engine/Source/Runtime/Renderer/Private/VT/VirtualTextureSpace.cpp b/Engine/Source/Runtime/Renderer/Private/VT/VirtualTextureSpace.cpp
--- a/Engine/Source/Runtime/Renderer/Private/VT/VirtualTextureSpace.cpp
+++ b/Engine/Source/Runtime/Renderer/Private/VT/VirtualTextureSpace.cpp
@@ -10,6 +10,7 @@
 #include "PipelineStateCache.h"
 #include "HAL/IConsoleManager.h"
 #include "SceneUtils.h"
+#include "PostProcess/VisualizeTexture.h"
 
 static TAutoConsoleVariable<int32> CVarVTRefreshEntirePageTable(
 	TEXT("r.VT.RefreshEntirePageTable"),
@@ -243,7 +244,7 @@ void FVirtualTextureSpace::ApplyUpdates( FRHICommandList& RHICmdList )
 	{
 		if( PageTableUpdates.Num() == 0 )
 		{
-			GRenderTargetPool.VisualizeTexture.SetCheckPoint( RHICmdList, PageTable );
+			GVisualizeTexture.SetCheckPoint( RHICmdList, PageTable );
 			return;
 		}
 
@@ -373,5 +374,5 @@ void FVirtualTextureSpace::ApplyUpdates( FRHICommandList& RHICmdList )
 
 	RHICmdList.CopyToResolveTarget( PageTableTarget.TargetableTexture, PageTableTarget.ShaderResourceTexture, FResolveParams() );
 
-	GRenderTargetPool.VisualizeTexture.SetCheckPoint( RHICmdList, PageTable );
+	GVisualizeTexture.SetCheckPoint( RHICmdList, PageTable );
 }
\ No newline at end of file
diff --git a/Engine/Source/Runtime/Renderer/Private/VT/VirtualTextureTest.cpp b/Engine/Source/Runtime/Renderer/Private/VT/VirtualTextureTest.cpp
--- a/Engine/Source/Runtime/Renderer/Private/VT/VirtualTextureTest.cpp
+++ b/Engine/Source/Runtime/Renderer/Private/VT/VirtualTextureTest.cpp
@@ -219,7 +219,7 @@ void FVirtualTextureTest::ProducePageData( FRHICommandList& RHICmdList, ERHIFeat
 
 	RHICmdList.CopyToResolveTarget( RenderTarget.TargetableTexture, RenderTarget.ShaderResourceTexture, false, FResolveParams() );
 
-	GRenderTargetPool.VisualizeTexture.SetCheckPoint( RHICmdList, PhysicalTexture );
+	GVisualizeTexture.SetCheckPoint( RHICmdList, PhysicalTexture );
 }
 
 void FVirtualTextureTest::DumpToConsole()
diff --git a/Engine/Source/Runtime/Renderer/Private/VelocityRendering.cpp b/Engine/Source/Runtime/Renderer/Private/VelocityRendering.cpp
--- a/Engine/Source/Runtime/Renderer/Private/VelocityRendering.cpp
+++ b/Engine/Source/Runtime/Renderer/Private/VelocityRendering.cpp
@@ -127,7 +127,7 @@ protected:
 		PrevTransformBuffer.Bind(Initializer.ParameterMap, TEXT("PrevTransformBuffer"));
 		InstancedEyeIndexParameter.Bind(Initializer.ParameterMap, TEXT("InstancedEyeIndex"));
 		IsInstancedStereoParameter.Bind(Initializer.ParameterMap, TEXT("bIsInstancedStereo"));
-		PassUniformBuffer.Bind(Initializer.ParameterMap, FSceneTexturesUniformParameters::StaticStruct.GetShaderVariableName());
+		PassUniformBuffer.Bind(Initializer.ParameterMap, FSceneTexturesUniformParameters::StaticStructMetadata.GetShaderVariableName());
 	}
 
 	FVelocityVS() {}
@@ -240,7 +240,7 @@ public:
 	FVelocityPS(const ShaderMetaType::CompiledShaderInitializerType& Initializer)
 		: FMeshMaterialShader(Initializer)
 	{
-		PassUniformBuffer.Bind(Initializer.ParameterMap, FSceneTexturesUniformParameters::StaticStruct.GetShaderVariableName());
+		PassUniformBuffer.Bind(Initializer.ParameterMap, FSceneTexturesUniformParameters::StaticStructMetadata.GetShaderVariableName());
 	}
 
 	void SetParameters(FRHICommandList& RHICmdList, const FVertexFactory* VertexFactory,const FMaterialRenderProxy* MaterialRenderProxy,const FViewInfo& View, const FDrawingPolicyRenderState& DrawRenderState)
@@ -941,7 +941,7 @@ void FDeferredShadingSceneRenderer::RenderVelocities(FRHICommandListImmediate& R
 	}
 
 	// to be able to observe results with VisualizeTexture
-	GRenderTargetPool.VisualizeTexture.SetCheckPoint(RHICmdList, VelocityRT);
+	GVisualizeTexture.SetCheckPoint(RHICmdList, VelocityRT);
 }
 
 FPooledRenderTargetDesc FVelocityRendering::GetRenderTargetDesc()
diff --git a/Engine/Source/Runtime/Renderer/Private/VolumetricFog.cpp b/Engine/Source/Runtime/Renderer/Private/VolumetricFog.cpp
--- a/Engine/Source/Runtime/Renderer/Private/VolumetricFog.cpp
+++ b/Engine/Source/Runtime/Renderer/Private/VolumetricFog.cpp
@@ -17,6 +17,7 @@
 #include "ScreenRendering.h"
 #include "VolumeLighting.h"
 #include "PipelineStateCache.h"
+#include "ShaderParameterStruct.h"
 
 int32 GVolumetricFog = 1;
 FAutoConsoleVariableRef CVarVolumetricFog(
@@ -99,7 +100,7 @@ FAutoConsoleVariableRef CVarInverseSquaredLightDistanceBiasScale(
 	ECVF_Scalability | ECVF_RenderThreadSafe
 );
 
-IMPLEMENT_UNIFORM_BUFFER_STRUCT(FVolumetricFogGlobalData,TEXT("VolumetricFog"));
+IMPLEMENT_GLOBAL_SHADER_PARAMETER_STRUCT(FVolumetricFogGlobalData, "VolumetricFog");
 
 FVolumetricFogGlobalData::FVolumetricFogGlobalData()
 {}
@@ -136,6 +137,20 @@ uint32 VolumetricFogGridInjectionGroupSize = 4;
 class FVolumetricFogMaterialSetupCS : public FGlobalShader
 {
 	DECLARE_SHADER_TYPE(FVolumetricFogMaterialSetupCS,Global)
+	//SHADER_USE_PARAMETER_STRUCT(FVolumetricFogMaterialSetupCS, FGlobalShader)
+
+	BEGIN_SHADER_PARAMETER_STRUCT(FParameters, )
+		SHADER_PARAMETER( FLinearColor,	GlobalAlbedo )
+		SHADER_PARAMETER( FLinearColor,	GlobalEmissive )
+		SHADER_PARAMETER( float,		GlobalExtinctionScale )
+		
+		SHADER_PARAMETER_STRUCT_REF( FFogUniformParameters,			FogUniformParameters )
+		SHADER_PARAMETER_STRUCT_REF( FViewUniformShaderParameters,	View )
+		
+		SHADER_PARAMETER_GRAPH_UAV(	RWTexture3D<float4>, RWVBufferA )
+		SHADER_PARAMETER_GRAPH_UAV(	RWTexture3D<float4>, RWVBufferB )
+	END_SHADER_PARAMETER_STRUCT()
+
 public:
 
 	static bool ShouldCompilePermutation(const FGlobalShaderPermutationParameters& Parameters)
@@ -153,10 +168,9 @@ public:
 	FVolumetricFogMaterialSetupCS(const ShaderMetaType::CompiledShaderInitializerType& Initializer)
 		: FGlobalShader(Initializer)
 	{
+		Bindings.Bind(Initializer.ParameterMap, *FParameters::FTypeInfo::GetStructMetadata());
+
 		VolumetricFogParameters.Bind(Initializer.ParameterMap);
-		GlobalAlbedo.Bind(Initializer.ParameterMap, TEXT("GlobalAlbedo"));
-		GlobalEmissive.Bind(Initializer.ParameterMap, TEXT("GlobalEmissive"));
-		GlobalExtinctionScale.Bind(Initializer.ParameterMap, TEXT("GlobalExtinctionScale"));
 	} 
 
 	FVolumetricFogMaterialSetupCS()
@@ -166,45 +180,22 @@ public:
 	void SetParameters(
 		FRHICommandList& RHICmdList, 
 		const FViewInfo& View, 
-		const FVolumetricFogIntegrationParameterData& IntegrationData,
-		const FExponentialHeightFogSceneInfo& FogInfo)
+		const FVolumetricFogIntegrationParameterData& IntegrationData)
 	{
 		FComputeShaderRHIParamRef ShaderRHI = GetComputeShader();
-		FGlobalShader::SetParameters<FViewUniformShaderParameters>(RHICmdList, ShaderRHI, View.ViewUniformBuffer);
 		VolumetricFogParameters.Set(RHICmdList, ShaderRHI, View, IntegrationData);
-		SetShaderValue(RHICmdList, ShaderRHI, GlobalAlbedo, FogInfo.VolumetricFogAlbedo);
-		SetShaderValue(RHICmdList, ShaderRHI, GlobalEmissive, FogInfo.VolumetricFogEmissive);
-		SetShaderValue(RHICmdList, ShaderRHI, GlobalExtinctionScale, FogInfo.VolumetricFogExtinctionScale);
-
-		FFogUniformParameters FogUniformParameters;
-		SetupFogUniformParameters(View, FogUniformParameters);
-		SetUniformBufferParameterImmediate(RHICmdList, ShaderRHI, GetUniformBufferParameter<FFogUniformParameters>(), FogUniformParameters);
-	}
-
-	void UnsetParameters(FRHICommandList& RHICmdList, 
-		const FViewInfo& View,
-		IPooledRenderTarget* VBufferA,
-		IPooledRenderTarget* VBufferB)
-	{
-		VolumetricFogParameters.UnsetParameters(RHICmdList, GetComputeShader(), View, VBufferA, VBufferB, NULL, false);
 	}
 
 	virtual bool Serialize(FArchive& Ar)
 	{		
 		bool bShaderHasOutdatedParameters = FGlobalShader::Serialize(Ar);
 		Ar << VolumetricFogParameters;
-		Ar << GlobalAlbedo;
-		Ar << GlobalEmissive;
-		Ar << GlobalExtinctionScale;
 		return bShaderHasOutdatedParameters;
 	}
 
 private:
 
 	FVolumetricFogIntegrationParameters VolumetricFogParameters;
-	FShaderParameter GlobalAlbedo;
-	FShaderParameter GlobalEmissive;
-	FShaderParameter GlobalExtinctionScale;
 };
 
 IMPLEMENT_SHADER_TYPE(,FVolumetricFogMaterialSetupCS,TEXT("/Engine/Private/VolumetricFog.usf"),TEXT("MaterialSetupCS"),SF_Compute);
@@ -271,18 +262,22 @@ private:
 IMPLEMENT_SHADER_TYPE(,FWriteToBoundingSphereVS,TEXT("/Engine/Private/VolumetricFog.usf"),TEXT("WriteToBoundingSphereVS"),SF_Vertex);
 
 /** Shader that adds direct lighting contribution from the given light to the current volume lighting cascade. */
-template<bool bDynamicallyShadowed, bool bInverseSquared, bool bTemporalReprojection>
 class TInjectShadowedLocalLightPS : public FGlobalShader
 {
-	DECLARE_SHADER_TYPE(TInjectShadowedLocalLightPS,Global);
-public:
+	DECLARE_GLOBAL_SHADER(TInjectShadowedLocalLightPS);
+
+	class FDynamicallyShadowed	: SHADER_PERMUTATION_BOOL("DYNAMICALLY_SHADOWED");
+	class FInverseSquared		: SHADER_PERMUTATION_BOOL("INVERSE_SQUARED_FALLOFF");
+	class FTemporalReprojection	: SHADER_PERMUTATION_BOOL("USE_TEMPORAL_REPROJECTION");
+
+	using FPermutationDomain = TShaderPermutationDomain<
+		FDynamicallyShadowed,
+		FInverseSquared,
+		FTemporalReprojection >;
 
 	static void ModifyCompilationEnvironment(const FGlobalShaderPermutationParameters& Parameters, FShaderCompilerEnvironment& OutEnvironment)
 	{
 		FGlobalShader::ModifyCompilationEnvironment(Parameters, OutEnvironment);
-		OutEnvironment.SetDefine(TEXT("DYNAMICALLY_SHADOWED"), (uint32)bDynamicallyShadowed);
-		OutEnvironment.SetDefine(TEXT("INVERSE_SQUARED_FALLOFF"), (uint32)bInverseSquared);
-		OutEnvironment.SetDefine(TEXT("USE_TEMPORAL_REPROJECTION"), bTemporalReprojection);
 		FVolumetricFogIntegrationParameters::ModifyCompilationEnvironment(Parameters.Platform, OutEnvironment);
 	}
 
@@ -302,6 +297,7 @@ public:
 
 	TInjectShadowedLocalLightPS() {}
 
+public:
 	// @param InnerSplitIndex which CSM shadow map level, INDEX_NONE if no directional light
 	// @param VolumeCascadeIndexValue which volume we render to
 	void SetParameters(
@@ -310,7 +306,8 @@ public:
 		const FVolumetricFogIntegrationParameterData& IntegrationData,
 		const FLightSceneInfo* LightSceneInfo,
 		const FExponentialHeightFogSceneInfo& FogInfo,
-		const FProjectedShadowInfo* ShadowMap)
+		const FProjectedShadowInfo* ShadowMap,
+		bool bDynamicallyShadowed)
 	{
 		const FPixelShaderRHIParamRef ShaderRHI = GetPixelShader();
 
@@ -343,18 +340,7 @@ private:
 	FVolumeShadowingParameters VolumeShadowingParameters;
 };
 
-#define IMPLEMENT_LOCAL_LIGHT_INJECTION_PIXELSHADER_TYPE(bDynamicallyShadowed,bInverseSquared,bTemporalReprojection) \
-	typedef TInjectShadowedLocalLightPS<bDynamicallyShadowed,bInverseSquared,bTemporalReprojection> TInjectShadowedLocalLightPS##bDynamicallyShadowed##bInverseSquared##bTemporalReprojection; \
-	IMPLEMENT_SHADER_TYPE(template<>,TInjectShadowedLocalLightPS##bDynamicallyShadowed##bInverseSquared##bTemporalReprojection,TEXT("/Engine/Private/VolumetricFog.usf"),TEXT("InjectShadowedLocalLightPS"),SF_Pixel);
-
-IMPLEMENT_LOCAL_LIGHT_INJECTION_PIXELSHADER_TYPE(true, true, true);
-IMPLEMENT_LOCAL_LIGHT_INJECTION_PIXELSHADER_TYPE(true, false, true);
-IMPLEMENT_LOCAL_LIGHT_INJECTION_PIXELSHADER_TYPE(false, true, true);
-IMPLEMENT_LOCAL_LIGHT_INJECTION_PIXELSHADER_TYPE(false, false, true);
-IMPLEMENT_LOCAL_LIGHT_INJECTION_PIXELSHADER_TYPE(true, true, false);
-IMPLEMENT_LOCAL_LIGHT_INJECTION_PIXELSHADER_TYPE(true, false, false);
-IMPLEMENT_LOCAL_LIGHT_INJECTION_PIXELSHADER_TYPE(false, true, false);
-IMPLEMENT_LOCAL_LIGHT_INJECTION_PIXELSHADER_TYPE(false, false, false);
+IMPLEMENT_GLOBAL_SHADER(TInjectShadowedLocalLightPS, "/Engine/Private/VolumetricFog.usf", "InjectShadowedLocalLightPS", SF_Pixel);
 
 FProjectedShadowInfo* GetShadowForInjectionIntoVolumetricFog(const FLightSceneProxy* LightProxy, FVisibleLightInfo& VisibleLightInfo)
 {
@@ -407,46 +393,6 @@ FIntPoint CalculateVolumetricFogBoundsForLight(const FSphere& LightBounds, const
 	return VolumeZBounds;
 }
 
-template<bool bDynamicallyShadowed, bool bInverseSquared, bool bUseTemporalReprojection>
-void SetInjectShadowedLocalLightShaders(
-	FRHICommandList& RHICmdList, 
-	const FViewInfo& View, 
-	const FVolumetricFogIntegrationParameterData& IntegrationData,
-	const FLightSceneInfo* LightSceneInfo,
-	const FSphere& LightBounds,
-	const FExponentialHeightFogSceneInfo& FogInfo,
-	const FProjectedShadowInfo* ProjectedShadowInfo,
-	FIntVector VolumetricFogGridSize,
-	int32 MinZ)
-{
-	TShaderMapRef<FWriteToBoundingSphereVS> VertexShader(View.ShaderMap);
-	TOptionalShaderMapRef<FWriteToSliceGS> GeometryShader(View.ShaderMap);
-	TShaderMapRef<TInjectShadowedLocalLightPS<bDynamicallyShadowed, bInverseSquared, bUseTemporalReprojection>> PixelShader(View.ShaderMap);
-
-	FGraphicsPipelineStateInitializer GraphicsPSOInit;
-	RHICmdList.ApplyCachedRenderTargets(GraphicsPSOInit);
-	GraphicsPSOInit.RasterizerState = TStaticRasterizerState<FM_Solid, CM_None>::GetRHI();
-	GraphicsPSOInit.DepthStencilState = TStaticDepthStencilState<false, CF_Always>::GetRHI();
-	// Accumulate the contribution of multiple lights
-	GraphicsPSOInit.BlendState = TStaticBlendState<CW_RGBA, BO_Add, BF_One, BF_One, BO_Add, BF_Zero, BF_One>::GetRHI();
-
-	GraphicsPSOInit.BoundShaderState.VertexDeclarationRHI = GScreenVertexDeclaration.VertexDeclarationRHI;
-	GraphicsPSOInit.BoundShaderState.VertexShaderRHI = GETSAFERHISHADER_VERTEX(*VertexShader);
-	GraphicsPSOInit.BoundShaderState.GeometryShaderRHI = GETSAFERHISHADER_GEOMETRY(*GeometryShader);
-	GraphicsPSOInit.BoundShaderState.PixelShaderRHI = GETSAFERHISHADER_PIXEL(*PixelShader);
-	GraphicsPSOInit.PrimitiveType = PT_TriangleList;
-
-	SetGraphicsPipelineState(RHICmdList, GraphicsPSOInit);
-
-	PixelShader->SetParameters(RHICmdList, View, IntegrationData, LightSceneInfo, FogInfo, ProjectedShadowInfo);
-	VertexShader->SetParameters(RHICmdList, View, IntegrationData, LightBounds, MinZ);
-
-	if (GeometryShader.IsValid())
-	{
-		GeometryShader->SetParameters(RHICmdList, MinZ);
-	}
-}
-
 /**  */
 class FCircleRasterizeVertexBuffer : public FVertexBuffer
 {
@@ -517,7 +463,7 @@ public:
 TGlobalResource<FCircleRasterizeIndexBuffer> GCircleRasterizeIndexBuffer;
 
 void FDeferredShadingSceneRenderer::RenderLocalLightsForVolumetricFog(
-	FRHICommandListImmediate& RHICmdList,
+	FRenderGraphBuilder& GraphBuilder,
 	FViewInfo& View,
 	bool bUseTemporalReprojection,
 	const FVolumetricFogIntegrationParameterData& IntegrationData,
@@ -525,7 +471,7 @@ void FDeferredShadingSceneRenderer::RenderLocalLightsForVolumetricFog(
 	FIntVector VolumetricFogGridSize, 
 	FVector GridZParams,
 	const FPooledRenderTargetDesc& VolumeDesc,
-	TRefCountPtr<IPooledRenderTarget>& OutLocalShadowedLightScattering)
+	const FGraphTexture*& OutLocalShadowedLightScattering)
 {
 	TArray<const FLightSceneInfo*, SceneRenderingAllocator> LightsToInject;
 
@@ -550,94 +496,92 @@ void FDeferredShadingSceneRenderer::RenderLocalLightsForVolumetricFog(
 
 	if (LightsToInject.Num() > 0)
 	{
-		SCOPED_DRAW_EVENT(RHICmdList, ShadowedLights);
-
-		GRenderTargetPool.FindFreeElement(RHICmdList, VolumeDesc, OutLocalShadowedLightScattering, TEXT("LocalShadowedLightScattering"));
-
-		FRHIRenderTargetView ColorView(OutLocalShadowedLightScattering->GetRenderTargetItem().TargetableTexture, 0, -1, ERenderTargetLoadAction::EClear, ERenderTargetStoreAction::EStore);
-		FRHISetRenderTargetsInfo Info(1, &ColorView, FRHIDepthRenderTargetView());
-		RHICmdList.SetRenderTargetsAndClear(Info);
-
-		for (int32 LightIndex = 0; LightIndex < LightsToInject.Num(); LightIndex++)
-		{
-			const FLightSceneInfo* LightSceneInfo = LightsToInject[LightIndex];
-			FProjectedShadowInfo* ProjectedShadowInfo = GetShadowForInjectionIntoVolumetricFog(LightSceneInfo->Proxy, VisibleLightInfos[LightSceneInfo->Id]);
+		OutLocalShadowedLightScattering = GraphBuilder.CreateTexture(VolumeDesc, TEXT("LocalShadowedLightScattering"));
 
-			const bool bInverseSquared = LightSceneInfo->Proxy->IsInverseSquared();
-			const bool bDynamicallyShadowed = ProjectedShadowInfo != NULL;
-			const FSphere LightBounds = LightSceneInfo->Proxy->GetBoundingSphere();
-			const FIntPoint VolumeZBounds = CalculateVolumetricFogBoundsForLight(LightBounds, View, VolumetricFogGridSize, GridZParams);
+		FRenderTargetParameters* PassParameters;
+		GraphBuilder.CreateParameters(&PassParameters);
+		PassParameters->RenderTargets[0] = FRenderTargetBinding( OutLocalShadowedLightScattering, ERenderTargetLoadAction::EClear, ERenderTargetStoreAction::ENoAction );
 
-			if (VolumeZBounds.X < VolumeZBounds.Y)
+		GraphBuilder.AddPass(TEXT("LocalShadowedLightScattering"),
+			PassParameters,
+			ERenderGraphPassFlags::None,
+			[PassParameters, &View, this, LightsToInject, VolumetricFogGridSize, GridZParams, bUseTemporalReprojection, IntegrationData, FogInfo](FRHICommandListImmediate& RHICmdList)
 			{
-				if (bUseTemporalReprojection)
-				{
-					if (bDynamicallyShadowed)
-					{
-						if (bInverseSquared)
-						{
-							SetInjectShadowedLocalLightShaders<true, true, true>(RHICmdList, View, IntegrationData, LightSceneInfo, LightBounds, FogInfo, ProjectedShadowInfo, VolumetricFogGridSize, VolumeZBounds.X);
-						}
-						else
-						{
-							SetInjectShadowedLocalLightShaders<true, false, true>(RHICmdList, View, IntegrationData, LightSceneInfo, LightBounds, FogInfo, ProjectedShadowInfo, VolumetricFogGridSize, VolumeZBounds.X);
-						}
-					}
-					else
-					{
-						if (bInverseSquared)
-						{
-							SetInjectShadowedLocalLightShaders<false, true, true>(RHICmdList, View, IntegrationData, LightSceneInfo, LightBounds, FogInfo, ProjectedShadowInfo, VolumetricFogGridSize, VolumeZBounds.X);
-						}
-						else
-						{
-							SetInjectShadowedLocalLightShaders<false, false, true>(RHICmdList, View, IntegrationData, LightSceneInfo, LightBounds, FogInfo, ProjectedShadowInfo, VolumetricFogGridSize, VolumeZBounds.X);
-						}
-					}
-				}
-				else
+				SCOPED_DRAW_EVENT(RHICmdList, ShadowedLights);
+				for (int32 LightIndex = 0; LightIndex < LightsToInject.Num(); LightIndex++)
 				{
-					if (bDynamicallyShadowed)
-					{
-						if (bInverseSquared)
-						{
-							SetInjectShadowedLocalLightShaders<true, true, false>(RHICmdList, View, IntegrationData, LightSceneInfo, LightBounds, FogInfo, ProjectedShadowInfo, VolumetricFogGridSize, VolumeZBounds.X);
-						}
-						else
-						{
-							SetInjectShadowedLocalLightShaders<true, false, false>(RHICmdList, View, IntegrationData, LightSceneInfo, LightBounds, FogInfo, ProjectedShadowInfo, VolumetricFogGridSize, VolumeZBounds.X);
-						}
-					}
-					else
+					const FLightSceneInfo* LightSceneInfo = LightsToInject[LightIndex];
+					FProjectedShadowInfo* ProjectedShadowInfo = GetShadowForInjectionIntoVolumetricFog(LightSceneInfo->Proxy, VisibleLightInfos[LightSceneInfo->Id]);
+
+					const bool bInverseSquared = LightSceneInfo->Proxy->IsInverseSquared();
+					const bool bDynamicallyShadowed = ProjectedShadowInfo != NULL;
+					const FSphere LightBounds = LightSceneInfo->Proxy->GetBoundingSphere();
+					const FIntPoint VolumeZBounds = CalculateVolumetricFogBoundsForLight(LightBounds, View, VolumetricFogGridSize, GridZParams);
+
+					if (VolumeZBounds.X < VolumeZBounds.Y)
 					{
-						if (bInverseSquared)
-						{
-							SetInjectShadowedLocalLightShaders<false, true, false>(RHICmdList, View, IntegrationData, LightSceneInfo, LightBounds, FogInfo, ProjectedShadowInfo, VolumetricFogGridSize, VolumeZBounds.X);
-						}
-						else
+						TInjectShadowedLocalLightPS::FPermutationDomain PermutationVector;
+						PermutationVector.Set< TInjectShadowedLocalLightPS::FDynamicallyShadowed >( bDynamicallyShadowed );
+						PermutationVector.Set< TInjectShadowedLocalLightPS::FInverseSquared >( bInverseSquared );
+						PermutationVector.Set< TInjectShadowedLocalLightPS::FTemporalReprojection >( bUseTemporalReprojection );
+
+						auto VertexShader	= View.ShaderMap->GetShader< FWriteToBoundingSphereVS >();
+						auto GeometryShader	= View.ShaderMap->GetShader< FWriteToSliceGS >();
+						auto PixelShader	= View.ShaderMap->GetShader< TInjectShadowedLocalLightPS >( PermutationVector );
+
+						FGraphicsPipelineStateInitializer GraphicsPSOInit;
+						RHICmdList.ApplyCachedRenderTargets(GraphicsPSOInit);
+						GraphicsPSOInit.RasterizerState = TStaticRasterizerState<FM_Solid, CM_None>::GetRHI();
+						GraphicsPSOInit.DepthStencilState = TStaticDepthStencilState<false, CF_Always>::GetRHI();
+						// Accumulate the contribution of multiple lights
+						GraphicsPSOInit.BlendState = TStaticBlendState<CW_RGBA, BO_Add, BF_One, BF_One, BO_Add, BF_Zero, BF_One>::GetRHI();
+
+						GraphicsPSOInit.BoundShaderState.VertexDeclarationRHI = GScreenVertexDeclaration.VertexDeclarationRHI;
+						GraphicsPSOInit.BoundShaderState.VertexShaderRHI = GETSAFERHISHADER_VERTEX(VertexShader);
+						GraphicsPSOInit.BoundShaderState.GeometryShaderRHI = GETSAFERHISHADER_GEOMETRY(GeometryShader);
+						GraphicsPSOInit.BoundShaderState.PixelShaderRHI = GETSAFERHISHADER_PIXEL(PixelShader);
+						GraphicsPSOInit.PrimitiveType = PT_TriangleList;
+
+						SetGraphicsPipelineState(RHICmdList, GraphicsPSOInit);
+
+						PixelShader->SetParameters(RHICmdList, View, IntegrationData, LightSceneInfo, FogInfo, ProjectedShadowInfo, bDynamicallyShadowed);
+						VertexShader->SetParameters(RHICmdList, View, IntegrationData, LightBounds, VolumeZBounds.X);
+
+						if (GeometryShader)
 						{
-							SetInjectShadowedLocalLightShaders<false, false, false>(RHICmdList, View, IntegrationData, LightSceneInfo, LightBounds, FogInfo, ProjectedShadowInfo, VolumetricFogGridSize, VolumeZBounds.X);
+							GeometryShader->SetParameters(RHICmdList, VolumeZBounds.X);
 						}
+
+						RHICmdList.SetStreamSource(0, GCircleRasterizeVertexBuffer.VertexBufferRHI, 0);
+						const int32 NumInstances = VolumeZBounds.Y - VolumeZBounds.X;
+						const int32 NumTriangles = FCircleRasterizeVertexBuffer::NumVertices - 2;
+						RHICmdList.DrawIndexedPrimitive(GCircleRasterizeIndexBuffer.IndexBufferRHI, 0, 0, FCircleRasterizeVertexBuffer::NumVertices, 0, NumTriangles, NumInstances);
 					}
 				}
-
-				RHICmdList.SetStreamSource(0, GCircleRasterizeVertexBuffer.VertexBufferRHI, 0);
-				const int32 NumInstances = VolumeZBounds.Y - VolumeZBounds.X;
-				const int32 NumTriangles = FCircleRasterizeVertexBuffer::NumVertices - 2;
-				RHICmdList.DrawIndexedPrimitive(GCircleRasterizeIndexBuffer.IndexBufferRHI, 0, 0, FCircleRasterizeVertexBuffer::NumVertices, 0, NumTriangles, NumInstances);
-			}
-		}
-
-		RHICmdList.CopyToResolveTarget(OutLocalShadowedLightScattering->GetRenderTargetItem().TargetableTexture, OutLocalShadowedLightScattering->GetRenderTargetItem().ShaderResourceTexture, FResolveParams());
-	
-		GRenderTargetPool.VisualizeTexture.SetCheckPoint(RHICmdList, OutLocalShadowedLightScattering);
+			});
 	}
 }
 
-template<bool bTemporalReprojection, bool bDistanceFieldSkyOcclusion>
 class TVolumetricFogLightScatteringCS : public FGlobalShader
 {
-	DECLARE_SHADER_TYPE(TVolumetricFogLightScatteringCS,Global)
+	DECLARE_GLOBAL_SHADER(TVolumetricFogLightScatteringCS)
+
+	class FTemporalReprojection			: SHADER_PERMUTATION_BOOL("USE_TEMPORAL_REPROJECTION");
+	class FDistanceFieldSkyOcclusion	: SHADER_PERMUTATION_BOOL("DISTANCE_FIELD_SKY_OCCLUSION");
+
+	using FPermutationDomain = TShaderPermutationDomain<
+		FTemporalReprojection,
+		FDistanceFieldSkyOcclusion >;
+
+	BEGIN_SHADER_PARAMETER_STRUCT( FParameters, )
+		SHADER_PARAMETER_GRAPH_TEXTURE( Texture2D, VBufferA )
+		SHADER_PARAMETER_GRAPH_TEXTURE( Texture2D, VBufferB )
+		SHADER_PARAMETER_GRAPH_TEXTURE( Texture2D, LocalShadowedLightScattering )
+		SHADER_PARAMETER_GRAPH_TEXTURE( Texture2D, LightFunctionTexture )
+
+		SHADER_PARAMETER_GRAPH_UAV( RWTexture2D, RWLightScattering )
+	END_SHADER_PARAMETER_STRUCT()
+
 public:
 
 	static bool ShouldCompilePermutation(const FGlobalShaderPermutationParameters& Parameters)
@@ -651,13 +595,13 @@ public:
 		OutEnvironment.SetDefine(TEXT("THREADGROUP_SIZE"), VolumetricFogGridInjectionGroupSize);
 		FVolumetricFogIntegrationParameters::ModifyCompilationEnvironment(Parameters.Platform, OutEnvironment);
 		FForwardLightingParameters::ModifyCompilationEnvironment(Parameters.Platform, OutEnvironment);
-		OutEnvironment.SetDefine(TEXT("USE_TEMPORAL_REPROJECTION"), bTemporalReprojection);
-		OutEnvironment.SetDefine(TEXT("DISTANCE_FIELD_SKY_OCCLUSION"), bDistanceFieldSkyOcclusion);
 	}
 
 	TVolumetricFogLightScatteringCS(const ShaderMetaType::CompiledShaderInitializerType& Initializer)
 		: FGlobalShader(Initializer)
 	{
+		Bindings.Bind(Initializer.ParameterMap, *FParameters::FTypeInfo::GetStructMetadata());
+
 		LocalShadowedLightScattering.Bind(Initializer.ParameterMap, TEXT("LocalShadowedLightScattering"));
 		LightScatteringHistory.Bind(Initializer.ParameterMap, TEXT("LightScatteringHistory"));
 		LightScatteringHistorySampler.Bind(Initializer.ParameterMap, TEXT("LightScatteringHistorySampler"));
@@ -686,24 +630,13 @@ public:
 		const FViewInfo& View,
 		const FVolumetricFogIntegrationParameterData& IntegrationData,
 		const FExponentialHeightFogSceneInfo& FogInfo,
-		IPooledRenderTarget* LocalShadowedLightScatteringTarget,
 		FTextureRHIParamRef LightScatteringHistoryTexture,
 		bool bUseDirectionalLightShadowing,
-		const FMatrix& DirectionalLightFunctionWorldToShadowValue,
-		TRefCountPtr<IPooledRenderTarget>& LightFunctionTextureValue)
+		const FMatrix& DirectionalLightFunctionWorldToShadowValue)
 	{
 		FComputeShaderRHIParamRef ShaderRHI = GetComputeShader();
 		FGlobalShader::SetParameters<FViewUniformShaderParameters>(RHICmdList, ShaderRHI, View.ViewUniformBuffer);
 
-		FTextureRHIParamRef LocalShadowedLightScatteringTexture = GBlackVolumeTexture->TextureRHI;
-
-		if (LocalShadowedLightScatteringTarget)
-		{
-			LocalShadowedLightScatteringTexture = LocalShadowedLightScatteringTarget->GetRenderTargetItem().ShaderResourceTexture;
-		}
-
-		SetTextureParameter(RHICmdList, ShaderRHI, LocalShadowedLightScattering, LocalShadowedLightScatteringTexture);
-
 		if (!LightScatteringHistoryTexture)
 		{
 			LightScatteringHistoryTexture = GBlackVolumeTexture->TextureRHI;
@@ -722,14 +655,7 @@ public:
 
 		SetShaderValue(RHICmdList, ShaderRHI, DirectionalLightFunctionWorldToShadow, DirectionalLightFunctionWorldToShadowValue);
 
-		FTextureRHIParamRef LightFunctionRHITexture = GWhiteTexture->TextureRHI;
-
-		if (LightFunctionTextureValue)
-		{
-			LightFunctionRHITexture = LightFunctionTextureValue->GetRenderTargetItem().ShaderResourceTexture;
-		}
-
-		SetTextureParameter(RHICmdList, ShaderRHI, LightFunctionTexture, LightFunctionSampler, TStaticSamplerState<SF_Bilinear, AM_Clamp, AM_Clamp, AM_Clamp>::GetRHI(), LightFunctionRHITexture);
+		SetSamplerParameter(RHICmdList, ShaderRHI, LightFunctionSampler, TStaticSamplerState<SF_Bilinear, AM_Clamp, AM_Clamp, AM_Clamp>::GetRHI());
 
 		FScene* Scene = (FScene*)View.Family->Scene;
 		FDistanceFieldAOParameters AOParameterData(Scene->DefaultMaxDistanceFieldOcclusionDistance);
@@ -782,11 +708,6 @@ public:
 		SetUniformBufferParameterImmediate(RHICmdList, ShaderRHI, GetUniformBufferParameter<FFogUniformParameters>(), FogUniformParameters);
 	}
 
-	void UnsetParameters(FRHICommandList& RHICmdList, const FViewInfo& View, IPooledRenderTarget* LightScatteringRenderTarget)
-	{
-		VolumetricFogParameters.UnsetParameters(RHICmdList, GetComputeShader(), View, NULL, NULL, LightScatteringRenderTarget, true);
-	}
-
 	virtual bool Serialize(FArchive& Ar)
 	{		
 		bool bShaderHasOutdatedParameters = FGlobalShader::Serialize(Ar);
@@ -831,21 +752,18 @@ private:
 	FGlobalDistanceFieldParameters GlobalDistanceFieldParameters;
 };
 
-#define IMPLEMENT_VOLUMETRICFOG_LIGHT_SCATTERING_CS_TYPE(bTemporalReprojection, bDistanceFieldSkyOcclusion) \
-	typedef TVolumetricFogLightScatteringCS<bTemporalReprojection, bDistanceFieldSkyOcclusion> TVolumetricFogLightScatteringCS##bTemporalReprojection##bDistanceFieldSkyOcclusion; \
-	IMPLEMENT_SHADER_TYPE(template<>,TVolumetricFogLightScatteringCS##bTemporalReprojection##bDistanceFieldSkyOcclusion,TEXT("/Engine/Private/VolumetricFog.usf"),TEXT("LightScatteringCS"),SF_Compute);
-
-IMPLEMENT_VOLUMETRICFOG_LIGHT_SCATTERING_CS_TYPE(true, true)
-IMPLEMENT_VOLUMETRICFOG_LIGHT_SCATTERING_CS_TYPE(false, true)
-IMPLEMENT_VOLUMETRICFOG_LIGHT_SCATTERING_CS_TYPE(true, false)
-IMPLEMENT_VOLUMETRICFOG_LIGHT_SCATTERING_CS_TYPE(false, false)
+IMPLEMENT_GLOBAL_SHADER(TVolumetricFogLightScatteringCS, "/Engine/Private/VolumetricFog.usf", "LightScatteringCS", SF_Compute);
 
 uint32 VolumetricFogIntegrationGroupSize = 8;
 
 class FVolumetricFogFinalIntegrationCS : public FGlobalShader
 {
 	DECLARE_SHADER_TYPE(FVolumetricFogFinalIntegrationCS,Global)
-public:
+
+	BEGIN_SHADER_PARAMETER_STRUCT( FParameters, )
+		SHADER_PARAMETER_GRAPH_TEXTURE(	Texture3D<float4>,		LightScattering )
+		SHADER_PARAMETER_GRAPH_UAV(		RWTexture3D<float4>,	RWIntegratedLightScattering )
+	END_SHADER_PARAMETER_STRUCT()
 
 	static bool ShouldCompilePermutation(const FGlobalShaderPermutationParameters& Parameters)
 	{
@@ -862,6 +780,8 @@ public:
 	FVolumetricFogFinalIntegrationCS(const ShaderMetaType::CompiledShaderInitializerType& Initializer)
 		: FGlobalShader(Initializer)
 	{
+		Bindings.Bind(Initializer.ParameterMap, *FParameters::FTypeInfo::GetStructMetadata());
+
 		VolumetricFogParameters.Bind(Initializer.ParameterMap);
 	}
 
@@ -869,6 +789,7 @@ public:
 	{
 	}
 
+public:
 	void SetParameters(FRHICommandList& RHICmdList, const FViewInfo& View, const FVolumetricFogIntegrationParameterData& IntegrationData)
 	{
 		FComputeShaderRHIParamRef ShaderRHI = GetComputeShader();
@@ -876,11 +797,6 @@ public:
 		VolumetricFogParameters.Set(RHICmdList, ShaderRHI, View, IntegrationData);
 	}
 
-	void UnsetParameters(FRHICommandList& RHICmdList, const FViewInfo& View)
-	{
-		VolumetricFogParameters.UnsetParameters(RHICmdList, GetComputeShader(), View, NULL, NULL, NULL, true);
-	}
-
 	virtual bool Serialize(FArchive& Ar)
 	{		
 		bool bShaderHasOutdatedParameters = FGlobalShader::Serialize(Ar);
@@ -1045,6 +961,8 @@ void FDeferredShadingSceneRenderer::ComputeVolumetricFog(FRHICommandListImmediat
 {
 	if (ShouldRenderVolumetricFog())
 	{
+		QUICK_SCOPE_CYCLE_COUNTER(STAT_VolumetricFog);
+
 		const FExponentialHeightFogSceneInfo& FogInfo = Scene->ExponentialFogs[0];
 
 		for (int32 ViewIndex = 0; ViewIndex < Views.Num(); ViewIndex++)
@@ -1054,7 +972,8 @@ void FDeferredShadingSceneRenderer::ComputeVolumetricFog(FRHICommandListImmediat
 			const FIntVector VolumetricFogGridSize = GetVolumetricFogGridSize(View.ViewRect.Size());
 			const FVector GridZParams = GetVolumetricFogGridZParams(View.NearClippingDistance, FogInfo.VolumetricFogDistance, VolumetricFogGridSize.Z);
 
-			SCOPED_DRAW_EVENT(RHICmdList, VolumetricFog);
+			//@DW - graph todo
+			//SCOPED_DRAW_EVENT(RHICmdList, VolumetricFog);
 
 			const FVector FrameJitterOffsetValue = VolumetricFogTemporalRandom(View.Family->FrameNumber);
 
@@ -1077,15 +996,18 @@ void FDeferredShadingSceneRenderer::ComputeVolumetricFog(FRHICommandListImmediat
 				&& !View.bCameraCut
 				&& !View.bPrevTransformsReset
 				&& ViewFamily.bRealtimeUpdate
-				&& View.ViewState->LightScatteringHistory
-				&& View.ViewState->LightScatteringHistory->GetDesc().GetSize() == VolumetricFogGridSize;
+				&& View.ViewState->LightScatteringHistory;
 
 			FMatrix LightFunctionWorldToShadow;
-			TRefCountPtr<IPooledRenderTarget> LightFunctionTexture;
+
+			FRenderGraphBuilder GraphBuilder(RHICmdList);
+
+			//@DW - register GWhiteTexture as a graph external for when there's no light function - later a shader is going to bind it whether we rendered to it or not
+			const FGraphTexture* LightFunctionTexture = GraphBuilder.RegisterExternalTexture(GSystemTextures.WhiteDummy);
 			bool bUseDirectionalLightShadowing;
 
 			RenderLightFunctionForVolumetricFog(
-				RHICmdList,
+				GraphBuilder,
 				View,
 				VolumetricFogGridSize,
 				FogInfo.VolumetricFogDistance,
@@ -1093,168 +1015,168 @@ void FDeferredShadingSceneRenderer::ComputeVolumetricFog(FRHICommandListImmediat
 				LightFunctionTexture,
 				bUseDirectionalLightShadowing);
 
-			TRefCountPtr<IPooledRenderTarget> VBufferA;
-			TRefCountPtr<IPooledRenderTarget> VBufferB;
-
 			const uint32 Flags = TexCreate_ShaderResource | TexCreate_RenderTargetable | TexCreate_UAV | TexCreate_ReduceMemoryWithTilingMode;
 			FPooledRenderTargetDesc VolumeDesc(FPooledRenderTargetDesc::CreateVolumeDesc(VolumetricFogGridSize.X, VolumetricFogGridSize.Y, VolumetricFogGridSize.Z, PF_FloatRGBA, FClearValueBinding::Black, TexCreate_None, Flags, false));
 			FPooledRenderTargetDesc VolumeDescFastVRAM = VolumeDesc;
 			VolumeDescFastVRAM.Flags |= GFastVRamConfig.VolumetricFog;
-			GRenderTargetPool.FindFreeElement(RHICmdList, VolumeDescFastVRAM, VBufferA, TEXT("VBufferA"));
-			GRenderTargetPool.FindFreeElement(RHICmdList, VolumeDescFastVRAM, VBufferB, TEXT("VBufferB"));
-
-			IntegrationData.VBufferARenderTarget = VBufferA.GetReference();
-			IntegrationData.VBufferBRenderTarget = VBufferB.GetReference();
 
-			// Unbind render targets, the shadow depth target may still be bound
-			SetRenderTarget(RHICmdList, NULL, NULL);
+			//@DW - Explicit creation of graph resource handles
+			//@DW - Passing these around in a struct to ease manual wiring
+			IntegrationData.VBufferA = GraphBuilder.CreateTexture(VolumeDescFastVRAM, TEXT("VBufferA"));
+			IntegrationData.VBufferB = GraphBuilder.CreateTexture(VolumeDescFastVRAM, TEXT("VBufferB"));
+			IntegrationData.VBufferA_UAV = GraphBuilder.CreateUAV(FGraphUAVDesc(IntegrationData.VBufferA));
+			IntegrationData.VBufferB_UAV = GraphBuilder.CreateUAV(FGraphUAVDesc(IntegrationData.VBufferB));
 
 			{
-				const FIntVector NumGroups = FIntVector::DivideAndRoundUp(VolumetricFogGridSize, VolumetricFogGridInjectionGroupSize);
+				FVolumetricFogMaterialSetupCS::FParameters* PassParameters;
+				GraphBuilder.CreateParameters(&PassParameters);
+				PassParameters->GlobalAlbedo			= FogInfo.VolumetricFogAlbedo;
+				PassParameters->GlobalEmissive			= FogInfo.VolumetricFogEmissive;
+				PassParameters->GlobalExtinctionScale	= FogInfo.VolumetricFogExtinctionScale;
+				
+				PassParameters->RWVBufferA = IntegrationData.VBufferA_UAV;
+				PassParameters->RWVBufferB = IntegrationData.VBufferB_UAV;
+
+				FFogUniformParameters FogUniformParameters;
+				SetupFogUniformParameters(View, FogUniformParameters);
+				PassParameters->FogUniformParameters = CreateUniformBufferImmediate(FogUniformParameters, UniformBuffer_SingleDraw);
+				PassParameters->View = View.ViewUniformBuffer;
+
+				auto ComputeShader = View.ShaderMap->GetShader< FVolumetricFogMaterialSetupCS >();
+				ClearUnusedGraphResources(ComputeShader, PassParameters);
+
+				//@DW - this pass only reads external textures, we don't have any graph inputs
+				GraphBuilder.AddPass(TEXT("VolumetricFogMaterialSetup"),
+					PassParameters,
+					ERenderGraphPassFlags::Compute,
+					[PassParameters, &View, VolumetricFogGridSize, IntegrationData, ComputeShader](FRHICommandListImmediate& RHICmdList)
+					{
+						const FIntVector NumGroups = FIntVector::DivideAndRoundUp(VolumetricFogGridSize, VolumetricFogGridInjectionGroupSize);
+						SCOPED_DRAW_EVENT(RHICmdList, InitializeVolumeAttributes);
 
-				{
-					SCOPED_DRAW_EVENT(RHICmdList, InitializeVolumeAttributes);
-					TShaderMapRef<FVolumetricFogMaterialSetupCS> ComputeShader(View.ShaderMap);
-					RHICmdList.SetComputeShader(ComputeShader->GetComputeShader());
-					ComputeShader->SetParameters(RHICmdList, View, IntegrationData, FogInfo);
-					DispatchComputeShader(RHICmdList, *ComputeShader, NumGroups.X, NumGroups.Y, NumGroups.Z);
-					ComputeShader->UnsetParameters(RHICmdList, View, VBufferA.GetReference(), VBufferB.GetReference());
-				}
+						RHICmdList.SetComputeShader(ComputeShader->GetComputeShader());
 
-				FUnorderedAccessViewRHIParamRef VoxelizeUAVs[2] =
-				{
-					VBufferA->GetRenderTargetItem().UAV.GetReference(),
-					VBufferB->GetRenderTargetItem().UAV.GetReference()
-				};
+						ComputeShader->SetParameters(RHICmdList, View, IntegrationData);
 
-				RHICmdList.TransitionResources(EResourceTransitionAccess::EWritable, EResourceTransitionPipeline::EComputeToGfx, VoxelizeUAVs, ARRAY_COUNT(VoxelizeUAVs), nullptr);
+						SetShaderParameters( RHICmdList, ComputeShader, ComputeShader->GetComputeShader(), *PassParameters );
+						DispatchComputeShader(RHICmdList, ComputeShader, NumGroups.X, NumGroups.Y, NumGroups.Z);
+						UnsetShaderUAVs( RHICmdList, ComputeShader, ComputeShader->GetComputeShader() );
+					});
 
 				VoxelizeFogVolumePrimitives(
-					RHICmdList, 
+					GraphBuilder,
 					View, 
 					IntegrationData,
 					VolumetricFogGridSize, 
 					GridZParams,
 					FogInfo.VolumetricFogDistance);
-
-				FTextureRHIParamRef VoxelizeRenderTargets[2] =
-				{
-					VBufferA->GetRenderTargetItem().TargetableTexture,
-					VBufferB->GetRenderTargetItem().TargetableTexture
-				};
-
-				RHICmdList.TransitionResources(EResourceTransitionAccess::EReadable, VoxelizeRenderTargets, ARRAY_COUNT(VoxelizeRenderTargets));
-
-				GRenderTargetPool.VisualizeTexture.SetCheckPoint(RHICmdList, VBufferA);
-				GRenderTargetPool.VisualizeTexture.SetCheckPoint(RHICmdList, VBufferB);
 			}
 
-			TRefCountPtr<IPooledRenderTarget> LocalShadowedLightScattering = NULL;
-			RenderLocalLightsForVolumetricFog(RHICmdList, View, bUseTemporalReprojection, IntegrationData, FogInfo, VolumetricFogGridSize, GridZParams, VolumeDescFastVRAM, LocalShadowedLightScattering);
-
-			TRefCountPtr<IPooledRenderTarget> LightScattering;
-			GRenderTargetPool.FindFreeElement(RHICmdList, VolumeDesc, LightScattering, TEXT("LightScattering"));
-
-			IntegrationData.LightScatteringRenderTarget = LightScattering.GetReference();
+			const FGraphTexture* LocalShadowedLightScattering = GraphBuilder.RegisterExternalTexture(GSystemTextures.VolumetricBlackDummy);
+			RenderLocalLightsForVolumetricFog(GraphBuilder, View, bUseTemporalReprojection, IntegrationData, FogInfo, VolumetricFogGridSize, GridZParams, VolumeDescFastVRAM, LocalShadowedLightScattering);
 
-			SetRenderTarget(RHICmdList, NULL, NULL);
+			IntegrationData.LightScattering = GraphBuilder.CreateTexture(VolumeDesc, TEXT("LightScattering"));
+			IntegrationData.LightScatteringUAV = GraphBuilder.CreateUAV(FGraphUAVDesc(IntegrationData.LightScattering));
 
 			{
-				const FIntVector NumGroups = FIntVector::DivideAndRoundUp(VolumetricFogGridSize, VolumetricFogGridInjectionGroupSize);
+				TVolumetricFogLightScatteringCS::FParameters* PassParameters;
+				GraphBuilder.CreateParameters(&PassParameters);
+
+				PassParameters->VBufferA = IntegrationData.VBufferA;
+				PassParameters->VBufferB = IntegrationData.VBufferB;
+				PassParameters->LocalShadowedLightScattering = LocalShadowedLightScattering;
+				PassParameters->LightFunctionTexture = LightFunctionTexture;
+				PassParameters->RWLightScattering = IntegrationData.LightScatteringUAV;
 
 				const bool bUseGlobalDistanceField = UseGlobalDistanceField() && Scene->DistanceFieldSceneData.NumObjectsInBuffer > 0;
 
 				const bool bUseDistanceFieldSkyOcclusion =
 					ViewFamily.EngineShowFlags.AmbientOcclusion
 					&& Scene->SkyLight
-					&& Scene->SkyLight->bCastShadows 
+					&& Scene->SkyLight->bCastShadows
 					&& Scene->SkyLight->bCastVolumetricShadow
-					&& ShouldRenderDistanceFieldAO() 
+					&& ShouldRenderDistanceFieldAO()
 					&& SupportsDistanceFieldAO(View.GetFeatureLevel(), View.GetShaderPlatform())
 					&& bUseGlobalDistanceField
 					&& Views.Num() == 1
 					&& View.IsPerspectiveProjection();
 
-				SCOPED_DRAW_EVENTF(RHICmdList, LightScattering, TEXT("LightScattering %dx%dx%d %s %s"), 
-					VolumetricFogGridSize.X, 
-					VolumetricFogGridSize.Y, 
-					VolumetricFogGridSize.Z,
-					bUseDistanceFieldSkyOcclusion ? TEXT("DFAO") : TEXT(""),
-					LightFunctionTexture ? TEXT("LF") : TEXT(""));
+				TVolumetricFogLightScatteringCS::FPermutationDomain PermutationVector;
+				PermutationVector.Set< TVolumetricFogLightScatteringCS::FTemporalReprojection >(bUseTemporalReprojection);
+				PermutationVector.Set< TVolumetricFogLightScatteringCS::FDistanceFieldSkyOcclusion >(bUseDistanceFieldSkyOcclusion);
 
-				if (bUseTemporalReprojection)
-				{
-					FTextureRHIParamRef LightScatteringHistoryTexture = View.ViewState->LightScatteringHistory 
-						? View.ViewState->LightScatteringHistory->GetRenderTargetItem().ShaderResourceTexture
-						: GBlackVolumeTexture->TextureRHI;
+				auto ComputeShader = View.ShaderMap->GetShader< TVolumetricFogLightScatteringCS >(PermutationVector);
+				ClearUnusedGraphResources(ComputeShader, PassParameters);
 
-					if (bUseDistanceFieldSkyOcclusion)
-					{
-						TShaderMapRef<TVolumetricFogLightScatteringCS<true, true> > ComputeShader(View.ShaderMap);
-						RHICmdList.SetComputeShader(ComputeShader->GetComputeShader());
-						ComputeShader->SetParameters(RHICmdList, View, IntegrationData, FogInfo, LocalShadowedLightScattering.GetReference(), LightScatteringHistoryTexture, bUseDirectionalLightShadowing, LightFunctionWorldToShadow, LightFunctionTexture);
-						DispatchComputeShader(RHICmdList, *ComputeShader, NumGroups.X, NumGroups.Y, NumGroups.Z);
-						ComputeShader->UnsetParameters(RHICmdList, View, LightScattering.GetReference());
-					}
-					else
-					{
-						TShaderMapRef<TVolumetricFogLightScatteringCS<true, false> > ComputeShader(View.ShaderMap);
-						RHICmdList.SetComputeShader(ComputeShader->GetComputeShader());
-						ComputeShader->SetParameters(RHICmdList, View, IntegrationData, FogInfo, LocalShadowedLightScattering.GetReference(), LightScatteringHistoryTexture, bUseDirectionalLightShadowing, LightFunctionWorldToShadow, LightFunctionTexture);
-						DispatchComputeShader(RHICmdList, *ComputeShader, NumGroups.X, NumGroups.Y, NumGroups.Z);
-						ComputeShader->UnsetParameters(RHICmdList, View, LightScattering.GetReference());
-					}
-				}
-				else
-				{
-					if (bUseDistanceFieldSkyOcclusion)
+				GraphBuilder.AddPass(TEXT("LightScattering"),
+					PassParameters,
+					ERenderGraphPassFlags::Compute,
+					[PassParameters, ComputeShader, &View, this, FogInfo, bUseTemporalReprojection, VolumetricFogGridSize, IntegrationData, bUseDirectionalLightShadowing, bUseDistanceFieldSkyOcclusion, LightFunctionWorldToShadow](FRHICommandListImmediate& RHICmdList)
 					{
-						TShaderMapRef<TVolumetricFogLightScatteringCS<false, true> > ComputeShader(View.ShaderMap);
-						RHICmdList.SetComputeShader(ComputeShader->GetComputeShader());
-						ComputeShader->SetParameters(RHICmdList, View, IntegrationData, FogInfo, LocalShadowedLightScattering.GetReference(), NULL, bUseDirectionalLightShadowing, LightFunctionWorldToShadow, LightFunctionTexture);
-						DispatchComputeShader(RHICmdList, *ComputeShader, NumGroups.X, NumGroups.Y, NumGroups.Z);
-						ComputeShader->UnsetParameters(RHICmdList, View, LightScattering.GetReference());
-					}
-					else
-					{
-						TShaderMapRef<TVolumetricFogLightScatteringCS<false, false> > ComputeShader(View.ShaderMap);
+						SetRenderTarget(RHICmdList, NULL, NULL);
+						const FIntVector NumGroups = FIntVector::DivideAndRoundUp(VolumetricFogGridSize, VolumetricFogGridInjectionGroupSize);
+
+						SCOPED_DRAW_EVENTF(RHICmdList, LightScattering, TEXT("LightScattering %dx%dx%d %s %s"), 
+							VolumetricFogGridSize.X, 
+							VolumetricFogGridSize.Y, 
+							VolumetricFogGridSize.Z,
+							bUseDistanceFieldSkyOcclusion ? TEXT("DFAO") : TEXT(""),
+							PassParameters->LightFunctionTexture ? TEXT("LF") : TEXT(""));
+
 						RHICmdList.SetComputeShader(ComputeShader->GetComputeShader());
-						ComputeShader->SetParameters(RHICmdList, View, IntegrationData, FogInfo, LocalShadowedLightScattering.GetReference(), NULL, bUseDirectionalLightShadowing, LightFunctionWorldToShadow, LightFunctionTexture);
-						DispatchComputeShader(RHICmdList, *ComputeShader, NumGroups.X, NumGroups.Y, NumGroups.Z);
-						ComputeShader->UnsetParameters(RHICmdList, View, LightScattering.GetReference());
-					}
-				}
 
-				GRenderTargetPool.VisualizeTexture.SetCheckPoint(RHICmdList, LightScattering);
+						FTextureRHIParamRef LightScatteringHistoryTexture = View.ViewState->LightScatteringHistory.IsValid()
+							? View.ViewState->LightScatteringHistory->GetRenderTargetItem().ShaderResourceTexture
+							: GBlackVolumeTexture->TextureRHI;
 
-				if (bUseTemporalReprojection)
-				{
-					View.ViewState->LightScatteringHistory = LightScattering;
-				}
-				else if (View.ViewState)
-				{
-					View.ViewState->LightScatteringHistory = NULL;
-				}
+						ComputeShader->SetParameters(RHICmdList, View, IntegrationData, FogInfo, LightScatteringHistoryTexture, bUseDirectionalLightShadowing, LightFunctionWorldToShadow);
+						
+						SetShaderParameters( RHICmdList, ComputeShader, ComputeShader->GetComputeShader(), *PassParameters );
+						DispatchComputeShader(RHICmdList, ComputeShader, NumGroups.X, NumGroups.Y, NumGroups.Z);
+						UnsetShaderUAVs( RHICmdList, ComputeShader, ComputeShader->GetComputeShader() );
+					});
 			}
 
-			VBufferA = NULL;
-			VBufferB = NULL;
-			LightFunctionTexture = NULL;
+			const FGraphTexture* IntegratedLightScattering = GraphBuilder.CreateTexture(VolumeDesc, TEXT("IntegratedLightScattering"));
+			const FGraphUAV* IntegratedLightScatteringUAV = GraphBuilder.CreateUAV(FGraphUAVDesc(IntegratedLightScattering));
+
+			{
+				FVolumetricFogFinalIntegrationCS::FParameters* PassParameters;
+				GraphBuilder.CreateParameters(&PassParameters);
+				PassParameters->LightScattering = IntegrationData.LightScattering;
+				PassParameters->RWIntegratedLightScattering = IntegratedLightScatteringUAV;
+
+				GraphBuilder.AddPass(TEXT("FinalIntegration"),
+					PassParameters,
+					ERenderGraphPassFlags::Compute,
+					[PassParameters, &View, VolumetricFogGridSize, IntegrationData, this](FRHICommandListImmediate& RHICmdList)
+					{
+						SCOPED_DRAW_EVENT(RHICmdList, FinalIntegration);
 
-			GRenderTargetPool.FindFreeElement(RHICmdList, VolumeDesc, View.VolumetricFogResources.IntegratedLightScattering, TEXT("IntegratedLightScattering"));
+						const FIntVector NumGroups = FIntVector::DivideAndRoundUp(VolumetricFogGridSize, VolumetricFogIntegrationGroupSize);
 
+						auto ComputeShader = View.ShaderMap->GetShader< FVolumetricFogFinalIntegrationCS >();
+						RHICmdList.SetComputeShader(ComputeShader->GetComputeShader());
+						ComputeShader->SetParameters(RHICmdList, View, IntegrationData);
+
+						SetShaderParameters( RHICmdList, ComputeShader, ComputeShader->GetComputeShader(), *PassParameters );
+						DispatchComputeShader(RHICmdList, ComputeShader, NumGroups.X, NumGroups.Y, 1);
+						UnsetShaderUAVs( RHICmdList, ComputeShader, ComputeShader->GetComputeShader() );
+					});
+			}
+
+			GraphBuilder.GetInternalTexture(IntegratedLightScattering, &View.VolumetricFogResources.IntegratedLightScattering);
+
+			if (bUseTemporalReprojection)
+			{
+				GraphBuilder.GetInternalTexture(IntegrationData.LightScattering, &View.ViewState->LightScatteringHistory);
+			}
+			else if (View.ViewState)
 			{
-				SCOPED_DRAW_EVENT(RHICmdList, FinalIntegration);
-
-				const FIntVector NumGroups = FIntVector::DivideAndRoundUp(VolumetricFogGridSize, VolumetricFogIntegrationGroupSize);
-				TShaderMapRef<FVolumetricFogFinalIntegrationCS> ComputeShader(View.ShaderMap);
-				RHICmdList.SetComputeShader(ComputeShader->GetComputeShader());
-				ComputeShader->SetParameters(RHICmdList, View, IntegrationData);
-				DispatchComputeShader(RHICmdList, *ComputeShader, NumGroups.X, NumGroups.Y, 1);
-				ComputeShader->UnsetParameters(RHICmdList, View);
+				View.ViewState->LightScatteringHistory = NULL;
 			}
 
-			GRenderTargetPool.VisualizeTexture.SetCheckPoint(RHICmdList, View.VolumetricFogResources.IntegratedLightScattering);
+			GraphBuilder.Execute();
 		}
 	}
 }
diff --git a/Engine/Source/Runtime/Renderer/Private/VolumetricFogLightFunction.cpp b/Engine/Source/Runtime/Renderer/Private/VolumetricFogLightFunction.cpp
--- a/Engine/Source/Runtime/Renderer/Private/VolumetricFogLightFunction.cpp
+++ b/Engine/Source/Runtime/Renderer/Private/VolumetricFogLightFunction.cpp
@@ -97,16 +97,15 @@ private:
 IMPLEMENT_MATERIAL_SHADER_TYPE(,FVolumetricFogLightFunctionPS,TEXT("/Engine/Private/VolumetricFogLightFunction.usf"),TEXT("Main"),SF_Pixel);
 
 void FDeferredShadingSceneRenderer::RenderLightFunctionForVolumetricFog(
-	FRHICommandListImmediate& RHICmdList, 
+	FRenderGraphBuilder& GraphBuilder,
 	FViewInfo& View, 
 	FIntVector VolumetricFogGridSize,
 	float VolumetricFogMaxDistance,
 	FMatrix& OutLightFunctionWorldToShadow,
-	TRefCountPtr<IPooledRenderTarget>& OutLightFunctionTexture,
+	const FGraphTexture*& OutLightFunctionTexture,
 	bool& bOutUseDirectionalLightShadowing)
 {
 	OutLightFunctionWorldToShadow = FMatrix::Identity;
-	OutLightFunctionTexture = NULL;
 	bOutUseDirectionalLightShadowing = true;
 
 	FLightSceneInfo* DirectionalLightSceneInfo = NULL;
@@ -191,52 +190,58 @@ void FDeferredShadingSceneRenderer::RenderLightFunctionForVolumetricFog(
 
 		if (MaterialProxy && MaterialProxy->GetMaterial(Scene->GetFeatureLevel())->IsLightFunction())
 		{
-			FPooledRenderTargetDesc LightFunctionTextureDesc(FPooledRenderTargetDesc::Create2DDesc(LightFunctionResolution, PF_G8, FClearValueBinding::None, TexCreate_None, TexCreate_ShaderResource | TexCreate_RenderTargetable, false));
+			FPooledRenderTargetDesc LightFunctionTextureDesc = FPooledRenderTargetDesc::Create2DDesc(LightFunctionResolution, PF_G8, FClearValueBinding::None, TexCreate_None, TexCreate_ShaderResource | TexCreate_RenderTargetable, false);
 			LightFunctionTextureDesc.Flags |= GFastVRamConfig.VolumetricFog;
-			GRenderTargetPool.FindFreeElement(RHICmdList, LightFunctionTextureDesc, OutLightFunctionTexture, TEXT("VolumetricFogLightFunction"));
-			
-			const FMatrix WorldToShadowValue = FTranslationMatrix(ProjectedShadowInfo.PreShadowTranslation) * ProjectedShadowInfo.SubjectAndReceiverMatrix;
-			OutLightFunctionWorldToShadow = WorldToShadowValue;
-
-			const FMaterial* Material = MaterialProxy->GetMaterial(Scene->GetFeatureLevel());
-			SCOPED_DRAW_EVENTF(RHICmdList, LightFunction, TEXT("LightFunction %ux%u Material=%s"), LightFunctionResolution.X, LightFunctionResolution.Y, *Material->GetFriendlyName());
-
-			SetRenderTarget(RHICmdList, OutLightFunctionTexture->GetRenderTargetItem().TargetableTexture, NULL, ESimpleRenderTargetMode::EExistingColorAndDepth, FExclusiveDepthStencil::DepthNop_StencilNop, true);
-			RHICmdList.SetViewport(0, 0, 0.0f, LightFunctionResolution.X, LightFunctionResolution.Y, 1.0f);
-
-			FGraphicsPipelineStateInitializer GraphicsPSOInit;
-			RHICmdList.ApplyCachedRenderTargets(GraphicsPSOInit);
-			GraphicsPSOInit.RasterizerState = TStaticRasterizerState<FM_Solid, CM_None>::GetRHI();
-			GraphicsPSOInit.BlendState = TStaticBlendState<>::GetRHI();
-			GraphicsPSOInit.DepthStencilState = TStaticDepthStencilState<false, CF_Always>::GetRHI();
-			GraphicsPSOInit.PrimitiveType = PT_TriangleList;
-
-			const FMaterialShaderMap* MaterialShaderMap = Material->GetRenderingThreadShaderMap();
-			TShaderMapRef<FPostProcessVS> VertexShader(View.ShaderMap);
-			FVolumetricFogLightFunctionPS* PixelShader = MaterialShaderMap->GetShader<FVolumetricFogLightFunctionPS>();
 
-			GraphicsPSOInit.BoundShaderState.VertexDeclarationRHI = GFilterVertexDeclaration.VertexDeclarationRHI;
-			GraphicsPSOInit.BoundShaderState.VertexShaderRHI = GETSAFERHISHADER_VERTEX(*VertexShader);
-			GraphicsPSOInit.BoundShaderState.PixelShaderRHI = GETSAFERHISHADER_PIXEL(PixelShader);
+			OutLightFunctionTexture = GraphBuilder.CreateTexture(LightFunctionTextureDesc, TEXT("VolumetricFogLightFunction"));
 
-			SetGraphicsPipelineState(RHICmdList, GraphicsPSOInit);
-
-			VertexShader->SetParameters(RHICmdList, View.ViewUniformBuffer);
-			PixelShader->SetParameters(RHICmdList, View, DirectionalLightSceneInfo, MaterialProxy, FVector2D(1.0f / LightFunctionResolution.X, 1.0f / LightFunctionResolution.Y), WorldToShadowValue.Inverse());
-
-			DrawRectangle( 
-				RHICmdList,
-				0, 0, 
-				LightFunctionResolution.X, LightFunctionResolution.Y,
-				0, 0, 
-				LightFunctionResolution.X, LightFunctionResolution.Y,
-				LightFunctionResolution,
-				LightFunctionResolution,
-				*VertexShader);
-
-			RHICmdList.TransitionResource(EResourceTransitionAccess::EReadable, OutLightFunctionTexture->GetRenderTargetItem().TargetableTexture);
+			const FMatrix WorldToShadowValue = FTranslationMatrix(ProjectedShadowInfo.PreShadowTranslation) * ProjectedShadowInfo.SubjectAndReceiverMatrix;
+			OutLightFunctionWorldToShadow = WorldToShadowValue;
 
-			GRenderTargetPool.VisualizeTexture.SetCheckPoint(RHICmdList, OutLightFunctionTexture);
+			FRenderTargetParameters* PassParameters;
+			GraphBuilder.CreateParameters(&PassParameters);
+			PassParameters->RenderTargets[0] = FRenderTargetBinding( OutLightFunctionTexture, ERenderTargetLoadAction::ENoAction, ERenderTargetStoreAction::EStore );
+
+			GraphBuilder.AddPass(TEXT("VolumetricFogLightFunction"),
+				PassParameters,
+				ERenderGraphPassFlags::None,
+				[PassParameters, &View, MaterialProxy, LightFunctionResolution, DirectionalLightSceneInfo, WorldToShadowValue, this](FRHICommandListImmediate& RHICmdList)
+				{
+					const FMaterial* Material = MaterialProxy->GetMaterial(Scene->GetFeatureLevel());
+					SCOPED_DRAW_EVENTF(RHICmdList, LightFunction, TEXT("LightFunction %ux%u Material=%s"), LightFunctionResolution.X, LightFunctionResolution.Y, *Material->GetFriendlyName());
+
+					RHICmdList.SetViewport(0, 0, 0.0f, LightFunctionResolution.X, LightFunctionResolution.Y, 1.0f);
+
+					FGraphicsPipelineStateInitializer GraphicsPSOInit;
+					RHICmdList.ApplyCachedRenderTargets(GraphicsPSOInit);
+					GraphicsPSOInit.RasterizerState = TStaticRasterizerState<FM_Solid, CM_None>::GetRHI();
+					GraphicsPSOInit.BlendState = TStaticBlendState<>::GetRHI();
+					GraphicsPSOInit.DepthStencilState = TStaticDepthStencilState<false, CF_Always>::GetRHI();
+					GraphicsPSOInit.PrimitiveType = PT_TriangleList;
+
+					const FMaterialShaderMap* MaterialShaderMap = Material->GetRenderingThreadShaderMap();
+					TShaderMapRef<FPostProcessVS> VertexShader(View.ShaderMap);
+					FVolumetricFogLightFunctionPS* PixelShader = MaterialShaderMap->GetShader<FVolumetricFogLightFunctionPS>();
+
+					GraphicsPSOInit.BoundShaderState.VertexDeclarationRHI = GFilterVertexDeclaration.VertexDeclarationRHI;
+					GraphicsPSOInit.BoundShaderState.VertexShaderRHI = GETSAFERHISHADER_VERTEX(*VertexShader);
+					GraphicsPSOInit.BoundShaderState.PixelShaderRHI = GETSAFERHISHADER_PIXEL(PixelShader);
+
+					SetGraphicsPipelineState(RHICmdList, GraphicsPSOInit);
+
+					VertexShader->SetParameters(RHICmdList, View.ViewUniformBuffer);
+					PixelShader->SetParameters(RHICmdList, View, DirectionalLightSceneInfo, MaterialProxy, FVector2D(1.0f / LightFunctionResolution.X, 1.0f / LightFunctionResolution.Y), WorldToShadowValue.Inverse());
+
+					DrawRectangle(
+						RHICmdList,
+						0, 0,
+						LightFunctionResolution.X, LightFunctionResolution.Y,
+						0, 0,
+						LightFunctionResolution.X, LightFunctionResolution.Y,
+						LightFunctionResolution,
+						LightFunctionResolution,
+						*VertexShader);
+				});
 		}
 	}
 }
diff --git a/Engine/Source/Runtime/Renderer/Private/VolumetricFogShared.h b/Engine/Source/Runtime/Renderer/Private/VolumetricFogShared.h
--- a/Engine/Source/Runtime/Renderer/Private/VolumetricFogShared.h
+++ b/Engine/Source/Runtime/Renderer/Private/VolumetricFogShared.h
@@ -15,17 +15,18 @@ extern FVector VolumetricFogTemporalRandom(uint32 FrameNumber);
 struct FVolumetricFogIntegrationParameterData
 {
 	FVolumetricFogIntegrationParameterData() :
-		bTemporalHistoryIsValid(false),
-		VBufferARenderTarget(NULL),
-		VBufferBRenderTarget(NULL),
-		LightScatteringRenderTarget(NULL)
+		bTemporalHistoryIsValid(false)
 	{}
 
 	bool bTemporalHistoryIsValid;
 	TArray<FVector4, TInlineAllocator<16>> FrameJitterOffsetValues;
-	IPooledRenderTarget* VBufferARenderTarget;
-	IPooledRenderTarget* VBufferBRenderTarget;
-	IPooledRenderTarget* LightScatteringRenderTarget;
+	const FGraphTexture* VBufferA;
+	const FGraphTexture* VBufferB;
+	const FGraphUAV* VBufferA_UAV;
+	const FGraphUAV* VBufferB_UAV;
+
+	const FGraphTexture* LightScattering;
+	const FGraphUAV* LightScatteringUAV;
 };
 
 /**  */
@@ -59,31 +60,6 @@ public:
 		const FViewInfo& View, 
 		const FVolumetricFogIntegrationParameterData& IntegrationData) const
 	{
-		if (VBufferA.IsBound())
-		{
-			const FSceneRenderTargetItem& VBufferAItem = IntegrationData.VBufferARenderTarget->GetRenderTargetItem();
-			VBufferA.SetTexture(RHICmdList, ShaderRHI, VBufferAItem.ShaderResourceTexture, VBufferAItem.UAV);
-		}
-
-		if (VBufferB.IsBound())
-		{
-			const FSceneRenderTargetItem& VBufferBItem = IntegrationData.VBufferBRenderTarget->GetRenderTargetItem();
-			VBufferB.SetTexture(RHICmdList, ShaderRHI, VBufferBItem.ShaderResourceTexture, VBufferBItem.UAV);
-		}
-
-		if (LightScattering.IsBound())
-		{
-			const FSceneRenderTargetItem& LightScatteringItem = IntegrationData.LightScatteringRenderTarget->GetRenderTargetItem();
-			LightScattering.SetTexture(RHICmdList, ShaderRHI, LightScatteringItem.ShaderResourceTexture, LightScatteringItem.UAV);
-		}
-
-		if (IntegratedLightScattering.IsBound())
-		{
-			const FSceneRenderTargetItem& IntegratedLightScatteringItem = View.VolumetricFogResources.IntegratedLightScattering->GetRenderTargetItem();
-
-			IntegratedLightScattering.SetTexture(RHICmdList, ShaderRHI, IntegratedLightScatteringItem.ShaderResourceTexture, IntegratedLightScatteringItem.UAV);
-		}
-
 		SetSamplerParameter(RHICmdList, ShaderRHI, IntegratedLightScatteringSampler, TStaticSamplerState<SF_Bilinear, AM_Clamp, AM_Clamp, AM_Clamp>::GetRHI());
 		
 		if (VolumetricFogData.IsBound())
@@ -115,51 +91,6 @@ public:
 		SetShaderValue(RHICmdList, ShaderRHI, HistoryMissSuperSampleCount, FMath::Clamp(GVolumetricFogHistoryMissSupersampleCount, 1, 16));
 	}
 
-	template<typename ShaderRHIParamRef>
-	void UnsetParameters(
-		FRHICommandList& RHICmdList, 
-		const ShaderRHIParamRef& ShaderRHI, 
-		const FViewInfo& View, 
-		IPooledRenderTarget* VBufferARenderTarget, 
-		IPooledRenderTarget* VBufferBRenderTarget, 
-		IPooledRenderTarget* LightScatteringRenderTarget, 
-		bool bTransitionToReadable)
-	{
-		VBufferA.UnsetUAV(RHICmdList, ShaderRHI);
-		VBufferB.UnsetUAV(RHICmdList, ShaderRHI);
-		LightScattering.UnsetUAV(RHICmdList, ShaderRHI);
-		IntegratedLightScattering.UnsetUAV(RHICmdList, ShaderRHI);
-
-		const FSceneRenderTargetItem& IntegratedLightScatteringItem = View.VolumetricFogResources.IntegratedLightScattering->GetRenderTargetItem();
-
-		TArray<FUnorderedAccessViewRHIParamRef, TInlineAllocator<4>> OutUAVs;
-
-		if (VBufferA.IsUAVBound())
-		{
-			OutUAVs.Add(VBufferARenderTarget->GetRenderTargetItem().UAV);
-		}
-
-		if (VBufferB.IsUAVBound())
-		{
-			OutUAVs.Add(VBufferBRenderTarget->GetRenderTargetItem().UAV);
-		}
-
-		if (LightScattering.IsUAVBound())
-		{
-			OutUAVs.Add(LightScatteringRenderTarget->GetRenderTargetItem().UAV);
-		}
-
-		if (IntegratedLightScattering.IsUAVBound())
-		{
-			OutUAVs.Add(IntegratedLightScatteringItem.UAV);
-		}
-
-		if (OutUAVs.Num() > 0 && bTransitionToReadable)
-		{
-			RHICmdList.TransitionResources(EResourceTransitionAccess::EReadable, EResourceTransitionPipeline::EComputeToGfx, OutUAVs.GetData(), OutUAVs.Num());
-		}
-	}
-
 	/** Serializer. */
 	friend FArchive& operator<<(FArchive& Ar,FVolumetricFogIntegrationParameters& P)
 	{
diff --git a/Engine/Source/Runtime/Renderer/Private/VolumetricFogVoxelization.cpp b/Engine/Source/Runtime/Renderer/Private/VolumetricFogVoxelization.cpp
--- a/Engine/Source/Runtime/Renderer/Private/VolumetricFogVoxelization.cpp
+++ b/Engine/Source/Runtime/Renderer/Private/VolumetricFogVoxelization.cpp
@@ -123,7 +123,7 @@ class FVoxelizeVolumeVS : public FMeshMaterialShader
 		VoxelizationPassIndex.Bind(Initializer.ParameterMap, TEXT("VoxelizationPassIndex"));
 		ViewToVolumeClip.Bind(Initializer.ParameterMap, TEXT("ViewToVolumeClip"));
 		VolumetricFogParameters.Bind(Initializer.ParameterMap);
-		PassUniformBuffer.Bind(Initializer.ParameterMap, FSceneTexturesUniformParameters::StaticStruct.GetShaderVariableName());
+		PassUniformBuffer.Bind(Initializer.ParameterMap, FSceneTexturesUniformParameters::StaticStructMetadata.GetShaderVariableName());
 	}
 
 	FVoxelizeVolumeVS()
@@ -250,7 +250,7 @@ protected:
 		VoxelizationPassIndex.Bind(Initializer.ParameterMap, TEXT("VoxelizationPassIndex"));
 		ViewToVolumeClip.Bind(Initializer.ParameterMap, TEXT("ViewToVolumeClip"));
 		VolumetricFogParameters.Bind(Initializer.ParameterMap);
-		PassUniformBuffer.Bind(Initializer.ParameterMap, FSceneTexturesUniformParameters::StaticStruct.GetShaderVariableName());
+		PassUniformBuffer.Bind(Initializer.ParameterMap, FSceneTexturesUniformParameters::StaticStructMetadata.GetShaderVariableName());
 	}
 
 	FVoxelizeVolumeGS()
@@ -360,7 +360,7 @@ protected:
 		:	FMeshMaterialShader(Initializer)
 	{
 		VolumetricFogParameters.Bind(Initializer.ParameterMap);
-		PassUniformBuffer.Bind(Initializer.ParameterMap, FSceneTexturesUniformParameters::StaticStruct.GetShaderVariableName());
+		PassUniformBuffer.Bind(Initializer.ParameterMap, FSceneTexturesUniformParameters::StaticStructMetadata.GetShaderVariableName());
 	}
 
 	FVoxelizeVolumePS()
@@ -681,7 +681,7 @@ void VoxelizeVolumePrimitive(
 }
 
 void FDeferredShadingSceneRenderer::VoxelizeFogVolumePrimitives(
-	FRHICommandListImmediate& RHICmdList, 
+	FRenderGraphBuilder& GraphBuilder,
 	const FViewInfo& View,
 	const FVolumetricFogIntegrationParameterData& IntegrationData,
 	FIntVector VolumetricFogGridSize,
@@ -691,78 +691,81 @@ void FDeferredShadingSceneRenderer::VoxelizeFogVolumePrimitives(
 	if (View.VolumetricPrimSet.NumPrims() > 0
 		&& DoesPlatformSupportVolumetricFogVoxelization(View.GetShaderPlatform()))
 	{
-		SCOPED_DRAW_EVENT(RHICmdList, VoxelizeVolumePrimitives);
+		FRenderTargetParameters* PassParameters;
+		GraphBuilder.CreateParameters(&PassParameters);
+		PassParameters->RenderTargets[0] = FRenderTargetBinding( IntegrationData.VBufferA, ERenderTargetLoadAction::ENoAction, ERenderTargetStoreAction::EStore );
+		PassParameters->RenderTargets[1] = FRenderTargetBinding( IntegrationData.VBufferB, ERenderTargetLoadAction::ENoAction, ERenderTargetStoreAction::EStore );
 
-		FViewUniformShaderParameters VoxelizeParameters = *View.CachedViewUniformShaderParameters;
+		GraphBuilder.AddPass(TEXT("VoxelizeVolumePrimitives"),
+			PassParameters,
+			ERenderGraphPassFlags::None,
+			[PassParameters, &View, VolumetricFogGridSize, IntegrationData, VolumetricFogDistance, GridZParams](FRHICommandListImmediate& RHICmdList)
+			{
+				SCOPED_DRAW_EVENT(RHICmdList, VoxelizeVolumePrimitives);
 
-		// Update the parts of VoxelizeParameters which are dependent on the buffer size and view rect
-		View.SetupViewRectUniformBufferParameters(
-			VoxelizeParameters,
-			FIntPoint(VolumetricFogGridSize.X, VolumetricFogGridSize.Y),
-			FIntRect(0, 0, VolumetricFogGridSize.X, VolumetricFogGridSize.Y),
-			View.ViewMatrices,
-			View.PrevViewInfo.ViewMatrices
-		);
+				FViewUniformShaderParameters VoxelizeParameters = *View.CachedViewUniformShaderParameters;
 
-		TUniformBufferRef<FViewUniformShaderParameters> VoxelizeViewUniformBuffer = TUniformBufferRef<FViewUniformShaderParameters>::CreateUniformBufferImmediate(VoxelizeParameters, UniformBuffer_SingleFrame);
+				// Update the parts of VoxelizeParameters which are dependent on the buffer size and view rect
+				View.SetupViewRectUniformBufferParameters(
+					VoxelizeParameters,
+					FIntPoint(VolumetricFogGridSize.X, VolumetricFogGridSize.Y),
+					FIntRect(0, 0, VolumetricFogGridSize.X, VolumetricFogGridSize.Y),
+					View.ViewMatrices,
+					View.PrevViewInfo.ViewMatrices
+				);
 
-		FSceneTexturesUniformParameters SceneTextureParameters;
-		FSceneRenderTargets& SceneContext = FSceneRenderTargets::Get(RHICmdList);
-		SetupSceneTextureUniformParameters(SceneContext, View.FeatureLevel, ESceneTextureSetupMode::None, SceneTextureParameters);
-		TUniformBufferRef<FSceneTexturesUniformParameters> PassUniformBuffer = TUniformBufferRef<FSceneTexturesUniformParameters>::CreateUniformBufferImmediate(SceneTextureParameters, UniformBuffer_SingleFrame);	
+				TUniformBufferRef<FViewUniformShaderParameters> VoxelizeViewUniformBuffer = TUniformBufferRef<FViewUniformShaderParameters>::CreateUniformBufferImmediate(VoxelizeParameters, UniformBuffer_SingleFrame);
 
-		FDrawingPolicyRenderState DrawRenderState(View, PassUniformBuffer);
-		DrawRenderState.SetViewUniformBuffer(VoxelizeViewUniformBuffer);
+				FSceneTexturesUniformParameters SceneTextureParameters;
+				FSceneRenderTargets& SceneContext = FSceneRenderTargets::Get(RHICmdList);
+				SetupSceneTextureUniformParameters(SceneContext, View.FeatureLevel, ESceneTextureSetupMode::None, SceneTextureParameters);
+				TUniformBufferRef<FSceneTexturesUniformParameters> PassUniformBuffer = TUniformBufferRef<FSceneTexturesUniformParameters>::CreateUniformBufferImmediate(SceneTextureParameters, UniformBuffer_SingleFrame);
 
-		FVector2D Jitter(IntegrationData.FrameJitterOffsetValues[0].X / VolumetricFogGridSize.X, IntegrationData.FrameJitterOffsetValues[0].Y / VolumetricFogGridSize.Y);
+				FDrawingPolicyRenderState DrawRenderState(View, PassUniformBuffer);
+				DrawRenderState.SetViewUniformBuffer(VoxelizeViewUniformBuffer);
 
-		FTextureRHIParamRef RenderTargets[2] =
-		{
-			IntegrationData.VBufferARenderTarget->GetRenderTargetItem().TargetableTexture,
-			IntegrationData.VBufferBRenderTarget->GetRenderTargetItem().TargetableTexture
-		};
+				FVector2D Jitter(IntegrationData.FrameJitterOffsetValues[0].X / VolumetricFogGridSize.X, IntegrationData.FrameJitterOffsetValues[0].Y / VolumetricFogGridSize.Y);
 
-		SetRenderTargets(RHICmdList, ARRAY_COUNT(RenderTargets), RenderTargets, FTextureRHIParamRef(), 0, NULL);
-
-		for (int32 PrimIdx = 0; PrimIdx < View.VolumetricPrimSet.NumPrims(); PrimIdx++)
-		{
-			const FPrimitiveSceneProxy* PrimitiveSceneProxy = View.VolumetricPrimSet.GetPrim(PrimIdx);
-			const FPrimitiveSceneInfo* PrimitiveSceneInfo = PrimitiveSceneProxy->GetPrimitiveSceneInfo();
-
-			if (View.PrimitiveVisibilityMap[PrimitiveSceneInfo->GetIndex()])
-			{
-				const FPrimitiveViewRelevance& ViewRelevance = View.PrimitiveViewRelevanceMap[PrimitiveSceneInfo->GetIndex()];
-				FBoxSphereBounds Bounds = PrimitiveSceneProxy->GetBounds();
-
-				if ((View.ViewMatrices.GetViewOrigin() - Bounds.Origin).SizeSquared() < (VolumetricFogDistance + Bounds.SphereRadius) * (VolumetricFogDistance + Bounds.SphereRadius))
+				for (int32 PrimIdx = 0; PrimIdx < View.VolumetricPrimSet.NumPrims(); PrimIdx++)
 				{
-					// Range in View.DynamicMeshElements[] corresponding to this PrimitiveSceneInfo
-					FInt32Range Range = View.GetDynamicMeshElementRange(PrimitiveSceneInfo->GetIndex());
+					const FPrimitiveSceneProxy* PrimitiveSceneProxy = View.VolumetricPrimSet.GetPrim(PrimIdx);
+					const FPrimitiveSceneInfo* PrimitiveSceneInfo = PrimitiveSceneProxy->GetPrimitiveSceneInfo();
 
-					for (int32 MeshBatchIndex = Range.GetLowerBoundValue(); MeshBatchIndex < Range.GetUpperBoundValue(); MeshBatchIndex++)
+					if (View.PrimitiveVisibilityMap[PrimitiveSceneInfo->GetIndex()])
 					{
-						const FMeshBatchAndRelevance& MeshBatchAndRelevance = View.DynamicMeshElements[MeshBatchIndex];
+						const FPrimitiveViewRelevance& ViewRelevance = View.PrimitiveViewRelevanceMap[PrimitiveSceneInfo->GetIndex()];
+						FBoxSphereBounds Bounds = PrimitiveSceneProxy->GetBounds();
 
-						checkSlow(MeshBatchAndRelevance.PrimitiveSceneProxy == PrimitiveSceneInfo->Proxy);
+						if ((View.ViewMatrices.GetViewOrigin() - Bounds.Origin).SizeSquared() < (VolumetricFogDistance + Bounds.SphereRadius) * (VolumetricFogDistance + Bounds.SphereRadius))
+						{
+							// Range in View.DynamicMeshElements[] corresponding to this PrimitiveSceneInfo
+							FInt32Range Range = View.GetDynamicMeshElementRange(PrimitiveSceneInfo->GetIndex());
 
-						const FMeshBatch& MeshBatch = *MeshBatchAndRelevance.Mesh;
-						VoxelizeVolumePrimitive(RHICmdList, View, IntegrationData, DrawRenderState, Jitter, VolumetricFogGridSize, GridZParams, PrimitiveSceneProxy, MeshBatch);
-					}
-				}
+							for (int32 MeshBatchIndex = Range.GetLowerBoundValue(); MeshBatchIndex < Range.GetUpperBoundValue(); MeshBatchIndex++)
+							{
+								const FMeshBatchAndRelevance& MeshBatchAndRelevance = View.DynamicMeshElements[MeshBatchIndex];
 
-				if (ViewRelevance.bStaticRelevance)
-				{
-					for (int32 StaticMeshIdx = 0; StaticMeshIdx < PrimitiveSceneInfo->StaticMeshes.Num(); StaticMeshIdx++)
-					{
-						const FStaticMesh& StaticMesh = PrimitiveSceneInfo->StaticMeshes[StaticMeshIdx];
+								checkSlow(MeshBatchAndRelevance.PrimitiveSceneProxy == PrimitiveSceneInfo->Proxy);
 
-						if (View.StaticMeshVisibilityMap[StaticMesh.Id])
+								const FMeshBatch& MeshBatch = *MeshBatchAndRelevance.Mesh;
+								VoxelizeVolumePrimitive(RHICmdList, View, IntegrationData, DrawRenderState, Jitter, VolumetricFogGridSize, GridZParams, PrimitiveSceneProxy, MeshBatch);
+							}
+						}
+
+						if (ViewRelevance.bStaticRelevance)
 						{
-							VoxelizeVolumePrimitive(RHICmdList, View, IntegrationData, DrawRenderState, Jitter, VolumetricFogGridSize, GridZParams, PrimitiveSceneProxy, StaticMesh);
+							for (int32 StaticMeshIdx = 0; StaticMeshIdx < PrimitiveSceneInfo->StaticMeshes.Num(); StaticMeshIdx++)
+							{
+								const FStaticMesh& StaticMesh = PrimitiveSceneInfo->StaticMeshes[StaticMeshIdx];
+
+								if (View.StaticMeshVisibilityMap[StaticMesh.Id])
+								{
+									VoxelizeVolumePrimitive(RHICmdList, View, IntegrationData, DrawRenderState, Jitter, VolumetricFogGridSize, GridZParams, PrimitiveSceneProxy, StaticMesh);
+								}
+							}
 						}
 					}
 				}
-			}
-		}
+			});
 	}
 }
\ No newline at end of file
diff --git a/Engine/Source/Runtime/Renderer/Public/SceneRenderTargetParameters.h b/Engine/Source/Runtime/Renderer/Public/SceneRenderTargetParameters.h
--- a/Engine/Source/Runtime/Renderer/Public/SceneRenderTargetParameters.h
+++ b/Engine/Source/Runtime/Renderer/Public/SceneRenderTargetParameters.h
@@ -15,50 +15,50 @@ class FShaderParameterMap;
 
 // Convenience parameters used by the material graph or many global shaders
 //@todo - replace with rendergraph inputs and outputs, whose lifetimes can be validated (eg GBuffers not available in BasePass)
-BEGIN_UNIFORM_BUFFER_STRUCT(FSceneTexturesUniformParameters, RENDERER_API)
+BEGIN_GLOBAL_SHADER_PARAMETER_STRUCT(FSceneTexturesUniformParameters, RENDERER_API)
 	// Scene Color / Depth
-	UNIFORM_MEMBER_TEXTURE(Texture2D, SceneColorTexture)
-	UNIFORM_MEMBER_SAMPLER(SamplerState, SceneColorTextureSampler)
-	UNIFORM_MEMBER_TEXTURE(Texture2D, SceneDepthTexture)
-	UNIFORM_MEMBER_SAMPLER(SamplerState, SceneDepthTextureSampler)
-	UNIFORM_MEMBER_TEXTURE(Texture2D<float>, SceneDepthTextureNonMS)
+	SHADER_PARAMETER_TEXTURE(Texture2D, SceneColorTexture)
+	SHADER_PARAMETER_SAMPLER(SamplerState, SceneColorTextureSampler)
+	SHADER_PARAMETER_TEXTURE(Texture2D, SceneDepthTexture)
+	SHADER_PARAMETER_SAMPLER(SamplerState, SceneDepthTextureSampler)
+	SHADER_PARAMETER_TEXTURE(Texture2D<float>, SceneDepthTextureNonMS)
 
 	// GBuffer
-	UNIFORM_MEMBER_TEXTURE(Texture2D, GBufferATexture)
-	UNIFORM_MEMBER_TEXTURE(Texture2D, GBufferBTexture)
-	UNIFORM_MEMBER_TEXTURE(Texture2D, GBufferCTexture)
-	UNIFORM_MEMBER_TEXTURE(Texture2D, GBufferDTexture)
-	UNIFORM_MEMBER_TEXTURE(Texture2D, GBufferETexture)
-	UNIFORM_MEMBER_TEXTURE(Texture2D, GBufferVelocityTexture)
-	UNIFORM_MEMBER_TEXTURE(Texture2D<float4>, GBufferATextureNonMS)
-	UNIFORM_MEMBER_TEXTURE(Texture2D<float4>, GBufferBTextureNonMS)
-	UNIFORM_MEMBER_TEXTURE(Texture2D<float4>, GBufferCTextureNonMS)
-	UNIFORM_MEMBER_TEXTURE(Texture2D<float4>, GBufferDTextureNonMS)
-	UNIFORM_MEMBER_TEXTURE(Texture2D<float4>, GBufferETextureNonMS)
-	UNIFORM_MEMBER_TEXTURE(Texture2D<float4>, GBufferVelocityTextureNonMS)
-	UNIFORM_MEMBER_SAMPLER(SamplerState, GBufferATextureSampler)
-	UNIFORM_MEMBER_SAMPLER(SamplerState, GBufferBTextureSampler)
-	UNIFORM_MEMBER_SAMPLER(SamplerState, GBufferCTextureSampler)
-	UNIFORM_MEMBER_SAMPLER(SamplerState, GBufferDTextureSampler)
-	UNIFORM_MEMBER_SAMPLER(SamplerState, GBufferETextureSampler)
-	UNIFORM_MEMBER_SAMPLER(SamplerState, GBufferVelocityTextureSampler)
+	SHADER_PARAMETER_TEXTURE(Texture2D, GBufferATexture)
+	SHADER_PARAMETER_TEXTURE(Texture2D, GBufferBTexture)
+	SHADER_PARAMETER_TEXTURE(Texture2D, GBufferCTexture)
+	SHADER_PARAMETER_TEXTURE(Texture2D, GBufferDTexture)
+	SHADER_PARAMETER_TEXTURE(Texture2D, GBufferETexture)
+	SHADER_PARAMETER_TEXTURE(Texture2D, GBufferVelocityTexture)
+	SHADER_PARAMETER_TEXTURE(Texture2D<float4>, GBufferATextureNonMS)
+	SHADER_PARAMETER_TEXTURE(Texture2D<float4>, GBufferBTextureNonMS)
+	SHADER_PARAMETER_TEXTURE(Texture2D<float4>, GBufferCTextureNonMS)
+	SHADER_PARAMETER_TEXTURE(Texture2D<float4>, GBufferDTextureNonMS)
+	SHADER_PARAMETER_TEXTURE(Texture2D<float4>, GBufferETextureNonMS)
+	SHADER_PARAMETER_TEXTURE(Texture2D<float4>, GBufferVelocityTextureNonMS)
+	SHADER_PARAMETER_SAMPLER(SamplerState, GBufferATextureSampler)
+	SHADER_PARAMETER_SAMPLER(SamplerState, GBufferBTextureSampler)
+	SHADER_PARAMETER_SAMPLER(SamplerState, GBufferCTextureSampler)
+	SHADER_PARAMETER_SAMPLER(SamplerState, GBufferDTextureSampler)
+	SHADER_PARAMETER_SAMPLER(SamplerState, GBufferETextureSampler)
+	SHADER_PARAMETER_SAMPLER(SamplerState, GBufferVelocityTextureSampler)
 
 	// SSAO
-	UNIFORM_MEMBER_TEXTURE(Texture2D, ScreenSpaceAOTexture)
-	UNIFORM_MEMBER_SAMPLER(SamplerState, ScreenSpaceAOTextureSampler)
+	SHADER_PARAMETER_TEXTURE(Texture2D, ScreenSpaceAOTexture)
+	SHADER_PARAMETER_SAMPLER(SamplerState, ScreenSpaceAOTextureSampler)
 
 	// Custom Depth / Stencil
-	UNIFORM_MEMBER_TEXTURE(Texture2D<float>, CustomDepthTextureNonMS)
-	UNIFORM_MEMBER_TEXTURE(Texture2D, CustomDepthTexture)
-	UNIFORM_MEMBER_SAMPLER(SamplerState, CustomDepthTextureSampler)
-	UNIFORM_MEMBER_SRV(Texture2D<uint2>, CustomStencilTexture)
-	UNIFORM_MEMBER_SRV(Texture2D<uint2>, SceneStencilTexture)
+	SHADER_PARAMETER_TEXTURE(Texture2D<float>, CustomDepthTextureNonMS)
+	SHADER_PARAMETER_TEXTURE(Texture2D, CustomDepthTexture)
+	SHADER_PARAMETER_SAMPLER(SamplerState, CustomDepthTextureSampler)
+	SHADER_PARAMETER_SRV(Texture2D<uint2>, CustomStencilTexture)
+	SHADER_PARAMETER_SRV(Texture2D<uint2>, SceneStencilTexture)
 
 	// Misc
-	UNIFORM_MEMBER_TEXTURE(Texture2D, EyeAdaptation)
-	UNIFORM_MEMBER_TEXTURE(Texture2D, SceneColorCopyTexture)
-	UNIFORM_MEMBER_SAMPLER(SamplerState, SceneColorCopyTextureSampler)
-END_UNIFORM_BUFFER_STRUCT(FSceneTexturesUniformParameters)
+	SHADER_PARAMETER_TEXTURE(Texture2D, EyeAdaptation)
+	SHADER_PARAMETER_TEXTURE(Texture2D, SceneColorCopyTexture)
+	SHADER_PARAMETER_SAMPLER(SamplerState, SceneColorCopyTextureSampler)
+END_GLOBAL_SHADER_PARAMETER_STRUCT()
 
 enum class ESceneTextureSetupMode : uint32
 {
@@ -95,18 +95,18 @@ extern RENDERER_API void SetupSceneTextureUniformParameters(
 template< typename TRHICmdList >
 RENDERER_API TUniformBufferRef<FSceneTexturesUniformParameters> CreateSceneTextureUniformBufferSingleDraw(TRHICmdList& RHICmdList, ESceneTextureSetupMode SceneTextureSetupMode, ERHIFeatureLevel::Type FeatureLevel);
 
-BEGIN_UNIFORM_BUFFER_STRUCT(FMobileSceneTextureUniformParameters, RENDERER_API)
-	UNIFORM_MEMBER_TEXTURE(Texture2D, SceneColorTexture)
-	UNIFORM_MEMBER_SAMPLER(SamplerState, SceneColorTextureSampler)
-	UNIFORM_MEMBER_TEXTURE(Texture2D, SceneDepthTexture)
-	UNIFORM_MEMBER_SAMPLER(SamplerState, SceneDepthTextureSampler)
-	UNIFORM_MEMBER_TEXTURE(Texture2D, SceneAlphaCopyTexture)
-	UNIFORM_MEMBER_SAMPLER(SamplerState, SceneAlphaCopyTextureSampler)
-	UNIFORM_MEMBER_TEXTURE(Texture2D, CustomDepthTexture)
-	UNIFORM_MEMBER_SAMPLER(SamplerState, CustomDepthTextureSampler)
-	UNIFORM_MEMBER_TEXTURE(Texture2D, MobileCustomStencilTexture)
-	UNIFORM_MEMBER_SAMPLER(SamplerState, MobileCustomStencilTextureSampler)
-END_UNIFORM_BUFFER_STRUCT(FMobileSceneTextureUniformParameters)
+BEGIN_GLOBAL_SHADER_PARAMETER_STRUCT(FMobileSceneTextureUniformParameters, RENDERER_API)
+	SHADER_PARAMETER_TEXTURE(Texture2D, SceneColorTexture)
+	SHADER_PARAMETER_SAMPLER(SamplerState, SceneColorTextureSampler)
+	SHADER_PARAMETER_TEXTURE(Texture2D, SceneDepthTexture)
+	SHADER_PARAMETER_SAMPLER(SamplerState, SceneDepthTextureSampler)
+	SHADER_PARAMETER_TEXTURE(Texture2D, SceneAlphaCopyTexture)
+	SHADER_PARAMETER_SAMPLER(SamplerState, SceneAlphaCopyTextureSampler)
+	SHADER_PARAMETER_TEXTURE(Texture2D, CustomDepthTexture)
+	SHADER_PARAMETER_SAMPLER(SamplerState, CustomDepthTextureSampler)
+	SHADER_PARAMETER_TEXTURE(Texture2D, MobileCustomStencilTexture)
+	SHADER_PARAMETER_SAMPLER(SamplerState, MobileCustomStencilTextureSampler)
+END_GLOBAL_SHADER_PARAMETER_STRUCT()
 
 extern void SetupMobileSceneTextureUniformParameters(
 	FSceneRenderTargets& SceneContext,
