Fix for InfiltratorForward failing on startup on DX12

diff --git a/Engine/Source/Runtime/Renderer/Private/CompositionLighting/PostProcessDeferredDecals.cpp b/Engine/Source/Runtime/Renderer/Private/CompositionLighting/PostProcessDeferredDecals.cpp
--- a/Engine/Source/Runtime/Renderer/Private/CompositionLighting/PostProcessDeferredDecals.cpp
+++ b/Engine/Source/Runtime/Renderer/Private/CompositionLighting/PostProcessDeferredDecals.cpp
@@ -1149,7 +1149,7 @@ void FDecalRenderTargetManager::SetRenderTargetMode(FDecalRenderingCommon::ERend
 			RTCount = 4;
 		}
 
-		SetRenderTargets(RHICmdList, RTCount, &TargetsToResolve[DBufferAIndex], SceneContext.GetSceneDepthSurface(), ESimpleRenderTargetMode::EExistingColorAndDepth, FExclusiveDepthStencil::DepthRead_StencilWrite, TargetsToTransitionWritable[CurrentRenderTargetMode]);
+		SetRenderTargets(RHICmdList, RTCount, &TargetsToResolve[DBufferAIndex], SceneContext.GetSceneDepthTexture(), ESimpleRenderTargetMode::EExistingColorAndDepth, FExclusiveDepthStencil::DepthRead_StencilWrite, TargetsToTransitionWritable[CurrentRenderTargetMode]);
 		break;
 	}
 
diff --git a/Engine/Source/Runtime/Renderer/Private/PostProcess/SceneRenderTargets.cpp b/Engine/Source/Runtime/Renderer/Private/PostProcess/SceneRenderTargets.cpp
--- a/Engine/Source/Runtime/Renderer/Private/PostProcess/SceneRenderTargets.cpp
+++ b/Engine/Source/Runtime/Renderer/Private/PostProcess/SceneRenderTargets.cpp
@@ -1919,7 +1919,7 @@ void FSceneRenderTargets::AllocateCommonDepthTargets(FRHICommandList& RHICmdList
 		const bool bHMDAllocated = StereoRenderTargetManager && StereoRenderTargetManager->AllocateDepthTexture(0, BufferSize.X, BufferSize.Y, PF_X24_G8, 0, TexCreate_None, TexCreate_DepthStencilTargetable, DepthTex, SRTex, GetNumSceneColorMSAASamples(CurrentFeatureLevel));
 
 		// Create a texture to store the resolved scene depth, and a render-targetable surface to hold the unresolved scene depth.
-		FPooledRenderTargetDesc Desc(FPooledRenderTargetDesc::Create2DDesc(BufferSize, PF_DepthStencil, DefaultDepthClear, TexCreate_None, TexCreate_DepthStencilTargetable, false));
+		FPooledRenderTargetDesc Desc(FPooledRenderTargetDesc::Create2DDesc(BufferSize, PF_DepthStencil, DefaultDepthClear, TexCreate_None, TexCreate_DepthStencilTargetable | TexCreate_ShaderResource, false));
 		Desc.NumSamples = GetNumSceneColorMSAASamples(CurrentFeatureLevel);
 		Desc.Flags |= GFastVRamConfig.SceneDepth;
 		GRenderTargetPool.FindFreeElement(RHICmdList, Desc, SceneDepthZ, TEXT("SceneDepthZ"));
