Fixes declaration shadowing warning on some compilers.

diff --git a/Engine/Source/Runtime/Renderer/Private/SceneOcclusion.cpp b/Engine/Source/Runtime/Renderer/Private/SceneOcclusion.cpp
--- a/Engine/Source/Runtime/Renderer/Private/SceneOcclusion.cpp
+++ b/Engine/Source/Runtime/Renderer/Private/SceneOcclusion.cpp
@@ -1089,11 +1089,11 @@ class FHZBBuildPS : public FGlobalShader
 IMPLEMENT_GLOBAL_SHADER(FHZBBuildPS, "/Engine/Private/HZBOcclusion.usf", "HZBBuildPS", SF_Pixel);
 
 
-void BuildHZB(FRHICommandListImmediate& RHICmdList, FViewInfo& View)
+void BuildHZB(FRHICommandListImmediate& RHICmdListImmediate, FViewInfo& View)
 {
 	QUICK_SCOPE_CYCLE_COUNTER(STAT_BuildHZB);
 	
-	FRDGBuilder GraphBuilder(RHICmdList);
+	FRDGBuilder GraphBuilder(RHICmdListImmediate);
 
 	// View.ViewRect.{Width,Height}() are most likely to be < 2^24, so the float
 	// conversion won't loss any precision (assuming float have 23bits for mantissa)
diff --git a/Engine/Source/Runtime/Renderer/Private/VolumetricFog.cpp b/Engine/Source/Runtime/Renderer/Private/VolumetricFog.cpp
--- a/Engine/Source/Runtime/Renderer/Private/VolumetricFog.cpp
+++ b/Engine/Source/Runtime/Renderer/Private/VolumetricFog.cpp
@@ -957,7 +957,7 @@ void FDeferredShadingSceneRenderer::SetupVolumetricFog()
 	}
 }
 
-void FDeferredShadingSceneRenderer::ComputeVolumetricFog(FRHICommandListImmediate& RHICmdList)
+void FDeferredShadingSceneRenderer::ComputeVolumetricFog(FRHICommandListImmediate& RHICmdListImmediate)
 {
 	if (ShouldRenderVolumetricFog())
 	{
@@ -1000,7 +1000,7 @@ void FDeferredShadingSceneRenderer::ComputeVolumetricFog(FRHICommandListImmediat
 
 			FMatrix LightFunctionWorldToShadow;
 
-			FRDGBuilder GraphBuilder(RHICmdList);
+			FRDGBuilder GraphBuilder(RHICmdListImmediate);
 
 			//@DW - register GWhiteTexture as a graph external for when there's no light function - later a shader is going to bind it whether we rendered to it or not
 			const FRDGTexture* LightFunctionTexture = GraphBuilder.RegisterExternalTexture(GSystemTextures.WhiteDummy);
