Fixes more static analysis...

diff --git a/Engine/Source/Runtime/RenderCore/Public/ShaderParameterStruct.h b/Engine/Source/Runtime/RenderCore/Public/ShaderParameterStruct.h
--- a/Engine/Source/Runtime/RenderCore/Public/ShaderParameterStruct.h
+++ b/Engine/Source/Runtime/RenderCore/Public/ShaderParameterStruct.h
@@ -74,6 +74,7 @@ inline void SetShaderUAVs(TRHICmdList& RHICmdList, const TShaderClass* Shader, F
 			}
 		}
 
+		checkSlow(GraphUAV);
 		RHICmdList.SetUAVParameter(ShadeRHI, ParameterBinding.BaseIndex, GraphUAV->CachedRHI.UAV);
 	}
 }
@@ -166,6 +167,7 @@ inline void SetShaderParameters(TRHICmdList& RHICmdList, const TShaderClass* Sha
 			}
 		}
 
+		checkSlow(GraphTexture);
 		RHICmdList.SetShaderTexture(ShadeRHI, ParameterBinding.BaseIndex, GraphTexture->GetRHITexture());
 	}
 
@@ -187,6 +189,7 @@ inline void SetShaderParameters(TRHICmdList& RHICmdList, const TShaderClass* Sha
 			}
 		}
 
+		checkSlow(GraphSRV);
 		RHICmdList.SetShaderResourceViewParameter(ShadeRHI, ParameterBinding.BaseIndex, GraphSRV->CachedRHI.SRV);
 	}
 
diff --git a/Engine/Source/Runtime/Renderer/Private/SceneOcclusion.cpp b/Engine/Source/Runtime/Renderer/Private/SceneOcclusion.cpp
--- a/Engine/Source/Runtime/Renderer/Private/SceneOcclusion.cpp
+++ b/Engine/Source/Runtime/Renderer/Private/SceneOcclusion.cpp
@@ -1138,12 +1138,12 @@ void BuildHZB(FRHICommandListImmediate& RHICmdListImmediate, FViewInfo& View)
 				FHZBBuildPS::FPermutationDomain PermutationVector;
 				PermutationVector.Set<FHZBBuildPS::FStageDim>(false);
 
-				auto VertexShader = View.ShaderMap->GetShader< FPostProcessVS >();
-				auto PixelShader  = View.ShaderMap->GetShader< FHZBBuildPS >( PermutationVector );
+				TShaderMapRef<FPostProcessVS> VertexShader(View.ShaderMap);
+				TShaderMapRef<FHZBBuildPS> PixelShader(View.ShaderMap, PermutationVector);
 
 				GraphicsPSOInit.BoundShaderState.VertexDeclarationRHI = GFilterVertexDeclaration.VertexDeclarationRHI;
-				GraphicsPSOInit.BoundShaderState.VertexShaderRHI = GETSAFERHISHADER_VERTEX(VertexShader);
-				GraphicsPSOInit.BoundShaderState.PixelShaderRHI = GETSAFERHISHADER_PIXEL(PixelShader);
+				GraphicsPSOInit.BoundShaderState.VertexShaderRHI = GETSAFERHISHADER_VERTEX(*VertexShader);
+				GraphicsPSOInit.BoundShaderState.PixelShaderRHI = GETSAFERHISHADER_PIXEL(*PixelShader);
 				GraphicsPSOInit.PrimitiveType = PT_TriangleList;
 
 				SetGraphicsPipelineState(RHICmdList, GraphicsPSOInit);
@@ -1167,7 +1167,7 @@ void BuildHZB(FRHICommandListImmediate& RHICmdListImmediate, FViewInfo& View)
 				Parameters.View = View.ViewUniformBuffer;
 				Parameters.SceneTextures = CreateSceneTextureUniformBufferSingleDraw(RHICmdList, ESceneTextureSetupMode::SceneDepth, View.FeatureLevel);
 
-				SetShaderParameters(RHICmdList, PixelShader, PixelShader->GetPixelShader(), Parameters);
+				SetShaderParameters(RHICmdList, *PixelShader, PixelShader->GetPixelShader(), Parameters);
 
 				RHICmdList.SetViewport(0, 0, 0.0f, HZBSize.X, HZBSize.Y, 1.0f);
 
@@ -1179,7 +1179,7 @@ void BuildHZB(FRHICommandListImmediate& RHICmdListImmediate, FViewInfo& View)
 					View.ViewRect.Width(), View.ViewRect.Height(),
 					HZBSize,
 					FSceneRenderTargets::Get(RHICmdList).GetBufferSizeXY(),
-					VertexShader,
+					*VertexShader,
 					EDRF_UseTriangleOptimization);
 			});
 	}
@@ -1213,8 +1213,8 @@ void BuildHZB(FRHICommandListImmediate& RHICmdListImmediate, FViewInfo& View)
 				FHZBBuildPS::FPermutationDomain PermutationVector;
 				PermutationVector.Set<FHZBBuildPS::FStageDim>(true);
 
-				auto VertexShader = View.ShaderMap->GetShader< FPostProcessVS >();
-				auto PixelShader  = View.ShaderMap->GetShader< FHZBBuildPS >( PermutationVector );
+				TShaderMapRef<FPostProcessVS> VertexShader(View.ShaderMap);
+				TShaderMapRef<FHZBBuildPS> PixelShader(View.ShaderMap, PermutationVector);
 
 				FGraphicsPipelineStateInitializer GraphicsPSOInit;
 				GraphicsPSOInit.BlendState = TStaticBlendState<>::GetRHI();
@@ -1224,8 +1224,8 @@ void BuildHZB(FRHICommandListImmediate& RHICmdListImmediate, FViewInfo& View)
 				RHICmdList.ApplyCachedRenderTargets(GraphicsPSOInit);
 
 				GraphicsPSOInit.BoundShaderState.VertexDeclarationRHI = GFilterVertexDeclaration.VertexDeclarationRHI;
-				GraphicsPSOInit.BoundShaderState.VertexShaderRHI = GETSAFERHISHADER_VERTEX(VertexShader);
-				GraphicsPSOInit.BoundShaderState.PixelShaderRHI = GETSAFERHISHADER_PIXEL(PixelShader);
+				GraphicsPSOInit.BoundShaderState.VertexShaderRHI = GETSAFERHISHADER_VERTEX(*VertexShader);
+				GraphicsPSOInit.BoundShaderState.PixelShaderRHI = GETSAFERHISHADER_PIXEL(*PixelShader);
 				GraphicsPSOInit.PrimitiveType = PT_TriangleList;
 
 				SetGraphicsPipelineState(RHICmdList, GraphicsPSOInit);
@@ -1235,7 +1235,7 @@ void BuildHZB(FRHICommandListImmediate& RHICmdListImmediate, FViewInfo& View)
 				Parameters.Pass = *PassParameters;
 				Parameters.View = View.ViewUniformBuffer;
 
-				SetShaderParameters(RHICmdList, PixelShader, PixelShader->GetPixelShader(), Parameters);
+				SetShaderParameters(RHICmdList, *PixelShader, PixelShader->GetPixelShader(), Parameters);
 
 				RHICmdList.SetViewport(0, 0, 0.0f, DstSize.X, DstSize.Y, 1.0f);
 
@@ -1247,7 +1247,7 @@ void BuildHZB(FRHICommandListImmediate& RHICmdListImmediate, FViewInfo& View)
 					SrcSize.X, SrcSize.Y,
 					DstSize,
 					SrcSize,
-					VertexShader,
+					*VertexShader,
 					EDRF_UseTriangleOptimization);
 			});
 
