Render Passes

Self-contained PostProcessing effects.

diff --git a/Engine/Source/Runtime/Renderer/Private/PostProcess/DiaphragmDOFPasses.cpp b/Engine/Source/Runtime/Renderer/Private/PostProcess/DiaphragmDOFPasses.cpp
--- a/Engine/Source/Runtime/Renderer/Private/PostProcess/DiaphragmDOFPasses.cpp
+++ b/Engine/Source/Runtime/Renderer/Private/PostProcess/DiaphragmDOFPasses.cpp
@@ -338,7 +338,8 @@ public:
 	{
 		DestViewport.Min = FIntPoint::ZeroValue;
 
-		SetRenderTarget(Context.RHICmdList, FTextureRHIRef(), FTextureRHIRef());
+		// #todo-renderpasses remove once everything is a renderpass
+		UnbindRenderTargets(Context.RHICmdList);
 		Context.RHICmdList.SetComputeShader(ShaderRHI);
 
 		for (int32 i = 0; i < OutputCount; i++)
@@ -1668,70 +1669,72 @@ void FRCPassDiaphragmDOFHybridScatter::Process(FRenderingCompositePassContext& C
 		EResourceTransitionPipeline::EComputeToGfx,
 		DestRenderTarget.UAV);
 
-	SetRenderTarget(Context.RHICmdList, DestRenderTarget.TargetableTexture, FTextureRHIRef());
-	Context.SetViewportAndCallRHI(DestViewport, 0.0f, 1.0f);
+	FRHIRenderPassInfo RPInfo(DestRenderTarget.TargetableTexture, ERenderTargetActions::Load_Store, DestRenderTarget.ShaderResourceTexture);
+	Context.RHICmdList.BeginRenderPass(RPInfo, TEXT("DOFHybridScatter"));
+	{
+		Context.SetViewportAndCallRHI(DestViewport, 0.0f, 1.0f);
 
-	SCOPED_DRAW_EVENTF(Context.RHICmdList, DiaphragmDOFIndirectScatter, TEXT("DiaphragmDOF IndirectScatter(%s Bokeh=%s Occlusion=%s 1/2) %dx%d"),
-		GetEventName(bIsForeground ? EDiaphragmDOFLayerProcessing::ForegroundOnly : EDiaphragmDOFLayerProcessing::BackgroundOnly),
-		PermutationVector.Get<FPostProcessDiaphragmDOFHybridScatterPS::FBokehSimulationDim>() ? TEXT("Generic") : TEXT("None"),
-		PermutationVector.Get<FDDOFScatterOcclusionDim>() ? TEXT("Yes") : TEXT("No"),
-		DestViewport.Width(), DestViewport.Height());
+		SCOPED_DRAW_EVENTF(Context.RHICmdList, DiaphragmDOFIndirectScatter, TEXT("DiaphragmDOF IndirectScatter(%s Bokeh=%s Occlusion=%s 1/2) %dx%d"),
+			GetEventName(bIsForeground ? EDiaphragmDOFLayerProcessing::ForegroundOnly : EDiaphragmDOFLayerProcessing::BackgroundOnly),
+			PermutationVector.Get<FPostProcessDiaphragmDOFHybridScatterPS::FBokehSimulationDim>() ? TEXT("Generic") : TEXT("None"),
+			PermutationVector.Get<FDDOFScatterOcclusionDim>() ? TEXT("Yes") : TEXT("No"),
+			DestViewport.Width(), DestViewport.Height());
 
-	{
-		FGraphicsPipelineStateInitializer GraphicsPSOInit;
-		Context.RHICmdList.ApplyCachedRenderTargets(GraphicsPSOInit);
-		GraphicsPSOInit.DepthStencilState = TStaticDepthStencilState<false, CF_Always>::GetRHI();
-		GraphicsPSOInit.BlendState = TStaticBlendState<CW_RGBA, BO_Add, BF_One, BF_One, BO_Add, BF_One, BF_One>::GetRHI();
-		GraphicsPSOInit.RasterizerState = TStaticRasterizerState<>::GetRHI();
-		GraphicsPSOInit.PrimitiveType = PrimitiveType;
-		GraphicsPSOInit.BoundShaderState.VertexDeclarationRHI = GEmptyVertexDeclaration.VertexDeclarationRHI;
-		GraphicsPSOInit.BoundShaderState.VertexShaderRHI = GETSAFERHISHADER_VERTEX(*VertexShader);
-		GraphicsPSOInit.BoundShaderState.PixelShaderRHI = GETSAFERHISHADER_PIXEL(*PixelShader);
-		SetGraphicsPipelineState(Context.RHICmdList, GraphicsPSOInit);
-	}
-	
-	float ScatteringScaling = float(Params.OutputViewSize.X) / float(Params.InputViewSize.X);
+		{
+			FGraphicsPipelineStateInitializer GraphicsPSOInit;
+			Context.RHICmdList.ApplyCachedRenderTargets(GraphicsPSOInit);
+			GraphicsPSOInit.DepthStencilState = TStaticDepthStencilState<false, CF_Always>::GetRHI();
+			GraphicsPSOInit.BlendState = TStaticBlendState<CW_RGBA, BO_Add, BF_One, BF_One, BO_Add, BF_One, BF_One>::GetRHI();
+			GraphicsPSOInit.RasterizerState = TStaticRasterizerState<>::GetRHI();
+			GraphicsPSOInit.PrimitiveType = PrimitiveType;
+			GraphicsPSOInit.BoundShaderState.VertexDeclarationRHI = GEmptyVertexDeclaration.VertexDeclarationRHI;
+			GraphicsPSOInit.BoundShaderState.VertexShaderRHI = GETSAFERHISHADER_VERTEX(*VertexShader);
+			GraphicsPSOInit.BoundShaderState.PixelShaderRHI = GETSAFERHISHADER_PIXEL(*PixelShader);
+			SetGraphicsPipelineState(Context.RHICmdList, GraphicsPSOInit);
+		}
 
-	{
-		FVertexShaderRHIParamRef ShaderRHI = VertexShader->GetVertexShader();
-		
-		SetShaderValue(Context.RHICmdList, ShaderRHI, VertexShader->CocRadiusToCircumscribedRadius,
-			BokehModel.CocRadiusToCircumscribedRadius);
+		float ScatteringScaling = float(Params.OutputViewSize.X) / float(Params.InputViewSize.X);
 
-		SetSRVParameter(Context.RHICmdList, ShaderRHI, VertexShader->ScatterDrawList,
-			ScatterDrawListBuffer->SRV);
+		{
+			FVertexShaderRHIParamRef ShaderRHI = VertexShader->GetVertexShader();
 
-		VertexShader->PostprocessParameter.SetVS(ShaderRHI, Context);
+			SetShaderValue(Context.RHICmdList, ShaderRHI, VertexShader->CocRadiusToCircumscribedRadius,
+				BokehModel.CocRadiusToCircumscribedRadius);
 
-		SetShaderValue(Context.RHICmdList, ShaderRHI, VertexShader->ScatteringScaling, ScatteringScaling);
-	}
+			SetSRVParameter(Context.RHICmdList, ShaderRHI, VertexShader->ScatterDrawList,
+				ScatterDrawListBuffer->SRV);
 
-	{
-		FPixelShaderRHIParamRef ShaderRHI = PixelShader->GetPixelShader();
-		PixelShader->PostprocessParameter.SetPS(Context.RHICmdList, ShaderRHI, Context);
-		PixelShader->SetParameters<FViewUniformShaderParameters>(
-			Context.RHICmdList, ShaderRHI, Context.View.ViewUniformBuffer);
+			VertexShader->PostprocessParameter.SetVS(ShaderRHI, Context);
 
-		SetShaderValue(Context.RHICmdList, ShaderRHI, PixelShader->ScatteringScaling, ScatteringScaling);
-	}
+			SetShaderValue(Context.RHICmdList, ShaderRHI, VertexShader->ScatteringScaling, ScatteringScaling);
+		}
 
-	Context.RHICmdList.SetStreamSource(0, NULL, 0);
+		{
+			FPixelShaderRHIParamRef ShaderRHI = PixelShader->GetPixelShader();
+			PixelShader->PostprocessParameter.SetPS(Context.RHICmdList, ShaderRHI, Context);
+			PixelShader->SetParameters<FViewUniformShaderParameters>(
+				Context.RHICmdList, ShaderRHI, Context.View.ViewUniformBuffer);
 
-	if (GRHISupportsRectTopology)
-	{
-		Context.RHICmdList.DrawPrimitiveIndirect(
-			DrawIndirectParametersBuffer->Buffer,
-			sizeof(FRHIDrawIndirectParameters) * DrawIndirectParametersOffset);
-	}
-	else
-	{
-		Context.RHICmdList.DrawIndexedPrimitiveIndirect(
-			GDiaphragmDOFGlobalResource.ScatterIndexBuffer.IndexBufferRHI,
-			DrawIndirectParametersBuffer->Buffer,
-			sizeof(FRHIDrawIndexedIndirectParameters) * DrawIndirectParametersOffset);
-	}
+			SetShaderValue(Context.RHICmdList, ShaderRHI, PixelShader->ScatteringScaling, ScatteringScaling);
+		}
 
-	Context.RHICmdList.CopyToResolveTarget(DestRenderTarget.TargetableTexture, DestRenderTarget.ShaderResourceTexture, FResolveParams());
+		Context.RHICmdList.SetStreamSource(0, NULL, 0);
+
+		if (GRHISupportsRectTopology)
+		{
+			Context.RHICmdList.DrawPrimitiveIndirect(
+				DrawIndirectParametersBuffer->Buffer,
+				sizeof(FRHIDrawIndirectParameters) * DrawIndirectParametersOffset);
+		}
+		else
+		{
+			Context.RHICmdList.DrawIndexedPrimitiveIndirect(
+				GDiaphragmDOFGlobalResource.ScatterIndexBuffer.IndexBufferRHI,
+				DrawIndirectParametersBuffer->Buffer,
+				sizeof(FRHIDrawIndexedIndirectParameters) * DrawIndirectParametersOffset);
+		}
+	}
+	Context.RHICmdList.EndRenderPass();
 
 	{
 		FVertexShaderRHIParamRef ShaderRHI = VertexShader->GetVertexShader();
diff --git a/Engine/Source/Runtime/Renderer/Private/PostProcess/PostProcessAA.cpp b/Engine/Source/Runtime/Renderer/Private/PostProcess/PostProcessAA.cpp
--- a/Engine/Source/Runtime/Renderer/Private/PostProcess/PostProcessAA.cpp
+++ b/Engine/Source/Runtime/Renderer/Private/PostProcess/PostProcessAA.cpp
@@ -262,34 +262,36 @@ void FRCPassPostProcessAA::Process(FRenderingCompositePassContext& Context)
 	SCOPED_DRAW_EVENTF(Context.RHICmdList, PostProcessFXAA, TEXT("PostProcessFXAA %dx%d"), DestRect.Width(), DestRect.Height());
 
 	// Set the view family's render target/viewport.
-	SetRenderTarget(Context.RHICmdList, DestRenderTarget.TargetableTexture, FTextureRHIRef());
-	Context.SetViewportAndCallRHI(0, 0, 0.0f, DestSize.X, DestSize.Y, 1.0f );
-
-	switch(Quality)
+	FRHIRenderPassInfo RPInfo(DestRenderTarget.TargetableTexture, ERenderTargetActions::Load_Store, DestRenderTarget.ShaderResourceTexture);
+	Context.RHICmdList.BeginRenderPass(RPInfo, TEXT("PostProcessAA"));
 	{
+		Context.SetViewportAndCallRHI(0, 0, 0.0f, DestSize.X, DestSize.Y, 1.0f);
+
+		switch (Quality)
+		{
 		case 1:  SetShaderTemplAA<1>(Context.RHICmdList, Context); break;
 		case 2:  SetShaderTemplAA<2>(Context.RHICmdList, Context); break;
 		case 3:  SetShaderTemplAA<3>(Context.RHICmdList, Context); break;
 		case 4:  SetShaderTemplAA<4>(Context.RHICmdList, Context); break;
 		case 5:  SetShaderTemplAA<5>(Context.RHICmdList, Context); break;
 		default: SetShaderTemplAA<6>(Context.RHICmdList, Context); break;
-	}
-	
-	TShaderMapRef<FFXAAVS> VertexShader(Context.GetShaderMap());
+		}
 
-	DrawPostProcessPass(Context.RHICmdList,
-		DestRect.Min.X, DestRect.Min.Y,
-		DestRect.Width(), DestRect.Height(),
-		SrcRect.Min.X, SrcRect.Min.Y,
-		SrcRect.Width(), SrcRect.Height(),
-		DestSize,
-		SrcSize,
-		*VertexShader,
-		View.StereoPass,
-		Context.HasHmdMesh(),
-		EDRF_Default);
-
-	Context.RHICmdList.CopyToResolveTarget(DestRenderTarget.TargetableTexture, DestRenderTarget.ShaderResourceTexture, FResolveParams());
+		TShaderMapRef<FFXAAVS> VertexShader(Context.GetShaderMap());
+
+		DrawPostProcessPass(Context.RHICmdList,
+			DestRect.Min.X, DestRect.Min.Y,
+			DestRect.Width(), DestRect.Height(),
+			SrcRect.Min.X, SrcRect.Min.Y,
+			SrcRect.Width(), SrcRect.Height(),
+			DestSize,
+			SrcSize,
+			*VertexShader,
+			View.StereoPass,
+			Context.HasHmdMesh(),
+			EDRF_Default);
+	}
+	Context.RHICmdList.EndRenderPass();
 }
 
 FPooledRenderTargetDesc FRCPassPostProcessAA::ComputeOutputDesc(EPassOutputId InPassOutputId) const
diff --git a/Engine/Source/Runtime/Renderer/Private/PostProcess/PostProcessBloomSetup.cpp b/Engine/Source/Runtime/Renderer/Private/PostProcess/PostProcessBloomSetup.cpp
--- a/Engine/Source/Runtime/Renderer/Private/PostProcess/PostProcessBloomSetup.cpp
+++ b/Engine/Source/Runtime/Renderer/Private/PostProcess/PostProcessBloomSetup.cpp
@@ -244,7 +244,8 @@ void FRCPassPostProcessBloomSetup::Process(FRenderingCompositePassContext& Conte
 		DestRect = {View.ViewRect.Min, View.ViewRect.Min + DestSize};
 	
 		// Common setup
-		SetRenderTarget(Context.RHICmdList, nullptr, nullptr);
+		// #todo-renderpasses remove once everything is renderpasses
+		UnbindRenderTargets(Context.RHICmdList);
 		Context.SetViewportAndCallRHI(DestRect, 0.0f, 1.0f);
 		
 		static FName AsyncEndFenceName(TEXT("AsyncBloomSetupEndFence"));
@@ -285,44 +286,44 @@ void FRCPassPostProcessBloomSetup::Process(FRenderingCompositePassContext& Conte
 		WaitForInputPassComputeFences(Context.RHICmdList);
 
 		// Set the view family's render target/viewport.
-		FRHIRenderTargetView RtView = FRHIRenderTargetView(DestRenderTarget.TargetableTexture, ERenderTargetLoadAction::ENoAction);
-		FRHISetRenderTargetsInfo Info(1, &RtView, FRHIDepthRenderTargetView());
-		Context.RHICmdList.SetRenderTargetsAndClear(Info);
-		Context.SetViewportAndCallRHI(0, 0, 0.0f, DestSize.X, DestSize.Y, 1.0f );
-
-		FGraphicsPipelineStateInitializer GraphicsPSOInit;
-		Context.RHICmdList.ApplyCachedRenderTargets(GraphicsPSOInit);
-
-		// set the state
-		GraphicsPSOInit.BlendState = TStaticBlendState<>::GetRHI();
-		GraphicsPSOInit.RasterizerState = TStaticRasterizerState<>::GetRHI();
-		GraphicsPSOInit.DepthStencilState = TStaticDepthStencilState<false, CF_Always>::GetRHI();
-
-		TShaderMapRef<FPostProcessBloomSetupVS> VertexShader(Context.GetShaderMap());
-		TShaderMapRef<FPostProcessBloomSetupPS> PixelShader(Context.GetShaderMap());
-
-		GraphicsPSOInit.BoundShaderState.VertexDeclarationRHI = GFilterVertexDeclaration.VertexDeclarationRHI;
-		GraphicsPSOInit.BoundShaderState.VertexShaderRHI = GETSAFERHISHADER_VERTEX(*VertexShader);
-		GraphicsPSOInit.BoundShaderState.PixelShaderRHI = GETSAFERHISHADER_PIXEL(*PixelShader);
-		GraphicsPSOInit.PrimitiveType = PT_TriangleList;
-
-		SetGraphicsPipelineState(Context.RHICmdList, GraphicsPSOInit);
-
-		VertexShader->SetVS(Context);
-		PixelShader->SetPS(Context.RHICmdList, Context);
-
-		DrawRectangle(
-			Context.RHICmdList,
-			DestRect.Min.X, DestRect.Min.Y,
-			DestRect.Width(), DestRect.Height(),
-			SrcRect.Min.X, SrcRect.Min.Y,
-			SrcRect.Width(), SrcRect.Height(),
-			DestSize,
-			SrcSize,
-			*VertexShader,
-			EDRF_UseTriangleOptimization);
-
-		Context.RHICmdList.CopyToResolveTarget(DestRenderTarget.TargetableTexture, DestRenderTarget.ShaderResourceTexture, FResolveParams());
+		FRHIRenderPassInfo RPInfo(DestRenderTarget.TargetableTexture, ERenderTargetActions::Clear_Store, DestRenderTarget.ShaderResourceTexture);
+		Context.RHICmdList.BeginRenderPass(RPInfo, TEXT("Bloom"));
+		{
+			Context.SetViewportAndCallRHI(0, 0, 0.0f, DestSize.X, DestSize.Y, 1.0f);
+
+			FGraphicsPipelineStateInitializer GraphicsPSOInit;
+			Context.RHICmdList.ApplyCachedRenderTargets(GraphicsPSOInit);
+
+			// set the state
+			GraphicsPSOInit.BlendState = TStaticBlendState<>::GetRHI();
+			GraphicsPSOInit.RasterizerState = TStaticRasterizerState<>::GetRHI();
+			GraphicsPSOInit.DepthStencilState = TStaticDepthStencilState<false, CF_Always>::GetRHI();
+
+			TShaderMapRef<FPostProcessBloomSetupVS> VertexShader(Context.GetShaderMap());
+			TShaderMapRef<FPostProcessBloomSetupPS> PixelShader(Context.GetShaderMap());
+
+			GraphicsPSOInit.BoundShaderState.VertexDeclarationRHI = GFilterVertexDeclaration.VertexDeclarationRHI;
+			GraphicsPSOInit.BoundShaderState.VertexShaderRHI = GETSAFERHISHADER_VERTEX(*VertexShader);
+			GraphicsPSOInit.BoundShaderState.PixelShaderRHI = GETSAFERHISHADER_PIXEL(*PixelShader);
+			GraphicsPSOInit.PrimitiveType = PT_TriangleList;
+
+			SetGraphicsPipelineState(Context.RHICmdList, GraphicsPSOInit);
+
+			VertexShader->SetVS(Context);
+			PixelShader->SetPS(Context.RHICmdList, Context);
+
+			DrawRectangle(
+				Context.RHICmdList,
+				DestRect.Min.X, DestRect.Min.Y,
+				DestRect.Width(), DestRect.Height(),
+				SrcRect.Min.X, SrcRect.Min.Y,
+				SrcRect.Width(), SrcRect.Height(),
+				DestSize,
+				SrcSize,
+				*VertexShader,
+				EDRF_UseTriangleOptimization);
+		}
+		Context.RHICmdList.EndRenderPass();
 	}
 }
 
@@ -427,45 +428,46 @@ void FRCPassPostProcessVisualizeBloomSetup::Process(FRenderingCompositePassConte
 	FIntRect DestRect = SrcRect;
 
 	const FSceneRenderTargetItem& DestRenderTarget = PassOutputs[0].RequestSurface(Context);
-	FRHIRenderTargetView RtView = FRHIRenderTargetView(DestRenderTarget.TargetableTexture, ERenderTargetLoadAction::ENoAction);
-	FRHISetRenderTargetsInfo Info(1, &RtView, FRHIDepthRenderTargetView());
-	Context.RHICmdList.SetRenderTargetsAndClear(Info);
-	Context.SetViewportAndCallRHI(0, 0, 0.0f, DestRect.Width(), DestRect.Height(), 1.0f );
-
-	FGraphicsPipelineStateInitializer GraphicsPSOInit;
-	Context.RHICmdList.ApplyCachedRenderTargets(GraphicsPSOInit);
-
-	// set the state
-	GraphicsPSOInit.BlendState = TStaticBlendState<>::GetRHI();
-	GraphicsPSOInit.RasterizerState = TStaticRasterizerState<>::GetRHI();
-	GraphicsPSOInit.DepthStencilState = TStaticDepthStencilState<false, CF_Always>::GetRHI();
-
-	TShaderMapRef<FPostProcessVS> VertexShader(Context.GetShaderMap());
-	TShaderMapRef<FPostProcessVisualizeBloomSetupPS> PixelShader(Context.GetShaderMap());
-
-	GraphicsPSOInit.BoundShaderState.VertexDeclarationRHI = GFilterVertexDeclaration.VertexDeclarationRHI;
-	GraphicsPSOInit.BoundShaderState.VertexShaderRHI = GETSAFERHISHADER_VERTEX(*VertexShader);
-	GraphicsPSOInit.BoundShaderState.PixelShaderRHI = GETSAFERHISHADER_PIXEL(*PixelShader);
-	GraphicsPSOInit.PrimitiveType = PT_TriangleList;
-
-	SetGraphicsPipelineState(Context.RHICmdList, GraphicsPSOInit);
-
-	PixelShader->SetParameters(Context.RHICmdList, Context);
-	VertexShader->SetParameters(Context);
-
-	// Draw a quad mapping scene color to the view's render target
-	DrawRectangle(
-		Context.RHICmdList,
-		DestRect.Min.X, DestRect.Min.Y,
-		DestRect.Width(), DestRect.Height(),
-		SrcRect.Min.X, SrcRect.Min.Y,
-		SrcRect.Width(), SrcRect.Height(),
-		DestRect.Size(),
-		SrcSize,
-		*VertexShader,
-		EDRF_UseTriangleOptimization);
-
-	Context.RHICmdList.CopyToResolveTarget(DestRenderTarget.TargetableTexture, DestRenderTarget.ShaderResourceTexture, FResolveParams());
+	
+	FRHIRenderPassInfo RPInfo(DestRenderTarget.TargetableTexture, ERenderTargetActions::DontLoad_Store, DestRenderTarget.ShaderResourceTexture);
+	Context.RHICmdList.BeginRenderPass(RPInfo, TEXT("VisualizeBloom"));
+	{
+		Context.SetViewportAndCallRHI(0, 0, 0.0f, DestRect.Width(), DestRect.Height(), 1.0f);
+
+		FGraphicsPipelineStateInitializer GraphicsPSOInit;
+		Context.RHICmdList.ApplyCachedRenderTargets(GraphicsPSOInit);
+
+		// set the state
+		GraphicsPSOInit.BlendState = TStaticBlendState<>::GetRHI();
+		GraphicsPSOInit.RasterizerState = TStaticRasterizerState<>::GetRHI();
+		GraphicsPSOInit.DepthStencilState = TStaticDepthStencilState<false, CF_Always>::GetRHI();
+
+		TShaderMapRef<FPostProcessVS> VertexShader(Context.GetShaderMap());
+		TShaderMapRef<FPostProcessVisualizeBloomSetupPS> PixelShader(Context.GetShaderMap());
+
+		GraphicsPSOInit.BoundShaderState.VertexDeclarationRHI = GFilterVertexDeclaration.VertexDeclarationRHI;
+		GraphicsPSOInit.BoundShaderState.VertexShaderRHI = GETSAFERHISHADER_VERTEX(*VertexShader);
+		GraphicsPSOInit.BoundShaderState.PixelShaderRHI = GETSAFERHISHADER_PIXEL(*PixelShader);
+		GraphicsPSOInit.PrimitiveType = PT_TriangleList;
+
+		SetGraphicsPipelineState(Context.RHICmdList, GraphicsPSOInit);
+
+		PixelShader->SetParameters(Context.RHICmdList, Context);
+		VertexShader->SetParameters(Context);
+
+		// Draw a quad mapping scene color to the view's render target
+		DrawRectangle(
+			Context.RHICmdList,
+			DestRect.Min.X, DestRect.Min.Y,
+			DestRect.Width(), DestRect.Height(),
+			SrcRect.Min.X, SrcRect.Min.Y,
+			SrcRect.Width(), SrcRect.Height(),
+			DestRect.Size(),
+			SrcSize,
+			*VertexShader,
+			EDRF_UseTriangleOptimization);
+	}
+	Context.RHICmdList.EndRenderPass();
 }
 
 
@@ -569,47 +571,47 @@ void FRCPassPostProcessVisualizeBloomOverlay::Process(FRenderingCompositePassCon
 
 	const FSceneRenderTargetItem& DestRenderTarget = PassOutputs[0].RequestSurface(Context);
 
-	SetRenderTarget(Context.RHICmdList, DestRenderTarget.TargetableTexture, FTextureRHIRef());
-
-	// is optimized away if possible (RT size=view size, )
-	DrawClearQuad(Context.RHICmdList, true, FLinearColor(0, 0, 0 ,0), false, 0, false, 0, DestSize, DestRect);
+	FRHIRenderPassInfo RPInfo(DestRenderTarget.TargetableTexture, ERenderTargetActions::Load_Store, DestRenderTarget.ShaderResourceTexture);
 
-	Context.SetViewportAndCallRHI(0, 0, 0.0f, DestRect.Width(), DestRect.Height(), 1.0f );
-
-	FGraphicsPipelineStateInitializer GraphicsPSOInit;
-	Context.RHICmdList.ApplyCachedRenderTargets(GraphicsPSOInit);
+	Context.RHICmdList.BeginRenderPass(RPInfo, TEXT("VisualizeBloomOverlay"));
+	{
+		DrawClearQuad(Context.RHICmdList, true, FLinearColor(0, 0, 0, 0), false, 0, false, 0, DestSize, DestRect);
+		Context.SetViewportAndCallRHI(0, 0, 0.0f, DestRect.Width(), DestRect.Height(), 1.0f);
 
-	// set the state
-	GraphicsPSOInit.BlendState = TStaticBlendState<>::GetRHI();
-	GraphicsPSOInit.RasterizerState = TStaticRasterizerState<>::GetRHI();
-	GraphicsPSOInit.DepthStencilState = TStaticDepthStencilState<false, CF_Always>::GetRHI();
+		FGraphicsPipelineStateInitializer GraphicsPSOInit;
+		Context.RHICmdList.ApplyCachedRenderTargets(GraphicsPSOInit);
 
-	TShaderMapRef<FPostProcessVS> VertexShader(Context.GetShaderMap());
-	TShaderMapRef<FPostProcessVisualizeBloomOverlayPS> PixelShader(Context.GetShaderMap());
+		// set the state
+		GraphicsPSOInit.BlendState = TStaticBlendState<>::GetRHI();
+		GraphicsPSOInit.RasterizerState = TStaticRasterizerState<>::GetRHI();
+		GraphicsPSOInit.DepthStencilState = TStaticDepthStencilState<false, CF_Always>::GetRHI();
 
-	GraphicsPSOInit.BoundShaderState.VertexDeclarationRHI = GFilterVertexDeclaration.VertexDeclarationRHI;
-	GraphicsPSOInit.BoundShaderState.VertexShaderRHI = GETSAFERHISHADER_VERTEX(*VertexShader);
-	GraphicsPSOInit.BoundShaderState.PixelShaderRHI = GETSAFERHISHADER_PIXEL(*PixelShader);
-	GraphicsPSOInit.PrimitiveType = PT_TriangleList;
+		TShaderMapRef<FPostProcessVS> VertexShader(Context.GetShaderMap());
+		TShaderMapRef<FPostProcessVisualizeBloomOverlayPS> PixelShader(Context.GetShaderMap());
 
-	SetGraphicsPipelineState(Context.RHICmdList, GraphicsPSOInit);
+		GraphicsPSOInit.BoundShaderState.VertexDeclarationRHI = GFilterVertexDeclaration.VertexDeclarationRHI;
+		GraphicsPSOInit.BoundShaderState.VertexShaderRHI = GETSAFERHISHADER_VERTEX(*VertexShader);
+		GraphicsPSOInit.BoundShaderState.PixelShaderRHI = GETSAFERHISHADER_PIXEL(*PixelShader);
+		GraphicsPSOInit.PrimitiveType = PT_TriangleList;
 
-	PixelShader->SetParameters(Context);
-	VertexShader->SetParameters(Context);
+		SetGraphicsPipelineState(Context.RHICmdList, GraphicsPSOInit);
 
-	// Draw a quad mapping scene color to the view's render target
-	DrawRectangle(
-		Context.RHICmdList,
-		DestRect.Min.X, DestRect.Min.Y,
-		DestRect.Width(), DestRect.Height(),
-		SrcRect.Min.X, SrcRect.Min.Y,
-		SrcRect.Width(), SrcRect.Height(),
-		DestRect.Size(),
-		SrcSize,
-		*VertexShader,
-		EDRF_UseTriangleOptimization);
+		PixelShader->SetParameters(Context);
+		VertexShader->SetParameters(Context);
 
-	Context.RHICmdList.CopyToResolveTarget(DestRenderTarget.TargetableTexture, DestRenderTarget.ShaderResourceTexture, FResolveParams());
+		// Draw a quad mapping scene color to the view's render target
+		DrawRectangle(
+			Context.RHICmdList,
+			DestRect.Min.X, DestRect.Min.Y,
+			DestRect.Width(), DestRect.Height(),
+			SrcRect.Min.X, SrcRect.Min.Y,
+			SrcRect.Width(), SrcRect.Height(),
+			DestRect.Size(),
+			SrcSize,
+			*VertexShader,
+			EDRF_UseTriangleOptimization);
+	}
+	Context.RHICmdList.EndRenderPass();
 }
 
 
diff --git a/Engine/Source/Runtime/Renderer/Private/PostProcess/PostProcessBokehDOF.cpp b/Engine/Source/Runtime/Renderer/Private/PostProcess/PostProcessBokehDOF.cpp
--- a/Engine/Source/Runtime/Renderer/Private/PostProcess/PostProcessBokehDOF.cpp
+++ b/Engine/Source/Runtime/Renderer/Private/PostProcess/PostProcessBokehDOF.cpp
@@ -164,51 +164,53 @@ void FRCPassPostProcessVisualizeDOF::Process(FRenderingCompositePassContext& Con
 	const FSceneRenderTargetItem& DestRenderTarget = PassOutputs[0].RequestSurface(Context);
 
 	// Set the view family's render target/viewport.
-	SetRenderTarget(Context.RHICmdList, DestRenderTarget.TargetableTexture, FTextureRHIRef());
+	FRHIRenderPassInfo RPInfo(DestRenderTarget.TargetableTexture, ERenderTargetActions::Load_Store);
+	Context.RHICmdList.BeginRenderPass(RPInfo, TEXT("VisualizeDOF"));
+	{
+		// can be optimized (don't clear areas we overwrite, don't clear when full screen),
+		// needed when a camera (matinee) has black borders or with multiple viewports
+		// focal distance depth is stored in the alpha channel to avoid DOF artifacts
+		DrawClearQuad(Context.RHICmdList, true, FLinearColor(0, 0, 0, View.FinalPostProcessSettings.DepthOfFieldFocalDistance), false, 0, false, 0, PassOutputs[0].RenderTargetDesc.Extent, DestRect);
 
-	// can be optimized (don't clear areas we overwrite, don't clear when full screen),
-	// needed when a camera (matinee) has black borders or with multiple viewports
-	// focal distance depth is stored in the alpha channel to avoid DOF artifacts
-	DrawClearQuad(Context.RHICmdList, true, FLinearColor(0, 0, 0, View.FinalPostProcessSettings.DepthOfFieldFocalDistance), false, 0, false, 0, PassOutputs[0].RenderTargetDesc.Extent, DestRect);
+		Context.SetViewportAndCallRHI(0, 0, 0.0f, DestSize.X, DestSize.Y, 1.0f);
 
-	Context.SetViewportAndCallRHI(0, 0, 0.0f, DestSize.X, DestSize.Y, 1.0f );
+		FGraphicsPipelineStateInitializer GraphicsPSOInit;
+		Context.RHICmdList.ApplyCachedRenderTargets(GraphicsPSOInit);
 
-	FGraphicsPipelineStateInitializer GraphicsPSOInit;
-	Context.RHICmdList.ApplyCachedRenderTargets(GraphicsPSOInit);
+		// set the state
+		GraphicsPSOInit.BlendState = TStaticBlendState<>::GetRHI();
+		GraphicsPSOInit.RasterizerState = TStaticRasterizerState<>::GetRHI();
+		GraphicsPSOInit.DepthStencilState = TStaticDepthStencilState<false, CF_Always>::GetRHI();
 
-	// set the state
-	GraphicsPSOInit.BlendState = TStaticBlendState<>::GetRHI();
-	GraphicsPSOInit.RasterizerState = TStaticRasterizerState<>::GetRHI();
-	GraphicsPSOInit.DepthStencilState = TStaticDepthStencilState<false, CF_Always>::GetRHI();
+		// setup shader
+		TShaderMapRef<FPostProcessVS> VertexShader(Context.GetShaderMap());
+		GraphicsPSOInit.BoundShaderState.VertexDeclarationRHI = GFilterVertexDeclaration.VertexDeclarationRHI;
+		GraphicsPSOInit.BoundShaderState.VertexShaderRHI = GETSAFERHISHADER_VERTEX(*VertexShader);
+		GraphicsPSOInit.PrimitiveType = PT_TriangleList;
 
-	// setup shader
-	TShaderMapRef<FPostProcessVS> VertexShader(Context.GetShaderMap());
-	GraphicsPSOInit.BoundShaderState.VertexDeclarationRHI = GFilterVertexDeclaration.VertexDeclarationRHI;
-	GraphicsPSOInit.BoundShaderState.VertexShaderRHI = GETSAFERHISHADER_VERTEX(*VertexShader);
-	GraphicsPSOInit.PrimitiveType = PT_TriangleList;
+		{
+			TShaderMapRef<PostProcessVisualizeDOFPS> PixelShader(Context.GetShaderMap());
 
-	{
-		TShaderMapRef<PostProcessVisualizeDOFPS> PixelShader(Context.GetShaderMap());
+			GraphicsPSOInit.BoundShaderState.PixelShaderRHI = GETSAFERHISHADER_PIXEL(*PixelShader);
+			SetGraphicsPipelineState(Context.RHICmdList, GraphicsPSOInit);
 
-		GraphicsPSOInit.BoundShaderState.PixelShaderRHI = GETSAFERHISHADER_PIXEL(*PixelShader);
-		SetGraphicsPipelineState(Context.RHICmdList, GraphicsPSOInit);
+			VertexShader->SetParameters(Context);
+			PixelShader->SetParameters(Context.RHICmdList, Context, DepthOfFieldStats);
+		}
 
-		VertexShader->SetParameters(Context);
-		PixelShader->SetParameters(Context.RHICmdList, Context, DepthOfFieldStats);
+		// Draw a quad mapping scene color to the view's render target
+		DrawRectangle(
+			Context.RHICmdList,
+			DestRect.Min.X, DestRect.Min.Y,
+			DestRect.Width(), DestRect.Height(),
+			SrcRect.Min.X, SrcRect.Min.Y,
+			SrcRect.Width(), SrcRect.Height(),
+			DestSize,
+			SrcSize,
+			*VertexShader,
+			EDRF_UseTriangleOptimization);
 	}
-
-	// Draw a quad mapping scene color to the view's render target
-	DrawRectangle(
-		Context.RHICmdList,
-		DestRect.Min.X, DestRect.Min.Y,
-		DestRect.Width(), DestRect.Height(),
-		SrcRect.Min.X, SrcRect.Min.Y,
-		SrcRect.Width(), SrcRect.Height(),
-		DestSize,
-		SrcSize,
-		*VertexShader,
-		EDRF_UseTriangleOptimization);
-
+	Context.RHICmdList.EndRenderPass();
 
 	{
 		FRenderTargetTemp TempRenderTarget(View, (const FTexture2DRHIRef&)DestRenderTarget.TargetableTexture);
@@ -226,100 +228,101 @@ void FRCPassPostProcessVisualizeDOF::Process(FRenderingCompositePassContext& Con
 		Y += YStep;
 
 		PRAGMA_DISABLE_DEPRECATION_WARNINGS
-		EDepthOfFieldMethod MethodId = View.FinalPostProcessSettings.DepthOfFieldMethod;
+			EDepthOfFieldMethod MethodId = View.FinalPostProcessSettings.DepthOfFieldMethod;
 		PRAGMA_ENABLE_DEPRECATION_WARNINGS
 
-		if(MethodId == DOFM_BokehDOF)
-		{
-			PRAGMA_DISABLE_DEPRECATION_WARNINGS
-			Line = FString::Printf(TEXT("Method: BokehDOF (blue is far, green is near, black is in focus)"));
-			Canvas.DrawShadowedString(X, Y += YStep, *Line, GetStatsFont(), FLinearColor(1, 1, 1));
-			Y += YStep;
-			Line = FString::Printf(TEXT("FocalDistance: %.2f"), View.FinalPostProcessSettings.DepthOfFieldFocalDistance);
-			Canvas.DrawShadowedString(X, Y += YStep, *Line, GetStatsFont(), FLinearColor(1, 1, 1));
-			Line = FString::Printf(TEXT("FocalRegion (Artificial, avoid): %.2f"), View.FinalPostProcessSettings.DepthOfFieldFocalRegion);
-			Canvas.DrawShadowedString(X, Y += YStep, *Line, GetStatsFont(), FLinearColor(1, 1, 1));
-			Y += YStep;
-			Line = FString::Printf(TEXT("Scale: %.2f"), View.FinalPostProcessSettings.DepthOfFieldScale);
-			Canvas.DrawShadowedString(X, Y += YStep, *Line, GetStatsFont(), FLinearColor(1, 1, 1));
-			Line = FString::Printf(TEXT("MaxBokehSize: %.2f"), View.FinalPostProcessSettings.DepthOfFieldMaxBokehSize);
-			Canvas.DrawShadowedString(X, Y += YStep, *Line, GetStatsFont(), FLinearColor(1, 1, 1));
-			Line = FString::Printf(TEXT("NearTransitionRegion: %.2f"), View.FinalPostProcessSettings.DepthOfFieldNearTransitionRegion);
-			Canvas.DrawShadowedString(X, Y += YStep, *Line, GetStatsFont(), FLinearColor(1, 1, 1));
-			Line = FString::Printf(TEXT("FarTransitionRegion: %.2f"), View.FinalPostProcessSettings.DepthOfFieldFarTransitionRegion);
-			Canvas.DrawShadowedString(X, Y += YStep, *Line, GetStatsFont(), FLinearColor(1, 1, 1));
-			Line = FString::Printf(TEXT("ColorThreshold: %.2f"), View.FinalPostProcessSettings.DepthOfFieldColorThreshold);
-			Canvas.DrawShadowedString(X, Y += YStep, *Line, GetStatsFont(), FLinearColor(1, 1, 1));
-			Line = FString::Printf(TEXT("SizeThreshold: %.2f"), View.FinalPostProcessSettings.DepthOfFieldSizeThreshold);
-			Canvas.DrawShadowedString(X, Y += YStep, *Line, GetStatsFont(), FLinearColor(1, 1, 1));
-			Line = FString::Printf(TEXT("Occlusion: %.2f"), View.FinalPostProcessSettings.DepthOfFieldOcclusion);
-			Canvas.DrawShadowedString(X, Y += YStep, *Line, GetStatsFont(), FLinearColor(1, 1, 1));
-			PRAGMA_ENABLE_DEPRECATION_WARNINGS
-		}
-		else if(MethodId == DOFM_Gaussian)
-		{
-			Line = FString::Printf(TEXT("Method: GaussianDOF (blue is far, green is near, grey is disabled, black is in focus)"));
-			Canvas.DrawShadowedString(X, Y += YStep, *Line, GetStatsFont(), FLinearColor(1, 1, 1));
-			Y += YStep;
-			Line = FString::Printf(TEXT("FocalDistance: %.2f"), View.FinalPostProcessSettings.DepthOfFieldFocalDistance);
-			Canvas.DrawShadowedString(X, Y += YStep, *Line, GetStatsFont(), FLinearColor(1, 1, 1));
-			Line = FString::Printf(TEXT("FocalRegion (Artificial, avoid): %.2f"), View.FinalPostProcessSettings.DepthOfFieldFocalRegion);
-			Canvas.DrawShadowedString(X, Y += YStep, *Line, GetStatsFont(), FLinearColor(1, 1, 1));
-			Y += YStep;
-			Line = FString::Printf(TEXT("NearTransitionRegion: %.2f"), View.FinalPostProcessSettings.DepthOfFieldNearTransitionRegion);
-			Canvas.DrawShadowedString(X, Y += YStep, *Line, GetStatsFont(), FLinearColor(1, 1, 1));
-			Line = FString::Printf(TEXT("FarTransitionRegion: %.2f"), View.FinalPostProcessSettings.DepthOfFieldFarTransitionRegion);
-			Canvas.DrawShadowedString(X, Y += YStep, *Line, GetStatsFont(), FLinearColor(1, 1, 1));
-			Line = FString::Printf(TEXT("NearBlurSize: %.2f"), View.FinalPostProcessSettings.DepthOfFieldNearBlurSize);
-			Canvas.DrawShadowedString(X, Y += YStep, *Line, GetStatsFont(), FLinearColor(1, 1, 1));
-			Line = FString::Printf(TEXT("FarBlurSize: %.2f"), View.FinalPostProcessSettings.DepthOfFieldFarBlurSize);
-			Canvas.DrawShadowedString(X, Y += YStep, *Line, GetStatsFont(), FLinearColor(1, 1, 1));
-			Line = FString::Printf(TEXT("Occlusion: %.2f"), View.FinalPostProcessSettings.DepthOfFieldOcclusion);
-			Canvas.DrawShadowedString(X, Y += YStep, *Line, GetStatsFont(), FLinearColor(1, 1, 1));
-			Line = FString::Printf(TEXT("SkyFocusDistance: %.2f"), View.FinalPostProcessSettings.DepthOfFieldSkyFocusDistance);
-			Canvas.DrawShadowedString(X, Y += YStep, *Line, GetStatsFont(), FLinearColor(1, 1, 1));
-			Line = FString::Printf(TEXT("VignetteRadius: %.2f"), View.FinalPostProcessSettings.DepthOfFieldVignetteSize);
-			Canvas.DrawShadowedString(X, Y += YStep, *Line, GetStatsFont(), FLinearColor(1, 1, 1));
-			Y += YStep;
-			Line = FString::Printf(TEXT("Near:%d Far:%d"), DepthOfFieldStats.bNear ? 1 : 0, DepthOfFieldStats.bFar ? 1 : 0);
-			Canvas.DrawShadowedString(X, Y += YStep, *Line, GetStatsFont(), FLinearColor(1, 1, 1));
-		}
-		else if(MethodId == DOFM_CircleDOF)
-		{
-			Line = FString::Printf(TEXT("Method: CircleDOF (blue is far, green is near, black is in focus, cross hair shows Depth and CoC radius in pixel)"));
-			Canvas.DrawShadowedString(X, Y += YStep, *Line, GetStatsFont(), FLinearColor(1, 1, 1));
-			Y += YStep;
-			Line = FString::Printf(TEXT("FocalDistance: %.2f"), View.FinalPostProcessSettings.DepthOfFieldFocalDistance);
-			Canvas.DrawShadowedString(X, Y += YStep, *Line, GetStatsFont(), FLinearColor(1, 1, 1));
-			Line = FString::Printf(TEXT("Aperture F-stop: %.2f"), View.FinalPostProcessSettings.DepthOfFieldFstop);
-			Canvas.DrawShadowedString(X, Y += YStep, *Line, GetStatsFont(), FLinearColor(1, 1, 1));
-			Line = FString::Printf(TEXT("Aperture: f/%.2f"), View.FinalPostProcessSettings.DepthOfFieldFstop);
-			Canvas.DrawShadowedString(X, Y += YStep, *Line, GetStatsFont(), FLinearColor(1, 1, 1));
-			Y += YStep;
-			Line = FString::Printf(TEXT("DepthBlur (not related to Depth of Field, due to light traveling long distances in atmosphere)"));
-			Canvas.DrawShadowedString(X, Y += YStep, *Line, GetStatsFont(), FLinearColor(1, 1, 1));
-			Line = FString::Printf(TEXT("  km for 50%: %.2f"), View.FinalPostProcessSettings.DepthOfFieldDepthBlurAmount);
-			Canvas.DrawShadowedString(X, Y += YStep, *Line, GetStatsFont(), FLinearColor(1, 1, 1));
-			Line = FString::Printf(TEXT("  Radius (pixels in 1920x): %.2f"), View.FinalPostProcessSettings.DepthOfFieldDepthBlurRadius);
-			Canvas.DrawShadowedString(X, Y += YStep, *Line, GetStatsFont(), FLinearColor(1, 1, 1));
-			Y += YStep;
-
-			FVector2D Fov = View.ViewMatrices.ComputeHalfFieldOfViewPerAxis();
-			
-			float FocalLength = ComputeFocalLengthFromFov(View);
-
-			Line = FString::Printf(TEXT("Field Of View in deg. (computed): %.1f x %.1f"), FMath::RadiansToDegrees(Fov.X) * 2.0f, FMath::RadiansToDegrees(Fov.Y) * 2.0f);
-			Canvas.DrawShadowedString(X, Y += YStep, *Line, GetStatsFont(), FLinearColor(0.5, 0.5, 1));
-			Line = FString::Printf(TEXT("Focal Length (computed): %.1f"), FocalLength);
-			Canvas.DrawShadowedString(X, Y += YStep, *Line, GetStatsFont(), FLinearColor(0.5, 0.5, 10));
-			Line = FString::Printf(TEXT("Sensor: APS-C 24.576 mm sensor, crop-factor 1.61x"));
-			Canvas.DrawShadowedString(X, Y += YStep, *Line, GetStatsFont(), FLinearColor(0.5, 0.5, 1));
-		}
+			if (MethodId == DOFM_BokehDOF)
+			{
+				PRAGMA_DISABLE_DEPRECATION_WARNINGS
+					Line = FString::Printf(TEXT("Method: BokehDOF (blue is far, green is near, black is in focus)"));
+				Canvas.DrawShadowedString(X, Y += YStep, *Line, GetStatsFont(), FLinearColor(1, 1, 1));
+				Y += YStep;
+				Line = FString::Printf(TEXT("FocalDistance: %.2f"), View.FinalPostProcessSettings.DepthOfFieldFocalDistance);
+				Canvas.DrawShadowedString(X, Y += YStep, *Line, GetStatsFont(), FLinearColor(1, 1, 1));
+				Line = FString::Printf(TEXT("FocalRegion (Artificial, avoid): %.2f"), View.FinalPostProcessSettings.DepthOfFieldFocalRegion);
+				Canvas.DrawShadowedString(X, Y += YStep, *Line, GetStatsFont(), FLinearColor(1, 1, 1));
+				Y += YStep;
+				Line = FString::Printf(TEXT("Scale: %.2f"), View.FinalPostProcessSettings.DepthOfFieldScale);
+				Canvas.DrawShadowedString(X, Y += YStep, *Line, GetStatsFont(), FLinearColor(1, 1, 1));
+				Line = FString::Printf(TEXT("MaxBokehSize: %.2f"), View.FinalPostProcessSettings.DepthOfFieldMaxBokehSize);
+				Canvas.DrawShadowedString(X, Y += YStep, *Line, GetStatsFont(), FLinearColor(1, 1, 1));
+				Line = FString::Printf(TEXT("NearTransitionRegion: %.2f"), View.FinalPostProcessSettings.DepthOfFieldNearTransitionRegion);
+				Canvas.DrawShadowedString(X, Y += YStep, *Line, GetStatsFont(), FLinearColor(1, 1, 1));
+				Line = FString::Printf(TEXT("FarTransitionRegion: %.2f"), View.FinalPostProcessSettings.DepthOfFieldFarTransitionRegion);
+				Canvas.DrawShadowedString(X, Y += YStep, *Line, GetStatsFont(), FLinearColor(1, 1, 1));
+				Line = FString::Printf(TEXT("ColorThreshold: %.2f"), View.FinalPostProcessSettings.DepthOfFieldColorThreshold);
+				Canvas.DrawShadowedString(X, Y += YStep, *Line, GetStatsFont(), FLinearColor(1, 1, 1));
+				Line = FString::Printf(TEXT("SizeThreshold: %.2f"), View.FinalPostProcessSettings.DepthOfFieldSizeThreshold);
+				Canvas.DrawShadowedString(X, Y += YStep, *Line, GetStatsFont(), FLinearColor(1, 1, 1));
+				Line = FString::Printf(TEXT("Occlusion: %.2f"), View.FinalPostProcessSettings.DepthOfFieldOcclusion);
+				Canvas.DrawShadowedString(X, Y += YStep, *Line, GetStatsFont(), FLinearColor(1, 1, 1));
+				PRAGMA_ENABLE_DEPRECATION_WARNINGS
+			}
+			else if (MethodId == DOFM_Gaussian)
+			{
+				Line = FString::Printf(TEXT("Method: GaussianDOF (blue is far, green is near, grey is disabled, black is in focus)"));
+				Canvas.DrawShadowedString(X, Y += YStep, *Line, GetStatsFont(), FLinearColor(1, 1, 1));
+				Y += YStep;
+				Line = FString::Printf(TEXT("FocalDistance: %.2f"), View.FinalPostProcessSettings.DepthOfFieldFocalDistance);
+				Canvas.DrawShadowedString(X, Y += YStep, *Line, GetStatsFont(), FLinearColor(1, 1, 1));
+				Line = FString::Printf(TEXT("FocalRegion (Artificial, avoid): %.2f"), View.FinalPostProcessSettings.DepthOfFieldFocalRegion);
+				Canvas.DrawShadowedString(X, Y += YStep, *Line, GetStatsFont(), FLinearColor(1, 1, 1));
+				Y += YStep;
+				Line = FString::Printf(TEXT("NearTransitionRegion: %.2f"), View.FinalPostProcessSettings.DepthOfFieldNearTransitionRegion);
+				Canvas.DrawShadowedString(X, Y += YStep, *Line, GetStatsFont(), FLinearColor(1, 1, 1));
+				Line = FString::Printf(TEXT("FarTransitionRegion: %.2f"), View.FinalPostProcessSettings.DepthOfFieldFarTransitionRegion);
+				Canvas.DrawShadowedString(X, Y += YStep, *Line, GetStatsFont(), FLinearColor(1, 1, 1));
+				Line = FString::Printf(TEXT("NearBlurSize: %.2f"), View.FinalPostProcessSettings.DepthOfFieldNearBlurSize);
+				Canvas.DrawShadowedString(X, Y += YStep, *Line, GetStatsFont(), FLinearColor(1, 1, 1));
+				Line = FString::Printf(TEXT("FarBlurSize: %.2f"), View.FinalPostProcessSettings.DepthOfFieldFarBlurSize);
+				Canvas.DrawShadowedString(X, Y += YStep, *Line, GetStatsFont(), FLinearColor(1, 1, 1));
+				Line = FString::Printf(TEXT("Occlusion: %.2f"), View.FinalPostProcessSettings.DepthOfFieldOcclusion);
+				Canvas.DrawShadowedString(X, Y += YStep, *Line, GetStatsFont(), FLinearColor(1, 1, 1));
+				Line = FString::Printf(TEXT("SkyFocusDistance: %.2f"), View.FinalPostProcessSettings.DepthOfFieldSkyFocusDistance);
+				Canvas.DrawShadowedString(X, Y += YStep, *Line, GetStatsFont(), FLinearColor(1, 1, 1));
+				Line = FString::Printf(TEXT("VignetteRadius: %.2f"), View.FinalPostProcessSettings.DepthOfFieldVignetteSize);
+				Canvas.DrawShadowedString(X, Y += YStep, *Line, GetStatsFont(), FLinearColor(1, 1, 1));
+				Y += YStep;
+				Line = FString::Printf(TEXT("Near:%d Far:%d"), DepthOfFieldStats.bNear ? 1 : 0, DepthOfFieldStats.bFar ? 1 : 0);
+				Canvas.DrawShadowedString(X, Y += YStep, *Line, GetStatsFont(), FLinearColor(1, 1, 1));
+			}
+			else if (MethodId == DOFM_CircleDOF)
+			{
+				Line = FString::Printf(TEXT("Method: CircleDOF (blue is far, green is near, black is in focus, cross hair shows Depth and CoC radius in pixel)"));
+				Canvas.DrawShadowedString(X, Y += YStep, *Line, GetStatsFont(), FLinearColor(1, 1, 1));
+				Y += YStep;
+				Line = FString::Printf(TEXT("FocalDistance: %.2f"), View.FinalPostProcessSettings.DepthOfFieldFocalDistance);
+				Canvas.DrawShadowedString(X, Y += YStep, *Line, GetStatsFont(), FLinearColor(1, 1, 1));
+				Line = FString::Printf(TEXT("Aperture F-stop: %.2f"), View.FinalPostProcessSettings.DepthOfFieldFstop);
+				Canvas.DrawShadowedString(X, Y += YStep, *Line, GetStatsFont(), FLinearColor(1, 1, 1));
+				Line = FString::Printf(TEXT("Aperture: f/%.2f"), View.FinalPostProcessSettings.DepthOfFieldFstop);
+				Canvas.DrawShadowedString(X, Y += YStep, *Line, GetStatsFont(), FLinearColor(1, 1, 1));
+				Y += YStep;
+				Line = FString::Printf(TEXT("DepthBlur (not related to Depth of Field, due to light traveling long distances in atmosphere)"));
+				Canvas.DrawShadowedString(X, Y += YStep, *Line, GetStatsFont(), FLinearColor(1, 1, 1));
+				Line = FString::Printf(TEXT("  km for 50%: %.2f"), View.FinalPostProcessSettings.DepthOfFieldDepthBlurAmount);
+				Canvas.DrawShadowedString(X, Y += YStep, *Line, GetStatsFont(), FLinearColor(1, 1, 1));
+				Line = FString::Printf(TEXT("  Radius (pixels in 1920x): %.2f"), View.FinalPostProcessSettings.DepthOfFieldDepthBlurRadius);
+				Canvas.DrawShadowedString(X, Y += YStep, *Line, GetStatsFont(), FLinearColor(1, 1, 1));
+				Y += YStep;
+
+				FVector2D Fov = View.ViewMatrices.ComputeHalfFieldOfViewPerAxis();
+
+				float FocalLength = ComputeFocalLengthFromFov(View);
+
+				Line = FString::Printf(TEXT("Field Of View in deg. (computed): %.1f x %.1f"), FMath::RadiansToDegrees(Fov.X) * 2.0f, FMath::RadiansToDegrees(Fov.Y) * 2.0f);
+				Canvas.DrawShadowedString(X, Y += YStep, *Line, GetStatsFont(), FLinearColor(0.5, 0.5, 1));
+				Line = FString::Printf(TEXT("Focal Length (computed): %.1f"), FocalLength);
+				Canvas.DrawShadowedString(X, Y += YStep, *Line, GetStatsFont(), FLinearColor(0.5, 0.5, 10));
+				Line = FString::Printf(TEXT("Sensor: APS-C 24.576 mm sensor, crop-factor 1.61x"));
+				Canvas.DrawShadowedString(X, Y += YStep, *Line, GetStatsFont(), FLinearColor(0.5, 0.5, 1));
+			}
 
 		Canvas.Flush_RenderThread(Context.RHICmdList);
 	}
 
 	Context.RHICmdList.CopyToResolveTarget(DestRenderTarget.TargetableTexture, DestRenderTarget.ShaderResourceTexture, FResolveParams());
+
 }
 
 FPooledRenderTargetDesc FRCPassPostProcessVisualizeDOF::ComputeOutputDesc(EPassOutputId InPassOutputId) const
@@ -501,7 +504,8 @@ void FRCPassPostProcessBokehDOFSetup::Process(FRenderingCompositePassContext& Co
 	{
 		DestRect = {View.ViewRect.Min, View.ViewRect.Min + DestSize};
 
-		SetRenderTarget(Context.RHICmdList, nullptr, nullptr);
+		// #todo-renderpasses Remove once everything is converted to renderpasses
+		UnbindRenderTargets(Context.RHICmdList);
 		Context.SetViewportAndCallRHI(DestRect, 0.0f, 1.0f);
 
 		WaitForInputPassComputeFences(Context.RHICmdList);
@@ -514,56 +518,58 @@ void FRCPassPostProcessBokehDOFSetup::Process(FRenderingCompositePassContext& Co
 		WaitForInputPassComputeFences(Context.RHICmdList);
 
 		// Set the view family's render target/viewport.
-		SetRenderTarget(Context.RHICmdList, DestRenderTarget.TargetableTexture, FTextureRHIRef());
-	
-		if (View.StereoPass == eSSP_FULL)
+		FRHIRenderPassInfo RPInfo(DestRenderTarget.TargetableTexture, ERenderTargetActions::Load_Store, DestRenderTarget.ShaderResourceTexture);
+		Context.RHICmdList.BeginRenderPass(RPInfo, TEXT("BokehDOFSetup"));
 		{
-			// can be optimized (don't clear areas we overwrite, don't clear when full screen),
-			// needed when a camera (matinee) has black borders or with multiple viewports
-			// focal distance depth is stored in the alpha channel to avoid DOF artifacts
-			DrawClearQuad(Context.RHICmdList, true, FLinearColor(0, 0, 0, View.FinalPostProcessSettings.DepthOfFieldFocalDistance), false, 0, false, 0, DestSize, DestRect);
-		}
 
-		Context.SetViewportAndCallRHI(0, 0, 0.0f, DestSize.X, DestSize.Y, 1.0f );
+			if (View.StereoPass == eSSP_FULL)
+			{
+				// can be optimized (don't clear areas we overwrite, don't clear when full screen),
+				// needed when a camera (matinee) has black borders or with multiple viewports
+				// focal distance depth is stored in the alpha channel to avoid DOF artifacts
+				DrawClearQuad(Context.RHICmdList, true, FLinearColor(0, 0, 0, View.FinalPostProcessSettings.DepthOfFieldFocalDistance), false, 0, false, 0, DestSize, DestRect);
+			}
 
-		FGraphicsPipelineStateInitializer GraphicsPSOInit;
-		Context.RHICmdList.ApplyCachedRenderTargets(GraphicsPSOInit);
+			Context.SetViewportAndCallRHI(0, 0, 0.0f, DestSize.X, DestSize.Y, 1.0f);
 
-		// set the state
-		GraphicsPSOInit.BlendState = TStaticBlendState<>::GetRHI();
-		GraphicsPSOInit.RasterizerState = TStaticRasterizerState<>::GetRHI();
-		GraphicsPSOInit.DepthStencilState = TStaticDepthStencilState<false, CF_Always>::GetRHI();
+			FGraphicsPipelineStateInitializer GraphicsPSOInit;
+			Context.RHICmdList.ApplyCachedRenderTargets(GraphicsPSOInit);
 
-		// setup shader
-		TShaderMapRef<FPostProcessVS> VertexShader(Context.GetShaderMap());
-		GraphicsPSOInit.BoundShaderState.VertexDeclarationRHI = GFilterVertexDeclaration.VertexDeclarationRHI;
-		GraphicsPSOInit.BoundShaderState.VertexShaderRHI = GETSAFERHISHADER_VERTEX(*VertexShader);
-		GraphicsPSOInit.PrimitiveType = PT_TriangleList;
+			// set the state
+			GraphicsPSOInit.BlendState = TStaticBlendState<>::GetRHI();
+			GraphicsPSOInit.RasterizerState = TStaticRasterizerState<>::GetRHI();
+			GraphicsPSOInit.DepthStencilState = TStaticDepthStencilState<false, CF_Always>::GetRHI();
 
-		{
-			TShaderMapRef<PostProcessBokehDOFSetupPS> PixelShader(Context.GetShaderMap());
+			// setup shader
+			TShaderMapRef<FPostProcessVS> VertexShader(Context.GetShaderMap());
+			GraphicsPSOInit.BoundShaderState.VertexDeclarationRHI = GFilterVertexDeclaration.VertexDeclarationRHI;
+			GraphicsPSOInit.BoundShaderState.VertexShaderRHI = GETSAFERHISHADER_VERTEX(*VertexShader);
+			GraphicsPSOInit.PrimitiveType = PT_TriangleList;
 
-			GraphicsPSOInit.BoundShaderState.PixelShaderRHI = GETSAFERHISHADER_PIXEL(*PixelShader);
-			SetGraphicsPipelineState(Context.RHICmdList, GraphicsPSOInit);
+			{
+				TShaderMapRef<PostProcessBokehDOFSetupPS> PixelShader(Context.GetShaderMap());
 
-			VertexShader->SetParameters(Context);
-			PixelShader->SetParameters(Context);
-		}
+				GraphicsPSOInit.BoundShaderState.PixelShaderRHI = GETSAFERHISHADER_PIXEL(*PixelShader);
+				SetGraphicsPipelineState(Context.RHICmdList, GraphicsPSOInit);
 
-		DrawPostProcessPass(
-			Context.RHICmdList,
-			DestRect.Min.X, DestRect.Min.Y,
-			DestRect.Width(), DestRect.Height(),
-			SrcRect.Min.X, SrcRect.Min.Y,
-			SrcRect.Width(), SrcRect.Height(),
-			DestSize,
-			SrcSize,
-			*VertexShader,
-			View.StereoPass,
-			Context.HasHmdMesh(),
-			EDRF_UseTriangleOptimization);
+				VertexShader->SetParameters(Context);
+				PixelShader->SetParameters(Context);
+			}
 
-		Context.RHICmdList.CopyToResolveTarget(DestRenderTarget.TargetableTexture, DestRenderTarget.ShaderResourceTexture, FResolveParams());
+			DrawPostProcessPass(
+				Context.RHICmdList,
+				DestRect.Min.X, DestRect.Min.Y,
+				DestRect.Width(), DestRect.Height(),
+				SrcRect.Min.X, SrcRect.Min.Y,
+				SrcRect.Width(), SrcRect.Height(),
+				DestSize,
+				SrcSize,
+				*VertexShader,
+				View.StereoPass,
+				Context.HasHmdMesh(),
+				EDRF_UseTriangleOptimization);
+		}
+		Context.RHICmdList.EndRenderPass();
 	}
 }
 
@@ -873,60 +879,63 @@ void FRCPassPostProcessBokehDOF::Process(FRenderingCompositePassContext& Context
 	WaitForInputPassComputeFences(Context.RHICmdList);
 
 	// Set the view family's render target/viewport.
-	SetRenderTarget(Context.RHICmdList, DestRenderTarget.TargetableTexture, FTextureRHIRef());
+	FRHIRenderPassInfo RPInfo(DestRenderTarget.TargetableTexture, ERenderTargetActions::Load_Store, DestRenderTarget.ShaderResourceTexture);
+	Context.RHICmdList.BeginRenderPass(RPInfo, TEXT("BokehDOF"));
+	{
 
-	// This clean is required to make the accumulation working
-	DrawClearQuad(Context.RHICmdList, true, FLinearColor(0, 0, 0, 0), false, 0, false, 0, GetOutput(ePId_Output0)->RenderTargetDesc.Extent, FIntRect());
+		// This clean is required to make the accumulation working
+		DrawClearQuad(Context.RHICmdList, true, FLinearColor(0, 0, 0, 0), false, 0, false, 0, GetOutput(ePId_Output0)->RenderTargetDesc.Extent, FIntRect());
 
-	// we need to output to the whole rendertarget
-	Context.SetViewportAndCallRHI(0, 0, 0.0f, PassOutputs[0].RenderTargetDesc.Extent.X, PassOutputs[0].RenderTargetDesc.Extent.Y, 1.0f);
-	
-	static const auto CVar = IConsoleManager::Get().FindTConsoleVariableDataInt(TEXT("r.DepthOfFieldQuality"));
-	check(CVar);
-	int32 DOFQuality = CVar->GetValueOnRenderThread();
-	check( DOFQuality > 0 );
+		// we need to output to the whole rendertarget
+		Context.SetViewportAndCallRHI(0, 0, 0.0f, PassOutputs[0].RenderTargetDesc.Extent.X, PassOutputs[0].RenderTargetDesc.Extent.Y, 1.0f);
 
-	bool bHighQuality = DOFQuality > 1;
+		static const auto CVar = IConsoleManager::Get().FindTConsoleVariableDataInt(TEXT("r.DepthOfFieldQuality"));
+		check(CVar);
+		int32 DOFQuality = CVar->GetValueOnRenderThread();
+		check(DOFQuality > 0);
 
-	// 1: one quad per 1 half res texel
-	// 2: one quad per 4 half res texel (faster, can alias - need to verify that with bilinear filtering)
-	uint32 TileSize = bHighQuality ? 1 : 2;
+		bool bHighQuality = DOFQuality > 1;
 
-	// input is half res, don't process last pixel line where we don't have input
-	LocalViewSize.X &= ~1;
-	LocalViewSize.Y &= ~1;
+		// 1: one quad per 1 half res texel
+		// 2: one quad per 4 half res texel (faster, can alias - need to verify that with bilinear filtering)
+		uint32 TileSize = bHighQuality ? 1 : 2;
 
-	FIntPoint TileCount = LocalViewSize / TileSize;
+		// input is half res, don't process last pixel line where we don't have input
+		LocalViewSize.X &= ~1;
+		LocalViewSize.Y &= ~1;
 
-	PRAGMA_DISABLE_DEPRECATION_WARNINGS
-	float PixelKernelSize = Context.View.FinalPostProcessSettings.DepthOfFieldMaxBokehSize / 100.0f * LocalViewSize.X;
-	PRAGMA_ENABLE_DEPRECATION_WARNINGS
+		FIntPoint TileCount = LocalViewSize / TileSize;
 
-	FIntPoint LeftTop = LocalViewRect.Min;
-	
-	if(bHighQuality)
-	{
-		if(View.Family->EngineShowFlags.VisualizeAdaptiveDOF)
+		PRAGMA_DISABLE_DEPRECATION_WARNINGS
+			float PixelKernelSize = Context.View.FinalPostProcessSettings.DepthOfFieldMaxBokehSize / 100.0f * LocalViewSize.X;
+		PRAGMA_ENABLE_DEPRECATION_WARNINGS
+
+			FIntPoint LeftTop = LocalViewRect.Min;
+
+		if (bHighQuality)
 		{
-			// high quality, visualize in red and green where we spend more performance
-			SetShaderTempl<2>(Context, LeftTop, TileCount, TileSize, PixelKernelSize);
+			if (View.Family->EngineShowFlags.VisualizeAdaptiveDOF)
+			{
+				// high quality, visualize in red and green where we spend more performance
+				SetShaderTempl<2>(Context, LeftTop, TileCount, TileSize, PixelKernelSize);
+			}
+			else
+			{
+				// high quality
+				SetShaderTempl<1>(Context, LeftTop, TileCount, TileSize, PixelKernelSize);
+			}
 		}
 		else
 		{
-			// high quality
-			SetShaderTempl<1>(Context, LeftTop, TileCount, TileSize, PixelKernelSize);
+			// low quality
+			SetShaderTempl<0>(Context, LeftTop, TileCount, TileSize, PixelKernelSize);
 		}
-	}
-	else
-	{
-		// low quality
-		SetShaderTempl<0>(Context, LeftTop, TileCount, TileSize, PixelKernelSize);
-	}
 
-	Context.RHICmdList.SetStreamSource(0, NULL, 0);
-	Context.RHICmdList.DrawIndexedPrimitive(GBokehIndexBuffer.IndexBufferRHI, 0, 0, 4 * GBokehDOFQuadsPerInstance, 0, 2 * GBokehDOFQuadsPerInstance, FMath::DivideAndRoundUp(TileCount.X * TileCount.Y, GBokehDOFQuadsPerInstance));
+		Context.RHICmdList.SetStreamSource(0, NULL, 0);
+		Context.RHICmdList.DrawIndexedPrimitive(GBokehIndexBuffer.IndexBufferRHI, 0, 0, 4 * GBokehDOFQuadsPerInstance, 0, 2 * GBokehDOFQuadsPerInstance, FMath::DivideAndRoundUp(TileCount.X * TileCount.Y, GBokehDOFQuadsPerInstance));
+	}
 
-	Context.RHICmdList.CopyToResolveTarget(DestRenderTarget.TargetableTexture, DestRenderTarget.ShaderResourceTexture, FResolveParams());
+	Context.RHICmdList.EndRenderPass();
 }
 
 FPooledRenderTargetDesc FRCPassPostProcessBokehDOF::ComputeOutputDesc(EPassOutputId InPassOutputId) const
diff --git a/Engine/Source/Runtime/Renderer/Private/PostProcess/PostProcessBokehDOFRecombine.cpp b/Engine/Source/Runtime/Renderer/Private/PostProcess/PostProcessBokehDOFRecombine.cpp
--- a/Engine/Source/Runtime/Renderer/Private/PostProcess/PostProcessBokehDOFRecombine.cpp
+++ b/Engine/Source/Runtime/Renderer/Private/PostProcess/PostProcessBokehDOFRecombine.cpp
@@ -387,7 +387,8 @@ void FRCPassPostProcessBokehDOFRecombine::Process(FRenderingCompositePassContext
 		const float UVScaling = (float)HalfResViewRect.Height() / (float)TexSize.Y;
 	
 		// Common setup
-		SetRenderTarget(Context.RHICmdList, nullptr, nullptr);
+		// #todo-renderpasses remove once everything is renderpasses
+		UnbindRenderTargets(Context.RHICmdList);
 		Context.SetViewportAndCallRHI(DestRect, 0.0f, 1.0f);
 		
 		static FName AsyncEndFenceName(TEXT("AsyncBokehDOFRecombineEndFence"));
@@ -421,18 +422,19 @@ void FRCPassPostProcessBokehDOFRecombine::Process(FRenderingCompositePassContext
 		WaitForInputPassComputeFences(Context.RHICmdList);
 
 		// Set the view family's render target/viewport.
-		SetRenderTarget(Context.RHICmdList, DestRenderTarget.TargetableTexture, FTextureRHIRef());
-
-		if (View.StereoPass == eSSP_FULL)
+		FRHIRenderPassInfo RPInfo(DestRenderTarget.TargetableTexture, ERenderTargetActions::Load_Store, DestRenderTarget.ShaderResourceTexture);
+		Context.RHICmdList.BeginRenderPass(RPInfo, TEXT("BokenDOFRecombine"));
 		{
-			// is optimized away if possible (RT size=view size, )
-			DrawClearQuad(Context.RHICmdList, true, FLinearColor::Black, false, 1.0f, false, 0, PassOutputs[0].RenderTargetDesc.Extent, View.ViewRect);
-		}
+			if (View.StereoPass == eSSP_FULL)
+			{
+				// is optimized away if possible (RT size=view size, )
+				DrawClearQuad(Context.RHICmdList, true, FLinearColor::Black, false, 1.0f, false, 0, PassOutputs[0].RenderTargetDesc.Extent, View.ViewRect);
+			}
 
-		Context.SetViewportAndCallRHI(View.ViewRect);
+			Context.SetViewportAndCallRHI(View.ViewRect);
 
-		switch(Method)
-		{
+			switch (Method)
+			{
 			case 1: SetShader<1>(Context); break;
 			case 2: SetShader<2>(Context); break;
 			case 3: SetShader<3>(Context); break;
@@ -440,24 +442,24 @@ void FRCPassPostProcessBokehDOFRecombine::Process(FRenderingCompositePassContext
 			case 5: SetShader<5>(Context); break;
 			default:
 				check(0);
-		}
+			}
 
-		TShaderMapRef<FPostProcessVS> VertexShader(Context.GetShaderMap());
-
-		DrawPostProcessPass(
-			Context.RHICmdList,
-			0, 0,
-			View.ViewRect.Width(), View.ViewRect.Height(),
-			HalfResViewRect.Min.X, HalfResViewRect.Min.Y,
-			HalfResViewRect.Width(), HalfResViewRect.Height(),
-			View.ViewRect.Size(),
-			TexSize,
-			*VertexShader,
-			View.StereoPass,
-			false, // Disabled for correctness
-			EDRF_UseTriangleOptimization);
-
-		Context.RHICmdList.CopyToResolveTarget(DestRenderTarget.TargetableTexture, DestRenderTarget.ShaderResourceTexture, FResolveParams());
+			TShaderMapRef<FPostProcessVS> VertexShader(Context.GetShaderMap());
+
+			DrawPostProcessPass(
+				Context.RHICmdList,
+				0, 0,
+				View.ViewRect.Width(), View.ViewRect.Height(),
+				HalfResViewRect.Min.X, HalfResViewRect.Min.Y,
+				HalfResViewRect.Width(), HalfResViewRect.Height(),
+				View.ViewRect.Size(),
+				TexSize,
+				*VertexShader,
+				View.StereoPass,
+				false, // Disabled for correctness
+				EDRF_UseTriangleOptimization);
+		}
+		Context.RHICmdList.EndRenderPass();
 	}
 }
 
diff --git a/Engine/Source/Runtime/Renderer/Private/PostProcess/PostProcessBufferInspector.cpp b/Engine/Source/Runtime/Renderer/Private/PostProcess/PostProcessBufferInspector.cpp
--- a/Engine/Source/Runtime/Renderer/Private/PostProcess/PostProcessBufferInspector.cpp
+++ b/Engine/Source/Runtime/Renderer/Private/PostProcess/PostProcessBufferInspector.cpp
@@ -270,29 +270,31 @@ void FRCPassPostProcessBufferInspector::Process(FRenderingCompositePassContext&
 		Scene->PixelInspectorData.Requests.Remove(RequestKey);
 	}
 
-	FIntRect ViewRect = Context.SceneColorViewRect;
-	FIntPoint SrcSize = InputDesc->Extent;
-
 	const FSceneRenderTargetItem& DestRenderTarget = PassOutputs[0].RequestSurface(Context);
 
-	// Set the view family's render target/viewport.
-	SetRenderTarget(Context.RHICmdList, DestRenderTarget.TargetableTexture, FTextureRHIRef());
-	Context.SetViewportAndCallRHI(ViewRect);
-
+	FRHIRenderPassInfo RPInfo(DestRenderTarget.TargetableTexture, ERenderTargetActions::Load_Store);
+	Context.RHICmdList.BeginRenderPass(RPInfo, TEXT("PostProcessBufferInspector"));
 	{
-		FShader* VertexShader = SetShaderTempl(Context.RHICmdList, Context);
-
-		DrawRectangle(
-			Context.RHICmdList,
-			0, 0,
-			ViewRect.Width(), ViewRect.Height(),
-			ViewRect.Min.X, ViewRect.Min.Y,
-			ViewRect.Width(), ViewRect.Height(),
-			ViewRect.Size(),
-			SrcSize,
-			VertexShader,
-			EDRF_UseTriangleOptimization);
+		FIntPoint SrcSize = InputDesc->Extent;
+		FIntRect ViewRect = Context.SceneColorViewRect;
+		Context.SetViewportAndCallRHI(ViewRect);
+
+		{
+			FShader* VertexShader = SetShaderTempl(Context.RHICmdList, Context);
+
+			DrawRectangle(
+				Context.RHICmdList,
+				0, 0,
+				ViewRect.Width(), ViewRect.Height(),
+				ViewRect.Min.X, ViewRect.Min.Y,
+				ViewRect.Width(), ViewRect.Height(),
+				ViewRect.Size(),
+				SrcSize,
+				VertexShader,
+				EDRF_UseTriangleOptimization);
+		}
 	}
+	Context.RHICmdList.EndRenderPass();
 
 	FRenderTargetTemp TempRenderTarget(View, (const FTexture2DRHIRef&)DestRenderTarget.TargetableTexture);
 	FCanvas Canvas(&TempRenderTarget, NULL, ViewFamily.CurrentRealTime, ViewFamily.CurrentWorldTime, ViewFamily.DeltaWorldTime, Context.GetFeatureLevel());
@@ -301,7 +303,7 @@ void FRCPassPostProcessBufferInspector::Process(FRenderingCompositePassContext&
 	Canvas.Flush_RenderThread(Context.RHICmdList);
 
 	Context.RHICmdList.CopyToResolveTarget(DestRenderTarget.TargetableTexture, DestRenderTarget.ShaderResourceTexture, FResolveParams());
-
+	
 	// AdjustGBufferRefCount(1) call is done in constructor
 	FSceneRenderTargets::Get(Context.RHICmdList).AdjustGBufferRefCount(Context.RHICmdList, -1);
 #endif
diff --git a/Engine/Source/Runtime/Renderer/Private/PostProcess/PostProcessBusyWait.cpp b/Engine/Source/Runtime/Renderer/Private/PostProcess/PostProcessBusyWait.cpp
--- a/Engine/Source/Runtime/Renderer/Private/PostProcess/PostProcessBusyWait.cpp
+++ b/Engine/Source/Runtime/Renderer/Private/PostProcess/PostProcessBusyWait.cpp
@@ -88,50 +88,51 @@ void FRCPassPostProcessBusyWait::Process(FRenderingCompositePassContext& Context
 	SCOPED_DRAW_EVENT(Context.RHICmdList, BusyWait);
 	FSceneRenderTargets& SceneContext = FSceneRenderTargets::Get(Context.RHICmdList);
 
-	const FViewInfo& View = Context.View;
-	
-	FIntRect SrcRect = View.ViewRect;
-	FIntRect DestRect = View.UnscaledViewRect;
-	
 	SceneContext.BeginRenderingLightAttenuation(Context.RHICmdList);
 	
 	const FSceneRenderTargetItem& DestRenderTarget = SceneContext.GetLightAttenuation()->GetRenderTargetItem();
 
-	// Set the view family's render target/viewport.
-	SetRenderTarget(Context.RHICmdList, DestRenderTarget.TargetableTexture, FTextureRHIRef());
-	Context.SetViewportAndCallRHI(DestRect);
-
-	FGraphicsPipelineStateInitializer GraphicsPSOInit;
-	Context.RHICmdList.ApplyCachedRenderTargets(GraphicsPSOInit);
-	GraphicsPSOInit.BlendState = TStaticBlendState<>::GetRHI();
-	GraphicsPSOInit.RasterizerState = TStaticRasterizerState<>::GetRHI();
-	GraphicsPSOInit.DepthStencilState = TStaticDepthStencilState<false, CF_Always>::GetRHI();
-
-	TShaderMapRef<FPostProcessVS> VertexShader(Context.GetShaderMap());
-	TShaderMapRef<FPostProcessBusyWaitPS> PixelShader(Context.GetShaderMap());
-
-	GraphicsPSOInit.BoundShaderState.VertexDeclarationRHI = GFilterVertexDeclaration.VertexDeclarationRHI;
-	GraphicsPSOInit.BoundShaderState.VertexShaderRHI = GETSAFERHISHADER_VERTEX(*VertexShader);
-	GraphicsPSOInit.BoundShaderState.PixelShaderRHI = GETSAFERHISHADER_PIXEL(*PixelShader);
-	GraphicsPSOInit.PrimitiveType = PT_TriangleList;
-
-	SetGraphicsPipelineState(Context.RHICmdList, GraphicsPSOInit);
-
-	PixelShader->SetPS(Context.RHICmdList, Context);
-
-	// Draw a quad mapping scene color to the view's render target
-	DrawRectangle(
-		Context.RHICmdList,
-		0, 0,
-		DestRect.Width(), DestRect.Height(),
-		SrcRect.Min.X, SrcRect.Min.Y,
-		SrcRect.Width(), SrcRect.Height(),
-		DestRect.Size(),
-		SrcRect.Size(),
-		*VertexShader,
-		EDRF_UseTriangleOptimization);
-
-	Context.RHICmdList.CopyToResolveTarget(DestRenderTarget.TargetableTexture, DestRenderTarget.ShaderResourceTexture, FResolveParams());
+	FRHIRenderPassInfo RPInfo(DestRenderTarget.TargetableTexture, ERenderTargetActions::Load_Store, DestRenderTarget.ShaderResourceTexture);
+	Context.RHICmdList.BeginRenderPass(RPInfo, TEXT("PostProcessBusyWait"));
+	{
+		const FViewInfo& View = Context.View;
+
+		FIntRect SrcRect = View.ViewRect;
+		FIntRect DestRect = View.UnscaledViewRect;
+
+		Context.SetViewportAndCallRHI(DestRect);
+
+		FGraphicsPipelineStateInitializer GraphicsPSOInit;
+		Context.RHICmdList.ApplyCachedRenderTargets(GraphicsPSOInit);
+		GraphicsPSOInit.BlendState = TStaticBlendState<>::GetRHI();
+		GraphicsPSOInit.RasterizerState = TStaticRasterizerState<>::GetRHI();
+		GraphicsPSOInit.DepthStencilState = TStaticDepthStencilState<false, CF_Always>::GetRHI();
+
+		TShaderMapRef<FPostProcessVS> VertexShader(Context.GetShaderMap());
+		TShaderMapRef<FPostProcessBusyWaitPS> PixelShader(Context.GetShaderMap());
+
+		GraphicsPSOInit.BoundShaderState.VertexDeclarationRHI = GFilterVertexDeclaration.VertexDeclarationRHI;
+		GraphicsPSOInit.BoundShaderState.VertexShaderRHI = GETSAFERHISHADER_VERTEX(*VertexShader);
+		GraphicsPSOInit.BoundShaderState.PixelShaderRHI = GETSAFERHISHADER_PIXEL(*PixelShader);
+		GraphicsPSOInit.PrimitiveType = PT_TriangleList;
+
+		SetGraphicsPipelineState(Context.RHICmdList, GraphicsPSOInit);
+
+		PixelShader->SetPS(Context.RHICmdList, Context);
+
+		// Draw a quad mapping scene color to the view's render target
+		DrawRectangle(
+			Context.RHICmdList,
+			0, 0,
+			DestRect.Width(), DestRect.Height(),
+			SrcRect.Min.X, SrcRect.Min.Y,
+			SrcRect.Width(), SrcRect.Height(),
+			DestRect.Size(),
+			SrcRect.Size(),
+			*VertexShader,
+			EDRF_UseTriangleOptimization);
+	}
+	Context.RHICmdList.EndRenderPass();
 
 	SceneContext.SetLightAttenuation(0);
 }
diff --git a/Engine/Source/Runtime/Renderer/Private/PostProcess/PostProcessCircleDOF.cpp b/Engine/Source/Runtime/Renderer/Private/PostProcess/PostProcessCircleDOF.cpp
--- a/Engine/Source/Runtime/Renderer/Private/PostProcess/PostProcessCircleDOF.cpp
+++ b/Engine/Source/Runtime/Renderer/Private/PostProcess/PostProcessCircleDOF.cpp
@@ -178,21 +178,6 @@ void FRCPassPostProcessCircleDOFSetup::Process(FRenderingCompositePassContext& C
 		return;
 	}
 
-	const FViewInfo& View = Context.View;
-	const FSceneViewFamily& ViewFamily = *(View.Family);
-
-	const auto FeatureLevel = Context.GetFeatureLevel();
-	auto ShaderMap = Context.GetShaderMap();
-
-	FIntPoint SrcSize = InputDesc->Extent;
-	FIntPoint DestSize = PassOutputs[0].RenderTargetDesc.Extent;
-
-	// e.g. 4 means the input texture is 4x smaller than the buffer size
-	uint32 ScaleFactor = FSceneRenderTargets::Get(Context.RHICmdList).GetBufferSizeXY().X / SrcSize.X;
-
-	FIntRect SrcRect = View.ViewRect / ScaleFactor;
-	FIntRect DestRect = SrcRect / 2;
-
 	const FSceneRenderTargetItem& DestRenderTarget0 = PassOutputs[0].RequestSurface(Context);
 	const FSceneRenderTargetItem& DestRenderTarget1 = FPostProcessing::HasAlphaChannelSupport() ? PassOutputs[1].RequestSurface(Context) : FSceneRenderTargetItem();
 
@@ -203,74 +188,97 @@ void FRCPassPostProcessCircleDOFSetup::Process(FRenderingCompositePassContext& C
 		DestRenderTarget1.TargetableTexture
 	};
 	uint32 NumRenderTargets = FPostProcessing::HasAlphaChannelSupport() ? 2 : 1;
-	SetRenderTargets(Context.RHICmdList, NumRenderTargets, RenderTargets, FTextureRHIParamRef(), 0, NULL);
+	
+	FTextureRHIParamRef ResolveTargets[2] =
+	{
+		DestRenderTarget0.ShaderResourceTexture,
+		nullptr,
+	};
 
-	if (View.StereoPass == eSSP_FULL)
+	if (DestRenderTarget1.TargetableTexture)
 	{
-		FLinearColor ClearColors[2] =
-		{
-			FLinearColor(0, 0, 0, 0),
-			FLinearColor(0, 0, 0, 0)
-		};
-		// is optimized away if possible (RT size=view size, )
-		DrawClearQuadMRT(Context.RHICmdList, true, NumRenderTargets, ClearColors, false, 0, false, 0, PassOutputs[0].RenderTargetDesc.Extent, DestRect);
+		ResolveTargets[1] = DestRenderTarget1.ShaderResourceTexture;
 	}
 
-	Context.SetViewportAndCallRHI(0, 0, 0.0f, DestSize.X, DestSize.Y, 1.0f );
+	FRHIRenderPassInfo RPInfo(NumRenderTargets, RenderTargets, ERenderTargetActions::Load_Store, ResolveTargets);
+	Context.RHICmdList.BeginRenderPass(RPInfo, TEXT("CircleDOFSetup"));
+	{
+		const FViewInfo& View = Context.View;
+		const FSceneViewFamily& ViewFamily = *(View.Family);
 
-	FGraphicsPipelineStateInitializer GraphicsPSOInit;
-	Context.RHICmdList.ApplyCachedRenderTargets(GraphicsPSOInit);
-	GraphicsPSOInit.BlendState = TStaticBlendState<>::GetRHI();
-	GraphicsPSOInit.RasterizerState = TStaticRasterizerState<>::GetRHI();
-	GraphicsPSOInit.DepthStencilState = TStaticDepthStencilState<false, CF_Always>::GetRHI();
-	GraphicsPSOInit.PrimitiveType = PT_TriangleList;
+		const auto FeatureLevel = Context.GetFeatureLevel();
+		auto ShaderMap = Context.GetShaderMap();
 
-	TShaderMapRef<FPostProcessVS> VertexShader(ShaderMap);
+		FIntPoint SrcSize = InputDesc->Extent;
+		FIntPoint DestSize = PassOutputs[0].RenderTargetDesc.Extent;
 
-	if(CVarDepthOfFieldFarBlur.GetValueOnRenderThread())
-	{
-		TShaderMapRef< FPostProcessCircleDOFSetupPS<1> > PixelShader(ShaderMap);
+		// e.g. 4 means the input texture is 4x smaller than the buffer size
+		uint32 ScaleFactor = FSceneRenderTargets::Get(Context.RHICmdList).GetBufferSizeXY().X / SrcSize.X;
 
-		GraphicsPSOInit.BoundShaderState.VertexDeclarationRHI = GFilterVertexDeclaration.VertexDeclarationRHI;
-		GraphicsPSOInit.BoundShaderState.VertexShaderRHI = GETSAFERHISHADER_VERTEX(*VertexShader);
-		GraphicsPSOInit.BoundShaderState.PixelShaderRHI = GETSAFERHISHADER_PIXEL(*PixelShader);
-		SetGraphicsPipelineState(Context.RHICmdList, GraphicsPSOInit);
+		FIntRect SrcRect = View.ViewRect / ScaleFactor;
+		FIntRect DestRect = SrcRect / 2;
 
-		PixelShader->SetParameters(Context);
-	}
-	else
-	{
-		TShaderMapRef< FPostProcessCircleDOFSetupPS<0> > PixelShader(ShaderMap);
+		if (View.StereoPass == eSSP_FULL)
+		{
+			FLinearColor ClearColors[2] =
+			{
+				FLinearColor(0, 0, 0, 0),
+				FLinearColor(0, 0, 0, 0)
+			};
+			// is optimized away if possible (RT size=view size, )
+			DrawClearQuadMRT(Context.RHICmdList, true, NumRenderTargets, ClearColors, false, 0, false, 0, PassOutputs[0].RenderTargetDesc.Extent, DestRect);
+		}
 
-		GraphicsPSOInit.BoundShaderState.VertexDeclarationRHI = GFilterVertexDeclaration.VertexDeclarationRHI;
-		GraphicsPSOInit.BoundShaderState.VertexShaderRHI = GETSAFERHISHADER_VERTEX(*VertexShader);
-		GraphicsPSOInit.BoundShaderState.PixelShaderRHI = GETSAFERHISHADER_PIXEL(*PixelShader);
-		SetGraphicsPipelineState(Context.RHICmdList, GraphicsPSOInit);
+		Context.SetViewportAndCallRHI(0, 0, 0.0f, DestSize.X, DestSize.Y, 1.0f);
 
-		PixelShader->SetParameters(Context);
-	}
+		FGraphicsPipelineStateInitializer GraphicsPSOInit;
+		Context.RHICmdList.ApplyCachedRenderTargets(GraphicsPSOInit);
+		GraphicsPSOInit.BlendState = TStaticBlendState<>::GetRHI();
+		GraphicsPSOInit.RasterizerState = TStaticRasterizerState<>::GetRHI();
+		GraphicsPSOInit.DepthStencilState = TStaticDepthStencilState<false, CF_Always>::GetRHI();
+		GraphicsPSOInit.PrimitiveType = PT_TriangleList;
 
-	VertexShader->SetParameters(Context);
+		TShaderMapRef<FPostProcessVS> VertexShader(ShaderMap);
 
-	DrawPostProcessPass(
-		Context.RHICmdList,
-		DestRect.Min.X, DestRect.Min.Y,
-		DestRect.Width() + 1, DestRect.Height() + 1,
-		SrcRect.Min.X, SrcRect.Min.Y,
-		SrcRect.Width() + 1, SrcRect.Height() + 1,
-		DestSize,
-		SrcSize,
-		*VertexShader,
-		View.StereoPass,
-		Context.HasHmdMesh(),
-		EDRF_UseTriangleOptimization);
-
-	Context.RHICmdList.CopyToResolveTarget(DestRenderTarget0.TargetableTexture, DestRenderTarget0.ShaderResourceTexture, FResolveParams());
+		if (CVarDepthOfFieldFarBlur.GetValueOnRenderThread())
+		{
+			TShaderMapRef< FPostProcessCircleDOFSetupPS<1> > PixelShader(ShaderMap);
 
-	if (DestRenderTarget1.TargetableTexture)
-	{
-		Context.RHICmdList.CopyToResolveTarget(DestRenderTarget1.TargetableTexture, DestRenderTarget1.ShaderResourceTexture, FResolveParams());
-	}
+			GraphicsPSOInit.BoundShaderState.VertexDeclarationRHI = GFilterVertexDeclaration.VertexDeclarationRHI;
+			GraphicsPSOInit.BoundShaderState.VertexShaderRHI = GETSAFERHISHADER_VERTEX(*VertexShader);
+			GraphicsPSOInit.BoundShaderState.PixelShaderRHI = GETSAFERHISHADER_PIXEL(*PixelShader);
+			SetGraphicsPipelineState(Context.RHICmdList, GraphicsPSOInit);
+
+			PixelShader->SetParameters(Context);
+		}
+		else
+		{
+			TShaderMapRef< FPostProcessCircleDOFSetupPS<0> > PixelShader(ShaderMap);
+
+			GraphicsPSOInit.BoundShaderState.VertexDeclarationRHI = GFilterVertexDeclaration.VertexDeclarationRHI;
+			GraphicsPSOInit.BoundShaderState.VertexShaderRHI = GETSAFERHISHADER_VERTEX(*VertexShader);
+			GraphicsPSOInit.BoundShaderState.PixelShaderRHI = GETSAFERHISHADER_PIXEL(*PixelShader);
+			SetGraphicsPipelineState(Context.RHICmdList, GraphicsPSOInit);
+
+			PixelShader->SetParameters(Context);
+		}
+
+		VertexShader->SetParameters(Context);
+
+		DrawPostProcessPass(
+			Context.RHICmdList,
+			DestRect.Min.X, DestRect.Min.Y,
+			DestRect.Width() + 1, DestRect.Height() + 1,
+			SrcRect.Min.X, SrcRect.Min.Y,
+			SrcRect.Width() + 1, SrcRect.Height() + 1,
+			DestSize,
+			SrcSize,
+			*VertexShader,
+			View.StereoPass,
+			Context.HasHmdMesh(),
+			EDRF_UseTriangleOptimization);
+	}
+	Context.RHICmdList.EndRenderPass();
 }
 
 FPooledRenderTargetDesc FRCPassPostProcessCircleDOFSetup::ComputeOutputDesc(EPassOutputId InPassOutputId) const
@@ -394,21 +402,6 @@ void FRCPassPostProcessCircleDOFDilate::Process(FRenderingCompositePassContext&
 
 	uint32 NumRenderTargets = 1;
 
-	const FViewInfo& View = Context.View;
-	const FSceneViewFamily& ViewFamily = *(View.Family);
-
-	const auto FeatureLevel = Context.GetFeatureLevel();
-	auto ShaderMap = Context.GetShaderMap();
-
-	FIntPoint SrcSize = InputDesc->Extent;
-	FIntPoint DestSize = PassOutputs[0].RenderTargetDesc.Extent;
-
-	// e.g. 4 means the input texture is 4x smaller than the buffer size
-	uint32 ScaleFactor = FSceneRenderTargets::Get(Context.RHICmdList).GetBufferSizeXY().X / SrcSize.X;
-
-	FIntRect SrcRect = View.ViewRect / ScaleFactor;
-	FIntRect DestRect = SrcRect / 2;
-
 	const FSceneRenderTargetItem& DestRenderTarget0 = PassOutputs[0].RequestSurface(Context);
 	const FSceneRenderTargetItem& DestRenderTarget1 = FSceneRenderTargetItem();
 
@@ -418,74 +411,97 @@ void FRCPassPostProcessCircleDOFDilate::Process(FRenderingCompositePassContext&
 		DestRenderTarget0.TargetableTexture,
 		DestRenderTarget1.TargetableTexture
 	};
-	SetRenderTargets(Context.RHICmdList, NumRenderTargets, RenderTargets, FTextureRHIParamRef(), 0, nullptr);
 
-	if (View.StereoPass == eSSP_FULL)
+	FTextureRHIParamRef ResolveTargets[2] =
 	{
-		FLinearColor ClearColors[2] =
-		{
-			FLinearColor(0, 0, 0, 0),
-			FLinearColor(0, 0, 0, 0)
-		};
-		// is optimized away if possible (RT size=view size, )
-		DrawClearQuadMRT(Context.RHICmdList, true, NumRenderTargets, ClearColors, false, 0, false, 0, PassOutputs[0].RenderTargetDesc.Extent, DestRect);
+		DestRenderTarget0.ShaderResourceTexture,
+		nullptr,
+	};
+
+	if (DestRenderTarget1.TargetableTexture)
+	{
+		ResolveTargets[1] = DestRenderTarget1.ShaderResourceTexture;
 	}
 
-	Context.SetViewportAndCallRHI(0, 0, 0.0f, DestSize.X, DestSize.Y, 1.0f );
+	FRHIRenderPassInfo RPInfo(NumRenderTargets, RenderTargets, ERenderTargetActions::Load_Store, ResolveTargets);
+	Context.RHICmdList.BeginRenderPass(RPInfo, TEXT("CircleDOFDilate"));
+	{
+		const FViewInfo& View = Context.View;
+		const FSceneViewFamily& ViewFamily = *(View.Family);
 
-	FGraphicsPipelineStateInitializer GraphicsPSOInit;
-	Context.RHICmdList.ApplyCachedRenderTargets(GraphicsPSOInit);
-	GraphicsPSOInit.BlendState = TStaticBlendState<>::GetRHI();
-	GraphicsPSOInit.RasterizerState = TStaticRasterizerState<>::GetRHI();
-	GraphicsPSOInit.DepthStencilState = TStaticDepthStencilState<false, CF_Always>::GetRHI();
-	GraphicsPSOInit.PrimitiveType = PT_TriangleList;
+		const auto FeatureLevel = Context.GetFeatureLevel();
+		auto ShaderMap = Context.GetShaderMap();
 
-	TShaderMapRef<FPostProcessVS> VertexShader(ShaderMap);
+		FIntPoint SrcSize = InputDesc->Extent;
+		FIntPoint DestSize = PassOutputs[0].RenderTargetDesc.Extent;
 
-	if (false)
-	{
-		TShaderMapRef< FPostProcessCircleDOFDilatePS<1> > PixelShader(ShaderMap);
+		// e.g. 4 means the input texture is 4x smaller than the buffer size
+		uint32 ScaleFactor = FSceneRenderTargets::Get(Context.RHICmdList).GetBufferSizeXY().X / SrcSize.X;
 
-		GraphicsPSOInit.BoundShaderState.VertexDeclarationRHI = GFilterVertexDeclaration.VertexDeclarationRHI;
-		GraphicsPSOInit.BoundShaderState.VertexShaderRHI = GETSAFERHISHADER_VERTEX(*VertexShader);
-		GraphicsPSOInit.BoundShaderState.PixelShaderRHI = GETSAFERHISHADER_PIXEL(*PixelShader);
+		FIntRect SrcRect = View.ViewRect / ScaleFactor;
+		FIntRect DestRect = SrcRect / 2;
 
-		SetGraphicsPipelineState(Context.RHICmdList, GraphicsPSOInit);
-		PixelShader->SetParameters(Context.RHICmdList, Context);
-	}
-	else
-	{
-		TShaderMapRef< FPostProcessCircleDOFDilatePS<0> > PixelShader(ShaderMap);
+		if (View.StereoPass == eSSP_FULL)
+		{
+			FLinearColor ClearColors[2] =
+			{
+				FLinearColor(0, 0, 0, 0),
+				FLinearColor(0, 0, 0, 0)
+			};
+			// is optimized away if possible (RT size=view size, )
+			DrawClearQuadMRT(Context.RHICmdList, true, NumRenderTargets, ClearColors, false, 0, false, 0, PassOutputs[0].RenderTargetDesc.Extent, DestRect);
+		}
 
-		GraphicsPSOInit.BoundShaderState.VertexDeclarationRHI = GFilterVertexDeclaration.VertexDeclarationRHI;
-		GraphicsPSOInit.BoundShaderState.VertexShaderRHI = GETSAFERHISHADER_VERTEX(*VertexShader);
-		GraphicsPSOInit.BoundShaderState.PixelShaderRHI = GETSAFERHISHADER_PIXEL(*PixelShader);
-		SetGraphicsPipelineState(Context.RHICmdList, GraphicsPSOInit);
+		Context.SetViewportAndCallRHI(0, 0, 0.0f, DestSize.X, DestSize.Y, 1.0f);
 
-		PixelShader->SetParameters(Context.RHICmdList, Context);
-	}
+		FGraphicsPipelineStateInitializer GraphicsPSOInit;
+		Context.RHICmdList.ApplyCachedRenderTargets(GraphicsPSOInit);
+		GraphicsPSOInit.BlendState = TStaticBlendState<>::GetRHI();
+		GraphicsPSOInit.RasterizerState = TStaticRasterizerState<>::GetRHI();
+		GraphicsPSOInit.DepthStencilState = TStaticDepthStencilState<false, CF_Always>::GetRHI();
+		GraphicsPSOInit.PrimitiveType = PT_TriangleList;
 
-	VertexShader->SetParameters(Context);
+		TShaderMapRef<FPostProcessVS> VertexShader(ShaderMap);
 
-	DrawPostProcessPass(
-		Context.RHICmdList,
-		DestRect.Min.X, DestRect.Min.Y,
-		DestRect.Width() + 1, DestRect.Height() + 1,
-		SrcRect.Min.X, SrcRect.Min.Y,
-		SrcRect.Width() + 1, SrcRect.Height() + 1,
-		DestSize,
-		SrcSize,
-		*VertexShader,
-		View.StereoPass,
-		Context.HasHmdMesh(),
-		EDRF_UseTriangleOptimization);
-	
-	Context.RHICmdList.CopyToResolveTarget(DestRenderTarget0.TargetableTexture, DestRenderTarget0.ShaderResourceTexture, FResolveParams());
+		if (false)
+		{
+			TShaderMapRef< FPostProcessCircleDOFDilatePS<1> > PixelShader(ShaderMap);
 
-	if (DestRenderTarget1.TargetableTexture)
-	{
-		Context.RHICmdList.CopyToResolveTarget(DestRenderTarget1.TargetableTexture, DestRenderTarget1.ShaderResourceTexture, FResolveParams());
-	}
+			GraphicsPSOInit.BoundShaderState.VertexDeclarationRHI = GFilterVertexDeclaration.VertexDeclarationRHI;
+			GraphicsPSOInit.BoundShaderState.VertexShaderRHI = GETSAFERHISHADER_VERTEX(*VertexShader);
+			GraphicsPSOInit.BoundShaderState.PixelShaderRHI = GETSAFERHISHADER_PIXEL(*PixelShader);
+
+			SetGraphicsPipelineState(Context.RHICmdList, GraphicsPSOInit);
+			PixelShader->SetParameters(Context.RHICmdList, Context);
+		}
+		else
+		{
+			TShaderMapRef< FPostProcessCircleDOFDilatePS<0> > PixelShader(ShaderMap);
+
+			GraphicsPSOInit.BoundShaderState.VertexDeclarationRHI = GFilterVertexDeclaration.VertexDeclarationRHI;
+			GraphicsPSOInit.BoundShaderState.VertexShaderRHI = GETSAFERHISHADER_VERTEX(*VertexShader);
+			GraphicsPSOInit.BoundShaderState.PixelShaderRHI = GETSAFERHISHADER_PIXEL(*PixelShader);
+			SetGraphicsPipelineState(Context.RHICmdList, GraphicsPSOInit);
+
+			PixelShader->SetParameters(Context.RHICmdList, Context);
+		}
+
+		VertexShader->SetParameters(Context);
+
+		DrawPostProcessPass(
+			Context.RHICmdList,
+			DestRect.Min.X, DestRect.Min.Y,
+			DestRect.Width() + 1, DestRect.Height() + 1,
+			SrcRect.Min.X, SrcRect.Min.Y,
+			SrcRect.Width() + 1, SrcRect.Height() + 1,
+			DestSize,
+			SrcSize,
+			*VertexShader,
+			View.StereoPass,
+			Context.HasHmdMesh(),
+			EDRF_UseTriangleOptimization);
+	}
+	Context.RHICmdList.EndRenderPass();
 }
 
 FPooledRenderTargetDesc FRCPassPostProcessCircleDOFDilate::ComputeOutputDesc(EPassOutputId InPassOutputId) const
@@ -677,21 +693,6 @@ void FRCPassPostProcessCircleDOF::Process(FRenderingCompositePassContext& Contex
 		return;
 	}
 
-	const FViewInfo& View = Context.View;
-	const FSceneViewFamily& ViewFamily = *(View.Family);
-
-	const auto FeatureLevel = Context.GetFeatureLevel();
-	auto ShaderMap = Context.GetShaderMap();
-
-	FIntPoint SrcSize = InputDesc->Extent;
-	FIntPoint DestSize = PassOutputs[0].RenderTargetDesc.Extent;
-
-	// e.g. 4 means the input texture is 4x smaller than the buffer size
-	uint32 ScaleFactor = FSceneRenderTargets::Get(Context.RHICmdList).GetBufferSizeXY().X / SrcSize.X;
-
-	FIntRect SrcRect = View.ViewRect / ScaleFactor;
-	FIntRect DestRect = SrcRect;
-
 	const FSceneRenderTargetItem& DestRenderTarget0 = PassOutputs[0].RequestSurface(Context);
 	const FSceneRenderTargetItem& DestRenderTarget1 = FPostProcessing::HasAlphaChannelSupport() ? PassOutputs[1].RequestSurface(Context) : FSceneRenderTargetItem();
 
@@ -702,56 +703,79 @@ void FRCPassPostProcessCircleDOF::Process(FRenderingCompositePassContext& Contex
 		DestRenderTarget1.TargetableTexture
 	};
 	uint32 NumRenderTargets = FPostProcessing::HasAlphaChannelSupport() ? 2 : 1;
-	SetRenderTargets(Context.RHICmdList, NumRenderTargets, RenderTargets, FTextureRHIParamRef(), 0, NULL);
 
-	if (View.StereoPass == eSSP_FULL)
+	FTextureRHIParamRef ResolveTargets[2] =
 	{
-		FLinearColor ClearColors[2] =
-		{
-			FLinearColor(0, 0, 0, 0),
-			FLinearColor(0, 0, 0, 0)
-		};
+		DestRenderTarget0.ShaderResourceTexture,
+		nullptr,
+	};
 
-		// is optimized away if possible (RT size=view size, )
-		DrawClearQuadMRT(Context.RHICmdList, true, NumRenderTargets, ClearColors, false, 0, false, 0, PassOutputs[0].RenderTargetDesc.Extent, DestRect);
+	if (DestRenderTarget1.TargetableTexture)
+	{
+		ResolveTargets[1] = DestRenderTarget1.ShaderResourceTexture;
 	}
 
-	Context.SetViewportAndCallRHI(0, 0, 0.0f, DestSize.X, DestSize.Y, 1.0f );
-		
-	static const auto CVar = IConsoleManager::Get().FindTConsoleVariableDataInt(TEXT("r.DepthOfFieldQuality"));
-	check(CVar);
-	int32 DOFQualityCVarValue = CVar->GetValueOnRenderThread();
-	
-	
-
-	FShader* VertexShader = 0;
-	
-	switch (DOFQualityCVarValue)
+	FRHIRenderPassInfo RPInfo(NumRenderTargets, RenderTargets, ERenderTargetActions::Load_Store, ResolveTargets);
+	Context.RHICmdList.BeginRenderPass(RPInfo, TEXT("CircleDOFApply"));
 	{
-		case 3:  VertexShader = SetShaderTempl<1>(Context); break;
-		case 4:  VertexShader = SetShaderTempl<2>(Context); break;	
-		default: VertexShader = SetShaderTempl<0>(Context);
-	}
+		const FViewInfo& View = Context.View;
+		const FSceneViewFamily& ViewFamily = *(View.Family);
 
-	DrawPostProcessPass(
-		Context.RHICmdList,
-		DestRect.Min.X, DestRect.Min.Y,
-		DestRect.Width() + 1, DestRect.Height() + 1,
-		SrcRect.Min.X, SrcRect.Min.Y,
-		SrcRect.Width() + 1, SrcRect.Height() + 1,
-		DestSize,
-		SrcSize,
-		VertexShader,
-		View.StereoPass,
-		Context.HasHmdMesh(),
-		EDRF_UseTriangleOptimization);
+		const auto FeatureLevel = Context.GetFeatureLevel();
+		auto ShaderMap = Context.GetShaderMap();
 
-	Context.RHICmdList.CopyToResolveTarget(DestRenderTarget0.TargetableTexture, DestRenderTarget0.ShaderResourceTexture, FResolveParams());
+		FIntPoint SrcSize = InputDesc->Extent;
+		FIntPoint DestSize = PassOutputs[0].RenderTargetDesc.Extent;
 
-	if (DestRenderTarget1.TargetableTexture)
-	{
-		Context.RHICmdList.CopyToResolveTarget(DestRenderTarget1.TargetableTexture, DestRenderTarget1.ShaderResourceTexture, FResolveParams());
-	}
+		// e.g. 4 means the input texture is 4x smaller than the buffer size
+		uint32 ScaleFactor = FSceneRenderTargets::Get(Context.RHICmdList).GetBufferSizeXY().X / SrcSize.X;
+
+		FIntRect SrcRect = View.ViewRect / ScaleFactor;
+		FIntRect DestRect = SrcRect;
+
+		if (View.StereoPass == eSSP_FULL)
+		{
+			FLinearColor ClearColors[2] =
+			{
+				FLinearColor(0, 0, 0, 0),
+				FLinearColor(0, 0, 0, 0)
+			};
+
+			// is optimized away if possible (RT size=view size, )
+			DrawClearQuadMRT(Context.RHICmdList, true, NumRenderTargets, ClearColors, false, 0, false, 0, PassOutputs[0].RenderTargetDesc.Extent, DestRect);
+		}
+
+		Context.SetViewportAndCallRHI(0, 0, 0.0f, DestSize.X, DestSize.Y, 1.0f);
+
+		static const auto CVar = IConsoleManager::Get().FindTConsoleVariableDataInt(TEXT("r.DepthOfFieldQuality"));
+		check(CVar);
+		int32 DOFQualityCVarValue = CVar->GetValueOnRenderThread();
+
+
+
+		FShader* VertexShader = 0;
+
+		switch (DOFQualityCVarValue)
+		{
+		case 3:  VertexShader = SetShaderTempl<1>(Context); break;
+		case 4:  VertexShader = SetShaderTempl<2>(Context); break;
+		default: VertexShader = SetShaderTempl<0>(Context);
+		}
+
+		DrawPostProcessPass(
+			Context.RHICmdList,
+			DestRect.Min.X, DestRect.Min.Y,
+			DestRect.Width() + 1, DestRect.Height() + 1,
+			SrcRect.Min.X, SrcRect.Min.Y,
+			SrcRect.Width() + 1, SrcRect.Height() + 1,
+			DestSize,
+			SrcSize,
+			VertexShader,
+			View.StereoPass,
+			Context.HasHmdMesh(),
+			EDRF_UseTriangleOptimization);
+	}
+	Context.RHICmdList.EndRenderPass();
 }
 
 FPooledRenderTargetDesc FRCPassPostProcessCircleDOF::ComputeOutputDesc(EPassOutputId InPassOutputId) const
@@ -915,65 +939,66 @@ void FRCPassPostProcessCircleDOFRecombine::Process(FRenderingCompositePassContex
 		return;
 	}
 
-	const FViewInfo& View = Context.View;
+	const FSceneRenderTargetItem& DestRenderTarget = PassOutputs[0].RequestSurface(Context);
+
+	ERenderTargetActions LoadStoreAction = ERenderTargetActions::Clear_Store;
+
+	if (!Context.HasHmdMesh())
+	{
+		LoadStoreAction = ERenderTargetActions::DontLoad_Store;
+	}
+
+	FRHIRenderPassInfo RPInfo(DestRenderTarget.TargetableTexture, LoadStoreAction, DestRenderTarget.ShaderResourceTexture);
+	Context.RHICmdList.BeginRenderPass(RPInfo, TEXT("CircleDOFRecombine"));
+	{
+		const FViewInfo& View = Context.View;
 
-	const auto FeatureLevel = Context.GetFeatureLevel();
-	auto ShaderMap = Context.GetShaderMap();
+		const auto FeatureLevel = Context.GetFeatureLevel();
+		auto ShaderMap = Context.GetShaderMap();
 
-	FIntPoint TexSize = InputDesc->Extent;
+		FIntPoint TexSize = InputDesc->Extent;
 
-	// usually 1, 2, 4 or 8
-	uint32 ScaleToFullRes = FSceneRenderTargets::Get(Context.RHICmdList).GetBufferSizeXY().X / TexSize.X;
+		// usually 1, 2, 4 or 8
+		uint32 ScaleToFullRes = FSceneRenderTargets::Get(Context.RHICmdList).GetBufferSizeXY().X / TexSize.X;
 
-	FIntRect HalfResViewRect = View.ViewRect / ScaleToFullRes;
+		FIntRect HalfResViewRect = View.ViewRect / ScaleToFullRes;
 
-	const FSceneRenderTargetItem& DestRenderTarget = PassOutputs[0].RequestSurface(Context);
+		if (!Context.HasHmdMesh())
+		{
+			// is optimized away if possible (RT size=view size, )
+			DrawClearQuad(Context.RHICmdList, true, FLinearColor::Black, false, 0, false, 0, PassOutputs[0].RenderTargetDesc.Extent, View.ViewRect);
+		}
 
-	if (!Context.HasHmdMesh())
-	{
-		// Set the view family's render target/viewport.
-		SetRenderTarget(Context.RHICmdList, DestRenderTarget.TargetableTexture, FTextureRHIRef());
+		Context.SetViewportAndCallRHI(View.ViewRect);
 
-		// is optimized away if possible (RT size=view size, )
-		DrawClearQuad(Context.RHICmdList, true, FLinearColor::Black, false, 0, false, 0, PassOutputs[0].RenderTargetDesc.Extent, View.ViewRect);
-	}
-	else
-	{
-		FRHIRenderTargetView RtView = FRHIRenderTargetView(DestRenderTarget.TargetableTexture, ERenderTargetLoadAction::ENoAction);
-		FRHISetRenderTargetsInfo Info(1, &RtView, FRHIDepthRenderTargetView());
-		Context.RHICmdList.SetRenderTargetsAndClear(Info);
-	}
+		static const auto CVar = IConsoleManager::Get().FindTConsoleVariableDataInt(TEXT("r.DepthOfFieldQuality"));
+		check(CVar);
+		int32 DOFQualityCVarValue = CVar->GetValueOnRenderThread();
 
-	Context.SetViewportAndCallRHI(View.ViewRect);
+		// 0:normal / 1:slow but very high quality
+		uint32 Quality = DOFQualityCVarValue >= 3;
 
-	static const auto CVar = IConsoleManager::Get().FindTConsoleVariableDataInt(TEXT("r.DepthOfFieldQuality"));
-	check(CVar);
-	int32 DOFQualityCVarValue = CVar->GetValueOnRenderThread();
-	
-	// 0:normal / 1:slow but very high quality
-	uint32 Quality = DOFQualityCVarValue >= 3;
+		FShader* VertexShader = 0;
 
-	FShader* VertexShader = 0;
+		if (Quality)
+			VertexShader = SetShaderTempl<1>(Context);
+		else
+			VertexShader = SetShaderTempl<0>(Context);
 
-	if (Quality)
-		VertexShader = SetShaderTempl<1>(Context);
-	else
-		VertexShader = SetShaderTempl<0>(Context);
-
-	DrawPostProcessPass(
-		Context.RHICmdList,
-		0, 0,
-		View.ViewRect.Width(), View.ViewRect.Height(),
-		View.ViewRect.Min.X, View.ViewRect.Min.Y,
-		View.ViewRect.Width(), View.ViewRect.Height(),
-		View.ViewRect.Size(),
-		TexSize,
-		VertexShader,
-		View.StereoPass,
-		Context.HasHmdMesh(),
-		EDRF_UseTriangleOptimization);
-
-	Context.RHICmdList.CopyToResolveTarget(DestRenderTarget.TargetableTexture, DestRenderTarget.ShaderResourceTexture, FResolveParams());
+		DrawPostProcessPass(
+			Context.RHICmdList,
+			0, 0,
+			View.ViewRect.Width(), View.ViewRect.Height(),
+			View.ViewRect.Min.X, View.ViewRect.Min.Y,
+			View.ViewRect.Width(), View.ViewRect.Height(),
+			View.ViewRect.Size(),
+			TexSize,
+			VertexShader,
+			View.StereoPass,
+			Context.HasHmdMesh(),
+			EDRF_UseTriangleOptimization);
+	}
+	Context.RHICmdList.EndRenderPass();
 }
 
 FPooledRenderTargetDesc FRCPassPostProcessCircleDOFRecombine::ComputeOutputDesc(EPassOutputId InPassOutputId) const
diff --git a/Engine/Source/Runtime/Renderer/Private/PostProcess/PostProcessCombineLUTs.cpp b/Engine/Source/Runtime/Renderer/Private/PostProcess/PostProcessCombineLUTs.cpp
--- a/Engine/Source/Runtime/Renderer/Private/PostProcess/PostProcessCombineLUTs.cpp
+++ b/Engine/Source/Runtime/Renderer/Private/PostProcess/PostProcessCombineLUTs.cpp
@@ -917,7 +917,8 @@ void FRCPassPostProcessCombineLUTs::Process(FRenderingCompositePassContext& Cont
 		FIntRect DestRect(0, 0, bUseVolumeTextureLUT ? GLUTSize : GLUTSize * GLUTSize, GLUTSize);
 	
 		// Common setup
-		SetRenderTarget(Context.RHICmdList, nullptr, nullptr);
+		// #todo-renderpasses remove once everything is renderpasses
+		UnbindRenderTargets(Context.RHICmdList);
 		Context.SetViewportAndCallRHI(DestRect, 0.0f, 1.0f);
 		
 		static FName AsyncEndFenceName(TEXT("AsyncCombineLUTsEndFence"));
@@ -946,45 +947,44 @@ void FRCPassPostProcessCombineLUTs::Process(FRenderingCompositePassContext& Cont
 	else
 	{
 		// Set the view family's render target/viewport.
+		ERenderTargetActions LoadStoreAction = ERenderTargetActions::DontLoad_Store;
 		if (IsMobilePlatform(ShaderPlatform))
 		{
-			// Full clear to avoid restore
-			SetRenderTarget(Context.RHICmdList, DestRenderTarget->TargetableTexture, FTextureRHIRef(), ESimpleRenderTargetMode::EClearColorAndDepth);
+			LoadStoreAction = ERenderTargetActions::Clear_Store;
 		}
-		else
+
+		FRHIRenderPassInfo RPInfo(DestRenderTarget->TargetableTexture, LoadStoreAction, DestRenderTarget->ShaderResourceTexture);
+		Context.RHICmdList.BeginRenderPass(RPInfo, TEXT("CombineLUTs"));
 		{
-			SetRenderTarget(Context.RHICmdList, DestRenderTarget->TargetableTexture, FTextureRHIRef(), ESimpleRenderTargetMode::EUninitializedColorAndDepth);
-		}
-		
-		Context.SetViewportAndCallRHI(0, 0, 0.0f, DestSize.X, DestSize.Y, 1.0f );
+			Context.SetViewportAndCallRHI(0, 0, 0.0f, DestSize.X, DestSize.Y, 1.0f);
 
-		const FVolumeBounds VolumeBounds(GLUTSize);
+			const FVolumeBounds VolumeBounds(GLUTSize);
 
-		SetLUTBlenderShader(Context, Context.RHICmdList, LocalCount, LocalTextures, LocalWeights, VolumeBounds, bUseVolumeTextureLUT);
+			SetLUTBlenderShader(Context, Context.RHICmdList, LocalCount, LocalTextures, LocalWeights, VolumeBounds, bUseVolumeTextureLUT);
 
-		if (bUseVolumeTextureLUT)
-		{
-			// use volume texture 16x16x16
-			RasterizeToVolumeTexture(Context.RHICmdList, VolumeBounds);
-		}
-		else
-		{
-			// use unwrapped 2d texture 256x16
-			TShaderMapRef<FPostProcessVS> VertexShader(Context.GetShaderMap());
-
-			DrawRectangle( 
-				Context.RHICmdList,
-				0, 0,										// XY
-				GLUTSize * GLUTSize, GLUTSize,				// SizeXY
-				0, 0,										// UV
-				GLUTSize * GLUTSize, GLUTSize,				// SizeUV
-				FIntPoint(GLUTSize * GLUTSize, GLUTSize),	// TargetSize
-				FIntPoint(GLUTSize * GLUTSize, GLUTSize),	// TextureSize
-				*VertexShader,
-				EDRF_UseTriangleOptimization);
+			if (bUseVolumeTextureLUT)
+			{
+				// use volume texture 16x16x16
+				RasterizeToVolumeTexture(Context.RHICmdList, VolumeBounds);
+			}
+			else
+			{
+				// use unwrapped 2d texture 256x16
+				TShaderMapRef<FPostProcessVS> VertexShader(Context.GetShaderMap());
+
+				DrawRectangle(
+					Context.RHICmdList,
+					0, 0,										// XY
+					GLUTSize * GLUTSize, GLUTSize,				// SizeXY
+					0, 0,										// UV
+					GLUTSize * GLUTSize, GLUTSize,				// SizeUV
+					FIntPoint(GLUTSize * GLUTSize, GLUTSize),	// TargetSize
+					FIntPoint(GLUTSize * GLUTSize, GLUTSize),	// TextureSize
+					*VertexShader,
+					EDRF_UseTriangleOptimization);
+			}
 		}
-
-		Context.RHICmdList.CopyToResolveTarget(DestRenderTarget->TargetableTexture, DestRenderTarget->ShaderResourceTexture, FResolveParams());
+		Context.RHICmdList.EndRenderPass();
 	}
 
 	Context.View.SetValidTonemappingLUT();
