Render Passes

diff --git a/Engine/Source/Runtime/Renderer/Private/PostProcess/PostProcessCompositeEditorPrimitives.cpp b/Engine/Source/Runtime/Renderer/Private/PostProcess/PostProcessCompositeEditorPrimitives.cpp
--- a/Engine/Source/Runtime/Renderer/Private/PostProcess/PostProcessCompositeEditorPrimitives.cpp
+++ b/Engine/Source/Runtime/Renderer/Private/PostProcess/PostProcessCompositeEditorPrimitives.cpp
@@ -442,89 +442,112 @@ void FRCPassPostProcessCompositeEditorPrimitives::Process(FRenderingCompositePas
 	FIntPoint SrcSize = InputDesc->Extent;
 
 	{
-		ESimpleRenderTargetMode RTMode = bClearIsNeeded ? ESimpleRenderTargetMode::EClearColorAndDepth : ESimpleRenderTargetMode::EExistingColorAndDepth;
-
-		SetRenderTarget(Context.RHICmdList, EditorColorTarget, EditorDepthTarget, RTMode);
-		Context.SetViewportAndCallRHI(ViewRect);
-
-		// Populate depth from scene depth.
+		FRHIRenderPassInfo RPInfo;
+		RPInfo.ColorRenderTargets[0].RenderTarget = EditorColorTarget;
+		RPInfo.DepthStencilRenderTarget.DepthStencilTarget = EditorDepthTarget;
+		RPInfo.DepthStencilRenderTarget.ExclusiveDepthStencil = FExclusiveDepthStencil::DepthWrite_StencilWrite;
 		if (bClearIsNeeded)
 		{
-			SCOPED_DRAW_EVENTF(Context.RHICmdList, TemporalAA, TEXT("PopulateEditorPrimitivesDepthBuffer %dx%d msaa=%d"),
-				ViewRect.Width(), ViewRect.Height(), MSAASampleCount);
-
-			if (MSAASampleCount == 1)
-			{
-				SetPopulateSceneDepthForEditorPrimitivesShaderTempl<1>(Context);
-			}
-			else
-			{
-				SetPopulateSceneDepthForEditorPrimitivesShaderTempl<2>(Context);
-			}
-
-			TShaderMapRef<FPostProcessVS> VertexShader(Context.GetShaderMap());
-
-			// Draw a quad mapping our render targets to the view's render target
-			DrawRectangle(
-				Context.RHICmdList,
-				0, 0,
-				ViewRect.Width(), ViewRect.Height(),
-				Context.View.ViewRect.Min.X, Context.View.ViewRect.Min.Y,
-				Context.View.ViewRect.Width(), Context.View.ViewRect.Height(),
-				ViewRect.Size(),
-				SrcSize,
-				*VertexShader,
-				EDRF_UseTriangleOptimization);
-		}
-
-		TUniformBufferRef<FOpaqueBasePassUniformParameters> OpaqueBasePassUniformBuffer;
-		TUniformBufferRef<FMobileBasePassUniformParameters> MobileBasePassUniformBuffer;
-		FUniformBufferRHIParamRef BasePassUniformBuffer = nullptr;
-
-		if (bDeferredBasePass)
-		{
-			CreateOpaqueBasePassUniformBuffer(Context.RHICmdList, EditorView, nullptr, OpaqueBasePassUniformBuffer);
-			BasePassUniformBuffer = OpaqueBasePassUniformBuffer;
+			RPInfo.ColorRenderTargets[0].Action = ERenderTargetActions::Clear_Store;
+			RPInfo.DepthStencilRenderTarget.Action = EDepthStencilTargetActions::ClearDepthStencil_StoreDepthStencil;
 		}
 		else
 		{
-			CreateMobileBasePassUniformBuffer(Context.RHICmdList, EditorView, true, MobileBasePassUniformBuffer);
-			BasePassUniformBuffer = MobileBasePassUniformBuffer;
+			RPInfo.ColorRenderTargets[0].Action = ERenderTargetActions::Load_Store;
+			RPInfo.DepthStencilRenderTarget.Action = EDepthStencilTargetActions::LoadDepthStencil_StoreDepthStencil;
 		}
 
-		FDrawingPolicyRenderState DrawRenderState(EditorView, BasePassUniformBuffer);
-		DrawRenderState.SetDepthStencilAccess(FExclusiveDepthStencil::DepthWrite_StencilWrite);
-		DrawRenderState.SetBlendState(TStaticBlendStateWriteMask<CW_RGBA>::GetRHI());
+		// It's possible to have no depth target here.
+		if (!EditorDepthTarget)
+		{
+			RPInfo.DepthStencilRenderTarget.Action = EDepthStencilTargetActions::DontLoad_DontStore;
+			RPInfo.DepthStencilRenderTarget.ExclusiveDepthStencil = FExclusiveDepthStencil::DepthNop_StencilNop;
+		}
 
-		// Draw editor primitives.
+		Context.RHICmdList.BeginRenderPass(RPInfo, TEXT("CompositeEditorPrimitives"));
 		{
-			SCOPED_DRAW_EVENTF(Context.RHICmdList, TemporalAA, TEXT("RenderViewEditorPrimitives %dx%d msaa=%d"),
-				ViewRect.Width(), ViewRect.Height(), MSAASampleCount);
-				
+			Context.SetViewportAndCallRHI(ViewRect);
+
+			// Populate depth from scene depth.
+			if (bClearIsNeeded)
+			{
+				SCOPED_DRAW_EVENTF(Context.RHICmdList, TemporalAA, TEXT("PopulateEditorPrimitivesDepthBuffer %dx%d msaa=%d"),
+					ViewRect.Width(), ViewRect.Height(), MSAASampleCount);
+
+				if (MSAASampleCount == 1)
+				{
+					SetPopulateSceneDepthForEditorPrimitivesShaderTempl<1>(Context);
+				}
+				else
+				{
+					SetPopulateSceneDepthForEditorPrimitivesShaderTempl<2>(Context);
+				}
+
+				TShaderMapRef<FPostProcessVS> VertexShader(Context.GetShaderMap());
+
+				// Draw a quad mapping our render targets to the view's render target
+				DrawRectangle(
+					Context.RHICmdList,
+					0, 0,
+					ViewRect.Width(), ViewRect.Height(),
+					Context.View.ViewRect.Min.X, Context.View.ViewRect.Min.Y,
+					Context.View.ViewRect.Width(), Context.View.ViewRect.Height(),
+					ViewRect.Size(),
+					SrcSize,
+					*VertexShader,
+					EDRF_UseTriangleOptimization);
+			}
+
+			TUniformBufferRef<FOpaqueBasePassUniformParameters> OpaqueBasePassUniformBuffer;
+			TUniformBufferRef<FMobileBasePassUniformParameters> MobileBasePassUniformBuffer;
+			FUniformBufferRHIParamRef BasePassUniformBuffer = nullptr;
+
 			if (bDeferredBasePass)
 			{
-				RenderEditorPrimitives<FBasePassOpaqueDrawingPolicyFactory>(Context.RHICmdList, EditorView, DrawRenderState);
+				CreateOpaqueBasePassUniformBuffer(Context.RHICmdList, EditorView, nullptr, OpaqueBasePassUniformBuffer);
+				BasePassUniformBuffer = OpaqueBasePassUniformBuffer;
 			}
 			else
 			{
-				RenderEditorPrimitives<FMobileBasePassOpaqueDrawingPolicyFactory>(Context.RHICmdList, EditorView, DrawRenderState);
+				CreateMobileBasePassUniformBuffer(Context.RHICmdList, EditorView, true, MobileBasePassUniformBuffer);
+				BasePassUniformBuffer = MobileBasePassUniformBuffer;
 			}
-		}
 
-		// Draw foreground editor primitives.
-		{
-			SCOPED_DRAW_EVENTF(Context.RHICmdList, TemporalAA, TEXT("RenderViewEditorForegroundPrimitives %dx%d msaa=%d"),
-				ViewRect.Width(), ViewRect.Height(), MSAASampleCount);
+			FDrawingPolicyRenderState DrawRenderState(EditorView, BasePassUniformBuffer);
+			DrawRenderState.SetDepthStencilAccess(FExclusiveDepthStencil::DepthWrite_StencilWrite);
+			DrawRenderState.SetBlendState(TStaticBlendStateWriteMask<CW_RGBA>::GetRHI());
 
-			if (bDeferredBasePass)
+			// Draw editor primitives.
 			{
-				RenderForegroundEditorPrimitives<FBasePassOpaqueDrawingPolicyFactory>(Context.RHICmdList, EditorView, DrawRenderState, /* bIgnoreExistingDepth = */ true);
+				SCOPED_DRAW_EVENTF(Context.RHICmdList, TemporalAA, TEXT("RenderViewEditorPrimitives %dx%d msaa=%d"),
+					ViewRect.Width(), ViewRect.Height(), MSAASampleCount);
+
+				if (bDeferredBasePass)
+				{
+					RenderEditorPrimitives<FBasePassOpaqueDrawingPolicyFactory>(Context.RHICmdList, EditorView, DrawRenderState);
+				}
+				else
+				{
+					RenderEditorPrimitives<FMobileBasePassOpaqueDrawingPolicyFactory>(Context.RHICmdList, EditorView, DrawRenderState);
+				}
 			}
-			else
+
+			// Draw foreground editor primitives.
 			{
-				RenderForegroundEditorPrimitives<FMobileBasePassOpaqueDrawingPolicyFactory>(Context.RHICmdList, EditorView, DrawRenderState, /* bIgnoreExistingDepth = */ true);
+				SCOPED_DRAW_EVENTF(Context.RHICmdList, TemporalAA, TEXT("RenderViewEditorForegroundPrimitives %dx%d msaa=%d"),
+					ViewRect.Width(), ViewRect.Height(), MSAASampleCount);
+
+				if (bDeferredBasePass)
+				{
+					RenderForegroundEditorPrimitives<FBasePassOpaqueDrawingPolicyFactory>(Context.RHICmdList, EditorView, DrawRenderState, /* bIgnoreExistingDepth = */ true);
+				}
+				else
+				{
+					RenderForegroundEditorPrimitives<FMobileBasePassOpaqueDrawingPolicyFactory>(Context.RHICmdList, EditorView, DrawRenderState, /* bIgnoreExistingDepth = */ true);
+				}
 			}
 		}
+		Context.RHICmdList.EndRenderPass();
 
 		GVisualizeTexture.SetCheckPoint(Context.RHICmdList, SceneContext.EditorPrimitivesColor);
 
@@ -547,55 +570,56 @@ void FRCPassPostProcessCompositeEditorPrimitives::Process(FRenderingCompositePas
 		FIntRect DestRect = Context.GetSceneColorDestRect(DestRenderTarget);
 
 		// Set the view family's render target/viewport.
-		SetRenderTarget(Context.RHICmdList, DestRenderTargetSurface, FTextureRHIRef());
+		FRHIRenderPassInfo RPInfo(DestRenderTargetSurface, ERenderTargetActions::Load_Store, DestRenderTarget.ShaderResourceTexture);
+		Context.RHICmdList.BeginRenderPass(RPInfo, TEXT("ComposeEditorPrimitives"));
+		{
+			Context.SetViewportAndCallRHI(DestRect);
 
-		Context.SetViewportAndCallRHI(DestRect);
+			// If clear is not needed, that mean already have something in MSAA buffers. Because not populating scene depth buffer
+			// into MSAA depth buffer, then force to compose any sample that have non null depth as if alpha channel was 1.
+			bool bComposeAnyNonNullDepth = !bClearIsNeeded;
 
-		// If clear is not needed, that mean already have something in MSAA buffers. Because not populating scene depth buffer
-		// into MSAA depth buffer, then force to compose any sample that have non null depth as if alpha channel was 1.
-		bool bComposeAnyNonNullDepth = !bClearIsNeeded;
+			if (!bDeferredBasePass)
+			{
+				SetCompositePrimitivesShaderTempl<0>(Context, bComposeAnyNonNullDepth);
+			}
+			else if (MSAASampleCount == 1)
+			{
+				SetCompositePrimitivesShaderTempl<1>(Context, bComposeAnyNonNullDepth);
+			}
+			else if (MSAASampleCount == 2)
+			{
+				SetCompositePrimitivesShaderTempl<2>(Context, bComposeAnyNonNullDepth);
+			}
+			else if (MSAASampleCount == 4)
+			{
+				SetCompositePrimitivesShaderTempl<4>(Context, bComposeAnyNonNullDepth);
+			}
+			else if (MSAASampleCount == 8)
+			{
+				SetCompositePrimitivesShaderTempl<8>(Context, bComposeAnyNonNullDepth);
+			}
+			else
+			{
+				// not supported, internal error
+				check(0);
+			}
 
-		if (!bDeferredBasePass)
-		{
-			SetCompositePrimitivesShaderTempl<0>(Context, bComposeAnyNonNullDepth);
-		}
-		else if (MSAASampleCount == 1)
-		{
-			SetCompositePrimitivesShaderTempl<1>(Context, bComposeAnyNonNullDepth);
-		}
-		else if (MSAASampleCount == 2)
-		{
-			SetCompositePrimitivesShaderTempl<2>(Context, bComposeAnyNonNullDepth);
-		}
-		else if (MSAASampleCount == 4)
-		{
-			SetCompositePrimitivesShaderTempl<4>(Context, bComposeAnyNonNullDepth);
-		}
-		else if (MSAASampleCount == 8)
-		{
-			SetCompositePrimitivesShaderTempl<8>(Context, bComposeAnyNonNullDepth);
-		}
-		else
-		{
-			// not supported, internal error
-			check(0);
-		}
+			TShaderMapRef<FPostProcessVS> VertexShader(Context.GetShaderMap());
 
-		TShaderMapRef<FPostProcessVS> VertexShader(Context.GetShaderMap());
-
-		// Draw a quad mapping our render targets to the view's render target
-		DrawRectangle(
-			Context.RHICmdList,
-			0, 0,
-			DestRect.Width(), DestRect.Height(),
-			ViewRect.Min.X, ViewRect.Min.Y,
-			ViewRect.Width(), ViewRect.Height(),
-			DestRect.Size(),
-			SrcSize,
-			*VertexShader,
-			EDRF_UseTriangleOptimization);
-
-		Context.RHICmdList.CopyToResolveTarget(DestRenderTargetSurface, DestRenderTarget.ShaderResourceTexture, FResolveParams());
+			// Draw a quad mapping our render targets to the view's render target
+			DrawRectangle(
+				Context.RHICmdList,
+				0, 0,
+				DestRect.Width(), DestRect.Height(),
+				ViewRect.Min.X, ViewRect.Min.Y,
+				ViewRect.Width(), ViewRect.Height(),
+				DestRect.Size(),
+				SrcSize,
+				*VertexShader,
+				EDRF_UseTriangleOptimization);
+		}
+		Context.RHICmdList.EndRenderPass();
 	}
 
 	// Clean up targets
diff --git a/Engine/Source/Runtime/Renderer/Private/PostProcess/PostProcessDOF.cpp b/Engine/Source/Runtime/Renderer/Private/PostProcess/PostProcessDOF.cpp
--- a/Engine/Source/Runtime/Renderer/Private/PostProcess/PostProcessDOF.cpp
+++ b/Engine/Source/Runtime/Renderer/Private/PostProcess/PostProcessDOF.cpp
@@ -145,7 +145,7 @@ void FRCPassPostProcessDOFSetup::Process(FRenderingCompositePassContext& Context
 
 	const FPooledRenderTargetDesc* InputDesc = GetInputDesc(ePId_Input0);
 
-	if(!InputDesc)
+	if (!InputDesc)
 	{
 		// input is not hooked up correctly
 		return;
@@ -177,88 +177,96 @@ void FRCPassPostProcessDOFSetup::Process(FRenderingCompositePassContext& Context
 		DestRenderTarget0.TargetableTexture,
 		DestRenderTarget1.TargetableTexture
 	};
+
+	FTextureRHIParamRef ResolveTargets[2] =
+	{
+		DestRenderTarget0.ShaderResourceTexture,
+		nullptr
+	};
+	if (DestRenderTarget1.TargetableTexture)
+	{
+		ResolveTargets[1] = DestRenderTarget1.ShaderResourceTexture;
+	}
+	ERenderTargetActions LoadStoreAction = ERenderTargetActions::Load_Store;
+
 	//@todo Ronin find a way to use the same codepath for all platforms.
 	const EShaderPlatform ShaderPlatform = GShaderPlatformForFeatureLevel[Context.GetFeatureLevel()];
 	if (IsVulkanMobilePlatform(ShaderPlatform))
 	{
-		SetRenderTargets(Context.RHICmdList, NumRenderTargets, RenderTargets, FTextureRHIParamRef(), ESimpleRenderTargetMode::EClearColorAndDepth, FExclusiveDepthStencil());
+		LoadStoreAction = ERenderTargetActions::Clear_Store;
 	}
-	else
-	{
-		SetRenderTargets(Context.RHICmdList, NumRenderTargets, RenderTargets, FTextureRHIParamRef(), 0, NULL);
-	}
-	
-	if (View.StereoPass == eSSP_FULL)
+
+	FRHIRenderPassInfo RPInfo(NumRenderTargets, RenderTargets, LoadStoreAction, ResolveTargets);
+	Context.RHICmdList.BeginRenderPass(RPInfo, TEXT("DOFSetup"));
 	{
-		FLinearColor ClearColors[2] =
+		if (View.StereoPass == eSSP_FULL)
 		{
-			FLinearColor(0, 0, 0, 0),
-			FLinearColor(0, 0, 0, 0)
-		};
-		// is optimized away if possible (RT size=view size, )
-		DrawClearQuadMRT(Context.RHICmdList, true, NumRenderTargets, ClearColors, false, 0, false, 0, DestSize, DestRect);
-	}
+			FLinearColor ClearColors[2] =
+			{
+				FLinearColor(0, 0, 0, 0),
+				FLinearColor(0, 0, 0, 0)
+			};
+			// is optimized away if possible (RT size=view size, )
+			DrawClearQuadMRT(Context.RHICmdList, true, NumRenderTargets, ClearColors, false, 0, false, 0, DestSize, DestRect);
+		}
 
-	Context.SetViewportAndCallRHI(DestRect.Min.X, DestRect.Min.Y, 0.0f, DestRect.Max.X + 1, DestRect.Max.Y + 1, 1.0f );
-	
-	const float DOFVignetteSize = FMath::Max(0.0f, View.FinalPostProcessSettings.DepthOfFieldVignetteSize);
+		Context.SetViewportAndCallRHI(DestRect.Min.X, DestRect.Min.Y, 0.0f, DestRect.Max.X + 1, DestRect.Max.Y + 1, 1.0f);
 
-	// todo: test is conservative, with bad content we would waste a bit of performance
-	const bool bDOFVignette = bNearBlur && DOFVignetteSize < 200.0f;
-	
-	// 0:off, 1:on, 2:on with Vignette
-	uint32 NearBlur = 0;
-	{
-		if(bNearBlur)
+		const float DOFVignetteSize = FMath::Max(0.0f, View.FinalPostProcessSettings.DepthOfFieldVignetteSize);
+
+		// todo: test is conservative, with bad content we would waste a bit of performance
+		const bool bDOFVignette = bNearBlur && DOFVignetteSize < 200.0f;
+
+		// 0:off, 1:on, 2:on with Vignette
+		uint32 NearBlur = 0;
 		{
-			NearBlur = (DOFVignetteSize < 200.0f) ? 2 : 1;
+			if (bNearBlur)
+			{
+				NearBlur = (DOFVignetteSize < 200.0f) ? 2 : 1;
+			}
 		}
-	}
-		
-	FShader* VertexShader = 0;
 
-	if(bFarBlur)
-	{
-		switch(NearBlur)
+		FShader* VertexShader = 0;
+
+		if (bFarBlur)
 		{
-			case 0: VertexShader = SetDOFShaderTempl<1, 0>(Context);break;
-			case 1: VertexShader = SetDOFShaderTempl<1, 1>(Context);break;
-			case 2: VertexShader = SetDOFShaderTempl<1, 2>(Context);break;
+			switch (NearBlur)
+			{
+			case 0: VertexShader = SetDOFShaderTempl<1, 0>(Context); break;
+			case 1: VertexShader = SetDOFShaderTempl<1, 1>(Context); break;
+			case 2: VertexShader = SetDOFShaderTempl<1, 2>(Context); break;
 			default: check(!"Invalid NearBlur input");
+			}
 		}
-	}
-	else
-	{
-		switch(NearBlur)
+		else
 		{
-			case 0: check(!"Internal error: In this case we should not run the pass");break;
-			case 1: VertexShader = SetDOFShaderTempl<0, 1>(Context);break;
-			case 2: VertexShader = SetDOFShaderTempl<0, 2>(Context);break;
+			switch (NearBlur)
+			{
+			case 0: check(!"Internal error: In this case we should not run the pass"); break;
+			case 1: VertexShader = SetDOFShaderTempl<0, 1>(Context); break;
+			case 2: VertexShader = SetDOFShaderTempl<0, 2>(Context); break;
 			default: check(!"Invalid NearBlur input");
+			}
 		}
-	}
 
-	DrawPostProcessPass(
-		Context.RHICmdList,
-		0, 0,
-		DestRect.Width() + 1, DestRect.Height() + 1,
-		SrcRect.Min.X, SrcRect.Min.Y,
-		SrcRect.Width() + 1, SrcRect.Height() + 1,
-		DestRect.Size() + FIntPoint(1, 1),
-		SrcSize,
-		VertexShader,
-		View.StereoPass,
-		Context.HasHmdMesh(),
-		EDRF_UseTriangleOptimization);
+		DrawPostProcessPass(
+			Context.RHICmdList,
+			0, 0,
+			DestRect.Width() + 1, DestRect.Height() + 1,
+			SrcRect.Min.X, SrcRect.Min.Y,
+			SrcRect.Width() + 1, SrcRect.Height() + 1,
+			DestRect.Size() + FIntPoint(1, 1),
+			SrcSize,
+			VertexShader,
+			View.StereoPass,
+			Context.HasHmdMesh(),
+			EDRF_UseTriangleOptimization);
+	}
+	Context.RHICmdList.EndRenderPass();
 
 	// #todo-rco: needed to avoid multiple resolves clearing the RT with VK.
-	SetRenderTarget(Context.RHICmdList, nullptr, nullptr);
-
-	Context.RHICmdList.CopyToResolveTarget(DestRenderTarget0.TargetableTexture, DestRenderTarget0.ShaderResourceTexture, FResolveParams());
-	if (DestRenderTarget1.TargetableTexture)
-	{
-		Context.RHICmdList.CopyToResolveTarget(DestRenderTarget1.TargetableTexture, DestRenderTarget1.ShaderResourceTexture, FResolveParams());
-	}
+	// #todo mattc this is probably busted now since the resolves previously happened after this SetRenderTarget.
+	UnbindRenderTargets(Context.RHICmdList);
 }
 
 FPooledRenderTargetDesc FRCPassPostProcessDOFSetup::ComputeOutputDesc(EPassOutputId InPassOutputId) const
@@ -443,60 +451,63 @@ void FRCPassPostProcessDOFRecombine::Process(FRenderingCompositePassContext& Con
 
 	//@todo Ronin find a way to use the same codepath for all platforms.
 	const EShaderPlatform ShaderPlatform = GShaderPlatformForFeatureLevel[Context.GetFeatureLevel()];
+	ERenderTargetActions LoadStoreAction = ERenderTargetActions::Load_Store;
 	if (IsVulkanMobilePlatform(ShaderPlatform))
 	{
-		SetRenderTarget(Context.RHICmdList, DestRenderTarget.TargetableTexture, FTextureRHIRef(), ESimpleRenderTargetMode::EClearColorAndDepth);
+		LoadStoreAction = ERenderTargetActions::Clear_Store;
 	}
-	else
+	else if (View.StereoPass == eSSP_FULL)
 	{
-		SetRenderTarget(Context.RHICmdList, DestRenderTarget.TargetableTexture, FTextureRHIRef());
+		// #todo-renderpasses a possible optimization here is to use DontLoad since we'll clear immediately.
+		LoadStoreAction = ERenderTargetActions::Load_Store;
+	}
 
+	FRHIRenderPassInfo RPInfo(DestRenderTarget.TargetableTexture, LoadStoreAction, DestRenderTarget.ShaderResourceTexture);
+	Context.RHICmdList.BeginRenderPass(RPInfo, TEXT("DOFRecombine"));
+	{
 		if (View.StereoPass == eSSP_FULL)
 		{
 			// is optimized away if possible (RT size=view size, )
 			DrawClearQuad(Context.RHICmdList, true, FLinearColor::Black, false, 0, false, 0, PassOutputs[0].RenderTargetDesc.Extent, View.ViewRect);
 		}
-	}
-
-	Context.SetViewportAndCallRHI(View.ViewRect);
+		Context.SetViewportAndCallRHI(View.ViewRect);
 
-	bool bFarBlur = GetInputDesc(ePId_Input1) != 0;
-	bool bNearBlur = GetInputDesc(ePId_Input2) != 0;
-	bool bSeparateTranslucency = GetInputDesc(ePId_Input3) != 0;
+		bool bFarBlur = GetInputDesc(ePId_Input1) != 0;
+		bool bNearBlur = GetInputDesc(ePId_Input2) != 0;
+		bool bSeparateTranslucency = GetInputDesc(ePId_Input3) != 0;
 
-	FShader* VertexShader = 0;
+		FShader* VertexShader = 0;
 
-	if(bFarBlur)
-	{
-		if(bNearBlur)
+		if (bFarBlur)
 		{
-			VertexShader = SetDOFRecombineShaderTempl<1, 1>(Context, bSeparateTranslucency);
+			if (bNearBlur)
+			{
+				VertexShader = SetDOFRecombineShaderTempl<1, 1>(Context, bSeparateTranslucency);
+			}
+			else
+			{
+				VertexShader = SetDOFRecombineShaderTempl<1, 0>(Context, bSeparateTranslucency);
+			}
 		}
 		else
 		{
-			VertexShader = SetDOFRecombineShaderTempl<1, 0>(Context, bSeparateTranslucency);
+			VertexShader = SetDOFRecombineShaderTempl<0, 1>(Context, bSeparateTranslucency);
 		}
-	}
-	else
-	{
-		VertexShader = SetDOFRecombineShaderTempl<0, 1>(Context, bSeparateTranslucency);
-	}
 
-	DrawPostProcessPass(
-		Context.RHICmdList,
-		0, 0,
-		View.ViewRect.Width(), View.ViewRect.Height(),
-		HalfResViewRect.Min.X, HalfResViewRect.Min.Y,
-		HalfResViewRect.Width(), HalfResViewRect.Height(),
-		View.ViewRect.Size(),
-		TexSize,
-		VertexShader,
-		View.StereoPass,
-		Context.HasHmdMesh(),
-		EDRF_UseTriangleOptimization);
-
-
-	Context.RHICmdList.CopyToResolveTarget(DestRenderTarget.TargetableTexture, DestRenderTarget.ShaderResourceTexture, FResolveParams());
+		DrawPostProcessPass(
+			Context.RHICmdList,
+			0, 0,
+			View.ViewRect.Width(), View.ViewRect.Height(),
+			HalfResViewRect.Min.X, HalfResViewRect.Min.Y,
+			HalfResViewRect.Width(), HalfResViewRect.Height(),
+			View.ViewRect.Size(),
+			TexSize,
+			VertexShader,
+			View.StereoPass,
+			Context.HasHmdMesh(),
+			EDRF_UseTriangleOptimization);
+	}
+	Context.RHICmdList.EndRenderPass();
 }
 
 FPooledRenderTargetDesc FRCPassPostProcessDOFRecombine::ComputeOutputDesc(EPassOutputId InPassOutputId) const
diff --git a/Engine/Source/Runtime/Renderer/Private/PostProcess/PostProcessDownsample.cpp b/Engine/Source/Runtime/Renderer/Private/PostProcess/PostProcessDownsample.cpp
--- a/Engine/Source/Runtime/Renderer/Private/PostProcess/PostProcessDownsample.cpp
+++ b/Engine/Source/Runtime/Renderer/Private/PostProcess/PostProcessDownsample.cpp
@@ -308,7 +308,8 @@ void FRCPassPostProcessDownsample::Process(FRenderingCompositePassContext& Conte
 		DestRect = {Context.SceneColorViewRect.Min, Context.SceneColorViewRect.Min + DestSize};
 
 		// Common setup
-		SetRenderTarget(Context.RHICmdList, nullptr, nullptr);
+		// #todo-renderpasses remove once everything is renderpasses
+		UnbindRenderTargets(Context.RHICmdList);
 		Context.SetViewportAndCallRHI(DestRect, 0.0f, 1.0f);
 		
 		static FName AsyncEndFenceName(TEXT("AsyncDownsampleEndFence"));
@@ -365,82 +366,69 @@ void FRCPassPostProcessDownsample::Process(FRenderingCompositePassContext& Conte
 	}
 	else
 	{
-		// check if we have to clear the whole surface.
-		// Otherwise perform the clear when the dest rectangle has been computed.
-		auto FeatureLevel = Context.View.GetFeatureLevel();
-		if (FeatureLevel == ERHIFeatureLevel::ES2 || FeatureLevel == ERHIFeatureLevel::ES3_1)
+		// #todo-renderpasses only clear dest rectangle if it's been computed
+		FRHIRenderPassInfo RPInfo(DestRenderTarget.TargetableTexture, ERenderTargetActions::Clear_Store, DestRenderTarget.ShaderResourceTexture);
+		Context.RHICmdList.BeginRenderPass(RPInfo, TEXT("PostProcessDownsample"));
 		{
-			// Set the view family's render target/viewport.
-			SetRenderTarget(Context.RHICmdList, DestRenderTarget.TargetableTexture, FTextureRHIRef(), ESimpleRenderTargetMode::EClearColorAndDepth);
 			Context.SetViewportAndCallRHI(0, 0, 0.0f, DestSize.X, DestSize.Y, 1.0f);
-		}
-		else
-		{
-			// Set the view family's render target/viewport.
-			FRHIRenderTargetView RtView = FRHIRenderTargetView(DestRenderTarget.TargetableTexture, ERenderTargetLoadAction::ENoAction);
-			FRHISetRenderTargetsInfo Info(1, &RtView, FRHIDepthRenderTargetView());
-			Context.RHICmdList.SetRenderTargetsAndClear(Info);
-			Context.SetViewportAndCallRHI(0, 0, 0.0f, DestSize.X, DestSize.Y, 1.0f);
-		}
-
-		// InflateSize increases the size of the source/dest rectangle to compensate for bilinear reads and UIBlur pass requirements.
-		int32 InflateSize;
-		// if second input is hooked up
+			// InflateSize increases the size of the source/dest rectangle to compensate for bilinear reads and UIBlur pass requirements.
+			int32 InflateSize;
+			// if second input is hooked up
 
-		if (bManuallyClampUV)
-		{
-			if (bIsDepthInputAvailable)
-			{
-				// also put depth in alpha
-				InflateSize = 2;
-				SetShader<2, true>(Context, InputDesc, SrcSize, SrcRect);
-			}
-			else if (Quality == 0)
+			if (bManuallyClampUV)
 			{
-				SetShader<0, true>(Context, InputDesc, SrcSize, SrcRect);
-				InflateSize = 1;
-			}
-			else
-			{
-				SetShader<1, true>(Context, InputDesc, SrcSize, SrcRect);
-				InflateSize = 2;
-			}
-		}
-		else
-		{
-			if (bIsDepthInputAvailable)
-			{
-				// also put depth in alpha
-				InflateSize = 2;
-				SetShader<2, false>(Context, InputDesc, SrcSize, SrcRect);
-			}
-			else if (Quality == 0)
-			{
-				SetShader<0, false>(Context, InputDesc, SrcSize, SrcRect);
-				InflateSize = 1;
+				if (bIsDepthInputAvailable)
+				{
+					// also put depth in alpha
+					InflateSize = 2;
+					SetShader<2, true>(Context, InputDesc, SrcSize, SrcRect);
+				}
+				else if (Quality == 0)
+				{
+					SetShader<0, true>(Context, InputDesc, SrcSize, SrcRect);
+					InflateSize = 1;
+				}
+				else
+				{
+					SetShader<1, true>(Context, InputDesc, SrcSize, SrcRect);
+					InflateSize = 2;
+				}
 			}
 			else
 			{
-				SetShader<1, false>(Context, InputDesc, SrcSize, SrcRect);
-				InflateSize = 2;
+				if (bIsDepthInputAvailable)
+				{
+					// also put depth in alpha
+					InflateSize = 2;
+					SetShader<2, false>(Context, InputDesc, SrcSize, SrcRect);
+				}
+				else if (Quality == 0)
+				{
+					SetShader<0, false>(Context, InputDesc, SrcSize, SrcRect);
+					InflateSize = 1;
+				}
+				else
+				{
+					SetShader<1, false>(Context, InputDesc, SrcSize, SrcRect);
+					InflateSize = 2;
+				}
 			}
-		}
 
-		TShaderMapRef<FPostProcessDownsampleVS> VertexShader(Context.GetShaderMap());
-
-		SrcRect = DestRect * 2;
-		DrawRectangle(
-			Context.RHICmdList,
-			DestRect.Min.X, DestRect.Min.Y,
-			DestRect.Width(), DestRect.Height(),
-			SrcRect.Min.X, SrcRect.Min.Y,
-			SrcRect.Width(), SrcRect.Height(),
-			DestSize,
-			SrcSize,
-			*VertexShader,
-			EDRF_UseTriangleOptimization);
-
-		Context.RHICmdList.CopyToResolveTarget(DestRenderTarget.TargetableTexture, DestRenderTarget.ShaderResourceTexture, FResolveParams());
+			TShaderMapRef<FPostProcessDownsampleVS> VertexShader(Context.GetShaderMap());
+
+			SrcRect = DestRect * 2;
+			DrawRectangle(
+				Context.RHICmdList,
+				DestRect.Min.X, DestRect.Min.Y,
+				DestRect.Width(), DestRect.Height(),
+				SrcRect.Min.X, SrcRect.Min.Y,
+				SrcRect.Width(), SrcRect.Height(),
+				DestSize,
+				SrcSize,
+				*VertexShader,
+				EDRF_UseTriangleOptimization);
+		}
+		Context.RHICmdList.EndRenderPass();
 	}
 }
 
diff --git a/Engine/Source/Runtime/Renderer/Private/PostProcess/PostProcessEyeAdaptation.cpp b/Engine/Source/Runtime/Renderer/Private/PostProcess/PostProcessEyeAdaptation.cpp
--- a/Engine/Source/Runtime/Renderer/Private/PostProcess/PostProcessEyeAdaptation.cpp
+++ b/Engine/Source/Runtime/Renderer/Private/PostProcess/PostProcessEyeAdaptation.cpp
@@ -307,7 +307,8 @@ void FRCPassPostProcessEyeAdaptation::Process(FRenderingCompositePassContext& Co
 		FIntRect DestRect(0, 0, DestSize.X, DestSize.Y);
 
 		// Common setup
-		SetRenderTarget(Context.RHICmdList, nullptr, nullptr);
+		// #todo-renderpasses remove once everything is renderpasses
+		UnbindRenderTargets(Context.RHICmdList);
 		Context.SetViewportAndCallRHI(DestRect, 0.0f, 1.0f);
 		
 		static FName AsyncEndFenceName(TEXT("AsyncEyeAdaptationEndFence"));
@@ -346,41 +347,44 @@ void FRCPassPostProcessEyeAdaptation::Process(FRenderingCompositePassContext& Co
 
 		// we render to our own output render target, not the intermediate one created by the compositing system
 		// Set the view family's render target/viewport.
-		SetRenderTarget(Context.RHICmdList, DestRenderTarget.TargetableTexture, FTextureRHIRef(), true);
-		Context.SetViewportAndCallRHI(0, 0, 0.0f, DestSize.X, DestSize.Y, 1.0f);
-
-		FGraphicsPipelineStateInitializer GraphicsPSOInit;
-		Context.RHICmdList.ApplyCachedRenderTargets(GraphicsPSOInit);
-		GraphicsPSOInit.BlendState = TStaticBlendState<>::GetRHI();
-		GraphicsPSOInit.RasterizerState = TStaticRasterizerState<>::GetRHI();
-		GraphicsPSOInit.DepthStencilState = TStaticDepthStencilState<false, CF_Always>::GetRHI();
-
-		TShaderMapRef<FPostProcessVS> VertexShader(Context.GetShaderMap());
-		TShaderMapRef<FPostProcessEyeAdaptationPS> PixelShader(Context.GetShaderMap());
-
-		GraphicsPSOInit.BoundShaderState.VertexDeclarationRHI = GFilterVertexDeclaration.VertexDeclarationRHI;
-		GraphicsPSOInit.BoundShaderState.VertexShaderRHI = GETSAFERHISHADER_VERTEX(*VertexShader);
-		GraphicsPSOInit.BoundShaderState.PixelShaderRHI = GETSAFERHISHADER_PIXEL(*PixelShader);
-		GraphicsPSOInit.PrimitiveType = PT_TriangleList;
-
-		SetGraphicsPipelineState(Context.RHICmdList, GraphicsPSOInit);
-
-		PixelShader->SetPS(Context, Context.RHICmdList, LastEyeAdaptation);
-
-		// Draw a quad mapping scene color to the view's render target
-		DrawRectangle(
-			Context.RHICmdList,
-			0, 0,
-			DestSize.X, DestSize.Y,
-			0, 0,
-			DestSize.X, DestSize.Y,
-			DestSize,
-			DestSize,
-			*VertexShader,
-			EDRF_UseTriangleOptimization);
-
-		Context.RHICmdList.CopyToResolveTarget(DestRenderTarget.TargetableTexture, DestRenderTarget.ShaderResourceTexture, FResolveParams());
 
+		FRHIRenderPassInfo RPInfo(DestRenderTarget.TargetableTexture, ERenderTargetActions::Load_Store, DestRenderTarget.ShaderResourceTexture);
+		TransitionRenderPassTargets(Context.RHICmdList, RPInfo);
+		Context.RHICmdList.BeginRenderPass(RPInfo, TEXT("EyeAdaptation"));
+		{
+			Context.SetViewportAndCallRHI(0, 0, 0.0f, DestSize.X, DestSize.Y, 1.0f);
+
+			FGraphicsPipelineStateInitializer GraphicsPSOInit;
+			Context.RHICmdList.ApplyCachedRenderTargets(GraphicsPSOInit);
+			GraphicsPSOInit.BlendState = TStaticBlendState<>::GetRHI();
+			GraphicsPSOInit.RasterizerState = TStaticRasterizerState<>::GetRHI();
+			GraphicsPSOInit.DepthStencilState = TStaticDepthStencilState<false, CF_Always>::GetRHI();
+
+			TShaderMapRef<FPostProcessVS> VertexShader(Context.GetShaderMap());
+			TShaderMapRef<FPostProcessEyeAdaptationPS> PixelShader(Context.GetShaderMap());
+
+			GraphicsPSOInit.BoundShaderState.VertexDeclarationRHI = GFilterVertexDeclaration.VertexDeclarationRHI;
+			GraphicsPSOInit.BoundShaderState.VertexShaderRHI = GETSAFERHISHADER_VERTEX(*VertexShader);
+			GraphicsPSOInit.BoundShaderState.PixelShaderRHI = GETSAFERHISHADER_PIXEL(*PixelShader);
+			GraphicsPSOInit.PrimitiveType = PT_TriangleList;
+
+			SetGraphicsPipelineState(Context.RHICmdList, GraphicsPSOInit);
+
+			PixelShader->SetPS(Context, Context.RHICmdList, LastEyeAdaptation);
+
+			// Draw a quad mapping scene color to the view's render target
+			DrawRectangle(
+				Context.RHICmdList,
+				0, 0,
+				DestSize.X, DestSize.Y,
+				0, 0,
+				DestSize.X, DestSize.Y,
+				DestSize,
+				DestSize,
+				*VertexShader,
+				EDRF_UseTriangleOptimization);
+		}
+		Context.RHICmdList.EndRenderPass();
 		// Inform MultiGPU systems that we've finished updating this texture for this frame
 		Context.RHICmdList.EndUpdateMultiFrameResource(DestRenderTarget.ShaderResourceTexture);
 	}
@@ -534,61 +538,61 @@ void FRCPassPostProcessBasicEyeAdaptationSetUp::Process(FRenderingCompositePassC
 		return;
 	}
 
-	const FViewInfo& View = Context.View;
-	const FSceneViewFamily& ViewFamily = *(View.Family);
+	const FSceneRenderTargetItem& DestRenderTarget = PassOutputs[0].RequestSurface(Context);
+
+	ERenderTargetLoadAction LoadAction = Context.GetLoadActionForRenderTarget(DestRenderTarget);
 
-	FIntPoint SrcSize = InputDesc->Extent;
-	FIntPoint DestSize = PassOutputs[0].RenderTargetDesc.Extent;
+	FRHIRenderPassInfo RPInfo(DestRenderTarget.TargetableTexture, MakeRenderTargetActions(LoadAction, ERenderTargetStoreAction::EStore), DestRenderTarget.ShaderResourceTexture);
+	Context.RHICmdList.BeginRenderPass(RPInfo, TEXT("PostProcessBasicEyeAdaptationSetUp"));
+	{
+		const FViewInfo& View = Context.View;
+		const FSceneViewFamily& ViewFamily = *(View.Family);
 
-	// e.g. 4 means the input texture is 4x smaller than the buffer size
-	uint32 ScaleFactor = Context.ReferenceBufferSize.X / SrcSize.X;
+		FIntPoint SrcSize = InputDesc->Extent;
+		FIntPoint DestSize = PassOutputs[0].RenderTargetDesc.Extent;
 
-	FIntRect SrcRect = Context.SceneColorViewRect / ScaleFactor;
-	FIntRect DestRect = SrcRect;
+		// e.g. 4 means the input texture is 4x smaller than the buffer size
+		uint32 ScaleFactor = Context.ReferenceBufferSize.X / SrcSize.X;
 
-	SCOPED_DRAW_EVENTF(Context.RHICmdList, PostProcessBasicEyeAdaptationSetup, TEXT("PostProcessBasicEyeAdaptationSetup %dx%d"), DestRect.Width(), DestRect.Height());
+		FIntRect SrcRect = Context.SceneColorViewRect / ScaleFactor;
+		FIntRect DestRect = SrcRect;
 
-	const FSceneRenderTargetItem& DestRenderTarget = PassOutputs[0].RequestSurface(Context);
+		SCOPED_DRAW_EVENTF(Context.RHICmdList, PostProcessBasicEyeAdaptationSetup, TEXT("PostProcessBasicEyeAdaptationSetup %dx%d"), DestRect.Width(), DestRect.Height());
 
-	ERenderTargetLoadAction LoadAction = Context.GetLoadActionForRenderTarget(DestRenderTarget);
+		Context.SetViewportAndCallRHI(0, 0, 0.0f, DestSize.X, DestSize.Y, 1.0f);
+
+		FGraphicsPipelineStateInitializer GraphicsPSOInit;
+		Context.RHICmdList.ApplyCachedRenderTargets(GraphicsPSOInit);
+		GraphicsPSOInit.BlendState = TStaticBlendState<>::GetRHI();
+		GraphicsPSOInit.RasterizerState = TStaticRasterizerState<>::GetRHI();
+		GraphicsPSOInit.DepthStencilState = TStaticDepthStencilState<false, CF_Always>::GetRHI();
 
-	FRHIRenderTargetView RtView = FRHIRenderTargetView(DestRenderTarget.TargetableTexture, LoadAction);
-	FRHISetRenderTargetsInfo Info(1, &RtView, FRHIDepthRenderTargetView());
-	Context.RHICmdList.SetRenderTargetsAndClear(Info);
-	Context.SetViewportAndCallRHI(0, 0, 0.0f, DestSize.X, DestSize.Y, 1.0f);
-
-	FGraphicsPipelineStateInitializer GraphicsPSOInit;
-	Context.RHICmdList.ApplyCachedRenderTargets(GraphicsPSOInit);
-	GraphicsPSOInit.BlendState = TStaticBlendState<>::GetRHI();
-	GraphicsPSOInit.RasterizerState = TStaticRasterizerState<>::GetRHI();
-	GraphicsPSOInit.DepthStencilState = TStaticDepthStencilState<false, CF_Always>::GetRHI();
-
-	TShaderMapRef<FPostProcessVS> VertexShader(Context.GetShaderMap());
-	TShaderMapRef<FPostProcessBasicEyeAdaptationSetupPS> PixelShader(Context.GetShaderMap());
-
-	GraphicsPSOInit.BoundShaderState.VertexDeclarationRHI = GFilterVertexDeclaration.VertexDeclarationRHI;
-	GraphicsPSOInit.BoundShaderState.VertexShaderRHI = GETSAFERHISHADER_VERTEX(*VertexShader);
-	GraphicsPSOInit.BoundShaderState.PixelShaderRHI = GETSAFERHISHADER_PIXEL(*PixelShader);
-	GraphicsPSOInit.PrimitiveType = PT_TriangleList;
-
-	SetGraphicsPipelineState(Context.RHICmdList, GraphicsPSOInit);
-
-	PixelShader->SetPS(Context);
-
-	DrawPostProcessPass(
-		Context.RHICmdList,
-		DestRect.Min.X, DestRect.Min.Y,
-		DestRect.Width(), DestRect.Height(),
-		SrcRect.Min.X, SrcRect.Min.Y,
-		SrcRect.Width(), SrcRect.Height(),
-		DestSize,
-		SrcSize,
-		*VertexShader,
-		View.StereoPass,
-		Context.HasHmdMesh(),
-		EDRF_UseTriangleOptimization);
-
-	Context.RHICmdList.CopyToResolveTarget(DestRenderTarget.TargetableTexture, DestRenderTarget.ShaderResourceTexture, FResolveParams());
+		TShaderMapRef<FPostProcessVS> VertexShader(Context.GetShaderMap());
+		TShaderMapRef<FPostProcessBasicEyeAdaptationSetupPS> PixelShader(Context.GetShaderMap());
+
+		GraphicsPSOInit.BoundShaderState.VertexDeclarationRHI = GFilterVertexDeclaration.VertexDeclarationRHI;
+		GraphicsPSOInit.BoundShaderState.VertexShaderRHI = GETSAFERHISHADER_VERTEX(*VertexShader);
+		GraphicsPSOInit.BoundShaderState.PixelShaderRHI = GETSAFERHISHADER_PIXEL(*PixelShader);
+		GraphicsPSOInit.PrimitiveType = PT_TriangleList;
+
+		SetGraphicsPipelineState(Context.RHICmdList, GraphicsPSOInit);
+
+		PixelShader->SetPS(Context);
+
+		DrawPostProcessPass(
+			Context.RHICmdList,
+			DestRect.Min.X, DestRect.Min.Y,
+			DestRect.Width(), DestRect.Height(),
+			SrcRect.Min.X, SrcRect.Min.Y,
+			SrcRect.Width(), SrcRect.Height(),
+			DestSize,
+			SrcSize,
+			*VertexShader,
+			View.StereoPass,
+			Context.HasHmdMesh(),
+			EDRF_UseTriangleOptimization);
+	}
+	Context.RHICmdList.EndRenderPass();
 }
 
 FPooledRenderTargetDesc FRCPassPostProcessBasicEyeAdaptationSetUp::ComputeOutputDesc(EPassOutputId InPassOutputId) const
@@ -727,42 +731,46 @@ void FRCPassPostProcessBasicEyeAdaptation::Process(FRenderingCompositePassContex
 
 	// we render to our own output render target, not the intermediate one created by the compositing system
 	// Set the view family's render target/viewport.
-	SetRenderTarget(Context.RHICmdList, DestRenderTarget.TargetableTexture, FTextureRHIRef(), true);
-	Context.SetViewportAndCallRHI(0, 0, 0.0f, DestSize.X, DestSize.Y, 1.0f);
-
-	FGraphicsPipelineStateInitializer GraphicsPSOInit;
-	Context.RHICmdList.ApplyCachedRenderTargets(GraphicsPSOInit);
-	GraphicsPSOInit.BlendState = TStaticBlendState<>::GetRHI();
-	GraphicsPSOInit.RasterizerState = TStaticRasterizerState<>::GetRHI();
-	GraphicsPSOInit.DepthStencilState = TStaticDepthStencilState<false, CF_Always>::GetRHI();
-
-	TShaderMapRef<FPostProcessVS> VertexShader(Context.GetShaderMap());
-	TShaderMapRef<FPostProcessLogLuminance2ExposureScalePS> PixelShader(Context.GetShaderMap());
-
-	GraphicsPSOInit.BoundShaderState.VertexDeclarationRHI = GFilterVertexDeclaration.VertexDeclarationRHI;
-	GraphicsPSOInit.BoundShaderState.VertexShaderRHI = GETSAFERHISHADER_VERTEX(*VertexShader);
-	GraphicsPSOInit.BoundShaderState.PixelShaderRHI = GETSAFERHISHADER_PIXEL(*PixelShader);
-	GraphicsPSOInit.PrimitiveType = PT_TriangleList;
-
-	SetGraphicsPipelineState(Context.RHICmdList, GraphicsPSOInit);
-
-	// Set the parameters used by the pixel shader.
-
-	PixelShader->SetPS(Context, SrcRect, EyeAdaptationLastFrameRT);
-
-	// Draw a quad mapping scene color to the view's render target
-	DrawRectangle(
-		Context.RHICmdList,
-		0, 0,
-		DestSize.X, DestSize.Y,
-		0, 0,
-		DestSize.X, DestSize.Y,
-		DestSize,
-		DestSize,
-		*VertexShader,
-		EDRF_UseTriangleOptimization);
-
-	Context.RHICmdList.CopyToResolveTarget(DestRenderTarget.TargetableTexture, DestRenderTarget.ShaderResourceTexture, FResolveParams());
+
+	FRHIRenderPassInfo RPInfo(DestRenderTarget.TargetableTexture, ERenderTargetActions::Load_Store, DestRenderTarget.ShaderResourceTexture);
+	TransitionRenderPassTargets(Context.RHICmdList, RPInfo);
+	Context.RHICmdList.BeginRenderPass(RPInfo, TEXT("BasicEyeAdaptation"));
+	{
+		Context.SetViewportAndCallRHI(0, 0, 0.0f, DestSize.X, DestSize.Y, 1.0f);
+
+		FGraphicsPipelineStateInitializer GraphicsPSOInit;
+		Context.RHICmdList.ApplyCachedRenderTargets(GraphicsPSOInit);
+		GraphicsPSOInit.BlendState = TStaticBlendState<>::GetRHI();
+		GraphicsPSOInit.RasterizerState = TStaticRasterizerState<>::GetRHI();
+		GraphicsPSOInit.DepthStencilState = TStaticDepthStencilState<false, CF_Always>::GetRHI();
+
+		TShaderMapRef<FPostProcessVS> VertexShader(Context.GetShaderMap());
+		TShaderMapRef<FPostProcessLogLuminance2ExposureScalePS> PixelShader(Context.GetShaderMap());
+
+		GraphicsPSOInit.BoundShaderState.VertexDeclarationRHI = GFilterVertexDeclaration.VertexDeclarationRHI;
+		GraphicsPSOInit.BoundShaderState.VertexShaderRHI = GETSAFERHISHADER_VERTEX(*VertexShader);
+		GraphicsPSOInit.BoundShaderState.PixelShaderRHI = GETSAFERHISHADER_PIXEL(*PixelShader);
+		GraphicsPSOInit.PrimitiveType = PT_TriangleList;
+
+		SetGraphicsPipelineState(Context.RHICmdList, GraphicsPSOInit);
+
+		// Set the parameters used by the pixel shader.
+
+		PixelShader->SetPS(Context, SrcRect, EyeAdaptationLastFrameRT);
+
+		// Draw a quad mapping scene color to the view's render target
+		DrawRectangle(
+			Context.RHICmdList,
+			0, 0,
+			DestSize.X, DestSize.Y,
+			0, 0,
+			DestSize.X, DestSize.Y,
+			DestSize,
+			DestSize,
+			*VertexShader,
+			EDRF_UseTriangleOptimization);
+	}
+	Context.RHICmdList.EndRenderPass();
 
 	// Inform MultiGPU systems that we've finished with this texture for this frame
 	Context.RHICmdList.EndUpdateMultiFrameResource(DestRenderTarget.ShaderResourceTexture);
diff --git a/Engine/Source/Runtime/Renderer/Private/PostProcess/PostProcessFFTBloom.cpp b/Engine/Source/Runtime/Renderer/Private/PostProcess/PostProcessFFTBloom.cpp
--- a/Engine/Source/Runtime/Renderer/Private/PostProcess/PostProcessFFTBloom.cpp
+++ b/Engine/Source/Runtime/Renderer/Private/PostProcess/PostProcessFFTBloom.cpp
@@ -421,8 +421,8 @@ void ResizeAndCenterTexture(FRenderingCompositePassContext& Context,
 	// Get a pointer to the shader
 	FResizeAndCenterTextureCS* ComputeShader = ShaderMap.GetShader< FResizeAndCenterTextureCS >();
 
-	//
-	SetRenderTarget(RHICmdList, FTextureRHIRef(), FTextureRHIRef());
+	// #todo-renderpasses remove once everything is renderpasses
+	UnbindRenderTargets(Context.RHICmdList);
 	RHICmdList.SetComputeShader(ComputeShader->GetComputeShader());
 
 	// set destination
@@ -477,8 +477,8 @@ void CaptureKernelWeight(FRenderingCompositePassContext& Context,
 	// Get a pointer to the shader
 	FCaptureKernelWeightsCS* ComputeShader = ShaderMap.GetShader< FCaptureKernelWeightsCS >();
 
-
-	SetRenderTarget(RHICmdList, FTextureRHIRef(), FTextureRHIRef());
+	// #todo-renderpasses remove once everything is renderpasses
+	UnbindRenderTargets(Context.RHICmdList);
 	RHICmdList.SetComputeShader(ComputeShader->GetComputeShader());
 
 	// set destination
@@ -536,8 +536,8 @@ void BlendLowRes(FRenderingCompositePassContext& Context,
 	// Get a pointer to the shader
 	FBlendLowResCS* ComputeShader = ShaderMap.GetShader< FBlendLowResCS>();
 
-
-	SetRenderTarget(RHICmdList, FTextureRHIRef(), FTextureRHIRef());
+	// #todo-renderpasses remove once everything is renderpasses
+	UnbindRenderTargets(Context.RHICmdList);
 	RHICmdList.SetComputeShader(ComputeShader->GetComputeShader());
 
 	// set destination
@@ -584,7 +584,8 @@ void CopyImageRect(FRenderingCompositePassContext& Context, const FSceneRenderTa
 	// Get a pointer to the shader
 	FPassThroughCS* ComputeShader = ShaderMap.GetShader< FPassThroughCS >();
 
-	SetRenderTarget(Context.RHICmdList, FTextureRHIRef(), FTextureRHIRef());
+	// #todo-renderpasses remove once everything is renderpasses
+	UnbindRenderTargets(Context.RHICmdList);
 	Context.RHICmdList.SetComputeShader(ComputeShader->GetComputeShader());
 
 	// set destination
diff --git a/Engine/Source/Runtime/Renderer/Private/PostProcess/PostProcessGBufferHints.cpp b/Engine/Source/Runtime/Renderer/Private/PostProcess/PostProcessGBufferHints.cpp
--- a/Engine/Source/Runtime/Renderer/Private/PostProcess/PostProcessGBufferHints.cpp
+++ b/Engine/Source/Runtime/Renderer/Private/PostProcess/PostProcessGBufferHints.cpp
@@ -103,39 +103,42 @@ void FRCPassPostProcessGBufferHints::Process(FRenderingCompositePassContext& Con
 	const FSceneRenderTargetItem& DestRenderTarget = PassOutputs[0].RequestSurface(Context);
 
 	// Set the view family's render target/viewport.
-	SetRenderTarget(Context.RHICmdList, DestRenderTarget.TargetableTexture, FTextureRHIRef());
-	Context.SetViewportAndCallRHI(DestRect);
-
-	FGraphicsPipelineStateInitializer GraphicsPSOInit;
-	Context.RHICmdList.ApplyCachedRenderTargets(GraphicsPSOInit);
-	GraphicsPSOInit.BlendState = TStaticBlendState<>::GetRHI();
-	GraphicsPSOInit.RasterizerState = TStaticRasterizerState<>::GetRHI();
-	GraphicsPSOInit.DepthStencilState = TStaticDepthStencilState<false, CF_Always>::GetRHI();
-
-	TShaderMapRef<FPostProcessVS> VertexShader(Context.GetShaderMap());
-	TShaderMapRef<FPostProcessGBufferHintsPS> PixelShader(Context.GetShaderMap());
-
-	GraphicsPSOInit.BoundShaderState.VertexDeclarationRHI = GFilterVertexDeclaration.VertexDeclarationRHI;
-	GraphicsPSOInit.BoundShaderState.VertexShaderRHI = GETSAFERHISHADER_VERTEX(*VertexShader);
-	GraphicsPSOInit.BoundShaderState.PixelShaderRHI = GETSAFERHISHADER_PIXEL(*PixelShader);
-	GraphicsPSOInit.PrimitiveType = PT_TriangleList;
-
-	SetGraphicsPipelineState(Context.RHICmdList, GraphicsPSOInit);
-
-	PixelShader->SetPS(Context.RHICmdList, Context);
-
-	// Draw a quad mapping scene color to the view's render target
-	DrawRectangle(
-		Context.RHICmdList,
-		0, 0,
-		DestRect.Width(), DestRect.Height(),
-		SrcRect.Min.X, SrcRect.Min.Y,
-		SrcRect.Width(), SrcRect.Height(),
-		DestRect.Size(),
-		SrcSize,
-		*VertexShader,
-		EDRF_UseTriangleOptimization);
-
+	FRHIRenderPassInfo RPInfo(DestRenderTarget.TargetableTexture, ERenderTargetActions::Load_Store);
+	Context.RHICmdList.BeginRenderPass(RPInfo, TEXT("GBufferHints"));
+	{
+		Context.SetViewportAndCallRHI(DestRect);
+
+		FGraphicsPipelineStateInitializer GraphicsPSOInit;
+		Context.RHICmdList.ApplyCachedRenderTargets(GraphicsPSOInit);
+		GraphicsPSOInit.BlendState = TStaticBlendState<>::GetRHI();
+		GraphicsPSOInit.RasterizerState = TStaticRasterizerState<>::GetRHI();
+		GraphicsPSOInit.DepthStencilState = TStaticDepthStencilState<false, CF_Always>::GetRHI();
+
+		TShaderMapRef<FPostProcessVS> VertexShader(Context.GetShaderMap());
+		TShaderMapRef<FPostProcessGBufferHintsPS> PixelShader(Context.GetShaderMap());
+
+		GraphicsPSOInit.BoundShaderState.VertexDeclarationRHI = GFilterVertexDeclaration.VertexDeclarationRHI;
+		GraphicsPSOInit.BoundShaderState.VertexShaderRHI = GETSAFERHISHADER_VERTEX(*VertexShader);
+		GraphicsPSOInit.BoundShaderState.PixelShaderRHI = GETSAFERHISHADER_PIXEL(*PixelShader);
+		GraphicsPSOInit.PrimitiveType = PT_TriangleList;
+
+		SetGraphicsPipelineState(Context.RHICmdList, GraphicsPSOInit);
+
+		PixelShader->SetPS(Context.RHICmdList, Context);
+
+		// Draw a quad mapping scene color to the view's render target
+		DrawRectangle(
+			Context.RHICmdList,
+			0, 0,
+			DestRect.Width(), DestRect.Height(),
+			SrcRect.Min.X, SrcRect.Min.Y,
+			SrcRect.Width(), SrcRect.Height(),
+			DestRect.Size(),
+			SrcSize,
+			*VertexShader,
+			EDRF_UseTriangleOptimization);
+	}
+	Context.RHICmdList.EndRenderPass();
 
 	FRenderTargetTemp TempRenderTarget(View, (const FTexture2DRHIRef&)DestRenderTarget.TargetableTexture);
 	FCanvas Canvas(&TempRenderTarget, NULL, ViewFamily.CurrentRealTime, ViewFamily.CurrentWorldTime, ViewFamily.DeltaWorldTime, View.GetFeatureLevel());
@@ -148,19 +151,19 @@ void FRCPassPostProcessGBufferHints::Process(FRenderingCompositePassContext& Con
 	FString Line;
 
 	Line = FString::Printf(TEXT("GBufferHints"));
-	Canvas.DrawShadowedString( X, Y += YStep, *Line, GetStatsFont(), FLinearColor(1, 1, 1));
+	Canvas.DrawShadowedString(X, Y += YStep, *Line, GetStatsFont(), FLinearColor(1, 1, 1));
 
 	Y += YStep;
-	
+
 	Line = FString::Printf(TEXT("Yellow: Unrealistic material (In nature even black materials reflect a small amount of light)"));
-	Canvas.DrawShadowedString( X, Y += YStep, *Line, GetStatsFont(), FLinearColor(0.8f, 0.8f, 0));
+	Canvas.DrawShadowedString(X, Y += YStep, *Line, GetStatsFont(), FLinearColor(0.8f, 0.8f, 0));
 
 	Line = FString::Printf(TEXT("Red: Impossible material (This material emits more light than it receives)"));
-	Canvas.DrawShadowedString( X, Y += YStep, *Line, GetStatsFont(), FLinearColor(1, 0, 0));
+	Canvas.DrawShadowedString(X, Y += YStep, *Line, GetStatsFont(), FLinearColor(1, 0, 0));
 
 	Canvas.Flush_RenderThread(Context.RHICmdList);
 
-	Context.RHICmdList.CopyToResolveTarget(DestRenderTarget.TargetableTexture, DestRenderTarget.ShaderResourceTexture, FResolveParams());
+	Context.RHICmdList.CopyToResolveTarget(DestRenderTarget.TargetableTexture, DestRenderTarget.ShaderResourceTexture, FResolveParams());	
 
 	// AdjustGBufferRefCount(1) call is done in constructor
 	FSceneRenderTargets::Get(Context.RHICmdList).AdjustGBufferRefCount(Context.RHICmdList, -1);
diff --git a/Engine/Source/Runtime/Renderer/Private/PostProcess/PostProcessHMD.cpp b/Engine/Source/Runtime/Renderer/Private/PostProcess/PostProcessHMD.cpp
--- a/Engine/Source/Runtime/Renderer/Private/PostProcess/PostProcessHMD.cpp
+++ b/Engine/Source/Runtime/Renderer/Private/PostProcess/PostProcessHMD.cpp
@@ -162,58 +162,61 @@ void FRCPassPostProcessHMD::Process(FRenderingCompositePassContext& Context)
 		return;
 	}
 
-	const FViewInfo& View = Context.View;
-	const FSceneViewFamily& ViewFamily = *(View.Family);
-	
-	const FIntRect SrcRect = View.ViewRect;
-	const FIntRect DestRect = View.UnscaledViewRect;
-	const FIntPoint SrcSize = InputDesc->Extent;
-
 	const FSceneRenderTargetItem& DestRenderTarget = PassOutputs[0].RequestSurface(Context);
 
-	// Set the view family's render target/viewport
-	if (!ensure(DestRenderTarget.TargetableTexture.IsValid()) || DestRenderTarget.TargetableTexture->GetClearColor() == FLinearColor::Black)
+	bool bUseClearQuad = !(!ensure(DestRenderTarget.TargetableTexture.IsValid()) || DestRenderTarget.TargetableTexture->GetClearColor() == FLinearColor::Black);
+
+	ERenderTargetLoadAction LoadAction = ERenderTargetLoadAction::EClear;
+	if (bUseClearQuad)
 	{
-		FRHIRenderTargetView RtView = FRHIRenderTargetView(DestRenderTarget.TargetableTexture, ERenderTargetLoadAction::EClear);
-		FRHISetRenderTargetsInfo Info(1, &RtView, FRHIDepthRenderTargetView());
-		Context.RHICmdList.SetRenderTargetsAndClear(Info);
-		Context.SetViewportAndCallRHI(DestRect);
+		LoadAction = ERenderTargetLoadAction::ENoAction;
 	}
-	else
+	
+	FRHIRenderPassInfo RPInfo(DestRenderTarget.TargetableTexture, MakeRenderTargetActions(LoadAction, ERenderTargetStoreAction::EStore), DestRenderTarget.ShaderResourceTexture);
+	Context.RHICmdList.BeginRenderPass(RPInfo, TEXT("PostProcessHMD"));
 	{
-		SetRenderTarget(Context.RHICmdList, DestRenderTarget.TargetableTexture, FTextureRHIRef());
+		const FViewInfo& View = Context.View;
+		const FSceneViewFamily& ViewFamily = *(View.Family);
+
+		const FIntRect SrcRect = View.ViewRect;
+		const FIntRect DestRect = View.UnscaledViewRect;
+		const FIntPoint SrcSize = InputDesc->Extent;
+
 		Context.SetViewportAndCallRHI(DestRect);
-		DrawClearQuad(Context.RHICmdList, FLinearColor::Black);
-	}
 
-	FGraphicsPipelineStateInitializer GraphicsPSOInit;
-	Context.RHICmdList.ApplyCachedRenderTargets(GraphicsPSOInit);
-	GraphicsPSOInit.BlendState = TStaticBlendState<>::GetRHI();
-	GraphicsPSOInit.RasterizerState = TStaticRasterizerState<>::GetRHI();
-	GraphicsPSOInit.DepthStencilState = TStaticDepthStencilState<false, CF_Always>::GetRHI();
+		if (bUseClearQuad)
+		{
+			DrawClearQuad(Context.RHICmdList, FLinearColor::Black);
+		}
 
-	FMatrix QuadTexTransform = FMatrix::Identity;
-	FMatrix QuadPosTransform = FMatrix::Identity;
+		FGraphicsPipelineStateInitializer GraphicsPSOInit;
+		Context.RHICmdList.ApplyCachedRenderTargets(GraphicsPSOInit);
+		GraphicsPSOInit.BlendState = TStaticBlendState<>::GetRHI();
+		GraphicsPSOInit.RasterizerState = TStaticRasterizerState<>::GetRHI();
+		GraphicsPSOInit.DepthStencilState = TStaticDepthStencilState<false, CF_Always>::GetRHI();
 
-	check(GEngine->XRSystem.IsValid() && GEngine->XRSystem->GetHMDDevice());
+		FMatrix QuadTexTransform = FMatrix::Identity;
+		FMatrix QuadPosTransform = FMatrix::Identity;
 
-	{
-		TShaderMapRef<FPostProcessHMDVS> VertexShader(Context.GetShaderMap());
-		TShaderMapRef<FPostProcessHMDPS> PixelShader(Context.GetShaderMap());
+		check(GEngine->XRSystem.IsValid() && GEngine->XRSystem->GetHMDDevice());
 
-		GraphicsPSOInit.BoundShaderState.VertexDeclarationRHI = GDistortionVertexDeclaration.VertexDeclarationRHI;
-		GraphicsPSOInit.BoundShaderState.VertexShaderRHI = GETSAFERHISHADER_VERTEX(*VertexShader);
-		GraphicsPSOInit.BoundShaderState.PixelShaderRHI = GETSAFERHISHADER_PIXEL(*PixelShader);
-		GraphicsPSOInit.PrimitiveType = PT_TriangleList;
+		{
+			TShaderMapRef<FPostProcessHMDVS> VertexShader(Context.GetShaderMap());
+			TShaderMapRef<FPostProcessHMDPS> PixelShader(Context.GetShaderMap());
 
-		SetGraphicsPipelineState(Context.RHICmdList, GraphicsPSOInit);
+			GraphicsPSOInit.BoundShaderState.VertexDeclarationRHI = GDistortionVertexDeclaration.VertexDeclarationRHI;
+			GraphicsPSOInit.BoundShaderState.VertexShaderRHI = GETSAFERHISHADER_VERTEX(*VertexShader);
+			GraphicsPSOInit.BoundShaderState.PixelShaderRHI = GETSAFERHISHADER_PIXEL(*PixelShader);
+			GraphicsPSOInit.PrimitiveType = PT_TriangleList;
 
-		VertexShader->SetVS(Context);
-		PixelShader->SetPS(Context.RHICmdList, Context, SrcRect, SrcSize, View.StereoPass, QuadTexTransform);
-	}
-	GEngine->XRSystem->GetHMDDevice()->DrawDistortionMesh_RenderThread(Context, SrcSize);
+			SetGraphicsPipelineState(Context.RHICmdList, GraphicsPSOInit);
 
-	Context.RHICmdList.CopyToResolveTarget(DestRenderTarget.TargetableTexture, DestRenderTarget.ShaderResourceTexture, FResolveParams());
+			VertexShader->SetVS(Context);
+			PixelShader->SetPS(Context.RHICmdList, Context, SrcRect, SrcSize, View.StereoPass, QuadTexTransform);
+		}
+		GEngine->XRSystem->GetHMDDevice()->DrawDistortionMesh_RenderThread(Context, SrcSize);
+	}
+	Context.RHICmdList.EndRenderPass();
 }
 
 FPooledRenderTargetDesc FRCPassPostProcessHMD::ComputeOutputDesc(EPassOutputId InPassOutputId) const
diff --git a/Engine/Source/Runtime/Renderer/Private/PostProcess/PostProcessHierarchical.cpp b/Engine/Source/Runtime/Renderer/Private/PostProcess/PostProcessHierarchical.cpp
--- a/Engine/Source/Runtime/Renderer/Private/PostProcess/PostProcessHierarchical.cpp
+++ b/Engine/Source/Runtime/Renderer/Private/PostProcess/PostProcessHierarchical.cpp
@@ -158,53 +158,56 @@ void FRCPassPostProcessBuildHCB::Process(FRenderingCompositePassContext& Context
 		GraphicsPSOInit.BoundShaderState.PixelShaderRHI = GETSAFERHISHADER_PIXEL(*PixelShader);
 		GraphicsPSOInit.PrimitiveType = PT_TriangleList;
 
-		SetRenderTarget(RHICmdList, HCBRenderTarget.TargetableTexture, 0, NULL);
-		Context.RHICmdList.ApplyCachedRenderTargets(GraphicsPSOInit);
+		FRHIRenderPassInfo RPInfo(HCBRenderTarget.TargetableTexture, ERenderTargetActions::Load_Store, HCBRenderTarget.ShaderResourceTexture);
+		RPInfo.ResolveParameters = FResolveParams(FResolveRect(), CubeFace_PosX, 0);
+		RHICmdList.BeginRenderPass(RPInfo, TEXT("BuildHCB_0"));
+		{
+			Context.RHICmdList.ApplyCachedRenderTargets(GraphicsPSOInit);
+
+			SetGraphicsPipelineState(Context.RHICmdList, GraphicsPSOInit);
+
+			VertexShader->SetParameters(Context);
+			PixelShader->SetParameters(Context.RHICmdList, Context);
+
 
-		SetGraphicsPipelineState(Context.RHICmdList, GraphicsPSOInit);
-		
-		VertexShader->SetParameters(Context);
-		PixelShader->SetParameters(Context.RHICmdList, Context);
-			
-			
 #if 0
-		const FIntPoint DstDrawSize (
-			FMath::Min(FMath::CeilToInt(HCBUsedPercentage.X * float(HCBSize.X)) + 1, HCBSize.X),
-			FMath::Min(FMath::CeilToInt(HCBUsedPercentage.Y * float(HCBSize.Y)) + 1, HCBSize.Y)
+			const FIntPoint DstDrawSize(
+				FMath::Min(FMath::CeilToInt(HCBUsedPercentage.X * float(HCBSize.X)) + 1, HCBSize.X),
+				FMath::Min(FMath::CeilToInt(HCBUsedPercentage.Y * float(HCBSize.Y)) + 1, HCBSize.Y)
 			);
 
-		RHICmdList.SetViewport(0, 0, 0.0f, DstDrawSize.X, DstDrawSize.Y, 1.0f);
-		
-		DrawPostProcessPass(
-			RHICmdList,
-			0, 0,
-			HCBSize.X, HCBSize.Y,
-			View.ViewRect.Min.X, View.ViewRect.Min.Y,
-			View.ViewRect.Width(), View.ViewRect.Height(),
-			DstDrawSize,
-			SceneContext.GetBufferSizeXY(),
-			*VertexShader,
-			View.StereoPass,
-			Context.HasHmdMesh(),
-			EDRF_UseTriangleOptimization);
+			RHICmdList.SetViewport(0, 0, 0.0f, DstDrawSize.X, DstDrawSize.Y, 1.0f);
+
+			DrawPostProcessPass(
+				RHICmdList,
+				0, 0,
+				HCBSize.X, HCBSize.Y,
+				View.ViewRect.Min.X, View.ViewRect.Min.Y,
+				View.ViewRect.Width(), View.ViewRect.Height(),
+				DstDrawSize,
+				SceneContext.GetBufferSizeXY(),
+				*VertexShader,
+				View.StereoPass,
+				Context.HasHmdMesh(),
+				EDRF_UseTriangleOptimization);
 #else
-		RHICmdList.SetViewport(0, 0, 0.0f, HCBSize.X, HCBSize.Y, 1.0f);
-			
-		DrawPostProcessPass(
-			RHICmdList,
-			0, 0,
-			HCBSize.X, HCBSize.Y,
-			View.ViewRect.Min.X, View.ViewRect.Min.Y,
-			View.ViewRect.Width(), View.ViewRect.Height(),
-			HCBSize,
-			SceneContext.GetBufferSizeXY(),
-			*VertexShader,
-			View.StereoPass,
-			Context.HasHmdMesh(),
-			EDRF_UseTriangleOptimization);
+			RHICmdList.SetViewport(0, 0, 0.0f, HCBSize.X, HCBSize.Y, 1.0f);
+
+			DrawPostProcessPass(
+				RHICmdList,
+				0, 0,
+				HCBSize.X, HCBSize.Y,
+				View.ViewRect.Min.X, View.ViewRect.Min.Y,
+				View.ViewRect.Width(), View.ViewRect.Height(),
+				HCBSize,
+				SceneContext.GetBufferSizeXY(),
+				*VertexShader,
+				View.StereoPass,
+				Context.HasHmdMesh(),
+				EDRF_UseTriangleOptimization);
 #endif
-
-		RHICmdList.CopyToResolveTarget(HCBRenderTarget.TargetableTexture, HCBRenderTarget.ShaderResourceTexture, FResolveParams(FResolveRect(), CubeFace_PosX, 0));
+		}
+		Context.RHICmdList.EndRenderPass();
 	}
 		
 	FIntPoint SrcSize = HCBSize;
@@ -213,62 +216,69 @@ void FRCPassPostProcessBuildHCB::Process(FRenderingCompositePassContext& Context
 	// Downsampling...
 	for( uint8 MipIndex = 1; MipIndex < HCBMipCount; MipIndex++ )
 	{
-		DstSize.X = FMath::Max(DstSize.X, 1);
-		DstSize.Y = FMath::Max(DstSize.Y, 1);
+		TCHAR RenderPassLabel[13];
+		FCString::Snprintf(RenderPassLabel, 13, TEXT("BuildHCB_%u"), MipIndex);
 
-		SetRenderTarget(RHICmdList, HCBRenderTarget.TargetableTexture, MipIndex, NULL);
-		Context.RHICmdList.ApplyCachedRenderTargets(GraphicsPSOInit);
+		FRHIRenderPassInfo RPInfo(HCBRenderTarget.TargetableTexture, ERenderTargetActions::Load_Store, HCBRenderTarget.ShaderResourceTexture);
+		RPInfo.ColorRenderTargets[0].MipIndex = MipIndex;
+		RPInfo.ResolveParameters = FResolveParams(FResolveRect(), CubeFace_PosX, MipIndex);
+		RHICmdList.BeginRenderPass(RPInfo, RenderPassLabel);
+		{
+			DstSize.X = FMath::Max(DstSize.X, 1);
+			DstSize.Y = FMath::Max(DstSize.Y, 1);
 
-		TShaderMapRef< FPostProcessVS >	VertexShader(Context.GetShaderMap());
-		TShaderMapRef< TPostProcessBuildHCBPS<1> >	PixelShader(Context.GetShaderMap());
+			Context.RHICmdList.ApplyCachedRenderTargets(GraphicsPSOInit);
 
-		GraphicsPSOInit.BoundShaderState.VertexDeclarationRHI = GFilterVertexDeclaration.VertexDeclarationRHI;
-		GraphicsPSOInit.BoundShaderState.VertexShaderRHI = GETSAFERHISHADER_VERTEX(*VertexShader);
-		GraphicsPSOInit.BoundShaderState.PixelShaderRHI = GETSAFERHISHADER_PIXEL(*PixelShader);
-		GraphicsPSOInit.PrimitiveType = PT_TriangleList;
+			TShaderMapRef< FPostProcessVS >	VertexShader(Context.GetShaderMap());
+			TShaderMapRef< TPostProcessBuildHCBPS<1> >	PixelShader(Context.GetShaderMap());
 
-		SetGraphicsPipelineState(Context.RHICmdList, GraphicsPSOInit);
+			GraphicsPSOInit.BoundShaderState.VertexDeclarationRHI = GFilterVertexDeclaration.VertexDeclarationRHI;
+			GraphicsPSOInit.BoundShaderState.VertexShaderRHI = GETSAFERHISHADER_VERTEX(*VertexShader);
+			GraphicsPSOInit.BoundShaderState.PixelShaderRHI = GETSAFERHISHADER_PIXEL(*PixelShader);
+			GraphicsPSOInit.PrimitiveType = PT_TriangleList;
+
+			SetGraphicsPipelineState(Context.RHICmdList, GraphicsPSOInit);
+
+			PixelShader->SetParameters(RHICmdList, View, SrcSize, HCBRenderTarget.MipSRVs[MipIndex - 1]);
 
-		PixelShader->SetParameters(RHICmdList, View, SrcSize, HCBRenderTarget.MipSRVs[ MipIndex - 1 ] );
-			
 #if 0
-		const FIntPoint DstDrawSize (
-			FMath::Min(FMath::CeilToInt(HCBUsedPercentage.X * float(DstSize.X)) + 1, DstSize.X),
-			FMath::Min(FMath::CeilToInt(HCBUsedPercentage.Y * float(DstSize.Y)) + 1, DstSize.Y)
+			const FIntPoint DstDrawSize(
+				FMath::Min(FMath::CeilToInt(HCBUsedPercentage.X * float(DstSize.X)) + 1, DstSize.X),
+				FMath::Min(FMath::CeilToInt(HCBUsedPercentage.Y * float(DstSize.Y)) + 1, DstSize.Y)
 			);
 
-		RHICmdList.SetViewport(0, 0, 0.0f, DstDrawSize.X, DstDrawSize.Y, 1.0f);
-
-		DrawPostProcessPass(
-			RHICmdList,
-			0, 0,
-			DstSize.X, DstSize.Y,
-			0, 0,
-			SrcSize.X, SrcSize.Y,
-			DstDrawSize,
-			SrcSize,
-			*VertexShader,
-			View.StereoPass,
-			Context.HasHmdMesh(),
-			EDRF_UseTriangleOptimization);
+			RHICmdList.SetViewport(0, 0, 0.0f, DstDrawSize.X, DstDrawSize.Y, 1.0f);
+
+			DrawPostProcessPass(
+				RHICmdList,
+				0, 0,
+				DstSize.X, DstSize.Y,
+				0, 0,
+				SrcSize.X, SrcSize.Y,
+				DstDrawSize,
+				SrcSize,
+				*VertexShader,
+				View.StereoPass,
+				Context.HasHmdMesh(),
+				EDRF_UseTriangleOptimization);
 #else
-		RHICmdList.SetViewport(0, 0, 0.0f, DstSize.X, DstSize.Y, 1.0f);
-
-		DrawPostProcessPass(
-			RHICmdList,
-			0, 0,
-			DstSize.X, DstSize.Y,
-			0, 0,
-			SrcSize.X, SrcSize.Y,
-			DstSize,
-			SrcSize,
-			*VertexShader,
-			View.StereoPass,
-			Context.HasHmdMesh(),
-			EDRF_UseTriangleOptimization);
+			RHICmdList.SetViewport(0, 0, 0.0f, DstSize.X, DstSize.Y, 1.0f);
+
+			DrawPostProcessPass(
+				RHICmdList,
+				0, 0,
+				DstSize.X, DstSize.Y,
+				0, 0,
+				SrcSize.X, SrcSize.Y,
+				DstSize,
+				SrcSize,
+				*VertexShader,
+				View.StereoPass,
+				Context.HasHmdMesh(),
+				EDRF_UseTriangleOptimization);
 #endif
-
-		RHICmdList.CopyToResolveTarget(HCBRenderTarget.TargetableTexture, HCBRenderTarget.ShaderResourceTexture, FResolveParams(FResolveRect(), CubeFace_PosX, MipIndex));
+		}
+		Context.RHICmdList.EndRenderPass();
 
 		SrcSize /= 2;
 		DstSize /= 2;
diff --git a/Engine/Source/Runtime/Renderer/Private/PostProcess/PostProcessHistogram.cpp b/Engine/Source/Runtime/Renderer/Private/PostProcess/PostProcessHistogram.cpp
--- a/Engine/Source/Runtime/Renderer/Private/PostProcess/PostProcessHistogram.cpp
+++ b/Engine/Source/Runtime/Renderer/Private/PostProcess/PostProcessHistogram.cpp
@@ -107,7 +107,8 @@ void FRCPassPostProcessHistogram::Process(FRenderingCompositePassContext& Contex
 
 	TShaderMapRef<FPostProcessHistogramCS> ComputeShader(Context.GetShaderMap());
 
-	SetRenderTarget(Context.RHICmdList, FTextureRHIRef(), FTextureRHIRef());
+	// #todo-renderpasses remove once everything is renderpasses
+	UnbindRenderTargets(Context.RHICmdList);
     Context.RHICmdList.SetComputeShader(ComputeShader->GetComputeShader());
     
 
diff --git a/Engine/Source/Runtime/Renderer/Private/PostProcess/PostProcessHistogramReduce.cpp b/Engine/Source/Runtime/Renderer/Private/PostProcess/PostProcessHistogramReduce.cpp
--- a/Engine/Source/Runtime/Renderer/Private/PostProcess/PostProcessHistogramReduce.cpp
+++ b/Engine/Source/Runtime/Renderer/Private/PostProcess/PostProcessHistogramReduce.cpp
@@ -107,43 +107,44 @@ void FRCPassPostProcessHistogramReduce::Process(FRenderingCompositePassContext&
 	const FSceneRenderTargetItem& DestRenderTarget = PassOutputs[0].RequestSurface(Context);
 
 	// Set the view family's render target/viewport.
-	SetRenderTarget(Context.RHICmdList, DestRenderTarget.TargetableTexture, FTextureRHIRef());
-
-	FGraphicsPipelineStateInitializer GraphicsPSOInit;
-	Context.RHICmdList.ApplyCachedRenderTargets(GraphicsPSOInit);
-	GraphicsPSOInit.BlendState = TStaticBlendState<>::GetRHI();
-	GraphicsPSOInit.RasterizerState = TStaticRasterizerState<>::GetRHI();
-	GraphicsPSOInit.DepthStencilState = TStaticDepthStencilState<false, CF_Always>::GetRHI();
-
-	TShaderMapRef<FPostProcessVS> VertexShader(Context.GetShaderMap());
-	TShaderMapRef<FPostProcessHistogramReducePS> PixelShader(Context.GetShaderMap());
-
-	GraphicsPSOInit.BoundShaderState.VertexDeclarationRHI = GFilterVertexDeclaration.VertexDeclarationRHI;
-	GraphicsPSOInit.BoundShaderState.VertexShaderRHI = GETSAFERHISHADER_VERTEX(*VertexShader);
-	GraphicsPSOInit.BoundShaderState.PixelShaderRHI = GETSAFERHISHADER_PIXEL(*PixelShader);
-	GraphicsPSOInit.PrimitiveType = PT_TriangleList;
-
-	SetGraphicsPipelineState(Context.RHICmdList, GraphicsPSOInit);
-
-	// we currently assume the input is half res, one full res pixel less to avoid getting bilinear filtered input
-	FIntPoint GatherExtent = (Context.SceneColorViewRect.Size() - FIntPoint(1, 1)) / 2;
-
-	uint32 LoopSizeValue = ComputeLoopSize(GatherExtent);
-
-	PixelShader->SetPS(Context.RHICmdList, Context, LoopSizeValue);
-
-	DrawRectangle(
-		Context.RHICmdList,
-		0, 0,
-		DestSize.X, DestSize.Y,
-		0, 0,
-		SrcSize.X, 0,
-		DestSize,
-		SrcSize,
-		*VertexShader,
-		EDRF_UseTriangleOptimization);
-
-	Context.RHICmdList.CopyToResolveTarget(DestRenderTarget.TargetableTexture, DestRenderTarget.ShaderResourceTexture, FResolveParams());
+	FRHIRenderPassInfo RPInfo(DestRenderTarget.TargetableTexture, ERenderTargetActions::Load_Store, DestRenderTarget.ShaderResourceTexture);
+	Context.RHICmdList.BeginRenderPass(RPInfo, TEXT("HistogramReduce"));
+	{
+		FGraphicsPipelineStateInitializer GraphicsPSOInit;
+		Context.RHICmdList.ApplyCachedRenderTargets(GraphicsPSOInit);
+		GraphicsPSOInit.BlendState = TStaticBlendState<>::GetRHI();
+		GraphicsPSOInit.RasterizerState = TStaticRasterizerState<>::GetRHI();
+		GraphicsPSOInit.DepthStencilState = TStaticDepthStencilState<false, CF_Always>::GetRHI();
+
+		TShaderMapRef<FPostProcessVS> VertexShader(Context.GetShaderMap());
+		TShaderMapRef<FPostProcessHistogramReducePS> PixelShader(Context.GetShaderMap());
+
+		GraphicsPSOInit.BoundShaderState.VertexDeclarationRHI = GFilterVertexDeclaration.VertexDeclarationRHI;
+		GraphicsPSOInit.BoundShaderState.VertexShaderRHI = GETSAFERHISHADER_VERTEX(*VertexShader);
+		GraphicsPSOInit.BoundShaderState.PixelShaderRHI = GETSAFERHISHADER_PIXEL(*PixelShader);
+		GraphicsPSOInit.PrimitiveType = PT_TriangleList;
+
+		SetGraphicsPipelineState(Context.RHICmdList, GraphicsPSOInit);
+
+		// we currently assume the input is half res, one full res pixel less to avoid getting bilinear filtered input
+		FIntPoint GatherExtent = (Context.SceneColorViewRect.Size() - FIntPoint(1, 1)) / 2;
+
+		uint32 LoopSizeValue = ComputeLoopSize(GatherExtent);
+
+		PixelShader->SetPS(Context.RHICmdList, Context, LoopSizeValue);
+
+		DrawRectangle(
+			Context.RHICmdList,
+			0, 0,
+			DestSize.X, DestSize.Y,
+			0, 0,
+			SrcSize.X, 0,
+			DestSize,
+			SrcSize,
+			*VertexShader,
+			EDRF_UseTriangleOptimization);
+	}
+	Context.RHICmdList.EndRenderPass();
 }
 
 uint32 FRCPassPostProcessHistogramReduce::ComputeLoopSize(FIntPoint PixelExtent)
diff --git a/Engine/Source/Runtime/Renderer/Private/PostProcess/PostProcessLensBlur.cpp b/Engine/Source/Runtime/Renderer/Private/PostProcess/PostProcessLensBlur.cpp
--- a/Engine/Source/Runtime/Renderer/Private/PostProcess/PostProcessLensBlur.cpp
+++ b/Engine/Source/Runtime/Renderer/Private/PostProcess/PostProcessLensBlur.cpp
@@ -188,45 +188,46 @@ void FRCPassPostProcessLensBlur::Process(FRenderingCompositePassContext& Context
 	const FSceneRenderTargetItem& DestRenderTarget = PassOutputs[0].RequestSurface(Context);
 
 	// Set the view family's render target/viewport.
-	SetRenderTarget(Context.RHICmdList, DestRenderTarget.TargetableTexture, FTextureRHIRef(), ESimpleRenderTargetMode::EClearColorExistingDepth);
-
-	Context.SetViewportAndCallRHI(ViewRect);
-
-	FGraphicsPipelineStateInitializer GraphicsPSOInit;
-	Context.RHICmdList.ApplyCachedRenderTargets(GraphicsPSOInit);
+	FRHIRenderPassInfo RPInfo(DestRenderTarget.TargetableTexture, ERenderTargetActions::Clear_Store, DestRenderTarget.ShaderResourceTexture);
+	Context.RHICmdList.BeginRenderPass(RPInfo, TEXT("LensBlur"));
+	{
+		Context.SetViewportAndCallRHI(ViewRect);
 
-	// set the state (additive blending)
-	GraphicsPSOInit.BlendState = TStaticBlendState<CW_RGB, BO_Add, BF_One, BF_One>::GetRHI();
-	GraphicsPSOInit.RasterizerState = TStaticRasterizerState<>::GetRHI();
-	GraphicsPSOInit.DepthStencilState = TStaticDepthStencilState<false, CF_Always>::GetRHI();
+		FGraphicsPipelineStateInitializer GraphicsPSOInit;
+		Context.RHICmdList.ApplyCachedRenderTargets(GraphicsPSOInit);
 
-	TShaderMapRef<FPostProcessLensBlurVS> VertexShader(Context.GetShaderMap());
-	TShaderMapRef<FPostProcessLensBlurPS> PixelShader(Context.GetShaderMap());
+		// set the state (additive blending)
+		GraphicsPSOInit.BlendState = TStaticBlendState<CW_RGB, BO_Add, BF_One, BF_One>::GetRHI();
+		GraphicsPSOInit.RasterizerState = TStaticRasterizerState<>::GetRHI();
+		GraphicsPSOInit.DepthStencilState = TStaticDepthStencilState<false, CF_Always>::GetRHI();
 
-	GraphicsPSOInit.BoundShaderState.VertexDeclarationRHI = GEmptyVertexDeclaration.VertexDeclarationRHI;
-	GraphicsPSOInit.BoundShaderState.VertexShaderRHI = GETSAFERHISHADER_VERTEX(*VertexShader);
-	GraphicsPSOInit.BoundShaderState.PixelShaderRHI = GETSAFERHISHADER_PIXEL(*PixelShader);
-	GraphicsPSOInit.PrimitiveType = PT_TriangleList;
+		TShaderMapRef<FPostProcessLensBlurVS> VertexShader(Context.GetShaderMap());
+		TShaderMapRef<FPostProcessLensBlurPS> PixelShader(Context.GetShaderMap());
 
-	SetGraphicsPipelineState(Context.RHICmdList, GraphicsPSOInit);
+		GraphicsPSOInit.BoundShaderState.VertexDeclarationRHI = GEmptyVertexDeclaration.VertexDeclarationRHI;
+		GraphicsPSOInit.BoundShaderState.VertexShaderRHI = GETSAFERHISHADER_VERTEX(*VertexShader);
+		GraphicsPSOInit.BoundShaderState.PixelShaderRHI = GETSAFERHISHADER_PIXEL(*PixelShader);
+		GraphicsPSOInit.PrimitiveType = PT_TriangleList;
 
-	uint32 TileSize = 1;
+		SetGraphicsPipelineState(Context.RHICmdList, GraphicsPSOInit);
 
-	FIntPoint TileCount = ViewSize / TileSize;
+		uint32 TileSize = 1;
 
-	float PixelKernelSize = PercentKernelSize / 100.0f * ViewSize.X;
+		FIntPoint TileCount = ViewSize / TileSize;
 
-	VertexShader->SetParameters(Context, TileCount, TileSize, PixelKernelSize, Threshold);
-	PixelShader->SetParameters(Context.RHICmdList, Context, PixelKernelSize);
+		float PixelKernelSize = PercentKernelSize / 100.0f * ViewSize.X;
 
-	Context.RHICmdList.SetStreamSource(0, NULL, 0);
+		VertexShader->SetParameters(Context, TileCount, TileSize, PixelKernelSize, Threshold);
+		PixelShader->SetParameters(Context.RHICmdList, Context, PixelKernelSize);
 
-	// needs to be the same on shader side (faster on NVIDIA and AMD)
-	int32 QuadsPerInstance = 4;
+		Context.RHICmdList.SetStreamSource(0, NULL, 0);
 
-	Context.RHICmdList.DrawPrimitive(0, 2, FMath::DivideAndRoundUp(TileCount.X * TileCount.Y, QuadsPerInstance));
+		// needs to be the same on shader side (faster on NVIDIA and AMD)
+		int32 QuadsPerInstance = 4;
 
-	Context.RHICmdList.CopyToResolveTarget(DestRenderTarget.TargetableTexture, DestRenderTarget.ShaderResourceTexture, FResolveParams());
+		Context.RHICmdList.DrawPrimitive(0, 2, FMath::DivideAndRoundUp(TileCount.X * TileCount.Y, QuadsPerInstance));
+	}
+	Context.RHICmdList.EndRenderPass();
 }
 
 
diff --git a/Engine/Source/Runtime/Renderer/Private/PostProcess/PostProcessLensFlares.cpp b/Engine/Source/Runtime/Renderer/Private/PostProcess/PostProcessLensFlares.cpp
--- a/Engine/Source/Runtime/Renderer/Private/PostProcess/PostProcessLensFlares.cpp
+++ b/Engine/Source/Runtime/Renderer/Private/PostProcess/PostProcessLensFlares.cpp
@@ -147,152 +147,151 @@ void FRCPassPostProcessLensFlares::Process(FRenderingCompositePassContext& Conte
 		return;
 	}
 
-	const FViewInfo& View = Context.View;
-	const FSceneViewFamily& ViewFamily = *(View.Family);
-
-	FIntPoint TexSize1 = InputDesc1->Extent;
-	FIntPoint TexSize2 = InputDesc2->Extent;
-
-	uint32 ScaleToFullRes1 = Context.ReferenceBufferSize.X / TexSize1.X;
-	uint32 ScaleToFullRes2 = Context.ReferenceBufferSize.X / TexSize2.X;
-
-	FIntRect ViewRect1 = FIntRect::DivideAndRoundUp(Context.SceneColorViewRect, ScaleToFullRes1);
-	FIntRect ViewRect2 = FIntRect::DivideAndRoundUp(Context.SceneColorViewRect, ScaleToFullRes2);
-
-	FIntPoint ViewSize1 = ViewRect1.Size();
-	FIntPoint ViewSize2 = ViewRect2.Size();
-
 	const FSceneRenderTargetItem& DestRenderTarget = PassOutputs[0].RequestSurface(Context);
 
-	FRHIRenderTargetView RtView = FRHIRenderTargetView(DestRenderTarget.TargetableTexture, ERenderTargetLoadAction::ENoAction);
-	FRHISetRenderTargetsInfo Info(1, &RtView, FRHIDepthRenderTargetView());
-	Context.RHICmdList.SetRenderTargetsAndClear(Info);
+	FRHIRenderPassInfo RPInfo(DestRenderTarget.TargetableTexture, ERenderTargetActions::DontLoad_Store, DestRenderTarget.ShaderResourceTexture);
+	Context.RHICmdList.BeginRenderPass(RPInfo, TEXT("PassPostProcessLensFlares"));
+	{
+		const FViewInfo& View = Context.View;
+		const FSceneViewFamily& ViewFamily = *(View.Family);
 
-	Context.SetViewportAndCallRHI(ViewRect1);
+		FIntPoint TexSize1 = InputDesc1->Extent;
+		FIntPoint TexSize2 = InputDesc2->Extent;
 
-	FGraphicsPipelineStateInitializer GraphicsPSOInit;
-	Context.RHICmdList.ApplyCachedRenderTargets(GraphicsPSOInit);
-	GraphicsPSOInit.BlendState = TStaticBlendState<>::GetRHI();
-	GraphicsPSOInit.RasterizerState = TStaticRasterizerState<>::GetRHI();
-	GraphicsPSOInit.DepthStencilState = TStaticDepthStencilState<false, CF_Always>::GetRHI();
+		uint32 ScaleToFullRes1 = Context.ReferenceBufferSize.X / TexSize1.X;
+		uint32 ScaleToFullRes2 = Context.ReferenceBufferSize.X / TexSize2.X;
 
-	TShaderMapRef<FPostProcessVS> VertexShader(Context.GetShaderMap());
+		FIntRect ViewRect1 = FIntRect::DivideAndRoundUp(Context.SceneColorViewRect, ScaleToFullRes1);
+		FIntRect ViewRect2 = FIntRect::DivideAndRoundUp(Context.SceneColorViewRect, ScaleToFullRes2);
 
-	
-	// setup background (bloom), can be implemented to use additive blending to avoid the read here
-	if (bCompositeBloom)
-	{
-		TShaderMapRef<TPostProcessLensFlareBasePS<false>> PixelShader(Context.GetShaderMap());
-
-		GraphicsPSOInit.BoundShaderState.VertexDeclarationRHI = GFilterVertexDeclaration.VertexDeclarationRHI;
-		GraphicsPSOInit.BoundShaderState.VertexShaderRHI = GETSAFERHISHADER_VERTEX(*VertexShader);
-		GraphicsPSOInit.BoundShaderState.PixelShaderRHI = GETSAFERHISHADER_PIXEL(*PixelShader);
-		GraphicsPSOInit.PrimitiveType = PT_TriangleList;
-
-		SetGraphicsPipelineState(Context.RHICmdList, GraphicsPSOInit);
-
-		VertexShader->SetParameters(Context);
-		PixelShader->SetParameters(Context.RHICmdList, Context);
-
-		// Draw a quad mapping scene color to the view's render target
-		DrawRectangle(
-			Context.RHICmdList,
-			0, 0,
-			ViewSize1.X, ViewSize1.Y,
-			ViewRect1.Min.X, ViewRect1.Min.Y,
-			ViewSize1.X, ViewSize1.Y,
-			ViewSize1,
-			TexSize1,
-			*VertexShader,
-			EDRF_UseTriangleOptimization);
-	}
-	else
-	{
-		TShaderMapRef<TPostProcessLensFlareBasePS<true>> PixelShader(Context.GetShaderMap());
-
-		GraphicsPSOInit.BoundShaderState.VertexDeclarationRHI = GFilterVertexDeclaration.VertexDeclarationRHI;
-		GraphicsPSOInit.BoundShaderState.VertexShaderRHI = GETSAFERHISHADER_VERTEX(*VertexShader);
-		GraphicsPSOInit.BoundShaderState.PixelShaderRHI = GETSAFERHISHADER_PIXEL(*PixelShader);
-		GraphicsPSOInit.PrimitiveType = PT_TriangleList;
-
-		SetGraphicsPipelineState(Context.RHICmdList, GraphicsPSOInit);
-
-		VertexShader->SetParameters(Context);
-		PixelShader->SetParameters(Context.RHICmdList, Context);
-
-		// Draw a quad mapping scene color to the view's render target
-		DrawRectangle(
-			Context.RHICmdList,
-			0, 0,
-			ViewSize1.X, ViewSize1.Y,
-			ViewRect1.Min.X, ViewRect1.Min.Y,
-			ViewSize1.X, ViewSize1.Y,
-			ViewSize1,
-			TexSize1,
-			*VertexShader,
-			EDRF_UseTriangleOptimization);
+		FIntPoint ViewSize1 = ViewRect1.Size();
+		FIntPoint ViewSize2 = ViewRect2.Size();
 
-	}
+		Context.SetViewportAndCallRHI(ViewRect1);
 
-	// additive blend
-	GraphicsPSOInit.BlendState = TStaticBlendState<CW_RGB, BO_Add, BF_One, BF_One>::GetRHI();
+		FGraphicsPipelineStateInitializer GraphicsPSOInit;
+		Context.RHICmdList.ApplyCachedRenderTargets(GraphicsPSOInit);
+		GraphicsPSOInit.BlendState = TStaticBlendState<>::GetRHI();
+		GraphicsPSOInit.RasterizerState = TStaticRasterizerState<>::GetRHI();
+		GraphicsPSOInit.DepthStencilState = TStaticDepthStencilState<false, CF_Always>::GetRHI();
 
-	// add lens flares on top of that
-	{
-		TShaderMapRef<FPostProcessLensFlaresPS> PixelShader(Context.GetShaderMap());
+		TShaderMapRef<FPostProcessVS> VertexShader(Context.GetShaderMap());
 
-		GraphicsPSOInit.BoundShaderState.VertexDeclarationRHI = GFilterVertexDeclaration.VertexDeclarationRHI;
-		GraphicsPSOInit.BoundShaderState.VertexShaderRHI = GETSAFERHISHADER_VERTEX(*VertexShader);
-		GraphicsPSOInit.BoundShaderState.PixelShaderRHI = GETSAFERHISHADER_PIXEL(*PixelShader);
-		GraphicsPSOInit.PrimitiveType = PT_TriangleList;
 
-		SetGraphicsPipelineState(Context.RHICmdList, GraphicsPSOInit);
+		// setup background (bloom), can be implemented to use additive blending to avoid the read here
+		if (bCompositeBloom)
+		{
+			TShaderMapRef<TPostProcessLensFlareBasePS<false>> PixelShader(Context.GetShaderMap());
 
-		FVector2D TexScaleValue = FVector2D(TexSize2) / ViewSize2;
+			GraphicsPSOInit.BoundShaderState.VertexDeclarationRHI = GFilterVertexDeclaration.VertexDeclarationRHI;
+			GraphicsPSOInit.BoundShaderState.VertexShaderRHI = GETSAFERHISHADER_VERTEX(*VertexShader);
+			GraphicsPSOInit.BoundShaderState.PixelShaderRHI = GETSAFERHISHADER_PIXEL(*PixelShader);
+			GraphicsPSOInit.PrimitiveType = PT_TriangleList;
 
-		VertexShader->SetParameters(Context);
-		PixelShader->SetParameters(Context.RHICmdList, Context, TexScaleValue);
+			SetGraphicsPipelineState(Context.RHICmdList, GraphicsPSOInit);
 
-		// todo: expose
-		const uint32 Count = 8;
+			VertexShader->SetParameters(Context);
+			PixelShader->SetParameters(Context.RHICmdList, Context);
 
-		// we assume the center of the view is the center of the lens (would not be correct for tiled rendering)
-		FVector2D Center = FVector2D(ViewSize1) * 0.5f;
+			// Draw a quad mapping scene color to the view's render target
+			DrawRectangle(
+				Context.RHICmdList,
+				0, 0,
+				ViewSize1.X, ViewSize1.Y,
+				ViewRect1.Min.X, ViewRect1.Min.Y,
+				ViewSize1.X, ViewSize1.Y,
+				ViewSize1,
+				TexSize1,
+				*VertexShader,
+				EDRF_UseTriangleOptimization);
+		}
+		else
+		{
+			TShaderMapRef<TPostProcessLensFlareBasePS<true>> PixelShader(Context.GetShaderMap());
 
-		FLinearColor LensFlareHDRColor = Context.View.FinalPostProcessSettings.LensFlareTint * Context.View.FinalPostProcessSettings.LensFlareIntensity;
-	
-		// to get the same brightness with 4x more quads (TileSize=1 in LensBlur)
-		LensFlareHDRColor.R *= 0.25f;
-		LensFlareHDRColor.G *= 0.25f;
-		LensFlareHDRColor.B *= 0.25f;
+			GraphicsPSOInit.BoundShaderState.VertexDeclarationRHI = GFilterVertexDeclaration.VertexDeclarationRHI;
+			GraphicsPSOInit.BoundShaderState.VertexShaderRHI = GETSAFERHISHADER_VERTEX(*VertexShader);
+			GraphicsPSOInit.BoundShaderState.PixelShaderRHI = GETSAFERHISHADER_PIXEL(*PixelShader);
+			GraphicsPSOInit.PrimitiveType = PT_TriangleList;
 
-		for(uint32 i = 0; i < Count; ++i)
-		{
-			FLinearColor FlareColor = Context.View.FinalPostProcessSettings.LensFlareTints[i % 8];
-			float NormalizedAlpha = FlareColor.A;
-			float Alpha = NormalizedAlpha * 7.0f - 3.5f; 
+			SetGraphicsPipelineState(Context.RHICmdList, GraphicsPSOInit);
 
-			// scale to blur outside of the view (only if we use LensBlur)
-			Alpha *= SizeScale;
-			
-			// set the individual flare color
-			SetShaderValue(Context.RHICmdList, PixelShader->GetPixelShader(), PixelShader->FlareColor, FlareColor * LensFlareHDRColor);
+			VertexShader->SetParameters(Context);
+			PixelShader->SetParameters(Context.RHICmdList, Context);
 
 			// Draw a quad mapping scene color to the view's render target
 			DrawRectangle(
 				Context.RHICmdList,
-				Center.X - 0.5f * ViewSize1.X * Alpha, Center.Y - 0.5f * ViewSize1.Y * Alpha,
-				ViewSize1.X * Alpha, ViewSize1.Y * Alpha,
-				ViewRect2.Min.X, ViewRect2.Min.Y,
-				ViewSize2.X, ViewSize2.Y,
+				0, 0,
+				ViewSize1.X, ViewSize1.Y,
+				ViewRect1.Min.X, ViewRect1.Min.Y,
+				ViewSize1.X, ViewSize1.Y,
 				ViewSize1,
-				TexSize2,
+				TexSize1,
 				*VertexShader,
-				EDRF_Default);
+				EDRF_UseTriangleOptimization);
+
 		}
-	}
 
-	Context.RHICmdList.CopyToResolveTarget(DestRenderTarget.TargetableTexture, DestRenderTarget.ShaderResourceTexture, FResolveParams());
+		// additive blend
+		GraphicsPSOInit.BlendState = TStaticBlendState<CW_RGB, BO_Add, BF_One, BF_One>::GetRHI();
+
+		// add lens flares on top of that
+		{
+			TShaderMapRef<FPostProcessLensFlaresPS> PixelShader(Context.GetShaderMap());
+
+			GraphicsPSOInit.BoundShaderState.VertexDeclarationRHI = GFilterVertexDeclaration.VertexDeclarationRHI;
+			GraphicsPSOInit.BoundShaderState.VertexShaderRHI = GETSAFERHISHADER_VERTEX(*VertexShader);
+			GraphicsPSOInit.BoundShaderState.PixelShaderRHI = GETSAFERHISHADER_PIXEL(*PixelShader);
+			GraphicsPSOInit.PrimitiveType = PT_TriangleList;
+
+			SetGraphicsPipelineState(Context.RHICmdList, GraphicsPSOInit);
+
+			FVector2D TexScaleValue = FVector2D(TexSize2) / ViewSize2;
+
+			VertexShader->SetParameters(Context);
+			PixelShader->SetParameters(Context.RHICmdList, Context, TexScaleValue);
+
+			// todo: expose
+			const uint32 Count = 8;
+
+			// we assume the center of the view is the center of the lens (would not be correct for tiled rendering)
+			FVector2D Center = FVector2D(ViewSize1) * 0.5f;
+
+			FLinearColor LensFlareHDRColor = Context.View.FinalPostProcessSettings.LensFlareTint * Context.View.FinalPostProcessSettings.LensFlareIntensity;
+
+			// to get the same brightness with 4x more quads (TileSize=1 in LensBlur)
+			LensFlareHDRColor.R *= 0.25f;
+			LensFlareHDRColor.G *= 0.25f;
+			LensFlareHDRColor.B *= 0.25f;
+
+			for (uint32 i = 0; i < Count; ++i)
+			{
+				FLinearColor FlareColor = Context.View.FinalPostProcessSettings.LensFlareTints[i % 8];
+				float NormalizedAlpha = FlareColor.A;
+				float Alpha = NormalizedAlpha * 7.0f - 3.5f;
+
+				// scale to blur outside of the view (only if we use LensBlur)
+				Alpha *= SizeScale;
+
+				// set the individual flare color
+				SetShaderValue(Context.RHICmdList, PixelShader->GetPixelShader(), PixelShader->FlareColor, FlareColor * LensFlareHDRColor);
+
+				// Draw a quad mapping scene color to the view's render target
+				DrawRectangle(
+					Context.RHICmdList,
+					Center.X - 0.5f * ViewSize1.X * Alpha, Center.Y - 0.5f * ViewSize1.Y * Alpha,
+					ViewSize1.X * Alpha, ViewSize1.Y * Alpha,
+					ViewRect2.Min.X, ViewRect2.Min.Y,
+					ViewSize2.X, ViewSize2.Y,
+					ViewSize1,
+					TexSize2,
+					*VertexShader,
+					EDRF_Default);
+			}
+		}
+	}
+	Context.RHICmdList.EndRenderPass();
 }
 
 FPooledRenderTargetDesc FRCPassPostProcessLensFlares::ComputeOutputDesc(EPassOutputId InPassOutputId) const
diff --git a/Engine/Source/Runtime/Renderer/Private/PostProcess/PostProcessMaterial.cpp b/Engine/Source/Runtime/Renderer/Private/PostProcess/PostProcessMaterial.cpp
--- a/Engine/Source/Runtime/Renderer/Private/PostProcess/PostProcessMaterial.cpp
+++ b/Engine/Source/Runtime/Renderer/Private/PostProcess/PostProcessMaterial.cpp
@@ -208,9 +208,9 @@ void FRCPassPostProcessMaterial::Process(FRenderingCompositePassContext& Context
 	check(Proxy);
 
 	ERHIFeatureLevel::Type FeatureLevel = Context.View.GetFeatureLevel();
-	
+
 	const FMaterial* Material = Proxy->GetMaterial(FeatureLevel);
-	
+
 	check(Material);
 
 	const FViewInfo& View = Context.View;
@@ -229,84 +229,85 @@ void FRCPassPostProcessMaterial::Process(FRenderingCompositePassContext& Context
 	SCOPED_DRAW_EVENTF(Context.RHICmdList, PostProcessMaterial, TEXT("PostProcessMaterial %dx%d Material=%s"), DestRect.Width(), DestRect.Height(), *Material->GetFriendlyName());
 
 	ERenderTargetLoadAction LoadAction = Context.GetLoadActionForRenderTarget(DestRenderTarget);
-	FRHIRenderTargetView RtView = FRHIRenderTargetView(DestRenderTarget.TargetableTexture, LoadAction);
-	FRHISetRenderTargetsInfo Info(1, &RtView, FRHIDepthRenderTargetView());
-	Context.RHICmdList.SetRenderTargetsAndClear(Info);
-	Context.SetViewportAndCallRHI(DestRect);
-
-	FGraphicsPipelineStateInitializer GraphicsPSOInit;
-	Context.RHICmdList.ApplyCachedRenderTargets(GraphicsPSOInit);
-	GraphicsPSOInit.BlendState = TStaticBlendState<>::GetRHI();
-	GraphicsPSOInit.RasterizerState = TStaticRasterizerState<>::GetRHI();
-	GraphicsPSOInit.DepthStencilState = TStaticDepthStencilState<false, CF_Always>::GetRHI();
-	GraphicsPSOInit.PrimitiveType = PT_TriangleList;
-
-	const FMaterialShaderMap* MaterialShaderMap = Material->GetRenderingThreadShaderMap();
-	FShader* VertexShader = nullptr;
-
-	// uses mobile's post process material.
-	if (FeatureLevel <= ERHIFeatureLevel::ES3_1)
+	
+	FRHIRenderPassInfo RPInfo(DestRenderTarget.TargetableTexture, MakeRenderTargetActions(LoadAction, ERenderTargetStoreAction::EStore), DestRenderTarget.ShaderResourceTexture);
+	Context.RHICmdList.BeginRenderPass(RPInfo, TEXT("PostProcessMaterial"));
 	{
-		FPostProcessMaterialPS_Mobile* PixelShader_Mobile = MaterialShaderMap->GetShader<FPostProcessMaterialPS_Mobile>();
-		FPostProcessMaterialVS_Mobile* VertexShader_Mobile = MaterialShaderMap->GetShader<FPostProcessMaterialVS_Mobile>();
+		Context.SetViewportAndCallRHI(DestRect);
 
-		GraphicsPSOInit.BoundShaderState.VertexDeclarationRHI = GFilterVertexDeclaration.VertexDeclarationRHI;
-		GraphicsPSOInit.BoundShaderState.VertexShaderRHI = GETSAFERHISHADER_VERTEX(VertexShader_Mobile);
-		GraphicsPSOInit.BoundShaderState.PixelShaderRHI = GETSAFERHISHADER_PIXEL(PixelShader_Mobile);
+		FGraphicsPipelineStateInitializer GraphicsPSOInit;
+		Context.RHICmdList.ApplyCachedRenderTargets(GraphicsPSOInit);
+		GraphicsPSOInit.BlendState = TStaticBlendState<>::GetRHI();
+		GraphicsPSOInit.RasterizerState = TStaticRasterizerState<>::GetRHI();
+		GraphicsPSOInit.DepthStencilState = TStaticDepthStencilState<false, CF_Always>::GetRHI();
+		GraphicsPSOInit.PrimitiveType = PT_TriangleList;
 
-		SetGraphicsPipelineState(Context.RHICmdList, GraphicsPSOInit);
+		const FMaterialShaderMap* MaterialShaderMap = Material->GetRenderingThreadShaderMap();
+		FShader* VertexShader = nullptr;
 
-		VertexShader_Mobile->SetParameters(Context.RHICmdList, Context, Proxy);
-		PixelShader_Mobile->SetParameters(Context.RHICmdList, Context, Proxy);
-		VertexShader = VertexShader_Mobile;
-	}
-	// Uses highend post process material that assumed ViewSize == BufferSize.
-	else if (View.ViewRect == Context.SceneColorViewRect && View.ViewRect.Size() == SrcSize && View.ViewRect.Min == FIntPoint::ZeroValue)
-	{
-		FFPostProcessMaterialPS_HighEnd0* PixelShader_HighEnd = MaterialShaderMap->GetShader<FFPostProcessMaterialPS_HighEnd0>();
-		FPostProcessMaterialVS_HighEnd* VertexShader_HighEnd = MaterialShaderMap->GetShader<FPostProcessMaterialVS_HighEnd>();
+		// uses mobile's post process material.
+		if (FeatureLevel <= ERHIFeatureLevel::ES3_1)
+		{
+			FPostProcessMaterialPS_Mobile* PixelShader_Mobile = MaterialShaderMap->GetShader<FPostProcessMaterialPS_Mobile>();
+			FPostProcessMaterialVS_Mobile* VertexShader_Mobile = MaterialShaderMap->GetShader<FPostProcessMaterialVS_Mobile>();
 
-		GraphicsPSOInit.BoundShaderState.VertexDeclarationRHI = GPostProcessMaterialVertexDeclaration.VertexDeclarationRHI;
-		GraphicsPSOInit.BoundShaderState.VertexShaderRHI = GETSAFERHISHADER_VERTEX(VertexShader_HighEnd);
-		GraphicsPSOInit.BoundShaderState.PixelShaderRHI = GETSAFERHISHADER_PIXEL(PixelShader_HighEnd);
+			GraphicsPSOInit.BoundShaderState.VertexDeclarationRHI = GFilterVertexDeclaration.VertexDeclarationRHI;
+			GraphicsPSOInit.BoundShaderState.VertexShaderRHI = GETSAFERHISHADER_VERTEX(VertexShader_Mobile);
+			GraphicsPSOInit.BoundShaderState.PixelShaderRHI = GETSAFERHISHADER_PIXEL(PixelShader_Mobile);
 
-		SetGraphicsPipelineState(Context.RHICmdList, GraphicsPSOInit);
+			SetGraphicsPipelineState(Context.RHICmdList, GraphicsPSOInit);
 
-		VertexShader_HighEnd->SetParameters(Context.RHICmdList, Context, Proxy);
-		PixelShader_HighEnd->SetParameters(Context.RHICmdList, Context, Proxy);
-		VertexShader = VertexShader_HighEnd;
-	}
-	// Uses highend post process material that handle ViewSize != BufferSize.
-	else
-	{
-		FFPostProcessMaterialPS_HighEnd1* PixelShader_HighEnd = MaterialShaderMap->GetShader<FFPostProcessMaterialPS_HighEnd1>();
-		FPostProcessMaterialVS_HighEnd* VertexShader_HighEnd = MaterialShaderMap->GetShader<FPostProcessMaterialVS_HighEnd>();
+			VertexShader_Mobile->SetParameters(Context.RHICmdList, Context, Proxy);
+			PixelShader_Mobile->SetParameters(Context.RHICmdList, Context, Proxy);
+			VertexShader = VertexShader_Mobile;
+		}
+		// Uses highend post process material that assumed ViewSize == BufferSize.
+		else if (View.ViewRect == Context.SceneColorViewRect && View.ViewRect.Size() == SrcSize && View.ViewRect.Min == FIntPoint::ZeroValue)
+		{
+			FFPostProcessMaterialPS_HighEnd0* PixelShader_HighEnd = MaterialShaderMap->GetShader<FFPostProcessMaterialPS_HighEnd0>();
+			FPostProcessMaterialVS_HighEnd* VertexShader_HighEnd = MaterialShaderMap->GetShader<FPostProcessMaterialVS_HighEnd>();
 
-		GraphicsPSOInit.BoundShaderState.VertexDeclarationRHI = GPostProcessMaterialVertexDeclaration.VertexDeclarationRHI;
-		GraphicsPSOInit.BoundShaderState.VertexShaderRHI = GETSAFERHISHADER_VERTEX(VertexShader_HighEnd);
-		GraphicsPSOInit.BoundShaderState.PixelShaderRHI = GETSAFERHISHADER_PIXEL(PixelShader_HighEnd);
+			GraphicsPSOInit.BoundShaderState.VertexDeclarationRHI = GPostProcessMaterialVertexDeclaration.VertexDeclarationRHI;
+			GraphicsPSOInit.BoundShaderState.VertexShaderRHI = GETSAFERHISHADER_VERTEX(VertexShader_HighEnd);
+			GraphicsPSOInit.BoundShaderState.PixelShaderRHI = GETSAFERHISHADER_PIXEL(PixelShader_HighEnd);
 
-		SetGraphicsPipelineState(Context.RHICmdList, GraphicsPSOInit);
+			SetGraphicsPipelineState(Context.RHICmdList, GraphicsPSOInit);
 
-		VertexShader_HighEnd->SetParameters(Context.RHICmdList, Context, Proxy);
-		PixelShader_HighEnd->SetParameters(Context.RHICmdList, Context, Proxy);
-		VertexShader = VertexShader_HighEnd;
-	}
+			VertexShader_HighEnd->SetParameters(Context.RHICmdList, Context, Proxy);
+			PixelShader_HighEnd->SetParameters(Context.RHICmdList, Context, Proxy);
+			VertexShader = VertexShader_HighEnd;
+		}
+		// Uses highend post process material that handle ViewSize != BufferSize.
+		else
+		{
+			FFPostProcessMaterialPS_HighEnd1* PixelShader_HighEnd = MaterialShaderMap->GetShader<FFPostProcessMaterialPS_HighEnd1>();
+			FPostProcessMaterialVS_HighEnd* VertexShader_HighEnd = MaterialShaderMap->GetShader<FPostProcessMaterialVS_HighEnd>();
 
-	DrawPostProcessPass(
-		Context.RHICmdList,
-		0, 0,
-		DestRect.Width(), DestRect.Height(),
-		SrcRect.Min.X, SrcRect.Min.Y,
-		SrcRect.Width(), SrcRect.Height(),
-		DestRect.Size(),
-		SrcSize,
-		VertexShader,
-		View.StereoPass,
-		Context.HasHmdMesh(),
-		EDRF_UseTriangleOptimization);
-
-	Context.RHICmdList.CopyToResolveTarget(DestRenderTarget.TargetableTexture, DestRenderTarget.ShaderResourceTexture, FResolveParams());
+			GraphicsPSOInit.BoundShaderState.VertexDeclarationRHI = GPostProcessMaterialVertexDeclaration.VertexDeclarationRHI;
+			GraphicsPSOInit.BoundShaderState.VertexShaderRHI = GETSAFERHISHADER_VERTEX(VertexShader_HighEnd);
+			GraphicsPSOInit.BoundShaderState.PixelShaderRHI = GETSAFERHISHADER_PIXEL(PixelShader_HighEnd);
+
+			SetGraphicsPipelineState(Context.RHICmdList, GraphicsPSOInit);
+
+			VertexShader_HighEnd->SetParameters(Context.RHICmdList, Context, Proxy);
+			PixelShader_HighEnd->SetParameters(Context.RHICmdList, Context, Proxy);
+			VertexShader = VertexShader_HighEnd;
+		}
+
+		DrawPostProcessPass(
+			Context.RHICmdList,
+			0, 0,
+			DestRect.Width(), DestRect.Height(),
+			SrcRect.Min.X, SrcRect.Min.Y,
+			SrcRect.Width(), SrcRect.Height(),
+			DestRect.Size(),
+			SrcSize,
+			VertexShader,
+			View.StereoPass,
+			Context.HasHmdMesh(),
+			EDRF_UseTriangleOptimization);
+	}
+	Context.RHICmdList.EndRenderPass();
 
 	if(Material->NeedsGBuffer())
 	{
