Render Passes

diff --git a/Engine/Source/Runtime/Renderer/Private/PostProcess/PostProcessMobile.cpp b/Engine/Source/Runtime/Renderer/Private/PostProcess/PostProcessMobile.cpp
--- a/Engine/Source/Runtime/Renderer/Private/PostProcess/PostProcessMobile.cpp
+++ b/Engine/Source/Runtime/Renderer/Private/PostProcess/PostProcessMobile.cpp
@@ -228,37 +228,43 @@ void FRCPassPostProcessBloomSetupES2::Process(FRenderingCompositePassContext& Co
 
 	const FSceneRenderTargetItem& DestRenderTarget = PassOutputs[0].RequestSurface(Context);
 
+	ERenderTargetLoadAction LoadAction = ERenderTargetLoadAction::EClear;
+
+	bool bUseClearQuad = (DestRenderTarget.TargetableTexture->GetClearColor() != FLinearColor::Black);
+
 	// OverrideRenderTarget might patch out final render target and we have no control of the clear color anymore
-	if (DestRenderTarget.TargetableTexture->GetClearColor() == FLinearColor::Black)
+	if (bUseClearQuad)
 	{
-		FRHIRenderTargetView View = FRHIRenderTargetView(DestRenderTarget.TargetableTexture, ERenderTargetLoadAction::EClear);
-		FRHISetRenderTargetsInfo Info(1, &View, FRHIDepthRenderTargetView());
-		Context.RHICmdList.SetRenderTargetsAndClear(Info);
+		LoadAction = ERenderTargetLoadAction::ENoAction;
 	}
-	else
+
+	FRHIRenderPassInfo RPInfo(DestRenderTarget.TargetableTexture, MakeRenderTargetActions(LoadAction, ERenderTargetStoreAction::EStore), DestRenderTarget.ShaderResourceTexture);
+	
+	Context.RHICmdList.BeginRenderPass(RPInfo, TEXT("PostProcessBloomSetupES2"));
 	{
-		SetRenderTarget(Context.RHICmdList, DestRenderTarget.TargetableTexture, FTextureRHIRef());
-		DrawClearQuad(Context.RHICmdList, FLinearColor::Black);
-	}
+		if (bUseClearQuad)
+		{
+			DrawClearQuad(Context.RHICmdList, FLinearColor::Black);
+		}
 
-	Context.SetViewportAndCallRHI(0, 0, 0.0f, DstX, DstY, 1.0f );
+		Context.SetViewportAndCallRHI(0, 0, 0.0f, DstX, DstY, 1.0f);
 
-	SetShader(Context);
+		SetShader(Context);
 
-	TShaderMapRef<FPostProcessBloomSetupVS_ES2> VertexShader(Context.GetShaderMap());
+		TShaderMapRef<FPostProcessBloomSetupVS_ES2> VertexShader(Context.GetShaderMap());
 
-	DrawRectangle(
-		Context.RHICmdList,
-		0, 0,
-		DstX, DstY,
-		SrcRect.Min.X, SrcRect.Min.Y,
-		SrcRect.Width(), SrcRect.Height(),
-		DstSize,
-		SrcSize,
-		*VertexShader,
-		EDRF_UseTriangleOptimization);
-
-	Context.RHICmdList.CopyToResolveTarget(DestRenderTarget.TargetableTexture, DestRenderTarget.ShaderResourceTexture, FResolveParams());
+		DrawRectangle(
+			Context.RHICmdList,
+			0, 0,
+			DstX, DstY,
+			SrcRect.Min.X, SrcRect.Min.Y,
+			SrcRect.Width(), SrcRect.Height(),
+			DstSize,
+			SrcSize,
+			*VertexShader,
+			EDRF_UseTriangleOptimization);
+	}
+	Context.RHICmdList.EndRenderPass();
 }
 
 FPooledRenderTargetDesc FRCPassPostProcessBloomSetupES2::ComputeOutputDesc(EPassOutputId InPassOutputId) const
@@ -435,37 +441,43 @@ void FRCPassPostProcessBloomSetupSmallES2::Process(FRenderingCompositePassContex
 	}
 
 	const FSceneRenderTargetItem& DestRenderTarget = PassOutputs[0].RequestSurface(Context);
+
+	const bool bUseClearQuad = DestRenderTarget.TargetableTexture->GetClearColor() != FLinearColor::Black;
 	
 	// OverrideRenderTarget might patch out final render target and we have no control of the clear color anymore
-	if (DestRenderTarget.TargetableTexture->GetClearColor() == FLinearColor::Black)
+	ERenderTargetLoadAction LoadAction = ERenderTargetLoadAction::EClear;
+
+	if (bUseClearQuad)
 	{
-		FRHIRenderTargetView View = FRHIRenderTargetView(DestRenderTarget.TargetableTexture, ERenderTargetLoadAction::EClear);
-		FRHISetRenderTargetsInfo Info(1, &View, FRHIDepthRenderTargetView());
-		Context.RHICmdList.SetRenderTargetsAndClear(Info);
+		LoadAction = ERenderTargetLoadAction::ENoAction;
 	}
-	else
+
+	FRHIRenderPassInfo RPInfo(DestRenderTarget.TargetableTexture, MakeRenderTargetActions(LoadAction, ERenderTargetStoreAction::EStore), DestRenderTarget.ShaderResourceTexture);
+	
+	Context.RHICmdList.BeginRenderPass(RPInfo, TEXT("PostProcessBloomSetupSmallES2"));
 	{
-		SetRenderTarget(Context.RHICmdList, DestRenderTarget.TargetableTexture, FTextureRHIRef());
-		DrawClearQuad(Context.RHICmdList, FLinearColor::Black);
-	}
+		if (bUseClearQuad)
+		{
+			DrawClearQuad(Context.RHICmdList, FLinearColor::Black);
+		}
 
-	Context.SetViewportAndCallRHI(0, 0, 0.0f, DstX, DstY, 1.0f );
+		Context.SetViewportAndCallRHI(0, 0, 0.0f, DstX, DstY, 1.0f);
 
-	SetShader(Context);
+		SetShader(Context);
 
-	TShaderMapRef<FPostProcessBloomSetupSmallVS_ES2> VertexShader(Context.GetShaderMap());
-	DrawRectangle(
-		Context.RHICmdList,
-		0, 0,
-		DstX, DstY,
-		SrcRect.Min.X, SrcRect.Min.Y,
-		SrcRect.Width(), SrcRect.Height(),
-		DstSize,
-		SrcSize,
-		*VertexShader,
-		EDRF_UseTriangleOptimization);
-
-	Context.RHICmdList.CopyToResolveTarget(DestRenderTarget.TargetableTexture, DestRenderTarget.ShaderResourceTexture, FResolveParams());
+		TShaderMapRef<FPostProcessBloomSetupSmallVS_ES2> VertexShader(Context.GetShaderMap());
+		DrawRectangle(
+			Context.RHICmdList,
+			0, 0,
+			DstX, DstY,
+			SrcRect.Min.X, SrcRect.Min.Y,
+			SrcRect.Width(), SrcRect.Height(),
+			DstSize,
+			SrcSize,
+			*VertexShader,
+			EDRF_UseTriangleOptimization);
+	}
+	Context.RHICmdList.EndRenderPass();
 }
 
 FPooledRenderTargetDesc FRCPassPostProcessBloomSetupSmallES2::ComputeOutputDesc(EPassOutputId InPassOutputId) const
@@ -592,53 +604,58 @@ void FRCPassPostProcessBloomDownES2::Process(FRenderingCompositePassContext& Con
 	const FSceneRenderTargetItem& DestRenderTarget = PassOutputs[0].RequestSurface(Context);
 	
 	// OverrideRenderTarget might patch out final render target and we have no control of the clear color anymore
-	if (DestRenderTarget.TargetableTexture->GetClearColor() == FLinearColor::Black)
+	const bool bUseClearQuad = DestRenderTarget.TargetableTexture->GetClearColor() != FLinearColor::Black;
+
+	ERenderTargetLoadAction LoadAction = ERenderTargetLoadAction::EClear;
+	if (bUseClearQuad)
 	{
-		FRHIRenderTargetView View = FRHIRenderTargetView(DestRenderTarget.TargetableTexture, ERenderTargetLoadAction::EClear);
-		FRHISetRenderTargetsInfo Info(1, &View, FRHIDepthRenderTargetView());
-		Context.RHICmdList.SetRenderTargetsAndClear(Info);
+		LoadAction = ERenderTargetLoadAction::ENoAction;
 	}
-	else
+
+	FRHIRenderPassInfo RPInfo(DestRenderTarget.TargetableTexture, MakeRenderTargetActions(LoadAction, ERenderTargetStoreAction::EStore), DestRenderTarget.ShaderResourceTexture);
+	
+	Context.RHICmdList.BeginRenderPass(RPInfo, TEXT("PostProcessBloomDownES2"));
 	{
-		SetRenderTarget(Context.RHICmdList, DestRenderTarget.TargetableTexture, FTextureRHIRef());
-		DrawClearQuad(Context.RHICmdList, FLinearColor::Black);
-	}
+		if (bUseClearQuad)
+		{
+			DrawClearQuad(Context.RHICmdList, FLinearColor::Black);
+		}
 
-	Context.SetViewportAndCallRHI(0, 0, 0.0f, DstX, DstY, 1.0f );
+		Context.SetViewportAndCallRHI(0, 0, 0.0f, DstX, DstY, 1.0f);
 
-	FGraphicsPipelineStateInitializer GraphicsPSOInit;
-	Context.RHICmdList.ApplyCachedRenderTargets(GraphicsPSOInit);
-	GraphicsPSOInit.BlendState = TStaticBlendState<>::GetRHI();
-	GraphicsPSOInit.RasterizerState = TStaticRasterizerState<>::GetRHI();
-	GraphicsPSOInit.DepthStencilState = TStaticDepthStencilState<false, CF_Always>::GetRHI();
+		FGraphicsPipelineStateInitializer GraphicsPSOInit;
+		Context.RHICmdList.ApplyCachedRenderTargets(GraphicsPSOInit);
+		GraphicsPSOInit.BlendState = TStaticBlendState<>::GetRHI();
+		GraphicsPSOInit.RasterizerState = TStaticRasterizerState<>::GetRHI();
+		GraphicsPSOInit.DepthStencilState = TStaticDepthStencilState<false, CF_Always>::GetRHI();
 
-	TShaderMapRef<FPostProcessBloomDownVS_ES2> VertexShader(Context.GetShaderMap());
-	TShaderMapRef<FPostProcessBloomDownPS_ES2> PixelShader(Context.GetShaderMap());
+		TShaderMapRef<FPostProcessBloomDownVS_ES2> VertexShader(Context.GetShaderMap());
+		TShaderMapRef<FPostProcessBloomDownPS_ES2> PixelShader(Context.GetShaderMap());
 
-	GraphicsPSOInit.BoundShaderState.VertexDeclarationRHI = GFilterVertexDeclaration.VertexDeclarationRHI;
-	GraphicsPSOInit.BoundShaderState.VertexShaderRHI = GETSAFERHISHADER_VERTEX(*VertexShader);
-	GraphicsPSOInit.BoundShaderState.PixelShaderRHI = GETSAFERHISHADER_PIXEL(*PixelShader);
-	GraphicsPSOInit.PrimitiveType = PT_TriangleList;
+		GraphicsPSOInit.BoundShaderState.VertexDeclarationRHI = GFilterVertexDeclaration.VertexDeclarationRHI;
+		GraphicsPSOInit.BoundShaderState.VertexShaderRHI = GETSAFERHISHADER_VERTEX(*VertexShader);
+		GraphicsPSOInit.BoundShaderState.PixelShaderRHI = GETSAFERHISHADER_PIXEL(*PixelShader);
+		GraphicsPSOInit.PrimitiveType = PT_TriangleList;
 
-	SetGraphicsPipelineState(Context.RHICmdList, GraphicsPSOInit);
-
-	VertexShader->SetVS(Context, Scale);
-	PixelShader->SetPS(Context.RHICmdList, Context);
+		SetGraphicsPipelineState(Context.RHICmdList, GraphicsPSOInit);
 
-	FIntPoint SrcDstSize = PrePostSourceViewportSize/2;
+		VertexShader->SetVS(Context, Scale);
+		PixelShader->SetPS(Context.RHICmdList, Context);
 
-	DrawRectangle(
-		Context.RHICmdList,
-		0, 0,
-		DstX, DstY,
-		0, 0,
-		DstX, DstY,
-		SrcDstSize,
-		SrcDstSize,
-		*VertexShader,
-		EDRF_UseTriangleOptimization);
+		FIntPoint SrcDstSize = PrePostSourceViewportSize / 2;
 
-	Context.RHICmdList.CopyToResolveTarget(DestRenderTarget.TargetableTexture, DestRenderTarget.ShaderResourceTexture, FResolveParams());
+		DrawRectangle(
+			Context.RHICmdList,
+			0, 0,
+			DstX, DstY,
+			0, 0,
+			DstX, DstY,
+			SrcDstSize,
+			SrcDstSize,
+			*VertexShader,
+			EDRF_UseTriangleOptimization);
+	}
+	Context.RHICmdList.EndRenderPass();
 }
 
 FPooledRenderTargetDesc FRCPassPostProcessBloomDownES2::ComputeOutputDesc(EPassOutputId InPassOutputId) const
@@ -770,56 +787,61 @@ void FRCPassPostProcessBloomUpES2::Process(FRenderingCompositePassContext& Conte
 	const FSceneRenderTargetItem& DestRenderTarget = PassOutputs[0].RequestSurface(Context);
 	
 	// OverrideRenderTarget might patch out final render target and we have no control of the clear color anymore
-	if (DestRenderTarget.TargetableTexture->GetClearColor() == FLinearColor::Black)
+	const bool bUseClearQuad = DestRenderTarget.TargetableTexture->GetClearColor() != FLinearColor::Black;
+	ERenderTargetLoadAction LoadAction = ERenderTargetLoadAction::EClear;
+
+	if (bUseClearQuad)
 	{
-		FRHIRenderTargetView View = FRHIRenderTargetView(DestRenderTarget.TargetableTexture, ERenderTargetLoadAction::EClear);
-		FRHISetRenderTargetsInfo Info(1, &View, FRHIDepthRenderTargetView());
-		Context.RHICmdList.SetRenderTargetsAndClear(Info);
+		LoadAction = ERenderTargetLoadAction::ENoAction;
 	}
-	else
+
+	FRHIRenderPassInfo RPInfo(DestRenderTarget.TargetableTexture, MakeRenderTargetActions(LoadAction, ERenderTargetStoreAction::EStore), DestRenderTarget.ShaderResourceTexture);
+		
+	Context.RHICmdList.BeginRenderPass(RPInfo, TEXT("PostProcessBloomUpES2"));
 	{
-		SetRenderTarget(Context.RHICmdList, DestRenderTarget.TargetableTexture, FTextureRHIRef());
-		DrawClearQuad(Context.RHICmdList, FLinearColor::Black);
-	}
+		if (bUseClearQuad)
+		{
+			DrawClearQuad(Context.RHICmdList, FLinearColor::Black);
+		}
 
-	Context.SetViewportAndCallRHI(0, 0, 0.0f, DstX, DstY, 1.0f );
+		Context.SetViewportAndCallRHI(0, 0, 0.0f, DstX, DstY, 1.0f);
 
-	FGraphicsPipelineStateInitializer GraphicsPSOInit;
-	Context.RHICmdList.ApplyCachedRenderTargets(GraphicsPSOInit);
-	GraphicsPSOInit.BlendState = TStaticBlendState<>::GetRHI();
-	GraphicsPSOInit.RasterizerState = TStaticRasterizerState<>::GetRHI();
-	GraphicsPSOInit.DepthStencilState = TStaticDepthStencilState<false, CF_Always>::GetRHI();
+		FGraphicsPipelineStateInitializer GraphicsPSOInit;
+		Context.RHICmdList.ApplyCachedRenderTargets(GraphicsPSOInit);
+		GraphicsPSOInit.BlendState = TStaticBlendState<>::GetRHI();
+		GraphicsPSOInit.RasterizerState = TStaticRasterizerState<>::GetRHI();
+		GraphicsPSOInit.DepthStencilState = TStaticDepthStencilState<false, CF_Always>::GetRHI();
 
-	TShaderMapRef<FPostProcessBloomUpVS_ES2> VertexShader(Context.GetShaderMap());
-	TShaderMapRef<FPostProcessBloomUpPS_ES2> PixelShader(Context.GetShaderMap());
+		TShaderMapRef<FPostProcessBloomUpVS_ES2> VertexShader(Context.GetShaderMap());
+		TShaderMapRef<FPostProcessBloomUpPS_ES2> PixelShader(Context.GetShaderMap());
 
-	GraphicsPSOInit.BoundShaderState.VertexDeclarationRHI = GFilterVertexDeclaration.VertexDeclarationRHI;
-	GraphicsPSOInit.BoundShaderState.VertexShaderRHI = GETSAFERHISHADER_VERTEX(*VertexShader);
-	GraphicsPSOInit.BoundShaderState.PixelShaderRHI = GETSAFERHISHADER_PIXEL(*PixelShader);
-	GraphicsPSOInit.PrimitiveType = PT_TriangleList;
+		GraphicsPSOInit.BoundShaderState.VertexDeclarationRHI = GFilterVertexDeclaration.VertexDeclarationRHI;
+		GraphicsPSOInit.BoundShaderState.VertexShaderRHI = GETSAFERHISHADER_VERTEX(*VertexShader);
+		GraphicsPSOInit.BoundShaderState.PixelShaderRHI = GETSAFERHISHADER_PIXEL(*PixelShader);
+		GraphicsPSOInit.PrimitiveType = PT_TriangleList;
 
-	SetGraphicsPipelineState(Context.RHICmdList, GraphicsPSOInit);
+		SetGraphicsPipelineState(Context.RHICmdList, GraphicsPSOInit);
+
+		// The 1/8 factor is because bloom is using 8 taps in the filter.
+		VertexShader->SetVS(Context, FVector2D(ScaleAB.X, ScaleAB.Y));
+		FVector4 TintAScaled = TintA * (1.0f / 8.0f);
+		FVector4 TintBScaled = TintB * (1.0f / 8.0f);
+		PixelShader->SetPS(Context.RHICmdList, Context, TintAScaled, TintBScaled);
 
-	// The 1/8 factor is because bloom is using 8 taps in the filter.
-	VertexShader->SetVS(Context, FVector2D(ScaleAB.X, ScaleAB.Y));
-	FVector4 TintAScaled = TintA * (1.0f/8.0f);
-	FVector4 TintBScaled = TintB * (1.0f/8.0f);
-	PixelShader->SetPS(Context.RHICmdList, Context, TintAScaled, TintBScaled);
-
-	FIntPoint SrcDstSize = PrePostSourceViewportSize;
-
-	DrawRectangle(
-		Context.RHICmdList,
-		0, 0,
-		DstX, DstY,
-		0, 0,
-		DstX, DstY,
-		SrcDstSize,
-		SrcDstSize,
-		*VertexShader,
-		EDRF_UseTriangleOptimization);
-
-	Context.RHICmdList.CopyToResolveTarget(DestRenderTarget.TargetableTexture, DestRenderTarget.ShaderResourceTexture, FResolveParams());
+		FIntPoint SrcDstSize = PrePostSourceViewportSize;
+
+		DrawRectangle(
+			Context.RHICmdList,
+			0, 0,
+			DstX, DstY,
+			0, 0,
+			DstX, DstY,
+			SrcDstSize,
+			SrcDstSize,
+			*VertexShader,
+			EDRF_UseTriangleOptimization);
+	}
+	Context.RHICmdList.EndRenderPass();
 }
 
 FPooledRenderTargetDesc FRCPassPostProcessBloomUpES2::ComputeOutputDesc(EPassOutputId InPassOutputId) const
@@ -1040,6 +1062,7 @@ void FRCPassPostProcessSunMaskES2::Process(FRenderingCompositePassContext& Conte
 	TShaderMapRef<FPostProcessSunMaskVS_ES2> VertexShader(Context.GetShaderMap());
 	if(bOnChip)
 	{
+		// #todo-renderpasses how does this match up with renderpasses?
 		SrcSize = DstSize;
 		//SrcRect = DstRect;
 		SrcRect = View.ViewRect;
@@ -1077,41 +1100,45 @@ void FRCPassPostProcessSunMaskES2::Process(FRenderingCompositePassContext& Conte
 		const FSceneRenderTargetItem& DestRenderTarget = PassOutputs[0].RequestSurface(Context);
 		
 		// OverrideRenderTarget might patch out final render target and we have no control of the clear color anymore
-		if (DestRenderTarget.TargetableTexture->GetClearColor() == FLinearColor::Black)
-		{
-			FRHIRenderTargetView RtView = FRHIRenderTargetView(DestRenderTarget.TargetableTexture, ERenderTargetLoadAction::EClear);
-			FRHISetRenderTargetsInfo Info(1, &RtView, FRHIDepthRenderTargetView());
-			Context.RHICmdList.SetRenderTargetsAndClear(Info);
-		}
-		else
-		{
-			SetRenderTarget(Context.RHICmdList, DestRenderTarget.TargetableTexture, FTextureRHIRef());
-			DrawClearQuad(Context.RHICmdList, FLinearColor::Black);
-		}
-
-		Context.SetViewportAndCallRHI(0, 0, 0.0f, DstX, DstY, 1.0f );
+		const bool bUseClearQuad = DestRenderTarget.TargetableTexture->GetClearColor() != FLinearColor::Black;
+		ERenderTargetLoadAction LoadAction = ERenderTargetLoadAction::EClear;
 
-		if (InputDesc != NULL && InputDesc->Format == PF_FloatR11G11B10)
+		if (bUseClearQuad)
 		{
-			SetShader<true>(Context);
+			LoadAction = ERenderTargetLoadAction::ENoAction;
 		}
-		else
+
+		FRHIRenderPassInfo RPInfo(DestRenderTarget.TargetableTexture, MakeRenderTargetActions(LoadAction, ERenderTargetStoreAction::EStore), DestRenderTarget.ShaderResourceTexture);
+		Context.RHICmdList.BeginRenderPass(RPInfo, TEXT("PostProcessSunMaskES2"));
 		{
-			SetShader<false>(Context);
-		}
+			if (bUseClearQuad)
+			{
+				DrawClearQuad(Context.RHICmdList, FLinearColor::Black);
+			}
 
-		DrawRectangle(
-			Context.RHICmdList,
-			SrcRect.Min.X, SrcRect.Min.Y,
-			SrcRect.Width(), SrcRect.Height(),
-			SrcRect.Min.X, SrcRect.Min.Y,
-			SrcRect.Width(), SrcRect.Height(),
-			DstSize,
-			SrcSize,
-			*VertexShader,
-			EDRF_UseTriangleOptimization);
+			Context.SetViewportAndCallRHI(0, 0, 0.0f, DstX, DstY, 1.0f);
+
+			if (InputDesc != NULL && InputDesc->Format == PF_FloatR11G11B10)
+			{
+				SetShader<true>(Context);
+			}
+			else
+			{
+				SetShader<false>(Context);
+			}
 
-		Context.RHICmdList.CopyToResolveTarget(DestRenderTarget.TargetableTexture, DestRenderTarget.ShaderResourceTexture, FResolveParams());
+			DrawRectangle(
+				Context.RHICmdList,
+				SrcRect.Min.X, SrcRect.Min.Y,
+				SrcRect.Width(), SrcRect.Height(),
+				SrcRect.Min.X, SrcRect.Min.Y,
+				SrcRect.Width(), SrcRect.Height(),
+				DstSize,
+				SrcSize,
+				*VertexShader,
+				EDRF_UseTriangleOptimization);
+		}
+		Context.RHICmdList.EndRenderPass();
 	}
 }
 
@@ -1280,37 +1307,41 @@ void FRCPassPostProcessSunAlphaES2::Process(FRenderingCompositePassContext& Cont
 	const FSceneRenderTargetItem& DestRenderTarget = PassOutputs[0].RequestSurface(Context);
 	
 	// OverrideRenderTarget might patch out final render target and we have no control of the clear color anymore
-	if (DestRenderTarget.TargetableTexture->GetClearColor() == FLinearColor::Black)
+	const bool bUseClearQuad = DestRenderTarget.TargetableTexture->GetClearColor() != FLinearColor::Black;
+	ERenderTargetLoadAction LoadAction = ERenderTargetLoadAction::EClear;
+
+	if (bUseClearQuad)
 	{
-		FRHIRenderTargetView View = FRHIRenderTargetView(DestRenderTarget.TargetableTexture, ERenderTargetLoadAction::EClear);
-		FRHISetRenderTargetsInfo Info(1, &View, FRHIDepthRenderTargetView());
-		Context.RHICmdList.SetRenderTargetsAndClear(Info);
+		LoadAction = ERenderTargetLoadAction::ENoAction;
 	}
-	else
+
+	FRHIRenderPassInfo RPInfo(DestRenderTarget.TargetableTexture, MakeRenderTargetActions(LoadAction, ERenderTargetStoreAction::EStore), DestRenderTarget.ShaderResourceTexture);
+	Context.RHICmdList.BeginRenderPass(RPInfo, TEXT("PostProcessSunAlphaES2"));
 	{
-		SetRenderTarget(Context.RHICmdList, DestRenderTarget.TargetableTexture, FTextureRHIRef());
-		DrawClearQuad(Context.RHICmdList, FLinearColor::Black);
-	}
+		if (bUseClearQuad)
+		{
+			DrawClearQuad(Context.RHICmdList, FLinearColor::Black);
+		}
 
-	Context.SetViewportAndCallRHI(0, 0, 0.0f, DstX, DstY, 1.0f );
+		Context.SetViewportAndCallRHI(0, 0, 0.0f, DstX, DstY, 1.0f);
 
-	SetShader(Context);
+		SetShader(Context);
 
-	FIntPoint SrcDstSize = PrePostSourceViewportSize / 4;
-	TShaderMapRef<FPostProcessSunAlphaVS_ES2> VertexShader(Context.GetShaderMap());
+		FIntPoint SrcDstSize = PrePostSourceViewportSize / 4;
+		TShaderMapRef<FPostProcessSunAlphaVS_ES2> VertexShader(Context.GetShaderMap());
 
-	DrawRectangle(
-		Context.RHICmdList,
-		0, 0,
-		DstX, DstY,
-		0, 0,
-		DstX, DstY,
-		SrcDstSize,
-		SrcDstSize,
-		*VertexShader,
-		EDRF_UseTriangleOptimization);
-
-	Context.RHICmdList.CopyToResolveTarget(DestRenderTarget.TargetableTexture, DestRenderTarget.ShaderResourceTexture, FResolveParams());
+		DrawRectangle(
+			Context.RHICmdList,
+			0, 0,
+			DstX, DstY,
+			0, 0,
+			DstX, DstY,
+			SrcDstSize,
+			SrcDstSize,
+			*VertexShader,
+			EDRF_UseTriangleOptimization);
+	}
+	Context.RHICmdList.EndRenderPass();
 }
 
 FPooledRenderTargetDesc FRCPassPostProcessSunAlphaES2::ComputeOutputDesc(EPassOutputId InPassOutputId) const
@@ -1438,53 +1469,56 @@ void FRCPassPostProcessSunBlurES2::Process(FRenderingCompositePassContext& Conte
 	const FSceneRenderTargetItem& DestRenderTarget = PassOutputs[0].RequestSurface(Context);
 	
 	// OverrideRenderTarget might patch out final render target and we have no control of the clear color anymore
-	if (DestRenderTarget.TargetableTexture->GetClearColor() == FLinearColor::Black)
+	const bool bUseClearQuad = DestRenderTarget.TargetableTexture->GetClearColor() != FLinearColor::Black;
+	ERenderTargetLoadAction LoadAction = ERenderTargetLoadAction::EClear;
+	if (bUseClearQuad)
 	{
-		FRHIRenderTargetView View = FRHIRenderTargetView(DestRenderTarget.TargetableTexture, ERenderTargetLoadAction::EClear);
-		FRHISetRenderTargetsInfo Info(1, &View, FRHIDepthRenderTargetView());
-		Context.RHICmdList.SetRenderTargetsAndClear(Info);
+		LoadAction = ERenderTargetLoadAction::ENoAction;
 	}
-	else
+
+	FRHIRenderPassInfo RPInfo(DestRenderTarget.TargetableTexture, MakeRenderTargetActions(LoadAction, ERenderTargetStoreAction::EStore), DestRenderTarget.ShaderResourceTexture);
+	Context.RHICmdList.BeginRenderPass(RPInfo, TEXT("SunBlurES2"));
 	{
-		SetRenderTarget(Context.RHICmdList, DestRenderTarget.TargetableTexture, FTextureRHIRef());
-		DrawClearQuad(Context.RHICmdList, FLinearColor::Black);
-	}
+		if (bUseClearQuad)
+		{
+			DrawClearQuad(Context.RHICmdList, FLinearColor::Black);
+		}
 
-	Context.SetViewportAndCallRHI(0, 0, 0.0f, DstX, DstY, 1.0f );
+		Context.SetViewportAndCallRHI(0, 0, 0.0f, DstX, DstY, 1.0f);
 
-	FGraphicsPipelineStateInitializer GraphicsPSOInit;
-	Context.RHICmdList.ApplyCachedRenderTargets(GraphicsPSOInit);
-	GraphicsPSOInit.BlendState = TStaticBlendState<>::GetRHI();
-	GraphicsPSOInit.RasterizerState = TStaticRasterizerState<>::GetRHI();
-	GraphicsPSOInit.DepthStencilState = TStaticDepthStencilState<false, CF_Always>::GetRHI();
+		FGraphicsPipelineStateInitializer GraphicsPSOInit;
+		Context.RHICmdList.ApplyCachedRenderTargets(GraphicsPSOInit);
+		GraphicsPSOInit.BlendState = TStaticBlendState<>::GetRHI();
+		GraphicsPSOInit.RasterizerState = TStaticRasterizerState<>::GetRHI();
+		GraphicsPSOInit.DepthStencilState = TStaticDepthStencilState<false, CF_Always>::GetRHI();
 
-	TShaderMapRef<FPostProcessSunBlurVS_ES2> VertexShader(Context.GetShaderMap());
-	TShaderMapRef<FPostProcessSunBlurPS_ES2> PixelShader(Context.GetShaderMap());
+		TShaderMapRef<FPostProcessSunBlurVS_ES2> VertexShader(Context.GetShaderMap());
+		TShaderMapRef<FPostProcessSunBlurPS_ES2> PixelShader(Context.GetShaderMap());
 
-	GraphicsPSOInit.BoundShaderState.VertexDeclarationRHI = GFilterVertexDeclaration.VertexDeclarationRHI;
-	GraphicsPSOInit.BoundShaderState.VertexShaderRHI = GETSAFERHISHADER_VERTEX(*VertexShader);
-	GraphicsPSOInit.BoundShaderState.PixelShaderRHI = GETSAFERHISHADER_PIXEL(*PixelShader);
-	GraphicsPSOInit.PrimitiveType = PT_TriangleList;
+		GraphicsPSOInit.BoundShaderState.VertexDeclarationRHI = GFilterVertexDeclaration.VertexDeclarationRHI;
+		GraphicsPSOInit.BoundShaderState.VertexShaderRHI = GETSAFERHISHADER_VERTEX(*VertexShader);
+		GraphicsPSOInit.BoundShaderState.PixelShaderRHI = GETSAFERHISHADER_PIXEL(*PixelShader);
+		GraphicsPSOInit.PrimitiveType = PT_TriangleList;
 
-	SetGraphicsPipelineState(Context.RHICmdList, GraphicsPSOInit);
-
-	VertexShader->SetVS(Context);
-	PixelShader->SetPS(Context);
+		SetGraphicsPipelineState(Context.RHICmdList, GraphicsPSOInit);
 
-	FIntPoint SrcDstSize = PrePostSourceViewportSize / 4;
+		VertexShader->SetVS(Context);
+		PixelShader->SetPS(Context);
 
-	DrawRectangle(
-		Context.RHICmdList,
-		0, 0,
-		DstX, DstY,
-		0, 0,
-		DstX, DstY,
-		SrcDstSize,
-		SrcDstSize,
-		*VertexShader,
-		EDRF_UseTriangleOptimization);
+		FIntPoint SrcDstSize = PrePostSourceViewportSize / 4;
 
-	Context.RHICmdList.CopyToResolveTarget(DestRenderTarget.TargetableTexture, DestRenderTarget.ShaderResourceTexture, FResolveParams());
+		DrawRectangle(
+			Context.RHICmdList,
+			0, 0,
+			DstX, DstY,
+			0, 0,
+			DstX, DstY,
+			SrcDstSize,
+			SrcDstSize,
+			*VertexShader,
+			EDRF_UseTriangleOptimization);
+	}
+	Context.RHICmdList.EndRenderPass();
 }
 
 FPooledRenderTargetDesc FRCPassPostProcessSunBlurES2::ComputeOutputDesc(EPassOutputId InPassOutputId) const
@@ -1676,8 +1710,8 @@ void FRCPassPostProcessSunMergeES2::Process(FRenderingCompositePassContext& Cont
 {
 	SCOPED_DRAW_EVENT(Context.RHICmdList, PostProcessSunMerge);
 
-	uint32 DstX = FMath::Max(1, PrePostSourceViewportSize.X/4);
-	uint32 DstY = FMath::Max(1, PrePostSourceViewportSize.Y/4);
+	uint32 DstX = FMath::Max(1, PrePostSourceViewportSize.X / 4);
+	uint32 DstY = FMath::Max(1, PrePostSourceViewportSize.Y / 4);
 
 	FIntRect DstRect;
 	DstRect.Min.X = 0;
@@ -1686,38 +1720,40 @@ void FRCPassPostProcessSunMergeES2::Process(FRenderingCompositePassContext& Cont
 	DstRect.Max.Y = DstY;
 
 	const FSceneRenderTargetItem& DestRenderTarget = PassOutputs[0].RequestSurface(Context);
-	
+
 	// OverrideRenderTarget might patch out final render target and we have no control of the clear color anymore
-	if (DestRenderTarget.TargetableTexture->GetClearColor() == FLinearColor::Black)
+	const bool bUseClearQuad = DestRenderTarget.TargetableTexture->GetClearColor() != FLinearColor::Black;
+	ERenderTargetLoadAction LoadAction = ERenderTargetLoadAction::EClear;
+	if (bUseClearQuad)
 	{
-		FRHIRenderTargetView View = FRHIRenderTargetView(DestRenderTarget.TargetableTexture, ERenderTargetLoadAction::EClear);
-		FRHISetRenderTargetsInfo Info(1, &View, FRHIDepthRenderTargetView());
-		Context.RHICmdList.SetRenderTargetsAndClear(Info);
+		LoadAction = ERenderTargetLoadAction::ENoAction;
 	}
-	else
+	FRHIRenderPassInfo RPInfo(DestRenderTarget.TargetableTexture, MakeRenderTargetActions(LoadAction, ERenderTargetStoreAction::EStore), DestRenderTarget.ShaderResourceTexture);
+	Context.RHICmdList.BeginRenderPass(RPInfo, TEXT("SunMergeES2"));
 	{
-		SetRenderTarget(Context.RHICmdList, DestRenderTarget.TargetableTexture, FTextureRHIRef());
-		DrawClearQuad(Context.RHICmdList, FLinearColor::Black);
-	}
-
-	Context.SetViewportAndCallRHI(0, 0, 0.0f, DstX, DstY, 1.0f );
+		if (bUseClearQuad)
+		{
+			DrawClearQuad(Context.RHICmdList, FLinearColor::Black);
+		}
 
-	FShader* VertexShader = SetShader(Context);
+		Context.SetViewportAndCallRHI(0, 0, 0.0f, DstX, DstY, 1.0f);
 
-	FIntPoint SrcDstSize = PrePostSourceViewportSize / 4;
+		FShader* VertexShader = SetShader(Context);
 
-	DrawRectangle(
-		Context.RHICmdList,
-		0, 0,
-		DstX, DstY,
-		0, 0,
-		DstX, DstY,
-		SrcDstSize,
-		SrcDstSize,
-		VertexShader,
-		EDRF_UseTriangleOptimization);
+		FIntPoint SrcDstSize = PrePostSourceViewportSize / 4;
 
-	Context.RHICmdList.CopyToResolveTarget(DestRenderTarget.TargetableTexture, DestRenderTarget.ShaderResourceTexture, FResolveParams());
+		DrawRectangle(
+			Context.RHICmdList,
+			0, 0,
+			DstX, DstY,
+			0, 0,
+			DstX, DstY,
+			SrcDstSize,
+			SrcDstSize,
+			VertexShader,
+			EDRF_UseTriangleOptimization);
+	}
+	Context.RHICmdList.EndRenderPass();
 
 	// Double buffer sun+bloom+vignette composite.
 	if(Context.View.AntiAliasingMethod == AAM_TemporalAA)
@@ -1889,37 +1925,39 @@ void FRCPassPostProcessSunMergeSmallES2::Process(FRenderingCompositePassContext&
 	const FSceneRenderTargetItem& DestRenderTarget = PassOutputs[0].RequestSurface(Context);
 	
 	// OverrideRenderTarget might patch out final render target and we have no control of the clear color anymore
-	if (DestRenderTarget.TargetableTexture->GetClearColor() == FLinearColor::Black)
+	const bool bUseClearQuad = DestRenderTarget.TargetableTexture->GetClearColor() != FLinearColor::Black;
+	ERenderTargetLoadAction LoadAction = ERenderTargetLoadAction::EClear;
+	if (bUseClearQuad)
 	{
-		FRHIRenderTargetView View = FRHIRenderTargetView(DestRenderTarget.TargetableTexture, ERenderTargetLoadAction::EClear);
-		FRHISetRenderTargetsInfo Info(1, &View, FRHIDepthRenderTargetView());
-		Context.RHICmdList.SetRenderTargetsAndClear(Info);
+		LoadAction = ERenderTargetLoadAction::ENoAction;
 	}
-	else
+	FRHIRenderPassInfo RPInfo(DestRenderTarget.TargetableTexture, MakeRenderTargetActions(LoadAction, ERenderTargetStoreAction::EStore), DestRenderTarget.ShaderResourceTexture);
+	Context.RHICmdList.BeginRenderPass(RPInfo, TEXT("SunMergeSmallES2"));
 	{
-		SetRenderTarget(Context.RHICmdList, DestRenderTarget.TargetableTexture, FTextureRHIRef());
-		DrawClearQuad(Context.RHICmdList, FLinearColor::Black);
-	}
+		if (bUseClearQuad)
+		{
+			DrawClearQuad(Context.RHICmdList, FLinearColor::Black);
+		}
 
-	Context.SetViewportAndCallRHI(0, 0, 0.0f, DstX, DstY, 1.0f );
+		Context.SetViewportAndCallRHI(0, 0, 0.0f, DstX, DstY, 1.0f);
 
-	SetShader(Context);
+		SetShader(Context);
 
-	FIntPoint SrcDstSize = PrePostSourceViewportSize / 4;
-	TShaderMapRef<FPostProcessSunMergeSmallVS_ES2> VertexShader(Context.GetShaderMap());
+		FIntPoint SrcDstSize = PrePostSourceViewportSize / 4;
+		TShaderMapRef<FPostProcessSunMergeSmallVS_ES2> VertexShader(Context.GetShaderMap());
 
-	DrawRectangle(
-		Context.RHICmdList, 
-		0, 0,
-		DstX, DstY,
-		0, 0,
-		DstX, DstY,
-		SrcDstSize,
-		SrcDstSize,
-		*VertexShader,
-		EDRF_UseTriangleOptimization);
-
-	Context.RHICmdList.CopyToResolveTarget(DestRenderTarget.TargetableTexture, DestRenderTarget.ShaderResourceTexture, FResolveParams());
+		DrawRectangle(
+			Context.RHICmdList,
+			0, 0,
+			DstX, DstY,
+			0, 0,
+			DstX, DstY,
+			SrcDstSize,
+			SrcDstSize,
+			*VertexShader,
+			EDRF_UseTriangleOptimization);
+	}
+	Context.RHICmdList.EndRenderPass();
 
 	// Double buffer sun+bloom+vignette composite.
 
@@ -2129,36 +2167,38 @@ void FRCPassPostProcessDofDownES2::Process(FRenderingCompositePassContext& Conte
 	const FSceneRenderTargetItem& DestRenderTarget = PassOutputs[0].RequestSurface(Context);
 
 	// OverrideRenderTarget might patch out final render target and we have no control of the clear color anymore
-	if (DestRenderTarget.TargetableTexture->GetClearColor() == FLinearColor::Black)
+	const bool bUseClearQuad = DestRenderTarget.TargetableTexture->GetClearColor() != FLinearColor::Black;
+	ERenderTargetLoadAction LoadAction = ERenderTargetLoadAction::EClear;
+	if (bUseClearQuad)
 	{
-		FRHIRenderTargetView View = FRHIRenderTargetView(DestRenderTarget.TargetableTexture, ERenderTargetLoadAction::EClear);
-		FRHISetRenderTargetsInfo Info(1, &View, FRHIDepthRenderTargetView());
-		Context.RHICmdList.SetRenderTargetsAndClear(Info);
+		LoadAction = ERenderTargetLoadAction::ENoAction;
 	}
-	else
+	FRHIRenderPassInfo RPInfo(DestRenderTarget.TargetableTexture, MakeRenderTargetActions(LoadAction, ERenderTargetStoreAction::EStore), DestRenderTarget.ShaderResourceTexture);
+	Context.RHICmdList.BeginRenderPass(RPInfo, TEXT("DofDownES2"));
 	{
-		SetRenderTarget(Context.RHICmdList, DestRenderTarget.TargetableTexture, FTextureRHIRef());
-		DrawClearQuad(Context.RHICmdList, FLinearColor::Black);
-	}
+		if (bUseClearQuad)
+		{
+			DrawClearQuad(Context.RHICmdList, FLinearColor::Black);
+		}
 
-	Context.SetViewportAndCallRHI(0, 0, 0.0f, DstX, DstY, 1.0f );
+		Context.SetViewportAndCallRHI(0, 0, 0.0f, DstX, DstY, 1.0f);
 
-	SetShader(Context);
+		SetShader(Context);
 
-	TShaderMapRef<FPostProcessDofDownVS_ES2> VertexShader(Context.GetShaderMap());
+		TShaderMapRef<FPostProcessDofDownVS_ES2> VertexShader(Context.GetShaderMap());
 
-	DrawRectangle(
-		Context.RHICmdList,
-		0, 0,
-		DstX, DstY,
-		SrcRect.Min.X, SrcRect.Min.Y,
-		SrcRect.Width(), SrcRect.Height(),
-		DstSize,
-		SrcSize,
-		*VertexShader,
-		EDRF_UseTriangleOptimization);
-
-	Context.RHICmdList.CopyToResolveTarget(DestRenderTarget.TargetableTexture, DestRenderTarget.ShaderResourceTexture, FResolveParams());
+		DrawRectangle(
+			Context.RHICmdList,
+			0, 0,
+			DstX, DstY,
+			SrcRect.Min.X, SrcRect.Min.Y,
+			SrcRect.Width(), SrcRect.Height(),
+			DstSize,
+			SrcSize,
+			*VertexShader,
+			EDRF_UseTriangleOptimization);
+	}
+	Context.RHICmdList.EndRenderPass();
 }
 
 FPooledRenderTargetDesc FRCPassPostProcessDofDownES2::ComputeOutputDesc(EPassOutputId InPassOutputId) const
@@ -2332,37 +2372,39 @@ void FRCPassPostProcessDofNearES2::Process(FRenderingCompositePassContext& Conte
 	const FSceneRenderTargetItem& DestRenderTarget = PassOutputs[0].RequestSurface(Context);
 
 	// OverrideRenderTarget might patch out final render target and we have no control of the clear color anymore
-	if (DestRenderTarget.TargetableTexture->GetClearColor() == FLinearColor::Black)
+	const bool bUseClearQuad = DestRenderTarget.TargetableTexture->GetClearColor() != FLinearColor::Black;
+	ERenderTargetLoadAction LoadAction = ERenderTargetLoadAction::EClear;
+	if (bUseClearQuad)
 	{
-		FRHIRenderTargetView View = FRHIRenderTargetView(DestRenderTarget.TargetableTexture, ERenderTargetLoadAction::EClear);
-		FRHISetRenderTargetsInfo Info(1, &View, FRHIDepthRenderTargetView());
-		Context.RHICmdList.SetRenderTargetsAndClear(Info);
+		LoadAction = ERenderTargetLoadAction::ENoAction;
 	}
-	else
+	FRHIRenderPassInfo RPInfo(DestRenderTarget.TargetableTexture, MakeRenderTargetActions(LoadAction, ERenderTargetStoreAction::EStore), DestRenderTarget.ShaderResourceTexture);
+	Context.RHICmdList.BeginRenderPass(RPInfo, TEXT("DofNearES2"));
 	{
-		SetRenderTarget(Context.RHICmdList, DestRenderTarget.TargetableTexture, FTextureRHIRef());
-		DrawClearQuad(Context.RHICmdList, FLinearColor::Black);
-	}
+		if (bUseClearQuad)
+		{
+			DrawClearQuad(Context.RHICmdList, FLinearColor::Black);
+		}
 
-	Context.SetViewportAndCallRHI(0, 0, 0.0f, DstX, DstY, 1.0f );
+		Context.SetViewportAndCallRHI(0, 0, 0.0f, DstX, DstY, 1.0f);
 
-	SetShader(Context);
+		SetShader(Context);
 
-	FIntPoint SrcDstSize = PrePostSourceViewportSize / 4;
-	TShaderMapRef<FPostProcessDofNearVS_ES2> VertexShader(Context.GetShaderMap());
+		FIntPoint SrcDstSize = PrePostSourceViewportSize / 4;
+		TShaderMapRef<FPostProcessDofNearVS_ES2> VertexShader(Context.GetShaderMap());
 
-	DrawRectangle(
-		Context.RHICmdList,
-		0, 0,
-		DstX, DstY,
-		0, 0,
-		DstX, DstY,
-		SrcDstSize,
-		SrcSize,
-		*VertexShader,
-		EDRF_UseTriangleOptimization);
-
-	Context.RHICmdList.CopyToResolveTarget(DestRenderTarget.TargetableTexture, DestRenderTarget.ShaderResourceTexture, FResolveParams());
+		DrawRectangle(
+			Context.RHICmdList,
+			0, 0,
+			DstX, DstY,
+			0, 0,
+			DstX, DstY,
+			SrcDstSize,
+			SrcSize,
+			*VertexShader,
+			EDRF_UseTriangleOptimization);
+	}
+	Context.RHICmdList.EndRenderPass();
 }
 
 FPooledRenderTargetDesc FRCPassPostProcessDofNearES2::ComputeOutputDesc(EPassOutputId InPassOutputId) const
@@ -2483,53 +2525,55 @@ void FRCPassPostProcessDofBlurES2::Process(FRenderingCompositePassContext& Conte
 	const FSceneRenderTargetItem& DestRenderTarget = PassOutputs[0].RequestSurface(Context);
 	
 	// OverrideRenderTarget might patch out final render target and we have no control of the clear color anymore
-	if (DestRenderTarget.TargetableTexture->GetClearColor() == FLinearColor::Black)
+	const bool bUseClearQuad = DestRenderTarget.TargetableTexture->GetClearColor() != FLinearColor::Black;
+	ERenderTargetLoadAction LoadAction = ERenderTargetLoadAction::EClear;
+	if (bUseClearQuad)
 	{
-		FRHIRenderTargetView View = FRHIRenderTargetView(DestRenderTarget.TargetableTexture, ERenderTargetLoadAction::EClear);
-		FRHISetRenderTargetsInfo Info(1, &View, FRHIDepthRenderTargetView());
-		Context.RHICmdList.SetRenderTargetsAndClear(Info);
+		LoadAction = ERenderTargetLoadAction::ENoAction;
 	}
-	else
+	FRHIRenderPassInfo RPInfo(DestRenderTarget.TargetableTexture, MakeRenderTargetActions(LoadAction, ERenderTargetStoreAction::EStore), DestRenderTarget.ShaderResourceTexture);
+	Context.RHICmdList.BeginRenderPass(RPInfo, TEXT("DofBlurES2"));
 	{
-		SetRenderTarget(Context.RHICmdList, DestRenderTarget.TargetableTexture, FTextureRHIRef());
-		DrawClearQuad(Context.RHICmdList, FLinearColor::Black);
-	}
-
-	Context.SetViewportAndCallRHI(0, 0, 0.0f, DstX, DstY, 1.0f );
+		if (bUseClearQuad)
+		{
+			DrawClearQuad(Context.RHICmdList, FLinearColor::Black);
+		}
 
-	FGraphicsPipelineStateInitializer GraphicsPSOInit;
-	Context.RHICmdList.ApplyCachedRenderTargets(GraphicsPSOInit);
-	GraphicsPSOInit.BlendState = TStaticBlendState<>::GetRHI();
-	GraphicsPSOInit.RasterizerState = TStaticRasterizerState<>::GetRHI();
-	GraphicsPSOInit.DepthStencilState = TStaticDepthStencilState<false, CF_Always>::GetRHI();
+		Context.SetViewportAndCallRHI(0, 0, 0.0f, DstX, DstY, 1.0f);
 
-	TShaderMapRef<FPostProcessDofBlurVS_ES2> VertexShader(Context.GetShaderMap());
-	TShaderMapRef<FPostProcessDofBlurPS_ES2> PixelShader(Context.GetShaderMap());
+		FGraphicsPipelineStateInitializer GraphicsPSOInit;
+		Context.RHICmdList.ApplyCachedRenderTargets(GraphicsPSOInit);
+		GraphicsPSOInit.BlendState = TStaticBlendState<>::GetRHI();
+		GraphicsPSOInit.RasterizerState = TStaticRasterizerState<>::GetRHI();
+		GraphicsPSOInit.DepthStencilState = TStaticDepthStencilState<false, CF_Always>::GetRHI();
 
-	GraphicsPSOInit.BoundShaderState.VertexDeclarationRHI = GFilterVertexDeclaration.VertexDeclarationRHI;
-	GraphicsPSOInit.BoundShaderState.VertexShaderRHI = GETSAFERHISHADER_VERTEX(*VertexShader);
-	GraphicsPSOInit.BoundShaderState.PixelShaderRHI = GETSAFERHISHADER_PIXEL(*PixelShader);
-	GraphicsPSOInit.PrimitiveType = PT_TriangleList;
+		TShaderMapRef<FPostProcessDofBlurVS_ES2> VertexShader(Context.GetShaderMap());
+		TShaderMapRef<FPostProcessDofBlurPS_ES2> PixelShader(Context.GetShaderMap());
 
-	SetGraphicsPipelineState(Context.RHICmdList, GraphicsPSOInit);
+		GraphicsPSOInit.BoundShaderState.VertexDeclarationRHI = GFilterVertexDeclaration.VertexDeclarationRHI;
+		GraphicsPSOInit.BoundShaderState.VertexShaderRHI = GETSAFERHISHADER_VERTEX(*VertexShader);
+		GraphicsPSOInit.BoundShaderState.PixelShaderRHI = GETSAFERHISHADER_PIXEL(*PixelShader);
+		GraphicsPSOInit.PrimitiveType = PT_TriangleList;
 
-	VertexShader->SetVS(Context);
-	PixelShader->SetPS(Context);
+		SetGraphicsPipelineState(Context.RHICmdList, GraphicsPSOInit);
 
-	FIntPoint SrcDstSize = PrePostSourceViewportSize / 2;
+		VertexShader->SetVS(Context);
+		PixelShader->SetPS(Context);
 
-	DrawRectangle(
-		Context.RHICmdList,
-		0, 0,
-		DstX, DstY,
-		0, 0,
-		DstX, DstY,
-		SrcDstSize,
-		SrcDstSize,
-		*VertexShader,
-		EDRF_UseTriangleOptimization);
+		FIntPoint SrcDstSize = PrePostSourceViewportSize / 2;
 
-	Context.RHICmdList.CopyToResolveTarget(DestRenderTarget.TargetableTexture, DestRenderTarget.ShaderResourceTexture, FResolveParams());
+		DrawRectangle(
+			Context.RHICmdList,
+			0, 0,
+			DstX, DstY,
+			0, 0,
+			DstX, DstY,
+			SrcDstSize,
+			SrcDstSize,
+			*VertexShader,
+			EDRF_UseTriangleOptimization);
+	}
+	Context.RHICmdList.EndRenderPass();
 }
 
 FPooledRenderTargetDesc FRCPassPostProcessDofBlurES2::ComputeOutputDesc(EPassOutputId InPassOutputId) const
@@ -2681,37 +2725,39 @@ void FRCPassPostProcessSunAvgES2::Process(FRenderingCompositePassContext& Contex
 	const FSceneRenderTargetItem& DestRenderTarget = PassOutputs[0].RequestSurface(Context);
 	
 	// OverrideRenderTarget might patch out final render target and we have no control of the clear color anymore
-	if (DestRenderTarget.TargetableTexture->GetClearColor() == FLinearColor::Black)
+	const bool bUseClearQuad = DestRenderTarget.TargetableTexture->GetClearColor() != FLinearColor::Black;
+	ERenderTargetLoadAction LoadAction = ERenderTargetLoadAction::EClear;
+	if (bUseClearQuad)
 	{
-		FRHIRenderTargetView View = FRHIRenderTargetView(DestRenderTarget.TargetableTexture, ERenderTargetLoadAction::EClear);
-		FRHISetRenderTargetsInfo Info(1, &View, FRHIDepthRenderTargetView());
-		Context.RHICmdList.SetRenderTargetsAndClear(Info);
+		LoadAction = ERenderTargetLoadAction::ENoAction;
 	}
-	else
+	FRHIRenderPassInfo RPInfo(DestRenderTarget.TargetableTexture, MakeRenderTargetActions(LoadAction, ERenderTargetStoreAction::EStore), DestRenderTarget.ShaderResourceTexture);
+	Context.RHICmdList.BeginRenderPass(RPInfo, TEXT("SunAvgES2"));
 	{
-		SetRenderTarget(Context.RHICmdList, DestRenderTarget.TargetableTexture, FTextureRHIRef());
-		DrawClearQuad(Context.RHICmdList, FLinearColor::Black);
-	}
+		if (bUseClearQuad)
+		{
+			DrawClearQuad(Context.RHICmdList, FLinearColor::Black);
+		}
 
-	Context.SetViewportAndCallRHI(0, 0, 0.0f, DstX, DstY, 1.0f );
+		Context.SetViewportAndCallRHI(0, 0, 0.0f, DstX, DstY, 1.0f);
 
-	SetShader(Context);
+		SetShader(Context);
 
-	FIntPoint SrcDstSize = PrePostSourceViewportSize / 4;
-	TShaderMapRef<FPostProcessSunAvgVS_ES2> VertexShader(Context.GetShaderMap());
+		FIntPoint SrcDstSize = PrePostSourceViewportSize / 4;
+		TShaderMapRef<FPostProcessSunAvgVS_ES2> VertexShader(Context.GetShaderMap());
 
-	DrawRectangle(
-		Context.RHICmdList,
-		0, 0,
-		DstX, DstY,
-		0, 0,
-		DstX, DstY,
-		SrcDstSize,
-		SrcDstSize,
-		*VertexShader,
-		EDRF_UseTriangleOptimization);
-
-	Context.RHICmdList.CopyToResolveTarget(DestRenderTarget.TargetableTexture, DestRenderTarget.ShaderResourceTexture, FResolveParams());
+		DrawRectangle(
+			Context.RHICmdList,
+			0, 0,
+			DstX, DstY,
+			0, 0,
+			DstX, DstY,
+			SrcDstSize,
+			SrcDstSize,
+			*VertexShader,
+			EDRF_UseTriangleOptimization);
+	}
+	Context.RHICmdList.EndRenderPass();
 }
 
 FPooledRenderTargetDesc FRCPassPostProcessSunAvgES2::ComputeOutputDesc(EPassOutputId InPassOutputId) const
@@ -2933,40 +2979,40 @@ void FRCPassPostProcessAaES2::Process(FRenderingCompositePassContext& Context)
 	
 	check(SrcSize == DestSize);
 
+	ERenderTargetActions LoadStoreAction = ERenderTargetActions::Load_Store;
 	if (Context.View.StereoPass != eSSP_RIGHT_EYE)
 	{
 		// Full clear to avoid restore
-		SetRenderTarget(Context.RHICmdList, DestRenderTarget.TargetableTexture, FTextureRHIRef(), ESimpleRenderTargetMode::EClearColorAndDepth);
+		LoadStoreAction = ERenderTargetActions::Clear_Store;
 	}
-	else
+	
+	FRHIRenderPassInfo RPInfo(DestRenderTarget.TargetableTexture, LoadStoreAction, DestRenderTarget.ShaderResourceTexture);
+	Context.RHICmdList.BeginRenderPass(RPInfo, TEXT("AaES2"));
 	{
-		SetRenderTarget(Context.RHICmdList, DestRenderTarget.TargetableTexture, FTextureRHIRef(), ESimpleRenderTargetMode::EExistingColorAndDepth);
-	}
-
-	Context.SetViewportAndCallRHI(0, 0, 0.0f, DestSize.X, DestSize.Y, 1.0f);
+		Context.SetViewportAndCallRHI(0, 0, 0.0f, DestSize.X, DestSize.Y, 1.0f);
 
-	SetShader(Context);
+		SetShader(Context);
 
-	const FIntRect& ViewRect = Context.View.UnscaledViewRect; // Simple upscaling, ES2 post process does not currently have a specific upscaling pass.
-	float XPos = ViewRect.Min.X;
-	float YPos = ViewRect.Min.Y;
-	float Width = ViewRect.Width();
-	float Height = ViewRect.Height();
+		const FIntRect& ViewRect = Context.View.UnscaledViewRect; // Simple upscaling, ES2 post process does not currently have a specific upscaling pass.
+		float XPos = ViewRect.Min.X;
+		float YPos = ViewRect.Min.Y;
+		float Width = ViewRect.Width();
+		float Height = ViewRect.Height();
 
-	TShaderMapRef<FPostProcessAaVS_ES2> VertexShader(Context.GetShaderMap());
+		TShaderMapRef<FPostProcessAaVS_ES2> VertexShader(Context.GetShaderMap());
 
-	DrawRectangle(
-		Context.RHICmdList,
-		XPos, YPos,
-		Width, Height,
-		XPos, YPos,
-		Width, Height,
-		DestSize,
-		SrcSize,
-		*VertexShader,
-		EDRF_UseTriangleOptimization);
-
-	Context.RHICmdList.CopyToResolveTarget(DestRenderTarget.TargetableTexture, DestRenderTarget.ShaderResourceTexture, FResolveParams());
+		DrawRectangle(
+			Context.RHICmdList,
+			XPos, YPos,
+			Width, Height,
+			XPos, YPos,
+			Width, Height,
+			DestSize,
+			SrcSize,
+			*VertexShader,
+			EDRF_UseTriangleOptimization);
+	}
+	Context.RHICmdList.EndRenderPass();
 
 	if (FSceneRenderer::ShouldCompositeEditorPrimitives(Context.View))
 	{
diff --git a/Engine/Source/Runtime/Renderer/Private/PostProcess/PostProcessMorpheus.cpp b/Engine/Source/Runtime/Renderer/Private/PostProcess/PostProcessMorpheus.cpp
--- a/Engine/Source/Runtime/Renderer/Private/PostProcess/PostProcessMorpheus.cpp
+++ b/Engine/Source/Runtime/Renderer/Private/PostProcess/PostProcessMorpheus.cpp
@@ -223,51 +223,52 @@ void FRCPassPostProcessMorpheus::Process(FRenderingCompositePassContext& Context
 	const FSceneRenderTargetItem& DestRenderTarget = PassOutputs[0].RequestSurface(Context);
 
 	// Set the view family's render target/viewport.
-	SetRenderTarget(Context.RHICmdList, DestRenderTarget.TargetableTexture, FTextureRHIRef());
-
-	Context.SetViewportAndCallRHI(DestRect);
+	FRHIRenderPassInfo RPInfo(DestRenderTarget.TargetableTexture, ERenderTargetActions::Load_Store, DestRenderTarget.ShaderResourceTexture);
+	Context.RHICmdList.BeginRenderPass(RPInfo, TEXT("MorpheusPostProcess"));
+	{
+		Context.SetViewportAndCallRHI(DestRect);
 
-	FGraphicsPipelineStateInitializer GraphicsPSOInit;
-	Context.RHICmdList.ApplyCachedRenderTargets(GraphicsPSOInit);
-	GraphicsPSOInit.BlendState = TStaticBlendState<>::GetRHI();
-	GraphicsPSOInit.RasterizerState = TStaticRasterizerState<>::GetRHI();
-	GraphicsPSOInit.DepthStencilState = TStaticDepthStencilState<false, CF_Always>::GetRHI();
+		FGraphicsPipelineStateInitializer GraphicsPSOInit;
+		Context.RHICmdList.ApplyCachedRenderTargets(GraphicsPSOInit);
+		GraphicsPSOInit.BlendState = TStaticBlendState<>::GetRHI();
+		GraphicsPSOInit.RasterizerState = TStaticRasterizerState<>::GetRHI();
+		GraphicsPSOInit.DepthStencilState = TStaticDepthStencilState<false, CF_Always>::GetRHI();
 
 #if defined(MORPHEUS_ENGINE_DISTORTION) && MORPHEUS_ENGINE_DISTORTION
-	TShaderMapRef<FPostProcessMorpheusVS> VertexShader(Context.GetShaderMap());
-	TShaderMapRef<FPostProcessMorpheusPS> PixelShader(Context.GetShaderMap());
-
-	GraphicsPSOInit.BoundShaderState.VertexDeclarationRHI = GFilterVertexDeclaration.VertexDeclarationRHI;
-	GraphicsPSOInit.BoundShaderState.VertexShaderRHI = GETSAFERHISHADER_VERTEX(*VertexShader);
-	GraphicsPSOInit.BoundShaderState.PixelShaderRHI = GETSAFERHISHADER_PIXEL(*PixelShader);
-	GraphicsPSOInit.PrimitiveType = PT_TriangleList;
-
-	SetGraphicsPipelineState(Context.RHICmdList, GraphicsPSOInit);
-
-	FMatrix QuadTexTransform;
-	FMatrix QuadPosTransform = FMatrix::Identity;
-
-	PixelShader->SetPS(Context.RHICmdList, Context, SrcRect, SrcSize, View.StereoPass, QuadTexTransform);
-
-	// Draw a quad mapping scene color to the view's render target
-	DrawTransformedRectangle(
-		Context.RHICmdList,
-		0, 0,
-		DestRect.Width(), DestRect.Height(),
-		QuadPosTransform,
-		SrcRect.Min.X, SrcRect.Min.Y,
-		SrcRect.Width(), SrcRect.Height(),
-		QuadTexTransform,
-		DestRect.Size(),
-		SrcSize
+		TShaderMapRef<FPostProcessMorpheusVS> VertexShader(Context.GetShaderMap());
+		TShaderMapRef<FPostProcessMorpheusPS> PixelShader(Context.GetShaderMap());
+
+		GraphicsPSOInit.BoundShaderState.VertexDeclarationRHI = GFilterVertexDeclaration.VertexDeclarationRHI;
+		GraphicsPSOInit.BoundShaderState.VertexShaderRHI = GETSAFERHISHADER_VERTEX(*VertexShader);
+		GraphicsPSOInit.BoundShaderState.PixelShaderRHI = GETSAFERHISHADER_PIXEL(*PixelShader);
+		GraphicsPSOInit.PrimitiveType = PT_TriangleList;
+
+		SetGraphicsPipelineState(Context.RHICmdList, GraphicsPSOInit);
+
+		FMatrix QuadTexTransform;
+		FMatrix QuadPosTransform = FMatrix::Identity;
+
+		PixelShader->SetPS(Context.RHICmdList, Context, SrcRect, SrcSize, View.StereoPass, QuadTexTransform);
+
+		// Draw a quad mapping scene color to the view's render target
+		DrawTransformedRectangle(
+			Context.RHICmdList,
+			0, 0,
+			DestRect.Width(), DestRect.Height(),
+			QuadPosTransform,
+			SrcRect.Min.X, SrcRect.Min.Y,
+			SrcRect.Width(), SrcRect.Height(),
+			QuadTexTransform,
+			DestRect.Size(),
+			SrcSize
 		);
 #elif PLATFORM_PS4
-	checkf(false, TEXT("PS4 uses SDK distortion."));
+		checkf(false, TEXT("PS4 uses SDK distortion."));
 #else
-	checkf(false, TEXT("Unsupported path.  Morpheus should be disabled."));
+		checkf(false, TEXT("Unsupported path.  Morpheus should be disabled."));
 #endif
-
-	Context.RHICmdList.CopyToResolveTarget(DestRenderTarget.TargetableTexture, DestRenderTarget.ShaderResourceTexture, FResolveParams());
+	}
+	Context.RHICmdList.EndRenderPass();
 }
 
 FPooledRenderTargetDesc FRCPassPostProcessMorpheus::ComputeOutputDesc(EPassOutputId InPassOutputId) const
diff --git a/Engine/Source/Runtime/Renderer/Private/PostProcess/PostProcessMotionBlur.cpp b/Engine/Source/Runtime/Renderer/Private/PostProcess/PostProcessMotionBlur.cpp
--- a/Engine/Source/Runtime/Renderer/Private/PostProcess/PostProcessMotionBlur.cpp
+++ b/Engine/Source/Runtime/Renderer/Private/PostProcess/PostProcessMotionBlur.cpp
@@ -169,7 +169,8 @@ void FRCPassPostProcessVelocityFlatten::Process(FRenderingCompositePassContext&
 
 	TShaderMapRef< FPostProcessVelocityFlattenCS > ComputeShader( Context.GetShaderMap() );
 
-	SetRenderTarget(Context.RHICmdList, FTextureRHIRef(), FTextureRHIRef());
+	// #todo-renderpasses remove once everything is renderpasses
+	UnbindRenderTargets(Context.RHICmdList);
 
 	Context.SetViewportAndCallRHI( View.ViewRect );
 	Context.RHICmdList.SetComputeShader(ComputeShader->GetComputeShader());
@@ -338,57 +339,59 @@ void FRCPassPostProcessVelocityScatter::Process(FRenderingCompositePassContext&
 	GRenderTargetPool.FindFreeElement(Context.RHICmdList, Desc, DepthTarget, TEXT("VelocityScatterDepth") );
 
 	// Set the view family's render target/viewport.
-	FRHIRenderTargetView ColorView(DestRenderTarget.TargetableTexture, 0, -1, ERenderTargetLoadAction::ELoad, ERenderTargetStoreAction::EStore);
-	FRHIDepthRenderTargetView DepthView(DepthTarget->GetRenderTargetItem().TargetableTexture, ERenderTargetLoadAction::EClear, ERenderTargetStoreAction::EStore, ERenderTargetLoadAction::ELoad, ERenderTargetStoreAction::EStore);
-	FRHISetRenderTargetsInfo RTInfo(1, &ColorView, DepthView);
-
-	// clear depth
-	// Max >= Min so no need to clear on second pass
-	Context.RHICmdList.SetRenderTargetsAndClear(RTInfo);
-	Context.SetViewportAndCallRHI(0, 0, 0.0f, TileCount.X, TileCount.Y, 1.0f);
-	
-	FGraphicsPipelineStateInitializer GraphicsPSOInit;
-	Context.RHICmdList.ApplyCachedRenderTargets(GraphicsPSOInit);
+	FRHIRenderPassInfo RPInfo(DestRenderTarget.TargetableTexture, ERenderTargetActions::Load_Store, DestRenderTarget.ShaderResourceTexture);
+	RPInfo.DepthStencilRenderTarget.DepthStencilTarget = DepthTarget->GetRenderTargetItem().TargetableTexture;
+	RPInfo.DepthStencilRenderTarget.Action = MakeDepthStencilTargetActions(ERenderTargetActions::Clear_Store, ERenderTargetActions::Load_Store);
+	RPInfo.DepthStencilRenderTarget.ExclusiveDepthStencil = FExclusiveDepthStencil::DepthWrite_StencilWrite;
 
-	// Min,Max
-	for( int i = 0; i < 2; i++ )
+	Context.RHICmdList.BeginRenderPass(RPInfo, TEXT("VelocityScatterDepth"));
 	{
-		if( i == 0 )
-		{
-			// min
-			GraphicsPSOInit.BlendState = TStaticBlendStateWriteMask< CW_RGBA >::GetRHI();
-			GraphicsPSOInit.RasterizerState = TStaticRasterizerState<>::GetRHI();
-			GraphicsPSOInit.DepthStencilState = TStaticDepthStencilState< true, CF_Less >::GetRHI();
-		}
-		else
+		// clear depth
+		// Max >= Min so no need to clear on second pass
+		Context.SetViewportAndCallRHI(0, 0, 0.0f, TileCount.X, TileCount.Y, 1.0f);
+
+		FGraphicsPipelineStateInitializer GraphicsPSOInit;
+		Context.RHICmdList.ApplyCachedRenderTargets(GraphicsPSOInit);
+
+		// Min,Max
+		for (int i = 0; i < 2; i++)
 		{
-			// max
-			GraphicsPSOInit.BlendState = TStaticBlendStateWriteMask< CW_BA >::GetRHI();
-			GraphicsPSOInit.RasterizerState = TStaticRasterizerState<>::GetRHI();
-			GraphicsPSOInit.DepthStencilState = TStaticDepthStencilState< true, CF_Greater >::GetRHI();
-		}
+			if (i == 0)
+			{
+				// min
+				GraphicsPSOInit.BlendState = TStaticBlendStateWriteMask< CW_RGBA >::GetRHI();
+				GraphicsPSOInit.RasterizerState = TStaticRasterizerState<>::GetRHI();
+				GraphicsPSOInit.DepthStencilState = TStaticDepthStencilState< true, CF_Less >::GetRHI();
+			}
+			else
+			{
+				// max
+				GraphicsPSOInit.BlendState = TStaticBlendStateWriteMask< CW_BA >::GetRHI();
+				GraphicsPSOInit.RasterizerState = TStaticRasterizerState<>::GetRHI();
+				GraphicsPSOInit.DepthStencilState = TStaticDepthStencilState< true, CF_Greater >::GetRHI();
+			}
 
-		TShaderMapRef< FPostProcessVelocityScatterVS > VertexShader(Context.GetShaderMap());
-		TShaderMapRef< FPostProcessVelocityScatterPS > PixelShader(Context.GetShaderMap());
+			TShaderMapRef< FPostProcessVelocityScatterVS > VertexShader(Context.GetShaderMap());
+			TShaderMapRef< FPostProcessVelocityScatterPS > PixelShader(Context.GetShaderMap());
 
-		GraphicsPSOInit.BoundShaderState.VertexDeclarationRHI = GEmptyVertexDeclaration.VertexDeclarationRHI;
-		GraphicsPSOInit.BoundShaderState.VertexShaderRHI = GETSAFERHISHADER_VERTEX(*VertexShader);
-		GraphicsPSOInit.BoundShaderState.PixelShaderRHI = GETSAFERHISHADER_PIXEL(*PixelShader);
-		GraphicsPSOInit.PrimitiveType = PT_TriangleList;
+			GraphicsPSOInit.BoundShaderState.VertexDeclarationRHI = GEmptyVertexDeclaration.VertexDeclarationRHI;
+			GraphicsPSOInit.BoundShaderState.VertexShaderRHI = GETSAFERHISHADER_VERTEX(*VertexShader);
+			GraphicsPSOInit.BoundShaderState.PixelShaderRHI = GETSAFERHISHADER_PIXEL(*PixelShader);
+			GraphicsPSOInit.PrimitiveType = PT_TriangleList;
 
-		SetGraphicsPipelineState(Context.RHICmdList, GraphicsPSOInit);
+			SetGraphicsPipelineState(Context.RHICmdList, GraphicsPSOInit);
 
-		VertexShader->SetParameters( Context, i );
-		PixelShader->SetParameters( Context );
+			VertexShader->SetParameters(Context, i);
+			PixelShader->SetParameters(Context);
 
-		// needs to be the same on shader side (faster on NVIDIA and AMD)
-		int32 QuadsPerInstance = 8;
+			// needs to be the same on shader side (faster on NVIDIA and AMD)
+			int32 QuadsPerInstance = 8;
 
-		Context.RHICmdList.SetStreamSource(0, NULL, 0);
-		Context.RHICmdList.DrawIndexedPrimitive(GScatterQuadIndexBuffer.IndexBufferRHI, 0, 0, 32, 0, 2 * QuadsPerInstance, FMath::DivideAndRoundUp(TileCount.X * TileCount.Y, QuadsPerInstance));
+			Context.RHICmdList.SetStreamSource(0, NULL, 0);
+			Context.RHICmdList.DrawIndexedPrimitive(GScatterQuadIndexBuffer.IndexBufferRHI, 0, 0, 32, 0, 2 * QuadsPerInstance, FMath::DivideAndRoundUp(TileCount.X * TileCount.Y, QuadsPerInstance));
+		}
 	}
-
-	Context.RHICmdList.CopyToResolveTarget(DestRenderTarget.TargetableTexture, DestRenderTarget.ShaderResourceTexture, FResolveParams());
+	Context.RHICmdList.EndRenderPass();
 }
 
 FPooledRenderTargetDesc FRCPassPostProcessVelocityScatter::ComputeOutputDesc(EPassOutputId InPassOutputId) const
@@ -466,7 +469,8 @@ void FRCPassPostProcessVelocityGather::Process(FRenderingCompositePassContext& C
 	
 	const FSceneRenderTargetItem& DestRenderTarget = PassOutputs[0].RequestSurface(Context);
 
-	SetRenderTarget(Context.RHICmdList, FTextureRHIRef(), FTextureRHIRef());
+	// #todo-renderpasses remove when everything is renderpasses.
+	UnbindRenderTargets(Context.RHICmdList);
 	Context.SetViewportAndCallRHI( 0, 0, 0.0f, TileCount.X, TileCount.Y, 1.0f );
 	
 	TShaderMapRef< FPostProcessVelocityGatherCS > ComputeShader( Context.GetShaderMap() );
@@ -872,7 +876,8 @@ void FRCPassPostProcessMotionBlur::Process(FRenderingCompositePassContext& Conte
 	if (bIsComputePass)
 	{
 		// Common setup
-		SetRenderTarget(Context.RHICmdList, nullptr, nullptr);
+		// #todo-renderpasses remove once everything is renderpasses.
+		UnbindRenderTargets(Context.RHICmdList);
 		Context.SetViewportAndCallRHI(DestRect, 0.0f, 1.0f);
 		
 		static FName AsyncEndFenceName(TEXT("AsyncMotionBlurEndFence"));
@@ -906,45 +911,45 @@ void FRCPassPostProcessMotionBlur::Process(FRenderingCompositePassContext& Conte
 
 		ERenderTargetLoadAction LoadAction = Context.GetLoadActionForRenderTarget(DestRenderTarget);
 
-		FRHIRenderTargetView RtView = FRHIRenderTargetView(DestRenderTarget.TargetableTexture, LoadAction);
-		FRHISetRenderTargetsInfo Info(1, &RtView, FRHIDepthRenderTargetView());
-		Context.RHICmdList.SetRenderTargetsAndClear(Info);
-		Context.SetViewportAndCallRHI(SrcRect);
-
-		if(Quality == 1)
-		{
-			SetMotionBlurShaderNewTempl<1>( Context, Scale, InnerLoopConfig );
-		}
-		else if(Quality == 2)
+		FRHIRenderPassInfo RPInfo(DestRenderTarget.TargetableTexture, MakeRenderTargetActions(LoadAction, ERenderTargetStoreAction::EStore), DestRenderTarget.ShaderResourceTexture);
+		Context.RHICmdList.BeginRenderPass(RPInfo, TEXT("MotionBlur"));
 		{
-			SetMotionBlurShaderNewTempl<2>( Context, Scale, InnerLoopConfig );
-		}
-		else if(Quality == 3 || Pass > 0 )
-		{
-			SetMotionBlurShaderNewTempl<3>( Context, Scale, InnerLoopConfig );
-		}
-		else
-		{
-			check(Quality == 4);
-			SetMotionBlurShaderNewTempl<4>( Context, Scale, InnerLoopConfig );
-		}
+			Context.SetViewportAndCallRHI(SrcRect);
 
-		TShaderMapRef<FPostProcessVS> VertexShader(Context.GetShaderMap());
-	
-		DrawPostProcessPass(
-			Context.RHICmdList,
-			0, 0,
-			SrcRect.Width(), SrcRect.Height(),
-			SrcRect.Min.X, SrcRect.Min.Y,
-			SrcRect.Width(), SrcRect.Height(),
-			SrcRect.Size(),
-			SrcSize,
-			*VertexShader,
-			View.StereoPass,
-			Context.HasHmdMesh(),
-			EDRF_UseTriangleOptimization);
+			if (Quality == 1)
+			{
+				SetMotionBlurShaderNewTempl<1>(Context, Scale, InnerLoopConfig);
+			}
+			else if (Quality == 2)
+			{
+				SetMotionBlurShaderNewTempl<2>(Context, Scale, InnerLoopConfig);
+			}
+			else if (Quality == 3 || Pass > 0)
+			{
+				SetMotionBlurShaderNewTempl<3>(Context, Scale, InnerLoopConfig);
+			}
+			else
+			{
+				check(Quality == 4);
+				SetMotionBlurShaderNewTempl<4>(Context, Scale, InnerLoopConfig);
+			}
 
-		Context.RHICmdList.CopyToResolveTarget(DestRenderTarget.TargetableTexture, DestRenderTarget.ShaderResourceTexture, FResolveParams());
+			TShaderMapRef<FPostProcessVS> VertexShader(Context.GetShaderMap());
+
+			DrawPostProcessPass(
+				Context.RHICmdList,
+				0, 0,
+				SrcRect.Width(), SrcRect.Height(),
+				SrcRect.Min.X, SrcRect.Min.Y,
+				SrcRect.Width(), SrcRect.Height(),
+				SrcRect.Size(),
+				SrcSize,
+				*VertexShader,
+				View.StereoPass,
+				Context.HasHmdMesh(),
+				EDRF_UseTriangleOptimization);
+		}
+		Context.RHICmdList.EndRenderPass();
 	}
 }
 
@@ -1153,41 +1158,43 @@ void FRCPassPostProcessVisualizeMotionBlur::Process(FRenderingCompositePassConte
 	const FSceneRenderTargetItem& DestRenderTarget = PassOutputs[0].RequestSurface(Context);
 
 	// Set the view family's render target/viewport.
-	FRHIRenderTargetView RtView = FRHIRenderTargetView(DestRenderTarget.TargetableTexture, ERenderTargetLoadAction::ENoAction);
-	FRHISetRenderTargetsInfo Info(1, &RtView, FRHIDepthRenderTargetView());
-	Context.RHICmdList.SetRenderTargetsAndClear(Info);
-	Context.SetViewportAndCallRHI(SrcRect);
+	FRHIRenderPassInfo RPInfo(DestRenderTarget.TargetableTexture, ERenderTargetActions::DontLoad_Store);
+	Context.RHICmdList.BeginRenderPass(RPInfo, TEXT("VisualizeMotionBlur"));
+	{
+		Context.SetViewportAndCallRHI(SrcRect);
 
-	FGraphicsPipelineStateInitializer GraphicsPSOInit;
-	Context.RHICmdList.ApplyCachedRenderTargets(GraphicsPSOInit);
-	GraphicsPSOInit.BlendState = TStaticBlendState<>::GetRHI();
-	GraphicsPSOInit.RasterizerState = TStaticRasterizerState<>::GetRHI();
-	GraphicsPSOInit.DepthStencilState = TStaticDepthStencilState<false, CF_Always>::GetRHI();
+		FGraphicsPipelineStateInitializer GraphicsPSOInit;
+		Context.RHICmdList.ApplyCachedRenderTargets(GraphicsPSOInit);
+		GraphicsPSOInit.BlendState = TStaticBlendState<>::GetRHI();
+		GraphicsPSOInit.RasterizerState = TStaticRasterizerState<>::GetRHI();
+		GraphicsPSOInit.DepthStencilState = TStaticDepthStencilState<false, CF_Always>::GetRHI();
 
-	TShaderMapRef<FPostProcessVS> VertexShader(Context.GetShaderMap());
-	TShaderMapRef<FPostProcessVisualizeMotionBlurPS> PixelShader(Context.GetShaderMap());
+		TShaderMapRef<FPostProcessVS> VertexShader(Context.GetShaderMap());
+		TShaderMapRef<FPostProcessVisualizeMotionBlurPS> PixelShader(Context.GetShaderMap());
 
-	GraphicsPSOInit.BoundShaderState.VertexDeclarationRHI = GFilterVertexDeclaration.VertexDeclarationRHI;
-	GraphicsPSOInit.BoundShaderState.VertexShaderRHI = GETSAFERHISHADER_VERTEX(*VertexShader);
-	GraphicsPSOInit.BoundShaderState.PixelShaderRHI = GETSAFERHISHADER_PIXEL(*PixelShader);
-	GraphicsPSOInit.PrimitiveType = PT_TriangleList;
+		GraphicsPSOInit.BoundShaderState.VertexDeclarationRHI = GFilterVertexDeclaration.VertexDeclarationRHI;
+		GraphicsPSOInit.BoundShaderState.VertexShaderRHI = GETSAFERHISHADER_VERTEX(*VertexShader);
+		GraphicsPSOInit.BoundShaderState.PixelShaderRHI = GETSAFERHISHADER_PIXEL(*PixelShader);
+		GraphicsPSOInit.PrimitiveType = PT_TriangleList;
 
-	SetGraphicsPipelineState(Context.RHICmdList, GraphicsPSOInit);
+		SetGraphicsPipelineState(Context.RHICmdList, GraphicsPSOInit);
 
-	VertexShader->SetParameters(Context);
-	PixelShader->SetParameters(Context.RHICmdList, Context);
-
-	// Draw a quad mapping scene color to the view's render target
-	DrawRectangle(
-		Context.RHICmdList,
-		0, 0,
-		SrcRect.Width(), SrcRect.Height(),
-		SrcRect.Min.X, SrcRect.Min.Y, 
-		SrcRect.Width(), SrcRect.Height(),
-		SrcRect.Size(),
-		SrcSize,
-		*VertexShader,
-		EDRF_UseTriangleOptimization);
+		VertexShader->SetParameters(Context);
+		PixelShader->SetParameters(Context.RHICmdList, Context);
+
+		// Draw a quad mapping scene color to the view's render target
+		DrawRectangle(
+			Context.RHICmdList,
+			0, 0,
+			SrcRect.Width(), SrcRect.Height(),
+			SrcRect.Min.X, SrcRect.Min.Y,
+			SrcRect.Width(), SrcRect.Height(),
+			SrcRect.Size(),
+			SrcSize,
+			*VertexShader,
+			EDRF_UseTriangleOptimization);
+	}
+	Context.RHICmdList.EndRenderPass();
 
 	FRenderTargetTemp TempRenderTarget(View, (const FTexture2DRHIRef&)DestRenderTarget.TargetableTexture);
 	FCanvas Canvas(&TempRenderTarget, NULL, ViewFamily.CurrentRealTime, ViewFamily.CurrentWorldTime, ViewFamily.DeltaWorldTime, Context.GetFeatureLevel());
@@ -1201,7 +1208,7 @@ void FRCPassPostProcessVisualizeMotionBlur::Process(FRenderingCompositePassConte
 
 	Line = FString::Printf(TEXT("Visualize MotionBlur"));
 	Canvas.DrawShadowedString(X, Y += YStep, *Line, GetStatsFont(), FLinearColor(1, 1, 0));
-	
+
 	static const auto MotionBlurDebugVar = IConsoleManager::Get().FindTConsoleVariableDataInt(TEXT("r.MotionBlurDebug"));
 	const int32 MotionBlurDebug = MotionBlurDebugVar ? MotionBlurDebugVar->GetValueOnRenderThread() : 0;
 
@@ -1211,7 +1218,7 @@ void FRCPassPostProcessVisualizeMotionBlur::Process(FRenderingCompositePassConte
 
 	static const auto VelocityTestVar = IConsoleManager::Get().FindTConsoleVariableDataInt(TEXT("r.VelocityTest"));
 	const int32 VelocityTest = VelocityTestVar ? VelocityTestVar->GetValueOnRenderThread() : 0;
-	
+
 	extern bool IsParallelVelocity();
 
 	Line = FString::Printf(TEXT("%d, %d, %d"), ViewFamily.bWorldIsPaused, VelocityTest, IsParallelVelocity());
@@ -1221,7 +1228,7 @@ void FRCPassPostProcessVisualizeMotionBlur::Process(FRenderingCompositePassConte
 	const FScene* Scene = (const FScene*)View.Family->Scene;
 
 	Canvas.DrawShadowedString(X, Y += YStep, TEXT("MotionBlurInfoData (per object):"), GetStatsFont(), FLinearColor(1, 1, 0));
-	Canvas.DrawShadowedString(X + ColumnWidth, Y, *Scene->MotionBlurInfoData.GetDebugString(), GetStatsFont(), FLinearColor(1, 1, 0));	
+	Canvas.DrawShadowedString(X + ColumnWidth, Y, *Scene->MotionBlurInfoData.GetDebugString(), GetStatsFont(), FLinearColor(1, 1, 0));
 
 	const FSceneViewState *SceneViewState = (const FSceneViewState*)View.State;
 
diff --git a/Engine/Source/Runtime/Renderer/Private/PostProcess/PostProcessNoiseBlur.cpp b/Engine/Source/Runtime/Renderer/Private/PostProcess/PostProcessNoiseBlur.cpp
--- a/Engine/Source/Runtime/Renderer/Private/PostProcess/PostProcessNoiseBlur.cpp
+++ b/Engine/Source/Runtime/Renderer/Private/PostProcess/PostProcessNoiseBlur.cpp
@@ -149,45 +149,47 @@ void FRCPassPostProcessNoiseBlur::Process(FRenderingCompositePassContext& Contex
 	const FSceneRenderTargetItem& DestRenderTarget = PassOutputs[0].RequestSurface(Context);
 
 	// Set the view family's render target/viewport.
-	SetRenderTarget(Context.RHICmdList, DestRenderTarget.TargetableTexture, FTextureRHIRef());
-
-	if (View.StereoPass == eSSP_FULL)
+	FRHIRenderPassInfo RPInfo(DestRenderTarget.TargetableTexture, ERenderTargetActions::Load_Store, DestRenderTarget.ShaderResourceTexture);
+	Context.RHICmdList.BeginRenderPass(RPInfo, TEXT("NoiseBlur"));
 	{
-		// is optimized away if possible (RT size=view size, )
-		DrawClearQuad(Context.RHICmdList, true, FLinearColor(0, 0, 0, 0), false, 0, false, 0, PassOutputs[0].RenderTargetDesc.Extent, DestRect);
-	}
+		// #todo-renderpasses perhaps an optimization here. Use NoAction if this will clear the whole RT
+		if (View.StereoPass == eSSP_FULL)
+		{
+			// is optimized away if possible (RT size=view size, )
+			DrawClearQuad(Context.RHICmdList, true, FLinearColor(0, 0, 0, 0), false, 0, false, 0, PassOutputs[0].RenderTargetDesc.Extent, DestRect);
+		}
 
-	Context.SetViewportAndCallRHI(0, 0, 0.0f, DestSize.X, DestSize.Y, 1.0f );
-	
-	if(Quality == 0)
-	{
-		SetNoiseBlurShader<0>(Context, Radius);
-	}
-	else if(Quality == 1)
-	{
-		SetNoiseBlurShader<1>(Context, Radius);
-	}
-	else
-	{
-		SetNoiseBlurShader<2>(Context, Radius);
-	}
+		Context.SetViewportAndCallRHI(0, 0, 0.0f, DestSize.X, DestSize.Y, 1.0f);
 
-	TShaderMapRef<FPostProcessVS> VertexShader(Context.GetShaderMap());
+		if (Quality == 0)
+		{
+			SetNoiseBlurShader<0>(Context, Radius);
+		}
+		else if (Quality == 1)
+		{
+			SetNoiseBlurShader<1>(Context, Radius);
+		}
+		else
+		{
+			SetNoiseBlurShader<2>(Context, Radius);
+		}
 
-	DrawPostProcessPass(
-		Context.RHICmdList,
-		DestRect.Min.X, DestRect.Min.Y,
-		DestRect.Width(), DestRect.Height(),
-		SrcRect.Min.X, SrcRect.Min.Y,
-		SrcRect.Width(), SrcRect.Height(),
-		DestSize,
-		SrcSize,
-		*VertexShader,
-		View.StereoPass,
-		Context.HasHmdMesh(),
-		EDRF_UseTriangleOptimization);
-
-	Context.RHICmdList.CopyToResolveTarget(DestRenderTarget.TargetableTexture, DestRenderTarget.ShaderResourceTexture, FResolveParams());
+		TShaderMapRef<FPostProcessVS> VertexShader(Context.GetShaderMap());
+
+		DrawPostProcessPass(
+			Context.RHICmdList,
+			DestRect.Min.X, DestRect.Min.Y,
+			DestRect.Width(), DestRect.Height(),
+			SrcRect.Min.X, SrcRect.Min.Y,
+			SrcRect.Width(), SrcRect.Height(),
+			DestSize,
+			SrcSize,
+			*VertexShader,
+			View.StereoPass,
+			Context.HasHmdMesh(),
+			EDRF_UseTriangleOptimization);
+	}
+	Context.RHICmdList.EndRenderPass();
 }
 
 
diff --git a/Engine/Source/Runtime/Renderer/Private/PostProcess/PostProcessSelectionOutline.cpp b/Engine/Source/Runtime/Renderer/Private/PostProcess/PostProcessSelectionOutline.cpp
--- a/Engine/Source/Runtime/Renderer/Private/PostProcess/PostProcessSelectionOutline.cpp
+++ b/Engine/Source/Runtime/Renderer/Private/PostProcess/PostProcessSelectionOutline.cpp
@@ -47,8 +47,11 @@ void FRCPassPostProcessSelectionOutlineColor::Process(FRenderingCompositePassCon
 	const FSceneRenderTargetItem& DestRenderTarget = PassOutputs[0].RequestSurface(Context);
 	
 	// Set the render target/viewport.
-	FRHIDepthRenderTargetView DepthRT(DestRenderTarget.TargetableTexture, ERenderTargetLoadAction::EClear, ERenderTargetStoreAction::ENoAction, ERenderTargetLoadAction::EClear, ERenderTargetStoreAction::ENoAction);
-	FRHISetRenderTargetsInfo RTInfo(0, nullptr, DepthRT);
+	FRHIRenderPassInfo RPInfo;
+	RPInfo.DepthStencilRenderTarget.DepthStencilTarget = DestRenderTarget.TargetableTexture;
+	RPInfo.DepthStencilRenderTarget.ResolveTarget = DestRenderTarget.ShaderResourceTexture;
+	RPInfo.DepthStencilRenderTarget.Action = MakeDepthStencilTargetActions(ERenderTargetActions::Clear_DontStore, ERenderTargetActions::Clear_DontStore);
+	RPInfo.DepthStencilRenderTarget.ExclusiveDepthStencil = FExclusiveDepthStencil::DepthWrite_StencilWrite;
 
 	if (GRHIRequiresRenderTargetForPixelShaderUAVs)
 	{
@@ -59,118 +62,118 @@ void FRCPassPostProcessSelectionOutlineColor::Process(FRenderingCompositePassCon
 		GRenderTargetPool.FindFreeElement(Context.RHICmdList, Desc, Dummy, TEXT("Dummy"));
 		FRHIRenderTargetView DummyRTView(Dummy->GetRenderTargetItem().TargetableTexture, ERenderTargetLoadAction::ENoAction);
 
-		RTInfo = FRHISetRenderTargetsInfo(1, &DummyRTView, DepthRT);
+		RPInfo.ColorRenderTargets[0].RenderTarget = Dummy->GetRenderTargetItem().TargetableTexture;
+		RPInfo.ColorRenderTargets[0].Action = ERenderTargetActions::DontLoad_Store;
 	}
 
-	Context.RHICmdList.SetRenderTargetsAndClear(RTInfo);
-
-	Context.SetViewportAndCallRHI(ViewRect);
-
-	WaitForInputPassComputeFences(Context.RHICmdList);
-
-	if (View.Family->EngineShowFlags.Selection)
+	Context.RHICmdList.BeginRenderPass(RPInfo, TEXT("SelectionOutlineColor"));
 	{
-		FViewInfo& EditorView = *Context.View.CreateSnapshot();
+		Context.SetViewportAndCallRHI(ViewRect);
 
+		WaitForInputPassComputeFences(Context.RHICmdList);
+
+		if (View.Family->EngineShowFlags.Selection)
 		{
-			// Patch view rect.
-			EditorView.ViewRect = ViewRect;
+			FViewInfo& EditorView = *Context.View.CreateSnapshot();
 
-			// Override pre exposure to 1.0f, because rendering after tonemapper. 
-			EditorView.PreExposure = 1.0f;
+			{
+				// Patch view rect.
+				EditorView.ViewRect = ViewRect;
 
-			// Kills material texture mipbias because after TAA.
-			EditorView.MaterialTextureMipBias = 0.0f;
+				// Override pre exposure to 1.0f, because rendering after tonemapper. 
+				EditorView.PreExposure = 1.0f;
 
-			if (EditorView.AntiAliasingMethod == AAM_TemporalAA)
-			{
-				EditorView.ViewMatrices.HackRemoveTemporalAAProjectionJitter();
-			}
+				// Kills material texture mipbias because after TAA.
+				EditorView.MaterialTextureMipBias = 0.0f;
 
-			EditorView.CachedViewUniformShaderParameters = MakeUnique<FViewUniformShaderParameters>();
+				if (EditorView.AntiAliasingMethod == AAM_TemporalAA)
+				{
+					EditorView.ViewMatrices.HackRemoveTemporalAAProjectionJitter();
+				}
 
-			FBox VolumeBounds[TVC_MAX];
-			EditorView.SetupUniformBufferParameters(SceneContext, VolumeBounds, TVC_MAX, *EditorView.CachedViewUniformShaderParameters);
-			EditorView.ViewUniformBuffer = TUniformBufferRef<FViewUniformShaderParameters>::CreateUniformBufferImmediate(*EditorView.CachedViewUniformShaderParameters, UniformBuffer_SingleFrame);
-		}
+				EditorView.CachedViewUniformShaderParameters = MakeUnique<FViewUniformShaderParameters>();
 
-		FSceneTexturesUniformParameters SceneTextureParameters;
-		SetupSceneTextureUniformParameters(SceneContext, EditorView.FeatureLevel, ESceneTextureSetupMode::None, SceneTextureParameters);
-		TUniformBufferRef<FSceneTexturesUniformParameters> PassUniformBuffer = TUniformBufferRef<FSceneTexturesUniformParameters>::CreateUniformBufferImmediate(SceneTextureParameters, UniformBuffer_SingleFrame);
+				FBox VolumeBounds[TVC_MAX];
+				EditorView.SetupUniformBufferParameters(SceneContext, VolumeBounds, TVC_MAX, *EditorView.CachedViewUniformShaderParameters);
+				EditorView.ViewUniformBuffer = TUniformBufferRef<FViewUniformShaderParameters>::CreateUniformBufferImmediate(*EditorView.CachedViewUniformShaderParameters, UniformBuffer_SingleFrame);
+			}
 
-		FDrawingPolicyRenderState DrawRenderState(EditorView, PassUniformBuffer);
+			FSceneTexturesUniformParameters SceneTextureParameters;
+			SetupSceneTextureUniformParameters(SceneContext, EditorView.FeatureLevel, ESceneTextureSetupMode::None, SceneTextureParameters);
+			TUniformBufferRef<FSceneTexturesUniformParameters> PassUniformBuffer = TUniformBufferRef<FSceneTexturesUniformParameters>::CreateUniformBufferImmediate(SceneTextureParameters, UniformBuffer_SingleFrame);
 
-		FHitProxyDrawingPolicyFactory::ContextType FactoryContext;
-		DrawRenderState.ModifyViewOverrideFlags() |= EDrawingPolicyOverrideFlags::TwoSided;
-		DrawRenderState.SetBlendState(TStaticBlendStateWriteMask<CW_NONE, CW_NONE, CW_NONE, CW_NONE>::GetRHI());
-		DrawRenderState.SetDepthStencilState(TStaticDepthStencilState<false, CF_Always>::GetRHI());
+			FDrawingPolicyRenderState DrawRenderState(EditorView, PassUniformBuffer);
 
-		// Note that the stencil value will overflow with enough selected objects
-		FEditorSelectionDrawingPolicy::ResetStencilValues();
+			FHitProxyDrawingPolicyFactory::ContextType FactoryContext;
+			DrawRenderState.ModifyViewOverrideFlags() |= EDrawingPolicyOverrideFlags::TwoSided;
+			DrawRenderState.SetBlendState(TStaticBlendStateWriteMask<CW_NONE, CW_NONE, CW_NONE, CW_NONE>::GetRHI());
+			DrawRenderState.SetDepthStencilState(TStaticDepthStencilState<false, CF_Always>::GetRHI());
 
-		// Run selection pass on static elements
-		FScene* Scene = View.Family->Scene->GetRenderScene();
-		if(Scene)
-		{	
-			Scene->EditorSelectionDrawList.DrawVisible(Context.RHICmdList, EditorView, DrawRenderState, View.StaticMeshEditorSelectionMap, View.StaticMeshBatchVisibility);
-		}
-	
-		for (int32 MeshBatchIndex = 0; MeshBatchIndex < View.DynamicMeshElements.Num(); MeshBatchIndex++)
-		{
-			const FMeshBatchAndRelevance& MeshBatchAndRelevance = View.DynamicMeshElements[MeshBatchIndex];
-			const FPrimitiveSceneProxy* PrimitiveSceneProxy = MeshBatchAndRelevance.PrimitiveSceneProxy;
+			// Note that the stencil value will overflow with enough selected objects
+			FEditorSelectionDrawingPolicy::ResetStencilValues();
 
-			// Selected actors should be subdued if any component is individually selected
-			bool bActorSelectionColorIsSubdued = View.bHasSelectedComponents;
+			// Run selection pass on static elements
+			FScene* Scene = View.Family->Scene->GetRenderScene();
+			if (Scene)
+			{
+				Scene->EditorSelectionDrawList.DrawVisible(Context.RHICmdList, EditorView, DrawRenderState, View.StaticMeshEditorSelectionMap, View.StaticMeshBatchVisibility);
+			}
 
-			if (PrimitiveSceneProxy->IsSelected() && MeshBatchAndRelevance.Mesh->bUseSelectionOutline && PrimitiveSceneProxy->WantsSelectionOutline())
+			for (int32 MeshBatchIndex = 0; MeshBatchIndex < View.DynamicMeshElements.Num(); MeshBatchIndex++)
 			{
-				int32 StencilValue = 1;
-				if(PrimitiveSceneProxy->GetOwnerName() != NAME_BSP)
+				const FMeshBatchAndRelevance& MeshBatchAndRelevance = View.DynamicMeshElements[MeshBatchIndex];
+				const FPrimitiveSceneProxy* PrimitiveSceneProxy = MeshBatchAndRelevance.PrimitiveSceneProxy;
+
+				// Selected actors should be subdued if any component is individually selected
+				bool bActorSelectionColorIsSubdued = View.bHasSelectedComponents;
+
+				if (PrimitiveSceneProxy->IsSelected() && MeshBatchAndRelevance.Mesh->bUseSelectionOutline && PrimitiveSceneProxy->WantsSelectionOutline())
 				{
-					StencilValue = FEditorSelectionDrawingPolicy::GetStencilValue(View, PrimitiveSceneProxy);
+					int32 StencilValue = 1;
+					if (PrimitiveSceneProxy->GetOwnerName() != NAME_BSP)
+					{
+						StencilValue = FEditorSelectionDrawingPolicy::GetStencilValue(View, PrimitiveSceneProxy);
 
-				}
+					}
 
-				// Note that the stencil value will overflow with enough selected objects
-				DrawRenderState.SetDepthStencilState(TStaticDepthStencilState<true, CF_DepthNearOrEqual, true, CF_Always, SO_Keep, SO_Keep, SO_Replace>::GetRHI());
-				DrawRenderState.SetStencilRef(StencilValue);
+					// Note that the stencil value will overflow with enough selected objects
+					DrawRenderState.SetDepthStencilState(TStaticDepthStencilState<true, CF_DepthNearOrEqual, true, CF_Always, SO_Keep, SO_Keep, SO_Replace>::GetRHI());
+					DrawRenderState.SetStencilRef(StencilValue);
 
-				const FMeshBatch& MeshBatch = *MeshBatchAndRelevance.Mesh;
-				FHitProxyDrawingPolicyFactory::DrawDynamicMesh(Context.RHICmdList, EditorView, FactoryContext, MeshBatch, true, DrawRenderState, MeshBatchAndRelevance.PrimitiveSceneProxy, MeshBatch.BatchHitProxyId);
+					const FMeshBatch& MeshBatch = *MeshBatchAndRelevance.Mesh;
+					FHitProxyDrawingPolicyFactory::DrawDynamicMesh(Context.RHICmdList, EditorView, FactoryContext, MeshBatch, true, DrawRenderState, MeshBatchAndRelevance.PrimitiveSceneProxy, MeshBatch.BatchHitProxyId);
+				}
 			}
-		}
 
-		// to get an outline around the objects if it's partly outside of the screen
-		{
-			FIntRect InnerRect = ViewRect;
-
-			// 1 as we have an outline that is that thick
-			InnerRect.InflateRect(-1);
-
-			// We could use Clear with InnerRect but this is just an optimization - on some hardware it might do a full clear (and we cannot disable yet)
-			//			RHICmdList.Clear(false, FLinearColor(0, 0, 0, 0), true, 0.0f, true, 0, InnerRect);
-			// so we to 4 clears - one for each border.
-
-			// top
-			Context.RHICmdList.SetScissorRect(true, ViewRect.Min.X, ViewRect.Min.Y, ViewRect.Max.X, InnerRect.Min.Y);
-			DrawClearQuad(Context.RHICmdList, false, FLinearColor(), true, (float)ERHIZBuffer::FarPlane, true, 0, PassOutputs[0].RenderTargetDesc.Extent, FIntRect());
-			// bottom
-			Context.RHICmdList.SetScissorRect(true, ViewRect.Min.X, InnerRect.Max.Y, ViewRect.Max.X, ViewRect.Max.Y);
-			DrawClearQuad(Context.RHICmdList, false, FLinearColor(), true, (float)ERHIZBuffer::FarPlane, true, 0, PassOutputs[0].RenderTargetDesc.Extent, FIntRect());
-			// left
-			Context.RHICmdList.SetScissorRect(true, ViewRect.Min.X, ViewRect.Min.Y, InnerRect.Min.X, ViewRect.Max.Y);
-			DrawClearQuad(Context.RHICmdList, false, FLinearColor(), true, (float)ERHIZBuffer::FarPlane, true, 0, PassOutputs[0].RenderTargetDesc.Extent, FIntRect());
-			// right
-			Context.RHICmdList.SetScissorRect(true, InnerRect.Max.X, ViewRect.Min.Y, ViewRect.Max.X, ViewRect.Max.Y);
-			DrawClearQuad(Context.RHICmdList, false, FLinearColor(), true, (float)ERHIZBuffer::FarPlane, true, 0, PassOutputs[0].RenderTargetDesc.Extent, FIntRect());
-
-			Context.RHICmdList.SetScissorRect(false, 0, 0, 0, 0);
+			// to get an outline around the objects if it's partly outside of the screen
+			{
+				FIntRect InnerRect = ViewRect;
+
+				// 1 as we have an outline that is that thick
+				InnerRect.InflateRect(-1);
+
+				// We could use Clear with InnerRect but this is just an optimization - on some hardware it might do a full clear (and we cannot disable yet)
+				//			RHICmdList.Clear(false, FLinearColor(0, 0, 0, 0), true, 0.0f, true, 0, InnerRect);
+				// so we to 4 clears - one for each border.
+
+				// top
+				Context.RHICmdList.SetScissorRect(true, ViewRect.Min.X, ViewRect.Min.Y, ViewRect.Max.X, InnerRect.Min.Y);
+				DrawClearQuad(Context.RHICmdList, false, FLinearColor(), true, (float)ERHIZBuffer::FarPlane, true, 0, PassOutputs[0].RenderTargetDesc.Extent, FIntRect());
+				// bottom
+				Context.RHICmdList.SetScissorRect(true, ViewRect.Min.X, InnerRect.Max.Y, ViewRect.Max.X, ViewRect.Max.Y);
+				DrawClearQuad(Context.RHICmdList, false, FLinearColor(), true, (float)ERHIZBuffer::FarPlane, true, 0, PassOutputs[0].RenderTargetDesc.Extent, FIntRect());
+				// left
+				Context.RHICmdList.SetScissorRect(true, ViewRect.Min.X, ViewRect.Min.Y, InnerRect.Min.X, ViewRect.Max.Y);
+				DrawClearQuad(Context.RHICmdList, false, FLinearColor(), true, (float)ERHIZBuffer::FarPlane, true, 0, PassOutputs[0].RenderTargetDesc.Extent, FIntRect());
+				// right
+				Context.RHICmdList.SetScissorRect(true, InnerRect.Max.X, ViewRect.Min.Y, ViewRect.Max.X, ViewRect.Max.Y);
+				DrawClearQuad(Context.RHICmdList, false, FLinearColor(), true, (float)ERHIZBuffer::FarPlane, true, 0, PassOutputs[0].RenderTargetDesc.Extent, FIntRect());
+
+				Context.RHICmdList.SetScissorRect(false, 0, 0, 0, 0);
+			}
 		}
 	}
-	
-	// Resolve to the output
-	Context.RHICmdList.CopyToResolveTarget(DestRenderTarget.TargetableTexture, DestRenderTarget.ShaderResourceTexture, FResolveParams());
+	Context.RHICmdList.EndRenderPass();
 #endif	
 }
 
@@ -406,47 +409,49 @@ void FRCPassPostProcessSelectionOutline::Process(FRenderingCompositePassContext&
 	FIntPoint SrcSize = SceneColorInputDesc->Extent;
 
 	// Set the view family's render target/viewport.
-	SetRenderTarget(Context.RHICmdList, DestRenderTarget.TargetableTexture, FTextureRHIRef());
-	Context.SetViewportAndCallRHI(DestRect);
+	FRHIRenderPassInfo RPInfo(DestRenderTarget.TargetableTexture, ERenderTargetActions::Load_Store, DestRenderTarget.ShaderResourceTexture);
+	Context.RHICmdList.BeginRenderPass(RPInfo, TEXT("SelectionOutline"));
+	{
+		Context.SetViewportAndCallRHI(DestRect);
 
-	const uint32 MSAASampleCount = FSceneRenderTargets::Get(Context.RHICmdList).GetEditorMSAACompositingSampleCount();
+		const uint32 MSAASampleCount = FSceneRenderTargets::Get(Context.RHICmdList).GetEditorMSAACompositingSampleCount();
 
-	if(MSAASampleCount == 1)
-	{
-		SetSelectionOutlineShaderTempl<1>(Context);
-	}
-	else if(MSAASampleCount == 2)
-	{
-		SetSelectionOutlineShaderTempl<2>(Context);
-	}
-	else if(MSAASampleCount == 4)
-	{
-		SetSelectionOutlineShaderTempl<4>(Context);
-	}
-	else if(MSAASampleCount == 8)
-	{
-		SetSelectionOutlineShaderTempl<8>(Context);
-	}
-	else
-	{
-		// not supported, internal error
-		check(0);
-	}
+		if (MSAASampleCount == 1)
+		{
+			SetSelectionOutlineShaderTempl<1>(Context);
+		}
+		else if (MSAASampleCount == 2)
+		{
+			SetSelectionOutlineShaderTempl<2>(Context);
+		}
+		else if (MSAASampleCount == 4)
+		{
+			SetSelectionOutlineShaderTempl<4>(Context);
+		}
+		else if (MSAASampleCount == 8)
+		{
+			SetSelectionOutlineShaderTempl<8>(Context);
+		}
+		else
+		{
+			// not supported, internal error
+			check(0);
+		}
 
-	// Draw a quad mapping scene color to the view's render target
-	TShaderMapRef<FPostProcessVS> VertexShader(Context.GetShaderMap());
-	DrawRectangle(
-		Context.RHICmdList,
-		0, 0,
-		DestRect.Width(), DestRect.Height(),
-		SrcRect.Min.X, SrcRect.Min.Y,
-		SrcRect.Width(), SrcRect.Height(),
-		DestRect.Size(),
-		SrcSize,
-		*VertexShader,
-		EDRF_UseTriangleOptimization);
-
-	Context.RHICmdList.CopyToResolveTarget(DestRenderTarget.TargetableTexture, DestRenderTarget.ShaderResourceTexture, FResolveParams());
+		// Draw a quad mapping scene color to the view's render target
+		TShaderMapRef<FPostProcessVS> VertexShader(Context.GetShaderMap());
+		DrawRectangle(
+			Context.RHICmdList,
+			0, 0,
+			DestRect.Width(), DestRect.Height(),
+			SrcRect.Min.X, SrcRect.Min.Y,
+			SrcRect.Width(), SrcRect.Height(),
+			DestRect.Size(),
+			SrcSize,
+			*VertexShader,
+			EDRF_UseTriangleOptimization);
+	}
+	Context.RHICmdList.EndRenderPass();
 }
 
 FPooledRenderTargetDesc FRCPassPostProcessSelectionOutline::ComputeOutputDesc(EPassOutputId InPassOutputId) const
diff --git a/Engine/Source/Runtime/Renderer/Private/PostProcess/PostProcessSubsurface.cpp b/Engine/Source/Runtime/Renderer/Private/PostProcess/PostProcessSubsurface.cpp
--- a/Engine/Source/Runtime/Renderer/Private/PostProcess/PostProcessSubsurface.cpp
+++ b/Engine/Source/Runtime/Renderer/Private/PostProcess/PostProcessSubsurface.cpp
@@ -288,27 +288,31 @@ void FRCPassPostProcessSubsurfaceVisualize::Process(FRenderingCompositePassConte
 	const FSceneRenderTargetItem& DestRenderTarget = PassOutputs[0].RequestSurface(Context);
 
 	// Set the view family's render target/viewport.
-	SetRenderTarget(Context.RHICmdList, DestRenderTarget.TargetableTexture, FTextureRHIRef());
-
-	// is optimized away if possible (RT size=view size, )
-	DrawClearQuad(Context.RHICmdList, true, FLinearColor::Black, false, 0, false, 0, PassOutputs[0].RenderTargetDesc.Extent, DestRect);
-
-	Context.SetViewportAndCallRHI(0, 0, 0.0f, DestSize.X, DestSize.Y, 1.0f );
-
-	SetSubsurfaceVisualizeShader(Context);
-
-	// Draw a quad mapping scene color to the view's render target
-	TShaderMapRef<FPostProcessVS> VertexShader(Context.GetShaderMap());
-	DrawRectangle(
-		Context.RHICmdList,
-		DestRect.Min.X, DestRect.Min.Y,
-		DestRect.Width(), DestRect.Height(),
-		SrcRect.Min.X, SrcRect.Min.Y,
-		SrcRect.Width(), SrcRect.Height(),
-		DestSize,
-		SrcSize,
-		*VertexShader,
-		EDRF_UseTriangleOptimization);
+	// #todo-renderpasses Use Clear? Load maintains previous behavior.
+	FRHIRenderPassInfo RPInfo(DestRenderTarget.TargetableTexture, ERenderTargetActions::Load_Store);
+	Context.RHICmdList.BeginRenderPass(RPInfo, TEXT("SubsurfaceVisualize"));
+	{
+		// is optimized away if possible (RT size=view size, )
+		DrawClearQuad(Context.RHICmdList, true, FLinearColor::Black, false, 0, false, 0, PassOutputs[0].RenderTargetDesc.Extent, DestRect);
+
+		Context.SetViewportAndCallRHI(0, 0, 0.0f, DestSize.X, DestSize.Y, 1.0f);
+
+		SetSubsurfaceVisualizeShader(Context);
+
+		// Draw a quad mapping scene color to the view's render target
+		TShaderMapRef<FPostProcessVS> VertexShader(Context.GetShaderMap());
+		DrawRectangle(
+			Context.RHICmdList,
+			DestRect.Min.X, DestRect.Min.Y,
+			DestRect.Width(), DestRect.Height(),
+			SrcRect.Min.X, SrcRect.Min.Y,
+			SrcRect.Width(), SrcRect.Height(),
+			DestSize,
+			SrcSize,
+			*VertexShader,
+			EDRF_UseTriangleOptimization);
+	}
+	Context.RHICmdList.EndRenderPass();
 
 	{
 		FRenderTargetTemp TempRenderTarget(View, (const FTexture2DRHIRef&)DestRenderTarget.TargetableTexture);
@@ -335,6 +339,7 @@ void FRCPassPostProcessSubsurfaceVisualize::Process(FRenderingCompositePassConte
 	}
 
 	Context.RHICmdList.CopyToResolveTarget(DestRenderTarget.TargetableTexture, DestRenderTarget.ShaderResourceTexture, FResolveParams());
+	
 
 	// we no longer need the GBuffer
 	SceneContext.AdjustGBufferRefCount(Context.RHICmdList, -1);
@@ -500,56 +505,56 @@ void FRCPassPostProcessSubsurfaceSetup::Process(FRenderingCompositePassContext&
 
 	ERenderTargetLoadAction LoadAction = Context.GetLoadActionForRenderTarget(DestRenderTarget);
 
-	FRHIRenderTargetView RtView = FRHIRenderTargetView(DestRenderTarget.TargetableTexture, LoadAction);
-	FRHISetRenderTargetsInfo Info(1, &RtView, FRHIDepthRenderTargetView());
-	Context.RHICmdList.SetRenderTargetsAndClear(Info);
-	Context.SetViewportAndCallRHI(DestRect);
-
-	if(bHalfRes)
+	FRHIRenderPassInfo RPInfo(DestRenderTarget.TargetableTexture, MakeRenderTargetActions(LoadAction, ERenderTargetStoreAction::EStore), DestRenderTarget.ShaderResourceTexture);
+	Context.RHICmdList.BeginRenderPass(RPInfo, TEXT("SubsurfaceSetup"));
 	{
-		if (bCheckerboard)
-		{
-			SetSubsurfaceSetupShader<1, 1>(Context);
-		}
-		else
-		{
-			SetSubsurfaceSetupShader<1,0>(Context);
-		}
-	}
-	else
-	{
-		if (bCheckerboard)
+		Context.SetViewportAndCallRHI(DestRect);
+
+		if (bHalfRes)
 		{
-			SetSubsurfaceSetupShader<0, 1>(Context);
+			if (bCheckerboard)
+			{
+				SetSubsurfaceSetupShader<1, 1>(Context);
+			}
+			else
+			{
+				SetSubsurfaceSetupShader<1, 0>(Context);
+			}
 		}
 		else
 		{
-			SetSubsurfaceSetupShader<0, 0>(Context);
+			if (bCheckerboard)
+			{
+				SetSubsurfaceSetupShader<0, 1>(Context);
+			}
+			else
+			{
+				SetSubsurfaceSetupShader<0, 0>(Context);
+			}
 		}
-	}
 
-	// Draw a quad mapping scene color to the view's render target
-	TShaderMapRef<FPostProcessVS> VertexShader(Context.GetShaderMap());
-
-	// Align up, so downsample always pickups correct pixel from the full-res buffer.
-	FIntPoint TargetSize = SrcRect.Size();
-	TargetSize.X = Align(TargetSize.X, ScaleFactor);
-	TargetSize.Y = Align(TargetSize.Y, ScaleFactor);
-
-	DrawPostProcessPass(
-		Context.RHICmdList,
-		0, 0,
-		SrcRect.Width(), SrcRect.Height(),
-		SrcRect.Min.X, SrcRect.Min.Y,
-		SrcRect.Width(), SrcRect.Height(),
-		TargetSize,
-		SrcSize,
-		*VertexShader,
-		View.StereoPass,
-		Context.HasHmdMesh(),
-		EDRF_UseTriangleOptimization);
-
-	Context.RHICmdList.CopyToResolveTarget(DestRenderTarget.TargetableTexture, DestRenderTarget.ShaderResourceTexture, FResolveParams());
+		// Draw a quad mapping scene color to the view's render target
+		TShaderMapRef<FPostProcessVS> VertexShader(Context.GetShaderMap());
+
+		// Align up, so downsample always pickups correct pixel from the full-res buffer.
+		FIntPoint TargetSize = SrcRect.Size();
+		TargetSize.X = Align(TargetSize.X, ScaleFactor);
+		TargetSize.Y = Align(TargetSize.Y, ScaleFactor);
+
+		DrawPostProcessPass(
+			Context.RHICmdList,
+			0, 0,
+			SrcRect.Width(), SrcRect.Height(),
+			SrcRect.Min.X, SrcRect.Min.Y,
+			SrcRect.Width(), SrcRect.Height(),
+			TargetSize,
+			SrcSize,
+			*VertexShader,
+			View.StereoPass,
+			Context.HasHmdMesh(),
+			EDRF_UseTriangleOptimization);
+	}
+	Context.RHICmdList.EndRenderPass();
 }
 
 FPooledRenderTargetDesc FRCPassPostProcessSubsurfaceSetup::ComputeOutputDesc(EPassOutputId InPassOutputId) const
@@ -759,40 +764,40 @@ void FRCPassPostProcessSubsurface::Process(FRenderingCompositePassContext& Conte
 
 	ERenderTargetLoadAction LoadAction = Context.GetLoadActionForRenderTarget(DestRenderTarget);
 
-	FRHIRenderTargetView RtView = FRHIRenderTargetView(DestRenderTarget.TargetableTexture, LoadAction);
-	FRHISetRenderTargetsInfo Info(1, &RtView, FRHIDepthRenderTargetView());
-	Context.RHICmdList.SetRenderTargetsAndClear(Info);
-	Context.SetViewportAndCallRHI(DestRect);
-
-	TShaderMapRef<FPostProcessVS> VertexShader(Context.GetShaderMap());
+	FRHIRenderPassInfo RPInfo(DestRenderTarget.TargetableTexture, MakeRenderTargetActions(LoadAction, ERenderTargetStoreAction::EStore), DestRenderTarget.ShaderResourceTexture);
+	Context.RHICmdList.BeginRenderPass(RPInfo, TEXT("Subsurface"));
+	{
+		Context.SetViewportAndCallRHI(DestRect);
 
-	SCOPED_DRAW_EVENTF(Context.RHICmdList, SubsurfacePass, TEXT("SubsurfaceDirection#%d"), Direction);
+		TShaderMapRef<FPostProcessVS> VertexShader(Context.GetShaderMap());
 
-	uint32 SampleSet = FMath::Clamp(CVarSSSSampleSet.GetValueOnRenderThread(), 0, 2);
+		SCOPED_DRAW_EVENTF(Context.RHICmdList, SubsurfacePass, TEXT("SubsurfaceDirection#%d"), Direction);
 
-	if (Direction == 0)
-	{
-		SetSubsurfaceShaderSampleSet<0>(Context, VertexShader, SampleSet);
-	}
-	else
-	{
-		SetSubsurfaceShaderSampleSet<1>(Context, VertexShader, SampleSet);
-	}
+		uint32 SampleSet = FMath::Clamp(CVarSSSSampleSet.GetValueOnRenderThread(), 0, 2);
 
-	DrawPostProcessPass(
-		Context.RHICmdList,
-		0, 0,
-		DestRect.Width(), DestRect.Height(),
-		SrcRect.Min.X, SrcRect.Min.Y,
-		SrcRect.Width(), SrcRect.Height(),
-		DestRect.Size(),
-		SrcSize,
-		*VertexShader,
-		View.StereoPass,
-		Context.HasHmdMesh(),
-		EDRF_UseTriangleOptimization);
+		if (Direction == 0)
+		{
+			SetSubsurfaceShaderSampleSet<0>(Context, VertexShader, SampleSet);
+		}
+		else
+		{
+			SetSubsurfaceShaderSampleSet<1>(Context, VertexShader, SampleSet);
+		}
 
-	Context.RHICmdList.CopyToResolveTarget(DestRenderTarget.TargetableTexture, DestRenderTarget.ShaderResourceTexture, FResolveParams());
+		DrawPostProcessPass(
+			Context.RHICmdList,
+			0, 0,
+			DestRect.Width(), DestRect.Height(),
+			SrcRect.Min.X, SrcRect.Min.Y,
+			SrcRect.Width(), SrcRect.Height(),
+			DestRect.Size(),
+			SrcSize,
+			*VertexShader,
+			View.StereoPass,
+			Context.HasHmdMesh(),
+			EDRF_UseTriangleOptimization);
+	}
+	Context.RHICmdList.EndRenderPass();
 }
 
 
@@ -948,106 +953,106 @@ void FRCPassPostProcessSubsurfaceRecombine::Process(FRenderingCompositePassConte
 
 	ERenderTargetLoadAction LoadAction = Context.GetLoadActionForRenderTarget(DestRenderTarget);
 
-	FRHIRenderTargetView RtView = FRHIRenderTargetView(DestRenderTarget.TargetableTexture, LoadAction);
-	FRHISetRenderTargetsInfo Info(1, &RtView, FRHIDepthRenderTargetView());
-	Context.RHICmdList.SetRenderTargetsAndClear(Info);
+	FRHIRenderPassInfo RPInfo(DestRenderTarget.TargetableTexture, MakeRenderTargetActions(LoadAction, ERenderTargetStoreAction::EStore), DestRenderTarget.ShaderResourceTexture);
+	Context.RHICmdList.BeginRenderPass(RPInfo, TEXT("SubsurfaceRecombine"));
+	{
 
-	CopyOverOtherViewportsIfNeeded(Context, View);
+		CopyOverOtherViewportsIfNeeded(Context, View);
 
-	Context.SetViewportAndCallRHI(DestRect);
+		Context.SetViewportAndCallRHI(DestRect);
 
-	TShaderMapRef<FPostProcessVS> VertexShader(Context.GetShaderMap());
+		TShaderMapRef<FPostProcessVS> VertexShader(Context.GetShaderMap());
 
-	uint32 QualityCVar = CVarSSSQuality.GetValueOnRenderThread();
-	const bool bCheckerboard = FRCPassPostProcessSubsurface::RequiresCheckerboardSubsurfaceRendering( SceneContext.GetSceneColorFormat() );
+		uint32 QualityCVar = CVarSSSQuality.GetValueOnRenderThread();
+		const bool bCheckerboard = FRCPassPostProcessSubsurface::RequiresCheckerboardSubsurfaceRendering(SceneContext.GetSceneColorFormat());
 
-	// 0:low / 1:high
-	uint32 RecombineQuality = 0;
-	{
-		if(QualityCVar == -1)
-		{
-			RecombineQuality = (View.AntiAliasingMethod == AAM_TemporalAA) ? 0 : 1;
-		}
-		else if(QualityCVar == 1)
+		// 0:low / 1:high
+		uint32 RecombineQuality = 0;
 		{
-			RecombineQuality = 1;
+			if (QualityCVar == -1)
+			{
+				RecombineQuality = (View.AntiAliasingMethod == AAM_TemporalAA) ? 0 : 1;
+			}
+			else if (QualityCVar == 1)
+			{
+				RecombineQuality = 1;
+			}
 		}
-	}
 
-	// needed for Scalability
-	// 0:fullres, 1: halfres, 2:no scattering, just reconstruct the lighting (needed for scalability)
-	uint32 RecombineMode = 2;
+		// needed for Scalability
+		// 0:fullres, 1: halfres, 2:no scattering, just reconstruct the lighting (needed for scalability)
+		uint32 RecombineMode = 2;
 
-	if(GetInput(ePId_Input1)->IsValid())
-	{
-		RecombineMode = bHalfRes ? 1 : 0;
-	}
+		if (GetInput(ePId_Input1)->IsValid())
+		{
+			RecombineMode = bHalfRes ? 1 : 0;
+		}
 
-	SCOPED_DRAW_EVENTF(Context.RHICmdList, SubsurfacePassRecombine, TEXT("SubsurfacePassRecombine Mode:%d Quality:%d"), RecombineMode, RecombineQuality);
+		SCOPED_DRAW_EVENTF(Context.RHICmdList, SubsurfacePassRecombine, TEXT("SubsurfacePassRecombine Mode:%d Quality:%d"), RecombineMode, RecombineQuality);
 
-	{
-		if (bCheckerboard)
 		{
-			if (RecombineQuality == 0)
+			if (bCheckerboard)
 			{
-				switch (RecombineMode)
+				if (RecombineQuality == 0)
 				{
-				case 0: SetSubsurfaceRecombineShader<0, 0, 1>(Context, VertexShader); break;
-				case 1: SetSubsurfaceRecombineShader<1, 0, 1>(Context, VertexShader); break;
-				case 2: SetSubsurfaceRecombineShader<2, 0, 1>(Context, VertexShader); break;
-				default: check(0);
+					switch (RecombineMode)
+					{
+					case 0: SetSubsurfaceRecombineShader<0, 0, 1>(Context, VertexShader); break;
+					case 1: SetSubsurfaceRecombineShader<1, 0, 1>(Context, VertexShader); break;
+					case 2: SetSubsurfaceRecombineShader<2, 0, 1>(Context, VertexShader); break;
+					default: check(0);
+					}
 				}
-			}
-			else
-			{
-				switch (RecombineMode)
+				else
 				{
-				case 0: SetSubsurfaceRecombineShader<0, 1, 1>(Context, VertexShader); break;
-				case 1: SetSubsurfaceRecombineShader<1, 1, 1>(Context, VertexShader); break;
-				case 2: SetSubsurfaceRecombineShader<2, 1, 1>(Context, VertexShader); break;
-				default: check(0);
+					switch (RecombineMode)
+					{
+					case 0: SetSubsurfaceRecombineShader<0, 1, 1>(Context, VertexShader); break;
+					case 1: SetSubsurfaceRecombineShader<1, 1, 1>(Context, VertexShader); break;
+					case 2: SetSubsurfaceRecombineShader<2, 1, 1>(Context, VertexShader); break;
+					default: check(0);
+					}
 				}
 			}
-		}
-		else
-		{
-			if (RecombineQuality == 0)
+			else
 			{
-				switch (RecombineMode)
+				if (RecombineQuality == 0)
 				{
-				case 0: SetSubsurfaceRecombineShader<0, 0, 0>(Context, VertexShader); break;
-				case 1: SetSubsurfaceRecombineShader<1, 0, 0>(Context, VertexShader); break;
-				case 2: SetSubsurfaceRecombineShader<2, 0, 0>(Context, VertexShader); break;
-				default: check(0);
+					switch (RecombineMode)
+					{
+					case 0: SetSubsurfaceRecombineShader<0, 0, 0>(Context, VertexShader); break;
+					case 1: SetSubsurfaceRecombineShader<1, 0, 0>(Context, VertexShader); break;
+					case 2: SetSubsurfaceRecombineShader<2, 0, 0>(Context, VertexShader); break;
+					default: check(0);
+					}
 				}
-			}
-			else
-			{
-				switch (RecombineMode)
+				else
 				{
-				case 0: SetSubsurfaceRecombineShader<0, 1, 0>(Context, VertexShader); break;
-				case 1: SetSubsurfaceRecombineShader<1, 1, 0>(Context, VertexShader); break;
-				case 2: SetSubsurfaceRecombineShader<2, 1, 0>(Context, VertexShader); break;
-				default: check(0);
+					switch (RecombineMode)
+					{
+					case 0: SetSubsurfaceRecombineShader<0, 1, 0>(Context, VertexShader); break;
+					case 1: SetSubsurfaceRecombineShader<1, 1, 0>(Context, VertexShader); break;
+					case 2: SetSubsurfaceRecombineShader<2, 1, 0>(Context, VertexShader); break;
+					default: check(0);
+					}
 				}
 			}
 		}
-	}
-
-	DrawPostProcessPass(
-		Context.RHICmdList,
-		0, 0,
-		DestRect.Width(), DestRect.Height(),
-		SrcRect.Min.X, SrcRect.Min.Y,
-		SrcRect.Width(), SrcRect.Height(),
-		DestRect.Size(),
-		SrcSize,
-		*VertexShader,
-		View.StereoPass,
-		Context.HasHmdMesh(),
-		EDRF_UseTriangleOptimization);
 
-	Context.RHICmdList.CopyToResolveTarget(DestRenderTarget.TargetableTexture, DestRenderTarget.ShaderResourceTexture, FResolveParams());
+		DrawPostProcessPass(
+			Context.RHICmdList,
+			0, 0,
+			DestRect.Width(), DestRect.Height(),
+			SrcRect.Min.X, SrcRect.Min.Y,
+			SrcRect.Width(), SrcRect.Height(),
+			DestRect.Size(),
+			SrcSize,
+			*VertexShader,
+			View.StereoPass,
+			Context.HasHmdMesh(),
+			EDRF_UseTriangleOptimization);
+	}
+	Context.RHICmdList.EndRenderPass();
 
 	// replace the current SceneColor with this one
 	SceneContext.SetSceneColor(PassOutputs[0].PooledRenderTarget);
