Render Passes

diff --git a/Engine/Source/Runtime/Renderer/Private/PostProcess/PostProcessTemporalAA.cpp b/Engine/Source/Runtime/Renderer/Private/PostProcess/PostProcessTemporalAA.cpp
--- a/Engine/Source/Runtime/Renderer/Private/PostProcess/PostProcessTemporalAA.cpp
+++ b/Engine/Source/Runtime/Renderer/Private/PostProcess/PostProcessTemporalAA.cpp
@@ -768,7 +768,8 @@ void FRCPassPostProcessTemporalAA::Process(FRenderingCompositePassContext& Conte
 			PassName, Parameters.bUseFast ? TEXT(" Fast") : TEXT(""), PracticableSrcRect.Width(), PracticableSrcRect.Height(), PracticableDestRect.Width(), PracticableDestRect.Height());
 
 		// Common setup
-		SetRenderTarget(Context.RHICmdList, nullptr, nullptr);
+		// #todo-renderpass remove once everything is renderpasses
+		UnbindRenderTargets(Context.RHICmdList);
 		Context.SetViewportAndCallRHI(PracticableDestRect, 0.0f, 1.0f);
 		
 		static FName AsyncEndFenceName(TEXT("AsyncTemporalAAEndFence"));
@@ -870,78 +871,79 @@ void FRCPassPostProcessTemporalAA::Process(FRenderingCompositePassContext& Conte
 		TransitionPixelPassResources(Context);
 
 		// Setup render targets.
-		{
-			FRHIDepthRenderTargetView DepthRTV(SceneContext.GetSceneDepthTexture(),
-				ERenderTargetLoadAction::ENoAction, ERenderTargetStoreAction::ENoAction,
-				ERenderTargetLoadAction::ELoad, ERenderTargetStoreAction::ENoAction,
-				FExclusiveDepthStencil::DepthRead_StencilWrite);
-
-			FRHIRenderTargetView RTVs[2];
-
-			RTVs[0] = FRHIRenderTargetView(DestRenderTarget[0]->TargetableTexture, ERenderTargetLoadAction::ENoAction);
+		
+		// Inform MultiGPU systems that we're starting to update this resource
+		Context.RHICmdList.BeginUpdateMultiFrameResource(DestRenderTarget[0]->ShaderResourceTexture);
+		FRHIRenderPassInfo RPInfo(DestRenderTarget[0]->TargetableTexture, ERenderTargetActions::DontLoad_Store, DestRenderTarget[0]->ShaderResourceTexture);
 
-			// Inform MultiGPU systems that we're starting to update this resource
-			Context.RHICmdList.BeginUpdateMultiFrameResource(DestRenderTarget[0]->ShaderResourceTexture);
+		if (RenderTargetCount == 2)
+		{
+			RPInfo.ColorRenderTargets[1].RenderTarget = DestRenderTarget[1]->TargetableTexture;
+			RPInfo.ColorRenderTargets[1].ResolveTarget = DestRenderTarget[1]->ShaderResourceTexture;
+			RPInfo.ColorRenderTargets[1].Action = ERenderTargetActions::DontLoad_Store;
+			RPInfo.ColorRenderTargets[1].ArraySlice = -1;
+			RPInfo.ColorRenderTargets[1].MipIndex = 0;
 
-			if (RenderTargetCount == 2)
-			{
-				RTVs[1] = FRHIRenderTargetView(DestRenderTarget[1]->TargetableTexture, ERenderTargetLoadAction::ENoAction);
-				Context.RHICmdList.BeginUpdateMultiFrameResource(DestRenderTarget[1]->ShaderResourceTexture);
-			}
-			Context.RHICmdList.SetRenderTargets(RenderTargetCount, RTVs, &DepthRTV, 0, nullptr);
+			Context.RHICmdList.BeginUpdateMultiFrameResource(DestRenderTarget[1]->ShaderResourceTexture);
 		}
 
-		Context.SetViewportAndCallRHI(ViewRect);
-
-		FPostProcessTemporalAAPS::FPermutationDomain PermutationVector;
-		PermutationVector.Set<FTAAPassConfigDim>(Parameters.Pass);
-		PermutationVector.Set<FTAAFastDim>(Parameters.bUseFast);
-		PermutationVector.Set<FTAACameraCutDim>(bCameraCut);
+		RPInfo.DepthStencilRenderTarget.DepthStencilTarget = SceneContext.GetSceneDepthTexture();
+		RPInfo.DepthStencilRenderTarget.ResolveTarget = nullptr;
+		RPInfo.DepthStencilRenderTarget.Action = MakeDepthStencilTargetActions(ERenderTargetActions::DontLoad_DontStore, ERenderTargetActions::Load_DontStore);
+		RPInfo.DepthStencilRenderTarget.ExclusiveDepthStencil = FExclusiveDepthStencil::DepthRead_StencilWrite;
 
-		if(bUseResponsiveStencilTest)
+		Context.RHICmdList.BeginRenderPass(RPInfo, TEXT("TemporalAA"));
 		{
-			// Normal temporal feedback
-			// Draw to pixels where stencil == 0
-			FDepthStencilStateRHIParamRef DepthStencilState = TStaticDepthStencilState<
-				false, CF_Always,
-				true, CF_Equal, SO_Keep, SO_Keep, SO_Keep,
-				false, CF_Always, SO_Keep, SO_Keep, SO_Keep,
-				STENCIL_TEMPORAL_RESPONSIVE_AA_MASK, STENCIL_TEMPORAL_RESPONSIVE_AA_MASK>::GetRHI();
-
-			DrawPixelPassTemplate(
-				Context, PermutationVector, SrcSize, ViewRect,
-				InputHistory, Parameters, bUseDither,
-				DepthStencilState);
-	
-			// Responsive feedback for tagged pixels
-			// Draw to pixels where stencil != 0
-			DepthStencilState = TStaticDepthStencilState<
-				false, CF_Always,
-				true, CF_NotEqual, SO_Keep, SO_Keep, SO_Keep,
-				false, CF_Always, SO_Keep, SO_Keep, SO_Keep,
-				STENCIL_TEMPORAL_RESPONSIVE_AA_MASK, STENCIL_TEMPORAL_RESPONSIVE_AA_MASK>::GetRHI();
-
-			PermutationVector.Set<FTAAResponsiveDim>(true);
-			DrawPixelPassTemplate(
-				Context, PermutationVector, SrcSize, ViewRect,
-				InputHistory, Parameters, bUseDither,
-				DepthStencilState);
-		}
-		else
-		{
-			DrawPixelPassTemplate(
-				Context, PermutationVector, SrcSize, ViewRect,
-				InputHistory, Parameters, bUseDither,
-				TStaticDepthStencilState<false, CF_Always>::GetRHI());
-		}
+			Context.SetViewportAndCallRHI(ViewRect);
 
-		Context.RHICmdList.CopyToResolveTarget(DestRenderTarget[0]->TargetableTexture, DestRenderTarget[0]->ShaderResourceTexture, FResolveParams());
+			FPostProcessTemporalAAPS::FPermutationDomain PermutationVector;
+			PermutationVector.Set<FTAAPassConfigDim>(Parameters.Pass);
+			PermutationVector.Set<FTAAFastDim>(Parameters.bUseFast);
+			PermutationVector.Set<FTAACameraCutDim>(bCameraCut);
 
-		if (RenderTargetCount == 2)
-		{
-			Context.RHICmdList.CopyToResolveTarget(DestRenderTarget[1]->TargetableTexture, DestRenderTarget[1]->ShaderResourceTexture, FResolveParams());
-			Context.RHICmdList.EndUpdateMultiFrameResource(DestRenderTarget[1]->ShaderResourceTexture);
+			if (bUseResponsiveStencilTest)
+			{
+				// Normal temporal feedback
+				// Draw to pixels where stencil == 0
+				FDepthStencilStateRHIParamRef DepthStencilState = TStaticDepthStencilState<
+					false, CF_Always,
+					true, CF_Equal, SO_Keep, SO_Keep, SO_Keep,
+					false, CF_Always, SO_Keep, SO_Keep, SO_Keep,
+					STENCIL_TEMPORAL_RESPONSIVE_AA_MASK, STENCIL_TEMPORAL_RESPONSIVE_AA_MASK>::GetRHI();
+
+				DrawPixelPassTemplate(
+					Context, PermutationVector, SrcSize, ViewRect,
+					InputHistory, Parameters, bUseDither,
+					DepthStencilState);
+
+				// Responsive feedback for tagged pixels
+				// Draw to pixels where stencil != 0
+				DepthStencilState = TStaticDepthStencilState<
+					false, CF_Always,
+					true, CF_NotEqual, SO_Keep, SO_Keep, SO_Keep,
+					false, CF_Always, SO_Keep, SO_Keep, SO_Keep,
+					STENCIL_TEMPORAL_RESPONSIVE_AA_MASK, STENCIL_TEMPORAL_RESPONSIVE_AA_MASK>::GetRHI();
+
+				PermutationVector.Set<FTAAResponsiveDim>(true);
+				DrawPixelPassTemplate(
+					Context, PermutationVector, SrcSize, ViewRect,
+					InputHistory, Parameters, bUseDither,
+					DepthStencilState);
+			}
+			else
+			{
+				DrawPixelPassTemplate(
+					Context, PermutationVector, SrcSize, ViewRect,
+					InputHistory, Parameters, bUseDither,
+					TStaticDepthStencilState<false, CF_Always>::GetRHI());
+			}
+
+			if (RenderTargetCount == 2)
+			{
+				Context.RHICmdList.EndUpdateMultiFrameResource(DestRenderTarget[1]->ShaderResourceTexture);
+			}
 		}
+		Context.RHICmdList.EndRenderPass();
 
 		if (IsDOFTAAConfig(Parameters.Pass))
 		{
diff --git a/Engine/Source/Runtime/Renderer/Private/PostProcess/PostProcessTestImage.cpp b/Engine/Source/Runtime/Renderer/Private/PostProcess/PostProcessTestImage.cpp
--- a/Engine/Source/Runtime/Renderer/Private/PostProcess/PostProcessTestImage.cpp
+++ b/Engine/Source/Runtime/Renderer/Private/PostProcess/PostProcessTestImage.cpp
@@ -103,38 +103,43 @@ void FRCPassPostProcessTestImage::Process(FRenderingCompositePassContext& Contex
 	const FSceneRenderTargetItem& DestRenderTarget = PassOutputs[0].RequestSurface(Context);
 
 	// Set the view family's render target/viewport.
-	SetRenderTarget(Context.RHICmdList, DestRenderTarget.TargetableTexture, FTextureRHIRef());
-	Context.SetViewportAndCallRHI(DestRect);
-
-	FGraphicsPipelineStateInitializer GraphicsPSOInit;
-	Context.RHICmdList.ApplyCachedRenderTargets(GraphicsPSOInit);
-	GraphicsPSOInit.BlendState = TStaticBlendState<>::GetRHI();
-	GraphicsPSOInit.RasterizerState = TStaticRasterizerState<>::GetRHI();
-	GraphicsPSOInit.DepthStencilState = TStaticDepthStencilState<false, CF_Always>::GetRHI();
-
-	TShaderMapRef<FPostProcessVS> VertexShader(Context.GetShaderMap());
-	TShaderMapRef<FPostProcessTestImagePS> PixelShader(Context.GetShaderMap());
-
-	GraphicsPSOInit.BoundShaderState.VertexDeclarationRHI = GFilterVertexDeclaration.VertexDeclarationRHI;
-	GraphicsPSOInit.BoundShaderState.VertexShaderRHI = GETSAFERHISHADER_VERTEX(*VertexShader);
-	GraphicsPSOInit.BoundShaderState.PixelShaderRHI = GETSAFERHISHADER_PIXEL(*PixelShader);
-	GraphicsPSOInit.PrimitiveType = PT_TriangleList;
-
-	SetGraphicsPipelineState(Context.RHICmdList, GraphicsPSOInit);
-
-	PixelShader->SetPS(Context.RHICmdList, Context);
-
-	// Draw a quad mapping scene color to the view's render target
-	DrawRectangle(
-		Context.RHICmdList,
-		0, 0,
-		DestRect.Width(), DestRect.Height(),
-		SrcRect.Min.X, SrcRect.Min.Y,
-		SrcRect.Width(), SrcRect.Height(),
-		DestRect.Size(),
-		FSceneRenderTargets::Get(Context.RHICmdList).GetBufferSizeXY(),
-		*VertexShader,
-		EDRF_UseTriangleOptimization);
+	FRHIRenderPassInfo RPInfo(DestRenderTarget.TargetableTexture, ERenderTargetActions::Load_Store);
+	Context.RHICmdList.BeginRenderPass(RPInfo, TEXT("TestImage"));
+	{
+		Context.SetViewportAndCallRHI(DestRect);
+
+		FGraphicsPipelineStateInitializer GraphicsPSOInit;
+		Context.RHICmdList.ApplyCachedRenderTargets(GraphicsPSOInit);
+		GraphicsPSOInit.BlendState = TStaticBlendState<>::GetRHI();
+		GraphicsPSOInit.RasterizerState = TStaticRasterizerState<>::GetRHI();
+		GraphicsPSOInit.DepthStencilState = TStaticDepthStencilState<false, CF_Always>::GetRHI();
+
+		TShaderMapRef<FPostProcessVS> VertexShader(Context.GetShaderMap());
+		TShaderMapRef<FPostProcessTestImagePS> PixelShader(Context.GetShaderMap());
+
+		GraphicsPSOInit.BoundShaderState.VertexDeclarationRHI = GFilterVertexDeclaration.VertexDeclarationRHI;
+		GraphicsPSOInit.BoundShaderState.VertexShaderRHI = GETSAFERHISHADER_VERTEX(*VertexShader);
+		GraphicsPSOInit.BoundShaderState.PixelShaderRHI = GETSAFERHISHADER_PIXEL(*PixelShader);
+		GraphicsPSOInit.PrimitiveType = PT_TriangleList;
+
+		SetGraphicsPipelineState(Context.RHICmdList, GraphicsPSOInit);
+
+		PixelShader->SetPS(Context.RHICmdList, Context);
+
+		// Draw a quad mapping scene color to the view's render target
+		DrawRectangle(
+			Context.RHICmdList,
+			0, 0,
+			DestRect.Width(), DestRect.Height(),
+			SrcRect.Min.X, SrcRect.Min.Y,
+			SrcRect.Width(), SrcRect.Height(),
+			DestRect.Size(),
+			FSceneRenderTargets::Get(Context.RHICmdList).GetBufferSizeXY(),
+			*VertexShader,
+			EDRF_UseTriangleOptimization);
+
+	}
+	Context.RHICmdList.EndRenderPass();
 
 	{
 		FRenderTargetTemp TempRenderTarget(Context.View, (const FTexture2DRHIRef&)DestRenderTarget.TargetableTexture);
@@ -148,43 +153,42 @@ void FRCPassPostProcessTestImage::Process(FRenderingCompositePassContext& Contex
 		FString Line;
 
 		Line = FString::Printf(TEXT("Top bars:"));
-		Canvas.DrawShadowedString( X, Y += YStep, *Line, GetStatsFont(), FLinearColor(1, 1, 1));
+		Canvas.DrawShadowedString(X, Y += YStep, *Line, GetStatsFont(), FLinearColor(1, 1, 1));
 		Line = FString::Printf(TEXT("   Moving bars using FrameTime"));
-		Canvas.DrawShadowedString( X, Y += YStep, *Line, GetStatsFont(), FLinearColor(1, 1, 1));
+		Canvas.DrawShadowedString(X, Y += YStep, *Line, GetStatsFont(), FLinearColor(1, 1, 1));
 		Line = FString::Printf(TEXT("   Black and white raster, Pixel sized, Watch for Moire pattern"));
-		Canvas.DrawShadowedString( X, Y += YStep, *Line, GetStatsFont(), FLinearColor(1, 1, 1));
+		Canvas.DrawShadowedString(X, Y += YStep, *Line, GetStatsFont(), FLinearColor(1, 1, 1));
 		Line = FString::Printf(TEXT("   Black and white raster, 2x2 block sized"));
-		Canvas.DrawShadowedString( X, Y += YStep, *Line, GetStatsFont(), FLinearColor(1, 1, 1));
+		Canvas.DrawShadowedString(X, Y += YStep, *Line, GetStatsFont(), FLinearColor(1, 1, 1));
 		Line = FString::Printf(TEXT("Bottom bars:"));
-		Canvas.DrawShadowedString( X, Y += YStep, *Line, GetStatsFont(), FLinearColor(1, 1, 1));
+		Canvas.DrawShadowedString(X, Y += YStep, *Line, GetStatsFont(), FLinearColor(1, 1, 1));
 		Line = FString::Printf(TEXT("   8 bars near white, 4 right bars should appear as one (HDTV)"));
-		Canvas.DrawShadowedString( X, Y += YStep, *Line, GetStatsFont(), FLinearColor(1, 1, 1));
+		Canvas.DrawShadowedString(X, Y += YStep, *Line, GetStatsFont(), FLinearColor(1, 1, 1));
 		Line = FString::Printf(TEXT("   8 bars near black, 4 left bars should appear as one (HDTV)"));
-		Canvas.DrawShadowedString( X, Y += YStep, *Line, GetStatsFont(), FLinearColor(1, 1, 1));
+		Canvas.DrawShadowedString(X, Y += YStep, *Line, GetStatsFont(), FLinearColor(1, 1, 1));
 		Line = FString::Printf(TEXT("   Linear Greyscale in sRGB from 0 to 255"));
-		Canvas.DrawShadowedString( X, Y += YStep, *Line, GetStatsFont(), FLinearColor(1, 1, 1));
+		Canvas.DrawShadowedString(X, Y += YStep, *Line, GetStatsFont(), FLinearColor(1, 1, 1));
 		Line = FString::Printf(TEXT("Color bars:"));
-		Canvas.DrawShadowedString( X, Y += YStep, *Line, GetStatsFont(), FLinearColor(1, 1, 1));
+		Canvas.DrawShadowedString(X, Y += YStep, *Line, GetStatsFont(), FLinearColor(1, 1, 1));
 		Line = FString::Printf(TEXT("   Red, Green, Blue"));
-		Canvas.DrawShadowedString( X, Y += YStep, *Line, GetStatsFont(), FLinearColor(1, 1, 1));
+		Canvas.DrawShadowedString(X, Y += YStep, *Line, GetStatsFont(), FLinearColor(1, 1, 1));
 		Line = FString::Printf(TEXT("Outside:"));
-		Canvas.DrawShadowedString( X, Y += YStep, *Line, GetStatsFont(), FLinearColor(1, 1, 1));
+		Canvas.DrawShadowedString(X, Y += YStep, *Line, GetStatsFont(), FLinearColor(1, 1, 1));
 		Line = FString::Printf(TEXT("   Moving bars using FrameNumber, Tearing without VSync"));
-		Canvas.DrawShadowedString( X, Y += YStep, *Line, GetStatsFont(), FLinearColor(1, 1, 1));
+		Canvas.DrawShadowedString(X, Y += YStep, *Line, GetStatsFont(), FLinearColor(1, 1, 1));
 		Line = FString::Printf(TEXT("Circles:"));
-		Canvas.DrawShadowedString( X, Y += YStep, *Line, GetStatsFont(), FLinearColor(1, 1, 1));
+		Canvas.DrawShadowedString(X, Y += YStep, *Line, GetStatsFont(), FLinearColor(1, 1, 1));
 		Line = FString::Printf(TEXT("   Should be round and centered"));
-		Canvas.DrawShadowedString( X, Y += YStep, *Line, GetStatsFont(), FLinearColor(1, 1, 1));
+		Canvas.DrawShadowedString(X, Y += YStep, *Line, GetStatsFont(), FLinearColor(1, 1, 1));
 		Line = FString::Printf(TEXT("Border:"));
-		Canvas.DrawShadowedString( X, Y += YStep, *Line, GetStatsFont(), FLinearColor(1, 1, 1));
+		Canvas.DrawShadowedString(X, Y += YStep, *Line, GetStatsFont(), FLinearColor(1, 1, 1));
 		Line = FString::Printf(TEXT("   4 white pixel sized lines (only visible without overscan)"));
-		Canvas.DrawShadowedString( X, Y += YStep, *Line, GetStatsFont(), FLinearColor(1, 1, 1));
+		Canvas.DrawShadowedString(X, Y += YStep, *Line, GetStatsFont(), FLinearColor(1, 1, 1));
 
 		Canvas.Flush_RenderThread(Context.RHICmdList);
 	}
-
-
 	Context.RHICmdList.CopyToResolveTarget(DestRenderTarget.TargetableTexture, DestRenderTarget.ShaderResourceTexture, FResolveParams());
+	
 }
 
 FPooledRenderTargetDesc FRCPassPostProcessTestImage::ComputeOutputDesc(EPassOutputId InPassOutputId) const
diff --git a/Engine/Source/Runtime/Renderer/Private/PostProcess/PostProcessTonemap.cpp b/Engine/Source/Runtime/Renderer/Private/PostProcess/PostProcessTonemap.cpp
--- a/Engine/Source/Runtime/Renderer/Private/PostProcess/PostProcessTonemap.cpp
+++ b/Engine/Source/Runtime/Renderer/Private/PostProcess/PostProcessTonemap.cpp
@@ -999,7 +999,8 @@ void FRCPassPostProcessTonemap::Process(FRenderingCompositePassContext& Context)
 		DestRect = { DestRect.Min, DestRect.Min + DestSize };
 
 		// Common setup
-		SetRenderTarget(Context.RHICmdList, nullptr, nullptr);
+		// #todo-renderpass remove once everything is renderpasses
+		UnbindRenderTargets(Context.RHICmdList);
 		Context.SetViewportAndCallRHI(DestRect, 0.0f, 1.0f);
 
 		static FName AsyncEndFenceName(TEXT("AsyncTonemapEndFence"));
@@ -1049,6 +1050,8 @@ void FRCPassPostProcessTonemap::Process(FRenderingCompositePassContext& Context)
 			PostProcessTonemapUtil::ShaderTransitionResources<false>(Context);
 		}
 
+		ERenderTargetLoadAction LoadAction = ERenderTargetLoadAction::ELoad;
+
 		if (IsMobilePlatform(ShaderPlatform))
 		{
 			// clear target when processing first view in case of splitscreen
@@ -1057,84 +1060,79 @@ void FRCPassPostProcessTonemap::Process(FRenderingCompositePassContext& Context)
 			// Full clear to avoid restore
 			if ((View.StereoPass == eSSP_FULL && bFirstView) || View.StereoPass == eSSP_LEFT_EYE)
 			{
-				SetRenderTarget(Context.RHICmdList, DestRenderTarget.TargetableTexture, FTextureRHIParamRef(), ESimpleRenderTargetMode::EClearColorAndDepth);
-			}
-			else
-			{
-				SetRenderTarget(Context.RHICmdList, DestRenderTarget.TargetableTexture, FTextureRHIParamRef());
+				LoadAction = ERenderTargetLoadAction::EClear;
 			}
 		}
 		else
 		{
-			ERenderTargetLoadAction LoadAction = Context.GetLoadActionForRenderTarget(DestRenderTarget);
+			LoadAction = Context.GetLoadActionForRenderTarget(DestRenderTarget);
 			if (Context.View.AntiAliasingMethod == AAM_FXAA)
 			{
 				check(LoadAction != ERenderTargetLoadAction::ELoad);
 				// needed to not have PostProcessAA leaking in content (e.g. Matinee black borders).
 				LoadAction = ERenderTargetLoadAction::EClear;
 			}
-
-			FRHIRenderTargetView RtView = FRHIRenderTargetView(DestRenderTarget.TargetableTexture, LoadAction);
-			FRHISetRenderTargetsInfo Info(1, &RtView, FRHIDepthRenderTargetView());
-			Context.RHICmdList.SetRenderTargetsAndClear(Info);
 		}
 
-		Context.SetViewportAndCallRHI(DestRect, 0.0f, 1.0f);
-
-		FShader* VertexShader;
+		FRHIRenderPassInfo RPInfo(DestRenderTarget.TargetableTexture, MakeRenderTargetActions(LoadAction, ERenderTargetStoreAction::EStore), DestRenderTarget.ShaderResourceTexture);
+		Context.RHICmdList.BeginRenderPass(RPInfo, TEXT("Tonemap"));
 		{
-			FGraphicsPipelineStateInitializer GraphicsPSOInit;
-			Context.RHICmdList.ApplyCachedRenderTargets(GraphicsPSOInit);
-			GraphicsPSOInit.BlendState = TStaticBlendState<>::GetRHI();
-			GraphicsPSOInit.RasterizerState = TStaticRasterizerState<>::GetRHI();
-			GraphicsPSOInit.DepthStencilState = TStaticDepthStencilState<false, CF_Always>::GetRHI();
+			Context.SetViewportAndCallRHI(DestRect, 0.0f, 1.0f);
 
-			if (bDoEyeAdaptation)
-			{
-				VertexShader = Context.GetShaderMap()->GetShader<TPostProcessTonemapVS<true>>();
-			}
-			else
+			FShader* VertexShader;
 			{
-				VertexShader = Context.GetShaderMap()->GetShader<TPostProcessTonemapVS<false>>();
-			}
+				FGraphicsPipelineStateInitializer GraphicsPSOInit;
+				Context.RHICmdList.ApplyCachedRenderTargets(GraphicsPSOInit);
+				GraphicsPSOInit.BlendState = TStaticBlendState<>::GetRHI();
+				GraphicsPSOInit.RasterizerState = TStaticRasterizerState<>::GetRHI();
+				GraphicsPSOInit.DepthStencilState = TStaticDepthStencilState<false, CF_Always>::GetRHI();
 
-			TShaderMapRef<FPostProcessTonemapPS> PixelShader(Context.GetShaderMap(), DesktopPermutationVector);
+				if (bDoEyeAdaptation)
+				{
+					VertexShader = Context.GetShaderMap()->GetShader<TPostProcessTonemapVS<true>>();
+				}
+				else
+				{
+					VertexShader = Context.GetShaderMap()->GetShader<TPostProcessTonemapVS<false>>();
+				}
 
-			GraphicsPSOInit.BoundShaderState.VertexDeclarationRHI = GFilterVertexDeclaration.VertexDeclarationRHI;
-			GraphicsPSOInit.BoundShaderState.VertexShaderRHI = VertexShader->GetVertexShader();
-			GraphicsPSOInit.BoundShaderState.PixelShaderRHI = GETSAFERHISHADER_PIXEL(*PixelShader);
-			GraphicsPSOInit.PrimitiveType = PT_TriangleList;
+				TShaderMapRef<FPostProcessTonemapPS> PixelShader(Context.GetShaderMap(), DesktopPermutationVector);
 
-			SetGraphicsPipelineState(Context.RHICmdList, GraphicsPSOInit);
+				GraphicsPSOInit.BoundShaderState.VertexDeclarationRHI = GFilterVertexDeclaration.VertexDeclarationRHI;
+				GraphicsPSOInit.BoundShaderState.VertexShaderRHI = VertexShader->GetVertexShader();
+				GraphicsPSOInit.BoundShaderState.PixelShaderRHI = GETSAFERHISHADER_PIXEL(*PixelShader);
+				GraphicsPSOInit.PrimitiveType = PT_TriangleList;
 
-			if (bDoEyeAdaptation)
-			{
-				TShaderMapRef<TPostProcessTonemapVS<true>> VertexShaderMapRef(Context.GetShaderMap());
-				VertexShaderMapRef->SetVS(Context);
-			}
-			else
-			{
-				TShaderMapRef<TPostProcessTonemapVS<false>> VertexShaderMapRef(Context.GetShaderMap());
-				VertexShaderMapRef->SetVS(Context);
-			}
+				SetGraphicsPipelineState(Context.RHICmdList, GraphicsPSOInit);
 
-			PixelShader->SetPS(Context);
-		}
+				if (bDoEyeAdaptation)
+				{
+					TShaderMapRef<TPostProcessTonemapVS<true>> VertexShaderMapRef(Context.GetShaderMap());
+					VertexShaderMapRef->SetVS(Context);
+				}
+				else
+				{
+					TShaderMapRef<TPostProcessTonemapVS<false>> VertexShaderMapRef(Context.GetShaderMap());
+					VertexShaderMapRef->SetVS(Context);
+				}
 
-		DrawPostProcessPass(
-			Context.RHICmdList,
-			0, 0,
-			DestRect.Width(), DestRect.Height(),
-			SrcRect.Min.X, SrcRect.Min.Y,
-			SrcRect.Width(), SrcRect.Height(),
-			DestRect.Size(),
-			SrcSize,
-			VertexShader,
-			View.StereoPass,
-			Context.HasHmdMesh(),
-			EDRF_UseTriangleOptimization);
+				PixelShader->SetPS(Context);
+			}
 
-		Context.RHICmdList.CopyToResolveTarget(DestRenderTarget.TargetableTexture, DestRenderTarget.ShaderResourceTexture, FResolveParams());
+			DrawPostProcessPass(
+				Context.RHICmdList,
+				0, 0,
+				DestRect.Width(), DestRect.Height(),
+				SrcRect.Min.X, SrcRect.Min.Y,
+				SrcRect.Width(), SrcRect.Height(),
+				DestRect.Size(),
+				SrcSize,
+				VertexShader,
+				View.StereoPass,
+				Context.HasHmdMesh(),
+				EDRF_UseTriangleOptimization);
+		}
+		Context.RHICmdList.EndRenderPass();
 
 		// We only release the SceneColor after the last view was processed (SplitScreen)
 		if(Context.View.Family->Views[Context.View.Family->Views.Num() - 1] == &Context.View && !GIsEditor)
@@ -1655,6 +1653,8 @@ void FRCPassPostProcessTonemapES2::Process(FRenderingCompositePassContext& Conte
 	FIntPoint SrcSize = InputDesc->Extent;
 	FIntPoint DstSize = OutputDesc.Extent;
 
+	ERenderTargetLoadAction LoadAction = ERenderTargetLoadAction::ELoad;
+
 	// Set the view family's render target/viewport.
 	{
 		// clear target when processing first view in case of splitscreen
@@ -1663,53 +1663,52 @@ void FRCPassPostProcessTonemapES2::Process(FRenderingCompositePassContext& Conte
 		// Full clear to avoid restore
 		if ((View.StereoPass == eSSP_FULL && bFirstView) || View.StereoPass == eSSP_LEFT_EYE)
 		{
-			SetRenderTarget(Context.RHICmdList, DestRenderTarget.TargetableTexture, FTextureRHIParamRef(), ESimpleRenderTargetMode::EClearColorAndDepth);
-		}
-		else
-		{
-			SetRenderTarget(Context.RHICmdList, DestRenderTarget.TargetableTexture, FTextureRHIParamRef());
+			LoadAction = ERenderTargetLoadAction::EClear;
 		}
 	}
 
-	Context.SetViewportAndCallRHI(DestRect);
+	FRHIRenderPassInfo RPInfo(DestRenderTarget.TargetableTexture, MakeRenderTargetActions(LoadAction, ERenderTargetStoreAction::EStore), DestRenderTarget.ShaderResourceTexture);
+	Context.RHICmdList.BeginRenderPass(RPInfo, TEXT("TonemapES2"));
+	{
+		Context.SetViewportAndCallRHI(DestRect);
 
-	auto PermutationVector = FPostProcessTonemapPS_ES2::BuildPermutationVector(View);
+		auto PermutationVector = FPostProcessTonemapPS_ES2::BuildPermutationVector(View);
 
-	TShaderMapRef<FPostProcessTonemapVS_ES2> VertexShader(Context.GetShaderMap());
-	TShaderMapRef<FPostProcessTonemapPS_ES2> PixelShader(Context.GetShaderMap(), PermutationVector);
+		TShaderMapRef<FPostProcessTonemapVS_ES2> VertexShader(Context.GetShaderMap());
+		TShaderMapRef<FPostProcessTonemapPS_ES2> PixelShader(Context.GetShaderMap(), PermutationVector);
 
-	{
-		FGraphicsPipelineStateInitializer GraphicsPSOInit;
-		Context.RHICmdList.ApplyCachedRenderTargets(GraphicsPSOInit);
-		GraphicsPSOInit.BlendState = TStaticBlendState<>::GetRHI();
-		GraphicsPSOInit.RasterizerState = TStaticRasterizerState<>::GetRHI();
-		GraphicsPSOInit.DepthStencilState = TStaticDepthStencilState<false, CF_Always>::GetRHI();
+		{
+			FGraphicsPipelineStateInitializer GraphicsPSOInit;
+			Context.RHICmdList.ApplyCachedRenderTargets(GraphicsPSOInit);
+			GraphicsPSOInit.BlendState = TStaticBlendState<>::GetRHI();
+			GraphicsPSOInit.RasterizerState = TStaticRasterizerState<>::GetRHI();
+			GraphicsPSOInit.DepthStencilState = TStaticDepthStencilState<false, CF_Always>::GetRHI();
 
-		VertexShader->bUsedFramebufferFetch = bUsedFramebufferFetch;
+			VertexShader->bUsedFramebufferFetch = bUsedFramebufferFetch;
 
-		GraphicsPSOInit.BoundShaderState.VertexDeclarationRHI = GFilterVertexDeclaration.VertexDeclarationRHI;
-		GraphicsPSOInit.BoundShaderState.VertexShaderRHI = GETSAFERHISHADER_VERTEX(*VertexShader);
-		GraphicsPSOInit.BoundShaderState.PixelShaderRHI = GETSAFERHISHADER_PIXEL(*PixelShader);
-		GraphicsPSOInit.PrimitiveType = PT_TriangleList;
+			GraphicsPSOInit.BoundShaderState.VertexDeclarationRHI = GFilterVertexDeclaration.VertexDeclarationRHI;
+			GraphicsPSOInit.BoundShaderState.VertexShaderRHI = GETSAFERHISHADER_VERTEX(*VertexShader);
+			GraphicsPSOInit.BoundShaderState.PixelShaderRHI = GETSAFERHISHADER_PIXEL(*PixelShader);
+			GraphicsPSOInit.PrimitiveType = PT_TriangleList;
 
-		SetGraphicsPipelineState(Context.RHICmdList, GraphicsPSOInit);
+			SetGraphicsPipelineState(Context.RHICmdList, GraphicsPSOInit);
 
-		VertexShader->SetVS(Context);
-		PixelShader->SetPS(Context.RHICmdList, Context, PermutationVector, bSRGBAwareTarget);
-	}
+			VertexShader->SetVS(Context);
+			PixelShader->SetPS(Context.RHICmdList, Context, PermutationVector, bSRGBAwareTarget);
+		}
 
-	DrawRectangle(
-		Context.RHICmdList,
-		0, 0,
-		DstSize.X, DstSize.Y,
-		SrcRect.Min.X, SrcRect.Min.Y,
-		SrcRect.Width(), SrcRect.Height(),
-		DstSize,
-		SrcSize,
-		*VertexShader,
-		EDRF_UseTriangleOptimization);
-
-	Context.RHICmdList.CopyToResolveTarget(DestRenderTarget.TargetableTexture, DestRenderTarget.ShaderResourceTexture, FResolveParams());
+		DrawRectangle(
+			Context.RHICmdList,
+			0, 0,
+			DstSize.X, DstSize.Y,
+			SrcRect.Min.X, SrcRect.Min.Y,
+			SrcRect.Width(), SrcRect.Height(),
+			DstSize,
+			SrcSize,
+			*VertexShader,
+			EDRF_UseTriangleOptimization);
+	}
+	Context.RHICmdList.EndRenderPass();
 }
 
 FPooledRenderTargetDesc FRCPassPostProcessTonemapES2::ComputeOutputDesc(EPassOutputId InPassOutputId) const
diff --git a/Engine/Source/Runtime/Renderer/Private/PostProcess/PostProcessUpscale.cpp b/Engine/Source/Runtime/Renderer/Private/PostProcess/PostProcessUpscale.cpp
--- a/Engine/Source/Runtime/Renderer/Private/PostProcess/PostProcessUpscale.cpp
+++ b/Engine/Source/Runtime/Renderer/Private/PostProcess/PostProcessUpscale.cpp
@@ -398,36 +398,37 @@ void FRCPassPostProcessUpscale::Process(FRenderingCompositePassContext& Context)
 	FIntPoint SrcSize = InputDesc->Extent;
 
 	bool bMobilePlatform = IsMobilePlatform(GShaderPlatformForFeatureLevel[Context.GetFeatureLevel()]);
+
+	ERenderTargetLoadAction LoadAction = ERenderTargetLoadAction::ELoad;
+
 	if (bMobilePlatform)
 	{
 		// clear on mobile to avoid restore
-		SetRenderTarget(Context.RHICmdList, DestRenderTarget.TargetableTexture, FTextureRHIRef(), ESimpleRenderTargetMode::EClearColorAndDepth);
+		LoadAction = ERenderTargetLoadAction::EClear;
 	}
-	else
+
+	FRHIRenderPassInfo RPInfo(DestRenderTarget.TargetableTexture, MakeRenderTargetActions(LoadAction, ERenderTargetStoreAction::EStore), DestRenderTarget.ShaderResourceTexture);
+	Context.RHICmdList.BeginRenderPass(RPInfo, TEXT("Upscale"));
 	{
-		// Set the view family's render target/viewport.
-		SetRenderTarget(Context.RHICmdList, DestRenderTarget.TargetableTexture, FTextureRHIRef());
-	}
-	
-	bool bTessellatedQuad = PaniniConfig.D >= 0.01f;
+		bool bTessellatedQuad = PaniniConfig.D >= 0.01f;
 
-	// with distortion (bTessellatedQuad) we need to clear the background
-	FIntRect ExcludeRect = bTessellatedQuad ? FIntRect() : DestRect;
+		// with distortion (bTessellatedQuad) we need to clear the background
+		FIntRect ExcludeRect = bTessellatedQuad ? FIntRect() : DestRect;
 
-	Context.SetViewportAndCallRHI(DestRect);
-	if (Context.GetLoadActionForRenderTarget(DestRenderTarget) == ERenderTargetLoadAction::EClear && !bMobilePlatform)
-	{
-		DrawClearQuad(Context.RHICmdList, true, FLinearColor::Black, false, 0, false, 0, PassOutputs[0].RenderTargetDesc.Extent, ExcludeRect);
-	}
+		Context.SetViewportAndCallRHI(DestRect);
+		if (Context.GetLoadActionForRenderTarget(DestRenderTarget) == ERenderTargetLoadAction::EClear && !bMobilePlatform)
+		{
+			DrawClearQuad(Context.RHICmdList, true, FLinearColor::Black, false, 0, false, 0, PassOutputs[0].RenderTargetDesc.Extent, ExcludeRect);
+		}
 
-	FShader* VertexShader = 0;
+		FShader* VertexShader = 0;
 
-	bool ManullyClampUV = SrcRect.Min != FIntPoint::ZeroValue || SrcRect.Max != SrcSize;
+		bool ManullyClampUV = SrcRect.Min != FIntPoint::ZeroValue || SrcRect.Max != SrcSize;
 
-	if(bTessellatedQuad)
-	{
-		switch (UpscaleQuality)
+		if (bTessellatedQuad)
 		{
+			switch (UpscaleQuality)
+			{
 			case 0:	VertexShader = SetShader<0, 1>(Context, PaniniConfig, false); break;
 			case 1:	VertexShader = SetShader<1, 1>(Context, PaniniConfig, ManullyClampUV); break;
 			case 2:	VertexShader = SetShader<2, 1>(Context, PaniniConfig, ManullyClampUV); break;
@@ -438,12 +439,12 @@ void FRCPassPostProcessUpscale::Process(FRenderingCompositePassContext& Context)
 			default:
 				checkNoEntry();
 				break;
+			}
 		}
-	}
-	else
-	{
-		switch (UpscaleQuality)
+		else
 		{
+			switch (UpscaleQuality)
+			{
 			case 0:	VertexShader = SetShader<0, 0>(Context, PaniniParams::Default, false); break;
 			case 1:	VertexShader = SetShader<1, 0>(Context, PaniniParams::Default, ManullyClampUV); break;
 			case 2:	VertexShader = SetShader<2, 0>(Context, PaniniParams::Default, ManullyClampUV); break;
@@ -454,22 +455,22 @@ void FRCPassPostProcessUpscale::Process(FRenderingCompositePassContext& Context)
 			default:
 				checkNoEntry();
 				break;
+			}
 		}
-	}
 
-	// Draw a quad, a triangle or a tessellated quad
-	DrawRectangle(
-		Context.RHICmdList,
-		0, 0,
-		DestRect.Width(), DestRect.Height(),
-		SrcRect.Min.X, SrcRect.Min.Y,
-		SrcRect.Width(), SrcRect.Height(),
-		DestRect.Size(),
-		SrcSize,
-		VertexShader,
-		bTessellatedQuad ? EDRF_UseTesselatedIndexBuffer: EDRF_UseTriangleOptimization);
-
-	Context.RHICmdList.CopyToResolveTarget(DestRenderTarget.TargetableTexture, DestRenderTarget.ShaderResourceTexture, FResolveParams());
+		// Draw a quad, a triangle or a tessellated quad
+		DrawRectangle(
+			Context.RHICmdList,
+			0, 0,
+			DestRect.Width(), DestRect.Height(),
+			SrcRect.Min.X, SrcRect.Min.Y,
+			SrcRect.Width(), SrcRect.Height(),
+			DestRect.Size(),
+			SrcSize,
+			VertexShader,
+			bTessellatedQuad ? EDRF_UseTesselatedIndexBuffer : EDRF_UseTriangleOptimization);
+	}
+	Context.RHICmdList.EndRenderPass();
 
 	// Update scene color view rectangle for secondary upscale.
 	Context.SceneColorViewRect = DestRect;
diff --git a/Engine/Source/Runtime/Renderer/Private/PostProcess/PostProcessVisualizeBuffer.cpp b/Engine/Source/Runtime/Renderer/Private/PostProcess/PostProcessVisualizeBuffer.cpp
--- a/Engine/Source/Runtime/Renderer/Private/PostProcess/PostProcessVisualizeBuffer.cpp
+++ b/Engine/Source/Runtime/Renderer/Private/PostProcess/PostProcessVisualizeBuffer.cpp
@@ -155,46 +155,6 @@ void FRCPassPostProcessVisualizeBuffer::Process(FRenderingCompositePassContext&
 	FIntRect DestRect = View.ViewRect;
 	FIntPoint SrcSize = InputDesc->Extent;
 
-	const FSceneRenderTargetItem& DestRenderTarget = PassOutputs[0].RequestSurface(Context);
-
-	// Set the view family's render target/viewport.
-	SetRenderTarget(Context.RHICmdList, DestRenderTarget.TargetableTexture, FTextureRHIRef());
-	Context.SetViewportAndCallRHI(DestRect);
-
-	{
-		FShader* VertexShader = SetShaderTempl<false>(Context);
-
-	// Draw a quad mapping scene color to the view's render target
-	DrawRectangle(
-		Context.RHICmdList,
-		0, 0,
-		DestRect.Width(), DestRect.Height(),
-		SrcRect.Min.X, SrcRect.Min.Y,
-		SrcRect.Width(), SrcRect.Height(),
-		DestRect.Size(),
-		SrcSize,
-			VertexShader,
-		EDRF_UseTriangleOptimization);
-	}
-
-	FGraphicsPipelineStateInitializer GraphicsPSOInit;
-	Context.RHICmdList.ApplyCachedRenderTargets(GraphicsPSOInit);
-	GraphicsPSOInit.BlendState = TStaticBlendState<CW_RGB, BO_Add, BF_SourceAlpha, BF_InverseSourceAlpha>::GetRHI();
-	GraphicsPSOInit.RasterizerState = TStaticRasterizerState<>::GetRHI();
-	GraphicsPSOInit.DepthStencilState = TStaticDepthStencilState<false, CF_Always>::GetRHI();
-
-	TShaderMapRef<FPostProcessVS> VertexShader(Context.GetShaderMap());
-	TShaderMapRef<FPostProcessVisualizeBufferPS<true> > PixelShader(Context.GetShaderMap());
-
-	GraphicsPSOInit.BoundShaderState.VertexDeclarationRHI = GFilterVertexDeclaration.VertexDeclarationRHI;
-	GraphicsPSOInit.BoundShaderState.VertexShaderRHI = GETSAFERHISHADER_VERTEX(*VertexShader);
-	GraphicsPSOInit.BoundShaderState.PixelShaderRHI = GETSAFERHISHADER_PIXEL(*PixelShader);
-	GraphicsPSOInit.PrimitiveType = PT_TriangleList;
-
-	SetGraphicsPipelineState(Context.RHICmdList, GraphicsPSOInit);
-
-	PixelShader->SetPS(Context.RHICmdList, Context);
-
 	// Track the name and position of each tile we draw so we can write text labels over them
 	struct LabelRecord
 	{
@@ -202,49 +162,91 @@ void FRCPassPostProcessVisualizeBuffer::Process(FRenderingCompositePassContext&
 		int32 LocationX;
 		int32 LocationY;
 	};
-
 	TArray<LabelRecord> Labels;
 
-	const int32 MaxTilesX = 4;
-	const int32 MaxTilesY = 4;
-	const int32 TileWidth = DestRect.Width() / MaxTilesX;
-	const int32 TileHeight = DestRect.Height() / MaxTilesY;
-	int32 CurrentTileIndex = 0; 
+	const FSceneRenderTargetItem& DestRenderTarget = PassOutputs[0].RequestSurface(Context);
 
-	for (TArray<TileData>::TConstIterator It = Tiles.CreateConstIterator(); It; ++It, ++CurrentTileIndex)
+	// Set the view family's render target/viewport.
+	FRHIRenderPassInfo RPInfo(DestRenderTarget.TargetableTexture, ERenderTargetActions::Load_Store);
+	Context.RHICmdList.BeginRenderPass(RPInfo, TEXT("VisualizeBuffer"));
 	{
-		FRenderingCompositeOutputRef Tile = It->Source;
+		Context.SetViewportAndCallRHI(DestRect);
 
-		if (Tile.IsValid())
 		{
-			FTextureRHIRef Texture = Tile.GetOutput()->PooledRenderTarget->GetRenderTargetItem().TargetableTexture;
-
-			int32 TileX = CurrentTileIndex % MaxTilesX;
-			int32 TileY = CurrentTileIndex / MaxTilesX;
-
-			PixelShader->SetSourceTexture(Context.RHICmdList, Texture);
+			FShader* VertexShader = SetShaderTempl<false>(Context);
 
+			// Draw a quad mapping scene color to the view's render target
 			DrawRectangle(
 				Context.RHICmdList,
-				TileX * TileWidth, TileY * TileHeight,
-				TileWidth, TileHeight,
+				0, 0,
+				DestRect.Width(), DestRect.Height(),
 				SrcRect.Min.X, SrcRect.Min.Y,
 				SrcRect.Width(), SrcRect.Height(),
 				DestRect.Size(),
 				SrcSize,
-				*VertexShader,
-				EDRF_Default
+				VertexShader,
+				EDRF_UseTriangleOptimization);
+		}
+
+		FGraphicsPipelineStateInitializer GraphicsPSOInit;
+		Context.RHICmdList.ApplyCachedRenderTargets(GraphicsPSOInit);
+		GraphicsPSOInit.BlendState = TStaticBlendState<CW_RGB, BO_Add, BF_SourceAlpha, BF_InverseSourceAlpha>::GetRHI();
+		GraphicsPSOInit.RasterizerState = TStaticRasterizerState<>::GetRHI();
+		GraphicsPSOInit.DepthStencilState = TStaticDepthStencilState<false, CF_Always>::GetRHI();
+
+		TShaderMapRef<FPostProcessVS> VertexShader(Context.GetShaderMap());
+		TShaderMapRef<FPostProcessVisualizeBufferPS<true> > PixelShader(Context.GetShaderMap());
+
+		GraphicsPSOInit.BoundShaderState.VertexDeclarationRHI = GFilterVertexDeclaration.VertexDeclarationRHI;
+		GraphicsPSOInit.BoundShaderState.VertexShaderRHI = GETSAFERHISHADER_VERTEX(*VertexShader);
+		GraphicsPSOInit.BoundShaderState.PixelShaderRHI = GETSAFERHISHADER_PIXEL(*PixelShader);
+		GraphicsPSOInit.PrimitiveType = PT_TriangleList;
+
+		SetGraphicsPipelineState(Context.RHICmdList, GraphicsPSOInit);
+
+		PixelShader->SetPS(Context.RHICmdList, Context);
+
+		const int32 MaxTilesX = 4;
+		const int32 MaxTilesY = 4;
+		const int32 TileWidth = DestRect.Width() / MaxTilesX;
+		const int32 TileHeight = DestRect.Height() / MaxTilesY;
+		int32 CurrentTileIndex = 0;
+
+		for (TArray<TileData>::TConstIterator It = Tiles.CreateConstIterator(); It; ++It, ++CurrentTileIndex)
+		{
+			FRenderingCompositeOutputRef Tile = It->Source;
+
+			if (Tile.IsValid())
+			{
+				FTextureRHIRef Texture = Tile.GetOutput()->PooledRenderTarget->GetRenderTargetItem().TargetableTexture;
+
+				int32 TileX = CurrentTileIndex % MaxTilesX;
+				int32 TileY = CurrentTileIndex / MaxTilesX;
+
+				PixelShader->SetSourceTexture(Context.RHICmdList, Texture);
+
+				DrawRectangle(
+					Context.RHICmdList,
+					TileX * TileWidth, TileY * TileHeight,
+					TileWidth, TileHeight,
+					SrcRect.Min.X, SrcRect.Min.Y,
+					SrcRect.Width(), SrcRect.Height(),
+					DestRect.Size(),
+					SrcSize,
+					*VertexShader,
+					EDRF_Default
 				);
 
-			Labels.Add(LabelRecord());
-			Labels.Last().Label = It->Name;
-			Labels.Last().LocationX = 8 + TileX * TileWidth;
-			Labels.Last().LocationY = (TileY + 1) * TileHeight - 19;
+				Labels.Add(LabelRecord());
+				Labels.Last().Label = It->Name;
+				Labels.Last().LocationX = 8 + TileX * TileWidth;
+				Labels.Last().LocationY = (TileY + 1) * TileHeight - 19;
+			}
 		}
 	}
+	Context.RHICmdList.EndRenderPass();
 
 	// Draw tile labels
-
 	FRenderTargetTemp TempRenderTarget(View, (const FTexture2DRHIRef&)DestRenderTarget.TargetableTexture);
 	FCanvas Canvas(&TempRenderTarget, NULL, ViewFamily.CurrentRealTime, ViewFamily.CurrentWorldTime, ViewFamily.DeltaWorldTime, Context.GetFeatureLevel());
 	FLinearColor LabelColor(1, 1, 0);
@@ -254,8 +256,8 @@ void FRCPassPostProcessVisualizeBuffer::Process(FRenderingCompositePassContext&
 	}
 	Canvas.Flush_RenderThread(Context.RHICmdList);
 
-
 	Context.RHICmdList.CopyToResolveTarget(DestRenderTarget.TargetableTexture, DestRenderTarget.ShaderResourceTexture, FResolveParams());
+	
 }
 
 FPooledRenderTargetDesc FRCPassPostProcessVisualizeBuffer::ComputeOutputDesc(EPassOutputId InPassOutputId) const
diff --git a/Engine/Source/Runtime/Renderer/Private/PostProcess/PostProcessVisualizeComplexity.cpp b/Engine/Source/Runtime/Renderer/Private/PostProcess/PostProcessVisualizeComplexity.cpp
--- a/Engine/Source/Runtime/Renderer/Private/PostProcess/PostProcessVisualizeComplexity.cpp
+++ b/Engine/Source/Runtime/Renderer/Private/PostProcess/PostProcessVisualizeComplexity.cpp
@@ -135,50 +135,54 @@ void FRCPassPostProcessVisualizeComplexity::Process(FRenderingCompositePassConte
 	const FSceneRenderTargetItem& DestRenderTarget = PassOutputs[0].RequestSurface(Context);
 
 	// Set the view family's render target/viewport.
-	SetRenderTarget(Context.RHICmdList, DestRenderTarget.TargetableTexture, FTextureRHIRef());
-	Context.SetViewportAndCallRHI(DestRect);
+	FRHIRenderPassInfo RPInfo(DestRenderTarget.TargetableTexture, ERenderTargetActions::Load_Store);
+	Context.RHICmdList.BeginRenderPass(RPInfo, TEXT("VisualizeComplexity"));
+	{
+		Context.SetViewportAndCallRHI(DestRect);
 
-	FGraphicsPipelineStateInitializer GraphicsPSOInit;
-	Context.RHICmdList.ApplyCachedRenderTargets(GraphicsPSOInit);
+		FGraphicsPipelineStateInitializer GraphicsPSOInit;
+		Context.RHICmdList.ApplyCachedRenderTargets(GraphicsPSOInit);
 
-	// turn off culling and blending
-	GraphicsPSOInit.RasterizerState = TStaticRasterizerState<FM_Solid, CM_None>::GetRHI();
-	GraphicsPSOInit.BlendState = TStaticBlendState<>::GetRHI();
+		// turn off culling and blending
+		GraphicsPSOInit.RasterizerState = TStaticRasterizerState<FM_Solid, CM_None>::GetRHI();
+		GraphicsPSOInit.BlendState = TStaticBlendState<>::GetRHI();
 
-	// turn off depth reads/writes
-	GraphicsPSOInit.DepthStencilState = TStaticDepthStencilState<false, CF_Always>::GetRHI();
+		// turn off depth reads/writes
+		GraphicsPSOInit.DepthStencilState = TStaticDepthStencilState<false, CF_Always>::GetRHI();
 
-	//reuse this generic vertex shader
-	TShaderMapRef<FPostProcessVS> VertexShader(Context.GetShaderMap());
-	TShaderMapRef<FVisualizeComplexityApplyPS> PixelShader(Context.GetShaderMap());
+		//reuse this generic vertex shader
+		TShaderMapRef<FPostProcessVS> VertexShader(Context.GetShaderMap());
+		TShaderMapRef<FVisualizeComplexityApplyPS> PixelShader(Context.GetShaderMap());
 
-	GraphicsPSOInit.BoundShaderState.VertexDeclarationRHI = GFilterVertexDeclaration.VertexDeclarationRHI;
-	GraphicsPSOInit.BoundShaderState.VertexShaderRHI = GETSAFERHISHADER_VERTEX(*VertexShader);
-	GraphicsPSOInit.BoundShaderState.PixelShaderRHI = GETSAFERHISHADER_PIXEL(*PixelShader);
-	GraphicsPSOInit.PrimitiveType = PT_TriangleList;
+		GraphicsPSOInit.BoundShaderState.VertexDeclarationRHI = GFilterVertexDeclaration.VertexDeclarationRHI;
+		GraphicsPSOInit.BoundShaderState.VertexShaderRHI = GETSAFERHISHADER_VERTEX(*VertexShader);
+		GraphicsPSOInit.BoundShaderState.PixelShaderRHI = GETSAFERHISHADER_PIXEL(*PixelShader);
+		GraphicsPSOInit.PrimitiveType = PT_TriangleList;
 
-	SetGraphicsPipelineState(Context.RHICmdList, GraphicsPSOInit);
+		SetGraphicsPipelineState(Context.RHICmdList, GraphicsPSOInit);
 
-	PixelShader->SetParameters(Context.RHICmdList, Context, Colors, ColorSampling, ComplexityScale, bLegend);
-	
-	DrawRectangle(
-		Context.RHICmdList,
-		0, 0,
-		DestRect.Width(), DestRect.Height(),
-		SrcRect.Min.X, SrcRect.Min.Y,
-		SrcRect.Width(), SrcRect.Height(),
-		DestRect.Size(),
-		SrcSize,
-		*VertexShader,
-		EDRF_UseTriangleOptimization);
-
-	if(bLegend)
+		PixelShader->SetParameters(Context.RHICmdList, Context, Colors, ColorSampling, ComplexityScale, bLegend);
+
+		DrawRectangle(
+			Context.RHICmdList,
+			0, 0,
+			DestRect.Width(), DestRect.Height(),
+			SrcRect.Min.X, SrcRect.Min.Y,
+			SrcRect.Width(), SrcRect.Height(),
+			DestRect.Size(),
+			SrcSize,
+			*VertexShader,
+			EDRF_UseTriangleOptimization);
+	}
+	Context.RHICmdList.EndRenderPass();
+
+	if (bLegend)
 	{
 		FRenderTargetTemp TempRenderTarget(View, (const FTexture2DRHIRef&)DestRenderTarget.TargetableTexture);
 		FCanvas Canvas(&TempRenderTarget, NULL, ViewFamily.CurrentRealTime, ViewFamily.CurrentWorldTime, ViewFamily.DeltaWorldTime, Context.GetFeatureLevel());
 
-//later?		Canvas.DrawShadowedString(DestRect.Max.X - DestRect.Width() / 3 - 64 + 8, DestRect.Max.Y - 80, TEXT("Overdraw"), GetStatsFont(), FLinearColor(0.7f, 0.7f, 0.7f), FLinearColor(0,0,0,0));
-//later?		Canvas.DrawShadowedString(DestRect.Min.X + 64 + 4, DestRect.Max.Y - 80, TEXT("VS Instructions"), GetStatsFont(), FLinearColor(0.0f, 0.0f, 0.0f), FLinearColor(0,0,0,0));
+		//later?		Canvas.DrawShadowedString(DestRect.Max.X - DestRect.Width() / 3 - 64 + 8, DestRect.Max.Y - 80, TEXT("Overdraw"), GetStatsFont(), FLinearColor(0.7f, 0.7f, 0.7f), FLinearColor(0,0,0,0));
+		//later?		Canvas.DrawShadowedString(DestRect.Min.X + 64 + 4, DestRect.Max.Y - 80, TEXT("VS Instructions"), GetStatsFont(), FLinearColor(0.0f, 0.0f, 0.0f), FLinearColor(0,0,0,0));
 
 		if (View.Family->GetDebugViewShaderMode() == DVSM_QuadComplexity)
 		{
@@ -207,7 +211,7 @@ void FRCPassPostProcessVisualizeComplexity::Process(FRenderingCompositePassConte
 
 		Canvas.Flush_RenderThread(Context.RHICmdList);
 	}
-	
+
 	Context.RHICmdList.CopyToResolveTarget(DestRenderTarget.TargetableTexture, DestRenderTarget.ShaderResourceTexture, FResolveParams());
 }
 
diff --git a/Engine/Source/Runtime/Renderer/Private/PostProcess/PostProcessVisualizeHDR.cpp b/Engine/Source/Runtime/Renderer/Private/PostProcess/PostProcessVisualizeHDR.cpp
--- a/Engine/Source/Runtime/Renderer/Private/PostProcess/PostProcessVisualizeHDR.cpp
+++ b/Engine/Source/Runtime/Renderer/Private/PostProcess/PostProcessVisualizeHDR.cpp
@@ -233,38 +233,42 @@ void FRCPassPostProcessVisualizeHDR::Process(FRenderingCompositePassContext& Con
 	FIntPoint SrcSize = InputDesc->Extent;
 
 	// Set the view family's render target/viewport.
-	SetRenderTarget(Context.RHICmdList, DestRenderTarget.TargetableTexture, FTextureRHIRef());
-	Context.SetViewportAndCallRHI(DestRect);
-
-	FGraphicsPipelineStateInitializer GraphicsPSOInit;
-	Context.RHICmdList.ApplyCachedRenderTargets(GraphicsPSOInit);
-	GraphicsPSOInit.BlendState = TStaticBlendState<>::GetRHI();
-	GraphicsPSOInit.RasterizerState = TStaticRasterizerState<>::GetRHI();
-	GraphicsPSOInit.DepthStencilState = TStaticDepthStencilState<false, CF_Always>::GetRHI();
-
-	TShaderMapRef<FPostProcessVS> VertexShader(Context.GetShaderMap());
-	TShaderMapRef<FPostProcessVisualizeHDRPS> PixelShader(Context.GetShaderMap());
-
-	GraphicsPSOInit.BoundShaderState.VertexDeclarationRHI = GFilterVertexDeclaration.VertexDeclarationRHI;
-	GraphicsPSOInit.BoundShaderState.VertexShaderRHI = GETSAFERHISHADER_VERTEX(*VertexShader);
-	GraphicsPSOInit.BoundShaderState.PixelShaderRHI = GETSAFERHISHADER_PIXEL(*PixelShader);
-	GraphicsPSOInit.PrimitiveType = PT_TriangleList;
-
-	SetGraphicsPipelineState(Context.RHICmdList, GraphicsPSOInit);
-
-	PixelShader->SetPS(Context.RHICmdList, Context);
-
-	// Draw a quad mapping scene color to the view's render target
-	DrawRectangle(
-		Context.RHICmdList,
-		0, 0,
-		DestRect.Width(), DestRect.Height(),
-		SrcRect.Min.X, SrcRect.Min.Y,
-		SrcRect.Width(), SrcRect.Height(),
-		DestRect.Size(),
-		SrcSize,
-		*VertexShader,
-		EDRF_UseTriangleOptimization);
+	FRHIRenderPassInfo RPInfo(DestRenderTarget.TargetableTexture, ERenderTargetActions::Load_Store);
+	Context.RHICmdList.BeginRenderPass(RPInfo, TEXT("VisualizeHDR"));
+	{
+		Context.SetViewportAndCallRHI(DestRect);
+
+		FGraphicsPipelineStateInitializer GraphicsPSOInit;
+		Context.RHICmdList.ApplyCachedRenderTargets(GraphicsPSOInit);
+		GraphicsPSOInit.BlendState = TStaticBlendState<>::GetRHI();
+		GraphicsPSOInit.RasterizerState = TStaticRasterizerState<>::GetRHI();
+		GraphicsPSOInit.DepthStencilState = TStaticDepthStencilState<false, CF_Always>::GetRHI();
+
+		TShaderMapRef<FPostProcessVS> VertexShader(Context.GetShaderMap());
+		TShaderMapRef<FPostProcessVisualizeHDRPS> PixelShader(Context.GetShaderMap());
+
+		GraphicsPSOInit.BoundShaderState.VertexDeclarationRHI = GFilterVertexDeclaration.VertexDeclarationRHI;
+		GraphicsPSOInit.BoundShaderState.VertexShaderRHI = GETSAFERHISHADER_VERTEX(*VertexShader);
+		GraphicsPSOInit.BoundShaderState.PixelShaderRHI = GETSAFERHISHADER_PIXEL(*PixelShader);
+		GraphicsPSOInit.PrimitiveType = PT_TriangleList;
+
+		SetGraphicsPipelineState(Context.RHICmdList, GraphicsPSOInit);
+
+		PixelShader->SetPS(Context.RHICmdList, Context);
+
+		// Draw a quad mapping scene color to the view's render target
+		DrawRectangle(
+			Context.RHICmdList,
+			0, 0,
+			DestRect.Width(), DestRect.Height(),
+			SrcRect.Min.X, SrcRect.Min.Y,
+			SrcRect.Width(), SrcRect.Height(),
+			DestRect.Size(),
+			SrcSize,
+			*VertexShader,
+			EDRF_UseTriangleOptimization);
+	}
+	Context.RHICmdList.EndRenderPass();
 
 	FRenderTargetTemp TempRenderTarget(View, (const FTexture2DRHIRef&)DestRenderTarget.TargetableTexture);
 	FCanvas Canvas(&TempRenderTarget, NULL, ViewFamily.CurrentRealTime, ViewFamily.CurrentWorldTime, ViewFamily.DeltaWorldTime, Context.GetFeatureLevel());
@@ -277,33 +281,33 @@ void FRCPassPostProcessVisualizeHDR::Process(FRenderingCompositePassContext& Con
 	FString Line;
 
 	Line = FString::Printf(TEXT("HDR Histogram (EV100, max of RGB)"));
-	Canvas.DrawShadowedString( X, Y += YStep, *Line, GetStatsFont(), FLinearColor(1, 1, 1));
-	
+	Canvas.DrawShadowedString(X, Y += YStep, *Line, GetStatsFont(), FLinearColor(1, 1, 1));
+
 	Y += 160;
 
 	float MinX = DestRect.Min.X + 64 + 10;
 	float MaxY = DestRect.Max.Y - 64;
 	float SizeX = DestRect.Size().X - 64 * 2 - 20;
 
-	for(uint32 i = 0; i <= 4; ++i)
+	for (uint32 i = 0; i <= 4; ++i)
 	{
 		int XAdd = (int)(i * SizeX / 4);
-		float HistogramPosition =  i / 4.0f;
+		float HistogramPosition = i / 4.0f;
 		float EV100Value = FMath::Lerp(View.FinalPostProcessSettings.HistogramLogMin, View.FinalPostProcessSettings.HistogramLogMax, HistogramPosition);
 		if (!bExtendedLuminanceRange)
-		{	
+		{
 			// In this case the post process settings are actually Log2 values.
 			EV100Value = Log2ToEV100(EV100Value);
 		}
 
 		Line = FString::Printf(TEXT("%.2g"), EV100Value);
-		Canvas.DrawShadowedString( MinX + XAdd - 5, MaxY + YStep, *Line, GetStatsFont(), FLinearColor(1, 0.3f, 0.3f));
+		Canvas.DrawShadowedString(MinX + XAdd - 5, MaxY + YStep, *Line, GetStatsFont(), FLinearColor(1, 0.3f, 0.3f));
 	}
 	Y += 3 * YStep;
 	if (AutoExposureMethod == EAutoExposureMethod::AEM_Basic)
 	{
 		Line = FString::Printf(TEXT("Basic"));
-	} 
+	}
 	else
 	{
 		Line = FString::Printf(TEXT("Histogram"));
@@ -312,8 +316,8 @@ void FRCPassPostProcessVisualizeHDR::Process(FRenderingCompositePassContext& Con
 	Canvas.DrawShadowedString(X + ColumnWidth, Y, *Line, GetStatsFont(), FLinearColor(1, 1, 1));
 
 	Line = FString::Printf(TEXT("%g%% .. %g%%"), View.FinalPostProcessSettings.AutoExposureLowPercent, View.FinalPostProcessSettings.AutoExposureHighPercent);
-	Canvas.DrawShadowedString( X, Y += YStep, TEXT("Percent Low/High:"), GetStatsFont(), FLinearColor(1, 1, 1));
-	Canvas.DrawShadowedString( X + ColumnWidth, Y, *Line, GetStatsFont(), FLinearColor(1, 1, 1));
+	Canvas.DrawShadowedString(X, Y += YStep, TEXT("Percent Low/High:"), GetStatsFont(), FLinearColor(1, 1, 1));
+	Canvas.DrawShadowedString(X + ColumnWidth, Y, *Line, GetStatsFont(), FLinearColor(1, 1, 1));
 
 	if (bExtendedLuminanceRange)
 	{
@@ -323,22 +327,22 @@ void FRCPassPostProcessVisualizeHDR::Process(FRenderingCompositePassContext& Con
 	{
 		Line = FString::Printf(TEXT("%.1f .. %.1f"), LuminanceToEV100(View.FinalPostProcessSettings.AutoExposureMinBrightness), LuminanceToEV100(View.FinalPostProcessSettings.AutoExposureMaxBrightness));
 	}
-	Canvas.DrawShadowedString( X, Y += YStep, TEXT("EV100 Min/Max"), GetStatsFont(), FLinearColor(1, 1, 1));
-	Canvas.DrawShadowedString( X + ColumnWidth, Y, *Line, GetStatsFont(), FLinearColor(0.3f, 0.3f, 1));
+	Canvas.DrawShadowedString(X, Y += YStep, TEXT("EV100 Min/Max"), GetStatsFont(), FLinearColor(1, 1, 1));
+	Canvas.DrawShadowedString(X + ColumnWidth, Y, *Line, GetStatsFont(), FLinearColor(0.3f, 0.3f, 1));
 
 	Line = FString::Printf(TEXT("%g / %g"), View.FinalPostProcessSettings.AutoExposureSpeedUp, View.FinalPostProcessSettings.AutoExposureSpeedDown);
-	Canvas.DrawShadowedString( X, Y += YStep, TEXT("Speed Up/Down:"), GetStatsFont(), FLinearColor(1, 1, 1));
-	Canvas.DrawShadowedString( X + ColumnWidth, Y, *Line, GetStatsFont(), FLinearColor(1, 1, 1));
+	Canvas.DrawShadowedString(X, Y += YStep, TEXT("Speed Up/Down:"), GetStatsFont(), FLinearColor(1, 1, 1));
+	Canvas.DrawShadowedString(X + ColumnWidth, Y, *Line, GetStatsFont(), FLinearColor(1, 1, 1));
 
 	Line = FString::Printf(TEXT("%.2g"), View.FinalPostProcessSettings.AutoExposureBias);
-	Canvas.DrawShadowedString( X, Y += YStep, TEXT("Exposure Compensation: "), GetStatsFont(), FLinearColor(1, 1, 1));
-	Canvas.DrawShadowedString( X + ColumnWidth, Y, *Line, GetStatsFont(), FLinearColor(1, 0.3f, 0.3f));
+	Canvas.DrawShadowedString(X, Y += YStep, TEXT("Exposure Compensation: "), GetStatsFont(), FLinearColor(1, 1, 1));
+	Canvas.DrawShadowedString(X + ColumnWidth, Y, *Line, GetStatsFont(), FLinearColor(1, 0.3f, 0.3f));
 
 	if (AutoExposureMethod == AEM_Basic || AutoExposureMethod == AEM_Histogram)
 	{
 		Line = FString::Printf(TEXT("%g%%"), AutoExposureMethod == AEM_Basic ? View.FinalPostProcessSettings.AutoExposureCalibrationConstant : 100.f);
-		Canvas.DrawShadowedString( X, Y += YStep, TEXT("Calibration Constant: "), GetStatsFont(), FLinearColor(1, 1, 1));
-		Canvas.DrawShadowedString( X + ColumnWidth, Y, *Line, GetStatsFont(), FLinearColor(1, 0.3f, 0.3f));
+		Canvas.DrawShadowedString(X, Y += YStep, TEXT("Calibration Constant: "), GetStatsFont(), FLinearColor(1, 1, 1));
+		Canvas.DrawShadowedString(X + ColumnWidth, Y, *Line, GetStatsFont(), FLinearColor(1, 0.3f, 0.3f));
 	}
 
 	if (bExtendedLuminanceRange)
@@ -350,8 +354,8 @@ void FRCPassPostProcessVisualizeHDR::Process(FRenderingCompositePassContext& Con
 		Line = FString::Printf(TEXT("%.1f .. %.1f"), Log2ToEV100(View.FinalPostProcessSettings.HistogramLogMin), Log2ToEV100(View.FinalPostProcessSettings.HistogramLogMax));
 	}
 
-	Canvas.DrawShadowedString( X, Y += YStep, TEXT("Histogram EV100 Min/Max:"), GetStatsFont(), FLinearColor(1, 1, 1));
-	Canvas.DrawShadowedString( X + ColumnWidth, Y, *Line, GetStatsFont(), FLinearColor(0.3f, 0.3f, 1));
+	Canvas.DrawShadowedString(X, Y += YStep, TEXT("Histogram EV100 Min/Max:"), GetStatsFont(), FLinearColor(1, 1, 1));
+	Canvas.DrawShadowedString(X + ColumnWidth, Y, *Line, GetStatsFont(), FLinearColor(0.3f, 0.3f, 1));
 
 	if (AutoExposureMethod == EAutoExposureMethod::AEM_Basic)
 	{
diff --git a/Engine/Source/Runtime/Renderer/Private/PostProcess/PostProcessWeightedSampleSum.cpp b/Engine/Source/Runtime/Renderer/Private/PostProcess/PostProcessWeightedSampleSum.cpp
--- a/Engine/Source/Runtime/Renderer/Private/PostProcess/PostProcessWeightedSampleSum.cpp
+++ b/Engine/Source/Runtime/Renderer/Private/PostProcess/PostProcessWeightedSampleSum.cpp
@@ -998,7 +998,8 @@ void FRCPassPostProcessWeightedSampleSum::Process(FRenderingCompositePassContext
 		DestRect = {Context.SceneColorViewRect.Min, Context.SceneColorViewRect.Min + DestSize};
 
 		// Common setup
-		SetRenderTarget(Context.RHICmdList, nullptr, nullptr);
+		// #todo-renderpasses remove once everything is renderpasses
+		UnbindRenderTargets(Context.RHICmdList);
 		Context.SetViewportAndCallRHI(DestRect, 0.0f, 1.0f);
 
 		static FName AsyncEndFenceName(TEXT("AsyncWeightedSampleSumEndFence"));
@@ -1037,45 +1038,45 @@ void FRCPassPostProcessWeightedSampleSum::Process(FRenderingCompositePassContext
 		bool bRequiresClear = true;
 		// check if we have to clear the whole surface.
 		// Otherwise perform the clear when the dest rectangle has been computed.
+		ERenderTargetLoadAction LoadAction = ERenderTargetLoadAction::ELoad;
 		if (FeatureLevel == ERHIFeatureLevel::ES2 || FeatureLevel == ERHIFeatureLevel::ES3_1)
 		{
 			bRequiresClear = false;
-			SetRenderTarget(Context.RHICmdList, DestRenderTarget.TargetableTexture, FTextureRHIRef(), ESimpleRenderTargetMode::EClearColorAndDepth);
+			LoadAction = ERenderTargetLoadAction::EClear;
 		}
-		else
+
+		FRHIRenderPassInfo RPInfo(DestRenderTarget.TargetableTexture, MakeRenderTargetActions(LoadAction, ERenderTargetStoreAction::EStore), DestRenderTarget.ShaderResourceTexture);
+		Context.RHICmdList.BeginRenderPass(RPInfo, TEXT("WeightedSampleSum"));
 		{
-			SetRenderTarget(Context.RHICmdList, DestRenderTarget.TargetableTexture, FTextureRHIRef(), ESimpleRenderTargetMode::EExistingColorAndDepth);
-		}
+			Context.SetViewportAndCallRHI(0, 0, 0.0f, DestSize.X, DestSize.Y, 1.0f);
 
-		Context.SetViewportAndCallRHI(0, 0, 0.0f, DestSize.X, DestSize.Y, 1.0f);
+			FIntRect SrcRect = FIntRect::DivideAndRoundUp(Context.SceneColorViewRect, SrcScaleFactor);
+			if (bRequiresClear)
+			{
+				DrawClear(Context.RHICmdList, FeatureLevel, bDoFastBlur, SrcRect, DestRect, DestSize);
+			}
 
-		FIntRect SrcRect = FIntRect::DivideAndRoundUp(Context.SceneColorViewRect, SrcScaleFactor);
-		if (bRequiresClear)
-		{
-			DrawClear(Context.RHICmdList, FeatureLevel, bDoFastBlur, SrcRect, DestRect, DestSize);
+			FShader* VertexShader = nullptr;
+			SetFilterShaders(
+				Context.RHICmdList,
+				FeatureLevel,
+				TStaticSamplerState<SF_Bilinear, AM_Border, AM_Border, AM_Clamp>::GetRHI(),
+				FilterTexture,
+				AdditiveTexture,
+				CombineMethodInt,
+				BlurOffsets,
+				BlurWeights,
+				NumSamples,
+				&VertexShader,
+				SrcRect,
+				SrcSize,
+				AdditiveSrcRect,
+				AdditiveSrcSize
+			);
+
+			DrawQuad(Context.RHICmdList, FeatureLevel, bDoFastBlur, SrcRect, DestRect, DestSize, SrcSize, VertexShader);
 		}
-
-		FShader* VertexShader = nullptr;
-		SetFilterShaders(
-			Context.RHICmdList,
-			FeatureLevel,
-			TStaticSamplerState<SF_Bilinear, AM_Border, AM_Border, AM_Clamp>::GetRHI(),
-			FilterTexture,
-			AdditiveTexture,
-			CombineMethodInt,
-			BlurOffsets,
-			BlurWeights,
-			NumSamples,
-			&VertexShader,
-			SrcRect,
-			SrcSize,
-			AdditiveSrcRect,
-			AdditiveSrcSize
-		);
-
-		DrawQuad(Context.RHICmdList, FeatureLevel, bDoFastBlur, SrcRect, DestRect, DestSize, SrcSize, VertexShader);
-
-		Context.RHICmdList.CopyToResolveTarget(DestRenderTarget.TargetableTexture, DestRenderTarget.ShaderResourceTexture, FResolveParams());
+		Context.RHICmdList.EndRenderPass();
 	}
 }
 
