RenderPasses prep.
Adding FinishRenderingSceneColor. Temporarily a noop until later passes are committed.
When the heavy passes are checked in this will end the open scene color renderpass.

diff --git a/Engine/Source/Runtime/Renderer/Private/PostProcess/SceneRenderTargets.cpp b/Engine/Source/Runtime/Renderer/Private/PostProcess/SceneRenderTargets.cpp
--- a/Engine/Source/Runtime/Renderer/Private/PostProcess/SceneRenderTargets.cpp
+++ b/Engine/Source/Runtime/Renderer/Private/PostProcess/SceneRenderTargets.cpp
@@ -618,6 +618,11 @@ void FSceneRenderTargets::BeginRenderingSceneColor(FRHICommandList& RHICmdList,
 	SetRenderTarget(RHICmdList, GetSceneColorSurface(), GetSceneDepthSurface(), RenderTargetMode, DepthStencilAccess, bTransitionWritable);
 } 
 
+void FSceneRenderTargets::FinishRenderingSceneColor(FRHICommandList& RHICmdList)
+{
+	// does nothing until renderpasses are all checked in
+}
+
 void FSceneRenderTargets::BeginRenderingSceneMonoColor(FRHICommandList& RHICmdList, ESimpleRenderTargetMode RenderTargetMode, FExclusiveDepthStencil DepthStencilAccess)
 {
 	SCOPED_DRAW_EVENT(RHICmdList, BeginRenderingSceneMonoColor);
diff --git a/Engine/Source/Runtime/Renderer/Private/PostProcess/SceneRenderTargets.h b/Engine/Source/Runtime/Renderer/Private/PostProcess/SceneRenderTargets.h
--- a/Engine/Source/Runtime/Renderer/Private/PostProcess/SceneRenderTargets.h
+++ b/Engine/Source/Runtime/Renderer/Private/PostProcess/SceneRenderTargets.h
@@ -212,6 +212,7 @@ public:
 	 * Sets the scene color target and restores its contents if necessary
 	 */
 	void BeginRenderingSceneColor(FRHICommandList& FRHICommandListImmediate, ESimpleRenderTargetMode RenderTargetMode = ESimpleRenderTargetMode::EUninitializedColorExistingDepth, FExclusiveDepthStencil DepthStencilAccess = FExclusiveDepthStencil::DepthWrite_StencilWrite, bool bTransitionWritable = true);
+	void FinishRenderingSceneColor(FRHICommandList& RHICmdList);
 
 	void BeginRenderingSceneMonoColor(FRHICommandList& RHICmdList, ESimpleRenderTargetMode RenderTargetMode = ESimpleRenderTargetMode::EUninitializedColorExistingDepth, FExclusiveDepthStencil DepthStencilAccess = FExclusiveDepthStencil::DepthWrite_StencilWrite);
 
