AtmosphereRendering RenderPass updates.
It does call BeginRenderingSceneColor but is contained. The Begin/FinishRenderingSceneColor entry points will be switched to use Renderpasses later.

diff --git a/Engine/Source/Runtime/Renderer/Private/AtmosphereRendering.cpp b/Engine/Source/Runtime/Renderer/Private/AtmosphereRendering.cpp
--- a/Engine/Source/Runtime/Renderer/Private/AtmosphereRendering.cpp
+++ b/Engine/Source/Runtime/Renderer/Private/AtmosphereRendering.cpp
@@ -422,6 +422,8 @@ void FDeferredShadingSceneRenderer::RenderAtmosphere(FRHICommandListImmediate& R
 			RHICmdList.SetStreamSource(0, GScreenSpaceVertexBuffer.VertexBufferRHI, 0);
 			RHICmdList.DrawIndexedPrimitive(GTwoTrianglesIndexBuffer.IndexBufferRHI, 0, 0, 4, 0, 2, 1);
 		}
+
+		SceneContext.FinishRenderingSceneColor(RHICmdList);
 	}
 }
 
@@ -1106,95 +1108,104 @@ void FAtmosphericFogSceneInfo::RenderAtmosphereShaders(FRHICommandList& RHICmdLi
 	auto ShaderMap = View.ShaderMap;
 
 	check(Component != NULL);
-	switch(AtmospherePhase)
+	switch (AtmospherePhase)
 	{
 	case AP_Transmittance:
 		{
 			const FSceneRenderTargetItem& DestRenderTarget = AtmosphereTextures->AtmosphereTransmittance->GetRenderTargetItem();
-			SetRenderTarget(RHICmdList, DestRenderTarget.TargetableTexture, FTextureRHIRef());
-			RHICmdList.ApplyCachedRenderTargets(GraphicsPSOInit);
 
-			TShaderMapRef<FAtmospherePrecomputeVS> VertexShader(ShaderMap);
-			TShaderMapRef<FAtmosphereTransmittancePS> PixelShader(ShaderMap);
+			FRHIRenderPassInfo RPInfo(DestRenderTarget.TargetableTexture, MakeRenderTargetActions(ERenderTargetLoadAction::ELoad, ERenderTargetStoreAction::EStore));
+			RHICmdList.BeginRenderPass(RPInfo, TEXT("AP_Transmittance"));
+			{
+				RHICmdList.ApplyCachedRenderTargets(GraphicsPSOInit);
 
-			GraphicsPSOInit.BoundShaderState.VertexDeclarationRHI = GFilterVertexDeclaration.VertexDeclarationRHI;
-			GraphicsPSOInit.BoundShaderState.VertexShaderRHI = GETSAFERHISHADER_VERTEX(*VertexShader);
-			GraphicsPSOInit.BoundShaderState.PixelShaderRHI = GETSAFERHISHADER_PIXEL(*PixelShader);
-			GraphicsPSOInit.PrimitiveType = PT_TriangleList;
+				TShaderMapRef<FAtmospherePrecomputeVS> VertexShader(ShaderMap);
+				TShaderMapRef<FAtmosphereTransmittancePS> PixelShader(ShaderMap);
 
-			SetGraphicsPipelineState(RHICmdList, GraphicsPSOInit);
-			PixelShader->SetParameters(RHICmdList, View);
-			DrawQuad(RHICmdList, ViewRect, *VertexShader);
+				GraphicsPSOInit.BoundShaderState.VertexDeclarationRHI = GFilterVertexDeclaration.VertexDeclarationRHI;
+				GraphicsPSOInit.BoundShaderState.VertexShaderRHI = GETSAFERHISHADER_VERTEX(*VertexShader);
+				GraphicsPSOInit.BoundShaderState.PixelShaderRHI = GETSAFERHISHADER_PIXEL(*PixelShader);
+				GraphicsPSOInit.PrimitiveType = PT_TriangleList;
 
-			RHICmdList.CopyToResolveTarget(DestRenderTarget.TargetableTexture, DestRenderTarget.ShaderResourceTexture, FResolveParams());
+				SetGraphicsPipelineState(RHICmdList, GraphicsPSOInit);
+				PixelShader->SetParameters(RHICmdList, View);
+				DrawQuad(RHICmdList, ViewRect, *VertexShader);
+			}
+			RHICmdList.EndRenderPass();
 		}
 		break;
 	case AP_Irradiance1:
 		{
 			const FSceneRenderTargetItem& DestRenderTarget = AtmosphereTextures->AtmosphereDeltaE->GetRenderTargetItem();
-			SetRenderTarget(RHICmdList, DestRenderTarget.TargetableTexture, FTextureRHIRef());
-			RHICmdList.ApplyCachedRenderTargets(GraphicsPSOInit);
-
-			TShaderMapRef<FAtmospherePrecomputeVS> VertexShader(ShaderMap);
-			TShaderMapRef<FAtmosphereIrradiance1PS> PixelShader(ShaderMap);
+			FRHIRenderPassInfo RPInfo(DestRenderTarget.TargetableTexture, MakeRenderTargetActions(ERenderTargetLoadAction::ELoad, ERenderTargetStoreAction::EStore));
+			RHICmdList.BeginRenderPass(RPInfo, TEXT("AP_Transmittance"));
+			{
+				RHICmdList.ApplyCachedRenderTargets(GraphicsPSOInit);
 
-			GraphicsPSOInit.BoundShaderState.VertexDeclarationRHI = GFilterVertexDeclaration.VertexDeclarationRHI;
-			GraphicsPSOInit.BoundShaderState.VertexShaderRHI = GETSAFERHISHADER_VERTEX(*VertexShader);
-			GraphicsPSOInit.BoundShaderState.PixelShaderRHI = GETSAFERHISHADER_PIXEL(*PixelShader);
-			GraphicsPSOInit.PrimitiveType = PT_TriangleList;
+				TShaderMapRef<FAtmospherePrecomputeVS> VertexShader(ShaderMap);
+				TShaderMapRef<FAtmosphereIrradiance1PS> PixelShader(ShaderMap);
 
-			SetGraphicsPipelineState(RHICmdList, GraphicsPSOInit);
+				GraphicsPSOInit.BoundShaderState.VertexDeclarationRHI = GFilterVertexDeclaration.VertexDeclarationRHI;
+				GraphicsPSOInit.BoundShaderState.VertexShaderRHI = GETSAFERHISHADER_VERTEX(*VertexShader);
+				GraphicsPSOInit.BoundShaderState.PixelShaderRHI = GETSAFERHISHADER_PIXEL(*PixelShader);
+				GraphicsPSOInit.PrimitiveType = PT_TriangleList;
 
-			PixelShader->SetParameters(RHICmdList, AtmosphereTextures);
+				SetGraphicsPipelineState(RHICmdList, GraphicsPSOInit);
 
-			DrawQuad(RHICmdList, ViewRect, *VertexShader);
+				PixelShader->SetParameters(RHICmdList, AtmosphereTextures);
 
-			RHICmdList.CopyToResolveTarget(DestRenderTarget.TargetableTexture, DestRenderTarget.ShaderResourceTexture, FResolveParams());
+				DrawQuad(RHICmdList, ViewRect, *VertexShader);
+			}
+			RHICmdList.EndRenderPass();
 		}
 		break;
 	case AP_Inscatter1:
 		{
 			int32 Layer = Atmosphere3DTextureIndex;
 			{
-				FRHIDepthRenderTargetView DepthRTV;
-				FRHIRenderTargetView RenderTargets[2];
-				RenderTargets[0] = FRHIRenderTargetView(AtmosphereTextures->AtmosphereDeltaSR->GetRenderTargetItem().TargetableTexture, ERenderTargetLoadAction::ELoad);
-				RenderTargets[1] = FRHIRenderTargetView(AtmosphereTextures->AtmosphereDeltaSM->GetRenderTargetItem().TargetableTexture, ERenderTargetLoadAction::ELoad);
-				
-				RHICmdList.SetRenderTargets(ARRAY_COUNT(RenderTargets), RenderTargets, &DepthRTV, 0, nullptr);
-				RHICmdList.ApplyCachedRenderTargets(GraphicsPSOInit);
-
-				TShaderMapRef<FAtmospherePrecomputeInscatterVS> VertexShader(ShaderMap);
-				TOptionalShaderMapRef<FAtmosphereGS> GeometryShader(ShaderMap);
-				TShaderMapRef<FAtmosphereInscatter1PS> PixelShader(ShaderMap);
-
-				GraphicsPSOInit.BoundShaderState.VertexDeclarationRHI = GFilterVertexDeclaration.VertexDeclarationRHI;
-				GraphicsPSOInit.BoundShaderState.VertexShaderRHI = GETSAFERHISHADER_VERTEX(*VertexShader);
-				GraphicsPSOInit.BoundShaderState.GeometryShaderRHI = GETSAFERHISHADER_GEOMETRY(*GeometryShader);
-				GraphicsPSOInit.BoundShaderState.PixelShaderRHI = GETSAFERHISHADER_PIXEL(*PixelShader);
-				GraphicsPSOInit.PrimitiveType = PT_TriangleList;
+				FTextureRHIParamRef RenderTargets[] =
+				{
+					AtmosphereTextures->AtmosphereDeltaSR->GetRenderTargetItem().TargetableTexture,
+					AtmosphereTextures->AtmosphereDeltaSM->GetRenderTargetItem().TargetableTexture,
+				};
 
-				SetGraphicsPipelineState(RHICmdList, GraphicsPSOInit);
+				FRHIRenderPassInfo RPInfo(2, RenderTargets, MakeRenderTargetActions(ERenderTargetLoadAction::ELoad, ERenderTargetStoreAction::EStore));
 
-				//
-				float r;
-				FVector4 DhdH;
-				GetLayerValue(Layer, r, DhdH);
-				VertexShader->SetParameters(RHICmdList, Layer);
-				if (GeometryShader.IsValid())
+				if (Atmosphere3DTextureIndex == Component->PrecomputeParams.InscatterAltitudeSampleNum - 1)
 				{
-					GeometryShader->SetParameters(RHICmdList, Layer);
+					RPInfo.ColorRenderTargets[0].ResolveTarget = AtmosphereTextures->AtmosphereDeltaSR->GetRenderTargetItem().ShaderResourceTexture;
+					RPInfo.ColorRenderTargets[1].ResolveTarget = AtmosphereTextures->AtmosphereDeltaSM->GetRenderTargetItem().ShaderResourceTexture;
 				}
-				PixelShader->SetParameters(RHICmdList, View, r, DhdH, AtmosphereTextures);
-				DrawQuad(RHICmdList, ViewRect, *VertexShader);
 
-				if (Atmosphere3DTextureIndex == Component->PrecomputeParams.InscatterAltitudeSampleNum - 1)
+				RHICmdList.BeginRenderPass(RPInfo, TEXT("AP_Inscatter"));
 				{
-					RHICmdList.CopyToResolveTarget(AtmosphereTextures->AtmosphereDeltaSR->GetRenderTargetItem().TargetableTexture,
-						AtmosphereTextures->AtmosphereDeltaSR->GetRenderTargetItem().ShaderResourceTexture, FResolveParams());
-					RHICmdList.CopyToResolveTarget(AtmosphereTextures->AtmosphereDeltaSM->GetRenderTargetItem().TargetableTexture,
-						AtmosphereTextures->AtmosphereDeltaSM->GetRenderTargetItem().ShaderResourceTexture, FResolveParams());
+					RHICmdList.ApplyCachedRenderTargets(GraphicsPSOInit);
+
+					TShaderMapRef<FAtmospherePrecomputeInscatterVS> VertexShader(ShaderMap);
+					TOptionalShaderMapRef<FAtmosphereGS> GeometryShader(ShaderMap);
+					TShaderMapRef<FAtmosphereInscatter1PS> PixelShader(ShaderMap);
+
+					GraphicsPSOInit.BoundShaderState.VertexDeclarationRHI = GFilterVertexDeclaration.VertexDeclarationRHI;
+					GraphicsPSOInit.BoundShaderState.VertexShaderRHI = GETSAFERHISHADER_VERTEX(*VertexShader);
+					GraphicsPSOInit.BoundShaderState.GeometryShaderRHI = GETSAFERHISHADER_GEOMETRY(*GeometryShader);
+					GraphicsPSOInit.BoundShaderState.PixelShaderRHI = GETSAFERHISHADER_PIXEL(*PixelShader);
+					GraphicsPSOInit.PrimitiveType = PT_TriangleList;
+
+					SetGraphicsPipelineState(RHICmdList, GraphicsPSOInit);
+
+					//
+					float r;
+					FVector4 DhdH;
+					GetLayerValue(Layer, r, DhdH);
+					VertexShader->SetParameters(RHICmdList, Layer);
+					if (GeometryShader.IsValid())
+					{
+						GeometryShader->SetParameters(RHICmdList, Layer);
+					}
+					PixelShader->SetParameters(RHICmdList, View, r, DhdH, AtmosphereTextures);
+					DrawQuad(RHICmdList, ViewRect, *VertexShader);
 				}
+				RHICmdList.EndRenderPass();
 			}
 		}
 		break;
@@ -1202,8 +1213,10 @@ void FAtmosphericFogSceneInfo::RenderAtmosphereShaders(FRHICommandList& RHICmdLi
 		{
 			const FSceneRenderTargetItem& DestRenderTarget = AtmosphereTextures->AtmosphereIrradiance->GetRenderTargetItem();
 			ensure(DestRenderTarget.TargetableTexture->GetClearColor() == FLinearColor::Black);
-			SetRenderTarget(RHICmdList, DestRenderTarget.TargetableTexture, FTextureRHIRef(), ESimpleRenderTargetMode::EClearColorExistingDepth);
-			RHICmdList.CopyToResolveTarget(DestRenderTarget.TargetableTexture, DestRenderTarget.ShaderResourceTexture, FResolveParams());
+
+			FRHIRenderPassInfo RPInfo(DestRenderTarget.TargetableTexture, ERenderTargetActions::Clear_Store, DestRenderTarget.ShaderResourceTexture);
+			RHICmdList.BeginRenderPass(RPInfo, TEXT("AP_ClearIrradiance"));
+			RHICmdList.EndRenderPass();
 		}
 		break;
 	case AP_CopyInscatter1:
@@ -1211,36 +1224,42 @@ void FAtmosphericFogSceneInfo::RenderAtmosphereShaders(FRHICommandList& RHICmdLi
 			int32 Layer = Atmosphere3DTextureIndex;
 			{
 				const FSceneRenderTargetItem& DestRenderTarget = AtmosphereTextures->AtmosphereInscatter->GetRenderTargetItem();
-				SetRenderTarget(RHICmdList, DestRenderTarget.TargetableTexture, FTextureRHIRef());
-				RHICmdList.ApplyCachedRenderTargets(GraphicsPSOInit);
-
-				TShaderMapRef<FAtmospherePrecomputeInscatterVS> VertexShader(ShaderMap);
-				TOptionalShaderMapRef<FAtmosphereGS> GeometryShader(ShaderMap);
-				TShaderMapRef<FAtmosphereCopyInscatter1PS> PixelShader(ShaderMap);
 
-				GraphicsPSOInit.BoundShaderState.VertexDeclarationRHI = GFilterVertexDeclaration.VertexDeclarationRHI;
-				GraphicsPSOInit.BoundShaderState.VertexShaderRHI = GETSAFERHISHADER_VERTEX(*VertexShader);
-				GraphicsPSOInit.BoundShaderState.GeometryShaderRHI = GETSAFERHISHADER_GEOMETRY(*GeometryShader);
-				GraphicsPSOInit.BoundShaderState.PixelShaderRHI = GETSAFERHISHADER_PIXEL(*PixelShader);
-				GraphicsPSOInit.PrimitiveType = PT_TriangleList;
-
-				SetGraphicsPipelineState(RHICmdList, GraphicsPSOInit);
+				FRHIRenderPassInfo RPInfo(DestRenderTarget.TargetableTexture, ERenderTargetActions::Load_Store);
 
-				float r;
-				FVector4 DhdH;
-				GetLayerValue(Layer, r, DhdH);
-				VertexShader->SetParameters(RHICmdList, Layer);
-				if (GeometryShader.IsValid())
+				if (Atmosphere3DTextureIndex == Component->PrecomputeParams.InscatterAltitudeSampleNum - 1)
 				{
-					GeometryShader->SetParameters(RHICmdList, Layer);
+					RPInfo.ColorRenderTargets[0].ResolveTarget = DestRenderTarget.ShaderResourceTexture;
 				}
-				PixelShader->SetParameters(RHICmdList, View, r, DhdH, Layer, AtmosphereTextures);
-				DrawQuad(RHICmdList, ViewRect, *VertexShader);
 
-				if (Atmosphere3DTextureIndex == Component->PrecomputeParams.InscatterAltitudeSampleNum - 1)
+				RHICmdList.BeginRenderPass(RPInfo, TEXT("AP_CopyInscatter1"));
 				{
-					RHICmdList.CopyToResolveTarget(DestRenderTarget.TargetableTexture, DestRenderTarget.ShaderResourceTexture, FResolveParams());
+					RHICmdList.ApplyCachedRenderTargets(GraphicsPSOInit);
+
+					TShaderMapRef<FAtmospherePrecomputeInscatterVS> VertexShader(ShaderMap);
+					TOptionalShaderMapRef<FAtmosphereGS> GeometryShader(ShaderMap);
+					TShaderMapRef<FAtmosphereCopyInscatter1PS> PixelShader(ShaderMap);
+
+					GraphicsPSOInit.BoundShaderState.VertexDeclarationRHI = GFilterVertexDeclaration.VertexDeclarationRHI;
+					GraphicsPSOInit.BoundShaderState.VertexShaderRHI = GETSAFERHISHADER_VERTEX(*VertexShader);
+					GraphicsPSOInit.BoundShaderState.GeometryShaderRHI = GETSAFERHISHADER_GEOMETRY(*GeometryShader);
+					GraphicsPSOInit.BoundShaderState.PixelShaderRHI = GETSAFERHISHADER_PIXEL(*PixelShader);
+					GraphicsPSOInit.PrimitiveType = PT_TriangleList;
+
+					SetGraphicsPipelineState(RHICmdList, GraphicsPSOInit);
+
+					float r;
+					FVector4 DhdH;
+					GetLayerValue(Layer, r, DhdH);
+					VertexShader->SetParameters(RHICmdList, Layer);
+					if (GeometryShader.IsValid())
+					{
+						GeometryShader->SetParameters(RHICmdList, Layer);
+					}
+					PixelShader->SetParameters(RHICmdList, View, r, DhdH, Layer, AtmosphereTextures);
+					DrawQuad(RHICmdList, ViewRect, *VertexShader);
 				}
+				RHICmdList.EndRenderPass();
 			}
 		}
 		break;
@@ -1249,60 +1268,66 @@ void FAtmosphericFogSceneInfo::RenderAtmosphereShaders(FRHICommandList& RHICmdLi
 			int32 Layer = Atmosphere3DTextureIndex;
 			{
 				const FSceneRenderTargetItem& DestRenderTarget = AtmosphereTextures->AtmosphereDeltaJ->GetRenderTargetItem();
-				SetRenderTarget(RHICmdList, DestRenderTarget.TargetableTexture, FTextureRHIRef());
-				RHICmdList.ApplyCachedRenderTargets(GraphicsPSOInit);
-
-				TShaderMapRef<FAtmospherePrecomputeInscatterVS> VertexShader(ShaderMap);
-				TOptionalShaderMapRef<FAtmosphereGS> GeometryShader(ShaderMap);
-				TShaderMapRef<FAtmosphereInscatterSPS> PixelShader(ShaderMap);
-
-				GraphicsPSOInit.BoundShaderState.VertexDeclarationRHI = GFilterVertexDeclaration.VertexDeclarationRHI;
-				GraphicsPSOInit.BoundShaderState.VertexShaderRHI = GETSAFERHISHADER_VERTEX(*VertexShader);
-				GraphicsPSOInit.BoundShaderState.GeometryShaderRHI = GETSAFERHISHADER_GEOMETRY(*GeometryShader);
-				GraphicsPSOInit.BoundShaderState.PixelShaderRHI = GETSAFERHISHADER_PIXEL(*PixelShader);
-				GraphicsPSOInit.PrimitiveType = PT_TriangleList;
 
-				SetGraphicsPipelineState(RHICmdList, GraphicsPSOInit);
-
-				float r;
-				FVector4 DhdH;
-				GetLayerValue(Layer, r, DhdH);
-				VertexShader->SetParameters(RHICmdList, Layer);
-				if (GeometryShader.IsValid())
+				FRHIRenderPassInfo RPInfo(DestRenderTarget.TargetableTexture, ERenderTargetActions::Load_Store);
+				if (Atmosphere3DTextureIndex == Component->PrecomputeParams.InscatterAltitudeSampleNum - 1)
 				{
-					GeometryShader->SetParameters(RHICmdList, Layer);
+					RPInfo.ColorRenderTargets[0].ResolveTarget = DestRenderTarget.ShaderResourceTexture;
 				}
-				PixelShader->SetParameters(RHICmdList, View, r, DhdH, AtmoshpereOrder == 2 ? 1.f : 0.f, AtmosphereTextures);
-				DrawQuad(RHICmdList, ViewRect, *VertexShader);
-
-				if (Atmosphere3DTextureIndex == Component->PrecomputeParams.InscatterAltitudeSampleNum - 1)
+				RHICmdList.BeginRenderPass(RPInfo, TEXT("InscatterS"));
 				{
-					RHICmdList.CopyToResolveTarget(DestRenderTarget.TargetableTexture, DestRenderTarget.ShaderResourceTexture, FResolveParams());
+					RHICmdList.ApplyCachedRenderTargets(GraphicsPSOInit);
+
+					TShaderMapRef<FAtmospherePrecomputeInscatterVS> VertexShader(ShaderMap);
+					TOptionalShaderMapRef<FAtmosphereGS> GeometryShader(ShaderMap);
+					TShaderMapRef<FAtmosphereInscatterSPS> PixelShader(ShaderMap);
+
+					GraphicsPSOInit.BoundShaderState.VertexDeclarationRHI = GFilterVertexDeclaration.VertexDeclarationRHI;
+					GraphicsPSOInit.BoundShaderState.VertexShaderRHI = GETSAFERHISHADER_VERTEX(*VertexShader);
+					GraphicsPSOInit.BoundShaderState.GeometryShaderRHI = GETSAFERHISHADER_GEOMETRY(*GeometryShader);
+					GraphicsPSOInit.BoundShaderState.PixelShaderRHI = GETSAFERHISHADER_PIXEL(*PixelShader);
+					GraphicsPSOInit.PrimitiveType = PT_TriangleList;
+
+					SetGraphicsPipelineState(RHICmdList, GraphicsPSOInit);
+
+					float r;
+					FVector4 DhdH;
+					GetLayerValue(Layer, r, DhdH);
+					VertexShader->SetParameters(RHICmdList, Layer);
+					if (GeometryShader.IsValid())
+					{
+						GeometryShader->SetParameters(RHICmdList, Layer);
+					}
+					PixelShader->SetParameters(RHICmdList, View, r, DhdH, AtmoshpereOrder == 2 ? 1.f : 0.f, AtmosphereTextures);
+					DrawQuad(RHICmdList, ViewRect, *VertexShader);
 				}
+				RHICmdList.EndRenderPass();
 			}
 		}
 		break;
 	case AP_IrradianceN:
 		{
 			const FSceneRenderTargetItem& DestRenderTarget = AtmosphereTextures->AtmosphereDeltaE->GetRenderTargetItem();
-			SetRenderTarget(RHICmdList, DestRenderTarget.TargetableTexture, FTextureRHIRef());
-			RHICmdList.ApplyCachedRenderTargets(GraphicsPSOInit);
-
-			TShaderMapRef<FAtmospherePrecomputeVS> VertexShader(ShaderMap);
-			TShaderMapRef<FAtmosphereIrradianceNPS> PixelShader(ShaderMap);
+			FRHIRenderPassInfo RPInfo(DestRenderTarget.TargetableTexture, ERenderTargetActions::Load_Store, DestRenderTarget.ShaderResourceTexture);
+			RHICmdList.BeginRenderPass(RPInfo, TEXT("IrradianceN"));
+			{
+				RHICmdList.ApplyCachedRenderTargets(GraphicsPSOInit);
 
-			GraphicsPSOInit.BoundShaderState.VertexDeclarationRHI = GFilterVertexDeclaration.VertexDeclarationRHI;
-			GraphicsPSOInit.BoundShaderState.VertexShaderRHI = GETSAFERHISHADER_VERTEX(*VertexShader);
-			GraphicsPSOInit.BoundShaderState.PixelShaderRHI = GETSAFERHISHADER_PIXEL(*PixelShader);
-			GraphicsPSOInit.PrimitiveType = PT_TriangleList;
+				TShaderMapRef<FAtmospherePrecomputeVS> VertexShader(ShaderMap);
+				TShaderMapRef<FAtmosphereIrradianceNPS> PixelShader(ShaderMap);
 
-			SetGraphicsPipelineState(RHICmdList, GraphicsPSOInit);
+				GraphicsPSOInit.BoundShaderState.VertexDeclarationRHI = GFilterVertexDeclaration.VertexDeclarationRHI;
+				GraphicsPSOInit.BoundShaderState.VertexShaderRHI = GETSAFERHISHADER_VERTEX(*VertexShader);
+				GraphicsPSOInit.BoundShaderState.PixelShaderRHI = GETSAFERHISHADER_PIXEL(*PixelShader);
+				GraphicsPSOInit.PrimitiveType = PT_TriangleList;
 
-			PixelShader->SetParameters(RHICmdList, View, AtmoshpereOrder == 2 ? 1.f : 0.f, AtmosphereTextures);
+				SetGraphicsPipelineState(RHICmdList, GraphicsPSOInit);
 
-			DrawQuad(RHICmdList, ViewRect, *VertexShader);
+				PixelShader->SetParameters(RHICmdList, View, AtmoshpereOrder == 2 ? 1.f : 0.f, AtmosphereTextures);
 
-			RHICmdList.CopyToResolveTarget(DestRenderTarget.TargetableTexture, DestRenderTarget.ShaderResourceTexture, FResolveParams());
+				DrawQuad(RHICmdList, ViewRect, *VertexShader);
+			}
+			RHICmdList.EndRenderPass();
 		}
 		break;
 
@@ -1311,62 +1336,69 @@ void FAtmosphericFogSceneInfo::RenderAtmosphereShaders(FRHICommandList& RHICmdLi
 			int32 Layer = Atmosphere3DTextureIndex;
 			{
 				const FSceneRenderTargetItem& DestRenderTarget = AtmosphereTextures->AtmosphereDeltaSR->GetRenderTargetItem();
-				SetRenderTarget(RHICmdList, DestRenderTarget.TargetableTexture, FTextureRHIRef());
-				RHICmdList.ApplyCachedRenderTargets(GraphicsPSOInit);
-
-				TShaderMapRef<FAtmospherePrecomputeInscatterVS> VertexShader(ShaderMap);
-				TOptionalShaderMapRef<FAtmosphereGS> GeometryShader(ShaderMap);
-				TShaderMapRef<FAtmosphereInscatterNPS> PixelShader(ShaderMap);
 
-				GraphicsPSOInit.BoundShaderState.VertexDeclarationRHI = GFilterVertexDeclaration.VertexDeclarationRHI;
-				GraphicsPSOInit.BoundShaderState.VertexShaderRHI = GETSAFERHISHADER_VERTEX(*VertexShader);
-				GraphicsPSOInit.BoundShaderState.GeometryShaderRHI = GETSAFERHISHADER_GEOMETRY(*GeometryShader);
-				GraphicsPSOInit.BoundShaderState.PixelShaderRHI = GETSAFERHISHADER_PIXEL(*PixelShader);
-				GraphicsPSOInit.PrimitiveType = PT_TriangleList;
-
-				SetGraphicsPipelineState(RHICmdList, GraphicsPSOInit);
-
-				float r;
-				FVector4 DhdH;
-				GetLayerValue(Layer, r, DhdH);
-				VertexShader->SetParameters(RHICmdList, Layer);
-				if (GeometryShader.IsValid())
+				FRHIRenderPassInfo RPInfo(DestRenderTarget.TargetableTexture, ERenderTargetActions::Load_Store);
+				if (Atmosphere3DTextureIndex == Component->PrecomputeParams.InscatterAltitudeSampleNum - 1)
 				{
-					GeometryShader->SetParameters(RHICmdList, Layer);
+					RPInfo.ColorRenderTargets[0].ResolveTarget = DestRenderTarget.ShaderResourceTexture;
 				}
-				PixelShader->SetParameters(RHICmdList, View, r, DhdH, AtmoshpereOrder == 2 ? 1.f : 0.f, AtmosphereTextures);
-				DrawQuad(RHICmdList, ViewRect, *VertexShader);
-
-				if (Atmosphere3DTextureIndex == Component->PrecomputeParams.InscatterAltitudeSampleNum - 1)
+				RHICmdList.BeginRenderPass(RPInfo, TEXT("InscatterN"));
 				{
-					RHICmdList.CopyToResolveTarget(DestRenderTarget.TargetableTexture, DestRenderTarget.ShaderResourceTexture, FResolveParams());
+					RHICmdList.ApplyCachedRenderTargets(GraphicsPSOInit);
+
+					TShaderMapRef<FAtmospherePrecomputeInscatterVS> VertexShader(ShaderMap);
+					TOptionalShaderMapRef<FAtmosphereGS> GeometryShader(ShaderMap);
+					TShaderMapRef<FAtmosphereInscatterNPS> PixelShader(ShaderMap);
+
+					GraphicsPSOInit.BoundShaderState.VertexDeclarationRHI = GFilterVertexDeclaration.VertexDeclarationRHI;
+					GraphicsPSOInit.BoundShaderState.VertexShaderRHI = GETSAFERHISHADER_VERTEX(*VertexShader);
+					GraphicsPSOInit.BoundShaderState.GeometryShaderRHI = GETSAFERHISHADER_GEOMETRY(*GeometryShader);
+					GraphicsPSOInit.BoundShaderState.PixelShaderRHI = GETSAFERHISHADER_PIXEL(*PixelShader);
+					GraphicsPSOInit.PrimitiveType = PT_TriangleList;
+
+					SetGraphicsPipelineState(RHICmdList, GraphicsPSOInit);
+
+					float r;
+					FVector4 DhdH;
+					GetLayerValue(Layer, r, DhdH);
+					VertexShader->SetParameters(RHICmdList, Layer);
+					if (GeometryShader.IsValid())
+					{
+						GeometryShader->SetParameters(RHICmdList, Layer);
+					}
+					PixelShader->SetParameters(RHICmdList, View, r, DhdH, AtmoshpereOrder == 2 ? 1.f : 0.f, AtmosphereTextures);
+					DrawQuad(RHICmdList, ViewRect, *VertexShader);
 				}
+				RHICmdList.EndRenderPass();
 			}
 		}
 		break;
 	case AP_CopyIrradiance:
 		{
 			const FSceneRenderTargetItem& DestRenderTarget = AtmosphereTextures->AtmosphereIrradiance->GetRenderTargetItem();
-			SetRenderTarget(RHICmdList, DestRenderTarget.TargetableTexture, FTextureRHIRef());
-			RHICmdList.ApplyCachedRenderTargets(GraphicsPSOInit);
 
-			GraphicsPSOInit.BlendState = TStaticBlendState<CW_RGBA, BO_Add, BF_One, BF_One, BO_Add, BF_One, BF_One>::GetRHI();
+			FRHIRenderPassInfo RPInfo(DestRenderTarget.TargetableTexture, ERenderTargetActions::Load_Store, DestRenderTarget.ShaderResourceTexture);
+			RHICmdList.BeginRenderPass(RPInfo, TEXT("CopyIrradiance"));
+			{
+				RHICmdList.ApplyCachedRenderTargets(GraphicsPSOInit);
 
-			TShaderMapRef<FAtmospherePrecomputeVS> VertexShader(ShaderMap);
-			TShaderMapRef<FAtmosphereCopyIrradiancePS> PixelShader(ShaderMap);
+				GraphicsPSOInit.BlendState = TStaticBlendState<CW_RGBA, BO_Add, BF_One, BF_One, BO_Add, BF_One, BF_One>::GetRHI();
 
-			GraphicsPSOInit.BoundShaderState.VertexDeclarationRHI = GFilterVertexDeclaration.VertexDeclarationRHI;
-			GraphicsPSOInit.BoundShaderState.VertexShaderRHI = GETSAFERHISHADER_VERTEX(*VertexShader);
-			GraphicsPSOInit.BoundShaderState.PixelShaderRHI = GETSAFERHISHADER_PIXEL(*PixelShader);
-			GraphicsPSOInit.PrimitiveType = PT_TriangleList;
+				TShaderMapRef<FAtmospherePrecomputeVS> VertexShader(ShaderMap);
+				TShaderMapRef<FAtmosphereCopyIrradiancePS> PixelShader(ShaderMap);
 
-			SetGraphicsPipelineState(RHICmdList, GraphicsPSOInit);
+				GraphicsPSOInit.BoundShaderState.VertexDeclarationRHI = GFilterVertexDeclaration.VertexDeclarationRHI;
+				GraphicsPSOInit.BoundShaderState.VertexShaderRHI = GETSAFERHISHADER_VERTEX(*VertexShader);
+				GraphicsPSOInit.BoundShaderState.PixelShaderRHI = GETSAFERHISHADER_PIXEL(*PixelShader);
+				GraphicsPSOInit.PrimitiveType = PT_TriangleList;
 
-			PixelShader->SetParameters(RHICmdList, AtmosphereTextures);
+				SetGraphicsPipelineState(RHICmdList, GraphicsPSOInit);
 
-			DrawQuad(RHICmdList, ViewRect, *VertexShader);
+				PixelShader->SetParameters(RHICmdList, AtmosphereTextures);
 
-			RHICmdList.CopyToResolveTarget(DestRenderTarget.TargetableTexture, DestRenderTarget.ShaderResourceTexture, FResolveParams());
+				DrawQuad(RHICmdList, ViewRect, *VertexShader);
+			}
+			RHICmdList.EndRenderPass();
 
 			GraphicsPSOInit.BlendState = TStaticBlendState<>::GetRHI();
 		}
@@ -1377,38 +1409,43 @@ void FAtmosphericFogSceneInfo::RenderAtmosphereShaders(FRHICommandList& RHICmdLi
 			int32 Layer = Atmosphere3DTextureIndex;
 			{
 				const FSceneRenderTargetItem& DestRenderTarget = AtmosphereTextures->AtmosphereInscatter->GetRenderTargetItem();
-				SetRenderTarget(RHICmdList, DestRenderTarget.TargetableTexture, FTextureRHIRef());
-				RHICmdList.ApplyCachedRenderTargets(GraphicsPSOInit);
-
-				GraphicsPSOInit.BlendState = TStaticBlendState<CW_RGBA, BO_Add, BF_One, BF_One, BO_Add, BF_One, BF_One>::GetRHI();
-
-				TShaderMapRef<FAtmospherePrecomputeInscatterVS> VertexShader(ShaderMap);
-				TOptionalShaderMapRef<FAtmosphereGS> GeometryShader(ShaderMap);
-				TShaderMapRef<FAtmosphereCopyInscatterNPS> PixelShader(ShaderMap);
-
-				GraphicsPSOInit.BoundShaderState.VertexDeclarationRHI = GFilterVertexDeclaration.VertexDeclarationRHI;
-				GraphicsPSOInit.BoundShaderState.VertexShaderRHI = GETSAFERHISHADER_VERTEX(*VertexShader);
-				GraphicsPSOInit.BoundShaderState.GeometryShaderRHI = GETSAFERHISHADER_GEOMETRY(*GeometryShader);
-				GraphicsPSOInit.BoundShaderState.PixelShaderRHI = GETSAFERHISHADER_PIXEL(*PixelShader);
-				GraphicsPSOInit.PrimitiveType = PT_TriangleList;
-
-				SetGraphicsPipelineState(RHICmdList, GraphicsPSOInit);
-
-				float r;
-				FVector4 DhdH;
-				GetLayerValue(Layer, r, DhdH);
-				VertexShader->SetParameters(RHICmdList, Layer);
-				if (GeometryShader.IsValid())
+				
+				FRHIRenderPassInfo RPInfo(DestRenderTarget.TargetableTexture, ERenderTargetActions::Load_Store);
+				if (Atmosphere3DTextureIndex == Component->PrecomputeParams.InscatterAltitudeSampleNum - 1)
 				{
-					GeometryShader->SetParameters(RHICmdList, Layer);
+					RPInfo.ColorRenderTargets[0].ResolveTarget = DestRenderTarget.ShaderResourceTexture;
 				}
-				PixelShader->SetParameters(RHICmdList, View, r, DhdH, Layer, AtmosphereTextures);
-				DrawQuad(RHICmdList, ViewRect, *VertexShader);
 
-				if (Atmosphere3DTextureIndex == Component->PrecomputeParams.InscatterAltitudeSampleNum - 1)
+				RHICmdList.BeginRenderPass(RPInfo, TEXT("CopyInscatterN"));
 				{
-					RHICmdList.CopyToResolveTarget(DestRenderTarget.TargetableTexture, DestRenderTarget.ShaderResourceTexture, FResolveParams());
+					RHICmdList.ApplyCachedRenderTargets(GraphicsPSOInit);
+
+					GraphicsPSOInit.BlendState = TStaticBlendState<CW_RGBA, BO_Add, BF_One, BF_One, BO_Add, BF_One, BF_One>::GetRHI();
+
+					TShaderMapRef<FAtmospherePrecomputeInscatterVS> VertexShader(ShaderMap);
+					TOptionalShaderMapRef<FAtmosphereGS> GeometryShader(ShaderMap);
+					TShaderMapRef<FAtmosphereCopyInscatterNPS> PixelShader(ShaderMap);
+
+					GraphicsPSOInit.BoundShaderState.VertexDeclarationRHI = GFilterVertexDeclaration.VertexDeclarationRHI;
+					GraphicsPSOInit.BoundShaderState.VertexShaderRHI = GETSAFERHISHADER_VERTEX(*VertexShader);
+					GraphicsPSOInit.BoundShaderState.GeometryShaderRHI = GETSAFERHISHADER_GEOMETRY(*GeometryShader);
+					GraphicsPSOInit.BoundShaderState.PixelShaderRHI = GETSAFERHISHADER_PIXEL(*PixelShader);
+					GraphicsPSOInit.PrimitiveType = PT_TriangleList;
+
+					SetGraphicsPipelineState(RHICmdList, GraphicsPSOInit);
+
+					float r;
+					FVector4 DhdH;
+					GetLayerValue(Layer, r, DhdH);
+					VertexShader->SetParameters(RHICmdList, Layer);
+					if (GeometryShader.IsValid())
+					{
+						GeometryShader->SetParameters(RHICmdList, Layer);
+					}
+					PixelShader->SetParameters(RHICmdList, View, r, DhdH, Layer, AtmosphereTextures);
+					DrawQuad(RHICmdList, ViewRect, *VertexShader);
 				}
+				RHICmdList.EndRenderPass();
 			}
 
 			GraphicsPSOInit.BlendState = TStaticBlendState<>::GetRHI();
@@ -1420,36 +1457,40 @@ void FAtmosphericFogSceneInfo::RenderAtmosphereShaders(FRHICommandList& RHICmdLi
 			int32 Layer = Atmosphere3DTextureIndex;
 			{
 				const FSceneRenderTargetItem& DestRenderTarget = AtmosphereTextures->AtmosphereDeltaSR->GetRenderTargetItem();
-				SetRenderTarget(RHICmdList, DestRenderTarget.TargetableTexture, FTextureRHIRef());
-				RHICmdList.ApplyCachedRenderTargets(GraphicsPSOInit);
-
-				TShaderMapRef<FAtmospherePrecomputeInscatterVS> VertexShader(ShaderMap);
-				TOptionalShaderMapRef<FAtmosphereGS> GeometryShader(ShaderMap);
-				TShaderMapRef<FAtmosphereCopyInscatterFPS> PixelShader(ShaderMap);
-
-				GraphicsPSOInit.BoundShaderState.VertexDeclarationRHI = GFilterVertexDeclaration.VertexDeclarationRHI;
-				GraphicsPSOInit.BoundShaderState.VertexShaderRHI = GETSAFERHISHADER_VERTEX(*VertexShader);
-				GraphicsPSOInit.BoundShaderState.GeometryShaderRHI = GETSAFERHISHADER_GEOMETRY(*GeometryShader);
-				GraphicsPSOInit.BoundShaderState.PixelShaderRHI = GETSAFERHISHADER_PIXEL(*PixelShader);
-				GraphicsPSOInit.PrimitiveType = PT_TriangleList;
-
-				SetGraphicsPipelineState(RHICmdList, GraphicsPSOInit);
-
-				float r;
-				FVector4 DhdH;
-				GetLayerValue(Layer, r, DhdH);
-				VertexShader->SetParameters(RHICmdList, Layer);
-				if (GeometryShader.IsValid())
+				
+				FRHIRenderPassInfo RPInfo(DestRenderTarget.TargetableTexture, ERenderTargetActions::Load_Store);
+				if (Atmosphere3DTextureIndex == Component->PrecomputeParams.InscatterAltitudeSampleNum - 1)
 				{
-					GeometryShader->SetParameters(RHICmdList, Layer);
+					RPInfo.ColorRenderTargets[0].ResolveTarget = DestRenderTarget.ShaderResourceTexture;
 				}
-				PixelShader->SetParameters(RHICmdList, r, DhdH, Layer, AtmosphereTextures);
-				DrawQuad(RHICmdList, ViewRect, *VertexShader);
-
-				if (Atmosphere3DTextureIndex == Component->PrecomputeParams.InscatterAltitudeSampleNum - 1)
+				RHICmdList.BeginRenderPass(RPInfo, TEXT("CopyInscatterF"));
 				{
-					RHICmdList.CopyToResolveTarget(DestRenderTarget.TargetableTexture, DestRenderTarget.ShaderResourceTexture, FResolveParams());
+					RHICmdList.ApplyCachedRenderTargets(GraphicsPSOInit);
+
+					TShaderMapRef<FAtmospherePrecomputeInscatterVS> VertexShader(ShaderMap);
+					TOptionalShaderMapRef<FAtmosphereGS> GeometryShader(ShaderMap);
+					TShaderMapRef<FAtmosphereCopyInscatterFPS> PixelShader(ShaderMap);
+
+					GraphicsPSOInit.BoundShaderState.VertexDeclarationRHI = GFilterVertexDeclaration.VertexDeclarationRHI;
+					GraphicsPSOInit.BoundShaderState.VertexShaderRHI = GETSAFERHISHADER_VERTEX(*VertexShader);
+					GraphicsPSOInit.BoundShaderState.GeometryShaderRHI = GETSAFERHISHADER_GEOMETRY(*GeometryShader);
+					GraphicsPSOInit.BoundShaderState.PixelShaderRHI = GETSAFERHISHADER_PIXEL(*PixelShader);
+					GraphicsPSOInit.PrimitiveType = PT_TriangleList;
+
+					SetGraphicsPipelineState(RHICmdList, GraphicsPSOInit);
+
+					float r;
+					FVector4 DhdH;
+					GetLayerValue(Layer, r, DhdH);
+					VertexShader->SetParameters(RHICmdList, Layer);
+					if (GeometryShader.IsValid())
+					{
+						GeometryShader->SetParameters(RHICmdList, Layer);
+					}
+					PixelShader->SetParameters(RHICmdList, r, DhdH, Layer, AtmosphereTextures);
+					DrawQuad(RHICmdList, ViewRect, *VertexShader);
 				}
+				RHICmdList.EndRenderPass();				
 			}
 		}
 		break;
@@ -1458,36 +1499,40 @@ void FAtmosphericFogSceneInfo::RenderAtmosphereShaders(FRHICommandList& RHICmdLi
 			int32 Layer = Atmosphere3DTextureIndex;
 			{
 				const FSceneRenderTargetItem& DestRenderTarget = AtmosphereTextures->AtmosphereInscatter->GetRenderTargetItem();
-				SetRenderTarget(RHICmdList, DestRenderTarget.TargetableTexture, FTextureRHIRef());
-				RHICmdList.ApplyCachedRenderTargets(GraphicsPSOInit);
 
-				TShaderMapRef<FAtmospherePrecomputeInscatterVS> VertexShader(ShaderMap);
-				TOptionalShaderMapRef<FAtmosphereGS> GeometryShader(ShaderMap);
-				TShaderMapRef<FAtmosphereCopyInscatterFBackPS> PixelShader(ShaderMap);
-
-				GraphicsPSOInit.BoundShaderState.VertexDeclarationRHI = GFilterVertexDeclaration.VertexDeclarationRHI;
-				GraphicsPSOInit.BoundShaderState.VertexShaderRHI = GETSAFERHISHADER_VERTEX(*VertexShader);
-				GraphicsPSOInit.BoundShaderState.GeometryShaderRHI = GETSAFERHISHADER_GEOMETRY(*GeometryShader);
-				GraphicsPSOInit.BoundShaderState.PixelShaderRHI = GETSAFERHISHADER_PIXEL(*PixelShader);
-				GraphicsPSOInit.PrimitiveType = PT_TriangleList;
-
-				SetGraphicsPipelineState(RHICmdList, GraphicsPSOInit);
-
-				float r;
-				FVector4 DhdH;
-				GetLayerValue(Layer, r, DhdH);
-				VertexShader->SetParameters(RHICmdList, Layer);
-				if (GeometryShader.IsValid())
+				FRHIRenderPassInfo RPInfo(DestRenderTarget.TargetableTexture, ERenderTargetActions::Load_Store);
+				if (Atmosphere3DTextureIndex == Component->PrecomputeParams.InscatterAltitudeSampleNum - 1)
 				{
-					GeometryShader->SetParameters(RHICmdList, Layer);
+					RPInfo.ColorRenderTargets[0].ResolveTarget = DestRenderTarget.ShaderResourceTexture;
 				}
-				PixelShader->SetParameters(RHICmdList, r, DhdH, Layer, AtmosphereTextures);
-				DrawQuad(RHICmdList, ViewRect, *VertexShader);
-
-				if (Atmosphere3DTextureIndex == Component->PrecomputeParams.InscatterAltitudeSampleNum - 1)
+				RHICmdList.BeginRenderPass(RPInfo, TEXT("CopyInscatterFBack"));
 				{
-					RHICmdList.CopyToResolveTarget(DestRenderTarget.TargetableTexture, DestRenderTarget.ShaderResourceTexture, FResolveParams());
+					RHICmdList.ApplyCachedRenderTargets(GraphicsPSOInit);
+
+					TShaderMapRef<FAtmospherePrecomputeInscatterVS> VertexShader(ShaderMap);
+					TOptionalShaderMapRef<FAtmosphereGS> GeometryShader(ShaderMap);
+					TShaderMapRef<FAtmosphereCopyInscatterFBackPS> PixelShader(ShaderMap);
+
+					GraphicsPSOInit.BoundShaderState.VertexDeclarationRHI = GFilterVertexDeclaration.VertexDeclarationRHI;
+					GraphicsPSOInit.BoundShaderState.VertexShaderRHI = GETSAFERHISHADER_VERTEX(*VertexShader);
+					GraphicsPSOInit.BoundShaderState.GeometryShaderRHI = GETSAFERHISHADER_GEOMETRY(*GeometryShader);
+					GraphicsPSOInit.BoundShaderState.PixelShaderRHI = GETSAFERHISHADER_PIXEL(*PixelShader);
+					GraphicsPSOInit.PrimitiveType = PT_TriangleList;
+
+					SetGraphicsPipelineState(RHICmdList, GraphicsPSOInit);
+
+					float r;
+					FVector4 DhdH;
+					GetLayerValue(Layer, r, DhdH);
+					VertexShader->SetParameters(RHICmdList, Layer);
+					if (GeometryShader.IsValid())
+					{
+						GeometryShader->SetParameters(RHICmdList, Layer);
+					}
+					PixelShader->SetParameters(RHICmdList, r, DhdH, Layer, AtmosphereTextures);
+					DrawQuad(RHICmdList, ViewRect, *VertexShader);
 				}
+				RHICmdList.EndRenderPass();
 			}
 		}
 		break;
diff --git a/Engine/Source/Runtime/Renderer/Private/AtmosphereTextures.cpp b/Engine/Source/Runtime/Renderer/Private/AtmosphereTextures.cpp
--- a/Engine/Source/Runtime/Renderer/Private/AtmosphereTextures.cpp
+++ b/Engine/Source/Runtime/Renderer/Private/AtmosphereTextures.cpp
@@ -23,10 +23,11 @@ void FAtmosphereTextures::InitDynamicRHI()
 		FPooledRenderTargetDesc TransmittanceDesc(FPooledRenderTargetDesc::Create2DDesc(GTransmittanceTexSize, PF_FloatRGBA, FClearValueBinding::Black, TexCreate_None, TexCreate_RenderTargetable, false));
 		GRenderTargetPool.FindFreeElement(RHICmdList, TransmittanceDesc, AtmosphereTransmittance, TEXT("AtmosphereTransmittance"));
 		{
-			FRHIRenderTargetView View = FRHIRenderTargetView(AtmosphereTransmittance->GetRenderTargetItem().TargetableTexture, ERenderTargetLoadAction::EClear);
-			FRHISetRenderTargetsInfo Info(1, &View, FRHIDepthRenderTargetView());
-			RHICmdList.SetRenderTargetsAndClear(Info);
-			RHICmdList.CopyToResolveTarget(AtmosphereTransmittance->GetRenderTargetItem().TargetableTexture, AtmosphereTransmittance->GetRenderTargetItem().ShaderResourceTexture, FResolveParams());
+			const FSceneRenderTargetItem& TransmittanceTarget = AtmosphereTransmittance->GetRenderTargetItem();
+
+			FRHIRenderPassInfo RPInfo(TransmittanceTarget.TargetableTexture, ERenderTargetActions::Clear_Store, TransmittanceTarget.ShaderResourceTexture);
+			RHICmdList.BeginRenderPass(RPInfo, TEXT("ClearTransmittance"));
+			RHICmdList.EndRenderPass();
 		}
 
 		// Irradiance
@@ -34,10 +35,11 @@ void FAtmosphereTextures::InitDynamicRHI()
 		FPooledRenderTargetDesc IrradianceDesc(FPooledRenderTargetDesc::Create2DDesc(GIrradianceTexSize, PF_FloatRGBA, FClearValueBinding::Black, TexCreate_None, TexCreate_RenderTargetable, false));
 		GRenderTargetPool.FindFreeElement(RHICmdList, IrradianceDesc, AtmosphereIrradiance, TEXT("AtmosphereIrradiance"));
 		{
-			FRHIRenderTargetView View = FRHIRenderTargetView(AtmosphereIrradiance->GetRenderTargetItem().TargetableTexture, ERenderTargetLoadAction::EClear);
-			FRHISetRenderTargetsInfo Info(1, &View, FRHIDepthRenderTargetView());
-			RHICmdList.SetRenderTargetsAndClear(Info);
-			RHICmdList.CopyToResolveTarget(AtmosphereIrradiance->GetRenderTargetItem().TargetableTexture, AtmosphereIrradiance->GetRenderTargetItem().ShaderResourceTexture, FResolveParams());
+			const FSceneRenderTargetItem& IrradianceTarget = AtmosphereIrradiance->GetRenderTargetItem();
+
+			FRHIRenderPassInfo RPInfo(IrradianceTarget.TargetableTexture, ERenderTargetActions::Clear_Store, IrradianceTarget.ShaderResourceTexture);
+			RHICmdList.BeginRenderPass(RPInfo, TEXT("ClearIrradiance"));
+			RHICmdList.EndRenderPass();
 		}
 		// DeltaE
 		GRenderTargetPool.FindFreeElement(RHICmdList, IrradianceDesc, AtmosphereDeltaE, TEXT("AtmosphereDeltaE"));
