Renderpass updates.

CompositionLighting updates. More self-contained postprocess updates.

diff --git a/Engine/Source/Runtime/Renderer/Private/CompositionLighting/PostProcessAmbient.cpp b/Engine/Source/Runtime/Renderer/Private/CompositionLighting/PostProcessAmbient.cpp
--- a/Engine/Source/Runtime/Renderer/Private/CompositionLighting/PostProcessAmbient.cpp
+++ b/Engine/Source/Runtime/Renderer/Private/CompositionLighting/PostProcessAmbient.cpp
@@ -117,28 +117,31 @@ void FRCPassPostProcessAmbient::Process(FRenderingCompositePassContext& Context)
 	const FSceneRenderTargetItem& DestRenderTarget = FSceneRenderTargets::Get(Context.RHICmdList).GetSceneColor()->GetRenderTargetItem();
 
 	// Set the view family's render target/viewport.
-	SetRenderTarget(Context.RHICmdList, DestRenderTarget.TargetableTexture, FTextureRHIRef(), true);
-	Context.SetViewportAndCallRHI(View.ViewRect);
+	FRHIRenderPassInfo RPInfo(DestRenderTarget.TargetableTexture, ERenderTargetActions::Load_Store, DestRenderTarget.ShaderResourceTexture);
+	TransitionRenderPassTargets(Context.RHICmdList, RPInfo);
+	Context.RHICmdList.BeginRenderPass(RPInfo, TEXT("PostProcessAmbient"));
+	{
+		Context.SetViewportAndCallRHI(View.ViewRect);
 
-	FGraphicsPipelineStateInitializer GraphicsPSOInit;
-	Context.RHICmdList.ApplyCachedRenderTargets(GraphicsPSOInit);
+		FGraphicsPipelineStateInitializer GraphicsPSOInit;
+		Context.RHICmdList.ApplyCachedRenderTargets(GraphicsPSOInit);
 
-	// set the state
-	GraphicsPSOInit.BlendState = TStaticBlendState<CW_RGB, BO_Add, BF_One, BF_One, BO_Add, BF_One, BF_One>::GetRHI();
-	GraphicsPSOInit.RasterizerState = TStaticRasterizerState<>::GetRHI();
-	GraphicsPSOInit.DepthStencilState = TStaticDepthStencilState<false, CF_Always>::GetRHI();
+		// set the state
+		GraphicsPSOInit.BlendState = TStaticBlendState<CW_RGB, BO_Add, BF_One, BF_One, BO_Add, BF_One, BF_One>::GetRHI();
+		GraphicsPSOInit.RasterizerState = TStaticRasterizerState<>::GetRHI();
+		GraphicsPSOInit.DepthStencilState = TStaticDepthStencilState<false, CF_Always>::GetRHI();
 
-	FScene* Scene = ViewFamily.Scene->GetRenderScene();
-	check(Scene);
-	uint32 NumReflectionCaptures = Scene->ReflectionSceneData.RegisteredReflectionCaptures.Num();
-	
-	// Ambient cubemap specular will be applied in the reflection environment pass if it is enabled
-	//const bool bApplySpecular = View.Family->EngineShowFlags.ReflectionEnvironment == 0 || NumReflectionCaptures == 0;
-	//bool bClearCoatNeeded = (View.ShadingModelMaskInView & (1 << MSM_ClearCoat)) != 0;
+		FScene* Scene = ViewFamily.Scene->GetRenderScene();
+		check(Scene);
+		uint32 NumReflectionCaptures = Scene->ReflectionSceneData.RegisteredReflectionCaptures.Num();
 
-	Render(Context, GraphicsPSOInit);
+		// Ambient cubemap specular will be applied in the reflection environment pass if it is enabled
+		//const bool bApplySpecular = View.Family->EngineShowFlags.ReflectionEnvironment == 0 || NumReflectionCaptures == 0;
+		//bool bClearCoatNeeded = (View.ShadingModelMaskInView & (1 << MSM_ClearCoat)) != 0;
 
-	Context.RHICmdList.CopyToResolveTarget(DestRenderTarget.TargetableTexture, DestRenderTarget.ShaderResourceTexture, FResolveParams());
+		Render(Context, GraphicsPSOInit);
+	}
+	Context.RHICmdList.EndRenderPass();
 }
 
 FPooledRenderTargetDesc FRCPassPostProcessAmbient::ComputeOutputDesc(EPassOutputId InPassOutputId) const
diff --git a/Engine/Source/Runtime/Renderer/Private/CompositionLighting/PostProcessAmbientOcclusion.cpp b/Engine/Source/Runtime/Renderer/Private/CompositionLighting/PostProcessAmbientOcclusion.cpp
--- a/Engine/Source/Runtime/Renderer/Private/CompositionLighting/PostProcessAmbientOcclusion.cpp
+++ b/Engine/Source/Runtime/Renderer/Private/CompositionLighting/PostProcessAmbientOcclusion.cpp
@@ -378,57 +378,59 @@ void FRCPassPostProcessAmbientOcclusionSetup::Process(FRenderingCompositePassCon
 	SCOPED_DRAW_EVENTF(Context.RHICmdList, AmbientOcclusionSetup, TEXT("AmbientOcclusionSetup %dx%d"), DestRect.Width(), DestRect.Height());
 
 	// Set the view family's render target/viewport.
-	SetRenderTarget(Context.RHICmdList, DestRenderTarget.TargetableTexture, FTextureRHIParamRef());
+	FRHIRenderPassInfo RPInfo(DestRenderTarget.TargetableTexture, ERenderTargetActions::Load_Store, DestRenderTarget.ShaderResourceTexture);
+	Context.RHICmdList.BeginRenderPass(RPInfo, TEXT("AmbientOcclusionSetup"));
+	{
 
-	Context.SetViewportAndCallRHI(DestRect);
+		Context.SetViewportAndCallRHI(DestRect);
 
-	FGraphicsPipelineStateInitializer GraphicsPSOInit;
-	Context.RHICmdList.ApplyCachedRenderTargets(GraphicsPSOInit);
-	GraphicsPSOInit.BlendState = TStaticBlendState<>::GetRHI();
-	GraphicsPSOInit.RasterizerState = TStaticRasterizerState<>::GetRHI();
-	GraphicsPSOInit.DepthStencilState = TStaticDepthStencilState<false, CF_Always>::GetRHI();
-	GraphicsPSOInit.PrimitiveType = PT_TriangleList;
+		FGraphicsPipelineStateInitializer GraphicsPSOInit;
+		Context.RHICmdList.ApplyCachedRenderTargets(GraphicsPSOInit);
+		GraphicsPSOInit.BlendState = TStaticBlendState<>::GetRHI();
+		GraphicsPSOInit.RasterizerState = TStaticRasterizerState<>::GetRHI();
+		GraphicsPSOInit.DepthStencilState = TStaticDepthStencilState<false, CF_Always>::GetRHI();
+		GraphicsPSOInit.PrimitiveType = PT_TriangleList;
 
-	TShaderMapRef<FPostProcessVS> VertexShader(Context.GetShaderMap());
+		TShaderMapRef<FPostProcessVS> VertexShader(Context.GetShaderMap());
 
-	if(IsInitialPass())
-	{
-		TShaderMapRef<FPostProcessAmbientOcclusionSetupPS<1> > PixelShader(Context.GetShaderMap());
+		if (IsInitialPass())
+		{
+			TShaderMapRef<FPostProcessAmbientOcclusionSetupPS<1> > PixelShader(Context.GetShaderMap());
 
-		GraphicsPSOInit.BoundShaderState.VertexDeclarationRHI = GFilterVertexDeclaration.VertexDeclarationRHI;
-		GraphicsPSOInit.BoundShaderState.VertexShaderRHI = GETSAFERHISHADER_VERTEX(*VertexShader);
-		GraphicsPSOInit.BoundShaderState.PixelShaderRHI = GETSAFERHISHADER_PIXEL(*PixelShader);
-		SetGraphicsPipelineState(Context.RHICmdList, GraphicsPSOInit);
+			GraphicsPSOInit.BoundShaderState.VertexDeclarationRHI = GFilterVertexDeclaration.VertexDeclarationRHI;
+			GraphicsPSOInit.BoundShaderState.VertexShaderRHI = GETSAFERHISHADER_VERTEX(*VertexShader);
+			GraphicsPSOInit.BoundShaderState.PixelShaderRHI = GETSAFERHISHADER_PIXEL(*PixelShader);
+			SetGraphicsPipelineState(Context.RHICmdList, GraphicsPSOInit);
 
-		PixelShader->SetParameters(Context);
-	}
-	else
-	{
-		TShaderMapRef<FPostProcessAmbientOcclusionSetupPS<0> > PixelShader(Context.GetShaderMap());
+			PixelShader->SetParameters(Context);
+		}
+		else
+		{
+			TShaderMapRef<FPostProcessAmbientOcclusionSetupPS<0> > PixelShader(Context.GetShaderMap());
 
-		GraphicsPSOInit.BoundShaderState.VertexDeclarationRHI = GFilterVertexDeclaration.VertexDeclarationRHI;
-		GraphicsPSOInit.BoundShaderState.VertexShaderRHI = GETSAFERHISHADER_VERTEX(*VertexShader);
-		GraphicsPSOInit.BoundShaderState.PixelShaderRHI = GETSAFERHISHADER_PIXEL(*PixelShader);
-		SetGraphicsPipelineState(Context.RHICmdList, GraphicsPSOInit);
+			GraphicsPSOInit.BoundShaderState.VertexDeclarationRHI = GFilterVertexDeclaration.VertexDeclarationRHI;
+			GraphicsPSOInit.BoundShaderState.VertexShaderRHI = GETSAFERHISHADER_VERTEX(*VertexShader);
+			GraphicsPSOInit.BoundShaderState.PixelShaderRHI = GETSAFERHISHADER_PIXEL(*PixelShader);
+			SetGraphicsPipelineState(Context.RHICmdList, GraphicsPSOInit);
 
-		PixelShader->SetParameters(Context);
-	}
+			PixelShader->SetParameters(Context);
+		}
 
-	VertexShader->SetParameters(Context);
-	DrawPostProcessPass(
-		Context.RHICmdList,
-		0, 0,
-		DestRect.Width(), DestRect.Height(),
-		SrcRect.Min.X, SrcRect.Min.Y, 
-		SrcRect.Width(), SrcRect.Height(),
-		DestRect.Size(),
-		FSceneRenderTargets::Get(Context.RHICmdList).GetBufferSizeXY(),
-		*VertexShader,
-		View.StereoPass,
-		Context.HasHmdMesh(),
-		EDRF_UseTriangleOptimization);
-
-	Context.RHICmdList.CopyToResolveTarget(DestRenderTarget.TargetableTexture, DestRenderTarget.ShaderResourceTexture, FResolveParams());
+		VertexShader->SetParameters(Context);
+		DrawPostProcessPass(
+			Context.RHICmdList,
+			0, 0,
+			DestRect.Width(), DestRect.Height(),
+			SrcRect.Min.X, SrcRect.Min.Y,
+			SrcRect.Width(), SrcRect.Height(),
+			DestRect.Size(),
+			FSceneRenderTargets::Get(Context.RHICmdList).GetBufferSizeXY(),
+			*VertexShader,
+			View.StereoPass,
+			Context.HasHmdMesh(),
+			EDRF_UseTriangleOptimization);
+	}
+	Context.RHICmdList.EndRenderPass();
 }
 
 FPooledRenderTargetDesc FRCPassPostProcessAmbientOcclusionSetup::ComputeOutputDesc(EPassOutputId InPassOutputId) const
@@ -725,7 +727,8 @@ void FRCPassPostProcessAmbientOcclusion::ProcessCS(FRenderingCompositePassContex
 	} \
 	break
 
-	SetRenderTarget(Context.RHICmdList, nullptr, nullptr);
+	// #todo-renderpasses remove once everything is renderpasses
+	UnbindRenderTargets(Context.RHICmdList);
 	Context.SetViewportAndCallRHI(ViewRect, 0.0f, 1.0f);
 
 	//for async compute we need to set up a fence to make sure the resource is ready before we start.
@@ -782,55 +785,61 @@ void FRCPassPostProcessAmbientOcclusion::ProcessPS(FRenderingCompositePassContex
 	const bool bDepthBoundsTestEnabled = GSupportsDepthBoundsTest && SceneDepthBuffer && CVarAmbientOcclusionDepthBoundsTest.GetValueOnRenderThread() && SceneDepthBuffer->TargetableTexture->GetNumSamples() == 1;
 
 	// Set the view family's render target/viewport.
-
-	SetRenderTarget(Context.RHICmdList, DestRenderTarget->TargetableTexture,
-		bDepthBoundsTestEnabled ? SceneDepthBuffer->TargetableTexture : FTextureRHIRef(),
-		ESimpleRenderTargetMode::EUninitializedColorExistingDepth, // Color target will be fully over-written and old contents is not needed.
-		FExclusiveDepthStencil::DepthRead_StencilRead); // DepthStencil buffer will not be modified, but will be used for culling.
-
-	Context.SetViewportAndCallRHI(ViewRect);
-
-	float DepthFar = 0.0f;
-
+	// Rendertarget will be completely overwritten.
+	FRHIRenderPassInfo RPInfo(DestRenderTarget->TargetableTexture, ERenderTargetActions::DontLoad_Store);
 	if (bDepthBoundsTestEnabled)
 	{
-		const FFinalPostProcessSettings& Settings = Context.View.FinalPostProcessSettings;
-		const FMatrix& ProjectionMatrix = Context.View.ViewMatrices.GetProjectionMatrix();
-		const FVector4 Far = ProjectionMatrix.TransformFVector4(FVector4(0, 0, Settings.AmbientOcclusionFadeDistance));
-		DepthFar = FMath::Min(1.0f, Far.Z / Far.W);
+		// We'll use the depth/stencil buffer for read but it will not be modified.
+		RPInfo.DepthStencilRenderTarget.DepthStencilTarget = SceneDepthBuffer->TargetableTexture;
+		RPInfo.DepthStencilRenderTarget.Action = MakeDepthStencilTargetActions(ERenderTargetActions::Load_DontStore, ERenderTargetActions::Load_DontStore);
+		RPInfo.DepthStencilRenderTarget.ExclusiveDepthStencil = FExclusiveDepthStencil::DepthRead_StencilRead;
+	}
 
-		static_assert(bool(ERHIZBuffer::IsInverted), "Inverted depth buffer is assumed when setting depth bounds test for AO.");
+	Context.RHICmdList.BeginRenderPass(RPInfo, TEXT("PSAmbientOcclusion"));
+	{
+		Context.SetViewportAndCallRHI(ViewRect);
 
-		// We must clear all pixels that won't be touched by AO shader.
-		FClearQuadCallbacks Callbacks;
-		Callbacks.PSOModifier = [](FGraphicsPipelineStateInitializer& PSOInitializer)
-		{
-			PSOInitializer.bDepthBounds = true;
-		};
-		Callbacks.PreClear = [DepthFar](FRHICommandList& InRHICmdList)
-		{
-			// This is done by rendering a clear quad over a depth range from AmbientOcclusionFadeDistance to far plane.
-			InRHICmdList.SetDepthBounds(0, DepthFar);	// NOTE: Inverted depth
-		};
-		Callbacks.PostClear = [DepthFar](FRHICommandList& InRHICmdList)
+		float DepthFar = 0.0f;
+
+		if (bDepthBoundsTestEnabled)
 		{
-			// Set depth bounds test to cover everything from near plane to AmbientOcclusionFadeDistance and run AO pixel shader.
-			InRHICmdList.SetDepthBounds(DepthFar, 1.0f);
-		};
-		DrawClearQuad(Context.RHICmdList, FLinearColor::White, Callbacks);
-	}
+			const FFinalPostProcessSettings& Settings = Context.View.FinalPostProcessSettings;
+			const FMatrix& ProjectionMatrix = Context.View.ViewMatrices.GetProjectionMatrix();
+			const FVector4 Far = ProjectionMatrix.TransformFVector4(FVector4(0, 0, Settings.AmbientOcclusionFadeDistance));
+			DepthFar = FMath::Min(1.0f, Far.Z / Far.W);
 
-	FGraphicsPipelineStateInitializer GraphicsPSOInit;
-	Context.RHICmdList.ApplyCachedRenderTargets(GraphicsPSOInit);
+			static_assert(bool(ERHIZBuffer::IsInverted), "Inverted depth buffer is assumed when setting depth bounds test for AO.");
 
-	// set the state
-	GraphicsPSOInit.BlendState = TStaticBlendState<>::GetRHI();
-	GraphicsPSOInit.RasterizerState = TStaticRasterizerState<>::GetRHI();
-	GraphicsPSOInit.DepthStencilState = TStaticDepthStencilState<false, CF_Always>::GetRHI();
-	GraphicsPSOInit.PrimitiveType = PT_TriangleList;
-	GraphicsPSOInit.bDepthBounds = bDepthBoundsTestEnabled;
+			// We must clear all pixels that won't be touched by AO shader.
+			FClearQuadCallbacks Callbacks;
+			Callbacks.PSOModifier = [](FGraphicsPipelineStateInitializer& PSOInitializer)
+			{
+				PSOInitializer.bDepthBounds = true;
+			};
+			Callbacks.PreClear = [DepthFar](FRHICommandList& InRHICmdList)
+			{
+				// This is done by rendering a clear quad over a depth range from AmbientOcclusionFadeDistance to far plane.
+				InRHICmdList.SetDepthBounds(0, DepthFar);	// NOTE: Inverted depth
+			};
+			Callbacks.PostClear = [DepthFar](FRHICommandList& InRHICmdList)
+			{
+				// Set depth bounds test to cover everything from near plane to AmbientOcclusionFadeDistance and run AO pixel shader.
+				InRHICmdList.SetDepthBounds(DepthFar, 1.0f);
+			};
+			DrawClearQuad(Context.RHICmdList, FLinearColor::White, Callbacks);
+		}
+
+		FGraphicsPipelineStateInitializer GraphicsPSOInit;
+		Context.RHICmdList.ApplyCachedRenderTargets(GraphicsPSOInit);
 
-	FShader* VertexShader = 0;
+		// set the state
+		GraphicsPSOInit.BlendState = TStaticBlendState<>::GetRHI();
+		GraphicsPSOInit.RasterizerState = TStaticRasterizerState<>::GetRHI();
+		GraphicsPSOInit.DepthStencilState = TStaticDepthStencilState<false, CF_Always>::GetRHI();
+		GraphicsPSOInit.PrimitiveType = PT_TriangleList;
+		GraphicsPSOInit.bDepthBounds = bDepthBoundsTestEnabled;
+
+		FShader* VertexShader = 0;
 
 #define SET_SHADER_CASE(ShaderQualityCase) \
 		case ShaderQualityCase: \
@@ -846,34 +855,36 @@ void FRCPassPostProcessAmbientOcclusion::ProcessPS(FRenderingCompositePassContex
 	} \
 	break
 
-	switch (ShaderQuality)
-	{
-		SET_SHADER_CASE(0);
-		SET_SHADER_CASE(1);
-		SET_SHADER_CASE(2);
-		SET_SHADER_CASE(3);
-		SET_SHADER_CASE(4);
-	default:
-		break;
-	};
+		switch (ShaderQuality)
+		{
+			SET_SHADER_CASE(0);
+			SET_SHADER_CASE(1);
+			SET_SHADER_CASE(2);
+			SET_SHADER_CASE(3);
+			SET_SHADER_CASE(4);
+		default:
+			break;
+		};
 #undef SET_SHADER_CASE
 
-	if (bDepthBoundsTestEnabled)
-	{
-		Context.RHICmdList.SetDepthBounds(DepthFar, 1.0f);
-	}
+		if (bDepthBoundsTestEnabled)
+		{
+			Context.RHICmdList.SetDepthBounds(DepthFar, 1.0f);
+		}
 
-	// Draw a quad mapping scene color to the view's render target
-	DrawRectangle(
-		Context.RHICmdList,
-		0, 0,
-		ViewRect.Width(), ViewRect.Height(),
-		ViewRect.Min.X, ViewRect.Min.Y,
-		ViewRect.Width(), ViewRect.Height(),
-		ViewRect.Size(),
-		TexSize,
-		VertexShader,
-		EDRF_UseTriangleOptimization);
+		// Draw a quad mapping scene color to the view's render target
+		DrawRectangle(
+			Context.RHICmdList,
+			0, 0,
+			ViewRect.Width(), ViewRect.Height(),
+			ViewRect.Min.X, ViewRect.Min.Y,
+			ViewRect.Width(), ViewRect.Height(),
+			ViewRect.Size(),
+			TexSize,
+			VertexShader,
+			EDRF_UseTriangleOptimization);
+	}
+	Context.RHICmdList.EndRenderPass();
 
 	Context.RHICmdList.TransitionResource(EResourceTransitionAccess::EReadable, DestRenderTarget->TargetableTexture);
 
@@ -1036,44 +1047,47 @@ void FRCPassPostProcessBasePassAO::Process(FRenderingCompositePassContext& Conte
 
 	// Set the view family's render target/viewport.
 	Context.RHICmdList.TransitionResource(EResourceTransitionAccess::EWritable, DestRenderTarget.TargetableTexture);
-	SetRenderTarget(Context.RHICmdList, DestRenderTarget.TargetableTexture,	FTextureRHIParamRef(), ESimpleRenderTargetMode::EExistingColorAndDepth);
-	Context.SetViewportAndCallRHI(View.ViewRect);
 
-	FGraphicsPipelineStateInitializer GraphicsPSOInit;
-	Context.RHICmdList.ApplyCachedRenderTargets(GraphicsPSOInit);
+	FRHIRenderPassInfo RPInfo(DestRenderTarget.TargetableTexture, ERenderTargetActions::Load_Store, DestRenderTarget.ShaderResourceTexture);
+	Context.RHICmdList.BeginRenderPass(RPInfo, TEXT("BasePassAmbientOcclusion"));
+	{
+		Context.SetViewportAndCallRHI(View.ViewRect);
 
-	// set the state
-	GraphicsPSOInit.BlendState = TStaticBlendState<CW_RGBA, BO_Add, BF_DestColor, BF_Zero, BO_Add, BF_DestAlpha, BF_Zero>::GetRHI();
-	GraphicsPSOInit.RasterizerState = TStaticRasterizerState<>::GetRHI();
-	GraphicsPSOInit.DepthStencilState = TStaticDepthStencilState<false, CF_Always>::GetRHI();
+		FGraphicsPipelineStateInitializer GraphicsPSOInit;
+		Context.RHICmdList.ApplyCachedRenderTargets(GraphicsPSOInit);
 
-	TShaderMapRef<FPostProcessVS> VertexShader(Context.GetShaderMap());
-	TShaderMapRef<FPostProcessBasePassAOPS> PixelShader(Context.GetShaderMap());
+		// set the state
+		GraphicsPSOInit.BlendState = TStaticBlendState<CW_RGBA, BO_Add, BF_DestColor, BF_Zero, BO_Add, BF_DestAlpha, BF_Zero>::GetRHI();
+		GraphicsPSOInit.RasterizerState = TStaticRasterizerState<>::GetRHI();
+		GraphicsPSOInit.DepthStencilState = TStaticDepthStencilState<false, CF_Always>::GetRHI();
 
-	GraphicsPSOInit.BoundShaderState.VertexDeclarationRHI = GFilterVertexDeclaration.VertexDeclarationRHI;
-	GraphicsPSOInit.BoundShaderState.VertexShaderRHI = GETSAFERHISHADER_VERTEX(*VertexShader);
-	GraphicsPSOInit.BoundShaderState.PixelShaderRHI = GETSAFERHISHADER_PIXEL(*PixelShader);
-	GraphicsPSOInit.PrimitiveType = PT_TriangleList;
+		TShaderMapRef<FPostProcessVS> VertexShader(Context.GetShaderMap());
+		TShaderMapRef<FPostProcessBasePassAOPS> PixelShader(Context.GetShaderMap());
 
-	SetGraphicsPipelineState(Context.RHICmdList, GraphicsPSOInit);
-	
-	VertexShader->SetParameters(Context);
-	PixelShader->SetParameters(Context.RHICmdList, Context, SceneContext.GetBufferSizeXY());
-
-	DrawPostProcessPass(
-		Context.RHICmdList,
-		0, 0,
-		View.ViewRect.Width(), View.ViewRect.Height(),
-		View.ViewRect.Min.X, View.ViewRect.Min.Y,
-		View.ViewRect.Width(), View.ViewRect.Height(),
-		View.ViewRect.Size(),
-		SceneContext.GetBufferSizeXY(),
-		*VertexShader,
-		View.StereoPass, 
-		Context.HasHmdMesh(),
-		EDRF_UseTriangleOptimization);
-
-	Context.RHICmdList.CopyToResolveTarget(DestRenderTarget.TargetableTexture, DestRenderTarget.ShaderResourceTexture, FResolveParams());
+		GraphicsPSOInit.BoundShaderState.VertexDeclarationRHI = GFilterVertexDeclaration.VertexDeclarationRHI;
+		GraphicsPSOInit.BoundShaderState.VertexShaderRHI = GETSAFERHISHADER_VERTEX(*VertexShader);
+		GraphicsPSOInit.BoundShaderState.PixelShaderRHI = GETSAFERHISHADER_PIXEL(*PixelShader);
+		GraphicsPSOInit.PrimitiveType = PT_TriangleList;
+
+		SetGraphicsPipelineState(Context.RHICmdList, GraphicsPSOInit);
+
+		VertexShader->SetParameters(Context);
+		PixelShader->SetParameters(Context.RHICmdList, Context, SceneContext.GetBufferSizeXY());
+
+		DrawPostProcessPass(
+			Context.RHICmdList,
+			0, 0,
+			View.ViewRect.Width(), View.ViewRect.Height(),
+			View.ViewRect.Min.X, View.ViewRect.Min.Y,
+			View.ViewRect.Width(), View.ViewRect.Height(),
+			View.ViewRect.Size(),
+			SceneContext.GetBufferSizeXY(),
+			*VertexShader,
+			View.StereoPass,
+			Context.HasHmdMesh(),
+			EDRF_UseTriangleOptimization);
+	}
+	Context.RHICmdList.EndRenderPass();
 }
 
 FPooledRenderTargetDesc FRCPassPostProcessBasePassAO::ComputeOutputDesc(EPassOutputId InPassOutputId) const
diff --git a/Engine/Source/Runtime/Renderer/Private/CompositionLighting/PostProcessLpvIndirect.cpp b/Engine/Source/Runtime/Renderer/Private/CompositionLighting/PostProcessLpvIndirect.cpp
--- a/Engine/Source/Runtime/Renderer/Private/CompositionLighting/PostProcessLpvIndirect.cpp
+++ b/Engine/Source/Runtime/Renderer/Private/CompositionLighting/PostProcessLpvIndirect.cpp
@@ -280,72 +280,76 @@ void FRCPassPostProcessLpvIndirect::Process(FRenderingCompositePassContext& Cont
 		NumRenderTargets = 2;
 	}
 
-	SetRenderTargets(Context.RHICmdList, NumRenderTargets, RenderTargets, 0, ESimpleRenderTargetMode::EExistingColorAndDepth, FExclusiveDepthStencil::DepthRead_StencilNop);
-	Context.SetViewportAndCallRHI(View.ViewRect);
-
-	FGraphicsPipelineStateInitializer GraphicsPSOInit;
-	Context.RHICmdList.ApplyCachedRenderTargets(GraphicsPSOInit);
-	// set the state
-	GraphicsPSOInit.BlendState = TStaticBlendState<CW_RGB,BO_Add,BF_One,BF_One,BO_Add,BF_One,BF_One>::GetRHI();
-	GraphicsPSOInit.RasterizerState = TStaticRasterizerState<>::GetRHI();
-	GraphicsPSOInit.DepthStencilState = TStaticDepthStencilState<false,CF_Always>::GetRHI();
-
-	TShaderMapRef<FPostProcessVS> VertexShader(Context.GetShaderMap());
-	GraphicsPSOInit.BoundShaderState.VertexDeclarationRHI = GFilterVertexDeclaration.VertexDeclarationRHI;
-	GraphicsPSOInit.BoundShaderState.VertexShaderRHI = GETSAFERHISHADER_VERTEX(*VertexShader);
-
-	FPostProcessLpvIndirectPS* PixelShader = NULL;
-	if ( bApplySeparateSpecularRT )
+	FRHIRenderPassInfo RPInfo(NumRenderTargets, RenderTargets, ERenderTargetActions::Load_Store);
+	RPInfo.ColorRenderTargets[0].ResolveTarget = DestColorRenderTarget.ShaderResourceTexture;
+	if (bApplySeparateSpecularRT)
 	{
-		TShaderMapRef< TPostProcessLpvIndirectPS<true> > PixelShaderWithSpecular(Context.GetShaderMap());
-		PixelShader = (FPostProcessLpvIndirectPS*)*PixelShaderWithSpecular;
+		RPInfo.ColorRenderTargets[1].ResolveTarget = DestSpecularRenderTarget.ShaderResourceTexture;
 	}
-	else
+
+	Context.RHICmdList.BeginRenderPass(RPInfo, TEXT("LPVIndirect"));
 	{
-		TShaderMapRef< TPostProcessLpvIndirectPS<false> > PixelShaderNoSpecular(Context.GetShaderMap());
-		PixelShader = (FPostProcessLpvIndirectPS*)*PixelShaderNoSpecular;
-	}
+		Context.SetViewportAndCallRHI(View.ViewRect);
 
-	GraphicsPSOInit.BoundShaderState.PixelShaderRHI = GETSAFERHISHADER_PIXEL(PixelShader);
-	GraphicsPSOInit.PrimitiveType = PT_TriangleList;
+		FGraphicsPipelineStateInitializer GraphicsPSOInit;
+		Context.RHICmdList.ApplyCachedRenderTargets(GraphicsPSOInit);
+		// set the state
+		GraphicsPSOInit.BlendState = TStaticBlendState<CW_RGB, BO_Add, BF_One, BF_One, BO_Add, BF_One, BF_One>::GetRHI();
+		GraphicsPSOInit.RasterizerState = TStaticRasterizerState<>::GetRHI();
+		GraphicsPSOInit.DepthStencilState = TStaticDepthStencilState<false, CF_Always>::GetRHI();
 
-	SetGraphicsPipelineState(Context.RHICmdList, GraphicsPSOInit);
+		TShaderMapRef<FPostProcessVS> VertexShader(Context.GetShaderMap());
+		GraphicsPSOInit.BoundShaderState.VertexDeclarationRHI = GFilterVertexDeclaration.VertexDeclarationRHI;
+		GraphicsPSOInit.BoundShaderState.VertexShaderRHI = GETSAFERHISHADER_VERTEX(*VertexShader);
 
-	FLpvReadUniformBufferParameters	LpvReadUniformBufferParams;
-	FLpvReadUniformBufferRef LpvReadUniformBuffer;
+		FPostProcessLpvIndirectPS* PixelShader = NULL;
+		if (bApplySeparateSpecularRT)
+		{
+			TShaderMapRef< TPostProcessLpvIndirectPS<true> > PixelShaderWithSpecular(Context.GetShaderMap());
+			PixelShader = (FPostProcessLpvIndirectPS*)*PixelShaderWithSpecular;
+		}
+		else
+		{
+			TShaderMapRef< TPostProcessLpvIndirectPS<false> > PixelShaderNoSpecular(Context.GetShaderMap());
+			PixelShader = (FPostProcessLpvIndirectPS*)*PixelShaderNoSpecular;
+		}
 
-	LpvReadUniformBufferParams = Lpv->GetReadUniformBufferParams();
-	LpvReadUniformBuffer = FLpvReadUniformBufferRef::CreateUniformBufferImmediate( LpvReadUniformBufferParams, UniformBuffer_SingleDraw ); 
+		GraphicsPSOInit.BoundShaderState.PixelShaderRHI = GETSAFERHISHADER_PIXEL(PixelShader);
+		GraphicsPSOInit.PrimitiveType = PT_TriangleList;
 
-	FTextureRHIParamRef LpvBufferSrvs[7];
-	for ( int i = 0; i < 7; i++ ) 
-	{
-		LpvBufferSrvs[i] = Lpv->GetLpvBufferSrv(i);
-	}
+		SetGraphicsPipelineState(Context.RHICmdList, GraphicsPSOInit);
 
-	PixelShader->SetParameters(Context.RHICmdList, LpvBufferSrvs, Lpv->GetAOVolumeTextureSRV(), LpvReadUniformBuffer, Context );
+		FLpvReadUniformBufferParameters	LpvReadUniformBufferParams;
+		FLpvReadUniformBufferRef LpvReadUniformBuffer;
 
-	{
-		SCOPED_DRAW_EVENT(Context.RHICmdList, PostProcessLpvIndirect );
+		LpvReadUniformBufferParams = Lpv->GetReadUniformBufferParams();
+		LpvReadUniformBuffer = FLpvReadUniformBufferRef::CreateUniformBufferImmediate(LpvReadUniformBufferParams, UniformBuffer_SingleDraw);
 
-		DrawPostProcessPass(
-			Context.RHICmdList,
-			0, 0,
-			View.ViewRect.Width(), View.ViewRect.Height(),
-			View.ViewRect.Min.X, View.ViewRect.Min.Y,
-			View.ViewRect.Width(), View.ViewRect.Height(),
-			View.ViewRect.Size(),
-			SceneContext.GetBufferSizeXY(),
-			*VertexShader,
-			View.StereoPass,
-			Context.HasHmdMesh());
+		FTextureRHIParamRef LpvBufferSrvs[7];
+		for (int i = 0; i < 7; i++)
+		{
+			LpvBufferSrvs[i] = Lpv->GetLpvBufferSrv(i);
+		}
+
+		PixelShader->SetParameters(Context.RHICmdList, LpvBufferSrvs, Lpv->GetAOVolumeTextureSRV(), LpvReadUniformBuffer, Context);
 
-		Context.RHICmdList.CopyToResolveTarget(DestColorRenderTarget.TargetableTexture, DestColorRenderTarget.ShaderResourceTexture, FResolveParams());
-		if(bApplySeparateSpecularRT)
 		{
-			Context.RHICmdList.CopyToResolveTarget(DestSpecularRenderTarget.TargetableTexture, DestSpecularRenderTarget.ShaderResourceTexture, FResolveParams());
+			SCOPED_DRAW_EVENT(Context.RHICmdList, PostProcessLpvIndirect);
+
+			DrawPostProcessPass(
+				Context.RHICmdList,
+				0, 0,
+				View.ViewRect.Width(), View.ViewRect.Height(),
+				View.ViewRect.Min.X, View.ViewRect.Min.Y,
+				View.ViewRect.Width(), View.ViewRect.Height(),
+				View.ViewRect.Size(),
+				SceneContext.GetBufferSizeXY(),
+				*VertexShader,
+				View.StereoPass,
+				Context.HasHmdMesh());
 		}
 	}
+	Context.RHICmdList.EndRenderPass();
 
 	if ( LPVSettings.LPVDirectionalOcclusionIntensity > 0.0001f )
 	{
@@ -357,67 +361,70 @@ void FRCPassPostProcessLpvIndirect::DoDirectionalOcclusionPass(FRenderingComposi
 {
 	FSceneRenderTargets& SceneContext = FSceneRenderTargets::Get(Context.RHICmdList);
 
-	SCOPED_DRAW_EVENT(Context.RHICmdList, PostProcessLpvDirectionalOcclusion );
+	SCOPED_DRAW_EVENT(Context.RHICmdList, PostProcessLpvDirectionalOcclusion);
 	const FSceneRenderTargetItem& DestDirectionalOcclusionRenderTarget = SceneContext.DirectionalOcclusion->GetRenderTargetItem();
 	const FViewInfo& View = Context.View;
 	FSceneViewState* ViewState = (FSceneViewState*)View.State;
 
-	if(!ViewState)
+	if (!ViewState)
 	{
 		return;
 	}
 
 	const FFinalPostProcessSettings& PostprocessSettings = Context.View.FinalPostProcessSettings;
 	const FLightPropagationVolumeSettings& LPVSettings = PostprocessSettings.BlendableManager.GetSingleFinalDataConst<FLightPropagationVolumeSettings>();
-	
+
 	FLightPropagationVolume* Lpv = ViewState->GetLightPropagationVolume(Context.GetFeatureLevel());
 
-	if(!Lpv || LPVSettings.LPVIntensity == 0.0f)
+	if (!Lpv || LPVSettings.LPVIntensity == 0.0f)
 	{
 		return;
 	}
 
 	FTextureRHIParamRef RenderTarget = DestDirectionalOcclusionRenderTarget.TargetableTexture;
 
-	SetRenderTargets(Context.RHICmdList, 1, &RenderTarget, NULL, ESimpleRenderTargetMode::EExistingColorAndDepth, FExclusiveDepthStencil::DepthRead_StencilNop);
-
-	Context.SetViewportAndCallRHI(View.ViewRect);
+	FRHIRenderPassInfo RPInfo(RenderTarget, ERenderTargetActions::Clear_Store);
+	Context.RHICmdList.BeginRenderPass(RPInfo, TEXT("DoDirectionalOcclusionPass"));
+	{
+		Context.SetViewportAndCallRHI(View.ViewRect);
 
-	FGraphicsPipelineStateInitializer GraphicsPSOInit;
-	Context.RHICmdList.ApplyCachedRenderTargets(GraphicsPSOInit);
-	GraphicsPSOInit.BlendState = TStaticBlendState<>::GetRHI();
-	GraphicsPSOInit.RasterizerState = TStaticRasterizerState<>::GetRHI();
-	GraphicsPSOInit.DepthStencilState = TStaticDepthStencilState<false,CF_Always>::GetRHI();
-	TShaderMapRef<FPostProcessVS> VertexShader(View.ShaderMap);
+		FGraphicsPipelineStateInitializer GraphicsPSOInit;
+		Context.RHICmdList.ApplyCachedRenderTargets(GraphicsPSOInit);
+		GraphicsPSOInit.BlendState = TStaticBlendState<>::GetRHI();
+		GraphicsPSOInit.RasterizerState = TStaticRasterizerState<>::GetRHI();
+		GraphicsPSOInit.DepthStencilState = TStaticDepthStencilState<false, CF_Always>::GetRHI();
+		TShaderMapRef<FPostProcessVS> VertexShader(View.ShaderMap);
 
-	FLpvReadUniformBufferParameters	LpvReadUniformBufferParams;
-	FLpvReadUniformBufferRef LpvReadUniformBuffer;
+		FLpvReadUniformBufferParameters	LpvReadUniformBufferParams;
+		FLpvReadUniformBufferRef LpvReadUniformBuffer;
 
-	TShaderMapRef< FPostProcessLpvDirectionalOcclusionPS > PixelShader(View.ShaderMap);
+		TShaderMapRef< FPostProcessLpvDirectionalOcclusionPS > PixelShader(View.ShaderMap);
 
-	GraphicsPSOInit.BoundShaderState.VertexDeclarationRHI = GFilterVertexDeclaration.VertexDeclarationRHI;
-	GraphicsPSOInit.BoundShaderState.VertexShaderRHI = GETSAFERHISHADER_VERTEX(*VertexShader);
-	GraphicsPSOInit.BoundShaderState.PixelShaderRHI = GETSAFERHISHADER_PIXEL(*PixelShader);
-	GraphicsPSOInit.PrimitiveType = PT_TriangleList;
+		GraphicsPSOInit.BoundShaderState.VertexDeclarationRHI = GFilterVertexDeclaration.VertexDeclarationRHI;
+		GraphicsPSOInit.BoundShaderState.VertexShaderRHI = GETSAFERHISHADER_VERTEX(*VertexShader);
+		GraphicsPSOInit.BoundShaderState.PixelShaderRHI = GETSAFERHISHADER_PIXEL(*PixelShader);
+		GraphicsPSOInit.PrimitiveType = PT_TriangleList;
 
-	SetGraphicsPipelineState(Context.RHICmdList, GraphicsPSOInit);
+		SetGraphicsPipelineState(Context.RHICmdList, GraphicsPSOInit);
 
-	LpvReadUniformBufferParams = Lpv->GetReadUniformBufferParams();
-	LpvReadUniformBuffer = FLpvReadUniformBufferRef::CreateUniformBufferImmediate( LpvReadUniformBufferParams, UniformBuffer_SingleDraw ); 
+		LpvReadUniformBufferParams = Lpv->GetReadUniformBufferParams();
+		LpvReadUniformBuffer = FLpvReadUniformBufferRef::CreateUniformBufferImmediate(LpvReadUniformBufferParams, UniformBuffer_SingleDraw);
 
-	PixelShader->SetParameters( Lpv->GetAOVolumeTextureSRV(), LpvReadUniformBuffer, Context );
+		PixelShader->SetParameters(Lpv->GetAOVolumeTextureSRV(), LpvReadUniformBuffer, Context);
 
-	DrawPostProcessPass(
-		Context.RHICmdList,
-		0, 0,
-		View.ViewRect.Width(), View.ViewRect.Height(),
-		View.ViewRect.Min.X, View.ViewRect.Min.Y,
-		View.ViewRect.Width(), View.ViewRect.Height(),
-		View.ViewRect.Size(),
-		SceneContext.GetBufferSizeXY(),
-		*VertexShader, 
-		View.StereoPass, 
-		Context.HasHmdMesh());
+		DrawPostProcessPass(
+			Context.RHICmdList,
+			0, 0,
+			View.ViewRect.Width(), View.ViewRect.Height(),
+			View.ViewRect.Min.X, View.ViewRect.Min.Y,
+			View.ViewRect.Width(), View.ViewRect.Height(),
+			View.ViewRect.Size(),
+			SceneContext.GetBufferSizeXY(),
+			*VertexShader,
+			View.StereoPass,
+			Context.HasHmdMesh());
+	}
+	Context.RHICmdList.EndRenderPass();
 }
 
 FPooledRenderTargetDesc FRCPassPostProcessLpvIndirect::ComputeOutputDesc(EPassOutputId InPassOutputId) const
@@ -442,8 +449,6 @@ void FRCPassPostProcessVisualizeLPV::Process(FRenderingCompositePassContext& Con
 	const FSceneRenderTargetItem& DestRenderTarget = RenderTarget->GetRenderTargetItem();
 
 	// Set the view family's render target/viewport.
-	SetRenderTarget(Context.RHICmdList, DestRenderTarget.TargetableTexture, FTextureRHIRef());
-
 	{
 		FRenderTargetTemp TempRenderTarget(View, (const FTexture2DRHIRef&)DestRenderTarget.TargetableTexture);
 		FCanvas Canvas(&TempRenderTarget, NULL, ViewFamily.CurrentRealTime, ViewFamily.CurrentWorldTime, ViewFamily.DeltaWorldTime, View.GetFeatureLevel());
@@ -480,7 +485,6 @@ void FRCPassPostProcessVisualizeLPV::Process(FRenderingCompositePassContext& Con
 
 		Canvas.Flush_RenderThread(Context.RHICmdList);
 	}
-
 	Context.RHICmdList.CopyToResolveTarget(DestRenderTarget.TargetableTexture, DestRenderTarget.ShaderResourceTexture, FResolveParams());
 	
 	// to satify following passws
diff --git a/Engine/Source/Runtime/Renderer/Private/CompositionLighting/PostProcessPassThrough.cpp b/Engine/Source/Runtime/Renderer/Private/CompositionLighting/PostProcessPassThrough.cpp
--- a/Engine/Source/Runtime/Renderer/Private/CompositionLighting/PostProcessPassThrough.cpp
+++ b/Engine/Source/Runtime/Renderer/Private/CompositionLighting/PostProcessPassThrough.cpp
@@ -86,55 +86,57 @@ void FRCPassPostProcessPassThrough::Process(FRenderingCompositePassContext& Cont
 	checkf(DestRect.Size() == SrcRect.Size(), TEXT("Pass through should not be used as upscaling pass."));
 
 	// Set the view family's render target/viewport.
-	SetRenderTarget(Context.RHICmdList, DestRenderTarget.TargetableTexture, FTextureRHIRef());
-	Context.SetViewportAndCallRHI(0, 0, 0.0f, DestSize.X, DestSize.Y, 1.0f);
-
-	FGraphicsPipelineStateInitializer GraphicsPSOInit;
-	Context.RHICmdList.ApplyCachedRenderTargets(GraphicsPSOInit);
-
-	// set the state
-	if(bAdditiveBlend)
-	{
-		GraphicsPSOInit.BlendState = TStaticBlendState<CW_RGB, BO_Add, BF_One, BF_One, BO_Add, BF_One, BF_One>::GetRHI();
-	}
-	else
+	FRHIRenderPassInfo RPInfo(DestRenderTarget.TargetableTexture, ERenderTargetActions::Load_Store, DestRenderTarget.ShaderResourceTexture);
+	Context.RHICmdList.BeginRenderPass(RPInfo, TEXT("Passthrough"));
 	{
-		GraphicsPSOInit.BlendState = TStaticBlendState<>::GetRHI();
-	}
+		Context.SetViewportAndCallRHI(0, 0, 0.0f, DestSize.X, DestSize.Y, 1.0f);
 
-	GraphicsPSOInit.RasterizerState = TStaticRasterizerState<>::GetRHI();
-	GraphicsPSOInit.DepthStencilState = TStaticDepthStencilState<false, CF_Always>::GetRHI();
+		FGraphicsPipelineStateInitializer GraphicsPSOInit;
+		Context.RHICmdList.ApplyCachedRenderTargets(GraphicsPSOInit);
 
-	TShaderMapRef<FPostProcessVS> VertexShader(Context.GetShaderMap());
-	TShaderMapRef<FPostProcessPassThroughPS> PixelShader(Context.GetShaderMap());
+		// set the state
+		if (bAdditiveBlend)
+		{
+			GraphicsPSOInit.BlendState = TStaticBlendState<CW_RGB, BO_Add, BF_One, BF_One, BO_Add, BF_One, BF_One>::GetRHI();
+		}
+		else
+		{
+			GraphicsPSOInit.BlendState = TStaticBlendState<>::GetRHI();
+		}
 
-	GraphicsPSOInit.BoundShaderState.VertexDeclarationRHI = GFilterVertexDeclaration.VertexDeclarationRHI;
-	GraphicsPSOInit.BoundShaderState.VertexShaderRHI = GETSAFERHISHADER_VERTEX(*VertexShader);
-	GraphicsPSOInit.BoundShaderState.PixelShaderRHI = GETSAFERHISHADER_PIXEL(*PixelShader);
-	GraphicsPSOInit.PrimitiveType = PT_TriangleList;
+		GraphicsPSOInit.RasterizerState = TStaticRasterizerState<>::GetRHI();
+		GraphicsPSOInit.DepthStencilState = TStaticDepthStencilState<false, CF_Always>::GetRHI();
 
-	SetGraphicsPipelineState(Context.RHICmdList, GraphicsPSOInit);
+		TShaderMapRef<FPostProcessVS> VertexShader(Context.GetShaderMap());
+		TShaderMapRef<FPostProcessPassThroughPS> PixelShader(Context.GetShaderMap());
 
-	VertexShader->SetParameters(Context);
-	PixelShader->SetParameters(Context.RHICmdList, Context);
+		GraphicsPSOInit.BoundShaderState.VertexDeclarationRHI = GFilterVertexDeclaration.VertexDeclarationRHI;
+		GraphicsPSOInit.BoundShaderState.VertexShaderRHI = GETSAFERHISHADER_VERTEX(*VertexShader);
+		GraphicsPSOInit.BoundShaderState.PixelShaderRHI = GETSAFERHISHADER_PIXEL(*PixelShader);
+		GraphicsPSOInit.PrimitiveType = PT_TriangleList;
 
-	DrawPostProcessPass(
-		Context.RHICmdList,
-		DestRect.Min.X, DestRect.Min.Y,
-		DestRect.Width(), DestRect.Height(),
-		SrcRect.Min.X, SrcRect.Min.Y,
-		SrcRect.Width(), SrcRect.Height(),
-		DestSize,
-		SrcSize,
-		*VertexShader,
-		View.StereoPass,
-		Context.HasHmdMesh(),
-		EDRF_UseTriangleOptimization);
+		SetGraphicsPipelineState(Context.RHICmdList, GraphicsPSOInit);
 
-	// Draw custom data (like legends) for derived types
-	DrawCustom(Context);
+		VertexShader->SetParameters(Context);
+		PixelShader->SetParameters(Context.RHICmdList, Context);
 
-	Context.RHICmdList.CopyToResolveTarget(DestRenderTarget.TargetableTexture, DestRenderTarget.ShaderResourceTexture, FResolveParams());
+		DrawPostProcessPass(
+			Context.RHICmdList,
+			DestRect.Min.X, DestRect.Min.Y,
+			DestRect.Width(), DestRect.Height(),
+			SrcRect.Min.X, SrcRect.Min.Y,
+			SrcRect.Width(), SrcRect.Height(),
+			DestSize,
+			SrcSize,
+			*VertexShader,
+			View.StereoPass,
+			Context.HasHmdMesh(),
+			EDRF_UseTriangleOptimization);
+
+		// Draw custom data (like legends) for derived types
+		DrawCustom(Context);
+	}
+	Context.RHICmdList.EndRenderPass();
 }
 
 FPooledRenderTargetDesc FRCPassPostProcessPassThrough::ComputeOutputDesc(EPassOutputId InPassOutputId) const
