Renderpass updates.
Self-contained texture/shading model visualization.

diff --git a/Engine/Source/Runtime/Renderer/Private/PostProcess/VisualizeShadingModels.cpp b/Engine/Source/Runtime/Renderer/Private/PostProcess/VisualizeShadingModels.cpp
--- a/Engine/Source/Runtime/Renderer/Private/PostProcess/VisualizeShadingModels.cpp
+++ b/Engine/Source/Runtime/Renderer/Private/PostProcess/VisualizeShadingModels.cpp
@@ -105,37 +105,42 @@ void FRCPassPostProcessVisualizeShadingModels::Process(FRenderingCompositePassCo
 	const FSceneRenderTargetItem& DestRenderTarget = PassOutputs[0].RequestSurface(Context);
 
 	// Set the view family's render target/viewport.
-	SetRenderTarget(Context.RHICmdList, DestRenderTarget.TargetableTexture, FTextureRHIRef());
-	Context.SetViewportAndCallRHI(DestRect);
-
-	FGraphicsPipelineStateInitializer GraphicsPSOInit;
-	Context.RHICmdList.ApplyCachedRenderTargets(GraphicsPSOInit);
-	GraphicsPSOInit.BlendState = TStaticBlendState<>::GetRHI();
-	GraphicsPSOInit.RasterizerState = TStaticRasterizerState<>::GetRHI();
-	GraphicsPSOInit.DepthStencilState = TStaticDepthStencilState<false, CF_Always>::GetRHI();
-
-	TShaderMapRef<FPostProcessVS> VertexShader(Context.GetShaderMap());
-	TShaderMapRef<FPostProcessVisualizeShadingModelsPS> PixelShader(Context.GetShaderMap());
-
-	GraphicsPSOInit.BoundShaderState.VertexDeclarationRHI = GFilterVertexDeclaration.VertexDeclarationRHI;
-	GraphicsPSOInit.BoundShaderState.VertexShaderRHI = GETSAFERHISHADER_VERTEX(*VertexShader);
-	GraphicsPSOInit.BoundShaderState.PixelShaderRHI = GETSAFERHISHADER_PIXEL(*PixelShader);
-	GraphicsPSOInit.PrimitiveType = PT_TriangleList;
-	SetGraphicsPipelineState(Context.RHICmdList, GraphicsPSOInit);
-
-	PixelShader->SetPS(Context.RHICmdList, Context, ((FViewInfo&)View).ShadingModelMaskInView);
-
-	// Draw a quad mapping scene color to the view's render target
-	DrawRectangle(
-		Context.RHICmdList,
-		DestRect.Min.X, DestRect.Min.Y,
-		DestRect.Width(), DestRect.Height(),
-		SrcRect.Min.X, SrcRect.Min.Y,
-		SrcRect.Width(), SrcRect.Height(),
-		DestRect.Size(),
-		SrcSize,
-		*VertexShader,
-		EDRF_UseTriangleOptimization);
+	FRHIRenderPassInfo RPInfo(DestRenderTarget.TargetableTexture, ERenderTargetActions::Load_Store);
+	Context.RHICmdList.BeginRenderPass(RPInfo, TEXT("VisualizeShadingModels"));
+	{
+
+		Context.SetViewportAndCallRHI(DestRect);
+
+		FGraphicsPipelineStateInitializer GraphicsPSOInit;
+		Context.RHICmdList.ApplyCachedRenderTargets(GraphicsPSOInit);
+		GraphicsPSOInit.BlendState = TStaticBlendState<>::GetRHI();
+		GraphicsPSOInit.RasterizerState = TStaticRasterizerState<>::GetRHI();
+		GraphicsPSOInit.DepthStencilState = TStaticDepthStencilState<false, CF_Always>::GetRHI();
+
+		TShaderMapRef<FPostProcessVS> VertexShader(Context.GetShaderMap());
+		TShaderMapRef<FPostProcessVisualizeShadingModelsPS> PixelShader(Context.GetShaderMap());
+
+		GraphicsPSOInit.BoundShaderState.VertexDeclarationRHI = GFilterVertexDeclaration.VertexDeclarationRHI;
+		GraphicsPSOInit.BoundShaderState.VertexShaderRHI = GETSAFERHISHADER_VERTEX(*VertexShader);
+		GraphicsPSOInit.BoundShaderState.PixelShaderRHI = GETSAFERHISHADER_PIXEL(*PixelShader);
+		GraphicsPSOInit.PrimitiveType = PT_TriangleList;
+		SetGraphicsPipelineState(Context.RHICmdList, GraphicsPSOInit);
+
+		PixelShader->SetPS(Context.RHICmdList, Context, ((FViewInfo&)View).ShadingModelMaskInView);
+
+		// Draw a quad mapping scene color to the view's render target
+		DrawRectangle(
+			Context.RHICmdList,
+			DestRect.Min.X, DestRect.Min.Y,
+			DestRect.Width(), DestRect.Height(),
+			SrcRect.Min.X, SrcRect.Min.Y,
+			SrcRect.Width(), SrcRect.Height(),
+			DestRect.Size(),
+			SrcSize,
+			*VertexShader,
+			EDRF_UseTriangleOptimization);
+	}
+	Context.RHICmdList.EndRenderPass();
 
 	FRenderTargetTemp TempRenderTarget(View, (const FTexture2DRHIRef&)DestRenderTarget.TargetableTexture);
 	FCanvas Canvas(&TempRenderTarget, NULL, 0, 0, 0, Context.GetFeatureLevel());
diff --git a/Engine/Source/Runtime/Renderer/Private/PostProcess/VisualizeTexture.cpp b/Engine/Source/Runtime/Renderer/Private/PostProcess/VisualizeTexture.cpp
--- a/Engine/Source/Runtime/Renderer/Private/PostProcess/VisualizeTexture.cpp
+++ b/Engine/Source/Runtime/Renderer/Private/PostProcess/VisualizeTexture.cpp
@@ -386,9 +386,9 @@ void FVisualizeTexture::GenerateContent(FRHICommandListImmediate& RHICmdList, co
 	check(ViewRect != FIntRect(0, 0, 0, 0))
 
 
-	FTexture2DRHIRef VisTexture = (FTexture2DRHIRef&)RenderTargetItem.ShaderResourceTexture;
+		FTexture2DRHIRef VisTexture = (FTexture2DRHIRef&)RenderTargetItem.ShaderResourceTexture;
 
-	if(!IsValidRef(VisTexture) || !Desc.IsValid())
+	if (!IsValidRef(VisTexture) || !Desc.IsValid())
 	{
 		// todo: improve
 		return;
@@ -403,101 +403,104 @@ void FVisualizeTexture::GenerateContent(FRHICommandListImmediate& RHICmdList, co
 	Size.Y = FMath::Max(Size.Y, 1);
 
 	FPooledRenderTargetDesc OutputDesc(FPooledRenderTargetDesc::Create2DDesc(Size, PF_B8G8R8A8, FClearValueBinding(FLinearColor(1, 1, 0, 1)), TexCreate_None, TexCreate_RenderTargetable | TexCreate_ShaderResource, false));
-	
+
 	GRenderTargetPool.FindFreeElement(RHICmdList, OutputDesc, VisualizeTextureContent, TEXT("VisualizeTexture"));
 
-	if(!VisualizeTextureContent)
+	if (!VisualizeTextureContent)
 	{
 		return;
 	}
 
 	const FSceneRenderTargetItem& DestRenderTarget = VisualizeTextureContent->GetRenderTargetItem();
 
-	TransitionSetRenderTargetsHelper(RHICmdList, DestRenderTarget.TargetableTexture, FTextureRHIParamRef(), FExclusiveDepthStencil::DepthWrite_StencilWrite);
-
-	FRHIRenderTargetView RtView = FRHIRenderTargetView(DestRenderTarget.TargetableTexture, ERenderTargetLoadAction::EClear);
-	FRHISetRenderTargetsInfo Info(1, &RtView, FRHIDepthRenderTargetView());
-	RHICmdList.SetRenderTargetsAndClear(Info);
+	FRHIRenderPassInfo RPInfo(DestRenderTarget.TargetableTexture, ERenderTargetActions::Clear_Store, DestRenderTarget.ShaderResourceTexture);
+	TransitionRenderPassTargets(RHICmdList, RPInfo);
+	
+	RHICmdList.BeginRenderPass(RPInfo, TEXT("VisualizeTexture"));
+	{
 
-	FIntPoint RTExtent = FSceneRenderTargets::Get(RHICmdList).GetBufferSizeXY();
+		FIntPoint RTExtent = FSceneRenderTargets::Get(RHICmdList).GetBufferSizeXY();
 
-	FVector2D Tex00 = FVector2D(0, 0);
-	FVector2D Tex11 = FVector2D(1, 1);
+		FVector2D Tex00 = FVector2D(0, 0);
+		FVector2D Tex11 = FVector2D(1, 1);
 
-	uint32 LocalVisualizeTextureInputMapping = UVInputMapping;
+		uint32 LocalVisualizeTextureInputMapping = UVInputMapping;
 
-	if(!Desc.Is2DTexture())
-	{
-		LocalVisualizeTextureInputMapping = 1;
-	}
+		if (!Desc.Is2DTexture())
+		{
+			LocalVisualizeTextureInputMapping = 1;
+		}
 
-	// set UV
-	switch(LocalVisualizeTextureInputMapping)
-	{
-		// UV in left top
+		// set UV
+		switch (LocalVisualizeTextureInputMapping)
+		{
+			// UV in left top
 		case 0:
 			Tex11 = FVector2D((float)ViewRect.Width() / RTExtent.X, (float)ViewRect.Height() / RTExtent.Y);
 			break;
 
-		// whole texture
+			// whole texture
 		default:
 			break;
-	}
+		}
 
-	bool bIsDefault = StencilSRVSrc == GBlackTexture->TextureRHI;
-	bool bDepthStencil = Desc.Is2DTexture() && Desc.Format == PF_DepthStencil;
+		bool bIsDefault = StencilSRVSrc == GBlackTexture->TextureRHI;
+		bool bDepthStencil = Desc.Is2DTexture() && Desc.Format == PF_DepthStencil;
 
-	//clear if this is a new different Stencil buffer, or it's not a stencil buffer and we haven't switched to the default yet.
-	bool bNeedsClear = bDepthStencil && (StencilSRVSrc != RenderTargetItem.TargetableTexture);
-	bNeedsClear |= !bDepthStencil && !bIsDefault;
-	if (bNeedsClear)
-	{
-		StencilSRVSrc = nullptr;
-		StencilSRV.SafeRelease();
-	}
+		//clear if this is a new different Stencil buffer, or it's not a stencil buffer and we haven't switched to the default yet.
+		bool bNeedsClear = bDepthStencil && (StencilSRVSrc != RenderTargetItem.TargetableTexture);
+		bNeedsClear |= !bDepthStencil && !bIsDefault;
+		if (bNeedsClear)
+		{
+			StencilSRVSrc = nullptr;
+			StencilSRV.SafeRelease();
+		}
 
-	//always set something into the StencilSRV slot for platforms that require a full resource binding, even if
-	//dynamic branching will cause them not to be used.	
-	if(bDepthStencil && !StencilSRVSrc)
-	{
-		StencilSRVSrc = RenderTargetItem.TargetableTexture;
-		StencilSRV = RHICreateShaderResourceView((FTexture2DRHIRef&) RenderTargetItem.TargetableTexture, 0, 1, PF_X24_G8);
-	}
-	else if(!StencilSRVSrc)
-	{
-		StencilSRVSrc = GBlackTexture->TextureRHI;
-		StencilSRV = RHICreateShaderResourceView((FTexture2DRHIRef&) GBlackTexture->TextureRHI, 0, 1, PF_B8G8R8A8);
-	}	
-
-	FVisualizeTextureData VisualizeTextureData(RenderTargetItem, Desc);
-
-	// distinguish between standard depth and shadow depth to produce more reasonable default value mapping in the pixel shader.
-	const bool bDepthTexture = (Desc.TargetableFlags & TexCreate_DepthStencilTargetable) != 0;
-	const bool bShadowDepth = (Desc.Format == PF_ShadowDepth);
-
-	VisualizeTextureData.RGBMul = RGBMul;
-	VisualizeTextureData.SingleChannelMul = SingleChannelMul;
-	VisualizeTextureData.SingleChannel = SingleChannel;
-	VisualizeTextureData.AMul = AMul;
-	VisualizeTextureData.Tex00 = Tex00;
-	VisualizeTextureData.Tex11 = Tex11;
-	VisualizeTextureData.bSaturateInsteadOfFrac = (Flags & 1) != 0;
-	VisualizeTextureData.InputValueMapping = bShadowDepth ? 2 : (bDepthTexture ? 1 : 0);
-	VisualizeTextureData.ArrayIndex = ArrayIndex;
-	VisualizeTextureData.CustomMip = CustomMip;
-	VisualizeTextureData.StencilSRV = StencilSRV;
-
-	if(!(Desc.Flags & TexCreate_CPUReadback))		// We cannot make a texture lookup on such elements
-	{	
-		SCOPED_DRAW_EVENT(RHICmdList, VisualizeTexture);
-		// continue rendering to HDR if necessary
-		RenderVisualizeTexture(RHICmdList, FeatureLevel, VisualizeTextureData);
+		//always set something into the StencilSRV slot for platforms that require a full resource binding, even if
+		//dynamic branching will cause them not to be used.	
+		if (bDepthStencil && !StencilSRVSrc)
+		{
+			StencilSRVSrc = RenderTargetItem.TargetableTexture;
+			StencilSRV = RHICreateShaderResourceView((FTexture2DRHIRef&)RenderTargetItem.TargetableTexture, 0, 1, PF_X24_G8);
+		}
+		else if (!StencilSRVSrc)
+		{
+			StencilSRVSrc = GBlackTexture->TextureRHI;
+			StencilSRV = RHICreateShaderResourceView((FTexture2DRHIRef&)GBlackTexture->TextureRHI, 0, 1, PF_B8G8R8A8);
+		}
+
+		FVisualizeTextureData VisualizeTextureData(RenderTargetItem, Desc);
+
+		// distinguish between standard depth and shadow depth to produce more reasonable default value mapping in the pixel shader.
+		const bool bDepthTexture = (Desc.TargetableFlags & TexCreate_DepthStencilTargetable) != 0;
+		const bool bShadowDepth = (Desc.Format == PF_ShadowDepth);
+
+		VisualizeTextureData.RGBMul = RGBMul;
+		VisualizeTextureData.SingleChannelMul = SingleChannelMul;
+		VisualizeTextureData.SingleChannel = SingleChannel;
+		VisualizeTextureData.AMul = AMul;
+		VisualizeTextureData.Tex00 = Tex00;
+		VisualizeTextureData.Tex11 = Tex11;
+		VisualizeTextureData.bSaturateInsteadOfFrac = (Flags & 1) != 0;
+		VisualizeTextureData.InputValueMapping = bShadowDepth ? 2 : (bDepthTexture ? 1 : 0);
+		VisualizeTextureData.ArrayIndex = ArrayIndex;
+		VisualizeTextureData.CustomMip = CustomMip;
+		VisualizeTextureData.StencilSRV = StencilSRV;
+
+		if (!(Desc.Flags & TexCreate_CPUReadback))		// We cannot make a texture lookup on such elements
+		{
+			SCOPED_DRAW_EVENT(RHICmdList, VisualizeTexture);
+			// continue rendering to HDR if necessary
+			RenderVisualizeTexture(RHICmdList, FeatureLevel, VisualizeTextureData);
+		}
 	}
 
-	{
+	RHICmdList.EndRenderPass();
+
+	/*{
 		SCOPED_DRAW_EVENT(RHICmdList, VisCopy);
 		RHICmdList.CopyToResolveTarget(DestRenderTarget.TargetableTexture, DestRenderTarget.ShaderResourceTexture, FResolveParams());
-	}
+	}*/
 
 	VisualizeTextureDesc = Desc;
 
@@ -697,11 +700,8 @@ void FVisualizeTexture::PresentContent(FRHICommandListImmediate& RHICmdList, con
 		if (DisplayExtent.X > 50 && DisplayExtent.Y > 50)
 		{
 			FRenderTargetPool::SMemoryStats MemoryStats = GRenderTargetPool.ComputeView();
-
-			SetRenderTarget(RHICmdList, View.Family->RenderTarget->GetRenderTargetTexture(), FTextureRHIRef());
 			RHICmdList.SetViewport(0, 0, 0.0f, FSceneRenderTargets::Get(RHICmdList).GetBufferSizeXY().X, FSceneRenderTargets::Get(RHICmdList).GetBufferSizeXY().Y, 1.0f);
 
-
 			FRenderTargetTemp TempRenderTarget(View, View.UnscaledViewRect.Size());
 			FCanvas Canvas(&TempRenderTarget, NULL, View.Family->CurrentRealTime, View.Family->CurrentWorldTime, View.Family->DeltaWorldTime, View.GetFeatureLevel());
 
@@ -859,44 +859,49 @@ void FVisualizeTexture::PresentContent(FRHICommandListImmediate& RHICmdList, con
 	FPooledRenderTargetDesc Desc = VisualizeTextureDesc;
 
 	auto& RenderTarget = View.Family->RenderTarget->GetRenderTargetTexture();
-	SetRenderTarget(RHICmdList, RenderTarget, FTextureRHIRef(), true);
-	RHICmdList.SetViewport(0, 0, 0.0f, RenderTarget->GetSizeX(), RenderTarget->GetSizeY(), 1.0f);
 
-	FGraphicsPipelineStateInitializer GraphicsPSOInit;
-	RHICmdList.ApplyCachedRenderTargets(GraphicsPSOInit);
-	GraphicsPSOInit.BlendState = TStaticBlendState<>::GetRHI();
-	GraphicsPSOInit.RasterizerState = TStaticRasterizerState<>::GetRHI();
-	GraphicsPSOInit.DepthStencilState = TStaticDepthStencilState<false, CF_Always>::GetRHI();
+	FRHIRenderPassInfo RPInfo(RenderTarget, ERenderTargetActions::Load_Store);
+	RHICmdList.BeginRenderPass(RPInfo, TEXT("VisualizeTexture"));
+	{
+		RHICmdList.SetViewport(0, 0, 0.0f, RenderTarget->GetSizeX(), RenderTarget->GetSizeY(), 1.0f);
 
-	auto ShaderMap = View.ShaderMap;
-	TShaderMapRef<FPostProcessVS> VertexShader(ShaderMap);
-	TShaderMapRef<FVisualizeTexturePresentPS> PixelShader(ShaderMap);
+		FGraphicsPipelineStateInitializer GraphicsPSOInit;
+		RHICmdList.ApplyCachedRenderTargets(GraphicsPSOInit);
+		GraphicsPSOInit.BlendState = TStaticBlendState<>::GetRHI();
+		GraphicsPSOInit.RasterizerState = TStaticRasterizerState<>::GetRHI();
+		GraphicsPSOInit.DepthStencilState = TStaticDepthStencilState<false, CF_Always>::GetRHI();
 
-	GraphicsPSOInit.BoundShaderState.VertexDeclarationRHI = GFilterVertexDeclaration.VertexDeclarationRHI;
-	GraphicsPSOInit.BoundShaderState.VertexShaderRHI = GETSAFERHISHADER_VERTEX(*VertexShader);
-	GraphicsPSOInit.BoundShaderState.PixelShaderRHI = GETSAFERHISHADER_PIXEL(*PixelShader);
-	GraphicsPSOInit.PrimitiveType = PT_TriangleList;
-	SetGraphicsPipelineState(RHICmdList, GraphicsPSOInit);
+		auto ShaderMap = View.ShaderMap;
+		TShaderMapRef<FPostProcessVS> VertexShader(ShaderMap);
+		TShaderMapRef<FVisualizeTexturePresentPS> PixelShader(ShaderMap);
 
-	VertexShader->SetParameters(RHICmdList, View.ViewUniformBuffer);
-	PixelShader->SetParameters(RHICmdList, View, *VisualizeTextureContent);
+		GraphicsPSOInit.BoundShaderState.VertexDeclarationRHI = GFilterVertexDeclaration.VertexDeclarationRHI;
+		GraphicsPSOInit.BoundShaderState.VertexShaderRHI = GETSAFERHISHADER_VERTEX(*VertexShader);
+		GraphicsPSOInit.BoundShaderState.PixelShaderRHI = GETSAFERHISHADER_PIXEL(*PixelShader);
+		GraphicsPSOInit.PrimitiveType = PT_TriangleList;
+		SetGraphicsPipelineState(RHICmdList, GraphicsPSOInit);
 
-	FIntRect VisualizeTextureRect = ComputeVisualizeTextureRect(Desc.Extent);
+		VertexShader->SetParameters(RHICmdList, View.ViewUniformBuffer);
+		PixelShader->SetParameters(RHICmdList, View, *VisualizeTextureContent);
 
-	{
-		SCOPED_DRAW_EVENT(RHICmdList, VisCopyToMain);
-		// Draw a quad mapping scene color to the view's render target
-		DrawRectangle(
-			RHICmdList,
-			VisualizeTextureRect.Min.X, VisualizeTextureRect.Min.Y,
-			VisualizeTextureRect.Width(), VisualizeTextureRect.Height(),
-			0, 0,
-			VisualizeTextureRect.Width(), VisualizeTextureRect.Height(),
-			FIntPoint(RenderTarget->GetSizeX(), RenderTarget->GetSizeY()),
-			VisualizeTextureRect.Size(),
-			*VertexShader,
-			EDRF_Default);
+		FIntRect VisualizeTextureRect = ComputeVisualizeTextureRect(Desc.Extent);
+
+		{
+			SCOPED_DRAW_EVENT(RHICmdList, VisCopyToMain);
+			// Draw a quad mapping scene color to the view's render target
+			DrawRectangle(
+				RHICmdList,
+				VisualizeTextureRect.Min.X, VisualizeTextureRect.Min.Y,
+				VisualizeTextureRect.Width(), VisualizeTextureRect.Height(),
+				0, 0,
+				VisualizeTextureRect.Width(), VisualizeTextureRect.Height(),
+				FIntPoint(RenderTarget->GetSizeX(), RenderTarget->GetSizeY()),
+				VisualizeTextureRect.Size(),
+				*VertexShader,
+				EDRF_Default);
+		}
 	}
+	RHICmdList.EndRenderPass();
 
 	FRenderTargetTemp TempRenderTarget(View, View.UnscaledViewRect.Size());
 	FCanvas Canvas(&TempRenderTarget, NULL, View.Family->CurrentRealTime, View.Family->CurrentWorldTime, View.Family->DeltaWorldTime, View.GetFeatureLevel());
@@ -996,7 +1001,7 @@ void FVisualizeTexture::SetObserveTarget(const FString& InObservedDebugName, uin
 
 void FVisualizeTexture::SetCheckPoint(FRHICommandList& RHICmdList, const IPooledRenderTarget* PooledRenderTarget)
 {
-
+	check(RHICmdList.IsOutsideRenderPass());
 #if !(UE_BUILD_SHIPPING || UE_BUILD_TEST)
 	check(IsInRenderingThread());
 	if (!PooledRenderTarget || !bEnabled)
diff --git a/Engine/Source/Runtime/Renderer/Private/VisualizeVolumetricLightmap.cpp b/Engine/Source/Runtime/Renderer/Private/VisualizeVolumetricLightmap.cpp
--- a/Engine/Source/Runtime/Renderer/Private/VisualizeVolumetricLightmap.cpp
+++ b/Engine/Source/Runtime/Renderer/Private/VisualizeVolumetricLightmap.cpp
@@ -174,37 +174,45 @@ void FDeferredShadingSceneRenderer::VisualizeVolumetricLightmap(FRHICommandListI
 			NumRenderTargets++;
 		}
 
-		SetRenderTargets(RHICmdList, NumRenderTargets, RenderTargets, SceneContext.GetSceneDepthSurface(), ESimpleRenderTargetMode::EExistingColorAndDepth, FExclusiveDepthStencil::DepthWrite_StencilWrite);
+		FRHIRenderPassInfo RPInfo(NumRenderTargets, RenderTargets, ERenderTargetActions::Load_Store);
+		RPInfo.DepthStencilRenderTarget.Action = EDepthStencilTargetActions::LoadDepthStencil_StoreDepthStencil;
+		RPInfo.DepthStencilRenderTarget.DepthStencilTarget = SceneContext.GetSceneDepthSurface();
+		RPInfo.DepthStencilRenderTarget.ExclusiveDepthStencil = FExclusiveDepthStencil::DepthWrite_StencilWrite;
 
-		for (int32 ViewIndex = 0; ViewIndex < Views.Num(); ViewIndex++)
+		RHICmdList.BeginRenderPass(RPInfo, TEXT("VisualizeVolumetricLightmap"));
 		{
-			const FViewInfo& View = Views[ViewIndex];
-			FGraphicsPipelineStateInitializer GraphicsPSOInit;
-			RHICmdList.ApplyCachedRenderTargets(GraphicsPSOInit);
+			for (int32 ViewIndex = 0; ViewIndex < Views.Num(); ViewIndex++)
+			{
+				const FViewInfo& View = Views[ViewIndex];
+				FGraphicsPipelineStateInitializer GraphicsPSOInit;
+				RHICmdList.ApplyCachedRenderTargets(GraphicsPSOInit);
 
-			RHICmdList.SetViewport(View.ViewRect.Min.X, View.ViewRect.Min.Y, 0.0f, View.ViewRect.Max.X, View.ViewRect.Max.Y, 1.0f);
-			GraphicsPSOInit.RasterizerState = TStaticRasterizerState<FM_Solid, CM_None>::GetRHI();
-			GraphicsPSOInit.DepthStencilState = TStaticDepthStencilState<true, CF_DepthNearOrEqual>::GetRHI();
-			GraphicsPSOInit.BlendState = TStaticBlendStateWriteMask<CW_RGB,CW_RGBA>::GetRHI();
-			GraphicsPSOInit.PrimitiveType = PT_TriangleList;
+				RHICmdList.SetViewport(View.ViewRect.Min.X, View.ViewRect.Min.Y, 0.0f, View.ViewRect.Max.X, View.ViewRect.Max.Y, 1.0f);
+				GraphicsPSOInit.RasterizerState = TStaticRasterizerState<FM_Solid, CM_None>::GetRHI();
+				GraphicsPSOInit.DepthStencilState = TStaticDepthStencilState<true, CF_DepthNearOrEqual>::GetRHI();
+				GraphicsPSOInit.BlendState = TStaticBlendStateWriteMask<CW_RGB, CW_RGBA>::GetRHI();
+				GraphicsPSOInit.PrimitiveType = PT_TriangleList;
 
-			TShaderMapRef<FVisualizeVolumetricLightmapVS> VertexShader(View.ShaderMap);
-			TShaderMapRef<FVisualizeVolumetricLightmapPS> PixelShader(View.ShaderMap);
+				TShaderMapRef<FVisualizeVolumetricLightmapVS> VertexShader(View.ShaderMap);
+				TShaderMapRef<FVisualizeVolumetricLightmapPS> PixelShader(View.ShaderMap);
 
-			GraphicsPSOInit.BoundShaderState.VertexDeclarationRHI = GEmptyVertexDeclaration.VertexDeclarationRHI;
-			GraphicsPSOInit.BoundShaderState.VertexShaderRHI = GETSAFERHISHADER_VERTEX(*VertexShader);
-			GraphicsPSOInit.BoundShaderState.PixelShaderRHI = GETSAFERHISHADER_PIXEL(*PixelShader);
+				GraphicsPSOInit.BoundShaderState.VertexDeclarationRHI = GEmptyVertexDeclaration.VertexDeclarationRHI;
+				GraphicsPSOInit.BoundShaderState.VertexShaderRHI = GETSAFERHISHADER_VERTEX(*VertexShader);
+				GraphicsPSOInit.BoundShaderState.PixelShaderRHI = GETSAFERHISHADER_PIXEL(*PixelShader);
 
-			SetGraphicsPipelineState(RHICmdList, GraphicsPSOInit);
+				SetGraphicsPipelineState(RHICmdList, GraphicsPSOInit);
 
-			VertexShader->SetParameters(RHICmdList, View);
-			PixelShader->SetParameters(RHICmdList, View);
+				VertexShader->SetParameters(RHICmdList, View);
+				PixelShader->SetParameters(RHICmdList, View);
 
-			int32 BrickSize = VolumetricLightmapData->BrickSize;
-			int32 NumQuads = VolumetricLightmapData->IndirectionTextureDimensions.X * VolumetricLightmapData->IndirectionTextureDimensions.Y * VolumetricLightmapData->IndirectionTextureDimensions.Z * BrickSize * BrickSize * BrickSize;
+				int32 BrickSize = VolumetricLightmapData->BrickSize;
+				int32 NumQuads = VolumetricLightmapData->IndirectionTextureDimensions.X * VolumetricLightmapData->IndirectionTextureDimensions.Y * VolumetricLightmapData->IndirectionTextureDimensions.Z * BrickSize * BrickSize * BrickSize;
 
-			RHICmdList.SetStreamSource(0, NULL, 0);
-			RHICmdList.DrawIndexedPrimitive(GVisualizeQuadIndexBuffer.IndexBufferRHI, 0, 0, 4 * GQuadsPerVisualizeInstance, 0, 2 * GQuadsPerVisualizeInstance, FMath::DivideAndRoundUp(NumQuads, GQuadsPerVisualizeInstance));
+				RHICmdList.SetStreamSource(0, NULL, 0);
+				RHICmdList.DrawIndexedPrimitive(GVisualizeQuadIndexBuffer.IndexBufferRHI, 0, 0, 4 * GQuadsPerVisualizeInstance, 0, 2 * GQuadsPerVisualizeInstance, FMath::DivideAndRoundUp(NumQuads, GQuadsPerVisualizeInstance));
+			}
 		}
+		RHICmdList.EndRenderPass();
+
 	}
 }
\ No newline at end of file
