Renderpass updates.
VirtualTexturing. Self contained updates.

diff --git a/Engine/Source/Runtime/Renderer/Private/VT/VirtualTextureSpace.cpp b/Engine/Source/Runtime/Renderer/Private/VT/VirtualTextureSpace.cpp
--- a/Engine/Source/Runtime/Renderer/Private/VT/VirtualTextureSpace.cpp
+++ b/Engine/Source/Runtime/Renderer/Private/VT/VirtualTextureSpace.cpp
@@ -311,59 +311,65 @@ void FVirtualTextureSpace::ApplyUpdates( FRHICommandList& RHICmdList )
 		uint32 NumUpdates = ExpandedUpdates[ Mip ].Num();
 		if( NumUpdates )
 		{
-			SetRenderTarget( RHICmdList, PageTableTarget.TargetableTexture, Mip, NULL );
-			RHICmdList.SetViewport( 0, 0, 0.0f, MipSize, MipSize, 1.0f );
+			// We will resolve this target once after all mips are updated.
+			FRHIRenderPassInfo RPInfo(PageTableTarget.TargetableTexture, ERenderTargetActions::Load_Store);
+			RPInfo.ColorRenderTargets[0].MipIndex = Mip;
+			RHICmdList.BeginRenderPass(RPInfo, TEXT("VirtualTextureUpdate"));
+			{
+				RHICmdList.SetViewport(0, 0, 0.0f, MipSize, MipSize, 1.0f);
 
-			FGraphicsPipelineStateInitializer GraphicsPSOInit;
-			RHICmdList.ApplyCachedRenderTargets( GraphicsPSOInit );
+				FGraphicsPipelineStateInitializer GraphicsPSOInit;
+				RHICmdList.ApplyCachedRenderTargets(GraphicsPSOInit);
 
-			GraphicsPSOInit.BlendState = TStaticBlendState<>::GetRHI();
-			GraphicsPSOInit.RasterizerState = TStaticRasterizerState<>::GetRHI();
-			GraphicsPSOInit.DepthStencilState = TStaticDepthStencilState<false, CF_Always>::GetRHI();
-			GraphicsPSOInit.PrimitiveType = PT_TriangleList;
+				GraphicsPSOInit.BlendState = TStaticBlendState<>::GetRHI();
+				GraphicsPSOInit.RasterizerState = TStaticRasterizerState<>::GetRHI();
+				GraphicsPSOInit.DepthStencilState = TStaticDepthStencilState<false, CF_Always>::GetRHI();
+				GraphicsPSOInit.PrimitiveType = PT_TriangleList;
 
-			FPageTableUpdateVS* VertexShader = nullptr;
-			FPageTableUpdatePS* PixelShader = nullptr;
+				FPageTableUpdateVS* VertexShader = nullptr;
+				FPageTableUpdatePS* PixelShader = nullptr;
 
-			switch( PageTableFormat )
-			{
-			case PF_R16_UINT:
+				switch (PageTableFormat)
+				{
+				case PF_R16_UINT:
 				{
-					VertexShader = *TShaderMapRef< TPageTableUpdateVS<0> >( ShaderMap );
-					PixelShader  = *TShaderMapRef< TPageTableUpdatePS<0> >( ShaderMap );
+					VertexShader = *TShaderMapRef< TPageTableUpdateVS<0> >(ShaderMap);
+					PixelShader = *TShaderMapRef< TPageTableUpdatePS<0> >(ShaderMap);
 				}
 				break;
-			case PF_R8G8B8A8:
+				case PF_R8G8B8A8:
 				{
-					VertexShader = *TShaderMapRef< TPageTableUpdateVS<1> >( ShaderMap );
-					PixelShader  = *TShaderMapRef< TPageTableUpdatePS<1> >( ShaderMap );
+					VertexShader = *TShaderMapRef< TPageTableUpdateVS<1> >(ShaderMap);
+					PixelShader = *TShaderMapRef< TPageTableUpdatePS<1> >(ShaderMap);
 				}
 				break;
-			default:
-				check(0);
-			}
-			checkSlow( VertexShader && PixelShader );
-			
-			GraphicsPSOInit.BoundShaderState.VertexDeclarationRHI = GEmptyVertexDeclaration.VertexDeclarationRHI;
-			GraphicsPSOInit.BoundShaderState.VertexShaderRHI = GETSAFERHISHADER_VERTEX(VertexShader);
-			GraphicsPSOInit.BoundShaderState.PixelShaderRHI = GETSAFERHISHADER_PIXEL(PixelShader);
+				default:
+					check(0);
+				}
+				checkSlow(VertexShader && PixelShader);
 
-			SetGraphicsPipelineState( RHICmdList, GraphicsPSOInit );
+				GraphicsPSOInit.BoundShaderState.VertexDeclarationRHI = GEmptyVertexDeclaration.VertexDeclarationRHI;
+				GraphicsPSOInit.BoundShaderState.VertexShaderRHI = GETSAFERHISHADER_VERTEX(VertexShader);
+				GraphicsPSOInit.BoundShaderState.PixelShaderRHI = GETSAFERHISHADER_PIXEL(PixelShader);
 
-			{
-				const FVertexShaderRHIParamRef ShaderRHI = VertexShader->GetVertexShader();
+				SetGraphicsPipelineState(RHICmdList, GraphicsPSOInit);
 
-				SetShaderValue( RHICmdList, ShaderRHI, VertexShader->PageTableSize,	PageTableSize );
-				SetShaderValue( RHICmdList, ShaderRHI, VertexShader->FirstUpdate,	FirstUpdate );
-				SetShaderValue( RHICmdList, ShaderRHI, VertexShader->NumUpdates,	NumUpdates );
-				SetSRVParameter( RHICmdList, ShaderRHI, VertexShader->UpdateBuffer,	UpdateBufferSRV );
-			}
+				{
+					const FVertexShaderRHIParamRef ShaderRHI = VertexShader->GetVertexShader();
 
-			// needs to be the same on shader side (faster on NVIDIA and AMD)
-			uint32 QuadsPerInstance = 8;
+					SetShaderValue(RHICmdList, ShaderRHI, VertexShader->PageTableSize, PageTableSize);
+					SetShaderValue(RHICmdList, ShaderRHI, VertexShader->FirstUpdate, FirstUpdate);
+					SetShaderValue(RHICmdList, ShaderRHI, VertexShader->NumUpdates, NumUpdates);
+					SetSRVParameter(RHICmdList, ShaderRHI, VertexShader->UpdateBuffer, UpdateBufferSRV);
+				}
 
-			RHICmdList.SetStreamSource( 0, NULL, 0 );
-			RHICmdList.DrawIndexedPrimitive( GQuadIndexBuffer.IndexBufferRHI, 0, 0, 32, 0, 2 * QuadsPerInstance, FMath::DivideAndRoundUp( NumUpdates, QuadsPerInstance ) );
+				// needs to be the same on shader side (faster on NVIDIA and AMD)
+				uint32 QuadsPerInstance = 8;
+
+				RHICmdList.SetStreamSource(0, NULL, 0);
+				RHICmdList.DrawIndexedPrimitive(GQuadIndexBuffer.IndexBufferRHI, 0, 0, 32, 0, 2 * QuadsPerInstance, FMath::DivideAndRoundUp(NumUpdates, QuadsPerInstance));
+			}
+			RHICmdList.EndRenderPass();
 
 			ExpandedUpdates[ Mip ].Reset();
 		}
diff --git a/Engine/Source/Runtime/Renderer/Private/VT/VirtualTextureTest.cpp b/Engine/Source/Runtime/Renderer/Private/VT/VirtualTextureTest.cpp
--- a/Engine/Source/Runtime/Renderer/Private/VT/VirtualTextureTest.cpp
+++ b/Engine/Source/Runtime/Renderer/Private/VT/VirtualTextureTest.cpp
@@ -187,37 +187,39 @@ void FVirtualTextureTest::ProducePageData( FRHICommandList& RHICmdList, ERHIFeat
 
 	FSceneRenderTargetItem& RenderTarget = PhysicalTexture->GetRenderTargetItem();
 
-	SetRenderTarget( RHICmdList, RenderTarget.TargetableTexture, NULL );
-	RHICmdList.SetViewport( DstRect.Min.X, DstRect.Min.Y, 0.0f, DstRect.Max.X, DstRect.Max.Y, 1.0f );
-
-	FGraphicsPipelineStateInitializer GraphicsPSOInit;
-	RHICmdList.ApplyCachedRenderTargets( GraphicsPSOInit );
-
-	GraphicsPSOInit.BlendState = TStaticBlendState<>::GetRHI();
-	GraphicsPSOInit.RasterizerState = TStaticRasterizerState<>::GetRHI();
-	GraphicsPSOInit.DepthStencilState = TStaticDepthStencilState<false, CF_Always>::GetRHI();
-	GraphicsPSOInit.PrimitiveType = PT_TriangleList;
-
-	TShaderMapRef< FScreenVS >				VertexShader( ShaderMap );
-	TShaderMapRef< FVirtualTextureTestPS >	PixelShader( ShaderMap );
-
-	GraphicsPSOInit.BoundShaderState.VertexDeclarationRHI = GFilterVertexDeclaration.VertexDeclarationRHI;
-	GraphicsPSOInit.BoundShaderState.VertexShaderRHI = GETSAFERHISHADER_VERTEX(*VertexShader);
-	GraphicsPSOInit.BoundShaderState.PixelShaderRHI = GETSAFERHISHADER_PIXEL(*PixelShader);
-	SetGraphicsPipelineState( RHICmdList, GraphicsPSOInit );
-
-	DrawRectangle(
-		RHICmdList,
-		0,					0,
-		DstRect.Width(),	DstRect.Height(),
-		SrcRect.Min.X,		SrcRect.Min.Y,
-		SrcRect.Width(),	SrcRect.Height(),
-		DstRect.Size(),
-		SrcSize,
-		*VertexShader,
-		EDRF_UseTriangleOptimization );
-
-	RHICmdList.CopyToResolveTarget( RenderTarget.TargetableTexture, RenderTarget.ShaderResourceTexture, false, FResolveParams() );
+	FRHIRenderPassInfo RPInfo(RenderTarget.TargetableTexture, ERenderTargetActions::Load_Store, RenderTarget.ShaderResourceTexture);
+	RHICmdList.BeginRenderPass(RPInfo, TEXT("ProducePageData"));
+	{
+		RHICmdList.SetViewport(DstRect.Min.X, DstRect.Min.Y, 0.0f, DstRect.Max.X, DstRect.Max.Y, 1.0f);
+
+		FGraphicsPipelineStateInitializer GraphicsPSOInit;
+		RHICmdList.ApplyCachedRenderTargets(GraphicsPSOInit);
+
+		GraphicsPSOInit.BlendState = TStaticBlendState<>::GetRHI();
+		GraphicsPSOInit.RasterizerState = TStaticRasterizerState<>::GetRHI();
+		GraphicsPSOInit.DepthStencilState = TStaticDepthStencilState<false, CF_Always>::GetRHI();
+		GraphicsPSOInit.PrimitiveType = PT_TriangleList;
+
+		TShaderMapRef< FScreenVS >				VertexShader(ShaderMap);
+		TShaderMapRef< FVirtualTextureTestPS >	PixelShader(ShaderMap);
+
+		GraphicsPSOInit.BoundShaderState.VertexDeclarationRHI = GFilterVertexDeclaration.VertexDeclarationRHI;
+		GraphicsPSOInit.BoundShaderState.VertexShaderRHI = GETSAFERHISHADER_VERTEX(*VertexShader);
+		GraphicsPSOInit.BoundShaderState.PixelShaderRHI = GETSAFERHISHADER_PIXEL(*PixelShader);
+		SetGraphicsPipelineState(RHICmdList, GraphicsPSOInit);
+
+		DrawRectangle(
+			RHICmdList,
+			0, 0,
+			DstRect.Width(), DstRect.Height(),
+			SrcRect.Min.X, SrcRect.Min.Y,
+			SrcRect.Width(), SrcRect.Height(),
+			DstRect.Size(),
+			SrcSize,
+			*VertexShader,
+			EDRF_UseTriangleOptimization);
+	}
+	RHICmdList.EndRenderPass();
 
 	GVisualizeTexture.SetCheckPoint( RHICmdList, PhysicalTexture );
 }
