Renderpasses.

Prep work. Verifies we are not leaving renderpasses open.

diff --git a/Engine/Source/Runtime/Renderer/Private/PostProcess/RenderingCompositionGraph.cpp b/Engine/Source/Runtime/Renderer/Private/PostProcess/RenderingCompositionGraph.cpp
--- a/Engine/Source/Runtime/Renderer/Private/PostProcess/RenderingCompositionGraph.cpp
+++ b/Engine/Source/Runtime/Renderer/Private/PostProcess/RenderingCompositionGraph.cpp
@@ -679,7 +679,9 @@ void FRenderingCompositionGraph::RecursivelyProcess(const FRenderingCompositeOut
 	Context.SetViewportInvalid();
 
 	// then process the pass itself
+	check(!Context.RHICmdList.IsInsideRenderPass());
 	Pass->Process(Context);
+	check(!Context.RHICmdList.IsInsideRenderPass());
 
 	// for VisualizeTexture and output buffer dumping
 #if !(UE_BUILD_SHIPPING || UE_BUILD_TEST)
