Renderpass updates.
Fog Rendering. Now calling FinishRenderingSceneColor.

diff --git a/Engine/Source/Runtime/Renderer/Private/FogRendering.cpp b/Engine/Source/Runtime/Renderer/Private/FogRendering.cpp
--- a/Engine/Source/Runtime/Renderer/Private/FogRendering.cpp
+++ b/Engine/Source/Runtime/Renderer/Private/FogRendering.cpp
@@ -426,6 +426,8 @@ void FDeferredShadingSceneRenderer::RenderViewFog(FRHICommandList& RHICmdList, c
 
 bool FDeferredShadingSceneRenderer::RenderFog(FRHICommandListImmediate& RHICmdList, const FLightShaftsOutput& LightShaftsOutput)
 {
+	check(RHICmdList.IsOutsideRenderPass());
+
 	if (Scene->ExponentialFogs.Num() > 0 
 		// Fog must be done in the base pass for MSAA to work
 		&& !IsForwardShadingEnabled(ShaderPlatform))
@@ -444,6 +446,8 @@ bool FDeferredShadingSceneRenderer::RenderFog(FRHICommandListImmediate& RHICmdLi
 			}
 		}
 
+		SceneContext.FinishRenderingSceneColor(RHICmdList);
+
 		return true;
 	}
 
