Renderpass updates.

System textures. Self-contained.

diff --git a/Engine/Source/Runtime/Renderer/Private/SystemTextures.cpp b/Engine/Source/Runtime/Renderer/Private/SystemTextures.cpp
--- a/Engine/Source/Runtime/Renderer/Private/SystemTextures.cpp
+++ b/Engine/Source/Runtime/Renderer/Private/SystemTextures.cpp
@@ -28,8 +28,12 @@ void FSystemTextures::InitializeCommonTextures(FRHICommandListImmediate& RHICmdL
 			Desc.AutoWritable = false;
 			GRenderTargetPool.FindFreeElement(RHICmdList, Desc, WhiteDummy, TEXT("WhiteDummy"), true, ERenderTargetTransience::NonTransient);
 
-			SetRenderTarget(RHICmdList, WhiteDummy->GetRenderTargetItem().TargetableTexture, FTextureRHIRef(), ESimpleRenderTargetMode::EClearColorExistingDepth);
-			RHICmdList.CopyToResolveTarget(WhiteDummy->GetRenderTargetItem().TargetableTexture, WhiteDummy->GetRenderTargetItem().ShaderResourceTexture, FResolveParams());
+			//SetRenderTarget(RHICmdList, WhiteDummy->GetRenderTargetItem().TargetableTexture, FTextureRHIRef(), ESimpleRenderTargetMode::EClearColorExistingDepth);
+			//RHICmdList.CopyToResolveTarget(WhiteDummy->GetRenderTargetItem().TargetableTexture, WhiteDummy->GetRenderTargetItem().ShaderResourceTexture, FResolveParams());
+
+			FRHIRenderPassInfo RPInfo(WhiteDummy->GetRenderTargetItem().TargetableTexture, ERenderTargetActions::Clear_Store, WhiteDummy->GetRenderTargetItem().ShaderResourceTexture);
+			RHICmdList.BeginRenderPass(RPInfo, TEXT("WhiteDummy"));
+			RHICmdList.EndRenderPass();
 		}
 
 		// Create a BlackDummy texture
@@ -38,8 +42,12 @@ void FSystemTextures::InitializeCommonTextures(FRHICommandListImmediate& RHICmdL
 			Desc.AutoWritable = false;
 			GRenderTargetPool.FindFreeElement(RHICmdList, Desc, BlackDummy, TEXT("BlackDummy"), true, ERenderTargetTransience::NonTransient);
 
-			SetRenderTarget(RHICmdList, BlackDummy->GetRenderTargetItem().TargetableTexture, FTextureRHIRef(), ESimpleRenderTargetMode::EClearColorExistingDepth);
-			RHICmdList.CopyToResolveTarget(BlackDummy->GetRenderTargetItem().TargetableTexture, BlackDummy->GetRenderTargetItem().ShaderResourceTexture, FResolveParams());
+			/*SetRenderTarget(RHICmdList, BlackDummy->GetRenderTargetItem().TargetableTexture, FTextureRHIRef(), ESimpleRenderTargetMode::EClearColorExistingDepth);
+			RHICmdList.CopyToResolveTarget(BlackDummy->GetRenderTargetItem().TargetableTexture, BlackDummy->GetRenderTargetItem().ShaderResourceTexture, FResolveParams());*/
+
+			FRHIRenderPassInfo RPInfo(BlackDummy->GetRenderTargetItem().TargetableTexture, ERenderTargetActions::Clear_Store, BlackDummy->GetRenderTargetItem().ShaderResourceTexture);
+			RHICmdList.BeginRenderPass(RPInfo, TEXT("BlackDummy"));
+			RHICmdList.EndRenderPass();
 		}
 
 		// Create a BlackAlphaOneDummy texture
@@ -48,8 +56,12 @@ void FSystemTextures::InitializeCommonTextures(FRHICommandListImmediate& RHICmdL
 			Desc.AutoWritable = false;
 			GRenderTargetPool.FindFreeElement(RHICmdList, Desc, BlackAlphaOneDummy, TEXT("BlackAlphaOneDummy"), true, ERenderTargetTransience::NonTransient);
 
-			SetRenderTarget(RHICmdList, BlackAlphaOneDummy->GetRenderTargetItem().TargetableTexture, FTextureRHIRef(), ESimpleRenderTargetMode::EClearColorExistingDepth);
-			RHICmdList.CopyToResolveTarget(BlackAlphaOneDummy->GetRenderTargetItem().TargetableTexture, BlackAlphaOneDummy->GetRenderTargetItem().ShaderResourceTexture, FResolveParams());
+			/*SetRenderTarget(RHICmdList, BlackAlphaOneDummy->GetRenderTargetItem().TargetableTexture, FTextureRHIRef(), ESimpleRenderTargetMode::EClearColorExistingDepth);
+			RHICmdList.CopyToResolveTarget(BlackAlphaOneDummy->GetRenderTargetItem().TargetableTexture, BlackAlphaOneDummy->GetRenderTargetItem().ShaderResourceTexture, FResolveParams());*/
+
+			FRHIRenderPassInfo RPInfo(BlackAlphaOneDummy->GetRenderTargetItem().TargetableTexture, ERenderTargetActions::Clear_Store, BlackAlphaOneDummy->GetRenderTargetItem().ShaderResourceTexture);
+			RHICmdList.BeginRenderPass(RPInfo, TEXT("BlackAlphaOneDummy"));
+			RHICmdList.EndRenderPass();
 		}
 
 		// Create a GreenDummy texture
@@ -58,8 +70,12 @@ void FSystemTextures::InitializeCommonTextures(FRHICommandListImmediate& RHICmdL
 			Desc.AutoWritable = false;
 			GRenderTargetPool.FindFreeElement(RHICmdList, Desc, GreenDummy, TEXT("GreenDummy"), true, ERenderTargetTransience::NonTransient);
 
-			SetRenderTarget(RHICmdList, GreenDummy->GetRenderTargetItem().TargetableTexture, FTextureRHIRef(), ESimpleRenderTargetMode::EClearColorExistingDepth);
-			RHICmdList.CopyToResolveTarget(GreenDummy->GetRenderTargetItem().TargetableTexture, GreenDummy->GetRenderTargetItem().ShaderResourceTexture, FResolveParams());
+			/*SetRenderTarget(RHICmdList, GreenDummy->GetRenderTargetItem().TargetableTexture, FTextureRHIRef(), ESimpleRenderTargetMode::EClearColorExistingDepth);
+			RHICmdList.CopyToResolveTarget(GreenDummy->GetRenderTargetItem().TargetableTexture, GreenDummy->GetRenderTargetItem().ShaderResourceTexture, FResolveParams());*/
+
+			FRHIRenderPassInfo RPInfo(GreenDummy->GetRenderTargetItem().TargetableTexture, ERenderTargetActions::Clear_Store, GreenDummy->GetRenderTargetItem().ShaderResourceTexture);
+			RHICmdList.BeginRenderPass(RPInfo, TEXT("GreenDummy"));
+			RHICmdList.EndRenderPass();
 		}
 
 		// Create a DefaultNormal8Bit texture
@@ -68,8 +84,12 @@ void FSystemTextures::InitializeCommonTextures(FRHICommandListImmediate& RHICmdL
 			Desc.AutoWritable = false;
 			GRenderTargetPool.FindFreeElement(RHICmdList, Desc, DefaultNormal8Bit, TEXT("DefaultNormal8Bit"), true, ERenderTargetTransience::NonTransient);
 
-			SetRenderTarget(RHICmdList, DefaultNormal8Bit->GetRenderTargetItem().TargetableTexture, FTextureRHIRef(), ESimpleRenderTargetMode::EClearColorExistingDepth);
-			RHICmdList.CopyToResolveTarget(DefaultNormal8Bit->GetRenderTargetItem().TargetableTexture, DefaultNormal8Bit->GetRenderTargetItem().ShaderResourceTexture, FResolveParams());
+			/*SetRenderTarget(RHICmdList, DefaultNormal8Bit->GetRenderTargetItem().TargetableTexture, FTextureRHIRef(), ESimpleRenderTargetMode::EClearColorExistingDepth);
+			RHICmdList.CopyToResolveTarget(DefaultNormal8Bit->GetRenderTargetItem().TargetableTexture, DefaultNormal8Bit->GetRenderTargetItem().ShaderResourceTexture, FResolveParams());*/
+
+			FRHIRenderPassInfo RPInfo(DefaultNormal8Bit->GetRenderTargetItem().TargetableTexture, ERenderTargetActions::Clear_Store, DefaultNormal8Bit->GetRenderTargetItem().ShaderResourceTexture);
+			RHICmdList.BeginRenderPass(RPInfo, TEXT("DefaultNormal8Bit"));
+			RHICmdList.EndRenderPass();
 		}
 
 		// Create the PerlinNoiseGradient texture
@@ -110,10 +130,9 @@ void FSystemTextures::InitializeCommonTextures(FRHICommandListImmediate& RHICmdL
 		Desc.AutoWritable = false;
 		GRenderTargetPool.FindFreeElement(RHICmdList, Desc, MaxFP16Depth, TEXT("MaxFP16Depth"), true, ERenderTargetTransience::NonTransient);
 
-		FRHIRenderTargetView RtView = FRHIRenderTargetView(MaxFP16Depth->GetRenderTargetItem().TargetableTexture, ERenderTargetLoadAction::EClear);
-		FRHISetRenderTargetsInfo Info(1, &RtView, FRHIDepthRenderTargetView());
-		RHICmdList.SetRenderTargetsAndClear(Info);
-		RHICmdList.CopyToResolveTarget(MaxFP16Depth->GetRenderTargetItem().TargetableTexture, MaxFP16Depth->GetRenderTargetItem().ShaderResourceTexture, FResolveParams());
+		FRHIRenderPassInfo RPInfo(MaxFP16Depth->GetRenderTargetItem().TargetableTexture, ERenderTargetActions::Clear_Store, MaxFP16Depth->GetRenderTargetItem().ShaderResourceTexture);
+		RHICmdList.BeginRenderPass(RPInfo, TEXT("MaxFP16Depth"));
+		RHICmdList.EndRenderPass();
 	}
 
 	// Create dummy 1x1 depth texture		
@@ -122,9 +141,10 @@ void FSystemTextures::InitializeCommonTextures(FRHICommandListImmediate& RHICmdL
 		Desc.AutoWritable = false;
 		GRenderTargetPool.FindFreeElement(RHICmdList, Desc, DepthDummy, TEXT("DepthDummy"), true, ERenderTargetTransience::NonTransient);
 
-		FRHISetRenderTargetsInfo Info(0, nullptr, FRHIDepthRenderTargetView(DepthDummy->GetRenderTargetItem().TargetableTexture, ERenderTargetLoadAction::EClear, ERenderTargetStoreAction::EStore));
-		RHICmdList.SetRenderTargetsAndClear(Info);
-		RHICmdList.CopyToResolveTarget(DepthDummy->GetRenderTargetItem().TargetableTexture, DepthDummy->GetRenderTargetItem().ShaderResourceTexture, FResolveParams());
+		FRHIRenderPassInfo RPInfo(DepthDummy->GetRenderTargetItem().TargetableTexture, EDepthStencilTargetActions::ClearDepthStencil_StoreDepthStencil, DepthDummy->GetRenderTargetItem().ShaderResourceTexture, FExclusiveDepthStencil::DepthWrite_StencilWrite);
+
+		RHICmdList.BeginRenderPass(RPInfo, TEXT("DepthDummy"));
+		RHICmdList.EndRenderPass();
 	}
 }
 
