Renderpass Updates.

Volumetric Fog. Self-contained.

diff --git a/Engine/Source/Runtime/Renderer/Private/VolumetricFog.cpp b/Engine/Source/Runtime/Renderer/Private/VolumetricFog.cpp
--- a/Engine/Source/Runtime/Renderer/Private/VolumetricFog.cpp
+++ b/Engine/Source/Runtime/Renderer/Private/VolumetricFog.cpp
@@ -959,6 +959,8 @@ void FDeferredShadingSceneRenderer::SetupVolumetricFog()
 
 void FDeferredShadingSceneRenderer::ComputeVolumetricFog(FRHICommandListImmediate& RHICmdListImmediate)
 {
+	check(RHICmdListImmediate.IsOutsideRenderPass());
+
 	if (ShouldRenderVolumetricFog())
 	{
 		QUICK_SCOPE_CYCLE_COUNTER(STAT_VolumetricFog);
@@ -1113,7 +1115,7 @@ void FDeferredShadingSceneRenderer::ComputeVolumetricFog(FRHICommandListImmediat
 					ERenderGraphPassFlags::Compute,
 					[PassParameters, ComputeShader, &View, this, FogInfo, bUseTemporalReprojection, VolumetricFogGridSize, IntegrationData, bUseDirectionalLightShadowing, bUseDistanceFieldSkyOcclusion, LightFunctionWorldToShadow](FRHICommandListImmediate& RHICmdList)
 					{
-						SetRenderTarget(RHICmdList, NULL, NULL);
+						UnbindRenderTargets(RHICmdList);
 						const FIntVector NumGroups = FIntVector::DivideAndRoundUp(VolumetricFogGridSize, VolumetricFogGridInjectionGroupSize);
 
 						SCOPED_DRAW_EVENTF(RHICmdList, LightScattering, TEXT("LightScattering %dx%dx%d %s %s"), 
