Renderpass updates.

Gamma correction. Self-contained.

diff --git a/Engine/Source/Runtime/Renderer/Private/GammaCorrection.cpp b/Engine/Source/Runtime/Renderer/Private/GammaCorrection.cpp
--- a/Engine/Source/Runtime/Renderer/Private/GammaCorrection.cpp
+++ b/Engine/Source/Runtime/Renderer/Private/GammaCorrection.cpp
@@ -88,27 +88,26 @@ IMPLEMENT_SHADER_TYPE(,FGammaCorrectionVS,TEXT("/Engine/Private/GammaCorrection.
 void FSceneRenderer::GammaCorrectToViewportRenderTarget(FRHICommandList& RHICmdList, const FViewInfo* View, float OverrideGamma)
 {
 	// Set the view family's render target/viewport.
-	SetRenderTarget(RHICmdList, ViewFamily.RenderTarget->GetRenderTargetTexture(), FTextureRHIRef());
+	FRHIRenderPassInfo RPInfo(ViewFamily.RenderTarget->GetRenderTargetTexture(), ERenderTargetActions::DontLoad_Store);
 
 	// Deferred the clear until here so the garbage left in the non rendered regions by the post process effects do not show up
 	if( ViewFamily.bDeferClear )
 	{
 		if (ensure(ViewFamily.RenderTarget->GetRenderTargetTexture()->GetClearColor() == FLinearColor::Black))
 		{
-			FRHIRenderTargetView RtView = FRHIRenderTargetView(ViewFamily.RenderTarget->GetRenderTargetTexture(), ERenderTargetLoadAction::EClear);
-			FRHISetRenderTargetsInfo Info(1, &RtView, FRHIDepthRenderTargetView());
-			RHICmdList.SetRenderTargetsAndClear(Info);
+			RPInfo.ColorRenderTargets[0].Action = ERenderTargetActions::Clear_Store;
+			RHICmdList.BeginRenderPass(RPInfo, TEXT("GammaCorrectToViewportRenderTarget"));
 		}
 		else
 		{
-			SetRenderTarget(RHICmdList, ViewFamily.RenderTarget->GetRenderTargetTexture(), FTextureRHIRef());
+			RHICmdList.BeginRenderPass(RPInfo, TEXT("GammaCorrectToViewportRenderTarget"));
 			DrawClearQuad(RHICmdList, FLinearColor::Black);
 		}
 		ViewFamily.bDeferClear = false;
 	}
 	else
 	{
-		SetRenderTarget(RHICmdList, ViewFamily.RenderTarget->GetRenderTargetTexture(), FTextureRHIRef());
+		RHICmdList.BeginRenderPass(RPInfo, TEXT("GammaCorrectToViewportRenderTarget"));
 	}
 
 	SCOPED_DRAW_EVENT(RHICmdList, GammaCorrection);
@@ -174,4 +173,6 @@ void FSceneRenderer::GammaCorrectToViewportRenderTarget(FRHICommandList& RHICmdL
 		SceneContext.GetBufferSizeXY(),
 		*VertexShader,
 		EDRF_UseTriangleOptimization);
+
+	RHICmdList.EndRenderPass();
 }
