Renderpass updates.

GPUBenchmark. Self-contained.

diff --git a/Engine/Source/Runtime/Renderer/Private/GPUBenchmark.cpp b/Engine/Source/Runtime/Renderer/Private/GPUBenchmark.cpp
--- a/Engine/Source/Runtime/Renderer/Private/GPUBenchmark.cpp
+++ b/Engine/Source/Runtime/Renderer/Private/GPUBenchmark.cpp
@@ -576,14 +576,16 @@ void RendererGPUBenchmark(FRHICommandListImmediate& RHICmdList, FSynthBenchmarkR
 
 				GVisualizeTexture.SetCheckPoint(RHICmdList, RTItems[DestRTIndex]);
 
-				SetRenderTarget(RHICmdList, RTItems[DestRTIndex]->GetRenderTargetItem().TargetableTexture, FTextureRHIRef(), true);	
-
 				// decide how much work we do in this pass
 				LocalWorkScale[Iteration] = (Iteration / 10.f + 1.f) * WorkScale;
 
-				RunBenchmarkShader(RHICmdList, VertexBuffer, View, MethodId, RTItems[SrcRTIndex], LocalWorkScale[Iteration]);
-
-				RHICmdList.CopyToResolveTarget(RTItems[DestRTIndex]->GetRenderTargetItem().TargetableTexture, RTItems[DestRTIndex]->GetRenderTargetItem().ShaderResourceTexture, FResolveParams());
+				FRHIRenderPassInfo RPInfo(RTItems[DestRTIndex]->GetRenderTargetItem().TargetableTexture, ERenderTargetActions::Load_Store, RTItems[DestRTIndex]->GetRenderTargetItem().ShaderResourceTexture);
+				TransitionRenderPassTargets(RHICmdList, RPInfo);
+				RHICmdList.BeginRenderPass(RPInfo, TEXT("GPUBenchmark"));
+				{
+					RunBenchmarkShader(RHICmdList, VertexBuffer, View, MethodId, RTItems[SrcRTIndex], LocalWorkScale[Iteration]);
+				}
+				RHICmdList.EndRenderPass();
 
 				/*if(bGPUCPUSync)
 				{
