Renderpass updates.

Monoscopic rendering. Self-contained.

diff --git a/Engine/Source/Runtime/Renderer/Private/MonoscopicFarFieldRendering.cpp b/Engine/Source/Runtime/Renderer/Private/MonoscopicFarFieldRendering.cpp
--- a/Engine/Source/Runtime/Renderer/Private/MonoscopicFarFieldRendering.cpp
+++ b/Engine/Source/Runtime/Renderer/Private/MonoscopicFarFieldRendering.cpp
@@ -277,96 +277,103 @@ void FSceneRenderer::CompositeMonoscopicFarField(FRHICommandListImmediate& RHICm
 		const FTextureRHIParamRef CurrentSceneColor = GetMultiViewSceneColor(SceneContext);
 
 		const FTextureRHIParamRef SceneDepth = (LeftView.bIsMobileMultiViewEnabled) ? SceneContext.MobileMultiViewSceneDepthZ->GetRenderTargetItem().TargetableTexture : static_cast<FTextureRHIRef>(SceneContext.GetSceneDepthTexture());
-		SetRenderTarget(RHICmdList, CurrentSceneColor, SceneDepth, ESimpleRenderTargetMode::EExistingColorAndDepth);
 
-		FGraphicsPipelineStateInitializer GraphicsPSOInit;
-		RHICmdList.ApplyCachedRenderTargets(GraphicsPSOInit);
-
-		GraphicsPSOInit.RasterizerState = TStaticRasterizerState<FM_Solid, CM_None>::GetRHI();
-		GraphicsPSOInit.DepthStencilState = TStaticDepthStencilState<false, CF_Always>::GetRHI();
-
-		if (ViewFamily.MonoParameters.Mode == EMonoscopicFarFieldMode::MonoOnly)
-		{
-			GraphicsPSOInit.BlendState = TStaticBlendState<CW_RGBA, BO_Add, BF_One, BF_Zero>::GetRHI();
-		}
-		else
-		{
-			GraphicsPSOInit.BlendState = TStaticBlendState<CW_RGBA, BO_Add, BF_InverseDestAlpha, BF_One>::GetRHI();
-		}
-
-		if (!LeftView.bIsMobileMultiViewEnabled)
-		{
-			TShaderMapRef<FCompositeMonoscopicFarFieldViewVS<false>> VertexShader(MonoView.ShaderMap);
-			TShaderMapRef<FCompositeMonoscopicFarFieldViewPS<false>> PixelShader(MonoView.ShaderMap);
-
-			GraphicsPSOInit.BoundShaderState.VertexDeclarationRHI = GFilterVertexDeclaration.VertexDeclarationRHI;
-			GraphicsPSOInit.BoundShaderState.VertexShaderRHI = GETSAFERHISHADER_VERTEX(*VertexShader);
-			GraphicsPSOInit.BoundShaderState.PixelShaderRHI = GETSAFERHISHADER_PIXEL(*PixelShader);
-			GraphicsPSOInit.PrimitiveType = PT_TriangleList;
-
-			SetGraphicsPipelineState(RHICmdList, GraphicsPSOInit);
-
-			PixelShader->SetParameters(RHICmdList, MonoView);
-
-			VertexShader->SetParameters(RHICmdList, MonoView, ViewFamily.MonoParameters.LateralOffset);
-
-			const int32 LateralOffsetInPixels = static_cast<int32>(roundf(ViewFamily.MonoParameters.LateralOffset * static_cast<float>(MonoView.ViewRect.Width())));
-
-			RHICmdList.SetViewport(LeftView.ViewRect.Min.X, LeftView.ViewRect.Min.Y, 0.0f, LeftView.ViewRect.Max.X, LeftView.ViewRect.Max.Y, 1.0f);
-			DrawRectangle(
-				RHICmdList,
-				0, 0,
-				LeftView.ViewRect.Width(), LeftView.ViewRect.Height(),
-				MonoView.ViewRect.Min.X + LateralOffsetInPixels, MonoView.ViewRect.Min.Y,
-				LeftView.ViewRect.Width(), LeftView.ViewRect.Height(),
-				LeftView.ViewRect.Size(),
-				MonoView.ViewRect.Max,
-				*VertexShader,
-				EDRF_UseTriangleOptimization);
-
-			RHICmdList.SetViewport(RightView.ViewRect.Min.X, RightView.ViewRect.Min.Y, 0.0f, RightView.ViewRect.Max.X, RightView.ViewRect.Max.Y, 1.0f);
-
-			// Composite into right eye
-			DrawRectangle(
-				RHICmdList,
-				0, 0,
-				LeftView.ViewRect.Width(), LeftView.ViewRect.Height(),
-				MonoView.ViewRect.Min.X - LateralOffsetInPixels, MonoView.ViewRect.Min.Y,
-				LeftView.ViewRect.Width(), LeftView.ViewRect.Height(),
-				LeftView.ViewRect.Size(),
-				MonoView.ViewRect.Max,
-				*VertexShader,
-				EDRF_UseTriangleOptimization);
-
-		}
-		else
+		FRHIRenderPassInfo RPInfo(CurrentSceneColor, ERenderTargetActions::Load_Store);
+		RPInfo.DepthStencilRenderTarget.Action = EDepthStencilTargetActions::LoadDepthStencil_StoreDepthStencil;
+		RPInfo.DepthStencilRenderTarget.DepthStencilTarget = SceneDepth;
+		RPInfo.DepthStencilRenderTarget.ExclusiveDepthStencil = FExclusiveDepthStencil::DepthWrite_StencilWrite;
+		RHICmdList.BeginRenderPass(RPInfo, TEXT("ComputeMonoscopicFarField"));
 		{
-			TShaderMapRef<FCompositeMonoscopicFarFieldViewVS<true>> VertexShader(MonoView.ShaderMap);
-			TShaderMapRef<FCompositeMonoscopicFarFieldViewPS<true>> PixelShader(MonoView.ShaderMap);
-
-			GraphicsPSOInit.BoundShaderState.VertexDeclarationRHI = GFilterVertexDeclaration.VertexDeclarationRHI;
-			GraphicsPSOInit.BoundShaderState.VertexShaderRHI = GETSAFERHISHADER_VERTEX(*VertexShader);
-			GraphicsPSOInit.BoundShaderState.PixelShaderRHI = GETSAFERHISHADER_PIXEL(*PixelShader);
-			GraphicsPSOInit.PrimitiveType = PT_TriangleList;
-
-			SetGraphicsPipelineState(RHICmdList, GraphicsPSOInit);
-
-			PixelShader->SetParameters(RHICmdList, MonoView);
-
-			VertexShader->SetParameters(RHICmdList, MonoView, ViewFamily.MonoParameters.LateralOffset);
-
-			RHICmdList.SetViewport(LeftView.ViewRect.Min.X, LeftView.ViewRect.Min.Y, 0.0f, LeftView.ViewRect.Max.X, LeftView.ViewRect.Max.Y, 1.0f);
-			DrawRectangle(
-				RHICmdList,
-				0, 0,
-				LeftView.ViewRect.Width(), LeftView.ViewRect.Height(),
-				MonoView.ViewRect.Min.X, MonoView.ViewRect.Min.Y,
-				LeftView.ViewRect.Width(), LeftView.ViewRect.Height(),
-				LeftView.ViewRect.Size(),
-				MonoView.ViewRect.Max,
-				*VertexShader,
-				EDRF_UseTriangleOptimization);
+			FGraphicsPipelineStateInitializer GraphicsPSOInit;
+			RHICmdList.ApplyCachedRenderTargets(GraphicsPSOInit);
+
+			GraphicsPSOInit.RasterizerState = TStaticRasterizerState<FM_Solid, CM_None>::GetRHI();
+			GraphicsPSOInit.DepthStencilState = TStaticDepthStencilState<false, CF_Always>::GetRHI();
+
+			if (ViewFamily.MonoParameters.Mode == EMonoscopicFarFieldMode::MonoOnly)
+			{
+				GraphicsPSOInit.BlendState = TStaticBlendState<CW_RGBA, BO_Add, BF_One, BF_Zero>::GetRHI();
+			}
+			else
+			{
+				GraphicsPSOInit.BlendState = TStaticBlendState<CW_RGBA, BO_Add, BF_InverseDestAlpha, BF_One>::GetRHI();
+			}
+
+			if (!LeftView.bIsMobileMultiViewEnabled)
+			{
+				TShaderMapRef<FCompositeMonoscopicFarFieldViewVS<false>> VertexShader(MonoView.ShaderMap);
+				TShaderMapRef<FCompositeMonoscopicFarFieldViewPS<false>> PixelShader(MonoView.ShaderMap);
+
+				GraphicsPSOInit.BoundShaderState.VertexDeclarationRHI = GFilterVertexDeclaration.VertexDeclarationRHI;
+				GraphicsPSOInit.BoundShaderState.VertexShaderRHI = GETSAFERHISHADER_VERTEX(*VertexShader);
+				GraphicsPSOInit.BoundShaderState.PixelShaderRHI = GETSAFERHISHADER_PIXEL(*PixelShader);
+				GraphicsPSOInit.PrimitiveType = PT_TriangleList;
+
+				SetGraphicsPipelineState(RHICmdList, GraphicsPSOInit);
+
+				PixelShader->SetParameters(RHICmdList, MonoView);
+
+				VertexShader->SetParameters(RHICmdList, MonoView, ViewFamily.MonoParameters.LateralOffset);
+
+				const int32 LateralOffsetInPixels = static_cast<int32>(roundf(ViewFamily.MonoParameters.LateralOffset * static_cast<float>(MonoView.ViewRect.Width())));
+
+				RHICmdList.SetViewport(LeftView.ViewRect.Min.X, LeftView.ViewRect.Min.Y, 0.0f, LeftView.ViewRect.Max.X, LeftView.ViewRect.Max.Y, 1.0f);
+				DrawRectangle(
+					RHICmdList,
+					0, 0,
+					LeftView.ViewRect.Width(), LeftView.ViewRect.Height(),
+					MonoView.ViewRect.Min.X + LateralOffsetInPixels, MonoView.ViewRect.Min.Y,
+					LeftView.ViewRect.Width(), LeftView.ViewRect.Height(),
+					LeftView.ViewRect.Size(),
+					MonoView.ViewRect.Max,
+					*VertexShader,
+					EDRF_UseTriangleOptimization);
+
+				RHICmdList.SetViewport(RightView.ViewRect.Min.X, RightView.ViewRect.Min.Y, 0.0f, RightView.ViewRect.Max.X, RightView.ViewRect.Max.Y, 1.0f);
+
+				// Composite into right eye
+				DrawRectangle(
+					RHICmdList,
+					0, 0,
+					LeftView.ViewRect.Width(), LeftView.ViewRect.Height(),
+					MonoView.ViewRect.Min.X - LateralOffsetInPixels, MonoView.ViewRect.Min.Y,
+					LeftView.ViewRect.Width(), LeftView.ViewRect.Height(),
+					LeftView.ViewRect.Size(),
+					MonoView.ViewRect.Max,
+					*VertexShader,
+					EDRF_UseTriangleOptimization);
+
+			}
+			else
+			{
+				TShaderMapRef<FCompositeMonoscopicFarFieldViewVS<true>> VertexShader(MonoView.ShaderMap);
+				TShaderMapRef<FCompositeMonoscopicFarFieldViewPS<true>> PixelShader(MonoView.ShaderMap);
+
+				GraphicsPSOInit.BoundShaderState.VertexDeclarationRHI = GFilterVertexDeclaration.VertexDeclarationRHI;
+				GraphicsPSOInit.BoundShaderState.VertexShaderRHI = GETSAFERHISHADER_VERTEX(*VertexShader);
+				GraphicsPSOInit.BoundShaderState.PixelShaderRHI = GETSAFERHISHADER_PIXEL(*PixelShader);
+				GraphicsPSOInit.PrimitiveType = PT_TriangleList;
+
+				SetGraphicsPipelineState(RHICmdList, GraphicsPSOInit);
+
+				PixelShader->SetParameters(RHICmdList, MonoView);
+
+				VertexShader->SetParameters(RHICmdList, MonoView, ViewFamily.MonoParameters.LateralOffset);
+
+				RHICmdList.SetViewport(LeftView.ViewRect.Min.X, LeftView.ViewRect.Min.Y, 0.0f, LeftView.ViewRect.Max.X, LeftView.ViewRect.Max.Y, 1.0f);
+				DrawRectangle(
+					RHICmdList,
+					0, 0,
+					LeftView.ViewRect.Width(), LeftView.ViewRect.Height(),
+					MonoView.ViewRect.Min.X, MonoView.ViewRect.Min.Y,
+					LeftView.ViewRect.Width(), LeftView.ViewRect.Height(),
+					LeftView.ViewRect.Size(),
+					MonoView.ViewRect.Max,
+					*VertexShader,
+					EDRF_UseTriangleOptimization);
+			}
 		}
+		RHICmdList.EndRenderPass();
 	}
 
 	// Remove the mono view before moving to post.
