Renderpass updates.
Distortion Rendering. Self contained.

diff --git a/Engine/Source/Runtime/Renderer/Private/DistortionRendering.cpp b/Engine/Source/Runtime/Renderer/Private/DistortionRendering.cpp
--- a/Engine/Source/Runtime/Renderer/Private/DistortionRendering.cpp
+++ b/Engine/Source/Runtime/Renderer/Private/DistortionRendering.cpp
@@ -1092,41 +1092,44 @@ void FSceneRenderer::RenderDistortion(FRHICommandListImmediate& RHICmdList)
 		// DistortionRT==0 should never happen but better we don't crash
 		if(DistortionRT)
 		{
-			FRHIRenderTargetView ColorView( DistortionRT->GetRenderTargetItem().TargetableTexture, 0, -1, ERenderTargetLoadAction::EClear, ERenderTargetStoreAction::EStore );
-			FRHIDepthRenderTargetView DepthView( SceneContext.GetSceneDepthSurface(), ERenderTargetLoadAction::ELoad, ERenderTargetStoreAction::ENoAction, FExclusiveDepthStencil::DepthRead_StencilWrite );
-			FRHISetRenderTargetsInfo Info( 1, &ColorView, DepthView );
+			FRHIRenderPassInfo RPInfo(DistortionRT->GetRenderTargetItem().TargetableTexture, ERenderTargetActions::Clear_Store);
+			RPInfo.DepthStencilRenderTarget.Action = MakeDepthStencilTargetActions(ERenderTargetActions::Load_DontStore, ERenderTargetActions::DontLoad_Store);
+			RPInfo.DepthStencilRenderTarget.DepthStencilTarget = SceneContext.GetSceneDepthSurface();
+			RPInfo.DepthStencilRenderTarget.ExclusiveDepthStencil = FExclusiveDepthStencil::DepthRead_StencilWrite;
 
-			RHICmdList.SetRenderTargetsAndClear(Info);
-
-			for(int32 ViewIndex = 0;ViewIndex < Views.Num();ViewIndex++)
+			RHICmdList.BeginRenderPass(RPInfo, TEXT("RenderDistortion"));
 			{
-				SCOPED_CONDITIONAL_DRAW_EVENTF(RHICmdList, EventView, Views.Num() > 1, TEXT("View%d"), ViewIndex);
+				for (int32 ViewIndex = 0; ViewIndex < Views.Num(); ViewIndex++)
+				{
+					SCOPED_CONDITIONAL_DRAW_EVENTF(RHICmdList, EventView, Views.Num() > 1, TEXT("View%d"), ViewIndex);
 
-				FViewInfo& View = Views[ViewIndex];
-				// viewport to match view size
-				RHICmdList.SetViewport(View.ViewRect.Min.X, View.ViewRect.Min.Y, 0.0f, View.ViewRect.Max.X, View.ViewRect.Max.Y, 1.0f);
+					FViewInfo& View = Views[ViewIndex];
+					// viewport to match view size
+					RHICmdList.SetViewport(View.ViewRect.Min.X, View.ViewRect.Min.Y, 0.0f, View.ViewRect.Max.X, View.ViewRect.Max.Y, 1.0f);
 
-				FSceneTexturesUniformParameters SceneTextureParameters;
-				SetupSceneTextureUniformParameters(SceneContext, View.FeatureLevel, ESceneTextureSetupMode::All, SceneTextureParameters);
-				TUniformBufferRef<FSceneTexturesUniformParameters> PassUniformBuffer = TUniformBufferRef<FSceneTexturesUniformParameters>::CreateUniformBufferImmediate(SceneTextureParameters, UniformBuffer_SingleFrame);
+					FSceneTexturesUniformParameters SceneTextureParameters;
+					SetupSceneTextureUniformParameters(SceneContext, View.FeatureLevel, ESceneTextureSetupMode::All, SceneTextureParameters);
+					TUniformBufferRef<FSceneTexturesUniformParameters> PassUniformBuffer = TUniformBufferRef<FSceneTexturesUniformParameters>::CreateUniformBufferImmediate(SceneTextureParameters, UniformBuffer_SingleFrame);
 
-				FDrawingPolicyRenderState DrawRenderState(View, PassUniformBuffer);
+					FDrawingPolicyRenderState DrawRenderState(View, PassUniformBuffer);
 
-				// test against depth and write stencil mask
-				DrawRenderState.SetDepthStencilState(TStaticDepthStencilState<
-					false, CF_DepthNearOrEqual,
-					true, CF_Always, SO_Keep, SO_Keep, SO_Replace,
-					false, CF_Always, SO_Keep, SO_Keep, SO_Keep,
-					kStencilMaskBit, kStencilMaskBit>::GetRHI());
-				DrawRenderState.SetStencilRef(kStencilMaskBit);
+					// test against depth and write stencil mask
+					DrawRenderState.SetDepthStencilState(TStaticDepthStencilState<
+						false, CF_DepthNearOrEqual,
+						true, CF_Always, SO_Keep, SO_Keep, SO_Replace,
+						false, CF_Always, SO_Keep, SO_Keep, SO_Keep,
+						kStencilMaskBit, kStencilMaskBit>::GetRHI());
+					DrawRenderState.SetStencilRef(kStencilMaskBit);
 
-				// additive blending of offsets (or complexity if the shader complexity viewmode is enabled)
-				DrawRenderState.SetBlendState(TStaticBlendState<CW_RGBA, BO_Add, BF_One, BF_One, BO_Add, BF_One, BF_One>::GetRHI());
+					// additive blending of offsets (or complexity if the shader complexity viewmode is enabled)
+					DrawRenderState.SetBlendState(TStaticBlendState<CW_RGBA, BO_Add, BF_One, BF_One, BO_Add, BF_One, BF_One>::GetRHI());
 
-				// draw only distortion meshes to accumulate their offsets
-				bDirty |= View.DistortionPrimSet.DrawAccumulatedOffsets(RHICmdList, View, DrawRenderState, false);
-				bDebugShader |= View.Family->GetDebugViewShaderMode() != DVSM_None;
+					// draw only distortion meshes to accumulate their offsets
+					bDirty |= View.DistortionPrimSet.DrawAccumulatedOffsets(RHICmdList, View, DrawRenderState, false);
+					bDebugShader |= View.Family->GetDebugViewShaderMode() != DVSM_None;
+				}
 			}
+			RHICmdList.EndRenderPass();
 
 			if (bDirty)
 			{
@@ -1152,43 +1155,63 @@ void FSceneRenderer::RenderDistortion(FRHICommandListImmediate& RHICmdList)
 		const FSceneRenderTargetItem& DestRenderTarget = NewSceneColor->GetRenderTargetItem();
 
 		// Apply distortion and store off-screen
-		SetRenderTarget(RHICmdList, DestRenderTarget.TargetableTexture, SceneContext.GetSceneDepthSurface(), ESimpleRenderTargetMode::EExistingColorAndDepth, FExclusiveDepthStencil::DepthRead_StencilRead);
-
-		for(int32 ViewIndex = 0, Num = Views.Num(); ViewIndex < Num; ++ViewIndex)
 		{
-			QUICK_SCOPE_CYCLE_COUNTER(STAT_FSceneRenderer_RenderDistortion_PostView1);
-			SCOPED_CONDITIONAL_DRAW_EVENTF(RHICmdList, EventView, Views.Num() > 1, TEXT("View%d"), ViewIndex);
+			FRHIRenderPassInfo RPInfo(DestRenderTarget.TargetableTexture, ERenderTargetActions::Load_Store);
+			RPInfo.DepthStencilRenderTarget.Action = MakeDepthStencilTargetActions(ERenderTargetActions::Load_DontStore, ERenderTargetActions::Load_DontStore);
+			RPInfo.DepthStencilRenderTarget.DepthStencilTarget = SceneContext.GetSceneDepthSurface();
+			RPInfo.DepthStencilRenderTarget.ExclusiveDepthStencil = FExclusiveDepthStencil::DepthRead_StencilRead;
 
-			FViewInfo& View = Views[ViewIndex];
-
-			if (MSAACount == 1)
-			{
-				DrawDistortionApplyScreenPass<false>(RHICmdList, SceneContext, View, *DistortionRT);
-			}
-			else
+			RHICmdList.BeginRenderPass(RPInfo, TEXT("DistortedSceneColor"));
 			{
-				DrawDistortionApplyScreenPass<true>(RHICmdList, SceneContext, View, *DistortionRT);
+				for (int32 ViewIndex = 0, Num = Views.Num(); ViewIndex < Num; ++ViewIndex)
+				{
+					QUICK_SCOPE_CYCLE_COUNTER(STAT_FSceneRenderer_RenderDistortion_PostView1);
+					SCOPED_CONDITIONAL_DRAW_EVENTF(RHICmdList, EventView, Views.Num() > 1, TEXT("View%d"), ViewIndex);
+
+					FViewInfo& View = Views[ViewIndex];
+
+					if (MSAACount == 1)
+					{
+						DrawDistortionApplyScreenPass<false>(RHICmdList, SceneContext, View, *DistortionRT);
+					}
+					else
+					{
+						DrawDistortionApplyScreenPass<true>(RHICmdList, SceneContext, View, *DistortionRT);
+					}
+				}
 			}
-		}
+			RHICmdList.EndRenderPass();
 
-		RHICmdList.TransitionResource(EResourceTransitionAccess::EReadable, DestRenderTarget.TargetableTexture);
-		SetRenderTarget(RHICmdList, SceneContext.GetSceneColor()->GetRenderTargetItem().TargetableTexture, SceneContext.GetSceneDepthSurface(), ESimpleRenderTargetMode::EExistingColorAndDepth, FExclusiveDepthStencil::DepthRead_StencilWrite, true);
+			RHICmdList.TransitionResource(EResourceTransitionAccess::EReadable, DestRenderTarget.TargetableTexture);
+		}
 
-		for(int32 ViewIndex = 0, Num = Views.Num(); ViewIndex < Num; ++ViewIndex)
 		{
-			QUICK_SCOPE_CYCLE_COUNTER(STAT_FSceneRenderer_RenderDistortion_PostView2);
-			SCOPED_CONDITIONAL_DRAW_EVENTF(RHICmdList, EventView, Views.Num() > 1, TEXT("View%d"), ViewIndex);
-
-			FViewInfo& View = Views[ViewIndex];
+			FRHIRenderPassInfo RPInfo(SceneContext.GetSceneColor()->GetRenderTargetItem().TargetableTexture, ERenderTargetActions::Load_Store);
+			RPInfo.DepthStencilRenderTarget.Action = MakeDepthStencilTargetActions(ERenderTargetActions::Load_DontStore, ERenderTargetActions::Load_Store);
+			RPInfo.DepthStencilRenderTarget.DepthStencilTarget = SceneContext.GetSceneDepthSurface();
+			RPInfo.DepthStencilRenderTarget.ExclusiveDepthStencil = FExclusiveDepthStencil::DepthRead_StencilWrite;
 
-			if (MSAACount == 1)
+			TransitionRenderPassTargets(RHICmdList, RPInfo);
+			RHICmdList.BeginRenderPass(RPInfo, TEXT("DistortionMerge"));
 			{
-				DrawDistortionMergePass<false>(RHICmdList, SceneContext, View, DestRenderTarget.TargetableTexture);
-			}
-			else
-			{
-				DrawDistortionMergePass<true>(RHICmdList, SceneContext, View, DestRenderTarget.TargetableTexture);
+				for (int32 ViewIndex = 0, Num = Views.Num(); ViewIndex < Num; ++ViewIndex)
+				{
+					QUICK_SCOPE_CYCLE_COUNTER(STAT_FSceneRenderer_RenderDistortion_PostView2);
+					SCOPED_CONDITIONAL_DRAW_EVENTF(RHICmdList, EventView, Views.Num() > 1, TEXT("View%d"), ViewIndex);
+
+					FViewInfo& View = Views[ViewIndex];
+
+					if (MSAACount == 1)
+					{
+						DrawDistortionMergePass<false>(RHICmdList, SceneContext, View, DestRenderTarget.TargetableTexture);
+					}
+					else
+					{
+						DrawDistortionMergePass<true>(RHICmdList, SceneContext, View, DestRenderTarget.TargetableTexture);
+					}
+				}
 			}
+			RHICmdList.EndRenderPass();
 		}
 	}
 }
