Renderpasses. Fixes vulkan rendering.

#ryi rolando.caloca

diff --git a/Engine/Source/Runtime/Renderer/Private/PostProcess/PostProcessSelectionOutline.cpp b/Engine/Source/Runtime/Renderer/Private/PostProcess/PostProcessSelectionOutline.cpp
--- a/Engine/Source/Runtime/Renderer/Private/PostProcess/PostProcessSelectionOutline.cpp
+++ b/Engine/Source/Runtime/Renderer/Private/PostProcess/PostProcessSelectionOutline.cpp
@@ -50,7 +50,7 @@ void FRCPassPostProcessSelectionOutlineColor::Process(FRenderingCompositePassCon
 	FRHIRenderPassInfo RPInfo;
 	RPInfo.DepthStencilRenderTarget.DepthStencilTarget = DestRenderTarget.TargetableTexture;
 	RPInfo.DepthStencilRenderTarget.ResolveTarget = DestRenderTarget.ShaderResourceTexture;
-	RPInfo.DepthStencilRenderTarget.Action = MakeDepthStencilTargetActions(ERenderTargetActions::Clear_DontStore, ERenderTargetActions::Clear_DontStore);
+	RPInfo.DepthStencilRenderTarget.Action = MakeDepthStencilTargetActions(ERenderTargetActions::Clear_Store, ERenderTargetActions::Clear_Store);
 	RPInfo.DepthStencilRenderTarget.ExclusiveDepthStencil = FExclusiveDepthStencil::DepthWrite_StencilWrite;
 
 	if (GRHIRequiresRenderTargetForPixelShaderUAVs)
diff --git a/Engine/Source/Runtime/Renderer/Private/PostProcess/PostProcessTemporalAA.cpp b/Engine/Source/Runtime/Renderer/Private/PostProcess/PostProcessTemporalAA.cpp
--- a/Engine/Source/Runtime/Renderer/Private/PostProcess/PostProcessTemporalAA.cpp
+++ b/Engine/Source/Runtime/Renderer/Private/PostProcess/PostProcessTemporalAA.cpp
@@ -889,7 +889,7 @@ void FRCPassPostProcessTemporalAA::Process(FRenderingCompositePassContext& Conte
 
 		RPInfo.DepthStencilRenderTarget.DepthStencilTarget = SceneContext.GetSceneDepthTexture();
 		RPInfo.DepthStencilRenderTarget.ResolveTarget = nullptr;
-		RPInfo.DepthStencilRenderTarget.Action = MakeDepthStencilTargetActions(ERenderTargetActions::DontLoad_DontStore, ERenderTargetActions::Load_DontStore);
+		RPInfo.DepthStencilRenderTarget.Action = MakeDepthStencilTargetActions(ERenderTargetActions::DontLoad_DontStore, ERenderTargetActions::Load_Store);
 		RPInfo.DepthStencilRenderTarget.ExclusiveDepthStencil = FExclusiveDepthStencil::DepthRead_StencilWrite;
 
 		Context.RHICmdList.BeginRenderPass(RPInfo, TEXT("TemporalAA"));
