Fix some leaked RHI resources

diff --git a/Engine/Source/Runtime/Renderer/Private/MobileShadingRenderer.cpp b/Engine/Source/Runtime/Renderer/Private/MobileShadingRenderer.cpp
--- a/Engine/Source/Runtime/Renderer/Private/MobileShadingRenderer.cpp
+++ b/Engine/Source/Runtime/Renderer/Private/MobileShadingRenderer.cpp
@@ -80,10 +80,36 @@ FMobileSceneRenderer::FMobileSceneRenderer(const FSceneViewFamily* InViewFamily,
 	bModulatedShadowsInUse = false;
 }
 
+class FMobileDirLightShaderParamsRenderResource : public FRenderResource
+{
+public:
+	using MobileDirLightUniformBufferRef = TUniformBufferRef<FMobileDirectionalLightShaderParameters>;
+
+	virtual void InitRHI() override
+	{
+		UniformBufferRHI =
+			MobileDirLightUniformBufferRef::CreateUniformBufferImmediate(
+				FMobileDirectionalLightShaderParameters(),
+				UniformBuffer_MultiFrame);
+	}
+
+	virtual void ReleaseRHI() override
+	{
+		UniformBufferRHI.SafeRelease();
+	}
+
+	MobileDirLightUniformBufferRef UniformBufferRHI;
+};
+
 TUniformBufferRef<FMobileDirectionalLightShaderParameters>& GetNullMobileDirectionalLightShaderParameters()
 {
-	static TUniformBufferRef<FMobileDirectionalLightShaderParameters> NullLightParams = TUniformBufferRef<FMobileDirectionalLightShaderParameters>::CreateUniformBufferImmediate(FMobileDirectionalLightShaderParameters(), UniformBuffer_MultiFrame);
-	return NullLightParams;
+	static TGlobalResource<FMobileDirLightShaderParamsRenderResource>* NullLightParams;
+	if (!NullLightParams)
+	{
+		NullLightParams = new TGlobalResource<FMobileDirLightShaderParamsRenderResource>();
+	}
+	check(!!NullLightParams->UniformBufferRHI);
+	return NullLightParams->UniformBufferRHI;
 }
 
 /**
diff --git a/Engine/Source/Runtime/Renderer/Private/ScenePrivate.h b/Engine/Source/Runtime/Renderer/Private/ScenePrivate.h
--- a/Engine/Source/Runtime/Renderer/Private/ScenePrivate.h
+++ b/Engine/Source/Runtime/Renderer/Private/ScenePrivate.h
@@ -1224,7 +1224,11 @@ public:
 		CapsuleTileIntersectionCountsBuffer.Release();
 		TranslucencyTimer.Release();
 		SeparateTranslucencyTimer.Release();
-		ForwardLightingResources.Release();
+		if (!!ForwardLightingResources)
+		{
+			ForwardLightingResources->Release();
+			ForwardLightingResources = nullptr;
+		}
 		ForwardLightingCullingResources.Release();
 		LightScatteringHistory.SafeRelease();
 	}
diff --git a/Engine/Source/Runtime/Renderer/Private/SceneRendering.cpp b/Engine/Source/Runtime/Renderer/Private/SceneRendering.cpp
--- a/Engine/Source/Runtime/Renderer/Private/SceneRendering.cpp
+++ b/Engine/Source/Runtime/Renderer/Private/SceneRendering.cpp
@@ -1433,10 +1433,11 @@ FViewInfo* FViewInfo::CreateSnapshot() const
 	for (size_t i = 0; i < ARRAY_COUNT(Result->MobileDirectionalLightUniformBuffers); i++)
 	{
 		// This memcpy is necessary to clear the reference from the memcpy of the whole of this -> Result without releasing the pointer
-		FMemory::Memcpy(Result->MobileDirectionalLightUniformBuffers[i], NullMobileDirectionalLightUniformBuffer);
-		
 		// But what we really want is the null buffer.
-		Result->MobileDirectionalLightUniformBuffers[i] = GetNullMobileDirectionalLightShaderParameters();
+		FMemory::Memcpy(
+			&Result->MobileDirectionalLightUniformBuffers[i],
+			&GetNullMobileDirectionalLightShaderParameters(),
+			sizeof(TUniformBufferRef<FMobileDirectionalLightShaderParameters>));
 	}
 
 	TUniquePtr<FViewUniformShaderParameters> NullViewParameters;
