Clamping r.TranslucencyLightingVolumeDim in order to fix crash when it's 0.

diff --git a/Engine/Source/Runtime/Renderer/Private/PostProcess/SceneRenderTargets.cpp b/Engine/Source/Runtime/Renderer/Private/PostProcess/SceneRenderTargets.cpp
--- a/Engine/Source/Runtime/Renderer/Private/PostProcess/SceneRenderTargets.cpp
+++ b/Engine/Source/Runtime/Renderer/Private/PostProcess/SceneRenderTargets.cpp
@@ -551,7 +551,7 @@ void FSceneRenderTargets::Allocate(FRHICommandListImmediate& RHICmdList, const F
 		}
 	}
 
-	int32 TranslucencyLightingVolumeDim = GTranslucencyLightingVolumeDim;
+	const int32 TranslucencyLightingVolumeDim = GetTranslucencyLightingVolumeDim();
 
 	uint32 Mobile32bpp = !IsMobileHDR() || IsMobileHDR32bpp();
 
@@ -2054,6 +2054,8 @@ void FSceneRenderTargets::AllocateDeferredShadingPathRenderTargets(FRHICommandLi
 		}
 		
 		{
+			const int32 TranslucencyLightingVolumeDim = GetTranslucencyLightingVolumeDim();
+
 			// TODO: We can skip the and TLV allocations when rendering in forward shading mode
 			uint32 TranslucencyTargetFlags = TexCreate_ShaderResource | TexCreate_RenderTargetable | TexCreate_ReduceMemoryWithTilingMode;
 
@@ -2070,9 +2072,9 @@ void FSceneRenderTargets::AllocateDeferredShadingPathRenderTargets(FRHICommandLi
 				GRenderTargetPool.FindFreeElement(
 					RHICmdList,
 					FPooledRenderTargetDesc(FPooledRenderTargetDesc::CreateVolumeDesc(
-						GTranslucencyLightingVolumeDim,
-						GTranslucencyLightingVolumeDim,
-						GTranslucencyLightingVolumeDim,
+						TranslucencyLightingVolumeDim,
+						TranslucencyLightingVolumeDim,
+						TranslucencyLightingVolumeDim,
 						PF_FloatRGBA,
 						FClearValueBinding::Transparent,
 						0,
@@ -2089,15 +2091,15 @@ void FSceneRenderTargets::AllocateDeferredShadingPathRenderTargets(FRHICommandLi
 				//Tests to catch UE-31578, UE-32536 and UE-22073 crash (Defferred Render Targets not being allocated)
 				ensureMsgf(TranslucencyLightingVolumeAmbient[RTSetIndex], TEXT("Failed to allocate render target %s with dimension %i and flags %i"),
 					GetVolumeName(RTSetIndex, false),
-					GTranslucencyLightingVolumeDim,
+					TranslucencyLightingVolumeDim,
 					TranslucencyTargetFlags);
 
 				GRenderTargetPool.FindFreeElement(
 					RHICmdList,
 					FPooledRenderTargetDesc(FPooledRenderTargetDesc::CreateVolumeDesc(
-						GTranslucencyLightingVolumeDim,
-						GTranslucencyLightingVolumeDim,
-						GTranslucencyLightingVolumeDim,
+						TranslucencyLightingVolumeDim,
+						TranslucencyLightingVolumeDim,
+						TranslucencyLightingVolumeDim,
 						PF_FloatRGBA,
 						FClearValueBinding::Transparent,
 						0,
@@ -2114,7 +2116,7 @@ void FSceneRenderTargets::AllocateDeferredShadingPathRenderTargets(FRHICommandLi
 				//Tests to catch UE-31578, UE-32536 and UE-22073 crash
 				ensureMsgf(TranslucencyLightingVolumeDirectional[RTSetIndex], TEXT("Failed to allocate render target %s with dimension %i and flags %i"),
 					GetVolumeName(RTSetIndex, true),
-					GTranslucencyLightingVolumeDim,
+					TranslucencyLightingVolumeDim,
 					TranslucencyTargetFlags);
 			}
 
@@ -2228,7 +2230,9 @@ void FSceneRenderTargets::ClearVolumeTextures(FRHICommandList& RHICmdList, ERHIF
 	GraphicsPSOInit.DepthStencilState = TStaticDepthStencilState<false, CF_Always>::GetRHI();
 	GraphicsPSOInit.BlendState = TStaticBlendState<>::GetRHI();
 
-	const FVolumeBounds VolumeBounds(GTranslucencyLightingVolumeDim);
+	const int32 TranslucencyLightingVolumeDim = GetTranslucencyLightingVolumeDim();
+
+	const FVolumeBounds VolumeBounds(TranslucencyLightingVolumeDim);
 	auto ShaderMap = GetGlobalShaderMap(FeatureLevel);
 	TShaderMapRef<FWriteToSliceVS> VertexShader(ShaderMap);
 	TOptionalShaderMapRef<FWriteToSliceGS> GeometryShader(ShaderMap);
@@ -2242,7 +2246,7 @@ void FSceneRenderTargets::ClearVolumeTextures(FRHICommandList& RHICmdList, ERHIF
 
 	SetGraphicsPipelineState(RHICmdList, GraphicsPSOInit);
 
-	VertexShader->SetParameters(RHICmdList, VolumeBounds, FIntVector(GTranslucencyLightingVolumeDim));
+	VertexShader->SetParameters(RHICmdList, VolumeBounds, FIntVector(TranslucencyLightingVolumeDim));
 	if (GeometryShader.IsValid())
 	{
 		GeometryShader->SetParameters(RHICmdList, VolumeBounds.MinZ);
diff --git a/Engine/Source/Runtime/Renderer/Private/PostProcess/SceneRenderTargets.h b/Engine/Source/Runtime/Renderer/Private/PostProcess/SceneRenderTargets.h
--- a/Engine/Source/Runtime/Renderer/Private/PostProcess/SceneRenderTargets.h
+++ b/Engine/Source/Runtime/Renderer/Private/PostProcess/SceneRenderTargets.h
@@ -30,8 +30,12 @@ static const int32 NumTranslucentVolumeRenderTargetSets = (TVC_MAX + 1);
 /** Forward declaration of console variable controlling translucent volume blur */
 extern int32 GUseTranslucencyVolumeBlur;
 
-/** Forward declaration of console variable controlling translucent lighting volume dimensions */
-extern int32 GTranslucencyLightingVolumeDim;
+/** Function returning current translucency lighting volume dimensions. */
+inline int32 GetTranslucencyLightingVolumeDim()
+{
+	extern int32 GTranslucencyLightingVolumeDim;
+	return FMath::Clamp(GTranslucencyLightingVolumeDim, 4, 2048);
+}
 
 /** Function to select the index of the volume texture that we will hold the final translucency lighting volume texture */
 inline int SelectTranslucencyVolumeTarget(ETranslucencyVolumeCascade InCascade)
diff --git a/Engine/Source/Runtime/Renderer/Private/SceneRendering.cpp b/Engine/Source/Runtime/Renderer/Private/SceneRendering.cpp
--- a/Engine/Source/Runtime/Renderer/Private/SceneRendering.cpp
+++ b/Engine/Source/Runtime/Renderer/Private/SceneRendering.cpp
@@ -1217,11 +1217,13 @@ void FViewInfo::SetupUniformBufferParameters(
 		PrimaryView->CalcTranslucencyLightingVolumeBounds(OutTranslucentCascadeBoundsArray, NumTranslucentCascades);
 	}
 
+	const int32 TranslucencyLightingVolumeDim = GetTranslucencyLightingVolumeDim();
+
 	for (int32 CascadeIndex = 0; CascadeIndex < NumTranslucentCascades; CascadeIndex++)
 	{
-		const float VolumeVoxelSize = (OutTranslucentCascadeBoundsArray[CascadeIndex].Max.X - OutTranslucentCascadeBoundsArray[CascadeIndex].Min.X) / GTranslucencyLightingVolumeDim;
+		const float VolumeVoxelSize = (OutTranslucentCascadeBoundsArray[CascadeIndex].Max.X - OutTranslucentCascadeBoundsArray[CascadeIndex].Min.X) / TranslucencyLightingVolumeDim;
 		const FVector VolumeSize = OutTranslucentCascadeBoundsArray[CascadeIndex].Max - OutTranslucentCascadeBoundsArray[CascadeIndex].Min;
-		ViewUniformShaderParameters.TranslucencyLightingVolumeMin[CascadeIndex] = FVector4(OutTranslucentCascadeBoundsArray[CascadeIndex].Min, 1.0f / GTranslucencyLightingVolumeDim);
+		ViewUniformShaderParameters.TranslucencyLightingVolumeMin[CascadeIndex] = FVector4(OutTranslucentCascadeBoundsArray[CascadeIndex].Min, 1.0f / TranslucencyLightingVolumeDim);
 		ViewUniformShaderParameters.TranslucencyLightingVolumeInvSize[CascadeIndex] = FVector4(FVector(1.0f) / VolumeSize, VolumeVoxelSize);
 	}
 	
@@ -1387,10 +1389,12 @@ void FViewInfo::InitRHIResources()
 
 	ViewUniformBuffer = TUniformBufferRef<FViewUniformShaderParameters>::CreateUniformBufferImmediate(*CachedViewUniformShaderParameters, UniformBuffer_SingleFrame);
 
+	const int32 TranslucencyLightingVolumeDim = GetTranslucencyLightingVolumeDim();
+
 	for (int32 CascadeIndex = 0; CascadeIndex < TVC_MAX; CascadeIndex++)
 	{
 		TranslucencyLightingVolumeMin[CascadeIndex] = VolumeBounds[CascadeIndex].Min;
-		TranslucencyVolumeVoxelSize[CascadeIndex] = (VolumeBounds[CascadeIndex].Max.X - VolumeBounds[CascadeIndex].Min.X) / GTranslucencyLightingVolumeDim;
+		TranslucencyVolumeVoxelSize[CascadeIndex] = (VolumeBounds[CascadeIndex].Max.X - VolumeBounds[CascadeIndex].Min.X) / TranslucencyLightingVolumeDim;
 		TranslucencyLightingVolumeSize[CascadeIndex] = VolumeBounds[CascadeIndex].Max - VolumeBounds[CascadeIndex].Min;
 	}
 
diff --git a/Engine/Source/Runtime/Renderer/Private/TranslucentLighting.cpp b/Engine/Source/Runtime/Renderer/Private/TranslucentLighting.cpp
--- a/Engine/Source/Runtime/Renderer/Private/TranslucentLighting.cpp
+++ b/Engine/Source/Runtime/Renderer/Private/TranslucentLighting.cpp
@@ -174,9 +174,10 @@ void FViewInfo::CalcTranslucencyLightingVolumeBounds(FBox* InOutCascadeBoundsArr
 		FSphere SphereBounds(Center, FMath::Sqrt(RadiusSquared));
 
 		// Snap the center to a multiple of the volume dimension for stability
-		SphereBounds.Center.X = SphereBounds.Center.X - FMath::Fmod(SphereBounds.Center.X, SphereBounds.W * 2 / GTranslucencyLightingVolumeDim);
-		SphereBounds.Center.Y = SphereBounds.Center.Y - FMath::Fmod(SphereBounds.Center.Y, SphereBounds.W * 2 / GTranslucencyLightingVolumeDim);
-		SphereBounds.Center.Z = SphereBounds.Center.Z - FMath::Fmod(SphereBounds.Center.Z, SphereBounds.W * 2 / GTranslucencyLightingVolumeDim);
+		const int32 TranslucencyLightingVolumeDim = GetTranslucencyLightingVolumeDim();
+		SphereBounds.Center.X = SphereBounds.Center.X - FMath::Fmod(SphereBounds.Center.X, SphereBounds.W * 2 / TranslucencyLightingVolumeDim);
+		SphereBounds.Center.Y = SphereBounds.Center.Y - FMath::Fmod(SphereBounds.Center.Y, SphereBounds.W * 2 / TranslucencyLightingVolumeDim);
+		SphereBounds.Center.Z = SphereBounds.Center.Z - FMath::Fmod(SphereBounds.Center.Z, SphereBounds.W * 2 / TranslucencyLightingVolumeDim);
 
 		InOutCascadeBoundsArray[CascadeIndex] = FBox(SphereBounds.Center - SphereBounds.W, SphereBounds.Center + SphereBounds.W);
 	}
@@ -645,7 +646,8 @@ public:
 
 		FGlobalShader::SetParameters<FViewUniformShaderParameters>(RHICmdList, ShaderRHI, View.ViewUniformBuffer);
 
-		SetShaderValue(RHICmdList, ShaderRHI, TexelSize, 1.0f / GTranslucencyLightingVolumeDim);
+		const int32 TranslucencyLightingVolumeDim = GetTranslucencyLightingVolumeDim();
+		SetShaderValue(RHICmdList, ShaderRHI, TexelSize, 1.0f / TranslucencyLightingVolumeDim);
 
 		SetTextureParameter(
 			RHICmdList,
@@ -1010,10 +1012,12 @@ void FDeferredShadingSceneRenderer::ClearTranslucentVolumeLightingAsyncCompute(F
 		//write fence on the Gfx pipe so the async clear compute shader won't clear until the Gfx pipe is caught up.
 		RHICmdList.TransitionResources(EResourceTransitionAccess::ERWBarrier, EResourceTransitionPipeline::EGfxToCompute, VolumeUAVs, NumUAVs, ClearBeginFence);
 
+		const int32 TranslucencyLightingVolumeDim = GetTranslucencyLightingVolumeDim();
+
 		//Grab the async compute commandlist.
 		FRHIAsyncComputeCommandListImmediate& RHICmdListComputeImmediate = FRHICommandListExecutor::GetImmediateAsyncComputeCommandList();
 		{
-			SCOPED_COMPUTE_EVENTF(RHICmdListComputeImmediate, ClearTranslucencyLightingVolume, TEXT("ClearTranslucencyLightingVolumeCompute %d"), GTranslucencyLightingVolumeDim);
+			SCOPED_COMPUTE_EVENTF(RHICmdListComputeImmediate, ClearTranslucencyLightingVolume, TEXT("ClearTranslucencyLightingVolumeCompute %d"), TranslucencyLightingVolumeDim);
 
 			//we must wait on the fence written from the Gfx pipe to let us know all our dependencies are ready.
 			RHICmdListComputeImmediate.WaitComputeFence(ClearBeginFence);
@@ -1023,7 +1027,7 @@ void FDeferredShadingSceneRenderer::ClearTranslucentVolumeLightingAsyncCompute(F
 
 			ComputeShader->SetParameters(RHICmdListComputeImmediate, VolumeUAVs, NumUAVs);
 		
-			int32 GroupsPerDim = GTranslucencyLightingVolumeDim / FClearTranslucentLightingVolumeCS::CLEAR_BLOCK_SIZE;
+			int32 GroupsPerDim = TranslucencyLightingVolumeDim / FClearTranslucentLightingVolumeCS::CLEAR_BLOCK_SIZE;
 			DispatchComputeShader(RHICmdListComputeImmediate, ComputeShader, GroupsPerDim, GroupsPerDim, GroupsPerDim);
 
 			ComputeShader->UnsetParameters(RHICmdListComputeImmediate);
@@ -1091,7 +1095,9 @@ void FDeferredShadingSceneRenderer::InjectAmbientCubemapTranslucentVolumeLightin
 		SCOPED_DRAW_EVENT(RHICmdList, InjectAmbientCubemapTranslucentVolumeLighting);
 		SCOPED_GPU_STAT(RHICmdList, TranslucentLighting);
 
-		const FVolumeBounds VolumeBounds(GTranslucencyLightingVolumeDim);
+		const int32 TranslucencyLightingVolumeDim = GetTranslucencyLightingVolumeDim();
+
+		const FVolumeBounds VolumeBounds(TranslucencyLightingVolumeDim);
 
 		auto ShaderMap = GetGlobalShaderMap(FeatureLevel);
 		
@@ -1122,7 +1128,7 @@ void FDeferredShadingSceneRenderer::InjectAmbientCubemapTranslucentVolumeLightin
 									
 			SetGraphicsPipelineState(RHICmdList, GraphicsPSOInit);
 
-			VertexShader->SetParameters(RHICmdList, VolumeBounds, FIntVector(GTranslucencyLightingVolumeDim));
+			VertexShader->SetParameters(RHICmdList, VolumeBounds, FIntVector(TranslucencyLightingVolumeDim));
 			if(GeometryShader.IsValid())
 			{
 				GeometryShader->SetParameters(RHICmdList, VolumeBounds.MinZ);
@@ -1168,11 +1174,13 @@ void FDeferredShadingSceneRenderer::ClearTranslucentVolumePerObjectShadowing(FRH
 /** Calculates volume texture bounds for the given light in the given translucent lighting volume cascade. */
 FVolumeBounds CalculateLightVolumeBounds(const FSphere& LightBounds, const FViewInfo& View, uint32 VolumeCascadeIndex, bool bDirectionalLight)
 {
+	const int32 TranslucencyLightingVolumeDim = GetTranslucencyLightingVolumeDim();
+
 	FVolumeBounds VolumeBounds;
 
 	if (bDirectionalLight)
 	{
-		VolumeBounds = FVolumeBounds(GTranslucencyLightingVolumeDim);
+		VolumeBounds = FVolumeBounds(TranslucencyLightingVolumeDim);
 	}
 	else
 	{
@@ -1184,9 +1192,9 @@ FVolumeBounds CalculateLightVolumeBounds(const FSphere& LightBounds, const FView
 		VolumeBounds.MinY = FMath::Max(FMath::TruncToInt(MinPosition.Y), 0);
 		VolumeBounds.MinZ = FMath::Max(FMath::TruncToInt(MinPosition.Z), 0);
 
-		VolumeBounds.MaxX = FMath::Min(FMath::TruncToInt(MaxPosition.X) + 1, GTranslucencyLightingVolumeDim);
-		VolumeBounds.MaxY = FMath::Min(FMath::TruncToInt(MaxPosition.Y) + 1, GTranslucencyLightingVolumeDim);
-		VolumeBounds.MaxZ = FMath::Min(FMath::TruncToInt(MaxPosition.Z) + 1, GTranslucencyLightingVolumeDim);
+		VolumeBounds.MaxX = FMath::Min(FMath::TruncToInt(MaxPosition.X) + 1, TranslucencyLightingVolumeDim);
+		VolumeBounds.MaxY = FMath::Min(FMath::TruncToInt(MaxPosition.Y) + 1, TranslucencyLightingVolumeDim);
+		VolumeBounds.MaxZ = FMath::Min(FMath::TruncToInt(MaxPosition.Z) + 1, TranslucencyLightingVolumeDim);
 	}
 
 	return VolumeBounds;
@@ -1251,7 +1259,9 @@ void FDeferredShadingSceneRenderer::AccumulateTranslucentVolumeObjectShadowing(F
 
 				SetGraphicsPipelineState(RHICmdList, GraphicsPSOInit);
 
-				VertexShader->SetParameters(RHICmdList, VolumeBounds, FIntVector(GTranslucencyLightingVolumeDim));
+				const int32 TranslucencyLightingVolumeDim = GetTranslucencyLightingVolumeDim();
+
+				VertexShader->SetParameters(RHICmdList, VolumeBounds, FIntVector(TranslucencyLightingVolumeDim));
 				if(GeometryShader.IsValid())
 				{
 					GeometryShader->SetParameters(RHICmdList, VolumeBounds.MinZ);
@@ -1507,7 +1517,9 @@ static void InjectTranslucentLightArray(FRHICommandListImmediate& RHICmdList, co
 					}
 				}
 
-				VertexShader->SetParameters(RHICmdList, VolumeBounds, FIntVector(GTranslucencyLightingVolumeDim));
+				const int32 TranslucencyLightingVolumeDim = GetTranslucencyLightingVolumeDim();
+
+				VertexShader->SetParameters(RHICmdList, VolumeBounds, FIntVector(TranslucencyLightingVolumeDim));
 				if(GeometryShader.IsValid())
 				{
 					GeometryShader->SetParameters(RHICmdList, VolumeBounds.MinZ);
@@ -1688,7 +1700,9 @@ void FDeferredShadingSceneRenderer::InjectSimpleTranslucentVolumeLightingArray(F
 
 						SetGraphicsPipelineState(RHICmdList, GraphicsPSOInit);
 
-						VertexShader->SetParameters(RHICmdList, VolumeBounds, FIntVector(GTranslucencyLightingVolumeDim));
+						const int32 TranslucencyLightingVolumeDim = GetTranslucencyLightingVolumeDim();
+
+						VertexShader->SetParameters(RHICmdList, VolumeBounds, FIntVector(TranslucencyLightingVolumeDim));
 						if(GeometryShader.IsValid())
 						{
 							GeometryShader->SetParameters(RHICmdList, VolumeBounds.MinZ);
@@ -1728,8 +1742,9 @@ void FDeferredShadingSceneRenderer::FilterTranslucentVolumeLighting(FRHICommandL
 
 		if (GUseTranslucencyVolumeBlur)
 		{
+			const int32 TranslucencyLightingVolumeDim = GetTranslucencyLightingVolumeDim();
 			SCOPED_DRAW_EVENTF(RHICmdList, FilterTranslucentVolume, TEXT("FilterTranslucentVolume %dx%dx%d Cascades:%d"),
-				GTranslucencyLightingVolumeDim, GTranslucencyLightingVolumeDim, GTranslucencyLightingVolumeDim, TVC_MAX);
+				TranslucencyLightingVolumeDim, TranslucencyLightingVolumeDim, TranslucencyLightingVolumeDim, TVC_MAX);
 
 			SCOPED_GPU_STAT(RHICmdList, TranslucentLighting);
 
@@ -1772,7 +1787,7 @@ void FDeferredShadingSceneRenderer::FilterTranslucentVolumeLighting(FRHICommandL
 				SetRenderTargets(RHICmdList, ARRAY_COUNT(RenderTargets), RenderTargets, FTextureRHIRef(), 0, NULL, true);
 				RHICmdList.ApplyCachedRenderTargets(GraphicsPSOInit);
 
-				const FVolumeBounds VolumeBounds(GTranslucencyLightingVolumeDim);
+				const FVolumeBounds VolumeBounds(TranslucencyLightingVolumeDim);
 				TShaderMapRef<FWriteToSliceVS> VertexShader(View.ShaderMap);
 				TOptionalShaderMapRef<FWriteToSliceGS> GeometryShader(View.ShaderMap);
 				TShaderMapRef<FFilterTranslucentVolumePS> PixelShader(View.ShaderMap);
@@ -1785,7 +1800,7 @@ void FDeferredShadingSceneRenderer::FilterTranslucentVolumeLighting(FRHICommandL
 
 				SetGraphicsPipelineState(RHICmdList, GraphicsPSOInit);
 
-				VertexShader->SetParameters(RHICmdList, VolumeBounds, FIntVector(GTranslucencyLightingVolumeDim));
+				VertexShader->SetParameters(RHICmdList, VolumeBounds, FIntVector(TranslucencyLightingVolumeDim));
 				if(GeometryShader.IsValid())
 				{
 					GeometryShader->SetParameters(RHICmdList, VolumeBounds.MinZ);
