Switching back to CopyToResolveTarget to make resolves explicitly clear.

diff --git a/Engine/Source/Runtime/Engine/Private/Particles/ParticleCurveTexture.cpp b/Engine/Source/Runtime/Engine/Private/Particles/ParticleCurveTexture.cpp
--- a/Engine/Source/Runtime/Engine/Private/Particles/ParticleCurveTexture.cpp
+++ b/Engine/Source/Runtime/Engine/Private/Particles/ParticleCurveTexture.cpp
@@ -196,7 +196,7 @@ static void InjectCurves(
 		bFirstCall = false;
 	}
 
-	FRHIRenderPassInfo RPInfo(CurveTextureTargetRHI, MakeRenderTargetActions(LoadAction, ERenderTargetStoreAction::EStore), CurveTextureRHI);
+	FRHIRenderPassInfo RPInfo(CurveTextureTargetRHI, MakeRenderTargetActions(LoadAction, ERenderTargetStoreAction::EStore));
 	TransitionRenderPassTargets(RHICmdList, RPInfo);
 	RHICmdList.BeginRenderPass(RPInfo, TEXT("InjectCurves"));
 	{
@@ -287,6 +287,7 @@ static void InjectCurves(
 		}
 	}
 	RHICmdList.EndRenderPass();
+	RHICmdList.CopyToResolveTarget(CurveTextureTargetRHI, CurveTextureRHI, FResolveParams());
 	RHICmdList.EndUpdateMultiFrameResource(CurveTextureTargetRHI);
 }
 
diff --git a/Engine/Source/Runtime/Engine/Private/Particles/ParticleGpuSimulation.cpp b/Engine/Source/Runtime/Engine/Private/Particles/ParticleGpuSimulation.cpp
--- a/Engine/Source/Runtime/Engine/Private/Particles/ParticleGpuSimulation.cpp
+++ b/Engine/Source/Runtime/Engine/Private/Particles/ParticleGpuSimulation.cpp
@@ -4737,11 +4737,11 @@ void FFXSystem::SimulateGPUParticles(
 
 			{
 				FRHIRenderPassInfo RPInfo(2, PreviousStateRenderTargets, ERenderTargetActions::Clear_Store);
-				RPInfo.ColorRenderTargets[0].ResolveTarget = PrevStateTextures.PositionTextureTargetRHI;
-				RPInfo.ColorRenderTargets[1].ResolveTarget = PrevStateTextures.VelocityTextureTargetRHI;
-
 				RHICmdList.BeginRenderPass(RPInfo, TEXT("GPUParticlesClearPreviousStateTextures"));
 				RHICmdList.EndRenderPass();
+
+				RHICmdList.CopyToResolveTarget(PreviousStateRenderTargets[0], PrevStateTextures.PositionTextureTargetRHI, FResolveParams());
+				RHICmdList.CopyToResolveTarget(PreviousStateRenderTargets[1], PrevStateTextures.VelocityTextureTargetRHI, FResolveParams());
 			}
 			
 			PrevStateTextures.bTexturesCleared = true;
@@ -4959,8 +4959,6 @@ void FFXSystem::SimulateGPUParticles(
 		RHICmdList.BeginUpdateMultiFrameResource(ParticleSimulationResources->SimulationAttributesTexture.TextureTargetRHI);
 
 		FRHIRenderPassInfo RPInfo(4, InjectRenderTargets, ERenderTargetActions::Load_Store);
-		RPInfo.ColorRenderTargets[2].ResolveTarget = ParticleSimulationResources->RenderAttributesTexture.TextureRHI;
-		RPInfo.ColorRenderTargets[3].ResolveTarget = ParticleSimulationResources->SimulationAttributesTexture.TextureRHI;
 		{
 			TransitionRenderPassTargets(RHICmdList, RPInfo);
 			RHICmdList.BeginRenderPass(RPInfo, TEXT("ParticleInjection"));
@@ -4977,6 +4975,18 @@ void FFXSystem::SimulateGPUParticles(
 			InjectNewParticles<false>(RHICmdList, GraphicsPSOInit, this->FeatureLevel, NewParticles);
 
 			RHICmdList.EndRenderPass();
+
+			// Resolve attributes textures. State textures are resolved later.
+			RHICmdList.CopyToResolveTarget(
+				ParticleSimulationResources->RenderAttributesTexture.TextureTargetRHI,
+				ParticleSimulationResources->RenderAttributesTexture.TextureRHI,
+				FResolveParams()
+			);
+			RHICmdList.CopyToResolveTarget(
+				ParticleSimulationResources->SimulationAttributesTexture.TextureTargetRHI,
+				ParticleSimulationResources->SimulationAttributesTexture.TextureRHI,
+				FResolveParams()
+			);
 		}
 
 		if (GNumAlternateFrameRenderingGroups > 1 && CVarGPUParticleAFRReinject.GetValueOnRenderThread() == 1)
diff --git a/Engine/Source/Runtime/Renderer/Private/AtmosphereRendering.cpp b/Engine/Source/Runtime/Renderer/Private/AtmosphereRendering.cpp
--- a/Engine/Source/Runtime/Renderer/Private/AtmosphereRendering.cpp
+++ b/Engine/Source/Runtime/Renderer/Private/AtmosphereRendering.cpp
@@ -1132,6 +1132,7 @@ void FAtmosphericFogSceneInfo::RenderAtmosphereShaders(FRHICommandList& RHICmdLi
 				DrawQuad(RHICmdList, ViewRect, *VertexShader);
 			}
 			RHICmdList.EndRenderPass();
+			RHICmdList.CopyToResolveTarget(DestRenderTarget.TargetableTexture, DestRenderTarget.ShaderResourceTexture, FResolveParams());
 		}
 		break;
 	case AP_Irradiance1:
@@ -1157,6 +1158,7 @@ void FAtmosphericFogSceneInfo::RenderAtmosphereShaders(FRHICommandList& RHICmdLi
 				DrawQuad(RHICmdList, ViewRect, *VertexShader);
 			}
 			RHICmdList.EndRenderPass();
+			RHICmdList.CopyToResolveTarget(DestRenderTarget.TargetableTexture, DestRenderTarget.ShaderResourceTexture, FResolveParams());
 		}
 		break;
 	case AP_Inscatter1:
@@ -1170,13 +1172,6 @@ void FAtmosphericFogSceneInfo::RenderAtmosphereShaders(FRHICommandList& RHICmdLi
 				};
 
 				FRHIRenderPassInfo RPInfo(2, RenderTargets, MakeRenderTargetActions(ERenderTargetLoadAction::ELoad, ERenderTargetStoreAction::EStore));
-
-				if (Atmosphere3DTextureIndex == Component->PrecomputeParams.InscatterAltitudeSampleNum - 1)
-				{
-					RPInfo.ColorRenderTargets[0].ResolveTarget = AtmosphereTextures->AtmosphereDeltaSR->GetRenderTargetItem().ShaderResourceTexture;
-					RPInfo.ColorRenderTargets[1].ResolveTarget = AtmosphereTextures->AtmosphereDeltaSM->GetRenderTargetItem().ShaderResourceTexture;
-				}
-
 				RHICmdList.BeginRenderPass(RPInfo, TEXT("AP_Inscatter"));
 				{
 					RHICmdList.ApplyCachedRenderTargets(GraphicsPSOInit);
@@ -1206,6 +1201,13 @@ void FAtmosphericFogSceneInfo::RenderAtmosphereShaders(FRHICommandList& RHICmdLi
 					DrawQuad(RHICmdList, ViewRect, *VertexShader);
 				}
 				RHICmdList.EndRenderPass();
+				if (Atmosphere3DTextureIndex == Component->PrecomputeParams.InscatterAltitudeSampleNum - 1)
+				{
+					RHICmdList.CopyToResolveTarget(AtmosphereTextures->AtmosphereDeltaSR->GetRenderTargetItem().TargetableTexture,
+						AtmosphereTextures->AtmosphereDeltaSR->GetRenderTargetItem().ShaderResourceTexture, FResolveParams());
+					RHICmdList.CopyToResolveTarget(AtmosphereTextures->AtmosphereDeltaSM->GetRenderTargetItem().TargetableTexture,
+						AtmosphereTextures->AtmosphereDeltaSM->GetRenderTargetItem().ShaderResourceTexture, FResolveParams());
+				}
 			}
 		}
 		break;
@@ -1214,9 +1216,10 @@ void FAtmosphericFogSceneInfo::RenderAtmosphereShaders(FRHICommandList& RHICmdLi
 			const FSceneRenderTargetItem& DestRenderTarget = AtmosphereTextures->AtmosphereIrradiance->GetRenderTargetItem();
 			ensure(DestRenderTarget.TargetableTexture->GetClearColor() == FLinearColor::Black);
 
-			FRHIRenderPassInfo RPInfo(DestRenderTarget.TargetableTexture, ERenderTargetActions::Clear_Store, DestRenderTarget.ShaderResourceTexture);
+			FRHIRenderPassInfo RPInfo(DestRenderTarget.TargetableTexture, ERenderTargetActions::Clear_Store);
 			RHICmdList.BeginRenderPass(RPInfo, TEXT("AP_ClearIrradiance"));
 			RHICmdList.EndRenderPass();
+			RHICmdList.CopyToResolveTarget(DestRenderTarget.TargetableTexture, DestRenderTarget.ShaderResourceTexture, FResolveParams());
 		}
 		break;
 	case AP_CopyInscatter1:
@@ -1226,12 +1229,6 @@ void FAtmosphericFogSceneInfo::RenderAtmosphereShaders(FRHICommandList& RHICmdLi
 				const FSceneRenderTargetItem& DestRenderTarget = AtmosphereTextures->AtmosphereInscatter->GetRenderTargetItem();
 
 				FRHIRenderPassInfo RPInfo(DestRenderTarget.TargetableTexture, ERenderTargetActions::Load_Store);
-
-				if (Atmosphere3DTextureIndex == Component->PrecomputeParams.InscatterAltitudeSampleNum - 1)
-				{
-					RPInfo.ColorRenderTargets[0].ResolveTarget = DestRenderTarget.ShaderResourceTexture;
-				}
-
 				RHICmdList.BeginRenderPass(RPInfo, TEXT("AP_CopyInscatter1"));
 				{
 					RHICmdList.ApplyCachedRenderTargets(GraphicsPSOInit);
@@ -1260,6 +1257,10 @@ void FAtmosphericFogSceneInfo::RenderAtmosphereShaders(FRHICommandList& RHICmdLi
 					DrawQuad(RHICmdList, ViewRect, *VertexShader);
 				}
 				RHICmdList.EndRenderPass();
+				if (Atmosphere3DTextureIndex == Component->PrecomputeParams.InscatterAltitudeSampleNum - 1)
+				{
+					RHICmdList.CopyToResolveTarget(DestRenderTarget.TargetableTexture, DestRenderTarget.ShaderResourceTexture, FResolveParams());
+				}
 			}
 		}
 		break;
@@ -1270,10 +1271,6 @@ void FAtmosphericFogSceneInfo::RenderAtmosphereShaders(FRHICommandList& RHICmdLi
 				const FSceneRenderTargetItem& DestRenderTarget = AtmosphereTextures->AtmosphereDeltaJ->GetRenderTargetItem();
 
 				FRHIRenderPassInfo RPInfo(DestRenderTarget.TargetableTexture, ERenderTargetActions::Load_Store);
-				if (Atmosphere3DTextureIndex == Component->PrecomputeParams.InscatterAltitudeSampleNum - 1)
-				{
-					RPInfo.ColorRenderTargets[0].ResolveTarget = DestRenderTarget.ShaderResourceTexture;
-				}
 				RHICmdList.BeginRenderPass(RPInfo, TEXT("InscatterS"));
 				{
 					RHICmdList.ApplyCachedRenderTargets(GraphicsPSOInit);
@@ -1302,13 +1299,17 @@ void FAtmosphericFogSceneInfo::RenderAtmosphereShaders(FRHICommandList& RHICmdLi
 					DrawQuad(RHICmdList, ViewRect, *VertexShader);
 				}
 				RHICmdList.EndRenderPass();
+				if (Atmosphere3DTextureIndex == Component->PrecomputeParams.InscatterAltitudeSampleNum - 1)
+				{
+					RHICmdList.CopyToResolveTarget(DestRenderTarget.TargetableTexture, DestRenderTarget.ShaderResourceTexture, FResolveParams());
+				}
 			}
 		}
 		break;
 	case AP_IrradianceN:
 		{
 			const FSceneRenderTargetItem& DestRenderTarget = AtmosphereTextures->AtmosphereDeltaE->GetRenderTargetItem();
-			FRHIRenderPassInfo RPInfo(DestRenderTarget.TargetableTexture, ERenderTargetActions::Load_Store, DestRenderTarget.ShaderResourceTexture);
+			FRHIRenderPassInfo RPInfo(DestRenderTarget.TargetableTexture, ERenderTargetActions::Load_Store);
 			RHICmdList.BeginRenderPass(RPInfo, TEXT("IrradianceN"));
 			{
 				RHICmdList.ApplyCachedRenderTargets(GraphicsPSOInit);
@@ -1328,6 +1329,7 @@ void FAtmosphericFogSceneInfo::RenderAtmosphereShaders(FRHICommandList& RHICmdLi
 				DrawQuad(RHICmdList, ViewRect, *VertexShader);
 			}
 			RHICmdList.EndRenderPass();
+			RHICmdList.CopyToResolveTarget(DestRenderTarget.TargetableTexture, DestRenderTarget.ShaderResourceTexture, FResolveParams());
 		}
 		break;
 
@@ -1338,10 +1340,6 @@ void FAtmosphericFogSceneInfo::RenderAtmosphereShaders(FRHICommandList& RHICmdLi
 				const FSceneRenderTargetItem& DestRenderTarget = AtmosphereTextures->AtmosphereDeltaSR->GetRenderTargetItem();
 
 				FRHIRenderPassInfo RPInfo(DestRenderTarget.TargetableTexture, ERenderTargetActions::Load_Store);
-				if (Atmosphere3DTextureIndex == Component->PrecomputeParams.InscatterAltitudeSampleNum - 1)
-				{
-					RPInfo.ColorRenderTargets[0].ResolveTarget = DestRenderTarget.ShaderResourceTexture;
-				}
 				RHICmdList.BeginRenderPass(RPInfo, TEXT("InscatterN"));
 				{
 					RHICmdList.ApplyCachedRenderTargets(GraphicsPSOInit);
@@ -1370,6 +1368,10 @@ void FAtmosphericFogSceneInfo::RenderAtmosphereShaders(FRHICommandList& RHICmdLi
 					DrawQuad(RHICmdList, ViewRect, *VertexShader);
 				}
 				RHICmdList.EndRenderPass();
+				if (Atmosphere3DTextureIndex == Component->PrecomputeParams.InscatterAltitudeSampleNum - 1)
+				{
+					RHICmdList.CopyToResolveTarget(DestRenderTarget.TargetableTexture, DestRenderTarget.ShaderResourceTexture, FResolveParams());
+				}
 			}
 		}
 		break;
@@ -1377,7 +1379,7 @@ void FAtmosphericFogSceneInfo::RenderAtmosphereShaders(FRHICommandList& RHICmdLi
 		{
 			const FSceneRenderTargetItem& DestRenderTarget = AtmosphereTextures->AtmosphereIrradiance->GetRenderTargetItem();
 
-			FRHIRenderPassInfo RPInfo(DestRenderTarget.TargetableTexture, ERenderTargetActions::Load_Store, DestRenderTarget.ShaderResourceTexture);
+			FRHIRenderPassInfo RPInfo(DestRenderTarget.TargetableTexture, ERenderTargetActions::Load_Store);
 			RHICmdList.BeginRenderPass(RPInfo, TEXT("CopyIrradiance"));
 			{
 				RHICmdList.ApplyCachedRenderTargets(GraphicsPSOInit);
@@ -1399,6 +1401,7 @@ void FAtmosphericFogSceneInfo::RenderAtmosphereShaders(FRHICommandList& RHICmdLi
 				DrawQuad(RHICmdList, ViewRect, *VertexShader);
 			}
 			RHICmdList.EndRenderPass();
+			RHICmdList.CopyToResolveTarget(DestRenderTarget.TargetableTexture, DestRenderTarget.ShaderResourceTexture, FResolveParams());
 
 			GraphicsPSOInit.BlendState = TStaticBlendState<>::GetRHI();
 		}
@@ -1411,11 +1414,6 @@ void FAtmosphericFogSceneInfo::RenderAtmosphereShaders(FRHICommandList& RHICmdLi
 				const FSceneRenderTargetItem& DestRenderTarget = AtmosphereTextures->AtmosphereInscatter->GetRenderTargetItem();
 				
 				FRHIRenderPassInfo RPInfo(DestRenderTarget.TargetableTexture, ERenderTargetActions::Load_Store);
-				if (Atmosphere3DTextureIndex == Component->PrecomputeParams.InscatterAltitudeSampleNum - 1)
-				{
-					RPInfo.ColorRenderTargets[0].ResolveTarget = DestRenderTarget.ShaderResourceTexture;
-				}
-
 				RHICmdList.BeginRenderPass(RPInfo, TEXT("CopyInscatterN"));
 				{
 					RHICmdList.ApplyCachedRenderTargets(GraphicsPSOInit);
@@ -1446,6 +1444,10 @@ void FAtmosphericFogSceneInfo::RenderAtmosphereShaders(FRHICommandList& RHICmdLi
 					DrawQuad(RHICmdList, ViewRect, *VertexShader);
 				}
 				RHICmdList.EndRenderPass();
+				if (Atmosphere3DTextureIndex == Component->PrecomputeParams.InscatterAltitudeSampleNum - 1)
+				{
+					RHICmdList.CopyToResolveTarget(DestRenderTarget.TargetableTexture, DestRenderTarget.ShaderResourceTexture, FResolveParams());
+				}
 			}
 
 			GraphicsPSOInit.BlendState = TStaticBlendState<>::GetRHI();
@@ -1459,10 +1461,6 @@ void FAtmosphericFogSceneInfo::RenderAtmosphereShaders(FRHICommandList& RHICmdLi
 				const FSceneRenderTargetItem& DestRenderTarget = AtmosphereTextures->AtmosphereDeltaSR->GetRenderTargetItem();
 				
 				FRHIRenderPassInfo RPInfo(DestRenderTarget.TargetableTexture, ERenderTargetActions::Load_Store);
-				if (Atmosphere3DTextureIndex == Component->PrecomputeParams.InscatterAltitudeSampleNum - 1)
-				{
-					RPInfo.ColorRenderTargets[0].ResolveTarget = DestRenderTarget.ShaderResourceTexture;
-				}
 				RHICmdList.BeginRenderPass(RPInfo, TEXT("CopyInscatterF"));
 				{
 					RHICmdList.ApplyCachedRenderTargets(GraphicsPSOInit);
@@ -1490,7 +1488,11 @@ void FAtmosphericFogSceneInfo::RenderAtmosphereShaders(FRHICommandList& RHICmdLi
 					PixelShader->SetParameters(RHICmdList, r, DhdH, Layer, AtmosphereTextures);
 					DrawQuad(RHICmdList, ViewRect, *VertexShader);
 				}
-				RHICmdList.EndRenderPass();				
+				RHICmdList.EndRenderPass();
+				if (Atmosphere3DTextureIndex == Component->PrecomputeParams.InscatterAltitudeSampleNum - 1)
+				{
+					RHICmdList.CopyToResolveTarget(DestRenderTarget.TargetableTexture, DestRenderTarget.ShaderResourceTexture, FResolveParams());
+				}
 			}
 		}
 		break;
@@ -1501,10 +1503,6 @@ void FAtmosphericFogSceneInfo::RenderAtmosphereShaders(FRHICommandList& RHICmdLi
 				const FSceneRenderTargetItem& DestRenderTarget = AtmosphereTextures->AtmosphereInscatter->GetRenderTargetItem();
 
 				FRHIRenderPassInfo RPInfo(DestRenderTarget.TargetableTexture, ERenderTargetActions::Load_Store);
-				if (Atmosphere3DTextureIndex == Component->PrecomputeParams.InscatterAltitudeSampleNum - 1)
-				{
-					RPInfo.ColorRenderTargets[0].ResolveTarget = DestRenderTarget.ShaderResourceTexture;
-				}
 				RHICmdList.BeginRenderPass(RPInfo, TEXT("CopyInscatterFBack"));
 				{
 					RHICmdList.ApplyCachedRenderTargets(GraphicsPSOInit);
@@ -1533,6 +1531,10 @@ void FAtmosphericFogSceneInfo::RenderAtmosphereShaders(FRHICommandList& RHICmdLi
 					DrawQuad(RHICmdList, ViewRect, *VertexShader);
 				}
 				RHICmdList.EndRenderPass();
+				if (Atmosphere3DTextureIndex == Component->PrecomputeParams.InscatterAltitudeSampleNum - 1)
+				{
+					RHICmdList.CopyToResolveTarget(DestRenderTarget.TargetableTexture, DestRenderTarget.ShaderResourceTexture, FResolveParams());
+				}
 			}
 		}
 		break;
diff --git a/Engine/Source/Runtime/Renderer/Private/AtmosphereTextures.cpp b/Engine/Source/Runtime/Renderer/Private/AtmosphereTextures.cpp
--- a/Engine/Source/Runtime/Renderer/Private/AtmosphereTextures.cpp
+++ b/Engine/Source/Runtime/Renderer/Private/AtmosphereTextures.cpp
@@ -25,9 +25,10 @@ void FAtmosphereTextures::InitDynamicRHI()
 		{
 			const FSceneRenderTargetItem& TransmittanceTarget = AtmosphereTransmittance->GetRenderTargetItem();
 
-			FRHIRenderPassInfo RPInfo(TransmittanceTarget.TargetableTexture, ERenderTargetActions::Clear_Store, TransmittanceTarget.ShaderResourceTexture);
+			FRHIRenderPassInfo RPInfo(TransmittanceTarget.TargetableTexture, ERenderTargetActions::Clear_Store);
 			RHICmdList.BeginRenderPass(RPInfo, TEXT("ClearTransmittance"));
 			RHICmdList.EndRenderPass();
+			RHICmdList.CopyToResolveTarget(TransmittanceTarget.TargetableTexture, TransmittanceTarget.ShaderResourceTexture, FResolveParams());
 		}
 
 		// Irradiance
@@ -37,9 +38,10 @@ void FAtmosphereTextures::InitDynamicRHI()
 		{
 			const FSceneRenderTargetItem& IrradianceTarget = AtmosphereIrradiance->GetRenderTargetItem();
 
-			FRHIRenderPassInfo RPInfo(IrradianceTarget.TargetableTexture, ERenderTargetActions::Clear_Store, IrradianceTarget.ShaderResourceTexture);
+			FRHIRenderPassInfo RPInfo(IrradianceTarget.TargetableTexture, ERenderTargetActions::Clear_Store);
 			RHICmdList.BeginRenderPass(RPInfo, TEXT("ClearIrradiance"));
 			RHICmdList.EndRenderPass();
+			RHICmdList.CopyToResolveTarget(IrradianceTarget.TargetableTexture, IrradianceTarget.ShaderResourceTexture, FResolveParams());
 		}
 		// DeltaE
 		GRenderTargetPool.FindFreeElement(RHICmdList, IrradianceDesc, AtmosphereDeltaE, TEXT("AtmosphereDeltaE"));
diff --git a/Engine/Source/Runtime/Renderer/Private/CompositionLighting/PostProcessAmbient.cpp b/Engine/Source/Runtime/Renderer/Private/CompositionLighting/PostProcessAmbient.cpp
--- a/Engine/Source/Runtime/Renderer/Private/CompositionLighting/PostProcessAmbient.cpp
+++ b/Engine/Source/Runtime/Renderer/Private/CompositionLighting/PostProcessAmbient.cpp
@@ -117,7 +117,7 @@ void FRCPassPostProcessAmbient::Process(FRenderingCompositePassContext& Context)
 	const FSceneRenderTargetItem& DestRenderTarget = FSceneRenderTargets::Get(Context.RHICmdList).GetSceneColor()->GetRenderTargetItem();
 
 	// Set the view family's render target/viewport.
-	FRHIRenderPassInfo RPInfo(DestRenderTarget.TargetableTexture, ERenderTargetActions::Load_Store, DestRenderTarget.ShaderResourceTexture);
+	FRHIRenderPassInfo RPInfo(DestRenderTarget.TargetableTexture, ERenderTargetActions::Load_Store);
 	TransitionRenderPassTargets(Context.RHICmdList, RPInfo);
 	Context.RHICmdList.BeginRenderPass(RPInfo, TEXT("PostProcessAmbient"));
 	{
@@ -142,6 +142,7 @@ void FRCPassPostProcessAmbient::Process(FRenderingCompositePassContext& Context)
 		Render(Context, GraphicsPSOInit);
 	}
 	Context.RHICmdList.EndRenderPass();
+	Context.RHICmdList.CopyToResolveTarget(DestRenderTarget.TargetableTexture, DestRenderTarget.ShaderResourceTexture, FResolveParams());
 }
 
 FPooledRenderTargetDesc FRCPassPostProcessAmbient::ComputeOutputDesc(EPassOutputId InPassOutputId) const
diff --git a/Engine/Source/Runtime/Renderer/Private/CompositionLighting/PostProcessAmbientOcclusion.cpp b/Engine/Source/Runtime/Renderer/Private/CompositionLighting/PostProcessAmbientOcclusion.cpp
--- a/Engine/Source/Runtime/Renderer/Private/CompositionLighting/PostProcessAmbientOcclusion.cpp
+++ b/Engine/Source/Runtime/Renderer/Private/CompositionLighting/PostProcessAmbientOcclusion.cpp
@@ -378,7 +378,7 @@ void FRCPassPostProcessAmbientOcclusionSetup::Process(FRenderingCompositePassCon
 	SCOPED_DRAW_EVENTF(Context.RHICmdList, AmbientOcclusionSetup, TEXT("AmbientOcclusionSetup %dx%d"), DestRect.Width(), DestRect.Height());
 
 	// Set the view family's render target/viewport.
-	FRHIRenderPassInfo RPInfo(DestRenderTarget.TargetableTexture, ERenderTargetActions::Load_Store, DestRenderTarget.ShaderResourceTexture);
+	FRHIRenderPassInfo RPInfo(DestRenderTarget.TargetableTexture, ERenderTargetActions::Load_Store);
 	Context.RHICmdList.BeginRenderPass(RPInfo, TEXT("AmbientOcclusionSetup"));
 	{
 
@@ -431,6 +431,7 @@ void FRCPassPostProcessAmbientOcclusionSetup::Process(FRenderingCompositePassCon
 			EDRF_UseTriangleOptimization);
 	}
 	Context.RHICmdList.EndRenderPass();
+	Context.RHICmdList.CopyToResolveTarget(DestRenderTarget.TargetableTexture, DestRenderTarget.ShaderResourceTexture, FResolveParams());
 }
 
 FPooledRenderTargetDesc FRCPassPostProcessAmbientOcclusionSetup::ComputeOutputDesc(EPassOutputId InPassOutputId) const
@@ -791,7 +792,7 @@ void FRCPassPostProcessAmbientOcclusion::ProcessPS(FRenderingCompositePassContex
 	{
 		// We'll use the depth/stencil buffer for read but it will not be modified.
 		RPInfo.DepthStencilRenderTarget.DepthStencilTarget = SceneDepthBuffer->TargetableTexture;
-		RPInfo.DepthStencilRenderTarget.Action = MakeDepthStencilTargetActions(ERenderTargetActions::Load_DontStore, ERenderTargetActions::Load_DontStore);
+		RPInfo.DepthStencilRenderTarget.Action = MakeDepthStencilTargetActions(ERenderTargetActions::Load_Store, ERenderTargetActions::Load_Store);
 		RPInfo.DepthStencilRenderTarget.ExclusiveDepthStencil = FExclusiveDepthStencil::DepthRead_StencilRead;
 	}
 
@@ -1048,7 +1049,7 @@ void FRCPassPostProcessBasePassAO::Process(FRenderingCompositePassContext& Conte
 	// Set the view family's render target/viewport.
 	Context.RHICmdList.TransitionResource(EResourceTransitionAccess::EWritable, DestRenderTarget.TargetableTexture);
 
-	FRHIRenderPassInfo RPInfo(DestRenderTarget.TargetableTexture, ERenderTargetActions::Load_Store, DestRenderTarget.ShaderResourceTexture);
+	FRHIRenderPassInfo RPInfo(DestRenderTarget.TargetableTexture, ERenderTargetActions::Load_Store);
 	Context.RHICmdList.BeginRenderPass(RPInfo, TEXT("BasePassAmbientOcclusion"));
 	{
 		Context.SetViewportAndCallRHI(View.ViewRect);
@@ -1088,6 +1089,7 @@ void FRCPassPostProcessBasePassAO::Process(FRenderingCompositePassContext& Conte
 			EDRF_UseTriangleOptimization);
 	}
 	Context.RHICmdList.EndRenderPass();
+	Context.RHICmdList.CopyToResolveTarget(DestRenderTarget.TargetableTexture, DestRenderTarget.ShaderResourceTexture, FResolveParams());
 }
 
 FPooledRenderTargetDesc FRCPassPostProcessBasePassAO::ComputeOutputDesc(EPassOutputId InPassOutputId) const
diff --git a/Engine/Source/Runtime/Renderer/Private/CompositionLighting/PostProcessLpvIndirect.cpp b/Engine/Source/Runtime/Renderer/Private/CompositionLighting/PostProcessLpvIndirect.cpp
--- a/Engine/Source/Runtime/Renderer/Private/CompositionLighting/PostProcessLpvIndirect.cpp
+++ b/Engine/Source/Runtime/Renderer/Private/CompositionLighting/PostProcessLpvIndirect.cpp
@@ -281,11 +281,6 @@ void FRCPassPostProcessLpvIndirect::Process(FRenderingCompositePassContext& Cont
 	}
 
 	FRHIRenderPassInfo RPInfo(NumRenderTargets, RenderTargets, ERenderTargetActions::Load_Store);
-	RPInfo.ColorRenderTargets[0].ResolveTarget = DestColorRenderTarget.ShaderResourceTexture;
-	if (bApplySeparateSpecularRT)
-	{
-		RPInfo.ColorRenderTargets[1].ResolveTarget = DestSpecularRenderTarget.ShaderResourceTexture;
-	}
 
 	Context.RHICmdList.BeginRenderPass(RPInfo, TEXT("LPVIndirect"));
 	{
@@ -350,6 +345,12 @@ void FRCPassPostProcessLpvIndirect::Process(FRenderingCompositePassContext& Cont
 		}
 	}
 	Context.RHICmdList.EndRenderPass();
+	Context.RHICmdList.CopyToResolveTarget(DestColorRenderTarget.TargetableTexture, DestColorRenderTarget.ShaderResourceTexture, FResolveParams());
+
+	if (bApplySeparateSpecularRT)
+	{
+		Context.RHICmdList.CopyToResolveTarget(DestSpecularRenderTarget.TargetableTexture, DestSpecularRenderTarget.ShaderResourceTexture, FResolveParams());
+	}
 
 	if ( LPVSettings.LPVDirectionalOcclusionIntensity > 0.0001f )
 	{
diff --git a/Engine/Source/Runtime/Renderer/Private/CompositionLighting/PostProcessPassThrough.cpp b/Engine/Source/Runtime/Renderer/Private/CompositionLighting/PostProcessPassThrough.cpp
--- a/Engine/Source/Runtime/Renderer/Private/CompositionLighting/PostProcessPassThrough.cpp
+++ b/Engine/Source/Runtime/Renderer/Private/CompositionLighting/PostProcessPassThrough.cpp
@@ -86,7 +86,7 @@ void FRCPassPostProcessPassThrough::Process(FRenderingCompositePassContext& Cont
 	checkf(DestRect.Size() == SrcRect.Size(), TEXT("Pass through should not be used as upscaling pass."));
 
 	// Set the view family's render target/viewport.
-	FRHIRenderPassInfo RPInfo(DestRenderTarget.TargetableTexture, ERenderTargetActions::Load_Store, DestRenderTarget.ShaderResourceTexture);
+	FRHIRenderPassInfo RPInfo(DestRenderTarget.TargetableTexture, ERenderTargetActions::Load_Store);
 	Context.RHICmdList.BeginRenderPass(RPInfo, TEXT("Passthrough"));
 	{
 		Context.SetViewportAndCallRHI(0, 0, 0.0f, DestSize.X, DestSize.Y, 1.0f);
@@ -137,6 +137,7 @@ void FRCPassPostProcessPassThrough::Process(FRenderingCompositePassContext& Cont
 		DrawCustom(Context);
 	}
 	Context.RHICmdList.EndRenderPass();
+	Context.RHICmdList.CopyToResolveTarget(DestRenderTarget.TargetableTexture, DestRenderTarget.ShaderResourceTexture, FResolveParams());
 }
 
 FPooledRenderTargetDesc FRCPassPostProcessPassThrough::ComputeOutputDesc(EPassOutputId InPassOutputId) const
diff --git a/Engine/Source/Runtime/Renderer/Private/GPUBenchmark.cpp b/Engine/Source/Runtime/Renderer/Private/GPUBenchmark.cpp
--- a/Engine/Source/Runtime/Renderer/Private/GPUBenchmark.cpp
+++ b/Engine/Source/Runtime/Renderer/Private/GPUBenchmark.cpp
@@ -579,13 +579,14 @@ void RendererGPUBenchmark(FRHICommandListImmediate& RHICmdList, FSynthBenchmarkR
 				// decide how much work we do in this pass
 				LocalWorkScale[Iteration] = (Iteration / 10.f + 1.f) * WorkScale;
 
-				FRHIRenderPassInfo RPInfo(RTItems[DestRTIndex]->GetRenderTargetItem().TargetableTexture, ERenderTargetActions::Load_Store, RTItems[DestRTIndex]->GetRenderTargetItem().ShaderResourceTexture);
+				FRHIRenderPassInfo RPInfo(RTItems[DestRTIndex]->GetRenderTargetItem().TargetableTexture, ERenderTargetActions::Load_Store);
 				TransitionRenderPassTargets(RHICmdList, RPInfo);
 				RHICmdList.BeginRenderPass(RPInfo, TEXT("GPUBenchmark"));
 				{
 					RunBenchmarkShader(RHICmdList, VertexBuffer, View, MethodId, RTItems[SrcRTIndex], LocalWorkScale[Iteration]);
 				}
 				RHICmdList.EndRenderPass();
+				RHICmdList.CopyToResolveTarget(RTItems[DestRTIndex]->GetRenderTargetItem().TargetableTexture, RTItems[DestRTIndex]->GetRenderTargetItem().ShaderResourceTexture, FResolveParams());
 
 				/*if(bGPUCPUSync)
 				{
diff --git a/Engine/Source/Runtime/Renderer/Private/PostProcess/DiaphragmDOFPasses.cpp b/Engine/Source/Runtime/Renderer/Private/PostProcess/DiaphragmDOFPasses.cpp
--- a/Engine/Source/Runtime/Renderer/Private/PostProcess/DiaphragmDOFPasses.cpp
+++ b/Engine/Source/Runtime/Renderer/Private/PostProcess/DiaphragmDOFPasses.cpp
@@ -1669,7 +1669,7 @@ void FRCPassDiaphragmDOFHybridScatter::Process(FRenderingCompositePassContext& C
 		EResourceTransitionPipeline::EComputeToGfx,
 		DestRenderTarget.UAV);
 
-	FRHIRenderPassInfo RPInfo(DestRenderTarget.TargetableTexture, ERenderTargetActions::Load_Store, DestRenderTarget.ShaderResourceTexture);
+	FRHIRenderPassInfo RPInfo(DestRenderTarget.TargetableTexture, ERenderTargetActions::Load_Store);
 	Context.RHICmdList.BeginRenderPass(RPInfo, TEXT("DOFHybridScatter"));
 	{
 		Context.SetViewportAndCallRHI(DestViewport, 0.0f, 1.0f);
@@ -1735,6 +1735,7 @@ void FRCPassDiaphragmDOFHybridScatter::Process(FRenderingCompositePassContext& C
 		}
 	}
 	Context.RHICmdList.EndRenderPass();
+	Context.RHICmdList.CopyToResolveTarget(DestRenderTarget.TargetableTexture, DestRenderTarget.ShaderResourceTexture, FResolveParams());
 
 	{
 		FVertexShaderRHIParamRef ShaderRHI = VertexShader->GetVertexShader();
diff --git a/Engine/Source/Runtime/Renderer/Private/PostProcess/PostProcessAA.cpp b/Engine/Source/Runtime/Renderer/Private/PostProcess/PostProcessAA.cpp
--- a/Engine/Source/Runtime/Renderer/Private/PostProcess/PostProcessAA.cpp
+++ b/Engine/Source/Runtime/Renderer/Private/PostProcess/PostProcessAA.cpp
@@ -262,7 +262,7 @@ void FRCPassPostProcessAA::Process(FRenderingCompositePassContext& Context)
 	SCOPED_DRAW_EVENTF(Context.RHICmdList, PostProcessFXAA, TEXT("PostProcessFXAA %dx%d"), DestRect.Width(), DestRect.Height());
 
 	// Set the view family's render target/viewport.
-	FRHIRenderPassInfo RPInfo(DestRenderTarget.TargetableTexture, ERenderTargetActions::Load_Store, DestRenderTarget.ShaderResourceTexture);
+	FRHIRenderPassInfo RPInfo(DestRenderTarget.TargetableTexture, ERenderTargetActions::Load_Store);
 	Context.RHICmdList.BeginRenderPass(RPInfo, TEXT("PostProcessAA"));
 	{
 		Context.SetViewportAndCallRHI(0, 0, 0.0f, DestSize.X, DestSize.Y, 1.0f);
@@ -292,6 +292,7 @@ void FRCPassPostProcessAA::Process(FRenderingCompositePassContext& Context)
 			EDRF_Default);
 	}
 	Context.RHICmdList.EndRenderPass();
+	Context.RHICmdList.CopyToResolveTarget(DestRenderTarget.TargetableTexture, DestRenderTarget.ShaderResourceTexture, FResolveParams());
 }
 
 FPooledRenderTargetDesc FRCPassPostProcessAA::ComputeOutputDesc(EPassOutputId InPassOutputId) const
diff --git a/Engine/Source/Runtime/Renderer/Private/PostProcess/PostProcessBloomSetup.cpp b/Engine/Source/Runtime/Renderer/Private/PostProcess/PostProcessBloomSetup.cpp
--- a/Engine/Source/Runtime/Renderer/Private/PostProcess/PostProcessBloomSetup.cpp
+++ b/Engine/Source/Runtime/Renderer/Private/PostProcess/PostProcessBloomSetup.cpp
@@ -324,6 +324,7 @@ void FRCPassPostProcessBloomSetup::Process(FRenderingCompositePassContext& Conte
 				EDRF_UseTriangleOptimization);
 		}
 		Context.RHICmdList.EndRenderPass();
+		Context.RHICmdList.CopyToResolveTarget(DestRenderTarget.TargetableTexture, DestRenderTarget.ShaderResourceTexture, FResolveParams());
 	}
 }
 
@@ -429,7 +430,7 @@ void FRCPassPostProcessVisualizeBloomSetup::Process(FRenderingCompositePassConte
 
 	const FSceneRenderTargetItem& DestRenderTarget = PassOutputs[0].RequestSurface(Context);
 	
-	FRHIRenderPassInfo RPInfo(DestRenderTarget.TargetableTexture, ERenderTargetActions::DontLoad_Store, DestRenderTarget.ShaderResourceTexture);
+	FRHIRenderPassInfo RPInfo(DestRenderTarget.TargetableTexture, ERenderTargetActions::DontLoad_Store);
 	Context.RHICmdList.BeginRenderPass(RPInfo, TEXT("VisualizeBloom"));
 	{
 		Context.SetViewportAndCallRHI(0, 0, 0.0f, DestRect.Width(), DestRect.Height(), 1.0f);
@@ -468,6 +469,7 @@ void FRCPassPostProcessVisualizeBloomSetup::Process(FRenderingCompositePassConte
 			EDRF_UseTriangleOptimization);
 	}
 	Context.RHICmdList.EndRenderPass();
+	Context.RHICmdList.CopyToResolveTarget(DestRenderTarget.TargetableTexture, DestRenderTarget.ShaderResourceTexture, FResolveParams());
 }
 
 
@@ -571,7 +573,7 @@ void FRCPassPostProcessVisualizeBloomOverlay::Process(FRenderingCompositePassCon
 
 	const FSceneRenderTargetItem& DestRenderTarget = PassOutputs[0].RequestSurface(Context);
 
-	FRHIRenderPassInfo RPInfo(DestRenderTarget.TargetableTexture, ERenderTargetActions::Load_Store, DestRenderTarget.ShaderResourceTexture);
+	FRHIRenderPassInfo RPInfo(DestRenderTarget.TargetableTexture, ERenderTargetActions::Load_Store);
 
 	Context.RHICmdList.BeginRenderPass(RPInfo, TEXT("VisualizeBloomOverlay"));
 	{
@@ -612,6 +614,7 @@ void FRCPassPostProcessVisualizeBloomOverlay::Process(FRenderingCompositePassCon
 			EDRF_UseTriangleOptimization);
 	}
 	Context.RHICmdList.EndRenderPass();
+	Context.RHICmdList.CopyToResolveTarget(DestRenderTarget.TargetableTexture, DestRenderTarget.ShaderResourceTexture, FResolveParams());
 }
 
 
diff --git a/Engine/Source/Runtime/Renderer/Private/PostProcess/PostProcessBokehDOF.cpp b/Engine/Source/Runtime/Renderer/Private/PostProcess/PostProcessBokehDOF.cpp
--- a/Engine/Source/Runtime/Renderer/Private/PostProcess/PostProcessBokehDOF.cpp
+++ b/Engine/Source/Runtime/Renderer/Private/PostProcess/PostProcessBokehDOF.cpp
@@ -518,7 +518,7 @@ void FRCPassPostProcessBokehDOFSetup::Process(FRenderingCompositePassContext& Co
 		WaitForInputPassComputeFences(Context.RHICmdList);
 
 		// Set the view family's render target/viewport.
-		FRHIRenderPassInfo RPInfo(DestRenderTarget.TargetableTexture, ERenderTargetActions::Load_Store, DestRenderTarget.ShaderResourceTexture);
+		FRHIRenderPassInfo RPInfo(DestRenderTarget.TargetableTexture, ERenderTargetActions::Load_Store);
 		Context.RHICmdList.BeginRenderPass(RPInfo, TEXT("BokehDOFSetup"));
 		{
 
@@ -570,6 +570,7 @@ void FRCPassPostProcessBokehDOFSetup::Process(FRenderingCompositePassContext& Co
 				EDRF_UseTriangleOptimization);
 		}
 		Context.RHICmdList.EndRenderPass();
+		Context.RHICmdList.CopyToResolveTarget(DestRenderTarget.TargetableTexture, DestRenderTarget.ShaderResourceTexture, FResolveParams());
 	}
 }
 
@@ -879,7 +880,7 @@ void FRCPassPostProcessBokehDOF::Process(FRenderingCompositePassContext& Context
 	WaitForInputPassComputeFences(Context.RHICmdList);
 
 	// Set the view family's render target/viewport.
-	FRHIRenderPassInfo RPInfo(DestRenderTarget.TargetableTexture, ERenderTargetActions::Load_Store, DestRenderTarget.ShaderResourceTexture);
+	FRHIRenderPassInfo RPInfo(DestRenderTarget.TargetableTexture, ERenderTargetActions::Load_Store);
 	Context.RHICmdList.BeginRenderPass(RPInfo, TEXT("BokehDOF"));
 	{
 
@@ -936,6 +937,7 @@ void FRCPassPostProcessBokehDOF::Process(FRenderingCompositePassContext& Context
 	}
 
 	Context.RHICmdList.EndRenderPass();
+	Context.RHICmdList.CopyToResolveTarget(DestRenderTarget.TargetableTexture, DestRenderTarget.ShaderResourceTexture, FResolveParams());
 }
 
 FPooledRenderTargetDesc FRCPassPostProcessBokehDOF::ComputeOutputDesc(EPassOutputId InPassOutputId) const
diff --git a/Engine/Source/Runtime/Renderer/Private/PostProcess/PostProcessBokehDOFRecombine.cpp b/Engine/Source/Runtime/Renderer/Private/PostProcess/PostProcessBokehDOFRecombine.cpp
--- a/Engine/Source/Runtime/Renderer/Private/PostProcess/PostProcessBokehDOFRecombine.cpp
+++ b/Engine/Source/Runtime/Renderer/Private/PostProcess/PostProcessBokehDOFRecombine.cpp
@@ -422,7 +422,7 @@ void FRCPassPostProcessBokehDOFRecombine::Process(FRenderingCompositePassContext
 		WaitForInputPassComputeFences(Context.RHICmdList);
 
 		// Set the view family's render target/viewport.
-		FRHIRenderPassInfo RPInfo(DestRenderTarget.TargetableTexture, ERenderTargetActions::Load_Store, DestRenderTarget.ShaderResourceTexture);
+		FRHIRenderPassInfo RPInfo(DestRenderTarget.TargetableTexture, ERenderTargetActions::Load_Store);
 		Context.RHICmdList.BeginRenderPass(RPInfo, TEXT("BokenDOFRecombine"));
 		{
 			if (View.StereoPass == eSSP_FULL)
@@ -460,6 +460,7 @@ void FRCPassPostProcessBokehDOFRecombine::Process(FRenderingCompositePassContext
 				EDRF_UseTriangleOptimization);
 		}
 		Context.RHICmdList.EndRenderPass();
+		Context.RHICmdList.CopyToResolveTarget(DestRenderTarget.TargetableTexture, DestRenderTarget.ShaderResourceTexture, FResolveParams());
 	}
 }
 
diff --git a/Engine/Source/Runtime/Renderer/Private/PostProcess/PostProcessBusyWait.cpp b/Engine/Source/Runtime/Renderer/Private/PostProcess/PostProcessBusyWait.cpp
--- a/Engine/Source/Runtime/Renderer/Private/PostProcess/PostProcessBusyWait.cpp
+++ b/Engine/Source/Runtime/Renderer/Private/PostProcess/PostProcessBusyWait.cpp
@@ -92,7 +92,7 @@ void FRCPassPostProcessBusyWait::Process(FRenderingCompositePassContext& Context
 	
 	const FSceneRenderTargetItem& DestRenderTarget = SceneContext.GetLightAttenuation()->GetRenderTargetItem();
 
-	FRHIRenderPassInfo RPInfo(DestRenderTarget.TargetableTexture, ERenderTargetActions::Load_Store, DestRenderTarget.ShaderResourceTexture);
+	FRHIRenderPassInfo RPInfo(DestRenderTarget.TargetableTexture, ERenderTargetActions::Load_Store);
 	Context.RHICmdList.BeginRenderPass(RPInfo, TEXT("PostProcessBusyWait"));
 	{
 		const FViewInfo& View = Context.View;
@@ -133,6 +133,7 @@ void FRCPassPostProcessBusyWait::Process(FRenderingCompositePassContext& Context
 			EDRF_UseTriangleOptimization);
 	}
 	Context.RHICmdList.EndRenderPass();
+	Context.RHICmdList.CopyToResolveTarget(DestRenderTarget.TargetableTexture, DestRenderTarget.ShaderResourceTexture, FResolveParams());
 
 	SceneContext.SetLightAttenuation(0);
 }
diff --git a/Engine/Source/Runtime/Renderer/Private/PostProcess/PostProcessCircleDOF.cpp b/Engine/Source/Runtime/Renderer/Private/PostProcess/PostProcessCircleDOF.cpp
--- a/Engine/Source/Runtime/Renderer/Private/PostProcess/PostProcessCircleDOF.cpp
+++ b/Engine/Source/Runtime/Renderer/Private/PostProcess/PostProcessCircleDOF.cpp
@@ -189,18 +189,7 @@ void FRCPassPostProcessCircleDOFSetup::Process(FRenderingCompositePassContext& C
 	};
 	uint32 NumRenderTargets = FPostProcessing::HasAlphaChannelSupport() ? 2 : 1;
 	
-	FTextureRHIParamRef ResolveTargets[2] =
-	{
-		DestRenderTarget0.ShaderResourceTexture,
-		nullptr,
-	};
-
-	if (DestRenderTarget1.TargetableTexture)
-	{
-		ResolveTargets[1] = DestRenderTarget1.ShaderResourceTexture;
-	}
-
-	FRHIRenderPassInfo RPInfo(NumRenderTargets, RenderTargets, ERenderTargetActions::Load_Store, ResolveTargets);
+	FRHIRenderPassInfo RPInfo(NumRenderTargets, RenderTargets, ERenderTargetActions::Load_Store);
 	Context.RHICmdList.BeginRenderPass(RPInfo, TEXT("CircleDOFSetup"));
 	{
 		const FViewInfo& View = Context.View;
@@ -279,6 +268,12 @@ void FRCPassPostProcessCircleDOFSetup::Process(FRenderingCompositePassContext& C
 			EDRF_UseTriangleOptimization);
 	}
 	Context.RHICmdList.EndRenderPass();
+	Context.RHICmdList.CopyToResolveTarget(DestRenderTarget0.TargetableTexture, DestRenderTarget0.ShaderResourceTexture, FResolveParams());
+
+	if (DestRenderTarget1.TargetableTexture)
+	{
+		Context.RHICmdList.CopyToResolveTarget(DestRenderTarget1.TargetableTexture, DestRenderTarget1.ShaderResourceTexture, FResolveParams());
+	}
 }
 
 FPooledRenderTargetDesc FRCPassPostProcessCircleDOFSetup::ComputeOutputDesc(EPassOutputId InPassOutputId) const
@@ -412,18 +407,7 @@ void FRCPassPostProcessCircleDOFDilate::Process(FRenderingCompositePassContext&
 		DestRenderTarget1.TargetableTexture
 	};
 
-	FTextureRHIParamRef ResolveTargets[2] =
-	{
-		DestRenderTarget0.ShaderResourceTexture,
-		nullptr,
-	};
-
-	if (DestRenderTarget1.TargetableTexture)
-	{
-		ResolveTargets[1] = DestRenderTarget1.ShaderResourceTexture;
-	}
-
-	FRHIRenderPassInfo RPInfo(NumRenderTargets, RenderTargets, ERenderTargetActions::Load_Store, ResolveTargets);
+	FRHIRenderPassInfo RPInfo(NumRenderTargets, RenderTargets, ERenderTargetActions::Load_Store);
 	Context.RHICmdList.BeginRenderPass(RPInfo, TEXT("CircleDOFDilate"));
 	{
 		const FViewInfo& View = Context.View;
@@ -502,6 +486,12 @@ void FRCPassPostProcessCircleDOFDilate::Process(FRenderingCompositePassContext&
 			EDRF_UseTriangleOptimization);
 	}
 	Context.RHICmdList.EndRenderPass();
+	Context.RHICmdList.CopyToResolveTarget(DestRenderTarget0.TargetableTexture, DestRenderTarget0.ShaderResourceTexture, FResolveParams());
+
+	if (DestRenderTarget1.TargetableTexture)
+	{
+		Context.RHICmdList.CopyToResolveTarget(DestRenderTarget1.TargetableTexture, DestRenderTarget1.ShaderResourceTexture, FResolveParams());
+	}
 }
 
 FPooledRenderTargetDesc FRCPassPostProcessCircleDOFDilate::ComputeOutputDesc(EPassOutputId InPassOutputId) const
@@ -704,18 +694,7 @@ void FRCPassPostProcessCircleDOF::Process(FRenderingCompositePassContext& Contex
 	};
 	uint32 NumRenderTargets = FPostProcessing::HasAlphaChannelSupport() ? 2 : 1;
 
-	FTextureRHIParamRef ResolveTargets[2] =
-	{
-		DestRenderTarget0.ShaderResourceTexture,
-		nullptr,
-	};
-
-	if (DestRenderTarget1.TargetableTexture)
-	{
-		ResolveTargets[1] = DestRenderTarget1.ShaderResourceTexture;
-	}
-
-	FRHIRenderPassInfo RPInfo(NumRenderTargets, RenderTargets, ERenderTargetActions::Load_Store, ResolveTargets);
+	FRHIRenderPassInfo RPInfo(NumRenderTargets, RenderTargets, ERenderTargetActions::Load_Store);
 	Context.RHICmdList.BeginRenderPass(RPInfo, TEXT("CircleDOFApply"));
 	{
 		const FViewInfo& View = Context.View;
@@ -776,6 +755,12 @@ void FRCPassPostProcessCircleDOF::Process(FRenderingCompositePassContext& Contex
 			EDRF_UseTriangleOptimization);
 	}
 	Context.RHICmdList.EndRenderPass();
+	Context.RHICmdList.CopyToResolveTarget(DestRenderTarget0.TargetableTexture, DestRenderTarget0.ShaderResourceTexture, FResolveParams());
+
+	if (DestRenderTarget1.TargetableTexture)
+	{
+		Context.RHICmdList.CopyToResolveTarget(DestRenderTarget1.TargetableTexture, DestRenderTarget1.ShaderResourceTexture, FResolveParams());
+	}
 }
 
 FPooledRenderTargetDesc FRCPassPostProcessCircleDOF::ComputeOutputDesc(EPassOutputId InPassOutputId) const
@@ -948,7 +933,7 @@ void FRCPassPostProcessCircleDOFRecombine::Process(FRenderingCompositePassContex
 		LoadStoreAction = ERenderTargetActions::DontLoad_Store;
 	}
 
-	FRHIRenderPassInfo RPInfo(DestRenderTarget.TargetableTexture, LoadStoreAction, DestRenderTarget.ShaderResourceTexture);
+	FRHIRenderPassInfo RPInfo(DestRenderTarget.TargetableTexture, LoadStoreAction);
 	Context.RHICmdList.BeginRenderPass(RPInfo, TEXT("CircleDOFRecombine"));
 	{
 		const FViewInfo& View = Context.View;
@@ -999,6 +984,7 @@ void FRCPassPostProcessCircleDOFRecombine::Process(FRenderingCompositePassContex
 			EDRF_UseTriangleOptimization);
 	}
 	Context.RHICmdList.EndRenderPass();
+	Context.RHICmdList.CopyToResolveTarget(DestRenderTarget.TargetableTexture, DestRenderTarget.ShaderResourceTexture, FResolveParams());
 }
 
 FPooledRenderTargetDesc FRCPassPostProcessCircleDOFRecombine::ComputeOutputDesc(EPassOutputId InPassOutputId) const
diff --git a/Engine/Source/Runtime/Renderer/Private/PostProcess/PostProcessCombineLUTs.cpp b/Engine/Source/Runtime/Renderer/Private/PostProcess/PostProcessCombineLUTs.cpp
--- a/Engine/Source/Runtime/Renderer/Private/PostProcess/PostProcessCombineLUTs.cpp
+++ b/Engine/Source/Runtime/Renderer/Private/PostProcess/PostProcessCombineLUTs.cpp
@@ -953,7 +953,7 @@ void FRCPassPostProcessCombineLUTs::Process(FRenderingCompositePassContext& Cont
 			LoadStoreAction = ERenderTargetActions::Clear_Store;
 		}
 
-		FRHIRenderPassInfo RPInfo(DestRenderTarget->TargetableTexture, LoadStoreAction, DestRenderTarget->ShaderResourceTexture);
+		FRHIRenderPassInfo RPInfo(DestRenderTarget->TargetableTexture, LoadStoreAction);
 		Context.RHICmdList.BeginRenderPass(RPInfo, TEXT("CombineLUTs"));
 		{
 			Context.SetViewportAndCallRHI(0, 0, 0.0f, DestSize.X, DestSize.Y, 1.0f);
@@ -985,6 +985,7 @@ void FRCPassPostProcessCombineLUTs::Process(FRenderingCompositePassContext& Cont
 			}
 		}
 		Context.RHICmdList.EndRenderPass();
+		Context.RHICmdList.CopyToResolveTarget(DestRenderTarget->TargetableTexture, DestRenderTarget->ShaderResourceTexture, FResolveParams());
 	}
 
 	Context.View.SetValidTonemappingLUT();
diff --git a/Engine/Source/Runtime/Renderer/Private/PostProcess/PostProcessCompositeEditorPrimitives.cpp b/Engine/Source/Runtime/Renderer/Private/PostProcess/PostProcessCompositeEditorPrimitives.cpp
--- a/Engine/Source/Runtime/Renderer/Private/PostProcess/PostProcessCompositeEditorPrimitives.cpp
+++ b/Engine/Source/Runtime/Renderer/Private/PostProcess/PostProcessCompositeEditorPrimitives.cpp
@@ -570,7 +570,7 @@ void FRCPassPostProcessCompositeEditorPrimitives::Process(FRenderingCompositePas
 		FIntRect DestRect = Context.GetSceneColorDestRect(DestRenderTarget);
 
 		// Set the view family's render target/viewport.
-		FRHIRenderPassInfo RPInfo(DestRenderTargetSurface, ERenderTargetActions::Load_Store, DestRenderTarget.ShaderResourceTexture);
+		FRHIRenderPassInfo RPInfo(DestRenderTargetSurface, ERenderTargetActions::Load_Store);
 		Context.RHICmdList.BeginRenderPass(RPInfo, TEXT("ComposeEditorPrimitives"));
 		{
 			Context.SetViewportAndCallRHI(DestRect);
@@ -620,6 +620,7 @@ void FRCPassPostProcessCompositeEditorPrimitives::Process(FRenderingCompositePas
 				EDRF_UseTriangleOptimization);
 		}
 		Context.RHICmdList.EndRenderPass();
+		Context.RHICmdList.CopyToResolveTarget(DestRenderTargetSurface, DestRenderTarget.ShaderResourceTexture, FResolveParams());
 	}
 
 	// Clean up targets
diff --git a/Engine/Source/Runtime/Renderer/Private/PostProcess/PostProcessDOF.cpp b/Engine/Source/Runtime/Renderer/Private/PostProcess/PostProcessDOF.cpp
--- a/Engine/Source/Runtime/Renderer/Private/PostProcess/PostProcessDOF.cpp
+++ b/Engine/Source/Runtime/Renderer/Private/PostProcess/PostProcessDOF.cpp
@@ -177,16 +177,7 @@ void FRCPassPostProcessDOFSetup::Process(FRenderingCompositePassContext& Context
 		DestRenderTarget0.TargetableTexture,
 		DestRenderTarget1.TargetableTexture
 	};
-
-	FTextureRHIParamRef ResolveTargets[2] =
-	{
-		DestRenderTarget0.ShaderResourceTexture,
-		nullptr
-	};
-	if (DestRenderTarget1.TargetableTexture)
-	{
-		ResolveTargets[1] = DestRenderTarget1.ShaderResourceTexture;
-	}
+	
 	ERenderTargetActions LoadStoreAction = ERenderTargetActions::Load_Store;
 
 	//@todo Ronin find a way to use the same codepath for all platforms.
@@ -196,7 +187,7 @@ void FRCPassPostProcessDOFSetup::Process(FRenderingCompositePassContext& Context
 		LoadStoreAction = ERenderTargetActions::Clear_Store;
 	}
 
-	FRHIRenderPassInfo RPInfo(NumRenderTargets, RenderTargets, LoadStoreAction, ResolveTargets);
+	FRHIRenderPassInfo RPInfo(NumRenderTargets, RenderTargets, LoadStoreAction);
 	Context.RHICmdList.BeginRenderPass(RPInfo, TEXT("DOFSetup"));
 	{
 		if (View.StereoPass == eSSP_FULL)
@@ -263,6 +254,12 @@ void FRCPassPostProcessDOFSetup::Process(FRenderingCompositePassContext& Context
 			EDRF_UseTriangleOptimization);
 	}
 	Context.RHICmdList.EndRenderPass();
+	Context.RHICmdList.CopyToResolveTarget(DestRenderTarget0.TargetableTexture, DestRenderTarget0.ShaderResourceTexture, FResolveParams());
+
+	if (DestRenderTarget1.TargetableTexture)
+	{
+		Context.RHICmdList.CopyToResolveTarget(DestRenderTarget1.TargetableTexture, DestRenderTarget1.ShaderResourceTexture, FResolveParams());
+	}
 
 	// #todo-rco: needed to avoid multiple resolves clearing the RT with VK.
 	// #todo mattc this is probably busted now since the resolves previously happened after this SetRenderTarget.
@@ -462,7 +459,7 @@ void FRCPassPostProcessDOFRecombine::Process(FRenderingCompositePassContext& Con
 		LoadStoreAction = ERenderTargetActions::Load_Store;
 	}
 
-	FRHIRenderPassInfo RPInfo(DestRenderTarget.TargetableTexture, LoadStoreAction, DestRenderTarget.ShaderResourceTexture);
+	FRHIRenderPassInfo RPInfo(DestRenderTarget.TargetableTexture, LoadStoreAction);
 	Context.RHICmdList.BeginRenderPass(RPInfo, TEXT("DOFRecombine"));
 	{
 		if (View.StereoPass == eSSP_FULL)
@@ -508,6 +505,7 @@ void FRCPassPostProcessDOFRecombine::Process(FRenderingCompositePassContext& Con
 			EDRF_UseTriangleOptimization);
 	}
 	Context.RHICmdList.EndRenderPass();
+	Context.RHICmdList.CopyToResolveTarget(DestRenderTarget.TargetableTexture, DestRenderTarget.ShaderResourceTexture, FResolveParams());
 }
 
 FPooledRenderTargetDesc FRCPassPostProcessDOFRecombine::ComputeOutputDesc(EPassOutputId InPassOutputId) const
diff --git a/Engine/Source/Runtime/Renderer/Private/PostProcess/PostProcessDownsample.cpp b/Engine/Source/Runtime/Renderer/Private/PostProcess/PostProcessDownsample.cpp
--- a/Engine/Source/Runtime/Renderer/Private/PostProcess/PostProcessDownsample.cpp
+++ b/Engine/Source/Runtime/Renderer/Private/PostProcess/PostProcessDownsample.cpp
@@ -367,7 +367,7 @@ void FRCPassPostProcessDownsample::Process(FRenderingCompositePassContext& Conte
 	else
 	{
 		// #todo-renderpasses only clear dest rectangle if it's been computed
-		FRHIRenderPassInfo RPInfo(DestRenderTarget.TargetableTexture, ERenderTargetActions::Clear_Store, DestRenderTarget.ShaderResourceTexture);
+		FRHIRenderPassInfo RPInfo(DestRenderTarget.TargetableTexture, ERenderTargetActions::Clear_Store);
 		Context.RHICmdList.BeginRenderPass(RPInfo, TEXT("PostProcessDownsample"));
 		{
 			Context.SetViewportAndCallRHI(0, 0, 0.0f, DestSize.X, DestSize.Y, 1.0f);
@@ -429,6 +429,7 @@ void FRCPassPostProcessDownsample::Process(FRenderingCompositePassContext& Conte
 				EDRF_UseTriangleOptimization);
 		}
 		Context.RHICmdList.EndRenderPass();
+		Context.RHICmdList.CopyToResolveTarget(DestRenderTarget.TargetableTexture, DestRenderTarget.ShaderResourceTexture, FResolveParams());
 	}
 }
 
diff --git a/Engine/Source/Runtime/Renderer/Private/PostProcess/PostProcessEyeAdaptation.cpp b/Engine/Source/Runtime/Renderer/Private/PostProcess/PostProcessEyeAdaptation.cpp
--- a/Engine/Source/Runtime/Renderer/Private/PostProcess/PostProcessEyeAdaptation.cpp
+++ b/Engine/Source/Runtime/Renderer/Private/PostProcess/PostProcessEyeAdaptation.cpp
@@ -348,7 +348,7 @@ void FRCPassPostProcessEyeAdaptation::Process(FRenderingCompositePassContext& Co
 		// we render to our own output render target, not the intermediate one created by the compositing system
 		// Set the view family's render target/viewport.
 
-		FRHIRenderPassInfo RPInfo(DestRenderTarget.TargetableTexture, ERenderTargetActions::Load_Store, DestRenderTarget.ShaderResourceTexture);
+		FRHIRenderPassInfo RPInfo(DestRenderTarget.TargetableTexture, ERenderTargetActions::Load_Store);
 		TransitionRenderPassTargets(Context.RHICmdList, RPInfo);
 		Context.RHICmdList.BeginRenderPass(RPInfo, TEXT("EyeAdaptation"));
 		{
@@ -385,6 +385,7 @@ void FRCPassPostProcessEyeAdaptation::Process(FRenderingCompositePassContext& Co
 				EDRF_UseTriangleOptimization);
 		}
 		Context.RHICmdList.EndRenderPass();
+		Context.RHICmdList.CopyToResolveTarget(DestRenderTarget.TargetableTexture, DestRenderTarget.ShaderResourceTexture, FResolveParams());
 		// Inform MultiGPU systems that we've finished updating this texture for this frame
 		Context.RHICmdList.EndUpdateMultiFrameResource(DestRenderTarget.ShaderResourceTexture);
 	}
@@ -542,7 +543,7 @@ void FRCPassPostProcessBasicEyeAdaptationSetUp::Process(FRenderingCompositePassC
 
 	ERenderTargetLoadAction LoadAction = Context.GetLoadActionForRenderTarget(DestRenderTarget);
 
-	FRHIRenderPassInfo RPInfo(DestRenderTarget.TargetableTexture, MakeRenderTargetActions(LoadAction, ERenderTargetStoreAction::EStore), DestRenderTarget.ShaderResourceTexture);
+	FRHIRenderPassInfo RPInfo(DestRenderTarget.TargetableTexture, MakeRenderTargetActions(LoadAction, ERenderTargetStoreAction::EStore));
 	Context.RHICmdList.BeginRenderPass(RPInfo, TEXT("PostProcessBasicEyeAdaptationSetUp"));
 	{
 		const FViewInfo& View = Context.View;
@@ -593,6 +594,7 @@ void FRCPassPostProcessBasicEyeAdaptationSetUp::Process(FRenderingCompositePassC
 			EDRF_UseTriangleOptimization);
 	}
 	Context.RHICmdList.EndRenderPass();
+	Context.RHICmdList.CopyToResolveTarget(DestRenderTarget.TargetableTexture, DestRenderTarget.ShaderResourceTexture, FResolveParams());
 }
 
 FPooledRenderTargetDesc FRCPassPostProcessBasicEyeAdaptationSetUp::ComputeOutputDesc(EPassOutputId InPassOutputId) const
@@ -732,7 +734,7 @@ void FRCPassPostProcessBasicEyeAdaptation::Process(FRenderingCompositePassContex
 	// we render to our own output render target, not the intermediate one created by the compositing system
 	// Set the view family's render target/viewport.
 
-	FRHIRenderPassInfo RPInfo(DestRenderTarget.TargetableTexture, ERenderTargetActions::Load_Store, DestRenderTarget.ShaderResourceTexture);
+	FRHIRenderPassInfo RPInfo(DestRenderTarget.TargetableTexture, ERenderTargetActions::Load_Store);
 	TransitionRenderPassTargets(Context.RHICmdList, RPInfo);
 	Context.RHICmdList.BeginRenderPass(RPInfo, TEXT("BasicEyeAdaptation"));
 	{
@@ -771,6 +773,7 @@ void FRCPassPostProcessBasicEyeAdaptation::Process(FRenderingCompositePassContex
 			EDRF_UseTriangleOptimization);
 	}
 	Context.RHICmdList.EndRenderPass();
+	Context.RHICmdList.CopyToResolveTarget(DestRenderTarget.TargetableTexture, DestRenderTarget.ShaderResourceTexture, FResolveParams());
 
 	// Inform MultiGPU systems that we've finished with this texture for this frame
 	Context.RHICmdList.EndUpdateMultiFrameResource(DestRenderTarget.ShaderResourceTexture);
diff --git a/Engine/Source/Runtime/Renderer/Private/PostProcess/PostProcessHMD.cpp b/Engine/Source/Runtime/Renderer/Private/PostProcess/PostProcessHMD.cpp
--- a/Engine/Source/Runtime/Renderer/Private/PostProcess/PostProcessHMD.cpp
+++ b/Engine/Source/Runtime/Renderer/Private/PostProcess/PostProcessHMD.cpp
@@ -172,7 +172,7 @@ void FRCPassPostProcessHMD::Process(FRenderingCompositePassContext& Context)
 		LoadAction = ERenderTargetLoadAction::ENoAction;
 	}
 	
-	FRHIRenderPassInfo RPInfo(DestRenderTarget.TargetableTexture, MakeRenderTargetActions(LoadAction, ERenderTargetStoreAction::EStore), DestRenderTarget.ShaderResourceTexture);
+	FRHIRenderPassInfo RPInfo(DestRenderTarget.TargetableTexture, MakeRenderTargetActions(LoadAction, ERenderTargetStoreAction::EStore));
 	Context.RHICmdList.BeginRenderPass(RPInfo, TEXT("PostProcessHMD"));
 	{
 		const FViewInfo& View = Context.View;
@@ -217,6 +217,7 @@ void FRCPassPostProcessHMD::Process(FRenderingCompositePassContext& Context)
 		GEngine->XRSystem->GetHMDDevice()->DrawDistortionMesh_RenderThread(Context, SrcSize);
 	}
 	Context.RHICmdList.EndRenderPass();
+	Context.RHICmdList.CopyToResolveTarget(DestRenderTarget.TargetableTexture, DestRenderTarget.ShaderResourceTexture, FResolveParams());
 }
 
 FPooledRenderTargetDesc FRCPassPostProcessHMD::ComputeOutputDesc(EPassOutputId InPassOutputId) const
diff --git a/Engine/Source/Runtime/Renderer/Private/PostProcess/PostProcessHierarchical.cpp b/Engine/Source/Runtime/Renderer/Private/PostProcess/PostProcessHierarchical.cpp
--- a/Engine/Source/Runtime/Renderer/Private/PostProcess/PostProcessHierarchical.cpp
+++ b/Engine/Source/Runtime/Renderer/Private/PostProcess/PostProcessHierarchical.cpp
@@ -158,7 +158,7 @@ void FRCPassPostProcessBuildHCB::Process(FRenderingCompositePassContext& Context
 		GraphicsPSOInit.BoundShaderState.PixelShaderRHI = GETSAFERHISHADER_PIXEL(*PixelShader);
 		GraphicsPSOInit.PrimitiveType = PT_TriangleList;
 
-		FRHIRenderPassInfo RPInfo(HCBRenderTarget.TargetableTexture, ERenderTargetActions::Load_Store, HCBRenderTarget.ShaderResourceTexture);
+		FRHIRenderPassInfo RPInfo(HCBRenderTarget.TargetableTexture, ERenderTargetActions::Load_Store);
 		RPInfo.ResolveParameters = FResolveParams(FResolveRect(), CubeFace_PosX, 0);
 		RHICmdList.BeginRenderPass(RPInfo, TEXT("BuildHCB_0"));
 		{
@@ -208,6 +208,7 @@ void FRCPassPostProcessBuildHCB::Process(FRenderingCompositePassContext& Context
 #endif
 		}
 		Context.RHICmdList.EndRenderPass();
+		Context.RHICmdList.CopyToResolveTarget(HCBRenderTarget.TargetableTexture, HCBRenderTarget.ShaderResourceTexture, FResolveRect());
 	}
 		
 	FIntPoint SrcSize = HCBSize;
@@ -219,7 +220,7 @@ void FRCPassPostProcessBuildHCB::Process(FRenderingCompositePassContext& Context
 		TCHAR RenderPassLabel[13];
 		FCString::Snprintf(RenderPassLabel, 13, TEXT("BuildHCB_%u"), MipIndex);
 
-		FRHIRenderPassInfo RPInfo(HCBRenderTarget.TargetableTexture, ERenderTargetActions::Load_Store, HCBRenderTarget.ShaderResourceTexture);
+		FRHIRenderPassInfo RPInfo(HCBRenderTarget.TargetableTexture, ERenderTargetActions::Load_Store);
 		RPInfo.ColorRenderTargets[0].MipIndex = MipIndex;
 		RPInfo.ResolveParameters = FResolveParams(FResolveRect(), CubeFace_PosX, MipIndex);
 		RHICmdList.BeginRenderPass(RPInfo, RenderPassLabel);
@@ -279,6 +280,7 @@ void FRCPassPostProcessBuildHCB::Process(FRenderingCompositePassContext& Context
 #endif
 		}
 		Context.RHICmdList.EndRenderPass();
+		Context.RHICmdList.CopyToResolveTarget(HCBRenderTarget.TargetableTexture, HCBRenderTarget.ShaderResourceTexture, FResolveParams());
 
 		SrcSize /= 2;
 		DstSize /= 2;
diff --git a/Engine/Source/Runtime/Renderer/Private/PostProcess/PostProcessHistogramReduce.cpp b/Engine/Source/Runtime/Renderer/Private/PostProcess/PostProcessHistogramReduce.cpp
--- a/Engine/Source/Runtime/Renderer/Private/PostProcess/PostProcessHistogramReduce.cpp
+++ b/Engine/Source/Runtime/Renderer/Private/PostProcess/PostProcessHistogramReduce.cpp
@@ -107,7 +107,7 @@ void FRCPassPostProcessHistogramReduce::Process(FRenderingCompositePassContext&
 	const FSceneRenderTargetItem& DestRenderTarget = PassOutputs[0].RequestSurface(Context);
 
 	// Set the view family's render target/viewport.
-	FRHIRenderPassInfo RPInfo(DestRenderTarget.TargetableTexture, ERenderTargetActions::Load_Store, DestRenderTarget.ShaderResourceTexture);
+	FRHIRenderPassInfo RPInfo(DestRenderTarget.TargetableTexture, ERenderTargetActions::Load_Store);
 	Context.RHICmdList.BeginRenderPass(RPInfo, TEXT("HistogramReduce"));
 	{
 		FGraphicsPipelineStateInitializer GraphicsPSOInit;
@@ -145,6 +145,7 @@ void FRCPassPostProcessHistogramReduce::Process(FRenderingCompositePassContext&
 			EDRF_UseTriangleOptimization);
 	}
 	Context.RHICmdList.EndRenderPass();
+	Context.RHICmdList.CopyToResolveTarget(DestRenderTarget.TargetableTexture, DestRenderTarget.ShaderResourceTexture, FResolveParams());
 }
 
 uint32 FRCPassPostProcessHistogramReduce::ComputeLoopSize(FIntPoint PixelExtent)
diff --git a/Engine/Source/Runtime/Renderer/Private/PostProcess/PostProcessLensBlur.cpp b/Engine/Source/Runtime/Renderer/Private/PostProcess/PostProcessLensBlur.cpp
--- a/Engine/Source/Runtime/Renderer/Private/PostProcess/PostProcessLensBlur.cpp
+++ b/Engine/Source/Runtime/Renderer/Private/PostProcess/PostProcessLensBlur.cpp
@@ -188,7 +188,7 @@ void FRCPassPostProcessLensBlur::Process(FRenderingCompositePassContext& Context
 	const FSceneRenderTargetItem& DestRenderTarget = PassOutputs[0].RequestSurface(Context);
 
 	// Set the view family's render target/viewport.
-	FRHIRenderPassInfo RPInfo(DestRenderTarget.TargetableTexture, ERenderTargetActions::Clear_Store, DestRenderTarget.ShaderResourceTexture);
+	FRHIRenderPassInfo RPInfo(DestRenderTarget.TargetableTexture, ERenderTargetActions::Clear_Store);
 	Context.RHICmdList.BeginRenderPass(RPInfo, TEXT("LensBlur"));
 	{
 		Context.SetViewportAndCallRHI(ViewRect);
@@ -228,6 +228,7 @@ void FRCPassPostProcessLensBlur::Process(FRenderingCompositePassContext& Context
 		Context.RHICmdList.DrawPrimitive(0, 2, FMath::DivideAndRoundUp(TileCount.X * TileCount.Y, QuadsPerInstance));
 	}
 	Context.RHICmdList.EndRenderPass();
+	Context.RHICmdList.CopyToResolveTarget(DestRenderTarget.TargetableTexture, DestRenderTarget.ShaderResourceTexture, FResolveParams());
 }
 
 
diff --git a/Engine/Source/Runtime/Renderer/Private/PostProcess/PostProcessLensFlares.cpp b/Engine/Source/Runtime/Renderer/Private/PostProcess/PostProcessLensFlares.cpp
--- a/Engine/Source/Runtime/Renderer/Private/PostProcess/PostProcessLensFlares.cpp
+++ b/Engine/Source/Runtime/Renderer/Private/PostProcess/PostProcessLensFlares.cpp
@@ -149,7 +149,7 @@ void FRCPassPostProcessLensFlares::Process(FRenderingCompositePassContext& Conte
 
 	const FSceneRenderTargetItem& DestRenderTarget = PassOutputs[0].RequestSurface(Context);
 
-	FRHIRenderPassInfo RPInfo(DestRenderTarget.TargetableTexture, ERenderTargetActions::DontLoad_Store, DestRenderTarget.ShaderResourceTexture);
+	FRHIRenderPassInfo RPInfo(DestRenderTarget.TargetableTexture, ERenderTargetActions::DontLoad_Store);
 	Context.RHICmdList.BeginRenderPass(RPInfo, TEXT("PassPostProcessLensFlares"));
 	{
 		const FViewInfo& View = Context.View;
@@ -292,6 +292,7 @@ void FRCPassPostProcessLensFlares::Process(FRenderingCompositePassContext& Conte
 		}
 	}
 	Context.RHICmdList.EndRenderPass();
+	Context.RHICmdList.CopyToResolveTarget(DestRenderTarget.TargetableTexture, DestRenderTarget.ShaderResourceTexture, FResolveParams());
 }
 
 FPooledRenderTargetDesc FRCPassPostProcessLensFlares::ComputeOutputDesc(EPassOutputId InPassOutputId) const
diff --git a/Engine/Source/Runtime/Renderer/Private/PostProcess/PostProcessMaterial.cpp b/Engine/Source/Runtime/Renderer/Private/PostProcess/PostProcessMaterial.cpp
--- a/Engine/Source/Runtime/Renderer/Private/PostProcess/PostProcessMaterial.cpp
+++ b/Engine/Source/Runtime/Renderer/Private/PostProcess/PostProcessMaterial.cpp
@@ -230,7 +230,7 @@ void FRCPassPostProcessMaterial::Process(FRenderingCompositePassContext& Context
 
 	ERenderTargetLoadAction LoadAction = Context.GetLoadActionForRenderTarget(DestRenderTarget);
 	
-	FRHIRenderPassInfo RPInfo(DestRenderTarget.TargetableTexture, MakeRenderTargetActions(LoadAction, ERenderTargetStoreAction::EStore), DestRenderTarget.ShaderResourceTexture);
+	FRHIRenderPassInfo RPInfo(DestRenderTarget.TargetableTexture, MakeRenderTargetActions(LoadAction, ERenderTargetStoreAction::EStore));
 	Context.RHICmdList.BeginRenderPass(RPInfo, TEXT("PostProcessMaterial"));
 	{
 		Context.SetViewportAndCallRHI(DestRect);
@@ -308,6 +308,7 @@ void FRCPassPostProcessMaterial::Process(FRenderingCompositePassContext& Context
 			EDRF_UseTriangleOptimization);
 	}
 	Context.RHICmdList.EndRenderPass();
+	Context.RHICmdList.CopyToResolveTarget(DestRenderTarget.TargetableTexture, DestRenderTarget.ShaderResourceTexture, FResolveParams());
 
 	if(Material->NeedsGBuffer())
 	{
diff --git a/Engine/Source/Runtime/Renderer/Private/PostProcess/PostProcessMobile.cpp b/Engine/Source/Runtime/Renderer/Private/PostProcess/PostProcessMobile.cpp
--- a/Engine/Source/Runtime/Renderer/Private/PostProcess/PostProcessMobile.cpp
+++ b/Engine/Source/Runtime/Renderer/Private/PostProcess/PostProcessMobile.cpp
@@ -238,7 +238,7 @@ void FRCPassPostProcessBloomSetupES2::Process(FRenderingCompositePassContext& Co
 		LoadAction = ERenderTargetLoadAction::ENoAction;
 	}
 
-	FRHIRenderPassInfo RPInfo(DestRenderTarget.TargetableTexture, MakeRenderTargetActions(LoadAction, ERenderTargetStoreAction::EStore), DestRenderTarget.ShaderResourceTexture);
+	FRHIRenderPassInfo RPInfo(DestRenderTarget.TargetableTexture, MakeRenderTargetActions(LoadAction, ERenderTargetStoreAction::EStore));
 	
 	Context.RHICmdList.BeginRenderPass(RPInfo, TEXT("PostProcessBloomSetupES2"));
 	{
@@ -265,6 +265,7 @@ void FRCPassPostProcessBloomSetupES2::Process(FRenderingCompositePassContext& Co
 			EDRF_UseTriangleOptimization);
 	}
 	Context.RHICmdList.EndRenderPass();
+	Context.RHICmdList.CopyToResolveTarget(DestRenderTarget.TargetableTexture, DestRenderTarget.ShaderResourceTexture, FResolveParams());
 }
 
 FPooledRenderTargetDesc FRCPassPostProcessBloomSetupES2::ComputeOutputDesc(EPassOutputId InPassOutputId) const
@@ -452,7 +453,7 @@ void FRCPassPostProcessBloomSetupSmallES2::Process(FRenderingCompositePassContex
 		LoadAction = ERenderTargetLoadAction::ENoAction;
 	}
 
-	FRHIRenderPassInfo RPInfo(DestRenderTarget.TargetableTexture, MakeRenderTargetActions(LoadAction, ERenderTargetStoreAction::EStore), DestRenderTarget.ShaderResourceTexture);
+	FRHIRenderPassInfo RPInfo(DestRenderTarget.TargetableTexture, MakeRenderTargetActions(LoadAction, ERenderTargetStoreAction::EStore));
 	
 	Context.RHICmdList.BeginRenderPass(RPInfo, TEXT("PostProcessBloomSetupSmallES2"));
 	{
@@ -478,6 +479,7 @@ void FRCPassPostProcessBloomSetupSmallES2::Process(FRenderingCompositePassContex
 			EDRF_UseTriangleOptimization);
 	}
 	Context.RHICmdList.EndRenderPass();
+	Context.RHICmdList.CopyToResolveTarget(DestRenderTarget.TargetableTexture, DestRenderTarget.ShaderResourceTexture, FResolveParams());
 }
 
 FPooledRenderTargetDesc FRCPassPostProcessBloomSetupSmallES2::ComputeOutputDesc(EPassOutputId InPassOutputId) const
@@ -612,7 +614,7 @@ void FRCPassPostProcessBloomDownES2::Process(FRenderingCompositePassContext& Con
 		LoadAction = ERenderTargetLoadAction::ENoAction;
 	}
 
-	FRHIRenderPassInfo RPInfo(DestRenderTarget.TargetableTexture, MakeRenderTargetActions(LoadAction, ERenderTargetStoreAction::EStore), DestRenderTarget.ShaderResourceTexture);
+	FRHIRenderPassInfo RPInfo(DestRenderTarget.TargetableTexture, MakeRenderTargetActions(LoadAction, ERenderTargetStoreAction::EStore));
 	
 	Context.RHICmdList.BeginRenderPass(RPInfo, TEXT("PostProcessBloomDownES2"));
 	{
@@ -656,6 +658,7 @@ void FRCPassPostProcessBloomDownES2::Process(FRenderingCompositePassContext& Con
 			EDRF_UseTriangleOptimization);
 	}
 	Context.RHICmdList.EndRenderPass();
+	Context.RHICmdList.CopyToResolveTarget(DestRenderTarget.TargetableTexture, DestRenderTarget.ShaderResourceTexture, FResolveParams());
 }
 
 FPooledRenderTargetDesc FRCPassPostProcessBloomDownES2::ComputeOutputDesc(EPassOutputId InPassOutputId) const
@@ -795,7 +798,7 @@ void FRCPassPostProcessBloomUpES2::Process(FRenderingCompositePassContext& Conte
 		LoadAction = ERenderTargetLoadAction::ENoAction;
 	}
 
-	FRHIRenderPassInfo RPInfo(DestRenderTarget.TargetableTexture, MakeRenderTargetActions(LoadAction, ERenderTargetStoreAction::EStore), DestRenderTarget.ShaderResourceTexture);
+	FRHIRenderPassInfo RPInfo(DestRenderTarget.TargetableTexture, MakeRenderTargetActions(LoadAction, ERenderTargetStoreAction::EStore));
 		
 	Context.RHICmdList.BeginRenderPass(RPInfo, TEXT("PostProcessBloomUpES2"));
 	{
@@ -842,6 +845,7 @@ void FRCPassPostProcessBloomUpES2::Process(FRenderingCompositePassContext& Conte
 			EDRF_UseTriangleOptimization);
 	}
 	Context.RHICmdList.EndRenderPass();
+	Context.RHICmdList.CopyToResolveTarget(DestRenderTarget.TargetableTexture, DestRenderTarget.ShaderResourceTexture, FResolveParams());
 }
 
 FPooledRenderTargetDesc FRCPassPostProcessBloomUpES2::ComputeOutputDesc(EPassOutputId InPassOutputId) const
@@ -1108,7 +1112,7 @@ void FRCPassPostProcessSunMaskES2::Process(FRenderingCompositePassContext& Conte
 			LoadAction = ERenderTargetLoadAction::ENoAction;
 		}
 
-		FRHIRenderPassInfo RPInfo(DestRenderTarget.TargetableTexture, MakeRenderTargetActions(LoadAction, ERenderTargetStoreAction::EStore), DestRenderTarget.ShaderResourceTexture);
+		FRHIRenderPassInfo RPInfo(DestRenderTarget.TargetableTexture, MakeRenderTargetActions(LoadAction, ERenderTargetStoreAction::EStore));
 		Context.RHICmdList.BeginRenderPass(RPInfo, TEXT("PostProcessSunMaskES2"));
 		{
 			if (bUseClearQuad)
@@ -1139,6 +1143,7 @@ void FRCPassPostProcessSunMaskES2::Process(FRenderingCompositePassContext& Conte
 				EDRF_UseTriangleOptimization);
 		}
 		Context.RHICmdList.EndRenderPass();
+		Context.RHICmdList.CopyToResolveTarget(DestRenderTarget.TargetableTexture, DestRenderTarget.ShaderResourceTexture, FResolveParams());
 	}
 }
 
@@ -1315,7 +1320,7 @@ void FRCPassPostProcessSunAlphaES2::Process(FRenderingCompositePassContext& Cont
 		LoadAction = ERenderTargetLoadAction::ENoAction;
 	}
 
-	FRHIRenderPassInfo RPInfo(DestRenderTarget.TargetableTexture, MakeRenderTargetActions(LoadAction, ERenderTargetStoreAction::EStore), DestRenderTarget.ShaderResourceTexture);
+	FRHIRenderPassInfo RPInfo(DestRenderTarget.TargetableTexture, MakeRenderTargetActions(LoadAction, ERenderTargetStoreAction::EStore));
 	Context.RHICmdList.BeginRenderPass(RPInfo, TEXT("PostProcessSunAlphaES2"));
 	{
 		if (bUseClearQuad)
@@ -1342,6 +1347,7 @@ void FRCPassPostProcessSunAlphaES2::Process(FRenderingCompositePassContext& Cont
 			EDRF_UseTriangleOptimization);
 	}
 	Context.RHICmdList.EndRenderPass();
+	Context.RHICmdList.CopyToResolveTarget(DestRenderTarget.TargetableTexture, DestRenderTarget.ShaderResourceTexture, FResolveParams());
 }
 
 FPooledRenderTargetDesc FRCPassPostProcessSunAlphaES2::ComputeOutputDesc(EPassOutputId InPassOutputId) const
@@ -1476,7 +1482,7 @@ void FRCPassPostProcessSunBlurES2::Process(FRenderingCompositePassContext& Conte
 		LoadAction = ERenderTargetLoadAction::ENoAction;
 	}
 
-	FRHIRenderPassInfo RPInfo(DestRenderTarget.TargetableTexture, MakeRenderTargetActions(LoadAction, ERenderTargetStoreAction::EStore), DestRenderTarget.ShaderResourceTexture);
+	FRHIRenderPassInfo RPInfo(DestRenderTarget.TargetableTexture, MakeRenderTargetActions(LoadAction, ERenderTargetStoreAction::EStore));
 	Context.RHICmdList.BeginRenderPass(RPInfo, TEXT("SunBlurES2"));
 	{
 		if (bUseClearQuad)
@@ -1519,6 +1525,7 @@ void FRCPassPostProcessSunBlurES2::Process(FRenderingCompositePassContext& Conte
 			EDRF_UseTriangleOptimization);
 	}
 	Context.RHICmdList.EndRenderPass();
+	Context.RHICmdList.CopyToResolveTarget(DestRenderTarget.TargetableTexture, DestRenderTarget.ShaderResourceTexture, FResolveParams());
 }
 
 FPooledRenderTargetDesc FRCPassPostProcessSunBlurES2::ComputeOutputDesc(EPassOutputId InPassOutputId) const
@@ -1728,7 +1735,7 @@ void FRCPassPostProcessSunMergeES2::Process(FRenderingCompositePassContext& Cont
 	{
 		LoadAction = ERenderTargetLoadAction::ENoAction;
 	}
-	FRHIRenderPassInfo RPInfo(DestRenderTarget.TargetableTexture, MakeRenderTargetActions(LoadAction, ERenderTargetStoreAction::EStore), DestRenderTarget.ShaderResourceTexture);
+	FRHIRenderPassInfo RPInfo(DestRenderTarget.TargetableTexture, MakeRenderTargetActions(LoadAction, ERenderTargetStoreAction::EStore));
 	Context.RHICmdList.BeginRenderPass(RPInfo, TEXT("SunMergeES2"));
 	{
 		if (bUseClearQuad)
@@ -1754,6 +1761,7 @@ void FRCPassPostProcessSunMergeES2::Process(FRenderingCompositePassContext& Cont
 			EDRF_UseTriangleOptimization);
 	}
 	Context.RHICmdList.EndRenderPass();
+	Context.RHICmdList.CopyToResolveTarget(DestRenderTarget.TargetableTexture, DestRenderTarget.ShaderResourceTexture, FResolveParams());
 
 	// Double buffer sun+bloom+vignette composite.
 	if(Context.View.AntiAliasingMethod == AAM_TemporalAA)
@@ -1931,7 +1939,7 @@ void FRCPassPostProcessSunMergeSmallES2::Process(FRenderingCompositePassContext&
 	{
 		LoadAction = ERenderTargetLoadAction::ENoAction;
 	}
-	FRHIRenderPassInfo RPInfo(DestRenderTarget.TargetableTexture, MakeRenderTargetActions(LoadAction, ERenderTargetStoreAction::EStore), DestRenderTarget.ShaderResourceTexture);
+	FRHIRenderPassInfo RPInfo(DestRenderTarget.TargetableTexture, MakeRenderTargetActions(LoadAction, ERenderTargetStoreAction::EStore));
 	Context.RHICmdList.BeginRenderPass(RPInfo, TEXT("SunMergeSmallES2"));
 	{
 		if (bUseClearQuad)
@@ -1958,6 +1966,7 @@ void FRCPassPostProcessSunMergeSmallES2::Process(FRenderingCompositePassContext&
 			EDRF_UseTriangleOptimization);
 	}
 	Context.RHICmdList.EndRenderPass();
+	Context.RHICmdList.CopyToResolveTarget(DestRenderTarget.TargetableTexture, DestRenderTarget.ShaderResourceTexture, FResolveParams());
 
 	// Double buffer sun+bloom+vignette composite.
 
@@ -2173,7 +2182,7 @@ void FRCPassPostProcessDofDownES2::Process(FRenderingCompositePassContext& Conte
 	{
 		LoadAction = ERenderTargetLoadAction::ENoAction;
 	}
-	FRHIRenderPassInfo RPInfo(DestRenderTarget.TargetableTexture, MakeRenderTargetActions(LoadAction, ERenderTargetStoreAction::EStore), DestRenderTarget.ShaderResourceTexture);
+	FRHIRenderPassInfo RPInfo(DestRenderTarget.TargetableTexture, MakeRenderTargetActions(LoadAction, ERenderTargetStoreAction::EStore));
 	Context.RHICmdList.BeginRenderPass(RPInfo, TEXT("DofDownES2"));
 	{
 		if (bUseClearQuad)
@@ -2199,6 +2208,7 @@ void FRCPassPostProcessDofDownES2::Process(FRenderingCompositePassContext& Conte
 			EDRF_UseTriangleOptimization);
 	}
 	Context.RHICmdList.EndRenderPass();
+	Context.RHICmdList.CopyToResolveTarget(DestRenderTarget.TargetableTexture, DestRenderTarget.ShaderResourceTexture, FResolveParams());
 }
 
 FPooledRenderTargetDesc FRCPassPostProcessDofDownES2::ComputeOutputDesc(EPassOutputId InPassOutputId) const
@@ -2378,7 +2388,7 @@ void FRCPassPostProcessDofNearES2::Process(FRenderingCompositePassContext& Conte
 	{
 		LoadAction = ERenderTargetLoadAction::ENoAction;
 	}
-	FRHIRenderPassInfo RPInfo(DestRenderTarget.TargetableTexture, MakeRenderTargetActions(LoadAction, ERenderTargetStoreAction::EStore), DestRenderTarget.ShaderResourceTexture);
+	FRHIRenderPassInfo RPInfo(DestRenderTarget.TargetableTexture, MakeRenderTargetActions(LoadAction, ERenderTargetStoreAction::EStore));
 	Context.RHICmdList.BeginRenderPass(RPInfo, TEXT("DofNearES2"));
 	{
 		if (bUseClearQuad)
@@ -2405,6 +2415,7 @@ void FRCPassPostProcessDofNearES2::Process(FRenderingCompositePassContext& Conte
 			EDRF_UseTriangleOptimization);
 	}
 	Context.RHICmdList.EndRenderPass();
+	Context.RHICmdList.CopyToResolveTarget(DestRenderTarget.TargetableTexture, DestRenderTarget.ShaderResourceTexture, FResolveParams());
 }
 
 FPooledRenderTargetDesc FRCPassPostProcessDofNearES2::ComputeOutputDesc(EPassOutputId InPassOutputId) const
@@ -2531,7 +2542,7 @@ void FRCPassPostProcessDofBlurES2::Process(FRenderingCompositePassContext& Conte
 	{
 		LoadAction = ERenderTargetLoadAction::ENoAction;
 	}
-	FRHIRenderPassInfo RPInfo(DestRenderTarget.TargetableTexture, MakeRenderTargetActions(LoadAction, ERenderTargetStoreAction::EStore), DestRenderTarget.ShaderResourceTexture);
+	FRHIRenderPassInfo RPInfo(DestRenderTarget.TargetableTexture, MakeRenderTargetActions(LoadAction, ERenderTargetStoreAction::EStore));
 	Context.RHICmdList.BeginRenderPass(RPInfo, TEXT("DofBlurES2"));
 	{
 		if (bUseClearQuad)
@@ -2574,6 +2585,7 @@ void FRCPassPostProcessDofBlurES2::Process(FRenderingCompositePassContext& Conte
 			EDRF_UseTriangleOptimization);
 	}
 	Context.RHICmdList.EndRenderPass();
+	Context.RHICmdList.CopyToResolveTarget(DestRenderTarget.TargetableTexture, DestRenderTarget.ShaderResourceTexture, FResolveParams());
 }
 
 FPooledRenderTargetDesc FRCPassPostProcessDofBlurES2::ComputeOutputDesc(EPassOutputId InPassOutputId) const
@@ -2731,7 +2743,7 @@ void FRCPassPostProcessSunAvgES2::Process(FRenderingCompositePassContext& Contex
 	{
 		LoadAction = ERenderTargetLoadAction::ENoAction;
 	}
-	FRHIRenderPassInfo RPInfo(DestRenderTarget.TargetableTexture, MakeRenderTargetActions(LoadAction, ERenderTargetStoreAction::EStore), DestRenderTarget.ShaderResourceTexture);
+	FRHIRenderPassInfo RPInfo(DestRenderTarget.TargetableTexture, MakeRenderTargetActions(LoadAction, ERenderTargetStoreAction::EStore));
 	Context.RHICmdList.BeginRenderPass(RPInfo, TEXT("SunAvgES2"));
 	{
 		if (bUseClearQuad)
@@ -2758,6 +2770,7 @@ void FRCPassPostProcessSunAvgES2::Process(FRenderingCompositePassContext& Contex
 			EDRF_UseTriangleOptimization);
 	}
 	Context.RHICmdList.EndRenderPass();
+	Context.RHICmdList.CopyToResolveTarget(DestRenderTarget.TargetableTexture, DestRenderTarget.ShaderResourceTexture, FResolveParams());
 }
 
 FPooledRenderTargetDesc FRCPassPostProcessSunAvgES2::ComputeOutputDesc(EPassOutputId InPassOutputId) const
@@ -2986,7 +2999,7 @@ void FRCPassPostProcessAaES2::Process(FRenderingCompositePassContext& Context)
 		LoadStoreAction = ERenderTargetActions::Clear_Store;
 	}
 	
-	FRHIRenderPassInfo RPInfo(DestRenderTarget.TargetableTexture, LoadStoreAction, DestRenderTarget.ShaderResourceTexture);
+	FRHIRenderPassInfo RPInfo(DestRenderTarget.TargetableTexture, LoadStoreAction);
 	Context.RHICmdList.BeginRenderPass(RPInfo, TEXT("AaES2"));
 	{
 		Context.SetViewportAndCallRHI(0, 0, 0.0f, DestSize.X, DestSize.Y, 1.0f);
@@ -3013,6 +3026,7 @@ void FRCPassPostProcessAaES2::Process(FRenderingCompositePassContext& Context)
 			EDRF_UseTriangleOptimization);
 	}
 	Context.RHICmdList.EndRenderPass();
+	Context.RHICmdList.CopyToResolveTarget(DestRenderTarget.TargetableTexture, DestRenderTarget.ShaderResourceTexture, FResolveParams());
 
 	if (FSceneRenderer::ShouldCompositeEditorPrimitives(Context.View))
 	{
diff --git a/Engine/Source/Runtime/Renderer/Private/PostProcess/PostProcessMorpheus.cpp b/Engine/Source/Runtime/Renderer/Private/PostProcess/PostProcessMorpheus.cpp
--- a/Engine/Source/Runtime/Renderer/Private/PostProcess/PostProcessMorpheus.cpp
+++ b/Engine/Source/Runtime/Renderer/Private/PostProcess/PostProcessMorpheus.cpp
@@ -223,7 +223,7 @@ void FRCPassPostProcessMorpheus::Process(FRenderingCompositePassContext& Context
 	const FSceneRenderTargetItem& DestRenderTarget = PassOutputs[0].RequestSurface(Context);
 
 	// Set the view family's render target/viewport.
-	FRHIRenderPassInfo RPInfo(DestRenderTarget.TargetableTexture, ERenderTargetActions::Load_Store, DestRenderTarget.ShaderResourceTexture);
+	FRHIRenderPassInfo RPInfo(DestRenderTarget.TargetableTexture, ERenderTargetActions::Load_Store);
 	Context.RHICmdList.BeginRenderPass(RPInfo, TEXT("MorpheusPostProcess"));
 	{
 		Context.SetViewportAndCallRHI(DestRect);
@@ -269,6 +269,7 @@ void FRCPassPostProcessMorpheus::Process(FRenderingCompositePassContext& Context
 #endif
 	}
 	Context.RHICmdList.EndRenderPass();
+	Context.RHICmdList.CopyToResolveTarget(DestRenderTarget.TargetableTexture, DestRenderTarget.ShaderResourceTexture, FResolveParams());
 }
 
 FPooledRenderTargetDesc FRCPassPostProcessMorpheus::ComputeOutputDesc(EPassOutputId InPassOutputId) const
diff --git a/Engine/Source/Runtime/Renderer/Private/PostProcess/PostProcessMotionBlur.cpp b/Engine/Source/Runtime/Renderer/Private/PostProcess/PostProcessMotionBlur.cpp
--- a/Engine/Source/Runtime/Renderer/Private/PostProcess/PostProcessMotionBlur.cpp
+++ b/Engine/Source/Runtime/Renderer/Private/PostProcess/PostProcessMotionBlur.cpp
@@ -339,7 +339,7 @@ void FRCPassPostProcessVelocityScatter::Process(FRenderingCompositePassContext&
 	GRenderTargetPool.FindFreeElement(Context.RHICmdList, Desc, DepthTarget, TEXT("VelocityScatterDepth") );
 
 	// Set the view family's render target/viewport.
-	FRHIRenderPassInfo RPInfo(DestRenderTarget.TargetableTexture, ERenderTargetActions::Load_Store, DestRenderTarget.ShaderResourceTexture);
+	FRHIRenderPassInfo RPInfo(DestRenderTarget.TargetableTexture, ERenderTargetActions::Load_Store);
 	RPInfo.DepthStencilRenderTarget.DepthStencilTarget = DepthTarget->GetRenderTargetItem().TargetableTexture;
 	RPInfo.DepthStencilRenderTarget.Action = MakeDepthStencilTargetActions(ERenderTargetActions::Clear_Store, ERenderTargetActions::Load_Store);
 	RPInfo.DepthStencilRenderTarget.ExclusiveDepthStencil = FExclusiveDepthStencil::DepthWrite_StencilWrite;
@@ -392,6 +392,7 @@ void FRCPassPostProcessVelocityScatter::Process(FRenderingCompositePassContext&
 		}
 	}
 	Context.RHICmdList.EndRenderPass();
+	Context.RHICmdList.CopyToResolveTarget(DestRenderTarget.TargetableTexture, DestRenderTarget.ShaderResourceTexture, FResolveParams());
 }
 
 FPooledRenderTargetDesc FRCPassPostProcessVelocityScatter::ComputeOutputDesc(EPassOutputId InPassOutputId) const
@@ -911,7 +912,7 @@ void FRCPassPostProcessMotionBlur::Process(FRenderingCompositePassContext& Conte
 
 		ERenderTargetLoadAction LoadAction = Context.GetLoadActionForRenderTarget(DestRenderTarget);
 
-		FRHIRenderPassInfo RPInfo(DestRenderTarget.TargetableTexture, MakeRenderTargetActions(LoadAction, ERenderTargetStoreAction::EStore), DestRenderTarget.ShaderResourceTexture);
+		FRHIRenderPassInfo RPInfo(DestRenderTarget.TargetableTexture, MakeRenderTargetActions(LoadAction, ERenderTargetStoreAction::EStore));
 		Context.RHICmdList.BeginRenderPass(RPInfo, TEXT("MotionBlur"));
 		{
 			Context.SetViewportAndCallRHI(SrcRect);
@@ -950,6 +951,7 @@ void FRCPassPostProcessMotionBlur::Process(FRenderingCompositePassContext& Conte
 				EDRF_UseTriangleOptimization);
 		}
 		Context.RHICmdList.EndRenderPass();
+		Context.RHICmdList.CopyToResolveTarget(DestRenderTarget.TargetableTexture, DestRenderTarget.ShaderResourceTexture, FResolveParams());
 	}
 }
 
diff --git a/Engine/Source/Runtime/Renderer/Private/PostProcess/PostProcessNoiseBlur.cpp b/Engine/Source/Runtime/Renderer/Private/PostProcess/PostProcessNoiseBlur.cpp
--- a/Engine/Source/Runtime/Renderer/Private/PostProcess/PostProcessNoiseBlur.cpp
+++ b/Engine/Source/Runtime/Renderer/Private/PostProcess/PostProcessNoiseBlur.cpp
@@ -149,7 +149,7 @@ void FRCPassPostProcessNoiseBlur::Process(FRenderingCompositePassContext& Contex
 	const FSceneRenderTargetItem& DestRenderTarget = PassOutputs[0].RequestSurface(Context);
 
 	// Set the view family's render target/viewport.
-	FRHIRenderPassInfo RPInfo(DestRenderTarget.TargetableTexture, ERenderTargetActions::Load_Store, DestRenderTarget.ShaderResourceTexture);
+	FRHIRenderPassInfo RPInfo(DestRenderTarget.TargetableTexture, ERenderTargetActions::Load_Store);
 	Context.RHICmdList.BeginRenderPass(RPInfo, TEXT("NoiseBlur"));
 	{
 		// #todo-renderpasses perhaps an optimization here. Use NoAction if this will clear the whole RT
@@ -190,6 +190,7 @@ void FRCPassPostProcessNoiseBlur::Process(FRenderingCompositePassContext& Contex
 			EDRF_UseTriangleOptimization);
 	}
 	Context.RHICmdList.EndRenderPass();
+	Context.RHICmdList.CopyToResolveTarget(DestRenderTarget.TargetableTexture, DestRenderTarget.ShaderResourceTexture, FResolveParams());
 }
 
 
diff --git a/Engine/Source/Runtime/Renderer/Private/PostProcess/PostProcessSelectionOutline.cpp b/Engine/Source/Runtime/Renderer/Private/PostProcess/PostProcessSelectionOutline.cpp
--- a/Engine/Source/Runtime/Renderer/Private/PostProcess/PostProcessSelectionOutline.cpp
+++ b/Engine/Source/Runtime/Renderer/Private/PostProcess/PostProcessSelectionOutline.cpp
@@ -49,7 +49,6 @@ void FRCPassPostProcessSelectionOutlineColor::Process(FRenderingCompositePassCon
 	// Set the render target/viewport.
 	FRHIRenderPassInfo RPInfo;
 	RPInfo.DepthStencilRenderTarget.DepthStencilTarget = DestRenderTarget.TargetableTexture;
-	RPInfo.DepthStencilRenderTarget.ResolveTarget = DestRenderTarget.ShaderResourceTexture;
 	RPInfo.DepthStencilRenderTarget.Action = MakeDepthStencilTargetActions(ERenderTargetActions::Clear_Store, ERenderTargetActions::Clear_Store);
 	RPInfo.DepthStencilRenderTarget.ExclusiveDepthStencil = FExclusiveDepthStencil::DepthWrite_StencilWrite;
 
@@ -174,6 +173,7 @@ void FRCPassPostProcessSelectionOutlineColor::Process(FRenderingCompositePassCon
 		}
 	}
 	Context.RHICmdList.EndRenderPass();
+	Context.RHICmdList.CopyToResolveTarget(DestRenderTarget.TargetableTexture, DestRenderTarget.ShaderResourceTexture, FResolveParams());
 #endif	
 }
 
@@ -409,7 +409,7 @@ void FRCPassPostProcessSelectionOutline::Process(FRenderingCompositePassContext&
 	FIntPoint SrcSize = SceneColorInputDesc->Extent;
 
 	// Set the view family's render target/viewport.
-	FRHIRenderPassInfo RPInfo(DestRenderTarget.TargetableTexture, ERenderTargetActions::Load_Store, DestRenderTarget.ShaderResourceTexture);
+	FRHIRenderPassInfo RPInfo(DestRenderTarget.TargetableTexture, ERenderTargetActions::Load_Store);
 	Context.RHICmdList.BeginRenderPass(RPInfo, TEXT("SelectionOutline"));
 	{
 		Context.SetViewportAndCallRHI(DestRect);
@@ -452,6 +452,7 @@ void FRCPassPostProcessSelectionOutline::Process(FRenderingCompositePassContext&
 			EDRF_UseTriangleOptimization);
 	}
 	Context.RHICmdList.EndRenderPass();
+	Context.RHICmdList.CopyToResolveTarget(DestRenderTarget.TargetableTexture, DestRenderTarget.ShaderResourceTexture, FResolveParams());
 }
 
 FPooledRenderTargetDesc FRCPassPostProcessSelectionOutline::ComputeOutputDesc(EPassOutputId InPassOutputId) const
diff --git a/Engine/Source/Runtime/Renderer/Private/PostProcess/PostProcessTemporalAA.cpp b/Engine/Source/Runtime/Renderer/Private/PostProcess/PostProcessTemporalAA.cpp
--- a/Engine/Source/Runtime/Renderer/Private/PostProcess/PostProcessTemporalAA.cpp
+++ b/Engine/Source/Runtime/Renderer/Private/PostProcess/PostProcessTemporalAA.cpp
@@ -874,12 +874,11 @@ void FRCPassPostProcessTemporalAA::Process(FRenderingCompositePassContext& Conte
 		
 		// Inform MultiGPU systems that we're starting to update this resource
 		Context.RHICmdList.BeginUpdateMultiFrameResource(DestRenderTarget[0]->ShaderResourceTexture);
-		FRHIRenderPassInfo RPInfo(DestRenderTarget[0]->TargetableTexture, ERenderTargetActions::DontLoad_Store, DestRenderTarget[0]->ShaderResourceTexture);
+		FRHIRenderPassInfo RPInfo(DestRenderTarget[0]->TargetableTexture, ERenderTargetActions::DontLoad_Store);
 
 		if (RenderTargetCount == 2)
 		{
 			RPInfo.ColorRenderTargets[1].RenderTarget = DestRenderTarget[1]->TargetableTexture;
-			RPInfo.ColorRenderTargets[1].ResolveTarget = DestRenderTarget[1]->ShaderResourceTexture;
 			RPInfo.ColorRenderTargets[1].Action = ERenderTargetActions::DontLoad_Store;
 			RPInfo.ColorRenderTargets[1].ArraySlice = -1;
 			RPInfo.ColorRenderTargets[1].MipIndex = 0;
@@ -944,6 +943,12 @@ void FRCPassPostProcessTemporalAA::Process(FRenderingCompositePassContext& Conte
 			}
 		}
 		Context.RHICmdList.EndRenderPass();
+		Context.RHICmdList.CopyToResolveTarget(DestRenderTarget[0]->TargetableTexture, DestRenderTarget[0]->ShaderResourceTexture, FResolveParams());
+
+		if (RenderTargetCount == 2)
+		{
+			Context.RHICmdList.CopyToResolveTarget(DestRenderTarget[1]->TargetableTexture, DestRenderTarget[1]->ShaderResourceTexture, FResolveParams());
+		}
 
 		if (IsDOFTAAConfig(Parameters.Pass))
 		{
diff --git a/Engine/Source/Runtime/Renderer/Private/PostProcess/PostProcessTonemap.cpp b/Engine/Source/Runtime/Renderer/Private/PostProcess/PostProcessTonemap.cpp
--- a/Engine/Source/Runtime/Renderer/Private/PostProcess/PostProcessTonemap.cpp
+++ b/Engine/Source/Runtime/Renderer/Private/PostProcess/PostProcessTonemap.cpp
@@ -1074,7 +1074,7 @@ void FRCPassPostProcessTonemap::Process(FRenderingCompositePassContext& Context)
 			}
 		}
 
-		FRHIRenderPassInfo RPInfo(DestRenderTarget.TargetableTexture, MakeRenderTargetActions(LoadAction, ERenderTargetStoreAction::EStore), DestRenderTarget.ShaderResourceTexture);
+		FRHIRenderPassInfo RPInfo(DestRenderTarget.TargetableTexture, MakeRenderTargetActions(LoadAction, ERenderTargetStoreAction::EStore));
 		Context.RHICmdList.BeginRenderPass(RPInfo, TEXT("Tonemap"));
 		{
 			Context.SetViewportAndCallRHI(DestRect, 0.0f, 1.0f);
@@ -1133,6 +1133,7 @@ void FRCPassPostProcessTonemap::Process(FRenderingCompositePassContext& Context)
 				EDRF_UseTriangleOptimization);
 		}
 		Context.RHICmdList.EndRenderPass();
+		Context.RHICmdList.CopyToResolveTarget(DestRenderTarget.TargetableTexture, DestRenderTarget.ShaderResourceTexture, FResolveParams());
 
 		// We only release the SceneColor after the last view was processed (SplitScreen)
 		if(Context.View.Family->Views[Context.View.Family->Views.Num() - 1] == &Context.View && !GIsEditor)
@@ -1667,7 +1668,7 @@ void FRCPassPostProcessTonemapES2::Process(FRenderingCompositePassContext& Conte
 		}
 	}
 
-	FRHIRenderPassInfo RPInfo(DestRenderTarget.TargetableTexture, MakeRenderTargetActions(LoadAction, ERenderTargetStoreAction::EStore), DestRenderTarget.ShaderResourceTexture);
+	FRHIRenderPassInfo RPInfo(DestRenderTarget.TargetableTexture, MakeRenderTargetActions(LoadAction, ERenderTargetStoreAction::EStore));
 	Context.RHICmdList.BeginRenderPass(RPInfo, TEXT("TonemapES2"));
 	{
 		Context.SetViewportAndCallRHI(DestRect);
@@ -1709,6 +1710,7 @@ void FRCPassPostProcessTonemapES2::Process(FRenderingCompositePassContext& Conte
 			EDRF_UseTriangleOptimization);
 	}
 	Context.RHICmdList.EndRenderPass();
+	Context.RHICmdList.CopyToResolveTarget(DestRenderTarget.TargetableTexture, DestRenderTarget.ShaderResourceTexture, FResolveParams());
 }
 
 FPooledRenderTargetDesc FRCPassPostProcessTonemapES2::ComputeOutputDesc(EPassOutputId InPassOutputId) const
diff --git a/Engine/Source/Runtime/Renderer/Private/PostProcess/PostProcessUpscale.cpp b/Engine/Source/Runtime/Renderer/Private/PostProcess/PostProcessUpscale.cpp
--- a/Engine/Source/Runtime/Renderer/Private/PostProcess/PostProcessUpscale.cpp
+++ b/Engine/Source/Runtime/Renderer/Private/PostProcess/PostProcessUpscale.cpp
@@ -407,7 +407,7 @@ void FRCPassPostProcessUpscale::Process(FRenderingCompositePassContext& Context)
 		LoadAction = ERenderTargetLoadAction::EClear;
 	}
 
-	FRHIRenderPassInfo RPInfo(DestRenderTarget.TargetableTexture, MakeRenderTargetActions(LoadAction, ERenderTargetStoreAction::EStore), DestRenderTarget.ShaderResourceTexture);
+	FRHIRenderPassInfo RPInfo(DestRenderTarget.TargetableTexture, MakeRenderTargetActions(LoadAction, ERenderTargetStoreAction::EStore));
 	Context.RHICmdList.BeginRenderPass(RPInfo, TEXT("Upscale"));
 	{
 		bool bTessellatedQuad = PaniniConfig.D >= 0.01f;
@@ -471,6 +471,7 @@ void FRCPassPostProcessUpscale::Process(FRenderingCompositePassContext& Context)
 			bTessellatedQuad ? EDRF_UseTesselatedIndexBuffer : EDRF_UseTriangleOptimization);
 	}
 	Context.RHICmdList.EndRenderPass();
+	Context.RHICmdList.CopyToResolveTarget(DestRenderTarget.TargetableTexture, DestRenderTarget.ShaderResourceTexture, FResolveParams());
 
 	// Update scene color view rectangle for secondary upscale.
 	Context.SceneColorViewRect = DestRect;
diff --git a/Engine/Source/Runtime/Renderer/Private/PostProcess/PostProcessWeightedSampleSum.cpp b/Engine/Source/Runtime/Renderer/Private/PostProcess/PostProcessWeightedSampleSum.cpp
--- a/Engine/Source/Runtime/Renderer/Private/PostProcess/PostProcessWeightedSampleSum.cpp
+++ b/Engine/Source/Runtime/Renderer/Private/PostProcess/PostProcessWeightedSampleSum.cpp
@@ -1045,7 +1045,7 @@ void FRCPassPostProcessWeightedSampleSum::Process(FRenderingCompositePassContext
 			LoadAction = ERenderTargetLoadAction::EClear;
 		}
 
-		FRHIRenderPassInfo RPInfo(DestRenderTarget.TargetableTexture, MakeRenderTargetActions(LoadAction, ERenderTargetStoreAction::EStore), DestRenderTarget.ShaderResourceTexture);
+		FRHIRenderPassInfo RPInfo(DestRenderTarget.TargetableTexture, MakeRenderTargetActions(LoadAction, ERenderTargetStoreAction::EStore));
 		Context.RHICmdList.BeginRenderPass(RPInfo, TEXT("WeightedSampleSum"));
 		{
 			Context.SetViewportAndCallRHI(0, 0, 0.0f, DestSize.X, DestSize.Y, 1.0f);
@@ -1077,6 +1077,7 @@ void FRCPassPostProcessWeightedSampleSum::Process(FRenderingCompositePassContext
 			DrawQuad(Context.RHICmdList, FeatureLevel, bDoFastBlur, SrcRect, DestRect, DestSize, SrcSize, VertexShader);
 		}
 		Context.RHICmdList.EndRenderPass();
+		Context.RHICmdList.CopyToResolveTarget(DestRenderTarget.TargetableTexture, DestRenderTarget.ShaderResourceTexture, FResolveParams());
 	}
 }
 
diff --git a/Engine/Source/Runtime/Renderer/Private/PostProcess/VisualizeTexture.cpp b/Engine/Source/Runtime/Renderer/Private/PostProcess/VisualizeTexture.cpp
--- a/Engine/Source/Runtime/Renderer/Private/PostProcess/VisualizeTexture.cpp
+++ b/Engine/Source/Runtime/Renderer/Private/PostProcess/VisualizeTexture.cpp
@@ -413,7 +413,7 @@ void FVisualizeTexture::GenerateContent(FRHICommandListImmediate& RHICmdList, co
 
 	const FSceneRenderTargetItem& DestRenderTarget = VisualizeTextureContent->GetRenderTargetItem();
 
-	FRHIRenderPassInfo RPInfo(DestRenderTarget.TargetableTexture, ERenderTargetActions::Clear_Store, DestRenderTarget.ShaderResourceTexture);
+	FRHIRenderPassInfo RPInfo(DestRenderTarget.TargetableTexture, ERenderTargetActions::Clear_Store);
 	TransitionRenderPassTargets(RHICmdList, RPInfo);
 	
 	RHICmdList.BeginRenderPass(RPInfo, TEXT("VisualizeTexture"));
@@ -497,10 +497,10 @@ void FVisualizeTexture::GenerateContent(FRHICommandListImmediate& RHICmdList, co
 
 	RHICmdList.EndRenderPass();
 
-	/*{
+	{
 		SCOPED_DRAW_EVENT(RHICmdList, VisCopy);
 		RHICmdList.CopyToResolveTarget(DestRenderTarget.TargetableTexture, DestRenderTarget.ShaderResourceTexture, FResolveParams());
-	}*/
+	}
 
 	VisualizeTextureDesc = Desc;
 
diff --git a/Engine/Source/Runtime/Renderer/Private/ReflectionEnvironmentCapture.cpp b/Engine/Source/Runtime/Renderer/Private/ReflectionEnvironmentCapture.cpp
--- a/Engine/Source/Runtime/Renderer/Private/ReflectionEnvironmentCapture.cpp
+++ b/Engine/Source/Runtime/Renderer/Private/ReflectionEnvironmentCapture.cpp
@@ -383,7 +383,7 @@ void FilterReflectionEnvironment(FRHICommandListImmediate& RHICmdList, ERHIFeatu
 
 			for (int32 CubeFace = 0; CubeFace < CubeFace_MAX; CubeFace++)
 			{
-				FRHIRenderPassInfo RPInfo(FilteredCube.TargetableTexture, ERenderTargetActions::DontLoad_Store, FilteredCube.ShaderResourceTexture, MipIndex, CubeFace);
+				FRHIRenderPassInfo RPInfo(FilteredCube.TargetableTexture, ERenderTargetActions::DontLoad_Store, nullptr, MipIndex, CubeFace);
 				RHICmdList.BeginRenderPass(RPInfo, TEXT("FilterMips"));
 
 				RHICmdList.ApplyCachedRenderTargets(GraphicsPSOInit);
@@ -434,6 +434,7 @@ void FilterReflectionEnvironment(FRHICommandListImmediate& RHICmdList, ERHIFeatu
 					*VertexShader);
 
 				RHICmdList.EndRenderPass();
+				RHICmdList.CopyToResolveTarget(FilteredCube.TargetableTexture, FilteredCube.ShaderResourceTexture, FResolveParams());
 			}
 		}
 
diff --git a/Engine/Source/Runtime/Renderer/Private/SystemTextures.cpp b/Engine/Source/Runtime/Renderer/Private/SystemTextures.cpp
--- a/Engine/Source/Runtime/Renderer/Private/SystemTextures.cpp
+++ b/Engine/Source/Runtime/Renderer/Private/SystemTextures.cpp
@@ -28,12 +28,10 @@ void FSystemTextures::InitializeCommonTextures(FRHICommandListImmediate& RHICmdL
 			Desc.AutoWritable = false;
 			GRenderTargetPool.FindFreeElement(RHICmdList, Desc, WhiteDummy, TEXT("WhiteDummy"), true, ERenderTargetTransience::NonTransient);
 
-			//SetRenderTarget(RHICmdList, WhiteDummy->GetRenderTargetItem().TargetableTexture, FTextureRHIRef(), ESimpleRenderTargetMode::EClearColorExistingDepth);
-			//RHICmdList.CopyToResolveTarget(WhiteDummy->GetRenderTargetItem().TargetableTexture, WhiteDummy->GetRenderTargetItem().ShaderResourceTexture, FResolveParams());
-
-			FRHIRenderPassInfo RPInfo(WhiteDummy->GetRenderTargetItem().TargetableTexture, ERenderTargetActions::Clear_Store, WhiteDummy->GetRenderTargetItem().ShaderResourceTexture);
+			FRHIRenderPassInfo RPInfo(WhiteDummy->GetRenderTargetItem().TargetableTexture, ERenderTargetActions::Clear_Store);
 			RHICmdList.BeginRenderPass(RPInfo, TEXT("WhiteDummy"));
 			RHICmdList.EndRenderPass();
+			RHICmdList.CopyToResolveTarget(WhiteDummy->GetRenderTargetItem().TargetableTexture, WhiteDummy->GetRenderTargetItem().ShaderResourceTexture, FResolveParams());
 		}
 
 		// Create a BlackDummy texture
@@ -42,12 +40,10 @@ void FSystemTextures::InitializeCommonTextures(FRHICommandListImmediate& RHICmdL
 			Desc.AutoWritable = false;
 			GRenderTargetPool.FindFreeElement(RHICmdList, Desc, BlackDummy, TEXT("BlackDummy"), true, ERenderTargetTransience::NonTransient);
 
-			/*SetRenderTarget(RHICmdList, BlackDummy->GetRenderTargetItem().TargetableTexture, FTextureRHIRef(), ESimpleRenderTargetMode::EClearColorExistingDepth);
-			RHICmdList.CopyToResolveTarget(BlackDummy->GetRenderTargetItem().TargetableTexture, BlackDummy->GetRenderTargetItem().ShaderResourceTexture, FResolveParams());*/
-
-			FRHIRenderPassInfo RPInfo(BlackDummy->GetRenderTargetItem().TargetableTexture, ERenderTargetActions::Clear_Store, BlackDummy->GetRenderTargetItem().ShaderResourceTexture);
+			FRHIRenderPassInfo RPInfo(BlackDummy->GetRenderTargetItem().TargetableTexture, ERenderTargetActions::Clear_Store);
 			RHICmdList.BeginRenderPass(RPInfo, TEXT("BlackDummy"));
 			RHICmdList.EndRenderPass();
+			RHICmdList.CopyToResolveTarget(BlackDummy->GetRenderTargetItem().TargetableTexture, BlackDummy->GetRenderTargetItem().ShaderResourceTexture, FResolveParams());
 		}
 
 		// Create a BlackAlphaOneDummy texture
@@ -56,12 +52,10 @@ void FSystemTextures::InitializeCommonTextures(FRHICommandListImmediate& RHICmdL
 			Desc.AutoWritable = false;
 			GRenderTargetPool.FindFreeElement(RHICmdList, Desc, BlackAlphaOneDummy, TEXT("BlackAlphaOneDummy"), true, ERenderTargetTransience::NonTransient);
 
-			/*SetRenderTarget(RHICmdList, BlackAlphaOneDummy->GetRenderTargetItem().TargetableTexture, FTextureRHIRef(), ESimpleRenderTargetMode::EClearColorExistingDepth);
-			RHICmdList.CopyToResolveTarget(BlackAlphaOneDummy->GetRenderTargetItem().TargetableTexture, BlackAlphaOneDummy->GetRenderTargetItem().ShaderResourceTexture, FResolveParams());*/
-
-			FRHIRenderPassInfo RPInfo(BlackAlphaOneDummy->GetRenderTargetItem().TargetableTexture, ERenderTargetActions::Clear_Store, BlackAlphaOneDummy->GetRenderTargetItem().ShaderResourceTexture);
+			FRHIRenderPassInfo RPInfo(BlackAlphaOneDummy->GetRenderTargetItem().TargetableTexture, ERenderTargetActions::Clear_Store);
 			RHICmdList.BeginRenderPass(RPInfo, TEXT("BlackAlphaOneDummy"));
 			RHICmdList.EndRenderPass();
+			RHICmdList.CopyToResolveTarget(BlackAlphaOneDummy->GetRenderTargetItem().TargetableTexture, BlackAlphaOneDummy->GetRenderTargetItem().ShaderResourceTexture, FResolveParams());
 		}
 
 		// Create a GreenDummy texture
@@ -70,12 +64,10 @@ void FSystemTextures::InitializeCommonTextures(FRHICommandListImmediate& RHICmdL
 			Desc.AutoWritable = false;
 			GRenderTargetPool.FindFreeElement(RHICmdList, Desc, GreenDummy, TEXT("GreenDummy"), true, ERenderTargetTransience::NonTransient);
 
-			/*SetRenderTarget(RHICmdList, GreenDummy->GetRenderTargetItem().TargetableTexture, FTextureRHIRef(), ESimpleRenderTargetMode::EClearColorExistingDepth);
-			RHICmdList.CopyToResolveTarget(GreenDummy->GetRenderTargetItem().TargetableTexture, GreenDummy->GetRenderTargetItem().ShaderResourceTexture, FResolveParams());*/
-
-			FRHIRenderPassInfo RPInfo(GreenDummy->GetRenderTargetItem().TargetableTexture, ERenderTargetActions::Clear_Store, GreenDummy->GetRenderTargetItem().ShaderResourceTexture);
+			FRHIRenderPassInfo RPInfo(GreenDummy->GetRenderTargetItem().TargetableTexture, ERenderTargetActions::Clear_Store);
 			RHICmdList.BeginRenderPass(RPInfo, TEXT("GreenDummy"));
 			RHICmdList.EndRenderPass();
+			RHICmdList.CopyToResolveTarget(GreenDummy->GetRenderTargetItem().TargetableTexture, GreenDummy->GetRenderTargetItem().ShaderResourceTexture, FResolveParams());
 		}
 
 		// Create a DefaultNormal8Bit texture
@@ -84,12 +76,10 @@ void FSystemTextures::InitializeCommonTextures(FRHICommandListImmediate& RHICmdL
 			Desc.AutoWritable = false;
 			GRenderTargetPool.FindFreeElement(RHICmdList, Desc, DefaultNormal8Bit, TEXT("DefaultNormal8Bit"), true, ERenderTargetTransience::NonTransient);
 
-			/*SetRenderTarget(RHICmdList, DefaultNormal8Bit->GetRenderTargetItem().TargetableTexture, FTextureRHIRef(), ESimpleRenderTargetMode::EClearColorExistingDepth);
-			RHICmdList.CopyToResolveTarget(DefaultNormal8Bit->GetRenderTargetItem().TargetableTexture, DefaultNormal8Bit->GetRenderTargetItem().ShaderResourceTexture, FResolveParams());*/
-
-			FRHIRenderPassInfo RPInfo(DefaultNormal8Bit->GetRenderTargetItem().TargetableTexture, ERenderTargetActions::Clear_Store, DefaultNormal8Bit->GetRenderTargetItem().ShaderResourceTexture);
+			FRHIRenderPassInfo RPInfo(DefaultNormal8Bit->GetRenderTargetItem().TargetableTexture, ERenderTargetActions::Clear_Store);
 			RHICmdList.BeginRenderPass(RPInfo, TEXT("DefaultNormal8Bit"));
 			RHICmdList.EndRenderPass();
+			RHICmdList.CopyToResolveTarget(DefaultNormal8Bit->GetRenderTargetItem().TargetableTexture, DefaultNormal8Bit->GetRenderTargetItem().ShaderResourceTexture, FResolveParams());
 		}
 
 		// Create the PerlinNoiseGradient texture
@@ -130,9 +120,10 @@ void FSystemTextures::InitializeCommonTextures(FRHICommandListImmediate& RHICmdL
 		Desc.AutoWritable = false;
 		GRenderTargetPool.FindFreeElement(RHICmdList, Desc, MaxFP16Depth, TEXT("MaxFP16Depth"), true, ERenderTargetTransience::NonTransient);
 
-		FRHIRenderPassInfo RPInfo(MaxFP16Depth->GetRenderTargetItem().TargetableTexture, ERenderTargetActions::Clear_Store, MaxFP16Depth->GetRenderTargetItem().ShaderResourceTexture);
+		FRHIRenderPassInfo RPInfo(MaxFP16Depth->GetRenderTargetItem().TargetableTexture, ERenderTargetActions::Clear_Store);
 		RHICmdList.BeginRenderPass(RPInfo, TEXT("MaxFP16Depth"));
 		RHICmdList.EndRenderPass();
+		RHICmdList.CopyToResolveTarget(MaxFP16Depth->GetRenderTargetItem().TargetableTexture, MaxFP16Depth->GetRenderTargetItem().ShaderResourceTexture, FResolveParams());
 	}
 
 	// Create dummy 1x1 depth texture		
@@ -141,10 +132,11 @@ void FSystemTextures::InitializeCommonTextures(FRHICommandListImmediate& RHICmdL
 		Desc.AutoWritable = false;
 		GRenderTargetPool.FindFreeElement(RHICmdList, Desc, DepthDummy, TEXT("DepthDummy"), true, ERenderTargetTransience::NonTransient);
 
-		FRHIRenderPassInfo RPInfo(DepthDummy->GetRenderTargetItem().TargetableTexture, EDepthStencilTargetActions::ClearDepthStencil_StoreDepthStencil, DepthDummy->GetRenderTargetItem().ShaderResourceTexture, FExclusiveDepthStencil::DepthWrite_StencilWrite);
+		FRHIRenderPassInfo RPInfo(DepthDummy->GetRenderTargetItem().TargetableTexture, EDepthStencilTargetActions::ClearDepthStencil_StoreDepthStencil, nullptr, FExclusiveDepthStencil::DepthWrite_StencilWrite);
 
 		RHICmdList.BeginRenderPass(RPInfo, TEXT("DepthDummy"));
 		RHICmdList.EndRenderPass();
+		RHICmdList.CopyToResolveTarget(DepthDummy->GetRenderTargetItem().TargetableTexture, DepthDummy->GetRenderTargetItem().ShaderResourceTexture, FResolveParams());
 	}
 }
 
diff --git a/Engine/Source/Runtime/Renderer/Private/VT/VirtualTextureTest.cpp b/Engine/Source/Runtime/Renderer/Private/VT/VirtualTextureTest.cpp
--- a/Engine/Source/Runtime/Renderer/Private/VT/VirtualTextureTest.cpp
+++ b/Engine/Source/Runtime/Renderer/Private/VT/VirtualTextureTest.cpp
@@ -187,7 +187,7 @@ void FVirtualTextureTest::ProducePageData( FRHICommandList& RHICmdList, ERHIFeat
 
 	FSceneRenderTargetItem& RenderTarget = PhysicalTexture->GetRenderTargetItem();
 
-	FRHIRenderPassInfo RPInfo(RenderTarget.TargetableTexture, ERenderTargetActions::Load_Store, RenderTarget.ShaderResourceTexture);
+	FRHIRenderPassInfo RPInfo(RenderTarget.TargetableTexture, ERenderTargetActions::Load_Store);
 	RHICmdList.BeginRenderPass(RPInfo, TEXT("ProducePageData"));
 	{
 		RHICmdList.SetViewport(DstRect.Min.X, DstRect.Min.Y, 0.0f, DstRect.Max.X, DstRect.Max.Y, 1.0f);
@@ -220,6 +220,7 @@ void FVirtualTextureTest::ProducePageData( FRHICommandList& RHICmdList, ERHIFeat
 			EDRF_UseTriangleOptimization);
 	}
 	RHICmdList.EndRenderPass();
+	RHICmdList.CopyToResolveTarget(RenderTarget.TargetableTexture, RenderTarget.ShaderResourceTexture, FResolveParams());
 
 	GVisualizeTexture.SetCheckPoint( RHICmdList, PhysicalTexture );
 }
diff --git a/Engine/Source/Runtime/VulkanRHI/Private/VulkanRenderTarget.cpp b/Engine/Source/Runtime/VulkanRHI/Private/VulkanRenderTarget.cpp
--- a/Engine/Source/Runtime/VulkanRHI/Private/VulkanRenderTarget.cpp
+++ b/Engine/Source/Runtime/VulkanRHI/Private/VulkanRenderTarget.cpp
@@ -1630,8 +1630,6 @@ void FVulkanCommandListContext::RHIEndRenderPass()
 		TransitionAndLayoutManager.EndRealRenderPass(CmdBuffer);
 	}
 
-	IRHICommandContext::RHIEndRenderPass();
-
 	RHIPopEvent();
 }
 
