Stencil Store action must be EStore if we ever use this attachment again in vulkan. By the spec, DontCare is allowed to leave it in an undefined state. This led to busted Distortion rendering.
Added extra validation to catch this case.

diff --git a/Engine/Source/Runtime/RHI/Private/RHI.cpp b/Engine/Source/Runtime/RHI/Private/RHI.cpp
--- a/Engine/Source/Runtime/RHI/Private/RHI.cpp
+++ b/Engine/Source/Runtime/RHI/Private/RHI.cpp
@@ -967,9 +967,28 @@ void FRHIRenderPassInfo::Validate() const
 		ensure(!bIsMSAAResolve || DepthStencilRenderTarget.DepthStencilTarget->GetNumSamples() > 1);
 		// Don't resolve to null
 		//ensure(DepthStencilRenderTarget.ResolveTarget || DepthStore != ERenderTargetStoreAction::EStore);
+
 		// Don't write to depth if read-only
-		ensure((DepthStencilRenderTarget.ExclusiveDepthStencil.IsDepthWrite() && DepthStore == ERenderTargetStoreAction::EStore) || DepthStore != ERenderTargetStoreAction::EStore);
-		ensure((DepthStencilRenderTarget.ExclusiveDepthStencil.IsStencilWrite() && StencilStore == ERenderTargetStoreAction::EStore) || StencilStore != ERenderTargetStoreAction::EStore);
+		ensure(DepthStencilRenderTarget.ExclusiveDepthStencil.IsDepthWrite() || DepthStore != ERenderTargetStoreAction::EStore);
+		// This is not true for stencil. VK and Metal specify that the DontCare store action MAY leave the attachment in an undefined state.
+		/*ensure(DepthStencilRenderTarget.ExclusiveDepthStencil.IsStencilWrite() || StencilStore != ERenderTargetStoreAction::EStore);*/
+
+		// If we have a depthstencil target we MUST Store it or it will be undefined after rendering.
+		if (DepthStencilRenderTarget.DepthStencilTarget->GetFormat() != PF_D24)
+		{
+			// If this is DepthStencil we must store it out unless we are absolutely sure it will never be used again.
+			ensure(StencilStore == ERenderTargetStoreAction::EStore);
+		}
+
+		if (DepthStencilRenderTarget.ExclusiveDepthStencil.IsDepthWrite())
+		{
+			ensure(DepthStore == ERenderTargetStoreAction::EStore);
+		}
+
+		if (DepthStencilRenderTarget.ExclusiveDepthStencil.IsStencilWrite())
+		{
+			ensure(StencilStore == ERenderTargetStoreAction::EStore);
+		}
 	}
 	else
 	{
diff --git a/Engine/Source/Runtime/RHI/Public/RHIResources.h b/Engine/Source/Runtime/RHI/Public/RHIResources.h
--- a/Engine/Source/Runtime/RHI/Public/RHIResources.h
+++ b/Engine/Source/Runtime/RHI/Public/RHIResources.h
@@ -1325,9 +1325,10 @@ public:
 	}
 
 	void Validate() const
-	{		
+	{
+		// VK and Metal MAY leave the attachment in an undefined state if the StoreAction is DontCare. So we can't assume read-only implies it should be DontCare unless we know for sure it will never be used again.
 		ensureMsgf(DepthStencilAccess.IsDepthWrite() || DepthStoreAction == ERenderTargetStoreAction::ENoAction, TEXT("Depth is read-only, but we are performing a store.  This is a waste on mobile.  If depth can't change, we don't need to store it out again"));
-		ensureMsgf(DepthStencilAccess.IsStencilWrite() || StencilStoreAction == ERenderTargetStoreAction::ENoAction, TEXT("Stencil is read-only, but we are performing a store.  This is a waste on mobile.  If stencil can't change, we don't need to store it out again"));
+		/*ensureMsgf(DepthStencilAccess.IsStencilWrite() || StencilStoreAction == ERenderTargetStoreAction::ENoAction, TEXT("Stencil is read-only, but we are performing a store.  This is a waste on mobile.  If stencil can't change, we don't need to store it out again"));*/
 	}
 
 	bool operator==(const FRHIDepthRenderTargetView& Other) const
diff --git a/Engine/Source/Runtime/Renderer/Private/DistortionRendering.cpp b/Engine/Source/Runtime/Renderer/Private/DistortionRendering.cpp
--- a/Engine/Source/Runtime/Renderer/Private/DistortionRendering.cpp
+++ b/Engine/Source/Runtime/Renderer/Private/DistortionRendering.cpp
@@ -1093,7 +1093,7 @@ void FSceneRenderer::RenderDistortion(FRHICommandListImmediate& RHICmdList)
 		if(DistortionRT)
 		{
 			FRHIRenderPassInfo RPInfo(DistortionRT->GetRenderTargetItem().TargetableTexture, ERenderTargetActions::Clear_Store);
-			RPInfo.DepthStencilRenderTarget.Action = MakeDepthStencilTargetActions(ERenderTargetActions::Load_DontStore, ERenderTargetActions::DontLoad_Store);
+			RPInfo.DepthStencilRenderTarget.Action = MakeDepthStencilTargetActions(ERenderTargetActions::Load_DontStore, ERenderTargetActions::Load_Store);
 			RPInfo.DepthStencilRenderTarget.DepthStencilTarget = SceneContext.GetSceneDepthSurface();
 			RPInfo.DepthStencilRenderTarget.ExclusiveDepthStencil = FExclusiveDepthStencil::DepthRead_StencilWrite;
 
@@ -1157,11 +1157,11 @@ void FSceneRenderer::RenderDistortion(FRHICommandListImmediate& RHICmdList)
 		// Apply distortion and store off-screen
 		{
 			FRHIRenderPassInfo RPInfo(DestRenderTarget.TargetableTexture, ERenderTargetActions::Load_Store);
-			RPInfo.DepthStencilRenderTarget.Action = MakeDepthStencilTargetActions(ERenderTargetActions::Load_DontStore, ERenderTargetActions::Load_DontStore);
+			RPInfo.DepthStencilRenderTarget.Action = MakeDepthStencilTargetActions(ERenderTargetActions::Load_DontStore, ERenderTargetActions::Load_Store);
 			RPInfo.DepthStencilRenderTarget.DepthStencilTarget = SceneContext.GetSceneDepthSurface();
 			RPInfo.DepthStencilRenderTarget.ExclusiveDepthStencil = FExclusiveDepthStencil::DepthRead_StencilRead;
 
-			RHICmdList.BeginRenderPass(RPInfo, TEXT("DistortedSceneColor"));
+			RHICmdList.BeginRenderPass(RPInfo, TEXT("DistortionApply"));
 			{
 				for (int32 ViewIndex = 0, Num = Views.Num(); ViewIndex < Num; ++ViewIndex)
 				{
diff --git a/Engine/Source/Runtime/Renderer/Private/PostProcess/PostProcessTemporalAA.cpp b/Engine/Source/Runtime/Renderer/Private/PostProcess/PostProcessTemporalAA.cpp
--- a/Engine/Source/Runtime/Renderer/Private/PostProcess/PostProcessTemporalAA.cpp
+++ b/Engine/Source/Runtime/Renderer/Private/PostProcess/PostProcessTemporalAA.cpp
@@ -889,7 +889,7 @@ void FRCPassPostProcessTemporalAA::Process(FRenderingCompositePassContext& Conte
 		RPInfo.DepthStencilRenderTarget.DepthStencilTarget = SceneContext.GetSceneDepthTexture();
 		RPInfo.DepthStencilRenderTarget.ResolveTarget = nullptr;
 		RPInfo.DepthStencilRenderTarget.Action = MakeDepthStencilTargetActions(ERenderTargetActions::DontLoad_DontStore, ERenderTargetActions::Load_Store);
-		RPInfo.DepthStencilRenderTarget.ExclusiveDepthStencil = FExclusiveDepthStencil::DepthRead_StencilWrite;
+		RPInfo.DepthStencilRenderTarget.ExclusiveDepthStencil = FExclusiveDepthStencil::DepthNop_StencilWrite;
 
 		Context.RHICmdList.BeginRenderPass(RPInfo, TEXT("TemporalAA"));
 		{
