Renderpass updates.

Screen Space Reflections

diff --git a/Engine/Source/Runtime/Renderer/Private/PostProcess/ScreenSpaceReflections.cpp b/Engine/Source/Runtime/Renderer/Private/PostProcess/ScreenSpaceReflections.cpp
--- a/Engine/Source/Runtime/Renderer/Private/PostProcess/ScreenSpaceReflections.cpp
+++ b/Engine/Source/Runtime/Renderer/Private/PostProcess/ScreenSpaceReflections.cpp
@@ -399,16 +399,23 @@ void FRCPassPostProcessScreenSpaceReflections::Process(FRenderingCompositePassCo
 	{ // ScreenSpaceReflectionsStencil draw event
 		SCOPED_DRAW_EVENTF(Context.RHICmdList, ScreenSpaceReflectionsStencil, TEXT("ScreenSpaceReflectionsStencil %dx%d"), View.ViewRect.Width(), View.ViewRect.Height());
 
+		FRHIRenderPassInfo RPInfo(DestRenderTarget.TargetableTexture,
+			MakeRenderTargetActions(ERenderTargetLoadAction::ENoAction, ERenderTargetStoreAction::EStore),
+			SceneContext.GetSceneDepthSurface(),
+			MakeDepthStencilTargetActions(MakeRenderTargetActions(ERenderTargetLoadAction::ENoAction, ERenderTargetStoreAction::ENoAction),
+				MakeRenderTargetActions(ERenderTargetLoadAction::ENoAction, ERenderTargetStoreAction::EStore)),
+			FExclusiveDepthStencil::DepthRead_StencilWrite);
+
+		RHICmdList.BeginRenderPass(RPInfo, TEXT("SSRStencil"));
+		
 		TShaderMapRef< FPostProcessVS > VertexShader(Context.GetShaderMap());
 		TShaderMapRef< FPostProcessScreenSpaceReflectionsStencilPS > PixelShader(Context.GetShaderMap());
-		
-		// bind the dest render target and the depth stencil render target
-		SetRenderTarget(RHICmdList, DestRenderTarget.TargetableTexture, SceneContext.GetSceneDepthSurface(), ESimpleRenderTargetMode::EUninitializedColorAndDepth, FExclusiveDepthStencil::DepthRead_StencilWrite);
+
 		Context.SetViewportAndCallRHI(View.ViewRect);
 
 		// Clear stencil to 0
 		DrawClearQuad(RHICmdList, false, FLinearColor(), false, 0, true, 0, PassOutputs[0].RenderTargetDesc.Extent, View.ViewRect);
-	
+
 		FGraphicsPipelineStateInitializer GraphicsPSOInit;
 		RHICmdList.ApplyCachedRenderTargets(GraphicsPSOInit);
 
@@ -432,7 +439,7 @@ void FRCPassPostProcessScreenSpaceReflections::Process(FRenderingCompositePassCo
 
 		VertexShader->SetParameters(Context);
 		PixelShader->SetParameters(Context.RHICmdList, Context, SSRQuality, true);
-	
+
 		DrawPostProcessPass(
 			Context.RHICmdList,
 			0, 0,
@@ -454,23 +461,30 @@ void FRCPassPostProcessScreenSpaceReflections::Process(FRenderingCompositePassCo
 		FGraphicsPipelineStateInitializer GraphicsPSOInit;
 		if (SSRStencilPrePass)
 		{
+			check(RHICmdList.IsInsideRenderPass());
+
 			// set up the stencil test to match 0, meaning FPostProcessScreenSpaceReflectionsStencilPS has been discarded
 			GraphicsPSOInit.DepthStencilState = TStaticDepthStencilState<false, CF_Always, true, CF_Equal, SO_Keep, SO_Keep, SO_Keep>::GetRHI();
 			RHICmdList.ApplyCachedRenderTargets(GraphicsPSOInit);
 		}
 		else
 		{
+			check(!RHICmdList.IsInsideRenderPass());
+
 			ERenderTargetLoadAction LoadAction = Context.GetLoadActionForRenderTarget(DestRenderTarget);
 
-			FRHIRenderTargetView RtView = FRHIRenderTargetView(DestRenderTarget.TargetableTexture, LoadAction);
-			FRHISetRenderTargetsInfo Info(1, &RtView, FRHIDepthRenderTargetView());
-			Context.RHICmdList.SetRenderTargetsAndClear(Info);
+			// #todo-renderpasses LoadAction isn't actually used since we clear immediately
+			FRHIRenderPassInfo RPInfo(DestRenderTarget.TargetableTexture, MakeRenderTargetActions(ERenderTargetLoadAction::EClear, ERenderTargetStoreAction::EStore), DestRenderTarget.ShaderResourceTexture);
+			RHICmdList.BeginRenderPass(RPInfo, TEXT("ScreenSpaceReflections"));
+
 			Context.SetViewportAndCallRHI(View.ViewRect);
 
 			GraphicsPSOInit.DepthStencilState = TStaticDepthStencilState<false, CF_Always>::GetRHI();
 			RHICmdList.ApplyCachedRenderTargets(GraphicsPSOInit);
 		}
 
+		check(RHICmdList.IsInsideRenderPass());
+
 		// set the state
 		GraphicsPSOInit.BlendState = TStaticBlendState<>::GetRHI();
 		GraphicsPSOInit.RasterizerState = TStaticRasterizerState<>::GetRHI();
@@ -502,9 +516,9 @@ void FRCPassPostProcessScreenSpaceReflections::Process(FRenderingCompositePassCo
 			View.StereoPass,
 			Context.HasHmdMesh(),
 			EDRF_UseTriangleOptimization);
-	
-		RHICmdList.CopyToResolveTarget(DestRenderTarget.TargetableTexture, DestRenderTarget.ShaderResourceTexture, FResolveParams());
-	} // ScreenSpaceReflections
+	}
+	RHICmdList.EndRenderPass();
+	RHICmdList.CopyToResolveTarget(DestRenderTarget.TargetableTexture, DestRenderTarget.ShaderResourceTexture, FResolveParams());
 }
 
 FPooledRenderTargetDesc FRCPassPostProcessScreenSpaceReflections::ComputeOutputDesc(EPassOutputId InPassOutputId) const
@@ -520,6 +534,8 @@ FPooledRenderTargetDesc FRCPassPostProcessScreenSpaceReflections::ComputeOutputD
 
 void RenderScreenSpaceReflections(FRHICommandListImmediate& RHICmdList, FViewInfo& View, TRefCountPtr<IPooledRenderTarget>& SSROutput, TRefCountPtr<IPooledRenderTarget>& VelocityRT)
 {
+	check(RHICmdList.IsOutsideRenderPass());
+
 	FRenderingCompositePassContext CompositeContext(RHICmdList, View);	
 	FPostprocessContext Context(RHICmdList, CompositeContext.Graph, View );
 
