Renderpass updates.

Planar reflections. Self-contained.

diff --git a/Engine/Source/Runtime/Renderer/Private/PlanarReflectionRendering.cpp b/Engine/Source/Runtime/Renderer/Private/PlanarReflectionRendering.cpp
--- a/Engine/Source/Runtime/Renderer/Private/PlanarReflectionRendering.cpp
+++ b/Engine/Source/Runtime/Renderer/Private/PlanarReflectionRendering.cpp
@@ -129,50 +129,51 @@ void PrefilterPlanarReflection(FRHICommandListImmediate& RHICmdList, FViewInfo&
 		// Workaround for a possible driver bug on S7 Adreno, missing planar reflections
 		ERenderTargetLoadAction RTLoadAction = IsVulkanMobilePlatform(View.GetShaderPlatform()) ?  ERenderTargetLoadAction::EClear : ERenderTargetLoadAction::ENoAction;
 
-		FRHIRenderTargetView ColorView(Target->GetRenderTargetTexture(), 0, -1, RTLoadAction, ERenderTargetStoreAction::EStore);
-		FRHISetRenderTargetsInfo Info(1, &ColorView, FRHIDepthRenderTargetView());
-		RHICmdList.SetRenderTargetsAndClear(Info);
+		FRHIRenderPassInfo RPInfo(Target->GetRenderTargetTexture(), MakeRenderTargetActions(RTLoadAction, ERenderTargetStoreAction::EStore));
+		RHICmdList.BeginRenderPass(RPInfo, TEXT("PrefilterPlanarReflections"));
+		{
+			RHICmdList.SetViewport(View.ViewRect.Min.X, View.ViewRect.Min.Y, 0.0f, View.ViewRect.Max.X, View.ViewRect.Max.Y, 1.0f);
 
-		RHICmdList.SetViewport(View.ViewRect.Min.X, View.ViewRect.Min.Y, 0.0f, View.ViewRect.Max.X, View.ViewRect.Max.Y, 1.0f);
+			FGraphicsPipelineStateInitializer GraphicsPSOInit;
+			RHICmdList.ApplyCachedRenderTargets(GraphicsPSOInit);
+			GraphicsPSOInit.BlendState = TStaticBlendState<>::GetRHI();
+			GraphicsPSOInit.RasterizerState = TStaticRasterizerState<FM_Solid, CM_None>::GetRHI();
+			GraphicsPSOInit.DepthStencilState = TStaticDepthStencilState<false, CF_Always>::GetRHI();
 
-		FGraphicsPipelineStateInitializer GraphicsPSOInit;
-		RHICmdList.ApplyCachedRenderTargets(GraphicsPSOInit);
-		GraphicsPSOInit.BlendState = TStaticBlendState<>::GetRHI();
-		GraphicsPSOInit.RasterizerState = TStaticRasterizerState<FM_Solid, CM_None>::GetRHI();
-		GraphicsPSOInit.DepthStencilState = TStaticDepthStencilState<false, CF_Always>::GetRHI();
+			TShaderMapRef<TDeferredLightVS<false> > VertexShader(View.ShaderMap);
+			TShaderMapRef<FPrefilterPlanarReflectionPS<bEnablePlanarReflectionPrefilter> > PixelShader(View.ShaderMap);
 
-		TShaderMapRef<TDeferredLightVS<false> > VertexShader(View.ShaderMap);
-		TShaderMapRef<FPrefilterPlanarReflectionPS<bEnablePlanarReflectionPrefilter> > PixelShader(View.ShaderMap);
+			GraphicsPSOInit.BoundShaderState.VertexDeclarationRHI = GFilterVertexDeclaration.VertexDeclarationRHI;
+			GraphicsPSOInit.BoundShaderState.VertexShaderRHI = GETSAFERHISHADER_VERTEX(*VertexShader);
+			GraphicsPSOInit.BoundShaderState.PixelShaderRHI = GETSAFERHISHADER_PIXEL(*PixelShader);
+			GraphicsPSOInit.PrimitiveType = PT_TriangleList;
 
-		GraphicsPSOInit.BoundShaderState.VertexDeclarationRHI = GFilterVertexDeclaration.VertexDeclarationRHI;
-		GraphicsPSOInit.BoundShaderState.VertexShaderRHI = GETSAFERHISHADER_VERTEX(*VertexShader);
-		GraphicsPSOInit.BoundShaderState.PixelShaderRHI = GETSAFERHISHADER_PIXEL(*PixelShader);
-		GraphicsPSOInit.PrimitiveType = PT_TriangleList;
+			SetGraphicsPipelineState(RHICmdList, GraphicsPSOInit);
 
-		SetGraphicsPipelineState(RHICmdList, GraphicsPSOInit);
+			PixelShader->SetParameters(RHICmdList, View, ReflectionSceneProxy, SceneColorInput, View.ViewRect.Width());
+			VertexShader->SetSimpleLightParameters(RHICmdList, View, FSphere(0));
 
-		PixelShader->SetParameters(RHICmdList, View, ReflectionSceneProxy, SceneColorInput, View.ViewRect.Width());
-		VertexShader->SetSimpleLightParameters(RHICmdList, View, FSphere(0));
+			FIntPoint UV = View.ViewRect.Min;
+			FIntPoint UVSize = View.ViewRect.Size();
 
-		FIntPoint UV = View.ViewRect.Min;
-		FIntPoint UVSize = View.ViewRect.Size();
+			if (RHINeedsToSwitchVerticalAxis(GShaderPlatformForFeatureLevel[View.FeatureLevel]) && !IsMobileHDR())
+			{
+				UV.Y = UV.Y + UVSize.Y;
+				UVSize.Y = -UVSize.Y;
+			}
 
-		if (RHINeedsToSwitchVerticalAxis(GShaderPlatformForFeatureLevel[View.FeatureLevel]) && !IsMobileHDR())
-		{
-			UV.Y = UV.Y + UVSize.Y;
-			UVSize.Y = -UVSize.Y;
+			DrawRectangle(
+				RHICmdList,
+				0, 0,
+				View.ViewRect.Width(), View.ViewRect.Height(),
+				UV.X, UV.Y,
+				UVSize.X, UVSize.Y,
+				View.ViewRect.Size(),
+				FSceneRenderTargets::Get(RHICmdList).GetBufferSizeXY(),
+				*VertexShader,
+				EDRF_UseTriangleOptimization);
 		}
-
-		DrawRectangle(
-			RHICmdList,
-			0, 0,
-			View.ViewRect.Width(), View.ViewRect.Height(),
-			UV.X, UV.Y,
-			UVSize.X, UVSize.Y,
-			View.ViewRect.Size(),
-			FSceneRenderTargets::Get(RHICmdList).GetBufferSizeXY(),
-			*VertexShader,
-			EDRF_UseTriangleOptimization);
+		RHICmdList.EndRenderPass();
 	}
 }
 
@@ -582,6 +583,7 @@ IMPLEMENT_SHADER_TYPE(,FPlanarReflectionPS,TEXT("/Engine/Private/PlanarReflectio
 
 bool FDeferredShadingSceneRenderer::RenderDeferredPlanarReflections(FRHICommandListImmediate& RHICmdList, const FViewInfo& View, bool bLightAccumulationIsInUse, TRefCountPtr<IPooledRenderTarget>& Output)
 {
+	check(RHICmdList.IsOutsideRenderPass());
 	// Prevent rendering unsupported views when ViewIndex >= GMaxPlanarReflectionViews
 	// Planar reflections in those views will fallback to other reflection methods
 	{
@@ -632,61 +634,63 @@ bool FDeferredShadingSceneRenderer::RenderDeferredPlanarReflections(FRHICommandL
 			}
 		}
 
-		SetRenderTarget(RHICmdList, Output->GetRenderTargetItem().TargetableTexture, nullptr);
-
-		if (!bSSRAsInput)
-		{
-			DrawClearQuad(RHICmdList, FLinearColor(0, 0, 0, 0));
-		}
-
+		FRHIRenderPassInfo RPInfo(Output->GetRenderTargetItem().TargetableTexture, ERenderTargetActions::Load_Store);
+		RHICmdList.BeginRenderPass(RPInfo, TEXT("DeferredPlanarReflections"));
 		{
-			RHICmdList.SetViewport(View.ViewRect.Min.X, View.ViewRect.Min.Y, 0.0f, View.ViewRect.Max.X, View.ViewRect.Max.Y, 1.0f);
+			if (!bSSRAsInput)
+			{
+				DrawClearQuad(RHICmdList, FLinearColor(0, 0, 0, 0));
+			}
 
-			FGraphicsPipelineStateInitializer GraphicsPSOInit;
-			RHICmdList.ApplyCachedRenderTargets(GraphicsPSOInit);
+			{
+				RHICmdList.SetViewport(View.ViewRect.Min.X, View.ViewRect.Min.Y, 0.0f, View.ViewRect.Max.X, View.ViewRect.Max.Y, 1.0f);
 
-			// Blend over previous reflections in the output target (SSR or planar reflections that have already been rendered)
-			// Planar reflections win over SSR and reflection environment
-			//@todo - this is order dependent blending, but ordering is coming from registration order
-			GraphicsPSOInit.BlendState = TStaticBlendState<CW_RGBA, BO_Add, BF_One, BF_InverseSourceAlpha, BO_Max, BF_One, BF_One>::GetRHI();
-			GraphicsPSOInit.RasterizerState = TStaticRasterizerState<FM_Solid, CM_None>::GetRHI();
-			GraphicsPSOInit.DepthStencilState = TStaticDepthStencilState<false, CF_Always>::GetRHI();
+				FGraphicsPipelineStateInitializer GraphicsPSOInit;
+				RHICmdList.ApplyCachedRenderTargets(GraphicsPSOInit);
 
-			for (int32 PlanarReflectionIndex = 0; PlanarReflectionIndex < Scene->PlanarReflections.Num(); PlanarReflectionIndex++)
-			{
-				FPlanarReflectionSceneProxy* ReflectionSceneProxy = Scene->PlanarReflections[PlanarReflectionIndex];
+				// Blend over previous reflections in the output target (SSR or planar reflections that have already been rendered)
+				// Planar reflections win over SSR and reflection environment
+				//@todo - this is order dependent blending, but ordering is coming from registration order
+				GraphicsPSOInit.BlendState = TStaticBlendState<CW_RGBA, BO_Add, BF_One, BF_InverseSourceAlpha, BO_Max, BF_One, BF_One>::GetRHI();
+				GraphicsPSOInit.RasterizerState = TStaticRasterizerState<FM_Solid, CM_None>::GetRHI();
+				GraphicsPSOInit.DepthStencilState = TStaticDepthStencilState<false, CF_Always>::GetRHI();
 
-				if (View.ViewFrustum.IntersectBox(ReflectionSceneProxy->WorldBounds.GetCenter(), ReflectionSceneProxy->WorldBounds.GetExtent()))
+				for (int32 PlanarReflectionIndex = 0; PlanarReflectionIndex < Scene->PlanarReflections.Num(); PlanarReflectionIndex++)
 				{
-					SCOPED_DRAW_EVENTF(RHICmdList, PlanarReflection, *ReflectionSceneProxy->OwnerName.ToString());
-
-					TShaderMapRef<TDeferredLightVS<false> > VertexShader(View.ShaderMap);
-					TShaderMapRef<FPlanarReflectionPS> PixelShader(View.ShaderMap);
-
-					GraphicsPSOInit.BoundShaderState.VertexDeclarationRHI = GFilterVertexDeclaration.VertexDeclarationRHI;
-					GraphicsPSOInit.BoundShaderState.VertexShaderRHI = GETSAFERHISHADER_VERTEX(*VertexShader);
-					GraphicsPSOInit.BoundShaderState.PixelShaderRHI = GETSAFERHISHADER_PIXEL(*PixelShader);
-					GraphicsPSOInit.PrimitiveType = PT_TriangleList;
-
-					SetGraphicsPipelineState(RHICmdList, GraphicsPSOInit);
-
-					PixelShader->SetParameters(RHICmdList, View, ReflectionSceneProxy);
-					VertexShader->SetSimpleLightParameters(RHICmdList, View, FSphere(0));
-
-					DrawRectangle(
-						RHICmdList,
-						0, 0,
-						View.ViewRect.Width(), View.ViewRect.Height(),
-						View.ViewRect.Min.X, View.ViewRect.Min.Y,
-						View.ViewRect.Width(), View.ViewRect.Height(),
-						View.ViewRect.Size(),
-						FSceneRenderTargets::Get(RHICmdList).GetBufferSizeXY(),
-						*VertexShader,
-						EDRF_UseTriangleOptimization);
+					FPlanarReflectionSceneProxy* ReflectionSceneProxy = Scene->PlanarReflections[PlanarReflectionIndex];
+
+					if (View.ViewFrustum.IntersectBox(ReflectionSceneProxy->WorldBounds.GetCenter(), ReflectionSceneProxy->WorldBounds.GetExtent()))
+					{
+						SCOPED_DRAW_EVENTF(RHICmdList, PlanarReflection, *ReflectionSceneProxy->OwnerName.ToString());
+
+						TShaderMapRef<TDeferredLightVS<false> > VertexShader(View.ShaderMap);
+						TShaderMapRef<FPlanarReflectionPS> PixelShader(View.ShaderMap);
+
+						GraphicsPSOInit.BoundShaderState.VertexDeclarationRHI = GFilterVertexDeclaration.VertexDeclarationRHI;
+						GraphicsPSOInit.BoundShaderState.VertexShaderRHI = GETSAFERHISHADER_VERTEX(*VertexShader);
+						GraphicsPSOInit.BoundShaderState.PixelShaderRHI = GETSAFERHISHADER_PIXEL(*PixelShader);
+						GraphicsPSOInit.PrimitiveType = PT_TriangleList;
+
+						SetGraphicsPipelineState(RHICmdList, GraphicsPSOInit);
+
+						PixelShader->SetParameters(RHICmdList, View, ReflectionSceneProxy);
+						VertexShader->SetSimpleLightParameters(RHICmdList, View, FSphere(0));
+
+						DrawRectangle(
+							RHICmdList,
+							0, 0,
+							View.ViewRect.Width(), View.ViewRect.Height(),
+							View.ViewRect.Min.X, View.ViewRect.Min.Y,
+							View.ViewRect.Width(), View.ViewRect.Height(),
+							View.ViewRect.Size(),
+							FSceneRenderTargets::Get(RHICmdList).GetBufferSizeXY(),
+							*VertexShader,
+							EDRF_UseTriangleOptimization);
+					}
 				}
 			}
 		}
-
+		RHICmdList.EndRenderPass();
 		RHICmdList.CopyToResolveTarget(Output->GetRenderTargetItem().TargetableTexture, Output->GetRenderTargetItem().ShaderResourceTexture, FResolveParams());
 
 		return true;
