Renderpass updates.

Reflection updates. Self-contained.

diff --git a/Engine/Source/Runtime/Renderer/Private/ReflectionEnvironment.cpp b/Engine/Source/Runtime/Renderer/Private/ReflectionEnvironment.cpp
--- a/Engine/Source/Runtime/Renderer/Private/ReflectionEnvironment.cpp
+++ b/Engine/Source/Runtime/Renderer/Private/ReflectionEnvironment.cpp
@@ -771,6 +771,8 @@ void FDeferredShadingSceneRenderer::SetupReflectionCaptureBuffers(FViewInfo& Vie
 
 void FDeferredShadingSceneRenderer::RenderDeferredReflectionsAndSkyLighting(FRHICommandListImmediate& RHICmdList, TRefCountPtr<IPooledRenderTarget>& DynamicBentNormalAO, TRefCountPtr<IPooledRenderTarget>& VelocityRT)
 {
+	check(RHICmdList.IsOutsideRenderPass());
+
 	if (ViewFamily.EngineShowFlags.VisualizeLightCulling || !ViewFamily.EngineShowFlags.Lighting)
 	{
 		return;
@@ -808,6 +810,8 @@ void FDeferredShadingSceneRenderer::RenderDeferredReflectionsAndSkyLighting(FRHI
 		}
 	}
 
+	check(RHICmdList.IsOutsideRenderPass());
+
 	FSceneRenderTargets& SceneContext = FSceneRenderTargets::Get(RHICmdList);
 
 	const bool bReflectionEnv = ShouldDoReflectionEnvironment();
@@ -914,6 +918,8 @@ void FDeferredShadingSceneRenderer::RenderDeferredReflectionsAndSkyLighting(FRHI
 					*VertexShader);
 			}
 
+			SceneContext.FinishRenderingSceneColor(RHICmdList);
+
 			ResolveSceneColor(RHICmdList);
 		}
 	}
diff --git a/Engine/Source/Runtime/Renderer/Private/ReflectionEnvironmentCapture.cpp b/Engine/Source/Runtime/Renderer/Private/ReflectionEnvironmentCapture.cpp
--- a/Engine/Source/Runtime/Renderer/Private/ReflectionEnvironmentCapture.cpp
+++ b/Engine/Source/Runtime/Renderer/Private/ReflectionEnvironmentCapture.cpp
@@ -233,37 +233,41 @@ float ComputeSingleAverageBrightnessFromCubemap(FRHICommandListImmediate& RHICmd
 	GRenderTargetPool.FindFreeElement(RHICmdList, Desc, ReflectionBrightnessTarget, TEXT("ReflectionBrightness"));
 
 	FTextureRHIRef& BrightnessTarget = ReflectionBrightnessTarget->GetRenderTargetItem().TargetableTexture;
-	SetRenderTarget(RHICmdList, BrightnessTarget, NULL, true);
 
-	FGraphicsPipelineStateInitializer GraphicsPSOInit;
-	RHICmdList.ApplyCachedRenderTargets(GraphicsPSOInit);
-	GraphicsPSOInit.RasterizerState = TStaticRasterizerState<FM_Solid, CM_None>::GetRHI();
-	GraphicsPSOInit.DepthStencilState = TStaticDepthStencilState<false, CF_Always>::GetRHI();
-	GraphicsPSOInit.BlendState = TStaticBlendState<>::GetRHI();
+	FRHIRenderPassInfo RPInfo(BrightnessTarget, ERenderTargetActions::Load_Store);
+	TransitionRenderPassTargets(RHICmdList, RPInfo);
+	RHICmdList.BeginRenderPass(RPInfo, TEXT("ReflectionBrightness"));
+	{
+		FGraphicsPipelineStateInitializer GraphicsPSOInit;
+		RHICmdList.ApplyCachedRenderTargets(GraphicsPSOInit);
+		GraphicsPSOInit.RasterizerState = TStaticRasterizerState<FM_Solid, CM_None>::GetRHI();
+		GraphicsPSOInit.DepthStencilState = TStaticDepthStencilState<false, CF_Always>::GetRHI();
+		GraphicsPSOInit.BlendState = TStaticBlendState<>::GetRHI();
 
-	auto ShaderMap = GetGlobalShaderMap(FeatureLevel);
-	TShaderMapRef<FPostProcessVS> VertexShader(ShaderMap);
-	TShaderMapRef<FComputeBrightnessPS> PixelShader(ShaderMap);
+		auto ShaderMap = GetGlobalShaderMap(FeatureLevel);
+		TShaderMapRef<FPostProcessVS> VertexShader(ShaderMap);
+		TShaderMapRef<FComputeBrightnessPS> PixelShader(ShaderMap);
 
-	GraphicsPSOInit.BoundShaderState.VertexDeclarationRHI = GFilterVertexDeclaration.VertexDeclarationRHI;
-	GraphicsPSOInit.BoundShaderState.VertexShaderRHI = GETSAFERHISHADER_VERTEX(*VertexShader);
-	GraphicsPSOInit.BoundShaderState.PixelShaderRHI = GETSAFERHISHADER_PIXEL(*PixelShader);
-	GraphicsPSOInit.PrimitiveType = PT_TriangleList;
-	
-	SetGraphicsPipelineState(RHICmdList, GraphicsPSOInit);
+		GraphicsPSOInit.BoundShaderState.VertexDeclarationRHI = GFilterVertexDeclaration.VertexDeclarationRHI;
+		GraphicsPSOInit.BoundShaderState.VertexShaderRHI = GETSAFERHISHADER_VERTEX(*VertexShader);
+		GraphicsPSOInit.BoundShaderState.PixelShaderRHI = GETSAFERHISHADER_PIXEL(*PixelShader);
+		GraphicsPSOInit.PrimitiveType = PT_TriangleList;
 
-	PixelShader->SetParameters(RHICmdList, TargetSize, Cubemap);
+		SetGraphicsPipelineState(RHICmdList, GraphicsPSOInit);
 
-	DrawRectangle( 
-		RHICmdList,
-		0, 0, 
-		1, 1,
-		0, 0, 
-		1, 1,
-		FIntPoint(1, 1),
-		FIntPoint(1, 1),
-		*VertexShader);
+		PixelShader->SetParameters(RHICmdList, TargetSize, Cubemap);
 
+		DrawRectangle(
+			RHICmdList,
+			0, 0,
+			1, 1,
+			0, 0,
+			1, 1,
+			FIntPoint(1, 1),
+			FIntPoint(1, 1),
+			*VertexShader);
+	}
+	RHICmdList.EndRenderPass();
 	RHICmdList.CopyToResolveTarget(BrightnessTarget, BrightnessTarget, FResolveParams());
 
 	FSceneRenderTargetItem& EffectiveRT = ReflectionBrightnessTarget->GetRenderTargetItem();
@@ -666,15 +670,16 @@ void ClearScratchCubemaps(FRHICommandList& RHICmdList, int32 TargetSize)
 	{
 		SCOPED_DRAW_EVENT(RHICmdList, ClearScratchCubemapsRT0);
 
-		TransitionSetRenderTargetsHelper(RHICmdList, RT0.TargetableTexture, FTextureRHIParamRef(), FExclusiveDepthStencil::DepthWrite_StencilWrite);
+		RHICmdList.TransitionResource(EResourceTransitionAccess::EWritable, RT0.TargetableTexture);
 
 		for (int32 MipIndex = 0; MipIndex < NumMips; MipIndex++)
 		{
 			for (int32 CubeFace = 0; CubeFace < CubeFace_MAX; CubeFace++)
 			{
-				FRHIRenderTargetView RtView = FRHIRenderTargetView(RT0.TargetableTexture, ERenderTargetLoadAction::EClear, MipIndex, CubeFace);
-				FRHISetRenderTargetsInfo Info(1, &RtView, FRHIDepthRenderTargetView());
-				RHICmdList.SetRenderTargetsAndClear(Info);
+				FRHIRenderPassInfo RPInfo(RT0.TargetableTexture, ERenderTargetActions::Clear_Store, nullptr, MipIndex, CubeFace);
+
+				RHICmdList.BeginRenderPass(RPInfo, TEXT("ClearCubeFace"));
+				RHICmdList.EndRenderPass();
 			}
 		}
 	}
@@ -685,15 +690,16 @@ void ClearScratchCubemaps(FRHICommandList& RHICmdList, int32 TargetSize)
 		FSceneRenderTargetItem& RT1 = SceneContext.ReflectionColorScratchCubemap[1]->GetRenderTargetItem();
 		NumMips = (int32)RT1.TargetableTexture->GetNumMips();
 
-		TransitionSetRenderTargetsHelper(RHICmdList, RT1.TargetableTexture, FTextureRHIParamRef(), FExclusiveDepthStencil::DepthWrite_StencilWrite);
+		RHICmdList.TransitionResource(EResourceTransitionAccess::EWritable, RT1.TargetableTexture);
 
 		for (int32 MipIndex = 0; MipIndex < NumMips; MipIndex++)
 		{
 			for (int32 CubeFace = 0; CubeFace < CubeFace_MAX; CubeFace++)
 			{
-				FRHIRenderTargetView RtView = FRHIRenderTargetView(RT1.TargetableTexture, ERenderTargetLoadAction::EClear, MipIndex, CubeFace);
-				FRHISetRenderTargetsInfo Info(1, &RtView, FRHIDepthRenderTargetView());
-				RHICmdList.SetRenderTargetsAndClear(Info);
+				FRHIRenderPassInfo RPInfo(RT1.TargetableTexture, ERenderTargetActions::Clear_Store, nullptr, MipIndex, CubeFace);
+
+				RHICmdList.BeginRenderPass(RPInfo, TEXT("ClearCubeFace"));
+				RHICmdList.EndRenderPass();
 			}
 		}
 	}
