Renderpass updates.

Deferred decals. Self-contained.

diff --git a/Engine/Source/Runtime/Renderer/Private/CompositionLighting/PostProcessDeferredDecals.cpp b/Engine/Source/Runtime/Renderer/Private/CompositionLighting/PostProcessDeferredDecals.cpp
--- a/Engine/Source/Runtime/Renderer/Private/CompositionLighting/PostProcessDeferredDecals.cpp
+++ b/Engine/Source/Runtime/Renderer/Private/CompositionLighting/PostProcessDeferredDecals.cpp
@@ -632,7 +632,8 @@ void FRCPassPostProcessDeferredDecals::DecodeRTWriteMask(FRenderingCompositePass
 
 	TShaderMapRef< FRTWriteMaskDecodeCS > ComputeShader(Context.GetShaderMap());
 
-	SetRenderTarget(Context.RHICmdList, FTextureRHIRef(), FTextureRHIRef());
+	//  #todo-renderpasses remove once everything is renderpasses
+	UnbindRenderTargets(Context.RHICmdList);
 	Context.SetViewportAndCallRHI(ViewRect);
 	Context.RHICmdList.SetComputeShader(ComputeShader->GetComputeShader());
 
@@ -677,6 +678,7 @@ void FRCPassPostProcessDeferredDecals::Process(FRenderingCompositePassContext& C
 	const FViewInfo& View = Context.View;
 	const FSceneViewFamily& ViewFamily = *(View.Family);
 	bool bNeedsDBufferTargets = false;
+	bool bDidClearDBuffer = false;
 
 	if (CurrentStage == DRS_BeforeBasePass)
 	{
@@ -753,23 +755,25 @@ void FRCPassPostProcessDeferredDecals::Process(FRenderingCompositePassContext& C
 			{
 				SCOPED_DRAW_EVENT(RHICmdList, DBufferClear);
 
-				FRHIRenderTargetView RenderTargets[4];
+				FTextureRHIParamRef RenderTargets[4];
 
-				RenderTargets[0] = FRHIRenderTargetView(SceneContext.DBufferA->GetRenderTargetItem().TargetableTexture, 0, -1, ERenderTargetLoadAction::EClear, ERenderTargetStoreAction::EStore);
-				RenderTargets[1] = FRHIRenderTargetView(SceneContext.DBufferB->GetRenderTargetItem().TargetableTexture, 0, -1, ERenderTargetLoadAction::EClear, ERenderTargetStoreAction::EStore);
-				RenderTargets[2] = FRHIRenderTargetView(SceneContext.DBufferC->GetRenderTargetItem().TargetableTexture, 0, -1, ERenderTargetLoadAction::EClear, ERenderTargetStoreAction::EStore);
-				uint32 RTCount = 3;
+				RenderTargets[0] = SceneContext.DBufferA->GetRenderTargetItem().TargetableTexture;
+				RenderTargets[1] = SceneContext.DBufferB->GetRenderTargetItem().TargetableTexture;
+				RenderTargets[2] = SceneContext.DBufferC->GetRenderTargetItem().TargetableTexture;
+				int32 RTCount = 3;
 
 				if (bPerPixelDBufferMask)
 				{
-					RenderTargets[3] = FRHIRenderTargetView(SceneContext.DBufferMask->GetRenderTargetItem().TargetableTexture, 0, -1, ERenderTargetLoadAction::EClear, ERenderTargetStoreAction::EStore);
+					RenderTargets[3] = SceneContext.DBufferMask->GetRenderTargetItem().TargetableTexture;
 					RTCount = 4;
 				}
 
-				FRHIDepthRenderTargetView DepthView(SceneContext.GetSceneDepthTexture(), ERenderTargetLoadAction::ELoad, ERenderTargetStoreAction::ENoAction, ERenderTargetLoadAction::ELoad, ERenderTargetStoreAction::ENoAction, FExclusiveDepthStencil(FExclusiveDepthStencil::DepthRead_StencilWrite));
-
-				FRHISetRenderTargetsInfo Info(RTCount, RenderTargets, DepthView);
-				RHICmdList.SetRenderTargetsAndClear(Info);
+				FRHIRenderPassInfo RPInfo(RTCount, RenderTargets, ERenderTargetActions::Clear_Store);
+				RPInfo.DepthStencilRenderTarget.DepthStencilTarget = SceneContext.GetSceneDepthTexture();
+				RPInfo.DepthStencilRenderTarget.Action = MakeDepthStencilTargetActions(ERenderTargetActions::Load_DontStore, ERenderTargetActions::Load_Store);
+				RPInfo.DepthStencilRenderTarget.ExclusiveDepthStencil = FExclusiveDepthStencil::DepthRead_StencilNop;
+				RHICmdList.BeginRenderPass(RPInfo, TEXT("InitialDeferredDecals"));
+				bDidClearDBuffer = true;
 			}
 		} // if ( bNeedsDBufferTargets )
 	}
@@ -787,6 +791,10 @@ void FRCPassPostProcessDeferredDecals::Process(FRenderingCompositePassContext& C
 				check(bNeedsDBufferTargets || CurrentStage != DRS_BeforeBasePass);
 				RenderMeshDecals(Context, CurrentStage);
 
+				// Note: There will be an open renderpass at this point.
+				// We are not ending it here in case the next decal uses the same rendertarget.
+				// There is a catch-all to end an active renderpass after the scene decal rendering.
+				
 				bShouldResolveTargets = true;
 			}
 		}
@@ -812,7 +820,12 @@ void FRCPassPostProcessDeferredDecals::Process(FRenderingCompositePassContext& C
 				SCOPED_DRAW_EVENTF(RHICmdList, DeferredDecalsInner, TEXT("DeferredDecalsInner %d/%d"), SortedDecals.Num(), Scene.Decals.Num());
 
 				FGraphicsPipelineStateInitializer GraphicsPSOInit;
-				RHICmdList.ApplyCachedRenderTargets(GraphicsPSOInit);
+				if (bDidClearDBuffer)
+				{
+					// If we cleared the Dbuffer above we'll be inside a renderpass here.
+					check(RHICmdList.IsInsideRenderPass());
+					RHICmdList.ApplyCachedRenderTargets(GraphicsPSOInit);
+				}
 
 				// optimization to have less state changes
 				EDecalRasterizerState LastDecalRasterizerState = DRS_Undefined;
@@ -866,6 +879,8 @@ void FRCPassPostProcessDeferredDecals::Process(FRenderingCompositePassContext& C
 						RHICmdList.ApplyCachedRenderTargets(GraphicsPSOInit);
 					}
 
+					check(RHICmdList.IsInsideRenderPass());
+
 					bool bThisDecalUsesStencil = false;
 
 					if (bStencilThisDecal && bStencilSizeThreshold)
@@ -934,18 +949,32 @@ void FRCPassPostProcessDeferredDecals::Process(FRenderingCompositePassContext& C
 					RHICmdList.DrawIndexedPrimitive(GetUnitCubeIndexBuffer(), 0, 0, 8, 0, ARRAY_COUNT(GCubeIndices) / 3, 1);
 					RenderTargetManager.bGufferADirty |= (RenderTargetManager.TargetsToResolve[FDecalRenderTargetManager::GBufferAIndex] != nullptr);
 				}
+
+				check(RHICmdList.IsInsideRenderPass());
+				// Finished rendering sorted decals, so end the renderpass.
+				RHICmdList.EndRenderPass();
+			}
+
+			if (RHICmdList.IsInsideRenderPass())
+			{
+				// If the SortedDecals list is empty we may have started a renderpass to clear the dbuffer.
+				// If we only draw mesh decals we'll have an active renderpass here as well.
+				RHICmdList.EndRenderPass();
 			}
 
 			if ((Scene.Decals.Num() > 0) && CurrentStage == DRS_Emissive)
 			{
 				// we don't modify stencil but if out input was having stencil for us (after base pass - we need to clear)
 				// Clear stencil to 0, which is the assumed default by other passes
-				FRHIDepthRenderTargetView DepthView(SceneContext.GetSceneDepthTexture(), 
-													ERenderTargetLoadAction::ENoAction, ERenderTargetStoreAction::ENoAction, 
-													ERenderTargetLoadAction::EClear,	ERenderTargetStoreAction::EStore, 
-													FExclusiveDepthStencil(FExclusiveDepthStencil::DepthNop_StencilWrite));
-				FRHISetRenderTargetsInfo Info(0, nullptr, DepthView);
-				RHICmdList.SetRenderTargetsAndClear(Info);
+
+				FRHIRenderPassInfo RPInfo;
+				RPInfo.DepthStencilRenderTarget.Action = MakeDepthStencilTargetActions(ERenderTargetActions::DontLoad_DontStore, ERenderTargetActions::Clear_Store);
+				RPInfo.DepthStencilRenderTarget.DepthStencilTarget = SceneContext.GetSceneDepthTexture();
+				RPInfo.DepthStencilRenderTarget.ExclusiveDepthStencil = FExclusiveDepthStencil::DepthNop_StencilWrite;
+				RPInfo.DepthStencilRenderTarget.ResolveTarget = nullptr;
+
+				RHICmdList.BeginRenderPass(RPInfo, TEXT("ClearStencil"));
+				RHICmdList.EndRenderPass();
 				//DrawClearQuad(RHICmdList, false, FLinearColor(), false, 0, true, 0, SceneContext.GetSceneDepthSurface()->GetSizeXY(), FIntRect());
 			}
 
@@ -989,7 +1018,6 @@ void FRCPassPostProcessDeferredDecals::Process(FRenderingCompositePassContext& C
 			RenderTargetManager.ResolveTargets();
 		}
 
-
 		if (CurrentStage == DRS_BeforeBasePass && bNeedsDBufferTargets)
 		{
 			// before BasePass
@@ -1006,6 +1034,8 @@ void FRCPassPostProcessDeferredDecals::Process(FRenderingCompositePassContext& C
 		// This allows us to ignore the DBufferMask on frames without decals, without having to explicitly clear the texture.
 		SceneContext.DBufferMask = nullptr;
 	}
+
+	check(!RHICmdList.IsInsideRenderPass());
 }
 
 FPooledRenderTargetDesc FRCPassPostProcessDeferredDecals::ComputeOutputDesc(EPassOutputId InPassOutputId) const
@@ -1090,6 +1120,15 @@ FDecalRenderTargetManager::FDecalRenderTargetManager(FRHICommandList& InRHICmdLi
 
 void FDecalRenderTargetManager::SetRenderTargetMode(FDecalRenderingCommon::ERenderTargetMode CurrentRenderTargetMode, bool bHasNormal, bool bPerPixelDBufferMask)
 {
+	// There are several situations where we do not have a renderpass active when we get here.
+	// The first decal or mesh to draw, etc.
+	if (RHICmdList.IsInsideRenderPass())
+	{
+		RHICmdList.EndRenderPass();
+	}
+
+	check(!RHICmdList.IsInsideRenderPass());
+
 	FSceneRenderTargets& SceneContext = FSceneRenderTargets::Get(RHICmdList);
 
 	// If GBufferA was resolved for read, and we want to write to it again.
@@ -1101,39 +1140,51 @@ void FDecalRenderTargetManager::SetRenderTargetMode(FDecalRenderingCommon::ERend
 		RHICmdList.TransitionResource(EResourceTransitionAccess::EWritable, SceneContext.GBufferA->GetRenderTargetItem().TargetableTexture);
 	}
 
-	// @todo Workaround Vulkan requires we pointlessly bind into RT1 even though we don't write to it,
-	// otherwise the writes to RT2 and RT3 go haywire. This isn't really possible to fix lower down the stack.
-	const bool bRequiresDummyRenderTarget = IsVulkanPlatform(GMaxRHIShaderPlatform);
+	FExclusiveDepthStencil DepthStencilAccess = FExclusiveDepthStencil::DepthRead_StencilWrite;
+	ERenderTargetActions DepthTargetActions = ERenderTargetActions::Load_DontStore;
+	uint32 NumColorTargets = 1;
+	FTextureRHIParamRef* TargetsToBind = &TargetsToResolve[0];
 
+	// The SceneColorAndGBuffer modes do not actually need GBufferA bound when there's no normal.
+	// The apis based on renderpasses fill fail to actually bind anything past a null entry in their RT list so we have to bind it anyway.
+	// DX11 will drop writes to an unbound RT so it worked there.
 	switch (CurrentRenderTargetMode)
 	{
 	case FDecalRenderingCommon::RTM_SceneColorAndGBufferWithNormal:
 	case FDecalRenderingCommon::RTM_SceneColorAndGBufferNoNormal:
 		TargetsToResolve[SceneColorIndex] = SceneContext.GetSceneColor()->GetRenderTargetItem().TargetableTexture;
-		TargetsToResolve[GBufferAIndex] = bHasNormal ? SceneContext.GBufferA->GetRenderTargetItem().TargetableTexture : (bRequiresDummyRenderTarget ? SceneContext.GBufferB->GetRenderTargetItem().TargetableTexture : nullptr);
+		TargetsToResolve[GBufferAIndex] = SceneContext.GBufferA->GetRenderTargetItem().TargetableTexture;
+		//TargetsToResolve[GBufferAIndex] = bHasNormal ? SceneContext.GBufferA->GetRenderTargetItem().TargetableTexture : nullptr;
 		TargetsToResolve[GBufferBIndex] = SceneContext.GBufferB->GetRenderTargetItem().TargetableTexture;
 		TargetsToResolve[GBufferCIndex] = SceneContext.GBufferC->GetRenderTargetItem().TargetableTexture;
-		SetRenderTargets(RHICmdList, 4, TargetsToResolve, SceneContext.GetSceneDepthSurface(), ESimpleRenderTargetMode::EExistingColorAndDepth, FExclusiveDepthStencil::DepthRead_StencilWrite, TargetsToTransitionWritable[CurrentRenderTargetMode]);
+
+		NumColorTargets = 4;
 		break;
 
 	case FDecalRenderingCommon::RTM_SceneColorAndGBufferDepthWriteWithNormal:
 	case FDecalRenderingCommon::RTM_SceneColorAndGBufferDepthWriteNoNormal:
 		TargetsToResolve[SceneColorIndex] = SceneContext.GetSceneColor()->GetRenderTargetItem().TargetableTexture;
-		TargetsToResolve[GBufferAIndex] = bHasNormal ? SceneContext.GBufferA->GetRenderTargetItem().TargetableTexture : (bRequiresDummyRenderTarget ? SceneContext.GBufferB->GetRenderTargetItem().TargetableTexture : nullptr);
+		TargetsToResolve[GBufferAIndex] = SceneContext.GBufferA->GetRenderTargetItem().TargetableTexture;
+		//TargetsToResolve[GBufferAIndex] = bHasNormal ? SceneContext.GBufferA->GetRenderTargetItem().TargetableTexture : nullptr;
 		TargetsToResolve[GBufferBIndex] = SceneContext.GBufferB->GetRenderTargetItem().TargetableTexture;
 		TargetsToResolve[GBufferCIndex] = SceneContext.GBufferC->GetRenderTargetItem().TargetableTexture;
 		TargetsToResolve[GBufferEIndex] = SceneContext.GBufferE->GetRenderTargetItem().TargetableTexture;
-		SetRenderTargets(RHICmdList, 5, TargetsToResolve, SceneContext.GetSceneDepthSurface(), ESimpleRenderTargetMode::EExistingColorAndDepth, FExclusiveDepthStencil::DepthWrite_StencilWrite, TargetsToTransitionWritable[CurrentRenderTargetMode]);
+
+		NumColorTargets = 5;
+		DepthStencilAccess = FExclusiveDepthStencil::DepthWrite_StencilWrite;
+		DepthTargetActions = ERenderTargetActions::Load_Store;
 		break;
 
 	case FDecalRenderingCommon::RTM_GBufferNormal:
 		TargetsToResolve[GBufferAIndex] = SceneContext.GBufferA->GetRenderTargetItem().TargetableTexture;
-		SetRenderTarget(RHICmdList, TargetsToResolve[GBufferAIndex], SceneContext.GetSceneDepthSurface(), ESimpleRenderTargetMode::EExistingColorAndDepth, FExclusiveDepthStencil::DepthRead_StencilWrite, TargetsToTransitionWritable[CurrentRenderTargetMode]);
+
+		TargetsToBind = &TargetsToResolve[GBufferAIndex];
 		break;
 
 	case FDecalRenderingCommon::RTM_SceneColor:
 		TargetsToResolve[SceneColorIndex] = SceneContext.GetSceneColor()->GetRenderTargetItem().TargetableTexture;
-		SetRenderTarget(RHICmdList, TargetsToResolve[SceneColorIndex], SceneContext.GetSceneDepthSurface(), ESimpleRenderTargetMode::EExistingColorAndDepth, FExclusiveDepthStencil::DepthRead_StencilWrite, TargetsToTransitionWritable[CurrentRenderTargetMode]);
+
+		TargetsToBind = &TargetsToResolve[SceneColorIndex];
 		break;
 
 	case FDecalRenderingCommon::RTM_DBuffer:
@@ -1141,15 +1192,15 @@ void FDecalRenderTargetManager::SetRenderTargetMode(FDecalRenderingCommon::ERend
 		TargetsToResolve[DBufferAIndex] = SceneContext.DBufferA->GetRenderTargetItem().TargetableTexture;
 		TargetsToResolve[DBufferBIndex] = SceneContext.DBufferB->GetRenderTargetItem().TargetableTexture;
 		TargetsToResolve[DBufferCIndex] = SceneContext.DBufferC->GetRenderTargetItem().TargetableTexture;
-		uint32 RTCount = 3;
+		NumColorTargets = 3;
 
 		if (bPerPixelDBufferMask)
 		{
 			TargetsToResolve[DBufferMaskIndex] = SceneContext.DBufferMask->GetRenderTargetItem().TargetableTexture;
-			RTCount = 4;
+			NumColorTargets = 4;
 		}
 
-		SetRenderTargets(RHICmdList, RTCount, &TargetsToResolve[DBufferAIndex], SceneContext.GetSceneDepthTexture(), ESimpleRenderTargetMode::EExistingColorAndDepth, FExclusiveDepthStencil::DepthRead_StencilWrite, TargetsToTransitionWritable[CurrentRenderTargetMode]);
+		TargetsToBind = &TargetsToResolve[DBufferAIndex];
 		break;
 	}
 
@@ -1157,8 +1208,9 @@ void FDecalRenderTargetManager::SetRenderTargetMode(FDecalRenderingCommon::ERend
 	{
 		TargetsToResolve[SceneColorIndex] = SceneContext.ScreenSpaceAO->GetRenderTargetItem().TargetableTexture;
 		RHICmdList.TransitionResource(EResourceTransitionAccess::EWritable, TargetsToResolve[SceneColorIndex]);
-		ESimpleRenderTargetMode RTMode = SceneContext.bScreenSpaceAOIsValid ? ESimpleRenderTargetMode::EExistingColorAndDepth : ESimpleRenderTargetMode::EClearColorExistingDepth;
-		SetRenderTarget(RHICmdList, TargetsToResolve[SceneColorIndex], SceneContext.GetSceneDepthSurface(), ESimpleRenderTargetMode::EExistingColorAndDepth, FExclusiveDepthStencil::DepthRead_StencilWrite, TargetsToTransitionWritable[CurrentRenderTargetMode]);
+
+		TargetsToBind = &TargetsToResolve[SceneColorIndex];
+
 		SceneContext.bScreenSpaceAOIsValid = true;
 		break;
 	}
@@ -1167,6 +1219,18 @@ void FDecalRenderTargetManager::SetRenderTargetMode(FDecalRenderingCommon::ERend
 		check(0);
 		break;
 	}
+
+	FRHIRenderPassInfo RPInfo(NumColorTargets, TargetsToBind, ERenderTargetActions::Load_Store);
+	RPInfo.DepthStencilRenderTarget.DepthStencilTarget = SceneContext.GetSceneDepthSurface();
+	RPInfo.DepthStencilRenderTarget.Action = MakeDepthStencilTargetActions(DepthTargetActions, ERenderTargetActions::Load_Store);
+	RPInfo.DepthStencilRenderTarget.ExclusiveDepthStencil = DepthStencilAccess;
+
+	if (TargetsToTransitionWritable[CurrentRenderTargetMode])
+	{
+		TransitionRenderPassTargets(RHICmdList, RPInfo);
+	}
+	RHICmdList.BeginRenderPass(RPInfo, TEXT("DecalPass"));
+
 	TargetsToTransitionWritable[CurrentRenderTargetMode] = false;
 }
 
