This needs to have DepthRead tagged on the FExclusiveDepthStencil struct. Which is kind of weird in this situation. Will have to further investigate. For now this fixes TAA.

diff --git a/Engine/Source/Runtime/Renderer/Private/PostProcess/PostProcessTemporalAA.cpp b/Engine/Source/Runtime/Renderer/Private/PostProcess/PostProcessTemporalAA.cpp
--- a/Engine/Source/Runtime/Renderer/Private/PostProcess/PostProcessTemporalAA.cpp
+++ b/Engine/Source/Runtime/Renderer/Private/PostProcess/PostProcessTemporalAA.cpp
@@ -889,7 +889,7 @@ void FRCPassPostProcessTemporalAA::Process(FRenderingCompositePassContext& Conte
 		RPInfo.DepthStencilRenderTarget.DepthStencilTarget = SceneContext.GetSceneDepthTexture();
 		RPInfo.DepthStencilRenderTarget.ResolveTarget = nullptr;
 		RPInfo.DepthStencilRenderTarget.Action = MakeDepthStencilTargetActions(ERenderTargetActions::DontLoad_DontStore, ERenderTargetActions::Load_Store);
-		RPInfo.DepthStencilRenderTarget.ExclusiveDepthStencil = FExclusiveDepthStencil::DepthNop_StencilWrite;
+		RPInfo.DepthStencilRenderTarget.ExclusiveDepthStencil = FExclusiveDepthStencil::DepthRead_StencilWrite;
 
 		Context.RHICmdList.BeginRenderPass(RPInfo, TEXT("TemporalAA"));
 		{
