Can't have Stencil Store and StencilNop on the depthstencil attachment. Fixes Metal decals.

diff --git a/Engine/Source/Runtime/Renderer/Private/CompositionLighting/PostProcessDeferredDecals.cpp b/Engine/Source/Runtime/Renderer/Private/CompositionLighting/PostProcessDeferredDecals.cpp
--- a/Engine/Source/Runtime/Renderer/Private/CompositionLighting/PostProcessDeferredDecals.cpp
+++ b/Engine/Source/Runtime/Renderer/Private/CompositionLighting/PostProcessDeferredDecals.cpp
@@ -771,7 +771,7 @@ void FRCPassPostProcessDeferredDecals::Process(FRenderingCompositePassContext& C
 				FRHIRenderPassInfo RPInfo(RTCount, RenderTargets, ERenderTargetActions::Clear_Store);
 				RPInfo.DepthStencilRenderTarget.DepthStencilTarget = SceneContext.GetSceneDepthTexture();
 				RPInfo.DepthStencilRenderTarget.Action = MakeDepthStencilTargetActions(ERenderTargetActions::Load_DontStore, ERenderTargetActions::Load_Store);
-				RPInfo.DepthStencilRenderTarget.ExclusiveDepthStencil = FExclusiveDepthStencil::DepthRead_StencilNop;
+				RPInfo.DepthStencilRenderTarget.ExclusiveDepthStencil = FExclusiveDepthStencil::DepthRead_StencilWrite;
 				RHICmdList.BeginRenderPass(RPInfo, TEXT("InitialDeferredDecals"));
 				bDidClearDBuffer = true;
 			}
