Capsule Shadows converted to renderpasses.

diff --git a/Engine/Source/Runtime/Renderer/Private/CapsuleShadowRendering.cpp b/Engine/Source/Runtime/Renderer/Private/CapsuleShadowRendering.cpp
--- a/Engine/Source/Runtime/Renderer/Private/CapsuleShadowRendering.cpp
+++ b/Engine/Source/Runtime/Renderer/Private/CapsuleShadowRendering.cpp
@@ -793,6 +793,8 @@ bool FDeferredShadingSceneRenderer::RenderCapsuleDirectShadows(
 	const TArray<FProjectedShadowInfo*, SceneRenderingAllocator>& CapsuleShadows, 
 	bool bProjectingForForwardShading) const
 {
+	check(RHICmdList.IsOutsideRenderPass());
+
 	bool bAllViewsHaveViewState = true;
 
 	for (int32 ViewIndex = 0; ViewIndex < Views.Num(); ViewIndex++)
@@ -874,7 +876,8 @@ bool FDeferredShadingSceneRenderer::RenderCapsuleDirectShadows(
 				FPlatformMemory::Memcpy(CapsuleShapeLockedData, CapsuleShapeData.GetData(), DataSize);
 				RHIUnlockVertexBuffer(LightSceneInfo.ShadowCapsuleShapesVertexBuffer);
 
-				SetRenderTarget(RHICmdList, NULL, NULL);
+				// #todo-renderpasses Remove once everything is converted to renderpasses
+				UnbindRenderTargets(RHICmdList);
 
 				const bool bDirectionalLight = LightSceneInfo.Proxy->GetLightType() == LightType_Directional;
 				FIntRect ScissorRect;
@@ -952,11 +955,17 @@ bool FDeferredShadingSceneRenderer::RenderCapsuleDirectShadows(
 					}
 				}
 
+				FRHIRenderPassInfo RPInfo(	ScreenShadowMaskTexture->GetRenderTargetItem().TargetableTexture, 
+											ERenderTargetActions::Load_Store,
+											FSceneRenderTargets::Get(RHICmdList).GetSceneDepthSurface(), 
+											MakeDepthStencilTargetActions(ERenderTargetActions::Load_DontStore, ERenderTargetActions::Load_Store), 
+											FExclusiveDepthStencil::DepthRead_StencilWrite);
+
+				TransitionRenderPassTargets(RHICmdList, RPInfo);
+				RHICmdList.BeginRenderPass(RPInfo, TEXT("UpsampleCapsuleShadow"));
 				{
 					SCOPED_DRAW_EVENTF(RHICmdList, Upsample, TEXT("Upsample %dx%d"),
 						ScissorRect.Width(), ScissorRect.Height());
-						
-					SetRenderTarget(RHICmdList, ScreenShadowMaskTexture->GetRenderTargetItem().TargetableTexture, FSceneRenderTargets::Get(RHICmdList).GetSceneDepthSurface(), ESimpleRenderTargetMode::EExistingColorAndDepth, FExclusiveDepthStencil::DepthRead_StencilWrite, true);
 
 					FGraphicsPipelineStateInitializer GraphicsPSOInit;
 					RHICmdList.ApplyCachedRenderTargets(GraphicsPSOInit);
@@ -1006,6 +1015,7 @@ bool FDeferredShadingSceneRenderer::RenderCapsuleDirectShadows(
 					RHICmdList.SetStreamSource(0, GTileTexCoordVertexBuffer.VertexBufferRHI, 0);
 					RHICmdList.DrawIndexedPrimitive(GTileIndexBuffer.IndexBufferRHI, 0, 0, 4, 0, 2 * NumTileQuadsInBuffer, FMath::DivideAndRoundUp(GroupSize.X * GroupSize.Y, NumTileQuadsInBuffer));
 				}
+				RHICmdList.EndRenderPass();
 			}
 		}
 
@@ -1286,6 +1296,8 @@ void FDeferredShadingSceneRenderer::RenderIndirectCapsuleShadows(
 	FTextureRHIParamRef IndirectLightingTexture, 
 	FTextureRHIParamRef ExistingIndirectOcclusionTexture) const
 {
+	check(RHICmdList.IsOutsideRenderPass());
+
 	if (SupportsCapsuleIndirectShadows(FeatureLevel, GShaderPlatformForFeatureLevel[FeatureLevel])
 		&& ViewFamily.EngineShowFlags.DynamicShadows
 		&& ViewFamily.EngineShowFlags.CapsuleShadows
@@ -1338,11 +1350,26 @@ void FDeferredShadingSceneRenderer::RenderIndirectCapsuleShadows(
 
 				SCOPED_DRAW_EVENT(RHICmdList, ClearIndirectOcclusion);
 				// We are the first users of the indirect occlusion texture so we must clear to unoccluded
-				FRHIRenderTargetView RtView = FRHIRenderTargetView(SceneContext.ScreenSpaceAO->GetRenderTargetItem().TargetableTexture, ERenderTargetLoadAction::EClear);
-				FRHISetRenderTargetsInfo Info(1, &RtView, FRHIDepthRenderTargetView());
-				RHICmdList.SetRenderTargetsAndClear(Info);
+				{
+					FRHIRenderPassInfo RPInfo(SceneContext.ScreenSpaceAO->GetRenderTargetItem().TargetableTexture, ERenderTargetActions::Clear_Store);
+					RHICmdList.BeginRenderPass(RPInfo, TEXT("IndirectCapsuleShadowsClearAO"));
+					RHICmdList.EndRenderPass();
+				}
 
-				SetRenderTargets(RHICmdList, RenderTargets.Num(), RenderTargets.GetData(), FTextureRHIParamRef(), 0, NULL, true);
+				// #todo is this transition really needed?
+				{
+					FTextureRHIParamRef Transitions[MaxSimultaneousRenderTargets + 1];
+					int32 TransitionIndex = 0;
+					for (int32 Index = 0; Index < RenderTargets.Num(); Index++)
+					{
+						if (RenderTargets.GetData()[Index])
+						{
+							Transitions[TransitionIndex] = RenderTargets.GetData()[Index];
+							++TransitionIndex;
+						}
+					}
+					RHICmdList.TransitionResources(EResourceTransitionAccess::EWritable, Transitions, TransitionIndex);
+				}
 			}
 							
 			check(RenderTargets.Num() > 0);
@@ -1365,7 +1392,8 @@ void FDeferredShadingSceneRenderer::RenderIndirectCapsuleShadows(
 					{
 						check(IndirectShadowLightDirectionSRV);
 
-						SetRenderTarget(RHICmdList, NULL, NULL);
+						// #todo-renderpasses remove once everything is converted to renderpasses
+						UnbindRenderTargets(RHICmdList);
 
 						FIntRect ScissorRect = View.ViewRect;
 
@@ -1429,11 +1457,12 @@ void FDeferredShadingSceneRenderer::RenderIndirectCapsuleShadows(
 			
 						RHICmdList.TransitionResource(EResourceTransitionAccess::EReadable, EResourceTransitionPipeline::EComputeToGfx, RayTracedShadowsRT->GetRenderTargetItem().UAV);
 
+						FRHIRenderPassInfo RPInfo(RenderTargets.Num(), RenderTargets.GetData(), ERenderTargetActions::Load_Store);
+						TransitionRenderPassTargets(RHICmdList, RPInfo);
+						RHICmdList.BeginRenderPass(RPInfo, TEXT("UpsampleIndirectCapsuleShadows"));
 						{
 							SCOPED_DRAW_EVENTF(RHICmdList, Upsample, TEXT("Upsample %dx%d"), ScissorRect.Width(), ScissorRect.Height());
 
-							SetRenderTargets(RHICmdList, RenderTargets.Num(), RenderTargets.GetData(), FTextureRHIParamRef(), 0, NULL, true);
-
 							FGraphicsPipelineStateInitializer GraphicsPSOInit;
 							RHICmdList.ApplyCachedRenderTargets(GraphicsPSOInit);
 
@@ -1501,6 +1530,7 @@ void FDeferredShadingSceneRenderer::RenderIndirectCapsuleShadows(
 							RHICmdList.SetStreamSource(0, GTileTexCoordVertexBuffer.VertexBufferRHI, 0);
 							RHICmdList.DrawIndexedPrimitive(GTileIndexBuffer.IndexBufferRHI, 0, 0, 4, 0, 2 * NumTileQuadsInBuffer, FMath::DivideAndRoundUp(GroupSize.X * GroupSize.Y, NumTileQuadsInBuffer));
 						}
+						RHICmdList.EndRenderPass();
 					}
 				}
 			}
@@ -1541,6 +1571,8 @@ bool FDeferredShadingSceneRenderer::ShouldPrepareForDFInsetIndirectShadow() cons
 
 void FDeferredShadingSceneRenderer::RenderCapsuleShadowsForMovableSkylight(FRHICommandListImmediate& RHICmdList, TRefCountPtr<IPooledRenderTarget>& BentNormalOutput) const
 {
+	check(RHICmdList.IsOutsideRenderPass());
+
 	if (SupportsCapsuleIndirectShadows(FeatureLevel, GShaderPlatformForFeatureLevel[FeatureLevel])
 		&& ViewFamily.EngineShowFlags.CapsuleShadows)
 	{
@@ -1586,7 +1618,9 @@ void FDeferredShadingSceneRenderer::RenderCapsuleShadowsForMovableSkylight(FRHIC
 					if (NumCapsuleShapes > 0 || NumMeshDistanceFieldCasters > 0)
 					{
 						check(IndirectShadowLightDirectionSRV);
-						SetRenderTarget(RHICmdList, NULL, NULL);
+
+						// #todo-renderpasses remove once everything is converted to renderpasses
+						UnbindRenderTargets(RHICmdList);
 
 						FIntRect ScissorRect = View.ViewRect;
 
