Renderpass updates.

SceneOcclusion updates. Self-contained.

diff --git a/Engine/Source/Runtime/Renderer/Private/SceneOcclusion.cpp b/Engine/Source/Runtime/Renderer/Private/SceneOcclusion.cpp
--- a/Engine/Source/Runtime/Renderer/Private/SceneOcclusion.cpp
+++ b/Engine/Source/Runtime/Renderer/Private/SceneOcclusion.cpp
@@ -1010,39 +1010,42 @@ void FHZBOcclusionTester::Submit(FRHICommandListImmediate& RHICmdList, const FVi
 	{
 		SCOPED_DRAW_EVENT(RHICmdList, TestHZB);
 
-		SetRenderTarget(RHICmdList, ResultsTextureGPU->GetRenderTargetItem().TargetableTexture, NULL);
-
-		FGraphicsPipelineStateInitializer GraphicsPSOInit;
-		RHICmdList.ApplyCachedRenderTargets(GraphicsPSOInit);
-		GraphicsPSOInit.BlendState = TStaticBlendState<>::GetRHI();
-		GraphicsPSOInit.RasterizerState = TStaticRasterizerState<>::GetRHI();
-		GraphicsPSOInit.DepthStencilState = TStaticDepthStencilState<false, CF_Always>::GetRHI();
-
-		TShaderMapRef< FScreenVS >	VertexShader(View.ShaderMap);
-		TShaderMapRef< FHZBTestPS >	PixelShader(View.ShaderMap);
-
-		GraphicsPSOInit.BoundShaderState.VertexDeclarationRHI = GFilterVertexDeclaration.VertexDeclarationRHI;
-		GraphicsPSOInit.BoundShaderState.VertexShaderRHI = GETSAFERHISHADER_VERTEX(*VertexShader);
-		GraphicsPSOInit.BoundShaderState.PixelShaderRHI = GETSAFERHISHADER_PIXEL(*PixelShader);
-		GraphicsPSOInit.PrimitiveType = PT_TriangleList;
-
-		SetGraphicsPipelineState(RHICmdList, GraphicsPSOInit);
-
-		PixelShader->SetParameters(RHICmdList, View, BoundsCenterTexture->GetRenderTargetItem().ShaderResourceTexture, BoundsExtentTexture->GetRenderTargetItem().ShaderResourceTexture );
-
-		RHICmdList.SetViewport(0, 0, 0.0f, SizeX, SizeY, 1.0f);
-
-		// TODO draw quads covering blocks added above
-		DrawRectangle(
-			RHICmdList,
-			0, 0,
-			SizeX, SizeY,
-			0, 0,
-			SizeX, SizeY,
-			FIntPoint( SizeX, SizeY ),
-			FIntPoint( SizeX, SizeY ),
-			*VertexShader,
-			EDRF_UseTriangleOptimization);
+		FRHIRenderPassInfo RPInfo(ResultsTextureGPU->GetRenderTargetItem().TargetableTexture, ERenderTargetActions::Load_Store);
+		RHICmdList.BeginRenderPass(RPInfo, TEXT("TestHZB"));
+		{
+			FGraphicsPipelineStateInitializer GraphicsPSOInit;
+			RHICmdList.ApplyCachedRenderTargets(GraphicsPSOInit);
+			GraphicsPSOInit.BlendState = TStaticBlendState<>::GetRHI();
+			GraphicsPSOInit.RasterizerState = TStaticRasterizerState<>::GetRHI();
+			GraphicsPSOInit.DepthStencilState = TStaticDepthStencilState<false, CF_Always>::GetRHI();
+
+			TShaderMapRef< FScreenVS >	VertexShader(View.ShaderMap);
+			TShaderMapRef< FHZBTestPS >	PixelShader(View.ShaderMap);
+
+			GraphicsPSOInit.BoundShaderState.VertexDeclarationRHI = GFilterVertexDeclaration.VertexDeclarationRHI;
+			GraphicsPSOInit.BoundShaderState.VertexShaderRHI = GETSAFERHISHADER_VERTEX(*VertexShader);
+			GraphicsPSOInit.BoundShaderState.PixelShaderRHI = GETSAFERHISHADER_PIXEL(*PixelShader);
+			GraphicsPSOInit.PrimitiveType = PT_TriangleList;
+
+			SetGraphicsPipelineState(RHICmdList, GraphicsPSOInit);
+
+			PixelShader->SetParameters(RHICmdList, View, BoundsCenterTexture->GetRenderTargetItem().ShaderResourceTexture, BoundsExtentTexture->GetRenderTargetItem().ShaderResourceTexture);
+
+			RHICmdList.SetViewport(0, 0, 0.0f, SizeX, SizeY, 1.0f);
+
+			// TODO draw quads covering blocks added above
+			DrawRectangle(
+				RHICmdList,
+				0, 0,
+				SizeX, SizeY,
+				0, 0,
+				SizeX, SizeY,
+				FIntPoint(SizeX, SizeY),
+				FIntPoint(SizeX, SizeY),
+				*VertexShader,
+				EDRF_UseTriangleOptimization);
+		}
+		RHICmdList.EndRenderPass();
 	}
 
 	GVisualizeTexture.SetCheckPoint(RHICmdList, ResultsTextureGPU);
@@ -1617,6 +1620,7 @@ void FSceneRenderer::BeginOcclusionTests(FRHICommandListImmediate& RHICmdList, b
 			if (bUseDownsampledDepth)
 			{
 				// Restore default render target
+				// #todo-renderpasses this is not ideal. This pass should be self-contained. Can we refactor this?
 				SceneContext.BeginRenderingSceneColor(RHICmdList, ESimpleRenderTargetMode::EUninitializedColorExistingDepth, FExclusiveDepthStencil::DepthRead_StencilWrite);
 			}
 		}
