Renderpass updates.

Heightfield lighting. Self-contained.

diff --git a/Engine/Source/Runtime/Renderer/Private/HeightfieldLighting.cpp b/Engine/Source/Runtime/Renderer/Private/HeightfieldLighting.cpp
--- a/Engine/Source/Runtime/Renderer/Private/HeightfieldLighting.cpp
+++ b/Engine/Source/Runtime/Renderer/Private/HeightfieldLighting.cpp
@@ -472,40 +472,44 @@ void FHeightfieldLightingViewInfo::SetupVisibleHeightfields(const FViewInfo& Vie
 					};
 
 					RHICmdList.SetViewport(0, 0, 0.0f, LightingAtlasSize.X, LightingAtlasSize.Y, 1.0f);
-					SetRenderTargets(RHICmdList, ARRAY_COUNT(RenderTargets), RenderTargets, FTextureRHIParamRef(), ESimpleRenderTargetMode::EClearColorExistingDepth, FExclusiveDepthStencil::DepthRead_StencilRead);
-					
-					FGraphicsPipelineStateInitializer GraphicsPSOInit;
-					RHICmdList.ApplyCachedRenderTargets(GraphicsPSOInit);
 
-					GraphicsPSOInit.RasterizerState = TStaticRasterizerState<FM_Solid, CM_None>::GetRHI();
-					GraphicsPSOInit.DepthStencilState = TStaticDepthStencilState<false, CF_Always>::GetRHI();
-					GraphicsPSOInit.BlendState = TStaticBlendState<>::GetRHI();
-					RHICmdList.SetStreamSource(0, GQuadVertexBuffer.VertexBufferRHI, 0);
+					FRHIRenderPassInfo RPInfo(ARRAY_COUNT(RenderTargets), RenderTargets, ERenderTargetActions::Clear_Store);
+					RHICmdList.BeginRenderPass(RPInfo, TEXT("SetupHeightfields"));
+					{
+						FGraphicsPipelineStateInitializer GraphicsPSOInit;
+						RHICmdList.ApplyCachedRenderTargets(GraphicsPSOInit);
 
-					TShaderMapRef<FHeightfieldSubsectionQuadVS> VertexShader(View.ShaderMap);
-					TShaderMapRef<FInitializeHeightfieldsPS> PixelShader(View.ShaderMap);
+						GraphicsPSOInit.RasterizerState = TStaticRasterizerState<FM_Solid, CM_None>::GetRHI();
+						GraphicsPSOInit.DepthStencilState = TStaticDepthStencilState<false, CF_Always>::GetRHI();
+						GraphicsPSOInit.BlendState = TStaticBlendState<>::GetRHI();
+						RHICmdList.SetStreamSource(0, GQuadVertexBuffer.VertexBufferRHI, 0);
 
-					GraphicsPSOInit.BoundShaderState.VertexDeclarationRHI = GScreenVertexDeclaration.VertexDeclarationRHI;
-					GraphicsPSOInit.BoundShaderState.VertexShaderRHI = GETSAFERHISHADER_VERTEX(*VertexShader);
-					GraphicsPSOInit.BoundShaderState.PixelShaderRHI = GETSAFERHISHADER_PIXEL(*PixelShader);
-					GraphicsPSOInit.PrimitiveType = PT_TriangleList;
+						TShaderMapRef<FHeightfieldSubsectionQuadVS> VertexShader(View.ShaderMap);
+						TShaderMapRef<FInitializeHeightfieldsPS> PixelShader(View.ShaderMap);
 
-					SetGraphicsPipelineState(RHICmdList, GraphicsPSOInit);
+						GraphicsPSOInit.BoundShaderState.VertexDeclarationRHI = GScreenVertexDeclaration.VertexDeclarationRHI;
+						GraphicsPSOInit.BoundShaderState.VertexShaderRHI = GETSAFERHISHADER_VERTEX(*VertexShader);
+						GraphicsPSOInit.BoundShaderState.PixelShaderRHI = GETSAFERHISHADER_PIXEL(*PixelShader);
+						GraphicsPSOInit.PrimitiveType = PT_TriangleList;
 
-					for (TMap<FHeightfieldComponentTextures, TArray<FHeightfieldComponentDescription>>::TConstIterator It(Heightfield.ComponentDescriptions); It; ++It)
-					{
-						const TArray<FHeightfieldComponentDescription>& HeightfieldDescriptions = It.Value();
+						SetGraphicsPipelineState(RHICmdList, GraphicsPSOInit);
 
-						if (HeightfieldDescriptions.Num() > 0)
+						for (TMap<FHeightfieldComponentTextures, TArray<FHeightfieldComponentDescription>>::TConstIterator It(Heightfield.ComponentDescriptions); It; ++It)
 						{
-							const int32 NumQuads = UploadSubsectionHeightfieldDescriptions(HeightfieldDescriptions, InvLightingAtlasSize, 1.0f / Heightfield.DownsampleFactor);
+							const TArray<FHeightfieldComponentDescription>& HeightfieldDescriptions = It.Value();
 
-							VertexShader->SetParameters(RHICmdList, View);
-							PixelShader->SetParameters(RHICmdList, View, It.Key().HeightAndNormal, It.Key().DiffuseColor);
+							if (HeightfieldDescriptions.Num() > 0)
+							{
+								const int32 NumQuads = UploadSubsectionHeightfieldDescriptions(HeightfieldDescriptions, InvLightingAtlasSize, 1.0f / Heightfield.DownsampleFactor);
 
-							RHICmdList.DrawPrimitive(0, 2, NumQuads);
+								VertexShader->SetParameters(RHICmdList, View);
+								PixelShader->SetParameters(RHICmdList, View, It.Key().HeightAndNormal, It.Key().DiffuseColor);
+
+								RHICmdList.DrawPrimitive(0, 2, NumQuads);
+							}
 						}
 					}
+					RHICmdList.EndRenderPass();
 
 					RHICmdList.TransitionResources(EResourceTransitionAccess::EReadable, RenderTargets, ARRAY_COUNT(RenderTargets));
 				}
@@ -811,7 +815,10 @@ void FHeightfieldLightingViewInfo::ClearShadowing(const FViewInfo& View, FRHICom
 	{
 		FSceneViewState* ViewState = (FSceneViewState*)View.State;
 		const FHeightfieldLightingAtlas& Atlas = *ViewState->HeightfieldLightingAtlas;
-		SetRenderTarget(RHICmdList, Atlas.DirectionalLightShadowing->GetRenderTargetItem().TargetableTexture, NULL, ESimpleRenderTargetMode::EClearColorExistingDepth);		
+
+		FRHIRenderPassInfo RPInfo(Atlas.DirectionalLightShadowing->GetRenderTargetItem().TargetableTexture, ERenderTargetActions::Clear_Store);
+		RHICmdList.BeginRenderPass(RPInfo, TEXT("ClearShadowing"));
+		RHICmdList.EndRenderPass();
 	}
 }
 
@@ -833,43 +840,47 @@ void FHeightfieldLightingViewInfo::ComputeShadowMapShadowing(const FViewInfo& Vi
 
 			const FIntPoint LightingAtlasSize = Atlas.GetAtlasSize();
 			const FVector2D InvLightingAtlasSize(1.0f / LightingAtlasSize.X, 1.0f / LightingAtlasSize.Y);
-			SetRenderTarget(RHICmdList, Atlas.DirectionalLightShadowing->GetRenderTargetItem().TargetableTexture, NULL);
 
-			FGraphicsPipelineStateInitializer GraphicsPSOInit;
-			RHICmdList.ApplyCachedRenderTargets(GraphicsPSOInit);
-
-			RHICmdList.SetViewport(0, 0, 0.0f, LightingAtlasSize.X, LightingAtlasSize.Y, 1.0f);
-			GraphicsPSOInit.RasterizerState = TStaticRasterizerState<FM_Solid, CM_None>::GetRHI();
-			GraphicsPSOInit.DepthStencilState = TStaticDepthStencilState<false, CF_Always>::GetRHI();
-			// Combine with other shadow types with min (ray traced)
-			GraphicsPSOInit.BlendState = TStaticBlendState<CW_RED, BO_Min, BF_One, BF_One>::GetRHI();
-			RHICmdList.SetStreamSource(0, GQuadVertexBuffer.VertexBufferRHI, 0);
+			FRHIRenderPassInfo RPInfo(Atlas.DirectionalLightShadowing->GetRenderTargetItem().TargetableTexture, ERenderTargetActions::Load_Store);
+			RHICmdList.BeginRenderPass(RPInfo, TEXT("ComputeShadowMapShadowing"));
+			{
+				FGraphicsPipelineStateInitializer GraphicsPSOInit;
+				RHICmdList.ApplyCachedRenderTargets(GraphicsPSOInit);
 
-			TShaderMapRef<FHeightfieldComponentQuadVS> VertexShader(View.ShaderMap);
-			TShaderMapRef<FShadowHeightfieldsPS> PixelShader(View.ShaderMap);
+				RHICmdList.SetViewport(0, 0, 0.0f, LightingAtlasSize.X, LightingAtlasSize.Y, 1.0f);
+				GraphicsPSOInit.RasterizerState = TStaticRasterizerState<FM_Solid, CM_None>::GetRHI();
+				GraphicsPSOInit.DepthStencilState = TStaticDepthStencilState<false, CF_Always>::GetRHI();
+				// Combine with other shadow types with min (ray traced)
+				GraphicsPSOInit.BlendState = TStaticBlendState<CW_RED, BO_Min, BF_One, BF_One>::GetRHI();
+				RHICmdList.SetStreamSource(0, GQuadVertexBuffer.VertexBufferRHI, 0);
 
-			GraphicsPSOInit.BoundShaderState.VertexDeclarationRHI = GScreenVertexDeclaration.VertexDeclarationRHI;
-			GraphicsPSOInit.BoundShaderState.VertexShaderRHI = GETSAFERHISHADER_VERTEX(*VertexShader);
-			GraphicsPSOInit.BoundShaderState.PixelShaderRHI = GETSAFERHISHADER_PIXEL(*PixelShader);
-			GraphicsPSOInit.PrimitiveType = PT_TriangleList;
+				TShaderMapRef<FHeightfieldComponentQuadVS> VertexShader(View.ShaderMap);
+				TShaderMapRef<FShadowHeightfieldsPS> PixelShader(View.ShaderMap);
 
-			SetGraphicsPipelineState(RHICmdList, GraphicsPSOInit);
+				GraphicsPSOInit.BoundShaderState.VertexDeclarationRHI = GScreenVertexDeclaration.VertexDeclarationRHI;
+				GraphicsPSOInit.BoundShaderState.VertexShaderRHI = GETSAFERHISHADER_VERTEX(*VertexShader);
+				GraphicsPSOInit.BoundShaderState.PixelShaderRHI = GETSAFERHISHADER_PIXEL(*PixelShader);
+				GraphicsPSOInit.PrimitiveType = PT_TriangleList;
 
-			for (TMap<FHeightfieldComponentTextures, TArray<FHeightfieldComponentDescription>>::TConstIterator It(Heightfield.ComponentDescriptions); It; ++It)
-			{
-				const TArray<FHeightfieldComponentDescription>& HeightfieldDescriptions = It.Value();
+				SetGraphicsPipelineState(RHICmdList, GraphicsPSOInit);
 
-				if (HeightfieldDescriptions.Num() > 0)
+				for (TMap<FHeightfieldComponentTextures, TArray<FHeightfieldComponentDescription>>::TConstIterator It(Heightfield.ComponentDescriptions); It; ++It)
 				{
-					//@todo - cull heightfield tiles with shadow bounds
-					UploadHeightfieldDescriptions(HeightfieldDescriptions, InvLightingAtlasSize, 1.0f / Heightfield.DownsampleFactor);
+					const TArray<FHeightfieldComponentDescription>& HeightfieldDescriptions = It.Value();
 
-					VertexShader->SetParameters(RHICmdList, View, HeightfieldDescriptions.Num());
-					PixelShader->SetParameters(RHICmdList, View, ProjectedShadowInfo, HeightfieldDescriptions.Num(), Atlas);
+					if (HeightfieldDescriptions.Num() > 0)
+					{
+						//@todo - cull heightfield tiles with shadow bounds
+						UploadHeightfieldDescriptions(HeightfieldDescriptions, InvLightingAtlasSize, 1.0f / Heightfield.DownsampleFactor);
 
-					RHICmdList.DrawPrimitive(0, 2, HeightfieldDescriptions.Num());
+						VertexShader->SetParameters(RHICmdList, View, HeightfieldDescriptions.Num());
+						PixelShader->SetParameters(RHICmdList, View, ProjectedShadowInfo, HeightfieldDescriptions.Num(), Atlas);
+
+						RHICmdList.DrawPrimitive(0, 2, HeightfieldDescriptions.Num());
+					}
 				}
 			}
+			RHICmdList.EndRenderPass();
 		}
 	}
 }
@@ -999,42 +1010,46 @@ void FHeightfieldLightingViewInfo::ComputeRayTracedShadowing(
 
 			const FIntPoint LightingAtlasSize = Atlas.GetAtlasSize();
 			const FVector2D InvLightingAtlasSize(1.0f / LightingAtlasSize.X, 1.0f / LightingAtlasSize.Y);
-			SetRenderTarget(RHICmdList, Atlas.DirectionalLightShadowing->GetRenderTargetItem().TargetableTexture, NULL);
 
-			FGraphicsPipelineStateInitializer GraphicsPSOInit;
-			RHICmdList.ApplyCachedRenderTargets(GraphicsPSOInit);
-
-			RHICmdList.SetViewport(0, 0, 0.0f, LightingAtlasSize.X, LightingAtlasSize.Y, 1.0f);
-			GraphicsPSOInit.RasterizerState = TStaticRasterizerState<FM_Solid, CM_None>::GetRHI();
-			GraphicsPSOInit.DepthStencilState = TStaticDepthStencilState<false, CF_Always>::GetRHI();
-			// Combine with other shadow types with min (CSM)
-			GraphicsPSOInit.BlendState = TStaticBlendState<CW_RED, BO_Min, BF_One, BF_One>::GetRHI();
-			RHICmdList.SetStreamSource(0, GQuadVertexBuffer.VertexBufferRHI, 0);
+			FRHIRenderPassInfo RPInfo(Atlas.DirectionalLightShadowing->GetRenderTargetItem().TargetableTexture, ERenderTargetActions::Load_Store);
+			RHICmdList.BeginRenderPass(RPInfo, TEXT("ComputeRayTracedShadowing"));
+			{
+				FGraphicsPipelineStateInitializer GraphicsPSOInit;
+				RHICmdList.ApplyCachedRenderTargets(GraphicsPSOInit);
 
-			TShaderMapRef<FHeightfieldComponentQuadVS> VertexShader(View.ShaderMap);
-			TShaderMapRef<FRayTracedShadowHeightfieldsPS> PixelShader(View.ShaderMap);
+				RHICmdList.SetViewport(0, 0, 0.0f, LightingAtlasSize.X, LightingAtlasSize.Y, 1.0f);
+				GraphicsPSOInit.RasterizerState = TStaticRasterizerState<FM_Solid, CM_None>::GetRHI();
+				GraphicsPSOInit.DepthStencilState = TStaticDepthStencilState<false, CF_Always>::GetRHI();
+				// Combine with other shadow types with min (CSM)
+				GraphicsPSOInit.BlendState = TStaticBlendState<CW_RED, BO_Min, BF_One, BF_One>::GetRHI();
+				RHICmdList.SetStreamSource(0, GQuadVertexBuffer.VertexBufferRHI, 0);
 
-			GraphicsPSOInit.BoundShaderState.VertexDeclarationRHI = GScreenVertexDeclaration.VertexDeclarationRHI;
-			GraphicsPSOInit.BoundShaderState.VertexShaderRHI = GETSAFERHISHADER_VERTEX(*VertexShader);
-			GraphicsPSOInit.BoundShaderState.PixelShaderRHI = GETSAFERHISHADER_PIXEL(*PixelShader);
-			GraphicsPSOInit.PrimitiveType = PT_TriangleList;
+				TShaderMapRef<FHeightfieldComponentQuadVS> VertexShader(View.ShaderMap);
+				TShaderMapRef<FRayTracedShadowHeightfieldsPS> PixelShader(View.ShaderMap);
 
-			SetGraphicsPipelineState(RHICmdList, GraphicsPSOInit);
+				GraphicsPSOInit.BoundShaderState.VertexDeclarationRHI = GScreenVertexDeclaration.VertexDeclarationRHI;
+				GraphicsPSOInit.BoundShaderState.VertexShaderRHI = GETSAFERHISHADER_VERTEX(*VertexShader);
+				GraphicsPSOInit.BoundShaderState.PixelShaderRHI = GETSAFERHISHADER_PIXEL(*PixelShader);
+				GraphicsPSOInit.PrimitiveType = PT_TriangleList;
 
-			for (TMap<FHeightfieldComponentTextures, TArray<FHeightfieldComponentDescription>>::TConstIterator It(Heightfield.ComponentDescriptions); It; ++It)
-			{
-				const TArray<FHeightfieldComponentDescription>& HeightfieldDescriptions = It.Value();
+				SetGraphicsPipelineState(RHICmdList, GraphicsPSOInit);
 
-				if (HeightfieldDescriptions.Num() > 0)
+				for (TMap<FHeightfieldComponentTextures, TArray<FHeightfieldComponentDescription>>::TConstIterator It(Heightfield.ComponentDescriptions); It; ++It)
 				{
-					UploadHeightfieldDescriptions(HeightfieldDescriptions, InvLightingAtlasSize, 1.0f / Heightfield.DownsampleFactor);
+					const TArray<FHeightfieldComponentDescription>& HeightfieldDescriptions = It.Value();
 
-					VertexShader->SetParameters(RHICmdList, View, HeightfieldDescriptions.Num());
-					PixelShader->SetParameters(RHICmdList, View, ProjectedShadowInfo, HeightfieldDescriptions.Num(), Atlas, TileIntersectionResources, CulledObjectBuffers);
+					if (HeightfieldDescriptions.Num() > 0)
+					{
+						UploadHeightfieldDescriptions(HeightfieldDescriptions, InvLightingAtlasSize, 1.0f / Heightfield.DownsampleFactor);
 
-					RHICmdList.DrawPrimitive(0, 2, HeightfieldDescriptions.Num());
+						VertexShader->SetParameters(RHICmdList, View, HeightfieldDescriptions.Num());
+						PixelShader->SetParameters(RHICmdList, View, ProjectedShadowInfo, HeightfieldDescriptions.Num(), Atlas, TileIntersectionResources, CulledObjectBuffers);
+
+						RHICmdList.DrawPrimitive(0, 2, HeightfieldDescriptions.Num());
+					}
 				}
 			}
+			RHICmdList.EndRenderPass();
 		}
 	}
 }
@@ -1141,6 +1156,8 @@ void FHeightfieldLightingViewInfo::ComputeLighting(const FViewInfo& View, FRHICo
 {
 	const ERHIFeatureLevel::Type FeatureLevel = View.GetFeatureLevel();
 
+	check(RHICmdList.IsOutsideRenderPass());
+
 	if (AllowHeightfieldGI(View) 
 		&& SupportsHeightfieldLighting(FeatureLevel, View.GetShaderPlatform())
 		//@todo - handle local lights
@@ -1156,62 +1173,65 @@ void FHeightfieldLightingViewInfo::ComputeLighting(const FViewInfo& View, FRHICo
 		const FIntPoint LightingAtlasSize = Atlas.GetAtlasSize();
 		const FVector2D InvLightingAtlasSize(1.0f / LightingAtlasSize.X, 1.0f / LightingAtlasSize.Y);
 
-		SetRenderTarget(RHICmdList, Atlas.Lighting->GetRenderTargetItem().TargetableTexture, NULL, ESimpleRenderTargetMode::EClearColorExistingDepth);		
-
-		const bool bApplyLightFunction = (View.Family->EngineShowFlags.LightFunctions &&
-			LightSceneInfo.Proxy->GetLightFunctionMaterial() && 
-			LightSceneInfo.Proxy->GetLightFunctionMaterial()->GetMaterial(FeatureLevel)->IsLightFunction());
-
-		const FMaterialRenderProxy* MaterialProxy = bApplyLightFunction ? 
-			LightSceneInfo.Proxy->GetLightFunctionMaterial() : 
-			UMaterial::GetDefaultMaterial(MD_LightFunction)->GetRenderProxy(false);
+		FRHIRenderPassInfo RPInfo(Atlas.Lighting->GetRenderTargetItem().TargetableTexture, ERenderTargetActions::Clear_Store);
+		RHICmdList.BeginRenderPass(RPInfo, TEXT("ComputeHeightfieldLighting"));
+		{
+			const bool bApplyLightFunction = (View.Family->EngineShowFlags.LightFunctions &&
+				LightSceneInfo.Proxy->GetLightFunctionMaterial() &&
+				LightSceneInfo.Proxy->GetLightFunctionMaterial()->GetMaterial(FeatureLevel)->IsLightFunction());
 
-		const FScene* Scene = (const FScene*)View.Family->Scene;
+			const FMaterialRenderProxy* MaterialProxy = bApplyLightFunction ?
+				LightSceneInfo.Proxy->GetLightFunctionMaterial() :
+				UMaterial::GetDefaultMaterial(MD_LightFunction)->GetRenderProxy(false);
 
-		const float SkyLightIndirectScale = ShouldRenderDeferredDynamicSkyLight(Scene, *View.Family) ? Scene->SkyLight->IndirectLightingIntensity : 0;
+			const FScene* Scene = (const FScene*)View.Family->Scene;
 
-		// Skip rendering if the DefaultLightFunctionMaterial isn't compiled yet
-		if (MaterialProxy->GetMaterial(FeatureLevel)->IsLightFunction())
-		{
-			RHICmdList.SetViewport(0, 0, 0.0f, LightingAtlasSize.X, LightingAtlasSize.Y, 1.0f);
+			const float SkyLightIndirectScale = ShouldRenderDeferredDynamicSkyLight(Scene, *View.Family) ? Scene->SkyLight->IndirectLightingIntensity : 0;
 
-			FGraphicsPipelineStateInitializer GraphicsPSOInit;
-			RHICmdList.ApplyCachedRenderTargets(GraphicsPSOInit);
+			// Skip rendering if the DefaultLightFunctionMaterial isn't compiled yet
+			if (MaterialProxy->GetMaterial(FeatureLevel)->IsLightFunction())
+			{
+				RHICmdList.SetViewport(0, 0, 0.0f, LightingAtlasSize.X, LightingAtlasSize.Y, 1.0f);
 
-			GraphicsPSOInit.RasterizerState = TStaticRasterizerState<FM_Solid, CM_None>::GetRHI();
-			GraphicsPSOInit.DepthStencilState = TStaticDepthStencilState<false, CF_Always>::GetRHI();
-			GraphicsPSOInit.BlendState = TStaticBlendState<>::GetRHI();
+				FGraphicsPipelineStateInitializer GraphicsPSOInit;
+				RHICmdList.ApplyCachedRenderTargets(GraphicsPSOInit);
 
-			RHICmdList.SetStreamSource(0, GQuadVertexBuffer.VertexBufferRHI, 0);
+				GraphicsPSOInit.RasterizerState = TStaticRasterizerState<FM_Solid, CM_None>::GetRHI();
+				GraphicsPSOInit.DepthStencilState = TStaticDepthStencilState<false, CF_Always>::GetRHI();
+				GraphicsPSOInit.BlendState = TStaticBlendState<>::GetRHI();
 
-			TShaderMapRef<FHeightfieldComponentQuadVS> VertexShader(View.ShaderMap);
+				RHICmdList.SetStreamSource(0, GQuadVertexBuffer.VertexBufferRHI, 0);
 
-			const FMaterial* Material = MaterialProxy->GetMaterial(FeatureLevel);
-			const FMaterialShaderMap* MaterialShaderMap = Material->GetRenderingThreadShaderMap();
-			FLightHeightfieldsPS* PixelShader = MaterialShaderMap->GetShader<FLightHeightfieldsPS>();
+				TShaderMapRef<FHeightfieldComponentQuadVS> VertexShader(View.ShaderMap);
 
-			for (TMap<FHeightfieldComponentTextures, TArray<FHeightfieldComponentDescription>>::TConstIterator It(Heightfield.ComponentDescriptions); It; ++It)
-			{
-				const TArray<FHeightfieldComponentDescription>& HeightfieldDescriptions = It.Value();
+				const FMaterial* Material = MaterialProxy->GetMaterial(FeatureLevel);
+				const FMaterialShaderMap* MaterialShaderMap = Material->GetRenderingThreadShaderMap();
+				FLightHeightfieldsPS* PixelShader = MaterialShaderMap->GetShader<FLightHeightfieldsPS>();
 
-				if (HeightfieldDescriptions.Num() > 0)
+				for (TMap<FHeightfieldComponentTextures, TArray<FHeightfieldComponentDescription>>::TConstIterator It(Heightfield.ComponentDescriptions); It; ++It)
 				{
-					UploadHeightfieldDescriptions(HeightfieldDescriptions, InvLightingAtlasSize, 1.0f / Heightfield.DownsampleFactor);
+					const TArray<FHeightfieldComponentDescription>& HeightfieldDescriptions = It.Value();
+
+					if (HeightfieldDescriptions.Num() > 0)
+					{
+						UploadHeightfieldDescriptions(HeightfieldDescriptions, InvLightingAtlasSize, 1.0f / Heightfield.DownsampleFactor);
 
-					GraphicsPSOInit.BoundShaderState.VertexDeclarationRHI = GScreenVertexDeclaration.VertexDeclarationRHI;
-					GraphicsPSOInit.BoundShaderState.VertexShaderRHI = GETSAFERHISHADER_VERTEX(*VertexShader);
-					GraphicsPSOInit.BoundShaderState.PixelShaderRHI = GETSAFERHISHADER_PIXEL(PixelShader);
-					GraphicsPSOInit.PrimitiveType = PT_TriangleList;
+						GraphicsPSOInit.BoundShaderState.VertexDeclarationRHI = GScreenVertexDeclaration.VertexDeclarationRHI;
+						GraphicsPSOInit.BoundShaderState.VertexShaderRHI = GETSAFERHISHADER_VERTEX(*VertexShader);
+						GraphicsPSOInit.BoundShaderState.PixelShaderRHI = GETSAFERHISHADER_PIXEL(PixelShader);
+						GraphicsPSOInit.PrimitiveType = PT_TriangleList;
 
-					SetGraphicsPipelineState(RHICmdList, GraphicsPSOInit);
+						SetGraphicsPipelineState(RHICmdList, GraphicsPSOInit);
 
-					VertexShader->SetParameters(RHICmdList, View, HeightfieldDescriptions.Num());
-					PixelShader->SetParameters(RHICmdList, View, LightSceneInfo, MaterialProxy, HeightfieldDescriptions.Num(), Atlas, SkyLightIndirectScale);
+						VertexShader->SetParameters(RHICmdList, View, HeightfieldDescriptions.Num());
+						PixelShader->SetParameters(RHICmdList, View, LightSceneInfo, MaterialProxy, HeightfieldDescriptions.Num(), Atlas, SkyLightIndirectScale);
 
-					RHICmdList.DrawPrimitive(0, 2, HeightfieldDescriptions.Num());
+						RHICmdList.DrawPrimitive(0, 2, HeightfieldDescriptions.Num());
+					}
 				}
 			}
 		}
+		RHICmdList.EndRenderPass();
 	}
 }
 
