Renderpass updates.

SceneHitProxies.

diff --git a/Engine/Source/Runtime/Renderer/Private/SceneHitProxyRendering.cpp b/Engine/Source/Runtime/Renderer/Private/SceneHitProxyRendering.cpp
--- a/Engine/Source/Runtime/Renderer/Private/SceneHitProxyRendering.cpp
+++ b/Engine/Source/Runtime/Renderer/Private/SceneHitProxyRendering.cpp
@@ -407,8 +407,10 @@ bool FHitProxyDrawingPolicyFactory::DrawDynamicMesh(
 
 #if WITH_EDITOR
 
-void InitHitProxyRender(FRHICommandListImmediate& RHICmdList, const FSceneRenderer* SceneRenderer, TRefCountPtr<IPooledRenderTarget>& OutHitProxyRT, TRefCountPtr<IPooledRenderTarget>& OutHitProxyDepthRT)
+void InitHitProxyRenderAndBeginRenderPass(FRHICommandListImmediate& RHICmdList, const FSceneRenderer* SceneRenderer, TRefCountPtr<IPooledRenderTarget>& OutHitProxyRT, TRefCountPtr<IPooledRenderTarget>& OutHitProxyDepthRT)
 {
+	check(!RHICmdList.IsInsideRenderPass());
+
 	auto& ViewFamily = SceneRenderer->ViewFamily;
 	auto FeatureLevel = ViewFamily.Scene->GetFeatureLevel();
 
@@ -439,20 +441,29 @@ void InitHitProxyRender(FRHICommandListImmediate& RHICmdList, const FSceneRender
 		}
 	}
 
-	SetRenderTarget(RHICmdList, OutHitProxyRT->GetRenderTargetItem().TargetableTexture, OutHitProxyDepthRT->GetRenderTargetItem().TargetableTexture, ESimpleRenderTargetMode::EExistingColorAndDepth, FExclusiveDepthStencil::DepthWrite_StencilWrite, true);
+	FRHIRenderPassInfo RPInfo(OutHitProxyRT->GetRenderTargetItem().TargetableTexture, ERenderTargetActions::Load_Store);
+	RPInfo.DepthStencilRenderTarget.Action = EDepthStencilTargetActions::LoadDepthStencil_StoreDepthStencil;
+	RPInfo.DepthStencilRenderTarget.DepthStencilTarget = OutHitProxyDepthRT->GetRenderTargetItem().TargetableTexture;
+	RPInfo.DepthStencilRenderTarget.ExclusiveDepthStencil = FExclusiveDepthStencil::DepthWrite_StencilWrite;
+	TransitionRenderPassTargets(RHICmdList, RPInfo);
 
-	// Clear color for each view.
-	auto& Views = SceneRenderer->Views;
-	for (int32 ViewIndex = 0; ViewIndex < Views.Num(); ViewIndex++)
+	RHICmdList.BeginRenderPass(RPInfo, TEXT("ClearHitProxies"));
 	{
-		const FViewInfo& View = Views[ViewIndex];
-		RHICmdList.SetViewport(View.ViewRect.Min.X, View.ViewRect.Min.Y, 0.0f, View.ViewRect.Max.X, View.ViewRect.Max.Y, 1.0f);
-		DrawClearQuad(RHICmdList, true, FLinearColor::White, false, 0, false, 0, OutHitProxyRT->GetDesc().Extent, FIntRect());
+		// Clear color for each view.
+		auto& Views = SceneRenderer->Views;
+		for (int32 ViewIndex = 0; ViewIndex < Views.Num(); ViewIndex++)
+		{
+			const FViewInfo& View = Views[ViewIndex];
+			RHICmdList.SetViewport(View.ViewRect.Min.X, View.ViewRect.Min.Y, 0.0f, View.ViewRect.Max.X, View.ViewRect.Max.Y, 1.0f);
+			DrawClearQuad(RHICmdList, true, FLinearColor::White, false, 0, false, 0, OutHitProxyRT->GetDesc().Extent, FIntRect());
+		}
 	}
 }
 
 static void DoRenderHitProxies(FRHICommandListImmediate& RHICmdList, const FSceneRenderer* SceneRenderer, TRefCountPtr<IPooledRenderTarget> HitProxyRT, TRefCountPtr<IPooledRenderTarget> HitProxyDepthRT)
 {
+	check(RHICmdList.IsInsideRenderPass());
+
 	auto & ViewFamily = SceneRenderer->ViewFamily;
 	auto & Views = SceneRenderer->Views;
 
@@ -555,6 +566,8 @@ static void DoRenderHitProxies(FRHICommandListImmediate& RHICmdList, const FScen
 
 		View.TopBatchedViewElements.Draw(RHICmdList, DrawRenderState, FeatureLevel, bNeedToSwitchVerticalAxis, View, true);
 	}
+	// Was started in init, but ends here.
+	RHICmdList.EndRenderPass();
 
 	// Finish drawing to the hit proxy render target.
 	RHICmdList.CopyToResolveTarget(HitProxyRT->GetRenderTargetItem().TargetableTexture, HitProxyRT->GetRenderTargetItem().ShaderResourceTexture, FResolveParams());
@@ -612,23 +625,26 @@ static void DoRenderHitProxies(FRHICommandListImmediate& RHICmdList, const FScen
 
 	{
 		// Draw the triangles to the view family's render target.
-		SetRenderTarget(RHICmdList, ViewFamily.RenderTarget->GetRenderTargetTexture(), FTextureRHIRef());
-
-		FSceneView SceneView = FBatchedElements::CreateProxySceneView(PixelToView, FIntRect(0, 0, RenderTargetSize.X, RenderTargetSize.Y));
-		FDrawingPolicyRenderState DrawRenderState(SceneView);
-
-		DrawRenderState.SetDepthStencilState(TStaticDepthStencilState<false, CF_Always>::GetRHI());
-		DrawRenderState.SetBlendState(TStaticBlendState<>::GetRHI());
-
-		BatchedElements.Draw(
-			RHICmdList,
-			DrawRenderState,
-			FeatureLevel,
-			bNeedToSwitchVerticalAxis,
-			SceneView,
-			false,
-			1.0f
-		);
+		FRHIRenderPassInfo RPInfo(ViewFamily.RenderTarget->GetRenderTargetTexture(), ERenderTargetActions::Load_Store);
+		RHICmdList.BeginRenderPass(RPInfo, TEXT("HitProxies"));
+		{
+			FSceneView SceneView = FBatchedElements::CreateProxySceneView(PixelToView, FIntRect(0, 0, RenderTargetSize.X, RenderTargetSize.Y));
+			FDrawingPolicyRenderState DrawRenderState(SceneView);
+
+			DrawRenderState.SetDepthStencilState(TStaticDepthStencilState<false, CF_Always>::GetRHI());
+			DrawRenderState.SetBlendState(TStaticBlendState<>::GetRHI());
+
+			BatchedElements.Draw(
+				RHICmdList,
+				DrawRenderState,
+				FeatureLevel,
+				bNeedToSwitchVerticalAxis,
+				SceneView,
+				false,
+				1.0f
+			);
+		}
+		RHICmdList.EndRenderPass();
 	}
 
 	RHICmdList.EndScene();
@@ -642,7 +658,7 @@ void FMobileSceneRenderer::RenderHitProxies(FRHICommandListImmediate& RHICmdList
 #if WITH_EDITOR
 	TRefCountPtr<IPooledRenderTarget> HitProxyRT;
 	TRefCountPtr<IPooledRenderTarget> HitProxyDepthRT;
-	InitHitProxyRender(RHICmdList, this, HitProxyRT, HitProxyDepthRT);
+	InitHitProxyRenderAndBeginRenderPass(RHICmdList, this, HitProxyRT, HitProxyDepthRT);
 	// HitProxyRT==0 should never happen but better we don't crash
 	if (HitProxyRT)
 	{
@@ -650,6 +666,9 @@ void FMobileSceneRenderer::RenderHitProxies(FRHICommandListImmediate& RHICmdList
 		InitViews(RHICmdList);
 		::DoRenderHitProxies(RHICmdList, this, HitProxyRT, HitProxyDepthRT);
 	}
+
+	check(RHICmdList.IsOutsideRenderPass());
+
 #endif
 }
 
@@ -660,7 +679,7 @@ void FDeferredShadingSceneRenderer::RenderHitProxies(FRHICommandListImmediate& R
 #if WITH_EDITOR
 	TRefCountPtr<IPooledRenderTarget> HitProxyRT;
 	TRefCountPtr<IPooledRenderTarget> HitProxyDepthRT;
-	InitHitProxyRender(RHICmdList, this, HitProxyRT, HitProxyDepthRT);
+	InitHitProxyRenderAndBeginRenderPass(RHICmdList, this, HitProxyRT, HitProxyDepthRT);
 
 	// HitProxyRT==0 should never happen but better we don't crash
 	if (HitProxyRT)
@@ -684,5 +703,6 @@ void FDeferredShadingSceneRenderer::RenderHitProxies(FRHICommandListImmediate& R
 		::DoRenderHitProxies(RHICmdList, this, HitProxyRT, HitProxyDepthRT);
 		ClearPrimitiveSingleFramePrecomputedLightingBuffers();
 	}
+	check(RHICmdList.IsOutsideRenderPass());
 #endif
 }
