Renderpass updates.

SceneCapture. Self-contained.

diff --git a/Engine/Source/Runtime/Renderer/Private/SceneCaptureRendering.cpp b/Engine/Source/Runtime/Renderer/Private/SceneCaptureRendering.cpp
--- a/Engine/Source/Runtime/Renderer/Private/SceneCaptureRendering.cpp
+++ b/Engine/Source/Runtime/Renderer/Private/SceneCaptureRendering.cpp
@@ -198,105 +198,108 @@ void FDeferredShadingSceneRenderer::CopySceneCaptureComponentToTarget(FRHIComman
 		for (int32 ViewIndex = 0; ViewIndex < Views.Num(); ViewIndex++)
 		{
 			FViewInfo& View = Views[ViewIndex];
-			FRHIRenderTargetView ColorView(ViewFamily.RenderTarget->GetRenderTargetTexture(), 0, -1, ERenderTargetLoadAction::ENoAction, ERenderTargetStoreAction::EStore);
-			FRHISetRenderTargetsInfo Info(1, &ColorView, FRHIDepthRenderTargetView());
-			RHICmdList.SetRenderTargetsAndClear(Info);
-			RHICmdList.ApplyCachedRenderTargets(GraphicsPSOInit);
 
-			if (SceneCaptureSource == SCS_SceneColorHDR && ViewFamily.SceneCaptureCompositeMode == SCCM_Composite)
+			FRHIRenderPassInfo RPInfo(ViewFamily.RenderTarget->GetRenderTargetTexture(), ERenderTargetActions::DontLoad_Store);
+			RHICmdList.BeginRenderPass(RPInfo, TEXT("ViewCapture"));
 			{
-				// Blend with existing render target color. Scene capture color is already pre-multiplied by alpha.
-				GraphicsPSOInit.BlendState = TStaticBlendState<CW_RGBA, BO_Add, BF_One, BF_SourceAlpha, BO_Add, BF_Zero, BF_SourceAlpha>::GetRHI();
-			}
-			else if (SceneCaptureSource == SCS_SceneColorHDR && ViewFamily.SceneCaptureCompositeMode == SCCM_Additive)
-			{
-				// Add to existing render target color. Scene capture color is already pre-multiplied by alpha.
-				GraphicsPSOInit.BlendState = TStaticBlendState<CW_RGBA, BO_Add, BF_One, BF_One, BO_Add, BF_Zero, BF_SourceAlpha>::GetRHI();
-			}
-			else
-			{
-				GraphicsPSOInit.BlendState = TStaticBlendState<>::GetRHI();
-			}
+				RHICmdList.ApplyCachedRenderTargets(GraphicsPSOInit);
 
-			TShaderMapRef<FScreenVS> VertexShader(View.ShaderMap);
-			GraphicsPSOInit.BoundShaderState.VertexDeclarationRHI = GFilterVertexDeclaration.VertexDeclarationRHI;
-			GraphicsPSOInit.BoundShaderState.VertexShaderRHI = GETSAFERHISHADER_VERTEX(*VertexShader);
-			GraphicsPSOInit.PrimitiveType = PT_TriangleList;
+				if (SceneCaptureSource == SCS_SceneColorHDR && ViewFamily.SceneCaptureCompositeMode == SCCM_Composite)
+				{
+					// Blend with existing render target color. Scene capture color is already pre-multiplied by alpha.
+					GraphicsPSOInit.BlendState = TStaticBlendState<CW_RGBA, BO_Add, BF_One, BF_SourceAlpha, BO_Add, BF_Zero, BF_SourceAlpha>::GetRHI();
+				}
+				else if (SceneCaptureSource == SCS_SceneColorHDR && ViewFamily.SceneCaptureCompositeMode == SCCM_Additive)
+				{
+					// Add to existing render target color. Scene capture color is already pre-multiplied by alpha.
+					GraphicsPSOInit.BlendState = TStaticBlendState<CW_RGBA, BO_Add, BF_One, BF_One, BO_Add, BF_Zero, BF_SourceAlpha>::GetRHI();
+				}
+				else
+				{
+					GraphicsPSOInit.BlendState = TStaticBlendState<>::GetRHI();
+				}
 
-			if (SceneCaptureSource == SCS_SceneColorHDR)
-			{
-				TShaderMapRef<TSceneCapturePS<SCS_SceneColorHDR> > PixelShader(View.ShaderMap);
-				GraphicsPSOInit.BoundShaderState.PixelShaderRHI = GETSAFERHISHADER_PIXEL(*PixelShader);
-				SetGraphicsPipelineState(RHICmdList, GraphicsPSOInit);
+				TShaderMapRef<FScreenVS> VertexShader(View.ShaderMap);
+				GraphicsPSOInit.BoundShaderState.VertexDeclarationRHI = GFilterVertexDeclaration.VertexDeclarationRHI;
+				GraphicsPSOInit.BoundShaderState.VertexShaderRHI = GETSAFERHISHADER_VERTEX(*VertexShader);
+				GraphicsPSOInit.PrimitiveType = PT_TriangleList;
 
-				PixelShader->SetParameters(RHICmdList, View);
-			}
-			else if (SceneCaptureSource == SCS_SceneColorHDRNoAlpha)
-			{
-				TShaderMapRef<TSceneCapturePS<SCS_SceneColorHDRNoAlpha> > PixelShader(View.ShaderMap);
-				GraphicsPSOInit.BoundShaderState.PixelShaderRHI = GETSAFERHISHADER_PIXEL(*PixelShader);
-				SetGraphicsPipelineState(RHICmdList, GraphicsPSOInit);
+				if (SceneCaptureSource == SCS_SceneColorHDR)
+				{
+					TShaderMapRef<TSceneCapturePS<SCS_SceneColorHDR> > PixelShader(View.ShaderMap);
+					GraphicsPSOInit.BoundShaderState.PixelShaderRHI = GETSAFERHISHADER_PIXEL(*PixelShader);
+					SetGraphicsPipelineState(RHICmdList, GraphicsPSOInit);
 
-				PixelShader->SetParameters(RHICmdList, View);
-			}
-			else if (SceneCaptureSource == SCS_SceneColorSceneDepth)
-			{
-				TShaderMapRef<TSceneCapturePS<SCS_SceneColorSceneDepth> > PixelShader(View.ShaderMap);
-				GraphicsPSOInit.BoundShaderState.PixelShaderRHI = GETSAFERHISHADER_PIXEL(*PixelShader);
-				SetGraphicsPipelineState(RHICmdList, GraphicsPSOInit);
+					PixelShader->SetParameters(RHICmdList, View);
+				}
+				else if (SceneCaptureSource == SCS_SceneColorHDRNoAlpha)
+				{
+					TShaderMapRef<TSceneCapturePS<SCS_SceneColorHDRNoAlpha> > PixelShader(View.ShaderMap);
+					GraphicsPSOInit.BoundShaderState.PixelShaderRHI = GETSAFERHISHADER_PIXEL(*PixelShader);
+					SetGraphicsPipelineState(RHICmdList, GraphicsPSOInit);
 
-				PixelShader->SetParameters(RHICmdList, View);
-			}
-			else if (SceneCaptureSource == SCS_SceneDepth)
-			{
-				TShaderMapRef<TSceneCapturePS<SCS_SceneDepth> > PixelShader(View.ShaderMap);
-				GraphicsPSOInit.BoundShaderState.PixelShaderRHI = GETSAFERHISHADER_PIXEL(*PixelShader);
-				SetGraphicsPipelineState(RHICmdList, GraphicsPSOInit);
+					PixelShader->SetParameters(RHICmdList, View);
+				}
+				else if (SceneCaptureSource == SCS_SceneColorSceneDepth)
+				{
+					TShaderMapRef<TSceneCapturePS<SCS_SceneColorSceneDepth> > PixelShader(View.ShaderMap);
+					GraphicsPSOInit.BoundShaderState.PixelShaderRHI = GETSAFERHISHADER_PIXEL(*PixelShader);
+					SetGraphicsPipelineState(RHICmdList, GraphicsPSOInit);
 
-				PixelShader->SetParameters(RHICmdList, View);
-			}
-			else if (ViewFamily.SceneCaptureSource == SCS_DeviceDepth)
-			{
-				TShaderMapRef<TSceneCapturePS<SCS_DeviceDepth> > PixelShader(View.ShaderMap);
-				GraphicsPSOInit.BoundShaderState.PixelShaderRHI = GETSAFERHISHADER_PIXEL(*PixelShader);
-				SetGraphicsPipelineState(RHICmdList, GraphicsPSOInit);
+					PixelShader->SetParameters(RHICmdList, View);
+				}
+				else if (SceneCaptureSource == SCS_SceneDepth)
+				{
+					TShaderMapRef<TSceneCapturePS<SCS_SceneDepth> > PixelShader(View.ShaderMap);
+					GraphicsPSOInit.BoundShaderState.PixelShaderRHI = GETSAFERHISHADER_PIXEL(*PixelShader);
+					SetGraphicsPipelineState(RHICmdList, GraphicsPSOInit);
 
-				PixelShader->SetParameters(RHICmdList, View);
-			}
-			else if (SceneCaptureSource == SCS_Normal)
-			{
-				TShaderMapRef<TSceneCapturePS<SCS_Normal> > PixelShader(View.ShaderMap);
-				GraphicsPSOInit.BoundShaderState.PixelShaderRHI = GETSAFERHISHADER_PIXEL(*PixelShader);
-				SetGraphicsPipelineState(RHICmdList, GraphicsPSOInit);
+					PixelShader->SetParameters(RHICmdList, View);
+				}
+				else if (ViewFamily.SceneCaptureSource == SCS_DeviceDepth)
+				{
+					TShaderMapRef<TSceneCapturePS<SCS_DeviceDepth> > PixelShader(View.ShaderMap);
+					GraphicsPSOInit.BoundShaderState.PixelShaderRHI = GETSAFERHISHADER_PIXEL(*PixelShader);
+					SetGraphicsPipelineState(RHICmdList, GraphicsPSOInit);
 
-				PixelShader->SetParameters(RHICmdList, View);
-			}
-			else if (SceneCaptureSource == SCS_BaseColor)
-			{
-				TShaderMapRef<TSceneCapturePS<SCS_BaseColor> > PixelShader(View.ShaderMap);
-				GraphicsPSOInit.BoundShaderState.PixelShaderRHI = GETSAFERHISHADER_PIXEL(*PixelShader);
-				SetGraphicsPipelineState(RHICmdList, GraphicsPSOInit);
+					PixelShader->SetParameters(RHICmdList, View);
+				}
+				else if (SceneCaptureSource == SCS_Normal)
+				{
+					TShaderMapRef<TSceneCapturePS<SCS_Normal> > PixelShader(View.ShaderMap);
+					GraphicsPSOInit.BoundShaderState.PixelShaderRHI = GETSAFERHISHADER_PIXEL(*PixelShader);
+					SetGraphicsPipelineState(RHICmdList, GraphicsPSOInit);
 
-				PixelShader->SetParameters(RHICmdList, View);
-			}
-			else
-			{
-				check(0);
-			}
+					PixelShader->SetParameters(RHICmdList, View);
+				}
+				else if (SceneCaptureSource == SCS_BaseColor)
+				{
+					TShaderMapRef<TSceneCapturePS<SCS_BaseColor> > PixelShader(View.ShaderMap);
+					GraphicsPSOInit.BoundShaderState.PixelShaderRHI = GETSAFERHISHADER_PIXEL(*PixelShader);
+					SetGraphicsPipelineState(RHICmdList, GraphicsPSOInit);
 
-			VertexShader->SetParameters(RHICmdList, View.ViewUniformBuffer);
+					PixelShader->SetParameters(RHICmdList, View);
+				}
+				else
+				{
+					check(0);
+				}
 
-			DrawRectangle(
-				RHICmdList,
-				View.ViewRect.Min.X, View.ViewRect.Min.Y,
-				View.ViewRect.Width(), View.ViewRect.Height(),
-				View.ViewRect.Min.X, View.ViewRect.Min.Y,
-				View.ViewRect.Width(), View.ViewRect.Height(),
-				View.UnconstrainedViewRect.Size(),
-				FSceneRenderTargets::Get(RHICmdList).GetBufferSizeXY(),
-				*VertexShader,
-				EDRF_UseTriangleOptimization);
-		}
+				VertexShader->SetParameters(RHICmdList, View.ViewUniformBuffer);
+
+				DrawRectangle(
+					RHICmdList,
+					View.ViewRect.Min.X, View.ViewRect.Min.Y,
+					View.ViewRect.Width(), View.ViewRect.Height(),
+					View.ViewRect.Min.X, View.ViewRect.Min.Y,
+					View.ViewRect.Width(), View.ViewRect.Height(),
+					View.UnconstrainedViewRect.Size(),
+					FSceneRenderTargets::Get(RHICmdList).GetBufferSizeXY(),
+					*VertexShader,
+					EDRF_UseTriangleOptimization);
+			}
+			RHICmdList.EndRenderPass();
+		} // foreach view
 	}
 }
 
@@ -323,18 +326,24 @@ static void UpdateSceneCaptureContentDeferred_RenderThread(
 
 		// TODO: Could avoid the clear by replacing with dummy black system texture.
 		FViewInfo& View = SceneRenderer->Views[0];
-		SetRenderTarget(RHICmdList, Target->GetRenderTargetTexture(), nullptr, true);
-		DrawClearQuad(RHICmdList, true, FLinearColor::Black, false, 0, false, 0, Target->GetSizeXY(), View.UnscaledViewRect);
-
-		// Render the scene normally
-		{
-			SCOPED_DRAW_EVENT(RHICmdList, RenderScene);
-			SceneRenderer->Render(RHICmdList);
-		}
 
 		// Note: When the ViewFamily.SceneCaptureSource requires scene textures (i.e. SceneCaptureSource != SCS_FinalColorLDR), the copy to RenderTarget 
 		// will be done in CopySceneCaptureComponentToTarget while the GBuffers are still alive for the frame.
-		RHICmdList.CopyToResolveTarget(RenderTarget->GetRenderTargetTexture(), RenderTargetTexture->TextureRHI, ResolveParams);
+		FRHIRenderPassInfo RPInfo(Target->GetRenderTargetTexture(), ERenderTargetActions::DontLoad_Store, RenderTargetTexture->TextureRHI);
+		RPInfo.ResolveParameters = ResolveParams;
+		TransitionRenderPassTargets(RHICmdList, RPInfo);
+
+		RHICmdList.BeginRenderPass(RPInfo, TEXT("UpdateSceneCaptureContent"));
+		{
+			DrawClearQuad(RHICmdList, true, FLinearColor::Black, false, 0, false, 0, Target->GetSizeXY(), View.UnscaledViewRect);
+
+			// Render the scene normally
+			{
+				SCOPED_DRAW_EVENT(RHICmdList, RenderScene);
+				SceneRenderer->Render(RHICmdList);
+			}
+		}
+		RHICmdList.EndRenderPass();
 	}
 
 	FSceneRenderer::WaitForTasksClearSnapshotsAndDeleteSceneRenderer(RHICmdList, SceneRenderer);
@@ -347,41 +356,46 @@ static void ODSCapture_RenderThread(
 	FRenderTarget* const RenderTarget, 
 	const ERHIFeatureLevel::Type FeatureLevel)
 {
-	SetRenderTarget(RHICmdList, RenderTarget->GetRenderTargetTexture(), nullptr, ESimpleRenderTargetMode::EExistingColorAndDepth, FExclusiveDepthStencil::DepthNop_StencilNop, true);
-
-	FGraphicsPipelineStateInitializer GraphicsPSOInit;
-	RHICmdList.ApplyCachedRenderTargets(GraphicsPSOInit);
-	GraphicsPSOInit.BlendState = TStaticBlendState<>::GetRHI();
-	GraphicsPSOInit.RasterizerState = TStaticRasterizerState<>::GetRHI();
-	GraphicsPSOInit.DepthStencilState = TStaticDepthStencilState<false, CF_Always>::GetRHI();
-
-	const auto ShaderMap = GetGlobalShaderMap(FeatureLevel);	
-	TShaderMapRef<FScreenVS> VertexShader(ShaderMap);
-	TShaderMapRef<FODSCapturePS> PixelShader(ShaderMap);
-	extern TGlobalResource<FFilterVertexDeclaration> GFilterVertexDeclaration;
-
-	GraphicsPSOInit.BoundShaderState.VertexDeclarationRHI = GFilterVertexDeclaration.VertexDeclarationRHI;
-	GraphicsPSOInit.BoundShaderState.VertexShaderRHI = GETSAFERHISHADER_VERTEX(*VertexShader);
-	GraphicsPSOInit.BoundShaderState.PixelShaderRHI = GETSAFERHISHADER_PIXEL(*PixelShader);
-	GraphicsPSOInit.PrimitiveType = PT_TriangleList;
-
-	SetGraphicsPipelineState(RHICmdList, GraphicsPSOInit);
-
-	PixelShader->SetParameters(RHICmdList, LeftEyeTexture->TextureRHI->GetTextureCube(), RightEyeTexture->TextureRHI->GetTextureCube());
-
-	const FIntPoint& TargetSize = RenderTarget->GetSizeXY();
-	RHICmdList.SetViewport(0, 0, 0.0f, TargetSize.X, TargetSize.Y, 1.0f);
-
-	DrawRectangle(
-		RHICmdList,
-		0, 0,
-		static_cast<float>(TargetSize.X), static_cast<float>(TargetSize.Y),
-		0, 0,
-		TargetSize.X, TargetSize.Y,
-		TargetSize,
-		TargetSize,
-		*VertexShader,
-		EDRF_UseTriangleOptimization);
+	FRHIRenderPassInfo RPInfo(RenderTarget->GetRenderTargetTexture(), ERenderTargetActions::Load_Store);
+	TransitionRenderPassTargets(RHICmdList, RPInfo);
+	RHICmdList.BeginRenderPass(RPInfo, TEXT("ODSCapture"));
+	{
+
+		FGraphicsPipelineStateInitializer GraphicsPSOInit;
+		RHICmdList.ApplyCachedRenderTargets(GraphicsPSOInit);
+		GraphicsPSOInit.BlendState = TStaticBlendState<>::GetRHI();
+		GraphicsPSOInit.RasterizerState = TStaticRasterizerState<>::GetRHI();
+		GraphicsPSOInit.DepthStencilState = TStaticDepthStencilState<false, CF_Always>::GetRHI();
+
+		const auto ShaderMap = GetGlobalShaderMap(FeatureLevel);
+		TShaderMapRef<FScreenVS> VertexShader(ShaderMap);
+		TShaderMapRef<FODSCapturePS> PixelShader(ShaderMap);
+		extern TGlobalResource<FFilterVertexDeclaration> GFilterVertexDeclaration;
+
+		GraphicsPSOInit.BoundShaderState.VertexDeclarationRHI = GFilterVertexDeclaration.VertexDeclarationRHI;
+		GraphicsPSOInit.BoundShaderState.VertexShaderRHI = GETSAFERHISHADER_VERTEX(*VertexShader);
+		GraphicsPSOInit.BoundShaderState.PixelShaderRHI = GETSAFERHISHADER_PIXEL(*PixelShader);
+		GraphicsPSOInit.PrimitiveType = PT_TriangleList;
+
+		SetGraphicsPipelineState(RHICmdList, GraphicsPSOInit);
+
+		PixelShader->SetParameters(RHICmdList, LeftEyeTexture->TextureRHI->GetTextureCube(), RightEyeTexture->TextureRHI->GetTextureCube());
+
+		const FIntPoint& TargetSize = RenderTarget->GetSizeXY();
+		RHICmdList.SetViewport(0, 0, 0.0f, TargetSize.X, TargetSize.Y, 1.0f);
+
+		DrawRectangle(
+			RHICmdList,
+			0, 0,
+			static_cast<float>(TargetSize.X), static_cast<float>(TargetSize.Y),
+			0, 0,
+			TargetSize.X, TargetSize.Y,
+			TargetSize,
+			TargetSize,
+			*VertexShader,
+			EDRF_UseTriangleOptimization);
+	}
+	RHICmdList.EndRenderPass();
 }
 
 static void UpdateSceneCaptureContent_RenderThread(
