Fixes assertion failure.

diff --git a/Engine/Source/Runtime/Renderer/Private/SceneCaptureRendering.cpp b/Engine/Source/Runtime/Renderer/Private/SceneCaptureRendering.cpp
--- a/Engine/Source/Runtime/Renderer/Private/SceneCaptureRendering.cpp
+++ b/Engine/Source/Runtime/Renderer/Private/SceneCaptureRendering.cpp
@@ -327,23 +327,23 @@ static void UpdateSceneCaptureContentDeferred_RenderThread(
 		// TODO: Could avoid the clear by replacing with dummy black system texture.
 		FViewInfo& View = SceneRenderer->Views[0];
 
-		// Note: When the ViewFamily.SceneCaptureSource requires scene textures (i.e. SceneCaptureSource != SCS_FinalColorLDR), the copy to RenderTarget 
-		// will be done in CopySceneCaptureComponentToTarget while the GBuffers are still alive for the frame.
-		FRHIRenderPassInfo RPInfo(Target->GetRenderTargetTexture(), ERenderTargetActions::DontLoad_Store, RenderTargetTexture->TextureRHI);
+		FRHIRenderPassInfo RPInfo(Target->GetRenderTargetTexture(), ERenderTargetActions::DontLoad_Store);
 		RPInfo.ResolveParameters = ResolveParams;
 		TransitionRenderPassTargets(RHICmdList, RPInfo);
 
-		RHICmdList.BeginRenderPass(RPInfo, TEXT("UpdateSceneCaptureContent"));
-		{
-			DrawClearQuad(RHICmdList, true, FLinearColor::Black, false, 0, false, 0, Target->GetSizeXY(), View.UnscaledViewRect);
+		RHICmdList.BeginRenderPass(RPInfo, TEXT("ClearSceneCaptureContent"));
+		DrawClearQuad(RHICmdList, true, FLinearColor::Black, false, 0, false, 0, Target->GetSizeXY(), View.UnscaledViewRect);
+		RHICmdList.EndRenderPass();
 
-			// Render the scene normally
-			{
-				SCOPED_DRAW_EVENT(RHICmdList, RenderScene);
-				SceneRenderer->Render(RHICmdList);
-			}
+		// Render the scene normally
+		{
+			SCOPED_DRAW_EVENT(RHICmdList, RenderScene);
+			SceneRenderer->Render(RHICmdList);
 		}
-		RHICmdList.EndRenderPass();
+
+		// Note: When the ViewFamily.SceneCaptureSource requires scene textures (i.e. SceneCaptureSource != SCS_FinalColorLDR), the copy to RenderTarget 
+		// will be done in CopySceneCaptureComponentToTarget while the GBuffers are still alive for the frame.
+		RHICmdList.CopyToResolveTarget(RenderTarget->GetRenderTargetTexture(), RenderTargetTexture->TextureRHI, ResolveParams);		
 	}
 
 	FSceneRenderer::WaitForTasksClearSnapshotsAndDeleteSceneRenderer(RHICmdList, SceneRenderer);
