Distance field rendering converted to renderpasses

diff --git a/Engine/Source/Runtime/Renderer/Private/DistanceFieldAmbientOcclusion.cpp b/Engine/Source/Runtime/Renderer/Private/DistanceFieldAmbientOcclusion.cpp
--- a/Engine/Source/Runtime/Renderer/Private/DistanceFieldAmbientOcclusion.cpp
+++ b/Engine/Source/Runtime/Renderer/Private/DistanceFieldAmbientOcclusion.cpp
@@ -394,7 +394,8 @@ void ComputeDistanceFieldNormal(FRHICommandListImmediate& RHICmdList, const TArr
 {
 	if (GAOComputeShaderNormalCalculation)
 	{
-		SetRenderTarget(RHICmdList, NULL, NULL);
+		// #todo-renderpasses remove once everything is converted to renderpasses
+		UnbindRenderTargets(RHICmdList);
 
 		for (int32 ViewIndex = 0; ViewIndex < Views.Num(); ViewIndex++)
 		{
@@ -417,49 +418,49 @@ void ComputeDistanceFieldNormal(FRHICommandListImmediate& RHICmdList, const TArr
 	}
 	else
 	{
-		SetRenderTarget(RHICmdList,
-			DistanceFieldNormal.TargetableTexture,
-			NULL,
-			ESimpleRenderTargetMode::EClearColorExistingDepth,
-			FExclusiveDepthStencil::DepthNop_StencilNop,
-			true);
-		FGraphicsPipelineStateInitializer GraphicsPSOInit;
-		RHICmdList.ApplyCachedRenderTargets(GraphicsPSOInit);
+		FRHIRenderPassInfo RPInfo(DistanceFieldNormal.TargetableTexture, ERenderTargetActions::Clear_Store);
+		TransitionRenderPassTargets(RHICmdList, RPInfo);
+		RHICmdList.BeginRenderPass(RPInfo, TEXT("ComputeDistanceFieldNormal"));
+		{
+			FGraphicsPipelineStateInitializer GraphicsPSOInit;
+			RHICmdList.ApplyCachedRenderTargets(GraphicsPSOInit);
 
-		GraphicsPSOInit.RasterizerState = TStaticRasterizerState<FM_Solid, CM_None>::GetRHI();
-		GraphicsPSOInit.DepthStencilState = TStaticDepthStencilState<false, CF_Always>::GetRHI();
-		GraphicsPSOInit.BlendState = TStaticBlendState<>::GetRHI();
-		GraphicsPSOInit.PrimitiveType = PT_TriangleList;
+			GraphicsPSOInit.RasterizerState = TStaticRasterizerState<FM_Solid, CM_None>::GetRHI();
+			GraphicsPSOInit.DepthStencilState = TStaticDepthStencilState<false, CF_Always>::GetRHI();
+			GraphicsPSOInit.BlendState = TStaticBlendState<>::GetRHI();
+			GraphicsPSOInit.PrimitiveType = PT_TriangleList;
 
-		for (int32 ViewIndex = 0; ViewIndex < Views.Num(); ViewIndex++)
-		{
-			const FViewInfo& View = Views[ViewIndex];
+			for (int32 ViewIndex = 0; ViewIndex < Views.Num(); ViewIndex++)
+			{
+				const FViewInfo& View = Views[ViewIndex];
 
-			SCOPED_DRAW_EVENT(RHICmdList, ComputeNormal);
+				SCOPED_DRAW_EVENT(RHICmdList, ComputeNormal);
 
-			RHICmdList.SetViewport(0, 0, 0.0f, View.ViewRect.Width() / GAODownsampleFactor, View.ViewRect.Height() / GAODownsampleFactor, 1.0f);
+				RHICmdList.SetViewport(0, 0, 0.0f, View.ViewRect.Width() / GAODownsampleFactor, View.ViewRect.Height() / GAODownsampleFactor, 1.0f);
 
-			TShaderMapRef<FPostProcessVS> VertexShader(View.ShaderMap);
-			TShaderMapRef<FComputeDistanceFieldNormalPS> PixelShader(View.ShaderMap);
+				TShaderMapRef<FPostProcessVS> VertexShader(View.ShaderMap);
+				TShaderMapRef<FComputeDistanceFieldNormalPS> PixelShader(View.ShaderMap);
 
-			GraphicsPSOInit.BoundShaderState.VertexDeclarationRHI = GFilterVertexDeclaration.VertexDeclarationRHI;
-			GraphicsPSOInit.BoundShaderState.VertexShaderRHI = GETSAFERHISHADER_VERTEX(*VertexShader);
-			GraphicsPSOInit.BoundShaderState.PixelShaderRHI = GETSAFERHISHADER_PIXEL(*PixelShader);
+				GraphicsPSOInit.BoundShaderState.VertexDeclarationRHI = GFilterVertexDeclaration.VertexDeclarationRHI;
+				GraphicsPSOInit.BoundShaderState.VertexShaderRHI = GETSAFERHISHADER_VERTEX(*VertexShader);
+				GraphicsPSOInit.BoundShaderState.PixelShaderRHI = GETSAFERHISHADER_PIXEL(*PixelShader);
 
-			SetGraphicsPipelineState(RHICmdList, GraphicsPSOInit);
+				SetGraphicsPipelineState(RHICmdList, GraphicsPSOInit);
 
-			PixelShader->SetParameters(RHICmdList, View, Parameters);
+				PixelShader->SetParameters(RHICmdList, View, Parameters);
 
-			DrawRectangle(
-				RHICmdList,
-				0, 0,
-				View.ViewRect.Width() / GAODownsampleFactor, View.ViewRect.Height() / GAODownsampleFactor,
-				0, 0,
-				View.ViewRect.Width(), View.ViewRect.Height(),
-				FIntPoint(View.ViewRect.Width() / GAODownsampleFactor, View.ViewRect.Height() / GAODownsampleFactor),
-				FSceneRenderTargets::Get(RHICmdList).GetBufferSizeXY(),
-				*VertexShader);
+				DrawRectangle(
+					RHICmdList,
+					0, 0,
+					View.ViewRect.Width() / GAODownsampleFactor, View.ViewRect.Height() / GAODownsampleFactor,
+					0, 0,
+					View.ViewRect.Width(), View.ViewRect.Height(),
+					FIntPoint(View.ViewRect.Width() / GAODownsampleFactor, View.ViewRect.Height() / GAODownsampleFactor),
+					FSceneRenderTargets::Get(RHICmdList).GetBufferSizeXY(),
+					*VertexShader);
+			}
 		}
+		RHICmdList.EndRenderPass();
 
 		RHICmdList.TransitionResource(EResourceTransitionAccess::EReadable, DistanceFieldNormal.TargetableTexture);
 	}
@@ -794,6 +795,8 @@ bool FDeferredShadingSceneRenderer::RenderDistanceFieldLighting(
 	bool bModulateToSceneColor,
 	bool bVisualizeAmbientOcclusion)
 {
+	check(RHICmdList.IsOutsideRenderPass());
+
 	SCOPED_DRAW_EVENT(RHICmdList, RenderDistanceFieldLighting);
 
 	//@todo - support multiple views
@@ -876,11 +879,16 @@ bool FDeferredShadingSceneRenderer::RenderDistanceFieldLighting(
 
 			if (bVisualizeAmbientOcclusion)
 			{
-				SceneContext.BeginRenderingSceneColor(RHICmdList, ESimpleRenderTargetMode::EExistingColorAndDepth, FExclusiveDepthStencil::DepthRead_StencilNop);
+				SceneContext.BeginRenderingSceneColor(RHICmdList, ESimpleRenderTargetMode::EExistingColorAndDepth, FExclusiveDepthStencil::DepthRead_StencilRead);
 			}
 			else
 			{
-				SetRenderTarget(RHICmdList, SceneContext.GetSceneColorSurface(), SceneContext.GetSceneDepthSurface(), ESimpleRenderTargetMode::EExistingColorAndDepth, FExclusiveDepthStencil::DepthRead_StencilNop);
+				FRHIRenderPassInfo RPInfo(	SceneContext.GetSceneColorSurface(), 
+											ERenderTargetActions::Load_Store, 
+											SceneContext.GetSceneDepthSurface(), 
+											EDepthStencilTargetActions::LoadDepthStencil_StoreStencilNotDepth, 
+											FExclusiveDepthStencil::DepthRead_StencilWrite);
+				RHICmdList.BeginRenderPass(RPInfo, TEXT("DistanceFieldAO"));
 			}
 
 			// Upsample to full resolution, write to output in case of debug AO visualization or scene color modulation (standard upsampling is done later together with sky lighting and reflection environment)
@@ -891,8 +899,13 @@ bool FDeferredShadingSceneRenderer::RenderDistanceFieldLighting(
 
 			OutDynamicBentNormalAO = BentNormalOutput;
 
-			if (!bVisualizeAmbientOcclusion)
+			if (bVisualizeAmbientOcclusion)
+			{
+				SceneContext.FinishRenderingSceneColor(RHICmdList);
+			}
+			else
 			{
+				RHICmdList.EndRenderPass();
 				RHICmdList.CopyToResolveTarget(OutDynamicBentNormalAO->GetRenderTargetItem().TargetableTexture, OutDynamicBentNormalAO->GetRenderTargetItem().ShaderResourceTexture, FResolveParams());
 			}
 
diff --git a/Engine/Source/Runtime/Renderer/Private/DistanceFieldLightingPost.cpp b/Engine/Source/Runtime/Renderer/Private/DistanceFieldLightingPost.cpp
--- a/Engine/Source/Runtime/Renderer/Private/DistanceFieldLightingPost.cpp
+++ b/Engine/Source/Runtime/Renderer/Private/DistanceFieldLightingPost.cpp
@@ -481,7 +481,8 @@ void GeometryAwareUpsample(FRHICommandList& RHICmdList, const FViewInfo& View, T
 
 	FSceneRenderTargets& SceneContext = FSceneRenderTargets::Get(RHICmdList);
 
-	SetRenderTarget(RHICmdList, DistanceFieldAOBentNormal->GetRenderTargetItem().TargetableTexture, FTextureRHIParamRef());
+	FRHIRenderPassInfo RPInfo(DistanceFieldAOBentNormal->GetRenderTargetItem().TargetableTexture, ERenderTargetActions::Load_Store);
+	RHICmdList.BeginRenderPass(RPInfo, TEXT("GeometryAwareUpsample"));
 
 	FGraphicsPipelineStateInitializer GraphicsPSOInit;
 	RHICmdList.ApplyCachedRenderTargets(GraphicsPSOInit);
@@ -513,8 +514,10 @@ void GeometryAwareUpsample(FRHICommandList& RHICmdList, const FViewInfo& View, T
 			SceneContext.GetBufferSizeXY() / FIntPoint(GAODownsampleFactor, GAODownsampleFactor),
 			*VertexShader);
 	}
-
+	RHICmdList.EndRenderPass();
 	RHICmdList.CopyToResolveTarget(DistanceFieldAOBentNormal->GetRenderTargetItem().TargetableTexture, DistanceFieldAOBentNormal->GetRenderTargetItem().ShaderResourceTexture, FResolveParams());
+
+	
 }
 
 void UpdateHistory(
@@ -552,38 +555,41 @@ void UpdateHistory(
 			SCOPED_DRAW_EVENT(RHICmdList, UpdateHistory);
 
 			{
-				SetRenderTarget(RHICmdList, NewBentNormalHistory->GetRenderTargetItem().TargetableTexture, FTextureRHIParamRef());
-				RHICmdList.SetViewport(0, 0, 0.0f, View.ViewRect.Width() / GAODownsampleFactor, View.ViewRect.Height() / GAODownsampleFactor, 1.0f);
-
-				FGraphicsPipelineStateInitializer GraphicsPSOInit;
-				RHICmdList.ApplyCachedRenderTargets(GraphicsPSOInit);
-				GraphicsPSOInit.RasterizerState = TStaticRasterizerState<FM_Solid, CM_None>::GetRHI();
-				GraphicsPSOInit.DepthStencilState = TStaticDepthStencilState<false, CF_Always>::GetRHI();
-				GraphicsPSOInit.BlendState = TStaticBlendState<>::GetRHI();
-
-				TShaderMapRef<FPostProcessVS> VertexShader(View.ShaderMap);
-				GraphicsPSOInit.BoundShaderState.VertexDeclarationRHI = GFilterVertexDeclaration.VertexDeclarationRHI;
-				GraphicsPSOInit.BoundShaderState.VertexShaderRHI = GETSAFERHISHADER_VERTEX(*VertexShader);
-				GraphicsPSOInit.PrimitiveType = PT_TriangleList;
-
-				TShaderMapRef<FUpdateHistoryDepthRejectionPS> PixelShader(View.ShaderMap);
-
-				GraphicsPSOInit.BoundShaderState.PixelShaderRHI = GETSAFERHISHADER_PIXEL(*PixelShader);
-				SetGraphicsPipelineState(RHICmdList, GraphicsPSOInit);
-				PixelShader->SetParameters(RHICmdList, View, *DistanceFieldAOHistoryViewRect, DistanceFieldNormal, BentNormalInterpolation,(*BentNormalHistoryState)->GetRenderTargetItem(), VelocityTexture, Parameters);
-
-				VertexShader->SetParameters(RHICmdList, View.ViewUniformBuffer);
-
-				DrawRectangle( 
-					RHICmdList,
-					0, 0, 
-					View.ViewRect.Width() / GAODownsampleFactor, View.ViewRect.Height() / GAODownsampleFactor,
-					View.ViewRect.Min.X / GAODownsampleFactor, View.ViewRect.Min.Y / GAODownsampleFactor,
-					View.ViewRect.Width() / GAODownsampleFactor, View.ViewRect.Height() / GAODownsampleFactor,
-					FIntPoint(View.ViewRect.Width() / GAODownsampleFactor, View.ViewRect.Height() / GAODownsampleFactor),
-					SceneContext.GetBufferSizeXY() / FIntPoint(GAODownsampleFactor, GAODownsampleFactor),
-					*VertexShader);
+				FRHIRenderPassInfo RPInfo(NewBentNormalHistory->GetRenderTargetItem().TargetableTexture, ERenderTargetActions::Load_Store);
+				RHICmdList.BeginRenderPass(RPInfo, TEXT("UpdateHistory"));
+				{
+					RHICmdList.SetViewport(0, 0, 0.0f, View.ViewRect.Width() / GAODownsampleFactor, View.ViewRect.Height() / GAODownsampleFactor, 1.0f);
+
+					FGraphicsPipelineStateInitializer GraphicsPSOInit;
+					RHICmdList.ApplyCachedRenderTargets(GraphicsPSOInit);
+					GraphicsPSOInit.RasterizerState = TStaticRasterizerState<FM_Solid, CM_None>::GetRHI();
+					GraphicsPSOInit.DepthStencilState = TStaticDepthStencilState<false, CF_Always>::GetRHI();
+					GraphicsPSOInit.BlendState = TStaticBlendState<>::GetRHI();
+
+					TShaderMapRef<FPostProcessVS> VertexShader(View.ShaderMap);
+					GraphicsPSOInit.BoundShaderState.VertexDeclarationRHI = GFilterVertexDeclaration.VertexDeclarationRHI;
+					GraphicsPSOInit.BoundShaderState.VertexShaderRHI = GETSAFERHISHADER_VERTEX(*VertexShader);
+					GraphicsPSOInit.PrimitiveType = PT_TriangleList;
+
+					TShaderMapRef<FUpdateHistoryDepthRejectionPS> PixelShader(View.ShaderMap);
 
+					GraphicsPSOInit.BoundShaderState.PixelShaderRHI = GETSAFERHISHADER_PIXEL(*PixelShader);
+					SetGraphicsPipelineState(RHICmdList, GraphicsPSOInit);
+					PixelShader->SetParameters(RHICmdList, View, *DistanceFieldAOHistoryViewRect, DistanceFieldNormal, BentNormalInterpolation, (*BentNormalHistoryState)->GetRenderTargetItem(), VelocityTexture, Parameters);
+
+					VertexShader->SetParameters(RHICmdList, View.ViewUniformBuffer);
+
+					DrawRectangle(
+						RHICmdList,
+						0, 0,
+						View.ViewRect.Width() / GAODownsampleFactor, View.ViewRect.Height() / GAODownsampleFactor,
+						View.ViewRect.Min.X / GAODownsampleFactor, View.ViewRect.Min.Y / GAODownsampleFactor,
+						View.ViewRect.Width() / GAODownsampleFactor, View.ViewRect.Height() / GAODownsampleFactor,
+						FIntPoint(View.ViewRect.Width() / GAODownsampleFactor, View.ViewRect.Height() / GAODownsampleFactor),
+						SceneContext.GetBufferSizeXY() / FIntPoint(GAODownsampleFactor, GAODownsampleFactor),
+						*VertexShader);
+				}
+				RHICmdList.EndRenderPass();
 				RHICmdList.CopyToResolveTarget(NewBentNormalHistory->GetRenderTargetItem().TargetableTexture, NewBentNormalHistory->GetRenderTargetItem().ShaderResourceTexture, FResolveParams());
 			}
 
@@ -600,8 +606,9 @@ void UpdateHistory(
 					AllocateOrReuseAORenderTarget(RHICmdList, *BentNormalHistoryState, BentNormalHistoryRTName, PF_FloatRGBA);
 				}
 
+				FRHIRenderPassInfo RPInfo((*BentNormalHistoryState)->GetRenderTargetItem().TargetableTexture, ERenderTargetActions::Load_Store);
+				RHICmdList.BeginRenderPass(RPInfo, TEXT("UpdateHistoryStability"));
 				{
-					SetRenderTarget(RHICmdList, (*BentNormalHistoryState)->GetRenderTargetItem().TargetableTexture, FTextureRHIParamRef());
 					RHICmdList.SetViewport(0, 0, 0.0f, View.ViewRect.Width() / GAODownsampleFactor, View.ViewRect.Height() / GAODownsampleFactor, 1.0f);
 
 					FGraphicsPipelineStateInitializer GraphicsPSOInit;
@@ -643,9 +650,9 @@ void UpdateHistory(
 						FIntPoint(View.ViewRect.Width() / GAODownsampleFactor, View.ViewRect.Height() / GAODownsampleFactor),
 						SceneContext.GetBufferSizeXY() / FIntPoint(GAODownsampleFactor, GAODownsampleFactor),
 						*VertexShader);
-
-					RHICmdList.CopyToResolveTarget((*BentNormalHistoryState)->GetRenderTargetItem().TargetableTexture, (*BentNormalHistoryState)->GetRenderTargetItem().ShaderResourceTexture, FResolveParams());
 				}
+				RHICmdList.EndRenderPass();
+				RHICmdList.CopyToResolveTarget((*BentNormalHistoryState)->GetRenderTargetItem().TargetableTexture, (*BentNormalHistoryState)->GetRenderTargetItem().ShaderResourceTexture, FResolveParams());
 
 				BentNormalHistoryOutput = *BentNormalHistoryState;
 			}
diff --git a/Engine/Source/Runtime/Renderer/Private/DistanceFieldObjectCulling.cpp b/Engine/Source/Runtime/Renderer/Private/DistanceFieldObjectCulling.cpp
--- a/Engine/Source/Runtime/Renderer/Private/DistanceFieldObjectCulling.cpp
+++ b/Engine/Source/Runtime/Renderer/Private/DistanceFieldObjectCulling.cpp
@@ -558,48 +558,54 @@ void ScatterTilesToObjects(FRHICommandListImmediate& RHICmdList, const FViewInfo
 	TArray<FUnorderedAccessViewRHIParamRef> UAVs;
 	PixelShader->GetUAVs(View, UAVs);
 	RHICmdList.TransitionResources(EResourceTransitionAccess::ERWBarrier, EResourceTransitionPipeline::EComputeToGfx, UAVs.GetData(), UAVs.Num());
+
+	FRHIRenderPassInfo RPInfo(UAVs.Num(), UAVs.GetData());
 	if (GRHIRequiresRenderTargetForPixelShaderUAVs)
 	{
 		TRefCountPtr<IPooledRenderTarget> Dummy;
 		FPooledRenderTargetDesc Desc(FPooledRenderTargetDesc::Create2DDesc(TileListGroupSize, PF_B8G8R8A8, FClearValueBinding::None, TexCreate_None, TexCreate_RenderTargetable, false));
 		GRenderTargetPool.FindFreeElement(RHICmdList, Desc, Dummy, TEXT("Dummy"));
-		FRHIRenderTargetView DummyRTView(Dummy->GetRenderTargetItem().TargetableTexture, ERenderTargetLoadAction::ENoAction);
-		RHICmdList.SetRenderTargets(1, &DummyRTView, NULL, UAVs.Num(), UAVs.GetData());
+
+		RPInfo.ColorRenderTargets[0].Action = ERenderTargetActions::DontLoad_DontStore;
+		RPInfo.ColorRenderTargets[0].ArraySlice = -1;
+		RPInfo.ColorRenderTargets[0].MipIndex = 0;
+		RPInfo.ColorRenderTargets[0].RenderTarget = Dummy->GetRenderTargetItem().TargetableTexture;
 	}
-	else
+
+	RHICmdList.BeginRenderPass(RPInfo, TEXT("ScatterTilesToObjects"));
 	{
-		RHICmdList.SetRenderTargets(0, (const FRHIRenderTargetView*)NULL, NULL, UAVs.Num(), UAVs.GetData());
-	}
-	
-	FGraphicsPipelineStateInitializer GraphicsPSOInit;
-	RHICmdList.ApplyCachedRenderTargets(GraphicsPSOInit);
+		FGraphicsPipelineStateInitializer GraphicsPSOInit;
+		RHICmdList.ApplyCachedRenderTargets(GraphicsPSOInit);
 
-	RHICmdList.SetViewport(0, 0, 0.0f, TileListGroupSize.X, TileListGroupSize.Y, 1.0f);
+		RHICmdList.SetViewport(0, 0, 0.0f, TileListGroupSize.X, TileListGroupSize.Y, 1.0f);
 
-	// Render backfaces since camera may intersect
-	GraphicsPSOInit.RasterizerState = View.bReverseCulling ? TStaticRasterizerState<FM_Solid, CM_CW>::GetRHI() : TStaticRasterizerState<FM_Solid, CM_CCW>::GetRHI();
-	GraphicsPSOInit.DepthStencilState = TStaticDepthStencilState<false, CF_Always>::GetRHI();
-	GraphicsPSOInit.BlendState = TStaticBlendState<>::GetRHI();
-	GraphicsPSOInit.PrimitiveType = PT_TriangleList;
+		// Render backfaces since camera may intersect
+		GraphicsPSOInit.RasterizerState = View.bReverseCulling ? TStaticRasterizerState<FM_Solid, CM_CW>::GetRHI() : TStaticRasterizerState<FM_Solid, CM_CCW>::GetRHI();
+		GraphicsPSOInit.DepthStencilState = TStaticDepthStencilState<false, CF_Always>::GetRHI();
+		GraphicsPSOInit.BlendState = TStaticBlendState<>::GetRHI();
+		GraphicsPSOInit.PrimitiveType = PT_TriangleList;
 
-	GraphicsPSOInit.BoundShaderState.VertexDeclarationRHI = GetVertexDeclarationFVector4();
-	GraphicsPSOInit.BoundShaderState.VertexShaderRHI = GETSAFERHISHADER_VERTEX(*VertexShader);
-	GraphicsPSOInit.BoundShaderState.PixelShaderRHI = GETSAFERHISHADER_PIXEL(*PixelShader);
+		GraphicsPSOInit.BoundShaderState.VertexDeclarationRHI = GetVertexDeclarationFVector4();
+		GraphicsPSOInit.BoundShaderState.VertexShaderRHI = GETSAFERHISHADER_VERTEX(*VertexShader);
+		GraphicsPSOInit.BoundShaderState.PixelShaderRHI = GETSAFERHISHADER_PIXEL(*PixelShader);
 
-	SetGraphicsPipelineState(RHICmdList, GraphicsPSOInit);
+		SetGraphicsPipelineState(RHICmdList, GraphicsPSOInit);
 
-	VertexShader->SetParameters(RHICmdList, View, Parameters);
-	PixelShader->SetParameters(RHICmdList, View, FVector2D(TileListGroupSize.X, TileListGroupSize.Y), Parameters);
+		VertexShader->SetParameters(RHICmdList, View, Parameters);
+		PixelShader->SetParameters(RHICmdList, View, FVector2D(TileListGroupSize.X, TileListGroupSize.Y), Parameters);
 
-	RHICmdList.SetStreamSource(0, StencilingGeometry::GLowPolyStencilSphereVertexBuffer.VertexBufferRHI, 0);
+		RHICmdList.SetStreamSource(0, StencilingGeometry::GLowPolyStencilSphereVertexBuffer.VertexBufferRHI, 0);
 
-	RHICmdList.DrawIndexedPrimitiveIndirect(
-		StencilingGeometry::GLowPolyStencilSphereIndexBuffer.IndexBufferRHI, 
-		GAOCulledObjectBuffers.Buffers.ObjectIndirectArguments.Buffer,
-		0);
-	RHICmdList.TransitionResources(EResourceTransitionAccess::ERWBarrier, EResourceTransitionPipeline::EGfxToCompute, UAVs.GetData(), UAVs.Num());
+		RHICmdList.DrawIndexedPrimitiveIndirect(
+			StencilingGeometry::GLowPolyStencilSphereIndexBuffer.IndexBufferRHI,
+			GAOCulledObjectBuffers.Buffers.ObjectIndirectArguments.Buffer,
+			0);
+	}
+	RHICmdList.EndRenderPass();
 
-	SetRenderTarget(RHICmdList, NULL, NULL);
+	// #todo-renderpasses Remove once SetRenderTargets is completely removed.
+	UnbindRenderTargets(RHICmdList);
+	RHICmdList.TransitionResources(EResourceTransitionAccess::ERWBarrier, EResourceTransitionPipeline::EGfxToCompute, UAVs.GetData(), UAVs.Num());
 }
 
 FIntPoint GetTileListGroupSizeForView(const FViewInfo& View)
@@ -612,7 +618,7 @@ FIntPoint GetTileListGroupSizeForView(const FViewInfo& View)
 void BuildTileObjectLists(FRHICommandListImmediate& RHICmdList, FScene* Scene, TArray<FViewInfo>& Views, FSceneRenderTargetItem& DistanceFieldNormal, const FDistanceFieldAOParameters& Parameters)
 {
 	SCOPED_DRAW_EVENT(RHICmdList, BuildTileList);
-	SetRenderTarget(RHICmdList, NULL, NULL);
+	UnbindRenderTargets(RHICmdList);
 
 	for (int32 ViewIndex = 0; ViewIndex < Views.Num(); ViewIndex++)
 	{
diff --git a/Engine/Source/Runtime/Renderer/Private/DistanceFieldScreenGridLighting.cpp b/Engine/Source/Runtime/Renderer/Private/DistanceFieldScreenGridLighting.cpp
--- a/Engine/Source/Runtime/Renderer/Private/DistanceFieldScreenGridLighting.cpp
+++ b/Engine/Source/Runtime/Renderer/Private/DistanceFieldScreenGridLighting.cpp
@@ -558,7 +558,7 @@ void FDeferredShadingSceneRenderer::RenderDistanceFieldAOScreenGrid(
 	}
 	ScreenGridResources->AcquireTransientResource();
 
-	SetRenderTarget(RHICmdList, NULL, NULL);
+	UnbindRenderTargets(RHICmdList);
 
 	if (bUseGlobalDistanceField)
 	{
diff --git a/Engine/Source/Runtime/Renderer/Private/DistanceFieldShadowing.cpp b/Engine/Source/Runtime/Renderer/Private/DistanceFieldShadowing.cpp
--- a/Engine/Source/Runtime/Renderer/Private/DistanceFieldShadowing.cpp
+++ b/Engine/Source/Runtime/Renderer/Private/DistanceFieldShadowing.cpp
@@ -688,45 +688,51 @@ void ScatterObjectsToShadowTiles(
 	TArray<FUnorderedAccessViewRHIParamRef> UAVs;
 	PixelShader->GetUAVs(View, TileIntersectionResources, UAVs);
 	RHICmdList.TransitionResources(EResourceTransitionAccess::ERWBarrier, EResourceTransitionPipeline::EComputeToGfx, UAVs.GetData(), UAVs.Num());
+
+	FRHIRenderPassInfo RPInfo(UAVs.Num(), UAVs.GetData());
 	if (GRHIRequiresRenderTargetForPixelShaderUAVs)
 	{
 		TRefCountPtr<IPooledRenderTarget> Dummy;
 		FPooledRenderTargetDesc Desc(FPooledRenderTargetDesc::Create2DDesc(LightTileDimensions, PF_B8G8R8A8, FClearValueBinding::None, TexCreate_None, TexCreate_RenderTargetable, false));
 		GRenderTargetPool.FindFreeElement(RHICmdList, Desc, Dummy, TEXT("Dummy"));
-		FRHIRenderTargetView DummyRTView(Dummy->GetRenderTargetItem().TargetableTexture, ERenderTargetLoadAction::ENoAction);
-		RHICmdList.SetRenderTargets(1, &DummyRTView, NULL, UAVs.Num(), UAVs.GetData());
+
+		RPInfo.ColorRenderTargets[0].Action = ERenderTargetActions::DontLoad_DontStore;
+		RPInfo.ColorRenderTargets[0].ArraySlice = -1;
+		RPInfo.ColorRenderTargets[0].MipIndex = 0;
+		RPInfo.ColorRenderTargets[0].RenderTarget = Dummy->GetRenderTargetItem().TargetableTexture;
 	}
-	else
+	
+	RHICmdList.BeginRenderPass(RPInfo, TEXT("ScatterObjectsToShadowTiles"));
 	{
-		RHICmdList.SetRenderTargets(0, (const FRHIRenderTargetView *)NULL, NULL, UAVs.Num(), UAVs.GetData());
+		FGraphicsPipelineStateInitializer GraphicsPSOInit;
+		RHICmdList.ApplyCachedRenderTargets(GraphicsPSOInit);
+
+		RHICmdList.SetViewport(0, 0, 0.0f, LightTileDimensions.X, LightTileDimensions.Y, 1.0f);
+
+		// Render backfaces since camera may intersect
+		GraphicsPSOInit.RasterizerState = View.bReverseCulling ? TStaticRasterizerState<FM_Solid, CM_CW>::GetRHI() : TStaticRasterizerState<FM_Solid, CM_CCW>::GetRHI();
+		GraphicsPSOInit.DepthStencilState = TStaticDepthStencilState<false, CF_Always>::GetRHI();
+		GraphicsPSOInit.BlendState = TStaticBlendState<>::GetRHI();
+		GraphicsPSOInit.BoundShaderState.VertexDeclarationRHI = GetVertexDeclarationFVector4();
+		GraphicsPSOInit.BoundShaderState.VertexShaderRHI = GETSAFERHISHADER_VERTEX(*VertexShader);
+		GraphicsPSOInit.BoundShaderState.PixelShaderRHI = GETSAFERHISHADER_PIXEL(*PixelShader);
+		GraphicsPSOInit.PrimitiveType = PT_TriangleList;
+		SetGraphicsPipelineState(RHICmdList, GraphicsPSOInit);
+
+		VertexShader->SetParameters(RHICmdList, View, FVector2D(LightTileDimensions.X, LightTileDimensions.Y), WorldToShadowValue, ShadowBoundingRadius);
+		PixelShader->SetParameters(RHICmdList, View, TileIntersectionResources);
+
+		RHICmdList.SetStreamSource(0, GetUnitCubeVertexBuffer(), 0);
+
+		RHICmdList.DrawIndexedPrimitiveIndirect(
+			GetUnitCubeIndexBuffer(),
+			GShadowCulledObjectBuffers.Buffers.ObjectIndirectArguments.Buffer,
+			0);
 	}
+	RHICmdList.EndRenderPass();
 
-	FGraphicsPipelineStateInitializer GraphicsPSOInit;
-	RHICmdList.ApplyCachedRenderTargets(GraphicsPSOInit);
-
-	RHICmdList.SetViewport(0, 0, 0.0f, LightTileDimensions.X, LightTileDimensions.Y, 1.0f);
-
-	// Render backfaces since camera may intersect
-	GraphicsPSOInit.RasterizerState = View.bReverseCulling ? TStaticRasterizerState<FM_Solid, CM_CW>::GetRHI() : TStaticRasterizerState<FM_Solid, CM_CCW>::GetRHI();
-	GraphicsPSOInit.DepthStencilState = TStaticDepthStencilState<false, CF_Always>::GetRHI();
-	GraphicsPSOInit.BlendState = TStaticBlendState<>::GetRHI();
-	GraphicsPSOInit.BoundShaderState.VertexDeclarationRHI = GetVertexDeclarationFVector4();
-	GraphicsPSOInit.BoundShaderState.VertexShaderRHI = GETSAFERHISHADER_VERTEX(*VertexShader);
-	GraphicsPSOInit.BoundShaderState.PixelShaderRHI = GETSAFERHISHADER_PIXEL(*PixelShader);
-	GraphicsPSOInit.PrimitiveType = PT_TriangleList;
-	SetGraphicsPipelineState(RHICmdList, GraphicsPSOInit);
-
-	VertexShader->SetParameters(RHICmdList, View, FVector2D(LightTileDimensions.X, LightTileDimensions.Y), WorldToShadowValue, ShadowBoundingRadius);
-	PixelShader->SetParameters(RHICmdList, View, TileIntersectionResources);
-
-	RHICmdList.SetStreamSource(0, GetUnitCubeVertexBuffer(), 0);
-
-	RHICmdList.DrawIndexedPrimitiveIndirect(
-		GetUnitCubeIndexBuffer(),
-		GShadowCulledObjectBuffers.Buffers.ObjectIndirectArguments.Buffer,
-		0);
-
-	SetRenderTarget(RHICmdList, NULL, NULL);
+	// #todo-renderpasses Needed for now. Once SetRenderTarget is removed completely this will be redundant.
+	UnbindRenderTargets(RHICmdList);
 	RHICmdList.TransitionResources(EResourceTransitionAccess::EReadable, EResourceTransitionPipeline::EGfxToCompute, UAVs.GetData(), UAVs.Num());
 }
 
@@ -948,7 +954,7 @@ void FProjectedShadowInfo::BeginRenderRayTracedDistanceFieldProjection(FRHIComma
 		{
 			check(!Scene->DistanceFieldSceneData.HasPendingOperations());
 
-			SetRenderTarget(RHICmdList, NULL, NULL);
+			UnbindRenderTargets(RHICmdList);
 
 			int32 NumPlanes = 0;
 			const FPlane* PlaneData = NULL;
@@ -997,7 +1003,7 @@ void FProjectedShadowInfo::BeginRenderRayTracedDistanceFieldProjection(FRHIComma
 			}
 
 			SCOPED_DRAW_EVENT(RHICmdList, RayTraceShadows);
-			SetRenderTarget(RHICmdList, NULL, NULL);
+			UnbindRenderTargets(RHICmdList);
 
 			RayTraceShadows(RHICmdList, View, this, TileIntersectionResources);
 		}
@@ -1026,57 +1032,66 @@ void FProjectedShadowInfo::RenderRayTracedDistanceFieldProjection(FRHICommandLis
 		}
 
 		{
-			SetRenderTarget(RHICmdList, ScreenShadowMaskTexture->GetRenderTargetItem().TargetableTexture, FSceneRenderTargets::Get(RHICmdList).GetSceneDepthSurface(), ESimpleRenderTargetMode::EExistingColorAndDepth, FExclusiveDepthStencil::DepthRead_StencilWrite, true);
+			FRHIRenderPassInfo RPInfo(ScreenShadowMaskTexture->GetRenderTargetItem().TargetableTexture, ERenderTargetActions::Load_Store);
+			RPInfo.DepthStencilRenderTarget.Action = MakeDepthStencilTargetActions(ERenderTargetActions::Load_DontStore, ERenderTargetActions::Load_Store);
+			RPInfo.DepthStencilRenderTarget.DepthStencilTarget = FSceneRenderTargets::Get(RHICmdList).GetSceneDepthSurface();
+			RPInfo.DepthStencilRenderTarget.ExclusiveDepthStencil = FExclusiveDepthStencil::DepthRead_StencilWrite;
 
-			SCOPED_DRAW_EVENT(RHICmdList, Upsample);
+			TransitionRenderPassTargets(RHICmdList, RPInfo);
+			RHICmdList.BeginRenderPass(RPInfo, TEXT("RenderRayTracedDistanceFieldShadows"));
+			{
 
-			FGraphicsPipelineStateInitializer GraphicsPSOInit;
-			RHICmdList.ApplyCachedRenderTargets(GraphicsPSOInit);
+				SCOPED_DRAW_EVENT(RHICmdList, Upsample);
 
-			RHICmdList.SetViewport(ScissorRect.Min.X, ScissorRect.Min.Y, 0.0f, ScissorRect.Max.X, ScissorRect.Max.Y, 1.0f);
-			GraphicsPSOInit.RasterizerState = TStaticRasterizerState<FM_Solid, CM_None>::GetRHI();
-			GraphicsPSOInit.DepthStencilState = TStaticDepthStencilState<false, CF_Always>::GetRHI();
-				
-			SetBlendStateForProjection(GraphicsPSOInit, bProjectingForForwardShading, false);
+				FGraphicsPipelineStateInitializer GraphicsPSOInit;
+				RHICmdList.ApplyCachedRenderTargets(GraphicsPSOInit);
 
-			TShaderMapRef<FPostProcessVS> VertexShader(View.ShaderMap);
-			GraphicsPSOInit.BoundShaderState.VertexDeclarationRHI = GFilterVertexDeclaration.VertexDeclarationRHI;
-			GraphicsPSOInit.BoundShaderState.VertexShaderRHI = GETSAFERHISHADER_VERTEX(*VertexShader);
-			GraphicsPSOInit.PrimitiveType = PT_TriangleList;
-			GraphicsPSOInit.bDepthBounds = bDirectionalLight;
+				RHICmdList.SetViewport(ScissorRect.Min.X, ScissorRect.Min.Y, 0.0f, ScissorRect.Max.X, ScissorRect.Max.Y, 1.0f);
+				GraphicsPSOInit.RasterizerState = TStaticRasterizerState<FM_Solid, CM_None>::GetRHI();
+				GraphicsPSOInit.DepthStencilState = TStaticDepthStencilState<false, CF_Always>::GetRHI();
 
-			if (GFullResolutionDFShadowing)
-			{
-				TShaderMapRef<TDistanceFieldShadowingUpsamplePS<false> > PixelShader(View.ShaderMap);
-				GraphicsPSOInit.BoundShaderState.PixelShaderRHI = GETSAFERHISHADER_PIXEL(*PixelShader);
-				SetGraphicsPipelineState(RHICmdList, GraphicsPSOInit);
-				VertexShader->SetParameters(RHICmdList, View.ViewUniformBuffer);
-				PixelShader->SetParameters(RHICmdList, View, this, ScissorRect, RayTracedShadowsRT);
-			}
-			else
-			{
-				TShaderMapRef<TDistanceFieldShadowingUpsamplePS<true> > PixelShader(View.ShaderMap);
-				GraphicsPSOInit.BoundShaderState.PixelShaderRHI = GETSAFERHISHADER_PIXEL(*PixelShader);
-				SetGraphicsPipelineState(RHICmdList, GraphicsPSOInit);
-				VertexShader->SetParameters(RHICmdList, View.ViewUniformBuffer);
-				PixelShader->SetParameters(RHICmdList, View, this, ScissorRect, RayTracedShadowsRT);
-			}
+				SetBlendStateForProjection(GraphicsPSOInit, bProjectingForForwardShading, false);
 
-			//@todo - depth bounds test for local lights
-			if (bDirectionalLight)
-			{
-				SetDepthBoundsTest(RHICmdList, CascadeSettings.SplitNear - CascadeSettings.SplitNearFadeRegion, CascadeSettings.SplitFar, View.ViewMatrices.GetProjectionMatrix());
-			}
+				TShaderMapRef<FPostProcessVS> VertexShader(View.ShaderMap);
+				GraphicsPSOInit.BoundShaderState.VertexDeclarationRHI = GFilterVertexDeclaration.VertexDeclarationRHI;
+				GraphicsPSOInit.BoundShaderState.VertexShaderRHI = GETSAFERHISHADER_VERTEX(*VertexShader);
+				GraphicsPSOInit.PrimitiveType = PT_TriangleList;
+				GraphicsPSOInit.bDepthBounds = bDirectionalLight;
 
-			DrawRectangle( 
-				RHICmdList,
-				0, 0, 
-				ScissorRect.Width(), ScissorRect.Height(),
-				ScissorRect.Min.X / GetDFShadowDownsampleFactor(), ScissorRect.Min.Y / GetDFShadowDownsampleFactor(), 
-				ScissorRect.Width() / GetDFShadowDownsampleFactor(), ScissorRect.Height() / GetDFShadowDownsampleFactor(),
-				FIntPoint(ScissorRect.Width(), ScissorRect.Height()),
-				GetBufferSizeForDFShadows(),
-				*VertexShader);
+				if (GFullResolutionDFShadowing)
+				{
+					TShaderMapRef<TDistanceFieldShadowingUpsamplePS<false> > PixelShader(View.ShaderMap);
+					GraphicsPSOInit.BoundShaderState.PixelShaderRHI = GETSAFERHISHADER_PIXEL(*PixelShader);
+					SetGraphicsPipelineState(RHICmdList, GraphicsPSOInit);
+					VertexShader->SetParameters(RHICmdList, View.ViewUniformBuffer);
+					PixelShader->SetParameters(RHICmdList, View, this, ScissorRect, RayTracedShadowsRT);
+				}
+				else
+				{
+					TShaderMapRef<TDistanceFieldShadowingUpsamplePS<true> > PixelShader(View.ShaderMap);
+					GraphicsPSOInit.BoundShaderState.PixelShaderRHI = GETSAFERHISHADER_PIXEL(*PixelShader);
+					SetGraphicsPipelineState(RHICmdList, GraphicsPSOInit);
+					VertexShader->SetParameters(RHICmdList, View.ViewUniformBuffer);
+					PixelShader->SetParameters(RHICmdList, View, this, ScissorRect, RayTracedShadowsRT);
+				}
+
+				//@todo - depth bounds test for local lights
+				if (bDirectionalLight)
+				{
+					SetDepthBoundsTest(RHICmdList, CascadeSettings.SplitNear - CascadeSettings.SplitNearFadeRegion, CascadeSettings.SplitFar, View.ViewMatrices.GetProjectionMatrix());
+				}
+
+				DrawRectangle(
+					RHICmdList,
+					0, 0,
+					ScissorRect.Width(), ScissorRect.Height(),
+					ScissorRect.Min.X / GetDFShadowDownsampleFactor(), ScissorRect.Min.Y / GetDFShadowDownsampleFactor(),
+					ScissorRect.Width() / GetDFShadowDownsampleFactor(), ScissorRect.Height() / GetDFShadowDownsampleFactor(),
+					FIntPoint(ScissorRect.Width(), ScissorRect.Height()),
+					GetBufferSizeForDFShadows(),
+					*VertexShader);
+			}
+			RHICmdList.EndRenderPass();
 		}
 
 		RayTracedShadowsRT = NULL;
diff --git a/Engine/Source/Runtime/Renderer/Private/DistanceFieldVisualization.cpp b/Engine/Source/Runtime/Renderer/Private/DistanceFieldVisualization.cpp
--- a/Engine/Source/Runtime/Renderer/Private/DistanceFieldVisualization.cpp
+++ b/Engine/Source/Runtime/Renderer/Private/DistanceFieldVisualization.cpp
@@ -198,7 +198,7 @@ void FDeferredShadingSceneRenderer::RenderMeshDistanceFieldVisualization(FRHICom
 			}
 
 			{
-				SetRenderTarget(RHICmdList, NULL, NULL);
+				UnbindRenderTargets(RHICmdList);
 
 				for (int32 ViewIndex = 0; ViewIndex < Views.Num(); ViewIndex++)
 				{
@@ -240,6 +240,8 @@ void FDeferredShadingSceneRenderer::RenderMeshDistanceFieldVisualization(FRHICom
 				GAOCulledObjectBuffers.Buffers.DiscardTransientResource();
 			}
 
+			check(RHICmdList.IsOutsideRenderPass());
+
 			{
 				FSceneRenderTargets::Get(RHICmdList).BeginRenderingSceneColor(RHICmdList, ESimpleRenderTargetMode::EExistingColorAndDepth, FExclusiveDepthStencil::DepthRead_StencilRead);
 
@@ -278,6 +280,8 @@ void FDeferredShadingSceneRenderer::RenderMeshDistanceFieldVisualization(FRHICom
 						GetBufferSizeForAO(),
 						*VertexShader);
 				}
+
+				FSceneRenderTargets::Get(RHICmdList).FinishRenderingSceneColor(RHICmdList);
 			}
 		}
 	}
