Geometry shaders must remain TOptionalShaderMapRef because Metal does not have geometry shaders.
Fixes mac automation

diff --git a/Engine/Source/Runtime/Renderer/Private/VolumetricFog.cpp b/Engine/Source/Runtime/Renderer/Private/VolumetricFog.cpp
--- a/Engine/Source/Runtime/Renderer/Private/VolumetricFog.cpp
+++ b/Engine/Source/Runtime/Renderer/Private/VolumetricFog.cpp
@@ -526,7 +526,7 @@ void FDeferredShadingSceneRenderer::RenderLocalLightsForVolumetricFog(
 						PermutationVector.Set< TInjectShadowedLocalLightPS::FTemporalReprojection >( bUseTemporalReprojection );
 
 						auto VertexShader	= View.ShaderMap->GetShader< FWriteToBoundingSphereVS >();
-						auto GeometryShader	= View.ShaderMap->GetShader< FWriteToSliceGS >();
+						TOptionalShaderMapRef<FWriteToSliceGS> GeometryShader(View.ShaderMap);
 						auto PixelShader	= View.ShaderMap->GetShader< TInjectShadowedLocalLightPS >( PermutationVector );
 
 						FGraphicsPipelineStateInitializer GraphicsPSOInit;
@@ -538,7 +538,7 @@ void FDeferredShadingSceneRenderer::RenderLocalLightsForVolumetricFog(
 
 						GraphicsPSOInit.BoundShaderState.VertexDeclarationRHI = GScreenVertexDeclaration.VertexDeclarationRHI;
 						GraphicsPSOInit.BoundShaderState.VertexShaderRHI = GETSAFERHISHADER_VERTEX(VertexShader);
-						GraphicsPSOInit.BoundShaderState.GeometryShaderRHI = GETSAFERHISHADER_GEOMETRY(GeometryShader);
+						GraphicsPSOInit.BoundShaderState.GeometryShaderRHI = GETSAFERHISHADER_GEOMETRY(*GeometryShader);
 						GraphicsPSOInit.BoundShaderState.PixelShaderRHI = GETSAFERHISHADER_PIXEL(PixelShader);
 						GraphicsPSOInit.PrimitiveType = PT_TriangleList;
 
@@ -547,7 +547,7 @@ void FDeferredShadingSceneRenderer::RenderLocalLightsForVolumetricFog(
 						PixelShader->SetParameters(RHICmdList, View, IntegrationData, LightSceneInfo, FogInfo, ProjectedShadowInfo, bDynamicallyShadowed);
 						VertexShader->SetParameters(RHICmdList, View, IntegrationData, LightBounds, VolumeZBounds.X);
 
-						if (GeometryShader)
+						if (GeometryShader.IsValid())
 						{
 							GeometryShader->SetParameters(RHICmdList, VolumeZBounds.X);
 						}
