[NO NON-PII COMMIT MESSAGE]

diff --git a/Engine/Source/Runtime/Renderer/Private/SceneRendering.h b/Engine/Source/Runtime/Renderer/Private/SceneRendering.h
--- a/Engine/Source/Runtime/Renderer/Private/SceneRendering.h
+++ b/Engine/Source/Runtime/Renderer/Private/SceneRendering.h
@@ -1263,7 +1263,12 @@ public:
 
 	/** Custom Data Memstack functions.	*/
 	FORCEINLINE FMemStackBase& GetCustomDataGlobalMemStack() { return PrimitiveCustomDataMemStack[0]; }
-	FORCEINLINE FMemStackBase& AllocateCustomDataMemStack() { return *new(PrimitiveCustomDataMemStack) FMemStackBase(0); }
+	FORCEINLINE FMemStackBase& AllocateCustomDataMemStack() 
+	{ 
+		// Don't reallocate since we keep references in FRelevancePacket.
+		check(PrimitiveCustomDataMemStack.GetSlack() > 0); 
+		return *new(PrimitiveCustomDataMemStack) FMemStackBase(0);
+	}
 
 private:
 	// Cache of TEXTUREGROUP_World to create view's samplers on render thread.
diff --git a/Engine/Source/Runtime/Renderer/Private/SceneVisibility.cpp b/Engine/Source/Runtime/Renderer/Private/SceneVisibility.cpp
--- a/Engine/Source/Runtime/Renderer/Private/SceneVisibility.cpp
+++ b/Engine/Source/Runtime/Renderer/Private/SceneVisibility.cpp
@@ -2149,12 +2149,6 @@ static void ComputeAndMarkRelevanceForViewParallel(
 
 	bool WillExecuteInParallel = FApp::ShouldUseThreadingForPerformance() && CVarParallelInitViews.GetValueOnRenderThread() > 0;
 
-	if (WillExecuteInParallel)
-	{
-		// We must reserve to prevent realloc otherwise it will cause memory leak if WillExecuteInParallel == true
-		View.PrimitiveCustomDataMemStack.Reserve(View.PrimitiveCustomDataMemStack.Num() + FMath::TruncToInt((float)NumMesh / (float)FRelevancePrimSet<int32>::MaxInputPrims + 1.0f));
-	}
-
 	{
 		FSceneSetBitIterator BitIt(View.PrimitiveVisibilityMap);
 		if (BitIt)
@@ -2871,6 +2865,11 @@ void FSceneRenderer::ComputeViewVisibility(FRHICommandListImmediate& RHICmdList)
 
 		View.PrimitivesCustomData.Init(nullptr, Scene->Primitives.Num());
 		View.PrimitivesWithCustomData.Reserve(Scene->Primitives.Num());
+
+		// We must reserve to prevent realloc otherwise it will cause memory leak if we Execute In Parallel
+		bool WillExecuteInParallel = FApp::ShouldUseThreadingForPerformance() && CVarParallelInitViews.GetValueOnRenderThread() > 0;
+		View.PrimitiveCustomDataMemStack.Reserve(WillExecuteInParallel ? FMath::CeilToInt(((float)View.PrimitiveVisibilityMap.Num() / (float)FRelevancePrimSet<int32>::MaxInputPrims)) + 1 : 1);
+
 		View.AllocateCustomDataMemStack();
 
 		View.VisibleLightInfos.Empty(Scene->Lights.GetMaxIndex());
