Clear RT size history when we hit an aspect ratio change or the next frame will pass that check and end up looking back in the 3-frame history again as the buffers have already been reallocated and match the new ratio. This leads to repeated allocations when the RT size decreases until the history clears naturally.

diff --git a/Engine/Source/Runtime/Renderer/Private/PostProcess/SceneRenderTargets.cpp b/Engine/Source/Runtime/Renderer/Private/PostProcess/SceneRenderTargets.cpp
--- a/Engine/Source/Runtime/Renderer/Private/PostProcess/SceneRenderTargets.cpp
+++ b/Engine/Source/Runtime/Renderer/Private/PostProcess/SceneRenderTargets.cpp
@@ -423,7 +423,23 @@ FIntPoint FSceneRenderTargets::ComputeDesiredSize(const FSceneViewFamily& ViewFa
 			!FMath::IsNearlyEqual(
 				(float)BufferSize.X / BufferSize.Y,
 				(float)DesiredBufferSize.X / DesiredBufferSize.Y);
-		bAllowDelayResize = bAllowDelayResize && !bAspectRatioChanged;
+
+		if (bAspectRatioChanged)
+		{
+			bAllowDelayResize = false;
+
+			// At this point we're assuming a simple output resize and forcing a hard swap so clear the history.
+			// If we don't the next frame will fail this check as the allocated aspect ratio will match the new
+			// frame's forced size so we end up looking through the history again, finding the previous old size
+			// and reallocating. Only after a few frames can the results actually settle when the history clears 
+			for (int32 i = 0; i < FrameSizeHistoryCount; ++i)
+			{
+				LargestDesiredSizes[i] = FIntPoint::ZeroValue;
+#if PREVENT_RENDERTARGET_SIZE_THRASHING
+				HistoryFlags[i] = 0;
+#endif
+			}
+		}
 	}
 
 	if(bAllowDelayResize)
