When TAA is not the final PP output do not output to UnscaledViewRect.
Fixes incorrect letterboxing when using constrained aspect ratio.

diff --git a/Engine/Source/Runtime/Renderer/Private/PostProcess/PostProcessMobile.cpp b/Engine/Source/Runtime/Renderer/Private/PostProcess/PostProcessMobile.cpp
--- a/Engine/Source/Runtime/Renderer/Private/PostProcess/PostProcessMobile.cpp
+++ b/Engine/Source/Runtime/Renderer/Private/PostProcess/PostProcessMobile.cpp
@@ -2947,7 +2947,10 @@ void FRCPassPostProcessAaES2::Process(FRenderingCompositePassContext& Context)
 
 	SetShader(Context);
 
-	const FIntRect& ViewRect = Context.View.UnscaledViewRect; // Simple upscaling, ES2 post process does not currently have a specific upscaling pass.
+	const bool bIsFinalPass = Context.IsViewFamilyRenderTarget(DestRenderTarget);
+	// If final pass then perform simple upscaling
+	const FIntRect& ViewRect = bIsFinalPass ? Context.View.UnscaledViewRect : Context.View.ViewRect;
+
 	float XPos = ViewRect.Min.X;
 	float YPos = ViewRect.Min.Y;
 	float Width = ViewRect.Width();
