Allow r.Mobile.TonemapperFilm with es3.1 feature level only.

diff --git a/Engine/Source/Runtime/Renderer/Private/PostProcess/PostProcessing.cpp b/Engine/Source/Runtime/Renderer/Private/PostProcess/PostProcessing.cpp
--- a/Engine/Source/Runtime/Renderer/Private/PostProcess/PostProcessing.cpp
+++ b/Engine/Source/Runtime/Renderer/Private/PostProcess/PostProcessing.cpp
@@ -2644,7 +2644,7 @@ void FPostProcessing::ProcessES2(FRHICommandListImmediate& RHICmdList, const FVi
 		}
 		
 		static const auto VarTonemapperFilm = IConsoleManager::Get().FindTConsoleVariableDataInt(TEXT("r.Mobile.TonemapperFilm"));
-		const bool bUseTonemapperFilm = IsMobileHDR() && !bMobileHDR32bpp && GSupportsRenderTargetFormat_PF_FloatRGBA && (VarTonemapperFilm && VarTonemapperFilm->GetValueOnRenderThread());
+		const bool bUseTonemapperFilm = Context.View.GetFeatureLevel() == ERHIFeatureLevel::ES3_1 && IsMobileHDR() && !bMobileHDR32bpp && GSupportsRenderTargetFormat_PF_FloatRGBA && (VarTonemapperFilm && VarTonemapperFilm->GetValueOnRenderThread());
 
 
 		static const auto VarTonemapperUpscale = IConsoleManager::Get().FindTConsoleVariableDataInt(TEXT("r.MobileTonemapperUpscale"));
