UE-64084 UMG material displays darker texture from mediaplayer when running in mobile preview(ES2)

#4.21

Adding a new global GISMobilePreview to properly identify the game running in the Editor's ES2 Mobile Preview

diff --git a/Engine/Source/Runtime/Renderer/Private/SceneRendering.cpp b/Engine/Source/Runtime/Renderer/Private/SceneRendering.cpp
--- a/Engine/Source/Runtime/Renderer/Private/SceneRendering.cpp
+++ b/Engine/Source/Runtime/Renderer/Private/SceneRendering.cpp
@@ -1325,10 +1325,7 @@ void FViewInfo::SetupUniformBufferParameters(
 	checkSlow(sizeof(ViewUniformShaderParameters.SkyIrradianceEnvironmentMap) == sizeof(FVector4)* 7);
 	SetupSkyIrradianceEnvironmentMapConstants((FVector4*)&ViewUniformShaderParameters.SkyIrradianceEnvironmentMap);
 
-	ViewUniformShaderParameters.MobilePreviewMode =
-		(GIsEditor &&
-		(RHIFeatureLevel == ERHIFeatureLevel::ES2 || RHIFeatureLevel == ERHIFeatureLevel::ES3_1) &&
-		GMaxRHIFeatureLevel > ERHIFeatureLevel::ES3_1) ? 1.0f : 0.0f;
+	ViewUniformShaderParameters.MobilePreviewMode = Scene == nullptr ? 0.0f: (IsSimulatedPlatform(Scene->GetShaderPlatform()) ? 1.0f : 0.0f);
 
 	// Padding between the left and right eye may be introduced by an HMD, which instanced stereo needs to account for.
 	if ((Family != nullptr) && (StereoPass != eSSP_FULL) && (Family->Views.Num() > 1))
diff --git a/Engine/Source/Runtime/SlateRHIRenderer/Private/SlateRHIRenderingPolicy.cpp b/Engine/Source/Runtime/SlateRHIRenderer/Private/SlateRHIRenderingPolicy.cpp
--- a/Engine/Source/Runtime/SlateRHIRenderer/Private/SlateRHIRenderingPolicy.cpp
+++ b/Engine/Source/Runtime/SlateRHIRenderer/Private/SlateRHIRenderingPolicy.cpp
@@ -287,10 +287,8 @@ static FSceneView* CreateSceneView( FSceneViewFamilyContext* ViewFamilyContext,
 
 	ERHIFeatureLevel::Type RHIFeatureLevel = View->GetFeatureLevel();
 
-	ViewUniformShaderParameters.MobilePreviewMode =
-		(GIsEditor &&
-		(RHIFeatureLevel == ERHIFeatureLevel::ES2 || RHIFeatureLevel == ERHIFeatureLevel::ES3_1) &&
-		GMaxRHIFeatureLevel > ERHIFeatureLevel::ES3_1) ? 1.0f : 0.0f;
+	ViewUniformShaderParameters.MobilePreviewMode = IsSimulatedPlatform(View->GetShaderPlatform()) ? 1.0f : 0.0f;
+
 
 	UpdateNoiseTextureParameters(ViewUniformShaderParameters);
 
