Add movable spotlight support to mobile renderer.
Add point  light fallback for rect lights
Add new spotlight mobile shader permuation option (defaults off)

diff --git a/Engine/Shaders/Private/MobileBasePassPixelShader.usf b/Engine/Shaders/Private/MobileBasePassPixelShader.usf
--- a/Engine/Shaders/Private/MobileBasePassPixelShader.usf
+++ b/Engine/Shaders/Private/MobileBasePassPixelShader.usf
@@ -31,6 +31,12 @@ int NumDynamicPointLights;
 #endif
 float4 LightPositionAndInvRadius[MAX_DYNAMIC_POINT_LIGHTS];
 float4 LightColorAndFalloffExponent[MAX_DYNAMIC_POINT_LIGHTS];
+float4 SpotLightAngles[MAX_DYNAMIC_POINT_LIGHTS];
+float4 SpotLightDirection[MAX_DYNAMIC_POINT_LIGHTS];
+#endif
+
+#ifndef PROJECT_MOBILE_ENABLE_MOVABLE_SPOTLIGHTS
+#define PROJECT_MOBILE_ENABLE_MOVABLE_SPOTLIGHTS 0
 #endif
 
 #ifndef MOBILE_QL_FORCE_FULLY_ROUGH
@@ -589,6 +595,13 @@ void Main(
 				{
 					Attenuation = RadialAttenuation(ToLight * LightPositionAndInvRadius[i].w, LightColorAndFalloffExponent[i].w);		
 				}
+				
+				#if PROJECT_MOBILE_ENABLE_MOVABLE_SPOTLIGHTS
+				if(SpotLightAngles[i].w > 0.0)
+				{
+					Attenuation *= SpotAttenuation(L, -SpotLightDirection[i].xyz, SpotLightAngles[i].xy);
+				}
+				#endif
 
 				#if !FULLY_ROUGH
 					Color += min(65000.0, (Attenuation * PointNoL) * LightColorAndFalloffExponent[i].rgb * (1.0/PI) * (DiffuseColor + SpecularColor * CalcSpecular(Roughness, PointRoL, PointNoH, H, MaterialParameters.WorldNormal)));
diff --git a/Engine/Source/Runtime/Engine/Classes/Engine/RendererSettings.h b/Engine/Source/Runtime/Engine/Classes/Engine/RendererSettings.h
--- a/Engine/Source/Runtime/Engine/Classes/Engine/RendererSettings.h
+++ b/Engine/Source/Runtime/Engine/Classes/Engine/RendererSettings.h
@@ -634,17 +634,24 @@ class ENGINE_API URendererSettings : public UDeveloperSettings
 		uint32 bMobileAllowMovableDirectionalLights : 1;
 
 	UPROPERTY(config, EditAnywhere, Category = MobileShaderPermutationReduction, meta = (
-		ConsoleVariable = "r.MobileNumDynamicPointLights", DisplayName = "Max Movable Point Lights", ClampMax = 4,
-		ToolTip = "The number of dynamic point lights to support on mobile devices. Setting this to 0 for games which do not require dynamic point lights will reduce the number of shaders generated. Changing this setting requires restarting the editor.",
+		ConsoleVariable = "r.MobileNumDynamicPointLights", DisplayName = "Max Movable Spotlights / Point Lights", ClampMax = 4,
+		ToolTip = "The number of dynamic spotlights or point lights to support on mobile devices. Setting this to 0 for games which do not require dynamic spotlights or point lights will reduce the number of shaders generated. Changing this setting requires restarting the editor.",
 		ConfigRestartRequired = true))
 		uint32 MobileNumDynamicPointLights;
 
 	UPROPERTY(config, EditAnywhere, Category = MobileShaderPermutationReduction, meta = (
-		ConsoleVariable = "r.MobileDynamicPointLightsUseStaticBranch", DisplayName = "Use Shared Movable Point Light Shaders",
-		ToolTip = "If this setting is enabled, the same shader will be used for any number of dynamic point lights (up to the maximum specified above) hitting a surface. This is slightly slower but reduces the number of shaders generated. Changing this setting requires restarting the editor.",
+		ConsoleVariable = "r.MobileDynamicPointLightsUseStaticBranch", DisplayName = "Use Shared Movable Spotlight / Point Light Shaders",
+		ToolTip = "If this setting is enabled, the same shader will be used for any number of dynamic spotlights or point lights (up to the maximum specified above) hitting a surface. This is slightly slower but reduces the number of shaders generated. Changing this setting requires restarting the editor.",
 		ConfigRestartRequired = true))
 		uint32 bMobileDynamicPointLightsUseStaticBranch : 1;
 
+	UPROPERTY(config, EditAnywhere, Category = MobileShaderPermutationReduction, meta = (
+		ConsoleVariable = "r.Mobile.EnableMovableSpotlights",
+		DisplayName = "Support Movable Spotlights",
+		ToolTip = "Generate shaders for primitives to receive lighting from movable spotlights. This incurs an additional cost when processing movable lights. Changing this setting requires restarting the editor.",
+		ConfigRestartRequired = true))
+		uint32 bMobileAllowMovableSpotlights : 1;
+
 	UPROPERTY(config, EditAnywhere, Category = Optimizations, meta = (
 		ConsoleVariable = "r.SkinCache.SceneMemoryLimitInMB", DisplayName = "Maximum memory for Compute Skincache per world (MB)",
 		ToolTip = "Maximum amount of memory (in MB) per world/scene allowed for the Compute Skincache to generate output vertex data and recompute tangents."))
diff --git a/Engine/Source/Runtime/Engine/Private/ShaderCompiler/ShaderCompiler.cpp b/Engine/Source/Runtime/Engine/Private/ShaderCompiler/ShaderCompiler.cpp
--- a/Engine/Source/Runtime/Engine/Private/ShaderCompiler/ShaderCompiler.cpp
+++ b/Engine/Source/Runtime/Engine/Private/ShaderCompiler/ShaderCompiler.cpp
@@ -3241,6 +3241,13 @@ void GlobalBeginCompileShader(
 
 	Input.Environment.SetDefine(TEXT("PLATFORM_SUPPORTS_PER_PIXEL_DBUFFER_MASK"), IsUsingPerPixelDBufferMask(EShaderPlatform(Target.Platform)) ? 1 : 0);
 
+
+	if (IsMobilePlatform((EShaderPlatform)Target.Platform))
+	{
+		static IConsoleVariable* CVar = IConsoleManager::Get().FindConsoleVariable(TEXT("r.Mobile.EnableMovableSpotlights"));
+		Input.Environment.SetDefine(TEXT("PROJECT_MOBILE_ENABLE_MOVABLE_SPOTLIGHTS"), CVar ? (CVar->GetInt() != 0) : 0);
+	}
+
 	// Allow the target shader format to modify the shader input before we add it as a job
 	const IShaderFormat* Format = GetTargetPlatformManagerRef().FindShaderFormat(Input.ShaderFormat);
 	Format->ModifyShaderCompilerInput(Input);
diff --git a/Engine/Source/Runtime/Renderer/Private/MobileBasePassRendering.cpp b/Engine/Source/Runtime/Renderer/Private/MobileBasePassRendering.cpp
--- a/Engine/Source/Runtime/Renderer/Private/MobileBasePassRendering.cpp
+++ b/Engine/Source/Runtime/Renderer/Private/MobileBasePassRendering.cpp
@@ -28,6 +28,12 @@ static TAutoConsoleVariable<int32> CVarMobileUseLegacyShadingModel(
 	TEXT("If 1 then use legacy (pre 4.20) shading model (such as spherical guassian specular calculation.) (will cause a shader rebuild)"),
 	ECVF_ReadOnly | ECVF_RenderThreadSafe);
 
+static TAutoConsoleVariable<int32> CVarMobileEnableMovableSpotLights(
+	TEXT("r.Mobile.EnableMovableSpotlights"),
+	0,
+	TEXT("If 1 then enable movable spotlight support"),
+	ECVF_ReadOnly | ECVF_RenderThreadSafe);
+
 
 // Changing this causes a full shader recompile
 static TAutoConsoleVariable<int32> CVarMobileSeparateMaskedPass(
@@ -126,7 +132,7 @@ void GetSkyTextureParams(const FScene* Scene, float& AverageBrightnessOUT, FText
 	}
 }
 
-FMobileBasePassMovablePointLightInfo::FMobileBasePassMovablePointLightInfo(const FPrimitiveSceneProxy* InSceneProxy)
+FMobileBasePassMovableLightInfo::FMobileBasePassMovableLightInfo(const FPrimitiveSceneProxy* InSceneProxy)
 : NumMovablePointLights(0)
 {
 	static auto* MobileNumDynamicPointLightsCVar = IConsoleManager::Get().FindTConsoleVariableDataInt(TEXT("r.MobileNumDynamicPointLights"));
@@ -137,7 +143,13 @@ FMobileBasePassMovablePointLightInfo::FMobileBasePassMovablePointLightInfo(const
 		for (FLightPrimitiveInteraction* LPI = InSceneProxy->GetPrimitiveSceneInfo()->LightList; LPI && NumMovablePointLights < MobileNumDynamicPointLights; LPI = LPI->GetNextLight())
 		{
 			FLightSceneProxy* LightProxy = LPI->GetLight()->Proxy;
-			if (LightProxy->GetLightType() == LightType_Point && LightProxy->IsMovable() && (LightProxy->GetLightingChannelMask() & InSceneProxy->GetLightingChannelMask()) != 0)
+			const uint8 LightType = LightProxy->GetLightType(); 
+			const bool bIsValidLightType =
+				  LightType == LightType_Point
+				|| LightType == LightType_Rect
+				|| (LightType == LightType_Spot && CVarMobileEnableMovableSpotLights.GetValueOnRenderThread());
+				
+			if (bIsValidLightType && LightProxy->IsMovable() && (LightProxy->GetLightingChannelMask() & InSceneProxy->GetLightingChannelMask()) != 0)
 			{
 				FLightParameters LightParameters;
 
@@ -145,6 +157,14 @@ FMobileBasePassMovablePointLightInfo::FMobileBasePassMovablePointLightInfo(const
 
 				LightPositionAndInvRadius[NumMovablePointLights] = LightParameters.LightPositionAndInvRadius;
 				LightColorAndFalloffExponent[NumMovablePointLights] = LightParameters.LightColorAndFalloffExponent;
+				SpotLightDirection[NumMovablePointLights] = LightParameters.NormalizedLightDirection;
+				SpotLightAngles[NumMovablePointLights].Set(LightParameters.SpotAngles.X, LightParameters.SpotAngles.Y, 0.f, LightType == LightType_Spot ? 1.0f : 0.0f);
+
+				if (LightType == LightType_Rect)
+				{
+					// Treat rect lights as point lights.
+					LightColorAndFalloffExponent[NumMovablePointLights] = FVector4(FVector(LightProxy->GetColor()), LightParameters.LightColorAndFalloffExponent.W);
+				}
 
 				if (LightProxy->IsInverseSquared())
 				{
diff --git a/Engine/Source/Runtime/Renderer/Private/MobileBasePassRendering.h b/Engine/Source/Runtime/Renderer/Private/MobileBasePassRendering.h
--- a/Engine/Source/Runtime/Renderer/Private/MobileBasePassRendering.h
+++ b/Engine/Source/Runtime/Renderer/Private/MobileBasePassRendering.h
@@ -49,15 +49,17 @@ enum EOutputFormat
 
 const FLightSceneInfo* GetSceneMobileDirectionalLights(FScene const* Scene, uint32 LightChannel);
 
-/* Info for dynamic point lights rendered in base pass */
-class FMobileBasePassMovablePointLightInfo
+/* Info for dynamic point or spot lights rendered in base pass */
+class FMobileBasePassMovableLightInfo
 {
 public:
-	FMobileBasePassMovablePointLightInfo(const FPrimitiveSceneProxy* InSceneProxy);
+	FMobileBasePassMovableLightInfo(const FPrimitiveSceneProxy* InSceneProxy);
 
 	int32 NumMovablePointLights;
 	FVector4 LightPositionAndInvRadius[MAX_BASEPASS_DYNAMIC_POINT_LIGHTS];
 	FVector4 LightColorAndFalloffExponent[MAX_BASEPASS_DYNAMIC_POINT_LIGHTS];
+	FVector4 SpotLightDirection[MAX_BASEPASS_DYNAMIC_POINT_LIGHTS];
+	FVector4 SpotLightAngles[MAX_BASEPASS_DYNAMIC_POINT_LIGHTS];
 };
 
 static bool ShouldCacheShaderByPlatformAndOutputFormat(EShaderPlatform Platform, EOutputFormat OutputFormat)
@@ -216,7 +218,7 @@ public:
 		InvReflectionCubemapAverageBrightness.Bind(Initializer.ParameterMap, TEXT("InvReflectionCubemapAverageBrightness"));
 		LightPositionAndInvRadiusParameter.Bind(Initializer.ParameterMap, TEXT("LightPositionAndInvRadius"));
 		LightColorAndFalloffExponentParameter.Bind(Initializer.ParameterMap, TEXT("LightColorAndFalloffExponent"));
-			NumDynamicPointLightsParameter.Bind(Initializer.ParameterMap, TEXT("NumDynamicPointLights"));
+		NumDynamicPointLightsParameter.Bind(Initializer.ParameterMap, TEXT("NumDynamicPointLights"));
 		ReflectionPositionsAndRadii.Bind(Initializer.ParameterMap, TEXT("ReflectionPositionsAndRadii"));
 		ReflectionCubemap1.Bind(Initializer.ParameterMap, TEXT("ReflectionCubemap1"));
 		ReflectionSampler1.Bind(Initializer.ParameterMap, TEXT("ReflectionCubemapSampler1"));
@@ -228,6 +230,9 @@ public:
 		CSMDebugHintParams.Bind(Initializer.ParameterMap, TEXT("CSMDebugHint"));
 
 		PlanarReflectionParams.Bind(Initializer.ParameterMap);
+
+		SpotLightAnglesParameter.Bind(Initializer.ParameterMap, TEXT("SpotLightAngles"));
+		SpotLightDirectionParameter.Bind(Initializer.ParameterMap, TEXT("SpotLightDirection"));
 	}
 	TMobileBasePassPSPolicyParamType() {}
 
@@ -344,7 +349,7 @@ public:
 
 		if (NumMovablePointLights > 0)
 		{
-			FMobileBasePassMovablePointLightInfo LightInfo(Proxy);
+			FMobileBasePassMovableLightInfo LightInfo(Proxy);
 
 			if (NumMovablePointLights == INT32_MAX)
 			{
@@ -354,6 +359,8 @@ public:
 			// Set dynamic point lights
 			SetShaderValueArray(RHICmdList, PixelShader, LightPositionAndInvRadiusParameter, LightInfo.LightPositionAndInvRadius, LightInfo.NumMovablePointLights);
 			SetShaderValueArray(RHICmdList, PixelShader, LightColorAndFalloffExponentParameter, LightInfo.LightColorAndFalloffExponent, LightInfo.NumMovablePointLights);
+			SetShaderValueArray(RHICmdList, PixelShader, SpotLightDirectionParameter, LightInfo.SpotLightDirection, LightInfo.NumMovablePointLights);
+			SetShaderValueArray(RHICmdList, PixelShader, SpotLightAnglesParameter, LightInfo.SpotLightAngles, LightInfo.NumMovablePointLights);
 		}
 
 		if (CSMDebugHintParams.IsBound())
@@ -392,6 +399,9 @@ public:
 		Ar << PlanarReflectionParams;
 
 		Ar << CSMDebugHintParams;
+		
+		Ar << SpotLightAnglesParameter;
+		Ar << SpotLightDirectionParameter;
 
 		return bShaderHasOutdatedParameters;
 	}
@@ -407,6 +417,8 @@ private:
 	FShaderParameter LightPositionAndInvRadiusParameter;
 	FShaderParameter MobileSkyReflectionParam;
 	FShaderParameter LightColorAndFalloffExponentParameter;
+	FShaderParameter SpotLightDirectionParameter;
+	FShaderParameter SpotLightAnglesParameter;
 	FShaderParameter NumDynamicPointLightsParameter;
 
 
diff --git a/Engine/Source/Runtime/Renderer/Private/SceneCore.cpp b/Engine/Source/Runtime/Renderer/Private/SceneCore.cpp
--- a/Engine/Source/Runtime/Renderer/Private/SceneCore.cpp
+++ b/Engine/Source/Runtime/Renderer/Private/SceneCore.cpp
@@ -165,7 +165,7 @@ FLightPrimitiveInteraction::FLightPrimitiveInteraction(
 	bHasTranslucentObjectShadow(bInHasTranslucentObjectShadow),
 	bHasInsetObjectShadow(bInHasInsetObjectShadow),
 	bSelfShadowOnly(false),
-	bES2DynamicPointLight(false)
+	bMobileDynamicPointLight(false)
 {
 	// Determine whether this light-primitive interaction produces a shadow.
 	if(PrimitiveSceneInfo->Proxy->HasStaticLighting())
@@ -224,13 +224,24 @@ FLightPrimitiveInteraction::FLightPrimitiveInteraction(
 		// Add the interaction to the light's interaction list.
 		PrevPrimitiveLink = PrimitiveSceneInfo->Proxy->IsMeshShapeOftenMoving() ? &LightSceneInfo->DynamicInteractionOftenMovingPrimitiveList : &LightSceneInfo->DynamicInteractionStaticPrimitiveList;
 
-		// mobile movable point lights
-		if (PrimitiveSceneInfo->Scene->GetFeatureLevel() < ERHIFeatureLevel::SM4 && LightSceneInfo->Proxy->GetLightType() == LightType_Point && LightSceneInfo->Proxy->IsMovable())
+		// mobile movable spotlights / point lights
+		if (PrimitiveSceneInfo->Scene->GetShadingPath() == EShadingPath::Mobile && LightSceneInfo->Proxy->IsMovable())
 		{
-			bES2DynamicPointLight = true;
-			PrimitiveSceneInfo->NumMobileMovablePointLights++;
-			// The mobile renderer needs to use a different shader for movable point lights, so we have to update any static meshes in drawlists
-			PrimitiveSceneInfo->BeginDeferredUpdateStaticMeshes();
+			const uint8 LightType = LightSceneInfo->Proxy->GetLightType();
+			static const auto CVarMobileEnableMovableSpotLights = IConsoleManager::Get().FindTConsoleVariableDataInt(TEXT("r.Mobile.EnableMovableSpotLights"));
+
+			const bool bIsValidLightType = 
+				   LightType == LightType_Rect
+				|| LightType == LightType_Point
+				|| (LightType == LightType_Spot && CVarMobileEnableMovableSpotLights->GetValueOnRenderThread());
+
+			if( bIsValidLightType )
+			{
+				bMobileDynamicPointLight = true;
+				PrimitiveSceneInfo->NumMobileMovablePointLights++;
+				// The mobile renderer needs to use a different shader for movable point lights, so we have to update any static meshes in drawlists
+				PrimitiveSceneInfo->BeginDeferredUpdateStaticMeshes();
+			} 
 		}
 	}
 
@@ -272,7 +283,7 @@ FLightPrimitiveInteraction::~FLightPrimitiveInteraction()
 	FlushCachedShadowMapData();
 
 	// Track mobile movable point light count
-	if (bES2DynamicPointLight)
+	if (bMobileDynamicPointLight)
 	{
 		PrimitiveSceneInfo->NumMobileMovablePointLights--;
 		// The mobile renderer needs to use a different shader for movable point lights, so we have to update any static meshes in drawlists
diff --git a/Engine/Source/Runtime/Renderer/Private/SceneCore.h b/Engine/Source/Runtime/Renderer/Private/SceneCore.h
--- a/Engine/Source/Runtime/Renderer/Private/SceneCore.h
+++ b/Engine/Source/Runtime/Renderer/Private/SceneCore.h
@@ -41,7 +41,7 @@ public:
 	bool HasTranslucentObjectShadow() const { return bHasTranslucentObjectShadow; }
 	bool HasInsetObjectShadow() const { return bHasInsetObjectShadow; }
 	bool CastsSelfShadowOnly() const { return bSelfShadowOnly; }
-	bool IsES2DynamicPointLight() const { return bES2DynamicPointLight; }
+	bool IsMobileDynamicPointLight() const { return bMobileDynamicPointLight; }
 	FLightSceneInfo* GetLight() const { return LightSceneInfo; }
 	int32 GetLightId() const { return LightId; }
 	FPrimitiveSceneInfo* GetPrimitiveSceneInfo() const { return PrimitiveSceneInfo; }
@@ -108,7 +108,7 @@ private:
 	uint32 bSelfShadowOnly : 1;
 
 	/** True this is an ES2 dynamic point light interaction. */
-	uint32 bES2DynamicPointLight : 1;	
+	uint32 bMobileDynamicPointLight : 1;
 
 	/** Initialization constructor. */
 	FLightPrimitiveInteraction(FLightSceneInfo* InLightSceneInfo,FPrimitiveSceneInfo* InPrimitiveSceneInfo,
diff --git a/Engine/Source/Runtime/ShaderCore/Private/Shader.cpp b/Engine/Source/Runtime/ShaderCore/Private/Shader.cpp
--- a/Engine/Source/Runtime/ShaderCore/Private/Shader.cpp
+++ b/Engine/Source/Runtime/ShaderCore/Private/Shader.cpp
@@ -1997,6 +1997,11 @@ void ShaderMapAppendKeyString(EShaderPlatform Platform, FString& KeyString)
 			static IConsoleVariable* CVar = IConsoleManager::Get().FindConsoleVariable(TEXT("OpenGL.UseEmulatedUBs"));
 			KeyString += (CVar && CVar->GetInt() != 0) ? TEXT("_NoUB") : TEXT("");
 		}
+
+		{
+			static IConsoleVariable* CVar = IConsoleManager::Get().FindConsoleVariable(TEXT("r.Mobile.EnableMovableSpotlights"));
+			KeyString += (CVar && CVar->GetInt() != 0) ? TEXT("_MSPTL") : TEXT("");
+		}
 	}
 
 	if (Platform == SP_PS4)
