Removed redundant null check tripping up static analysis

diff --git a/Engine/Source/Runtime/Renderer/Private/SceneRendering.cpp b/Engine/Source/Runtime/Renderer/Private/SceneRendering.cpp
--- a/Engine/Source/Runtime/Renderer/Private/SceneRendering.cpp
+++ b/Engine/Source/Runtime/Renderer/Private/SceneRendering.cpp
@@ -1338,7 +1338,7 @@ void FViewInfo::SetupUniformBufferParameters(
 	ViewUniformShaderParameters.MobilePreviewMode = Scene == nullptr ? 0.0f: (IsSimulatedPlatform(Scene->GetShaderPlatform()) ? 1.0f : 0.0f);
 
 	// Padding between the left and right eye may be introduced by an HMD, which instanced stereo needs to account for.
-	if ((Family != nullptr) && (StereoPass != eSSP_FULL) && (Family->Views.Num() > 1))
+	if ((StereoPass != eSSP_FULL) && (Family->Views.Num() > 1))
 	{
 		check(Family->Views.Num() >= 2);
 
