Tag potential blocking calls for HMDs with profiling macros.

diff --git a/Engine/Source/Runtime/Apple/MetalRHI/Private/MetalContext.cpp b/Engine/Source/Runtime/Apple/MetalRHI/Private/MetalContext.cpp
--- a/Engine/Source/Runtime/Apple/MetalRHI/Private/MetalContext.cpp
+++ b/Engine/Source/Runtime/Apple/MetalRHI/Private/MetalContext.cpp
@@ -598,7 +598,10 @@ void FMetalDeviceContext::EndDrawingViewport(FMetalViewport* Viewport, bool bPre
 		if (CustomPresent != nullptr)
 		{
 			int32 SyncInterval = 0;
-			CustomPresent->Present(SyncInterval);
+			{
+				SCOPE_CYCLE_COUNTER(STAT_MetalCustomPresentTime);
+				CustomPresent->Present(SyncInterval);
+			}
 			
 			mtlpp::CommandBuffer CurrentCommandBuffer = GetCurrentCommandBuffer();
 			check(CurrentCommandBuffer);
diff --git a/Engine/Source/Runtime/Apple/MetalRHI/Private/MetalProfiler.cpp b/Engine/Source/Runtime/Apple/MetalRHI/Private/MetalProfiler.cpp
--- a/Engine/Source/Runtime/Apple/MetalRHI/Private/MetalProfiler.cpp
+++ b/Engine/Source/Runtime/Apple/MetalRHI/Private/MetalProfiler.cpp
@@ -29,6 +29,7 @@ DEFINE_STAT(STAT_MetalTexturePageOffTime);
 DEFINE_STAT(STAT_MetalGPUWorkTime);
 DEFINE_STAT(STAT_MetalGPUIdleTime);
 DEFINE_STAT(STAT_MetalPresentTime);
+DEFINE_STAT(STAT_MetalCustomPresentTime);
 DEFINE_STAT(STAT_MetalCommandBufferCreatedPerFrame);
 DEFINE_STAT(STAT_MetalCommandBufferCommittedPerFrame);
 DEFINE_STAT(STAT_MetalBufferMemory);
diff --git a/Engine/Source/Runtime/Apple/MetalRHI/Private/MetalProfiler.h b/Engine/Source/Runtime/Apple/MetalRHI/Private/MetalProfiler.h
--- a/Engine/Source/Runtime/Apple/MetalRHI/Private/MetalProfiler.h
+++ b/Engine/Source/Runtime/Apple/MetalRHI/Private/MetalProfiler.h
@@ -37,6 +37,7 @@ DECLARE_CYCLE_STAT_EXTERN(TEXT("Texture Page-On time"), STAT_MetalTexturePageOnT
 DECLARE_CYCLE_STAT_EXTERN(TEXT("GPU Work time"), STAT_MetalGPUWorkTime, STATGROUP_MetalRHI, );
 DECLARE_CYCLE_STAT_EXTERN(TEXT("GPU Idle time"), STAT_MetalGPUIdleTime, STATGROUP_MetalRHI, );
 DECLARE_CYCLE_STAT_EXTERN(TEXT("Present time"), STAT_MetalPresentTime, STATGROUP_MetalRHI, );
+DECLARE_CYCLE_STAT_EXTERN(TEXT("CustomPresent time"), STAT_MetalCustomPresentTime, STATGROUP_MetalRHI, );
 #if STATS
 extern int64 volatile GMetalTexturePageOnTime;
 extern int64 volatile GMetalGPUWorkTime;
diff --git a/Engine/Source/Runtime/D3D12RHI/Private/D3D12Stats.h b/Engine/Source/Runtime/D3D12RHI/Private/D3D12Stats.h
--- a/Engine/Source/Runtime/D3D12RHI/Private/D3D12Stats.h
+++ b/Engine/Source/Runtime/D3D12RHI/Private/D3D12Stats.h
@@ -10,6 +10,7 @@ D3D12Stats.h: D3D12 Statistics and Timing Interfaces
 */
 
 DECLARE_CYCLE_STAT_EXTERN(TEXT("Present time"), STAT_D3D12PresentTime, STATGROUP_D3D12RHI, );
+DECLARE_CYCLE_STAT_EXTERN(TEXT("CustomPresent time"), STAT_D3D12CustomPresentTime, STATGROUP_D3D12RHI, );
 
 DECLARE_DWORD_ACCUMULATOR_STAT_EXTERN(TEXT("Num command allocators (3D, Compute, Copy)"), STAT_D3D12NumCommandAllocators, STATGROUP_D3D12RHI, );
 DECLARE_DWORD_ACCUMULATOR_STAT_EXTERN(TEXT("Num command lists (3D, Compute, Copy)"), STAT_D3D12NumCommandLists, STATGROUP_D3D12RHI, );
diff --git a/Engine/Source/Runtime/D3D12RHI/Private/D3D12Util.cpp b/Engine/Source/Runtime/D3D12RHI/Private/D3D12Util.cpp
--- a/Engine/Source/Runtime/D3D12RHI/Private/D3D12Util.cpp
+++ b/Engine/Source/Runtime/D3D12RHI/Private/D3D12Util.cpp
@@ -732,6 +732,7 @@ bool AssertResourceState(ID3D12CommandList* pCommandList, FD3D12Resource* pResou
 //
 
 DEFINE_STAT(STAT_D3D12PresentTime);
+DEFINE_STAT(STAT_D3D12CustomPresentTime);
 
 DEFINE_STAT(STAT_D3D12NumCommandAllocators);
 DEFINE_STAT(STAT_D3D12NumCommandLists);
diff --git a/Engine/Source/Runtime/D3D12RHI/Private/D3D12Viewport.cpp b/Engine/Source/Runtime/D3D12RHI/Private/D3D12Viewport.cpp
--- a/Engine/Source/Runtime/D3D12RHI/Private/D3D12Viewport.cpp
+++ b/Engine/Source/Runtime/D3D12RHI/Private/D3D12Viewport.cpp
@@ -495,6 +495,7 @@ bool FD3D12Viewport::PresentChecked(int32 SyncInterval)
 
 	if (IsValidRef(CustomPresent))
 	{
+		SCOPE_CYCLE_COUNTER(STAT_D3D12CustomPresentTime);
 		bNeedNativePresent = CustomPresent->Present(SyncInterval);
 	}
 	if (bNeedNativePresent)
diff --git a/Engine/Source/Runtime/OpenGLDrv/Private/Android/AndroidOpenGL.cpp b/Engine/Source/Runtime/OpenGLDrv/Private/Android/AndroidOpenGL.cpp
--- a/Engine/Source/Runtime/OpenGLDrv/Private/Android/AndroidOpenGL.cpp
+++ b/Engine/Source/Runtime/OpenGLDrv/Private/Android/AndroidOpenGL.cpp
@@ -180,6 +180,7 @@ bool PlatformBlitToViewport( FPlatformOpenGLDevice* Device, const FOpenGLViewpor
 {
 	if (bPresent && Viewport.GetCustomPresent())
 	{
+		QUICK_SCOPE_CYCLE_COUNTER(STAT_FAndroidOpenGL_PlatformBlitToViewport_CustomPresent);
 		bPresent = Viewport.GetCustomPresent()->Present(SyncInterval);
 	}
 	if (bPresent)
diff --git a/Engine/Source/Runtime/Renderer/Private/DeferredShadingRenderer.cpp b/Engine/Source/Runtime/Renderer/Private/DeferredShadingRenderer.cpp
--- a/Engine/Source/Runtime/Renderer/Private/DeferredShadingRenderer.cpp
+++ b/Engine/Source/Runtime/Renderer/Private/DeferredShadingRenderer.cpp
@@ -145,6 +145,7 @@ DECLARE_CYCLE_STAT(TEXT("DeferredShadingSceneRenderer RenderAtmosphere"), STAT_F
 DECLARE_CYCLE_STAT(TEXT("DeferredShadingSceneRenderer RenderFog"), STAT_FDeferredShadingSceneRenderer_RenderFog, STATGROUP_SceneRendering);
 DECLARE_CYCLE_STAT(TEXT("DeferredShadingSceneRenderer RenderLightShaftBloom"), STAT_FDeferredShadingSceneRenderer_RenderLightShaftBloom, STATGROUP_SceneRendering);
 DECLARE_CYCLE_STAT(TEXT("DeferredShadingSceneRenderer RenderFinish"), STAT_FDeferredShadingSceneRenderer_RenderFinish, STATGROUP_SceneRendering);
+DECLARE_CYCLE_STAT(TEXT("DeferredShadingSceneRenderer ViewExtensionPostRenderBasePass"), STAT_FDeferredShadingSceneRenderer_ViewExtensionPostRenderBasePass, STATGROUP_SceneRendering);
 
 DECLARE_GPU_STAT(Postprocessing);
 DECLARE_GPU_STAT(HZB);
@@ -968,12 +969,15 @@ void FDeferredShadingSceneRenderer::Render(FRHICommandListImmediate& RHICmdList)
 	RenderBasePass(RHICmdList, BasePassDepthStencilAccess, ForwardScreenSpaceShadowMask.GetReference());
 	RHICmdList.SetCurrentStat(GET_STATID(STAT_CLM_AfterBasePass));
 	ServiceLocalQueue();
-	
-	for (int32 ViewExt = 0; ViewExt < ViewFamily.ViewExtensions.Num(); ++ViewExt)
+
 	{
-		for (int32 ViewIndex = 0; ViewIndex < ViewFamily.Views.Num(); ++ViewIndex)
+		SCOPE_CYCLE_COUNTER(STAT_FDeferredShadingSceneRenderer_ViewExtensionPostRenderBasePass);
+		for (int32 ViewExt = 0; ViewExt < ViewFamily.ViewExtensions.Num(); ++ViewExt)
 		{
-			ViewFamily.ViewExtensions[ViewExt]->PostRenderBasePass_RenderThread(RHICmdList, Views[ViewIndex]);
+			for (int32 ViewIndex = 0; ViewIndex < ViewFamily.Views.Num(); ++ViewIndex)
+			{
+				ViewFamily.ViewExtensions[ViewExt]->PostRenderBasePass_RenderThread(RHICmdList, Views[ViewIndex]);
+			}
 		}
 	}
 
diff --git a/Engine/Source/Runtime/Renderer/Private/MobileShadingRenderer.cpp b/Engine/Source/Runtime/Renderer/Private/MobileShadingRenderer.cpp
--- a/Engine/Source/Runtime/Renderer/Private/MobileShadingRenderer.cpp
+++ b/Engine/Source/Runtime/Renderer/Private/MobileShadingRenderer.cpp
@@ -274,11 +274,14 @@ void FMobileSceneRenderer::Render(FRHICommandListImmediate& RHICmdList)
 
 	RHICmdList.SetCurrentStat(GET_STATID(STAT_CLMM_Post));
 
-	for (int32 ViewExt = 0; ViewExt < ViewFamily.ViewExtensions.Num(); ++ViewExt)
 	{
-		for (int32 ViewIndex = 0; ViewIndex < ViewFamily.Views.Num(); ++ViewIndex)
+		QUICK_SCOPE_CYCLE_COUNTER(STAT_FMobileSceneRenderer_ViewExtensionPostRenderBasePass);
+		for (int32 ViewExt = 0; ViewExt < ViewFamily.ViewExtensions.Num(); ++ViewExt)
 		{
-			ViewFamily.ViewExtensions[ViewExt]->PostRenderBasePass_RenderThread(RHICmdList, Views[ViewIndex]);
+			for (int32 ViewIndex = 0; ViewIndex < ViewFamily.Views.Num(); ++ViewIndex)
+			{
+				ViewFamily.ViewExtensions[ViewExt]->PostRenderBasePass_RenderThread(RHICmdList, Views[ViewIndex]);
+			}
 		}
 	}
 
diff --git a/Engine/Source/Runtime/Renderer/Private/SceneRendering.cpp b/Engine/Source/Runtime/Renderer/Private/SceneRendering.cpp
--- a/Engine/Source/Runtime/Renderer/Private/SceneRendering.cpp
+++ b/Engine/Source/Runtime/Renderer/Private/SceneRendering.cpp
@@ -324,6 +324,8 @@ FGraphEventRef FSceneRenderer::OcclusionSubmittedFence[FOcclusionQueryHelpers::M
 
 
 DECLARE_CYCLE_STAT(TEXT("DeferredShadingSceneRenderer UpdateMotionBlurCache"), STAT_FDeferredShadingSceneRenderer_UpdateMotionBlurCache, STATGROUP_SceneRendering);
+DECLARE_CYCLE_STAT(TEXT("DeferredShadingSceneRenderer ViewExtensionPostRenderView"), STAT_FDeferredShadingSceneRenderer_ViewExtensionPostRenderView, STATGROUP_SceneRendering);
+DECLARE_CYCLE_STAT(TEXT("DeferredShadingSceneRenderer ViewExtensionPreRenderView"), STAT_FDeferredShadingSceneRenderer_ViewExtensionPreRenderView, STATGROUP_SceneRendering);
 
 #define FASTVRAM_CVAR(Name,DefaultValue) static TAutoConsoleVariable<int32> CVarFastVRam_##Name(TEXT("r.FastVRam."#Name), DefaultValue, TEXT(""))
 
@@ -2425,12 +2427,15 @@ void FSceneRenderer::RenderFinish(FRHICommandListImmediate& RHICmdList)
 	}
 #endif
 
-	for(int32 ViewExt = 0; ViewExt < ViewFamily.ViewExtensions.Num(); ++ViewExt)
 	{
-		ViewFamily.ViewExtensions[ViewExt]->PostRenderViewFamily_RenderThread(RHICmdList, ViewFamily);
-		for(int32 ViewIndex = 0; ViewIndex < ViewFamily.Views.Num(); ++ViewIndex)
+		SCOPE_CYCLE_COUNTER(STAT_FDeferredShadingSceneRenderer_ViewExtensionPostRenderView);
+		for(int32 ViewExt = 0; ViewExt < ViewFamily.ViewExtensions.Num(); ++ViewExt)
 		{
-			ViewFamily.ViewExtensions[ViewExt]->PostRenderView_RenderThread(RHICmdList, Views[ViewIndex]);
+			ViewFamily.ViewExtensions[ViewExt]->PostRenderViewFamily_RenderThread(RHICmdList, ViewFamily);
+			for(int32 ViewIndex = 0; ViewIndex < ViewFamily.Views.Num(); ++ViewIndex)
+			{
+				ViewFamily.ViewExtensions[ViewExt]->PostRenderView_RenderThread(RHICmdList, Views[ViewIndex]);
+			}
 		}
 	}
 
@@ -2779,12 +2784,15 @@ static void ViewExtensionPreRender_RenderThread(FRHICommandListImmediate& RHICmd
 {
 	FMemMark MemStackMark(FMemStack::Get());
 
-	for (int ViewExt = 0; ViewExt < SceneRenderer->ViewFamily.ViewExtensions.Num(); ViewExt++)
 	{
-		SceneRenderer->ViewFamily.ViewExtensions[ViewExt]->PreRenderViewFamily_RenderThread(RHICmdList, SceneRenderer->ViewFamily);
-		for (int ViewIndex = 0; ViewIndex < SceneRenderer->ViewFamily.Views.Num(); ViewIndex++)
+		SCOPE_CYCLE_COUNTER(STAT_FDeferredShadingSceneRenderer_ViewExtensionPreRenderView);
+		for (int ViewExt = 0; ViewExt < SceneRenderer->ViewFamily.ViewExtensions.Num(); ViewExt++)
 		{
-			SceneRenderer->ViewFamily.ViewExtensions[ViewExt]->PreRenderView_RenderThread(RHICmdList, SceneRenderer->Views[ViewIndex]);
+			SceneRenderer->ViewFamily.ViewExtensions[ViewExt]->PreRenderViewFamily_RenderThread(RHICmdList, SceneRenderer->ViewFamily);
+			for (int ViewIndex = 0; ViewIndex < SceneRenderer->ViewFamily.Views.Num(); ViewIndex++)
+			{
+				SceneRenderer->ViewFamily.ViewExtensions[ViewExt]->PreRenderView_RenderThread(RHICmdList, SceneRenderer->Views[ViewIndex]);
+			}
 		}
 	}
 	
diff --git a/Engine/Source/Runtime/VulkanRHI/Private/VulkanRHIPrivate.h b/Engine/Source/Runtime/VulkanRHI/Private/VulkanRHIPrivate.h
--- a/Engine/Source/Runtime/VulkanRHI/Private/VulkanRHIPrivate.h
+++ b/Engine/Source/Runtime/VulkanRHI/Private/VulkanRHIPrivate.h
@@ -333,6 +333,7 @@ DECLARE_STATS_GROUP(TEXT("Vulkan RHI"), STATGROUP_VulkanRHI, STATCAT_Advanced);
 DECLARE_CYCLE_STAT_EXTERN(TEXT("Draw call time"), STAT_VulkanDrawCallTime, STATGROUP_VulkanRHI, );
 DECLARE_CYCLE_STAT_EXTERN(TEXT("Dispatch call time"), STAT_VulkanDispatchCallTime, STATGROUP_VulkanRHI, );
 DECLARE_CYCLE_STAT_EXTERN(TEXT("Draw call prep time"), STAT_VulkanDrawCallPrepareTime, STATGROUP_VulkanRHI, );
+DECLARE_CYCLE_STAT_EXTERN(TEXT("CustomPresent time"), STAT_VulkanCustomPresentTime, STATGROUP_VulkanRHI, );
 DECLARE_CYCLE_STAT_EXTERN(TEXT("Dispatch call prep time"), STAT_VulkanDispatchCallPrepareTime, STATGROUP_VulkanRHI, );
 DECLARE_CYCLE_STAT_EXTERN(TEXT("Get Or Create Pipeline"), STAT_VulkanGetOrCreatePipeline, STATGROUP_VulkanRHI, );
 DECLARE_CYCLE_STAT_EXTERN(TEXT("Get DescriptorSet"), STAT_VulkanGetDescriptorSet, STATGROUP_VulkanRHI, );
diff --git a/Engine/Source/Runtime/VulkanRHI/Private/VulkanUtil.cpp b/Engine/Source/Runtime/VulkanRHI/Private/VulkanUtil.cpp
--- a/Engine/Source/Runtime/VulkanRHI/Private/VulkanUtil.cpp
+++ b/Engine/Source/Runtime/VulkanRHI/Private/VulkanUtil.cpp
@@ -620,6 +620,7 @@ namespace VulkanRHI
 DEFINE_STAT(STAT_VulkanDrawCallTime);
 DEFINE_STAT(STAT_VulkanDispatchCallTime);
 DEFINE_STAT(STAT_VulkanDrawCallPrepareTime);
+DEFINE_STAT(STAT_VulkanCustomPresentTime);
 DEFINE_STAT(STAT_VulkanDispatchCallPrepareTime);
 DEFINE_STAT(STAT_VulkanGetOrCreatePipeline);
 DEFINE_STAT(STAT_VulkanGetDescriptorSet);
diff --git a/Engine/Source/Runtime/VulkanRHI/Private/VulkanViewport.cpp b/Engine/Source/Runtime/VulkanRHI/Private/VulkanViewport.cpp
--- a/Engine/Source/Runtime/VulkanRHI/Private/VulkanViewport.cpp
+++ b/Engine/Source/Runtime/VulkanRHI/Private/VulkanViewport.cpp
@@ -641,6 +641,7 @@ bool FVulkanViewport::Present(FVulkanCommandListContext* Context, FVulkanCmdBuff
 	const bool bHasCustomPresent = IsValidRef(CustomPresent);
 	if (bHasCustomPresent)
 	{
+		SCOPE_CYCLE_COUNTER(STAT_VulkanCustomPresentTime);
 		bNeedNativePresent = CustomPresent->Present(SyncInterval);
 	}
 
diff --git a/Engine/Source/Runtime/Windows/D3D11RHI/Private/D3D11RHIPrivate.h b/Engine/Source/Runtime/Windows/D3D11RHI/Private/D3D11RHIPrivate.h
--- a/Engine/Source/Runtime/Windows/D3D11RHI/Private/D3D11RHIPrivate.h
+++ b/Engine/Source/Runtime/Windows/D3D11RHI/Private/D3D11RHIPrivate.h
@@ -56,6 +56,7 @@ extern bool GDX11NVAfterMathEnabled;
  * The D3D RHI stats.
  */
 DECLARE_CYCLE_STAT_EXTERN(TEXT("Present time"),STAT_D3D11PresentTime,STATGROUP_D3D11RHI, );
+DECLARE_CYCLE_STAT_EXTERN(TEXT("CustomPresent time"), STAT_D3D11CustomPresentTime, STATGROUP_D3D11RHI, );
 DECLARE_CYCLE_STAT_EXTERN(TEXT("CreateTexture time"),STAT_D3D11CreateTextureTime,STATGROUP_D3D11RHI, );
 DECLARE_CYCLE_STAT_EXTERN(TEXT("LockTexture time"),STAT_D3D11LockTextureTime,STATGROUP_D3D11RHI, );
 DECLARE_CYCLE_STAT_EXTERN(TEXT("UnlockTexture time"),STAT_D3D11UnlockTextureTime,STATGROUP_D3D11RHI, );
diff --git a/Engine/Source/Runtime/Windows/D3D11RHI/Private/D3D11Util.cpp b/Engine/Source/Runtime/Windows/D3D11RHI/Private/D3D11Util.cpp
--- a/Engine/Source/Runtime/Windows/D3D11RHI/Private/D3D11Util.cpp
+++ b/Engine/Source/Runtime/Windows/D3D11RHI/Private/D3D11Util.cpp
@@ -507,6 +507,7 @@ ID3D11Buffer* FD3D11DynamicBuffer::Unlock()
 
 
 DEFINE_STAT(STAT_D3D11PresentTime);
+DEFINE_STAT(STAT_D3D11CustomPresentTime);
 DEFINE_STAT(STAT_D3D11TexturesAllocated);
 DEFINE_STAT(STAT_D3D11TexturesReleased);
 DEFINE_STAT(STAT_D3D11ClearShaderResourceTime);
diff --git a/Engine/Source/Runtime/Windows/D3D11RHI/Private/D3D11Viewport.cpp b/Engine/Source/Runtime/Windows/D3D11RHI/Private/D3D11Viewport.cpp
--- a/Engine/Source/Runtime/Windows/D3D11RHI/Private/D3D11Viewport.cpp
+++ b/Engine/Source/Runtime/Windows/D3D11RHI/Private/D3D11Viewport.cpp
@@ -287,6 +287,7 @@ bool FD3D11Viewport::PresentChecked(int32 SyncInterval)
 
 	if (IsValidRef(CustomPresent))
 	{
+		SCOPE_CYCLE_COUNTER(STAT_D3D11CustomPresentTime);
 		bNeedNativePresent = CustomPresent->Present(SyncInterval);
 	}
 
