Work around for csm shadow culling bug on Lumiun

diff --git a/Engine/Source/Runtime/Renderer/Private/ShadowSetupMobile.cpp b/Engine/Source/Runtime/Renderer/Private/ShadowSetupMobile.cpp
--- a/Engine/Source/Runtime/Renderer/Private/ShadowSetupMobile.cpp
+++ b/Engine/Source/Runtime/Renderer/Private/ShadowSetupMobile.cpp
@@ -22,9 +22,10 @@ static TAutoConsoleVariable<int32> CVarCsmShaderCullingDebugGfx(
 	TEXT(""),
 	ECVF_RenderThreadSafe);
 
+const uint32 CSMShaderCullingMethodDefault = !PLATFORM_LUMIN;
 static TAutoConsoleVariable<int32> CVarsCsmShaderCullingMethod(
 	TEXT("r.Mobile.Shadow.CSMShaderCullingMethod"),
-	1,
+	CSMShaderCullingMethodDefault,
 	TEXT("Method to determine which primitives will receive CSM shaders:\n")
 	TEXT("0 - disabled (all primitives will receive CSM)\n")
 	TEXT("1 - Light frustum, all primitives whose bounding box is within CSM receiving distance. (default)\n")
