[NO NON-PII COMMIT MESSAGE]

diff --git a/Engine/Source/Runtime/Renderer/Private/SceneVisibility.cpp b/Engine/Source/Runtime/Renderer/Private/SceneVisibility.cpp
--- a/Engine/Source/Runtime/Renderer/Private/SceneVisibility.cpp
+++ b/Engine/Source/Runtime/Renderer/Private/SceneVisibility.cpp
@@ -2149,12 +2149,6 @@ static void ComputeAndMarkRelevanceForViewParallel(
 
 	bool WillExecuteInParallel = FApp::ShouldUseThreadingForPerformance() && CVarParallelInitViews.GetValueOnRenderThread() > 0;
 
-	if (WillExecuteInParallel)
-	{
-		// We must reserve to prevent realloc otherwise it will cause memory leak if WillExecuteInParallel == true
-		View.PrimitiveCustomDataMemStack.Reserve(View.PrimitiveCustomDataMemStack.Num() + FMath::CeilToInt(((float)View.PrimitiveVisibilityMap.Num() / (float)FRelevancePrimSet<int32>::MaxInputPrims)));
-	}
-
 	{
 		FSceneSetBitIterator BitIt(View.PrimitiveVisibilityMap);
 		if (BitIt)
@@ -2871,7 +2865,10 @@ void FSceneRenderer::ComputeViewVisibility(FRHICommandListImmediate& RHICmdList)
 
 		View.PrimitivesCustomData.Init(nullptr, Scene->Primitives.Num());
 		View.PrimitivesWithCustomData.Reserve(Scene->Primitives.Num());
-		View.PrimitiveCustomDataMemStack.Reserve(1);
+
+		// We must reserve to prevent realloc otherwise it will cause memory leak if we Execute In Parallel
+		View.PrimitiveCustomDataMemStack.Reserve(FApp::ShouldUseThreadingForPerformance() && CVarParallelInitViews.GetValueOnRenderThread() > 0 ? FMath::CeilToInt(((float)View.PrimitiveVisibilityMap.Num() / (float)FRelevancePrimSet<int32>::MaxInputPrims)) + 1 : 1);
+
 		View.AllocateCustomDataMemStack();
 
 		View.VisibleLightInfos.Empty(Scene->Lights.GetMaxIndex());
