Renderpass updates.

Light rendering. These are self-contained.

diff --git a/Engine/Source/Runtime/Renderer/Private/LightFunctionRendering.cpp b/Engine/Source/Runtime/Renderer/Private/LightFunctionRendering.cpp
--- a/Engine/Source/Runtime/Renderer/Private/LightFunctionRendering.cpp
+++ b/Engine/Source/Runtime/Renderer/Private/LightFunctionRendering.cpp
@@ -281,110 +281,131 @@ bool FDeferredShadingSceneRenderer::RenderLightFunctionForMaterial(
 {
 	bool bRenderedLightFunction = false;
 
+	check(RHICmdList.IsOutsideRenderPass());
 	check(ScreenShadowMaskTexture);
-	if (MaterialProxy && MaterialProxy->GetMaterial(Scene->GetFeatureLevel())->IsLightFunction() )
+	if (MaterialProxy && MaterialProxy->GetMaterial(Scene->GetFeatureLevel())->IsLightFunction())
 	{
-		SetRenderTarget(RHICmdList, ScreenShadowMaskTexture->GetRenderTargetItem().TargetableTexture, FSceneRenderTargets::Get(RHICmdList).GetSceneDepthSurface(), ESimpleRenderTargetMode::EExistingColorAndDepth, FExclusiveDepthStencil::DepthRead_StencilWrite, true);
-
-		bRenderedLightFunction = true;
-
-		const FMaterial* Material = MaterialProxy->GetMaterial(Scene->GetFeatureLevel());
-		SCOPED_DRAW_EVENTF(RHICmdList, LightFunction, TEXT("LightFunction Material=%s"), *Material->GetFriendlyName());
-
-		const FMaterialShaderMap* MaterialShaderMap = Material->GetRenderingThreadShaderMap();
-		FLightFunctionVS* VertexShader = MaterialShaderMap->GetShader<FLightFunctionVS>();
-		FLightFunctionPS* PixelShader = MaterialShaderMap->GetShader<FLightFunctionPS>();
+		{
+			// Scope out this FRHIRenderPassInfo so we don't pollute ourselves elsewhere.
+			FRHIRenderPassInfo RPInfo(ScreenShadowMaskTexture->GetRenderTargetItem().TargetableTexture, ERenderTargetActions::Load_Store);
+			RPInfo.DepthStencilRenderTarget.Action = MakeDepthStencilTargetActions(ERenderTargetActions::Load_DontStore, ERenderTargetActions::Load_Store);
+			RPInfo.DepthStencilRenderTarget.DepthStencilTarget = FSceneRenderTargets::Get(RHICmdList).GetSceneDepthSurface();
+			RPInfo.DepthStencilRenderTarget.ExclusiveDepthStencil = FExclusiveDepthStencil::DepthRead_StencilWrite;
+
+			TransitionRenderPassTargets(RHICmdList, RPInfo);
+			RHICmdList.BeginRenderPass(RPInfo, TEXT("RenderLightFunctionForMaterial"));
+		}
+		{
+			bRenderedLightFunction = true;
 
-		FGraphicsPipelineStateInitializer GraphicsPSOInit;
-		RHICmdList.ApplyCachedRenderTargets(GraphicsPSOInit);
+			const FMaterial* Material = MaterialProxy->GetMaterial(Scene->GetFeatureLevel());
+			SCOPED_DRAW_EVENTF(RHICmdList, LightFunction, TEXT("LightFunction Material=%s"), *Material->GetFriendlyName());
 
-		GraphicsPSOInit.BoundShaderState.VertexDeclarationRHI = GetVertexDeclarationFVector4();
-		GraphicsPSOInit.BoundShaderState.VertexShaderRHI = VertexShader->GetVertexShader();
-		GraphicsPSOInit.BoundShaderState.PixelShaderRHI = PixelShader->GetPixelShader();
+			const FMaterialShaderMap* MaterialShaderMap = Material->GetRenderingThreadShaderMap();
+			FLightFunctionVS* VertexShader = MaterialShaderMap->GetShader<FLightFunctionVS>();
+			FLightFunctionPS* PixelShader = MaterialShaderMap->GetShader<FLightFunctionPS>();
 
-		FSphere LightBounds = LightSceneInfo->Proxy->GetBoundingSphere();
+			FGraphicsPipelineStateInitializer GraphicsPSOInit;
+			RHICmdList.ApplyCachedRenderTargets(GraphicsPSOInit);
 
-		// Render to the light attenuation buffer for all views.
-		for (int32 ViewIndex = 0; ViewIndex < Views.Num(); ViewIndex++)
-		{
-			SCOPED_CONDITIONAL_DRAW_EVENTF(RHICmdList, EventView, Views.Num() > 1, TEXT("View%d"), ViewIndex);
+			GraphicsPSOInit.BoundShaderState.VertexDeclarationRHI = GetVertexDeclarationFVector4();
+			GraphicsPSOInit.BoundShaderState.VertexShaderRHI = VertexShader->GetVertexShader();
+			GraphicsPSOInit.BoundShaderState.PixelShaderRHI = PixelShader->GetPixelShader();
 
-			const FViewInfo& View = Views[ViewIndex];
+			FSphere LightBounds = LightSceneInfo->Proxy->GetBoundingSphere();
 
-			if (View.VisibleLightInfos[LightSceneInfo->Id].bInViewFrustum)
+			// Render to the light attenuation buffer for all views.
+			for (int32 ViewIndex = 0; ViewIndex < Views.Num(); ViewIndex++)
 			{
-				if (LightSceneInfo->Proxy->GetLightType() == LightType_Directional)
-				{
-					LightBounds.Center = View.ViewMatrices.GetViewOrigin();
-				}
+				SCOPED_CONDITIONAL_DRAW_EVENTF(RHICmdList, EventView, Views.Num() > 1, TEXT("View%d"), ViewIndex);
+
+				const FViewInfo& View = Views[ViewIndex];
 
-				const float FadeAlpha = GetLightFunctionFadeFraction(View, LightBounds);
-				// Don't draw the light function if it has completely faded out
-				if (FadeAlpha < 1.0f / 256.0f)
+				if (View.VisibleLightInfos[LightSceneInfo->Id].bInViewFrustum)
 				{
-					if( !bLightAttenuationCleared )
+					if (LightSceneInfo->Proxy->GetLightType() == LightType_Directional)
 					{
-						LightSceneInfo->Proxy->SetScissorRect(RHICmdList, View, View.ViewRect);
-						SetRenderTarget(RHICmdList, ScreenShadowMaskTexture->GetRenderTargetItem().TargetableTexture, FSceneRenderTargets::Get(RHICmdList).GetSceneDepthSurface(), ESimpleRenderTargetMode::EClearColorExistingDepth, FExclusiveDepthStencil::DepthRead_StencilWrite, true);
+						LightBounds.Center = View.ViewMatrices.GetViewOrigin();
 					}
-				}
-				else
-				{
-					// Set the device viewport for the view.
-					RHICmdList.SetViewport(View.ViewRect.Min.X, View.ViewRect.Min.Y, 0.0f, View.ViewRect.Max.X, View.ViewRect.Max.Y, 1.0f);
 
-					// Set the states to modulate the light function with the render target.
-					GraphicsPSOInit.DepthStencilState = TStaticDepthStencilState<false,CF_Always>::GetRHI();
-					GraphicsPSOInit.PrimitiveType = PT_TriangleList;
+					const float FadeAlpha = GetLightFunctionFadeFraction(View, LightBounds);
+					// Don't draw the light function if it has completely faded out
+					if (FadeAlpha < 1.0f / 256.0f)
+					{
+						if (!bLightAttenuationCleared)
+						{
+							LightSceneInfo->Proxy->SetScissorRect(RHICmdList, View, View.ViewRect);
+
+							// #todo-renderpasses This is extremely gross with renderpasses. Should refactor this whole thing.
+							RHICmdList.EndRenderPass();
+							FRHIRenderPassInfo RPInfoClear(ScreenShadowMaskTexture->GetRenderTargetItem().TargetableTexture, ERenderTargetActions::Clear_Store);
+							RPInfoClear.DepthStencilRenderTarget.Action = MakeDepthStencilTargetActions(ERenderTargetActions::Load_DontStore, ERenderTargetActions::Load_Store);
+							RPInfoClear.DepthStencilRenderTarget.DepthStencilTarget = FSceneRenderTargets::Get(RHICmdList).GetSceneDepthSurface();
+							RPInfoClear.DepthStencilRenderTarget.ExclusiveDepthStencil = FExclusiveDepthStencil::DepthRead_StencilWrite;
 
-					if( bLightAttenuationCleared )
+							TransitionRenderPassTargets(RHICmdList, RPInfoClear);
+							RHICmdList.BeginRenderPass(RPInfoClear, TEXT("RenderLightFunction_PlusClearAttenuation"));
+						}
+					}
+					else
 					{
-						if (bRenderingPreviewShadowsIndicator)
+						// Set the device viewport for the view.
+						RHICmdList.SetViewport(View.ViewRect.Min.X, View.ViewRect.Min.Y, 0.0f, View.ViewRect.Max.X, View.ViewRect.Max.Y, 1.0f);
+
+						// Set the states to modulate the light function with the render target.
+						GraphicsPSOInit.DepthStencilState = TStaticDepthStencilState<false, CF_Always>::GetRHI();
+						GraphicsPSOInit.PrimitiveType = PT_TriangleList;
+
+						if (bLightAttenuationCleared)
 						{
-							GraphicsPSOInit.BlendState = TStaticBlendState<CW_RGBA,BO_Max,BF_One,BF_One,BO_Max,BF_One,BF_One>::GetRHI();
+							if (bRenderingPreviewShadowsIndicator)
+							{
+								GraphicsPSOInit.BlendState = TStaticBlendState<CW_RGBA, BO_Max, BF_One, BF_One, BO_Max, BF_One, BF_One>::GetRHI();
+							}
+							else
+							{
+								FProjectedShadowInfo::SetBlendStateForProjection(
+									GraphicsPSOInit,
+									LightSceneInfo->GetDynamicShadowMapChannel(),
+									false,
+									false,
+									bProjectingForForwardShading,
+									false);
+							}
 						}
 						else
 						{
-							FProjectedShadowInfo::SetBlendStateForProjection(
-								GraphicsPSOInit,
-								LightSceneInfo->GetDynamicShadowMapChannel(), 
-								false,
-								false,
-								bProjectingForForwardShading, 
-								false);
+							GraphicsPSOInit.BlendState = TStaticBlendState<CW_RGBA>::GetRHI();
 						}
-					}
-					else
-					{
-						GraphicsPSOInit.BlendState = TStaticBlendState<CW_RGBA>::GetRHI();
-					}
 
-					if (((FVector)View.ViewMatrices.GetViewOrigin() - LightBounds.Center).SizeSquared() < FMath::Square(LightBounds.W * 1.05f + View.NearClippingDistance * 2.0f))
-					{
-						// Render backfaces with depth tests disabled since the camera is inside (or close to inside) the light function geometry
-						GraphicsPSOInit.RasterizerState = View.bReverseCulling ? TStaticRasterizerState<FM_Solid, CM_CW>::GetRHI() : TStaticRasterizerState<FM_Solid, CM_CCW>::GetRHI();
-					}
-					else
-					{
-						// Render frontfaces with depth tests on to get the speedup from HiZ since the camera is outside the light function geometry
-						GraphicsPSOInit.DepthStencilState = TStaticDepthStencilState<false,CF_DepthNearOrEqual>::GetRHI();
-						GraphicsPSOInit.RasterizerState = View.bReverseCulling ? TStaticRasterizerState<FM_Solid, CM_CCW>::GetRHI() : TStaticRasterizerState<FM_Solid, CM_CW>::GetRHI();
-					}
+						if (((FVector)View.ViewMatrices.GetViewOrigin() - LightBounds.Center).SizeSquared() < FMath::Square(LightBounds.W * 1.05f + View.NearClippingDistance * 2.0f))
+						{
+							// Render backfaces with depth tests disabled since the camera is inside (or close to inside) the light function geometry
+							GraphicsPSOInit.RasterizerState = View.bReverseCulling ? TStaticRasterizerState<FM_Solid, CM_CW>::GetRHI() : TStaticRasterizerState<FM_Solid, CM_CCW>::GetRHI();
+						}
+						else
+						{
+							// Render frontfaces with depth tests on to get the speedup from HiZ since the camera is outside the light function geometry
+							GraphicsPSOInit.DepthStencilState = TStaticDepthStencilState<false, CF_DepthNearOrEqual>::GetRHI();
+							GraphicsPSOInit.RasterizerState = View.bReverseCulling ? TStaticRasterizerState<FM_Solid, CM_CCW>::GetRHI() : TStaticRasterizerState<FM_Solid, CM_CW>::GetRHI();
+						}
 
-					// Set the light's scissor rectangle.
-					LightSceneInfo->Proxy->SetScissorRect(RHICmdList, View, View.ViewRect);
+						// Set the light's scissor rectangle.
+						LightSceneInfo->Proxy->SetScissorRect(RHICmdList, View, View.ViewRect);
 
-					// Render a bounding light sphere.
-					SetGraphicsPipelineState(RHICmdList, GraphicsPSOInit);
-					VertexShader->SetParameters(RHICmdList, View, LightSceneInfo);
-					PixelShader->SetParameters(RHICmdList, View, LightSceneInfo, MaterialProxy, bRenderingPreviewShadowsIndicator, FadeAlpha);
+						// Render a bounding light sphere.
+						SetGraphicsPipelineState(RHICmdList, GraphicsPSOInit);
+						VertexShader->SetParameters(RHICmdList, View, LightSceneInfo);
+						PixelShader->SetParameters(RHICmdList, View, LightSceneInfo, MaterialProxy, bRenderingPreviewShadowsIndicator, FadeAlpha);
 
-					// Project the light function using a sphere around the light
-					//@todo - could use a cone for spotlights
-					StencilingGeometry::DrawSphere(RHICmdList);
+						// Project the light function using a sphere around the light
+						//@todo - could use a cone for spotlights
+						StencilingGeometry::DrawSphere(RHICmdList);
+					}
 				}
 			}
 		}
+		RHICmdList.EndRenderPass();
 
 		// Restore states.
 		RHICmdList.SetScissorRect(false, 0, 0, 0, 0);
diff --git a/Engine/Source/Runtime/Renderer/Private/LightGridInjection.cpp b/Engine/Source/Runtime/Renderer/Private/LightGridInjection.cpp
--- a/Engine/Source/Runtime/Renderer/Private/LightGridInjection.cpp
+++ b/Engine/Source/Runtime/Renderer/Private/LightGridInjection.cpp
@@ -810,6 +810,8 @@ void FDeferredShadingSceneRenderer::ComputeLightGrid(FRHICommandListImmediate& R
 
 void FDeferredShadingSceneRenderer::RenderForwardShadingShadowProjections(FRHICommandListImmediate& RHICmdList, TRefCountPtr<IPooledRenderTarget>& ForwardScreenSpaceShadowMask)
 {
+	check(RHICmdList.IsOutsideRenderPass());
+
 	bool bScreenShadowMaskNeeded = false;
 
 	for (TSparseArray<FLightSceneInfoCompact>::TConstIterator LightIt(Scene->Lights); LightIt; ++LightIt)
@@ -830,34 +832,39 @@ void FDeferredShadingSceneRenderer::RenderForwardShadingShadowProjections(FRHICo
 		SCOPED_GPU_STAT(RHICmdList, ShadowProjection);
 
 		// All shadows render with min blending
-		bool bClearToWhite = true;
-		SetRenderTarget(RHICmdList, ForwardScreenSpaceShadowMask->GetRenderTargetItem().TargetableTexture, SceneRenderTargets.GetSceneDepthSurface(), bClearToWhite ? ESimpleRenderTargetMode::EClearColorExistingDepth : ESimpleRenderTargetMode::EExistingColorAndDepth, FExclusiveDepthStencil::DepthRead_StencilWrite, true);
+		FRHIRenderPassInfo RPInfo(ForwardScreenSpaceShadowMask->GetRenderTargetItem().TargetableTexture, ERenderTargetActions::Clear_Store);
+		TransitionRenderPassTargets(RHICmdList, RPInfo);
+		RHICmdList.BeginRenderPass(RPInfo, TEXT("RenderForwardShadingShadowProjectionsClear"));
+		RHICmdList.EndRenderPass();
 
-		for (TSparseArray<FLightSceneInfoCompact>::TConstIterator LightIt(Scene->Lights); LightIt; ++LightIt)
+		// Note: all calls here will set up renderpasses internally.
+		// #todo-renderpasses might be worth refactoring all this and splitting into lists of draws for each renderpass
 		{
-			const FLightSceneInfoCompact& LightSceneInfoCompact = *LightIt;
-			const FLightSceneInfo* const LightSceneInfo = LightSceneInfoCompact.LightSceneInfo;
-			FVisibleLightInfo& VisibleLightInfo = VisibleLightInfos[LightSceneInfo->Id];
+			for (TSparseArray<FLightSceneInfoCompact>::TConstIterator LightIt(Scene->Lights); LightIt; ++LightIt)
+			{
+				const FLightSceneInfoCompact& LightSceneInfoCompact = *LightIt;
+				const FLightSceneInfo* const LightSceneInfo = LightSceneInfoCompact.LightSceneInfo;
+				FVisibleLightInfo& VisibleLightInfo = VisibleLightInfos[LightSceneInfo->Id];
 
-			const bool bIssueLightDrawEvent = VisibleLightInfo.ShadowsToProject.Num() > 0 || VisibleLightInfo.CapsuleShadowsToProject.Num() > 0;
+				const bool bIssueLightDrawEvent = VisibleLightInfo.ShadowsToProject.Num() > 0 || VisibleLightInfo.CapsuleShadowsToProject.Num() > 0;
 
-			FString LightNameWithLevel;
-			GetLightNameForDrawEvent(LightSceneInfo->Proxy, LightNameWithLevel);
-			SCOPED_CONDITIONAL_DRAW_EVENTF(RHICmdList, EventLightPass, bIssueLightDrawEvent, *LightNameWithLevel);
+				FString LightNameWithLevel;
+				GetLightNameForDrawEvent(LightSceneInfo->Proxy, LightNameWithLevel);
+				SCOPED_CONDITIONAL_DRAW_EVENTF(RHICmdList, EventLightPass, bIssueLightDrawEvent, *LightNameWithLevel);
 
-			if (VisibleLightInfo.ShadowsToProject.Num() > 0)
-			{
-				FSceneRenderer::RenderShadowProjections(RHICmdList, LightSceneInfo, ForwardScreenSpaceShadowMask, true, false);
-			}
+				if (VisibleLightInfo.ShadowsToProject.Num() > 0)
+				{
+					FSceneRenderer::RenderShadowProjections(RHICmdList, LightSceneInfo, ForwardScreenSpaceShadowMask, true, false);
+				}
 
-			RenderCapsuleDirectShadows(RHICmdList, *LightSceneInfo, ForwardScreenSpaceShadowMask, VisibleLightInfo.CapsuleShadowsToProject, true);
+				RenderCapsuleDirectShadows(RHICmdList, *LightSceneInfo, ForwardScreenSpaceShadowMask, VisibleLightInfo.CapsuleShadowsToProject, true);
 
-			if (LightSceneInfo->GetDynamicShadowMapChannel() >= 0 && LightSceneInfo->GetDynamicShadowMapChannel() < 4)
-			{
-				RenderLightFunction(RHICmdList, LightSceneInfo, ForwardScreenSpaceShadowMask, true, true);
+				if (LightSceneInfo->GetDynamicShadowMapChannel() >= 0 && LightSceneInfo->GetDynamicShadowMapChannel() < 4)
+				{
+					RenderLightFunction(RHICmdList, LightSceneInfo, ForwardScreenSpaceShadowMask, true, true);
+				}
 			}
 		}
-
 		RHICmdList.CopyToResolveTarget(ForwardScreenSpaceShadowMask->GetRenderTargetItem().TargetableTexture, ForwardScreenSpaceShadowMask->GetRenderTargetItem().ShaderResourceTexture, FResolveParams(FResolveRect()));
 	}
 }
diff --git a/Engine/Source/Runtime/Renderer/Private/LightMapDensityRendering.cpp b/Engine/Source/Runtime/Renderer/Private/LightMapDensityRendering.cpp
--- a/Engine/Source/Runtime/Renderer/Private/LightMapDensityRendering.cpp
+++ b/Engine/Source/Runtime/Renderer/Private/LightMapDensityRendering.cpp
@@ -39,6 +39,8 @@ IMPLEMENT_DENSITY_LIGHTMAPPED_SHADER_TYPE( TUniformLightMapPolicy<LMP_HQ_LIGHTMA
 
 bool FDeferredShadingSceneRenderer::RenderLightMapDensities(FRHICommandListImmediate& RHICmdList)
 {
+	check(RHICmdList.IsInsideRenderPass());
+
 	bool bDirty=0;
 	if (Scene->GetFeatureLevel() >= ERHIFeatureLevel::SM4)
 	{
diff --git a/Engine/Source/Runtime/Renderer/Private/LightRendering.cpp b/Engine/Source/Runtime/Renderer/Private/LightRendering.cpp
--- a/Engine/Source/Runtime/Renderer/Private/LightRendering.cpp
+++ b/Engine/Source/Runtime/Renderer/Private/LightRendering.cpp
@@ -458,6 +458,8 @@ uint32 GetShadowQuality();
 /** Renders the scene's lighting. */
 void FDeferredShadingSceneRenderer::RenderLights(FRHICommandListImmediate& RHICmdList)
 {
+	check(RHICmdList.IsOutsideRenderPass());
+
 	SCOPED_NAMED_EVENT(FDeferredShadingSceneRenderer_RenderLights, FColor::Emerald);
 	SCOPED_DRAW_EVENT(RHICmdList, Lights);
 	SCOPED_GPU_STAT(RHICmdList, Lights);
@@ -603,6 +605,7 @@ void FDeferredShadingSceneRenderer::RenderLights(FRHICommandListImmediate& RHICm
 			{
 				SceneContext.BeginRenderingSceneColor(RHICmdList, ESimpleRenderTargetMode::EExistingColorAndDepth, FExclusiveDepthStencil::DepthRead_StencilWrite);
 				RenderSimpleLightsStandardDeferred(RHICmdList, SimpleLights);
+				SceneContext.FinishRenderingSceneColor(RHICmdList);
 			}
 
 			{
@@ -620,6 +623,8 @@ void FDeferredShadingSceneRenderer::RenderLights(FRHICommandListImmediate& RHICm
 					// Render the light to the scene color buffer, using a 1x1 white texture as input 
 					RenderLight(RHICmdList, LightSceneInfo, NULL, false, false);
 				}
+
+				SceneContext.FinishRenderingSceneColor(RHICmdList);
 			}
 
 			if (GUseTranslucentLightingVolumes && GSupportsVolumeTextureRendering)
@@ -758,37 +763,49 @@ void FDeferredShadingSceneRenderer::RenderLights(FRHICommandListImmediate& RHICm
 					// All shadows render with min blending
 					bool bClearToWhite = !bClearLightScreenExtentsOnly;
 
-					SetRenderTarget(RHICmdList, ScreenShadowMaskTexture->GetRenderTargetItem().TargetableTexture, SceneContext.GetSceneDepthSurface(), bClearToWhite ? ESimpleRenderTargetMode::EClearColorExistingDepth : ESimpleRenderTargetMode::EExistingColorAndDepth, FExclusiveDepthStencil::DepthRead_StencilWrite, true);
-
-					if (bClearLightScreenExtentsOnly)
+					FRHIRenderPassInfo RPInfo(ScreenShadowMaskTexture->GetRenderTargetItem().TargetableTexture, ERenderTargetActions::Load_Store);
+					RPInfo.DepthStencilRenderTarget.Action = MakeDepthStencilTargetActions(ERenderTargetActions::Load_DontStore, ERenderTargetActions::Load_Store);
+					RPInfo.DepthStencilRenderTarget.DepthStencilTarget = SceneContext.GetSceneDepthSurface();
+					RPInfo.DepthStencilRenderTarget.ExclusiveDepthStencil = FExclusiveDepthStencil::DepthRead_StencilWrite;
+					if (bClearToWhite)
 					{
-						SCOPED_DRAW_EVENT(RHICmdList, ClearQuad);
+						RPInfo.ColorRenderTargets[0].Action = ERenderTargetActions::Clear_Store;
+					}
 
-						for (int32 ViewIndex = 0; ViewIndex < Views.Num(); ViewIndex++)
+					TransitionRenderPassTargets(RHICmdList, RPInfo);
+					RHICmdList.BeginRenderPass(RPInfo, TEXT("ClearScreenShadowMask"));
+					{
+						if (bClearLightScreenExtentsOnly)
 						{
-							const FViewInfo& View = Views[ViewIndex];
-							FIntRect ScissorRect;
+							SCOPED_DRAW_EVENT(RHICmdList, ClearQuad);
 
-							if (!LightSceneInfo.Proxy->GetScissorRect(ScissorRect, View, View.ViewRect))
+							for (int32 ViewIndex = 0; ViewIndex < Views.Num(); ViewIndex++)
 							{
-								ScissorRect = View.ViewRect;
-							}
+								const FViewInfo& View = Views[ViewIndex];
+								FIntRect ScissorRect;
 
-							if (ScissorRect.Min.X < ScissorRect.Max.X && ScissorRect.Min.Y < ScissorRect.Max.Y)
-							{
-								RHICmdList.SetViewport(ScissorRect.Min.X, ScissorRect.Min.Y, 0.0f, ScissorRect.Max.X, ScissorRect.Max.Y, 1.0f);
-								DrawClearQuad(RHICmdList, true, FLinearColor(1, 1, 1, 1), false, 0, false, 0);
-							}
-							else
-							{
-								LightSceneInfo.Proxy->GetScissorRect(ScissorRect, View, View.ViewRect);
+								if (!LightSceneInfo.Proxy->GetScissorRect(ScissorRect, View, View.ViewRect))
+								{
+									ScissorRect = View.ViewRect;
+								}
+
+								if (ScissorRect.Min.X < ScissorRect.Max.X && ScissorRect.Min.Y < ScissorRect.Max.Y)
+								{
+									RHICmdList.SetViewport(ScissorRect.Min.X, ScissorRect.Min.Y, 0.0f, ScissorRect.Max.X, ScissorRect.Max.Y, 1.0f);
+									DrawClearQuad(RHICmdList, true, FLinearColor(1, 1, 1, 1), false, 0, false, 0);
+								}
+								else
+								{
+									LightSceneInfo.Proxy->GetScissorRect(ScissorRect, View, View.ViewRect);
+								}
 							}
 						}
 					}
+					RHICmdList.EndRenderPass();
 
 					RenderShadowProjections(RHICmdList, &LightSceneInfo, ScreenShadowMaskTexture, bInjectedTranslucentVolume);
-
-					bUsedShadowMaskTexture = true;
+					
+					bUsedShadowMaskTexture = true;					
 				}
 
 				for (int32 ViewIndex = 0; ViewIndex < Views.Num(); ViewIndex++)
@@ -816,7 +833,7 @@ void FDeferredShadingSceneRenderer::RenderLights(FRHICommandListImmediate& RHICm
 						INC_DWORD_STAT(STAT_NumLightFunctionOnlyLights);
 					}
 				}
-			
+
 				if (bUsedShadowMaskTexture)
 				{
 					RHICmdList.CopyToResolveTarget(ScreenShadowMaskTexture->GetRenderTargetItem().TargetableTexture, ScreenShadowMaskTexture->GetRenderTargetItem().ShaderResourceTexture, FResolveParams(FResolveRect()));
@@ -841,6 +858,8 @@ void FDeferredShadingSceneRenderer::RenderLights(FRHICommandListImmediate& RHICm
 				{
 					RenderLight(RHICmdList, &LightSceneInfo, ScreenShadowMaskTexture, false, true);
 				}
+
+				SceneContext.FinishRenderingSceneColor(RHICmdList);
 			}
 		}
 	}
@@ -882,7 +901,8 @@ void FDeferredShadingSceneRenderer::RenderStationaryLightOverlap(FRHICommandList
 {
 	if (Scene->bIsEditorScene)
 	{
-		FSceneRenderTargets::Get(RHICmdList).BeginRenderingSceneColor(RHICmdList, ESimpleRenderTargetMode::EUninitializedColorExistingDepth, FExclusiveDepthStencil::DepthRead_StencilWrite);
+		FSceneRenderTargets& SceneContext = FSceneRenderTargets::Get(RHICmdList);
+		SceneContext.BeginRenderingSceneColor(RHICmdList, ESimpleRenderTargetMode::EUninitializedColorExistingDepth, FExclusiveDepthStencil::DepthRead_StencilWrite);
 
 		// Clear to discard base pass values in scene color since we didn't skip that, to have valid scene depths
 		DrawClearQuad(RHICmdList, FLinearColor::Black);
@@ -892,6 +912,8 @@ void FDeferredShadingSceneRenderer::RenderStationaryLightOverlap(FRHICommandList
 		//Note: making use of FScene::InvisibleLights, which contains lights that haven't been added to the scene in the same way as visible lights
 		// So code called by RenderLightArrayForOverlapViewmode must be careful what it accesses
 		RenderLightArrayForOverlapViewmode(RHICmdList, Scene->InvisibleLights);
+
+		SceneContext.FinishRenderingSceneColor(RHICmdList);
 	}
 }
 
diff --git a/Engine/Source/Runtime/Renderer/Private/LightShaftRendering.cpp b/Engine/Source/Runtime/Renderer/Private/LightShaftRendering.cpp
--- a/Engine/Source/Runtime/Renderer/Private/LightShaftRendering.cpp
+++ b/Engine/Source/Runtime/Renderer/Private/LightShaftRendering.cpp
@@ -474,7 +474,9 @@ void AllocateOrReuseLightShaftRenderTarget(FRHICommandListImmediate& RHICmdList,
 		Desc.AutoWritable = false;
 		GRenderTargetPool.FindFreeElement(RHICmdList, Desc, Target, Name);
 
-		SetRenderTarget(RHICmdList, Target->GetRenderTargetItem().TargetableTexture, FTextureRHIRef(), ESimpleRenderTargetMode::EClearColorExistingDepth);
+		FRHIRenderPassInfo RPInfo(Target->GetRenderTargetItem().TargetableTexture, ERenderTargetActions::Clear_Store);
+		RHICmdList.BeginRenderPass(RPInfo, TEXT("ClearLightShaft"));
+		RHICmdList.EndRenderPass();
 	}
 }
 
@@ -489,25 +491,28 @@ void DownsamplePass(FRHICommandListImmediate& RHICmdList, const FViewInfo& View,
 	const uint32 DownsampledSizeX = View.ViewRect.Width() / DownsampleFactor;
 	const uint32 DownsampledSizeY = View.ViewRect.Height() / DownsampleFactor;
 
-	SetRenderTarget(RHICmdList, LightShaftsDest->GetRenderTargetItem().TargetableTexture, FTextureRHIRef(), true);
-	RHICmdList.SetViewport(DownSampledXY.X, DownSampledXY.Y, 0.0f, DownSampledXY.X + DownsampledSizeX, DownSampledXY.Y + DownsampledSizeY, 1.0f);
+	FRHIRenderPassInfo RPInfo(LightShaftsDest->GetRenderTargetItem().TargetableTexture, ERenderTargetActions::Load_Store);
+	TransitionRenderPassTargets(RHICmdList, RPInfo);
+	RHICmdList.BeginRenderPass(RPInfo, TEXT("DownsampleLightshaftMask"));
+	{
+		RHICmdList.SetViewport(DownSampledXY.X, DownSampledXY.Y, 0.0f, DownSampledXY.X + DownsampledSizeX, DownSampledXY.Y + DownsampledSizeY, 1.0f);
 
-	FGraphicsPipelineStateInitializer GraphicsPSOInit;
-	RHICmdList.ApplyCachedRenderTargets(GraphicsPSOInit);
+		FGraphicsPipelineStateInitializer GraphicsPSOInit;
+		RHICmdList.ApplyCachedRenderTargets(GraphicsPSOInit);
 
-	// No depth tests, no backface culling.
-	GraphicsPSOInit.BlendState = TStaticBlendState<>::GetRHI();
-	GraphicsPSOInit.RasterizerState = TStaticRasterizerState<FM_Solid, CM_None>::GetRHI();
-	GraphicsPSOInit.DepthStencilState = TStaticDepthStencilState<false, CF_Always>::GetRHI();
-	GraphicsPSOInit.PrimitiveType = PT_TriangleList;
+		// No depth tests, no backface culling.
+		GraphicsPSOInit.BlendState = TStaticBlendState<>::GetRHI();
+		GraphicsPSOInit.RasterizerState = TStaticRasterizerState<FM_Solid, CM_None>::GetRHI();
+		GraphicsPSOInit.DepthStencilState = TStaticDepthStencilState<false, CF_Always>::GetRHI();
+		GraphicsPSOInit.PrimitiveType = PT_TriangleList;
 
-	// Set shaders and texture
-	TShaderMapRef<FDownsampleLightShaftsVertexShader> DownsampleLightShaftsVertexShader(View.ShaderMap);
-	TRefCountPtr<IPooledRenderTarget> UnusedRT;
+		// Set shaders and texture
+		TShaderMapRef<FDownsampleLightShaftsVertexShader> DownsampleLightShaftsVertexShader(View.ShaderMap);
+		TRefCountPtr<IPooledRenderTarget> UnusedRT;
 
-	switch(LightSceneInfo->Proxy->GetLightType())
-	{
-	case LightType_Directional:
+		switch (LightSceneInfo->Proxy->GetLightType())
+		{
+		case LightType_Directional:
 		{
 			TShaderMapRef<TDownsampleLightShaftsPixelShader<LightType_Directional, bDownsampleOcclusion> > DownsampleLightShaftsPixelShader(View.ShaderMap);
 			GraphicsPSOInit.BoundShaderState.VertexDeclarationRHI = GFilterVertexDeclaration.VertexDeclarationRHI;
@@ -517,7 +522,7 @@ void DownsamplePass(FRHICommandListImmediate& RHICmdList, const FViewInfo& View,
 			DownsampleLightShaftsPixelShader->SetParameters(RHICmdList, LightSceneInfo, View, UnusedRT);
 		}
 		break;
-	case LightType_Spot:
+		case LightType_Spot:
 		{
 			TShaderMapRef<TDownsampleLightShaftsPixelShader<LightType_Spot, bDownsampleOcclusion> > DownsampleLightShaftsPixelShader(View.ShaderMap);
 			GraphicsPSOInit.BoundShaderState.VertexDeclarationRHI = GFilterVertexDeclaration.VertexDeclarationRHI;
@@ -527,8 +532,8 @@ void DownsamplePass(FRHICommandListImmediate& RHICmdList, const FViewInfo& View,
 			DownsampleLightShaftsPixelShader->SetParameters(RHICmdList, LightSceneInfo, View, UnusedRT);
 		}
 		break;
-	default:
-	case LightType_Point:
+		default:
+		case LightType_Point:
 		{
 			TShaderMapRef<TDownsampleLightShaftsPixelShader<LightType_Point, bDownsampleOcclusion> > DownsampleLightShaftsPixelShader(View.ShaderMap);
 			GraphicsPSOInit.BoundShaderState.VertexDeclarationRHI = GFilterVertexDeclaration.VertexDeclarationRHI;
@@ -538,22 +543,23 @@ void DownsamplePass(FRHICommandListImmediate& RHICmdList, const FViewInfo& View,
 			DownsampleLightShaftsPixelShader->SetParameters(RHICmdList, LightSceneInfo, View, UnusedRT);
 		}
 		break;
-	}
+		}
 
-	DownsampleLightShaftsVertexShader->SetParameters(RHICmdList, View);
-
-	// Downsample scene color and depth, and convert them into a bloom term and an occlusion masking term
-	DrawRectangle( 
-		RHICmdList,
-		0, 0, 
-		DownsampledSizeX, DownsampledSizeY,
-		View.ViewRect.Min.X, View.ViewRect.Min.Y, 
-		View.ViewRect.Width(), View.ViewRect.Height(),
-		FIntPoint(DownsampledSizeX, DownsampledSizeY), 
-		BufferSize,
-		*DownsampleLightShaftsVertexShader,
-		EDRF_UseTriangleOptimization);
+		DownsampleLightShaftsVertexShader->SetParameters(RHICmdList, View);
 
+		// Downsample scene color and depth, and convert them into a bloom term and an occlusion masking term
+		DrawRectangle(
+			RHICmdList,
+			0, 0,
+			DownsampledSizeX, DownsampledSizeY,
+			View.ViewRect.Min.X, View.ViewRect.Min.Y,
+			View.ViewRect.Width(), View.ViewRect.Height(),
+			FIntPoint(DownsampledSizeX, DownsampledSizeY),
+			BufferSize,
+			*DownsampleLightShaftsVertexShader,
+			EDRF_UseTriangleOptimization);
+	}
+	RHICmdList.EndRenderPass();
 	RHICmdList.CopyToResolveTarget(LightShaftsDest->GetRenderTargetItem().TargetableTexture, LightShaftsDest->GetRenderTargetItem().ShaderResourceTexture, FResolveParams());
 
 	Swap(LightShaftsSource, LightShaftsDest);
@@ -625,6 +631,8 @@ void ApplyRadialBlurPasses(
 	/** First pass dest. */
 	TRefCountPtr<IPooledRenderTarget>& LightShaftsDest)
 {
+	check(RHICmdList.IsOutsideRenderPass());
+
 	TShaderMapRef<FScreenVS> ScreenVertexShader(View.ShaderMap);
 
 	const FIntPoint BufferSize = FSceneRenderTargets::Get(RHICmdList).GetBufferSizeXY();
@@ -637,42 +645,46 @@ void ApplyRadialBlurPasses(
 
 	for (uint32 PassIndex = 0; PassIndex < NumPasses; PassIndex++)
 	{
-		SetRenderTarget(RHICmdList, LightShaftsDest->GetRenderTargetItem().TargetableTexture, FTextureRHIRef(), true);
-		RHICmdList.SetViewport(0, 0, 0.0f, FilterBufferSize.X, FilterBufferSize.Y, 1.0f);
+		FRHIRenderPassInfo RPInfo(LightShaftsDest->GetRenderTargetItem().TargetableTexture, ERenderTargetActions::Load_Store);
+		TransitionRenderPassTargets(RHICmdList, RPInfo);
+		RHICmdList.BeginRenderPass(RPInfo, TEXT("LightShaftPass"));
+		{
+			RHICmdList.SetViewport(0, 0, 0.0f, FilterBufferSize.X, FilterBufferSize.Y, 1.0f);
 
-		FGraphicsPipelineStateInitializer GraphicsPSOInit;
-		RHICmdList.ApplyCachedRenderTargets(GraphicsPSOInit);
+			FGraphicsPipelineStateInitializer GraphicsPSOInit;
+			RHICmdList.ApplyCachedRenderTargets(GraphicsPSOInit);
 
-		GraphicsPSOInit.BlendState = TStaticBlendState<>::GetRHI();
-		GraphicsPSOInit.RasterizerState = TStaticRasterizerState<>::GetRHI();
-		GraphicsPSOInit.DepthStencilState = TStaticDepthStencilState<false, CF_Always>::GetRHI();
+			GraphicsPSOInit.BlendState = TStaticBlendState<>::GetRHI();
+			GraphicsPSOInit.RasterizerState = TStaticRasterizerState<>::GetRHI();
+			GraphicsPSOInit.DepthStencilState = TStaticDepthStencilState<false, CF_Always>::GetRHI();
 
-		TShaderMapRef<FBlurLightShaftsPixelShader> BlurLightShaftsPixelShader(View.ShaderMap);
+			TShaderMapRef<FBlurLightShaftsPixelShader> BlurLightShaftsPixelShader(View.ShaderMap);
 
-		GraphicsPSOInit.BoundShaderState.VertexDeclarationRHI = GFilterVertexDeclaration.VertexDeclarationRHI;
-		GraphicsPSOInit.BoundShaderState.VertexShaderRHI = GETSAFERHISHADER_VERTEX(*ScreenVertexShader);
-		GraphicsPSOInit.BoundShaderState.PixelShaderRHI = GETSAFERHISHADER_PIXEL(*BlurLightShaftsPixelShader);
-		GraphicsPSOInit.PrimitiveType = PT_TriangleList;
+			GraphicsPSOInit.BoundShaderState.VertexDeclarationRHI = GFilterVertexDeclaration.VertexDeclarationRHI;
+			GraphicsPSOInit.BoundShaderState.VertexShaderRHI = GETSAFERHISHADER_VERTEX(*ScreenVertexShader);
+			GraphicsPSOInit.BoundShaderState.PixelShaderRHI = GETSAFERHISHADER_PIXEL(*BlurLightShaftsPixelShader);
+			GraphicsPSOInit.PrimitiveType = PT_TriangleList;
 
-		SetGraphicsPipelineState(RHICmdList, GraphicsPSOInit);
+			SetGraphicsPipelineState(RHICmdList, GraphicsPSOInit);
 
-		TRefCountPtr<IPooledRenderTarget>& EffectiveSource = PassIndex == 0 ? FirstPassSource : LightShaftsSource;
-		/// ?
-		BlurLightShaftsPixelShader->SetParameters(RHICmdList, LightSceneInfo, View, PassIndex, EffectiveSource);
+			TRefCountPtr<IPooledRenderTarget>& EffectiveSource = PassIndex == 0 ? FirstPassSource : LightShaftsSource;
+			/// ?
+			BlurLightShaftsPixelShader->SetParameters(RHICmdList, LightSceneInfo, View, PassIndex, EffectiveSource);
 
-		{
-			// Apply a radial blur to the bloom and occlusion mask
-			DrawRectangle( 
-				RHICmdList,
-				DownSampledXY.X, DownSampledXY.Y, 
-				DownsampledSizeX, DownsampledSizeY,
-				DownSampledXY.X, DownSampledXY.Y, 
-				DownsampledSizeX, DownsampledSizeY,
-				FilterBufferSize, FilterBufferSize,
-				*ScreenVertexShader,
-				EDRF_UseTriangleOptimization);
+			{
+				// Apply a radial blur to the bloom and occlusion mask
+				DrawRectangle(
+					RHICmdList,
+					DownSampledXY.X, DownSampledXY.Y,
+					DownsampledSizeX, DownsampledSizeY,
+					DownSampledXY.X, DownSampledXY.Y,
+					DownsampledSizeX, DownsampledSizeY,
+					FilterBufferSize, FilterBufferSize,
+					*ScreenVertexShader,
+					EDRF_UseTriangleOptimization);
+			}
 		}
-
+		RHICmdList.EndRenderPass();
 		RHICmdList.CopyToResolveTarget(LightShaftsDest->GetRenderTargetItem().TargetableTexture, LightShaftsDest->GetRenderTargetItem().ShaderResourceTexture, FResolveParams());
 
 		// Swap input and output for the next pass
@@ -682,6 +694,8 @@ void ApplyRadialBlurPasses(
 
 void FinishOcclusionTerm(FRHICommandList& RHICmdList, const FViewInfo& View, const FLightSceneInfo* const LightSceneInfo, TRefCountPtr<IPooledRenderTarget>& LightShaftsSource, TRefCountPtr<IPooledRenderTarget>& LightShaftsDest)
 {
+	check(RHICmdList.IsOutsideRenderPass());
+
 	TShaderMapRef<FScreenVS> ScreenVertexShader(View.ShaderMap);
 
 	const FIntPoint BufferSize = FSceneRenderTargets::Get(RHICmdList).GetBufferSizeXY();
@@ -691,39 +705,43 @@ void FinishOcclusionTerm(FRHICommandList& RHICmdList, const FViewInfo& View, con
 	const uint32 DownsampledSizeX = View.ViewRect.Width() / DownsampleFactor;
 	const uint32 DownsampledSizeY = View.ViewRect.Height() / DownsampleFactor;
 
-	SetRenderTarget(RHICmdList, LightShaftsDest->GetRenderTargetItem().TargetableTexture, FTextureRHIRef(), true);
-	RHICmdList.SetViewport(0, 0, 0.0f, FilterBufferSize.X, FilterBufferSize.Y, 1.0f);
+	FRHIRenderPassInfo RPInfo(LightShaftsDest->GetRenderTargetItem().TargetableTexture, ERenderTargetActions::Load_Store);
+	TransitionRenderPassTargets(RHICmdList, RPInfo);
+	RHICmdList.BeginRenderPass(RPInfo, TEXT("FinishOcclusionTerm"));
+	{
+		RHICmdList.SetViewport(0, 0, 0.0f, FilterBufferSize.X, FilterBufferSize.Y, 1.0f);
 
-	FGraphicsPipelineStateInitializer GraphicsPSOInit;
-	RHICmdList.ApplyCachedRenderTargets(GraphicsPSOInit);
-	GraphicsPSOInit.BlendState = TStaticBlendState<>::GetRHI();
-	GraphicsPSOInit.RasterizerState = TStaticRasterizerState<>::GetRHI();
-	GraphicsPSOInit.DepthStencilState = TStaticDepthStencilState<false, CF_Always>::GetRHI();
+		FGraphicsPipelineStateInitializer GraphicsPSOInit;
+		RHICmdList.ApplyCachedRenderTargets(GraphicsPSOInit);
+		GraphicsPSOInit.BlendState = TStaticBlendState<>::GetRHI();
+		GraphicsPSOInit.RasterizerState = TStaticRasterizerState<>::GetRHI();
+		GraphicsPSOInit.DepthStencilState = TStaticDepthStencilState<false, CF_Always>::GetRHI();
 
-	TShaderMapRef<FFinishOcclusionPixelShader> MaskOcclusionTermPixelShader(View.ShaderMap);
+		TShaderMapRef<FFinishOcclusionPixelShader> MaskOcclusionTermPixelShader(View.ShaderMap);
 
-	GraphicsPSOInit.BoundShaderState.VertexDeclarationRHI = GFilterVertexDeclaration.VertexDeclarationRHI;
-	GraphicsPSOInit.BoundShaderState.VertexShaderRHI = GETSAFERHISHADER_VERTEX(*ScreenVertexShader);
-	GraphicsPSOInit.BoundShaderState.PixelShaderRHI = GETSAFERHISHADER_PIXEL(*MaskOcclusionTermPixelShader);
-	GraphicsPSOInit.PrimitiveType = PT_TriangleList;
+		GraphicsPSOInit.BoundShaderState.VertexDeclarationRHI = GFilterVertexDeclaration.VertexDeclarationRHI;
+		GraphicsPSOInit.BoundShaderState.VertexShaderRHI = GETSAFERHISHADER_VERTEX(*ScreenVertexShader);
+		GraphicsPSOInit.BoundShaderState.PixelShaderRHI = GETSAFERHISHADER_PIXEL(*MaskOcclusionTermPixelShader);
+		GraphicsPSOInit.PrimitiveType = PT_TriangleList;
 
-	SetGraphicsPipelineState(RHICmdList, GraphicsPSOInit);
-	/// ?
-	MaskOcclusionTermPixelShader->SetParameters(RHICmdList, LightSceneInfo, View, LightShaftsSource);
+		SetGraphicsPipelineState(RHICmdList, GraphicsPSOInit);
+		/// ?
+		MaskOcclusionTermPixelShader->SetParameters(RHICmdList, LightSceneInfo, View, LightShaftsSource);
 
-	{
-		// Apply a radial blur to the bloom and occlusion mask
-		DrawRectangle( 
-			RHICmdList,
-			DownSampledXY.X, DownSampledXY.Y, 
-			DownsampledSizeX, DownsampledSizeY,
-			DownSampledXY.X, DownSampledXY.Y, 
-			DownsampledSizeX, DownsampledSizeY,
-			FilterBufferSize, FilterBufferSize,
-			*ScreenVertexShader,
-			EDRF_UseTriangleOptimization);
+		{
+			// Apply a radial blur to the bloom and occlusion mask
+			DrawRectangle(
+				RHICmdList,
+				DownSampledXY.X, DownSampledXY.Y,
+				DownsampledSizeX, DownsampledSizeY,
+				DownSampledXY.X, DownSampledXY.Y,
+				DownsampledSizeX, DownsampledSizeY,
+				FilterBufferSize, FilterBufferSize,
+				*ScreenVertexShader,
+				EDRF_UseTriangleOptimization);
+		}
 	}
-
+	RHICmdList.EndRenderPass();
 	RHICmdList.CopyToResolveTarget(LightShaftsDest->GetRenderTargetItem().TargetableTexture, LightShaftsDest->GetRenderTargetItem().ShaderResourceTexture, FResolveParams());
 }
 
@@ -976,6 +994,10 @@ void ApplyLightShaftBloom(FRHICommandListImmediate& RHICmdList, const FViewInfo&
 	{
 		SceneContext.FinishRenderingSeparateTranslucency(RHICmdList, View);
 	}
+	else
+	{
+		SceneContext.FinishRenderingSceneColor(RHICmdList);
+	}
 }
 
 void FSceneViewState::TrimHistoryRenderTargets(const FScene* Scene)
