PR #5238: Fix for Decals rendering in main pass when bRenderInMainPass is false (Contributed by [REDACTED_PERSON])

diff --git a/Engine/Source/Runtime/Renderer/Private/SceneVisibility.cpp b/Engine/Source/Runtime/Renderer/Private/SceneVisibility.cpp
--- a/Engine/Source/Runtime/Renderer/Private/SceneVisibility.cpp
+++ b/Engine/Source/Runtime/Renderer/Private/SceneVisibility.cpp
@@ -1817,7 +1817,7 @@ struct FRelevancePacket
 				continue;
 			}
 
-			if (ViewRelevance.bDecal)
+			if (ViewRelevance.bDecal && ViewRelevance.bRenderInMainPass)
 			{
 				MeshDecalPrimSet.AddPrim(FMeshDecalPrimSet::GenerateKey(PrimitiveSceneInfo, PrimitiveSceneInfo->Proxy->GetTranslucencySortPriority()));
 			}
