PR #5245: Fix for Lights with Source Length > 0.5f in Tiled Deferred lighting path (Contributed by [REDACTED_PERSON])
Use negative source length to specify simple lights, rather than previous behavior of source length > 0.5, which is valid value for non-simple lights

diff --git a/Engine/Shaders/Private/TiledDeferredLightShaders.usf b/Engine/Shaders/Private/TiledDeferredLightShaders.usf
--- a/Engine/Shaders/Private/TiledDeferredLightShaders.usf
+++ b/Engine/Shaders/Private/TiledDeferredLightShaders.usf
@@ -216,9 +216,9 @@ void TiledDeferredLightingMain(
 			BRANCH
 			if (bInDepthRange) 
 			{
-				// Separate lights into different index lists depending on whether they will use simple shading or not
+				// Separate lights into different index lists depending on whether they will use simple shading or not (designated by a negative SourceLength)
 				// This prevents the need for an inner loop branch when shading
-				if (TiledDeferred2.SpotAnglesAndSourceRadiusAndSimpleLighting[LightIndex].w > .5f)
+				if (TiledDeferred2.SpotAnglesAndSourceRadiusAndSimpleLighting[LightIndex].w < 0.0f)
 				{
 					uint ListIndex;
 					InterlockedAdd(TileNumSimpleLights, 1U, ListIndex);
diff --git a/Engine/Source/Runtime/Renderer/Private/TiledDeferredLightRendering.cpp b/Engine/Source/Runtime/Renderer/Private/TiledDeferredLightRendering.cpp
--- a/Engine/Source/Runtime/Renderer/Private/TiledDeferredLightRendering.cpp
+++ b/Engine/Source/Runtime/Renderer/Private/TiledDeferredLightRendering.cpp
@@ -250,7 +250,7 @@ public:
 				LightData.LightPositionAndInvRadius[LightIndex] = FVector4(SimpleLightPerViewData.Position, 1.0f / FMath::Max(SimpleLight.Radius, KINDA_SMALL_NUMBER));
 				LightData.LightColorAndFalloffExponent[LightIndex] = FVector4(SimpleLight.Color, SimpleLight.Exponent);
 				LightData2.LightDirectionAndSpotlightMaskAndSpecularScale[LightIndex] = FVector4(FVector(1, 0, 0), 0);
-				LightData2.SpotAnglesAndSourceRadiusAndSimpleLighting[LightIndex] = FVector4(-2, 1, 0, 1);
+				LightData2.SpotAnglesAndSourceRadiusAndSimpleLighting[LightIndex] = FVector4(-2, 1, 0, -1);
 				LightData2.ShadowMapChannelMask[LightIndex] = FVector4(0, 0, 0, 0);
 			}
 		}
