Not ready for this check until the basepass renderpass conversion is checked in.

diff --git a/Engine/Source/Runtime/Renderer/Private/LightMapDensityRendering.cpp b/Engine/Source/Runtime/Renderer/Private/LightMapDensityRendering.cpp
--- a/Engine/Source/Runtime/Renderer/Private/LightMapDensityRendering.cpp
+++ b/Engine/Source/Runtime/Renderer/Private/LightMapDensityRendering.cpp
@@ -39,8 +39,6 @@ IMPLEMENT_DENSITY_LIGHTMAPPED_SHADER_TYPE( TUniformLightMapPolicy<LMP_HQ_LIGHTMA
 
 bool FDeferredShadingSceneRenderer::RenderLightMapDensities(FRHICommandListImmediate& RHICmdList)
 {
-	check(RHICmdList.IsInsideRenderPass());
-
 	bool bDirty=0;
 	if (Scene->GetFeatureLevel() >= ERHIFeatureLevel::SM4)
 	{
