- Exclude lights with non-zero source length from tiled deferred pass (these don't work due to light tangent not being passed to tiled deferred shader)
- Clean up logic around tiled deferred pass, entire pass is no longer disabled by unsupported light types

diff --git a/Engine/Shaders/Private/TiledDeferredLightShaders.usf b/Engine/Shaders/Private/TiledDeferredLightShaders.usf
--- a/Engine/Shaders/Private/TiledDeferredLightShaders.usf
+++ b/Engine/Shaders/Private/TiledDeferredLightShaders.usf
@@ -218,7 +218,7 @@ void TiledDeferredLightingMain(
 			{
 				// Separate lights into different index lists depending on whether they will use simple shading or not (designated by a negative SourceLength)
 				// This prevents the need for an inner loop branch when shading
-				if (TiledDeferred2.SpotAnglesAndSourceRadiusAndSimpleLighting[LightIndex].w < 0.0f)
+				if (TiledDeferred2.SpotAnglesAndSourceRadiusAndSimpleLighting[LightIndex].w > 0.0f)
 				{
 					uint ListIndex;
 					InterlockedAdd(TileNumSimpleLights, 1U, ListIndex);
@@ -261,7 +261,9 @@ void TiledDeferredLightingMain(
 				LightData.bSpotLight = Value.w > 0;
 				LightData.SpecularScale = abs(Value.w);
 			}
-			LightData.SpotAnglesAndSourceRadius = TiledDeferred2.SpotAnglesAndSourceRadiusAndSimpleLighting[LightIndex];
+
+			// don't support lights with non-0 length in tiled-deferred path, so make sure it stays zero-initialized
+			LightData.SpotAnglesAndSourceRadius.xyz = TiledDeferred2.SpotAnglesAndSourceRadiusAndSimpleLighting[LightIndex].xyz;
 			LightData.ShadowMapChannelMask = TiledDeferred2.ShadowMapChannelMask[LightIndex];
 			LightData.bInverseSquared = LightData.LightColorAndFalloffExponent.w == 0;
 			// Only radial lights supported with tiled deferred
diff --git a/Engine/Source/Runtime/Engine/Private/Components/LightComponent.cpp b/Engine/Source/Runtime/Engine/Private/Components/LightComponent.cpp
--- a/Engine/Source/Runtime/Engine/Private/Components/LightComponent.cpp
+++ b/Engine/Source/Runtime/Engine/Private/Components/LightComponent.cpp
@@ -237,6 +237,7 @@ FLightSceneProxy::FLightSceneProxy(const ULightComponent* InLightComponent)
 	, bUseRayTracedDistanceFieldShadows(InLightComponent->bUseRayTracedDistanceFieldShadows)
 	, bCastModulatedShadows(false)
 	, bUseWholeSceneCSMForMovableObjects(false)
+	, bTiledDeferredLightingSupported(false)
 	, LightType(InLightComponent->GetLightType())	
 	, LightingChannelMask(GetLightingChannelMaskForStruct(InLightComponent->LightingChannels))
 	, StatId(InLightComponent->GetStatID(true))
diff --git a/Engine/Source/Runtime/Engine/Private/PointLightSceneProxy.h b/Engine/Source/Runtime/Engine/Private/PointLightSceneProxy.h
--- a/Engine/Source/Runtime/Engine/Private/PointLightSceneProxy.h
+++ b/Engine/Source/Runtime/Engine/Private/PointLightSceneProxy.h
@@ -187,6 +187,9 @@ public:
 	,	bInverseSquared(Component->bUseInverseSquaredFalloff)
 	{
 		UpdateRadius(Component->AttenuationRadius);
+
+		// tiled deferred only supported for point/spot lights with 0 length
+		bTiledDeferredLightingSupported = (SourceLength == 0.0f);
 	}
 
 	virtual float GetSourceRadius() const override
diff --git a/Engine/Source/Runtime/Engine/Public/SceneManagement.h b/Engine/Source/Runtime/Engine/Public/SceneManagement.h
--- a/Engine/Source/Runtime/Engine/Public/SceneManagement.h
+++ b/Engine/Source/Runtime/Engine/Public/SceneManagement.h
@@ -1110,6 +1110,7 @@ public:
 	inline bool Transmission() const { return bTransmission; }
 	inline bool UseRayTracedDistanceFieldShadows() const { return bUseRayTracedDistanceFieldShadows; }
 	inline float GetRayStartOffsetDepthScale() const { return RayStartOffsetDepthScale; }
+	inline bool IsTiledDeferredLightingSupported() const { return bTiledDeferredLightingSupported;  }
 	inline uint8 GetLightType() const { return LightType; }
 	inline uint8 GetLightingChannelMask() const { return LightingChannelMask; }
 	inline FName GetComponentName() const { return ComponentName; }
@@ -1267,6 +1268,9 @@ protected:
 	/** Whether to render csm shadows for movable objects only (mobile). */
 	uint8 bUseWholeSceneCSMForMovableObjects : 1;
 
+	/** Whether the light supports rendering in tiled deferred pass */
+	uint8 bTiledDeferredLightingSupported : 1;
+
 	/** The light type (ELightComponentType) */
 	const uint8 LightType;
 
diff --git a/Engine/Source/Runtime/Renderer/Private/LightRendering.cpp b/Engine/Source/Runtime/Renderer/Private/LightRendering.cpp
--- a/Engine/Source/Runtime/Renderer/Private/LightRendering.cpp
+++ b/Engine/Source/Runtime/Renderer/Private/LightRendering.cpp
@@ -504,6 +504,14 @@ void FDeferredShadingSceneRenderer::RenderLights(FRHICommandListImmediate& RHICm
 					SortedLightInfo->SortKey.Fields.bShadowed = bDynamicShadows && CheckForProjectedShadows(LightSceneInfo);
 					SortedLightInfo->SortKey.Fields.bLightFunction = ViewFamily.EngineShowFlags.LightFunctions && CheckForLightFunction(LightSceneInfo);
 					SortedLightInfo->SortKey.Fields.bUsesLightingChannels = Views[ViewIndex].bUsesLightingChannels && LightSceneInfo->Proxy->GetLightingChannelMask() != GetDefaultLightingChannelMask();
+
+					// tiled deferred lighting only supported for certain lights that don't use any additional features
+					const bool bTiledDeferredSupported = LightSceneInfo->Proxy->IsTiledDeferredLightingSupported() &&
+						!SortedLightInfo->SortKey.Fields.bTextureProfile &&
+						!SortedLightInfo->SortKey.Fields.bShadowed &&
+						!SortedLightInfo->SortKey.Fields.bLightFunction &&
+						!SortedLightInfo->SortKey.Fields.bUsesLightingChannels;
+					SortedLightInfo->SortKey.Fields.bTiledDeferredNotSupported = !bTiledDeferredSupported;
 					break;
 				}
 			}
@@ -527,41 +535,29 @@ void FDeferredShadingSceneRenderer::RenderLights(FRHICommandListImmediate& RHICm
 
 		int32 AttenuationLightStart = SortedLights.Num();
 		int32 SupportedByTiledDeferredLightEnd = SortedLights.Num();
-		bool bAnyUnsupportedByTiledDeferred = false;
 
 		// Iterate over all lights to be rendered and build ranges for tiled deferred and unshadowed lights
 		for (int32 LightIndex = 0; LightIndex < SortedLights.Num(); LightIndex++)
 		{
 			const FSortedLightSceneInfo& SortedLightInfo = SortedLights[LightIndex];
-			bool bDrawShadows = SortedLightInfo.SortKey.Fields.bShadowed;
-			bool bDrawLightFunction = SortedLightInfo.SortKey.Fields.bLightFunction;
-			bool bTextureLightProfile = SortedLightInfo.SortKey.Fields.bTextureProfile;
-			bool bLightingChannels = SortedLightInfo.SortKey.Fields.bUsesLightingChannels;
+			const bool bDrawShadows = SortedLightInfo.SortKey.Fields.bShadowed;
+			const bool bDrawLightFunction = SortedLightInfo.SortKey.Fields.bLightFunction;
+			const bool bTextureLightProfile = SortedLightInfo.SortKey.Fields.bTextureProfile;
+			const bool bLightingChannels = SortedLightInfo.SortKey.Fields.bUsesLightingChannels;
 
-			if (bTextureLightProfile && SupportedByTiledDeferredLightEnd == SortedLights.Num())
+			if (SortedLightInfo.SortKey.Fields.bTiledDeferredNotSupported && SupportedByTiledDeferredLightEnd == SortedLights.Num())
 			{
-				// Mark the first index to not support tiled deferred due to texture light profile
+				// Mark the first index to not support tiled deferred
 				SupportedByTiledDeferredLightEnd = LightIndex;
 			}
 
 			if (bDrawShadows || bDrawLightFunction || bLightingChannels)
 			{
+				// Once we find a shadowed light, we can exit the loop, these lights should never support tiled deferred rendering either
+				check(SortedLightInfo.SortKey.Fields.bTiledDeferredNotSupported);
 				AttenuationLightStart = LightIndex;
-
-				if (SupportedByTiledDeferredLightEnd == SortedLights.Num())
-				{
-					// Mark the first index to not support tiled deferred due to shadowing
-					SupportedByTiledDeferredLightEnd = LightIndex;
-				}
 				break;
 			}
-
-			if (LightIndex < SupportedByTiledDeferredLightEnd)
-			{
-				// Directional lights currently not supported by tiled deferred
-				bAnyUnsupportedByTiledDeferred = bAnyUnsupportedByTiledDeferred 
-					|| (SortedLightInfo.SortKey.Fields.LightType != LightType_Point && SortedLightInfo.SortKey.Fields.LightType != LightType_Spot);
-			}
 		}
 		
 		if (GbEnableAsyncComputeTranslucencyLightingVolumeClear && GSupportsEfficientAsyncCompute)
@@ -592,7 +588,7 @@ void FDeferredShadingSceneRenderer::RenderLights(FRHICommandListImmediate& RHICm
 				}
 
 				// Use tiled deferred shading on any unshadowed lights without a texture light profile
-				if (ShouldUseTiledDeferred(SupportedByTiledDeferredLightEnd, SimpleLights.InstanceData.Num()) && !bAnyUnsupportedByTiledDeferred && !bAnyViewIsStereo)
+				if (ShouldUseTiledDeferred(SupportedByTiledDeferredLightEnd, SimpleLights.InstanceData.Num()) && !bAnyViewIsStereo)
 				{
 					// Update the range that needs to be processed by standard deferred to exclude the lights done with tiled
 					StandardDeferredStart = SupportedByTiledDeferredLightEnd;
diff --git a/Engine/Source/Runtime/Renderer/Private/LightSceneInfo.h b/Engine/Source/Runtime/Renderer/Private/LightSceneInfo.h
--- a/Engine/Source/Runtime/Renderer/Private/LightSceneInfo.h
+++ b/Engine/Source/Runtime/Renderer/Private/LightSceneInfo.h
@@ -69,7 +69,7 @@ struct FSortedLightSceneInfo
 		struct
 		{
 			// Note: the order of these members controls the light sort order!
-			// Currently bUsesLightingChannels is the MSB and LightType is LSB
+			// Currently bTiledDeferredNotSupported is the MSB and LightType is LSB
 			/** The type of light. */
 			uint32 LightType : LightType_NumBits;
 			/** Whether the light has a texture profile. */
@@ -80,6 +80,8 @@ struct FSortedLightSceneInfo
 			uint32 bShadowed : 1;
 			/** Whether the light uses lighting channels. */
 			uint32 bUsesLightingChannels : 1;
+			/** True if the light doesn't support tiled deferred, logic is inverted so that lights that DO support tiled deferred will sort first in list */
+			uint32 bTiledDeferredNotSupported : 1;
 		} Fields;
 		/** Sort key bits packed into an integer. */
 		int32 Packed;
diff --git a/Engine/Source/Runtime/Renderer/Private/TiledDeferredLightRendering.cpp b/Engine/Source/Runtime/Renderer/Private/TiledDeferredLightRendering.cpp
--- a/Engine/Source/Runtime/Renderer/Private/TiledDeferredLightRendering.cpp
+++ b/Engine/Source/Runtime/Renderer/Private/TiledDeferredLightRendering.cpp
@@ -223,11 +223,14 @@ public:
 					LightData2.LightDirectionAndSpotlightMaskAndSpecularScale[LightIndex] = FVector4(LightParameters.NormalizedLightDirection, W);
 				}
 
+				// Lights with non-0 length don't support tiled deferred pass, should not have gotten into this list
+				ensure(LightParameters.LightSourceLength==0.0f);
+
 				LightData2.SpotAnglesAndSourceRadiusAndSimpleLighting[LightIndex] = FVector4(
 						LightParameters.SpotAngles.X,
 						LightParameters.SpotAngles.Y,
 						LightParameters.LightSourceRadius,
-						LightParameters.LightSourceLength);
+						0.0f);
 
 				int32 ShadowMapChannel = LightSceneInfo->Proxy->GetShadowMapChannel();
 
@@ -250,7 +253,7 @@ public:
 				LightData.LightPositionAndInvRadius[LightIndex] = FVector4(SimpleLightPerViewData.Position, 1.0f / FMath::Max(SimpleLight.Radius, KINDA_SMALL_NUMBER));
 				LightData.LightColorAndFalloffExponent[LightIndex] = FVector4(SimpleLight.Color, SimpleLight.Exponent);
 				LightData2.LightDirectionAndSpotlightMaskAndSpecularScale[LightIndex] = FVector4(FVector(1, 0, 0), 0);
-				LightData2.SpotAnglesAndSourceRadiusAndSimpleLighting[LightIndex] = FVector4(-2, 1, 0, -1);
+				LightData2.SpotAnglesAndSourceRadiusAndSimpleLighting[LightIndex] = FVector4(-2, 1, 0, 1.0f);
 				LightData2.ShadowMapChannelMask[LightIndex] = FVector4(0, 0, 0, 0);
 			}
 		}
