Don't render preshadow projections for objects that don't render in the main pass

diff --git a/Engine/Source/Runtime/Renderer/Private/ShadowSetup.cpp b/Engine/Source/Runtime/Renderer/Private/ShadowSetup.cpp
--- a/Engine/Source/Runtime/Renderer/Private/ShadowSetup.cpp
+++ b/Engine/Source/Runtime/Renderer/Private/ShadowSetup.cpp
@@ -1542,8 +1542,13 @@ void FSceneRenderer::CreatePerObjectProjectedShadow(
 				((FSceneViewState*)View.State)->IsShadowOccluded(RHICmdList, TranslucentKey, NumBufferedFrames)
 			);
 
-		const bool bSubjectIsVisibleInThisView = View.PrimitiveVisibilityMap[PrimitiveSceneInfo->GetIndex()];
-		bSubjectIsVisible |= bSubjectIsVisibleInThisView;
+		// if subject doesn't render in the main pass, it's never considered visible
+		// (in this case, there will be no need to generate any preshadows for the subject)
+		if (PrimitiveSceneInfo->Proxy->ShouldRenderInMainPass())
+		{
+			const bool bSubjectIsVisibleInThisView = View.PrimitiveVisibilityMap[PrimitiveSceneInfo->GetIndex()];
+			bSubjectIsVisible |= bSubjectIsVisibleInThisView;
+		}
 
 		// The shadow is visible if it is view relevant and unoccluded.
 		bOpaqueShadowIsVisibleThisFrame |= (bPrimitiveIsShadowRelevant && !bOpaqueShadowIsOccluded);
