Fixes hit proxy rendering in vulkan.

diff --git a/Engine/Source/Runtime/Renderer/Private/SceneHitProxyRendering.cpp b/Engine/Source/Runtime/Renderer/Private/SceneHitProxyRendering.cpp
--- a/Engine/Source/Runtime/Renderer/Private/SceneHitProxyRendering.cpp
+++ b/Engine/Source/Runtime/Renderer/Private/SceneHitProxyRendering.cpp
@@ -407,7 +407,7 @@ bool FHitProxyDrawingPolicyFactory::DrawDynamicMesh(
 
 #if WITH_EDITOR
 
-void InitHitProxyRenderAndBeginRenderPass(FRHICommandListImmediate& RHICmdList, const FSceneRenderer* SceneRenderer, TRefCountPtr<IPooledRenderTarget>& OutHitProxyRT, TRefCountPtr<IPooledRenderTarget>& OutHitProxyDepthRT)
+void InitHitProxyRender(FRHICommandListImmediate& RHICmdList, const FSceneRenderer* SceneRenderer, TRefCountPtr<IPooledRenderTarget>& OutHitProxyRT, TRefCountPtr<IPooledRenderTarget>& OutHitProxyDepthRT)
 {
 	check(!RHICmdList.IsInsideRenderPass());
 
@@ -440,14 +440,17 @@ void InitHitProxyRenderAndBeginRenderPass(FRHICommandListImmediate& RHICmdList,
 			OutHitProxyDepthRT = SceneContext.SceneDepthZ;
 		}
 	}
+}
 
-	FRHIRenderPassInfo RPInfo(OutHitProxyRT->GetRenderTargetItem().TargetableTexture, ERenderTargetActions::Load_Store);
+static void BeginHitProxyRenderpass(FRHICommandListImmediate& RHICmdList, const FSceneRenderer* SceneRenderer, TRefCountPtr<IPooledRenderTarget> HitProxyRT, TRefCountPtr<IPooledRenderTarget> HitProxyDepthRT)
+{
+	FRHIRenderPassInfo RPInfo(HitProxyRT->GetRenderTargetItem().TargetableTexture, ERenderTargetActions::Load_Store);
 	RPInfo.DepthStencilRenderTarget.Action = EDepthStencilTargetActions::LoadDepthStencil_StoreDepthStencil;
-	RPInfo.DepthStencilRenderTarget.DepthStencilTarget = OutHitProxyDepthRT->GetRenderTargetItem().TargetableTexture;
+	RPInfo.DepthStencilRenderTarget.DepthStencilTarget = HitProxyDepthRT->GetRenderTargetItem().TargetableTexture;
 	RPInfo.DepthStencilRenderTarget.ExclusiveDepthStencil = FExclusiveDepthStencil::DepthWrite_StencilWrite;
 	TransitionRenderPassTargets(RHICmdList, RPInfo);
 
-	RHICmdList.BeginRenderPass(RPInfo, TEXT("ClearHitProxies"));
+	RHICmdList.BeginRenderPass(RPInfo, TEXT("Clear_HitProxies"));
 	{
 		// Clear color for each view.
 		auto& Views = SceneRenderer->Views;
@@ -455,14 +458,14 @@ void InitHitProxyRenderAndBeginRenderPass(FRHICommandListImmediate& RHICmdList,
 		{
 			const FViewInfo& View = Views[ViewIndex];
 			RHICmdList.SetViewport(View.ViewRect.Min.X, View.ViewRect.Min.Y, 0.0f, View.ViewRect.Max.X, View.ViewRect.Max.Y, 1.0f);
-			DrawClearQuad(RHICmdList, true, FLinearColor::White, false, 0, false, 0, OutHitProxyRT->GetDesc().Extent, FIntRect());
+			DrawClearQuad(RHICmdList, true, FLinearColor::White, false, 0, false, 0, HitProxyRT->GetDesc().Extent, FIntRect());
 		}
 	}
 }
 
 static void DoRenderHitProxies(FRHICommandListImmediate& RHICmdList, const FSceneRenderer* SceneRenderer, TRefCountPtr<IPooledRenderTarget> HitProxyRT, TRefCountPtr<IPooledRenderTarget> HitProxyDepthRT)
 {
-	check(RHICmdList.IsInsideRenderPass());
+	BeginHitProxyRenderpass(RHICmdList, SceneRenderer, HitProxyRT, HitProxyDepthRT);
 
 	auto & ViewFamily = SceneRenderer->ViewFamily;
 	auto & Views = SceneRenderer->Views;
@@ -658,7 +661,7 @@ void FMobileSceneRenderer::RenderHitProxies(FRHICommandListImmediate& RHICmdList
 #if WITH_EDITOR
 	TRefCountPtr<IPooledRenderTarget> HitProxyRT;
 	TRefCountPtr<IPooledRenderTarget> HitProxyDepthRT;
-	InitHitProxyRenderAndBeginRenderPass(RHICmdList, this, HitProxyRT, HitProxyDepthRT);
+	InitHitProxyRender(RHICmdList, this, HitProxyRT, HitProxyDepthRT);
 	// HitProxyRT==0 should never happen but better we don't crash
 	if (HitProxyRT)
 	{
@@ -676,11 +679,11 @@ void FDeferredShadingSceneRenderer::RenderHitProxies(FRHICommandListImmediate& R
 {
 	PrepareViewRectsForRendering();
 
-#if WITH_EDITOR
 	TRefCountPtr<IPooledRenderTarget> HitProxyRT;
 	TRefCountPtr<IPooledRenderTarget> HitProxyDepthRT;
-	InitHitProxyRenderAndBeginRenderPass(RHICmdList, this, HitProxyRT, HitProxyDepthRT);
+	InitHitProxyRender(RHICmdList, this, HitProxyRT, HitProxyDepthRT);
 
+#if WITH_EDITOR
 	// HitProxyRT==0 should never happen but better we don't crash
 	if (HitProxyRT)
 	{
