Reapplies Jules's fix to the logic that clears stencil.

diff --git a/Engine/Source/Runtime/Renderer/Private/CompositionLighting/PostProcessDeferredDecals.cpp b/Engine/Source/Runtime/Renderer/Private/CompositionLighting/PostProcessDeferredDecals.cpp
--- a/Engine/Source/Runtime/Renderer/Private/CompositionLighting/PostProcessDeferredDecals.cpp
+++ b/Engine/Source/Runtime/Renderer/Private/CompositionLighting/PostProcessDeferredDecals.cpp
@@ -962,7 +962,10 @@ void FRCPassPostProcessDeferredDecals::Process(FRenderingCompositePassContext& C
 				RHICmdList.EndRenderPass();
 			}
 
-			if ((Scene.Decals.Num() > 0) && CurrentStage == DRS_Emissive)
+			// This stops the targets from being resolved and decoded until the last view is rendered.
+			// This is done so as to not run eliminate fast clear on the views before the end.
+			bool bLastView = Context.View.Family->Views.Last() == &Context.View;
+			if ((Scene.Decals.Num() > 0) && bLastView && CurrentStage == DRS_Emissive)
 			{
 				// we don't modify stencil but if out input was having stencil for us (after base pass - we need to clear)
 				// Clear stencil to 0, which is the assumed default by other passes
@@ -975,12 +978,8 @@ void FRCPassPostProcessDeferredDecals::Process(FRenderingCompositePassContext& C
 
 				RHICmdList.BeginRenderPass(RPInfo, TEXT("ClearStencil"));
 				RHICmdList.EndRenderPass();
-				//DrawClearQuad(RHICmdList, false, FLinearColor(), false, 0, true, 0, SceneContext.GetSceneDepthSurface()->GetSizeXY(), FIntRect());
 			}
 
-			// This stops the targets from being resolved and decoded until the last view is rendered.
-			// This is done so as to not run eliminate fast clear on the views before the end.
-			bool bLastView = Context.View.Family->Views.Last() == &Context.View;
 			if (CurrentStage == DRS_BeforeBasePass)
 			{
 				// combine DBuffer RTWriteMasks; will end up in one texture we can load from in the base pass PS and decide whether to do the actual work or not
@@ -1035,7 +1034,7 @@ void FRCPassPostProcessDeferredDecals::Process(FRenderingCompositePassContext& C
 		SceneContext.DBufferMask = nullptr;
 	}
 
-	check(!RHICmdList.IsInsideRenderPass());
+	check(RHICmdList.IsOutsideRenderPass());
 }
 
 FPooledRenderTargetDesc FRCPassPostProcessDeferredDecals::ComputeOutputDesc(EPassOutputId InPassOutputId) const
