Error checking to make sure we don't try to initialize views within a renderpass

diff --git a/Engine/Source/Runtime/Renderer/Private/SceneVisibility.cpp b/Engine/Source/Runtime/Renderer/Private/SceneVisibility.cpp
--- a/Engine/Source/Runtime/Renderer/Private/SceneVisibility.cpp
+++ b/Engine/Source/Runtime/Renderer/Private/SceneVisibility.cpp
@@ -3369,6 +3369,7 @@ bool FDeferredShadingSceneRenderer::InitViews(FRHICommandListImmediate& RHICmdLi
 {
 	SCOPED_NAMED_EVENT(FDeferredShadingSceneRenderer_InitViews, FColor::Emerald);
 	SCOPE_CYCLE_COUNTER(STAT_InitViewsTime);
+	check(RHICmdList.IsOutsideRenderPass());
 
 	PreVisibilityFrameSetup(RHICmdList);
 	RHICmdList.ImmediateFlush(EImmediateFlushType::DispatchToRHIThread);
