Renderpass updates. Base pass, Depth, Shadows, Translucency and Velocity.

diff --git a/Engine/Source/Runtime/Renderer/Private/BasePassRendering.cpp b/Engine/Source/Runtime/Renderer/Private/BasePassRendering.cpp
--- a/Engine/Source/Runtime/Renderer/Private/BasePassRendering.cpp
+++ b/Engine/Source/Runtime/Renderer/Private/BasePassRendering.cpp
@@ -33,13 +33,6 @@ static TAutoConsoleVariable<int32> CVarVertexFoggingForOpaque(
 	TEXT("Causes opaque materials to use per-vertex fogging, which costs less and integrates properly with MSAA.  Only supported with forward shading."),
 	ECVF_ReadOnly | ECVF_RenderThreadSafe);
 
-static TAutoConsoleVariable<int32> CVarParallelBasePass(
-	TEXT("r.ParallelBasePass"),
-	1,
-	TEXT("Toggles parallel base pass rendering. Parallel rendering must be enabled for this to have an effect."),
-	ECVF_RenderThreadSafe
-	);
-
 static TAutoConsoleVariable<int32> CVarRHICmdBasePassDeferredContexts(
 	TEXT("r.RHICmdBasePassDeferredContexts"),
 	1,
@@ -886,10 +879,10 @@ void CreateOpaqueBasePassUniformBuffer(
 }
 
 /**
- * Renders the scene's base pass 
+ * Renders the scene's base pass. This assumes there is a current renderpass active. 
  * @return true if anything was rendered
  */
-bool FDeferredShadingSceneRenderer::RenderBasePass(FRHICommandListImmediate& RHICmdList, FExclusiveDepthStencil::Type BasePassDepthStencilAccess, IPooledRenderTarget* ForwardScreenSpaceShadowMask)
+bool FDeferredShadingSceneRenderer::RenderBasePass(FRHICommandListImmediate& RHICmdList, FExclusiveDepthStencil::Type BasePassDepthStencilAccess, IPooledRenderTarget* ForwardScreenSpaceShadowMask, bool bParallelBasePass, bool bRenderLightmapDensity)
 {
 	SCOPED_NAMED_EVENT(FDeferredShadingSceneRenderer_RenderBasePass, FColor::Emerald);
 
@@ -897,7 +890,7 @@ bool FDeferredShadingSceneRenderer::RenderBasePass(FRHICommandListImmediate& RHI
 
 	RHICmdList.AutomaticCacheFlushAfterComputeShader(false);
 
-	if (ViewFamily.EngineShowFlags.LightMapDensity && AllowDebugViewmodes())
+	if (bRenderLightmapDensity)
 	{
 		// Override the base pass with the lightmap density pass if the viewmode is enabled.
 		bDirty = RenderLightMapDensities(RHICmdList);
@@ -908,8 +901,10 @@ bool FDeferredShadingSceneRenderer::RenderBasePass(FRHICommandListImmediate& RHI
 		SCOPE_CYCLE_COUNTER(STAT_BasePassDrawTime);
 		SCOPED_GPU_STAT(RHICmdList, Basepass);
 
-		if (GRHICommandList.UseParallelAlgorithms() && CVarParallelBasePass.GetValueOnRenderThread())
+		if (bParallelBasePass)
 		{
+			check(RHICmdList.IsOutsideRenderPass());
+
 			FScopedCommandListWaitForTasks Flusher(CVarRHICmdFlushRenderThreadTasksBasePass.GetValueOnRenderThread() > 0 || CVarRHICmdFlushRenderThreadTasks.GetValueOnRenderThread() > 0, RHICmdList);
 			for (int32 ViewIndex = 0; ViewIndex < Views.Num(); ViewIndex++)
 			{
@@ -926,14 +921,21 @@ bool FDeferredShadingSceneRenderer::RenderBasePass(FRHICommandListImmediate& RHI
 				{
 					RenderBasePassViewParallel(View, RHICmdList, BasePassDepthStencilAccess, DrawRenderState);
 				}
+				
+				check(RHICmdList.IsOutsideRenderPass());
 
+				FSceneRenderTargets::Get(RHICmdList).BeginRenderingGBuffer(RHICmdList, ERenderTargetLoadAction::ELoad, ERenderTargetLoadAction::ELoad, BasePassDepthStencilAccess, this->ViewFamily.EngineShowFlags.ShaderComplexity);
 				RenderEditorPrimitives(RHICmdList, View, BasePassDepthStencilAccess, DrawRenderState, bDirty);
+				RHICmdList.EndRenderPass();
 			}
 
 			bDirty = true; // assume dirty since we are not going to wait
 		}
 		else
 		{
+			// Must have an open renderpass before getting here in single threaded mode.
+			check(RHICmdList.IsInsideRenderPass());
+
 			for (int32 ViewIndex = 0; ViewIndex < Views.Num(); ViewIndex++)
 			{
 				SCOPED_CONDITIONAL_DRAW_EVENTF(RHICmdList, EventView, Views.Num() > 1, TEXT("View%d"), ViewIndex);
@@ -1157,8 +1159,12 @@ public:
 
 	void DoTask(ENamedThreads::Type CurrentThread, const FGraphEventRef& MyCompletionGraphEvent)
 	{
+		check(RHICmdList.IsInsideRenderPass());
+
 		bool OutDirty = false;
 		ThisRenderer.RenderBasePassDynamicData(RHICmdList, View, DrawRenderState, OutDirty);
+
+		RHICmdList.EndRenderPass();
 		RHICmdList.HandleRTThreadTaskCompletion(MyCompletionGraphEvent);
 	}
 };
@@ -1249,7 +1255,6 @@ public:
 		: FParallelCommandListSet(GET_STATID(STAT_CLP_Basepass), InView, InSceneRenderer, InParentCmdList, bInParallelExecute, bInCreateSceneContext, InDrawRenderState)
 		, BasePassDepthStencilAccess(InBasePassDepthStencilAccess)
 	{
-		SetStateOnCommandList(ParentCmdList);
 	}
 
 	virtual ~FBasePassParallelCommandListSet()
@@ -1267,6 +1272,8 @@ public:
 
 void FDeferredShadingSceneRenderer::RenderBasePassViewParallel(FViewInfo& View, FRHICommandListImmediate& ParentCmdList, FExclusiveDepthStencil::Type BasePassDepthStencilAccess, const FDrawingPolicyRenderState& InDrawRenderState)
 {
+	check(ParentCmdList.IsOutsideRenderPass());
+
 	FBasePassParallelCommandListSet ParallelSet(View, ParentCmdList, 
 		CVarRHICmdBasePassDeferredContexts.GetValueOnRenderThread() > 0, 
 		CVarRHICmdFlushRenderThreadTasksBasePass.GetValueOnRenderThread() == 0 && CVarRHICmdFlushRenderThreadTasks.GetValueOnRenderThread() == 0,
diff --git a/Engine/Source/Runtime/Renderer/Private/DeferredShadingRenderer.cpp b/Engine/Source/Runtime/Renderer/Private/DeferredShadingRenderer.cpp
--- a/Engine/Source/Runtime/Renderer/Private/DeferredShadingRenderer.cpp
+++ b/Engine/Source/Runtime/Renderer/Private/DeferredShadingRenderer.cpp
@@ -131,6 +131,13 @@ static TAutoConsoleVariable<int32> CVarBasePassWriteDepthEvenWithFullPrepass(
 	TEXT("0 to allow a readonly base pass, which skips an MSAA depth resolve, and allows masked materials to get EarlyZ (writing to depth while doing clip() disables EarlyZ) (default)\n")
 	TEXT("1 to force depth writes in the base pass.  Useful for debugging when the prepass and base pass don't match what they render."));
 
+static TAutoConsoleVariable<int32> CVarParallelBasePass(
+	TEXT("r.ParallelBasePass"),
+	1,
+	TEXT("Toggles parallel base pass rendering. Parallel rendering must be enabled for this to have an effect."),
+	ECVF_RenderThreadSafe
+);
+
 DECLARE_CYCLE_STAT(TEXT("PostInitViews FlushDel"), STAT_PostInitViews_FlushDel, STATGROUP_InitViews);
 DECLARE_CYCLE_STAT(TEXT("InitViews Intentional Stall"), STAT_InitViews_Intentional_Stall, STATGROUP_InitViews);
 
@@ -255,6 +262,7 @@ float GetSceneColorClearAlpha()
 void FDeferredShadingSceneRenderer::ClearGBufferAtMaxZ(FRHICommandList& RHICmdList)
 {
 	// Assumes BeginRenderingSceneColor() has been called before this function
+	check(RHICmdList.IsInsideRenderPass());
 	SCOPED_DRAW_EVENT(RHICmdList, ClearGBufferAtMaxZ);
 
 	// Clear the G Buffer render targets
@@ -425,19 +433,6 @@ static FORCEINLINE bool NeedsPrePass(const FDeferredShadingSceneRenderer* Render
 		(Renderer->EarlyZPassMode != DDM_None || Renderer->bEarlyZPassMovable != 0);
 }
 
-static void SetAndClearViewGBuffer(FRHICommandListImmediate& RHICmdList, FViewInfo& View, FExclusiveDepthStencil::Type DepthStencilAccess, bool bClearDepth)
-{
-	// if we didn't to the prepass above, then we will need to clear now, otherwise, it's already been cleared and rendered to
-	ERenderTargetLoadAction DepthLoadAction = bClearDepth ? ERenderTargetLoadAction::EClear : ERenderTargetLoadAction::ELoad;
-
-	const bool bClearBlack = View.Family->EngineShowFlags.ShaderComplexity || View.Family->EngineShowFlags.StationaryLightOverlap;
-	const float ClearAlpha = GetSceneColorClearAlpha();
-	const FLinearColor ClearColor = bClearBlack ? FLinearColor(0, 0, 0, ClearAlpha) : FLinearColor(View.BackgroundColor.R, View.BackgroundColor.G, View.BackgroundColor.B, ClearAlpha);
-
-	// clearing the GBuffer
-	FSceneRenderTargets::Get(RHICmdList).BeginRenderingGBuffer(RHICmdList, ERenderTargetLoadAction::EClear, DepthLoadAction, DepthStencilAccess, View.Family->EngineShowFlags.ShaderComplexity, ClearColor);
-}
-
 bool FDeferredShadingSceneRenderer::RenderHzb(FRHICommandListImmediate& RHICmdList)
 {
 	FSceneRenderTargets& SceneContext = FSceneRenderTargets::Get(RHICmdList);
@@ -480,7 +475,9 @@ bool FDeferredShadingSceneRenderer::RenderHzb(FRHICommandListImmediate& RHICmdLi
 }
 
 void FDeferredShadingSceneRenderer::RenderOcclusion(FRHICommandListImmediate& RHICmdList)
-{		
+{
+	check(RHICmdList.IsOutsideRenderPass());
+
 	SCOPED_GPU_STAT(RHICmdList, HZB);
 
 	{
@@ -558,6 +555,8 @@ static TAutoConsoleVariable<float> CVarStallInitViews(
 
 void FDeferredShadingSceneRenderer::Render(FRHICommandListImmediate& RHICmdList)
 {
+	check(RHICmdList.IsOutsideRenderPass());
+
 	PrepareViewRectsForRendering();
 
 	SCOPED_NAMED_EVENT(FDeferredShadingSceneRenderer_Render, FColor::Emerald);
@@ -568,7 +567,7 @@ void FDeferredShadingSceneRenderer::Render(FRHICommandListImmediate& RHICmdList)
 #endif // WITH_MGPU
 
 	FSceneRenderTargets& SceneContext = FSceneRenderTargets::Get(RHICmdList);
-	
+
 	//make sure all the targets we're going to use will be safely writable.
 	GRenderTargetPool.TransitionTargetsWritable(RHICmdList);
 
@@ -588,7 +587,7 @@ void FDeferredShadingSceneRenderer::Render(FRHICommandListImmediate& RHICmdList)
 	// Anything rendered inside Render() which isn't accounted for will fall into this stat
 	// This works because child stat events do not contribute to their parents' times (see GPU_STATS_CHILD_TIMES_INCLUDED)
 	SCOPED_GPU_STAT(RHICmdList, Unaccounted);
-	
+
 	{
 		SCOPE_CYCLE_COUNTER(STAT_FDeferredShadingSceneRenderer_Render_Init);
 
@@ -611,7 +610,7 @@ void FDeferredShadingSceneRenderer::Render(FRHICommandListImmediate& RHICmdList)
 	if (CVarVirtualTextureLightmaps && CVarVirtualTextureLightmaps->GetValueOnRenderThread())
 	{
 		// TODO should probably be in InitViews
-		GetVirtualTextureSystem()->Update( RHICmdList, FeatureLevel );
+		GetVirtualTextureSystem()->Update(RHICmdList, FeatureLevel);
 	}
 
 	TGuardValue<bool> LockDrawLists(GDrawListsLocked, true);
@@ -629,7 +628,7 @@ void FDeferredShadingSceneRenderer::Render(FRHICommandListImmediate& RHICmdList)
 		for (int32 ViewIndex = 0; ViewIndex < Views.Num(); ViewIndex++)
 		{
 			Views[ViewIndex].GetEyeAdaptation(RHICmdList);
-		}	
+		}
 	}
 
 	if (GDoPrepareDistanceFieldSceneAfterRHIFlush && (GRHINeedsExtraDeletionLatency || !GRHICommandList.Bypass()))
@@ -694,25 +693,25 @@ void FDeferredShadingSceneRenderer::Render(FRHICommandListImmediate& RHICmdList)
 		switch (ClearMethod)
 		{
 		case 0: // No clear
-			{
-				bRequiresRHIClear = false;
-				bRequiresFarZQuadClear = false;
-				break;
-			}
-		
+		{
+			bRequiresRHIClear = false;
+			bRequiresFarZQuadClear = false;
+			break;
+		}
+
 		case 1: // RHICmdList.Clear
-			{
-				bRequiresRHIClear = true;
-				bRequiresFarZQuadClear = false;
-				break;
-			}
+		{
+			bRequiresRHIClear = true;
+			bRequiresFarZQuadClear = false;
+			break;
+		}
 
 		case 2: // Clear using far-z quad
-			{
-				bRequiresFarZQuadClear = true;
-				bRequiresRHIClear = false;
-				break;
-			}
+		{
+			bRequiresFarZQuadClear = true;
+			bRequiresRHIClear = false;
+			break;
+		}
 		}
 	}
 
@@ -724,15 +723,15 @@ void FDeferredShadingSceneRenderer::Render(FRHICommandListImmediate& RHICmdList)
 
 	// force using occ queries for wireframe if rendering is parented or frozen in the first view
 	check(Views.Num());
-	#if (UE_BUILD_SHIPPING || UE_BUILD_TEST)
-		const bool bIsViewFrozen = false;
-		const bool bHasViewParent = false;
-	#else
-		const bool bIsViewFrozen = Views[0].State && ((FSceneViewState*)Views[0].State)->bIsFrozen;
-		const bool bHasViewParent = Views[0].State && ((FSceneViewState*)Views[0].State)->HasViewParent();
-	#endif
+#if (UE_BUILD_SHIPPING || UE_BUILD_TEST)
+	const bool bIsViewFrozen = false;
+	const bool bHasViewParent = false;
+#else
+	const bool bIsViewFrozen = Views[0].State && ((FSceneViewState*)Views[0].State)->bIsFrozen;
+	const bool bHasViewParent = Views[0].State && ((FSceneViewState*)Views[0].State)->HasViewParent();
+#endif
+
 
-	
 	const bool bIsOcclusionTesting = DoOcclusionQueries(FeatureLevel) && (!bIsWireframe || bIsViewFrozen || bHasViewParent);
 
 	// Dynamic vertex and index buffers need to be committed before rendering.
@@ -805,15 +804,19 @@ void FDeferredShadingSceneRenderer::Render(FRHICommandListImmediate& RHICmdList)
 		QUICK_SCOPE_CYCLE_COUNTER(STAT_FDeferredShadingSceneRenderer_AsyncUpdateViewCustomData_Wait);
 		FTaskGraphInterface::Get().WaitUntilTasksComplete(UpdateViewCustomDataEvents, ENamedThreads::GetRenderThread());
 	}
-	
+
 	if (CVarVirtualTextureLightmaps && CVarVirtualTextureLightmaps->GetValueOnRenderThread())
 	{
 		// Create VT feedback buffer
 		FIntPoint Size = SceneContext.GetBufferSizeXY();
-		Size = FIntPoint::DivideAndRoundUp( Size, 16 );
-		GVirtualTextureFeedback.CreateResourceGPU( RHICmdList, Size.X, Size.Y );
+		Size = FIntPoint::DivideAndRoundUp(Size, 16);
+		GVirtualTextureFeedback.CreateResourceGPU(RHICmdList, Size.X, Size.Y);
 	}
 
+	checkSlow(RHICmdList.IsOutsideRenderPass());
+
+	// The Z-prepass
+
 	// Draw the scene pre-pass / early z pass, populating the scene depth buffer and HiZ
 	GRenderTargetPool.AddPhaseEvent(TEXT("EarlyZPass"));
 	const bool bNeedsPrePass = NeedsPrePass(this);
@@ -833,6 +836,10 @@ void FDeferredShadingSceneRenderer::Render(FRHICommandListImmediate& RHICmdList)
 	RHICmdList.SetCurrentStat(GET_STATID(STAT_CLM_AfterPrePass));
 	ServiceLocalQueue();
 
+	// Z-Prepass End
+
+	checkSlow(RHICmdList.IsOutsideRenderPass());
+
 	const bool bShouldRenderVelocities = ShouldRenderVelocities();
 	const bool bBasePassCanOutputVelocity = FVelocityRendering::BasePassCanOutputVelocity(FeatureLevel);
 	const bool bUseSelectiveBasePassOutputs = bUseGBuffer && UseSelectiveBasePassOutputs();
@@ -845,9 +852,9 @@ void FDeferredShadingSceneRenderer::Render(FRHICommandListImmediate& RHICmdList)
 		&& !bIsWireframe
 		&& !ViewFamily.EngineShowFlags.LightMapDensity;
 
-	const FExclusiveDepthStencil::Type BasePassDepthStencilAccess = 
+	const FExclusiveDepthStencil::Type BasePassDepthStencilAccess =
 		bAllowReadonlyDepthBasePass
-		? FExclusiveDepthStencil::DepthRead_StencilWrite 
+		? FExclusiveDepthStencil::DepthRead_StencilWrite
 		: FExclusiveDepthStencil::DepthWrite_StencilWrite;
 
 	SceneContext.ResolveSceneDepthTexture(RHICmdList, FResolveRect(0, 0, FamilySize.X, FamilySize.Y));
@@ -860,8 +867,11 @@ void FDeferredShadingSceneRenderer::Render(FRHICommandListImmediate& RHICmdList)
 		SCOPE_CYCLE_COUNTER(STAT_FDeferredShadingSceneRenderer_AllocGBufferTargets);
 		SceneContext.PreallocGBufferTargets(); // Even if !bShouldRenderVelocities, the velocity buffer must be bound because it's a compile time option for the shader.
 		SceneContext.AllocGBufferTargets(RHICmdList);
-	}	
+	}
+
+	checkSlow(RHICmdList.IsOutsideRenderPass());
 
+	// Early occlusion queries
 	const bool bOcclusionBeforeBasePass = (EarlyZPassMode == EDepthDrawingMode::DDM_AllOccluders) || (EarlyZPassMode == EDepthDrawingMode::DDM_AllOpaque);
 
 	if (bOcclusionBeforeBasePass)
@@ -871,7 +881,7 @@ void FDeferredShadingSceneRenderer::Render(FRHICommandListImmediate& RHICmdList)
 			RenderOcclusion(RHICmdList);
 		}
 		bool bUseHzbOcclusion = RenderHzb(RHICmdList);
-		
+
 		SCOPED_GPU_STAT(RHICmdList, HZB);
 		if (bUseHzbOcclusion || bIsOcclusionTesting)
 		{
@@ -884,7 +894,11 @@ void FDeferredShadingSceneRenderer::Render(FRHICommandListImmediate& RHICmdList)
 	}
 
 	ServiceLocalQueue();
+	// End early occlusion queries
+
+	checkSlow(RHICmdList.IsOutsideRenderPass());
 
+	// Early Shadow depth rendering
 	if (bOcclusionBeforeBasePass)
 	{
 		// Before starting the shadow render, all async task for the shadow Custom data must be completed
@@ -897,6 +911,9 @@ void FDeferredShadingSceneRenderer::Render(FRHICommandListImmediate& RHICmdList)
 		RenderShadowDepthMaps(RHICmdList);
 		ServiceLocalQueue();
 	}
+	// End early Shadow depth rendering
+
+	checkSlow(RHICmdList.IsOutsideRenderPass());
 
 	// Clear LPVs for all views
 	if (FeatureLevel >= ERHIFeatureLevel::SM5)
@@ -906,7 +923,7 @@ void FDeferredShadingSceneRenderer::Render(FRHICommandListImmediate& RHICmdList)
 		ServiceLocalQueue();
 	}
 
-	if(GetCustomDepthPassLocation() == 0)
+	if (GetCustomDepthPassLocation() == 0)
 	{
 		QUICK_SCOPE_CYCLE_COUNTER(STAT_FDeferredShadingSceneRenderer_CustomDepthPass0);
 		RenderCustomDepthPassAtLocation(RHICmdList, 0);
@@ -924,8 +941,8 @@ void FDeferredShadingSceneRenderer::Render(FRHICommandListImmediate& RHICmdList)
 		RenderForwardShadingShadowProjections(RHICmdList, ForwardScreenSpaceShadowMask);
 
 		RenderIndirectCapsuleShadows(
-			RHICmdList, 
-			NULL, 
+			RHICmdList,
+			NULL,
 			NULL);
 	}
 
@@ -939,9 +956,9 @@ void FDeferredShadingSceneRenderer::Render(FRHICommandListImmediate& RHICmdList)
 		SCOPE_CYCLE_COUNTER(STAT_FDeferredShadingSceneRenderer_DBuffer);
 
 		// e.g. DBuffer deferred decals
-		for(int32 ViewIndex = 0;ViewIndex < Views.Num();ViewIndex++)
-		{	
-			SCOPED_CONDITIONAL_DRAW_EVENTF(RHICmdList, EventView,Views.Num() > 1, TEXT("View%d"), ViewIndex);
+		for (int32 ViewIndex = 0; ViewIndex < Views.Num(); ViewIndex++)
+		{
+			SCOPED_CONDITIONAL_DRAW_EVENTF(RHICmdList, EventView, Views.Num() > 1, TEXT("View%d"), ViewIndex);
 			FViewInfo& View = Views[ViewIndex];
 
 			uint32 SSAOLevels = FSSAOHelper::ComputeAmbientOcclusionPassCount(View);
@@ -957,16 +974,18 @@ void FDeferredShadingSceneRenderer::Render(FRHICommandListImmediate& RHICmdList)
 
 		ServiceLocalQueue();
 	}
-	
+
+	checkSlow(RHICmdList.IsOutsideRenderPass());
+
 	if (bRenderDeferredLighting)
 	{
 		bool bShouldAllocateDeferredShadingPathRenderTargets = false;
 		const char* str = SceneContext.ScreenSpaceAO ? "Allocated" : "Unallocated"; //ScreenSpaceAO is determining factor of detecting render target allocation
-		for(int Index = 0; Index < (NumTranslucentVolumeRenderTargetSets * Views.Num()); ++Index)
+		for (int Index = 0; Index < (NumTranslucentVolumeRenderTargetSets * Views.Num()); ++Index)
 		{
-			if(!SceneContext.TranslucencyLightingVolumeAmbient[Index] || !SceneContext.TranslucencyLightingVolumeDirectional[Index])
+			if (!SceneContext.TranslucencyLightingVolumeAmbient[Index] || !SceneContext.TranslucencyLightingVolumeDirectional[Index])
 			{
-				
+
 				ensureMsgf(SceneContext.TranslucencyLightingVolumeAmbient[Index], TEXT("%s%d is unallocated, Deferred Render Targets would be detected as: %s"), "TranslucencyLightingVolumeAmbient", Index, str);
 				ensureMsgf(SceneContext.TranslucencyLightingVolumeDirectional[Index], TEXT("%s%d is unallocated, Deferred Render Targets would be detected as: %s"), "TranslucencyLightingVolumeDirectional", Index, str);
 				bShouldAllocateDeferredShadingPathRenderTargets = true;
@@ -974,50 +993,107 @@ void FDeferredShadingSceneRenderer::Render(FRHICommandListImmediate& RHICmdList)
 			}
 		}
 
-		if(bShouldAllocateDeferredShadingPathRenderTargets)
+		if (bShouldAllocateDeferredShadingPathRenderTargets)
 		{
 			SceneContext.AllocateDeferredShadingPathRenderTargets(RHICmdList);
 		}
-		
+
 		if (GbEnableAsyncComputeTranslucencyLightingVolumeClear && GSupportsEfficientAsyncCompute)
 		{
 			ClearTranslucentVolumeLightingAsyncCompute(RHICmdList);
 		}
 	}
 
+	checkSlow(RHICmdList.IsOutsideRenderPass());
+
+	bool bIsWireframeRenderpass = bIsWireframe && FSceneRenderer::ShouldCompositeEditorPrimitives(Views[0]);
+	bool bRenderLightmapDensity = ViewFamily.EngineShowFlags.LightMapDensity && AllowDebugViewmodes();
+	bool bDoParallelBasePass = GRHICommandList.UseParallelAlgorithms() && CVarParallelBasePass.GetValueOnRenderThread();
+	
+	// BASE PASS AND GBUFFER SETUP
+	// Gross logic to cover all the cases of special rendering modes + parallel dispatch
 	// Clear the GBuffer render targets
 	bool bIsGBufferCurrent = false;
 	if (bRequiresRHIClear)
 	{
 		SCOPE_CYCLE_COUNTER(STAT_FDeferredShadingSceneRenderer_SetAndClearViewGBuffer);
-		// set GBuffer to be current, and clear it
-		SetAndClearViewGBuffer(RHICmdList, Views[0], BasePassDepthStencilAccess, !bDepthWasCleared);
-		
-		// depth was cleared now no matter what
-		bDepthWasCleared = true;
-		bIsGBufferCurrent = true;
+		bool bClearDepth = !bDepthWasCleared;
+		 
+		// if we didn't to the prepass above, then we will need to clear now, otherwise, it's already been cleared and rendered to
+		ERenderTargetLoadAction ColorLoadAction = ERenderTargetLoadAction::ELoad;
+		ERenderTargetLoadAction DepthLoadAction = bClearDepth ? ERenderTargetLoadAction::EClear : ERenderTargetLoadAction::ENoAction;
+
+		const bool bClearBlack = ViewFamily.EngineShowFlags.ShaderComplexity || ViewFamily.EngineShowFlags.StationaryLightOverlap;
+		const float ClearAlpha = GetSceneColorClearAlpha();
+		FLinearColor ClearColor = bClearBlack ? FLinearColor(0, 0, 0, ClearAlpha) : FLinearColor(Views[0].BackgroundColor.R, Views[0].BackgroundColor.G, Views[0].BackgroundColor.B, ClearAlpha);
+		ColorLoadAction = ERenderTargetLoadAction::EClear;
+
+		// The first time through we'll clear the Overdraw UAVs.
+		SceneContext.BeginRenderingGBuffer(RHICmdList, ColorLoadAction, DepthLoadAction, BasePassDepthStencilAccess, ViewFamily.EngineShowFlags.ShaderComplexity, true, ClearColor);
+
+		// If we are in wireframe mode or will go wide later this pass is just the clear.
+		if (bIsWireframeRenderpass || bDoParallelBasePass)
+		{
+			RHICmdList.EndRenderPass();
+		}
+		else
+		{
+			bIsGBufferCurrent = true;
+		}
 		ServiceLocalQueue();
 	}
-	
-	if (bIsWireframe && FSceneRenderer::ShouldCompositeEditorPrimitives(Views[0]))
+
+	// Wireframe mode requires bRequiresRHIClear to be true. 
+	// Rendering will be very funny without it and the call to BeginRenderingGBuffer will call AllocSceneColor which is needed for the EditorPrimitives resolve.
+	if (bIsWireframeRenderpass)
 	{
+		check(bRequiresRHIClear);
+
 		// In Editor we want wire frame view modes to be MSAA for better quality. Resolve will be done with EditorPrimitives
-		SetRenderTarget(RHICmdList, SceneContext.GetEditorPrimitivesColor(RHICmdList), SceneContext.GetEditorPrimitivesDepth(RHICmdList), ESimpleRenderTargetMode::EClearColorAndDepth);
+		FRHIRenderPassInfo RPInfo(SceneContext.GetEditorPrimitivesColor(RHICmdList), ERenderTargetActions::Clear_Store);
+		RPInfo.DepthStencilRenderTarget.Action = EDepthStencilTargetActions::ClearDepthStencil_StoreDepthStencil;
+		RPInfo.DepthStencilRenderTarget.DepthStencilTarget = SceneContext.GetEditorPrimitivesDepth(RHICmdList);
+		RPInfo.DepthStencilRenderTarget.ExclusiveDepthStencil = FExclusiveDepthStencil::DepthWrite_StencilWrite;
+		RHICmdList.BeginRenderPass(RPInfo, TEXT("Wireframe"));
+
+		// #todo-renderpasses In serial mode wireframe rendering only binds one target
+		// In parallel the entire gbuffer is bound. This was the previous SetRenderTarget behavior, preserved here.
+		// This is just a clear in the parallel case.
+		if (bDoParallelBasePass)
+		{
+			RHICmdList.EndRenderPass();
+		}
 	}
-	else if (!bIsGBufferCurrent)
+	else if (!bIsGBufferCurrent && (!bDoParallelBasePass || bRenderLightmapDensity))
 	{
-		// make sure the GBuffer is set, in case we didn't need to clear above
+		// Make sure we have began the renderpass
 		ERenderTargetLoadAction DepthLoadAction = bDepthWasCleared ? ERenderTargetLoadAction::ELoad : ERenderTargetLoadAction::EClear;
 		SceneContext.BeginRenderingGBuffer(RHICmdList, ERenderTargetLoadAction::ENoAction, DepthLoadAction, BasePassDepthStencilAccess, ViewFamily.EngineShowFlags.ShaderComplexity);
 	}
-
 	GRenderTargetPool.AddPhaseEvent(TEXT("BasePass"));
 
 	RHICmdList.SetCurrentStat(GET_STATID(STAT_CLM_BasePass));
-	RenderBasePass(RHICmdList, BasePassDepthStencilAccess, ForwardScreenSpaceShadowMask.GetReference());
+	RenderBasePass(RHICmdList, BasePassDepthStencilAccess, ForwardScreenSpaceShadowMask.GetReference(), bDoParallelBasePass, bRenderLightmapDensity);
 	RHICmdList.SetCurrentStat(GET_STATID(STAT_CLM_AfterBasePass));
 	ServiceLocalQueue();
 
+	// If we ran parallel in the basepass there will be no renderpass at this point.
+	if (bDoParallelBasePass && !bRenderLightmapDensity)
+	{
+		if (bIsWireframeRenderpass)
+		{
+			FRHIRenderPassInfo RPInfo(SceneContext.GetEditorPrimitivesColor(RHICmdList), ERenderTargetActions::Clear_Store);
+			RPInfo.DepthStencilRenderTarget.Action = EDepthStencilTargetActions::ClearDepthStencil_StoreDepthStencil;
+			RPInfo.DepthStencilRenderTarget.DepthStencilTarget = SceneContext.GetEditorPrimitivesDepth(RHICmdList);
+			RPInfo.DepthStencilRenderTarget.ExclusiveDepthStencil = FExclusiveDepthStencil::DepthWrite_StencilWrite;
+			RHICmdList.BeginRenderPass(RPInfo, TEXT("Wireframe"));
+		}
+		else
+		{
+			SceneContext.BeginRenderingGBuffer(RHICmdList, ERenderTargetLoadAction::ELoad, ERenderTargetLoadAction::ELoad, BasePassDepthStencilAccess, ViewFamily.EngineShowFlags.ShaderComplexity);
+		}
+	}
+
 	{
 		SCOPE_CYCLE_COUNTER(STAT_FDeferredShadingSceneRenderer_ViewExtensionPostRenderBasePass);
 		for (int32 ViewExt = 0; ViewExt < ViewFamily.ViewExtensions.Num(); ++ViewExt)
@@ -1029,18 +1105,46 @@ void FDeferredShadingSceneRenderer::Render(FRHICommandListImmediate& RHICmdList)
 		}
 	}
 
+	// #todo-renderpasses Should this be further below?
+	if (bRequiresFarZQuadClear)
+	{
+		SCOPE_CYCLE_COUNTER(STAT_FDeferredShadingSceneRenderer_ClearGBufferAtMaxZ);
+		// Clears view by drawing quad at maximum Z
+		// TODO: if all the platforms have fast color clears, we can replace this with an RHICmdList.Clear.
+		ClearGBufferAtMaxZ(RHICmdList);
+		ServiceLocalQueue();
+
+		bRequiresFarZQuadClear = false;
+	}
+
+	if (bUseGBuffer)
+	{
+		SCOPE_CYCLE_COUNTER(STAT_FDeferredShadingSceneRenderer_Resolve_After_Basepass);
+		SceneContext.FinishGBufferPassAndResolve(RHICmdList);
+	}
+	else
+	{
+		// #todo-renderpasses which paths can lead us here?
+		RHICmdList.EndRenderPass();
+	}
+
 	if (!bAllowReadonlyDepthBasePass)
 	{
 		SceneContext.ResolveSceneDepthTexture(RHICmdList, FResolveRect(0, 0, FamilySize.X, FamilySize.Y));
 		SceneContext.ResolveSceneDepthToAuxiliaryTexture(RHICmdList);
 	}
 
+	// BASE PASS ENDS HERE.
+
 	if (ViewFamily.EngineShowFlags.VisualizeLightCulling)
 	{
 		// clear out emissive and baked lighting (not too efficient but simple and only needed for this debug view)
 		SceneContext.BeginRenderingSceneColor(RHICmdList);
 		DrawClearQuad(RHICmdList, FLinearColor(0, 0, 0, 0));
+		SceneContext.FinishRenderingSceneColor(RHICmdList);
 	}
+	
+	checkSlow(RHICmdList.IsOutsideRenderPass());
 
 	SceneContext.DBufferA.SafeRelease();
 	SceneContext.DBufferB.SafeRelease();
@@ -1057,23 +1161,16 @@ void FDeferredShadingSceneRenderer::Render(FRHICommandListImmediate& RHICmdList)
 		GVirtualTextureFeedback.TransferGPUToCPU( RHICmdList );
 	}
 
-	if (bRequiresFarZQuadClear)
-	{
-		SCOPE_CYCLE_COUNTER(STAT_FDeferredShadingSceneRenderer_ClearGBufferAtMaxZ);
-		// Clears view by drawing quad at maximum Z
-		// TODO: if all the platforms have fast color clears, we can replace this with an RHICmdList.Clear.
-		ClearGBufferAtMaxZ(RHICmdList);
-		ServiceLocalQueue();
-
-		bRequiresFarZQuadClear = false;
-	}
+	// #todo-renderpass Zfar clear was here. where should it really go?
 	
 	VisualizeVolumetricLightmap(RHICmdList);
 
 	SceneContext.ResolveSceneDepthToAuxiliaryTexture(RHICmdList);
 
+	// Occlusion after base pass
 	if (!bOcclusionBeforeBasePass)
 	{
+		// #todo-renderpasses Needs its own renderpass. Does this need more than the depth?
 		if (bIsOcclusionTesting)
 		{
 			RenderOcclusion(RHICmdList);
@@ -1091,17 +1188,16 @@ void FDeferredShadingSceneRenderer::Render(FRHICommandListImmediate& RHICmdList)
 	}
 
 	ServiceLocalQueue();
+	// End occlusion after base
 
-	if (bUseGBuffer)
-	{
-		SCOPE_CYCLE_COUNTER(STAT_FDeferredShadingSceneRenderer_Resolve_After_Basepass);
-		SceneContext.FinishRenderingGBuffer(RHICmdList);
-	}
-	else
+	checkSlow(RHICmdList.IsOutsideRenderPass());
+
+	if (!bUseGBuffer)
 	{		
 		ResolveSceneColor(RHICmdList);
 	}
 
+	// Shadow and fog after base pass
 	if (!bOcclusionBeforeBasePass)
 	{
 		// Before starting the shadow render, all async task for the shadow Custom data must be completed
@@ -1112,9 +1208,15 @@ void FDeferredShadingSceneRenderer::Render(FRHICommandListImmediate& RHICmdList)
 		}
 
 		RenderShadowDepthMaps(RHICmdList);
+
+		checkSlow(RHICmdList.IsOutsideRenderPass());
+
 		ComputeVolumetricFog(RHICmdList);
 		ServiceLocalQueue();
 	}
+	// End shadow and fog after base pass
+
+	checkSlow(RHICmdList.IsOutsideRenderPass());
 
 	if(GetCustomDepthPassLocation() == 1)
 	{
@@ -1142,6 +1244,8 @@ void FDeferredShadingSceneRenderer::Render(FRHICommandListImmediate& RHICmdList)
 		ServiceLocalQueue();
 	}
 
+	checkSlow(RHICmdList.IsOutsideRenderPass());
+
 	TRefCountPtr<IPooledRenderTarget> VelocityRT;
 
 	if (bBasePassCanOutputVelocity)
@@ -1159,13 +1263,19 @@ void FDeferredShadingSceneRenderer::Render(FRHICommandListImmediate& RHICmdList)
 		ServiceLocalQueue();
 	}
 
+	checkSlow(RHICmdList.IsOutsideRenderPass());
+
 	// Copy lighting channels out of stencil before deferred decals which overwrite those values
 	CopyStencilToLightingChannelTexture(RHICmdList);
 
+	checkSlow(RHICmdList.IsOutsideRenderPass());
+
 	{
 		GCompositionLighting.GfxWaitForAsyncSSAO(RHICmdList);
 	}
 
+	checkSlow(RHICmdList.IsOutsideRenderPass());
+
 	// Pre-lighting composition lighting stage
 	// e.g. deferred decals, SSAO
 	if (FeatureLevel >= ERHIFeatureLevel::SM4)
@@ -1201,6 +1311,8 @@ void FDeferredShadingSceneRenderer::Render(FRHICommandListImmediate& RHICmdList)
 		RHICmdList.TransitionResource(EResourceTransitionAccess::EReadable, SceneContext.GetSceneDepthSurface());
 	}
 
+	checkSlow(RHICmdList.IsOutsideRenderPass());
+
 	// Render lighting.
 	if (bRenderDeferredLighting)
 	{
@@ -1233,6 +1345,8 @@ void FDeferredShadingSceneRenderer::Render(FRHICommandListImmediate& RHICmdList)
 		RHICmdList.SetCurrentStat(GET_STATID(STAT_CLM_AfterLighting));
 		ServiceLocalQueue();
 
+		checkSlow(RHICmdList.IsOutsideRenderPass());
+
 		GRenderTargetPool.AddPhaseEvent(TEXT("AfterRenderLights"));
 
 		for(int32 ViewIndex = 0; ViewIndex < Views.Num(); ++ViewIndex)
@@ -1248,6 +1362,8 @@ void FDeferredShadingSceneRenderer::Render(FRHICommandListImmediate& RHICmdList)
 		}
 		ServiceLocalQueue();
 
+		checkSlow(RHICmdList.IsOutsideRenderPass());
+
 		// Pre-lighting composition lighting stage
 		// e.g. LPV indirect
 		for (int32 ViewIndex = 0; ViewIndex < Views.Num(); ++ViewIndex)
@@ -1263,6 +1379,8 @@ void FDeferredShadingSceneRenderer::Render(FRHICommandListImmediate& RHICmdList)
 			}
 		}
 
+		checkSlow(RHICmdList.IsOutsideRenderPass());
+
 		// Render diffuse sky lighting and reflections that only operate on opaque pixels
 		RenderDeferredReflectionsAndSkyLighting(RHICmdList, DynamicBentNormalAO, VelocityRT);
 
@@ -1283,6 +1401,8 @@ void FDeferredShadingSceneRenderer::Render(FRHICommandListImmediate& RHICmdList)
 		ServiceLocalQueue();
 	}
 
+	checkSlow(RHICmdList.IsOutsideRenderPass());
+
 	FLightShaftsOutput LightShaftOutput;
 
 	// Draw Lightshafts
@@ -1293,6 +1413,8 @@ void FDeferredShadingSceneRenderer::Render(FRHICommandListImmediate& RHICmdList)
 		ServiceLocalQueue();
 	}
 
+	checkSlow(RHICmdList.IsOutsideRenderPass());
+
 	// Draw atmosphere
 	if (ShouldRenderAtmosphere(ViewFamily))
 	{
@@ -1320,6 +1442,8 @@ void FDeferredShadingSceneRenderer::Render(FRHICommandListImmediate& RHICmdList)
 		}
 	}
 
+	checkSlow(RHICmdList.IsOutsideRenderPass());
+
 	GRenderTargetPool.AddPhaseEvent(TEXT("Fog"));
 
 	// Draw fog.
@@ -1330,6 +1454,8 @@ void FDeferredShadingSceneRenderer::Render(FRHICommandListImmediate& RHICmdList)
 		ServiceLocalQueue();
 	}
 
+	checkSlow(RHICmdList.IsOutsideRenderPass());
+
 	IRendererModule& RendererModule = GetRendererModule();
 	if (RendererModule.HasPostOpaqueExtentions())
 	{
@@ -1344,8 +1470,9 @@ void FDeferredShadingSceneRenderer::Render(FRHICommandListImmediate& RHICmdList)
 			RHICmdList.SetViewport(View.ViewRect.Min.X, View.ViewRect.Min.Y, 0.0f, View.ViewRect.Max.X, View.ViewRect.Max.Y, 1.0f);
 			RendererModule.RenderPostOpaqueExtensions(View, RHICmdList, SceneContext, SceneTextureUniformBuffer);
 		}
+		SceneContext.FinishRenderingSceneColor(RHICmdList);
 	}
-	SetRenderTarget(RHICmdList, nullptr, 0, 0, nullptr);
+	checkSlow(RHICmdList.IsOutsideRenderPass());
 	RendererModule.DispatchPostOpaqueCompute(RHICmdList, Views[0].ViewUniformBuffer);
 
 	// No longer needed, release
@@ -1392,6 +1519,8 @@ void FDeferredShadingSceneRenderer::Render(FRHICommandListImmediate& RHICmdList)
 		RHICmdList.SetCurrentStat(GET_STATID(STAT_CLM_AfterTranslucency));
 	}
 
+	checkSlow(RHICmdList.IsOutsideRenderPass());
+
 	if (ViewFamily.EngineShowFlags.LightShafts)
 	{
 		SCOPE_CYCLE_COUNTER(STAT_FDeferredShadingSceneRenderer_RenderLightShaftBloom);
@@ -1417,6 +1546,8 @@ void FDeferredShadingSceneRenderer::Render(FRHICommandListImmediate& RHICmdList)
 		ServiceLocalQueue();
 	}
 
+	checkSlow(RHICmdList.IsOutsideRenderPass());
+
 	// Draw visualizations just before use to avoid target contamination
 	if (ViewFamily.EngineShowFlags.VisualizeMeshDistanceFields || ViewFamily.EngineShowFlags.VisualizeGlobalDistanceField)
 	{
@@ -1579,49 +1710,55 @@ void FDeferredShadingSceneRenderer::UpdateDownsampledDepthSurface(FRHICommandLis
 void FDeferredShadingSceneRenderer::DownsampleDepthSurface(FRHICommandList& RHICmdList, const FTexture2DRHIRef& RenderTarget, const FViewInfo& View, float ScaleFactor, bool bUseMaxDepth)
 {
 	FSceneRenderTargets& SceneContext = FSceneRenderTargets::Get(RHICmdList);
-	SetRenderTarget(RHICmdList, NULL, RenderTarget);
-	SCOPED_DRAW_EVENT(RHICmdList, DownsampleDepth);
 
-	// Set shaders and texture
-	TShaderMapRef<FScreenVS> ScreenVertexShader(View.ShaderMap);
-	TShaderMapRef<FDownsampleSceneDepthPS> PixelShader(View.ShaderMap);
+	FRHIRenderPassInfo RPInfo;
+	RPInfo.DepthStencilRenderTarget.Action = EDepthStencilTargetActions::LoadDepthStencil_StoreDepthStencil;
+	RPInfo.DepthStencilRenderTarget.DepthStencilTarget = RenderTarget;
+	RPInfo.DepthStencilRenderTarget.ExclusiveDepthStencil = FExclusiveDepthStencil::DepthWrite_StencilWrite;
+	RHICmdList.BeginRenderPass(RPInfo, TEXT("DownsampleDepth"));
+	{
+		SCOPED_DRAW_EVENT(RHICmdList, DownsampleDepth);
 
-	extern TGlobalResource<FFilterVertexDeclaration> GFilterVertexDeclaration;
+		// Set shaders and texture
+		TShaderMapRef<FScreenVS> ScreenVertexShader(View.ShaderMap);
+		TShaderMapRef<FDownsampleSceneDepthPS> PixelShader(View.ShaderMap);
 
-	FGraphicsPipelineStateInitializer GraphicsPSOInit;
-	RHICmdList.ApplyCachedRenderTargets(GraphicsPSOInit);
+		extern TGlobalResource<FFilterVertexDeclaration> GFilterVertexDeclaration;
 
-	GraphicsPSOInit.BlendState = TStaticBlendState<CW_NONE>::GetRHI();
-	GraphicsPSOInit.RasterizerState = TStaticRasterizerState<FM_Solid, CM_None>::GetRHI();
-	GraphicsPSOInit.DepthStencilState = TStaticDepthStencilState<true, CF_Always>::GetRHI();
+		FGraphicsPipelineStateInitializer GraphicsPSOInit;
+		RHICmdList.ApplyCachedRenderTargets(GraphicsPSOInit);
+
+		GraphicsPSOInit.BlendState = TStaticBlendState<CW_NONE>::GetRHI();
+		GraphicsPSOInit.RasterizerState = TStaticRasterizerState<FM_Solid, CM_None>::GetRHI();
+		GraphicsPSOInit.DepthStencilState = TStaticDepthStencilState<true, CF_Always>::GetRHI();
 
-	GraphicsPSOInit.BoundShaderState.VertexDeclarationRHI = GFilterVertexDeclaration.VertexDeclarationRHI;
-	GraphicsPSOInit.BoundShaderState.VertexShaderRHI = GETSAFERHISHADER_VERTEX(*ScreenVertexShader);
-	GraphicsPSOInit.BoundShaderState.PixelShaderRHI = GETSAFERHISHADER_PIXEL(*PixelShader);
-	GraphicsPSOInit.PrimitiveType = PT_TriangleList;
+		GraphicsPSOInit.BoundShaderState.VertexDeclarationRHI = GFilterVertexDeclaration.VertexDeclarationRHI;
+		GraphicsPSOInit.BoundShaderState.VertexShaderRHI = GETSAFERHISHADER_VERTEX(*ScreenVertexShader);
+		GraphicsPSOInit.BoundShaderState.PixelShaderRHI = GETSAFERHISHADER_PIXEL(*PixelShader);
+		GraphicsPSOInit.PrimitiveType = PT_TriangleList;
 
-	SetGraphicsPipelineState(RHICmdList, GraphicsPSOInit);
+		SetGraphicsPipelineState(RHICmdList, GraphicsPSOInit);
 
-	PixelShader->SetParameters(RHICmdList, View, bUseMaxDepth, View.ViewRect.Size());
-	const uint32 DownsampledX = FMath::TruncToInt(View.ViewRect.Min.X * ScaleFactor);
-	const uint32 DownsampledY = FMath::TruncToInt(View.ViewRect.Min.Y * ScaleFactor);
-	const uint32 DownsampledSizeX = FMath::TruncToInt(View.ViewRect.Width() * ScaleFactor);
-	const uint32 DownsampledSizeY = FMath::TruncToInt(View.ViewRect.Height() * ScaleFactor);
-
-	RHICmdList.SetViewport(DownsampledX, DownsampledY, 0.0f, DownsampledX + DownsampledSizeX, DownsampledY + DownsampledSizeY, 1.0f);
-
-	DrawRectangle(
-		RHICmdList,
-		0, 0,
-		DownsampledSizeX, DownsampledSizeY,
-		View.ViewRect.Min.X, View.ViewRect.Min.Y,
-		View.ViewRect.Width(), View.ViewRect.Height(),
-		FIntPoint(DownsampledSizeX, DownsampledSizeY),
-		SceneContext.GetBufferSizeXY(),
-		*ScreenVertexShader,
-		EDRF_UseTriangleOptimization);
-
-	RHICmdList.TransitionResource(EResourceTransitionAccess::EReadable, RenderTarget);
+		PixelShader->SetParameters(RHICmdList, View, bUseMaxDepth, View.ViewRect.Size());
+		const uint32 DownsampledX = FMath::TruncToInt(View.ViewRect.Min.X * ScaleFactor);
+		const uint32 DownsampledY = FMath::TruncToInt(View.ViewRect.Min.Y * ScaleFactor);
+		const uint32 DownsampledSizeX = FMath::TruncToInt(View.ViewRect.Width() * ScaleFactor);
+		const uint32 DownsampledSizeY = FMath::TruncToInt(View.ViewRect.Height() * ScaleFactor);
+
+		RHICmdList.SetViewport(DownsampledX, DownsampledY, 0.0f, DownsampledX + DownsampledSizeX, DownsampledY + DownsampledSizeY, 1.0f);
+
+		DrawRectangle(
+			RHICmdList,
+			0, 0,
+			DownsampledSizeX, DownsampledSizeY,
+			View.ViewRect.Min.X, View.ViewRect.Min.Y,
+			View.ViewRect.Width(), View.ViewRect.Height(),
+			FIntPoint(DownsampledSizeX, DownsampledSizeY),
+			SceneContext.GetBufferSizeXY(),
+			*ScreenVertexShader,
+			EDRF_UseTriangleOptimization);
+	}
+	RHICmdList.EndRenderPass();
 }
 
 /** */
@@ -1688,50 +1825,48 @@ void FDeferredShadingSceneRenderer::CopyStencilToLightingChannelTexture(FRHIComm
 		SceneContext.AllocateLightingChannelTexture(RHICmdList);
 
 		// Set the light attenuation surface as the render target, and the scene depth buffer as the depth-stencil surface.
-		SetRenderTarget(RHICmdList, SceneContext.LightingChannels->GetRenderTargetItem().TargetableTexture, nullptr, ESimpleRenderTargetMode::EExistingColorAndDepth, FExclusiveDepthStencil::DepthNop_StencilNop, true);
-
-		FGraphicsPipelineStateInitializer GraphicsPSOInit;
-		RHICmdList.ApplyCachedRenderTargets(GraphicsPSOInit);
-		GraphicsPSOInit.BlendState = TStaticBlendState<CW_RGBA>::GetRHI();
-		GraphicsPSOInit.RasterizerState = TStaticRasterizerState<FM_Solid, CM_None>::GetRHI();
-		GraphicsPSOInit.DepthStencilState = TStaticDepthStencilState<false, CF_Always>::GetRHI();
-
-		extern TGlobalResource<FFilterVertexDeclaration> GFilterVertexDeclaration;
-		GraphicsPSOInit.BoundShaderState.VertexDeclarationRHI = GFilterVertexDeclaration.VertexDeclarationRHI;
-
-		for (int32 ViewIndex = 0; ViewIndex < Views.Num(); ViewIndex++)
+		FRHIRenderPassInfo RPInfo(SceneContext.LightingChannels->GetRenderTargetItem().TargetableTexture, ERenderTargetActions::Load_Store);
+		TransitionRenderPassTargets(RHICmdList, RPInfo);
+		RHICmdList.BeginRenderPass(RPInfo, TEXT("CopyStencilToLightingChannel"));
 		{
-			const FViewInfo& View = Views[ViewIndex];
-			// Set shaders and texture
-			TShaderMapRef<FScreenVS> ScreenVertexShader(View.ShaderMap);
-			TShaderMapRef<FCopyStencilToLightingChannelsPS> PixelShader(View.ShaderMap);
-
-			GraphicsPSOInit.BoundShaderState.VertexShaderRHI = GETSAFERHISHADER_VERTEX(*ScreenVertexShader);
-			GraphicsPSOInit.BoundShaderState.PixelShaderRHI = GETSAFERHISHADER_PIXEL(*PixelShader);
-			GraphicsPSOInit.PrimitiveType = PT_TriangleList;
-
-			SetGraphicsPipelineState(RHICmdList, GraphicsPSOInit);
-			PixelShader->SetParameters(RHICmdList, View);
-
-			RHICmdList.SetViewport(View.ViewRect.Min.X, View.ViewRect.Min.Y, 0.0f, View.ViewRect.Min.X + View.ViewRect.Width(), View.ViewRect.Min.Y + View.ViewRect.Height(), 1.0f);
-
-			DrawRectangle(
-				RHICmdList,
-				0, 0,
-				View.ViewRect.Width(), View.ViewRect.Height(),
-				View.ViewRect.Min.X, View.ViewRect.Min.Y,
-				View.ViewRect.Width(), View.ViewRect.Height(),
-				FIntPoint(View.ViewRect.Width(), View.ViewRect.Height()),
-				SceneContext.GetBufferSizeXY(),
-				*ScreenVertexShader,
-				EDRF_UseTriangleOptimization);
-		}
+			FGraphicsPipelineStateInitializer GraphicsPSOInit;
+			RHICmdList.ApplyCachedRenderTargets(GraphicsPSOInit);
+			GraphicsPSOInit.BlendState = TStaticBlendState<CW_RGBA>::GetRHI();
+			GraphicsPSOInit.RasterizerState = TStaticRasterizerState<FM_Solid, CM_None>::GetRHI();
+			GraphicsPSOInit.DepthStencilState = TStaticDepthStencilState<false, CF_Always>::GetRHI();
 
-		FResolveParams ResolveParams;
+			extern TGlobalResource<FFilterVertexDeclaration> GFilterVertexDeclaration;
+			GraphicsPSOInit.BoundShaderState.VertexDeclarationRHI = GFilterVertexDeclaration.VertexDeclarationRHI;
 
-		RHICmdList.CopyToResolveTarget(
-			SceneContext.LightingChannels->GetRenderTargetItem().TargetableTexture, 
-			SceneContext.LightingChannels->GetRenderTargetItem().TargetableTexture,
-			ResolveParams);
+			for (int32 ViewIndex = 0; ViewIndex < Views.Num(); ViewIndex++)
+			{
+				const FViewInfo& View = Views[ViewIndex];
+				// Set shaders and texture
+				TShaderMapRef<FScreenVS> ScreenVertexShader(View.ShaderMap);
+				TShaderMapRef<FCopyStencilToLightingChannelsPS> PixelShader(View.ShaderMap);
+
+				GraphicsPSOInit.BoundShaderState.VertexShaderRHI = GETSAFERHISHADER_VERTEX(*ScreenVertexShader);
+				GraphicsPSOInit.BoundShaderState.PixelShaderRHI = GETSAFERHISHADER_PIXEL(*PixelShader);
+				GraphicsPSOInit.PrimitiveType = PT_TriangleList;
+
+				SetGraphicsPipelineState(RHICmdList, GraphicsPSOInit);
+				PixelShader->SetParameters(RHICmdList, View);
+
+				RHICmdList.SetViewport(View.ViewRect.Min.X, View.ViewRect.Min.Y, 0.0f, View.ViewRect.Min.X + View.ViewRect.Width(), View.ViewRect.Min.Y + View.ViewRect.Height(), 1.0f);
+
+				DrawRectangle(
+					RHICmdList,
+					0, 0,
+					View.ViewRect.Width(), View.ViewRect.Height(),
+					View.ViewRect.Min.X, View.ViewRect.Min.Y,
+					View.ViewRect.Width(), View.ViewRect.Height(),
+					FIntPoint(View.ViewRect.Width(), View.ViewRect.Height()),
+					SceneContext.GetBufferSizeXY(),
+					*ScreenVertexShader,
+					EDRF_UseTriangleOptimization);
+			}
+		}
+		RHICmdList.EndRenderPass();
+		RHICmdList.CopyToResolveTarget(SceneContext.LightingChannels->GetRenderTargetItem().TargetableTexture, SceneContext.LightingChannels->GetRenderTargetItem().TargetableTexture, FResolveParams());
 	}
 }
diff --git a/Engine/Source/Runtime/Renderer/Private/DeferredShadingRenderer.h b/Engine/Source/Runtime/Renderer/Private/DeferredShadingRenderer.h
--- a/Engine/Source/Runtime/Renderer/Private/DeferredShadingRenderer.h
+++ b/Engine/Source/Runtime/Renderer/Private/DeferredShadingRenderer.h
@@ -108,7 +108,7 @@ public:
 	* Renders the scene's base pass 
 	* @return true if anything was rendered
 	*/
-	bool RenderBasePass(FRHICommandListImmediate& RHICmdList, FExclusiveDepthStencil::Type BasePassDepthStencilAccess, IPooledRenderTarget* ForwardScreenSpaceShadowMask);
+	bool RenderBasePass(FRHICommandListImmediate& RHICmdList, FExclusiveDepthStencil::Type BasePassDepthStencilAccess, IPooledRenderTarget* ForwardScreenSpaceShadowMask, bool bParallelBasePass, bool bRenderLightmapDensity);
 
 	/** Finishes the view family rendering. */
 	void RenderFinish(FRHICommandListImmediate& RHICmdList);
diff --git a/Engine/Source/Runtime/Renderer/Private/DepthRendering.cpp b/Engine/Source/Runtime/Renderer/Private/DepthRendering.cpp
--- a/Engine/Source/Runtime/Renderer/Private/DepthRendering.cpp
+++ b/Engine/Source/Runtime/Renderer/Private/DepthRendering.cpp
@@ -992,7 +992,9 @@ public:
 
 	void DoTask(ENamedThreads::Type CurrentThread, const FGraphEventRef& MyCompletionGraphEvent)
 	{
+		check(RHICmdList.IsInsideRenderPass());
 		ThisRenderer.RenderPrePassViewDynamic(RHICmdList, View, DrawRenderState);
+		RHICmdList.EndRenderPass();
 		RHICmdList.HandleRTThreadTaskCompletion(MyCompletionGraphEvent);
 	}
 };
@@ -1019,7 +1021,6 @@ public:
 	virtual ~FPrePassParallelCommandListSet()
 	{
 		// Do not copy-paste. this is a very unusual FParallelCommandListSet because it is a prepass and we want to do some work after starting some tasks
-		SetStateOnCommandList(ParentCmdList);
 		Dispatch(true);
 	}
 
@@ -1035,6 +1036,8 @@ bool FDeferredShadingSceneRenderer::RenderPrePassViewParallel(const FViewInfo& V
 {
 	bool bDepthWasCleared = false;
 
+	check(ParentCmdList.IsOutsideRenderPass());
+
 	{
 		FPrePassParallelCommandListSet ParallelCommandListSet(View, this, ParentCmdList,
 			CVarRHICmdPrePassDeferredContexts.GetValueOnRenderThread() > 0, 
@@ -1093,9 +1096,6 @@ bool FDeferredShadingSceneRenderer::RenderPrePassViewParallel(const FViewInfo& V
 			CVarRHICmdFlushRenderThreadTasksPrePass.GetValueOnRenderThread() == 0 && CVarRHICmdFlushRenderThreadTasks.GetValueOnRenderThread() == 0,
 			DrawRenderState);
 
-		FSceneRenderTargets& SceneContext = FSceneRenderTargets::Get(ParentCmdList);
-		SceneContext.BeginRenderingPrePass(ParentCmdList, false);
-
 		// Dynamic
 		FRHICommandList* CmdList = ParallelCommandListSet.NewParallelCommandList();
 
@@ -1155,6 +1155,7 @@ bool FDeferredShadingSceneRenderer::PreRenderPrePass(FRHICommandListImmediate& R
 	SCOPED_GPU_MASK(RHICmdList, FRHIGPUMask::All()); // Required otherwise emulatestereo gets broken.
 
 	RHICmdList.SetCurrentStat(GET_STATID(STAT_CLM_PrePass));
+	// RenderPrePassHMD clears the depth buffer. If this changes we must change RenderPrePass to maintain the correct behavior!
 	bool bDepthWasCleared = RenderPrePassHMD(RHICmdList);
 
 	FSceneRenderTargets& SceneContext = FSceneRenderTargets::Get(RHICmdList);
@@ -1211,6 +1212,9 @@ bool FDeferredShadingSceneRenderer::PreRenderPrePass(FRHICommandListImmediate& R
 				EDRF_UseTriangleOptimization);
 		}
 	}
+	// Need to close the renderpass here since we may call BeginRenderingPrePass later
+	RHICmdList.EndRenderPass();
+
 	return bDepthWasCleared;
 }
 
@@ -1241,6 +1245,8 @@ void FDeferredShadingSceneRenderer::RenderPrePassEditorPrimitives(FRHICommandLis
 
 bool FDeferredShadingSceneRenderer::RenderPrePass(FRHICommandListImmediate& RHICmdList, TFunctionRef<void()> AfterTasksAreStarted)
 {
+	check(RHICmdList.IsOutsideRenderPass());
+
 	SCOPED_NAMED_EVENT(FDeferredShadingSceneRenderer_RenderPrePass, FColor::Emerald);
 	bool bDepthWasCleared = false;
 
@@ -1260,8 +1266,12 @@ bool FDeferredShadingSceneRenderer::RenderPrePass(FRHICommandListImmediate& RHIC
 	{
 		// nothing to be gained by delaying this.
 		AfterTasksAreStarted();
+		// Note: the depth buffer will be cleared under PreRenderPrePass.
 		bDepthWasCleared = PreRenderPrePass(RHICmdList);
 		bDidPrePre = true;
+
+		// PreRenderPrePass will end up clearing the depth buffer so do not clear it again.
+		SceneContext.BeginRenderingPrePass(RHICmdList, false);
 	}
 	else
 	{
@@ -1295,6 +1305,7 @@ bool FDeferredShadingSceneRenderer::RenderPrePass(FRHICommandListImmediate& RHIC
 
 				if (bParallel)
 				{
+					check(RHICmdList.IsOutsideRenderPass());
 					bDepthWasCleared = RenderPrePassViewParallel(View, RHICmdList, DrawRenderState, AfterTasksAreStarted, !bDidPrePre) || bDepthWasCleared;
 					bDirty = true; // assume dirty since we are not going to wait
 					bDidPrePre = true;
@@ -1305,17 +1316,33 @@ bool FDeferredShadingSceneRenderer::RenderPrePass(FRHICommandListImmediate& RHIC
 				}
 			}
 
+			// Parallel rendering has self contained renderpasses so we need a new one for editor primitives.
+			if (bParallel)
+			{
+				SceneContext.BeginRenderingPrePass(RHICmdList, false);
+			}
 			RenderPrePassEditorPrimitives(RHICmdList, View, DrawRenderState, FDepthDrawingPolicyFactory::ContextType(EarlyZPassMode, true));
+			if (bParallel)
+			{
+				RHICmdList.EndRenderPass();
+			}
 		}
 	}
 	if (!bDidPrePre)
 	{
+		// Only parallel rendering with all views marked as not-to-be-rendered will get here.
 		// For some reason we haven't done this yet. Best do it now for consistency with the old code.
 		AfterTasksAreStarted();
 		bDepthWasCleared = PreRenderPrePass(RHICmdList);
 		bDidPrePre = true;
 	}
 
+	if (bParallel)
+	{
+		// In parallel mode there will be no renderpass here. Need to restart.
+		SceneContext.BeginRenderingPrePass(RHICmdList, false);
+	}
+
 	// Dithered transition stencil mask clear, accounting for all active viewports
 	if (bDitheredLODTransitionsUseStencil)
 	{
@@ -1331,6 +1358,7 @@ bool FDeferredShadingSceneRenderer::RenderPrePass(FRHICommandListImmediate& RHIC
 		DrawClearQuad(RHICmdList, false, FLinearColor::Transparent, false, 0, true, 0);
 	}
 
+	// Now we are finally finished.
 	SceneContext.FinishRenderingPrePass(RHICmdList);
 
 	return bDepthWasCleared;
@@ -1357,6 +1385,7 @@ bool FDeferredShadingSceneRenderer::RenderPrePassHMD(FRHICommandListImmediate& R
 		return false;
 	}
 
+	// This is the only place the depth buffer is cleared. If this changes we MUST change RenderPrePass and others to maintain the behavior.
 	FSceneRenderTargets& SceneContext = FSceneRenderTargets::Get(RHICmdList);
 	SceneContext.BeginRenderingPrePass(RHICmdList, true);
 
diff --git a/Engine/Source/Runtime/Renderer/Private/DistanceFieldVisualization.cpp b/Engine/Source/Runtime/Renderer/Private/DistanceFieldVisualization.cpp
--- a/Engine/Source/Runtime/Renderer/Private/DistanceFieldVisualization.cpp
+++ b/Engine/Source/Runtime/Renderer/Private/DistanceFieldVisualization.cpp
@@ -243,7 +243,8 @@ void FDeferredShadingSceneRenderer::RenderMeshDistanceFieldVisualization(FRHICom
 			check(RHICmdList.IsOutsideRenderPass());
 
 			{
-				FSceneRenderTargets::Get(RHICmdList).BeginRenderingSceneColor(RHICmdList, ESimpleRenderTargetMode::EExistingColorAndDepth, FExclusiveDepthStencil::DepthRead_StencilRead);
+				// We must specify StencilWrite or VK will lose the attachment
+				FSceneRenderTargets::Get(RHICmdList).BeginRenderingSceneColor(RHICmdList, ESimpleRenderTargetMode::EExistingColorAndDepth, FExclusiveDepthStencil::DepthRead_StencilWrite);
 
 				for (int32 ViewIndex = 0; ViewIndex < Views.Num(); ViewIndex++)
 				{
diff --git a/Engine/Source/Runtime/Renderer/Private/DynamicPrimitiveDrawing.inl b/Engine/Source/Runtime/Renderer/Private/DynamicPrimitiveDrawing.inl
--- a/Engine/Source/Runtime/Renderer/Private/DynamicPrimitiveDrawing.inl
+++ b/Engine/Source/Runtime/Renderer/Private/DynamicPrimitiveDrawing.inl
@@ -39,6 +39,7 @@ void DrawViewElementsInner(
 	int32 LastIndex
 	)
 {
+	check(RHICmdList.IsInsideRenderPass());
 	// Get the correct element list based on dpg index
 	const TIndirectArray<FMeshBatch>& ViewMeshElementList = ( DPGIndex == SDPG_Foreground ? View.TopViewMeshElements : View.ViewMeshElements );
 	// Draw the view's mesh elements.
@@ -117,7 +118,9 @@ public:
 
 	void DoTask(ENamedThreads::Type CurrentThread, const FGraphEventRef& MyCompletionGraphEvent)
 	{
+		check(RHICmdList.IsInsideRenderPass());
 		DrawViewElementsInner<DrawingPolicyFactoryType>(RHICmdList, View, DrawRenderState, DrawingContext, DPGIndex, bPreFog, FirstIndex, LastIndex);
+		RHICmdList.EndRenderPass();
 		RHICmdList.HandleRTThreadTaskCompletion(MyCompletionGraphEvent);
 	}
 };
diff --git a/Engine/Source/Runtime/Renderer/Private/PostProcess/SceneRenderTargets.cpp b/Engine/Source/Runtime/Renderer/Private/PostProcess/SceneRenderTargets.cpp
--- a/Engine/Source/Runtime/Renderer/Private/PostProcess/SceneRenderTargets.cpp
+++ b/Engine/Source/Runtime/Renderer/Private/PostProcess/SceneRenderTargets.cpp
@@ -613,21 +613,112 @@ void FSceneRenderTargets::Allocate(FRHICommandListImmediate& RHICmdList, const F
 
 void FSceneRenderTargets::BeginRenderingSceneColor(FRHICommandList& RHICmdList, ESimpleRenderTargetMode RenderTargetMode/*=EUninitializedColorExistingDepth*/, FExclusiveDepthStencil DepthStencilAccess, bool bTransitionWritable)
 {
+	check(RHICmdList.IsOutsideRenderPass());
+
 	SCOPED_DRAW_EVENT(RHICmdList, BeginRenderingSceneColor);
 	AllocSceneColor(RHICmdList);
-	SetRenderTarget(RHICmdList, GetSceneColorSurface(), GetSceneDepthSurface(), RenderTargetMode, DepthStencilAccess, bTransitionWritable);
+
+	ERenderTargetLoadAction ColorLoadAction, DepthLoadAction, StencilLoadAction;
+	ERenderTargetStoreAction ColorStoreAction, DepthStoreAction, StencilStoreAction;
+	DecodeRenderTargetMode(RenderTargetMode, ColorLoadAction, ColorStoreAction, DepthLoadAction, DepthStoreAction, StencilLoadAction, StencilStoreAction, DepthStencilAccess);
+
+	FRHIRenderPassInfo RPInfo(GetSceneColorSurface(), MakeRenderTargetActions(ColorLoadAction, ColorStoreAction));
+	RPInfo.DepthStencilRenderTarget.Action = MakeDepthStencilTargetActions(MakeRenderTargetActions(DepthLoadAction, DepthStoreAction), MakeRenderTargetActions(StencilLoadAction, StencilStoreAction));
+	RPInfo.DepthStencilRenderTarget.DepthStencilTarget = GetSceneDepthSurface();
+	RPInfo.DepthStencilRenderTarget.ExclusiveDepthStencil = DepthStencilAccess;
+
+	if (bTransitionWritable)
+	{
+		TransitionRenderPassTargets(RHICmdList, RPInfo);
+	}
+	RHICmdList.BeginRenderPass(RPInfo, TEXT("BeginRenderingSceneColor"));
 } 
 
 void FSceneRenderTargets::FinishRenderingSceneColor(FRHICommandList& RHICmdList)
 {
-	// does nothing until renderpasses are all checked in
+	RHICmdList.EndRenderPass();
 }
 
 void FSceneRenderTargets::BeginRenderingSceneMonoColor(FRHICommandList& RHICmdList, ESimpleRenderTargetMode RenderTargetMode, FExclusiveDepthStencil DepthStencilAccess)
 {
 	SCOPED_DRAW_EVENT(RHICmdList, BeginRenderingSceneMonoColor);
 
-	SetRenderTarget(RHICmdList, GetSceneMonoColorSurface(), GetSceneMonoDepthSurface(), RenderTargetMode, DepthStencilAccess, true);
+	ERenderTargetLoadAction ColorLoadAction, DepthLoadAction, StencilLoadAction;
+	ERenderTargetStoreAction ColorStoreAction, DepthStoreAction, StencilStoreAction;
+	DecodeRenderTargetMode(RenderTargetMode, ColorLoadAction, ColorStoreAction, DepthLoadAction, DepthStoreAction, StencilLoadAction, StencilStoreAction, DepthStencilAccess);
+
+	FRHIRenderPassInfo RPInfo(GetSceneColorSurface(), MakeRenderTargetActions(ColorLoadAction, ColorStoreAction));
+	RPInfo.DepthStencilRenderTarget.Action = MakeDepthStencilTargetActions(MakeRenderTargetActions(DepthLoadAction, DepthStoreAction), MakeRenderTargetActions(StencilLoadAction, StencilStoreAction));
+	RPInfo.DepthStencilRenderTarget.DepthStencilTarget = GetSceneDepthSurface();
+	RPInfo.DepthStencilRenderTarget.ExclusiveDepthStencil = DepthStencilAccess;
+
+	TransitionRenderPassTargets(RHICmdList, RPInfo);
+	RHICmdList.BeginRenderPass(RPInfo, TEXT("BeginRenderingSceneMonoColor"));
+}
+
+int32 FSceneRenderTargets::FillGBufferRenderPassInfo(ERenderTargetLoadAction ColorLoadAction, FRHIRenderPassInfo& OutRenderPassInfo, int32& OutVelocityRTIndex)
+{
+	int32 MRTCount = 0;
+	OutRenderPassInfo.ColorRenderTargets[MRTCount].Action = MakeRenderTargetActions(ColorLoadAction, ERenderTargetStoreAction::EStore);
+	OutRenderPassInfo.ColorRenderTargets[MRTCount].RenderTarget = GetSceneColorSurface();
+	OutRenderPassInfo.ColorRenderTargets[MRTCount].ArraySlice = -1;
+	OutRenderPassInfo.ColorRenderTargets[MRTCount].MipIndex = 0;
+	MRTCount++;
+
+	OutRenderPassInfo.ColorRenderTargets[MRTCount].Action = MakeRenderTargetActions(ColorLoadAction, ERenderTargetStoreAction::EStore);
+	OutRenderPassInfo.ColorRenderTargets[MRTCount].RenderTarget = GBufferA->GetRenderTargetItem().TargetableTexture;
+	OutRenderPassInfo.ColorRenderTargets[MRTCount].ArraySlice = -1;
+	OutRenderPassInfo.ColorRenderTargets[MRTCount].MipIndex = 0;
+	MRTCount++;
+
+	OutRenderPassInfo.ColorRenderTargets[MRTCount].Action = MakeRenderTargetActions(ColorLoadAction, ERenderTargetStoreAction::EStore);
+	OutRenderPassInfo.ColorRenderTargets[MRTCount].RenderTarget = GBufferB->GetRenderTargetItem().TargetableTexture;
+	OutRenderPassInfo.ColorRenderTargets[MRTCount].ArraySlice = -1;
+	OutRenderPassInfo.ColorRenderTargets[MRTCount].MipIndex = 0;
+	MRTCount++;
+
+	OutRenderPassInfo.ColorRenderTargets[MRTCount].Action = MakeRenderTargetActions(ColorLoadAction, ERenderTargetStoreAction::EStore);
+	OutRenderPassInfo.ColorRenderTargets[MRTCount].RenderTarget = GBufferC->GetRenderTargetItem().TargetableTexture;
+	OutRenderPassInfo.ColorRenderTargets[MRTCount].ArraySlice = -1;
+	OutRenderPassInfo.ColorRenderTargets[MRTCount].MipIndex = 0;
+	MRTCount++;
+
+	// The velocity buffer needs to be bound before other optionnal rendertargets (when UseSelectiveBasePassOutputs() is true).
+	// Otherwise there is an issue on some AMD hardware where the target does not get updated. Seems to be related to the velocity buffer format as it works fine with other targets.
+	if (bAllocateVelocityGBuffer)
+	{
+		OutVelocityRTIndex = MRTCount;
+		check(OutVelocityRTIndex == 4); // As defined in BasePassPixelShader.usf
+
+		OutRenderPassInfo.ColorRenderTargets[MRTCount].Action = MakeRenderTargetActions(ColorLoadAction, ERenderTargetStoreAction::EStore);
+		OutRenderPassInfo.ColorRenderTargets[MRTCount].RenderTarget = GBufferVelocity->GetRenderTargetItem().TargetableTexture;
+		OutRenderPassInfo.ColorRenderTargets[MRTCount].ArraySlice = -1;
+		OutRenderPassInfo.ColorRenderTargets[MRTCount].MipIndex = 0;
+		MRTCount++;
+	}
+	else
+	{
+		OutVelocityRTIndex = -1;
+	}
+
+	OutRenderPassInfo.ColorRenderTargets[MRTCount].Action = MakeRenderTargetActions(ColorLoadAction, ERenderTargetStoreAction::EStore);
+	OutRenderPassInfo.ColorRenderTargets[MRTCount].RenderTarget = GBufferD->GetRenderTargetItem().TargetableTexture;
+	OutRenderPassInfo.ColorRenderTargets[MRTCount].ArraySlice = -1;
+	OutRenderPassInfo.ColorRenderTargets[MRTCount].MipIndex = 0;
+	MRTCount++;
+
+	if (bAllowStaticLighting)
+	{
+		check(MRTCount == (bAllocateVelocityGBuffer ? 6 : 5)); // As defined in BasePassPixelShader.usf
+		OutRenderPassInfo.ColorRenderTargets[MRTCount].Action = MakeRenderTargetActions(ColorLoadAction, ERenderTargetStoreAction::EStore);
+		OutRenderPassInfo.ColorRenderTargets[MRTCount].RenderTarget = GBufferE->GetRenderTargetItem().TargetableTexture;
+		OutRenderPassInfo.ColorRenderTargets[MRTCount].ArraySlice = -1;
+		OutRenderPassInfo.ColorRenderTargets[MRTCount].MipIndex = 0;
+		MRTCount++;
+	}
+
+	check(MRTCount <= MaxSimultaneousRenderTargets);
+	return MRTCount;
 }
 
 int32 FSceneRenderTargets::GetGBufferRenderTargets(ERenderTargetLoadAction ColorLoadAction, FRHIRenderTargetView OutRenderTargets[MaxSimultaneousRenderTargets], int32& OutVelocityRTIndex)
@@ -663,7 +754,7 @@ int32 FSceneRenderTargets::GetGBufferRenderTargets(ERenderTargetLoadAction Color
 	return MRTCount;
 }
 
-void FSceneRenderTargets::SetQuadOverdrawUAV(FRHICommandList& RHICmdList, bool bBindQuadOverdrawBuffers, FRHISetRenderTargetsInfo& Info)
+void FSceneRenderTargets::SetQuadOverdrawUAV(FRHICommandList& RHICmdList, bool bBindQuadOverdrawBuffers, bool bClearQuadOverdrawBuffers, FRHIRenderPassInfo& OutInfo)
 {
 	if (bBindQuadOverdrawBuffers && AllowDebugViewPS(DVSM_QuadComplexity, GetFeatureLevelShaderPlatform(CurrentFeatureLevel)))
 	{
@@ -672,13 +763,13 @@ void FSceneRenderTargets::SetQuadOverdrawUAV(FRHICommandList& RHICmdList, bool b
 			QuadOverdrawIndex = IsAnyForwardShadingEnabled(GetFeatureLevelShaderPlatform(CurrentFeatureLevel)) ? 1 : 7; // As defined in QuadOverdraw.usf
 
 			// Increase the rendertarget count in order to control the bound slot of the UAV.
-			check(Info.NumColorRenderTargets <= QuadOverdrawIndex);
-			Info.NumColorRenderTargets = QuadOverdrawIndex;
-			Info.UnorderedAccessView[Info.NumUAVs++] = QuadOverdrawBuffer->GetRenderTargetItem().UAV;
+			check(OutInfo.GetNumColorRenderTargets() <= QuadOverdrawIndex);
+			OutInfo.UAVIndex = QuadOverdrawIndex;
+			OutInfo.UAVs[OutInfo.NumUAVs++] = QuadOverdrawBuffer->GetRenderTargetItem().UAV;
 
-			if (Info.bClearColor)
+			// Clear to default value
+			if (bClearQuadOverdrawBuffers)
 			{
-				// Clear to default value
 				const uint32 ClearValue[4] = { 0, 0, 0, 0 };
 				ClearUAV(RHICmdList, QuadOverdrawBuffer->GetRenderTargetItem(), ClearValue);
 				RHICmdList.TransitionResource(EResourceTransitionAccess::ERWBarrier, EResourceTransitionPipeline::EGfxToGfx, QuadOverdrawBuffer->GetRenderTargetItem().UAV);
@@ -687,15 +778,15 @@ void FSceneRenderTargets::SetQuadOverdrawUAV(FRHICommandList& RHICmdList, bool b
 	}
 }
 
-void FSceneRenderTargets::BeginRenderingGBuffer(FRHICommandList& RHICmdList, ERenderTargetLoadAction ColorLoadAction, ERenderTargetLoadAction DepthLoadAction, FExclusiveDepthStencil::Type DepthStencilAccess, bool bBindQuadOverdrawBuffers, const FLinearColor& ClearColor/*=(0,0,0,1)*/)
+void FSceneRenderTargets::BeginRenderingGBuffer(FRHICommandList& RHICmdList, ERenderTargetLoadAction ColorLoadAction, ERenderTargetLoadAction DepthLoadAction, FExclusiveDepthStencil::Type DepthStencilAccess, bool bBindQuadOverdrawBuffers, bool bClearQuadOverdrawBuffers, const FLinearColor& ClearColor/*=(0,0,0,1)*/)
 {
+	check(RHICmdList.IsOutsideRenderPass());
+
 	SCOPED_DRAW_EVENT(RHICmdList, BeginRenderingGBuffer);
 	check(CurrentFeatureLevel >= ERHIFeatureLevel::SM4);
 	AllocSceneColor(RHICmdList);
 
-	FRHIRenderTargetView RenderTargets[MaxSimultaneousRenderTargets];
 	const ERenderTargetStoreAction DepthStoreAction = (DepthStencilAccess & FExclusiveDepthStencil::DepthWrite) ? ERenderTargetStoreAction::EStore : ERenderTargetStoreAction::ENoAction;
-	FRHIDepthRenderTargetView DepthView(GetSceneDepthSurface(), DepthLoadAction, DepthStoreAction, DepthStencilAccess);
 
 	bool bClearColor = ColorLoadAction == ERenderTargetLoadAction::EClear;
 	bool bClearDepth = DepthLoadAction == ERenderTargetLoadAction::EClear;
@@ -720,64 +811,72 @@ void FSceneRenderTargets::BeginRenderingGBuffer(FRHICommandList& RHICmdList, ERe
 	int32 VelocityRTIndex = -1;
 	int32 MRTCount;
 
+	FRHIRenderPassInfo RPInfo;
+
 	if (IsAnyForwardShadingEnabled(GetFeatureLevelShaderPlatform(CurrentFeatureLevel)))
 	{
 		MRTCount = 1;
-		RenderTargets[0] = FRHIRenderTargetView(SceneColorTex, 0, -1, ColorLoadAction, ERenderTargetStoreAction::EStore);
+		RPInfo.ColorRenderTargets[0].Action = MakeRenderTargetActions(ColorLoadAction, ERenderTargetStoreAction::EStore);
+		RPInfo.ColorRenderTargets[0].RenderTarget = SceneColorTex;
+		RPInfo.ColorRenderTargets[0].ArraySlice = -1;
+		RPInfo.ColorRenderTargets[0].MipIndex = 0;
 	}
 	else
 	{
-		MRTCount = GetGBufferRenderTargets(ColorLoadAction, RenderTargets, VelocityRTIndex);
+		MRTCount = FillGBufferRenderPassInfo(ColorLoadAction, RPInfo, VelocityRTIndex);
 	}
 
 	//make sure our conditions for shader clear fallback are valid.
-	check(RenderTargets[0].Texture == SceneColorTex);
-
-	FRHISetRenderTargetsInfo Info(MRTCount, RenderTargets, DepthView);
+	check(RPInfo.ColorRenderTargets[0].RenderTarget == SceneColorTex);
 
 	if (bClearDepth)
 	{
 		bSceneDepthCleared = true;
 	}
 
-	SetQuadOverdrawUAV(RHICmdList, bBindQuadOverdrawBuffers, Info);
+	SetQuadOverdrawUAV(RHICmdList, bBindQuadOverdrawBuffers, bClearQuadOverdrawBuffers, RPInfo);
+
+	// Stencil always has to be store or certain VK drivers will leave the attachment in an undefined state.
+	RPInfo.DepthStencilRenderTarget.Action = MakeDepthStencilTargetActions(MakeRenderTargetActions(DepthLoadAction, DepthStoreAction), MakeRenderTargetActions(DepthLoadAction, ERenderTargetStoreAction::EStore));
+	RPInfo.DepthStencilRenderTarget.DepthStencilTarget = GetSceneDepthSurface();
+	RPInfo.DepthStencilRenderTarget.ExclusiveDepthStencil = DepthStencilAccess;
 
 	static const auto CVarVirtualTextureLightmaps = IConsoleManager::Get().FindTConsoleVariableDataInt(TEXT("r.VirtualTexturedLightmaps"));
 	if(CVarVirtualTextureLightmaps && CVarVirtualTextureLightmaps->GetValueOnRenderThread() && !bBindQuadOverdrawBuffers )
 	{
 		int32 FeedbackIndex = 7;
+		check(RPInfo.GetNumColorRenderTargets() <= FeedbackIndex);
+		RPInfo.UAVIndex = FeedbackIndex;
+		RPInfo.UAVs[RPInfo.NumUAVs++] = GVirtualTextureFeedback.FeedbackTextureGPU->GetRenderTargetItem().UAV;		
+	}	
 
-		// Increase the rendertarget count in order to control the bound slot of the UAV.
-		check( Info.NumColorRenderTargets <= FeedbackIndex );
-		Info.NumColorRenderTargets = FeedbackIndex;
-		Info.UnorderedAccessView[ Info.NumUAVs++ ] = GVirtualTextureFeedback.FeedbackTextureGPU->GetRenderTargetItem().UAV;
-	}
+	RHICmdList.BeginRenderPass(RPInfo, TEXT("GBuffer"));
 
-	// set the render target
-	RHICmdList.SetRenderTargetsAndClear(Info);
 	if (bShaderClear)
 	{
 		FLinearColor ClearColors[MaxSimultaneousRenderTargets];
 		FTextureRHIParamRef Textures[MaxSimultaneousRenderTargets];
 		ClearColors[0] = ClearColor;
-		Textures[0] = RenderTargets[0].Texture;
+		Textures[0] = RPInfo.ColorRenderTargets[0].RenderTarget;
 		for (int32 i = 1; i < MRTCount; ++i)
 		{
-			ClearColors[i] = RenderTargets[i].Texture->GetClearColor();
-			Textures[i] = RenderTargets[i].Texture;
+			ClearColors[i] = RPInfo.ColorRenderTargets[i].RenderTarget->GetClearColor();
+			Textures[i] = RPInfo.ColorRenderTargets[i].RenderTarget;
 		}
 		//depth/stencil should have been handled by the fast clear.  only color for RT0 can get changed.
 		DrawClearQuadMRT(RHICmdList, true, MRTCount, ClearColors, false, 0, false, 0);
 	}
 
-	//bind any clear data that won't be bound automatically by the preceding SetRenderTargetsAndClear
+	//bind any clear data that won't be bound automatically
 	bool bBindClearColor = !bClearColor && bGBuffersFastCleared;
 	bool bBindClearDepth = !bClearDepth && bSceneDepthCleared;
 	RHICmdList.BindClearMRTValues(bBindClearColor, bBindClearDepth, bBindClearDepth);
 }
 
-void FSceneRenderTargets::FinishRenderingGBuffer(FRHICommandListImmediate& RHICmdList)
+void FSceneRenderTargets::FinishGBufferPassAndResolve(FRHICommandListImmediate& RHICmdList)
 {
+	RHICmdList.EndRenderPass();
+
 	int32 VelocityRTIndex;
 	FRHIRenderTargetView RenderTargets[MaxSimultaneousRenderTargets];
 	int32 NumMRTs = GetGBufferRenderTargets(ERenderTargetLoadAction::ELoad, RenderTargets, VelocityRTIndex);
@@ -1191,12 +1290,16 @@ void FSceneRenderTargets::AdjustGBufferRefCount(FRHICommandList& RHICmdList, int
 
 bool FSceneRenderTargets::BeginRenderingCustomDepth(FRHICommandListImmediate& RHICmdList, bool bPrimitives)
 {
+	check(!RHICmdList.IsInsideRenderPass());
+
 	IPooledRenderTarget* CustomDepthRenderTarget = RequestCustomDepth(RHICmdList, bPrimitives);
 
 	if(CustomDepthRenderTarget)
 	{
 		SCOPED_DRAW_EVENT(RHICmdList, BeginRenderingCustomDepth);
 
+		FRHIRenderPassInfo RPInfo;
+
 		const bool bWritesCustomStencilValues = IsCustomDepthPassWritingStencil();
 		const bool bRequiresStencilColorTarget = (bWritesCustomStencilValues && CurrentFeatureLevel <= ERHIFeatureLevel::ES3_1);
 
@@ -1205,15 +1308,30 @@ bool FSceneRenderTargets::BeginRenderingCustomDepth(FRHICommandListImmediate& RH
 		if (bRequiresStencilColorTarget)
 		{
 			checkSlow(MobileCustomStencil.IsValid());
+
+			RPInfo.ColorRenderTargets[0].Action = ERenderTargetActions::Clear_Store;
+			RPInfo.ColorRenderTargets[0].ArraySlice = -1;
+			RPInfo.ColorRenderTargets[0].MipIndex = 0;
+			RPInfo.ColorRenderTargets[0].RenderTarget = MobileCustomStencil->GetRenderTargetItem().ShaderResourceTexture;
+
 			ColorView =	FRHIRenderTargetView(MobileCustomStencil->GetRenderTargetItem().ShaderResourceTexture, 0, -1, ERenderTargetLoadAction::EClear, ERenderTargetStoreAction::EStore);
 			NumColorTargets = 1;
 		}
 
-		FRHIDepthRenderTargetView DepthView(CustomDepthRenderTarget->GetRenderTargetItem().ShaderResourceTexture, ERenderTargetLoadAction::EClear, ERenderTargetStoreAction::EStore, ERenderTargetLoadAction::EClear, ERenderTargetStoreAction::EStore);
-		FRHISetRenderTargetsInfo Info(NumColorTargets, &ColorView, DepthView);
-		Info.bClearStencil = bWritesCustomStencilValues;
-		check(DepthView.Texture->GetStencilClearValue() == 0);
-		RHICmdList.SetRenderTargetsAndClear(Info);
+		RPInfo.DepthStencilRenderTarget.Action = MakeDepthStencilTargetActions(ERenderTargetActions::Clear_Store, ERenderTargetActions::Load_Store);
+		if (bWritesCustomStencilValues)
+		{
+			RPInfo.DepthStencilRenderTarget.Action = MakeDepthStencilTargetActions(ERenderTargetActions::Clear_Store, ERenderTargetActions::Clear_Store);
+		}
+
+		RPInfo.DepthStencilRenderTarget.DepthStencilTarget = CustomDepthRenderTarget->GetRenderTargetItem().ShaderResourceTexture;
+		RPInfo.DepthStencilRenderTarget.ExclusiveDepthStencil = FExclusiveDepthStencil::DepthWrite_StencilWrite;
+
+		check(RPInfo.DepthStencilRenderTarget.DepthStencilTarget->GetStencilClearValue() == 0);
+
+		// #todo-renderpasses need ResolveRect here
+		//RPInfo.DepthStencilRenderTarget.ResolveTarget = CustomDepth->GetRenderTargetItem().ShaderResourceTexture;
+		RHICmdList.BeginRenderPass(RPInfo, TEXT("BeginRenderCustomDepth"));
 
 		return true;
 	}
@@ -1223,67 +1341,85 @@ bool FSceneRenderTargets::BeginRenderingCustomDepth(FRHICommandListImmediate& RH
 
 void FSceneRenderTargets::FinishRenderingCustomDepth(FRHICommandListImmediate& RHICmdList, const FResolveRect& ResolveRect)
 {
+	check(RHICmdList.IsInsideRenderPass());
+
 	SCOPED_DRAW_EVENT(RHICmdList, FinishRenderingCustomDepth);
 
-	RHICmdList.CopyToResolveTarget(CustomDepth->GetRenderTargetItem().TargetableTexture, CustomDepth->GetRenderTargetItem().ShaderResourceTexture, FResolveParams(ResolveRect));
+	RHICmdList.EndRenderPass();
 	
 	if (CurrentFeatureLevel <= ERHIFeatureLevel::ES3_1 && IsCustomDepthPassWritingStencil() && MobileCustomStencil.IsValid())
 	{
 		RHICmdList.CopyToResolveTarget(MobileCustomStencil->GetRenderTargetItem().TargetableTexture, MobileCustomStencil->GetRenderTargetItem().ShaderResourceTexture, FResolveParams(ResolveRect));
 	}
 
+	RHICmdList.CopyToResolveTarget(CustomDepth->GetRenderTargetItem().TargetableTexture, CustomDepth->GetRenderTargetItem().ShaderResourceTexture, FResolveParams(ResolveRect));
+
 	bCustomDepthIsValid = true;
 }
 
 void FSceneRenderTargets::BeginRenderingPrePass(FRHICommandList& RHICmdList, bool bPerformClear)
 {
+	check(RHICmdList.IsOutsideRenderPass());
+
 	SCOPED_DRAW_EVENT(RHICmdList, BeginRenderingPrePass);
 
-	FTextureRHIRef ColorTarget;
 	FTexture2DRHIRef DepthTarget = GetSceneDepthSurface();
 	
+	// No color target bound for the prepass.
+	FRHIRenderPassInfo RPInfo;
+	RPInfo.DepthStencilRenderTarget.DepthStencilTarget = DepthTarget;
+	RPInfo.DepthStencilRenderTarget.ExclusiveDepthStencil = FExclusiveDepthStencil::DepthWrite_StencilWrite;
+	
 	if (bPerformClear)
 	{				
-		FRHIRenderTargetView ColorView(ColorTarget, 0, -1, ERenderTargetLoadAction::ENoAction, ERenderTargetStoreAction::ENoAction);
-		FRHIDepthRenderTargetView DepthView(DepthTarget, ERenderTargetLoadAction::EClear, ERenderTargetStoreAction::EStore);
-
 		// Clear the depth buffer.
 		// Note, this is a reversed Z depth surface, so 0.0f is the far plane.
-		FRHISetRenderTargetsInfo Info(1, &ColorView, DepthView);
-	
-		RHICmdList.SetRenderTargetsAndClear(Info);
+		RPInfo.DepthStencilRenderTarget.Action = MakeDepthStencilTargetActions(ERenderTargetActions::Clear_Store, ERenderTargetActions::Clear_Store);
 		bSceneDepthCleared = true;	
 	}
 	else
 	{
 		// Set the scene depth surface and a DUMMY buffer as color buffer
 		// (as long as it's the same dimension as the depth buffer),	
-		FRHIRenderTargetView ColorView(ColorTarget, 0, -1, ERenderTargetLoadAction::ENoAction, ERenderTargetStoreAction::ENoAction);
-		FRHIDepthRenderTargetView DepthRTV(DepthTarget, ERenderTargetLoadAction::ELoad, ERenderTargetStoreAction::EStore);
+		RPInfo.DepthStencilRenderTarget.Action = MakeDepthStencilTargetActions(ERenderTargetActions::Load_Store, ERenderTargetActions::Load_Store);
+	}
 		
-		RHICmdList.SetRenderTargets(1, &ColorView, &DepthRTV, 0, NULL);
+	RHICmdList.BeginRenderPass(RPInfo, TEXT("BeginRenderingPrePass"));
 		
+	if (!bPerformClear)
+	{
+		// Needs to be called after we start a renderpass.
 		RHICmdList.BindClearMRTValues(false, true, true);
 	}
 }
 
 void FSceneRenderTargets::FinishRenderingPrePass(FRHICommandListImmediate& RHICmdList)
 {
+	check(RHICmdList.IsInsideRenderPass());
+
 	SCOPED_DRAW_EVENT(RHICmdList, FinishRenderingPrePass);
+	RHICmdList.EndRenderPass();
+
 	GVisualizeTexture.SetCheckPoint(RHICmdList, SceneDepthZ);
 }
 
 void FSceneRenderTargets::BeginRenderingSceneAlphaCopy(FRHICommandListImmediate& RHICmdList)
 {
+	check(!RHICmdList.IsInsideRenderPass());
+
 	SCOPED_DRAW_EVENT(RHICmdList, BeginRenderingSceneAlphaCopy);
 	GVisualizeTexture.SetCheckPoint(RHICmdList, SceneAlphaCopy);
-	SetRenderTarget(RHICmdList, GetSceneAlphaCopySurface(), 0);
+
+	FRHIRenderPassInfo RPInfo(GetSceneAlphaCopySurface(), ERenderTargetActions::Load_Store, SceneAlphaCopy->GetRenderTargetItem().ShaderResourceTexture);
+	RHICmdList.BeginRenderPass(RPInfo, TEXT("BeginRenderingSceneAlphaCopy"));
 }
 
 void FSceneRenderTargets::FinishRenderingSceneAlphaCopy(FRHICommandListImmediate& RHICmdList)
 {
+	check(RHICmdList.IsInsideRenderPass());
+
 	SCOPED_DRAW_EVENT(RHICmdList, FinishRenderingSceneAlphaCopy);
-	RHICmdList.CopyToResolveTarget(GetSceneAlphaCopySurface(), SceneAlphaCopy->GetRenderTargetItem().ShaderResourceTexture, FResolveParams(FResolveRect()));
+	RHICmdList.EndRenderPass();
 	GVisualizeTexture.SetCheckPoint(RHICmdList, SceneAlphaCopy);
 }
 
@@ -1297,27 +1433,44 @@ void FSceneRenderTargets::BeginRenderingLightAttenuation(FRHICommandList& RHICmd
 
 	GVisualizeTexture.SetCheckPoint(RHICmdList, GetLightAttenuation());
 
+	FRHIRenderPassInfo RPInfo;
+
 	// Set the light attenuation surface as the render target, and the scene depth buffer as the depth-stencil surface.
 	if (bClearToWhite)
 	{
-		TransitionSetRenderTargetsHelper(RHICmdList, GetLightAttenuationSurface(), GetSceneDepthSurface(), FExclusiveDepthStencil::DepthRead_StencilWrite);
-
-		FRHIRenderTargetView View = FRHIRenderTargetView(GetLightAttenuationSurface(), ERenderTargetLoadAction::EClear);
-		FRHISetRenderTargetsInfo Info(1, &View, FRHIDepthRenderTargetView(GetSceneDepthSurface(), ERenderTargetLoadAction::ELoad, ERenderTargetStoreAction::ENoAction, ERenderTargetLoadAction::ELoad, ERenderTargetStoreAction::ENoAction, FExclusiveDepthStencil::DepthRead_StencilWrite));
-		RHICmdList.SetRenderTargetsAndClear(Info);
+		RPInfo.ColorRenderTargets[0].RenderTarget = GetLightAttenuationSurface();
+		RPInfo.ColorRenderTargets[0].Action = ERenderTargetActions::Clear_Store;
+		RPInfo.ColorRenderTargets[0].ArraySlice = -1;
+		RPInfo.ColorRenderTargets[0].MipIndex = 0;
+		RPInfo.ColorRenderTargets[0].ResolveTarget = LightAttenuation->GetRenderTargetItem().ShaderResourceTexture;
+		RPInfo.DepthStencilRenderTarget.Action = MakeDepthStencilTargetActions(ERenderTargetActions::Load_DontStore, ERenderTargetActions::Load_DontStore);
+		RPInfo.DepthStencilRenderTarget.DepthStencilTarget = GetSceneDepthSurface();
+		RPInfo.DepthStencilRenderTarget.ExclusiveDepthStencil = FExclusiveDepthStencil::DepthRead_StencilWrite;
+		RPInfo.DepthStencilRenderTarget.ResolveTarget = nullptr;
 	}
 	else
 	{
-		SetRenderTarget(RHICmdList, GetLightAttenuationSurface(), GetSceneDepthSurface(), ESimpleRenderTargetMode::EExistingColorAndDepth, FExclusiveDepthStencil::DepthRead_StencilWrite, true);
+		RPInfo.ColorRenderTargets[0].RenderTarget = GetLightAttenuationSurface();
+		RPInfo.ColorRenderTargets[0].Action = ERenderTargetActions::Load_Store;
+		RPInfo.ColorRenderTargets[0].ArraySlice = -1;
+		RPInfo.ColorRenderTargets[0].MipIndex = 0;
+		RPInfo.ColorRenderTargets[0].ResolveTarget = LightAttenuation->GetRenderTargetItem().ShaderResourceTexture;
+		RPInfo.DepthStencilRenderTarget.Action = MakeDepthStencilTargetActions(ERenderTargetActions::Load_Store, ERenderTargetActions::Load_Store);
+		RPInfo.DepthStencilRenderTarget.DepthStencilTarget = GetSceneDepthSurface();
+		RPInfo.DepthStencilRenderTarget.ExclusiveDepthStencil = FExclusiveDepthStencil::DepthRead_StencilWrite;
+		RPInfo.DepthStencilRenderTarget.ResolveTarget = nullptr;
 	}
+
+	TransitionRenderPassTargets(RHICmdList, RPInfo);
+
+	RHICmdList.BeginRenderPass(RPInfo, TEXT("BeginRenderingLightAttenuation"));
 }
 
 void FSceneRenderTargets::FinishRenderingLightAttenuation(FRHICommandList& RHICmdList)
 {
 	SCOPED_DRAW_EVENT(RHICmdList, FinishRenderingLightAttenuation);
 
-	// Resolve the light attenuation surface.
-	RHICmdList.CopyToResolveTarget(GetLightAttenuationSurface(), LightAttenuation->GetRenderTargetItem().ShaderResourceTexture, FResolveParams(FResolveRect()));
+	RHICmdList.EndRenderPass();
 	
 	GVisualizeTexture.SetCheckPoint(RHICmdList, GetLightAttenuation());
 }
@@ -1352,6 +1505,8 @@ TRefCountPtr<IPooledRenderTarget>& FSceneRenderTargets::GetDownsampledTranslucen
 
 void FSceneRenderTargets::BeginRenderingTranslucency(FRHICommandList& RHICmdList, const FViewInfo& View, const FSceneRenderer& Renderer, bool bFirstTimeThisFrame)
 {
+	check(RHICmdList.IsOutsideRenderPass());
+
 	// Use the scene color buffer.
 	BeginRenderingSceneColor(RHICmdList, ESimpleRenderTargetMode::EExistingColorAndDepth, FExclusiveDepthStencil::DepthRead_StencilWrite);
 
@@ -1390,8 +1545,15 @@ void FSceneRenderTargets::BeginRenderingTranslucency(FRHICommandList& RHICmdList
 	}
 }
 
+void FSceneRenderTargets::FinishRenderingTranslucency(FRHICommandList& RHICmdList)
+{
+	// This is a nop on purpose for consistency's sake
+}
+
 void FSceneRenderTargets::BeginRenderingSeparateTranslucency(FRHICommandList& RHICmdList, const FViewInfo& View, const FSceneRenderer& Renderer, bool bFirstTimeThisFrame)
 {
+	check(RHICmdList.IsOutsideRenderPass());
+
 	bSeparateTranslucencyPass = true;
 
 	SCOPED_DRAW_EVENT(RHICmdList, BeginSeparateTranslucency);
@@ -1409,10 +1571,21 @@ void FSceneRenderTargets::BeginRenderingSeparateTranslucency(FRHICommandList& RH
 	const FTexture2DRHIRef &SeparateTranslucencyDepth = SeparateTranslucencyScale < 1.0f ? (const FTexture2DRHIRef&)GetDownsampledTranslucencyDepth(RHICmdList, SeparateTranslucencyBufferSize)->GetRenderTargetItem().TargetableTexture : GetSceneDepthSurface();
 
 	check((*SeparateTranslucency)->GetRenderTargetItem().TargetableTexture->GetClearColor() == FLinearColor::Black);
-	// clear the render target the first time, re-use afterwards
-	SetRenderTarget(RHICmdList, (*SeparateTranslucency)->GetRenderTargetItem().TargetableTexture, SeparateTranslucencyDepth,
-		bFirstTimeThisFrame ? ESimpleRenderTargetMode::EClearColorExistingDepth : ESimpleRenderTargetMode::EExistingColorAndDepth, FExclusiveDepthStencil::DepthRead_StencilWrite);
+	FRHIRenderPassInfo RPInfo((*SeparateTranslucency)->GetRenderTargetItem().TargetableTexture, ERenderTargetActions::Load_Store);
+	// #todo-renderpass do we ever want to write to depth in these passes?
+	// when do we write to stencil?
+	RPInfo.DepthStencilRenderTarget.Action = MakeDepthStencilTargetActions(ERenderTargetActions::Load_DontStore, ERenderTargetActions::Load_Store);
+	RPInfo.DepthStencilRenderTarget.DepthStencilTarget = SeparateTranslucencyDepth;
+	RPInfo.DepthStencilRenderTarget.ExclusiveDepthStencil = FExclusiveDepthStencil::DepthRead_StencilWrite;
+	RPInfo.DepthStencilRenderTarget.ResolveTarget = nullptr;
+
+	if (bFirstTimeThisFrame)
+	{
+		// Clear the color target the first pass through and re-use after
+		RPInfo.ColorRenderTargets[0].Action = ERenderTargetActions::Clear_Store;
+	}
 
+	RHICmdList.BeginRenderPass(RPInfo, TEXT("BeginRenderingSeparateTranslucency"));
 
 	if (!bFirstTimeThisFrame)
 	{
@@ -1453,6 +1626,8 @@ void FSceneRenderTargets::FinishRenderingSeparateTranslucency(FRHICommandList& R
 {
 	SCOPED_DRAW_EVENT(RHICmdList, FinishSeparateTranslucency);
 
+	check(RHICmdList.IsOutsideRenderPass());
+
 	TRefCountPtr<IPooledRenderTarget>* SeparateTranslucency;
 	TRefCountPtr<IPooledRenderTarget>* SeparateTranslucencyDepth;
 	if (bSnapshot)
@@ -1494,12 +1669,20 @@ FResolveRect FSceneRenderTargets::GetDefaultRect(const FResolveRect& Rect, uint3
 
 void FSceneRenderTargets::ResolveDepthTexture(FRHICommandList& RHICmdList, const FTexture2DRHIRef& SourceTexture, const FTexture2DRHIRef& DestTexture, const FResolveParams& ResolveParams)
 {
+	check(!RHICmdList.IsInsideRenderPass());
+
 	FResolveRect ResolveRect = ResolveParams.Rect;
 
 	RHICmdList.TransitionResource(EResourceTransitionAccess::EReadable, SourceTexture);
 
-	SetRenderTargets(RHICmdList, 0, nullptr, DestTexture, ESimpleRenderTargetMode::EClearColorExistingDepth, FExclusiveDepthStencil::DepthWrite_StencilWrite, true);
+	FRHIRenderPassInfo RPInfo;
+	RPInfo.DepthStencilRenderTarget.Action = EDepthStencilTargetActions::LoadDepthStencil_StoreDepthStencil;
+	RPInfo.DepthStencilRenderTarget.DepthStencilTarget = DestTexture;
+	RPInfo.DepthStencilRenderTarget.ExclusiveDepthStencil = FExclusiveDepthStencil::DepthWrite_StencilWrite;
 
+	TransitionRenderPassTargets(RHICmdList, RPInfo);
+	RHICmdList.BeginRenderPass(RPInfo, TEXT("ResolveDepthTexture"));
+	{
 	FGraphicsPipelineStateInitializer GraphicsPSOInit;
 	RHICmdList.ApplyCachedRenderTargets(GraphicsPSOInit);
 
@@ -1606,9 +1789,13 @@ void FSceneRenderTargets::ResolveDepthTexture(FRHICommandList& RHICmdList, const
 	RHIUnlockVertexBuffer(VertexBufferRHI);
 	RHICmdList.SetStreamSource(0, VertexBufferRHI, 0);
 	RHICmdList.DrawPrimitive(0, 2, 1);
+}
+	RHICmdList.EndRenderPass();
 }
 void FSceneRenderTargets::ResolveSceneDepthTexture(FRHICommandList& RHICmdList, const FResolveRect& ResolveRect)
 {
+	check(RHICmdList.IsOutsideRenderPass());
+
 	SCOPED_DRAW_EVENT(RHICmdList, ResolveSceneDepthTexture);
 
 	if (ResolveRect.IsValid())
@@ -1936,7 +2123,7 @@ void FSceneRenderTargets::AllocateCommonDepthTargets(FRHICommandList& RHICmdList
 		{
 			const uint32 OldElementSize = SceneDepthZ->ComputeMemorySize();
 			bHMDAllocatedDepthTarget = true;
-
+		
 			// If SRT and texture are different (MSAA), only modify the resolve render target, to avoid creating a swapchain of MSAA textures
 			if (SceneDepthZ->GetRenderTargetItem().ShaderResourceTexture == SceneDepthZ->GetRenderTargetItem().TargetableTexture)
 			{
@@ -2219,8 +2406,13 @@ void FSceneRenderTargets::ClearTranslucentVolumeLighting(FRHICommandListImmediat
 template<int32 NumRenderTargets>
 void FSceneRenderTargets::ClearVolumeTextures(FRHICommandList& RHICmdList, ERHIFeatureLevel::Type FeatureLevel, FTextureRHIParamRef* RenderTargets, const FLinearColor* ClearColors)
 {
-	SetRenderTargets(RHICmdList, NumRenderTargets, RenderTargets, FTextureRHIRef(), 0, NULL, true);
+	check(!RHICmdList.IsInsideRenderPass());
+
+	FRHIRenderPassInfo RPInfo(NumRenderTargets, RenderTargets, ERenderTargetActions::DontLoad_Store);
+	TransitionRenderPassTargets(RHICmdList, RPInfo);
 
+	RHICmdList.BeginRenderPass(RPInfo, TEXT("ClearVolumeTextures"));
+	{
 	FGraphicsPipelineStateInitializer GraphicsPSOInit;
 	RHICmdList.ApplyCachedRenderTargets(GraphicsPSOInit);
 	GraphicsPSOInit.RasterizerState = TStaticRasterizerState<FM_Solid, CM_None>::GetRHI();
@@ -2251,6 +2443,8 @@ void FSceneRenderTargets::ClearVolumeTextures(FRHICommandList& RHICmdList, ERHIF
 	PixelShader->SetColors(RHICmdList, ClearColors, NumRenderTargets);
 
 	RasterizeToVolumeTexture(RHICmdList, VolumeBounds);
+	}
+	RHICmdList.EndRenderPass();
 
 	RHICmdList.TransitionResources(EResourceTransitionAccess::EReadable, (FTextureRHIParamRef*)RenderTargets, NumRenderTargets);
 }
diff --git a/Engine/Source/Runtime/Renderer/Private/PostProcess/SceneRenderTargets.h b/Engine/Source/Runtime/Renderer/Private/PostProcess/SceneRenderTargets.h
--- a/Engine/Source/Runtime/Renderer/Private/PostProcess/SceneRenderTargets.h
+++ b/Engine/Source/Runtime/Renderer/Private/PostProcess/SceneRenderTargets.h
@@ -201,6 +201,8 @@ public:
 	 */
 	void Allocate(FRHICommandListImmediate& RHICmdList, const FSceneRenderer* SceneRenderer);
 
+	void AllocSceneColor(FRHICommandList& RHICmdList);
+
 	/**
 	 *
 	 */
@@ -209,8 +211,9 @@ public:
 	void SetSeparateTranslucencyBufferSize(bool bAnyViewWantsDownsampledSeparateTranslucency);
 
 	void SetQuadOverdrawUAV(FRHICommandList& RHICmdList, bool bBindQuadOverdrawBuffers, FRHISetRenderTargetsInfo& Info);
-	void BeginRenderingGBuffer(FRHICommandList& RHICmdList, ERenderTargetLoadAction ColorLoadAction, ERenderTargetLoadAction DepthLoadAction, FExclusiveDepthStencil::Type DepthStencilAccess, bool bBindQuadOverdrawBuffers, const FLinearColor& ClearColor = FLinearColor(0, 0, 0, 1));
-	void FinishRenderingGBuffer(FRHICommandListImmediate& RHICmdList);
+	void SetQuadOverdrawUAV(FRHICommandList& RHICmdList, bool bBindQuadOverdrawBuffers, bool bClearQuadOverdrawBuffers, FRHIRenderPassInfo& Info);
+	void BeginRenderingGBuffer(FRHICommandList& RHICmdList, ERenderTargetLoadAction ColorLoadAction, ERenderTargetLoadAction DepthLoadAction, FExclusiveDepthStencil::Type DepthStencilAccess, bool bBindQuadOverdrawBuffers, bool bClearQuadOverdrawBuffers = false, const FLinearColor& ClearColor = FLinearColor(0, 0, 0, 1));
+	void FinishGBufferPassAndResolve(FRHICommandListImmediate& RHICmdList);
 
 	/**
 	 * Sets the scene color target and restores its contents if necessary
@@ -230,6 +233,8 @@ public:
 
 	/** Begin rendering translucency in the scene color. */
 	void BeginRenderingTranslucency(FRHICommandList& RHICmdList, const class FViewInfo& View, const FSceneRenderer& Renderer, bool bFirstTimeThisFrame = true);
+	void FinishRenderingTranslucency(FRHICommandList& RHICmdList);
+
 	/** Begin rendering translucency in a separate (offscreen) buffer. This can be any translucency pass. */
 	void BeginRenderingSeparateTranslucency(FRHICommandList& RHICmdList, const FViewInfo& View, const FSceneRenderer& Renderer, bool bFirstTimeThisFrame);
 	void FinishRenderingSeparateTranslucency(FRHICommandList& RHICmdList, const FViewInfo& View);
@@ -627,8 +632,6 @@ private:
 	/** Determine the appropriate render target dimensions. */
 	FIntPoint ComputeDesiredSize(const FSceneViewFamily& ViewFamily);
 
-	void AllocSceneColor(FRHICommandList& RHICmdList);
-
 	/** Allocates the mobile multi-view scene color texture array render target. */
 	void AllocMobileMultiViewSceneColor(FRHICommandList& RHICmdList, const int32 ScaleFactor);
 
@@ -662,6 +665,9 @@ private:
 	/** Gets all GBuffers to use.  Returns the number actually used. */
 	int32 GetGBufferRenderTargets(ERenderTargetLoadAction ColorLoadAction, FRHIRenderTargetView OutRenderTargets[MaxSimultaneousRenderTargets], int32& OutVelocityRTIndex);
 
+	/** Fills the given FRenderPassInfo with the current GBuffer */
+	int32 FillGBufferRenderPassInfo(ERenderTargetLoadAction ColorLoadAction, FRHIRenderPassInfo& OutRenderPassInfo, int32& OutVelocityRTIndex);
+
 	ESceneColorFormatType GetSceneColorFormatType() const
 	{
 		if (CurrentShadingPath == EShadingPath::Mobile)
diff --git a/Engine/Source/Runtime/Renderer/Private/SceneRendering.cpp b/Engine/Source/Runtime/Renderer/Private/SceneRendering.cpp
--- a/Engine/Source/Runtime/Renderer/Private/SceneRendering.cpp
+++ b/Engine/Source/Runtime/Renderer/Private/SceneRendering.cpp
@@ -590,6 +590,11 @@ void FParallelCommandListSet::Dispatch(bool bHighPriority)
 	check(CommandLists.Num() == Events.Num());
 	check(CommandLists.Num() == NumAlloc);
 
+	// We should not be submitting work off a parent command list if it's still in the middle of a renderpass.
+	// This is a bit weird since we will (likely) end up opening one in the parallel translate case but until we have
+	// a cleaner way for the RHI to specify parallel passes this is what we've got.
+	check(ParentCmdList.IsOutsideRenderPass());
+
 	// We don't support the mask changing while the parallel commandlists are being processed.
 	ensure(GPUMask == ParentCmdList.GetGPUMask());
 
@@ -636,7 +641,9 @@ void FParallelCommandListSet::Dispatch(bool bHighPriority)
 		check(GRHISupportsParallelRHIExecute);
 		NumAlloc -= CommandLists.Num();
 		ParentCmdList.QueueParallelAsyncCommandListSubmit(&Events[0], bHighPriority, &CommandLists[0], &NumDrawsIfKnown[0], CommandLists.Num(), (MinDrawsPerCommandList * 4) / 3, bSpewBalance);
+		// #todo-renderpasses PS4 breaks if this isn't here. Why?
 		SetStateOnCommandList(ParentCmdList);
+		ParentCmdList.EndRenderPass();
 	}
 	else
 	{
@@ -3507,115 +3514,119 @@ void FSceneRenderer::ResolveSceneColor(FRHICommandList& RHICmdList)
 		FTexture2DRHIRef FMaskTexture = GDynamicRHI->RHIGetFMaskTexture(SceneContext.GetSceneColorSurface());
 
 		// Custom shader based color resolve for HDR color to emulate mobile.
-		SetRenderTarget(RHICmdList, SceneContext.GetSceneColorTexture(), FTextureRHIParamRef());
-		
-		for (int32 ViewIndex = 0; ViewIndex < Views.Num(); ViewIndex++)
+		FRHIRenderPassInfo RPInfo(SceneContext.GetSceneColorTexture(), ERenderTargetActions::Load_Store);
+		RHICmdList.BeginRenderPass(RPInfo, TEXT("ResolveColor"));
 		{
-			const FViewInfo& View = Views[ViewIndex];
-
-			FGraphicsPipelineStateInitializer GraphicsPSOInit;
-			RHICmdList.ApplyCachedRenderTargets(GraphicsPSOInit);
+			for (int32 ViewIndex = 0; ViewIndex < Views.Num(); ViewIndex++)
+			{
+				const FViewInfo& View = Views[ViewIndex];
 
-			GraphicsPSOInit.BlendState = TStaticBlendState<>::GetRHI();
-			GraphicsPSOInit.RasterizerState = TStaticRasterizerState<>::GetRHI();
-			GraphicsPSOInit.DepthStencilState = TStaticDepthStencilState<false, CF_Always>::GetRHI();
-			RHICmdList.SetStreamSource(0, GResolveDummyVertexBuffer.VertexBufferRHI, 0);
+				FGraphicsPipelineStateInitializer GraphicsPSOInit;
+				RHICmdList.ApplyCachedRenderTargets(GraphicsPSOInit);
 
-			// Resolve views individually
-			// In the case of adaptive resolution, the view family will be much larger than the views individually
-			RHICmdList.SetScissorRect(true, View.ViewRect.Min.X, View.ViewRect.Min.Y, View.ViewRect.Max.X, View.ViewRect.Max.Y);
+				GraphicsPSOInit.BlendState = TStaticBlendState<>::GetRHI();
+				GraphicsPSOInit.RasterizerState = TStaticRasterizerState<>::GetRHI();
+				GraphicsPSOInit.DepthStencilState = TStaticDepthStencilState<false, CF_Always>::GetRHI();
+				RHICmdList.SetStreamSource(0, GResolveDummyVertexBuffer.VertexBufferRHI, 0);
 
-			int32 ResolveWidth = CVarWideCustomResolve.GetValueOnRenderThread();
+				// Resolve views individually
+				// In the case of adaptive resolution, the view family will be much larger than the views individually
+				RHICmdList.SetScissorRect(true, View.ViewRect.Min.X, View.ViewRect.Min.Y, View.ViewRect.Max.X, View.ViewRect.Max.Y);
 
-			if (CurrentNumSamples <= 1)
-			{
-				ResolveWidth = 0;
-			}
+				int32 ResolveWidth = CVarWideCustomResolve.GetValueOnRenderThread();
 
-			if (ResolveWidth != 0)
-			{
-				ResolveFilterWide(RHICmdList, GraphicsPSOInit, SceneContext.GetCurrentFeatureLevel(), CurrentSceneColor->GetRenderTargetItem().TargetableTexture, FMaskTexture, FIntPoint(0, 0), CurrentNumSamples, ResolveWidth);
-			}
-			else
-			{
-				auto ShaderMap = GetGlobalShaderMap(SceneContext.GetCurrentFeatureLevel());
-				TShaderMapRef<FHdrCustomResolveVS> VertexShader(ShaderMap);
-				GraphicsPSOInit.BoundShaderState.VertexDeclarationRHI = GetVertexDeclarationFVector4();
-				GraphicsPSOInit.BoundShaderState.VertexShaderRHI = GETSAFERHISHADER_VERTEX(*VertexShader);
-				GraphicsPSOInit.PrimitiveType = PT_TriangleList;
+				if (CurrentNumSamples <= 1)
+				{
+					ResolveWidth = 0;
+				}
 
-				if (FMaskTexture.IsValid())
+				if (ResolveWidth != 0)
 				{
-					if (CurrentNumSamples == 2)
-					{
-						TShaderMapRef<FHdrCustomResolveFMask2xPS> PixelShader(ShaderMap);
-						GraphicsPSOInit.BoundShaderState.PixelShaderRHI = GETSAFERHISHADER_PIXEL(*PixelShader);
+					ResolveFilterWide(RHICmdList, GraphicsPSOInit, SceneContext.GetCurrentFeatureLevel(), CurrentSceneColor->GetRenderTargetItem().TargetableTexture, FMaskTexture, FIntPoint(0, 0), CurrentNumSamples, ResolveWidth);
+				}
+				else
+				{
+					auto ShaderMap = GetGlobalShaderMap(SceneContext.GetCurrentFeatureLevel());
+					TShaderMapRef<FHdrCustomResolveVS> VertexShader(ShaderMap);
+					GraphicsPSOInit.BoundShaderState.VertexDeclarationRHI = GetVertexDeclarationFVector4();
+					GraphicsPSOInit.BoundShaderState.VertexShaderRHI = GETSAFERHISHADER_VERTEX(*VertexShader);
+					GraphicsPSOInit.PrimitiveType = PT_TriangleList;
 
-						SetGraphicsPipelineState(RHICmdList, GraphicsPSOInit);
-						PixelShader->SetParameters(RHICmdList, CurrentSceneColor->GetRenderTargetItem().TargetableTexture, FMaskTexture);
-					}
-					else if (CurrentNumSamples == 4)
+					if (FMaskTexture.IsValid())
 					{
-						TShaderMapRef<FHdrCustomResolveFMask4xPS> PixelShader(ShaderMap);
-						GraphicsPSOInit.BoundShaderState.PixelShaderRHI = GETSAFERHISHADER_PIXEL(*PixelShader);
+						if (CurrentNumSamples == 2)
+						{
+							TShaderMapRef<FHdrCustomResolveFMask2xPS> PixelShader(ShaderMap);
+							GraphicsPSOInit.BoundShaderState.PixelShaderRHI = GETSAFERHISHADER_PIXEL(*PixelShader);
 
-						SetGraphicsPipelineState(RHICmdList, GraphicsPSOInit);
-						PixelShader->SetParameters(RHICmdList, CurrentSceneColor->GetRenderTargetItem().TargetableTexture, FMaskTexture);
-					}
-					else if (CurrentNumSamples == 8)
-					{
-						TShaderMapRef<FHdrCustomResolveFMask8xPS> PixelShader(ShaderMap);
-						GraphicsPSOInit.BoundShaderState.PixelShaderRHI = GETSAFERHISHADER_PIXEL(*PixelShader);
+							SetGraphicsPipelineState(RHICmdList, GraphicsPSOInit);
+							PixelShader->SetParameters(RHICmdList, CurrentSceneColor->GetRenderTargetItem().TargetableTexture, FMaskTexture);
+						}
+						else if (CurrentNumSamples == 4)
+						{
+							TShaderMapRef<FHdrCustomResolveFMask4xPS> PixelShader(ShaderMap);
+							GraphicsPSOInit.BoundShaderState.PixelShaderRHI = GETSAFERHISHADER_PIXEL(*PixelShader);
+
+							SetGraphicsPipelineState(RHICmdList, GraphicsPSOInit);
+							PixelShader->SetParameters(RHICmdList, CurrentSceneColor->GetRenderTargetItem().TargetableTexture, FMaskTexture);
+						}
+						else if (CurrentNumSamples == 8)
+						{
+							TShaderMapRef<FHdrCustomResolveFMask8xPS> PixelShader(ShaderMap);
+							GraphicsPSOInit.BoundShaderState.PixelShaderRHI = GETSAFERHISHADER_PIXEL(*PixelShader);
 
-						SetGraphicsPipelineState(RHICmdList, GraphicsPSOInit);
-						PixelShader->SetParameters(RHICmdList, CurrentSceneColor->GetRenderTargetItem().TargetableTexture, FMaskTexture);
+							SetGraphicsPipelineState(RHICmdList, GraphicsPSOInit);
+							PixelShader->SetParameters(RHICmdList, CurrentSceneColor->GetRenderTargetItem().TargetableTexture, FMaskTexture);
+						}
+						else
+						{
+							// Everything other than 2,4,8 samples is not implemented.
+							check(0);
+							break;
+						}
 					}
 					else
 					{
-						// Everything other than 2,4,8 samples is not implemented.
-						check(0);
-						break;
-					}
-				}
-				else
-				{
-				if (CurrentNumSamples == 2)
-				{
-					TShaderMapRef<FHdrCustomResolve2xPS> PixelShader(ShaderMap);
-					GraphicsPSOInit.BoundShaderState.PixelShaderRHI = GETSAFERHISHADER_PIXEL(*PixelShader);
+						if (CurrentNumSamples == 2)
+						{
+							TShaderMapRef<FHdrCustomResolve2xPS> PixelShader(ShaderMap);
+							GraphicsPSOInit.BoundShaderState.PixelShaderRHI = GETSAFERHISHADER_PIXEL(*PixelShader);
 
-					SetGraphicsPipelineState(RHICmdList, GraphicsPSOInit);
-					PixelShader->SetParameters(RHICmdList, CurrentSceneColor->GetRenderTargetItem().TargetableTexture);
-				}
-				else if (CurrentNumSamples == 4)
-				{
-					TShaderMapRef<FHdrCustomResolve4xPS> PixelShader(ShaderMap);
-					GraphicsPSOInit.BoundShaderState.PixelShaderRHI = GETSAFERHISHADER_PIXEL(*PixelShader);
+							SetGraphicsPipelineState(RHICmdList, GraphicsPSOInit);
+							PixelShader->SetParameters(RHICmdList, CurrentSceneColor->GetRenderTargetItem().TargetableTexture);
+						}
+						else if (CurrentNumSamples == 4)
+						{
+							TShaderMapRef<FHdrCustomResolve4xPS> PixelShader(ShaderMap);
+							GraphicsPSOInit.BoundShaderState.PixelShaderRHI = GETSAFERHISHADER_PIXEL(*PixelShader);
 
-					SetGraphicsPipelineState(RHICmdList, GraphicsPSOInit);
-					PixelShader->SetParameters(RHICmdList, CurrentSceneColor->GetRenderTargetItem().TargetableTexture);
-				}
-				else if (CurrentNumSamples == 8)
-				{
-					TShaderMapRef<FHdrCustomResolve8xPS> PixelShader(ShaderMap);
-					GraphicsPSOInit.BoundShaderState.PixelShaderRHI = GETSAFERHISHADER_PIXEL(*PixelShader);
+							SetGraphicsPipelineState(RHICmdList, GraphicsPSOInit);
+							PixelShader->SetParameters(RHICmdList, CurrentSceneColor->GetRenderTargetItem().TargetableTexture);
+						}
+						else if (CurrentNumSamples == 8)
+						{
+							TShaderMapRef<FHdrCustomResolve8xPS> PixelShader(ShaderMap);
+							GraphicsPSOInit.BoundShaderState.PixelShaderRHI = GETSAFERHISHADER_PIXEL(*PixelShader);
 
-					SetGraphicsPipelineState(RHICmdList, GraphicsPSOInit);
-					PixelShader->SetParameters(RHICmdList, CurrentSceneColor->GetRenderTargetItem().TargetableTexture);
-				}
-				else
-				{
-					// Everything other than 2,4,8 samples is not implemented.
-					check(0);
-						break;
+							SetGraphicsPipelineState(RHICmdList, GraphicsPSOInit);
+							PixelShader->SetParameters(RHICmdList, CurrentSceneColor->GetRenderTargetItem().TargetableTexture);
+						}
+						else
+						{
+							// Everything other than 2,4,8 samples is not implemented.
+							check(0);
+							break;
+						}
+					}
+
+					RHICmdList.DrawPrimitive(0, 1, 1);
 				}
 			}
 
-				RHICmdList.DrawPrimitive(0, 1, 1);
-			}
+			RHICmdList.SetScissorRect(false, 0, 0, 0, 0);
 		}
 
-		RHICmdList.SetScissorRect(false, 0, 0, 0, 0);
+		RHICmdList.EndRenderPass();
 	}
 }
 
diff --git a/Engine/Source/Runtime/Renderer/Private/ShadowDepthRendering.cpp b/Engine/Source/Runtime/Renderer/Private/ShadowDepthRendering.cpp
--- a/Engine/Source/Runtime/Renderer/Private/ShadowDepthRendering.cpp
+++ b/Engine/Source/Runtime/Renderer/Private/ShadowDepthRendering.cpp
@@ -1168,6 +1168,8 @@ static void CheckShadowDepthMaterials(const FMaterialRenderProxy* InRenderProxy,
 
 void FProjectedShadowInfo::ClearDepth(FRHICommandList& RHICmdList, class FSceneRenderer* SceneRenderer, int32 NumColorTextures, FTextureRHIParamRef* ColorTextures, FTextureRHIParamRef DepthTexture, bool bPerformClear)
 {
+	check(RHICmdList.IsInsideRenderPass());
+
 	uint32 ViewportMinX = X;
 	uint32 ViewportMinY = Y;
 	float ViewportMinZ = 0.0f;
@@ -1371,6 +1373,8 @@ public:
 	{
 		SCOPE_CYCLE_COUNTER(STAT_WholeSceneStaticShadowDepthsTime);
 
+		check(RHICmdList.IsInsideRenderPass());
+
 		if (bReflective)
 		{
 			// reflective shadow map
@@ -1381,6 +1385,8 @@ public:
 			// normal shadow map
 			DrawShadowMeshElements<false>(RHICmdList, View, DrawRenderState, ThisShadow);
 		}
+
+		RHICmdList.EndRenderPass();
 		RHICmdList.HandleRTThreadTaskCompletion(MyCompletionGraphEvent);
 	}
 };
@@ -1419,7 +1425,9 @@ public:
 
 	void DoTask(ENamedThreads::Type CurrentThread, const FGraphEventRef& MyCompletionGraphEvent)
 	{
+		check(RHICmdList.IsInsideRenderPass());
 		ThisShadow.RenderDepthDynamic(RHICmdList, SceneRenderer, &View, DrawRenderState);
+		RHICmdList.EndRenderPass();
 		RHICmdList.HandleRTThreadTaskCompletion(MyCompletionGraphEvent);
 	}
 };
@@ -1498,7 +1506,7 @@ DECLARE_CYCLE_STAT(TEXT("Shadow"), STAT_CLP_Shadow, STATGROUP_ParallelCommandLis
 class FShadowParallelCommandListSet : public FParallelCommandListSet
 {
 	FProjectedShadowInfo& ProjectedShadowInfo;
-	FSetShadowRenderTargetFunction SetShadowRenderTargets;
+	FBeginShadowRenderPassFunction BeginShadowRenderPass;
 	EShadowDepthRenderMode RenderMode;
 
 public:
@@ -1510,14 +1518,13 @@ public:
 		bool bInCreateSceneContext, 
 		const FDrawingPolicyRenderState& InDrawRenderState,
 		FProjectedShadowInfo& InProjectedShadowInfo, 
-		FSetShadowRenderTargetFunction InSetShadowRenderTargets, 
+		FBeginShadowRenderPassFunction InBeginShadowRenderPass,
 		EShadowDepthRenderMode RenderModeIn )
 		: FParallelCommandListSet(GET_STATID(STAT_CLP_Shadow), InView, InSceneRenderer, InParentCmdList, bInParallelExecute, bInCreateSceneContext, InDrawRenderState)
 		, ProjectedShadowInfo(InProjectedShadowInfo)
-		, SetShadowRenderTargets(InSetShadowRenderTargets)
+		, BeginShadowRenderPass(InBeginShadowRenderPass)
 		, RenderMode(RenderModeIn)
 	{
-		SetStateOnCommandList(ParentCmdList);
 		// It is difficult to get accurate counts from the last frame. Cascades interfere with each other.
 		bBalanceCommandsWithLastFrame = false;
 		bBalanceCommands = false;
@@ -1531,7 +1538,7 @@ public:
 	virtual void SetStateOnCommandList(FRHICommandList& CmdList) override
 	{
 		FParallelCommandListSet::SetStateOnCommandList(CmdList);
-		SetShadowRenderTargets(CmdList, false);
+		BeginShadowRenderPass(CmdList, false);
 		ProjectedShadowInfo.SetStateForDepth(CmdList,RenderMode, DrawRenderState);
 	}
 };
@@ -1634,7 +1641,7 @@ public:
 
 IMPLEMENT_SHADER_TYPE(,FCopyShadowMaps2DPS,TEXT("/Engine/Private/CopyShadowMaps.usf"),TEXT("Copy2DDepthPS"),SF_Pixel);
 
-void FProjectedShadowInfo::CopyCachedShadowMap(FRHICommandList& RHICmdList, const FDrawingPolicyRenderState& DrawRenderState, FSceneRenderer* SceneRenderer, const FViewInfo& View, FSetShadowRenderTargetFunction SetShadowRenderTargets)
+void FProjectedShadowInfo::CopyCachedShadowMap(FRHICommandList& RHICmdList, const FDrawingPolicyRenderState& DrawRenderState, FSceneRenderer* SceneRenderer, const FViewInfo& View)
 {
 	check(CacheMode == SDCM_MovablePrimitivesOnly);
 	const FCachedShadowMapData& CachedShadowMapData = SceneRenderer->Scene->CachedShadowMaps.FindChecked(GetLightSceneInfo().Id);
@@ -1747,7 +1754,11 @@ void FProjectedShadowInfo::CopyCachedShadowMap(FRHICommandList& RHICmdList, cons
 	}
 }
 
-void FProjectedShadowInfo::RenderDepthInner(FRHICommandList& RHICmdList, FSceneRenderer* SceneRenderer, const FViewInfo* FoundView, FSetShadowRenderTargetFunction SetShadowRenderTargets, EShadowDepthRenderMode RenderMode)
+/*
+	The current structure of the code means we can get here with an active renderpass on RHICmdList (in serial mode) OR without one (in parallel dispatch mode)
+	So there's some tracking going on. In the future I'd like to refactor
+*/
+void FProjectedShadowInfo::RenderDepthInner(FRHICommandList& RHICmdList, FSceneRenderer* SceneRenderer, const FViewInfo* FoundView, FBeginShadowRenderPassFunction BeginShadowRenderPass, EShadowDepthRenderMode RenderMode, bool bDoParallelDispatch)
 {
 	TUniformBufferRef<FSceneTexturesUniformParameters> DeferredPassUniformBuffer;
 	TUniformBufferRef<FMobileSceneTextureUniformParameters> MobilePassUniformBuffer;
@@ -1776,8 +1787,19 @@ void FProjectedShadowInfo::RenderDepthInner(FRHICommandList& RHICmdList, FSceneR
 	
 	if (CacheMode == SDCM_MovablePrimitivesOnly)
 	{
+		// In parallel mode we will not have a renderpass active at this point.
+		if (bDoParallelDispatch)
+		{
+			BeginShadowRenderPass(RHICmdList, false);
+		}
+
 		// Copy in depths of static primitives before we render movable primitives
-		CopyCachedShadowMap(RHICmdList, DrawRenderState, SceneRenderer, *FoundView, SetShadowRenderTargets);
+		CopyCachedShadowMap(RHICmdList, DrawRenderState, SceneRenderer, *FoundView);
+
+		if (bDoParallelDispatch)
+		{
+			RHICmdList.EndRenderPass();
+		}
 	}
 
 	FShadowDepthDrawingPolicyContext StackPolicyContext(this);
@@ -1785,12 +1807,11 @@ void FProjectedShadowInfo::RenderDepthInner(FRHICommandList& RHICmdList, FSceneR
 
 	bool bIsWholeSceneDirectionalShadow = IsWholeSceneDirectionalShadow();
 
-	if (RHICmdList.IsImmediate() &&  // translucent shadows are draw on the render thread, using a recursive cmdlist (which is not immediate)
-		GRHICommandList.UseParallelAlgorithms() && CVarParallelShadows.GetValueOnRenderThread() &&
-		(bIsWholeSceneDirectionalShadow || CVarParallelShadowsNonWholeScene.GetValueOnRenderThread() > 0)
-		)
+	if (bDoParallelDispatch)
 	{
 		check(IsInRenderingThread());
+		// Parallel encoding requires its own renderpass.
+		check(RHICmdList.IsOutsideRenderPass());
 
 		// parallel version
 		bool bFlush = CVarRHICmdFlushRenderThreadTasksShadowPass.GetValueOnRenderThread() > 0 || CVarRHICmdFlushRenderThreadTasks.GetValueOnRenderThread() > 0; 
@@ -1803,10 +1824,12 @@ void FProjectedShadowInfo::RenderDepthInner(FRHICommandList& RHICmdList, FSceneR
 			check(IsInRenderingThread() && FMemStack::Get().GetNumMarks() == 1); // we do not want this popped before the end of the scene and it better be the scene allocator
 			PolicyContext = new (FMemStack::Get()) FShadowDepthDrawingPolicyContext(this);
 		}
+
+		// Dispatch commands
 		{
 			check(RHICmdList.IsImmediate());
 			FRHICommandListImmediate& Immed = static_cast<FRHICommandListImmediate&>(RHICmdList);
-			FShadowParallelCommandListSet ParallelCommandListSet(*FoundView, SceneRenderer, Immed, CVarRHICmdShadowDeferredContexts.GetValueOnRenderThread() > 0, !bFlush, DrawRenderState, *this, SetShadowRenderTargets, RenderMode);
+			FShadowParallelCommandListSet ParallelCommandListSet(*FoundView, SceneRenderer, Immed, CVarRHICmdShadowDeferredContexts.GetValueOnRenderThread() > 0, !bFlush, DrawRenderState, *this, BeginShadowRenderPass, RenderMode);
 
 			// Draw the subject's static elements using static draw lists
 			if (bIsWholeSceneDirectionalShadow && RenderMode != ShadowDepthRenderMode_EmissiveOnly && RenderMode != ShadowDepthRenderMode_GIBlockingVolumes)
@@ -1843,9 +1866,15 @@ void FProjectedShadowInfo::RenderDepthInner(FRHICommandList& RHICmdList, FSceneR
 				ParallelCommandListSet.AddParallelCommandList(CmdList, AnyThreadCompletionEvent, DynamicSubjectMeshElements.Num());
 			}
 		}
+
+		// Renderpass must be closed once we get here.
+		check(RHICmdList.IsOutsideRenderPass());
 	}
 	else
 	{
+		// We must have already opened the renderpass by the time we get here.
+		check(RHICmdList.IsInsideRenderPass());
+
 		// single threaded version
 		SetStateForDepth(RHICmdList, RenderMode, DrawRenderState);
 
@@ -1881,6 +1910,9 @@ void FProjectedShadowInfo::RenderDepthInner(FRHICommandList& RHICmdList, FSceneR
 			}
 		}
 		RenderDepthDynamic(RHICmdList, SceneRenderer, FoundView, DrawRenderState);
+
+		// Renderpass must still be open when we reach here
+		check(RHICmdList.IsInsideRenderPass());
 	}
 }
 
@@ -1941,7 +1973,7 @@ FViewInfo* FProjectedShadowInfo::FindViewForShadow(FSceneRenderer* SceneRenderer
 	return FoundView;
 }
 
-void FProjectedShadowInfo::RenderDepth(FRHICommandList& RHICmdList, FSceneRenderer* SceneRenderer, FSetShadowRenderTargetFunction SetShadowRenderTargets, EShadowDepthRenderMode RenderMode)
+void FProjectedShadowInfo::RenderDepth(FRHICommandList& RHICmdList, FSceneRenderer* SceneRenderer, FBeginShadowRenderPassFunction BeginShadowRenderPass, EShadowDepthRenderMode RenderMode, bool bDoParallelDispatch)
 {
 	// Select the correct set of arrays for the current render mode
 	TArray<FShadowStaticMeshElement,SceneRenderingAllocator>* PtrCurrentMeshElements = nullptr;
@@ -1988,7 +2020,7 @@ void FProjectedShadowInfo::RenderDepth(FRHICommandList& RHICmdList, FSceneRender
 		return;
 	}
 
-	RenderDepthInner(RHICmdList, SceneRenderer, ShadowDepthView, SetShadowRenderTargets, RenderMode);
+	RenderDepthInner(RHICmdList, SceneRenderer, ShadowDepthView, BeginShadowRenderPass, RenderMode, bDoParallelDispatch);
 }
 
 void FProjectedShadowInfo::SetupShadowDepthView(FRHICommandListImmediate& RHICmdList, FSceneRenderer* SceneRenderer)
@@ -2062,6 +2094,9 @@ void FSceneRenderer::RenderShadowDepthMapAtlases(FRHICommandListImmediate& RHICm
 {
 	FSceneRenderTargets& SceneContext = FSceneRenderTargets::Get(RHICmdList);
 
+	bool bCanUseParallelDispatch =	RHICmdList.IsImmediate() &&  // translucent shadows are draw on the render thread, using a recursive cmdlist (which is not immediate)
+									GRHICommandList.UseParallelAlgorithms() && CVarParallelShadows.GetValueOnRenderThread();
+
 	for (int32 AtlasIndex = 0; AtlasIndex < SortedShadowsForShadowDepthPass.ShadowMapAtlases.Num(); AtlasIndex++)
 	{
 		const FSortedShadowMapAtlas& ShadowMapAtlas = SortedShadowsForShadowDepthPass.ShadowMapAtlases[AtlasIndex];
@@ -2072,25 +2107,22 @@ void FSceneRenderer::RenderShadowDepthMapAtlases(FRHICommandListImmediate& RHICm
 
 		SCOPED_DRAW_EVENTF(RHICmdList, EventShadowDepths, TEXT("Atlas%u %ux%u"), AtlasIndex, AtlasSize.X, AtlasSize.Y);
 
-		auto SetShadowRenderTargets = [this, &RenderTarget, &SceneContext](FRHICommandList& InRHICmdList, bool bPerformClear)
+		auto BeginShadowRenderPass = [this, &RenderTarget, &SceneContext](FRHICommandList& InRHICmdList, bool bPerformClear)
 		{
-			FRHISetRenderTargetsInfo Info(0, nullptr, FRHIDepthRenderTargetView(RenderTarget.TargetableTexture, 
-				bPerformClear ? ERenderTargetLoadAction::EClear : ERenderTargetLoadAction::ELoad, 
-				ERenderTargetStoreAction::EStore, 
-				ERenderTargetLoadAction::ELoad, 
-				ERenderTargetStoreAction::EStore));
+			check(RenderTarget.TargetableTexture->GetDepthClearValue() == 1.0f);	
 
-			check(Info.DepthStencilRenderTarget.Texture->GetDepthClearValue() == 1.0f);	
-			Info.ColorRenderTarget[0].StoreAction = ERenderTargetStoreAction::ENoAction;
+			ERenderTargetLoadAction DepthLoadAction = bPerformClear ? ERenderTargetLoadAction::EClear : ERenderTargetLoadAction::ELoad;
+
+			FRHIRenderPassInfo RPInfo(RenderTarget.TargetableTexture, MakeDepthStencilTargetActions(MakeRenderTargetActions(DepthLoadAction, ERenderTargetStoreAction::EStore), ERenderTargetActions::Load_Store), nullptr, FExclusiveDepthStencil::DepthWrite_StencilWrite);
 
 			if (!GSupportsDepthRenderTargetWithoutColorRenderTarget)
 			{
-				Info.NumColorRenderTargets = 1;
-				Info.ColorRenderTarget[0].Texture = SceneContext.GetOptionalShadowDepthColorSurface(InRHICmdList, Info.DepthStencilRenderTarget.Texture->GetTexture2D()->GetSizeX(), Info.DepthStencilRenderTarget.Texture->GetTexture2D()->GetSizeY());
-				InRHICmdList.TransitionResource(EResourceTransitionAccess::EWritable, Info.ColorRenderTarget[0].Texture);
+				RPInfo.ColorRenderTargets[0].Action = ERenderTargetActions::DontLoad_DontStore;
+				RPInfo.ColorRenderTargets[0].RenderTarget = SceneContext.GetOptionalShadowDepthColorSurface(InRHICmdList, RPInfo.DepthStencilRenderTarget.DepthStencilTarget->GetTexture2D()->GetSizeX(), RPInfo.DepthStencilRenderTarget.DepthStencilTarget->GetTexture2D()->GetSizeY());
+				InRHICmdList.TransitionResource(EResourceTransitionAccess::EWritable, RPInfo.ColorRenderTargets[0].RenderTarget);
 			}
-			InRHICmdList.TransitionResource(EResourceTransitionAccess::EWritable, Info.DepthStencilRenderTarget.Texture);
-			InRHICmdList.SetRenderTargetsAndClear(Info);
+			InRHICmdList.TransitionResource(EResourceTransitionAccess::EWritable, RPInfo.DepthStencilRenderTarget.DepthStencilTarget);
+			InRHICmdList.BeginRenderPass(RPInfo, TEXT("ShadowMapAtlases"));
 
 			if (!bPerformClear)
 			{
@@ -2098,9 +2130,25 @@ void FSceneRenderer::RenderShadowDepthMapAtlases(FRHICommandListImmediate& RHICm
 			}
 		};
 
+		TArray<FProjectedShadowInfo*, SceneRenderingAllocator> ParallelShadowPasses;
+		TArray<FProjectedShadowInfo*, SceneRenderingAllocator> SerialShadowPasses;
+
+		// Gather our passes here to minimize switching renderpasses
+		for (int32 ShadowIndex = 0; ShadowIndex < ShadowMapAtlas.Shadows.Num(); ShadowIndex++)
 		{
-			SCOPED_DRAW_EVENT(RHICmdList, SetShadowRTsAndClear);
-			SetShadowRenderTargets(RHICmdList, true);	
+			FProjectedShadowInfo* ProjectedShadowInfo = ShadowMapAtlas.Shadows[ShadowIndex];
+
+			const bool bDoParallelDispatch = bCanUseParallelDispatch &&
+				(ProjectedShadowInfo->IsWholeSceneDirectionalShadow() || CVarParallelShadowsNonWholeScene.GetValueOnRenderThread());
+
+			if (bDoParallelDispatch)
+			{
+				ParallelShadowPasses.Add(ProjectedShadowInfo);
+			}
+			else
+			{
+				SerialShadowPasses.Add(ProjectedShadowInfo);
+			}
 		}
 
 		FLightSceneProxy* CurrentLightForDrawEvent = NULL;
@@ -2109,29 +2157,80 @@ void FSceneRenderer::RenderShadowDepthMapAtlases(FRHICommandListImmediate& RHICm
 		TDrawEvent<FRHICommandList> LightEvent;
 #endif
 
-		for (int32 ShadowIndex = 0; ShadowIndex < ShadowMapAtlas.Shadows.Num(); ShadowIndex++)
+		if (ParallelShadowPasses.Num() > 0)
 		{
-			FProjectedShadowInfo* ProjectedShadowInfo = ShadowMapAtlas.Shadows[ShadowIndex];
+			{
+				// Clear before going wide.
+				SCOPED_DRAW_EVENT(RHICmdList, SetShadowRTsAndClear);
+				BeginShadowRenderPass(RHICmdList, true); 
+				RHICmdList.EndRenderPass();
+			}
 
-			if (!CurrentLightForDrawEvent || ProjectedShadowInfo->GetLightSceneInfo().Proxy != CurrentLightForDrawEvent)
+			for (int32 ShadowIndex = 0; ShadowIndex < ParallelShadowPasses.Num(); ShadowIndex++)
 			{
-				if (CurrentLightForDrawEvent)
+				FProjectedShadowInfo* ProjectedShadowInfo = ParallelShadowPasses[ShadowIndex];
+
+				if (!CurrentLightForDrawEvent || ProjectedShadowInfo->GetLightSceneInfo().Proxy != CurrentLightForDrawEvent)
 				{
-					STOP_DRAW_EVENT(LightEvent);
+					if (CurrentLightForDrawEvent)
+					{
+						STOP_DRAW_EVENT(LightEvent);
+					}
+
+					CurrentLightForDrawEvent = ProjectedShadowInfo->GetLightSceneInfo().Proxy;
+					FString LightNameWithLevel;
+					GetLightNameForDrawEvent(CurrentLightForDrawEvent, LightNameWithLevel);
+
+					BEGIN_DRAW_EVENTF(
+						RHICmdList,
+						LightNameEvent,
+						LightEvent,
+						*LightNameWithLevel);
 				}
 
-				CurrentLightForDrawEvent = ProjectedShadowInfo->GetLightSceneInfo().Proxy;
-				FString LightNameWithLevel;
-				GetLightNameForDrawEvent(CurrentLightForDrawEvent, LightNameWithLevel);
-				
-				BEGIN_DRAW_EVENTF(
-					RHICmdList, 
-					LightNameEvent, 
-					LightEvent, 
-					*LightNameWithLevel);
+				ProjectedShadowInfo->RenderDepth(RHICmdList, this, BeginShadowRenderPass, ShadowDepthRenderMode_Normal, true);
 			}
+		}
 
-			ProjectedShadowInfo->RenderDepth(RHICmdList, this, SetShadowRenderTargets, ShadowDepthRenderMode_Normal);
+		if (CurrentLightForDrawEvent)
+		{
+			STOP_DRAW_EVENT(LightEvent);
+		}
+
+		CurrentLightForDrawEvent = nullptr;
+
+		if (SerialShadowPasses.Num() > 0)
+		{
+			{
+				SCOPED_DRAW_EVENT(RHICmdList, SetShadowRTsAndClear);
+				BeginShadowRenderPass(RHICmdList, true);
+			}
+
+			for (int32 ShadowIndex = 0; ShadowIndex < SerialShadowPasses.Num(); ShadowIndex++)
+			{
+				FProjectedShadowInfo* ProjectedShadowInfo = SerialShadowPasses[ShadowIndex];
+
+				if (!CurrentLightForDrawEvent || ProjectedShadowInfo->GetLightSceneInfo().Proxy != CurrentLightForDrawEvent)
+				{
+					if (CurrentLightForDrawEvent)
+					{
+						STOP_DRAW_EVENT(LightEvent);
+					}
+
+					CurrentLightForDrawEvent = ProjectedShadowInfo->GetLightSceneInfo().Proxy;
+					FString LightNameWithLevel;
+					GetLightNameForDrawEvent(CurrentLightForDrawEvent, LightNameWithLevel);
+
+					BEGIN_DRAW_EVENTF(
+						RHICmdList,
+						LightNameEvent,
+						LightEvent,
+						*LightNameWithLevel);
+				}
+
+				ProjectedShadowInfo->RenderDepth(RHICmdList, this, BeginShadowRenderPass, ShadowDepthRenderMode_Normal, false);
+			}
+			RHICmdList.EndRenderPass();
 		}
 
 		if (CurrentLightForDrawEvent)
@@ -2146,6 +2245,8 @@ void FSceneRenderer::RenderShadowDepthMapAtlases(FRHICommandListImmediate& RHICm
 
 void FSceneRenderer::RenderShadowDepthMaps(FRHICommandListImmediate& RHICmdList)
 {
+	check(RHICmdList.IsOutsideRenderPass());
+
 	SCOPED_NAMED_EVENT(FSceneRenderer_RenderShadowDepthMaps, FColor::Emerald);
 	FSceneRenderTargets& SceneContext = FSceneRenderTargets::Get(RHICmdList);
 
@@ -2163,33 +2264,37 @@ void FSceneRenderer::RenderShadowDepthMaps(FRHICommandListImmediate& RHICmdList)
 		check(ShadowMap.Shadows.Num() == 1);
 		FProjectedShadowInfo* ProjectedShadowInfo = ShadowMap.Shadows[0];
 
+		const bool bDoParallelDispatch =	RHICmdList.IsImmediate() &&  // translucent shadows are draw on the render thread, using a recursive cmdlist (which is not immediate)
+											GRHICommandList.UseParallelAlgorithms() && CVarParallelShadows.GetValueOnRenderThread() &&
+											(ProjectedShadowInfo->IsWholeSceneDirectionalShadow() || CVarParallelShadowsNonWholeScene.GetValueOnRenderThread());
+
 		GVisualizeTexture.SetCheckPoint(RHICmdList, ShadowMap.RenderTargets.DepthTarget.GetReference());
 
 		FString LightNameWithLevel;
 		GetLightNameForDrawEvent(ProjectedShadowInfo->GetLightSceneInfo().Proxy, LightNameWithLevel);
 		SCOPED_DRAW_EVENTF(RHICmdList, EventShadowDepths, TEXT("Cubemap %s %u^2"), *LightNameWithLevel, TargetSize.X, TargetSize.Y);
 
-		auto SetShadowRenderTargets = [this, &RenderTarget, &SceneContext](FRHICommandList& InRHICmdList, bool bPerformClear)
+		auto BeginShadowRenderPass = [this, &RenderTarget, &SceneContext](FRHICommandList& InRHICmdList, bool bPerformClear)
 		{
-			FRHISetRenderTargetsInfo Info(0, nullptr, FRHIDepthRenderTargetView(RenderTarget.TargetableTexture, 
-				bPerformClear ? ERenderTargetLoadAction::EClear : ERenderTargetLoadAction::ELoad, 
-				ERenderTargetStoreAction::EStore, 
-				ERenderTargetLoadAction::ELoad, 
-				ERenderTargetStoreAction::EStore));
+			FTextureRHIParamRef DepthTarget = RenderTarget.TargetableTexture;
+			ERenderTargetLoadAction DepthLoadAction = bPerformClear ? ERenderTargetLoadAction::EClear : ERenderTargetLoadAction::ELoad;
 
-			check(Info.DepthStencilRenderTarget.Texture->GetDepthClearValue() == 1.0f);	
-			Info.ColorRenderTarget[0].StoreAction = ERenderTargetStoreAction::ENoAction;
+			check(DepthTarget->GetDepthClearValue() == 1.0f);
+			FRHIRenderPassInfo RPInfo(DepthTarget, MakeDepthStencilTargetActions(MakeRenderTargetActions(DepthLoadAction, ERenderTargetStoreAction::EStore), ERenderTargetActions::Load_Store), nullptr, FExclusiveDepthStencil::DepthWrite_StencilWrite);
 
 			if (!GSupportsDepthRenderTargetWithoutColorRenderTarget)
 			{
-				Info.NumColorRenderTargets = 1;
-				Info.ColorRenderTarget[0].Texture = SceneContext.GetOptionalShadowDepthColorSurface(InRHICmdList, Info.DepthStencilRenderTarget.Texture->GetTexture2D()->GetSizeX(), Info.DepthStencilRenderTarget.Texture->GetTexture2D()->GetSizeY());
-				InRHICmdList.TransitionResource(EResourceTransitionAccess::EWritable, Info.ColorRenderTarget[0].Texture);
+				RPInfo.ColorRenderTargets[0].Action = ERenderTargetActions::DontLoad_DontStore;
+				RPInfo.ColorRenderTargets[0].ArraySlice = -1;
+				RPInfo.ColorRenderTargets[0].MipIndex = 0;
+				RPInfo.ColorRenderTargets[0].RenderTarget = SceneContext.GetOptionalShadowDepthColorSurface(InRHICmdList, DepthTarget->GetTexture2D()->GetSizeX(), DepthTarget->GetTexture2D()->GetSizeY());
+				
+				InRHICmdList.TransitionResource(EResourceTransitionAccess::EWritable, RPInfo.ColorRenderTargets[0].RenderTarget);
 			}
-			InRHICmdList.TransitionResource(EResourceTransitionAccess::EWritable, Info.DepthStencilRenderTarget.Texture);
-			InRHICmdList.SetRenderTargetsAndClear(Info);
+			InRHICmdList.TransitionResource(EResourceTransitionAccess::EWritable, DepthTarget);
+			InRHICmdList.BeginRenderPass(RPInfo, TEXT("ShadowDepthCubeMaps"));
 		};
-			
+		
 		{
 			bool bDoClear = true;
 
@@ -2201,10 +2306,21 @@ void FSceneRenderer::RenderShadowDepthMaps(FRHICommandListImmediate& RHICmdList)
 			}
 
 			SCOPED_CONDITIONAL_DRAW_EVENT(RHICmdList, Clear, bDoClear);
-			SetShadowRenderTargets(RHICmdList, bDoClear);	
+			BeginShadowRenderPass(RHICmdList, bDoClear);	
+		}
+
+		if (bDoParallelDispatch)
+		{
+			// In parallel mode this first pass will just be the clear.
+			RHICmdList.EndRenderPass();
 		}
 
-		ProjectedShadowInfo->RenderDepth(RHICmdList, this, SetShadowRenderTargets, ShadowDepthRenderMode_Normal);
+		ProjectedShadowInfo->RenderDepth(RHICmdList, this, BeginShadowRenderPass, ShadowDepthRenderMode_Normal, bDoParallelDispatch);
+
+		if (!bDoParallelDispatch)
+		{
+			RHICmdList.EndRenderPass();
+		}
 
 		RHICmdList.TransitionResource(EResourceTransitionAccess::EReadable, RenderTarget.TargetableTexture);
 	}
@@ -2223,19 +2339,37 @@ void FSceneRenderer::RenderShadowDepthMaps(FRHICommandListImmediate& RHICmdList)
 
 			if (!ProjectedShadowInfo->bDepthsCached)
 			{
-				auto SetShadowRenderTargets = [this, ProjectedShadowInfo](FRHICommandList& InRHICmdList, bool bPerformClear)
+				const bool bDoParallelDispatch = RHICmdList.IsImmediate() &&  // translucent shadows are draw on the render thread, using a recursive cmdlist (which is not immediate)
+					GRHICommandList.UseParallelAlgorithms() && CVarParallelShadows.GetValueOnRenderThread() &&
+					(ProjectedShadowInfo->IsWholeSceneDirectionalShadow() || CVarParallelShadowsNonWholeScene.GetValueOnRenderThread());
+
+				auto BeginShadowRenderPass = [this, ProjectedShadowInfo](FRHICommandList& InRHICmdList, bool bPerformClear)
 				{
 					FTextureRHIParamRef PreShadowCacheDepthZ = Scene->PreShadowCacheDepthZ->GetRenderTargetItem().TargetableTexture.GetReference();
 					InRHICmdList.TransitionResources(EResourceTransitionAccess::EWritable, &PreShadowCacheDepthZ, 1);
 
+					FRHIRenderPassInfo RPInfo(PreShadowCacheDepthZ, EDepthStencilTargetActions::LoadDepthStencil_StoreDepthStencil, nullptr, FExclusiveDepthStencil::DepthWrite_StencilWrite);
+
 					// Must preserve existing contents as the clear will be scissored
-					SetRenderTarget(InRHICmdList, FTextureRHIRef(), PreShadowCacheDepthZ, ESimpleRenderTargetMode::EExistingColorAndDepth);
+					InRHICmdList.BeginRenderPass(RPInfo, TEXT("ShadowDepthMaps"));
 					ProjectedShadowInfo->ClearDepth(InRHICmdList, this, 0, nullptr, PreShadowCacheDepthZ, bPerformClear);
 				};
 
-				SetShadowRenderTargets(RHICmdList, true);
+				BeginShadowRenderPass(RHICmdList, true);
+
+				if (bDoParallelDispatch)
+				{
+					// In parallel mode the first pass is just the clear.
+					RHICmdList.EndRenderPass();
+				}
+
+				ProjectedShadowInfo->RenderDepth(RHICmdList, this, BeginShadowRenderPass, ShadowDepthRenderMode_Normal, bDoParallelDispatch);
+
+				if (!bDoParallelDispatch)
+				{
+					RHICmdList.EndRenderPass();
+				}
 
-				ProjectedShadowInfo->RenderDepth(RHICmdList, this, SetShadowRenderTargets, ShadowDepthRenderMode_Normal);
 				ProjectedShadowInfo->bDepthsCached = true;
 			}
 		}
@@ -2258,13 +2392,18 @@ void FSceneRenderer::RenderShadowDepthMaps(FRHICommandListImmediate& RHICmdList)
 			ColorTarget0.TargetableTexture,
 			ColorTarget1.TargetableTexture
 		};
-		SetRenderTargets(RHICmdList, ARRAY_COUNT(RenderTargetArray), RenderTargetArray, FTextureRHIParamRef(), 0, NULL, true);
 
-		for (int32 ShadowIndex = 0; ShadowIndex < ShadowMapAtlas.Shadows.Num(); ShadowIndex++)
+		FRHIRenderPassInfo RPInfo(ARRAY_COUNT(RenderTargetArray), RenderTargetArray, ERenderTargetActions::Load_Store);
+		TransitionRenderPassTargets(RHICmdList, RPInfo);
+		RHICmdList.BeginRenderPass(RPInfo, TEXT("RenderTranslucencyDepths"));
 		{
-			FProjectedShadowInfo* ProjectedShadowInfo = ShadowMapAtlas.Shadows[ShadowIndex];
-			ProjectedShadowInfo->RenderTranslucencyDepths(RHICmdList, this);
+			for (int32 ShadowIndex = 0; ShadowIndex < ShadowMapAtlas.Shadows.Num(); ShadowIndex++)
+			{
+				FProjectedShadowInfo* ProjectedShadowInfo = ShadowMapAtlas.Shadows[ShadowIndex];
+				ProjectedShadowInfo->RenderTranslucencyDepths(RHICmdList, this);
+			}
 		}
+		RHICmdList.EndRenderPass();
 
 		RHICmdList.TransitionResource(EResourceTransitionAccess::EReadable, ColorTarget0.TargetableTexture);
 		RHICmdList.TransitionResource(EResourceTransitionAccess::EReadable, ColorTarget1.TargetableTexture);
@@ -2291,6 +2430,8 @@ void FSceneRenderer::RenderShadowDepthMaps(FRHICommandListImmediate& RHICmdList)
 
 	for (int32 AtlasIndex = 0; AtlasIndex < SortedShadowsForShadowDepthPass.RSMAtlases.Num(); AtlasIndex++)
 	{
+		check(RHICmdList.IsOutsideRenderPass());
+
 		const FSortedShadowMapAtlas& ShadowMapAtlas = SortedShadowsForShadowDepthPass.RSMAtlases[AtlasIndex];
 		FSceneRenderTargetItem ColorTarget0 = ShadowMapAtlas.RenderTargets.ColorTargets[0]->GetRenderTargetItem();
 		FSceneRenderTargetItem ColorTarget1 = ShadowMapAtlas.RenderTargets.ColorTargets[1]->GetRenderTargetItem();
@@ -2303,10 +2444,14 @@ void FSceneRenderer::RenderShadowDepthMaps(FRHICommandListImmediate& RHICmdList)
 		{
 			FProjectedShadowInfo* ProjectedShadowInfo = ShadowMapAtlas.Shadows[ShadowIndex];
 
+			const bool bDoParallelDispatch =	RHICmdList.IsImmediate() &&  // translucent shadows are draw on the render thread, using a recursive cmdlist (which is not immediate)
+												GRHICommandList.UseParallelAlgorithms() && CVarParallelShadows.GetValueOnRenderThread() &&
+												(ProjectedShadowInfo->IsWholeSceneDirectionalShadow() || CVarParallelShadowsNonWholeScene.GetValueOnRenderThread());
+
 			FSceneViewState* ViewState = (FSceneViewState*)ProjectedShadowInfo->DependentView->State;
 			FLightPropagationVolume* LightPropagationVolume = ViewState->GetLightPropagationVolume(FeatureLevel);
 
-			auto SetShadowRenderTargets = [this, LightPropagationVolume, ProjectedShadowInfo, &ColorTarget0, &ColorTarget1, &DepthTarget](FRHICommandList& InRHICmdList, bool bPerformClear)
+			auto BeginShadowRenderPass = [this, LightPropagationVolume, ProjectedShadowInfo, &ColorTarget0, &ColorTarget1, &DepthTarget](FRHICommandList& InRHICmdList, bool bPerformClear)
 			{
 				FTextureRHIParamRef RenderTargets[2];
 				RenderTargets[0] = ColorTarget0.TargetableTexture;
@@ -2318,25 +2463,48 @@ void FSceneRenderer::RenderShadowDepthMaps(FRHICommandListImmediate& RHICmdList)
 				Uavs[1] = LightPropagationVolume->GetGvListHeadBufferUav();
 				Uavs[2] = LightPropagationVolume->GetVplListBufferUav();
 				Uavs[3] = LightPropagationVolume->GetVplListHeadBufferUav();
+
+				FRHIRenderPassInfo RPInfo(ARRAY_COUNT(RenderTargets), RenderTargets, ERenderTargetActions::Load_Store);
+				RPInfo.DepthStencilRenderTarget.Action = EDepthStencilTargetActions::LoadDepthStencil_StoreDepthStencil;
+				RPInfo.DepthStencilRenderTarget.DepthStencilTarget = DepthTarget.TargetableTexture;
+				RPInfo.DepthStencilRenderTarget.ExclusiveDepthStencil = FExclusiveDepthStencil::DepthWrite_StencilWrite;
+
+				// Set starting UAV bind index
+				RPInfo.UAVIndex = ARRAY_COUNT(RenderTargets);
+				RPInfo.NumUAVs = ARRAY_COUNT(Uavs);
+				for (int32 Index = 0; Index < RPInfo.NumUAVs; Index++)
+				{
+					RPInfo.UAVs[Index] = Uavs[Index];
+				}
 					
 				InRHICmdList.TransitionResources(EResourceTransitionAccess::ERWBarrier, EResourceTransitionPipeline::EGfxToGfx, Uavs, ARRAY_COUNT(Uavs));
-				SetRenderTargets(InRHICmdList, ARRAY_COUNT(RenderTargets), RenderTargets, DepthTarget.TargetableTexture, ARRAY_COUNT(Uavs), Uavs);
+				InRHICmdList.BeginRenderPass(RPInfo, TEXT("ShadowAtlas"));
 
 				ProjectedShadowInfo->ClearDepth(InRHICmdList, this, ARRAY_COUNT(RenderTargets), RenderTargets, DepthTarget.TargetableTexture, bPerformClear);
 			};					
-
+			
 			{
 				SCOPED_DRAW_EVENT(RHICmdList, Clear);
-				SetShadowRenderTargets(RHICmdList, true);	
-			}				
+				BeginShadowRenderPass(RHICmdList, true);
+			}
 
-			ProjectedShadowInfo->RenderDepth(RHICmdList, this, SetShadowRenderTargets, ShadowDepthRenderMode_Normal);
+			// In parallel mode the first renderpass is just the clear.
+			if(bDoParallelDispatch)
+			{
+				RHICmdList.EndRenderPass();
+			}
+
+			ProjectedShadowInfo->RenderDepth(RHICmdList, this, BeginShadowRenderPass, ShadowDepthRenderMode_Normal, bDoParallelDispatch);
 
 			// Render emissive only meshes as they are held in a separate list.
-			ProjectedShadowInfo->RenderDepth(RHICmdList, this, SetShadowRenderTargets, ShadowDepthRenderMode_EmissiveOnly);
+			ProjectedShadowInfo->RenderDepth(RHICmdList, this, BeginShadowRenderPass, ShadowDepthRenderMode_EmissiveOnly, bDoParallelDispatch);
 			// Render gi blocking volume meshes.
-			ProjectedShadowInfo->RenderDepth(RHICmdList, this, SetShadowRenderTargets, ShadowDepthRenderMode_GIBlockingVolumes);
+			ProjectedShadowInfo->RenderDepth(RHICmdList, this, BeginShadowRenderPass, ShadowDepthRenderMode_GIBlockingVolumes, bDoParallelDispatch);
 
+			if (!bDoParallelDispatch)
+			{
+				RHICmdList.EndRenderPass();
+			}
 			{
 				// Resolve the shadow depth z surface.
 				RHICmdList.CopyToResolveTarget(DepthTarget.TargetableTexture, DepthTarget.ShaderResourceTexture, FResolveParams());
@@ -2347,12 +2515,10 @@ void FSceneRenderer::RenderShadowDepthMaps(FRHICommandListImmediate& RHICmdList)
 				UavsToReadable[0] = LightPropagationVolume->GetGvListBufferUav();
 				UavsToReadable[1] = LightPropagationVolume->GetGvListHeadBufferUav();	
 				RHICmdList.TransitionResources(EResourceTransitionAccess::EReadable, EResourceTransitionPipeline::EGfxToGfx, UavsToReadable, ARRAY_COUNT(UavsToReadable));
-
-				// Unset render targets
-				FTextureRHIParamRef RenderTargets[2] = {NULL};
-				FUnorderedAccessViewRHIParamRef Uavs[2] = {NULL};
-				SetRenderTargets(RHICmdList, ARRAY_COUNT(RenderTargets), RenderTargets, FTextureRHIParamRef(), ARRAY_COUNT(Uavs), Uavs);
 			}
+			checkSlow(RHICmdList.IsOutsideRenderPass());
 		}
 	}
+
+	checkSlow(RHICmdList.IsOutsideRenderPass());
 }
diff --git a/Engine/Source/Runtime/Renderer/Private/ShadowRendering.cpp b/Engine/Source/Runtime/Renderer/Private/ShadowRendering.cpp
--- a/Engine/Source/Runtime/Renderer/Private/ShadowRendering.cpp
+++ b/Engine/Source/Runtime/Renderer/Private/ShadowRendering.cpp
@@ -1402,52 +1402,105 @@ bool FDeferredShadingSceneRenderer::InjectReflectiveShadowMaps(FRHICommandListIm
 
 bool FSceneRenderer::RenderShadowProjections(FRHICommandListImmediate& RHICmdList, const FLightSceneInfo* LightSceneInfo, IPooledRenderTarget* ScreenShadowMaskTexture, bool bProjectingForForwardShading, bool bMobileModulatedProjections)
 {
+	check(RHICmdList.IsOutsideRenderPass());
+
 	FVisibleLightInfo& VisibleLightInfo = VisibleLightInfos[LightSceneInfo->Id];
 	FSceneRenderTargets& SceneContext = FSceneRenderTargets::Get(RHICmdList);
 
-	if (bMobileModulatedProjections)
-	{
-		SceneContext.BeginRenderingSceneColor(RHICmdList, ESimpleRenderTargetMode::EExistingColorAndDepth, FExclusiveDepthStencil::DepthRead_StencilWrite);
-	}
-	else
+	// Gather up our work real quick so we can do everything in one renderpass later.
+	// #todo-renderpasses How many ShadowsToProject do we have usually?
+	TArray<FProjectedShadowInfo*> DistanceFieldShadows;
+	TArray<FProjectedShadowInfo*> NormalShadows;
+
+	for (int32 ShadowIndex = 0; ShadowIndex < VisibleLightInfo.ShadowsToProject.Num(); ShadowIndex++)
 	{
-		// Normal deferred shadows render to the shadow mask
-		SetRenderTarget(RHICmdList, ScreenShadowMaskTexture->GetRenderTargetItem().TargetableTexture, SceneContext.GetSceneDepthSurface(), ESimpleRenderTargetMode::EExistingColorAndDepth, FExclusiveDepthStencil::DepthRead_StencilWrite, true);
+		FProjectedShadowInfo* ProjectedShadowInfo = VisibleLightInfo.ShadowsToProject[ShadowIndex];
+		if (ProjectedShadowInfo->bRayTracedDistanceField)
+		{
+			DistanceFieldShadows.Add(ProjectedShadowInfo);
+		}
+		else
+		{
+			NormalShadows.Add(ProjectedShadowInfo);
+		}
 	}
 
-	for (int32 ViewIndex = 0; ViewIndex < Views.Num(); ViewIndex++)
+	if (DistanceFieldShadows.Num() > 0)
 	{
-		SCOPED_CONDITIONAL_DRAW_EVENTF(RHICmdList, EventView, Views.Num() > 1, TEXT("View%d"), ViewIndex);
-
-		const FViewInfo& View = Views[ViewIndex];
+		// Dispatch distance field shadows first
+		for (int32 ViewIndex = 0; ViewIndex < Views.Num(); ViewIndex++)
+		{
+			SCOPED_CONDITIONAL_DRAW_EVENTF(RHICmdList, EventView, Views.Num() > 1, TEXT("DistanceFieldShadows_View%d"), ViewIndex);
 
-		// Set the device viewport for the view.
-		RHICmdList.SetViewport(View.ViewRect.Min.X, View.ViewRect.Min.Y, 0.0f, View.ViewRect.Max.X, View.ViewRect.Max.Y, 1.0f);
+			const FViewInfo& View = Views[ViewIndex];
 
-		// Set the light's scissor rectangle.
-		LightSceneInfo->Proxy->SetScissorRect(RHICmdList, View, View.ViewRect);
+			// Set the device viewport for the view.
+			RHICmdList.SetViewport(View.ViewRect.Min.X, View.ViewRect.Min.Y, 0.0f, View.ViewRect.Max.X, View.ViewRect.Max.Y, 1.0f);
 
-		// Project the shadow depth buffers onto the scene.
-		for (int32 ShadowIndex = 0; ShadowIndex < VisibleLightInfo.ShadowsToProject.Num(); ShadowIndex++)
-		{
-			FProjectedShadowInfo* ProjectedShadowInfo = VisibleLightInfo.ShadowsToProject[ShadowIndex];
+			// Set the light's scissor rectangle.
+			LightSceneInfo->Proxy->SetScissorRect(RHICmdList, View, View.ViewRect);
 
-			if (ProjectedShadowInfo->bRayTracedDistanceField)
+			// Project the shadow depth buffers onto the scene.
+			for (int32 ShadowIndex = 0; ShadowIndex < DistanceFieldShadows.Num(); ShadowIndex++)
 			{
+				FProjectedShadowInfo* ProjectedShadowInfo = DistanceFieldShadows[ShadowIndex];
 				ProjectedShadowInfo->RenderRayTracedDistanceFieldProjection(RHICmdList, View, ScreenShadowMaskTexture, bProjectingForForwardShading);
 			}
-			else if (ProjectedShadowInfo->bAllocated)
+
+			// Reset the scissor rectangle.
+			RHICmdList.SetScissorRect(false, 0, 0, 0, 0);
+		}
+	}
+
+	if (NormalShadows.Num() > 0)
+	{
+		// Render normal shadows
+		if (bMobileModulatedProjections)
+		{
+			SceneContext.BeginRenderingSceneColor(RHICmdList, ESimpleRenderTargetMode::EExistingColorAndDepth, FExclusiveDepthStencil::DepthRead_StencilWrite);
+		}
+		else
+		{
+			// Normal deferred shadows render to the shadow mask
+			FRHIRenderPassInfo RPInfo(ScreenShadowMaskTexture->GetRenderTargetItem().TargetableTexture, ERenderTargetActions::Load_Store);
+			RPInfo.DepthStencilRenderTarget.Action = MakeDepthStencilTargetActions(ERenderTargetActions::Load_DontStore, ERenderTargetActions::Load_Store);
+			RPInfo.DepthStencilRenderTarget.DepthStencilTarget = SceneContext.GetSceneDepthSurface();
+			RPInfo.DepthStencilRenderTarget.ExclusiveDepthStencil = FExclusiveDepthStencil::DepthRead_StencilWrite;
+
+			TransitionRenderPassTargets(RHICmdList, RPInfo);
+			RHICmdList.BeginRenderPass(RPInfo, TEXT("RenderShadowProjection"));
+		}
+
+		for (int32 ViewIndex = 0; ViewIndex < Views.Num(); ViewIndex++)
+		{
+			SCOPED_CONDITIONAL_DRAW_EVENTF(RHICmdList, EventView, Views.Num() > 1, TEXT("View%d"), ViewIndex);
+
+			const FViewInfo& View = Views[ViewIndex];
+
+			// Set the device viewport for the view.
+			RHICmdList.SetViewport(View.ViewRect.Min.X, View.ViewRect.Min.Y, 0.0f, View.ViewRect.Max.X, View.ViewRect.Max.Y, 1.0f);
+
+			// Set the light's scissor rectangle.
+			LightSceneInfo->Proxy->SetScissorRect(RHICmdList, View, View.ViewRect);
+
+			// Project the shadow depth buffers onto the scene.
+			for (int32 ShadowIndex = 0; ShadowIndex < NormalShadows.Num(); ShadowIndex++)
 			{
-				// Only project the shadow if it's large enough in this particular view (split screen, etc... may have shadows that are large in one view but irrelevantly small in others)
-				if (ProjectedShadowInfo->FadeAlphas[ViewIndex] > 1.0f / 256.0f)
+				FProjectedShadowInfo* ProjectedShadowInfo = NormalShadows[ShadowIndex];
+
+				if (ProjectedShadowInfo->bAllocated)
 				{
-					if (ProjectedShadowInfo->bOnePassPointLightShadow)
+					// Only project the shadow if it's large enough in this particular view (split screen, etc... may have shadows that are large in one view but irrelevantly small in others)
+					if (ProjectedShadowInfo->FadeAlphas[ViewIndex] > 1.0f / 256.0f)
 					{
-						ProjectedShadowInfo->RenderOnePassPointLightProjection(RHICmdList, ViewIndex, View, bProjectingForForwardShading);
-					}
-					else 
-					{
-						ProjectedShadowInfo->RenderProjection(RHICmdList, ViewIndex, &View, this, bProjectingForForwardShading, bMobileModulatedProjections);
+						if (ProjectedShadowInfo->bOnePassPointLightShadow)
+						{
+							ProjectedShadowInfo->RenderOnePassPointLightProjection(RHICmdList, ViewIndex, View, bProjectingForForwardShading);
+						}
+						else
+						{
+							ProjectedShadowInfo->RenderProjection(RHICmdList, ViewIndex, &View, this, bProjectingForForwardShading, bMobileModulatedProjections);
+						}
 					}
 				}
 			}
@@ -1455,6 +1508,15 @@ bool FSceneRenderer::RenderShadowProjections(FRHICommandListImmediate& RHICmdLis
 
 		// Reset the scissor rectangle.
 		RHICmdList.SetScissorRect(false, 0, 0, 0, 0);
+
+		if (bMobileModulatedProjections)
+		{
+			SceneContext.FinishRenderingSceneColor(RHICmdList);
+		}
+		else
+		{
+			RHICmdList.EndRenderPass();
+		}
 	}
 
 	return true;
@@ -1467,10 +1529,13 @@ bool FDeferredShadingSceneRenderer::RenderShadowProjections(FRHICommandListImmed
 	SCOPED_DRAW_EVENT(RHICmdList, ShadowProjectionOnOpaque);
 	SCOPED_GPU_STAT(RHICmdList, ShadowProjection);
 
+	check(RHICmdList.IsOutsideRenderPass());
+
 	FVisibleLightInfo& VisibleLightInfo = VisibleLightInfos[LightSceneInfo->Id];
 
 	FSceneRenderer::RenderShadowProjections(RHICmdList, LightSceneInfo, ScreenShadowMaskTexture, false, false);
 
+	checkSlow(RHICmdList.IsOutsideRenderPass());
 	for (int32 ShadowIndex = 0; ShadowIndex < VisibleLightInfo.ShadowsToProject.Num(); ShadowIndex++)
 	{
 		FProjectedShadowInfo* ProjectedShadowInfo = VisibleLightInfo.ShadowsToProject[ShadowIndex];
diff --git a/Engine/Source/Runtime/Renderer/Private/ShadowRendering.h b/Engine/Source/Runtime/Renderer/Private/ShadowRendering.h
--- a/Engine/Source/Runtime/Renderer/Private/ShadowRendering.h
+++ b/Engine/Source/Runtime/Renderer/Private/ShadowRendering.h
@@ -304,7 +304,7 @@ public:
 	}
 };
 
-typedef TFunctionRef<void(FRHICommandList& RHICmdList, bool bFirst)> FSetShadowRenderTargetFunction;
+typedef TFunctionRef<void(FRHICommandList& RHICmdList, bool bFirst)> FBeginShadowRenderPassFunction;
 
 /**
  * Information about a projected shadow.
@@ -481,7 +481,7 @@ public:
 	/**
 	 * Renders the shadow subject depth.
 	 */
-	void RenderDepth(FRHICommandList& RHICmdList, class FSceneRenderer* SceneRenderer, FSetShadowRenderTargetFunction SetShadowRenderTargets, EShadowDepthRenderMode RenderMode);
+	void RenderDepth(FRHICommandList& RHICmdList, class FSceneRenderer* SceneRenderer, FBeginShadowRenderPassFunction BeginShadowRenderPass, EShadowDepthRenderMode RenderMode, bool bDoParallelDispatch);
 
 	/** Set state for depth rendering */
 	void SetStateForDepth(FRHICommandList& RHICmdList, EShadowDepthRenderMode RenderMode, FDrawingPolicyRenderState& DrawRenderState);
@@ -647,12 +647,12 @@ private:
 	 */
 	float ShaderDepthBias;
 
-	void CopyCachedShadowMap(FRHICommandList& RHICmdList, const FDrawingPolicyRenderState& DrawRenderState, FSceneRenderer* SceneRenderer, const FViewInfo& View, FSetShadowRenderTargetFunction SetShadowRenderTargets);
+	void CopyCachedShadowMap(FRHICommandList& RHICmdList, const FDrawingPolicyRenderState& DrawRenderState, FSceneRenderer* SceneRenderer, const FViewInfo& View);
 
 	/**
 	* Renders the shadow subject depth, to a particular hacked view
 	*/
-	void RenderDepthInner(FRHICommandList& RHICmdList, class FSceneRenderer* SceneRenderer, const FViewInfo* FoundView, FSetShadowRenderTargetFunction SetShadowRenderTargets, EShadowDepthRenderMode RenderMode );
+	void RenderDepthInner(FRHICommandList& RHICmdList, class FSceneRenderer* SceneRenderer, const FViewInfo* FoundView, FBeginShadowRenderPassFunction BeginShadowRenderPass, EShadowDepthRenderMode RenderMode, bool bDoParallelDispatch );
 
 	/**
 	* Modifies the passed in view for this shadow
diff --git a/Engine/Source/Runtime/Renderer/Private/StaticMeshDrawList.inl b/Engine/Source/Runtime/Renderer/Private/StaticMeshDrawList.inl
--- a/Engine/Source/Runtime/Renderer/Private/StaticMeshDrawList.inl
+++ b/Engine/Source/Runtime/Renderer/Private/StaticMeshDrawList.inl
@@ -472,6 +472,8 @@ public:
 
 	void DoTask(ENamedThreads::Type CurrentThread, const FGraphEventRef& MyCompletionGraphEvent)
 	{
+		// In parallel rendering we must have a renderpass open before starting our task.
+		check(RHICmdList.IsInsideRenderPass());
 		LLM_SCOPE(ELLMTag::StaticMesh);
 
 		FScopeCycleCounter ScopeOuter(RHICmdList.ExecuteStat);
@@ -502,6 +504,9 @@ public:
 		{
 			this->Caller.DrawVisibleInner(this->RHICmdList, this->View, this->PolicyContext, this->DrawRenderState, this->StaticMeshVisibilityMap, this->BatchVisibilityArray, this->FirstPolicy, this->LastPolicy, true);
 		}
+
+		// Now that this task is complete we must close our renderpass.
+		this->RHICmdList.EndRenderPass();
 		this->RHICmdList.HandleRTThreadTaskCompletion(MyCompletionGraphEvent);
 	}
 };
diff --git a/Engine/Source/Runtime/Renderer/Private/TranslucentLighting.cpp b/Engine/Source/Runtime/Renderer/Private/TranslucentLighting.cpp
--- a/Engine/Source/Runtime/Renderer/Private/TranslucentLighting.cpp
+++ b/Engine/Source/Runtime/Renderer/Private/TranslucentLighting.cpp
@@ -530,6 +530,7 @@ public:
 
 void FProjectedShadowInfo::RenderTranslucencyDepths(FRHICommandList& RHICmdList, FSceneRenderer* SceneRenderer)
 {
+	check(RHICmdList.IsInsideRenderPass());
 	check(IsInRenderingThread());
 	checkSlow(!bWholeSceneShadow);
 	SCOPE_CYCLE_COUNTER(STAT_RenderPerObjectShadowDepthsTime);
@@ -1112,40 +1113,42 @@ void FDeferredShadingSceneRenderer::InjectAmbientCubemapTranslucentVolumeLightin
 			const IPooledRenderTarget* RT0 = SceneContext.TranslucencyLightingVolumeAmbient[VolumeCascadeIndex + NumTranslucentVolumeRenderTargetSets * ViewIndex];
 
 			// we don't update the directional volume (could be a HQ option)
-			SetRenderTarget(RHICmdList, RT0->GetRenderTargetItem().TargetableTexture, FTextureRHIRef(), true);
-			RHICmdList.ApplyCachedRenderTargets(GraphicsPSOInit);
-
+			FRHIRenderPassInfo RPInfo(RT0->GetRenderTargetItem().TargetableTexture, ERenderTargetActions::Load_Store);
+			TransitionRenderPassTargets(RHICmdList, RPInfo);
+			RHICmdList.BeginRenderPass(RPInfo, TEXT("InjectAmbientCubemapTranslucentVolumeLighting"));
+			{
+				RHICmdList.ApplyCachedRenderTargets(GraphicsPSOInit);
 
-			TShaderMapRef<FWriteToSliceVS> VertexShader(ShaderMap);
-			TOptionalShaderMapRef<FWriteToSliceGS> GeometryShader(ShaderMap);
-			TShaderMapRef<FInjectAmbientCubemapPS> PixelShader(ShaderMap);
+				TShaderMapRef<FWriteToSliceVS> VertexShader(ShaderMap);
+				TOptionalShaderMapRef<FWriteToSliceGS> GeometryShader(ShaderMap);
+				TShaderMapRef<FInjectAmbientCubemapPS> PixelShader(ShaderMap);
 
-			GraphicsPSOInit.BoundShaderState.VertexDeclarationRHI = GScreenVertexDeclaration.VertexDeclarationRHI;
-			GraphicsPSOInit.BoundShaderState.VertexShaderRHI = GETSAFERHISHADER_VERTEX(*VertexShader);
-			GraphicsPSOInit.BoundShaderState.GeometryShaderRHI = GETSAFERHISHADER_GEOMETRY(*GeometryShader);
-			GraphicsPSOInit.BoundShaderState.PixelShaderRHI = GETSAFERHISHADER_PIXEL(*PixelShader);
-			GraphicsPSOInit.PrimitiveType = PT_TriangleStrip;
+				GraphicsPSOInit.BoundShaderState.VertexDeclarationRHI = GScreenVertexDeclaration.VertexDeclarationRHI;
+				GraphicsPSOInit.BoundShaderState.VertexShaderRHI = GETSAFERHISHADER_VERTEX(*VertexShader);
+				GraphicsPSOInit.BoundShaderState.GeometryShaderRHI = GETSAFERHISHADER_GEOMETRY(*GeometryShader);
+				GraphicsPSOInit.BoundShaderState.PixelShaderRHI = GETSAFERHISHADER_PIXEL(*PixelShader);
+				GraphicsPSOInit.PrimitiveType = PT_TriangleStrip;
 									
-			SetGraphicsPipelineState(RHICmdList, GraphicsPSOInit);
+				SetGraphicsPipelineState(RHICmdList, GraphicsPSOInit);
 
-			VertexShader->SetParameters(RHICmdList, VolumeBounds, FIntVector(TranslucencyLightingVolumeDim));
-			if(GeometryShader.IsValid())
-			{
-				GeometryShader->SetParameters(RHICmdList, VolumeBounds.MinZ);
-			}
+				VertexShader->SetParameters(RHICmdList, VolumeBounds, FIntVector(TranslucencyLightingVolumeDim));
+				if(GeometryShader.IsValid())
+				{
+					GeometryShader->SetParameters(RHICmdList, VolumeBounds.MinZ);
+				}
 
-			uint32 Count = View.FinalPostProcessSettings.ContributingCubemaps.Num();
-			for(uint32 i = 0; i < Count; ++i)
-			{
-				const FFinalPostProcessSettings::FCubemapEntry& CubemapEntry = View.FinalPostProcessSettings.ContributingCubemaps[i];
+				uint32 Count = View.FinalPostProcessSettings.ContributingCubemaps.Num();
+				for(uint32 i = 0; i < Count; ++i)
+				{
+					const FFinalPostProcessSettings::FCubemapEntry& CubemapEntry = View.FinalPostProcessSettings.ContributingCubemaps[i];
 
-				PixelShader->SetParameters(RHICmdList, View, CubemapEntry);
+					PixelShader->SetParameters(RHICmdList, View, CubemapEntry);
 
-				RasterizeToVolumeTexture(RHICmdList, VolumeBounds);
+					RasterizeToVolumeTexture(RHICmdList, VolumeBounds);
+				}
 			}
-
-			RHICmdList.CopyToResolveTarget(RT0->GetRenderTargetItem().TargetableTexture,
-			RT0->GetRenderTargetItem().ShaderResourceTexture, FResolveParams());
+			RHICmdList.EndRenderPass();
+			RHICmdList.CopyToResolveTarget(RT0->GetRenderTargetItem().TargetableTexture, RT0->GetRenderTargetItem().ShaderResourceTexture, FResolveParams());
 		}
 	}
 }
@@ -1236,39 +1239,42 @@ void FDeferredShadingSceneRenderer::AccumulateTranslucentVolumeObjectShadowing(F
 					RenderTarget = SceneContext.GetTranslucencyVolumeDirectional(TVC_Inner, ViewIndex)->GetRenderTargetItem().TargetableTexture;
 				}
 
-				SetRenderTarget(RHICmdList, RenderTarget, FTextureRHIRef());
+				FRHIRenderPassInfo RPInfo(RenderTarget, ERenderTargetActions::Load_Store);
+				RHICmdList.BeginRenderPass(RPInfo, TEXT("AccumulateVolumeObjectShadowing"));
+				{
+					FGraphicsPipelineStateInitializer GraphicsPSOInit;
+					RHICmdList.ApplyCachedRenderTargets(GraphicsPSOInit);
 
-				FGraphicsPipelineStateInitializer GraphicsPSOInit;
-				RHICmdList.ApplyCachedRenderTargets(GraphicsPSOInit);
+					GraphicsPSOInit.RasterizerState = TStaticRasterizerState<FM_Solid, CM_None>::GetRHI();
+					GraphicsPSOInit.DepthStencilState = TStaticDepthStencilState<false, CF_Always>::GetRHI();
 
-				GraphicsPSOInit.RasterizerState = TStaticRasterizerState<FM_Solid, CM_None>::GetRHI();
-				GraphicsPSOInit.DepthStencilState = TStaticDepthStencilState<false, CF_Always>::GetRHI();
+					// Modulate the contribution of multiple object shadows in rgb
+					GraphicsPSOInit.BlendState = TStaticBlendState<CW_RGB, BO_Add, BF_DestColor, BF_Zero>::GetRHI();
 
-				// Modulate the contribution of multiple object shadows in rgb
-				GraphicsPSOInit.BlendState = TStaticBlendState<CW_RGB, BO_Add, BF_DestColor, BF_Zero>::GetRHI();
+					TShaderMapRef<FWriteToSliceVS> VertexShader(ShaderMap);
+					TOptionalShaderMapRef<FWriteToSliceGS> GeometryShader(ShaderMap);
+					TShaderMapRef<FTranslucentObjectShadowingPS> PixelShader(ShaderMap);
 
-				TShaderMapRef<FWriteToSliceVS> VertexShader(ShaderMap);
-				TOptionalShaderMapRef<FWriteToSliceGS> GeometryShader(ShaderMap);
-				TShaderMapRef<FTranslucentObjectShadowingPS> PixelShader(ShaderMap);
-
-				GraphicsPSOInit.BoundShaderState.VertexDeclarationRHI = GScreenVertexDeclaration.VertexDeclarationRHI;
-				GraphicsPSOInit.BoundShaderState.VertexShaderRHI = GETSAFERHISHADER_VERTEX(*VertexShader);
-				GraphicsPSOInit.BoundShaderState.GeometryShaderRHI = GETSAFERHISHADER_GEOMETRY(*GeometryShader);
-				GraphicsPSOInit.BoundShaderState.PixelShaderRHI = GETSAFERHISHADER_PIXEL(*PixelShader);
-				GraphicsPSOInit.PrimitiveType = PT_TriangleStrip;
+					GraphicsPSOInit.BoundShaderState.VertexDeclarationRHI = GScreenVertexDeclaration.VertexDeclarationRHI;
+					GraphicsPSOInit.BoundShaderState.VertexShaderRHI = GETSAFERHISHADER_VERTEX(*VertexShader);
+					GraphicsPSOInit.BoundShaderState.GeometryShaderRHI = GETSAFERHISHADER_GEOMETRY(*GeometryShader);
+					GraphicsPSOInit.BoundShaderState.PixelShaderRHI = GETSAFERHISHADER_PIXEL(*PixelShader);
+					GraphicsPSOInit.PrimitiveType = PT_TriangleStrip;
 
-				SetGraphicsPipelineState(RHICmdList, GraphicsPSOInit);
+					SetGraphicsPipelineState(RHICmdList, GraphicsPSOInit);
 
-				const int32 TranslucencyLightingVolumeDim = GetTranslucencyLightingVolumeDim();
+					const int32 TranslucencyLightingVolumeDim = GetTranslucencyLightingVolumeDim();
 
-				VertexShader->SetParameters(RHICmdList, VolumeBounds, FIntVector(TranslucencyLightingVolumeDim));
-				if(GeometryShader.IsValid())
-				{
-					GeometryShader->SetParameters(RHICmdList, VolumeBounds.MinZ);
-				}
-				PixelShader->SetParameters(RHICmdList, View, LightSceneInfo, InProjectedShadowInfo, VolumeCascadeIndex);
+					VertexShader->SetParameters(RHICmdList, VolumeBounds, FIntVector(TranslucencyLightingVolumeDim));
+					if(GeometryShader.IsValid())
+					{
+						GeometryShader->SetParameters(RHICmdList, VolumeBounds.MinZ);
+					}
+					PixelShader->SetParameters(RHICmdList, View, LightSceneInfo, InProjectedShadowInfo, VolumeCascadeIndex);
 				
-				RasterizeToVolumeTexture(RHICmdList, VolumeBounds);
+					RasterizeToVolumeTexture(RHICmdList, VolumeBounds);
+				}
+				RHICmdList.EndRenderPass();
 
 				RHICmdList.CopyToResolveTarget(SceneContext.GetTranslucencyVolumeAmbient((ETranslucencyVolumeCascade)VolumeCascadeIndex, ViewIndex)->GetRenderTargetItem().TargetableTexture,
 					SceneContext.GetTranslucencyVolumeAmbient((ETranslucencyVolumeCascade)VolumeCascadeIndex, ViewIndex)->GetRenderTargetItem().ShaderResourceTexture, FResolveParams());
@@ -1434,6 +1440,7 @@ static void AddLightForInjection(
 /** Injects all the lights in LightInjectionData into the translucent lighting volume textures. */
 static void InjectTranslucentLightArray(FRHICommandListImmediate& RHICmdList, const FViewInfo& View, const TArray<FTranslucentLightInjectionData, SceneRenderingAllocator>& LightInjectionData, int32 ViewIndex)
 {
+	check(RHICmdList.IsOutsideRenderPass());
 	FSceneRenderTargets& SceneContext = FSceneRenderTargets::Get(RHICmdList);
 	INC_DWORD_STAT_BY(STAT_NumLightsInjectedIntoTranslucency, LightInjectionData.Num());
 
@@ -1451,83 +1458,87 @@ static void InjectTranslucentLightArray(FRHICommandListImmediate& RHICmdList, co
 		RenderTargets[0] = RT0->GetRenderTargetItem().TargetableTexture;
 		RenderTargets[1] = RT1->GetRenderTargetItem().TargetableTexture;
 
-		SetRenderTargets(RHICmdList, ARRAY_COUNT(RenderTargets), RenderTargets, FTextureRHIRef(), 0, NULL, true);
-
-		FGraphicsPipelineStateInitializer GraphicsPSOInit;
-		RHICmdList.ApplyCachedRenderTargets(GraphicsPSOInit);
-
-		GraphicsPSOInit.RasterizerState = TStaticRasterizerState<FM_Solid, CM_None>::GetRHI();
-		GraphicsPSOInit.DepthStencilState = TStaticDepthStencilState<false, CF_Always>::GetRHI();
-		GraphicsPSOInit.PrimitiveType = PT_TriangleStrip;
+		FRHIRenderPassInfo RPInfo(ARRAY_COUNT(RenderTargets), RenderTargets, ERenderTargetActions::Load_Store);
+		TransitionRenderPassTargets(RHICmdList, RPInfo);
 
-		for (int32 LightIndex = 0; LightIndex < LightInjectionData.Num(); LightIndex++)
+		RHICmdList.BeginRenderPass(RPInfo, TEXT("InjectTranslucentLightArray"));
 		{
-			const FTranslucentLightInjectionData& InjectionData = LightInjectionData[LightIndex];
-			const FLightSceneInfo* const LightSceneInfo = InjectionData.LightSceneInfo;
-			const bool bInverseSquared = LightSceneInfo->Proxy->IsInverseSquared();
-			const bool bDirectionalLight = LightSceneInfo->Proxy->GetLightType() == LightType_Directional;
-			const FVolumeBounds VolumeBounds = CalculateLightVolumeBounds(LightSceneInfo->Proxy->GetBoundingSphere(), View, VolumeCascadeIndex, bDirectionalLight);
+			FGraphicsPipelineStateInitializer GraphicsPSOInit;
+			RHICmdList.ApplyCachedRenderTargets(GraphicsPSOInit);
 
-			if (VolumeBounds.IsValid())
+			GraphicsPSOInit.RasterizerState = TStaticRasterizerState<FM_Solid, CM_None>::GetRHI();
+			GraphicsPSOInit.DepthStencilState = TStaticDepthStencilState<false, CF_Always>::GetRHI();
+			GraphicsPSOInit.PrimitiveType = PT_TriangleStrip;
+
+			for (int32 LightIndex = 0; LightIndex < LightInjectionData.Num(); LightIndex++)
 			{
-				TShaderMapRef<FWriteToSliceVS> VertexShader(View.ShaderMap);
-				TOptionalShaderMapRef<FWriteToSliceGS> GeometryShader(View.ShaderMap);
+				const FTranslucentLightInjectionData& InjectionData = LightInjectionData[LightIndex];
+				const FLightSceneInfo* const LightSceneInfo = InjectionData.LightSceneInfo;
+				const bool bInverseSquared = LightSceneInfo->Proxy->IsInverseSquared();
+				const bool bDirectionalLight = LightSceneInfo->Proxy->GetLightType() == LightType_Directional;
+				const FVolumeBounds VolumeBounds = CalculateLightVolumeBounds(LightSceneInfo->Proxy->GetBoundingSphere(), View, VolumeCascadeIndex, bDirectionalLight);
 
-				if (bDirectionalLight)
+				if (VolumeBounds.IsValid())
 				{
-					// Accumulate the contribution of multiple lights
-					// Directional lights write their shadowing into alpha of the ambient texture
-					GraphicsPSOInit.BlendState = TStaticBlendState<
-						CW_RGBA, BO_Add, BF_One, BF_One, BO_Add, BF_One, BF_One,
-						CW_RGB, BO_Add, BF_One, BF_One, BO_Add, BF_One, BF_One>::GetRHI();
+					TShaderMapRef<FWriteToSliceVS> VertexShader(View.ShaderMap);
+					TOptionalShaderMapRef<FWriteToSliceGS> GeometryShader(View.ShaderMap);
 
-					if (InjectionData.ProjectedShadowInfo)
+					if (bDirectionalLight)
 					{
-						// shadows, restricting light contribution to the cascade bounds (except last cascade far to get light functions and no shadows there)
-						SetInjectionShader<LightType_Directional, true>(RHICmdList, GraphicsPSOInit, View, InjectionData.LightFunctionMaterialProxy, LightSceneInfo,
-							InjectionData.ProjectedShadowInfo, InjectionData.ProjectedShadowInfo->CascadeSettings.ShadowSplitIndex, VolumeCascadeIndex,
-							*VertexShader, *GeometryShader, InjectionData.bApplyLightFunction, false);
-					}
-					else
-					{
-						// no shadows
-						SetInjectionShader<LightType_Directional, false>(RHICmdList, GraphicsPSOInit, View, InjectionData.LightFunctionMaterialProxy, LightSceneInfo,
-							InjectionData.ProjectedShadowInfo, -1, VolumeCascadeIndex,
-							*VertexShader, *GeometryShader, InjectionData.bApplyLightFunction, false);
-					}
-				}
-				else
-				{
-					// Accumulate the contribution of multiple lights
-					GraphicsPSOInit.BlendState = TStaticBlendState<
-						CW_RGB, BO_Add, BF_One, BF_One, BO_Add, BF_Zero, BF_One,
-						CW_RGB, BO_Add, BF_One, BF_One, BO_Add, BF_Zero, BF_One>::GetRHI();
+						// Accumulate the contribution of multiple lights
+						// Directional lights write their shadowing into alpha of the ambient texture
+						GraphicsPSOInit.BlendState = TStaticBlendState<
+							CW_RGBA, BO_Add, BF_One, BF_One, BO_Add, BF_One, BF_One,
+							CW_RGB, BO_Add, BF_One, BF_One, BO_Add, BF_One, BF_One>::GetRHI();
 
-					if (InjectionData.ProjectedShadowInfo)
-					{
-						SetInjectionShader<LightType_Point, true>(RHICmdList, GraphicsPSOInit, View, InjectionData.LightFunctionMaterialProxy, LightSceneInfo,
-							InjectionData.ProjectedShadowInfo, -1, VolumeCascadeIndex,
-							*VertexShader, *GeometryShader, InjectionData.bApplyLightFunction, bInverseSquared);
+						if (InjectionData.ProjectedShadowInfo)
+						{
+							// shadows, restricting light contribution to the cascade bounds (except last cascade far to get light functions and no shadows there)
+							SetInjectionShader<LightType_Directional, true>(RHICmdList, GraphicsPSOInit, View, InjectionData.LightFunctionMaterialProxy, LightSceneInfo,
+								InjectionData.ProjectedShadowInfo, InjectionData.ProjectedShadowInfo->CascadeSettings.ShadowSplitIndex, VolumeCascadeIndex,
+								*VertexShader, *GeometryShader, InjectionData.bApplyLightFunction, false);
+						}
+						else
+						{
+							// no shadows
+							SetInjectionShader<LightType_Directional, false>(RHICmdList, GraphicsPSOInit, View, InjectionData.LightFunctionMaterialProxy, LightSceneInfo,
+								InjectionData.ProjectedShadowInfo, -1, VolumeCascadeIndex,
+								*VertexShader, *GeometryShader, InjectionData.bApplyLightFunction, false);
+						}
 					}
 					else
 					{
-						SetInjectionShader<LightType_Point, false>(RHICmdList, GraphicsPSOInit, View, InjectionData.LightFunctionMaterialProxy, LightSceneInfo,
-							InjectionData.ProjectedShadowInfo, -1, VolumeCascadeIndex,
-							*VertexShader, *GeometryShader, InjectionData.bApplyLightFunction, bInverseSquared);
+						// Accumulate the contribution of multiple lights
+						GraphicsPSOInit.BlendState = TStaticBlendState<
+							CW_RGB, BO_Add, BF_One, BF_One, BO_Add, BF_Zero, BF_One,
+							CW_RGB, BO_Add, BF_One, BF_One, BO_Add, BF_Zero, BF_One>::GetRHI();
+
+						if (InjectionData.ProjectedShadowInfo)
+						{
+							SetInjectionShader<LightType_Point, true>(RHICmdList, GraphicsPSOInit, View, InjectionData.LightFunctionMaterialProxy, LightSceneInfo,
+								InjectionData.ProjectedShadowInfo, -1, VolumeCascadeIndex,
+								*VertexShader, *GeometryShader, InjectionData.bApplyLightFunction, bInverseSquared);
+						}
+						else
+						{
+							SetInjectionShader<LightType_Point, false>(RHICmdList, GraphicsPSOInit, View, InjectionData.LightFunctionMaterialProxy, LightSceneInfo,
+								InjectionData.ProjectedShadowInfo, -1, VolumeCascadeIndex,
+								*VertexShader, *GeometryShader, InjectionData.bApplyLightFunction, bInverseSquared);
+						}
 					}
-				}
 
-				const int32 TranslucencyLightingVolumeDim = GetTranslucencyLightingVolumeDim();
+					const int32 TranslucencyLightingVolumeDim = GetTranslucencyLightingVolumeDim();
 
-				VertexShader->SetParameters(RHICmdList, VolumeBounds, FIntVector(TranslucencyLightingVolumeDim));
-				if(GeometryShader.IsValid())
-				{
-					GeometryShader->SetParameters(RHICmdList, VolumeBounds.MinZ);
+					VertexShader->SetParameters(RHICmdList, VolumeBounds, FIntVector(TranslucencyLightingVolumeDim));
+					if(GeometryShader.IsValid())
+					{
+						GeometryShader->SetParameters(RHICmdList, VolumeBounds.MinZ);
+					}
+					RasterizeToVolumeTexture(RHICmdList, VolumeBounds);
 				}
-				RasterizeToVolumeTexture(RHICmdList, VolumeBounds);
 			}
 		}
-
+		RHICmdList.EndRenderPass();
 		RHICmdList.CopyToResolveTarget(RT0->GetRenderTargetItem().TargetableTexture, RT0->GetRenderTargetItem().ShaderResourceTexture, FResolveParams());
 		RHICmdList.CopyToResolveTarget(RT1->GetRenderTargetItem().TargetableTexture, RT1->GetRenderTargetItem().ShaderResourceTexture, FResolveParams());
 	}
@@ -1633,6 +1644,7 @@ IMPLEMENT_SHADER_TYPE(,FSimpleLightTranslucentLightingInjectPS,TEXT("/Engine/Pri
 
 void FDeferredShadingSceneRenderer::InjectSimpleTranslucentVolumeLightingArray(FRHICommandListImmediate& RHICmdList, const FSimpleLightArray& SimpleLights, const FViewInfo& View, const int32 ViewIndex)
 {
+	check(RHICmdList.IsOutsideRenderPass());
 	SCOPE_CYCLE_COUNTER(STAT_TranslucentInjectTime);
 
 	int32 NumLightsToInject = 0;
@@ -1664,58 +1676,61 @@ void FDeferredShadingSceneRenderer::InjectSimpleTranslucentVolumeLightingArray(F
 			RenderTargets[0] = RT0->GetRenderTargetItem().TargetableTexture;
 			RenderTargets[1] = RT1->GetRenderTargetItem().TargetableTexture;
 
-			SetRenderTargets(RHICmdList, ARRAY_COUNT(RenderTargets), RenderTargets, FTextureRHIRef(), 0, NULL, true);
-
-			FGraphicsPipelineStateInitializer GraphicsPSOInit;
-			RHICmdList.ApplyCachedRenderTargets(GraphicsPSOInit);
-
-			GraphicsPSOInit.RasterizerState = TStaticRasterizerState<FM_Solid, CM_None>::GetRHI();
-			GraphicsPSOInit.DepthStencilState = TStaticDepthStencilState<false, CF_Always>::GetRHI();
-			// Accumulate the contribution of multiple lights
-			GraphicsPSOInit.BlendState = TStaticBlendState<
-				CW_RGB, BO_Add, BF_One, BF_One, BO_Add, BF_Zero, BF_One,
-				CW_RGB, BO_Add, BF_One, BF_One, BO_Add, BF_Zero, BF_One>::GetRHI();
-			GraphicsPSOInit.PrimitiveType = PT_TriangleStrip;
-		
-			for (int32 LightIndex = 0; LightIndex < SimpleLights.InstanceData.Num(); LightIndex++)
+			FRHIRenderPassInfo RPInfo(ARRAY_COUNT(RenderTargets), RenderTargets, ERenderTargetActions::Load_Store);
+			TransitionRenderPassTargets(RHICmdList, RPInfo);
+			RHICmdList.BeginRenderPass(RPInfo, TEXT("InjectSimpleTranslucentVolumeLightingArray"));
 			{
-				const FSimpleLightEntry& SimpleLight = SimpleLights.InstanceData[LightIndex];
-				const FSimpleLightPerViewEntry& SimpleLightPerViewData = SimpleLights.GetViewDependentData(LightIndex, ViewIndex, Views.Num());
+				FGraphicsPipelineStateInitializer GraphicsPSOInit;
+				RHICmdList.ApplyCachedRenderTargets(GraphicsPSOInit);
 
-				if (SimpleLight.bAffectTranslucency)
+				GraphicsPSOInit.RasterizerState = TStaticRasterizerState<FM_Solid, CM_None>::GetRHI();
+				GraphicsPSOInit.DepthStencilState = TStaticDepthStencilState<false, CF_Always>::GetRHI();
+				// Accumulate the contribution of multiple lights
+				GraphicsPSOInit.BlendState = TStaticBlendState<
+					CW_RGB, BO_Add, BF_One, BF_One, BO_Add, BF_Zero, BF_One,
+					CW_RGB, BO_Add, BF_One, BF_One, BO_Add, BF_Zero, BF_One>::GetRHI();
+				GraphicsPSOInit.PrimitiveType = PT_TriangleStrip;
+		
+				for (int32 LightIndex = 0; LightIndex < SimpleLights.InstanceData.Num(); LightIndex++)
 				{
-					const FSphere LightBounds(SimpleLightPerViewData.Position, SimpleLight.Radius);
-					const FVolumeBounds VolumeBounds = CalculateLightVolumeBounds(LightBounds, View, VolumeCascadeIndex, false);
+					const FSimpleLightEntry& SimpleLight = SimpleLights.InstanceData[LightIndex];
+					const FSimpleLightPerViewEntry& SimpleLightPerViewData = SimpleLights.GetViewDependentData(LightIndex, ViewIndex, Views.Num());
 
-					if (VolumeBounds.IsValid())
+					if (SimpleLight.bAffectTranslucency)
 					{
-						TShaderMapRef<FWriteToSliceVS> VertexShader(View.ShaderMap);
-						TOptionalShaderMapRef<FWriteToSliceGS> GeometryShader(View.ShaderMap);
-						TShaderMapRef<FSimpleLightTranslucentLightingInjectPS> PixelShader(View.ShaderMap);
+						const FSphere LightBounds(SimpleLightPerViewData.Position, SimpleLight.Radius);
+						const FVolumeBounds VolumeBounds = CalculateLightVolumeBounds(LightBounds, View, VolumeCascadeIndex, false);
 
-						GraphicsPSOInit.BoundShaderState.VertexDeclarationRHI = GScreenVertexDeclaration.VertexDeclarationRHI;
-						GraphicsPSOInit.BoundShaderState.VertexShaderRHI = GETSAFERHISHADER_VERTEX(*VertexShader);
-						GraphicsPSOInit.BoundShaderState.GeometryShaderRHI = GETSAFERHISHADER_GEOMETRY(*GeometryShader);
-						GraphicsPSOInit.BoundShaderState.PixelShaderRHI = GETSAFERHISHADER_PIXEL(*PixelShader);
+						if (VolumeBounds.IsValid())
+						{
+							TShaderMapRef<FWriteToSliceVS> VertexShader(View.ShaderMap);
+							TOptionalShaderMapRef<FWriteToSliceGS> GeometryShader(View.ShaderMap);
+							TShaderMapRef<FSimpleLightTranslucentLightingInjectPS> PixelShader(View.ShaderMap);
 
-						SetGraphicsPipelineState(RHICmdList, GraphicsPSOInit);
+							GraphicsPSOInit.BoundShaderState.VertexDeclarationRHI = GScreenVertexDeclaration.VertexDeclarationRHI;
+							GraphicsPSOInit.BoundShaderState.VertexShaderRHI = GETSAFERHISHADER_VERTEX(*VertexShader);
+							GraphicsPSOInit.BoundShaderState.GeometryShaderRHI = GETSAFERHISHADER_GEOMETRY(*GeometryShader);
+							GraphicsPSOInit.BoundShaderState.PixelShaderRHI = GETSAFERHISHADER_PIXEL(*PixelShader);
 
-						const int32 TranslucencyLightingVolumeDim = GetTranslucencyLightingVolumeDim();
+							SetGraphicsPipelineState(RHICmdList, GraphicsPSOInit);
 
-						VertexShader->SetParameters(RHICmdList, VolumeBounds, FIntVector(TranslucencyLightingVolumeDim));
-						if(GeometryShader.IsValid())
-						{
-							GeometryShader->SetParameters(RHICmdList, VolumeBounds.MinZ);
-						}
-						PixelShader->SetParameters(RHICmdList, View, SimpleLight, SimpleLightPerViewData, VolumeCascadeIndex);
+							const int32 TranslucencyLightingVolumeDim = GetTranslucencyLightingVolumeDim();
+
+							VertexShader->SetParameters(RHICmdList, VolumeBounds, FIntVector(TranslucencyLightingVolumeDim));
+							if(GeometryShader.IsValid())
+							{
+								GeometryShader->SetParameters(RHICmdList, VolumeBounds.MinZ);
+							}
+							PixelShader->SetParameters(RHICmdList, View, SimpleLight, SimpleLightPerViewData, VolumeCascadeIndex);
 
 
 
-						RasterizeToVolumeTexture(RHICmdList, VolumeBounds);
+							RasterizeToVolumeTexture(RHICmdList, VolumeBounds);
+						}
 					}
 				}
 			}
-
+			RHICmdList.EndRenderPass();
 			RHICmdList.CopyToResolveTarget(RT0->GetRenderTargetItem().TargetableTexture, RT0->GetRenderTargetItem().ShaderResourceTexture, FResolveParams());
 			RHICmdList.CopyToResolveTarget(RT1->GetRenderTargetItem().TargetableTexture, RT1->GetRenderTargetItem().ShaderResourceTexture, FResolveParams());
 		}
@@ -1784,30 +1799,35 @@ void FDeferredShadingSceneRenderer::FilterTranslucentVolumeLighting(FRHICommandL
 				}
 				RHICmdList.TransitionResources(EResourceTransitionAccess::EReadable, Inputs, 2);
 
-				SetRenderTargets(RHICmdList, ARRAY_COUNT(RenderTargets), RenderTargets, FTextureRHIRef(), 0, NULL, true);
-				RHICmdList.ApplyCachedRenderTargets(GraphicsPSOInit);
+				FRHIRenderPassInfo RPInfo(ARRAY_COUNT(RenderTargets), RenderTargets, ERenderTargetActions::Load_Store);
+				TransitionRenderPassTargets(RHICmdList, RPInfo);
+				RHICmdList.BeginRenderPass(RPInfo, TEXT("FilterTranslucentVolumeLighting"));
+				{
+					RHICmdList.ApplyCachedRenderTargets(GraphicsPSOInit);
 
-				const FVolumeBounds VolumeBounds(TranslucencyLightingVolumeDim);
-				TShaderMapRef<FWriteToSliceVS> VertexShader(View.ShaderMap);
-				TOptionalShaderMapRef<FWriteToSliceGS> GeometryShader(View.ShaderMap);
-				TShaderMapRef<FFilterTranslucentVolumePS> PixelShader(View.ShaderMap);
+					const FVolumeBounds VolumeBounds(TranslucencyLightingVolumeDim);
+					TShaderMapRef<FWriteToSliceVS> VertexShader(View.ShaderMap);
+					TOptionalShaderMapRef<FWriteToSliceGS> GeometryShader(View.ShaderMap);
+					TShaderMapRef<FFilterTranslucentVolumePS> PixelShader(View.ShaderMap);
 				
-				GraphicsPSOInit.BoundShaderState.VertexDeclarationRHI = GScreenVertexDeclaration.VertexDeclarationRHI;
-				GraphicsPSOInit.BoundShaderState.VertexShaderRHI = GETSAFERHISHADER_VERTEX(*VertexShader);
-				GraphicsPSOInit.BoundShaderState.GeometryShaderRHI = GETSAFERHISHADER_GEOMETRY(*GeometryShader);
-				GraphicsPSOInit.BoundShaderState.PixelShaderRHI = GETSAFERHISHADER_PIXEL(*PixelShader);
-				GraphicsPSOInit.PrimitiveType = PT_TriangleStrip;
+					GraphicsPSOInit.BoundShaderState.VertexDeclarationRHI = GScreenVertexDeclaration.VertexDeclarationRHI;
+					GraphicsPSOInit.BoundShaderState.VertexShaderRHI = GETSAFERHISHADER_VERTEX(*VertexShader);
+					GraphicsPSOInit.BoundShaderState.GeometryShaderRHI = GETSAFERHISHADER_GEOMETRY(*GeometryShader);
+					GraphicsPSOInit.BoundShaderState.PixelShaderRHI = GETSAFERHISHADER_PIXEL(*PixelShader);
+					GraphicsPSOInit.PrimitiveType = PT_TriangleStrip;
 
-				SetGraphicsPipelineState(RHICmdList, GraphicsPSOInit);
+					SetGraphicsPipelineState(RHICmdList, GraphicsPSOInit);
 
-				VertexShader->SetParameters(RHICmdList, VolumeBounds, FIntVector(TranslucencyLightingVolumeDim));
-				if(GeometryShader.IsValid())
-				{
-					GeometryShader->SetParameters(RHICmdList, VolumeBounds.MinZ);
-				}
-				PixelShader->SetParameters(RHICmdList, View, VolumeCascadeIndex, ViewIndex);
+					VertexShader->SetParameters(RHICmdList, VolumeBounds, FIntVector(TranslucencyLightingVolumeDim));
+					if(GeometryShader.IsValid())
+					{
+						GeometryShader->SetParameters(RHICmdList, VolumeBounds.MinZ);
+					}
+					PixelShader->SetParameters(RHICmdList, View, VolumeCascadeIndex, ViewIndex);
 
-				RasterizeToVolumeTexture(RHICmdList, VolumeBounds);
+					RasterizeToVolumeTexture(RHICmdList, VolumeBounds);
+				}
+				RHICmdList.EndRenderPass();
 
 				//only do readable transition on the final loop since the other ones will do this up front.
 				//if (VolumeCascadeIndex == TVC_MAX - 1)
diff --git a/Engine/Source/Runtime/Renderer/Private/TranslucentRendering.cpp b/Engine/Source/Runtime/Renderer/Private/TranslucentRendering.cpp
--- a/Engine/Source/Runtime/Renderer/Private/TranslucentRendering.cpp
+++ b/Engine/Source/Runtime/Renderer/Private/TranslucentRendering.cpp
@@ -672,6 +672,7 @@ public:
 
 	void DoTask(ENamedThreads::Type CurrentThread, const FGraphEventRef& MyCompletionGraphEvent)
 	{
+		check(RHICmdList.IsOutsideRenderPass());
 		// Never needs clear here as it is already done in RenderTranslucency.
 		FSceneRenderTargets& SceneContext = FSceneRenderTargets::Get(RHICmdList);
 		if (SceneContext.IsSeparateTranslucencyPass())
@@ -684,6 +685,8 @@ public:
 		}
 
 		PrimSet.DrawAPrimitive(RHICmdList, View, DrawRenderState, Renderer, TranslucencyPass, Index);
+
+		RHICmdList.EndRenderPass();
 		RHICmdList.HandleRTThreadTaskCompletion(MyCompletionGraphEvent);
 	}
 };
@@ -697,6 +700,7 @@ void FTranslucentPrimSet::DrawPrimitivesParallel(
 	int32 FirstPrimIdx, int32 LastPrimIdx
 	) const
 {
+	TArray<int32> Volumetrics;
 	// Draw sorted scene prims
 	for (int32 PrimIdx = FirstPrimIdx; PrimIdx <= LastPrimIdx; PrimIdx++)
 	{
@@ -708,18 +712,45 @@ void FTranslucentPrimSet::DrawPrimitivesParallel(
 
 		if (PrimitiveSceneInfo->Proxy && PrimitiveSceneInfo->Proxy->CastsVolumetricTranslucentShadow())
 		{
+			Volumetrics.Add(PrimIdx);
 			check(!IsInActualRenderingThread());
 			// can't do this in parallel, defer
-			FRHICommandList* CmdList = new FRHICommandList(View.GPUMask);
-			CmdList->CopyRenderThreadContexts(RHICmdList);
-			FGraphEventRef RenderThreadCompletionEvent = TGraphTask<FVolumetricTranslucentShadowRenderThreadTask>::CreateTask().ConstructAndDispatchWhenReady(*CmdList, *this, View, DrawRenderState, Renderer, TranslucencyPass, PrimIdx);
-			RHICmdList.QueueRenderThreadCommandListSubmit(RenderThreadCompletionEvent, CmdList);
+			//FRHICommandList* CmdList = new FRHICommandList(View.GPUMask);
+			//CmdList->CopyRenderThreadContexts(RHICmdList);
+			//FGraphEventRef RenderThreadCompletionEvent = TGraphTask<FVolumetricTranslucentShadowRenderThreadTask>::CreateTask().ConstructAndDispatchWhenReady(*CmdList, *this, View, DrawRenderState, Renderer, TranslucencyPass, PrimIdx);
+			//RHICmdList.QueueRenderThreadCommandListSubmit(RenderThreadCompletionEvent, CmdList);
 		}
 		else
 		{
 			RenderPrimitive(RHICmdList, View, DrawRenderState, PrimitiveSceneInfo, ViewRelevance, nullptr, TranslucencyPass);
 		}
 	}
+
+	RHICmdList.EndRenderPass();
+
+	if (Volumetrics.Num() > 0)
+	{
+		FRHICommandList* CmdList = new FRHICommandList(View.GPUMask);
+		// #todo-renderpasses Do we actually need the context here?
+		CmdList->CopyRenderThreadContexts(RHICmdList);
+		FSceneRenderTargets& SceneContext = FSceneRenderTargets::Get(RHICmdList);
+		if (SceneContext.IsSeparateTranslucencyPass())
+		{
+			SceneContext.BeginRenderingSeparateTranslucency(*CmdList, View, Renderer, false);
+		}
+		else
+		{
+			SceneContext.BeginRenderingTranslucency(*CmdList, View, Renderer, false);
+		}
+
+		for (int32 Index = 0; Index < Volumetrics.Num(); Index++)
+		{
+			DrawAPrimitive(*CmdList, View, DrawRenderState, Renderer, TranslucencyPass, Index);
+		}
+		CmdList->EndRenderPass();
+
+		RHICmdList.QueueCommandListSubmit(CmdList);
+	}
 }
 
 void FTranslucentPrimSet::DrawPrimitives(
@@ -953,6 +984,7 @@ public:
 
 	void DoTask(ENamedThreads::Type CurrentThread, const FGraphEventRef& MyCompletionGraphEvent)
 	{
+		check(RHICmdList.IsInsideRenderPass());
 		FScopeCycleCounter ScopeOuter(RHICmdList.ExecuteStat);
 		View.TranslucentPrimSet.DrawPrimitivesParallel(RHICmdList, View, DrawRenderState, Renderer, TranslucencyPass, FirstIndex, LastIndex);
 		RHICmdList.HandleRTThreadTaskCompletion(MyCompletionGraphEvent);
@@ -973,7 +1005,6 @@ public:
 		, TranslucencyPass(InTranslucencyPass)
 		, bRenderInSeparateTranslucency(InRenderInSeparateTranslucency)
 	{
-		SetStateOnCommandList(ParentCmdList);
 	}
 
 	virtual ~FTranslucencyPassParallelCommandListSet()
@@ -1021,6 +1052,8 @@ static FTranslucencyDrawingPolicyFactory::ContextType GParallelTranslucencyConte
 
 void FDeferredShadingSceneRenderer::RenderViewTranslucency(FRHICommandListImmediate& RHICmdList, const FViewInfo& View, const FDrawingPolicyRenderState& DrawRenderState, ETranslucencyPass::Type TranslucencyPass)
 {
+	check(RHICmdList.IsInsideRenderPass());
+
 	SCOPED_GPU_MASK(RHICmdList, View.GPUMask);
 
 	// Draw translucent prims
@@ -1052,6 +1085,8 @@ void FDeferredShadingSceneRenderer::RenderViewTranslucency(FRHICommandListImmedi
 }
 void FDeferredShadingSceneRenderer::RenderViewTranslucencyParallel(FRHICommandListImmediate& RHICmdList, const FViewInfo& View, const FDrawingPolicyRenderState& DrawRenderState, ETranslucencyPass::Type TranslucencyPass)
 {
+	check(RHICmdList.IsOutsideRenderPass());
+
 	SCOPED_GPU_MASK(RHICmdList, View.GPUMask);
 
 	FTranslucencyPassParallelCommandListSet ParallelCommandListSet(
@@ -1182,41 +1217,45 @@ void FDeferredShadingSceneRenderer::ConditionalResolveSceneColorForTranslucentMa
 				GRenderTargetPool.FindFreeElement(RHICmdList, Desc, SceneColorCopy, TEXT("SceneColorCopy"));
 			}
 
-			SetRenderTarget(RHICmdList, SceneColorCopy->GetRenderTargetItem().TargetableTexture, nullptr, ESimpleRenderTargetMode::EUninitializedColorAndDepth, FExclusiveDepthStencil::DepthNop_StencilNop, true);
+			FRHIRenderPassInfo RPInfo(SceneColorCopy->GetRenderTargetItem().TargetableTexture, ERenderTargetActions::DontLoad_Store);
+			TransitionRenderPassTargets(RHICmdList, RPInfo);
+			RHICmdList.BeginRenderPass(RPInfo, TEXT("ResolveColorForTranslucentMaterials"));
+			{
 
-			RHICmdList.SetViewport(View.ViewRect.Min.X, View.ViewRect.Min.Y, 0.0f, View.ViewRect.Max.X, View.ViewRect.Max.Y, 1.0f);
+				RHICmdList.SetViewport(View.ViewRect.Min.X, View.ViewRect.Min.Y, 0.0f, View.ViewRect.Max.X, View.ViewRect.Max.Y, 1.0f);
 
-			FGraphicsPipelineStateInitializer GraphicsPSOInit;
-			RHICmdList.ApplyCachedRenderTargets(GraphicsPSOInit);
-			GraphicsPSOInit.RasterizerState = TStaticRasterizerState<FM_Solid, CM_None>::GetRHI();
-			GraphicsPSOInit.DepthStencilState = TStaticDepthStencilState<false, CF_Always>::GetRHI();
-			GraphicsPSOInit.BlendState = TStaticBlendState<>::GetRHI();
+				FGraphicsPipelineStateInitializer GraphicsPSOInit;
+				RHICmdList.ApplyCachedRenderTargets(GraphicsPSOInit);
+				GraphicsPSOInit.RasterizerState = TStaticRasterizerState<FM_Solid, CM_None>::GetRHI();
+				GraphicsPSOInit.DepthStencilState = TStaticDepthStencilState<false, CF_Always>::GetRHI();
+				GraphicsPSOInit.BlendState = TStaticBlendState<>::GetRHI();
 
 
-			TShaderMapRef<FScreenVS> ScreenVertexShader(View.ShaderMap);
-			TShaderMapRef<FCopySceneColorPS> PixelShader(View.ShaderMap);
+				TShaderMapRef<FScreenVS> ScreenVertexShader(View.ShaderMap);
+				TShaderMapRef<FCopySceneColorPS> PixelShader(View.ShaderMap);
 
-			GraphicsPSOInit.BoundShaderState.VertexDeclarationRHI = GFilterVertexDeclaration.VertexDeclarationRHI;
-			GraphicsPSOInit.BoundShaderState.VertexShaderRHI = GETSAFERHISHADER_VERTEX(*ScreenVertexShader);
-			GraphicsPSOInit.BoundShaderState.PixelShaderRHI = GETSAFERHISHADER_PIXEL(*PixelShader);
-			GraphicsPSOInit.PrimitiveType = PT_TriangleList;
+				GraphicsPSOInit.BoundShaderState.VertexDeclarationRHI = GFilterVertexDeclaration.VertexDeclarationRHI;
+				GraphicsPSOInit.BoundShaderState.VertexShaderRHI = GETSAFERHISHADER_VERTEX(*ScreenVertexShader);
+				GraphicsPSOInit.BoundShaderState.PixelShaderRHI = GETSAFERHISHADER_PIXEL(*PixelShader);
+				GraphicsPSOInit.PrimitiveType = PT_TriangleList;
 
-			SetGraphicsPipelineState(RHICmdList, GraphicsPSOInit);
+				SetGraphicsPipelineState(RHICmdList, GraphicsPSOInit);
 
-			PixelShader->SetParameters(RHICmdList, View);
+				PixelShader->SetParameters(RHICmdList, View);
 
-			DrawRectangle(
-				RHICmdList,
-				0, 0, 
-				View.ViewRect.Width(), View.ViewRect.Height(),
-				View.ViewRect.Min.X, View.ViewRect.Min.Y, 
-				View.ViewRect.Width(), View.ViewRect.Height(),
-				FIntPoint(View.ViewRect.Width(), View.ViewRect.Height()),
-				SceneContext.GetBufferSizeXY(),
-				*ScreenVertexShader,
-				EDRF_UseTriangleOptimization);
-
-			RHICmdList.CopyToResolveTarget(SceneColorCopy->GetRenderTargetItem().TargetableTexture, SceneColorCopy->GetRenderTargetItem().ShaderResourceTexture, FResolveParams(FResolveRect()));
+				DrawRectangle(
+					RHICmdList,
+					0, 0,
+					View.ViewRect.Width(), View.ViewRect.Height(),
+					View.ViewRect.Min.X, View.ViewRect.Min.Y,
+					View.ViewRect.Width(), View.ViewRect.Height(),
+					FIntPoint(View.ViewRect.Width(), View.ViewRect.Height()),
+					SceneContext.GetBufferSizeXY(),
+					*ScreenVertexShader,
+					EDRF_UseTriangleOptimization);
+			}
+			RHICmdList.EndRenderPass();
+			RHICmdList.CopyToResolveTarget(SceneColorCopy->GetRenderTargetItem().TargetableTexture, SceneColorCopy->GetRenderTargetItem().ShaderResourceTexture, FResolveParams());
 		}
 	}
 }
@@ -1339,6 +1378,8 @@ void CreateTranslucentBasePassUniformBuffer(
 
 void FDeferredShadingSceneRenderer::RenderTranslucency(FRHICommandListImmediate& RHICmdList, ETranslucencyPass::Type TranslucencyPass, IPooledRenderTarget* SceneColorCopy)
 {
+	check(RHICmdList.IsOutsideRenderPass());
+
 	if (!ShouldRenderTranslucency(TranslucencyPass))
 	{
 		return; // Early exit if nothing needs to be done.
@@ -1359,6 +1400,8 @@ void FDeferredShadingSceneRenderer::RenderTranslucency(FRHICommandListImmediate&
 
 	for (int32 ViewIndex = 0; ViewIndex < Views.Num(); ViewIndex++)
 	{
+		checkSlow(RHICmdList.IsOutsideRenderPass());
+
 		SCOPED_CONDITIONAL_DRAW_EVENTF(RHICmdList, EventView, Views.Num() > 1, TEXT("View%d"), ViewIndex);
 		FViewInfo& View = Views[ViewIndex];
 		if (!View.ShouldRenderView())
@@ -1380,6 +1423,8 @@ void FDeferredShadingSceneRenderer::RenderTranslucency(FRHICommandListImmediate&
 		// If downsampling we need to render in the separate buffer. Otherwise we also need to render offscreen to apply TPT_TranslucencyAfterDOF
 		if (RenderInSeparateTranslucency(SceneContext, TranslucencyPass, View.TranslucentPrimSet.SortedPrimsNum.DisableOffscreenRendering(TranslucencyPass)))
 		{
+			checkSlow(RHICmdList.IsOutsideRenderPass());
+
 			FIntPoint ScaledSize;
 			float DownsamplingScale = 1.f;
 			SceneContext.GetSeparateTranslucencyDimensions(ScaledSize, DownsamplingScale);
@@ -1394,21 +1439,25 @@ void FDeferredShadingSceneRenderer::RenderTranslucency(FRHICommandListImmediate&
 			{
 				BeginTimingSeparateTranslucencyPass(RHICmdList, View);
 			}
-			SceneContext.BeginRenderingSeparateTranslucency(RHICmdList, View, *this, ViewIndex == 0 || View.Family->bMultiGPUForkAndJoin);
 
+			SceneContext.BeginRenderingSeparateTranslucency(RHICmdList, View, *this, ViewIndex == 0 || View.Family->bMultiGPUForkAndJoin);
+			
 			// Draw only translucent prims that are in the SeparateTranslucency pass
 			DrawRenderState.SetDepthStencilState(TStaticDepthStencilState<false, CF_DepthNearOrEqual>::GetRHI());
 
 			if (bUseParallel)
 			{
+				RHICmdList.EndRenderPass();
 				RenderViewTranslucencyParallel(RHICmdList, View, DrawRenderState, TranslucencyPass);
 			}
 			else
 			{
 				RenderViewTranslucency(RHICmdList, View, DrawRenderState, TranslucencyPass);
+				RHICmdList.EndRenderPass();
 			}
 
 			SceneContext.FinishRenderingSeparateTranslucency(RHICmdList, View);
+
 			if (TranslucencyPass == ETranslucencyPass::TPT_TranslucencyAfterDOF)
 			{
 				EndTimingSeparateTranslucencyPass(RHICmdList, View);
@@ -1417,6 +1466,8 @@ void FDeferredShadingSceneRenderer::RenderTranslucency(FRHICommandListImmediate&
 			{
 				FTranslucencyDrawingPolicyFactory::UpsampleTranslucency(RHICmdList, View, false);
 			}
+
+			checkSlow(RHICmdList.IsOutsideRenderPass());
 		}
 		else
 		{
@@ -1425,12 +1476,16 @@ void FDeferredShadingSceneRenderer::RenderTranslucency(FRHICommandListImmediate&
 	
 			if (bUseParallel)
 			{
+				RHICmdList.EndRenderPass();
 				RenderViewTranslucencyParallel(RHICmdList, View, DrawRenderState, TranslucencyPass);
 			}
 			else
 			{
 				RenderViewTranslucency(RHICmdList, View, DrawRenderState, TranslucencyPass);
+				RHICmdList.EndRenderPass();
 			}
+
+			SceneContext.FinishRenderingTranslucency(RHICmdList);
 		}
 
 #if STATS
@@ -1440,6 +1495,8 @@ void FDeferredShadingSceneRenderer::RenderTranslucency(FRHICommandListImmediate&
 		}
 #endif
 	}
+
+	checkSlow(RHICmdList.IsOutsideRenderPass());
 }
 
 class FTranslucencyUpsamplingPS : public FGlobalShader
@@ -1603,4 +1660,6 @@ void FTranslucencyDrawingPolicyFactory::UpsampleTranslucency(FRHICommandList& RH
 		FIntPoint(TextureWidth, TextureHeight),
 		*ScreenVertexShader,
 		EDRF_UseTriangleOptimization);
+
+	SceneContext.FinishRenderingSceneColor(RHICmdList);
 }
diff --git a/Engine/Source/Runtime/Renderer/Private/VelocityRendering.cpp b/Engine/Source/Runtime/Renderer/Private/VelocityRendering.cpp
--- a/Engine/Source/Runtime/Renderer/Private/VelocityRendering.cpp
+++ b/Engine/Source/Runtime/Renderer/Private/VelocityRendering.cpp
@@ -623,6 +623,7 @@ bool IsMotionBlurEnabled(const FViewInfo& View)
 
 void FDeferredShadingSceneRenderer::RenderDynamicVelocitiesMeshElementsInner(FRHICommandList& RHICmdList, const FViewInfo& View, const FDrawingPolicyRenderState& DrawRenderState, int32 FirstIndex, int32 LastIndex)
 {
+	check(RHICmdList.IsInsideRenderPass());
 	FVelocityDrawingPolicyFactory::ContextType Context(DDM_AllOccluders, false);
 
 	for (int32 MeshBatchIndex = FirstIndex; MeshBatchIndex <= LastIndex; MeshBatchIndex++)
@@ -674,30 +675,34 @@ public:
 
 	void DoTask(ENamedThreads::Type CurrentThread, const FGraphEventRef& MyCompletionGraphEvent)
 	{
+		check(RHICmdList.IsInsideRenderPass());
 		ThisRenderer.RenderDynamicVelocitiesMeshElementsInner(RHICmdList, View, DrawRenderState, FirstIndex, LastIndex);
+		RHICmdList.EndRenderPass();
 		RHICmdList.HandleRTThreadTaskCompletion(MyCompletionGraphEvent);
 	}
 };
 
 static void BeginVelocityRendering(FRHICommandList& RHICmdList, TRefCountPtr<IPooledRenderTarget>& VelocityRT, bool bPerformClear)
 {
+	check(RHICmdList.IsOutsideRenderPass());
+
 	FTextureRHIRef VelocityTexture = VelocityRT->GetRenderTargetItem().TargetableTexture;
 	FTexture2DRHIRef DepthTexture = FSceneRenderTargets::Get(RHICmdList).GetSceneDepthTexture();	
+
+	FRHIRenderPassInfo RPInfo(VelocityTexture, ERenderTargetActions::Load_Store);
+	RPInfo.DepthStencilRenderTarget.Action = MakeDepthStencilTargetActions(ERenderTargetActions::Load_Store, ERenderTargetActions::Load_Store);
+	RPInfo.DepthStencilRenderTarget.DepthStencilTarget = DepthTexture;
+	RPInfo.DepthStencilRenderTarget.ExclusiveDepthStencil = FExclusiveDepthStencil::DepthWrite_StencilWrite;
+
 	if (bPerformClear)
 	{
-		// now make the FRHISetRenderTargetsInfo that encapsulates all of the info
-		FRHIRenderTargetView ColorView(VelocityTexture, 0, -1, ERenderTargetLoadAction::EClear, ERenderTargetStoreAction::EStore);
-		FRHIDepthRenderTargetView DepthView(DepthTexture, ERenderTargetLoadAction::ELoad, ERenderTargetStoreAction::ENoAction, FExclusiveDepthStencil::DepthRead_StencilWrite);
+		RPInfo.ColorRenderTargets[0].Action = ERenderTargetActions::Clear_Store;
+	}
 
-		FRHISetRenderTargetsInfo Info(1, &ColorView, DepthView);		
+	RHICmdList.BeginRenderPass(RPInfo, TEXT("VelocityRendering"));
 
-		// Clear the velocity buffer (0.0f means "use static background velocity").
-		RHICmdList.SetRenderTargetsAndClear(Info);		
-	}
-	else
+	if (!bPerformClear)
 	{
-		SetRenderTarget(RHICmdList, VelocityTexture, DepthTexture, ESimpleRenderTargetMode::EExistingColorAndDepth, FExclusiveDepthStencil::DepthRead_StencilWrite);
-
 		// some platforms need the clear color when rendertargets transition to SRVs.  We propagate here to allow parallel rendering to always
 		// have the proper mapping when the RT is transitioned.
 		RHICmdList.BindClearMRTValues(true, false, false);
@@ -764,7 +769,6 @@ public:
 		: FParallelCommandListSet(GET_STATID(STAT_CLP_Velocity), InView, InSceneRenderer, InParentCmdList, bInParallelExecute, bInCreateSceneContext, InDrawRenderState)
 		, VelocityRT(InVelocityRT)
 	{
-		SetStateOnCommandList(ParentCmdList);
 	}
 
 	virtual ~FVelocityPassParallelCommandListSet()
@@ -787,6 +791,8 @@ static TAutoConsoleVariable<int32> CVarRHICmdFlushRenderThreadTasksVelocityPass(
 
 void FDeferredShadingSceneRenderer::RenderVelocitiesInnerParallel(FRHICommandListImmediate& RHICmdList, TRefCountPtr<IPooledRenderTarget>& VelocityRT)
 {
+	// Parallel rendering requires its own renderpasses so we cannot have an active one at this point
+	check(RHICmdList.IsOutsideRenderPass());
 	// parallel version
 	FScopedCommandListWaitForTasks Flusher(CVarRHICmdFlushRenderThreadTasksVelocityPass.GetValueOnRenderThread() > 0 || CVarRHICmdFlushRenderThreadTasks.GetValueOnRenderThread() > 0, RHICmdList);
 
@@ -848,6 +854,7 @@ void FDeferredShadingSceneRenderer::RenderVelocitiesInnerParallel(FRHICommandLis
 
 void FDeferredShadingSceneRenderer::RenderVelocitiesInner(FRHICommandListImmediate& RHICmdList, TRefCountPtr<IPooledRenderTarget>& VelocityRT)
 {
+	check(RHICmdList.IsInsideRenderPass());
 	for(int32 ViewIndex = 0; ViewIndex < Views.Num(); ViewIndex++)
 	{
 		const FViewInfo& View = Views[ViewIndex];
@@ -927,14 +934,21 @@ void FDeferredShadingSceneRenderer::RenderVelocities(FRHICommandListImmediate& R
 	{
 		// In this case, basepass also outputs some of the velocities, so append is already started, and don't clear the buffer.
 		BeginVelocityRendering(RHICmdList, VelocityRT, !FVelocityRendering::BasePassCanOutputVelocity(FeatureLevel));
+	}
 
+	{
 		if (IsParallelVelocity())
 		{
+			// This initial renderpass will just be a clear in the parallel case.
+			RHICmdList.EndRenderPass();
+
+			// Now do parallel encoding.
 			RenderVelocitiesInnerParallel(RHICmdList, VelocityRT);
 		}
 		else
 		{
 			RenderVelocitiesInner(RHICmdList, VelocityRT);
+			RHICmdList.EndRenderPass();
 		}
 
 		RHICmdList.CopyToResolveTarget(VelocityRT->GetRenderTargetItem().TargetableTexture, VelocityRT->GetRenderTargetItem().ShaderResourceTexture, FResolveParams());
