Vulkan requires stencil store (instead of don't store) or it is allowed to leave the attachment in an undefined state.

diff --git a/Engine/Source/Runtime/Renderer/Private/CompositionLighting/PostProcessAmbientOcclusion.cpp b/Engine/Source/Runtime/Renderer/Private/CompositionLighting/PostProcessAmbientOcclusion.cpp
--- a/Engine/Source/Runtime/Renderer/Private/CompositionLighting/PostProcessAmbientOcclusion.cpp
+++ b/Engine/Source/Runtime/Renderer/Private/CompositionLighting/PostProcessAmbientOcclusion.cpp
@@ -791,9 +791,10 @@ void FRCPassPostProcessAmbientOcclusion::ProcessPS(FRenderingCompositePassContex
 	if (bDepthBoundsTestEnabled)
 	{
 		// We'll use the depth/stencil buffer for read but it will not be modified.
+		// Note: VK requires us to store stencil or it (may) leave the attachment in an undefined state.
 		RPInfo.DepthStencilRenderTarget.DepthStencilTarget = SceneDepthBuffer->TargetableTexture;
-		RPInfo.DepthStencilRenderTarget.Action = MakeDepthStencilTargetActions(ERenderTargetActions::Load_Store, ERenderTargetActions::Load_Store);
-		RPInfo.DepthStencilRenderTarget.ExclusiveDepthStencil = FExclusiveDepthStencil::DepthRead_StencilRead;
+		RPInfo.DepthStencilRenderTarget.Action = MakeDepthStencilTargetActions(ERenderTargetActions::Load_DontStore, ERenderTargetActions::Load_Store);
+		RPInfo.DepthStencilRenderTarget.ExclusiveDepthStencil = FExclusiveDepthStencil::DepthRead_StencilWrite;
 	}
 
 	Context.RHICmdList.BeginRenderPass(RPInfo, TEXT("PSAmbientOcclusion"));
