Lost Emil's fix for Dbuffer decals + MSAA + Forward. Now restored.

diff --git a/Engine/Source/Runtime/Renderer/Private/CompositionLighting/PostProcessDeferredDecals.cpp b/Engine/Source/Runtime/Renderer/Private/CompositionLighting/PostProcessDeferredDecals.cpp
--- a/Engine/Source/Runtime/Renderer/Private/CompositionLighting/PostProcessDeferredDecals.cpp
+++ b/Engine/Source/Runtime/Renderer/Private/CompositionLighting/PostProcessDeferredDecals.cpp
@@ -1143,6 +1143,7 @@ void FDecalRenderTargetManager::SetRenderTargetMode(FDecalRenderingCommon::ERend
 	ERenderTargetActions DepthTargetActions = ERenderTargetActions::Load_DontStore;
 	uint32 NumColorTargets = 1;
 	FTextureRHIParamRef* TargetsToBind = &TargetsToResolve[0];
+	FTextureRHIParamRef DepthTarget = SceneContext.GetSceneDepthSurface();
 
 	// The SceneColorAndGBuffer modes do not actually need GBufferA bound when there's no normal.
 	// The apis based on renderpasses fill fail to actually bind anything past a null entry in their RT list so we have to bind it anyway.
@@ -1199,6 +1200,8 @@ void FDecalRenderTargetManager::SetRenderTargetMode(FDecalRenderingCommon::ERend
 			NumColorTargets = 4;
 		}
 
+		DepthTarget = SceneContext.GetSceneDepthTexture();
+
 		TargetsToBind = &TargetsToResolve[DBufferAIndex];
 		break;
 	}
@@ -1220,7 +1223,7 @@ void FDecalRenderTargetManager::SetRenderTargetMode(FDecalRenderingCommon::ERend
 	}
 
 	FRHIRenderPassInfo RPInfo(NumColorTargets, TargetsToBind, ERenderTargetActions::Load_Store);
-	RPInfo.DepthStencilRenderTarget.DepthStencilTarget = SceneContext.GetSceneDepthSurface();
+	RPInfo.DepthStencilRenderTarget.DepthStencilTarget = DepthTarget;
 	RPInfo.DepthStencilRenderTarget.Action = MakeDepthStencilTargetActions(DepthTargetActions, ERenderTargetActions::Load_Store);
 	RPInfo.DepthStencilRenderTarget.ExclusiveDepthStencil = DepthStencilAccess;
 
