Fixing a case where AO Decals render incorrectly.

diff --git a/Engine/Source/Runtime/Renderer/Private/CompositionLighting/PostProcessDeferredDecals.cpp b/Engine/Source/Runtime/Renderer/Private/CompositionLighting/PostProcessDeferredDecals.cpp
--- a/Engine/Source/Runtime/Renderer/Private/CompositionLighting/PostProcessDeferredDecals.cpp
+++ b/Engine/Source/Runtime/Renderer/Private/CompositionLighting/PostProcessDeferredDecals.cpp
@@ -1139,6 +1139,10 @@ void FDecalRenderTargetManager::SetRenderTargetMode(FDecalRenderingCommon::ERend
 		RHICmdList.TransitionResource(EResourceTransitionAccess::EWritable, SceneContext.GBufferA->GetRenderTargetItem().TargetableTexture);
 	}
 
+	// Color
+	ERenderTargetActions ColorTargetActions = ERenderTargetActions::Load_Store;
+
+	// Depth
 	FExclusiveDepthStencil DepthStencilAccess = FExclusiveDepthStencil::DepthRead_StencilWrite;
 	ERenderTargetActions DepthTargetActions = ERenderTargetActions::Load_DontStore;
 	uint32 NumColorTargets = 1;
@@ -1213,6 +1217,11 @@ void FDecalRenderTargetManager::SetRenderTargetMode(FDecalRenderingCommon::ERend
 
 		TargetsToBind = &TargetsToResolve[SceneColorIndex];
 
+		if (!SceneContext.bScreenSpaceAOIsValid)
+		{
+			ColorTargetActions = ERenderTargetActions::Clear_Store;
+		}
+
 		SceneContext.bScreenSpaceAOIsValid = true;
 		break;
 	}
@@ -1222,7 +1231,7 @@ void FDecalRenderTargetManager::SetRenderTargetMode(FDecalRenderingCommon::ERend
 		break;
 	}
 
-	FRHIRenderPassInfo RPInfo(NumColorTargets, TargetsToBind, ERenderTargetActions::Load_Store);
+	FRHIRenderPassInfo RPInfo(NumColorTargets, TargetsToBind, ColorTargetActions);
 	RPInfo.DepthStencilRenderTarget.DepthStencilTarget = DepthTarget;
 	RPInfo.DepthStencilRenderTarget.Action = MakeDepthStencilTargetActions(DepthTargetActions, ERenderTargetActions::Load_Store);
 	RPInfo.DepthStencilRenderTarget.ExclusiveDepthStencil = DepthStencilAccess;
