Fixing a couple more mismatched exclusive depth tags.

diff --git a/Engine/Source/Runtime/Renderer/Private/MobileTranslucentRendering.cpp b/Engine/Source/Runtime/Renderer/Private/MobileTranslucentRendering.cpp
--- a/Engine/Source/Runtime/Renderer/Private/MobileTranslucentRendering.cpp
+++ b/Engine/Source/Runtime/Renderer/Private/MobileTranslucentRendering.cpp
@@ -340,7 +340,7 @@ void FMobileSceneRenderer::RenderTranslucency(FRHICommandListImmediate& RHICmdLi
 				FSceneRenderTargets& SceneContext = FSceneRenderTargets::Get(RHICmdList);
 				// Use begin rendering scene color with FExclusiveDepthStencil::DepthRead_StencilRead to avoid starting a new render pass on vulkan.
 				// #todo-renderpasses we'll need to clean this up once we verify VK won't trash the stencil buffer if we make it DontStore
-				SceneContext.BeginRenderingSceneColor(RHICmdList, ESimpleRenderTargetMode::EExistingColorAndDepth, FExclusiveDepthStencil::DepthWrite_StencilWrite);
+				SceneContext.BeginRenderingSceneColor(RHICmdList, ESimpleRenderTargetMode::EExistingColorAndDepth, FExclusiveDepthStencil::DepthRead_StencilWrite);
 			}
 			else
 			{
@@ -352,7 +352,7 @@ void FMobileSceneRenderer::RenderTranslucency(FRHICommandListImmediate& RHICmdLi
 				FRHIRenderPassInfo RPInfo(SceneColor, ERenderTargetActions::Load_Store);
 				RPInfo.DepthStencilRenderTarget.Action = EDepthStencilTargetActions::LoadDepthStencil_StoreDepthStencil;
 				RPInfo.DepthStencilRenderTarget.DepthStencilTarget = SceneDepth;
-				RPInfo.DepthStencilRenderTarget.ExclusiveDepthStencil = FExclusiveDepthStencil::DepthWrite_StencilWrite;
+				RPInfo.DepthStencilRenderTarget.ExclusiveDepthStencil = FExclusiveDepthStencil::DepthRead_StencilWrite;
 
 				RHICmdList.BeginRenderPass(RPInfo, TEXT("RenderMobileTranslucency"));
 			}
diff --git a/Engine/Source/Runtime/Renderer/Private/VelocityRendering.cpp b/Engine/Source/Runtime/Renderer/Private/VelocityRendering.cpp
--- a/Engine/Source/Runtime/Renderer/Private/VelocityRendering.cpp
+++ b/Engine/Source/Runtime/Renderer/Private/VelocityRendering.cpp
@@ -692,7 +692,7 @@ static void BeginVelocityRendering(FRHICommandList& RHICmdList, TRefCountPtr<IPo
 	FRHIRenderPassInfo RPInfo(VelocityTexture, ERenderTargetActions::Load_Store);
 	RPInfo.DepthStencilRenderTarget.Action = MakeDepthStencilTargetActions(ERenderTargetActions::Load_Store, ERenderTargetActions::Load_Store);
 	RPInfo.DepthStencilRenderTarget.DepthStencilTarget = DepthTexture;
-	RPInfo.DepthStencilRenderTarget.ExclusiveDepthStencil = FExclusiveDepthStencil::DepthWrite_StencilWrite;
+	RPInfo.DepthStencilRenderTarget.ExclusiveDepthStencil = FExclusiveDepthStencil::DepthRead_StencilWrite;
 
 	if (bPerformClear)
 	{
