Fixing lightshaft assert.
Cleaned up naming so it's more clear what ends renderpasses.

diff --git a/Engine/Source/Runtime/Renderer/Private/LightShaftRendering.cpp b/Engine/Source/Runtime/Renderer/Private/LightShaftRendering.cpp
--- a/Engine/Source/Runtime/Renderer/Private/LightShaftRendering.cpp
+++ b/Engine/Source/Runtime/Renderer/Private/LightShaftRendering.cpp
@@ -992,7 +992,8 @@ void ApplyLightShaftBloom(FRHICommandListImmediate& RHICmdList, const FViewInfo&
 
 	if (bUseSeparateTranslucency)
 	{
-		SceneContext.FinishRenderingSeparateTranslucency(RHICmdList, View);
+		RHICmdList.EndRenderPass();
+		SceneContext.ResolveSeparateTranslucency(RHICmdList, View);
 	}
 	else
 	{
diff --git a/Engine/Source/Runtime/Renderer/Private/PostProcess/SceneRenderTargets.cpp b/Engine/Source/Runtime/Renderer/Private/PostProcess/SceneRenderTargets.cpp
--- a/Engine/Source/Runtime/Renderer/Private/PostProcess/SceneRenderTargets.cpp
+++ b/Engine/Source/Runtime/Renderer/Private/PostProcess/SceneRenderTargets.cpp
@@ -1622,9 +1622,9 @@ void FSceneRenderTargets::BeginRenderingSeparateTranslucency(FRHICommandList& RH
 	}
 }
 
-void FSceneRenderTargets::FinishRenderingSeparateTranslucency(FRHICommandList& RHICmdList, const FViewInfo& View)
+void FSceneRenderTargets::ResolveSeparateTranslucency(FRHICommandList& RHICmdList, const FViewInfo& View)
 {
-	SCOPED_DRAW_EVENT(RHICmdList, FinishSeparateTranslucency);
+	SCOPED_DRAW_EVENT(RHICmdList, ResolveSeparateTranslucency);
 
 	check(RHICmdList.IsOutsideRenderPass());
 
diff --git a/Engine/Source/Runtime/Renderer/Private/PostProcess/SceneRenderTargets.h b/Engine/Source/Runtime/Renderer/Private/PostProcess/SceneRenderTargets.h
--- a/Engine/Source/Runtime/Renderer/Private/PostProcess/SceneRenderTargets.h
+++ b/Engine/Source/Runtime/Renderer/Private/PostProcess/SceneRenderTargets.h
@@ -237,7 +237,7 @@ public:
 
 	/** Begin rendering translucency in a separate (offscreen) buffer. This can be any translucency pass. */
 	void BeginRenderingSeparateTranslucency(FRHICommandList& RHICmdList, const FViewInfo& View, const FSceneRenderer& Renderer, bool bFirstTimeThisFrame);
-	void FinishRenderingSeparateTranslucency(FRHICommandList& RHICmdList, const FViewInfo& View);
+	void ResolveSeparateTranslucency(FRHICommandList& RHICmdList, const FViewInfo& View);
 
 	void FreeSeparateTranslucency()
 	{
diff --git a/Engine/Source/Runtime/Renderer/Private/TranslucentRendering.cpp b/Engine/Source/Runtime/Renderer/Private/TranslucentRendering.cpp
--- a/Engine/Source/Runtime/Renderer/Private/TranslucentRendering.cpp
+++ b/Engine/Source/Runtime/Renderer/Private/TranslucentRendering.cpp
@@ -1456,7 +1456,7 @@ void FDeferredShadingSceneRenderer::RenderTranslucency(FRHICommandListImmediate&
 				RHICmdList.EndRenderPass();
 			}
 
-			SceneContext.FinishRenderingSeparateTranslucency(RHICmdList, View);
+			SceneContext.ResolveSeparateTranslucency(RHICmdList, View);
 
 			if (TranslucencyPass == ETranslucencyPass::TPT_TranslucencyAfterDOF)
 			{
