Need to end the renderpass when we are done rendering.

Fixes CopyCaptureToTarget.

diff --git a/Engine/Source/Runtime/Renderer/Private/MobileTranslucentRendering.cpp b/Engine/Source/Runtime/Renderer/Private/MobileTranslucentRendering.cpp
--- a/Engine/Source/Runtime/Renderer/Private/MobileTranslucentRendering.cpp
+++ b/Engine/Source/Runtime/Renderer/Private/MobileTranslucentRendering.cpp
@@ -736,6 +736,7 @@ bool FMobileSceneRenderer::RenderInverseOpacity(FRHICommandListImmediate& RHICmd
 
 			if (!bGammaSpace)
 			{
+				RHICmdList.EndRenderPass();
 				SceneContext.FinishRenderingTranslucency(RHICmdList);
 			}
 			else
