Fixing MeshUV density debug mode.
Changed the post process DrawCustom path to require its own renderpass.

diff --git a/Engine/Source/Runtime/Renderer/Private/PostProcess/PostProcessStreamingAccuracyLegend.cpp b/Engine/Source/Runtime/Renderer/Private/PostProcess/PostProcessStreamingAccuracyLegend.cpp
--- a/Engine/Source/Runtime/Renderer/Private/PostProcess/PostProcessStreamingAccuracyLegend.cpp
+++ b/Engine/Source/Runtime/Renderer/Private/PostProcess/PostProcessStreamingAccuracyLegend.cpp
@@ -37,6 +37,7 @@ void FRCPassPostProcessStreamingAccuracyLegend::DrawCheckerBoard(FCanvas& Canvas
 
 void FRCPassPostProcessStreamingAccuracyLegend::DrawCustom(FRenderingCompositePassContext& Context)
 {
+	check(Context.RHICmdList.IsOutsideRenderPass());
 	if (Colors.Num() != NumStreamingAccuracyColors) return;
 
 	SCOPED_DRAW_EVENT(Context.RHICmdList, PostProcessStreamingAccuracyLegend);
