Fixed volumetric translucent shadow parallel rendering. Also removed its thread task and all its special cases, as PrepareTranslucentShadowMap no longer renders translucent shadow map - it just returns an already prepared shadow map.

diff --git a/Engine/Source/Runtime/Renderer/Private/DeferredShadingRenderer.h b/Engine/Source/Runtime/Renderer/Private/DeferredShadingRenderer.h
--- a/Engine/Source/Runtime/Renderer/Private/DeferredShadingRenderer.h
+++ b/Engine/Source/Runtime/Renderer/Private/DeferredShadingRenderer.h
@@ -265,8 +265,8 @@ private:
 
 	void RenderLightShaftBloom(FRHICommandListImmediate& RHICmdList);
 
-	/** Reuses an existing translucent shadow map if possible or re-renders one if necessary. */
-	const FProjectedShadowInfo* PrepareTranslucentShadowMap(FRHICommandList& RHICmdList, const FViewInfo& View, FPrimitiveSceneInfo* PrimitiveSceneInfo, ETranslucencyPass::Type TranslucenyPassType);
+	/** Returns an existing translucent shadow map for a given primitive. */
+	const FProjectedShadowInfo* GetTranslucentShadowMap(FPrimitiveSceneInfo* PrimitiveSceneInfo, ETranslucencyPass::Type TranslucenyPassType);
 
 	bool ShouldRenderVelocities() const;
 
diff --git a/Engine/Source/Runtime/Renderer/Private/SceneRendering.h b/Engine/Source/Runtime/Renderer/Private/SceneRendering.h
--- a/Engine/Source/Runtime/Renderer/Private/SceneRendering.h
+++ b/Engine/Source/Runtime/Renderer/Private/SceneRendering.h
@@ -367,15 +367,6 @@ public:
 	*/
 	void DrawPrimitivesParallel(FRHICommandList& RHICmdList, const class FViewInfo& View, const FDrawingPolicyRenderState& DrawRenderState, class FDeferredShadingSceneRenderer& Renderer, ETranslucencyPass::Type TranslucenyPassType, int32 FirstPrimIdx, int32 LastPrimIdx) const;
 
-	/**
-	* Draw a single primitive...this is used when we are rendering in parallel and we need to handlke a translucent shadow
-	* @param View - current view used to draw items
-	* @param PhaseSortedPrimitives - array with the primitives we want to draw
-	* @param TranslucenyPassType
-	* @param PrimIdx in SortedPrims[]
-	*/
-	void DrawAPrimitive(FRHICommandList& RHICmdList, const class FViewInfo& View, const FDrawingPolicyRenderState& DrawRenderState, class FDeferredShadingSceneRenderer& Renderer, ETranslucencyPass::Type TranslucenyPassType, int32 PrimIdx) const;
-
 	/** 
 	* Draw all the primitives in this set for the mobile pipeline. 
 	*/
diff --git a/Engine/Source/Runtime/Renderer/Private/TranslucentRendering.cpp b/Engine/Source/Runtime/Renderer/Private/TranslucentRendering.cpp
--- a/Engine/Source/Runtime/Renderer/Private/TranslucentRendering.cpp
+++ b/Engine/Source/Runtime/Renderer/Private/TranslucentRendering.cpp
@@ -257,9 +257,9 @@ void FDeferredShadingSceneRenderer::EndTimingSeparateTranslucencyPass(FRHIComman
 	}
 }
 
-const FProjectedShadowInfo* FDeferredShadingSceneRenderer::PrepareTranslucentShadowMap(FRHICommandList& RHICmdList, const FViewInfo& View, FPrimitiveSceneInfo* PrimitiveSceneInfo, ETranslucencyPass::Type TranslucencyPass)
+const FProjectedShadowInfo* FDeferredShadingSceneRenderer::GetTranslucentShadowMap(FPrimitiveSceneInfo* PrimitiveSceneInfo, ETranslucencyPass::Type TranslucencyPass)
 {
-	QUICK_SCOPE_CYCLE_COUNTER(STAT_FDeferredShadingSceneRenderer_PrepareTranslucentShadowMap);
+	QUICK_SCOPE_CYCLE_COUNTER(STAT_FDeferredShadingSceneRenderer_GetTranslucentShadowMap);
 	const FVisibleLightInfo* VisibleLightInfo = NULL;
 	FProjectedShadowInfo* TranslucentSelfShadow = NULL;
 
@@ -615,82 +615,6 @@ bool FTranslucencyDrawingPolicyFactory::DrawStaticMesh(
 FTranslucentPrimSet
 -----------------------------------------------------------------------------*/
 
-void FTranslucentPrimSet::DrawAPrimitive(
-	FRHICommandList& RHICmdList,
-	const FViewInfo& View,
-	const FDrawingPolicyRenderState& DrawRenderState,
-	FDeferredShadingSceneRenderer& Renderer,
-	ETranslucencyPass::Type TranslucencyPass,
-	int32 PrimIdx
-	) const
-{
-	FPrimitiveSceneInfo* PrimitiveSceneInfo = SortedPrims[PrimIdx].PrimitiveSceneInfo;
-	int32 PrimitiveId = PrimitiveSceneInfo->GetIndex();
-	const FPrimitiveViewRelevance& ViewRelevance = View.PrimitiveViewRelevanceMap[PrimitiveId];
-
-	checkSlow(ViewRelevance.HasTranslucency());
-
-	const FProjectedShadowInfo* TranslucentSelfShadow = Renderer.PrepareTranslucentShadowMap(RHICmdList, View, PrimitiveSceneInfo, TranslucencyPass);
-
-	RenderPrimitive(RHICmdList, View, DrawRenderState, PrimitiveSceneInfo, ViewRelevance, TranslucentSelfShadow, TranslucencyPass);
-}
-
-class FVolumetricTranslucentShadowRenderThreadTask
-{
-	FRHICommandList& RHICmdList;
-	const FTranslucentPrimSet &PrimSet;
-	const FViewInfo& View;
-	FDrawingPolicyRenderState DrawRenderState;
-	FDeferredShadingSceneRenderer& Renderer;
-	ETranslucencyPass::Type TranslucencyPass;
-	int32 Index;
-
-public:
-
-	FORCEINLINE_DEBUGGABLE FVolumetricTranslucentShadowRenderThreadTask(FRHICommandList& InRHICmdList, const FTranslucentPrimSet& InPrimSet, const FViewInfo& InView, const FDrawingPolicyRenderState& InDrawRenderState, FDeferredShadingSceneRenderer& InRenderer, ETranslucencyPass::Type InTranslucencyPass, int32 InIndex)
-		: RHICmdList(InRHICmdList)
-		, PrimSet(InPrimSet)
-		, View(InView)
-		, DrawRenderState(InDrawRenderState)
-		, Renderer(InRenderer)
-		, TranslucencyPass(InTranslucencyPass)
-		, Index(InIndex)
-	{
-	}
-
-	FORCEINLINE TStatId GetStatId() const
-	{
-		RETURN_QUICK_DECLARE_CYCLE_STAT(FVolumetricTranslucentShadowRenderThreadTask, STATGROUP_TaskGraphTasks);
-	}
-
-	ENamedThreads::Type GetDesiredThread()
-	{
-		return ENamedThreads::GetRenderThread_Local();
-	}
-
-	static ESubsequentsMode::Type GetSubsequentsMode() { return ESubsequentsMode::TrackSubsequents; }
-
-	void DoTask(ENamedThreads::Type CurrentThread, const FGraphEventRef& MyCompletionGraphEvent)
-	{
-		check(RHICmdList.IsOutsideRenderPass());
-		// Never needs clear here as it is already done in RenderTranslucency.
-		FSceneRenderTargets& SceneContext = FSceneRenderTargets::Get(RHICmdList);
-		if (SceneContext.IsSeparateTranslucencyPass())
-		{
-			SceneContext.BeginRenderingSeparateTranslucency(RHICmdList, View, Renderer, false);
-		}
-		else
-		{
-			SceneContext.BeginRenderingTranslucency(RHICmdList, View, Renderer, false);
-		}
-
-		PrimSet.DrawAPrimitive(RHICmdList, View, DrawRenderState, Renderer, TranslucencyPass, Index);
-
-		RHICmdList.EndRenderPass();
-		RHICmdList.HandleRTThreadTaskCompletion(MyCompletionGraphEvent);
-	}
-};
-
 void FTranslucentPrimSet::DrawPrimitivesParallel(
 	FRHICommandList& RHICmdList,
 	const FViewInfo& View,
@@ -700,7 +624,6 @@ void FTranslucentPrimSet::DrawPrimitivesParallel(
 	int32 FirstPrimIdx, int32 LastPrimIdx
 	) const
 {
-	TArray<int32> Volumetrics;
 	// Draw sorted scene prims
 	for (int32 PrimIdx = FirstPrimIdx; PrimIdx <= LastPrimIdx; PrimIdx++)
 	{
@@ -710,47 +633,12 @@ void FTranslucentPrimSet::DrawPrimitivesParallel(
 
 		checkSlow(ViewRelevance.HasTranslucency());
 
-		if (PrimitiveSceneInfo->Proxy && PrimitiveSceneInfo->Proxy->CastsVolumetricTranslucentShadow())
-		{
-			Volumetrics.Add(PrimIdx);
-			check(!IsInActualRenderingThread());
-			// can't do this in parallel, defer
-			//FRHICommandList* CmdList = new FRHICommandList(View.GPUMask);
-			//CmdList->CopyRenderThreadContexts(RHICmdList);
-			//FGraphEventRef RenderThreadCompletionEvent = TGraphTask<FVolumetricTranslucentShadowRenderThreadTask>::CreateTask().ConstructAndDispatchWhenReady(*CmdList, *this, View, DrawRenderState, Renderer, TranslucencyPass, PrimIdx);
-			//RHICmdList.QueueRenderThreadCommandListSubmit(RenderThreadCompletionEvent, CmdList);
-		}
-		else
-		{
-			RenderPrimitive(RHICmdList, View, DrawRenderState, PrimitiveSceneInfo, ViewRelevance, nullptr, TranslucencyPass);
-		}
+		const FProjectedShadowInfo* TranslucentSelfShadow = Renderer.GetTranslucentShadowMap(PrimitiveSceneInfo, TranslucencyPass);
+
+		RenderPrimitive(RHICmdList, View, DrawRenderState, PrimitiveSceneInfo, ViewRelevance, TranslucentSelfShadow, TranslucencyPass);
 	}
 
 	RHICmdList.EndRenderPass();
-
-	if (Volumetrics.Num() > 0)
-	{
-		FRHICommandList* CmdList = new FRHICommandList(View.GPUMask);
-		// #todo-renderpasses Do we actually need the context here?
-		CmdList->CopyRenderThreadContexts(RHICmdList);
-		FSceneRenderTargets& SceneContext = FSceneRenderTargets::Get(RHICmdList);
-		if (SceneContext.IsSeparateTranslucencyPass())
-		{
-			SceneContext.BeginRenderingSeparateTranslucency(*CmdList, View, Renderer, false);
-		}
-		else
-		{
-			SceneContext.BeginRenderingTranslucency(*CmdList, View, Renderer, false);
-		}
-
-		for (int32 Index = 0; Index < Volumetrics.Num(); Index++)
-		{
-			DrawAPrimitive(*CmdList, View, DrawRenderState, Renderer, TranslucencyPass, Index);
-		}
-		CmdList->EndRenderPass();
-
-		RHICmdList.QueueCommandListSubmit(CmdList);
-	}
 }
 
 void FTranslucentPrimSet::DrawPrimitives(
@@ -772,7 +660,7 @@ void FTranslucentPrimSet::DrawPrimitives(
 
 		checkSlow(ViewRelevance.HasTranslucency());
 			
-		const FProjectedShadowInfo* TranslucentSelfShadow = Renderer.PrepareTranslucentShadowMap(RHICmdList, View, PrimitiveSceneInfo, TranslucencyPass);
+		const FProjectedShadowInfo* TranslucentSelfShadow = Renderer.GetTranslucentShadowMap(PrimitiveSceneInfo, TranslucencyPass);
 
 		RenderPrimitive(RHICmdList, View, DrawRenderState, PrimitiveSceneInfo, ViewRelevance, TranslucentSelfShadow, TranslucencyPass);
 	}
