Foreground DPG now works for SimpleElementCollector which is used by many engine debug utilities like DrawDebugSphere.  Foreground primitives first render with depth test off, depth write on to effectively clear the depth buffer for the pixels they affect.  Then they are rendered a second time with depth testing on to achieve the correct result.

diff --git a/Engine/Source/Runtime/Engine/Private/SceneManagement.cpp b/Engine/Source/Runtime/Engine/Private/SceneManagement.cpp
--- a/Engine/Source/Runtime/Engine/Private/SceneManagement.cpp
+++ b/Engine/Source/Runtime/Engine/Private/SceneManagement.cpp
@@ -101,7 +101,9 @@ void FSimpleElementCollector::DrawSprite(
 	uint8 BlendMode
 	)
 {
-	BatchedElements.AddSprite(
+	FBatchedElements& Elements = DepthPriorityGroup == SDPG_World ? BatchedElements : TopBatchedElements;
+
+	Elements.AddSprite(
 		Position,
 		SizeX,
 		SizeY,
@@ -126,7 +128,9 @@ void FSimpleElementCollector::DrawLine(
 	bool bScreenSpace/* = false*/
 	)
 {
-	BatchedElements.AddLine(
+	FBatchedElements& Elements = DepthPriorityGroup == SDPG_World ? BatchedElements : TopBatchedElements;
+
+	Elements.AddLine(
 		Start,
 		End,
 		Color,
@@ -144,7 +148,9 @@ void FSimpleElementCollector::DrawPoint(
 	uint8 DepthPriorityGroup
 	)
 {
-	BatchedElements.AddPoint(
+	FBatchedElements& Elements = DepthPriorityGroup == SDPG_World ? BatchedElements : TopBatchedElements;
+
+	Elements.AddPoint(
 		Position,
 		PointSize,
 		Color,
@@ -161,13 +167,15 @@ void FSimpleElementCollector::RegisterDynamicResource(FDynamicPrimitiveResource*
 	DynamicResource->InitPrimitiveResource();
 }
 
-void FSimpleElementCollector::DrawBatchedElements(FRHICommandList& RHICmdList, const FDrawingPolicyRenderState& DrawRenderState, const FSceneView& InView, EBlendModeFilter::Type Filter) const
+void FSimpleElementCollector::DrawBatchedElements(FRHICommandList& RHICmdList, const FDrawingPolicyRenderState& DrawRenderState, const FSceneView& InView, EBlendModeFilter::Type Filter, ESceneDepthPriorityGroup DepthPriorityGroup) const
 {
 	// Mobile HDR does not execute post process, so does not need to render flipped
 	const bool bNeedToSwitchVerticalAxis = RHINeedsToSwitchVerticalAxis(InView.GetShaderPlatform()) && !bIsMobileHDR;
 
+	const FBatchedElements& Elements = DepthPriorityGroup == SDPG_World ? BatchedElements : TopBatchedElements;
+
 	// Draw the batched elements.
-	BatchedElements.Draw(
+	Elements.Draw(
 		RHICmdList,
 		DrawRenderState,
 		InView.GetFeatureLevel(),
diff --git a/Engine/Source/Runtime/Engine/Public/SceneManagement.h b/Engine/Source/Runtime/Engine/Public/SceneManagement.h
--- a/Engine/Source/Runtime/Engine/Public/SceneManagement.h
+++ b/Engine/Source/Runtime/Engine/Public/SceneManagement.h
@@ -1678,10 +1678,21 @@ public:
 		return 0;
 	}
 
-	void DrawBatchedElements(FRHICommandList& RHICmdList, const FDrawingPolicyRenderState& DrawRenderState, const FSceneView& InView, EBlendModeFilter::Type Filter) const;
+	void DrawBatchedElements(FRHICommandList& RHICmdList, const FDrawingPolicyRenderState& DrawRenderState, const FSceneView& InView, EBlendModeFilter::Type Filter, ESceneDepthPriorityGroup DPG) const;
+
+	bool HasPrimitives(ESceneDepthPriorityGroup DPG) const
+	{
+		if (DPG == SDPG_World)
+		{
+			return BatchedElements.HasPrimsToDraw();
+		}
+
+		return TopBatchedElements.HasPrimsToDraw();
+	}
 
 	/** The batched simple elements. */
 	FBatchedElements BatchedElements;
+	FBatchedElements TopBatchedElements;
 
 private:
 
diff --git a/Engine/Source/Runtime/Renderer/Private/BasePassRendering.cpp b/Engine/Source/Runtime/Renderer/Private/BasePassRendering.cpp
--- a/Engine/Source/Runtime/Renderer/Private/BasePassRendering.cpp
+++ b/Engine/Source/Runtime/Renderer/Private/BasePassRendering.cpp
@@ -1285,12 +1285,58 @@ void FDeferredShadingSceneRenderer::RenderBasePassViewParallel(FViewInfo& View,
 	RenderBasePassDynamicDataParallel(ParallelSet);
 }
 
-void FDeferredShadingSceneRenderer::RenderEditorPrimitives(FRHICommandList& RHICmdList, const FViewInfo& View, FExclusiveDepthStencil::Type BasePassDepthStencilAccess, const FDrawingPolicyRenderState& InDrawRenderState, bool& bOutDirty) 
+bool HasEditorPrimitivesForDPG(const FViewInfo& View, ESceneDepthPriorityGroup DepthPriorityGroup)
+{
+	bool bHasPrimitives = View.SimpleElementCollector.HasPrimitives(DepthPriorityGroup);
+
+	if (!View.Family->EngineShowFlags.CompositeEditorPrimitives)
+	{
+		const TIndirectArray<FMeshBatch>& ViewMeshElementList = (DepthPriorityGroup == SDPG_Foreground ? View.TopViewMeshElements : View.ViewMeshElements);
+		bHasPrimitives |= ViewMeshElementList.Num() > 0;
+
+		const FBatchedElements& BatchedViewElements = DepthPriorityGroup == SDPG_World ? View.BatchedViewElements : View.TopBatchedViewElements;
+		bHasPrimitives |= BatchedViewElements.HasPrimsToDraw();
+	}
+
+	return bHasPrimitives;
+}
+
+void FDeferredShadingSceneRenderer::RenderEditorPrimitives(FRHICommandListImmediate& RHICmdList, const FViewInfo& View, FExclusiveDepthStencil::Type BasePassDepthStencilAccess, const FDrawingPolicyRenderState& InDrawRenderState, bool& bOutDirty)
 {
 	FDrawingPolicyRenderState DrawRenderState(InDrawRenderState);
 	SetupBasePassView(RHICmdList, View, this, DrawRenderState, BasePassDepthStencilAccess, ViewFamily.EngineShowFlags.ShaderComplexity, true);
 
-	View.SimpleElementCollector.DrawBatchedElements(RHICmdList, DrawRenderState, View, EBlendModeFilter::OpaqueAndMasked);
+	RenderEditorPrimitivesForDPG(RHICmdList, View, BasePassDepthStencilAccess, DrawRenderState, SDPG_World, bOutDirty);
+
+	if (HasEditorPrimitivesForDPG(View, SDPG_Foreground))
+	{
+		RHICmdList.EndRenderPass();
+
+		// Write foreground primitives into depth buffer without testing 
+		{
+			// Change to depth writable
+			FSceneRenderTargets& SceneContext = FSceneRenderTargets::Get(RHICmdList);
+			SceneContext.BeginRenderingGBuffer(RHICmdList, ERenderTargetLoadAction::ELoad, ERenderTargetLoadAction::ELoad, FExclusiveDepthStencil::DepthWrite_StencilWrite, false);
+
+			FDrawingPolicyRenderState NoDepthTestDrawRenderState(DrawRenderState);
+			NoDepthTestDrawRenderState.SetDepthStencilState(TStaticDepthStencilState<true, CF_Always>::GetRHI());
+			NoDepthTestDrawRenderState.SetDepthStencilAccess(FExclusiveDepthStencil::DepthWrite_StencilWrite);
+			RenderEditorPrimitivesForDPG(RHICmdList, View, BasePassDepthStencilAccess, NoDepthTestDrawRenderState, SDPG_Foreground, bOutDirty);
+
+			RHICmdList.EndRenderPass();
+
+			// Restore default base pass depth access
+			SceneContext.BeginRenderingGBuffer(RHICmdList, ERenderTargetLoadAction::ELoad, ERenderTargetLoadAction::ELoad, BasePassDepthStencilAccess, false);
+		}
+
+		// Render foreground primitives with depth testing
+		RenderEditorPrimitivesForDPG(RHICmdList, View, BasePassDepthStencilAccess, DrawRenderState, SDPG_Foreground, bOutDirty);
+	}
+}
+
+void FDeferredShadingSceneRenderer::RenderEditorPrimitivesForDPG(FRHICommandList& RHICmdList, const FViewInfo& View, FExclusiveDepthStencil::Type BasePassDepthStencilAccess, const FDrawingPolicyRenderState& InDrawRenderState, ESceneDepthPriorityGroup DepthPriorityGroup, bool& bOutDirty) 
+{
+	View.SimpleElementCollector.DrawBatchedElements(RHICmdList, InDrawRenderState, View, EBlendModeFilter::OpaqueAndMasked, DepthPriorityGroup);
 
 	bool bDirty = false;
 	if (!View.Family->EngineShowFlags.CompositeEditorPrimitives)
@@ -1298,17 +1344,12 @@ void FDeferredShadingSceneRenderer::RenderEditorPrimitives(FRHICommandList& RHIC
 		const bool bNeedToSwitchVerticalAxis = RHINeedsToSwitchVerticalAxis(ShaderPlatform);
 
 		// Draw the base pass for the view's batched mesh elements.
-		bDirty |= DrawViewElements<FBasePassOpaqueDrawingPolicyFactory>(RHICmdList, View, DrawRenderState, FBasePassOpaqueDrawingPolicyFactory::ContextType(), SDPG_World, true) || bDirty;
+		bDirty |= DrawViewElements<FBasePassOpaqueDrawingPolicyFactory>(RHICmdList, View, InDrawRenderState, FBasePassOpaqueDrawingPolicyFactory::ContextType(), DepthPriorityGroup, true) || bDirty;
 
-		// Draw the view's batched simple elements(lines, sprites, etc).
-		bDirty |= View.BatchedViewElements.Draw(RHICmdList, DrawRenderState, FeatureLevel, bNeedToSwitchVerticalAxis, View, false) || bDirty;
-
-		// Draw foreground objects last
-		bDirty |= DrawViewElements<FBasePassOpaqueDrawingPolicyFactory>(RHICmdList, View, DrawRenderState, FBasePassOpaqueDrawingPolicyFactory::ContextType(), SDPG_Foreground, true) || bDirty;
+		const FBatchedElements& BatchedViewElements = DepthPriorityGroup == SDPG_World ? View.BatchedViewElements : View.TopBatchedViewElements;
 
 		// Draw the view's batched simple elements(lines, sprites, etc).
-		bDirty |= View.TopBatchedViewElements.Draw(RHICmdList, DrawRenderState, FeatureLevel, bNeedToSwitchVerticalAxis, View, false) || bDirty;
-
+		bDirty |= BatchedViewElements.Draw(RHICmdList, InDrawRenderState, FeatureLevel, bNeedToSwitchVerticalAxis, View, false) || bDirty;
 	}
 
 	if (bDirty)
diff --git a/Engine/Source/Runtime/Renderer/Private/DeferredShadingRenderer.h b/Engine/Source/Runtime/Renderer/Private/DeferredShadingRenderer.h
--- a/Engine/Source/Runtime/Renderer/Private/DeferredShadingRenderer.h
+++ b/Engine/Source/Runtime/Renderer/Private/DeferredShadingRenderer.h
@@ -102,7 +102,10 @@ public:
 	bool RenderBasePassView(FRHICommandListImmediate& RHICmdList, FViewInfo& View, FExclusiveDepthStencil::Type BasePassDepthStencilAccess, const FDrawingPolicyRenderState& InDrawRenderState);
 
 	/** Renders editor primitives for a given View. */
-	void RenderEditorPrimitives(FRHICommandList& RHICmdList, const FViewInfo& View, FExclusiveDepthStencil::Type BasePassDepthStencilAccess, const FDrawingPolicyRenderState& DrawRenderState, bool& bOutDirty);
+	void RenderEditorPrimitives(FRHICommandListImmediate& RHICmdList, const FViewInfo& View, FExclusiveDepthStencil::Type BasePassDepthStencilAccess, const FDrawingPolicyRenderState& DrawRenderState, bool& bOutDirty);
+
+	/** Renders editor primitives for a given View. */
+	void RenderEditorPrimitivesForDPG(FRHICommandList& RHICmdList, const FViewInfo& View, FExclusiveDepthStencil::Type BasePassDepthStencilAccess, const FDrawingPolicyRenderState& DrawRenderState, ESceneDepthPriorityGroup DepthPriorityGroup, bool& bOutDirty);
 
 	/** 
 	* Renders the scene's base pass 
diff --git a/Engine/Source/Runtime/Renderer/Private/DepthRendering.cpp b/Engine/Source/Runtime/Renderer/Private/DepthRendering.cpp
--- a/Engine/Source/Runtime/Renderer/Private/DepthRendering.cpp
+++ b/Engine/Source/Runtime/Renderer/Private/DepthRendering.cpp
@@ -1222,7 +1222,8 @@ void FDeferredShadingSceneRenderer::RenderPrePassEditorPrimitives(FRHICommandLis
 {
 	SetupPrePassView(RHICmdList, View, this, true);
 
-	View.SimpleElementCollector.DrawBatchedElements(RHICmdList, DrawRenderState, View, EBlendModeFilter::OpaqueAndMasked);
+	View.SimpleElementCollector.DrawBatchedElements(RHICmdList, DrawRenderState, View, EBlendModeFilter::OpaqueAndMasked, SDPG_World);
+	View.SimpleElementCollector.DrawBatchedElements(RHICmdList, DrawRenderState, View, EBlendModeFilter::OpaqueAndMasked, SDPG_Foreground);
 
 	bool bDirty = false;
 	if (!View.Family->EngineShowFlags.CompositeEditorPrimitives)
diff --git a/Engine/Source/Runtime/Renderer/Private/MobileBasePassRendering.cpp b/Engine/Source/Runtime/Renderer/Private/MobileBasePassRendering.cpp
--- a/Engine/Source/Runtime/Renderer/Private/MobileBasePassRendering.cpp
+++ b/Engine/Source/Runtime/Renderer/Private/MobileBasePassRendering.cpp
@@ -657,7 +657,8 @@ void FMobileSceneRenderer::RenderMobileEditorPrimitives(FRHICommandList& RHICmdL
 	QUICK_SCOPE_CYCLE_COUNTER(STAT_EditorDynamicPrimitiveDrawTime);
 	SCOPED_DRAW_EVENT(RHICmdList, DynamicEd);
 
-	View.SimpleElementCollector.DrawBatchedElements(RHICmdList, DrawRenderState, View, EBlendModeFilter::OpaqueAndMasked);
+	View.SimpleElementCollector.DrawBatchedElements(RHICmdList, DrawRenderState, View, EBlendModeFilter::OpaqueAndMasked, SDPG_World);
+	View.SimpleElementCollector.DrawBatchedElements(RHICmdList, DrawRenderState, View, EBlendModeFilter::OpaqueAndMasked, SDPG_Foreground);
 
 	if (!View.Family->EngineShowFlags.CompositeEditorPrimitives)
 	{
diff --git a/Engine/Source/Runtime/Renderer/Private/MobileTranslucentRendering.cpp b/Engine/Source/Runtime/Renderer/Private/MobileTranslucentRendering.cpp
--- a/Engine/Source/Runtime/Renderer/Private/MobileTranslucentRendering.cpp
+++ b/Engine/Source/Runtime/Renderer/Private/MobileTranslucentRendering.cpp
@@ -368,7 +368,8 @@ void FMobileSceneRenderer::RenderTranslucency(FRHICommandListImmediate& RHICmdLi
 			FMobileTranslucencyDrawingPolicyFactory::ContextType DrawingContext(TranslucencyPass);
 			View.TranslucentPrimSet.DrawPrimitivesForMobile<FMobileTranslucencyDrawingPolicyFactory>(RHICmdList, View, DrawRenderState, DrawingContext);
 
-			View.SimpleElementCollector.DrawBatchedElements(RHICmdList, DrawRenderState, View, EBlendModeFilter::Translucent);
+			View.SimpleElementCollector.DrawBatchedElements(RHICmdList, DrawRenderState, View, EBlendModeFilter::Translucent, SDPG_World);
+			View.SimpleElementCollector.DrawBatchedElements(RHICmdList, DrawRenderState, View, EBlendModeFilter::Translucent, SDPG_Foreground);
 
 			// Editor and debug rendering
 			DrawViewElements<FMobileTranslucencyDrawingPolicyFactory>(RHICmdList, View, DrawRenderState, DrawingContext, SDPG_World, false);
diff --git a/Engine/Source/Runtime/Renderer/Private/PostProcess/PostProcessCompositeEditorPrimitives.cpp b/Engine/Source/Runtime/Renderer/Private/PostProcess/PostProcessCompositeEditorPrimitives.cpp
--- a/Engine/Source/Runtime/Renderer/Private/PostProcess/PostProcessCompositeEditorPrimitives.cpp
+++ b/Engine/Source/Runtime/Renderer/Private/PostProcess/PostProcessCompositeEditorPrimitives.cpp
@@ -325,7 +325,7 @@ static void RenderEditorPrimitives(FRHICommandListImmediate& RHICmdList, const F
 		}
 	}
 
-	View.EditorSimpleElementCollector.DrawBatchedElements(RHICmdList, DrawRenderState, View, EBlendModeFilter::OpaqueAndMasked);
+	View.EditorSimpleElementCollector.DrawBatchedElements(RHICmdList, DrawRenderState, View, EBlendModeFilter::OpaqueAndMasked, SDPG_World);
 
 	// Draw the base pass for the view's batched mesh elements.
 	//DrawViewElements<TBasePass>(RHICmdList, View, DrawRenderState, Context, SDPG_World, false);
@@ -349,6 +349,7 @@ static void RenderForegroundEditorPrimitives(FRHICommandListImmediate& RHICmdLis
 	if (bIgnoreExistingDepth)
 	{
 		DrawRenderState.SetDepthStencilState(TStaticDepthStencilState<true, CF_Always>::GetRHI());
+		View.EditorSimpleElementCollector.DrawBatchedElements(RHICmdList, DrawRenderState, View, EBlendModeFilter::OpaqueAndMasked, SDPG_Foreground);
 		DrawViewElements<TBasePass>(RHICmdList, View, DrawRenderState,
 			typename TBasePass::ContextType(), SDPG_Foreground, false);
 		View.TopBatchedViewElements.Draw(RHICmdList, DrawRenderState, FeatureLevel, bNeedToSwitchVerticalAxis, View, false);
@@ -357,6 +358,7 @@ static void RenderForegroundEditorPrimitives(FRHICommandListImmediate& RHICmdLis
 	// Draw a second time the foreground primitive with depth test to have proper depth test between foreground primitives.
 	{
 		DrawRenderState.SetDepthStencilState(TStaticDepthStencilState<true, CF_DepthNearOrEqual>::GetRHI());
+		View.EditorSimpleElementCollector.DrawBatchedElements(RHICmdList, DrawRenderState, View, EBlendModeFilter::OpaqueAndMasked, SDPG_Foreground);
 		DrawViewElements<TBasePass>(RHICmdList, View, DrawRenderState,
 			typename TBasePass::ContextType(), SDPG_Foreground, false);
 		View.TopBatchedViewElements.Draw(RHICmdList, DrawRenderState, FeatureLevel, bNeedToSwitchVerticalAxis, View, false);
diff --git a/Engine/Source/Runtime/Renderer/Private/SceneHitProxyRendering.cpp b/Engine/Source/Runtime/Renderer/Private/SceneHitProxyRendering.cpp
--- a/Engine/Source/Runtime/Renderer/Private/SceneHitProxyRendering.cpp
+++ b/Engine/Source/Runtime/Renderer/Private/SceneHitProxyRendering.cpp
@@ -530,7 +530,8 @@ static void DoRenderHitProxies(FRHICommandListImmediate& RHICmdList, const FScen
 			}
 		}
 
-		View.SimpleElementCollector.DrawBatchedElements(RHICmdList, DrawRenderState, View, EBlendModeFilter::All);
+		View.SimpleElementCollector.DrawBatchedElements(RHICmdList, DrawRenderState, View, EBlendModeFilter::All, SDPG_World);
+		View.SimpleElementCollector.DrawBatchedElements(RHICmdList, DrawRenderState, View, EBlendModeFilter::All, SDPG_Foreground);
 
 		for (int32 MeshBatchIndex = 0; MeshBatchIndex < View.DynamicEditorMeshElements.Num(); MeshBatchIndex++)
 		{
@@ -544,8 +545,9 @@ static void DoRenderHitProxies(FRHICommandListImmediate& RHICmdList, const FScen
 			}
 		}
 
-		View.EditorSimpleElementCollector.DrawBatchedElements(RHICmdList, DrawRenderState, View, EBlendModeFilter::All);
-		
+		View.EditorSimpleElementCollector.DrawBatchedElements(RHICmdList, DrawRenderState, View, EBlendModeFilter::All, SDPG_World);
+		View.EditorSimpleElementCollector.DrawBatchedElements(RHICmdList, DrawRenderState, View, EBlendModeFilter::All, SDPG_Foreground);
+
 		// Draw the view's elements.
 		DrawViewElements<FHitProxyDrawingPolicyFactory>(RHICmdList, View, DrawRenderState, FHitProxyDrawingPolicyFactory::ContextType(), SDPG_World, bPreFog);
 
diff --git a/Engine/Source/Runtime/Renderer/Private/TranslucentRendering.cpp b/Engine/Source/Runtime/Renderer/Private/TranslucentRendering.cpp
--- a/Engine/Source/Runtime/Renderer/Private/TranslucentRendering.cpp
+++ b/Engine/Source/Runtime/Renderer/Private/TranslucentRendering.cpp
@@ -665,7 +665,8 @@ void FTranslucentPrimSet::DrawPrimitives(
 		RenderPrimitive(RHICmdList, View, DrawRenderState, PrimitiveSceneInfo, ViewRelevance, TranslucentSelfShadow, TranslucencyPass);
 	}
 
-	View.SimpleElementCollector.DrawBatchedElements(RHICmdList, DrawRenderState, View, EBlendModeFilter::Translucent);
+	View.SimpleElementCollector.DrawBatchedElements(RHICmdList, DrawRenderState, View, EBlendModeFilter::Translucent, SDPG_World);
+	View.SimpleElementCollector.DrawBatchedElements(RHICmdList, DrawRenderState, View, EBlendModeFilter::Translucent, SDPG_Foreground);
 }
 
 void FTranslucentPrimSet::RenderPrimitive(
@@ -949,7 +950,8 @@ void FDeferredShadingSceneRenderer::RenderViewTranslucency(FRHICommandListImmedi
 
 	if (IsMainTranslucencyPass(TranslucencyPass))
 	{
-		View.SimpleElementCollector.DrawBatchedElements(RHICmdList, DrawRenderState, View, EBlendModeFilter::Translucent);
+		View.SimpleElementCollector.DrawBatchedElements(RHICmdList, DrawRenderState, View, EBlendModeFilter::Translucent, SDPG_World);
+		View.SimpleElementCollector.DrawBatchedElements(RHICmdList, DrawRenderState, View, EBlendModeFilter::Translucent, SDPG_Foreground);
 
 		// editor and debug rendering
 		if (View.bHasTranslucentViewMeshElements)
@@ -1021,7 +1023,8 @@ void FDeferredShadingSceneRenderer::RenderViewTranslucencyParallel(FRHICommandLi
 
 	if (IsMainTranslucencyPass(TranslucencyPass))
 	{
-		View.SimpleElementCollector.DrawBatchedElements(RHICmdList, DrawRenderState, View, EBlendModeFilter::Translucent);
+		View.SimpleElementCollector.DrawBatchedElements(RHICmdList, DrawRenderState, View, EBlendModeFilter::Translucent, SDPG_World);
+		View.SimpleElementCollector.DrawBatchedElements(RHICmdList, DrawRenderState, View, EBlendModeFilter::Translucent, SDPG_Foreground);
 
 		// editor and debug rendering
 		if (View.bHasTranslucentViewMeshElements)
