Cherry-pick 4529073: Adds support for the visualize texture tool in render graph.

This split the GVisualizeTexture in two:
* first part that capture any texture using a pixel shader, implemented using render pass, that render graph can just automatically insert where needed;
* Second part to present on screem still in Renderer module, because using things like FCanvas, requiring

Depth buffer visualization is broken, because needs to have more functionality in FRDGTextureSRVDesc.

diff --git a/Engine/Plugins/Experimental/OSVR/Source/OSVR/Private/OSVRRender.cpp b/Engine/Plugins/Experimental/OSVR/Source/OSVR/Private/OSVRRender.cpp
--- a/Engine/Plugins/Experimental/OSVR/Source/OSVR/Private/OSVRRender.cpp
+++ b/Engine/Plugins/Experimental/OSVR/Source/OSVR/Private/OSVRRender.cpp
@@ -56,7 +56,7 @@ void FOSVRHMD::RenderTexture_RenderThread(FRHICommandListImmediate& rhiCmdList,
 		TShaderMapRef<FScreenVS> vertexShader(shaderMap);
 		TShaderMapRef<FScreenPS> pixelShader(shaderMap);
 
-		GraphicsPSOInit.BoundShaderState.VertexDeclarationRHI = RendererModule->GetFilterVertexDeclaration().VertexDeclarationRHI;
+		GraphicsPSOInit.BoundShaderState.VertexDeclarationRHI = GFilterVertexDeclaration.VertexDeclarationRHI;
 		GraphicsPSOInit.BoundShaderState.VertexShaderRHI = GETSAFERHISHADER_VERTEX(*vertexShader);
 		GraphicsPSOInit.BoundShaderState.PixelShaderRHI = GETSAFERHISHADER_PIXEL(*pixelShader);
 		GraphicsPSOInit.PrimitiveType = PT_TriangleList;
diff --git a/Engine/Plugins/Experimental/PixelStreaming/Source/PixelStreaming/Private/NvVideoEncoder.cpp b/Engine/Plugins/Experimental/PixelStreaming/Source/PixelStreaming/Private/NvVideoEncoder.cpp
--- a/Engine/Plugins/Experimental/PixelStreaming/Source/PixelStreaming/Private/NvVideoEncoder.cpp
+++ b/Engine/Plugins/Experimental/PixelStreaming/Source/PixelStreaming/Private/NvVideoEncoder.cpp
@@ -15,6 +15,7 @@
 #include "HAL/ThreadSafeCounter.h"
 #include "Misc/CoreDelegates.h"
 #include "Modules/ModuleManager.h"
+#include "DummyRenderResources.h"
 
 #if defined PLATFORM_WINDOWS
 // Disable macro redefinition warning for compatibility with Windows SDK 8+
@@ -443,7 +444,7 @@ void FNvVideoEncoder::FNvVideoEncoderImpl::CopyBackBuffer(const FTexture2DRHIRef
 			TShaderMapRef<FScreenVS> VertexShader(ShaderMap);
 			TShaderMapRef<FScreenPS> PixelShader(ShaderMap);
 
-			GraphicsPSOInit.BoundShaderState.VertexDeclarationRHI = RendererModule->GetFilterVertexDeclaration().VertexDeclarationRHI;
+			GraphicsPSOInit.BoundShaderState.VertexDeclarationRHI = GFilterVertexDeclaration.VertexDeclarationRHI;
 			GraphicsPSOInit.BoundShaderState.VertexShaderRHI = GETSAFERHISHADER_VERTEX(*VertexShader);
 			GraphicsPSOInit.BoundShaderState.PixelShaderRHI = GETSAFERHISHADER_PIXEL(*PixelShader);
 			GraphicsPSOInit.PrimitiveType = PT_TriangleList;
diff --git a/Engine/Plugins/Experimental/RemoteSession/Source/RemoteSession/Private/Channels/RemoteSessionARCameraChannel.cpp b/Engine/Plugins/Experimental/RemoteSession/Source/RemoteSession/Private/Channels/RemoteSessionARCameraChannel.cpp
--- a/Engine/Plugins/Experimental/RemoteSession/Source/RemoteSession/Private/Channels/RemoteSessionARCameraChannel.cpp
+++ b/Engine/Plugins/Experimental/RemoteSession/Source/RemoteSession/Private/Channels/RemoteSessionARCameraChannel.cpp
@@ -38,6 +38,7 @@
 
 #include "IImageWrapper.h"
 #include "IImageWrapperModule.h"
+#include "DummyRenderResources.h"
 
 #define CAMERA_MESSAGE_ADDRESS TEXT("/ARCamera")
 
@@ -302,7 +303,7 @@ void FARCameraSceneViewExtension::RenderARCamera_RenderThread(FRHICommandListImm
 	GraphicsPSOInit.RasterizerState = TStaticRasterizerState<>::GetRHI();
 	GraphicsPSOInit.DepthStencilState = TStaticDepthStencilState<false, CF_DepthNearOrEqual>::GetRHI();
 	GraphicsPSOInit.PrimitiveType = PT_TriangleList;
-	GraphicsPSOInit.BoundShaderState.VertexDeclarationRHI = RendererModule.GetFilterVertexDeclaration().VertexDeclarationRHI;
+	GraphicsPSOInit.BoundShaderState.VertexDeclarationRHI = GFilterVertexDeclaration.VertexDeclarationRHI;
 
 	const bool bIsMobileRenderer = FeatureLevel <= ERHIFeatureLevel::ES3_1;
 	FMaterialShader* VertexShader = nullptr;
diff --git a/Engine/Plugins/Lumin/MagicLeap/Source/MagicLeap/Private/MagicLeapHMD.cpp b/Engine/Plugins/Lumin/MagicLeap/Source/MagicLeap/Private/MagicLeapHMD.cpp
--- a/Engine/Plugins/Lumin/MagicLeap/Source/MagicLeap/Private/MagicLeapHMD.cpp
+++ b/Engine/Plugins/Lumin/MagicLeap/Source/MagicLeap/Private/MagicLeapHMD.cpp
@@ -1623,7 +1623,7 @@ void FMagicLeapHMD::RenderTexture_RenderThread(class FRHICommandListImmediate& R
 			GraphicsPSOInit.BlendState = TStaticBlendState<>::GetRHI();
 			GraphicsPSOInit.RasterizerState = TStaticRasterizerState<>::GetRHI();
 			GraphicsPSOInit.DepthStencilState = TStaticDepthStencilState<false, CF_Always>::GetRHI();
-			GraphicsPSOInit.BoundShaderState.VertexDeclarationRHI = RendererModule->GetFilterVertexDeclaration().VertexDeclarationRHI;
+			GraphicsPSOInit.BoundShaderState.VertexDeclarationRHI = GFilterVertexDeclaration.VertexDeclarationRHI;
 			GraphicsPSOInit.BoundShaderState.VertexShaderRHI = GETSAFERHISHADER_VERTEX(*VertexShader);
 			GraphicsPSOInit.BoundShaderState.PixelShaderRHI = GETSAFERHISHADER_PIXEL(*PixelShader);
 			GraphicsPSOInit.PrimitiveType = PT_TriangleList;
diff --git a/Engine/Plugins/Runtime/AppleARKit/Source/AppleARKit/Private/AppleARKitVideoOverlay.cpp b/Engine/Plugins/Runtime/AppleARKit/Source/AppleARKit/Private/AppleARKitVideoOverlay.cpp
--- a/Engine/Plugins/Runtime/AppleARKit/Source/AppleARKit/Private/AppleARKitVideoOverlay.cpp
+++ b/Engine/Plugins/Runtime/AppleARKit/Source/AppleARKit/Private/AppleARKitVideoOverlay.cpp
@@ -20,6 +20,7 @@
 #include "Containers/DynamicRHIResourceArray.h"
 #include "PostProcess/SceneFilterRendering.h"
 #include "PostProcessParameters.h"
+#include "DummyRenderResources.h"
 #if SUPPORTS_ARKIT_1_0
 	#include "IOSAppDelegate.h"
 #endif
@@ -394,7 +395,7 @@ void FAppleARKitVideoOverlay::RenderVideoOverlay_RenderThread(FRHICommandListImm
 		GraphicsPSOInit.RasterizerState = TStaticRasterizerState<>::GetRHI();
 		GraphicsPSOInit.DepthStencilState = TStaticDepthStencilState<false, CF_DepthNearOrEqual>::GetRHI();
 		GraphicsPSOInit.PrimitiveType = PT_TriangleList;
-		GraphicsPSOInit.BoundShaderState.VertexDeclarationRHI = RendererModule.GetFilterVertexDeclaration().VertexDeclarationRHI;
+		GraphicsPSOInit.BoundShaderState.VertexDeclarationRHI = GFilterVertexDeclaration.VertexDeclarationRH;
 		
 		const bool bIsMobileRenderer = FeatureLevel <= ERHIFeatureLevel::ES3_1;
 		FMaterialShader* VertexShader = nullptr;
diff --git a/Engine/Plugins/Runtime/GoogleARCore/Source/GoogleARCoreRendering/Private/GoogleARCorePassthroughCameraRenderer.cpp b/Engine/Plugins/Runtime/GoogleARCore/Source/GoogleARCoreRendering/Private/GoogleARCorePassthroughCameraRenderer.cpp
--- a/Engine/Plugins/Runtime/GoogleARCore/Source/GoogleARCoreRendering/Private/GoogleARCorePassthroughCameraRenderer.cpp
+++ b/Engine/Plugins/Runtime/GoogleARCore/Source/GoogleARCoreRendering/Private/GoogleARCorePassthroughCameraRenderer.cpp
@@ -18,6 +18,7 @@
 #include "ExternalTexture.h"
 #include "GoogleARCorePassthroughCameraExternalTextureGuid.h"
 #include "GoogleARCoreAndroidHelper.h"
+#include "DummyRenderResources.h"
 
 FGoogleARCorePassthroughCameraRenderer::FGoogleARCorePassthroughCameraRenderer()
 	: OverlayQuadUVs{ 0.0f, 0.0f, 0.0f, 1.0f, 1.0f, 0.0f, 1.0f, 1.0f }
@@ -286,7 +287,7 @@ void FGoogleARCorePassthroughCameraRenderer::RenderVideoOverlay_RenderThread(FRH
 		GraphicsPSOInit.RasterizerState = TStaticRasterizerState<>::GetRHI();
 		GraphicsPSOInit.DepthStencilState = TStaticDepthStencilState<false, CF_DepthNearOrEqual>::GetRHI();
 
-		GraphicsPSOInit.BoundShaderState.VertexDeclarationRHI = RendererModule.GetFilterVertexDeclaration().VertexDeclarationRHI;
+		GraphicsPSOInit.BoundShaderState.VertexDeclarationRHI = GFilterVertexDeclaration.VertexDeclarationRHI;
 		GraphicsPSOInit.BoundShaderState.VertexShaderRHI = GETSAFERHISHADER_VERTEX(VertexShader);
 		GraphicsPSOInit.BoundShaderState.PixelShaderRHI = GETSAFERHISHADER_PIXEL(PixelShader);
 		GraphicsPSOInit.PrimitiveType = PT_TriangleList;
diff --git a/Engine/Plugins/Runtime/GoogleVR/GoogleVRHMD/Source/GoogleVRHMD/Private/GoogleVRHMD.cpp b/Engine/Plugins/Runtime/GoogleVR/GoogleVRHMD/Source/GoogleVRHMD/Private/GoogleVRHMD.cpp
--- a/Engine/Plugins/Runtime/GoogleVR/GoogleVRHMD/Source/GoogleVRHMD/Private/GoogleVRHMD.cpp
+++ b/Engine/Plugins/Runtime/GoogleVR/GoogleVRHMD/Source/GoogleVRHMD/Private/GoogleVRHMD.cpp
@@ -1324,7 +1324,7 @@ void FGoogleVRHMD::PostRenderViewFamily_RenderThread(FRHICommandListImmediate& R
 			GraphicsPSOInit.RasterizerState = TStaticRasterizerState<>::GetRHI();
 			GraphicsPSOInit.DepthStencilState = TStaticDepthStencilState<false, CF_Always>::GetRHI();
 
-			GraphicsPSOInit.BoundShaderState.VertexDeclarationRHI = RendererModule->GetFilterVertexDeclaration().VertexDeclarationRHI;
+			GraphicsPSOInit.BoundShaderState.VertexDeclarationRHI = GFilterVertexDeclaration.VertexDeclarationRHI;
 			GraphicsPSOInit.BoundShaderState.VertexShaderRHI = GETSAFERHISHADER_VERTEX(*VertexShader);
 			GraphicsPSOInit.BoundShaderState.PixelShaderRHI = GETSAFERHISHADER_PIXEL(*PixelShader);
 			GraphicsPSOInit.PrimitiveType = PT_TriangleList;
diff --git a/Engine/Plugins/Runtime/GoogleVR/GoogleVRHMD/Source/GoogleVRHMD/Private/GoogleVRHMDRender.cpp b/Engine/Plugins/Runtime/GoogleVR/GoogleVRHMD/Source/GoogleVRHMD/Private/GoogleVRHMDRender.cpp
--- a/Engine/Plugins/Runtime/GoogleVR/GoogleVRHMD/Source/GoogleVRHMD/Private/GoogleVRHMDRender.cpp
+++ b/Engine/Plugins/Runtime/GoogleVR/GoogleVRHMD/Source/GoogleVRHMD/Private/GoogleVRHMDRender.cpp
@@ -215,7 +215,7 @@ static void ResolvePendingRenderTarget(FRHICommandListImmediate& RHICmdList, FGr
 		TShaderMapRef<FScreenPS> PixelShader(ShaderMap);
 		RHICmdList.ApplyCachedRenderTargets(GraphicsPSOInit);
 
-		GraphicsPSOInit.BoundShaderState.VertexDeclarationRHI = RendererModule->GetFilterVertexDeclaration().VertexDeclarationRHI;
+		GraphicsPSOInit.BoundShaderState.VertexDeclarationRHI = GFilterVertexDeclaration.VertexDeclarationRHI;
 		GraphicsPSOInit.BoundShaderState.VertexShaderRHI = GETSAFERHISHADER_VERTEX(*VertexShader);
 		GraphicsPSOInit.BoundShaderState.PixelShaderRHI = GETSAFERHISHADER_PIXEL(*PixelShader);
 		GraphicsPSOInit.PrimitiveType = PT_TriangleList;
@@ -302,7 +302,7 @@ void FGoogleVRHMD::RenderTexture_RenderThread(FRHICommandListImmediate& RHICmdLi
 			TShaderMapRef<FScreenVS> VertexShader(ShaderMap);
 			TShaderMapRef<FScreenPS> PixelShader(ShaderMap);
 
-			GraphicsPSOInit.BoundShaderState.VertexDeclarationRHI = RendererModule->GetFilterVertexDeclaration().VertexDeclarationRHI;
+			GraphicsPSOInit.BoundShaderState.VertexDeclarationRHI = GFilterVertexDeclaration.VertexDeclarationRHI;
 			GraphicsPSOInit.BoundShaderState.VertexShaderRHI = GETSAFERHISHADER_VERTEX(*VertexShader);
 			GraphicsPSOInit.BoundShaderState.PixelShaderRHI = GETSAFERHISHADER_PIXEL(*PixelShader);
 			GraphicsPSOInit.PrimitiveType = PT_TriangleList;
diff --git a/Engine/Plugins/Runtime/GoogleVR/GoogleVRHMD/Source/GoogleVRHMD/Private/GoogleVRSplash.cpp b/Engine/Plugins/Runtime/GoogleVR/GoogleVRHMD/Source/GoogleVRHMD/Private/GoogleVRSplash.cpp
--- a/Engine/Plugins/Runtime/GoogleVR/GoogleVRHMD/Source/GoogleVRHMD/Private/GoogleVRSplash.cpp
+++ b/Engine/Plugins/Runtime/GoogleVR/GoogleVRHMD/Source/GoogleVRHMD/Private/GoogleVRSplash.cpp
@@ -227,7 +227,7 @@ void FGoogleVRSplash::RenderStereoSplashScreen(FRHICommandListImmediate& RHICmdL
 			GraphicsPSOInit.RasterizerState = TStaticRasterizerState<>::GetRHI();
 			GraphicsPSOInit.DepthStencilState = TStaticDepthStencilState<false, CF_Always>::GetRHI();
 
-			GraphicsPSOInit.BoundShaderState.VertexDeclarationRHI = RendererModule->GetFilterVertexDeclaration().VertexDeclarationRHI;
+			GraphicsPSOInit.BoundShaderState.VertexDeclarationRHI = GFilterVertexDeclaration.VertexDeclarationRHI;
 			GraphicsPSOInit.BoundShaderState.VertexShaderRHI = GETSAFERHISHADER_VERTEX(*VertexShader);
 			GraphicsPSOInit.BoundShaderState.PixelShaderRHI = GETSAFERHISHADER_PIXEL(*PixelShader);
 			GraphicsPSOInit.PrimitiveType = PT_TriangleList;
diff --git a/Engine/Plugins/Runtime/Oculus/OculusVR/Source/OculusHMD/Private/OculusHMD_CustomPresent.cpp b/Engine/Plugins/Runtime/Oculus/OculusVR/Source/OculusHMD/Private/OculusHMD_CustomPresent.cpp
--- a/Engine/Plugins/Runtime/Oculus/OculusVR/Source/OculusHMD/Private/OculusHMD_CustomPresent.cpp
+++ b/Engine/Plugins/Runtime/Oculus/OculusVR/Source/OculusHMD/Private/OculusHMD_CustomPresent.cpp
@@ -8,6 +8,7 @@
 #include "PipelineStateCache.h"
 #include "ClearQuad.h"
 #include "OculusShaders.h"
+#include "DummyRenderResources.h"
 
 #if PLATFORM_ANDROID
 #include "Android/AndroidJNI.h"
@@ -392,7 +393,7 @@ void FCustomPresent::CopyTexture_RenderThread(FRHICommandListImmediate& RHICmdLi
 	const auto FeatureLevel = GMaxRHIFeatureLevel;
 	auto ShaderMap = GetGlobalShaderMap(FeatureLevel);
 	TShaderMapRef<FScreenVS> VertexShader(ShaderMap);
-	GraphicsPSOInit.BoundShaderState.VertexDeclarationRHI = RendererModule->GetFilterVertexDeclaration().VertexDeclarationRHI;
+	GraphicsPSOInit.BoundShaderState.VertexDeclarationRHI = GFilterVertexDeclaration.VertexDeclarationRHI;
 	GraphicsPSOInit.BoundShaderState.VertexShaderRHI = GETSAFERHISHADER_VERTEX(*VertexShader);
 
 	if (DstTexture2D)
diff --git a/Engine/Plugins/Runtime/Steam/SteamVR/Source/SteamVR/Private/SteamVRSpectatorScreen.cpp b/Engine/Plugins/Runtime/Steam/SteamVR/Source/SteamVR/Private/SteamVRSpectatorScreen.cpp
--- a/Engine/Plugins/Runtime/Steam/SteamVR/Source/SteamVR/Private/SteamVRSpectatorScreen.cpp
+++ b/Engine/Plugins/Runtime/Steam/SteamVR/Source/SteamVR/Private/SteamVRSpectatorScreen.cpp
@@ -79,7 +79,7 @@ void FSteamVRHMD::CopyTexture_RenderThread(FRHICommandListImmediate& RHICmdList,
 		TShaderMapRef<FScreenVS> VertexShader(ShaderMap);
 		TShaderMapRef<FScreenPS> PixelShader(ShaderMap);
 
-		GraphicsPSOInit.BoundShaderState.VertexDeclarationRHI = RendererModule->GetFilterVertexDeclaration().VertexDeclarationRHI;
+		GraphicsPSOInit.BoundShaderState.VertexDeclarationRHI = GFilterVertexDeclaration.VertexDeclarationRHI;
 		GraphicsPSOInit.BoundShaderState.VertexShaderRHI = GETSAFERHISHADER_VERTEX(*VertexShader);
 		GraphicsPSOInit.BoundShaderState.PixelShaderRHI = GETSAFERHISHADER_PIXEL(*PixelShader);
 
diff --git a/Engine/Plugins/Runtime/WindowsMixedReality/Source/WindowsMixedRealityHMD/Private/WindowsMixedRealityHMD.cpp b/Engine/Plugins/Runtime/WindowsMixedReality/Source/WindowsMixedRealityHMD/Private/WindowsMixedRealityHMD.cpp
--- a/Engine/Plugins/Runtime/WindowsMixedReality/Source/WindowsMixedRealityHMD/Private/WindowsMixedRealityHMD.cpp
+++ b/Engine/Plugins/Runtime/WindowsMixedReality/Source/WindowsMixedRealityHMD/Private/WindowsMixedRealityHMD.cpp
@@ -736,7 +736,7 @@ namespace WindowsMixedReality
 		GraphicsPSOInit.BlendState = TStaticBlendState<>::GetRHI();
 		GraphicsPSOInit.RasterizerState = TStaticRasterizerState<>::GetRHI();
 		GraphicsPSOInit.DepthStencilState = TStaticDepthStencilState<false, CF_Always>::GetRHI();
-		GraphicsPSOInit.BoundShaderState.VertexDeclarationRHI = RendererModule->GetFilterVertexDeclaration().VertexDeclarationRHI;
+		GraphicsPSOInit.BoundShaderState.VertexDeclarationRHI = GFilterVertexDeclaration.VertexDeclarationRHI;
 		GraphicsPSOInit.BoundShaderState.VertexShaderRHI = GETSAFERHISHADER_VERTEX(*VertexShader);
 		GraphicsPSOInit.BoundShaderState.PixelShaderRHI = GETSAFERHISHADER_PIXEL(*PixelShader);
 		GraphicsPSOInit.PrimitiveType = PT_TriangleList;
@@ -893,7 +893,7 @@ namespace WindowsMixedReality
 			TShaderMapRef<FScreenVS> vertexShader(shaderMap);
 			TShaderMapRef<FDepthConversionPS> pixelShader(shaderMap);
 
-			GraphicsPSOInit.BoundShaderState.VertexDeclarationRHI = RendererModule->GetFilterVertexDeclaration().VertexDeclarationRHI;
+			GraphicsPSOInit.BoundShaderState.VertexDeclarationRHI = GFilterVertexDeclaration.VertexDeclarationRHI;
 			GraphicsPSOInit.BoundShaderState.VertexShaderRHI = GETSAFERHISHADER_VERTEX(*vertexShader);
 			GraphicsPSOInit.BoundShaderState.PixelShaderRHI = GETSAFERHISHADER_PIXEL(*pixelShader);
 			GraphicsPSOInit.PrimitiveType = PT_TriangleList;
diff --git a/Engine/Shaders/Private/Tools/FullscreenVertexShader.usf b/Engine/Shaders/Private/Tools/FullscreenVertexShader.usf
new file mode 100644
--- /dev/null
+++ b/Engine/Shaders/Private/Tools/FullscreenVertexShader.usf
@@ -0,0 +1,17 @@
+// Copyright 1998-2018 Epic Games, Inc. All Rights Reserved.
+
+/*=============================================================================
+	FullscreenVertexShader.usf: Most basic multi platform full screen vertex shader
+=============================================================================*/
+
+#include "../Common.ush"
+
+void MainVS(
+	float2 InPosition : ATTRIBUTE0,
+	float2 InUV       : ATTRIBUTE1,
+	out FScreenVertexOutput Output
+	)
+{
+	Output.UV = InUV;
+	Output.Position = float4(InPosition.x * 2.0 - 1.0, 1.0 - 2.0 * InPosition.y, 0, 1);
+}
diff --git a/Engine/Shaders/Private/Tools/VisualizeTexture.usf b/Engine/Shaders/Private/Tools/VisualizeTexture.usf
new file mode 100644
--- /dev/null
+++ b/Engine/Shaders/Private/Tools/VisualizeTexture.usf
@@ -0,0 +1,183 @@
+// Copyright 1998-2018 Epic Games, Inc. All Rights Reserved.
+
+/*=============================================================================
+	VisualizeTexture.usf: VisualizeTexture shader
+=============================================================================*/
+
+#include "../Common.ush"
+
+
+// to visualize the content of a texture to the screen, useful for debugging
+// use "VisualizeTexture" in the console to activate
+Texture2D VisualizeTexture2D;
+SamplerState VisualizeTexture2DSampler;
+Texture3D VisualizeTexture3D;
+SamplerState VisualizeTexture3DSampler;  
+TextureCube VisualizeTextureCube;
+SamplerState VisualizeTextureCubeSampler;
+
+//default Texture2D type is <float4>, which won't read the X28_G8UINT format properly.
+Texture2D<uint4> VisualizeDepthStencilTexture; 
+
+Texture2D<uint> VisualizeUINT8Texture2D;
+
+#if FEATURE_LEVEL >= FEATURE_LEVEL_SM5
+Texture2DMS<float4> VisualizeTexture2DMS;
+TextureCubeArray VisualizeTextureCubeArray;
+SamplerState VisualizeTextureCubeArraySampler;
+#endif
+
+// [0]: r:RGBAdd, g:SingleChannelMul, b:Add, a:FracScale
+// [1]: r:BlinkState=0/1, g:Frac when 0, saturate when 1, b:array index, a:custom mip
+// [2]: r:0=normal 1=depth, 2=shadow depth, g:unwrapcubemap, b:singlechannel(0-3), a:unused
+float4 VisualizeParam[3];
+// x, y (0 for 2D textures), z(only !=0 if 3d texture), w unused
+float3 TextureExtent;
+
+void VisualizeTexturePS(
+	FScreenVertexOutput Input,
+	out float4 OutColor : SV_Target0)
+{
+	float2 UV = Input.UV;
+
+	// Longitude Latitude unwrapping
+	float3 SpherePos;
+	{
+		float2 Angles = float2(2 * PI * (UV.x + 0.5f), PI * (1 - UV.y));
+		float s = sin(Angles.y);
+		SpherePos = float3(s * sin(Angles.x), -s * cos(Angles.x), -cos(Angles.y));
+	}
+
+	// volume texture unwrapping
+	float3 VolumePos;
+	{
+		uint LayerCount = TextureExtent.z;
+		uint LayerUCount = (uint)sqrt(TextureExtent.z);
+		uint LayerVCount = LayerCount / LayerUCount;
+
+		uint LayerUId = (uint)(UV.x * LayerUCount);
+		uint LayerVId = (uint)(UV.y * LayerVCount);
+		uint LayerId = LayerUId + LayerUCount * LayerVId;
+
+		VolumePos.z = LayerId / (TextureExtent.z - 1);
+		VolumePos.x = frac(UV.x * LayerUCount);
+		VolumePos.y = frac(UV.y * LayerVCount);
+	}
+
+	float MipLevel = VisualizeParam[1].a;
+	float ArrayIndex = VisualizeParam[1].b;
+
+
+
+#if TEXTURE_TYPE == 0
+	float4 TexLookup = TextureCubeSampleLevel(VisualizeTextureCube, VisualizeTextureCubeSampler, SpherePos, MipLevel);
+#elif TEXTURE_TYPE == 2 || TEXTURE_TYPE == 6
+	float4 TexLookup = Texture2DSampleLevel(VisualizeTexture2D, VisualizeTexture2DSampler, UV, MipLevel);
+#elif TEXTURE_TYPE == 3
+	#if FEATURE_LEVEL < FEATURE_LEVEL_SM4
+		float4 TexLookup = float4(1.0f, 0.0f, 1.0f, 1.0f);
+	#else
+		float4 TexLookup = Texture3DSampleLevel(VisualizeTexture3D, VisualizeTexture3DSampler, VolumePos, MipLevel);
+	#endif
+#elif TEXTURE_TYPE == 4
+	#if FEATURE_LEVEL < FEATURE_LEVEL_SM5
+		float4 TexLookup = float4(1.0f, 0.0f, 1.0f, 1.0f);
+	#else
+		float4 TexLookup = TextureCubeArraySampleLevel(VisualizeTextureCubeArray, VisualizeTextureCubeArraySampler, SpherePos, ArrayIndex, MipLevel);
+	#endif
+#elif TEXTURE_TYPE == 5
+	float4 TexLookup = float4(1,1,0,0);	// yellow, not yet supported
+	
+	#if FEATURE_LEVEL >= FEATURE_LEVEL_SM5
+		TexLookup = VisualizeTexture2DMS.Load(int2(UV * TextureExtent.xy), 0);
+	#endif
+#elif TEXTURE_TYPE == 7
+    float IntTexLookup = float(VisualizeUINT8Texture2D.Load(int3(UV * TextureExtent.xy, 0)));
+	float4 TexLookup = IntTexLookup.xxxx;
+#endif
+	
+	// Normal = 0, SceneDepth = 1 and ShadowDepth = 2
+	int InputValueMapping = (int)VisualizeParam[2].r;
+	if(InputValueMapping != 0)
+	{
+		float Depth = TexLookup.r;
+		if (InputValueMapping == 1)  // SceneDepth 
+		{
+			// current z buffer is negated (1 is near, 0 is far)
+#if HAS_INVERTED_Z_BUFFER
+			Depth = 1 - Depth;
+#endif
+
+#if FEATURE_LEVEL >= FEATURE_LEVEL_SM4
+			// map depth to nice colors
+			float Scale = 64.0; // This 64 looks like a hardware limit.
+			Depth = (exp2(Depth * Scale) - 1.0) * (1.0 / (exp2(Scale) - 1.0));
+			Depth *= Depth;
+			Depth *= Depth;
+
+			float Value = Depth * Depth;
+			TexLookup = float4(pow(Value, 0.25f), pow(Value, 1), pow(Value, 4), 0);
+		} 
+		else if (InputValueMapping == 2) // ShadowDepth
+		{
+			// linearly blend from White (near) to Teal (far).
+			float3 TealBlue = float3(54.f/255.f, 117.f/255.f, 136.f/255.f);
+			float3 White = float3(1.0f, 1.0f, 1.0f);
+			TexLookup.rgb = lerp(White, TealBlue, Depth);
+		}
+#if TEXTURE_TYPE == 6
+		// inside #ifdef to avoid D3DDebug error:
+		//    The resource return type for component 0 declared in the shader code (UINT) is not compatible with the Shader Resource View format bound to slot 1 of the Pixel Shader unit (UNORM).  This mismatch is invalid if the shader actually uses the view (e.g. it is not skipped due to shader code branching). [ EXECUTION ERROR #361: DEVICE_DRAW_RESOURCE_RETURN_TYPE_MISMATCH]
+		float fStencil = VisualizeDepthStencilTexture.Load(int3(UV * TextureExtent.xy, 0)).g;
+		fStencil /= 255.0f;
+		TexLookup.a = fStencil;
+#endif
+
+#else
+		}
+		TexLookup.r = Depth;
+#endif
+	}
+
+	float4 RGBAMask = 0;
+	int SingleChannel = (int)VisualizeParam[2].z;
+	switch (SingleChannel)
+	{
+		case 0: RGBAMask = float4(1,0,0,0); break;
+		case 1: RGBAMask = float4(0,1,0,0); break;
+		case 2: RGBAMask = float4(0,0,1,0); break;
+		case 3: RGBAMask = float4(0,0,0,1); break;
+		default: RGBAMask = float4(0,0,0,0); break;
+	}
+	float Add = dot( TexLookup, RGBAMask ) * VisualizeParam[0].y + VisualizeParam[0].z;
+
+
+	OutColor.rgb = TexLookup.rgb * VisualizeParam[0].xxx + Add.xxx;
+
+	float3 ScaledColor = OutColor.rgb * VisualizeParam[0].w;
+	float3 FracOutColor = lerp(ScaledColor, saturate(ScaledColor), VisualizeParam[1].y);
+
+	float3 AlternateColor = FracOutColor;
+	
+	// blink red if <0
+	if(OutColor.r < 0.0f || OutColor.g < 0.0f || OutColor.b < 0.0f)
+	{
+		AlternateColor = float3(1, 0, 0);
+	}
+	
+	// blink blue if not finite
+#if FEATURE_LEVEL >= FEATURE_LEVEL_SM4
+	if (isfinite(OutColor.rgb).x == false)
+	{
+		AlternateColor = float3(0, 0, 1);
+	}
+#endif
+	OutColor.rgb = lerp(FracOutColor, AlternateColor, VisualizeParam[1].x);	
+	OutColor.a = 0;
+}
+
+
+void PresentPS(float4 InUV : TEXCOORD0, out float4 OutColor : SV_Target0)
+{
+	OutColor = Texture2DSample(VisualizeTexture2D, VisualizeTexture2DSampler, InUV.xy);	
+}
diff --git a/Engine/Shaders/Private/VisualizeTexture.usf b/Engine/Shaders/Private/VisualizeTexture.usf
deleted file mode 100644
diff --git a/Engine/Shaders/Public/ShaderVersion.ush b/Engine/Shaders/Public/ShaderVersion.ush
--- a/Engine/Shaders/Public/ShaderVersion.ush
+++ b/Engine/Shaders/Public/ShaderVersion.ush
@@ -2,4 +2,4 @@
 // Each time the console command is executed it generates a new GUID. As this file is included
 // in Platform.ush (which should be included in any shader) it allows to invalidate the shader DDC.
 // 
-// GUID = A7CB4E7E6E3347189A794F837FDB1A6D
+// GUID = 8301458A46E1B310E41BA7B2A499BFC1aaaaa
\ No newline at end of file
diff --git a/Engine/Source/Editor/MovieSceneTools/Private/TrackEditorThumbnail/TrackEditorThumbnail.cpp b/Engine/Source/Editor/MovieSceneTools/Private/TrackEditorThumbnail/TrackEditorThumbnail.cpp
--- a/Engine/Source/Editor/MovieSceneTools/Private/TrackEditorThumbnail/TrackEditorThumbnail.cpp
+++ b/Engine/Source/Editor/MovieSceneTools/Private/TrackEditorThumbnail/TrackEditorThumbnail.cpp
@@ -17,6 +17,7 @@
 #include "ScreenRendering.h"
 #include "TrackEditorThumbnail/TrackEditorThumbnailPool.h"
 #include "PipelineStateCache.h"
+#include "DummyRenderResources.h"
 
 namespace TrackEditorThumbnailConstants
 {
@@ -155,7 +156,7 @@ void FTrackEditorThumbnail::CopyTextureIn(FTexture2DRHIRef SourceTexture)
 			TShaderMapRef<FScreenVS> VertexShader(ShaderMap);
 			TShaderMapRef<FScreenPS> PixelShader(ShaderMap);
 
-			GraphicsPSOInit.BoundShaderState.VertexDeclarationRHI = RendererModule->GetFilterVertexDeclaration().VertexDeclarationRHI;
+			GraphicsPSOInit.BoundShaderState.VertexDeclarationRHI = GFilterVertexDeclaration.VertexDeclarationRHI;
 			GraphicsPSOInit.BoundShaderState.VertexShaderRHI = GETSAFERHISHADER_VERTEX(*VertexShader);
 			GraphicsPSOInit.BoundShaderState.PixelShaderRHI = GETSAFERHISHADER_PIXEL(*PixelShader);
 
diff --git a/Engine/Source/Runtime/Engine/Private/UserInterface/VisualizeRT.cpp b/Engine/Source/Runtime/Engine/Private/UserInterface/VisualizeRT.cpp
--- a/Engine/Source/Runtime/Engine/Private/UserInterface/VisualizeRT.cpp
+++ b/Engine/Source/Runtime/Engine/Private/UserInterface/VisualizeRT.cpp
@@ -24,6 +24,7 @@
 #include "Widgets/Views/STableViewBase.h"
 #include "Widgets/Views/STableRow.h"
 #include "Widgets/Views/SListView.h"
+#include "VisualizeTexture.h"
 
 #define ENABLE_IMAGES	0
 
@@ -220,7 +221,9 @@ public:
 
 		// Get List from Renderer
 		FQueryVisualizeTexureInfo VisTextureInfo;
-		RendererModule.QueryVisualizeTexture(VisTextureInfo);
+		{
+			GVisualizeTexture.QueryInfo_GameThread(VisTextureInfo);
+		}
 		uint32 TotalSize = 0;
 		for (int32 Index = 0; Index < VisTextureInfo.Entries.Num(); ++Index)
 		{
diff --git a/Engine/Source/Runtime/HeadMountedDisplay/Private/DefaultStereoLayers.cpp b/Engine/Source/Runtime/HeadMountedDisplay/Private/DefaultStereoLayers.cpp
--- a/Engine/Source/Runtime/HeadMountedDisplay/Private/DefaultStereoLayers.cpp
+++ b/Engine/Source/Runtime/HeadMountedDisplay/Private/DefaultStereoLayers.cpp
@@ -13,6 +13,7 @@
 #include "PipelineStateCache.h"
 #include "ClearQuad.h"
 #include "SceneView.h"
+#include "DummyRenderResources.h"
 
 namespace 
 {
@@ -86,7 +87,7 @@ void FDefaultStereoLayers::StereoLayerRender(FRHICommandListImmediate& RHICmdLis
 	TShaderMapRef<FStereoLayerVS> VertexShader(ShaderMap);
 	TShaderMapRef<FStereoLayerPS> PixelShader(ShaderMap);
 
-	GraphicsPSOInit.BoundShaderState.VertexDeclarationRHI = RendererModule.GetFilterVertexDeclaration().VertexDeclarationRHI;
+	GraphicsPSOInit.BoundShaderState.VertexDeclarationRHI = GFilterVertexDeclaration.VertexDeclarationRHI;
 	GraphicsPSOInit.BoundShaderState.VertexShaderRHI = GETSAFERHISHADER_VERTEX(*VertexShader);
 	GraphicsPSOInit.BoundShaderState.PixelShaderRHI = GETSAFERHISHADER_PIXEL(*PixelShader);
 	GraphicsPSOInit.PrimitiveType = PT_TriangleList;
diff --git a/Engine/Source/Runtime/MovieSceneCapture/Private/FrameGrabber.cpp b/Engine/Source/Runtime/MovieSceneCapture/Private/FrameGrabber.cpp
--- a/Engine/Source/Runtime/MovieSceneCapture/Private/FrameGrabber.cpp
+++ b/Engine/Source/Runtime/MovieSceneCapture/Private/FrameGrabber.cpp
@@ -15,6 +15,7 @@
 #include "GlobalShader.h"
 #include "ScreenRendering.h"
 #include "PipelineStateCache.h"
+#include "DummyRenderResources.h"
 
 int32 GFrameGrabberFrameLatency = 0;
 static FAutoConsoleVariableRef CVarFrameGrabberFrameLatency(
@@ -167,7 +168,7 @@ void FViewportSurfaceReader::ResolveRenderTarget(FViewportSurfaceReader* RenderT
 			TShaderMapRef<FScreenVS> VertexShader(ShaderMap);
 			TShaderMapRef<FScreenPS> PixelShader(ShaderMap);
 
-			GraphicsPSOInit.BoundShaderState.VertexDeclarationRHI = RendererModule->GetFilterVertexDeclaration().VertexDeclarationRHI;
+			GraphicsPSOInit.BoundShaderState.VertexDeclarationRHI = GFilterVertexDeclaration.VertexDeclarationRHI;
 			GraphicsPSOInit.BoundShaderState.VertexShaderRHI = GETSAFERHISHADER_VERTEX(*VertexShader);
 			GraphicsPSOInit.BoundShaderState.PixelShaderRHI = GETSAFERHISHADER_PIXEL(*PixelShader);
 			GraphicsPSOInit.PrimitiveType = PT_TriangleList;
diff --git a/Engine/Source/Runtime/RenderCore/Private/DummyRenderResources.cpp b/Engine/Source/Runtime/RenderCore/Private/DummyRenderResources.cpp
new file mode 100644
--- /dev/null
+++ b/Engine/Source/Runtime/RenderCore/Private/DummyRenderResources.cpp
@@ -0,0 +1,62 @@
+// Copyright 1998-2018 Epic Games, Inc. All Rights Reserved.
+
+/*=============================================================================
+	DummyRenderResources.cpp: Implementations of frequently used render resources.
+=============================================================================*/
+
+#include "DummyRenderResources.h"
+#include "Containers/DynamicRHIResourceArray.h"
+
+
+TGlobalResource<FFilterVertexDeclaration> GFilterVertexDeclaration;
+TGlobalResource<FEmptyVertexDeclaration> GEmptyVertexDeclaration;
+
+TGlobalResource<FScreenRectangleVertexBuffer> GScreenRectangleVertexBuffer;
+TGlobalResource<FScreenRectangleIndexBuffer> GScreenRectangleIndexBuffer;
+
+IMPLEMENT_GLOBAL_SHADER(FVisualizeTextureVS, "/Engine/Private/Tools/FullscreenVertexShader.usf", "MainVS", SF_Vertex);
+
+
+void FScreenRectangleVertexBuffer::InitRHI()
+{
+	TResourceArray<FFilterVertex, VERTEXBUFFER_ALIGNMENT> Vertices;
+	Vertices.SetNumUninitialized(6);
+
+	Vertices[0].Position = FVector4(1, 1, 0, 1);
+	Vertices[0].UV = FVector2D(1, 1);
+
+	Vertices[1].Position = FVector4(0, 1, 0, 1);
+	Vertices[1].UV = FVector2D(0, 1);
+
+	Vertices[2].Position = FVector4(1, 0, 0, 1);
+	Vertices[2].UV = FVector2D(1, 0);
+
+	Vertices[3].Position = FVector4(0, 0, 0, 1);
+	Vertices[3].UV = FVector2D(0, 0);
+
+	//The final two vertices are used for the triangle optimization (a single triangle spans the entire viewport )
+	Vertices[4].Position = FVector4(-1, 1, 0, 1);
+	Vertices[4].UV = FVector2D(-1, 1);
+
+	Vertices[5].Position = FVector4(1, -1, 0, 1);
+	Vertices[5].UV = FVector2D(1, -1);
+
+	// Create vertex buffer. Fill buffer with initial data upon creation
+	FRHIResourceCreateInfo CreateInfo(&Vertices);
+	VertexBufferRHI = RHICreateVertexBuffer(Vertices.GetResourceDataSize(), BUF_Static, CreateInfo);
+}
+
+void FScreenRectangleIndexBuffer::InitRHI()
+{
+	// Indices 0 - 5 are used for rendering a quad. Indices 6 - 8 are used for triangle optimization.
+	const uint16 Indices[] = { 0, 1, 2, 2, 1, 3, 0, 4, 5 };
+
+	TResourceArray<uint16, INDEXBUFFER_ALIGNMENT> IndexBuffer;
+	uint32 NumIndices = ARRAY_COUNT(Indices);
+	IndexBuffer.AddUninitialized(NumIndices);
+	FMemory::Memcpy(IndexBuffer.GetData(), Indices, NumIndices * sizeof(uint16));
+
+	// Create index buffer. Fill buffer with initial data upon creation
+	FRHIResourceCreateInfo CreateInfo(&IndexBuffer);
+	IndexBufferRHI = RHICreateIndexBuffer(sizeof(uint16), IndexBuffer.GetResourceDataSize(), BUF_Static, CreateInfo);
+}
diff --git a/Engine/Source/Runtime/RenderCore/Private/PixelShaderUtils.cpp b/Engine/Source/Runtime/RenderCore/Private/PixelShaderUtils.cpp
new file mode 100644
--- /dev/null
+++ b/Engine/Source/Runtime/RenderCore/Private/PixelShaderUtils.cpp
@@ -0,0 +1,40 @@
+// Copyright 1998-2018 Epic Games, Inc. All Rights Reserved.
+
+/*=============================================================================
+	PixelShaderUtils.cpp: Implementations of utilities for pixel shaders.
+=============================================================================*/
+
+#include "PixelShaderUtils.h"
+#include "DummyRenderResources.h"
+#include "RenderGraph.h"
+
+
+// static
+void FPixelShaderUtils::DrawFullscreenTriangle(FRHICommandList& RHICmdList, uint32 InstanceCount)
+{
+	RHICmdList.SetStreamSource(0, GScreenRectangleVertexBuffer.VertexBufferRHI, 0);
+
+	RHICmdList.DrawIndexedPrimitive(
+		GScreenRectangleIndexBuffer.IndexBufferRHI,
+		/*BaseVertexIndex=*/ 0,
+		/*MinIndex=*/ 0,
+		/*NumVertices=*/ 3,
+		/*StartIndex=*/ 6,
+		/*NumPrimitives=*/ 1,
+		/*NumInstances=*/ InstanceCount);
+}
+
+// static
+void FPixelShaderUtils::DrawFullscreenQuad(FRHICommandList& RHICmdList, uint32 InstanceCount)
+{
+	RHICmdList.SetStreamSource(0, GScreenRectangleVertexBuffer.VertexBufferRHI, 0);
+
+	RHICmdList.DrawIndexedPrimitive(
+		GScreenRectangleIndexBuffer.IndexBufferRHI,
+		/*BaseVertexIndex=*/ 0,
+		/*MinIndex=*/ 0,
+		/*NumVertices=*/ 4,
+		/*StartIndex=*/ 0,
+		/*NumPrimitives=*/ 2,
+		/*NumInstances=*/ InstanceCount);
+}
diff --git a/Engine/Source/Runtime/RenderCore/Private/RenderGraphBuilder.cpp b/Engine/Source/Runtime/RenderCore/Private/RenderGraphBuilder.cpp
--- a/Engine/Source/Runtime/RenderCore/Private/RenderGraphBuilder.cpp
+++ b/Engine/Source/Runtime/RenderCore/Private/RenderGraphBuilder.cpp
@@ -4,6 +4,7 @@
 #include "RenderCore.h"
 #include "RenderTargetPool.h"
 #include "RenderGraphResourcePool.h"
+#include "VisualizeTexture.h"
 
 
 #if RENDER_GRAPH_DEBUGGING
@@ -60,6 +61,14 @@ void FRDGBuilder::DebugPass(const FRenderGraphPass* Pass)
 	{
 		ExecutePass(Pass);
 	}
+
+	// If visualizing a texture, look for any output of the pass. This must be done after the
+	// GRenderGraphImmediateMode's ExecutePass() because this will actually increate a capturing
+	// pass if needed that would have to be executed right away as well.
+	if (GVisualizeTexture.bEnabled)
+	{
+		CaptureAnyInterestingPassOutput(Pass);
+	}
 }
 
 void FRDGBuilder::ValidatePass(const FRenderGraphPass* Pass) const
@@ -131,6 +140,58 @@ void FRDGBuilder::ValidatePass(const FRenderGraphPass* Pass) const
 	}
 }
 
+void FRDGBuilder::CaptureAnyInterestingPassOutput(const FRenderGraphPass* Pass)
+{
+#if WITH_ENGINE
+	FShaderParameterStructRef ParameterStruct = Pass->GetParameters();
+	for (int ResourceIndex = 0, Num = ParameterStruct.Layout->Resources.Num(); ResourceIndex < Num; ResourceIndex++)
+	{
+		uint8  Type = ParameterStruct.Layout->Resources[ResourceIndex];
+		uint16 Offset = ParameterStruct.Layout->ResourceOffsets[ResourceIndex];
+
+		switch (Type)
+		{
+		case UBMT_GRAPH_TRACKED_UAV:
+		{
+			FRDGTextureUAV* RESTRICT UAV = *ParameterStruct.GetMemberPtrAtOffset<FRDGTextureUAV*>(Offset);
+			if (UAV && GVisualizeTexture.ShouldCapture(UAV->Desc.Texture->Name))
+			{
+				GVisualizeTexture.CreateContentCapturePass(*this, UAV->Desc.Texture);
+			}
+		}
+		break;
+		case UBMT_RENDER_TARGET_BINDING_SLOTS:
+		{
+			FRenderTargetBindingSlots* RESTRICT RenderTargets = ParameterStruct.GetMemberPtrAtOffset<FRenderTargetBindingSlots>(Offset);
+			if (RenderTargets->DepthStencil.Texture && 
+				(RenderTargets->DepthStencil.DepthStoreAction != ERenderTargetStoreAction::ENoAction || RenderTargets->DepthStencil.StencilStoreAction != ERenderTargetStoreAction::ENoAction) &&
+				GVisualizeTexture.ShouldCapture(RenderTargets->DepthStencil.Texture->Name))
+			{
+				GVisualizeTexture.CreateContentCapturePass(*this, RenderTargets->DepthStencil.Texture);
+			}
+			for (int i = 0; i < RenderTargets->Output.Num(); i++)
+			{
+				const FRenderTargetBinding& RenderTarget = RenderTargets->Output[i];
+				if (RenderTarget.GetTexture() &&
+					RenderTarget.GetStoreAction() != ERenderTargetStoreAction::ENoAction &&
+					GVisualizeTexture.ShouldCapture(RenderTarget.GetTexture()->Name))
+				{
+					GVisualizeTexture.CreateContentCapturePass(*this, RenderTarget.GetTexture());
+				}
+				else
+				{
+					break;
+				}
+			}
+		}
+		break;
+		default:
+			break;
+		}
+	}
+#endif // WITH_ENGINE
+}
+
 void FRDGBuilder::WalkGraphDependencies()
 {
 	for (const FRenderGraphPass* Pass : Passes)
diff --git a/Engine/Source/Runtime/RenderCore/Private/VisualizeTexture.cpp b/Engine/Source/Runtime/RenderCore/Private/VisualizeTexture.cpp
new file mode 100644
--- /dev/null
+++ b/Engine/Source/Runtime/RenderCore/Private/VisualizeTexture.cpp
@@ -0,0 +1,451 @@
+// Copyright 1998-2018 Epic Games, Inc. All Rights Reserved.
+
+/*=============================================================================
+	VisualizeTexture.cpp: Post processing visualize texture.
+=============================================================================*/
+
+#include "VisualizeTexture.h"
+#include "ShaderParameters.h"
+#include "RHIStaticStates.h"
+#include "Shader.h"
+#include "StaticBoundShaderState.h"
+//#include "HAL/FileManager.h"
+//#include "Misc/FileHelper.h"
+//#include "Misc/Paths.h"
+//#include "Misc/App.h"
+#include "RenderTargetPool.h"
+#include "GlobalShader.h"
+#include "PipelineStateCache.h"
+#include "DummyRenderResources.h"
+#include "PixelShaderUtils.h"
+
+
+FVisualizeTexture::FVisualizeTexture()
+{
+	Mode = 0;
+	RGBMul = 1.0f;
+	SingleChannelMul = 0.0f;
+	SingleChannel = -1;
+	AMul = 0.0f;
+	UVInputMapping = 3;
+	Flags = 0;
+	ObservedDebugNameReusedGoal = 0xffffffff;
+	ArrayIndex = 0;
+	CustomMip = 0;
+	bSaveBitmap = false;
+	bOutputStencil = false;
+	bFullList = false;
+	SortOrder = -1;
+	bEnabled = true;
+}
+
+
+#if WITH_ENGINE
+
+enum class EVisualisePSType
+{
+	Cube = 0,
+	Texture1D = 1, //not supported
+	Texture2DNoMSAA = 2,
+	Texture3D = 3,
+	CubeArray = 4,
+	Texture2DMSAA = 5,
+	Texture2DDepthStencilNoMSAA = 6,
+	Texture2DUINT8 = 7,
+	MAX
+};
+
+TGlobalResource<FVisualizeTexture> GVisualizeTexture;
+
+
+/** A pixel shader which filters a texture. */
+// @param TextureType 0:Cube, 1:1D(not yet supported), 2:2D no MSAA, 3:3D, 4:Cube[], 5:2D MSAA, 6:2D DepthStencil no MSAA (needed to avoid D3DDebug error)
+class FVisualizeTexturePS : public FGlobalShader
+{
+	DECLARE_GLOBAL_SHADER(FVisualizeTexturePS);
+	SHADER_USE_PARAMETER_STRUCT(FVisualizeTexturePS, FGlobalShader);
+
+	class FVisualisePSTypeDim : SHADER_PERMUTATION_ENUM_CLASS("TEXTURE_TYPE", EVisualisePSType);
+
+	using FPermutationDomain = TShaderPermutationDomain<FVisualisePSTypeDim>;
+
+	static bool ShouldCompilePermutation(const FGlobalShaderPermutationParameters& Parameters)
+	{
+		FPermutationDomain PermutationVector(Parameters.PermutationId);
+		return PermutationVector.Get<FVisualisePSTypeDim>() != EVisualisePSType::Texture1D;
+	}
+
+	BEGIN_SHADER_PARAMETER_STRUCT(FParameters, )
+		SHADER_PARAMETER(FVector, TextureExtent)
+		SHADER_PARAMETER_ARRAY(FVector4, VisualizeParam, [3])
+
+		SHADER_PARAMETER_RDG_TEXTURE(Texture2D, VisualizeTexture2D)
+		SHADER_PARAMETER_SAMPLER(SamplerState, VisualizeTexture2DSampler)
+		SHADER_PARAMETER_RDG_TEXTURE(Texture3D, VisualizeTexture3D)
+		SHADER_PARAMETER_SAMPLER(SamplerState, VisualizeTexture3DSampler)
+		SHADER_PARAMETER_RDG_TEXTURE(TextureCube, VisualizeTextureCube)
+		SHADER_PARAMETER_SAMPLER(SamplerState, VisualizeTextureCubeSampler)
+		SHADER_PARAMETER_RDG_TEXTURE(TextureCubeArray, VisualizeTextureCubeArray)
+		SHADER_PARAMETER_SAMPLER(SamplerState, VisualizeTextureCubeArraySampler)
+		SHADER_PARAMETER_RDG_TEXTURE(Texture2D<uint4>, VisualizeDepthStencilTexture)
+		SHADER_PARAMETER_RDG_TEXTURE(Texture2DMS<float4>, VisualizeTexture2DMS)
+		SHADER_PARAMETER_RDG_TEXTURE(Texture2D<uint>, VisualizeUINT8Texture2D)
+
+		RENDER_TARGET_BINDING_SLOTS()
+	END_SHADER_PARAMETER_STRUCT()
+};
+
+IMPLEMENT_GLOBAL_SHADER(FVisualizeTexturePS, "/Engine/Private/Tools/VisualizeTexture.usf", "VisualizeTexturePS", SF_Pixel);
+
+
+static EVisualisePSType GetVisualizePSType(const FRDGTextureDesc& Desc)
+{
+	if(Desc.Is2DTexture())
+	{
+		// 2D		
+		if(Desc.NumSamples > 1)
+		{
+			// MSAA
+			return EVisualisePSType::Texture2DMSAA;
+		}
+		else
+		{
+			if(Desc.Format == PF_DepthStencil)
+			{
+				// DepthStencil non MSAA (needed to avoid D3DDebug error)
+				return EVisualisePSType::Texture2DDepthStencilNoMSAA;
+			}
+			else if (Desc.Format == PF_R8_UINT)
+			{
+				return EVisualisePSType::Texture2DUINT8;
+			}
+			else
+			{
+				// non MSAA
+				return EVisualisePSType::Texture2DNoMSAA;
+			}
+		}
+	}
+	else if(Desc.IsCubemap())
+	{		
+		if(Desc.IsArray())
+		{
+			// Cube[]
+			return EVisualisePSType::CubeArray;
+		}
+		else
+		{
+			// Cube
+			return EVisualisePSType::Cube;
+		}
+	}
+
+	check(Desc.Is3DTexture());
+	return EVisualisePSType::Texture3D;
+}
+
+void FVisualizeTexture::ReleaseDynamicRHI()
+{
+	VisualizeTextureContent.SafeRelease();
+	StencilSRV.SafeRelease();
+}
+
+void FVisualizeTexture::CreateContentCapturePass(FRDGBuilder& GraphBuilder, const FRDGTextureRef SrcTexture)
+{
+#if !(UE_BUILD_SHIPPING || UE_BUILD_TEST)
+	if (!SrcTexture || !SrcTexture->Desc.IsValid())
+	{
+		// todo: improve
+		return;
+	}
+
+	const FRDGTextureDesc& SrcDesc = SrcTexture->Desc;
+
+	if ((SrcDesc.Flags & TexCreate_CPUReadback))
+	{
+		// We cannot make a texture lookup on such elements
+		return;
+	}
+
+	FRDGTextureRef CopyTexture;
+	{
+		FIntPoint Size = SrcTexture->Desc.Extent;
+
+		// clamp to reasonable value to prevent crash
+		Size.X = FMath::Max(Size.X, 1);
+		Size.Y = FMath::Max(Size.Y, 1);
+
+		FRDGTextureDesc CopyDesc = FRDGTextureDesc::Create2DDesc(
+			Size, PF_B8G8R8A8, FClearValueBinding(FLinearColor(1, 1, 0, 1)),
+			TexCreate_None, TexCreate_RenderTargetable | TexCreate_ShaderResource, false);
+
+		CopyTexture = GraphBuilder.CreateTexture(CopyDesc, TEXT("VisualizeTexture"));
+	}
+
+	FIntPoint RTExtent = SrcTexture->Desc.Extent;
+
+	uint32 LocalVisualizeTextureInputMapping = UVInputMapping;
+
+	if (!SrcDesc.Is2DTexture())
+	{
+		LocalVisualizeTextureInputMapping = 1;
+	}
+
+#if 0 // TODO(RDG): requires changes in FRDGTextureSRVDesc
+	bool bIsDefault = StencilSRVSrc == GBlackTexture->TextureRHI;
+	bool bDepthStencil = SrcDesc.Is2DTexture() && SrcDesc.Format == PF_DepthStencil;
+
+	//clear if this is a new different Stencil buffer, or it's not a stencil buffer and we haven't switched to the default yet.
+	bool bNeedsClear = bDepthStencil && (StencilSRVSrc != RenderTargetItem.TargetableTexture);
+	bNeedsClear |= !bDepthStencil && !bIsDefault;
+	if (bNeedsClear)
+	{
+		StencilSRVSrc = nullptr;
+		StencilSRV.SafeRelease();
+	}
+
+	//always set something into the StencilSRV slot for platforms that require a full resource binding, even if
+	//dynamic branching will cause them not to be used.	
+	if (bDepthStencil && !StencilSRVSrc)
+	{
+		StencilSRVSrc = RenderTargetItem.TargetableTexture;
+		StencilSRV = RHICreateShaderResourceView((FTexture2DRHIRef&)RenderTargetItem.TargetableTexture, 0, 1, PF_X24_G8);
+	}
+	else if (!StencilSRVSrc)
+	{
+		StencilSRVSrc = GBlackTexture->TextureRHI;
+		StencilSRV = RHICreateShaderResourceView((FTexture2DRHIRef&)GBlackTexture->TextureRHI, 0, 1, PF_B8G8R8A8);
+	}
+#endif
+
+	// distinguish between standard depth and shadow depth to produce more reasonable default value mapping in the pixel shader.
+	const bool bDepthTexture = (SrcDesc.TargetableFlags & TexCreate_DepthStencilTargetable) != 0;
+	const bool bShadowDepth = (SrcDesc.Format == PF_ShadowDepth);
+
+	bool bSaturateInsteadOfFrac = (Flags & 1) != 0;
+	int32 InputValueMapping = bShadowDepth ? 2 : (bDepthTexture ? 1 : 0);
+
+	FVisualizeTexturePS::FParameters *PassParameters;
+	GraphBuilder.CreateParameters(&PassParameters);
+	{
+		PassParameters->TextureExtent = FVector(SrcDesc.Extent.X, SrcDesc.Extent.Y, SrcDesc.Depth);
+
+		{
+			// alternates between 0 and 1 with a short pause
+			const float FracTimeScale = 2.0f;
+			float FracTime = FApp::GetCurrentTime() * FracTimeScale - floor(FApp::GetCurrentTime() * FracTimeScale);
+			float BlinkState = (FracTime > 0.5f) ? 1.0f : 0.0f;
+
+			FVector4 VisualizeParamValue[3];
+
+			float Add = 0.0f;
+			float FracScale = 1.0f;
+
+			// w * almost_1 to avoid frac(1) => 0
+			PassParameters->VisualizeParam[0] = FVector4(RGBMul, SingleChannelMul, Add, FracScale * 0.9999f);
+			PassParameters->VisualizeParam[1] = FVector4(BlinkState, bSaturateInsteadOfFrac ? 1.0f : 0.0f, ArrayIndex, CustomMip);
+			PassParameters->VisualizeParam[2] = FVector4(InputValueMapping, 0.0f, SingleChannel);
+		}
+
+		FSamplerStateRHIParamRef PointSampler = TStaticSamplerState<SF_Point, AM_Clamp, AM_Clamp, AM_Clamp>::GetRHI();
+
+		PassParameters->VisualizeTexture2D = SrcTexture;
+		PassParameters->VisualizeTexture2DSampler = PointSampler;
+		PassParameters->VisualizeTexture3D = SrcTexture;
+		PassParameters->VisualizeTexture3DSampler = PointSampler;
+		PassParameters->VisualizeTextureCube = SrcTexture;
+		PassParameters->VisualizeTextureCubeSampler = PointSampler;
+		PassParameters->VisualizeTextureCubeArray = SrcTexture;
+		PassParameters->VisualizeTextureCubeArraySampler = PointSampler;
+
+		PassParameters->VisualizeDepthStencilTexture = nullptr; // TODO: StencilSRV
+		PassParameters->VisualizeTexture2DMS = SrcTexture;
+		PassParameters->VisualizeUINT8Texture2D = SrcTexture;
+
+		PassParameters->RenderTargets[0] = FRenderTargetBinding(CopyTexture, ERenderTargetLoadAction::EClear, ERenderTargetStoreAction::EStore);
+	}
+
+	auto ShaderMap = GetGlobalShaderMap(FeatureLevel);
+	FVisualizeTexturePS::FPermutationDomain PermutationVector;
+	PermutationVector.Set<FVisualizeTexturePS::FVisualisePSTypeDim>(GetVisualizePSType(SrcDesc));
+
+	TShaderMapRef<FVisualizeTexturePS> PixelShader(ShaderMap, PermutationVector);
+
+	FPixelShaderUtils::AddFullscreenPass(
+		GraphBuilder, ShaderMap,
+		TEXT("VisualizeTexture"),
+		*PixelShader, PassParameters);
+
+	// Save the copied texture and descriptor about original informations.
+	{
+		VisualizeTextureDesc = SrcDesc;
+		VisualizeTextureContent = nullptr;
+		GraphBuilder.GetInternalTexture(CopyTexture, &VisualizeTextureContent);
+	}
+
+#if 0 // TODO(RDG): requires some kind of CPU readback pass.
+	// save to disk
+	if (bSaveBitmap)
+	{
+		bSaveBitmap = false;
+
+		uint32 MipAdjustedExtentX = FMath::Clamp(Desc.Extent.X >> CustomMip, 0, Desc.Extent.X);
+		uint32 MipAdjustedExtentY = FMath::Clamp(Desc.Extent.Y >> CustomMip, 0, Desc.Extent.Y);
+		FIntPoint Extent(MipAdjustedExtentX, MipAdjustedExtentY);
+
+		FReadSurfaceDataFlags ReadDataFlags;
+		ReadDataFlags.SetLinearToGamma(false);
+		ReadDataFlags.SetOutputStencil(bOutputStencil);
+		ReadDataFlags.SetMip(CustomMip);
+
+		FTextureRHIRef Texture = RenderTargetItem.TargetableTexture ? RenderTargetItem.TargetableTexture : RenderTargetItem.ShaderResourceTexture;
+
+		check(Texture);
+
+		TArray<FColor> Bitmap;
+
+
+
+		RHICmdList.ReadSurfaceData(Texture, FIntRect(0, 0, Extent.X, Extent.Y), Bitmap, ReadDataFlags);
+
+		// if the format and texture type is supported
+		if (Bitmap.Num())
+		{
+			// Create screenshot folder if not already present.
+			IFileManager::Get().MakeDirectory(*FPaths::ScreenShotDir(), true);
+
+			const FString ScreenFileName(FPaths::ScreenShotDir() / TEXT("VisualizeTexture"));
+
+			uint32 ExtendXWithMSAA = Bitmap.Num() / Extent.Y;
+
+			// Save the contents of the array to a bitmap file. (24bit only so alpha channel is dropped)
+			FFileHelper::CreateBitmap(*ScreenFileName, ExtendXWithMSAA, Extent.Y, Bitmap.GetData());
+
+			UE_LOG(LogConsoleResponse, Display, TEXT("Content was saved to \"%s\""), *FPaths::ScreenShotDir());
+		}
+		else
+		{
+			UE_LOG(LogConsoleResponse, Error, TEXT("Failed to save BMP for VisualizeTexture, format or texture type is not supported"));
+		}
+	}
+#endif
+
+
+#endif // !(UE_BUILD_SHIPPING || UE_BUILD_TEST)
+}
+
+bool FVisualizeTexture::ShouldCapture(const TCHAR* DebugName)
+{
+#if (UE_BUILD_SHIPPING || UE_BUILD_TEST)
+	return false;
+#else
+	if (!bEnabled)
+	{
+		return false;
+	}
+
+
+	uint32* UsageCountPtr = VisualizeTextureCheckpoints.Find(DebugName);
+
+	if (!UsageCountPtr)
+	{
+		// create a new element with count 0
+		UsageCountPtr = &VisualizeTextureCheckpoints.Add(DebugName, 0);
+	}
+
+	// is this is the name we are observing with visualize texture?
+	// First check if we need to find anything to avoid string the comparison
+	if (!ObservedDebugName.IsEmpty() && ObservedDebugName == DebugName)
+	{
+		// if multiple times reused during the frame, is that the one we want to look at?
+		if (*UsageCountPtr == ObservedDebugNameReusedGoal || ObservedDebugNameReusedGoal == 0xffffffff)
+		{
+			*UsageCountPtr = *UsageCountPtr + 1;
+			return true;
+		}
+	}
+	// only needed for VisualizeTexture (todo: optimize out when possible)
+	*UsageCountPtr = *UsageCountPtr + 1;
+	return false;
+#endif
+}
+
+#if !(UE_BUILD_SHIPPING || UE_BUILD_TEST)
+
+void FVisualizeTexture::SetCheckPoint(FRHICommandList& RHICmdList, const IPooledRenderTarget* PooledRenderTarget)
+{
+	check(IsInRenderingThread());
+	if (!PooledRenderTarget)
+	{
+		return;
+	}
+
+	const TCHAR* DebugName = PooledRenderTarget->GetDesc().DebugName;
+
+	if (!ShouldCapture(DebugName))
+	{
+		return;
+	}
+
+	FRHICommandListImmediate& RHICmdListIm = FRHICommandListExecutor::GetImmediateCommandList();
+	if (RHICmdListIm.IsExecuting())
+	{
+		UE_LOG(LogConsoleResponse, Fatal, TEXT("We can't create a checkpoint because that requires the immediate commandlist, which is currently executing. You might try disabling parallel rendering."));
+	}
+
+	if (&RHICmdList != &RHICmdListIm)
+	{
+		UE_LOG(LogConsoleResponse, Warning, TEXT("Attempt to checkpoint a render target from a non-immediate command list. We will flush it and hope that works. If it doesn't you might try disabling parallel rendering."));
+		RHICmdList.Flush();
+	}
+
+	FRDGBuilder GraphBuilder(RHICmdListIm);
+
+	// Sorry for the const cast here only required for reference count of the pooled render target the graph needs to do.
+	// Long therm this SetCheckPoint() method should no longer since it is done exclusively by render graph automatically. 
+	TRefCountPtr<IPooledRenderTarget> PooledRenderTargetRef(const_cast<IPooledRenderTarget*>(PooledRenderTarget));
+	FRDGTextureRef TextureToCapture = GraphBuilder.RegisterExternalTexture(PooledRenderTargetRef, DebugName);
+
+	CreateContentCapturePass(GraphBuilder, TextureToCapture);
+	GraphBuilder.Execute();
+
+	if (&RHICmdList != &RHICmdListIm)
+	{
+		RHICmdListIm.Flush();
+	}
+}
+
+#endif // !(UE_BUILD_SHIPPING || UE_BUILD_TEST)
+
+void FVisualizeTexture::QueryInfo_GameThread( FQueryVisualizeTexureInfo& Out )
+{
+	check(IsInGameThread());
+	FlushRenderingCommands();
+
+	for(uint32 i = 0, Num = GRenderTargetPool.GetElementCount(); i < Num; ++i)
+	{
+		FPooledRenderTarget* RT = GRenderTargetPool.GetElementById(i);
+
+		if(!RT)
+		{
+			continue;
+		}
+
+		FPooledRenderTargetDesc Desc = RT->GetDesc();
+		uint32 SizeInKB = (RT->ComputeMemorySize() + 1023) / 1024;
+		FString Entry = FString::Printf(TEXT("%s %d %s %d"),
+				*Desc.GenerateInfoString(),
+				i + 1,
+				Desc.DebugName ? Desc.DebugName : TEXT("<Unnamed>"),
+				SizeInKB);
+		Out.Entries.Add(Entry);
+	}
+}
+
+void FVisualizeTexture::SetRenderTargetNameToObserve(const FString& InObservedDebugName, uint32 InObservedDebugNameReusedGoal)
+{
+	ObservedDebugName = InObservedDebugName;
+	ObservedDebugNameReusedGoal = InObservedDebugNameReusedGoal;
+}
+
+#endif // WITH_ENGINE
diff --git a/Engine/Source/Runtime/RenderCore/Public/DummyRenderResources.h b/Engine/Source/Runtime/RenderCore/Public/DummyRenderResources.h
new file mode 100644
--- /dev/null
+++ b/Engine/Source/Runtime/RenderCore/Public/DummyRenderResources.h
@@ -0,0 +1,105 @@
+// Copyright 1998-2018 Epic Games, Inc. All Rights Reserved.
+
+/*=============================================================================
+	DummyRenderResource.h: Frequently used rendering resources
+=============================================================================*/
+
+#pragma once
+
+#include "RenderResource.h"
+#include "GlobalShader.h"
+#include "ShaderParameterStruct.h"
+
+
+/** The vertex data used to filter a texture. */
+struct FFilterVertex
+{
+	FVector4 Position;
+	FVector2D UV;
+};
+
+/** The filter vertex declaration resource type. */
+class FFilterVertexDeclaration : public FRenderResource
+{
+public:
+	FVertexDeclarationRHIRef VertexDeclarationRHI;
+
+	/** Destructor. */
+	virtual ~FFilterVertexDeclaration() {}
+
+	virtual void InitRHI()
+	{
+		FVertexDeclarationElementList Elements;
+		uint32 Stride = sizeof(FFilterVertex);
+		Elements.Add(FVertexElement(0, STRUCT_OFFSET(FFilterVertex, Position), VET_Float4, 0, Stride));
+		Elements.Add(FVertexElement(0, STRUCT_OFFSET(FFilterVertex, UV), VET_Float2, 1, Stride));
+		VertexDeclarationRHI = RHICreateVertexDeclaration(Elements);
+	}
+
+	virtual void ReleaseRHI()
+	{
+		VertexDeclarationRHI.SafeRelease();
+	}
+};
+
+extern RENDERCORE_API TGlobalResource<FFilterVertexDeclaration> GFilterVertexDeclaration;
+
+/** The empty vertex declaration resource type. */
+class FEmptyVertexDeclaration : public FRenderResource
+{
+public:
+	FVertexDeclarationRHIRef VertexDeclarationRHI;
+
+	/** Destructor. */
+	virtual ~FEmptyVertexDeclaration() {}
+
+	virtual void InitRHI()
+	{
+		FVertexDeclarationElementList Elements;
+		VertexDeclarationRHI = RHICreateVertexDeclaration(Elements);
+	}
+
+	virtual void ReleaseRHI()
+	{
+		VertexDeclarationRHI.SafeRelease();
+	}
+};
+
+extern RENDERCORE_API TGlobalResource<FEmptyVertexDeclaration> GEmptyVertexDeclaration;
+
+/**
+ * Static vertex and index buffer used for 2D screen rectangles.
+ */
+class FScreenRectangleVertexBuffer : public FVertexBuffer
+{
+public:
+	/** Initialize the RHI for this rendering resource */
+	void InitRHI() override;
+};
+
+extern RENDERCORE_API TGlobalResource<FScreenRectangleVertexBuffer> GScreenRectangleVertexBuffer;
+
+
+class FScreenRectangleIndexBuffer : public FIndexBuffer
+{
+public:
+	/** Initialize the RHI for this rendering resource */
+	void InitRHI() override;
+};
+
+extern RENDERCORE_API TGlobalResource<FScreenRectangleIndexBuffer> GScreenRectangleIndexBuffer;
+
+
+/** Vertex shader to draw a full screen quad that works on all platforms. */
+class RENDERCORE_API FVisualizeTextureVS : public FGlobalShader
+{
+	DECLARE_GLOBAL_SHADER(FVisualizeTextureVS);
+	SHADER_USE_PARAMETER_STRUCT(FVisualizeTextureVS, FGlobalShader);
+
+	static bool ShouldCompilePermutation(const FGlobalShaderPermutationParameters& Parameters) {
+		return true;
+	}
+
+	BEGIN_SHADER_PARAMETER_STRUCT(FParameters, )
+	END_SHADER_PARAMETER_STRUCT()
+};
diff --git a/Engine/Source/Runtime/RenderCore/Public/PixelShaderUtils.h b/Engine/Source/Runtime/RenderCore/Public/PixelShaderUtils.h
new file mode 100644
--- /dev/null
+++ b/Engine/Source/Runtime/RenderCore/Public/PixelShaderUtils.h
@@ -0,0 +1,70 @@
+// Copyright 1998-2018 Epic Games, Inc. All Rights Reserved.
+
+/*=============================================================================
+	PixelShaderUtils.h: Utilities for pixel shaders.
+=============================================================================*/
+
+#pragma once
+
+#include "RenderResource.h"
+#include "RenderGraphUtils.h"
+
+
+/** All utils for pixel shaders. */
+struct RENDERCORE_API FPixelShaderUtils
+{
+	/** Draw a single triangle on the entire viewport. */
+	static void DrawFullscreenTriangle(FRHICommandList& RHICmdList, uint32 InstanceCount = 1);
+
+	/** Draw a two triangle on the entire viewport. */
+	static void DrawFullscreenQuad(FRHICommandList& RHICmdList, uint32 InstanceCount = 1);
+
+
+	/** Dispatch a full screen pixel shader to rhi command list with its parameters. */
+	template<typename TShaderClass>
+	static inline void DrawFullscreenPixelShader(
+		FRHICommandList& RHICmdList, 
+		const TShaderMap<FGlobalShaderType>* GlobalShaderMap,
+		const TShaderClass* PixelShader,
+		const typename TShaderClass::FParameters& Parameters)
+	{
+		TShaderMapRef<FVisualizeTextureVS> VertexShader(GlobalShaderMap);
+
+		FGraphicsPipelineStateInitializer GraphicsPSOInit;
+		RHICmdList.ApplyCachedRenderTargets(GraphicsPSOInit);
+		GraphicsPSOInit.BlendState = TStaticBlendState<>::GetRHI();
+		GraphicsPSOInit.RasterizerState = TStaticRasterizerState<>::GetRHI();
+		GraphicsPSOInit.DepthStencilState = TStaticDepthStencilState<false, CF_Always>::GetRHI();
+
+		GraphicsPSOInit.BoundShaderState.VertexDeclarationRHI = GFilterVertexDeclaration.VertexDeclarationRHI;
+		GraphicsPSOInit.BoundShaderState.VertexShaderRHI = GETSAFERHISHADER_VERTEX(*VertexShader);
+		GraphicsPSOInit.BoundShaderState.PixelShaderRHI = GETSAFERHISHADER_PIXEL(PixelShader);
+		GraphicsPSOInit.PrimitiveType = PT_TriangleList;
+		SetGraphicsPipelineState(RHICmdList, GraphicsPSOInit);
+
+		SetShaderParameters(RHICmdList, PixelShader, PixelShader->GetPixelShader(), Parameters);
+
+		DrawFullscreenTriangle(RHICmdList);
+	}
+
+	/** Dispatch a pixel shader to render graph builder with its parameters. */
+	template<typename TShaderClass>
+	static inline void AddFullscreenPass(
+		FRDGBuilder& GraphBuilder,
+		const TShaderMap<FGlobalShaderType>* GlobalShaderMap,
+		const TCHAR* PassName,
+		const TShaderClass* PixelShader,
+		typename TShaderClass::FParameters* Parameters)
+	{
+		ClearUnusedGraphResources(PixelShader, Parameters);
+
+		GraphBuilder.AddPass(
+			PassName,
+			Parameters,
+			ERenderGraphPassFlags::None,
+			[Parameters, GlobalShaderMap, PixelShader](FRHICommandList& RHICmdList)
+		{
+			FPixelShaderUtils::DrawFullscreenPixelShader(RHICmdList, GlobalShaderMap, PixelShader, *Parameters);
+		});
+	}
+};
diff --git a/Engine/Source/Runtime/RenderCore/Public/RenderGraphBuilder.h b/Engine/Source/Runtime/RenderCore/Public/RenderGraphBuilder.h
--- a/Engine/Source/Runtime/RenderCore/Public/RenderGraphBuilder.h
+++ b/Engine/Source/Runtime/RenderCore/Public/RenderGraphBuilder.h
@@ -333,6 +333,7 @@ private:
 
 	void DebugPass(const FRenderGraphPass* Pass);
 	void ValidatePass(const FRenderGraphPass* Pass) const;
+	void CaptureAnyInterestingPassOutput(const FRenderGraphPass* Pass);
 
 	void WalkGraphDependencies();
 
diff --git a/Engine/Source/Runtime/RenderCore/Public/RenderTargetPool.h b/Engine/Source/Runtime/RenderCore/Public/RenderTargetPool.h
--- a/Engine/Source/Runtime/RenderCore/Public/RenderTargetPool.h
+++ b/Engine/Source/Runtime/RenderCore/Public/RenderTargetPool.h
@@ -260,9 +260,6 @@ public:
 	// @return -1 if not found
 	int32 FindIndex(IPooledRenderTarget* In) const;
 
-	// Render visualization
-	void RenderVisualizeTexture(class FDeferredShadingSceneRenderer& Scene);
-
 	void SetObserveTarget(const FString& InObservedDebugName, uint32 InObservedDebugNameReusedGoal = 0xffffffff);
 
 	// Logs out usage information.
@@ -360,6 +357,7 @@ private:
 
 	friend struct FPooledRenderTarget;
 	friend class FVisualizeTexture;
+	friend class FVisualizeTexturePresent;
 };
 
 /** The global render targets for easy shading. */
diff --git a/Engine/Source/Runtime/RenderCore/Public/RendererInterface.h b/Engine/Source/Runtime/RenderCore/Public/RendererInterface.h
--- a/Engine/Source/Runtime/RenderCore/Public/RendererInterface.h
+++ b/Engine/Source/Runtime/RenderCore/Public/RendererInterface.h
@@ -491,58 +491,6 @@ struct FQueryVisualizeTexureInfo
 };
 
 
-/** The vertex data used to filter a texture. */
-struct FFilterVertex
-{
-	FVector4 Position;
-	FVector2D UV;
-};
-
-/** The filter vertex declaration resource type. */
-class FFilterVertexDeclaration : public FRenderResource
-{
-public:
-	FVertexDeclarationRHIRef VertexDeclarationRHI;
-
-	/** Destructor. */
-	virtual ~FFilterVertexDeclaration() {}
-
-	virtual void InitRHI()
-	{
-		FVertexDeclarationElementList Elements;
-		uint32 Stride = sizeof(FFilterVertex);
-		Elements.Add(FVertexElement(0,STRUCT_OFFSET(FFilterVertex,Position),VET_Float4,0,Stride));
-		Elements.Add(FVertexElement(0,STRUCT_OFFSET(FFilterVertex,UV),VET_Float2,1,Stride));
-		VertexDeclarationRHI = RHICreateVertexDeclaration(Elements);
-	}
-
-	virtual void ReleaseRHI()
-	{
-		VertexDeclarationRHI.SafeRelease();
-	}
-};
-
-/** The empty vertex declaration resource type. */
-class FEmptyVertexDeclaration : public FRenderResource
-{
-public:
-	FVertexDeclarationRHIRef VertexDeclarationRHI;
-
-	/** Destructor. */
-	virtual ~FEmptyVertexDeclaration() {}
-
-	virtual void InitRHI()
-	{
-		FVertexDeclarationElementList Elements;
-		VertexDeclarationRHI = RHICreateVertexDeclaration(Elements);
-	}
-
-	virtual void ReleaseRHI()
-	{
-		VertexDeclarationRHI.SafeRelease();
-	}
-};
-
 // use r.DrawDenormalizedQuadMode to override the function call setting (quick way to see if an artifact is caused by this optimization)
 enum EDrawRectangleFlags
 {
@@ -762,9 +710,11 @@ public:
 	virtual void DrawTileMesh(FRHICommandListImmediate& RHICmdList, struct FDrawingPolicyRenderState& DrawRenderState, const FSceneView& View, const FMeshBatch& Mesh, bool bIsHitTesting, const class FHitProxyId& HitProxyId) = 0;
 
 	/** Render thread side, use TRefCountPtr<IPooledRenderTarget>, allows to use sharing and VisualizeTexture */
+	// TODO(RDG): Kill that guy.
 	virtual void RenderTargetPoolFindFreeElement(FRHICommandListImmediate& RHICmdList, const FPooledRenderTargetDesc& Desc, TRefCountPtr<IPooledRenderTarget> &Out, const TCHAR* InDebugName) = 0;
 	
 	/** Render thread side, to age the pool elements so they get released at some point */
+	// TODO(RDG): Kill that guy.
 	virtual void TickRenderTargetPool() = 0;
 
 	virtual const TSet<FSceneInterface*>& GetAllocatedScenes() = 0;
@@ -775,7 +725,6 @@ public:
 	// @param WorkScale >0, 10 for normal precision and runtime of less than a second
 	virtual void GPUBenchmark(FSynthBenchmarkResults& InOut, float WorkScale = 10.0f) = 0;
 
-	virtual void QueryVisualizeTexture(FQueryVisualizeTexureInfo& Out) = 0;
 	virtual void ExecVisualizeTextureCmd(const FString& Cmd) = 0;
 
 	virtual void UpdateMapNeedsLightingFullyRebuiltState(UWorld* World) = 0;
@@ -811,9 +760,6 @@ public:
 		EDrawRectangleFlags Flags = EDRF_Default
 		) = 0;
 
-	/** @return Returns a vertex declaration that can be used with with the DrawRectangle() function */
-	virtual TGlobalResource<FFilterVertexDeclaration>& GetFilterVertexDeclaration() = 0;
-
 	/** Register/unregister a custom occlusion culling implementation */
 	virtual void RegisterCustomCullingImpl(ICustomCulling* impl) = 0;
 	virtual void UnregisterCustomCullingImpl(ICustomCulling* impl) = 0;
diff --git a/Engine/Source/Runtime/RenderCore/Public/VisualizeTexture.h b/Engine/Source/Runtime/RenderCore/Public/VisualizeTexture.h
new file mode 100644
--- /dev/null
+++ b/Engine/Source/Runtime/RenderCore/Public/VisualizeTexture.h
@@ -0,0 +1,116 @@
+// Copyright 1998-2018 Epic Games, Inc. All Rights Reserved.
+
+/*=============================================================================
+	VisualizeTexture.h: Post processing visualize texture.
+=============================================================================*/
+
+#pragma once
+
+#include "CoreMinimal.h"
+#include "RendererInterface.h"
+#include "RenderResource.h"
+#include "RenderGraph.h"
+
+
+class RENDERCORE_API FVisualizeTexture : public FRenderResource
+{
+public:
+	FVisualizeTexture();
+
+#if WITH_ENGINE
+	virtual void ReleaseDynamicRHI() override;
+	
+	/**
+	 * calling this allows to grab the state of the texture at this point to be queried by visualizetexture e.g. "vis LightAttenuation@2"
+	 * @param PooledRenderTarget 0 is silently ignored
+	 * Warning: this may change the active render target and other state
+	 */
+#if !(UE_BUILD_SHIPPING || UE_BUILD_TEST)
+	void SetCheckPoint(FRHICommandList& RHICmdList, const IPooledRenderTarget* PooledRenderTarget);
+#else
+	inline void SetCheckPoint(FRHICommandList& RHICmdList, const IPooledRenderTarget* PooledRenderTarget) {}
+#endif
+
+	/** Query some information from game thread. */
+	// TODO: refactor
+	void QueryInfo_GameThread( FQueryVisualizeTexureInfo& Out );
+
+	/** Sets the render  */
+	void SetRenderTargetNameToObserve(const FString& InObservedDebugName, uint32 InObservedDebugNameReusedGoal = 0xffffffff);
+#endif
+
+	// VisualizeTexture console command settings:
+	// written on game thread, read on render thread (uses FlushRenderingCommands to avoid the threading issues)
+
+	// 0=off, >0=texture id, changed by "VisualizeTexture" console command, useful for debugging
+	int32 Mode;
+	//
+	float RGBMul;
+
+	// -1=off, 0=R, 1=G, 2=B, 3=A
+	int32 SingleChannel;
+
+	// Multiplier for the single channel
+	float SingleChannelMul;
+
+	//
+	float AMul;
+	// 0=view in left top, 1=whole texture, 2=pixel perfect centered, 3=Picture in Picture
+	int32 UVInputMapping;
+	// bit 1: if 1, saturation mode, if 0, frac mode
+	int32 Flags;
+	//
+	int32 CustomMip;
+	//
+	int32 ArrayIndex;
+	//
+	bool bSaveBitmap;
+	// stencil normally displays in the alpha channel of depth buffer visualization.  This option is just for BMP writeout to get a stencil only BMP.
+	bool bOutputStencil;
+	//
+	bool bFullList;
+	// -1:by index, 0:by name, 1:by size
+	int32 SortOrder;
+
+#if !(UE_BUILD_SHIPPING || UE_BUILD_TEST)
+	// [DebugName of the RT] = ReuseCount this frame
+	TMap<const TCHAR*, uint32> VisualizeTextureCheckpoints;
+#endif
+
+	// render target DebugName that is observed, "" if the feature is deactivated
+	FString ObservedDebugName;
+	// each frame this is counting up each time a RT with the same name is reused
+	uint32 ObservedDebugNameReusedCurrent;
+	// this is the count we want to reach, 0xffffffff if the last one
+	uint32 ObservedDebugNameReusedGoal;
+
+private:
+	// Copy of the texture being visualized.
+	TRefCountPtr<IPooledRenderTarget> VisualizeTextureContent;
+
+	/** Descriptor of the texture being visualized. */
+	FPooledRenderTargetDesc VisualizeTextureDesc;
+
+	TRefCountPtr<FRHIShaderResourceView> StencilSRV;
+	FTextureRHIRef StencilSRVSrc;
+
+	// Flag to determine whether texture visualization is enabled, currently based on the feature level we are rendering with
+	bool bEnabled;
+
+	// Store feature level that we're currently using
+	ERHIFeatureLevel::Type FeatureLevel;
+
+#if WITH_ENGINE
+	bool ShouldCapture(const TCHAR* DebugName);
+
+	/** Create a pass capturing a texture. */
+	void CreateContentCapturePass(FRDGBuilder& GraphBuilder, const FRDGTextureRef Texture);
+#endif
+
+	friend class FRDGBuilder;
+	friend class FVisualizeTexturePresent;
+};
+
+
+/** The global render targets for easy shading. */
+extern RENDERCORE_API TGlobalResource<FVisualizeTexture> GVisualizeTexture;
diff --git a/Engine/Source/Runtime/Renderer/Private/CompositionLighting/CompositionLighting.cpp b/Engine/Source/Runtime/Renderer/Private/CompositionLighting/CompositionLighting.cpp
--- a/Engine/Source/Runtime/Renderer/Private/CompositionLighting/CompositionLighting.cpp
+++ b/Engine/Source/Runtime/Renderer/Private/CompositionLighting/CompositionLighting.cpp
@@ -16,6 +16,7 @@
 #include "PostProcess/PostProcessSubsurface.h"
 #include "LightPropagationVolumeSettings.h"
 #include "DecalRenderingShared.h"
+#include "VisualizeTexture.h"
 
 /** The global center for all deferred lighting activities. */
 FCompositionLighting GCompositionLighting;
diff --git a/Engine/Source/Runtime/Renderer/Private/CompositionLighting/PostProcessDeferredDecals.cpp b/Engine/Source/Runtime/Renderer/Private/CompositionLighting/PostProcessDeferredDecals.cpp
--- a/Engine/Source/Runtime/Renderer/Private/CompositionLighting/PostProcessDeferredDecals.cpp
+++ b/Engine/Source/Runtime/Renderer/Private/CompositionLighting/PostProcessDeferredDecals.cpp
@@ -10,6 +10,7 @@
 #include "DecalRenderingShared.h"
 #include "ClearQuad.h"
 #include "PipelineStateCache.h"
+#include "VisualizeTexture.h"
 
 
 static TAutoConsoleVariable<int32> CVarGenerateDecalRTWriteMaskTexture(
diff --git a/Engine/Source/Runtime/Renderer/Private/CompositionLighting/PostProcessLpvIndirect.cpp b/Engine/Source/Runtime/Renderer/Private/CompositionLighting/PostProcessLpvIndirect.cpp
--- a/Engine/Source/Runtime/Renderer/Private/CompositionLighting/PostProcessLpvIndirect.cpp
+++ b/Engine/Source/Runtime/Renderer/Private/CompositionLighting/PostProcessLpvIndirect.cpp
@@ -23,6 +23,7 @@
 #include "PostProcess/PostProcessing.h"
 #include "LightPropagationVolumeSettings.h"
 #include "PipelineStateCache.h"
+#include "VisualizeTexture.h"
 
 IMPLEMENT_GLOBAL_SHADER_PARAMETER_STRUCT(FLpvReadUniformBufferParameters,"LpvRead");
 typedef TUniformBufferRef<FLpvReadUniformBufferParameters> FLpvReadUniformBufferRef;
diff --git a/Engine/Source/Runtime/Renderer/Private/DistanceFieldAmbientOcclusion.cpp b/Engine/Source/Runtime/Renderer/Private/DistanceFieldAmbientOcclusion.cpp
--- a/Engine/Source/Runtime/Renderer/Private/DistanceFieldAmbientOcclusion.cpp
+++ b/Engine/Source/Runtime/Renderer/Private/DistanceFieldAmbientOcclusion.cpp
@@ -17,6 +17,7 @@
 #include "DistanceFieldGlobalIllumination.h"
 #include "RendererModule.h"
 #include "PipelineStateCache.h"
+#include "VisualizeTexture.h"
 
 int32 GDistanceFieldAO = 1;
 FAutoConsoleVariableRef CVarDistanceFieldAO(
diff --git a/Engine/Source/Runtime/Renderer/Private/DistanceFieldScreenGridLighting.cpp b/Engine/Source/Runtime/Renderer/Private/DistanceFieldScreenGridLighting.cpp
--- a/Engine/Source/Runtime/Renderer/Private/DistanceFieldScreenGridLighting.cpp
+++ b/Engine/Source/Runtime/Renderer/Private/DistanceFieldScreenGridLighting.cpp
@@ -26,6 +26,7 @@
 #include "GlobalDistanceField.h"
 #include "PipelineStateCache.h"
 #include "ClearQuad.h"
+#include "VisualizeTexture.h"
 
 int32 GAOUseJitter = 1;
 FAutoConsoleVariableRef CVarAOUseJitter(
diff --git a/Engine/Source/Runtime/Renderer/Private/DistortionRendering.cpp b/Engine/Source/Runtime/Renderer/Private/DistortionRendering.cpp
--- a/Engine/Source/Runtime/Renderer/Private/DistortionRendering.cpp
+++ b/Engine/Source/Runtime/Renderer/Private/DistortionRendering.cpp
@@ -27,6 +27,7 @@
 #include "UnrealEngine.h"
 #include "PipelineStateCache.h"
 #include "ScenePrivate.h"
+#include "VisualizeTexture.h"
 
 DECLARE_GPU_STAT(Distortion);
 
diff --git a/Engine/Source/Runtime/Renderer/Private/GPUBenchmark.cpp b/Engine/Source/Runtime/Renderer/Private/GPUBenchmark.cpp
--- a/Engine/Source/Runtime/Renderer/Private/GPUBenchmark.cpp
+++ b/Engine/Source/Runtime/Renderer/Private/GPUBenchmark.cpp
@@ -21,7 +21,8 @@
 #include "GPUProfiler.h"
 #include "PipelineStateCache.h"
 #include "LongGPUTask.h"
-#include "PostProcess/VisualizeTexture.h"
+#include "VisualizeTexture.h"
+#include "DummyRenderResources.h"
 
 static const uint32 GBenchmarkResolution = 512;
 static const uint32 GBenchmarkPrimitives = 200000;
diff --git a/Engine/Source/Runtime/Renderer/Private/LightRendering.cpp b/Engine/Source/Runtime/Renderer/Private/LightRendering.cpp
--- a/Engine/Source/Runtime/Renderer/Private/LightRendering.cpp
+++ b/Engine/Source/Runtime/Renderer/Private/LightRendering.cpp
@@ -13,6 +13,7 @@
 #include "ClearQuad.h"
 #include "Engine/SubsurfaceProfile.h"
 #include "ShowFlags.h"
+#include "VisualizeTexture.h"
 
 DECLARE_GPU_STAT(Lights);
 
diff --git a/Engine/Source/Runtime/Renderer/Private/MobileShadingRenderer.cpp b/Engine/Source/Runtime/Renderer/Private/MobileShadingRenderer.cpp
--- a/Engine/Source/Runtime/Renderer/Private/MobileShadingRenderer.cpp
+++ b/Engine/Source/Runtime/Renderer/Private/MobileShadingRenderer.cpp
@@ -40,6 +40,7 @@
 #include "ClearQuad.h"
 #include "MobileSeparateTranslucencyPass.h"
 #include "MobileDistortionPass.h"
+#include "VisualizeTexturePresent.h"
 
 
 uint32 GetShadowQuality();
@@ -235,7 +236,7 @@ void FMobileSceneRenderer::Render(FRHICommandListImmediate& RHICmdList)
 	FRHICommandListExecutor::GetImmediateCommandList().PollOcclusionQueries();
 	RHICmdList.ImmediateFlush(EImmediateFlushType::DispatchToRHIThread);
 
-	GVisualizeTexture.OnStartFrame(Views[0]);
+	FVisualizeTexturePresent::OnStartRender(Views[0]);
 
 	RHICmdList.SetCurrentStat(GET_STATID(STAT_CLMM_Shadows));
 
@@ -665,41 +666,41 @@ void FMobileSceneRenderer::ConditionalResolveSceneDepth(FRHICommandListImmediate
 						RPInfo.DepthStencilRenderTarget.DepthStencilTarget = DummyDepthTarget;
 						RPInfo.DepthStencilRenderTarget.ExclusiveDepthStencil = FExclusiveDepthStencil::DepthWrite_StencilWrite;
 						RHICmdList.BeginRenderPass(RPInfo, TEXT("ResolveDepth"));
-						{
-							FGraphicsPipelineStateInitializer GraphicsPSOInit;
-							RHICmdList.ApplyCachedRenderTargets(GraphicsPSOInit);
-							GraphicsPSOInit.BlendState = TStaticBlendState<>::GetRHI();
-							GraphicsPSOInit.RasterizerState = TStaticRasterizerState<>::GetRHI();
-							GraphicsPSOInit.DepthStencilState = TStaticDepthStencilState<false, CF_Always>::GetRHI();
-
-							// for devices that do not support framebuffer fetch we rely on undocumented behavior:
-							// Depth reading features will have the depth bound as an attachment AND as a sampler this means
-							// some driver implementations will ignore our attempts to resolve, here we draw with the depth texture to force a resolve.
-							// See UE-37809 for a description of the desired fix.
-							// The results of this draw are irrelevant.
-							TShaderMapRef<FScreenVS> ScreenVertexShader(View.ShaderMap);
-							TShaderMapRef<FScreenPS> PixelShader(View.ShaderMap);
-
-							GraphicsPSOInit.BoundShaderState.VertexDeclarationRHI = GFilterVertexDeclaration.VertexDeclarationRHI;
-							GraphicsPSOInit.BoundShaderState.VertexShaderRHI = GETSAFERHISHADER_VERTEX(*ScreenVertexShader);
-							GraphicsPSOInit.BoundShaderState.PixelShaderRHI = GETSAFERHISHADER_PIXEL(*PixelShader);
-							GraphicsPSOInit.PrimitiveType = PT_TriangleList;
-
-							SetGraphicsPipelineState(RHICmdList, GraphicsPSOInit);
-
-							ScreenVertexShader->SetParameters(RHICmdList, View.ViewUniformBuffer);
-							PixelShader->SetParameters(RHICmdList, TStaticSamplerState<SF_Point>::GetRHI(), SceneContext.GetSceneDepthTexture());
-							DrawRectangle(
-								RHICmdList,
-								0, 0,
-								0, 0,
-								0, 0,
-								1, 1,
-								FIntPoint(1, 1),
-								FIntPoint(1, 1),
-								*ScreenVertexShader,
-								EDRF_UseTriangleOptimization);
-						}
+					{
+						FGraphicsPipelineStateInitializer GraphicsPSOInit;
+						RHICmdList.ApplyCachedRenderTargets(GraphicsPSOInit);
+						GraphicsPSOInit.BlendState = TStaticBlendState<>::GetRHI();
+						GraphicsPSOInit.RasterizerState = TStaticRasterizerState<>::GetRHI();
+						GraphicsPSOInit.DepthStencilState = TStaticDepthStencilState<false, CF_Always>::GetRHI();
+
+						// for devices that do not support framebuffer fetch we rely on undocumented behavior:
+						// Depth reading features will have the depth bound as an attachment AND as a sampler this means
+						// some driver implementations will ignore our attempts to resolve, here we draw with the depth texture to force a resolve.
+						// See UE-37809 for a description of the desired fix.
+						// The results of this draw are irrelevant.
+						TShaderMapRef<FScreenVS> ScreenVertexShader(View.ShaderMap);
+						TShaderMapRef<FScreenPS> PixelShader(View.ShaderMap);
+
+						GraphicsPSOInit.BoundShaderState.VertexDeclarationRHI = GFilterVertexDeclaration.VertexDeclarationRHI;
+						GraphicsPSOInit.BoundShaderState.VertexShaderRHI = GETSAFERHISHADER_VERTEX(*ScreenVertexShader);
+						GraphicsPSOInit.BoundShaderState.PixelShaderRHI = GETSAFERHISHADER_PIXEL(*PixelShader);
+						GraphicsPSOInit.PrimitiveType = PT_TriangleList;
+
+						SetGraphicsPipelineState(RHICmdList, GraphicsPSOInit);
+
+						ScreenVertexShader->SetParameters(RHICmdList, View.ViewUniformBuffer);
+						PixelShader->SetParameters(RHICmdList, TStaticSamplerState<SF_Point>::GetRHI(), SceneContext.GetSceneDepthTexture());
+						DrawRectangle(
+							RHICmdList,
+							0, 0,
+							0, 0,
+							0, 0,
+							1, 1,
+							FIntPoint(1, 1),
+							FIntPoint(1, 1),
+							*ScreenVertexShader,
+							EDRF_UseTriangleOptimization);
+					}
 						RHICmdList.EndRenderPass();
 					} // force depth resolve
 				}
@@ -824,45 +825,45 @@ void FMobileSceneRenderer::CopyMobileMultiViewSceneColor(FRHICommandListImmediat
 	TransitionRenderPassTargets(RHICmdList, RPInfo);
 	RHICmdList.BeginRenderPass(RPInfo, TEXT("CopyMobileMultiViewColor"));
 	{
-		FGraphicsPipelineStateInitializer GraphicsPSOInit;
-		RHICmdList.ApplyCachedRenderTargets(GraphicsPSOInit);
-		GraphicsPSOInit.BlendState = TStaticBlendState<>::GetRHI();
-		GraphicsPSOInit.RasterizerState = TStaticRasterizerState<>::GetRHI();
-		GraphicsPSOInit.DepthStencilState = TStaticDepthStencilState<false, CF_Always>::GetRHI();
+	FGraphicsPipelineStateInitializer GraphicsPSOInit;
+	RHICmdList.ApplyCachedRenderTargets(GraphicsPSOInit);
+	GraphicsPSOInit.BlendState = TStaticBlendState<>::GetRHI();
+	GraphicsPSOInit.RasterizerState = TStaticRasterizerState<>::GetRHI();
+	GraphicsPSOInit.DepthStencilState = TStaticDepthStencilState<false, CF_Always>::GetRHI();
 
-		const auto ShaderMap = GetGlobalShaderMap(FeatureLevel);
-		TShaderMapRef<FScreenVS> VertexShader(ShaderMap);
-		TShaderMapRef<FCopyMobileMultiViewSceneColorPS> PixelShader(ShaderMap);
-		extern TGlobalResource<FFilterVertexDeclaration> GFilterVertexDeclaration;
+	const auto ShaderMap = GetGlobalShaderMap(FeatureLevel);
+	TShaderMapRef<FScreenVS> VertexShader(ShaderMap);
+	TShaderMapRef<FCopyMobileMultiViewSceneColorPS> PixelShader(ShaderMap);
+	extern TGlobalResource<FFilterVertexDeclaration> GFilterVertexDeclaration;
 
-		GraphicsPSOInit.BoundShaderState.VertexDeclarationRHI = GFilterVertexDeclaration.VertexDeclarationRHI;
-		GraphicsPSOInit.BoundShaderState.VertexShaderRHI = GETSAFERHISHADER_VERTEX(*VertexShader);
-		GraphicsPSOInit.BoundShaderState.PixelShaderRHI = GETSAFERHISHADER_PIXEL(*PixelShader);
-		GraphicsPSOInit.PrimitiveType = PT_TriangleList;
+	GraphicsPSOInit.BoundShaderState.VertexDeclarationRHI = GFilterVertexDeclaration.VertexDeclarationRHI;
+	GraphicsPSOInit.BoundShaderState.VertexShaderRHI = GETSAFERHISHADER_VERTEX(*VertexShader);
+	GraphicsPSOInit.BoundShaderState.PixelShaderRHI = GETSAFERHISHADER_PIXEL(*PixelShader);
+	GraphicsPSOInit.PrimitiveType = PT_TriangleList;
 
-		SetGraphicsPipelineState(RHICmdList, GraphicsPSOInit);
+	SetGraphicsPipelineState(RHICmdList, GraphicsPSOInit);
 
-		for (int32 ViewIndex = 0; ViewIndex < Views.Num(); ++ViewIndex)
-		{
-			const FViewInfo& View = Views[ViewIndex];
-
-			// Multi-view color target is our input texture array
-			PixelShader->SetParameters(RHICmdList, View.ViewUniformBuffer, SceneContext.MobileMultiViewSceneColor->GetRenderTargetItem().ShaderResourceTexture);
-
-			RHICmdList.SetViewport(View.ViewRect.Min.X, View.ViewRect.Min.Y, 0.0f, View.ViewRect.Min.X + View.ViewRect.Width(), View.ViewRect.Min.Y + View.ViewRect.Height(), 1.0f);
-			const FIntPoint TargetSize(View.ViewRect.Width(), View.ViewRect.Height());
-
-			DrawRectangle(
-				RHICmdList,
-				0, 0,
-				View.ViewRect.Width(), View.ViewRect.Height(),
-				0, 0,
-				View.ViewRect.Width(), View.ViewRect.Height(),
-				TargetSize,
-				TargetSize,
-				*VertexShader,
-				EDRF_UseTriangleOptimization);
-		}
+	for (int32 ViewIndex = 0; ViewIndex < Views.Num(); ++ViewIndex)
+	{
+		const FViewInfo& View = Views[ViewIndex];
+
+		// Multi-view color target is our input texture array
+		PixelShader->SetParameters(RHICmdList, View.ViewUniformBuffer, SceneContext.MobileMultiViewSceneColor->GetRenderTargetItem().ShaderResourceTexture);
+
+		RHICmdList.SetViewport(View.ViewRect.Min.X, View.ViewRect.Min.Y, 0.0f, View.ViewRect.Min.X + View.ViewRect.Width(), View.ViewRect.Min.Y + View.ViewRect.Height(), 1.0f);
+		const FIntPoint TargetSize(View.ViewRect.Width(), View.ViewRect.Height());
+
+		DrawRectangle(
+			RHICmdList,
+			0, 0,
+			View.ViewRect.Width(), View.ViewRect.Height(),
+			0, 0,
+			View.ViewRect.Width(), View.ViewRect.Height(),
+			TargetSize,
+			TargetSize,
+			*VertexShader,
+			EDRF_UseTriangleOptimization);
+	}
 	}
 	RHICmdList.EndRenderPass();
 }
diff --git a/Engine/Source/Runtime/Renderer/Private/PostProcess/PostProcessCompositeEditorPrimitives.cpp b/Engine/Source/Runtime/Renderer/Private/PostProcess/PostProcessCompositeEditorPrimitives.cpp
--- a/Engine/Source/Runtime/Renderer/Private/PostProcess/PostProcessCompositeEditorPrimitives.cpp
+++ b/Engine/Source/Runtime/Renderer/Private/PostProcess/PostProcessCompositeEditorPrimitives.cpp
@@ -12,6 +12,7 @@
 #include "DynamicPrimitiveDrawing.h"
 #include "ClearQuad.h"
 #include "PipelineStateCache.h"
+#include "VisualizeTexture.h"
 
 #if WITH_EDITOR
 
diff --git a/Engine/Source/Runtime/Renderer/Private/PostProcess/PostProcessSubsurface.cpp b/Engine/Source/Runtime/Renderer/Private/PostProcess/PostProcessSubsurface.cpp
--- a/Engine/Source/Runtime/Renderer/Private/PostProcess/PostProcessSubsurface.cpp
+++ b/Engine/Source/Runtime/Renderer/Private/PostProcess/PostProcessSubsurface.cpp
@@ -20,6 +20,7 @@
 #include "CompositionLighting/PostProcessPassThrough.h"
 #include "ClearQuad.h"
 #include "PipelineStateCache.h"
+#include "VisualizeTexture.h"
 
 ENGINE_API const IPooledRenderTarget* GetSubsufaceProfileTexture_RT(FRHICommandListImmediate& RHICmdList);
 
diff --git a/Engine/Source/Runtime/Renderer/Private/PostProcess/PostProcessVisualizeComplexity.cpp b/Engine/Source/Runtime/Renderer/Private/PostProcess/PostProcessVisualizeComplexity.cpp
--- a/Engine/Source/Runtime/Renderer/Private/PostProcess/PostProcessVisualizeComplexity.cpp
+++ b/Engine/Source/Runtime/Renderer/Private/PostProcess/PostProcessVisualizeComplexity.cpp
@@ -111,9 +111,6 @@ void FVisualizeComplexityApplyPS::SetParameters(
 
 IMPLEMENT_SHADER_TYPE(,FVisualizeComplexityApplyPS,TEXT("/Engine/Private/ShaderComplexityApplyPixelShader.usf"),TEXT("Main"),SF_Pixel);
 
-//reuse the generic filter vertex declaration
-extern TGlobalResource<FFilterVertexDeclaration> GFilterVertexDeclaration;
-
 void FRCPassPostProcessVisualizeComplexity::Process(FRenderingCompositePassContext& Context)
 {
 	SCOPED_DRAW_EVENT(Context.RHICmdList, PostProcessVisualizeComplexity);
diff --git a/Engine/Source/Runtime/Renderer/Private/PostProcess/RenderingCompositionGraph.cpp b/Engine/Source/Runtime/Renderer/Private/PostProcess/RenderingCompositionGraph.cpp
--- a/Engine/Source/Runtime/Renderer/Private/PostProcess/RenderingCompositionGraph.cpp
+++ b/Engine/Source/Runtime/Renderer/Private/PostProcess/RenderingCompositionGraph.cpp
@@ -17,6 +17,7 @@
 #include "IXRTrackingSystem.h"
 #include "PostProcess/SceneRenderTargets.h"
 #include "SceneRendering.h"
+#include "VisualizeTexture.h"
 
 #include "IImageWrapper.h"
 #include "ImageWriteQueue.h"
diff --git a/Engine/Source/Runtime/Renderer/Private/PostProcess/SceneFilterRendering.cpp b/Engine/Source/Runtime/Renderer/Private/PostProcess/SceneFilterRendering.cpp
--- a/Engine/Source/Runtime/Renderer/Private/PostProcess/SceneFilterRendering.cpp
+++ b/Engine/Source/Runtime/Renderer/Private/PostProcess/SceneFilterRendering.cpp
@@ -10,67 +10,8 @@
 #include "Engine/Engine.h"
 #include "IHeadMountedDisplay.h"
 #include "IXRTrackingSystem.h"
-
-/**
-* Static vertex and index buffer used for 2D screen rectangles.
-*/
-class FScreenRectangleVertexBuffer : public FVertexBuffer
-{
-public:
-	/** Initialize the RHI for this rendering resource */
-	void InitRHI() override
-	{
-		TResourceArray<FFilterVertex, VERTEXBUFFER_ALIGNMENT> Vertices;
-		Vertices.SetNumUninitialized(6);
-
-		Vertices[0].Position = FVector4(1,  1,	0,	1);
-		Vertices[0].UV = FVector2D(1,	1);
-
-		Vertices[1].Position = FVector4(0,  1,	0,	1);
-		Vertices[1].UV = FVector2D(0,	1);
-
-		Vertices[2].Position = FVector4(1,	0,	0,	1);
-		Vertices[2].UV = FVector2D(1,	0);
-
-		Vertices[3].Position = FVector4(0,	0,	0,	1);
-		Vertices[3].UV = FVector2D(0,	0);
-
-		//The final two vertices are used for the triangle optimization (a single triangle spans the entire viewport )
-		Vertices[4].Position = FVector4(-1,  1,	0,	1);
-		Vertices[4].UV = FVector2D(-1,	1);
-
-		Vertices[5].Position = FVector4(1,  -1,	0,	1);
-		Vertices[5].UV = FVector2D(1, -1);
-
-		// Create vertex buffer. Fill buffer with initial data upon creation
-		FRHIResourceCreateInfo CreateInfo(&Vertices);
-		VertexBufferRHI = RHICreateVertexBuffer(Vertices.GetResourceDataSize(), BUF_Static, CreateInfo);
-	}
-};
-
-class FScreenRectangleIndexBuffer : public FIndexBuffer
-{
-public:
-	/** Initialize the RHI for this rendering resource */
-	void InitRHI() override
-	{
-		// Indices 0 - 5 are used for rendering a quad. Indices 6 - 8 are used for triangle optimization.
-		const uint16 Indices[] = { 0, 1, 2, 2, 1, 3, 0, 4, 5 };
-
-		TResourceArray<uint16, INDEXBUFFER_ALIGNMENT> IndexBuffer;
-		uint32 NumIndices = ARRAY_COUNT(Indices);
-		IndexBuffer.AddUninitialized(NumIndices);
-		FMemory::Memcpy(IndexBuffer.GetData(), Indices, NumIndices * sizeof(uint16));
-
-		// Create index buffer. Fill buffer with initial data upon creation
-		FRHIResourceCreateInfo CreateInfo(&IndexBuffer);
-		IndexBufferRHI = RHICreateIndexBuffer(sizeof(uint16), IndexBuffer.GetResourceDataSize(), BUF_Static, CreateInfo);
-	}
-};
-
-/** Global resource  */
-static TGlobalResource<FScreenRectangleVertexBuffer> GScreenRectangleVertexBuffer;
-static TGlobalResource<FScreenRectangleIndexBuffer> GScreenRectangleIndexBuffer;
+#include "PixelShaderUtils.h"
+#include "DummyRenderResources.h"
 
 void FTesselatedScreenRectangleIndexBuffer::InitRHI()
 {
@@ -122,11 +63,6 @@ uint32 FTesselatedScreenRectangleIndexBuffer::NumPrimitives() const
 static TGlobalResource<FTesselatedScreenRectangleIndexBuffer> GTesselatedScreenRectangleIndexBuffer;
 
 
-/** Vertex declaration for the 2D screen rectangle. */
-TGlobalResource<FFilterVertexDeclaration> GFilterVertexDeclaration;
-/** Vertex declaration for vertex shaders that don't require any inputs (eg. generated via vertex ID). */
-TGlobalResource<FEmptyVertexDeclaration> GEmptyVertexDeclaration;
-
 IMPLEMENT_GLOBAL_SHADER_PARAMETER_STRUCT(FDrawRectangleParameters, "DrawRectangleParameters");
 
 typedef TUniformBufferRef<FDrawRectangleParameters> FDrawRectangleBufferRef;
@@ -214,37 +150,13 @@ static inline void InternalDrawRectangle(
 	}
 	else
 	{
-		RHICmdList.SetStreamSource(0, GScreenRectangleVertexBuffer.VertexBufferRHI, 0);
-
 		if (Flags == EDRF_UseTriangleOptimization)
 		{
-			// A single triangle spans the entire viewport this results in a quad that fill the viewport. This can increase rasterization efficiency
-			// as we do not have a diagonal edge (through the center) for the rasterizer/span-dispatch. Although the actual benefit of this technique is dependent upon hardware.
-
-			// We offset into the index buffer when using the triangle optimization to access the correct vertices.
-			RHICmdList.DrawIndexedPrimitive(
-				GScreenRectangleIndexBuffer.IndexBufferRHI,
-				/*BaseVertexIndex=*/ 0,
-				/*MinIndex=*/ 0,
-				/*NumVertices=*/ 3,
-				/*StartIndex=*/ 6,
-				/*NumPrimitives=*/ 1,
-				/*NumInstances=*/ InstanceCount
-				);
+			FPixelShaderUtils::DrawFullscreenTriangle(RHICmdList, InstanceCount);
 		}
 		else
 		{
-			RHICmdList.SetStreamSource(0, GScreenRectangleVertexBuffer.VertexBufferRHI, 0);
-
-			RHICmdList.DrawIndexedPrimitive(
-				GScreenRectangleIndexBuffer.IndexBufferRHI,
-				/*BaseVertexIndex=*/ 0,
-				/*MinIndex=*/ 0,
-				/*NumVertices=*/ 4,
-				/*StartIndex=*/ 0,
-				/*NumPrimitives=*/ 2,
-				/*NumInstances=*/ InstanceCount
-				);
+			FPixelShaderUtils::DrawFullscreenQuad(RHICmdList, InstanceCount);
 		}
 	}
 }
diff --git a/Engine/Source/Runtime/Renderer/Private/PostProcess/SceneFilterRendering.h b/Engine/Source/Runtime/Renderer/Private/PostProcess/SceneFilterRendering.h
--- a/Engine/Source/Runtime/Renderer/Private/PostProcess/SceneFilterRendering.h
+++ b/Engine/Source/Runtime/Renderer/Private/PostProcess/SceneFilterRendering.h
@@ -108,10 +108,6 @@ extern void DrawPostProcessPass(
 	EDrawRectangleFlags Flags = EDRF_Default
 	);
 
-extern TGlobalResource<FFilterVertexDeclaration> GFilterVertexDeclaration;
-extern TGlobalResource<FEmptyVertexDeclaration> GEmptyVertexDeclaration;
-
-
 
 
 /*-----------------------------------------------------------------------------
diff --git a/Engine/Source/Runtime/Renderer/Private/PostProcess/SceneRenderTargets.cpp b/Engine/Source/Runtime/Renderer/Private/PostProcess/SceneRenderTargets.cpp
--- a/Engine/Source/Runtime/Renderer/Private/PostProcess/SceneRenderTargets.cpp
+++ b/Engine/Source/Runtime/Renderer/Private/PostProcess/SceneRenderTargets.cpp
@@ -26,6 +26,7 @@
 #include "StereoRenderTargetManager.h"
 #include "VT/VirtualTextureSystem.h"
 #include "VT/VirtualTextureFeedback.h"
+#include "VisualizeTexture.h"
 
 static TAutoConsoleVariable<int32> CVarRSMResolution(
 	TEXT("r.LPV.RSMResolution"),
diff --git a/Engine/Source/Runtime/Renderer/Private/PostProcess/VisualizeTexture.cpp b/Engine/Source/Runtime/Renderer/Private/PostProcess/VisualizeTexture.cpp
deleted file mode 100644
diff --git a/Engine/Source/Runtime/Renderer/Private/PostProcess/VisualizeTexture.h b/Engine/Source/Runtime/Renderer/Private/PostProcess/VisualizeTexture.h
deleted file mode 100644
diff --git a/Engine/Source/Runtime/Renderer/Private/ReflectionEnvironmentDiffuseIrradiance.cpp b/Engine/Source/Runtime/Renderer/Private/ReflectionEnvironmentDiffuseIrradiance.cpp
--- a/Engine/Source/Runtime/Renderer/Private/ReflectionEnvironmentDiffuseIrradiance.cpp
+++ b/Engine/Source/Runtime/Renderer/Private/ReflectionEnvironmentDiffuseIrradiance.cpp
@@ -16,6 +16,8 @@
 #include "PostProcess/SceneFilterRendering.h"
 #include "ScreenRendering.h"
 #include "PipelineStateCache.h"
+#include "VisualizeTexture.h"
+#include "DummyRenderResources.h"
 
 extern int32 GDiffuseIrradianceCubemapSize;
 
diff --git a/Engine/Source/Runtime/Renderer/Private/Renderer.cpp b/Engine/Source/Runtime/Renderer/Private/Renderer.cpp
--- a/Engine/Source/Runtime/Renderer/Private/Renderer.cpp
+++ b/Engine/Source/Runtime/Renderer/Private/Renderer.cpp
@@ -16,7 +16,7 @@
 #include "SceneView.h"
 #include "RenderTargetPool.h"
 #include "PostProcess/SceneRenderTargets.h"
-#include "PostProcess/VisualizeTexture.h"
+#include "VisualizeTexture.h"
 #include "SceneCore.h"
 #include "SceneHitProxyRendering.h"
 #include "SceneRendering.h"
@@ -26,6 +26,7 @@
 #include "RendererModule.h"
 #include "GPUBenchmark.h"
 #include "SystemSettings.h"
+#include "VisualizeTexturePresent.h"
 
 DEFINE_LOG_CATEGORY(LogRenderer);
 
@@ -173,7 +174,7 @@ void FRendererModule::DebugLogOnCrash()
 {
 	GVisualizeTexture.SortOrder = 1;
 	GVisualizeTexture.bFullList = true;
-	GVisualizeTexture.DebugLog(false);
+	FVisualizeTexturePresent::DebugLog(false);
 	
 	GEngine->Exec(NULL, TEXT("rhi.DumpMemory"), *GLog);
 
@@ -239,14 +240,6 @@ void FRendererModule::GPUBenchmark(FSynthBenchmarkResults& InOut, float WorkScal
 	FlushRenderingCommands();
 }
 
-void FRendererModule::QueryVisualizeTexture(FQueryVisualizeTexureInfo& Out)
-{
-	check(IsInGameThread());
-	FlushRenderingCommands();
-
-	GVisualizeTexture.QueryInfo(Out);
-}
-
 static void VisualizeTextureExec( const TCHAR* Cmd, FOutputDevice &Ar )
 {
 	check(IsInGameThread());
@@ -322,18 +315,18 @@ static void VisualizeTextureExec( const TCHAR* Cmd, FOutputDevice &Ar )
 				{
 					// user specified a reuse goal
 					FString NameWithoutAt = Parameter.Left(AfterAt - *Parameter);
-					GVisualizeTexture.SetObserveTarget(*NameWithoutAt, FCString::Atoi(AfterAt + 1));
+					GVisualizeTexture.SetRenderTargetNameToObserve(*NameWithoutAt, FCString::Atoi(AfterAt + 1));
 				}
 				else
 				{
 					// we take the last one
-					GVisualizeTexture.SetObserveTarget(*Parameter);
+					GVisualizeTexture.SetRenderTargetNameToObserve(*Parameter);
 				}
 			}
 			else
 			{
 				// the index was used
-				GVisualizeTexture.SetObserveTarget(TEXT(""));
+				GVisualizeTexture.SetRenderTargetNameToObserve(TEXT(""));
 			}
 		}
 		// GRenderTargetPoolInputMapping mode
@@ -469,7 +462,7 @@ static void VisualizeTextureExec( const TCHAR* Cmd, FOutputDevice &Ar )
 		Ar.Logf(TEXT("TextureId:"));
 		Ar.Logf(TEXT("  0        = <off>"));
 
-		GVisualizeTexture.DebugLog(true);
+		FVisualizeTexturePresent::DebugLog(true);
 	}
 	//		Ar.Logf(TEXT("VisualizeTexture %d"), GVisualizeTexture.Mode);
 }
diff --git a/Engine/Source/Runtime/Renderer/Private/RendererModule.h b/Engine/Source/Runtime/Renderer/Private/RendererModule.h
--- a/Engine/Source/Runtime/Renderer/Private/RendererModule.h
+++ b/Engine/Source/Runtime/Renderer/Private/RendererModule.h
@@ -48,7 +48,6 @@ public:
 	virtual void TickRenderTargetPool() override;
 	virtual void DebugLogOnCrash() override;
 	virtual void GPUBenchmark(FSynthBenchmarkResults& InOut, float WorkScale) override;
-	virtual void QueryVisualizeTexture(FQueryVisualizeTexureInfo& Out) override;
 	virtual void ExecVisualizeTextureCmd(const FString& Cmd) override;
 	virtual void UpdateMapNeedsLightingFullyRebuiltState(UWorld* World) override;
 	virtual void DrawRectangle(
@@ -66,7 +65,6 @@ public:
 		class FShader* VertexShader,
 		EDrawRectangleFlags Flags = EDRF_Default
 		) override;
-	virtual TGlobalResource<FFilterVertexDeclaration>& GetFilterVertexDeclaration() override;
 
 	virtual const TSet<FSceneInterface*>& GetAllocatedScenes() override
 	{
diff --git a/Engine/Source/Runtime/Renderer/Private/SceneHitProxyRendering.cpp b/Engine/Source/Runtime/Renderer/Private/SceneHitProxyRendering.cpp
--- a/Engine/Source/Runtime/Renderer/Private/SceneHitProxyRendering.cpp
+++ b/Engine/Source/Runtime/Renderer/Private/SceneHitProxyRendering.cpp
@@ -17,6 +17,7 @@
 #include "ScenePrivate.h"
 #include "DynamicPrimitiveDrawing.h"
 #include "ClearQuad.h"
+#include "VisualizeTexture.h"
 
 /**
  * A vertex shader for rendering the depth of a mesh.
diff --git a/Engine/Source/Runtime/Renderer/Private/SceneOcclusion.cpp b/Engine/Source/Runtime/Renderer/Private/SceneOcclusion.cpp
--- a/Engine/Source/Runtime/Renderer/Private/SceneOcclusion.cpp
+++ b/Engine/Source/Runtime/Renderer/Private/SceneOcclusion.cpp
@@ -16,6 +16,7 @@
 #include "PlanarReflectionSceneProxy.h"
 #include "PipelineStateCache.h"
 #include "ShaderParameterStruct.h"
+#include "VisualizeTexture.h"
 
 
 /*-----------------------------------------------------------------------------
diff --git a/Engine/Source/Runtime/Renderer/Private/ScenePrivate.h b/Engine/Source/Runtime/Renderer/Private/ScenePrivate.h
--- a/Engine/Source/Runtime/Renderer/Private/ScenePrivate.h
+++ b/Engine/Source/Runtime/Renderer/Private/ScenePrivate.h
@@ -41,7 +41,7 @@
 #include "MobileBasePassRendering.h"
 #include "VolumeRendering.h"
 #include "SceneSoftwareOcclusion.h"
-#include "PostProcess/VisualizeTexture.h"
+#include "DummyRenderResources.h"
 
 /** Factor by which to grow occlusion tests **/
 #define OCCLUSION_SLOP (1.0f)
diff --git a/Engine/Source/Runtime/Renderer/Private/SceneRendering.cpp b/Engine/Source/Runtime/Renderer/Private/SceneRendering.cpp
--- a/Engine/Source/Runtime/Renderer/Private/SceneRendering.cpp
+++ b/Engine/Source/Runtime/Renderer/Private/SceneRendering.cpp
@@ -44,6 +44,7 @@
 #include "PostProcess/PostProcessing.h"
 #include "SceneSoftwareOcclusion.h"
 #include "VirtualTexturing.h"
+#include "VisualizeTexturePresent.h"
 
 /*-----------------------------------------------------------------------------
 	Globals
@@ -2527,7 +2528,7 @@ void FSceneRenderer::RenderFinish(FRHICommandListImmediate& RHICmdList)
 				continue;
 			}
 
-			GVisualizeTexture.PresentContent(RHICmdList, View);
+			FVisualizeTexturePresent::PresentContent(RHICmdList, View);
 		}
 	}
 #endif
@@ -2700,7 +2701,7 @@ void FSceneRenderer::OnStartFrame(FRHICommandListImmediate& RHICmdList)
 {
 	FSceneRenderTargets& SceneContext = FSceneRenderTargets::Get(RHICmdList);
 
-	GVisualizeTexture.OnStartFrame(Views[0]);
+	FVisualizeTexturePresent::OnStartRender(Views[0]);
 	CompositionGraph_OnStartFrame();
 	SceneContext.bScreenSpaceAOIsValid = false;
 	SceneContext.bCustomDepthIsValid = false;
@@ -3230,12 +3231,6 @@ void FRendererModule::DrawRectangle(
 	::DrawRectangle( RHICmdList, X, Y, SizeX, SizeY, U, V, SizeU, SizeV, TargetSize, TextureSize, VertexShader, Flags );
 }
 
-
-TGlobalResource<FFilterVertexDeclaration>& FRendererModule::GetFilterVertexDeclaration()
-{
-	return GFilterVertexDeclaration;
-}
-
 void FRendererModule::RegisterPostOpaqueComputeDispatcher(FComputeDispatcher *Dispatcher)
 {
 	PostOpaqueDispatchers.AddUnique(Dispatcher);
@@ -3517,115 +3512,115 @@ void FSceneRenderer::ResolveSceneColor(FRHICommandList& RHICmdList)
 		FRHIRenderPassInfo RPInfo(SceneContext.GetSceneColorTexture(), ERenderTargetActions::Load_Store);
 		RHICmdList.BeginRenderPass(RPInfo, TEXT("ResolveColor"));
 		{
-			for (int32 ViewIndex = 0; ViewIndex < Views.Num(); ViewIndex++)
-			{
-				const FViewInfo& View = Views[ViewIndex];
+		for (int32 ViewIndex = 0; ViewIndex < Views.Num(); ViewIndex++)
+		{
+			const FViewInfo& View = Views[ViewIndex];
 
-				FGraphicsPipelineStateInitializer GraphicsPSOInit;
-				RHICmdList.ApplyCachedRenderTargets(GraphicsPSOInit);
+			FGraphicsPipelineStateInitializer GraphicsPSOInit;
+			RHICmdList.ApplyCachedRenderTargets(GraphicsPSOInit);
 
-				GraphicsPSOInit.BlendState = TStaticBlendState<>::GetRHI();
-				GraphicsPSOInit.RasterizerState = TStaticRasterizerState<>::GetRHI();
-				GraphicsPSOInit.DepthStencilState = TStaticDepthStencilState<false, CF_Always>::GetRHI();
-				RHICmdList.SetStreamSource(0, GResolveDummyVertexBuffer.VertexBufferRHI, 0);
+			GraphicsPSOInit.BlendState = TStaticBlendState<>::GetRHI();
+			GraphicsPSOInit.RasterizerState = TStaticRasterizerState<>::GetRHI();
+			GraphicsPSOInit.DepthStencilState = TStaticDepthStencilState<false, CF_Always>::GetRHI();
+			RHICmdList.SetStreamSource(0, GResolveDummyVertexBuffer.VertexBufferRHI, 0);
 
-				// Resolve views individually
-				// In the case of adaptive resolution, the view family will be much larger than the views individually
-				RHICmdList.SetScissorRect(true, View.ViewRect.Min.X, View.ViewRect.Min.Y, View.ViewRect.Max.X, View.ViewRect.Max.Y);
+			// Resolve views individually
+			// In the case of adaptive resolution, the view family will be much larger than the views individually
+			RHICmdList.SetScissorRect(true, View.ViewRect.Min.X, View.ViewRect.Min.Y, View.ViewRect.Max.X, View.ViewRect.Max.Y);
 
-				int32 ResolveWidth = CVarWideCustomResolve.GetValueOnRenderThread();
+			int32 ResolveWidth = CVarWideCustomResolve.GetValueOnRenderThread();
 
-				if (CurrentNumSamples <= 1)
-				{
-					ResolveWidth = 0;
-				}
+			if (CurrentNumSamples <= 1)
+			{
+				ResolveWidth = 0;
+			}
 
-				if (ResolveWidth != 0)
-				{
-					ResolveFilterWide(RHICmdList, GraphicsPSOInit, SceneContext.GetCurrentFeatureLevel(), CurrentSceneColor->GetRenderTargetItem().TargetableTexture, FMaskTexture, FIntPoint(0, 0), CurrentNumSamples, ResolveWidth);
-				}
-				else
-				{
-					auto ShaderMap = GetGlobalShaderMap(SceneContext.GetCurrentFeatureLevel());
-					TShaderMapRef<FHdrCustomResolveVS> VertexShader(ShaderMap);
-					GraphicsPSOInit.BoundShaderState.VertexDeclarationRHI = GetVertexDeclarationFVector4();
-					GraphicsPSOInit.BoundShaderState.VertexShaderRHI = GETSAFERHISHADER_VERTEX(*VertexShader);
-					GraphicsPSOInit.PrimitiveType = PT_TriangleList;
+			if (ResolveWidth != 0)
+			{
+				ResolveFilterWide(RHICmdList, GraphicsPSOInit, SceneContext.GetCurrentFeatureLevel(), CurrentSceneColor->GetRenderTargetItem().TargetableTexture, FMaskTexture, FIntPoint(0, 0), CurrentNumSamples, ResolveWidth);
+			}
+			else
+			{
+				auto ShaderMap = GetGlobalShaderMap(SceneContext.GetCurrentFeatureLevel());
+				TShaderMapRef<FHdrCustomResolveVS> VertexShader(ShaderMap);
+				GraphicsPSOInit.BoundShaderState.VertexDeclarationRHI = GetVertexDeclarationFVector4();
+				GraphicsPSOInit.BoundShaderState.VertexShaderRHI = GETSAFERHISHADER_VERTEX(*VertexShader);
+				GraphicsPSOInit.PrimitiveType = PT_TriangleList;
 
-					if (FMaskTexture.IsValid())
+				if (FMaskTexture.IsValid())
+				{
+					if (CurrentNumSamples == 2)
 					{
-						if (CurrentNumSamples == 2)
-						{
-							TShaderMapRef<FHdrCustomResolveFMask2xPS> PixelShader(ShaderMap);
-							GraphicsPSOInit.BoundShaderState.PixelShaderRHI = GETSAFERHISHADER_PIXEL(*PixelShader);
+						TShaderMapRef<FHdrCustomResolveFMask2xPS> PixelShader(ShaderMap);
+						GraphicsPSOInit.BoundShaderState.PixelShaderRHI = GETSAFERHISHADER_PIXEL(*PixelShader);
 
-							SetGraphicsPipelineState(RHICmdList, GraphicsPSOInit);
-							PixelShader->SetParameters(RHICmdList, CurrentSceneColor->GetRenderTargetItem().TargetableTexture, FMaskTexture);
-						}
-						else if (CurrentNumSamples == 4)
-						{
-							TShaderMapRef<FHdrCustomResolveFMask4xPS> PixelShader(ShaderMap);
-							GraphicsPSOInit.BoundShaderState.PixelShaderRHI = GETSAFERHISHADER_PIXEL(*PixelShader);
+						SetGraphicsPipelineState(RHICmdList, GraphicsPSOInit);
+						PixelShader->SetParameters(RHICmdList, CurrentSceneColor->GetRenderTargetItem().TargetableTexture, FMaskTexture);
+					}
+					else if (CurrentNumSamples == 4)
+					{
+						TShaderMapRef<FHdrCustomResolveFMask4xPS> PixelShader(ShaderMap);
+						GraphicsPSOInit.BoundShaderState.PixelShaderRHI = GETSAFERHISHADER_PIXEL(*PixelShader);
 
-							SetGraphicsPipelineState(RHICmdList, GraphicsPSOInit);
-							PixelShader->SetParameters(RHICmdList, CurrentSceneColor->GetRenderTargetItem().TargetableTexture, FMaskTexture);
-						}
-						else if (CurrentNumSamples == 8)
-						{
-							TShaderMapRef<FHdrCustomResolveFMask8xPS> PixelShader(ShaderMap);
-							GraphicsPSOInit.BoundShaderState.PixelShaderRHI = GETSAFERHISHADER_PIXEL(*PixelShader);
+						SetGraphicsPipelineState(RHICmdList, GraphicsPSOInit);
+						PixelShader->SetParameters(RHICmdList, CurrentSceneColor->GetRenderTargetItem().TargetableTexture, FMaskTexture);
+					}
+					else if (CurrentNumSamples == 8)
+					{
+						TShaderMapRef<FHdrCustomResolveFMask8xPS> PixelShader(ShaderMap);
+						GraphicsPSOInit.BoundShaderState.PixelShaderRHI = GETSAFERHISHADER_PIXEL(*PixelShader);
 
-							SetGraphicsPipelineState(RHICmdList, GraphicsPSOInit);
-							PixelShader->SetParameters(RHICmdList, CurrentSceneColor->GetRenderTargetItem().TargetableTexture, FMaskTexture);
-						}
-						else
-						{
-							// Everything other than 2,4,8 samples is not implemented.
-							check(0);
-							break;
-						}
+						SetGraphicsPipelineState(RHICmdList, GraphicsPSOInit);
+						PixelShader->SetParameters(RHICmdList, CurrentSceneColor->GetRenderTargetItem().TargetableTexture, FMaskTexture);
 					}
 					else
 					{
-						if (CurrentNumSamples == 2)
-						{
-							TShaderMapRef<FHdrCustomResolve2xPS> PixelShader(ShaderMap);
-							GraphicsPSOInit.BoundShaderState.PixelShaderRHI = GETSAFERHISHADER_PIXEL(*PixelShader);
-
-							SetGraphicsPipelineState(RHICmdList, GraphicsPSOInit);
-							PixelShader->SetParameters(RHICmdList, CurrentSceneColor->GetRenderTargetItem().TargetableTexture);
-						}
-						else if (CurrentNumSamples == 4)
-						{
-							TShaderMapRef<FHdrCustomResolve4xPS> PixelShader(ShaderMap);
-							GraphicsPSOInit.BoundShaderState.PixelShaderRHI = GETSAFERHISHADER_PIXEL(*PixelShader);
+						// Everything other than 2,4,8 samples is not implemented.
+						check(0);
+						break;
+					}
+				}
+				else
+				{
+				if (CurrentNumSamples == 2)
+				{
+					TShaderMapRef<FHdrCustomResolve2xPS> PixelShader(ShaderMap);
+					GraphicsPSOInit.BoundShaderState.PixelShaderRHI = GETSAFERHISHADER_PIXEL(*PixelShader);
 
-							SetGraphicsPipelineState(RHICmdList, GraphicsPSOInit);
-							PixelShader->SetParameters(RHICmdList, CurrentSceneColor->GetRenderTargetItem().TargetableTexture);
-						}
-						else if (CurrentNumSamples == 8)
-						{
-							TShaderMapRef<FHdrCustomResolve8xPS> PixelShader(ShaderMap);
-							GraphicsPSOInit.BoundShaderState.PixelShaderRHI = GETSAFERHISHADER_PIXEL(*PixelShader);
+					SetGraphicsPipelineState(RHICmdList, GraphicsPSOInit);
+					PixelShader->SetParameters(RHICmdList, CurrentSceneColor->GetRenderTargetItem().TargetableTexture);
+				}
+				else if (CurrentNumSamples == 4)
+				{
+					TShaderMapRef<FHdrCustomResolve4xPS> PixelShader(ShaderMap);
+					GraphicsPSOInit.BoundShaderState.PixelShaderRHI = GETSAFERHISHADER_PIXEL(*PixelShader);
 
-							SetGraphicsPipelineState(RHICmdList, GraphicsPSOInit);
-							PixelShader->SetParameters(RHICmdList, CurrentSceneColor->GetRenderTargetItem().TargetableTexture);
-						}
-						else
-						{
-							// Everything other than 2,4,8 samples is not implemented.
-							check(0);
-							break;
-						}
-					}
+					SetGraphicsPipelineState(RHICmdList, GraphicsPSOInit);
+					PixelShader->SetParameters(RHICmdList, CurrentSceneColor->GetRenderTargetItem().TargetableTexture);
+				}
+				else if (CurrentNumSamples == 8)
+				{
+					TShaderMapRef<FHdrCustomResolve8xPS> PixelShader(ShaderMap);
+					GraphicsPSOInit.BoundShaderState.PixelShaderRHI = GETSAFERHISHADER_PIXEL(*PixelShader);
 
-					RHICmdList.DrawPrimitive(0, 1, 1);
+					SetGraphicsPipelineState(RHICmdList, GraphicsPSOInit);
+					PixelShader->SetParameters(RHICmdList, CurrentSceneColor->GetRenderTargetItem().TargetableTexture);
+				}
+				else
+				{
+					// Everything other than 2,4,8 samples is not implemented.
+					check(0);
+						break;
 				}
 			}
 
-			RHICmdList.SetScissorRect(false, 0, 0, 0, 0);
+				RHICmdList.DrawPrimitive(0, 1, 1);
+			}
 		}
 
+		RHICmdList.SetScissorRect(false, 0, 0, 0, 0);
+	}
+
 		RHICmdList.EndRenderPass();
 	}
 }
diff --git a/Engine/Source/Runtime/Renderer/Private/ShadowDepthRendering.cpp b/Engine/Source/Runtime/Renderer/Private/ShadowDepthRendering.cpp
--- a/Engine/Source/Runtime/Renderer/Private/ShadowDepthRendering.cpp
+++ b/Engine/Source/Runtime/Renderer/Private/ShadowDepthRendering.cpp
@@ -37,6 +37,7 @@
 #include "ScreenRendering.h"
 #include "ClearQuad.h"
 #include "PipelineStateCache.h"
+#include "VisualizeTexture.h"
 
 DECLARE_GPU_STAT_NAMED(ShadowDepths, TEXT("Shadow Depths"));
 
diff --git a/Engine/Source/Runtime/Renderer/Private/TranslucentLighting.cpp b/Engine/Source/Runtime/Renderer/Private/TranslucentLighting.cpp
--- a/Engine/Source/Runtime/Renderer/Private/TranslucentLighting.cpp
+++ b/Engine/Source/Runtime/Renderer/Private/TranslucentLighting.cpp
@@ -44,6 +44,7 @@
 #include "VolumeRendering.h"
 #include "VolumeLighting.h"
 #include "PipelineStateCache.h"
+#include "VisualizeTexture.h"
 
 class FMaterial;
 
diff --git a/Engine/Source/Runtime/Renderer/Private/VT/VirtualTextureFeedback.cpp b/Engine/Source/Runtime/Renderer/Private/VT/VirtualTextureFeedback.cpp
--- a/Engine/Source/Runtime/Renderer/Private/VT/VirtualTextureFeedback.cpp
+++ b/Engine/Source/Runtime/Renderer/Private/VT/VirtualTextureFeedback.cpp
@@ -1,7 +1,7 @@
 // Copyright 1998-2018 Epic Games, Inc. All Rights Reserved.
 
 #include "VirtualTextureFeedback.h"
-#include "PostProcess/VisualizeTexture.h"
+#include "VisualizeTexture.h"
 
 #include "ClearQuad.h"
 
diff --git a/Engine/Source/Runtime/Renderer/Private/VT/VirtualTextureSpace.cpp b/Engine/Source/Runtime/Renderer/Private/VT/VirtualTextureSpace.cpp
--- a/Engine/Source/Runtime/Renderer/Private/VT/VirtualTextureSpace.cpp
+++ b/Engine/Source/Runtime/Renderer/Private/VT/VirtualTextureSpace.cpp
@@ -10,7 +10,8 @@
 #include "PipelineStateCache.h"
 #include "HAL/IConsoleManager.h"
 #include "SceneUtils.h"
-#include "PostProcess/VisualizeTexture.h"
+#include "VisualizeTexture.h"
+#include "DummyRenderResources.h"
 
 static TAutoConsoleVariable<int32> CVarVTRefreshEntirePageTable(
 	TEXT("r.VT.RefreshEntirePageTable"),
diff --git a/Engine/Source/Runtime/Renderer/Private/VelocityRendering.cpp b/Engine/Source/Runtime/Renderer/Private/VelocityRendering.cpp
--- a/Engine/Source/Runtime/Renderer/Private/VelocityRendering.cpp
+++ b/Engine/Source/Runtime/Renderer/Private/VelocityRendering.cpp
@@ -23,6 +23,7 @@
 #include "Interfaces/ITargetPlatform.h"
 #include "Interfaces/ITargetPlatformManagerModule.h"
 #endif
+#include "VisualizeTexture.h"
 
 // Changing this causes a full shader recompile
 static TAutoConsoleVariable<int32> CVarBasePassOutputsVelocity(
diff --git a/Engine/Source/Runtime/Renderer/Private/VisualizeTexturePresent.cpp b/Engine/Source/Runtime/Renderer/Private/VisualizeTexturePresent.cpp
new file mode 100644
--- /dev/null
+++ b/Engine/Source/Runtime/Renderer/Private/VisualizeTexturePresent.cpp
@@ -0,0 +1,824 @@
+// Copyright 1998-2018 Epic Games, Inc. All Rights Reserved.
+
+/*=============================================================================
+	VisualizeTexture.cpp: Post processing visualize texture.
+=============================================================================*/
+
+#include "VisualizeTexturePresent.h"
+#include "VisualizeTexture.h"
+#include "ShaderParameters.h"
+#include "RHIStaticStates.h"
+#include "Shader.h"
+#include "StaticBoundShaderState.h"
+#include "SceneUtils.h"
+#include "HAL/FileManager.h"
+#include "Misc/FileHelper.h"
+#include "Misc/Paths.h"
+#include "Misc/App.h"
+#include "CanvasTypes.h"
+#include "Engine/Canvas.h"
+#include "UnrealEngine.h"
+#include "RenderTargetPool.h"
+#include "PostProcess/SceneRenderTargets.h"
+#include "GlobalShader.h"
+#include "RenderTargetTemp.h"
+#include "ScreenRendering.h"
+#include "PostProcess/SceneFilterRendering.h"
+#include "PostProcess/PostProcessing.h"
+#include "PipelineStateCache.h"
+
+
+/** Encapsulates a simple copy pixel shader. */
+class FVisualizeTexturePresentPS : public FGlobalShader
+{
+	DECLARE_GLOBAL_SHADER(FVisualizeTexturePresentPS);
+	SHADER_USE_PARAMETER_STRUCT(FVisualizeTexturePresentPS, FGlobalShader);
+
+	static bool ShouldCompilePermutation(const FGlobalShaderPermutationParameters& Parameters)
+	{
+		return true;
+	}
+
+	BEGIN_SHADER_PARAMETER_STRUCT(FParameters, )
+		SHADER_PARAMETER_TEXTURE(Texture2D, VisualizeTexture2D)
+		SHADER_PARAMETER_SAMPLER(SamplerState, VisualizeTexture2DSampler)
+	END_SHADER_PARAMETER_STRUCT()
+};
+
+IMPLEMENT_GLOBAL_SHADER(FVisualizeTexturePresentPS, "/Engine/Private/Tools/VisualizeTexture.usf", "PresentPS", SF_Pixel);
+
+
+// draw a single pixel sized rectangle using 4 sub elements
+static void DrawBorder(FCanvas& Canvas, const FIntRect Rect, FLinearColor Color)
+{
+	// top
+	Canvas.DrawTile(Rect.Min.X, Rect.Min.Y, Rect.Max.X - Rect.Min.X, 1, 0, 0, 1, 1, Color);
+	// bottom
+	Canvas.DrawTile(Rect.Min.X, Rect.Max.Y - 1, Rect.Max.X - Rect.Min.X, 1, 0, 0, 1, 1, Color);
+	// left
+	Canvas.DrawTile(Rect.Min.X, Rect.Min.Y + 1, 1, Rect.Max.Y - Rect.Min.Y - 2, 0, 0, 1, 1, Color);
+	// right
+	Canvas.DrawTile(Rect.Max.X - 1, Rect.Min.Y + 1, 1, Rect.Max.Y - Rect.Min.Y - 2, 0, 0, 1, 1, Color);
+}
+
+// helper class to get a consistent layout in multiple functions
+// MaxX and Y are the output value that can be requested during or after iteration
+// Examples usages:
+//    FRenderTargetPoolEventIterator It(RenderTargetPoolEvents, OptionalStartIndex);
+//    while(FRenderTargetPoolEvent* Event = It.Iterate()) {}
+struct FRenderTargetPoolEventIterator
+{
+	int32 Index;
+	TArray<FRenderTargetPoolEvent>& RenderTargetPoolEvents;
+	bool bLineContent;
+	uint32 TotalWidth;
+	int32 Y;
+
+	// constructor
+	FRenderTargetPoolEventIterator(TArray<FRenderTargetPoolEvent>& InRenderTargetPoolEvents, int32 InIndex = 0)
+		: Index(InIndex)
+		, RenderTargetPoolEvents(InRenderTargetPoolEvents)
+		, bLineContent(false)
+		, TotalWidth(1)
+		, Y(0)
+	{
+		Touch();
+	}
+
+	FRenderTargetPoolEvent* operator*()
+	{
+		if (Index < RenderTargetPoolEvents.Num())
+		{
+			return &RenderTargetPoolEvents[Index];
+		}
+
+		return 0;
+	}
+
+	// @return 0 if end was reached
+	FRenderTargetPoolEventIterator& operator++()
+	{
+		if (Index < RenderTargetPoolEvents.Num())
+		{
+			++Index;
+		}
+
+		Touch();
+
+		return *this;
+	}
+
+	int32 FindClosingEventY() const
+	{
+		FRenderTargetPoolEventIterator It = *this;
+
+		const ERenderTargetPoolEventType StartType = (*It)->GetEventType();
+
+		if (StartType == ERTPE_Alloc)
+		{
+			int32 PoolEntryId = RenderTargetPoolEvents[Index].GetPoolEntryId();
+
+			++It;
+
+			// search for next Dealloc of the same PoolEntryId
+			for (; *It; ++It)
+			{
+				FRenderTargetPoolEvent* Event = *It;
+
+				if (Event->GetEventType() == ERTPE_Dealloc && Event->GetPoolEntryId() == PoolEntryId)
+				{
+					break;
+				}
+			}
+		}
+		else if (StartType == ERTPE_Phase)
+		{
+			++It;
+
+			// search for next Phase
+			for (; *It; ++It)
+			{
+				FRenderTargetPoolEvent* Event = *It;
+
+				if (Event->GetEventType() == ERTPE_Phase)
+				{
+					break;
+				}
+			}
+		}
+		else
+		{
+			check(0);
+		}
+
+		return It.Y;
+	}
+
+private:
+
+	void Touch()
+	{
+		if (Index < RenderTargetPoolEvents.Num())
+		{
+			const FRenderTargetPoolEvent& Event = RenderTargetPoolEvents[Index];
+
+			const ERenderTargetPoolEventType Type = Event.GetEventType();
+
+			if (Type == ERTPE_Alloc)
+			{
+				// for now they are all equal width
+				TotalWidth = FMath::Max(TotalWidth, Event.GetColumnX() + Event.GetColumnSize());
+			}
+			Y = Event.GetTimeStep();
+		}
+	}
+};
+
+// static
+uint32 FVisualizeTexturePresent::ComputeEventDisplayHeight()
+{
+	FRenderTargetPoolEventIterator It(GRenderTargetPool.RenderTargetPoolEvents);
+
+	for (; *It; ++It)
+	{
+	}
+
+	return It.Y;
+}
+
+// static
+void FVisualizeTexturePresent::OnStartRender(const FViewInfo& View)
+{
+	GVisualizeTexture.FeatureLevel = View.GetFeatureLevel();
+	GVisualizeTexture.bEnabled = true;
+
+#if !(UE_BUILD_SHIPPING || UE_BUILD_TEST)
+	// VisualizeTexture observed render target is set each frame
+	GVisualizeTexture.VisualizeTextureContent.SafeRelease();
+	GVisualizeTexture.VisualizeTextureDesc = FPooledRenderTargetDesc();
+	GVisualizeTexture.VisualizeTextureDesc.DebugName = TEXT("VisualizeTexture");
+
+	GVisualizeTexture.ObservedDebugNameReusedCurrent = 0;
+
+	// only needed for VisualizeTexture (todo: optimize out when possible)
+	{
+		for (TMap<const TCHAR*, uint32>::TIterator It(GVisualizeTexture.VisualizeTextureCheckpoints); It; ++It)
+		{
+			uint32& Value = It.Value();
+
+			// 0 as it was not used this frame yet
+			Value = 0;
+		}
+	}
+#endif // !(UE_BUILD_SHIPPING || UE_BUILD_TEST)
+}
+
+// static
+void FVisualizeTexturePresent::PresentContent(FRHICommandListImmediate& RHICmdList, const FViewInfo& View)
+{
+	if (GRenderTargetPool.RenderTargetPoolEvents.Num())
+	{
+		GRenderTargetPool.AddPhaseEvent(TEXT("FrameEnd"));
+
+		FIntPoint DisplayLeftTop(20, 50);
+		// on the right we leave more space to make the mouse tooltip readable
+		FIntPoint DisplayExtent(View.ViewRect.Width() - DisplayLeftTop.X * 2 - 140, View.ViewRect.Height() - DisplayLeftTop.Y * 2);
+
+		// if the area is not too small
+		if (DisplayExtent.X > 50 && DisplayExtent.Y > 50)
+		{
+			FRenderTargetPool::SMemoryStats MemoryStats = GRenderTargetPool.ComputeView();
+
+			SetRenderTarget(RHICmdList, View.Family->RenderTarget->GetRenderTargetTexture(), FTextureRHIRef());
+			RHICmdList.SetViewport(0, 0, 0.0f, FSceneRenderTargets::Get(RHICmdList).GetBufferSizeXY().X, FSceneRenderTargets::Get(RHICmdList).GetBufferSizeXY().Y, 1.0f);
+
+
+			FRenderTargetTemp TempRenderTarget(View, View.UnconstrainedViewRect.Size());
+			FCanvas Canvas(&TempRenderTarget, NULL, View.Family->CurrentRealTime, View.Family->CurrentWorldTime, View.Family->DeltaWorldTime, View.GetFeatureLevel());
+
+			// TinyFont property
+			const int32 FontHeight = 12;
+
+			FIntPoint MousePos = View.CursorPos;
+
+			FLinearColor BackgroundColor = FLinearColor(0.0f, 0.0f, 0.0f, 0.7f);
+			FLinearColor PhaseColor = FLinearColor(0.2f, 0.1f, 0.05f, 0.8f);
+			FLinearColor ElementColor = FLinearColor(0.3f, 0.3f, 0.3f, 0.9f);
+			FLinearColor ElementColorVRam = FLinearColor(0.4f, 0.25f, 0.25f, 0.9f);
+
+			UTexture2D* GradientTexture = UCanvas::StaticClass()->GetDefaultObject<UCanvas>()->GradientTexture0;
+
+			// background rectangle
+			Canvas.DrawTile(DisplayLeftTop.X, DisplayLeftTop.Y - 1 * FontHeight - 1, DisplayExtent.X, DisplayExtent.Y + FontHeight, 0, 0, 1, 1, BackgroundColor);
+
+			{
+				uint32 MB = 1024 * 1024;
+				uint32 MBm1 = MB - 1;
+
+				FString Headline = *FString::Printf(TEXT("RenderTargetPool elements(x) over time(y) >= %dKB, Displayed/Total:%d/%dMB"),
+					GRenderTargetPool.EventRecordingSizeThreshold,
+					(uint32)((MemoryStats.DisplayedUsageInBytes + MBm1) / MB),
+					(uint32)((MemoryStats.TotalUsageInBytes + MBm1) / MB));
+				Canvas.DrawShadowedString(DisplayLeftTop.X, DisplayLeftTop.Y - 1 * FontHeight - 1, *Headline, GEngine->GetTinyFont(), FLinearColor(1, 1, 1));
+			}
+
+			uint32 EventDisplayHeight = FVisualizeTexturePresent::ComputeEventDisplayHeight();
+
+			float ScaleX = DisplayExtent.X / (float)MemoryStats.TotalColumnSize;
+			float ScaleY = DisplayExtent.Y / (float)EventDisplayHeight;
+
+			// 0 if none
+			FRenderTargetPoolEvent* HighlightedEvent = 0;
+			FIntRect HighlightedRect;
+
+			// Phase events
+			for (FRenderTargetPoolEventIterator It(GRenderTargetPool.RenderTargetPoolEvents); *It; ++It)
+			{
+				FRenderTargetPoolEvent* Event = *It;
+
+				if (Event->GetEventType() == ERTPE_Phase)
+				{
+					int32 Y0 = It.Y;
+					int32 Y1 = It.FindClosingEventY();
+
+					FIntPoint PixelLeftTop((int32)(DisplayLeftTop.X), (int32)(DisplayLeftTop.Y + ScaleY * Y0));
+					FIntPoint PixelRightBottom((int32)(DisplayLeftTop.X + DisplayExtent.X), (int32)(DisplayLeftTop.Y + ScaleY * Y1));
+
+					bool bHighlight = MousePos.X >= PixelLeftTop.X && MousePos.X < PixelRightBottom.X && MousePos.Y >= PixelLeftTop.Y && MousePos.Y <= PixelRightBottom.Y;
+
+					if (bHighlight)
+					{
+						HighlightedEvent = Event;
+						HighlightedRect = FIntRect(PixelLeftTop, PixelRightBottom);
+					}
+
+					// UMax is 0.9f to avoid getting some wrap texture leaking in at the bottom
+					Canvas.DrawTile(PixelLeftTop.X, PixelLeftTop.Y, PixelRightBottom.X - PixelLeftTop.X, PixelRightBottom.Y - PixelLeftTop.Y, 0, 0, 1, 0.9f, PhaseColor, GradientTexture->Resource);
+				}
+			}
+
+			// Alloc / Dealloc events
+			for (FRenderTargetPoolEventIterator It(GRenderTargetPool.RenderTargetPoolEvents); *It; ++It)
+			{
+				FRenderTargetPoolEvent* Event = *It;
+
+				if (Event->GetEventType() == ERTPE_Alloc && Event->GetColumnSize())
+				{
+					int32 Y0 = It.Y;
+					int32 Y1 = It.FindClosingEventY();
+
+					int32 X0 = Event->GetColumnX();
+					// for now they are all equal width
+					int32 X1 = X0 + Event->GetColumnSize();
+
+					FIntPoint PixelLeftTop((int32)(DisplayLeftTop.X + ScaleX * X0), (int32)(DisplayLeftTop.Y + ScaleY * Y0));
+					FIntPoint PixelRightBottom((int32)(DisplayLeftTop.X + ScaleX * X1), (int32)(DisplayLeftTop.Y + ScaleY * Y1));
+
+					bool bHighlight = MousePos.X >= PixelLeftTop.X && MousePos.X < PixelRightBottom.X && MousePos.Y >= PixelLeftTop.Y && MousePos.Y <= PixelRightBottom.Y;
+
+					if (bHighlight)
+					{
+						HighlightedEvent = Event;
+						HighlightedRect = FIntRect(PixelLeftTop, PixelRightBottom);
+					}
+
+					FLinearColor Color = ElementColor;
+
+					// Highlight EDRAM/FastVRAM usage
+					if (Event->GetDesc().Flags & TexCreate_FastVRAM)
+					{
+						Color = ElementColorVRam;
+					}
+
+					Canvas.DrawTile(
+						PixelLeftTop.X, PixelLeftTop.Y,
+						PixelRightBottom.X - PixelLeftTop.X - 1, PixelRightBottom.Y - PixelLeftTop.Y - 1,
+						0, 0, 1, 1, Color);
+				}
+			}
+
+			if (HighlightedEvent)
+			{
+				DrawBorder(Canvas, HighlightedRect, FLinearColor(0.8f, 0, 0, 0.5f));
+
+				// Offset to not intersect with crosshair (in editor) or arrow (in game).
+				FIntPoint Pos = MousePos + FIntPoint(12, 4);
+
+				if (HighlightedEvent->GetEventType() == ERTPE_Phase)
+				{
+					FString PhaseText = *FString::Printf(TEXT("Phase: %s"), *HighlightedEvent->GetPhaseName());
+
+					Canvas.DrawShadowedString(Pos.X, Pos.Y + 0 * FontHeight, *PhaseText, GEngine->GetTinyFont(), FLinearColor(0.5f, 0.5f, 1));
+				}
+				else
+				{
+					FString SizeString = FString::Printf(TEXT("%d KB"), (HighlightedEvent->GetSizeInBytes() + 1024) / 1024);
+
+					Canvas.DrawShadowedString(Pos.X, Pos.Y + 0 * FontHeight, HighlightedEvent->GetDesc().DebugName, GEngine->GetTinyFont(), FLinearColor(1, 1, 0));
+					Canvas.DrawShadowedString(Pos.X, Pos.Y + 1 * FontHeight, *HighlightedEvent->GetDesc().GenerateInfoString(), GEngine->GetTinyFont(), FLinearColor(1, 1, 0));
+					Canvas.DrawShadowedString(Pos.X, Pos.Y + 2 * FontHeight, *SizeString, GEngine->GetTinyFont(), FLinearColor(1, 1, 0));
+				}
+			}
+
+			Canvas.Flush_RenderThread(RHICmdList);
+
+			GRenderTargetPool.CurrentEventRecordingTime = 0;
+			GRenderTargetPool.RenderTargetPoolEvents.Empty();
+		}
+	}
+
+	if (GVisualizeTexture.Mode != 0)
+	{
+		// old mode is used, lets copy the specified texture to do it similar to the new system
+		FPooledRenderTarget* Element = GRenderTargetPool.GetElementById(GVisualizeTexture.Mode - 1);
+		if (Element)
+		{
+			FRDGBuilder GraphBuilder(RHICmdList);
+
+			TRefCountPtr<IPooledRenderTarget> ElementRefCount;
+
+			GVisualizeTexture.CreateContentCapturePass(GraphBuilder, GraphBuilder.RegisterExternalTexture(ElementRefCount));
+			GraphBuilder.Execute();
+		}
+	}
+
+	const FTexture2DRHIRef& RenderTargetTexture = View.Family->RenderTarget->GetRenderTargetTexture();
+
+	if (!GVisualizeTexture.VisualizeTextureContent
+		|| !IsValidRef(RenderTargetTexture)
+		|| !GVisualizeTexture.bEnabled)
+	{
+		// visualize feature is deactivated
+		return;
+	}
+
+	FPooledRenderTargetDesc Desc = GVisualizeTexture.VisualizeTextureDesc;
+
+	FIntPoint SrcSize = Desc.Extent;
+	FIntRect SrcRect(0, 0, SrcSize.X, SrcSize.Y);
+	FIntRect DestRect = SrcRect;
+	{
+
+		// set ViewRect
+		switch (GVisualizeTexture.UVInputMapping)
+		{
+		case 0:
+		{
+			DestRect = SrcRect;
+		}
+		break;
+
+		case 1:
+		{
+			DestRect = SrcRect;
+		}
+		break;
+
+		case 2:
+		{
+			FIntPoint Center = View.UnconstrainedViewRect.Size() / 2;
+			FIntPoint HalfMin = SrcSize / 2;
+			FIntPoint HalfMax = SrcSize - HalfMin;
+
+			DestRect = FIntRect(Center - HalfMin, Center + HalfMax);
+			break;
+		}
+		break;
+
+		case 3:
+		{
+			float SrcAspectRatio = float(Desc.Extent.X) / float(Desc.Extent.Y);
+
+			int32 TargetedHeight = 0.3f * View.UnconstrainedViewRect.Height();
+			int32 TargetedWidth = SrcAspectRatio * TargetedHeight;
+			int32 OffsetFromBorder = 100;
+
+			DestRect.Min.X = View.UnconstrainedViewRect.Min.X + OffsetFromBorder;
+			DestRect.Max.X = DestRect.Min.X + TargetedWidth;
+			DestRect.Max.Y = View.UnconstrainedViewRect.Max.Y - OffsetFromBorder;
+			DestRect.Min.Y = DestRect.Max.Y - TargetedHeight;
+			break;
+		}
+		break;
+
+		default:
+			break;
+		}
+	}
+
+	auto& RenderTarget = View.Family->RenderTarget->GetRenderTargetTexture();
+	SetRenderTarget(RHICmdList, RenderTarget, FTextureRHIRef(), true);
+	RHICmdList.SetViewport(DestRect.Min.X, DestRect.Min.Y, 0.0f, DestRect.Max.X, DestRect.Max.Y, 1.0f);
+
+	FGraphicsPipelineStateInitializer GraphicsPSOInit;
+	RHICmdList.ApplyCachedRenderTargets(GraphicsPSOInit);
+	GraphicsPSOInit.BlendState = TStaticBlendState<>::GetRHI();
+	GraphicsPSOInit.RasterizerState = TStaticRasterizerState<>::GetRHI();
+	GraphicsPSOInit.DepthStencilState = TStaticDepthStencilState<false, CF_Always>::GetRHI();
+
+	auto ShaderMap = View.ShaderMap;
+	TShaderMapRef<FPostProcessVS> VertexShader(ShaderMap);
+	TShaderMapRef<FVisualizeTexturePresentPS> PixelShader(ShaderMap);
+
+	GraphicsPSOInit.BoundShaderState.VertexDeclarationRHI = GFilterVertexDeclaration.VertexDeclarationRHI;
+	GraphicsPSOInit.BoundShaderState.VertexShaderRHI = GETSAFERHISHADER_VERTEX(*VertexShader);
+	GraphicsPSOInit.BoundShaderState.PixelShaderRHI = GETSAFERHISHADER_PIXEL(*PixelShader);
+	GraphicsPSOInit.PrimitiveType = PT_TriangleList;
+	SetGraphicsPipelineState(RHICmdList, GraphicsPSOInit);
+
+	VertexShader->SetParameters(RHICmdList, View.ViewUniformBuffer);
+	{
+		FVisualizeTexturePresentPS::FParameters Parameters;
+		Parameters.VisualizeTexture2D = GVisualizeTexture.VisualizeTextureContent->GetRenderTargetItem().ShaderResourceTexture;
+		Parameters.VisualizeTexture2DSampler = TStaticSamplerState<SF_Point, AM_Clamp, AM_Clamp, AM_Clamp>::GetRHI();
+		SetShaderParameters(RHICmdList, *PixelShader, PixelShader->GetPixelShader(), Parameters);
+	}
+
+	{
+		SCOPED_DRAW_EVENT(RHICmdList, VisCopyToMain);
+		DrawRectangle(
+			RHICmdList,
+			0, 0,
+			DestRect.Width(), DestRect.Height(),
+			SrcRect.Min.X, SrcRect.Min.Y,
+			SrcRect.Width(), SrcRect.Height(),
+			DestRect.Size(),
+			SrcSize,
+			*VertexShader,
+			EDRF_Default);
+	}
+
+	FRenderTargetTemp TempRenderTarget(View, View.UnconstrainedViewRect.Size());
+	FCanvas Canvas(&TempRenderTarget, NULL, View.Family->CurrentRealTime, View.Family->CurrentWorldTime, View.Family->DeltaWorldTime, View.GetFeatureLevel());
+
+	float X = 100 + View.UnconstrainedViewRect.Min.X;
+	float Y = 160 + View.UnconstrainedViewRect.Min.Y;
+	float YStep = 14;
+
+	{
+		uint32 ReuseCount = GVisualizeTexture.ObservedDebugNameReusedCurrent;
+
+		FString ExtendedName;
+		if (ReuseCount)
+		{
+			uint32 ReuseGoal = FMath::Min(ReuseCount - 1, GVisualizeTexture.ObservedDebugNameReusedGoal);
+
+			// was reused this frame
+			ExtendedName = FString::Printf(TEXT("%s@%d @0..%d"), Desc.DebugName, ReuseGoal, ReuseCount - 1);
+		}
+		else
+		{
+			// was not reused this frame but can be referenced
+			ExtendedName = FString::Printf(TEXT("%s"), Desc.DebugName);
+		}
+
+		FString Channels = TEXT("RGB");
+		switch (GVisualizeTexture.SingleChannel)
+		{
+		case 0: Channels = TEXT("R"); break;
+		case 1: Channels = TEXT("G"); break;
+		case 2: Channels = TEXT("B"); break;
+		case 3: Channels = TEXT("A"); break;
+		}
+		float Multiplier = (GVisualizeTexture.SingleChannel == -1) ? GVisualizeTexture.RGBMul : GVisualizeTexture.SingleChannelMul;
+
+		FString Line = FString::Printf(TEXT("VisualizeTexture: %d \"%s\" %s*%g UV%d"),
+			GVisualizeTexture.Mode,
+			*ExtendedName,
+			*Channels,
+			Multiplier,
+			GVisualizeTexture.UVInputMapping);
+
+		Canvas.DrawShadowedString(X, Y += YStep, *Line, GetStatsFont(), FLinearColor(1, 1, 1));
+	}
+	{
+		FString Line = FString::Printf(TEXT("   TextureInfoString(): %s"), *(Desc.GenerateInfoString()));
+		Canvas.DrawShadowedString(X + 10, Y += YStep, *Line, GetStatsFont(), FLinearColor(1, 1, 1));
+	}
+	{
+		FString Line = FString::Printf(TEXT("  BufferSize:(%d,%d)"), FSceneRenderTargets::Get(RHICmdList).GetBufferSizeXY().X, FSceneRenderTargets::Get(RHICmdList).GetBufferSizeXY().Y);
+		Canvas.DrawShadowedString(X + 10, Y += YStep, *Line, GetStatsFont(), FLinearColor(1, 1, 1));
+	}
+
+	const FSceneViewFamily& ViewFamily = *View.Family;
+
+	for (int32 ViewId = 0; ViewId < ViewFamily.Views.Num(); ++ViewId)
+	{
+		const FViewInfo* ViewIt = static_cast<const FViewInfo*>(ViewFamily.Views[ViewId]);
+		FString Line = FString::Printf(TEXT("   View #%d: (%d,%d)-(%d,%d)"), ViewId + 1,
+			ViewIt->UnscaledViewRect.Min.X, ViewIt->UnscaledViewRect.Min.Y, ViewIt->UnscaledViewRect.Max.X, ViewIt->UnscaledViewRect.Max.Y);
+		Canvas.DrawShadowedString(X + 10, Y += YStep, *Line, GetStatsFont(), FLinearColor(1, 1, 1));
+	}
+
+	X += 40;
+
+	if (Desc.Flags & TexCreate_CPUReadback)
+	{
+		Canvas.DrawShadowedString(X, Y += YStep, TEXT("Content cannot be visualized on the GPU (TexCreate_CPUReadback)"), GetStatsFont(), FLinearColor(1, 1, 0));
+	}
+	else
+	{
+		Canvas.DrawShadowedString(X, Y += YStep, TEXT("Blinking Red: <0"), GetStatsFont(), FLinearColor(1, 0, 0));
+		Canvas.DrawShadowedString(X, Y += YStep, TEXT("Blinking Blue: NAN or Inf"), GetStatsFont(), FLinearColor(0, 0, 1));
+
+		// add explicit legend for SceneDepth and ShadowDepth as the display coloring is an artificial choice. 
+		const bool bDepthTexture = (Desc.TargetableFlags & TexCreate_DepthStencilTargetable) != 0;
+		const bool bShadowDepth = (Desc.Format == PF_ShadowDepth);
+		if (bShadowDepth)
+		{
+			Canvas.DrawShadowedString(X, Y += YStep, TEXT("Color Key: Linear with white near and teal distant"), GetStatsFont(), FLinearColor(54.f / 255.f, 117.f / 255.f, 136.f / 255.f));
+		}
+		else if (bDepthTexture)
+		{
+			Canvas.DrawShadowedString(X, Y += YStep, TEXT("Color Key: Nonlinear with white distant"), GetStatsFont(), FLinearColor(0.5, 0, 0));
+		}
+	}
+
+	Canvas.Flush_RenderThread(RHICmdList);
+}
+
+struct FSortedLines
+{
+	FString Line;
+	int32 SortIndex;
+	uint32 PoolIndex;
+
+	FORCEINLINE bool operator<(const FSortedLines &B) const
+	{
+		// first large ones
+		if (SortIndex < B.SortIndex)
+		{
+			return true;
+		}
+		if (SortIndex > B.SortIndex)
+		{
+			return false;
+		}
+
+		return Line < B.Line;
+	}
+};
+
+// static
+void FVisualizeTexturePresent::DebugLog(bool bExtended)
+{
+#if !(UE_BUILD_SHIPPING || UE_BUILD_TEST)
+	{
+		TArray<FSortedLines> SortedLines;
+
+		for (uint32 i = 0, Num = GRenderTargetPool.GetElementCount(); i < Num; ++i)
+		{
+			FPooledRenderTarget* RT = GRenderTargetPool.GetElementById(i);
+
+			if (!RT)
+			{
+				continue;
+			}
+
+			FPooledRenderTargetDesc Desc = RT->GetDesc();
+
+			if (GVisualizeTexture.bFullList || (Desc.Flags & TexCreate_HideInVisualizeTexture) == 0)
+			{
+				uint32 SizeInKB = (RT->ComputeMemorySize() + 1023) / 1024;
+
+				FString UnusedStr;
+
+				if (RT->GetUnusedForNFrames() > 0)
+				{
+					if (!GVisualizeTexture.bFullList)
+					{
+						continue;
+					}
+
+					UnusedStr = FString::Printf(TEXT(" unused(%d)"), RT->GetUnusedForNFrames());
+				}
+
+				FSortedLines Element;
+
+				Element.PoolIndex = i;
+
+				// sort by index
+				Element.SortIndex = i;
+
+				FString InfoString = Desc.GenerateInfoString();
+				if (GVisualizeTexture.SortOrder == -1)
+				{
+					// constant works well with the average name length
+					const uint32 TotelSpacerSize = 36;
+					uint32 SpaceCount = FMath::Max<int32>(0, TotelSpacerSize - InfoString.Len());
+
+					for (uint32 Space = 0; Space < SpaceCount; ++Space)
+					{
+						InfoString.AppendChar((TCHAR)' ');
+					}
+
+					// sort by index
+					Element.Line = FString::Printf(TEXT("%s %s %d KB%s"), *InfoString, Desc.DebugName, SizeInKB, *UnusedStr);
+				}
+				else if (GVisualizeTexture.SortOrder == 0)
+				{
+					// sort by name
+					Element.Line = FString::Printf(TEXT("%s %s %d KB%s"), Desc.DebugName, *InfoString, SizeInKB, *UnusedStr);
+					Element.SortIndex = 0;
+				}
+				else if (GVisualizeTexture.SortOrder == 1)
+				{
+					// sort by size (large ones first)
+					Element.Line = FString::Printf(TEXT("%d KB %s %s%s"), SizeInKB, *InfoString, Desc.DebugName, *UnusedStr);
+					Element.SortIndex = -(int32)SizeInKB;
+				}
+				else
+				{
+					check(0);
+				}
+
+				if (Desc.Flags & TexCreate_FastVRAM)
+				{
+					FRHIResourceInfo Info;
+
+					FTextureRHIRef Texture = RT->GetRenderTargetItem().ShaderResourceTexture;
+
+					if (!IsValidRef(Texture))
+					{
+						Texture = RT->GetRenderTargetItem().TargetableTexture;
+					}
+
+					if (IsValidRef(Texture))
+					{
+						RHIGetResourceInfo(Texture, Info);
+					}
+
+					if (Info.VRamAllocation.AllocationSize)
+					{
+						// note we do KB for more readable numbers but this can cause quantization loss
+						Element.Line += FString::Printf(TEXT(" VRamInKB(Start/Size):%d/%d"),
+							Info.VRamAllocation.AllocationStart / 1024,
+							(Info.VRamAllocation.AllocationSize + 1023) / 1024);
+					}
+					else
+					{
+						Element.Line += TEXT(" VRamInKB(Start/Size):<NONE>");
+					}
+				}
+
+				SortedLines.Add(Element);
+			}
+		}
+
+		SortedLines.Sort();
+
+		{
+			for (int32 Index = 0; Index < SortedLines.Num(); Index++)
+			{
+				const FSortedLines& Entry = SortedLines[Index];
+
+				UE_LOG(LogConsoleResponse, Log, TEXT("   %3d = %s"), Entry.PoolIndex + 1, *Entry.Line);
+			}
+		}
+
+		// clean flags for next use
+		GVisualizeTexture.bFullList = false;
+		GVisualizeTexture.SortOrder = -1;
+	}
+
+	UE_LOG(LogConsoleResponse, Log, TEXT(""));
+
+	// log names (alternative method to look at the rendertargets)
+	if (bExtended)
+	{
+		UE_LOG(LogConsoleResponse, Log, TEXT("CheckpointName (what was rendered this frame, use <Name>@<Number> to get intermediate versions):"));
+
+		TArray<FString> Entries;
+
+		// sorted by pointer for efficiency, now we want to print sorted alphabetically
+		for (TMap<const TCHAR*, uint32>::TIterator It(GVisualizeTexture.VisualizeTextureCheckpoints); It; ++It)
+		{
+			const TCHAR* Key = It.Key();
+			uint32 Value = It.Value();
+
+			/*					if(Value)
+						{
+							// was reused this frame
+							Entries.Add(FString::Printf(TEXT("%s @0..%d"), *Key.GetPlainNameString(), Value - 1));
+						}
+						else
+			*/ {
+			// was not reused this frame but can be referenced
+				Entries.Add(Key);
+			}
+		}
+
+		Entries.Sort();
+
+		// that number works well with the name length we have
+		const uint32 ColumnCount = 5;
+		const uint32 SpaceBetweenColumns = 1;
+		uint32 ColumnHeight = FMath::DivideAndRoundUp((uint32)Entries.Num(), ColumnCount);
+
+		// width of the column in characters, init with 0
+		uint32 ColumnWidths[ColumnCount] = {};
+
+		for (int32 Index = 0; Index < Entries.Num(); ++Index)
+		{
+			uint32 Column = Index / ColumnHeight;
+
+			const FString& Entry = *Entries[Index];
+
+			ColumnWidths[Column] = FMath::Max(ColumnWidths[Column], (uint32)Entry.Len());
+		}
+
+		// print them sorted, if possible multiple in a line
+		{
+			FString Line;
+
+			for (int32 OutputIndex = 0; OutputIndex < Entries.Num(); ++OutputIndex)
+			{
+				// 0..ColumnCount-1
+				uint32 Column = OutputIndex % ColumnCount;
+				int32 Row = OutputIndex / ColumnCount;
+
+				uint32 Index = Row + Column * ColumnHeight;
+
+				bool bLineEnd = true;
+
+				if (Index < (uint32)Entries.Num())
+				{
+					bLineEnd = (Column + 1 == ColumnCount);
+
+					// for human readability we order them to be per column
+					const FString& Entry = *Entries[Index];
+
+					Line += Entry;
+
+					int32 SpaceCount = ColumnWidths[Column] + SpaceBetweenColumns - Entry.Len();
+
+					// otehrwise a fomer pass was producing bad data
+					check(SpaceCount >= 0);
+
+					for (int32 Space = 0; Space < SpaceCount; ++Space)
+					{
+						Line.AppendChar((TCHAR)' ');
+					}
+				}
+
+				if (bLineEnd)
+				{
+					Line.TrimEndInline();
+					UE_LOG(LogConsoleResponse, Log, TEXT("   %s"), *Line);
+					Line.Empty();
+				}
+			}
+		}
+	}
+
+	{
+		uint32 WholeCount;
+		uint32 WholePoolInKB;
+		uint32 UsedInKB;
+
+		GRenderTargetPool.GetStats(WholeCount, WholePoolInKB, UsedInKB);
+
+		UE_LOG(LogConsoleResponse, Log, TEXT("Pool: %d/%d MB (referenced/allocated)"), (UsedInKB + 1023) / 1024, (WholePoolInKB + 1023) / 1024);
+	}
+#endif
+}
diff --git a/Engine/Source/Runtime/Renderer/Private/VisualizeTexturePresent.h b/Engine/Source/Runtime/Renderer/Private/VisualizeTexturePresent.h
new file mode 100644
--- /dev/null
+++ b/Engine/Source/Runtime/Renderer/Private/VisualizeTexturePresent.h
@@ -0,0 +1,29 @@
+// Copyright 1998-2018 Epic Games, Inc. All Rights Reserved.
+
+/*=============================================================================
+	VisualizeTexturePresent.h: Display texture visualization on screen.
+=============================================================================*/
+
+#include "CoreMinimal.h"
+#include "RHI.h"
+
+class FViewInfo;
+
+
+/** Presents */
+class FVisualizeTexturePresent
+{
+public:
+	/** Starts texture visualization capture. */
+	static void OnStartRender(const FViewInfo& View);
+
+	/** Present the visualize texture tool on screen. */
+	static void PresentContent(FRHICommandListImmediate& RHICmdList, const FViewInfo& View);
+
+	/** Dumps every information to log. */
+	static void DebugLog(bool bExtended);
+
+private:
+	static uint32 ComputeEventDisplayHeight();
+
+};
diff --git a/Engine/Source/Runtime/Renderer/Private/VolumetricFog.cpp b/Engine/Source/Runtime/Renderer/Private/VolumetricFog.cpp
--- a/Engine/Source/Runtime/Renderer/Private/VolumetricFog.cpp
+++ b/Engine/Source/Runtime/Renderer/Private/VolumetricFog.cpp
@@ -168,7 +168,7 @@ public:
 	FVolumetricFogMaterialSetupCS(const ShaderMetaType::CompiledShaderInitializerType& Initializer)
 		: FGlobalShader(Initializer)
 	{
-		Bindings.Bind(Initializer.ParameterMap, *FParameters::FTypeInfo::GetStructMetadata());
+		Bindings.BindForLegacyShaderParameters(Initializer.ParameterMap, *FParameters::FTypeInfo::GetStructMetadata());
 
 		VolumetricFogParameters.Bind(Initializer.ParameterMap);
 	} 
@@ -600,7 +600,7 @@ public:
 	TVolumetricFogLightScatteringCS(const ShaderMetaType::CompiledShaderInitializerType& Initializer)
 		: FGlobalShader(Initializer)
 	{
-		Bindings.Bind(Initializer.ParameterMap, *FParameters::FTypeInfo::GetStructMetadata());
+		Bindings.BindForLegacyShaderParameters(Initializer.ParameterMap, *FParameters::FTypeInfo::GetStructMetadata());
 
 		LocalShadowedLightScattering.Bind(Initializer.ParameterMap, TEXT("LocalShadowedLightScattering"));
 		LightScatteringHistory.Bind(Initializer.ParameterMap, TEXT("LightScatteringHistory"));
@@ -780,7 +780,7 @@ class FVolumetricFogFinalIntegrationCS : public FGlobalShader
 	FVolumetricFogFinalIntegrationCS(const ShaderMetaType::CompiledShaderInitializerType& Initializer)
 		: FGlobalShader(Initializer)
 	{
-		Bindings.Bind(Initializer.ParameterMap, *FParameters::FTypeInfo::GetStructMetadata());
+		Bindings.BindForLegacyShaderParameters(Initializer.ParameterMap, *FParameters::FTypeInfo::GetStructMetadata());
 
 		VolumetricFogParameters.Bind(Initializer.ParameterMap);
 	}
diff --git a/Engine/Source/Runtime/SlateRHIRenderer/Private/SlatePostProcessor.cpp b/Engine/Source/Runtime/SlateRHIRenderer/Private/SlatePostProcessor.cpp
--- a/Engine/Source/Runtime/SlateRHIRenderer/Private/SlatePostProcessor.cpp
+++ b/Engine/Source/Runtime/SlateRHIRenderer/Private/SlatePostProcessor.cpp
@@ -8,6 +8,7 @@
 #include "RendererInterface.h"
 #include "StaticBoundShaderState.h"
 #include "PipelineStateCache.h"
+#include "DummyRenderResources.h"
 
 DECLARE_CYCLE_STAT(TEXT("Slate PostProcessing RT"), STAT_SlatePostProcessingRTTime, STATGROUP_Slate);
 DECLARE_CYCLE_STAT(TEXT("Slate ColorDeficiency RT"), STAT_SlateColorDeficiencyRTTime, STATGROUP_Slate);
@@ -75,7 +76,7 @@ void FSlatePostProcessor::BlurRect(FRHICommandListImmediate& RHICmdList, IRender
 	const FSlateRect& SourceRect = RectParams.SourceRect;
 	const FSlateRect& DestRect = RectParams.DestRect;
 
-	FVertexDeclarationRHIRef VertexDecl = RendererModule.GetFilterVertexDeclaration().VertexDeclarationRHI;
+	FVertexDeclarationRHIRef VertexDecl = GFilterVertexDeclaration.VertexDeclarationRHI;
 	check(IsValidRef(VertexDecl));
 	
 	FGraphicsPipelineStateInitializer GraphicsPSOInit;
@@ -226,7 +227,7 @@ void FSlatePostProcessor::ColorDeficiency(FRHICommandListImmediate& RHICmdList,
 	const FSlateRect& SourceRect = RectParams.SourceRect;
 	const FSlateRect& DestRect = RectParams.DestRect;
 
-	FVertexDeclarationRHIRef VertexDecl = RendererModule.GetFilterVertexDeclaration().VertexDeclarationRHI;
+	FVertexDeclarationRHIRef VertexDecl = GFilterVertexDeclaration.VertexDeclarationRHI;
 	check(IsValidRef(VertexDecl));
 
 	FGraphicsPipelineStateInitializer GraphicsPSOInit;
@@ -337,7 +338,7 @@ void FSlatePostProcessor::DownsampleRect(FRHICommandListImmediate& RHICmdList, I
 			GraphicsPSOInit.RasterizerState = TStaticRasterizerState<>::GetRHI();
 			GraphicsPSOInit.DepthStencilState = TStaticDepthStencilState<false, CF_Always>::GetRHI();
 
-			GraphicsPSOInit.BoundShaderState.VertexDeclarationRHI = RendererModule.GetFilterVertexDeclaration().VertexDeclarationRHI;
+			GraphicsPSOInit.BoundShaderState.VertexDeclarationRHI = GFilterVertexDeclaration.VertexDeclarationRHI;
 			GraphicsPSOInit.BoundShaderState.VertexShaderRHI = GETSAFERHISHADER_VERTEX(*VertexShader);
 			GraphicsPSOInit.BoundShaderState.PixelShaderRHI = GETSAFERHISHADER_PIXEL(*PixelShader);
 			GraphicsPSOInit.PrimitiveType = PT_TriangleList;
@@ -415,7 +416,7 @@ void FSlatePostProcessor::UpsampleRect(FRHICommandListImmediate& RHICmdList, IRe
 
 		TShaderMapRef<FScreenPS> PixelShader(ShaderMap);
 
-		GraphicsPSOInit.BoundShaderState.VertexDeclarationRHI = RendererModule.GetFilterVertexDeclaration().VertexDeclarationRHI;
+		GraphicsPSOInit.BoundShaderState.VertexDeclarationRHI = GFilterVertexDeclaration.VertexDeclarationRHI;
 		GraphicsPSOInit.BoundShaderState.VertexShaderRHI = GETSAFERHISHADER_VERTEX(*VertexShader);
 		GraphicsPSOInit.BoundShaderState.PixelShaderRHI = GETSAFERHISHADER_PIXEL(*PixelShader);
 		GraphicsPSOInit.PrimitiveType = PT_TriangleList;
diff --git a/Engine/Source/Runtime/SlateRHIRenderer/Private/SlateRHIRenderer.cpp b/Engine/Source/Runtime/SlateRHIRenderer/Private/SlateRHIRenderer.cpp
--- a/Engine/Source/Runtime/SlateRHIRenderer/Private/SlateRHIRenderer.cpp
+++ b/Engine/Source/Runtime/SlateRHIRenderer/Private/SlateRHIRenderer.cpp
@@ -27,6 +27,7 @@
 #include "EngineModule.h"
 #include "Interfaces/ISlate3DRenderer.h"
 #include "Slate/SlateTextureAtlasInterface.h"
+#include "DummyRenderResources.h"
 
 DECLARE_CYCLE_STAT(TEXT("Slate RT: Rendering"), STAT_SlateRenderingRTTime, STATGROUP_Slate);
 DECLARE_CYCLE_STAT(TEXT("Slate RT: Create Batches"), STAT_SlateRTCreateBatches, STATGROUP_Slate);
@@ -873,7 +874,7 @@ void FSlateRHIRenderer::DrawWindow_RenderThread(FRHICommandListImmediate& RHICmd
 							// ScRGB encoding
 							TShaderMapRef<FCompositePS<1>> PixelShader(ShaderMap);
 
-							GraphicsPSOInit.BoundShaderState.VertexDeclarationRHI = RendererModule.GetFilterVertexDeclaration().VertexDeclarationRHI;
+							GraphicsPSOInit.BoundShaderState.VertexDeclarationRHI = GFilterVertexDeclaration.VertexDeclarationRHI;
 							GraphicsPSOInit.BoundShaderState.VertexShaderRHI = GETSAFERHISHADER_VERTEX(*VertexShader);
 							GraphicsPSOInit.BoundShaderState.PixelShaderRHI = GETSAFERHISHADER_PIXEL(*PixelShader);
 							GraphicsPSOInit.PrimitiveType = PT_TriangleList;
@@ -887,7 +888,7 @@ void FSlateRHIRenderer::DrawWindow_RenderThread(FRHICommandListImmediate& RHICmd
 							// ST2084 (PQ) encoding
 							TShaderMapRef<FCompositePS<0>> PixelShader(ShaderMap);
 
-							GraphicsPSOInit.BoundShaderState.VertexDeclarationRHI = RendererModule.GetFilterVertexDeclaration().VertexDeclarationRHI;
+							GraphicsPSOInit.BoundShaderState.VertexDeclarationRHI = GFilterVertexDeclaration.VertexDeclarationRHI;
 							GraphicsPSOInit.BoundShaderState.VertexShaderRHI = GETSAFERHISHADER_VERTEX(*VertexShader);
 							GraphicsPSOInit.BoundShaderState.PixelShaderRHI = GETSAFERHISHADER_PIXEL(*PixelShader);
 							GraphicsPSOInit.PrimitiveType = PT_TriangleList;
