Fixes missing #pragma once

diff --git a/Engine/Source/Runtime/Renderer/Private/VisualizeTexturePresent.h b/Engine/Source/Runtime/Renderer/Private/VisualizeTexturePresent.h
--- a/Engine/Source/Runtime/Renderer/Private/VisualizeTexturePresent.h
+++ b/Engine/Source/Runtime/Renderer/Private/VisualizeTexturePresent.h
@@ -4,6 +4,8 @@
 	VisualizeTexturePresent.h: Display texture visualization on screen.
 =============================================================================*/
 
+#pragma once
+
 #include "CoreMinimal.h"
 #include "RHI.h"
 
