Cherry-pick 4542242: Changes FRDGBuilder::CreateParameters() to FRDGBuilder::AllocParameters()

diff --git a/Engine/Source/Runtime/RenderCore/Private/VisualizeTexture.cpp b/Engine/Source/Runtime/RenderCore/Private/VisualizeTexture.cpp
--- a/Engine/Source/Runtime/RenderCore/Private/VisualizeTexture.cpp
+++ b/Engine/Source/Runtime/RenderCore/Private/VisualizeTexture.cpp
@@ -225,8 +225,7 @@ void FVisualizeTexture::CreateContentCapturePass(FRDGBuilder& GraphBuilder, cons
 	bool bSaturateInsteadOfFrac = (Flags & 1) != 0;
 	int32 InputValueMapping = bShadowDepth ? 2 : (bDepthTexture ? 1 : 0);
 
-	FVisualizeTexturePS::FParameters *PassParameters;
-	GraphBuilder.CreateParameters(&PassParameters);
+	FVisualizeTexturePS::FParameters* PassParameters = GraphBuilder.AllocParameters<FVisualizeTexturePS::FParameters>();
 	{
 		PassParameters->TextureExtent = FVector(SrcDesc.Extent.X, SrcDesc.Extent.Y, SrcDesc.Depth);
 
diff --git a/Engine/Source/Runtime/RenderCore/Public/RenderGraphBuilder.h b/Engine/Source/Runtime/RenderCore/Public/RenderGraphBuilder.h
--- a/Engine/Source/Runtime/RenderCore/Public/RenderGraphBuilder.h
+++ b/Engine/Source/Runtime/RenderCore/Public/RenderGraphBuilder.h
@@ -320,16 +320,11 @@ public:
 
 	/** Allocates parameter struct specifically to survive through the life time of the render graph. */
 	template< typename ParameterStructType >
-	inline void CreateParameters(ParameterStructType** OutParameterPtr) const
+	inline ParameterStructType* AllocParameters() const
 	{
-		// Check because destructor called by the pass's destructor.
-		#if RENDER_GRAPH_DEBUGGING
-		{
-			checkf(!bHasExecuted, TEXT("Render graph allocated shader parameters %s needs to be allocated before the builder execution."), ParameterStructType::FTypeInfo::GetStructMetadata()->GetStructTypeName());
-		}
-		#endif
-		*OutParameterPtr = new(FMemStack::Get()) ParameterStructType;
-		FMemory::Memzero( *OutParameterPtr, sizeof( ParameterStructType ) );
+		ParameterStructType* OutParameterPtr = new(FMemStack::Get()) ParameterStructType;
+		FMemory::Memzero(OutParameterPtr, sizeof(ParameterStructType));
+		return OutParameterPtr;
 	}
 
 	/** 
diff --git a/Engine/Source/Runtime/Renderer/Private/SceneOcclusion.cpp b/Engine/Source/Runtime/Renderer/Private/SceneOcclusion.cpp
--- a/Engine/Source/Runtime/Renderer/Private/SceneOcclusion.cpp
+++ b/Engine/Source/Runtime/Renderer/Private/SceneOcclusion.cpp
@@ -1116,8 +1116,7 @@ void BuildHZB(FRDGBuilder& GraphBuilder, FViewInfo& View)
 	FRDGTextureRef HZBTexture = GraphBuilder.CreateTexture(HZBDesc, TEXT("HZB"));
 
 	{
-		FHZBBuildPassParameters* PassParameters;
-		GraphBuilder.CreateParameters(&PassParameters);
+		FHZBBuildPassParameters* PassParameters = GraphBuilder.AllocParameters<FHZBBuildPassParameters>();
 		PassParameters->RenderTargets[0] = FRenderTargetBinding( HZBTexture, ERenderTargetLoadAction::ENoAction, ERenderTargetStoreAction::EStore);
 
 		//@DW - this pass only reads external textures, we don't have any graph inputs
@@ -1197,8 +1196,7 @@ void BuildHZB(FRDGBuilder& GraphBuilder, FViewInfo& View)
 		FRDGTextureSRVDesc Desc(HZBTexture, MipIndex - 1);
 		FRDGTextureSRVRef ParentMipSRV = GraphBuilder.CreateSRV(Desc);
 
-		FHZBBuildPassParameters* PassParameters;
-		GraphBuilder.CreateParameters(&PassParameters);
+		FHZBBuildPassParameters* PassParameters = GraphBuilder.AllocParameters<FHZBBuildPassParameters>();
 		PassParameters->RenderTargets[0] = FRenderTargetBinding( HZBTexture, ERenderTargetLoadAction::ENoAction, ERenderTargetStoreAction::EStore, MipIndex );
 		PassParameters->Texture = ParentMipSRV;
 		PassParameters->TextureSampler = TStaticSamplerState<SF_Point, AM_Clamp, AM_Clamp, AM_Clamp>::GetRHI();
diff --git a/Engine/Source/Runtime/Renderer/Private/VolumetricFog.cpp b/Engine/Source/Runtime/Renderer/Private/VolumetricFog.cpp
--- a/Engine/Source/Runtime/Renderer/Private/VolumetricFog.cpp
+++ b/Engine/Source/Runtime/Renderer/Private/VolumetricFog.cpp
@@ -498,8 +498,7 @@ void FDeferredShadingSceneRenderer::RenderLocalLightsForVolumetricFog(
 	{
 		OutLocalShadowedLightScattering = GraphBuilder.CreateTexture(VolumeDesc, TEXT("LocalShadowedLightScattering"));
 
-		FRenderTargetParameters* PassParameters;
-		GraphBuilder.CreateParameters(&PassParameters);
+		FRenderTargetParameters* PassParameters = GraphBuilder.AllocParameters<FRenderTargetParameters>();
 		PassParameters->RenderTargets[0] = FRenderTargetBinding( OutLocalShadowedLightScattering, ERenderTargetLoadAction::EClear, ERenderTargetStoreAction::ENoAction );
 
 		GraphBuilder.AddPass(
@@ -1030,8 +1029,7 @@ void FDeferredShadingSceneRenderer::ComputeVolumetricFog(FRHICommandListImmediat
 			IntegrationData.VBufferB_UAV = GraphBuilder.CreateUAV(FRDGTextureUAVDesc(IntegrationData.VBufferB));
 
 			{
-				FVolumetricFogMaterialSetupCS::FParameters* PassParameters;
-				GraphBuilder.CreateParameters(&PassParameters);
+				FVolumetricFogMaterialSetupCS::FParameters* PassParameters = GraphBuilder.AllocParameters<FVolumetricFogMaterialSetupCS::FParameters>();
 				PassParameters->GlobalAlbedo			= FogInfo.VolumetricFogAlbedo;
 				PassParameters->GlobalEmissive			= FogInfo.VolumetricFogEmissive;
 				PassParameters->GlobalExtinctionScale	= FogInfo.VolumetricFogExtinctionScale;
@@ -1081,8 +1079,7 @@ void FDeferredShadingSceneRenderer::ComputeVolumetricFog(FRHICommandListImmediat
 			IntegrationData.LightScatteringUAV = GraphBuilder.CreateUAV(FRDGTextureUAVDesc(IntegrationData.LightScattering));
 
 			{
-				TVolumetricFogLightScatteringCS::FParameters* PassParameters;
-				GraphBuilder.CreateParameters(&PassParameters);
+				TVolumetricFogLightScatteringCS::FParameters* PassParameters = GraphBuilder.AllocParameters<TVolumetricFogLightScatteringCS::FParameters>();
 
 				PassParameters->VBufferA = IntegrationData.VBufferA;
 				PassParameters->VBufferB = IntegrationData.VBufferB;
@@ -1142,8 +1139,7 @@ void FDeferredShadingSceneRenderer::ComputeVolumetricFog(FRHICommandListImmediat
 			const FRDGTextureUAV* IntegratedLightScatteringUAV = GraphBuilder.CreateUAV(FRDGTextureUAVDesc(IntegratedLightScattering));
 
 			{
-				FVolumetricFogFinalIntegrationCS::FParameters* PassParameters;
-				GraphBuilder.CreateParameters(&PassParameters);
+				FVolumetricFogFinalIntegrationCS::FParameters* PassParameters = GraphBuilder.AllocParameters<FVolumetricFogFinalIntegrationCS::FParameters>();
 				PassParameters->LightScattering = IntegrationData.LightScattering;
 				PassParameters->RWIntegratedLightScattering = IntegratedLightScatteringUAV;
 
diff --git a/Engine/Source/Runtime/Renderer/Private/VolumetricFogLightFunction.cpp b/Engine/Source/Runtime/Renderer/Private/VolumetricFogLightFunction.cpp
--- a/Engine/Source/Runtime/Renderer/Private/VolumetricFogLightFunction.cpp
+++ b/Engine/Source/Runtime/Renderer/Private/VolumetricFogLightFunction.cpp
@@ -198,8 +198,7 @@ void FDeferredShadingSceneRenderer::RenderLightFunctionForVolumetricFog(
 			const FMatrix WorldToShadowValue = FTranslationMatrix(ProjectedShadowInfo.PreShadowTranslation) * ProjectedShadowInfo.SubjectAndReceiverMatrix;
 			OutLightFunctionWorldToShadow = WorldToShadowValue;
 
-			FRenderTargetParameters* PassParameters;
-			GraphBuilder.CreateParameters(&PassParameters);
+			FRenderTargetParameters* PassParameters = GraphBuilder.AllocParameters<FRenderTargetParameters>();
 			PassParameters->RenderTargets[0] = FRenderTargetBinding( OutLightFunctionTexture, ERenderTargetLoadAction::ENoAction, ERenderTargetStoreAction::EStore );
 
 			GraphBuilder.AddPass(
diff --git a/Engine/Source/Runtime/Renderer/Private/VolumetricFogVoxelization.cpp b/Engine/Source/Runtime/Renderer/Private/VolumetricFogVoxelization.cpp
--- a/Engine/Source/Runtime/Renderer/Private/VolumetricFogVoxelization.cpp
+++ b/Engine/Source/Runtime/Renderer/Private/VolumetricFogVoxelization.cpp
@@ -691,8 +691,7 @@ void FDeferredShadingSceneRenderer::VoxelizeFogVolumePrimitives(
 	if (View.VolumetricPrimSet.NumPrims() > 0
 		&& DoesPlatformSupportVolumetricFogVoxelization(View.GetShaderPlatform()))
 	{
-		FRenderTargetParameters* PassParameters;
-		GraphBuilder.CreateParameters(&PassParameters);
+		FRenderTargetParameters* PassParameters = GraphBuilder.AllocParameters<FRenderTargetParameters>();
 		PassParameters->RenderTargets[0] = FRenderTargetBinding( IntegrationData.VBufferA, ERenderTargetLoadAction::ENoAction, ERenderTargetStoreAction::EStore );
 		PassParameters->RenderTargets[1] = FRenderTargetBinding( IntegrationData.VBufferB, ERenderTargetLoadAction::ENoAction, ERenderTargetStoreAction::EStore );
 
