Adds missing include to fix monolithic builds.

diff --git a/Engine/Source/Runtime/Renderer/Private/GammaCorrection.cpp b/Engine/Source/Runtime/Renderer/Private/GammaCorrection.cpp
--- a/Engine/Source/Runtime/Renderer/Private/GammaCorrection.cpp
+++ b/Engine/Source/Runtime/Renderer/Private/GammaCorrection.cpp
@@ -16,6 +16,7 @@
 #include "PostProcess/SceneFilterRendering.h"
 #include "PipelineStateCache.h"
 #include "ClearQuad.h"
+#include "DummyRenderResources.h"
 
 /** Encapsulates the gamma correction pixel shader. */
 class FGammaCorrectionPS : public FGlobalShader
diff --git a/Engine/Source/Runtime/Renderer/Private/MobileSceneCaptureRendering.cpp b/Engine/Source/Runtime/Renderer/Private/MobileSceneCaptureRendering.cpp
--- a/Engine/Source/Runtime/Renderer/Private/MobileSceneCaptureRendering.cpp
+++ b/Engine/Source/Runtime/Renderer/Private/MobileSceneCaptureRendering.cpp
@@ -26,6 +26,7 @@ MobileSceneCaptureRendering.cpp - Mobile specific scene capture code.
 #include "ScreenRendering.h"
 #include "ClearQuad.h"
 #include "PipelineStateCache.h"
+#include "DummyRenderResources.h"
 
 /**
 * Shader set for the copy of scene color to capture target, decoding mosaic or RGBE encoded HDR image as part of a
