DR - Default clear the LightShaftOcclusion to White because it is also used by the FogRendering and when only one View is rendered the other View decomes black. The default for the FogRendering is White if no LightShaftOcclusion is present, therefore we should clear to White.

diff --git a/Engine/Source/Runtime/Renderer/Private/LightShaftRendering.cpp b/Engine/Source/Runtime/Renderer/Private/LightShaftRendering.cpp
--- a/Engine/Source/Runtime/Renderer/Private/LightShaftRendering.cpp
+++ b/Engine/Source/Runtime/Renderer/Private/LightShaftRendering.cpp
@@ -470,7 +470,7 @@ void AllocateOrReuseLightShaftRenderTarget(FRHICommandListImmediate& RHICmdList,
 		EPixelFormat LightShaftFilterBufferFormat = PF_FloatRGB;
 		const FIntPoint BufferSize = FSceneRenderTargets::Get(RHICmdList).GetBufferSizeXY();
 		FIntPoint LightShaftSize(FMath::Max<uint32>(BufferSize.X / GetLightShaftDownsampleFactor(), 1), FMath::Max<uint32>(BufferSize.Y / GetLightShaftDownsampleFactor(), 1));
-		FPooledRenderTargetDesc Desc(FPooledRenderTargetDesc::Create2DDesc(LightShaftSize, LightShaftFilterBufferFormat, FClearValueBinding::Black, TexCreate_None, TexCreate_RenderTargetable, false));
+		FPooledRenderTargetDesc Desc(FPooledRenderTargetDesc::Create2DDesc(LightShaftSize, LightShaftFilterBufferFormat, FClearValueBinding::White, TexCreate_None, TexCreate_RenderTargetable, false));
 		Desc.AutoWritable = false;
 		GRenderTargetPool.FindFreeElement(RHICmdList, Desc, Target, Name);
 
