Fixes a crash in volumetric fog

diff --git a/Engine/Source/Runtime/Renderer/Private/VolumetricFog.cpp b/Engine/Source/Runtime/Renderer/Private/VolumetricFog.cpp
--- a/Engine/Source/Runtime/Renderer/Private/VolumetricFog.cpp
+++ b/Engine/Source/Runtime/Renderer/Private/VolumetricFog.cpp
@@ -1123,7 +1123,7 @@ void FDeferredShadingSceneRenderer::ComputeVolumetricFog(FRHICommandListImmediat
 
 						RHICmdList.SetComputeShader(ComputeShader->GetComputeShader());
 
-						FTextureRHIParamRef LightScatteringHistoryTexture = View.ViewState->LightScatteringHistory.IsValid()
+						FTextureRHIParamRef LightScatteringHistoryTexture = bUseTemporalReprojection && View.ViewState->LightScatteringHistory.IsValid()
 							? View.ViewState->LightScatteringHistory->GetRenderTargetItem().ShaderResourceTexture
 							: GBlackVolumeTexture->TextureRHI;
 
