Duplicate 4267103

diff --git a/Engine/Source/Runtime/Renderer/Private/ShadowSetup.cpp b/Engine/Source/Runtime/Renderer/Private/ShadowSetup.cpp
--- a/Engine/Source/Runtime/Renderer/Private/ShadowSetup.cpp
+++ b/Engine/Source/Runtime/Renderer/Private/ShadowSetup.cpp
@@ -3895,7 +3895,7 @@ void FSceneRenderer::InitDynamicShadows(FRHICommandListImmediate& RHICmdList)
 
 					const bool bCreateShadowForMovableLight = 
 						bShouldCreateShadowForMovableLight
-						&& (bPointLightShadow || bProjectEnablePointLightShadows);
+						&& (!bPointLightShadow || bProjectEnablePointLightShadows);
 
 					// Also create a whole scene shadow for lights with precomputed shadows that are unbuilt
 					const bool bShouldCreateShadowToPreviewStaticLight =
@@ -3905,7 +3905,7 @@ void FSceneRenderer::InitDynamicShadows(FRHICommandListImmediate& RHICmdList)
 
 					const bool bCreateShadowToPreviewStaticLight = 
 						bShouldCreateShadowToPreviewStaticLight						
-						&& (bPointLightShadow || bProjectEnablePointLightShadows);
+						&& (!bPointLightShadow || bProjectEnablePointLightShadows);
 
 					// Create a whole scene shadow for lights that want static shadowing but didn't get assigned to a valid shadowmap channel due to overlap
 					const bool bShouldCreateShadowForOverflowStaticShadowing =
@@ -3917,7 +3917,7 @@ void FSceneRenderer::InitDynamicShadows(FRHICommandListImmediate& RHICmdList)
 
 					const bool bCreateShadowForOverflowStaticShadowing =
 						bShouldCreateShadowForOverflowStaticShadowing
-						&& (bPointLightShadow || bProjectEnablePointLightShadows);
+						&& (!bPointLightShadow || bProjectEnablePointLightShadows);
 
 					const bool bPointLightWholeSceneShadow = (bShouldCreateShadowForMovableLight || bShouldCreateShadowForOverflowStaticShadowing || bShouldCreateShadowToPreviewStaticLight) && bPointLightShadow;
 					if (bPointLightWholeSceneShadow)
