Fix clang platform compilation

diff --git a/Engine/Source/Runtime/Renderer/Private/BasePassRendering.cpp b/Engine/Source/Runtime/Renderer/Private/BasePassRendering.cpp
--- a/Engine/Source/Runtime/Renderer/Private/BasePassRendering.cpp
+++ b/Engine/Source/Runtime/Renderer/Private/BasePassRendering.cpp
@@ -1279,7 +1279,7 @@ void FBasePassMeshProcessor::AddMeshBatch(const FMeshBatch& RESTRICT MeshBatch,
 
 				FSelfShadowLightCacheElementData ElementData;
 				ElementData.LCI = MeshBatch.LCI;
-				ElementData.SelfShadowTranslucencyUniformBuffer = UniformBufferPtr ? *UniformBufferPtr : GEmptyTranslucentSelfShadowUniformBuffer.GetUniformBufferRHI();
+				ElementData.SelfShadowTranslucencyUniformBuffer = UniformBufferPtr ? (*UniformBufferPtr).GetReference() : GEmptyTranslucentSelfShadowUniformBuffer.GetUniformBufferRHI();
 
 				if (bIsLitMaterial
 					&& bAllowStaticLighting
