MeshProcessors - renamed MeshId to StaticMeshId in order to explain better that this is an index into scene's static mesh array.

diff --git a/Engine/Source/Runtime/Landscape/Private/LandscapeGrass.cpp b/Engine/Source/Runtime/Landscape/Private/LandscapeGrass.cpp
--- a/Engine/Source/Runtime/Landscape/Private/LandscapeGrass.cpp
+++ b/Engine/Source/Runtime/Landscape/Private/LandscapeGrass.cpp
@@ -349,7 +349,7 @@ public:
 		float PassOffsetX,
 		const FPrimitiveSceneProxy* RESTRICT PrimitiveSceneProxy);
 
-	virtual void AddMeshBatch(const FMeshBatch& RESTRICT MeshBatch, uint64 BatchElementMask, const FPrimitiveSceneProxy* RESTRICT PrimitiveSceneProxy, int32 MeshId = -1) override final
+	virtual void AddMeshBatch(const FMeshBatch& RESTRICT MeshBatch, uint64 BatchElementMask, const FPrimitiveSceneProxy* RESTRICT PrimitiveSceneProxy, int32 StaticMeshId = -1) override final
 	{
 		checkf(false, TEXT("Default AddMeshBatch can't be used as rendering requires extra parameters per pass."));
 	}
diff --git a/Engine/Source/Runtime/Renderer/Private/BasePassRendering.cpp b/Engine/Source/Runtime/Renderer/Private/BasePassRendering.cpp
--- a/Engine/Source/Runtime/Renderer/Private/BasePassRendering.cpp
+++ b/Engine/Source/Runtime/Renderer/Private/BasePassRendering.cpp
@@ -131,16 +131,16 @@ IMPLEMENT_BASEPASS_LIGHTMAPPED_SHADER_TYPE( TUniformLightMapPolicy<LMP_DISTANCE_
 
 DECLARE_GPU_STAT(Basepass);
 
-void SetBasePassDitheredLODTransitionState(const FSceneView* SceneView, const FMeshBatch& RESTRICT Mesh, int32 MeshId, FDrawingPolicyRenderState& DrawRenderState)
+void SetBasePassDitheredLODTransitionState(const FSceneView* SceneView, const FMeshBatch& RESTRICT Mesh, int32 StaticMeshId, FDrawingPolicyRenderState& DrawRenderState)
 {
-	if (SceneView && MeshId >= 0 && Mesh.bDitheredLODTransition)
+	if (SceneView && StaticMeshId >= 0 && Mesh.bDitheredLODTransition)
 	{
 		checkSlow(SceneView->bIsViewInfo);
 		const FViewInfo* ViewInfo = (FViewInfo*)SceneView;
 
 		if (ViewInfo->bAllowStencilDither)
 		{
-			if (ViewInfo->StaticMeshFadeOutDitheredLODMap[MeshId])
+			if (ViewInfo->StaticMeshFadeOutDitheredLODMap[StaticMeshId])
 			{
 				DrawRenderState.SetDepthStencilState(
 					TStaticDepthStencilState<
@@ -150,7 +150,7 @@ void SetBasePassDitheredLODTransitionState(const FSceneView* SceneView, const FM
 					0xFF, GET_STENCIL_BIT_MASK(RECEIVE_DECAL, 1) | STENCIL_LIGHTING_CHANNELS_MASK(0x7)
 					>::GetRHI());
 			}
-			else if (ViewInfo->StaticMeshFadeInDitheredLODMap[MeshId])
+			else if (ViewInfo->StaticMeshFadeInDitheredLODMap[StaticMeshId])
 			{
 				DrawRenderState.SetDepthStencilState(
 					TStaticDepthStencilState<
@@ -938,7 +938,7 @@ template<typename LightMapPolicyType>
 void FBasePassMeshProcessor::Process(
 	const FMeshBatch& RESTRICT MeshBatch,
 	uint64 BatchElementMask,
-	int32 MeshId,
+	int32 StaticMeshId,
 	const FPrimitiveSceneProxy* RESTRICT PrimitiveSceneProxy,
 	const FMaterialRenderProxy& RESTRICT MaterialRenderProxy,
 	const FMaterial& RESTRICT MaterialResource,
@@ -984,7 +984,7 @@ void FBasePassMeshProcessor::Process(
 		SetTranslucentRenderState(DrawRenderState, MaterialResource);
 	}
 
-	SetBasePassDitheredLODTransitionState(ViewIfDynamicMeshCommand, MeshBatch, MeshId, DrawRenderState);
+	SetBasePassDitheredLODTransitionState(ViewIfDynamicMeshCommand, MeshBatch, StaticMeshId, DrawRenderState);
 
 	FMeshDrawCommandSortKey SortKey = FMeshDrawCommandSortKey::Default;
 
@@ -998,7 +998,7 @@ void FBasePassMeshProcessor::Process(
 	}
 
 	TBasePassShaderElementData<LightMapPolicyType> ShaderElementData(LightMapElementData);
-	ShaderElementData.InitializeMeshMaterialData(ViewIfDynamicMeshCommand, PrimitiveSceneProxy, MeshBatch, MeshId, true);
+	ShaderElementData.InitializeMeshMaterialData(ViewIfDynamicMeshCommand, PrimitiveSceneProxy, MeshBatch, StaticMeshId, true);
 
 	BuildMeshDrawCommands(
 		MeshBatch,
@@ -1019,7 +1019,7 @@ void FBasePassMeshProcessor::Process(
 void FBasePassMeshProcessor::AddMeshBatchForSimpleForwardShading(
 	const FMeshBatch& RESTRICT MeshBatch,
 	uint64 BatchElementMask,
-	int32 MeshId,
+	int32 StaticMeshId,
 	const FPrimitiveSceneProxy* RESTRICT PrimitiveSceneProxy,
 	const FMaterialRenderProxy& MaterialRenderProxy,
 	const FMaterial& Material,
@@ -1045,7 +1045,7 @@ void FBasePassMeshProcessor::AddMeshBatchForSimpleForwardShading(
 			Process< FUniformLightMapPolicy >(
 				MeshBatch,
 				BatchElementMask,
-				MeshId,
+				StaticMeshId,
 				PrimitiveSceneProxy,
 				MaterialRenderProxy,
 				Material,
@@ -1061,7 +1061,7 @@ void FBasePassMeshProcessor::AddMeshBatchForSimpleForwardShading(
 			Process< FUniformLightMapPolicy >(
 				MeshBatch,
 				BatchElementMask,
-				MeshId,
+				StaticMeshId,
 				PrimitiveSceneProxy,
 				MaterialRenderProxy,
 				Material,
@@ -1081,7 +1081,7 @@ void FBasePassMeshProcessor::AddMeshBatchForSimpleForwardShading(
 		Process< FUniformLightMapPolicy >(
 			MeshBatch,
 			BatchElementMask,
-			MeshId,
+			StaticMeshId,
 			PrimitiveSceneProxy,
 			MaterialRenderProxy,
 			Material,
@@ -1113,7 +1113,7 @@ void FBasePassMeshProcessor::AddMeshBatchForSimpleForwardShading(
 			Process< FUniformLightMapPolicy >(
 				MeshBatch,
 				BatchElementMask,
-				MeshId,
+				StaticMeshId,
 				PrimitiveSceneProxy,
 				MaterialRenderProxy,
 				Material,
@@ -1129,7 +1129,7 @@ void FBasePassMeshProcessor::AddMeshBatchForSimpleForwardShading(
 			Process< FUniformLightMapPolicy >(
 				MeshBatch,
 				BatchElementMask,
-				MeshId,
+				StaticMeshId,
 				PrimitiveSceneProxy,
 				MaterialRenderProxy,
 				Material,
@@ -1147,7 +1147,7 @@ void FBasePassMeshProcessor::AddMeshBatchForSimpleForwardShading(
 		Process< FUniformLightMapPolicy >(
 			MeshBatch,
 			BatchElementMask,
-			MeshId,
+			StaticMeshId,
 			PrimitiveSceneProxy,
 			MaterialRenderProxy,
 			Material,
@@ -1163,7 +1163,7 @@ void FBasePassMeshProcessor::AddMeshBatchForSimpleForwardShading(
 		Process< FUniformLightMapPolicy >(
 			MeshBatch,
 			BatchElementMask,
-			MeshId,
+			StaticMeshId,
 			PrimitiveSceneProxy,
 			MaterialRenderProxy,
 			Material,
@@ -1176,7 +1176,7 @@ void FBasePassMeshProcessor::AddMeshBatchForSimpleForwardShading(
 	}
 }
 
-void FBasePassMeshProcessor::AddMeshBatch(const FMeshBatch& RESTRICT MeshBatch, uint64 BatchElementMask, const FPrimitiveSceneProxy* RESTRICT PrimitiveSceneProxy, int32 MeshId)
+void FBasePassMeshProcessor::AddMeshBatch(const FMeshBatch& RESTRICT MeshBatch, uint64 BatchElementMask, const FPrimitiveSceneProxy* RESTRICT PrimitiveSceneProxy, int32 StaticMeshId)
 {
 	if (MeshBatch.bUseForMaterial)
 	{
@@ -1243,7 +1243,7 @@ void FBasePassMeshProcessor::AddMeshBatch(const FMeshBatch& RESTRICT MeshBatch,
 			const bool bUseVolumetricLightmap = Scene && Scene->VolumetricLightmapSceneData.HasData();
 
 			FMeshMaterialShaderElementData MeshMaterialShaderElementData;
-			MeshMaterialShaderElementData.InitializeMeshMaterialData(ViewIfDynamicMeshCommand, PrimitiveSceneProxy, MeshBatch, MeshId, true);
+			MeshMaterialShaderElementData.InitializeMeshMaterialData(ViewIfDynamicMeshCommand, PrimitiveSceneProxy, MeshBatch, StaticMeshId, true);
 
 			if (IsSimpleForwardShadingEnabled(GetFeatureLevelShaderPlatform(FeatureLevel)))
 			{
@@ -1252,7 +1252,7 @@ void FBasePassMeshProcessor::AddMeshBatch(const FMeshBatch& RESTRICT MeshBatch,
 				AddMeshBatchForSimpleForwardShading(
 					MeshBatch,
 					BatchElementMask,
-					MeshId,
+					StaticMeshId,
 					PrimitiveSceneProxy,
 					MaterialRenderProxy,
 					Material,
@@ -1289,7 +1289,7 @@ void FBasePassMeshProcessor::AddMeshBatch(const FMeshBatch& RESTRICT MeshBatch,
 					Process< FSelfShadowedVolumetricLightmapPolicy >(
 						MeshBatch,
 						BatchElementMask,
-						MeshId,
+						StaticMeshId,
 						PrimitiveSceneProxy,
 						MaterialRenderProxy,
 						Material,
@@ -1308,7 +1308,7 @@ void FBasePassMeshProcessor::AddMeshBatch(const FMeshBatch& RESTRICT MeshBatch,
 					Process< FSelfShadowedCachedPointIndirectLightingPolicy >(
 						MeshBatch,
 						BatchElementMask,
-						MeshId,
+						StaticMeshId,
 						PrimitiveSceneProxy,
 						MaterialRenderProxy,
 						Material,
@@ -1324,7 +1324,7 @@ void FBasePassMeshProcessor::AddMeshBatch(const FMeshBatch& RESTRICT MeshBatch,
 					Process< FSelfShadowedTranslucencyPolicy >(
 						MeshBatch,
 						BatchElementMask,
-						MeshId,
+						StaticMeshId,
 						PrimitiveSceneProxy,
 						MaterialRenderProxy,
 						Material,
@@ -1355,7 +1355,7 @@ void FBasePassMeshProcessor::AddMeshBatch(const FMeshBatch& RESTRICT MeshBatch,
 							Process< FUniformLightMapPolicy >(
 								MeshBatch,
 								BatchElementMask,
-								MeshId,
+								StaticMeshId,
 								PrimitiveSceneProxy,
 								MaterialRenderProxy,
 								Material,
@@ -1371,7 +1371,7 @@ void FBasePassMeshProcessor::AddMeshBatch(const FMeshBatch& RESTRICT MeshBatch,
 							Process< FUniformLightMapPolicy >(
 								MeshBatch,
 								BatchElementMask,
-								MeshId,
+								StaticMeshId,
 								PrimitiveSceneProxy,
 								MaterialRenderProxy,
 								Material,
@@ -1388,7 +1388,7 @@ void FBasePassMeshProcessor::AddMeshBatch(const FMeshBatch& RESTRICT MeshBatch,
 						Process< FUniformLightMapPolicy >(
 							MeshBatch,
 							BatchElementMask,
-							MeshId,
+							StaticMeshId,
 							PrimitiveSceneProxy,
 							MaterialRenderProxy,
 							Material,
@@ -1404,7 +1404,7 @@ void FBasePassMeshProcessor::AddMeshBatch(const FMeshBatch& RESTRICT MeshBatch,
 						Process< FUniformLightMapPolicy >(
 							MeshBatch,
 							BatchElementMask,
-							MeshId,
+							StaticMeshId,
 							PrimitiveSceneProxy,
 							MaterialRenderProxy,
 							Material,
@@ -1429,7 +1429,7 @@ void FBasePassMeshProcessor::AddMeshBatch(const FMeshBatch& RESTRICT MeshBatch,
 						Process< FUniformLightMapPolicy >(
 							MeshBatch,
 							BatchElementMask,
-							MeshId,
+							StaticMeshId,
 							PrimitiveSceneProxy,
 							MaterialRenderProxy,
 							Material,
@@ -1468,7 +1468,7 @@ void FBasePassMeshProcessor::AddMeshBatch(const FMeshBatch& RESTRICT MeshBatch,
 								Process< FUniformLightMapPolicy >(
 									MeshBatch,
 									BatchElementMask,
-									MeshId,
+									StaticMeshId,
 									PrimitiveSceneProxy,
 									MaterialRenderProxy,
 									Material,
@@ -1485,7 +1485,7 @@ void FBasePassMeshProcessor::AddMeshBatch(const FMeshBatch& RESTRICT MeshBatch,
 								Process< FUniformLightMapPolicy >(
 									MeshBatch,
 									BatchElementMask,
-									MeshId,
+									StaticMeshId,
 									PrimitiveSceneProxy,
 									MaterialRenderProxy,
 									Material,
@@ -1502,7 +1502,7 @@ void FBasePassMeshProcessor::AddMeshBatch(const FMeshBatch& RESTRICT MeshBatch,
 							Process< FUniformLightMapPolicy >(
 								MeshBatch,
 								BatchElementMask,
-								MeshId,
+								StaticMeshId,
 								PrimitiveSceneProxy,
 								MaterialRenderProxy,
 								Material,
@@ -1519,7 +1519,7 @@ void FBasePassMeshProcessor::AddMeshBatch(const FMeshBatch& RESTRICT MeshBatch,
 						Process< FUniformLightMapPolicy >(
 							MeshBatch,
 							BatchElementMask,
-							MeshId,
+							StaticMeshId,
 							PrimitiveSceneProxy,
 							MaterialRenderProxy,
 							Material,
diff --git a/Engine/Source/Runtime/Renderer/Private/BasePassRendering.h b/Engine/Source/Runtime/Renderer/Private/BasePassRendering.h
--- a/Engine/Source/Runtime/Renderer/Private/BasePassRendering.h
+++ b/Engine/Source/Runtime/Renderer/Private/BasePassRendering.h
@@ -677,7 +677,7 @@ public:
 
 	FBasePassMeshProcessor(const FScene* InScene, ERHIFeatureLevel::Type InFeatureLevel, const FSceneView* InViewIfDynamicMeshCommand, const FDrawingPolicyRenderState& InDrawRenderState, FMeshPassDrawListContext& InDrawListContext, ETranslucencyPass::Type InTranslucencyPassType = ETranslucencyPass::TPT_MAX);
 
-	virtual void AddMeshBatch(const FMeshBatch& RESTRICT MeshBatch, uint64 BatchElementMask, const FPrimitiveSceneProxy* RESTRICT PrimitiveSceneProxy, int32 MeshId = -1) override final;
+	virtual void AddMeshBatch(const FMeshBatch& RESTRICT MeshBatch, uint64 BatchElementMask, const FPrimitiveSceneProxy* RESTRICT PrimitiveSceneProxy, int32 StaticMeshId = -1) override final;
 
 	FDrawingPolicyRenderState PassDrawRenderState;
 
@@ -686,7 +686,7 @@ private:
 	void AddMeshBatchForSimpleForwardShading(
 		const FMeshBatch& RESTRICT MeshBatch,
 		uint64 BatchElementMask,
-		int32 MeshId,
+		int32 StaticMeshId,
 		const FPrimitiveSceneProxy* RESTRICT PrimitiveSceneProxy,
 		const FMaterialRenderProxy& MaterialRenderProxy,
 		const FMaterial& Material,
@@ -702,7 +702,7 @@ private:
 	void Process(
 		const FMeshBatch& RESTRICT MeshBatch,
 		uint64 BatchElementMask,
-		int32 MeshId,
+		int32 StaticMeshId,
 		const FPrimitiveSceneProxy* RESTRICT PrimitiveSceneProxy,
 		const FMaterialRenderProxy& RESTRICT MaterialRenderProxy,
 		const FMaterial& RESTRICT MaterialResource,
diff --git a/Engine/Source/Runtime/Renderer/Private/CompositionLighting/PostProcessMeshDecals.cpp b/Engine/Source/Runtime/Renderer/Private/CompositionLighting/PostProcessMeshDecals.cpp
--- a/Engine/Source/Runtime/Renderer/Private/CompositionLighting/PostProcessMeshDecals.cpp
+++ b/Engine/Source/Runtime/Renderer/Private/CompositionLighting/PostProcessMeshDecals.cpp
@@ -214,13 +214,13 @@ public:
 		FDecalRenderingCommon::ERenderTargetMode InRenderTargetMode,
 		FMeshPassDrawListContext& InDrawListContext);
 
-	virtual void AddMeshBatch(const FMeshBatch& RESTRICT MeshBatch, uint64 BatchElementMask, const FPrimitiveSceneProxy* RESTRICT PrimitiveSceneProxy, int32 MeshId = -1) override final;
+	virtual void AddMeshBatch(const FMeshBatch& RESTRICT MeshBatch, uint64 BatchElementMask, const FPrimitiveSceneProxy* RESTRICT PrimitiveSceneProxy, int32 StaticMeshId = -1) override final;
 
 private:
 	void Process(
 		const FMeshBatch& MeshBatch,
 		uint64 BatchElementMask,
-		int32 MeshId,
+		int32 StaticMeshId,
 		const FPrimitiveSceneProxy* RESTRICT PrimitiveSceneProxy,
 		const FMaterialRenderProxy& RESTRICT MaterialRenderProxy,
 		const FMaterial& RESTRICT MaterialResource,
@@ -247,7 +247,7 @@ FMeshDecalMeshProcessor::FMeshDecalMeshProcessor(const FScene* Scene,
 	PassDrawRenderState.SetPassUniformBuffer(Scene->UniformBuffers.MeshDecalPassUniformBuffer);
 }
 
-void FMeshDecalMeshProcessor::AddMeshBatch(const FMeshBatch& RESTRICT MeshBatch, uint64 BatchElementMask, const FPrimitiveSceneProxy* RESTRICT PrimitiveSceneProxy, int32 MeshId)
+void FMeshDecalMeshProcessor::AddMeshBatch(const FMeshBatch& RESTRICT MeshBatch, uint64 BatchElementMask, const FPrimitiveSceneProxy* RESTRICT PrimitiveSceneProxy, int32 StaticMeshId)
 {
 	if (MeshBatch.bUseForMaterial)
 	{
@@ -307,7 +307,7 @@ void FMeshDecalMeshProcessor::AddMeshBatch(const FMeshBatch& RESTRICT MeshBatch,
 						PassDrawRenderState.SetBlendState(GetDecalBlendState(FeatureLevel, PassDecalStage, DecalBlendMode, bHasNormal));
 					}
 
-					Process(MeshBatch, BatchElementMask, MeshId, PrimitiveSceneProxy, *MaterialRenderProxy, *Material, MeshFillMode, MeshCullMode);
+					Process(MeshBatch, BatchElementMask, StaticMeshId, PrimitiveSceneProxy, *MaterialRenderProxy, *Material, MeshFillMode, MeshCullMode);
 				}
 			}
 		}
@@ -317,7 +317,7 @@ void FMeshDecalMeshProcessor::AddMeshBatch(const FMeshBatch& RESTRICT MeshBatch,
 void FMeshDecalMeshProcessor::Process(
 	const FMeshBatch& MeshBatch,
 	uint64 BatchElementMask,
-	int32 MeshId,
+	int32 StaticMeshId,
 	const FPrimitiveSceneProxy* RESTRICT PrimitiveSceneProxy,
 	const FMaterialRenderProxy& RESTRICT MaterialRenderProxy,
 	const FMaterial& RESTRICT MaterialResource,
@@ -354,7 +354,7 @@ void FMeshDecalMeshProcessor::Process(
 
 
 	FMeshMaterialShaderElementData ShaderElementData;
-	ShaderElementData.InitializeMeshMaterialData(ViewIfDynamicMeshCommand, PrimitiveSceneProxy, MeshBatch, MeshId, true);
+	ShaderElementData.InitializeMeshMaterialData(ViewIfDynamicMeshCommand, PrimitiveSceneProxy, MeshBatch, StaticMeshId, true);
 
 	BuildMeshDrawCommands(
 		MeshBatch,
diff --git a/Engine/Source/Runtime/Renderer/Private/ConvertToUniformMesh.cpp b/Engine/Source/Runtime/Renderer/Private/ConvertToUniformMesh.cpp
--- a/Engine/Source/Runtime/Renderer/Private/ConvertToUniformMesh.cpp
+++ b/Engine/Source/Runtime/Renderer/Private/ConvertToUniformMesh.cpp
@@ -154,7 +154,7 @@ class FConvertToUniformMeshProcessor : public FMeshPassProcessor
 public:
 	FConvertToUniformMeshProcessor(const FScene* Scene, const FViewInfo* InViewIfDynamicMeshCommand, FMeshPassDrawListContext& InDrawListContext);
 
-	void AddMeshBatch(const FMeshBatch& RESTRICT MeshBatch, uint64 BatchElementMask, const FPrimitiveSceneProxy* RESTRICT PrimitiveSceneProxy, int32 MeshId = -1) override final;
+	void AddMeshBatch(const FMeshBatch& RESTRICT MeshBatch, uint64 BatchElementMask, const FPrimitiveSceneProxy* RESTRICT PrimitiveSceneProxy, int32 StaticMeshId = -1) override final;
 
 private:
 	void Process(
@@ -179,7 +179,7 @@ FConvertToUniformMeshProcessor::FConvertToUniformMeshProcessor(const FScene* Sce
 	PassDrawRenderState.SetPassUniformBuffer(Scene->UniformBuffers.ConvertToUniformMeshPassUniformBuffer);
 }
 
-void FConvertToUniformMeshProcessor::AddMeshBatch(const FMeshBatch& RESTRICT MeshBatch, uint64 BatchElementMask, const FPrimitiveSceneProxy* RESTRICT PrimitiveSceneProxy, int32 MeshId)
+void FConvertToUniformMeshProcessor::AddMeshBatch(const FMeshBatch& RESTRICT MeshBatch, uint64 BatchElementMask, const FPrimitiveSceneProxy* RESTRICT PrimitiveSceneProxy, int32 StaticMeshId)
 {
 	const FMaterialRenderProxy* FallbackMaterialRenderProxyPtr = nullptr;
 	const FMaterial& Material = MeshBatch.MaterialRenderProxy->GetMaterialWithFallback(FeatureLevel, FallbackMaterialRenderProxyPtr);
diff --git a/Engine/Source/Runtime/Renderer/Private/CustomDepthRendering.cpp b/Engine/Source/Runtime/Renderer/Private/CustomDepthRendering.cpp
--- a/Engine/Source/Runtime/Renderer/Private/CustomDepthRendering.cpp
+++ b/Engine/Source/Runtime/Renderer/Private/CustomDepthRendering.cpp
@@ -121,14 +121,14 @@ class FCustomDepthPassMeshProcessor : public FMeshPassProcessor
 public:
 	FCustomDepthPassMeshProcessor(const FScene* Scene, const FSceneView* InViewIfDynamicMeshCommand, FMeshPassDrawListContext& InDrawListContext);
 
-	virtual void AddMeshBatch(const FMeshBatch& RESTRICT MeshBatch, uint64 BatchElementMask, const FPrimitiveSceneProxy* RESTRICT PrimitiveSceneProxy, int32 MeshId = -1) override final;
+	virtual void AddMeshBatch(const FMeshBatch& RESTRICT MeshBatch, uint64 BatchElementMask, const FPrimitiveSceneProxy* RESTRICT PrimitiveSceneProxy, int32 StaticMeshId = -1) override final;
 
 private:
 	template<bool bPositionOnly>
 	void Process(
 		const FMeshBatch& MeshBatch,
 		uint64 BatchElementMask,
-		int32 MeshId,
+		int32 StaticMeshId,
 		const FPrimitiveSceneProxy* RESTRICT PrimitiveSceneProxy,
 		const FMaterialRenderProxy& RESTRICT MaterialRenderProxy,
 		const FMaterial& RESTRICT MaterialResource,
@@ -158,7 +158,7 @@ FCustomDepthPassMeshProcessor::FCustomDepthPassMeshProcessor(const FScene* Scene
 	PassDrawRenderState.SetDepthStencilState(TStaticDepthStencilState<true, CF_DepthNearOrEqual>::GetRHI());
 }
 
-void FCustomDepthPassMeshProcessor::AddMeshBatch(const FMeshBatch& RESTRICT MeshBatch, uint64 BatchElementMask, const FPrimitiveSceneProxy* RESTRICT PrimitiveSceneProxy, int32 MeshId)
+void FCustomDepthPassMeshProcessor::AddMeshBatch(const FMeshBatch& RESTRICT MeshBatch, uint64 BatchElementMask, const FPrimitiveSceneProxy* RESTRICT PrimitiveSceneProxy, int32 StaticMeshId)
 {
 	if (PrimitiveSceneProxy->ShouldRenderCustomDepth())
 	{
@@ -219,7 +219,7 @@ void FCustomDepthPassMeshProcessor::AddMeshBatch(const FMeshBatch& RESTRICT Mesh
 		{
 			const FMaterialRenderProxy& DefaultProxy = *UMaterial::GetDefaultMaterial(MD_Surface)->GetRenderProxy();
 			const FMaterial& DefaultMaterial = *DefaultProxy.GetMaterial(FeatureLevel);
-			Process<true>(MeshBatch, BatchElementMask, MeshId, PrimitiveSceneProxy, DefaultProxy, DefaultMaterial, MeshFillMode, MeshCullMode, bMobileColorValue, bUsesMobileColorValue);
+			Process<true>(MeshBatch, BatchElementMask, StaticMeshId, PrimitiveSceneProxy, DefaultProxy, DefaultMaterial, MeshFillMode, MeshCullMode, bMobileColorValue, bUsesMobileColorValue);
 		}
 		else if (!IsTranslucentBlendMode(BlendMode) || Material.IsTranslucencyWritingCustomDepth())
 		{
@@ -235,7 +235,7 @@ void FCustomDepthPassMeshProcessor::AddMeshBatch(const FMeshBatch& RESTRICT Mesh
 				EffectiveMaterial = EffectiveMaterialRenderProxy->GetMaterial(FeatureLevel);
 			}
 
-			Process<false>(MeshBatch, BatchElementMask, MeshId, PrimitiveSceneProxy, *EffectiveMaterialRenderProxy, *EffectiveMaterial, MeshFillMode, MeshCullMode, bMobileColorValue, bUsesMobileColorValue);
+			Process<false>(MeshBatch, BatchElementMask, StaticMeshId, PrimitiveSceneProxy, *EffectiveMaterialRenderProxy, *EffectiveMaterial, MeshFillMode, MeshCullMode, bMobileColorValue, bUsesMobileColorValue);
 		}
 	}
 }
@@ -244,7 +244,7 @@ template<bool bPositionOnly>
 void FCustomDepthPassMeshProcessor::Process(
 	const FMeshBatch& RESTRICT MeshBatch,
 	uint64 BatchElementMask,
-	int32 MeshId,
+	int32 StaticMeshId,
 	const FPrimitiveSceneProxy* RESTRICT PrimitiveSceneProxy,
 	const FMaterialRenderProxy& RESTRICT MaterialRenderProxy,
 	const FMaterial& RESTRICT MaterialResource,
@@ -276,7 +276,7 @@ void FCustomDepthPassMeshProcessor::Process(
 		);
 
 	FDepthOnlyShaderElementData ShaderElementData(MobileColorValue);
-	ShaderElementData.InitializeMeshMaterialData(ViewIfDynamicMeshCommand, PrimitiveSceneProxy, MeshBatch, MeshId, false);
+	ShaderElementData.InitializeMeshMaterialData(ViewIfDynamicMeshCommand, PrimitiveSceneProxy, MeshBatch, StaticMeshId, false);
 
 	BuildMeshDrawCommands(
 		MeshBatch,
diff --git a/Engine/Source/Runtime/Renderer/Private/DebugViewModeRendering.cpp b/Engine/Source/Runtime/Renderer/Private/DebugViewModeRendering.cpp
--- a/Engine/Source/Runtime/Renderer/Private/DebugViewModeRendering.cpp
+++ b/Engine/Source/Runtime/Renderer/Private/DebugViewModeRendering.cpp
@@ -313,7 +313,7 @@ FDebugViewModeMeshProcessor::FDebugViewModeMeshProcessor(
 	}
 }
 
-void FDebugViewModeMeshProcessor::AddMeshBatch(const FMeshBatch& RESTRICT MeshBatch, uint64 BatchElementMask, const FPrimitiveSceneProxy* RESTRICT PrimitiveSceneProxy, int32 MeshId)
+void FDebugViewModeMeshProcessor::AddMeshBatch(const FMeshBatch& RESTRICT MeshBatch, uint64 BatchElementMask, const FPrimitiveSceneProxy* RESTRICT PrimitiveSceneProxy, int32 StaticMeshId)
 {
 	const FMaterialRenderProxy* MaterialRenderProxy = nullptr;
 	const FMaterial* Material = nullptr;
@@ -379,7 +379,7 @@ void FDebugViewModeMeshProcessor::AddMeshBatch(const FMeshBatch& RESTRICT MeshBa
 		ShaderElementData.NumPSInstructions = BatchMaterial->GetShader<TBasePassPS<TUniformLightMapPolicy<LMP_NO_LIGHTMAP>, false>>(VertexFactoryType)->GetNumInstructions();
 	}
 
-	ShaderElementData.InitializeMeshMaterialData(ViewIfDynamicMeshCommand, PrimitiveSceneProxy, MeshBatch, MeshId, true);
+	ShaderElementData.InitializeMeshMaterialData(ViewIfDynamicMeshCommand, PrimitiveSceneProxy, MeshBatch, StaticMeshId, true);
 
 	BuildMeshDrawCommands(
 		MeshBatch,
diff --git a/Engine/Source/Runtime/Renderer/Private/DebugViewModeRendering.h b/Engine/Source/Runtime/Renderer/Private/DebugViewModeRendering.h
--- a/Engine/Source/Runtime/Renderer/Private/DebugViewModeRendering.h
+++ b/Engine/Source/Runtime/Renderer/Private/DebugViewModeRendering.h
@@ -310,7 +310,7 @@ class FDebugViewModeMeshProcessor : public FMeshPassProcessor
 {
 public:
 	FDebugViewModeMeshProcessor(const FScene* InScene, ERHIFeatureLevel::Type InFeatureLevel, const FSceneView* InViewIfDynamicMeshCommand, FUniformBufferRHIParamRef InPassUniformBuffer, bool bTranslucentBasePass, FMeshPassDrawListContext& InDrawListContext);
-	virtual void AddMeshBatch(const FMeshBatch& RESTRICT MeshBatch, uint64 BatchElementMask, const FPrimitiveSceneProxy* RESTRICT PrimitiveSceneProxy, int32 MeshId = -1) override final;
+	virtual void AddMeshBatch(const FMeshBatch& RESTRICT MeshBatch, uint64 BatchElementMask, const FPrimitiveSceneProxy* RESTRICT PrimitiveSceneProxy, int32 StaticMeshId = -1) override final;
 
 private:
 
diff --git a/Engine/Source/Runtime/Renderer/Private/DepthRendering.cpp b/Engine/Source/Runtime/Renderer/Private/DepthRendering.cpp
--- a/Engine/Source/Runtime/Renderer/Private/DepthRendering.cpp
+++ b/Engine/Source/Runtime/Renderer/Private/DepthRendering.cpp
@@ -238,16 +238,16 @@ FDepthDrawingPolicy::FDepthDrawingPolicy(
 	BaseVertexShader = VertexShader;
 }
 
-void SetDepthPassDitheredLODTransitionState(const FSceneView* SceneView, const FMeshBatch& RESTRICT Mesh, int32 MeshId, FDrawingPolicyRenderState& DrawRenderState)
+void SetDepthPassDitheredLODTransitionState(const FSceneView* SceneView, const FMeshBatch& RESTRICT Mesh, int32 StaticMeshId, FDrawingPolicyRenderState& DrawRenderState)
 {
-	if (SceneView && MeshId >= 0 && Mesh.bDitheredLODTransition)
+	if (SceneView && StaticMeshId >= 0 && Mesh.bDitheredLODTransition)
 	{
 		checkSlow(SceneView->bIsViewInfo);
 		const FViewInfo* ViewInfo = (FViewInfo*)SceneView;
 
 		if (ViewInfo->bAllowStencilDither)
 		{
-			if (ViewInfo->StaticMeshFadeOutDitheredLODMap[MeshId])
+			if (ViewInfo->StaticMeshFadeOutDitheredLODMap[StaticMeshId])
 			{
 				DrawRenderState.SetDepthStencilState(
 					TStaticDepthStencilState<true, CF_DepthNearOrEqual,
@@ -257,7 +257,7 @@ void SetDepthPassDitheredLODTransitionState(const FSceneView* SceneView, const F
 					>::GetRHI());
 				DrawRenderState.SetStencilRef(STENCIL_SANDBOX_MASK);
 			}
-			else if (ViewInfo->StaticMeshFadeInDitheredLODMap[MeshId])
+			else if (ViewInfo->StaticMeshFadeInDitheredLODMap[StaticMeshId])
 			{
 				DrawRenderState.SetDepthStencilState(
 					TStaticDepthStencilState<true, CF_DepthNearOrEqual,
@@ -1458,7 +1458,7 @@ template<bool bPositionOnly>
 void FDepthPassMeshProcessor::Process(
 	const FMeshBatch& RESTRICT MeshBatch,
 	uint64 BatchElementMask,
-	int32 MeshId,
+	int32 StaticMeshId,
 	const FPrimitiveSceneProxy* RESTRICT PrimitiveSceneProxy,
 	const FMaterialRenderProxy& RESTRICT MaterialRenderProxy,
 	const FMaterial& RESTRICT MaterialResource,
@@ -1489,10 +1489,10 @@ void FDepthPassMeshProcessor::Process(
 
 	FDrawingPolicyRenderState DrawRenderState(PassDrawRenderState);
 
-	SetDepthPassDitheredLODTransitionState(ViewIfDynamicMeshCommand, MeshBatch, MeshId, DrawRenderState);
+	SetDepthPassDitheredLODTransitionState(ViewIfDynamicMeshCommand, MeshBatch, StaticMeshId, DrawRenderState);
 
 	FDepthOnlyShaderElementData ShaderElementData(0.0f);
-	ShaderElementData.InitializeMeshMaterialData(ViewIfDynamicMeshCommand, PrimitiveSceneProxy, MeshBatch, MeshId, true);
+	ShaderElementData.InitializeMeshMaterialData(ViewIfDynamicMeshCommand, PrimitiveSceneProxy, MeshBatch, StaticMeshId, true);
 
 	BuildMeshDrawCommands(
 		MeshBatch,
@@ -1510,7 +1510,7 @@ void FDepthPassMeshProcessor::Process(
 		ShaderElementData);
 }
 
-void FDepthPassMeshProcessor::AddMeshBatch(const FMeshBatch& RESTRICT MeshBatch, uint64 BatchElementMask, const FPrimitiveSceneProxy* RESTRICT PrimitiveSceneProxy, int32 MeshId)
+void FDepthPassMeshProcessor::AddMeshBatch(const FMeshBatch& RESTRICT MeshBatch, uint64 BatchElementMask, const FPrimitiveSceneProxy* RESTRICT PrimitiveSceneProxy, int32 StaticMeshId)
 {
 	bool bDraw = true;
 
@@ -1559,7 +1559,7 @@ void FDepthPassMeshProcessor::AddMeshBatch(const FMeshBatch& RESTRICT MeshBatch,
 			{
 				const FMaterialRenderProxy& DefaultProxy = *UMaterial::GetDefaultMaterial(MD_Surface)->GetRenderProxy();
 				const FMaterial& DefaultMaterial = *DefaultProxy.GetMaterial(FeatureLevel);
-				Process<true>(MeshBatch, BatchElementMask, MeshId, PrimitiveSceneProxy, DefaultProxy, DefaultMaterial, MeshFillMode, MeshCullMode);
+				Process<true>(MeshBatch, BatchElementMask, StaticMeshId, PrimitiveSceneProxy, DefaultProxy, DefaultMaterial, MeshFillMode, MeshCullMode);
 			}
 			else
 			{
@@ -1577,7 +1577,7 @@ void FDepthPassMeshProcessor::AddMeshBatch(const FMeshBatch& RESTRICT MeshBatch,
 						EffectiveMaterial = EffectiveMaterialRenderProxy->GetMaterial(FeatureLevel);
 					}
 
-					Process<false>(MeshBatch, BatchElementMask, MeshId, PrimitiveSceneProxy, *EffectiveMaterialRenderProxy, *EffectiveMaterial, MeshFillMode, MeshCullMode);
+					Process<false>(MeshBatch, BatchElementMask, StaticMeshId, PrimitiveSceneProxy, *EffectiveMaterialRenderProxy, *EffectiveMaterial, MeshFillMode, MeshCullMode);
 				}
 			}
 		}
diff --git a/Engine/Source/Runtime/Renderer/Private/DepthRendering.h b/Engine/Source/Runtime/Renderer/Private/DepthRendering.h
--- a/Engine/Source/Runtime/Renderer/Private/DepthRendering.h
+++ b/Engine/Source/Runtime/Renderer/Private/DepthRendering.h
@@ -504,7 +504,7 @@ public:
 		const bool InbEarlyZPassMovable,
 		FMeshPassDrawListContext& InDrawListContext);
 
-	virtual void AddMeshBatch(const FMeshBatch& RESTRICT MeshBatch, uint64 BatchElementMask, const FPrimitiveSceneProxy* RESTRICT PrimitiveSceneProxy, int32 MeshId = -1) override final;
+	virtual void AddMeshBatch(const FMeshBatch& RESTRICT MeshBatch, uint64 BatchElementMask, const FPrimitiveSceneProxy* RESTRICT PrimitiveSceneProxy, int32 StaticMeshId = -1) override final;
 
 private:
 	
@@ -512,7 +512,7 @@ private:
 	void Process(
 		const FMeshBatch& MeshBatch,
 		uint64 BatchElementMask,
-		int32 MeshId,
+		int32 StaticMeshId,
 		const FPrimitiveSceneProxy* RESTRICT PrimitiveSceneProxy,
 		const FMaterialRenderProxy& RESTRICT MaterialRenderProxy,
 		const FMaterial& RESTRICT MaterialResource,
diff --git a/Engine/Source/Runtime/Renderer/Private/DistortionRendering.cpp b/Engine/Source/Runtime/Renderer/Private/DistortionRendering.cpp
--- a/Engine/Source/Runtime/Renderer/Private/DistortionRendering.cpp
+++ b/Engine/Source/Runtime/Renderer/Private/DistortionRendering.cpp
@@ -1183,7 +1183,7 @@ void FSceneRenderer::RenderDistortion(FRHICommandListImmediate& RHICmdList)
 	}
 }
 
-void FDistortionMeshProcessor::AddMeshBatch(const FMeshBatch& RESTRICT MeshBatch, uint64 BatchElementMask, const FPrimitiveSceneProxy* RESTRICT PrimitiveSceneProxy, int32 MeshId)
+void FDistortionMeshProcessor::AddMeshBatch(const FMeshBatch& RESTRICT MeshBatch, uint64 BatchElementMask, const FPrimitiveSceneProxy* RESTRICT PrimitiveSceneProxy, int32 StaticMeshId)
 {
 	if (MeshBatch.bUseForMaterial)
 	{
@@ -1204,7 +1204,7 @@ void FDistortionMeshProcessor::AddMeshBatch(const FMeshBatch& RESTRICT MeshBatch
 			&& ShouldIncludeDomainInMeshPass(Material.GetMaterialDomain())
 			&& Material.IsDistorted())
 		{
-			Process(MeshBatch, BatchElementMask, PrimitiveSceneProxy, MeshId, MaterialRenderProxy, Material, MeshFillMode, MeshCullMode);
+			Process(MeshBatch, BatchElementMask, PrimitiveSceneProxy, StaticMeshId, MaterialRenderProxy, Material, MeshFillMode, MeshCullMode);
 		}
 	}
 }
@@ -1238,7 +1238,7 @@ void FDistortionMeshProcessor::Process(
 	const FMeshBatch& MeshBatch, 
 	uint64 BatchElementMask, 
 	const FPrimitiveSceneProxy* RESTRICT PrimitiveSceneProxy,
-	int32 MeshId,
+	int32 StaticMeshId,
 	const FMaterialRenderProxy& RESTRICT MaterialRenderProxy,
 	const FMaterial& RESTRICT MaterialResource,
 	ERasterizerFillMode MeshFillMode,
@@ -1263,7 +1263,7 @@ void FDistortionMeshProcessor::Process(
 		);
 
 	FMeshMaterialShaderElementData ShaderElementData;
-	ShaderElementData.InitializeMeshMaterialData(ViewIfDynamicMeshCommand, PrimitiveSceneProxy, MeshBatch, MeshId, false);
+	ShaderElementData.InitializeMeshMaterialData(ViewIfDynamicMeshCommand, PrimitiveSceneProxy, MeshBatch, StaticMeshId, false);
 
 	BuildMeshDrawCommands(
 		MeshBatch,
diff --git a/Engine/Source/Runtime/Renderer/Private/DistortionRendering.h b/Engine/Source/Runtime/Renderer/Private/DistortionRendering.h
--- a/Engine/Source/Runtime/Renderer/Private/DistortionRendering.h
+++ b/Engine/Source/Runtime/Renderer/Private/DistortionRendering.h
@@ -79,7 +79,7 @@ public:
 
 	FDistortionMeshProcessor(const FScene* Scene, const FSceneView* InViewIfDynamicMeshCommand, const FDrawingPolicyRenderState& InPassDrawRenderState, FMeshPassDrawListContext& InDrawListContext);
 
-	virtual void AddMeshBatch(const FMeshBatch& RESTRICT MeshBatch, uint64 BatchElementMask, const FPrimitiveSceneProxy* RESTRICT PrimitiveSceneProxy, int32 MeshId = -1) override final;
+	virtual void AddMeshBatch(const FMeshBatch& RESTRICT MeshBatch, uint64 BatchElementMask, const FPrimitiveSceneProxy* RESTRICT PrimitiveSceneProxy, int32 StaticMeshId = -1) override final;
 
 	FDrawingPolicyRenderState PassDrawRenderState;
 
@@ -88,7 +88,7 @@ private:
 		const FMeshBatch& MeshBatch,
 		uint64 BatchElementMask,
 		const FPrimitiveSceneProxy* RESTRICT PrimitiveSceneProxy,
-		int32 MeshId,
+		int32 StaticMeshId,
 		const FMaterialRenderProxy& RESTRICT MaterialRenderProxy,
 		const FMaterial& RESTRICT MaterialResource,
 		ERasterizerFillMode MeshFillMode,
diff --git a/Engine/Source/Runtime/Renderer/Private/EditorPrimitivesRendering.cpp b/Engine/Source/Runtime/Renderer/Private/EditorPrimitivesRendering.cpp
--- a/Engine/Source/Runtime/Renderer/Private/EditorPrimitivesRendering.cpp
+++ b/Engine/Source/Runtime/Renderer/Private/EditorPrimitivesRendering.cpp
@@ -16,7 +16,7 @@ FEditorPrimitivesBasePassMeshProcessor::FEditorPrimitivesBasePassMeshProcessor(c
 	, bTranslucentBasePass(bInTranslucentBasePass)
 {}
 
-void FEditorPrimitivesBasePassMeshProcessor::AddMeshBatch(const FMeshBatch& RESTRICT MeshBatch, uint64 BatchElementMask, const FPrimitiveSceneProxy* RESTRICT PrimitiveSceneProxy, int32 MeshId)
+void FEditorPrimitivesBasePassMeshProcessor::AddMeshBatch(const FMeshBatch& RESTRICT MeshBatch, uint64 BatchElementMask, const FPrimitiveSceneProxy* RESTRICT PrimitiveSceneProxy, int32 StaticMeshId)
 {
 	if (MeshBatch.bUseForMaterial)
 	{
@@ -35,11 +35,11 @@ void FEditorPrimitivesBasePassMeshProcessor::AddMeshBatch(const FMeshBatch& REST
 		{
 			if (Scene->GetShadingPath(FeatureLevel) == EShadingPath::Mobile)
 			{
-				ProcessMobileShadingPath(MeshBatch, BatchElementMask, Material, MaterialRenderProxy, PrimitiveSceneProxy, MeshId); 
+				ProcessMobileShadingPath(MeshBatch, BatchElementMask, Material, MaterialRenderProxy, PrimitiveSceneProxy, StaticMeshId);
 			}
 			else
 			{
-				ProcessDeferredShadingPath(MeshBatch, BatchElementMask, Material, MaterialRenderProxy, PrimitiveSceneProxy, MeshId);
+				ProcessDeferredShadingPath(MeshBatch, BatchElementMask, Material, MaterialRenderProxy, PrimitiveSceneProxy, StaticMeshId);
 			}
 		}
 	}
@@ -103,7 +103,7 @@ void FEditorPrimitivesBasePassMeshProcessor::ProcessDeferredShadingPath(const FM
 		ShaderElementData);
 }
 
-void FEditorPrimitivesBasePassMeshProcessor::ProcessMobileShadingPath(const FMeshBatch& MeshBatch, uint64 BatchElementMask, const FMaterial& Material, const FMaterialRenderProxy& MaterialRenderProxy, const FPrimitiveSceneProxy* PrimitiveSceneProxy, int32 MeshId)
+void FEditorPrimitivesBasePassMeshProcessor::ProcessMobileShadingPath(const FMeshBatch& MeshBatch, uint64 BatchElementMask, const FMaterial& Material, const FMaterialRenderProxy& MaterialRenderProxy, const FPrimitiveSceneProxy* PrimitiveSceneProxy, int32 StaticMeshId)
 {
 	FUniformLightMapPolicy NoLightmapPolicy(LMP_NO_LIGHTMAP);
 	typedef FUniformLightMapPolicy LightMapPolicyType;
@@ -138,7 +138,7 @@ void FEditorPrimitivesBasePassMeshProcessor::ProcessMobileShadingPath(const FMes
 	ERasterizerCullMode MeshCullMode = ComputeMeshCullMode(MeshBatch, Material);
 	
 	TMobileBasePassShaderElementData<LightMapPolicyType> ShaderElementData(nullptr);
-	ShaderElementData.InitializeMeshMaterialData(ViewIfDynamicMeshCommand, PrimitiveSceneProxy, MeshBatch, MeshId, false);
+	ShaderElementData.InitializeMeshMaterialData(ViewIfDynamicMeshCommand, PrimitiveSceneProxy, MeshBatch, StaticMeshId, false);
 
 	BuildMeshDrawCommands(
 		MeshBatch,
diff --git a/Engine/Source/Runtime/Renderer/Private/EditorPrimitivesRendering.h b/Engine/Source/Runtime/Renderer/Private/EditorPrimitivesRendering.h
--- a/Engine/Source/Runtime/Renderer/Private/EditorPrimitivesRendering.h
+++ b/Engine/Source/Runtime/Renderer/Private/EditorPrimitivesRendering.h
@@ -11,12 +11,12 @@ class FEditorPrimitivesBasePassMeshProcessor : public FMeshPassProcessor
 public:
 	FEditorPrimitivesBasePassMeshProcessor(const FScene* InScene, ERHIFeatureLevel::Type InFeatureLevel, const FSceneView* InViewIfDynamicMeshCommand, const FDrawingPolicyRenderState& InDrawRenderState, bool bTranslucentBasePass, FMeshPassDrawListContext& InDrawListContext);
 
-	virtual void AddMeshBatch(const FMeshBatch& RESTRICT MeshBatch, uint64 BatchElementMask, const FPrimitiveSceneProxy* RESTRICT PrimitiveSceneProxy, int32 MeshId = -1) override final;
+	virtual void AddMeshBatch(const FMeshBatch& RESTRICT MeshBatch, uint64 BatchElementMask, const FPrimitiveSceneProxy* RESTRICT PrimitiveSceneProxy, int32 StaticMeshId = -1) override final;
 
 	const FDrawingPolicyRenderState PassDrawRenderState;
 	const bool bTranslucentBasePass;
 
 private:
-	void ProcessDeferredShadingPath(const FMeshBatch& MeshBatch, uint64 BatchElementMask, const FMaterial& Material, const FMaterialRenderProxy& MaterialRenderProxy, const FPrimitiveSceneProxy* PrimitiveSceneProxy, int32 MeshId);
-	void ProcessMobileShadingPath(const FMeshBatch& MeshBatch, uint64 BatchElementMask, const FMaterial& Material, const FMaterialRenderProxy& MaterialRenderProxy, const FPrimitiveSceneProxy* PrimitiveSceneProxy, int32 MeshId);
+	void ProcessDeferredShadingPath(const FMeshBatch& MeshBatch, uint64 BatchElementMask, const FMaterial& Material, const FMaterialRenderProxy& MaterialRenderProxy, const FPrimitiveSceneProxy* PrimitiveSceneProxy, int32 StaticMeshId);
+	void ProcessMobileShadingPath(const FMeshBatch& MeshBatch, uint64 BatchElementMask, const FMaterial& Material, const FMaterialRenderProxy& MaterialRenderProxy, const FPrimitiveSceneProxy* PrimitiveSceneProxy, int32 StaticMeshId);
 };
\ No newline at end of file
diff --git a/Engine/Source/Runtime/Renderer/Private/LightMapDensityRendering.cpp b/Engine/Source/Runtime/Renderer/Private/LightMapDensityRendering.cpp
--- a/Engine/Source/Runtime/Renderer/Private/LightMapDensityRendering.cpp
+++ b/Engine/Source/Runtime/Renderer/Private/LightMapDensityRendering.cpp
@@ -97,7 +97,7 @@ void FLightmapDensityMeshProcessor::Process(
 	const FMeshBatch& MeshBatch,
 	uint64 BatchElementMask,
 	const FPrimitiveSceneProxy* RESTRICT PrimitiveSceneProxy,
-	int32 MeshId, 
+	int32 StaticMeshId,
 	const FMaterialRenderProxy& RESTRICT MaterialRenderProxy,
 	const FMaterial& RESTRICT MaterialResource,
 	const LightMapPolicyType& RESTRICT LightMapPolicy,
@@ -130,7 +130,7 @@ void FLightmapDensityMeshProcessor::Process(
 	LightmapDensityPassShaders.PixelShader = MaterialResource.GetShader<TLightMapDensityPS<LightMapPolicyType>>(VertexFactoryType);
 
 	TLightMapDensityElementData<LightMapPolicyType> ShaderElementData(LightMapElementData);
-	ShaderElementData.InitializeMeshMaterialData(ViewIfDynamicMeshCommand, PrimitiveSceneProxy, MeshBatch, MeshId, true);
+	ShaderElementData.InitializeMeshMaterialData(ViewIfDynamicMeshCommand, PrimitiveSceneProxy, MeshBatch, StaticMeshId, true);
 
 	{		
 		// BuiltLightingAndSelectedFlags informs the shader is lighting is built or not for this primitive
@@ -203,7 +203,7 @@ void FLightmapDensityMeshProcessor::Process(
 		ShaderElementData.LightMapResolutionScale *= 0.5f;
 	}
 
-	ShaderElementData.InitializeMeshMaterialData(ViewIfDynamicMeshCommand, PrimitiveSceneProxy, MeshBatch, MeshId, false);
+	ShaderElementData.InitializeMeshMaterialData(ViewIfDynamicMeshCommand, PrimitiveSceneProxy, MeshBatch, StaticMeshId, false);
 
 	BuildMeshDrawCommands(
 		MeshBatch,
@@ -221,7 +221,7 @@ void FLightmapDensityMeshProcessor::Process(
 		ShaderElementData);
 }
 
-void FLightmapDensityMeshProcessor::AddMeshBatch(const FMeshBatch& RESTRICT MeshBatch, uint64 BatchElementMask, const FPrimitiveSceneProxy* RESTRICT PrimitiveSceneProxy, int32 MeshId)
+void FLightmapDensityMeshProcessor::AddMeshBatch(const FMeshBatch& RESTRICT MeshBatch, uint64 BatchElementMask, const FPrimitiveSceneProxy* RESTRICT PrimitiveSceneProxy, int32 StaticMeshId)
 {
 	if (FeatureLevel >= ERHIFeatureLevel::SM4 && ViewIfDynamicMeshCommand->Family->EngineShowFlags.LightMapDensity && AllowDebugViewmodes() && MeshBatch.bUseForMaterial)
 	{
@@ -272,7 +272,7 @@ void FLightmapDensityMeshProcessor::AddMeshBatch(const FMeshBatch& RESTRICT Mesh
 							MeshBatch, 
 							BatchElementMask, 
 							PrimitiveSceneProxy, 
-							MeshId,
+							StaticMeshId,
 							*MaterialRenderProxy, 
 							*Material, 
 							TUniformLightMapPolicy<LMP_HQ_LIGHTMAP>(), 
@@ -286,7 +286,7 @@ void FLightmapDensityMeshProcessor::AddMeshBatch(const FMeshBatch& RESTRICT Mesh
 							MeshBatch, 
 							BatchElementMask, 
 							PrimitiveSceneProxy, 
-							MeshId,
+							StaticMeshId,
 							*MaterialRenderProxy, 
 							*Material, 
 							TUniformLightMapPolicy<LMP_LQ_LIGHTMAP>(), 
@@ -301,7 +301,7 @@ void FLightmapDensityMeshProcessor::AddMeshBatch(const FMeshBatch& RESTRICT Mesh
 						MeshBatch, 
 						BatchElementMask, 
 						PrimitiveSceneProxy, 
-						MeshId,
+						StaticMeshId,
 						*MaterialRenderProxy, 
 						*Material, 
 						TUniformLightMapPolicy<LMP_DUMMY>(), 
@@ -316,7 +316,7 @@ void FLightmapDensityMeshProcessor::AddMeshBatch(const FMeshBatch& RESTRICT Mesh
 					MeshBatch, 
 					BatchElementMask, 
 					PrimitiveSceneProxy, 
-					MeshId,
+					StaticMeshId,
 					*MaterialRenderProxy, 
 					*Material, 
 					TUniformLightMapPolicy<LMP_NO_LIGHTMAP>(), 
diff --git a/Engine/Source/Runtime/Renderer/Private/LightMapDensityRendering.h b/Engine/Source/Runtime/Renderer/Private/LightMapDensityRendering.h
--- a/Engine/Source/Runtime/Renderer/Private/LightMapDensityRendering.h
+++ b/Engine/Source/Runtime/Renderer/Private/LightMapDensityRendering.h
@@ -299,7 +299,7 @@ public:
 
 	FLightmapDensityMeshProcessor(const FScene* Scene, const FSceneView* InViewIfDynamicMeshCommand, FMeshPassDrawListContext& InDrawListContext);
 
-	virtual void AddMeshBatch(const FMeshBatch& RESTRICT MeshBatch, uint64 BatchElementMask, const FPrimitiveSceneProxy* RESTRICT PrimitiveSceneProxy, int32 MeshId = -1) override final;
+	virtual void AddMeshBatch(const FMeshBatch& RESTRICT MeshBatch, uint64 BatchElementMask, const FPrimitiveSceneProxy* RESTRICT PrimitiveSceneProxy, int32 StaticMeshId = -1) override final;
 
 
 private:
@@ -309,7 +309,7 @@ private:
 		const FMeshBatch& MeshBatch,
 		uint64 BatchElementMask,
 		const FPrimitiveSceneProxy* RESTRICT PrimitiveSceneProxy,
-		int32 MeshId, 
+		int32 StaticMeshId,
 		const FMaterialRenderProxy& RESTRICT MaterialRenderProxy,
 		const FMaterial& RESTRICT MaterialResource,
 		const LightMapPolicyType& RESTRICT LightMapPolicy,
diff --git a/Engine/Source/Runtime/Renderer/Private/MobileBasePass.cpp b/Engine/Source/Runtime/Renderer/Private/MobileBasePass.cpp
--- a/Engine/Source/Runtime/Renderer/Private/MobileBasePass.cpp
+++ b/Engine/Source/Runtime/Renderer/Private/MobileBasePass.cpp
@@ -691,7 +691,7 @@ void FMobileBasePassMeshProcessor::AddMeshBatch(const FMeshBatch& RESTRICT MeshB
 void FMobileBasePassMeshProcessor::Process(
 		const FMeshBatch& RESTRICT MeshBatch,
 		uint64 BatchElementMask,
-		int32 MeshId,
+		int32 StaticMeshId,
 		const FPrimitiveSceneProxy* RESTRICT PrimitiveSceneProxy,
 		const FMaterialRenderProxy& RESTRICT MaterialRenderProxy,
 		const FMaterial& RESTRICT MaterialResource,
@@ -745,7 +745,7 @@ void FMobileBasePassMeshProcessor::Process(
 	ERasterizerCullMode MeshCullMode = ComputeMeshCullMode(MeshBatch, MaterialResource);
 
 	TMobileBasePassShaderElementData<FUniformLightMapPolicy> ShaderElementData(LightMapElementData);
-	ShaderElementData.InitializeMeshMaterialData(ViewIfDynamicMeshCommand, PrimitiveSceneProxy, MeshBatch, MeshId, false);
+	ShaderElementData.InitializeMeshMaterialData(ViewIfDynamicMeshCommand, PrimitiveSceneProxy, MeshBatch, StaticMeshId, false);
 
 	BuildMeshDrawCommands(
 		MeshBatch,
diff --git a/Engine/Source/Runtime/Renderer/Private/MobileBasePassRendering.h b/Engine/Source/Runtime/Renderer/Private/MobileBasePassRendering.h
--- a/Engine/Source/Runtime/Renderer/Private/MobileBasePassRendering.h
+++ b/Engine/Source/Runtime/Renderer/Private/MobileBasePassRendering.h
@@ -941,7 +941,7 @@ public:
 		bool bInCanReceiveCSM,
 		ETranslucencyPass::Type InTranslucencyPassType = ETranslucencyPass::TPT_MAX);
 
-	virtual void AddMeshBatch(const FMeshBatch& RESTRICT MeshBatch, uint64 BatchElementMask, const FPrimitiveSceneProxy* RESTRICT PrimitiveSceneProxy, int32 MeshId = -1) override final;
+	virtual void AddMeshBatch(const FMeshBatch& RESTRICT MeshBatch, uint64 BatchElementMask, const FPrimitiveSceneProxy* RESTRICT PrimitiveSceneProxy, int32 StaticMeshId = -1) override final;
 
 	FDrawingPolicyRenderState PassDrawRenderState;
 
@@ -949,7 +949,7 @@ private:
 	void Process(
 		const FMeshBatch& RESTRICT MeshBatch,
 		uint64 BatchElementMask,
-		int32 MeshId,
+		int32 StaticMeshId,
 		const FPrimitiveSceneProxy* RESTRICT PrimitiveSceneProxy,
 		const FMaterialRenderProxy& RESTRICT MaterialRenderProxy,
 		const FMaterial& RESTRICT MaterialResource,
diff --git a/Engine/Source/Runtime/Renderer/Private/MobileTranslucentRendering.cpp b/Engine/Source/Runtime/Renderer/Private/MobileTranslucentRendering.cpp
--- a/Engine/Source/Runtime/Renderer/Private/MobileTranslucentRendering.cpp
+++ b/Engine/Source/Runtime/Renderer/Private/MobileTranslucentRendering.cpp
@@ -824,7 +824,7 @@ public:
 		check(InViewIfDynamicMeshCommand);
 	}
 
-	virtual void AddMeshBatch(const FMeshBatch& RESTRICT MeshBatch, uint64 BatchElementMask, const FPrimitiveSceneProxy* RESTRICT PrimitiveSceneProxy, int32 MeshId = -1) override final
+	virtual void AddMeshBatch(const FMeshBatch& RESTRICT MeshBatch, uint64 BatchElementMask, const FPrimitiveSceneProxy* RESTRICT PrimitiveSceneProxy, int32 StaticMeshId = -1) override final
 	{
 		if (MeshBatch.bUseForMaterial)
 		{
@@ -837,13 +837,13 @@ public:
 
 			if (bIsTranslucent)
 			{
-				Process(MeshBatch, BatchElementMask, Material, MaterialRenderProxy, PrimitiveSceneProxy, MeshId); 
+				Process(MeshBatch, BatchElementMask, Material, MaterialRenderProxy, PrimitiveSceneProxy, StaticMeshId);
 			}
 		}
 	}
 
 private:
-	void Process(const FMeshBatch& MeshBatch, uint64 BatchElementMask, const FMaterial& Material, const FMaterialRenderProxy& MaterialRenderProxy, const FPrimitiveSceneProxy* PrimitiveSceneProxy, int32 MeshId)
+	void Process(const FMeshBatch& MeshBatch, uint64 BatchElementMask, const FMaterial& Material, const FMaterialRenderProxy& MaterialRenderProxy, const FPrimitiveSceneProxy* PrimitiveSceneProxy, int32 StaticMeshId)
 	{
 		const FVertexFactory* VertexFactory = MeshBatch.VertexFactory;
 
diff --git a/Engine/Source/Runtime/Renderer/Private/SceneHitProxyRendering.cpp b/Engine/Source/Runtime/Renderer/Private/SceneHitProxyRendering.cpp
--- a/Engine/Source/Runtime/Renderer/Private/SceneHitProxyRendering.cpp
+++ b/Engine/Source/Runtime/Renderer/Private/SceneHitProxyRendering.cpp
@@ -587,7 +587,7 @@ void FDeferredShadingSceneRenderer::RenderHitProxies(FRHICommandListImmediate& R
 #endif
 }
 
-void FHitProxyMeshProcessor::AddMeshBatch(const FMeshBatch& RESTRICT MeshBatch, uint64 BatchElementMask, const FPrimitiveSceneProxy* RESTRICT PrimitiveSceneProxy, int32 MeshId)
+void FHitProxyMeshProcessor::AddMeshBatch(const FMeshBatch& RESTRICT MeshBatch, uint64 BatchElementMask, const FPrimitiveSceneProxy* RESTRICT PrimitiveSceneProxy, int32 StaticMeshId)
 {
 	if (MeshBatch.bUseForMaterial && MeshBatch.bSelectable && Scene->RequiresHitProxies() && (!PrimitiveSceneProxy || PrimitiveSceneProxy->IsSelectable()))
 	{
@@ -614,7 +614,7 @@ void FHitProxyMeshProcessor::AddMeshBatch(const FMeshBatch& RESTRICT MeshBatch,
 
 		if (!bOpaqueOnly || !IsTranslucentBlendMode(BlendMode))
 		{
-			Process(MeshBatch, BatchElementMask, MeshId, PrimitiveSceneProxy, *MaterialRenderProxy, *Material, MeshFillMode, MeshCullMode);
+			Process(MeshBatch, BatchElementMask, StaticMeshId, PrimitiveSceneProxy, *MaterialRenderProxy, *Material, MeshFillMode, MeshCullMode);
 		}
 	}
 }
@@ -647,7 +647,7 @@ void GetHitProxyPassShaders(
 void FHitProxyMeshProcessor::Process(
 	const FMeshBatch& MeshBatch,
 	uint64 BatchElementMask,
-	int32 MeshId,
+	int32 StaticMeshId,
 	const FPrimitiveSceneProxy* RESTRICT PrimitiveSceneProxy,
 	const FMaterialRenderProxy& RESTRICT MaterialRenderProxy,
 	const FMaterial& RESTRICT MaterialResource,
@@ -673,7 +673,7 @@ void FHitProxyMeshProcessor::Process(
 	);
 
 	FHitProxyShaderElementData ShaderElementData(MeshBatch.BatchHitProxyId);
-	ShaderElementData.InitializeMeshMaterialData(ViewIfDynamicMeshCommand, PrimitiveSceneProxy, MeshBatch, MeshId, false);
+	ShaderElementData.InitializeMeshMaterialData(ViewIfDynamicMeshCommand, PrimitiveSceneProxy, MeshBatch, StaticMeshId, false);
 
 	BuildMeshDrawCommands(
 		MeshBatch,
@@ -722,7 +722,7 @@ FRegisterPassProcessorCreateFunction RegisterMobileHitProxyOpaqueOnlyPass(&Creat
 
 #if WITH_EDITOR
 
-void FEditorSelectionMeshProcessor::AddMeshBatch(const FMeshBatch& RESTRICT MeshBatch, uint64 BatchElementMask, const FPrimitiveSceneProxy* RESTRICT PrimitiveSceneProxy, int32 MeshId)
+void FEditorSelectionMeshProcessor::AddMeshBatch(const FMeshBatch& RESTRICT MeshBatch, uint64 BatchElementMask, const FPrimitiveSceneProxy* RESTRICT PrimitiveSceneProxy, int32 StaticMeshId)
 {
 	if (MeshBatch.bUseForMaterial 
 		&& MeshBatch.bUseSelectionOutline 
@@ -750,14 +750,14 @@ void FEditorSelectionMeshProcessor::AddMeshBatch(const FMeshBatch& RESTRICT Mesh
 
 		check(Material && MaterialRenderProxy);
 
-		Process(MeshBatch, BatchElementMask, MeshId, PrimitiveSceneProxy, *MaterialRenderProxy, *Material, MeshFillMode, MeshCullMode);
+		Process(MeshBatch, BatchElementMask, StaticMeshId, PrimitiveSceneProxy, *MaterialRenderProxy, *Material, MeshFillMode, MeshCullMode);
 	}
 }
 
 void FEditorSelectionMeshProcessor::Process(
 	const FMeshBatch& MeshBatch,
 	uint64 BatchElementMask,
-	int32 MeshId,
+	int32 StaticMeshId,
 	const FPrimitiveSceneProxy* RESTRICT PrimitiveSceneProxy,
 	const FMaterialRenderProxy& RESTRICT MaterialRenderProxy,
 	const FMaterial& RESTRICT MaterialResource,
@@ -787,7 +787,7 @@ void FEditorSelectionMeshProcessor::Process(
 
 	FHitProxyId DummyId;
 	FHitProxyShaderElementData ShaderElementData(DummyId);
-	ShaderElementData.InitializeMeshMaterialData(ViewIfDynamicMeshCommand, PrimitiveSceneProxy, MeshBatch, MeshId, false);
+	ShaderElementData.InitializeMeshMaterialData(ViewIfDynamicMeshCommand, PrimitiveSceneProxy, MeshBatch, StaticMeshId, false);
 
 	BuildMeshDrawCommands(
 		MeshBatch,
diff --git a/Engine/Source/Runtime/Renderer/Private/SceneHitProxyRendering.h b/Engine/Source/Runtime/Renderer/Private/SceneHitProxyRendering.h
--- a/Engine/Source/Runtime/Renderer/Private/SceneHitProxyRendering.h
+++ b/Engine/Source/Runtime/Renderer/Private/SceneHitProxyRendering.h
@@ -22,7 +22,7 @@ public:
 
 	FHitProxyMeshProcessor(const FScene* Scene, const FSceneView* InViewIfDynamicMeshCommand, bool InbOpaqueOnly, const FDrawingPolicyRenderState& InRenderState, FMeshPassDrawListContext& InDrawListContext);
 
-	virtual void AddMeshBatch(const FMeshBatch& RESTRICT MeshBatch, uint64 BatchElementMask, const FPrimitiveSceneProxy* RESTRICT PrimitiveSceneProxy, int32 MeshId = -1) override final;
+	virtual void AddMeshBatch(const FMeshBatch& RESTRICT MeshBatch, uint64 BatchElementMask, const FPrimitiveSceneProxy* RESTRICT PrimitiveSceneProxy, int32 StaticMeshId = -1) override final;
 
 	FDrawingPolicyRenderState PassDrawRenderState;
 
@@ -31,7 +31,7 @@ private:
 	void Process(
 		const FMeshBatch& MeshBatch,
 		uint64 BatchElementMask,
-		int32 MeshId,
+		int32 StaticMeshId,
 		const FPrimitiveSceneProxy* RESTRICT PrimitiveSceneProxy,
 		const FMaterialRenderProxy& RESTRICT MaterialRenderProxy,
 		const FMaterial& RESTRICT MaterialResource,
@@ -50,7 +50,7 @@ public:
 
 	FEditorSelectionMeshProcessor(const FScene* Scene, const FSceneView* InViewIfDynamicMeshCommand, FMeshPassDrawListContext& InDrawListContext);
 
-	virtual void AddMeshBatch(const FMeshBatch& RESTRICT MeshBatch, uint64 BatchElementMask, const FPrimitiveSceneProxy* RESTRICT PrimitiveSceneProxy, int32 MeshId = -1) override final;
+	virtual void AddMeshBatch(const FMeshBatch& RESTRICT MeshBatch, uint64 BatchElementMask, const FPrimitiveSceneProxy* RESTRICT PrimitiveSceneProxy, int32 StaticMeshId = -1) override final;
 
 	FDrawingPolicyRenderState PassDrawRenderState;
 
@@ -58,7 +58,7 @@ private:
 	void Process(
 		const FMeshBatch& MeshBatch,
 		uint64 BatchElementMask,
-		int32 MeshId,
+		int32 StaticMeshId,
 		const FPrimitiveSceneProxy* RESTRICT PrimitiveSceneProxy,
 		const FMaterialRenderProxy& RESTRICT MaterialRenderProxy,
 		const FMaterial& RESTRICT MaterialResource,
diff --git a/Engine/Source/Runtime/Renderer/Private/ShaderBaseClasses.cpp b/Engine/Source/Runtime/Renderer/Private/ShaderBaseClasses.cpp
--- a/Engine/Source/Runtime/Renderer/Private/ShaderBaseClasses.cpp
+++ b/Engine/Source/Runtime/Renderer/Private/ShaderBaseClasses.cpp
@@ -23,23 +23,23 @@ FAutoConsoleVariableRef FMaterialShader::CVarAllowCachedUniformExpressions(
 	TEXT("Allow uniform expressions to be cached."),
 	ECVF_RenderThreadSafe);
 
-void FMeshMaterialShaderElementData::InitializeMeshMaterialData(const FSceneView* SceneView, const FPrimitiveSceneProxy* RESTRICT PrimitiveSceneProxy, const FMeshBatch& RESTRICT MeshBatch, int32 MeshId, bool bAllowStencilDither)
+void FMeshMaterialShaderElementData::InitializeMeshMaterialData(const FSceneView* SceneView, const FPrimitiveSceneProxy* RESTRICT PrimitiveSceneProxy, const FMeshBatch& RESTRICT MeshBatch, int32 StaticMeshId, bool bAllowStencilDither)
 {
 	FadeUniformBuffer = GDistanceCullFadedInUniformBuffer.GetUniformBufferRHI();
 	DitherUniformBuffer = GDitherFadedInUniformBuffer.GetUniformBufferRHI();
 
-	if (SceneView && MeshId >= 0)
+	if (SceneView && StaticMeshId >= 0)
 	{
 		checkSlow(SceneView->bIsViewInfo);
 		const FViewInfo* ViewInfo = (FViewInfo*)SceneView;
 
 		if (MeshBatch.bDitheredLODTransition && !(bAllowStencilDither && ViewInfo->bAllowStencilDither))
 		{
-			if (ViewInfo->StaticMeshFadeOutDitheredLODMap[MeshId])
+			if (ViewInfo->StaticMeshFadeOutDitheredLODMap[StaticMeshId])
 			{
 				DitherUniformBuffer = ViewInfo->DitherFadeOutUniformBuffer;
 			}
-			else if (ViewInfo->StaticMeshFadeInDitheredLODMap[MeshId])
+			else if (ViewInfo->StaticMeshFadeInDitheredLODMap[StaticMeshId])
 			{
 				DitherUniformBuffer = ViewInfo->DitherFadeInUniformBuffer;
 			}
diff --git a/Engine/Source/Runtime/Renderer/Private/ShadowDepthRendering.cpp b/Engine/Source/Runtime/Renderer/Private/ShadowDepthRendering.cpp
--- a/Engine/Source/Runtime/Renderer/Private/ShadowDepthRendering.cpp
+++ b/Engine/Source/Runtime/Renderer/Private/ShadowDepthRendering.cpp
@@ -2537,7 +2537,7 @@ template<bool bRenderReflectiveShadowMap>
 void FShadowDepthPassMeshProcessor::Process(
 	const FMeshBatch& RESTRICT MeshBatch,
 	uint64 BatchElementMask,
-	int32 MeshId,
+	int32 StaticMeshId,
 	const FPrimitiveSceneProxy* RESTRICT PrimitiveSceneProxy,
 	const FMaterialRenderProxy& RESTRICT MaterialRenderProxy,
 	const FMaterial& RESTRICT MaterialResource,
@@ -2572,7 +2572,7 @@ void FShadowDepthPassMeshProcessor::Process(
 		ShadowDepthPassShaders.GeometryShader);
 
 	FShadowDepthShaderElementData ShaderElementData;
-	ShaderElementData.InitializeMeshMaterialData(ViewIfDynamicMeshCommand, PrimitiveSceneProxy, MeshBatch, MeshId, false);
+	ShaderElementData.InitializeMeshMaterialData(ViewIfDynamicMeshCommand, PrimitiveSceneProxy, MeshBatch, StaticMeshId, false);
 
 	const uint32 InstanceFactor = !ShadowDepthType.bOnePassPointLightShadow || RHISupportsGeometryShaders(GShaderPlatformForFeatureLevel[FeatureLevel]) ? 1 : 6;
 	for (uint32 i = 0; i < InstanceFactor; i++)
@@ -2596,7 +2596,7 @@ void FShadowDepthPassMeshProcessor::Process(
 	}
 }
 
-void FShadowDepthPassMeshProcessor::AddMeshBatch(const FMeshBatch& RESTRICT MeshBatch, uint64 BatchElementMask, const FPrimitiveSceneProxy* RESTRICT PrimitiveSceneProxy, int32 MeshId)
+void FShadowDepthPassMeshProcessor::AddMeshBatch(const FMeshBatch& RESTRICT MeshBatch, uint64 BatchElementMask, const FPrimitiveSceneProxy* RESTRICT PrimitiveSceneProxy, int32 StaticMeshId)
 {
 	if (MeshBatch.CastShadow)
 	{
@@ -2642,11 +2642,11 @@ void FShadowDepthPassMeshProcessor::AddMeshBatch(const FMeshBatch& RESTRICT Mesh
 
 			if (ShadowDepthType.bReflectiveShadowmap)
 			{
-				Process<true>(MeshBatch, BatchElementMask, MeshId, PrimitiveSceneProxy, *EffectiveMaterialRenderProxy, *EffectiveMaterial, MeshFillMode, FinalCullMode);
+				Process<true>(MeshBatch, BatchElementMask, StaticMeshId, PrimitiveSceneProxy, *EffectiveMaterialRenderProxy, *EffectiveMaterial, MeshFillMode, FinalCullMode);
 			}
 			else
 			{
-				Process<false>(MeshBatch, BatchElementMask, MeshId, PrimitiveSceneProxy, *EffectiveMaterialRenderProxy, *EffectiveMaterial, MeshFillMode, FinalCullMode);
+				Process<false>(MeshBatch, BatchElementMask, StaticMeshId, PrimitiveSceneProxy, *EffectiveMaterialRenderProxy, *EffectiveMaterial, MeshFillMode, FinalCullMode);
 			}
 		}
 	}
diff --git a/Engine/Source/Runtime/Renderer/Private/ShadowRendering.h b/Engine/Source/Runtime/Renderer/Private/ShadowRendering.h
--- a/Engine/Source/Runtime/Renderer/Private/ShadowRendering.h
+++ b/Engine/Source/Runtime/Renderer/Private/ShadowRendering.h
@@ -306,7 +306,7 @@ public:
 		FShadowDepthType InShadowDepthType,
 		FMeshPassDrawListContext& InDrawListContext);
 
-	virtual void AddMeshBatch(const FMeshBatch& RESTRICT MeshBatch, uint64 BatchElementMask, const FPrimitiveSceneProxy* RESTRICT PrimitiveSceneProxy, int32 MeshId = -1) override final;
+	virtual void AddMeshBatch(const FMeshBatch& RESTRICT MeshBatch, uint64 BatchElementMask, const FPrimitiveSceneProxy* RESTRICT PrimitiveSceneProxy, int32 StaticMeshId = -1) override final;
 
 	FDrawingPolicyRenderState PassDrawRenderState;
 
@@ -316,7 +316,7 @@ private:
 	void Process(
 		const FMeshBatch& RESTRICT MeshBatch,
 		uint64 BatchElementMask,
-		int32 MeshId,
+		int32 StaticMeshId,
 		const FPrimitiveSceneProxy* RESTRICT PrimitiveSceneProxy,
 		const FMaterialRenderProxy& RESTRICT MaterialRenderProxy,
 		const FMaterial& RESTRICT MaterialResource,
diff --git a/Engine/Source/Runtime/Renderer/Private/VelocityRendering.cpp b/Engine/Source/Runtime/Renderer/Private/VelocityRendering.cpp
--- a/Engine/Source/Runtime/Renderer/Private/VelocityRendering.cpp
+++ b/Engine/Source/Runtime/Renderer/Private/VelocityRendering.cpp
@@ -602,7 +602,7 @@ bool FVelocityRendering::PrimitiveHasVelocity(const FViewInfo& View, const FPrim
 }
 
 
-void FVelocityMeshProcessor::AddMeshBatch(const FMeshBatch& RESTRICT MeshBatch, uint64 BatchElementMask, const FPrimitiveSceneProxy* RESTRICT PrimitiveSceneProxy, int32 MeshId)
+void FVelocityMeshProcessor::AddMeshBatch(const FMeshBatch& RESTRICT MeshBatch, uint64 BatchElementMask, const FPrimitiveSceneProxy* RESTRICT PrimitiveSceneProxy, int32 StaticMeshId)
 {
 	// Velocity only needs to be directly rendered for movable meshes
 	bool bRequiresSeparateVelocity = false;
@@ -652,7 +652,7 @@ void FVelocityMeshProcessor::AddMeshBatch(const FMeshBatch& RESTRICT MeshBatch,
 
 			check(Material && MaterialRenderProxy);
 
-			Process(MeshBatch, BatchElementMask, MeshId, PrimitiveSceneProxy, *MaterialRenderProxy, *Material, MeshFillMode, MeshCullMode);
+			Process(MeshBatch, BatchElementMask, StaticMeshId, PrimitiveSceneProxy, *MaterialRenderProxy, *Material, MeshFillMode, MeshCullMode);
 		}
 	}
 }
@@ -699,7 +699,7 @@ void GetVelocityPassShaders(
 void FVelocityMeshProcessor::Process(
 	const FMeshBatch& MeshBatch,
 	uint64 BatchElementMask,
-	int32 MeshId,
+	int32 StaticMeshId,
 	const FPrimitiveSceneProxy* RESTRICT PrimitiveSceneProxy,
 	const FMaterialRenderProxy& RESTRICT MaterialRenderProxy,
 	const FMaterial& RESTRICT MaterialResource,
@@ -725,7 +725,7 @@ void FVelocityMeshProcessor::Process(
 	);
 
 	FMeshMaterialShaderElementData ShaderElementData;
-	ShaderElementData.InitializeMeshMaterialData(ViewIfDynamicMeshCommand, PrimitiveSceneProxy, MeshBatch, MeshId, false);
+	ShaderElementData.InitializeMeshMaterialData(ViewIfDynamicMeshCommand, PrimitiveSceneProxy, MeshBatch, StaticMeshId, false);
 
 	BuildMeshDrawCommands(
 		MeshBatch,
diff --git a/Engine/Source/Runtime/Renderer/Private/VelocityRendering.h b/Engine/Source/Runtime/Renderer/Private/VelocityRendering.h
--- a/Engine/Source/Runtime/Renderer/Private/VelocityRendering.h
+++ b/Engine/Source/Runtime/Renderer/Private/VelocityRendering.h
@@ -50,7 +50,7 @@ public:
 
 	FVelocityMeshProcessor(const FScene* Scene, const FSceneView* InViewIfDynamicMeshCommand, const FDrawingPolicyRenderState& InPassDrawRenderState, FMeshPassDrawListContext& InDrawListContext);
 
-	virtual void AddMeshBatch(const FMeshBatch& RESTRICT MeshBatch, uint64 BatchElementMask, const FPrimitiveSceneProxy* RESTRICT PrimitiveSceneProxy, int32 MeshId = -1) override final;
+	virtual void AddMeshBatch(const FMeshBatch& RESTRICT MeshBatch, uint64 BatchElementMask, const FPrimitiveSceneProxy* RESTRICT PrimitiveSceneProxy, int32 StaticMeshId = -1) override final;
 
 	FDrawingPolicyRenderState PassDrawRenderState;
 
@@ -58,7 +58,7 @@ private:
 	void Process(
 		const FMeshBatch& MeshBatch,
 		uint64 BatchElementMask,
-		int32 MeshId,
+		int32 StaticMeshId,
 		const FPrimitiveSceneProxy* RESTRICT PrimitiveSceneProxy,
 		const FMaterialRenderProxy& RESTRICT MaterialRenderProxy,
 		const FMaterial& RESTRICT MaterialResource,
diff --git a/Engine/Source/Runtime/Renderer/Private/VolumetricFogVoxelization.cpp b/Engine/Source/Runtime/Renderer/Private/VolumetricFogVoxelization.cpp
--- a/Engine/Source/Runtime/Renderer/Private/VolumetricFogVoxelization.cpp
+++ b/Engine/Source/Runtime/Renderer/Private/VolumetricFogVoxelization.cpp
@@ -445,7 +445,7 @@ public:
 
 	void AddMeshBatch(const FMeshBatch& RESTRICT MeshBatch, uint64 BatchElementMask, int32 NumVoxelizationPasses, const FPrimitiveSceneProxy* RESTRICT PrimitiveSceneProxy);
 
-	virtual void AddMeshBatch(const FMeshBatch& RESTRICT MeshBatch, uint64 BatchElementMask, const FPrimitiveSceneProxy* RESTRICT PrimitiveSceneProxy, int32 MeshId = -1) override final
+	virtual void AddMeshBatch(const FMeshBatch& RESTRICT MeshBatch, uint64 BatchElementMask, const FPrimitiveSceneProxy* RESTRICT PrimitiveSceneProxy, int32 StaticMeshId = -1) override final
 	{
 		checkf(false, TEXT("Default AddMeshBatch can't be used as rendering requires extra parameters per pass."));
 	}
diff --git a/Engine/Source/Runtime/Renderer/Public/MeshMaterialShader.h b/Engine/Source/Runtime/Renderer/Public/MeshMaterialShader.h
--- a/Engine/Source/Runtime/Renderer/Public/MeshMaterialShader.h
+++ b/Engine/Source/Runtime/Renderer/Public/MeshMaterialShader.h
@@ -26,7 +26,7 @@ public:
 	FUniformBufferRHIParamRef FadeUniformBuffer = nullptr;
 	FUniformBufferRHIParamRef DitherUniformBuffer = nullptr;
 
-	RENDERER_API void InitializeMeshMaterialData(const FSceneView* SceneView, const FPrimitiveSceneProxy* RESTRICT PrimitiveSceneProxy, const FMeshBatch& RESTRICT MeshBatch, int32 MeshId, bool bAllowStencilDither);
+	RENDERER_API void InitializeMeshMaterialData(const FSceneView* SceneView, const FPrimitiveSceneProxy* RESTRICT PrimitiveSceneProxy, const FMeshBatch& RESTRICT MeshBatch, int32 StaticMeshId, bool bAllowStencilDither);
 };
 
 /** Base class of all shaders that need material and vertex factory parameters. */
diff --git a/Engine/Source/Runtime/Renderer/Public/MeshPassProcessor.h b/Engine/Source/Runtime/Renderer/Public/MeshPassProcessor.h
--- a/Engine/Source/Runtime/Renderer/Public/MeshPassProcessor.h
+++ b/Engine/Source/Runtime/Renderer/Public/MeshPassProcessor.h
@@ -650,7 +650,7 @@ public:
 
 	// FMeshPassProcessor interface
 	// Add a FMeshBatch to the pass
-	virtual void AddMeshBatch(const FMeshBatch& RESTRICT MeshBatch, uint64 BatchElementMask, const FPrimitiveSceneProxy* RESTRICT PrimitiveSceneProxy, int32 MeshId = -1) = 0;
+	virtual void AddMeshBatch(const FMeshBatch& RESTRICT MeshBatch, uint64 BatchElementMask, const FPrimitiveSceneProxy* RESTRICT PrimitiveSceneProxy, int32 StaticMeshId = -1) = 0;
 
 	static FORCEINLINE_DEBUGGABLE ERasterizerCullMode InverseCullMode(ERasterizerCullMode CullMode)
 	{
