Fix FReadBuffer shadow leak

diff --git a/Engine/Source/Runtime/Renderer/Private/SceneRendering.cpp b/Engine/Source/Runtime/Renderer/Private/SceneRendering.cpp
--- a/Engine/Source/Runtime/Renderer/Private/SceneRendering.cpp
+++ b/Engine/Source/Runtime/Renderer/Private/SceneRendering.cpp
@@ -1568,6 +1568,7 @@ void FViewInfo::DestroyAllSnapshots()
 	{
 		Snapshot->ViewUniformBuffer.SafeRelease();
 		Snapshot->CachedViewUniformShaderParameters.Reset();
+		Snapshot->OneFramePrimitiveIdBufferEmulation.Release();
 		FreeViewInfoSnapshots.Add(Snapshot);
 	}
 	ViewInfoSnapshots.Reset();
