-FMobileInverseOpacityMeshProcessor - passing StaticMeshId instead of -1, as it is required for LOD dithering.
-Removed virtual void GenerateDynamicMeshDrawCommands, as it's no longer specialized between mobile and deferred renders.

diff --git a/Engine/Source/Runtime/Renderer/Private/DeferredShadingRenderer.cpp b/Engine/Source/Runtime/Renderer/Private/DeferredShadingRenderer.cpp
--- a/Engine/Source/Runtime/Renderer/Private/DeferredShadingRenderer.cpp
+++ b/Engine/Source/Runtime/Renderer/Private/DeferredShadingRenderer.cpp
@@ -375,13 +375,6 @@ void FDeferredShadingSceneRenderer::RenderVisualizeTexturePool(FRHICommandListIm
 }
 #endif
 
-void FDeferredShadingSceneRenderer::GenerateDynamicMeshDrawCommands()
-{
-	QUICK_SCOPE_CYCLE_COUNTER(STAT_GenerateDynamicMeshDrawCommands);
-
-	FSceneRenderer::GenerateDynamicMeshDrawCommands();
-}
-
 /** 
 * Finishes the view family rendering.
 */
diff --git a/Engine/Source/Runtime/Renderer/Private/DeferredShadingRenderer.h b/Engine/Source/Runtime/Renderer/Private/DeferredShadingRenderer.h
--- a/Engine/Source/Runtime/Renderer/Private/DeferredShadingRenderer.h
+++ b/Engine/Source/Runtime/Renderer/Private/DeferredShadingRenderer.h
@@ -115,10 +115,6 @@ public:
 	FORCEINLINE void DoDebugViewModePostProcessing(FRHICommandListImmediate& RHICmdList, const FViewInfo& View, TRefCountPtr<IPooledRenderTarget>& VelocityRT) {}
 #endif
 
-protected:
-
-	virtual void GenerateDynamicMeshDrawCommands() override;
-
 private:
 
 	static FGraphEventRef TranslucencyTimestampQuerySubmittedFence[FOcclusionQueryHelpers::MaxBufferedOcclusionFrames + 1];
diff --git a/Engine/Source/Runtime/Renderer/Private/EditorPrimitivesRendering.cpp b/Engine/Source/Runtime/Renderer/Private/EditorPrimitivesRendering.cpp
--- a/Engine/Source/Runtime/Renderer/Private/EditorPrimitivesRendering.cpp
+++ b/Engine/Source/Runtime/Renderer/Private/EditorPrimitivesRendering.cpp
@@ -45,7 +45,7 @@ void FEditorPrimitivesBasePassMeshProcessor::AddMeshBatch(const FMeshBatch& REST
 	}
 }
 
-void FEditorPrimitivesBasePassMeshProcessor::ProcessDeferredShadingPath(const FMeshBatch& MeshBatch, uint64 BatchElementMask, const FMaterial& Material, const FMaterialRenderProxy& MaterialRenderProxy, const FPrimitiveSceneProxy* PrimitiveSceneProxy, int32 MeshId)
+void FEditorPrimitivesBasePassMeshProcessor::ProcessDeferredShadingPath(const FMeshBatch& MeshBatch, uint64 BatchElementMask, const FMaterial& Material, const FMaterialRenderProxy& MaterialRenderProxy, const FPrimitiveSceneProxy* PrimitiveSceneProxy, int32 StaticMeshId)
 {
 	FUniformLightMapPolicy NoLightmapPolicy(LMP_NO_LIGHTMAP);
 	typedef FUniformLightMapPolicy LightMapPolicyType;
@@ -85,7 +85,7 @@ void FEditorPrimitivesBasePassMeshProcessor::ProcessDeferredShadingPath(const FM
 	ERasterizerCullMode MeshCullMode = ComputeMeshCullMode(MeshBatch, Material);
 	
 	TBasePassShaderElementData<LightMapPolicyType> ShaderElementData(nullptr);
-	ShaderElementData.InitializeMeshMaterialData(ViewIfDynamicMeshCommand, PrimitiveSceneProxy, MeshBatch, MeshId, false);
+	ShaderElementData.InitializeMeshMaterialData(ViewIfDynamicMeshCommand, PrimitiveSceneProxy, MeshBatch, StaticMeshId, false);
 
 	BuildMeshDrawCommands(
 		MeshBatch,
diff --git a/Engine/Source/Runtime/Renderer/Private/MobileBasePass.cpp b/Engine/Source/Runtime/Renderer/Private/MobileBasePass.cpp
--- a/Engine/Source/Runtime/Renderer/Private/MobileBasePass.cpp
+++ b/Engine/Source/Runtime/Renderer/Private/MobileBasePass.cpp
@@ -646,7 +646,7 @@ FMobileBasePassMeshProcessor::FMobileBasePassMeshProcessor(const FScene* Scene,
 	, bCanReceiveCSM(bInCanReceiveCSM)
 {}
 
-void FMobileBasePassMeshProcessor::AddMeshBatch(const FMeshBatch& RESTRICT MeshBatch, uint64 BatchElementMask, const FPrimitiveSceneProxy* RESTRICT PrimitiveSceneProxy, int32 MeshId)
+void FMobileBasePassMeshProcessor::AddMeshBatch(const FMeshBatch& RESTRICT MeshBatch, uint64 BatchElementMask, const FPrimitiveSceneProxy* RESTRICT PrimitiveSceneProxy, int32 StaticMeshId)
 {
 	if (!MeshBatch.bUseForMaterial)
 	{
@@ -673,7 +673,7 @@ void FMobileBasePassMeshProcessor::AddMeshBatch(const FMeshBatch& RESTRICT MeshB
 			check(bCanReceiveCSM == false);
 			const FLightSceneInfo* MobileDirectionalLight = MobileBasePass::GetDirectionalLightInfo(Scene, PrimitiveSceneProxy);
 			ELightMapPolicyType LightmapPolicyType = MobileBasePass::SelectMeshLightmapPolicy(Scene, MeshBatch, PrimitiveSceneProxy, MobileDirectionalLight, ShadingModel, bCanReceiveCSM, FeatureLevel);
-			Process(MeshBatch, BatchElementMask, MeshId, PrimitiveSceneProxy, MaterialRenderProxy, Material, BlendMode, ShadingModel, LightmapPolicyType, MeshBatch.LCI);
+			Process(MeshBatch, BatchElementMask, StaticMeshId, PrimitiveSceneProxy, MaterialRenderProxy, Material, BlendMode, ShadingModel, LightmapPolicyType, MeshBatch.LCI);
 		}
 	}
 	else
@@ -683,7 +683,7 @@ void FMobileBasePassMeshProcessor::AddMeshBatch(const FMeshBatch& RESTRICT MeshB
 		{
 			const FLightSceneInfo* MobileDirectionalLight = MobileBasePass::GetDirectionalLightInfo(Scene, PrimitiveSceneProxy);
 			ELightMapPolicyType LightmapPolicyType = MobileBasePass::SelectMeshLightmapPolicy(Scene, MeshBatch, PrimitiveSceneProxy, MobileDirectionalLight, ShadingModel, bCanReceiveCSM, FeatureLevel);
-			Process(MeshBatch, BatchElementMask, MeshId, PrimitiveSceneProxy, MaterialRenderProxy, Material, BlendMode, ShadingModel, LightmapPolicyType, MeshBatch.LCI);
+			Process(MeshBatch, BatchElementMask, StaticMeshId, PrimitiveSceneProxy, MaterialRenderProxy, Material, BlendMode, ShadingModel, LightmapPolicyType, MeshBatch.LCI);
 		}
 	}
 }
diff --git a/Engine/Source/Runtime/Renderer/Private/MobileShadingRenderer.cpp b/Engine/Source/Runtime/Renderer/Private/MobileShadingRenderer.cpp
--- a/Engine/Source/Runtime/Renderer/Private/MobileShadingRenderer.cpp
+++ b/Engine/Source/Runtime/Renderer/Private/MobileShadingRenderer.cpp
@@ -221,11 +221,6 @@ void FMobileSceneRenderer::PostInitViewCustomData()
 	}
 }
 
-void FMobileSceneRenderer::GenerateDynamicMeshDrawCommands()
-{
-	FSceneRenderer::GenerateDynamicMeshDrawCommands();
-}
-
 /** 
 * Renders the view family. 
 */
diff --git a/Engine/Source/Runtime/Renderer/Private/MobileTranslucentRendering.cpp b/Engine/Source/Runtime/Renderer/Private/MobileTranslucentRendering.cpp
--- a/Engine/Source/Runtime/Renderer/Private/MobileTranslucentRendering.cpp
+++ b/Engine/Source/Runtime/Renderer/Private/MobileTranslucentRendering.cpp
@@ -863,7 +863,7 @@ private:
 		ERasterizerCullMode MeshCullMode = ComputeMeshCullMode(MeshBatch, Material);
 	
 		FMeshMaterialShaderElementData ShaderElementData;
-		ShaderElementData.InitializeMeshMaterialData(ViewIfDynamicMeshCommand, PrimitiveSceneProxy, MeshBatch, -1, false);
+		ShaderElementData.InitializeMeshMaterialData(ViewIfDynamicMeshCommand, PrimitiveSceneProxy, MeshBatch, StaticMeshId, false);
 
 		FMeshDrawCommandSortKey SortKey = CalculateStaticTranslucentMeshSortKey(PrimitiveSceneProxy);
 		
diff --git a/Engine/Source/Runtime/Renderer/Private/SceneRendering.h b/Engine/Source/Runtime/Renderer/Private/SceneRendering.h
--- a/Engine/Source/Runtime/Renderer/Private/SceneRendering.h
+++ b/Engine/Source/Runtime/Renderer/Private/SceneRendering.h
@@ -1533,7 +1533,7 @@ protected:
 
 	void InitDynamicShadows(FRHICommandListImmediate& RHICmdList);
 
-	virtual void GenerateDynamicMeshDrawCommands();
+	void GenerateDynamicMeshDrawCommands();
 
 	void ApplyViewOverridesToMeshDrawCommands(FExclusiveDepthStencil::Type BasePassDepthStencilAccess);
 
@@ -1715,8 +1715,6 @@ protected:
 
 	void InitViews(FRHICommandListImmediate& RHICmdList);
 
-	virtual void GenerateDynamicMeshDrawCommands() override;
-
 	/** Renders the opaque base pass for mobile. */
 	void RenderMobileBasePass(FRHICommandListImmediate& RHICmdList, const TArrayView<const FViewInfo*> PassViews);
 
