Fixed planar reflection crash when there are more than 2 views (planar reflection has hardcoded max 2 views).

diff --git a/Engine/Source/Runtime/Renderer/Private/PlanarReflectionRendering.cpp b/Engine/Source/Runtime/Renderer/Private/PlanarReflectionRendering.cpp
--- a/Engine/Source/Runtime/Renderer/Private/PlanarReflectionRendering.cpp
+++ b/Engine/Source/Runtime/Renderer/Private/PlanarReflectionRendering.cpp
@@ -68,7 +68,7 @@ void SetupPlanarReflectionUniformParameters(const class FSceneView& View, const
 			FVector4(0, 0, 0, 0),
 			FVector4(0, 0, 0, 0),
 		};
-		for (int32 ViewIndex = 0; ViewIndex < View.Family->Views.Num(); ViewIndex++)
+		for (int32 ViewIndex = 0; ViewIndex < FMath::Min(View.Family->Views.Num(), GMaxPlanarReflectionViews); ViewIndex++)
 		{
 			FIntRect ViewRect = ReflectionSceneProxy->ViewRect[ViewIndex];
 			ScreenScaleBiasValue[ViewIndex] = FVector4(
