Recache cached mesh draw commands during UpdateStaticDrawListsForMaterials. This fixes crashes when FMaterialUpdateContext is created with only SyncWithRenderingThread flag and won't recreate render state contexts. Thus leaving pointers to deleted uniform buffers inside cached mesh draw commands and causing a crash later.

diff --git a/Engine/Source/Runtime/Renderer/Private/RendererScene.cpp b/Engine/Source/Runtime/Renderer/Private/RendererScene.cpp
--- a/Engine/Source/Runtime/Renderer/Private/RendererScene.cpp
+++ b/Engine/Source/Runtime/Renderer/Private/RendererScene.cpp
@@ -3348,12 +3348,67 @@ void FRendererModule::RemoveScene(FSceneInterface* Scene)
 	AllocatedScenes.Remove(Scene);
 }
 
+void UpdateStaticMeshesForMaterials(const TArray<const FMaterial*>& MaterialResourcesToUpdate)
+{
+	TArray<UMaterialInterface*> UsedMaterials;
+
+	for (TObjectIterator<UPrimitiveComponent> PrimitiveIt; PrimitiveIt; ++PrimitiveIt)
+	{
+		UPrimitiveComponent* PrimitiveComponent = *PrimitiveIt;
+
+		if (PrimitiveComponent->IsRenderStateCreated() && PrimitiveComponent->SceneProxy)
+		{
+			UsedMaterials.Reset();
+			bool bPrimitiveIsDependentOnMaterial = false;
+
+			// Note: relying on GetUsedMaterials to be accurate, or else we won't propagate to the right primitives and the renderer will crash later
+			// FPrimitiveSceneProxy::VerifyUsedMaterial is used to make sure that all materials used for rendering are reported in GetUsedMaterials
+			PrimitiveComponent->GetUsedMaterials(UsedMaterials);
+
+			if (UsedMaterials.Num() > 0)
+			{
+				for (TArray<const FMaterial*>::TConstIterator MaterialIt(MaterialResourcesToUpdate); MaterialIt; ++MaterialIt)
+				{
+					UMaterialInterface* UpdatedMaterialInterface = (*MaterialIt)->GetMaterialInterface();
+
+					if (UpdatedMaterialInterface)
+					{
+						for (int32 MaterialIndex = 0; MaterialIndex < UsedMaterials.Num(); MaterialIndex++)
+						{
+							UMaterialInterface* TestMaterial = UsedMaterials[MaterialIndex];
+
+							if (TestMaterial && (TestMaterial == UpdatedMaterialInterface || TestMaterial->IsDependent(UpdatedMaterialInterface)))
+							{
+								bPrimitiveIsDependentOnMaterial = true;
+								break;
+							}
+						}
+					}
+				}
+
+				if (bPrimitiveIsDependentOnMaterial)
+				{
+					ENQUEUE_UNIQUE_RENDER_COMMAND_ONEPARAMETER(
+						FUpdateStaticMeshesForMaterials,
+						FPrimitiveSceneProxy*, SceneProxy, PrimitiveComponent->SceneProxy,
+						{
+							SceneProxy->GetPrimitiveSceneInfo()->UpdateStaticMeshes(RHICmdList);
+						});
+				}
+			}
+		}
+	}
+}
+
 void FRendererModule::UpdateStaticDrawListsForMaterials(const TArray<const FMaterial*>& Materials)
 {
 	for (TSet<FSceneInterface*>::TConstIterator SceneIt(AllocatedScenes); SceneIt; ++SceneIt)
 	{
 		(*SceneIt)->UpdateStaticDrawListsForMaterials(Materials);
 	}
+
+	// Update static meshes in order to recache cached mesh draw commands.
+	UpdateStaticMeshesForMaterials(Materials);
 }
 
 FSceneViewStateInterface* FRendererModule::AllocateViewState()
