You *must* use ELoadAction::ELoad to re-use *any* render-target data that was stored using EStoreAction::EStore or EStoreAction::EMultisampleResolve on Metal - ELoadAction::ENoAction is (in practice) a synonym for ELoadAction::EClear!

diff --git a/Engine/Source/Runtime/Renderer/Private/DeferredShadingRenderer.cpp b/Engine/Source/Runtime/Renderer/Private/DeferredShadingRenderer.cpp
--- a/Engine/Source/Runtime/Renderer/Private/DeferredShadingRenderer.cpp
+++ b/Engine/Source/Runtime/Renderer/Private/DeferredShadingRenderer.cpp
@@ -1040,7 +1040,7 @@ void FDeferredShadingSceneRenderer::Render(FRHICommandListImmediate& RHICmdList)
 		 
 		// if we didn't to the prepass above, then we will need to clear now, otherwise, it's already been cleared and rendered to
 		ERenderTargetLoadAction ColorLoadAction = ERenderTargetLoadAction::ELoad;
-		ERenderTargetLoadAction DepthLoadAction = bClearDepth ? ERenderTargetLoadAction::EClear : ERenderTargetLoadAction::ENoAction;
+		ERenderTargetLoadAction DepthLoadAction = bClearDepth ? ERenderTargetLoadAction::EClear : (!IsMetalPlatform(ShaderPlatform) ? ERenderTargetLoadAction::ENoAction : ERenderTargetLoadAction::ELoad);
 
 		const bool bClearBlack = ViewFamily.EngineShowFlags.ShaderComplexity || ViewFamily.EngineShowFlags.StationaryLightOverlap;
 		const float ClearAlpha = GetSceneColorClearAlpha();
