Removed ScenePrimitiveId from FVisibleMeshDrawCommand. ItÆs the same as DrawPrimitiveId for the cases when we want to sort draw commands (except draws with PrimitiveUniformBufferResource which have DrawPrimitiveId > Scene.Primitives.Num()).

diff --git a/Engine/Source/Runtime/Renderer/Private/MeshPassProcessor.cpp b/Engine/Source/Runtime/Renderer/Private/MeshPassProcessor.cpp
--- a/Engine/Source/Runtime/Renderer/Private/MeshPassProcessor.cpp
+++ b/Engine/Source/Runtime/Renderer/Private/MeshPassProcessor.cpp
@@ -818,9 +818,9 @@ void FMeshPassProcessor::SetDrawCommandEvent(const FPrimitiveSceneProxy* Primiti
 #endif
 }
 
-void FMeshPassProcessor::GetDrawCommandPrimitiveId(const FPrimitiveSceneInfo* RESTRICT PrimitiveSceneInfo, const FMeshBatchElement& BatchElement, int32& DrawPrimitiveId, int32& ScenePrimitiveId) const
+int32 FMeshPassProcessor::GetDrawCommandPrimitiveId(const FPrimitiveSceneInfo* RESTRICT PrimitiveSceneInfo, const FMeshBatchElement& BatchElement) const
 {
-	DrawPrimitiveId = 0;
+	int32 DrawPrimitiveId;
 
 	if (UseGPUScene(GMaxRHIShaderPlatform, FeatureLevel))
 	{
@@ -836,10 +836,15 @@ void FMeshPassProcessor::GetDrawCommandPrimitiveId(const FPrimitiveSceneInfo* RE
 		else
 		{
 			check(BatchElement.PrimitiveIdMode == PrimID_ForceZero);
+			DrawPrimitiveId = 0;
 		}
 	}
+	else
+	{
+		DrawPrimitiveId = PrimitiveSceneInfo ? PrimitiveSceneInfo->GetIndex() : Scene->Primitives.Num();
+	}
 
-	ScenePrimitiveId = PrimitiveSceneInfo ? PrimitiveSceneInfo->GetIndex() : -1;
+	return DrawPrimitiveId;
 }
 
 FCachedPassMeshDrawListContext::FCachedPassMeshDrawListContext(int32& InCommandIndex, FCachedPassMeshDrawList& InDrawList, FScene& InScene) :
@@ -863,7 +868,6 @@ void FCachedPassMeshDrawListContext::FinalizeCommand(
 	const FMeshBatch& MeshBatch, 
 	const FMeshBatchElement& BatchElement,
 	int32 DrawPrimitiveId,
-	int32 ScenePrimitiveId,
 	ERasterizerFillMode MeshFillMode,
 	ERasterizerCullMode MeshCullMode,
 	int32 InstanceFactor,
diff --git a/Engine/Source/Runtime/Renderer/Private/MobileBasePass.cpp b/Engine/Source/Runtime/Renderer/Private/MobileBasePass.cpp
--- a/Engine/Source/Runtime/Renderer/Private/MobileBasePass.cpp
+++ b/Engine/Source/Runtime/Renderer/Private/MobileBasePass.cpp
@@ -484,13 +484,12 @@ void MobileBasePass::ComputeBasePassSortKeys(const FScene& Scene, const FViewInf
 	for (int32 CmdIdx = 0; CmdIdx < NumCmds; ++CmdIdx)
 	{
 		FVisibleMeshDrawCommand& Cmd = VisibleMeshCommands[CmdIdx];
-		int32 PrimitiveIndex = Cmd.ScenePrimitiveId;
 
 		float PrimitiveDistance = 0;
 
-		if (PrimitiveIndex >= 0)
+		if (Cmd.DrawPrimitiveId < Scene.PrimitiveBounds.Num())
 		{
-			const FPrimitiveBounds& PrimitiveBounds = Scene.PrimitiveBounds[PrimitiveIndex];
+			const FPrimitiveBounds& PrimitiveBounds = Scene.PrimitiveBounds[Cmd.DrawPrimitiveId];
 			PrimitiveDistance = (PrimitiveBounds.BoxSphereBounds.Origin - ViewLocation).Size();
 		}
 		
@@ -503,17 +502,16 @@ void MobileBasePass::ComputeBasePassSortKeys(const FScene& Scene, const FViewInf
 	for (int32 CmdIdx = 0; CmdIdx < NumCmds; ++CmdIdx)
 	{
 		FVisibleMeshDrawCommand& Cmd = VisibleMeshCommands[CmdIdx];
-		int32 PrimitiveIndex = Cmd.ScenePrimitiveId;
 		float PrimitiveDistance = 0;
 		bool bMasked = false;
 		bool bBackground = false;
 
-		if (PrimitiveIndex >= 0)
+		if (Cmd.DrawPrimitiveId < Scene.PrimitiveBounds.Num())
 		{
-			const FPrimitiveBounds& PrimitiveBounds = Scene.PrimitiveBounds[PrimitiveIndex];
+			const FPrimitiveBounds& PrimitiveBounds = Scene.PrimitiveBounds[Cmd.DrawPrimitiveId];
 			PrimitiveDistance = (PrimitiveBounds.BoxSphereBounds.Origin - ViewLocation).Size();
 			bBackground = (PrimitiveBounds.BoxSphereBounds.SphereRadius > HALF_WORLD_MAX/4.0f); // TODO: per-primitive bBackgound flag
-			bMasked = View.PrimitiveViewRelevanceMap[PrimitiveIndex].bMaskedRelevance;
+			bMasked = View.PrimitiveViewRelevanceMap[Cmd.DrawPrimitiveId].bMaskedRelevance;
 		}
 
 		int32 PipelineId = Cmd.MeshDrawCommand->CachedPipelineId.GetId();
diff --git a/Engine/Source/Runtime/Renderer/Private/SceneRendering.cpp b/Engine/Source/Runtime/Renderer/Private/SceneRendering.cpp
--- a/Engine/Source/Runtime/Renderer/Private/SceneRendering.cpp
+++ b/Engine/Source/Runtime/Renderer/Private/SceneRendering.cpp
@@ -2698,7 +2698,6 @@ void FSceneRenderer::ApplyViewOverridesToMeshDrawCommands(FExclusiveDepthStencil
 						NewVisibleMeshDrawCommand.Setup(
 							&NewMeshCommand,
 							VisibleMeshDrawCommand.DrawPrimitiveId,
-							VisibleMeshDrawCommand.ScenePrimitiveId,
 							VisibleMeshDrawCommand.StateBucketId,
 							VisibleMeshDrawCommand.MeshFillMode,
 							VisibleMeshDrawCommand.MeshCullMode,
@@ -2816,7 +2815,6 @@ void SortPassMeshDrawCommands(
 						NewVisibleMeshDrawCommand.Setup(
 							&NewCommand, 
 							VisibleMeshDrawCommand.DrawPrimitiveId,
-							VisibleMeshDrawCommand.ScenePrimitiveId,
 							VisibleMeshDrawCommand.StateBucketId,
 							VisibleMeshDrawCommand.MeshFillMode,
 							VisibleMeshDrawCommand.MeshCullMode,
@@ -2878,8 +2876,11 @@ void UpdateTranslucentMeshSortKeys(FScene* Scene, FViewInfo& View, ETranslucency
 	{
 		FVisibleMeshDrawCommand& VisibleCommand = VisibleMeshCommands[CommandIndex];
 
-		const int32 PrimitiveIndex = VisibleCommand.ScenePrimitiveId;
-		const FVector BoundsOrigin = PrimitiveIndex >= 0 ? Scene->PrimitiveBounds[PrimitiveIndex].BoxSphereBounds.Origin : FVector::ZeroVector;
+		FVector BoundsOrigin = FVector::ZeroVector;
+		if (VisibleCommand.DrawPrimitiveId < Scene->PrimitiveBounds.Num())
+		{
+			BoundsOrigin = Scene->PrimitiveBounds[VisibleCommand.DrawPrimitiveId].BoxSphereBounds.Origin;
+		}
 
 		float Distance = 0.0f;
 		if (View.TranslucentSortPolicy == ETranslucentSortPolicy::SortByDistance)
diff --git a/Engine/Source/Runtime/Renderer/Private/SceneVisibility.cpp b/Engine/Source/Runtime/Renderer/Private/SceneVisibility.cpp
--- a/Engine/Source/Runtime/Renderer/Private/SceneVisibility.cpp
+++ b/Engine/Source/Runtime/Renderer/Private/SceneVisibility.cpp
@@ -1713,7 +1713,6 @@ struct FDrawCommandRelevancePacket
 				NewVisibleMeshDrawCommand.Setup(
 					&SceneDrawList.MeshDrawCommands[CommandIndex],
 					PrimitiveIndex,
-					PrimitiveIndex,
 					CachedMeshDrawCommand.StateBucketId,
 					CachedMeshDrawCommand.MeshFillMode,
 					CachedMeshDrawCommand.MeshCullMode,
diff --git a/Engine/Source/Runtime/Renderer/Private/ShadowSetup.cpp b/Engine/Source/Runtime/Renderer/Private/ShadowSetup.cpp
--- a/Engine/Source/Runtime/Renderer/Private/ShadowSetup.cpp
+++ b/Engine/Source/Runtime/Renderer/Private/ShadowSetup.cpp
@@ -819,7 +819,6 @@ void FProjectedShadowInfo::AddCachedMeshDrawCommandsForPass(
 			NewVisibleMeshDrawCommand.Setup(
 				&SceneDrawList.MeshDrawCommands[CommandIndex], 
 				PrimitiveIndex,
-				PrimitiveIndex,
 				CachedMeshDrawCommand.StateBucketId,
 				CachedMeshDrawCommand.MeshFillMode,
 				CachedMeshDrawCommand.MeshCullMode,
@@ -1377,7 +1376,6 @@ void FProjectedShadowInfo::ApplyViewOverridesToMeshDrawCommands(const FViewInfo&
 			NewVisibleMeshDrawCommand.Setup(
 				&NewMeshCommand,
 				VisibleMeshDrawCommand.DrawPrimitiveId,
-				VisibleMeshDrawCommand.ScenePrimitiveId,
 				VisibleMeshDrawCommand.StateBucketId,
 				VisibleMeshDrawCommand.MeshFillMode,
 				VisibleMeshDrawCommand.MeshCullMode,
diff --git a/Engine/Source/Runtime/Renderer/Private/ShadowSetupMobile.cpp b/Engine/Source/Runtime/Renderer/Private/ShadowSetupMobile.cpp
--- a/Engine/Source/Runtime/Renderer/Private/ShadowSetupMobile.cpp
+++ b/Engine/Source/Runtime/Renderer/Private/ShadowSetupMobile.cpp
@@ -352,9 +352,9 @@ void FMobileSceneRenderer::InitDynamicShadows(FRHICommandListImmediate& RHICmdLi
 				FVisibleMeshDrawCommand& MeshCommand = MeshCommands[i];
 				FVisibleMeshDrawCommand& MeshCommandCSM = MeshCommandsCSM[i];
 
-				if (bAlwaysUseCSM || (MeshCommand.ScenePrimitiveId >= 0 && MobileCSMVisibilityInfo.MobilePrimitiveCSMReceiverVisibilityMap[MeshCommand.ScenePrimitiveId]))
+				if (bAlwaysUseCSM || (MeshCommand.DrawPrimitiveId < Scene->Primitives.Num() && MobileCSMVisibilityInfo.MobilePrimitiveCSMReceiverVisibilityMap[MeshCommand.DrawPrimitiveId]))
 				{
-					checkf(MeshCommand.ScenePrimitiveId == MeshCommandCSM.ScenePrimitiveId, TEXT("VisibleMeshDrawCommands of BasePass and MobileBasePassCSM are expected to match."));
+					checkf(MeshCommand.DrawPrimitiveId == MeshCommandCSM.DrawPrimitiveId, TEXT("VisibleMeshDrawCommands of BasePass and MobileBasePassCSM are expected to match."));
 					// Use CSM's VisibleMeshDrawCommand.
 					MeshCommand = MeshCommandCSM;
 				}
diff --git a/Engine/Source/Runtime/Renderer/Public/MeshPassProcessor.h b/Engine/Source/Runtime/Renderer/Public/MeshPassProcessor.h
--- a/Engine/Source/Runtime/Renderer/Public/MeshPassProcessor.h
+++ b/Engine/Source/Runtime/Renderer/Public/MeshPassProcessor.h
@@ -391,7 +391,6 @@ public:
 		const FMeshBatch& MeshBatch, 
 		const FMeshBatchElement& BatchElement, 
 		int32 DrawPrimitiveId,
-		int32 ScenePrimitiveId,
 		ERasterizerFillMode MeshFillMode,
 		ERasterizerCullMode MeshCullMode,
 		int32 InstanceFactor,
@@ -417,11 +416,10 @@ public:
 
 	// Note: no ctor as TChunkedArray::CopyToLinearArray requires POD types
 
-	FORCEINLINE_DEBUGGABLE void Setup(const FMeshDrawCommand* InMeshDrawCommand, int32 InDrawPrimitiveIndex, int32 InScenePrimitiveIndex, int32 InStateBucketId, ERasterizerFillMode InMeshFillMode, ERasterizerCullMode InMeshCullMode, FMeshDrawCommandSortKey InSortKey)
+	FORCEINLINE_DEBUGGABLE void Setup(const FMeshDrawCommand* InMeshDrawCommand, int32 InDrawPrimitiveIndex, int32 InStateBucketId, ERasterizerFillMode InMeshFillMode, ERasterizerCullMode InMeshCullMode, FMeshDrawCommandSortKey InSortKey)
 	{
 		MeshDrawCommand = InMeshDrawCommand;
 		DrawPrimitiveId = InDrawPrimitiveIndex;
-		ScenePrimitiveId = InScenePrimitiveIndex;
 		PrimitiveIdBufferOffset = -1;
 		StateBucketId = InStateBucketId;
 		MeshFillMode = InMeshFillMode;
@@ -436,11 +434,9 @@ public:
 	FMeshDrawCommandSortKey SortKey;
 
 	// Draw PrimitiveId this draw command is associated with - used by the shader to fetch primitive data from the PrimitiveSceneData SRV.
+	// If it's < Scene->Primitives.Num() then it's a valid Scene PrimitiveIndex and can be used to backtrack to the FPrimitiveSceneInfo.
 	int32 DrawPrimitiveId;
 
-	// Scene PrimitiveIndex that generated this draw command, or -1 if no FPrimitiveSceneInfo.  Can be used to backtrack to the FPrimitiveSceneInfo.
-	int32 ScenePrimitiveId;
-
 	// Offset into the buffer of PrimitiveIds built for this pass, in int32's.
 	int32 PrimitiveIdBufferOffset;
 
@@ -484,7 +480,6 @@ public:
 		const FMeshBatch& MeshBatch, 
 		const FMeshBatchElement& BatchElement,
 		int32 DrawPrimitiveId,
-		int32 ScenePrimitiveId,
 		ERasterizerFillMode MeshFillMode,
 		ERasterizerCullMode MeshCullMode,
 		int32 InstanceFactor,
@@ -496,7 +491,7 @@ public:
 		FVisibleMeshDrawCommand NewVisibleMeshDrawCommand;
 		//@todo MeshCommandPipeline - assign usable state ID for dynamic path draws
 		// Currently dynamic path draws will not get dynamic instancing, but they will be roughly sorted by state
-		NewVisibleMeshDrawCommand.Setup(&MeshDrawCommand, DrawPrimitiveId, ScenePrimitiveId, -1, MeshFillMode, MeshCullMode, SortKey);
+		NewVisibleMeshDrawCommand.Setup(&MeshDrawCommand, DrawPrimitiveId, -1, MeshFillMode, MeshCullMode, SortKey);
 		DrawList.Add(NewVisibleMeshDrawCommand);
 	}
 
@@ -550,7 +545,6 @@ public:
 		const FMeshBatch& MeshBatch, 
 		const FMeshBatchElement& BatchElement, 
 		int32 DrawPrimitiveId,
-		int32 ScenePrimitiveId,
 		ERasterizerFillMode MeshFillMode,
 		ERasterizerCullMode MeshCullMode,
 		int32 InstanceFactor,
@@ -678,7 +672,7 @@ public:
 
 private:
 	RENDERER_API void SetDrawCommandEvent(const FPrimitiveSceneProxy* PrimitiveSceneProxy, const FMaterial& RESTRICT MaterialResource, FMeshDrawCommand& MeshDrawCommand) const;
-	RENDERER_API void GetDrawCommandPrimitiveId(const FPrimitiveSceneInfo* RESTRICT PrimitiveSceneInfo, const FMeshBatchElement& BatchElement, int32& DrawPrimitiveId, int32& ScenePrimitiveId) const;
+	RENDERER_API int32 GetDrawCommandPrimitiveId(const FPrimitiveSceneInfo* RESTRICT PrimitiveSceneInfo, const FMeshBatchElement& BatchElement) const;
 };
 
 typedef FMeshPassProcessor* (*PassProcessorCreateFunction)(const FScene* Scene, const FSceneView* InViewIfDynamicMeshCommand, FMeshPassDrawListContext& InDrawListContext);
diff --git a/Engine/Source/Runtime/Renderer/Public/MeshPassProcessor.inl b/Engine/Source/Runtime/Renderer/Public/MeshPassProcessor.inl
--- a/Engine/Source/Runtime/Renderer/Public/MeshPassProcessor.inl
+++ b/Engine/Source/Runtime/Renderer/Public/MeshPassProcessor.inl
@@ -126,11 +126,9 @@ void FMeshPassProcessor::BuildMeshDrawCommands(
 
 			SetDrawCommandEvent(PrimitiveSceneProxy, MaterialResource, MeshDrawCommand);
 
-			int32 DrawPrimitiveId;
-			int32 ScenePrimitiveId;
-			GetDrawCommandPrimitiveId(PrimitiveSceneInfo, BatchElement, DrawPrimitiveId, ScenePrimitiveId);
+			const int32 DrawPrimitiveId = GetDrawCommandPrimitiveId(PrimitiveSceneInfo, BatchElement);
 
-			DrawListContext.FinalizeCommand(MeshBatch, BatchElement, DrawPrimitiveId, ScenePrimitiveId, MeshFillMode, MeshCullMode, InstanceFactor, SortKey, MeshDrawCommand);
+			DrawListContext.FinalizeCommand(MeshBatch, BatchElement, DrawPrimitiveId, MeshFillMode, MeshCullMode, InstanceFactor, SortKey, MeshDrawCommand);
 		}
 	}
 }
