FrustumCull optimization - changed order of branches in order to skip most Scene->Primitives[Index]->Proxy->IsUsingDistanceCullFade() calls. This doubles the performance of FrustumCull function when profiling level_rooftops in RoboRecall.

diff --git a/Engine/Source/Runtime/Renderer/Private/SceneVisibility.cpp b/Engine/Source/Runtime/Renderer/Private/SceneVisibility.cpp
--- a/Engine/Source/Runtime/Renderer/Private/SceneVisibility.cpp
+++ b/Engine/Source/Runtime/Renderer/Private/SceneVisibility.cpp
@@ -460,19 +460,23 @@ static int32 FrustumCull(const FScene* Scene, FViewInfo& View)
 					}
 					else
 					{
-						const bool bUsesDistanceCullFade = Scene->Primitives[Index]->Proxy->IsUsingDistanceCullFade();
-
-						if (bUsesDistanceCullFade && DistanceSquared > FMath::Square(MaxDrawDistance))
+						if (DistanceSquared > FMath::Square(MaxDrawDistance))
 						{
-							FadingBits |= Mask;
+							if (Scene->Primitives[Index]->Proxy->IsUsingDistanceCullFade())
+							{
+								FadingBits |= Mask;
+							}
 						}
 						else
 						{
 							// The primitive is visible!
 							VisBits |= Mask;
-							if (bUsesDistanceCullFade && DistanceSquared > FMath::Square(MaxDrawDistance - FadeRadius))
+							if (DistanceSquared > FMath::Square(MaxDrawDistance - FadeRadius))
 							{
-								FadingBits |= Mask;
+								if (Scene->Primitives[Index]->Proxy->IsUsingDistanceCullFade())
+								{
+									FadingBits |= Mask;
+								}
 							}
 						}
 					}
