InitViews optimizations for cache efficiency:
-Replaced FPrimitiveSceneInfo::StaticMesh TIndirectArray with TArray (properly presized before adding meshes).
-Added StaticMeshRelevances array, which keeps InitView specific FStaticMesh data.
-Changed order of branches inside AddCommandsForMesh, so if mesh draw processor can't use cached mesh draw commands then it never deferences cached mesh draw command specific data.

diff --git a/Engine/Source/Runtime/Engine/Private/Components/TextRenderComponent.cpp b/Engine/Source/Runtime/Engine/Private/Components/TextRenderComponent.cpp
--- a/Engine/Source/Runtime/Engine/Private/Components/TextRenderComponent.cpp
+++ b/Engine/Source/Runtime/Engine/Private/Components/TextRenderComponent.cpp
@@ -755,6 +755,8 @@ void FTextRenderSceneProxy::DrawStaticElements(FStaticPrimitiveDrawInterface* PD
 	// Vertex factory will not been initialized when the font is invalid or the text string is empty.
 	if(VertexFactory.IsInitialized())
 	{
+		PDI->ReserveMemoryForMeshes(TextBatches.Num());
+
 		for (const FTextBatch& TextBatch : TextBatches)
 		{
 			// Draw the mesh.
diff --git a/Engine/Source/Runtime/Engine/Private/ModelRender.cpp b/Engine/Source/Runtime/Engine/Private/ModelRender.cpp
--- a/Engine/Source/Runtime/Engine/Private/ModelRender.cpp
+++ b/Engine/Source/Runtime/Engine/Private/ModelRender.cpp
@@ -513,6 +513,8 @@ public:
 			// Determine the DPG the primitive should be drawn in.
 			uint8 PrimitiveDPG = GetStaticDepthPriorityGroup();
 
+			PDI->ReserveMemoryForMeshes(Elements.Num());
+
 			for(int32 ElementIndex = 0;ElementIndex < Elements.Num();ElementIndex++)
 			{
 				const FModelElement& ModelElement = Component->GetElements()[ElementIndex];
diff --git a/Engine/Source/Runtime/Engine/Private/SceneManagement.cpp b/Engine/Source/Runtime/Engine/Private/SceneManagement.cpp
--- a/Engine/Source/Runtime/Engine/Private/SceneManagement.cpp
+++ b/Engine/Source/Runtime/Engine/Private/SceneManagement.cpp
@@ -260,8 +260,6 @@ void FMeshElementCollector::AddMesh(int32 ViewIndex, FMeshBatch& MeshBatch)
 			*this);
 	}
 
-	MeshBatch.bSupportsCachingMeshDrawCommands = SupportsCachingMeshDrawCommands(MeshBatch.VertexFactory, PrimitiveSceneProxy);
-
 	MeshBatch.PreparePrimitiveUniformBuffer(PrimitiveSceneProxy, FeatureLevel);
 
 	for (int32 Index = 0; Index < MeshBatch.Elements.Num(); ++Index)
@@ -538,12 +536,12 @@ int8 ComputeStaticMeshLOD( const FStaticMeshRenderData* RenderData, const FVecto
 	return MinLOD;
 }
 
-FLODMask ComputeLODForMeshes( const TIndirectArray<class FStaticMesh>& StaticMeshes, const FSceneView& View, const FVector4& Origin, float SphereRadius, int32 ForcedLODLevel, float& OutScreenRadiusSquared, float ScreenSizeScale)
+FLODMask ComputeLODForMeshes( const TArray<class FStaticMeshRelevance>& StaticMeshRelevances, const FSceneView& View, const FVector4& Origin, float SphereRadius, int32 ForcedLODLevel, float& OutScreenRadiusSquared, float ScreenSizeScale)
 {
 	FLODMask LODToRender;
 	const FSceneView& LODView = GetLODView(View);
 
-	const int32 NumMeshes = StaticMeshes.Num();
+	const int32 NumMeshes = StaticMeshRelevances.Num();
 
 	// Handle forced LOD level first
 	if (ForcedLODLevel >= 0)
@@ -551,9 +549,9 @@ FLODMask ComputeLODForMeshes( const TIndirectArray<class FStaticMesh>& StaticMes
 		OutScreenRadiusSquared = 0.0f;
 
 		int8 MinLOD = 127, MaxLOD = 0;
-		for (int32 MeshIndex = 0; MeshIndex < StaticMeshes.Num(); ++MeshIndex)
+		for (int32 MeshIndex = 0; MeshIndex < StaticMeshRelevances.Num(); ++MeshIndex)
 		{
-			const FStaticMesh&  Mesh = StaticMeshes[MeshIndex];
+			const FStaticMeshRelevance& Mesh = StaticMeshRelevances[MeshIndex];
 			MinLOD = FMath::Min(MinLOD, Mesh.LODIndex);
 			MaxLOD = FMath::Max(MaxLOD, Mesh.LODIndex);
 		}
@@ -561,7 +559,7 @@ FLODMask ComputeLODForMeshes( const TIndirectArray<class FStaticMesh>& StaticMes
 	}
 	else if (LODView.Family->EngineShowFlags.LOD && NumMeshes)
 	{
-		if (StaticMeshes[0].bDitheredLODTransition)
+		if (StaticMeshRelevances[0].bDitheredLODTransition)
 		{
 			for (int32 SampleIndex = 0; SampleIndex < 2; SampleIndex++)
 			{
@@ -571,7 +569,7 @@ FLODMask ComputeLODForMeshes( const TIndirectArray<class FStaticMesh>& StaticMes
 
 				for(int32 MeshIndex = NumMeshes-1 ; MeshIndex >= 0 ; --MeshIndex)
 				{
-					const FStaticMesh& Mesh = StaticMeshes[MeshIndex];
+					const FStaticMeshRelevance& Mesh = StaticMeshRelevances[MeshIndex];
 
 					float MeshScreenSize = Mesh.ScreenSize * ScreenSizeScale;
 
@@ -599,7 +597,7 @@ FLODMask ComputeLODForMeshes( const TIndirectArray<class FStaticMesh>& StaticMes
 
 			for(int32 MeshIndex = NumMeshes-1 ; MeshIndex >= 0 ; --MeshIndex)
 			{
-				const FStaticMesh& Mesh = StaticMeshes[MeshIndex];
+				const FStaticMeshRelevance& Mesh = StaticMeshRelevances[MeshIndex];
 
 				float MeshScreenSize = Mesh.ScreenSize * ScreenSizeScale;
 
diff --git a/Engine/Source/Runtime/Engine/Private/StaticMeshRender.cpp b/Engine/Source/Runtime/Engine/Private/StaticMeshRender.cpp
--- a/Engine/Source/Runtime/Engine/Private/StaticMeshRender.cpp
+++ b/Engine/Source/Runtime/Engine/Private/StaticMeshRender.cpp
@@ -792,6 +792,8 @@ void FStaticMeshSceneProxy::DrawStaticElements(FStaticPrimitiveDrawInterface* PD
 
 				const int32 NumBatches = GetNumMeshBatches();
 
+				PDI->ReserveMemoryForMeshes(NumBatches);
+
 				for (int32 BatchIndex = 0; BatchIndex < NumBatches; BatchIndex++)
 				{
 					FMeshBatch MeshBatch;
@@ -857,6 +859,8 @@ void FStaticMeshSceneProxy::DrawStaticElements(FStaticPrimitiveDrawInterface* PD
 						{
 							const int32 NumBatches = GetNumMeshBatches();
 
+							PDI->ReserveMemoryForMeshes(NumBatches);
+
 							for (int32 BatchIndex = 0; BatchIndex < NumBatches; BatchIndex++)
 							{
 								FMeshBatch MeshBatch;
@@ -891,6 +895,8 @@ void FStaticMeshSceneProxy::DrawStaticElements(FStaticPrimitiveDrawInterface* PD
 
 					const int32 NumBatches = GetNumMeshBatches();
 
+					PDI->ReserveMemoryForMeshes(NumBatches);
+
 					for (int32 BatchIndex = 0; BatchIndex < NumBatches; BatchIndex++)
 					{
 						FMeshBatch MeshBatch;
diff --git a/Engine/Source/Runtime/Engine/Public/MeshBatch.h b/Engine/Source/Runtime/Engine/Public/MeshBatch.h
--- a/Engine/Source/Runtime/Engine/Public/MeshBatch.h
+++ b/Engine/Source/Runtime/Engine/Public/MeshBatch.h
@@ -186,9 +186,7 @@ struct FMeshBatch
 	/** Whether the mesh batch should apply dithered LOD. */
 	uint32 bDitheredLODTransition : 1;
 
-	/** Cached from vertex factory to avoid dereferencing VF in InitViews. */
-	uint32 bSupportsCachingMeshDrawCommands : 1;
-
+	//@todo MeshCommandPipeline remove it together with mesh draw policies.
 	/** If bDitheredLODTransition and this is a dynamic mesh element, then this is the alpha for dither fade (static draw lists need to derive this later as it is changes every frame) */
 	float DitheredLODTransitionAlpha;
 
@@ -279,7 +277,6 @@ struct FMeshBatch
 	,	bSelectable(true)
 	,	bRequiresPerElementVisibility(false)
 	,	bDitheredLODTransition(false)
-	,	bSupportsCachingMeshDrawCommands(false)
 	,   DitheredLODTransitionAlpha(0.0f)
 	,	LCI(NULL)
 	,	VertexFactory(NULL)
diff --git a/Engine/Source/Runtime/Engine/Public/SceneManagement.h b/Engine/Source/Runtime/Engine/Public/SceneManagement.h
--- a/Engine/Source/Runtime/Engine/Public/SceneManagement.h
+++ b/Engine/Source/Runtime/Engine/Public/SceneManagement.h
@@ -1690,7 +1690,14 @@ class FStaticPrimitiveDrawInterface
 {
 public:
 	virtual ~FStaticPrimitiveDrawInterface() { }
+
 	virtual void SetHitProxy(HHitProxy* HitProxy) = 0;
+
+	/**
+	  * Reserve memory for specified number of meshes in order to minimize number of allocations inside DrawMesh.
+	  */
+	virtual void ReserveMemoryForMeshes(int32 MeshNum) = 0;
+
 	virtual void DrawMesh(
 		const FMeshBatch& Mesh,
 		float ScreenSize
@@ -2580,7 +2587,7 @@ struct FLODMask
 		return DitheredLODIndices[0] != DitheredLODIndices[1];
 	}
 };
-FLODMask ENGINE_API ComputeLODForMeshes(const TIndirectArray<class FStaticMesh>& StaticMeshes, const FSceneView& View, const FVector4& Origin, float SphereRadius, int32 ForcedLODLevel, float& OutScreenRadiusSquared, float ScreenSizeScale = 1.0f);
+FLODMask ENGINE_API ComputeLODForMeshes(const TArray<class FStaticMeshRelevance>& StaticMeshRelevances, const FSceneView& View, const FVector4& Origin, float SphereRadius, int32 ForcedLODLevel, float& OutScreenRadiusSquared, float ScreenSizeScale = 1.0f);
 
 class FSharedSamplerState : public FRenderResource
 {
diff --git a/Engine/Source/Runtime/Landscape/Private/LandscapeRender.cpp b/Engine/Source/Runtime/Landscape/Private/LandscapeRender.cpp
--- a/Engine/Source/Runtime/Landscape/Private/LandscapeRender.cpp
+++ b/Engine/Source/Runtime/Landscape/Private/LandscapeRender.cpp
@@ -1606,6 +1606,8 @@ void FLandscapeComponentSceneProxy::DrawStaticElements(FStaticPrimitiveDrawInter
 
 	int32 CurrentLODIndex = 0;
 
+	PDI->ReserveMemoryForMeshes(MaterialCount);
+
 	for (int32 i = 0; i < MaterialCount; ++i)
 	{
 		UMaterialInstance* MaterialInstance = Cast<UMaterialInstance>(AvailableMaterials[i]);
diff --git a/Engine/Source/Runtime/Renderer/Private/MeshPassProcessor.cpp b/Engine/Source/Runtime/Renderer/Private/MeshPassProcessor.cpp
--- a/Engine/Source/Runtime/Renderer/Private/MeshPassProcessor.cpp
+++ b/Engine/Source/Runtime/Renderer/Private/MeshPassProcessor.cpp
@@ -847,8 +847,8 @@ int32 FMeshPassProcessor::GetDrawCommandPrimitiveId(const FPrimitiveSceneInfo* R
 	return DrawPrimitiveId;
 }
 
-FCachedPassMeshDrawListContext::FCachedPassMeshDrawListContext(int32& InCommandIndex, FCachedPassMeshDrawList& InDrawList, FScene& InScene) :
-	CommandIndex(InCommandIndex),
+FCachedPassMeshDrawListContext::FCachedPassMeshDrawListContext(FCachedMeshDrawCommandInfo& InCommandInfo, FCachedPassMeshDrawList& InDrawList, FScene& InScene) :
+	CommandInfo(InCommandInfo),
 	DrawList(InDrawList),
 	Scene(InScene)
 {}
@@ -856,12 +856,9 @@ FCachedPassMeshDrawListContext::FCachedPassMeshDrawListContext(int32& InCommandI
 FMeshDrawCommand& FCachedPassMeshDrawListContext::AddCommand(const FMeshDrawCommand& Initializer)
 {
 	// Only one FMeshDrawCommand supported per FStaticMesh in a pass
-	check(CommandIndex == -1);
-	const int32 Index = DrawList.MeshDrawCommands.Add(Initializer);
-	const int32 SecondIndex = DrawList.MeshDrawCommandInfo.Add(FCachedMeshDrawCommandInfo());
-	check(Index == SecondIndex);
-	CommandIndex = Index;
-	return DrawList.MeshDrawCommands[Index];
+	check(CommandInfo.CommandIndex == -1);
+	CommandInfo.CommandIndex = DrawList.MeshDrawCommands.Add(Initializer);
+	return DrawList.MeshDrawCommands[CommandInfo.CommandIndex];
 }
 
 void FCachedPassMeshDrawListContext::FinalizeCommand(
@@ -878,8 +875,7 @@ void FCachedPassMeshDrawListContext::FinalizeCommand(
 
 	MeshDrawCommand.SetDrawParametersAndFinalize(MeshBatch, BatchElement, InstanceFactor);
 
-	check(CommandIndex != -1);
-	FCachedMeshDrawCommandInfo& CachedMeshDrawCommandInfo = DrawList.MeshDrawCommandInfo[CommandIndex];
+	check(CommandInfo.CommandIndex != -1);
 
 	FSetElementId SetId = Scene.CachedMeshDrawCommandStateBuckets.FindId(MeshDrawCommand);
 
@@ -892,10 +888,10 @@ void FCachedPassMeshDrawListContext::FinalizeCommand(
 		SetId = Scene.CachedMeshDrawCommandStateBuckets.Add(FMeshDrawCommandStateBucket(1, MeshDrawCommand));
 	}
 
-	CachedMeshDrawCommandInfo.StateBucketId = SetId.AsInteger();
-	CachedMeshDrawCommandInfo.SortKey = SortKey;
-	CachedMeshDrawCommandInfo.MeshFillMode = MeshFillMode;
-	CachedMeshDrawCommandInfo.MeshCullMode = MeshCullMode;
+	CommandInfo.StateBucketId = SetId.AsInteger();
+	CommandInfo.SortKey = SortKey;
+	CommandInfo.MeshFillMode = MeshFillMode;
+	CommandInfo.MeshCullMode = MeshCullMode;
 }
 
 PassProcessorCreateFunction FPassProcessorManager::JumpTable[(int32)EShadingPath::Num][EMeshPass::Num] = {};
diff --git a/Engine/Source/Runtime/Renderer/Private/PrimitiveSceneInfo.cpp b/Engine/Source/Runtime/Renderer/Private/PrimitiveSceneInfo.cpp
--- a/Engine/Source/Runtime/Renderer/Private/PrimitiveSceneInfo.cpp
+++ b/Engine/Source/Runtime/Renderer/Private/PrimitiveSceneInfo.cpp
@@ -39,6 +39,12 @@ public:
 		}
 	}
 
+	virtual void ReserveMemoryForMeshes(int32 MeshNum)
+	{
+		PrimitiveSceneInfo->StaticMeshRelevances.Reserve(PrimitiveSceneInfo->StaticMeshRelevances.Max() + MeshNum);
+		PrimitiveSceneInfo->StaticMeshes.Reserve(PrimitiveSceneInfo->StaticMeshes.Max() + MeshNum);
+	}
+
 	virtual void DrawMesh(const FMeshBatch& Mesh, float ScreenSize) final override
 	{
 		if (Mesh.GetNumPrimitives() > 0)
@@ -53,13 +59,18 @@ public:
 			FStaticMesh* StaticMesh = new(PrimitiveSceneInfo->StaticMeshes) FStaticMesh(
 				PrimitiveSceneInfo,
 				Mesh,
-				ScreenSize,
 				CurrentHitProxy ? CurrentHitProxy->Id : FHitProxyId()
 				);
 
-			StaticMesh->bSupportsCachingMeshDrawCommands = SupportsCachingMeshDrawCommands(StaticMesh->VertexFactory, PrimitiveSceneProxy);
-
 			StaticMesh->PreparePrimitiveUniformBuffer(PrimitiveSceneProxy, PrimitiveSceneInfo->Scene->GetFeatureLevel());
+
+			const bool bSupportsCachingMeshDrawCommands = SupportsCachingMeshDrawCommands(StaticMesh->VertexFactory, PrimitiveSceneProxy);
+
+			FStaticMeshRelevance* StaticMeshRelevance = new(PrimitiveSceneInfo->StaticMeshRelevances) FStaticMeshRelevance(
+				*StaticMesh, 
+				ScreenSize, 
+				bSupportsCachingMeshDrawCommands
+			);
 		}
 	}
 
@@ -154,15 +165,20 @@ void FPrimitiveSceneInfo::AddStaticMeshes(FRHICommandListImmediate& RHICmdList,
 	BatchingSPDI.SetHitProxy(DefaultDynamicHitProxy);
 	Proxy->DrawStaticElements(&BatchingSPDI);
 	StaticMeshes.Shrink();
+	StaticMeshRelevances.Shrink();
+
+	check(StaticMeshRelevances.Num() == StaticMeshes.Num());
 
 	for(int32 MeshIndex = 0;MeshIndex < StaticMeshes.Num();MeshIndex++)
 	{
+		FStaticMeshRelevance& MeshRelevance = StaticMeshRelevances[MeshIndex];
 		FStaticMesh& Mesh = StaticMeshes[MeshIndex];
 
 		// Add the static mesh to the scene's static mesh list.
 		FSparseArrayAllocationInfo SceneArrayAllocation = Scene->StaticMeshes.AddUninitialized();
 		Scene->StaticMeshes[SceneArrayAllocation.Index] = &Mesh;
 		Mesh.Id = SceneArrayAllocation.Index;
+		MeshRelevance.Id = SceneArrayAllocation.Index;
 
 		if (Mesh.bRequiresPerElementVisibility)
 		{
@@ -342,13 +358,14 @@ void FPrimitiveSceneInfo::AddToScene(FRHICommandListImmediate& RHICmdList, bool
 		}
 	}
 
-	INC_MEMORY_STAT_BY(STAT_PrimitiveInfoMemory, sizeof(*this) + StaticMeshes.GetAllocatedSize() + Proxy->GetMemoryFootprint());
+	INC_MEMORY_STAT_BY(STAT_PrimitiveInfoMemory, sizeof(*this) + StaticMeshes.GetAllocatedSize() + StaticMeshRelevances.GetAllocatedSize() + Proxy->GetMemoryFootprint());
 }
 
 void FPrimitiveSceneInfo::RemoveStaticMeshes()
 {
 	// Remove static meshes from the scene.
 	StaticMeshes.Empty();
+	StaticMeshRelevances.Empty();
 }
 
 void FPrimitiveSceneInfo::RemoveFromScene(bool bUpdateStaticDrawLists)
@@ -380,7 +397,7 @@ void FPrimitiveSceneInfo::RemoveFromScene(bool bUpdateStaticDrawLists)
 	IndirectLightingCacheAllocation = NULL;
 	ClearIndirectLightingCacheBuffer(false);
 
-	DEC_MEMORY_STAT_BY(STAT_PrimitiveInfoMemory, sizeof(*this) + StaticMeshes.GetAllocatedSize() + Proxy->GetMemoryFootprint());
+	DEC_MEMORY_STAT_BY(STAT_PrimitiveInfoMemory, sizeof(*this) + StaticMeshes.GetAllocatedSize() + StaticMeshRelevances.GetAllocatedSize() + Proxy->GetMemoryFootprint());
 
 	if (bNeedsStaticMeshUpdate)
 	{
@@ -615,7 +632,7 @@ FBoxSphereBounds FPrimitiveSceneInfo::GetAttachmentGroupBounds() const
 
 uint32 FPrimitiveSceneInfo::GetMemoryFootprint()
 {
-	return( sizeof( *this ) + HitProxies.GetAllocatedSize() + StaticMeshes.GetAllocatedSize() );
+	return( sizeof( *this ) + HitProxies.GetAllocatedSize() + StaticMeshes.GetAllocatedSize() + StaticMeshRelevances.GetAllocatedSize() );
 }
 
 bool FPrimitiveSceneInfo::ShouldRenderVelocity(const FViewInfo& View, bool bCheckVisibility) const
@@ -752,11 +769,11 @@ void FPrimitiveSceneInfo::GetStaticMeshesLODRange(int8& OutMinLOD, int8& OutMaxL
 	OutMinLOD = MAX_int8;
 	OutMaxLOD = 0;
 
-	for (int32 MeshIndex = 0; MeshIndex < StaticMeshes.Num(); ++MeshIndex)
+	for (int32 MeshIndex = 0; MeshIndex < StaticMeshRelevances.Num(); ++MeshIndex)
 	{
-		const FMeshBatch&  Mesh = StaticMeshes[MeshIndex];
-		OutMinLOD = FMath::Min(OutMinLOD, Mesh.LODIndex);
-		OutMaxLOD = FMath::Max(OutMaxLOD, Mesh.LODIndex);
+		const FStaticMeshRelevance& MeshRelevance = StaticMeshRelevances[MeshIndex];
+		OutMinLOD = FMath::Min(OutMinLOD, MeshRelevance.LODIndex);
+		OutMaxLOD = FMath::Max(OutMaxLOD, MeshRelevance.LODIndex);
 	}
 }
 
diff --git a/Engine/Source/Runtime/Renderer/Private/SceneCore.cpp b/Engine/Source/Runtime/Renderer/Private/SceneCore.cpp
--- a/Engine/Source/Runtime/Renderer/Private/SceneCore.cpp
+++ b/Engine/Source/Runtime/Renderer/Private/SceneCore.cpp
@@ -388,7 +388,7 @@ void FStaticMesh::CacheMeshDrawCommands(FScene* Scene)
 		if ((FPassProcessorManager::GetPassFlags(ShadingPath, PassType) & EMeshPassFlags::CachedMeshCommands) != EMeshPassFlags::None)
 		{
 			FCachedPassMeshDrawList& SceneDrawList = Scene->CachedDrawLists[PassType];
-			FCachedPassMeshDrawListContext CachedPassMeshDrawListContext(PassMeshCommandIndices[PassType], SceneDrawList, *Scene);
+			FCachedPassMeshDrawListContext CachedPassMeshDrawListContext(CachedMeshDrawCommands[PassType], SceneDrawList, *Scene);
 
 			PassProcessorCreateFunction CreateFunction = FPassProcessorManager::GetCreateFunction(ShadingPath, PassType);
 			FMeshPassProcessor* PassMeshProcessor = CreateFunction(Scene, nullptr, CachedPassMeshDrawListContext);
@@ -432,15 +432,14 @@ void FStaticMesh::RemoveFromDrawLists(bool bMeshIsBeingDestroyed)
 
 	FScene* Scene = PrimitiveSceneInfo->Scene;
 
-	for (int32 PassIndex = 0; PassIndex < ARRAY_COUNT(PassMeshCommandIndices); PassIndex++)
+	for (int32 PassIndex = 0; PassIndex < ARRAY_COUNT(CachedMeshDrawCommands); PassIndex++)
 	{
-		const int32 CommandIndex = PassMeshCommandIndices[PassIndex];
-		if (CommandIndex != -1)
+		const FCachedMeshDrawCommandInfo& CachedCommand = CachedMeshDrawCommands[PassIndex];
+		if (CachedCommand.CommandIndex != -1)
 		{
 			FCachedPassMeshDrawList& PassDrawList = Scene->CachedDrawLists[PassIndex];
 
-			FCachedMeshDrawCommandInfo& CachedMeshDrawCommandInfo = PassDrawList.MeshDrawCommandInfo[CommandIndex];
-			const FSetElementId StateBucketId = FSetElementId::FromInteger(CachedMeshDrawCommandInfo.StateBucketId);
+			const FSetElementId StateBucketId = FSetElementId::FromInteger(CachedCommand.StateBucketId);
 			checkSlow(StateBucketId.IsValidId());
 			FMeshDrawCommandStateBucket& StateBucket = Scene->CachedMeshDrawCommandStateBuckets[StateBucketId];
 
@@ -453,11 +452,10 @@ void FStaticMesh::RemoveFromDrawLists(bool bMeshIsBeingDestroyed)
 				StateBucket.Num--;
 			}
 		
-			PassDrawList.MeshDrawCommandInfo.RemoveAt(CommandIndex);
-			PassDrawList.MeshDrawCommands.RemoveAt(CommandIndex);
+			PassDrawList.MeshDrawCommands.RemoveAt(CachedCommand.CommandIndex);
 		}
 
-		PassMeshCommandIndices[PassIndex] = -1;
+		CachedMeshDrawCommands[PassIndex] = FCachedMeshDrawCommandInfo();
 	}
 }
 
diff --git a/Engine/Source/Runtime/Renderer/Private/SceneCore.h b/Engine/Source/Runtime/Renderer/Private/SceneCore.h
--- a/Engine/Source/Runtime/Renderer/Private/SceneCore.h
+++ b/Engine/Source/Runtime/Renderer/Private/SceneCore.h
@@ -140,9 +140,6 @@ public:
 		virtual void Remove(const bool bUnlinkMesh) = 0;
 	};
 
-	/** The screen space size to draw this primitive at */
-	float ScreenSize;
-
 	/** The render info for the primitive which created this mesh. */
 	FPrimitiveSceneInfo* PrimitiveSceneInfo;
 
@@ -152,28 +149,21 @@ public:
 	/** Index of the mesh into the scene's StaticMeshBatchVisibility array. */
 	int32 BatchVisibilityId;
 
-	/** Maps from this FStaticMesh to its FMeshDrawCommands cached in the scene. */
-	int32 PassMeshCommandIndices[EMeshPass::Num];
+	/** Cached mesh draw commands for this FStaticMesh. */
+	FCachedMeshDrawCommandInfo CachedMeshDrawCommands[EMeshPass::Num];
 
 	// Constructor/destructor.
 	FStaticMesh(
 		FPrimitiveSceneInfo* InPrimitiveSceneInfo,
 		const FMeshBatch& InMesh,
-		float InScreenSize,
 		FHitProxyId InHitProxyId
 		):
 		FMeshBatch(InMesh),
-		ScreenSize(InScreenSize),
 		PrimitiveSceneInfo(InPrimitiveSceneInfo),
 		Id(INDEX_NONE),
 		BatchVisibilityId(INDEX_NONE)
 	{
 		BatchHitProxyId = InHitProxyId;
-
-		for (int32 i = 0; i < ARRAY_COUNT(PassMeshCommandIndices); i++)
-		{
-			PassMeshCommandIndices[i] = -1;
-		}
 	}
 
 	~FStaticMesh();
@@ -199,13 +189,54 @@ private:
 	/** Private copy constructor. */
 	FStaticMesh(const FStaticMesh& InStaticMesh):
 		FMeshBatch(InStaticMesh),
-		ScreenSize(InStaticMesh.ScreenSize),
 		PrimitiveSceneInfo(InStaticMesh.PrimitiveSceneInfo),
 		Id(InStaticMesh.Id),
 		BatchVisibilityId(InStaticMesh.BatchVisibilityId)
 	{}
 };
 
+/**
+ * FStaticMesh data which is InitViews specific. Stored separately for cache efficiency.
+ */
+class FStaticMeshRelevance
+{
+public:
+	FStaticMeshRelevance(const FStaticMesh& StaticMesh, float InScreenSize, bool InbSupportsCachingMeshDrawCommands)
+		: Id(StaticMesh.Id)
+		, ScreenSize(InScreenSize)
+		, LODIndex(StaticMesh.LODIndex)
+		, bDitheredLODTransition(StaticMesh.bDitheredLODTransition)
+		, bRequiresPerElementVisibility(StaticMesh.bRequiresPerElementVisibility)
+		, CastShadow(StaticMesh.CastShadow)
+		, bUseForMaterial(StaticMesh.bUseForMaterial)
+		, bUseAsOccluder(StaticMesh.bUseAsOccluder)
+		, bSupportsCachingMeshDrawCommands(InbSupportsCachingMeshDrawCommands)
+	{
+	}
+
+	/** The index of the mesh in the scene's static meshes array. */
+	int32 Id;
+
+	/** The screen space size to draw this primitive at */
+	float ScreenSize;
+
+	/** LOD index of the mesh, used for fading LOD transitions. */
+	int8 LODIndex;
+
+	/** Whether the mesh batch should apply dithered LOD. */
+	uint32 bDitheredLODTransition : 1;
+
+	/** Whether the mesh batch needs VertexFactory->GetStaticBatchElementVisibility to be called each frame to determine which elements of the batch are visible. */
+	uint32 bRequiresPerElementVisibility : 1;
+
+	uint32 CastShadow : 1;			// Whether it can be used in shadow renderpasses.
+	uint32 bUseForMaterial : 1;		// Whether it can be used in renderpasses requiring material outputs.
+	uint32 bUseAsOccluder : 1;		// Whether it can be used in renderpasses only depending on the raw geometry (i.e. Depth Prepass).
+
+	/** Cached from vertex factory to avoid dereferencing VF in InitViews. */
+	uint32 bSupportsCachingMeshDrawCommands : 1;
+};
+
 /** The properties of a exponential height fog layer which are used for rendering. */
 class FExponentialHeightFogSceneInfo
 {
diff --git a/Engine/Source/Runtime/Renderer/Private/SceneVisibility.cpp b/Engine/Source/Runtime/Renderer/Private/SceneVisibility.cpp
--- a/Engine/Source/Runtime/Renderer/Private/SceneVisibility.cpp
+++ b/Engine/Source/Runtime/Renderer/Private/SceneVisibility.cpp
@@ -1699,35 +1699,37 @@ struct FDrawCommandRelevancePacket
 	bool bUseMeshDrawCommandPipeline;
 	bool bUseCachedMeshDrawCommands;
 
-	void AddCommandsForMesh(int32 PrimitiveIndex, const FStaticMesh& RESTRICT StaticMesh, const FScene* RESTRICT Scene, bool bCanCache, EMeshPass::Type PassType)
+	void AddCommandsForMesh(int32 PrimitiveIndex, const FStaticMeshRelevance& RESTRICT StaticMeshRelevance, const FStaticMesh& RESTRICT StaticMesh, const FScene* RESTRICT Scene, bool bCanCache, EMeshPass::Type PassType)
 	{
 		if (bUseMeshDrawCommandPipeline)
 		{
-			const bool bCacheMeshCommand = bUseCachedMeshDrawCommands && bCanCache;
-			const int32 CommandIndex = StaticMesh.PassMeshCommandIndices[PassType];
 			const EShadingPath ShadingPath = Scene->GetShadingPath();
+			const bool bUseCachedMeshCommand = bUseCachedMeshDrawCommands
+				&& !!(FPassProcessorManager::GetPassFlags(ShadingPath, PassType) & EMeshPassFlags::CachedMeshCommands)
+				&& StaticMeshRelevance.bSupportsCachingMeshDrawCommands
+				&& bCanCache;
 
-			if (bCacheMeshCommand && CommandIndex != -1)
+			if (bUseCachedMeshCommand)
 			{
-				const FCachedPassMeshDrawList& SceneDrawList = Scene->CachedDrawLists[PassType];
+				const FCachedMeshDrawCommandInfo& CachedMeshDrawCommand = StaticMesh.CachedMeshDrawCommands[PassType];
+				if (CachedMeshDrawCommand.CommandIndex != -1)
+				{
+					const FCachedPassMeshDrawList& SceneDrawList = Scene->CachedDrawLists[PassType];
 
-				const FCachedMeshDrawCommandInfo& CachedMeshDrawCommand = SceneDrawList.MeshDrawCommandInfo[CommandIndex];
-				FVisibleMeshDrawCommand NewVisibleMeshDrawCommand;
+					FVisibleMeshDrawCommand NewVisibleMeshDrawCommand;
 
-				NewVisibleMeshDrawCommand.Setup(
-					&SceneDrawList.MeshDrawCommands[CommandIndex],
-					PrimitiveIndex,
-					CachedMeshDrawCommand.StateBucketId,
-					CachedMeshDrawCommand.MeshFillMode,
-					CachedMeshDrawCommand.MeshCullMode,
-					CachedMeshDrawCommand.SortKey);
+					NewVisibleMeshDrawCommand.Setup(
+						&SceneDrawList.MeshDrawCommands[CachedMeshDrawCommand.CommandIndex],
+						PrimitiveIndex,
+						CachedMeshDrawCommand.StateBucketId,
+						CachedMeshDrawCommand.MeshFillMode,
+						CachedMeshDrawCommand.MeshCullMode,
+						CachedMeshDrawCommand.SortKey);
 
-				VisibleCachedDrawCommands[(uint32)PassType].AddElement(NewVisibleMeshDrawCommand);
+					VisibleCachedDrawCommands[(uint32)PassType].AddElement(NewVisibleMeshDrawCommand);
+				}
 			}
-			
-			if (!(FPassProcessorManager::GetPassFlags(ShadingPath, PassType) & EMeshPassFlags::CachedMeshCommands)
-				|| !StaticMesh.bSupportsCachingMeshDrawCommands
-				|| !bCacheMeshCommand)
+			else
 			{
 				DynamicBuildRequests[PassType].AddElement(&StaticMesh);
 			}
@@ -2017,7 +2019,7 @@ struct FRelevancePacket
 			}
 			else
 			{
-				LODToRender = ComputeLODForMeshes(PrimitiveSceneInfo->StaticMeshes, View, Bounds.BoxSphereBounds.Origin, Bounds.BoxSphereBounds.SphereRadius, ViewData.ForcedLODLevel, MeshScreenSizeSquared, ViewData.LODScale);
+				LODToRender = ComputeLODForMeshes(PrimitiveSceneInfo->StaticMeshRelevances, View, Bounds.BoxSphereBounds.Origin, Bounds.BoxSphereBounds.SphereRadius, ViewData.ForcedLODLevel, MeshScreenSizeSquared, ViewData.LODScale);
 			}
 
 			PrimitivesLODMask.AddPrim(FRelevancePacket::FPrimitiveLODMask(PrimitiveIndex, LODToRender));
@@ -2045,11 +2047,13 @@ struct FRelevancePacket
 
 			const bool bAddLightmapDensityCommands = View.Family->EngineShowFlags.LightMapDensity && AllowDebugViewmodes();
 
-			const int32 NumStaticMeshes = PrimitiveSceneInfo->StaticMeshes.Num();
+			const int32 NumStaticMeshes = PrimitiveSceneInfo->StaticMeshRelevances.Num();
 			for(int32 MeshIndex = 0;MeshIndex < NumStaticMeshes;MeshIndex++)
 			{
+				const FStaticMeshRelevance& StaticMeshRelevance = PrimitiveSceneInfo->StaticMeshRelevances[MeshIndex];
 				const FStaticMesh& StaticMesh = PrimitiveSceneInfo->StaticMeshes[MeshIndex];
-				if (LODToRender.ContainsLOD(StaticMesh.LODIndex))
+
+				if (LODToRender.ContainsLOD(StaticMeshRelevance.LODIndex))
 				{
 					uint8 MarkMask = 0;
 					bool bNeedsBatchVisibility = false;
@@ -2059,7 +2063,7 @@ struct FRelevancePacket
 					{
 						if (bIsHLODFadingOut)
 						{
-							if (bIsLODDithered && LODToRender.DitheredLODIndices[1] == StaticMesh.LODIndex)
+							if (bIsLODDithered && LODToRender.DitheredLODIndices[1] == StaticMeshRelevance.LODIndex)
 							{
 								bHiddenByHLODFade = true;
 							}
@@ -2070,7 +2074,7 @@ struct FRelevancePacket
 						}
 						else
 						{
-							if (bIsLODDithered && LODToRender.DitheredLODIndices[0] == StaticMesh.LODIndex)
+							if (bIsLODDithered && LODToRender.DitheredLODIndices[0] == StaticMeshRelevance.LODIndex)
 							{
 								bHiddenByHLODFade = true;
 							}
@@ -2082,7 +2086,7 @@ struct FRelevancePacket
 					}
 					else if (bIsLODDithered)
 					{
-						if (LODToRender.DitheredLODIndices[0] == StaticMesh.LODIndex)
+						if (LODToRender.DitheredLODIndices[0] == StaticMeshRelevance.LODIndex)
 						{
 							MarkMask |= EMarkMaskBits::StaticMeshFadeOutDitheredLODMapMask;
 						}
@@ -2093,57 +2097,57 @@ struct FRelevancePacket
 					}
 
 					// Don't cache if it requires per view per mesh state for LOD dithering or distance cull fade.
-					const bool bIsMeshDitheringLOD = StaticMesh.bDitheredLODTransition && (MarkMask & (EMarkMaskBits::StaticMeshFadeOutDitheredLODMapMask | EMarkMaskBits::StaticMeshFadeInDitheredLODMapMask));
+					const bool bIsMeshDitheringLOD = StaticMeshRelevance.bDitheredLODTransition && (MarkMask & (EMarkMaskBits::StaticMeshFadeOutDitheredLODMapMask | EMarkMaskBits::StaticMeshFadeInDitheredLODMapMask));
 					const bool bCanCache = !bIsPrimitiveDistanceCullFading && !bIsMeshDitheringLOD;
 
-					if (ViewRelevance.bShadowRelevance && bDrawShadowDepth && StaticMesh.CastShadow)
+					if (ViewRelevance.bShadowRelevance && bDrawShadowDepth && StaticMeshRelevance.CastShadow)
 					{
 						// Mark static mesh as visible in shadows.
 						MarkMask |= EMarkMaskBits::StaticMeshShadowDepthMapMask;
 						bNeedsBatchVisibility = true;
 					}
 
-					if(ViewRelevance.bDrawRelevance && (StaticMesh.bUseForMaterial || StaticMesh.bUseAsOccluder) && (ViewRelevance.bRenderInMainPass || ViewRelevance.bRenderCustomDepth) && !bHiddenByHLODFade)
+					if(ViewRelevance.bDrawRelevance && (StaticMeshRelevance.bUseForMaterial || StaticMeshRelevance.bUseAsOccluder) && (ViewRelevance.bRenderInMainPass || ViewRelevance.bRenderCustomDepth) && !bHiddenByHLODFade)
 					{
 						// Mark static mesh as visible for rendering
-						if (StaticMesh.bUseForMaterial)
+						if (StaticMeshRelevance.bUseForMaterial)
 						{
-							DrawCommandPacket.AddCommandsForMesh(PrimitiveIndex, StaticMesh, Scene, bCanCache, EMeshPass::BasePass);
+							DrawCommandPacket.AddCommandsForMesh(PrimitiveIndex, StaticMeshRelevance, StaticMesh, Scene, bCanCache, EMeshPass::BasePass);
 							MarkMask |= EMarkMaskBits::StaticMeshVisibilityMapMask;
 
 							if (ShadingPath == EShadingPath::Mobile)
 							{
-								DrawCommandPacket.AddCommandsForMesh(PrimitiveIndex, StaticMesh, Scene, bCanCache, EMeshPass::MobileBasePassCSM);
+								DrawCommandPacket.AddCommandsForMesh(PrimitiveIndex, StaticMeshRelevance, StaticMesh, Scene, bCanCache, EMeshPass::MobileBasePassCSM);
 							}
 
 							if (ViewRelevance.bRenderCustomDepth)
 							{
-								DrawCommandPacket.AddCommandsForMesh(PrimitiveIndex, StaticMesh, Scene, bCanCache, EMeshPass::CustomDepth);
+								DrawCommandPacket.AddCommandsForMesh(PrimitiveIndex, StaticMeshRelevance, StaticMesh, Scene, bCanCache, EMeshPass::CustomDepth);
 							}
 
 							if (bAddLightmapDensityCommands)
 							{
-								DrawCommandPacket.AddCommandsForMesh(PrimitiveIndex, StaticMesh, Scene, bCanCache, EMeshPass::LightmapDensity);
+								DrawCommandPacket.AddCommandsForMesh(PrimitiveIndex, StaticMeshRelevance, StaticMesh, Scene, bCanCache, EMeshPass::LightmapDensity);
 							}
 #if !(UE_BUILD_SHIPPING || UE_BUILD_TEST)
 							else if (View.Family->UseDebugViewPS())
 							{
-								DrawCommandPacket.AddCommandsForMesh(PrimitiveIndex, StaticMesh, Scene, bCanCache, EMeshPass::DebugViewMode);
+								DrawCommandPacket.AddCommandsForMesh(PrimitiveIndex, StaticMeshRelevance, StaticMesh, Scene, bCanCache, EMeshPass::DebugViewMode);
 							}
 #endif
 							if (StaticMesh.bSelectable)
 							{
-								DrawCommandPacket.AddCommandsForMesh(PrimitiveIndex, StaticMesh, Scene, bCanCache, EMeshPass::HitProxy);
+								DrawCommandPacket.AddCommandsForMesh(PrimitiveIndex, StaticMeshRelevance, StaticMesh, Scene, bCanCache, EMeshPass::HitProxy);
 
 								if (!ViewRelevance.HasTranslucency())
 								{
-									DrawCommandPacket.AddCommandsForMesh(PrimitiveIndex, StaticMesh, Scene, bCanCache, EMeshPass::HitProxyOpaqueOnly);
+									DrawCommandPacket.AddCommandsForMesh(PrimitiveIndex, StaticMeshRelevance, StaticMesh, Scene, bCanCache, EMeshPass::HitProxyOpaqueOnly);
 								}
 							}
 
 							if (PrimitiveSceneInfo->ShouldRenderVelocity(View, false))
 							{
-								DrawCommandPacket.AddCommandsForMesh(PrimitiveIndex, StaticMesh, Scene, bCanCache, EMeshPass::Velocity);
+								DrawCommandPacket.AddCommandsForMesh(PrimitiveIndex, StaticMeshRelevance, StaticMesh, Scene, bCanCache, EMeshPass::Velocity);
 								MarkMask |= EMarkMaskBits::StaticMeshVelocityMapMask;
 							}
 
@@ -2151,15 +2155,15 @@ struct FRelevancePacket
 						}
 
 						// If the static mesh is an occluder, check whether it covers enough of the screen to be used as an occluder.
-						if(	StaticMesh.bUseAsOccluder && bDrawDepthOnly )
+						if (StaticMeshRelevance.bUseAsOccluder && bDrawDepthOnly)
 						{
-							DrawCommandPacket.AddCommandsForMesh(PrimitiveIndex, StaticMesh, Scene, bCanCache, EMeshPass::DepthPass);
+							DrawCommandPacket.AddCommandsForMesh(PrimitiveIndex, StaticMeshRelevance, StaticMesh, Scene, bCanCache, EMeshPass::DepthPass);
 							MarkMask |= EMarkMaskBits::StaticMeshOccluderMapMask;
 						}
 						bNeedsBatchVisibility = true;
 					}
 
-					if (StaticMesh.bUseForMaterial 
+					if (StaticMeshRelevance.bUseForMaterial
 						&& ViewRelevance.HasTranslucency() 
 						&& !ViewRelevance.bEditorPrimitiveRelevance
 						&& ViewRelevance.bRenderInMainPass)
@@ -2168,48 +2172,48 @@ struct FRelevancePacket
 						{
 							if (ViewRelevance.bNormalTranslucencyRelevance)
 							{
-								DrawCommandPacket.AddCommandsForMesh(PrimitiveIndex, StaticMesh, Scene, bCanCache, EMeshPass::TranslucencyStandard);
+								DrawCommandPacket.AddCommandsForMesh(PrimitiveIndex, StaticMeshRelevance, StaticMesh, Scene, bCanCache, EMeshPass::TranslucencyStandard);
 							}
 
 							if (ViewRelevance.bSeparateTranslucencyRelevance)
 							{
-								DrawCommandPacket.AddCommandsForMesh(PrimitiveIndex, StaticMesh, Scene, bCanCache, EMeshPass::TranslucencyAfterDOF);
+								DrawCommandPacket.AddCommandsForMesh(PrimitiveIndex, StaticMeshRelevance, StaticMesh, Scene, bCanCache, EMeshPass::TranslucencyAfterDOF);
 							}
 						}
 						else
 						{
 							// Otherwise, everything is rendered in a single bucket. This is not related to whether DOF is currently enabled or not.
 							// When using all translucency, Standard and AfterDOF are sorted together instead of being rendered like 2 buckets.
-							DrawCommandPacket.AddCommandsForMesh(PrimitiveIndex, StaticMesh, Scene, bCanCache, EMeshPass::TranslucencyAll);
+							DrawCommandPacket.AddCommandsForMesh(PrimitiveIndex, StaticMeshRelevance, StaticMesh, Scene, bCanCache, EMeshPass::TranslucencyAll);
 						}
 
 						if (ViewRelevance.bDistortionRelevance)
 						{
-							DrawCommandPacket.AddCommandsForMesh(PrimitiveIndex, StaticMesh, Scene, bCanCache, EMeshPass::Distortion);
+							DrawCommandPacket.AddCommandsForMesh(PrimitiveIndex, StaticMeshRelevance, StaticMesh, Scene, bCanCache, EMeshPass::Distortion);
 						}
 
 						if (ShadingPath == EShadingPath::Mobile && View.bIsSceneCapture)
 						{
-							DrawCommandPacket.AddCommandsForMesh(PrimitiveIndex, StaticMesh, Scene, bCanCache, EMeshPass::MobileInverseOpacity);
+							DrawCommandPacket.AddCommandsForMesh(PrimitiveIndex, StaticMeshRelevance, StaticMesh, Scene, bCanCache, EMeshPass::MobileInverseOpacity);
 						}
 					}
 
 #if WITH_EDITOR
 					if(ViewRelevance.bDrawRelevance && ViewRelevance.bEditorStaticSelectionRelevance)
 					{
-						DrawCommandPacket.AddCommandsForMesh(PrimitiveIndex, StaticMesh, Scene, bCanCache, EMeshPass::EditorSelection);
+						DrawCommandPacket.AddCommandsForMesh(PrimitiveIndex, StaticMeshRelevance, StaticMesh, Scene, bCanCache, EMeshPass::EditorSelection);
 
 						MarkMask |= EMarkMaskBits::StaticMeshEditorSelectedMask;
 					}
 #endif
 					if (MarkMask)
 					{
-						MarkMasks[StaticMesh.Id] = MarkMask;
+						MarkMasks[StaticMeshRelevance.Id] = MarkMask;
 					}
 
 					// Static meshes which don't need per-element visibility always draw all elements
-					if (bNeedsBatchVisibility && StaticMesh.bRequiresPerElementVisibility)
-					{						
+					if (bNeedsBatchVisibility && StaticMeshRelevance.bRequiresPerElementVisibility)
+					{
 						WriteView.StaticMeshBatchVisibility[StaticMesh.BatchVisibilityId] = StaticMesh.VertexFactory->GetStaticBatchElementVisibility(View, &StaticMesh, UserViewCustomData);
 					}
 				}
@@ -2998,6 +3002,8 @@ static TAutoConsoleVariable<int32> CVarAlsoUseSphereForFrustumCull(
 
 void UpdateReflectionSceneData(FScene* Scene)
 {
+	QUICK_SCOPE_CYCLE_COUNTER(STAT_UpdateReflectionSceneData)
+
 	FReflectionEnvironmentSceneData& ReflectionSceneData = Scene->ReflectionSceneData;
 
 	ReflectionSceneData.SortedCaptures.Reset(ReflectionSceneData.RegisteredReflectionCaptures.Num());
@@ -4129,10 +4135,10 @@ void FLODSceneTree::UpdateVisibilityStates(FViewInfo& View)
 			}
 
 			FHLODSceneNodeVisibilityState& NodeVisibility = VisibilityStates.FindOrAdd(SceneInfo->PrimitiveComponentId);
-			const TIndirectArray<FStaticMesh>& NodeMeshes = SceneInfo->StaticMeshes;
+			const TArray<FStaticMeshRelevance>& NodeMeshRelevances = SceneInfo->StaticMeshRelevances;
 
 			// Ignore already updated nodes, or those that we can't work with
-			if (NodeVisibility.UpdateCount == UpdateCount || !NodeMeshes.IsValidIndex(0))
+			if (NodeVisibility.UpdateCount == UpdateCount || !NodeMeshRelevances.IsValidIndex(0))
 			{
 				continue;
 			}
@@ -4153,7 +4159,7 @@ void FLODSceneTree::UpdateVisibilityStates(FViewInfo& View)
 			const bool bIsInDrawRange = DistanceSquared >= Bounds.MinDrawDistanceSq * HLODState.FOVDistanceScaleSq;
 
 			const bool bWasFadingPreUpdate = !!NodeVisibility.bIsFading;
-			const bool bIsDitheredTransition = NodeMeshes[0].bDitheredLODTransition;
+			const bool bIsDitheredTransition = NodeMeshRelevances[0].bDitheredLODTransition;
 
 			if (bIsDitheredTransition && !bForcedIntoView)
 			{		
diff --git a/Engine/Source/Runtime/Renderer/Private/ShadowRendering.h b/Engine/Source/Runtime/Renderer/Private/ShadowRendering.h
--- a/Engine/Source/Runtime/Renderer/Private/ShadowRendering.h
+++ b/Engine/Source/Runtime/Renderer/Private/ShadowRendering.h
@@ -763,6 +763,7 @@ private:
 
 	void AddCachedMeshDrawCommandsForPass(
 		int32 PrimitiveIndex,
+		const FStaticMeshRelevance& RESTRICT StaticMeshRelevance,
 		const FStaticMesh& StaticMesh,
 		const FScene* Scene,
 		EMeshPass::Type PassType,
diff --git a/Engine/Source/Runtime/Renderer/Private/ShadowSetup.cpp b/Engine/Source/Runtime/Renderer/Private/ShadowSetup.cpp
--- a/Engine/Source/Runtime/Renderer/Private/ShadowSetup.cpp
+++ b/Engine/Source/Runtime/Renderer/Private/ShadowSetup.cpp
@@ -798,6 +798,7 @@ void FProjectedShadowInfo::SetupWholeSceneProjection(
 
 void FProjectedShadowInfo::AddCachedMeshDrawCommandsForPass(
 	int32 PrimitiveIndex, 
+	const FStaticMeshRelevance& RESTRICT StaticMeshRelevance,
 	const FStaticMesh& RESTRICT StaticMesh, 
 	const FScene* RESTRICT Scene, 
 	EMeshPass::Type PassType,
@@ -806,32 +807,32 @@ void FProjectedShadowInfo::AddCachedMeshDrawCommandsForPass(
 {
 	if (UseMeshDrawCommandPipeline())
 	{
-		const bool bUseCachedMeshDrawCommands = UseCachedMeshDrawCommands();
-		const int32 CommandIndex = StaticMesh.PassMeshCommandIndices[PassType];
+		EShadingPath ShadingPath = Scene->GetShadingPath();
+		const bool bUseCachedMeshCommand = UseCachedMeshDrawCommands()
+			&& !!(FPassProcessorManager::GetPassFlags(ShadingPath, PassType) & EMeshPassFlags::CachedMeshCommands)
+			&& StaticMeshRelevance.bSupportsCachingMeshDrawCommands;
 
-		if (CommandIndex != -1 && bUseCachedMeshDrawCommands)
+		if (bUseCachedMeshCommand)
 		{
-			const FCachedPassMeshDrawList& SceneDrawList = Scene->CachedDrawLists[PassType];
+			const FCachedMeshDrawCommandInfo& CachedMeshDrawCommand = StaticMesh.CachedMeshDrawCommands[PassType];
+			if (CachedMeshDrawCommand.CommandIndex != -1)
+			{
+				const FCachedPassMeshDrawList& SceneDrawList = Scene->CachedDrawLists[PassType];
 
-			const FCachedMeshDrawCommandInfo& CachedMeshDrawCommand = SceneDrawList.MeshDrawCommandInfo[CommandIndex];
-			FVisibleMeshDrawCommand NewVisibleMeshDrawCommand;
+				FVisibleMeshDrawCommand NewVisibleMeshDrawCommand;
 
-			NewVisibleMeshDrawCommand.Setup(
-				&SceneDrawList.MeshDrawCommands[CommandIndex], 
-				PrimitiveIndex,
-				CachedMeshDrawCommand.StateBucketId,
-				CachedMeshDrawCommand.MeshFillMode,
-				CachedMeshDrawCommand.MeshCullMode,
-				CachedMeshDrawCommand.SortKey);
-
-			VisibleMeshCommands.Add(NewVisibleMeshDrawCommand);
-		}
+				NewVisibleMeshDrawCommand.Setup(
+					&SceneDrawList.MeshDrawCommands[CachedMeshDrawCommand.CommandIndex],
+					PrimitiveIndex,
+					CachedMeshDrawCommand.StateBucketId,
+					CachedMeshDrawCommand.MeshFillMode,
+					CachedMeshDrawCommand.MeshCullMode,
+					CachedMeshDrawCommand.SortKey);
 
-		EShadingPath ShadingPath = Scene->GetShadingPath();
-			
-		if (!(FPassProcessorManager::GetPassFlags(ShadingPath, PassType) & EMeshPassFlags::CachedMeshCommands)
-			|| !StaticMesh.bSupportsCachingMeshDrawCommands
-			|| !bUseCachedMeshDrawCommands)
+				VisibleMeshCommands.Add(NewVisibleMeshDrawCommand);
+			}
+		}
+		else
 		{
 			MeshCommandBuildRequests.Add(&StaticMesh);
 		}
@@ -860,7 +861,7 @@ bool FProjectedShadowInfo::ShouldDrawStaticMeshes(FViewInfo& InCurrentView, bool
 			else
 			{
 				const FBoxSphereBounds& Bounds = InPrimitiveSceneInfo->Proxy->GetBounds();
-				ViewLODToRender = ComputeLODForMeshes(InPrimitiveSceneInfo->StaticMeshes, InCurrentView, Bounds.Origin, Bounds.SphereRadius, ForcedLOD, MeshScreenSizeSquared, InCurrentView.LODDistanceFactor);
+				ViewLODToRender = ComputeLODForMeshes(InPrimitiveSceneInfo->StaticMeshRelevances, InCurrentView, Bounds.Origin, Bounds.SphereRadius, ForcedLOD, MeshScreenSizeSquared, InCurrentView.LODDistanceFactor);
 			}	
 
 			InCurrentView.PrimitivesLODMask[PrimitiveId] = ViewLODToRender;
@@ -886,9 +887,9 @@ bool FProjectedShadowInfo::ShouldDrawStaticMeshes(FViewInfo& InCurrentView, bool
 			int8 LODToRenderScan = -MAX_int8;
 			FLODMask LODToRender;
 
-			for (int32 Index = 0; Index < InPrimitiveSceneInfo->StaticMeshes.Num(); Index++)
+			for (int32 Index = 0; Index < InPrimitiveSceneInfo->StaticMeshRelevances.Num(); Index++)
 			{
-				LODToRenderScan = FMath::Max<int8>(InPrimitiveSceneInfo->StaticMeshes[Index].LODIndex, LODToRenderScan);
+				LODToRenderScan = FMath::Max<int8>(InPrimitiveSceneInfo->StaticMeshRelevances[Index].LODIndex, LODToRenderScan);
 			}
 			if (LODToRenderScan != -MAX_int8)
 			{
@@ -909,10 +910,12 @@ bool FProjectedShadowInfo::ShouldDrawStaticMeshes(FViewInfo& InCurrentView, bool
 			const bool bIsPrimitiveDistanceCullFading = InCurrentView.PotentiallyFadingPrimitiveMap[InPrimitiveSceneInfo->GetIndex()];
 			const bool bCanCache = !bIsPrimitiveDistanceCullFading;
 
-			for (int32 MeshIndex = 0; MeshIndex < InPrimitiveSceneInfo->StaticMeshes.Num(); MeshIndex++)
+			for (int32 MeshIndex = 0; MeshIndex < InPrimitiveSceneInfo->StaticMeshRelevances.Num(); MeshIndex++)
 			{
+				const FStaticMeshRelevance& StaticMeshRelevance = InPrimitiveSceneInfo->StaticMeshRelevances[MeshIndex];
 				const FStaticMesh& StaticMesh = InPrimitiveSceneInfo->StaticMeshes[MeshIndex];
-				if ((StaticMesh.CastShadow || (bSelfShadowOnly && StaticMesh.bUseAsOccluder)) && ShadowLODToRender.ContainsLOD(StaticMesh.LODIndex))
+
+				if ((StaticMeshRelevance.CastShadow || (bSelfShadowOnly && StaticMeshRelevance.bUseAsOccluder)) && ShadowLODToRender.ContainsLOD(StaticMeshRelevance.LODIndex))
 				{
 					if (UseMeshDrawCommandPipeline())
 					{
@@ -920,6 +923,7 @@ bool FProjectedShadowInfo::ShouldDrawStaticMeshes(FViewInfo& InCurrentView, bool
 						{
 							AddCachedMeshDrawCommandsForPass(
 								PrimitiveId,
+								StaticMeshRelevance,
 								StaticMesh,
 								InPrimitiveSceneInfo->Scene,
 								EMeshPass::CSMShadowDepth,
@@ -931,12 +935,12 @@ bool FProjectedShadowInfo::ShouldDrawStaticMeshes(FViewInfo& InCurrentView, bool
 							SubjectMeshCommandBuildRequests.Add(&StaticMesh);
 						}
 					}
-					else if (StaticMesh.CastShadow)
+					else if (StaticMeshRelevance.CastShadow)
 					{
-						StaticMeshWholeSceneShadowDepthMap[StaticMesh.Id] = true;
+						StaticMeshWholeSceneShadowDepthMap[StaticMeshRelevance.Id] = true;
 					}
 
-					if (StaticMesh.bRequiresPerElementVisibility)
+					if (StaticMeshRelevance.bRequiresPerElementVisibility)
 					{
 						StaticMeshWholeSceneShadowBatchVisibility[StaticMesh.BatchVisibilityId] = StaticMesh.VertexFactory->GetStaticBatchElementVisibility(InCurrentView, &StaticMesh);
 					}
@@ -947,21 +951,23 @@ bool FProjectedShadowInfo::ShouldDrawStaticMeshes(FViewInfo& InCurrentView, bool
 		}
 		else
 		{
-			for (int32 MeshIndex = 0; MeshIndex < InPrimitiveSceneInfo->StaticMeshes.Num(); MeshIndex++)
+			for (int32 MeshIndex = 0; MeshIndex < InPrimitiveSceneInfo->StaticMeshRelevances.Num(); MeshIndex++)
 			{
+				const FStaticMeshRelevance& StaticMeshRelevance = InPrimitiveSceneInfo->StaticMeshRelevances[MeshIndex];
 				const FStaticMesh& StaticMesh = InPrimitiveSceneInfo->StaticMeshes[MeshIndex];
-				if ((StaticMesh.CastShadow || (bSelfShadowOnly && StaticMesh.bUseAsOccluder)) && ShadowLODToRender.ContainsLOD(StaticMesh.LODIndex))
+
+				if ((StaticMeshRelevance.CastShadow || (bSelfShadowOnly && StaticMeshRelevance.bUseAsOccluder)) && ShadowLODToRender.ContainsLOD(StaticMeshRelevance.LODIndex))
 				{
 					if (UseMeshDrawCommandPipeline())
 					{
 						SubjectMeshCommandBuildRequests.Add(&StaticMesh);
 					}
-					else if (StaticMesh.CastShadow)
+					else if (StaticMeshRelevance.CastShadow)
 					{
-						InCurrentView.StaticMeshShadowDepthMap[StaticMesh.Id] = true;
+						InCurrentView.StaticMeshShadowDepthMap[StaticMeshRelevance.Id] = true;
 					}
 
-					if (StaticMesh.bRequiresPerElementVisibility)
+					if (StaticMeshRelevance.bRequiresPerElementVisibility)
 					{
 						InCurrentView.StaticMeshBatchVisibility[StaticMesh.BatchVisibilityId] = StaticMesh.VertexFactory->GetStaticBatchElementVisibility(InCurrentView, &StaticMesh);
 					}
diff --git a/Engine/Source/Runtime/Renderer/Public/MeshPassProcessor.h b/Engine/Source/Runtime/Renderer/Public/MeshPassProcessor.h
--- a/Engine/Source/Runtime/Renderer/Public/MeshPassProcessor.h
+++ b/Engine/Source/Runtime/Renderer/Public/MeshPassProcessor.h
@@ -509,12 +509,15 @@ class FCachedMeshDrawCommandInfo
 public:
 	explicit FCachedMeshDrawCommandInfo() :
 		SortKey(FMeshDrawCommandSortKey::Default),
+		CommandIndex(-1),
 		StateBucketId(-1),
 		MeshFillMode(ERasterizerFillMode_Num),
 		MeshCullMode(ERasterizerCullMode_Num)
 	{}
 
 	FMeshDrawCommandSortKey SortKey;
+
+	int32 CommandIndex;
 	int32 StateBucketId;
 
 	// Needed for view overrides
@@ -525,10 +528,7 @@ public:
 class FCachedPassMeshDrawList
 {
 public:
-	/** Indices held by FStaticMesh::DrawCommandIndices must be stable */
-	TSparseArray<FCachedMeshDrawCommandInfo> MeshDrawCommandInfo;
-
-	/** Indices held by FStaticMesh::DrawCommandIndices must be stable */
+	/** Indices held by FStaticMesh::CachedMeshDrawCommands must be stable */
 	TSparseArray<FMeshDrawCommand> MeshDrawCommands;
 };
 
@@ -537,7 +537,7 @@ typedef TArray<int32, TInlineAllocator<5>> FDrawCommandIndices;
 class FCachedPassMeshDrawListContext : public FMeshPassDrawListContext
 {
 public:
-	FCachedPassMeshDrawListContext(int32& InCommandIndex, FCachedPassMeshDrawList& InDrawList, FScene& InScene);
+	FCachedPassMeshDrawListContext(FCachedMeshDrawCommandInfo& InCommandInfo, FCachedPassMeshDrawList& InDrawList, FScene& InScene);
 
 	virtual FMeshDrawCommand& AddCommand(const FMeshDrawCommand& Initializer) override final;
 
@@ -552,7 +552,7 @@ public:
 		FMeshDrawCommand& MeshDrawCommand) override final;
 
 private:
-	int32& CommandIndex;
+	FCachedMeshDrawCommandInfo& CommandInfo;
 	FCachedPassMeshDrawList& DrawList;
 	FScene& Scene;
 };
diff --git a/Engine/Source/Runtime/Renderer/Public/PrimitiveSceneInfo.h b/Engine/Source/Runtime/Renderer/Public/PrimitiveSceneInfo.h
--- a/Engine/Source/Runtime/Renderer/Public/PrimitiveSceneInfo.h
+++ b/Engine/Source/Runtime/Renderer/Public/PrimitiveSceneInfo.h
@@ -211,8 +211,11 @@ public:
 	 */
 	FPrimitiveComponentId LODParentComponentId;
 
+	/** The primitive's static mesh relevances. Must be in sync with StaticMeshes. */
+	TArray<class FStaticMeshRelevance> StaticMeshRelevances;
+
 	/** The primitive's static meshes. */
-	TIndirectArray<class FStaticMesh> StaticMeshes;
+	TArray<class FStaticMesh> StaticMeshes;
 
 	/** The identifier for the primitive in Scene->PrimitiveOctree. */
 	FOctreeElementId OctreeId;
