Fixed transitions in UploadDynamicPrimitiveShaderDataForView.  Fixes some flickering in dx12 when repeatedly deleting / undoing.

diff --git a/Engine/Source/Runtime/Renderer/Private/ByteBuffer.cpp b/Engine/Source/Runtime/Renderer/Private/ByteBuffer.cpp
--- a/Engine/Source/Runtime/Renderer/Private/ByteBuffer.cpp
+++ b/Engine/Source/Runtime/Renderer/Private/ByteBuffer.cpp
@@ -149,7 +149,7 @@ void MemcpyBuffer(FRHICommandList& RHICmdList, const FRWBufferStructured& SrcBuf
 	SetUAVParameter( RHICmdList, ShaderRHI, ComputeShader->DstBuffer, FUnorderedAccessViewRHIRef() );
 }
 
-void ResizeBufferIfNeeded(FRHICommandList& RHICmdList, FRWBufferStructured& Buffer, uint32 NumFloat4s)
+bool ResizeBufferIfNeeded(FRHICommandList& RHICmdList, FRWBufferStructured& Buffer, uint32 NumFloat4s)
 {
 	EPixelFormat BufferFormat = PF_A32B32G32R32F;
 	uint32 BytesPerElement = GPixelFormats[BufferFormat].BlockBytes;
@@ -168,7 +168,10 @@ void ResizeBufferIfNeeded(FRHICommandList& RHICmdList, FRWBufferStructured& Buff
 		MemcpyBuffer(RHICmdList, Buffer, NewBuffer, CopyFloat4s);
 
 		Buffer = NewBuffer;
+		return true;
 	}
+
+	return false;
 }
 
 class FScatterCopyCS : public FGlobalShader
diff --git a/Engine/Source/Runtime/Renderer/Private/ByteBuffer.h b/Engine/Source/Runtime/Renderer/Private/ByteBuffer.h
--- a/Engine/Source/Runtime/Renderer/Private/ByteBuffer.h
+++ b/Engine/Source/Runtime/Renderer/Private/ByteBuffer.h
@@ -7,7 +7,7 @@
 
 void MemsetBuffer(FRHICommandList& RHICmdList, const FRWBufferStructured& DstBuffer, const FVector4& Value, uint32 NumFloat4s, uint32 DstOffsetInFloat4s);
 void MemcpyBuffer(FRHICommandList& RHICmdList, const FRWBufferStructured& SrcBuffer, const FRWBufferStructured& DstBuffer, uint32 NumFloat4s, uint32 SrcOffset = 0, uint32 DstOffset = 0);
-void ResizeBufferIfNeeded(FRHICommandList& RHICmdList, FRWBufferStructured& Buffer, uint32 NumFloat4s);
+bool ResizeBufferIfNeeded(FRHICommandList& RHICmdList, FRWBufferStructured& Buffer, uint32 NumFloat4s);
 
 class FScatterUploadBuilder
 {
diff --git a/Engine/Source/Runtime/Renderer/Private/GPUScene.cpp b/Engine/Source/Runtime/Renderer/Private/GPUScene.cpp
--- a/Engine/Source/Runtime/Renderer/Private/GPUScene.cpp
+++ b/Engine/Source/Runtime/Renderer/Private/GPUScene.cpp
@@ -167,17 +167,19 @@ void UpdateGPUScene(FRHICommandList& RHICmdList, FScene& Scene)
 			Scene.GPUScene.bUpdateAllPrimitives = false;
 		}
 
+		bool bResizedPrimitiveData = false;
+		bool bResizedLightmapData = false;
 		{
 			const int32 NumPrimitiveEntries = Scene.Primitives.Num();
 			const uint32 PrimitiveSceneNumFloat4s = NumPrimitiveEntries * FPrimitiveSceneShaderData::PrimitiveDataStrideInFloat4s;
 			// Reserve enough space
-			ResizeBufferIfNeeded(RHICmdList, Scene.GPUScene.PrimitiveBuffer, FMath::RoundUpToPowerOfTwo(PrimitiveSceneNumFloat4s));
+			bResizedPrimitiveData = ResizeBufferIfNeeded(RHICmdList, Scene.GPUScene.PrimitiveBuffer, FMath::RoundUpToPowerOfTwo(PrimitiveSceneNumFloat4s));
 		}
 		
 		{
 			const int32 NumLightmapDataEntries = Scene.GPUScene.LightmapDataAllocator.GetMaxSize();
 			const uint32 LightmapDataNumFloat4s = NumLightmapDataEntries * FLightmapSceneShaderData::LightmapDataStrideInFloat4s;
-			ResizeBufferIfNeeded(RHICmdList, Scene.GPUScene.LightmapDataBuffer, FMath::RoundUpToPowerOfTwo(LightmapDataNumFloat4s));
+			bResizedLightmapData = ResizeBufferIfNeeded(RHICmdList, Scene.GPUScene.LightmapDataBuffer, FMath::RoundUpToPowerOfTwo(LightmapDataNumFloat4s));
 		}
 
 		const int32 NumPrimitiveDataUploads = Scene.GPUScene.PrimitivesToUpdate.Num();
@@ -209,7 +211,14 @@ void UpdateGPUScene(FRHICommandList& RHICmdList, FScene& Scene)
 				}
 			}
 
-			RHICmdList.TransitionResource(EResourceTransitionAccess::EWritable, EResourceTransitionPipeline::EGfxToCompute, Scene.GPUScene.PrimitiveBuffer.UAV);
+			if (bResizedPrimitiveData)
+			{
+				RHICmdList.TransitionResource(EResourceTransitionAccess::ERWBarrier, EResourceTransitionPipeline::EComputeToCompute, Scene.GPUScene.PrimitiveBuffer.UAV);
+			}
+			else
+			{
+				RHICmdList.TransitionResource(EResourceTransitionAccess::EWritable, EResourceTransitionPipeline::EGfxToCompute, Scene.GPUScene.PrimitiveBuffer.UAV);
+			}
 
 			PrimitivesUploadBuilder.UploadTo(RHICmdList, Scene.GPUScene.PrimitiveBuffer);
 
@@ -254,7 +263,14 @@ void UpdateGPUScene(FRHICommandList& RHICmdList, FScene& Scene)
 					}
 				}
 
-				RHICmdList.TransitionResource(EResourceTransitionAccess::EWritable, EResourceTransitionPipeline::EGfxToCompute, Scene.GPUScene.LightmapDataBuffer.UAV);
+				if (bResizedLightmapData)
+				{
+					RHICmdList.TransitionResource(EResourceTransitionAccess::ERWBarrier, EResourceTransitionPipeline::EComputeToCompute, Scene.GPUScene.LightmapDataBuffer.UAV);
+				}
+				else
+				{
+					RHICmdList.TransitionResource(EResourceTransitionAccess::EWritable, EResourceTransitionPipeline::EGfxToCompute, Scene.GPUScene.LightmapDataBuffer.UAV);
+				}
 
 				LightmapDataUploadBuilder.UploadTo(RHICmdList, Scene.GPUScene.LightmapDataBuffer);
 
@@ -289,8 +305,15 @@ void UploadDynamicPrimitiveShaderDataForView(FRHICommandList& RHICmdList, FScene
 		const int32 NumPrimitiveEntries = Scene.Primitives.Num() + View.DynamicPrimitiveShaderData.Num();
 		const uint32 PrimitiveSceneNumFloat4s = NumPrimitiveEntries * FPrimitiveSceneShaderData::PrimitiveDataStrideInFloat4s;
 
+		uint32 ViewPrimitiveSceneNumFloat4s = FMath::RoundUpToPowerOfTwo(PrimitiveSceneNumFloat4s);
+		uint32 BytesPerElement = GPixelFormats[PF_A32B32G32R32F].BlockBytes;
+
 		// Reserve enough space
-		ResizeBufferIfNeeded(RHICmdList, ViewPrimitiveShaderDataBuffer, FMath::RoundUpToPowerOfTwo(PrimitiveSceneNumFloat4s));
+		if (ViewPrimitiveSceneNumFloat4s * BytesPerElement != ViewPrimitiveShaderDataBuffer.NumBytes)
+		{
+			ViewPrimitiveShaderDataBuffer.Release();
+			ViewPrimitiveShaderDataBuffer.Initialize(BytesPerElement, ViewPrimitiveSceneNumFloat4s);
+		}
 
 		// Copy scene primitive data into view primitive data
 		MemcpyBuffer(RHICmdList, Scene.GPUScene.PrimitiveBuffer, ViewPrimitiveShaderDataBuffer, Scene.Primitives.Num() * FPrimitiveSceneShaderData::PrimitiveDataStrideInFloat4s);
@@ -309,18 +332,18 @@ void UploadDynamicPrimitiveShaderDataForView(FRHICommandList& RHICmdList, FScene
 				PrimitivesUploadBuilder.Add(Scene.Primitives.Num() + DynamicUploadIndex, &PrimitiveSceneData.Data[0]);
 			}
 
-			RHICmdList.TransitionResource(EResourceTransitionAccess::EWritable, EResourceTransitionPipeline::EGfxToCompute, ViewPrimitiveShaderDataBuffer.UAV);
+			RHICmdList.TransitionResource(EResourceTransitionAccess::ERWBarrier, EResourceTransitionPipeline::EComputeToCompute, ViewPrimitiveShaderDataBuffer.UAV);
 
 			PrimitivesUploadBuilder.UploadTo(RHICmdList, ViewPrimitiveShaderDataBuffer);
 		}
 
+		RHICmdList.TransitionResource(EResourceTransitionAccess::EReadable, EResourceTransitionPipeline::EComputeToGfx, ViewPrimitiveShaderDataBuffer.UAV);
+
 		// Update view uniform buffer
 		View.CachedViewUniformShaderParameters->LightmapSceneData = Scene.GPUScene.LightmapDataBuffer.SRV;
 		View.CachedViewUniformShaderParameters->PrimitiveSceneData = ViewPrimitiveShaderDataBuffer.SRV;
 		View.ViewUniformBuffer.UpdateUniformBufferImmediate(*View.CachedViewUniformShaderParameters);
 
-		RHICmdList.TransitionResource(EResourceTransitionAccess::EReadable, EResourceTransitionPipeline::EComputeToGfx, ViewPrimitiveShaderDataBuffer.UAV);
-
 		{
 			QUICK_SCOPE_CYCLE_COUNTER(STAT_PrimitiveIdBufferEmulation);
 
