DR - Fix for crash on ActionRPG
#rnx

diff --git a/Engine/Source/Runtime/Renderer/Private/DepthRendering.cpp b/Engine/Source/Runtime/Renderer/Private/DepthRendering.cpp
--- a/Engine/Source/Runtime/Renderer/Private/DepthRendering.cpp
+++ b/Engine/Source/Runtime/Renderer/Private/DepthRendering.cpp
@@ -1518,7 +1518,7 @@ void FDepthPassMeshProcessor::AddMeshBatch(const FMeshBatch& RESTRICT MeshBatch,
 	{
 		bDraw = MeshBatch.bUseAsOccluder || EarlyZPassMode == DDM_AllOpaque;
 
-		if (EarlyZPassMode < DDM_AllOpaque)
+		if (EarlyZPassMode < DDM_AllOpaque && PrimitiveSceneProxy)
 		{
 			// Only render primitives marked as occluders.
 			bDraw = bDraw && PrimitiveSceneProxy->ShouldUseAsOccluder()
