Fixed: Decals are rendering incorrectly on Nexus devices with HLR

diff --git a/Engine/Source/Runtime/Renderer/Private/MobileShadingRenderer.cpp b/Engine/Source/Runtime/Renderer/Private/MobileShadingRenderer.cpp
--- a/Engine/Source/Runtime/Renderer/Private/MobileShadingRenderer.cpp
+++ b/Engine/Source/Runtime/Renderer/Private/MobileShadingRenderer.cpp
@@ -702,14 +702,12 @@ void FMobileSceneRenderer::ConditionalResolveSceneDepth(FRHICommandListImmediate
 					// Switch target to force hardware flush current depth to texture
 					FTextureRHIRef DummySceneColor = GSystemTextures.BlackDummy->GetRenderTargetItem().TargetableTexture;
 					FTextureRHIRef DummyDepthTarget = GSystemTextures.DepthDummy->GetRenderTargetItem().TargetableTexture;
-
-					if (CVarMobileForceDepthResolve.GetValueOnRenderThread() != 0)
-					{
-						FRHIRenderPassInfo RPInfo(DummySceneColor, ERenderTargetActions::DontLoad_DontStore);
-						RPInfo.DepthStencilRenderTarget.Action = EDepthStencilTargetActions::ClearDepthStencil_StoreDepthStencil;
-						RPInfo.DepthStencilRenderTarget.DepthStencilTarget = DummyDepthTarget;
-						RPInfo.DepthStencilRenderTarget.ExclusiveDepthStencil = FExclusiveDepthStencil::DepthWrite_StencilWrite;
-						RHICmdList.BeginRenderPass(RPInfo, TEXT("ResolveDepth"));
+					
+					FRHIRenderPassInfo RPInfo(DummySceneColor, ERenderTargetActions::DontLoad_DontStore);
+					RPInfo.DepthStencilRenderTarget.Action = EDepthStencilTargetActions::ClearDepthStencil_StoreDepthStencil;
+					RPInfo.DepthStencilRenderTarget.DepthStencilTarget = DummyDepthTarget;
+					RPInfo.DepthStencilRenderTarget.ExclusiveDepthStencil = FExclusiveDepthStencil::DepthWrite_StencilWrite;
+					RHICmdList.BeginRenderPass(RPInfo, TEXT("ResolveDepth"));
 					{
 						FGraphicsPipelineStateInitializer GraphicsPSOInit;
 						RHICmdList.ApplyCachedRenderTargets(GraphicsPSOInit);
@@ -744,9 +742,8 @@ void FMobileSceneRenderer::ConditionalResolveSceneDepth(FRHICommandListImmediate
 							FIntPoint(1, 1),
 							*ScreenVertexShader,
 							EDRF_UseTriangleOptimization);
-					}
-						RHICmdList.EndRenderPass();
 					} // force depth resolve
+					RHICmdList.EndRenderPass();
 				}
 			}
 		}
