Fix viewport for editor primitives rendered during base pass, needed for views that don't render the full viewport (sub windows like mesh editor)

diff --git a/Engine/Source/Runtime/Renderer/Private/BasePassRendering.cpp b/Engine/Source/Runtime/Renderer/Private/BasePassRendering.cpp
--- a/Engine/Source/Runtime/Renderer/Private/BasePassRendering.cpp
+++ b/Engine/Source/Runtime/Renderer/Private/BasePassRendering.cpp
@@ -846,6 +846,9 @@ void FDeferredShadingSceneRenderer::RenderEditorPrimitives(FRHICommandListImmedi
 			FSceneRenderTargets& SceneContext = FSceneRenderTargets::Get(RHICmdList);
 			SceneContext.BeginRenderingGBuffer(RHICmdList, ERenderTargetLoadAction::ELoad, ERenderTargetLoadAction::ELoad, FExclusiveDepthStencil::DepthWrite_StencilWrite, false);
 
+			// need to setup view again after reconfiguring render targets
+			SetupBasePassView(RHICmdList, View, this, true);
+
 			FDrawingPolicyRenderState NoDepthTestDrawRenderState(DrawRenderState);
 			NoDepthTestDrawRenderState.SetDepthStencilState(TStaticDepthStencilState<true, CF_Always>::GetRHI());
 			NoDepthTestDrawRenderState.SetDepthStencilAccess(FExclusiveDepthStencil::DepthWrite_StencilWrite);
@@ -855,6 +858,7 @@ void FDeferredShadingSceneRenderer::RenderEditorPrimitives(FRHICommandListImmedi
 
 			// Restore default base pass depth access
 			SceneContext.BeginRenderingGBuffer(RHICmdList, ERenderTargetLoadAction::ELoad, ERenderTargetLoadAction::ELoad, BasePassDepthStencilAccess, false);
+			SetupBasePassView(RHICmdList, View, this, true);
 		}
 
 		// Render foreground primitives with depth testing
