Fixed redundant cached depth pass mesh draw commands.

diff --git a/Engine/Source/Runtime/Renderer/Private/DepthRendering.cpp b/Engine/Source/Runtime/Renderer/Private/DepthRendering.cpp
--- a/Engine/Source/Runtime/Renderer/Private/DepthRendering.cpp
+++ b/Engine/Source/Runtime/Renderer/Private/DepthRendering.cpp
@@ -1512,26 +1512,21 @@ void FDepthPassMeshProcessor::Process(
 
 void FDepthPassMeshProcessor::AddMeshBatch(const FMeshBatch& RESTRICT MeshBatch, uint64 BatchElementMask, const FPrimitiveSceneProxy* RESTRICT PrimitiveSceneProxy, int32 StaticMeshId)
 {
-	bool bDraw = true;
+	bool bDraw = MeshBatch.bUseAsOccluder;
 
-	if (bEarlyZPass)
+	if (bRespectUseAsOccluderFlag && EarlyZPassMode < DDM_AllOpaque && PrimitiveSceneProxy)
 	{
-		bDraw = MeshBatch.bUseAsOccluder || EarlyZPassMode == DDM_AllOpaque;
+		// Only render primitives marked as occluders.
+		bDraw = bDraw && PrimitiveSceneProxy->ShouldUseAsOccluder()
+			// Only render static objects unless movable are requested.
+			&& (!PrimitiveSceneProxy->IsMovable() || bEarlyZPassMovable);
 
-		if (EarlyZPassMode < DDM_AllOpaque && PrimitiveSceneProxy)
+		// Filter dynamic mesh commands by screen size.
+		if (ViewIfDynamicMeshCommand)
 		{
-			// Only render primitives marked as occluders.
-			bDraw = bDraw && PrimitiveSceneProxy->ShouldUseAsOccluder()
-				// Only render static objects unless movable are requested.
-				&& (!PrimitiveSceneProxy->IsMovable() || bEarlyZPassMovable);
-
-			// Filter dynamic mesh commands by screen size.
-			if (ViewIfDynamicMeshCommand)
-			{
-				extern float GMinScreenRadiusForDepthPrepass;
-				const float LODFactorDistanceSquared = (PrimitiveSceneProxy->GetBounds().Origin - ViewIfDynamicMeshCommand->ViewMatrices.GetViewOrigin()).SizeSquared() * FMath::Square(ViewIfDynamicMeshCommand->LODDistanceFactor);
-				bDraw = bDraw && FMath::Square(PrimitiveSceneProxy->GetBounds().SphereRadius) > GMinScreenRadiusForDepthPrepass * GMinScreenRadiusForDepthPrepass * LODFactorDistanceSquared;
-			}
+			extern float GMinScreenRadiusForDepthPrepass;
+			const float LODFactorDistanceSquared = (PrimitiveSceneProxy->GetBounds().Origin - ViewIfDynamicMeshCommand->ViewMatrices.GetViewOrigin()).SizeSquared() * FMath::Square(ViewIfDynamicMeshCommand->LODDistanceFactor);
+			bDraw = bDraw && FMath::Square(PrimitiveSceneProxy->GetBounds().SphereRadius) > GMinScreenRadiusForDepthPrepass * GMinScreenRadiusForDepthPrepass * LODFactorDistanceSquared;
 		}
 	}
 
@@ -1587,12 +1582,12 @@ void FDepthPassMeshProcessor::AddMeshBatch(const FMeshBatch& RESTRICT MeshBatch,
 FDepthPassMeshProcessor::FDepthPassMeshProcessor(const FScene* Scene,
 	const FSceneView* InViewIfDynamicMeshCommand,
 	const FDrawingPolicyRenderState& InPassDrawRenderState,
-	const bool InbEarlyZPass,
+	const bool InbRespectUseAsOccluderFlag,
 	const EDepthDrawingMode InEarlyZPassMode,
 	const bool InbEarlyZPassMovable,
 	FMeshPassDrawListContext& InDrawListContext)
 	: FMeshPassProcessor(Scene, Scene->GetFeatureLevel(), InViewIfDynamicMeshCommand, InDrawListContext)
-	, bEarlyZPass(InbEarlyZPass)
+	, bRespectUseAsOccluderFlag(InbRespectUseAsOccluderFlag)
 	, EarlyZPassMode(InEarlyZPassMode)
 	, bEarlyZPassMovable(InbEarlyZPassMovable)
 {
diff --git a/Engine/Source/Runtime/Renderer/Private/DepthRendering.h b/Engine/Source/Runtime/Renderer/Private/DepthRendering.h
--- a/Engine/Source/Runtime/Renderer/Private/DepthRendering.h
+++ b/Engine/Source/Runtime/Renderer/Private/DepthRendering.h
@@ -499,7 +499,7 @@ public:
 	FDepthPassMeshProcessor(const FScene* Scene, 
 		const FSceneView* InViewIfDynamicMeshCommand, 
 		const FDrawingPolicyRenderState& InPassDrawRenderState, 
-		const bool InbEarlyZPass,
+		const bool InbRespectUseAsOccluderFlag,
 		const EDepthDrawingMode InEarlyZPassMode,
 		const bool InbEarlyZPassMovable,
 		FMeshPassDrawListContext& InDrawListContext);
@@ -521,7 +521,7 @@ private:
 
 	FDrawingPolicyRenderState PassDrawRenderState;
 
-	const bool bEarlyZPass;
+	const bool bRespectUseAsOccluderFlag;
 	const EDepthDrawingMode EarlyZPassMode;
 	const bool bEarlyZPassMovable;
 };
