Fix incredibly spammy clang warnings.

diff --git a/Engine/Source/Runtime/RHI/Public/DynamicRHI.h b/Engine/Source/Runtime/RHI/Public/DynamicRHI.h
--- a/Engine/Source/Runtime/RHI/Public/DynamicRHI.h
+++ b/Engine/Source/Runtime/RHI/Public/DynamicRHI.h
@@ -102,12 +102,12 @@ struct FRayTracingGeometrySegment
 struct FRayTracingGeometryInitializer
 {
 	FRayTracingGeometryInitializer()
-		: BaseVertexIndex(0)
+		: VertexBufferByteOffset(0)
 		, VertexBufferStride(0)
-		, VertexBufferByteOffset(0)
-		, TotalPrimitiveCount(0)
 		, VertexBufferElementType(VET_Float3)
+		, BaseVertexIndex(0)
 		, PrimitiveType(PT_TriangleList)
+		, TotalPrimitiveCount(0)
 		, bFastBuild(false)
 		, bAllowUpdate(false)
 	{}
diff --git a/Engine/Source/Runtime/Renderer/Public/MeshPassProcessor.inl b/Engine/Source/Runtime/Renderer/Public/MeshPassProcessor.inl
--- a/Engine/Source/Runtime/Renderer/Public/MeshPassProcessor.inl
+++ b/Engine/Source/Runtime/Renderer/Public/MeshPassProcessor.inl
@@ -180,7 +180,10 @@ void FMeshPassProcessor::BuildRayTracingDrawCommands(
 	check(VertexFactory && VertexFactory->IsInitialized());
 	VertexFactory->GetStreams(FeatureLevel, SharedMeshDrawCommand.VertexStreams);
 
+#if RHI_RAYTRACING
 	SharedMeshDrawCommand.SetRayTracingShaders(PassShaders.GetUntypedShaders());
+#endif
+
 	SharedMeshDrawCommand.SetDebugData(PrimitiveSceneProxy, &MaterialResource, &MaterialRenderProxy, PassShaders.GetUntypedShaders());
 
 	if (PassShaders.RayHitGroupShader)
