Fixed dynamic mesh prepass. MeshBatch now uses bUseAsOccluder as a hint if it's an occluder or not, and new bUseForDepthPass flag is used to indicate whether this should be used in depth passes. This allows to not cache redundantly static meshes, while still being able to draw prepass for dynamic meshes.

diff --git a/Engine/Source/Runtime/Engine/Private/StaticMeshRender.cpp b/Engine/Source/Runtime/Engine/Private/StaticMeshRender.cpp
--- a/Engine/Source/Runtime/Engine/Private/StaticMeshRender.cpp
+++ b/Engine/Source/Runtime/Engine/Private/StaticMeshRender.cpp
@@ -397,8 +397,9 @@ bool FStaticMeshSceneProxy::GetShadowMeshElement(int32 LODIndex, int32 BatchInde
 	}
 
 	// By default this will be a shadow only mesh.
-	OutMeshBatch.bUseAsOccluder = false;
 	OutMeshBatch.bUseForMaterial = false;
+	OutMeshBatch.bUseForDepthPass = false;
+	OutMeshBatch.bUseAsOccluder = false;
 
 	return true;
 }
@@ -890,6 +891,7 @@ void FStaticMeshSceneProxy::DrawStaticElements(FStaticPrimitiveDrawInterface* PD
 									bUseUnifiedMeshForShadow = bAllSectionsCastShadow;
 
 									MeshBatch.CastShadow = bUseUnifiedMeshForShadow;
+									MeshBatch.bUseForDepthPass = bUseUnifiedMeshForDepth;
 									MeshBatch.bUseAsOccluder = bUseUnifiedMeshForDepth;
 									MeshBatch.bUseForMaterial = false;
 
@@ -926,6 +928,7 @@ void FStaticMeshSceneProxy::DrawStaticElements(FStaticPrimitiveDrawInterface* PD
 							// If we have submitted an optimized shadow-only mesh, remaining mesh elements must not cast shadows.
 							MeshBatch.CastShadow &= !bUseUnifiedMeshForShadow;
 							MeshBatch.bUseAsOccluder &= !bUseUnifiedMeshForDepth;
+							MeshBatch.bUseForDepthPass &= !bUseUnifiedMeshForDepth;
 
 							PDI->DrawMesh(MeshBatch, ScreenSize);
 						}
diff --git a/Engine/Source/Runtime/Engine/Public/MeshBatch.h b/Engine/Source/Runtime/Engine/Public/MeshBatch.h
--- a/Engine/Source/Runtime/Engine/Public/MeshBatch.h
+++ b/Engine/Source/Runtime/Engine/Public/MeshBatch.h
@@ -153,10 +153,11 @@ struct FMeshBatch
 	 * Pass feature relevance flags.  Allows a proxy to submit fast representations for passes which can take advantage of it, 
 	 * for example separate index buffer for depth-only rendering since vertices can be merged based on position and ignore UV differences.
 	 */
-	uint32 CastShadow : 1;			// Wheter it can be used in shadow renderpasses.
-	uint32 bUseForMaterial : 1;		// Whether it can be used in renderpasses requiring material outputs.
-	uint32 bUseAsOccluder : 1;		// Whether it can be used in renderpasses only depending on the raw geometry (i.e. Depth Prepass).
-	uint32 bWireframe : 1;
+	uint32 CastShadow		: 1;	// Whether it can be used in shadow renderpasses.
+	uint32 bUseForMaterial	: 1;	// Whether it can be used in renderpasses requiring material outputs.
+	uint32 bUseForDepthPass : 1;	// Whether it can be used in depth pass.
+	uint32 bUseAsOccluder	: 1;	// Hint whether this mesh is a good occluder.
+	uint32 bWireframe		: 1;
 	// e.g. PT_TriangleList(default), PT_LineList, ..
 	uint32 Type : PT_NumBits;
 	// e.g. SDPG_World (default), SDPG_Foreground
@@ -269,6 +270,7 @@ struct FMeshBatch
 	,	bDisableBackfaceCulling(false)
 	,	CastShadow(true)
 	,   bUseForMaterial(true)
+	,	bUseForDepthPass(true)
 	,	bUseAsOccluder(true)
 	,	bWireframe(false)
 	,	Type(PT_TriangleList)
diff --git a/Engine/Source/Runtime/Renderer/Private/DepthRendering.cpp b/Engine/Source/Runtime/Renderer/Private/DepthRendering.cpp
--- a/Engine/Source/Runtime/Renderer/Private/DepthRendering.cpp
+++ b/Engine/Source/Runtime/Renderer/Private/DepthRendering.cpp
@@ -1512,21 +1512,29 @@ void FDepthPassMeshProcessor::Process(
 
 void FDepthPassMeshProcessor::AddMeshBatch(const FMeshBatch& RESTRICT MeshBatch, uint64 BatchElementMask, const FPrimitiveSceneProxy* RESTRICT PrimitiveSceneProxy, int32 StaticMeshId)
 {
-	bool bDraw = MeshBatch.bUseAsOccluder;
+	bool bDraw = MeshBatch.bUseForDepthPass;
 
-	if (bRespectUseAsOccluderFlag && EarlyZPassMode < DDM_AllOpaque && PrimitiveSceneProxy)
+	// Filter by occluder flags and settings if required.
+	if (bDraw && bRespectUseAsOccluderFlag && !MeshBatch.bUseAsOccluder && EarlyZPassMode < DDM_AllOpaque)
 	{
-		// Only render primitives marked as occluders.
-		bDraw = bDraw && PrimitiveSceneProxy->ShouldUseAsOccluder()
-			// Only render static objects unless movable are requested.
-			&& (!PrimitiveSceneProxy->IsMovable() || bEarlyZPassMovable);
+		if (PrimitiveSceneProxy)
+		{
+			// Only render primitives marked as occluders.
+			bDraw = PrimitiveSceneProxy->ShouldUseAsOccluder()
+				// Only render static objects unless movable are requested.
+				&& (!PrimitiveSceneProxy->IsMovable() || bEarlyZPassMovable);
 
-		// Filter dynamic mesh commands by screen size.
-		if (ViewIfDynamicMeshCommand)
+			// Filter dynamic mesh commands by screen size.
+			if (ViewIfDynamicMeshCommand)
+			{
+				extern float GMinScreenRadiusForDepthPrepass;
+				const float LODFactorDistanceSquared = (PrimitiveSceneProxy->GetBounds().Origin - ViewIfDynamicMeshCommand->ViewMatrices.GetViewOrigin()).SizeSquared() * FMath::Square(ViewIfDynamicMeshCommand->LODDistanceFactor);
+				bDraw = bDraw && FMath::Square(PrimitiveSceneProxy->GetBounds().SphereRadius) > GMinScreenRadiusForDepthPrepass * GMinScreenRadiusForDepthPrepass * LODFactorDistanceSquared;
+			}
+		}
+		else
 		{
-			extern float GMinScreenRadiusForDepthPrepass;
-			const float LODFactorDistanceSquared = (PrimitiveSceneProxy->GetBounds().Origin - ViewIfDynamicMeshCommand->ViewMatrices.GetViewOrigin()).SizeSquared() * FMath::Square(ViewIfDynamicMeshCommand->LODDistanceFactor);
-			bDraw = bDraw && FMath::Square(PrimitiveSceneProxy->GetBounds().SphereRadius) > GMinScreenRadiusForDepthPrepass * GMinScreenRadiusForDepthPrepass * LODFactorDistanceSquared;
+			bDraw = false;
 		}
 	}
 
diff --git a/Engine/Source/Runtime/Renderer/Private/SceneCore.h b/Engine/Source/Runtime/Renderer/Private/SceneCore.h
--- a/Engine/Source/Runtime/Renderer/Private/SceneCore.h
+++ b/Engine/Source/Runtime/Renderer/Private/SceneCore.h
@@ -209,6 +209,7 @@ public:
 		, bRequiresPerElementVisibility(StaticMesh.bRequiresPerElementVisibility)
 		, CastShadow(StaticMesh.CastShadow)
 		, bUseForMaterial(StaticMesh.bUseForMaterial)
+		, bUseForDepthPass(StaticMesh.bUseForDepthPass)
 		, bUseAsOccluder(StaticMesh.bUseAsOccluder)
 		, bSupportsCachingMeshDrawCommands(InbSupportsCachingMeshDrawCommands)
 	{
@@ -229,9 +230,10 @@ public:
 	/** Whether the mesh batch needs VertexFactory->GetStaticBatchElementVisibility to be called each frame to determine which elements of the batch are visible. */
 	uint32 bRequiresPerElementVisibility : 1;
 
-	uint32 CastShadow : 1;			// Whether it can be used in shadow renderpasses.
-	uint32 bUseForMaterial : 1;		// Whether it can be used in renderpasses requiring material outputs.
-	uint32 bUseAsOccluder : 1;		// Whether it can be used in renderpasses only depending on the raw geometry (i.e. Depth Prepass).
+	uint32 CastShadow		: 1; // Whether it can be used in shadow renderpasses.
+	uint32 bUseForMaterial	: 1; // Whether it can be used in renderpasses requiring material outputs.
+	uint32 bUseForDepthPass	: 1; // Whether it can be used in depth pass.
+	uint32 bUseAsOccluder	: 1; // User hint whether it's a good occluder.
 
 	/** Cached from vertex factory to avoid dereferencing VF in InitViews. */
 	uint32 bSupportsCachingMeshDrawCommands : 1;
diff --git a/Engine/Source/Runtime/Renderer/Private/SceneVisibility.cpp b/Engine/Source/Runtime/Renderer/Private/SceneVisibility.cpp
--- a/Engine/Source/Runtime/Renderer/Private/SceneVisibility.cpp
+++ b/Engine/Source/Runtime/Renderer/Private/SceneVisibility.cpp
@@ -2109,8 +2109,7 @@ struct FRelevancePacket
 
 					if(ViewRelevance.bDrawRelevance && (StaticMeshRelevance.bUseForMaterial || StaticMeshRelevance.bUseAsOccluder) && (ViewRelevance.bRenderInMainPass || ViewRelevance.bRenderCustomDepth) && !bHiddenByHLODFade)
 					{
-						// If the static mesh is an occluder, check whether it covers enough of the screen to be used as an occluder.
-						if (StaticMeshRelevance.bUseAsOccluder && bDrawDepthOnly)
+						if (StaticMeshRelevance.bUseForDepthPass && bDrawDepthOnly)
 						{
 							DrawCommandPacket.AddCommandsForMesh(PrimitiveIndex, StaticMeshRelevance, StaticMesh, Scene, bCanCache, EMeshPass::DepthPass);
 							MarkMask |= EMarkMaskBits::StaticMeshOccluderMapMask;
diff --git a/Engine/Source/Runtime/Renderer/Private/ShadowSetup.cpp b/Engine/Source/Runtime/Renderer/Private/ShadowSetup.cpp
--- a/Engine/Source/Runtime/Renderer/Private/ShadowSetup.cpp
+++ b/Engine/Source/Runtime/Renderer/Private/ShadowSetup.cpp
@@ -915,7 +915,7 @@ bool FProjectedShadowInfo::ShouldDrawStaticMeshes(FViewInfo& InCurrentView, bool
 				const FStaticMeshRelevance& StaticMeshRelevance = InPrimitiveSceneInfo->StaticMeshRelevances[MeshIndex];
 				const FStaticMesh& StaticMesh = InPrimitiveSceneInfo->StaticMeshes[MeshIndex];
 
-				if ((StaticMeshRelevance.CastShadow || (bSelfShadowOnly && StaticMeshRelevance.bUseAsOccluder)) && ShadowLODToRender.ContainsLOD(StaticMeshRelevance.LODIndex))
+				if ((StaticMeshRelevance.CastShadow || (bSelfShadowOnly && StaticMeshRelevance.bUseForDepthPass)) && ShadowLODToRender.ContainsLOD(StaticMeshRelevance.LODIndex))
 				{
 					if (UseMeshDrawCommandPipeline())
 					{
@@ -956,7 +956,7 @@ bool FProjectedShadowInfo::ShouldDrawStaticMeshes(FViewInfo& InCurrentView, bool
 				const FStaticMeshRelevance& StaticMeshRelevance = InPrimitiveSceneInfo->StaticMeshRelevances[MeshIndex];
 				const FStaticMesh& StaticMesh = InPrimitiveSceneInfo->StaticMeshes[MeshIndex];
 
-				if ((StaticMeshRelevance.CastShadow || (bSelfShadowOnly && StaticMeshRelevance.bUseAsOccluder)) && ShadowLODToRender.ContainsLOD(StaticMeshRelevance.LODIndex))
+				if ((StaticMeshRelevance.CastShadow || (bSelfShadowOnly && StaticMeshRelevance.bUseForDepthPass)) && ShadowLODToRender.ContainsLOD(StaticMeshRelevance.LODIndex))
 				{
 					if (UseMeshDrawCommandPipeline())
 					{
