Fixed VR editor crashes from FViewInfo snapshots due to non PDO object memcpy.

diff --git a/Engine/Source/Runtime/Renderer/Private/SceneRendering.cpp b/Engine/Source/Runtime/Renderer/Private/SceneRendering.cpp
--- a/Engine/Source/Runtime/Renderer/Private/SceneRendering.cpp
+++ b/Engine/Source/Runtime/Renderer/Private/SceneRendering.cpp
@@ -1551,6 +1551,13 @@ FViewInfo* FViewInfo::CreateSnapshot() const
 	TUniquePtr<FViewUniformShaderParameters> NullViewParameters;
 	FMemory::Memcpy(Result->CachedViewUniformShaderParameters, NullViewParameters); 
 
+	TArray<FPrimitiveUniformShaderParameters> NullDynamicPrimitiveShaderData;
+	FMemory::Memcpy(Result->DynamicPrimitiveShaderData, NullDynamicPrimitiveShaderData);
+	Result->DynamicPrimitiveShaderData = DynamicPrimitiveShaderData;
+
+	FRWBufferStructured NullOneFramePrimitiveShaderDataBuffer;
+	FMemory::Memcpy(Result->OneFramePrimitiveShaderDataBuffer, NullOneFramePrimitiveShaderDataBuffer);
+
 	FReadBuffer NullReadBuffer;
 	FMemory::Memcpy(Result->OneFramePrimitiveIdBufferEmulation, NullReadBuffer);
 
@@ -1578,6 +1585,8 @@ void FViewInfo::DestroyAllSnapshots()
 	{
 		Snapshot->ViewUniformBuffer.SafeRelease();
 		Snapshot->CachedViewUniformShaderParameters.Reset();
+		Snapshot->DynamicPrimitiveShaderData.Empty();
+		Snapshot->OneFramePrimitiveShaderDataBuffer.Release();
 		Snapshot->OneFramePrimitiveIdBufferEmulation.Release();
 		FreeViewInfoSnapshots.Add(Snapshot);
 	}
