Don't compile RT test shaders when RT not supported

diff --git a/Engine/Source/Runtime/Renderer/Private/RayTracing/RayTracingTestbed.cpp b/Engine/Source/Runtime/Renderer/Private/RayTracing/RayTracingTestbed.cpp
--- a/Engine/Source/Runtime/Renderer/Private/RayTracing/RayTracingTestbed.cpp
+++ b/Engine/Source/Runtime/Renderer/Private/RayTracing/RayTracingTestbed.cpp
@@ -168,7 +168,7 @@ public:
 	static bool ShouldCompilePermutation(const FGlobalShaderPermutationParameters& Parameters)
 	{
 		// #dxr_todo: this should also check if ray tracing is enabled for the target platform & project
-		return IsFeatureLevelSupported(Parameters.Platform, ERHIFeatureLevel::SM5);
+		return IsFeatureLevelSupported(Parameters.Platform, ERHIFeatureLevel::SM5) && IsRayTracingSupportedForThisProject();
 	}
 
 	FTestRaygenShader() {}
