Disable texture importance sampling when source texture is not present.

diff --git a/Engine/Shaders/Private/RayTracing/RayTracingRectLightOcclusionRGS.usf b/Engine/Shaders/Private/RayTracing/RayTracingRectLightOcclusionRGS.usf
--- a/Engine/Shaders/Private/RayTracing/RayTracingRectLightOcclusionRGS.usf
+++ b/Engine/Shaders/Private/RayTracing/RayTracingRectLightOcclusionRGS.usf
@@ -205,9 +205,13 @@ void RectLightOcclusionRGS()
 		float2 RandSample = RandomSequence_GenerateSample2D(RandSequence, DummyVariable);
 		float RayPdf = 0.0;
 
+#if TEXTURE_IMPORTANCE_SAMPLING
 		bool bIsValid = RectLight.bIsTextureImportanceSampling ?
 			GenerateRectLightImportanceSampledRay(WorldPosition, WorldNormal, LightPosition, LightNormal, LightdPdu, LightdPdv, LightWidth, LightHeight, RandSequence, Ray.Origin, Ray.Direction, Ray.TMin, Ray.TMax, RayPdf) :
 			GenerateRectLightOcclusionRay(WorldPosition, WorldNormal, LightPosition, LightNormal, LightdPdu, LightdPdv, LightWidth, LightHeight, RandSample, Ray.Origin, Ray.Direction, Ray.TMin, Ray.TMax, RayPdf);
+#else
+		bool bIsValid = GenerateRectLightOcclusionRay(WorldPosition, WorldNormal, LightPosition, LightNormal, LightdPdu, LightdPdv, LightWidth, LightHeight, RandSample, Ray.Origin, Ray.Direction, Ray.TMin, Ray.TMax, RayPdf);
+#endif
 		if (!bIsValid)
 		{
 			continue;
diff --git a/Engine/Source/Runtime/Renderer/Private/DeferredShadingRenderer.h b/Engine/Source/Runtime/Renderer/Private/DeferredShadingRenderer.h
--- a/Engine/Source/Runtime/Renderer/Private/DeferredShadingRenderer.h
+++ b/Engine/Source/Runtime/Renderer/Private/DeferredShadingRenderer.h
@@ -455,7 +455,7 @@ private:
 	void VisualizeRectLightMipTree(FRHICommandListImmediate& RHICmdList, const FViewInfo& View, const FRWBuffer& RectLightMipTree, const FIntVector& RectLightMipTreeDimensions);
 	void RenderRayTracingRectLight(FRHICommandListImmediate& RHICmdList, const FLightSceneInfo& RectLightSceneInfo, TRefCountPtr<IPooledRenderTarget>& RectLightRT, TRefCountPtr<IPooledRenderTarget>& HitDistanceRT);
 	void RenderRayTracingOcclusionForRectLight(FRHICommandListImmediate& RHICmdList, const FLightSceneInfo& RectLightSceneInfo, TRefCountPtr<IPooledRenderTarget>& ScreenShadowMaskTexture);
-	template <int CalcDirectLighting, int EncodeVisibility> void RenderRayTracingRectLightInternal(FRHICommandListImmediate& RHICmdList, const FLightSceneInfo& RectLightSceneInfo, TRefCountPtr<IPooledRenderTarget>& ScreenShadowMaskTexture, TRefCountPtr<IPooledRenderTarget>& RayDistanceTexture);
+	template <int CalcDirectLighting, int EncodeVisibility, int TextureImportanceSampling> void RenderRayTracingRectLightInternal(FRHICommandListImmediate& RHICmdList, const FLightSceneInfo& RectLightSceneInfo, TRefCountPtr<IPooledRenderTarget>& ScreenShadowMaskTexture, TRefCountPtr<IPooledRenderTarget>& RayDistanceTexture);
 
 	void RenderRayTracingAmbientOcclusion(FRHICommandListImmediate& RHICmdList, const FLightSceneInfo* LightSceneInfo, TRefCountPtr<IPooledRenderTarget>& AmbientOcclusionRT);
 	void CompositeRayTracingAmbientOcclusion(FRHICommandListImmediate& RHICmdList, TRefCountPtr<IPooledRenderTarget>& AmbientOcclusionRT);
diff --git a/Engine/Source/Runtime/Renderer/Private/RayTracing/RayTracingShadowsForRectLight.cpp b/Engine/Source/Runtime/Renderer/Private/RayTracing/RayTracingShadowsForRectLight.cpp
--- a/Engine/Source/Runtime/Renderer/Private/RayTracing/RayTracingShadowsForRectLight.cpp
+++ b/Engine/Source/Runtime/Renderer/Private/RayTracing/RayTracingShadowsForRectLight.cpp
@@ -81,7 +81,7 @@ DECLARE_GPU_STAT_NAMED(BuildRectLightMipTree, TEXT("Build RectLight Mip Tree"));
 
 IMPLEMENT_GLOBAL_SHADER_PARAMETER_STRUCT(FRectLightData, "RectLight");
 
-template <int CalcDirectLighting, int EncodeVisibility>
+template <int CalcDirectLighting, int EncodeVisibility, int TextureImportanceSampling>
 class FRectLightOcclusionRGS : public FGlobalShader
 {
 	DECLARE_SHADER_TYPE(FRectLightOcclusionRGS, Global)
@@ -90,6 +90,7 @@ class FRectLightOcclusionRGS : public FGlobalShader
 	{
 		OutEnvironment.SetDefine(TEXT("CALC_DIRECT_LIGHTING"), CalcDirectLighting);
 		OutEnvironment.SetDefine(TEXT("ENCODE_VISIBILITY"), EncodeVisibility);
+		OutEnvironment.SetDefine(TEXT("TEXTURE_IMPORTANCE_SAMPLING"), TextureImportanceSampling);
 	}
 
 public:
@@ -164,13 +165,16 @@ private:
 	FShaderResourceParameter RayDistanceUAVParameter;
 };
 
-#define IMPLEMENT_RECT_LIGHT_OCCLUSION_TYPE(CalcDirectLightingType, EncodeVisbilityType) \
-	typedef FRectLightOcclusionRGS<CalcDirectLightingType, EncodeVisbilityType> FRectLightOcclusionRGS##CalcDirectLightingType##EncodeVisbilityType; \
-	IMPLEMENT_SHADER_TYPE(template<>, FRectLightOcclusionRGS##CalcDirectLightingType##EncodeVisbilityType, TEXT("/Engine/Private/RayTracing/RayTracingRectLightOcclusionRGS.usf"), TEXT("RectLightOcclusionRGS"), SF_RayGen);
+#define IMPLEMENT_RECT_LIGHT_OCCLUSION_TYPE(CalcDirectLightingType, EncodeVisbilityType, TextureImportanceSampling) \
+	typedef FRectLightOcclusionRGS<CalcDirectLightingType, EncodeVisbilityType, TextureImportanceSampling> FRectLightOcclusionRGS##CalcDirectLightingType##EncodeVisbilityType##TextureImportanceSampling; \
+	IMPLEMENT_SHADER_TYPE(template<>, FRectLightOcclusionRGS##CalcDirectLightingType##EncodeVisbilityType##TextureImportanceSampling, TEXT("/Engine/Private/RayTracing/RayTracingRectLightOcclusionRGS.usf"), TEXT("RectLightOcclusionRGS"), SF_RayGen);
 
-IMPLEMENT_RECT_LIGHT_OCCLUSION_TYPE(0, 0);
-IMPLEMENT_RECT_LIGHT_OCCLUSION_TYPE(0, 1);
-IMPLEMENT_RECT_LIGHT_OCCLUSION_TYPE(1, 0);
+IMPLEMENT_RECT_LIGHT_OCCLUSION_TYPE(0, 0, 0);
+IMPLEMENT_RECT_LIGHT_OCCLUSION_TYPE(0, 0, 1);
+IMPLEMENT_RECT_LIGHT_OCCLUSION_TYPE(0, 1, 0);
+IMPLEMENT_RECT_LIGHT_OCCLUSION_TYPE(0, 1, 1);
+IMPLEMENT_RECT_LIGHT_OCCLUSION_TYPE(1, 0, 0);
+IMPLEMENT_RECT_LIGHT_OCCLUSION_TYPE(1, 0, 1);
 
 class FBuildRectLightMipTreeCS : public FGlobalShader
 {
@@ -443,7 +447,14 @@ void FDeferredShadingSceneRenderer::RenderRayTracingRectLight(
 	GRenderTargetPool.FindFreeElement(RHICmdList, Desc, HitDistanceRT, TEXT("RayTracingRectLightDistance"));
 	//ClearUAV(RHICmdList, HitDistanceRT->GetRenderTargetItem(), FLinearColor::Black);
 
-	 RenderRayTracingRectLightInternal<1, 0>(RHICmdList, RectLightSceneInfo, RectLightRT, HitDistanceRT);
+	if (RectLightSceneInfo.Proxy->HasSourceTexture())
+	{
+		RenderRayTracingRectLightInternal<1, 0, 1>(RHICmdList, RectLightSceneInfo, RectLightRT, HitDistanceRT);
+	}
+	else
+	{
+		RenderRayTracingRectLightInternal<1, 0, 0>(RHICmdList, RectLightSceneInfo, RectLightRT, HitDistanceRT);
+	}
 }
 
 void FDeferredShadingSceneRenderer::RenderRayTracingOcclusionForRectLight(
@@ -467,10 +478,17 @@ void FDeferredShadingSceneRenderer::RenderRayTracingOcclusionForRectLight(
 	GRenderTargetPool.FindFreeElement(RHICmdList, Desc, HitDistanceTexture, TEXT("RayTracingRectLightDistance"));
 	//ClearUAV(RHICmdList, HitDistanceTexture->GetRenderTargetItem(), FLinearColor::Black);
 
-	RenderRayTracingRectLightInternal<0, 1>(RHICmdList, RectLightSceneInfo, ScreenShadowMaskTexture, HitDistanceTexture);
+	if (RectLightSceneInfo.Proxy->HasSourceTexture())
+	{
+		RenderRayTracingRectLightInternal<0, 1, 1>(RHICmdList, RectLightSceneInfo, ScreenShadowMaskTexture, HitDistanceTexture);
+	}
+	else
+	{
+		RenderRayTracingRectLightInternal<0, 1, 0>(RHICmdList, RectLightSceneInfo, ScreenShadowMaskTexture, HitDistanceTexture);
+	}
 }
 
-template <int CalcDirectLighting, int EncodeVisibility>
+template <int CalcDirectLighting, int EncodeVisibility, int TextureImportanceSampling>
 void FDeferredShadingSceneRenderer::RenderRayTracingRectLightInternal(
 	FRHICommandListImmediate& RHICmdList,
 	const FLightSceneInfo& RectLightSceneInfo,
@@ -480,6 +498,7 @@ void FDeferredShadingSceneRenderer::RenderRayTracingRectLightInternal(
 {
 	check(RectLightSceneInfo.Proxy);
 
+	// #dxr_todo: Cache MIP tree build with RectLight
 	FRWBuffer RectLightMipTree;
 	FIntVector RectLightMipTreeDimensions;
  	BuildRectLightMipTree(RHICmdList, RectLightSceneInfo, RectLightMipTree, RectLightMipTreeDimensions);
@@ -523,7 +542,7 @@ void FDeferredShadingSceneRenderer::RenderRayTracingRectLightInternal(
 		FViewInfo& View = Views[ViewIndex];
 		FIntPoint ViewSize = View.ViewRect.Size();
 
-		TShaderMapRef<FRectLightOcclusionRGS<CalcDirectLighting, EncodeVisibility>> RectLightOcclusionRayGenerationShader(GetGlobalShaderMap(FeatureLevel));
+		TShaderMapRef<FRectLightOcclusionRGS<CalcDirectLighting, EncodeVisibility, TextureImportanceSampling>> RectLightOcclusionRayGenerationShader(GetGlobalShaderMap(FeatureLevel));
 
 		FSceneRenderTargets& SceneContext = FSceneRenderTargets::Get(RHICmdList);
 		FSceneTexturesUniformParameters SceneTextures;
