Added two new CVars to the ray tracing reflections shader to help profiling it
- GRayTracingReflectionsDirectLighting for enabling/disabling direct lighting computation
- GRayTracingReflectionsMaxRayDistance to set the reflection rays TMax

diff --git a/Engine/Shaders/Private/RayTracing/RayTracingReflections.usf b/Engine/Shaders/Private/RayTracing/RayTracingReflections.usf
--- a/Engine/Shaders/Private/RayTracing/RayTracingReflections.usf
+++ b/Engine/Shaders/Private/RayTracing/RayTracingReflections.usf
@@ -29,6 +29,8 @@ float2 SobolIndexToUniformUnitSquare(uint2 SobolRand)
 #include "RayTracingHitGroupCommon.ush"
 
 uint SamplesPerPixel;
+uint ShouldDoDirectLighting;
+float ReflectionMaxRayDistance;
 
 RaytracingAccelerationStructure TLAS;
 
@@ -217,13 +219,14 @@ FMaterialClosestHitPayload TraceRayAndAccumulateResults(
 		Ray        // RayDesc
 	);
 
+	if (ShouldDoDirectLighting)
+	{
+		float3 DirectLighting = ComputeDirectLighting(Ray, Payload);
 
-	float3 DirectLighting = ComputeDirectLighting(Ray, Payload);
-
-	// Transform NaNs to black, transform negative colors to black.
-	DirectLighting = -min(-DirectLighting, 0.0);
-
-	Radiance += DirectLighting;
+		// Transform NaNs to black, transform negative colors to black.
+		DirectLighting = -min(-DirectLighting, 0.0);
+		Radiance += DirectLighting;
+	}
 
 	return Payload;
 }
@@ -287,7 +290,7 @@ void RayTracingReflectionsRGS()
 				RayDesc Ray;
 				Ray.Origin = WorldPosition + WorldNormal * 0.01;
 				Ray.TMin = 0.0;
-				Ray.TMax = 1e27f;
+				Ray.TMax = ReflectionMaxRayDistance;
 
 				uint DummyVariable;
 				float2 RandSample = RandomSequence_GenerateSample2D(RandSequence, DummyVariable);
@@ -344,12 +347,11 @@ void RayTracingReflectionsRGS()
 		}
 	}
 
-	////uint NumLights = ForwardLightData.NumDirectionalLights;
-
 	if (SamplesPerPixel > 0)
 	{
 		ReflectedColor /= SamplesPerPixel;
 	}
+
 	ColorOutput[PixelCoord] = float4(ReflectedColor);
 }
 
diff --git a/Engine/Source/Runtime/Renderer/Private/RayTracing/RayTracingReflections.cpp b/Engine/Source/Runtime/Renderer/Private/RayTracing/RayTracingReflections.cpp
--- a/Engine/Source/Runtime/Renderer/Private/RayTracing/RayTracingReflections.cpp
+++ b/Engine/Source/Runtime/Renderer/Private/RayTracing/RayTracingReflections.cpp
@@ -23,6 +23,20 @@ static FAutoConsoleVariableRef CVarRayTracingReflectionsSamplesPerPixel(
 	TEXT("Sets the samples-per-pixel for reflections (default = 1)")
 );
 
+static int32 GRayTracingReflectionsDirectLighting = 1;
+static FAutoConsoleVariableRef CVarRayTracingReflectionsDirectLighting(
+	TEXT("r.RayTracing.Reflections.DirectLighting"),
+	GRayTracingReflectionsDirectLighting,
+	TEXT("Enables ray tracing reflections direct lighting (default = 1)")
+);
+
+static float GRayTracingReflectionsMaxRayDistance = 1.0e27;
+static FAutoConsoleVariableRef CVarRayTracingReflectionsMaxRayDistance(
+	TEXT("r.RayTracing.Reflections.MaxRayDistance"),
+	GRayTracingReflectionsMaxRayDistance,
+	TEXT("Sets the maximum ray distance for ray traced reflection rays (default = 1.0e27)")
+);
+
 static const int32 GReflectionLightCountMaximum = 64;
 
 BEGIN_GLOBAL_SHADER_PARAMETER_STRUCT(FReflectionsLightData, )
@@ -98,6 +112,8 @@ class FRayTracingReflectionsRG : public FGlobalShader
 
 	BEGIN_SHADER_PARAMETER_STRUCT(FParameters, )
 		SHADER_PARAMETER(int32, SamplesPerPixel)
+		SHADER_PARAMETER(int32, ShouldDoDirectLighting)
+		SHADER_PARAMETER(float, ReflectionMaxRayDistance)
 
 		SHADER_PARAMETER_SRV(RaytracingAccelerationStructure, TLAS)
 
@@ -171,6 +187,8 @@ void FDeferredShadingSceneRenderer::RayTraceReflections(
 	FRayTracingReflectionsRG::FParameters* PassParameters = GraphBuilder.AllocParameters<FRayTracingReflectionsRG::FParameters>();
 
 	PassParameters->SamplesPerPixel = GRayTracingReflectionsSamplesPerPixel;
+	PassParameters->ShouldDoDirectLighting = GRayTracingReflectionsDirectLighting;
+	PassParameters->ReflectionMaxRayDistance = GRayTracingReflectionsMaxRayDistance;
 	PassParameters->LTCMatTexture = GSystemTextures.LTCMat->GetRenderTargetItem().ShaderResourceTexture;
 	PassParameters->LTCMatSampler = TStaticSamplerState<SF_Bilinear, AM_Clamp, AM_Clamp, AM_Clamp>::GetRHI();
 	PassParameters->LTCAmpTexture = GSystemTextures.LTCAmp->GetRenderTargetItem().ShaderResourceTexture;
