Workarounds for HLSLCC platforms.  Rect lights will be broken for these platforms until HLSLCC is fixed after the break

diff --git a/Engine/Shaders/Private/DeferredLightingCommon.ush b/Engine/Shaders/Private/DeferredLightingCommon.ush
--- a/Engine/Shaders/Private/DeferredLightingCommon.ush
+++ b/Engine/Shaders/Private/DeferredLightingCommon.ush
@@ -323,7 +323,12 @@ float4 GetDynamicLighting(float3 WorldPosition, float3 CameraVector, FGBufferDat
 		#else
 			FDirectLighting Lighting;
 
-			if( LightData.bRectLight )
+			//temp workaround for HLSLCC bug with propagating SourceTexture parameter down
+			#if COMPILER_HLSLCC
+			if (false)
+			#else
+			if (LightData.bRectLight)
+			#endif
 			{
 				FRect Rect = GetRect( ToLight, LightData );
 
diff --git a/Engine/Source/Runtime/Renderer/Private/PathTracing/SkyLightCdfBuild.cpp b/Engine/Source/Runtime/Renderer/Private/PathTracing/SkyLightCdfBuild.cpp
--- a/Engine/Source/Runtime/Renderer/Private/PathTracing/SkyLightCdfBuild.cpp
+++ b/Engine/Source/Runtime/Renderer/Private/PathTracing/SkyLightCdfBuild.cpp
@@ -267,14 +267,9 @@ class FVisualizeCdfPS : public FGlobalShader
 {
 	DECLARE_SHADER_TYPE(FVisualizeCdfPS, Global);
 
-	static bool ShouldCache(EShaderPlatform Platform)
-	{
-		return IsFeatureLevelSupported(Platform, ERHIFeatureLevel::SM5);
-	}
-
 	static bool ShouldCompilePermutation(const FGlobalShaderPermutationParameters& Parameters)
 	{
-		return IsFeatureLevelSupported(Parameters.Platform, ERHIFeatureLevel::SM5);
+		return IsFeatureLevelSupported(Parameters.Platform, ERHIFeatureLevel::SM5) && IsRayTracingSupportedForThisProject();
 	}
 
 	static void ModifyCompilationEnvironment(const FGlobalShaderPermutationParameters& Parameters, FShaderCompilerEnvironment& OutEnvironment)
diff --git a/Engine/Source/Runtime/Renderer/Private/PathTracing/VarianceMipTreeBuild.cpp b/Engine/Source/Runtime/Renderer/Private/PathTracing/VarianceMipTreeBuild.cpp
--- a/Engine/Source/Runtime/Renderer/Private/PathTracing/VarianceMipTreeBuild.cpp
+++ b/Engine/Source/Runtime/Renderer/Private/PathTracing/VarianceMipTreeBuild.cpp
@@ -19,7 +19,7 @@ class FBuildVarianceMipTreeCS : public FGlobalShader
 public:
 	static bool ShouldCompilePermutation(const FGlobalShaderPermutationParameters& Parameters)
 	{
-		return RHISupportsComputeShaders(Parameters.Platform);
+		return RHISupportsComputeShaders(Parameters.Platform) && IsRayTracingSupportedForThisProject();
 	}
 
 	static uint32 GetGroupSize()
