temp work around for Android compilation issues

diff --git a/Engine/Source/Runtime/Renderer/Private/ReflectionEnvironment.cpp b/Engine/Source/Runtime/Renderer/Private/ReflectionEnvironment.cpp
--- a/Engine/Source/Runtime/Renderer/Private/ReflectionEnvironment.cpp
+++ b/Engine/Source/Runtime/Renderer/Private/ReflectionEnvironment.cpp
@@ -728,6 +728,7 @@ void FDeferredShadingSceneRenderer::RenderDeferredReflectionsAndSkyLighting(FRHI
 		TRefCountPtr<IPooledRenderTarget> ReflectionsColor = GSystemTextures.BlackDummy;
 		if (bRayTracedReflections)
 		{
+#if RHI_RAYTRACING
 			//SCOPED_DRAW_EVENT(RHICmdList, RayTracedReflections);
 
 			FRDGBuilder GraphBuilder(RHICmdList); // TODO: convert the entire reflections to render graph.
@@ -735,10 +736,12 @@ void FDeferredShadingSceneRenderer::RenderDeferredReflectionsAndSkyLighting(FRHI
 			FSceneViewFamilyBlackboard SceneBlackboard;
 			SetupSceneViewFamilyBlackboard(GraphBuilder, &SceneBlackboard);
 
+
 			// Ray trace the reflection.
 			IScreenSpaceDenoiser::FReflectionInputs DenoiserInputs;
 			RayTraceReflections(GraphBuilder, View, &DenoiserInputs.Color, &DenoiserInputs.RayHitDistance);
 
+
 			// Denoise the reflections.
 			if (int32 DenoiserMode = CVarUseReflectionDenoiser.GetValueOnRenderThread())
 			{
@@ -765,6 +768,7 @@ void FDeferredShadingSceneRenderer::RenderDeferredReflectionsAndSkyLighting(FRHI
 			}
 
 			GraphBuilder.Execute();
+#endif
 		}
 		else if (bScreenSpaceReflections)
 		{
diff --git a/Engine/Source/Runtime/Renderer/Private/ScreenSpaceDenoise.cpp b/Engine/Source/Runtime/Renderer/Private/ScreenSpaceDenoise.cpp
--- a/Engine/Source/Runtime/Renderer/Private/ScreenSpaceDenoise.cpp
+++ b/Engine/Source/Runtime/Renderer/Private/ScreenSpaceDenoise.cpp
@@ -319,7 +319,11 @@ static void DenoiseShadowPenumbra(
 		PassParameters->SignalOutput0Mip1 = GraphBuilder.CreateUAV(FRDGTextureUAVDesc(ReducedSignal0, /* MipLevel = */ 1));
 		PassParameters->SignalOutput0Mip2 = GraphBuilder.CreateUAV(FRDGTextureUAVDesc(ReducedSignal0, /* MipLevel = */ 2));
 		PassParameters->SignalOutput0Mip3 = GraphBuilder.CreateUAV(FRDGTextureUAVDesc(ReducedSignal0, /* MipLevel = */ 3));
-		PassParameters->TileClassificationOutput = GraphBuilder.CreateUAV(CommonParameters.TileClassificationTexture);
+
+		{
+			FRDGTextureRef TileClassificationTexture = CommonParameters.TileClassificationTexture;
+			PassParameters->TileClassificationOutput = GraphBuilder.CreateUAV(TileClassificationTexture);
+		}
 
 		PassParameters->DebugOutput = GraphBuilder.CreateUAV(GraphBuilder.CreateTexture(PenumbraRT->Desc, TEXT("SSDDebugReduce")));
 
@@ -572,6 +576,8 @@ public:
 	{
 		FShadowPenumbraOutputs Outputs;
 
+		//temp workaround for compilation errors with android compiler
+#if !PLATFORM_ANDROID
 		if (IsSupportedLightType(ELightComponentType(LightSceneInfo.Proxy->GetLightType())))
 		{
 			::DenoiseShadowPenumbra(
@@ -583,6 +589,7 @@ public:
 				&Outputs.DiffusePenumbra);
 		}
 		else
+#endif
 		{
 			Outputs.DiffusePenumbra = ShadowInputs.Penumbra;
 		}
