Adding profiling stats for ray traced reflections.

diff --git a/Engine/Source/Runtime/Renderer/Private/ReflectionEnvironment.cpp b/Engine/Source/Runtime/Renderer/Private/ReflectionEnvironment.cpp
--- a/Engine/Source/Runtime/Renderer/Private/ReflectionEnvironment.cpp
+++ b/Engine/Source/Runtime/Renderer/Private/ReflectionEnvironment.cpp
@@ -33,7 +33,7 @@
 #include "ScreenSpaceDenoise.h"
 
 DECLARE_GPU_STAT_NAMED(ReflectionEnvironment, TEXT("Reflection Environment"));
-DECLARE_GPU_STAT_NAMED(RayTracedReflections, TEXT("Ray Traced Reflections"));
+DECLARE_GPU_STAT_NAMED(RayTracingReflections, TEXT("Ray Tracing Reflections"));
 DECLARE_GPU_STAT(SkyLightDiffuse);
 
 extern TAutoConsoleVariable<int32> CVarLPVMixing;
@@ -729,7 +729,8 @@ void FDeferredShadingSceneRenderer::RenderDeferredReflectionsAndSkyLighting(FRHI
 		if (bRayTracedReflections)
 		{
 #if RHI_RAYTRACING
-			//SCOPED_DRAW_EVENT(RHICmdList, RayTracedReflections);
+			SCOPED_DRAW_EVENT(RHICmdList, RayTracingReflections);
+			SCOPED_GPU_STAT(RHICmdList, RayTracingReflections);
 
 			FRDGBuilder GraphBuilder(RHICmdList); // TODO: convert the entire reflections to render graph.
 
