More shipping fixes

diff --git a/Engine/Source/Runtime/Renderer/Private/RayTracing/RayTracingDynamicGeometry.cpp b/Engine/Source/Runtime/Renderer/Private/RayTracing/RayTracingDynamicGeometry.cpp
--- a/Engine/Source/Runtime/Renderer/Private/RayTracing/RayTracingDynamicGeometry.cpp
+++ b/Engine/Source/Runtime/Renderer/Private/RayTracing/RayTracingDynamicGeometry.cpp
@@ -130,10 +130,12 @@ void FRayTracingDynamicGeometryCollection::AddDynamicMeshBatchForGeometryUpdate(
 
 	check(DispatchCmd.TargetBuffer->NumBytes >= DispatchCmd.NumVertices * sizeof(FVector));
 
+#if MESH_DRAW_COMMAND_DEBUG_DATA
 	FMeshDrawCommandDebugData DebugData;
 	DebugData.Shaders = Shaders.GetUntypedShaders();
 
 	ShaderBindings.Finalize(DebugData);
+#endif
 
 	DispatchCommands.Add(DispatchCmd);
 
