Mesh processor - fixing crash when using DrawTileMesh and Scene can be nullptr.

diff --git a/Engine/Source/Runtime/Renderer/Private/MeshPassProcessor.cpp b/Engine/Source/Runtime/Renderer/Private/MeshPassProcessor.cpp
--- a/Engine/Source/Runtime/Renderer/Private/MeshPassProcessor.cpp
+++ b/Engine/Source/Runtime/Renderer/Private/MeshPassProcessor.cpp
@@ -1024,7 +1024,7 @@ int32 FMeshPassProcessor::GetDrawCommandPrimitiveId(const FPrimitiveSceneInfo* R
 	}
 	else
 	{
-		DrawPrimitiveId = PrimitiveSceneInfo ? PrimitiveSceneInfo->GetIndex() : Scene->Primitives.Num();
+		DrawPrimitiveId = PrimitiveSceneInfo ? PrimitiveSceneInfo->GetIndex() : INT32_MAX;
 	}
 
 	return DrawPrimitiveId;
