Fixed SSAO was not properly computed when ray tracing was not enabled.

diff --git a/Engine/Source/Runtime/Renderer/Private/CompositionLighting/CompositionLighting.cpp b/Engine/Source/Runtime/Renderer/Private/CompositionLighting/CompositionLighting.cpp
--- a/Engine/Source/Runtime/Renderer/Private/CompositionLighting/CompositionLighting.cpp
+++ b/Engine/Source/Runtime/Renderer/Private/CompositionLighting/CompositionLighting.cpp
@@ -125,7 +125,7 @@ bool ShouldRenderScreenSpaceAmbientOcclusion(const FViewInfo& View)
 
 	// #dxr_todo: hard-coded FindConsoleVariable is a bad hack. The current variable is hidden behind DeferredShadingRenderer
 	static IConsoleVariable* UseRayTracingAmbientOcclusionCVar = IConsoleManager::Get().FindConsoleVariable(TEXT("r.RayTracing.AmbientOcclusion"));
-	bool bUseRayTracingAmbientOcclusion = UseRayTracingAmbientOcclusionCVar && UseRayTracingAmbientOcclusionCVar->GetInt() > 0;
+	bool bUseRayTracingAmbientOcclusion = IsRayTracingSupportedForThisProject() && UseRayTracingAmbientOcclusionCVar && UseRayTracingAmbientOcclusionCVar->GetInt() > 0;
 	return bEnabled && !bUseRayTracingAmbientOcclusion;
 }
 
