Recache all cached mesh draw commands after invalidating entire material uniform expression cache, as they store pointers to material uniform buffers which were just invalidated.

diff --git a/Engine/Source/Runtime/Engine/Private/ShaderCompiler/ShaderCompiler.cpp b/Engine/Source/Runtime/Engine/Private/ShaderCompiler/ShaderCompiler.cpp
--- a/Engine/Source/Runtime/Engine/Private/ShaderCompiler/ShaderCompiler.cpp
+++ b/Engine/Source/Runtime/Engine/Private/ShaderCompiler/ShaderCompiler.cpp
@@ -3449,6 +3449,9 @@ public:
 
 			// fixup uniform expressions
 			UMaterialInterface::RecacheAllMaterialUniformExpressions();
+
+			// Need to recache all cached mesh draw commands, as they store pointers to material uniform buffers which we just invalidated.
+			GetRendererModule().UpdateStaticDrawLists();
 		}
 
 		ENQUEUE_UNIQUE_RENDER_COMMAND(
diff --git a/Engine/Source/Runtime/Engine/Private/UnrealEngine.cpp b/Engine/Source/Runtime/Engine/Private/UnrealEngine.cpp
--- a/Engine/Source/Runtime/Engine/Private/UnrealEngine.cpp
+++ b/Engine/Source/Runtime/Engine/Private/UnrealEngine.cpp
@@ -779,6 +779,9 @@ void RefreshSamplerStatesCallback()
 		}
 
 		UMaterialInterface::RecacheAllMaterialUniformExpressions();
+
+		// Need to recache all cached mesh draw commands, as they store pointers to material uniform buffers which we just invalidated.
+		GetRendererModule().UpdateStaticDrawLists();
 	}
 }
 
diff --git a/Engine/Source/Runtime/Engine/Public/SceneInterface.h b/Engine/Source/Runtime/Engine/Public/SceneInterface.h
--- a/Engine/Source/Runtime/Engine/Public/SceneInterface.h
+++ b/Engine/Source/Runtime/Engine/Public/SceneInterface.h
@@ -186,6 +186,9 @@ public:
 	/** Sets the precomputed volume distance field for the scene, or NULL to clear the current one. */
 	virtual void SetPrecomputedVolumeDistanceField(const class FPrecomputedVolumeDistanceField* PrecomputedVolumeDistanceField) {}
 
+	/** Updates all static draw lists. */
+	virtual void UpdateStaticDrawLists() {}
+
 	/** Updates static draw lists for the given set of materials. */
 	virtual void UpdateStaticDrawListsForMaterials(const TArray<const FMaterial*>& Materials) {}
 
diff --git a/Engine/Source/Runtime/RenderCore/Private/RenderCore.cpp b/Engine/Source/Runtime/RenderCore/Private/RenderCore.cpp
--- a/Engine/Source/Runtime/RenderCore/Private/RenderCore.cpp
+++ b/Engine/Source/Runtime/RenderCore/Private/RenderCore.cpp
@@ -185,6 +185,7 @@ DEFINE_STAT(STAT_AddScenePrimitiveGT);
 DEFINE_STAT(STAT_UpdatePrimitiveTransformGT);
 
 DEFINE_STAT(STAT_Scene_SetShaderMapsOnMaterialResources_RT);
+DEFINE_STAT(STAT_Scene_UpdateStaticDrawLists_RT);
 DEFINE_STAT(STAT_Scene_UpdateStaticDrawListsForMaterials_RT);
 DEFINE_STAT(STAT_GameToRendererMallocTotal);
 
diff --git a/Engine/Source/Runtime/RenderCore/Public/RenderCore.h b/Engine/Source/Runtime/RenderCore/Public/RenderCore.h
--- a/Engine/Source/Runtime/RenderCore/Public/RenderCore.h
+++ b/Engine/Source/Runtime/RenderCore/Public/RenderCore.h
@@ -136,6 +136,7 @@ DECLARE_CYCLE_STAT_EXTERN(TEXT("AddPrimitive (GT)"),STAT_AddScenePrimitiveGT,STA
 DECLARE_CYCLE_STAT_EXTERN(TEXT("UpdatePrimitiveTransform (GT)"),STAT_UpdatePrimitiveTransformGT,STATGROUP_SceneUpdate, RENDERCORE_API);
 
 DECLARE_CYCLE_STAT_EXTERN(TEXT("Set Shader Maps On Material Resources (RT)"),STAT_Scene_SetShaderMapsOnMaterialResources_RT,STATGROUP_SceneUpdate, RENDERCORE_API);
+DECLARE_CYCLE_STAT_EXTERN(TEXT("Update Static Draw Lists (RT)"), STAT_Scene_UpdateStaticDrawLists_RT, STATGROUP_SceneUpdate, RENDERCORE_API);
 DECLARE_CYCLE_STAT_EXTERN(TEXT("Update Static Draw Lists For Materials (RT)"),STAT_Scene_UpdateStaticDrawListsForMaterials_RT,STATGROUP_SceneUpdate, RENDERCORE_API);
 DECLARE_DWORD_COUNTER_STAT_EXTERN( TEXT( "Proxy Total" ), STAT_GameToRendererMallocTotal, STATGROUP_Memory, RENDERCORE_API );
 
diff --git a/Engine/Source/Runtime/RenderCore/Public/RendererInterface.h b/Engine/Source/Runtime/RenderCore/Public/RendererInterface.h
--- a/Engine/Source/Runtime/RenderCore/Public/RendererInterface.h
+++ b/Engine/Source/Runtime/RenderCore/Public/RendererInterface.h
@@ -687,6 +687,11 @@ public:
 	
 	virtual void RemoveScene(FSceneInterface* Scene) = 0;
 
+	/**
+	* Updates all static draw lists for each allocated scene.
+	*/
+	virtual void UpdateStaticDrawLists() = 0;
+
 	/**
 	 * Updates static draw lists for the given set of materials for each allocated scene.
 	 */
diff --git a/Engine/Source/Runtime/Renderer/Private/RendererModule.h b/Engine/Source/Runtime/Renderer/Private/RendererModule.h
--- a/Engine/Source/Runtime/Renderer/Private/RendererModule.h
+++ b/Engine/Source/Runtime/Renderer/Private/RendererModule.h
@@ -40,6 +40,7 @@ public:
 	virtual void CreateAndInitSingleView(FRHICommandListImmediate& RHICmdList, class FSceneViewFamily* ViewFamily, const struct FSceneViewInitOptions* ViewInitOptions) override;
 	virtual FSceneInterface* AllocateScene(UWorld* World, bool bInRequiresHitProxies, bool bCreateFXSystem, ERHIFeatureLevel::Type InFeatureLevel) override;
 	virtual void RemoveScene(FSceneInterface* Scene) override;
+	virtual void UpdateStaticDrawLists() override;
 	virtual void UpdateStaticDrawListsForMaterials(const TArray<const FMaterial*>& Materials) override;
 	virtual FSceneViewStateInterface* AllocateViewState() override;
 	virtual uint32 GetNumDynamicLightsAffectingPrimitive(const FPrimitiveSceneInfo* PrimitiveSceneInfo,const FLightCacheInterface* LCI) override;
diff --git a/Engine/Source/Runtime/Renderer/Private/RendererScene.cpp b/Engine/Source/Runtime/Renderer/Private/RendererScene.cpp
--- a/Engine/Source/Runtime/Renderer/Private/RendererScene.cpp
+++ b/Engine/Source/Runtime/Renderer/Private/RendererScene.cpp
@@ -2777,6 +2777,31 @@ void FScene::SetPrecomputedVisibility(const FPrecomputedVisibilityHandler* NewPr
 	});
 }
 
+void FScene::UpdateStaticDrawLists_RenderThread(FRHICommandListImmediate& RHICmdList)
+{
+	SCOPE_CYCLE_COUNTER(STAT_Scene_UpdateStaticDrawLists_RT);
+
+	const int32 NumPrimitives = Primitives.Num();
+
+	for (int32 PrimitiveIndex = 0; PrimitiveIndex < NumPrimitives; ++PrimitiveIndex)
+	{
+		FPrimitiveSceneInfo* Primitive = Primitives[PrimitiveIndex];
+
+		Primitive->RemoveStaticMeshes();
+		Primitive->AddStaticMeshes(RHICmdList);
+	}
+}
+
+void FScene::UpdateStaticDrawLists()
+{
+	ENQUEUE_UNIQUE_RENDER_COMMAND_ONEPARAMETER(
+		FUpdateDrawLists,
+		FScene*, Scene, this,
+		{
+			Scene->UpdateStaticDrawLists_RenderThread(RHICmdList);
+		});
+}
+
 void FScene::UpdateStaticDrawListsForMaterials_RenderThread(FRHICommandListImmediate& RHICmdList, const TArray<const FMaterial*>& Materials)
 {
 	SCOPE_CYCLE_COUNTER(STAT_Scene_UpdateStaticDrawListsForMaterials_RT);
@@ -2812,7 +2837,7 @@ void FScene::UpdateStaticDrawListsForMaterials_RenderThread(FRHICommandListImmed
 void FScene::UpdateStaticDrawListsForMaterials(const TArray<const FMaterial*>& Materials)
 {
 	ENQUEUE_UNIQUE_RENDER_COMMAND_TWOPARAMETER(
-		FUpdateDrawLists,
+		FUpdateDrawListsForMaterials,
 		FScene*,Scene,this,
 		TArray<const FMaterial*>,Materials,Materials,
 	{
@@ -3387,6 +3412,15 @@ void FRendererModule::RemoveScene(FSceneInterface* Scene)
 	AllocatedScenes.Remove(Scene);
 }
 
+void FRendererModule::UpdateStaticDrawLists()
+{
+	// Update all static meshes in order to recache cached mesh draw commands.
+	for (TSet<FSceneInterface*>::TConstIterator SceneIt(AllocatedScenes); SceneIt; ++SceneIt)
+	{
+		(*SceneIt)->UpdateStaticDrawLists();
+	}
+}
+
 void UpdateStaticMeshesForMaterials(const TArray<const FMaterial*>& MaterialResourcesToUpdate)
 {
 	TArray<UMaterialInterface*> UsedMaterials;
@@ -3446,7 +3480,7 @@ void FRendererModule::UpdateStaticDrawListsForMaterials(const TArray<const FMate
 		(*SceneIt)->UpdateStaticDrawListsForMaterials(Materials);
 	}
 
-	// Update static meshes in order to recache cached mesh draw commands.
+	// Update static meshes for a given set of materials in order to recache cached mesh draw commands.
 	UpdateStaticMeshesForMaterials(Materials);
 }
 
diff --git a/Engine/Source/Runtime/Renderer/Private/ScenePrivate.h b/Engine/Source/Runtime/Renderer/Private/ScenePrivate.h
--- a/Engine/Source/Runtime/Renderer/Private/ScenePrivate.h
+++ b/Engine/Source/Runtime/Renderer/Private/ScenePrivate.h
@@ -2727,6 +2727,9 @@ public:
 	/** Sets the precomputed visibility handler for the scene, or NULL to clear the current one. */
 	virtual void SetPrecomputedVisibility(const FPrecomputedVisibilityHandler* InPrecomputedVisibilityHandler) override;
 
+	/** Updates all static draw lists. */
+	virtual void UpdateStaticDrawLists() override;
+
 	/** Updates static draw lists for the given set of materials. */
 	virtual void UpdateStaticDrawListsForMaterials(const TArray<const FMaterial*>& Materials) override;
 
@@ -2937,6 +2940,9 @@ private:
 	/** Updates the contents of all reflection captures in the scene.  Must be called from the game thread. */
 	void UpdateAllReflectionCaptures(const TCHAR* CaptureReason, bool bVerifyOnlyCapturing);
 
+	/** Updates all static draw lists. */
+	void UpdateStaticDrawLists_RenderThread(FRHICommandListImmediate& RHICmdList);
+
 	/** Updates static draw lists for the given materials. */
 	void UpdateStaticDrawListsForMaterials_RenderThread(FRHICommandListImmediate& RHICmdList, const TArray<const FMaterial*>& Materials);
 
