Temp workaround for the apparently many older compilers not happy with the denoise code.

diff --git a/Engine/Source/Runtime/Renderer/Private/ScreenSpaceDenoise.cpp b/Engine/Source/Runtime/Renderer/Private/ScreenSpaceDenoise.cpp
--- a/Engine/Source/Runtime/Renderer/Private/ScreenSpaceDenoise.cpp
+++ b/Engine/Source/Runtime/Renderer/Private/ScreenSpaceDenoise.cpp
@@ -577,7 +577,7 @@ public:
 		FShadowPenumbraOutputs Outputs;
 
 		//temp workaround for compilation errors with android compiler
-#if !PLATFORM_ANDROID
+#if RHI_RAYTRACING
 		if (IsSupportedLightType(ELightComponentType(LightSceneInfo.Proxy->GetLightType())))
 		{
 			::DenoiseShadowPenumbra(
