VS2015 compile fixes

diff --git a/Engine/Source/Runtime/Engine/Private/SkeletalRenderCPUSkin.cpp b/Engine/Source/Runtime/Engine/Private/SkeletalRenderCPUSkin.cpp
--- a/Engine/Source/Runtime/Engine/Private/SkeletalRenderCPUSkin.cpp
+++ b/Engine/Source/Runtime/Engine/Private/SkeletalRenderCPUSkin.cpp
@@ -369,7 +369,7 @@ void FSkeletalMeshObjectCPUSkin::FSkeletalMeshObjectLOD::InitResources(FSkelMesh
 				GeometrySections.Reserve(RenderSections->Num());
 				for (const FSkelMeshRenderSection& Section : *RenderSections)
 				{
-					GeometrySections.Add(FRayTracingGeometrySegment{ Section.BaseIndex / 3, Section.NumTriangles });
+					GeometrySections.Add(FRayTracingGeometrySegment( (uint32)(Section.BaseIndex / 3), (uint32)(Section.NumTriangles) ));
 				}
 				Initializer.Segments = GeometrySections;
 
diff --git a/Engine/Source/Runtime/Engine/Private/SkeletalRenderGPUSkin.cpp b/Engine/Source/Runtime/Engine/Private/SkeletalRenderGPUSkin.cpp
--- a/Engine/Source/Runtime/Engine/Private/SkeletalRenderGPUSkin.cpp
+++ b/Engine/Source/Runtime/Engine/Private/SkeletalRenderGPUSkin.cpp
@@ -367,7 +367,7 @@ void FSkeletalMeshObjectGPUSkin::UpdateDynamicData_RenderThread(FGPUSkinCache* G
 				GeometrySections.Reserve(LODModel.RenderSections.Num());
 				for (const FSkelMeshRenderSection& Section : LODModel.RenderSections)
 				{
-					GeometrySections.Add(FRayTracingGeometrySegment{ Section.BaseIndex / 3, Section.NumTriangles });
+					GeometrySections.Add(FRayTracingGeometrySegment( (uint32)(Section.BaseIndex / 3), (uint32)(Section.NumTriangles) ));
 				}
 				Initializer.Segments = GeometrySections;
 
diff --git a/Engine/Source/Runtime/Engine/Private/SkeletalRenderStatic.cpp b/Engine/Source/Runtime/Engine/Private/SkeletalRenderStatic.cpp
--- a/Engine/Source/Runtime/Engine/Private/SkeletalRenderStatic.cpp
+++ b/Engine/Source/Runtime/Engine/Private/SkeletalRenderStatic.cpp
@@ -158,7 +158,7 @@ void FSkeletalMeshObjectStatic::FSkeletalMeshObjectLOD::InitResources(FSkelMeshC
 				GeometrySections.Reserve(RenderSections->Num());
 				for (const FSkelMeshRenderSection& Section : *RenderSections)
 				{
-					GeometrySections.Add(FRayTracingGeometrySegment{ Section.BaseIndex / 3, Section.NumTriangles });
+					GeometrySections.Add(FRayTracingGeometrySegment( (uint32)(Section.BaseIndex / 3), (uint32)Section.NumTriangles ));
 				}
 				Initializer.Segments = GeometrySections;
 
diff --git a/Engine/Source/Runtime/Engine/Private/StaticMesh.cpp b/Engine/Source/Runtime/Engine/Private/StaticMesh.cpp
--- a/Engine/Source/Runtime/Engine/Private/StaticMesh.cpp
+++ b/Engine/Source/Runtime/Engine/Private/StaticMesh.cpp
@@ -703,7 +703,7 @@ void FStaticMeshLODResources::InitResources(UStaticMesh* Parent)
 				GeometrySections.Reserve(Sections.Num());
 				for (const FStaticMeshSection& Section : Sections)
 				{
-					GeometrySections.Add(FRayTracingGeometrySegment{ Section.FirstIndex / 3, Section.NumTriangles});
+					GeometrySections.Add(FRayTracingGeometrySegment( (uint32)(Section.FirstIndex / 3), (uint32)(Section.NumTriangles)));
 				}
 				Initializer.Segments = GeometrySections;
 				
diff --git a/Engine/Source/Runtime/RHI/Public/DynamicRHI.h b/Engine/Source/Runtime/RHI/Public/DynamicRHI.h
--- a/Engine/Source/Runtime/RHI/Public/DynamicRHI.h
+++ b/Engine/Source/Runtime/RHI/Public/DynamicRHI.h
@@ -85,8 +85,18 @@ struct FRayTracingGeometryInstance
 
 struct FRayTracingGeometrySegment
 {
-	uint32 FirstPrimitive = 0;
-	uint32 NumPrimitives = 0;
+	FRayTracingGeometrySegment()
+		: FirstPrimitive(0)
+		, NumPrimitives(0)
+	{}
+
+	FRayTracingGeometrySegment(uint32 InFirstPrimitive, uint32 InNumPrimitives)
+		: FirstPrimitive(InFirstPrimitive)
+		, NumPrimitives(InNumPrimitives)
+	{}
+
+	uint32 FirstPrimitive;
+	uint32 NumPrimitives;
 
 	// #dxr_todo: what should be the sensible default values of these flags?
 
diff --git a/Engine/Source/Runtime/Renderer/Private/PathTracing/PathTracingTestbed.cpp b/Engine/Source/Runtime/Renderer/Private/PathTracing/PathTracingTestbed.cpp
--- a/Engine/Source/Runtime/Renderer/Private/PathTracing/PathTracingTestbed.cpp
+++ b/Engine/Source/Runtime/Renderer/Private/PathTracing/PathTracingTestbed.cpp
@@ -542,7 +542,7 @@ void TestPDFsIntegrateToOne(void)
 			}
 
 			float IntegralRes = PdfAccum / NSamples;
-			bool bPdfOK = abs(IntegralRes - 1.0f) < ValidPdfThreshold;
+			bool bPdfOK = fabsf(IntegralRes - 1.0f) < ValidPdfThreshold;
 			if (!bPdfOK)
 			{
 				Nfails++;
@@ -701,7 +701,7 @@ void TestBRDFandPDFConsistency(void)
 
 				if (ExpectedPercent > 0)
 				{
-					uint32 diff = abs(SampledPercent - ExpectedPercent);
+					uint32 diff = fabsf(SampledPercent - ExpectedPercent);
 					float Err = (float)diff;
 					if (Err > 20.0) //#dxr_todo: adjust threshold
 					{
