InitViews optimizations and velocity pass condition cleanup:
-Added bVelocityRelevance which prevents from dereferencing Proxy inside InitViews.
-Merged velocity pass conditions into two functions: PrimitiveHasVelocity and PrimitiveHasVelocityForView

diff --git a/Engine/Plugins/2D/Paper2D/Source/Paper2D/Private/PaperRenderSceneProxy.cpp b/Engine/Plugins/2D/Paper2D/Source/Paper2D/Private/PaperRenderSceneProxy.cpp
--- a/Engine/Plugins/2D/Paper2D/Source/Paper2D/Private/PaperRenderSceneProxy.cpp
+++ b/Engine/Plugins/2D/Paper2D/Source/Paper2D/Private/PaperRenderSceneProxy.cpp
@@ -434,6 +434,8 @@ FPrimitiveViewRelevance FPaperRenderSceneProxy::GetViewRelevance(const FSceneVie
 		Result.bOpaqueRelevance = true;
 	}
 
+	Result.bVelocityRelevance = IsMovable() && Result.bOpaqueRelevance && Result.bRenderInMainPass;
+
 	return Result;
 }
 
diff --git a/Engine/Plugins/Editor/MeshEditor/Source/MeshEditor/OverlayComponent.cpp b/Engine/Plugins/Editor/MeshEditor/Source/MeshEditor/OverlayComponent.cpp
--- a/Engine/Plugins/Editor/MeshEditor/Source/MeshEditor/OverlayComponent.cpp
+++ b/Engine/Plugins/Editor/MeshEditor/Source/MeshEditor/OverlayComponent.cpp
@@ -281,6 +281,7 @@ public:
 		Result.bRenderCustomDepth = ShouldRenderCustomDepth();
 		Result.bTranslucentSelfShadow = bCastVolumetricTranslucentShadow;
 		MaterialRelevance.SetPrimitiveViewRelevance( Result );
+		Result.bVelocityRelevance = IsMovable() && Result.bOpaqueRelevance && Result.bRenderInMainPass;
 		return Result;
 	}
 
diff --git a/Engine/Plugins/Editor/MeshEditor/Source/MeshEditor/WireframeMeshComponent.cpp b/Engine/Plugins/Editor/MeshEditor/Source/MeshEditor/WireframeMeshComponent.cpp
--- a/Engine/Plugins/Editor/MeshEditor/Source/MeshEditor/WireframeMeshComponent.cpp
+++ b/Engine/Plugins/Editor/MeshEditor/Source/MeshEditor/WireframeMeshComponent.cpp
@@ -371,6 +371,7 @@ public:
 		Result.bRenderCustomDepth = ShouldRenderCustomDepth();
 		Result.bTranslucentSelfShadow = bCastVolumetricTranslucentShadow;
 		MaterialRelevance.SetPrimitiveViewRelevance( Result );
+		Result.bVelocityRelevance = IsMovable() && Result.bOpaqueRelevance && Result.bRenderInMainPass;
 		return Result;
 	}
 
diff --git a/Engine/Plugins/Experimental/ImagePlate/Source/ImagePlate/Private/ImagePlateComponent.cpp b/Engine/Plugins/Experimental/ImagePlate/Source/ImagePlate/Private/ImagePlateComponent.cpp
--- a/Engine/Plugins/Experimental/ImagePlate/Source/ImagePlate/Private/ImagePlateComponent.cpp
+++ b/Engine/Plugins/Experimental/ImagePlate/Source/ImagePlate/Private/ImagePlateComponent.cpp
@@ -276,6 +276,9 @@ namespace
 			}
 			
 			MaterialRelevance.SetPrimitiveViewRelevance(Result);
+
+			Result.bVelocityRelevance = IsMovable() && Result.bOpaqueRelevance && Result.bRenderInMainPass;
+
 			return Result;
 		}
 		virtual bool CanBeOccluded() const override { return !MaterialRelevance.bDisableDepthTest; }
diff --git a/Engine/Plugins/Runtime/CustomMeshComponent/Source/CustomMeshComponent/Private/CustomMeshComponent.cpp b/Engine/Plugins/Runtime/CustomMeshComponent/Source/CustomMeshComponent/Private/CustomMeshComponent.cpp
--- a/Engine/Plugins/Runtime/CustomMeshComponent/Source/CustomMeshComponent/Private/CustomMeshComponent.cpp
+++ b/Engine/Plugins/Runtime/CustomMeshComponent/Source/CustomMeshComponent/Private/CustomMeshComponent.cpp
@@ -162,6 +162,7 @@ public:
 		Result.bRenderCustomDepth = ShouldRenderCustomDepth();
 		Result.bTranslucentSelfShadow = bCastVolumetricTranslucentShadow;
 		MaterialRelevance.SetPrimitiveViewRelevance(Result);
+		Result.bVelocityRelevance = IsMovable() && Result.bOpaqueRelevance && Result.bRenderInMainPass;
 		return Result;
 	}
 
diff --git a/Engine/Plugins/Runtime/ProceduralMeshComponent/Source/ProceduralMeshComponent/Private/ProceduralMeshComponent.cpp b/Engine/Plugins/Runtime/ProceduralMeshComponent/Source/ProceduralMeshComponent/Private/ProceduralMeshComponent.cpp
--- a/Engine/Plugins/Runtime/ProceduralMeshComponent/Source/ProceduralMeshComponent/Private/ProceduralMeshComponent.cpp
+++ b/Engine/Plugins/Runtime/ProceduralMeshComponent/Source/ProceduralMeshComponent/Private/ProceduralMeshComponent.cpp
@@ -355,6 +355,7 @@ public:
 		Result.bRenderCustomDepth = ShouldRenderCustomDepth();
 		Result.bTranslucentSelfShadow = bCastVolumetricTranslucentShadow;
 		MaterialRelevance.SetPrimitiveViewRelevance(Result);
+		Result.bVelocityRelevance = IsMovable() && Result.bOpaqueRelevance && Result.bRenderInMainPass;
 		return Result;
 	}
 
diff --git a/Engine/Source/Runtime/AugmentedReality/Private/ARPointCloudComponent.cpp b/Engine/Source/Runtime/AugmentedReality/Private/ARPointCloudComponent.cpp
--- a/Engine/Source/Runtime/AugmentedReality/Private/ARPointCloudComponent.cpp
+++ b/Engine/Source/Runtime/AugmentedReality/Private/ARPointCloudComponent.cpp
@@ -81,6 +81,7 @@ public:
 		Result.bRenderCustomDepth = false;
 		Result.bTranslucentSelfShadow = false;
 		MaterialRelevance.SetPrimitiveViewRelevance(Result);
+		Result.bVelocityRelevance = IsMovable() && Result.bOpaqueRelevance && Result.bRenderInMainPass;
 		return Result;
 	}
 
diff --git a/Engine/Source/Runtime/Engine/Private/Components/TextRenderComponent.cpp b/Engine/Source/Runtime/Engine/Private/Components/TextRenderComponent.cpp
--- a/Engine/Source/Runtime/Engine/Private/Components/TextRenderComponent.cpp
+++ b/Engine/Source/Runtime/Engine/Private/Components/TextRenderComponent.cpp
@@ -814,6 +814,7 @@ FPrimitiveViewRelevance FTextRenderSceneProxy::GetViewRelevance(const FSceneView
 	}
 
 	MaterialRelevance.SetPrimitiveViewRelevance(Result);
+	Result.bVelocityRelevance = IsMovable() && Result.bOpaqueRelevance && Result.bRenderInMainPass;
 	return Result;
 }
 
diff --git a/Engine/Source/Runtime/Engine/Private/Particles/ParticleSystemRender.cpp b/Engine/Source/Runtime/Engine/Private/Particles/ParticleSystemRender.cpp
--- a/Engine/Source/Runtime/Engine/Private/Particles/ParticleSystemRender.cpp
+++ b/Engine/Source/Runtime/Engine/Private/Particles/ParticleSystemRender.cpp
@@ -7220,6 +7220,8 @@ FPrimitiveViewRelevance FParticleSystemSceneProxy::GetViewRelevance(const FScene
 		Result.bOpaqueRelevance = true;
 	}
 
+	Result.bVelocityRelevance = IsMovable() && Result.bOpaqueRelevance && Result.bRenderInMainPass;
+
 	return Result;
 }
 
diff --git a/Engine/Source/Runtime/Engine/Private/SkeletalMesh.cpp b/Engine/Source/Runtime/Engine/Private/SkeletalMesh.cpp
--- a/Engine/Source/Runtime/Engine/Private/SkeletalMesh.cpp
+++ b/Engine/Source/Runtime/Engine/Private/SkeletalMesh.cpp
@@ -3912,6 +3912,8 @@ FPrimitiveViewRelevance FSkeletalMeshSceneProxy::GetViewRelevance(const FSceneVi
 	}
 #endif
 
+	Result.bVelocityRelevance = IsMovable() && Result.bOpaqueRelevance && Result.bRenderInMainPass;
+
 	return Result;
 }
 
diff --git a/Engine/Source/Runtime/Engine/Private/StaticMeshRender.cpp b/Engine/Source/Runtime/Engine/Private/StaticMeshRender.cpp
--- a/Engine/Source/Runtime/Engine/Private/StaticMeshRender.cpp
+++ b/Engine/Source/Runtime/Engine/Private/StaticMeshRender.cpp
@@ -1412,6 +1412,9 @@ FPrimitiveViewRelevance FStaticMeshSceneProxy::GetViewRelevance(const FSceneView
 	{
 		Result.bOpaqueRelevance = true;
 	}
+
+	Result.bVelocityRelevance = IsMovable() && Result.bOpaqueRelevance && Result.bRenderInMainPass;
+
 	return Result;
 }
 
diff --git a/Engine/Source/Runtime/Engine/Public/PrimitiveViewRelevance.h b/Engine/Source/Runtime/Engine/Public/PrimitiveViewRelevance.h
--- a/Engine/Source/Runtime/Engine/Public/PrimitiveViewRelevance.h
+++ b/Engine/Source/Runtime/Engine/Public/PrimitiveViewRelevance.h
@@ -45,6 +45,8 @@ struct FPrimitiveViewRelevance
 	uint32 bDrawRelevance : 1;
 	/** The primitive is casting a shadow. */
 	uint32 bShadowRelevance : 1;
+	/** The primitive should render velocity. */
+	uint32 bVelocityRelevance : 1;
 	/** The primitive should render to the custom depth pass. */
 	uint32 bRenderCustomDepth : 1;
 	/** The primitive should render to the base pass / normal depth / velocity rendering. */
diff --git a/Engine/Source/Runtime/Renderer/Private/PrimitiveSceneInfo.cpp b/Engine/Source/Runtime/Renderer/Private/PrimitiveSceneInfo.cpp
--- a/Engine/Source/Runtime/Renderer/Private/PrimitiveSceneInfo.cpp
+++ b/Engine/Source/Runtime/Renderer/Private/PrimitiveSceneInfo.cpp
@@ -681,60 +681,6 @@ uint32 FPrimitiveSceneInfo::GetMemoryFootprint()
 	return( sizeof( *this ) + HitProxies.GetAllocatedSize() + StaticMeshes.GetAllocatedSize() + StaticMeshRelevances.GetAllocatedSize() );
 }
 
-bool FPrimitiveSceneInfo::ShouldRenderVelocity(const FViewInfo& View, bool bCheckVisibility) const
-{
-	int32 PrimitiveId = GetIndex();
-	if (bCheckVisibility)
-	{
-		const bool bVisible = View.PrimitiveVisibilityMap[PrimitiveId];
-
-		// Only render if visible.
-		if (!bVisible)
-		{
-			return false;
-		}
-	}
-
-	const FPrimitiveViewRelevance& PrimitiveViewRelevance = View.PrimitiveViewRelevanceMap[PrimitiveId];
-
-	if (!Proxy->IsMovable())
-	{
-		return false;
-	}
-
-	// !Skip translucent objects as they don't support velocities and in the case of particles have a significant CPU overhead.
-	if (!PrimitiveViewRelevance.bOpaqueRelevance || !PrimitiveViewRelevance.bRenderInMainPass)
-	{
-		return false;
-	}
-
-	const float LODFactorDistanceSquared = (Proxy->GetBounds().Origin - View.ViewMatrices.GetViewOrigin()).SizeSquared() * FMath::Square(View.LODDistanceFactor);
-
-	// The minimum projected screen radius for a primitive to be drawn in the velocity pass, as a fraction of half the horizontal screen width (likely to be 0.08f)
-	float MinScreenRadiusForVelocityPass = View.FinalPostProcessSettings.MotionBlurPerObjectSize * (2.0f / 100.0f);
-	float MinScreenRadiusForVelocityPassSquared = FMath::Square(MinScreenRadiusForVelocityPass);
-
-	// Skip primitives that only cover a small amount of screenspace, motion blur on them won't be noticeable.
-	if (FMath::Square(Proxy->GetBounds().SphereRadius) <= MinScreenRadiusForVelocityPassSquared * LODFactorDistanceSquared)
-	{
-		return false;
-	}
-
-	// Only render primitives with velocity.
-	if (!FVelocityRendering::PrimitiveHasVelocity(View, this))
-	{
-		return false;
-	}
-
-	// If the base pass is allowed to render velocity in the GBuffer, only mesh with static lighting need the velocity pass.
-	if (FVelocityRendering::BasePassCanOutputVelocity(Scene->GetFeatureLevel()) && !(UseSelectiveBasePassOutputs() && Proxy->HasStaticLighting()))
-	{
-		return false;
-	}
-
-	return true;
-}
-
 void FPrimitiveSceneInfo::ApplyWorldOffset(FVector InOffset)
 {
 	Proxy->ApplyWorldOffset(InOffset);
diff --git a/Engine/Source/Runtime/Renderer/Private/SceneCore.h b/Engine/Source/Runtime/Renderer/Private/SceneCore.h
--- a/Engine/Source/Runtime/Renderer/Private/SceneCore.h
+++ b/Engine/Source/Runtime/Renderer/Private/SceneCore.h
@@ -207,6 +207,7 @@ public:
 		, LODIndex(StaticMesh.LODIndex)
 		, bDitheredLODTransition(StaticMesh.bDitheredLODTransition)
 		, bRequiresPerElementVisibility(StaticMesh.bRequiresPerElementVisibility)
+		, bSelectable(StaticMesh.bSelectable)
 		, CastShadow(StaticMesh.CastShadow)
 		, bUseForMaterial(StaticMesh.bUseForMaterial)
 		, bUseForDepthPass(StaticMesh.bUseForDepthPass)
@@ -230,6 +231,9 @@ public:
 	/** Whether the mesh batch needs VertexFactory->GetStaticBatchElementVisibility to be called each frame to determine which elements of the batch are visible. */
 	uint32 bRequiresPerElementVisibility : 1;
 
+	/** Whether the mesh batch can be selected through editor selection, aka hit proxies. */
+	uint32 bSelectable : 1;
+
 	uint32 CastShadow		: 1; // Whether it can be used in shadow renderpasses.
 	uint32 bUseForMaterial	: 1; // Whether it can be used in renderpasses requiring material outputs.
 	uint32 bUseForDepthPass	: 1; // Whether it can be used in depth pass.
diff --git a/Engine/Source/Runtime/Renderer/Private/SceneVisibility.cpp b/Engine/Source/Runtime/Renderer/Private/SceneVisibility.cpp
--- a/Engine/Source/Runtime/Renderer/Private/SceneVisibility.cpp
+++ b/Engine/Source/Runtime/Renderer/Private/SceneVisibility.cpp
@@ -2141,7 +2141,7 @@ struct FRelevancePacket
 								DrawCommandPacket.AddCommandsForMesh(PrimitiveIndex, StaticMeshRelevance, StaticMesh, Scene, bCanCache, EMeshPass::DebugViewMode);
 							}
 #endif
-							if (StaticMesh.bSelectable)
+							if (StaticMeshRelevance.bSelectable)
 							{
 								DrawCommandPacket.AddCommandsForMesh(PrimitiveIndex, StaticMeshRelevance, StaticMesh, Scene, bCanCache, EMeshPass::HitProxy);
 
@@ -2150,8 +2150,10 @@ struct FRelevancePacket
 									DrawCommandPacket.AddCommandsForMesh(PrimitiveIndex, StaticMeshRelevance, StaticMesh, Scene, bCanCache, EMeshPass::HitProxyOpaqueOnly);
 								}
 							}
-
-							if (PrimitiveSceneInfo->ShouldRenderVelocity(View, false))
+							 
+							if (ViewRelevance.bVelocityRelevance 
+								&& FVelocityRendering::PrimitiveHasVelocity(View.GetFeatureLevel(), PrimitiveSceneInfo)
+								&& FVelocityRendering::PrimitiveHasVelocityForView(View, Bounds.BoxSphereBounds, PrimitiveSceneInfo))
 							{
 								DrawCommandPacket.AddCommandsForMesh(PrimitiveIndex, StaticMeshRelevance, StaticMesh, Scene, bCanCache, EMeshPass::Velocity);
 								MarkMask |= EMarkMaskBits::StaticMeshVelocityMapMask;
diff --git a/Engine/Source/Runtime/Renderer/Private/VelocityRendering.cpp b/Engine/Source/Runtime/Renderer/Private/VelocityRendering.cpp
--- a/Engine/Source/Runtime/Renderer/Private/VelocityRendering.cpp
+++ b/Engine/Source/Runtime/Renderer/Private/VelocityRendering.cpp
@@ -567,13 +567,41 @@ bool FVelocityRendering::VertexFactoryOnlyOutputsVelocityInBasePass(EShaderPlatf
 	return BasePassCanOutputVelocity(ShaderPlatform) && !(UseSelectiveBasePassOutputs() && bVertexFactorySupportsStaticLighting);;
 }
 
-bool FVelocityRendering::PrimitiveHasVelocity(const FViewInfo& View, const FPrimitiveSceneInfo* PrimitiveSceneInfo)
+bool FVelocityRendering::PrimitiveHasVelocity(ERHIFeatureLevel::Type FeatureLevel, const FPrimitiveSceneInfo* PrimitiveSceneInfo)
 {
-	checkSlow(IsInParallelRenderingThread());
-	check(PrimitiveSceneInfo->Proxy);
+	// No velocity if motionblur is off, or if it's a non-moving object (treat as background in that case).
+	if (!GPixelFormats[PF_G16R16].Supported || !PrimitiveSceneInfo->Proxy->IsMovable())
+	{
+		return false;
+	}
+
+	// If the base pass is allowed to render velocity in the GBuffer, only mesh with static lighting need the velocity pass.
+	const bool bVelocityInGBuffer = FVelocityRendering::BasePassCanOutputVelocity(FeatureLevel) 
+		&& !(UseSelectiveBasePassOutputs() && (PrimitiveSceneInfo->Proxy->HasStaticLighting()));
+
+	if (bVelocityInGBuffer)
+	{
+		return false;
+	}
+
+	return true;
+}
 
-	// No velocity if motionblur is off, or if it's a non-moving object (treat as background in that case)
-	if (View.bCameraCut || !PrimitiveSceneInfo->Proxy->IsMovable())
+bool FVelocityRendering::PrimitiveHasVelocityForView(const FViewInfo& View, const FBoxSphereBounds& Bounds, const FPrimitiveSceneInfo* PrimitiveSceneInfo)
+{
+	if (View.bCameraCut)
+	{
+		return false;
+	}
+
+	const float LODFactorDistanceSquared = (Bounds.Origin - View.ViewMatrices.GetViewOrigin()).SizeSquared() * FMath::Square(View.LODDistanceFactor);
+
+	// The minimum projected screen radius for a primitive to be drawn in the velocity pass, as a fraction of half the horizontal screen width (likely to be 0.08f)
+	float MinScreenRadiusForVelocityPass = View.FinalPostProcessSettings.MotionBlurPerObjectSize * (2.0f / 100.0f);
+	float MinScreenRadiusForVelocityPassSquared = FMath::Square(MinScreenRadiusForVelocityPass);
+
+	// Skip primitives that only cover a small amount of screenspace, motion blur on them won't be noticeable.
+	if (FMath::Square(Bounds.SphereRadius) <= MinScreenRadiusForVelocityPassSquared * LODFactorDistanceSquared)
 	{
 		return false;
 	}
@@ -604,26 +632,19 @@ bool FVelocityRendering::PrimitiveHasVelocity(const FViewInfo& View, const FPrim
 
 void FVelocityMeshProcessor::AddMeshBatch(const FMeshBatch& RESTRICT MeshBatch, uint64 BatchElementMask, const FPrimitiveSceneProxy* RESTRICT PrimitiveSceneProxy, int32 StaticMeshId)
 {
-	// Velocity only needs to be directly rendered for movable meshes
 	bool bRequiresSeparateVelocity = false;
-	if (PrimitiveSceneProxy)
-	{
-		// When selective outputs are enable, only primitive with no static lighting output velocity in GBuffer.
-		const bool bVelocityInGBuffer = FVelocityRendering::BasePassCanOutputVelocity(FeatureLevel) && (!UseSelectiveBasePassOutputs() || !(PrimitiveSceneProxy->HasStaticLighting()));
 
-		bRequiresSeparateVelocity = GPixelFormats[PF_G16R16].Supported && PrimitiveSceneProxy->IsMovable() && !bVelocityInGBuffer;		
+	if (FVelocityRendering::PrimitiveHasVelocity(FeatureLevel, PrimitiveSceneProxy->GetPrimitiveSceneInfo()))
+	{
+		bRequiresSeparateVelocity = true;
 
 		// Cached mesh commands have identical check inside MarkRelevant.
-		if (bRequiresSeparateVelocity && ViewIfDynamicMeshCommand)
+		if (ViewIfDynamicMeshCommand)
 		{
 			checkSlow(ViewIfDynamicMeshCommand->bIsViewInfo);
-
 			FViewInfo* ViewInfo = (FViewInfo*)ViewIfDynamicMeshCommand;
 
-			if (!PrimitiveSceneProxy->GetPrimitiveSceneInfo()->ShouldRenderVelocity(*ViewInfo, false))
-			{
-				bRequiresSeparateVelocity = false;
-			}
+			bRequiresSeparateVelocity = FVelocityRendering::PrimitiveHasVelocityForView(*ViewInfo, PrimitiveSceneProxy->GetBounds(), PrimitiveSceneProxy->GetPrimitiveSceneInfo());
 		}
 	}
 
diff --git a/Engine/Source/Runtime/Renderer/Private/VelocityRendering.h b/Engine/Source/Runtime/Renderer/Private/VelocityRendering.h
--- a/Engine/Source/Runtime/Renderer/Private/VelocityRendering.h
+++ b/Engine/Source/Runtime/Renderer/Private/VelocityRendering.h
@@ -39,8 +39,11 @@ struct FVelocityRendering
 	/** Returns true the velocity is output in the BasePass. */
 	static bool VertexFactoryOnlyOutputsVelocityInBasePass(EShaderPlatform ShaderPlatform, bool bVertexFactorySupportsStaticLighting);
 
-	/** Determines whether this primitive has motionblur velocity to render. */
-	static bool PrimitiveHasVelocity(const FViewInfo& View, const FPrimitiveSceneInfo* PrimitiveSceneInfo);
+	/** Returns true if the object needs to be rendered in the velocity pass (is not moving like the world, needed for motionblur and TemporalAA). */
+	static bool PrimitiveHasVelocity(ERHIFeatureLevel::Type FeatureLevel, const FPrimitiveSceneInfo* PrimitiveSceneInfo);
+
+	/** Returns true if the object needs to be rendered in the velocity pass for a given view. */
+	static bool PrimitiveHasVelocityForView(const FViewInfo& View, const FBoxSphereBounds& Bounds, const FPrimitiveSceneInfo* PrimitiveSceneInfo);
 };
 
 
diff --git a/Engine/Source/Runtime/Renderer/Public/PrimitiveSceneInfo.h b/Engine/Source/Runtime/Renderer/Public/PrimitiveSceneInfo.h
--- a/Engine/Source/Runtime/Renderer/Public/PrimitiveSceneInfo.h
+++ b/Engine/Source/Runtime/Renderer/Public/PrimitiveSceneInfo.h
@@ -384,9 +384,6 @@ public:
 
 	/** Simple comparison against the invalid values used before/after scene add/remove. */
 	FORCEINLINE bool IsIndexValid() const { return PackedIndex != INDEX_NONE && PackedIndex != MAX_int32; }
-
-	/* @return true if the object needs to be rendered in the velocity pass (is not moving like the world, needed for motionblur and TemporalAA) */
-	bool ShouldRenderVelocity(const FViewInfo& View, bool bCheckVisibility = true) const;
 	
 	/**
 	 * Shifts primitive position and all relevant data by an arbitrary delta.
diff --git a/Engine/Source/Runtime/UMG/Private/Components/WidgetComponent.cpp b/Engine/Source/Runtime/UMG/Private/Components/WidgetComponent.cpp
--- a/Engine/Source/Runtime/UMG/Private/Components/WidgetComponent.cpp
+++ b/Engine/Source/Runtime/UMG/Private/Components/WidgetComponent.cpp
@@ -508,6 +508,7 @@ public:
 		Result.bShadowRelevance = IsShadowCast(View);
 		Result.bTranslucentSelfShadow = bCastVolumetricTranslucentShadow;
 		Result.bEditorPrimitiveRelevance = false;
+		Result.bVelocityRelevance = IsMovable() && Result.bOpaqueRelevance && Result.bRenderInMainPass;
 
 		return Result;
 	}
