Parallel mesh draw mesh draw command sorting and processing for instancing.

diff --git a/Engine/Source/Runtime/Renderer/Private/MeshPassProcessor.cpp b/Engine/Source/Runtime/Renderer/Private/MeshPassProcessor.cpp
--- a/Engine/Source/Runtime/Renderer/Private/MeshPassProcessor.cpp
+++ b/Engine/Source/Runtime/Renderer/Private/MeshPassProcessor.cpp
@@ -945,12 +945,16 @@ void SubmitMeshDrawCommandsForView(
 	FParallelCommandListSet* ParallelCommandListSet,
 	FRHICommandList& RHICmdList)
 {
-	FVertexBufferRHIParamRef PrimitiveIdsBuffer = UseGPUScene(GMaxRHIShaderPlatform, View.GetFeatureLevel()) ? View.VisibleMeshDrawCommandPrimitiveIdBuffers[PassType].VertexBuffer->VertexBufferRHI : nullptr;
-	const int32 BasePrimitiveIdsOffset = View.VisibleMeshDrawCommandPrimitiveIdBuffers[PassType].VertexOffset;
-	const bool bDynamicInstancing = IsDynamicInstancingEnabled() && UseGPUScene(GMaxRHIShaderPlatform, View.GetFeatureLevel());
 	const FMeshCommandOneFrameArray& VisibleMeshDrawCommands = View.VisibleMeshDrawCommands[PassType];
 
-	SubmitMeshDrawCommands(VisibleMeshDrawCommands, PrimitiveIdsBuffer, BasePrimitiveIdsOffset, bDynamicInstancing, ParallelCommandListSet, RHICmdList);
+	if (VisibleMeshDrawCommands.Num() > 0)
+	{
+		FVertexBufferRHIParamRef PrimitiveIdsBuffer = UseGPUScene(GMaxRHIShaderPlatform, View.GetFeatureLevel()) ? View.VisibleMeshDrawCommandPrimitiveIdBuffers[PassType].VertexBuffer->VertexBufferRHI : nullptr;
+		const int32 BasePrimitiveIdsOffset = View.VisibleMeshDrawCommandPrimitiveIdBuffers[PassType].VertexOffset;
+		const bool bDynamicInstancing = IsDynamicInstancingEnabled() && UseGPUScene(GMaxRHIShaderPlatform, View.GetFeatureLevel());
+
+		SubmitMeshDrawCommands(VisibleMeshDrawCommands, PrimitiveIdsBuffer, BasePrimitiveIdsOffset, bDynamicInstancing, ParallelCommandListSet, RHICmdList);
+	}
 }
 
 FMeshPassProcessor::FMeshPassProcessor(const FScene* InScene, ERHIFeatureLevel::Type InFeatureLevel, const FSceneView* InViewIfDynamicMeshCommand, FMeshPassDrawListContext& InDrawListContext) 
diff --git a/Engine/Source/Runtime/Renderer/Private/SceneRendering.cpp b/Engine/Source/Runtime/Renderer/Private/SceneRendering.cpp
--- a/Engine/Source/Runtime/Renderer/Private/SceneRendering.cpp
+++ b/Engine/Source/Runtime/Renderer/Private/SceneRendering.cpp
@@ -99,6 +99,12 @@ static TAutoConsoleVariable<int32> CVarMeshDrawCommandsDynamicInstancing(
 	TEXT("Whether to dynamically combine multiple compatible visible Mesh Draw Commands into one instanced draw on vertex factories that support it."),
 	ECVF_RenderThreadSafe);
 
+static TAutoConsoleVariable<int32> CVarMeshDrawCommandsParallelSort(
+	TEXT("r.MeshDrawCommands.ParallelSort"),
+	1,
+	TEXT("Whether to sort and process mesh draw commands in parallel."),
+	ECVF_RenderThreadSafe);
+
 bool IsDynamicInstancingEnabled()
 {
 	return CVarMeshDrawCommandsDynamicInstancing.GetValueOnRenderThread() > 0;
@@ -2854,7 +2860,7 @@ void SortPassMeshDrawCommands(
 					}
 				}
 
-				//@todo - refactor into memcpy
+				//@todo MeshCommandPipeline - refactor into memcpy
 				PrimitiveIds[DrawCommandIndex] = VisibleMeshDrawCommand.DrawPrimitiveId;
 			}
 
@@ -2929,19 +2935,188 @@ void UpdateTranslucentMeshSortKeys(FScene* Scene, FViewInfo& View, ETranslucency
 	}
 }
 
+/**
+ * Packet for parallel mesh draw command sorting and processing for instancing.
+ */
+class FMeshDrawCommandSortPacket
+{
+public:
+	FMeshDrawCommandSortPacket(FScene* InScene, 
+		FViewInfo& InView, 
+		bool InbUseGPUScene, 
+		EMeshPass::Type InPassType,
+		ETranslucencyPass::Type InTranslucencyPass, 
+		FMeshCommandOneFrameArray& InVisibleMeshDrawCommands, 
+		FDynamicMeshDrawCommandStorage& InMeshDrawCommandStorage,
+		FGlobalDynamicVertexBuffer::FAllocation& InPrimitiveIdBuffer)
+		: Scene(InScene)
+		, View(InView)
+		, bUseGPUScene(InbUseGPUScene)
+		, PassType(InPassType)
+		, TranslucencyPass(InTranslucencyPass)
+		, VisibleMeshDrawCommands(InVisibleMeshDrawCommands)
+		, MeshDrawCommandStorage(InMeshDrawCommandStorage)
+		, PrimitiveIdBuffer(InPrimitiveIdBuffer)
+		, VisibleMeshDrawCommandsNum(0)
+		, NewPassVisibleMeshDrawCommandsNum(0)
+		, MaxInstances(1)
+	{
+		// Preallocate NewPassVisibleMeshDrawCommands on rendering thread.
+		if (IsDynamicInstancingEnabled())
+		{
+			NewPassVisibleMeshDrawCommands.Empty(InVisibleMeshDrawCommands.Num());
+		}
+	}
+
+	void AnyThreadTask()
+	{
+		if (TranslucencyPass != ETranslucencyPass::TPT_MAX)
+		{
+			QUICK_SCOPE_CYCLE_COUNTER(STAT_UpdateTranslucentMeshSortKeys);
+			UpdateTranslucentMeshSortKeys(Scene, View, TranslucencyPass, VisibleMeshDrawCommands);
+		}
+
+		{
+			QUICK_SCOPE_CYCLE_COUNTER(STAT_SortVisibleMeshDrawCommands);
+			VisibleMeshDrawCommands.Sort(FCompareFMeshDrawCommands());
+		}
+
+		if (bUseGPUScene)
+		{
+			const FVisibleMeshDrawCommand* RESTRICT PassVisibleMeshDrawCommands = VisibleMeshDrawCommands.GetData();
+			const int32 NumDrawCommands = VisibleMeshDrawCommands.Num();
+
+			int32* RESTRICT PrimitiveIds = (int32*)PrimitiveIdBuffer.Buffer;
+
+			if (IsDynamicInstancingEnabled())
+			{
+				QUICK_SCOPE_CYCLE_COUNTER(STAT_DynamicInstancingOfVisibleMeshDrawCommands);
+
+				int32 CurrentStateBucketId = -1;
+				uint32* RESTRICT CurrentDynamicallyInstancedMeshCommandNumInstances = nullptr;
+				MaxInstances = 1;
+
+				for (int32 DrawCommandIndex = 0; DrawCommandIndex < NumDrawCommands; DrawCommandIndex++)
+				{
+					const FVisibleMeshDrawCommand& RESTRICT VisibleMeshDrawCommand = PassVisibleMeshDrawCommands[DrawCommandIndex];
+
+					if (VisibleMeshDrawCommand.StateBucketId == CurrentStateBucketId && VisibleMeshDrawCommand.StateBucketId != -1)
+					{
+						if (CurrentDynamicallyInstancedMeshCommandNumInstances)
+						{
+							const int32 CurrentNumInstances = *CurrentDynamicallyInstancedMeshCommandNumInstances;
+							*CurrentDynamicallyInstancedMeshCommandNumInstances = CurrentNumInstances + 1;
+							MaxInstances = FMath::Max(MaxInstances, CurrentNumInstances + 1);
+						}
+						else
+						{
+							FVisibleMeshDrawCommand NewVisibleMeshDrawCommand = VisibleMeshDrawCommand;
+							NewVisibleMeshDrawCommand.PrimitiveIdBufferOffset = DrawCommandIndex;
+							NewPassVisibleMeshDrawCommands.Emplace(MoveTemp(NewVisibleMeshDrawCommand));
+						}
+					}
+					else
+					{
+						// First time state bucket setup
+						CurrentStateBucketId = VisibleMeshDrawCommand.StateBucketId;
+
+						if (VisibleMeshDrawCommand.MeshDrawCommand->PrimitiveIdStreamIndex >= 0
+							&& VisibleMeshDrawCommand.MeshDrawCommand->NumInstances == 1
+							// Don't create a new FMeshDrawCommand for the last command and make it safe for us to look at the next command
+							&& DrawCommandIndex + 1 <= NumDrawCommands
+							// Only create a new FMeshDrawCommand if more than one draw in the state bucket
+							&& CurrentStateBucketId == PassVisibleMeshDrawCommands[DrawCommandIndex + 1].StateBucketId)
+						{
+							const int32 Index = MeshDrawCommandStorage.MeshDrawCommands.AddElement(*VisibleMeshDrawCommand.MeshDrawCommand);
+							FMeshDrawCommand& NewCommand = MeshDrawCommandStorage.MeshDrawCommands[Index];
+							FVisibleMeshDrawCommand NewVisibleMeshDrawCommand;
+
+							NewVisibleMeshDrawCommand.Setup(
+								&NewCommand,
+								VisibleMeshDrawCommand.DrawPrimitiveId,
+								VisibleMeshDrawCommand.StateBucketId,
+								VisibleMeshDrawCommand.MeshFillMode,
+								VisibleMeshDrawCommand.MeshCullMode,
+								VisibleMeshDrawCommand.SortKey);
+
+							NewVisibleMeshDrawCommand.PrimitiveIdBufferOffset = DrawCommandIndex;
+							NewPassVisibleMeshDrawCommands.Emplace(MoveTemp(NewVisibleMeshDrawCommand));
+
+							CurrentDynamicallyInstancedMeshCommandNumInstances = &NewCommand.NumInstances;
+						}
+						else
+						{
+							CurrentDynamicallyInstancedMeshCommandNumInstances = nullptr;
+							FVisibleMeshDrawCommand NewVisibleMeshDrawCommand = VisibleMeshDrawCommand;
+							NewVisibleMeshDrawCommand.PrimitiveIdBufferOffset = DrawCommandIndex;
+							NewPassVisibleMeshDrawCommands.Emplace(MoveTemp(NewVisibleMeshDrawCommand));
+						}
+					}
+
+					//@todo - refactor into memcpy
+					PrimitiveIds[DrawCommandIndex] = VisibleMeshDrawCommand.DrawPrimitiveId;
+				}
+
+				// Setup instancing stats for logging.
+				VisibleMeshDrawCommandsNum = VisibleMeshDrawCommands.Num();
+				NewPassVisibleMeshDrawCommandsNum = NewPassVisibleMeshDrawCommands.Num();
+
+				// Replace VisibleMeshDrawCommands
+				FMemory::Memswap(&VisibleMeshDrawCommands, &NewPassVisibleMeshDrawCommands, sizeof(NewPassVisibleMeshDrawCommands));
+			}
+			else
+			{
+				QUICK_SCOPE_CYCLE_COUNTER(STAT_BuildVisibleMeshDrawCommandPrimitiveIdBuffers);
+
+				for (int32 DrawCommandIndex = 0; DrawCommandIndex < NumDrawCommands; DrawCommandIndex++)
+				{
+					const FVisibleMeshDrawCommand& VisibleMeshDrawCommand = VisibleMeshDrawCommands[DrawCommandIndex];
+					PrimitiveIds[DrawCommandIndex] = VisibleMeshDrawCommand.DrawPrimitiveId;
+				}
+			}
+		}
+	}
+
+	void DumpMeshDrawCommandInstancingStats()
+	{
+		UE_LOG(LogRenderer, Log, TEXT("%s"), GetMeshPassName(PassType));
+		UE_LOG(LogRenderer, Log, TEXT("   %i Mesh Draw Commands in %i instancing state buckets"), VisibleMeshDrawCommandsNum, NewPassVisibleMeshDrawCommandsNum);
+		UE_LOG(LogRenderer, Log, TEXT("   Largest %i"), MaxInstances);
+		UE_LOG(LogRenderer, Log, TEXT("   %.1f Dynamic Instancing draw call reduction factor"), VisibleMeshDrawCommandsNum / (float)NewPassVisibleMeshDrawCommandsNum);
+	}
+
+
+private:
+	FScene* Scene;
+	FViewInfo& View;
+	const bool bUseGPUScene;
+	EMeshPass::Type PassType;
+	ETranslucencyPass::Type TranslucencyPass;
+	FMeshCommandOneFrameArray& VisibleMeshDrawCommands;
+	FDynamicMeshDrawCommandStorage& MeshDrawCommandStorage;
+	FGlobalDynamicVertexBuffer::FAllocation& PrimitiveIdBuffer;
+	FMeshCommandOneFrameArray NewPassVisibleMeshDrawCommands;
+
+	// For logging instancing stats.
+	int32 VisibleMeshDrawCommandsNum;
+	int32 NewPassVisibleMeshDrawCommandsNum;
+	int32 MaxInstances;
+};
+
 void FSceneRenderer::SortMeshDrawCommands()
 {
-	EShadingPath ShadingPath = Scene->GetShadingPath();
-	
+	QUICK_SCOPE_CYCLE_COUNTER(STAT_ProcessVisibleMeshDrawCommands);
+
+	const EShadingPath ShadingPath = Scene->GetShadingPath();
+	const bool bUseGPUScene = UseGPUScene(GMaxRHIShaderPlatform, FeatureLevel);
+	const bool bExecuteInParallel = FApp::ShouldUseThreadingForPerformance() && CVarMeshDrawCommandsParallelSort.GetValueOnRenderThread() > 0;
+
+	TArray<FMeshDrawCommandSortPacket, SceneRenderingAllocator> Packets;
+	Packets.Reserve(Views.Num() * EMeshPass::Num);
+
 	for (int32 ViewIndex = 0; ViewIndex < Views.Num(); ViewIndex++)
 	{
 		FViewInfo& View = Views[ViewIndex];
-	
-		QUICK_SCOPE_CYCLE_COUNTER(STAT_ProcessVisibleMeshDrawCommands);
-
-		UpdateTranslucentMeshSortKeys(Scene, View, ETranslucencyPass::TPT_StandardTranslucency, View.VisibleMeshDrawCommands[EMeshPass::TranslucencyStandard]);
-		UpdateTranslucentMeshSortKeys(Scene, View, ETranslucencyPass::TPT_TranslucencyAfterDOF, View.VisibleMeshDrawCommands[EMeshPass::TranslucencyAfterDOF]);
-		UpdateTranslucentMeshSortKeys(Scene, View, ETranslucencyPass::TPT_AllTranslucency, View.VisibleMeshDrawCommands[EMeshPass::TranslucencyAll]);
 
 		for (int32 PassIndex = 0; PassIndex < EMeshPass::Num; PassIndex++)
 		{
@@ -2949,21 +3124,54 @@ void FSceneRenderer::SortMeshDrawCommands()
 
 			if ((FPassProcessorManager::GetPassFlags(ShadingPath, PassType) & EMeshPassFlags::MainView) != EMeshPassFlags::None)
 			{
-				if (ShouldDumpMeshDrawCommandInstancingStats() && View.VisibleMeshDrawCommands[PassIndex].Num() > 0)
+				ETranslucencyPass::Type TranslucencyPass = ETranslucencyPass::TPT_MAX;
+				switch (PassType)
 				{
-					UE_LOG(LogRenderer, Log, TEXT("%s"), GetMeshPassName(PassType));
+					case EMeshPass::TranslucencyStandard:	TranslucencyPass = ETranslucencyPass::TPT_StandardTranslucency; break;
+					case EMeshPass::TranslucencyAfterDOF:	TranslucencyPass = ETranslucencyPass::TPT_TranslucencyAfterDOF; break;
+					case EMeshPass::TranslucencyAll:		TranslucencyPass = ETranslucencyPass::TPT_AllTranslucency;		break;
 				}
 
- 				SortPassMeshDrawCommands(
-					FeatureLevel,
-					View.VisibleMeshDrawCommands[PassIndex],
-					View.DynamicMeshDrawCommandStorage[PassIndex],
-					FGlobalDynamicVertexBuffer::Get(),
-					View.VisibleMeshDrawCommandPrimitiveIdBuffers[PassIndex],
-					ShouldDumpMeshDrawCommandInstancingStats());
+				const int32 NumDrawCommands = View.VisibleMeshDrawCommands[PassIndex].Num();
+				if (NumDrawCommands > 0)
+				{
+					// Preallocate primitive Id buffer on rendering thread.
+					if (bUseGPUScene)
+					{
+						View.VisibleMeshDrawCommandPrimitiveIdBuffers[PassIndex] = FGlobalDynamicVertexBuffer::Get().Allocate(NumDrawCommands * sizeof(int32));
+					}
+
+					FMeshDrawCommandSortPacket Packet(
+						Scene, 
+						View, 
+						bUseGPUScene, 
+						PassType,
+						TranslucencyPass, 
+						View.VisibleMeshDrawCommands[PassIndex],
+						View.DynamicMeshDrawCommandStorage[PassIndex],
+						View.VisibleMeshDrawCommandPrimitiveIdBuffers[PassIndex]);
+
+					Packets.Emplace(MoveTemp(Packet));
+				}
 			}
 		}
 	}
+
+	ParallelFor(Packets.Num(),
+		[&Packets](int32 Index)
+		{
+			Packets[Index].AnyThreadTask();
+		},
+		!bExecuteInParallel
+	);
+
+	if (ShouldDumpMeshDrawCommandInstancingStats())
+	{
+		for (int32 PacketIndex = 0; PacketIndex < Packets.Num(); ++PacketIndex)
+		{
+			Packets[PacketIndex].DumpMeshDrawCommandInstancingStats();
+		}
+	}
 }
 
 FSceneRenderer* FSceneRenderer::CreateSceneRenderer(const FSceneViewFamily* InViewFamily, FHitProxyConsumer* HitProxyConsumer)
