Fixed single shadowmap directional shadow map rendering. In case of a directional light CSMShadowDepthViewUniformBuffer was always used, but if bWholeSceneShadow was false, then it was never updated before rendering.

diff --git a/Engine/Source/Runtime/Renderer/Private/ShadowDepthRendering.cpp b/Engine/Source/Runtime/Renderer/Private/ShadowDepthRendering.cpp
--- a/Engine/Source/Runtime/Renderer/Private/ShadowDepthRendering.cpp
+++ b/Engine/Source/Runtime/Renderer/Private/ShadowDepthRendering.cpp
@@ -1974,7 +1974,7 @@ void FProjectedShadowInfo::ModifyViewForShadow(FRHICommandList& RHICmdList, FVie
 		TVC_MAX,
 		*FoundView->CachedViewUniformShaderParameters);
 
-	if (GetShadowDepthType() == CSMShadowDepthType)
+	if (IsWholeSceneDirectionalShadow())
 	{
 		FScene* Scene = (FScene*)FoundView->Family->Scene;
 		FoundView->ViewUniformBuffer = Scene->UniformBuffers.CSMShadowDepthViewUniformBuffer;
