Fixed some MSVC static analysis warnings.

diff --git a/Engine/Source/Runtime/D3D12RHI/Private/D3D12RayTracing.cpp b/Engine/Source/Runtime/D3D12RHI/Private/D3D12RayTracing.cpp
--- a/Engine/Source/Runtime/D3D12RHI/Private/D3D12RayTracing.cpp
+++ b/Engine/Source/Runtime/D3D12RHI/Private/D3D12RayTracing.cpp
@@ -664,7 +664,9 @@ public:
 		FD3D12Device* Device = GetParentDevice();
 
 		checkf(Data.Num(), TEXT("Shader table is expected to be initialized before copying to GPU."));
+#if DO_CHECK && DO_GUARD_SLOW
 		checkfSlow(NumValidHitGroupRecords == NumHitGroups, TEXT("Hit group shader table is expected to be fully populated before copying to GPU. Ensure that all shader records have been fille using WriteShaderIdentifier()."));
+#endif
 
 		FD3D12Adapter* Adapter = Device->GetParentAdapter();
 
diff --git a/Engine/Source/Runtime/Renderer/Private/MeshPassProcessor.cpp b/Engine/Source/Runtime/Renderer/Private/MeshPassProcessor.cpp
--- a/Engine/Source/Runtime/Renderer/Private/MeshPassProcessor.cpp
+++ b/Engine/Source/Runtime/Renderer/Private/MeshPassProcessor.cpp
@@ -83,6 +83,9 @@ void FMeshDrawShaderBindings::SetShaderBindings(
 	const FReadOnlyMeshDrawSingleShaderBindings& RESTRICT SingleShaderBindings,
 	FShaderBindingState& RESTRICT ShaderBindingState)
 {
+#pragma warning(push)
+#pragma warning(disable : 6386) // Suppress MSVC static analysis warning that Parameter.BaseIndex can index arrays out of bounds.
+
 	const FUniformBufferRHIParamRef* RESTRICT UniformBufferBindings = SingleShaderBindings.GetUniformBufferStart();
 	const FShaderParameterInfo* RESTRICT UniformBufferParameters = SingleShaderBindings.ParameterMapInfo.UniformBuffers.GetData();
 	const int32 NumUniformBuffers = SingleShaderBindings.ParameterMapInfo.UniformBuffers.Num();
@@ -170,6 +173,8 @@ void FMeshDrawShaderBindings::SetShaderBindings(
 			LooseDataStart += Parameter.Size;
 		}
 	}
+
+#pragma warning(pop)
 }
 
 template<class RHIShaderType>
@@ -178,6 +183,9 @@ void FMeshDrawShaderBindings::SetShaderBindings(
 	RHIShaderType Shader,
 	const FReadOnlyMeshDrawSingleShaderBindings& RESTRICT SingleShaderBindings)
 {
+#pragma warning(push)
+#pragma warning(disable : 6386) // Suppress MSVC static analysis warning that Parameter.BaseIndex can index arrays out of bounds.
+
 	const FUniformBufferRHIParamRef* RESTRICT UniformBufferBindings = SingleShaderBindings.GetUniformBufferStart();
 	const FShaderParameterInfo* RESTRICT UniformBufferParameters = SingleShaderBindings.ParameterMapInfo.UniformBuffers.GetData();
 	const int32 NumUniformBuffers = SingleShaderBindings.ParameterMapInfo.UniformBuffers.Num();
@@ -250,6 +258,8 @@ void FMeshDrawShaderBindings::SetShaderBindings(
 			LooseDataStart += Parameter.Size;
 		}
 	}
+
+#pragma warning(pop)
 }
 
 #if RHI_RAYTRACING
diff --git a/Engine/Source/Runtime/Renderer/Private/MobileBasePass.cpp b/Engine/Source/Runtime/Renderer/Private/MobileBasePass.cpp
--- a/Engine/Source/Runtime/Renderer/Private/MobileBasePass.cpp
+++ b/Engine/Source/Runtime/Renderer/Private/MobileBasePass.cpp
@@ -622,7 +622,7 @@ void TMobileBasePassPSPolicyParamType<FUniformLightMapPolicy>::GetShaderBindings
 	}
 
 	// Set directional light UB
-	if (MobileDirectionLightBufferParam.IsBound())
+	if (MobileDirectionLightBufferParam.IsBound() && Scene)
 	{
 		int32 UniformBufferIndex = PrimitiveSceneProxy ? GetFirstLightingChannelFromMask(PrimitiveSceneProxy->GetLightingChannelMask()) + 1 : 0;
 		ShaderBindings.Add(MobileDirectionLightBufferParam, Scene->UniformBuffers.MobileDirectionalLightUniformBuffers[UniformBufferIndex]);
diff --git a/Engine/Source/Runtime/Renderer/Private/TranslucentRendering.cpp b/Engine/Source/Runtime/Renderer/Private/TranslucentRendering.cpp
--- a/Engine/Source/Runtime/Renderer/Private/TranslucentRendering.cpp
+++ b/Engine/Source/Runtime/Renderer/Private/TranslucentRendering.cpp
@@ -761,7 +761,8 @@ void CreateTranslucentBasePassUniformBuffer(
 	BasePassParameters.SceneTextures.SceneColorCopyTexture = SceneColorCopy ? SceneColorCopy->GetRenderTargetItem().ShaderResourceTexture : GBlackTexture->TextureRHI;
 
 
-	FScene* Scene = View.Family->Scene ? View.Family->Scene->GetRenderScene() : nullptr;
+	check(View.Family->Scene);
+	FScene* Scene = View.Family->Scene->GetRenderScene();
 
 	Scene->UniformBuffers.TranslucentBasePassUniformBuffer.UpdateUniformBufferImmediate(BasePassParameters);
 	BasePassUniformBuffer = Scene->UniformBuffers.TranslucentBasePassUniformBuffer;
