Removed pragmas, as they didn't suppress MSVC static analysis out of bounds warnings inside MeshPassProcessor.

diff --git a/Engine/Source/Runtime/Renderer/Private/MeshPassProcessor.cpp b/Engine/Source/Runtime/Renderer/Private/MeshPassProcessor.cpp
--- a/Engine/Source/Runtime/Renderer/Private/MeshPassProcessor.cpp
+++ b/Engine/Source/Runtime/Renderer/Private/MeshPassProcessor.cpp
@@ -83,9 +83,6 @@ void FMeshDrawShaderBindings::SetShaderBindings(
 	const FReadOnlyMeshDrawSingleShaderBindings& RESTRICT SingleShaderBindings,
 	FShaderBindingState& RESTRICT ShaderBindingState)
 {
-#pragma warning(push)
-#pragma warning(disable : 6386) // Suppress MSVC static analysis warning that Parameter.BaseIndex can index arrays out of bounds.
-
 	const FUniformBufferRHIParamRef* RESTRICT UniformBufferBindings = SingleShaderBindings.GetUniformBufferStart();
 	const FShaderParameterInfo* RESTRICT UniformBufferParameters = SingleShaderBindings.ParameterMapInfo.UniformBuffers.GetData();
 	const int32 NumUniformBuffers = SingleShaderBindings.ParameterMapInfo.UniformBuffers.Num();
@@ -173,8 +170,6 @@ void FMeshDrawShaderBindings::SetShaderBindings(
 			LooseDataStart += Parameter.Size;
 		}
 	}
-
-#pragma warning(pop)
 }
 
 template<class RHIShaderType>
@@ -183,9 +178,6 @@ void FMeshDrawShaderBindings::SetShaderBindings(
 	RHIShaderType Shader,
 	const FReadOnlyMeshDrawSingleShaderBindings& RESTRICT SingleShaderBindings)
 {
-#pragma warning(push)
-#pragma warning(disable : 6386) // Suppress MSVC static analysis warning that Parameter.BaseIndex can index arrays out of bounds.
-
 	const FUniformBufferRHIParamRef* RESTRICT UniformBufferBindings = SingleShaderBindings.GetUniformBufferStart();
 	const FShaderParameterInfo* RESTRICT UniformBufferParameters = SingleShaderBindings.ParameterMapInfo.UniformBuffers.GetData();
 	const int32 NumUniformBuffers = SingleShaderBindings.ParameterMapInfo.UniformBuffers.Num();
@@ -258,8 +250,6 @@ void FMeshDrawShaderBindings::SetShaderBindings(
 			LooseDataStart += Parameter.Size;
 		}
 	}
-
-#pragma warning(pop)
 }
 
 #if RHI_RAYTRACING
diff --git a/Engine/Source/Runtime/Renderer/Private/SceneRendering.cpp b/Engine/Source/Runtime/Renderer/Private/SceneRendering.cpp
--- a/Engine/Source/Runtime/Renderer/Private/SceneRendering.cpp
+++ b/Engine/Source/Runtime/Renderer/Private/SceneRendering.cpp
@@ -2780,7 +2780,6 @@ void SortPassMeshDrawCommands(
 	{
 		QUICK_SCOPE_CYCLE_COUNTER(STAT_SortVisibleMeshDrawCommands);
 
-		//@todo MeshCommandPipeline - sort in parallel during visibility or bucket in parallel
 		VisibleMeshDrawCommands.Sort(FCompareFMeshDrawCommands());
 	}
 
