Fixed out of bounds array access during visible mesh draw command processing for instancing. Also renamed FMeshDrawCommandSortPacket to FVisibleMeshDrawCommandProcessPacket, as it's doing more things than sorting.

diff --git a/Engine/Source/Runtime/Renderer/Private/SceneRendering.cpp b/Engine/Source/Runtime/Renderer/Private/SceneRendering.cpp
--- a/Engine/Source/Runtime/Renderer/Private/SceneRendering.cpp
+++ b/Engine/Source/Runtime/Renderer/Private/SceneRendering.cpp
@@ -99,10 +99,10 @@ static TAutoConsoleVariable<int32> CVarMeshDrawCommandsDynamicInstancing(
 	TEXT("Whether to dynamically combine multiple compatible visible Mesh Draw Commands into one instanced draw on vertex factories that support it."),
 	ECVF_RenderThreadSafe);
 
-static TAutoConsoleVariable<int32> CVarMeshDrawCommandsParallelSort(
-	TEXT("r.MeshDrawCommands.ParallelSort"),
+static TAutoConsoleVariable<int32> CVarMeshDrawCommandsParallelProcess(
+	TEXT("r.MeshDrawCommands.ParallelProcess"),
 	1,
-	TEXT("Whether to sort and process mesh draw commands in parallel."),
+	TEXT("Whether to process (sort and setup primitive id buffer for instancing) visible mesh draw commands in parallel."),
 	ECVF_RenderThreadSafe);
 
 bool IsDynamicInstancingEnabled()
@@ -2935,12 +2935,12 @@ void UpdateTranslucentMeshSortKeys(FScene* Scene, FViewInfo& View, ETranslucency
 }
 
 /**
- * Packet for parallel mesh draw command sorting and processing for instancing.
+ * Packet for parallel mesh draw command sorting and setting up primitive id buffer for instancing.
  */
-class FMeshDrawCommandSortPacket
+class FVisibleMeshDrawCommandProcessPacket
 {
 public:
-	FMeshDrawCommandSortPacket(FScene* InScene, 
+	FVisibleMeshDrawCommandProcessPacket(FScene* InScene,
 		FViewInfo& InView, 
 		bool InbUseGPUScene, 
 		EMeshPass::Type InPassType,
@@ -3022,7 +3022,7 @@ public:
 						if (VisibleMeshDrawCommand.MeshDrawCommand->PrimitiveIdStreamIndex >= 0
 							&& VisibleMeshDrawCommand.MeshDrawCommand->NumInstances == 1
 							// Don't create a new FMeshDrawCommand for the last command and make it safe for us to look at the next command
-							&& DrawCommandIndex + 1 <= NumDrawCommands
+							&& DrawCommandIndex + 1 < NumDrawCommands
 							// Only create a new FMeshDrawCommand if more than one draw in the state bucket
 							&& CurrentStateBucketId == PassVisibleMeshDrawCommands[DrawCommandIndex + 1].StateBucketId)
 						{
@@ -3108,9 +3108,9 @@ void FSceneRenderer::SortMeshDrawCommands()
 
 	const EShadingPath ShadingPath = Scene->GetShadingPath();
 	const bool bUseGPUScene = UseGPUScene(GMaxRHIShaderPlatform, FeatureLevel);
-	const bool bExecuteInParallel = FApp::ShouldUseThreadingForPerformance() && CVarMeshDrawCommandsParallelSort.GetValueOnRenderThread() > 0;
+	const bool bExecuteInParallel = FApp::ShouldUseThreadingForPerformance() && CVarMeshDrawCommandsParallelProcess.GetValueOnRenderThread() > 0;
 
-	TArray<FMeshDrawCommandSortPacket, SceneRenderingAllocator> Packets;
+	TArray<FVisibleMeshDrawCommandProcessPacket, SceneRenderingAllocator> Packets;
 	Packets.Reserve(Views.Num() * EMeshPass::Num);
 
 	for (int32 ViewIndex = 0; ViewIndex < Views.Num(); ViewIndex++)
@@ -3140,7 +3140,7 @@ void FSceneRenderer::SortMeshDrawCommands()
 						View.VisibleMeshDrawCommandPrimitiveIdBuffers[PassIndex] = FGlobalDynamicVertexBuffer::Get().Allocate(NumDrawCommands * sizeof(int32));
 					}
 
-					FMeshDrawCommandSortPacket Packet(
+					FVisibleMeshDrawCommandProcessPacket Packet(
 						Scene, 
 						View, 
 						bUseGPUScene, 
