Fixed indirect lighting in cases when IndirectLightingCacheAllocation was updated after primitive was added to scene. During adding to scene cached mesh draw commands were created with null IndirectLightingCacheAllocation. Later IndirectLightingCacheAllocation was updated, but cached mesh draw commands didn't update and were rendering with a wrong shader. Also added bNeedsStaticMeshUpdate check into BeginDeferredUpdateStaticMeshes, so we don't add it redundantly.

diff --git a/Engine/Source/Runtime/Renderer/Private/IndirectLightingCache.cpp b/Engine/Source/Runtime/Renderer/Private/IndirectLightingCache.cpp
--- a/Engine/Source/Runtime/Renderer/Private/IndirectLightingCache.cpp
+++ b/Engine/Source/Runtime/Renderer/Private/IndirectLightingCache.cpp
@@ -678,6 +678,8 @@ void FIndirectLightingCache::UpdateCachePrimitive(
 			}
 		}
 
+		const FIndirectLightingCacheAllocation* OriginalCachedIndirectLightingCacheAllocation = PrimitiveSceneInfo->IndirectLightingCacheAllocation;
+
 		if (AttachmentParentAllocation)
 		{
 			// Reuse the attachment parent's lighting allocation if part of an attachment group
@@ -715,6 +717,12 @@ void FIndirectLightingCache::UpdateCachePrimitive(
 				PrimitiveSceneInfo->MarkIndirectLightingCacheBufferDirty();
 			}
 		}
+
+		if (OriginalCachedIndirectLightingCacheAllocation != PrimitiveSceneInfo->IndirectLightingCacheAllocation)
+		{
+			// Update cached mesh draw commands, so they can pick up new IndirectLightingCacheAllocation.
+			PrimitiveSceneInfo->BeginDeferredUpdateStaticMeshes();
+		}
 	}
 }
 
diff --git a/Engine/Source/Runtime/Renderer/Private/PrimitiveSceneInfo.cpp b/Engine/Source/Runtime/Renderer/Private/PrimitiveSceneInfo.cpp
--- a/Engine/Source/Runtime/Renderer/Private/PrimitiveSceneInfo.cpp
+++ b/Engine/Source/Runtime/Renderer/Private/PrimitiveSceneInfo.cpp
@@ -448,6 +448,8 @@ void FPrimitiveSceneInfo::RemoveFromScene(bool bUpdateStaticDrawLists)
 	if (bNeedsStaticMeshUpdate)
 	{
 		Scene->PrimitivesNeedingStaticMeshUpdate.Remove(this);
+
+		bNeedsStaticMeshUpdate = false;
 	}
 
 	if (bUpdateStaticDrawLists)
@@ -498,10 +500,13 @@ void FPrimitiveSceneInfo::UpdateUniformBuffer(FRHICommandListImmediate& RHICmdLi
 
 void FPrimitiveSceneInfo::BeginDeferredUpdateStaticMeshes()
 {
-	// Set a flag which causes InitViews to update the static meshes the next time the primitive is visible.
-	bNeedsStaticMeshUpdate = true;
+	if (!bNeedsStaticMeshUpdate)
+	{
+		// Set a flag which causes InitViews to update the static meshes the next time the primitive is visible.
+		bNeedsStaticMeshUpdate = true;
 
-	Scene->PrimitivesNeedingStaticMeshUpdate.Add(this);
+		Scene->PrimitivesNeedingStaticMeshUpdate.Add(this);
+	}
 }
 
 void FPrimitiveSceneInfo::LinkLODParentComponent()
