Fix ensure() upon QAGame launch

diff --git a/Engine/Source/Runtime/Renderer/Private/DeferredShadingRenderer.cpp b/Engine/Source/Runtime/Renderer/Private/DeferredShadingRenderer.cpp
--- a/Engine/Source/Runtime/Renderer/Private/DeferredShadingRenderer.cpp
+++ b/Engine/Source/Runtime/Renderer/Private/DeferredShadingRenderer.cpp
@@ -773,7 +773,7 @@ bool FDeferredShadingSceneRenderer::GatherRayTracingWorldInstances(FRHICommandLi
 								}
 							}
 
-							ensure(DrawCmdIndexCopy != DynamicMeshBatchStartOffset[ViewIndex]);
+							ensure(DrawCmdIndexCopy != VisibleDrawCommandStartOffset[ViewIndex]);
 							View.RayTracingGeometryInstances.Add(FRayTracingGeometryInstance { RayTracingGeometryInstance, Scene->PrimitiveTransforms[PrimitiveIndex], (uint32)PrimitiveIndex });
 						}
 					}
