Fixes auto test ActorMerging/SingleLODMaterialBaking. The new FLightmapResourceCluster was missing from the bake process.

diff --git a/Engine/Source/Developer/MaterialBaking/Private/MaterialRenderItem.cpp b/Engine/Source/Developer/MaterialBaking/Private/MaterialRenderItem.cpp
--- a/Engine/Source/Developer/MaterialBaking/Private/MaterialRenderItem.cpp
+++ b/Engine/Source/Developer/MaterialBaking/Private/MaterialRenderItem.cpp
@@ -14,7 +14,7 @@ FMeshMaterialRenderItem::FMeshMaterialRenderItem(const FMaterialData* InMaterial
 	: MeshSettings(InMeshSettings), MaterialSettings(InMaterialSettings), MaterialProperty(InMaterialProperty), MaterialRenderProxy(nullptr), ViewFamily(nullptr)
 {
 	GenerateRenderData();
-	LCI = new FMeshRenderInfo(InMeshSettings->LightMap, nullptr, nullptr);
+	LCI = new FMeshRenderInfo(InMeshSettings->LightMap, nullptr, nullptr, InMeshSettings->LightmapResourceCluster);
 }
 
 bool FMeshMaterialRenderItem::Render_RenderThread(FRHICommandListImmediate& RHICmdList, FDrawingPolicyRenderState& DrawRenderState, const FCanvas* Canvas)
diff --git a/Engine/Source/Developer/MaterialBaking/Private/MaterialRenderItemData.h b/Engine/Source/Developer/MaterialBaking/Private/MaterialRenderItemData.h
--- a/Engine/Source/Developer/MaterialBaking/Private/MaterialRenderItemData.h
+++ b/Engine/Source/Developer/MaterialBaking/Private/MaterialRenderItemData.h
@@ -11,12 +11,13 @@
 class FMeshRenderInfo : public FLightCacheInterface
 {
 public:
-	FMeshRenderInfo(const FLightMap* InLightMap, const FShadowMap* InShadowMap, FUniformBufferRHIRef Buffer)
+	FMeshRenderInfo(const FLightMap* InLightMap, const FShadowMap* InShadowMap, FUniformBufferRHIRef Buffer, const FLightmapResourceCluster* InLightmapResourceCluster)
 		: FLightCacheInterface()
 	{
 		SetLightMap(InLightMap);
 		SetShadowMap(InShadowMap);
 		SetPrecomputedLightingBuffer(Buffer);
+		SetResourceCluster(InLightmapResourceCluster);
 	}
 
 	virtual FLightInteraction GetInteraction(const class FLightSceneProxy* LightSceneProxy) const override
diff --git a/Engine/Source/Developer/MaterialBaking/Public/MaterialBakingStructures.h b/Engine/Source/Developer/MaterialBaking/Public/MaterialBakingStructures.h
--- a/Engine/Source/Developer/MaterialBaking/Public/MaterialBakingStructures.h
+++ b/Engine/Source/Developer/MaterialBaking/Public/MaterialBakingStructures.h
@@ -28,7 +28,7 @@ struct FMaterialData
 struct FMeshData
 {
 	FMeshData()
-		: RawMesh(nullptr), Mesh(nullptr), bMirrored(false), VertexColorHash(0), TextureCoordinateIndex(0), LightMapIndex(0), LightMap(nullptr)
+		: RawMesh(nullptr), Mesh(nullptr), bMirrored(false), VertexColorHash(0), TextureCoordinateIndex(0), LightMapIndex(0), LightMap(nullptr), LightmapResourceCluster(nullptr)
 	{}
 
 	/** Ptr to raw mesh data to use for baking out the material data, if nullptr a standard quad is used */
@@ -60,6 +60,9 @@ struct FMeshData
 
 	/** Reference to the lightmap texture part of the level in the currently being baked out mesh instance data is resident */
 	FLightMapRef LightMap;
+
+	/** Pointer to the LightmapResourceCluster to be passed on the the LightCacheInterface when baking */
+	const FLightmapResourceCluster* LightmapResourceCluster;
 };
 
 /** Structure containing data being processed while baking out materials*/
diff --git a/Engine/Source/Developer/MeshMergeUtilities/Private/StaticMeshComponentAdapter.cpp b/Engine/Source/Developer/MeshMergeUtilities/Private/StaticMeshComponentAdapter.cpp
--- a/Engine/Source/Developer/MeshMergeUtilities/Private/StaticMeshComponentAdapter.cpp
+++ b/Engine/Source/Developer/MeshMergeUtilities/Private/StaticMeshComponentAdapter.cpp
@@ -48,6 +48,7 @@ void FStaticMeshComponentAdapter::ApplySettings(int32 LODIndex, FMeshData& InOut
 		{
 			InOutMeshData.LightMap = MeshMapBuildData->LightMap;
 			InOutMeshData.LightMapIndex = StaticMeshComponent->GetStaticMesh()->LightMapCoordinateIndex;
+			InOutMeshData.LightmapResourceCluster = MeshMapBuildData->ResourceCluster;
 		}
 	}
 }
diff --git a/Engine/Source/Runtime/Renderer/Private/BasePassRendering.cpp b/Engine/Source/Runtime/Renderer/Private/BasePassRendering.cpp
--- a/Engine/Source/Runtime/Renderer/Private/BasePassRendering.cpp
+++ b/Engine/Source/Runtime/Renderer/Private/BasePassRendering.cpp
@@ -1550,7 +1550,7 @@ FBasePassMeshProcessor::FBasePassMeshProcessor(const FScene* Scene, ERHIFeatureL
 	extern void GetEarlyZPassMode(EShaderPlatform ShaderPlatform, EDepthDrawingMode& EarlyZPassMode, bool& bEarlyZPassMovable);
 
 	bool bEarlyZPassMovable;
-	GetEarlyZPassMode(Scene->GetShaderPlatform(), EarlyZPassMode, bEarlyZPassMovable);
+	GetEarlyZPassMode(GShaderPlatformForFeatureLevel[InFeatureLevel], EarlyZPassMode, bEarlyZPassMovable);
 }
 
 FMeshPassProcessor* CreateBasePassProcessor(const FScene* Scene, const FSceneView* InViewIfDynamicMeshCommand, FMeshPassDrawListContext& InDrawListContext)
diff --git a/Engine/Source/Runtime/Renderer/Private/Renderer.cpp b/Engine/Source/Runtime/Renderer/Private/Renderer.cpp
--- a/Engine/Source/Runtime/Renderer/Private/Renderer.cpp
+++ b/Engine/Source/Runtime/Renderer/Private/Renderer.cpp
@@ -187,13 +187,13 @@ void FRendererModule::DrawTileMesh(FRHICommandListImmediate& RHICmdList, FDrawin
 				DrawDynamicMeshPass(View, RHICmdList,
 					[&View, &DrawRenderState, &Mesh](FDynamicPassMeshDrawListContext& DynamicMeshPassContext)
 					{
-						FEditorPrimitivesBasePassMeshProcessor PassMeshProcessor(
+						FBasePassMeshProcessor PassMeshProcessor(
 							nullptr,
 							View.GetFeatureLevel(),
 							&View,
 							DrawRenderState,
-							true,
-							DynamicMeshPassContext);
+							DynamicMeshPassContext,
+							ETranslucencyPass::TPT_AllTranslucency);
 
 						const uint64 DefaultBatchElementMask = ~0ull;
 						PassMeshProcessor.AddMeshBatch(Mesh, DefaultBatchElementMask, nullptr);
@@ -243,12 +243,11 @@ void FRendererModule::DrawTileMesh(FRHICommandListImmediate& RHICmdList, FDrawin
 					DrawDynamicMeshPass(View, RHICmdList,
 						[&View, &DrawRenderState, &Mesh](FDynamicPassMeshDrawListContext& DynamicMeshPassContext)
 						{
-							FEditorPrimitivesBasePassMeshProcessor PassMeshProcessor(
+							FBasePassMeshProcessor PassMeshProcessor(
 								nullptr,
 								View.GetFeatureLevel(),
 								&View,
 								DrawRenderState,
-								false,
 								DynamicMeshPassContext);
 
 							const uint64 DefaultBatchElementMask = ~0ull;
