DR - Changes for Ray Tracing tier access in the RHIs
- GRHIRayTracingSupportTier is now used instead of GDynamicRHI->RHIGetRayTracingSupport()
- RHISupportsRayTracing() became RHISupportsRayTracingShaders() and only checks the RHI, not the project
- IsRayTracingSupportedForThisProject() is now split into ShouldCompileRayTracingShadersForProject() and IsRayTracingTierSupported() to fix cooking issues

diff --git a/Engine/Source/Editor/LevelEditor/Private/LevelEditorToolBar.cpp b/Engine/Source/Editor/LevelEditor/Private/LevelEditorToolBar.cpp
--- a/Engine/Source/Editor/LevelEditor/Private/LevelEditorToolBar.cpp
+++ b/Engine/Source/Editor/LevelEditor/Private/LevelEditorToolBar.cpp
@@ -1715,7 +1715,7 @@ static void MakeShaderModelPreviewMenu(FMenuBuilder& MenuBuilder)
 		for (int32 i = GMaxRHIFeatureLevel; i >= 0; --i)
 		{
 #if RHI_RAYTRACING
-			if (IsRayTracingSupportedForThisProject())
+			if (IsRayTracingTierSupported(2))
 			{
 				if (i != ERHIFeatureLevel::SM5)
 				{
diff --git a/Engine/Source/Editor/UnrealEd/Private/EditorEngine.cpp b/Engine/Source/Editor/UnrealEd/Private/EditorEngine.cpp
--- a/Engine/Source/Editor/UnrealEd/Private/EditorEngine.cpp
+++ b/Engine/Source/Editor/UnrealEd/Private/EditorEngine.cpp
@@ -7478,7 +7478,7 @@ void UEditorEngine::AllMaterialsCacheResourceShadersForRendering()
 void UEditorEngine::SetPreviewPlatform(const FName MaterialQualityPlatform, ERHIFeatureLevel::Type PreviewFeatureLevel, const bool bSaveSettings/* = true*/)
 {
 #if RHI_RAYTRACING
-	if (IsRayTracingSupportedForThisProject())
+	if (IsRayTracingTierSupported(2))
 	{
 		if (PreviewFeatureLevel != ERHIFeatureLevel::SM5)
 		{
diff --git a/Engine/Source/Editor/UnrealEd/Private/SEditorViewport.cpp b/Engine/Source/Editor/UnrealEd/Private/SEditorViewport.cpp
--- a/Engine/Source/Editor/UnrealEd/Private/SEditorViewport.cpp
+++ b/Engine/Source/Editor/UnrealEd/Private/SEditorViewport.cpp
@@ -359,7 +359,7 @@ void SEditorViewport::BindCommands()
 	MAP_VIEWMODE_ACTION( Commands.UnlitMode, VMI_Unlit );
 	MAP_VIEWMODE_ACTION( Commands.LitMode, VMI_Lit );
 #if RHI_RAYTRACING
-	if (IsRayTracingSupportedForThisProject() > 0)
+	if (IsRayTracingTierSupported(2))
 	{
 		MAP_VIEWMODE_ACTION(Commands.PathTracingMode, VMI_PathTracing);
 		MAP_VIEWMODE_ACTION(Commands.RayTracingDebugMode, VMI_RayTracingDebug);
diff --git a/Engine/Source/Runtime/D3D12RHI/Private/D3D12RHIPrivate.h b/Engine/Source/Runtime/D3D12RHI/Private/D3D12RHIPrivate.h
--- a/Engine/Source/Runtime/D3D12RHI/Private/D3D12RHIPrivate.h
+++ b/Engine/Source/Runtime/D3D12RHI/Private/D3D12RHIPrivate.h
@@ -459,8 +459,6 @@ public:
 
 #if D3D12_RHI_RAYTRACING
 
-	virtual uint32 RHIGetRayTracingSupport() final override;
-
 	virtual FRayTracingGeometryRHIRef RHICreateRayTracingGeometry(const FRayTracingGeometryInitializer& Initializer) final override;
 
 	virtual FRayTracingSceneRHIRef RHICreateRayTracingScene(const FRayTracingSceneInitializer& Initializer) final override;
diff --git a/Engine/Source/Runtime/D3D12RHI/Private/D3D12RayTracing.cpp b/Engine/Source/Runtime/D3D12RHI/Private/D3D12RayTracing.cpp
--- a/Engine/Source/Runtime/D3D12RHI/Private/D3D12RayTracing.cpp
+++ b/Engine/Source/Runtime/D3D12RHI/Private/D3D12RayTracing.cpp
@@ -247,11 +247,6 @@ static TRefCountPtr<ID3D12StateObject> CreateRayTracingStateObject(
 	return Result;
 }
 
-uint32 FD3D12DynamicRHI::RHIGetRayTracingSupport()
-{
-	return GetAdapter().GetD3DRayTracingDevice() ? 2 : 0;
-}
-
 // #dxr_todo: FD3D12Device::GlobalViewHeap/GlobalSamplerHeap should be used instead of ad-hoc heaps here.
 // Unfortunately, this requires a major refactor of how global heaps work.
 // FD3D12CommandContext-s should not get static chunks of the global heap, but instead should dynamically allocate
diff --git a/Engine/Source/Runtime/D3D12RHI/Private/Windows/WindowsD3D12Device.cpp b/Engine/Source/Runtime/D3D12RHI/Private/Windows/WindowsD3D12Device.cpp
--- a/Engine/Source/Runtime/D3D12RHI/Private/Windows/WindowsD3D12Device.cpp
+++ b/Engine/Source/Runtime/D3D12RHI/Private/Windows/WindowsD3D12Device.cpp
@@ -663,13 +663,17 @@ void FD3D12DynamicRHI::Init()
 	// - Standalones are added to the deferred deletion queue of its parent FD3D12Adapter
 	GRHIForceNoDeletionLatencyForStreamingTextures = !!PLATFORM_WINDOWS;
 
+#if D3D12_RHI_RAYTRACING
+	GRHIRayTracingSupportTier = GetAdapter().GetD3DRayTracingDevice() ? 2 : 0;
+#endif
+
 	// Set the RHI initialized flag.
 	GIsRHIInitialized = true;
 }
 
 void FD3D12DynamicRHI::PostInit()
 {
-	if (IsRayTracingSupportedForThisProject())
+	if (IsRayTracingTierSupported(2))
 	{
 		for (FD3D12Adapter*& Adapter : ChosenAdapters)
 		{
diff --git a/Engine/Source/Runtime/Engine/Private/GameViewportClient.cpp b/Engine/Source/Runtime/Engine/Private/GameViewportClient.cpp
--- a/Engine/Source/Runtime/Engine/Private/GameViewportClient.cpp
+++ b/Engine/Source/Runtime/Engine/Private/GameViewportClient.cpp
@@ -3020,7 +3020,7 @@ bool UGameViewportClient::HandleViewModeCommand( const TCHAR* Cmd, FOutputDevice
 	}
 
 #if RHI_RAYTRACING
-	if (RHIGetRayTracingSupport() == 0)
+	if (GRHIRayTracingSupportTier == 0)
 	{
 		if (ViewModeIndex == VMI_PathTracing)
 		{
diff --git a/Engine/Source/Runtime/Engine/Private/SceneView.cpp b/Engine/Source/Runtime/Engine/Private/SceneView.cpp
--- a/Engine/Source/Runtime/Engine/Private/SceneView.cpp
+++ b/Engine/Source/Runtime/Engine/Private/SceneView.cpp
@@ -2428,7 +2428,7 @@ void FSceneView::SetupRayTracedRendering()
 {
 	RayTracingRenderMode = ERayTracingRenderMode::Disabled;
 
-	if (!IsRayTracingSupportedForThisProject())
+	if (!IsRayTracingTierSupported(2))
 	{
 		return;
 	}
diff --git a/Engine/Source/Runtime/Engine/Private/ShaderCompiler/ShaderCompiler.cpp b/Engine/Source/Runtime/Engine/Private/ShaderCompiler/ShaderCompiler.cpp
--- a/Engine/Source/Runtime/Engine/Private/ShaderCompiler/ShaderCompiler.cpp
+++ b/Engine/Source/Runtime/Engine/Private/ShaderCompiler/ShaderCompiler.cpp
@@ -3893,7 +3893,7 @@ void VerifyGlobalShaders(EShaderPlatform Platform, bool bLoadedFromCacheFile)
 			!IsOpenGLPlatform(GMaxRHIShaderPlatform) && !IsVulkanPlatform(GMaxRHIShaderPlatform) &&
 			!IsMetalPlatform(GMaxRHIShaderPlatform) && !IsSwitchPlatform(GMaxRHIShaderPlatform) &&
 #if RHI_RAYTRACING
-			!IsRayTracingSupportedForThisProject() && //This is here because DXR is caching its BuiltIn Tier1 Shaders in PostInit (see FD3D12Device::InitRayTracing)
+			!IsRayTracingTierSupported(2) && //This is here because DXR is caching its BuiltIn Tier1 Shaders in PostInit (see FD3D12Device::InitRayTracing)
 #endif // RHI_RAYTRACING
 			GShaderCompilingManager->AllowAsynchronousShaderCompiling();
 
diff --git a/Engine/Source/Runtime/Engine/Private/SkeletalRenderCPUSkin.cpp b/Engine/Source/Runtime/Engine/Private/SkeletalRenderCPUSkin.cpp
--- a/Engine/Source/Runtime/Engine/Private/SkeletalRenderCPUSkin.cpp
+++ b/Engine/Source/Runtime/Engine/Private/SkeletalRenderCPUSkin.cpp
@@ -333,7 +333,7 @@ void FSkeletalMeshObjectCPUSkin::FSkeletalMeshObjectLOD::InitResources(FSkelMesh
 		});
 
 #if RHI_RAYTRACING
-	if (IsRayTracingSupportedForThisProject())
+	if (IsRayTracingTierSupported(2))
 	{
 		check(SkelMeshRenderData);
 		check(SkelMeshRenderData->LODRenderData.IsValidIndex(LODIndex));
diff --git a/Engine/Source/Runtime/Engine/Private/SkeletalRenderGPUSkin.cpp b/Engine/Source/Runtime/Engine/Private/SkeletalRenderGPUSkin.cpp
--- a/Engine/Source/Runtime/Engine/Private/SkeletalRenderGPUSkin.cpp
+++ b/Engine/Source/Runtime/Engine/Private/SkeletalRenderGPUSkin.cpp
@@ -326,7 +326,7 @@ void FSkeletalMeshObjectGPUSkin::UpdateDynamicData_RenderThread(FGPUSkinCache* G
 	}
 
 #if RHI_RAYTRACING
-	if (IsRayTracingSupportedForThisProject())
+	if (IsRayTracingTierSupported(2))
 	{
 		if (GEnableGPUSkinCache && SkinCacheEntry)
 		{
diff --git a/Engine/Source/Runtime/Engine/Private/SkeletalRenderStatic.cpp b/Engine/Source/Runtime/Engine/Private/SkeletalRenderStatic.cpp
--- a/Engine/Source/Runtime/Engine/Private/SkeletalRenderStatic.cpp
+++ b/Engine/Source/Runtime/Engine/Private/SkeletalRenderStatic.cpp
@@ -122,7 +122,7 @@ void FSkeletalMeshObjectStatic::FSkeletalMeshObjectLOD::InitResources(FSkelMeshC
 
 
 #if RHI_RAYTRACING
-	if (IsRayTracingSupportedForThisProject())
+	if (IsRayTracingTierSupported(2))
 	{
 		check(SkelMeshRenderData);
 		check(SkelMeshRenderData->LODRenderData.IsValidIndex(LODIndex));
diff --git a/Engine/Source/Runtime/Engine/Private/StaticMesh.cpp b/Engine/Source/Runtime/Engine/Private/StaticMesh.cpp
--- a/Engine/Source/Runtime/Engine/Private/StaticMesh.cpp
+++ b/Engine/Source/Runtime/Engine/Private/StaticMesh.cpp
@@ -683,7 +683,7 @@ void FStaticMeshLODResources::InitResources(UStaticMesh* Parent)
 	}
 
 #if RHI_RAYTRACING
-	if (IsRayTracingSupportedForThisProject())
+	if (IsRayTracingTierSupported(2))
 	{
 		ENQUEUE_RENDER_COMMAND(InitStaticMeshRayTracingGeometry)(
 			[this](FRHICommandListImmediate& RHICmdList)
diff --git a/Engine/Source/Runtime/RHI/Private/RHI.cpp b/Engine/Source/Runtime/RHI/Private/RHI.cpp
--- a/Engine/Source/Runtime/RHI/Private/RHI.cpp
+++ b/Engine/Source/Runtime/RHI/Private/RHI.cpp
@@ -498,6 +498,7 @@ bool GRHISupportsBaseVertexIndex = true;
 TRHIGlobal<bool> GRHISupportsInstancing(true);
 bool GRHISupportsFirstInstance = false;
 bool GRHISupportsDynamicResolution = false;
+int32 GRHIRayTracingSupportTier = 0;
 bool GRHISupportsRHIThread = false;
 bool GRHISupportsRHIOnTaskThread = false;
 bool GRHISupportsParallelRHIExecute = false;
@@ -823,13 +824,6 @@ RHI_API bool RHISupportsIndexBufferUAVs(const EShaderPlatform Platform)
 	return Platform == SP_PCD3D_SM5 || IsVulkanPlatform(Platform) || IsMetalSM5Platform(Platform) || Platform == SP_XBOXONE_D3D12 || Platform == SP_PS4;
 }
 
-RHI_API bool RHISupportsRayTracing(const EShaderPlatform Platform)
-{
-	return (
-		Platform == SP_PCD3D_SM5 &&
-		IsRayTracingSupportedForThisProject()); // dxr_todo: so cooking depends whether you cook on a DXR compatible platform?
-}
-
 
 static ERHIFeatureLevel::Type GRHIMobilePreviewFeatureLevel = ERHIFeatureLevel::Num;
 RHI_API void RHISetMobilePreviewFeatureLevel(ERHIFeatureLevel::Type MobilePreviewFeatureLevel)
diff --git a/Engine/Source/Runtime/RHI/Public/DynamicRHI.h b/Engine/Source/Runtime/RHI/Public/DynamicRHI.h
--- a/Engine/Source/Runtime/RHI/Public/DynamicRHI.h
+++ b/Engine/Source/Runtime/RHI/Public/DynamicRHI.h
@@ -1128,15 +1128,6 @@ public:
 	virtual bool RHIIsTypedUAVLoadSupported(EPixelFormat PixelFormat) { return true; }
 
 #if RHI_RAYTRACING
-	/**
-	 * Query RHI ray tracing support level.
-	 * @return Numeric value representing a set of ray tracing features supported by the RHI.
-	 *   0: Ray tracing is not supported.
-	 *   1: Basic support for acceleration structure construction and ray queries. Ray tracing shader types are not supported.
-	 *   2: Full ray tracing support, including ray tracing shader types, shader binding tables, etc.
-	 */
-	virtual uint32 RHIGetRayTracingSupport() { return 0; }
-
 	// #dxr_todo: reverted index buffer may be needed to support mirrored meshes
 	virtual FRayTracingGeometryRHIRef RHICreateRayTracingGeometry(const FRayTracingGeometryInitializer& Initializer)
 	{
@@ -1354,22 +1345,7 @@ FORCEINLINE class IRHICommandContextContainer* RHIGetCommandContextContainer(int
 	return GDynamicRHI->RHIGetCommandContextContainer(Index, Num, GPUMask);
 }
 
-FORCEINLINE uint32 IsRayTracingSupportedForThisProject()
-{
-#if RHI_RAYTRACING
-	// r.RayTracing is a read-only CVar. UE needs to be restarted to effectively change it
-	static const bool bRayTracingEnabled = IConsoleManager::Get().FindConsoleVariable(TEXT("r.RayTracing"))->GetInt() > 0;
-	return GDynamicRHI->RHIGetRayTracingSupport() && bRayTracingEnabled;
-#else
-	return false;
-#endif
-}
-
 #if RHI_RAYTRACING
-FORCEINLINE uint32 RHIGetRayTracingSupport()
-{
-	return GDynamicRHI->RHIGetRayTracingSupport();
-}
 
 FORCEINLINE FRayTracingGeometryRHIRef RHICreateRayTracingGeometry(const FRayTracingGeometryInitializer& Initializer)
 {
diff --git a/Engine/Source/Runtime/RHI/Public/RHI.h b/Engine/Source/Runtime/RHI/Public/RHI.h
--- a/Engine/Source/Runtime/RHI/Public/RHI.h
+++ b/Engine/Source/Runtime/RHI/Public/RHI.h
@@ -194,8 +194,13 @@ inline bool RHISupportsAbsoluteVertexID(EShaderPlatform InShaderPlatform)
 	return IsVulkanPlatform(InShaderPlatform) || IsVulkanMobilePlatform(InShaderPlatform);
 }
 
-// helper to check whether a shader platform supports ray tracing.
-RHI_API bool RHISupportsRayTracing(EShaderPlatform Platform);
+/** Can this platform compile ray tracing shaders (regardless of project settings).
+ *  To use at runtime, also check GRHIRayTracingSupportTier.
+ **/
+inline RHI_API bool RHISupportsRayTracingShaders(EShaderPlatform Platform)
+{
+	return Platform == SP_PCD3D_SM5;
+}
 
 // Wrapper for GRHI## global variables, allows values to be overridden for mobile preview modes.
 template <typename TValueType>
@@ -461,6 +466,13 @@ extern RHI_API bool GRHISupportsFirstInstance;
 /** Whether or not the RHI can handle dynamic resolution or not. */
 extern RHI_API bool GRHISupportsDynamicResolution;
 
+/** Ray tracing tier level:
+ *   0: Ray tracing is not supported.
+ *   1: Basic support for acceleration structure construction and ray queries. Ray tracing shader types are not supported.
+ *   2: Full ray tracing support, including ray tracing shader types, shader binding tables, etc.
+ **/
+extern RHI_API int32 GRHIRayTracingSupportTier;
+
 /** Whether or not the RHI supports an RHI thread.
 Requirements for RHI thread
 * Microresources (those in RHIStaticStates.h) need to be able to be created by any thread at any time and be able to work with a radically simplified rhi resource lifecycle. CreateSamplerState, CreateRasterizerState, CreateDepthStencilState, CreateBlendState
diff --git a/Engine/Source/Runtime/RenderCore/Public/RenderResource.h b/Engine/Source/Runtime/RenderCore/Public/RenderResource.h
--- a/Engine/Source/Runtime/RenderCore/Public/RenderResource.h
+++ b/Engine/Source/Runtime/RenderCore/Public/RenderResource.h
@@ -503,7 +503,29 @@ public:
 	virtual FString GetFriendlyName() const override { return TEXT("FIndexBuffer"); }
 };
 
+
+FORCEINLINE bool ShouldCompileRayTracingShadersForProject(EShaderPlatform ShaderPlatform)
+{
+#if RHI_RAYTRACING
+	if (RHISupportsRayTracingShaders(ShaderPlatform))
+	{
+		// r.RayTracing is a read-only CVar. UE needs to be restarted to effectively change it
+		static const bool bRayTracingEnabled = IConsoleManager::Get().FindConsoleVariable(TEXT("r.RayTracing"))->GetInt() > 0;
+		return bRayTracingEnabled;
+	}
+#endif
+	return false;
+}
+
+// This function is a runtime only function!
+FORCEINLINE bool IsRayTracingTierSupported(int32 RequiredTier)
+{
+	ensure(RequiredTier >= 0);
+	return GRHIRayTracingSupportTier > RequiredTier;
+}
+
 #if RHI_RAYTRACING
+
 /** A ray tracing geometry resource */
 class RENDERCORE_API FRayTracingGeometry : public FRenderResource
 {
@@ -533,7 +555,7 @@ public:
 
 	virtual void InitRHI() override
 	{
-		if (Initializer.IndexBuffer && Initializer.PositionVertexBuffer && IsRayTracingSupportedForThisProject())
+		if (Initializer.IndexBuffer && Initializer.PositionVertexBuffer && IsRayTracingTierSupported(2))
 		{
 			RayTracingGeometryRHI = RHICreateRayTracingGeometry(Initializer);
 			FRHICommandListExecutor::GetImmediateCommandList().BuildAccelerationStructure(RayTracingGeometryRHI);
diff --git a/Engine/Source/Runtime/Renderer/Private/CompositionLighting/CompositionLighting.cpp b/Engine/Source/Runtime/Renderer/Private/CompositionLighting/CompositionLighting.cpp
--- a/Engine/Source/Runtime/Renderer/Private/CompositionLighting/CompositionLighting.cpp
+++ b/Engine/Source/Runtime/Renderer/Private/CompositionLighting/CompositionLighting.cpp
@@ -125,7 +125,7 @@ bool ShouldRenderScreenSpaceAmbientOcclusion(const FViewInfo& View)
 
 	// #dxr_todo: hard-coded FindConsoleVariable is a bad hack. The current variable is hidden behind DeferredShadingRenderer
 	static IConsoleVariable* UseRayTracingAmbientOcclusionCVar = IConsoleManager::Get().FindConsoleVariable(TEXT("r.RayTracing.AmbientOcclusion"));
-	bool bUseRayTracingAmbientOcclusion = IsRayTracingSupportedForThisProject() && UseRayTracingAmbientOcclusionCVar && UseRayTracingAmbientOcclusionCVar->GetInt() > 0;
+	bool bUseRayTracingAmbientOcclusion = IsRayTracingTierSupported(2) && UseRayTracingAmbientOcclusionCVar && UseRayTracingAmbientOcclusionCVar->GetInt() > 0;
 	return bEnabled && !bUseRayTracingAmbientOcclusion;
 }
 
diff --git a/Engine/Source/Runtime/Renderer/Private/DeferredShadingRenderer.cpp b/Engine/Source/Runtime/Renderer/Private/DeferredShadingRenderer.cpp
--- a/Engine/Source/Runtime/Renderer/Private/DeferredShadingRenderer.cpp
+++ b/Engine/Source/Runtime/Renderer/Private/DeferredShadingRenderer.cpp
@@ -581,7 +581,7 @@ void FDeferredShadingSceneRenderer::PrepareDistanceFieldScene(FRHICommandListImm
 
 bool FDeferredShadingSceneRenderer::GatherRayTracingWorldInstances(FRHICommandListImmediate& RHICmdList)
 {
-	if (!IsRayTracingSupportedForThisProject())
+	if (!IsRayTracingTierSupported(2))
 	{
 		return false;
 	}
@@ -867,7 +867,7 @@ bool FDeferredShadingSceneRenderer::GatherRayTracingWorldInstances(FRHICommandLi
 
 bool FDeferredShadingSceneRenderer::DispatchRayTracingWorldUpdates(FRHICommandListImmediate& RHICmdList)
 {
-	if (!IsRayTracingSupportedForThisProject())
+	if (!IsRayTracingTierSupported(2))
 	{
 		return false;
 	}
@@ -1678,7 +1678,7 @@ void FDeferredShadingSceneRenderer::Render(FRHICommandListImmediate& RHICmdList)
 
 	checkSlow(RHICmdList.IsOutsideRenderPass());
 #if RHI_RAYTRACING
-	if (IsRayTracingSupportedForThisProject())
+	if (IsRayTracingTierSupported(2))
 	{
 		for (int32 ViewIndex = 0; ViewIndex < Views.Num(); ViewIndex++)
 		{
diff --git a/Engine/Source/Runtime/Renderer/Private/LightRendering.cpp b/Engine/Source/Runtime/Renderer/Private/LightRendering.cpp
--- a/Engine/Source/Runtime/Renderer/Private/LightRendering.cpp
+++ b/Engine/Source/Runtime/Renderer/Private/LightRendering.cpp
@@ -764,7 +764,7 @@ void FDeferredShadingSceneRenderer::RenderLights(FRHICommandListImmediate& RHICm
 						RenderRayTracingOcclusionForRectLight(RHICmdList, LightSceneInfo, ScreenShadowMaskTexture);
 					}
 					// #dxr_todo: ShouldRenderRayTracingOcclusion()
-					else if (IsRayTracingSupportedForThisProject() && GRayTracingOcclusion == 1)
+					else if (IsRayTracingTierSupported(2) && GRayTracingOcclusion == 1)
 					{
 						RenderRayTracingOcclusion(RHICmdList, &LightSceneInfo, ScreenShadowMaskTexture);
 					}
diff --git a/Engine/Source/Runtime/Renderer/Private/PathTracing.cpp b/Engine/Source/Runtime/Renderer/Private/PathTracing.cpp
--- a/Engine/Source/Runtime/Renderer/Private/PathTracing.cpp
+++ b/Engine/Source/Runtime/Renderer/Private/PathTracing.cpp
@@ -82,7 +82,7 @@ public:
 
 	static bool ShouldCompilePermutation(const FGlobalShaderPermutationParameters& Parameters)
 	{
-		return IsFeatureLevelSupported(Parameters.Platform, ERHIFeatureLevel::SM5) && IsRayTracingSupportedForThisProject();
+		return ShouldCompileRayTracingShadersForProject(Parameters.Platform);
 	}
 
 	FPathTracingRG() {}
@@ -291,7 +291,7 @@ public:
 public:
 	static bool ShouldCompilePermutation(const FGlobalShaderPermutationParameters& Parameters)
 	{
-		return IsFeatureLevelSupported(Parameters.Platform, ERHIFeatureLevel::SM5) && IsRayTracingSupportedForThisProject();
+		return ShouldCompileRayTracingShadersForProject(Parameters.Platform);
 	}
 
 	FPathTracingCHS() {}
@@ -312,7 +312,7 @@ public:
 public:
 	static bool ShouldCompilePermutation(const FGlobalShaderPermutationParameters& Parameters)
 	{
-		return IsFeatureLevelSupported(Parameters.Platform, ERHIFeatureLevel::SM5) && IsRayTracingSupportedForThisProject();
+		return ShouldCompileRayTracingShadersForProject(Parameters.Platform);
 	}
 
 	FPathTracingMS() {}
diff --git a/Engine/Source/Runtime/Renderer/Private/PathTracing/PathTracingTestbed.cpp b/Engine/Source/Runtime/Renderer/Private/PathTracing/PathTracingTestbed.cpp
--- a/Engine/Source/Runtime/Renderer/Private/PathTracing/PathTracingTestbed.cpp
+++ b/Engine/Source/Runtime/Renderer/Private/PathTracing/PathTracingTestbed.cpp
@@ -20,7 +20,7 @@ public:
 
 	static bool ShouldCompilePermutation(const FGlobalShaderPermutationParameters& Parameters)
 	{
-		return RHISupportsComputeShaders(Parameters.Platform) && IsRayTracingSupportedForThisProject();
+		return RHISupportsComputeShaders(Parameters.Platform) && ShouldCompileRayTracingShadersForProject(Parameters.Platform);
 	}
 
 	static uint32 GetGroupSize()
@@ -103,7 +103,7 @@ public:
 
 	static bool ShouldCompilePermutation(const FGlobalShaderPermutationParameters& Parameters)
 	{
-		return RHISupportsComputeShaders(Parameters.Platform) && IsRayTracingSupportedForThisProject();
+		return RHISupportsComputeShaders(Parameters.Platform) && ShouldCompileRayTracingShadersForProject(Parameters.Platform);
 	}
 
 	static uint32 GetGroupSize()
@@ -189,7 +189,7 @@ public:
 
 	static bool ShouldCompilePermutation(const FGlobalShaderPermutationParameters& Parameters)
 	{
-		return RHISupportsComputeShaders(Parameters.Platform) && IsRayTracingSupportedForThisProject();
+		return RHISupportsComputeShaders(Parameters.Platform) && ShouldCompileRayTracingShadersForProject(Parameters.Platform);
 	}
 
 	static uint32 GetGroupSize()
@@ -280,7 +280,7 @@ public:
 
 	static bool ShouldCompilePermutation(const FGlobalShaderPermutationParameters& Parameters)
 	{
-		return RHISupportsComputeShaders(Parameters.Platform) && IsRayTracingSupportedForThisProject();
+		return RHISupportsComputeShaders(Parameters.Platform) && ShouldCompileRayTracingShadersForProject(Parameters.Platform);
 	}
 
 	static uint32 GetGroupSize()
diff --git a/Engine/Source/Runtime/Renderer/Private/PathTracing/SkyLightCdfBuild.cpp b/Engine/Source/Runtime/Renderer/Private/PathTracing/SkyLightCdfBuild.cpp
--- a/Engine/Source/Runtime/Renderer/Private/PathTracing/SkyLightCdfBuild.cpp
+++ b/Engine/Source/Runtime/Renderer/Private/PathTracing/SkyLightCdfBuild.cpp
@@ -269,7 +269,7 @@ class FVisualizeCdfPS : public FGlobalShader
 
 	static bool ShouldCompilePermutation(const FGlobalShaderPermutationParameters& Parameters)
 	{
-		return IsFeatureLevelSupported(Parameters.Platform, ERHIFeatureLevel::SM5) && IsRayTracingSupportedForThisProject();
+		return ShouldCompileRayTracingShadersForProject(Parameters.Platform);
 	}
 
 	static void ModifyCompilationEnvironment(const FGlobalShaderPermutationParameters& Parameters, FShaderCompilerEnvironment& OutEnvironment)
diff --git a/Engine/Source/Runtime/Renderer/Private/PathTracing/VarianceMipTreeBuild.cpp b/Engine/Source/Runtime/Renderer/Private/PathTracing/VarianceMipTreeBuild.cpp
--- a/Engine/Source/Runtime/Renderer/Private/PathTracing/VarianceMipTreeBuild.cpp
+++ b/Engine/Source/Runtime/Renderer/Private/PathTracing/VarianceMipTreeBuild.cpp
@@ -19,7 +19,7 @@ class FBuildVarianceMipTreeCS : public FGlobalShader
 public:
 	static bool ShouldCompilePermutation(const FGlobalShaderPermutationParameters& Parameters)
 	{
-		return RHISupportsComputeShaders(Parameters.Platform) && IsRayTracingSupportedForThisProject();
+		return RHISupportsComputeShaders(Parameters.Platform) && ShouldCompileRayTracingShadersForProject(Parameters.Platform);
 	}
 
 	static uint32 GetGroupSize()
diff --git a/Engine/Source/Runtime/Renderer/Private/RayTracing/RayTracingAmbientOcclusion.cpp b/Engine/Source/Runtime/Renderer/Private/RayTracing/RayTracingAmbientOcclusion.cpp
--- a/Engine/Source/Runtime/Renderer/Private/RayTracing/RayTracingAmbientOcclusion.cpp
+++ b/Engine/Source/Runtime/Renderer/Private/RayTracing/RayTracingAmbientOcclusion.cpp
@@ -62,7 +62,7 @@ class FAmbientOcclusionRGS : public FGlobalShader
 public:
 	static bool ShouldCompilePermutation(const FGlobalShaderPermutationParameters& Parameters)
 	{
-		return IsFeatureLevelSupported(Parameters.Platform, ERHIFeatureLevel::SM5) && IsRayTracingSupportedForThisProject();
+		return ShouldCompileRayTracingShadersForProject(Parameters.Platform);
 	}
 
 	FAmbientOcclusionRGS() {}
diff --git a/Engine/Source/Runtime/Renderer/Private/RayTracing/RayTracingBarycentrics.cpp b/Engine/Source/Runtime/Renderer/Private/RayTracing/RayTracingBarycentrics.cpp
--- a/Engine/Source/Runtime/Renderer/Private/RayTracing/RayTracingBarycentrics.cpp
+++ b/Engine/Source/Runtime/Renderer/Private/RayTracing/RayTracingBarycentrics.cpp
@@ -22,7 +22,7 @@ class FRayTracingBarycentricsRGS : public FGlobalShader
 
 	static bool ShouldCompilePermutation(const FGlobalShaderPermutationParameters& Parameters)
 	{
-		return IsFeatureLevelSupported(Parameters.Platform, ERHIFeatureLevel::SM5) && IsRayTracingSupportedForThisProject();
+		return ShouldCompileRayTracingShadersForProject(Parameters.Platform);
 	}
 };
 IMPLEMENT_GLOBAL_SHADER(FRayTracingBarycentricsRGS, "/Engine/Private/RayTracing/RayTracingBarycentrics.usf", "RayTracingBarycentricsMainRGS", SF_RayGen);
@@ -36,7 +36,7 @@ public:
 
 	static bool ShouldCompilePermutation(const FGlobalShaderPermutationParameters& Parameters)
 	{
-		return IsFeatureLevelSupported(Parameters.Platform, ERHIFeatureLevel::SM5) && IsRayTracingSupportedForThisProject();
+		return ShouldCompileRayTracingShadersForProject(Parameters.Platform);
 	}
 
 	FRayTracingBarycentricsMS() = default;
@@ -61,7 +61,7 @@ public:
 
 	static bool ShouldCompilePermutation(const FGlobalShaderPermutationParameters& Parameters)
 	{
-		return IsFeatureLevelSupported(Parameters.Platform, ERHIFeatureLevel::SM5) && IsRayTracingSupportedForThisProject();
+		return ShouldCompileRayTracingShadersForProject(Parameters.Platform);
 	}
 
 	FRayTracingBarycentricsCHS() = default;
diff --git a/Engine/Source/Runtime/Renderer/Private/RayTracing/RayTracingDebug.cpp b/Engine/Source/Runtime/Renderer/Private/RayTracing/RayTracingDebug.cpp
--- a/Engine/Source/Runtime/Renderer/Private/RayTracing/RayTracingDebug.cpp
+++ b/Engine/Source/Runtime/Renderer/Private/RayTracing/RayTracingDebug.cpp
@@ -29,7 +29,7 @@ class FRayTracingDebugRGS : public FGlobalShader
 
 	static bool ShouldCompilePermutation(const FGlobalShaderPermutationParameters& Parameters)
 	{
-		return IsFeatureLevelSupported(Parameters.Platform, ERHIFeatureLevel::SM5) && IsRayTracingSupportedForThisProject();
+		return ShouldCompileRayTracingShadersForProject(Parameters.Platform);
 	}
 };
 IMPLEMENT_GLOBAL_SHADER(FRayTracingDebugRGS, "/Engine/Private/RayTracing/RayTracingDebug.usf", "RayTracingDebugMainRGS", SF_RayGen);
@@ -42,7 +42,7 @@ public:
 
 	static bool ShouldCompilePermutation(const FGlobalShaderPermutationParameters& Parameters)
 	{
-		return IsFeatureLevelSupported(Parameters.Platform, ERHIFeatureLevel::SM5) && IsRayTracingSupportedForThisProject();
+		return ShouldCompileRayTracingShadersForProject(Parameters.Platform);
 	}
 
 	FRayTracingDebugMS() = default;
@@ -61,7 +61,7 @@ public:
 
 	static bool ShouldCompilePermutation(const FGlobalShaderPermutationParameters& Parameters)
 	{
-		return IsFeatureLevelSupported(Parameters.Platform, ERHIFeatureLevel::SM5) && IsRayTracingSupportedForThisProject();
+		return ShouldCompileRayTracingShadersForProject(Parameters.Platform);
 	}
 
 	FRayTracingDebugCHS() = default;
diff --git a/Engine/Source/Runtime/Renderer/Private/RayTracing/RayTracingDynamicGeometry.cpp b/Engine/Source/Runtime/Renderer/Private/RayTracing/RayTracingDynamicGeometry.cpp
--- a/Engine/Source/Runtime/Renderer/Private/RayTracing/RayTracingDynamicGeometry.cpp
+++ b/Engine/Source/Runtime/Renderer/Private/RayTracing/RayTracingDynamicGeometry.cpp
@@ -30,7 +30,7 @@ public:
 	static bool ShouldCompilePermutation(EShaderPlatform Platform, const FMaterial* Material, const FVertexFactoryType* VertexFactoryType)
 	{
 		// #dxr_todo: this should also check if ray tracing is enabled for the target platform & project
-		return IsSupportedDynamicVertexFactoryType(VertexFactoryType) && IsRayTracingSupportedForThisProject();
+		return IsSupportedDynamicVertexFactoryType(VertexFactoryType) && ShouldCompileRayTracingShadersForProject(Platform);
 	}
 
 	virtual bool Serialize(FArchive& Ar) override
diff --git a/Engine/Source/Runtime/Renderer/Private/RayTracing/RayTracingMaterialHitShaders.cpp b/Engine/Source/Runtime/Renderer/Private/RayTracing/RayTracingMaterialHitShaders.cpp
--- a/Engine/Source/Runtime/Renderer/Private/RayTracing/RayTracingMaterialHitShaders.cpp
+++ b/Engine/Source/Runtime/Renderer/Private/RayTracing/RayTracingMaterialHitShaders.cpp
@@ -90,7 +90,7 @@ public:
 		return IsSupportedVertexFactoryType(VertexFactoryType)
 			&& (Material->IsMasked() == UseAnyHitShader)
 			&& LightMapPolicyType::ShouldCompilePermutation(Platform, Material, VertexFactoryType)
-			&& IsRayTracingSupportedForThisProject();
+			&& ShouldCompileRayTracingShadersForProject(Platform);
 	}
 
 	static void ModifyCompilationEnvironment(EShaderPlatform Platform, const FMaterial* Material, FShaderCompilerEnvironment& OutEnvironment)
@@ -127,7 +127,7 @@ IMPLEMENT_MATERIALCHS_TYPE(TUniformLightMapPolicy<LMP_DISTANCE_FIELD_SHADOWS_AND
 
 static FMeshPassProcessor* CreateRayTracingProcessor(const FScene* Scene, const FSceneView* InViewIfDynamicMeshCommand, FMeshPassDrawListContext& InDrawListContext)
 {
-	return IsRayTracingSupportedForThisProject() ? new(FMemStack::Get()) FRayTracingMeshProcessor(Scene, Scene->GetFeatureLevel(), InViewIfDynamicMeshCommand, InDrawListContext) : nullptr;
+	return IsRayTracingTierSupported(2) ? new(FMemStack::Get()) FRayTracingMeshProcessor(Scene, Scene->GetFeatureLevel(), InViewIfDynamicMeshCommand, InDrawListContext) : nullptr;
 }
 
 FRegisterPassProcessorCreateFunction RegisterRayTracing(&CreateRayTracingProcessor, EShadingPath::Deferred, EMeshPass::RayTracing, EMeshPassFlags::CachedMeshCommands);
diff --git a/Engine/Source/Runtime/Renderer/Private/RayTracing/RayTracingReflections.cpp b/Engine/Source/Runtime/Renderer/Private/RayTracing/RayTracingReflections.cpp
--- a/Engine/Source/Runtime/Renderer/Private/RayTracing/RayTracingReflections.cpp
+++ b/Engine/Source/Runtime/Renderer/Private/RayTracing/RayTracingReflections.cpp
@@ -148,7 +148,7 @@ class FRayTracingReflectionsRG : public FGlobalShader
 
 	static bool ShouldCompilePermutation(const FGlobalShaderPermutationParameters& Parameters)
 	{
-		return RHISupportsRayTracing(Parameters.Platform);
+		return ShouldCompileRayTracingShadersForProject(Parameters.Platform);
 	}
 };
 
@@ -159,7 +159,7 @@ class FRayTracingReflectionsCHS : public FGlobalShader
 		
 	static bool ShouldCompilePermutation(const FGlobalShaderPermutationParameters& Parameters)
 	{
-		return RHISupportsRayTracing(Parameters.Platform);
+		return ShouldCompileRayTracingShadersForProject(Parameters.Platform);
 	}
 	
 	using FParameters = FEmptyShaderParameters;
@@ -172,7 +172,7 @@ class FRayTracingReflectionsMS : public FGlobalShader
 		
 	static bool ShouldCompilePermutation(const FGlobalShaderPermutationParameters& Parameters)
 	{
-		return RHISupportsRayTracing(Parameters.Platform);
+		return ShouldCompileRayTracingShadersForProject(Parameters.Platform);
 	}
 
 	using FParameters = FEmptyShaderParameters;
diff --git a/Engine/Source/Runtime/Renderer/Private/RayTracing/RayTracingShadows.cpp b/Engine/Source/Runtime/Renderer/Private/RayTracing/RayTracingShadows.cpp
--- a/Engine/Source/Runtime/Renderer/Private/RayTracing/RayTracingShadows.cpp
+++ b/Engine/Source/Runtime/Renderer/Private/RayTracing/RayTracingShadows.cpp
@@ -47,7 +47,7 @@ class FOcclusionRGS : public FGlobalShader
 
 	static bool ShouldCompilePermutation(const FGlobalShaderPermutationParameters& Parameters)
 	{
-		return IsFeatureLevelSupported(Parameters.Platform, ERHIFeatureLevel::SM5) && IsRayTracingSupportedForThisProject();
+		return ShouldCompileRayTracingShadersForProject(Parameters.Platform);
 	}
 
 	BEGIN_SHADER_PARAMETER_STRUCT(FParameters, )
diff --git a/Engine/Source/Runtime/Renderer/Private/RayTracing/RayTracingShadowsForRectLight.cpp b/Engine/Source/Runtime/Renderer/Private/RayTracing/RayTracingShadowsForRectLight.cpp
--- a/Engine/Source/Runtime/Renderer/Private/RayTracing/RayTracingShadowsForRectLight.cpp
+++ b/Engine/Source/Runtime/Renderer/Private/RayTracing/RayTracingShadowsForRectLight.cpp
@@ -39,7 +39,7 @@ static FAutoConsoleVariableRef CVarRayTracingRecLightIsTextureImportanceSampling
 
 bool IsRayTracingRectLightSelected()
 {
-	return IsRayTracingSupportedForThisProject() && GRayTracingRectLight == 1;
+	return IsRayTracingTierSupported(2) && GRayTracingRectLight == 1;
 }
 
 bool ShouldRenderRayTracingStaticOrStationaryRectLight(const FLightSceneInfo& LightSceneInfo)
@@ -96,7 +96,7 @@ class FRectLightOcclusionRGS : public FGlobalShader
 public:
 	static bool ShouldCompilePermutation(const FGlobalShaderPermutationParameters& Parameters)
 	{
-		return IsFeatureLevelSupported(Parameters.Platform, ERHIFeatureLevel::SM5) && IsRayTracingSupportedForThisProject();
+		return ShouldCompileRayTracingShadersForProject(Parameters.Platform);
 	}
 
 	FRectLightOcclusionRGS() {}
diff --git a/Engine/Source/Runtime/Renderer/Private/RayTracing/RayTracingTestbed.cpp b/Engine/Source/Runtime/Renderer/Private/RayTracing/RayTracingTestbed.cpp
--- a/Engine/Source/Runtime/Renderer/Private/RayTracing/RayTracingTestbed.cpp
+++ b/Engine/Source/Runtime/Renderer/Private/RayTracing/RayTracingTestbed.cpp
@@ -10,7 +10,7 @@
 
 void TestBasicRayTracing(bool bValidateResults)
 {
-	if (RHIGetRayTracingSupport() == 0)
+	if (GRHIRayTracingSupportTier == 0)
 	{
 		return;
 	}
@@ -167,8 +167,7 @@ public:
 
 	static bool ShouldCompilePermutation(const FGlobalShaderPermutationParameters& Parameters)
 	{
-		// #dxr_todo: this should also check if ray tracing is enabled for the target platform & project
-		return IsFeatureLevelSupported(Parameters.Platform, ERHIFeatureLevel::SM5) && IsRayTracingSupportedForThisProject();
+		return ShouldCompileRayTracingShadersForProject(Parameters.Platform);
 	}
 
 	FTestRaygenShader() {}
diff --git a/Engine/Source/Runtime/Renderer/Private/RayTracing/RaytracingSkylight.cpp b/Engine/Source/Runtime/Renderer/Private/RayTracing/RaytracingSkylight.cpp
--- a/Engine/Source/Runtime/Renderer/Private/RayTracing/RaytracingSkylight.cpp
+++ b/Engine/Source/Runtime/Renderer/Private/RayTracing/RaytracingSkylight.cpp
@@ -10,7 +10,7 @@ static int32 GRayTracingSkyLight = 0;
 bool IsRayTracingSkyLightSelected()
 {
 #if RHI_RAYTRACING
-	return IsRayTracingSupportedForThisProject() && GRayTracingSkyLight > 0;
+	return IsRayTracingTierSupported(2) && GRayTracingSkyLight > 0;
 #else
 	return false;
 #endif
@@ -61,7 +61,7 @@ class FSkylightRG : public FGlobalShader
 
 	static bool ShouldCompilePermutation(const FGlobalShaderPermutationParameters& Parameters)
 	{
-		return IsFeatureLevelSupported(Parameters.Platform, ERHIFeatureLevel::SM5) && IsRayTracingSupportedForThisProject();
+		return IsRayTracingEnabledForThisProject(Parameters.Platform, false);
 	}
 
 	BEGIN_SHADER_PARAMETER_STRUCT(FParameters, )
@@ -179,7 +179,7 @@ class FSkylightRG : public FGlobalShader
 public:
 	static bool ShouldCompilePermutation(const FGlobalShaderPermutationParameters& Parameters)
 	{
-		return IsFeatureLevelSupported(Parameters.Platform, ERHIFeatureLevel::SM5) && IsRayTracingSupportedForThisProject();
+		return ShouldCompileRayTracingShadersForProject(Parameters.Platform);
 	}
 
 	FSkylightRG() {}
diff --git a/Engine/Source/Runtime/Renderer/Private/ReflectionEnvironment.cpp b/Engine/Source/Runtime/Renderer/Private/ReflectionEnvironment.cpp
--- a/Engine/Source/Runtime/Renderer/Private/ReflectionEnvironment.cpp
+++ b/Engine/Source/Runtime/Renderer/Private/ReflectionEnvironment.cpp
@@ -688,7 +688,7 @@ void FDeferredShadingSceneRenderer::RenderDeferredReflectionsAndSkyLighting(FRHI
 		bReflectionCapture = bReflectionCapture || View.bIsReflectionCapture;
 	}
 
-	bool bRayTracedReflections = IsRayTracingSupportedForThisProject() && GRayTracingReflections == 1;
+	bool bRayTracedReflections = IsRayTracingTierSupported(2) && GRayTracingReflections == 1;
 
 	const bool bSkyLight = Scene->SkyLight
 		&& Scene->SkyLight->ProcessedTexture
diff --git a/Engine/Source/Runtime/Renderer/Private/RendererScene.cpp b/Engine/Source/Runtime/Renderer/Private/RendererScene.cpp
--- a/Engine/Source/Runtime/Renderer/Private/RendererScene.cpp
+++ b/Engine/Source/Runtime/Renderer/Private/RendererScene.cpp
@@ -894,7 +894,7 @@ FScene::FScene(UWorld* InWorld, bool bInRequiresHitProxies, bool bInIsEditorScen
 	}
 
 #if RHI_RAYTRACING
-	if (IsRayTracingSupportedForThisProject())
+	if (IsRayTracingTierSupported(2))
 	{
 		RayTracingDynamicGeometryCollection = new FRayTracingDynamicGeometryCollection();
 	}
diff --git a/Engine/Source/Runtime/UtilityShaders/Public/BuiltInRayTracingShaders.h b/Engine/Source/Runtime/UtilityShaders/Public/BuiltInRayTracingShaders.h
--- a/Engine/Source/Runtime/UtilityShaders/Public/BuiltInRayTracingShaders.h
+++ b/Engine/Source/Runtime/UtilityShaders/Public/BuiltInRayTracingShaders.h
@@ -17,7 +17,7 @@ protected:
 public:
 	static bool ShouldCompilePermutation(const FGlobalShaderPermutationParameters& Parameters)
 	{
-		return IsFeatureLevelSupported(Parameters.Platform, ERHIFeatureLevel::SM5) && IsRayTracingSupportedForThisProject();
+		return ShouldCompileRayTracingShadersForProject(Parameters.Platform);
 	}
 
 	FBuiltInRayTracingShader() = default;
