DR - Fix for ES2 -game
- Minor change for better lambda debugging

diff --git a/Engine/Source/Runtime/Engine/Private/Components/SplineMeshComponent.cpp b/Engine/Source/Runtime/Engine/Private/Components/SplineMeshComponent.cpp
--- a/Engine/Source/Runtime/Engine/Private/Components/SplineMeshComponent.cpp
+++ b/Engine/Source/Runtime/Engine/Private/Components/SplineMeshComponent.cpp
@@ -546,14 +546,8 @@ void USplineMeshComponent::UpdateRenderStateAndCollision_Internal(bool bConcurre
 		float SplineMeshMinZ = 1.f;
 		CalculateScaleZAndMinZ(SplineMeshScaleZ, SplineMeshMinZ);
 
-		ENQUEUE_UNIQUE_RENDER_COMMAND_SIXPARAMETER(
-			UpdateSplineParamsRTCommand,
-			FSplineMeshSceneProxy*, SplineProxy, SplineProxy,
-			FSplineMeshParams, SplineParams, SplineParams,
-			TEnumAsByte<ESplineMeshAxis::Type>, ForwardAxis, ForwardAxis,
-			FVector, SplineUpDir, SplineUpDir,
-			float, SplineMeshScaleZ, SplineMeshScaleZ,
-			float, SplineMeshMinZ, SplineMeshMinZ,
+		ENQUEUE_RENDER_COMMAND(UpdateSplineParamsRTCommand)(
+			[SplineProxy, this, SplineMeshScaleZ, SplineMeshMinZ](FRHICommandList&)
 			{
 				SplineProxy->SplineParams = SplineParams;
 				SplineProxy->ForwardAxis = ForwardAxis;
diff --git a/Engine/Source/Runtime/Renderer/Private/RendererScene.cpp b/Engine/Source/Runtime/Renderer/Private/RendererScene.cpp
--- a/Engine/Source/Runtime/Renderer/Private/RendererScene.cpp
+++ b/Engine/Source/Runtime/Renderer/Private/RendererScene.cpp
@@ -179,13 +179,14 @@ FSceneViewState::FSceneViewState()
 
 	TotalRayCount = 0;
 	TotalRayCountBuffer = new FRWBuffer;
-	ENQUEUE_UNIQUE_RENDER_COMMAND_ONEPARAMETER(
-		InitializeSceneViewStateRWBuffer,
-		FRWBuffer*, RWBufferPtr, TotalRayCountBuffer,
-		{
-			RWBufferPtr->Initialize(sizeof(uint32), 1, PF_R32_UINT);
-		}
-	);
+	if (GetFeatureLevel() >= ERHIFeatureLevel::SM5)
+	{
+		ENQUEUE_RENDER_COMMAND(InitializeSceneViewStateRWBuffer)(
+			[this](FRHICommandList&)
+			{
+				TotalRayCountBuffer->Initialize(sizeof(uint32), 1, PF_R32_UINT);
+			});
+	}
 #endif
 }
 
@@ -193,26 +194,22 @@ void DestroyRenderResource(FRenderResource* RenderResource)
 {
 	if (RenderResource) 
 	{
-		ENQUEUE_UNIQUE_RENDER_COMMAND_ONEPARAMETER(
-			DestroySceneViewStateRenderResource,
-			FRenderResource*, RenderResourceRT, RenderResource,
+		ENQUEUE_RENDER_COMMAND(DestroySceneViewStateRenderResource)(
+			[RenderResource](FRHICommandList&)
 			{
-				RenderResourceRT->ReleaseResource();
-				delete RenderResourceRT;
-			}
-		);
+				RenderResource->ReleaseResource();
+				delete RenderResource;
+			});
 	}
 }
 
 void DestroyRWBuffer(FRWBuffer* RWBuffer)
 {
-	ENQUEUE_UNIQUE_RENDER_COMMAND_ONEPARAMETER(
-		DestroyRWBuffer,
-		FRWBuffer*, RWBufferPtr, RWBuffer,
+	ENQUEUE_RENDER_COMMAND(DestroyRWBuffer)(
+		[RWBuffer](FRHICommandList&)
 		{
-			delete RWBufferPtr;
-		}
-	);
+			delete RWBuffer;
+		});
 }
 
 FSceneViewState::~FSceneViewState()
@@ -1179,16 +1176,15 @@ void FScene::UpdatePrimitiveLightingAttachmentRoot(UPrimitiveComponent* Primitiv
 
 	if (Primitive->SceneProxy)
 	{
-		ENQUEUE_UNIQUE_RENDER_COMMAND_TWOPARAMETER(
-			UpdatePrimitiveAttachment,
-			FPrimitiveSceneProxy*,Proxy,Primitive->SceneProxy,
-			FPrimitiveComponentId,NewComponentId,NewComponentId,
-		{
-			FPrimitiveSceneInfo* PrimitiveInfo = Proxy->GetPrimitiveSceneInfo();
-			PrimitiveInfo->UnlinkAttachmentGroup();
-			PrimitiveInfo->LightingAttachmentRoot = NewComponentId;
-			PrimitiveInfo->LinkAttachmentGroup();
-		});
+		FPrimitiveSceneProxy* Proxy = Primitive->SceneProxy;
+		ENQUEUE_RENDER_COMMAND(UpdatePrimitiveAttachment)(
+			[Proxy,NewComponentId](FRHICommandList&)
+			{
+				FPrimitiveSceneInfo* PrimitiveInfo = Proxy->GetPrimitiveSceneInfo();
+				PrimitiveInfo->UnlinkAttachmentGroup();
+				PrimitiveInfo->LightingAttachmentRoot = NewComponentId;
+				PrimitiveInfo->LinkAttachmentGroup();
+			});
 	}
 }
 
@@ -1228,13 +1224,13 @@ void FScene::UpdatePrimitiveDistanceFieldSceneData_GameThread(UPrimitiveComponen
 
 		ENQUEUE_RENDER_COMMAND(UpdatePrimDFSceneDataCmd)(
 			[this, PrimitiveSceneProxy = Primitive->SceneProxy](FRHICommandList&)
-		{
-			if (PrimitiveSceneProxy && PrimitiveSceneProxy->GetPrimitiveSceneInfo())
 			{
-				FPrimitiveSceneInfo* Info = PrimitiveSceneProxy->GetPrimitiveSceneInfo();
-				DistanceFieldSceneData.UpdatePrimitive(Info);
-			}
-		});
+				if (PrimitiveSceneProxy && PrimitiveSceneProxy->GetPrimitiveSceneInfo())
+				{
+					FPrimitiveSceneInfo* Info = PrimitiveSceneProxy->GetPrimitiveSceneInfo();
+					DistanceFieldSceneData.UpdatePrimitive(Info);
+				}
+			});
 	}
 }
 
@@ -1378,11 +1374,10 @@ void FScene::RemovePrimitive( UPrimitiveComponent* Primitive )
 		Primitive->SceneProxy = NULL;
 
 		// Send a command to the rendering thread to remove the primitive from the scene.
-		ENQUEUE_UNIQUE_RENDER_COMMAND_THREEPARAMETER(
-			FRemovePrimitiveCommand,
-			FScene*,Scene,this,
-			FPrimitiveSceneInfo*,PrimitiveSceneInfo,PrimitiveSceneProxy->GetPrimitiveSceneInfo(),
-			FThreadSafeCounter*,AttachmentCounter,&Primitive->AttachmentCounter,
+		FScene* Scene = this;
+		FThreadSafeCounter* AttachmentCounter = &Primitive->AttachmentCounter;
+		ENQUEUE_RENDER_COMMAND(FRemovePrimitiveCommand)(
+			[Scene, PrimitiveSceneInfo, AttachmentCounter](FRHICommandList&)
 			{
 				FScopeCycleCounter Context(PrimitiveSceneInfo->Proxy->GetStatId());
 				Scene->RemovePrimitiveSceneInfo_RenderThread(PrimitiveSceneInfo);
@@ -1398,14 +1393,14 @@ void FScene::RemovePrimitive( UPrimitiveComponent* Primitive )
 void FScene::ReleasePrimitive( UPrimitiveComponent* PrimitiveComponent )
 {
 	// Send a command to the rendering thread to clean up any state dependent on this primitive
-	ENQUEUE_UNIQUE_RENDER_COMMAND_TWOPARAMETER(
-		FReleasePrimitiveCommand,
-		FScene*,Scene,this,
-		FPrimitiveComponentId,PrimitiveComponentId,PrimitiveComponent->ComponentId,
-	{
-		// Free the space in the indirect lighting cache
-		Scene->IndirectLightingCache.ReleasePrimitive(PrimitiveComponentId);
-	});
+	FScene* Scene = this;
+	FPrimitiveComponentId PrimitiveComponentId = PrimitiveComponent->ComponentId;
+	ENQUEUE_RENDER_COMMAND(FReleasePrimitiveCommand)(
+		[Scene, PrimitiveComponentId](FRHICommandList&)
+		{
+			// Free the space in the indirect lighting cache
+			Scene->IndirectLightingCache.ReleasePrimitive(PrimitiveComponentId);
+		});
 }
 
 void FScene::AssignAvailableShadowMapChannelForLight(FLightSceneInfo* LightSceneInfo)
