IndirectLightingCache:
-BeginDeferredUpdateStaticMeshes only when original IndirectLightingCacheAllocation was null or bPointSample changed (same conditions as in mesh draw command processor).
-Fixed crash when BeginDeferredUpdateStaticMeshes was called in parallel on non render thread.

diff --git a/Engine/Source/Runtime/Renderer/Private/IndirectLightingCache.cpp b/Engine/Source/Runtime/Renderer/Private/IndirectLightingCache.cpp
--- a/Engine/Source/Runtime/Renderer/Private/IndirectLightingCache.cpp
+++ b/Engine/Source/Runtime/Renderer/Private/IndirectLightingCache.cpp
@@ -395,16 +395,24 @@ class FUpdateCachePrimitivesTask
 	FSceneRenderer& Renderer;
 	TMap<FIntVector, FBlockUpdateInfo>& BlocksToUpdate;
 	TArray<FIndirectLightingCacheAllocation*>& TransitionsOverTimeToUpdate;
+	TArray<FPrimitiveSceneInfo*>& PrimitivesToUpdateStaticMeshes;
 	bool bAllowUnbuiltPreview;
 
 public:
 
-	FUpdateCachePrimitivesTask(FIndirectLightingCache* InILC, FScene* InScene, FSceneRenderer& InRenderer, bool bInAllowUnbuiltPreview, TMap<FIntVector, FBlockUpdateInfo>& OutBlocksToUpdate, TArray<FIndirectLightingCacheAllocation*>& OutTransitionsOverTimeToUpdate)
+	FUpdateCachePrimitivesTask(
+		FIndirectLightingCache* InILC, 
+		FScene* InScene, FSceneRenderer& InRenderer, 
+		bool bInAllowUnbuiltPreview, 
+		TMap<FIntVector, FBlockUpdateInfo>& OutBlocksToUpdate, 
+		TArray<FIndirectLightingCacheAllocation*>& OutTransitionsOverTimeToUpdate, 
+		TArray<FPrimitiveSceneInfo*>& OutPrimitivesToUpdateStaticMeshes)
 		: ILC(InILC)
 		, Scene(InScene)
 		, Renderer(InRenderer)
 		, BlocksToUpdate(OutBlocksToUpdate)
 		, TransitionsOverTimeToUpdate(OutTransitionsOverTimeToUpdate)
+		, PrimitivesToUpdateStaticMeshes(OutPrimitivesToUpdateStaticMeshes)
 		, bAllowUnbuiltPreview(bInAllowUnbuiltPreview)
 	{
 	}
@@ -423,20 +431,28 @@ public:
 
 	void DoTask(ENamedThreads::Type CurrentThread, const FGraphEventRef& MyCompletionGraphEvent)
 	{
-		ILC->UpdateCachePrimitivesInternal(Scene, Renderer, bAllowUnbuiltPreview, BlocksToUpdate, TransitionsOverTimeToUpdate);
+		ILC->UpdateCachePrimitivesInternal(Scene, Renderer, bAllowUnbuiltPreview, BlocksToUpdate, TransitionsOverTimeToUpdate, PrimitivesToUpdateStaticMeshes);
 	}
 };
 
 void FIndirectLightingCache::StartUpdateCachePrimitivesTask(FScene* Scene, FSceneRenderer& Renderer, bool bAllowUnbuiltPreview, FILCUpdatePrimTaskData& OutTaskData)
 {
-	OutTaskData.TaskRef = TGraphTask<FUpdateCachePrimitivesTask>::CreateTask().ConstructAndDispatchWhenReady(this, Scene, Renderer, bAllowUnbuiltPreview, OutTaskData.OutBlocksToUpdate, OutTaskData.OutTransitionsOverTimeToUpdate);
+	OutTaskData.TaskRef = TGraphTask<FUpdateCachePrimitivesTask>::CreateTask().ConstructAndDispatchWhenReady(
+		this, 
+		Scene, 
+		Renderer, 
+		bAllowUnbuiltPreview, 
+		OutTaskData.OutBlocksToUpdate, 
+		OutTaskData.OutTransitionsOverTimeToUpdate,
+		OutTaskData.OutPrimitivesToUpdateStaticMeshes
+	);
 }
 
 void FIndirectLightingCache::FinalizeCacheUpdates(FScene* Scene, FSceneRenderer& Renderer, FILCUpdatePrimTaskData& TaskData)
 {
 	SCOPE_CYCLE_COUNTER(STAT_UpdateIndirectLightingCacheFinalize);	
 	FTaskGraphInterface::Get().WaitUntilTaskCompletes(TaskData.TaskRef, ENamedThreads::GetRenderThread_Local());
-	FinalizeUpdateInternal_RenderThread(Scene, Renderer, TaskData.OutBlocksToUpdate, TaskData.OutTransitionsOverTimeToUpdate);
+	FinalizeUpdateInternal_RenderThread(Scene, Renderer, TaskData.OutBlocksToUpdate, TaskData.OutTransitionsOverTimeToUpdate, TaskData.OutPrimitivesToUpdateStaticMeshes);
 }
 
 void FIndirectLightingCache::UpdateCache(FScene* Scene, FSceneRenderer& Renderer, bool bAllowUnbuiltPreview)
@@ -445,9 +461,10 @@ void FIndirectLightingCache::UpdateCache(FScene* Scene, FSceneRenderer& Renderer
 
 	TMap<FIntVector, FBlockUpdateInfo>BlocksToUpdate;
 	TArray<FIndirectLightingCacheAllocation*> TransitionsOverTimeToUpdate;
+	TArray<FPrimitiveSceneInfo*> PrimitivesToUpdateStaticMeshes;
 
-	UpdateCachePrimitivesInternal(Scene, Renderer, bAllowUnbuiltPreview, BlocksToUpdate, TransitionsOverTimeToUpdate);
-	FinalizeUpdateInternal_RenderThread(Scene, Renderer, BlocksToUpdate, TransitionsOverTimeToUpdate);
+	UpdateCachePrimitivesInternal(Scene, Renderer, bAllowUnbuiltPreview, BlocksToUpdate, TransitionsOverTimeToUpdate, PrimitivesToUpdateStaticMeshes);
+	FinalizeUpdateInternal_RenderThread(Scene, Renderer, BlocksToUpdate, TransitionsOverTimeToUpdate, PrimitivesToUpdateStaticMeshes);
 }
 
 bool FIndirectLightingCache::IndirectLightingAllowed(FScene* Scene, FSceneRenderer& Renderer) const
@@ -463,13 +480,13 @@ bool FIndirectLightingCache::IndirectLightingAllowed(FScene* Scene, FSceneRender
 	return bAnyViewAllowsIndirectLightingCache;
 }
 
-void FIndirectLightingCache::ProcessPrimitiveUpdate(FScene* Scene, FViewInfo& View, int32 PrimitiveIndex, bool bAllowUnbuiltPreview, bool bAllowVolumeSample, TMap<FIntVector, FBlockUpdateInfo>& OutBlocksToUpdate, TArray<FIndirectLightingCacheAllocation*>& OutTransitionsOverTimeToUpdate)
+void FIndirectLightingCache::ProcessPrimitiveUpdate(FScene* Scene, FViewInfo& View, int32 PrimitiveIndex, bool bAllowUnbuiltPreview, bool bAllowVolumeSample, TMap<FIntVector, FBlockUpdateInfo>& OutBlocksToUpdate, TArray<FIndirectLightingCacheAllocation*>& OutTransitionsOverTimeToUpdate, TArray<FPrimitiveSceneInfo*>& OutPrimitivesToUpdateStaticMeshes)
 {
 	FPrimitiveSceneInfo* PrimitiveSceneInfo = Scene->Primitives[PrimitiveIndex];
 	const bool bIndirectLightingCacheBufferWasDirty = PrimitiveSceneInfo->NeedsIndirectLightingCacheBufferUpdate();
 	
 	const TMap<FPrimitiveComponentId, FAttachmentGroupSceneInfo>& AttachmentGroups = Scene->AttachmentGroups;
-	UpdateCachePrimitive(AttachmentGroups, PrimitiveSceneInfo, bAllowUnbuiltPreview, bAllowVolumeSample, OutBlocksToUpdate, OutTransitionsOverTimeToUpdate);
+	UpdateCachePrimitive(AttachmentGroups, PrimitiveSceneInfo, bAllowUnbuiltPreview, bAllowVolumeSample, OutBlocksToUpdate, OutTransitionsOverTimeToUpdate, OutPrimitivesToUpdateStaticMeshes);
 
 	// If it was already dirty, then the primitive is already in one of the view dirty primitive list at this point.
 	// This also ensures that a primitive does not get added twice to the list, which could create an array reallocation.
@@ -481,7 +498,7 @@ void FIndirectLightingCache::ProcessPrimitiveUpdate(FScene* Scene, FViewInfo& Vi
 	}
 }
 
-void FIndirectLightingCache::UpdateCachePrimitivesInternal(FScene* Scene, FSceneRenderer& Renderer, bool bAllowUnbuiltPreview, TMap<FIntVector, FBlockUpdateInfo>& OutBlocksToUpdate, TArray<FIndirectLightingCacheAllocation*>& OutTransitionsOverTimeToUpdate)
+void FIndirectLightingCache::UpdateCachePrimitivesInternal(FScene* Scene, FSceneRenderer& Renderer, bool bAllowUnbuiltPreview, TMap<FIntVector, FBlockUpdateInfo>& OutBlocksToUpdate, TArray<FIndirectLightingCacheAllocation*>& OutTransitionsOverTimeToUpdate, TArray<FPrimitiveSceneInfo*>& OutPrimitivesToUpdateStaticMeshes)
 {
 	SCOPE_CYCLE_COUNTER(STAT_UpdateIndirectLightingCachePrims);
 
@@ -498,7 +515,7 @@ void FIndirectLightingCache::UpdateCachePrimitivesInternal(FScene* Scene, FScene
 				FPrimitiveSceneInfo* PrimitiveSceneInfo = Scene->Primitives[PrimitiveIndex];
 				const bool bIndirectLightingCacheBufferWasDirty = PrimitiveSceneInfo->NeedsIndirectLightingCacheBufferUpdate();
 				
-				UpdateCachePrimitive(AttachmentGroups, PrimitiveSceneInfo, false, true, OutBlocksToUpdate, OutTransitionsOverTimeToUpdate);
+				UpdateCachePrimitive(AttachmentGroups, PrimitiveSceneInfo, false, true, OutBlocksToUpdate, OutTransitionsOverTimeToUpdate, OutPrimitivesToUpdateStaticMeshes);
 
 				// If it was already dirty, then the primitive is already in one of the view dirty primitive list at this point.
 				// This also ensures that a primitive does not get added twice to the list, which could create an array reallocation.
@@ -533,7 +550,7 @@ void FIndirectLightingCache::UpdateCachePrimitivesInternal(FScene* Scene, FScene
 					uint32 PrimitiveIndex = BitIt.GetIndex();
 					// FDrawTranslucentMeshAction::AllowIndirectLightingCacheVolumeTexture doesn't allow volume samples on translucency, so we only need to support one if the primitive has at least one opaque material
 					const bool bAllowVolumeSample = View.PrimitiveViewRelevanceMap[PrimitiveIndex].bOpaqueRelevance;
-					ProcessPrimitiveUpdate(Scene, View, PrimitiveIndex, bAllowUnbuiltPreview, bAllowVolumeSample, OutBlocksToUpdate, OutTransitionsOverTimeToUpdate);
+					ProcessPrimitiveUpdate(Scene, View, PrimitiveIndex, bAllowUnbuiltPreview, bAllowVolumeSample, OutBlocksToUpdate, OutTransitionsOverTimeToUpdate, OutPrimitivesToUpdateStaticMeshes);
 				}
 
 				// Any visible primitives with an indirect shadow need their ILC updated, since that determines the indirect shadow direction
@@ -543,7 +560,7 @@ void FIndirectLightingCache::UpdateCachePrimitivesInternal(FScene* Scene, FScene
 
 					if (!View.PrimitiveVisibilityMap[PrimitiveIndex])
 					{
-						ProcessPrimitiveUpdate(Scene, View, PrimitiveIndex, bAllowUnbuiltPreview, true, OutBlocksToUpdate, OutTransitionsOverTimeToUpdate);
+						ProcessPrimitiveUpdate(Scene, View, PrimitiveIndex, bAllowUnbuiltPreview, true, OutBlocksToUpdate, OutTransitionsOverTimeToUpdate, OutPrimitivesToUpdateStaticMeshes);
 					}
 				}
 			}			
@@ -553,7 +570,7 @@ void FIndirectLightingCache::UpdateCachePrimitivesInternal(FScene* Scene, FScene
 	}
 }
 
-void FIndirectLightingCache::FinalizeUpdateInternal_RenderThread(FScene* Scene, FSceneRenderer& Renderer, TMap<FIntVector, FBlockUpdateInfo>& BlocksToUpdate, const TArray<FIndirectLightingCacheAllocation*>& TransitionsOverTimeToUpdate)
+void FIndirectLightingCache::FinalizeUpdateInternal_RenderThread(FScene* Scene, FSceneRenderer& Renderer, TMap<FIntVector, FBlockUpdateInfo>& BlocksToUpdate, const TArray<FIndirectLightingCacheAllocation*>& TransitionsOverTimeToUpdate, TArray<FPrimitiveSceneInfo*>& PrimitivesToUpdateStaticMeshes)
 {
 	check(IsInRenderingThread());
 
@@ -568,6 +585,11 @@ void FIndirectLightingCache::FinalizeUpdateInternal_RenderThread(FScene* Scene,
 			SCOPE_CYCLE_COUNTER(STAT_UpdateIndirectLightingCacheTransitions);
 			UpdateTransitionsOverTime(TransitionsOverTimeToUpdate, Renderer.ViewFamily.DeltaWorldTime);
 		}
+
+		for (int32 PrimitiveIndex = 0; PrimitiveIndex < PrimitivesToUpdateStaticMeshes.Num(); ++PrimitiveIndex)
+		{
+			PrimitivesToUpdateStaticMeshes[PrimitiveIndex]->BeginDeferredUpdateStaticMeshes();
+		}
 	}	
 
 	if (GCacheDrawLightingSamples || Renderer.ViewFamily.EngineShowFlags.VolumeLightingSamples || GCacheDrawDirectionalShadowing)
@@ -654,7 +676,8 @@ void FIndirectLightingCache::UpdateCachePrimitive(
 	bool bAllowUnbuiltPreview,
 	bool bAllowVolumeSample,
 	TMap<FIntVector, FBlockUpdateInfo>& BlocksToUpdate,
-	TArray<FIndirectLightingCacheAllocation*>& TransitionsOverTimeToUpdate)
+	TArray<FIndirectLightingCacheAllocation*>& TransitionsOverTimeToUpdate,
+	TArray<FPrimitiveSceneInfo*>& PrimitivesToUpdateStaticMeshes)
 {
 	FPrimitiveSceneProxy* PrimitiveSceneProxy = PrimitiveSceneInfo->Proxy;	
 	FIndirectLightingCacheAllocation** PrimitiveAllocationPtr = PrimitiveAllocations.Find(PrimitiveSceneInfo->PrimitiveComponentId);
@@ -718,10 +741,12 @@ void FIndirectLightingCache::UpdateCachePrimitive(
 			}
 		}
 
-		if (OriginalCachedIndirectLightingCacheAllocation != PrimitiveSceneInfo->IndirectLightingCacheAllocation)
+		// Cached mesh draw command shader selection depends on IndirectLightingCacheAllocation existence and bPointSample value.
+		// If any of those changes, we need to re-cache mesh draw commands.
+		if (OriginalCachedIndirectLightingCacheAllocation == nullptr
+			|| OriginalCachedIndirectLightingCacheAllocation->bPointSample != PrimitiveSceneInfo->IndirectLightingCacheAllocation->bPointSample)
 		{
-			// Update cached mesh draw commands, so they can pick up new IndirectLightingCacheAllocation.
-			PrimitiveSceneInfo->BeginDeferredUpdateStaticMeshes();
+			PrimitivesToUpdateStaticMeshes.Add(PrimitiveSceneInfo);
 		}
 	}
 }
diff --git a/Engine/Source/Runtime/Renderer/Private/ScenePrivate.h b/Engine/Source/Runtime/Renderer/Private/ScenePrivate.h
--- a/Engine/Source/Runtime/Renderer/Private/ScenePrivate.h
+++ b/Engine/Source/Runtime/Renderer/Private/ScenePrivate.h
@@ -1918,6 +1918,7 @@ struct FILCUpdatePrimTaskData
 	FGraphEventRef TaskRef;
 	TMap<FIntVector, FBlockUpdateInfo> OutBlocksToUpdate;
 	TArray<FIndirectLightingCacheAllocation*> OutTransitionsOverTimeToUpdate;
+	TArray<FPrimitiveSceneInfo*> OutPrimitivesToUpdateStaticMeshes;
 };
 
 /** 
@@ -1964,13 +1965,13 @@ private:
 	/** Internal helper to determine if indirect lighting is enabled at all */
 	bool IndirectLightingAllowed(FScene* Scene, FSceneRenderer& Renderer) const;
 
-	void ProcessPrimitiveUpdate(FScene* Scene, FViewInfo& View, int32 PrimitiveIndex, bool bAllowUnbuiltPreview, bool bAllowVolumeSample, TMap<FIntVector, FBlockUpdateInfo>& OutBlocksToUpdate, TArray<FIndirectLightingCacheAllocation*>& OutTransitionsOverTimeToUpdate);
+	void ProcessPrimitiveUpdate(FScene* Scene, FViewInfo& View, int32 PrimitiveIndex, bool bAllowUnbuiltPreview, bool bAllowVolumeSample, TMap<FIntVector, FBlockUpdateInfo>& OutBlocksToUpdate, TArray<FIndirectLightingCacheAllocation*>& OutTransitionsOverTimeToUpdate, TArray<FPrimitiveSceneInfo*>& OutPrimitivesToUpdateStaticMeshes);
 
 	/** Internal helper to perform the work of updating the cache primitives.  Can be done on any thread as a task */
-	void UpdateCachePrimitivesInternal(FScene* Scene, FSceneRenderer& Renderer, bool bAllowUnbuiltPreview, TMap<FIntVector, FBlockUpdateInfo>& OutBlocksToUpdate, TArray<FIndirectLightingCacheAllocation*>& OutTransitionsOverTimeToUpdate);
+	void UpdateCachePrimitivesInternal(FScene* Scene, FSceneRenderer& Renderer, bool bAllowUnbuiltPreview, TMap<FIntVector, FBlockUpdateInfo>& OutBlocksToUpdate, TArray<FIndirectLightingCacheAllocation*>& OutTransitionsOverTimeToUpdate, TArray<FPrimitiveSceneInfo*>& OutPrimitivesToUpdateStaticMeshes);
 
 	/** Internal helper to perform blockupdates and transition updates on the results of UpdateCachePrimitivesInternal.  Must be on render thread. */
-	void FinalizeUpdateInternal_RenderThread(FScene* Scene, FSceneRenderer& Renderer, TMap<FIntVector, FBlockUpdateInfo>& BlocksToUpdate, const TArray<FIndirectLightingCacheAllocation*>& TransitionsOverTimeToUpdate);
+	void FinalizeUpdateInternal_RenderThread(FScene* Scene, FSceneRenderer& Renderer, TMap<FIntVector, FBlockUpdateInfo>& BlocksToUpdate, const TArray<FIndirectLightingCacheAllocation*>& TransitionsOverTimeToUpdate, TArray<FPrimitiveSceneInfo*>& PrimitivesToUpdateStaticMeshes);
 
 	/** Internal helper which adds an entry to the update lists for this allocation, if needed (due to movement, etc). Returns true if the allocation was updated or will be udpated */
 	bool UpdateCacheAllocation(
@@ -1992,7 +1993,8 @@ private:
 		bool bAllowUnbuiltPreview, 
 		bool bAllowVolumeSample, 
 		TMap<FIntVector, FBlockUpdateInfo>& BlocksToUpdate, 
-		TArray<FIndirectLightingCacheAllocation*>& TransitionsOverTimeToUpdate);	
+		TArray<FIndirectLightingCacheAllocation*>& TransitionsOverTimeToUpdate,
+		TArray<FPrimitiveSceneInfo*>& PrimitivesToUpdateStaticMeshes);
 
 	/** Updates the contents of the volume texture blocks in BlocksToUpdate. */
 	void UpdateBlocks(FScene* Scene, FViewInfo* DebugDrawingView, TMap<FIntVector, FBlockUpdateInfo>& BlocksToUpdate);
