Enabled r.ParallelInitViews and r.ParallelViewsCustomDataUpdate by default inside the editor, so it's consistent with a standalone game.

diff --git a/Engine/Source/Runtime/Renderer/Private/SceneVisibility.cpp b/Engine/Source/Runtime/Renderer/Private/SceneVisibility.cpp
--- a/Engine/Source/Runtime/Renderer/Private/SceneVisibility.cpp
+++ b/Engine/Source/Runtime/Renderer/Private/SceneVisibility.cpp
@@ -211,22 +211,14 @@ static FAutoConsoleVariableRef CVarDistanceFadeMaxTravel( TEXT("r.DistanceFadeMa
 
 static TAutoConsoleVariable<int32> CVarParallelInitViews(
 	TEXT("r.ParallelInitViews"),
-#if WITH_EDITOR
-	0,  
-#else
-	1,  
-#endif
+	1,
 	TEXT("Toggles parallel init views. 0 = off; 1 = on"),
 	ECVF_RenderThreadSafe
 	);          
 
 static TAutoConsoleVariable<int32> CVarParallelPostInitViewCustomData(
 	TEXT("r.ParallelViewsCustomDataUpdate"),
-#if WITH_EDITOR
-	0,
-#else
 	1,
-#endif
 	TEXT("Toggles parallel views custom data update. 0 = off; 1 = on"),
 	ECVF_RenderThreadSafe
 );
