Volumetric fog can be pricey, add a stat gpu for it

diff --git a/Engine/Source/Runtime/Renderer/Private/VolumetricFog.cpp b/Engine/Source/Runtime/Renderer/Private/VolumetricFog.cpp
--- a/Engine/Source/Runtime/Renderer/Private/VolumetricFog.cpp
+++ b/Engine/Source/Runtime/Renderer/Private/VolumetricFog.cpp
@@ -102,6 +102,8 @@ FAutoConsoleVariableRef CVarInverseSquaredLightDistanceBiasScale(
 
 IMPLEMENT_GLOBAL_SHADER_PARAMETER_STRUCT(FVolumetricFogGlobalData, "VolumetricFog");
 
+DECLARE_GPU_STAT(VolumetricFog);
+
 FVolumetricFogGlobalData::FVolumetricFogGlobalData()
 {}
 
@@ -970,8 +972,9 @@ void FDeferredShadingSceneRenderer::ComputeVolumetricFog(FRHICommandListImmediat
 	check(RHICmdListImmediate.IsOutsideRenderPass());
 
 	if (ShouldRenderVolumetricFog())
-	{
+	{		
 		QUICK_SCOPE_CYCLE_COUNTER(STAT_VolumetricFog);
+		SCOPED_GPU_STAT(RHICmdListImmediate, VolumetricFog);
 
 		const FExponentialHeightFogSceneInfo& FogInfo = Scene->ExponentialFogs[0];
 
