InitViews optimization. Packed mesh draw command infos into a continuous array in FPrimitiveSceneInfo for better cache efficiency. This allows entirely to skip fetching FStaticMesh data inside InitViews and just rely on two small continuous arrays: StaticMeshRelevances and StaticMeshCommandInfos.

diff --git a/Engine/Source/Runtime/Renderer/Private/DeferredShadingRenderer.cpp b/Engine/Source/Runtime/Renderer/Private/DeferredShadingRenderer.cpp
--- a/Engine/Source/Runtime/Renderer/Private/DeferredShadingRenderer.cpp
+++ b/Engine/Source/Runtime/Renderer/Private/DeferredShadingRenderer.cpp
@@ -675,12 +675,15 @@ bool FDeferredShadingSceneRenderer::GatherRayTracingWorldInstances(FRHICommandLi
 						const auto& StaticMeshMdcIndices = SceneProxy->RayTracingLodIndexToMeshDrawCommandIndicies[LODToRender.GetRayTracedLOD()];
 						for (int StaticMeshIndex = 0; StaticMeshIndex < StaticMeshMdcIndices.Num(); StaticMeshIndex++)
 						{
-							FStaticMesh& StaticMesh = Scene->Primitives[PrimitiveIndex]->StaticMeshes[StaticMeshMdcIndices[StaticMeshIndex].StaticMeshIndex];
+							const FPrimitiveSceneInfo* PrimitiveInfo = Scene->Primitives[PrimitiveIndex];
+							const FStaticMeshRelevance& StaticMeshRelevance = PrimitiveInfo->StaticMeshRelevances[StaticMeshMdcIndices[StaticMeshIndex].StaticMeshIndex];
+
 							const int32 CommandIndex = StaticMeshMdcIndices[StaticMeshIndex].CommandIndex;
 							if (CommandIndex >= 0)
 							{
 								const FCachedPassMeshDrawList& SceneDrawList = Scene->CachedDrawLists[EMeshPass::RayTracing];
-								const FCachedMeshDrawCommandInfo& CachedMeshDrawCommand = StaticMesh.CachedMeshDrawCommands[EMeshPass::RayTracing];
+								const int32 RayTracingStaticMeshCommandInfoIndex = StaticMeshRelevance.GetStaticMeshCommandInfoIndex(EMeshPass::RayTracing);
+								const FCachedMeshDrawCommandInfo& CachedMeshDrawCommand = Scene->Primitives[PrimitiveIndex]->StaticMeshCommandInfos[RayTracingStaticMeshCommandInfoIndex];
 								FVisibleMeshDrawCommand NewVisibleMeshDrawCommand;
 
 								NewVisibleMeshDrawCommand.Setup(
@@ -695,6 +698,8 @@ bool FDeferredShadingSceneRenderer::GatherRayTracingWorldInstances(FRHICommandLi
 							}
 							else
 							{
+								const FStaticMesh& StaticMesh = PrimitiveInfo->StaticMeshes[StaticMeshMdcIndices[StaticMeshIndex].StaticMeshIndex];
+
 								uint64 BatchVisibility = 1;
 								if (StaticMesh.bRequiresPerElementVisibility)
 								{
diff --git a/Engine/Source/Runtime/Renderer/Private/MeshPassProcessor.cpp b/Engine/Source/Runtime/Renderer/Private/MeshPassProcessor.cpp
--- a/Engine/Source/Runtime/Renderer/Private/MeshPassProcessor.cpp
+++ b/Engine/Source/Runtime/Renderer/Private/MeshPassProcessor.cpp
@@ -962,7 +962,9 @@ FMeshPassProcessor::FMeshPassProcessor(const FScene* InScene, ERHIFeatureLevel::
 	, FeatureLevel(InFeatureLevel)
 	, ViewIfDynamicMeshCommand(InViewIfDynamicMeshCommand)
 	, DrawListContext(InDrawListContext)
-{}
+{
+	static_assert(EMeshPass::Num <= EMeshPass::NumBits * 8, "EMeshPass::NumBits is too small.");
+}
 
 ERasterizerFillMode FMeshPassProcessor::ComputeMeshFillMode(const FMeshBatch& Mesh, const FMaterial& InMaterialResource) const
 {
diff --git a/Engine/Source/Runtime/Renderer/Private/PrimitiveSceneInfo.cpp b/Engine/Source/Runtime/Renderer/Private/PrimitiveSceneInfo.cpp
--- a/Engine/Source/Runtime/Renderer/Private/PrimitiveSceneInfo.cpp
+++ b/Engine/Source/Runtime/Renderer/Private/PrimitiveSceneInfo.cpp
@@ -156,6 +156,134 @@ FPrimitiveSceneInfo::FPrimitiveSceneInfo(UPrimitiveComponent* InComponent,FScene
 FPrimitiveSceneInfo::~FPrimitiveSceneInfo()
 {
 	check(!OctreeId.IsValidId());
+
+	RemoveCachedMeshDrawCommands();
+}
+
+void FPrimitiveSceneInfo::CacheMeshDrawCommands(FRHICommandListImmediate& RHICmdList)
+{
+	check(StaticMeshCommandInfos.Num() == 0);
+
+	// TEMP: optionally allow to switch off mesh command caching for cooked mobile
+	extern int32 MeshDrawCommandPipelineMobileCooked();
+	const int32 CacheMeshCommandsVal = MeshDrawCommandPipelineMobileCooked();
+	if (CacheMeshCommandsVal <= 0)
+	{
+		return;
+	}
+
+
+	int32 MeshWithCachedCommandsNum = 0;
+	for (int32 MeshIndex = 0; MeshIndex < StaticMeshes.Num(); MeshIndex++)
+	{
+		const FStaticMesh& Mesh = StaticMeshes[MeshIndex];
+		if (SupportsCachingMeshDrawCommands(Mesh.VertexFactory, Proxy))
+		{
+			++MeshWithCachedCommandsNum;
+		}
+	}
+
+	if (MeshWithCachedCommandsNum > 0)
+	{
+		//@todo - only need material uniform buffers to be created since we are going to cache pointers to them
+		// Any updates (after initial creation) don't need to be forced here
+		FMaterialRenderProxy::UpdateDeferredCachedUniformExpressions();
+
+		// Reserve based on assumption that we have on average 2 cached mesh draw commands per mesh.
+		StaticMeshCommandInfos.Reserve(MeshWithCachedCommandsNum * 2);
+
+		QUICK_SCOPE_CYCLE_COUNTER(STAT_CacheMeshDrawCommands);
+		FMemMark Mark(FMemStack::Get());
+
+		const EShadingPath ShadingPath = Scene->GetShadingPath();
+
+		for (int32 MeshIndex = 0; MeshIndex < StaticMeshes.Num(); MeshIndex++)
+		{
+			FStaticMeshRelevance& MeshRelevance = StaticMeshRelevances[MeshIndex];
+			FStaticMesh& Mesh = StaticMeshes[MeshIndex];
+
+			check(MeshRelevance.CommandInfosMask == 0);
+			MeshRelevance.CommandInfosBase = StaticMeshCommandInfos.Num();
+
+			if (SupportsCachingMeshDrawCommands(Mesh.VertexFactory, Proxy))
+			{
+				for (int32 PassIndex = 0; PassIndex < EMeshPass::Num; PassIndex++)
+				{
+					EMeshPass::Type PassType = (EMeshPass::Type)PassIndex;
+
+					if ((FPassProcessorManager::GetPassFlags(ShadingPath, PassType) & EMeshPassFlags::CachedMeshCommands) != EMeshPassFlags::None)
+					{
+						FCachedMeshDrawCommandInfo CommandInfo;
+						CommandInfo.CommandIndex = -1;
+						CommandInfo.MeshPass = PassType;
+
+						FCachedPassMeshDrawList& SceneDrawList = Scene->CachedDrawLists[PassType];
+						FCachedPassMeshDrawListContext CachedPassMeshDrawListContext(CommandInfo, SceneDrawList, *Scene);
+
+						PassProcessorCreateFunction CreateFunction = FPassProcessorManager::GetCreateFunction(ShadingPath, PassType);
+						FMeshPassProcessor* PassMeshProcessor = CreateFunction(Scene, nullptr, CachedPassMeshDrawListContext);
+
+						if (PassMeshProcessor != nullptr)
+						{
+							check(!Mesh.bRequiresPerElementVisibility);
+							uint64 BatchElementMask = ~0ull;
+							PassMeshProcessor->AddMeshBatch(Mesh, BatchElementMask, Proxy);
+
+							PassMeshProcessor->~FMeshPassProcessor();
+						}
+
+						if (CommandInfo.CommandIndex != -1)
+						{
+							static_assert(sizeof(MeshRelevance.CommandInfosMask) * 8 >= EMeshPass::Num, "CommandInfosMask is too small to contain all mesh passes.");
+
+							MeshRelevance.CommandInfosMask |= (1 << PassType);
+
+							StaticMeshCommandInfos.Add(CommandInfo);
+						}
+					}
+				}
+			}
+		}
+	}
+}
+
+void FPrimitiveSceneInfo::RemoveCachedMeshDrawCommands()
+{
+	for (int32 CommandIndex = 0; CommandIndex < StaticMeshCommandInfos.Num(); ++CommandIndex)
+	{
+		const FCachedMeshDrawCommandInfo& CachedCommand = StaticMeshCommandInfos[CommandIndex];
+		if (CachedCommand.CommandIndex != -1)
+		{
+			FCachedPassMeshDrawList& PassDrawList = Scene->CachedDrawLists[CachedCommand.MeshPass];
+
+			const FSetElementId StateBucketId = FSetElementId::FromInteger(CachedCommand.StateBucketId);
+			checkSlow(StateBucketId.IsValidId());
+			FMeshDrawCommandStateBucket& StateBucket = Scene->CachedMeshDrawCommandStateBuckets[StateBucketId];
+			if (CachedCommand.StateBucketId != -1)
+			{
+				if (StateBucket.Num == 1)
+				{
+					Scene->CachedMeshDrawCommandStateBuckets.Remove(StateBucketId);
+				}
+				else
+				{
+					StateBucket.Num--;
+				}
+			}
+
+			PassDrawList.MeshDrawCommands.RemoveAt(CachedCommand.CommandIndex);
+		}
+	}
+
+	for (int32 MeshIndex = 0; MeshIndex < StaticMeshRelevances.Num(); ++MeshIndex)
+	{
+		FStaticMeshRelevance& MeshRelevance = StaticMeshRelevances[MeshIndex];
+
+		MeshRelevance.CommandInfosBase = 0;
+		MeshRelevance.CommandInfosMask = 0;
+	}
+
+	StaticMeshCommandInfos.Empty();
 }
 
 void FPrimitiveSceneInfo::AddStaticMeshes(FRHICommandListImmediate& RHICmdList, bool bAddToStaticDrawLists)
@@ -195,6 +323,8 @@ void FPrimitiveSceneInfo::AddStaticMeshes(FRHICommandListImmediate& RHICmdList,
 		}
 	}
 
+	CacheMeshDrawCommands(RHICmdList);
+
 #if RHI_RAYTRACING
 	int MaxLOD = -1;
 
@@ -214,21 +344,23 @@ void FPrimitiveSceneInfo::AddStaticMeshes(FRHICommandListImmediate& RHICmdList,
 			FStaticMeshRelevance& MeshRelevance = StaticMeshRelevances[MeshIndex];
 			FStaticMesh& Mesh = StaticMeshes[MeshIndex];
 
-			if (Proxy->ScreenSizes[Mesh.LODIndex] != 0.0f)
+			if (Proxy->ScreenSizes[MeshRelevance.LODIndex] != 0.0f)
 			{
-				check(Proxy->ScreenSizes[Mesh.LODIndex] == MeshRelevance.ScreenSize);
+				check(Proxy->ScreenSizes[MeshRelevance.LODIndex] == MeshRelevance.ScreenSize);
 			}
 			else
 			{
 				check(MeshRelevance.ScreenSize != 0.0f);
-				Proxy->ScreenSizes[Mesh.LODIndex] = MeshRelevance.ScreenSize;
+				Proxy->ScreenSizes[MeshRelevance.LODIndex] = MeshRelevance.ScreenSize;
 			}
 
 			if (Mesh.Elements.Num() > 0)
 			{
-				const int32 CommandIndex = Mesh.CachedMeshDrawCommands[EMeshPass::RayTracing].CommandIndex;
-				if (SupportsCachingMeshDrawCommands(Mesh.VertexFactory, Proxy) && CommandIndex != -1)
+				const int32 RayTracingStaticMeshCommandInfoIndex = MeshRelevance.GetStaticMeshCommandInfoIndex(EMeshPass::RayTracing);
+
+				if (SupportsCachingMeshDrawCommands(Mesh.VertexFactory, Proxy) && RayTracingStaticMeshCommandInfoIndex != -1)
 				{
+					const int32 CommandIndex = StaticMeshCommandInfos[RayTracingStaticMeshCommandInfoIndex].CommandIndex;
 					Proxy->RayTracingLodIndexToMeshDrawCommandIndicies[Mesh.LODIndex].Add({ MeshIndex, CommandIndex });
 				}
 				else
@@ -412,6 +544,7 @@ void FPrimitiveSceneInfo::RemoveStaticMeshes()
 	// Remove static meshes from the scene.
 	StaticMeshes.Empty();
 	StaticMeshRelevances.Empty();
+	RemoveCachedMeshDrawCommands();
 }
 
 void FPrimitiveSceneInfo::RemoveFromScene(bool bUpdateStaticDrawLists)
@@ -489,6 +622,12 @@ void FPrimitiveSceneInfo::UpdateStaticMeshes(FRHICommandListImmediate& RHICmdLis
 			StaticMeshes[MeshIndex].AddToDrawLists(RHICmdList, Scene);
 		}
 	}
+
+	RemoveCachedMeshDrawCommands();
+	if (bReAddToDrawLists)
+	{
+		CacheMeshDrawCommands(RHICmdList);
+	}
 }
 
 void FPrimitiveSceneInfo::UpdateUniformBuffer(FRHICommandListImmediate& RHICmdList)
diff --git a/Engine/Source/Runtime/Renderer/Private/SceneCore.cpp b/Engine/Source/Runtime/Renderer/Private/SceneCore.cpp
--- a/Engine/Source/Runtime/Renderer/Private/SceneCore.cpp
+++ b/Engine/Source/Runtime/Renderer/Private/SceneCore.cpp
@@ -315,6 +315,28 @@ void FLightPrimitiveInteraction::FlushCachedShadowMapData()
 	}
 }
 
+int32 FStaticMeshRelevance::GetStaticMeshCommandInfoIndex(EMeshPass::Type MeshPass) const
+{
+	static_assert(sizeof(CommandInfosMask) * 8 >= EMeshPass::Num, "CommandInfosMask is too small to contain all mesh passes.");
+
+	int32 CommandInfoIndex = CommandInfosBase;
+
+	if ((CommandInfosMask & (1 << MeshPass)) == 0)
+	{
+		return -1;
+	}
+
+	for (int32 MeshPassIndex = 0; MeshPassIndex < MeshPass; ++MeshPassIndex)
+	{
+		if (CommandInfosMask & (1 << MeshPassIndex))
+		{
+			++CommandInfoIndex;
+		}
+	}
+
+	return CommandInfoIndex;
+}
+
 /*-----------------------------------------------------------------------------
 	FStaticMesh
 -----------------------------------------------------------------------------*/
@@ -340,11 +362,6 @@ void FStaticMesh::AddToDrawLists(FRHICommandListImmediate& RHICmdList, FScene* S
 	extern int32 MeshDrawCommandPipelineMobileCooked();
 	const int32 CacheMeshCommandsVal = MeshDrawCommandPipelineMobileCooked();
 
-	if (SupportsCachingMeshDrawCommands(VertexFactory, PrimitiveSceneInfo->Proxy) && CacheMeshCommandsVal > 0)
-	{
-		CacheMeshDrawCommands(Scene);
-	}
-
 	if (!PrimitiveSceneInfo->Proxy->ShouldRenderInMainPass() || !ShouldIncludeDomainInMeshPass(MaterialRenderProxy->GetMaterial(FeatureLevel)->GetMaterialDomain()))
 	{
 		return;
@@ -370,41 +387,6 @@ void FStaticMesh::AddToDrawLists(FRHICommandListImmediate& RHICmdList, FScene* S
 	}
 }
 
-void FStaticMesh::CacheMeshDrawCommands(FScene* Scene)
-{
-	//@todo - only need material uniform buffers to be created since we are going to cache pointers to them
-	// Any updates (after initial creation) don't need to be forced here
-	FMaterialRenderProxy::UpdateDeferredCachedUniformExpressions();
-
-	QUICK_SCOPE_CYCLE_COUNTER(STAT_CacheMeshDrawCommands);
-	FMemMark Mark(FMemStack::Get());
-
-	EShadingPath ShadingPath = Scene->GetShadingPath();
-
-	for (int32 PassIndex = 0; PassIndex < EMeshPass::Num; PassIndex++)
-	{
-		EMeshPass::Type PassType = (EMeshPass::Type)PassIndex;
-
-		if ((FPassProcessorManager::GetPassFlags(ShadingPath, PassType) & EMeshPassFlags::CachedMeshCommands) != EMeshPassFlags::None)
-		{
-			FCachedPassMeshDrawList& SceneDrawList = Scene->CachedDrawLists[PassType];
-			FCachedPassMeshDrawListContext CachedPassMeshDrawListContext(CachedMeshDrawCommands[PassType], SceneDrawList, *Scene);
-
-			PassProcessorCreateFunction CreateFunction = FPassProcessorManager::GetCreateFunction(ShadingPath, PassType);
-			FMeshPassProcessor* PassMeshProcessor = CreateFunction(Scene, nullptr, CachedPassMeshDrawListContext);
-
-			if (PassMeshProcessor != nullptr)
-			{
-				check(!bRequiresPerElementVisibility);
-				uint64 BatchElementMask = ~0ull;
-				PassMeshProcessor->AddMeshBatch(*this, BatchElementMask, PrimitiveSceneInfo->Proxy);
-
-				PassMeshProcessor->~FMeshPassProcessor();
-			}
-		}
-	}
-}
-
 void FStaticMesh::RemoveFromDrawLists(bool bMeshIsBeingDestroyed)
 {
 	if (bMeshIsBeingDestroyed)
@@ -432,36 +414,6 @@ void FStaticMesh::RemoveFromDrawLists(bool bMeshIsBeingDestroyed)
 			}
 		}
 	}
-
-	FScene* Scene = PrimitiveSceneInfo->Scene;
-
-	for (int32 PassIndex = 0; PassIndex < ARRAY_COUNT(CachedMeshDrawCommands); PassIndex++)
-	{
-		const FCachedMeshDrawCommandInfo& CachedCommand = CachedMeshDrawCommands[PassIndex];
-		if (CachedCommand.CommandIndex != -1)
-		{
-			FCachedPassMeshDrawList& PassDrawList = Scene->CachedDrawLists[PassIndex];
-
-			const FSetElementId StateBucketId = FSetElementId::FromInteger(CachedCommand.StateBucketId);
-			checkSlow(StateBucketId.IsValidId());
-			FMeshDrawCommandStateBucket& StateBucket = Scene->CachedMeshDrawCommandStateBuckets[StateBucketId];
-			if (CachedCommand.StateBucketId != -1)
-			{
-				if (StateBucket.Num == 1)
-				{
-					Scene->CachedMeshDrawCommandStateBuckets.Remove(StateBucketId);
-				}
-				else
-				{
-					StateBucket.Num--;
-				}
-			}
-		
-			PassDrawList.MeshDrawCommands.RemoveAt(CachedCommand.CommandIndex);
-		}
-
-		CachedMeshDrawCommands[PassIndex] = FCachedMeshDrawCommandInfo();
-	}
 }
 
 FStaticMesh::~FStaticMesh()
diff --git a/Engine/Source/Runtime/Renderer/Private/SceneCore.h b/Engine/Source/Runtime/Renderer/Private/SceneCore.h
--- a/Engine/Source/Runtime/Renderer/Private/SceneCore.h
+++ b/Engine/Source/Runtime/Renderer/Private/SceneCore.h
@@ -149,9 +149,6 @@ public:
 	/** Index of the mesh into the scene's StaticMeshBatchVisibility array. */
 	int32 BatchVisibilityId;
 
-	/** Cached mesh draw commands for this FStaticMesh. */
-	FCachedMeshDrawCommandInfo CachedMeshDrawCommands[EMeshPass::Num];
-
 	// Constructor/destructor.
 	FStaticMesh(
 		FPrimitiveSceneInfo* InPrimitiveSceneInfo,
@@ -177,8 +174,6 @@ public:
 	/** Adds the static mesh to the appropriate draw lists in a scene. */
 	void AddToDrawLists(FRHICommandListImmediate& RHICmdList, FScene* Scene);
 
-	void CacheMeshDrawCommands(FScene* Scene);
-
 	/** Removes the static mesh from all draw lists. */
 	void RemoveFromDrawLists(bool bMeshIsBeingDestroyed = false);
 
@@ -204,6 +199,8 @@ public:
 	FStaticMeshRelevance(const FStaticMesh& StaticMesh, float InScreenSize, bool InbSupportsCachingMeshDrawCommands)
 		: Id(StaticMesh.Id)
 		, ScreenSize(InScreenSize)
+		, CommandInfosMask(0)
+		, CommandInfosBase(0)
 		, LODIndex(StaticMesh.LODIndex)
 		, bDitheredLODTransition(StaticMesh.bDitheredLODTransition)
 		, bRequiresPerElementVisibility(StaticMesh.bRequiresPerElementVisibility)
@@ -222,6 +219,12 @@ public:
 	/** The screen space size to draw this primitive at */
 	float ScreenSize;
 
+	/** Starting offset into continuous array of command infos for this mesh in FPrimitiveSceneInfo::CachedMeshDrawCommandInfos. */
+	uint32 CommandInfosMask;
+
+	/* Every bit corresponds to one MeshPass. If bit is set, then FPrimitiveSceneInfo::CachedMeshDrawCommandInfos contains this mesh pass. */
+	uint16 CommandInfosBase;
+
 	/** LOD index of the mesh, used for fading LOD transitions. */
 	int8 LODIndex;
 
@@ -241,6 +244,9 @@ public:
 
 	/** Cached from vertex factory to avoid dereferencing VF in InitViews. */
 	uint32 bSupportsCachingMeshDrawCommands : 1;
+
+	/** Computes index of cached mesh draw command in FPrimitiveSceneInfo::CachedMeshDrawCommandInfos, for a given mesh pass. */
+	int32 GetStaticMeshCommandInfoIndex(EMeshPass::Type MeshPass) const;
 };
 
 /** The properties of a exponential height fog layer which are used for rendering. */
diff --git a/Engine/Source/Runtime/Renderer/Private/SceneVisibility.cpp b/Engine/Source/Runtime/Renderer/Private/SceneVisibility.cpp
--- a/Engine/Source/Runtime/Renderer/Private/SceneVisibility.cpp
+++ b/Engine/Source/Runtime/Renderer/Private/SceneVisibility.cpp
@@ -1691,7 +1691,14 @@ struct FDrawCommandRelevancePacket
 	bool bUseMeshDrawCommandPipeline;
 	bool bUseCachedMeshDrawCommands;
 
-	void AddCommandsForMesh(int32 PrimitiveIndex, const FStaticMeshRelevance& RESTRICT StaticMeshRelevance, const FStaticMesh& RESTRICT StaticMesh, const FScene* RESTRICT Scene, bool bCanCache, EMeshPass::Type PassType)
+	void AddCommandsForMesh(
+		int32 PrimitiveIndex, 
+		const FPrimitiveSceneInfo* InPrimitiveSceneInfo,
+		const FStaticMeshRelevance& RESTRICT StaticMeshRelevance, 
+		const FStaticMesh& RESTRICT StaticMesh, 
+		const FScene* RESTRICT Scene, 
+		bool bCanCache, 
+		EMeshPass::Type PassType)
 	{
 		if (bUseMeshDrawCommandPipeline)
 		{
@@ -1703,9 +1710,10 @@ struct FDrawCommandRelevancePacket
 
 			if (bUseCachedMeshCommand)
 			{
-				const FCachedMeshDrawCommandInfo& CachedMeshDrawCommand = StaticMesh.CachedMeshDrawCommands[PassType];
-				if (CachedMeshDrawCommand.CommandIndex != -1)
+				const int32 StaticMeshCommandInfoIndex = StaticMeshRelevance.GetStaticMeshCommandInfoIndex(PassType);
+				if (StaticMeshCommandInfoIndex >= 0)
 				{
+					const FCachedMeshDrawCommandInfo& CachedMeshDrawCommand = InPrimitiveSceneInfo->StaticMeshCommandInfos[StaticMeshCommandInfoIndex];
 					const FCachedPassMeshDrawList& SceneDrawList = Scene->CachedDrawLists[PassType];
 
 					FVisibleMeshDrawCommand NewVisibleMeshDrawCommand;
@@ -2103,43 +2111,43 @@ struct FRelevancePacket
 					{
 						if (StaticMeshRelevance.bUseForDepthPass && bDrawDepthOnly)
 						{
-							DrawCommandPacket.AddCommandsForMesh(PrimitiveIndex, StaticMeshRelevance, StaticMesh, Scene, bCanCache, EMeshPass::DepthPass);
+							DrawCommandPacket.AddCommandsForMesh(PrimitiveIndex, PrimitiveSceneInfo, StaticMeshRelevance, StaticMesh, Scene, bCanCache, EMeshPass::DepthPass);
 							MarkMask |= EMarkMaskBits::StaticMeshOccluderMapMask;
 						}
 
 						// Mark static mesh as visible for rendering
 						if (StaticMeshRelevance.bUseForMaterial)
 						{
-							DrawCommandPacket.AddCommandsForMesh(PrimitiveIndex, StaticMeshRelevance, StaticMesh, Scene, bCanCache, EMeshPass::BasePass);
+							DrawCommandPacket.AddCommandsForMesh(PrimitiveIndex, PrimitiveSceneInfo, StaticMeshRelevance, StaticMesh, Scene, bCanCache, EMeshPass::BasePass);
 							MarkMask |= EMarkMaskBits::StaticMeshVisibilityMapMask;
 
 							if (ShadingPath == EShadingPath::Mobile)
 							{
-								DrawCommandPacket.AddCommandsForMesh(PrimitiveIndex, StaticMeshRelevance, StaticMesh, Scene, bCanCache, EMeshPass::MobileBasePassCSM);
+								DrawCommandPacket.AddCommandsForMesh(PrimitiveIndex, PrimitiveSceneInfo, StaticMeshRelevance, StaticMesh, Scene, bCanCache, EMeshPass::MobileBasePassCSM);
 							}
 
 							if (ViewRelevance.bRenderCustomDepth)
 							{
-								DrawCommandPacket.AddCommandsForMesh(PrimitiveIndex, StaticMeshRelevance, StaticMesh, Scene, bCanCache, EMeshPass::CustomDepth);
+								DrawCommandPacket.AddCommandsForMesh(PrimitiveIndex, PrimitiveSceneInfo, StaticMeshRelevance, StaticMesh, Scene, bCanCache, EMeshPass::CustomDepth);
 							}
 
 							if (bAddLightmapDensityCommands)
 							{
-								DrawCommandPacket.AddCommandsForMesh(PrimitiveIndex, StaticMeshRelevance, StaticMesh, Scene, bCanCache, EMeshPass::LightmapDensity);
+								DrawCommandPacket.AddCommandsForMesh(PrimitiveIndex, PrimitiveSceneInfo, StaticMeshRelevance, StaticMesh, Scene, bCanCache, EMeshPass::LightmapDensity);
 							}
 #if !(UE_BUILD_SHIPPING || UE_BUILD_TEST)
 							else if (View.Family->UseDebugViewPS())
 							{
-								DrawCommandPacket.AddCommandsForMesh(PrimitiveIndex, StaticMeshRelevance, StaticMesh, Scene, bCanCache, EMeshPass::DebugViewMode);
+								DrawCommandPacket.AddCommandsForMesh(PrimitiveIndex, PrimitiveSceneInfo, StaticMeshRelevance, StaticMesh, Scene, bCanCache, EMeshPass::DebugViewMode);
 							}
 #endif
 							if (StaticMeshRelevance.bSelectable)
 							{
-								DrawCommandPacket.AddCommandsForMesh(PrimitiveIndex, StaticMeshRelevance, StaticMesh, Scene, bCanCache, EMeshPass::HitProxy);
+								DrawCommandPacket.AddCommandsForMesh(PrimitiveIndex, PrimitiveSceneInfo, StaticMeshRelevance, StaticMesh, Scene, bCanCache, EMeshPass::HitProxy);
 
 								if (!ViewRelevance.HasTranslucency())
 								{
-									DrawCommandPacket.AddCommandsForMesh(PrimitiveIndex, StaticMeshRelevance, StaticMesh, Scene, bCanCache, EMeshPass::HitProxyOpaqueOnly);
+									DrawCommandPacket.AddCommandsForMesh(PrimitiveIndex, PrimitiveSceneInfo, StaticMeshRelevance, StaticMesh, Scene, bCanCache, EMeshPass::HitProxyOpaqueOnly);
 								}
 							}
 							 
@@ -2147,7 +2155,7 @@ struct FRelevancePacket
 								&& FVelocityRendering::PrimitiveHasVelocity(View.GetFeatureLevel(), PrimitiveSceneInfo)
 								&& FVelocityRendering::PrimitiveHasVelocityForView(View, Bounds.BoxSphereBounds, PrimitiveSceneInfo))
 							{
-								DrawCommandPacket.AddCommandsForMesh(PrimitiveIndex, StaticMeshRelevance, StaticMesh, Scene, bCanCache, EMeshPass::Velocity);
+								DrawCommandPacket.AddCommandsForMesh(PrimitiveIndex, PrimitiveSceneInfo, StaticMeshRelevance, StaticMesh, Scene, bCanCache, EMeshPass::Velocity);
 								MarkMask |= EMarkMaskBits::StaticMeshVelocityMapMask;
 							}
 
@@ -2166,36 +2174,36 @@ struct FRelevancePacket
 						{
 							if (ViewRelevance.bNormalTranslucencyRelevance)
 							{
-								DrawCommandPacket.AddCommandsForMesh(PrimitiveIndex, StaticMeshRelevance, StaticMesh, Scene, bCanCache, EMeshPass::TranslucencyStandard);
+								DrawCommandPacket.AddCommandsForMesh(PrimitiveIndex, PrimitiveSceneInfo, StaticMeshRelevance, StaticMesh, Scene, bCanCache, EMeshPass::TranslucencyStandard);
 							}
 
 							if (ViewRelevance.bSeparateTranslucencyRelevance)
 							{
-								DrawCommandPacket.AddCommandsForMesh(PrimitiveIndex, StaticMeshRelevance, StaticMesh, Scene, bCanCache, EMeshPass::TranslucencyAfterDOF);
+								DrawCommandPacket.AddCommandsForMesh(PrimitiveIndex, PrimitiveSceneInfo, StaticMeshRelevance, StaticMesh, Scene, bCanCache, EMeshPass::TranslucencyAfterDOF);
 							}
 						}
 						else
 						{
 							// Otherwise, everything is rendered in a single bucket. This is not related to whether DOF is currently enabled or not.
 							// When using all translucency, Standard and AfterDOF are sorted together instead of being rendered like 2 buckets.
-							DrawCommandPacket.AddCommandsForMesh(PrimitiveIndex, StaticMeshRelevance, StaticMesh, Scene, bCanCache, EMeshPass::TranslucencyAll);
+							DrawCommandPacket.AddCommandsForMesh(PrimitiveIndex, PrimitiveSceneInfo, StaticMeshRelevance, StaticMesh, Scene, bCanCache, EMeshPass::TranslucencyAll);
 						}
 
 						if (ViewRelevance.bDistortionRelevance)
 						{
-							DrawCommandPacket.AddCommandsForMesh(PrimitiveIndex, StaticMeshRelevance, StaticMesh, Scene, bCanCache, EMeshPass::Distortion);
+							DrawCommandPacket.AddCommandsForMesh(PrimitiveIndex, PrimitiveSceneInfo, StaticMeshRelevance, StaticMesh, Scene, bCanCache, EMeshPass::Distortion);
 						}
 
 						if (ShadingPath == EShadingPath::Mobile && View.bIsSceneCapture)
 						{
-							DrawCommandPacket.AddCommandsForMesh(PrimitiveIndex, StaticMeshRelevance, StaticMesh, Scene, bCanCache, EMeshPass::MobileInverseOpacity);
+							DrawCommandPacket.AddCommandsForMesh(PrimitiveIndex, PrimitiveSceneInfo, StaticMeshRelevance, StaticMesh, Scene, bCanCache, EMeshPass::MobileInverseOpacity);
 						}
 					}
 
 #if WITH_EDITOR
 					if(ViewRelevance.bDrawRelevance && ViewRelevance.bEditorStaticSelectionRelevance)
 					{
-						DrawCommandPacket.AddCommandsForMesh(PrimitiveIndex, StaticMeshRelevance, StaticMesh, Scene, bCanCache, EMeshPass::EditorSelection);
+						DrawCommandPacket.AddCommandsForMesh(PrimitiveIndex, PrimitiveSceneInfo, StaticMeshRelevance, StaticMesh, Scene, bCanCache, EMeshPass::EditorSelection);
 
 						MarkMask |= EMarkMaskBits::StaticMeshEditorSelectedMask;
 					}
diff --git a/Engine/Source/Runtime/Renderer/Private/ShadowRendering.h b/Engine/Source/Runtime/Renderer/Private/ShadowRendering.h
--- a/Engine/Source/Runtime/Renderer/Private/ShadowRendering.h
+++ b/Engine/Source/Runtime/Renderer/Private/ShadowRendering.h
@@ -763,6 +763,7 @@ private:
 
 	void AddCachedMeshDrawCommandsForPass(
 		int32 PrimitiveIndex,
+		const FPrimitiveSceneInfo* InPrimitiveSceneInfo,
 		const FStaticMeshRelevance& RESTRICT StaticMeshRelevance,
 		const FStaticMesh& StaticMesh,
 		const FScene* Scene,
diff --git a/Engine/Source/Runtime/Renderer/Private/ShadowSetup.cpp b/Engine/Source/Runtime/Renderer/Private/ShadowSetup.cpp
--- a/Engine/Source/Runtime/Renderer/Private/ShadowSetup.cpp
+++ b/Engine/Source/Runtime/Renderer/Private/ShadowSetup.cpp
@@ -797,26 +797,28 @@ void FProjectedShadowInfo::SetupWholeSceneProjection(
 }
 
 void FProjectedShadowInfo::AddCachedMeshDrawCommandsForPass(
-	int32 PrimitiveIndex, 
+	int32 PrimitiveIndex,
+	const FPrimitiveSceneInfo* InPrimitiveSceneInfo,
 	const FStaticMeshRelevance& RESTRICT StaticMeshRelevance,
-	const FStaticMesh& RESTRICT StaticMesh, 
-	const FScene* RESTRICT Scene, 
+	const FStaticMesh& StaticMesh,
+	const FScene* Scene,
 	EMeshPass::Type PassType,
-	FMeshCommandOneFrameArray& RESTRICT VisibleMeshCommands,
-	TArray<const FStaticMesh*, SceneRenderingAllocator>& RESTRICT MeshCommandBuildRequests)
+	FMeshCommandOneFrameArray& VisibleMeshCommands,
+	TArray<const FStaticMesh*, SceneRenderingAllocator>& MeshCommandBuildRequests)
 {
 	if (UseMeshDrawCommandPipeline())
 	{
-		EShadingPath ShadingPath = Scene->GetShadingPath();
+		const EShadingPath ShadingPath = Scene->GetShadingPath();
 		const bool bUseCachedMeshCommand = UseCachedMeshDrawCommands()
 			&& !!(FPassProcessorManager::GetPassFlags(ShadingPath, PassType) & EMeshPassFlags::CachedMeshCommands)
 			&& StaticMeshRelevance.bSupportsCachingMeshDrawCommands;
 
 		if (bUseCachedMeshCommand)
 		{
-			const FCachedMeshDrawCommandInfo& CachedMeshDrawCommand = StaticMesh.CachedMeshDrawCommands[PassType];
-			if (CachedMeshDrawCommand.CommandIndex != -1)
+			const int32 StaticMeshCommandInfoIndex = StaticMeshRelevance.GetStaticMeshCommandInfoIndex(PassType);
+			if (StaticMeshCommandInfoIndex >= 0)
 			{
+				const FCachedMeshDrawCommandInfo& CachedMeshDrawCommand = InPrimitiveSceneInfo->StaticMeshCommandInfos[StaticMeshCommandInfoIndex];
 				const FCachedPassMeshDrawList& SceneDrawList = Scene->CachedDrawLists[PassType];
 
 				FVisibleMeshDrawCommand NewVisibleMeshDrawCommand;
@@ -923,6 +925,7 @@ bool FProjectedShadowInfo::ShouldDrawStaticMeshes(FViewInfo& InCurrentView, bool
 						{
 							AddCachedMeshDrawCommandsForPass(
 								PrimitiveId,
+								InPrimitiveSceneInfo,
 								StaticMeshRelevance,
 								StaticMesh,
 								InPrimitiveSceneInfo->Scene,
diff --git a/Engine/Source/Runtime/Renderer/Public/MeshPassProcessor.h b/Engine/Source/Runtime/Renderer/Public/MeshPassProcessor.h
--- a/Engine/Source/Runtime/Renderer/Public/MeshPassProcessor.h
+++ b/Engine/Source/Runtime/Renderer/Public/MeshPassProcessor.h
@@ -38,7 +38,8 @@ namespace EMeshPass
 		EditorSelection,
 #endif
 
-		Num
+		Num,
+		NumBits = 5,
 	};
 }
 
@@ -567,6 +568,7 @@ public:
 		SortKey(FMeshDrawCommandSortKey::Default),
 		CommandIndex(-1),
 		StateBucketId(-1),
+		MeshPass(EMeshPass::Num),
 		MeshFillMode(ERasterizerFillMode_Num),
 		MeshCullMode(ERasterizerCullMode_Num)
 	{}
@@ -576,6 +578,9 @@ public:
 	int32 CommandIndex;
 	int32 StateBucketId;
 
+	// Needed for easier debugging and faster removal of cached mesh draw commands.
+	EMeshPass::Type MeshPass : EMeshPass::NumBits;
+
 	// Needed for view overrides
 	ERasterizerFillMode MeshFillMode : ERasterizerFillMode_NumBits;
 	ERasterizerCullMode MeshCullMode : ERasterizerCullMode_NumBits;
diff --git a/Engine/Source/Runtime/Renderer/Public/PrimitiveSceneInfo.h b/Engine/Source/Runtime/Renderer/Public/PrimitiveSceneInfo.h
--- a/Engine/Source/Runtime/Renderer/Public/PrimitiveSceneInfo.h
+++ b/Engine/Source/Runtime/Renderer/Public/PrimitiveSceneInfo.h
@@ -211,7 +211,10 @@ public:
 	 */
 	FPrimitiveComponentId LODParentComponentId;
 
-	/** The primitive's static mesh relevances. Must be in sync with StaticMeshes. */
+	/** The primitive's cached mesh draw commands infos for all static meshes. Kept separately from StaticMeshes for cache efficiency inside InitViews. */
+	TArray<class FCachedMeshDrawCommandInfo> StaticMeshCommandInfos;
+
+	/** The primitive's static mesh relevances. Must be in sync with StaticMeshes. Kept separately from StaticMeshes for cache efficiency inside InitViews. */
 	TArray<class FStaticMeshRelevance> StaticMeshRelevances;
 
 	/** The primitive's static meshes. */
@@ -456,6 +459,12 @@ private:
 		FVector VolumetricLightmapLookupPosition,
 		uint32 SceneFrameNumber,
 		class FVolumetricLightmapSceneData* VolumetricLightmapSceneData);
+
+	/** Creates cached mesh draw commands for all meshes. */
+	void CacheMeshDrawCommands(FRHICommandListImmediate& RHICmdList);
+
+	/** Removes cached mesh draw commands for all meshes. */
+	void RemoveCachedMeshDrawCommands();
 };
 
 /** Defines how the primitive is stored in the scene's primitive octree. */
