Skip ray tracing instance if the corresponding mesh batches are filtered by ray tracing mesh processor

diff --git a/Engine/Source/Runtime/Renderer/Private/DeferredShadingRenderer.cpp b/Engine/Source/Runtime/Renderer/Private/DeferredShadingRenderer.cpp
--- a/Engine/Source/Runtime/Renderer/Private/DeferredShadingRenderer.cpp
+++ b/Engine/Source/Runtime/Renderer/Private/DeferredShadingRenderer.cpp
@@ -778,8 +778,10 @@ bool FDeferredShadingSceneRenderer::GatherRayTracingWorldInstances(FRHICommandLi
 								}
 							}
 
-							ensure(DrawCmdIndexCopy != VisibleDrawCommandStartOffset[ViewIndex]);
-							View.RayTracingGeometryInstances.Add(FRayTracingGeometryInstance { RayTracingGeometryInstance, Scene->PrimitiveTransforms[PrimitiveIndex], (uint32)PrimitiveIndex });
+							if (DrawCmdIndexCopy != VisibleDrawCommandStartOffset[ViewIndex])
+							{
+								View.RayTracingGeometryInstances.Add(FRayTracingGeometryInstance { RayTracingGeometryInstance, Scene->PrimitiveTransforms[PrimitiveIndex], (uint32)PrimitiveIndex });
+							}
 						}
 					}
 				}
