Skip baked static lights in ray traced reflections and path tracer

diff --git a/Engine/Source/Runtime/Renderer/Private/PathTracing.cpp b/Engine/Source/Runtime/Renderer/Private/PathTracing.cpp
--- a/Engine/Source/Runtime/Renderer/Private/PathTracing.cpp
+++ b/Engine/Source/Runtime/Renderer/Private/PathTracing.cpp
@@ -135,54 +135,58 @@ public:
 
 		// Lights
 		{
-			// Fill light buffer data
 			FPathTracingLightData LightData;
-			LightData.Count = FMath::Min(Lights.Num() + 1, GLightCountMaximum - 1);
+			LightData.Count = 0;
 
 			// Prepend SkyLight to light buffer
 			// WARNING: Until ray payload encodes Light data buffer, the execution depends on this ordering!
 			uint32 SkyLightIndex = 0;
 			LightData.Type[SkyLightIndex] = 0;
 			LightData.Color[SkyLightIndex] = FVector(SkyLightColorData);
+			LightData.Count++;
 
-			TSparseArray<FLightSceneInfoCompact>::TConstIterator LightItr = Lights.CreateConstIterator();
-			for (uint32 LightIndex = 1; LightIndex < LightData.Count; ++LightIndex, ++LightItr)
+			for (auto Light : Lights)
 			{
-				FLightSceneProxy* LightSceneProxy = LightItr->LightSceneInfo->Proxy;
+				if (LightData.Count >= GLightCountMaximum) break;
+
+				if (Light.LightSceneInfo->Proxy->HasStaticLighting() && Light.LightSceneInfo->IsPrecomputedLightingValid()) continue;
+
 				FLightShaderParameters LightParameters;
-				LightSceneProxy->GetLightShaderParameters(LightParameters);
-				
-				ELightComponentType LightComponentType = (ELightComponentType)LightSceneProxy->GetLightType();
+				Light.LightSceneInfo->Proxy->GetLightShaderParameters(LightParameters);
+
+				ELightComponentType LightComponentType = (ELightComponentType)Light.LightSceneInfo->Proxy->GetLightType();
 				switch (LightComponentType)
 				{
 					// TODO: LightType_Spot
 					case LightType_Directional:
 					{
-						LightData.Type[LightIndex] = 2;
-						LightData.Normal[LightIndex] = LightParameters.Direction;
-						LightData.Color[LightIndex] = LightParameters.Color;
+						LightData.Type[LightData.Count] = 2;
+						LightData.Normal[LightData.Count] = LightParameters.Direction;
+						LightData.Color[LightData.Count] = LightParameters.Color;
 						break;
 					}
 					case LightType_Rect:
 					{
-						LightData.Type[LightIndex] = 3;
-						LightData.Position[LightIndex] = LightParameters.Position;
-						LightData.Normal[LightIndex] = -LightParameters.Direction;
-						LightData.dPdu[LightIndex] = FVector::CrossProduct(LightParameters.Tangent, LightParameters.Direction);
-						LightData.dPdv[LightIndex] = LightParameters.Tangent;
-						LightData.Color[LightIndex] = LightParameters.Color * (2.0f * LightParameters.SourceRadius * LightParameters.SourceLength);
-						LightData.Dimensions[LightIndex] = FVector(2.0f * LightParameters.SourceRadius, 2.0f * LightParameters.SourceLength, 0.0f);
+						LightData.Type[LightData.Count] = 3;
+						LightData.Position[LightData.Count] = LightParameters.Position;
+						LightData.Normal[LightData.Count] = -LightParameters.Direction;
+						LightData.dPdu[LightData.Count] = FVector::CrossProduct(LightParameters.Tangent, LightParameters.Direction);
+						LightData.dPdv[LightData.Count] = LightParameters.Tangent;
+						LightData.Color[LightData.Count] = LightParameters.Color * (2.0f * LightParameters.SourceRadius * LightParameters.SourceLength);
+						LightData.Dimensions[LightData.Count] = FVector(2.0f * LightParameters.SourceRadius, 2.0f * LightParameters.SourceLength, 0.0f);
 						break;
 					}
 					case LightType_Point:
 					default:
 					{
-						LightData.Type[LightIndex] = 1;
-						LightData.Position[LightIndex] = LightParameters.Position;
-						LightData.Color[LightIndex] = LightParameters.Color;
+						LightData.Type[LightData.Count] = 1;
+						LightData.Position[LightData.Count] = LightParameters.Position;
+						LightData.Color[LightData.Count] = LightParameters.Color;
 						break;
 					}
 				};
+
+				LightData.Count++;
 			}
 
 			FUniformBufferRHIRef SceneLightsUniformBuffer = RHICreateUniformBuffer(&LightData, FPathTracingLightData::StaticStructMetadata.GetLayout(), EUniformBufferUsage::UniformBuffer_SingleDraw);
diff --git a/Engine/Source/Runtime/Renderer/Private/RayTracing/RayTracingReflections.cpp b/Engine/Source/Runtime/Renderer/Private/RayTracing/RayTracingReflections.cpp
--- a/Engine/Source/Runtime/Renderer/Private/RayTracing/RayTracingReflections.cpp
+++ b/Engine/Source/Runtime/Renderer/Private/RayTracing/RayTracingReflections.cpp
@@ -80,36 +80,34 @@ void SetupReflectionsLightData(
 	const FViewInfo& View,
 	FReflectionsLightData* LightData)
 {
-	// Fill light buffer data
-	LightData->Count = FMath::Min(Lights.Num(), GReflectionLightCountMaximum);
+	LightData->Count = 0;
 
-	TSparseArray<FLightSceneInfoCompact>::TConstIterator LightItr = Lights.CreateConstIterator();
-	for (uint32 LightIndex = 0; LightIndex < LightData->Count; ++LightIndex, ++LightItr)
+	for (auto Light : Lights)
 	{
-		FLightSceneProxy* LightSceneProxy = LightItr->LightSceneInfo->Proxy;
+		if (Light.LightSceneInfo->Proxy->HasStaticLighting() && Light.LightSceneInfo->IsPrecomputedLightingValid()) continue;
+
 		FLightShaderParameters LightParameters;
-		LightSceneProxy->GetLightShaderParameters(LightParameters);
-
-		if (LightSceneProxy->IsInverseSquared())
-		{
-			LightParameters.FalloffExponent = 0;
-		}
-
-		LightData->Type[LightIndex] = LightItr->LightType;
-		LightData->LightPosition[LightIndex] = LightParameters.Position;
-		LightData->LightInvRadius[LightIndex] = LightParameters.InvRadius;
-		LightData->LightColor[LightIndex] = LightParameters.Color;
-		LightData->LightFalloffExponent[LightIndex] = LightParameters.FalloffExponent;
-		LightData->Direction[LightIndex] = LightParameters.Direction;
-		LightData->Tangent[LightIndex] = LightParameters.Tangent;
-		LightData->SpotAngles[LightIndex] = LightParameters.SpotAngles;
-		LightData->SpecularScale[LightIndex] = LightParameters.SpecularScale;
-		LightData->SourceRadius[LightIndex] = LightParameters.SourceRadius;
-		LightData->SourceLength[LightIndex] = LightParameters.SourceLength;
-		LightData->SoftSourceRadius[LightIndex] = LightParameters.SoftSourceRadius;
-
-		const FVector2D FadeParams = LightSceneProxy->GetDirectionalLightDistanceFadeParameters(View.GetFeatureLevel(), LightItr->LightSceneInfo->IsPrecomputedLightingValid(), View.MaxShadowCascades);
-		LightData->DistanceFadeMAD[LightIndex] = FVector2D(FadeParams.Y, -FadeParams.X * FadeParams.Y);
+		Light.LightSceneInfo->Proxy->GetLightShaderParameters(LightParameters);
+
+		LightData->Type[LightData->Count] = Light.LightType;
+		LightData->LightPosition[LightData->Count] = LightParameters.Position;
+		LightData->LightInvRadius[LightData->Count] = LightParameters.InvRadius;
+		LightData->LightColor[LightData->Count] = LightParameters.Color;
+		LightData->LightFalloffExponent[LightData->Count] = LightParameters.FalloffExponent;
+		LightData->Direction[LightData->Count] = LightParameters.Direction;
+		LightData->Tangent[LightData->Count] = LightParameters.Tangent;
+		LightData->SpotAngles[LightData->Count] = LightParameters.SpotAngles;
+		LightData->SpecularScale[LightData->Count] = LightParameters.SpecularScale;
+		LightData->SourceRadius[LightData->Count] = LightParameters.SourceRadius;
+		LightData->SourceLength[LightData->Count] = LightParameters.SourceLength;
+		LightData->SoftSourceRadius[LightData->Count] = LightParameters.SoftSourceRadius;
+
+		const FVector2D FadeParams = Light.LightSceneInfo->Proxy->GetDirectionalLightDistanceFadeParameters(View.GetFeatureLevel(), Light.LightSceneInfo->IsPrecomputedLightingValid(), View.MaxShadowCascades);
+		LightData->DistanceFadeMAD[LightData->Count] = FVector2D(FadeParams.Y, -FadeParams.X * FadeParams.Y);
+
+		LightData->Count++;
+
+		if (LightData->Count >= GReflectionLightCountMaximum) break;
 	}
 
 	LightData->DummyRectLightTexture = GWhiteTexture->TextureRHI; //#dxr_todo: replace with valid textures per rect light
