Moved static mesh raytraycing code behind IsRayTracingTierSupported.

diff --git a/Engine/Source/Runtime/Renderer/Private/PrimitiveSceneInfo.cpp b/Engine/Source/Runtime/Renderer/Private/PrimitiveSceneInfo.cpp
--- a/Engine/Source/Runtime/Renderer/Private/PrimitiveSceneInfo.cpp
+++ b/Engine/Source/Runtime/Renderer/Private/PrimitiveSceneInfo.cpp
@@ -326,46 +326,49 @@ void FPrimitiveSceneInfo::AddStaticMeshes(FRHICommandListImmediate& RHICmdList,
 	CacheMeshDrawCommands(RHICmdList);
 
 #if RHI_RAYTRACING
-	int MaxLOD = -1;
-
-	for (int32 MeshIndex = 0; MeshIndex < StaticMeshes.Num(); MeshIndex++)
-	{
-		FStaticMesh& Mesh = StaticMeshes[MeshIndex];
-		MaxLOD = MaxLOD < Mesh.LODIndex ? Mesh.LODIndex : MaxLOD;
-	}
-
-	if (StaticMeshes.Num() > 0)
+	if (IsRayTracingTierSupported(1))
 	{
-		Proxy->ScreenSizes.AddDefaulted(MaxLOD + 1);
-		Proxy->RayTracingLodIndexToMeshDrawCommandIndicies.AddDefaulted(MaxLOD + 1);
+		int MaxLOD = -1;
 
 		for (int32 MeshIndex = 0; MeshIndex < StaticMeshes.Num(); MeshIndex++)
 		{
-			FStaticMeshRelevance& MeshRelevance = StaticMeshRelevances[MeshIndex];
 			FStaticMesh& Mesh = StaticMeshes[MeshIndex];
+			MaxLOD = MaxLOD < Mesh.LODIndex ? Mesh.LODIndex : MaxLOD;
+		}
 
-			if (Proxy->ScreenSizes[MeshRelevance.LODIndex] != 0.0f)
-			{
-				check(Proxy->ScreenSizes[MeshRelevance.LODIndex] == MeshRelevance.ScreenSize);
-			}
-			else
-			{
-				check(MeshRelevance.ScreenSize != 0.0f);
-				Proxy->ScreenSizes[MeshRelevance.LODIndex] = MeshRelevance.ScreenSize;
-			}
+		if (StaticMeshes.Num() > 0)
+		{
+			Proxy->ScreenSizes.AddDefaulted(MaxLOD + 1);
+			Proxy->RayTracingLodIndexToMeshDrawCommandIndicies.AddDefaulted(MaxLOD + 1);
 
-			if (Mesh.Elements.Num() > 0)
+			for (int32 MeshIndex = 0; MeshIndex < StaticMeshes.Num(); MeshIndex++)
 			{
-				const int32 RayTracingStaticMeshCommandInfoIndex = MeshRelevance.GetStaticMeshCommandInfoIndex(EMeshPass::RayTracing);
+				FStaticMeshRelevance& MeshRelevance = StaticMeshRelevances[MeshIndex];
+				FStaticMesh& Mesh = StaticMeshes[MeshIndex];
 
-				if (SupportsCachingMeshDrawCommands(Mesh.VertexFactory, Proxy) && RayTracingStaticMeshCommandInfoIndex != -1)
+				if (Proxy->ScreenSizes[MeshRelevance.LODIndex] != 0.0f)
 				{
-					const int32 CommandIndex = StaticMeshCommandInfos[RayTracingStaticMeshCommandInfoIndex].CommandIndex;
-					Proxy->RayTracingLodIndexToMeshDrawCommandIndicies[Mesh.LODIndex].Add({ MeshIndex, CommandIndex });
+					check(Proxy->ScreenSizes[MeshRelevance.LODIndex] == MeshRelevance.ScreenSize);
 				}
 				else
 				{
-					Proxy->RayTracingLodIndexToMeshDrawCommandIndicies[Mesh.LODIndex].Add({ MeshIndex, -1 });
+					check(MeshRelevance.ScreenSize != 0.0f);
+					Proxy->ScreenSizes[MeshRelevance.LODIndex] = MeshRelevance.ScreenSize;
+				}
+
+				if (Mesh.Elements.Num() > 0)
+				{
+					const int32 RayTracingStaticMeshCommandInfoIndex = MeshRelevance.GetStaticMeshCommandInfoIndex(EMeshPass::RayTracing);
+
+					if (SupportsCachingMeshDrawCommands(Mesh.VertexFactory, Proxy) && RayTracingStaticMeshCommandInfoIndex != -1)
+					{
+						const int32 CommandIndex = StaticMeshCommandInfos[RayTracingStaticMeshCommandInfoIndex].CommandIndex;
+						Proxy->RayTracingLodIndexToMeshDrawCommandIndicies[Mesh.LODIndex].Add({ MeshIndex, CommandIndex });
+					}
+					else
+					{
+						Proxy->RayTracingLodIndexToMeshDrawCommandIndicies[Mesh.LODIndex].Add({ MeshIndex, -1 });
+					}
 				}
 			}
 		}
diff --git a/Engine/Source/Runtime/Renderer/Public/MeshPassProcessor.h b/Engine/Source/Runtime/Renderer/Public/MeshPassProcessor.h
--- a/Engine/Source/Runtime/Renderer/Public/MeshPassProcessor.h
+++ b/Engine/Source/Runtime/Renderer/Public/MeshPassProcessor.h
@@ -495,7 +495,9 @@ public:
 	// A value of -1 means the draw is not in any state bucket and should be sorted by other factors instead.
 	int32 StateBucketId;
 
+#if RHI_RAYTRACING
 	uint32 RayTracedInstanceIndex;
+#endif
 
 	// Needed for view overrides
 	ERasterizerFillMode MeshFillMode : ERasterizerFillMode_NumBits + 1;
