Removing duplicated code from SortPassMeshDrawCommands. Now SortPassMeshDrawCommands also uses new code from FVisibleMeshDrawCommandProcessPacket, which contains a fix for the out of bounds array access.

diff --git a/Engine/Source/Runtime/Renderer/Private/SceneRendering.cpp b/Engine/Source/Runtime/Renderer/Private/SceneRendering.cpp
--- a/Engine/Source/Runtime/Renderer/Private/SceneRendering.cpp
+++ b/Engine/Source/Runtime/Renderer/Private/SceneRendering.cpp
@@ -2769,38 +2769,122 @@ struct FCompareFMeshDrawCommands
 	}
 };
 
-void SortPassMeshDrawCommands(
-	ERHIFeatureLevel::Type FeatureLevel,
-	FMeshCommandOneFrameArray& VisibleMeshDrawCommands, 
-	FDynamicMeshDrawCommandStorage& MeshDrawCommandStorage,
-	FGlobalDynamicVertexBuffer& DynamicVertexBuffer,
-	FGlobalDynamicVertexBuffer::FAllocation& PrimitiveIdBuffer,
-	bool bDumpStats)
+uint32 BitInvertIfNegativeFloat(uint32 f)
+{
+	unsigned mask = -int32(f >> 31) | 0x80000000;
+	return f ^ mask;
+}
+
+/**
+* Update mesh sort keys with view dependent data.
+*/
+void UpdateTranslucentMeshSortKeys(const FScene* Scene, const FViewInfo& View, ETranslucencyPass::Type TranslucencyPass, FMeshCommandOneFrameArray& VisibleMeshCommands)
 {
+	for (int32 CommandIndex = 0; CommandIndex < VisibleMeshCommands.Num(); ++CommandIndex)
 	{
-		QUICK_SCOPE_CYCLE_COUNTER(STAT_SortVisibleMeshDrawCommands);
+		FVisibleMeshDrawCommand& VisibleCommand = VisibleMeshCommands[CommandIndex];
+
+		FVector BoundsOrigin = FVector::ZeroVector;
+		if (VisibleCommand.DrawPrimitiveId < Scene->PrimitiveBounds.Num())
+		{
+			BoundsOrigin = Scene->PrimitiveBounds[VisibleCommand.DrawPrimitiveId].BoxSphereBounds.Origin;
+		}
+
+		float Distance = 0.0f;
+		if (View.TranslucentSortPolicy == ETranslucentSortPolicy::SortByDistance)
+		{
+			//sort based on distance to the view position, view rotation is not a factor
+			Distance = (BoundsOrigin - View.ViewMatrices.GetViewOrigin()).Size();
+		}
+		else if (View.TranslucentSortPolicy == ETranslucentSortPolicy::SortAlongAxis)
+		{
+			// Sort based on enforced orthogonal distance
+			const FVector CameraToObject = BoundsOrigin - View.ViewMatrices.GetViewOrigin();
+			Distance = FVector::DotProduct(CameraToObject, View.TranslucentSortAxis);
+		}
+		else
+		{
+			// Sort based on projected Z distance
+			check(View.TranslucentSortPolicy == ETranslucentSortPolicy::SortByProjectedZ);
+			Distance = View.ViewMatrices.GetViewMatrix().TransformPosition(BoundsOrigin).Z;
+		}
+
+		// Patch distance inside translucent mesh sort key.
+		UTranslucentMeshSortKey TranslucentMeshSortKey;
+		TranslucentMeshSortKey.PackedData = VisibleCommand.SortKey.PackedData;
+		TranslucentMeshSortKey.Fields.Distance = (uint32)~BitInvertIfNegativeFloat(*((uint32*)&Distance));
+		VisibleCommand.SortKey.PackedData = TranslucentMeshSortKey.PackedData;		
+	}
+}
+
+/**
+ * Packet for parallel mesh draw command sorting and setting up primitive id buffer for instancing.
+ */
+class FVisibleMeshDrawCommandProcessPacket
+{
+public:
+	FVisibleMeshDrawCommandProcessPacket(const FScene* InScene,
+		const FViewInfo* InView,
+		bool InbUseGPUScene, 
+		EMeshPass::Type InPassType,
+		ETranslucencyPass::Type InTranslucencyPass, 
+		FMeshCommandOneFrameArray& InVisibleMeshDrawCommands, 
+		FDynamicMeshDrawCommandStorage& InMeshDrawCommandStorage,
+		FGlobalDynamicVertexBuffer& InDynamicVertexBuffer,
+		FGlobalDynamicVertexBuffer::FAllocation& InPrimitiveIdBuffer)
+		: Scene(InScene)
+		, View(InView)
+		, bUseGPUScene(InbUseGPUScene)
+		, PassType(InPassType)
+		, TranslucencyPass(InTranslucencyPass)
+		, VisibleMeshDrawCommands(InVisibleMeshDrawCommands)
+		, MeshDrawCommandStorage(InMeshDrawCommandStorage)
+		, PrimitiveIdBuffer(InPrimitiveIdBuffer)
+		, VisibleMeshDrawCommandsNum(0)
+		, NewPassVisibleMeshDrawCommandsNum(0)
+		, MaxInstances(1)
+	{
+		// Preallocate memory and resources on rendering thread.
+
+		const int32 NumDrawCommands = VisibleMeshDrawCommands.Num();
+
+		if (bUseGPUScene && NumDrawCommands > 0)
+		{
+			PrimitiveIdBuffer = InDynamicVertexBuffer.Allocate(NumDrawCommands * sizeof(int32));
+			int32* RESTRICT PrimitiveIds = (int32*)PrimitiveIdBuffer.Buffer;
+		}
 
+		if (IsDynamicInstancingEnabled())
+		{
+			NewPassVisibleMeshDrawCommands.Empty(NumDrawCommands);
+		}
+	}
+
+	void Sort()
+	{
+		QUICK_SCOPE_CYCLE_COUNTER(STAT_SortVisibleMeshDrawCommands);
 		VisibleMeshDrawCommands.Sort(FCompareFMeshDrawCommands());
 	}
 
-	if (UseGPUScene(GMaxRHIShaderPlatform, FeatureLevel))
+	void BuildPrimitiveIdBuffer()
 	{
+		if (!bUseGPUScene)
+		{
+			return;
+		}
+		
 		const FVisibleMeshDrawCommand* RESTRICT PassVisibleMeshDrawCommands = VisibleMeshDrawCommands.GetData();
 		const int32 NumDrawCommands = VisibleMeshDrawCommands.Num();
 
-		PrimitiveIdBuffer = DynamicVertexBuffer.Allocate(NumDrawCommands * sizeof(int32));
 		int32* RESTRICT PrimitiveIds = (int32*)PrimitiveIdBuffer.Buffer;
 
 		if (IsDynamicInstancingEnabled())
 		{
 			QUICK_SCOPE_CYCLE_COUNTER(STAT_DynamicInstancingOfVisibleMeshDrawCommands);
 
-			FMeshCommandOneFrameArray NewPassVisibleMeshDrawCommands;
-			NewPassVisibleMeshDrawCommands.Empty(NumDrawCommands);
-
 			int32 CurrentStateBucketId = -1;
 			uint32* RESTRICT CurrentDynamicallyInstancedMeshCommandNumInstances = nullptr;
-			int32 MaxInstances = 1;
+			MaxInstances = 1;
 
 			for (int32 DrawCommandIndex = 0; DrawCommandIndex < NumDrawCommands; DrawCommandIndex++)
 			{
@@ -2826,10 +2910,10 @@ void SortPassMeshDrawCommands(
 					// First time state bucket setup
 					CurrentStateBucketId = VisibleMeshDrawCommand.StateBucketId;
 
-					if (VisibleMeshDrawCommand.MeshDrawCommand->PrimitiveIdStreamIndex >= 0 
+					if (VisibleMeshDrawCommand.MeshDrawCommand->PrimitiveIdStreamIndex >= 0
 						&& VisibleMeshDrawCommand.MeshDrawCommand->NumInstances == 1
 						// Don't create a new FMeshDrawCommand for the last command and make it safe for us to look at the next command
-						&& DrawCommandIndex + 1 <= NumDrawCommands 
+						&& DrawCommandIndex + 1 < NumDrawCommands
 						// Only create a new FMeshDrawCommand if more than one draw in the state bucket
 						&& CurrentStateBucketId == PassVisibleMeshDrawCommands[DrawCommandIndex + 1].StateBucketId)
 					{
@@ -2838,7 +2922,7 @@ void SortPassMeshDrawCommands(
 						FVisibleMeshDrawCommand NewVisibleMeshDrawCommand;
 
 						NewVisibleMeshDrawCommand.Setup(
-							&NewCommand, 
+							&NewCommand,
 							VisibleMeshDrawCommand.DrawPrimitiveId,
 							VisibleMeshDrawCommand.StateBucketId,
 							VisibleMeshDrawCommand.MeshFillMode,
@@ -2859,16 +2943,13 @@ void SortPassMeshDrawCommands(
 					}
 				}
 
-				//@todo MeshCommandPipeline - refactor into memcpy
+				//@todo - refactor into memcpy
 				PrimitiveIds[DrawCommandIndex] = VisibleMeshDrawCommand.DrawPrimitiveId;
 			}
 
-			if (bDumpStats && VisibleMeshDrawCommands.Num() > 0)
-			{
-				UE_LOG(LogRenderer, Log, TEXT("   %i Mesh Draw Commands in %i instancing state buckets"), VisibleMeshDrawCommands.Num(), NewPassVisibleMeshDrawCommands.Num());
-				UE_LOG(LogRenderer, Log, TEXT("   Largest %i"), MaxInstances);
-				UE_LOG(LogRenderer, Log, TEXT("   %.1f Dynamic Instancing draw call reduction factor"), VisibleMeshDrawCommands.Num() / (float)NewPassVisibleMeshDrawCommands.Num());
-			}
+			// Setup instancing stats for logging.
+			VisibleMeshDrawCommandsNum = VisibleMeshDrawCommands.Num();
+			NewPassVisibleMeshDrawCommandsNum = NewPassVisibleMeshDrawCommands.Num();
 
 			// Replace VisibleMeshDrawCommands
 			FMemory::Memswap(&VisibleMeshDrawCommands, &NewPassVisibleMeshDrawCommands, sizeof(NewPassVisibleMeshDrawCommands));
@@ -2884,201 +2965,29 @@ void SortPassMeshDrawCommands(
 			}
 		}
 	}
-}
-
-uint32 BitInvertIfNegativeFloat(uint32 f)
-{
-	unsigned mask = -int32(f >> 31) | 0x80000000;
-	return f ^ mask;
-}
-
-/**
-* Update mesh sort keys with view dependent data.
-*/
-void UpdateTranslucentMeshSortKeys(FScene* Scene, FViewInfo& View, ETranslucencyPass::Type TranslucencyPass, FMeshCommandOneFrameArray& VisibleMeshCommands)
-{
-	for (int32 CommandIndex = 0; CommandIndex < VisibleMeshCommands.Num(); ++CommandIndex)
-	{
-		FVisibleMeshDrawCommand& VisibleCommand = VisibleMeshCommands[CommandIndex];
-
-		FVector BoundsOrigin = FVector::ZeroVector;
-		if (VisibleCommand.DrawPrimitiveId < Scene->PrimitiveBounds.Num())
-		{
-			BoundsOrigin = Scene->PrimitiveBounds[VisibleCommand.DrawPrimitiveId].BoxSphereBounds.Origin;
-		}
-
-		float Distance = 0.0f;
-		if (View.TranslucentSortPolicy == ETranslucentSortPolicy::SortByDistance)
-		{
-			//sort based on distance to the view position, view rotation is not a factor
-			Distance = (BoundsOrigin - View.ViewMatrices.GetViewOrigin()).Size();
-		}
-		else if (View.TranslucentSortPolicy == ETranslucentSortPolicy::SortAlongAxis)
-		{
-			// Sort based on enforced orthogonal distance
-			const FVector CameraToObject = BoundsOrigin - View.ViewMatrices.GetViewOrigin();
-			Distance = FVector::DotProduct(CameraToObject, View.TranslucentSortAxis);
-		}
-		else
-		{
-			// Sort based on projected Z distance
-			check(View.TranslucentSortPolicy == ETranslucentSortPolicy::SortByProjectedZ);
-			Distance = View.ViewMatrices.GetViewMatrix().TransformPosition(BoundsOrigin).Z;
-		}
-
-		// Patch distance inside translucent mesh sort key.
-		UTranslucentMeshSortKey TranslucentMeshSortKey;
-		TranslucentMeshSortKey.PackedData = VisibleCommand.SortKey.PackedData;
-		TranslucentMeshSortKey.Fields.Distance = (uint32)~BitInvertIfNegativeFloat(*((uint32*)&Distance));
-		VisibleCommand.SortKey.PackedData = TranslucentMeshSortKey.PackedData;		
-	}
-}
-
-/**
- * Packet for parallel mesh draw command sorting and setting up primitive id buffer for instancing.
- */
-class FVisibleMeshDrawCommandProcessPacket
-{
-public:
-	FVisibleMeshDrawCommandProcessPacket(FScene* InScene,
-		FViewInfo& InView, 
-		bool InbUseGPUScene, 
-		EMeshPass::Type InPassType,
-		ETranslucencyPass::Type InTranslucencyPass, 
-		FMeshCommandOneFrameArray& InVisibleMeshDrawCommands, 
-		FDynamicMeshDrawCommandStorage& InMeshDrawCommandStorage,
-		FGlobalDynamicVertexBuffer::FAllocation& InPrimitiveIdBuffer)
-		: Scene(InScene)
-		, View(InView)
-		, bUseGPUScene(InbUseGPUScene)
-		, PassType(InPassType)
-		, TranslucencyPass(InTranslucencyPass)
-		, VisibleMeshDrawCommands(InVisibleMeshDrawCommands)
-		, MeshDrawCommandStorage(InMeshDrawCommandStorage)
-		, PrimitiveIdBuffer(InPrimitiveIdBuffer)
-		, VisibleMeshDrawCommandsNum(0)
-		, NewPassVisibleMeshDrawCommandsNum(0)
-		, MaxInstances(1)
-	{
-		// Preallocate NewPassVisibleMeshDrawCommands on rendering thread.
-		if (IsDynamicInstancingEnabled())
-		{
-			NewPassVisibleMeshDrawCommands.Empty(InVisibleMeshDrawCommands.Num());
-		}
-	}
 
 	void AnyThreadTask()
 	{
+		check(Scene && View);
+
 		if (TranslucencyPass != ETranslucencyPass::TPT_MAX)
 		{
 			QUICK_SCOPE_CYCLE_COUNTER(STAT_UpdateTranslucentMeshSortKeys);
-			UpdateTranslucentMeshSortKeys(Scene, View, TranslucencyPass, VisibleMeshDrawCommands);
+			UpdateTranslucentMeshSortKeys(Scene, *View, TranslucencyPass, VisibleMeshDrawCommands);
 		}
 
-		{
-			QUICK_SCOPE_CYCLE_COUNTER(STAT_SortVisibleMeshDrawCommands);
-			VisibleMeshDrawCommands.Sort(FCompareFMeshDrawCommands());
-		}
-
-		if (bUseGPUScene)
-		{
-			const FVisibleMeshDrawCommand* RESTRICT PassVisibleMeshDrawCommands = VisibleMeshDrawCommands.GetData();
-			const int32 NumDrawCommands = VisibleMeshDrawCommands.Num();
-
-			int32* RESTRICT PrimitiveIds = (int32*)PrimitiveIdBuffer.Buffer;
-
-			if (IsDynamicInstancingEnabled())
-			{
-				QUICK_SCOPE_CYCLE_COUNTER(STAT_DynamicInstancingOfVisibleMeshDrawCommands);
-
-				int32 CurrentStateBucketId = -1;
-				uint32* RESTRICT CurrentDynamicallyInstancedMeshCommandNumInstances = nullptr;
-				MaxInstances = 1;
-
-				for (int32 DrawCommandIndex = 0; DrawCommandIndex < NumDrawCommands; DrawCommandIndex++)
-				{
-					const FVisibleMeshDrawCommand& RESTRICT VisibleMeshDrawCommand = PassVisibleMeshDrawCommands[DrawCommandIndex];
-
-					if (VisibleMeshDrawCommand.StateBucketId == CurrentStateBucketId && VisibleMeshDrawCommand.StateBucketId != -1)
-					{
-						if (CurrentDynamicallyInstancedMeshCommandNumInstances)
-						{
-							const int32 CurrentNumInstances = *CurrentDynamicallyInstancedMeshCommandNumInstances;
-							*CurrentDynamicallyInstancedMeshCommandNumInstances = CurrentNumInstances + 1;
-							MaxInstances = FMath::Max(MaxInstances, CurrentNumInstances + 1);
-						}
-						else
-						{
-							FVisibleMeshDrawCommand NewVisibleMeshDrawCommand = VisibleMeshDrawCommand;
-							NewVisibleMeshDrawCommand.PrimitiveIdBufferOffset = DrawCommandIndex;
-							NewPassVisibleMeshDrawCommands.Emplace(MoveTemp(NewVisibleMeshDrawCommand));
-						}
-					}
-					else
-					{
-						// First time state bucket setup
-						CurrentStateBucketId = VisibleMeshDrawCommand.StateBucketId;
-
-						if (VisibleMeshDrawCommand.MeshDrawCommand->PrimitiveIdStreamIndex >= 0
-							&& VisibleMeshDrawCommand.MeshDrawCommand->NumInstances == 1
-							// Don't create a new FMeshDrawCommand for the last command and make it safe for us to look at the next command
-							&& DrawCommandIndex + 1 < NumDrawCommands
-							// Only create a new FMeshDrawCommand if more than one draw in the state bucket
-							&& CurrentStateBucketId == PassVisibleMeshDrawCommands[DrawCommandIndex + 1].StateBucketId)
-						{
-							const int32 Index = MeshDrawCommandStorage.MeshDrawCommands.AddElement(*VisibleMeshDrawCommand.MeshDrawCommand);
-							FMeshDrawCommand& NewCommand = MeshDrawCommandStorage.MeshDrawCommands[Index];
-							FVisibleMeshDrawCommand NewVisibleMeshDrawCommand;
-
-							NewVisibleMeshDrawCommand.Setup(
-								&NewCommand,
-								VisibleMeshDrawCommand.DrawPrimitiveId,
-								VisibleMeshDrawCommand.StateBucketId,
-								VisibleMeshDrawCommand.MeshFillMode,
-								VisibleMeshDrawCommand.MeshCullMode,
-								VisibleMeshDrawCommand.SortKey);
-
-							NewVisibleMeshDrawCommand.PrimitiveIdBufferOffset = DrawCommandIndex;
-							NewPassVisibleMeshDrawCommands.Emplace(MoveTemp(NewVisibleMeshDrawCommand));
-
-							CurrentDynamicallyInstancedMeshCommandNumInstances = &NewCommand.NumInstances;
-						}
-						else
-						{
-							CurrentDynamicallyInstancedMeshCommandNumInstances = nullptr;
-							FVisibleMeshDrawCommand NewVisibleMeshDrawCommand = VisibleMeshDrawCommand;
-							NewVisibleMeshDrawCommand.PrimitiveIdBufferOffset = DrawCommandIndex;
-							NewPassVisibleMeshDrawCommands.Emplace(MoveTemp(NewVisibleMeshDrawCommand));
-						}
-					}
-
-					//@todo - refactor into memcpy
-					PrimitiveIds[DrawCommandIndex] = VisibleMeshDrawCommand.DrawPrimitiveId;
-				}
-
-				// Setup instancing stats for logging.
-				VisibleMeshDrawCommandsNum = VisibleMeshDrawCommands.Num();
-				NewPassVisibleMeshDrawCommandsNum = NewPassVisibleMeshDrawCommands.Num();
-
-				// Replace VisibleMeshDrawCommands
-				FMemory::Memswap(&VisibleMeshDrawCommands, &NewPassVisibleMeshDrawCommands, sizeof(NewPassVisibleMeshDrawCommands));
-			}
-			else
-			{
-				QUICK_SCOPE_CYCLE_COUNTER(STAT_BuildVisibleMeshDrawCommandPrimitiveIdBuffers);
+		Sort();
 
-				for (int32 DrawCommandIndex = 0; DrawCommandIndex < NumDrawCommands; DrawCommandIndex++)
-				{
-					const FVisibleMeshDrawCommand& VisibleMeshDrawCommand = VisibleMeshDrawCommands[DrawCommandIndex];
-					PrimitiveIds[DrawCommandIndex] = VisibleMeshDrawCommand.DrawPrimitiveId;
-				}
-			}
-		}
+		BuildPrimitiveIdBuffer();
 	}
 
 	void DumpMeshDrawCommandInstancingStats()
 	{
-		UE_LOG(LogRenderer, Log, TEXT("%s"), GetMeshPassName(PassType));
+		if (PassType != EMeshPass::Num)
+		{
+			UE_LOG(LogRenderer, Log, TEXT("%s"), GetMeshPassName(PassType));
+		}
+
 		UE_LOG(LogRenderer, Log, TEXT("   %i Mesh Draw Commands in %i instancing state buckets"), VisibleMeshDrawCommandsNum, NewPassVisibleMeshDrawCommandsNum);
 		UE_LOG(LogRenderer, Log, TEXT("   Largest %i"), MaxInstances);
 		UE_LOG(LogRenderer, Log, TEXT("   %.1f Dynamic Instancing draw call reduction factor"), VisibleMeshDrawCommandsNum / (float)NewPassVisibleMeshDrawCommandsNum);
@@ -3086,8 +2995,8 @@ public:
 
 
 private:
-	FScene* Scene;
-	FViewInfo& View;
+	const FScene* Scene;
+	const FViewInfo* View;
 	const bool bUseGPUScene;
 	EMeshPass::Type PassType;
 	ETranslucencyPass::Type TranslucencyPass;
@@ -3102,6 +3011,40 @@ private:
 	int32 MaxInstances;
 };
 
+void SortPassMeshDrawCommands(
+	ERHIFeatureLevel::Type FeatureLevel,
+	FMeshCommandOneFrameArray& VisibleMeshDrawCommands,
+	FDynamicMeshDrawCommandStorage& MeshDrawCommandStorage,
+	FGlobalDynamicVertexBuffer& DynamicVertexBuffer,
+	FGlobalDynamicVertexBuffer::FAllocation& PrimitiveIdBuffer,
+	bool bDumpStats)
+{
+	const bool bUseGPUScene = UseGPUScene(GMaxRHIShaderPlatform, FeatureLevel);
+
+	const int32 NumDrawCommands = VisibleMeshDrawCommands.Num();
+	if (NumDrawCommands > 0)
+	{
+		FVisibleMeshDrawCommandProcessPacket Packet(
+			nullptr,
+			nullptr,
+			bUseGPUScene,
+			EMeshPass::Num,
+			ETranslucencyPass::TPT_MAX,
+			VisibleMeshDrawCommands,
+			MeshDrawCommandStorage,
+			DynamicVertexBuffer,
+			PrimitiveIdBuffer);
+
+		Packet.Sort();
+		Packet.BuildPrimitiveIdBuffer();
+
+		if (bDumpStats)
+		{
+			Packet.DumpMeshDrawCommandInstancingStats();
+		}
+	}
+}
+
 void FSceneRenderer::SortMeshDrawCommands()
 {
 	SCOPE_CYCLE_COUNTER(STAT_ProcessVisibleMeshDrawCommands);
@@ -3134,20 +3077,15 @@ void FSceneRenderer::SortMeshDrawCommands()
 				const int32 NumDrawCommands = View.VisibleMeshDrawCommands[PassIndex].Num();
 				if (NumDrawCommands > 0)
 				{
-					// Preallocate primitive Id buffer on rendering thread.
-					if (bUseGPUScene)
-					{
-						View.VisibleMeshDrawCommandPrimitiveIdBuffers[PassIndex] = FGlobalDynamicVertexBuffer::Get().Allocate(NumDrawCommands * sizeof(int32));
-					}
-
 					FVisibleMeshDrawCommandProcessPacket Packet(
 						Scene, 
-						View, 
+						&View, 
 						bUseGPUScene, 
 						PassType,
 						TranslucencyPass, 
 						View.VisibleMeshDrawCommands[PassIndex],
 						View.DynamicMeshDrawCommandStorage[PassIndex],
+						FGlobalDynamicVertexBuffer::Get(),
 						View.VisibleMeshDrawCommandPrimitiveIdBuffers[PassIndex]);
 
 					Packets.Emplace(MoveTemp(Packet));
diff --git a/Engine/Source/Runtime/Renderer/Private/SceneRendering.h b/Engine/Source/Runtime/Renderer/Private/SceneRendering.h
--- a/Engine/Source/Runtime/Renderer/Private/SceneRendering.h
+++ b/Engine/Source/Runtime/Renderer/Private/SceneRendering.h
@@ -1952,4 +1952,4 @@ extern void SortPassMeshDrawCommands(
 	FGlobalDynamicVertexBuffer::FAllocation& PrimitiveIdBuffer,
 	bool bDumpStats);
 
-extern void UpdateTranslucentMeshSortKeys(FScene* Scene, FViewInfo& View, ETranslucencyPass::Type TranslucencyPass, FMeshCommandOneFrameArray& VisibleMeshCommands);
\ No newline at end of file
+extern void UpdateTranslucentMeshSortKeys(const FScene* Scene, const FViewInfo& View, ETranslucencyPass::Type TranslucencyPass, FMeshCommandOneFrameArray& VisibleMeshCommands);
\ No newline at end of file
diff --git a/Engine/Source/Runtime/Renderer/Private/ShadowDepthRendering.cpp b/Engine/Source/Runtime/Renderer/Private/ShadowDepthRendering.cpp
--- a/Engine/Source/Runtime/Renderer/Private/ShadowDepthRendering.cpp
+++ b/Engine/Source/Runtime/Renderer/Private/ShadowDepthRendering.cpp
@@ -1847,11 +1847,14 @@ void FProjectedShadowInfo::RenderDepthInner(FRHICommandListImmediate& RHICmdList
 
 			if (UseMeshDrawCommandPipeline())
 			{
-				FVertexBufferRHIParamRef PrimitiveIdsBuffer = UseGPUScene(GMaxRHIShaderPlatform, FeatureLevel) ? ShadowDepthPass.PrimitiveIdBuffer.VertexBuffer->VertexBufferRHI : nullptr;
-				const int32 BasePrimitiveIdsOffset = ShadowDepthPass.PrimitiveIdBuffer.VertexOffset;
-				const bool bDynamicInstancing = IsDynamicInstancingEnabled() && UseGPUScene(GMaxRHIShaderPlatform, FeatureLevel);
+				if (ShadowDepthPass.VisibleMeshDrawCommands.Num() > 0)
+				{
+					FVertexBufferRHIParamRef PrimitiveIdsBuffer = UseGPUScene(GMaxRHIShaderPlatform, FeatureLevel) ? ShadowDepthPass.PrimitiveIdBuffer.VertexBuffer->VertexBufferRHI : nullptr;
+					const int32 BasePrimitiveIdsOffset = ShadowDepthPass.PrimitiveIdBuffer.VertexOffset;
+					const bool bDynamicInstancing = IsDynamicInstancingEnabled() && UseGPUScene(GMaxRHIShaderPlatform, FeatureLevel);
 
-				SubmitMeshDrawCommands(ShadowDepthPass.VisibleMeshDrawCommands, PrimitiveIdsBuffer, BasePrimitiveIdsOffset, bDynamicInstancing, &ParallelCommandListSet, RHICmdList);
+					SubmitMeshDrawCommands(ShadowDepthPass.VisibleMeshDrawCommands, PrimitiveIdsBuffer, BasePrimitiveIdsOffset, bDynamicInstancing, &ParallelCommandListSet, RHICmdList);
+				}
 			}
 			else
 			{
@@ -1902,11 +1905,14 @@ void FProjectedShadowInfo::RenderDepthInner(FRHICommandListImmediate& RHICmdList
 
 		if (UseMeshDrawCommandPipeline())
 		{
-			FVertexBufferRHIParamRef PrimitiveIdsBuffer = UseGPUScene(GMaxRHIShaderPlatform, FeatureLevel) ? ShadowDepthPass.PrimitiveIdBuffer.VertexBuffer->VertexBufferRHI : nullptr;
-			const int32 BasePrimitiveIdsOffset = ShadowDepthPass.PrimitiveIdBuffer.VertexOffset;
-			const bool bDynamicInstancing = IsDynamicInstancingEnabled() && UseGPUScene(GMaxRHIShaderPlatform, FeatureLevel);
+			if (ShadowDepthPass.VisibleMeshDrawCommands.Num() > 0)
+			{
+				FVertexBufferRHIParamRef PrimitiveIdsBuffer = UseGPUScene(GMaxRHIShaderPlatform, FeatureLevel) ? ShadowDepthPass.PrimitiveIdBuffer.VertexBuffer->VertexBufferRHI : nullptr;
+				const int32 BasePrimitiveIdsOffset = ShadowDepthPass.PrimitiveIdBuffer.VertexOffset;
+				const bool bDynamicInstancing = IsDynamicInstancingEnabled() && UseGPUScene(GMaxRHIShaderPlatform, FeatureLevel);
 
-			SubmitMeshDrawCommands(ShadowDepthPass.VisibleMeshDrawCommands, PrimitiveIdsBuffer, BasePrimitiveIdsOffset, bDynamicInstancing, nullptr, RHICmdList);
+				SubmitMeshDrawCommands(ShadowDepthPass.VisibleMeshDrawCommands, PrimitiveIdsBuffer, BasePrimitiveIdsOffset, bDynamicInstancing, nullptr, RHICmdList);
+			}
 		}
 		else
 		{
diff --git a/Engine/Source/Runtime/Renderer/Private/ShadowRendering.cpp b/Engine/Source/Runtime/Renderer/Private/ShadowRendering.cpp
--- a/Engine/Source/Runtime/Renderer/Private/ShadowRendering.cpp
+++ b/Engine/Source/Runtime/Renderer/Private/ShadowRendering.cpp
@@ -612,11 +612,14 @@ void FProjectedShadowInfo::SetupProjectionStencilMask(
 		if (UseMeshDrawCommandPipeline())
 		{
 			const FShadowMeshDrawCommandPass& ProjectionStencilingPass = ProjectionStencilingPasses[ViewIndex];
-			FVertexBufferRHIParamRef PrimitiveIdsBuffer = UseGPUScene(GMaxRHIShaderPlatform, View->GetFeatureLevel()) ? ProjectionStencilingPass.PrimitiveIdBuffer.VertexBuffer->VertexBufferRHI : nullptr;
-			const int32 BasePrimitiveIdsOffset = ProjectionStencilingPass.PrimitiveIdBuffer.VertexOffset;
-			const bool bDynamicInstancing = IsDynamicInstancingEnabled() && UseGPUScene(GMaxRHIShaderPlatform, View->GetFeatureLevel());
+			if (ProjectionStencilingPass.VisibleMeshDrawCommands.Num() > 0)
+			{
+				FVertexBufferRHIParamRef PrimitiveIdsBuffer = UseGPUScene(GMaxRHIShaderPlatform, View->GetFeatureLevel()) ? ProjectionStencilingPass.PrimitiveIdBuffer.VertexBuffer->VertexBufferRHI : nullptr;
+				const int32 BasePrimitiveIdsOffset = ProjectionStencilingPass.PrimitiveIdBuffer.VertexOffset;
+				const bool bDynamicInstancing = IsDynamicInstancingEnabled() && UseGPUScene(GMaxRHIShaderPlatform, View->GetFeatureLevel());
 
-			SubmitMeshDrawCommands(ProjectionStencilingPass.VisibleMeshDrawCommands, PrimitiveIdsBuffer, BasePrimitiveIdsOffset, bDynamicInstancing, nullptr, RHICmdList);
+				SubmitMeshDrawCommands(ProjectionStencilingPass.VisibleMeshDrawCommands, PrimitiveIdsBuffer, BasePrimitiveIdsOffset, bDynamicInstancing, nullptr, RHICmdList);
+			}
 		}
 		else
 		{
@@ -842,11 +845,14 @@ void FProjectedShadowInfo::SetupProjectionStencilMask(
 			if (UseMeshDrawCommandPipeline())
 			{
 				const FShadowMeshDrawCommandPass& ProjectionStencilingPass = ProjectionStencilingPasses[ViewIndex];
-				FVertexBufferRHIParamRef PrimitiveIdsBuffer = UseGPUScene(GMaxRHIShaderPlatform, View->GetFeatureLevel()) ? ProjectionStencilingPass.PrimitiveIdBuffer.VertexBuffer->VertexBufferRHI : nullptr;
-				const int32 BasePrimitiveIdsOffset = ProjectionStencilingPass.PrimitiveIdBuffer.VertexOffset;
-				const bool bDynamicInstancing = IsDynamicInstancingEnabled() && UseGPUScene(GMaxRHIShaderPlatform, View->GetFeatureLevel());
+				if (ProjectionStencilingPass.VisibleMeshDrawCommands.Num() > 0)
+				{
+					FVertexBufferRHIParamRef PrimitiveIdsBuffer = UseGPUScene(GMaxRHIShaderPlatform, View->GetFeatureLevel()) ? ProjectionStencilingPass.PrimitiveIdBuffer.VertexBuffer->VertexBufferRHI : nullptr;
+					const int32 BasePrimitiveIdsOffset = ProjectionStencilingPass.PrimitiveIdBuffer.VertexOffset;
+					const bool bDynamicInstancing = IsDynamicInstancingEnabled() && UseGPUScene(GMaxRHIShaderPlatform, View->GetFeatureLevel());
 
-				SubmitMeshDrawCommands(ProjectionStencilingPass.VisibleMeshDrawCommands, PrimitiveIdsBuffer, BasePrimitiveIdsOffset, bDynamicInstancing, nullptr, RHICmdList);
+					SubmitMeshDrawCommands(ProjectionStencilingPass.VisibleMeshDrawCommands, PrimitiveIdsBuffer, BasePrimitiveIdsOffset, bDynamicInstancing, nullptr, RHICmdList);
+				}
 			}
 			else
 			{
