Indirect lighting cached - fixed static analysis nullptr warning.

diff --git a/Engine/Source/Runtime/Renderer/Private/IndirectLightingCache.cpp b/Engine/Source/Runtime/Renderer/Private/IndirectLightingCache.cpp
--- a/Engine/Source/Runtime/Renderer/Private/IndirectLightingCache.cpp
+++ b/Engine/Source/Runtime/Renderer/Private/IndirectLightingCache.cpp
@@ -744,7 +744,9 @@ void FIndirectLightingCache::UpdateCachePrimitive(
 		// Cached mesh draw command shader selection depends on IndirectLightingCacheAllocation existence and bPointSample value.
 		// If any of those changes, we need to re-cache mesh draw commands.
 		if (OriginalCachedIndirectLightingCacheAllocation == nullptr
-			|| OriginalCachedIndirectLightingCacheAllocation->bPointSample != PrimitiveSceneInfo->IndirectLightingCacheAllocation->bPointSample)
+			|| (OriginalCachedIndirectLightingCacheAllocation 
+				&& PrimitiveSceneInfo->IndirectLightingCacheAllocation 
+				&& OriginalCachedIndirectLightingCacheAllocation->bPointSample != PrimitiveSceneInfo->IndirectLightingCacheAllocation->bPointSample))
 		{
 			PrimitivesToUpdateStaticMeshes.Add(PrimitiveSceneInfo);
 		}
