Removed check() call that was causing a crash in split screen mode with ray tracing enabled. More work is required to support multiple views in ray tracing mode.

diff --git a/Engine/Source/Runtime/Renderer/Private/RayTracing/RayTracingShadows.cpp b/Engine/Source/Runtime/Renderer/Private/RayTracing/RayTracingShadows.cpp
--- a/Engine/Source/Runtime/Renderer/Private/RayTracing/RayTracingShadows.cpp
+++ b/Engine/Source/Runtime/Renderer/Private/RayTracing/RayTracingShadows.cpp
@@ -75,8 +75,7 @@ void FDeferredShadingSceneRenderer::RenderRayTracingOcclusion(
 	FLightSceneProxy* LightSceneProxy = LightSceneInfo->Proxy;
 	check(LightSceneProxy);
 
-	check(Views.Num() == 1);
-	const FViewInfo& View = Views[0];
+	const FViewInfo& View = Views[0]; // #dxr_todo: what about multi-view case?
 
 	FSceneRenderTargets& SceneContext = FSceneRenderTargets::Get(RHICmdList);
 	FRDGBuilder GraphBuilder(RHICmdList);
