Fully disable scene textures in all ray tracing material hit shaders. Before it was only disabled in surface materials.

diff --git a/Engine/Source/Runtime/Renderer/Private/RayTracing/RayTracingMaterialHitShaders.cpp b/Engine/Source/Runtime/Renderer/Private/RayTracing/RayTracingMaterialHitShaders.cpp
--- a/Engine/Source/Runtime/Renderer/Private/RayTracing/RayTracingMaterialHitShaders.cpp
+++ b/Engine/Source/Runtime/Renderer/Private/RayTracing/RayTracingMaterialHitShaders.cpp
@@ -95,7 +95,7 @@ public:
 
 	static void ModifyCompilationEnvironment(EShaderPlatform Platform, const FMaterial* Material, FShaderCompilerEnvironment& OutEnvironment)
 	{
-		OutEnvironment.SetDefine(TEXT("SCENE_TEXTURES_DISABLED"), Material->GetMaterialDomain() != MD_Surface);
+		OutEnvironment.SetDefine(TEXT("SCENE_TEXTURES_DISABLED"), 1);
 		LightMapPolicyType::ModifyCompilationEnvironment(Platform, Material, OutEnvironment);
 		FMeshMaterialShader::ModifyCompilationEnvironment(Platform, Material, OutEnvironment);
 	}
