Assert when attempting to bind an already deleted uniform buffer on a FMeshDrawCommand, which greatly assists tracking down uniform buffer lifetime errors

diff --git a/Engine/Source/Runtime/RHI/Public/RHIResources.h b/Engine/Source/Runtime/RHI/Public/RHIResources.h
--- a/Engine/Source/Runtime/RHI/Public/RHIResources.h
+++ b/Engine/Source/Runtime/RHI/Public/RHIResources.h
@@ -101,6 +101,11 @@ public:
 		return bCommitted;
 	}
 
+	bool IsValid() const
+	{
+		return !MarkedForDelete && NumRefs.GetValue() > 0;
+	}
+
 private:
 	mutable FThreadSafeCounter NumRefs;
 	mutable int32 MarkedForDelete;
@@ -245,6 +250,7 @@ class FRHIRayTracingScene : public FRHIResource {};
 #define VALIDATE_UNIFORM_BUFFER_LAYOUT_LIFETIME 0
 
 // Whether to assert when a uniform buffer is being deleted while still referenced by a mesh draw command
+// Enabling this requires -norhithread to work correctly since FRHIResource lifetime is managed by both the RT and RHIThread
 #define VALIDATE_UNIFORM_BUFFER_LIFETIME 0
 
 /** The layout of a uniform buffer in memory. */
@@ -390,16 +396,22 @@ public:
 
 	FORCEINLINE_DEBUGGABLE uint32 Release() const
 	{
+		const FRHIUniformBufferLayout* LocalLayout = Layout;
+
+#if VALIDATE_UNIFORM_BUFFER_LIFETIME
+		int32 LocalNumMeshCommandReferencesForDebugging = NumMeshCommandReferencesForDebugging;
+#endif
+
 		uint32 NewRefCount = FRHIResource::Release();
 
 		if (NewRefCount == 0)
 		{
 #if VALIDATE_UNIFORM_BUFFER_LAYOUT_LIFETIME
-			Layout->NumUsesForDebugging--;
-			check(Layout->NumUsesForDebugging >= 0);
+			LocalLayout->NumUsesForDebugging--;
+			check(LocalLayout->NumUsesForDebugging >= 0);
 #endif
 #if VALIDATE_UNIFORM_BUFFER_LIFETIME
-			check(NumMeshCommandReferencesForDebugging == 0 || GIsRequestingExit);
+			check(LocalNumMeshCommandReferencesForDebugging == 0 || GIsRequestingExit);
 #endif
 		}
 
diff --git a/Engine/Source/Runtime/Renderer/Private/MeshPassProcessor.cpp b/Engine/Source/Runtime/Renderer/Private/MeshPassProcessor.cpp
--- a/Engine/Source/Runtime/Renderer/Private/MeshPassProcessor.cpp
+++ b/Engine/Source/Runtime/Renderer/Private/MeshPassProcessor.cpp
@@ -401,9 +401,11 @@ FMeshDrawShaderBindings::~FMeshDrawShaderBindings()
 	for (int32 ShaderBindingsIndex = 0; ShaderBindingsIndex < ShaderLayouts.Num(); ShaderBindingsIndex++)
 	{
 		FMeshDrawSingleShaderBindings SingleShaderBindings(ShaderLayouts[ShaderBindingsIndex], ShaderBindingDataPtr);
-		const FUniformBufferRHIParamRef* UniformBufferBindings = SingleShaderBindings.GetUniformBufferStart();
 
-		for (int32 UniformBufferIndex = 0; UniformBufferIndex < SingleShaderBindings.NumUniformBuffers; UniformBufferIndex++)
+		const FUniformBufferRHIParamRef* RESTRICT UniformBufferBindings = SingleShaderBindings.GetUniformBufferStart();
+		const int32 NumUniformBuffers = SingleShaderBindings.ParameterMapInfo.UniformBuffers.Num();
+
+		for (int32 UniformBufferIndex = 0; UniformBufferIndex < NumUniformBuffers; UniformBufferIndex++)
 		{
 			FUniformBufferRHIParamRef UniformBuffer = UniformBufferBindings[UniformBufferIndex];
 
@@ -576,9 +578,10 @@ void FMeshDrawShaderBindings::CopyFrom(const FMeshDrawShaderBindings& Other)
 	for (int32 ShaderBindingsIndex = 0; ShaderBindingsIndex < ShaderLayouts.Num(); ShaderBindingsIndex++)
 	{
 		FMeshDrawSingleShaderBindings SingleShaderBindings(ShaderLayouts[ShaderBindingsIndex], ShaderBindingDataPtr);
-		const FUniformBufferRHIParamRef* UniformBufferBindings = SingleShaderBindings.GetUniformBufferStart();
+		const FUniformBufferRHIParamRef* RESTRICT UniformBufferBindings = SingleShaderBindings.GetUniformBufferStart();
+		const int32 NumUniformBuffers = SingleShaderBindings.ParameterMapInfo.UniformBuffers.Num();
 
-		for (int32 UniformBufferIndex = 0; UniformBufferIndex < SingleShaderBindings.ParameterMapInfo.UniformBuffers.Num(); UniformBufferIndex++)
+		for (int32 UniformBufferIndex = 0; UniformBufferIndex < NumUniformBuffers; UniformBufferIndex++)
 		{
 			FUniformBufferRHIParamRef UniformBuffer = UniformBufferBindings[UniformBufferIndex];
 
diff --git a/Engine/Source/Runtime/Renderer/Public/MeshDrawShaderBindings.h b/Engine/Source/Runtime/Renderer/Public/MeshDrawShaderBindings.h
--- a/Engine/Source/Runtime/Renderer/Public/MeshDrawShaderBindings.h
+++ b/Engine/Source/Runtime/Renderer/Public/MeshDrawShaderBindings.h
@@ -99,6 +99,7 @@ public:
 		if (Parameter.IsBound())
 		{
 			checkf(Value.GetReference(), TEXT("Attempted to set null uniform buffer for type %s on %s"), UniformBufferStructType::StaticStructMetadata.GetStructTypeName(), GetShaderFrequencyString(Frequency));
+			checkfSlow(Value.GetReference()->IsValid(), TEXT("Attempted to set already deleted uniform buffer for type %s on %s"), UniformBufferStructType::StaticStructMetadata.GetStructTypeName(), GetShaderFrequencyString(Frequency));
 			WriteBindingUniformBuffer(Value.GetReference(), Parameter.GetBaseIndex());
 		}
 	}
@@ -111,6 +112,7 @@ public:
 		if (Parameter.IsBound())
 		{
 			checkf(Value.GetUniformBufferRHI(), TEXT("Attempted to set null uniform buffer for type %s on %s"), UniformBufferStructType::StaticStructMetadata.GetStructTypeName(), GetShaderFrequencyString(Frequency));
+			checkfSlow(Value.GetUniformBufferRHI()->IsValid(), TEXT("Attempted to set already deleted uniform buffer for type %s on %s"), UniformBufferStructType::StaticStructMetadata.GetStructTypeName(), GetShaderFrequencyString(Frequency));
 			WriteBindingUniformBuffer(Value.GetUniformBufferRHI(), Parameter.GetBaseIndex());
 		}
 	}
@@ -122,6 +124,7 @@ public:
 		if (Parameter.IsBound())
 		{
 			checkf(Value, TEXT("Attempted to set null uniform buffer with unknown type on %s"), GetShaderFrequencyString(Frequency));
+			checkfSlow(Value->IsValid(), TEXT("Attempted to set already deleted uniform buffer of type %s on %s"), *Value->GetLayout().GetDebugName().ToString(), GetShaderFrequencyString(Frequency));
 			WriteBindingUniformBuffer(Value, Parameter.GetBaseIndex());
 		}
 	}
@@ -133,6 +136,7 @@ public:
 		if (Parameter.IsBound())
 		{
 			checkf(Value, TEXT("Attempted to set null SRV on slot %u of %s"), Parameter.GetBaseIndex(), GetShaderFrequencyString(Frequency));
+			checkfSlow(Value->IsValid(), TEXT("Attempted to set already deleted SRV on slot %u of %s"), Parameter.GetBaseIndex(), GetShaderFrequencyString(Frequency));
 			WriteBindingSRV(Value, Parameter.GetBaseIndex());
 		}
 	}
