Removed SortEvents, as with static draw list removal it is no longer used.

diff --git a/Engine/Source/Runtime/Renderer/Private/DeferredShadingRenderer.cpp b/Engine/Source/Runtime/Renderer/Private/DeferredShadingRenderer.cpp
--- a/Engine/Source/Runtime/Renderer/Private/DeferredShadingRenderer.cpp
+++ b/Engine/Source/Runtime/Renderer/Private/DeferredShadingRenderer.cpp
@@ -168,7 +168,6 @@ static TAutoConsoleVariable<int32> CVarRayTracing(
 DECLARE_CYCLE_STAT(TEXT("PostInitViews FlushDel"), STAT_PostInitViews_FlushDel, STATGROUP_InitViews);
 DECLARE_CYCLE_STAT(TEXT("InitViews Intentional Stall"), STAT_InitViews_Intentional_Stall, STATGROUP_InitViews);
 
-DEFINE_STAT(STAT_FDeferredShadingSceneRenderer_AsyncSortBasePassStaticData);
 DECLARE_CYCLE_STAT(TEXT("DeferredShadingSceneRenderer UpdateDownsampledDepthSurface"), STAT_FDeferredShadingSceneRenderer_UpdateDownsampledDepthSurface, STATGROUP_SceneRendering);
 DECLARE_CYCLE_STAT(TEXT("DeferredShadingSceneRenderer Render Init"), STAT_FDeferredShadingSceneRenderer_Render_Init, STATGROUP_SceneRendering);
 DECLARE_CYCLE_STAT(TEXT("DeferredShadingSceneRenderer Render ServiceLocalQueue"), STAT_FDeferredShadingSceneRenderer_Render_ServiceLocalQueue, STATGROUP_SceneRendering);
@@ -1008,13 +1007,12 @@ void FDeferredShadingSceneRenderer::Render(FRHICommandListImmediate& RHICmdList)
 		? FExclusiveDepthStencil::DepthRead_StencilWrite 
 		: FExclusiveDepthStencil::DepthWrite_StencilWrite;
 
-	FGraphEventArray SortEvents;
 	FGraphEventArray UpdateViewCustomDataEvents;
 	FILCUpdatePrimTaskData ILCTaskData;
 
 	// Find the visible primitives.
 	RHICmdList.ImmediateFlush(EImmediateFlushType::DispatchToRHIThread);
-	bool bDoInitViewAftersPrepass = InitViews(RHICmdList, BasePassDepthStencilAccess, ILCTaskData, SortEvents, UpdateViewCustomDataEvents);
+	bool bDoInitViewAftersPrepass = InitViews(RHICmdList, BasePassDepthStencilAccess, ILCTaskData, UpdateViewCustomDataEvents);
 
 	static const auto CVarVirtualTextureLightmaps = IConsoleManager::Get().FindTConsoleVariableDataInt(TEXT("r.VirtualTexturedLightmaps"));
 	if (CVarVirtualTextureLightmaps && CVarVirtualTextureLightmaps->GetValueOnRenderThread())
@@ -1184,7 +1182,7 @@ void FDeferredShadingSceneRenderer::Render(FRHICommandListImmediate& RHICmdList)
 	}
 
 	bool bDidAfterTaskWork = false;
-	auto AfterTasksAreStarted = [&bDidAfterTaskWork, bDoInitViewAftersPrepass, this, &RHICmdList, &ILCTaskData, &SortEvents, &UpdateViewCustomDataEvents, bLateFXPrerender, bDoFXPrerender]()
+	auto AfterTasksAreStarted = [&bDidAfterTaskWork, bDoInitViewAftersPrepass, this, &RHICmdList, &ILCTaskData, &UpdateViewCustomDataEvents, bLateFXPrerender, bDoFXPrerender]()
 	{
 		if (!bDidAfterTaskWork)
 		{
@@ -1192,7 +1190,7 @@ void FDeferredShadingSceneRenderer::Render(FRHICommandListImmediate& RHICmdList)
 			bDidAfterTaskWork = true; // only do this once
 			if (bDoInitViewAftersPrepass)
 			{
-				InitViewsPossiblyAfterPrepass(RHICmdList, ILCTaskData, SortEvents, UpdateViewCustomDataEvents);
+				InitViewsPossiblyAfterPrepass(RHICmdList, ILCTaskData, UpdateViewCustomDataEvents);
 				PrepareDistanceFieldScene(RHICmdList, false);
 
 				if (!UseMeshDrawCommandPipeline())
diff --git a/Engine/Source/Runtime/Renderer/Private/DeferredShadingRenderer.h b/Engine/Source/Runtime/Renderer/Private/DeferredShadingRenderer.h
--- a/Engine/Source/Runtime/Renderer/Private/DeferredShadingRenderer.h
+++ b/Engine/Source/Runtime/Renderer/Private/DeferredShadingRenderer.h
@@ -142,9 +142,9 @@ private:
 	void PreVisibilityFrameSetup(FRHICommandListImmediate& RHICmdList);
 
 	/** Determines which primitives are visible for each view. */
-	bool InitViews(FRHICommandListImmediate& RHICmdList, FExclusiveDepthStencil::Type BasePassDepthStencilAccess, struct FILCUpdatePrimTaskData& ILCTaskData, FGraphEventArray& SortEvents, FGraphEventArray& UpdateViewCustomDataEvents);
+	bool InitViews(FRHICommandListImmediate& RHICmdList, FExclusiveDepthStencil::Type BasePassDepthStencilAccess, struct FILCUpdatePrimTaskData& ILCTaskData, FGraphEventArray& UpdateViewCustomDataEvents);
 
-	void InitViewsPossiblyAfterPrepass(FRHICommandListImmediate& RHICmdList, struct FILCUpdatePrimTaskData& ILCTaskData, FGraphEventArray& SortEvents, FGraphEventArray& UpdateViewCustomDataEvents);
+	void InitViewsPossiblyAfterPrepass(FRHICommandListImmediate& RHICmdList, struct FILCUpdatePrimTaskData& ILCTaskData, FGraphEventArray& UpdateViewCustomDataEvents);
 
 	void SetupSceneReflectionCaptureBuffer(FRHICommandListImmediate& RHICmdList);
 
@@ -511,5 +511,4 @@ private:
 	friend class FTranslucentPrimSet;
 };
 
-DECLARE_CYCLE_STAT_EXTERN(TEXT("PrePass"), STAT_CLM_PrePass, STATGROUP_CommandListMarkers, );
-DECLARE_CYCLE_STAT_EXTERN(TEXT("DeferredShadingSceneRenderer AsyncSortBasePassStaticData"), STAT_FDeferredShadingSceneRenderer_AsyncSortBasePassStaticData, STATGROUP_SceneRendering, );
\ No newline at end of file
+DECLARE_CYCLE_STAT_EXTERN(TEXT("PrePass"), STAT_CLM_PrePass, STATGROUP_CommandListMarkers, );
\ No newline at end of file
diff --git a/Engine/Source/Runtime/Renderer/Private/SceneHitProxyRendering.cpp b/Engine/Source/Runtime/Renderer/Private/SceneHitProxyRendering.cpp
--- a/Engine/Source/Runtime/Renderer/Private/SceneHitProxyRendering.cpp
+++ b/Engine/Source/Runtime/Renderer/Private/SceneHitProxyRendering.cpp
@@ -558,13 +558,12 @@ void FDeferredShadingSceneRenderer::RenderHitProxies(FRHICommandListImmediate& R
 	if (HitProxyRT)
 	{
 		// Find the visible primitives.
-		FGraphEventArray SortEvents;
 		FGraphEventArray UpdateViewCustomDataEvents;
 		FILCUpdatePrimTaskData ILCTaskData;
-		bool bDoInitViewAftersPrepass = InitViews(RHICmdList, FExclusiveDepthStencil::DepthWrite_StencilWrite, ILCTaskData, SortEvents, UpdateViewCustomDataEvents);
+		bool bDoInitViewAftersPrepass = InitViews(RHICmdList, FExclusiveDepthStencil::DepthWrite_StencilWrite, ILCTaskData, UpdateViewCustomDataEvents);
 		if (bDoInitViewAftersPrepass)
 		{
-			InitViewsPossiblyAfterPrepass(RHICmdList, ILCTaskData, SortEvents, UpdateViewCustomDataEvents);
+			InitViewsPossiblyAfterPrepass(RHICmdList, ILCTaskData, UpdateViewCustomDataEvents);
 		}
 
 		UpdateGPUScene(RHICmdList, *Scene);
diff --git a/Engine/Source/Runtime/Renderer/Private/SceneVisibility.cpp b/Engine/Source/Runtime/Renderer/Private/SceneVisibility.cpp
--- a/Engine/Source/Runtime/Renderer/Private/SceneVisibility.cpp
+++ b/Engine/Source/Runtime/Renderer/Private/SceneVisibility.cpp
@@ -3736,9 +3736,9 @@ void FDeferredShadingSceneRenderer::PreVisibilityFrameSetup(FRHICommandListImmed
 
 /**
  * Initialize scene's views.
- * Check visibility, sort translucent items, etc.
+ * Check visibility, build visible mesh commands, etc.
  */
-bool FDeferredShadingSceneRenderer::InitViews(FRHICommandListImmediate& RHICmdList, FExclusiveDepthStencil::Type BasePassDepthStencilAccess, struct FILCUpdatePrimTaskData& ILCTaskData, FGraphEventArray& SortEvents, FGraphEventArray& UpdateViewCustomDataEvents)
+bool FDeferredShadingSceneRenderer::InitViews(FRHICommandListImmediate& RHICmdList, FExclusiveDepthStencil::Type BasePassDepthStencilAccess, struct FILCUpdatePrimTaskData& ILCTaskData, FGraphEventArray& UpdateViewCustomDataEvents)
 {
 	SCOPED_NAMED_EVENT(FDeferredShadingSceneRenderer_InitViews, FColor::Emerald);
 	SCOPE_CYCLE_COUNTER(STAT_InitViewsTime);
@@ -3771,7 +3771,7 @@ bool FDeferredShadingSceneRenderer::InitViews(FRHICommandListImmediate& RHICmdLi
 
 	if (!bDoInitViewAftersPrepass)
 	{
-		InitViewsPossiblyAfterPrepass(RHICmdList, ILCTaskData, SortEvents, UpdateViewCustomDataEvents);
+		InitViewsPossiblyAfterPrepass(RHICmdList, ILCTaskData, UpdateViewCustomDataEvents);
 	}
 
 	if (!UseMeshDrawCommandPipeline())
@@ -3965,18 +3965,11 @@ void FDeferredShadingSceneRenderer::SetupSceneReflectionCaptureBuffer(FRHIComman
 	}
 }
 
-void FDeferredShadingSceneRenderer::InitViewsPossiblyAfterPrepass(FRHICommandListImmediate& RHICmdList, struct FILCUpdatePrimTaskData& ILCTaskData, FGraphEventArray& SortEvents, FGraphEventArray& UpdateViewCustomDataEvents)
+void FDeferredShadingSceneRenderer::InitViewsPossiblyAfterPrepass(FRHICommandListImmediate& RHICmdList, struct FILCUpdatePrimTaskData& ILCTaskData, FGraphEventArray& UpdateViewCustomDataEvents)
 {
 	SCOPED_NAMED_EVENT(FDeferredShadingSceneRenderer_InitViewsPossiblyAfterPrepass, FColor::Emerald);
 	SCOPE_CYCLE_COUNTER(STAT_InitViewsPossiblyAfterPrepass);
 
-	// this cannot be moved later because of static mesh updates for stuff that is only visible in shadows
-	if (SortEvents.Num())
-	{
-		QUICK_SCOPE_CYCLE_COUNTER(STAT_FDeferredShadingSceneRenderer_AsyncSortBasePassStaticData_Wait);
-		FTaskGraphInterface::Get().WaitUntilTasksComplete(SortEvents, ENamedThreads::GetRenderThread());
-	}
-
 	if (ViewFamily.EngineShowFlags.DynamicShadows 
 		&& !IsSimpleForwardShadingEnabled(ShaderPlatform)
 		&& !ViewFamily.EngineShowFlags.HitProxies)
