Move fix from CL 4585757 to dev-rendering: "Guard against divide-by-zero when setting SSAO shader parameters.".

diff --git a/Engine/Source/Runtime/Renderer/Private/CompositionLighting/PostProcessAmbientOcclusion.cpp b/Engine/Source/Runtime/Renderer/Private/CompositionLighting/PostProcessAmbientOcclusion.cpp
--- a/Engine/Source/Runtime/Renderer/Private/CompositionLighting/PostProcessAmbientOcclusion.cpp
+++ b/Engine/Source/Runtime/Renderer/Private/CompositionLighting/PostProcessAmbientOcclusion.cpp
@@ -262,9 +262,10 @@ public:
 			TemporalOffset = (View.State->GetCurrentTemporalAASampleIndex() % 8) * FVector2D(2.48f, 7.52f) / 64.0f;
 		}
 		const float HzbStepMipLevelFactorValue = FMath::Clamp(FSSAOHelper::GetAmbientOcclusionStepMipLevelFactor(), 0.0f, 100.0f);
+		const float InvAmbientOcclusionDistance = 1.0f / FMath::Max(Settings.AmbientOcclusionDistance_DEPRECATED, KINDA_SMALL_NUMBER);
 
 		// /1000 to be able to define the value in that distance
-		Value[0] = FVector4(Settings.AmbientOcclusionPower, Settings.AmbientOcclusionBias / 1000.0f, 1.0f / Settings.AmbientOcclusionDistance_DEPRECATED, Settings.AmbientOcclusionIntensity);
+		Value[0] = FVector4(Settings.AmbientOcclusionPower, Settings.AmbientOcclusionBias / 1000.0f, InvAmbientOcclusionDistance, Settings.AmbientOcclusionIntensity);
 		Value[1] = FVector4(ViewportUVToRandomUV.X, ViewportUVToRandomUV.Y, AORadiusInShader, Ratio);
 		Value[2] = FVector4(ScaleToFullRes, Settings.AmbientOcclusionMipThreshold / ScaleToFullRes, ScaleRadiusInWorldSpace, Settings.AmbientOcclusionMipBlend);
 		Value[3] = FVector4(TemporalOffset.X, TemporalOffset.Y, StaticFraction, InvTanHalfFov);
