Removed duplicate / redundant version of BuildRectLightMipTree() from FDeferredShadingSceneRenderer and associated duplicated shaders. This fixes linker errors in some configurations (duplicate FBuildRectLightMipTreeCS symbols).

diff --git a/Engine/Source/Runtime/Renderer/Private/DeferredShadingRenderer.h b/Engine/Source/Runtime/Renderer/Private/DeferredShadingRenderer.h
--- a/Engine/Source/Runtime/Renderer/Private/DeferredShadingRenderer.h
+++ b/Engine/Source/Runtime/Renderer/Private/DeferredShadingRenderer.h
@@ -451,7 +451,6 @@ private:
 
 	void RenderRayTracingOcclusion(FRHICommandListImmediate& RHICmdList, const FLightSceneInfo* LightSceneInfo, TRefCountPtr<IPooledRenderTarget>& ScreenShadowMaskTexture);
 
-	void BuildRectLightMipTree(FRHICommandListImmediate& RHICmdList, const FLightSceneInfo& RectLightSceneInfo, FRWBuffer& RectLightMipTree, FIntVector& RectLightMipTreeDimensions);
 	void VisualizeRectLightMipTree(FRHICommandListImmediate& RHICmdList, const FViewInfo& View, const FRWBuffer& RectLightMipTree, const FIntVector& RectLightMipTreeDimensions);
 	void RenderRayTracingRectLight(FRHICommandListImmediate& RHICmdList, const FLightSceneInfo& RectLightSceneInfo, TRefCountPtr<IPooledRenderTarget>& RectLightRT, TRefCountPtr<IPooledRenderTarget>& HitDistanceRT);
 	void RenderRayTracingOcclusionForRectLight(FRHICommandListImmediate& RHICmdList, const FLightSceneInfo& RectLightSceneInfo, TRefCountPtr<IPooledRenderTarget>& ScreenShadowMaskTexture);
diff --git a/Engine/Source/Runtime/Renderer/Private/RayTracing/RayTracingShadowsForRectLight.cpp b/Engine/Source/Runtime/Renderer/Private/RayTracing/RayTracingShadowsForRectLight.cpp
--- a/Engine/Source/Runtime/Renderer/Private/RayTracing/RayTracingShadowsForRectLight.cpp
+++ b/Engine/Source/Runtime/Renderer/Private/RayTracing/RayTracingShadowsForRectLight.cpp
@@ -177,136 +177,6 @@ IMPLEMENT_RECT_LIGHT_OCCLUSION_TYPE(0, 1, 1);
 IMPLEMENT_RECT_LIGHT_OCCLUSION_TYPE(1, 0, 0);
 IMPLEMENT_RECT_LIGHT_OCCLUSION_TYPE(1, 0, 1);
 
-class FBuildRectLightMipTreeCS : public FGlobalShader
-{
-	DECLARE_SHADER_TYPE(FBuildRectLightMipTreeCS, Global)
-
-public:
-	static bool ShouldCompilePermutation(const FGlobalShaderPermutationParameters& Parameters)
-	{
-		return RHISupportsComputeShaders(Parameters.Platform);
-	}
-
-	static uint32 GetGroupSize()
-	{
-		return 16;
-	}
-
-	static void ModifyCompilationEnvironment(const FGlobalShaderPermutationParameters& Parameters, FShaderCompilerEnvironment& OutEnvironment)
-	{
-		FGlobalShader::ModifyCompilationEnvironment(Parameters, OutEnvironment);
-		OutEnvironment.SetDefine(TEXT("THREADGROUP_SIZE"), GetGroupSize());
-	}
-
-	FBuildRectLightMipTreeCS(const ShaderMetaType::CompiledShaderInitializerType& Initializer)
-		: FGlobalShader(Initializer)
-	{
-		TextureParameter.Bind(Initializer.ParameterMap, TEXT("RectLightTexture"));
-		TextureSamplerParameter.Bind(Initializer.ParameterMap, TEXT("TextureSampler"));
-		DimensionsParameter.Bind(Initializer.ParameterMap, TEXT("Dimensions"));
-		MipLevelParameter.Bind(Initializer.ParameterMap, TEXT("MipLevel"));
-		MipTreeParameter.Bind(Initializer.ParameterMap, TEXT("MipTree"));
-	}
-
-	FBuildRectLightMipTreeCS() {}
-
-	void SetParameters(
-		FRHICommandList& RHICmdList,
-		FTextureRHIRef Texture,
-		const FIntVector& Dimensions,
-		uint32 MipLevel,
-		FRWBuffer& MipTree
-	)
-	{
-		FComputeShaderRHIParamRef ShaderRHI = GetComputeShader();
-
-		SetShaderValue(RHICmdList, ShaderRHI, DimensionsParameter, Dimensions);
-		SetShaderValue(RHICmdList, ShaderRHI, MipLevelParameter, MipLevel);
-		SetTextureParameter(RHICmdList, ShaderRHI, TextureParameter, TextureSamplerParameter, TStaticSamplerState<SF_Bilinear>::GetRHI(), Texture);
-
-		check(MipTreeParameter.IsBound());
-		MipTreeParameter.SetBuffer(RHICmdList, ShaderRHI, MipTree);
-	}
-
-	void UnsetParameters(
-		FRHICommandList& RHICmdList,
-		EResourceTransitionAccess TransitionAccess,
-		EResourceTransitionPipeline TransitionPipeline,
-		FRWBuffer& MipTree,
-		FComputeFenceRHIParamRef Fence)
-	{
-		FComputeShaderRHIParamRef ShaderRHI = GetComputeShader();
-
-		MipTreeParameter.UnsetUAV(RHICmdList, ShaderRHI);
-		RHICmdList.TransitionResource(TransitionAccess, TransitionPipeline, MipTree.UAV, Fence);
-	}
-
-	virtual bool Serialize(FArchive& Ar)
-	{
-		bool bShaderHasOutdatedParameters = FGlobalShader::Serialize(Ar);
-		Ar << TextureParameter;
-		Ar << TextureSamplerParameter;
-		Ar << DimensionsParameter;
-		Ar << MipLevelParameter;
-		Ar << MipTreeParameter;
-		return bShaderHasOutdatedParameters;
-	}
-
-private:
-	FShaderResourceParameter TextureParameter;
-	FShaderResourceParameter TextureSamplerParameter;
-
-	FShaderParameter DimensionsParameter;
-	FShaderParameter MipLevelParameter;
-	FRWShaderParameter MipTreeParameter;
-};
-
-IMPLEMENT_SHADER_TYPE(, FBuildRectLightMipTreeCS, TEXT("/Engine/Private/Raytracing/BuildMipTreeCS.usf"), TEXT("BuildRectLightMipTreeCS"), SF_Compute)
-
-void FDeferredShadingSceneRenderer::BuildRectLightMipTree(
-	FRHICommandListImmediate& RHICmdList,
-	const FLightSceneInfo& RectLightSceneInfo,
-	FRWBuffer& RectLightMipTree,
-	FIntVector& RectLightMipTreeDimensions
-)
-{
-	SCOPED_DRAW_EVENT(RHICmdList, BuildRectLightMipTree);
-	SCOPED_GPU_STAT(RHICmdList, BuildRectLightMipTree);
-
-	check(RectLightSceneInfo.Proxy);
-	const auto ShaderMap = GetGlobalShaderMap(FeatureLevel);
-	TShaderMapRef<FBuildRectLightMipTreeCS> BuildRectLightMipTreeComputeShader(ShaderMap);
-	RHICmdList.SetComputeShader(BuildRectLightMipTreeComputeShader->GetComputeShader());
-
-	// Allocate MIP tree
-	FLightShaderParameters LightParameters;
-	RectLightSceneInfo.Proxy->GetLightShaderParameters(LightParameters);
-	FIntVector TextureSize = LightParameters.SourceTexture->GetSizeXYZ();
-	uint32 MipLevelCount = FMath::Min(FMath::CeilLogTwo(TextureSize.X), FMath::CeilLogTwo(TextureSize.Y));
-	RectLightMipTreeDimensions = FIntVector(1 << MipLevelCount, 1 << MipLevelCount, 1);
-	uint32 NumElements = RectLightMipTreeDimensions.X * RectLightMipTreeDimensions.Y;
-	for (uint32 MipLevel = 1; MipLevel <= MipLevelCount; ++MipLevel)
-	{
-		uint32 NumElementsInLevel = (RectLightMipTreeDimensions.X >> MipLevel) * (RectLightMipTreeDimensions.Y >> MipLevel);
-		NumElements += NumElementsInLevel;
-	}
-
-	RectLightMipTree.Initialize(sizeof(float), NumElements, PF_R32_FLOAT, BUF_UnorderedAccess | BUF_ShaderResource);
-
-	// Execute hierarchical build
-	for (uint32 MipLevel = 0; MipLevel <= MipLevelCount; ++MipLevel)
-	{
-		FComputeFenceRHIRef MipLevelFence = RHICmdList.CreateComputeFence(TEXT("RectLightMipTree Build"));
-		BuildRectLightMipTreeComputeShader->SetParameters(RHICmdList, LightParameters.SourceTexture, RectLightMipTreeDimensions, MipLevel, RectLightMipTree);
-		FIntVector MipLevelDimensions = FIntVector(RectLightMipTreeDimensions.X >> MipLevel, RectLightMipTreeDimensions.Y >> MipLevel, 1);
-		FIntVector NumGroups = FIntVector::DivideAndRoundUp(MipLevelDimensions, FBuildRectLightMipTreeCS::GetGroupSize());
-		DispatchComputeShader(RHICmdList, *BuildRectLightMipTreeComputeShader, NumGroups.X, NumGroups.Y, 1);
-		BuildRectLightMipTreeComputeShader->UnsetParameters(RHICmdList, EResourceTransitionAccess::ERWBarrier, EResourceTransitionPipeline::EComputeToCompute, RectLightMipTree, MipLevelFence);
-	}
-	FComputeFenceRHIRef TransitionFence = RHICmdList.CreateComputeFence(TEXT("RectLightMipTree Transition"));
-	BuildRectLightMipTreeComputeShader->UnsetParameters(RHICmdList, EResourceTransitionAccess::ERWBarrier, EResourceTransitionPipeline::EComputeToCompute, RectLightMipTree, TransitionFence);
-}
-
 class FVisualizeRectLightMipTreePS : public FGlobalShader
 {
 	DECLARE_SHADER_TYPE(FVisualizeRectLightMipTreePS, Global);
