All ray tracing shaders now use ShouldCompileRayTracingShadersForProject() to decide if they should be compiled. This fixes cooking issues for incompatible platforms.

diff --git a/Engine/Source/Runtime/Engine/Private/Components/RectLightComponent.cpp b/Engine/Source/Runtime/Engine/Private/Components/RectLightComponent.cpp
--- a/Engine/Source/Runtime/Engine/Private/Components/RectLightComponent.cpp
+++ b/Engine/Source/Runtime/Engine/Private/Components/RectLightComponent.cpp
@@ -223,7 +223,7 @@ class FBuildRectLightMipTreeCS : public FGlobalShader
 public:
 	static bool ShouldCompilePermutation(const FGlobalShaderPermutationParameters& Parameters)
 	{
-		return RHISupportsComputeShaders(Parameters.Platform);
+		return ShouldCompileRayTracingShadersForProject(Parameters.Platform);
 	}
 
 	static uint32 GetGroupSize()
diff --git a/Engine/Source/Runtime/Renderer/Private/PathTracing/PathCompactionCompute.cpp b/Engine/Source/Runtime/Renderer/Private/PathTracing/PathCompactionCompute.cpp
--- a/Engine/Source/Runtime/Renderer/Private/PathTracing/PathCompactionCompute.cpp
+++ b/Engine/Source/Runtime/Renderer/Private/PathTracing/PathCompactionCompute.cpp
@@ -19,7 +19,7 @@ class FPathCompactionCS : public FGlobalShader
 public:
 	static bool ShouldCompilePermutation(const FGlobalShaderPermutationParameters& Parameters)
 	{
-		return RHISupportsComputeShaders(Parameters.Platform);
+		return ShouldCompileRayTracingShadersForProject(Parameters.Platform);
 	}
 
 	static uint32 GetGroupSize()
diff --git a/Engine/Source/Runtime/Renderer/Private/PathTracing/PathTracingTestbed.cpp b/Engine/Source/Runtime/Renderer/Private/PathTracing/PathTracingTestbed.cpp
--- a/Engine/Source/Runtime/Renderer/Private/PathTracing/PathTracingTestbed.cpp
+++ b/Engine/Source/Runtime/Renderer/Private/PathTracing/PathTracingTestbed.cpp
@@ -20,7 +20,7 @@ public:
 
 	static bool ShouldCompilePermutation(const FGlobalShaderPermutationParameters& Parameters)
 	{
-		return RHISupportsComputeShaders(Parameters.Platform) && ShouldCompileRayTracingShadersForProject(Parameters.Platform);
+		return ShouldCompileRayTracingShadersForProject(Parameters.Platform);
 	}
 
 	static uint32 GetGroupSize()
@@ -103,7 +103,7 @@ public:
 
 	static bool ShouldCompilePermutation(const FGlobalShaderPermutationParameters& Parameters)
 	{
-		return RHISupportsComputeShaders(Parameters.Platform) && ShouldCompileRayTracingShadersForProject(Parameters.Platform);
+		return ShouldCompileRayTracingShadersForProject(Parameters.Platform);
 	}
 
 	static uint32 GetGroupSize()
@@ -189,7 +189,7 @@ public:
 
 	static bool ShouldCompilePermutation(const FGlobalShaderPermutationParameters& Parameters)
 	{
-		return RHISupportsComputeShaders(Parameters.Platform) && ShouldCompileRayTracingShadersForProject(Parameters.Platform);
+		return ShouldCompileRayTracingShadersForProject(Parameters.Platform);
 	}
 
 	static uint32 GetGroupSize()
@@ -280,7 +280,7 @@ public:
 
 	static bool ShouldCompilePermutation(const FGlobalShaderPermutationParameters& Parameters)
 	{
-		return RHISupportsComputeShaders(Parameters.Platform) && ShouldCompileRayTracingShadersForProject(Parameters.Platform);
+		return ShouldCompileRayTracingShadersForProject(Parameters.Platform);
 	}
 
 	static uint32 GetGroupSize()
diff --git a/Engine/Source/Runtime/Renderer/Private/PathTracing/RayCounterCompute.cpp b/Engine/Source/Runtime/Renderer/Private/PathTracing/RayCounterCompute.cpp
--- a/Engine/Source/Runtime/Renderer/Private/PathTracing/RayCounterCompute.cpp
+++ b/Engine/Source/Runtime/Renderer/Private/PathTracing/RayCounterCompute.cpp
@@ -21,7 +21,7 @@ class FRayCounterCS : public FGlobalShader
 public:
 	static bool ShouldCompilePermutation(const FGlobalShaderPermutationParameters& Parameters)
 	{
-		return RHISupportsComputeShaders(Parameters.Platform);
+		return ShouldCompileRayTracingShadersForProject(Parameters.Platform);
 	}
 
 	static uint32 GetGroupSize()
diff --git a/Engine/Source/Runtime/Renderer/Private/PathTracing/SkyLightCdfBuild.cpp b/Engine/Source/Runtime/Renderer/Private/PathTracing/SkyLightCdfBuild.cpp
--- a/Engine/Source/Runtime/Renderer/Private/PathTracing/SkyLightCdfBuild.cpp
+++ b/Engine/Source/Runtime/Renderer/Private/PathTracing/SkyLightCdfBuild.cpp
@@ -19,7 +19,7 @@ class FBuildSkyLightRowCdfCS : public FGlobalShader
 public:
 	static bool ShouldCompilePermutation(const FGlobalShaderPermutationParameters& Parameters)
 	{
-		return RHISupportsComputeShaders(Parameters.Platform);
+		return ShouldCompileRayTracingShadersForProject(Parameters.Platform);
 	}
 
 	static uint32 GetGroupSize()
@@ -111,7 +111,7 @@ class FBuildSkyLightColumnCdfCS : public FGlobalShader
 public:
 	static bool ShouldCompilePermutation(const FGlobalShaderPermutationParameters& Parameters)
 	{
-		return RHISupportsComputeShaders(Parameters.Platform);
+		return ShouldCompileRayTracingShadersForProject(Parameters.Platform);
 	}
 
 	static uint32 GetGroupSize()
@@ -194,7 +194,7 @@ class FBuildSkyLightCubeFaceCdfCS : public FGlobalShader
 public:
 	static bool ShouldCompilePermutation(const FGlobalShaderPermutationParameters& Parameters)
 	{
-		return RHISupportsComputeShaders(Parameters.Platform);
+		return ShouldCompileRayTracingShadersForProject(Parameters.Platform);
 	}
 
 	static uint32 GetGroupSize()
diff --git a/Engine/Source/Runtime/Renderer/Private/PathTracing/VarianceMipTreeBuild.cpp b/Engine/Source/Runtime/Renderer/Private/PathTracing/VarianceMipTreeBuild.cpp
--- a/Engine/Source/Runtime/Renderer/Private/PathTracing/VarianceMipTreeBuild.cpp
+++ b/Engine/Source/Runtime/Renderer/Private/PathTracing/VarianceMipTreeBuild.cpp
@@ -19,7 +19,7 @@ class FBuildVarianceMipTreeCS : public FGlobalShader
 public:
 	static bool ShouldCompilePermutation(const FGlobalShaderPermutationParameters& Parameters)
 	{
-		return RHISupportsComputeShaders(Parameters.Platform) && ShouldCompileRayTracingShadersForProject(Parameters.Platform);
+		return ShouldCompileRayTracingShadersForProject(Parameters.Platform);
 	}
 
 	static uint32 GetGroupSize()
@@ -143,12 +143,12 @@ class FVisualizeMipTreePS : public FGlobalShader
 public:
 	static bool ShouldCache(EShaderPlatform Platform)
 	{
-		return IsFeatureLevelSupported(Platform, ERHIFeatureLevel::SM5);
+		return ShouldCompileRayTracingShadersForProject(Platform);
 	}
 
 	static bool ShouldCompilePermutation(const FGlobalShaderPermutationParameters& Parameters)
 	{
-		return IsFeatureLevelSupported(Parameters.Platform, ERHIFeatureLevel::SM5);
+		return ShouldCompileRayTracingShadersForProject(Parameters.Platform);
 	}
 
 	static void ModifyCompilationEnvironment(const FGlobalShaderPermutationParameters& Parameters, FShaderCompilerEnvironment& OutEnvironment)
diff --git a/Engine/Source/Runtime/Renderer/Private/RayTracing/RayTracingAmbientOcclusion.cpp b/Engine/Source/Runtime/Renderer/Private/RayTracing/RayTracingAmbientOcclusion.cpp
--- a/Engine/Source/Runtime/Renderer/Private/RayTracing/RayTracingAmbientOcclusion.cpp
+++ b/Engine/Source/Runtime/Renderer/Private/RayTracing/RayTracingAmbientOcclusion.cpp
@@ -204,7 +204,7 @@ class FCompositeAmbientOcclusionPS : public FGlobalShader
 public:
 	static bool ShouldCompilePermutation(const FGlobalShaderPermutationParameters& Parameters)
 	{
-		return IsFeatureLevelSupported(Parameters.Platform, ERHIFeatureLevel::SM5);
+		return ShouldCompileRayTracingShadersForProject(Parameters.Platform);
 	}
 
 	FCompositeAmbientOcclusionPS() {}
diff --git a/Engine/Source/Runtime/Renderer/Private/RayTracing/RayTracingShadowsForRectLight.cpp b/Engine/Source/Runtime/Renderer/Private/RayTracing/RayTracingShadowsForRectLight.cpp
--- a/Engine/Source/Runtime/Renderer/Private/RayTracing/RayTracingShadowsForRectLight.cpp
+++ b/Engine/Source/Runtime/Renderer/Private/RayTracing/RayTracingShadowsForRectLight.cpp
@@ -184,12 +184,12 @@ class FVisualizeRectLightMipTreePS : public FGlobalShader
 public:
 	static bool ShouldCache(EShaderPlatform Platform)
 	{
-		return IsFeatureLevelSupported(Platform, ERHIFeatureLevel::SM5);
+		return ShouldCompileRayTracingShadersForProject(Platform);
 	}
 
 	static bool ShouldCompilePermutation(const FGlobalShaderPermutationParameters& Parameters)
 	{
-		return IsFeatureLevelSupported(Parameters.Platform, ERHIFeatureLevel::SM5);
+		return ShouldCompileRayTracingShadersForProject(Parameters.Platform);
 	}
 
 	static void ModifyCompilationEnvironment(const FGlobalShaderPermutationParameters& Parameters, FShaderCompilerEnvironment& OutEnvironment)
diff --git a/Engine/Source/Runtime/Renderer/Private/RayTracing/RaytracingSkylight.cpp b/Engine/Source/Runtime/Renderer/Private/RayTracing/RaytracingSkylight.cpp
--- a/Engine/Source/Runtime/Renderer/Private/RayTracing/RaytracingSkylight.cpp
+++ b/Engine/Source/Runtime/Renderer/Private/RayTracing/RaytracingSkylight.cpp
@@ -61,7 +61,7 @@ class FSkylightRG : public FGlobalShader
 
 	static bool ShouldCompilePermutation(const FGlobalShaderPermutationParameters& Parameters)
 	{
-		return IsRayTracingEnabledForThisProject(Parameters.Platform, false);
+		return ShouldCompileRayTracingShadersForProject(Parameters.Platform);
 	}
 
 	BEGIN_SHADER_PARAMETER_STRUCT(FParameters, )
@@ -257,7 +257,7 @@ class FBuildMipTreeCS : public FGlobalShader
 public:
 	static bool ShouldCompilePermutation(const FGlobalShaderPermutationParameters& Parameters)
 	{
-		return RHISupportsComputeShaders(Parameters.Platform);
+		return ShouldCompileRayTracingShadersForProject(Parameters.Platform);
 	}
 
 	static uint32 GetGroupSize()
@@ -405,7 +405,7 @@ class FBuildSolidAnglePdfCS : public FGlobalShader
 public:
 	static bool ShouldCompilePermutation(const FGlobalShaderPermutationParameters& Parameters)
 	{
-		return RHISupportsComputeShaders(Parameters.Platform);
+		return ShouldCompileRayTracingShadersForProject(Parameters.Platform);
 	}
 
 	static uint32 GetGroupSize()
@@ -510,7 +510,7 @@ class FBuildMipTreePdfCS : public FGlobalShader
 public:
 	static bool ShouldCompilePermutation(const FGlobalShaderPermutationParameters& Parameters)
 	{
-		return RHISupportsComputeShaders(Parameters.Platform);
+		return ShouldCompileRayTracingShadersForProject(Parameters.Platform);
 	}
 
 	static uint32 GetGroupSize()
@@ -657,7 +657,7 @@ public:
 
 	static bool ShouldCompilePermutation(const FGlobalShaderPermutationParameters& Parameters)
 	{
-		return IsFeatureLevelSupported(Parameters.Platform, ERHIFeatureLevel::SM5);
+		return ShouldCompileRayTracingShadersForProject(Parameters.Platform);
 	}
 
 	static void ModifyCompilationEnvironment(const FGlobalShaderPermutationParameters& Parameters, FShaderCompilerEnvironment& OutEnvironment)
@@ -899,12 +899,12 @@ class FCompositeSkyLightPS : public FGlobalShader
 public:
 	static bool ShouldCache(EShaderPlatform Platform)
 	{
-		return IsFeatureLevelSupported(Platform, ERHIFeatureLevel::SM5);
+		return ShouldCompileRayTracingShadersForProject(Platform);
 	}
 
 	static bool ShouldCompilePermutation(const FGlobalShaderPermutationParameters& Parameters)
 	{
-		return IsFeatureLevelSupported(Parameters.Platform, ERHIFeatureLevel::SM5);
+		return ShouldCompileRayTracingShadersForProject(Parameters.Platform);
 	}
 
 	static void ModifyCompilationEnvironment(const FGlobalShaderPermutationParameters& Parameters, FShaderCompilerEnvironment& OutEnvironment)
