Fixed artifacts caused in ray traced reflections due incorrect setting of FalloffExponent (it should be 0 when IsInverseSquared() is true)

diff --git a/Engine/Source/Runtime/Renderer/Private/RayTracing/RayTracingReflections.cpp b/Engine/Source/Runtime/Renderer/Private/RayTracing/RayTracingReflections.cpp
--- a/Engine/Source/Runtime/Renderer/Private/RayTracing/RayTracingReflections.cpp
+++ b/Engine/Source/Runtime/Renderer/Private/RayTracing/RayTracingReflections.cpp
@@ -89,6 +89,11 @@ void SetupReflectionsLightData(
 		FLightShaderParameters LightParameters;
 		Light.LightSceneInfo->Proxy->GetLightShaderParameters(LightParameters);
 
+		if (Light.LightSceneInfo->Proxy->IsInverseSquared())
+		{
+			LightParameters.FalloffExponent = 0;
+		}
+
 		LightData->Type[LightData->Count] = Light.LightType;
 		LightData->LightPosition[LightData->Count] = LightParameters.Position;
 		LightData->LightInvRadius[LightData->Count] = LightParameters.InvRadius;
