Skip PrimitiveIdBufferEmulation building when Mesh Draw Command pipeline is enabled.  Saves 0.923ms in replay.

diff --git a/Engine/Source/Runtime/Renderer/Private/GPUScene.cpp b/Engine/Source/Runtime/Renderer/Private/GPUScene.cpp
--- a/Engine/Source/Runtime/Renderer/Private/GPUScene.cpp
+++ b/Engine/Source/Runtime/Renderer/Private/GPUScene.cpp
@@ -155,6 +155,8 @@ void UpdateGPUScene(FRHICommandList& RHICmdList, FScene& Scene)
 {
 	if (UseGPUScene(GMaxRHIShaderPlatform, Scene.GetFeatureLevel()))
 	{
+		QUICK_SCOPE_CYCLE_COUNTER(STAT_UpdateGPUScene);
+
 		if (GGPUSceneUploadEveryFrame || Scene.GPUScene.bUpdateAllPrimitives)
 		{
 			Scene.GPUScene.PrimitivesToUpdate.Reset();
@@ -300,6 +302,8 @@ void UploadDynamicPrimitiveShaderDataForView(FRHICommandList& RHICmdList, FScene
 {
 	if (UseGPUScene(GMaxRHIShaderPlatform, Scene.GetFeatureLevel()))
 	{
+		QUICK_SCOPE_CYCLE_COUNTER(STAT_UploadDynamicPrimitiveShaderDataForView);
+
 		FRWBufferStructured& ViewPrimitiveShaderDataBuffer = View.ViewState ? View.ViewState->PrimitiveShaderDataBuffer : View.OneFramePrimitiveShaderDataBuffer;
 
 		const int32 NumPrimitiveEntries = Scene.Primitives.Num() + View.DynamicPrimitiveShaderData.Num();
@@ -344,6 +348,7 @@ void UploadDynamicPrimitiveShaderDataForView(FRHICommandList& RHICmdList, FScene
 		View.CachedViewUniformShaderParameters->PrimitiveSceneData = ViewPrimitiveShaderDataBuffer.SRV;
 		View.ViewUniformBuffer.UpdateUniformBufferImmediate(*View.CachedViewUniformShaderParameters);
 
+		if (!UseMeshDrawCommandPipeline())
 		{
 			QUICK_SCOPE_CYCLE_COUNTER(STAT_PrimitiveIdBufferEmulation);
 
