Moved mesh draw command sorting and draw call merging to a separate task, which is chained to the actual draw task. This allows to move all this work from the rendering thread.
*FParallelMeshDrawCommandPass stored on View controls two parallel tasks and their data - sorting and merging task, and drawing task. It is initialized based on temporary visible mesh draw command lists inside InitViews. It is also the only way to submit mesh draw commands in parallel.
*Raytraycing mesh draw commands are now stored separately on the view.
*Local shadows now can be sorted in parallel too.
*Removed mobile SortMeshDrawCommands function override, which was required to merge NoCSM/CSM base pass lists. Instead mobile base pass is skipped inside InitViews and is sorted later - inside SortMobileBasePassAfterShadowInit.

diff --git a/Engine/Source/Runtime/RenderCore/Private/RenderCore.cpp b/Engine/Source/Runtime/RenderCore/Private/RenderCore.cpp
--- a/Engine/Source/Runtime/RenderCore/Private/RenderCore.cpp
+++ b/Engine/Source/Runtime/RenderCore/Private/RenderCore.cpp
@@ -108,7 +108,7 @@ DEFINE_STAT(STAT_CreateWholeSceneProjectedShadow);
 DEFINE_STAT(STAT_AddViewDependentWholeSceneShadowsForView);
 DEFINE_STAT(STAT_SetupInteractionShadows);
 DEFINE_STAT(STAT_GetDynamicMeshElements);
-DEFINE_STAT(STAT_ProcessVisibleMeshDrawCommands);
+DEFINE_STAT(STAT_SortAndMergeMeshDrawCommands);
 DEFINE_STAT(STAT_UpdateStaticMeshesTime);
 DEFINE_STAT(STAT_StaticRelevance);
 DEFINE_STAT(STAT_ViewRelevance);
diff --git a/Engine/Source/Runtime/RenderCore/Public/RenderCore.h b/Engine/Source/Runtime/RenderCore/Public/RenderCore.h
--- a/Engine/Source/Runtime/RenderCore/Public/RenderCore.h
+++ b/Engine/Source/Runtime/RenderCore/Public/RenderCore.h
@@ -59,7 +59,7 @@ DECLARE_CYCLE_STAT_EXTERN(TEXT("Compute View Relevance"),STAT_ComputeViewRelevan
 DECLARE_CYCLE_STAT_EXTERN(TEXT("Static Mesh Relevance"),STAT_StaticRelevance,STATGROUP_InitViews, RENDERCORE_API);
 DECLARE_CYCLE_STAT_EXTERN(TEXT("UpdateStaticMeshes"),STAT_UpdateStaticMeshesTime,STATGROUP_InitViews, RENDERCORE_API);
 DECLARE_CYCLE_STAT_EXTERN(TEXT("GetDynamicMeshElements"),STAT_GetDynamicMeshElements,STATGROUP_InitViews, RENDERCORE_API);
-DECLARE_CYCLE_STAT_EXTERN(TEXT("ProcessVisibleMeshDrawCommands"), STAT_ProcessVisibleMeshDrawCommands, STATGROUP_InitViews, RENDERCORE_API);
+DECLARE_CYCLE_STAT_EXTERN(TEXT("SortAndMergeVisibleMeshDrawCommands"), STAT_SortAndMergeMeshDrawCommands, STATGROUP_InitViews, RENDERCORE_API);
 DECLARE_CYCLE_STAT_EXTERN(TEXT("Init dynamic shadows"),STAT_InitDynamicShadowsTime,STATGROUP_InitViews, RENDERCORE_API);
 DECLARE_CYCLE_STAT_EXTERN(TEXT("Update Preshadow Cache"),STAT_UpdatePreshadowCache,STATGROUP_InitViews, RENDERCORE_API);
 DECLARE_CYCLE_STAT_EXTERN(TEXT("Init Projected Shadow"),STAT_InitProjectedShadowVisibility,STATGROUP_InitViews, RENDERCORE_API);
diff --git a/Engine/Source/Runtime/Renderer/Private/BasePassRendering.cpp b/Engine/Source/Runtime/Renderer/Private/BasePassRendering.cpp
--- a/Engine/Source/Runtime/Renderer/Private/BasePassRendering.cpp
+++ b/Engine/Source/Runtime/Renderer/Private/BasePassRendering.cpp
@@ -809,7 +809,7 @@ void FDeferredShadingSceneRenderer::RenderBasePassViewParallel(FViewInfo& View,
 		BasePassDepthStencilAccess,
 		InDrawRenderState);
 
-	SubmitMeshDrawCommandsForView(View, EMeshPass::BasePass, &ParallelSet, ParentCmdList);
+	View.ParallelMeshDrawCommandPasses[EMeshPass::BasePass].DispatchDraw(&ParallelSet, ParentCmdList);
 }
 
 bool HasEditorPrimitivesForDPG(const FViewInfo& View, ESceneDepthPriorityGroup DepthPriorityGroup)
@@ -933,7 +933,7 @@ bool FDeferredShadingSceneRenderer::RenderBasePassView(FRHICommandListImmediate&
 	FDrawingPolicyRenderState DrawRenderState(InDrawRenderState);
 	SetupBasePassView(RHICmdList, View, this);
 
-	SubmitMeshDrawCommandsForView(View, EMeshPass::BasePass, nullptr, RHICmdList);
+	View.ParallelMeshDrawCommandPasses[EMeshPass::BasePass].DispatchDraw(nullptr, RHICmdList);
 
 	return bDirty;
 }
diff --git a/Engine/Source/Runtime/Renderer/Private/DebugViewModeRendering.cpp b/Engine/Source/Runtime/Renderer/Private/DebugViewModeRendering.cpp
--- a/Engine/Source/Runtime/Renderer/Private/DebugViewModeRendering.cpp
+++ b/Engine/Source/Runtime/Renderer/Private/DebugViewModeRendering.cpp
@@ -264,7 +264,7 @@ bool FDeferredShadingSceneRenderer::RenderDebugViewMode(FRHICommandListImmediate
 		{
 			SCOPED_DRAW_EVENT(RHICmdList, Dynamic);
 
-			SubmitMeshDrawCommandsForView(View, EMeshPass::DebugViewMode, nullptr, RHICmdList);
+			View.ParallelMeshDrawCommandPasses[EMeshPass::DebugViewMode].DispatchDraw(nullptr, RHICmdList);
 		}
 	}
 
diff --git a/Engine/Source/Runtime/Renderer/Private/DeferredShadingRenderer.cpp b/Engine/Source/Runtime/Renderer/Private/DeferredShadingRenderer.cpp
--- a/Engine/Source/Runtime/Renderer/Private/DeferredShadingRenderer.cpp
+++ b/Engine/Source/Runtime/Renderer/Private/DeferredShadingRenderer.cpp
@@ -608,7 +608,7 @@ bool FDeferredShadingSceneRenderer::GatherRayTracingWorldInstances(FRHICommandLi
 				ViewFamily.GetFeatureLevel());
 
 			View.DynamicMeshDrawCommandStorage[EMeshPass::RayTracing].MeshDrawCommands.Reserve(Scene->Primitives.Num());
-			View.VisibleMeshDrawCommands[EMeshPass::RayTracing].Reserve(Scene->Primitives.Num());
+			View.RaytraycingVisibleMeshDrawCommands.Reserve(Scene->Primitives.Num());
 		}
 
 
@@ -693,7 +693,7 @@ bool FDeferredShadingSceneRenderer::GatherRayTracingWorldInstances(FRHICommandLi
 									CachedMeshDrawCommand.MeshCullMode,
 									FMeshDrawCommandSortKey::Default);
 
-								View.VisibleMeshDrawCommands[EMeshPass::RayTracing].Add(NewVisibleMeshDrawCommand);
+								View.RaytraycingVisibleMeshDrawCommands.Add(NewVisibleMeshDrawCommand);
 							}
 							else
 							{
@@ -708,7 +708,7 @@ bool FDeferredShadingSceneRenderer::GatherRayTracingWorldInstances(FRHICommandLi
 								FDynamicPassMeshDrawListContext DynamicPassMeshDrawListContext
 								(
 									View.DynamicMeshDrawCommandStorage[EMeshPass::RayTracing],
-									View.VisibleMeshDrawCommands[EMeshPass::RayTracing]
+									View.RaytraycingVisibleMeshDrawCommands
 								);
 								FRayTracingMeshProcessor RayTracingMeshProcessor(Scene, View.GetFeatureLevel(), &View, DynamicPassMeshDrawListContext);
 								RayTracingMeshProcessor.AddMeshBatch(StaticMesh, BatchVisibility, SceneProxy);
@@ -718,12 +718,12 @@ bool FDeferredShadingSceneRenderer::GatherRayTracingWorldInstances(FRHICommandLi
 						const int GeometryInstance = View.RayTracingGeometryInstances.Num();
 						int DrawCmdIndexCopy = VisibleDrawCommandStartOffset[ViewIndex];
 						int& DrawCmdIndex = VisibleDrawCommandStartOffset[ViewIndex];
-						while (DrawCmdIndex < View.VisibleMeshDrawCommands[EMeshPass::RayTracing].Num())
+						while (DrawCmdIndex < View.RaytraycingVisibleMeshDrawCommands.Num())
 						{
 							// #dxr_todo we have to remember which geometry instance we associate with the MDCs 
 							// therefore we patch the recently generated ones
 							// this might be better to be handles inside the MeshProcessor
-							View.VisibleMeshDrawCommands[EMeshPass::RayTracing][DrawCmdIndex++].RayTracedInstanceIndex = GeometryInstance;
+							View.RaytraycingVisibleMeshDrawCommands[DrawCmdIndex++].RayTracedInstanceIndex = GeometryInstance;
 						}
 
 						ensure(DrawCmdIndexCopy != DrawCmdIndex);
@@ -751,7 +751,7 @@ bool FDeferredShadingSceneRenderer::GatherRayTracingWorldInstances(FRHICommandLi
 							FDynamicPassMeshDrawListContext DynamicPassMeshDrawListContext
 							(
 								View.DynamicMeshDrawCommandStorage[EMeshPass::RayTracing],
-								View.VisibleMeshDrawCommands[EMeshPass::RayTracing]
+								View.RaytraycingVisibleMeshDrawCommands
 							);
 							FRayTracingMeshProcessor RayTracingMeshProcessor(Scene, View.GetFeatureLevel(), &View, DynamicPassMeshDrawListContext);
 
@@ -768,12 +768,12 @@ bool FDeferredShadingSceneRenderer::GatherRayTracingWorldInstances(FRHICommandLi
 								RayTracingMeshProcessor.AddMeshBatch(*BatchAndRelevance.Mesh, 1, BatchAndRelevance.PrimitiveSceneProxy);
 
 								int& DrawCmdIndex = VisibleDrawCommandStartOffset[ViewIndex];
-								while (DrawCmdIndex < View.VisibleMeshDrawCommands[EMeshPass::RayTracing].Num())
+								while (DrawCmdIndex < View.RaytraycingVisibleMeshDrawCommands.Num())
 								{
 									// #dxr_todo we have to remember which geometry instance we associate with the MDCs 
 									// therefore we patch the recently generated ones
 									// this might be better to be handles inside the MeshProcessor
-									View.VisibleMeshDrawCommands[EMeshPass::RayTracing][DrawCmdIndex++].RayTracedInstanceIndex = GeometryInstance;
+									View.RaytraycingVisibleMeshDrawCommands[DrawCmdIndex++].RayTracedInstanceIndex = GeometryInstance;
 								}
 							}
 
@@ -799,7 +799,7 @@ bool FDeferredShadingSceneRenderer::GatherRayTracingWorldInstances(FRHICommandLi
 							FDynamicPassMeshDrawListContext DynamicPassMeshDrawListContext
 							(
 								View.DynamicMeshDrawCommandStorage[EMeshPass::RayTracing],
-								View.VisibleMeshDrawCommands[EMeshPass::RayTracing]
+								View.RaytraycingVisibleMeshDrawCommands
 							);
 							FRayTracingMeshProcessor RayTracingMeshProcessor(Scene, View.GetFeatureLevel(), &View, DynamicPassMeshDrawListContext);
 
@@ -853,12 +853,12 @@ bool FDeferredShadingSceneRenderer::GatherRayTracingWorldInstances(FRHICommandLi
 								RayTracingMeshProcessor.AddMeshBatch(*BatchAndRelevance.Mesh, 1, BatchAndRelevance.PrimitiveSceneProxy);
 
 								int& DrawCmdIndex = VisibleDrawCommandStartOffset[ViewIndex];
-								while (DrawCmdIndex < View.VisibleMeshDrawCommands[EMeshPass::RayTracing].Num())
+								while (DrawCmdIndex < View.RaytraycingVisibleMeshDrawCommands.Num())
 								{
 									// #dxr_todo we have to remember which geometry instance we associate with the MDCs 
 									// therefore we patch the recently generated ones
 									// this might be better to be handles inside the MeshProcessor
-									View.VisibleMeshDrawCommands[EMeshPass::RayTracing][DrawCmdIndex++].RayTracedInstanceIndex = GeometryInstance + GeometryIndex;
+									View.RaytraycingVisibleMeshDrawCommands[DrawCmdIndex++].RayTracedInstanceIndex = GeometryInstance + GeometryIndex;
 								}
 							}
 						}
@@ -917,7 +917,7 @@ FRHIRayTracingPipelineState* FDeferredShadingSceneRenderer::BindRayTracingPipeli
 
 	if (GEnableRayTracingMaterials)
 	{
-		for (const FVisibleMeshDrawCommand& VisibleMeshDrawCommand : View.VisibleMeshDrawCommands[EMeshPass::RayTracing])
+		for (const FVisibleMeshDrawCommand& VisibleMeshDrawCommand : View.RaytraycingVisibleMeshDrawCommands)
 		{
 			const FMeshDrawCommand& MeshDrawCommand = *VisibleMeshDrawCommand.MeshDrawCommand;
 			MeshDrawCommand.ShaderBindings.SetOnRayTracingStructure(RHICmdList, View.PerViewRayTracingScene.RayTracingSceneRHI, VisibleMeshDrawCommand.RayTracedInstanceIndex, MeshDrawCommand.RayTracedSegmentIndex, PipelineState, MeshDrawCommand.RayTracingMaterialLibraryIndex);
diff --git a/Engine/Source/Runtime/Renderer/Private/DepthRendering.cpp b/Engine/Source/Runtime/Renderer/Private/DepthRendering.cpp
--- a/Engine/Source/Runtime/Renderer/Private/DepthRendering.cpp
+++ b/Engine/Source/Runtime/Renderer/Private/DepthRendering.cpp
@@ -871,7 +871,7 @@ bool FDeferredShadingSceneRenderer::RenderPrePassView(FRHICommandList& RHICmdLis
 
 	if (UseMeshDrawCommandPipeline())
 	{
-		SubmitMeshDrawCommandsForView(View, EMeshPass::DepthPass, nullptr, RHICmdList);
+		View.ParallelMeshDrawCommandPasses[EMeshPass::DepthPass].DispatchDraw(nullptr, RHICmdList);
 
 		return true;
 	}
@@ -990,7 +990,7 @@ bool FDeferredShadingSceneRenderer::RenderPrePassViewParallel(const FViewInfo& V
 
 		if (UseMeshDrawCommandPipeline())
 		{
-			SubmitMeshDrawCommandsForView(View, EMeshPass::DepthPass, &ParallelCommandListSet, ParentCmdList);
+			View.ParallelMeshDrawCommandPasses[EMeshPass::DepthPass].DispatchDraw(&ParallelCommandListSet, ParentCmdList);
 
 			if (bDoPrePre)
 			{
diff --git a/Engine/Source/Runtime/Renderer/Private/DistortionRendering.cpp b/Engine/Source/Runtime/Renderer/Private/DistortionRendering.cpp
--- a/Engine/Source/Runtime/Renderer/Private/DistortionRendering.cpp
+++ b/Engine/Source/Runtime/Renderer/Private/DistortionRendering.cpp
@@ -988,9 +988,9 @@ bool SubmitDistortionMeshDrawCommands(FRHICommandListImmediate& RHICmdList, cons
 	{
 		QUICK_SCOPE_CYCLE_COUNTER(STAT_FDistortionPrimSet_DrawAccumulatedOffsets_Prims);
 
-		SubmitMeshDrawCommandsForView(View, EMeshPass::Distortion, nullptr, RHICmdList);
+		View.ParallelMeshDrawCommandPasses[EMeshPass::Distortion].DispatchDraw(nullptr, RHICmdList);
 
-		bDirty |= View.VisibleMeshDrawCommands[EMeshPass::Distortion].Num() > 0;
+		bDirty |= View.ParallelMeshDrawCommandPasses[EMeshPass::Distortion].HasAnyDraw();
 	}
 
 	return bDirty;
diff --git a/Engine/Source/Runtime/Renderer/Private/LightMapDensityRendering.cpp b/Engine/Source/Runtime/Renderer/Private/LightMapDensityRendering.cpp
--- a/Engine/Source/Runtime/Renderer/Private/LightMapDensityRendering.cpp
+++ b/Engine/Source/Runtime/Renderer/Private/LightMapDensityRendering.cpp
@@ -59,7 +59,8 @@ void SetupLightmapDensityPassUniformBuffer(FRHICommandListImmediate& RHICmdList,
 
 bool FDeferredShadingSceneRenderer::RenderLightMapDensities(FRHICommandListImmediate& RHICmdList)
 {
-	bool bDirty=0;
+	bool bDirty = false;
+
 	if (Scene->GetFeatureLevel() >= ERHIFeatureLevel::SM4)
 	{
 		SCOPED_DRAW_EVENT(RHICmdList, LightMapDensity);
@@ -84,7 +85,9 @@ bool FDeferredShadingSceneRenderer::RenderLightMapDensities(FRHICommandListImmed
 			DrawRenderState.SetDepthStencilState(TStaticDepthStencilState<true,CF_DepthNearOrEqual>::GetRHI());
 			RHICmdList.SetViewport(View.ViewRect.Min.X, View.ViewRect.Min.Y, 0, View.ViewRect.Max.X, View.ViewRect.Max.Y, 1);
 
-			SubmitMeshDrawCommandsForView(View, EMeshPass::LightmapDensity, nullptr, RHICmdList);
+			View.ParallelMeshDrawCommandPasses[EMeshPass::LightmapDensity].DispatchDraw(nullptr, RHICmdList);
+
+			bDirty = bDirty && View.ParallelMeshDrawCommandPasses[EMeshPass::LightmapDensity].HasAnyDraw();
 		}
 	}
 
diff --git a/Engine/Source/Runtime/Renderer/Private/MeshPassProcessor.cpp b/Engine/Source/Runtime/Renderer/Private/MeshPassProcessor.cpp
--- a/Engine/Source/Runtime/Renderer/Private/MeshPassProcessor.cpp
+++ b/Engine/Source/Runtime/Renderer/Private/MeshPassProcessor.cpp
@@ -791,6 +791,16 @@ void FMeshDrawCommand::SubmitDraw(
 	}
 }
 
+void SubmitMeshDrawCommands(
+	const FMeshCommandOneFrameArray& VisibleMeshDrawCommands,
+	FVertexBufferRHIParamRef PrimitiveIdsBuffer,
+	int32 BasePrimitiveIdsOffset,
+	bool bDynamicInstancing,
+	FRHICommandList& RHICmdList)
+{
+	SubmitMeshDrawCommandsRange(VisibleMeshDrawCommands, PrimitiveIdsBuffer, BasePrimitiveIdsOffset, bDynamicInstancing, 0, VisibleMeshDrawCommands.Num(), RHICmdList);
+}
+
 void SubmitMeshDrawCommandsRange(
 	const FMeshCommandOneFrameArray& VisibleMeshDrawCommands,
 	FVertexBufferRHIParamRef PrimitiveIdsBuffer,
@@ -822,7 +832,7 @@ void DrawDynamicMeshPassPrivate(
 		FGlobalDynamicVertexBuffer DynamicVertexBuffer;
 		FGlobalDynamicVertexBuffer::FAllocation PrimitiveIdBufferAllocation;
 
-		SortPassMeshDrawCommands(View.GetFeatureLevel(), VisibleMeshDrawCommands, DynamicMeshDrawCommandStorage, DynamicVertexBuffer, PrimitiveIdBufferAllocation, false);
+		SortPassMeshDrawCommands(View.GetFeatureLevel(), VisibleMeshDrawCommands, DynamicMeshDrawCommandStorage, DynamicVertexBuffer, PrimitiveIdBufferAllocation);
 
 		DynamicVertexBuffer.Commit();
 
@@ -830,113 +840,10 @@ void DrawDynamicMeshPassPrivate(
 		const int32 BasePrimitiveIdsOffset = PrimitiveIdBufferAllocation.VertexOffset;
 		const bool bDynamicInstancing = IsDynamicInstancingEnabled() && UseGPUScene(GMaxRHIShaderPlatform, View.GetFeatureLevel());
 
-		SubmitMeshDrawCommands(VisibleMeshDrawCommands, PrimitiveIdsBuffer, BasePrimitiveIdsOffset, bDynamicInstancing, nullptr, RHICmdList);
-	}
-}
-
-class FDrawVisibleMeshCommandsAnyThreadTask : public FRenderTask
-{
-	FRHICommandList& RHICmdList;
-	const FMeshCommandOneFrameArray& VisibleMeshDrawCommands;
-	FVertexBufferRHIParamRef PrimitiveIdsBuffer;
-	int32 BasePrimitiveIdsOffset;
-	bool bDynamicInstancing;
-	int32 StartIndex;
-	int32 NumMeshDrawCommands;
-
-public:
-
-	FDrawVisibleMeshCommandsAnyThreadTask(
-		FRHICommandList& InRHICmdList,
-		const FMeshCommandOneFrameArray& InVisibleMeshDrawCommands,
-		FVertexBufferRHIParamRef InPrimitiveIdsBuffer,
-		int32 InBasePrimitiveIdsOffset,
-		bool bInDynamicInstancing,
-		int32 InStartIndex,
-		int32 InNumMeshDrawCommands
-		)
-		: RHICmdList(InRHICmdList)
-		, VisibleMeshDrawCommands(InVisibleMeshDrawCommands)
-		, PrimitiveIdsBuffer(InPrimitiveIdsBuffer)
-		, BasePrimitiveIdsOffset(InBasePrimitiveIdsOffset)
-		, bDynamicInstancing(bInDynamicInstancing)
-		, StartIndex(InStartIndex)
-		, NumMeshDrawCommands(InNumMeshDrawCommands)
-	{}
-
-	FORCEINLINE TStatId GetStatId() const
-	{
-		RETURN_QUICK_DECLARE_CYCLE_STAT(FDrawVisibleMeshCommandsAnyThreadTask, STATGROUP_TaskGraphTasks);
-	}
-
-	static ESubsequentsMode::Type GetSubsequentsMode() { return ESubsequentsMode::TrackSubsequents; }
-
-	void DoTask(ENamedThreads::Type CurrentThread, const FGraphEventRef& MyCompletionGraphEvent)
-	{
-		checkSlow(RHICmdList.IsInsideRenderPass());
-		SubmitMeshDrawCommandsRange(VisibleMeshDrawCommands, PrimitiveIdsBuffer, BasePrimitiveIdsOffset, bDynamicInstancing, StartIndex, NumMeshDrawCommands, RHICmdList);
-		RHICmdList.EndRenderPass();
-		RHICmdList.HandleRTThreadTaskCompletion(MyCompletionGraphEvent);
-	}
-};
-
-void SubmitMeshDrawCommands(
-	const FMeshCommandOneFrameArray& VisibleMeshDrawCommands,
-	FVertexBufferRHIParamRef PrimitiveIdsBuffer,
-	int32 BasePrimitiveIdsOffset,
-	bool bDynamicInstancing,
-	FParallelCommandListSet* ParallelCommandListSet,
-	FRHICommandList& RHICmdList)
-{
-	if (ParallelCommandListSet)
-	{
-		ENamedThreads::Type RenderThread = ENamedThreads::GetRenderThread();
-		const int32 NumMeshDrawCommands = VisibleMeshDrawCommands.Num();
-
-		if (NumMeshDrawCommands > 0)
-		{
-			// Our work is uniform, distribute evenly to the available task graph workers
-			const int32 NumTasks = FMath::Min<int32>(FTaskGraphInterface::Get().GetNumWorkerThreads(), FMath::DivideAndRoundUp(NumMeshDrawCommands, ParallelCommandListSet->MinDrawsPerCommandList));
-			const int32 NumDrawsPerTask = FMath::DivideAndRoundUp(NumMeshDrawCommands, NumTasks);
-
-			for (int32 TaskIndex = 0; TaskIndex < NumTasks; TaskIndex++)
-			{
-				const int32 StartIndex = TaskIndex * NumDrawsPerTask;
-				const int32 NumDraws = FMath::Min(NumDrawsPerTask, NumMeshDrawCommands - StartIndex);
-				checkSlow(NumDraws > 0);
-
-				FRHICommandList* CmdList = ParallelCommandListSet->NewParallelCommandList();
-
-				FGraphEventRef AnyThreadCompletionEvent = TGraphTask<FDrawVisibleMeshCommandsAnyThreadTask>::CreateTask(ParallelCommandListSet->GetPrereqs(), RenderThread)
-					.ConstructAndDispatchWhenReady(*CmdList, VisibleMeshDrawCommands, PrimitiveIdsBuffer, BasePrimitiveIdsOffset, bDynamicInstancing, StartIndex, NumDraws);
-				ParallelCommandListSet->AddParallelCommandList(CmdList, AnyThreadCompletionEvent, NumDraws);
-			}
-		}
-	}
-	else
-	{
 		SubmitMeshDrawCommandsRange(VisibleMeshDrawCommands, PrimitiveIdsBuffer, BasePrimitiveIdsOffset, bDynamicInstancing, 0, VisibleMeshDrawCommands.Num(), RHICmdList);
 	}
 }
 
-void SubmitMeshDrawCommandsForView(
-	const FViewInfo& View,
-	EMeshPass::Type PassType,
-	FParallelCommandListSet* ParallelCommandListSet,
-	FRHICommandList& RHICmdList)
-{
-	const FMeshCommandOneFrameArray& VisibleMeshDrawCommands = View.VisibleMeshDrawCommands[PassType];
-
-	if (VisibleMeshDrawCommands.Num() > 0)
-	{
-		FVertexBufferRHIParamRef PrimitiveIdsBuffer = UseGPUScene(GMaxRHIShaderPlatform, View.GetFeatureLevel()) ? View.VisibleMeshDrawCommandPrimitiveIdBuffers[PassType].VertexBuffer->VertexBufferRHI : nullptr;
-		const int32 BasePrimitiveIdsOffset = View.VisibleMeshDrawCommandPrimitiveIdBuffers[PassType].VertexOffset;
-		const bool bDynamicInstancing = IsDynamicInstancingEnabled() && UseGPUScene(GMaxRHIShaderPlatform, View.GetFeatureLevel());
-
-		SubmitMeshDrawCommands(VisibleMeshDrawCommands, PrimitiveIdsBuffer, BasePrimitiveIdsOffset, bDynamicInstancing, ParallelCommandListSet, RHICmdList);
-	}
-}
-
 FMeshPassProcessor::FMeshPassProcessor(const FScene* InScene, ERHIFeatureLevel::Type InFeatureLevel, const FSceneView* InViewIfDynamicMeshCommand, FMeshPassDrawListContext& InDrawListContext) 
 	: Scene(InScene)
 	, FeatureLevel(InFeatureLevel)
diff --git a/Engine/Source/Runtime/Renderer/Private/MobileBasePassRendering.cpp b/Engine/Source/Runtime/Renderer/Private/MobileBasePassRendering.cpp
--- a/Engine/Source/Runtime/Renderer/Private/MobileBasePassRendering.cpp
+++ b/Engine/Source/Runtime/Renderer/Private/MobileBasePassRendering.cpp
@@ -1007,8 +1007,7 @@ void FMobileSceneRenderer::RenderMobileBasePass(FRHICommandListImmediate& RHICmd
 			}
 		
 			RHICmdList.SetViewport(View.ViewRect.Min.X, View.ViewRect.Min.Y, 0, View.ViewRect.Max.X, View.ViewRect.Max.Y, 1);
-			SubmitMeshDrawCommandsForView(View, EMeshPass::BasePass, nullptr, RHICmdList);
-
+			View.ParallelMeshDrawCommandPasses[EMeshPass::BasePass].DispatchDraw(nullptr, RHICmdList);
 
 			{
 				FDrawingPolicyRenderState DrawRenderState(View, Scene->UniformBuffers.MobileOpaqueBasePassUniformBuffer);
diff --git a/Engine/Source/Runtime/Renderer/Private/MobileDistortionPass.cpp b/Engine/Source/Runtime/Renderer/Private/MobileDistortionPass.cpp
--- a/Engine/Source/Runtime/Renderer/Private/MobileDistortionPass.cpp
+++ b/Engine/Source/Runtime/Renderer/Private/MobileDistortionPass.cpp
@@ -19,7 +19,7 @@ bool IsMobileDistortionActive(const FViewInfo& View)
 	const EMobileHDRMode HDRMode = GetMobileHDRMode();
 
 	const bool bVisiblePrims = UseMeshDrawCommandPipeline() ?
-		View.VisibleMeshDrawCommands[EMeshPass::Distortion].Num() > 0 : View.DistortionPrimSet.NumPrims() > 0;
+		View.ParallelMeshDrawCommandPasses[EMeshPass::Distortion].HasAnyDraw() : View.DistortionPrimSet.NumPrims() > 0;
 
 	return
 		HDRMode == EMobileHDRMode::EnabledFloat16 &&
@@ -52,7 +52,7 @@ void FRCDistortionAccumulatePassES2::Process(FRenderingCompositePassContext& Con
 				Scene->UniformBuffers.MobileDistortionPassUniformBuffer.UpdateUniformBufferImmediate(Parameters);
 			}
 
-			SubmitMeshDrawCommandsForView(View, EMeshPass::Distortion, nullptr, Context.RHICmdList);
+			View.ParallelMeshDrawCommandPasses[EMeshPass::Distortion].DispatchDraw(nullptr, Context.RHICmdList);
 		}
 		else
 		{
diff --git a/Engine/Source/Runtime/Renderer/Private/MobileSeparateTranslucencyPass.cpp b/Engine/Source/Runtime/Renderer/Private/MobileSeparateTranslucencyPass.cpp
--- a/Engine/Source/Runtime/Renderer/Private/MobileSeparateTranslucencyPass.cpp
+++ b/Engine/Source/Runtime/Renderer/Private/MobileSeparateTranslucencyPass.cpp
@@ -1,7 +1,7 @@
 // Copyright 1998-2019 Epic Games, Inc. All Rights Reserved.
 
 /*=============================================================================
-MobileSeparateTranslucencyPass.cpp - Mobile specific separate translucensy pass
+MobileSeparateTranslucencyPass.cpp - Mobile specific separate translucency pass
 =============================================================================*/
 
 #include "MobileSeparateTranslucencyPass.h"
@@ -12,7 +12,7 @@ bool IsMobileSeparateTranslucencyActive(const FViewInfo& View)
 {
 	if (UseMeshDrawCommandPipeline())
 	{
-		return View.VisibleMeshDrawCommands[EMeshPass::TranslucencyAfterDOF].Num() > 0;
+		return View.ParallelMeshDrawCommandPasses[EMeshPass::TranslucencyAfterDOF].HasAnyDraw();
 	}
 	else
 	{
@@ -39,7 +39,7 @@ void FRCSeparateTranslucensyPassES2::Process(FRenderingCompositePassContext& Con
 
 		if (UseMeshDrawCommandPipeline())
 		{
-			SubmitMeshDrawCommandsForView(View, TranslucencyPassToMeshPass(ETranslucencyPass::TPT_TranslucencyAfterDOF), nullptr, Context.RHICmdList);
+			View.ParallelMeshDrawCommandPasses[EMeshPass::TranslucencyAfterDOF].DispatchDraw(nullptr, Context.RHICmdList);
 		}
 		else
 		{
diff --git a/Engine/Source/Runtime/Renderer/Private/MobileShadingRenderer.cpp b/Engine/Source/Runtime/Renderer/Private/MobileShadingRenderer.cpp
--- a/Engine/Source/Runtime/Renderer/Private/MobileShadingRenderer.cpp
+++ b/Engine/Source/Runtime/Renderer/Private/MobileShadingRenderer.cpp
@@ -113,29 +113,27 @@ TUniformBufferRef<FMobileDirectionalLightShaderParameters>& GetNullMobileDirecti
 	return NullLightParams->UniformBufferRHI;
 }
 
-void FMobileSceneRenderer::SortMeshDrawCommands()
-{
-	// ignoring sort here, mobile needs to sort after shadow init is performed.
-	//SortMobileMeshDrawCommands();
-}
-
-void FMobileSceneRenderer::SortMobileMeshDrawCommands()
+void FMobileSceneRenderer::SortMobileBasePassAfterShadowInit(FMeshDrawCommandsPerPassPerView& VisibleCommandsPerView)
 {
 	// Sort front to back on all platforms, even HSR benefits from it
 	//const bool bWantsFrontToBackSorting = (GHardwareHiddenSurfaceRemoval == false);
 
-	// compute keys for front to back sorting
+	// compute keys for front to back sorting and dispatch sort.
 	for (int32 ViewIndex = 0; ViewIndex < Views.Num(); ++ViewIndex)
 	{
 		FViewInfo& View = Views[ViewIndex];
+		FMeshDrawCommandsPerPass& VisibleCommands = VisibleCommandsPerView[ViewIndex];
 		
-		FMeshCommandOneFrameArray& BasePassVisibleCommands = View.VisibleMeshDrawCommands[EMeshPass::BasePass];
+		FMeshCommandOneFrameArray& BasePassVisibleCommands = VisibleCommands[EMeshPass::BasePass];
 		MobileBasePass::ComputeBasePassSortKeys(*Scene, View, BasePassVisibleCommands);
 
-		UpdateTranslucentMeshSortKeys(Scene, View, ETranslucencyPass::TPT_StandardTranslucency, View.VisibleMeshDrawCommands[EMeshPass::MobileInverseOpacity]);
+		// Run sorting on BasePass, as it's ignored inside FSceneRenderer::SortMeshDrawCommands, so it can be done after shadow init on mobile.
+		if (VisibleCommands[EMeshPass::BasePass].Num() > 0)
+		{
+			FParallelMeshDrawCommandPass& Pass = View.ParallelMeshDrawCommandPasses[EMeshPass::BasePass];
+			Pass.DispatchSortAndMerge(View, Scene->PrimitiveBounds, EMeshPass::BasePass, VisibleCommands[EMeshPass::BasePass]);
+		}
 	}
-
-	FSceneRenderer::SortMeshDrawCommands();
 }
 
 /**
@@ -151,8 +149,11 @@ void FMobileSceneRenderer::InitViews(FRHICommandListImmediate& RHICmdList)
 	SCOPE_CYCLE_COUNTER(STAT_InitViewsTime);
 
 	FILCUpdatePrimTaskData ILCTaskData;
+	FMeshDrawCommandsPerPassPerView VisibleCommandsPerView;
+	VisibleCommandsPerView.SetNum(Views.Num());
+
 	PreVisibilityFrameSetup(RHICmdList);
-	ComputeViewVisibility(RHICmdList, FExclusiveDepthStencil::DepthWrite_StencilWrite);
+	ComputeViewVisibility(RHICmdList, FExclusiveDepthStencil::DepthWrite_StencilWrite, VisibleCommandsPerView);
 	PostVisibilityFrameSetup(ILCTaskData);
 
 	const bool bDynamicShadows = ViewFamily.EngineShowFlags.DynamicShadows;
@@ -160,10 +161,10 @@ void FMobileSceneRenderer::InitViews(FRHICommandListImmediate& RHICmdList)
 	if (bDynamicShadows && !IsSimpleForwardShadingEnabled(ShaderPlatform))
 	{
 		// Setup dynamic shadows.
-		InitDynamicShadows(RHICmdList);		
+		InitDynamicShadows(RHICmdList, VisibleCommandsPerView);
 	}
 
-	SortMobileMeshDrawCommands();
+	SortMobileBasePassAfterShadowInit(VisibleCommandsPerView);
 
 	// if we kicked off ILC update via task, wait and finalize.
 	if (ILCTaskData.TaskRef.IsValid())
diff --git a/Engine/Source/Runtime/Renderer/Private/MobileTranslucentRendering.cpp b/Engine/Source/Runtime/Renderer/Private/MobileTranslucentRendering.cpp
--- a/Engine/Source/Runtime/Renderer/Private/MobileTranslucentRendering.cpp
+++ b/Engine/Source/Runtime/Renderer/Private/MobileTranslucentRendering.cpp
@@ -370,7 +370,8 @@ void FMobileSceneRenderer::RenderTranslucency(FRHICommandListImmediate& RHICmdLi
 					UpdateDirectionalLightUniformBuffers(RHICmdList, View);
 				}
 		
-				SubmitMeshDrawCommandsForView(View, TranslucencyPassToMeshPass(TranslucencyPass), nullptr, RHICmdList);
+				const EMeshPass::Type MeshPass = TranslucencyPassToMeshPass(TranslucencyPass);
+				View.ParallelMeshDrawCommandPasses[MeshPass].DispatchDraw(nullptr, RHICmdList);
 			}
 			else
 			{
@@ -747,9 +748,9 @@ bool FMobileSceneRenderer::RenderInverseOpacity(FRHICommandListImmediate& RHICmd
 					UpdateDirectionalLightUniformBuffers(RHICmdList, View);
 				}
 		
-				SubmitMeshDrawCommandsForView(View, EMeshPass::MobileInverseOpacity, nullptr, RHICmdList);
+				View.ParallelMeshDrawCommandPasses[EMeshPass::MobileInverseOpacity].DispatchDraw(nullptr, RHICmdList);
 				
-				bDirty |= View.VisibleMeshDrawCommands[EMeshPass::MobileInverseOpacity].Num() > 0;
+				bDirty |= View.ParallelMeshDrawCommandPasses[EMeshPass::MobileInverseOpacity].HasAnyDraw();
 			}
 			else
 			{
diff --git a/Engine/Source/Runtime/Renderer/Private/PostProcess/PostProcessSelectionOutline.cpp b/Engine/Source/Runtime/Renderer/Private/PostProcess/PostProcessSelectionOutline.cpp
--- a/Engine/Source/Runtime/Renderer/Private/PostProcess/PostProcessSelectionOutline.cpp
+++ b/Engine/Source/Runtime/Renderer/Private/PostProcess/PostProcessSelectionOutline.cpp
@@ -109,7 +109,7 @@ void FRCPassPostProcessSelectionOutlineColor::Process(FRenderingCompositePassCon
 			EditorView.ViewUniformBuffer = Scene->UniformBuffers.ViewUniformBuffer;
 
 			// Run selection pass on static elements
-			SubmitMeshDrawCommandsForView(View, EMeshPass::EditorSelection, nullptr, Context.RHICmdList);
+			View.ParallelMeshDrawCommandPasses[EMeshPass::EditorSelection].DispatchDraw(nullptr, Context.RHICmdList);
 
 			// to get an outline around the objects if it's partly outside of the screen
 			{
diff --git a/Engine/Source/Runtime/Renderer/Private/SceneHitProxyRendering.cpp b/Engine/Source/Runtime/Renderer/Private/SceneHitProxyRendering.cpp
--- a/Engine/Source/Runtime/Renderer/Private/SceneHitProxyRendering.cpp
+++ b/Engine/Source/Runtime/Renderer/Private/SceneHitProxyRendering.cpp
@@ -313,11 +313,11 @@ static void DoRenderHitProxies(FRHICommandListImmediate& RHICmdList, const FScen
 		// Adjust the visibility map for this view
 		if (!View.bAllowTranslucentPrimitivesInHitProxy)
 		{
-			SubmitMeshDrawCommandsForView(View, EMeshPass::HitProxyOpaqueOnly, nullptr, RHICmdList);
+			View.ParallelMeshDrawCommandPasses[EMeshPass::HitProxyOpaqueOnly].DispatchDraw(nullptr, RHICmdList);
 		}
 		else
 		{
-			SubmitMeshDrawCommandsForView(View, EMeshPass::HitProxy, nullptr, RHICmdList);
+			View.ParallelMeshDrawCommandPasses[EMeshPass::HitProxy].DispatchDraw(nullptr, RHICmdList);
 		}
 
 		DrawDynamicMeshPass(View, RHICmdList,
diff --git a/Engine/Source/Runtime/Renderer/Private/SceneRendering.cpp b/Engine/Source/Runtime/Renderer/Private/SceneRendering.cpp
--- a/Engine/Source/Runtime/Renderer/Private/SceneRendering.cpp
+++ b/Engine/Source/Runtime/Renderer/Private/SceneRendering.cpp
@@ -100,9 +100,9 @@ static TAutoConsoleVariable<int32> CVarMeshDrawCommandsDynamicInstancing(
 	ECVF_RenderThreadSafe);
 
 static TAutoConsoleVariable<int32> CVarMeshDrawCommandsParallelProcess(
-	TEXT("r.MeshDrawCommands.ParallelProcess"),
+	TEXT("r.MeshDrawCommands.ParallelSortAndMerge"),
 	1,
-	TEXT("Whether to process (sort and setup primitive id buffer for instancing) visible mesh draw commands in parallel."),
+	TEXT("Whether to sort and merge mesh visible mesh draw commands in parallel."),
 	ECVF_RenderThreadSafe);
 
 bool IsDynamicInstancingEnabled()
@@ -1567,6 +1567,9 @@ FViewInfo* FViewInfo::CreateSnapshot() const
 	FReadBuffer NullReadBuffer;
 	FMemory::Memcpy(Result->OneFramePrimitiveIdBufferEmulation, NullReadBuffer);
 
+	TStaticArray<FParallelMeshDrawCommandPass, EMeshPass::Num> NullParallelMeshDrawCommandPasses;
+	FMemory::Memcpy(Result->ParallelMeshDrawCommandPasses, NullParallelMeshDrawCommandPasses);
+
 	Result->bIsSnapshot = true;
 	ViewInfoSnapshots.Add(Result);
 	return Result;
@@ -2671,70 +2674,64 @@ void FSceneRenderer::RenderFinish(FRHICommandListImmediate& RHICmdList)
 	RHICmdList.EndScene();
 }
 
-void FSceneRenderer::ApplyViewOverridesToMeshDrawCommands(FExclusiveDepthStencil::Type BasePassDepthStencilAccess)
+void FSceneRenderer::ApplyViewOverridesToMeshDrawCommands(FExclusiveDepthStencil::Type BasePassDepthStencilAccess, FViewInfo& View, FMeshDrawCommandsPerPass& VisibleCommands)
 {
 	QUICK_SCOPE_CYCLE_COUNTER(STAT_FSceneRenderer_ApplyViewOverridesToMeshDrawCommands);
 
 	EShadingPath ShadingPath = Scene->GetShadingPath();
 	FExclusiveDepthStencil::Type DefaultBasePassDepthStencilAccess = GetDefaultBasePassDepthStencilAccess(ShadingPath);
 
-	for (int32 ViewIndex = 0; ViewIndex < Views.Num(); ViewIndex++)
+	for (int32 PassIndex = 0; PassIndex < EMeshPass::Num; PassIndex++)
 	{
-		FViewInfo& View = Views[ViewIndex];
-		EDrawingPolicyOverrideFlags ViewOverrideFlags = EDrawingPolicyOverrideFlags::None;
+		EMeshPass::Type PassType = (EMeshPass::Type)PassIndex;
 
-		for (int32 PassIndex = 0; PassIndex < EMeshPass::Num; PassIndex++)
+		if ((FPassProcessorManager::GetPassFlags(ShadingPath, PassType) & EMeshPassFlags::MainView) != EMeshPassFlags::None)
 		{
-			EMeshPass::Type PassType = (EMeshPass::Type)PassIndex;
-
-			if ((FPassProcessorManager::GetPassFlags(ShadingPath, PassType) & EMeshPassFlags::MainView) != EMeshPassFlags::None)
+			if (View.bReverseCulling || View.bRenderSceneTwoSided || (BasePassDepthStencilAccess != DefaultBasePassDepthStencilAccess && PassType == EMeshPass::BasePass))
 			{
-				if (View.bReverseCulling || View.bRenderSceneTwoSided || (BasePassDepthStencilAccess != DefaultBasePassDepthStencilAccess && PassType == EMeshPass::BasePass))
-				{
-					const FMeshCommandOneFrameArray& PassVisibleMeshDrawCommands = View.VisibleMeshDrawCommands[PassIndex];
-
-					FMeshCommandOneFrameArray ViewOverriddenMeshCommands;
-					ViewOverriddenMeshCommands.Empty(PassVisibleMeshDrawCommands.Num());
+				const FMeshCommandOneFrameArray& PassVisibleMeshDrawCommands = VisibleCommands[PassIndex];
 
-					FDynamicMeshDrawCommandStorage& MeshDrawCommandStorage = View.DynamicMeshDrawCommandStorage[PassIndex];
+				FMeshCommandOneFrameArray ViewOverriddenMeshCommands;
+				ViewOverriddenMeshCommands.Empty(PassVisibleMeshDrawCommands.Num());
 
-					for (int32 MeshCommandIndex = 0; MeshCommandIndex < PassVisibleMeshDrawCommands.Num(); MeshCommandIndex++)
-					{
-						MeshDrawCommandStorage.MeshDrawCommands.Add(1);
-						FMeshDrawCommand& NewMeshCommand = MeshDrawCommandStorage.MeshDrawCommands[MeshDrawCommandStorage.MeshDrawCommands.Num() - 1];
+				FDynamicMeshDrawCommandStorage& MeshDrawCommandStorage = View.DynamicMeshDrawCommandStorage[PassIndex];
 
-						const FVisibleMeshDrawCommand& VisibleMeshDrawCommand = PassVisibleMeshDrawCommands[MeshCommandIndex];
-						const FMeshDrawCommand& MeshCommand = *VisibleMeshDrawCommand.MeshDrawCommand;
-						NewMeshCommand = MeshCommand;
+				for (int32 MeshCommandIndex = 0; MeshCommandIndex < PassVisibleMeshDrawCommands.Num(); MeshCommandIndex++)
+				{
+					MeshDrawCommandStorage.MeshDrawCommands.Add(1);
+					FMeshDrawCommand& NewMeshCommand = MeshDrawCommandStorage.MeshDrawCommands[MeshDrawCommandStorage.MeshDrawCommands.Num() - 1];
 
-						const ERasterizerCullMode LocalCullMode = View.bRenderSceneTwoSided ? CM_None : View.bReverseCulling ? FMeshDrawingPolicy::InverseCullMode(VisibleMeshDrawCommand.MeshCullMode) : VisibleMeshDrawCommand.MeshCullMode;
-						NewMeshCommand.PipelineState.RasterizerState = GetStaticRasterizerState<true>(VisibleMeshDrawCommand.MeshFillMode, LocalCullMode);
+					const FVisibleMeshDrawCommand& VisibleMeshDrawCommand = PassVisibleMeshDrawCommands[MeshCommandIndex];
+					const FMeshDrawCommand& MeshCommand = *VisibleMeshDrawCommand.MeshDrawCommand;
+					NewMeshCommand = MeshCommand;
 
-						if (BasePassDepthStencilAccess != DefaultBasePassDepthStencilAccess && PassType == EMeshPass::BasePass)
-						{
-							FDrawingPolicyRenderState PassDrawRenderState;
-							SetupBasePassState(BasePassDepthStencilAccess, false, PassDrawRenderState);
-							NewMeshCommand.PipelineState.DepthStencilState = PassDrawRenderState.GetDepthStencilState();
-						}
+					const ERasterizerCullMode LocalCullMode = View.bRenderSceneTwoSided ? CM_None : View.bReverseCulling ? FMeshDrawingPolicy::InverseCullMode(VisibleMeshDrawCommand.MeshCullMode) : VisibleMeshDrawCommand.MeshCullMode;
+					NewMeshCommand.PipelineState.RasterizerState = GetStaticRasterizerState<true>(VisibleMeshDrawCommand.MeshFillMode, LocalCullMode);
 
-						NewMeshCommand.Finalize(true);
+					if (BasePassDepthStencilAccess != DefaultBasePassDepthStencilAccess && PassType == EMeshPass::BasePass)
+					{
+						FDrawingPolicyRenderState PassDrawRenderState;
+						SetupBasePassState(BasePassDepthStencilAccess, false, PassDrawRenderState);
+						NewMeshCommand.PipelineState.DepthStencilState = PassDrawRenderState.GetDepthStencilState();
+					}
 
-						FVisibleMeshDrawCommand NewVisibleMeshDrawCommand;
+					NewMeshCommand.Finalize(true);
 
-						NewVisibleMeshDrawCommand.Setup(
-							&NewMeshCommand,
-							VisibleMeshDrawCommand.DrawPrimitiveId,
-							VisibleMeshDrawCommand.StateBucketId,
-							VisibleMeshDrawCommand.MeshFillMode,
-							VisibleMeshDrawCommand.MeshCullMode,
-							VisibleMeshDrawCommand.SortKey);
+					FVisibleMeshDrawCommand NewVisibleMeshDrawCommand;
 
-						ViewOverriddenMeshCommands.Add(NewVisibleMeshDrawCommand);
-					}
+					NewVisibleMeshDrawCommand.Setup(
+						&NewMeshCommand,
+						VisibleMeshDrawCommand.DrawPrimitiveId,
+						VisibleMeshDrawCommand.StateBucketId,
+						VisibleMeshDrawCommand.MeshFillMode,
+						VisibleMeshDrawCommand.MeshCullMode,
+						VisibleMeshDrawCommand.SortKey);
 
-					// Replace VisibleMeshDrawCommands
-					FMemory::Memswap(&View.VisibleMeshDrawCommands[PassIndex], &ViewOverriddenMeshCommands, sizeof(ViewOverriddenMeshCommands));
+					ViewOverriddenMeshCommands.Add(NewVisibleMeshDrawCommand);
 				}
+
+				// Replace VisibleMeshDrawCommands
+				FMemory::Memswap(&VisibleCommands[PassIndex], &ViewOverriddenMeshCommands, sizeof(ViewOverriddenMeshCommands));
 			}
 		}
 	}
@@ -2778,35 +2775,45 @@ uint32 BitInvertIfNegativeFloat(uint32 f)
 /**
 * Update mesh sort keys with view dependent data.
 */
-void UpdateTranslucentMeshSortKeys(const FScene* Scene, const FViewInfo& View, ETranslucencyPass::Type TranslucencyPass, FMeshCommandOneFrameArray& VisibleMeshCommands)
+void UpdateTranslucentMeshSortKeys(
+	ETranslucentSortPolicy::Type TranslucentSortPolicy,
+	const FVector& TranslucentSortAxis,
+	const FVector& ViewOrigin,
+	const FMatrix& ViewMatrix,
+	const TArray<struct FPrimitiveBounds>& PrimitiveBounds,
+	ETranslucencyPass::Type TranslucencyPass, 
+	FMeshCommandOneFrameArray& VisibleMeshCommands
+	)
 {
+	QUICK_SCOPE_CYCLE_COUNTER(STAT_UpdateTranslucentMeshSortKeys);
+
 	for (int32 CommandIndex = 0; CommandIndex < VisibleMeshCommands.Num(); ++CommandIndex)
 	{
 		FVisibleMeshDrawCommand& VisibleCommand = VisibleMeshCommands[CommandIndex];
 
 		FVector BoundsOrigin = FVector::ZeroVector;
-		if (VisibleCommand.DrawPrimitiveId < Scene->PrimitiveBounds.Num())
+		if (VisibleCommand.DrawPrimitiveId < PrimitiveBounds.Num())
 		{
-			BoundsOrigin = Scene->PrimitiveBounds[VisibleCommand.DrawPrimitiveId].BoxSphereBounds.Origin;
+			BoundsOrigin = PrimitiveBounds[VisibleCommand.DrawPrimitiveId].BoxSphereBounds.Origin;
 		}
 
 		float Distance = 0.0f;
-		if (View.TranslucentSortPolicy == ETranslucentSortPolicy::SortByDistance)
+		if (TranslucentSortPolicy == ETranslucentSortPolicy::SortByDistance)
 		{
 			//sort based on distance to the view position, view rotation is not a factor
-			Distance = (BoundsOrigin - View.ViewMatrices.GetViewOrigin()).Size();
+			Distance = (BoundsOrigin - ViewOrigin).Size();
 		}
-		else if (View.TranslucentSortPolicy == ETranslucentSortPolicy::SortAlongAxis)
+		else if (TranslucentSortPolicy == ETranslucentSortPolicy::SortAlongAxis)
 		{
 			// Sort based on enforced orthogonal distance
-			const FVector CameraToObject = BoundsOrigin - View.ViewMatrices.GetViewOrigin();
-			Distance = FVector::DotProduct(CameraToObject, View.TranslucentSortAxis);
+			const FVector CameraToObject = BoundsOrigin - ViewOrigin;
+			Distance = FVector::DotProduct(CameraToObject, TranslucentSortAxis);
 		}
 		else
 		{
 			// Sort based on projected Z distance
-			check(View.TranslucentSortPolicy == ETranslucentSortPolicy::SortByProjectedZ);
-			Distance = View.ViewMatrices.GetViewMatrix().TransformPosition(BoundsOrigin).Z;
+			check(TranslucentSortPolicy == ETranslucentSortPolicy::SortByProjectedZ);
+			Distance = ViewMatrix.TransformPosition(BoundsOrigin).Z;
 		}
 
 		// Patch distance inside translucent mesh sort key.
@@ -2818,197 +2825,195 @@ void UpdateTranslucentMeshSortKeys(const FScene* Scene, const FViewInfo& View, E
 }
 
 /**
- * Packet for parallel mesh draw command sorting and setting up primitive id buffer for instancing.
+ * Build mesh draw command primitive Id buffer for instancing.
+ * NewPassVisibleMeshDrawCommands must be presized for NewPassVisibleMeshDrawCommands.
  */
-class FVisibleMeshDrawCommandProcessPacket
+void BuildMeshDrawCommandPrimitiveIdBuffer(
+	bool bUseGPUScene, 
+	FMeshCommandOneFrameArray& VisibleMeshDrawCommands,
+	FDynamicMeshDrawCommandStorage& MeshDrawCommandStorage,
+	FGlobalDynamicVertexBuffer::FAllocation& PrimitiveIdBuffer,
+	FMeshCommandOneFrameArray& NewPassVisibleMeshDrawCommands,
+	int32& MaxInstances,
+	int32& VisibleMeshDrawCommandsNum,
+	int32& NewPassVisibleMeshDrawCommandsNum
+	)
 {
-public:
-	FVisibleMeshDrawCommandProcessPacket(const FScene* InScene,
-		const FViewInfo* InView,
-		bool InbUseGPUScene, 
-		EMeshPass::Type InPassType,
-		ETranslucencyPass::Type InTranslucencyPass, 
-		FMeshCommandOneFrameArray& InVisibleMeshDrawCommands, 
-		FDynamicMeshDrawCommandStorage& InMeshDrawCommandStorage,
-		FGlobalDynamicVertexBuffer& InDynamicVertexBuffer,
-		FGlobalDynamicVertexBuffer::FAllocation& InPrimitiveIdBuffer)
-		: Scene(InScene)
-		, View(InView)
-		, bUseGPUScene(InbUseGPUScene)
-		, PassType(InPassType)
-		, TranslucencyPass(InTranslucencyPass)
-		, VisibleMeshDrawCommands(InVisibleMeshDrawCommands)
-		, MeshDrawCommandStorage(InMeshDrawCommandStorage)
-		, PrimitiveIdBuffer(InPrimitiveIdBuffer)
-		, VisibleMeshDrawCommandsNum(0)
-		, NewPassVisibleMeshDrawCommandsNum(0)
-		, MaxInstances(1)
+	if (!bUseGPUScene)
 	{
-		// Preallocate memory and resources on rendering thread.
+		return;
+	}
 
-		const int32 NumDrawCommands = VisibleMeshDrawCommands.Num();
 
-		if (bUseGPUScene && NumDrawCommands > 0)
-		{
-			PrimitiveIdBuffer = InDynamicVertexBuffer.Allocate(NumDrawCommands * sizeof(int32));
-			int32* RESTRICT PrimitiveIds = (int32*)PrimitiveIdBuffer.Buffer;
-		}
+	check(VisibleMeshDrawCommands.Num() <= NewPassVisibleMeshDrawCommands.Max())
 
-		if (IsDynamicInstancingEnabled())
-		{
-			NewPassVisibleMeshDrawCommands.Empty(NumDrawCommands);
-		}
-	}
+	const FVisibleMeshDrawCommand* RESTRICT PassVisibleMeshDrawCommands = VisibleMeshDrawCommands.GetData();
+	const int32 NumDrawCommands = VisibleMeshDrawCommands.Num();
 
-	void Sort()
-	{
-		QUICK_SCOPE_CYCLE_COUNTER(STAT_SortVisibleMeshDrawCommands);
-		VisibleMeshDrawCommands.Sort(FCompareFMeshDrawCommands());
-	}
+	int32* RESTRICT PrimitiveIds = (int32*)PrimitiveIdBuffer.Buffer;
 
-	void BuildPrimitiveIdBuffer()
+	if (IsDynamicInstancingEnabled())
 	{
-		if (!bUseGPUScene)
-		{
-			return;
-		}
-		
-		const FVisibleMeshDrawCommand* RESTRICT PassVisibleMeshDrawCommands = VisibleMeshDrawCommands.GetData();
-		const int32 NumDrawCommands = VisibleMeshDrawCommands.Num();
+		QUICK_SCOPE_CYCLE_COUNTER(STAT_DynamicInstancingOfVisibleMeshDrawCommands);
 
-		int32* RESTRICT PrimitiveIds = (int32*)PrimitiveIdBuffer.Buffer;
+		int32 CurrentStateBucketId = -1;
+		uint32* RESTRICT CurrentDynamicallyInstancedMeshCommandNumInstances = nullptr;
+		MaxInstances = 1;
 
-		if (IsDynamicInstancingEnabled())
+		for (int32 DrawCommandIndex = 0; DrawCommandIndex < NumDrawCommands; DrawCommandIndex++)
 		{
-			QUICK_SCOPE_CYCLE_COUNTER(STAT_DynamicInstancingOfVisibleMeshDrawCommands);
-
-			int32 CurrentStateBucketId = -1;
-			uint32* RESTRICT CurrentDynamicallyInstancedMeshCommandNumInstances = nullptr;
-			MaxInstances = 1;
+			const FVisibleMeshDrawCommand& RESTRICT VisibleMeshDrawCommand = PassVisibleMeshDrawCommands[DrawCommandIndex];
 
-			for (int32 DrawCommandIndex = 0; DrawCommandIndex < NumDrawCommands; DrawCommandIndex++)
+			if (VisibleMeshDrawCommand.StateBucketId == CurrentStateBucketId && VisibleMeshDrawCommand.StateBucketId != -1)
 			{
-				const FVisibleMeshDrawCommand& RESTRICT VisibleMeshDrawCommand = PassVisibleMeshDrawCommands[DrawCommandIndex];
-
-				if (VisibleMeshDrawCommand.StateBucketId == CurrentStateBucketId && VisibleMeshDrawCommand.StateBucketId != -1)
+				if (CurrentDynamicallyInstancedMeshCommandNumInstances)
 				{
-					if (CurrentDynamicallyInstancedMeshCommandNumInstances)
-					{
-						const int32 CurrentNumInstances = *CurrentDynamicallyInstancedMeshCommandNumInstances;
-						*CurrentDynamicallyInstancedMeshCommandNumInstances = CurrentNumInstances + 1;
-						MaxInstances = FMath::Max(MaxInstances, CurrentNumInstances + 1);
-					}
-					else
-					{
-						FVisibleMeshDrawCommand NewVisibleMeshDrawCommand = VisibleMeshDrawCommand;
-						NewVisibleMeshDrawCommand.PrimitiveIdBufferOffset = DrawCommandIndex;
-						NewPassVisibleMeshDrawCommands.Emplace(MoveTemp(NewVisibleMeshDrawCommand));
-					}
+					const int32 CurrentNumInstances = *CurrentDynamicallyInstancedMeshCommandNumInstances;
+					*CurrentDynamicallyInstancedMeshCommandNumInstances = CurrentNumInstances + 1;
+					MaxInstances = FMath::Max(MaxInstances, CurrentNumInstances + 1);
 				}
 				else
 				{
-					// First time state bucket setup
-					CurrentStateBucketId = VisibleMeshDrawCommand.StateBucketId;
-
-					if (VisibleMeshDrawCommand.MeshDrawCommand->PrimitiveIdStreamIndex >= 0
-						&& VisibleMeshDrawCommand.MeshDrawCommand->NumInstances == 1
-						// Don't create a new FMeshDrawCommand for the last command and make it safe for us to look at the next command
-						&& DrawCommandIndex + 1 < NumDrawCommands
-						// Only create a new FMeshDrawCommand if more than one draw in the state bucket
-						&& CurrentStateBucketId == PassVisibleMeshDrawCommands[DrawCommandIndex + 1].StateBucketId)
-					{
-						const int32 Index = MeshDrawCommandStorage.MeshDrawCommands.AddElement(*VisibleMeshDrawCommand.MeshDrawCommand);
-						FMeshDrawCommand& NewCommand = MeshDrawCommandStorage.MeshDrawCommands[Index];
-						FVisibleMeshDrawCommand NewVisibleMeshDrawCommand;
-
-						NewVisibleMeshDrawCommand.Setup(
-							&NewCommand,
-							VisibleMeshDrawCommand.DrawPrimitiveId,
-							VisibleMeshDrawCommand.StateBucketId,
-							VisibleMeshDrawCommand.MeshFillMode,
-							VisibleMeshDrawCommand.MeshCullMode,
-							VisibleMeshDrawCommand.SortKey);
-
-						NewVisibleMeshDrawCommand.PrimitiveIdBufferOffset = DrawCommandIndex;
-						NewPassVisibleMeshDrawCommands.Emplace(MoveTemp(NewVisibleMeshDrawCommand));
-
-						CurrentDynamicallyInstancedMeshCommandNumInstances = &NewCommand.NumInstances;
-					}
-					else
-					{
-						CurrentDynamicallyInstancedMeshCommandNumInstances = nullptr;
-						FVisibleMeshDrawCommand NewVisibleMeshDrawCommand = VisibleMeshDrawCommand;
-						NewVisibleMeshDrawCommand.PrimitiveIdBufferOffset = DrawCommandIndex;
-						NewPassVisibleMeshDrawCommands.Emplace(MoveTemp(NewVisibleMeshDrawCommand));
-					}
+					FVisibleMeshDrawCommand NewVisibleMeshDrawCommand = VisibleMeshDrawCommand;
+					NewVisibleMeshDrawCommand.PrimitiveIdBufferOffset = DrawCommandIndex;
+					NewPassVisibleMeshDrawCommands.Emplace(MoveTemp(NewVisibleMeshDrawCommand));
 				}
-
-				//@todo - refactor into memcpy
-				PrimitiveIds[DrawCommandIndex] = VisibleMeshDrawCommand.DrawPrimitiveId;
 			}
-
-			// Setup instancing stats for logging.
-			VisibleMeshDrawCommandsNum = VisibleMeshDrawCommands.Num();
-			NewPassVisibleMeshDrawCommandsNum = NewPassVisibleMeshDrawCommands.Num();
-
-			// Replace VisibleMeshDrawCommands
-			FMemory::Memswap(&VisibleMeshDrawCommands, &NewPassVisibleMeshDrawCommands, sizeof(NewPassVisibleMeshDrawCommands));
-		}
-		else
-		{
-			QUICK_SCOPE_CYCLE_COUNTER(STAT_BuildVisibleMeshDrawCommandPrimitiveIdBuffers);
-
-			for (int32 DrawCommandIndex = 0; DrawCommandIndex < NumDrawCommands; DrawCommandIndex++)
+			else
 			{
-				const FVisibleMeshDrawCommand& VisibleMeshDrawCommand = VisibleMeshDrawCommands[DrawCommandIndex];
-				PrimitiveIds[DrawCommandIndex] = VisibleMeshDrawCommand.DrawPrimitiveId;
+				// First time state bucket setup
+				CurrentStateBucketId = VisibleMeshDrawCommand.StateBucketId;
+
+				if (VisibleMeshDrawCommand.MeshDrawCommand->PrimitiveIdStreamIndex >= 0
+					&& VisibleMeshDrawCommand.MeshDrawCommand->NumInstances == 1
+					// Don't create a new FMeshDrawCommand for the last command and make it safe for us to look at the next command
+					&& DrawCommandIndex + 1 < NumDrawCommands
+					// Only create a new FMeshDrawCommand if more than one draw in the state bucket
+					&& CurrentStateBucketId == PassVisibleMeshDrawCommands[DrawCommandIndex + 1].StateBucketId)
+				{
+					const int32 Index = MeshDrawCommandStorage.MeshDrawCommands.AddElement(*VisibleMeshDrawCommand.MeshDrawCommand);
+					FMeshDrawCommand& NewCommand = MeshDrawCommandStorage.MeshDrawCommands[Index];
+					FVisibleMeshDrawCommand NewVisibleMeshDrawCommand;
+
+					NewVisibleMeshDrawCommand.Setup(
+						&NewCommand,
+						VisibleMeshDrawCommand.DrawPrimitiveId,
+						VisibleMeshDrawCommand.StateBucketId,
+						VisibleMeshDrawCommand.MeshFillMode,
+						VisibleMeshDrawCommand.MeshCullMode,
+						VisibleMeshDrawCommand.SortKey);
+
+					NewVisibleMeshDrawCommand.PrimitiveIdBufferOffset = DrawCommandIndex;
+					NewPassVisibleMeshDrawCommands.Emplace(MoveTemp(NewVisibleMeshDrawCommand));
+
+					CurrentDynamicallyInstancedMeshCommandNumInstances = &NewCommand.NumInstances;
+				}
+				else
+				{
+					CurrentDynamicallyInstancedMeshCommandNumInstances = nullptr;
+					FVisibleMeshDrawCommand NewVisibleMeshDrawCommand = VisibleMeshDrawCommand;
+					NewVisibleMeshDrawCommand.PrimitiveIdBufferOffset = DrawCommandIndex;
+					NewPassVisibleMeshDrawCommands.Emplace(MoveTemp(NewVisibleMeshDrawCommand));
+				}
 			}
+
+			//@todo - refactor into memcpy
+			PrimitiveIds[DrawCommandIndex] = VisibleMeshDrawCommand.DrawPrimitiveId;
 		}
-	}
 
-	void AnyThreadTask()
+		// Setup instancing stats for logging.
+		VisibleMeshDrawCommandsNum = VisibleMeshDrawCommands.Num();
+		NewPassVisibleMeshDrawCommandsNum = NewPassVisibleMeshDrawCommands.Num();
+
+		// Replace VisibleMeshDrawCommands
+		FMemory::Memswap(&VisibleMeshDrawCommands, &NewPassVisibleMeshDrawCommands, sizeof(NewPassVisibleMeshDrawCommands));
+	}
+	else
 	{
-		check(Scene && View);
+		QUICK_SCOPE_CYCLE_COUNTER(STAT_BuildVisibleMeshDrawCommandPrimitiveIdBuffers);
 
-		if (TranslucencyPass != ETranslucencyPass::TPT_MAX)
+		for (int32 DrawCommandIndex = 0; DrawCommandIndex < NumDrawCommands; DrawCommandIndex++)
 		{
-			QUICK_SCOPE_CYCLE_COUNTER(STAT_UpdateTranslucentMeshSortKeys);
-			UpdateTranslucentMeshSortKeys(Scene, *View, TranslucencyPass, VisibleMeshDrawCommands);
+			const FVisibleMeshDrawCommand& VisibleMeshDrawCommand = VisibleMeshDrawCommands[DrawCommandIndex];
+			PrimitiveIds[DrawCommandIndex] = VisibleMeshDrawCommand.DrawPrimitiveId;
 		}
+	}
+}
+
+FAutoConsoleTaskPriority CPrio_FVisibleMeshDrawCommandProcessTask(
+	TEXT("TaskGraph.TaskPriorities.VisibleMeshDrawCommandProcessTask"),
+	TEXT("Task and thread priority for FVisibleMeshDrawCommandProcessTask."),
+	ENamedThreads::HighThreadPriority, // if we have high priority task threads, then use them...
+	ENamedThreads::NormalTaskPriority, // .. at normal task priority
+	ENamedThreads::HighTaskPriority // if we don't have hi pri threads, then use normal priority threads at high task priority instead
+);
+
+/**
+ * Task for parallel sorting and building primitive id buffer for mesh draw commands.
+ */
+class FVisibleMeshDrawCommandProcessTask
+{
+public:
+	FVisibleMeshDrawCommandProcessTask(FVisibleMeshDrawCommandProcessTaskContext& InContext)
+		: Context(InContext)
+	{
+	}
+
+	FORCEINLINE TStatId GetStatId() const
+	{
+		RETURN_QUICK_DECLARE_CYCLE_STAT(FSortFrontToBackTask, STATGROUP_TaskGraphTasks);
+	}
 
-		Sort();
+	ENamedThreads::Type GetDesiredThread()
+	{
+		return CPrio_FVisibleMeshDrawCommandProcessTask.Get();
+	}
 
-		BuildPrimitiveIdBuffer();
+	static ESubsequentsMode::Type GetSubsequentsMode() 
+	{ 
+		return ESubsequentsMode::TrackSubsequents; 
 	}
 
-	void DumpMeshDrawCommandInstancingStats()
+	void AnyThreadTask()
 	{
-		if (PassType != EMeshPass::Num)
+		if (Context.TranslucencyPass != ETranslucencyPass::TPT_MAX)
+		{
+			UpdateTranslucentMeshSortKeys(
+				Context.TranslucentSortPolicy,
+				Context.TranslucentSortAxis,
+				Context.ViewOrigin,
+				Context.ViewMatrix,
+				*Context.PrimitiveBounds,
+				Context.TranslucencyPass,
+				Context.VisibleMeshDrawCommands
+			);
+		}
+
 		{
-			UE_LOG(LogRenderer, Log, TEXT("%s"), GetMeshPassName(PassType));
+			QUICK_SCOPE_CYCLE_COUNTER(STAT_SortVisibleMeshDrawCommands);
+			Context.VisibleMeshDrawCommands.Sort(FCompareFMeshDrawCommands());
 		}
 
-		UE_LOG(LogRenderer, Log, TEXT("   %i Mesh Draw Commands in %i instancing state buckets"), VisibleMeshDrawCommandsNum, NewPassVisibleMeshDrawCommandsNum);
-		UE_LOG(LogRenderer, Log, TEXT("   Largest %i"), MaxInstances);
-		UE_LOG(LogRenderer, Log, TEXT("   %.1f Dynamic Instancing draw call reduction factor"), VisibleMeshDrawCommandsNum / (float)NewPassVisibleMeshDrawCommandsNum);
+		BuildMeshDrawCommandPrimitiveIdBuffer(
+			Context.bUseGPUScene,
+			Context.VisibleMeshDrawCommands,
+			Context.MeshDrawCommandStorage,
+			Context.PrimitiveIdBuffer,
+			Context.NewPassVisibleMeshDrawCommands,
+			Context.MaxInstances,
+			Context.VisibleMeshDrawCommandsNum,
+			Context.NewPassVisibleMeshDrawCommandsNum
+		);
+	}
+
+	void DoTask(ENamedThreads::Type CurrentThread, const FGraphEventRef& MyCompletionGraphEvent)
+	{
+		AnyThreadTask();
 	}
 
 
 private:
-	const FScene* Scene;
-	const FViewInfo* View;
-	const bool bUseGPUScene;
-	EMeshPass::Type PassType;
-	ETranslucencyPass::Type TranslucencyPass;
-	FMeshCommandOneFrameArray& VisibleMeshDrawCommands;
-	FDynamicMeshDrawCommandStorage& MeshDrawCommandStorage;
-	FGlobalDynamicVertexBuffer::FAllocation& PrimitiveIdBuffer;
-	FMeshCommandOneFrameArray NewPassVisibleMeshDrawCommands;
-
-	// For logging instancing stats.
-	int32 VisibleMeshDrawCommandsNum;
-	int32 NewPassVisibleMeshDrawCommandsNum;
-	int32 MaxInstances;
+	FVisibleMeshDrawCommandProcessTaskContext& Context;
 };
 
 void SortPassMeshDrawCommands(
@@ -3016,97 +3021,270 @@ void SortPassMeshDrawCommands(
 	FMeshCommandOneFrameArray& VisibleMeshDrawCommands,
 	FDynamicMeshDrawCommandStorage& MeshDrawCommandStorage,
 	FGlobalDynamicVertexBuffer& DynamicVertexBuffer,
-	FGlobalDynamicVertexBuffer::FAllocation& PrimitiveIdBuffer,
-	bool bDumpStats)
+	FGlobalDynamicVertexBuffer::FAllocation& PrimitiveIdBuffer)
 {
 	const bool bUseGPUScene = UseGPUScene(GMaxRHIShaderPlatform, FeatureLevel);
 
 	const int32 NumDrawCommands = VisibleMeshDrawCommands.Num();
 	if (NumDrawCommands > 0)
 	{
-		FVisibleMeshDrawCommandProcessPacket Packet(
-			nullptr,
-			nullptr,
+		FMeshCommandOneFrameArray NewPassVisibleMeshDrawCommands;
+		int32 MaxInstances = 1;
+		int32 VisibleMeshDrawCommandsNum = 0;
+		int32 NewPassVisibleMeshDrawCommandsNum = 0;
+
+		// Preallocate context memory and resources.
+		if (bUseGPUScene && NumDrawCommands > 0)
+		{
+			PrimitiveIdBuffer = DynamicVertexBuffer.Allocate(NumDrawCommands * sizeof(int32));
+
+			if (IsDynamicInstancingEnabled())
+			{
+				NewPassVisibleMeshDrawCommands.Empty(NumDrawCommands);
+			}
+		}
+
+		VisibleMeshDrawCommands.Sort(FCompareFMeshDrawCommands());
+
+		BuildMeshDrawCommandPrimitiveIdBuffer(
 			bUseGPUScene,
-			EMeshPass::Num,
-			ETranslucencyPass::TPT_MAX,
 			VisibleMeshDrawCommands,
 			MeshDrawCommandStorage,
-			DynamicVertexBuffer,
-			PrimitiveIdBuffer);
+			PrimitiveIdBuffer,
+			NewPassVisibleMeshDrawCommands,
+			MaxInstances,
+			VisibleMeshDrawCommandsNum,
+			NewPassVisibleMeshDrawCommandsNum
+		);
+	}
+}
+
+void FParallelMeshDrawCommandPass::DispatchSortAndMerge(const FViewInfo& View, const TArray<FPrimitiveBounds>& PrimitiveBounds, EMeshPass::Type MeshPass, FMeshCommandOneFrameArray& InOutVisibleMeshDrawCommands)
+{
+	check(TaskEventRefs.Num() == 0);
+	check(InOutVisibleMeshDrawCommands.Num() > 0);
+
+	MaxDrawNum = InOutVisibleMeshDrawCommands.Num();
 
-		Packet.Sort();
-		Packet.BuildPrimitiveIdBuffer();
+	// Setup translucency sort key update pass based on view.
+	TaskContext.TranslucencyPass = ETranslucencyPass::TPT_MAX;
+	TaskContext.TranslucentSortPolicy = View.TranslucentSortPolicy;
+	TaskContext.TranslucentSortAxis = View.TranslucentSortAxis;
+	TaskContext.ViewOrigin = View.ViewMatrices.GetViewOrigin();
+	TaskContext.ViewMatrix = View.ViewMatrices.GetViewMatrix();
+	TaskContext.PrimitiveBounds = &PrimitiveBounds;
 
-		if (bDumpStats)
+	switch (MeshPass)
+	{
+		case EMeshPass::TranslucencyStandard: TaskContext.TranslucencyPass = ETranslucencyPass::TPT_StandardTranslucency; break;
+		case EMeshPass::TranslucencyAfterDOF: TaskContext.TranslucencyPass = ETranslucencyPass::TPT_TranslucencyAfterDOF; break;
+		case EMeshPass::TranslucencyAll: TaskContext.TranslucencyPass = ETranslucencyPass::TPT_AllTranslucency; break;
+		case EMeshPass::MobileInverseOpacity: TaskContext.TranslucencyPass = ETranslucencyPass::TPT_StandardTranslucency; break;
+	}
+
+	TaskContext.bUseGPUScene = UseGPUScene(GMaxRHIShaderPlatform, View.GetFeatureLevel());
+	TaskContext.bDynamicInstancing = TaskContext.bUseGPUScene && IsDynamicInstancingEnabled();
+
+	FMemory::Memswap(&TaskContext.VisibleMeshDrawCommands, &InOutVisibleMeshDrawCommands, sizeof(InOutVisibleMeshDrawCommands));
+
+	const bool bExecuteInParallel = FApp::ShouldUseThreadingForPerformance() && CVarMeshDrawCommandsParallelProcess.GetValueOnRenderThread() > 0;
+
+	// Preallocate context memory and resources on rendering thread.
+	if (TaskContext.bUseGPUScene && MaxDrawNum > 0)
+	{
+		TaskContext.PrimitiveIdBuffer = FGlobalDynamicVertexBuffer::Get().Allocate(MaxDrawNum * sizeof(int32));
+
+		if (TaskContext.bDynamicInstancing)
 		{
-			Packet.DumpMeshDrawCommandInstancingStats();
+			TaskContext.NewPassVisibleMeshDrawCommands.Empty(MaxDrawNum);
 		}
 	}
+
+	if (bExecuteInParallel)
+	{
+		TaskEventRefs.Add(TGraphTask<FVisibleMeshDrawCommandProcessTask>::CreateTask(
+			nullptr, ENamedThreads::GetRenderThread()).ConstructAndDispatchWhenReady(TaskContext));
+	}
+	else
+	{
+		FVisibleMeshDrawCommandProcessTask Task(TaskContext);
+		Task.AnyThreadTask();
+	}
 }
 
-void FSceneRenderer::SortMeshDrawCommands()
+void FParallelMeshDrawCommandPass::Empty()
 {
-	SCOPE_CYCLE_COUNTER(STAT_ProcessVisibleMeshDrawCommands);
+	// Need to wait in case if someone dispatched sort and draw merge task, but didn't draw it.
+	if (TaskEventRefs.Num() > 0)
+	{
+		FTaskGraphInterface::Get().WaitUntilTasksComplete(TaskEventRefs);
+	}
 
-	const EShadingPath ShadingPath = Scene->GetShadingPath();
-	const bool bUseGPUScene = UseGPUScene(GMaxRHIShaderPlatform, FeatureLevel);
-	const bool bExecuteInParallel = FApp::ShouldUseThreadingForPerformance() && CVarMeshDrawCommandsParallelProcess.GetValueOnRenderThread() > 0;
+	DumpInstancingStats();
+
+	MaxDrawNum = 0;
+	PassNameForStats.Empty();
+	TaskEventRefs.Empty();
+	TaskContext.VisibleMeshDrawCommands.Empty();
+	TaskContext.NewPassVisibleMeshDrawCommands.Empty();
+	TaskContext.PrimitiveIdBuffer = FGlobalDynamicVertexBuffer::FAllocation();
+}
 
-	TArray<FVisibleMeshDrawCommandProcessPacket, SceneRenderingAllocator> Packets;
-	Packets.Reserve(Views.Num() * EMeshPass::Num);
+FParallelMeshDrawCommandPass::~FParallelMeshDrawCommandPass()
+{
+	Empty();
+}
 
-	for (int32 ViewIndex = 0; ViewIndex < Views.Num(); ViewIndex++)
+class FDrawVisibleMeshCommandsAnyThreadTask : public FRenderTask
+{
+	FRHICommandList& RHICmdList;
+	const FMeshCommandOneFrameArray& VisibleMeshDrawCommands;
+	FVertexBufferRHIParamRef PrimitiveIdsBuffer;
+	int32 BasePrimitiveIdsOffset;
+	bool bDynamicInstancing;
+	int32 TaskIndex;
+	int32 TaskNum;
+
+public:
+
+	FDrawVisibleMeshCommandsAnyThreadTask(
+		FRHICommandList& InRHICmdList,
+		const FMeshCommandOneFrameArray& InVisibleMeshDrawCommands,
+		FVertexBufferRHIParamRef InPrimitiveIdsBuffer,
+		int32 InBasePrimitiveIdsOffset,
+		bool bInDynamicInstancing,
+		int32 InTaskIndex,
+		int32 InTaskNum
+	)
+		: RHICmdList(InRHICmdList)
+		, VisibleMeshDrawCommands(InVisibleMeshDrawCommands)
+		, PrimitiveIdsBuffer(InPrimitiveIdsBuffer)
+		, BasePrimitiveIdsOffset(InBasePrimitiveIdsOffset)
+		, bDynamicInstancing(bInDynamicInstancing)
+		, TaskIndex(InTaskIndex)
+		, TaskNum(InTaskNum)
+	{}
+
+	FORCEINLINE TStatId GetStatId() const
 	{
-		FViewInfo& View = Views[ViewIndex];
+		RETURN_QUICK_DECLARE_CYCLE_STAT(FDrawVisibleMeshCommandsAnyThreadTask, STATGROUP_TaskGraphTasks);
+	}
+
+	static ESubsequentsMode::Type GetSubsequentsMode() { return ESubsequentsMode::TrackSubsequents; }
+
+	void DoTask(ENamedThreads::Type CurrentThread, const FGraphEventRef& MyCompletionGraphEvent)
+	{
+		checkSlow(RHICmdList.IsInsideRenderPass());
 
-		for (int32 PassIndex = 0; PassIndex < EMeshPass::Num; PassIndex++)
+		// Recompute draw range.
+		const int32 DrawNum = VisibleMeshDrawCommands.Num();
+		const int32 NumDrawsPerTask = FMath::DivideAndRoundUp(DrawNum, TaskNum);
+		const int32 StartIndex = TaskIndex * NumDrawsPerTask;
+		const int32 NumDraws = FMath::Min(NumDrawsPerTask, DrawNum - StartIndex);
+
+		SubmitMeshDrawCommandsRange(VisibleMeshDrawCommands, PrimitiveIdsBuffer, BasePrimitiveIdsOffset, bDynamicInstancing, StartIndex, NumDraws, RHICmdList);
+
+		RHICmdList.EndRenderPass();
+		RHICmdList.HandleRTThreadTaskCompletion(MyCompletionGraphEvent);
+	}
+};
+
+void FParallelMeshDrawCommandPass::DispatchDraw(FParallelCommandListSet* ParallelCommandListSet, FRHICommandList& RHICmdList) const
+{
+	if (!HasAnyDraw())
+	{
+		return;
+	}
+
+
+	FVertexBufferRHIParamRef PrimitiveIdsBuffer = TaskContext.bUseGPUScene ? TaskContext.PrimitiveIdBuffer.VertexBuffer->VertexBufferRHI : nullptr;
+	const int32 BasePrimitiveIdsOffset = TaskContext.PrimitiveIdBuffer.VertexOffset;
+
+	if (ParallelCommandListSet)
+	{
+		const ENamedThreads::Type RenderThread = ENamedThreads::GetRenderThread();
+
+		FGraphEventArray Prereqs = TaskEventRefs;
+		if (ParallelCommandListSet->GetPrereqs())
 		{
-			EMeshPass::Type PassType = (EMeshPass::Type)PassIndex;
+			Prereqs.Append(*ParallelCommandListSet->GetPrereqs());
+		}
 
-			if ((FPassProcessorManager::GetPassFlags(ShadingPath, PassType) & EMeshPassFlags::MainView) != EMeshPassFlags::None)
-			{
-				ETranslucencyPass::Type TranslucencyPass = ETranslucencyPass::TPT_MAX;
-				switch (PassType)
-				{
-					case EMeshPass::TranslucencyStandard:	TranslucencyPass = ETranslucencyPass::TPT_StandardTranslucency; break;
-					case EMeshPass::TranslucencyAfterDOF:	TranslucencyPass = ETranslucencyPass::TPT_TranslucencyAfterDOF; break;
-					case EMeshPass::TranslucencyAll:		TranslucencyPass = ETranslucencyPass::TPT_AllTranslucency;		break;
-				}
+		// Distribute work evenly to the available task graph workers based on MaxDrawNum.
+		// Every task will then adjust it's working range based on FVisibleMeshDrawCommandProcessTask results.
+		const int32 NumThreads = FMath::Min<int32>(FTaskGraphInterface::Get().GetNumWorkerThreads(), ParallelCommandListSet->Width);
+		const int32 NumTasks = FMath::Min<int32>(NumThreads, FMath::DivideAndRoundUp(MaxDrawNum, ParallelCommandListSet->MinDrawsPerCommandList));
+		const int32 NumDrawsPerTask = FMath::DivideAndRoundUp(MaxDrawNum, NumTasks);
 
-				const int32 NumDrawCommands = View.VisibleMeshDrawCommands[PassIndex].Num();
-				if (NumDrawCommands > 0)
-				{
-					FVisibleMeshDrawCommandProcessPacket Packet(
-						Scene, 
-						&View, 
-						bUseGPUScene, 
-						PassType,
-						TranslucencyPass, 
-						View.VisibleMeshDrawCommands[PassIndex],
-						View.DynamicMeshDrawCommandStorage[PassIndex],
-						FGlobalDynamicVertexBuffer::Get(),
-						View.VisibleMeshDrawCommandPrimitiveIdBuffers[PassIndex]);
-
-					Packets.Emplace(MoveTemp(Packet));
-				}
-			}
+		for (int32 TaskIndex = 0; TaskIndex < NumTasks; TaskIndex++)
+		{
+			const int32 StartIndex = TaskIndex * NumDrawsPerTask;
+			const int32 NumDraws = FMath::Min(NumDrawsPerTask, MaxDrawNum - StartIndex);
+			checkSlow(NumDraws > 0);
+
+			FRHICommandList* CmdList = ParallelCommandListSet->NewParallelCommandList();
+
+			FGraphEventRef AnyThreadCompletionEvent = TGraphTask<FDrawVisibleMeshCommandsAnyThreadTask>::CreateTask(&Prereqs, RenderThread)
+				.ConstructAndDispatchWhenReady(*CmdList, TaskContext.VisibleMeshDrawCommands, PrimitiveIdsBuffer, BasePrimitiveIdsOffset, TaskContext.bDynamicInstancing, TaskIndex, NumTasks);
+			ParallelCommandListSet->AddParallelCommandList(CmdList, AnyThreadCompletionEvent, NumDraws);
 		}
 	}
-
-	ParallelFor(Packets.Num(),
-		[&Packets](int32 Index)
+	else
+	{
+		if (TaskEventRefs.Num() > 0)
 		{
-			Packets[Index].AnyThreadTask();
-		},
-		!bExecuteInParallel
-	);
+			FTaskGraphInterface::Get().WaitUntilTasksComplete(TaskEventRefs, ENamedThreads::GetRenderThread());
+		}
+
+		SubmitMeshDrawCommandsRange(TaskContext.VisibleMeshDrawCommands, PrimitiveIdsBuffer, BasePrimitiveIdsOffset, TaskContext.bDynamicInstancing, 0, TaskContext.VisibleMeshDrawCommands.Num(), RHICmdList);
+	}
+}
+
+void FParallelMeshDrawCommandPass::DumpInstancingStats() const
+{
+	if (!PassNameForStats.IsEmpty())
+	{
+		UE_LOG(LogRenderer, Log, TEXT("Instancing stats for %s"), *PassNameForStats);
+		UE_LOG(LogRenderer, Log, TEXT("   %i Mesh Draw Commands in %i instancing state buckets"), TaskContext.VisibleMeshDrawCommandsNum, TaskContext.NewPassVisibleMeshDrawCommandsNum);
+		UE_LOG(LogRenderer, Log, TEXT("   Largest %i"), TaskContext.MaxInstances);
+		UE_LOG(LogRenderer, Log, TEXT("   %.1f Dynamic Instancing draw call reduction factor"), TaskContext.VisibleMeshDrawCommandsNum / (float)TaskContext.NewPassVisibleMeshDrawCommandsNum);
+	}
+}
+
+void FParallelMeshDrawCommandPass::SetDumpInstancingStats(const FString& InPassNameForStats)
+{
+	PassNameForStats = InPassNameForStats;
+}
+
+void FSceneRenderer::SortAndMergeMeshDrawCommands(FViewInfo& View, FMeshDrawCommandsPerPass& VisibleCommands)
+{
+	SCOPE_CYCLE_COUNTER(STAT_SortAndMergeMeshDrawCommands);
 
-	if (ShouldDumpMeshDrawCommandInstancingStats())
+	const EShadingPath ShadingPath = Scene->GetShadingPath();
+
+	for (int32 PassIndex = 0; PassIndex < EMeshPass::Num; PassIndex++)
 	{
-		for (int32 PacketIndex = 0; PacketIndex < Packets.Num(); ++PacketIndex)
+		const EMeshPass::Type PassType = (EMeshPass::Type)PassIndex;
+
+		if ((FPassProcessorManager::GetPassFlags(ShadingPath, PassType) & EMeshPassFlags::MainView) != EMeshPassFlags::None)
 		{
-			Packets[PacketIndex].DumpMeshDrawCommandInstancingStats();
+			// Mobile: BasePass and MobileBasePassCSM lists need to be merged and sorted after shadow pass.
+			if (ShadingPath == EShadingPath::Mobile && (PassType == EMeshPass::BasePass || PassType == EMeshPass::MobileBasePassCSM))
+			{
+				continue;
+			}
+
+			if (VisibleCommands[PassIndex].Num() > 0)
+			{
+				FParallelMeshDrawCommandPass& Pass = View.ParallelMeshDrawCommandPasses[PassIndex];
+
+				if (ShouldDumpMeshDrawCommandInstancingStats())
+				{
+					Pass.SetDumpInstancingStats(GetMeshPassName(PassType));
+				}
+
+				Pass.DispatchSortAndMerge(View, Scene->PrimitiveBounds, PassType, VisibleCommands[PassIndex]);
+			}
 		}
 	}
 }
@@ -3241,7 +3419,7 @@ void FSceneRenderer::RenderCustomDepthPass(FRHICommandListImmediate& RHICmdList)
 	
 				if (UseMeshDrawCommandPipeline())
 				{
-					SubmitMeshDrawCommandsForView(View, EMeshPass::CustomDepth, nullptr, RHICmdList);
+					View.ParallelMeshDrawCommandPasses[EMeshPass::CustomDepth].DispatchDraw(nullptr, RHICmdList);
 				}
 				else
 				{
@@ -4206,57 +4384,52 @@ FTextureRHIParamRef FSceneRenderer::GetMultiViewSceneColor(const FSceneRenderTar
 	}
 }
 
-void FSceneRenderer::GenerateDynamicMeshDrawCommands()
+void FSceneRenderer::GenerateDynamicMeshDrawCommands(FViewInfo& View, FMeshDrawCommandsPerPass& VisibleCommands)
 {
 	QUICK_SCOPE_CYCLE_COUNTER(STAT_GenerateDynamicMeshDrawCommands);
 
 	EShadingPath ShadingPath = FSceneInterface::GetShadingPath(FeatureLevel);
 
-	for (int32 ViewIndex = 0; ViewIndex < Views.Num(); ViewIndex++)
+	for (int32 PassIndex = 0; PassIndex < EMeshPass::Num; PassIndex++)
 	{
-		FViewInfo& View = Views[ViewIndex];
+		EMeshPass::Type PassType = (EMeshPass::Type)PassIndex;
 
-		for (int32 PassIndex = 0; PassIndex < EMeshPass::Num; PassIndex++)
+		if ((FPassProcessorManager::GetPassFlags(ShadingPath, PassType) & EMeshPassFlags::MainView) != EMeshPassFlags::None)
 		{
-			EMeshPass::Type PassType = (EMeshPass::Type)PassIndex;
+			FDynamicPassMeshDrawListContext DynamicPassMeshDrawListContext
+			(
+				View.DynamicMeshDrawCommandStorage[PassType],
+				VisibleCommands[PassType]
+			);
 
-			if ((FPassProcessorManager::GetPassFlags(ShadingPath, PassType) & EMeshPassFlags::MainView) != EMeshPassFlags::None)
-			{
-				FDynamicPassMeshDrawListContext DynamicPassMeshDrawListContext
-				(
-					View.DynamicMeshDrawCommandStorage[PassType],
-					View.VisibleMeshDrawCommands[PassType]
-				);
+			PassProcessorCreateFunction CreateFunction = FPassProcessorManager::GetCreateFunction(ShadingPath, PassType);
+			FMeshPassProcessor* PassMeshProcessor = CreateFunction(Scene, &View, DynamicPassMeshDrawListContext);
 
-				PassProcessorCreateFunction CreateFunction = FPassProcessorManager::GetCreateFunction(ShadingPath, PassType);
-				FMeshPassProcessor* PassMeshProcessor = CreateFunction(Scene, &View, DynamicPassMeshDrawListContext);
+			if (PassMeshProcessor != nullptr)
+			{
+				const int32 NumDynamicMeshBatches = View.DynamicMeshElements.Num();
 
-				if (PassMeshProcessor != nullptr)
+				for (int32 MeshIndex = 0; MeshIndex < NumDynamicMeshBatches; MeshIndex++)
 				{
-					const int32 NumDynamicMeshBatches = View.DynamicMeshElements.Num();
-
-					for (int32 MeshIndex = 0; MeshIndex < NumDynamicMeshBatches; MeshIndex++)
-					{
-						const FMeshBatchAndRelevance& MeshAndRelevance = View.DynamicMeshElements[MeshIndex];
-						check(!MeshAndRelevance.Mesh->bRequiresPerElementVisibility);
-						const uint64 BatchElementMask = ~0ull;
-
-						PassMeshProcessor->AddMeshBatch(*MeshAndRelevance.Mesh, BatchElementMask, MeshAndRelevance.PrimitiveSceneProxy);
-					}
+					const FMeshBatchAndRelevance& MeshAndRelevance = View.DynamicMeshElements[MeshIndex];
+					check(!MeshAndRelevance.Mesh->bRequiresPerElementVisibility);
+					const uint64 BatchElementMask = ~0ull;
 
-					const int32 NumStaticMeshBatches = View.DynamicMeshCommandBuildRequests[PassType].Num();
+					PassMeshProcessor->AddMeshBatch(*MeshAndRelevance.Mesh, BatchElementMask, MeshAndRelevance.PrimitiveSceneProxy);
+				}
 
-					for (int32 MeshIndex = 0; MeshIndex < NumStaticMeshBatches; MeshIndex++)
-					{
-						const FStaticMesh* StaticMeshBatch = View.DynamicMeshCommandBuildRequests[PassType][MeshIndex];
-						const uint64 BatchElementMask = StaticMeshBatch->bRequiresPerElementVisibility ? View.StaticMeshBatchVisibility[StaticMeshBatch->BatchVisibilityId] : ~0ull;
+				const int32 NumStaticMeshBatches = View.DynamicMeshCommandBuildRequests[PassType].Num();
 
-						PassMeshProcessor->AddMeshBatch(*StaticMeshBatch, BatchElementMask, StaticMeshBatch->PrimitiveSceneInfo->Proxy, StaticMeshBatch->Id);
-					}
+				for (int32 MeshIndex = 0; MeshIndex < NumStaticMeshBatches; MeshIndex++)
+				{
+					const FStaticMesh* StaticMeshBatch = View.DynamicMeshCommandBuildRequests[PassType][MeshIndex];
+					const uint64 BatchElementMask = StaticMeshBatch->bRequiresPerElementVisibility ? View.StaticMeshBatchVisibility[StaticMeshBatch->BatchVisibilityId] : ~0ull;
 
-					PassMeshProcessor->~FMeshPassProcessor();
+					PassMeshProcessor->AddMeshBatch(*StaticMeshBatch, BatchElementMask, StaticMeshBatch->PrimitiveSceneInfo->Proxy, StaticMeshBatch->Id);
 				}
+
+				PassMeshProcessor->~FMeshPassProcessor();
 			}
 		}
 	}
-}
+}
\ No newline at end of file
diff --git a/Engine/Source/Runtime/Renderer/Private/SceneRendering.h b/Engine/Source/Runtime/Renderer/Private/SceneRendering.h
--- a/Engine/Source/Runtime/Renderer/Private/SceneRendering.h
+++ b/Engine/Source/Runtime/Renderer/Private/SceneRendering.h
@@ -32,16 +32,14 @@
 #include "RenderGraph.h"
 #include "MeshPassProcessor.h"
 
+// Forward declarations.
 class FScene;
 class FSceneViewState;
 class FViewInfo;
 struct FILCUpdatePrimTaskData;
-
-template<typename ShaderMetaType> class TShaderMap;
-
-// Forward declarations.
 class FPostprocessContext;
 struct FILCUpdatePrimTaskData;
+template<typename ShaderMetaType> class TShaderMap;
 
 DECLARE_STATS_GROUP(TEXT("Command List Markers"), STATGROUP_CommandListMarkers, STATCAT_Advanced);
 
@@ -879,6 +877,92 @@ struct FPreviousViewInfo
 	}
 };
 
+typedef TStaticArray<FMeshCommandOneFrameArray, EMeshPass::Num> FMeshDrawCommandsPerPass;
+typedef TArray<FMeshDrawCommandsPerPass, TInlineAllocator<4>> FMeshDrawCommandsPerPassPerView;
+
+/**
+ * Parallel mesh draw command processing and drawing context.
+ */
+class FVisibleMeshDrawCommandProcessTaskContext
+{
+public:
+	FVisibleMeshDrawCommandProcessTaskContext()
+		: bUseGPUScene(false)
+		, bDynamicInstancing(false)
+		, PrimitiveBounds(nullptr)
+		, VisibleMeshDrawCommandsNum(0)
+		, NewPassVisibleMeshDrawCommandsNum(0)
+		, MaxInstances(1)
+	{
+	}
+
+	bool bUseGPUScene;
+	bool bDynamicInstancing;
+
+	// Command storage.
+	FMeshCommandOneFrameArray VisibleMeshDrawCommands;
+	FDynamicMeshDrawCommandStorage MeshDrawCommandStorage;
+
+	// Preallocated resources.
+	FGlobalDynamicVertexBuffer::FAllocation PrimitiveIdBuffer;
+	FMeshCommandOneFrameArray NewPassVisibleMeshDrawCommands;
+
+	// For UpdateTranslucentMeshSortKeys.
+	ETranslucencyPass::Type TranslucencyPass;
+	ETranslucentSortPolicy::Type TranslucentSortPolicy;
+	FVector TranslucentSortAxis;
+	FVector ViewOrigin;
+	FMatrix ViewMatrix;
+	const TArray<struct FPrimitiveBounds>* PrimitiveBounds;
+
+	// For logging instancing stats.
+	int32 VisibleMeshDrawCommandsNum;
+	int32 NewPassVisibleMeshDrawCommandsNum;
+	int32 MaxInstances;
+};
+
+/**
+ * Parallel mesh draw command processing and rendering.
+ * DispatchSort - dispatched mesh command processing task.
+ * DispatchDraw - dispatched mesh draw task depending on DispatchSort.
+ */
+class FParallelMeshDrawCommandPass
+{
+public:
+	FParallelMeshDrawCommandPass()
+		: MaxDrawNum(0)
+	{
+	}
+
+	~FParallelMeshDrawCommandPass();
+
+	/**
+	 * Dispatch visible mesh draw command sort and merge task.
+	 */
+	void DispatchSortAndMerge(const FViewInfo& View, const TArray<FPrimitiveBounds>& PrimitiveBounds, EMeshPass::Type MeshPass, FMeshCommandOneFrameArray& InOutVisibleMeshDrawCommands);
+
+	/**
+	 * Dispatch visible mesh draw command draw task.
+	 */
+	void DispatchDraw(FParallelCommandListSet* ParallelCommandListSet, FRHICommandList& RHICmdList) const;
+
+	void Empty();
+
+	void SetDumpInstancingStats(const FString& InPassName);
+
+	bool HasAnyDraw() const { return MaxDrawNum > 0; }
+
+
+private:
+	FVisibleMeshDrawCommandProcessTaskContext TaskContext;
+	FGraphEventArray TaskEventRefs;
+	int32 MaxDrawNum;
+	FString PassNameForStats;
+
+	void DumpInstancingStats() const;
+	void DispatchSortInternal(ERHIFeatureLevel::Type FeatureLevel, FMeshCommandOneFrameArray& InOutVisibleMeshDrawCommands);
+};
+
 /** A FSceneView with additional state used by the scene renderer. */
 class FViewInfo : public FSceneView
 {
@@ -1020,8 +1104,10 @@ public:
 	FRWBufferStructured OneFramePrimitiveShaderDataBuffer;
 	FReadBuffer OneFramePrimitiveIdBufferEmulation;
 
-	TStaticArray<FMeshCommandOneFrameArray, EMeshPass::Num> VisibleMeshDrawCommands;
-	TStaticArray<FGlobalDynamicVertexBuffer::FAllocation, EMeshPass::Num> VisibleMeshDrawCommandPrimitiveIdBuffers;
+	TStaticArray<FParallelMeshDrawCommandPass, EMeshPass::Num> ParallelMeshDrawCommandPasses;
+	
+	FMeshCommandOneFrameArray RaytraycingVisibleMeshDrawCommands;
+
 	TStaticArray<TArray<const FStaticMesh*, SceneRenderingAllocator>, EMeshPass::Num> DynamicMeshCommandBuildRequests;
 
 	// Used by mobile renderer to determine whether static meshes will be rendered with CSM shaders or not.
@@ -1584,11 +1670,12 @@ protected:
 
 	void InitDynamicShadows(FRHICommandListImmediate& RHICmdList);
 
-	void GenerateDynamicMeshDrawCommands();
+	/** Converts each FMeshBatch into a set of FMeshDrawCommands for each relevant mesh pass. */
+	void GenerateDynamicMeshDrawCommands(FViewInfo& View, FMeshDrawCommandsPerPass& VisibleCommands);
 
-	void ApplyViewOverridesToMeshDrawCommands(FExclusiveDepthStencil::Type BasePassDepthStencilAccess);
+	void ApplyViewOverridesToMeshDrawCommands(FExclusiveDepthStencil::Type BasePassDepthStencilAccess, FViewInfo& View, FMeshDrawCommandsPerPass& VisibleCommands);
 
-	virtual void SortMeshDrawCommands();
+	void SortAndMergeMeshDrawCommands(FViewInfo& View, FMeshDrawCommandsPerPass& VisibleCommands);
 
 	bool RenderShadowProjections(FRHICommandListImmediate& RHICmdList, const FLightSceneInfo* LightSceneInfo, IPooledRenderTarget* ScreenShadowMaskTexture, bool bProjectingForForwardShading, bool bMobileModulatedProjections);
 
@@ -1684,7 +1771,7 @@ protected:
 	void PreVisibilityFrameSetup(FRHICommandListImmediate& RHICmdList);
 
 	/** Computes which primitives are visible and relevant for each view. */
-	void ComputeViewVisibility(FRHICommandListImmediate& RHICmdList, FExclusiveDepthStencil::Type BasePassDepthStencilAccess);
+	void ComputeViewVisibility(FRHICommandListImmediate& RHICmdList, FExclusiveDepthStencil::Type BasePassDepthStencilAccess, FMeshDrawCommandsPerPassPerView& VisibleCommandsPerView);
 
 	/** Performs once per frame setup after to visibility determination. */
 	void PostVisibilityFrameSetup(FILCUpdatePrimTaskData& OutILCTaskData);
@@ -1759,7 +1846,7 @@ public:
 
 protected:
 	/** Finds the visible dynamic shadows for each view. */
-	void InitDynamicShadows(FRHICommandListImmediate& RHICmdList);
+	void InitDynamicShadows(FRHICommandListImmediate& RHICmdList, FMeshDrawCommandsPerPassPerView& VisibleCommandsPerView);
 
 	/** Build visibility lists on CSM receivers and non-csm receivers. */
 	void BuildCSMVisibilityState(FLightSceneInfo* LightSceneInfo);
@@ -1799,7 +1886,7 @@ protected:
 	/** Creates uniform buffers with the mobile directional light parameters, for each lighting channel. Called by InitViews */
 	void CreateDirectionalLightUniformBuffers(FViewInfo& View);
 
-	/** Copy scene color from the mobile multi-view render targat array to side by side stereo scene color */
+	/** Copy scene color from the mobile multi-view render target array to side by side stereo scene color */
 	void CopyMobileMultiViewSceneColor(FRHICommandListImmediate& RHICmdList);
 
 	/** Render inverse opacity for the dynamic meshes. */
@@ -1808,9 +1895,7 @@ protected:
 	/** Will update the view custom data. */
 	void PostInitViewCustomData();
 
-	virtual void SortMeshDrawCommands() override;
-
-	void SortMobileMeshDrawCommands();
+	void SortMobileBasePassAfterShadowInit(FMeshDrawCommandsPerPassPerView& VisibleCommandsPerView);
 
 	void UpdateOpaqueBasePassUniformBuffer(FRHICommandListImmediate& RHICmdList, const FViewInfo& View);
 	void UpdateTranslucentBasePassUniformBuffer(FRHICommandListImmediate& RHICmdList, const FViewInfo& View);
@@ -1949,7 +2034,4 @@ extern void SortPassMeshDrawCommands(
 	FMeshCommandOneFrameArray& VisibleMeshDrawCommands,
 	FDynamicMeshDrawCommandStorage& MeshDrawCommandStorage,
 	FGlobalDynamicVertexBuffer& DynamicVertexBuffer,
-	FGlobalDynamicVertexBuffer::FAllocation& PrimitiveIdBuffer,
-	bool bDumpStats);
-
-extern void UpdateTranslucentMeshSortKeys(const FScene* Scene, const FViewInfo& View, ETranslucencyPass::Type TranslucencyPass, FMeshCommandOneFrameArray& VisibleMeshCommands);
\ No newline at end of file
+	FGlobalDynamicVertexBuffer::FAllocation& PrimitiveIdBuffer);
\ No newline at end of file
diff --git a/Engine/Source/Runtime/Renderer/Private/SceneVisibility.cpp b/Engine/Source/Runtime/Renderer/Private/SceneVisibility.cpp
--- a/Engine/Source/Runtime/Renderer/Private/SceneVisibility.cpp
+++ b/Engine/Source/Runtime/Renderer/Private/SceneVisibility.cpp
@@ -1745,6 +1745,7 @@ struct FRelevancePacket
 	FRHICommandListImmediate& RHICmdList;
 	const FScene* Scene;
 	const FViewInfo& View;
+	const FMeshDrawCommandsPerPass& VisibleCommands;
 	const uint8 ViewBit;
 	const FMarkRelevantStaticMeshesForViewData& ViewData;
 	FPrimitiveViewMasks& OutHasDynamicMeshElementsMasks;
@@ -1816,6 +1817,7 @@ struct FRelevancePacket
 		FRHICommandListImmediate& InRHICmdList,
 		const FScene* InScene, 
 		const FViewInfo& InView, 
+		const FMeshDrawCommandsPerPass& InVisibleCommands,
 		uint8 InViewBit,
 		const FMarkRelevantStaticMeshesForViewData& InViewData,
 		FPrimitiveViewMasks& InOutHasDynamicMeshElementsMasks,
@@ -1829,6 +1831,7 @@ struct FRelevancePacket
 		, RHICmdList(InRHICmdList)
 		, Scene(InScene)
 		, View(InView)
+		, VisibleCommands(InVisibleCommands)
 		, ViewBit(InViewBit)
 		, ViewData(InViewData)
 		, OutHasDynamicMeshElementsMasks(InOutHasDynamicMeshElementsMasks)
@@ -2228,6 +2231,7 @@ struct FRelevancePacket
 	void RenderThreadFinalize()
 	{
 		FViewInfo& WriteView = const_cast<FViewInfo&>(View);
+		FMeshDrawCommandsPerPass& WriteVisibleCommands = const_cast<FMeshDrawCommandsPerPass&>(VisibleCommands);
 		
 		for (int32 Index = 0; Index < NotDrawRelevant.NumPrims; Index++)
 		{
@@ -2265,7 +2269,7 @@ struct FRelevancePacket
 
 		for (int32 PassIndex = 0; PassIndex < EMeshPass::Num; PassIndex++)
 		{
-			DrawCommandPacket.VisibleCachedDrawCommands[PassIndex].CopyToLinearArray(WriteView.VisibleMeshDrawCommands[PassIndex]);
+			DrawCommandPacket.VisibleCachedDrawCommands[PassIndex].CopyToLinearArray(WriteVisibleCommands[PassIndex]);
 			DrawCommandPacket.DynamicBuildRequests[PassIndex].CopyToLinearArray(WriteView.DynamicMeshCommandBuildRequests[PassIndex]);
 		}
 
@@ -2290,6 +2294,7 @@ static void ComputeAndMarkRelevanceForViewParallel(
 	FRHICommandListImmediate& RHICmdList,
 	const FScene* Scene,
 	FViewInfo& View,
+	FMeshDrawCommandsPerPass& VisibleCommands,
 	uint8 ViewBit,
 	FPrimitiveViewMasks& OutHasDynamicMeshElementsMasks,
 	FPrimitiveViewMasks& OutHasDynamicEditorMeshElementsMasks,
@@ -2328,6 +2333,7 @@ static void ComputeAndMarkRelevanceForViewParallel(
 				RHICmdList,
 				Scene, 
 				View, 
+				VisibleCommands,
 				ViewBit,
 				ViewData,
 				OutHasDynamicMeshElementsMasks,
@@ -2353,6 +2359,7 @@ static void ComputeAndMarkRelevanceForViewParallel(
 							RHICmdList,
 							Scene, 
 							View, 
+							VisibleCommands,
 							ViewBit,
 							ViewData,
 							OutHasDynamicMeshElementsMasks,
@@ -2390,7 +2397,7 @@ static void ComputeAndMarkRelevanceForViewParallel(
 				NumDynamicBuildRequests += Packet->DrawCommandPacket.DynamicBuildRequests[PassIndex].Num();
 			}
 
-			View.VisibleMeshDrawCommands[PassIndex].Reserve(NumVisibleCachedMeshDrawCommands);
+			VisibleCommands[PassIndex].Reserve(NumVisibleCachedMeshDrawCommands);
 			View.DynamicMeshCommandBuildRequests[PassIndex].Reserve(NumDynamicBuildRequests);
 		}
 
@@ -3140,7 +3147,7 @@ void UpdateReflectionSceneData(FScene* Scene)
 	}
 }
 
-void FSceneRenderer::ComputeViewVisibility(FRHICommandListImmediate& RHICmdList, FExclusiveDepthStencil::Type BasePassDepthStencilAccess)
+void FSceneRenderer::ComputeViewVisibility(FRHICommandListImmediate& RHICmdList, FExclusiveDepthStencil::Type BasePassDepthStencilAccess, FMeshDrawCommandsPerPassPerView& VisibleCommandsPerView)
 {
 	SCOPE_CYCLE_COUNTER(STAT_ViewVisibilityTime);
 	SCOPED_NAMED_EVENT(FSceneRenderer_ComputeViewVisibility, FColor::Magenta);
@@ -3410,7 +3417,7 @@ void FSceneRenderer::ComputeViewVisibility(FRHICommandListImmediate& RHICmdList,
 		if (!View.IsInstancedStereoPass())
 		{
 			SCOPE_CYCLE_COUNTER(STAT_ViewRelevance);
-			ComputeAndMarkRelevanceForViewParallel(RHICmdList, Scene, View, ViewBit, HasDynamicMeshElementsMasks, HasDynamicEditorMeshElementsMasks, HasViewCustomDataMasks);
+			ComputeAndMarkRelevanceForViewParallel(RHICmdList, Scene, View, VisibleCommandsPerView[ViewIndex], ViewBit, HasDynamicMeshElementsMasks, HasDynamicEditorMeshElementsMasks, HasViewCustomDataMasks);
 		}
 
 #if !(UE_BUILD_SHIPPING || UE_BUILD_TEST)
@@ -3459,12 +3466,14 @@ void FSceneRenderer::ComputeViewVisibility(FRHICommandListImmediate& RHICmdList,
 	}
 
 	ViewBit = 0x1;
-	for (FViewInfo& View : Views)
+	for (int32 ViewIndex = 0; ViewIndex < Views.Num(); ++ViewIndex)
 	{
+		FViewInfo& View = Views[ViewIndex];
+		
 		if (View.IsInstancedStereoPass())
 		{
 			SCOPE_CYCLE_COUNTER(STAT_ViewRelevance);
-			ComputeAndMarkRelevanceForViewParallel(RHICmdList, Scene, View, ViewBit, HasDynamicMeshElementsMasks, HasDynamicEditorMeshElementsMasks, HasViewCustomDataMasks);
+			ComputeAndMarkRelevanceForViewParallel(RHICmdList, Scene, View, VisibleCommandsPerView[ViewIndex], ViewBit, HasDynamicMeshElementsMasks, HasDynamicEditorMeshElementsMasks, HasViewCustomDataMasks);
 		}
 		ViewBit <<= 1;
 	}
@@ -3474,12 +3483,17 @@ void FSceneRenderer::ComputeViewVisibility(FRHICommandListImmediate& RHICmdList,
 
 	if (UseMeshDrawCommandPipeline())
 	{
-		// Convert each FMeshBatch into a set of FMeshDrawCommands for each relevant mesh pass
-		GenerateDynamicMeshDrawCommands();
+		for (int32 ViewIndex = 0; ViewIndex < Views.Num(); ViewIndex++)
+		{
+			FViewInfo& View = Views[ViewIndex];
+			FMeshDrawCommandsPerPass& VisibleCommands = VisibleCommandsPerView[ViewIndex];
+
+			GenerateDynamicMeshDrawCommands(View, VisibleCommands);
 
-		ApplyViewOverridesToMeshDrawCommands(BasePassDepthStencilAccess);
+			ApplyViewOverridesToMeshDrawCommands(BasePassDepthStencilAccess, View, VisibleCommands);
 
-		SortMeshDrawCommands();
+			SortAndMergeMeshDrawCommands(View, VisibleCommands);
+		}
 	}
 
 	INC_DWORD_STAT_BY(STAT_ProcessedPrimitives,NumProcessedPrimitives);
@@ -3748,7 +3762,10 @@ bool FDeferredShadingSceneRenderer::InitViews(FRHICommandListImmediate& RHICmdLi
 
 	RHICmdList.ImmediateFlush(EImmediateFlushType::DispatchToRHIThread);
 
-	ComputeViewVisibility(RHICmdList, BasePassDepthStencilAccess);
+	FMeshDrawCommandsPerPassPerView VisibleCommandsPerView;
+	VisibleCommandsPerView.SetNum(Views.Num());
+
+	ComputeViewVisibility(RHICmdList, BasePassDepthStencilAccess, VisibleCommandsPerView);
 
 	RHICmdList.ImmediateFlush(EImmediateFlushType::DispatchToRHIThread);
 
diff --git a/Engine/Source/Runtime/Renderer/Private/ShadowDepthRendering.cpp b/Engine/Source/Runtime/Renderer/Private/ShadowDepthRendering.cpp
--- a/Engine/Source/Runtime/Renderer/Private/ShadowDepthRendering.cpp
+++ b/Engine/Source/Runtime/Renderer/Private/ShadowDepthRendering.cpp
@@ -1847,14 +1847,7 @@ void FProjectedShadowInfo::RenderDepthInner(FRHICommandListImmediate& RHICmdList
 
 			if (UseMeshDrawCommandPipeline())
 			{
-				if (ShadowDepthPass.VisibleMeshDrawCommands.Num() > 0)
-				{
-					FVertexBufferRHIParamRef PrimitiveIdsBuffer = UseGPUScene(GMaxRHIShaderPlatform, FeatureLevel) ? ShadowDepthPass.PrimitiveIdBuffer.VertexBuffer->VertexBufferRHI : nullptr;
-					const int32 BasePrimitiveIdsOffset = ShadowDepthPass.PrimitiveIdBuffer.VertexOffset;
-					const bool bDynamicInstancing = IsDynamicInstancingEnabled() && UseGPUScene(GMaxRHIShaderPlatform, FeatureLevel);
-
-					SubmitMeshDrawCommands(ShadowDepthPass.VisibleMeshDrawCommands, PrimitiveIdsBuffer, BasePrimitiveIdsOffset, bDynamicInstancing, &ParallelCommandListSet, RHICmdList);
-				}
+				ShadowDepthPass.DispatchDraw(&ParallelCommandListSet, RHICmdList);
 			}
 			else
 			{
@@ -1905,14 +1898,7 @@ void FProjectedShadowInfo::RenderDepthInner(FRHICommandListImmediate& RHICmdList
 
 		if (UseMeshDrawCommandPipeline())
 		{
-			if (ShadowDepthPass.VisibleMeshDrawCommands.Num() > 0)
-			{
-				FVertexBufferRHIParamRef PrimitiveIdsBuffer = UseGPUScene(GMaxRHIShaderPlatform, FeatureLevel) ? ShadowDepthPass.PrimitiveIdBuffer.VertexBuffer->VertexBufferRHI : nullptr;
-				const int32 BasePrimitiveIdsOffset = ShadowDepthPass.PrimitiveIdBuffer.VertexOffset;
-				const bool bDynamicInstancing = IsDynamicInstancingEnabled() && UseGPUScene(GMaxRHIShaderPlatform, FeatureLevel);
-
-				SubmitMeshDrawCommands(ShadowDepthPass.VisibleMeshDrawCommands, PrimitiveIdsBuffer, BasePrimitiveIdsOffset, bDynamicInstancing, nullptr, RHICmdList);
-			}
+			ShadowDepthPass.DispatchDraw(nullptr, RHICmdList);
 		}
 		else
 		{
diff --git a/Engine/Source/Runtime/Renderer/Private/ShadowRendering.cpp b/Engine/Source/Runtime/Renderer/Private/ShadowRendering.cpp
--- a/Engine/Source/Runtime/Renderer/Private/ShadowRendering.cpp
+++ b/Engine/Source/Runtime/Renderer/Private/ShadowRendering.cpp
@@ -618,7 +618,7 @@ void FProjectedShadowInfo::SetupProjectionStencilMask(
 				const int32 BasePrimitiveIdsOffset = ProjectionStencilingPass.PrimitiveIdBuffer.VertexOffset;
 				const bool bDynamicInstancing = IsDynamicInstancingEnabled() && UseGPUScene(GMaxRHIShaderPlatform, View->GetFeatureLevel());
 
-				SubmitMeshDrawCommands(ProjectionStencilingPass.VisibleMeshDrawCommands, PrimitiveIdsBuffer, BasePrimitiveIdsOffset, bDynamicInstancing, nullptr, RHICmdList);
+				SubmitMeshDrawCommands(ProjectionStencilingPass.VisibleMeshDrawCommands, PrimitiveIdsBuffer, BasePrimitiveIdsOffset, bDynamicInstancing, RHICmdList);
 			}
 		}
 		else
@@ -851,7 +851,7 @@ void FProjectedShadowInfo::SetupProjectionStencilMask(
 					const int32 BasePrimitiveIdsOffset = ProjectionStencilingPass.PrimitiveIdBuffer.VertexOffset;
 					const bool bDynamicInstancing = IsDynamicInstancingEnabled() && UseGPUScene(GMaxRHIShaderPlatform, View->GetFeatureLevel());
 
-					SubmitMeshDrawCommands(ProjectionStencilingPass.VisibleMeshDrawCommands, PrimitiveIdsBuffer, BasePrimitiveIdsOffset, bDynamicInstancing, nullptr, RHICmdList);
+					SubmitMeshDrawCommands(ProjectionStencilingPass.VisibleMeshDrawCommands, PrimitiveIdsBuffer, BasePrimitiveIdsOffset, bDynamicInstancing, RHICmdList);
 				}
 			}
 			else
diff --git a/Engine/Source/Runtime/Renderer/Private/ShadowRendering.h b/Engine/Source/Runtime/Renderer/Private/ShadowRendering.h
--- a/Engine/Source/Runtime/Renderer/Private/ShadowRendering.h
+++ b/Engine/Source/Runtime/Renderer/Private/ShadowRendering.h
@@ -733,7 +733,10 @@ private:
 	TArray<FMeshBatchAndRelevance,SceneRenderingAllocator> DynamicSubjectTranslucentMeshElements;
 
 	TArray<const FStaticMesh*, SceneRenderingAllocator> SubjectMeshCommandBuildRequests;
-	FShadowMeshDrawCommandPass ShadowDepthPass;
+
+	FMeshCommandOneFrameArray ShadowDepthPassVisibleCommands;
+	FParallelMeshDrawCommandPass ShadowDepthPass;
+
 	TArray<FShadowMeshDrawCommandPass, TInlineAllocator<2>> ProjectionStencilingPasses;
 
 	FDynamicMeshDrawCommandStorage DynamicMeshDrawCommandStorage;
diff --git a/Engine/Source/Runtime/Renderer/Private/ShadowSetup.cpp b/Engine/Source/Runtime/Renderer/Private/ShadowSetup.cpp
--- a/Engine/Source/Runtime/Renderer/Private/ShadowSetup.cpp
+++ b/Engine/Source/Runtime/Renderer/Private/ShadowSetup.cpp
@@ -930,7 +930,7 @@ bool FProjectedShadowInfo::ShouldDrawStaticMeshes(FViewInfo& InCurrentView, bool
 								StaticMesh,
 								InPrimitiveSceneInfo->Scene,
 								EMeshPass::CSMShadowDepth,
-								ShadowDepthPass.VisibleMeshDrawCommands,
+								ShadowDepthPassVisibleCommands,
 								SubjectMeshCommandBuildRequests);
 						}
 						else
@@ -1187,7 +1187,7 @@ bool FProjectedShadowInfo::HasSubjectPrims() const
 {
 	return DynamicSubjectPrimitives.Num() > 0
 		|| StaticSubjectMeshElements.Num() > 0
-		|| ShadowDepthPass.VisibleMeshDrawCommands.Num() > 0
+		|| ShadowDepthPass.HasAnyDraw()
 		|| SubjectMeshCommandBuildRequests.Num() > 0;
 }
 
@@ -1201,7 +1201,7 @@ void FProjectedShadowInfo::SetupMeshDrawCommandsForShadowDepth(FSceneRenderer& R
 {
 	QUICK_SCOPE_CYCLE_COUNTER(STAT_SetupMeshDrawCommandsForShadowDepth);
 
-	FDynamicPassMeshDrawListContext ShadowDepthContext(DynamicMeshDrawCommandStorage, ShadowDepthPass.VisibleMeshDrawCommands);
+	FDynamicPassMeshDrawListContext ShadowDepthContext(DynamicMeshDrawCommandStorage, ShadowDepthPassVisibleCommands);
 	FShadowDepthPassMeshProcessor ShadowDepthPassMeshProcessor(
 		Renderer.Scene, 
 		ShadowDepthView, 
@@ -1231,14 +1231,17 @@ void FProjectedShadowInfo::SetupMeshDrawCommandsForShadowDepth(FSceneRenderer& R
 		ShadowDepthPassMeshProcessor.AddMeshBatch(*StaticMeshBatch, BatchElementMask, StaticMeshBatch->PrimitiveSceneInfo->Proxy, StaticMeshBatch->Id);
 	}
 
-	if (Renderer.ShouldDumpMeshDrawCommandInstancingStats() && ShadowDepthPass.VisibleMeshDrawCommands.Num() > 0)
+	if (ShadowDepthPassVisibleCommands.Num() > 0)
 	{
-		FString ShadowName;
-		GetShadowTypeNameForDrawEvent(ShadowName);
-		UE_LOG(LogRenderer, Log, TEXT("ShadowDepth %s"), *ShadowName);
-	}
+		if (Renderer.ShouldDumpMeshDrawCommandInstancingStats())
+		{
+			FString PassNameForStats;
+			GetShadowTypeNameForDrawEvent(PassNameForStats);
+			ShadowDepthPass.SetDumpInstancingStats(TEXT("ShadowDepth ") + PassNameForStats);
+		}
 
-	SortPassMeshDrawCommands(Renderer.FeatureLevel, ShadowDepthPass.VisibleMeshDrawCommands, DynamicMeshDrawCommandStorage, FGlobalDynamicVertexBuffer::Get(), ShadowDepthPass.PrimitiveIdBuffer, Renderer.ShouldDumpMeshDrawCommandInstancingStats());
+		ShadowDepthPass.DispatchSortAndMerge(*ShadowDepthView, Renderer.Scene->PrimitiveBounds, EMeshPass::Num, ShadowDepthPassVisibleCommands);
+	}
 }
 
 void FProjectedShadowInfo::SetupMeshDrawCommandsForProjectionStenciling(FSceneRenderer& Renderer)
@@ -1352,7 +1355,7 @@ void FProjectedShadowInfo::SetupMeshDrawCommandsForProjectionStenciling(FSceneRe
 
 			ApplyViewOverridesToMeshDrawCommands(View, ProjectionStencilingPass.VisibleMeshDrawCommands);
 
-			SortPassMeshDrawCommands(Renderer.FeatureLevel, ProjectionStencilingPass.VisibleMeshDrawCommands, DynamicMeshDrawCommandStorage, FGlobalDynamicVertexBuffer::Get(), ProjectionStencilingPass.PrimitiveIdBuffer, false);
+			SortPassMeshDrawCommands(Renderer.FeatureLevel, ProjectionStencilingPass.VisibleMeshDrawCommands, DynamicMeshDrawCommandStorage, FGlobalDynamicVertexBuffer::Get(), ProjectionStencilingPass.PrimitiveIdBuffer);
 		}
 	}
 }
@@ -1551,8 +1554,9 @@ void FProjectedShadowInfo::ClearTransientArrays()
 	DynamicReceiverMeshElements.Empty();
 	DynamicSubjectTranslucentMeshElements.Empty();
 
-	ShadowDepthPass.VisibleMeshDrawCommands.Empty();
-	ShadowDepthPass.PrimitiveIdBuffer = FGlobalDynamicVertexBuffer::FAllocation();
+	ShadowDepthPassVisibleCommands.Empty();
+	ShadowDepthPass.Empty();
+
 	SubjectMeshCommandBuildRequests.Empty();
 
 	ProjectionStencilingPasses.Reset();
diff --git a/Engine/Source/Runtime/Renderer/Private/ShadowSetupMobile.cpp b/Engine/Source/Runtime/Renderer/Private/ShadowSetupMobile.cpp
--- a/Engine/Source/Runtime/Renderer/Private/ShadowSetupMobile.cpp
+++ b/Engine/Source/Runtime/Renderer/Private/ShadowSetupMobile.cpp
@@ -268,7 +268,7 @@ static void VisualizeMobileDynamicCSMSubjectCapsules(FViewInfo& View, FLightScen
 	}
 }
 /** Finds the visible dynamic shadows for each view. */
-void FMobileSceneRenderer::InitDynamicShadows(FRHICommandListImmediate& RHICmdList)
+void FMobileSceneRenderer::InitDynamicShadows(FRHICommandListImmediate& RHICmdList, FMeshDrawCommandsPerPassPerView& VisibleCommandsPerView)
 {
 	static auto* MyCVarMobileEnableStaticAndCSMShadowReceivers = IConsoleManager::Get().FindTConsoleVariableDataInt(TEXT("r.Mobile.EnableStaticAndCSMShadowReceivers"));
 	const bool bMobileEnableStaticAndCSMShadowReceivers = MyCVarMobileEnableStaticAndCSMShadowReceivers->GetValueOnRenderThread() == 1;
@@ -338,14 +338,17 @@ void FMobileSceneRenderer::InitDynamicShadows(FRHICommandListImmediate& RHICmdLi
 	}
 
 	// Where we've determined CSM is required we copy the MobileBasePassCSM mesh commands to the BasePass visible list:
-	for (auto& View : Views)
+	for (int32 ViewIndex = 0; ViewIndex < Views.Num(); ++ViewIndex)
 	{
+		FViewInfo& View = Views[ViewIndex];
+		FMeshDrawCommandsPerPass& VisibleCommands = VisibleCommandsPerView[ViewIndex];
+
 		FMobileCSMVisibilityInfo& MobileCSMVisibilityInfo = View.MobileCSMVisibilityInfo;
-		if(MobileCSMVisibilityInfo.bMobileDynamicCSMInUse)
+		if (MobileCSMVisibilityInfo.bMobileDynamicCSMInUse)
 		{
 			// determine per view CSM visibility
-			FMeshCommandOneFrameArray& MeshCommands = View.VisibleMeshDrawCommands[EMeshPass::BasePass];
-			FMeshCommandOneFrameArray& MeshCommandsCSM = View.VisibleMeshDrawCommands[EMeshPass::MobileBasePassCSM];
+			FMeshCommandOneFrameArray& MeshCommands = VisibleCommands[EMeshPass::BasePass];
+			FMeshCommandOneFrameArray& MeshCommandsCSM = VisibleCommands[EMeshPass::MobileBasePassCSM];
 			checkf(MeshCommands.Num() == MeshCommandsCSM.Num(), TEXT("VisibleMeshDrawCommands of BasePass and MobileBasePassCSM are expected to match."));
 			for (int32 i = MeshCommands.Num() - 1; i >= 0; --i)
 			{
diff --git a/Engine/Source/Runtime/Renderer/Private/TranslucentRendering.cpp b/Engine/Source/Runtime/Renderer/Private/TranslucentRendering.cpp
--- a/Engine/Source/Runtime/Renderer/Private/TranslucentRendering.cpp
+++ b/Engine/Source/Runtime/Renderer/Private/TranslucentRendering.cpp
@@ -451,7 +451,8 @@ void RenderViewTranslucencyInner(FRHICommandListImmediate& RHICmdList, const FVi
 	{
 		QUICK_SCOPE_CYCLE_COUNTER(RenderTranslucencyParallel_Start_FDrawSortedTransAnyThreadTask);
 
-		SubmitMeshDrawCommandsForView(View, TranslucencyPassToMeshPass(TranslucencyPass), ParallelCommandListSet, RHICmdList);
+		const EMeshPass::Type MeshPass = TranslucencyPassToMeshPass(TranslucencyPass);
+		View.ParallelMeshDrawCommandPasses[MeshPass].DispatchDraw(ParallelCommandListSet, RHICmdList);
 	}
 
 	if (IsMainTranslucencyPass(TranslucencyPass))
diff --git a/Engine/Source/Runtime/Renderer/Private/VelocityRendering.cpp b/Engine/Source/Runtime/Renderer/Private/VelocityRendering.cpp
--- a/Engine/Source/Runtime/Renderer/Private/VelocityRendering.cpp
+++ b/Engine/Source/Runtime/Renderer/Private/VelocityRendering.cpp
@@ -434,7 +434,7 @@ void FDeferredShadingSceneRenderer::RenderVelocitiesInnerParallel(FRHICommandLis
 				VelocityRT);
 
 			// Draw velocities.
-			SubmitMeshDrawCommandsForView(View, EMeshPass::Velocity, &ParallelCommandListSet, RHICmdList);
+			View.ParallelMeshDrawCommandPasses[EMeshPass::Velocity].DispatchDraw(&ParallelCommandListSet, RHICmdList);
 		}
 	}
 }
@@ -461,7 +461,7 @@ void FDeferredShadingSceneRenderer::RenderVelocitiesInner(FRHICommandListImmedia
 
 			SetVelocitiesState(RHICmdList, View, this, DrawRenderState, VelocityRT);
 
-			SubmitMeshDrawCommandsForView(View, EMeshPass::Velocity, nullptr, RHICmdList);
+			View.ParallelMeshDrawCommandPasses[EMeshPass::Velocity].DispatchDraw(nullptr, RHICmdList);
 		}
 	}
 }
diff --git a/Engine/Source/Runtime/Renderer/Public/MeshPassProcessor.h b/Engine/Source/Runtime/Renderer/Public/MeshPassProcessor.h
--- a/Engine/Source/Runtime/Renderer/Public/MeshPassProcessor.h
+++ b/Engine/Source/Runtime/Renderer/Public/MeshPassProcessor.h
@@ -829,13 +829,15 @@ extern void SubmitMeshDrawCommands(
 	FVertexBufferRHIParamRef PrimitiveIdsBuffer,
 	int32 BasePrimitiveIdsOffset,
 	bool bDynamicInstancing,
-	class FParallelCommandListSet* ParallelCommandListSet,
 	FRHICommandList& RHICmdList);
 
-extern void SubmitMeshDrawCommandsForView(
-	const FViewInfo& View,
-	EMeshPass::Type PassType,
-	class FParallelCommandListSet* ParallelCommandListSet,
+extern void SubmitMeshDrawCommandsRange(
+	const FMeshCommandOneFrameArray& VisibleMeshDrawCommands,
+	FVertexBufferRHIParamRef PrimitiveIdsBuffer,
+	int32 BasePrimitiveIdsOffset,
+	bool bDynamicInstancing,
+	int32 StartIndex,
+	int32 NumMeshDrawCommands,
 	FRHICommandList& RHICmdList);
 
 RENDERER_API extern void DrawDynamicMeshPassPrivate(
