Fixed mobile separate translucency missing when High Quality DOF is enabled

diff --git a/Engine/Source/Runtime/Renderer/Private/MobileSeparateTranslucencyPass.cpp b/Engine/Source/Runtime/Renderer/Private/MobileSeparateTranslucencyPass.cpp
--- a/Engine/Source/Runtime/Renderer/Private/MobileSeparateTranslucencyPass.cpp
+++ b/Engine/Source/Runtime/Renderer/Private/MobileSeparateTranslucencyPass.cpp
@@ -27,8 +27,9 @@ void FRCSeparateTranslucensyPassES2::Process(FRenderingCompositePassContext& Con
 	const FViewInfo& View = Context.View;
 
 	FSceneRenderTargets& SceneTargets = FSceneRenderTargets::Get(Context.RHICmdList);
-
-	FRHIRenderPassInfo RPInfo(SceneTargets.GetSceneColorSurface(), ERenderTargetActions::Load_Store, SceneTargets.GetSceneColorTexture());
+	const FSceneRenderTargetItem& DestRenderTarget = GetOutput(ePId_Output0)->RequestSurface(Context);
+		
+	FRHIRenderPassInfo RPInfo(DestRenderTarget.TargetableTexture, ERenderTargetActions::Load_Store);
 	RPInfo.DepthStencilRenderTarget.Action = EDepthStencilTargetActions::LoadDepthStencil_StoreDepthStencil;
 	RPInfo.DepthStencilRenderTarget.DepthStencilTarget = SceneTargets.GetSceneDepthSurface();
 	RPInfo.DepthStencilRenderTarget.ExclusiveDepthStencil = FExclusiveDepthStencil::DepthRead_StencilRead;
