Renamed FVisibleMeshDrawCommandProcessTask to FMeshDrawCommandSortAndMergeTask.

diff --git a/Engine/Source/Runtime/Renderer/Private/SceneRendering.cpp b/Engine/Source/Runtime/Renderer/Private/SceneRendering.cpp
--- a/Engine/Source/Runtime/Renderer/Private/SceneRendering.cpp
+++ b/Engine/Source/Runtime/Renderer/Private/SceneRendering.cpp
@@ -2940,9 +2940,9 @@ void BuildMeshDrawCommandPrimitiveIdBuffer(
 	}
 }
 
-FAutoConsoleTaskPriority CPrio_FVisibleMeshDrawCommandProcessTask(
-	TEXT("TaskGraph.TaskPriorities.VisibleMeshDrawCommandProcessTask"),
-	TEXT("Task and thread priority for FVisibleMeshDrawCommandProcessTask."),
+FAutoConsoleTaskPriority CPrio_FMeshDrawCommandSortAndMergeTask(
+	TEXT("TaskGraph.TaskPriorities.MeshDrawCommandSortAndMergeTask"),
+	TEXT("Task and thread priority for FMeshDrawCommandSortAndMergeTask."),
 	ENamedThreads::HighThreadPriority, // if we have high priority task threads, then use them...
 	ENamedThreads::NormalTaskPriority, // .. at normal task priority
 	ENamedThreads::HighTaskPriority // if we don't have hi pri threads, then use normal priority threads at high task priority instead
@@ -2951,10 +2951,10 @@ FAutoConsoleTaskPriority CPrio_FVisibleMeshDrawCommandProcessTask(
 /**
  * Task for parallel sorting and building primitive id buffer for mesh draw commands.
  */
-class FVisibleMeshDrawCommandProcessTask
+class FMeshDrawCommandSortAndMergeTask
 {
 public:
-	FVisibleMeshDrawCommandProcessTask(FVisibleMeshDrawCommandProcessTaskContext& InContext)
+	FMeshDrawCommandSortAndMergeTask(FMeshDrawCommandSortAndMergeTaskContext& InContext)
 		: Context(InContext)
 	{
 	}
@@ -2966,7 +2966,7 @@ public:
 
 	ENamedThreads::Type GetDesiredThread()
 	{
-		return CPrio_FVisibleMeshDrawCommandProcessTask.Get();
+		return CPrio_FMeshDrawCommandSortAndMergeTask.Get();
 	}
 
 	static ESubsequentsMode::Type GetSubsequentsMode() 
@@ -3013,7 +3013,7 @@ public:
 
 
 private:
-	FVisibleMeshDrawCommandProcessTaskContext& Context;
+	FMeshDrawCommandSortAndMergeTaskContext& Context;
 };
 
 void SortPassMeshDrawCommands(
@@ -3090,7 +3090,7 @@ void FParallelMeshDrawCommandPass::DispatchSortAndMerge(const FViewInfo& View, c
 	const bool bExecuteInParallel = FApp::ShouldUseThreadingForPerformance() && CVarMeshDrawCommandsParallelProcess.GetValueOnRenderThread() > 0;
 
 	// Preallocate context memory and resources on rendering thread.
-	if (TaskContext.bUseGPUScene && MaxDrawNum > 0)
+	if (TaskContext.bUseGPUScene)
 	{
 		TaskContext.PrimitiveIdBuffer = FGlobalDynamicVertexBuffer::Get().Allocate(MaxDrawNum * sizeof(int32));
 
@@ -3102,12 +3102,12 @@ void FParallelMeshDrawCommandPass::DispatchSortAndMerge(const FViewInfo& View, c
 
 	if (bExecuteInParallel)
 	{
-		TaskEventRefs.Add(TGraphTask<FVisibleMeshDrawCommandProcessTask>::CreateTask(
+		TaskEventRefs.Add(TGraphTask<FMeshDrawCommandSortAndMergeTask>::CreateTask(
 			nullptr, ENamedThreads::GetRenderThread()).ConstructAndDispatchWhenReady(TaskContext));
 	}
 	else
 	{
-		FVisibleMeshDrawCommandProcessTask Task(TaskContext);
+		FMeshDrawCommandSortAndMergeTask Task(TaskContext);
 		Task.AnyThreadTask();
 	}
 }
diff --git a/Engine/Source/Runtime/Renderer/Private/SceneRendering.h b/Engine/Source/Runtime/Renderer/Private/SceneRendering.h
--- a/Engine/Source/Runtime/Renderer/Private/SceneRendering.h
+++ b/Engine/Source/Runtime/Renderer/Private/SceneRendering.h
@@ -883,10 +883,10 @@ typedef TArray<FMeshDrawCommandsPerPass, TInlineAllocator<4>> FMeshDrawCommandsP
 /**
  * Parallel mesh draw command processing and drawing context.
  */
-class FVisibleMeshDrawCommandProcessTaskContext
+class FMeshDrawCommandSortAndMergeTaskContext
 {
 public:
-	FVisibleMeshDrawCommandProcessTaskContext()
+	FMeshDrawCommandSortAndMergeTaskContext()
 		: bUseGPUScene(false)
 		, bDynamicInstancing(false)
 		, PrimitiveBounds(nullptr)
@@ -954,7 +954,7 @@ public:
 
 
 private:
-	FVisibleMeshDrawCommandProcessTaskContext TaskContext;
+	FMeshDrawCommandSortAndMergeTaskContext TaskContext;
 	FGraphEventArray TaskEventRefs;
 	int32 MaxDrawNum;
 	FString PassNameForStats;
