FMeshDrawCommandSortAndMergeTask - renamed stat.

diff --git a/Engine/Source/Runtime/Renderer/Private/SceneRendering.cpp b/Engine/Source/Runtime/Renderer/Private/SceneRendering.cpp
--- a/Engine/Source/Runtime/Renderer/Private/SceneRendering.cpp
+++ b/Engine/Source/Runtime/Renderer/Private/SceneRendering.cpp
@@ -2961,7 +2961,7 @@ public:
 
 	FORCEINLINE TStatId GetStatId() const
 	{
-		RETURN_QUICK_DECLARE_CYCLE_STAT(FSortFrontToBackTask, STATGROUP_TaskGraphTasks);
+		RETURN_QUICK_DECLARE_CYCLE_STAT(FMeshDrawCommandSortAndMergeTask, STATGROUP_TaskGraphTasks);
 	}
 
 	ENamedThreads::Type GetDesiredThread()
