Moved ApplyViewOverridesToMeshDrawCommands from InitViews to a parallel task processing mesh draw command passes.

diff --git a/Engine/Source/Runtime/Renderer/Private/MeshPassProcessor.cpp b/Engine/Source/Runtime/Renderer/Private/MeshPassProcessor.cpp
--- a/Engine/Source/Runtime/Renderer/Private/MeshPassProcessor.cpp
+++ b/Engine/Source/Runtime/Renderer/Private/MeshPassProcessor.cpp
@@ -14,6 +14,7 @@
 #include "MeshPassProcessor.inl"
 
 TSet<FGraphicsMinimalPipelineStateInitializer> FGraphicsMinimalPipelineStateId::GlobalTable;
+FCriticalSection GlobalTableCriticalSection;
 
 const FMeshDrawCommandSortKey FMeshDrawCommandSortKey::Default = { 0 };
 
@@ -290,7 +291,8 @@ void FMeshDrawShaderBindings::SetOnRayTracingStructure(FRHICommandList& RHICmdLi
 
 void FGraphicsMinimalPipelineStateId::Setup(const FGraphicsMinimalPipelineStateInitializer& InPipelineState)
 {
-	checkSlow(IsInRenderingThread());
+	// Need to lock as this is called from multiple parallel tasks during mesh draw command generation or patching.
+	FScopeLock Lock(&GlobalTableCriticalSection);
 
 	FSetElementId TableId = GlobalTable.FindId(InPipelineState);
 
diff --git a/Engine/Source/Runtime/Renderer/Private/MobileShadingRenderer.cpp b/Engine/Source/Runtime/Renderer/Private/MobileShadingRenderer.cpp
--- a/Engine/Source/Runtime/Renderer/Private/MobileShadingRenderer.cpp
+++ b/Engine/Source/Runtime/Renderer/Private/MobileShadingRenderer.cpp
@@ -113,7 +113,7 @@ TUniformBufferRef<FMobileDirectionalLightShaderParameters>& GetNullMobileDirecti
 	return NullLightParams->UniformBufferRHI;
 }
 
-void FMobileSceneRenderer::SortMobileBasePassAfterShadowInit(FMeshDrawCommandsPerPassPerView& VisibleCommandsPerView)
+void FMobileSceneRenderer::SortMobileBasePassAfterShadowInit(FExclusiveDepthStencil::Type BasePassDepthStencilAccess, FMeshDrawCommandsPerPassPerView& VisibleCommandsPerView)
 {
 	// Sort front to back on all platforms, even HSR benefits from it
 	//const bool bWantsFrontToBackSorting = (GHardwareHiddenSurfaceRemoval == false);
@@ -131,7 +131,7 @@ void FMobileSceneRenderer::SortMobileBasePassAfterShadowInit(FMeshDrawCommandsPe
 		if (VisibleCommands[EMeshPass::BasePass].Num() > 0)
 		{
 			FParallelMeshDrawCommandPass& Pass = View.ParallelMeshDrawCommandPasses[EMeshPass::BasePass];
-			Pass.DispatchSortAndMerge(View, Scene->PrimitiveBounds, EMeshPass::BasePass, VisibleCommands[EMeshPass::BasePass]);
+			Pass.DispatchSortAndMerge(View, Scene->GetShadingPath(), Scene->PrimitiveBounds, EMeshPass::BasePass, BasePassDepthStencilAccess, VisibleCommands[EMeshPass::BasePass]);
 		}
 	}
 }
@@ -152,8 +152,10 @@ void FMobileSceneRenderer::InitViews(FRHICommandListImmediate& RHICmdList)
 	FMeshDrawCommandsPerPassPerView VisibleCommandsPerView;
 	VisibleCommandsPerView.SetNum(Views.Num());
 
+	const FExclusiveDepthStencil::Type BasePassDepthStencilAccess = FExclusiveDepthStencil::DepthWrite_StencilWrite;
+
 	PreVisibilityFrameSetup(RHICmdList);
-	ComputeViewVisibility(RHICmdList, FExclusiveDepthStencil::DepthWrite_StencilWrite, VisibleCommandsPerView);
+	ComputeViewVisibility(RHICmdList, BasePassDepthStencilAccess, VisibleCommandsPerView);
 	PostVisibilityFrameSetup(ILCTaskData);
 
 	const bool bDynamicShadows = ViewFamily.EngineShowFlags.DynamicShadows;
@@ -164,7 +166,7 @@ void FMobileSceneRenderer::InitViews(FRHICommandListImmediate& RHICmdList)
 		InitDynamicShadows(RHICmdList, VisibleCommandsPerView);
 	}
 
-	SortMobileBasePassAfterShadowInit(VisibleCommandsPerView);
+	SortMobileBasePassAfterShadowInit(BasePassDepthStencilAccess, VisibleCommandsPerView);
 
 	// if we kicked off ILC update via task, wait and finalize.
 	if (ILCTaskData.TaskRef.IsValid())
diff --git a/Engine/Source/Runtime/Renderer/Private/SceneRendering.cpp b/Engine/Source/Runtime/Renderer/Private/SceneRendering.cpp
--- a/Engine/Source/Runtime/Renderer/Private/SceneRendering.cpp
+++ b/Engine/Source/Runtime/Renderer/Private/SceneRendering.cpp
@@ -2674,69 +2674,6 @@ void FSceneRenderer::RenderFinish(FRHICommandListImmediate& RHICmdList)
 	RHICmdList.EndScene();
 }
 
-void FSceneRenderer::ApplyViewOverridesToMeshDrawCommands(FExclusiveDepthStencil::Type BasePassDepthStencilAccess, FViewInfo& View, FMeshDrawCommandsPerPass& VisibleCommands)
-{
-	QUICK_SCOPE_CYCLE_COUNTER(STAT_FSceneRenderer_ApplyViewOverridesToMeshDrawCommands);
-
-	EShadingPath ShadingPath = Scene->GetShadingPath();
-	FExclusiveDepthStencil::Type DefaultBasePassDepthStencilAccess = GetDefaultBasePassDepthStencilAccess(ShadingPath);
-
-	for (int32 PassIndex = 0; PassIndex < EMeshPass::Num; PassIndex++)
-	{
-		EMeshPass::Type PassType = (EMeshPass::Type)PassIndex;
-
-		if ((FPassProcessorManager::GetPassFlags(ShadingPath, PassType) & EMeshPassFlags::MainView) != EMeshPassFlags::None)
-		{
-			if (View.bReverseCulling || View.bRenderSceneTwoSided || (BasePassDepthStencilAccess != DefaultBasePassDepthStencilAccess && PassType == EMeshPass::BasePass))
-			{
-				const FMeshCommandOneFrameArray& PassVisibleMeshDrawCommands = VisibleCommands[PassIndex];
-
-				FMeshCommandOneFrameArray ViewOverriddenMeshCommands;
-				ViewOverriddenMeshCommands.Empty(PassVisibleMeshDrawCommands.Num());
-
-				FDynamicMeshDrawCommandStorage& MeshDrawCommandStorage = View.DynamicMeshDrawCommandStorage[PassIndex];
-
-				for (int32 MeshCommandIndex = 0; MeshCommandIndex < PassVisibleMeshDrawCommands.Num(); MeshCommandIndex++)
-				{
-					MeshDrawCommandStorage.MeshDrawCommands.Add(1);
-					FMeshDrawCommand& NewMeshCommand = MeshDrawCommandStorage.MeshDrawCommands[MeshDrawCommandStorage.MeshDrawCommands.Num() - 1];
-
-					const FVisibleMeshDrawCommand& VisibleMeshDrawCommand = PassVisibleMeshDrawCommands[MeshCommandIndex];
-					const FMeshDrawCommand& MeshCommand = *VisibleMeshDrawCommand.MeshDrawCommand;
-					NewMeshCommand = MeshCommand;
-
-					const ERasterizerCullMode LocalCullMode = View.bRenderSceneTwoSided ? CM_None : View.bReverseCulling ? FMeshDrawingPolicy::InverseCullMode(VisibleMeshDrawCommand.MeshCullMode) : VisibleMeshDrawCommand.MeshCullMode;
-					NewMeshCommand.PipelineState.RasterizerState = GetStaticRasterizerState<true>(VisibleMeshDrawCommand.MeshFillMode, LocalCullMode);
-
-					if (BasePassDepthStencilAccess != DefaultBasePassDepthStencilAccess && PassType == EMeshPass::BasePass)
-					{
-						FDrawingPolicyRenderState PassDrawRenderState;
-						SetupBasePassState(BasePassDepthStencilAccess, false, PassDrawRenderState);
-						NewMeshCommand.PipelineState.DepthStencilState = PassDrawRenderState.GetDepthStencilState();
-					}
-
-					NewMeshCommand.Finalize(true);
-
-					FVisibleMeshDrawCommand NewVisibleMeshDrawCommand;
-
-					NewVisibleMeshDrawCommand.Setup(
-						&NewMeshCommand,
-						VisibleMeshDrawCommand.DrawPrimitiveId,
-						VisibleMeshDrawCommand.StateBucketId,
-						VisibleMeshDrawCommand.MeshFillMode,
-						VisibleMeshDrawCommand.MeshCullMode,
-						VisibleMeshDrawCommand.SortKey);
-
-					ViewOverriddenMeshCommands.Add(NewVisibleMeshDrawCommand);
-				}
-
-				// Replace VisibleMeshDrawCommands
-				FMemory::Memswap(&VisibleCommands[PassIndex], &ViewOverriddenMeshCommands, sizeof(ViewOverriddenMeshCommands));
-			}
-		}
-	}
-}
-
 struct FCompareFMeshDrawCommands
 {
 	FORCEINLINE bool operator() (FVisibleMeshDrawCommand A, FVisibleMeshDrawCommand B) const
@@ -2826,27 +2763,27 @@ void UpdateTranslucentMeshSortKeys(
 
 /**
  * Build mesh draw command primitive Id buffer for instancing.
- * NewPassVisibleMeshDrawCommands must be presized for NewPassVisibleMeshDrawCommands.
+ * TempVisibleMeshDrawCommands must be presized for NewPassVisibleMeshDrawCommands.
  */
 void BuildMeshDrawCommandPrimitiveIdBuffer(
 	bool bUseGPUScene, 
 	FMeshCommandOneFrameArray& VisibleMeshDrawCommands,
 	FDynamicMeshDrawCommandStorage& MeshDrawCommandStorage,
 	FGlobalDynamicVertexBuffer::FAllocation& PrimitiveIdBuffer,
-	FMeshCommandOneFrameArray& NewPassVisibleMeshDrawCommands,
+	FMeshCommandOneFrameArray& TempVisibleMeshDrawCommands,
 	int32& MaxInstances,
 	int32& VisibleMeshDrawCommandsNum,
 	int32& NewPassVisibleMeshDrawCommandsNum
 	)
 {
+	QUICK_SCOPE_CYCLE_COUNTER(STAT_BuildMeshDrawCommandPrimitiveIdBuffer);
+	check(VisibleMeshDrawCommands.Num() <= TempVisibleMeshDrawCommands.Max() && TempVisibleMeshDrawCommands.Num() == 0)
 	if (!bUseGPUScene)
 	{
 		return;
 	}
 
 
-	check(VisibleMeshDrawCommands.Num() <= NewPassVisibleMeshDrawCommands.Max())
-
 	const FVisibleMeshDrawCommand* RESTRICT PassVisibleMeshDrawCommands = VisibleMeshDrawCommands.GetData();
 	const int32 NumDrawCommands = VisibleMeshDrawCommands.Num();
 
@@ -2876,7 +2813,7 @@ void BuildMeshDrawCommandPrimitiveIdBuffer(
 				{
 					FVisibleMeshDrawCommand NewVisibleMeshDrawCommand = VisibleMeshDrawCommand;
 					NewVisibleMeshDrawCommand.PrimitiveIdBufferOffset = DrawCommandIndex;
-					NewPassVisibleMeshDrawCommands.Emplace(MoveTemp(NewVisibleMeshDrawCommand));
+					TempVisibleMeshDrawCommands.Emplace(MoveTemp(NewVisibleMeshDrawCommand));
 				}
 			}
 			else
@@ -2904,7 +2841,7 @@ void BuildMeshDrawCommandPrimitiveIdBuffer(
 						VisibleMeshDrawCommand.SortKey);
 
 					NewVisibleMeshDrawCommand.PrimitiveIdBufferOffset = DrawCommandIndex;
-					NewPassVisibleMeshDrawCommands.Emplace(MoveTemp(NewVisibleMeshDrawCommand));
+					TempVisibleMeshDrawCommands.Emplace(MoveTemp(NewVisibleMeshDrawCommand));
 
 					CurrentDynamicallyInstancedMeshCommandNumInstances = &NewCommand.NumInstances;
 				}
@@ -2913,7 +2850,7 @@ void BuildMeshDrawCommandPrimitiveIdBuffer(
 					CurrentDynamicallyInstancedMeshCommandNumInstances = nullptr;
 					FVisibleMeshDrawCommand NewVisibleMeshDrawCommand = VisibleMeshDrawCommand;
 					NewVisibleMeshDrawCommand.PrimitiveIdBufferOffset = DrawCommandIndex;
-					NewPassVisibleMeshDrawCommands.Emplace(MoveTemp(NewVisibleMeshDrawCommand));
+					TempVisibleMeshDrawCommands.Emplace(MoveTemp(NewVisibleMeshDrawCommand));
 				}
 			}
 
@@ -2923,10 +2860,11 @@ void BuildMeshDrawCommandPrimitiveIdBuffer(
 
 		// Setup instancing stats for logging.
 		VisibleMeshDrawCommandsNum = VisibleMeshDrawCommands.Num();
-		NewPassVisibleMeshDrawCommandsNum = NewPassVisibleMeshDrawCommands.Num();
+		NewPassVisibleMeshDrawCommandsNum = TempVisibleMeshDrawCommands.Num();
 
 		// Replace VisibleMeshDrawCommands
-		FMemory::Memswap(&VisibleMeshDrawCommands, &NewPassVisibleMeshDrawCommands, sizeof(NewPassVisibleMeshDrawCommands));
+		FMemory::Memswap(&VisibleMeshDrawCommands, &TempVisibleMeshDrawCommands, sizeof(TempVisibleMeshDrawCommands));
+		TempVisibleMeshDrawCommands.Reset();
 	}
 	else
 	{
@@ -2940,6 +2878,72 @@ void BuildMeshDrawCommandPrimitiveIdBuffer(
 	}
 }
 
+/**
+* Apply view overrides to existing mesh draw commands (e.g. reverse culling mode for rendering planar reflections).
+* TempVisibleMeshDrawCommands must be presized for NewPassVisibleMeshDrawCommands.
+*/
+void ApplyViewOverridesToMeshDrawCommands(
+	EShadingPath ShadingPath,
+	EMeshPass::Type PassType,
+	bool bReverseCulling,
+	bool bRenderSceneTwoSided,
+	FExclusiveDepthStencil::Type BasePassDepthStencilAccess,
+	FMeshCommandOneFrameArray& VisibleMeshDrawCommands,
+	FDynamicMeshDrawCommandStorage& MeshDrawCommandStorage,
+	FMeshCommandOneFrameArray& TempVisibleMeshDrawCommands
+	)
+{
+	QUICK_SCOPE_CYCLE_COUNTER(STAT_ApplyViewOverridesToMeshDrawCommands);
+	check(VisibleMeshDrawCommands.Num() <= TempVisibleMeshDrawCommands.Max() && TempVisibleMeshDrawCommands.Num() == 0)
+
+	if (PassType != EMeshPass::Num
+		&& (FPassProcessorManager::GetPassFlags(ShadingPath, PassType) & EMeshPassFlags::MainView) != EMeshPassFlags::None)
+	{
+		const FExclusiveDepthStencil::Type DefaultBasePassDepthStencilAccess = GetDefaultBasePassDepthStencilAccess(ShadingPath);
+
+		if (bReverseCulling || bRenderSceneTwoSided || (BasePassDepthStencilAccess != DefaultBasePassDepthStencilAccess && PassType == EMeshPass::BasePass))
+		{
+			for (int32 MeshCommandIndex = 0; MeshCommandIndex < VisibleMeshDrawCommands.Num(); MeshCommandIndex++)
+			{
+				MeshDrawCommandStorage.MeshDrawCommands.Add(1);
+				FMeshDrawCommand& NewMeshCommand = MeshDrawCommandStorage.MeshDrawCommands[MeshDrawCommandStorage.MeshDrawCommands.Num() - 1];
+
+				const FVisibleMeshDrawCommand& VisibleMeshDrawCommand = VisibleMeshDrawCommands[MeshCommandIndex];
+				const FMeshDrawCommand& MeshCommand = *VisibleMeshDrawCommand.MeshDrawCommand;
+				NewMeshCommand = MeshCommand;
+
+				const ERasterizerCullMode LocalCullMode = bRenderSceneTwoSided ? CM_None : bReverseCulling ? FMeshDrawingPolicy::InverseCullMode(VisibleMeshDrawCommand.MeshCullMode) : VisibleMeshDrawCommand.MeshCullMode;
+				NewMeshCommand.PipelineState.RasterizerState = GetStaticRasterizerState<true>(VisibleMeshDrawCommand.MeshFillMode, LocalCullMode);
+
+				if (BasePassDepthStencilAccess != DefaultBasePassDepthStencilAccess && PassType == EMeshPass::BasePass)
+				{
+					FDrawingPolicyRenderState PassDrawRenderState;
+					SetupBasePassState(BasePassDepthStencilAccess, false, PassDrawRenderState);
+					NewMeshCommand.PipelineState.DepthStencilState = PassDrawRenderState.GetDepthStencilState();
+				}
+
+				NewMeshCommand.Finalize(true);
+
+				FVisibleMeshDrawCommand NewVisibleMeshDrawCommand;
+
+				NewVisibleMeshDrawCommand.Setup(
+					&NewMeshCommand,
+					VisibleMeshDrawCommand.DrawPrimitiveId,
+					VisibleMeshDrawCommand.StateBucketId,
+					VisibleMeshDrawCommand.MeshFillMode,
+					VisibleMeshDrawCommand.MeshCullMode,
+					VisibleMeshDrawCommand.SortKey);
+
+				TempVisibleMeshDrawCommands.Add(NewVisibleMeshDrawCommand);
+			}
+
+			// Replace VisibleMeshDrawCommands
+			FMemory::Memswap(&VisibleMeshDrawCommands, &TempVisibleMeshDrawCommands, sizeof(TempVisibleMeshDrawCommands));
+			TempVisibleMeshDrawCommands.Reset();
+		}
+	}
+}
+
 FAutoConsoleTaskPriority CPrio_FMeshDrawCommandSortAndMergeTask(
 	TEXT("TaskGraph.TaskPriorities.MeshDrawCommandSortAndMergeTask"),
 	TEXT("Task and thread priority for FMeshDrawCommandSortAndMergeTask."),
@@ -2976,6 +2980,17 @@ public:
 
 	void AnyThreadTask()
 	{
+		ApplyViewOverridesToMeshDrawCommands(
+			Context.ShadingPath,
+			Context.PassType,
+			Context.bReverseCulling,
+			Context.bRenderSceneTwoSided,
+			Context.BasePassDepthStencilAccess,
+			Context.VisibleMeshDrawCommands,
+			Context.MeshDrawCommandStorage,
+			Context.TempVisibleMeshDrawCommands
+		);
+
 		if (Context.TranslucencyPass != ETranslucencyPass::TPT_MAX)
 		{
 			UpdateTranslucentMeshSortKeys(
@@ -2999,7 +3014,7 @@ public:
 			Context.VisibleMeshDrawCommands,
 			Context.MeshDrawCommandStorage,
 			Context.PrimitiveIdBuffer,
-			Context.NewPassVisibleMeshDrawCommands,
+			Context.TempVisibleMeshDrawCommands,
 			Context.MaxInstances,
 			Context.VisibleMeshDrawCommandsNum,
 			Context.NewPassVisibleMeshDrawCommandsNum
@@ -3059,13 +3074,27 @@ void SortPassMeshDrawCommands(
 	}
 }
 
-void FParallelMeshDrawCommandPass::DispatchSortAndMerge(const FViewInfo& View, const TArray<FPrimitiveBounds>& PrimitiveBounds, EMeshPass::Type MeshPass, FMeshCommandOneFrameArray& InOutVisibleMeshDrawCommands)
+void FParallelMeshDrawCommandPass::DispatchSortAndMerge(
+	const FViewInfo& View, 
+	EShadingPath ShadingPath,
+	const TArray<FPrimitiveBounds>& PrimitiveBounds, 
+	EMeshPass::Type PassType, 
+	FExclusiveDepthStencil::Type BasePassDepthStencilAccess,
+	FMeshCommandOneFrameArray& InOutVisibleMeshDrawCommands)
 {
 	check(TaskEventRefs.Num() == 0);
 	check(InOutVisibleMeshDrawCommands.Num() > 0);
 
 	MaxDrawNum = InOutVisibleMeshDrawCommands.Num();
 
+	TaskContext.ShadingPath = ShadingPath;
+	TaskContext.PassType = PassType;
+	TaskContext.bUseGPUScene = UseGPUScene(GMaxRHIShaderPlatform, View.GetFeatureLevel());
+	TaskContext.bDynamicInstancing = TaskContext.bUseGPUScene && IsDynamicInstancingEnabled();
+	TaskContext.bReverseCulling = View.bReverseCulling;
+	TaskContext.bRenderSceneTwoSided = View.bRenderSceneTwoSided;
+	TaskContext.BasePassDepthStencilAccess = BasePassDepthStencilAccess;
+
 	// Setup translucency sort key update pass based on view.
 	TaskContext.TranslucencyPass = ETranslucencyPass::TPT_MAX;
 	TaskContext.TranslucentSortPolicy = View.TranslucentSortPolicy;
@@ -3074,7 +3103,7 @@ void FParallelMeshDrawCommandPass::DispatchSortAndMerge(const FViewInfo& View, c
 	TaskContext.ViewMatrix = View.ViewMatrices.GetViewMatrix();
 	TaskContext.PrimitiveBounds = &PrimitiveBounds;
 
-	switch (MeshPass)
+	switch (PassType)
 	{
 		case EMeshPass::TranslucencyStandard: TaskContext.TranslucencyPass = ETranslucencyPass::TPT_StandardTranslucency; break;
 		case EMeshPass::TranslucencyAfterDOF: TaskContext.TranslucencyPass = ETranslucencyPass::TPT_TranslucencyAfterDOF; break;
@@ -3082,9 +3111,6 @@ void FParallelMeshDrawCommandPass::DispatchSortAndMerge(const FViewInfo& View, c
 		case EMeshPass::MobileInverseOpacity: TaskContext.TranslucencyPass = ETranslucencyPass::TPT_StandardTranslucency; break;
 	}
 
-	TaskContext.bUseGPUScene = UseGPUScene(GMaxRHIShaderPlatform, View.GetFeatureLevel());
-	TaskContext.bDynamicInstancing = TaskContext.bUseGPUScene && IsDynamicInstancingEnabled();
-
 	FMemory::Memswap(&TaskContext.VisibleMeshDrawCommands, &InOutVisibleMeshDrawCommands, sizeof(InOutVisibleMeshDrawCommands));
 
 	const bool bExecuteInParallel = FApp::ShouldUseThreadingForPerformance() && CVarMeshDrawCommandsParallelProcess.GetValueOnRenderThread() > 0;
@@ -3096,7 +3122,7 @@ void FParallelMeshDrawCommandPass::DispatchSortAndMerge(const FViewInfo& View, c
 
 		if (TaskContext.bDynamicInstancing)
 		{
-			TaskContext.NewPassVisibleMeshDrawCommands.Empty(MaxDrawNum);
+			TaskContext.TempVisibleMeshDrawCommands.Empty(MaxDrawNum);
 		}
 	}
 
@@ -3126,7 +3152,7 @@ void FParallelMeshDrawCommandPass::Empty()
 	PassNameForStats.Empty();
 	TaskEventRefs.Empty();
 	TaskContext.VisibleMeshDrawCommands.Empty();
-	TaskContext.NewPassVisibleMeshDrawCommands.Empty();
+	TaskContext.TempVisibleMeshDrawCommands.Empty();
 	TaskContext.PrimitiveIdBuffer = FGlobalDynamicVertexBuffer::FAllocation();
 }
 
@@ -3256,7 +3282,7 @@ void FParallelMeshDrawCommandPass::SetDumpInstancingStats(const FString& InPassN
 	PassNameForStats = InPassNameForStats;
 }
 
-void FSceneRenderer::SortAndMergeMeshDrawCommands(FViewInfo& View, FMeshDrawCommandsPerPass& VisibleCommands)
+void FSceneRenderer::SortAndMergeMeshDrawCommands(FViewInfo& View, FExclusiveDepthStencil::Type BasePassDepthStencilAccess, FMeshDrawCommandsPerPass& VisibleCommands)
 {
 	SCOPE_CYCLE_COUNTER(STAT_SortAndMergeMeshDrawCommands);
 
@@ -3283,7 +3309,7 @@ void FSceneRenderer::SortAndMergeMeshDrawCommands(FViewInfo& View, FMeshDrawComm
 					Pass.SetDumpInstancingStats(GetMeshPassName(PassType));
 				}
 
-				Pass.DispatchSortAndMerge(View, Scene->PrimitiveBounds, PassType, VisibleCommands[PassIndex]);
+				Pass.DispatchSortAndMerge(View, Scene->GetShadingPath(), Scene->PrimitiveBounds, PassType, BasePassDepthStencilAccess, VisibleCommands[PassIndex]);
 			}
 		}
 	}
diff --git a/Engine/Source/Runtime/Renderer/Private/SceneRendering.h b/Engine/Source/Runtime/Renderer/Private/SceneRendering.h
--- a/Engine/Source/Runtime/Renderer/Private/SceneRendering.h
+++ b/Engine/Source/Runtime/Renderer/Private/SceneRendering.h
@@ -887,8 +887,13 @@ class FMeshDrawCommandSortAndMergeTaskContext
 {
 public:
 	FMeshDrawCommandSortAndMergeTaskContext()
-		: bUseGPUScene(false)
+		: ShadingPath(EShadingPath::Num)
+		, PassType(EMeshPass::Num)
+		, bUseGPUScene(false)
 		, bDynamicInstancing(false)
+		, bReverseCulling(false)
+		, bRenderSceneTwoSided(false)
+		, BasePassDepthStencilAccess(FExclusiveDepthStencil::DepthNop_StencilNop)
 		, PrimitiveBounds(nullptr)
 		, VisibleMeshDrawCommandsNum(0)
 		, NewPassVisibleMeshDrawCommandsNum(0)
@@ -896,16 +901,21 @@ public:
 	{
 	}
 
+	EShadingPath ShadingPath;
+	EMeshPass::Type PassType;
 	bool bUseGPUScene;
 	bool bDynamicInstancing;
+	bool bReverseCulling;
+	bool bRenderSceneTwoSided;
+	FExclusiveDepthStencil::Type BasePassDepthStencilAccess;
 
-	// Command storage.
+	// Commands.
 	FMeshCommandOneFrameArray VisibleMeshDrawCommands;
 	FDynamicMeshDrawCommandStorage MeshDrawCommandStorage;
 
 	// Preallocated resources.
 	FGlobalDynamicVertexBuffer::FAllocation PrimitiveIdBuffer;
-	FMeshCommandOneFrameArray NewPassVisibleMeshDrawCommands;
+	FMeshCommandOneFrameArray TempVisibleMeshDrawCommands;
 
 	// For UpdateTranslucentMeshSortKeys.
 	ETranslucencyPass::Type TranslucencyPass;
@@ -939,7 +949,14 @@ public:
 	/**
 	 * Dispatch visible mesh draw command sort and merge task.
 	 */
-	void DispatchSortAndMerge(const FViewInfo& View, const TArray<FPrimitiveBounds>& PrimitiveBounds, EMeshPass::Type MeshPass, FMeshCommandOneFrameArray& InOutVisibleMeshDrawCommands);
+	void DispatchSortAndMerge(
+		const FViewInfo& View, 
+		EShadingPath ShadingPath,
+		const TArray<FPrimitiveBounds>& PrimitiveBounds, 
+		EMeshPass::Type PassType, 
+		FExclusiveDepthStencil::Type BasePassDepthStencilAccess,
+		FMeshCommandOneFrameArray& InOutVisibleMeshDrawCommands
+	);
 
 	/**
 	 * Dispatch visible mesh draw command draw task.
@@ -1673,9 +1690,7 @@ protected:
 	/** Converts each FMeshBatch into a set of FMeshDrawCommands for each relevant mesh pass. */
 	void GenerateDynamicMeshDrawCommands(FViewInfo& View, FMeshDrawCommandsPerPass& VisibleCommands);
 
-	void ApplyViewOverridesToMeshDrawCommands(FExclusiveDepthStencil::Type BasePassDepthStencilAccess, FViewInfo& View, FMeshDrawCommandsPerPass& VisibleCommands);
-
-	void SortAndMergeMeshDrawCommands(FViewInfo& View, FMeshDrawCommandsPerPass& VisibleCommands);
+	void SortAndMergeMeshDrawCommands(FViewInfo& View, FExclusiveDepthStencil::Type BasePassDepthStencilAccess, FMeshDrawCommandsPerPass& VisibleCommands);
 
 	bool RenderShadowProjections(FRHICommandListImmediate& RHICmdList, const FLightSceneInfo* LightSceneInfo, IPooledRenderTarget* ScreenShadowMaskTexture, bool bProjectingForForwardShading, bool bMobileModulatedProjections);
 
@@ -1895,7 +1910,7 @@ protected:
 	/** Will update the view custom data. */
 	void PostInitViewCustomData();
 
-	void SortMobileBasePassAfterShadowInit(FMeshDrawCommandsPerPassPerView& VisibleCommandsPerView);
+	void SortMobileBasePassAfterShadowInit(FExclusiveDepthStencil::Type BasePassDepthStencilAccess, FMeshDrawCommandsPerPassPerView& VisibleCommandsPerView);
 
 	void UpdateOpaqueBasePassUniformBuffer(FRHICommandListImmediate& RHICmdList, const FViewInfo& View);
 	void UpdateTranslucentBasePassUniformBuffer(FRHICommandListImmediate& RHICmdList, const FViewInfo& View);
diff --git a/Engine/Source/Runtime/Renderer/Private/SceneVisibility.cpp b/Engine/Source/Runtime/Renderer/Private/SceneVisibility.cpp
--- a/Engine/Source/Runtime/Renderer/Private/SceneVisibility.cpp
+++ b/Engine/Source/Runtime/Renderer/Private/SceneVisibility.cpp
@@ -3490,9 +3490,7 @@ void FSceneRenderer::ComputeViewVisibility(FRHICommandListImmediate& RHICmdList,
 
 			GenerateDynamicMeshDrawCommands(View, VisibleCommands);
 
-			ApplyViewOverridesToMeshDrawCommands(BasePassDepthStencilAccess, View, VisibleCommands);
-
-			SortAndMergeMeshDrawCommands(View, VisibleCommands);
+			SortAndMergeMeshDrawCommands(View, BasePassDepthStencilAccess, VisibleCommands);
 		}
 	}
 
diff --git a/Engine/Source/Runtime/Renderer/Private/ShadowSetup.cpp b/Engine/Source/Runtime/Renderer/Private/ShadowSetup.cpp
--- a/Engine/Source/Runtime/Renderer/Private/ShadowSetup.cpp
+++ b/Engine/Source/Runtime/Renderer/Private/ShadowSetup.cpp
@@ -1240,7 +1240,7 @@ void FProjectedShadowInfo::SetupMeshDrawCommandsForShadowDepth(FSceneRenderer& R
 			ShadowDepthPass.SetDumpInstancingStats(TEXT("ShadowDepth ") + PassNameForStats);
 		}
 
-		ShadowDepthPass.DispatchSortAndMerge(*ShadowDepthView, Renderer.Scene->PrimitiveBounds, EMeshPass::Num, ShadowDepthPassVisibleCommands);
+		ShadowDepthPass.DispatchSortAndMerge(*ShadowDepthView, Renderer.Scene->GetShadingPath(), Renderer.Scene->PrimitiveBounds, EMeshPass::Num, FExclusiveDepthStencil::DepthNop_StencilNop, ShadowDepthPassVisibleCommands);
 	}
 }
 
