Request full update of the GlobalDistanceField when swtiching a scene

diff --git a/Engine/Source/Runtime/Renderer/Private/GlobalDistanceField.cpp b/Engine/Source/Runtime/Renderer/Private/GlobalDistanceField.cpp
--- a/Engine/Source/Runtime/Renderer/Private/GlobalDistanceField.cpp
+++ b/Engine/Source/Runtime/Renderer/Private/GlobalDistanceField.cpp
@@ -1014,6 +1014,7 @@ static void ComputeUpdateRegionsAndUpdateViewState(
 				|| ClipmapViewState.CachedMaxOcclusionDistance != MaxOcclusionDistance
 				|| ClipmapViewState.CachedGlobalDistanceFieldViewDistance != Scene->GlobalDistanceFieldViewDistance
 				|| ClipmapViewState.CacheMostlyStaticSeparately != GAOGlobalDistanceFieldCacheMostlyStaticSeparately
+				|| ClipmapViewState.LastUsedSceneDataForFullUpdate != &Scene->DistanceFieldSceneData
 				|| GAOGlobalDistanceFieldForceFullUpdate;
 
 			if (ShouldUpdateClipmapThisFrame(ClipmapIndex, View.ViewState->GlobalDistanceFieldUpdateIndex)
@@ -1036,6 +1037,7 @@ static void ComputeUpdateRegionsAndUpdateViewState(
 					// Store the location of the full update
 					ClipmapViewState.FullUpdateOrigin = GridCenter;
 					View.ViewState->bInitializedGlobalDistanceFieldOrigins = true;
+					ClipmapViewState.LastUsedSceneDataForFullUpdate = &Scene->DistanceFieldSceneData;
 				}
 
 				const FGlobalDFCacheType StartCacheType = GAOGlobalDistanceFieldCacheMostlyStaticSeparately ? GDF_MostlyStatic : GDF_Full;
diff --git a/Engine/Source/Runtime/Renderer/Private/ScenePrivate.h b/Engine/Source/Runtime/Renderer/Private/ScenePrivate.h
--- a/Engine/Source/Runtime/Renderer/Private/ScenePrivate.h
+++ b/Engine/Source/Runtime/Renderer/Private/ScenePrivate.h
@@ -459,6 +459,7 @@ public:
 		CachedMaxOcclusionDistance = 0;
 		CachedGlobalDistanceFieldViewDistance = 0;
 		CacheMostlyStaticSeparately = 1;
+		LastUsedSceneDataForFullUpdate = nullptr;
 	}
 
 	FIntVector FullUpdateOrigin;
@@ -468,6 +469,9 @@ public:
 	uint32 CacheMostlyStaticSeparately;
 
 	FGlobalDistanceFieldCacheTypeState Cache[GDF_Num];
+
+	// Used to perform a full update of the clip map when the scene data changes
+	const class FDistanceFieldSceneData* LastUsedSceneDataForFullUpdate;
 };
 
 /** Maps a single primitive to it's per-view fading state data */
