Fixed Non GPUScene path in mesh draw command tasks.

diff --git a/Engine/Source/Runtime/Renderer/Private/SceneRendering.cpp b/Engine/Source/Runtime/Renderer/Private/SceneRendering.cpp
--- a/Engine/Source/Runtime/Renderer/Private/SceneRendering.cpp
+++ b/Engine/Source/Runtime/Renderer/Private/SceneRendering.cpp
@@ -2766,7 +2766,6 @@ void UpdateTranslucentMeshSortKeys(
  * TempVisibleMeshDrawCommands must be presized for NewPassVisibleMeshDrawCommands.
  */
 void BuildMeshDrawCommandPrimitiveIdBuffer(
-	bool bUseGPUScene, 
 	FMeshCommandOneFrameArray& VisibleMeshDrawCommands,
 	FDynamicMeshDrawCommandStorage& MeshDrawCommandStorage,
 	FGlobalDynamicVertexBuffer::FAllocation& PrimitiveIdBuffer,
@@ -2777,12 +2776,7 @@ void BuildMeshDrawCommandPrimitiveIdBuffer(
 	)
 {
 	QUICK_SCOPE_CYCLE_COUNTER(STAT_BuildMeshDrawCommandPrimitiveIdBuffer);
-	check(VisibleMeshDrawCommands.Num() <= TempVisibleMeshDrawCommands.Max() && TempVisibleMeshDrawCommands.Num() == 0)
-	if (!bUseGPUScene)
-	{
-		return;
-	}
-
+	check(VisibleMeshDrawCommands.Num() <= TempVisibleMeshDrawCommands.Max() && TempVisibleMeshDrawCommands.Num() == 0 && PrimitiveIdBuffer.IsValid());
 
 	const FVisibleMeshDrawCommand* RESTRICT PassVisibleMeshDrawCommands = VisibleMeshDrawCommands.GetData();
 	const int32 NumDrawCommands = VisibleMeshDrawCommands.Num();
@@ -2894,10 +2888,9 @@ void ApplyViewOverridesToMeshDrawCommands(
 	)
 {
 	QUICK_SCOPE_CYCLE_COUNTER(STAT_ApplyViewOverridesToMeshDrawCommands);
-	check(VisibleMeshDrawCommands.Num() <= TempVisibleMeshDrawCommands.Max() && TempVisibleMeshDrawCommands.Num() == 0)
+	check(VisibleMeshDrawCommands.Num() <= TempVisibleMeshDrawCommands.Max() && TempVisibleMeshDrawCommands.Num() == 0 && PassType != EMeshPass::Num);
 
-	if (PassType != EMeshPass::Num
-		&& (FPassProcessorManager::GetPassFlags(ShadingPath, PassType) & EMeshPassFlags::MainView) != EMeshPassFlags::None)
+	if ((FPassProcessorManager::GetPassFlags(ShadingPath, PassType) & EMeshPassFlags::MainView) != EMeshPassFlags::None)
 	{
 		const FExclusiveDepthStencil::Type DefaultBasePassDepthStencilAccess = GetDefaultBasePassDepthStencilAccess(ShadingPath);
 
@@ -2980,16 +2973,19 @@ public:
 
 	void AnyThreadTask()
 	{
-		ApplyViewOverridesToMeshDrawCommands(
-			Context.ShadingPath,
-			Context.PassType,
-			Context.bReverseCulling,
-			Context.bRenderSceneTwoSided,
-			Context.BasePassDepthStencilAccess,
-			Context.VisibleMeshDrawCommands,
-			Context.MeshDrawCommandStorage,
-			Context.TempVisibleMeshDrawCommands
-		);
+		if (Context.PassType != EMeshPass::Num)
+		{
+			ApplyViewOverridesToMeshDrawCommands(
+				Context.ShadingPath,
+				Context.PassType,
+				Context.bReverseCulling,
+				Context.bRenderSceneTwoSided,
+				Context.BasePassDepthStencilAccess,
+				Context.VisibleMeshDrawCommands,
+				Context.MeshDrawCommandStorage,
+				Context.TempVisibleMeshDrawCommands
+			);
+		}
 
 		if (Context.TranslucencyPass != ETranslucencyPass::TPT_MAX)
 		{
@@ -3009,16 +3005,18 @@ public:
 			Context.VisibleMeshDrawCommands.Sort(FCompareFMeshDrawCommands());
 		}
 
-		BuildMeshDrawCommandPrimitiveIdBuffer(
-			Context.bUseGPUScene,
-			Context.VisibleMeshDrawCommands,
-			Context.MeshDrawCommandStorage,
-			Context.PrimitiveIdBuffer,
-			Context.TempVisibleMeshDrawCommands,
-			Context.MaxInstances,
-			Context.VisibleMeshDrawCommandsNum,
-			Context.NewPassVisibleMeshDrawCommandsNum
-		);
+		if (Context.bUseGPUScene)
+		{
+			BuildMeshDrawCommandPrimitiveIdBuffer(
+				Context.VisibleMeshDrawCommands,
+				Context.MeshDrawCommandStorage,
+				Context.PrimitiveIdBuffer,
+				Context.TempVisibleMeshDrawCommands,
+				Context.MaxInstances,
+				Context.VisibleMeshDrawCommandsNum,
+				Context.NewPassVisibleMeshDrawCommandsNum
+			);
+		}
 	}
 
 	void DoTask(ENamedThreads::Type CurrentThread, const FGraphEventRef& MyCompletionGraphEvent)
@@ -3049,7 +3047,7 @@ void SortPassMeshDrawCommands(
 		int32 NewPassVisibleMeshDrawCommandsNum = 0;
 
 		// Preallocate context memory and resources.
-		if (bUseGPUScene && NumDrawCommands > 0)
+		if (bUseGPUScene)
 		{
 			PrimitiveIdBuffer = DynamicVertexBuffer.Allocate(NumDrawCommands * sizeof(int32));
 
@@ -3061,16 +3059,18 @@ void SortPassMeshDrawCommands(
 
 		VisibleMeshDrawCommands.Sort(FCompareFMeshDrawCommands());
 
-		BuildMeshDrawCommandPrimitiveIdBuffer(
-			bUseGPUScene,
-			VisibleMeshDrawCommands,
-			MeshDrawCommandStorage,
-			PrimitiveIdBuffer,
-			NewPassVisibleMeshDrawCommands,
-			MaxInstances,
-			VisibleMeshDrawCommandsNum,
-			NewPassVisibleMeshDrawCommandsNum
-		);
+		if (bUseGPUScene)
+		{
+			BuildMeshDrawCommandPrimitiveIdBuffer(
+				VisibleMeshDrawCommands,
+				MeshDrawCommandStorage,
+				PrimitiveIdBuffer,
+				NewPassVisibleMeshDrawCommands,
+				MaxInstances,
+				VisibleMeshDrawCommandsNum,
+				NewPassVisibleMeshDrawCommandsNum
+			);
+		}
 	}
 }
 
@@ -3113,19 +3113,15 @@ void FParallelMeshDrawCommandPass::DispatchSortAndMerge(
 
 	FMemory::Memswap(&TaskContext.VisibleMeshDrawCommands, &InOutVisibleMeshDrawCommands, sizeof(InOutVisibleMeshDrawCommands));
 
-	const bool bExecuteInParallel = FApp::ShouldUseThreadingForPerformance() && CVarMeshDrawCommandsParallelProcess.GetValueOnRenderThread() > 0;
-
 	// Preallocate context memory and resources on rendering thread.
-	if (TaskContext.bUseGPUScene)
+	TaskContext.TempVisibleMeshDrawCommands.Empty(MaxDrawNum);
+	if (TaskContext.bDynamicInstancing)
 	{
 		TaskContext.PrimitiveIdBuffer = FGlobalDynamicVertexBuffer::Get().Allocate(MaxDrawNum * sizeof(int32));
-
-		if (TaskContext.bDynamicInstancing)
-		{
-			TaskContext.TempVisibleMeshDrawCommands.Empty(MaxDrawNum);
-		}
 	}
 
+	const bool bExecuteInParallel = FApp::ShouldUseThreadingForPerformance() && CVarMeshDrawCommandsParallelProcess.GetValueOnRenderThread() > 0;
+
 	if (bExecuteInParallel)
 	{
 		TaskEventRefs.Add(TGraphTask<FMeshDrawCommandSortAndMergeTask>::CreateTask(
