DR - Fix CustomStencilSRV requiring a view of a texture, not a buffer (crashes Vulkan)
#rnx

diff --git a/Engine/Source/Runtime/Renderer/Private/PostProcess/SceneRenderTargets.cpp b/Engine/Source/Runtime/Renderer/Private/PostProcess/SceneRenderTargets.cpp
--- a/Engine/Source/Runtime/Renderer/Private/PostProcess/SceneRenderTargets.cpp
+++ b/Engine/Source/Runtime/Renderer/Private/PostProcess/SceneRenderTargets.cpp
@@ -2953,7 +2953,7 @@ void SetupSceneTextureUniformParameters(
 		const bool bSetupCustomDepth = (SetupMode & ESceneTextureSetupMode::CustomDepth) != ESceneTextureSetupMode::None;
 
 		FTextureRHIParamRef CustomDepth = DepthDefault;
-		FShaderResourceViewRHIParamRef CustomStencilSRV = GNullColorVertexBuffer.VertexBufferSRV;
+		FShaderResourceViewRHIParamRef CustomStencilSRV = GSystemTextures.WhiteDummySRV;
 
 		// if there is no custom depth it's better to have the far distance there
 		IPooledRenderTarget* CustomDepthTarget = SceneContext.bCustomDepthIsValid ? SceneContext.CustomDepth.GetReference() : 0;
diff --git a/Engine/Source/Runtime/Renderer/Private/SystemTextures.cpp b/Engine/Source/Runtime/Renderer/Private/SystemTextures.cpp
--- a/Engine/Source/Runtime/Renderer/Private/SystemTextures.cpp
+++ b/Engine/Source/Runtime/Renderer/Private/SystemTextures.cpp
@@ -32,6 +32,8 @@ void FSystemTextures::InitializeCommonTextures(FRHICommandListImmediate& RHICmdL
 			RHICmdList.BeginRenderPass(RPInfo, TEXT("WhiteDummy"));
 			RHICmdList.EndRenderPass();
 			RHICmdList.CopyToResolveTarget(WhiteDummy->GetRenderTargetItem().TargetableTexture, WhiteDummy->GetRenderTargetItem().ShaderResourceTexture, FResolveParams());
+
+			WhiteDummySRV = RHICreateShaderResourceView((FRHITexture2D*)WhiteDummy->GetRenderTargetItem().ShaderResourceTexture.GetReference(), 0);
 		}
 
 		// Create a BlackDummy texture
@@ -558,6 +560,7 @@ void FSystemTextures::InitializeFeatureLevelDependentTextures(FRHICommandListImm
 
 void FSystemTextures::ReleaseDynamicRHI()
 {
+	WhiteDummySRV.SafeRelease();
 	WhiteDummy.SafeRelease();
 	BlackDummy.SafeRelease();
 	BlackAlphaOneDummy.SafeRelease();
