DR - Missed file
#rnx

diff --git a/Engine/Source/Runtime/Renderer/Private/SystemTextures.h b/Engine/Source/Runtime/Renderer/Private/SystemTextures.h
--- a/Engine/Source/Runtime/Renderer/Private/SystemTextures.h
+++ b/Engine/Source/Runtime/Renderer/Private/SystemTextures.h
@@ -84,6 +84,9 @@ public:
 	/** float4(0,0,0,0) volumetric texture. */
 	TRefCountPtr<IPooledRenderTarget> VolumetricBlackDummy;
 
+	// SRV for WhiteDummy Texture.
+	TRefCountPtr<FRHIShaderResourceView> WhiteDummySRV;
+
 protected:
 	/** Maximum feature level that the textures have been initialized up to */
 	ERHIFeatureLevel::Type FeatureLevelInitializedTo;
