Implements Tokoyashi specular lobe based bilateral kernel, to be continued...

diff --git a/Engine/Shaders/Private/DataStructures/SceneViewFamilyBlackboard.ush b/Engine/Shaders/Private/DataStructures/SceneViewFamilyBlackboard.ush
--- a/Engine/Shaders/Private/DataStructures/SceneViewFamilyBlackboard.ush
+++ b/Engine/Shaders/Private/DataStructures/SceneViewFamilyBlackboard.ush
@@ -13,16 +13,20 @@
 Texture2D SceneDepthBuffer;
 Texture2D SceneVelocityBuffer;
 Texture2D SceneGBufferA;
+Texture2D SceneGBufferB;
+Texture2D SceneGBufferC;
+Texture2D SceneGBufferD;
+Texture2D SceneGBufferE;
 
 
 // @param UV - UV space in the GBuffer textures (BufferSize resolution)
 FGBufferData GetGBufferDataFromSceneBlackboard(float2 UV, bool bGetNormalizedNormal = true)
 {
 	float4 GBufferA = SceneGBufferA.SampleLevel(GlobalPointClampedSampler, UV, 0);
-	float4 GBufferB = 0;
-	float4 GBufferC = 0;
-	float4 GBufferD = 0;
-	float4 GBufferE = 1;
+	float4 GBufferB = SceneGBufferB.SampleLevel(GlobalPointClampedSampler, UV, 0);
+	float4 GBufferC = SceneGBufferC.SampleLevel(GlobalPointClampedSampler, UV, 0);
+	float4 GBufferD = SceneGBufferD.SampleLevel(GlobalPointClampedSampler, UV, 0);
+	float4 GBufferE = SceneGBufferE.SampleLevel(GlobalPointClampedSampler, UV, 0);
 	float4 GBufferVelocity = SceneVelocityBuffer.SampleLevel(GlobalPointClampedSampler, UV, 0);
 
 	uint CustomStencil = 0;
diff --git a/Engine/Shaders/Private/ScreenSpaceDenoise/SSDBuildHistoryRejection.usf b/Engine/Shaders/Private/ScreenSpaceDenoise/SSDBuildHistoryRejection.usf
--- a/Engine/Shaders/Private/ScreenSpaceDenoise/SSDBuildHistoryRejection.usf
+++ b/Engine/Shaders/Private/ScreenSpaceDenoise/SSDBuildHistoryRejection.usf
@@ -215,10 +215,11 @@ void BuildHistoryRejection(
 						NeighborMipLevel = SharedArray0[LDSIndex].x;
 						float NeighborDepth = SharedArray0[LDSIndex].y;
 						float3 NeighborNormal = float3(0, 0, 1);
+						float Roughness = 1;
 
 						NeighborMetadata = InitSampleSceneInfos(
 							View.ScreenToTranslatedWorld,
-							NeighborScreenPosition, NeighborDepth, NeighborNormal);
+							NeighborScreenPosition, NeighborDepth, NeighborNormal, Roughness);
 					}
 				}
 
diff --git a/Engine/Shaders/Private/ScreenSpaceDenoise/SSDCommon.ush b/Engine/Shaders/Private/ScreenSpaceDenoise/SSDCommon.ush
--- a/Engine/Shaders/Private/ScreenSpaceDenoise/SSDCommon.ush
+++ b/Engine/Shaders/Private/ScreenSpaceDenoise/SSDCommon.ush
@@ -9,6 +9,7 @@
 #include "../Common.ush"
 #include "../DataStructures/SceneViewFamilyBlackboard.ush"
 #include "../ColorSpace.ush"
+#include "../SphericalGaussian.ush"
 
 
 //------------------------------------------------------- DEBUG CONFIG
@@ -100,6 +101,9 @@ struct FSSFSampleSceneInfos
 
 	/** Pixel depth in world space. */
 	float WorldDepth;
+
+	/** Roughness of the pixel. */
+	float Roughness;
 };
 
 FSSFSampleSceneInfos CreateSampleSceneInfos()
@@ -108,12 +112,13 @@ FSSFSampleSceneInfos CreateSampleSceneInfos()
 	Infos.WorldNormal             = 0;
 	Infos.TranslatedWorldPosition = 0;
 	Infos.WorldDepth              = 0;
+	Infos.Roughness               = 0;
 	return Infos;
 }
 
 FSSFSampleSceneInfos InitSampleSceneInfos(
 	float4x4 ScreenToTranslatedWorld,
-	float2 ScreenPosition, float WorldDepth, float3 WorldNormal)
+	float2 ScreenPosition, float WorldDepth, float3 WorldNormal, float Roughness)
 {
 	FSSFSampleSceneInfos Infos;
 
@@ -121,6 +126,7 @@ FSSFSampleSceneInfos InitSampleSceneInfos(
 	Infos.TranslatedWorldPosition = mul(float4(ClipPosition, WorldDepth, 1), ScreenToTranslatedWorld).xyz;
 	Infos.WorldDepth = WorldDepth;
 	Infos.WorldNormal = WorldNormal;
+	Infos.Roughness = Roughness;
 	return Infos;
 }
 
@@ -129,7 +135,7 @@ FSSFSampleSceneInfos FetchCurrentSceneInfos(float2 ScreenPosition, float2 Buffer
 	FGBufferData GBufferData = GetGBufferDataFromSceneBlackboard(BufferUV);
 	return InitSampleSceneInfos(
 		View.ScreenToTranslatedWorld,
-		ScreenPosition, GBufferData.Depth, GBufferData.WorldNormal);
+		ScreenPosition, GBufferData.Depth, GBufferData.WorldNormal, GBufferData.Roughness);
 }
 
 float ProjectWorldDistanceToScreenDistance(FSSFSampleSceneInfos SceneMetadata, float WorldDistance)
diff --git a/Engine/Shaders/Private/ScreenSpaceDenoise/SSDReduce.usf b/Engine/Shaders/Private/ScreenSpaceDenoise/SSDReduce.usf
--- a/Engine/Shaders/Private/ScreenSpaceDenoise/SSDReduce.usf
+++ b/Engine/Shaders/Private/ScreenSpaceDenoise/SSDReduce.usf
@@ -289,10 +289,11 @@ void MainCS(
 						NeighborMipLevel = SharedArray0[LDSIndex].x;
 						float NeighborDepth = SharedArray0[LDSIndex].y;
 						float3 NeighborNormal = float3(0, 0, 1);
+						float Roughness = 1;
 
 						NeighborMetadata = InitSampleSceneInfos(
 							View.ScreenToTranslatedWorld,
-							NeighborScreenPosition, NeighborDepth, NeighborNormal);
+							NeighborScreenPosition, NeighborDepth, NeighborNormal, Roughness);
 					}
 				}
 
diff --git a/Engine/Shaders/Private/ScreenSpaceDenoise/SSDReflectionsSignal.ush b/Engine/Shaders/Private/ScreenSpaceDenoise/SSDReflectionsSignal.ush
--- a/Engine/Shaders/Private/ScreenSpaceDenoise/SSDReflectionsSignal.ush
+++ b/Engine/Shaders/Private/ScreenSpaceDenoise/SSDReflectionsSignal.ush
@@ -7,6 +7,49 @@
 #include "SSDSignalCore.ush"
 
 
+FSphericalGaussian RoughnessLobe(float Roughness, float3 N, float3 V)
+{
+	float a = Pow2( max( 0.02, Roughness ) );
+	float a2 = a*a;
+	
+	float NoV = saturate( abs( dot(N, V) ) + 1e-5 );
+	
+	float3 R = 2 * NoV * N - V;
+
+	FSphericalGaussian SpecularSG;
+	SpecularSG.Axis = R;
+	SpecularSG.Sharpness = 0.5 / ( a2 * max( NoV, 0.1 ) );
+	SpecularSG.Amplitude = rcp( PI * a2 );
+
+	return SpecularSG;
+}
+
+FSphericalGaussian RoughnessLobe(FSSFSampleSceneInfos RefSceneMetadata)
+{
+	float3 N = RefSceneMetadata.WorldNormal;
+	float3 V = normalize(View.TranslatedWorldCameraOrigin - RefSceneMetadata.TranslatedWorldPosition);
+	float Roughness = RefSceneMetadata.Roughness;
+
+	return RoughnessLobe(Roughness, N, V);
+}
+
+// [ Tokoyashi 2015, "Specular Lobe-Aware Filtering and Upsampling for Interactive Indirect Illumination" ]
+float TokoyashiLobeBilateralWeight(FSSFSampleSceneInfos RefSceneMetadata, FSSFSampleSceneInfos NeighborSceneMetadata)
+{
+	FSphericalGaussian Ref = RoughnessLobe(RefSceneMetadata);
+	FSphericalGaussian Neighbor = RoughnessLobe(NeighborSceneMetadata);
+	
+#if 1
+	const float Beta = 5.05; // TODO: to be tunned with GGX.
+
+	float CombinedSharpness = sqrt(Ref.Sharpness * Neighbor.Sharpness) / (Ref.Sharpness + Neighbor.Sharpness);
+
+	return pow(2 * CombinedSharpness, Beta) * exp((Beta * CombinedSharpness) * (dot(Ref.Axis, Neighbor.Axis) - 1));
+#else
+	return Dot(Ref, Neighbor);
+#endif
+}
+
 
 /** COmpute the weight of the bilateral filter to use between a reference sample and a neighbor. */
 float ComputeBilateralWeight(
@@ -20,7 +63,8 @@ float ComputeBilateralWeight(
 		return 1;
 	}
 	
-	float WorldBilateralWeight;
+	float WorldBilateralWeight = 1;
+	if (1)
 	{
 		const float WorldRadius = MaxWorldBluringDistance;
 
@@ -40,7 +84,13 @@ float ComputeBilateralWeight(
 		WorldBilateralWeight = saturate(1 - DistSquare * Multiplier);
 	}
 
-	return WorldBilateralWeight;
+	float LobeBilateralWeight = 1;
+	//if (View.GeneralPurposeTweak > 0.5)
+	{
+		LobeBilateralWeight = TokoyashiLobeBilateralWeight(RefSceneMetadata, NeighborSceneMetadata);
+	}
+
+	return WorldBilateralWeight * LobeBilateralWeight;
 }
 
 float GetSignalWorldBluringRadius(FSSFSignalSample Sample, FSSFSampleSceneInfos SceneMetadata)
diff --git a/Engine/Shaders/Private/ScreenSpaceDenoise/SSDSpatialAccumulation.usf b/Engine/Shaders/Private/ScreenSpaceDenoise/SSDSpatialAccumulation.usf
--- a/Engine/Shaders/Private/ScreenSpaceDenoise/SSDSpatialAccumulation.usf
+++ b/Engine/Shaders/Private/ScreenSpaceDenoise/SSDSpatialAccumulation.usf
@@ -19,43 +19,42 @@ PostProcessTemporalAA.usf: Temporal AA
 
 #define TILE_PIXEL_SIZE 8
 
+#define CONFIG_SIGNAL_PROCESSING DIM_SIGNAL_PROCESSING
 
-#if DIM_SIGNAL_PROCESSING == SIGNAL_PROCESSING_PENUMBRA
-	#define CONFIG_SIGNAL_INPUT_LAYOUT \
-		SIGNAL_BUFFER_LAYOUT_PENUMBRA_INTERNAL
 
-	#define CONFIG_SIGNAL_OUTPUT_LAYOUT \
-		SIGNAL_BUFFER_LAYOUT_PENUMBRA_HISTORY_REJECTION
+#if CONFIG_SIGNAL_PROCESSING == SIGNAL_PROCESSING_PENUMBRA
+	#define CONFIG_SIGNAL_INPUT_LAYOUT  SIGNAL_BUFFER_LAYOUT_PENUMBRA_INTERNAL
+	#define CONFIG_SIGNAL_OUTPUT_LAYOUT SIGNAL_BUFFER_LAYOUT_PENUMBRA_HISTORY_REJECTION
 
-	#define CONFIG_SIGNAL_PROCESSING \
-		SIGNAL_PROCESSING_PENUMBRA
+	#define CONFIG_SAMPLE_SET SAMPLE_SET_3X3
 
-	#define CONFIG_SAMPLE_SET \
-		SAMPLE_SET_3X3
 
+#elif CONFIG_SIGNAL_PROCESSING == SIGNAL_PROCESSING_REFLECTIONS
+	#define CONFIG_SIGNAL_INPUT_LAYOUT  SIGNAL_BUFFER_LAYOUT_REFLECTIONS_HISTORY
+	#define CONFIG_SIGNAL_OUTPUT_LAYOUT SIGNAL_BUFFER_LAYOUT_REFLECTIONS_HISTORY
 
-#elif DIM_SIGNAL_PROCESSING == SIGNAL_PROCESSING_REFLECTIONS
-	#define CONFIG_SIGNAL_INPUT_LAYOUT \
-		SIGNAL_BUFFER_LAYOUT_REFLECTIONS_INPUT
-
-	#define CONFIG_SIGNAL_OUTPUT_LAYOUT \
-		SIGNAL_BUFFER_LAYOUT_REFLECTIONS_OUTPUT
-
-	#define CONFIG_SIGNAL_PROCESSING \
-		SIGNAL_PROCESSING_REFLECTIONS
-
-	#define CONFIG_SAMPLE_SET \
-		SAMPLE_SET_3X3
+	#define CONFIG_SAMPLE_SET SAMPLE_SET_3X3
 
+	#define CONFIG_WEIGHT_KERNEL_CENTER 1
 
 #else
 	#error Unknown signal processing.
 #endif
 
+#include "SSDSignalInterface.ush"
 
-//------------------------------------------------------- CONFIGS
 
-#include "SSDSignalInterface.ush"
+//------------------------------------------------------- CONFIG DISABLED DEFAULTS
+
+/** Whether the output signal should be normalized. */
+#ifndef CONFIG_NORMALIZE_OUTPUT
+	#define CONFIG_NORMALIZE_OUTPUT 0
+#endif
+
+/** Ensures the center of the kernel have a correct weight. */
+#ifndef CONFIG_WEIGHT_KERNEL_CENTER
+	#define CONFIG_WEIGHT_KERNEL_CENTER 0
+#endif
 
 
 //------------------------------------------------------- CONSTANTS
@@ -100,7 +99,13 @@ void MainCS(
 
 	float CurrentWorldBluringRadius = ComputePixelWorldRadius(CurrentSceneMetadata);
 
-	float AccumulatorWeight = 1.0;
+	float AccumulatorWeight = 1;
+	float BilateralWeightAdjustement = 1;
+	#if CONFIG_WEIGHT_KERNEL_CENTER
+	{
+		BilateralWeightAdjustement = rcp(ComputeBilateralWeight(100, CurrentSceneMetadata, CurrentSceneMetadata));
+	}
+	#endif
 
 	#if CONFIG_SAMPLE_SET == SAMPLE_SET_3X3
 		const int KernelRadius = 1;
@@ -140,7 +145,7 @@ void MainCS(
 		float NeighborWorldBluringRadius = ComputePixelWorldRadius(NeighborSceneMetadata);
 
 		// Compute the bilateral weight.
-		float BilateralWeight = ComputeBilateralWeight(
+		float BilateralWeight = BilateralWeightAdjustement * ComputeBilateralWeight(
 			WorldBluringRadiusToBilateralWorldDistance(min(CurrentWorldBluringRadius, NeighborWorldBluringRadius)),
 			CurrentSceneMetadata,
 			NeighborSceneMetadata);
@@ -168,6 +173,12 @@ void MainCS(
 		Accumulator = Add(Accumulator, Mul(NeighborSample, FinalWeight));
 	}
 
+	#if CONFIG_NORMALIZE_OUTPUT
+	{
+		Accumulator = NormalizeToOneSample(Accumulator);
+	}
+	#endif
+
 	// TODO:  LeaveRayCount = (LeaveRayCount - 1) * 9 / (9 - 2) post processing to reject when for history rejection
 
 	OutputSignal(SignalOutput0, SignalOutput1, CONFIG_SIGNAL_OUTPUT_LAYOUT, DispatchThreadId, Accumulator);
diff --git a/Engine/Shaders/Private/ScreenSpaceDenoise/SSDTemporalAccumulation.usf b/Engine/Shaders/Private/ScreenSpaceDenoise/SSDTemporalAccumulation.usf
--- a/Engine/Shaders/Private/ScreenSpaceDenoise/SSDTemporalAccumulation.usf
+++ b/Engine/Shaders/Private/ScreenSpaceDenoise/SSDTemporalAccumulation.usf
@@ -47,7 +47,7 @@
 	#define CONFIG_SIGNAL_INPUT_LAYOUT   SIGNAL_BUFFER_LAYOUT_REFLECTIONS_INPUT
 	#define CONFIG_SIGNAL_HISTORY_LAYOUT SIGNAL_BUFFER_LAYOUT_REFLECTIONS_HISTORY
 
-	#define CONFIG_HISTORY_REJECTION 1
+	#define CONFIG_HISTORY_REJECTION 2
 	#define CONFIG_HISTORY_REJECTION_COLOR_SPACE 0
 
 #else
@@ -138,6 +138,7 @@ Texture2D PrevHistory1;
 Texture2D<uint> PrevTileClassificationTexture;
 Texture2D PrevDepthBuffer;
 Texture2D PrevGBufferA;
+Texture2D PrevGBufferB;
 
 Texture2D HistoryRejectionSignal0;
 Texture2D HistoryRejectionSignal1;
@@ -179,10 +180,11 @@ void PointSampleHistory(float2 ScreenPosition, float2 BufferUV, out FSSFSignalSa
 	// TODO: so what to do with mip level of these guys?
 	float WorldSceneDepth = ConvertFromDeviceZ(PrevDepthBuffer.SampleLevel(GlobalPointClampedSampler, BufferUV, 0).r);
 	float3 WorldNormal = DecodeNormal(PrevGBufferA.SampleLevel(GlobalPointClampedSampler, BufferUV, 0).rgb);
+	float Roughness = PrevGBufferB.SampleLevel(GlobalPointClampedSampler, BufferUV, 0).b;
 
 	OutMetadata = InitSampleSceneInfos(
 		View.PrevScreenToTranslatedWorld,
-		ScreenPosition, WorldSceneDepth, WorldNormal);
+		ScreenPosition, WorldSceneDepth, WorldNormal, Roughness);
 
 	// TODO: PrevScreenToCurrentTranslatedWorld
 	OutMetadata.TranslatedWorldPosition += -View.PrevPreViewTranslation.xyz + View.PreViewTranslation.xyz;
@@ -411,10 +413,11 @@ void TemporallyAccumulate(
 							{
 								float WorldSceneDepth = ConvertFromDeviceZ(PrevDepthBuffer.SampleLevel(GlobalPointClampedSampler, SampleHistoryBufferUV2, 0).r);
 								float3 WorldNormal = DecodeNormal(PrevGBufferA.SampleLevel(GlobalPointClampedSampler, SampleHistoryBufferUV2, 0).rgb);
+								float Roughness = PrevGBufferB.SampleLevel(GlobalPointClampedSampler, SampleHistoryBufferUV2, 0).b;
 
 								HistoryMetadata2 = InitSampleSceneInfos(
 									View.PrevScreenToTranslatedWorld,
-									SampleScreenPosition2, WorldSceneDepth, WorldNormal);
+									SampleScreenPosition2, WorldSceneDepth, WorldNormal, Roughness);
 
 								// TODO: PrevScreenToCurrentTranslatedWorld
 								HistoryMetadata2.TranslatedWorldPosition += -View.PrevPreViewTranslation.xyz + View.PreViewTranslation.xyz;
@@ -496,10 +499,11 @@ void TemporallyAccumulate(
 					{
 						float WorldSceneDepth = ConvertFromDeviceZ(PrevDepthBuffer.SampleLevel(GlobalPointClampedSampler, SampleHistoryBufferUV, 0).r);
 						float3 WorldNormal = DecodeNormal(PrevGBufferA.SampleLevel(GlobalPointClampedSampler, SampleHistoryBufferUV, 0).rgb);
+						float Roughness = PrevGBufferB.SampleLevel(GlobalPointClampedSampler, SampleHistoryBufferUV, 0).b;
 
 						HistoryMetadata = InitSampleSceneInfos(
 							View.PrevScreenToTranslatedWorld,
-							SampleScreenPosition, WorldSceneDepth, WorldNormal);
+							SampleScreenPosition, WorldSceneDepth, WorldNormal, Roughness);
 
 						// TODO: PrevScreenToCurrentTranslatedWorld
 						HistoryMetadata.TranslatedWorldPosition += -View.PrevPreViewTranslation.xyz + View.PreViewTranslation.xyz;
diff --git a/Engine/Source/Runtime/Renderer/Private/SceneRendering.h b/Engine/Source/Runtime/Renderer/Private/SceneRendering.h
--- a/Engine/Source/Runtime/Renderer/Private/SceneRendering.h
+++ b/Engine/Source/Runtime/Renderer/Private/SceneRendering.h
@@ -843,6 +843,7 @@ struct FPreviousViewInfo
 	// Depth buffer and Normals of the previous frame generating this history entry for bilateral kernel rejection.
 	TRefCountPtr<IPooledRenderTarget> DepthBuffer;
 	TRefCountPtr<IPooledRenderTarget> GBufferA;
+	TRefCountPtr<IPooledRenderTarget> GBufferB;
 
 	// Temporal AA result of last frame
 	FTemporalAAHistory TemporalAAHistory;
@@ -867,6 +868,7 @@ struct FPreviousViewInfo
 	{
 		DepthBuffer.SafeRelease();
 		GBufferA.SafeRelease();
+		GBufferB.SafeRelease();
 		TemporalAAHistory.SafeRelease();
 		DOFPreGatherHistory.SafeRelease();
 		DOFPostGatherForegroundHistory.SafeRelease();
diff --git a/Engine/Source/Runtime/Renderer/Private/SceneViewFamilyBlackboard.cpp b/Engine/Source/Runtime/Renderer/Private/SceneViewFamilyBlackboard.cpp
--- a/Engine/Source/Runtime/Renderer/Private/SceneViewFamilyBlackboard.cpp
+++ b/Engine/Source/Runtime/Renderer/Private/SceneViewFamilyBlackboard.cpp
@@ -15,11 +15,12 @@ void SetupSceneViewFamilyBlackboard(
 {
 	const FSceneRenderTargets& SceneContext = FSceneRenderTargets::Get(GraphBuilder.RHICmdList);
 
-	FRDGTextureRef BlackDummy = GraphBuilder.RegisterExternalTexture(GSystemTextures.BlackDummy);
-	FRDGTextureRef DefaultNormal8Bit = GraphBuilder.RegisterExternalTexture(GSystemTextures.DefaultNormal8Bit);
-
 	*OutBlackboard = FSceneViewFamilyBlackboard();
 	OutBlackboard->SceneDepthBuffer = GraphBuilder.RegisterExternalTexture(SceneContext.SceneDepthZ);
-	OutBlackboard->SceneVelocityBuffer = SceneContext.SceneVelocity ? GraphBuilder.RegisterExternalTexture(SceneContext.SceneVelocity) : BlackDummy;
-	OutBlackboard->SceneGBufferA = SceneContext.GBufferA ? GraphBuilder.RegisterExternalTexture(SceneContext.GBufferA) : DefaultNormal8Bit;
+	OutBlackboard->SceneVelocityBuffer = RegisterExternalTextureWithFallback(GraphBuilder, SceneContext.SceneVelocity, GSystemTextures.BlackDummy);
+	OutBlackboard->SceneGBufferA = RegisterExternalTextureWithFallback(GraphBuilder, SceneContext.GBufferA, GSystemTextures.DefaultNormal8Bit, TEXT("GBufferA"));
+	OutBlackboard->SceneGBufferB = RegisterExternalTextureWithFallback(GraphBuilder, SceneContext.GBufferB, GSystemTextures.BlackDummy, TEXT("GBufferB"));
+	OutBlackboard->SceneGBufferC = RegisterExternalTextureWithFallback(GraphBuilder, SceneContext.GBufferC, GSystemTextures.BlackDummy, TEXT("GBufferC"));
+	OutBlackboard->SceneGBufferD = RegisterExternalTextureWithFallback(GraphBuilder, SceneContext.GBufferD, GSystemTextures.BlackDummy, TEXT("GBufferD"));
+	OutBlackboard->SceneGBufferE = RegisterExternalTextureWithFallback(GraphBuilder, SceneContext.GBufferE, GSystemTextures.BlackDummy, TEXT("GBufferE"));
 }
diff --git a/Engine/Source/Runtime/Renderer/Private/SceneViewFamilyBlackboard.h b/Engine/Source/Runtime/Renderer/Private/SceneViewFamilyBlackboard.h
--- a/Engine/Source/Runtime/Renderer/Private/SceneViewFamilyBlackboard.h
+++ b/Engine/Source/Runtime/Renderer/Private/SceneViewFamilyBlackboard.h
@@ -14,6 +14,10 @@ BEGIN_SHADER_PARAMETER_STRUCT(FSceneViewFamilyBlackboard, )
 	SHADER_PARAMETER_RDG_TEXTURE(Texture2D, SceneDepthBuffer)
 	SHADER_PARAMETER_RDG_TEXTURE(Texture2D, SceneVelocityBuffer)
 	SHADER_PARAMETER_RDG_TEXTURE(Texture2D, SceneGBufferA)
+	SHADER_PARAMETER_RDG_TEXTURE(Texture2D, SceneGBufferB)
+	SHADER_PARAMETER_RDG_TEXTURE(Texture2D, SceneGBufferC)
+	SHADER_PARAMETER_RDG_TEXTURE(Texture2D, SceneGBufferD)
+	SHADER_PARAMETER_RDG_TEXTURE(Texture2D, SceneGBufferE)
 END_SHADER_PARAMETER_STRUCT()
 
 
diff --git a/Engine/Source/Runtime/Renderer/Private/ScreenSpaceDenoise.cpp b/Engine/Source/Runtime/Renderer/Private/ScreenSpaceDenoise.cpp
--- a/Engine/Source/Runtime/Renderer/Private/ScreenSpaceDenoise.cpp
+++ b/Engine/Source/Runtime/Renderer/Private/ScreenSpaceDenoise.cpp
@@ -33,6 +33,13 @@ static TAutoConsoleVariable<int32> CVarDoDebugPass(
 	TEXT("Adds denoiser's the debug pass."),
 	ECVF_RenderThreadSafe);
 
+static TAutoConsoleVariable<int32> CVarReflectionHistoryConvolution(
+	TEXT("r.Reflection.Denoise.HistoryConvolution"), 1,
+	TEXT("Mode to use for history convolution.\n")
+	TEXT(" 0: disabled;\n")
+	TEXT(" 1: Spatial convolution (default)."),
+	ECVF_RenderThreadSafe);
+
 
 /** The maximum number of mip level supported in the denoiser. */
 static const int32 kMaxMipLevel = 4;
@@ -171,6 +178,8 @@ class FSSDSpatialAccumulationCS : public FScreenSpaceDenoisingShader
 {
 	DECLARE_GLOBAL_SHADER(FSSDSpatialAccumulationCS);
 	SHADER_USE_PARAMETER_STRUCT(FSSDSpatialAccumulationCS, FScreenSpaceDenoisingShader);
+	
+	using FPermutationDomain = TShaderPermutationDomain<FSignalProcessingDim>;
 
 	BEGIN_SHADER_PARAMETER_STRUCT(FParameters, )
 		SHADER_PARAMETER_STRUCT_INCLUDE(FSSDCommonParameters, CommonParameters)
@@ -227,6 +236,7 @@ class FSSDTemporalAccumulationCS : public FScreenSpaceDenoisingShader
 		SHADER_PARAMETER_RDG_TEXTURE(Texture2D, PrevTileClassificationTexture)
 		SHADER_PARAMETER_RDG_TEXTURE(Texture2D, PrevDepthBuffer)
 		SHADER_PARAMETER_RDG_TEXTURE(Texture2D, PrevGBufferA)
+		SHADER_PARAMETER_RDG_TEXTURE(Texture2D, PrevGBufferB)
 
 		SHADER_PARAMETER_RDG_TEXTURE(Texture2D, HistoryRejectionSignal0)
 		SHADER_PARAMETER_RDG_TEXTURE(Texture2D, HistoryRejectionSignal1)
@@ -401,7 +411,10 @@ static void DenoiseShadowPenumbra(
 		PassParameters->SignalOutput0 = GraphBuilder.CreateUAV(FRDGTextureUAVDesc(SignalOutput0, /* MipLevel = */ 0));
 		PassParameters->SignalOutput1 = GraphBuilder.CreateUAV(FRDGTextureUAVDesc(SignalOutput1, /* MipLevel = */ 0));
 
-		TShaderMapRef<FSSDSpatialAccumulationCS> ComputeShader(View.ShaderMap);
+		FSSDSpatialAccumulationCS::FPermutationDomain PermutationVector;
+		PermutationVector.Set<FSignalProcessingDim>(ESignalProcessing::Penumbra);
+
+		TShaderMapRef<FSSDSpatialAccumulationCS> ComputeShader(View.ShaderMap, PermutationVector);
 		FComputeShaderUtils::AddPass(
 			GraphBuilder,
 			RDG_EVENT_NAME("SSD SpatialAccumulation(Mip=0)"),
@@ -599,6 +612,7 @@ IScreenSpaceDenoiser::FReflectionOutputs DenoiseReflections(
 	FRDGTextureRef SignalHistory0 = GraphBuilder.CreateTexture(SignalProcessingDesc, TEXT("SSDReflectionsHistory0"));
 	FRDGTextureRef SignalHistory1 = GraphBuilder.CreateTexture(SignalProcessingDesc, TEXT("SSDReflectionsHistory1"));
 
+	// Temporal pass.
 	{
 		const int32 MipLevel = 0;
 
@@ -610,11 +624,6 @@ IScreenSpaceDenoiser::FReflectionOutputs DenoiseReflections(
 
 		FSSDTemporalAccumulationCS::FParameters* PassParameters = GraphBuilder.AllocParameters<FSSDTemporalAccumulationCS::FParameters>();
 		PassParameters->CommonParameters = CommonParameters;
-		PassParameters->iMipLevel = MipLevel;
-		PassParameters->iMipLevelPow2 = (1 << MipLevel);
-		PassParameters->MipLevel = MipLevel;
-		PassParameters->MipLevelPow2 = (1 << MipLevel);
-		PassParameters->InvMipLevelPow2 = 1.0f / PassParameters->MipLevelPow2;
 		PassParameters->bCameraCut = View.bCameraCut || !PrevFrameHistory.RT[0].IsValid();
 		PassParameters->SignalInput0 = ReflectionInputs.Color;
 		PassParameters->SignalInput1 = ReflectionInputs.RayHitDistance;
@@ -625,26 +634,54 @@ IScreenSpaceDenoiser::FReflectionOutputs DenoiseReflections(
 		PassParameters->PrevTileClassificationTexture = RegisterExternalTextureWithFallback(GraphBuilder, PrevFrameHistory.TileClassification, GSystemTextures.BlackDummy);
 		PassParameters->PrevDepthBuffer = RegisterExternalTextureWithFallback(GraphBuilder, View.PrevViewInfo.DepthBuffer, GSystemTextures.BlackDummy);
 		PassParameters->PrevGBufferA = RegisterExternalTextureWithFallback(GraphBuilder, View.PrevViewInfo.GBufferA, GSystemTextures.BlackDummy);
+		PassParameters->PrevGBufferB = RegisterExternalTextureWithFallback(GraphBuilder, View.PrevViewInfo.GBufferB, GSystemTextures.BlackDummy);
 
 		PassParameters->SignalHistoryOutput0 = GraphBuilder.CreateUAV(FRDGTextureUAVDesc(SignalHistory0, MipLevel));
 		PassParameters->SignalHistoryOutput1 = GraphBuilder.CreateUAV(FRDGTextureUAVDesc(SignalHistory1, MipLevel));
 
-		FIntPoint Resolution = FIntPoint::DivideAndRoundUp(DenoiseResolution, 1 << MipLevel);
+		FComputeShaderUtils::AddPass(
+			GraphBuilder,
+			RDG_EVENT_NAME("SSD TemporalAccumulation"),
+			*ComputeShader,
+			PassParameters,
+			FComputeShaderUtils::GetGroupCount(DenoiseResolution, FComputeShaderUtils::kGolden2DGroupSize));
+	}
+	
+	// Spatial filter, to converge history faster.
+	if (CVarReflectionHistoryConvolution.GetValueOnRenderThread())
+	{
+		FRDGTextureRef SignalOutput0 = GraphBuilder.CreateTexture(SignalProcessingDesc, TEXT("SSDReflectionsHistory0"));
+		FRDGTextureRef SignalOutput1 = GraphBuilder.CreateTexture(SignalProcessingDesc, TEXT("SSDReflectionsHistory1"));
+
+		FSSDSpatialAccumulationCS::FParameters* PassParameters = GraphBuilder.AllocParameters<FSSDSpatialAccumulationCS::FParameters>();
+		PassParameters->CommonParameters = CommonParameters;
+		PassParameters->SignalInput0 = SignalHistory0;
+		PassParameters->SignalInput1 = SignalHistory1;
+		PassParameters->SignalOutput0 = GraphBuilder.CreateUAV(SignalOutput0);
+		PassParameters->SignalOutput1 = GraphBuilder.CreateUAV(SignalOutput1);
 
+		FSSDSpatialAccumulationCS::FPermutationDomain PermutationVector;
+		PermutationVector.Set<FSignalProcessingDim>(ESignalProcessing::Reflections);
+
+		TShaderMapRef<FSSDSpatialAccumulationCS> ComputeShader(View.ShaderMap, PermutationVector);
 		FComputeShaderUtils::AddPass(
 			GraphBuilder,
-			RDG_EVENT_NAME("SSD TemporalAccumulation(Mip=%d)", MipLevel),
+			RDG_EVENT_NAME("SSD SpatialAccumulation"),
 			*ComputeShader,
 			PassParameters,
-			FComputeShaderUtils::GetGroupCount(Resolution, FComputeShaderUtils::kGolden2DGroupSize));
+			FComputeShaderUtils::GetGroupCount(DenoiseResolution, FComputeShaderUtils::kGolden2DGroupSize));
+
+		SignalHistory0 = SignalOutput0;
+		SignalHistory1 = SignalOutput1;
 	}
 
 	if (View.ViewState)
 	{
-		// Keep depth buffer and GBufferA around for next frame.
+		// Keep depth buffer and GBuffer around for next frame.
 		{
 			GraphBuilder.QueueTextureExtraction(SceneBlackboard.SceneDepthBuffer, &View.ViewState->PendingPrevFrameViewInfo.DepthBuffer);
 			GraphBuilder.QueueTextureExtraction(SceneBlackboard.SceneGBufferA, &View.ViewState->PendingPrevFrameViewInfo.GBufferA);
+			GraphBuilder.QueueTextureExtraction(SceneBlackboard.SceneGBufferB, &View.ViewState->PendingPrevFrameViewInfo.GBufferB);
 		}
 
 		// Saves history.
