Fixed linker error due to FDeferredShadingSceneRenderer::RenderDebugViewMode missing symbol in shipping configs

diff --git a/Engine/Source/Runtime/Renderer/Private/DebugViewModeRendering.cpp b/Engine/Source/Runtime/Renderer/Private/DebugViewModeRendering.cpp
--- a/Engine/Source/Runtime/Renderer/Private/DebugViewModeRendering.cpp
+++ b/Engine/Source/Runtime/Renderer/Private/DebugViewModeRendering.cpp
@@ -420,6 +420,13 @@ void InitDebugViewModeInterfaces()
 	FDebugViewModeInterface::SetInterface(DVSM_RequiredTextureResolution, new FRequiredTextureResolutionInterface());
 }
 
+#else // !(UE_BUILD_SHIPPING || UE_BUILD_TEST)
+
+bool FDeferredShadingSceneRenderer::RenderDebugViewMode(FRHICommandListImmediate& RHICmdList)
+{
+	return false;
+}
+
 #endif // !(UE_BUILD_SHIPPING || UE_BUILD_TEST)
 
 
