Add support for r.ForceLOD in ray tracing path

diff --git a/Engine/Source/Runtime/Renderer/Private/DeferredShadingRenderer.cpp b/Engine/Source/Runtime/Renderer/Private/DeferredShadingRenderer.cpp
--- a/Engine/Source/Runtime/Renderer/Private/DeferredShadingRenderer.cpp
+++ b/Engine/Source/Runtime/Renderer/Private/DeferredShadingRenderer.cpp
@@ -655,7 +655,7 @@ bool FDeferredShadingSceneRenderer::GatherRayTracingWorldInstances(FRHICommandLi
 
 						FLODMask LODToRender;
 						float MeshScreenSizeSquared = 0;
-						int32 ForcedLODLevel = -1;
+						int32 ForcedLODLevel = GetCVarForceLOD();
 						if (SceneProxy->IsUsingCustomLODRules())
 						{
 							LODToRender = SceneProxy->GetCustomLOD(View, View.LODDistanceFactor, ForcedLODLevel, MeshScreenSizeSquared);
