Implements r.Reflection.Denoise.TemporalAccumulation

diff --git a/Engine/Source/Runtime/Renderer/Private/ScreenSpaceDenoise.cpp b/Engine/Source/Runtime/Renderer/Private/ScreenSpaceDenoise.cpp
--- a/Engine/Source/Runtime/Renderer/Private/ScreenSpaceDenoise.cpp
+++ b/Engine/Source/Runtime/Renderer/Private/ScreenSpaceDenoise.cpp
@@ -38,6 +38,11 @@ static TAutoConsoleVariable<int32> CVarReflectionReconstructionSampleCount(
 	TEXT("Maximum number of samples for the reconstruction pass (default = 16)."),
 	ECVF_RenderThreadSafe);
 
+static TAutoConsoleVariable<int32> CVarReflectionTemporalAccumulation(
+	TEXT("r.Reflection.Denoise.TemporalAccumulation"), 1,
+	TEXT(""),
+	ECVF_RenderThreadSafe);
+
 static TAutoConsoleVariable<int32> CVarReflectionHistoryConvolution(
 	TEXT("r.Reflection.Denoise.HistoryConvolution"), 1,
 	TEXT("Mode to use for history convolution.\n")
@@ -680,6 +685,7 @@ IScreenSpaceDenoiser::FReflectionOutputs DenoiseReflections(
 	}
 
 	// Temporal pass.
+	if (CVarReflectionTemporalAccumulation.GetValueOnRenderThread())
 	{
 		FRDGTextureRef SignalOutput0 = GraphBuilder.CreateTexture(SignalProcessingDesc, TEXT("SSDReflectionsHistory0"));
 		FRDGTextureRef SignalOutput1 = GraphBuilder.CreateTexture(SignalProcessingDesc, TEXT("SSDReflectionsHistory1"));
