Enables the reflection denoiser by default.

diff --git a/Engine/Source/Runtime/Renderer/Private/ReflectionEnvironment.cpp b/Engine/Source/Runtime/Renderer/Private/ReflectionEnvironment.cpp
--- a/Engine/Source/Runtime/Renderer/Private/ReflectionEnvironment.cpp
+++ b/Engine/Source/Runtime/Renderer/Private/ReflectionEnvironment.cpp
@@ -111,11 +111,11 @@ static FAutoConsoleVariableRef CVarRayTracingReflectionss(
 
 static TAutoConsoleVariable<int32> CVarUseReflectionDenoiser(
 	TEXT("r.Reflection.Denoiser"),
-	0, // TODO: change default to 2.
+	2,
 	TEXT("Choose the denoising algorithm.\n")
-	TEXT(" 0: Disabled (default);\n")
+	TEXT(" 0: Disabled;\n")
 	TEXT(" 1: Forces the default denoiser of the renderer;\n")
-	TEXT(" 2: GScreenSpaceDenoiser witch may be overriden by a third party plugin.\n"),
+	TEXT(" 2: GScreenSpaceDenoiser witch may be overriden by a third party plugin (default).\n"),
 	ECVF_RenderThreadSafe);
 
 
diff --git a/Engine/Source/Runtime/Renderer/Private/ScreenSpaceDenoise.cpp b/Engine/Source/Runtime/Renderer/Private/ScreenSpaceDenoise.cpp
--- a/Engine/Source/Runtime/Renderer/Private/ScreenSpaceDenoise.cpp
+++ b/Engine/Source/Runtime/Renderer/Private/ScreenSpaceDenoise.cpp
@@ -44,7 +44,7 @@ static TAutoConsoleVariable<int32> CVarReflectionTemporalAccumulation(
 	ECVF_RenderThreadSafe);
 
 static TAutoConsoleVariable<int32> CVarReflectionHistoryConvolution(
-	TEXT("r.Reflection.Denoise.HistoryConvolution"), 1,
+	TEXT("r.Reflection.Denoise.HistoryConvolution"), 0, // TODO: enable by default
 	TEXT("Mode to use for history convolution.\n")
 	TEXT(" 0: disabled;\n")
 	TEXT(" 1: Spatial convolution (default)."),
