Moves r.RayTracing.Reflections.SamplesPerPixel to ReflectionEnvironement to have convenient swap between many spp as reference and denoised 1spp.

diff --git a/Engine/Source/Runtime/Renderer/Private/DeferredShadingRenderer.h b/Engine/Source/Runtime/Renderer/Private/DeferredShadingRenderer.h
--- a/Engine/Source/Runtime/Renderer/Private/DeferredShadingRenderer.h
+++ b/Engine/Source/Runtime/Renderer/Private/DeferredShadingRenderer.h
@@ -444,7 +444,8 @@ private:
 		FRDGBuilder& GraphBuilder,
 		const FViewInfo& View,
 		FRDGTextureRef* OutColorTexture,
-		FRDGTextureRef* OutRayHitDistanceTexture);
+		FRDGTextureRef* OutRayHitDistanceTexture,
+		int32 SamplePerPixel);
 
 #if RHI_RAYTRACING
 	bool ShouldRenderRayTracingAmbientOcclusion() const;
diff --git a/Engine/Source/Runtime/Renderer/Private/RayTracing/RayTracingReflections.cpp b/Engine/Source/Runtime/Renderer/Private/RayTracing/RayTracingReflections.cpp
--- a/Engine/Source/Runtime/Renderer/Private/RayTracing/RayTracingReflections.cpp
+++ b/Engine/Source/Runtime/Renderer/Private/RayTracing/RayTracingReflections.cpp
@@ -16,13 +16,6 @@
 #include "LightRendering.h"
 #include "SystemTextures.h"
 
-static int32 GRayTracingReflectionsSamplesPerPixel = 1;
-static FAutoConsoleVariableRef CVarRayTracingReflectionsSamplesPerPixel(
-	TEXT("r.RayTracing.Reflections.SamplesPerPixel"),
-	GRayTracingReflectionsSamplesPerPixel,
-	TEXT("Sets the samples-per-pixel for reflections (default = 1)")
-);
-
 static int32 GRayTracingReflectionsEmissiveAndIndirectLighting = 1;
 static FAutoConsoleVariableRef CVarRayTracingReflectionsEmissiveAndIndirectLighting(
 	TEXT("r.RayTracing.Reflections.EmissiveAndIndirectLighting"),
@@ -193,7 +186,8 @@ void FDeferredShadingSceneRenderer::RayTraceReflections(
 	FRDGBuilder& GraphBuilder,
 	const FViewInfo& View,
 	FRDGTextureRef* OutColorTexture,
-	FRDGTextureRef* OutRayHitDistanceTexture)
+	FRDGTextureRef* OutRayHitDistanceTexture,
+	int32 SamplePerPixel)
 #if RHI_RAYTRACING
 {
 	FSceneRenderTargets& SceneContext = FSceneRenderTargets::Get(GraphBuilder.RHICmdList);
@@ -211,7 +205,7 @@ void FDeferredShadingSceneRenderer::RayTraceReflections(
 
 	FRayTracingReflectionsRG::FParameters* PassParameters = GraphBuilder.AllocParameters<FRayTracingReflectionsRG::FParameters>();
 
-	PassParameters->SamplesPerPixel = GRayTracingReflectionsSamplesPerPixel;
+	PassParameters->SamplesPerPixel = SamplePerPixel;
 	PassParameters->ShouldDoDirectLighting = GRayTracingReflectionsDirectLighting;
 	PassParameters->ShouldDoReflectedShadows = GRayTracingReflectionsShadows;
 	PassParameters->ShouldDoEmissiveAndIndirectLighting = GRayTracingReflectionsEmissiveAndIndirectLighting;
diff --git a/Engine/Source/Runtime/Renderer/Private/ReflectionEnvironment.cpp b/Engine/Source/Runtime/Renderer/Private/ReflectionEnvironment.cpp
--- a/Engine/Source/Runtime/Renderer/Private/ReflectionEnvironment.cpp
+++ b/Engine/Source/Runtime/Renderer/Private/ReflectionEnvironment.cpp
@@ -109,13 +109,20 @@ static FAutoConsoleVariableRef CVarRayTracingReflectionss(
 	TEXT("1: use ray traced reflections (default when ray tracing is enabled)\n")
 );
 
+static int32 GRayTracingReflectionsSamplesPerPixel = 1;
+static FAutoConsoleVariableRef CVarRayTracingReflectionsSamplesPerPixel(
+	TEXT("r.RayTracing.Reflections.SamplesPerPixel"),
+	GRayTracingReflectionsSamplesPerPixel,
+	TEXT("Sets the samples-per-pixel for reflections (default = 1)")
+);
+
 static TAutoConsoleVariable<int32> CVarUseReflectionDenoiser(
 	TEXT("r.Reflection.Denoiser"),
 	2,
 	TEXT("Choose the denoising algorithm.\n")
 	TEXT(" 0: Disabled;\n")
 	TEXT(" 1: Forces the default denoiser of the renderer;\n")
-	TEXT(" 2: GScreenSpaceDenoiser witch may be overriden by a third party plugin (default).\n"),
+	TEXT(" 2: GScreenSpaceDenoiser witch may be overriden by a third party plugin (default)."),
 	ECVF_RenderThreadSafe);
 
 
@@ -740,11 +747,11 @@ void FDeferredShadingSceneRenderer::RenderDeferredReflectionsAndSkyLighting(FRHI
 
 			// Ray trace the reflection.
 			IScreenSpaceDenoiser::FReflectionInputs DenoiserInputs;
-			RayTraceReflections(GraphBuilder, View, &DenoiserInputs.Color, &DenoiserInputs.RayHitDistance);
+			RayTraceReflections(GraphBuilder, View, &DenoiserInputs.Color, &DenoiserInputs.RayHitDistance, GRayTracingReflectionsSamplesPerPixel);
 
 
 			// Denoise the reflections.
-			if (int32 DenoiserMode = CVarUseReflectionDenoiser.GetValueOnRenderThread())
+			if (int32 DenoiserMode = CVarUseReflectionDenoiser.GetValueOnRenderThread() && GRayTracingReflectionsSamplesPerPixel == 1)
 			{
 				const IScreenSpaceDenoiser* DefaultDenoiser = IScreenSpaceDenoiser::GetDefaultDenoiser();
 				const IScreenSpaceDenoiser* DenoiserToUse = DenoiserMode == 1 ? DefaultDenoiser : GScreenSpaceDenoiser;
