Wrapped uint32 for mesh pass mask into a separate class - FMeshPassMask.

diff --git a/Engine/Source/Runtime/Renderer/Private/MeshPassProcessor.cpp b/Engine/Source/Runtime/Renderer/Private/MeshPassProcessor.cpp
--- a/Engine/Source/Runtime/Renderer/Private/MeshPassProcessor.cpp
+++ b/Engine/Source/Runtime/Renderer/Private/MeshPassProcessor.cpp
@@ -852,7 +852,8 @@ FMeshPassProcessor::FMeshPassProcessor(const FScene* InScene, ERHIFeatureLevel::
 	, ViewIfDynamicMeshCommand(InViewIfDynamicMeshCommand)
 	, DrawListContext(InDrawListContext)
 {
-	static_assert(EMeshPass::Num <= EMeshPass::NumBits * 8, "EMeshPass::NumBits is too small.");
+	static_assert(EMeshPass::Num <= EMeshPass::NumBits * 8, "EMeshPass::NumBits is too small to fit all mesh passes.");
+	static_assert(sizeof(FMeshPassMask::Data) * 8 >= EMeshPass::Num, "FMeshPassMask::Data is too small to fit all mesh passes.");
 }
 
 ERasterizerFillMode FMeshPassProcessor::ComputeMeshFillMode(const FMeshBatch& Mesh, const FMaterial& InMaterialResource) const
diff --git a/Engine/Source/Runtime/Renderer/Private/PrimitiveSceneInfo.cpp b/Engine/Source/Runtime/Renderer/Private/PrimitiveSceneInfo.cpp
--- a/Engine/Source/Runtime/Renderer/Private/PrimitiveSceneInfo.cpp
+++ b/Engine/Source/Runtime/Renderer/Private/PrimitiveSceneInfo.cpp
@@ -202,7 +202,7 @@ void FPrimitiveSceneInfo::CacheMeshDrawCommands(FRHICommandListImmediate& RHICmd
 			FStaticMeshRelevance& MeshRelevance = StaticMeshRelevances[MeshIndex];
 			FStaticMesh& Mesh = StaticMeshes[MeshIndex];
 
-			check(MeshRelevance.CommandInfosMask == 0);
+			check(MeshRelevance.CommandInfosMask.IsEmpty());
 			MeshRelevance.CommandInfosBase = StaticMeshCommandInfos.Num();
 
 			if (SupportsCachingMeshDrawCommands(Mesh.VertexFactory, Proxy))
@@ -236,7 +236,7 @@ void FPrimitiveSceneInfo::CacheMeshDrawCommands(FRHICommandListImmediate& RHICmd
 						{
 							static_assert(sizeof(MeshRelevance.CommandInfosMask) * 8 >= EMeshPass::Num, "CommandInfosMask is too small to contain all mesh passes.");
 
-							MeshRelevance.CommandInfosMask |= (1 << PassType);
+							MeshRelevance.CommandInfosMask.Set(PassType);
 
 							StaticMeshCommandInfos.Add(CommandInfo);
 						}
@@ -280,7 +280,7 @@ void FPrimitiveSceneInfo::RemoveCachedMeshDrawCommands()
 		FStaticMeshRelevance& MeshRelevance = StaticMeshRelevances[MeshIndex];
 
 		MeshRelevance.CommandInfosBase = 0;
-		MeshRelevance.CommandInfosMask = 0;
+		MeshRelevance.CommandInfosMask.Reset();
 	}
 
 	StaticMeshCommandInfos.Empty();
diff --git a/Engine/Source/Runtime/Renderer/Private/SceneCore.cpp b/Engine/Source/Runtime/Renderer/Private/SceneCore.cpp
--- a/Engine/Source/Runtime/Renderer/Private/SceneCore.cpp
+++ b/Engine/Source/Runtime/Renderer/Private/SceneCore.cpp
@@ -317,18 +317,16 @@ void FLightPrimitiveInteraction::FlushCachedShadowMapData()
 
 int32 FStaticMeshRelevance::GetStaticMeshCommandInfoIndex(EMeshPass::Type MeshPass) const
 {
-	static_assert(sizeof(CommandInfosMask) * 8 >= EMeshPass::Num, "CommandInfosMask is too small to contain all mesh passes.");
-
 	int32 CommandInfoIndex = CommandInfosBase;
 
-	if ((CommandInfosMask & (1 << MeshPass)) == 0)
+	if (!CommandInfosMask.Get(MeshPass))
 	{
 		return -1;
 	}
 
 	for (int32 MeshPassIndex = 0; MeshPassIndex < MeshPass; ++MeshPassIndex)
 	{
-		if (CommandInfosMask & (1 << MeshPassIndex))
+		if (CommandInfosMask.Get((EMeshPass::Type) MeshPassIndex))
 		{
 			++CommandInfoIndex;
 		}
diff --git a/Engine/Source/Runtime/Renderer/Private/SceneCore.h b/Engine/Source/Runtime/Renderer/Private/SceneCore.h
--- a/Engine/Source/Runtime/Renderer/Private/SceneCore.h
+++ b/Engine/Source/Runtime/Renderer/Private/SceneCore.h
@@ -199,7 +199,6 @@ public:
 	FStaticMeshRelevance(const FStaticMesh& StaticMesh, float InScreenSize, bool InbSupportsCachingMeshDrawCommands)
 		: Id(StaticMesh.Id)
 		, ScreenSize(InScreenSize)
-		, CommandInfosMask(0)
 		, CommandInfosBase(0)
 		, LODIndex(StaticMesh.LODIndex)
 		, bDitheredLODTransition(StaticMesh.bDitheredLODTransition)
@@ -220,7 +219,7 @@ public:
 	float ScreenSize;
 
 	/** Starting offset into continuous array of command infos for this mesh in FPrimitiveSceneInfo::CachedMeshDrawCommandInfos. */
-	uint32 CommandInfosMask;
+	FMeshPassMask CommandInfosMask;
 
 	/* Every bit corresponds to one MeshPass. If bit is set, then FPrimitiveSceneInfo::CachedMeshDrawCommandInfos contains this mesh pass. */
 	uint16 CommandInfosBase;
diff --git a/Engine/Source/Runtime/Renderer/Public/MeshPassProcessor.h b/Engine/Source/Runtime/Renderer/Public/MeshPassProcessor.h
--- a/Engine/Source/Runtime/Renderer/Public/MeshPassProcessor.h
+++ b/Engine/Source/Runtime/Renderer/Public/MeshPassProcessor.h
@@ -72,6 +72,26 @@ inline const TCHAR* GetMeshPassName(EMeshPass::Type MeshPass)
 	return nullptr;
 }
 
+/** Mesh pass mask - stores one bit per mesh pass. */
+class FMeshPassMask
+{
+public:
+	FMeshPassMask()
+		: Data(0)
+	{
+	}
+
+	void Set(EMeshPass::Type Pass) { Data |= (1 << Pass); }
+	bool Get(EMeshPass::Type Pass) const { return Data & (1 << Pass); }
+
+	void AppendTo(FMeshPassMask& Mask) const { Mask.Data |= Data; }
+	void Reset() { Data = 0; }
+	bool IsEmpty() const { return Data == 0; }
+
+private:
+	uint32 Data;
+};
+
 /** Uniquely represents a FGraphicsMinimalPipelineStateInitializer for fast compares. */
 class FGraphicsMinimalPipelineStateId
 {
