Optimized size of FStaticMeshRelevance from 20 bytes to 16 bytes.

diff --git a/Engine/Source/Runtime/Renderer/Private/SceneCore.h b/Engine/Source/Runtime/Renderer/Private/SceneCore.h
--- a/Engine/Source/Runtime/Renderer/Private/SceneCore.h
+++ b/Engine/Source/Runtime/Renderer/Private/SceneCore.h
@@ -228,21 +228,21 @@ public:
 	int8 LODIndex;
 
 	/** Whether the mesh batch should apply dithered LOD. */
-	uint32 bDitheredLODTransition : 1;
+	uint8 bDitheredLODTransition : 1;
 
 	/** Whether the mesh batch needs VertexFactory->GetStaticBatchElementVisibility to be called each frame to determine which elements of the batch are visible. */
-	uint32 bRequiresPerElementVisibility : 1;
+	uint8 bRequiresPerElementVisibility : 1;
 
 	/** Whether the mesh batch can be selected through editor selection, aka hit proxies. */
-	uint32 bSelectable : 1;
+	uint8 bSelectable : 1;
 
-	uint32 CastShadow		: 1; // Whether it can be used in shadow renderpasses.
-	uint32 bUseForMaterial	: 1; // Whether it can be used in renderpasses requiring material outputs.
-	uint32 bUseForDepthPass	: 1; // Whether it can be used in depth pass.
-	uint32 bUseAsOccluder	: 1; // User hint whether it's a good occluder.
+	uint8 CastShadow		: 1; // Whether it can be used in shadow renderpasses.
+	uint8 bUseForMaterial	: 1; // Whether it can be used in renderpasses requiring material outputs.
+	uint8 bUseForDepthPass	: 1; // Whether it can be used in depth pass.
+	uint8 bUseAsOccluder	: 1; // User hint whether it's a good occluder.
 
 	/** Cached from vertex factory to avoid dereferencing VF in InitViews. */
-	uint32 bSupportsCachingMeshDrawCommands : 1;
+	uint8 bSupportsCachingMeshDrawCommands : 1;
 
 	/** Computes index of cached mesh draw command in FPrimitiveSceneInfo::CachedMeshDrawCommandInfos, for a given mesh pass. */
 	int32 GetStaticMeshCommandInfoIndex(EMeshPass::Type MeshPass) const;
