Fixed hit proxy draw filtering:
*Conditions in dynamic passes were flipped.
*Added filtering by View.bAllowTranslucentPrimitivesInHitProxy, as HitProxy and HitProxyOpaqueOnly lists are mutually exclusive.
*Removed filtering by !ViewRelevance.HasTranslucency() as this primitive may still contain some opaque draws.

diff --git a/Engine/Source/Runtime/Renderer/Private/SceneHitProxyRendering.cpp b/Engine/Source/Runtime/Renderer/Private/SceneHitProxyRendering.cpp
--- a/Engine/Source/Runtime/Renderer/Private/SceneHitProxyRendering.cpp
+++ b/Engine/Source/Runtime/Renderer/Private/SceneHitProxyRendering.cpp
@@ -311,13 +311,13 @@ static void DoRenderHitProxies(FRHICommandListImmediate& RHICmdList, const FScen
 		const bool bHitTesting = true;
 
 		// Adjust the visibility map for this view
-		if (!View.bAllowTranslucentPrimitivesInHitProxy)
+		if (View.bAllowTranslucentPrimitivesInHitProxy)
 		{
-			View.ParallelMeshDrawCommandPasses[EMeshPass::HitProxyOpaqueOnly].DispatchDraw(nullptr, RHICmdList);
+			View.ParallelMeshDrawCommandPasses[EMeshPass::HitProxy].DispatchDraw(nullptr, RHICmdList);
 		}
 		else
 		{
-			View.ParallelMeshDrawCommandPasses[EMeshPass::HitProxy].DispatchDraw(nullptr, RHICmdList);
+			View.ParallelMeshDrawCommandPasses[EMeshPass::HitProxyOpaqueOnly].DispatchDraw(nullptr, RHICmdList);
 		}
 
 		DrawDynamicMeshPass(View, RHICmdList,
@@ -588,6 +588,11 @@ void FDeferredShadingSceneRenderer::RenderHitProxies(FRHICommandListImmediate& R
 
 void FHitProxyMeshProcessor::AddMeshBatch(const FMeshBatch& RESTRICT MeshBatch, uint64 BatchElementMask, const FPrimitiveSceneProxy* RESTRICT PrimitiveSceneProxy, int32 StaticMeshId)
 {
+	if (ViewIfDynamicMeshCommand && ViewIfDynamicMeshCommand->bAllowTranslucentPrimitivesInHitProxy != bAllowTranslucentPrimitivesInHitProxy)
+	{
+		return;
+	}
+
 	if (MeshBatch.bUseForMaterial && MeshBatch.bSelectable && Scene->RequiresHitProxies() && (!PrimitiveSceneProxy || PrimitiveSceneProxy->IsSelectable()))
 	{
 		// Determine the mesh's material and blend mode.
@@ -611,7 +616,7 @@ void FHitProxyMeshProcessor::AddMeshBatch(const FMeshBatch& RESTRICT MeshBatch,
 
 		check(Material && MaterialRenderProxy);
 
-		if (!bOpaqueOnly || !IsTranslucentBlendMode(BlendMode))
+		if (bAllowTranslucentPrimitivesInHitProxy || !IsTranslucentBlendMode(BlendMode))
 		{
 			Process(MeshBatch, BatchElementMask, StaticMeshId, PrimitiveSceneProxy, *MaterialRenderProxy, *Material, MeshFillMode, MeshCullMode);
 		}
@@ -690,10 +695,10 @@ void FHitProxyMeshProcessor::Process(
 		ShaderElementData);
 }
 
-FHitProxyMeshProcessor::FHitProxyMeshProcessor(const FScene* Scene, const FSceneView* InViewIfDynamicMeshCommand, bool InbOpaqueOnly, const FDrawingPolicyRenderState& InRenderState, FMeshPassDrawListContext& InDrawListContext)
+FHitProxyMeshProcessor::FHitProxyMeshProcessor(const FScene* Scene, const FSceneView* InViewIfDynamicMeshCommand, bool InbAllowTranslucentPrimitivesInHitProxy, const FDrawingPolicyRenderState& InRenderState, FMeshPassDrawListContext& InDrawListContext)
 	: FMeshPassProcessor(Scene, Scene->GetFeatureLevel(), InViewIfDynamicMeshCommand, InDrawListContext)
 	, PassDrawRenderState(InRenderState)
-	, bOpaqueOnly(InbOpaqueOnly)
+	, bAllowTranslucentPrimitivesInHitProxy(InbAllowTranslucentPrimitivesInHitProxy)
 {
 }
 
@@ -702,7 +707,7 @@ FMeshPassProcessor* CreateHitProxyPassProcessor(const FScene* Scene, const FScen
 	FDrawingPolicyRenderState PassDrawRenderState(Scene->UniformBuffers.ViewUniformBuffer, Scene->UniformBuffers.HitProxyPassUniformBuffer);
 	PassDrawRenderState.SetDepthStencilState(TStaticDepthStencilState<true, CF_DepthNearOrEqual>::GetRHI());
 	PassDrawRenderState.SetBlendState(TStaticBlendState<>::GetRHI());
-	return new(FMemStack::Get()) FHitProxyMeshProcessor(Scene, InViewIfDynamicMeshCommand, false, PassDrawRenderState, InDrawListContext);
+	return new(FMemStack::Get()) FHitProxyMeshProcessor(Scene, InViewIfDynamicMeshCommand, true, PassDrawRenderState, InDrawListContext);
 }
 
 FMeshPassProcessor* CreateHitProxyOpaqueOnlyPassProcessor(const FScene* Scene, const FSceneView* InViewIfDynamicMeshCommand, FMeshPassDrawListContext& InDrawListContext)
@@ -710,7 +715,7 @@ FMeshPassProcessor* CreateHitProxyOpaqueOnlyPassProcessor(const FScene* Scene, c
 	FDrawingPolicyRenderState PassDrawRenderState(Scene->UniformBuffers.ViewUniformBuffer, Scene->UniformBuffers.HitProxyPassUniformBuffer);
 	PassDrawRenderState.SetDepthStencilState(TStaticDepthStencilState<true, CF_DepthNearOrEqual>::GetRHI());
 	PassDrawRenderState.SetBlendState(TStaticBlendState<>::GetRHI());
-	return new(FMemStack::Get()) FHitProxyMeshProcessor(Scene, InViewIfDynamicMeshCommand, true, PassDrawRenderState, InDrawListContext);
+	return new(FMemStack::Get()) FHitProxyMeshProcessor(Scene, InViewIfDynamicMeshCommand, false, PassDrawRenderState, InDrawListContext);
 }
 
 FRegisterPassProcessorCreateFunction RegisterHitProxyPass(&CreateHitProxyPassProcessor, EShadingPath::Deferred, EMeshPass::HitProxy, EMeshPassFlags::CachedMeshCommands | EMeshPassFlags::MainView);
diff --git a/Engine/Source/Runtime/Renderer/Private/SceneHitProxyRendering.h b/Engine/Source/Runtime/Renderer/Private/SceneHitProxyRendering.h
--- a/Engine/Source/Runtime/Renderer/Private/SceneHitProxyRendering.h
+++ b/Engine/Source/Runtime/Renderer/Private/SceneHitProxyRendering.h
@@ -20,7 +20,7 @@ class FHitProxyMeshProcessor : public FMeshPassProcessor
 {
 public:
 
-	FHitProxyMeshProcessor(const FScene* Scene, const FSceneView* InViewIfDynamicMeshCommand, bool InbOpaqueOnly, const FDrawingPolicyRenderState& InRenderState, FMeshPassDrawListContext& InDrawListContext);
+	FHitProxyMeshProcessor(const FScene* Scene, const FSceneView* InViewIfDynamicMeshCommand, bool InbAllowTranslucentPrimitivesInHitProxy, const FDrawingPolicyRenderState& InRenderState, FMeshPassDrawListContext& InDrawListContext);
 
 	virtual void AddMeshBatch(const FMeshBatch& RESTRICT MeshBatch, uint64 BatchElementMask, const FPrimitiveSceneProxy* RESTRICT PrimitiveSceneProxy, int32 StaticMeshId = -1) override final;
 
@@ -38,7 +38,7 @@ private:
 		ERasterizerFillMode MeshFillMode,
 		ERasterizerCullMode MeshCullMode);
 
-	const bool bOpaqueOnly;
+	const bool bAllowTranslucentPrimitivesInHitProxy;
 };
 
 
diff --git a/Engine/Source/Runtime/Renderer/Private/SceneVisibility.cpp b/Engine/Source/Runtime/Renderer/Private/SceneVisibility.cpp
--- a/Engine/Source/Runtime/Renderer/Private/SceneVisibility.cpp
+++ b/Engine/Source/Runtime/Renderer/Private/SceneVisibility.cpp
@@ -2146,9 +2146,11 @@ struct FRelevancePacket
 #endif
 							if (StaticMeshRelevance.bSelectable)
 							{
-								DrawCommandPacket.AddCommandsForMesh(PrimitiveIndex, PrimitiveSceneInfo, StaticMeshRelevance, StaticMesh, Scene, bCanCache, EMeshPass::HitProxy);
-
-								if (!ViewRelevance.HasTranslucency())
+								if (View.bAllowTranslucentPrimitivesInHitProxy)
+								{
+									DrawCommandPacket.AddCommandsForMesh(PrimitiveIndex, PrimitiveSceneInfo, StaticMeshRelevance, StaticMesh, Scene, bCanCache, EMeshPass::HitProxy);
+								}
+								else
 								{
 									DrawCommandPacket.AddCommandsForMesh(PrimitiveIndex, PrimitiveSceneInfo, StaticMeshRelevance, StaticMesh, Scene, bCanCache, EMeshPass::HitProxyOpaqueOnly);
 								}
