FMeshPassProcessor - added ability to change draw list context after it's creation, in order to be able to pass processor into a task and be able to work there with local task's context.

diff --git a/Engine/Source/Runtime/Landscape/Private/LandscapeGrass.cpp b/Engine/Source/Runtime/Landscape/Private/LandscapeGrass.cpp
--- a/Engine/Source/Runtime/Landscape/Private/LandscapeGrass.cpp
+++ b/Engine/Source/Runtime/Landscape/Private/LandscapeGrass.cpp
@@ -340,7 +340,7 @@ IMPLEMENT_MATERIAL_SHADER_TYPE(, FLandscapeGrassWeightPS, TEXT("/Engine/Private/
 class FLandscapeGrassWeightMeshProcessor : public FMeshPassProcessor
 {
 public:
-	FLandscapeGrassWeightMeshProcessor(const FScene* Scene, const FSceneView* InViewIfDynamicMeshCommand, FMeshPassDrawListContext& InDrawListContext);
+	FLandscapeGrassWeightMeshProcessor(const FScene* Scene, const FSceneView* InViewIfDynamicMeshCommand, FMeshPassDrawListContext* InDrawListContext);
 
 	void AddMeshBatch(const FMeshBatch& RESTRICT MeshBatch,
 		uint64 BatchElementMask,
@@ -369,7 +369,7 @@ private:
 	FDrawingPolicyRenderState PassDrawRenderState;
 };
 
-FLandscapeGrassWeightMeshProcessor::FLandscapeGrassWeightMeshProcessor(const FScene* Scene, const FSceneView* InViewIfDynamicMeshCommand, FMeshPassDrawListContext& InDrawListContext)
+FLandscapeGrassWeightMeshProcessor::FLandscapeGrassWeightMeshProcessor(const FScene* Scene, const FSceneView* InViewIfDynamicMeshCommand, FMeshPassDrawListContext* InDrawListContext)
 	: FMeshPassProcessor(Scene, InViewIfDynamicMeshCommand->GetFeatureLevel(), InViewIfDynamicMeshCommand, InDrawListContext)
 {
 	PassDrawRenderState.SetBlendState(TStaticBlendState<>::GetRHI());
@@ -519,7 +519,7 @@ public:
 		FMemMark Mark(FMemStack::Get());
 
 		DrawDynamicMeshPass(*View, RHICmdList,
-			[View, PassOffsetX = PassOffsetX, &ComponentInfos = ComponentInfos, NumPasses = NumPasses](FDynamicPassMeshDrawListContext& DynamicMeshPassContext)
+			[View, PassOffsetX = PassOffsetX, &ComponentInfos = ComponentInfos, NumPasses = NumPasses](FDynamicPassMeshDrawListContext* DynamicMeshPassContext)
 		{
 			FLandscapeGrassWeightMeshProcessor PassMeshProcessor(
 				nullptr,
diff --git a/Engine/Source/Runtime/Renderer/Private/BasePassRendering.cpp b/Engine/Source/Runtime/Renderer/Private/BasePassRendering.cpp
--- a/Engine/Source/Runtime/Renderer/Private/BasePassRendering.cpp
+++ b/Engine/Source/Runtime/Renderer/Private/BasePassRendering.cpp
@@ -876,7 +876,7 @@ void FDeferredShadingSceneRenderer::RenderEditorPrimitivesForDPG(FRHICommandList
 		const bool bNeedToSwitchVerticalAxis = RHINeedsToSwitchVerticalAxis(ShaderPlatform);
 
 		DrawDynamicMeshPass(View, RHICmdList,
-			[&View, &DrawRenderState](FDynamicPassMeshDrawListContext& DynamicMeshPassContext)
+			[&View, &DrawRenderState](FDynamicPassMeshDrawListContext* DynamicMeshPassContext)
 			{
 				FEditorPrimitivesBasePassMeshProcessor PassMeshProcessor(
 					View.Family->Scene->GetRenderScene(),
@@ -898,7 +898,7 @@ void FDeferredShadingSceneRenderer::RenderEditorPrimitivesForDPG(FRHICommandList
 		const FBatchedElements& BatchedViewElements = DepthPriorityGroup == SDPG_World ? View.BatchedViewElements : View.TopBatchedViewElements;
 
 		DrawDynamicMeshPass(View, RHICmdList,
-			[&View, &DrawRenderState](FDynamicPassMeshDrawListContext& DynamicMeshPassContext)
+			[&View, &DrawRenderState](FDynamicPassMeshDrawListContext* DynamicMeshPassContext)
 			{
 				FEditorPrimitivesBasePassMeshProcessor PassMeshProcessor(
 					View.Family->Scene->GetRenderScene(),
@@ -1541,7 +1541,7 @@ void FBasePassMeshProcessor::AddMeshBatch(const FMeshBatch& RESTRICT MeshBatch,
 	}
 }
 
-FBasePassMeshProcessor::FBasePassMeshProcessor(const FScene* Scene, ERHIFeatureLevel::Type InFeatureLevel, const FSceneView* InViewIfDynamicMeshCommand, const FDrawingPolicyRenderState& InDrawRenderState, FMeshPassDrawListContext& InDrawListContext, ETranslucencyPass::Type InTranslucencyPassType)
+FBasePassMeshProcessor::FBasePassMeshProcessor(const FScene* Scene, ERHIFeatureLevel::Type InFeatureLevel, const FSceneView* InViewIfDynamicMeshCommand, const FDrawingPolicyRenderState& InDrawRenderState, FMeshPassDrawListContext* InDrawListContext, ETranslucencyPass::Type InTranslucencyPassType)
 	: FMeshPassProcessor(Scene, InFeatureLevel, InViewIfDynamicMeshCommand, InDrawListContext)
 	, PassDrawRenderState(InDrawRenderState)
 	, TranslucencyPassType(InTranslucencyPassType)
@@ -1553,7 +1553,7 @@ FBasePassMeshProcessor::FBasePassMeshProcessor(const FScene* Scene, ERHIFeatureL
 	GetEarlyZPassMode(GShaderPlatformForFeatureLevel[InFeatureLevel], EarlyZPassMode, bEarlyZPassMovable);
 }
 
-FMeshPassProcessor* CreateBasePassProcessor(const FScene* Scene, const FSceneView* InViewIfDynamicMeshCommand, FMeshPassDrawListContext& InDrawListContext)
+FMeshPassProcessor* CreateBasePassProcessor(const FScene* Scene, const FSceneView* InViewIfDynamicMeshCommand, FMeshPassDrawListContext* InDrawListContext)
 {
 	FExclusiveDepthStencil::Type DefaultBasePassDepthStencilAccess = GetDefaultBasePassDepthStencilAccess(EShadingPath::Deferred);
 
@@ -1563,7 +1563,7 @@ FMeshPassProcessor* CreateBasePassProcessor(const FScene* Scene, const FSceneVie
 	return new(FMemStack::Get()) FBasePassMeshProcessor(Scene, Scene->GetFeatureLevel(), InViewIfDynamicMeshCommand, PassDrawRenderState, InDrawListContext);
 }
 
-FMeshPassProcessor* CreateTranslucencyStandardPassProcessor(const FScene* Scene, const FSceneView* InViewIfDynamicMeshCommand, FMeshPassDrawListContext& InDrawListContext)
+FMeshPassProcessor* CreateTranslucencyStandardPassProcessor(const FScene* Scene, const FSceneView* InViewIfDynamicMeshCommand, FMeshPassDrawListContext* InDrawListContext)
 {
 	FDrawingPolicyRenderState PassDrawRenderState(Scene->UniformBuffers.ViewUniformBuffer, Scene->UniformBuffers.TranslucentBasePassUniformBuffer);
 	PassDrawRenderState.SetDepthStencilState(TStaticDepthStencilState<false, CF_DepthNearOrEqual>::GetRHI());
@@ -1571,7 +1571,7 @@ FMeshPassProcessor* CreateTranslucencyStandardPassProcessor(const FScene* Scene,
 	return new(FMemStack::Get()) FBasePassMeshProcessor(Scene, Scene->GetFeatureLevel(), InViewIfDynamicMeshCommand, PassDrawRenderState, InDrawListContext, ETranslucencyPass::TPT_StandardTranslucency);
 }
 
-FMeshPassProcessor* CreateTranslucencyAfterDOFProcessor(const FScene* Scene, const FSceneView* InViewIfDynamicMeshCommand, FMeshPassDrawListContext& InDrawListContext)
+FMeshPassProcessor* CreateTranslucencyAfterDOFProcessor(const FScene* Scene, const FSceneView* InViewIfDynamicMeshCommand, FMeshPassDrawListContext* InDrawListContext)
 {
 	FDrawingPolicyRenderState PassDrawRenderState(Scene->UniformBuffers.ViewUniformBuffer, Scene->UniformBuffers.TranslucentBasePassUniformBuffer);
 	PassDrawRenderState.SetDepthStencilState(TStaticDepthStencilState<false, CF_DepthNearOrEqual>::GetRHI());
@@ -1579,7 +1579,7 @@ FMeshPassProcessor* CreateTranslucencyAfterDOFProcessor(const FScene* Scene, con
 	return new(FMemStack::Get()) FBasePassMeshProcessor(Scene, Scene->GetFeatureLevel(), InViewIfDynamicMeshCommand, PassDrawRenderState, InDrawListContext, ETranslucencyPass::TPT_TranslucencyAfterDOF);
 }
 
-FMeshPassProcessor* CreateTranslucencyAllPassProcessor(const FScene* Scene, const FSceneView* InViewIfDynamicMeshCommand, FMeshPassDrawListContext& InDrawListContext)
+FMeshPassProcessor* CreateTranslucencyAllPassProcessor(const FScene* Scene, const FSceneView* InViewIfDynamicMeshCommand, FMeshPassDrawListContext* InDrawListContext)
 {
 	FDrawingPolicyRenderState PassDrawRenderState(Scene->UniformBuffers.ViewUniformBuffer, Scene->UniformBuffers.TranslucentBasePassUniformBuffer);
 	PassDrawRenderState.SetDepthStencilState(TStaticDepthStencilState<false, CF_DepthNearOrEqual>::GetRHI());
diff --git a/Engine/Source/Runtime/Renderer/Private/BasePassRendering.h b/Engine/Source/Runtime/Renderer/Private/BasePassRendering.h
--- a/Engine/Source/Runtime/Renderer/Private/BasePassRendering.h
+++ b/Engine/Source/Runtime/Renderer/Private/BasePassRendering.h
@@ -675,7 +675,7 @@ class FBasePassMeshProcessor : public FMeshPassProcessor
 {
 public:
 
-	FBasePassMeshProcessor(const FScene* InScene, ERHIFeatureLevel::Type InFeatureLevel, const FSceneView* InViewIfDynamicMeshCommand, const FDrawingPolicyRenderState& InDrawRenderState, FMeshPassDrawListContext& InDrawListContext, ETranslucencyPass::Type InTranslucencyPassType = ETranslucencyPass::TPT_MAX);
+	FBasePassMeshProcessor(const FScene* InScene, ERHIFeatureLevel::Type InFeatureLevel, const FSceneView* InViewIfDynamicMeshCommand, const FDrawingPolicyRenderState& InDrawRenderState, FMeshPassDrawListContext* InDrawListContext, ETranslucencyPass::Type InTranslucencyPassType = ETranslucencyPass::TPT_MAX);
 
 	virtual void AddMeshBatch(const FMeshBatch& RESTRICT MeshBatch, uint64 BatchElementMask, const FPrimitiveSceneProxy* RESTRICT PrimitiveSceneProxy, int32 StaticMeshId = -1) override final;
 
diff --git a/Engine/Source/Runtime/Renderer/Private/CompositionLighting/PostProcessMeshDecals.cpp b/Engine/Source/Runtime/Renderer/Private/CompositionLighting/PostProcessMeshDecals.cpp
--- a/Engine/Source/Runtime/Renderer/Private/CompositionLighting/PostProcessMeshDecals.cpp
+++ b/Engine/Source/Runtime/Renderer/Private/CompositionLighting/PostProcessMeshDecals.cpp
@@ -212,7 +212,7 @@ public:
 		const FSceneView* InViewIfDynamicMeshCommand, 
 		EDecalRenderStage InPassDecalStage, 
 		FDecalRenderingCommon::ERenderTargetMode InRenderTargetMode,
-		FMeshPassDrawListContext& InDrawListContext);
+		FMeshPassDrawListContext* InDrawListContext);
 
 	virtual void AddMeshBatch(const FMeshBatch& RESTRICT MeshBatch, uint64 BatchElementMask, const FPrimitiveSceneProxy* RESTRICT PrimitiveSceneProxy, int32 StaticMeshId = -1) override final;
 
@@ -237,7 +237,7 @@ FMeshDecalMeshProcessor::FMeshDecalMeshProcessor(const FScene* Scene,
 	const FSceneView* InViewIfDynamicMeshCommand, 
 	EDecalRenderStage InPassDecalStage, 
 	FDecalRenderingCommon::ERenderTargetMode InRenderTargetMode,
-	FMeshPassDrawListContext& InDrawListContext)
+	FMeshPassDrawListContext* InDrawListContext)
 	: FMeshPassProcessor(Scene, Scene->GetFeatureLevel(), InViewIfDynamicMeshCommand, InDrawListContext)
 	, PassDecalStage(InPassDecalStage)
 	, RenderTargetMode(InRenderTargetMode)
@@ -388,7 +388,7 @@ void DrawDecalMeshCommands(FRenderingCompositePassContext& Context, EDecalRender
 
 
 	DrawDynamicMeshPass(View, RHICmdList,
-		[&View, CurrentDecalStage, RenderTargetMode](FDynamicPassMeshDrawListContext& DynamicMeshPassContext)
+		[&View, CurrentDecalStage, RenderTargetMode](FDynamicPassMeshDrawListContext* DynamicMeshPassContext)
 	{
 		FMeshDecalMeshProcessor PassMeshProcessor(
 			View.Family->Scene->GetRenderScene(),
diff --git a/Engine/Source/Runtime/Renderer/Private/ConvertToUniformMesh.cpp b/Engine/Source/Runtime/Renderer/Private/ConvertToUniformMesh.cpp
--- a/Engine/Source/Runtime/Renderer/Private/ConvertToUniformMesh.cpp
+++ b/Engine/Source/Runtime/Renderer/Private/ConvertToUniformMesh.cpp
@@ -152,7 +152,7 @@ IMPLEMENT_MATERIAL_SHADER_TYPE(,FConvertToUniformMeshGS,TEXT("/Engine/Private/Co
 class FConvertToUniformMeshProcessor : public FMeshPassProcessor
 {
 public:
-	FConvertToUniformMeshProcessor(const FScene* Scene, const FViewInfo* InViewIfDynamicMeshCommand, FMeshPassDrawListContext& InDrawListContext);
+	FConvertToUniformMeshProcessor(const FScene* Scene, const FViewInfo* InViewIfDynamicMeshCommand, FMeshPassDrawListContext* InDrawListContext);
 
 	void AddMeshBatch(const FMeshBatch& RESTRICT MeshBatch, uint64 BatchElementMask, const FPrimitiveSceneProxy* RESTRICT PrimitiveSceneProxy, int32 StaticMeshId = -1) override final;
 
@@ -169,7 +169,7 @@ private:
 	FDrawingPolicyRenderState PassDrawRenderState;
 };
 
-FConvertToUniformMeshProcessor::FConvertToUniformMeshProcessor(const FScene* Scene, const FViewInfo* InViewIfDynamicMeshCommand, FMeshPassDrawListContext& InDrawListContext)
+FConvertToUniformMeshProcessor::FConvertToUniformMeshProcessor(const FScene* Scene, const FViewInfo* InViewIfDynamicMeshCommand, FMeshPassDrawListContext* InDrawListContext)
 	: FMeshPassProcessor(Scene, Scene->GetFeatureLevel(), InViewIfDynamicMeshCommand, InDrawListContext)
 {
 	PassDrawRenderState.SetBlendState(TStaticBlendState<>::GetRHI());
@@ -302,7 +302,7 @@ int32 FUniformMeshConverter::Convert(
 				OutPrimitiveUniformBuffer = Mesh.Elements.Num() > 0 ? Mesh.Elements[0].PrimitiveUniformBuffer : nullptr;
 
 				DrawDynamicMeshPass(View, RHICmdList,
-					[&View, &RHICmdList, &Mesh, &PrimitiveSceneProxy](FDynamicPassMeshDrawListContext& DynamicMeshPassContext)
+					[&View, &RHICmdList, &Mesh, &PrimitiveSceneProxy](FDynamicPassMeshDrawListContext* DynamicMeshPassContext)
 				{
 					FConvertToUniformMeshProcessor PassMeshProcessor(
 						View.Family->Scene->GetRenderScene(),
diff --git a/Engine/Source/Runtime/Renderer/Private/CustomDepthRendering.cpp b/Engine/Source/Runtime/Renderer/Private/CustomDepthRendering.cpp
--- a/Engine/Source/Runtime/Renderer/Private/CustomDepthRendering.cpp
+++ b/Engine/Source/Runtime/Renderer/Private/CustomDepthRendering.cpp
@@ -119,7 +119,7 @@ bool FCustomDepthPrimSet::DrawPrims(FRHICommandListImmediate& RHICmdList, const
 class FCustomDepthPassMeshProcessor : public FMeshPassProcessor
 {
 public:
-	FCustomDepthPassMeshProcessor(const FScene* Scene, const FSceneView* InViewIfDynamicMeshCommand, FMeshPassDrawListContext& InDrawListContext);
+	FCustomDepthPassMeshProcessor(const FScene* Scene, const FSceneView* InViewIfDynamicMeshCommand, FMeshPassDrawListContext* InDrawListContext);
 
 	virtual void AddMeshBatch(const FMeshBatch& RESTRICT MeshBatch, uint64 BatchElementMask, const FPrimitiveSceneProxy* RESTRICT PrimitiveSceneProxy, int32 StaticMeshId = -1) override final;
 
@@ -140,7 +140,7 @@ private:
 	FDrawingPolicyRenderState PassDrawRenderState;
 };
 
-FCustomDepthPassMeshProcessor::FCustomDepthPassMeshProcessor(const FScene* Scene, const FSceneView* InViewIfDynamicMeshCommand, FMeshPassDrawListContext& InDrawListContext)
+FCustomDepthPassMeshProcessor::FCustomDepthPassMeshProcessor(const FScene* Scene, const FSceneView* InViewIfDynamicMeshCommand, FMeshPassDrawListContext* InDrawListContext)
 	: FMeshPassProcessor(Scene, Scene->GetFeatureLevel(), InViewIfDynamicMeshCommand, InDrawListContext)
 {
 	PassDrawRenderState.SetViewUniformBuffer(Scene->UniformBuffers.CustomDepthViewUniformBuffer);
@@ -294,7 +294,7 @@ void FCustomDepthPassMeshProcessor::Process(
 		ShaderElementData);
 }
 
-FMeshPassProcessor* CreateCustomDepthPassProcessor(const FScene* Scene, const FSceneView* InViewIfDynamicMeshCommand, FMeshPassDrawListContext& InDrawListContext)
+FMeshPassProcessor* CreateCustomDepthPassProcessor(const FScene* Scene, const FSceneView* InViewIfDynamicMeshCommand, FMeshPassDrawListContext* InDrawListContext)
 {
 	return new(FMemStack::Get()) FCustomDepthPassMeshProcessor(Scene, InViewIfDynamicMeshCommand, InDrawListContext);
 }
diff --git a/Engine/Source/Runtime/Renderer/Private/DebugViewModeRendering.cpp b/Engine/Source/Runtime/Renderer/Private/DebugViewModeRendering.cpp
--- a/Engine/Source/Runtime/Renderer/Private/DebugViewModeRendering.cpp
+++ b/Engine/Source/Runtime/Renderer/Private/DebugViewModeRendering.cpp
@@ -282,7 +282,7 @@ FDebugViewModeMeshProcessor::FDebugViewModeMeshProcessor(
 	const FSceneView* InViewIfDynamicMeshCommand, 
 	FUniformBufferRHIParamRef InPassUniformBuffer, 
 	bool bTranslucentBasePass,
-	FMeshPassDrawListContext& InDrawListContext
+	FMeshPassDrawListContext* InDrawListContext
 )
 	: FMeshPassProcessor(InScene, InFeatureLevel, InViewIfDynamicMeshCommand, InDrawListContext)
 	, PassUniformBuffer(InPassUniformBuffer)
@@ -397,7 +397,7 @@ void FDebugViewModeMeshProcessor::AddMeshBatch(const FMeshBatch& RESTRICT MeshBa
 		ShaderElementData);
 }
 
-FMeshPassProcessor* CreateDebugViewModePassProcessor(const FScene* Scene, const FSceneView* InViewIfDynamicMeshCommand, FMeshPassDrawListContext& InDrawListContext)
+FMeshPassProcessor* CreateDebugViewModePassProcessor(const FScene* Scene, const FSceneView* InViewIfDynamicMeshCommand, FMeshPassDrawListContext* InDrawListContext)
 {
 	const ERHIFeatureLevel::Type FeatureLevel = Scene ? Scene->GetFeatureLevel() : (InViewIfDynamicMeshCommand ? InViewIfDynamicMeshCommand->GetFeatureLevel() : GMaxRHIFeatureLevel);
 	return new(FMemStack::Get()) FDebugViewModeMeshProcessor(Scene, FeatureLevel, InViewIfDynamicMeshCommand, nullptr, false, InDrawListContext);
diff --git a/Engine/Source/Runtime/Renderer/Private/DebugViewModeRendering.h b/Engine/Source/Runtime/Renderer/Private/DebugViewModeRendering.h
--- a/Engine/Source/Runtime/Renderer/Private/DebugViewModeRendering.h
+++ b/Engine/Source/Runtime/Renderer/Private/DebugViewModeRendering.h
@@ -309,7 +309,7 @@ public:
 class FDebugViewModeMeshProcessor : public FMeshPassProcessor
 {
 public:
-	FDebugViewModeMeshProcessor(const FScene* InScene, ERHIFeatureLevel::Type InFeatureLevel, const FSceneView* InViewIfDynamicMeshCommand, FUniformBufferRHIParamRef InPassUniformBuffer, bool bTranslucentBasePass, FMeshPassDrawListContext& InDrawListContext);
+	FDebugViewModeMeshProcessor(const FScene* InScene, ERHIFeatureLevel::Type InFeatureLevel, const FSceneView* InViewIfDynamicMeshCommand, FUniformBufferRHIParamRef InPassUniformBuffer, bool bTranslucentBasePass, FMeshPassDrawListContext* InDrawListContext);
 	virtual void AddMeshBatch(const FMeshBatch& RESTRICT MeshBatch, uint64 BatchElementMask, const FPrimitiveSceneProxy* RESTRICT PrimitiveSceneProxy, int32 StaticMeshId = -1) override final;
 
 private:
diff --git a/Engine/Source/Runtime/Renderer/Private/DeferredShadingRenderer.cpp b/Engine/Source/Runtime/Renderer/Private/DeferredShadingRenderer.cpp
--- a/Engine/Source/Runtime/Renderer/Private/DeferredShadingRenderer.cpp
+++ b/Engine/Source/Runtime/Renderer/Private/DeferredShadingRenderer.cpp
@@ -710,7 +710,7 @@ bool FDeferredShadingSceneRenderer::GatherRayTracingWorldInstances(FRHICommandLi
 									View.DynamicMeshDrawCommandStorage[EMeshPass::RayTracing],
 									View.RaytraycingVisibleMeshDrawCommands
 								);
-								FRayTracingMeshProcessor RayTracingMeshProcessor(Scene, View.GetFeatureLevel(), &View, DynamicPassMeshDrawListContext);
+								FRayTracingMeshProcessor RayTracingMeshProcessor(Scene, View.GetFeatureLevel(), &View, &DynamicPassMeshDrawListContext);
 								RayTracingMeshProcessor.AddMeshBatch(StaticMesh, BatchVisibility, SceneProxy);
 							}
 						}
@@ -753,7 +753,7 @@ bool FDeferredShadingSceneRenderer::GatherRayTracingWorldInstances(FRHICommandLi
 								View.DynamicMeshDrawCommandStorage[EMeshPass::RayTracing],
 								View.RaytraycingVisibleMeshDrawCommands
 							);
-							FRayTracingMeshProcessor RayTracingMeshProcessor(Scene, View.GetFeatureLevel(), &View, DynamicPassMeshDrawListContext);
+							FRayTracingMeshProcessor RayTracingMeshProcessor(Scene, View.GetFeatureLevel(), &View, &DynamicPassMeshDrawListContext);
 
 							const int GeometryInstance = View.RayTracingGeometryInstances.Num();
 							int DrawCmdIndexCopy = VisibleDrawCommandStartOffset[ViewIndex];
@@ -801,7 +801,7 @@ bool FDeferredShadingSceneRenderer::GatherRayTracingWorldInstances(FRHICommandLi
 								View.DynamicMeshDrawCommandStorage[EMeshPass::RayTracing],
 								View.RaytraycingVisibleMeshDrawCommands
 							);
-							FRayTracingMeshProcessor RayTracingMeshProcessor(Scene, View.GetFeatureLevel(), &View, DynamicPassMeshDrawListContext);
+							FRayTracingMeshProcessor RayTracingMeshProcessor(Scene, View.GetFeatureLevel(), &View, &DynamicPassMeshDrawListContext);
 
 							int32 GeometryIndex = -1;
 							int32 SegmentPrefixSum = 0;
diff --git a/Engine/Source/Runtime/Renderer/Private/DepthRendering.cpp b/Engine/Source/Runtime/Renderer/Private/DepthRendering.cpp
--- a/Engine/Source/Runtime/Renderer/Private/DepthRendering.cpp
+++ b/Engine/Source/Runtime/Renderer/Private/DepthRendering.cpp
@@ -1191,7 +1191,7 @@ void FDeferredShadingSceneRenderer::RenderPrePassEditorPrimitives(FRHICommandLis
 		if (UseMeshDrawCommandPipeline())
 		{
 			DrawDynamicMeshPass(View, RHICmdList,
-				[&View, &DrawRenderState, LocalScene, &Context](FDynamicPassMeshDrawListContext& DynamicMeshPassContext)
+				[&View, &DrawRenderState, LocalScene, &Context](FDynamicPassMeshDrawListContext* DynamicMeshPassContext)
 				{
 					FDepthPassMeshProcessor PassMeshProcessor(
 						LocalScene,
@@ -1223,7 +1223,7 @@ void FDeferredShadingSceneRenderer::RenderPrePassEditorPrimitives(FRHICommandLis
 		if (UseMeshDrawCommandPipeline())
 		{
 			DrawDynamicMeshPass(View, RHICmdList,
-				[&View, &DrawRenderState, LocalScene, &Context](FDynamicPassMeshDrawListContext& DynamicMeshPassContext)
+				[&View, &DrawRenderState, LocalScene, &Context](FDynamicPassMeshDrawListContext* DynamicMeshPassContext)
 				{
 					FDepthPassMeshProcessor PassMeshProcessor(
 						LocalScene,
@@ -1593,7 +1593,7 @@ FDepthPassMeshProcessor::FDepthPassMeshProcessor(const FScene* Scene,
 	const bool InbRespectUseAsOccluderFlag,
 	const EDepthDrawingMode InEarlyZPassMode,
 	const bool InbEarlyZPassMovable,
-	FMeshPassDrawListContext& InDrawListContext)
+	FMeshPassDrawListContext* InDrawListContext)
 	: FMeshPassProcessor(Scene, Scene->GetFeatureLevel(), InViewIfDynamicMeshCommand, InDrawListContext)
 	, bRespectUseAsOccluderFlag(InbRespectUseAsOccluderFlag)
 	, EarlyZPassMode(InEarlyZPassMode)
@@ -1604,7 +1604,7 @@ FDepthPassMeshProcessor::FDepthPassMeshProcessor(const FScene* Scene,
 	PassDrawRenderState.SetPassUniformBuffer(Scene->UniformBuffers.DepthPassUniformBuffer);
 }
 
-FMeshPassProcessor* CreateDepthPassProcessor(const FScene* Scene, const FSceneView* InViewIfDynamicMeshCommand, FMeshPassDrawListContext& InDrawListContext)
+FMeshPassProcessor* CreateDepthPassProcessor(const FScene* Scene, const FSceneView* InViewIfDynamicMeshCommand, FMeshPassDrawListContext* InDrawListContext)
 {
 	extern void GetEarlyZPassMode(EShaderPlatform ShaderPlatform, EDepthDrawingMode& EarlyZPassMode, bool& bEarlyZPassMovable);
 
diff --git a/Engine/Source/Runtime/Renderer/Private/DepthRendering.h b/Engine/Source/Runtime/Renderer/Private/DepthRendering.h
--- a/Engine/Source/Runtime/Renderer/Private/DepthRendering.h
+++ b/Engine/Source/Runtime/Renderer/Private/DepthRendering.h
@@ -502,7 +502,7 @@ public:
 		const bool InbRespectUseAsOccluderFlag,
 		const EDepthDrawingMode InEarlyZPassMode,
 		const bool InbEarlyZPassMovable,
-		FMeshPassDrawListContext& InDrawListContext);
+		FMeshPassDrawListContext* InDrawListContext);
 
 	virtual void AddMeshBatch(const FMeshBatch& RESTRICT MeshBatch, uint64 BatchElementMask, const FPrimitiveSceneProxy* RESTRICT PrimitiveSceneProxy, int32 StaticMeshId = -1) override final;
 
diff --git a/Engine/Source/Runtime/Renderer/Private/DistortionRendering.cpp b/Engine/Source/Runtime/Renderer/Private/DistortionRendering.cpp
--- a/Engine/Source/Runtime/Renderer/Private/DistortionRendering.cpp
+++ b/Engine/Source/Runtime/Renderer/Private/DistortionRendering.cpp
@@ -1281,13 +1281,13 @@ void FDistortionMeshProcessor::Process(
 		ShaderElementData);
 }
 
-FDistortionMeshProcessor::FDistortionMeshProcessor(const FScene* Scene, const FSceneView* InViewIfDynamicMeshCommand, const FDrawingPolicyRenderState& InPassDrawRenderState, FMeshPassDrawListContext& InDrawListContext)
+FDistortionMeshProcessor::FDistortionMeshProcessor(const FScene* Scene, const FSceneView* InViewIfDynamicMeshCommand, const FDrawingPolicyRenderState& InPassDrawRenderState, FMeshPassDrawListContext* InDrawListContext)
 	: FMeshPassProcessor(Scene, Scene->GetFeatureLevel(), InViewIfDynamicMeshCommand, InDrawListContext)
 	, PassDrawRenderState(InPassDrawRenderState)
 {
 }
 
-FMeshPassProcessor* CreateDistortionPassProcessor(const FScene* Scene, const FSceneView* InViewIfDynamicMeshCommand, FMeshPassDrawListContext& InDrawListContext)
+FMeshPassProcessor* CreateDistortionPassProcessor(const FScene* Scene, const FSceneView* InViewIfDynamicMeshCommand, FMeshPassDrawListContext* InDrawListContext)
 {
 	FDrawingPolicyRenderState DistortionPassState;
 	DistortionPassState.SetViewUniformBuffer(Scene->UniformBuffers.ViewUniformBuffer);
@@ -1308,7 +1308,7 @@ FMeshPassProcessor* CreateDistortionPassProcessor(const FScene* Scene, const FSc
 	return new(FMemStack::Get()) FDistortionMeshProcessor(Scene, InViewIfDynamicMeshCommand, DistortionPassState, InDrawListContext);
 }
 
-FMeshPassProcessor* CreateMobileDistortionPassProcessor(const FScene* Scene, const FSceneView* InViewIfDynamicMeshCommand, FMeshPassDrawListContext& InDrawListContext)
+FMeshPassProcessor* CreateMobileDistortionPassProcessor(const FScene* Scene, const FSceneView* InViewIfDynamicMeshCommand, FMeshPassDrawListContext* InDrawListContext)
 {
 	FDrawingPolicyRenderState DistortionPassState;
 	DistortionPassState.SetViewUniformBuffer(Scene->UniformBuffers.ViewUniformBuffer);
diff --git a/Engine/Source/Runtime/Renderer/Private/DistortionRendering.h b/Engine/Source/Runtime/Renderer/Private/DistortionRendering.h
--- a/Engine/Source/Runtime/Renderer/Private/DistortionRendering.h
+++ b/Engine/Source/Runtime/Renderer/Private/DistortionRendering.h
@@ -77,7 +77,7 @@ class FDistortionMeshProcessor : public FMeshPassProcessor
 {
 public:
 
-	FDistortionMeshProcessor(const FScene* Scene, const FSceneView* InViewIfDynamicMeshCommand, const FDrawingPolicyRenderState& InPassDrawRenderState, FMeshPassDrawListContext& InDrawListContext);
+	FDistortionMeshProcessor(const FScene* Scene, const FSceneView* InViewIfDynamicMeshCommand, const FDrawingPolicyRenderState& InPassDrawRenderState, FMeshPassDrawListContext* InDrawListContext);
 
 	virtual void AddMeshBatch(const FMeshBatch& RESTRICT MeshBatch, uint64 BatchElementMask, const FPrimitiveSceneProxy* RESTRICT PrimitiveSceneProxy, int32 StaticMeshId = -1) override final;
 
diff --git a/Engine/Source/Runtime/Renderer/Private/EditorPrimitivesRendering.cpp b/Engine/Source/Runtime/Renderer/Private/EditorPrimitivesRendering.cpp
--- a/Engine/Source/Runtime/Renderer/Private/EditorPrimitivesRendering.cpp
+++ b/Engine/Source/Runtime/Renderer/Private/EditorPrimitivesRendering.cpp
@@ -10,7 +10,7 @@
 #include "MobileBasePassRendering.h"
 #include "MeshPassProcessor.inl"
 
-FEditorPrimitivesBasePassMeshProcessor::FEditorPrimitivesBasePassMeshProcessor(const FScene* Scene, ERHIFeatureLevel::Type InFeatureLevel, const FSceneView* InViewIfDynamicMeshCommand, const FDrawingPolicyRenderState& InDrawRenderState, bool bInTranslucentBasePass, FMeshPassDrawListContext& InDrawListContext) 
+FEditorPrimitivesBasePassMeshProcessor::FEditorPrimitivesBasePassMeshProcessor(const FScene* Scene, ERHIFeatureLevel::Type InFeatureLevel, const FSceneView* InViewIfDynamicMeshCommand, const FDrawingPolicyRenderState& InDrawRenderState, bool bInTranslucentBasePass, FMeshPassDrawListContext* InDrawListContext) 
 	: FMeshPassProcessor(Scene, InFeatureLevel, InViewIfDynamicMeshCommand, InDrawListContext)
 	, PassDrawRenderState(InDrawRenderState)
 	, bTranslucentBasePass(bInTranslucentBasePass)
diff --git a/Engine/Source/Runtime/Renderer/Private/EditorPrimitivesRendering.h b/Engine/Source/Runtime/Renderer/Private/EditorPrimitivesRendering.h
--- a/Engine/Source/Runtime/Renderer/Private/EditorPrimitivesRendering.h
+++ b/Engine/Source/Runtime/Renderer/Private/EditorPrimitivesRendering.h
@@ -9,7 +9,7 @@
 class FEditorPrimitivesBasePassMeshProcessor : public FMeshPassProcessor
 {
 public:
-	FEditorPrimitivesBasePassMeshProcessor(const FScene* InScene, ERHIFeatureLevel::Type InFeatureLevel, const FSceneView* InViewIfDynamicMeshCommand, const FDrawingPolicyRenderState& InDrawRenderState, bool bTranslucentBasePass, FMeshPassDrawListContext& InDrawListContext);
+	FEditorPrimitivesBasePassMeshProcessor(const FScene* InScene, ERHIFeatureLevel::Type InFeatureLevel, const FSceneView* InViewIfDynamicMeshCommand, const FDrawingPolicyRenderState& InDrawRenderState, bool bTranslucentBasePass, FMeshPassDrawListContext* InDrawListContext);
 
 	virtual void AddMeshBatch(const FMeshBatch& RESTRICT MeshBatch, uint64 BatchElementMask, const FPrimitiveSceneProxy* RESTRICT PrimitiveSceneProxy, int32 StaticMeshId = -1) override final;
 
diff --git a/Engine/Source/Runtime/Renderer/Private/LightMapDensityRendering.cpp b/Engine/Source/Runtime/Renderer/Private/LightMapDensityRendering.cpp
--- a/Engine/Source/Runtime/Renderer/Private/LightMapDensityRendering.cpp
+++ b/Engine/Source/Runtime/Renderer/Private/LightMapDensityRendering.cpp
@@ -331,7 +331,7 @@ void FLightmapDensityMeshProcessor::AddMeshBatch(const FMeshBatch& RESTRICT Mesh
 	}
 }
 
-FLightmapDensityMeshProcessor::FLightmapDensityMeshProcessor(const FScene* Scene, const FSceneView* InViewIfDynamicMeshCommand, FMeshPassDrawListContext& InDrawListContext)
+FLightmapDensityMeshProcessor::FLightmapDensityMeshProcessor(const FScene* Scene, const FSceneView* InViewIfDynamicMeshCommand, FMeshPassDrawListContext* InDrawListContext)
 	: FMeshPassProcessor(Scene, Scene->GetFeatureLevel(), InViewIfDynamicMeshCommand, InDrawListContext)
 {
 	// Opaque blending, depth tests and writes.
@@ -342,7 +342,7 @@ FLightmapDensityMeshProcessor::FLightmapDensityMeshProcessor(const FScene* Scene
 	PassDrawRenderState.SetPassUniformBuffer(Scene->UniformBuffers.LightmapDensityPassUniformBuffer);
 }
 
-FMeshPassProcessor* CreateLightmapDensityPassProcessor(const FScene* Scene, const FSceneView* InViewIfDynamicMeshCommand, FMeshPassDrawListContext& InDrawListContext)
+FMeshPassProcessor* CreateLightmapDensityPassProcessor(const FScene* Scene, const FSceneView* InViewIfDynamicMeshCommand, FMeshPassDrawListContext* InDrawListContext)
 {
 	return new(FMemStack::Get()) FLightmapDensityMeshProcessor(Scene, InViewIfDynamicMeshCommand, InDrawListContext);
 }
diff --git a/Engine/Source/Runtime/Renderer/Private/LightMapDensityRendering.h b/Engine/Source/Runtime/Renderer/Private/LightMapDensityRendering.h
--- a/Engine/Source/Runtime/Renderer/Private/LightMapDensityRendering.h
+++ b/Engine/Source/Runtime/Renderer/Private/LightMapDensityRendering.h
@@ -297,7 +297,7 @@ class FLightmapDensityMeshProcessor : public FMeshPassProcessor
 {
 public:
 
-	FLightmapDensityMeshProcessor(const FScene* Scene, const FSceneView* InViewIfDynamicMeshCommand, FMeshPassDrawListContext& InDrawListContext);
+	FLightmapDensityMeshProcessor(const FScene* Scene, const FSceneView* InViewIfDynamicMeshCommand, FMeshPassDrawListContext* InDrawListContext);
 
 	virtual void AddMeshBatch(const FMeshBatch& RESTRICT MeshBatch, uint64 BatchElementMask, const FPrimitiveSceneProxy* RESTRICT PrimitiveSceneProxy, int32 StaticMeshId = -1) override final;
 
diff --git a/Engine/Source/Runtime/Renderer/Private/MeshPassProcessor.cpp b/Engine/Source/Runtime/Renderer/Private/MeshPassProcessor.cpp
--- a/Engine/Source/Runtime/Renderer/Private/MeshPassProcessor.cpp
+++ b/Engine/Source/Runtime/Renderer/Private/MeshPassProcessor.cpp
@@ -846,7 +846,7 @@ void DrawDynamicMeshPassPrivate(
 	}
 }
 
-FMeshPassProcessor::FMeshPassProcessor(const FScene* InScene, ERHIFeatureLevel::Type InFeatureLevel, const FSceneView* InViewIfDynamicMeshCommand, FMeshPassDrawListContext& InDrawListContext) 
+FMeshPassProcessor::FMeshPassProcessor(const FScene* InScene, ERHIFeatureLevel::Type InFeatureLevel, const FSceneView* InViewIfDynamicMeshCommand, FMeshPassDrawListContext* InDrawListContext) 
 	: Scene(InScene)
 	, FeatureLevel(InFeatureLevel)
 	, ViewIfDynamicMeshCommand(InViewIfDynamicMeshCommand)
diff --git a/Engine/Source/Runtime/Renderer/Private/MobileBasePass.cpp b/Engine/Source/Runtime/Renderer/Private/MobileBasePass.cpp
--- a/Engine/Source/Runtime/Renderer/Private/MobileBasePass.cpp
+++ b/Engine/Source/Runtime/Renderer/Private/MobileBasePass.cpp
@@ -636,7 +636,7 @@ void TMobileBasePassPSPolicyParamType<FUniformLightMapPolicy>::GetShaderBindings
 	}
 }
 
-FMobileBasePassMeshProcessor::FMobileBasePassMeshProcessor(const FScene* Scene, ERHIFeatureLevel::Type InFeatureLevel, const FSceneView* InViewIfDynamicMeshCommand, const FDrawingPolicyRenderState& InDrawRenderState, FMeshPassDrawListContext& InDrawListContext, bool bInCanReceiveCSM, ETranslucencyPass::Type InTranslucencyPassType)
+FMobileBasePassMeshProcessor::FMobileBasePassMeshProcessor(const FScene* Scene, ERHIFeatureLevel::Type InFeatureLevel, const FSceneView* InViewIfDynamicMeshCommand, const FDrawingPolicyRenderState& InDrawRenderState, FMeshPassDrawListContext* InDrawListContext, bool bInCanReceiveCSM, ETranslucencyPass::Type InTranslucencyPassType)
 	: FMeshPassProcessor(Scene, InFeatureLevel, InViewIfDynamicMeshCommand, InDrawListContext)
 	, PassDrawRenderState(InDrawRenderState)
 	, TranslucencyPassType(InTranslucencyPassType)
@@ -761,7 +761,7 @@ void FMobileBasePassMeshProcessor::Process(
 		ShaderElementData);
 }
 
-FMeshPassProcessor* CreateMobileBasePassProcessor(const FScene* Scene, const FSceneView* InViewIfDynamicMeshCommand, FMeshPassDrawListContext& InDrawListContext)
+FMeshPassProcessor* CreateMobileBasePassProcessor(const FScene* Scene, const FSceneView* InViewIfDynamicMeshCommand, FMeshPassDrawListContext* InDrawListContext)
 {
 	FExclusiveDepthStencil::Type DefaultBasePassDepthStencilAccess = GetDefaultBasePassDepthStencilAccess(EShadingPath::Mobile);
 	
@@ -773,7 +773,7 @@ FMeshPassProcessor* CreateMobileBasePassProcessor(const FScene* Scene, const FSc
 	return new(FMemStack::Get()) FMobileBasePassMeshProcessor(Scene, Scene->GetFeatureLevel(), InViewIfDynamicMeshCommand, PassDrawRenderState, InDrawListContext, false);
 }
 
-FMeshPassProcessor* CreateMobileBasePassCSMProcessor(const FScene* Scene, const FSceneView* InViewIfDynamicMeshCommand, FMeshPassDrawListContext& InDrawListContext)
+FMeshPassProcessor* CreateMobileBasePassCSMProcessor(const FScene* Scene, const FSceneView* InViewIfDynamicMeshCommand, FMeshPassDrawListContext* InDrawListContext)
 {
 	FExclusiveDepthStencil::Type DefaultBasePassDepthStencilAccess = GetDefaultBasePassDepthStencilAccess(EShadingPath::Mobile);
 
@@ -785,7 +785,7 @@ FMeshPassProcessor* CreateMobileBasePassCSMProcessor(const FScene* Scene, const
 	return new(FMemStack::Get()) FMobileBasePassMeshProcessor(Scene, Scene->GetFeatureLevel(), InViewIfDynamicMeshCommand, PassDrawRenderState, InDrawListContext, true);
 }
 
-FMeshPassProcessor* CreateMobileTranslucencyStandardPassProcessor(const FScene* Scene, const FSceneView* InViewIfDynamicMeshCommand, FMeshPassDrawListContext& InDrawListContext)
+FMeshPassProcessor* CreateMobileTranslucencyStandardPassProcessor(const FScene* Scene, const FSceneView* InViewIfDynamicMeshCommand, FMeshPassDrawListContext* InDrawListContext)
 {
 	FDrawingPolicyRenderState PassDrawRenderState(Scene->UniformBuffers.ViewUniformBuffer, Scene->UniformBuffers.MobileTranslucentBasePassUniformBuffer);
 	PassDrawRenderState.SetDepthStencilState(TStaticDepthStencilState<false, CF_DepthNearOrEqual>::GetRHI());
@@ -794,7 +794,7 @@ FMeshPassProcessor* CreateMobileTranslucencyStandardPassProcessor(const FScene*
 	return new(FMemStack::Get()) FMobileBasePassMeshProcessor(Scene, Scene->GetFeatureLevel(), InViewIfDynamicMeshCommand, PassDrawRenderState, InDrawListContext, false, ETranslucencyPass::TPT_StandardTranslucency);
 }
 
-FMeshPassProcessor* CreateMobileTranslucencyAfterDOFProcessor(const FScene* Scene, const FSceneView* InViewIfDynamicMeshCommand, FMeshPassDrawListContext& InDrawListContext)
+FMeshPassProcessor* CreateMobileTranslucencyAfterDOFProcessor(const FScene* Scene, const FSceneView* InViewIfDynamicMeshCommand, FMeshPassDrawListContext* InDrawListContext)
 {
 	FDrawingPolicyRenderState PassDrawRenderState(Scene->UniformBuffers.ViewUniformBuffer, Scene->UniformBuffers.MobileTranslucentBasePassUniformBuffer);
 	PassDrawRenderState.SetDepthStencilState(TStaticDepthStencilState<false, CF_DepthNearOrEqual>::GetRHI());
@@ -803,7 +803,7 @@ FMeshPassProcessor* CreateMobileTranslucencyAfterDOFProcessor(const FScene* Scen
 	return new(FMemStack::Get()) FMobileBasePassMeshProcessor(Scene, Scene->GetFeatureLevel(), InViewIfDynamicMeshCommand, PassDrawRenderState, InDrawListContext, false, ETranslucencyPass::TPT_TranslucencyAfterDOF);
 }
 
-FMeshPassProcessor* CreateMobileTranslucencyAllPassProcessor(const FScene* Scene, const FSceneView* InViewIfDynamicMeshCommand, FMeshPassDrawListContext& InDrawListContext)
+FMeshPassProcessor* CreateMobileTranslucencyAllPassProcessor(const FScene* Scene, const FSceneView* InViewIfDynamicMeshCommand, FMeshPassDrawListContext* InDrawListContext)
 {
 	FDrawingPolicyRenderState PassDrawRenderState(Scene->UniformBuffers.ViewUniformBuffer, Scene->UniformBuffers.MobileTranslucentBasePassUniformBuffer);
 	PassDrawRenderState.SetDepthStencilState(TStaticDepthStencilState<false, CF_DepthNearOrEqual>::GetRHI());
diff --git a/Engine/Source/Runtime/Renderer/Private/MobileBasePassRendering.cpp b/Engine/Source/Runtime/Renderer/Private/MobileBasePassRendering.cpp
--- a/Engine/Source/Runtime/Renderer/Private/MobileBasePassRendering.cpp
+++ b/Engine/Source/Runtime/Renderer/Private/MobileBasePassRendering.cpp
@@ -708,7 +708,7 @@ void FMobileSceneRenderer::RenderMobileEditorPrimitives(FRHICommandList& RHICmdL
 		if (UseMeshDrawCommandPipeline())
 		{
 			DrawDynamicMeshPass(View, RHICmdList,
-				[&View, &DrawRenderState](FDynamicPassMeshDrawListContext& DynamicMeshPassContext)
+				[&View, &DrawRenderState](FDynamicPassMeshDrawListContext* DynamicMeshPassContext)
 				{
 					FEditorPrimitivesBasePassMeshProcessor PassMeshProcessor(
 						View.Family->Scene->GetRenderScene(),
@@ -739,7 +739,7 @@ void FMobileSceneRenderer::RenderMobileEditorPrimitives(FRHICommandList& RHICmdL
 		if (UseMeshDrawCommandPipeline())
 		{
 			DrawDynamicMeshPass(View, RHICmdList,
-				[&View, &DrawRenderState](FDynamicPassMeshDrawListContext& DynamicMeshPassContext)
+				[&View, &DrawRenderState](FDynamicPassMeshDrawListContext* DynamicMeshPassContext)
 			{
 				FEditorPrimitivesBasePassMeshProcessor PassMeshProcessor(
 					View.Family->Scene->GetRenderScene(),
diff --git a/Engine/Source/Runtime/Renderer/Private/MobileBasePassRendering.h b/Engine/Source/Runtime/Renderer/Private/MobileBasePassRendering.h
--- a/Engine/Source/Runtime/Renderer/Private/MobileBasePassRendering.h
+++ b/Engine/Source/Runtime/Renderer/Private/MobileBasePassRendering.h
@@ -937,7 +937,7 @@ public:
 		ERHIFeatureLevel::Type InFeatureLevel, 
 		const FSceneView* InViewIfDynamicMeshCommand, 
 		const FDrawingPolicyRenderState& InDrawRenderState, 
-		FMeshPassDrawListContext& InDrawListContext,
+		FMeshPassDrawListContext* InDrawListContext,
 		bool bInCanReceiveCSM,
 		ETranslucencyPass::Type InTranslucencyPassType = ETranslucencyPass::TPT_MAX);
 
diff --git a/Engine/Source/Runtime/Renderer/Private/MobileSceneCaptureRendering.cpp b/Engine/Source/Runtime/Renderer/Private/MobileSceneCaptureRendering.cpp
--- a/Engine/Source/Runtime/Renderer/Private/MobileSceneCaptureRendering.cpp
+++ b/Engine/Source/Runtime/Renderer/Private/MobileSceneCaptureRendering.cpp
@@ -419,7 +419,7 @@ void UpdateSceneCaptureContentMobile_RenderThread(
 		}
 
 		// Register pass for InverseOpacity for this scope
-		extern FMeshPassProcessor* CreateMobileInverseOpacityPassProcessor(const FScene* Scene, const FSceneView* InViewIfDynamicMeshCommand, FMeshPassDrawListContext& InDrawListContext);
+		extern FMeshPassProcessor* CreateMobileInverseOpacityPassProcessor(const FScene* Scene, const FSceneView* InViewIfDynamicMeshCommand, FMeshPassDrawListContext* InDrawListContext);
 		FRegisterPassProcessorCreateFunction RegisterMobileInverseOpacityPass(&CreateMobileInverseOpacityPassProcessor, EShadingPath::Mobile, EMeshPass::MobileInverseOpacity, EMeshPassFlags::MainView);
 		
 		// Render the scene normally
diff --git a/Engine/Source/Runtime/Renderer/Private/MobileTranslucentRendering.cpp b/Engine/Source/Runtime/Renderer/Private/MobileTranslucentRendering.cpp
--- a/Engine/Source/Runtime/Renderer/Private/MobileTranslucentRendering.cpp
+++ b/Engine/Source/Runtime/Renderer/Private/MobileTranslucentRendering.cpp
@@ -817,7 +817,7 @@ public:
 	const FDrawingPolicyRenderState PassDrawRenderState;
 
 public:
-	FMobileInverseOpacityMeshProcessor(const FScene* InScene, ERHIFeatureLevel::Type InFeatureLevel, const FSceneView* InViewIfDynamicMeshCommand, const FDrawingPolicyRenderState& InDrawRenderState, FMeshPassDrawListContext& InDrawListContext)
+	FMobileInverseOpacityMeshProcessor(const FScene* InScene, ERHIFeatureLevel::Type InFeatureLevel, const FSceneView* InViewIfDynamicMeshCommand, const FDrawingPolicyRenderState& InDrawRenderState, FMeshPassDrawListContext* InDrawListContext)
 		: FMeshPassProcessor(InScene, InFeatureLevel, InViewIfDynamicMeshCommand, InDrawListContext)
 		, PassDrawRenderState(InDrawRenderState)
 	{
@@ -886,7 +886,7 @@ private:
 };
 
 // This pass is registered only when we render to scene capture, see UpdateSceneCaptureContentMobile_RenderThread()
-FMeshPassProcessor* CreateMobileInverseOpacityPassProcessor(const FScene* Scene, const FSceneView* InViewIfDynamicMeshCommand, FMeshPassDrawListContext& InDrawListContext)
+FMeshPassProcessor* CreateMobileInverseOpacityPassProcessor(const FScene* Scene, const FSceneView* InViewIfDynamicMeshCommand, FMeshPassDrawListContext* InDrawListContext)
 {
 	FDrawingPolicyRenderState PassDrawRenderState(Scene->UniformBuffers.ViewUniformBuffer, Scene->UniformBuffers.MobileTranslucentBasePassUniformBuffer);
 	PassDrawRenderState.SetDepthStencilState(TStaticDepthStencilState<false, CF_DepthNearOrEqual>::GetRHI());
diff --git a/Engine/Source/Runtime/Renderer/Private/PostProcess/PostProcessCompositeEditorPrimitives.cpp b/Engine/Source/Runtime/Renderer/Private/PostProcess/PostProcessCompositeEditorPrimitives.cpp
--- a/Engine/Source/Runtime/Renderer/Private/PostProcess/PostProcessCompositeEditorPrimitives.cpp
+++ b/Engine/Source/Runtime/Renderer/Private/PostProcess/PostProcessCompositeEditorPrimitives.cpp
@@ -310,7 +310,7 @@ static void RenderEditorPrimitives(FRHICommandListImmediate& RHICmdList, const F
 		0xFF, GET_STENCIL_BIT_MASK(RECEIVE_DECAL, 1) | STENCIL_LIGHTING_CHANNELS_MASK(0x7)>::GetRHI());
 
 	DrawDynamicMeshPass(View, RHICmdList,
-		[&View, &DrawRenderState](FDynamicPassMeshDrawListContext& DynamicMeshPassContext)
+		[&View, &DrawRenderState](FDynamicPassMeshDrawListContext* DynamicMeshPassContext)
 	{
 		FEditorPrimitivesBasePassMeshProcessor PassMeshProcessor(
 			View.Family->Scene->GetRenderScene(),
@@ -403,7 +403,7 @@ static void RenderForegroundEditorPrimitives(FRHICommandListImmediate& RHICmdLis
 		View.EditorSimpleElementCollector.DrawBatchedElements(RHICmdList, DrawRenderState, View, EBlendModeFilter::OpaqueAndMasked, SDPG_Foreground);
 
 		DrawDynamicMeshPass(View, RHICmdList,
-				[&View, &DrawRenderState](FDynamicPassMeshDrawListContext& DynamicMeshPassContext)
+				[&View, &DrawRenderState](FDynamicPassMeshDrawListContext* DynamicMeshPassContext)
 			{
 				FEditorPrimitivesBasePassMeshProcessor PassMeshProcessor(
 					View.Family->Scene->GetRenderScene(),
@@ -432,7 +432,7 @@ static void RenderForegroundEditorPrimitives(FRHICommandListImmediate& RHICmdLis
 		View.EditorSimpleElementCollector.DrawBatchedElements(RHICmdList, DrawRenderState, View, EBlendModeFilter::OpaqueAndMasked, SDPG_Foreground);
 
 		DrawDynamicMeshPass(View, RHICmdList,
-				[&View, &DrawRenderState](FDynamicPassMeshDrawListContext& DynamicMeshPassContext)
+				[&View, &DrawRenderState](FDynamicPassMeshDrawListContext* DynamicMeshPassContext)
 			{
 				FEditorPrimitivesBasePassMeshProcessor PassMeshProcessor(
 					View.Family->Scene->GetRenderScene(),
diff --git a/Engine/Source/Runtime/Renderer/Private/PrimitiveSceneInfo.cpp b/Engine/Source/Runtime/Renderer/Private/PrimitiveSceneInfo.cpp
--- a/Engine/Source/Runtime/Renderer/Private/PrimitiveSceneInfo.cpp
+++ b/Engine/Source/Runtime/Renderer/Private/PrimitiveSceneInfo.cpp
@@ -221,7 +221,7 @@ void FPrimitiveSceneInfo::CacheMeshDrawCommands(FRHICommandListImmediate& RHICmd
 						FCachedPassMeshDrawListContext CachedPassMeshDrawListContext(CommandInfo, SceneDrawList, *Scene);
 
 						PassProcessorCreateFunction CreateFunction = FPassProcessorManager::GetCreateFunction(ShadingPath, PassType);
-						FMeshPassProcessor* PassMeshProcessor = CreateFunction(Scene, nullptr, CachedPassMeshDrawListContext);
+						FMeshPassProcessor* PassMeshProcessor = CreateFunction(Scene, nullptr, &CachedPassMeshDrawListContext);
 
 						if (PassMeshProcessor != nullptr)
 						{
diff --git a/Engine/Source/Runtime/Renderer/Private/RayTracing/RayTracingMaterialHitShaders.cpp b/Engine/Source/Runtime/Renderer/Private/RayTracing/RayTracingMaterialHitShaders.cpp
--- a/Engine/Source/Runtime/Renderer/Private/RayTracing/RayTracingMaterialHitShaders.cpp
+++ b/Engine/Source/Runtime/Renderer/Private/RayTracing/RayTracingMaterialHitShaders.cpp
@@ -125,14 +125,14 @@ IMPLEMENT_MATERIALCHS_TYPE(TUniformLightMapPolicy<LMP_LQ_LIGHTMAP>, TLightMapPol
 IMPLEMENT_MATERIALCHS_TYPE(TUniformLightMapPolicy<LMP_HQ_LIGHTMAP>, TLightMapPolicyHQ, FAnyHitShader);
 IMPLEMENT_MATERIALCHS_TYPE(TUniformLightMapPolicy<LMP_DISTANCE_FIELD_SHADOWS_AND_HQ_LIGHTMAP>, TDistanceFieldShadowsAndLightMapPolicyHQ, FAnyHitShader);
 
-static FMeshPassProcessor* CreateRayTracingProcessor(const FScene* Scene, const FSceneView* InViewIfDynamicMeshCommand, FMeshPassDrawListContext& InDrawListContext)
+static FMeshPassProcessor* CreateRayTracingProcessor(const FScene* Scene, const FSceneView* InViewIfDynamicMeshCommand, FMeshPassDrawListContext* InDrawListContext)
 {
 	return IsRayTracingEnabled() ? new(FMemStack::Get()) FRayTracingMeshProcessor(Scene, Scene->GetFeatureLevel(), InViewIfDynamicMeshCommand, InDrawListContext) : nullptr;
 }
 
 FRegisterPassProcessorCreateFunction RegisterRayTracing(&CreateRayTracingProcessor, EShadingPath::Deferred, EMeshPass::RayTracing, EMeshPassFlags::CachedMeshCommands);
 
-FRayTracingMeshProcessor::FRayTracingMeshProcessor(const FScene* Scene, ERHIFeatureLevel::Type InFeatureLevel, const FSceneView* InViewIfDynamicMeshCommand, FMeshPassDrawListContext& InDrawListContext)
+FRayTracingMeshProcessor::FRayTracingMeshProcessor(const FScene* Scene, ERHIFeatureLevel::Type InFeatureLevel, const FSceneView* InViewIfDynamicMeshCommand, FMeshPassDrawListContext* InDrawListContext)
 	: FMeshPassProcessor(Scene, InFeatureLevel, InViewIfDynamicMeshCommand, InDrawListContext)
 {}
 
diff --git a/Engine/Source/Runtime/Renderer/Private/RayTracing/RayTracingMaterialHitShaders.h b/Engine/Source/Runtime/Renderer/Private/RayTracing/RayTracingMaterialHitShaders.h
--- a/Engine/Source/Runtime/Renderer/Private/RayTracing/RayTracingMaterialHitShaders.h
+++ b/Engine/Source/Runtime/Renderer/Private/RayTracing/RayTracingMaterialHitShaders.h
@@ -12,7 +12,7 @@ class FRayTracingMeshProcessor : public FMeshPassProcessor
 {
 public:
 
-	FRayTracingMeshProcessor(const FScene* Scene, ERHIFeatureLevel::Type InFeatureLevel, const FSceneView* InViewIfDynamicMeshCommand, FMeshPassDrawListContext& InDrawListContext);
+	FRayTracingMeshProcessor(const FScene* Scene, ERHIFeatureLevel::Type InFeatureLevel, const FSceneView* InViewIfDynamicMeshCommand, FMeshPassDrawListContext* InDrawListContext);
 
 	virtual void AddMeshBatch(const FMeshBatch& RESTRICT MeshBatch, uint64 BatchElementMask, const FPrimitiveSceneProxy* RESTRICT PrimitiveSceneProxy, int32 MeshId = -1) override final;
 
diff --git a/Engine/Source/Runtime/Renderer/Private/Renderer.cpp b/Engine/Source/Runtime/Renderer/Private/Renderer.cpp
--- a/Engine/Source/Runtime/Renderer/Private/Renderer.cpp
+++ b/Engine/Source/Runtime/Renderer/Private/Renderer.cpp
@@ -161,7 +161,7 @@ void FRendererModule::DrawTileMesh(FRHICommandListImmediate& RHICmdList, FDrawin
 				);
 
 				DrawDynamicMeshPass(View, RHICmdList,
-					[&View, &DrawRenderState, &DebugViewModePassUniformBuffer, &Mesh](FMeshPassDrawListContext& InDrawListContext)
+					[&View, &DrawRenderState, &DebugViewModePassUniformBuffer, &Mesh](FMeshPassDrawListContext* InDrawListContext)
 				{
 					FDebugViewModeMeshProcessor PassMeshProcessor(nullptr, View.GetFeatureLevel(), &View, DebugViewModePassUniformBuffer, false, InDrawListContext);
 					const uint64 DefaultBatchElementMask = ~0ull;
@@ -185,7 +185,7 @@ void FRendererModule::DrawTileMesh(FRHICommandListImmediate& RHICmdList, FDrawin
 				DrawRenderState.SetPassUniformBuffer(BasePassUniformBuffer);
 
 				DrawDynamicMeshPass(View, RHICmdList,
-					[&View, &DrawRenderState, &Mesh](FDynamicPassMeshDrawListContext& DynamicMeshPassContext)
+					[&View, &DrawRenderState, &Mesh](FDynamicPassMeshDrawListContext* DynamicMeshPassContext)
 					{
 						FBasePassMeshProcessor PassMeshProcessor(
 							nullptr,
@@ -219,7 +219,7 @@ void FRendererModule::DrawTileMesh(FRHICommandListImmediate& RHICmdList, FDrawin
 				ensureMsgf(HitProxyId == Mesh.BatchHitProxyId, TEXT("Only Mesh.BatchHitProxyId is used for hit testing."));
 
 				DrawDynamicMeshPass(View, RHICmdList,
-					[&View, &DrawRenderState, &Mesh](FDynamicPassMeshDrawListContext& DynamicMeshPassContext)
+					[&View, &DrawRenderState, &Mesh](FDynamicPassMeshDrawListContext* DynamicMeshPassContext)
 				{
 					FHitProxyMeshProcessor PassMeshProcessor(
 						nullptr,
@@ -241,7 +241,7 @@ void FRendererModule::DrawTileMesh(FRHICommandListImmediate& RHICmdList, FDrawin
 					DrawRenderState.SetPassUniformBuffer(BasePassUniformBuffer);
 
 					DrawDynamicMeshPass(View, RHICmdList,
-						[&View, &DrawRenderState, &Mesh](FDynamicPassMeshDrawListContext& DynamicMeshPassContext)
+						[&View, &DrawRenderState, &Mesh](FDynamicPassMeshDrawListContext* DynamicMeshPassContext)
 						{
 							FBasePassMeshProcessor PassMeshProcessor(
 								nullptr,
diff --git a/Engine/Source/Runtime/Renderer/Private/SceneHitProxyRendering.cpp b/Engine/Source/Runtime/Renderer/Private/SceneHitProxyRendering.cpp
--- a/Engine/Source/Runtime/Renderer/Private/SceneHitProxyRendering.cpp
+++ b/Engine/Source/Runtime/Renderer/Private/SceneHitProxyRendering.cpp
@@ -321,7 +321,7 @@ static void DoRenderHitProxies(FRHICommandListImmediate& RHICmdList, const FScen
 		}
 
 		DrawDynamicMeshPass(View, RHICmdList,
-			[&View, &DrawRenderState, LocalScene](FDynamicPassMeshDrawListContext& DynamicMeshPassContext)
+			[&View, &DrawRenderState, LocalScene](FDynamicPassMeshDrawListContext* DynamicMeshPassContext)
 		{
 			FHitProxyMeshProcessor PassMeshProcessor(
 				LocalScene,
@@ -350,7 +350,7 @@ static void DoRenderHitProxies(FRHICommandListImmediate& RHICmdList, const FScen
 		View.EditorSimpleElementCollector.DrawBatchedElements(RHICmdList, DrawRenderState, View, EBlendModeFilter::All, SDPG_Foreground);
 
 		DrawDynamicMeshPass(View, RHICmdList,
-			[&View, &DrawRenderState, LocalScene](FDynamicPassMeshDrawListContext& DynamicMeshPassContext)
+			[&View, &DrawRenderState, LocalScene](FDynamicPassMeshDrawListContext* DynamicMeshPassContext)
 		{
 			FHitProxyMeshProcessor PassMeshProcessor(
 				LocalScene,
@@ -369,7 +369,7 @@ static void DoRenderHitProxies(FRHICommandListImmediate& RHICmdList, const FScen
 		});
 
 		DrawDynamicMeshPass(View, RHICmdList,
-			[&View, &DrawRenderState, LocalScene](FDynamicPassMeshDrawListContext& DynamicMeshPassContext)
+			[&View, &DrawRenderState, LocalScene](FDynamicPassMeshDrawListContext* DynamicMeshPassContext)
 		{
 			FHitProxyMeshProcessor PassMeshProcessor(
 				LocalScene,
@@ -397,7 +397,7 @@ static void DoRenderHitProxies(FRHICommandListImmediate& RHICmdList, const FScen
 		DrawRenderState.SetDepthStencilState(TStaticDepthStencilState<false, CF_Always>::GetRHI());
 
 		DrawDynamicMeshPass(View, RHICmdList,
-			[&View, &DrawRenderState, LocalScene](FDynamicPassMeshDrawListContext& DynamicMeshPassContext)
+			[&View, &DrawRenderState, LocalScene](FDynamicPassMeshDrawListContext* DynamicMeshPassContext)
 			{
 				FHitProxyMeshProcessor PassMeshProcessor(
 					LocalScene,
@@ -420,7 +420,7 @@ static void DoRenderHitProxies(FRHICommandListImmediate& RHICmdList, const FScen
 		DrawRenderState.SetDepthStencilState(TStaticDepthStencilState<true, CF_DepthNearOrEqual>::GetRHI());
 
 		DrawDynamicMeshPass(View, RHICmdList,
-			[&View, &DrawRenderState, LocalScene](FDynamicPassMeshDrawListContext& DynamicMeshPassContext)
+			[&View, &DrawRenderState, LocalScene](FDynamicPassMeshDrawListContext* DynamicMeshPassContext)
 			{
 				FHitProxyMeshProcessor PassMeshProcessor(
 					LocalScene,
@@ -695,14 +695,14 @@ void FHitProxyMeshProcessor::Process(
 		ShaderElementData);
 }
 
-FHitProxyMeshProcessor::FHitProxyMeshProcessor(const FScene* Scene, const FSceneView* InViewIfDynamicMeshCommand, bool InbAllowTranslucentPrimitivesInHitProxy, const FDrawingPolicyRenderState& InRenderState, FMeshPassDrawListContext& InDrawListContext)
+FHitProxyMeshProcessor::FHitProxyMeshProcessor(const FScene* Scene, const FSceneView* InViewIfDynamicMeshCommand, bool InbAllowTranslucentPrimitivesInHitProxy, const FDrawingPolicyRenderState& InRenderState, FMeshPassDrawListContext* InDrawListContext)
 	: FMeshPassProcessor(Scene, Scene->GetFeatureLevel(), InViewIfDynamicMeshCommand, InDrawListContext)
 	, PassDrawRenderState(InRenderState)
 	, bAllowTranslucentPrimitivesInHitProxy(InbAllowTranslucentPrimitivesInHitProxy)
 {
 }
 
-FMeshPassProcessor* CreateHitProxyPassProcessor(const FScene* Scene, const FSceneView* InViewIfDynamicMeshCommand, FMeshPassDrawListContext& InDrawListContext)
+FMeshPassProcessor* CreateHitProxyPassProcessor(const FScene* Scene, const FSceneView* InViewIfDynamicMeshCommand, FMeshPassDrawListContext* InDrawListContext)
 {
 	FDrawingPolicyRenderState PassDrawRenderState(Scene->UniformBuffers.ViewUniformBuffer, Scene->UniformBuffers.HitProxyPassUniformBuffer);
 	PassDrawRenderState.SetDepthStencilState(TStaticDepthStencilState<true, CF_DepthNearOrEqual>::GetRHI());
@@ -710,7 +710,7 @@ FMeshPassProcessor* CreateHitProxyPassProcessor(const FScene* Scene, const FScen
 	return new(FMemStack::Get()) FHitProxyMeshProcessor(Scene, InViewIfDynamicMeshCommand, true, PassDrawRenderState, InDrawListContext);
 }
 
-FMeshPassProcessor* CreateHitProxyOpaqueOnlyPassProcessor(const FScene* Scene, const FSceneView* InViewIfDynamicMeshCommand, FMeshPassDrawListContext& InDrawListContext)
+FMeshPassProcessor* CreateHitProxyOpaqueOnlyPassProcessor(const FScene* Scene, const FSceneView* InViewIfDynamicMeshCommand, FMeshPassDrawListContext* InDrawListContext)
 {
 	FDrawingPolicyRenderState PassDrawRenderState(Scene->UniformBuffers.ViewUniformBuffer, Scene->UniformBuffers.HitProxyPassUniformBuffer);
 	PassDrawRenderState.SetDepthStencilState(TStaticDepthStencilState<true, CF_DepthNearOrEqual>::GetRHI());
@@ -844,7 +844,7 @@ int32 FEditorSelectionMeshProcessor::GetStencilValue(const FSceneView* View, con
 TMap<const FPrimitiveSceneProxy*, int32> FEditorSelectionMeshProcessor::ProxyToStencilIndex;
 TMap<FName, int32> FEditorSelectionMeshProcessor::ActorNameToStencilIndex;
 
-FEditorSelectionMeshProcessor::FEditorSelectionMeshProcessor(const FScene* Scene, const FSceneView* InViewIfDynamicMeshCommand, FMeshPassDrawListContext& InDrawListContext)
+FEditorSelectionMeshProcessor::FEditorSelectionMeshProcessor(const FScene* Scene, const FSceneView* InViewIfDynamicMeshCommand, FMeshPassDrawListContext* InDrawListContext)
 	: FMeshPassProcessor(Scene, Scene->GetFeatureLevel(), InViewIfDynamicMeshCommand, InDrawListContext)
 {
 	checkf(InViewIfDynamicMeshCommand, TEXT("Editor selection mesh process required dynamic mesh command mode."));
@@ -860,7 +860,7 @@ FEditorSelectionMeshProcessor::FEditorSelectionMeshProcessor(const FScene* Scene
 	PassDrawRenderState.SetPassUniformBuffer(Scene->UniformBuffers.EditorSelectionPassUniformBuffer);
 }
 
-FMeshPassProcessor* CreateEditorSelectionPassProcessor(const FScene* Scene, const FSceneView* InViewIfDynamicMeshCommand, FMeshPassDrawListContext& InDrawListContext)
+FMeshPassProcessor* CreateEditorSelectionPassProcessor(const FScene* Scene, const FSceneView* InViewIfDynamicMeshCommand, FMeshPassDrawListContext* InDrawListContext)
 {
 	return new(FMemStack::Get()) FEditorSelectionMeshProcessor(Scene, InViewIfDynamicMeshCommand, InDrawListContext);
 }
diff --git a/Engine/Source/Runtime/Renderer/Private/SceneHitProxyRendering.h b/Engine/Source/Runtime/Renderer/Private/SceneHitProxyRendering.h
--- a/Engine/Source/Runtime/Renderer/Private/SceneHitProxyRendering.h
+++ b/Engine/Source/Runtime/Renderer/Private/SceneHitProxyRendering.h
@@ -20,7 +20,7 @@ class FHitProxyMeshProcessor : public FMeshPassProcessor
 {
 public:
 
-	FHitProxyMeshProcessor(const FScene* Scene, const FSceneView* InViewIfDynamicMeshCommand, bool InbAllowTranslucentPrimitivesInHitProxy, const FDrawingPolicyRenderState& InRenderState, FMeshPassDrawListContext& InDrawListContext);
+	FHitProxyMeshProcessor(const FScene* Scene, const FSceneView* InViewIfDynamicMeshCommand, bool InbAllowTranslucentPrimitivesInHitProxy, const FDrawingPolicyRenderState& InRenderState, FMeshPassDrawListContext* InDrawListContext);
 
 	virtual void AddMeshBatch(const FMeshBatch& RESTRICT MeshBatch, uint64 BatchElementMask, const FPrimitiveSceneProxy* RESTRICT PrimitiveSceneProxy, int32 StaticMeshId = -1) override final;
 
@@ -48,7 +48,7 @@ class FEditorSelectionMeshProcessor : public FMeshPassProcessor
 {
 public:
 
-	FEditorSelectionMeshProcessor(const FScene* Scene, const FSceneView* InViewIfDynamicMeshCommand, FMeshPassDrawListContext& InDrawListContext);
+	FEditorSelectionMeshProcessor(const FScene* Scene, const FSceneView* InViewIfDynamicMeshCommand, FMeshPassDrawListContext* InDrawListContext);
 
 	virtual void AddMeshBatch(const FMeshBatch& RESTRICT MeshBatch, uint64 BatchElementMask, const FPrimitiveSceneProxy* RESTRICT PrimitiveSceneProxy, int32 StaticMeshId = -1) override final;
 
diff --git a/Engine/Source/Runtime/Renderer/Private/SceneRendering.cpp b/Engine/Source/Runtime/Renderer/Private/SceneRendering.cpp
--- a/Engine/Source/Runtime/Renderer/Private/SceneRendering.cpp
+++ b/Engine/Source/Runtime/Renderer/Private/SceneRendering.cpp
@@ -4425,7 +4425,7 @@ void FSceneRenderer::GenerateDynamicMeshDrawCommands(FViewInfo& View, FMeshDrawC
 			);
 
 			PassProcessorCreateFunction CreateFunction = FPassProcessorManager::GetCreateFunction(ShadingPath, PassType);
-			FMeshPassProcessor* PassMeshProcessor = CreateFunction(Scene, &View, DynamicPassMeshDrawListContext);
+			FMeshPassProcessor* PassMeshProcessor = CreateFunction(Scene, &View, &DynamicPassMeshDrawListContext);
 
 			if (PassMeshProcessor != nullptr)
 			{
diff --git a/Engine/Source/Runtime/Renderer/Private/ShadowDepthRendering.cpp b/Engine/Source/Runtime/Renderer/Private/ShadowDepthRendering.cpp
--- a/Engine/Source/Runtime/Renderer/Private/ShadowDepthRendering.cpp
+++ b/Engine/Source/Runtime/Renderer/Private/ShadowDepthRendering.cpp
@@ -2650,7 +2650,7 @@ FShadowDepthPassMeshProcessor::FShadowDepthPassMeshProcessor(
 	const TUniformBufferRef<FViewUniformShaderParameters>& InViewUniformBuffer,
 	FUniformBufferRHIParamRef InPassUniformBuffer,
 	FShadowDepthType InShadowDepthType,
-	FMeshPassDrawListContext& InDrawListContext)
+	FMeshPassDrawListContext* InDrawListContext)
 	: FMeshPassProcessor(Scene, Scene->GetFeatureLevel(), InViewIfDynamicMeshCommand, InDrawListContext)
 	, PassDrawRenderState(FDrawingPolicyRenderState(InViewUniformBuffer, InPassUniformBuffer))
 	, ShadowDepthType(InShadowDepthType)
@@ -2660,7 +2660,7 @@ FShadowDepthPassMeshProcessor::FShadowDepthPassMeshProcessor(
 
 FShadowDepthType CSMShadowDepthType(true, false, false);
 
-FMeshPassProcessor* CreateCSMShadowDepthPassProcessor(const FScene* Scene, const FSceneView* InViewIfDynamicMeshCommand, FMeshPassDrawListContext& InDrawListContext)
+FMeshPassProcessor* CreateCSMShadowDepthPassProcessor(const FScene* Scene, const FSceneView* InViewIfDynamicMeshCommand, FMeshPassDrawListContext* InDrawListContext)
 {
 	FUniformBufferRHIParamRef PassUniformBuffer = nullptr;
 
diff --git a/Engine/Source/Runtime/Renderer/Private/ShadowRendering.h b/Engine/Source/Runtime/Renderer/Private/ShadowRendering.h
--- a/Engine/Source/Runtime/Renderer/Private/ShadowRendering.h
+++ b/Engine/Source/Runtime/Renderer/Private/ShadowRendering.h
@@ -304,7 +304,7 @@ public:
 		const TUniformBufferRef<FViewUniformShaderParameters>& InViewUniformBuffer,
 		FUniformBufferRHIParamRef InPassUniformBuffer,
 		FShadowDepthType InShadowDepthType,
-		FMeshPassDrawListContext& InDrawListContext);
+		FMeshPassDrawListContext* InDrawListContext);
 
 	virtual void AddMeshBatch(const FMeshBatch& RESTRICT MeshBatch, uint64 BatchElementMask, const FPrimitiveSceneProxy* RESTRICT PrimitiveSceneProxy, int32 StaticMeshId = -1) override final;
 
diff --git a/Engine/Source/Runtime/Renderer/Private/ShadowSetup.cpp b/Engine/Source/Runtime/Renderer/Private/ShadowSetup.cpp
--- a/Engine/Source/Runtime/Renderer/Private/ShadowSetup.cpp
+++ b/Engine/Source/Runtime/Renderer/Private/ShadowSetup.cpp
@@ -1208,7 +1208,7 @@ void FProjectedShadowInfo::SetupMeshDrawCommandsForShadowDepth(FSceneRenderer& R
 		ShadowDepthView->ViewUniformBuffer, 
 		PassUniformBuffer,
 		GetShadowDepthType(),
-		ShadowDepthContext);
+		&ShadowDepthContext);
 
 	const int32 NumDynamicMeshBatches = DynamicSubjectMeshElements.Num();
 
@@ -1300,7 +1300,7 @@ void FProjectedShadowInfo::SetupMeshDrawCommandsForProjectionStenciling(FSceneRe
 				false,
 				DDM_AllOccluders,
 				false,
-				ProjectionStencilingContext);
+				&ProjectionStencilingContext);
 
 			// Pre-shadows mask by receiver elements, self-shadow mask by subject elements.
 			// Note that self-shadow pre-shadows still mask by receiver elements.
diff --git a/Engine/Source/Runtime/Renderer/Private/TranslucentRendering.cpp b/Engine/Source/Runtime/Renderer/Private/TranslucentRendering.cpp
--- a/Engine/Source/Runtime/Renderer/Private/TranslucentRendering.cpp
+++ b/Engine/Source/Runtime/Renderer/Private/TranslucentRendering.cpp
@@ -467,7 +467,7 @@ void RenderViewTranslucencyInner(FRHICommandListImmediate& RHICmdList, const FVi
 				QUICK_SCOPE_CYCLE_COUNTER(RenderTranslucencyParallel_SDPG_World);
 
 				DrawDynamicMeshPass(View, RHICmdList,
-					[&View, &DrawRenderState](FDynamicPassMeshDrawListContext& DynamicMeshPassContext)
+					[&View, &DrawRenderState](FDynamicPassMeshDrawListContext* DynamicMeshPassContext)
 				{
 					FBasePassMeshProcessor PassMeshProcessor(
 						View.Family->Scene->GetRenderScene(),
@@ -491,7 +491,7 @@ void RenderViewTranslucencyInner(FRHICommandListImmediate& RHICmdList, const FVi
 				QUICK_SCOPE_CYCLE_COUNTER(RenderTranslucencyParallel_SDPG_Foreground);
 
 				DrawDynamicMeshPass(View, RHICmdList,
-					[&View, &DrawRenderState](FDynamicPassMeshDrawListContext& DynamicMeshPassContext)
+					[&View, &DrawRenderState](FDynamicPassMeshDrawListContext* DynamicMeshPassContext)
 				{
 					FBasePassMeshProcessor PassMeshProcessor(
 						View.Family->Scene->GetRenderScene(),
diff --git a/Engine/Source/Runtime/Renderer/Private/VelocityRendering.cpp b/Engine/Source/Runtime/Renderer/Private/VelocityRendering.cpp
--- a/Engine/Source/Runtime/Renderer/Private/VelocityRendering.cpp
+++ b/Engine/Source/Runtime/Renderer/Private/VelocityRendering.cpp
@@ -764,7 +764,7 @@ void FVelocityMeshProcessor::Process(
 		ShaderElementData);
 }
 
-FVelocityMeshProcessor::FVelocityMeshProcessor(const FScene* Scene, const FSceneView* InViewIfDynamicMeshCommand, const FDrawingPolicyRenderState& InPassDrawRenderState, FMeshPassDrawListContext& InDrawListContext)
+FVelocityMeshProcessor::FVelocityMeshProcessor(const FScene* Scene, const FSceneView* InViewIfDynamicMeshCommand, const FDrawingPolicyRenderState& InPassDrawRenderState, FMeshPassDrawListContext* InDrawListContext)
 	: FMeshPassProcessor(Scene, Scene->GetFeatureLevel(), InViewIfDynamicMeshCommand, InDrawListContext)
 {
 	PassDrawRenderState = InPassDrawRenderState;
@@ -772,7 +772,7 @@ FVelocityMeshProcessor::FVelocityMeshProcessor(const FScene* Scene, const FScene
 	PassDrawRenderState.SetPassUniformBuffer(Scene->UniformBuffers.VelocityPassUniformBuffer);
 }
 
-FMeshPassProcessor* CreateVelocityPassProcessor(const FScene* Scene, const FSceneView* InViewIfDynamicMeshCommand, FMeshPassDrawListContext& InDrawListContext)
+FMeshPassProcessor* CreateVelocityPassProcessor(const FScene* Scene, const FSceneView* InViewIfDynamicMeshCommand, FMeshPassDrawListContext* InDrawListContext)
 {
 	FDrawingPolicyRenderState VelocityPassState;
 	VelocityPassState.SetBlendState(TStaticBlendState<CW_RGBA>::GetRHI());
diff --git a/Engine/Source/Runtime/Renderer/Private/VelocityRendering.h b/Engine/Source/Runtime/Renderer/Private/VelocityRendering.h
--- a/Engine/Source/Runtime/Renderer/Private/VelocityRendering.h
+++ b/Engine/Source/Runtime/Renderer/Private/VelocityRendering.h
@@ -51,7 +51,7 @@ class FVelocityMeshProcessor : public FMeshPassProcessor
 {
 public:
 
-	FVelocityMeshProcessor(const FScene* Scene, const FSceneView* InViewIfDynamicMeshCommand, const FDrawingPolicyRenderState& InPassDrawRenderState, FMeshPassDrawListContext& InDrawListContext);
+	FVelocityMeshProcessor(const FScene* Scene, const FSceneView* InViewIfDynamicMeshCommand, const FDrawingPolicyRenderState& InPassDrawRenderState, FMeshPassDrawListContext* InDrawListContext);
 
 	virtual void AddMeshBatch(const FMeshBatch& RESTRICT MeshBatch, uint64 BatchElementMask, const FPrimitiveSceneProxy* RESTRICT PrimitiveSceneProxy, int32 StaticMeshId = -1) override final;
 
diff --git a/Engine/Source/Runtime/Renderer/Private/VolumetricFogVoxelization.cpp b/Engine/Source/Runtime/Renderer/Private/VolumetricFogVoxelization.cpp
--- a/Engine/Source/Runtime/Renderer/Private/VolumetricFogVoxelization.cpp
+++ b/Engine/Source/Runtime/Renderer/Private/VolumetricFogVoxelization.cpp
@@ -441,7 +441,7 @@ IMPLEMENT_MATERIAL_SHADER_TYPE(template<>,TVoxelizeVolumePS<VMode_Object_Box>,TE
 class FVoxelizeVolumeMeshProcessor : public FMeshPassProcessor
 {
 public:
-	FVoxelizeVolumeMeshProcessor(const FScene* Scene, const FViewInfo* InViewIfDynamicMeshCommand, FMeshPassDrawListContext& InDrawListContext);
+	FVoxelizeVolumeMeshProcessor(const FScene* Scene, const FViewInfo* InViewIfDynamicMeshCommand, FMeshPassDrawListContext* InDrawListContext);
 
 	void AddMeshBatch(const FMeshBatch& RESTRICT MeshBatch, uint64 BatchElementMask, int32 NumVoxelizationPasses, const FPrimitiveSceneProxy* RESTRICT PrimitiveSceneProxy);
 
@@ -464,7 +464,7 @@ private:
 	FDrawingPolicyRenderState PassDrawRenderState;
 };
 
-FVoxelizeVolumeMeshProcessor::FVoxelizeVolumeMeshProcessor(const FScene* Scene, const FViewInfo* InViewIfDynamicMeshCommand, FMeshPassDrawListContext& InDrawListContext)
+FVoxelizeVolumeMeshProcessor::FVoxelizeVolumeMeshProcessor(const FScene* Scene, const FViewInfo* InViewIfDynamicMeshCommand, FMeshPassDrawListContext* InDrawListContext)
 	: FMeshPassProcessor(Scene, Scene->GetFeatureLevel(), InViewIfDynamicMeshCommand, InDrawListContext)
 {
 	PassDrawRenderState.SetBlendState(TStaticBlendState<
@@ -658,7 +658,7 @@ void FDeferredShadingSceneRenderer::VoxelizeFogVolumePrimitives(
 			DrawRenderState.SetViewUniformBuffer(Scene->UniformBuffers.VoxelizeVolumeViewUniformBuffer);
 
 			DrawDynamicMeshPass(View, RHICmdList,
-				[&View, VolumetricFogDistance, &RHICmdList, &VolumetricFogGridSize, &GridZParams](FDynamicPassMeshDrawListContext& DynamicMeshPassContext)
+				[&View, VolumetricFogDistance, &RHICmdList, &VolumetricFogGridSize, &GridZParams](FDynamicPassMeshDrawListContext* DynamicMeshPassContext)
 			{
 				FVoxelizeVolumeMeshProcessor PassMeshProcessor(
 					View.Family->Scene->GetRenderScene(),
diff --git a/Engine/Source/Runtime/Renderer/Public/MeshPassProcessor.h b/Engine/Source/Runtime/Renderer/Public/MeshPassProcessor.h
--- a/Engine/Source/Runtime/Renderer/Public/MeshPassProcessor.h
+++ b/Engine/Source/Runtime/Renderer/Public/MeshPassProcessor.h
@@ -734,12 +734,17 @@ public:
 	const FScene* RESTRICT Scene;
 	ERHIFeatureLevel::Type FeatureLevel;
 	const FSceneView* ViewIfDynamicMeshCommand;
-	FMeshPassDrawListContext& DrawListContext;
+	FMeshPassDrawListContext* DrawListContext;
 
-	RENDERER_API FMeshPassProcessor(const FScene* InScene, ERHIFeatureLevel::Type InFeatureLevel, const FSceneView* InViewIfDynamicMeshCommand, FMeshPassDrawListContext& InDrawListContext);
+	RENDERER_API FMeshPassProcessor(const FScene* InScene, ERHIFeatureLevel::Type InFeatureLevel, const FSceneView* InViewIfDynamicMeshCommand, FMeshPassDrawListContext* InDrawListContext);
 
 	virtual ~FMeshPassProcessor() {}
 
+	void SetDrawListContext(FMeshPassDrawListContext* InDrawListContext)
+	{
+		DrawListContext = InDrawListContext;
+	}
+
 	// FMeshPassProcessor interface
 	// Add a FMeshBatch to the pass
 	virtual void AddMeshBatch(const FMeshBatch& RESTRICT MeshBatch, uint64 BatchElementMask, const FPrimitiveSceneProxy* RESTRICT PrimitiveSceneProxy, int32 StaticMeshId = -1) = 0;
@@ -788,7 +793,7 @@ private:
 	RENDERER_API int32 GetDrawCommandPrimitiveId(const FPrimitiveSceneInfo* RESTRICT PrimitiveSceneInfo, const FMeshBatchElement& BatchElement) const;
 };
 
-typedef FMeshPassProcessor* (*PassProcessorCreateFunction)(const FScene* Scene, const FSceneView* InViewIfDynamicMeshCommand, FMeshPassDrawListContext& InDrawListContext);
+typedef FMeshPassProcessor* (*PassProcessorCreateFunction)(const FScene* Scene, const FSceneView* InViewIfDynamicMeshCommand, FMeshPassDrawListContext* InDrawListContext);
 
 enum class EMeshPassFlags
 {
diff --git a/Engine/Source/Runtime/Renderer/Public/MeshPassProcessor.inl b/Engine/Source/Runtime/Renderer/Public/MeshPassProcessor.inl
--- a/Engine/Source/Runtime/Renderer/Public/MeshPassProcessor.inl
+++ b/Engine/Source/Runtime/Renderer/Public/MeshPassProcessor.inl
@@ -100,7 +100,7 @@ void FMeshPassProcessor::BuildMeshDrawCommands(
 		if ((1ull << BatchElementIndex) & BatchElementMask)
 		{
 			const FMeshBatchElement& BatchElement = MeshBatch.Elements[BatchElementIndex];
-			FMeshDrawCommand& MeshDrawCommand = DrawListContext.AddCommand(SharedMeshDrawCommand);
+			FMeshDrawCommand& MeshDrawCommand = DrawListContext->AddCommand(SharedMeshDrawCommand);
 
 			if (PassShaders.VertexShader)
 			{
@@ -133,7 +133,7 @@ void FMeshPassProcessor::BuildMeshDrawCommands(
 
 			const int32 DrawPrimitiveId = GetDrawCommandPrimitiveId(PrimitiveSceneInfo, BatchElement);
 
-			DrawListContext.FinalizeCommand(MeshBatch, BatchElementIndex, DrawPrimitiveId, MeshFillMode, MeshCullMode, InstanceFactor, SortKey, MeshDrawCommand, true);
+			DrawListContext->FinalizeCommand(MeshBatch, BatchElementIndex, DrawPrimitiveId, MeshFillMode, MeshCullMode, InstanceFactor, SortKey, MeshDrawCommand, true);
 		}
 	}
 }
@@ -150,7 +150,7 @@ void DrawDynamicMeshPass(const FSceneView& View, FRHICommandList& RHICmdList, co
 
 	FDynamicPassMeshDrawListContext DynamicMeshPassContext(DynamicMeshDrawCommandStorage, VisibleMeshDrawCommands);
 
-	BuildPassProcessorLambda(DynamicMeshPassContext);
+	BuildPassProcessorLambda(&DynamicMeshPassContext);
 
 	DrawDynamicMeshPassPrivate(View, RHICmdList, VisibleMeshDrawCommands, DynamicMeshDrawCommandStorage);
 }
@@ -201,7 +201,7 @@ void FMeshPassProcessor::BuildRayTracingDrawCommands(
 		if ((1ull << BatchElementIndex) & BatchElementMask)
 		{
 			const FMeshBatchElement& BatchElement = MeshBatch.Elements[BatchElementIndex];
-			FMeshDrawCommand& MeshDrawCommand = DrawListContext.AddCommand(SharedMeshDrawCommand);
+			FMeshDrawCommand& MeshDrawCommand = DrawListContext->AddCommand(SharedMeshDrawCommand);
 
 #if RHI_RAYTRACING
 			if (PassShaders.RayHitGroupShader)
@@ -215,7 +215,7 @@ void FMeshPassProcessor::BuildRayTracingDrawCommands(
 
 			const int32 DrawPrimitiveId = 0;
 			const int32 ScenePrimitiveId = 0;
-			DrawListContext.FinalizeCommand(MeshBatch, BatchElementIndex, DrawPrimitiveId, MeshFillMode, MeshCullMode, InstanceFactor, SortKey, MeshDrawCommand, false);
+			DrawListContext->FinalizeCommand(MeshBatch, BatchElementIndex, DrawPrimitiveId, MeshFillMode, MeshCullMode, InstanceFactor, SortKey, MeshDrawCommand, false);
 		}
 	}
 }
\ No newline at end of file
