Fixed compilation on clang.

diff --git a/Engine/Source/Runtime/Renderer/Public/MeshPassProcessor.h b/Engine/Source/Runtime/Renderer/Public/MeshPassProcessor.h
--- a/Engine/Source/Runtime/Renderer/Public/MeshPassProcessor.h
+++ b/Engine/Source/Runtime/Renderer/Public/MeshPassProcessor.h
@@ -88,7 +88,6 @@ public:
 	void Reset() { Data = 0; }
 	bool IsEmpty() const { return Data == 0; }
 
-private:
 	uint32 Data;
 };
 
