Removed hit proxy mesh processing from non editor builds.

diff --git a/Engine/Source/Runtime/Renderer/Private/Renderer.cpp b/Engine/Source/Runtime/Renderer/Private/Renderer.cpp
--- a/Engine/Source/Runtime/Renderer/Private/Renderer.cpp
+++ b/Engine/Source/Runtime/Renderer/Private/Renderer.cpp
@@ -218,6 +218,7 @@ void FRendererModule::DrawTileMesh(FRHICommandListImmediate& RHICmdList, FDrawin
 			{
 				ensureMsgf(HitProxyId == Mesh.BatchHitProxyId, TEXT("Only Mesh.BatchHitProxyId is used for hit testing."));
 
+#if WITH_EDITOR
 				DrawDynamicMeshPass(View, RHICmdList,
 					[&View, &DrawRenderState, &Mesh](FDynamicPassMeshDrawListContext* DynamicMeshPassContext)
 				{
@@ -231,6 +232,7 @@ void FRendererModule::DrawTileMesh(FRHICommandListImmediate& RHICmdList, FDrawin
 					const uint64 DefaultBatchElementMask = ~0ull;
 					PassMeshProcessor.AddMeshBatch(Mesh, DefaultBatchElementMask, nullptr);
 				});
+#endif
 			}
 			else
 			{
diff --git a/Engine/Source/Runtime/Renderer/Private/SceneVisibility.cpp b/Engine/Source/Runtime/Renderer/Private/SceneVisibility.cpp
--- a/Engine/Source/Runtime/Renderer/Private/SceneVisibility.cpp
+++ b/Engine/Source/Runtime/Renderer/Private/SceneVisibility.cpp
@@ -2144,6 +2144,8 @@ struct FRelevancePacket
 								DrawCommandPacket.AddCommandsForMesh(PrimitiveIndex, PrimitiveSceneInfo, StaticMeshRelevance, StaticMesh, Scene, bCanCache, EMeshPass::DebugViewMode);
 							}
 #endif
+
+#if WITH_EDITOR
 							if (StaticMeshRelevance.bSelectable)
 							{
 								if (View.bAllowTranslucentPrimitivesInHitProxy)
@@ -2155,6 +2157,7 @@ struct FRelevancePacket
 									DrawCommandPacket.AddCommandsForMesh(PrimitiveIndex, PrimitiveSceneInfo, StaticMeshRelevance, StaticMesh, Scene, bCanCache, EMeshPass::HitProxyOpaqueOnly);
 								}
 							}
+#endif
 							 
 							if (ViewRelevance.bVelocityRelevance 
 								&& FVelocityRendering::PrimitiveHasVelocity(View.GetFeatureLevel(), PrimitiveSceneInfo)
