Fixed vs2015 compile warning (implicit cast from bitmask to bool)

diff --git a/Engine/Source/Runtime/Renderer/Public/MeshPassProcessor.h b/Engine/Source/Runtime/Renderer/Public/MeshPassProcessor.h
--- a/Engine/Source/Runtime/Renderer/Public/MeshPassProcessor.h
+++ b/Engine/Source/Runtime/Renderer/Public/MeshPassProcessor.h
@@ -82,7 +82,7 @@ public:
 	}
 
 	void Set(EMeshPass::Type Pass) { Data |= (1 << Pass); }
-	bool Get(EMeshPass::Type Pass) const { return Data & (1 << Pass); }
+	bool Get(EMeshPass::Type Pass) const { return !!(Data & (1 << Pass)); }
 
 	void AppendTo(FMeshPassMask& Mask) const { Mask.Data |= Data; }
 	void Reset() { Data = 0; }
