Fix directory casing in Renderer.Build.cs

Linux was throwing errors in GPF due to folder case sensitivity

diff --git a/Engine/Source/Runtime/Renderer/Renderer.Build.cs b/Engine/Source/Runtime/Renderer/Renderer.Build.cs
--- a/Engine/Source/Runtime/Renderer/Renderer.Build.cs
+++ b/Engine/Source/Runtime/Renderer/Renderer.Build.cs
@@ -11,7 +11,7 @@ public class Renderer : ModuleRules
 				"Runtime/Renderer/Private",
 				"Runtime/Renderer/Private/CompositionLighting",
 				"Runtime/Renderer/Private/PostProcess",
-				"../Shaders/Private/Raytracing"
+				"../Shaders/Private/RayTracing"
     }
 		);
 
