Fixed bCullLightsToGrid in FDeferredShadingSceneRenderer::ComputeLightGrid always has to be set to true when ray tracing is enabled.

diff --git a/Engine/Source/Runtime/Renderer/Private/LightGridInjection.cpp b/Engine/Source/Runtime/Renderer/Private/LightGridInjection.cpp
--- a/Engine/Source/Runtime/Renderer/Private/LightGridInjection.cpp
+++ b/Engine/Source/Runtime/Renderer/Private/LightGridInjection.cpp
@@ -439,8 +439,8 @@ void FDeferredShadingSceneRenderer::ComputeLightGrid(FRHICommandListImmediate& R
 			bAnyViewUsesForwardLighting |= View.bTranslucentSurfaceLighting || ShouldRenderVolumetricFog();
 		}
 
-		const bool bCullLightsToGrid = GLightCullingQuality && ((IsForwardShadingEnabled(ShaderPlatform) || bAnyViewUsesForwardLighting) && ViewFamily.EngineShowFlags.DirectLighting);
-
+		const bool bCullLightsToGrid = GLightCullingQuality && (ViewFamily.EngineShowFlags.DirectLighting && (IsForwardShadingEnabled(ShaderPlatform) || bAnyViewUsesForwardLighting || IsRayTracingEnabled()));
+			   
 		FSimpleLightArray SimpleLights;
 
 		if (bCullLightsToGrid)
