Moved dynamic mesh draw command generation to a parallel task (FMeshDrawCommandPassSetupTask).
*InitViews computes approximate max number of dynamic draws and use it to allocate mesh draw command resources (primitive id vertex buffer, mesh draw command arrays, number of parallel draw tasks etc).
*DynamicMeshDrawCommandStorage was moved to the FMeshDrawCommandPassSetupTaskContext, as it's now internally handled by the task.
*Mobile BasePass sort key computation and BasePass / BasePassCSM list merging based on CSM visibility was fully moved to the FMeshDrawCommandPassSetupTask.

diff --git a/Engine/Source/Runtime/Engine/Private/Materials/Material.cpp b/Engine/Source/Runtime/Engine/Private/Materials/Material.cpp
--- a/Engine/Source/Runtime/Engine/Private/Materials/Material.cpp
+++ b/Engine/Source/Runtime/Engine/Private/Materials/Material.cpp
@@ -283,7 +283,7 @@ public:
 
 	virtual FMaterial* GetMaterialNoFallback(ERHIFeatureLevel::Type InFeatureLevel) const
 	{
-		checkSlow(IsInRenderingThread());
+		checkSlow(IsInParallelRenderingThread());
 		return Material->GetMaterialResource(InFeatureLevel);
 	}
 
diff --git a/Engine/Source/Runtime/RenderCore/Private/RenderCore.cpp b/Engine/Source/Runtime/RenderCore/Private/RenderCore.cpp
--- a/Engine/Source/Runtime/RenderCore/Private/RenderCore.cpp
+++ b/Engine/Source/Runtime/RenderCore/Private/RenderCore.cpp
@@ -108,7 +108,7 @@ DEFINE_STAT(STAT_CreateWholeSceneProjectedShadow);
 DEFINE_STAT(STAT_AddViewDependentWholeSceneShadowsForView);
 DEFINE_STAT(STAT_SetupInteractionShadows);
 DEFINE_STAT(STAT_GetDynamicMeshElements);
-DEFINE_STAT(STAT_SortAndMergeMeshDrawCommands);
+DEFINE_STAT(STAT_SetupMeshPass);
 DEFINE_STAT(STAT_UpdateStaticMeshesTime);
 DEFINE_STAT(STAT_StaticRelevance);
 DEFINE_STAT(STAT_ViewRelevance);
diff --git a/Engine/Source/Runtime/RenderCore/Public/RenderCore.h b/Engine/Source/Runtime/RenderCore/Public/RenderCore.h
--- a/Engine/Source/Runtime/RenderCore/Public/RenderCore.h
+++ b/Engine/Source/Runtime/RenderCore/Public/RenderCore.h
@@ -59,7 +59,7 @@ DECLARE_CYCLE_STAT_EXTERN(TEXT("Compute View Relevance"),STAT_ComputeViewRelevan
 DECLARE_CYCLE_STAT_EXTERN(TEXT("Static Mesh Relevance"),STAT_StaticRelevance,STATGROUP_InitViews, RENDERCORE_API);
 DECLARE_CYCLE_STAT_EXTERN(TEXT("UpdateStaticMeshes"),STAT_UpdateStaticMeshesTime,STATGROUP_InitViews, RENDERCORE_API);
 DECLARE_CYCLE_STAT_EXTERN(TEXT("GetDynamicMeshElements"),STAT_GetDynamicMeshElements,STATGROUP_InitViews, RENDERCORE_API);
-DECLARE_CYCLE_STAT_EXTERN(TEXT("SortAndMergeVisibleMeshDrawCommands"), STAT_SortAndMergeMeshDrawCommands, STATGROUP_InitViews, RENDERCORE_API);
+DECLARE_CYCLE_STAT_EXTERN(TEXT("SetupMeshPass"), STAT_SetupMeshPass, STATGROUP_InitViews, RENDERCORE_API);
 DECLARE_CYCLE_STAT_EXTERN(TEXT("Init dynamic shadows"),STAT_InitDynamicShadowsTime,STATGROUP_InitViews, RENDERCORE_API);
 DECLARE_CYCLE_STAT_EXTERN(TEXT("Update Preshadow Cache"),STAT_UpdatePreshadowCache,STATGROUP_InitViews, RENDERCORE_API);
 DECLARE_CYCLE_STAT_EXTERN(TEXT("Init Projected Shadow"),STAT_InitProjectedShadowVisibility,STATGROUP_InitViews, RENDERCORE_API);
diff --git a/Engine/Source/Runtime/Renderer/Private/DeferredShadingRenderer.cpp b/Engine/Source/Runtime/Renderer/Private/DeferredShadingRenderer.cpp
--- a/Engine/Source/Runtime/Renderer/Private/DeferredShadingRenderer.cpp
+++ b/Engine/Source/Runtime/Renderer/Private/DeferredShadingRenderer.cpp
@@ -607,7 +607,7 @@ bool FDeferredShadingSceneRenderer::GatherRayTracingWorldInstances(FRHICommandLi
 				&DummyDynamicPrimitiveShaderData,
 				ViewFamily.GetFeatureLevel());
 
-			View.DynamicMeshDrawCommandStorage[EMeshPass::RayTracing].MeshDrawCommands.Reserve(Scene->Primitives.Num());
+			View.RaytraycingDynamicMeshDrawCommandStorage.MeshDrawCommands.Reserve(Scene->Primitives.Num());
 			View.RaytraycingVisibleMeshDrawCommands.Reserve(Scene->Primitives.Num());
 		}
 
@@ -707,7 +707,7 @@ bool FDeferredShadingSceneRenderer::GatherRayTracingWorldInstances(FRHICommandLi
 
 								FDynamicPassMeshDrawListContext DynamicPassMeshDrawListContext
 								(
-									View.DynamicMeshDrawCommandStorage[EMeshPass::RayTracing],
+									View.RaytraycingDynamicMeshDrawCommandStorage,
 									View.RaytraycingVisibleMeshDrawCommands
 								);
 								FRayTracingMeshProcessor RayTracingMeshProcessor(Scene, View.GetFeatureLevel(), &View, &DynamicPassMeshDrawListContext);
@@ -750,7 +750,7 @@ bool FDeferredShadingSceneRenderer::GatherRayTracingWorldInstances(FRHICommandLi
 							FViewInfo& View = Views[ViewIndex];
 							FDynamicPassMeshDrawListContext DynamicPassMeshDrawListContext
 							(
-								View.DynamicMeshDrawCommandStorage[EMeshPass::RayTracing],
+								View.RaytraycingDynamicMeshDrawCommandStorage,
 								View.RaytraycingVisibleMeshDrawCommands
 							);
 							FRayTracingMeshProcessor RayTracingMeshProcessor(Scene, View.GetFeatureLevel(), &View, &DynamicPassMeshDrawListContext);
@@ -798,7 +798,7 @@ bool FDeferredShadingSceneRenderer::GatherRayTracingWorldInstances(FRHICommandLi
 							FViewInfo& View = Views[ViewIndex];
 							FDynamicPassMeshDrawListContext DynamicPassMeshDrawListContext
 							(
-								View.DynamicMeshDrawCommandStorage[EMeshPass::RayTracing],
+								View.RaytraycingDynamicMeshDrawCommandStorage,
 								View.RaytraycingVisibleMeshDrawCommands
 							);
 							FRayTracingMeshProcessor RayTracingMeshProcessor(Scene, View.GetFeatureLevel(), &View, &DynamicPassMeshDrawListContext);
diff --git a/Engine/Source/Runtime/Renderer/Private/MobileBasePass.cpp b/Engine/Source/Runtime/Renderer/Private/MobileBasePass.cpp
--- a/Engine/Source/Runtime/Renderer/Private/MobileBasePass.cpp
+++ b/Engine/Source/Runtime/Renderer/Private/MobileBasePass.cpp
@@ -439,90 +439,6 @@ static FMeshDrawCommandSortKey GetBasePassStaticSortKey(EBlendMode BlendMode)
 	return SortKey;
 }
 
-static uint64 GetSortKey(bool bMasked, bool bBackground, int32 PipelineId, int32 StateBucketId, float PipelineDistance, float PrimitiveDistance)
-{
-	const float PrimitiveDepthQuantization = ((1 << 14) - 1);
-	
-	union
-	{
-		uint64 PackedInt;
-		struct
-		{
-			uint64 DepthBits : 14;			// Order by primitive depth
-			uint64 StateBucketId : 20; 		// Order by state bucket
-			uint64 PipelineId : 20;			// Order by PSO
-			uint64 PipelineDepthBits : 8;	// Order PSOs front to back
-			uint64 Background : 1;			// Non-background meshes first 
-			uint64 Masked : 1;				// Non-masked first
-		} Fields;
-	} Key;
-
-	union FFloatToInt { float F; uint32 I; };
-	FFloatToInt F2I;
-	
-	Key.PackedInt = 0;
-	Key.Fields.Masked = bMasked;
-	Key.Fields.Background = bBackground;
-	F2I.F = PipelineDistance/HALF_WORLD_MAX;
-	Key.Fields.PipelineDepthBits = (F2I.I >> 23) & 0xff; // 8 bit exponent
-	Key.Fields.PipelineId = PipelineId;
-	Key.Fields.StateBucketId = StateBucketId;
-	Key.Fields.DepthBits = int32((FMath::Min<float>(PrimitiveDistance, HALF_WORLD_MAX) / HALF_WORLD_MAX) * PrimitiveDepthQuantization);
-	
-	return Key.PackedInt;
-}
-
-void MobileBasePass::ComputeBasePassSortKeys(const FScene& Scene, const FViewInfo& View, FMeshCommandOneFrameArray& VisibleMeshCommands)
-{
-	FVector ViewLocation = View.ViewMatrices.GetViewOrigin();
-	int32 NumCmds = VisibleMeshCommands.Num();
-
-	// pre-compute distance to a group of meshes that share same PSO
-	TMap<int32, float> PipelineDistances;
-	PipelineDistances.Reserve(256);
-
-	for (int32 CmdIdx = 0; CmdIdx < NumCmds; ++CmdIdx)
-	{
-		FVisibleMeshDrawCommand& Cmd = VisibleMeshCommands[CmdIdx];
-
-		float PrimitiveDistance = 0;
-
-		if (Cmd.DrawPrimitiveId < Scene.PrimitiveBounds.Num())
-		{
-			const FPrimitiveBounds& PrimitiveBounds = Scene.PrimitiveBounds[Cmd.DrawPrimitiveId];
-			PrimitiveDistance = (PrimitiveBounds.BoxSphereBounds.Origin - ViewLocation).Size();
-		}
-		
-		float& PipelineDistance = PipelineDistances.FindOrAdd(Cmd.MeshDrawCommand->CachedPipelineId.GetId());
-		// not sure what could be better: average distance, max or min
-		PipelineDistance = FMath::Max(PipelineDistance, PrimitiveDistance);
-	}
-	
-	// compute sort key for each mesh command
-	for (int32 CmdIdx = 0; CmdIdx < NumCmds; ++CmdIdx)
-	{
-		FVisibleMeshDrawCommand& Cmd = VisibleMeshCommands[CmdIdx];
-		float PrimitiveDistance = 0;
-		bool bMasked = false;
-		bool bBackground = false;
-
-		if (Cmd.DrawPrimitiveId < Scene.PrimitiveBounds.Num())
-		{
-			const FPrimitiveBounds& PrimitiveBounds = Scene.PrimitiveBounds[Cmd.DrawPrimitiveId];
-			PrimitiveDistance = (PrimitiveBounds.BoxSphereBounds.Origin - ViewLocation).Size();
-			bBackground = (PrimitiveBounds.BoxSphereBounds.SphereRadius > HALF_WORLD_MAX/4.0f); // TODO: per-primitive bBackgound flag
-			bMasked = View.PrimitiveViewRelevanceMap[Cmd.DrawPrimitiveId].bMaskedRelevance;
-		}
-
-		int32 PipelineId = Cmd.MeshDrawCommand->CachedPipelineId.GetId();
-		float PipelineDistance = PipelineDistances.FindRef(PipelineId);
-		// poor man StateID, can't use Cmd.StateBucketId as it is unique for each primitive if platform does not support auto-instancing
-		int32 StateBucketId = PointerHash(Cmd.MeshDrawCommand->IndexBuffer);
-				
-		Cmd.SortKey.PackedData = GetSortKey(bMasked, bBackground, PipelineId, StateBucketId, PipelineDistance, PrimitiveDistance);
-	}
-}
-
 template<>
 void TMobileBasePassPSPolicyParamType<FUniformLightMapPolicy>::GetShaderBindings(
 	const FScene* Scene,
diff --git a/Engine/Source/Runtime/Renderer/Private/MobileBasePassRendering.h b/Engine/Source/Runtime/Renderer/Private/MobileBasePassRendering.h
--- a/Engine/Source/Runtime/Renderer/Private/MobileBasePassRendering.h
+++ b/Engine/Source/Runtime/Renderer/Private/MobileBasePassRendering.h
@@ -525,8 +525,6 @@ namespace MobileBasePass
 
 	void SetOpaqueRenderState(FDrawingPolicyRenderState& DrawRenderState, const FPrimitiveSceneProxy* PrimitiveSceneProxy, const FMaterial& Material, bool bEnableReceiveDecalOutput);
 	void SetTranslucentRenderState(FDrawingPolicyRenderState& DrawRenderState, const FMaterial& Material);
-
-	void ComputeBasePassSortKeys(const FScene& Scene, const FViewInfo& View, FMeshCommandOneFrameArray& VisibleMeshCommands);
 };
 
 
diff --git a/Engine/Source/Runtime/Renderer/Private/MobileShadingRenderer.cpp b/Engine/Source/Runtime/Renderer/Private/MobileShadingRenderer.cpp
--- a/Engine/Source/Runtime/Renderer/Private/MobileShadingRenderer.cpp
+++ b/Engine/Source/Runtime/Renderer/Private/MobileShadingRenderer.cpp
@@ -113,26 +113,37 @@ TUniformBufferRef<FMobileDirectionalLightShaderParameters>& GetNullMobileDirecti
 	return NullLightParams->UniformBufferRHI;
 }
 
-void FMobileSceneRenderer::SortMobileBasePassAfterShadowInit(FExclusiveDepthStencil::Type BasePassDepthStencilAccess, FMeshDrawCommandsPerPassPerView& VisibleCommandsPerView)
+void FMobileSceneRenderer::SetupMobileBasePassAfterShadowInit(FExclusiveDepthStencil::Type BasePassDepthStencilAccess, FViewVisibleCommandsPerView& ViewCommandsPerView)
 {
 	// Sort front to back on all platforms, even HSR benefits from it
 	//const bool bWantsFrontToBackSorting = (GHardwareHiddenSurfaceRemoval == false);
 
-	// compute keys for front to back sorting and dispatch sort.
+	// compute keys for front to back sorting and dispatch pass setup.
 	for (int32 ViewIndex = 0; ViewIndex < Views.Num(); ++ViewIndex)
 	{
 		FViewInfo& View = Views[ViewIndex];
-		FMeshDrawCommandsPerPass& VisibleCommands = VisibleCommandsPerView[ViewIndex];
-		
-		FMeshCommandOneFrameArray& BasePassVisibleCommands = VisibleCommands[EMeshPass::BasePass];
-		MobileBasePass::ComputeBasePassSortKeys(*Scene, View, BasePassVisibleCommands);
-
-		// Run sorting on BasePass, as it's ignored inside FSceneRenderer::SortMeshDrawCommands, so it can be done after shadow init on mobile.
-		if (VisibleCommands[EMeshPass::BasePass].Num() > 0)
-		{
-			FParallelMeshDrawCommandPass& Pass = View.ParallelMeshDrawCommandPasses[EMeshPass::BasePass];
-			Pass.DispatchSortAndMerge(View, Scene->GetShadingPath(), Scene->PrimitiveBounds, EMeshPass::BasePass, BasePassDepthStencilAccess, VisibleCommands[EMeshPass::BasePass]);
-		}
+		FViewCommands& ViewCommands = ViewCommandsPerView[ViewIndex];
+
+		PassProcessorCreateFunction CreateFunction = FPassProcessorManager::GetCreateFunction(EShadingPath::Mobile, EMeshPass::BasePass);
+		FMeshPassProcessor* MeshPassProcessor = CreateFunction(Scene, &View, nullptr);
+
+		PassProcessorCreateFunction BasePassCSMCreateFunction = FPassProcessorManager::GetCreateFunction(EShadingPath::Mobile, EMeshPass::MobileBasePassCSM);
+		FMeshPassProcessor* BasePassCSMMeshPassProcessor = BasePassCSMCreateFunction(Scene, &View, nullptr);
+
+		// Run sorting on BasePass, as it's ignored inside FSceneRenderer::SetupMeshPass, so it can be done after shadow init on mobile.
+		FParallelMeshDrawCommandPass& Pass = View.ParallelMeshDrawCommandPasses[EMeshPass::BasePass];
+		Pass.DispatchPassSetup(
+			Scene,
+			View,
+			EMeshPass::BasePass,
+			BasePassDepthStencilAccess,
+			MeshPassProcessor,
+			View.DynamicMeshElements,
+			View.MaxNumVisibleDynamicMeshes[EMeshPass::BasePass],
+			ViewCommands.DynamicMeshCommandBuildRequests[EMeshPass::BasePass],
+			ViewCommands.MeshCommands[EMeshPass::BasePass],
+			BasePassCSMMeshPassProcessor,
+			&ViewCommands.MeshCommands[EMeshPass::MobileBasePassCSM]);
 	}
 }
 
@@ -149,13 +160,13 @@ void FMobileSceneRenderer::InitViews(FRHICommandListImmediate& RHICmdList)
 	SCOPE_CYCLE_COUNTER(STAT_InitViewsTime);
 
 	FILCUpdatePrimTaskData ILCTaskData;
-	FMeshDrawCommandsPerPassPerView VisibleCommandsPerView;
-	VisibleCommandsPerView.SetNum(Views.Num());
+	FViewVisibleCommandsPerView ViewCommandsPerView;
+	ViewCommandsPerView.SetNum(Views.Num());
 
 	const FExclusiveDepthStencil::Type BasePassDepthStencilAccess = FExclusiveDepthStencil::DepthWrite_StencilWrite;
 
 	PreVisibilityFrameSetup(RHICmdList);
-	ComputeViewVisibility(RHICmdList, BasePassDepthStencilAccess, VisibleCommandsPerView);
+	ComputeViewVisibility(RHICmdList, BasePassDepthStencilAccess, ViewCommandsPerView);
 	PostVisibilityFrameSetup(ILCTaskData);
 
 	const bool bDynamicShadows = ViewFamily.EngineShowFlags.DynamicShadows;
@@ -163,10 +174,10 @@ void FMobileSceneRenderer::InitViews(FRHICommandListImmediate& RHICmdList)
 	if (bDynamicShadows && !IsSimpleForwardShadingEnabled(ShaderPlatform))
 	{
 		// Setup dynamic shadows.
-		InitDynamicShadows(RHICmdList, VisibleCommandsPerView);
+		InitDynamicShadows(RHICmdList);
 	}
 
-	SortMobileBasePassAfterShadowInit(BasePassDepthStencilAccess, VisibleCommandsPerView);
+	SetupMobileBasePassAfterShadowInit(BasePassDepthStencilAccess, ViewCommandsPerView);
 
 	// if we kicked off ILC update via task, wait and finalize.
 	if (ILCTaskData.TaskRef.IsValid())
diff --git a/Engine/Source/Runtime/Renderer/Private/SceneRendering.cpp b/Engine/Source/Runtime/Renderer/Private/SceneRendering.cpp
--- a/Engine/Source/Runtime/Renderer/Private/SceneRendering.cpp
+++ b/Engine/Source/Runtime/Renderer/Private/SceneRendering.cpp
@@ -99,10 +99,10 @@ static TAutoConsoleVariable<int32> CVarMeshDrawCommandsDynamicInstancing(
 	TEXT("Whether to dynamically combine multiple compatible visible Mesh Draw Commands into one instanced draw on vertex factories that support it."),
 	ECVF_RenderThreadSafe);
 
-static TAutoConsoleVariable<int32> CVarMeshDrawCommandsParallelProcess(
-	TEXT("r.MeshDrawCommands.ParallelSortAndMerge"),
+static TAutoConsoleVariable<int32> CVarMeshDrawCommandsParallelPassSetup(
+	TEXT("r.MeshDrawCommands.ParallelPassSetup"),
 	1,
-	TEXT("Whether to sort and merge mesh visible mesh draw commands in parallel."),
+	TEXT("Whether to setup mesh draw command pass in parallel."),
 	ECVF_RenderThreadSafe);
 
 bool IsDynamicInstancingEnabled()
@@ -920,6 +920,9 @@ void FViewInfo::Init()
 
 	DitherFadeInUniformBuffer = nullptr;
 	DitherFadeOutUniformBuffer = nullptr;
+
+	VisibleDynamicMeshesPassMask.Reset();
+	FMemory::Memset(MaxNumVisibleDynamicMeshes, 0, sizeof(MaxNumVisibleDynamicMeshes));
 }
 
 FViewInfo::~FViewInfo()
@@ -2761,6 +2764,129 @@ void UpdateTranslucentMeshSortKeys(
 	}
 }
 
+static uint64 GetMobileBasePassSortKey(bool bMasked, bool bBackground, int32 PipelineId, int32 StateBucketId, float PipelineDistance, float PrimitiveDistance)
+{
+	const float PrimitiveDepthQuantization = ((1 << 14) - 1);
+
+	union
+	{
+		uint64 PackedInt;
+		struct
+		{
+			uint64 DepthBits : 14;			// Order by primitive depth
+			uint64 StateBucketId : 20; 		// Order by state bucket
+			uint64 PipelineId : 20;			// Order by PSO
+			uint64 PipelineDepthBits : 8;	// Order PSOs front to back
+			uint64 Background : 1;			// Non-background meshes first 
+			uint64 Masked : 1;				// Non-masked first
+		} Fields;
+	} Key;
+
+	union FFloatToInt { float F; uint32 I; };
+	FFloatToInt F2I;
+
+	Key.PackedInt = 0;
+	Key.Fields.Masked = bMasked;
+	Key.Fields.Background = bBackground;
+	F2I.F = PipelineDistance / HALF_WORLD_MAX;
+	Key.Fields.PipelineDepthBits = (F2I.I >> 23) & 0xff; // 8 bit exponent
+	Key.Fields.PipelineId = PipelineId;
+	Key.Fields.StateBucketId = StateBucketId;
+	Key.Fields.DepthBits = int32((FMath::Min<float>(PrimitiveDistance, HALF_WORLD_MAX) / HALF_WORLD_MAX) * PrimitiveDepthQuantization);
+
+	return Key.PackedInt;
+}
+
+/**
+* Merge mobile BasePass with BasePassCSM based on CSM visibility in order to select appropriate shader for given command.
+*/
+void MergeMobileBasePassMeshDrawCommands(
+	const FMobileCSMVisibilityInfo& MobileCSMVisibilityInfo,
+	int32 ScenePrimitiveNum,
+	FMeshCommandOneFrameArray& MeshCommands,
+	FMeshCommandOneFrameArray& MeshCommandsCSM
+	)
+{
+	if (MobileCSMVisibilityInfo.bMobileDynamicCSMInUse)
+	{
+		// determine per view CSM visibility
+		checkf(MeshCommands.Num() == MeshCommandsCSM.Num(), TEXT("VisibleMeshDrawCommands of BasePass and MobileBasePassCSM are expected to match."));
+		for (int32 i = MeshCommands.Num() - 1; i >= 0; --i)
+		{
+			FVisibleMeshDrawCommand& MeshCommand = MeshCommands[i];
+			FVisibleMeshDrawCommand& MeshCommandCSM = MeshCommandsCSM[i];
+
+			if (MobileCSMVisibilityInfo.bAlwaysUseCSM 
+				|| (MeshCommand.DrawPrimitiveId < ScenePrimitiveNum && MobileCSMVisibilityInfo.MobilePrimitiveCSMReceiverVisibilityMap[MeshCommand.DrawPrimitiveId]))
+			{
+				checkf(MeshCommand.DrawPrimitiveId == MeshCommandCSM.DrawPrimitiveId, TEXT("VisibleMeshDrawCommands of BasePass and MobileBasePassCSM are expected to match."));
+				// Use CSM's VisibleMeshDrawCommand.
+				MeshCommand = MeshCommandCSM;
+			}
+		}
+		MeshCommandsCSM.Reset();
+	}
+}
+
+/**
+* Compute mesh sort keys for the mobile base pass
+*/
+void UpdateMobileBasePassMeshSortKeys(
+	const FVector& ViewOrigin,
+	const TArray<struct FPrimitiveBounds>& ScenePrimitiveBounds,
+	FMeshCommandOneFrameArray& VisibleMeshCommands
+)
+{
+	QUICK_SCOPE_CYCLE_COUNTER(STAT_UpdateMobileBasePassMeshSortKeys);
+
+	int32 NumCmds = VisibleMeshCommands.Num();
+
+	// pre-compute distance to a group of meshes that share same PSO
+	TMap<int32, float> PipelineDistances;
+	PipelineDistances.Reserve(256);
+
+	for (int32 CmdIdx = 0; CmdIdx < NumCmds; ++CmdIdx)
+	{
+		FVisibleMeshDrawCommand& Cmd = VisibleMeshCommands[CmdIdx];
+
+		float PrimitiveDistance = 0;
+
+		if (Cmd.DrawPrimitiveId < ScenePrimitiveBounds.Num())
+		{
+			const FPrimitiveBounds& PrimitiveBounds = ScenePrimitiveBounds[Cmd.DrawPrimitiveId];
+			PrimitiveDistance = (PrimitiveBounds.BoxSphereBounds.Origin - ViewOrigin).Size();
+		}
+
+		float& PipelineDistance = PipelineDistances.FindOrAdd(Cmd.MeshDrawCommand->CachedPipelineId.GetId());
+		// not sure what could be better: average distance, max or min
+		PipelineDistance = FMath::Max(PipelineDistance, PrimitiveDistance);
+	}
+
+	// compute sort key for each mesh command
+	for (int32 CmdIdx = 0; CmdIdx < NumCmds; ++CmdIdx)
+	{
+		FVisibleMeshDrawCommand& Cmd = VisibleMeshCommands[CmdIdx];
+		float PrimitiveDistance = 0;
+		bool bMasked = false;
+		bool bBackground = false;
+
+		if (Cmd.DrawPrimitiveId < ScenePrimitiveBounds.Num())
+		{
+			const FPrimitiveBounds& PrimitiveBounds = ScenePrimitiveBounds[Cmd.DrawPrimitiveId];
+			PrimitiveDistance = (PrimitiveBounds.BoxSphereBounds.Origin - ViewOrigin).Size();
+			bBackground = (PrimitiveBounds.BoxSphereBounds.SphereRadius > HALF_WORLD_MAX / 4.0f); // TODO: per-primitive bBackgound flag
+			bMasked = Cmd.SortKey.PackedData & 0x1 ? true : false; // Set in MobileBasePass.cpp - GetBasePassStaticSortKey
+		}
+
+		int32 PipelineId = Cmd.MeshDrawCommand->CachedPipelineId.GetId();
+		float PipelineDistance = PipelineDistances.FindRef(PipelineId);
+		// poor man StateID, can't use Cmd.StateBucketId as it is unique for each primitive if platform does not support auto-instancing
+		int32 StateBucketId = PointerHash(Cmd.MeshDrawCommand->IndexBuffer);
+
+		Cmd.SortKey.PackedData = GetMobileBasePassSortKey(bMasked, bBackground, PipelineId, StateBucketId, PipelineDistance, PrimitiveDistance);
+	}
+}
+
 /**
  * Build mesh draw command primitive Id buffer for instancing.
  * TempVisibleMeshDrawCommands must be presized for NewPassVisibleMeshDrawCommands.
@@ -2872,6 +2998,119 @@ void BuildMeshDrawCommandPrimitiveIdBuffer(
 	}
 }
 
+/**
+ * Converts each FMeshBatch into a set of FMeshDrawCommands for a specific mesh pass type.
+ */
+void GenerateDynamicMeshDrawCommands(
+	const FViewInfo& View,
+	EShadingPath ShadingPath,
+	EMeshPass::Type PassType,
+	FMeshPassProcessor* PassMeshProcessor,
+	const TArray<FMeshBatchAndRelevance, SceneRenderingAllocator>& DynamicMeshElements,
+	const TArray<const FStaticMesh*, SceneRenderingAllocator>& DynamicMeshCommandBuildRequests,
+	FMeshCommandOneFrameArray& VisibleCommands,
+	FDynamicMeshDrawCommandStorage& MeshDrawCommandStorage
+)
+{
+	QUICK_SCOPE_CYCLE_COUNTER(STAT_GenerateDynamicMeshDrawCommands);
+	check(PassMeshProcessor);
+
+	FDynamicPassMeshDrawListContext DynamicPassMeshDrawListContext(
+		MeshDrawCommandStorage,
+		VisibleCommands
+	);
+	PassMeshProcessor->SetDrawListContext(&DynamicPassMeshDrawListContext);
+
+	const int32 NumDynamicMeshBatches = DynamicMeshElements.Num();
+
+	for (int32 MeshIndex = 0; MeshIndex < NumDynamicMeshBatches; MeshIndex++)
+	{
+		const FMeshBatchAndRelevance& MeshAndRelevance = DynamicMeshElements[MeshIndex];
+		check(!MeshAndRelevance.Mesh->bRequiresPerElementVisibility);
+		const uint64 BatchElementMask = ~0ull;
+
+		PassMeshProcessor->AddMeshBatch(*MeshAndRelevance.Mesh, BatchElementMask, MeshAndRelevance.PrimitiveSceneProxy);
+	}
+
+	const int32 NumStaticMeshBatches = DynamicMeshCommandBuildRequests.Num();
+
+	for (int32 MeshIndex = 0; MeshIndex < NumStaticMeshBatches; MeshIndex++)
+	{
+		const FStaticMesh* StaticMeshBatch = DynamicMeshCommandBuildRequests[MeshIndex];
+		const uint64 BatchElementMask = StaticMeshBatch->bRequiresPerElementVisibility ? View.StaticMeshBatchVisibility[StaticMeshBatch->BatchVisibilityId] : ~0ull;
+
+		PassMeshProcessor->AddMeshBatch(*StaticMeshBatch, BatchElementMask, StaticMeshBatch->PrimitiveSceneInfo->Proxy, StaticMeshBatch->Id);
+	}
+}
+
+/**
+ * Special version of GenerateDynamicMeshDrawCommands for the mobile base pass.
+ * Based on CSM visibility it will generate mesh draw commands using either normal base pass processor or CSM base pass processor.
+*/
+void GenerateMobileBasePassDynamicMeshDrawCommands(
+	const FViewInfo& View,
+	EShadingPath ShadingPath,
+	EMeshPass::Type PassType,
+	FMeshPassProcessor* PassMeshProcessor,
+	FMeshPassProcessor* MobilePassCSMPassMeshProcessor,
+	const TArray<FMeshBatchAndRelevance, SceneRenderingAllocator>& DynamicMeshElements,
+	const TArray<const FStaticMesh*, SceneRenderingAllocator>& DynamicMeshCommandBuildRequests,
+	FMeshCommandOneFrameArray& VisibleCommands,
+	FDynamicMeshDrawCommandStorage& MeshDrawCommandStorage
+)
+{
+	QUICK_SCOPE_CYCLE_COUNTER(STAT_GenerateMobileBasePassDynamicMeshDrawCommands);
+	check(PassMeshProcessor && MobilePassCSMPassMeshProcessor);
+
+	FDynamicPassMeshDrawListContext DynamicPassMeshDrawListContext(
+		MeshDrawCommandStorage,
+		VisibleCommands
+	);
+	PassMeshProcessor->SetDrawListContext(&DynamicPassMeshDrawListContext);
+	MobilePassCSMPassMeshProcessor->SetDrawListContext(&DynamicPassMeshDrawListContext);
+
+	const FMobileCSMVisibilityInfo& MobileCSMVisibilityInfo = View.MobileCSMVisibilityInfo;
+	
+	const int32 NumDynamicMeshBatches = DynamicMeshElements.Num();
+
+	for (int32 MeshIndex = 0; MeshIndex < NumDynamicMeshBatches; MeshIndex++)
+	{
+		const FMeshBatchAndRelevance& MeshAndRelevance = DynamicMeshElements[MeshIndex];
+		check(!MeshAndRelevance.Mesh->bRequiresPerElementVisibility);
+		const uint64 BatchElementMask = ~0ull;
+
+		const int32 PrimitiveIndex = MeshAndRelevance.PrimitiveSceneProxy->GetPrimitiveSceneInfo()->GetIndex();
+		if (MobileCSMVisibilityInfo.bMobileDynamicCSMInUse 
+			&& (MobileCSMVisibilityInfo.bAlwaysUseCSM || MobileCSMVisibilityInfo.MobilePrimitiveCSMReceiverVisibilityMap[PrimitiveIndex]))
+		{
+			MobilePassCSMPassMeshProcessor->AddMeshBatch(*MeshAndRelevance.Mesh, BatchElementMask, MeshAndRelevance.PrimitiveSceneProxy);
+		}
+		else
+		{
+			PassMeshProcessor->AddMeshBatch(*MeshAndRelevance.Mesh, BatchElementMask, MeshAndRelevance.PrimitiveSceneProxy);
+		}
+	}
+
+	const int32 NumStaticMeshBatches = DynamicMeshCommandBuildRequests.Num();
+
+	for (int32 MeshIndex = 0; MeshIndex < NumStaticMeshBatches; MeshIndex++)
+	{
+		const FStaticMesh* StaticMeshBatch = DynamicMeshCommandBuildRequests[MeshIndex];
+		const uint64 BatchElementMask = StaticMeshBatch->bRequiresPerElementVisibility ? View.StaticMeshBatchVisibility[StaticMeshBatch->BatchVisibilityId] : ~0ull;
+
+		const int32 PrimitiveIndex = StaticMeshBatch->PrimitiveSceneInfo->Proxy->GetPrimitiveSceneInfo()->GetIndex();
+		if (MobileCSMVisibilityInfo.bMobileDynamicCSMInUse
+			&& (MobileCSMVisibilityInfo.bAlwaysUseCSM || MobileCSMVisibilityInfo.MobilePrimitiveCSMReceiverVisibilityMap[PrimitiveIndex]))
+		{
+			MobilePassCSMPassMeshProcessor->AddMeshBatch(*StaticMeshBatch, BatchElementMask, StaticMeshBatch->PrimitiveSceneInfo->Proxy, StaticMeshBatch->Id);
+		}
+		else
+		{
+			PassMeshProcessor->AddMeshBatch(*StaticMeshBatch, BatchElementMask, StaticMeshBatch->PrimitiveSceneInfo->Proxy, StaticMeshBatch->Id);
+		}
+	}
+}
+
 /**
 * Apply view overrides to existing mesh draw commands (e.g. reverse culling mode for rendering planar reflections).
 * TempVisibleMeshDrawCommands must be presized for NewPassVisibleMeshDrawCommands.
@@ -2937,33 +3176,33 @@ void ApplyViewOverridesToMeshDrawCommands(
 	}
 }
 
-FAutoConsoleTaskPriority CPrio_FMeshDrawCommandSortAndMergeTask(
-	TEXT("TaskGraph.TaskPriorities.MeshDrawCommandSortAndMergeTask"),
-	TEXT("Task and thread priority for FMeshDrawCommandSortAndMergeTask."),
+FAutoConsoleTaskPriority CPrio_FMeshDrawCommandPassSetupTask(
+	TEXT("TaskGraph.TaskPriorities.FMeshDrawCommandPassSetupTask"),
+	TEXT("Task and thread priority for FMeshDrawCommandPassSetupTask."),
 	ENamedThreads::HighThreadPriority, // if we have high priority task threads, then use them...
 	ENamedThreads::NormalTaskPriority, // .. at normal task priority
 	ENamedThreads::HighTaskPriority // if we don't have hi pri threads, then use normal priority threads at high task priority instead
 );
 
 /**
- * Task for parallel sorting and building primitive id buffer for mesh draw commands.
+ * Task for a parallel setup of mesh draw commands. Includes generation of dynamic mesh draw commands, sorting, merging etc.
  */
-class FMeshDrawCommandSortAndMergeTask
+class FMeshDrawCommandPassSetupTask
 {
 public:
-	FMeshDrawCommandSortAndMergeTask(FMeshDrawCommandSortAndMergeTaskContext& InContext)
+	FMeshDrawCommandPassSetupTask(FMeshDrawCommandPassSetupTaskContext& InContext)
 		: Context(InContext)
 	{
 	}
 
 	FORCEINLINE TStatId GetStatId() const
 	{
-		RETURN_QUICK_DECLARE_CYCLE_STAT(FMeshDrawCommandSortAndMergeTask, STATGROUP_TaskGraphTasks);
+		RETURN_QUICK_DECLARE_CYCLE_STAT(FMeshDrawCommandPassSetupTask, STATGROUP_TaskGraphTasks);
 	}
 
 	ENamedThreads::Type GetDesiredThread()
 	{
-		return CPrio_FMeshDrawCommandSortAndMergeTask.Get();
+		return CPrio_FMeshDrawCommandPassSetupTask.Get();
 	}
 
 	static ESubsequentsMode::Type GetSubsequentsMode() 
@@ -2973,49 +3212,102 @@ public:
 
 	void AnyThreadTask()
 	{
-		if (Context.PassType != EMeshPass::Num)
+		// Mobile base pass is a special case, as final lists is created from two mesh passes based on CSM visibility.
+		const bool bMobileBasePass = Context.ShadingPath == EShadingPath::Mobile && Context.PassType == EMeshPass::BasePass;
+
+		if (bMobileBasePass)
 		{
-			ApplyViewOverridesToMeshDrawCommands(
+			MergeMobileBasePassMeshDrawCommands(
+				Context.View->MobileCSMVisibilityInfo,
+				Context.PrimitiveBounds->Num(),
+				Context.MeshDrawCommands,
+				Context.MobileBasePassCSMMeshDrawCommands
+			);
+
+			GenerateMobileBasePassDynamicMeshDrawCommands(
+				*Context.View,
 				Context.ShadingPath,
 				Context.PassType,
-				Context.bReverseCulling,
-				Context.bRenderSceneTwoSided,
-				Context.BasePassDepthStencilAccess,
-				Context.VisibleMeshDrawCommands,
-				Context.MeshDrawCommandStorage,
-				Context.TempVisibleMeshDrawCommands
+				Context.MeshPassProcessor,
+				Context.MobileBasePassCSMMeshPassProcessor,
+				*Context.DynamicMeshElements,
+				Context.DynamicMeshCommandBuildRequests,
+				Context.MeshDrawCommands,
+				Context.MeshDrawCommandStorage
 			);
 		}
-
-		if (Context.TranslucencyPass != ETranslucencyPass::TPT_MAX)
+		else
 		{
-			UpdateTranslucentMeshSortKeys(
-				Context.TranslucentSortPolicy,
-				Context.TranslucentSortAxis,
-				Context.ViewOrigin,
-				Context.ViewMatrix,
-				*Context.PrimitiveBounds,
-				Context.TranslucencyPass,
-				Context.VisibleMeshDrawCommands
+			GenerateDynamicMeshDrawCommands(
+				*Context.View,
+				Context.ShadingPath,
+				Context.PassType,
+				Context.MeshPassProcessor,
+				*Context.DynamicMeshElements,
+				Context.DynamicMeshCommandBuildRequests,
+				Context.MeshDrawCommands,
+				Context.MeshDrawCommandStorage
 			);
 		}
 
-		{
-			QUICK_SCOPE_CYCLE_COUNTER(STAT_SortVisibleMeshDrawCommands);
-			Context.VisibleMeshDrawCommands.Sort(FCompareFMeshDrawCommands());
-		}
+		// Check if we didn't generate more mesh draw commands, than we assumed in FParallelMeshDrawCommandPass::MaxNumDraws.
+		check(Context.MeshDrawCommands.Num() <= Context.TempVisibleMeshDrawCommands.Max());
 
-		if (Context.bUseGPUScene)
+		if (Context.MeshDrawCommands.Num() > 0)
 		{
-			BuildMeshDrawCommandPrimitiveIdBuffer(
-				Context.VisibleMeshDrawCommands,
-				Context.MeshDrawCommandStorage,
-				Context.PrimitiveIdBuffer,
-				Context.TempVisibleMeshDrawCommands,
-				Context.MaxInstances,
-				Context.VisibleMeshDrawCommandsNum,
-				Context.NewPassVisibleMeshDrawCommandsNum
-			);
+			if (Context.PassType != EMeshPass::Num)
+			{
+				ApplyViewOverridesToMeshDrawCommands(
+					Context.ShadingPath,
+					Context.PassType,
+					Context.bReverseCulling,
+					Context.bRenderSceneTwoSided,
+					Context.BasePassDepthStencilAccess,
+					Context.MeshDrawCommands,
+					Context.MeshDrawCommandStorage,
+					Context.TempVisibleMeshDrawCommands
+				);
+			}
+
+			// Update sort keys.
+			if (bMobileBasePass)
+			{
+				UpdateMobileBasePassMeshSortKeys(
+					Context.ViewOrigin,
+					*Context.PrimitiveBounds,
+					Context.MeshDrawCommands
+					);
+			}
+			else if (Context.TranslucencyPass != ETranslucencyPass::TPT_MAX)
+			{
+				UpdateTranslucentMeshSortKeys(
+					Context.TranslucentSortPolicy,
+					Context.TranslucentSortAxis,
+					Context.ViewOrigin,
+					Context.ViewMatrix,
+					*Context.PrimitiveBounds,
+					Context.TranslucencyPass,
+					Context.MeshDrawCommands
+				);
+			}
+
+			{
+				QUICK_SCOPE_CYCLE_COUNTER(STAT_SortVisibleMeshDrawCommands);
+				Context.MeshDrawCommands.Sort(FCompareFMeshDrawCommands());
+			}
+
+			if (Context.bUseGPUScene)
+			{
+				BuildMeshDrawCommandPrimitiveIdBuffer(
+					Context.MeshDrawCommands,
+					Context.MeshDrawCommandStorage,
+					Context.PrimitiveIdBuffer,
+					Context.TempVisibleMeshDrawCommands,
+					Context.MaxInstances,
+					Context.VisibleMeshDrawCommandsNum,
+					Context.NewPassVisibleMeshDrawCommandsNum
+				);
+			}
 		}
 	}
 
@@ -3026,7 +3318,7 @@ public:
 
 
 private:
-	FMeshDrawCommandSortAndMergeTaskContext& Context;
+	FMeshDrawCommandPassSetupTaskContext& Context;
 };
 
 void SortPassMeshDrawCommands(
@@ -3074,20 +3366,34 @@ void SortPassMeshDrawCommands(
 	}
 }
 
-void FParallelMeshDrawCommandPass::DispatchSortAndMerge(
-	const FViewInfo& View, 
-	EShadingPath ShadingPath,
-	const TArray<FPrimitiveBounds>& PrimitiveBounds, 
-	EMeshPass::Type PassType, 
+void FParallelMeshDrawCommandPass::DispatchPassSetup(
+	FScene* Scene,
+	const FViewInfo& View,
+	EMeshPass::Type PassType,
 	FExclusiveDepthStencil::Type BasePassDepthStencilAccess,
-	FMeshCommandOneFrameArray& InOutVisibleMeshDrawCommands)
+	FMeshPassProcessor* MeshPassProcessor,
+	const TArray<FMeshBatchAndRelevance, SceneRenderingAllocator>& DynamicMeshElements,
+	int32 MaxNumDynamicMeshesForThisPass,
+	TArray<const FStaticMesh*, SceneRenderingAllocator>& InOutDynamicMeshCommandBuildRequests,
+	FMeshCommandOneFrameArray& InOutMeshDrawCommands,
+	FMeshPassProcessor* MobileBasePassCSMMeshPassProcessor,
+	FMeshCommandOneFrameArray* InOutMobileBasePassCSMMeshDrawCommands
+)
 {
-	check(TaskEventRefs.Num() == 0);
-	check(InOutVisibleMeshDrawCommands.Num() > 0);
+	check(TaskEventRefs.Num() == 0 && MeshPassProcessor != nullptr);
+
+	MaxNumDraws = InOutMeshDrawCommands.Num() + InOutDynamicMeshCommandBuildRequests.Num() + MaxNumDynamicMeshesForThisPass;
+	if (MaxNumDraws <= 0)
+	{
+		return;
+	}
 
-	MaxDrawNum = InOutVisibleMeshDrawCommands.Num();
+	TaskContext.MeshPassProcessor = MeshPassProcessor;
+	TaskContext.MobileBasePassCSMMeshPassProcessor = MobileBasePassCSMMeshPassProcessor;
+	TaskContext.DynamicMeshElements = &DynamicMeshElements;
 
-	TaskContext.ShadingPath = ShadingPath;
+	TaskContext.View = &View;
+	TaskContext.ShadingPath = Scene->GetShadingPath();
 	TaskContext.PassType = PassType;
 	TaskContext.bUseGPUScene = UseGPUScene(GMaxRHIShaderPlatform, View.GetFeatureLevel());
 	TaskContext.bDynamicInstancing = TaskContext.bUseGPUScene && IsDynamicInstancingEnabled();
@@ -3101,7 +3407,7 @@ void FParallelMeshDrawCommandPass::DispatchSortAndMerge(
 	TaskContext.TranslucentSortAxis = View.TranslucentSortAxis;
 	TaskContext.ViewOrigin = View.ViewMatrices.GetViewOrigin();
 	TaskContext.ViewMatrix = View.ViewMatrices.GetViewMatrix();
-	TaskContext.PrimitiveBounds = &PrimitiveBounds;
+	TaskContext.PrimitiveBounds = &Scene->PrimitiveBounds;
 
 	switch (PassType)
 	{
@@ -3111,43 +3417,65 @@ void FParallelMeshDrawCommandPass::DispatchSortAndMerge(
 		case EMeshPass::MobileInverseOpacity: TaskContext.TranslucencyPass = ETranslucencyPass::TPT_StandardTranslucency; break;
 	}
 
-	FMemory::Memswap(&TaskContext.VisibleMeshDrawCommands, &InOutVisibleMeshDrawCommands, sizeof(InOutVisibleMeshDrawCommands));
+	FMemory::Memswap(&TaskContext.MeshDrawCommands, &InOutMeshDrawCommands, sizeof(InOutMeshDrawCommands));
+	FMemory::Memswap(&TaskContext.DynamicMeshCommandBuildRequests, &InOutDynamicMeshCommandBuildRequests, sizeof(InOutDynamicMeshCommandBuildRequests));
 
-	// Preallocate context memory and resources on rendering thread.
-	TaskContext.TempVisibleMeshDrawCommands.Empty(MaxDrawNum);
+	if (TaskContext.ShadingPath == EShadingPath::Mobile && TaskContext.PassType == EMeshPass::BasePass)
+	{
+		FMemory::Memswap(&TaskContext.MobileBasePassCSMMeshDrawCommands, InOutMobileBasePassCSMMeshDrawCommands, sizeof(*InOutMobileBasePassCSMMeshDrawCommands));
+	}
+	else
+	{
+		check(MobileBasePassCSMMeshPassProcessor == nullptr && InOutMobileBasePassCSMMeshDrawCommands == nullptr);
+	}
+
+	// Preallocate resources on rendering thread based on MaxNumDraws.
 	if (TaskContext.bDynamicInstancing)
 	{
-		TaskContext.PrimitiveIdBuffer = FGlobalDynamicVertexBuffer::Get().Allocate(MaxDrawNum * sizeof(int32));
+		TaskContext.PrimitiveIdBuffer = FGlobalDynamicVertexBuffer::Get().Allocate(MaxNumDraws * sizeof(int32));
 	}
+	TaskContext.MeshDrawCommands.Reserve(MaxNumDraws);
+	TaskContext.TempVisibleMeshDrawCommands.Reserve(MaxNumDraws);
 
-	const bool bExecuteInParallel = FApp::ShouldUseThreadingForPerformance() && CVarMeshDrawCommandsParallelProcess.GetValueOnRenderThread() > 0;
+	const bool bExecuteInParallel = FApp::ShouldUseThreadingForPerformance() && CVarMeshDrawCommandsParallelPassSetup.GetValueOnRenderThread() > 0;
 
 	if (bExecuteInParallel)
 	{
-		TaskEventRefs.Add(TGraphTask<FMeshDrawCommandSortAndMergeTask>::CreateTask(
+		TaskEventRefs.Add(TGraphTask<FMeshDrawCommandPassSetupTask>::CreateTask(
 			nullptr, ENamedThreads::GetRenderThread()).ConstructAndDispatchWhenReady(TaskContext));
 	}
 	else
 	{
-		FMeshDrawCommandSortAndMergeTask Task(TaskContext);
+		FMeshDrawCommandPassSetupTask Task(TaskContext);
 		Task.AnyThreadTask();
 	}
 }
 
-void FParallelMeshDrawCommandPass::Empty()
+void FParallelMeshDrawCommandPass::WaitForMeshPassSetupTask() const
 {
-	// Need to wait in case if someone dispatched sort and draw merge task, but didn't draw it.
 	if (TaskEventRefs.Num() > 0)
 	{
-		FTaskGraphInterface::Get().WaitUntilTasksComplete(TaskEventRefs);
+		// Need to wait on GetRenderThread_Local, as mesh pass setup task can wait on rendering thread inside InitResourceFromPossiblyParallelRendering().
+		FTaskGraphInterface::Get().WaitUntilTasksComplete(TaskEventRefs, ENamedThreads::GetRenderThread_Local());
 	}
+}
+
+void FParallelMeshDrawCommandPass::Empty()
+{
+	// Need to wait in case if someone dispatched sort and draw merge task, but didn't draw it.
+	WaitForMeshPassSetupTask();
+	TaskEventRefs.Empty();
 
 	DumpInstancingStats();
 
-	MaxDrawNum = 0;
+	MaxNumDraws = 0;
+	TaskContext.MeshPassProcessor = nullptr;
+	TaskContext.MobileBasePassCSMMeshPassProcessor = nullptr;
+	TaskContext.DynamicMeshElements = nullptr;
 	PassNameForStats.Empty();
-	TaskEventRefs.Empty();
-	TaskContext.VisibleMeshDrawCommands.Empty();
+	TaskContext.MeshDrawCommands.Empty();
+	TaskContext.MobileBasePassCSMMeshDrawCommands.Empty();
+	TaskContext.DynamicMeshCommandBuildRequests.Empty();
 	TaskContext.TempVisibleMeshDrawCommands.Empty();
 	TaskContext.PrimitiveIdBuffer = FGlobalDynamicVertexBuffer::FAllocation();
 }
@@ -3213,7 +3541,7 @@ public:
 
 void FParallelMeshDrawCommandPass::DispatchDraw(FParallelCommandListSet* ParallelCommandListSet, FRHICommandList& RHICmdList) const
 {
-	if (!HasAnyDraw())
+	if (MaxNumDraws <= 0)
 	{
 		return;
 	}
@@ -3232,33 +3560,29 @@ void FParallelMeshDrawCommandPass::DispatchDraw(FParallelCommandListSet* Paralle
 			Prereqs.Append(*ParallelCommandListSet->GetPrereqs());
 		}
 
-		// Distribute work evenly to the available task graph workers based on MaxDrawNum.
+		// Distribute work evenly to the available task graph workers based on NumEstimatedDraws.
 		// Every task will then adjust it's working range based on FVisibleMeshDrawCommandProcessTask results.
 		const int32 NumThreads = FMath::Min<int32>(FTaskGraphInterface::Get().GetNumWorkerThreads(), ParallelCommandListSet->Width);
-		const int32 NumTasks = FMath::Min<int32>(NumThreads, FMath::DivideAndRoundUp(MaxDrawNum, ParallelCommandListSet->MinDrawsPerCommandList));
-		const int32 NumDrawsPerTask = FMath::DivideAndRoundUp(MaxDrawNum, NumTasks);
+		const int32 NumTasks = FMath::Min<int32>(NumThreads, FMath::DivideAndRoundUp(MaxNumDraws, ParallelCommandListSet->MinDrawsPerCommandList));
+		const int32 NumDrawsPerTask = FMath::DivideAndRoundUp(MaxNumDraws, NumTasks);
 
 		for (int32 TaskIndex = 0; TaskIndex < NumTasks; TaskIndex++)
 		{
 			const int32 StartIndex = TaskIndex * NumDrawsPerTask;
-			const int32 NumDraws = FMath::Min(NumDrawsPerTask, MaxDrawNum - StartIndex);
+			const int32 NumDraws = FMath::Min(NumDrawsPerTask, MaxNumDraws - StartIndex);
 			checkSlow(NumDraws > 0);
 
 			FRHICommandList* CmdList = ParallelCommandListSet->NewParallelCommandList();
 
 			FGraphEventRef AnyThreadCompletionEvent = TGraphTask<FDrawVisibleMeshCommandsAnyThreadTask>::CreateTask(&Prereqs, RenderThread)
-				.ConstructAndDispatchWhenReady(*CmdList, TaskContext.VisibleMeshDrawCommands, PrimitiveIdsBuffer, BasePrimitiveIdsOffset, TaskContext.bDynamicInstancing, TaskIndex, NumTasks);
+				.ConstructAndDispatchWhenReady(*CmdList, TaskContext.MeshDrawCommands, PrimitiveIdsBuffer, BasePrimitiveIdsOffset, TaskContext.bDynamicInstancing, TaskIndex, NumTasks);
 			ParallelCommandListSet->AddParallelCommandList(CmdList, AnyThreadCompletionEvent, NumDraws);
 		}
 	}
 	else
 	{
-		if (TaskEventRefs.Num() > 0)
-		{
-			FTaskGraphInterface::Get().WaitUntilTasksComplete(TaskEventRefs, ENamedThreads::GetRenderThread());
-		}
-
-		SubmitMeshDrawCommandsRange(TaskContext.VisibleMeshDrawCommands, PrimitiveIdsBuffer, BasePrimitiveIdsOffset, TaskContext.bDynamicInstancing, 0, TaskContext.VisibleMeshDrawCommands.Num(), RHICmdList);
+		WaitForMeshPassSetupTask();
+		SubmitMeshDrawCommandsRange(TaskContext.MeshDrawCommands, PrimitiveIdsBuffer, BasePrimitiveIdsOffset, TaskContext.bDynamicInstancing, 0, TaskContext.MeshDrawCommands.Num(), RHICmdList);
 	}
 }
 
@@ -3278,9 +3602,9 @@ void FParallelMeshDrawCommandPass::SetDumpInstancingStats(const FString& InPassN
 	PassNameForStats = InPassNameForStats;
 }
 
-void FSceneRenderer::SortAndMergeMeshDrawCommands(FViewInfo& View, FExclusiveDepthStencil::Type BasePassDepthStencilAccess, FMeshDrawCommandsPerPass& VisibleCommands)
+void FSceneRenderer::SetupMeshPass(FViewInfo& View, FExclusiveDepthStencil::Type BasePassDepthStencilAccess, FViewCommands& ViewCommands)
 {
-	SCOPE_CYCLE_COUNTER(STAT_SortAndMergeMeshDrawCommands);
+	SCOPE_CYCLE_COUNTER(STAT_SetupMeshPass);
 
 	const EShadingPath ShadingPath = Scene->GetShadingPath();
 
@@ -3296,17 +3620,26 @@ void FSceneRenderer::SortAndMergeMeshDrawCommands(FViewInfo& View, FExclusiveDep
 				continue;
 			}
 
-			if (VisibleCommands[PassIndex].Num() > 0)
-			{
-				FParallelMeshDrawCommandPass& Pass = View.ParallelMeshDrawCommandPasses[PassIndex];
+			PassProcessorCreateFunction CreateFunction = FPassProcessorManager::GetCreateFunction(ShadingPath, PassType);
+			FMeshPassProcessor* MeshPassProcessor = CreateFunction(Scene, &View, nullptr);
 
-				if (ShouldDumpMeshDrawCommandInstancingStats())
-				{
-					Pass.SetDumpInstancingStats(GetMeshPassName(PassType));
-				}
+			FParallelMeshDrawCommandPass& Pass = View.ParallelMeshDrawCommandPasses[PassIndex];
 
-				Pass.DispatchSortAndMerge(View, Scene->GetShadingPath(), Scene->PrimitiveBounds, PassType, BasePassDepthStencilAccess, VisibleCommands[PassIndex]);
+			if (ShouldDumpMeshDrawCommandInstancingStats())
+			{
+				Pass.SetDumpInstancingStats(GetMeshPassName(PassType));
 			}
+
+			Pass.DispatchPassSetup(
+				Scene,
+				View,
+				PassType,
+				BasePassDepthStencilAccess,
+				MeshPassProcessor,
+				View.DynamicMeshElements,
+				View.MaxNumVisibleDynamicMeshes[PassType],
+				ViewCommands.DynamicMeshCommandBuildRequests[PassType],
+				ViewCommands.MeshCommands[PassIndex]);
 		}
 	}
 }
@@ -3772,8 +4105,8 @@ static void RenderViewFamily_RenderThread(FRHICommandListImmediate& RHICmdList,
 		}
 		else
 		{
-		FSceneRenderer::WaitForTasksClearSnapshotsAndDeleteSceneRenderer(RHICmdList, SceneRenderer);
-	}
+			FSceneRenderer::WaitForTasksClearSnapshotsAndDeleteSceneRenderer(RHICmdList, SceneRenderer);
+		}
 	}
 
 #if STATS
@@ -4404,54 +4737,4 @@ FTextureRHIParamRef FSceneRenderer::GetMultiViewSceneColor(const FSceneRenderTar
 	{
 		return static_cast<FTextureRHIRef>(ViewFamily.RenderTarget->GetRenderTargetTexture());
 	}
-}
-
-void FSceneRenderer::GenerateDynamicMeshDrawCommands(FViewInfo& View, FMeshDrawCommandsPerPass& VisibleCommands)
-{
-	QUICK_SCOPE_CYCLE_COUNTER(STAT_GenerateDynamicMeshDrawCommands);
-
-	EShadingPath ShadingPath = FSceneInterface::GetShadingPath(FeatureLevel);
-
-	for (int32 PassIndex = 0; PassIndex < EMeshPass::Num; PassIndex++)
-	{
-		EMeshPass::Type PassType = (EMeshPass::Type)PassIndex;
-
-		if ((FPassProcessorManager::GetPassFlags(ShadingPath, PassType) & EMeshPassFlags::MainView) != EMeshPassFlags::None)
-		{
-			FDynamicPassMeshDrawListContext DynamicPassMeshDrawListContext
-			(
-				View.DynamicMeshDrawCommandStorage[PassType],
-				VisibleCommands[PassType]
-			);
-
-			PassProcessorCreateFunction CreateFunction = FPassProcessorManager::GetCreateFunction(ShadingPath, PassType);
-			FMeshPassProcessor* PassMeshProcessor = CreateFunction(Scene, &View, &DynamicPassMeshDrawListContext);
-
-			if (PassMeshProcessor != nullptr)
-			{
-				const int32 NumDynamicMeshBatches = View.DynamicMeshElements.Num();
-
-				for (int32 MeshIndex = 0; MeshIndex < NumDynamicMeshBatches; MeshIndex++)
-				{
-					const FMeshBatchAndRelevance& MeshAndRelevance = View.DynamicMeshElements[MeshIndex];
-					check(!MeshAndRelevance.Mesh->bRequiresPerElementVisibility);
-					const uint64 BatchElementMask = ~0ull;
-
-					PassMeshProcessor->AddMeshBatch(*MeshAndRelevance.Mesh, BatchElementMask, MeshAndRelevance.PrimitiveSceneProxy);
-				}
-
-				const int32 NumStaticMeshBatches = View.DynamicMeshCommandBuildRequests[PassType].Num();
-
-				for (int32 MeshIndex = 0; MeshIndex < NumStaticMeshBatches; MeshIndex++)
-				{
-					const FStaticMesh* StaticMeshBatch = View.DynamicMeshCommandBuildRequests[PassType][MeshIndex];
-					const uint64 BatchElementMask = StaticMeshBatch->bRequiresPerElementVisibility ? View.StaticMeshBatchVisibility[StaticMeshBatch->BatchVisibilityId] : ~0ull;
-
-					PassMeshProcessor->AddMeshBatch(*StaticMeshBatch, BatchElementMask, StaticMeshBatch->PrimitiveSceneInfo->Proxy, StaticMeshBatch->Id);
-				}
-
-				PassMeshProcessor->~FMeshPassProcessor();
-			}
-		}
-	}
 }
\ No newline at end of file
diff --git a/Engine/Source/Runtime/Renderer/Private/SceneRendering.h b/Engine/Source/Runtime/Renderer/Private/SceneRendering.h
--- a/Engine/Source/Runtime/Renderer/Private/SceneRendering.h
+++ b/Engine/Source/Runtime/Renderer/Private/SceneRendering.h
@@ -51,6 +51,9 @@ public:
 	/** true if there are any primitives affected by CSM subjects */
 	uint32 bMobileDynamicCSMInUse : 1;
 
+	// true if all draws should be forced to use CSM shaders.
+	uint32 bAlwaysUseCSM : 1;
+
 	/** Visibility lists for static meshes that will use expensive CSM shaders. */
 	FSceneBitArray MobilePrimitiveCSMReceiverVisibilityMap;
 	FSceneBitArray MobileCSMStaticMeshVisibilityMap;
@@ -61,7 +64,7 @@ public:
 	TArray<uint64, SceneRenderingAllocator> MobileNonCSMStaticBatchVisibility;
 
 	/** Initialization constructor. */
-	FMobileCSMVisibilityInfo() : bMobileDynamicCSMInUse(false)
+	FMobileCSMVisibilityInfo() : bMobileDynamicCSMInUse(false), bAlwaysUseCSM(false)
 	{}
 };
 
@@ -879,23 +882,32 @@ struct FPreviousViewInfo
 	}
 };
 
-typedef TStaticArray<FMeshCommandOneFrameArray, EMeshPass::Num> FMeshDrawCommandsPerPass;
-typedef TArray<FMeshDrawCommandsPerPass, TInlineAllocator<4>> FMeshDrawCommandsPerPassPerView;
+class FViewCommands
+{
+public:
+	TStaticArray<FMeshCommandOneFrameArray, EMeshPass::Num> MeshCommands;
+	TStaticArray<TArray<const FStaticMesh*, SceneRenderingAllocator>, EMeshPass::Num> DynamicMeshCommandBuildRequests;
+};
+
+typedef TArray<FViewCommands, TInlineAllocator<4>> FViewVisibleCommandsPerView;
 
 /**
- * Parallel mesh draw command processing and drawing context.
+ * Parallel mesh draw command pass setup task context.
  */
-class FMeshDrawCommandSortAndMergeTaskContext
+class FMeshDrawCommandPassSetupTaskContext
 {
 public:
-	FMeshDrawCommandSortAndMergeTaskContext()
-		: ShadingPath(EShadingPath::Num)
+	FMeshDrawCommandPassSetupTaskContext()
+		: View(nullptr)
+		, ShadingPath(EShadingPath::Num)
 		, PassType(EMeshPass::Num)
 		, bUseGPUScene(false)
 		, bDynamicInstancing(false)
 		, bReverseCulling(false)
 		, bRenderSceneTwoSided(false)
 		, BasePassDepthStencilAccess(FExclusiveDepthStencil::DepthNop_StencilNop)
+		, MeshPassProcessor(nullptr)
+		, DynamicMeshElements(nullptr)
 		, PrimitiveBounds(nullptr)
 		, VisibleMeshDrawCommandsNum(0)
 		, NewPassVisibleMeshDrawCommandsNum(0)
@@ -903,6 +915,7 @@ public:
 	{
 	}
 
+	const FViewInfo* View;
 	EShadingPath ShadingPath;
 	EMeshPass::Type PassType;
 	bool bUseGPUScene;
@@ -911,8 +924,16 @@ public:
 	bool bRenderSceneTwoSided;
 	FExclusiveDepthStencil::Type BasePassDepthStencilAccess;
 
+	// Mesh pass processor.
+	FMeshPassProcessor* MeshPassProcessor;
+	FMeshPassProcessor* MobileBasePassCSMMeshPassProcessor;
+	const TArray<FMeshBatchAndRelevance, SceneRenderingAllocator>* DynamicMeshElements;
+
 	// Commands.
-	FMeshCommandOneFrameArray VisibleMeshDrawCommands;
+	FMeshCommandOneFrameArray MeshDrawCommands;
+	FMeshCommandOneFrameArray MobileBasePassCSMMeshDrawCommands;
+	TArray<const FStaticMesh*, SceneRenderingAllocator> DynamicMeshCommandBuildRequests;
+	TArray<const FStaticMesh*, SceneRenderingAllocator> MobileBasePassCSMDynamicMeshCommandBuildRequests;
 	FDynamicMeshDrawCommandStorage MeshDrawCommandStorage;
 
 	// Preallocated resources.
@@ -934,30 +955,35 @@ public:
 };
 
 /**
- * Parallel mesh draw command processing and rendering.
- * DispatchSort - dispatched mesh command processing task.
- * DispatchDraw - dispatched mesh draw task depending on DispatchSort.
+ * Parallel mesh draw command processing and rendering. 
+ * Encapsulates two parallel tasks - mesh command setup task and drawing task.
  */
 class FParallelMeshDrawCommandPass
 {
 public:
 	FParallelMeshDrawCommandPass()
-		: MaxDrawNum(0)
+		: MaxNumDraws(0)
 	{
 	}
 
 	~FParallelMeshDrawCommandPass();
 
 	/**
-	 * Dispatch visible mesh draw command sort and merge task.
+	 * Dispatch visible mesh draw command process task, which prepares this pass for drawing.
+	 * This includes generation of dynamic mesh draw commands, draw sorting and draw merging.
 	 */
-	void DispatchSortAndMerge(
+	void DispatchPassSetup(
+		FScene* Scene,
 		const FViewInfo& View, 
-		EShadingPath ShadingPath,
-		const TArray<FPrimitiveBounds>& PrimitiveBounds, 
 		EMeshPass::Type PassType, 
 		FExclusiveDepthStencil::Type BasePassDepthStencilAccess,
-		FMeshCommandOneFrameArray& InOutVisibleMeshDrawCommands
+		FMeshPassProcessor* MeshPassProcessor,
+		const TArray<FMeshBatchAndRelevance, SceneRenderingAllocator>& DynamicMeshElements,
+		int32 MaxNumDynamicMeshesForThisPass,
+		TArray<const FStaticMesh*, SceneRenderingAllocator>& InOutDynamicMeshCommandBuildRequests,
+		FMeshCommandOneFrameArray& InOutMeshDrawCommands,
+		FMeshPassProcessor* MobileBasePassCSMMeshPassProcessor = nullptr, // Required only for the mobile base pass.
+		FMeshCommandOneFrameArray* InOutMobileBasePassCSMMeshDrawCommands = nullptr // Required only for the mobile base pass.
 	);
 
 	/**
@@ -966,20 +992,21 @@ public:
 	void DispatchDraw(FParallelCommandListSet* ParallelCommandListSet, FRHICommandList& RHICmdList) const;
 
 	void Empty();
-
 	void SetDumpInstancingStats(const FString& InPassName);
-
-	bool HasAnyDraw() const { return MaxDrawNum > 0; }
+	bool HasAnyDraw() const { return MaxNumDraws > 0; }
 
 
 private:
-	FMeshDrawCommandSortAndMergeTaskContext TaskContext;
+	FMeshDrawCommandPassSetupTaskContext TaskContext;
 	FGraphEventArray TaskEventRefs;
-	int32 MaxDrawNum;
 	FString PassNameForStats;
 
+	// Maximum number of draws for this pass. Used to prealocate resources on rendering thread. 
+	// Has a guarantee that if there won't be any draws, then MaxNumDraws = 0;
+	int32 MaxNumDraws;
+
 	void DumpInstancingStats() const;
-	void DispatchSortInternal(ERHIFeatureLevel::Type FeatureLevel, FMeshCommandOneFrameArray& InOutVisibleMeshDrawCommands);
+	void WaitForMeshPassSetupTask() const;
 };
 
 /** A FSceneView with additional state used by the scene renderer. */
@@ -1054,7 +1081,13 @@ public:
 	TArray<uint64,SceneRenderingAllocator> StaticMeshBatchVisibility;
 
 	/** The dynamic primitives visible in this view. */
-	TArray<const FPrimitiveSceneInfo*,SceneRenderingAllocator> VisibleDynamicPrimitives;
+	TArray<FPrimitiveSceneInfo*,SceneRenderingAllocator> VisibleDynamicPrimitives;
+
+	/** Marks used mesh passes by visible dynamic meshes. */
+	FMeshPassMask VisibleDynamicMeshesPassMask;
+
+	/** Max number of dynamic meshes per mesh pass. */
+	int32 MaxNumVisibleDynamicMeshes[EMeshPass::Num];
 
 	/** The dynamic editor primitives visible in this view. */
 	TArray<const FPrimitiveSceneInfo*,SceneRenderingAllocator> VisibleEditorPrimitives;
@@ -1119,15 +1152,13 @@ public:
 	/** Tracks dynamic primitive data for upload to GPU Scene, when enabled. */
 	TArray<FPrimitiveUniformShaderParameters> DynamicPrimitiveShaderData;
 
-	TStaticArray<FDynamicMeshDrawCommandStorage, EMeshPass::Num> DynamicMeshDrawCommandStorage;
 	FRWBufferStructured OneFramePrimitiveShaderDataBuffer;
 	FReadBuffer OneFramePrimitiveIdBufferEmulation;
 
 	TStaticArray<FParallelMeshDrawCommandPass, EMeshPass::Num> ParallelMeshDrawCommandPasses;
 	
 	FMeshCommandOneFrameArray RaytraycingVisibleMeshDrawCommands;
-
-	TStaticArray<TArray<const FStaticMesh*, SceneRenderingAllocator>, EMeshPass::Num> DynamicMeshCommandBuildRequests;
+	FDynamicMeshDrawCommandStorage RaytraycingDynamicMeshDrawCommandStorage;
 
 	// Used by mobile renderer to determine whether static meshes will be rendered with CSM shaders or not.
 	FMobileCSMVisibilityInfo MobileCSMVisibilityInfo;
@@ -1689,10 +1720,7 @@ protected:
 
 	void InitDynamicShadows(FRHICommandListImmediate& RHICmdList);
 
-	/** Converts each FMeshBatch into a set of FMeshDrawCommands for each relevant mesh pass. */
-	void GenerateDynamicMeshDrawCommands(FViewInfo& View, FMeshDrawCommandsPerPass& VisibleCommands);
-
-	void SortAndMergeMeshDrawCommands(FViewInfo& View, FExclusiveDepthStencil::Type BasePassDepthStencilAccess, FMeshDrawCommandsPerPass& VisibleCommands);
+	void SetupMeshPass(FViewInfo& View, FExclusiveDepthStencil::Type BasePassDepthStencilAccess, FViewCommands& ViewCommands);
 
 	bool RenderShadowProjections(FRHICommandListImmediate& RHICmdList, const FLightSceneInfo* LightSceneInfo, IPooledRenderTarget* ScreenShadowMaskTexture, bool bProjectingForForwardShading, bool bMobileModulatedProjections);
 
@@ -1788,7 +1816,7 @@ protected:
 	void PreVisibilityFrameSetup(FRHICommandListImmediate& RHICmdList);
 
 	/** Computes which primitives are visible and relevant for each view. */
-	void ComputeViewVisibility(FRHICommandListImmediate& RHICmdList, FExclusiveDepthStencil::Type BasePassDepthStencilAccess, FMeshDrawCommandsPerPassPerView& VisibleCommandsPerView);
+	void ComputeViewVisibility(FRHICommandListImmediate& RHICmdList, FExclusiveDepthStencil::Type BasePassDepthStencilAccess, FViewVisibleCommandsPerView& ViewCommandsPerView);
 
 	/** Performs once per frame setup after to visibility determination. */
 	void PostVisibilityFrameSetup(FILCUpdatePrimTaskData& OutILCTaskData);
@@ -1863,7 +1891,7 @@ public:
 
 protected:
 	/** Finds the visible dynamic shadows for each view. */
-	void InitDynamicShadows(FRHICommandListImmediate& RHICmdList, FMeshDrawCommandsPerPassPerView& VisibleCommandsPerView);
+	void InitDynamicShadows(FRHICommandListImmediate& RHICmdList);
 
 	/** Build visibility lists on CSM receivers and non-csm receivers. */
 	void BuildCSMVisibilityState(FLightSceneInfo* LightSceneInfo);
@@ -1912,7 +1940,7 @@ protected:
 	/** Will update the view custom data. */
 	void PostInitViewCustomData();
 
-	void SortMobileBasePassAfterShadowInit(FExclusiveDepthStencil::Type BasePassDepthStencilAccess, FMeshDrawCommandsPerPassPerView& VisibleCommandsPerView);
+	void SetupMobileBasePassAfterShadowInit(FExclusiveDepthStencil::Type BasePassDepthStencilAccess, FViewVisibleCommandsPerView& ViewCommandsPerView);
 
 	void UpdateOpaqueBasePassUniformBuffer(FRHICommandListImmediate& RHICmdList, const FViewInfo& View);
 	void UpdateTranslucentBasePassUniformBuffer(FRHICommandListImmediate& RHICmdList, const FViewInfo& View);
diff --git a/Engine/Source/Runtime/Renderer/Private/SceneVisibility.cpp b/Engine/Source/Runtime/Renderer/Private/SceneVisibility.cpp
--- a/Engine/Source/Runtime/Renderer/Private/SceneVisibility.cpp
+++ b/Engine/Source/Runtime/Renderer/Private/SceneVisibility.cpp
@@ -1745,7 +1745,7 @@ struct FRelevancePacket
 	FRHICommandListImmediate& RHICmdList;
 	const FScene* Scene;
 	const FViewInfo& View;
-	const FMeshDrawCommandsPerPass& VisibleCommands;
+	const FViewCommands& ViewCommands;
 	const uint8 ViewBit;
 	const FMarkRelevantStaticMeshesForViewData& ViewData;
 	FPrimitiveViewMasks& OutHasDynamicMeshElementsMasks;
@@ -1757,6 +1757,7 @@ struct FRelevancePacket
 	FRelevancePrimSet<int32> NotDrawRelevant;
 	FRelevancePrimSet<int32> TranslucentSelfShadowPrimitives;
 	FRelevancePrimSet<FPrimitiveSceneInfo*> VisibleDynamicPrimitives;
+	FMeshPassMask VisibleDynamicMeshesPassMask;
 	FRelevancePrimSet<FTranslucentPrimSet::FTranslucentSortedPrim, ETranslucencyPass::TPT_MAX> TranslucencyPrims;
 	// belongs to TranslucencyPrims
 	FTranslucenyPrimCount TranslucencyPrimCount;
@@ -1817,7 +1818,7 @@ struct FRelevancePacket
 		FRHICommandListImmediate& InRHICmdList,
 		const FScene* InScene, 
 		const FViewInfo& InView, 
-		const FMeshDrawCommandsPerPass& InVisibleCommands,
+		const FViewCommands& InViewCommands,
 		uint8 InViewBit,
 		const FMarkRelevantStaticMeshesForViewData& InViewData,
 		FPrimitiveViewMasks& InOutHasDynamicMeshElementsMasks,
@@ -1831,7 +1832,7 @@ struct FRelevancePacket
 		, RHICmdList(InRHICmdList)
 		, Scene(InScene)
 		, View(InView)
-		, VisibleCommands(InVisibleCommands)
+		, ViewCommands(InViewCommands)
 		, ViewBit(InViewBit)
 		, ViewData(InViewData)
 		, OutHasDynamicMeshElementsMasks(InOutHasDynamicMeshElementsMasks)
@@ -1853,12 +1854,99 @@ struct FRelevancePacket
 		MarkRelevant();
 	}
 
+	void UpdateDynamicPrimMask(EShadingPath ShadingPath, bool bAddLightmapDensityCommands, const FPrimitiveViewRelevance& ViewRelevance)
+	{
+		// Compute number mesh pass mask for this dynamic primitive.
+		if (ViewRelevance.bDrawRelevance && (ViewRelevance.bRenderInMainPass || ViewRelevance.bRenderCustomDepth))
+		{
+			VisibleDynamicMeshesPassMask.Set(EMeshPass::DepthPass);
+			VisibleDynamicMeshesPassMask.Set(EMeshPass::BasePass);
+
+			if (ShadingPath == EShadingPath::Mobile)
+			{
+				VisibleDynamicMeshesPassMask.Set(EMeshPass::MobileBasePassCSM);
+			}
+
+			if (ViewRelevance.bRenderCustomDepth)
+			{
+				VisibleDynamicMeshesPassMask.Set(EMeshPass::CustomDepth);
+			}
+
+			if (bAddLightmapDensityCommands)
+			{
+				VisibleDynamicMeshesPassMask.Set(EMeshPass::LightmapDensity);
+			}
+#if !(UE_BUILD_SHIPPING || UE_BUILD_TEST)
+			else if (View.Family->UseDebugViewPS())
+			{
+				VisibleDynamicMeshesPassMask.Set(EMeshPass::DebugViewMode);
+			}
+#endif
+
+#if WITH_EDITOR
+			if (View.bAllowTranslucentPrimitivesInHitProxy)
+			{
+				VisibleDynamicMeshesPassMask.Set(EMeshPass::HitProxy);
+			}
+			else
+			{
+				VisibleDynamicMeshesPassMask.Set(EMeshPass::HitProxyOpaqueOnly);
+			}
+#endif
+
+			if (ViewRelevance.bVelocityRelevance)
+			{
+				VisibleDynamicMeshesPassMask.Set(EMeshPass::Velocity);
+			}
+		}
+
+		if (ViewRelevance.HasTranslucency()
+			&& !ViewRelevance.bEditorPrimitiveRelevance
+			&& ViewRelevance.bRenderInMainPass)
+		{
+			if (View.Family->AllowTranslucencyAfterDOF())
+			{
+				if (ViewRelevance.bNormalTranslucencyRelevance)
+				{
+					VisibleDynamicMeshesPassMask.Set(EMeshPass::TranslucencyStandard);
+				}
+
+				if (ViewRelevance.bSeparateTranslucencyRelevance)
+				{
+					VisibleDynamicMeshesPassMask.Set(EMeshPass::TranslucencyAfterDOF);
+				}
+			}
+			else
+			{
+				VisibleDynamicMeshesPassMask.Set(EMeshPass::TranslucencyAll);
+			}
+
+			if (ViewRelevance.bDistortionRelevance)
+			{
+				VisibleDynamicMeshesPassMask.Set(EMeshPass::Distortion);
+			}
+
+			if (ShadingPath == EShadingPath::Mobile && View.bIsSceneCapture)
+			{
+				VisibleDynamicMeshesPassMask.Set(EMeshPass::MobileInverseOpacity);
+			}
+		}
+#if WITH_EDITOR
+		if (ViewRelevance.bDrawRelevance)
+		{
+			VisibleDynamicMeshesPassMask.Set(EMeshPass::EditorSelection);
+		}
+#endif
+	}
+
 	void ComputeRelevance()
 	{
 		CombinedShadingModelMask = 0;
 		bUsesGlobalDistanceField = false;
 		bUsesLightingChannels = false;
 		bTranslucentSurfaceLighting = false;
+		const EShadingPath ShadingPath = Scene->GetShadingPath();
+		const bool bAddLightmapDensityCommands = View.Family->EngineShowFlags.LightMapDensity && AllowDebugViewmodes();
 
 		SCOPE_CYCLE_COUNTER(STAT_ComputeViewRelevance);
 		for (int32 Index = 0; Index < Input.NumPrims; Index++)
@@ -1914,6 +2002,7 @@ struct FRelevancePacket
 				// Keep track of visible dynamic primitives.
 				VisibleDynamicPrimitives.AddPrim(PrimitiveSceneInfo);
 				OutHasDynamicMeshElementsMasks[BitIndex] |= ViewBit;
+				UpdateDynamicPrimMask(ShadingPath, bAddLightmapDensityCommands, ViewRelevance);
 			}
 
 			if (ViewRelevance.bUseCustomViewData)
@@ -2236,7 +2325,7 @@ struct FRelevancePacket
 	void RenderThreadFinalize()
 	{
 		FViewInfo& WriteView = const_cast<FViewInfo&>(View);
-		FMeshDrawCommandsPerPass& WriteVisibleCommands = const_cast<FMeshDrawCommandsPerPass&>(VisibleCommands);
+		FViewCommands& WriteViewCommands = const_cast<FViewCommands&>(ViewCommands);
 		
 		for (int32 Index = 0; Index < NotDrawRelevant.NumPrims; Index++)
 		{
@@ -2250,6 +2339,7 @@ struct FRelevancePacket
 		WriteView.bUsesSceneDepth |= bUsesSceneDepth;
 		VisibleEditorPrimitives.AppendTo(WriteView.VisibleEditorPrimitives);
 		VisibleDynamicPrimitives.AppendTo(WriteView.VisibleDynamicPrimitives);
+		VisibleDynamicMeshesPassMask.AppendTo(WriteView.VisibleDynamicMeshesPassMask);
 		WriteView.TranslucentPrimSet.AppendScenePrimitives(TranslucencyPrims.Prims, TranslucencyPrims.NumPrims, TranslucencyPrimCount);
 		DistortionPrimSet.AppendTo(WriteView.DistortionPrimSet);
 		MeshDecalPrimSet.AppendTo(WriteView.MeshDecalPrimSet.Prims);
@@ -2274,8 +2364,8 @@ struct FRelevancePacket
 
 		for (int32 PassIndex = 0; PassIndex < EMeshPass::Num; PassIndex++)
 		{
-			DrawCommandPacket.VisibleCachedDrawCommands[PassIndex].CopyToLinearArray(WriteVisibleCommands[PassIndex]);
-			DrawCommandPacket.DynamicBuildRequests[PassIndex].CopyToLinearArray(WriteView.DynamicMeshCommandBuildRequests[PassIndex]);
+			DrawCommandPacket.VisibleCachedDrawCommands[PassIndex].CopyToLinearArray(WriteViewCommands.MeshCommands[PassIndex]);
+			DrawCommandPacket.DynamicBuildRequests[PassIndex].CopyToLinearArray(WriteViewCommands.DynamicMeshCommandBuildRequests[PassIndex]);
 		}
 
 		// Prepare translucent self shadow uniform buffers.
@@ -2299,7 +2389,7 @@ static void ComputeAndMarkRelevanceForViewParallel(
 	FRHICommandListImmediate& RHICmdList,
 	const FScene* Scene,
 	FViewInfo& View,
-	FMeshDrawCommandsPerPass& VisibleCommands,
+	FViewCommands& ViewCommands,
 	uint8 ViewBit,
 	FPrimitiveViewMasks& OutHasDynamicMeshElementsMasks,
 	FPrimitiveViewMasks& OutHasDynamicEditorMeshElementsMasks,
@@ -2338,7 +2428,7 @@ static void ComputeAndMarkRelevanceForViewParallel(
 				RHICmdList,
 				Scene, 
 				View, 
-				VisibleCommands,
+				ViewCommands,
 				ViewBit,
 				ViewData,
 				OutHasDynamicMeshElementsMasks,
@@ -2364,7 +2454,7 @@ static void ComputeAndMarkRelevanceForViewParallel(
 							RHICmdList,
 							Scene, 
 							View, 
-							VisibleCommands,
+							ViewCommands,
 							ViewBit,
 							ViewData,
 							OutHasDynamicMeshElementsMasks,
@@ -2402,8 +2492,8 @@ static void ComputeAndMarkRelevanceForViewParallel(
 				NumDynamicBuildRequests += Packet->DrawCommandPacket.DynamicBuildRequests[PassIndex].Num();
 			}
 
-			VisibleCommands[PassIndex].Reserve(NumVisibleCachedMeshDrawCommands);
-			View.DynamicMeshCommandBuildRequests[PassIndex].Reserve(NumDynamicBuildRequests);
+			ViewCommands.MeshCommands[PassIndex].Reserve(NumVisibleCachedMeshDrawCommands);
+			ViewCommands.DynamicMeshCommandBuildRequests[PassIndex].Reserve(NumDynamicBuildRequests);
 		}
 
 		for (auto Packet : Packets)
@@ -2598,6 +2688,20 @@ void FSceneRenderer::GatherDynamicMeshElements(
 		}
 	}
 	MeshCollector.ProcessTasks();
+
+	for (int32 ViewIndex = 0; ViewIndex < Views.Num(); ++ViewIndex)
+	{
+		FViewInfo& View = Views[ViewIndex];
+
+		for (int32 PassIndex = 0; PassIndex < EMeshPass::Num; ++PassIndex)
+		{
+			View.MaxNumVisibleDynamicMeshes[PassIndex] = 0;
+			if (View.VisibleDynamicMeshesPassMask.Get((EMeshPass::Type) PassIndex))
+			{
+				View.MaxNumVisibleDynamicMeshes[PassIndex] = View.DynamicMeshElements.Num();
+			}
+		}
+	}
 }
 
 /**
@@ -3152,7 +3256,7 @@ void UpdateReflectionSceneData(FScene* Scene)
 	}
 }
 
-void FSceneRenderer::ComputeViewVisibility(FRHICommandListImmediate& RHICmdList, FExclusiveDepthStencil::Type BasePassDepthStencilAccess, FMeshDrawCommandsPerPassPerView& VisibleCommandsPerView)
+void FSceneRenderer::ComputeViewVisibility(FRHICommandListImmediate& RHICmdList, FExclusiveDepthStencil::Type BasePassDepthStencilAccess, FViewVisibleCommandsPerView& ViewCommandsPerView)
 {
 	SCOPE_CYCLE_COUNTER(STAT_ViewVisibilityTime);
 	SCOPED_NAMED_EVENT(FSceneRenderer_ComputeViewVisibility, FColor::Magenta);
@@ -3191,6 +3295,7 @@ void FSceneRenderer::ComputeViewVisibility(FRHICommandListImmediate& RHICmdList,
 		STAT(NumProcessedPrimitives += NumPrimitives);
 
 		FViewInfo& View = Views[ViewIndex];
+		FViewCommands& ViewCommands = ViewCommandsPerView[ViewIndex];
 		FSceneViewState* ViewState = (FSceneViewState*)View.State;
 
 		// Allocate the view's visibility maps.
@@ -3422,7 +3527,7 @@ void FSceneRenderer::ComputeViewVisibility(FRHICommandListImmediate& RHICmdList,
 		if (!View.IsInstancedStereoPass())
 		{
 			SCOPE_CYCLE_COUNTER(STAT_ViewRelevance);
-			ComputeAndMarkRelevanceForViewParallel(RHICmdList, Scene, View, VisibleCommandsPerView[ViewIndex], ViewBit, HasDynamicMeshElementsMasks, HasDynamicEditorMeshElementsMasks, HasViewCustomDataMasks);
+			ComputeAndMarkRelevanceForViewParallel(RHICmdList, Scene, View, ViewCommands, ViewBit, HasDynamicMeshElementsMasks, HasDynamicEditorMeshElementsMasks, HasViewCustomDataMasks);
 		}
 
 #if !(UE_BUILD_SHIPPING || UE_BUILD_TEST)
@@ -3474,11 +3579,12 @@ void FSceneRenderer::ComputeViewVisibility(FRHICommandListImmediate& RHICmdList,
 	for (int32 ViewIndex = 0; ViewIndex < Views.Num(); ++ViewIndex)
 	{
 		FViewInfo& View = Views[ViewIndex];
+		FViewCommands& ViewCommands = ViewCommandsPerView[ViewIndex];
 		
 		if (View.IsInstancedStereoPass())
 		{
 			SCOPE_CYCLE_COUNTER(STAT_ViewRelevance);
-			ComputeAndMarkRelevanceForViewParallel(RHICmdList, Scene, View, VisibleCommandsPerView[ViewIndex], ViewBit, HasDynamicMeshElementsMasks, HasDynamicEditorMeshElementsMasks, HasViewCustomDataMasks);
+			ComputeAndMarkRelevanceForViewParallel(RHICmdList, Scene, View, ViewCommands, ViewBit, HasDynamicMeshElementsMasks, HasDynamicEditorMeshElementsMasks, HasViewCustomDataMasks);
 		}
 		ViewBit <<= 1;
 	}
@@ -3491,11 +3597,9 @@ void FSceneRenderer::ComputeViewVisibility(FRHICommandListImmediate& RHICmdList,
 		for (int32 ViewIndex = 0; ViewIndex < Views.Num(); ViewIndex++)
 		{
 			FViewInfo& View = Views[ViewIndex];
-			FMeshDrawCommandsPerPass& VisibleCommands = VisibleCommandsPerView[ViewIndex];
-
-			GenerateDynamicMeshDrawCommands(View, VisibleCommands);
+			FViewCommands& ViewCommands = ViewCommandsPerView[ViewIndex];
 
-			SortAndMergeMeshDrawCommands(View, BasePassDepthStencilAccess, VisibleCommands);
+			SetupMeshPass(View, BasePassDepthStencilAccess, ViewCommands);
 		}
 	}
 
@@ -3765,10 +3869,10 @@ bool FDeferredShadingSceneRenderer::InitViews(FRHICommandListImmediate& RHICmdLi
 
 	RHICmdList.ImmediateFlush(EImmediateFlushType::DispatchToRHIThread);
 
-	FMeshDrawCommandsPerPassPerView VisibleCommandsPerView;
-	VisibleCommandsPerView.SetNum(Views.Num());
+	FViewVisibleCommandsPerView ViewCommandsPerView;
+	ViewCommandsPerView.SetNum(Views.Num());
 
-	ComputeViewVisibility(RHICmdList, BasePassDepthStencilAccess, VisibleCommandsPerView);
+	ComputeViewVisibility(RHICmdList, BasePassDepthStencilAccess, ViewCommandsPerView);
 
 	RHICmdList.ImmediateFlush(EImmediateFlushType::DispatchToRHIThread);
 
diff --git a/Engine/Source/Runtime/Renderer/Private/ShadowSetup.cpp b/Engine/Source/Runtime/Renderer/Private/ShadowSetup.cpp
--- a/Engine/Source/Runtime/Renderer/Private/ShadowSetup.cpp
+++ b/Engine/Source/Runtime/Renderer/Private/ShadowSetup.cpp
@@ -1201,47 +1201,31 @@ void FProjectedShadowInfo::SetupMeshDrawCommandsForShadowDepth(FSceneRenderer& R
 {
 	QUICK_SCOPE_CYCLE_COUNTER(STAT_SetupMeshDrawCommandsForShadowDepth);
 
-	FDynamicPassMeshDrawListContext ShadowDepthContext(DynamicMeshDrawCommandStorage, ShadowDepthPassVisibleCommands);
-	FShadowDepthPassMeshProcessor ShadowDepthPassMeshProcessor(
-		Renderer.Scene, 
-		ShadowDepthView, 
-		ShadowDepthView->ViewUniformBuffer, 
+	FShadowDepthPassMeshProcessor* MeshPassProcessor = new(FMemStack::Get()) FShadowDepthPassMeshProcessor(
+		Renderer.Scene,
+		ShadowDepthView,
+		ShadowDepthView->ViewUniformBuffer,
 		PassUniformBuffer,
 		GetShadowDepthType(),
-		&ShadowDepthContext);
+		nullptr);
 
-	const int32 NumDynamicMeshBatches = DynamicSubjectMeshElements.Num();
-
-	for (int32 MeshIndex = 0; MeshIndex < NumDynamicMeshBatches; MeshIndex++)
-	{
-		const FMeshBatchAndRelevance& MeshAndRelevance = DynamicSubjectMeshElements[MeshIndex];
-		check(!MeshAndRelevance.Mesh->bRequiresPerElementVisibility);
-		const uint64 BatchElementMask = ~0ull;
-
-		ShadowDepthPassMeshProcessor.AddMeshBatch(*MeshAndRelevance.Mesh, BatchElementMask, MeshAndRelevance.PrimitiveSceneProxy);
-	}
-
-	const int32 NumStaticMeshBatches = SubjectMeshCommandBuildRequests.Num();
-
-	for (int32 MeshIndex = 0; MeshIndex < NumStaticMeshBatches; MeshIndex++)
+	if (Renderer.ShouldDumpMeshDrawCommandInstancingStats())
 	{
-		const FStaticMesh* StaticMeshBatch = SubjectMeshCommandBuildRequests[MeshIndex];
-		const uint64 BatchElementMask = StaticMeshBatch->bRequiresPerElementVisibility ? ShadowDepthView->StaticMeshBatchVisibility[StaticMeshBatch->BatchVisibilityId] : ~0ull;
-
-		ShadowDepthPassMeshProcessor.AddMeshBatch(*StaticMeshBatch, BatchElementMask, StaticMeshBatch->PrimitiveSceneInfo->Proxy, StaticMeshBatch->Id);
+		FString PassNameForStats;
+		GetShadowTypeNameForDrawEvent(PassNameForStats);
+		ShadowDepthPass.SetDumpInstancingStats(TEXT("ShadowDepth ") + PassNameForStats);
 	}
 
-	if (ShadowDepthPassVisibleCommands.Num() > 0)
-	{
-		if (Renderer.ShouldDumpMeshDrawCommandInstancingStats())
-		{
-			FString PassNameForStats;
-			GetShadowTypeNameForDrawEvent(PassNameForStats);
-			ShadowDepthPass.SetDumpInstancingStats(TEXT("ShadowDepth ") + PassNameForStats);
-		}
-
-		ShadowDepthPass.DispatchSortAndMerge(*ShadowDepthView, Renderer.Scene->GetShadingPath(), Renderer.Scene->PrimitiveBounds, EMeshPass::Num, FExclusiveDepthStencil::DepthNop_StencilNop, ShadowDepthPassVisibleCommands);
-	}
+	ShadowDepthPass.DispatchPassSetup(
+		Renderer.Scene,
+		*ShadowDepthView,
+		EMeshPass::Num,
+		FExclusiveDepthStencil::DepthNop_StencilNop,
+		MeshPassProcessor,
+		DynamicSubjectMeshElements,
+		DynamicSubjectMeshElements.Num(),
+		SubjectMeshCommandBuildRequests,
+		ShadowDepthPassVisibleCommands);
 }
 
 void FProjectedShadowInfo::SetupMeshDrawCommandsForProjectionStenciling(FSceneRenderer& Renderer)
diff --git a/Engine/Source/Runtime/Renderer/Private/ShadowSetupMobile.cpp b/Engine/Source/Runtime/Renderer/Private/ShadowSetupMobile.cpp
--- a/Engine/Source/Runtime/Renderer/Private/ShadowSetupMobile.cpp
+++ b/Engine/Source/Runtime/Renderer/Private/ShadowSetupMobile.cpp
@@ -268,7 +268,7 @@ static void VisualizeMobileDynamicCSMSubjectCapsules(FViewInfo& View, FLightScen
 	}
 }
 /** Finds the visible dynamic shadows for each view. */
-void FMobileSceneRenderer::InitDynamicShadows(FRHICommandListImmediate& RHICmdList, FMeshDrawCommandsPerPassPerView& VisibleCommandsPerView)
+void FMobileSceneRenderer::InitDynamicShadows(FRHICommandListImmediate& RHICmdList)
 {
 	static auto* MyCVarMobileEnableStaticAndCSMShadowReceivers = IConsoleManager::Get().FindTConsoleVariableDataInt(TEXT("r.Mobile.EnableStaticAndCSMShadowReceivers"));
 	const bool bMobileEnableStaticAndCSMShadowReceivers = MyCVarMobileEnableStaticAndCSMShadowReceivers->GetValueOnRenderThread() == 1;
@@ -337,33 +337,10 @@ void FMobileSceneRenderer::InitDynamicShadows(FRHICommandListImmediate& RHICmdLi
 		}
 	}
 
-	// Where we've determined CSM is required we copy the MobileBasePassCSM mesh commands to the BasePass visible list:
-	for (int32 ViewIndex = 0; ViewIndex < Views.Num(); ++ViewIndex)
+	for (auto& View : Views)
 	{
-		FViewInfo& View = Views[ViewIndex];
-		FMeshDrawCommandsPerPass& VisibleCommands = VisibleCommandsPerView[ViewIndex];
-
 		FMobileCSMVisibilityInfo& MobileCSMVisibilityInfo = View.MobileCSMVisibilityInfo;
-		if (MobileCSMVisibilityInfo.bMobileDynamicCSMInUse)
-		{
-			// determine per view CSM visibility
-			FMeshCommandOneFrameArray& MeshCommands = VisibleCommands[EMeshPass::BasePass];
-			FMeshCommandOneFrameArray& MeshCommandsCSM = VisibleCommands[EMeshPass::MobileBasePassCSM];
-			checkf(MeshCommands.Num() == MeshCommandsCSM.Num(), TEXT("VisibleMeshDrawCommands of BasePass and MobileBasePassCSM are expected to match."));
-			for (int32 i = MeshCommands.Num() - 1; i >= 0; --i)
-			{
-				FVisibleMeshDrawCommand& MeshCommand = MeshCommands[i];
-				FVisibleMeshDrawCommand& MeshCommandCSM = MeshCommandsCSM[i];
-
-				if (bAlwaysUseCSM || (MeshCommand.DrawPrimitiveId < Scene->Primitives.Num() && MobileCSMVisibilityInfo.MobilePrimitiveCSMReceiverVisibilityMap[MeshCommand.DrawPrimitiveId]))
-				{
-					checkf(MeshCommand.DrawPrimitiveId == MeshCommandCSM.DrawPrimitiveId, TEXT("VisibleMeshDrawCommands of BasePass and MobileBasePassCSM are expected to match."));
-					// Use CSM's VisibleMeshDrawCommand.
-					MeshCommand = MeshCommandCSM;
-				}
-			}
-			MeshCommandsCSM.Reset();
-		}
+		MobileCSMVisibilityInfo.bAlwaysUseCSM = bAlwaysUseCSM;
 	}
 
 	{
