Adapt MaxPayloadSizeInBytes so the new float3 IndirectRadiance field fits.

diff --git a/Engine/Shaders/Private/RayTracing/RayTracingCommon.ush b/Engine/Shaders/Private/RayTracing/RayTracingCommon.ush
--- a/Engine/Shaders/Private/RayTracing/RayTracingCommon.ush
+++ b/Engine/Shaders/Private/RayTracing/RayTracingCommon.ush
@@ -65,14 +65,14 @@ struct FMaterialClosestHitPayload
 	uint PrimitiveLightingChannelMask;	//96	0bytes	(removed)
 	float4 CustomData;					//112	0bytes	(removed)
 	float GBufferAO;					//116	0bytes	(removed)
-	float3 IndirectIrradiance;			//120	0bytes	48bits	
+	float3 IndirectIrradiance;			//128	0bytes	48bits	
 
 	// Quite some code assume FRayHitInfo has a WorldPos
 	// So keep it here and force to re-construct it in Unpack call using ray information.
 	// It is not packed in FRayHitInfoPacked
-	float3 WorldPos;					//132	0bytes	(derived)
+	float3 WorldPos;					//140	0bytes	(derived)
 
-	uint HitKind;						//136	0bytes 1bit	(packed)
+	uint HitKind;						//144	0bytes 1bit	(packed)
 };
 
 #define USE_PACKED_PAYLOAD 0
diff --git a/Engine/Shaders/Private/RayTracing/RayTracingDebug.usf b/Engine/Shaders/Private/RayTracing/RayTracingDebug.usf
--- a/Engine/Shaders/Private/RayTracing/RayTracingDebug.usf
+++ b/Engine/Shaders/Private/RayTracing/RayTracingDebug.usf
@@ -98,7 +98,7 @@ void RayTracingDebugMainRGS()
 			Result = float4(Payload.GBufferAO, Payload.GBufferAO, Payload.GBufferAO, 1.0f);
 			break;
 		case RAY_TRACING_DEBUG_VIZ_INDIRECT_IRRADIANCE:
-			Result = float4(Payload.IndirectIrradiance, Payload.IndirectIrradiance, Payload.IndirectIrradiance, 1.0f);
+			Result = float4(Payload.IndirectIrradiance, 1.0f);
 			break;
 		case RAY_TRACING_DEBUG_VIZ_WORLD_POSITION:
 			Result = float4(Payload.WorldPos, 1.0f);
diff --git a/Engine/Source/Runtime/Renderer/Private/DeferredShadingRenderer.cpp b/Engine/Source/Runtime/Renderer/Private/DeferredShadingRenderer.cpp
--- a/Engine/Source/Runtime/Renderer/Private/DeferredShadingRenderer.cpp
+++ b/Engine/Source/Runtime/Renderer/Private/DeferredShadingRenderer.cpp
@@ -907,7 +907,7 @@ FRHIRayTracingPipelineState* FDeferredShadingSceneRenderer::BindRayTracingPipeli
 	TArray<FRayTracingHitGroupInitializer> RayTracingMaterialLibrary;
 	FShaderResource::GetRayTracingMaterialLibrary(RayTracingMaterialLibrary);
 	Initializer.SetHitGroups(RayTracingMaterialLibrary);
-	Initializer.MaxPayloadSizeInBytes = 136; // #dxr_todo: set Initializer.MaxPayloadSizeInBytes based on shader requirements (it's not obvious how to compute this right now)
+	Initializer.MaxPayloadSizeInBytes = 144; // #dxr_todo: set Initializer.MaxPayloadSizeInBytes based on shader requirements (it's not obvious how to compute this right now)
 
 	FRHIRayTracingPipelineState* PipelineState = PipelineStateCache::GetAndOrCreateRayTracingPipelineState(Initializer);
 
