Temporary fix for dynamic num mesh draw calculation to unblock QA. Will fix it properly tomorrow.

diff --git a/Engine/Source/Runtime/Renderer/Private/MobileShadingRenderer.cpp b/Engine/Source/Runtime/Renderer/Private/MobileShadingRenderer.cpp
--- a/Engine/Source/Runtime/Renderer/Private/MobileShadingRenderer.cpp
+++ b/Engine/Source/Runtime/Renderer/Private/MobileShadingRenderer.cpp
@@ -139,7 +139,7 @@ void FMobileSceneRenderer::SetupMobileBasePassAfterShadowInit(FExclusiveDepthSte
 			BasePassDepthStencilAccess,
 			MeshPassProcessor,
 			View.DynamicMeshElements,
-			View.MaxNumVisibleDynamicMeshes[EMeshPass::BasePass],
+			View.NumVisibleDynamicMeshElements[EMeshPass::BasePass],
 			ViewCommands.DynamicMeshCommandBuildRequests[EMeshPass::BasePass],
 			ViewCommands.MeshCommands[EMeshPass::BasePass],
 			BasePassCSMMeshPassProcessor,
diff --git a/Engine/Source/Runtime/Renderer/Private/SceneRendering.cpp b/Engine/Source/Runtime/Renderer/Private/SceneRendering.cpp
--- a/Engine/Source/Runtime/Renderer/Private/SceneRendering.cpp
+++ b/Engine/Source/Runtime/Renderer/Private/SceneRendering.cpp
@@ -922,7 +922,7 @@ void FViewInfo::Init()
 	DitherFadeOutUniformBuffer = nullptr;
 
 	VisibleDynamicMeshesPassMask.Reset();
-	FMemory::Memset(MaxNumVisibleDynamicMeshes, 0, sizeof(MaxNumVisibleDynamicMeshes));
+	FMemory::Memset(NumVisibleDynamicMeshElements, 0, sizeof(NumVisibleDynamicMeshElements));
 }
 
 FViewInfo::~FViewInfo()
@@ -3373,7 +3373,7 @@ void FParallelMeshDrawCommandPass::DispatchPassSetup(
 	FExclusiveDepthStencil::Type BasePassDepthStencilAccess,
 	FMeshPassProcessor* MeshPassProcessor,
 	const TArray<FMeshBatchAndRelevance, SceneRenderingAllocator>& DynamicMeshElements,
-	int32 MaxNumDynamicMeshesForThisPass,
+	int32 MaxNumDynamicMeshElementsForThisPass,
 	TArray<const FStaticMesh*, SceneRenderingAllocator>& InOutDynamicMeshCommandBuildRequests,
 	FMeshCommandOneFrameArray& InOutMeshDrawCommands,
 	FMeshPassProcessor* MobileBasePassCSMMeshPassProcessor,
@@ -3382,7 +3382,7 @@ void FParallelMeshDrawCommandPass::DispatchPassSetup(
 {
 	check(TaskEventRefs.Num() == 0 && MeshPassProcessor != nullptr);
 
-	MaxNumDraws = InOutMeshDrawCommands.Num() + InOutDynamicMeshCommandBuildRequests.Num() + MaxNumDynamicMeshesForThisPass;
+	MaxNumDraws = InOutMeshDrawCommands.Num() + InOutDynamicMeshCommandBuildRequests.Num() + MaxNumDynamicMeshElementsForThisPass;
 	if (MaxNumDraws <= 0)
 	{
 		return;
@@ -3637,7 +3637,7 @@ void FSceneRenderer::SetupMeshPass(FViewInfo& View, FExclusiveDepthStencil::Type
 				BasePassDepthStencilAccess,
 				MeshPassProcessor,
 				View.DynamicMeshElements,
-				View.MaxNumVisibleDynamicMeshes[PassType],
+				View.NumVisibleDynamicMeshElements[PassType],
 				ViewCommands.DynamicMeshCommandBuildRequests[PassType],
 				ViewCommands.MeshCommands[PassIndex]);
 		}
diff --git a/Engine/Source/Runtime/Renderer/Private/SceneRendering.h b/Engine/Source/Runtime/Renderer/Private/SceneRendering.h
--- a/Engine/Source/Runtime/Renderer/Private/SceneRendering.h
+++ b/Engine/Source/Runtime/Renderer/Private/SceneRendering.h
@@ -979,7 +979,7 @@ public:
 		FExclusiveDepthStencil::Type BasePassDepthStencilAccess,
 		FMeshPassProcessor* MeshPassProcessor,
 		const TArray<FMeshBatchAndRelevance, SceneRenderingAllocator>& DynamicMeshElements,
-		int32 MaxNumDynamicMeshesForThisPass,
+		int32 MaxNumDynamicMeshElementsForThisPass,
 		TArray<const FStaticMesh*, SceneRenderingAllocator>& InOutDynamicMeshCommandBuildRequests,
 		FMeshCommandOneFrameArray& InOutMeshDrawCommands,
 		FMeshPassProcessor* MobileBasePassCSMMeshPassProcessor = nullptr, // Required only for the mobile base pass.
@@ -1086,8 +1086,8 @@ public:
 	/** Marks used mesh passes by visible dynamic meshes. */
 	FMeshPassMask VisibleDynamicMeshesPassMask;
 
-	/** Max number of dynamic meshes per mesh pass. */
-	int32 MaxNumVisibleDynamicMeshes[EMeshPass::Num];
+	/** Number of dynamic mesh elements per mesh pass (inside FViewInfo::DynamicMeshElements). */
+	int32 NumVisibleDynamicMeshElements[EMeshPass::Num];
 
 	/** The dynamic editor primitives visible in this view. */
 	TArray<const FPrimitiveSceneInfo*,SceneRenderingAllocator> VisibleEditorPrimitives;
diff --git a/Engine/Source/Runtime/Renderer/Private/SceneVisibility.cpp b/Engine/Source/Runtime/Renderer/Private/SceneVisibility.cpp
--- a/Engine/Source/Runtime/Renderer/Private/SceneVisibility.cpp
+++ b/Engine/Source/Runtime/Renderer/Private/SceneVisibility.cpp
@@ -2689,16 +2689,23 @@ void FSceneRenderer::GatherDynamicMeshElements(
 	}
 	MeshCollector.ProcessTasks();
 
+	// Compute MaxNumVisibleDynamicMeshes.
 	for (int32 ViewIndex = 0; ViewIndex < Views.Num(); ++ViewIndex)
 	{
 		FViewInfo& View = Views[ViewIndex];
 
+		int32 NumDynamicMeshElements = 0;
+		for (int32 DynamicElementIndex = 0; DynamicElementIndex < View.DynamicMeshElements.Num(); ++DynamicElementIndex)
+		{
+			NumDynamicMeshElements += View.DynamicMeshElements[DynamicElementIndex].Mesh->Elements.Num();
+		}
+
 		for (int32 PassIndex = 0; PassIndex < EMeshPass::Num; ++PassIndex)
 		{
-			View.MaxNumVisibleDynamicMeshes[PassIndex] = 0;
+			View.NumVisibleDynamicMeshElements[PassIndex] = 0;
 			if (View.VisibleDynamicMeshesPassMask.Get((EMeshPass::Type) PassIndex))
 			{
-				View.MaxNumVisibleDynamicMeshes[PassIndex] = View.DynamicMeshElements.Num();
+				View.NumVisibleDynamicMeshElements[PassIndex] = NumDynamicMeshElements;
 			}
 		}
 	}
