Ray tracing debug modes now use an exposure of 1 for the tonemapper like standard buffer visualization modes.

diff --git a/Engine/Source/Runtime/Renderer/Private/PostProcess/PostProcessEyeAdaptation.cpp b/Engine/Source/Runtime/Renderer/Private/PostProcess/PostProcessEyeAdaptation.cpp
--- a/Engine/Source/Runtime/Renderer/Private/PostProcess/PostProcessEyeAdaptation.cpp
+++ b/Engine/Source/Runtime/Renderer/Private/PostProcess/PostProcessEyeAdaptation.cpp
@@ -87,7 +87,7 @@ inline static void ComputeEyeAdaptationValues(const ERHIFeatureLevel::Type MinFe
 	float AverageLuminanceScale = 1.f;
 
 	// When any of those flags are set, make sure the tonemapper uses an exposure of 1.
-	if (!EngineShowFlags.Lighting || (EngineShowFlags.VisualizeBuffer && View.CurrentBufferVisualizationMode != NAME_None) || View.Family->UseDebugViewPS())
+	if (!EngineShowFlags.Lighting || (EngineShowFlags.VisualizeBuffer && View.CurrentBufferVisualizationMode != NAME_None) || View.Family->UseDebugViewPS() || EngineShowFlags.RayTracingDebug )
 	{
 		LocalExposureMultipler = 1.f;
 	}
