Removed old mobile mesh drawing pipeline

diff --git a/Engine/Source/Runtime/Renderer/Private/MobileBasePass.cpp b/Engine/Source/Runtime/Renderer/Private/MobileBasePass.cpp
--- a/Engine/Source/Runtime/Renderer/Private/MobileBasePass.cpp
+++ b/Engine/Source/Runtime/Renderer/Private/MobileBasePass.cpp
@@ -185,7 +185,7 @@ void MobileBasePass::GetShaders(
 	}
 }
 
-bool MobileBasePass::UseSkyReflectionCapture(const FScene* RenderScene)
+static bool UseSkyReflectionCapture(const FScene* RenderScene)
 {
 	return RenderScene
 		&& RenderScene->ReflectionSceneData.RegisteredReflectionCapturePositions.Num() == 0
@@ -193,7 +193,7 @@ bool MobileBasePass::UseSkyReflectionCapture(const FScene* RenderScene)
 		&& RenderScene->SkyLight->ProcessedTexture->TextureRHI;
 }
 
-void MobileBasePass::GetSkyTextureParams(const FScene* Scene, float& AverageBrightnessOUT, FTexture*& ReflectionTextureOUT, float& OutSkyMaxMipIndex)
+static void GetSkyTextureParams(const FScene* Scene, float& AverageBrightnessOUT, FTexture*& ReflectionTextureOUT, float& OutSkyMaxMipIndex)
 {
 	if (Scene && Scene->SkyLight && Scene->SkyLight->ProcessedTexture->TextureRHI)
 	{
@@ -470,7 +470,8 @@ void TMobileBasePassPSPolicyParamType<FUniformLightMapPolicy>::GetShaderBindings
 		FVector4 CapturePositions[MaxNumReflections] = { FVector4(0, 0, 0, 0), FVector4(0, 0, 0, 0), FVector4(0, 0, 0, 0) };
 		FVector4 ReflectionParams(1.0f, 1.0f, 1.0f, 0.0f);
 		
-		if (MobileBasePass::UseSkyReflectionCapture(Scene))
+		// If no reflection captures are available then attempt to use sky light's texture.
+		if (UseSkyReflectionCapture(Scene))
 		{
 			// if > 0 this will disable shader's RGBM decoding and enable sky light tinting of this envmap.
 			// ReflectionParams.X == inv average brightness
@@ -478,7 +479,7 @@ void TMobileBasePassPSPolicyParamType<FUniformLightMapPolicy>::GetShaderBindings
 			if (FeatureLevel > ERHIFeatureLevel::ES2) // not-supported on ES2 at the moment
 			{
 				float AverageBrightness = 1.0f;
-				MobileBasePass::GetSkyTextureParams(Scene, AverageBrightness, ReflectionCubemapTextures[0], ReflectionParams.W);
+				GetSkyTextureParams(Scene, AverageBrightness, ReflectionCubemapTextures[0], ReflectionParams.W);
 				ReflectionParams.X = 1.0f / AverageBrightness;
 			}
 		}
diff --git a/Engine/Source/Runtime/Renderer/Private/MobileBasePassRendering.cpp b/Engine/Source/Runtime/Renderer/Private/MobileBasePassRendering.cpp
--- a/Engine/Source/Runtime/Renderer/Private/MobileBasePassRendering.cpp
+++ b/Engine/Source/Runtime/Renderer/Private/MobileBasePassRendering.cpp
@@ -31,23 +31,8 @@ static TAutoConsoleVariable<int32> CVarMobileUseLegacyShadingModel(
 	TEXT("If 1 then use legacy (pre 4.20) shading model (such as spherical guassian specular calculation.) (will cause a shader rebuild)"),
 	ECVF_ReadOnly | ECVF_RenderThreadSafe);
 
-
-// Changing this causes a full shader recompile
-static TAutoConsoleVariable<int32> CVarMobileSeparateMaskedPass(
-	TEXT("r.Mobile.SeparateMaskedPass"),
-	1,
-	TEXT("Draw masked primitives in separate pass after all opaque (default)"),
-	ECVF_RenderThreadSafe);
 IMPLEMENT_GLOBAL_SHADER_PARAMETER_STRUCT(FMobileBasePassUniformParameters, "MobileBasePass");
 
-static TAutoConsoleVariable<int32> CVarMobileParallelBasePass(
-	TEXT("r.Mobile.ParallelBasePass"),
-	0,
-	TEXT("Toggles parallel base pass rendering for the mobile renderer. Parallel rendering must be enabled for this to have an effect."),
-	ECVF_RenderThreadSafe
-);
-
-
 #define IMPLEMENT_MOBILE_SHADING_BASEPASS_LIGHTMAPPED_VERTEX_SHADER_TYPE(LightMapPolicyType,LightMapPolicyName) \
 	typedef TMobileBasePassVS< LightMapPolicyType, LDR_GAMMA_32 > TMobileBasePassVS##LightMapPolicyName##LDRGamma32; \
 	typedef TMobileBasePassVS< LightMapPolicyType, HDR_LINEAR_64 > TMobileBasePassVS##LightMapPolicyName##HDRLinear64; \
@@ -136,491 +121,6 @@ FMobileBasePassMovablePointLightInfo::FMobileBasePassMovablePointLightInfo(const
 	}
 }
 
-/** The action used to draw a base pass static mesh element. */
-class FDrawMobileBasePassStaticMeshAction
-{
-public:
-
-	FScene* Scene;
-	FStaticMesh* StaticMesh;
-
-	/** Initialization constructor. */
-	FDrawMobileBasePassStaticMeshAction(FScene* InScene, FStaticMesh* InStaticMesh):
-		Scene(InScene),
-		StaticMesh(InStaticMesh)
-	{}
-
-	inline bool ShouldPackAmbientSH() const
-	{
-		return false;
-	}
-
-	static bool CanUseDrawlistToToggleCombinedStaticAndCSM(const FPrimitiveSceneProxy* PrimitiveSceneProxy, ELightMapPolicyType LightMapPolicyType)
-	{
-		switch (LightMapPolicyType)
-		{
-			case LMP_MOBILE_DISTANCE_FIELD_SHADOWS_LIGHTMAP_AND_CSM:
-			case LMP_MOBILE_DIRECTIONAL_LIGHT_CSM_AND_SH_INDIRECT:
-			case LMP_MOBILE_DISTANCE_FIELD_SHADOWS_AND_LQ_LIGHTMAP:
-			case LMP_MOBILE_DIRECTIONAL_LIGHT_AND_SH_INDIRECT:
-			{
-				static auto* CVarMobileEnableStaticAndCSMShadowReceivers = IConsoleManager::Get().FindTConsoleVariableDataInt(TEXT("r.Mobile.EnableStaticAndCSMShadowReceivers"));
-				return CVarMobileEnableStaticAndCSMShadowReceivers->GetValueOnRenderThread() != 0;
-			}
-			case LMP_MOBILE_MOVABLE_DIRECTIONAL_LIGHT_CSM_WITH_LIGHTMAP:
-			case LMP_MOBILE_MOVABLE_DIRECTIONAL_LIGHT_WITH_LIGHTMAP:
-			case LMP_MOBILE_MOVABLE_DIRECTIONAL_LIGHT_CSM_AND_SH_INDIRECT:
-			case LMP_MOBILE_MOVABLE_DIRECTIONAL_LIGHT_AND_SH_INDIRECT:
-			case LMP_MOBILE_MOVABLE_DIRECTIONAL_LIGHT_CSM:
-			case LMP_MOBILE_MOVABLE_DIRECTIONAL_LIGHT:
-			{
-				static auto* CVarMobileEnableMovableLightCSMShaderCulling = IConsoleManager::Get().FindTConsoleVariableDataInt(TEXT("r.Mobile.EnableMovableLightCSMShaderCulling"));
-				return CVarMobileEnableMovableLightCSMShaderCulling->GetValueOnRenderThread() != 0;
-			}
-			default:
-			{
-				return false;
-			}
-		}
-	}
-
-	static bool CanReceiveCSM(const FLightSceneInfo* LightSceneInfo, const FPrimitiveSceneProxy* PrimitiveSceneProxy)
-	{
-		return MobileBasePass::StaticCanReceiveCSM(LightSceneInfo, PrimitiveSceneProxy);
-	}
-
-	const FScene* GetScene() const 
-	{ 
-		return Scene;
-	}
-
-	/** Draws the mesh with a specific light-map type */
-	void Process(
-		FRHICommandList& RHICmdList, 
-		const FMobileProcessBasePassMeshParameters& Parameters,
-		const ELightMapPolicyType LightMapPolicy,
-		const FUniformLightMapPolicy::ElementDataType& LightMapElementData
-		) const
-	{
-		EBasePassDrawListType DrawType = EBasePass_Default;
-
-		if (StaticMesh->IsMasked(Parameters.FeatureLevel) && CVarMobileSeparateMaskedPass.GetValueOnRenderThread() != 0)
-		{
-			DrawType = EBasePass_Masked;	
-		}
-
-		if ( Scene )
-		{
-			// Determine if this primitive has the possibility of using combined static and CSM.
-			if (CanUseDrawlistToToggleCombinedStaticAndCSM(Parameters.PrimitiveSceneProxy, LightMapPolicy))
-			{
-				// if applicable, returns the corresponding CSM or non-CSM lightmap policy of LightMapPolicyType
-				auto GetAlternativeLightMapPolicy = [](ELightMapPolicyType LightMapPolicyType)
-				{
-					switch (LightMapPolicyType)
-					{
-						case LMP_MOBILE_DISTANCE_FIELD_SHADOWS_LIGHTMAP_AND_CSM:
-							return LMP_MOBILE_DISTANCE_FIELD_SHADOWS_AND_LQ_LIGHTMAP;
-						case LMP_MOBILE_DIRECTIONAL_LIGHT_CSM_AND_SH_INDIRECT:
-							return LMP_MOBILE_DIRECTIONAL_LIGHT_AND_SH_INDIRECT;
-						case LMP_MOBILE_DISTANCE_FIELD_SHADOWS_AND_LQ_LIGHTMAP:
-							return LMP_MOBILE_DISTANCE_FIELD_SHADOWS_LIGHTMAP_AND_CSM;
-						case LMP_MOBILE_DIRECTIONAL_LIGHT_AND_SH_INDIRECT:
-							return LMP_MOBILE_DIRECTIONAL_LIGHT_CSM_AND_SH_INDIRECT;
-
-						// movable light CSMs
-						case LMP_MOBILE_MOVABLE_DIRECTIONAL_LIGHT_CSM_WITH_LIGHTMAP:
-							return LMP_MOBILE_MOVABLE_DIRECTIONAL_LIGHT_WITH_LIGHTMAP;
-						case LMP_MOBILE_MOVABLE_DIRECTIONAL_LIGHT_WITH_LIGHTMAP:
-							return LMP_MOBILE_MOVABLE_DIRECTIONAL_LIGHT_CSM_WITH_LIGHTMAP;
-
-						case LMP_MOBILE_MOVABLE_DIRECTIONAL_LIGHT_CSM_AND_SH_INDIRECT:
-							return LMP_MOBILE_MOVABLE_DIRECTIONAL_LIGHT_AND_SH_INDIRECT;
-						case LMP_MOBILE_MOVABLE_DIRECTIONAL_LIGHT_AND_SH_INDIRECT:
-							return LMP_MOBILE_MOVABLE_DIRECTIONAL_LIGHT_CSM_AND_SH_INDIRECT;
-
-						case LMP_MOBILE_MOVABLE_DIRECTIONAL_LIGHT_CSM:
-							return LMP_MOBILE_MOVABLE_DIRECTIONAL_LIGHT;
-						case LMP_MOBILE_MOVABLE_DIRECTIONAL_LIGHT:
-							return LMP_MOBILE_MOVABLE_DIRECTIONAL_LIGHT_CSM;
-
-					}
-					return LightMapPolicyType;
-				};
-
-				const ELightMapPolicyType AlternativeLightMapPolicy = GetAlternativeLightMapPolicy(LightMapPolicy);
-				const bool bHasCSMCounterpart = AlternativeLightMapPolicy != LightMapPolicy;
-				if (bHasCSMCounterpart)
-				{
-					// Is the passed in lightmap policy CSM capable or not
-					const bool bIsCSMCapableLightPolicy = LightMapPolicy == LMP_MOBILE_DISTANCE_FIELD_SHADOWS_LIGHTMAP_AND_CSM
-						|| LightMapPolicy == LMP_MOBILE_DIRECTIONAL_LIGHT_CSM_AND_SH_INDIRECT
-						|| LightMapPolicy == LMP_MOBILE_MOVABLE_DIRECTIONAL_LIGHT_CSM_WITH_LIGHTMAP
-						|| LightMapPolicy == LMP_MOBILE_MOVABLE_DIRECTIONAL_LIGHT_CSM_AND_SH_INDIRECT
-						|| LightMapPolicy == LMP_MOBILE_MOVABLE_DIRECTIONAL_LIGHT_CSM;
-
-					if (bIsCSMCapableLightPolicy)
-					{
-						// Alternative policy is the non-CSM version.
-						AddMeshToStaticDrawList(Scene->GetMobileBasePassCSMDrawList(DrawType), Parameters, LightMapPolicy, LightMapElementData);
-						AddMeshToStaticDrawList(Scene->GetMobileBasePassDrawList(DrawType), Parameters, AlternativeLightMapPolicy, LightMapElementData);
-					}
-					else
-					{
-						// Alternative policy is the CSM version.
-						AddMeshToStaticDrawList(Scene->GetMobileBasePassCSMDrawList(DrawType), Parameters, AlternativeLightMapPolicy, LightMapElementData);
-						AddMeshToStaticDrawList(Scene->GetMobileBasePassDrawList(DrawType), Parameters, LightMapPolicy, LightMapElementData);
-					}
-
-					return; // avoid adding to draw list twice.
-				}
-			}
-
-			AddMeshToStaticDrawList(Scene->GetMobileBasePassDrawList(DrawType), Parameters, LightMapPolicy, LightMapElementData);
-		}
-	}
-
-	void AddMeshToStaticDrawList(TStaticMeshDrawList<FMobileBasePassUniformDrawingPolicy> &DrawList,
-		const FMobileProcessBasePassMeshParameters &Parameters, const ELightMapPolicyType LightMapPolicy, const FUniformLightMapPolicy::ElementDataType& LightMapElementData) const
-	{
-		ERHIFeatureLevel::Type FeatureLevel = Scene->GetFeatureLevel();
-		// Add the static mesh to the draw list.
-		DrawList.AddMesh(
-			StaticMesh,
-			LightMapElementData,
-			FMobileBasePassUniformDrawingPolicy(
-				StaticMesh->VertexFactory,
-				StaticMesh->MaterialRenderProxy,
-				Parameters.Material,
-				LightMapPolicy,
-				Parameters.NumMovablePointLights,
-				Parameters.BlendMode,
-				Parameters.ShadingModel != MSM_Unlit && Scene->ShouldRenderSkylightInBasePass(Parameters.BlendMode),
-				ComputeMeshOverrideSettings(Parameters.Mesh),
-				FeatureLevel,
-				IsMobileHDR() // bEnableReceiveDecalOutput
-				),
-			FeatureLevel
-		);
-	}
-};
-
-void FMobileBasePassOpaqueDrawingPolicyFactory::AddStaticMesh(FRHICommandList& RHICmdList, FScene* Scene, FStaticMesh* StaticMesh)
-{
-	// Determine the mesh's material and blend mode.
-	const ERHIFeatureLevel::Type FeatureLevel = Scene->GetFeatureLevel();
-	const FMaterial* Material = StaticMesh->MaterialRenderProxy->GetMaterial(FeatureLevel);
-	const EBlendMode BlendMode = Material->GetBlendMode();
-
-	// Only draw opaque materials
-	if (!IsTranslucentBlendMode(BlendMode))
-	{
-		// following check moved from ProcessMobileBasePassMesh to avoid passing feature level.
-		check(!AllowHighQualityLightmaps(Scene->GetFeatureLevel()));
-
-		bool bAllowFog = true;
-		FPrimitiveSceneProxy* Proxy = StaticMesh->PrimitiveSceneInfo->Proxy;
-		const FLightSceneInfo* MobileDirectionalLight = MobileBasePass::GetDirectionalLightInfo(Scene, Proxy);
-
-		FDrawMobileBasePassStaticMeshAction Action(
-			Scene, 
-			StaticMesh);
-
-		FMobileProcessBasePassMeshParameters Parameters(
-			*StaticMesh, 
-			*Material, 
-			Proxy, 
-			bAllowFog, 
-			FeatureLevel);
-
-
-		bool bCanReceiveCSM = Action.CanReceiveCSM(MobileDirectionalLight, Proxy);
-		EMaterialShadingModel ShadingModel = Material->GetShadingModel();
-		ELightMapPolicyType LightmapPolicyType = MobileBasePass::SelectMeshLightmapPolicy(Scene, *StaticMesh, Proxy, MobileDirectionalLight, ShadingModel, bCanReceiveCSM, FeatureLevel);
-			
-		Action.Process(RHICmdList, Parameters, LightmapPolicyType, StaticMesh->LCI);
-	}
-}
-
-/** The action used to draw a base pass dynamic mesh element. */
-class FDrawMobileBasePassDynamicMeshAction
-{
-public:
-
-	const FViewInfo& View;
-	FDrawingPolicyRenderState DrawRenderState;
-	FHitProxyId HitProxyId;
-
-	inline bool ShouldPackAmbientSH() const
-	{
-		return false;
-	}
-
-	bool CanReceiveCSM(const FLightSceneInfo* LightSceneInfo, const FPrimitiveSceneProxy* PrimitiveSceneProxy) const
-	{
-		return CanReceiveCSM(View, LightSceneInfo, PrimitiveSceneProxy);
-	}
-
-	static bool CanReceiveCSM(const FViewInfo& View, const FLightSceneInfo* LightSceneInfo, const FPrimitiveSceneProxy* PrimitiveSceneProxy)
-	{
-		if (PrimitiveSceneProxy == nullptr || LightSceneInfo == nullptr)
-		{
-			return false;
-		}
-
-		// Check that this primitive is eligible for CSM.
- 		const FVisibleLightViewInfo& VisibleLightViewInfo = View.VisibleLightInfos[LightSceneInfo->Id];
-
-		static auto* CVarMobileEnableStaticAndCSMShadowReceivers = IConsoleManager::Get().FindTConsoleVariableDataInt(TEXT("r.Mobile.EnableStaticAndCSMShadowReceivers"));
-		static auto* CVarMobileEnableMovableLightCSMShaderCulling = IConsoleManager::Get().FindTConsoleVariableDataInt(TEXT("r.Mobile.EnableMovableLightCSMShaderCulling"));
-		const bool bMobileEnableMovableLightCSMShaderCulling = CVarMobileEnableMovableLightCSMShaderCulling->GetValueOnRenderThread() == 1;
-		const bool bMobileEnableStaticAndCSMShadowReceivers = CVarMobileEnableStaticAndCSMShadowReceivers->GetValueOnRenderThread() == 1;
-
-		const bool bMovableLight = LightSceneInfo->Proxy->IsMovable();
-		const bool bMovableLightCastsCSM = bMovableLight && LightSceneInfo->ShouldRenderViewIndependentWholeSceneShadows();
-
-		return PrimitiveSceneProxy->ShouldReceiveMobileCSMShadows()
-			&&	(	// movable CSM culling is disabled and movable light is in use
-					(!bMobileEnableMovableLightCSMShaderCulling && bMovableLightCastsCSM)
-					||
-					// CSM culling is active
-					(View.MobileCSMVisibilityInfo.bMobileDynamicCSMInUse
-						&& (bMobileEnableStaticAndCSMShadowReceivers || bMobileEnableMovableLightCSMShaderCulling)
-						&& View.MobileCSMVisibilityInfo.MobilePrimitiveCSMReceiverVisibilityMap[PrimitiveSceneProxy->GetPrimitiveSceneInfo()->GetIndex()])
-				);
-	}
-
-	const FScene* GetScene() const
-	{
-		auto* Scene = (FScene*)View.Family->Scene;
-		return Scene;
-	}
-
-	/** Initialization constructor. */
-	FDrawMobileBasePassDynamicMeshAction(
-		FRHICommandList& InRHICmdList,
-		const FViewInfo& InView,
-		float InDitheredLODTransitionAlpha,
-		const FDrawingPolicyRenderState& InDrawRenderState,
-		const FHitProxyId InHitProxyId
-		)
-		: View(InView)
-		, DrawRenderState(InDrawRenderState)
-		, HitProxyId(InHitProxyId)
-	{
-		DrawRenderState.SetDitheredLODTransitionAlpha(InDitheredLODTransitionAlpha);
-	}
-
-	/** Draws mesh with a specific light-map type, and shader complexity predicate. */
-	void Process(
-		FRHICommandList& RHICmdList, 
-		const FMobileProcessBasePassMeshParameters& Parameters,
-		const ELightMapPolicyType LightMapPolicyType,
-		const FUniformLightMapPolicy::ElementDataType& LightMapElementData
-		)
-	{
-#if !(UE_BUILD_SHIPPING || UE_BUILD_TEST)
-		// Treat masked materials as if they don't occlude in shader complexity, which is PVR behavior
-		if(Parameters.BlendMode == BLEND_Masked && View.Family->EngineShowFlags.ShaderComplexity)
-		{
-			DrawRenderState.SetDepthStencilState(TStaticDepthStencilState<false,CF_DepthNearOrEqual>::GetRHI());
-		}
-#endif
-
-		const bool bIsLitMaterial = Parameters.ShadingModel != MSM_Unlit;
-		const FScene* Scene = Parameters.PrimitiveSceneProxy ? Parameters.PrimitiveSceneProxy->GetPrimitiveSceneInfo()->Scene : NULL;
-
-		FMobileBasePassUniformDrawingPolicy DrawingPolicy(
-			Parameters.Mesh.VertexFactory,
-			Parameters.Mesh.MaterialRenderProxy,
-			Parameters.Material,
-			LightMapPolicyType,
-			Parameters.NumMovablePointLights,
-			Parameters.BlendMode,
-			Parameters.ShadingModel != MSM_Unlit && Scene && Scene->ShouldRenderSkylightInBasePass(Parameters.BlendMode),
-			ComputeMeshOverrideSettings(Parameters.Mesh),
-			View.GetFeatureLevel(),
-			IsMobileHDR() // bEnableReceiveDecalOutput
-			);
-
-		DrawingPolicy.SetupPipelineState(DrawRenderState, View);
-		CommitGraphicsPipelineState(RHICmdList, DrawingPolicy, DrawRenderState, DrawingPolicy.GetBoundShaderStateInput(View.GetFeatureLevel()), Parameters.Mesh.MaterialRenderProxy);
-		DrawingPolicy.SetSharedState(RHICmdList, DrawRenderState, &View, FMobileBasePassUniformDrawingPolicy::ContextDataType());
-
-		for( int32 BatchElementIndex=0;BatchElementIndex<Parameters.Mesh.Elements.Num();BatchElementIndex++ )
-		{
-			TDrawEvent<FRHICommandList> MeshEvent;
-			BeginMeshDrawEvent(RHICmdList, Parameters.PrimitiveSceneProxy, Parameters.Mesh, MeshEvent, EnumHasAnyFlags(EShowMaterialDrawEventTypes(GShowMaterialDrawEventTypes), EShowMaterialDrawEventTypes::MobileBasePass));
-
-			DrawingPolicy.SetMeshRenderState(
-				RHICmdList, 
-				View,
-				Parameters.PrimitiveSceneProxy,
-				Parameters.Mesh,
-				BatchElementIndex,
-				DrawRenderState,
-				FMobileBasePassUniformDrawingPolicy::ElementDataType(LightMapElementData),
-				FMobileBasePassUniformDrawingPolicy::ContextDataType()
-				);
-			DrawingPolicy.DrawMesh(RHICmdList, View, Parameters.Mesh, BatchElementIndex);
-		}
-	}
-};
-
-bool FMobileBasePassOpaqueDrawingPolicyFactory::DrawDynamicMesh(
-	FRHICommandList& RHICmdList, 
-	const FViewInfo& View,
-	ContextType DrawingContext,
-	const FMeshBatch& Mesh,
-	bool bPreFog,
-	const FDrawingPolicyRenderState& DrawRenderState,
-	const FPrimitiveSceneProxy* PrimitiveSceneProxy,
-	FHitProxyId HitProxyId
-	)
-{
-	// Determine the mesh's material and blend mode.
-	const ERHIFeatureLevel::Type FeatureLevel = View.GetFeatureLevel();
-	const FMaterial* Material = Mesh.MaterialRenderProxy->GetMaterial(FeatureLevel);
-	const EBlendMode BlendMode = Material->GetBlendMode();
-
-	// Only draw opaque materials.
-	if (!IsTranslucentBlendMode(BlendMode))
-	{
-		bool bAllowFog = true;
-
-		FDrawMobileBasePassDynamicMeshAction Action(
-			RHICmdList,
-			View, 
-			Mesh.DitheredLODTransitionAlpha,
-			DrawRenderState, 
-			HitProxyId);
-		
-		FMobileProcessBasePassMeshParameters Parameters(
-			Mesh, 
-			*Material, 
-			PrimitiveSceneProxy, 
-			bAllowFog, 
-			FeatureLevel);
-		
-		const FScene* Scene = (FScene*)View.Family->Scene;
-		const FLightSceneInfo* MobileDirectionalLight = MobileBasePass::GetDirectionalLightInfo(Scene, PrimitiveSceneProxy);
-		bool bCanReceiveCSM = Action.CanReceiveCSM(MobileDirectionalLight, PrimitiveSceneProxy);
-		EMaterialShadingModel ShadingModel = Material->GetShadingModel();
-		ELightMapPolicyType LightmapPolicyType = MobileBasePass::SelectMeshLightmapPolicy(Scene, Mesh, PrimitiveSceneProxy, MobileDirectionalLight, ShadingModel, bCanReceiveCSM, FeatureLevel);
-					
-		Action.Process(RHICmdList, Parameters, LightmapPolicyType, Mesh.LCI);
-		return true;
-	}
-	else
-	{
-		return false;
-	}
-}
-
-/** Base pass sorting modes. */
-namespace EBasePassSort
-{
-	enum Type
-	{
-		/** Automatically select based on hardware/platform. */
-		Auto = 0,
-		/** No sorting. */
-		None = 1,
-		/** Sorts state buckets, not individual meshes. */
-		SortStateBuckets = 2,
-		/** Per mesh sorting. */
-		SortPerMesh = 3,
-
-		/** Useful range of sort modes. */
-		FirstForcedMode = None,
-		LastForcedMode = SortPerMesh
-	};
-};
-TAutoConsoleVariable<int32> GSortBasePass(TEXT("r.ForwardBasePassSort"),0,
-	TEXT("How to sort the mobile base pass:\n")
-	TEXT("\t0: Decide automatically based on the hardware and threading configuration.\n")
-	TEXT("\t1: No sorting.\n")
-	TEXT("\t2: Sort drawing policies.\n")
-	TEXT("\t3: Sort drawing policies and the meshes within them. Will not use the parallel path."),
-	ECVF_RenderThreadSafe);
-TAutoConsoleVariable<int32> GMaxBasePassDraws(TEXT("r.MaxForwardBasePassDraws"),0,TEXT("Stops rendering static mobile base pass draws after the specified number of times. Useful for seeing the order in which meshes render when optimizing."),ECVF_RenderThreadSafe);
-
-EBasePassSort::Type GetSortMode()
-{
-	int32 SortMode = GSortBasePass.GetValueOnRenderThread();
-	if (SortMode >= EBasePassSort::FirstForcedMode && SortMode <= EBasePassSort::LastForcedMode)
-	{
-		return (EBasePassSort::Type)SortMode;
-	}
-
-	// Determine automatically.
-	if (GRHICommandList.UseParallelAlgorithms() || GHardwareHiddenSurfaceRemoval)
-	{
-		return EBasePassSort::None;
-	}
-	else
-	{
-		return EBasePassSort::SortPerMesh;
-	}
-}
-
-/** Helper function for drawing sorted meshes */
-static void DrawVisibleFrontToBack(
-	FRHICommandListImmediate& RHICmdList,
-	FScene* const Scene,
-	const EBasePassDrawListType DrawListType,
-	const FViewInfo& View,
-	FDrawingPolicyRenderState& DrawRenderState,
-	const FMobileCSMVisibilityInfo& MobileCSMVisibilityInfo,
-	int32& MaxDraws)
-{
-#if FRAMEPRO_ENABLED
-	FFrameProProfiler::PushEvent();
-#else
-	QUICK_SCOPE_CYCLE_COUNTER(STAT_StaticDrawListDrawTimeFrontToBack);
-#endif // FRAMEPRO_ENABLED
-
-	const bool bIsCSM = MobileCSMVisibilityInfo.bMobileDynamicCSMInUse;
-	int32 NumDraws = 0;
-	if (bIsCSM)
-	{
-		NumDraws += Scene->MobileBasePassUniformLightMapPolicyDrawListWithCSM[DrawListType].DrawVisibleFrontToBack(RHICmdList, View, DrawRenderState, MobileCSMVisibilityInfo.MobileCSMStaticMeshVisibilityMap, MobileCSMVisibilityInfo.MobileCSMStaticBatchVisibility, MaxDraws);
-		NumDraws += Scene->MobileBasePassUniformLightMapPolicyDrawList[DrawListType].DrawVisibleFrontToBack(RHICmdList, View, DrawRenderState, MobileCSMVisibilityInfo.MobileNonCSMStaticMeshVisibilityMap, MobileCSMVisibilityInfo.MobileNonCSMStaticBatchVisibility, MaxDraws);
-	}
-	else
-	{
-		NumDraws += Scene->MobileBasePassUniformLightMapPolicyDrawList[DrawListType].DrawVisibleFrontToBack(RHICmdList, View, DrawRenderState, View.StaticMeshVisibilityMap, View.StaticMeshBatchVisibility, MaxDraws);
-	}
-
-	MaxDraws -= NumDraws;
-
-#if FRAMEPRO_ENABLED
-	FFrameProProfiler::PopEvent(*FString::Printf(TEXT("STAT_StaticDrawListDrawTimeFrontToBack (%d draws)"), NumDraws));
-#endif // FRAMEPRO_ENABLED
-}
-
-/** Helper function for drawing unsorted meshes */
-static void DrawVisible(
-	FRHICommandListImmediate& RHICmdList,
-	FScene* const Scene,
-	const EBasePassDrawListType DrawListType,
-	const FViewInfo& View,
-	const FDrawingPolicyRenderState& DrawRenderState,
-	const FMobileCSMVisibilityInfo& MobileCSMVisibilityInfo)
-{
-	SCOPE_CYCLE_COUNTER(STAT_StaticDrawListDrawTime);
-	const bool bIsCSM = MobileCSMVisibilityInfo.bMobileDynamicCSMInUse;
-	if (bIsCSM)
-	{
-		Scene->MobileBasePassUniformLightMapPolicyDrawListWithCSM[DrawListType].DrawVisible(RHICmdList, View, DrawRenderState, MobileCSMVisibilityInfo.MobileCSMStaticMeshVisibilityMap, MobileCSMVisibilityInfo.MobileCSMStaticBatchVisibility);
-		Scene->MobileBasePassUniformLightMapPolicyDrawList[DrawListType].DrawVisible(RHICmdList, View, DrawRenderState, MobileCSMVisibilityInfo.MobileNonCSMStaticMeshVisibilityMap, MobileCSMVisibilityInfo.MobileNonCSMStaticBatchVisibility);
-	}
-	else
-	{
-		Scene->MobileBasePassUniformLightMapPolicyDrawList[DrawListType].DrawVisible(RHICmdList, View, DrawRenderState, View.StaticMeshVisibilityMap, View.StaticMeshBatchVisibility);
-	}
-}
-
 void SetupMobileBasePassUniformParameters(
 	FRHICommandListImmediate& RHICmdList, 
 	const FViewInfo& View, 
@@ -693,6 +193,40 @@ void SetupMobileDirectionalLightUniformParameters(
 	}
 }
 
+void FMobileSceneRenderer::RenderMobileBasePass(FRHICommandListImmediate& RHICmdList, const TArrayView<const FViewInfo*> PassViews)
+{
+	SCOPED_DRAW_EVENT(RHICmdList, MobileBasePass);
+	SCOPE_CYCLE_COUNTER(STAT_BasePassDrawTime);
+
+	for (int32 ViewIndex = 0; ViewIndex < PassViews.Num(); ViewIndex++)
+	{
+		SCOPED_CONDITIONAL_DRAW_EVENTF(RHICmdList, EventView, Views.Num() > 1, TEXT("View%d"), ViewIndex);
+		const FViewInfo& View = *PassViews[ViewIndex];
+		if (!View.ShouldRenderView())
+		{
+			continue;
+		}
+
+		if (Scene->UniformBuffers.UpdateViewUniformBuffer(View))
+		{
+			UpdateOpaqueBasePassUniformBuffer(RHICmdList, View);
+			UpdateDirectionalLightUniformBuffers(RHICmdList, View);
+		}
+		
+		RHICmdList.SetViewport(View.ViewRect.Min.X, View.ViewRect.Min.Y, 0, View.ViewRect.Max.X, View.ViewRect.Max.Y, 1);
+		View.ParallelMeshDrawCommandPasses[EMeshPass::BasePass].DispatchDraw(nullptr, RHICmdList);
+
+		// editor primitives
+		{
+			FDrawingPolicyRenderState DrawRenderState(View, Scene->UniformBuffers.MobileOpaqueBasePassUniformBuffer);
+			DrawRenderState.SetBlendState(TStaticBlendStateWriteMask<CW_RGBA>::GetRHI());
+			DrawRenderState.SetDepthStencilAccess(GetDefaultBasePassDepthStencilAccess(EShadingPath::Mobile));
+			DrawRenderState.SetDepthStencilState(TStaticDepthStencilState<true, CF_DepthNearOrEqual>::GetRHI());
+			RenderMobileEditorPrimitives(RHICmdList, View, DrawRenderState);
+		}
+	}
+}
+
 void FMobileSceneRenderer::RenderMobileEditorPrimitives(FRHICommandList& RHICmdList, const FViewInfo& View, const FDrawingPolicyRenderState& DrawRenderState)
 {
 	QUICK_SCOPE_CYCLE_COUNTER(STAT_EditorDynamicPrimitiveDrawTime);
@@ -705,41 +239,8 @@ void FMobileSceneRenderer::RenderMobileEditorPrimitives(FRHICommandList& RHICmdL
 	{
 		const bool bNeedToSwitchVerticalAxis = RHINeedsToSwitchVerticalAxis(GShaderPlatformForFeatureLevel[FeatureLevel]) && !IsMobileHDR();
 
-		if (UseMeshDrawCommandPipeline())
-		{
-			DrawDynamicMeshPass(View, RHICmdList,
-				[&View, &DrawRenderState](FDynamicPassMeshDrawListContext* DynamicMeshPassContext)
-				{
-					FEditorPrimitivesBasePassMeshProcessor PassMeshProcessor(
-						View.Family->Scene->GetRenderScene(),
-						View.GetFeatureLevel(),
-						&View,
-						DrawRenderState,
-						false,
-						DynamicMeshPassContext);
-
-					const uint64 DefaultBatchElementMask = ~0ull;
-					
-					for (int32 MeshIndex = 0; MeshIndex < View.ViewMeshElements.Num(); MeshIndex++)
-					{
-						const FMeshBatch& MeshBatch = View.ViewMeshElements[MeshIndex];
-						PassMeshProcessor.AddMeshBatch(MeshBatch, DefaultBatchElementMask, nullptr);
-					}
-				});
-		}
-		else
-		{
-			// Draw the base pass for the view's batched mesh elements.
-			DrawViewElements<FMobileBasePassOpaqueDrawingPolicyFactory>(RHICmdList, View, DrawRenderState, FMobileBasePassOpaqueDrawingPolicyFactory::ContextType(), SDPG_World, true);
-		}
-
-		// Draw the view's batched simple elements(lines, sprites, etc).
-		View.BatchedViewElements.Draw(RHICmdList, DrawRenderState, FeatureLevel, bNeedToSwitchVerticalAxis, View, false);
-
-		if (UseMeshDrawCommandPipeline())
-		{
-			DrawDynamicMeshPass(View, RHICmdList,
-				[&View, &DrawRenderState](FDynamicPassMeshDrawListContext* DynamicMeshPassContext)
+		DrawDynamicMeshPass(View, RHICmdList,
+			[&View, &DrawRenderState](FDynamicPassMeshDrawListContext* DynamicMeshPassContext)
 			{
 				FEditorPrimitivesBasePassMeshProcessor PassMeshProcessor(
 					View.Family->Scene->GetRenderScene(),
@@ -750,396 +251,38 @@ void FMobileSceneRenderer::RenderMobileEditorPrimitives(FRHICommandList& RHICmdL
 					DynamicMeshPassContext);
 
 				const uint64 DefaultBatchElementMask = ~0ull;
-
-				for (int32 MeshIndex = 0; MeshIndex < View.TopViewMeshElements.Num(); MeshIndex++)
+					
+				for (int32 MeshIndex = 0; MeshIndex < View.ViewMeshElements.Num(); MeshIndex++)
 				{
-					const FMeshBatch& MeshBatch = View.TopViewMeshElements[MeshIndex];
+					const FMeshBatch& MeshBatch = View.ViewMeshElements[MeshIndex];
 					PassMeshProcessor.AddMeshBatch(MeshBatch, DefaultBatchElementMask, nullptr);
 				}
 			});
-		}
-		else
-		{
-			// Draw foreground objects last
-			DrawViewElements<FMobileBasePassOpaqueDrawingPolicyFactory>(RHICmdList, View, DrawRenderState, FMobileBasePassOpaqueDrawingPolicyFactory::ContextType(), SDPG_Foreground, true);
-		}
 
 		// Draw the view's batched simple elements(lines, sprites, etc).
-		View.TopBatchedViewElements.Draw(RHICmdList, DrawRenderState, FeatureLevel, bNeedToSwitchVerticalAxis, View, false);
-	}
-}
-
-void FMobileSceneRenderer::RenderMobileBasePassDynamicData(FRHICommandList& RHICmdList, const FViewInfo& View, const FDrawingPolicyRenderState& DrawRenderState, EBlendMode BlendMode, bool bWireFrame, int32 FirstElement, int32 AfterLastElement)
-{
-	AfterLastElement = FMath::Min(View.DynamicMeshElements.Num(), AfterLastElement);
-
-	if (FirstElement >= AfterLastElement)
-	{
-		return;
-	}
-	SCOPE_CYCLE_COUNTER(STAT_DynamicPrimitiveDrawTime);
-	SCOPED_DRAW_EVENT(RHICmdList, Dynamic);
-
-	FMobileBasePassOpaqueDrawingPolicyFactory::ContextType Context;
-
-	for (int32 Index = FirstElement; Index < AfterLastElement; Index++)
-	{
-		const FMeshBatchAndRelevance& MeshBatchAndRelevance = View.DynamicMeshElements[Index];
-
-		if (MeshBatchAndRelevance.GetRenderInMainPass() &&
-			((BlendMode == BLEND_Opaque && MeshBatchAndRelevance.GetHasOpaqueMaterial()) || 
-			(BlendMode == BLEND_Masked && MeshBatchAndRelevance.GetHasMaskedMaterial()) || bWireFrame))
-		{
-			const FMeshBatch& MeshBatch = *MeshBatchAndRelevance.Mesh;
-			FMobileBasePassOpaqueDrawingPolicyFactory::DrawDynamicMesh(RHICmdList, View, Context, MeshBatch, true, DrawRenderState, MeshBatchAndRelevance.PrimitiveSceneProxy, MeshBatch.BatchHitProxyId);
-		}
-	}
-}
-
-static void SetupMobileBasePassView(FRHICommandList& RHICmdList, const FViewInfo& View, FDrawingPolicyRenderState& DrawRenderState)
-{
-	// Opaque blending
-	if (View.bIsPlanarReflection)
-	{
-		DrawRenderState.SetBlendState(TStaticBlendState<CW_RGBA, BO_Add, BF_One, BF_Zero, BO_Add, BF_Zero, BF_Zero>::GetRHI());
-	}
-	else
-	{
-		DrawRenderState.SetBlendState(TStaticBlendStateWriteMask<CW_RGBA>::GetRHI());
-	}
-
-	DrawRenderState.SetDepthStencilState(TStaticDepthStencilState<true, CF_DepthNearOrEqual>::GetRHI());
-	RHICmdList.SetViewport(View.ViewRect.Min.X, View.ViewRect.Min.Y, 0, View.ViewRect.Max.X, View.ViewRect.Max.Y, 1);
-}
-
-DECLARE_CYCLE_STAT(TEXT("MobileBasepass"), STAT_CLP_MobileBasepass, STATGROUP_ParallelCommandListMarkers);
-
-class FMobileBasePassParallelCommandListSet : public FParallelCommandListSet
-{
-public:
-	const FSceneViewFamily& ViewFamily;
-
-	FMobileBasePassParallelCommandListSet(
-		const FViewInfo& InView, 
-		const FSceneRenderer* InSceneRenderer, 
-		FRHICommandListImmediate& InParentCmdList, 
-		bool bInParallelExecute, 
-		bool bInCreateSceneContext, 
-		const FSceneViewFamily& InViewFamily, 
-		const FDrawingPolicyRenderState& InDrawRenderState)
-		: FParallelCommandListSet(GET_STATID(STAT_CLP_MobileBasepass), InView, InSceneRenderer, InParentCmdList, bInParallelExecute, bInCreateSceneContext, InDrawRenderState)
-		, ViewFamily(InViewFamily)
-	{
-		SetStateOnCommandList(ParentCmdList);
-	}
-
-	virtual ~FMobileBasePassParallelCommandListSet()
-	{
-		Dispatch();
-	}
-
-	virtual void SetStateOnCommandList(FRHICommandList& CmdList) override
-	{
-		FParallelCommandListSet::SetStateOnCommandList(CmdList);
-		check(!bParallelExecute); // SetupMobileBasePassView is not (yet) complete enough for deferred contexts. Also need the rendertargets and ALL other state required!
-		SetupMobileBasePassView(CmdList, View, DrawRenderState);
-	}
-};
-
-class FRenderMobileBasePassDynamicDataThreadTask : public FRenderTask
-{
-	FMobileSceneRenderer& ThisRenderer;
-	FRHICommandList& RHICmdList;
-	const FViewInfo& View;
-	FDrawingPolicyRenderState DrawRenderState;
-	int32 FirstElement;
-	int32 AfterLastElement;
-	EBlendMode BlendMode;
-	bool bWireFrame;
-
-public:
-
-	FRenderMobileBasePassDynamicDataThreadTask(
-		FMobileSceneRenderer& InThisRenderer,
-		FRHICommandList& InRHICmdList,
-		const FViewInfo& InView,
-		const FDrawingPolicyRenderState& InDrawRenderState,
-		EBlendMode InBlendMode,
-		bool bInWireFrame,
-		int32 InFirstElement,
-		int32 InAfterLastElement
-
-	)
-		: ThisRenderer(InThisRenderer)
-		, RHICmdList(InRHICmdList)
-		, View(InView)
-		, DrawRenderState(InDrawRenderState)
-		, FirstElement(InFirstElement)
-		, AfterLastElement(InAfterLastElement)
-		, BlendMode(InBlendMode)
-		, bWireFrame(bInWireFrame)
-	{
-		check(FirstElement < AfterLastElement); // don't create useless tasks
-	}
-
-	FORCEINLINE TStatId GetStatId() const
-	{
-		RETURN_QUICK_DECLARE_CYCLE_STAT(FRenderMobileBasePassDynamicDataThreadTask, STATGROUP_TaskGraphTasks);
-	}
-
-	static ESubsequentsMode::Type GetSubsequentsMode() { return ESubsequentsMode::TrackSubsequents; }
-
-	void DoTask(ENamedThreads::Type CurrentThread, const FGraphEventRef& MyCompletionGraphEvent)
-	{
-		ThisRenderer.RenderMobileBasePassDynamicData(RHICmdList, View, DrawRenderState, BlendMode, bWireFrame, FirstElement, AfterLastElement);
-		RHICmdList.HandleRTThreadTaskCompletion(MyCompletionGraphEvent);
-	}
-};
-
-static void RenderMobileBasePassDynamicDataParallel(FMobileSceneRenderer& ThisRenderer, FMobileBasePassParallelCommandListSet& ParallelSet, EBlendMode BlendMode, bool bWireframe)
-{
-	if (ParallelSet.View.DynamicMeshElements.Num() > 0)
-	{
-		int32 NumExpectedPrimitives = ParallelSet.View.DynamicMeshElements.Num() / 2; // opaque and masked rendered separately 
-		int32 EffectiveThreads = FMath::Min<int32>(NumExpectedPrimitives, ParallelSet.Width);
-
-		int32 NumPer = ParallelSet.View.DynamicMeshElements.Num() / EffectiveThreads;
-		int32 Extra = ParallelSet.View.DynamicMeshElements.Num() - NumPer * EffectiveThreads;
-		int32 Start = 0;
-		for (int32 ThreadIndex = 0; ThreadIndex < EffectiveThreads; ThreadIndex++)
-		{
-			int32 Last = Start + (NumPer - 1) + (ThreadIndex < Extra);
-			check(Last >= Start);
-
-			{
-				FRHICommandList* CmdList = ParallelSet.NewParallelCommandList();
-				FGraphEventRef AnyThreadCompletionEvent = TGraphTask<FRenderMobileBasePassDynamicDataThreadTask>::CreateTask(ParallelSet.GetPrereqs(), ENamedThreads::ActualRenderingThread)
-					.ConstructAndDispatchWhenReady(ThisRenderer, *CmdList, ParallelSet.View, ParallelSet.DrawRenderState, BlendMode, bWireframe, Start, Last + 1);
-				ParallelSet.AddParallelCommandList(CmdList, AnyThreadCompletionEvent);
-			}
-
-			Start = Last + 1;
-		}
-		check(Start == ParallelSet.View.DynamicMeshElements.Num());
-	}
-}
-
-void FMobileSceneRenderer::RenderMobileBasePassViewParallel(const FViewInfo& View, FRHICommandListImmediate& ParentCmdList, TArray<FViewInfo>& InViews, const FDrawingPolicyRenderState& DrawRenderState)
-{
-	QUICK_SCOPE_CYCLE_COUNTER(STAT_RenderMobileBasePassViewParallel);
-
-	bool bCreateContexts = false; // CVarRHICmdFlushRenderThreadTasksBasePass.GetValueOnRenderThread() == 0 && CVarRHICmdFlushRenderThreadTasks.GetValueOnRenderThread() == 0;
-	const bool bIsCSM = View.MobileCSMVisibilityInfo.bMobileDynamicCSMInUse;
-	check(!View.bIsMobileMultiViewEnabled && GetSortMode() != EBasePassSort::SortPerMesh); // easy to support, just isn't supported yet
-
-	{
-		FMobileBasePassParallelCommandListSet ParallelSet(View, this, ParentCmdList,
-			false, // no support for deferred contexts yet
-			bCreateContexts,
-			ViewFamily,
-			DrawRenderState
-			);
-
-		if (bIsCSM)
-		{
-			Scene->MobileBasePassUniformLightMapPolicyDrawListWithCSM[EBasePass_Default].DrawVisibleParallel(View.MobileCSMVisibilityInfo.MobileCSMStaticMeshVisibilityMap, View.MobileCSMVisibilityInfo.MobileCSMStaticBatchVisibility, ParallelSet);
-			Scene->MobileBasePassUniformLightMapPolicyDrawList[EBasePass_Default].DrawVisibleParallel(View.MobileCSMVisibilityInfo.MobileNonCSMStaticMeshVisibilityMap, View.MobileCSMVisibilityInfo.MobileNonCSMStaticBatchVisibility, ParallelSet);
-		}
-		else
-		{
-			Scene->MobileBasePassUniformLightMapPolicyDrawList[EBasePass_Default].DrawVisibleParallel(View.StaticMeshVisibilityMap, View.StaticMeshBatchVisibility, ParallelSet);
-		}
-		
-		const bool bWireframe = !!ViewFamily.EngineShowFlags.Wireframe;
-		RenderMobileBasePassDynamicDataParallel(*this, ParallelSet, BLEND_Opaque, bWireframe);
-	}
-
-	RenderMobileEditorPrimitives(ParentCmdList, View, DrawRenderState);
-
-	{
-		FMobileBasePassParallelCommandListSet ParallelSet(View, this, ParentCmdList,
-			false, // no support for deferred contexts yet
-			bCreateContexts,
-			ViewFamily,
-			DrawRenderState
-		);
-
-		if (bIsCSM)
-		{
-			Scene->MobileBasePassUniformLightMapPolicyDrawListWithCSM[EBasePass_Masked].DrawVisibleParallel(View.MobileCSMVisibilityInfo.MobileCSMStaticMeshVisibilityMap, View.MobileCSMVisibilityInfo.MobileCSMStaticBatchVisibility, ParallelSet);
-			Scene->MobileBasePassUniformLightMapPolicyDrawList[EBasePass_Masked].DrawVisibleParallel(View.MobileCSMVisibilityInfo.MobileNonCSMStaticMeshVisibilityMap, View.MobileCSMVisibilityInfo.MobileNonCSMStaticBatchVisibility, ParallelSet);
-		}
-		else
-		{
-			Scene->MobileBasePassUniformLightMapPolicyDrawList[EBasePass_Masked].DrawVisibleParallel(View.StaticMeshVisibilityMap, View.StaticMeshBatchVisibility, ParallelSet);
-		}
-
-		const bool bWireframe = !!ViewFamily.EngineShowFlags.Wireframe;
-		if (!bWireframe)
-		{
-			RenderMobileBasePassDynamicDataParallel(*this, ParallelSet, BLEND_Masked, false);
-		}
-	}
-
-}
-
-
-void FMobileSceneRenderer::RenderMobileBasePass(FRHICommandListImmediate& RHICmdList, const TArrayView<const FViewInfo*> PassViews)
-{
-	SCOPED_DRAW_EVENT(RHICmdList, MobileBasePass);
-	SCOPE_CYCLE_COUNTER(STAT_BasePassDrawTime);
-
-	if (UseMeshDrawCommandPipeline())
-	{
-		for (int32 ViewIndex = 0; ViewIndex < PassViews.Num(); ViewIndex++)
-		{
-			SCOPED_CONDITIONAL_DRAW_EVENTF(RHICmdList, EventView, Views.Num() > 1, TEXT("View%d"), ViewIndex);
-			const FViewInfo& View = *PassViews[ViewIndex];
-			if (!View.ShouldRenderView())
-			{
-				continue;
-			}
-
-			if (Scene->UniformBuffers.UpdateViewUniformBuffer(View))
-			{
-				UpdateOpaqueBasePassUniformBuffer(RHICmdList, View);
-				UpdateDirectionalLightUniformBuffers(RHICmdList, View);
-			}
-		
-			RHICmdList.SetViewport(View.ViewRect.Min.X, View.ViewRect.Min.Y, 0, View.ViewRect.Max.X, View.ViewRect.Max.Y, 1);
-			View.ParallelMeshDrawCommandPasses[EMeshPass::BasePass].DispatchDraw(nullptr, RHICmdList);
-
-			{
-				FDrawingPolicyRenderState DrawRenderState(View, Scene->UniformBuffers.MobileOpaqueBasePassUniformBuffer);
-				DrawRenderState.SetBlendState(TStaticBlendStateWriteMask<CW_RGBA>::GetRHI());
-				DrawRenderState.SetDepthStencilAccess(GetDefaultBasePassDepthStencilAccess(EShadingPath::Mobile));
-				DrawRenderState.SetDepthStencilState(TStaticDepthStencilState<true, CF_DepthNearOrEqual>::GetRHI());
-				RenderMobileEditorPrimitives(RHICmdList, View, DrawRenderState);
-			}
-		}
-		return;
-	}
-	
-	EBasePassSort::Type SortMode = GetSortMode();
-#if UE_BUILD_DEVELOPMENT
-	if (SortMode == EBasePassSort::SortPerMesh)
-	{
-		FDrawListSortKey Test1;
-		FDrawListSortKey Test2;
-		FDrawListSortKey Test3;
-
-		ZeroDrawListSortKey(Test1);
-		ZeroDrawListSortKey(Test2);
-		ZeroDrawListSortKey(Test3);
-		Test1.Fields.bBackground = 1;
-		Test2.Fields.MeshElementIndex = 1;
-		Test3.Fields.DepthBits = 1;
-
-		UE_CLOG(Test1 < Test2 || Test3 < Test2, LogRHI, Fatal, TEXT("FDrawListSortKey is using non-portable code that doesn't work"));
-	}
-#endif	
-
-	int32 MaxDraws = GMaxBasePassDraws.GetValueOnRenderThread();
-	if (MaxDraws <= 0)
-	{
-		MaxDraws = MAX_int32;
-	}
-
-	if (SortMode == EBasePassSort::SortStateBuckets)
-	{
-		SCOPE_CYCLE_COUNTER(STAT_SortStaticDrawLists);
-
-		const FVector ViewLocation = Views[0].ViewMatrices.GetViewOrigin();
-		
-		for (int32 DrawType = 0; DrawType < EBasePass_MAX; DrawType++)
-		{
-			Scene->MobileBasePassUniformLightMapPolicyDrawList[DrawType].SortFrontToBack(ViewLocation);
-			Scene->MobileBasePassUniformLightMapPolicyDrawListWithCSM[DrawType].SortFrontToBack(ViewLocation);
-		}
-	}
-
-	if (MaxDraws == MAX_int32 && PassViews.Num() && !Views[0].bIsMobileMultiViewEnabled && // we don't support parallel multiview...it would not be hard to add
-		SortMode != EBasePassSort::SortPerMesh && // we don't support sorting...not sure how hard it would be to add
-		GRHICommandList.UseParallelAlgorithms() && CVarMobileParallelBasePass.GetValueOnRenderThread())
-	{
-		bool bFlush = true; // CVarRHICmdFlushRenderThreadTasksBasePass.GetValueOnRenderThread() > 0 || CVarRHICmdFlushRenderThreadTasks.GetValueOnRenderThread() > 0;
-		FScopedCommandListWaitForTasks Flusher(bFlush, RHICmdList); 
-		for (int32 ViewIndex = 0; ViewIndex < Views.Num(); ViewIndex++)
-		{
-			SCOPED_CONDITIONAL_DRAW_EVENTF(RHICmdList, EventView, Views.Num() > 1, TEXT("View%d"), ViewIndex);
-			FViewInfo& View = Views[ViewIndex];
-			if (!View.ShouldRenderView())
-			{
-				continue;
-			}
+		View.BatchedViewElements.Draw(RHICmdList, DrawRenderState, FeatureLevel, bNeedToSwitchVerticalAxis, View, false);
 
-			TUniformBufferRef<FMobileBasePassUniformParameters> BasePassUniformBuffer;
-			CreateMobileBasePassUniformBuffer(RHICmdList, View, false, BasePassUniformBuffer);
-			FDrawingPolicyRenderState DrawRenderState(View, BasePassUniformBuffer);
-			SetupMobileBasePassView(RHICmdList, View, DrawRenderState);
-			RenderMobileBasePassViewParallel(View, RHICmdList, Views, DrawRenderState);
-		}
-	}
-	else
-	{
-		// Draw the scene's emissive and light-map color.
-		for (int32 ViewIndex = 0; ViewIndex < PassViews.Num(); ViewIndex++)
+		DrawDynamicMeshPass(View, RHICmdList,
+			[&View, &DrawRenderState](FDynamicPassMeshDrawListContext* DynamicMeshPassContext)
 		{
-			SCOPED_CONDITIONAL_DRAW_EVENTF(RHICmdList, EventView, Views.Num() > 1, TEXT("View%d"), ViewIndex);
-			const FViewInfo& View = *PassViews[ViewIndex];
-
-			if (!View.ShouldRenderView())
-			{
-				continue;
-			}
-
-			TUniformBufferRef<FMobileBasePassUniformParameters> BasePassUniformBuffer;
-			CreateMobileBasePassUniformBuffer(RHICmdList, View, false, BasePassUniformBuffer);
-			FDrawingPolicyRenderState DrawRenderState(View, BasePassUniformBuffer);
-
-			SetupMobileBasePassView(RHICmdList, View, DrawRenderState);
-
-			// Render the base pass static data
-			if (SortMode == EBasePassSort::SortPerMesh)
-			{
-				DrawVisibleFrontToBack(RHICmdList, Scene, EBasePass_Default, View, DrawRenderState, View.MobileCSMVisibilityInfo, MaxDraws);
-			}
-			else
-			{
-				DrawVisible(RHICmdList, Scene, EBasePass_Default, View, DrawRenderState, View.MobileCSMVisibilityInfo);
-			}
-			FRHICommandListExecutor::GetImmediateCommandList().PollOcclusionQueries();
-			FRHICommandListExecutor::GetImmediateCommandList().ImmediateFlush(EImmediateFlushType::DispatchToRHIThread);
-
-			// render dynamic opaque primitives (or all if Wireframe)
-			const bool bWireframe = !!ViewFamily.EngineShowFlags.Wireframe;
-			RenderMobileBasePassDynamicData(RHICmdList, View, DrawRenderState, BLEND_Opaque, bWireframe);
-
-			FRHICommandListExecutor::GetImmediateCommandList().PollOcclusionQueries();
-			FRHICommandListExecutor::GetImmediateCommandList().ImmediateFlush(EImmediateFlushType::DispatchToRHIThread);
+			FEditorPrimitivesBasePassMeshProcessor PassMeshProcessor(
+				View.Family->Scene->GetRenderScene(),
+				View.GetFeatureLevel(),
+				&View,
+				DrawRenderState,
+				false,
+				DynamicMeshPassContext);
 
-			RenderMobileEditorPrimitives(RHICmdList, View, DrawRenderState);
+			const uint64 DefaultBatchElementMask = ~0ull;
 
-			// Issue static draw list masked draw calls last, as PVR wants it
-			if (SortMode == EBasePassSort::SortPerMesh)
+			for (int32 MeshIndex = 0; MeshIndex < View.TopViewMeshElements.Num(); MeshIndex++)
 			{
-				DrawVisibleFrontToBack(RHICmdList, Scene, EBasePass_Masked, View, DrawRenderState, View.MobileCSMVisibilityInfo, MaxDraws);
+				const FMeshBatch& MeshBatch = View.TopViewMeshElements[MeshIndex];
+				PassMeshProcessor.AddMeshBatch(MeshBatch, DefaultBatchElementMask, nullptr);
 			}
-			else
-			{
-				DrawVisible(RHICmdList, Scene, EBasePass_Masked, View, DrawRenderState, View.MobileCSMVisibilityInfo);
-			}
-
-			FRHICommandListExecutor::GetImmediateCommandList().PollOcclusionQueries();
-			FRHICommandListExecutor::GetImmediateCommandList().ImmediateFlush(EImmediateFlushType::DispatchToRHIThread);
+		});
 
-			// render dynamic masked primitives (or none if Wireframe)
-			if (!bWireframe)
-			{
-				RenderMobileBasePassDynamicData(RHICmdList, View, DrawRenderState, BLEND_Masked, false);
-				FRHICommandListExecutor::GetImmediateCommandList().PollOcclusionQueries();
-				FRHICommandListExecutor::GetImmediateCommandList().ImmediateFlush(EImmediateFlushType::DispatchToRHIThread);
-			}
-		}
+		// Draw the view's batched simple elements(lines, sprites, etc).
+		View.TopBatchedViewElements.Draw(RHICmdList, DrawRenderState, FeatureLevel, bNeedToSwitchVerticalAxis, View, false);
 	}
 }
diff --git a/Engine/Source/Runtime/Renderer/Private/MobileBasePassRendering.h b/Engine/Source/Runtime/Renderer/Private/MobileBasePassRendering.h
--- a/Engine/Source/Runtime/Renderer/Private/MobileBasePassRendering.h
+++ b/Engine/Source/Runtime/Renderer/Private/MobileBasePassRendering.h
@@ -61,14 +61,6 @@ enum EOutputFormat
 
 #define MAX_BASEPASS_DYNAMIC_POINT_LIGHTS 4
 
-namespace MobileBasePass
-{
-	// If no reflection captures are available then attempt to use sky light's texture.
-	bool UseSkyReflectionCapture(const FScene* RenderScene);
-	void GetSkyTextureParams(const FScene* Scene, float& AverageBrightnessOUT, FTexture*& ReflectionTextureOUT, float& OutSkyMaxMipIndex);
-};
-
-
 /* Info for dynamic point lights rendered in base pass */
 class FMobileBasePassMovablePointLightInfo
 {
@@ -134,32 +126,6 @@ public:
 		return bShaderHasOutdatedParameters;
 	}
 
-	void SetParameters(
-		FRHICommandList& RHICmdList,
-		const FViewInfo* View,
-		const FDrawingPolicyRenderState& DrawRenderState)
-	{
-
-		FMaterialShader::SetViewParameters(RHICmdList, GetVertexShader(), *View, DrawRenderState.GetViewUniformBuffer());
-		FMeshMaterialShader::SetPassUniformBuffer(RHICmdList, GetVertexShader(), DrawRenderState.GetPassUniformBuffer());
-	}
-
-	// Set parameters specific to mesh
-	void SetMesh(
-		FRHICommandList& RHICmdList,
-		const FMaterial& InMaterialResource,
-		const FSceneView& View,
-		const FVertexFactory* InVertexFactory, 
-		const FMaterialRenderProxy* InMaterialRenderProxy,
-		const FPrimitiveSceneProxy* Proxy,
-		const FMeshBatchElement& BatchElement, 
-		const FDrawingPolicyRenderState& DrawRenderState)
-	{
-		FMaterialShader::SetParametersInner(RHICmdList, GetVertexShader(), InMaterialRenderProxy, InMaterialResource, View);
-		uint32 DataFlags = 0;
-		FMeshMaterialShader::SetMesh(RHICmdList, GetVertexShader(), InVertexFactory, View, Proxy, BatchElement, DrawRenderState, DataFlags);
-	}
-
 	void GetShaderBindings(
 		const FScene* Scene,
 		ERHIFeatureLevel::Type FeatureLevel,
@@ -321,143 +287,6 @@ private:
 	FShaderParameter CSMDebugHintParams;
 
 public:
-	// Set parameters specific to PSO
-	void SetParameters(
-		FRHICommandList& RHICmdList, 
-		const FViewInfo* View, 
-		const FDrawingPolicyRenderState& DrawRenderState)
-	{
-		// If we're using only the sky for reflection then set it once here.
-		FScene* RenderScene = View->Family->Scene->GetRenderScene();
-		if (MobileBasePass::UseSkyReflectionCapture(RenderScene))
-		{
-			// MobileReflectionValues.x == max sky cube mip.
-			// if >0 this will disable shader's RGBM decoding and enable sky light tinting of this envmap.
-			FTexture* ReflectionTexture = GBlackTextureCube;
-			float AverageBrightness = 1.0f;
-			FVector4 MobileSkyReflectionValues(ForceInit);
-			if (View->GetFeatureLevel() > ERHIFeatureLevel::ES2) // not-supported on ES2 at the moment
-			{
-				MobileBasePass::GetSkyTextureParams(RenderScene, AverageBrightness, ReflectionTexture, MobileSkyReflectionValues.W);
-			}
-			MobileSkyReflectionValues.X = 1.0f / AverageBrightness;
-
-			FRHIPixelShader* PixelShader = GetPixelShader();
-			// Set the reflection cubemap
-			SetTextureParameter(RHICmdList, PixelShader, ReflectionCubemap, ReflectionSampler, ReflectionTexture);
-			SetShaderValue(RHICmdList, PixelShader, MobileReflectionParams, MobileSkyReflectionValues);
-		}
-
-		FMaterialShader::SetViewParameters(RHICmdList, GetPixelShader(), *View, DrawRenderState.GetViewUniformBuffer());
-		FMeshMaterialShader::SetPassUniformBuffer(RHICmdList, GetPixelShader(), DrawRenderState.GetPassUniformBuffer());
-	}
-
-	// Set parameters specific to Mesh
-	void SetMesh(
-		FRHICommandList& RHICmdList, 
-		const FMaterial& InMaterialResource,
-		const FSceneView& View,
-		const FVertexFactory* InVertexFactory, 
-		const FMaterialRenderProxy* InMaterialRenderProxy,
-		const FPrimitiveSceneProxy* Proxy,
-		const FMeshBatchElement& BatchElement, 
-		const FDrawingPolicyRenderState& DrawRenderState,
-		int32 NumMovablePointLights)
-	{
-		FRHIPixelShader* PixelShader = GetPixelShader();
-		FPrimitiveSceneInfo* PrimitiveSceneInfo = Proxy ? Proxy->GetPrimitiveSceneInfo() : NULL;
-		// test for HQ reflection parameter existence
-		if (ReflectionCubemap1.IsBound() || ReflectionCubemap2.IsBound() || ReflectionPositionsAndRadii.IsBound())
-		{
-			static const int32 MaxNumReflections = FPrimitiveSceneInfo::MaxCachedReflectionCaptureProxies;
-			static_assert(MaxNumReflections == 3, "Update reflection array initializations to match MaxCachedReflectionCaptureProxies");
-
-			// set high quality reflection parameters.
-			const FShaderResourceParameter* ReflectionTextureParameters[MaxNumReflections] = { &ReflectionCubemap, &ReflectionCubemap1, &ReflectionCubemap2 };
-			const FShaderResourceParameter* ReflectionSamplerParameters[MaxNumReflections] = { &ReflectionSampler, &ReflectionSampler1, &ReflectionSampler2 };
-			FTexture* ReflectionCubemapTextures[MaxNumReflections] = { GBlackTextureCube, GBlackTextureCube, GBlackTextureCube };
-			FVector4 CapturePositions[MaxNumReflections] = { FVector4(0, 0, 0, 0), FVector4(0, 0, 0, 0), FVector4(0, 0, 0, 0) };
-			FVector AverageBrightness(1, 1, 1);
-
-			if (PrimitiveSceneInfo)
-			{
-				for (int32 i = 0; i < MaxNumReflections; i++)
-				{
-					const FReflectionCaptureProxy* ReflectionProxy = PrimitiveSceneInfo->CachedReflectionCaptureProxies[i];
-					if (ReflectionProxy)
-					{
-						CapturePositions[i] = ReflectionProxy->Position;
-						CapturePositions[i].W = ReflectionProxy->InfluenceRadius;
-						if (ReflectionProxy->EncodedHDRCubemap && ReflectionProxy->EncodedHDRCubemap->IsInitialized())
-						{
-							ReflectionCubemapTextures[i] = PrimitiveSceneInfo->CachedReflectionCaptureProxies[i]->EncodedHDRCubemap;
-						}
-						AverageBrightness[i] = ReflectionProxy->EncodedHDRAverageBrightness;
-					}
-				}
-			}
-
-			for (int32 i = 0; i < MaxNumReflections; i++)
-			{
-				if (ReflectionTextureParameters[i]->IsBound())
-				{
-					SetTextureParameter(RHICmdList, PixelShader, *ReflectionTextureParameters[i], *ReflectionSamplerParameters[i], ReflectionCubemapTextures[i]);
-				}
-			}
-
-			if (ReflectionPositionsAndRadii.IsBound())
-			{
-				SetShaderValueArray(RHICmdList, PixelShader, ReflectionPositionsAndRadii, CapturePositions, MaxNumReflections);
-			}
-
-			SetShaderValue(RHICmdList, PixelShader, MobileReflectionParams, FVector4(1.0f / AverageBrightness.X, 1.0f / AverageBrightness.Y, 1.0f / AverageBrightness.Z, 0));
-		}
-		else if (ReflectionCubemap.IsBound() && (!PrimitiveSceneInfo || !MobileBasePass::UseSkyReflectionCapture(PrimitiveSceneInfo->Scene)))
-		{
-			FTexture* ReflectionTexture = GBlackTextureCube;
-			float AverageBrightness = 1.0f;
-
-			if (PrimitiveSceneInfo 
-				&& PrimitiveSceneInfo->CachedReflectionCaptureProxy
-				&& PrimitiveSceneInfo->CachedReflectionCaptureProxy->EncodedHDRCubemap
-				&& PrimitiveSceneInfo->CachedReflectionCaptureProxy->EncodedHDRCubemap->IsInitialized())
-			{
-				AverageBrightness = PrimitiveSceneInfo->CachedReflectionCaptureProxy->EncodedHDRAverageBrightness;
-				ReflectionTexture = PrimitiveSceneInfo->CachedReflectionCaptureProxy->EncodedHDRCubemap;
-			}
-
-			// Set the reflection cubemap
-			SetTextureParameter(RHICmdList, PixelShader, ReflectionCubemap, ReflectionSampler, ReflectionTexture);
-			SetShaderValue(RHICmdList, PixelShader, MobileReflectionParams, FVector4(1.0f / AverageBrightness, 0, 0, 0));
-		}
-
-		if (NumMovablePointLights > 0)
-		{
-			FMobileBasePassMovablePointLightInfo LightInfo(Proxy);
-
-			if (NumMovablePointLights == INT32_MAX)
-			{
-				SetShaderValue(RHICmdList, PixelShader, NumDynamicPointLightsParameter, LightInfo.NumMovablePointLights);
-			}
-
-			// Set dynamic point lights
-			SetShaderValueArray(RHICmdList, PixelShader, LightPositionAndInvRadiusParameter, LightInfo.LightPositionAndInvRadius, LightInfo.NumMovablePointLights);
-			SetShaderValueArray(RHICmdList, PixelShader, LightColorAndFalloffExponentParameter, LightInfo.LightColorAndFalloffExponent, LightInfo.NumMovablePointLights);
-		}
-
-		if (CSMDebugHintParams.IsBound())
-		{
-			static const auto CVarsCSMDebugHint = IConsoleManager::Get().FindTConsoleVariableDataFloat(TEXT("r.Mobile.Shadow.CSMDebugHint"));
-			
-			float CSMDebugValue = CVarsCSMDebugHint->GetValueOnRenderThread();
-			SetShaderValue(RHICmdList, PixelShader, CSMDebugHintParams, CSMDebugValue);
-		}
-
-		FMaterialShader::SetParametersInner(RHICmdList, PixelShader, InMaterialRenderProxy, InMaterialResource, View);
-		uint32 DataFlags = 0;
-		FMeshMaterialShader::SetMesh(RHICmdList, PixelShader, InVertexFactory, View, Proxy, BatchElement, DrawRenderState, DataFlags);
-	}
-
 	void GetShaderBindings(
 		const FScene* Scene,
 		ERHIFeatureLevel::Type FeatureLevel,
@@ -601,332 +430,6 @@ public:
 	TMobileBasePassPS() {}
 };
 
-/**
- * Draws the emissive color and the light-map of a mesh.
- */
-class FMobileBasePassUniformDrawingPolicy : public FMeshDrawingPolicy
-{
-public:
-	/** The data the drawing policy uses for each mesh element. */
-	class ElementDataType
-	{
-	public:
-		/** The element's light-map data. */
-		typename FUniformLightMapPolicy::ElementDataType LightMapElementData;
-
-		/** Default constructor. */
-		ElementDataType()
-		{}
-
-		/** Initialization constructor. */
-		ElementDataType(const FUniformLightMapPolicy::ElementDataType& InLightMapElementData)
-		:	LightMapElementData(InLightMapElementData)
-		{}
-	};
-
-	/** Initialization constructor. */
-	FMobileBasePassUniformDrawingPolicy(
-		const FVertexFactory* InVertexFactory,
-		const FMaterialRenderProxy* InMaterialRenderProxy,
-		const FMaterial& InMaterialResource,
-		ELightMapPolicyType InLightMapPolicyType,
-		int32 InNumMovablePointLights,
-		EBlendMode InBlendMode,
-		bool bInEnableSkyLight,
-		const FMeshDrawingPolicyOverrideSettings& InOverrideSettings,
-		ERHIFeatureLevel::Type FeatureLevel,
-		bool bInEnableReceiveDecalOutput = false
-		):
-		FMeshDrawingPolicy(nullptr, nullptr, InMaterialResource, InOverrideSettings),
-		VertexDeclaration(InVertexFactory->GetDeclaration()),
-		UniformLightMapPolicy(InLightMapPolicyType),
-		NumMovablePointLights(InNumMovablePointLights),
-		BlendMode(InBlendMode),
-		bEnableReceiveDecalOutput(bInEnableReceiveDecalOutput)
-	{
-		static_assert(MAX_BASEPASS_DYNAMIC_POINT_LIGHTS == 4, "If you change MAX_BASEPASS_DYNAMIC_POINT_LIGHTS, you need to change the switch statement below");
-
-		if (InMaterialRenderProxy != nullptr)
-		{
-			ImmutableSamplerState = InMaterialRenderProxy->ImmutableSamplerState;
-		}
-				
-		MobileBasePass::GetShaders(InLightMapPolicyType, InNumMovablePointLights, InMaterialResource, InVertexFactory->GetType(), bInEnableSkyLight, VertexShader, PixelShader);
-		BaseVertexShader = VertexShader;
-	}
-
-	// FMeshDrawingPolicy interface.
-
-	FDrawingPolicyMatchResult Matches(const FMobileBasePassUniformDrawingPolicy& Other, bool bForReals = false) const
-	{
-		DRAWING_POLICY_MATCH_BEGIN
-			DRAWING_POLICY_MATCH(MaterialResource == Other.MaterialResource) &&
-			DRAWING_POLICY_MATCH(VertexDeclaration == Other.VertexDeclaration) &&
-			DRAWING_POLICY_MATCH(bIsDitheredLODTransitionMaterial == Other.bIsDitheredLODTransitionMaterial) &&
-			DRAWING_POLICY_MATCH(bUsePositionOnlyVS == Other.bUsePositionOnlyVS) &&
-			DRAWING_POLICY_MATCH(MeshFillMode == Other.MeshFillMode) &&
-			DRAWING_POLICY_MATCH(MeshCullMode == Other.MeshCullMode) &&
-			DRAWING_POLICY_MATCH(MeshPrimitiveType == Other.MeshPrimitiveType) &&
-			DRAWING_POLICY_MATCH(VertexShader == Other.VertexShader) &&
-			DRAWING_POLICY_MATCH(PixelShader == Other.PixelShader) &&
-			DRAWING_POLICY_MATCH(UniformLightMapPolicy.GetIndirectPolicy() == Other.UniformLightMapPolicy.GetIndirectPolicy()) &&
-			DRAWING_POLICY_MATCH(NumMovablePointLights == Other.NumMovablePointLights) &&
-			DRAWING_POLICY_MATCH(bEnableReceiveDecalOutput == Other.bEnableReceiveDecalOutput) &&
-			DRAWING_POLICY_MATCH(ImmutableSamplerState == Other.ImmutableSamplerState);
-		DRAWING_POLICY_MATCH_END 
-	}
-
-	uint32 GetTypeHash() const
-	{
-		return PointerHash(VertexDeclaration, PointerHash(MaterialResource));
-	}
-
-	const FMaterialRenderProxy* GetPipelineMaterialRenderProxy(const FMaterialRenderProxy* ElementMaterialRenderProxy)
-	{
-		return ElementMaterialRenderProxy;
-	}
-
-	friend int32 CompareDrawingPolicy(const FMobileBasePassUniformDrawingPolicy& A, const FMobileBasePassUniformDrawingPolicy& B)
-	{
-		COMPAREDRAWINGPOLICYMEMBERS(MaterialResource);
-		COMPAREDRAWINGPOLICYMEMBERS(NumMovablePointLights);
-		return CompareDrawingPolicy(A.UniformLightMapPolicy, B.UniformLightMapPolicy);
-	}
-
-	void SetupPipelineState(FDrawingPolicyRenderState& DrawRenderState, const FViewInfo& View) const
-	{
-		bool bEncodedHDR = GetMobileHDRMode() == EMobileHDRMode::EnabledRGBE && MaterialResource->GetMaterialDomain() != MD_UI;;
-
-		static const auto CVarMonoscopicFarField = IConsoleManager::Get().FindTConsoleVariableDataInt(TEXT("vr.MonoscopicFarField"));
-		const bool bIsMobileMonoscopic = CVarMonoscopicFarField && (CVarMonoscopicFarField->GetValueOnRenderThread() != 0);
-
-		if (bEncodedHDR == false)
-		{
-			switch (BlendMode)
-			{
-			default:
-			case BLEND_Opaque:
-				// Opaque materials are rendered together in the base pass, where the blend state is set at a higher level
-				break;
-			case BLEND_Masked:
-				// Masked materials are rendered together in the base pass, where the blend state is set at a higher level
-				break;
-			case BLEND_Translucent:
-				if (MaterialResource->ShouldWriteOnlyAlpha())
-				{
-					DrawRenderState.SetBlendState(TStaticBlendState<CW_ALPHA, BO_Add, BF_Zero, BF_Zero, BO_Add, BF_One, BF_Zero>::GetRHI());
-				} 
-				else if (bIsMobileMonoscopic)
-				{
-					DrawRenderState.SetBlendState(TStaticBlendState<CW_RGBA, BO_Add, BF_SourceAlpha, BF_InverseSourceAlpha, BO_Add, BF_One, BF_One>::GetRHI());
-				}
-				else
-				{
-					DrawRenderState.SetBlendState(TStaticBlendState<CW_RGB, BO_Add, BF_SourceAlpha, BF_InverseSourceAlpha, BO_Add, BF_Zero, BF_InverseSourceAlpha>::GetRHI());
-				}
-				break;
-			case BLEND_Additive:
-				// Add to the existing scene color
-				DrawRenderState.SetBlendState(TStaticBlendState<CW_RGB, BO_Add, BF_One, BF_One, BO_Add, BF_Zero, BF_InverseSourceAlpha>::GetRHI());
-				break;
-			case BLEND_Modulate:
-				// Modulate with the existing scene color
-				DrawRenderState.SetBlendState(TStaticBlendState<CW_RGB, BO_Add, BF_DestColor, BF_Zero>::GetRHI());
-				break;
-			case BLEND_AlphaComposite:
-				// Blend with existing scene color. New color is already pre-multiplied by alpha.
-				DrawRenderState.SetBlendState(TStaticBlendState<CW_RGB, BO_Add, BF_One, BF_InverseSourceAlpha, BO_Add, BF_Zero, BF_InverseSourceAlpha>::GetRHI());
-				break;
-			};
-		}
-		else
-		{
-			DrawRenderState.SetBlendState(TStaticBlendState<>::GetRHI());
-		}
-
-		if (bEnableReceiveDecalOutput && View.bSceneHasDecals)
-		{
-			DrawRenderState.SetDepthStencilState(TStaticDepthStencilState<
-				true, CF_GreaterEqual,
-				true, CF_Always, SO_Keep, SO_Keep, SO_Replace,
-				false, CF_Always, SO_Keep, SO_Keep, SO_Keep,
-				0x00, GET_STENCIL_BIT_MASK(RECEIVE_DECAL, 1)>::GetRHI());
-		}
-	}
-
-	void SetSharedState(FRHICommandList& RHICmdList, const FDrawingPolicyRenderState& DrawRenderState, const FViewInfo* View, const ContextDataType PolicyContext) const
-	{
-		VertexShader->SetParameters(RHICmdList, View, DrawRenderState);
-		PixelShader->SetParameters(RHICmdList, View, DrawRenderState);
-	}
-
-	/** 
-	* Create bound shader state using the vertex decl from the mesh draw policy
-	* as well as the shaders needed to draw the mesh
-	* @return new bound shader state object
-	*/
-	FBoundShaderStateInput GetBoundShaderStateInput(ERHIFeatureLevel::Type InFeatureLevel) const
-	{
-		return FBoundShaderStateInput(
-			VertexDeclaration, 
-			VertexShader->GetVertexShader(),
-			FHullShaderRHIRef(), 
-			FDomainShaderRHIRef(), 
-			PixelShader->GetPixelShader(),
-			FGeometryShaderRHIRef());
-	}
-
-	void SetMeshRenderState(
-		FRHICommandList& RHICmdList, 
-		const FViewInfo& View,
-		const FPrimitiveSceneProxy* PrimitiveSceneProxy,
-		const FMeshBatch& Mesh,
-		int32 BatchElementIndex,
-		const FDrawingPolicyRenderState& DrawRenderState,
-		const ElementDataType& ElementData,
-		const ContextDataType PolicyContext
-		) const
-	{
-		const FMaterialRenderProxy* MeshMaterialRenderProxy = Mesh.MaterialRenderProxy;
-		const FVertexFactory* MeshVertexFactory = Mesh.VertexFactory;
-
-		check(MeshVertexFactory && MeshVertexFactory->IsInitialized());
-		MeshVertexFactory->SetStreams(View.FeatureLevel, RHICmdList);
-
-		// Set the light-map policy's mesh-specific settings.
-		UniformLightMapPolicy.SetMesh(
-			RHICmdList, 
-			View,
-			PrimitiveSceneProxy,
-			VertexShader,
-			PixelShader,
-			VertexShader,
-			PixelShader,
-			MeshVertexFactory,
-			MeshMaterialRenderProxy,
-			ElementData.LightMapElementData);
-
-		const FMeshBatchElement& BatchElement = Mesh.Elements[BatchElementIndex];
-		VertexShader->SetMesh(RHICmdList, *MaterialResource, View, MeshVertexFactory, MeshMaterialRenderProxy, PrimitiveSceneProxy, BatchElement, DrawRenderState);
-		PixelShader->SetMesh(RHICmdList, *MaterialResource, View, MeshVertexFactory, MeshMaterialRenderProxy, PrimitiveSceneProxy, BatchElement, DrawRenderState, NumMovablePointLights);
-
-		// Set directional light UB
-		const TShaderUniformBufferParameter<FMobileDirectionalLightShaderParameters>& MobileDirectionalLightParam = PixelShader->template GetUniformBufferParameter<FMobileDirectionalLightShaderParameters>();
-		if (MobileDirectionalLightParam.IsBound())
-		{
-			int32 UniformBufferIndex = PrimitiveSceneProxy ? GetFirstLightingChannelFromMask(PrimitiveSceneProxy->GetLightingChannelMask()) + 1 : 0;
-			SetUniformBufferParameter(RHICmdList, PixelShader->GetPixelShader(), MobileDirectionalLightParam, View.MobileDirectionalLightUniformBuffers[UniformBufferIndex]);
-		}
-
-		if (bEnableReceiveDecalOutput && View.bSceneHasDecals)
-		{
-			const uint8 StencilValue = (PrimitiveSceneProxy && !PrimitiveSceneProxy->ReceivesDecals() ? 0x01 : 0x00);
-			RHICmdList.SetStencilRef(GET_STENCIL_BIT_MASK(RECEIVE_DECAL, StencilValue)); // we hash the stencil group because we only have 6 bits.
-		}
-	}
-
-public:
-	TMobileBasePassVSPolicyParamType<FUniformLightMapPolicy>* VertexShader;
-	TMobileBasePassPSPolicyParamType<FUniformLightMapPolicy>* PixelShader;
-protected:
-	FVertexDeclarationRHIRef VertexDeclaration;
-	FUniformLightMapPolicy UniformLightMapPolicy;
-	int32 NumMovablePointLights;
-	FImmutableSamplerState ImmutableSamplerState;
-	EBlendMode BlendMode;
-	uint32 bEnableReceiveDecalOutput : 1;
-};
-
-/**
- * A drawing policy factory for the base pass drawing policy.
- */
-class FMobileBasePassOpaqueDrawingPolicyFactory
-{
-public:
-
-	enum { bAllowSimpleElements = true };
-	struct ContextType 
-	{
-		ContextType()
-		{}
-	};
-
-	static void AddStaticMesh(FRHICommandList& RHICmdList, FScene* Scene, FStaticMesh* StaticMesh);
-	static bool DrawDynamicMesh(
-		FRHICommandList& RHICmdList, 
-		const FViewInfo& View,
-		ContextType DrawingContext,
-		const FMeshBatch& Mesh,
-		bool bPreFog,
-		const FDrawingPolicyRenderState& DrawRenderState,
-		const FPrimitiveSceneProxy* PrimitiveSceneProxy,
-		FHitProxyId HitProxyId
-		);
-};
-
-
-/** The parameters used to process a mobile base pass mesh. */
-class FMobileProcessBasePassMeshParameters
-{
-public:
-	const FMeshBatch& Mesh;
-	const uint64 BatchElementMask;
-	const FMaterial& Material;
-	const FPrimitiveSceneProxy* PrimitiveSceneProxy;
-	int32 NumMovablePointLights;
-	EBlendMode BlendMode;
-	EMaterialShadingModel ShadingModel;
-	const bool bAllowFog;
-	ERHIFeatureLevel::Type FeatureLevel;
-	const bool bIsInstancedStereo;
-
-	/** Initialization constructor. */
-	FMobileProcessBasePassMeshParameters(
-		const FMeshBatch& InMesh,
-		const FMaterial& InMaterial,
-		const FPrimitiveSceneProxy* InPrimitiveSceneProxy,
-		bool InbAllowFog,
-		ERHIFeatureLevel::Type InFeatureLevel,
-		const bool InbIsInstancedStereo = false
-	) :
-		Mesh(InMesh),
-		BatchElementMask(Mesh.Elements.Num() == 1 ? 1 : (1 << Mesh.Elements.Num()) - 1), // 1 bit set for each mesh element
-		Material(InMaterial),
-		PrimitiveSceneProxy(InPrimitiveSceneProxy),
-		NumMovablePointLights(MobileBasePass::CalcNumMovablePointLights(InMaterial, InPrimitiveSceneProxy)),
-		BlendMode(InMaterial.GetBlendMode()),
-		ShadingModel(InMaterial.GetShadingModel()),
-		bAllowFog(InbAllowFog),
-		FeatureLevel(InFeatureLevel),
-		bIsInstancedStereo(InbIsInstancedStereo)
-	{
-	}
-
-	/** Initialization constructor. */
-	FMobileProcessBasePassMeshParameters(
-		const FMeshBatch& InMesh,
-		const uint64& InBatchElementMask,
-		const FMaterial& InMaterial,
-		const FPrimitiveSceneProxy* InPrimitiveSceneProxy,
-		bool InbAllowFog,
-		ERHIFeatureLevel::Type InFeatureLevel,
-		bool InbIsInstancedStereo = false
-	) :
-		Mesh(InMesh),
-		BatchElementMask(InBatchElementMask),
-		Material(InMaterial),
-		PrimitiveSceneProxy(InPrimitiveSceneProxy),
-		NumMovablePointLights(MobileBasePass::CalcNumMovablePointLights(InMaterial, InPrimitiveSceneProxy)),
-		BlendMode(InMaterial.GetBlendMode()),
-		ShadingModel(InMaterial.GetShadingModel()),
-		bAllowFog(InbAllowFog),
-		FeatureLevel(InFeatureLevel),
-		bIsInstancedStereo(InbIsInstancedStereo)
-	{
-	}
-};
-
-
 class FMobileBasePassMeshProcessor : public FMeshPassProcessor
 {
 public:
diff --git a/Engine/Source/Runtime/Renderer/Private/MobileDistortionPass.cpp b/Engine/Source/Runtime/Renderer/Private/MobileDistortionPass.cpp
--- a/Engine/Source/Runtime/Renderer/Private/MobileDistortionPass.cpp
+++ b/Engine/Source/Runtime/Renderer/Private/MobileDistortionPass.cpp
@@ -17,9 +17,7 @@ bool IsMobileDistortionActive(const FViewInfo& View)
 
 	// Distortion on mobile requires SceneDepth information in SceneColor.A channel
 	const EMobileHDRMode HDRMode = GetMobileHDRMode();
-
-	const bool bVisiblePrims = UseMeshDrawCommandPipeline() ?
-		View.ParallelMeshDrawCommandPasses[EMeshPass::Distortion].HasAnyDraw() : View.DistortionPrimSet.NumPrims() > 0;
+	const bool bVisiblePrims = View.ParallelMeshDrawCommandPasses[EMeshPass::Distortion].HasAnyDraw();
 
 	return
 		HDRMode == EMobileHDRMode::EnabledFloat16 &&
@@ -43,32 +41,14 @@ void FRCDistortionAccumulatePassES2::Process(FRenderingCompositePassContext& Con
 
 		Context.SetViewportAndCallRHI(View.ViewRect);
 
-		if (UseMeshDrawCommandPipeline())
-		{
-			if (Scene->UniformBuffers.UpdateViewUniformBuffer(View))
-			{
-				FMobileDistortionPassUniformParameters Parameters;
-				SetupMobileDistortionPassUniformBuffer(Context.RHICmdList, View, Parameters);
-				Scene->UniformBuffers.MobileDistortionPassUniformBuffer.UpdateUniformBufferImmediate(Parameters);
-			}
-
-			View.ParallelMeshDrawCommandPasses[EMeshPass::Distortion].DispatchDraw(nullptr, Context.RHICmdList);
-		}
-		else
+		if (Scene->UniformBuffers.UpdateViewUniformBuffer(View))
 		{
 			FMobileDistortionPassUniformParameters Parameters;
 			SetupMobileDistortionPassUniformBuffer(Context.RHICmdList, View, Parameters);
-			TUniformBufferRef<FMobileDistortionPassUniformParameters> PassUniformBuffer = TUniformBufferRef<FMobileDistortionPassUniformParameters>::CreateUniformBufferImmediate(Parameters, UniformBuffer_SingleFrame);
-
-			FDrawingPolicyRenderState DrawRenderState(View, PassUniformBuffer);
-			// We don't have depth, render all pixels, pixel shader will sample SceneDepth from SceneColor.A and discard if occluded
-			DrawRenderState.SetDepthStencilState(TStaticDepthStencilState<false, CF_Always>::GetRHI());
-			// additive blending of offsets
-			DrawRenderState.SetBlendState(TStaticBlendState<CW_RGBA, BO_Add, BF_One, BF_One, BO_Add, BF_One, BF_One>::GetRHI());
-
-			// draw only distortion meshes to accumulate their offsets
-			View.DistortionPrimSet.DrawAccumulatedOffsets(Context.RHICmdList, View, DrawRenderState, false);
+			Scene->UniformBuffers.MobileDistortionPassUniformBuffer.UpdateUniformBufferImmediate(Parameters);
 		}
+
+		View.ParallelMeshDrawCommandPasses[EMeshPass::Distortion].DispatchDraw(nullptr, Context.RHICmdList);
 	}
 	Context.RHICmdList.EndRenderPass();
 	Context.RHICmdList.CopyToResolveTarget(DestRenderTarget.TargetableTexture, DestRenderTarget.ShaderResourceTexture, FResolveParams());
diff --git a/Engine/Source/Runtime/Renderer/Private/MobileSeparateTranslucencyPass.cpp b/Engine/Source/Runtime/Renderer/Private/MobileSeparateTranslucencyPass.cpp
--- a/Engine/Source/Runtime/Renderer/Private/MobileSeparateTranslucencyPass.cpp
+++ b/Engine/Source/Runtime/Renderer/Private/MobileSeparateTranslucencyPass.cpp
@@ -10,14 +10,7 @@ MobileSeparateTranslucencyPass.cpp - Mobile specific separate translucency pass
 
 bool IsMobileSeparateTranslucencyActive(const FViewInfo& View)
 {
-	if (UseMeshDrawCommandPipeline())
-	{
-		return View.ParallelMeshDrawCommandPasses[EMeshPass::TranslucencyAfterDOF].HasAnyDraw();
-	}
-	else
-	{
-		return View.TranslucentPrimSet.SortedPrimsNum.Num(ETranslucencyPass::TPT_TranslucencyAfterDOF) > 0;
-	}
+	return View.ParallelMeshDrawCommandPasses[EMeshPass::TranslucencyAfterDOF].HasAnyDraw();
 }
 
 void FRCSeparateTranslucensyPassES2::Process(FRenderingCompositePassContext& Context)
@@ -37,25 +30,7 @@ void FRCSeparateTranslucensyPassES2::Process(FRenderingCompositePassContext& Con
 	{
 		// Set the view family's render target/viewport.
 		Context.SetViewportAndCallRHI(View.ViewRect);
-
-		if (UseMeshDrawCommandPipeline())
-		{
-			View.ParallelMeshDrawCommandPasses[EMeshPass::TranslucencyAfterDOF].DispatchDraw(nullptr, Context.RHICmdList);
-		}
-		else
-		{
-			TUniformBufferRef<FMobileBasePassUniformParameters> BasePassUniformBuffer;
-			CreateMobileBasePassUniformBuffer(Context.RHICmdList, View, true, BasePassUniformBuffer);
-
-			FDrawingPolicyRenderState DrawRenderState(View, BasePassUniformBuffer);
-
-			// Enable depth test, disable depth writes.
-			DrawRenderState.SetDepthStencilState(TStaticDepthStencilState<false, CF_DepthNearOrEqual>::GetRHI());
-
-			// Draw translucent prims
-			FMobileTranslucencyDrawingPolicyFactory::ContextType DrawingContext(ETranslucencyPass::TPT_TranslucencyAfterDOF);
-			View.TranslucentPrimSet.DrawPrimitivesForMobile<FMobileTranslucencyDrawingPolicyFactory>(Context.RHICmdList, View, DrawRenderState, DrawingContext);
-		}
+		View.ParallelMeshDrawCommandPasses[EMeshPass::TranslucencyAfterDOF].DispatchDraw(nullptr, Context.RHICmdList);
 	}
 	Context.RHICmdList.EndRenderPass();
 }
diff --git a/Engine/Source/Runtime/Renderer/Private/MobileShadingRenderer.cpp b/Engine/Source/Runtime/Renderer/Private/MobileShadingRenderer.cpp
--- a/Engine/Source/Runtime/Renderer/Private/MobileShadingRenderer.cpp
+++ b/Engine/Source/Runtime/Renderer/Private/MobileShadingRenderer.cpp
@@ -214,11 +214,6 @@ void FMobileSceneRenderer::InitViews(FRHICommandListImmediate& RHICmdList)
 	// Now that the indirect lighting cache is updated, we can update the uniform buffers.
 	UpdatePrimitiveIndirectLightingCacheBuffers();
 
-	if (!UseMeshDrawCommandPipeline())
-	{
-		PostInitViewCustomData();
-	}
-	
 	OnStartRender(RHICmdList);
 }
 
diff --git a/Engine/Source/Runtime/Renderer/Private/MobileTranslucentRendering.cpp b/Engine/Source/Runtime/Renderer/Private/MobileTranslucentRendering.cpp
--- a/Engine/Source/Runtime/Renderer/Private/MobileTranslucentRendering.cpp
+++ b/Engine/Source/Runtime/Renderer/Private/MobileTranslucentRendering.cpp
@@ -111,212 +111,6 @@ void FMobileSceneRenderer::CopySceneAlpha(FRHICommandListImmediate& RHICmdList,
 	SceneContext.FinishRenderingSceneAlphaCopy(RHICmdList);
 }
 
-
-
-
-
-
-
-/** The parameters used to draw a translucent mesh. */
-class FDrawMobileTranslucentMeshAction
-{
-public:
-	const FViewInfo& View;
-	FDrawingPolicyRenderState DrawRenderState;
-	FHitProxyId HitProxyId;
-
-	/** Initialization constructor. */
-	FDrawMobileTranslucentMeshAction(
-		FRHICommandList& InRHICmdList,
-		const FViewInfo& InView,
-		float InDitheredLODTransitionAlpha,
-		FDepthStencilStateRHIParamRef InDepthStencilState,
-		const FDrawingPolicyRenderState& InDrawRenderState,
-		FHitProxyId InHitProxyId
-		):
-		View(InView),
-		DrawRenderState(InDrawRenderState),
-		HitProxyId(InHitProxyId)
-	{
-		DrawRenderState.SetDitheredLODTransitionAlpha(InDitheredLODTransitionAlpha);
-		if (InDepthStencilState)
-		{
-			DrawRenderState.SetDepthStencilState(InDepthStencilState);
-		}
-	}
-	
-	inline bool ShouldPackAmbientSH() const
-	{
-		// So shader code can read a single constant to get the ambient term
-		return true;
-	}
-
-	static bool CanReceiveCSM(const FLightSceneInfo* LightSceneInfo, const FPrimitiveSceneProxy* PrimitiveSceneProxy) { return false; }
-
-	const FScene* GetScene() const
-	{ 
-		return ((FScene*)View.Family->Scene);
-	}
-
-	/** Draws the translucent mesh with a specific light-map type, and fog volume type */
-	void Process(
-		FRHICommandList& RHICmdList, 
-		const FMobileProcessBasePassMeshParameters& Parameters,
-		const ELightMapPolicyType& LightMapPolicy,
-		const FUniformLightMapPolicy::ElementDataType& LightMapElementData
-		)
-	{
-		const bool bIsLitMaterial = Parameters.ShadingModel != MSM_Unlit;
-		const FScene* Scene = Parameters.PrimitiveSceneProxy ? Parameters.PrimitiveSceneProxy->GetPrimitiveSceneInfo()->Scene : NULL;
-
-		FMobileBasePassUniformDrawingPolicy DrawingPolicy(
-			Parameters.Mesh.VertexFactory,
-			Parameters.Mesh.MaterialRenderProxy,
-			Parameters.Material,
-			LightMapPolicy,
-			Parameters.NumMovablePointLights,
-			Parameters.BlendMode,
-			Parameters.ShadingModel != MSM_Unlit && Scene && Scene->ShouldRenderSkylightInBasePass(Parameters.BlendMode),
-			ComputeMeshOverrideSettings(Parameters.Mesh),
-			View.GetFeatureLevel()
-			);
-
-		DrawingPolicy.SetupPipelineState(DrawRenderState, View);
-		CommitGraphicsPipelineState(RHICmdList, DrawingPolicy, DrawRenderState, DrawingPolicy.GetBoundShaderStateInput(View.GetFeatureLevel()), Parameters.Mesh.MaterialRenderProxy);
-		DrawingPolicy.SetSharedState(RHICmdList, DrawRenderState, &View, FMobileBasePassUniformDrawingPolicy::ContextDataType());
-
-		for (int32 BatchElementIndex = 0; BatchElementIndex<Parameters.Mesh.Elements.Num(); BatchElementIndex++)
-		{
-			TDrawEvent<FRHICommandList> MeshEvent;
-			BeginMeshDrawEvent(RHICmdList, Parameters.PrimitiveSceneProxy, Parameters.Mesh, MeshEvent, EnumHasAnyFlags(EShowMaterialDrawEventTypes(GShowMaterialDrawEventTypes), EShowMaterialDrawEventTypes::MobileTranslucent));
-
-			DrawingPolicy.SetMeshRenderState(
-				RHICmdList, 
-				View,
-				Parameters.PrimitiveSceneProxy,
-				Parameters.Mesh,
-				BatchElementIndex,
-				DrawRenderState,
-				FMobileBasePassUniformDrawingPolicy::ElementDataType(LightMapElementData),
-				FMobileBasePassUniformDrawingPolicy::ContextDataType()
-				);
-			DrawingPolicy.DrawMesh(RHICmdList, View, Parameters.Mesh, BatchElementIndex);
-		}
-	}
-};
-
-/**
- * Render a dynamic mesh using a translucent draw policy
- * @return true if the mesh rendered
- */
-bool FMobileTranslucencyDrawingPolicyFactory::DrawDynamicMesh(
-	FRHICommandList& RHICmdList, 
-	const FViewInfo& View,
-	ContextType DrawingContext,
-	const FMeshBatch& Mesh,
-	bool bPreFog,
-	const FDrawingPolicyRenderState& DrawRenderState,
-	const FPrimitiveSceneProxy* PrimitiveSceneProxy,
-	FHitProxyId HitProxyId
-	)
-{
-	bool bDirty = false;
-
-	// Determine the mesh's material and blend mode.
-	const auto FeatureLevel = View.GetFeatureLevel();
-	const auto ShaderPlatform = View.GetShaderPlatform();
-	const FMaterial* Material = Mesh.MaterialRenderProxy->GetMaterial(FeatureLevel);
-	const EBlendMode BlendMode = Material->GetBlendMode();
-
-	// Only render translucent materials.
-	if (IsTranslucentBlendMode(BlendMode))
-	{
-		FDepthStencilStateRHIParamRef DepthStencilState = nullptr;
-		if (DrawingContext.TranslucencyPass != ETranslucencyPass::TPT_TranslucencyAfterDOF && Material->ShouldDisableDepthTest())
-		{
-			DepthStencilState = TStaticDepthStencilState<false, CF_Always>::GetRHI();
-		}
-		bool bAllowFog = true;
-
-		FDrawMobileTranslucentMeshAction Action(
-				RHICmdList,
-				View,
-				Mesh.DitheredLODTransitionAlpha,
-				DepthStencilState,
-				DrawRenderState,
-				HitProxyId
-				);
-		
-		FMobileProcessBasePassMeshParameters Parameters(
-				Mesh,
-				*Material,
-				PrimitiveSceneProxy,
-				bAllowFog,
-				FeatureLevel, 
-				false // ISR disabled for mobile
-				);
-		
-		const FScene* Scene = (FScene*)View.Family->Scene;
-		const FLightSceneInfo* MobileDirectionalLight = MobileBasePass::GetDirectionalLightInfo(Scene, PrimitiveSceneProxy);
-		bool bCanReceiveCSM = Action.CanReceiveCSM(MobileDirectionalLight, PrimitiveSceneProxy);
-		EMaterialShadingModel ShadingModel = Material->GetShadingModel();
-		ELightMapPolicyType LightmapPolicyType = MobileBasePass::SelectMeshLightmapPolicy(Scene, Mesh, PrimitiveSceneProxy, MobileDirectionalLight, ShadingModel, bCanReceiveCSM, FeatureLevel);
-
-		Action.Process(RHICmdList, Parameters, LightmapPolicyType, Mesh.LCI);
-		bDirty = true;
-	}
-	return bDirty;
-}
-
-template <class TDrawingPolicyFactory>
-void FTranslucentPrimSet::DrawPrimitivesForMobile(FRHICommandListImmediate& RHICmdList, const FViewInfo& View, const FDrawingPolicyRenderState& DrawRenderState,
-	typename TDrawingPolicyFactory::ContextType& DrawingContext) const
-{
-	FInt32Range PassRange = SortedPrimsNum.GetPassRange(DrawingContext.TranslucencyPass);
-
-	// Draw sorted scene prims
-	for (int32 PrimIdx = PassRange.GetLowerBoundValue(); PrimIdx < PassRange.GetUpperBoundValue(); PrimIdx++)
-	{
-		FPrimitiveSceneInfo* PrimitiveSceneInfo = SortedPrims[PrimIdx].PrimitiveSceneInfo;
-		int32 PrimitiveId = PrimitiveSceneInfo->GetIndex();
-		const FPrimitiveViewRelevance& ViewRelevance = View.PrimitiveViewRelevanceMap[PrimitiveId];
-
-		checkSlow(ViewRelevance.HasTranslucency());
-
-		if (ViewRelevance.bDrawRelevance)
-		{
-			// range in View.DynamicMeshElements[]
-			FInt32Range range = View.GetDynamicMeshElementRange(PrimitiveSceneInfo->GetIndex());
-
-			for (int32 MeshBatchIndex = range.GetLowerBoundValue(); MeshBatchIndex < range.GetUpperBoundValue(); MeshBatchIndex++)
-			{
-				const FMeshBatchAndRelevance& MeshBatchAndRelevance = View.DynamicMeshElements[MeshBatchIndex];
-
-				checkSlow(MeshBatchAndRelevance.PrimitiveSceneProxy == PrimitiveSceneInfo->Proxy);
-
-				const FMeshBatch& MeshBatch = *MeshBatchAndRelevance.Mesh;
-				TDrawingPolicyFactory::DrawDynamicMesh(RHICmdList, View, DrawingContext, MeshBatch, false, DrawRenderState, MeshBatchAndRelevance.PrimitiveSceneProxy, MeshBatch.BatchHitProxyId);
-			}
-
-			// Render static scene prim
-			if (ViewRelevance.bStaticRelevance)
-			{
-				// Render static meshes from static scene prim
-				for (int32 StaticMeshIdx = 0; StaticMeshIdx < PrimitiveSceneInfo->StaticMeshes.Num(); StaticMeshIdx++)
-				{
-					FStaticMesh& StaticMesh = PrimitiveSceneInfo->StaticMeshes[StaticMeshIdx];
-					if (View.StaticMeshVisibilityMap[StaticMesh.Id]
-						// Only render static mesh elements using translucent materials
-						&& StaticMesh.IsTranslucent(View.GetFeatureLevel()))
-					{
-						TDrawingPolicyFactory::DrawDynamicMesh(RHICmdList, View, DrawingContext, StaticMesh, false, DrawRenderState, PrimitiveSceneInfo->Proxy, StaticMesh.BatchHitProxyId);
-					}
-				}
-			}
-		}
-	}
-}
-
 void FMobileSceneRenderer::RenderTranslucency(FRHICommandListImmediate& RHICmdList, const TArrayView<const FViewInfo*> PassViews, bool bRenderToSceneColor)
 {
 	ETranslucencyPass::Type TranslucencyPass = 
@@ -362,7 +156,7 @@ void FMobileSceneRenderer::RenderTranslucency(FRHICommandListImmediate& RHICmdLi
 			const FViewInfo& TranslucentViewport = (View.bIsMobileMultiViewEnabled) ? Views[0] : View;
 			RHICmdList.SetViewport(TranslucentViewport.ViewRect.Min.X, TranslucentViewport.ViewRect.Min.Y, 0.0f, TranslucentViewport.ViewRect.Max.X, TranslucentViewport.ViewRect.Max.Y, 1.0f);
 
-			if (UseMeshDrawCommandPipeline() && !View.Family->UseDebugViewPS())
+			if (!View.Family->UseDebugViewPS())
 			{
 				if (Scene->UniformBuffers.UpdateViewUniformBuffer(View))
 				{
@@ -373,26 +167,6 @@ void FMobileSceneRenderer::RenderTranslucency(FRHICommandListImmediate& RHICmdLi
 				const EMeshPass::Type MeshPass = TranslucencyPassToMeshPass(TranslucencyPass);
 				View.ParallelMeshDrawCommandPasses[MeshPass].DispatchDraw(nullptr, RHICmdList);
 			}
-			else
-			{
-				TUniformBufferRef<FMobileBasePassUniformParameters> BasePassUniformBuffer;
-				CreateMobileBasePassUniformBuffer(RHICmdList, View, true, BasePassUniformBuffer);
-				FDrawingPolicyRenderState DrawRenderState(View, BasePassUniformBuffer);
-								
-				// Enable depth test, disable depth writes.
-				DrawRenderState.SetDepthStencilState(TStaticDepthStencilState<false, CF_DepthNearOrEqual>::GetRHI());
-
-				// Draw only translucent prims that don't read from scene color
-				FMobileTranslucencyDrawingPolicyFactory::ContextType DrawingContext(TranslucencyPass);
-				View.TranslucentPrimSet.DrawPrimitivesForMobile<FMobileTranslucencyDrawingPolicyFactory>(RHICmdList, View, DrawRenderState, DrawingContext);
-			
-				View.SimpleElementCollector.DrawBatchedElements(RHICmdList, DrawRenderState, View, EBlendModeFilter::Translucent, SDPG_World);
-				View.SimpleElementCollector.DrawBatchedElements(RHICmdList, DrawRenderState, View, EBlendModeFilter::Translucent, SDPG_Foreground);
-
-				// Editor and debug rendering
-				DrawViewElements<FMobileTranslucencyDrawingPolicyFactory>(RHICmdList, View, DrawRenderState, DrawingContext, SDPG_World, false);
-				DrawViewElements<FMobileTranslucencyDrawingPolicyFactory>(RHICmdList, View, DrawRenderState, DrawingContext, SDPG_Foreground, false);
-			}
 
 			// #todo-renderpasses clean this up. Ending the renderpass is correct here but do we want to look nicer and call FinishRenderingSceneColor, etc?
 			if (bRenderToSceneColor)
@@ -513,207 +287,6 @@ public:
 IMPLEMENT_MATERIAL_SHADER_TYPE(, FOpacityOnlyPS, TEXT("/Engine/Private/MobileOpacityShaders.usf"), TEXT("MainPS"), SF_Pixel);
 IMPLEMENT_SHADERPIPELINE_TYPE_VSPS(MobileOpacityPipeline, FOpacityOnlyVS, FOpacityOnlyPS, true);
 
-class FMobileOpacityDrawingPolicy : public FMeshDrawingPolicy
-{
-public:
-
-	struct ContextDataType : public FMeshDrawingPolicy::ContextDataType
-	{
-	};
-
-	FMobileOpacityDrawingPolicy(
-		const FVertexFactory* InVertexFactory,
-		const FMaterialRenderProxy* InMaterialRenderProxy,
-		const FMaterial& InMaterialResource,
-		ERHIFeatureLevel::Type InFeatureLevel,
-		const FMeshDrawingPolicyOverrideSettings& InOverrideSettings
-	) :
-		FMeshDrawingPolicy(InVertexFactory, InMaterialRenderProxy, InMaterialResource, InOverrideSettings)
-	{
-		const EMaterialTessellationMode TessellationMode = InMaterialResource.GetTessellationMode();
-		static const auto* CVar = IConsoleManager::Get().FindTConsoleVariableDataInt(TEXT("r.ShaderPipelines"));
-		bool bUseShaderPipelines = CVar && CVar->GetValueOnAnyThread() != 0;
-
-		ShaderPipeline = bUseShaderPipelines ? InMaterialResource.GetShaderPipeline(&MobileOpacityPipeline, InVertexFactory->GetType()) : nullptr;
-
-		if (ShaderPipeline)
-		{
-			VertexShader = ShaderPipeline->GetShader<FOpacityOnlyVS >();
-			PixelShader = ShaderPipeline->GetShader<FOpacityOnlyPS>();
-		}
-		else
-		{
-			VertexShader = InMaterialResource.GetShader<FOpacityOnlyVS >(VertexFactory->GetType());
-			PixelShader = InMaterialResource.GetShader<FOpacityOnlyPS>(InVertexFactory->GetType());
-		}
-		BaseVertexShader = VertexShader;
-	}
-
-	// FMeshDrawingPolicy interface.
-	FDrawingPolicyMatchResult Matches(const FMobileOpacityDrawingPolicy& Other, bool bForReals = false) const
-	{
-		DRAWING_POLICY_MATCH_BEGIN
-			DRAWING_POLICY_MATCH(FMeshDrawingPolicy::Matches(Other, bForReals)) &&
-			DRAWING_POLICY_MATCH(VertexShader == Other.VertexShader) &&
-			DRAWING_POLICY_MATCH(PixelShader == Other.PixelShader);
-		DRAWING_POLICY_MATCH_END
-	}
-
-	void SetSharedState(FRHICommandList& RHICmdList, const FDrawingPolicyRenderState& DrawRenderState, const FSceneView* View, const FMobileOpacityDrawingPolicy::ContextDataType PolicyContext) const
-	{
-		VertexShader->SetParameters(RHICmdList, MaterialRenderProxy, *MaterialResource, *View, DrawRenderState);
-		PixelShader->SetParameters(RHICmdList, MaterialRenderProxy, *MaterialResource, View, DrawRenderState);
-
-		// Set the shared mesh resources.
-		FMeshDrawingPolicy::SetSharedState(RHICmdList, DrawRenderState, View, PolicyContext);
-	}
-
-	/**
-	* Create bound shader state using the vertex decl from the mesh draw policy
-	* as well as the shaders needed to draw the mesh
-	* @return new bound shader state object
-	*/
-	FBoundShaderStateInput GetBoundShaderStateInput(ERHIFeatureLevel::Type InFeatureLevel) const
-	{
-		return FBoundShaderStateInput(
-			FMeshDrawingPolicy::GetVertexDeclaration(),
-			VertexShader->GetVertexShader(),
-			nullptr,
-			nullptr,
-			PixelShader->GetPixelShader(),
-			NULL);
-	}
-
-	void SetMeshRenderState(
-		FRHICommandList& RHICmdList,
-		const FSceneView& View,
-		const FPrimitiveSceneProxy* PrimitiveSceneProxy,
-		const FMeshBatch& Mesh,
-		int32 BatchElementIndex,
-		const FDrawingPolicyRenderState& DrawRenderState,
-		const ElementDataType& ElementData,
-		const ContextDataType PolicyContext
-	) const
-	{
-		const FMeshBatchElement& BatchElement = Mesh.Elements[BatchElementIndex];
-		VertexShader->SetMesh(RHICmdList, VertexFactory, View, PrimitiveSceneProxy, BatchElement, DrawRenderState);
-		PixelShader->SetMesh(RHICmdList, VertexFactory, View, PrimitiveSceneProxy, BatchElement, DrawRenderState);
-	}
-
-	friend int32 CompareDrawingPolicy(const FMobileOpacityDrawingPolicy& A, const FMobileOpacityDrawingPolicy& B);
-
-private:
-
-	FShaderPipeline* ShaderPipeline;
-	FOpacityOnlyVS* VertexShader;
-	FOpacityOnlyPS* PixelShader;
-};
-
-int32 CompareDrawingPolicy(const FMobileOpacityDrawingPolicy& A, const FMobileOpacityDrawingPolicy& B)
-{
-	COMPAREDRAWINGPOLICYMEMBERS(VertexShader);
-	COMPAREDRAWINGPOLICYMEMBERS(PixelShader);
-	COMPAREDRAWINGPOLICYMEMBERS(VertexFactory);
-	COMPAREDRAWINGPOLICYMEMBERS(MaterialRenderProxy);
-	return 0;
-}
-
-class FMobileOpacityDrawingPolicyFactory
-{
-public:
-
-	struct ContextType
-	{
-		ETranslucencyPass::Type TranslucencyPass;
-
-		ContextType(ETranslucencyPass::Type InTranslucencyPass)
-		: TranslucencyPass(InTranslucencyPass)
-		{}
-	};
-
-	static bool DrawDynamicMesh(
-		FRHICommandList& RHICmdList,
-		const FViewInfo& View,
-		ContextType DrawingContext,
-		const FMeshBatch& Mesh,
-		bool bPreFog,
-		const FDrawingPolicyRenderState& DrawRenderState,
-		const FPrimitiveSceneProxy* PrimitiveSceneProxy,
-		FHitProxyId HitProxyId
-	)
-	{
-		return DrawMesh(
-			RHICmdList,
-			View,
-			DrawingContext,
-			Mesh,
-			Mesh.Elements.Num() == 1 ? 1ull : (1ull << Mesh.Elements.Num()) - 1ull,	// 1 bit set for each mesh element
-			DrawRenderState,
-			bPreFog,
-			PrimitiveSceneProxy,
-			HitProxyId
-		);
-	}
-
-private:
-	/**
-	* Render a dynamic or static mesh using the opacity draw policy
-	* @return true if the mesh rendered
-	*/
-	static bool DrawMesh(
-		FRHICommandList& RHICmdList,
-		const FViewInfo& View,
-		ContextType DrawingContext,
-		const FMeshBatch& Mesh,
-		const uint64& BatchElementMask,
-		const FDrawingPolicyRenderState& DrawRenderState,
-		bool bPreFog,
-		const FPrimitiveSceneProxy* PrimitiveSceneProxy,
-		FHitProxyId HitProxyId
-	)
-	{
-		bool bDirty = false;
-
-		const FMaterialRenderProxy* MaterialRenderProxy = Mesh.MaterialRenderProxy;
-		const FMaterial* Material = MaterialRenderProxy->GetMaterial(View.GetFeatureLevel());
-		const EBlendMode BlendMode = Material->GetBlendMode();
-
-		bool bDraw = IsTranslucentBlendMode(BlendMode);
-		if (bDraw)
-		{
-			FMeshDrawingPolicyOverrideSettings OverrideSettings = ComputeMeshOverrideSettings(Mesh);
-			OverrideSettings.MeshOverrideFlags |= Material->IsTwoSided() ? EDrawingPolicyOverrideFlags::TwoSided : EDrawingPolicyOverrideFlags::None;
-			FMobileOpacityDrawingPolicy DrawingPolicy(Mesh.VertexFactory, MaterialRenderProxy, *MaterialRenderProxy->GetMaterial(View.GetFeatureLevel()), View.GetFeatureLevel(), OverrideSettings);
-
-			FDrawingPolicyRenderState DrawRenderStateLocal(DrawRenderState);
-			DrawingPolicy.SetupPipelineState(DrawRenderStateLocal, View);
-			CommitGraphicsPipelineState(RHICmdList, DrawingPolicy, DrawRenderStateLocal, DrawingPolicy.GetBoundShaderStateInput(View.GetFeatureLevel()), Mesh.MaterialRenderProxy);
-			DrawingPolicy.SetSharedState(RHICmdList, DrawRenderStateLocal, &View, FMobileOpacityDrawingPolicy::ContextDataType());
-
-			int32 BatchElementIndex = 0;
-			uint64 Mask = BatchElementMask;
-			do
-			{
-				if (Mask & 1)
-				{
-					const bool bIsInstancedMesh = Mesh.Elements[BatchElementIndex].bIsInstancedMesh;
-					TDrawEvent<FRHICommandList> MeshEvent;
-					BeginMeshDrawEvent(RHICmdList, PrimitiveSceneProxy, Mesh, MeshEvent, EnumHasAnyFlags(EShowMaterialDrawEventTypes(GShowMaterialDrawEventTypes), EShowMaterialDrawEventTypes::MobileTranslucentOpacity));
-
-					DrawingPolicy.SetMeshRenderState(RHICmdList, View, PrimitiveSceneProxy, Mesh, BatchElementIndex, DrawRenderStateLocal, FMeshDrawingPolicy::ElementDataType(), FMobileOpacityDrawingPolicy::ContextDataType());
-					DrawingPolicy.DrawMesh(RHICmdList, View, Mesh, BatchElementIndex);
-				}
-				Mask >>= 1;
-				BatchElementIndex++;
-			} while (Mask);
-
-			bDirty = true;
-		}		
-
-		return bDirty;
-	}
-};
-
 bool FMobileSceneRenderer::RenderInverseOpacity(FRHICommandListImmediate& RHICmdList, const FViewInfo& View)
 {
 	// Function MUST be self-contained wrt RenderPasses
@@ -726,58 +299,37 @@ bool FMobileSceneRenderer::RenderInverseOpacity(FRHICommandListImmediate& RHICmd
 	{
 		const bool bGammaSpace = !IsMobileHDR();
 		const bool bLinearHDR64 = !bGammaSpace && !IsMobileHDR32bpp();
-
+		
+		if (!bGammaSpace)
 		{
-			if (!bGammaSpace)
-			{
-				SceneContext.BeginRenderingTranslucency(RHICmdList, View, *this);
-			}
-			else
-			{
-				SceneContext.BeginRenderingSceneColor(RHICmdList, ESimpleRenderTargetMode::EClearColorExistingDepth);
-				// Mobile multi-view is not side by side stereo
-				const FViewInfo& TranslucentViewport = (View.bIsMobileMultiViewEnabled) ? Views[0] : View;
-				RHICmdList.SetViewport(TranslucentViewport.ViewRect.Min.X, TranslucentViewport.ViewRect.Min.Y, 0.0f, TranslucentViewport.ViewRect.Max.X, TranslucentViewport.ViewRect.Max.Y, 1.0f);
-			}
+			SceneContext.BeginRenderingTranslucency(RHICmdList, View, *this);
+		}
+		else
+		{
+			SceneContext.BeginRenderingSceneColor(RHICmdList, ESimpleRenderTargetMode::EClearColorExistingDepth);
+			// Mobile multi-view is not side by side stereo
+			const FViewInfo& TranslucentViewport = (View.bIsMobileMultiViewEnabled) ? Views[0] : View;
+			RHICmdList.SetViewport(TranslucentViewport.ViewRect.Min.X, TranslucentViewport.ViewRect.Min.Y, 0.0f, TranslucentViewport.ViewRect.Max.X, TranslucentViewport.ViewRect.Max.Y, 1.0f);
+		}
 
-			if (UseMeshDrawCommandPipeline())
-			{
-				if (Scene->UniformBuffers.UpdateViewUniformBuffer(View))
-				{
-					UpdateTranslucentBasePassUniformBuffer(RHICmdList, View);
-					UpdateDirectionalLightUniformBuffers(RHICmdList, View);
-				}
+		if (Scene->UniformBuffers.UpdateViewUniformBuffer(View))
+		{
+			UpdateTranslucentBasePassUniformBuffer(RHICmdList, View);
+			UpdateDirectionalLightUniformBuffers(RHICmdList, View);
+		}
 		
-				View.ParallelMeshDrawCommandPasses[EMeshPass::MobileInverseOpacity].DispatchDraw(nullptr, RHICmdList);
+		View.ParallelMeshDrawCommandPasses[EMeshPass::MobileInverseOpacity].DispatchDraw(nullptr, RHICmdList);
 				
-				bDirty |= View.ParallelMeshDrawCommandPasses[EMeshPass::MobileInverseOpacity].HasAnyDraw();
-			}
-			else
-			{
-				FMobileBasePassUniformParameters PassParameters;
-				SetupMobileBasePassUniformParameters(RHICmdList, View, true, PassParameters);
-				TUniformBufferRef<FMobileBasePassUniformParameters> PassUniformBuffer = TUniformBufferRef<FMobileBasePassUniformParameters>::CreateUniformBufferImmediate(PassParameters, UniformBuffer_SingleFrame);
-
-				FDrawingPolicyRenderState DrawRenderState(View, PassUniformBuffer);
-				// Enable depth test, disable depth writes.
-				DrawRenderState.SetDepthStencilState(TStaticDepthStencilState<false, CF_DepthNearOrEqual>::GetRHI());
-				DrawRenderState.SetBlendState(TStaticBlendState<CW_ALPHA, BO_Add, BF_DestColor, BF_Zero, BO_Add, BF_Zero, BF_InverseSourceAlpha>::GetRHI());
+		bDirty |= View.ParallelMeshDrawCommandPasses[EMeshPass::MobileInverseOpacity].HasAnyDraw();
 
-				{
-					SCOPED_DRAW_EVENT(RHICmdList, InverseOpacity);
-					bDirty |= RenderInverseOpacityDynamic(RHICmdList, View, DrawRenderState);
-				}
-			}
-
-			if (!bGammaSpace)
-			{
-				RHICmdList.EndRenderPass();
-				SceneContext.FinishRenderingTranslucency(RHICmdList);
-			}
-			else
-			{
-				SceneContext.FinishRenderingSceneColor(RHICmdList);
-			}
+		if (!bGammaSpace)
+		{
+			RHICmdList.EndRenderPass();
+			SceneContext.FinishRenderingTranslucency(RHICmdList);
+		}
+		else
+		{
+			SceneContext.FinishRenderingSceneColor(RHICmdList);
 		}
 	}
 	else
@@ -789,28 +341,6 @@ bool FMobileSceneRenderer::RenderInverseOpacity(FRHICommandListImmediate& RHICmd
 	return bDirty;
 }
 
-bool FMobileSceneRenderer::RenderInverseOpacityDynamic(FRHICommandListImmediate& RHICmdList, const FViewInfo& View, const FDrawingPolicyRenderState& DrawRenderState)
-{
-	if (ViewFamily.AllowTranslucencyAfterDOF())
-	{
-		{
-			FMobileOpacityDrawingPolicyFactory::ContextType DrawingContext(ETranslucencyPass::TPT_StandardTranslucency);
-			View.TranslucentPrimSet.DrawPrimitivesForMobile<FMobileOpacityDrawingPolicyFactory>(RHICmdList, View, DrawRenderState, DrawingContext);
-		}
-		{
-			FMobileOpacityDrawingPolicyFactory::ContextType DrawingContext(ETranslucencyPass::TPT_TranslucencyAfterDOF);
-			View.TranslucentPrimSet.DrawPrimitivesForMobile<FMobileOpacityDrawingPolicyFactory>(RHICmdList, View, DrawRenderState, DrawingContext);
-		}
-	}
-	else
-	{
-		FMobileOpacityDrawingPolicyFactory::ContextType DrawingContext(ETranslucencyPass::TPT_AllTranslucency);
-		View.TranslucentPrimSet.DrawPrimitivesForMobile<FMobileOpacityDrawingPolicyFactory>(RHICmdList, View, DrawRenderState, DrawingContext);
-	}
-
-	return View.TranslucentPrimSet.NumPrims() > 0;
-}
-
 class FMobileInverseOpacityMeshProcessor : public FMeshPassProcessor
 {
 public:
diff --git a/Engine/Source/Runtime/Renderer/Private/PostProcess/PostProcessCompositeEditorPrimitives.cpp b/Engine/Source/Runtime/Renderer/Private/PostProcess/PostProcessCompositeEditorPrimitives.cpp
--- a/Engine/Source/Runtime/Renderer/Private/PostProcess/PostProcessCompositeEditorPrimitives.cpp
+++ b/Engine/Source/Runtime/Renderer/Private/PostProcess/PostProcessCompositeEditorPrimitives.cpp
@@ -640,42 +640,14 @@ void FRCPassPostProcessCompositeEditorPrimitives::Process(FRenderingCompositePas
 			DrawRenderState.SetDepthStencilAccess(FExclusiveDepthStencil::DepthWrite_StencilWrite);
 			DrawRenderState.SetBlendState(TStaticBlendStateWriteMask<CW_RGBA>::GetRHI());
 
-			if (bDeferredBasePass)
-			{
-				RenderEditorPrimitives(Context.RHICmdList, EditorView, DrawRenderState);
-			}
-			else
-			{
-				if (UseMeshDrawCommandPipeline())
-				{
-					RenderEditorPrimitives(Context.RHICmdList, EditorView, DrawRenderState);
-				}
-				else
-				{
-					RenderEditorPrimitivesWithDrawPolicies<FMobileBasePassOpaqueDrawingPolicyFactory>(Context.RHICmdList, EditorView, DrawRenderState);
-				}
-			}
+			RenderEditorPrimitives(Context.RHICmdList, EditorView, DrawRenderState);
 
 			// Draw editor primitives.
 			{
 				SCOPED_DRAW_EVENTF(Context.RHICmdList, TemporalAA, TEXT("RenderViewEditorPrimitives %dx%d msaa=%d"),
 					ViewRect.Width(), ViewRect.Height(), MSAASampleCount);
 
-				if (bDeferredBasePass)
-				{
-					RenderEditorPrimitives(Context.RHICmdList, EditorView, DrawRenderState);
-				}
-				else
-				{
-					if (UseMeshDrawCommandPipeline())
-					{
-						RenderEditorPrimitives(Context.RHICmdList, EditorView, DrawRenderState);
-					}
-					else
-					{
-						RenderEditorPrimitivesWithDrawPolicies<FMobileBasePassOpaqueDrawingPolicyFactory>(Context.RHICmdList, EditorView, DrawRenderState);
-					}
-				}
+				RenderEditorPrimitives(Context.RHICmdList, EditorView, DrawRenderState);
 			}
 
 			// Draw foreground editor primitives.
@@ -683,22 +655,7 @@ void FRCPassPostProcessCompositeEditorPrimitives::Process(FRenderingCompositePas
 				SCOPED_DRAW_EVENTF(Context.RHICmdList, TemporalAA, TEXT("RenderViewEditorForegroundPrimitives %dx%d msaa=%d"),
 					ViewRect.Width(), ViewRect.Height(), MSAASampleCount);
 
-				if (bDeferredBasePass)
-				{
-					RenderForegroundEditorPrimitives(Context.RHICmdList, EditorView, DrawRenderState);
-				}
-				else
-				{
-					if (UseMeshDrawCommandPipeline())
-					{
-						RenderForegroundEditorPrimitives(Context.RHICmdList, EditorView, DrawRenderState);
-					}
-					else
-					{
-						RenderForegroundEditorPrimitivesWithDrawPolicies<FMobileBasePassOpaqueDrawingPolicyFactory>(Context.RHICmdList, EditorView, DrawRenderState);
-					}
-
-				}
+				RenderForegroundEditorPrimitives(Context.RHICmdList, EditorView, DrawRenderState);
 			}
 		}
 		Context.RHICmdList.EndRenderPass();
diff --git a/Engine/Source/Runtime/Renderer/Private/PrimitiveSceneInfo.cpp b/Engine/Source/Runtime/Renderer/Private/PrimitiveSceneInfo.cpp
--- a/Engine/Source/Runtime/Renderer/Private/PrimitiveSceneInfo.cpp
+++ b/Engine/Source/Runtime/Renderer/Private/PrimitiveSceneInfo.cpp
@@ -164,15 +164,6 @@ void FPrimitiveSceneInfo::CacheMeshDrawCommands(FRHICommandListImmediate& RHICmd
 {
 	check(StaticMeshCommandInfos.Num() == 0);
 
-	// TEMP: optionally allow to switch off mesh command caching for cooked mobile
-	extern int32 MeshDrawCommandPipelineMobileCooked();
-	const int32 CacheMeshCommandsVal = MeshDrawCommandPipelineMobileCooked();
-	if (CacheMeshCommandsVal <= 0)
-	{
-		return;
-	}
-
-
 	int32 MeshWithCachedCommandsNum = 0;
 	for (int32 MeshIndex = 0; MeshIndex < StaticMeshes.Num(); MeshIndex++)
 	{
diff --git a/Engine/Source/Runtime/Renderer/Private/Renderer.cpp b/Engine/Source/Runtime/Renderer/Private/Renderer.cpp
--- a/Engine/Source/Runtime/Renderer/Private/Renderer.cpp
+++ b/Engine/Source/Runtime/Renderer/Private/Renderer.cpp
@@ -98,6 +98,7 @@ void FRendererModule::DrawTileMesh(FRHICommandListImmediate& RHICmdList, FDrawin
 		View.ViewRect = View.UnscaledViewRect;
 		
 		const auto FeatureLevel = View.GetFeatureLevel();
+		const EShadingPath ShadingPath = FSceneInterface::GetShadingPath(FeatureLevel);
 
 		Mesh.MaterialRenderProxy->UpdateUniformExpressionCacheIfNeeded(FeatureLevel);
 		FMaterialRenderProxy::UpdateDeferredCachedUniformExpressions();
@@ -130,7 +131,7 @@ void FRendererModule::DrawTileMesh(FRHICommandListImmediate& RHICmdList, FDrawin
 		View.InitRHIResources();
 		DrawRenderState.SetViewUniformBuffer(View.ViewUniformBuffer);
 
-		if (FeatureLevel <= ERHIFeatureLevel::ES3_1)
+		if (ShadingPath == EShadingPath::Mobile)
 		{
 			View.MobileDirectionalLightUniformBuffers[0] = TUniformBufferRef<FMobileDirectionalLightShaderParameters>::CreateUniformBufferImmediate(FMobileDirectionalLightShaderParameters(), UniformBuffer_SingleFrame);
 		}
@@ -149,8 +150,9 @@ void FRendererModule::DrawTileMesh(FRHICommandListImmediate& RHICmdList, FDrawin
 #if !(UE_BUILD_SHIPPING || UE_BUILD_TEST)
 			// make sure we are doing opaque drawing
 			DrawRenderState.SetBlendState(TStaticBlendState<>::GetRHI());
-
-			if (FeatureLevel >= ERHIFeatureLevel::SM4)
+			
+			// is this path used on mobile?
+			if (ShadingPath == EShadingPath::Deferred)
 			{
 				FSceneRenderTargets& SceneContext = FSceneRenderTargets::Get(RHICmdList);
 				FDebugViewModePassPassUniformParameters PassParameters;
@@ -172,7 +174,7 @@ void FRendererModule::DrawTileMesh(FRHICommandListImmediate& RHICmdList, FDrawin
 		}
 		else if (IsTranslucentBlendMode(MaterialBlendMode))
 		{
-			if (FeatureLevel >= ERHIFeatureLevel::SM4)
+			if (ShadingPath == EShadingPath::Deferred)
 			{
 				// Crash fix - reflection capture shader parameter is bound but we have no buffer during Build Texture Streaming
 				if(!View.ReflectionCaptureUniformBuffer.IsValid())
@@ -183,29 +185,28 @@ void FRendererModule::DrawTileMesh(FRHICommandListImmediate& RHICmdList, FDrawin
 				TUniformBufferRef<FTranslucentBasePassUniformParameters> BasePassUniformBuffer;
 				CreateTranslucentBasePassUniformBuffer(RHICmdList, View, nullptr, ESceneTextureSetupMode::None, BasePassUniformBuffer, 0);
 				DrawRenderState.SetPassUniformBuffer(BasePassUniformBuffer);
-
-				DrawDynamicMeshPass(View, RHICmdList,
-					[&View, &DrawRenderState, &Mesh](FDynamicPassMeshDrawListContext* DynamicMeshPassContext)
-					{
-						FBasePassMeshProcessor PassMeshProcessor(
-							nullptr,
-							View.GetFeatureLevel(),
-							&View,
-							DrawRenderState,
-							DynamicMeshPassContext,
-							ETranslucencyPass::TPT_AllTranslucency);
-
-						const uint64 DefaultBatchElementMask = ~0ull;
-						PassMeshProcessor.AddMeshBatch(Mesh, DefaultBatchElementMask, nullptr);
-					});
 			}
-			else
+			else // Mobile
 			{
 				TUniformBufferRef<FMobileBasePassUniformParameters> BasePassUniformBuffer;
-				CreateMobileBasePassUniformBuffer(RHICmdList, View, false, BasePassUniformBuffer);
+				CreateMobileBasePassUniformBuffer(RHICmdList, View, true, BasePassUniformBuffer);
 				DrawRenderState.SetPassUniformBuffer(BasePassUniformBuffer);
-				FMobileTranslucencyDrawingPolicyFactory::DrawDynamicMesh(RHICmdList, View, FMobileTranslucencyDrawingPolicyFactory::ContextType(ETranslucencyPass::TPT_AllTranslucency), Mesh, false, DrawRenderState, NULL, HitProxyId);
 			}
+			
+			DrawDynamicMeshPass(View, RHICmdList,
+				[&View, &DrawRenderState, &Mesh](FDynamicPassMeshDrawListContext* DynamicMeshPassContext)
+				{
+					FBasePassMeshProcessor PassMeshProcessor(
+						nullptr,
+						View.GetFeatureLevel(),
+						&View,
+						DrawRenderState,
+						DynamicMeshPassContext,
+						ETranslucencyPass::TPT_AllTranslucency);
+
+					const uint64 DefaultBatchElementMask = ~0ull;
+					PassMeshProcessor.AddMeshBatch(Mesh, DefaultBatchElementMask, nullptr);
+				});
 		}
 		// handle opaque materials
 		else
@@ -236,33 +237,32 @@ void FRendererModule::DrawTileMesh(FRHICommandListImmediate& RHICmdList, FDrawin
 			}
 			else
 			{
-				if (FeatureLevel >= ERHIFeatureLevel::SM4)
+				if (ShadingPath == EShadingPath::Deferred)
 				{
 					TUniformBufferRef<FOpaqueBasePassUniformParameters> BasePassUniformBuffer;
 					CreateOpaqueBasePassUniformBuffer(RHICmdList, View, nullptr, BasePassUniformBuffer);
 					DrawRenderState.SetPassUniformBuffer(BasePassUniformBuffer);
-
-					DrawDynamicMeshPass(View, RHICmdList,
-						[&View, &DrawRenderState, &Mesh](FDynamicPassMeshDrawListContext* DynamicMeshPassContext)
-						{
-							FBasePassMeshProcessor PassMeshProcessor(
-								nullptr,
-								View.GetFeatureLevel(),
-								&View,
-								DrawRenderState,
-								DynamicMeshPassContext);
-
-							const uint64 DefaultBatchElementMask = ~0ull;
-							PassMeshProcessor.AddMeshBatch(Mesh, DefaultBatchElementMask, nullptr);
-						});
 				}
-				else
+				else // Mobile
 				{
 					TUniformBufferRef<FMobileBasePassUniformParameters> BasePassUniformBuffer;
 					CreateMobileBasePassUniformBuffer(RHICmdList, View, false, BasePassUniformBuffer);
 					DrawRenderState.SetPassUniformBuffer(BasePassUniformBuffer);
-					FMobileBasePassOpaqueDrawingPolicyFactory::DrawDynamicMesh(RHICmdList, View, FMobileBasePassOpaqueDrawingPolicyFactory::ContextType(), Mesh, false, DrawRenderState, NULL, HitProxyId);
 				}
+
+				DrawDynamicMeshPass(View, RHICmdList,
+					[&View, &DrawRenderState, &Mesh](FDynamicPassMeshDrawListContext* DynamicMeshPassContext)
+					{
+						FBasePassMeshProcessor PassMeshProcessor(
+							nullptr,
+							View.GetFeatureLevel(),
+							&View,
+							DrawRenderState,
+							DynamicMeshPassContext);
+
+						const uint64 DefaultBatchElementMask = ~0ull;
+						PassMeshProcessor.AddMeshBatch(Mesh, DefaultBatchElementMask, nullptr);
+					});
 			}
 		}
 
diff --git a/Engine/Source/Runtime/Renderer/Private/RendererScene.cpp b/Engine/Source/Runtime/Renderer/Private/RendererScene.cpp
--- a/Engine/Source/Runtime/Renderer/Private/RendererScene.cpp
+++ b/Engine/Source/Runtime/Renderer/Private/RendererScene.cpp
@@ -2805,15 +2805,6 @@ void FScene::UpdateStaticDrawListsForMaterials_RenderThread(FRHICommandListImmed
 	TArray<FPrimitiveSceneInfo*> PrimitivesToUpdate;
 	auto SceneFeatureLevel = GetFeatureLevel();
 
-	if (GetShadingPath() == EShadingPath::Mobile)
-	{
-		for (int32 DrawType = 0; DrawType < EBasePass_MAX; DrawType++)
-		{
-			MobileBasePassUniformLightMapPolicyDrawList[DrawType].GetUsedPrimitivesBasedOnMaterials(SceneFeatureLevel, Materials, PrimitivesToUpdate);
-			MobileBasePassUniformLightMapPolicyDrawListWithCSM[DrawType].GetUsedPrimitivesBasedOnMaterials(SceneFeatureLevel, Materials, PrimitivesToUpdate);
-		}
-	}
-
 	PositionOnlyDepthDrawList.GetUsedPrimitivesBasedOnMaterials(SceneFeatureLevel, Materials, PrimitivesToUpdate);
 	DepthDrawList.GetUsedPrimitivesBasedOnMaterials(SceneFeatureLevel, Materials, PrimitivesToUpdate);
 	MaskedDepthDrawList.GetUsedPrimitivesBasedOnMaterials(SceneFeatureLevel, Materials, PrimitivesToUpdate);
@@ -3034,10 +3025,6 @@ void FScene::DumpStaticMeshDrawListStats() const
 	DUMP_DRAW_LIST(PositionOnlyDepthDrawList);
 	DUMP_DRAW_LIST(DepthDrawList);
 	DUMP_DRAW_LIST(MaskedDepthDrawList);
-	DUMP_DRAW_LIST(MobileBasePassUniformLightMapPolicyDrawList[EBasePass_Default]);
-	DUMP_DRAW_LIST(MobileBasePassUniformLightMapPolicyDrawList[EBasePass_Masked]);
-	DUMP_DRAW_LIST(MobileBasePassUniformLightMapPolicyDrawListWithCSM[EBasePass_Default]);
-	DUMP_DRAW_LIST(MobileBasePassUniformLightMapPolicyDrawListWithCSM[EBasePass_Masked]);
 	DUMP_DRAW_LIST(WholeSceneShadowDepthDrawList);
 #undef DUMP_DRAW_LIST
 }
@@ -3197,8 +3184,6 @@ void FScene::ApplyWorldOffset_RenderThread(FVector InOffset)
 	StaticMeshDrawListApplyWorldOffset(DepthDrawList, InOffset);
 	StaticMeshDrawListApplyWorldOffset(MaskedDepthDrawList, InOffset);
 	StaticMeshDrawListApplyWorldOffset(WholeSceneShadowDepthDrawList, InOffset);
-	StaticMeshDrawListApplyWorldOffset(MobileBasePassUniformLightMapPolicyDrawList, InOffset);
-	StaticMeshDrawListApplyWorldOffset(MobileBasePassUniformLightMapPolicyDrawListWithCSM, InOffset);
 
 	VelocityData.ApplyOffset(InOffset);
 }
@@ -3484,16 +3469,6 @@ FSceneViewStateInterface* FRendererModule::AllocateViewState()
 	return new FSceneViewState();
 }
 
-TStaticMeshDrawList<FMobileBasePassUniformDrawingPolicy>& FScene::GetMobileBasePassDrawList(EBasePassDrawListType DrawType)
-{
-	return MobileBasePassUniformLightMapPolicyDrawList[DrawType];
-}
-
-TStaticMeshDrawList<FMobileBasePassUniformDrawingPolicy>& FScene::GetMobileBasePassCSMDrawList(EBasePassDrawListType DrawType)
-{
-	return MobileBasePassUniformLightMapPolicyDrawListWithCSM[DrawType];
-}
-
 //////////////////////////////////////////////////////////////////////////
 
 FLatentGPUTimer::FLatentGPUTimer(int32 InAvgSamples)
diff --git a/Engine/Source/Runtime/Renderer/Private/SceneCore.cpp b/Engine/Source/Runtime/Renderer/Private/SceneCore.cpp
--- a/Engine/Source/Runtime/Renderer/Private/SceneCore.cpp
+++ b/Engine/Source/Runtime/Renderer/Private/SceneCore.cpp
@@ -356,33 +356,15 @@ void FStaticMesh::AddToDrawLists(FRHICommandListImmediate& RHICmdList, FScene* S
 {
 	const auto FeatureLevel = Scene->GetFeatureLevel();
 	
-	// TEMP: optionally allow to switch off mesh command caching for cooked mobile
-	extern int32 MeshDrawCommandPipelineMobileCooked();
-	const int32 CacheMeshCommandsVal = MeshDrawCommandPipelineMobileCooked();
-
 	if (!PrimitiveSceneInfo->Proxy->ShouldRenderInMainPass() || !ShouldIncludeDomainInMeshPass(MaterialRenderProxy->GetMaterial(FeatureLevel)->GetMaterialDomain()))
 	{
 		return;
 	}
 
-	if (CastShadow && CacheMeshCommandsVal < 2)
+	if (CastShadow)
 	{
 		FShadowDepthDrawingPolicyFactory::AddStaticMesh(Scene, this);
 	}
-
-	if (IsTranslucent(FeatureLevel))
-	{
-		return;
-	}
-
-	if (Scene->GetShadingPath() == EShadingPath::Mobile && CacheMeshCommandsVal < 2)
-	{
-		if (bUseForMaterial)
-		{
-			// Add the static mesh to the DPG's base pass draw list.
-			FMobileBasePassOpaqueDrawingPolicyFactory::AddStaticMesh(RHICmdList, Scene, this);
-		}
-	}
 }
 
 void FStaticMesh::RemoveFromDrawLists(bool bMeshIsBeingDestroyed)
diff --git a/Engine/Source/Runtime/Renderer/Private/ScenePrivate.h b/Engine/Source/Runtime/Renderer/Private/ScenePrivate.h
--- a/Engine/Source/Runtime/Renderer/Private/ScenePrivate.h
+++ b/Engine/Source/Runtime/Renderer/Private/ScenePrivate.h
@@ -2450,15 +2450,6 @@ public:
 	/** Draw list used for rendering whole scene reflective shadow maps.  */
 	TStaticMeshDrawList<FShadowDepthDrawingPolicy<true> > WholeSceneReflectiveShadowMapDrawList;
 
-	/** Mobile base pass draw lists */
-	TStaticMeshDrawList<FMobileBasePassUniformDrawingPolicy> MobileBasePassUniformLightMapPolicyDrawList[EBasePass_MAX];
-	TStaticMeshDrawList<FMobileBasePassUniformDrawingPolicy> MobileBasePassUniformLightMapPolicyDrawListWithCSM[EBasePass_MAX];
-
-
-	/** Maps a light-map type to the appropriate base pass draw list. */
-	TStaticMeshDrawList<FMobileBasePassUniformDrawingPolicy>& GetMobileBasePassDrawList(EBasePassDrawListType DrawType);
-	TStaticMeshDrawList<FMobileBasePassUniformDrawingPolicy>& GetMobileBasePassCSMDrawList(EBasePassDrawListType DrawType);
-
 	/**
 	 * The following arrays are densely packed primitive data needed by various
 	 * rendering passes. PrimitiveSceneInfo->PackedIndex maintains the index
diff --git a/Engine/Source/Runtime/Renderer/Private/SceneRendering.cpp b/Engine/Source/Runtime/Renderer/Private/SceneRendering.cpp
--- a/Engine/Source/Runtime/Renderer/Private/SceneRendering.cpp
+++ b/Engine/Source/Runtime/Renderer/Private/SceneRendering.cpp
@@ -54,23 +54,6 @@
 	Globals
 -----------------------------------------------------------------------------*/
 
-// TEMP: remove after new path is tested on mobile
-static TAutoConsoleVariable<int32> CVarMeshDrawCommandPipelineMobileCooked(
-	TEXT("r.MeshDrawCommandPipelineMobileCooked"),
-	1,
-	TEXT("  0: only old path\n")
-	TEXT("  1: cache both \n")
-	TEXT("  2: only new path\n"),
-	ECVF_ReadOnly | ECVF_RenderThreadSafe);
-
-int32 MeshDrawCommandPipelineMobileCooked()
-{
-#if PLATFORM_ANDROID || PLATFORM_IOS
-	return CVarMeshDrawCommandPipelineMobileCooked.GetValueOnAnyThread();
-#endif
-	return 1;
-}
-
 static TAutoConsoleVariable<int32> CVarMeshDrawCommandPipeline(
 	TEXT("r.MeshDrawCommandPipeline"),
 	1,
@@ -79,7 +62,7 @@ static TAutoConsoleVariable<int32> CVarMeshDrawCommandPipeline(
 
 bool UseMeshDrawCommandPipeline()
 {
-	return CVarMeshDrawCommandPipeline.GetValueOnRenderThread() > 0 && MeshDrawCommandPipelineMobileCooked() > 0;
+	return CVarMeshDrawCommandPipeline.GetValueOnRenderThread() > 0;
 }
 
 static TAutoConsoleVariable<int32> CVarCachedMeshDrawCommands(
diff --git a/Engine/Source/Runtime/Renderer/Private/SceneRendering.h b/Engine/Source/Runtime/Renderer/Private/SceneRendering.h
--- a/Engine/Source/Runtime/Renderer/Private/SceneRendering.h
+++ b/Engine/Source/Runtime/Renderer/Private/SceneRendering.h
@@ -351,12 +351,6 @@ public:
 		float SortKey;
 	};
 
-	/** 
-	* Draw all the primitives in this set for the mobile pipeline. 
-	*/
-	template <class TDrawingPolicyFactory>
-	void DrawPrimitivesForMobile(FRHICommandListImmediate& RHICmdList, const FViewInfo& View, const FDrawingPolicyRenderState& DrawRenderState, typename TDrawingPolicyFactory::ContextType& DrawingContext) const;
-
 	/**
 	* Insert a primitive to the translucency rendering list[s]
 	*/
@@ -1887,8 +1881,6 @@ public:
 
 	bool RenderInverseOpacity(FRHICommandListImmediate& RHICmdList, const FViewInfo& View);
 
-	void RenderMobileBasePassDynamicData(FRHICommandList& RHICmdList, const FViewInfo& View, const FDrawingPolicyRenderState& DrawRenderState, EBlendMode BlendMode, bool bWireFrame, int32 FirstElement = 0, int32 AfterLastElement = MAX_int32);
-
 protected:
 	/** Finds the visible dynamic shadows for each view. */
 	void InitDynamicShadows(FRHICommandListImmediate& RHICmdList);
@@ -1902,7 +1894,6 @@ protected:
 	void RenderMobileBasePass(FRHICommandListImmediate& RHICmdList, const TArrayView<const FViewInfo*> PassViews);
 
 	void RenderMobileEditorPrimitives(FRHICommandList& RHICmdList, const FViewInfo& View, const FDrawingPolicyRenderState& DrawRenderState);
-	void RenderMobileBasePassViewParallel(const FViewInfo& View, FRHICommandListImmediate& ParentCmdList, TArray<FViewInfo>& Views, const FDrawingPolicyRenderState& DrawRenderState);
 
 	/** Render modulated shadow projections in to the scene, loops over any unrendered shadows until all are processed.*/
 	void RenderModulatedShadowProjections(FRHICommandListImmediate& RHICmdList);
@@ -1934,9 +1925,6 @@ protected:
 	/** Copy scene color from the mobile multi-view render target array to side by side stereo scene color */
 	void CopyMobileMultiViewSceneColor(FRHICommandListImmediate& RHICmdList);
 
-	/** Render inverse opacity for the dynamic meshes. */
-	bool RenderInverseOpacityDynamic(FRHICommandListImmediate& RHICmdList, const FViewInfo& View, const FDrawingPolicyRenderState& DrawRenderState);
-
 	/** Will update the view custom data. */
 	void PostInitViewCustomData();
 
diff --git a/Engine/Source/Runtime/Renderer/Private/SceneVisibility.cpp b/Engine/Source/Runtime/Renderer/Private/SceneVisibility.cpp
--- a/Engine/Source/Runtime/Renderer/Private/SceneVisibility.cpp
+++ b/Engine/Source/Runtime/Renderer/Private/SceneVisibility.cpp
@@ -2065,7 +2065,7 @@ struct FRelevancePacket
 			if (PrimitiveSceneInfo->NeedsReflectionCaptureUpdate())
 			{
 				// mobile should not have any outstanding reflection capture update requests at this point
-				ensure(!(Scene->GetShadingPath() == EShadingPath::Mobile && UseMeshDrawCommandPipeline()));
+				ensure(Scene->GetShadingPath() != EShadingPath::Mobile);
 		
 				PrimitiveSceneInfo->CacheReflectionCaptures();
 			}
@@ -3239,7 +3239,7 @@ void UpdateReflectionSceneData(FScene* Scene)
 		if (Scene->ReflectionSceneData.bRegisteredReflectionCapturesHasChanged)
 		{
 			// Mobile needs to re-cache all mesh commands when scene capture data has changed
-			const bool bNeedsStaticMeshUpdate = Scene->GetShadingPath() == EShadingPath::Mobile && UseMeshDrawCommandPipeline();
+			const bool bNeedsStaticMeshUpdate = Scene->GetShadingPath() == EShadingPath::Mobile;
 			
 			// Mark all primitives as needing an update
 			// Note: Only visible primitives will actually update their reflection proxy
diff --git a/Engine/Source/Runtime/Renderer/Private/ShadowRendering.cpp b/Engine/Source/Runtime/Renderer/Private/ShadowRendering.cpp
--- a/Engine/Source/Runtime/Renderer/Private/ShadowRendering.cpp
+++ b/Engine/Source/Runtime/Renderer/Private/ShadowRendering.cpp
@@ -842,53 +842,14 @@ void FProjectedShadowInfo::SetupProjectionStencilMask(
 		// if rendering modulated shadows mask out subject mesh elements to prevent self shadowing.
 		if (bMobileModulatedProjections && !CVarEnableModulatedSelfShadow.GetValueOnRenderThread())
 		{
-			if (UseMeshDrawCommandPipeline())
-			{
-				const FShadowMeshDrawCommandPass& ProjectionStencilingPass = ProjectionStencilingPasses[ViewIndex];
-				if (ProjectionStencilingPass.VisibleMeshDrawCommands.Num() > 0)
-				{
-					FVertexBufferRHIParamRef PrimitiveIdsBuffer = UseGPUScene(GMaxRHIShaderPlatform, View->GetFeatureLevel()) ? ProjectionStencilingPass.PrimitiveIdBuffer.VertexBuffer->VertexBufferRHI : nullptr;
-					const int32 BasePrimitiveIdsOffset = ProjectionStencilingPass.PrimitiveIdBuffer.VertexOffset;
-					const bool bDynamicInstancing = IsDynamicInstancingEnabled() && UseGPUScene(GMaxRHIShaderPlatform, View->GetFeatureLevel());
-
-					SubmitMeshDrawCommands(ProjectionStencilingPass.VisibleMeshDrawCommands, PrimitiveIdsBuffer, BasePrimitiveIdsOffset, bDynamicInstancing, RHICmdList);
-				}
-			}
-			else
+			const FShadowMeshDrawCommandPass& ProjectionStencilingPass = ProjectionStencilingPasses[ViewIndex];
+			if (ProjectionStencilingPass.VisibleMeshDrawCommands.Num() > 0)
 			{
-				DrawRenderState.SetDepthStencilState(
-					TStaticDepthStencilState<
-					false, CF_DepthNearOrEqual,
-					true, CF_Always, SO_Keep, SO_Keep, SO_Replace,
-					true, CF_Always, SO_Keep, SO_Keep, SO_Replace,
-					0xff, 0xff
-					>::GetRHI());
-				DrawRenderState.SetStencilRef(0);
-
-				FDepthDrawingPolicyFactory::ContextType Context(DDM_AllOccluders, false);
-				for (int32 MeshBatchIndex = 0; MeshBatchIndex < DynamicSubjectMeshElements.Num(); MeshBatchIndex++)
-				{
-					const FMeshBatchAndRelevance& MeshBatchAndRelevance = DynamicSubjectMeshElements[MeshBatchIndex];
-					const FMeshBatch& MeshBatch = *MeshBatchAndRelevance.Mesh;
-					FDepthDrawingPolicyFactory::DrawDynamicMesh(RHICmdList, *View, Context, MeshBatch, true, DrawRenderState, MeshBatchAndRelevance.PrimitiveSceneProxy, MeshBatch.BatchHitProxyId);
-				}
+				FVertexBufferRHIParamRef PrimitiveIdsBuffer = UseGPUScene(GMaxRHIShaderPlatform, View->GetFeatureLevel()) ? ProjectionStencilingPass.PrimitiveIdBuffer.VertexBuffer->VertexBufferRHI : nullptr;
+				const int32 BasePrimitiveIdsOffset = ProjectionStencilingPass.PrimitiveIdBuffer.VertexOffset;
+				const bool bDynamicInstancing = IsDynamicInstancingEnabled() && UseGPUScene(GMaxRHIShaderPlatform, View->GetFeatureLevel());
 
-				for (int32 ElementIndex = 0; ElementIndex < StaticSubjectMeshElements.Num(); ++ElementIndex)
-				{
-					const FStaticMesh& StaticMesh = *StaticSubjectMeshElements[ElementIndex].Mesh;
-					FDepthDrawingPolicyFactory::DrawStaticMesh(
-						RHICmdList,
-						*View,
-						FDepthDrawingPolicyFactory::ContextType(DDM_AllOccluders, false),
-						StaticMesh,
-						StaticMesh.bRequiresPerElementVisibility ? View->StaticMeshBatchVisibility[StaticMesh.Id] : ((1ull << StaticMesh.Elements.Num()) - 1),
-						true,
-						DrawRenderState,
-						StaticMesh.PrimitiveSceneInfo->Proxy,
-						StaticMesh.BatchHitProxyId,
-						false
-					);
-				}
+				SubmitMeshDrawCommands(ProjectionStencilingPass.VisibleMeshDrawCommands, PrimitiveIdsBuffer, BasePrimitiveIdsOffset, bDynamicInstancing, RHICmdList);
 			}
 		}
 	}
