Fixed calculation of number of dynamic mesh elements and added more information to be able to track down potential bugs in this code.

diff --git a/Engine/Source/Runtime/Engine/Private/SceneManagement.cpp b/Engine/Source/Runtime/Engine/Private/SceneManagement.cpp
--- a/Engine/Source/Runtime/Engine/Private/SceneManagement.cpp
+++ b/Engine/Source/Runtime/Engine/Private/SceneManagement.cpp
@@ -293,6 +293,8 @@ void FMeshElementCollector::AddMesh(int32 ViewIndex, FMeshBatch& MeshBatch)
 	MeshBatch.MeshIdInPrimitive = MeshIdInPrimitivePerView[ViewIndex];
 	++MeshIdInPrimitivePerView[ViewIndex];
 
+	NumMeshBatchElementsPerView[ViewIndex] += MeshBatch.Elements.Num();
+
 	TArray<FMeshBatchAndRelevance,SceneRenderingAllocator>& ViewMeshBatches = *MeshBatches[ViewIndex];
 	new (ViewMeshBatches) FMeshBatchAndRelevance(MeshBatch, PrimitiveSceneProxy, FeatureLevel);	
 }
diff --git a/Engine/Source/Runtime/Engine/Public/SceneManagement.h b/Engine/Source/Runtime/Engine/Public/SceneManagement.h
--- a/Engine/Source/Runtime/Engine/Public/SceneManagement.h
+++ b/Engine/Source/Runtime/Engine/Public/SceneManagement.h
@@ -1948,6 +1948,12 @@ public:
 		return MeshBatches[ViewIndex]->Num();
 	}
 
+	// @return Number of elemenets collected so far for a given view.
+	uint32 GetMeshElementCount(uint32 ViewIndex) const
+	{
+		return NumMeshBatchElementsPerView[ViewIndex];
+	}
+
 	/** 
 	 * Adds a mesh batch to the collector for the specified view so that it can be rendered.
 	 */
@@ -2033,6 +2039,7 @@ private:
 		SimpleElementCollectors.Empty();
 		MeshIdInPrimitivePerView.Empty();
 		DynamicPrimitiveShaderDataPerView.Empty();
+		NumMeshBatchElementsPerView.Empty();
 	}
 
 	void AddViewMeshArrays(
@@ -2045,6 +2052,7 @@ private:
 		Views.Add(InView);
 		MeshIdInPrimitivePerView.Add(0);
 		MeshBatches.Add(ViewMeshes);
+		NumMeshBatchElementsPerView.Add(0);
 		SimpleElementCollectors.Add(ViewSimpleElementCollector);
 		DynamicPrimitiveShaderDataPerView.Add(InDynamicPrimitiveShaderData);
 	}
@@ -2058,6 +2066,9 @@ private:
 	/** Meshes to render */
 	TArray<TArray<FMeshBatchAndRelevance, SceneRenderingAllocator>*, TInlineAllocator<2> > MeshBatches;
 
+	/** Number of elements in gathered meshes per view. */
+	TArray<int32, TInlineAllocator<2> > NumMeshBatchElementsPerView;
+
 	/** PDIs */
 	TArray<FSimpleElementCollector*, TInlineAllocator<2> > SimpleElementCollectors;
 
diff --git a/Engine/Source/Runtime/Renderer/Private/MobileShadingRenderer.cpp b/Engine/Source/Runtime/Renderer/Private/MobileShadingRenderer.cpp
--- a/Engine/Source/Runtime/Renderer/Private/MobileShadingRenderer.cpp
+++ b/Engine/Source/Runtime/Renderer/Private/MobileShadingRenderer.cpp
@@ -141,6 +141,7 @@ void FMobileSceneRenderer::SetupMobileBasePassAfterShadowInit(FExclusiveDepthSte
 			View.DynamicMeshElements,
 			View.NumVisibleDynamicMeshElements[EMeshPass::BasePass],
 			ViewCommands.DynamicMeshCommandBuildRequests[EMeshPass::BasePass],
+			ViewCommands.NumDynamicMeshCommandBuildRequestElements[EMeshPass::BasePass],
 			ViewCommands.MeshCommands[EMeshPass::BasePass],
 			BasePassCSMMeshPassProcessor,
 			&ViewCommands.MeshCommands[EMeshPass::MobileBasePassCSM]);
diff --git a/Engine/Source/Runtime/Renderer/Private/SceneCore.h b/Engine/Source/Runtime/Renderer/Private/SceneCore.h
--- a/Engine/Source/Runtime/Renderer/Private/SceneCore.h
+++ b/Engine/Source/Runtime/Renderer/Private/SceneCore.h
@@ -201,6 +201,7 @@ public:
 		, ScreenSize(InScreenSize)
 		, CommandInfosBase(0)
 		, LODIndex(StaticMesh.LODIndex)
+		, NumElements(StaticMesh.Elements.Num())
 		, bDitheredLODTransition(StaticMesh.bDitheredLODTransition)
 		, bRequiresPerElementVisibility(StaticMesh.bRequiresPerElementVisibility)
 		, bSelectable(StaticMesh.bSelectable)
@@ -227,6 +228,9 @@ public:
 	/** LOD index of the mesh, used for fading LOD transitions. */
 	int8 LODIndex;
 
+	/** Number of elements in this mesh. */
+	uint16 NumElements;
+
 	/** Whether the mesh batch should apply dithered LOD. */
 	uint8 bDitheredLODTransition : 1;
 
diff --git a/Engine/Source/Runtime/Renderer/Private/SceneRendering.cpp b/Engine/Source/Runtime/Renderer/Private/SceneRendering.cpp
--- a/Engine/Source/Runtime/Renderer/Private/SceneRendering.cpp
+++ b/Engine/Source/Runtime/Renderer/Private/SceneRendering.cpp
@@ -905,7 +905,11 @@ void FViewInfo::Init()
 	DitherFadeOutUniformBuffer = nullptr;
 
 	VisibleDynamicMeshesPassMask.Reset();
-	FMemory::Memset(NumVisibleDynamicMeshElements, 0, sizeof(NumVisibleDynamicMeshElements));
+
+	for (int32 PassIndex = 0; PassIndex < EMeshPass::Num; ++PassIndex)
+	{
+		NumVisibleDynamicMeshElements[PassIndex] = 0;
+	}
 }
 
 FViewInfo::~FViewInfo()
@@ -2990,7 +2994,9 @@ void GenerateDynamicMeshDrawCommands(
 	EMeshPass::Type PassType,
 	FMeshPassProcessor* PassMeshProcessor,
 	const TArray<FMeshBatchAndRelevance, SceneRenderingAllocator>& DynamicMeshElements,
+	int32 NumDynamicMeshElements,
 	const TArray<const FStaticMesh*, SceneRenderingAllocator>& DynamicMeshCommandBuildRequests,
+	int32 NumDynamicMeshCommandBuildRequestElements,
 	FMeshCommandOneFrameArray& VisibleCommands,
 	FDynamicMeshDrawCommandStorage& MeshDrawCommandStorage
 )
@@ -3004,6 +3010,7 @@ void GenerateDynamicMeshDrawCommands(
 	);
 	PassMeshProcessor->SetDrawListContext(&DynamicPassMeshDrawListContext);
 
+	const int32 NumVisibleCommandsBeforeDynamicBatches = VisibleCommands.Num();
 	const int32 NumDynamicMeshBatches = DynamicMeshElements.Num();
 
 	for (int32 MeshIndex = 0; MeshIndex < NumDynamicMeshBatches; MeshIndex++)
@@ -3015,6 +3022,10 @@ void GenerateDynamicMeshDrawCommands(
 		PassMeshProcessor->AddMeshBatch(*MeshAndRelevance.Mesh, BatchElementMask, MeshAndRelevance.PrimitiveSceneProxy);
 	}
 
+	// Check if we didn't generate more mesh draw commands, than we assumed.
+	check(VisibleCommands.Num() - NumVisibleCommandsBeforeDynamicBatches <= NumDynamicMeshElements);
+
+	const int32 NumVisibleCommandsBeforeDynamicRequests = VisibleCommands.Num();
 	const int32 NumStaticMeshBatches = DynamicMeshCommandBuildRequests.Num();
 
 	for (int32 MeshIndex = 0; MeshIndex < NumStaticMeshBatches; MeshIndex++)
@@ -3024,6 +3035,9 @@ void GenerateDynamicMeshDrawCommands(
 
 		PassMeshProcessor->AddMeshBatch(*StaticMeshBatch, BatchElementMask, StaticMeshBatch->PrimitiveSceneInfo->Proxy, StaticMeshBatch->Id);
 	}
+
+	// Check if we didn't generate more mesh draw commands, than we assumed.
+	check(VisibleCommands.Num() - NumVisibleCommandsBeforeDynamicRequests <= NumDynamicMeshCommandBuildRequestElements);
 }
 
 /**
@@ -3037,7 +3051,9 @@ void GenerateMobileBasePassDynamicMeshDrawCommands(
 	FMeshPassProcessor* PassMeshProcessor,
 	FMeshPassProcessor* MobilePassCSMPassMeshProcessor,
 	const TArray<FMeshBatchAndRelevance, SceneRenderingAllocator>& DynamicMeshElements,
+	int32 NumDynamicMeshElements,
 	const TArray<const FStaticMesh*, SceneRenderingAllocator>& DynamicMeshCommandBuildRequests,
+	int32 NumDynamicMeshCommandBuildRequestElements,
 	FMeshCommandOneFrameArray& VisibleCommands,
 	FDynamicMeshDrawCommandStorage& MeshDrawCommandStorage
 )
@@ -3054,6 +3070,7 @@ void GenerateMobileBasePassDynamicMeshDrawCommands(
 
 	const FMobileCSMVisibilityInfo& MobileCSMVisibilityInfo = View.MobileCSMVisibilityInfo;
 	
+	const int32 NumVisibleCommandsBeforeDynamicBatches = VisibleCommands.Num();
 	const int32 NumDynamicMeshBatches = DynamicMeshElements.Num();
 
 	for (int32 MeshIndex = 0; MeshIndex < NumDynamicMeshBatches; MeshIndex++)
@@ -3074,6 +3091,10 @@ void GenerateMobileBasePassDynamicMeshDrawCommands(
 		}
 	}
 
+	// Check if we didn't generate more mesh draw commands, than we assumed.
+	check(VisibleCommands.Num() - NumVisibleCommandsBeforeDynamicBatches <= NumDynamicMeshElements);
+
+	const int32 NumVisibleCommandsBeforeDynamicRequests = VisibleCommands.Num();
 	const int32 NumStaticMeshBatches = DynamicMeshCommandBuildRequests.Num();
 
 	for (int32 MeshIndex = 0; MeshIndex < NumStaticMeshBatches; MeshIndex++)
@@ -3092,6 +3113,9 @@ void GenerateMobileBasePassDynamicMeshDrawCommands(
 			PassMeshProcessor->AddMeshBatch(*StaticMeshBatch, BatchElementMask, StaticMeshBatch->PrimitiveSceneInfo->Proxy, StaticMeshBatch->Id);
 		}
 	}
+
+	// Check if we didn't generate more mesh draw commands, than we assumed.
+	check(VisibleCommands.Num() - NumVisibleCommandsBeforeDynamicRequests <= NumDynamicMeshCommandBuildRequestElements);
 }
 
 /**
@@ -3214,7 +3238,9 @@ public:
 				Context.MeshPassProcessor,
 				Context.MobileBasePassCSMMeshPassProcessor,
 				*Context.DynamicMeshElements,
+				Context.NumDynamicMeshElements,
 				Context.DynamicMeshCommandBuildRequests,
+				Context.NumDynamicMeshCommandBuildRequestElements,
 				Context.MeshDrawCommands,
 				Context.MeshDrawCommandStorage
 			);
@@ -3227,15 +3253,14 @@ public:
 				Context.PassType,
 				Context.MeshPassProcessor,
 				*Context.DynamicMeshElements,
+				Context.NumDynamicMeshElements,
 				Context.DynamicMeshCommandBuildRequests,
+				Context.NumDynamicMeshCommandBuildRequestElements,
 				Context.MeshDrawCommands,
 				Context.MeshDrawCommandStorage
 			);
 		}
 
-		// Check if we didn't generate more mesh draw commands, than we assumed in FParallelMeshDrawCommandPass::MaxNumDraws.
-		check(Context.MeshDrawCommands.Num() <= Context.TempVisibleMeshDrawCommands.Max());
-
 		if (Context.MeshDrawCommands.Num() > 0)
 		{
 			if (Context.PassType != EMeshPass::Num)
@@ -3356,8 +3381,9 @@ void FParallelMeshDrawCommandPass::DispatchPassSetup(
 	FExclusiveDepthStencil::Type BasePassDepthStencilAccess,
 	FMeshPassProcessor* MeshPassProcessor,
 	const TArray<FMeshBatchAndRelevance, SceneRenderingAllocator>& DynamicMeshElements,
-	int32 MaxNumDynamicMeshElementsForThisPass,
+	int32 NumDynamicMeshElements,
 	TArray<const FStaticMesh*, SceneRenderingAllocator>& InOutDynamicMeshCommandBuildRequests,
+	int32 NumDynamicMeshCommandBuildRequestElements,
 	FMeshCommandOneFrameArray& InOutMeshDrawCommands,
 	FMeshPassProcessor* MobileBasePassCSMMeshPassProcessor,
 	FMeshCommandOneFrameArray* InOutMobileBasePassCSMMeshDrawCommands
@@ -3365,7 +3391,7 @@ void FParallelMeshDrawCommandPass::DispatchPassSetup(
 {
 	check(TaskEventRefs.Num() == 0 && MeshPassProcessor != nullptr);
 
-	MaxNumDraws = InOutMeshDrawCommands.Num() + InOutDynamicMeshCommandBuildRequests.Num() + MaxNumDynamicMeshElementsForThisPass;
+	MaxNumDraws = InOutMeshDrawCommands.Num() + NumDynamicMeshElements + NumDynamicMeshCommandBuildRequestElements;
 	if (MaxNumDraws <= 0)
 	{
 		return;
@@ -3383,6 +3409,8 @@ void FParallelMeshDrawCommandPass::DispatchPassSetup(
 	TaskContext.bReverseCulling = View.bReverseCulling;
 	TaskContext.bRenderSceneTwoSided = View.bRenderSceneTwoSided;
 	TaskContext.BasePassDepthStencilAccess = BasePassDepthStencilAccess;
+	TaskContext.NumDynamicMeshElements = NumDynamicMeshElements;
+	TaskContext.NumDynamicMeshCommandBuildRequestElements = NumDynamicMeshCommandBuildRequestElements;
 
 	// Setup translucency sort key update pass based on view.
 	TaskContext.TranslucencyPass = ETranslucencyPass::TPT_MAX;
@@ -3622,6 +3650,7 @@ void FSceneRenderer::SetupMeshPass(FViewInfo& View, FExclusiveDepthStencil::Type
 				View.DynamicMeshElements,
 				View.NumVisibleDynamicMeshElements[PassType],
 				ViewCommands.DynamicMeshCommandBuildRequests[PassType],
+				ViewCommands.NumDynamicMeshCommandBuildRequestElements[PassType],
 				ViewCommands.MeshCommands[PassIndex]);
 		}
 	}
diff --git a/Engine/Source/Runtime/Renderer/Private/SceneRendering.h b/Engine/Source/Runtime/Renderer/Private/SceneRendering.h
--- a/Engine/Source/Runtime/Renderer/Private/SceneRendering.h
+++ b/Engine/Source/Runtime/Renderer/Private/SceneRendering.h
@@ -879,7 +879,16 @@ struct FPreviousViewInfo
 class FViewCommands
 {
 public:
+	FViewCommands()
+	{
+		for (int32 PassIndex = 0; PassIndex < EMeshPass::Num; ++PassIndex)
+		{
+			NumDynamicMeshCommandBuildRequestElements[PassIndex] = 0;
+		}
+	}
+
 	TStaticArray<FMeshCommandOneFrameArray, EMeshPass::Num> MeshCommands;
+	TStaticArray<int32, EMeshPass::Num> NumDynamicMeshCommandBuildRequestElements;
 	TStaticArray<TArray<const FStaticMesh*, SceneRenderingAllocator>, EMeshPass::Num> DynamicMeshCommandBuildRequests;
 };
 
@@ -901,7 +910,10 @@ public:
 		, bRenderSceneTwoSided(false)
 		, BasePassDepthStencilAccess(FExclusiveDepthStencil::DepthNop_StencilNop)
 		, MeshPassProcessor(nullptr)
+		, MobileBasePassCSMMeshPassProcessor(nullptr)
 		, DynamicMeshElements(nullptr)
+		, NumDynamicMeshElements(0)
+		, NumDynamicMeshCommandBuildRequestElements(0)
 		, PrimitiveBounds(nullptr)
 		, VisibleMeshDrawCommandsNum(0)
 		, NewPassVisibleMeshDrawCommandsNum(0)
@@ -924,6 +936,8 @@ public:
 	const TArray<FMeshBatchAndRelevance, SceneRenderingAllocator>* DynamicMeshElements;
 
 	// Commands.
+	int32 NumDynamicMeshElements;
+	int32 NumDynamicMeshCommandBuildRequestElements;
 	FMeshCommandOneFrameArray MeshDrawCommands;
 	FMeshCommandOneFrameArray MobileBasePassCSMMeshDrawCommands;
 	TArray<const FStaticMesh*, SceneRenderingAllocator> DynamicMeshCommandBuildRequests;
@@ -973,8 +987,9 @@ public:
 		FExclusiveDepthStencil::Type BasePassDepthStencilAccess,
 		FMeshPassProcessor* MeshPassProcessor,
 		const TArray<FMeshBatchAndRelevance, SceneRenderingAllocator>& DynamicMeshElements,
-		int32 MaxNumDynamicMeshElementsForThisPass,
+		int32 NumDynamicMeshElements,
 		TArray<const FStaticMesh*, SceneRenderingAllocator>& InOutDynamicMeshCommandBuildRequests,
+		int32 NumDynamicMeshCommandBuildRequestElements,
 		FMeshCommandOneFrameArray& InOutMeshDrawCommands,
 		FMeshPassProcessor* MobileBasePassCSMMeshPassProcessor = nullptr, // Required only for the mobile base pass.
 		FMeshCommandOneFrameArray* InOutMobileBasePassCSMMeshDrawCommands = nullptr // Required only for the mobile base pass.
diff --git a/Engine/Source/Runtime/Renderer/Private/SceneVisibility.cpp b/Engine/Source/Runtime/Renderer/Private/SceneVisibility.cpp
--- a/Engine/Source/Runtime/Renderer/Private/SceneVisibility.cpp
+++ b/Engine/Source/Runtime/Renderer/Private/SceneVisibility.cpp
@@ -1684,10 +1684,16 @@ struct FDrawCommandRelevancePacket
 	{
 		bUseMeshDrawCommandPipeline = UseMeshDrawCommandPipeline();
 		bUseCachedMeshDrawCommands = UseCachedMeshDrawCommands();
+
+		for (int32 PassIndex = 0; PassIndex < EMeshPass::Num; ++PassIndex)
+		{
+			NumDynamicBuildRequestElements[PassIndex] = 0;
+		}
 	}
 
 	FPassDrawCommandArray VisibleCachedDrawCommands[EMeshPass::Num];
 	FPassDrawCommandBuildRequestArray DynamicBuildRequests[EMeshPass::Num];
+	int32 NumDynamicBuildRequestElements[EMeshPass::Num];
 	bool bUseMeshDrawCommandPipeline;
 	bool bUseCachedMeshDrawCommands;
 
@@ -1731,6 +1737,7 @@ struct FDrawCommandRelevancePacket
 			}
 			else
 			{
+				NumDynamicBuildRequestElements[PassType] += StaticMeshRelevance.NumElements;
 				DynamicBuildRequests[PassType].AddElement(&StaticMesh);
 			}
 		}
@@ -2366,6 +2373,7 @@ struct FRelevancePacket
 		{
 			DrawCommandPacket.VisibleCachedDrawCommands[PassIndex].CopyToLinearArray(WriteViewCommands.MeshCommands[PassIndex]);
 			DrawCommandPacket.DynamicBuildRequests[PassIndex].CopyToLinearArray(WriteViewCommands.DynamicMeshCommandBuildRequests[PassIndex]);
+			WriteViewCommands.NumDynamicMeshCommandBuildRequestElements[PassIndex] += DrawCommandPacket.NumDynamicBuildRequestElements[PassIndex];
 		}
 
 		// Prepare translucent self shadow uniform buffers.
@@ -2625,7 +2633,7 @@ void FSceneRenderer::GatherDynamicMeshElements(
 		{
 			Collector.AddViewMeshArrays(
 				&InViews[ViewIndex], 
-				&InViews[ViewIndex].DynamicMeshElements, 
+				&InViews[ViewIndex].DynamicMeshElements,
 				&InViews[ViewIndex].SimpleElementCollector,
 				&InViews[ViewIndex].DynamicPrimitiveShaderData, 
 				InViewFamily.GetFeatureLevel());
@@ -2656,6 +2664,20 @@ void FSceneRenderer::GatherDynamicMeshElements(
 				InViews[ViewIndex].DynamicMeshEndIndices[PrimitiveIndex] = Collector.GetMeshBatchCount(ViewIndex);
 			}
 		}
+
+		// Compute MaxNumVisibleDynamicMeshes.
+		for (int32 ViewIndex = 0; ViewIndex < Views.Num(); ++ViewIndex)
+		{
+			FViewInfo& View = Views[ViewIndex];
+
+			for (int32 PassIndex = 0; PassIndex < EMeshPass::Num; ++PassIndex)
+			{
+				if (View.VisibleDynamicMeshesPassMask.Get((EMeshPass::Type) PassIndex))
+				{
+					View.NumVisibleDynamicMeshElements[PassIndex] = Collector.GetMeshElementCount(ViewIndex);
+				}
+			}
+		}
 	}
 
 	if (GIsEditor)
@@ -2688,27 +2710,6 @@ void FSceneRenderer::GatherDynamicMeshElements(
 		}
 	}
 	MeshCollector.ProcessTasks();
-
-	// Compute MaxNumVisibleDynamicMeshes.
-	for (int32 ViewIndex = 0; ViewIndex < Views.Num(); ++ViewIndex)
-	{
-		FViewInfo& View = Views[ViewIndex];
-
-		int32 NumDynamicMeshElements = 0;
-		for (int32 DynamicElementIndex = 0; DynamicElementIndex < View.DynamicMeshElements.Num(); ++DynamicElementIndex)
-		{
-			NumDynamicMeshElements += View.DynamicMeshElements[DynamicElementIndex].Mesh->Elements.Num();
-		}
-
-		for (int32 PassIndex = 0; PassIndex < EMeshPass::Num; ++PassIndex)
-		{
-			View.NumVisibleDynamicMeshElements[PassIndex] = 0;
-			if (View.VisibleDynamicMeshesPassMask.Get((EMeshPass::Type) PassIndex))
-			{
-				View.NumVisibleDynamicMeshElements[PassIndex] = NumDynamicMeshElements;
-			}
-		}
-	}
 }
 
 /**
diff --git a/Engine/Source/Runtime/Renderer/Private/ShadowRendering.h b/Engine/Source/Runtime/Renderer/Private/ShadowRendering.h
--- a/Engine/Source/Runtime/Renderer/Private/ShadowRendering.h
+++ b/Engine/Source/Runtime/Renderer/Private/ShadowRendering.h
@@ -734,6 +734,12 @@ private:
 
 	TArray<const FStaticMesh*, SceneRenderingAllocator> SubjectMeshCommandBuildRequests;
 
+	/** Number of elements of DynamicSubjectMeshElements meshes. */
+	int32 NumDynamicSubjectMeshElements;
+
+	/** Number of elements of SubjectMeshCommandBuildRequests meshes. */
+	int32 NumSubjectMeshCommandBuildRequestElements;
+
 	FMeshCommandOneFrameArray ShadowDepthPassVisibleCommands;
 	FParallelMeshDrawCommandPass ShadowDepthPass;
 
@@ -796,7 +802,8 @@ private:
 		FSceneRenderer& Renderer, 
 		const PrimitiveArrayType& PrimitiveArray, 
 		const TArray<const FSceneView*>& ReusedViewsArray,
-		TArray<FMeshBatchAndRelevance,SceneRenderingAllocator>& OutDynamicMeshElements);
+		TArray<FMeshBatchAndRelevance,SceneRenderingAllocator>& OutDynamicMeshElements,
+		int32& OutNumDynamicSubjectMeshElements);
 
 	void SetupFrustumForProjection(const FViewInfo* View, TArray<FVector4, TInlineAllocator<8>>& OutFrustumVertices, bool& bOutCameraInsideShadowFrustum) const;
 
diff --git a/Engine/Source/Runtime/Renderer/Private/ShadowSetup.cpp b/Engine/Source/Runtime/Renderer/Private/ShadowSetup.cpp
--- a/Engine/Source/Runtime/Renderer/Private/ShadowSetup.cpp
+++ b/Engine/Source/Runtime/Renderer/Private/ShadowSetup.cpp
@@ -515,6 +515,8 @@ FProjectedShadowInfo::FProjectedShadowInfo()
 	, bTransmission(false)
 	, LightSceneInfo(0)
 	, ParentSceneInfo(0)
+	, NumDynamicSubjectMeshElements(0)
+	, NumSubjectMeshCommandBuildRequestElements(0)
 	, ShaderDepthBias(0.0f)
 {
 }
@@ -935,6 +937,7 @@ bool FProjectedShadowInfo::ShouldDrawStaticMeshes(FViewInfo& InCurrentView, bool
 						}
 						else
 						{
+							NumSubjectMeshCommandBuildRequestElements += StaticMeshRelevance.NumElements;
 							SubjectMeshCommandBuildRequests.Add(&StaticMesh);
 						}
 					}
@@ -963,6 +966,7 @@ bool FProjectedShadowInfo::ShouldDrawStaticMeshes(FViewInfo& InCurrentView, bool
 				{
 					if (UseMeshDrawCommandPipeline())
 					{
+						NumSubjectMeshCommandBuildRequestElements += StaticMeshRelevance.NumElements;
 						SubjectMeshCommandBuildRequests.Add(&StaticMesh);
 					}
 					else if (StaticMeshRelevance.CastShadow)
@@ -1223,8 +1227,9 @@ void FProjectedShadowInfo::SetupMeshDrawCommandsForShadowDepth(FSceneRenderer& R
 		FExclusiveDepthStencil::DepthNop_StencilNop,
 		MeshPassProcessor,
 		DynamicSubjectMeshElements,
-		DynamicSubjectMeshElements.Num(),
+		NumDynamicSubjectMeshElements,
 		SubjectMeshCommandBuildRequests,
+		NumSubjectMeshCommandBuildRequestElements,
 		ShadowDepthPassVisibleCommands);
 }
 
@@ -1410,23 +1415,25 @@ void FProjectedShadowInfo::GatherDynamicMeshElements(FSceneRenderer& Renderer, F
 		{
 			ShadowDepthView->SetPreShadowTranslation(FVector(0, 0, 0));
 			ShadowDepthView->SetDynamicMeshElementsShadowCullFrustum(&CascadeSettings.ShadowBoundsAccurate);
-			GatherDynamicMeshElementsArray(ShadowDepthView, Renderer, DynamicSubjectPrimitives, ReusedViewsArray, DynamicSubjectMeshElements);
+			GatherDynamicMeshElementsArray(ShadowDepthView, Renderer, DynamicSubjectPrimitives, ReusedViewsArray, DynamicSubjectMeshElements, NumDynamicSubjectMeshElements);
 			ShadowDepthView->SetPreShadowTranslation(PreShadowTranslation);
 		}
 		else
 		{
 			ShadowDepthView->SetPreShadowTranslation(PreShadowTranslation);
 			ShadowDepthView->SetDynamicMeshElementsShadowCullFrustum(&CasterFrustum);
-			GatherDynamicMeshElementsArray(ShadowDepthView, Renderer, DynamicSubjectPrimitives, ReusedViewsArray, DynamicSubjectMeshElements);
+			GatherDynamicMeshElementsArray(ShadowDepthView, Renderer, DynamicSubjectPrimitives, ReusedViewsArray, DynamicSubjectMeshElements, NumDynamicSubjectMeshElements);
 		}
 
 		ShadowDepthView->DrawDynamicFlags = EDrawDynamicFlags::None;
 
+		int32 NumDynamicReceiverMeshElements = 0;
 		ShadowDepthView->SetDynamicMeshElementsShadowCullFrustum(&ReceiverFrustum);
-		GatherDynamicMeshElementsArray(ShadowDepthView, Renderer, ReceiverPrimitives, ReusedViewsArray, DynamicReceiverMeshElements);
+		GatherDynamicMeshElementsArray(ShadowDepthView, Renderer, ReceiverPrimitives, ReusedViewsArray, DynamicReceiverMeshElements, NumDynamicReceiverMeshElements);
 
+		int32 NumDynamicSubjectTranslucentMeshElements = 0;
 		ShadowDepthView->SetDynamicMeshElementsShadowCullFrustum(&CasterFrustum);
-		GatherDynamicMeshElementsArray(ShadowDepthView, Renderer, SubjectTranslucentPrimitives, ReusedViewsArray, DynamicSubjectTranslucentMeshElements);
+		GatherDynamicMeshElementsArray(ShadowDepthView, Renderer, SubjectTranslucentPrimitives, ReusedViewsArray, DynamicSubjectTranslucentMeshElements, NumDynamicSubjectTranslucentMeshElements);
 
 		Renderer.MeshCollector.ProcessTasks();
     }
@@ -1460,7 +1467,8 @@ void FProjectedShadowInfo::GatherDynamicMeshElementsArray(
 	FSceneRenderer& Renderer, 
 	const PrimitiveArrayType& PrimitiveArray, 
 	const TArray<const FSceneView*>& ReusedViewsArray,
-	TArray<FMeshBatchAndRelevance,SceneRenderingAllocator>& OutDynamicMeshElements)
+	TArray<FMeshBatchAndRelevance,SceneRenderingAllocator>& OutDynamicMeshElements,
+	int32& OutNumDynamicSubjectMeshElements)
 {
 	// Simple elements not supported in shadow passes
 	FSimpleElementCollector DynamicSubjectSimpleElements;
@@ -1504,6 +1512,8 @@ void FProjectedShadowInfo::GatherDynamicMeshElementsArray(
 			PrimitiveSceneInfo->Proxy->GetDynamicMeshElements(ReusedViewsArray, Renderer.ViewFamily, 0x1, Renderer.MeshCollector);
 		}
 	}
+
+	OutNumDynamicSubjectMeshElements = Renderer.MeshCollector.GetMeshElementCount(0);
 }
 
 /** 
@@ -1530,6 +1540,9 @@ bool FProjectedShadowInfo::SubjectsVisible(const FViewInfo& View) const
  */
 void FProjectedShadowInfo::ClearTransientArrays()
 {
+	NumDynamicSubjectMeshElements = 0;
+	NumSubjectMeshCommandBuildRequestElements = 0;
+
 	SubjectTranslucentPrimitives.Empty();
 	DynamicSubjectPrimitives.Empty();
 	ReceiverPrimitives.Empty();
