Moved mesh draw command code related code to MeshDrawCommands.cpp/.h.

diff --git a/Engine/Source/Runtime/Renderer/Private/MeshDrawCommands.cpp b/Engine/Source/Runtime/Renderer/Private/MeshDrawCommands.cpp
new file mode 100644
--- /dev/null
+++ b/Engine/Source/Runtime/Renderer/Private/MeshDrawCommands.cpp
@@ -0,0 +1,966 @@
+// Copyright 1998-2019 Epic Games, Inc. All Rights Reserved.
+
+/*=============================================================================
+MeshDrawCommandSetup.cpp: Mesh draw command setup.
+=============================================================================*/
+
+#include "MeshDrawCommands.h"
+#include "RendererModule.h"
+#include "ScenePrivate.h"
+#include "TranslucentRendering.h"
+
+static TAutoConsoleVariable<int32> CVarMeshDrawCommandsParallelPassSetup(
+	TEXT("r.MeshDrawCommands.ParallelPassSetup"),
+	1,
+	TEXT("Whether to setup mesh draw command pass in parallel."),
+	ECVF_RenderThreadSafe);
+
+struct FCompareFMeshDrawCommands
+{
+	FORCEINLINE bool operator() (FVisibleMeshDrawCommand A, FVisibleMeshDrawCommand B) const
+	{
+		// First order by a sort key.
+		if (A.SortKey != B.SortKey)
+		{
+			return A.SortKey < B.SortKey;
+		}
+
+		// Next order by instancing bucket.
+		if (A.StateBucketId != B.StateBucketId)
+		{
+			return A.StateBucketId < B.StateBucketId;
+		}
+
+		// No valid StateBucketId - we cannot dynamically instance but we can still order by PSO to maximize state filtering
+		if (A.StateBucketId == -1 && B.StateBucketId == -1)
+		{
+			if (A.MeshDrawCommand->CachedPipelineId.GetId() != B.MeshDrawCommand->CachedPipelineId.GetId())
+			{
+				return A.MeshDrawCommand->CachedPipelineId.GetId() < B.MeshDrawCommand->CachedPipelineId.GetId();
+			}
+		}
+
+		return false;
+	}
+};
+
+uint32 BitInvertIfNegativeFloat(uint32 f)
+{
+	unsigned mask = -int32(f >> 31) | 0x80000000;
+	return f ^ mask;
+}
+
+/**
+* Update mesh sort keys with view dependent data.
+*/
+void UpdateTranslucentMeshSortKeys(
+	ETranslucentSortPolicy::Type TranslucentSortPolicy,
+	const FVector& TranslucentSortAxis,
+	const FVector& ViewOrigin,
+	const FMatrix& ViewMatrix,
+	const TArray<struct FPrimitiveBounds>& PrimitiveBounds,
+	ETranslucencyPass::Type TranslucencyPass, 
+	FMeshCommandOneFrameArray& VisibleMeshCommands
+	)
+{
+	QUICK_SCOPE_CYCLE_COUNTER(STAT_UpdateTranslucentMeshSortKeys);
+
+	for (int32 CommandIndex = 0; CommandIndex < VisibleMeshCommands.Num(); ++CommandIndex)
+	{
+		FVisibleMeshDrawCommand& VisibleCommand = VisibleMeshCommands[CommandIndex];
+
+		FVector BoundsOrigin = FVector::ZeroVector;
+		if (VisibleCommand.DrawPrimitiveId < PrimitiveBounds.Num())
+		{
+			BoundsOrigin = PrimitiveBounds[VisibleCommand.DrawPrimitiveId].BoxSphereBounds.Origin;
+		}
+
+		float Distance = 0.0f;
+		if (TranslucentSortPolicy == ETranslucentSortPolicy::SortByDistance)
+		{
+			//sort based on distance to the view position, view rotation is not a factor
+			Distance = (BoundsOrigin - ViewOrigin).Size();
+		}
+		else if (TranslucentSortPolicy == ETranslucentSortPolicy::SortAlongAxis)
+		{
+			// Sort based on enforced orthogonal distance
+			const FVector CameraToObject = BoundsOrigin - ViewOrigin;
+			Distance = FVector::DotProduct(CameraToObject, TranslucentSortAxis);
+		}
+		else
+		{
+			// Sort based on projected Z distance
+			check(TranslucentSortPolicy == ETranslucentSortPolicy::SortByProjectedZ);
+			Distance = ViewMatrix.TransformPosition(BoundsOrigin).Z;
+		}
+
+		// Patch distance inside translucent mesh sort key.
+		UTranslucentMeshSortKey TranslucentMeshSortKey;
+		TranslucentMeshSortKey.PackedData = VisibleCommand.SortKey.PackedData;
+		TranslucentMeshSortKey.Fields.Distance = (uint32)~BitInvertIfNegativeFloat(*((uint32*)&Distance));
+		VisibleCommand.SortKey.PackedData = TranslucentMeshSortKey.PackedData;		
+	}
+}
+
+static uint64 GetMobileBasePassSortKey(bool bMasked, bool bBackground, int32 PipelineId, int32 StateBucketId, float PipelineDistance, float PrimitiveDistance)
+{
+	const float PrimitiveDepthQuantization = ((1 << 14) - 1);
+
+	union
+	{
+		uint64 PackedInt;
+		struct
+		{
+			uint64 DepthBits : 14;			// Order by primitive depth
+			uint64 StateBucketId : 20; 		// Order by state bucket
+			uint64 PipelineId : 20;			// Order by PSO
+			uint64 PipelineDepthBits : 8;	// Order PSOs front to back
+			uint64 Background : 1;			// Non-background meshes first 
+			uint64 Masked : 1;				// Non-masked first
+		} Fields;
+	} Key;
+
+	union FFloatToInt { float F; uint32 I; };
+	FFloatToInt F2I;
+
+	Key.PackedInt = 0;
+	Key.Fields.Masked = bMasked;
+	Key.Fields.Background = bBackground;
+	F2I.F = PipelineDistance / HALF_WORLD_MAX;
+	Key.Fields.PipelineDepthBits = (F2I.I >> 23) & 0xff; // 8 bit exponent
+	Key.Fields.PipelineId = PipelineId;
+	Key.Fields.StateBucketId = StateBucketId;
+	Key.Fields.DepthBits = int32((FMath::Min<float>(PrimitiveDistance, HALF_WORLD_MAX) / HALF_WORLD_MAX) * PrimitiveDepthQuantization);
+
+	return Key.PackedInt;
+}
+
+/**
+* Merge mobile BasePass with BasePassCSM based on CSM visibility in order to select appropriate shader for given command.
+*/
+void MergeMobileBasePassMeshDrawCommands(
+	const FMobileCSMVisibilityInfo& MobileCSMVisibilityInfo,
+	int32 ScenePrimitiveNum,
+	FMeshCommandOneFrameArray& MeshCommands,
+	FMeshCommandOneFrameArray& MeshCommandsCSM
+	)
+{
+	if (MobileCSMVisibilityInfo.bMobileDynamicCSMInUse)
+	{
+		// determine per view CSM visibility
+		checkf(MeshCommands.Num() == MeshCommandsCSM.Num(), TEXT("VisibleMeshDrawCommands of BasePass and MobileBasePassCSM are expected to match."));
+		for (int32 i = MeshCommands.Num() - 1; i >= 0; --i)
+		{
+			FVisibleMeshDrawCommand& MeshCommand = MeshCommands[i];
+			FVisibleMeshDrawCommand& MeshCommandCSM = MeshCommandsCSM[i];
+
+			if (MobileCSMVisibilityInfo.bAlwaysUseCSM 
+				|| (MeshCommand.DrawPrimitiveId < ScenePrimitiveNum && MobileCSMVisibilityInfo.MobilePrimitiveCSMReceiverVisibilityMap[MeshCommand.DrawPrimitiveId]))
+			{
+				checkf(MeshCommand.DrawPrimitiveId == MeshCommandCSM.DrawPrimitiveId, TEXT("VisibleMeshDrawCommands of BasePass and MobileBasePassCSM are expected to match."));
+				// Use CSM's VisibleMeshDrawCommand.
+				MeshCommand = MeshCommandCSM;
+			}
+		}
+		MeshCommandsCSM.Reset();
+	}
+}
+
+/**
+* Compute mesh sort keys for the mobile base pass
+*/
+void UpdateMobileBasePassMeshSortKeys(
+	const FVector& ViewOrigin,
+	const TArray<struct FPrimitiveBounds>& ScenePrimitiveBounds,
+	FMeshCommandOneFrameArray& VisibleMeshCommands
+)
+{
+	QUICK_SCOPE_CYCLE_COUNTER(STAT_UpdateMobileBasePassMeshSortKeys);
+
+	int32 NumCmds = VisibleMeshCommands.Num();
+
+	// pre-compute distance to a group of meshes that share same PSO
+	TMap<int32, float> PipelineDistances;
+	PipelineDistances.Reserve(256);
+
+	for (int32 CmdIdx = 0; CmdIdx < NumCmds; ++CmdIdx)
+	{
+		FVisibleMeshDrawCommand& Cmd = VisibleMeshCommands[CmdIdx];
+
+		float PrimitiveDistance = 0;
+
+		if (Cmd.DrawPrimitiveId < ScenePrimitiveBounds.Num())
+		{
+			const FPrimitiveBounds& PrimitiveBounds = ScenePrimitiveBounds[Cmd.DrawPrimitiveId];
+			PrimitiveDistance = (PrimitiveBounds.BoxSphereBounds.Origin - ViewOrigin).Size();
+		}
+
+		float& PipelineDistance = PipelineDistances.FindOrAdd(Cmd.MeshDrawCommand->CachedPipelineId.GetId());
+		// not sure what could be better: average distance, max or min
+		PipelineDistance = FMath::Max(PipelineDistance, PrimitiveDistance);
+	}
+
+	// compute sort key for each mesh command
+	for (int32 CmdIdx = 0; CmdIdx < NumCmds; ++CmdIdx)
+	{
+		FVisibleMeshDrawCommand& Cmd = VisibleMeshCommands[CmdIdx];
+		float PrimitiveDistance = 0;
+		bool bMasked = false;
+		bool bBackground = false;
+
+		if (Cmd.DrawPrimitiveId < ScenePrimitiveBounds.Num())
+		{
+			const FPrimitiveBounds& PrimitiveBounds = ScenePrimitiveBounds[Cmd.DrawPrimitiveId];
+			PrimitiveDistance = (PrimitiveBounds.BoxSphereBounds.Origin - ViewOrigin).Size();
+			bBackground = (PrimitiveBounds.BoxSphereBounds.SphereRadius > HALF_WORLD_MAX / 4.0f); // TODO: per-primitive bBackgound flag
+			bMasked = Cmd.SortKey.PackedData & 0x1 ? true : false; // Set in MobileBasePass.cpp - GetBasePassStaticSortKey
+		}
+
+		int32 PipelineId = Cmd.MeshDrawCommand->CachedPipelineId.GetId();
+		float PipelineDistance = PipelineDistances.FindRef(PipelineId);
+		// poor man StateID, can't use Cmd.StateBucketId as it is unique for each primitive if platform does not support auto-instancing
+		int32 StateBucketId = PointerHash(Cmd.MeshDrawCommand->IndexBuffer);
+
+		Cmd.SortKey.PackedData = GetMobileBasePassSortKey(bMasked, bBackground, PipelineId, StateBucketId, PipelineDistance, PrimitiveDistance);
+	}
+}
+
+/**
+ * Build mesh draw command primitive Id buffer for instancing.
+ * TempVisibleMeshDrawCommands must be presized for NewPassVisibleMeshDrawCommands.
+ */
+void BuildMeshDrawCommandPrimitiveIdBuffer(
+	FMeshCommandOneFrameArray& VisibleMeshDrawCommands,
+	FDynamicMeshDrawCommandStorage& MeshDrawCommandStorage,
+	FGlobalDynamicVertexBuffer::FAllocation& PrimitiveIdBuffer,
+	FMeshCommandOneFrameArray& TempVisibleMeshDrawCommands,
+	int32& MaxInstances,
+	int32& VisibleMeshDrawCommandsNum,
+	int32& NewPassVisibleMeshDrawCommandsNum
+	)
+{
+	QUICK_SCOPE_CYCLE_COUNTER(STAT_BuildMeshDrawCommandPrimitiveIdBuffer);
+	check(VisibleMeshDrawCommands.Num() <= TempVisibleMeshDrawCommands.Max() && TempVisibleMeshDrawCommands.Num() == 0 && PrimitiveIdBuffer.IsValid());
+
+	const FVisibleMeshDrawCommand* RESTRICT PassVisibleMeshDrawCommands = VisibleMeshDrawCommands.GetData();
+	const int32 NumDrawCommands = VisibleMeshDrawCommands.Num();
+
+	int32* RESTRICT PrimitiveIds = (int32*)PrimitiveIdBuffer.Buffer;
+
+	if (IsDynamicInstancingEnabled())
+	{
+		QUICK_SCOPE_CYCLE_COUNTER(STAT_DynamicInstancingOfVisibleMeshDrawCommands);
+
+		int32 CurrentStateBucketId = -1;
+		uint32* RESTRICT CurrentDynamicallyInstancedMeshCommandNumInstances = nullptr;
+		MaxInstances = 1;
+
+		for (int32 DrawCommandIndex = 0; DrawCommandIndex < NumDrawCommands; DrawCommandIndex++)
+		{
+			const FVisibleMeshDrawCommand& RESTRICT VisibleMeshDrawCommand = PassVisibleMeshDrawCommands[DrawCommandIndex];
+
+			if (VisibleMeshDrawCommand.StateBucketId == CurrentStateBucketId && VisibleMeshDrawCommand.StateBucketId != -1)
+			{
+				if (CurrentDynamicallyInstancedMeshCommandNumInstances)
+				{
+					const int32 CurrentNumInstances = *CurrentDynamicallyInstancedMeshCommandNumInstances;
+					*CurrentDynamicallyInstancedMeshCommandNumInstances = CurrentNumInstances + 1;
+					MaxInstances = FMath::Max(MaxInstances, CurrentNumInstances + 1);
+				}
+				else
+				{
+					FVisibleMeshDrawCommand NewVisibleMeshDrawCommand = VisibleMeshDrawCommand;
+					NewVisibleMeshDrawCommand.PrimitiveIdBufferOffset = DrawCommandIndex;
+					TempVisibleMeshDrawCommands.Emplace(MoveTemp(NewVisibleMeshDrawCommand));
+				}
+			}
+			else
+			{
+				// First time state bucket setup
+				CurrentStateBucketId = VisibleMeshDrawCommand.StateBucketId;
+
+				if (VisibleMeshDrawCommand.MeshDrawCommand->PrimitiveIdStreamIndex >= 0
+					&& VisibleMeshDrawCommand.MeshDrawCommand->NumInstances == 1
+					// Don't create a new FMeshDrawCommand for the last command and make it safe for us to look at the next command
+					&& DrawCommandIndex + 1 < NumDrawCommands
+					// Only create a new FMeshDrawCommand if more than one draw in the state bucket
+					&& CurrentStateBucketId == PassVisibleMeshDrawCommands[DrawCommandIndex + 1].StateBucketId)
+				{
+					const int32 Index = MeshDrawCommandStorage.MeshDrawCommands.AddElement(*VisibleMeshDrawCommand.MeshDrawCommand);
+					FMeshDrawCommand& NewCommand = MeshDrawCommandStorage.MeshDrawCommands[Index];
+					FVisibleMeshDrawCommand NewVisibleMeshDrawCommand;
+
+					NewVisibleMeshDrawCommand.Setup(
+						&NewCommand,
+						VisibleMeshDrawCommand.DrawPrimitiveId,
+						VisibleMeshDrawCommand.StateBucketId,
+						VisibleMeshDrawCommand.MeshFillMode,
+						VisibleMeshDrawCommand.MeshCullMode,
+						VisibleMeshDrawCommand.SortKey);
+
+					NewVisibleMeshDrawCommand.PrimitiveIdBufferOffset = DrawCommandIndex;
+					TempVisibleMeshDrawCommands.Emplace(MoveTemp(NewVisibleMeshDrawCommand));
+
+					CurrentDynamicallyInstancedMeshCommandNumInstances = &NewCommand.NumInstances;
+				}
+				else
+				{
+					CurrentDynamicallyInstancedMeshCommandNumInstances = nullptr;
+					FVisibleMeshDrawCommand NewVisibleMeshDrawCommand = VisibleMeshDrawCommand;
+					NewVisibleMeshDrawCommand.PrimitiveIdBufferOffset = DrawCommandIndex;
+					TempVisibleMeshDrawCommands.Emplace(MoveTemp(NewVisibleMeshDrawCommand));
+				}
+			}
+
+			//@todo - refactor into memcpy
+			PrimitiveIds[DrawCommandIndex] = VisibleMeshDrawCommand.DrawPrimitiveId;
+		}
+
+		// Setup instancing stats for logging.
+		VisibleMeshDrawCommandsNum = VisibleMeshDrawCommands.Num();
+		NewPassVisibleMeshDrawCommandsNum = TempVisibleMeshDrawCommands.Num();
+
+		// Replace VisibleMeshDrawCommands
+		FMemory::Memswap(&VisibleMeshDrawCommands, &TempVisibleMeshDrawCommands, sizeof(TempVisibleMeshDrawCommands));
+		TempVisibleMeshDrawCommands.Reset();
+	}
+	else
+	{
+		QUICK_SCOPE_CYCLE_COUNTER(STAT_BuildVisibleMeshDrawCommandPrimitiveIdBuffers);
+
+		for (int32 DrawCommandIndex = 0; DrawCommandIndex < NumDrawCommands; DrawCommandIndex++)
+		{
+			const FVisibleMeshDrawCommand& VisibleMeshDrawCommand = VisibleMeshDrawCommands[DrawCommandIndex];
+			PrimitiveIds[DrawCommandIndex] = VisibleMeshDrawCommand.DrawPrimitiveId;
+		}
+	}
+}
+
+/**
+ * Converts each FMeshBatch into a set of FMeshDrawCommands for a specific mesh pass type.
+ */
+void GenerateDynamicMeshDrawCommands(
+	const FViewInfo& View,
+	EShadingPath ShadingPath,
+	EMeshPass::Type PassType,
+	FMeshPassProcessor* PassMeshProcessor,
+	const TArray<FMeshBatchAndRelevance, SceneRenderingAllocator>& DynamicMeshElements,
+	int32 NumDynamicMeshElements,
+	const TArray<const FStaticMesh*, SceneRenderingAllocator>& DynamicMeshCommandBuildRequests,
+	int32 NumDynamicMeshCommandBuildRequestElements,
+	FMeshCommandOneFrameArray& VisibleCommands,
+	FDynamicMeshDrawCommandStorage& MeshDrawCommandStorage
+)
+{
+	QUICK_SCOPE_CYCLE_COUNTER(STAT_GenerateDynamicMeshDrawCommands);
+	check(PassMeshProcessor);
+
+	FDynamicPassMeshDrawListContext DynamicPassMeshDrawListContext(
+		MeshDrawCommandStorage,
+		VisibleCommands
+	);
+	PassMeshProcessor->SetDrawListContext(&DynamicPassMeshDrawListContext);
+
+	const int32 NumVisibleCommandsBeforeDynamicBatches = VisibleCommands.Num();
+	const int32 NumDynamicMeshBatches = DynamicMeshElements.Num();
+
+	for (int32 MeshIndex = 0; MeshIndex < NumDynamicMeshBatches; MeshIndex++)
+	{
+		const FMeshBatchAndRelevance& MeshAndRelevance = DynamicMeshElements[MeshIndex];
+		check(!MeshAndRelevance.Mesh->bRequiresPerElementVisibility);
+		const uint64 BatchElementMask = ~0ull;
+
+		PassMeshProcessor->AddMeshBatch(*MeshAndRelevance.Mesh, BatchElementMask, MeshAndRelevance.PrimitiveSceneProxy);
+	}
+
+	// Check if we didn't generate more mesh draw commands, than we assumed.
+	check(VisibleCommands.Num() - NumVisibleCommandsBeforeDynamicBatches <= NumDynamicMeshElements);
+
+	const int32 NumVisibleCommandsBeforeDynamicRequests = VisibleCommands.Num();
+	const int32 NumStaticMeshBatches = DynamicMeshCommandBuildRequests.Num();
+
+	for (int32 MeshIndex = 0; MeshIndex < NumStaticMeshBatches; MeshIndex++)
+	{
+		const FStaticMesh* StaticMeshBatch = DynamicMeshCommandBuildRequests[MeshIndex];
+		const uint64 BatchElementMask = StaticMeshBatch->bRequiresPerElementVisibility ? View.StaticMeshBatchVisibility[StaticMeshBatch->BatchVisibilityId] : ~0ull;
+
+		PassMeshProcessor->AddMeshBatch(*StaticMeshBatch, BatchElementMask, StaticMeshBatch->PrimitiveSceneInfo->Proxy, StaticMeshBatch->Id);
+	}
+
+	// Check if we didn't generate more mesh draw commands, than we assumed.
+	check(VisibleCommands.Num() - NumVisibleCommandsBeforeDynamicRequests <= NumDynamicMeshCommandBuildRequestElements);
+}
+
+/**
+ * Special version of GenerateDynamicMeshDrawCommands for the mobile base pass.
+ * Based on CSM visibility it will generate mesh draw commands using either normal base pass processor or CSM base pass processor.
+*/
+void GenerateMobileBasePassDynamicMeshDrawCommands(
+	const FViewInfo& View,
+	EShadingPath ShadingPath,
+	EMeshPass::Type PassType,
+	FMeshPassProcessor* PassMeshProcessor,
+	FMeshPassProcessor* MobilePassCSMPassMeshProcessor,
+	const TArray<FMeshBatchAndRelevance, SceneRenderingAllocator>& DynamicMeshElements,
+	int32 NumDynamicMeshElements,
+	const TArray<const FStaticMesh*, SceneRenderingAllocator>& DynamicMeshCommandBuildRequests,
+	int32 NumDynamicMeshCommandBuildRequestElements,
+	FMeshCommandOneFrameArray& VisibleCommands,
+	FDynamicMeshDrawCommandStorage& MeshDrawCommandStorage
+)
+{
+	QUICK_SCOPE_CYCLE_COUNTER(STAT_GenerateMobileBasePassDynamicMeshDrawCommands);
+	check(PassMeshProcessor && MobilePassCSMPassMeshProcessor);
+
+	FDynamicPassMeshDrawListContext DynamicPassMeshDrawListContext(
+		MeshDrawCommandStorage,
+		VisibleCommands
+	);
+	PassMeshProcessor->SetDrawListContext(&DynamicPassMeshDrawListContext);
+	MobilePassCSMPassMeshProcessor->SetDrawListContext(&DynamicPassMeshDrawListContext);
+
+	const FMobileCSMVisibilityInfo& MobileCSMVisibilityInfo = View.MobileCSMVisibilityInfo;
+	
+	const int32 NumVisibleCommandsBeforeDynamicBatches = VisibleCommands.Num();
+	const int32 NumDynamicMeshBatches = DynamicMeshElements.Num();
+
+	for (int32 MeshIndex = 0; MeshIndex < NumDynamicMeshBatches; MeshIndex++)
+	{
+		const FMeshBatchAndRelevance& MeshAndRelevance = DynamicMeshElements[MeshIndex];
+		check(!MeshAndRelevance.Mesh->bRequiresPerElementVisibility);
+		const uint64 BatchElementMask = ~0ull;
+
+		const int32 PrimitiveIndex = MeshAndRelevance.PrimitiveSceneProxy->GetPrimitiveSceneInfo()->GetIndex();
+		if (MobileCSMVisibilityInfo.bMobileDynamicCSMInUse 
+			&& (MobileCSMVisibilityInfo.bAlwaysUseCSM || MobileCSMVisibilityInfo.MobilePrimitiveCSMReceiverVisibilityMap[PrimitiveIndex]))
+		{
+			MobilePassCSMPassMeshProcessor->AddMeshBatch(*MeshAndRelevance.Mesh, BatchElementMask, MeshAndRelevance.PrimitiveSceneProxy);
+		}
+		else
+		{
+			PassMeshProcessor->AddMeshBatch(*MeshAndRelevance.Mesh, BatchElementMask, MeshAndRelevance.PrimitiveSceneProxy);
+		}
+	}
+
+	// Check if we didn't generate more mesh draw commands, than we assumed.
+	check(VisibleCommands.Num() - NumVisibleCommandsBeforeDynamicBatches <= NumDynamicMeshElements);
+
+	const int32 NumVisibleCommandsBeforeDynamicRequests = VisibleCommands.Num();
+	const int32 NumStaticMeshBatches = DynamicMeshCommandBuildRequests.Num();
+
+	for (int32 MeshIndex = 0; MeshIndex < NumStaticMeshBatches; MeshIndex++)
+	{
+		const FStaticMesh* StaticMeshBatch = DynamicMeshCommandBuildRequests[MeshIndex];
+		const uint64 BatchElementMask = StaticMeshBatch->bRequiresPerElementVisibility ? View.StaticMeshBatchVisibility[StaticMeshBatch->BatchVisibilityId] : ~0ull;
+
+		const int32 PrimitiveIndex = StaticMeshBatch->PrimitiveSceneInfo->Proxy->GetPrimitiveSceneInfo()->GetIndex();
+		if (MobileCSMVisibilityInfo.bMobileDynamicCSMInUse
+			&& (MobileCSMVisibilityInfo.bAlwaysUseCSM || MobileCSMVisibilityInfo.MobilePrimitiveCSMReceiverVisibilityMap[PrimitiveIndex]))
+		{
+			MobilePassCSMPassMeshProcessor->AddMeshBatch(*StaticMeshBatch, BatchElementMask, StaticMeshBatch->PrimitiveSceneInfo->Proxy, StaticMeshBatch->Id);
+		}
+		else
+		{
+			PassMeshProcessor->AddMeshBatch(*StaticMeshBatch, BatchElementMask, StaticMeshBatch->PrimitiveSceneInfo->Proxy, StaticMeshBatch->Id);
+		}
+	}
+
+	// Check if we didn't generate more mesh draw commands, than we assumed.
+	check(VisibleCommands.Num() - NumVisibleCommandsBeforeDynamicRequests <= NumDynamicMeshCommandBuildRequestElements);
+}
+
+/**
+* Apply view overrides to existing mesh draw commands (e.g. reverse culling mode for rendering planar reflections).
+* TempVisibleMeshDrawCommands must be presized for NewPassVisibleMeshDrawCommands.
+*/
+void ApplyViewOverridesToMeshDrawCommands(
+	EShadingPath ShadingPath,
+	EMeshPass::Type PassType,
+	bool bReverseCulling,
+	bool bRenderSceneTwoSided,
+	FExclusiveDepthStencil::Type BasePassDepthStencilAccess,
+	FMeshCommandOneFrameArray& VisibleMeshDrawCommands,
+	FDynamicMeshDrawCommandStorage& MeshDrawCommandStorage,
+	FMeshCommandOneFrameArray& TempVisibleMeshDrawCommands
+	)
+{
+	QUICK_SCOPE_CYCLE_COUNTER(STAT_ApplyViewOverridesToMeshDrawCommands);
+	check(VisibleMeshDrawCommands.Num() <= TempVisibleMeshDrawCommands.Max() && TempVisibleMeshDrawCommands.Num() == 0 && PassType != EMeshPass::Num);
+
+	if ((FPassProcessorManager::GetPassFlags(ShadingPath, PassType) & EMeshPassFlags::MainView) != EMeshPassFlags::None)
+	{
+		const FExclusiveDepthStencil::Type DefaultBasePassDepthStencilAccess = GetDefaultBasePassDepthStencilAccess(ShadingPath);
+
+		if (bReverseCulling || bRenderSceneTwoSided || (BasePassDepthStencilAccess != DefaultBasePassDepthStencilAccess && PassType == EMeshPass::BasePass))
+		{
+			for (int32 MeshCommandIndex = 0; MeshCommandIndex < VisibleMeshDrawCommands.Num(); MeshCommandIndex++)
+			{
+				MeshDrawCommandStorage.MeshDrawCommands.Add(1);
+				FMeshDrawCommand& NewMeshCommand = MeshDrawCommandStorage.MeshDrawCommands[MeshDrawCommandStorage.MeshDrawCommands.Num() - 1];
+
+				const FVisibleMeshDrawCommand& VisibleMeshDrawCommand = VisibleMeshDrawCommands[MeshCommandIndex];
+				const FMeshDrawCommand& MeshCommand = *VisibleMeshDrawCommand.MeshDrawCommand;
+				NewMeshCommand = MeshCommand;
+
+				const ERasterizerCullMode LocalCullMode = bRenderSceneTwoSided ? CM_None : bReverseCulling ? FMeshDrawingPolicy::InverseCullMode(VisibleMeshDrawCommand.MeshCullMode) : VisibleMeshDrawCommand.MeshCullMode;
+				NewMeshCommand.PipelineState.RasterizerState = GetStaticRasterizerState<true>(VisibleMeshDrawCommand.MeshFillMode, LocalCullMode);
+
+				if (BasePassDepthStencilAccess != DefaultBasePassDepthStencilAccess && PassType == EMeshPass::BasePass)
+				{
+					FDrawingPolicyRenderState PassDrawRenderState;
+					SetupBasePassState(BasePassDepthStencilAccess, false, PassDrawRenderState);
+					NewMeshCommand.PipelineState.DepthStencilState = PassDrawRenderState.GetDepthStencilState();
+				}
+
+				NewMeshCommand.Finalize(true);
+
+				FVisibleMeshDrawCommand NewVisibleMeshDrawCommand;
+
+				NewVisibleMeshDrawCommand.Setup(
+					&NewMeshCommand,
+					VisibleMeshDrawCommand.DrawPrimitiveId,
+					VisibleMeshDrawCommand.StateBucketId,
+					VisibleMeshDrawCommand.MeshFillMode,
+					VisibleMeshDrawCommand.MeshCullMode,
+					VisibleMeshDrawCommand.SortKey);
+
+				TempVisibleMeshDrawCommands.Add(NewVisibleMeshDrawCommand);
+			}
+
+			// Replace VisibleMeshDrawCommands
+			FMemory::Memswap(&VisibleMeshDrawCommands, &TempVisibleMeshDrawCommands, sizeof(TempVisibleMeshDrawCommands));
+			TempVisibleMeshDrawCommands.Reset();
+		}
+	}
+}
+
+FAutoConsoleTaskPriority CPrio_FMeshDrawCommandPassSetupTask(
+	TEXT("TaskGraph.TaskPriorities.FMeshDrawCommandPassSetupTask"),
+	TEXT("Task and thread priority for FMeshDrawCommandPassSetupTask."),
+	ENamedThreads::HighThreadPriority, // if we have high priority task threads, then use them...
+	ENamedThreads::NormalTaskPriority, // .. at normal task priority
+	ENamedThreads::HighTaskPriority // if we don't have hi pri threads, then use normal priority threads at high task priority instead
+);
+
+/**
+ * Task for a parallel setup of mesh draw commands. Includes generation of dynamic mesh draw commands, sorting, merging etc.
+ */
+class FMeshDrawCommandPassSetupTask
+{
+public:
+	FMeshDrawCommandPassSetupTask(FMeshDrawCommandPassSetupTaskContext& InContext)
+		: Context(InContext)
+	{
+	}
+
+	FORCEINLINE TStatId GetStatId() const
+	{
+		RETURN_QUICK_DECLARE_CYCLE_STAT(FMeshDrawCommandPassSetupTask, STATGROUP_TaskGraphTasks);
+	}
+
+	ENamedThreads::Type GetDesiredThread()
+	{
+		return CPrio_FMeshDrawCommandPassSetupTask.Get();
+	}
+
+	static ESubsequentsMode::Type GetSubsequentsMode() 
+	{ 
+		return ESubsequentsMode::TrackSubsequents; 
+	}
+
+	void AnyThreadTask()
+	{
+		// Mobile base pass is a special case, as final lists is created from two mesh passes based on CSM visibility.
+		const bool bMobileBasePass = Context.ShadingPath == EShadingPath::Mobile && Context.PassType == EMeshPass::BasePass;
+
+		if (bMobileBasePass)
+		{
+			MergeMobileBasePassMeshDrawCommands(
+				Context.View->MobileCSMVisibilityInfo,
+				Context.PrimitiveBounds->Num(),
+				Context.MeshDrawCommands,
+				Context.MobileBasePassCSMMeshDrawCommands
+			);
+
+			GenerateMobileBasePassDynamicMeshDrawCommands(
+				*Context.View,
+				Context.ShadingPath,
+				Context.PassType,
+				Context.MeshPassProcessor,
+				Context.MobileBasePassCSMMeshPassProcessor,
+				*Context.DynamicMeshElements,
+				Context.NumDynamicMeshElements,
+				Context.DynamicMeshCommandBuildRequests,
+				Context.NumDynamicMeshCommandBuildRequestElements,
+				Context.MeshDrawCommands,
+				Context.MeshDrawCommandStorage
+			);
+		}
+		else
+		{
+			GenerateDynamicMeshDrawCommands(
+				*Context.View,
+				Context.ShadingPath,
+				Context.PassType,
+				Context.MeshPassProcessor,
+				*Context.DynamicMeshElements,
+				Context.NumDynamicMeshElements,
+				Context.DynamicMeshCommandBuildRequests,
+				Context.NumDynamicMeshCommandBuildRequestElements,
+				Context.MeshDrawCommands,
+				Context.MeshDrawCommandStorage
+			);
+		}
+
+		if (Context.MeshDrawCommands.Num() > 0)
+		{
+			if (Context.PassType != EMeshPass::Num)
+			{
+				ApplyViewOverridesToMeshDrawCommands(
+					Context.ShadingPath,
+					Context.PassType,
+					Context.bReverseCulling,
+					Context.bRenderSceneTwoSided,
+					Context.BasePassDepthStencilAccess,
+					Context.MeshDrawCommands,
+					Context.MeshDrawCommandStorage,
+					Context.TempVisibleMeshDrawCommands
+				);
+			}
+
+			// Update sort keys.
+			if (bMobileBasePass)
+			{
+				UpdateMobileBasePassMeshSortKeys(
+					Context.ViewOrigin,
+					*Context.PrimitiveBounds,
+					Context.MeshDrawCommands
+					);
+			}
+			else if (Context.TranslucencyPass != ETranslucencyPass::TPT_MAX)
+			{
+				UpdateTranslucentMeshSortKeys(
+					Context.TranslucentSortPolicy,
+					Context.TranslucentSortAxis,
+					Context.ViewOrigin,
+					Context.ViewMatrix,
+					*Context.PrimitiveBounds,
+					Context.TranslucencyPass,
+					Context.MeshDrawCommands
+				);
+			}
+
+			{
+				QUICK_SCOPE_CYCLE_COUNTER(STAT_SortVisibleMeshDrawCommands);
+				Context.MeshDrawCommands.Sort(FCompareFMeshDrawCommands());
+			}
+
+			if (Context.bUseGPUScene)
+			{
+				BuildMeshDrawCommandPrimitiveIdBuffer(
+					Context.MeshDrawCommands,
+					Context.MeshDrawCommandStorage,
+					Context.PrimitiveIdBuffer,
+					Context.TempVisibleMeshDrawCommands,
+					Context.MaxInstances,
+					Context.VisibleMeshDrawCommandsNum,
+					Context.NewPassVisibleMeshDrawCommandsNum
+				);
+			}
+		}
+	}
+
+	void DoTask(ENamedThreads::Type CurrentThread, const FGraphEventRef& MyCompletionGraphEvent)
+	{
+		AnyThreadTask();
+	}
+
+
+private:
+	FMeshDrawCommandPassSetupTaskContext& Context;
+};
+
+void SortPassMeshDrawCommands(
+	ERHIFeatureLevel::Type FeatureLevel,
+	FMeshCommandOneFrameArray& VisibleMeshDrawCommands,
+	FDynamicMeshDrawCommandStorage& MeshDrawCommandStorage,
+	FGlobalDynamicVertexBuffer& DynamicVertexBuffer,
+	FGlobalDynamicVertexBuffer::FAllocation& PrimitiveIdBuffer)
+{
+	const bool bUseGPUScene = UseGPUScene(GMaxRHIShaderPlatform, FeatureLevel);
+
+	const int32 NumDrawCommands = VisibleMeshDrawCommands.Num();
+	if (NumDrawCommands > 0)
+	{
+		FMeshCommandOneFrameArray NewPassVisibleMeshDrawCommands;
+		int32 MaxInstances = 1;
+		int32 VisibleMeshDrawCommandsNum = 0;
+		int32 NewPassVisibleMeshDrawCommandsNum = 0;
+
+		// Preallocate context memory and resources.
+		if (bUseGPUScene)
+		{
+			PrimitiveIdBuffer = DynamicVertexBuffer.Allocate(NumDrawCommands * sizeof(int32));
+
+			if (IsDynamicInstancingEnabled())
+			{
+				NewPassVisibleMeshDrawCommands.Empty(NumDrawCommands);
+			}
+		}
+
+		VisibleMeshDrawCommands.Sort(FCompareFMeshDrawCommands());
+
+		if (bUseGPUScene)
+		{
+			BuildMeshDrawCommandPrimitiveIdBuffer(
+				VisibleMeshDrawCommands,
+				MeshDrawCommandStorage,
+				PrimitiveIdBuffer,
+				NewPassVisibleMeshDrawCommands,
+				MaxInstances,
+				VisibleMeshDrawCommandsNum,
+				NewPassVisibleMeshDrawCommandsNum
+			);
+		}
+	}
+}
+
+
+void FParallelMeshDrawCommandPass::DispatchPassSetup(
+	FScene* Scene,
+	const FViewInfo& View,
+	EMeshPass::Type PassType,
+	FExclusiveDepthStencil::Type BasePassDepthStencilAccess,
+	FMeshPassProcessor* MeshPassProcessor,
+	const TArray<FMeshBatchAndRelevance, SceneRenderingAllocator>& DynamicMeshElements,
+	int32 NumDynamicMeshElements,
+	TArray<const FStaticMesh*, SceneRenderingAllocator>& InOutDynamicMeshCommandBuildRequests,
+	int32 NumDynamicMeshCommandBuildRequestElements,
+	FMeshCommandOneFrameArray& InOutMeshDrawCommands,
+	FMeshPassProcessor* MobileBasePassCSMMeshPassProcessor,
+	FMeshCommandOneFrameArray* InOutMobileBasePassCSMMeshDrawCommands
+)
+{
+	check(TaskEventRefs.Num() == 0 && MeshPassProcessor != nullptr);
+
+	MaxNumDraws = InOutMeshDrawCommands.Num() + NumDynamicMeshElements + NumDynamicMeshCommandBuildRequestElements;
+	if (MaxNumDraws <= 0)
+	{
+		return;
+	}
+
+	TaskContext.MeshPassProcessor = MeshPassProcessor;
+	TaskContext.MobileBasePassCSMMeshPassProcessor = MobileBasePassCSMMeshPassProcessor;
+	TaskContext.DynamicMeshElements = &DynamicMeshElements;
+
+	TaskContext.View = &View;
+	TaskContext.ShadingPath = Scene->GetShadingPath();
+	TaskContext.PassType = PassType;
+	TaskContext.bUseGPUScene = UseGPUScene(GMaxRHIShaderPlatform, View.GetFeatureLevel());
+	TaskContext.bDynamicInstancing = TaskContext.bUseGPUScene && IsDynamicInstancingEnabled();
+	TaskContext.bReverseCulling = View.bReverseCulling;
+	TaskContext.bRenderSceneTwoSided = View.bRenderSceneTwoSided;
+	TaskContext.BasePassDepthStencilAccess = BasePassDepthStencilAccess;
+	TaskContext.NumDynamicMeshElements = NumDynamicMeshElements;
+	TaskContext.NumDynamicMeshCommandBuildRequestElements = NumDynamicMeshCommandBuildRequestElements;
+
+	// Setup translucency sort key update pass based on view.
+	TaskContext.TranslucencyPass = ETranslucencyPass::TPT_MAX;
+	TaskContext.TranslucentSortPolicy = View.TranslucentSortPolicy;
+	TaskContext.TranslucentSortAxis = View.TranslucentSortAxis;
+	TaskContext.ViewOrigin = View.ViewMatrices.GetViewOrigin();
+	TaskContext.ViewMatrix = View.ViewMatrices.GetViewMatrix();
+	TaskContext.PrimitiveBounds = &Scene->PrimitiveBounds;
+
+	switch (PassType)
+	{
+		case EMeshPass::TranslucencyStandard: TaskContext.TranslucencyPass = ETranslucencyPass::TPT_StandardTranslucency; break;
+		case EMeshPass::TranslucencyAfterDOF: TaskContext.TranslucencyPass = ETranslucencyPass::TPT_TranslucencyAfterDOF; break;
+		case EMeshPass::TranslucencyAll: TaskContext.TranslucencyPass = ETranslucencyPass::TPT_AllTranslucency; break;
+		case EMeshPass::MobileInverseOpacity: TaskContext.TranslucencyPass = ETranslucencyPass::TPT_StandardTranslucency; break;
+	}
+
+	FMemory::Memswap(&TaskContext.MeshDrawCommands, &InOutMeshDrawCommands, sizeof(InOutMeshDrawCommands));
+	FMemory::Memswap(&TaskContext.DynamicMeshCommandBuildRequests, &InOutDynamicMeshCommandBuildRequests, sizeof(InOutDynamicMeshCommandBuildRequests));
+
+	if (TaskContext.ShadingPath == EShadingPath::Mobile && TaskContext.PassType == EMeshPass::BasePass)
+	{
+		FMemory::Memswap(&TaskContext.MobileBasePassCSMMeshDrawCommands, InOutMobileBasePassCSMMeshDrawCommands, sizeof(*InOutMobileBasePassCSMMeshDrawCommands));
+	}
+	else
+	{
+		check(MobileBasePassCSMMeshPassProcessor == nullptr && InOutMobileBasePassCSMMeshDrawCommands == nullptr);
+	}
+
+	// Preallocate resources on rendering thread based on MaxNumDraws.
+	if (TaskContext.bDynamicInstancing)
+	{
+		TaskContext.PrimitiveIdBuffer = FGlobalDynamicVertexBuffer::Get().Allocate(MaxNumDraws * sizeof(int32));
+	}
+	TaskContext.MeshDrawCommands.Reserve(MaxNumDraws);
+	TaskContext.TempVisibleMeshDrawCommands.Reserve(MaxNumDraws);
+
+	const bool bExecuteInParallel = FApp::ShouldUseThreadingForPerformance() && CVarMeshDrawCommandsParallelPassSetup.GetValueOnRenderThread() > 0;
+
+	if (bExecuteInParallel)
+	{
+		TaskEventRefs.Add(TGraphTask<FMeshDrawCommandPassSetupTask>::CreateTask(
+			nullptr, ENamedThreads::GetRenderThread()).ConstructAndDispatchWhenReady(TaskContext));
+	}
+	else
+	{
+		FMeshDrawCommandPassSetupTask Task(TaskContext);
+		Task.AnyThreadTask();
+	}
+}
+
+void FParallelMeshDrawCommandPass::WaitForMeshPassSetupTask() const
+{
+	if (TaskEventRefs.Num() > 0)
+	{
+		// Need to wait on GetRenderThread_Local, as mesh pass setup task can wait on rendering thread inside InitResourceFromPossiblyParallelRendering().
+		FTaskGraphInterface::Get().WaitUntilTasksComplete(TaskEventRefs, ENamedThreads::GetRenderThread_Local());
+	}
+}
+
+void FParallelMeshDrawCommandPass::Empty()
+{
+	// Need to wait in case if someone dispatched sort and draw merge task, but didn't draw it.
+	WaitForMeshPassSetupTask();
+	TaskEventRefs.Empty();
+
+	DumpInstancingStats();
+
+	MaxNumDraws = 0;
+	TaskContext.MeshPassProcessor = nullptr;
+	TaskContext.MobileBasePassCSMMeshPassProcessor = nullptr;
+	TaskContext.DynamicMeshElements = nullptr;
+	PassNameForStats.Empty();
+	TaskContext.MeshDrawCommands.Empty();
+	TaskContext.MobileBasePassCSMMeshDrawCommands.Empty();
+	TaskContext.DynamicMeshCommandBuildRequests.Empty();
+	TaskContext.TempVisibleMeshDrawCommands.Empty();
+	TaskContext.PrimitiveIdBuffer = FGlobalDynamicVertexBuffer::FAllocation();
+}
+
+FParallelMeshDrawCommandPass::~FParallelMeshDrawCommandPass()
+{
+	Empty();
+}
+
+class FDrawVisibleMeshCommandsAnyThreadTask : public FRenderTask
+{
+	FRHICommandList& RHICmdList;
+	const FMeshCommandOneFrameArray& VisibleMeshDrawCommands;
+	FVertexBufferRHIParamRef PrimitiveIdsBuffer;
+	int32 BasePrimitiveIdsOffset;
+	bool bDynamicInstancing;
+	int32 TaskIndex;
+	int32 TaskNum;
+
+public:
+
+	FDrawVisibleMeshCommandsAnyThreadTask(
+		FRHICommandList& InRHICmdList,
+		const FMeshCommandOneFrameArray& InVisibleMeshDrawCommands,
+		FVertexBufferRHIParamRef InPrimitiveIdsBuffer,
+		int32 InBasePrimitiveIdsOffset,
+		bool bInDynamicInstancing,
+		int32 InTaskIndex,
+		int32 InTaskNum
+	)
+		: RHICmdList(InRHICmdList)
+		, VisibleMeshDrawCommands(InVisibleMeshDrawCommands)
+		, PrimitiveIdsBuffer(InPrimitiveIdsBuffer)
+		, BasePrimitiveIdsOffset(InBasePrimitiveIdsOffset)
+		, bDynamicInstancing(bInDynamicInstancing)
+		, TaskIndex(InTaskIndex)
+		, TaskNum(InTaskNum)
+	{}
+
+	FORCEINLINE TStatId GetStatId() const
+	{
+		RETURN_QUICK_DECLARE_CYCLE_STAT(FDrawVisibleMeshCommandsAnyThreadTask, STATGROUP_TaskGraphTasks);
+	}
+
+	static ESubsequentsMode::Type GetSubsequentsMode() { return ESubsequentsMode::TrackSubsequents; }
+
+	void DoTask(ENamedThreads::Type CurrentThread, const FGraphEventRef& MyCompletionGraphEvent)
+	{
+		checkSlow(RHICmdList.IsInsideRenderPass());
+
+		// Recompute draw range.
+		const int32 DrawNum = VisibleMeshDrawCommands.Num();
+		const int32 NumDrawsPerTask = FMath::DivideAndRoundUp(DrawNum, TaskNum);
+		const int32 StartIndex = TaskIndex * NumDrawsPerTask;
+		const int32 NumDraws = FMath::Min(NumDrawsPerTask, DrawNum - StartIndex);
+
+		SubmitMeshDrawCommandsRange(VisibleMeshDrawCommands, PrimitiveIdsBuffer, BasePrimitiveIdsOffset, bDynamicInstancing, StartIndex, NumDraws, RHICmdList);
+
+		RHICmdList.EndRenderPass();
+		RHICmdList.HandleRTThreadTaskCompletion(MyCompletionGraphEvent);
+	}
+};
+
+void FParallelMeshDrawCommandPass::DispatchDraw(FParallelCommandListSet* ParallelCommandListSet, FRHICommandList& RHICmdList) const
+{
+	if (MaxNumDraws <= 0)
+	{
+		return;
+	}
+
+
+	FVertexBufferRHIParamRef PrimitiveIdsBuffer = TaskContext.bUseGPUScene ? TaskContext.PrimitiveIdBuffer.VertexBuffer->VertexBufferRHI : nullptr;
+	const int32 BasePrimitiveIdsOffset = TaskContext.PrimitiveIdBuffer.VertexOffset;
+
+	if (ParallelCommandListSet)
+	{
+		const ENamedThreads::Type RenderThread = ENamedThreads::GetRenderThread();
+
+		FGraphEventArray Prereqs = TaskEventRefs;
+		if (ParallelCommandListSet->GetPrereqs())
+		{
+			Prereqs.Append(*ParallelCommandListSet->GetPrereqs());
+		}
+
+		// Distribute work evenly to the available task graph workers based on NumEstimatedDraws.
+		// Every task will then adjust it's working range based on FVisibleMeshDrawCommandProcessTask results.
+		const int32 NumThreads = FMath::Min<int32>(FTaskGraphInterface::Get().GetNumWorkerThreads(), ParallelCommandListSet->Width);
+		const int32 NumTasks = FMath::Min<int32>(NumThreads, FMath::DivideAndRoundUp(MaxNumDraws, ParallelCommandListSet->MinDrawsPerCommandList));
+		const int32 NumDrawsPerTask = FMath::DivideAndRoundUp(MaxNumDraws, NumTasks);
+
+		for (int32 TaskIndex = 0; TaskIndex < NumTasks; TaskIndex++)
+		{
+			const int32 StartIndex = TaskIndex * NumDrawsPerTask;
+			const int32 NumDraws = FMath::Min(NumDrawsPerTask, MaxNumDraws - StartIndex);
+			checkSlow(NumDraws > 0);
+
+			FRHICommandList* CmdList = ParallelCommandListSet->NewParallelCommandList();
+
+			FGraphEventRef AnyThreadCompletionEvent = TGraphTask<FDrawVisibleMeshCommandsAnyThreadTask>::CreateTask(&Prereqs, RenderThread)
+				.ConstructAndDispatchWhenReady(*CmdList, TaskContext.MeshDrawCommands, PrimitiveIdsBuffer, BasePrimitiveIdsOffset, TaskContext.bDynamicInstancing, TaskIndex, NumTasks);
+			ParallelCommandListSet->AddParallelCommandList(CmdList, AnyThreadCompletionEvent, NumDraws);
+		}
+	}
+	else
+	{
+		WaitForMeshPassSetupTask();
+		SubmitMeshDrawCommandsRange(TaskContext.MeshDrawCommands, PrimitiveIdsBuffer, BasePrimitiveIdsOffset, TaskContext.bDynamicInstancing, 0, TaskContext.MeshDrawCommands.Num(), RHICmdList);
+	}
+}
+
+void FParallelMeshDrawCommandPass::DumpInstancingStats() const
+{
+	if (!PassNameForStats.IsEmpty())
+	{
+		UE_LOG(LogRenderer, Log, TEXT("Instancing stats for %s"), *PassNameForStats);
+		UE_LOG(LogRenderer, Log, TEXT("   %i Mesh Draw Commands in %i instancing state buckets"), TaskContext.VisibleMeshDrawCommandsNum, TaskContext.NewPassVisibleMeshDrawCommandsNum);
+		UE_LOG(LogRenderer, Log, TEXT("   Largest %i"), TaskContext.MaxInstances);
+		UE_LOG(LogRenderer, Log, TEXT("   %.1f Dynamic Instancing draw call reduction factor"), TaskContext.VisibleMeshDrawCommandsNum / (float)TaskContext.NewPassVisibleMeshDrawCommandsNum);
+	}
+}
+
+void FParallelMeshDrawCommandPass::SetDumpInstancingStats(const FString& InPassNameForStats)
+{
+	PassNameForStats = InPassNameForStats;
+}
\ No newline at end of file
diff --git a/Engine/Source/Runtime/Renderer/Private/MeshDrawCommands.h b/Engine/Source/Runtime/Renderer/Private/MeshDrawCommands.h
new file mode 100644
--- /dev/null
+++ b/Engine/Source/Runtime/Renderer/Private/MeshDrawCommands.h
@@ -0,0 +1,138 @@
+// Copyright 1998-2019 Epic Games, Inc. All Rights Reserved.
+
+/*=============================================================================
+MeshDrawCommands.h: Mesh draw commands.
+=============================================================================*/
+
+#pragma once
+
+#include "MeshPassProcessor.h"
+#include "TranslucencyPass.h"
+
+struct FMeshBatchAndRelevance;
+class FStaticMesh;
+class FParallelCommandListSet;
+
+/**	
+ * Parallel mesh draw command pass setup task context.
+ */
+class FMeshDrawCommandPassSetupTaskContext
+{
+public:
+	FMeshDrawCommandPassSetupTaskContext()
+		: View(nullptr)
+		, ShadingPath(EShadingPath::Num)
+		, PassType(EMeshPass::Num)
+		, bUseGPUScene(false)
+		, bDynamicInstancing(false)
+		, bReverseCulling(false)
+		, bRenderSceneTwoSided(false)
+		, BasePassDepthStencilAccess(FExclusiveDepthStencil::DepthNop_StencilNop)
+		, MeshPassProcessor(nullptr)
+		, MobileBasePassCSMMeshPassProcessor(nullptr)
+		, DynamicMeshElements(nullptr)
+		, NumDynamicMeshElements(0)
+		, NumDynamicMeshCommandBuildRequestElements(0)
+		, PrimitiveBounds(nullptr)
+		, VisibleMeshDrawCommandsNum(0)
+		, NewPassVisibleMeshDrawCommandsNum(0)
+		, MaxInstances(1)
+	{
+	}
+
+	const FViewInfo* View;
+	EShadingPath ShadingPath;
+	EMeshPass::Type PassType;
+	bool bUseGPUScene;
+	bool bDynamicInstancing;
+	bool bReverseCulling;
+	bool bRenderSceneTwoSided;
+	FExclusiveDepthStencil::Type BasePassDepthStencilAccess;
+
+	// Mesh pass processor.
+	FMeshPassProcessor* MeshPassProcessor;
+	FMeshPassProcessor* MobileBasePassCSMMeshPassProcessor;
+	const TArray<FMeshBatchAndRelevance, SceneRenderingAllocator>* DynamicMeshElements;
+
+	// Commands.
+	int32 NumDynamicMeshElements;
+	int32 NumDynamicMeshCommandBuildRequestElements;
+	FMeshCommandOneFrameArray MeshDrawCommands;
+	FMeshCommandOneFrameArray MobileBasePassCSMMeshDrawCommands;
+	TArray<const FStaticMesh*, SceneRenderingAllocator> DynamicMeshCommandBuildRequests;
+	TArray<const FStaticMesh*, SceneRenderingAllocator> MobileBasePassCSMDynamicMeshCommandBuildRequests;
+	FDynamicMeshDrawCommandStorage MeshDrawCommandStorage;
+
+	// Preallocated resources.
+	FGlobalDynamicVertexBuffer::FAllocation PrimitiveIdBuffer;
+	FMeshCommandOneFrameArray TempVisibleMeshDrawCommands;
+
+	// For UpdateTranslucentMeshSortKeys.
+	ETranslucencyPass::Type TranslucencyPass;
+	ETranslucentSortPolicy::Type TranslucentSortPolicy;
+	FVector TranslucentSortAxis;
+	FVector ViewOrigin;
+	FMatrix ViewMatrix;
+	const TArray<struct FPrimitiveBounds>* PrimitiveBounds;
+
+	// For logging instancing stats.
+	int32 VisibleMeshDrawCommandsNum;
+	int32 NewPassVisibleMeshDrawCommandsNum;
+	int32 MaxInstances;
+};
+
+/**
+ * Parallel mesh draw command processing and rendering. 
+ * Encapsulates two parallel tasks - mesh command setup task and drawing task.
+ */
+class FParallelMeshDrawCommandPass
+{
+public:
+	FParallelMeshDrawCommandPass()
+		: MaxNumDraws(0)
+	{
+	}
+
+	~FParallelMeshDrawCommandPass();
+
+	/**
+	 * Dispatch visible mesh draw command process task, which prepares this pass for drawing.
+	 * This includes generation of dynamic mesh draw commands, draw sorting and draw merging.
+	 */
+	void DispatchPassSetup(
+		FScene* Scene,
+		const FViewInfo& View, 
+		EMeshPass::Type PassType, 
+		FExclusiveDepthStencil::Type BasePassDepthStencilAccess,
+		FMeshPassProcessor* MeshPassProcessor,
+		const TArray<FMeshBatchAndRelevance, SceneRenderingAllocator>& DynamicMeshElements,
+		int32 NumDynamicMeshElements,
+		TArray<const FStaticMesh*, SceneRenderingAllocator>& InOutDynamicMeshCommandBuildRequests,
+		int32 NumDynamicMeshCommandBuildRequestElements,
+		FMeshCommandOneFrameArray& InOutMeshDrawCommands,
+		FMeshPassProcessor* MobileBasePassCSMMeshPassProcessor = nullptr, // Required only for the mobile base pass.
+		FMeshCommandOneFrameArray* InOutMobileBasePassCSMMeshDrawCommands = nullptr // Required only for the mobile base pass.
+	);
+
+	/**
+	 * Dispatch visible mesh draw command draw task.
+	 */
+	void DispatchDraw(FParallelCommandListSet* ParallelCommandListSet, FRHICommandList& RHICmdList) const;
+
+	void Empty();
+	void SetDumpInstancingStats(const FString& InPassName);
+	bool HasAnyDraw() const { return MaxNumDraws > 0; }
+
+
+private:
+	FMeshDrawCommandPassSetupTaskContext TaskContext;
+	FGraphEventArray TaskEventRefs;
+	FString PassNameForStats;
+
+	// Maximum number of draws for this pass. Used to prealocate resources on rendering thread. 
+	// Has a guarantee that if there won't be any draws, then MaxNumDraws = 0;
+	int32 MaxNumDraws;
+
+	void DumpInstancingStats() const;
+	void WaitForMeshPassSetupTask() const;
+};
diff --git a/Engine/Source/Runtime/Renderer/Private/SceneRendering.cpp b/Engine/Source/Runtime/Renderer/Private/SceneRendering.cpp
--- a/Engine/Source/Runtime/Renderer/Private/SceneRendering.cpp
+++ b/Engine/Source/Runtime/Renderer/Private/SceneRendering.cpp
@@ -49,6 +49,7 @@
 #include "TranslucentRendering.h"
 #include "VisualizeTexture.h"
 #include "VisualizeTexturePresent.h"
+#include "MeshDrawCommands.h"
 
 /*-----------------------------------------------------------------------------
 	Globals
@@ -82,12 +83,6 @@ static TAutoConsoleVariable<int32> CVarMeshDrawCommandsDynamicInstancing(
 	TEXT("Whether to dynamically combine multiple compatible visible Mesh Draw Commands into one instanced draw on vertex factories that support it."),
 	ECVF_RenderThreadSafe);
 
-static TAutoConsoleVariable<int32> CVarMeshDrawCommandsParallelPassSetup(
-	TEXT("r.MeshDrawCommands.ParallelPassSetup"),
-	1,
-	TEXT("Whether to setup mesh draw command pass in parallel."),
-	ECVF_RenderThreadSafe);
-
 bool IsDynamicInstancingEnabled()
 {
 	return CVarMeshDrawCommandsDynamicInstancing.GetValueOnRenderThread() > 0;
@@ -2664,955 +2659,6 @@ void FSceneRenderer::RenderFinish(FRHICommandListImmediate& RHICmdList)
 	RHICmdList.EndScene();
 }
 
-struct FCompareFMeshDrawCommands
-{
-	FORCEINLINE bool operator() (FVisibleMeshDrawCommand A, FVisibleMeshDrawCommand B) const
-	{
-		// First order by a sort key.
-		if (A.SortKey != B.SortKey)
-		{
-			return A.SortKey < B.SortKey;
-		}
-
-		// Next order by instancing bucket.
-		if (A.StateBucketId != B.StateBucketId)
-		{
-			return A.StateBucketId < B.StateBucketId;
-		}
-
-		// No valid StateBucketId - we cannot dynamically instance but we can still order by PSO to maximize state filtering
-		if (A.StateBucketId == -1 && B.StateBucketId == -1)
-		{
-			if (A.MeshDrawCommand->CachedPipelineId.GetId() != B.MeshDrawCommand->CachedPipelineId.GetId())
-			{
-				return A.MeshDrawCommand->CachedPipelineId.GetId() < B.MeshDrawCommand->CachedPipelineId.GetId();
-			}
-		}
-
-		return false;
-	}
-};
-
-uint32 BitInvertIfNegativeFloat(uint32 f)
-{
-	unsigned mask = -int32(f >> 31) | 0x80000000;
-	return f ^ mask;
-}
-
-/**
-* Update mesh sort keys with view dependent data.
-*/
-void UpdateTranslucentMeshSortKeys(
-	ETranslucentSortPolicy::Type TranslucentSortPolicy,
-	const FVector& TranslucentSortAxis,
-	const FVector& ViewOrigin,
-	const FMatrix& ViewMatrix,
-	const TArray<struct FPrimitiveBounds>& PrimitiveBounds,
-	ETranslucencyPass::Type TranslucencyPass, 
-	FMeshCommandOneFrameArray& VisibleMeshCommands
-	)
-{
-	QUICK_SCOPE_CYCLE_COUNTER(STAT_UpdateTranslucentMeshSortKeys);
-
-	for (int32 CommandIndex = 0; CommandIndex < VisibleMeshCommands.Num(); ++CommandIndex)
-	{
-		FVisibleMeshDrawCommand& VisibleCommand = VisibleMeshCommands[CommandIndex];
-
-		FVector BoundsOrigin = FVector::ZeroVector;
-		if (VisibleCommand.DrawPrimitiveId < PrimitiveBounds.Num())
-		{
-			BoundsOrigin = PrimitiveBounds[VisibleCommand.DrawPrimitiveId].BoxSphereBounds.Origin;
-		}
-
-		float Distance = 0.0f;
-		if (TranslucentSortPolicy == ETranslucentSortPolicy::SortByDistance)
-		{
-			//sort based on distance to the view position, view rotation is not a factor
-			Distance = (BoundsOrigin - ViewOrigin).Size();
-		}
-		else if (TranslucentSortPolicy == ETranslucentSortPolicy::SortAlongAxis)
-		{
-			// Sort based on enforced orthogonal distance
-			const FVector CameraToObject = BoundsOrigin - ViewOrigin;
-			Distance = FVector::DotProduct(CameraToObject, TranslucentSortAxis);
-		}
-		else
-		{
-			// Sort based on projected Z distance
-			check(TranslucentSortPolicy == ETranslucentSortPolicy::SortByProjectedZ);
-			Distance = ViewMatrix.TransformPosition(BoundsOrigin).Z;
-		}
-
-		// Patch distance inside translucent mesh sort key.
-		UTranslucentMeshSortKey TranslucentMeshSortKey;
-		TranslucentMeshSortKey.PackedData = VisibleCommand.SortKey.PackedData;
-		TranslucentMeshSortKey.Fields.Distance = (uint32)~BitInvertIfNegativeFloat(*((uint32*)&Distance));
-		VisibleCommand.SortKey.PackedData = TranslucentMeshSortKey.PackedData;		
-	}
-}
-
-static uint64 GetMobileBasePassSortKey(bool bMasked, bool bBackground, int32 PipelineId, int32 StateBucketId, float PipelineDistance, float PrimitiveDistance)
-{
-	const float PrimitiveDepthQuantization = ((1 << 14) - 1);
-
-	union
-	{
-		uint64 PackedInt;
-		struct
-		{
-			uint64 DepthBits : 14;			// Order by primitive depth
-			uint64 StateBucketId : 20; 		// Order by state bucket
-			uint64 PipelineId : 20;			// Order by PSO
-			uint64 PipelineDepthBits : 8;	// Order PSOs front to back
-			uint64 Background : 1;			// Non-background meshes first 
-			uint64 Masked : 1;				// Non-masked first
-		} Fields;
-	} Key;
-
-	union FFloatToInt { float F; uint32 I; };
-	FFloatToInt F2I;
-
-	Key.PackedInt = 0;
-	Key.Fields.Masked = bMasked;
-	Key.Fields.Background = bBackground;
-	F2I.F = PipelineDistance / HALF_WORLD_MAX;
-	Key.Fields.PipelineDepthBits = (F2I.I >> 23) & 0xff; // 8 bit exponent
-	Key.Fields.PipelineId = PipelineId;
-	Key.Fields.StateBucketId = StateBucketId;
-	Key.Fields.DepthBits = int32((FMath::Min<float>(PrimitiveDistance, HALF_WORLD_MAX) / HALF_WORLD_MAX) * PrimitiveDepthQuantization);
-
-	return Key.PackedInt;
-}
-
-/**
-* Merge mobile BasePass with BasePassCSM based on CSM visibility in order to select appropriate shader for given command.
-*/
-void MergeMobileBasePassMeshDrawCommands(
-	const FMobileCSMVisibilityInfo& MobileCSMVisibilityInfo,
-	int32 ScenePrimitiveNum,
-	FMeshCommandOneFrameArray& MeshCommands,
-	FMeshCommandOneFrameArray& MeshCommandsCSM
-	)
-{
-	if (MobileCSMVisibilityInfo.bMobileDynamicCSMInUse)
-	{
-		// determine per view CSM visibility
-		checkf(MeshCommands.Num() == MeshCommandsCSM.Num(), TEXT("VisibleMeshDrawCommands of BasePass and MobileBasePassCSM are expected to match."));
-		for (int32 i = MeshCommands.Num() - 1; i >= 0; --i)
-		{
-			FVisibleMeshDrawCommand& MeshCommand = MeshCommands[i];
-			FVisibleMeshDrawCommand& MeshCommandCSM = MeshCommandsCSM[i];
-
-			if (MobileCSMVisibilityInfo.bAlwaysUseCSM 
-				|| (MeshCommand.DrawPrimitiveId < ScenePrimitiveNum && MobileCSMVisibilityInfo.MobilePrimitiveCSMReceiverVisibilityMap[MeshCommand.DrawPrimitiveId]))
-			{
-				checkf(MeshCommand.DrawPrimitiveId == MeshCommandCSM.DrawPrimitiveId, TEXT("VisibleMeshDrawCommands of BasePass and MobileBasePassCSM are expected to match."));
-				// Use CSM's VisibleMeshDrawCommand.
-				MeshCommand = MeshCommandCSM;
-			}
-		}
-		MeshCommandsCSM.Reset();
-	}
-}
-
-/**
-* Compute mesh sort keys for the mobile base pass
-*/
-void UpdateMobileBasePassMeshSortKeys(
-	const FVector& ViewOrigin,
-	const TArray<struct FPrimitiveBounds>& ScenePrimitiveBounds,
-	FMeshCommandOneFrameArray& VisibleMeshCommands
-)
-{
-	QUICK_SCOPE_CYCLE_COUNTER(STAT_UpdateMobileBasePassMeshSortKeys);
-
-	int32 NumCmds = VisibleMeshCommands.Num();
-
-	// pre-compute distance to a group of meshes that share same PSO
-	TMap<int32, float> PipelineDistances;
-	PipelineDistances.Reserve(256);
-
-	for (int32 CmdIdx = 0; CmdIdx < NumCmds; ++CmdIdx)
-	{
-		FVisibleMeshDrawCommand& Cmd = VisibleMeshCommands[CmdIdx];
-
-		float PrimitiveDistance = 0;
-
-		if (Cmd.DrawPrimitiveId < ScenePrimitiveBounds.Num())
-		{
-			const FPrimitiveBounds& PrimitiveBounds = ScenePrimitiveBounds[Cmd.DrawPrimitiveId];
-			PrimitiveDistance = (PrimitiveBounds.BoxSphereBounds.Origin - ViewOrigin).Size();
-		}
-
-		float& PipelineDistance = PipelineDistances.FindOrAdd(Cmd.MeshDrawCommand->CachedPipelineId.GetId());
-		// not sure what could be better: average distance, max or min
-		PipelineDistance = FMath::Max(PipelineDistance, PrimitiveDistance);
-	}
-
-	// compute sort key for each mesh command
-	for (int32 CmdIdx = 0; CmdIdx < NumCmds; ++CmdIdx)
-	{
-		FVisibleMeshDrawCommand& Cmd = VisibleMeshCommands[CmdIdx];
-		float PrimitiveDistance = 0;
-		bool bMasked = false;
-		bool bBackground = false;
-
-		if (Cmd.DrawPrimitiveId < ScenePrimitiveBounds.Num())
-		{
-			const FPrimitiveBounds& PrimitiveBounds = ScenePrimitiveBounds[Cmd.DrawPrimitiveId];
-			PrimitiveDistance = (PrimitiveBounds.BoxSphereBounds.Origin - ViewOrigin).Size();
-			bBackground = (PrimitiveBounds.BoxSphereBounds.SphereRadius > HALF_WORLD_MAX / 4.0f); // TODO: per-primitive bBackgound flag
-			bMasked = Cmd.SortKey.PackedData & 0x1 ? true : false; // Set in MobileBasePass.cpp - GetBasePassStaticSortKey
-		}
-
-		int32 PipelineId = Cmd.MeshDrawCommand->CachedPipelineId.GetId();
-		float PipelineDistance = PipelineDistances.FindRef(PipelineId);
-		// poor man StateID, can't use Cmd.StateBucketId as it is unique for each primitive if platform does not support auto-instancing
-		int32 StateBucketId = PointerHash(Cmd.MeshDrawCommand->IndexBuffer);
-
-		Cmd.SortKey.PackedData = GetMobileBasePassSortKey(bMasked, bBackground, PipelineId, StateBucketId, PipelineDistance, PrimitiveDistance);
-	}
-}
-
-/**
- * Build mesh draw command primitive Id buffer for instancing.
- * TempVisibleMeshDrawCommands must be presized for NewPassVisibleMeshDrawCommands.
- */
-void BuildMeshDrawCommandPrimitiveIdBuffer(
-	FMeshCommandOneFrameArray& VisibleMeshDrawCommands,
-	FDynamicMeshDrawCommandStorage& MeshDrawCommandStorage,
-	FGlobalDynamicVertexBuffer::FAllocation& PrimitiveIdBuffer,
-	FMeshCommandOneFrameArray& TempVisibleMeshDrawCommands,
-	int32& MaxInstances,
-	int32& VisibleMeshDrawCommandsNum,
-	int32& NewPassVisibleMeshDrawCommandsNum
-	)
-{
-	QUICK_SCOPE_CYCLE_COUNTER(STAT_BuildMeshDrawCommandPrimitiveIdBuffer);
-	check(VisibleMeshDrawCommands.Num() <= TempVisibleMeshDrawCommands.Max() && TempVisibleMeshDrawCommands.Num() == 0 && PrimitiveIdBuffer.IsValid());
-
-	const FVisibleMeshDrawCommand* RESTRICT PassVisibleMeshDrawCommands = VisibleMeshDrawCommands.GetData();
-	const int32 NumDrawCommands = VisibleMeshDrawCommands.Num();
-
-	int32* RESTRICT PrimitiveIds = (int32*)PrimitiveIdBuffer.Buffer;
-
-	if (IsDynamicInstancingEnabled())
-	{
-		QUICK_SCOPE_CYCLE_COUNTER(STAT_DynamicInstancingOfVisibleMeshDrawCommands);
-
-		int32 CurrentStateBucketId = -1;
-		uint32* RESTRICT CurrentDynamicallyInstancedMeshCommandNumInstances = nullptr;
-		MaxInstances = 1;
-
-		for (int32 DrawCommandIndex = 0; DrawCommandIndex < NumDrawCommands; DrawCommandIndex++)
-		{
-			const FVisibleMeshDrawCommand& RESTRICT VisibleMeshDrawCommand = PassVisibleMeshDrawCommands[DrawCommandIndex];
-
-			if (VisibleMeshDrawCommand.StateBucketId == CurrentStateBucketId && VisibleMeshDrawCommand.StateBucketId != -1)
-			{
-				if (CurrentDynamicallyInstancedMeshCommandNumInstances)
-				{
-					const int32 CurrentNumInstances = *CurrentDynamicallyInstancedMeshCommandNumInstances;
-					*CurrentDynamicallyInstancedMeshCommandNumInstances = CurrentNumInstances + 1;
-					MaxInstances = FMath::Max(MaxInstances, CurrentNumInstances + 1);
-				}
-				else
-				{
-					FVisibleMeshDrawCommand NewVisibleMeshDrawCommand = VisibleMeshDrawCommand;
-					NewVisibleMeshDrawCommand.PrimitiveIdBufferOffset = DrawCommandIndex;
-					TempVisibleMeshDrawCommands.Emplace(MoveTemp(NewVisibleMeshDrawCommand));
-				}
-			}
-			else
-			{
-				// First time state bucket setup
-				CurrentStateBucketId = VisibleMeshDrawCommand.StateBucketId;
-
-				if (VisibleMeshDrawCommand.MeshDrawCommand->PrimitiveIdStreamIndex >= 0
-					&& VisibleMeshDrawCommand.MeshDrawCommand->NumInstances == 1
-					// Don't create a new FMeshDrawCommand for the last command and make it safe for us to look at the next command
-					&& DrawCommandIndex + 1 < NumDrawCommands
-					// Only create a new FMeshDrawCommand if more than one draw in the state bucket
-					&& CurrentStateBucketId == PassVisibleMeshDrawCommands[DrawCommandIndex + 1].StateBucketId)
-				{
-					const int32 Index = MeshDrawCommandStorage.MeshDrawCommands.AddElement(*VisibleMeshDrawCommand.MeshDrawCommand);
-					FMeshDrawCommand& NewCommand = MeshDrawCommandStorage.MeshDrawCommands[Index];
-					FVisibleMeshDrawCommand NewVisibleMeshDrawCommand;
-
-					NewVisibleMeshDrawCommand.Setup(
-						&NewCommand,
-						VisibleMeshDrawCommand.DrawPrimitiveId,
-						VisibleMeshDrawCommand.StateBucketId,
-						VisibleMeshDrawCommand.MeshFillMode,
-						VisibleMeshDrawCommand.MeshCullMode,
-						VisibleMeshDrawCommand.SortKey);
-
-					NewVisibleMeshDrawCommand.PrimitiveIdBufferOffset = DrawCommandIndex;
-					TempVisibleMeshDrawCommands.Emplace(MoveTemp(NewVisibleMeshDrawCommand));
-
-					CurrentDynamicallyInstancedMeshCommandNumInstances = &NewCommand.NumInstances;
-				}
-				else
-				{
-					CurrentDynamicallyInstancedMeshCommandNumInstances = nullptr;
-					FVisibleMeshDrawCommand NewVisibleMeshDrawCommand = VisibleMeshDrawCommand;
-					NewVisibleMeshDrawCommand.PrimitiveIdBufferOffset = DrawCommandIndex;
-					TempVisibleMeshDrawCommands.Emplace(MoveTemp(NewVisibleMeshDrawCommand));
-				}
-			}
-
-			//@todo - refactor into memcpy
-			PrimitiveIds[DrawCommandIndex] = VisibleMeshDrawCommand.DrawPrimitiveId;
-		}
-
-		// Setup instancing stats for logging.
-		VisibleMeshDrawCommandsNum = VisibleMeshDrawCommands.Num();
-		NewPassVisibleMeshDrawCommandsNum = TempVisibleMeshDrawCommands.Num();
-
-		// Replace VisibleMeshDrawCommands
-		FMemory::Memswap(&VisibleMeshDrawCommands, &TempVisibleMeshDrawCommands, sizeof(TempVisibleMeshDrawCommands));
-		TempVisibleMeshDrawCommands.Reset();
-	}
-	else
-	{
-		QUICK_SCOPE_CYCLE_COUNTER(STAT_BuildVisibleMeshDrawCommandPrimitiveIdBuffers);
-
-		for (int32 DrawCommandIndex = 0; DrawCommandIndex < NumDrawCommands; DrawCommandIndex++)
-		{
-			const FVisibleMeshDrawCommand& VisibleMeshDrawCommand = VisibleMeshDrawCommands[DrawCommandIndex];
-			PrimitiveIds[DrawCommandIndex] = VisibleMeshDrawCommand.DrawPrimitiveId;
-		}
-	}
-}
-
-/**
- * Converts each FMeshBatch into a set of FMeshDrawCommands for a specific mesh pass type.
- */
-void GenerateDynamicMeshDrawCommands(
-	const FViewInfo& View,
-	EShadingPath ShadingPath,
-	EMeshPass::Type PassType,
-	FMeshPassProcessor* PassMeshProcessor,
-	const TArray<FMeshBatchAndRelevance, SceneRenderingAllocator>& DynamicMeshElements,
-	int32 NumDynamicMeshElements,
-	const TArray<const FStaticMesh*, SceneRenderingAllocator>& DynamicMeshCommandBuildRequests,
-	int32 NumDynamicMeshCommandBuildRequestElements,
-	FMeshCommandOneFrameArray& VisibleCommands,
-	FDynamicMeshDrawCommandStorage& MeshDrawCommandStorage
-)
-{
-	QUICK_SCOPE_CYCLE_COUNTER(STAT_GenerateDynamicMeshDrawCommands);
-	check(PassMeshProcessor);
-
-	FDynamicPassMeshDrawListContext DynamicPassMeshDrawListContext(
-		MeshDrawCommandStorage,
-		VisibleCommands
-	);
-	PassMeshProcessor->SetDrawListContext(&DynamicPassMeshDrawListContext);
-
-	const int32 NumVisibleCommandsBeforeDynamicBatches = VisibleCommands.Num();
-	const int32 NumDynamicMeshBatches = DynamicMeshElements.Num();
-
-	for (int32 MeshIndex = 0; MeshIndex < NumDynamicMeshBatches; MeshIndex++)
-	{
-		const FMeshBatchAndRelevance& MeshAndRelevance = DynamicMeshElements[MeshIndex];
-		check(!MeshAndRelevance.Mesh->bRequiresPerElementVisibility);
-		const uint64 BatchElementMask = ~0ull;
-
-		PassMeshProcessor->AddMeshBatch(*MeshAndRelevance.Mesh, BatchElementMask, MeshAndRelevance.PrimitiveSceneProxy);
-	}
-
-	// Check if we didn't generate more mesh draw commands, than we assumed.
-	check(VisibleCommands.Num() - NumVisibleCommandsBeforeDynamicBatches <= NumDynamicMeshElements);
-
-	const int32 NumVisibleCommandsBeforeDynamicRequests = VisibleCommands.Num();
-	const int32 NumStaticMeshBatches = DynamicMeshCommandBuildRequests.Num();
-
-	for (int32 MeshIndex = 0; MeshIndex < NumStaticMeshBatches; MeshIndex++)
-	{
-		const FStaticMesh* StaticMeshBatch = DynamicMeshCommandBuildRequests[MeshIndex];
-		const uint64 BatchElementMask = StaticMeshBatch->bRequiresPerElementVisibility ? View.StaticMeshBatchVisibility[StaticMeshBatch->BatchVisibilityId] : ~0ull;
-
-		PassMeshProcessor->AddMeshBatch(*StaticMeshBatch, BatchElementMask, StaticMeshBatch->PrimitiveSceneInfo->Proxy, StaticMeshBatch->Id);
-	}
-
-	// Check if we didn't generate more mesh draw commands, than we assumed.
-	check(VisibleCommands.Num() - NumVisibleCommandsBeforeDynamicRequests <= NumDynamicMeshCommandBuildRequestElements);
-}
-
-/**
- * Special version of GenerateDynamicMeshDrawCommands for the mobile base pass.
- * Based on CSM visibility it will generate mesh draw commands using either normal base pass processor or CSM base pass processor.
-*/
-void GenerateMobileBasePassDynamicMeshDrawCommands(
-	const FViewInfo& View,
-	EShadingPath ShadingPath,
-	EMeshPass::Type PassType,
-	FMeshPassProcessor* PassMeshProcessor,
-	FMeshPassProcessor* MobilePassCSMPassMeshProcessor,
-	const TArray<FMeshBatchAndRelevance, SceneRenderingAllocator>& DynamicMeshElements,
-	int32 NumDynamicMeshElements,
-	const TArray<const FStaticMesh*, SceneRenderingAllocator>& DynamicMeshCommandBuildRequests,
-	int32 NumDynamicMeshCommandBuildRequestElements,
-	FMeshCommandOneFrameArray& VisibleCommands,
-	FDynamicMeshDrawCommandStorage& MeshDrawCommandStorage
-)
-{
-	QUICK_SCOPE_CYCLE_COUNTER(STAT_GenerateMobileBasePassDynamicMeshDrawCommands);
-	check(PassMeshProcessor && MobilePassCSMPassMeshProcessor);
-
-	FDynamicPassMeshDrawListContext DynamicPassMeshDrawListContext(
-		MeshDrawCommandStorage,
-		VisibleCommands
-	);
-	PassMeshProcessor->SetDrawListContext(&DynamicPassMeshDrawListContext);
-	MobilePassCSMPassMeshProcessor->SetDrawListContext(&DynamicPassMeshDrawListContext);
-
-	const FMobileCSMVisibilityInfo& MobileCSMVisibilityInfo = View.MobileCSMVisibilityInfo;
-	
-	const int32 NumVisibleCommandsBeforeDynamicBatches = VisibleCommands.Num();
-	const int32 NumDynamicMeshBatches = DynamicMeshElements.Num();
-
-	for (int32 MeshIndex = 0; MeshIndex < NumDynamicMeshBatches; MeshIndex++)
-	{
-		const FMeshBatchAndRelevance& MeshAndRelevance = DynamicMeshElements[MeshIndex];
-		check(!MeshAndRelevance.Mesh->bRequiresPerElementVisibility);
-		const uint64 BatchElementMask = ~0ull;
-
-		const int32 PrimitiveIndex = MeshAndRelevance.PrimitiveSceneProxy->GetPrimitiveSceneInfo()->GetIndex();
-		if (MobileCSMVisibilityInfo.bMobileDynamicCSMInUse 
-			&& (MobileCSMVisibilityInfo.bAlwaysUseCSM || MobileCSMVisibilityInfo.MobilePrimitiveCSMReceiverVisibilityMap[PrimitiveIndex]))
-		{
-			MobilePassCSMPassMeshProcessor->AddMeshBatch(*MeshAndRelevance.Mesh, BatchElementMask, MeshAndRelevance.PrimitiveSceneProxy);
-		}
-		else
-		{
-			PassMeshProcessor->AddMeshBatch(*MeshAndRelevance.Mesh, BatchElementMask, MeshAndRelevance.PrimitiveSceneProxy);
-		}
-	}
-
-	// Check if we didn't generate more mesh draw commands, than we assumed.
-	check(VisibleCommands.Num() - NumVisibleCommandsBeforeDynamicBatches <= NumDynamicMeshElements);
-
-	const int32 NumVisibleCommandsBeforeDynamicRequests = VisibleCommands.Num();
-	const int32 NumStaticMeshBatches = DynamicMeshCommandBuildRequests.Num();
-
-	for (int32 MeshIndex = 0; MeshIndex < NumStaticMeshBatches; MeshIndex++)
-	{
-		const FStaticMesh* StaticMeshBatch = DynamicMeshCommandBuildRequests[MeshIndex];
-		const uint64 BatchElementMask = StaticMeshBatch->bRequiresPerElementVisibility ? View.StaticMeshBatchVisibility[StaticMeshBatch->BatchVisibilityId] : ~0ull;
-
-		const int32 PrimitiveIndex = StaticMeshBatch->PrimitiveSceneInfo->Proxy->GetPrimitiveSceneInfo()->GetIndex();
-		if (MobileCSMVisibilityInfo.bMobileDynamicCSMInUse
-			&& (MobileCSMVisibilityInfo.bAlwaysUseCSM || MobileCSMVisibilityInfo.MobilePrimitiveCSMReceiverVisibilityMap[PrimitiveIndex]))
-		{
-			MobilePassCSMPassMeshProcessor->AddMeshBatch(*StaticMeshBatch, BatchElementMask, StaticMeshBatch->PrimitiveSceneInfo->Proxy, StaticMeshBatch->Id);
-		}
-		else
-		{
-			PassMeshProcessor->AddMeshBatch(*StaticMeshBatch, BatchElementMask, StaticMeshBatch->PrimitiveSceneInfo->Proxy, StaticMeshBatch->Id);
-		}
-	}
-
-	// Check if we didn't generate more mesh draw commands, than we assumed.
-	check(VisibleCommands.Num() - NumVisibleCommandsBeforeDynamicRequests <= NumDynamicMeshCommandBuildRequestElements);
-}
-
-/**
-* Apply view overrides to existing mesh draw commands (e.g. reverse culling mode for rendering planar reflections).
-* TempVisibleMeshDrawCommands must be presized for NewPassVisibleMeshDrawCommands.
-*/
-void ApplyViewOverridesToMeshDrawCommands(
-	EShadingPath ShadingPath,
-	EMeshPass::Type PassType,
-	bool bReverseCulling,
-	bool bRenderSceneTwoSided,
-	FExclusiveDepthStencil::Type BasePassDepthStencilAccess,
-	FMeshCommandOneFrameArray& VisibleMeshDrawCommands,
-	FDynamicMeshDrawCommandStorage& MeshDrawCommandStorage,
-	FMeshCommandOneFrameArray& TempVisibleMeshDrawCommands
-	)
-{
-	QUICK_SCOPE_CYCLE_COUNTER(STAT_ApplyViewOverridesToMeshDrawCommands);
-	check(VisibleMeshDrawCommands.Num() <= TempVisibleMeshDrawCommands.Max() && TempVisibleMeshDrawCommands.Num() == 0 && PassType != EMeshPass::Num);
-
-	if ((FPassProcessorManager::GetPassFlags(ShadingPath, PassType) & EMeshPassFlags::MainView) != EMeshPassFlags::None)
-	{
-		const FExclusiveDepthStencil::Type DefaultBasePassDepthStencilAccess = GetDefaultBasePassDepthStencilAccess(ShadingPath);
-
-		if (bReverseCulling || bRenderSceneTwoSided || (BasePassDepthStencilAccess != DefaultBasePassDepthStencilAccess && PassType == EMeshPass::BasePass))
-		{
-			for (int32 MeshCommandIndex = 0; MeshCommandIndex < VisibleMeshDrawCommands.Num(); MeshCommandIndex++)
-			{
-				MeshDrawCommandStorage.MeshDrawCommands.Add(1);
-				FMeshDrawCommand& NewMeshCommand = MeshDrawCommandStorage.MeshDrawCommands[MeshDrawCommandStorage.MeshDrawCommands.Num() - 1];
-
-				const FVisibleMeshDrawCommand& VisibleMeshDrawCommand = VisibleMeshDrawCommands[MeshCommandIndex];
-				const FMeshDrawCommand& MeshCommand = *VisibleMeshDrawCommand.MeshDrawCommand;
-				NewMeshCommand = MeshCommand;
-
-				const ERasterizerCullMode LocalCullMode = bRenderSceneTwoSided ? CM_None : bReverseCulling ? FMeshDrawingPolicy::InverseCullMode(VisibleMeshDrawCommand.MeshCullMode) : VisibleMeshDrawCommand.MeshCullMode;
-				NewMeshCommand.PipelineState.RasterizerState = GetStaticRasterizerState<true>(VisibleMeshDrawCommand.MeshFillMode, LocalCullMode);
-
-				if (BasePassDepthStencilAccess != DefaultBasePassDepthStencilAccess && PassType == EMeshPass::BasePass)
-				{
-					FDrawingPolicyRenderState PassDrawRenderState;
-					SetupBasePassState(BasePassDepthStencilAccess, false, PassDrawRenderState);
-					NewMeshCommand.PipelineState.DepthStencilState = PassDrawRenderState.GetDepthStencilState();
-				}
-
-				NewMeshCommand.Finalize(true);
-
-				FVisibleMeshDrawCommand NewVisibleMeshDrawCommand;
-
-				NewVisibleMeshDrawCommand.Setup(
-					&NewMeshCommand,
-					VisibleMeshDrawCommand.DrawPrimitiveId,
-					VisibleMeshDrawCommand.StateBucketId,
-					VisibleMeshDrawCommand.MeshFillMode,
-					VisibleMeshDrawCommand.MeshCullMode,
-					VisibleMeshDrawCommand.SortKey);
-
-				TempVisibleMeshDrawCommands.Add(NewVisibleMeshDrawCommand);
-			}
-
-			// Replace VisibleMeshDrawCommands
-			FMemory::Memswap(&VisibleMeshDrawCommands, &TempVisibleMeshDrawCommands, sizeof(TempVisibleMeshDrawCommands));
-			TempVisibleMeshDrawCommands.Reset();
-		}
-	}
-}
-
-FAutoConsoleTaskPriority CPrio_FMeshDrawCommandPassSetupTask(
-	TEXT("TaskGraph.TaskPriorities.FMeshDrawCommandPassSetupTask"),
-	TEXT("Task and thread priority for FMeshDrawCommandPassSetupTask."),
-	ENamedThreads::HighThreadPriority, // if we have high priority task threads, then use them...
-	ENamedThreads::NormalTaskPriority, // .. at normal task priority
-	ENamedThreads::HighTaskPriority // if we don't have hi pri threads, then use normal priority threads at high task priority instead
-);
-
-/**
- * Task for a parallel setup of mesh draw commands. Includes generation of dynamic mesh draw commands, sorting, merging etc.
- */
-class FMeshDrawCommandPassSetupTask
-{
-public:
-	FMeshDrawCommandPassSetupTask(FMeshDrawCommandPassSetupTaskContext& InContext)
-		: Context(InContext)
-	{
-	}
-
-	FORCEINLINE TStatId GetStatId() const
-	{
-		RETURN_QUICK_DECLARE_CYCLE_STAT(FMeshDrawCommandPassSetupTask, STATGROUP_TaskGraphTasks);
-	}
-
-	ENamedThreads::Type GetDesiredThread()
-	{
-		return CPrio_FMeshDrawCommandPassSetupTask.Get();
-	}
-
-	static ESubsequentsMode::Type GetSubsequentsMode() 
-	{ 
-		return ESubsequentsMode::TrackSubsequents; 
-	}
-
-	void AnyThreadTask()
-	{
-		// Mobile base pass is a special case, as final lists is created from two mesh passes based on CSM visibility.
-		const bool bMobileBasePass = Context.ShadingPath == EShadingPath::Mobile && Context.PassType == EMeshPass::BasePass;
-
-		if (bMobileBasePass)
-		{
-			MergeMobileBasePassMeshDrawCommands(
-				Context.View->MobileCSMVisibilityInfo,
-				Context.PrimitiveBounds->Num(),
-				Context.MeshDrawCommands,
-				Context.MobileBasePassCSMMeshDrawCommands
-			);
-
-			GenerateMobileBasePassDynamicMeshDrawCommands(
-				*Context.View,
-				Context.ShadingPath,
-				Context.PassType,
-				Context.MeshPassProcessor,
-				Context.MobileBasePassCSMMeshPassProcessor,
-				*Context.DynamicMeshElements,
-				Context.NumDynamicMeshElements,
-				Context.DynamicMeshCommandBuildRequests,
-				Context.NumDynamicMeshCommandBuildRequestElements,
-				Context.MeshDrawCommands,
-				Context.MeshDrawCommandStorage
-			);
-		}
-		else
-		{
-			GenerateDynamicMeshDrawCommands(
-				*Context.View,
-				Context.ShadingPath,
-				Context.PassType,
-				Context.MeshPassProcessor,
-				*Context.DynamicMeshElements,
-				Context.NumDynamicMeshElements,
-				Context.DynamicMeshCommandBuildRequests,
-				Context.NumDynamicMeshCommandBuildRequestElements,
-				Context.MeshDrawCommands,
-				Context.MeshDrawCommandStorage
-			);
-		}
-
-		if (Context.MeshDrawCommands.Num() > 0)
-		{
-			if (Context.PassType != EMeshPass::Num)
-			{
-				ApplyViewOverridesToMeshDrawCommands(
-					Context.ShadingPath,
-					Context.PassType,
-					Context.bReverseCulling,
-					Context.bRenderSceneTwoSided,
-					Context.BasePassDepthStencilAccess,
-					Context.MeshDrawCommands,
-					Context.MeshDrawCommandStorage,
-					Context.TempVisibleMeshDrawCommands
-				);
-			}
-
-			// Update sort keys.
-			if (bMobileBasePass)
-			{
-				UpdateMobileBasePassMeshSortKeys(
-					Context.ViewOrigin,
-					*Context.PrimitiveBounds,
-					Context.MeshDrawCommands
-					);
-			}
-			else if (Context.TranslucencyPass != ETranslucencyPass::TPT_MAX)
-			{
-				UpdateTranslucentMeshSortKeys(
-					Context.TranslucentSortPolicy,
-					Context.TranslucentSortAxis,
-					Context.ViewOrigin,
-					Context.ViewMatrix,
-					*Context.PrimitiveBounds,
-					Context.TranslucencyPass,
-					Context.MeshDrawCommands
-				);
-			}
-
-			{
-				QUICK_SCOPE_CYCLE_COUNTER(STAT_SortVisibleMeshDrawCommands);
-				Context.MeshDrawCommands.Sort(FCompareFMeshDrawCommands());
-			}
-
-			if (Context.bUseGPUScene)
-			{
-				BuildMeshDrawCommandPrimitiveIdBuffer(
-					Context.MeshDrawCommands,
-					Context.MeshDrawCommandStorage,
-					Context.PrimitiveIdBuffer,
-					Context.TempVisibleMeshDrawCommands,
-					Context.MaxInstances,
-					Context.VisibleMeshDrawCommandsNum,
-					Context.NewPassVisibleMeshDrawCommandsNum
-				);
-			}
-		}
-	}
-
-	void DoTask(ENamedThreads::Type CurrentThread, const FGraphEventRef& MyCompletionGraphEvent)
-	{
-		AnyThreadTask();
-	}
-
-
-private:
-	FMeshDrawCommandPassSetupTaskContext& Context;
-};
-
-void SortPassMeshDrawCommands(
-	ERHIFeatureLevel::Type FeatureLevel,
-	FMeshCommandOneFrameArray& VisibleMeshDrawCommands,
-	FDynamicMeshDrawCommandStorage& MeshDrawCommandStorage,
-	FGlobalDynamicVertexBuffer& DynamicVertexBuffer,
-	FGlobalDynamicVertexBuffer::FAllocation& PrimitiveIdBuffer)
-{
-	const bool bUseGPUScene = UseGPUScene(GMaxRHIShaderPlatform, FeatureLevel);
-
-	const int32 NumDrawCommands = VisibleMeshDrawCommands.Num();
-	if (NumDrawCommands > 0)
-	{
-		FMeshCommandOneFrameArray NewPassVisibleMeshDrawCommands;
-		int32 MaxInstances = 1;
-		int32 VisibleMeshDrawCommandsNum = 0;
-		int32 NewPassVisibleMeshDrawCommandsNum = 0;
-
-		// Preallocate context memory and resources.
-		if (bUseGPUScene)
-		{
-			PrimitiveIdBuffer = DynamicVertexBuffer.Allocate(NumDrawCommands * sizeof(int32));
-
-			if (IsDynamicInstancingEnabled())
-			{
-				NewPassVisibleMeshDrawCommands.Empty(NumDrawCommands);
-			}
-		}
-
-		VisibleMeshDrawCommands.Sort(FCompareFMeshDrawCommands());
-
-		if (bUseGPUScene)
-		{
-			BuildMeshDrawCommandPrimitiveIdBuffer(
-				VisibleMeshDrawCommands,
-				MeshDrawCommandStorage,
-				PrimitiveIdBuffer,
-				NewPassVisibleMeshDrawCommands,
-				MaxInstances,
-				VisibleMeshDrawCommandsNum,
-				NewPassVisibleMeshDrawCommandsNum
-			);
-		}
-	}
-}
-
-void FParallelMeshDrawCommandPass::DispatchPassSetup(
-	FScene* Scene,
-	const FViewInfo& View,
-	EMeshPass::Type PassType,
-	FExclusiveDepthStencil::Type BasePassDepthStencilAccess,
-	FMeshPassProcessor* MeshPassProcessor,
-	const TArray<FMeshBatchAndRelevance, SceneRenderingAllocator>& DynamicMeshElements,
-	int32 NumDynamicMeshElements,
-	TArray<const FStaticMesh*, SceneRenderingAllocator>& InOutDynamicMeshCommandBuildRequests,
-	int32 NumDynamicMeshCommandBuildRequestElements,
-	FMeshCommandOneFrameArray& InOutMeshDrawCommands,
-	FMeshPassProcessor* MobileBasePassCSMMeshPassProcessor,
-	FMeshCommandOneFrameArray* InOutMobileBasePassCSMMeshDrawCommands
-)
-{
-	check(TaskEventRefs.Num() == 0 && MeshPassProcessor != nullptr);
-
-	MaxNumDraws = InOutMeshDrawCommands.Num() + NumDynamicMeshElements + NumDynamicMeshCommandBuildRequestElements;
-	if (MaxNumDraws <= 0)
-	{
-		return;
-	}
-
-	TaskContext.MeshPassProcessor = MeshPassProcessor;
-	TaskContext.MobileBasePassCSMMeshPassProcessor = MobileBasePassCSMMeshPassProcessor;
-	TaskContext.DynamicMeshElements = &DynamicMeshElements;
-
-	TaskContext.View = &View;
-	TaskContext.ShadingPath = Scene->GetShadingPath();
-	TaskContext.PassType = PassType;
-	TaskContext.bUseGPUScene = UseGPUScene(GMaxRHIShaderPlatform, View.GetFeatureLevel());
-	TaskContext.bDynamicInstancing = TaskContext.bUseGPUScene && IsDynamicInstancingEnabled();
-	TaskContext.bReverseCulling = View.bReverseCulling;
-	TaskContext.bRenderSceneTwoSided = View.bRenderSceneTwoSided;
-	TaskContext.BasePassDepthStencilAccess = BasePassDepthStencilAccess;
-	TaskContext.NumDynamicMeshElements = NumDynamicMeshElements;
-	TaskContext.NumDynamicMeshCommandBuildRequestElements = NumDynamicMeshCommandBuildRequestElements;
-
-	// Setup translucency sort key update pass based on view.
-	TaskContext.TranslucencyPass = ETranslucencyPass::TPT_MAX;
-	TaskContext.TranslucentSortPolicy = View.TranslucentSortPolicy;
-	TaskContext.TranslucentSortAxis = View.TranslucentSortAxis;
-	TaskContext.ViewOrigin = View.ViewMatrices.GetViewOrigin();
-	TaskContext.ViewMatrix = View.ViewMatrices.GetViewMatrix();
-	TaskContext.PrimitiveBounds = &Scene->PrimitiveBounds;
-
-	switch (PassType)
-	{
-		case EMeshPass::TranslucencyStandard: TaskContext.TranslucencyPass = ETranslucencyPass::TPT_StandardTranslucency; break;
-		case EMeshPass::TranslucencyAfterDOF: TaskContext.TranslucencyPass = ETranslucencyPass::TPT_TranslucencyAfterDOF; break;
-		case EMeshPass::TranslucencyAll: TaskContext.TranslucencyPass = ETranslucencyPass::TPT_AllTranslucency; break;
-		case EMeshPass::MobileInverseOpacity: TaskContext.TranslucencyPass = ETranslucencyPass::TPT_StandardTranslucency; break;
-	}
-
-	FMemory::Memswap(&TaskContext.MeshDrawCommands, &InOutMeshDrawCommands, sizeof(InOutMeshDrawCommands));
-	FMemory::Memswap(&TaskContext.DynamicMeshCommandBuildRequests, &InOutDynamicMeshCommandBuildRequests, sizeof(InOutDynamicMeshCommandBuildRequests));
-
-	if (TaskContext.ShadingPath == EShadingPath::Mobile && TaskContext.PassType == EMeshPass::BasePass)
-	{
-		FMemory::Memswap(&TaskContext.MobileBasePassCSMMeshDrawCommands, InOutMobileBasePassCSMMeshDrawCommands, sizeof(*InOutMobileBasePassCSMMeshDrawCommands));
-	}
-	else
-	{
-		check(MobileBasePassCSMMeshPassProcessor == nullptr && InOutMobileBasePassCSMMeshDrawCommands == nullptr);
-	}
-
-	// Preallocate resources on rendering thread based on MaxNumDraws.
-	if (TaskContext.bDynamicInstancing)
-	{
-		TaskContext.PrimitiveIdBuffer = FGlobalDynamicVertexBuffer::Get().Allocate(MaxNumDraws * sizeof(int32));
-	}
-	TaskContext.MeshDrawCommands.Reserve(MaxNumDraws);
-	TaskContext.TempVisibleMeshDrawCommands.Reserve(MaxNumDraws);
-
-	const bool bExecuteInParallel = FApp::ShouldUseThreadingForPerformance() && CVarMeshDrawCommandsParallelPassSetup.GetValueOnRenderThread() > 0;
-
-	if (bExecuteInParallel)
-	{
-		TaskEventRefs.Add(TGraphTask<FMeshDrawCommandPassSetupTask>::CreateTask(
-			nullptr, ENamedThreads::GetRenderThread()).ConstructAndDispatchWhenReady(TaskContext));
-	}
-	else
-	{
-		FMeshDrawCommandPassSetupTask Task(TaskContext);
-		Task.AnyThreadTask();
-	}
-}
-
-void FParallelMeshDrawCommandPass::WaitForMeshPassSetupTask() const
-{
-	if (TaskEventRefs.Num() > 0)
-	{
-		// Need to wait on GetRenderThread_Local, as mesh pass setup task can wait on rendering thread inside InitResourceFromPossiblyParallelRendering().
-		FTaskGraphInterface::Get().WaitUntilTasksComplete(TaskEventRefs, ENamedThreads::GetRenderThread_Local());
-	}
-}
-
-void FParallelMeshDrawCommandPass::Empty()
-{
-	// Need to wait in case if someone dispatched sort and draw merge task, but didn't draw it.
-	WaitForMeshPassSetupTask();
-	TaskEventRefs.Empty();
-
-	DumpInstancingStats();
-
-	MaxNumDraws = 0;
-	TaskContext.MeshPassProcessor = nullptr;
-	TaskContext.MobileBasePassCSMMeshPassProcessor = nullptr;
-	TaskContext.DynamicMeshElements = nullptr;
-	PassNameForStats.Empty();
-	TaskContext.MeshDrawCommands.Empty();
-	TaskContext.MobileBasePassCSMMeshDrawCommands.Empty();
-	TaskContext.DynamicMeshCommandBuildRequests.Empty();
-	TaskContext.TempVisibleMeshDrawCommands.Empty();
-	TaskContext.PrimitiveIdBuffer = FGlobalDynamicVertexBuffer::FAllocation();
-}
-
-FParallelMeshDrawCommandPass::~FParallelMeshDrawCommandPass()
-{
-	Empty();
-}
-
-class FDrawVisibleMeshCommandsAnyThreadTask : public FRenderTask
-{
-	FRHICommandList& RHICmdList;
-	const FMeshCommandOneFrameArray& VisibleMeshDrawCommands;
-	FVertexBufferRHIParamRef PrimitiveIdsBuffer;
-	int32 BasePrimitiveIdsOffset;
-	bool bDynamicInstancing;
-	int32 TaskIndex;
-	int32 TaskNum;
-
-public:
-
-	FDrawVisibleMeshCommandsAnyThreadTask(
-		FRHICommandList& InRHICmdList,
-		const FMeshCommandOneFrameArray& InVisibleMeshDrawCommands,
-		FVertexBufferRHIParamRef InPrimitiveIdsBuffer,
-		int32 InBasePrimitiveIdsOffset,
-		bool bInDynamicInstancing,
-		int32 InTaskIndex,
-		int32 InTaskNum
-	)
-		: RHICmdList(InRHICmdList)
-		, VisibleMeshDrawCommands(InVisibleMeshDrawCommands)
-		, PrimitiveIdsBuffer(InPrimitiveIdsBuffer)
-		, BasePrimitiveIdsOffset(InBasePrimitiveIdsOffset)
-		, bDynamicInstancing(bInDynamicInstancing)
-		, TaskIndex(InTaskIndex)
-		, TaskNum(InTaskNum)
-	{}
-
-	FORCEINLINE TStatId GetStatId() const
-	{
-		RETURN_QUICK_DECLARE_CYCLE_STAT(FDrawVisibleMeshCommandsAnyThreadTask, STATGROUP_TaskGraphTasks);
-	}
-
-	static ESubsequentsMode::Type GetSubsequentsMode() { return ESubsequentsMode::TrackSubsequents; }
-
-	void DoTask(ENamedThreads::Type CurrentThread, const FGraphEventRef& MyCompletionGraphEvent)
-	{
-		checkSlow(RHICmdList.IsInsideRenderPass());
-
-		// Recompute draw range.
-		const int32 DrawNum = VisibleMeshDrawCommands.Num();
-		const int32 NumDrawsPerTask = FMath::DivideAndRoundUp(DrawNum, TaskNum);
-		const int32 StartIndex = TaskIndex * NumDrawsPerTask;
-		const int32 NumDraws = FMath::Min(NumDrawsPerTask, DrawNum - StartIndex);
-
-		SubmitMeshDrawCommandsRange(VisibleMeshDrawCommands, PrimitiveIdsBuffer, BasePrimitiveIdsOffset, bDynamicInstancing, StartIndex, NumDraws, RHICmdList);
-
-		RHICmdList.EndRenderPass();
-		RHICmdList.HandleRTThreadTaskCompletion(MyCompletionGraphEvent);
-	}
-};
-
-void FParallelMeshDrawCommandPass::DispatchDraw(FParallelCommandListSet* ParallelCommandListSet, FRHICommandList& RHICmdList) const
-{
-	if (MaxNumDraws <= 0)
-	{
-		return;
-	}
-
-
-	FVertexBufferRHIParamRef PrimitiveIdsBuffer = TaskContext.bUseGPUScene ? TaskContext.PrimitiveIdBuffer.VertexBuffer->VertexBufferRHI : nullptr;
-	const int32 BasePrimitiveIdsOffset = TaskContext.PrimitiveIdBuffer.VertexOffset;
-
-	if (ParallelCommandListSet)
-	{
-		const ENamedThreads::Type RenderThread = ENamedThreads::GetRenderThread();
-
-		FGraphEventArray Prereqs = TaskEventRefs;
-		if (ParallelCommandListSet->GetPrereqs())
-		{
-			Prereqs.Append(*ParallelCommandListSet->GetPrereqs());
-		}
-
-		// Distribute work evenly to the available task graph workers based on NumEstimatedDraws.
-		// Every task will then adjust it's working range based on FVisibleMeshDrawCommandProcessTask results.
-		const int32 NumThreads = FMath::Min<int32>(FTaskGraphInterface::Get().GetNumWorkerThreads(), ParallelCommandListSet->Width);
-		const int32 NumTasks = FMath::Min<int32>(NumThreads, FMath::DivideAndRoundUp(MaxNumDraws, ParallelCommandListSet->MinDrawsPerCommandList));
-		const int32 NumDrawsPerTask = FMath::DivideAndRoundUp(MaxNumDraws, NumTasks);
-
-		for (int32 TaskIndex = 0; TaskIndex < NumTasks; TaskIndex++)
-		{
-			const int32 StartIndex = TaskIndex * NumDrawsPerTask;
-			const int32 NumDraws = FMath::Min(NumDrawsPerTask, MaxNumDraws - StartIndex);
-			checkSlow(NumDraws > 0);
-
-			FRHICommandList* CmdList = ParallelCommandListSet->NewParallelCommandList();
-
-			FGraphEventRef AnyThreadCompletionEvent = TGraphTask<FDrawVisibleMeshCommandsAnyThreadTask>::CreateTask(&Prereqs, RenderThread)
-				.ConstructAndDispatchWhenReady(*CmdList, TaskContext.MeshDrawCommands, PrimitiveIdsBuffer, BasePrimitiveIdsOffset, TaskContext.bDynamicInstancing, TaskIndex, NumTasks);
-			ParallelCommandListSet->AddParallelCommandList(CmdList, AnyThreadCompletionEvent, NumDraws);
-		}
-	}
-	else
-	{
-		WaitForMeshPassSetupTask();
-		SubmitMeshDrawCommandsRange(TaskContext.MeshDrawCommands, PrimitiveIdsBuffer, BasePrimitiveIdsOffset, TaskContext.bDynamicInstancing, 0, TaskContext.MeshDrawCommands.Num(), RHICmdList);
-	}
-}
-
-void FParallelMeshDrawCommandPass::DumpInstancingStats() const
-{
-	if (!PassNameForStats.IsEmpty())
-	{
-		UE_LOG(LogRenderer, Log, TEXT("Instancing stats for %s"), *PassNameForStats);
-		UE_LOG(LogRenderer, Log, TEXT("   %i Mesh Draw Commands in %i instancing state buckets"), TaskContext.VisibleMeshDrawCommandsNum, TaskContext.NewPassVisibleMeshDrawCommandsNum);
-		UE_LOG(LogRenderer, Log, TEXT("   Largest %i"), TaskContext.MaxInstances);
-		UE_LOG(LogRenderer, Log, TEXT("   %.1f Dynamic Instancing draw call reduction factor"), TaskContext.VisibleMeshDrawCommandsNum / (float)TaskContext.NewPassVisibleMeshDrawCommandsNum);
-	}
-}
-
-void FParallelMeshDrawCommandPass::SetDumpInstancingStats(const FString& InPassNameForStats)
-{
-	PassNameForStats = InPassNameForStats;
-}
-
 void FSceneRenderer::SetupMeshPass(FViewInfo& View, FExclusiveDepthStencil::Type BasePassDepthStencilAccess, FViewCommands& ViewCommands)
 {
 	SCOPE_CYCLE_COUNTER(STAT_SetupMeshPass);
diff --git a/Engine/Source/Runtime/Renderer/Private/SceneRendering.h b/Engine/Source/Runtime/Renderer/Private/SceneRendering.h
--- a/Engine/Source/Runtime/Renderer/Private/SceneRendering.h
+++ b/Engine/Source/Runtime/Renderer/Private/SceneRendering.h
@@ -30,7 +30,7 @@
 #include "GlobalDistanceFieldParameters.h"
 #include "Templates/UniquePtr.h"
 #include "RenderGraph.h"
-#include "MeshPassProcessor.h"
+#include "MeshDrawCommands.h"
 
 // Forward declarations.
 class FScene;
@@ -152,20 +152,6 @@ public:
 	TArray<FProjectedShadowInfo*,SceneRenderingAllocator> OccludedPerObjectShadows;
 };
 
-// enum instead of bool to get better visibility when we pass around multiple bools, also allows for easier extensions
-namespace ETranslucencyPass
-{
-	enum Type
-	{
-		TPT_StandardTranslucency,
-		TPT_TranslucencyAfterDOF,
-
-		/** Drawing all translucency, regardless of separate or standard.  Used when drawing translucency outside of the main renderer, eg FRendererModule::DrawTile. */
-		TPT_AllTranslucency,
-		TPT_MAX
-	};
-};
-
 // Stores the primitive count of each translucency pass (redundant, could be computed after sorting but this way we touch less memory)
 struct FTranslucenyPrimCount
 {
@@ -894,130 +880,6 @@ public:
 
 typedef TArray<FViewCommands, TInlineAllocator<4>> FViewVisibleCommandsPerView;
 
-/**
- * Parallel mesh draw command pass setup task context.
- */
-class FMeshDrawCommandPassSetupTaskContext
-{
-public:
-	FMeshDrawCommandPassSetupTaskContext()
-		: View(nullptr)
-		, ShadingPath(EShadingPath::Num)
-		, PassType(EMeshPass::Num)
-		, bUseGPUScene(false)
-		, bDynamicInstancing(false)
-		, bReverseCulling(false)
-		, bRenderSceneTwoSided(false)
-		, BasePassDepthStencilAccess(FExclusiveDepthStencil::DepthNop_StencilNop)
-		, MeshPassProcessor(nullptr)
-		, MobileBasePassCSMMeshPassProcessor(nullptr)
-		, DynamicMeshElements(nullptr)
-		, NumDynamicMeshElements(0)
-		, NumDynamicMeshCommandBuildRequestElements(0)
-		, PrimitiveBounds(nullptr)
-		, VisibleMeshDrawCommandsNum(0)
-		, NewPassVisibleMeshDrawCommandsNum(0)
-		, MaxInstances(1)
-	{
-	}
-
-	const FViewInfo* View;
-	EShadingPath ShadingPath;
-	EMeshPass::Type PassType;
-	bool bUseGPUScene;
-	bool bDynamicInstancing;
-	bool bReverseCulling;
-	bool bRenderSceneTwoSided;
-	FExclusiveDepthStencil::Type BasePassDepthStencilAccess;
-
-	// Mesh pass processor.
-	FMeshPassProcessor* MeshPassProcessor;
-	FMeshPassProcessor* MobileBasePassCSMMeshPassProcessor;
-	const TArray<FMeshBatchAndRelevance, SceneRenderingAllocator>* DynamicMeshElements;
-
-	// Commands.
-	int32 NumDynamicMeshElements;
-	int32 NumDynamicMeshCommandBuildRequestElements;
-	FMeshCommandOneFrameArray MeshDrawCommands;
-	FMeshCommandOneFrameArray MobileBasePassCSMMeshDrawCommands;
-	TArray<const FStaticMesh*, SceneRenderingAllocator> DynamicMeshCommandBuildRequests;
-	TArray<const FStaticMesh*, SceneRenderingAllocator> MobileBasePassCSMDynamicMeshCommandBuildRequests;
-	FDynamicMeshDrawCommandStorage MeshDrawCommandStorage;
-
-	// Preallocated resources.
-	FGlobalDynamicVertexBuffer::FAllocation PrimitiveIdBuffer;
-	FMeshCommandOneFrameArray TempVisibleMeshDrawCommands;
-
-	// For UpdateTranslucentMeshSortKeys.
-	ETranslucencyPass::Type TranslucencyPass;
-	ETranslucentSortPolicy::Type TranslucentSortPolicy;
-	FVector TranslucentSortAxis;
-	FVector ViewOrigin;
-	FMatrix ViewMatrix;
-	const TArray<struct FPrimitiveBounds>* PrimitiveBounds;
-
-	// For logging instancing stats.
-	int32 VisibleMeshDrawCommandsNum;
-	int32 NewPassVisibleMeshDrawCommandsNum;
-	int32 MaxInstances;
-};
-
-/**
- * Parallel mesh draw command processing and rendering. 
- * Encapsulates two parallel tasks - mesh command setup task and drawing task.
- */
-class FParallelMeshDrawCommandPass
-{
-public:
-	FParallelMeshDrawCommandPass()
-		: MaxNumDraws(0)
-	{
-	}
-
-	~FParallelMeshDrawCommandPass();
-
-	/**
-	 * Dispatch visible mesh draw command process task, which prepares this pass for drawing.
-	 * This includes generation of dynamic mesh draw commands, draw sorting and draw merging.
-	 */
-	void DispatchPassSetup(
-		FScene* Scene,
-		const FViewInfo& View, 
-		EMeshPass::Type PassType, 
-		FExclusiveDepthStencil::Type BasePassDepthStencilAccess,
-		FMeshPassProcessor* MeshPassProcessor,
-		const TArray<FMeshBatchAndRelevance, SceneRenderingAllocator>& DynamicMeshElements,
-		int32 NumDynamicMeshElements,
-		TArray<const FStaticMesh*, SceneRenderingAllocator>& InOutDynamicMeshCommandBuildRequests,
-		int32 NumDynamicMeshCommandBuildRequestElements,
-		FMeshCommandOneFrameArray& InOutMeshDrawCommands,
-		FMeshPassProcessor* MobileBasePassCSMMeshPassProcessor = nullptr, // Required only for the mobile base pass.
-		FMeshCommandOneFrameArray* InOutMobileBasePassCSMMeshDrawCommands = nullptr // Required only for the mobile base pass.
-	);
-
-	/**
-	 * Dispatch visible mesh draw command draw task.
-	 */
-	void DispatchDraw(FParallelCommandListSet* ParallelCommandListSet, FRHICommandList& RHICmdList) const;
-
-	void Empty();
-	void SetDumpInstancingStats(const FString& InPassName);
-	bool HasAnyDraw() const { return MaxNumDraws > 0; }
-
-
-private:
-	FMeshDrawCommandPassSetupTaskContext TaskContext;
-	FGraphEventArray TaskEventRefs;
-	FString PassNameForStats;
-
-	// Maximum number of draws for this pass. Used to prealocate resources on rendering thread. 
-	// Has a guarantee that if there won't be any draws, then MaxNumDraws = 0;
-	int32 MaxNumDraws;
-
-	void DumpInstancingStats() const;
-	void WaitForMeshPassSetupTask() const;
-};
-
 /** A FSceneView with additional state used by the scene renderer. */
 class FViewInfo : public FSceneView
 {
diff --git a/Engine/Source/Runtime/Renderer/Private/TranslucencyPass.h b/Engine/Source/Runtime/Renderer/Private/TranslucencyPass.h
new file mode 100644
--- /dev/null
+++ b/Engine/Source/Runtime/Renderer/Private/TranslucencyPass.h
@@ -0,0 +1,21 @@
+// Copyright 1998-2019 Epic Games, Inc. All Rights Reserved.
+
+/*=============================================================================
+TranslucencyPass.h: Translucency pass enum.
+=============================================================================*/
+
+#pragma once
+
+// enum instead of bool to get better visibility when we pass around multiple bools, also allows for easier extensions
+namespace ETranslucencyPass
+{
+	enum Type
+	{
+		TPT_StandardTranslucency,
+		TPT_TranslucencyAfterDOF,
+
+		/** Drawing all translucency, regardless of separate or standard.  Used when drawing translucency outside of the main renderer, eg FRendererModule::DrawTile. */
+		TPT_AllTranslucency,
+		TPT_MAX
+	};
+};
\ No newline at end of file
