Fixed calculation of number of dynamic mesh elements for shadows.

diff --git a/Engine/Source/Runtime/Renderer/Private/ShadowRendering.h b/Engine/Source/Runtime/Renderer/Private/ShadowRendering.h
--- a/Engine/Source/Runtime/Renderer/Private/ShadowRendering.h
+++ b/Engine/Source/Runtime/Renderer/Private/ShadowRendering.h
@@ -778,7 +778,8 @@ private:
 		const FScene* Scene,
 		EMeshPass::Type PassType,
 		FMeshCommandOneFrameArray& VisibleMeshCommands,
-		TArray<const FStaticMesh*, SceneRenderingAllocator>& MeshCommandBuildRequests);
+		TArray<const FStaticMesh*, SceneRenderingAllocator>& MeshCommandBuildRequests,
+		int32& NumMeshCommandBuildRequestElements);
 
 	/** Will return if we should draw the static mesh for the shadow, and will perform lazy init of primitive if it wasn't visible */
 	bool ShouldDrawStaticMeshes(FViewInfo& InCurrentView, bool bInCustomDataRelevance, FPrimitiveSceneInfo* InPrimitiveSceneInfo);
diff --git a/Engine/Source/Runtime/Renderer/Private/ShadowSetup.cpp b/Engine/Source/Runtime/Renderer/Private/ShadowSetup.cpp
--- a/Engine/Source/Runtime/Renderer/Private/ShadowSetup.cpp
+++ b/Engine/Source/Runtime/Renderer/Private/ShadowSetup.cpp
@@ -806,7 +806,8 @@ void FProjectedShadowInfo::AddCachedMeshDrawCommandsForPass(
 	const FScene* Scene,
 	EMeshPass::Type PassType,
 	FMeshCommandOneFrameArray& VisibleMeshCommands,
-	TArray<const FStaticMesh*, SceneRenderingAllocator>& MeshCommandBuildRequests)
+	TArray<const FStaticMesh*, SceneRenderingAllocator>& MeshCommandBuildRequests,
+	int32& NumMeshCommandBuildRequestElements)
 {
 	if (UseMeshDrawCommandPipeline())
 	{
@@ -838,6 +839,7 @@ void FProjectedShadowInfo::AddCachedMeshDrawCommandsForPass(
 		}
 		else
 		{
+			NumMeshCommandBuildRequestElements += StaticMeshRelevance.NumElements;
 			MeshCommandBuildRequests.Add(&StaticMesh);
 		}
 	}
@@ -933,7 +935,8 @@ bool FProjectedShadowInfo::ShouldDrawStaticMeshes(FViewInfo& InCurrentView, bool
 								InPrimitiveSceneInfo->Scene,
 								EMeshPass::CSMShadowDepth,
 								ShadowDepthPassVisibleCommands,
-								SubjectMeshCommandBuildRequests);
+								SubjectMeshCommandBuildRequests,
+								NumSubjectMeshCommandBuildRequestElements);
 						}
 						else
 						{
