Fixed DrawTileMesh crashes due to BasePassUniformBuffer being destroyed before drawing.

diff --git a/Engine/Source/Runtime/Renderer/Private/Renderer.cpp b/Engine/Source/Runtime/Renderer/Private/Renderer.cpp
--- a/Engine/Source/Runtime/Renderer/Private/Renderer.cpp
+++ b/Engine/Source/Runtime/Renderer/Private/Renderer.cpp
@@ -174,6 +174,8 @@ void FRendererModule::DrawTileMesh(FRHICommandListImmediate& RHICmdList, FDrawin
 		}
 		else if (IsTranslucentBlendMode(MaterialBlendMode))
 		{
+			FUniformBufferRHIRef BasePassUniformBuffer;
+
 			if (ShadingPath == EShadingPath::Deferred)
 			{
 				// Crash fix - reflection capture shader parameter is bound but we have no buffer during Build Texture Streaming
@@ -182,16 +184,18 @@ void FRendererModule::DrawTileMesh(FRHICommandListImmediate& RHICmdList, FDrawin
 					View.ReflectionCaptureUniformBuffer = GDummyReflectionCaptureUniformBuffer.GetUniformBufferRef();
 				}
 			
-				TUniformBufferRef<FTranslucentBasePassUniformParameters> BasePassUniformBuffer;
-				CreateTranslucentBasePassUniformBuffer(RHICmdList, View, nullptr, ESceneTextureSetupMode::None, BasePassUniformBuffer, 0);
-				DrawRenderState.SetPassUniformBuffer(BasePassUniformBuffer);
+				TUniformBufferRef<FTranslucentBasePassUniformParameters> TranslucentBasePassUniformBuffer;
+				CreateTranslucentBasePassUniformBuffer(RHICmdList, View, nullptr, ESceneTextureSetupMode::None, TranslucentBasePassUniformBuffer, 0);
+				BasePassUniformBuffer = TranslucentBasePassUniformBuffer;
 			}
 			else // Mobile
 			{
-				TUniformBufferRef<FMobileBasePassUniformParameters> BasePassUniformBuffer;
-				CreateMobileBasePassUniformBuffer(RHICmdList, View, true, BasePassUniformBuffer);
-				DrawRenderState.SetPassUniformBuffer(BasePassUniformBuffer);
+				TUniformBufferRef<FMobileBasePassUniformParameters> MobileBasePassUniformBuffer;
+				CreateMobileBasePassUniformBuffer(RHICmdList, View, true, MobileBasePassUniformBuffer);
+				BasePassUniformBuffer = MobileBasePassUniformBuffer;
 			}
+
+			DrawRenderState.SetPassUniformBuffer(BasePassUniformBuffer);
 			
 			DrawDynamicMeshPass(View, RHICmdList,
 				[&View, &DrawRenderState, &Mesh](FDynamicPassMeshDrawListContext* DynamicMeshPassContext)
@@ -237,19 +241,23 @@ void FRendererModule::DrawTileMesh(FRHICommandListImmediate& RHICmdList, FDrawin
 			}
 			else
 			{
+				FUniformBufferRHIRef BasePassUniformBuffer;
+
 				if (ShadingPath == EShadingPath::Deferred)
 				{
-					TUniformBufferRef<FOpaqueBasePassUniformParameters> BasePassUniformBuffer;
-					CreateOpaqueBasePassUniformBuffer(RHICmdList, View, nullptr, BasePassUniformBuffer);
-					DrawRenderState.SetPassUniformBuffer(BasePassUniformBuffer);
+					TUniformBufferRef<FOpaqueBasePassUniformParameters> OpaqueBasePassUniformBuffer;
+					CreateOpaqueBasePassUniformBuffer(RHICmdList, View, nullptr, OpaqueBasePassUniformBuffer);
+					BasePassUniformBuffer = OpaqueBasePassUniformBuffer;
 				}
 				else // Mobile
 				{
-					TUniformBufferRef<FMobileBasePassUniformParameters> BasePassUniformBuffer;
-					CreateMobileBasePassUniformBuffer(RHICmdList, View, false, BasePassUniformBuffer);
-					DrawRenderState.SetPassUniformBuffer(BasePassUniformBuffer);
+					TUniformBufferRef<FMobileBasePassUniformParameters> MobileBasePassUniformBuffer;
+					CreateMobileBasePassUniformBuffer(RHICmdList, View, false, MobileBasePassUniformBuffer);
+					BasePassUniformBuffer = MobileBasePassUniformBuffer;
 				}
 
+				DrawRenderState.SetPassUniformBuffer(BasePassUniformBuffer);
+
 				DrawDynamicMeshPass(View, RHICmdList,
 					[&View, &DrawRenderState, &Mesh](FDynamicPassMeshDrawListContext* DynamicMeshPassContext)
 					{
