Restored r.EmitMeshDrawEvents and mesh draw call stat.

diff --git a/Engine/Source/Runtime/Renderer/Private/DrawingPolicy.cpp b/Engine/Source/Runtime/Renderer/Private/DrawingPolicy.cpp
--- a/Engine/Source/Runtime/Renderer/Private/DrawingPolicy.cpp
+++ b/Engine/Source/Runtime/Renderer/Private/DrawingPolicy.cpp
@@ -13,15 +13,6 @@
 #include "SceneCore.h"
 #include "ScenePrivate.h"
 
-int32 GEmitMeshDrawEvent = 0;
-static FAutoConsoleVariableRef CVarEmitMeshDrawEvent(
-	TEXT("r.EmitMeshDrawEvents"),
-	GEmitMeshDrawEvent,
-	TEXT("Emits a GPU event around each drawing policy draw call.  /n")
-	TEXT("Useful for seeing stats about each draw call, however it greatly distorts total time and time per draw call."),
-	ECVF_RenderThreadSafe
-	);
-
 FMeshDrawingPolicy::FMeshDrawingPolicy(
 	const FVertexFactory* InVertexFactory,
 	const FMaterialRenderProxy* InMaterialRenderProxy,
@@ -116,8 +107,6 @@ void FMeshDrawingPolicy::SetPrimitiveIdStream(
 void FMeshDrawingPolicy::DrawMesh(FRHICommandList& RHICmdList, const FSceneView& View, const FMeshBatch& Mesh, int32 BatchElementIndex, const bool bIsInstancedStereo) const
 {
 	DEFINE_LOG_CATEGORY_STATIC(LogFMeshDrawingPolicyDrawMesh, Warning, All);
-	INC_DWORD_STAT(STAT_MeshDrawCalls);
-	SCOPED_CONDITIONAL_DRAW_EVENTF(RHICmdList, MeshEvent, GEmitMeshDrawEvent != 0, TEXT("Mesh Draw"));
 
 	const FMeshBatchElement& BatchElement = Mesh.Elements[BatchElementIndex];
 
diff --git a/Engine/Source/Runtime/Renderer/Private/MeshPassProcessor.cpp b/Engine/Source/Runtime/Renderer/Private/MeshPassProcessor.cpp
--- a/Engine/Source/Runtime/Renderer/Private/MeshPassProcessor.cpp
+++ b/Engine/Source/Runtime/Renderer/Private/MeshPassProcessor.cpp
@@ -18,6 +18,15 @@ FCriticalSection GlobalTableCriticalSection;
 
 const FMeshDrawCommandSortKey FMeshDrawCommandSortKey::Default = { 0 };
 
+int32 GEmitMeshDrawEvent = 0;
+static FAutoConsoleVariableRef CVarEmitMeshDrawEvent(
+	TEXT("r.EmitMeshDrawEvents"),
+	GEmitMeshDrawEvent,
+	TEXT("Emits a GPU event around each drawing policy draw call.  /n")
+	TEXT("Useful for seeing stats about each draw call, however it greatly distorts total time and time per draw call."),
+	ECVF_RenderThreadSafe
+);
+
 enum { MAX_SRVs_PER_SHADER_STAGE = 128 };
 enum { MAX_UNIFORM_BUFFERS_PER_SHADER_STAGE = 14 };
 enum { MAX_SAMPLERS_PER_SHADER_STAGE = 32 };
@@ -816,6 +825,9 @@ void SubmitMeshDrawCommandsRange(
 
 	for (int32 DrawCommandIndex = StartIndex; DrawCommandIndex < StartIndex + NumMeshDrawCommands; DrawCommandIndex++)
 	{
+		INC_DWORD_STAT(STAT_MeshDrawCalls);
+		SCOPED_CONDITIONAL_DRAW_EVENTF(RHICmdList, MeshEvent, GEmitMeshDrawEvent != 0, TEXT("Mesh Draw"));
+
 		const FVisibleMeshDrawCommand& VisibleMeshDrawCommand = VisibleMeshDrawCommands[DrawCommandIndex];
 		const int32 PrimitiveIdBufferOffset = BasePrimitiveIdsOffset + (bDynamicInstancing ? VisibleMeshDrawCommand.PrimitiveIdBufferOffset : DrawCommandIndex) * sizeof(int32);
 		checkSlow(!bDynamicInstancing || VisibleMeshDrawCommand.PrimitiveIdBufferOffset >= 0);
