DR - Regain some performance in the raytracing code that was lost after the merge, due to touching the SceneInfo with data that is not used by the Raytracing.

diff --git a/Engine/Source/Runtime/Renderer/Private/DeferredShadingRenderer.cpp b/Engine/Source/Runtime/Renderer/Private/DeferredShadingRenderer.cpp
--- a/Engine/Source/Runtime/Renderer/Private/DeferredShadingRenderer.cpp
+++ b/Engine/Source/Runtime/Renderer/Private/DeferredShadingRenderer.cpp
@@ -674,29 +674,25 @@ bool FDeferredShadingSceneRenderer::GatherRayTracingWorldInstances(FRHICommandLi
 						const auto& StaticMeshMdcIndices = SceneProxy->RayTracingLodIndexToMeshDrawCommandIndicies[LODToRender.GetRayTracedLOD()];
 						for (int StaticMeshIndex = 0; StaticMeshIndex < StaticMeshMdcIndices.Num(); StaticMeshIndex++)
 						{
-							const FPrimitiveSceneInfo* PrimitiveInfo = Scene->Primitives[PrimitiveIndex];
-							const FStaticMeshRelevance& StaticMeshRelevance = PrimitiveInfo->StaticMeshRelevances[StaticMeshMdcIndices[StaticMeshIndex].StaticMeshIndex];
-
 							const int32 CommandIndex = StaticMeshMdcIndices[StaticMeshIndex].CommandIndex;
 							if (CommandIndex >= 0)
 							{
 								const FCachedPassMeshDrawList& SceneDrawList = Scene->CachedDrawLists[EMeshPass::RayTracing];
-								const int32 RayTracingStaticMeshCommandInfoIndex = StaticMeshRelevance.GetStaticMeshCommandInfoIndex(EMeshPass::RayTracing);
-								const FCachedMeshDrawCommandInfo& CachedMeshDrawCommand = Scene->Primitives[PrimitiveIndex]->StaticMeshCommandInfos[RayTracingStaticMeshCommandInfoIndex];
 								FVisibleMeshDrawCommand NewVisibleMeshDrawCommand;
 
 								NewVisibleMeshDrawCommand.Setup(
 									&SceneDrawList.MeshDrawCommands[CommandIndex],
 									PrimitiveIndex,
-									CachedMeshDrawCommand.StateBucketId,
-									CachedMeshDrawCommand.MeshFillMode,
-									CachedMeshDrawCommand.MeshCullMode,
+									0,
+									FM_Solid,
+									CM_CW,
 									FMeshDrawCommandSortKey::Default);
 
 								View.RaytraycingVisibleMeshDrawCommands.Add(NewVisibleMeshDrawCommand);
 							}
 							else
 							{
+								const FPrimitiveSceneInfo* PrimitiveInfo = Scene->Primitives[PrimitiveIndex];
 								const FStaticMesh& StaticMesh = PrimitiveInfo->StaticMeshes[StaticMeshMdcIndices[StaticMeshIndex].StaticMeshIndex];
 
 								uint64 BatchVisibility = 1;
