Remove all support for Metal prior to Metal 1.2 on iOS/tvOS and Metal 2.0 on macOS. This permits some fairly radical simplifications and removal of a number of workarounds.

diff --git a/Engine/Source/Runtime/Engine/Public/SceneManagement.h b/Engine/Source/Runtime/Engine/Public/SceneManagement.h
--- a/Engine/Source/Runtime/Engine/Public/SceneManagement.h
+++ b/Engine/Source/Runtime/Engine/Public/SceneManagement.h
@@ -981,7 +981,7 @@ inline bool DoesPlatformSupportDistanceFieldShadowing(EShaderPlatform Platform)
 {
 	// Hasn't been tested elsewhere yet
 	return Platform == SP_PCD3D_SM5 || Platform == SP_PS4
-		|| (IsMetalPlatform(Platform) && GetMaxSupportedFeatureLevel(Platform) >= ERHIFeatureLevel::SM5 && RHIGetShaderLanguageVersion(Platform) >= 2)
+		|| IsMetalSM5Platform(Platform)
 		|| Platform == SP_XBOXONE_D3D12
 		|| IsVulkanSM5Platform(Platform);
 }
diff --git a/Engine/Source/Runtime/IOS/IOSRuntimeSettings/Classes/IOSRuntimeSettings.h b/Engine/Source/Runtime/IOS/IOSRuntimeSettings/Classes/IOSRuntimeSettings.h
--- a/Engine/Source/Runtime/IOS/IOSRuntimeSettings/Classes/IOSRuntimeSettings.h
+++ b/Engine/Source/Runtime/IOS/IOSRuntimeSettings/Classes/IOSRuntimeSettings.h
@@ -56,16 +56,10 @@ UENUM()
 UENUM()
 enum class EIOSMetalShaderStandard : uint8
 {
-    /** Metal Shaders Compatible With iOS 8.0/tvOS 9.0 or later (std=ios-metal1.0) */
-    IOSMetalSLStandard_1_0 = 0 UMETA(DisplayName="Metal v1.0 (iOS 8.0/tvOS 9.0)"),
-    
-    /** Metal Shaders Compatible With iOS 9.0/tvOS 9.0 or later (std=ios-metal1.1) */
-    IOSMetalSLStandard_1_1 = 1 UMETA(DisplayName="Metal v1.1 (iOS 9.0/tvOS 9.0)"),
-    
-    /** Metal Shaders Compatible With iOS 10.0/tvOS 10.0 or later (std=ios-metal1.2) */
+	/** Metal Shaders Compatible With iOS 10.0/tvOS 10.0 or later (std=ios-metal1.2) */
 	IOSMetalSLStandard_1_2 = 2 UMETA(DisplayName="Metal v1.2 (iOS 10.0/tvOS 10.0)"),
 	
-	/** Metal Shaders Compatible With iOS 11.0/tvOS 11.0 or later (std=ios-metal2.0) */
+    /** Metal Shaders Compatible With iOS 11.0/tvOS 11.0 or later (std=ios-metal2.0) */
 	IOSMetalSLStandard_2_0 = 3 UMETA(DisplayName="Metal v2.0 (iOS 11.0/tvOS 11.0)"),
     
     /** Metal Shaders Compatible With iOS 12.0/tvOS 12.0 or later (std=ios-metal2.1) */
diff --git a/Engine/Source/Runtime/RHI/Private/RHI.cpp b/Engine/Source/Runtime/RHI/Private/RHI.cpp
--- a/Engine/Source/Runtime/RHI/Private/RHI.cpp
+++ b/Engine/Source/Runtime/RHI/Private/RHI.cpp
@@ -777,8 +777,8 @@ RHI_API uint32 RHIGetShaderLanguageVersion(const EShaderPlatform Platform)
 			static int32 MaxShaderVersion = -1;
 			if (MaxShaderVersion < 0)
 			{
-				MaxShaderVersion = 0;
-				int32 MinShaderVersion = 0;
+				MaxShaderVersion = 2;
+				int32 MinShaderVersion = 2;
 				if(!GConfig->GetInt(TEXT("/Script/IOSRuntimeSettings.IOSRuntimeSettings"), TEXT("MaxShaderLanguageVersion"), MaxShaderVersion, GEngineIni))
 				{
 					MaxShaderVersion = 0;
@@ -795,27 +795,17 @@ RHI_API bool RHISupportsTessellation(const EShaderPlatform Platform)
 {
 	if (IsFeatureLevelSupported(Platform, ERHIFeatureLevel::SM5) && !IsMetalPlatform(Platform))
 	{
-		return (Platform == SP_PCD3D_SM5) || (Platform == SP_XBOXONE_D3D12) || (Platform == SP_OPENGL_SM5) || (Platform == SP_OPENGL_ES31_EXT)/* || (IsVulkanSM5Platform(Platform)*/;
-	}
-	// For Metal we can only support tessellation if we are willing to sacrifice backward compatibility with OS versions.
-	// As such it becomes an opt-in project setting.
-	else if (Platform == SP_METAL_SM5)
-	{
-		return (RHIGetShaderLanguageVersion(Platform) >= 2);
+		return (Platform == SP_PCD3D_SM5) || (Platform == SP_XBOXONE_D3D12) || (Platform == SP_OPENGL_SM5) || (Platform == SP_OPENGL_ES31_EXT) || (Platform == SP_METAL_SM5) /* || (IsVulkanSM5Platform(Platform)*/;
 	}
 	return false;
 }
 
 RHI_API bool RHISupportsPixelShaderUAVs(const EShaderPlatform Platform)
 {
-	if (IsFeatureLevelSupported(Platform, ERHIFeatureLevel::SM5) && !IsMetalPlatform(Platform))
+	if (IsFeatureLevelSupported(Platform, ERHIFeatureLevel::SM5))
 	{
 		return true;
 	}
-	else if (IsMetalSM5Platform(Platform))
-	{
-		return (RHIGetShaderLanguageVersion(Platform) >= 2);
-	}
 	return false;
 }
 
diff --git a/Engine/Source/Runtime/RHI/Public/RHI.h b/Engine/Source/Runtime/RHI/Public/RHI.h
--- a/Engine/Source/Runtime/RHI/Public/RHI.h
+++ b/Engine/Source/Runtime/RHI/Public/RHI.h
@@ -133,7 +133,7 @@ inline bool RHISupportsInstancedStereo(const EShaderPlatform Platform)
 inline bool RHISupportsMultiView(const EShaderPlatform Platform)
 {
 	// Only PS4 and Metal SM5 from 10.13 onward supports Multi-View
-	return (Platform == EShaderPlatform::SP_PS4) || ((Platform == EShaderPlatform::SP_METAL_SM5 || Platform == SP_METAL_SM5_NOTESS) && RHIGetShaderLanguageVersion(Platform) >= 3);
+	return (Platform == EShaderPlatform::SP_PS4) || ((Platform == EShaderPlatform::SP_METAL_SM5 || Platform == SP_METAL_SM5_NOTESS));
 }
 
 inline bool RHISupportsMSAA(EShaderPlatform Platform)
@@ -153,12 +153,7 @@ inline bool RHISupportsMSAA(EShaderPlatform Platform)
 
 inline bool RHISupportsBufferLoadTypeConversion(EShaderPlatform Platform)
 {
-#if PLATFORM_MAC || PLATFORM_IOS
-	// Fixed from Metal v.2.0 onward.
-	return !IsMetalPlatform(Platform) || RHIGetShaderLanguageVersion(Platform) >= 3;
-#else
 	return true;
-#endif
 }
 
 /** Whether the platform supports reading from volume textures (does not cover rendering to volume textures). */
@@ -183,7 +178,7 @@ inline bool RHISupports4ComponentUAVReadWrite(EShaderPlatform Platform)
 	Shader Platform must not use the mobile renderer, and for Metal, the shader language must be at least 2. */
 inline bool RHISupportsManualVertexFetch(EShaderPlatform InShaderPlatform)
 {
-	return (!IsOpenGLPlatform(InShaderPlatform) || IsSwitchPlatform(InShaderPlatform)) && !IsMobilePlatform(InShaderPlatform) && (!IsMetalPlatform(InShaderPlatform) || RHIGetShaderLanguageVersion(InShaderPlatform) >= 2);
+	return (!IsOpenGLPlatform(InShaderPlatform) || IsSwitchPlatform(InShaderPlatform)) && !IsMobilePlatform(InShaderPlatform);
 }
 
 /** 
diff --git a/Engine/Source/Runtime/Renderer/Private/DistanceFieldAmbientOcclusion.h b/Engine/Source/Runtime/Renderer/Private/DistanceFieldAmbientOcclusion.h
--- a/Engine/Source/Runtime/Renderer/Private/DistanceFieldAmbientOcclusion.h
+++ b/Engine/Source/Runtime/Renderer/Private/DistanceFieldAmbientOcclusion.h
@@ -29,7 +29,7 @@ extern const uint32 UpdateObjectsGroupSize;
 inline bool DoesPlatformSupportDistanceFieldAO(EShaderPlatform Platform)
 {
 	return Platform == SP_PCD3D_SM5 || Platform == SP_PS4 || Platform == SP_XBOXONE_D3D12
-		|| (IsMetalPlatform(Platform) && GetMaxSupportedFeatureLevel(Platform) >= ERHIFeatureLevel::SM5 && RHIGetShaderLanguageVersion(Platform) >= 2)
+		|| IsMetalSM5Platform(Platform)
 		|| IsVulkanSM5Platform(Platform);
 }
 
diff --git a/Engine/Source/Runtime/Renderer/Private/GPUFastFourierTransform.cpp b/Engine/Source/Runtime/Renderer/Private/GPUFastFourierTransform.cpp
--- a/Engine/Source/Runtime/Renderer/Private/GPUFastFourierTransform.cpp
+++ b/Engine/Source/Runtime/Renderer/Private/GPUFastFourierTransform.cpp
@@ -200,7 +200,7 @@ namespace GPUFFT
 		static bool ShouldCompilePermutation(const FGlobalShaderPermutationParameters& Parameters)
 		{
 			// @todo MetalMRT: Metal MRT can't cope with the threadgroup storage requirements for these shaders right now
-			return IsFeatureLevelSupported(Parameters.Platform, ERHIFeatureLevel::SM5) && (!IsMetalPlatform(Parameters.Platform) || RHIGetShaderLanguageVersion(Parameters.Platform) >= 2) && !IsMetalMRTPlatform(Parameters.Platform);
+			return IsFeatureLevelSupported(Parameters.Platform, ERHIFeatureLevel::SM5) && !IsMetalMRTPlatform(Parameters.Platform);
 		}
 
 		static void ModifyCompilationEnvironment(const FGlobalShaderPermutationParameters& Parameters, FShaderCompilerEnvironment& OutEnvironment)
@@ -306,7 +306,7 @@ namespace GPUFFT
 		static bool ShouldCompilePermutation(const FGlobalShaderPermutationParameters& Parameters)
 		{
 			// @todo MetalMRT: Metal MRT can't cope with the threadgroup storage requirements for these shaders right now
-			return IsFeatureLevelSupported(Parameters.Platform, ERHIFeatureLevel::SM5) && (!IsMetalPlatform(Parameters.Platform) || RHIGetShaderLanguageVersion(Parameters.Platform) >= 2) && !IsMetalMRTPlatform(Parameters.Platform);
+			return IsFeatureLevelSupported(Parameters.Platform, ERHIFeatureLevel::SM5) && !IsMetalMRTPlatform(Parameters.Platform);
 		}
 
 		static void ModifyCompilationEnvironment(const FGlobalShaderPermutationParameters& Parameters, FShaderCompilerEnvironment& OutEnvironment)
@@ -408,7 +408,7 @@ namespace GPUFFT
 		static bool ShouldCompilePermutation(const FGlobalShaderPermutationParameters& Parameters)
 		{
 			// @todo MetalMRT: Metal MRT can't cope with the threadgroup storage requirements for these shaders right now
-			return IsFeatureLevelSupported(Parameters.Platform, ERHIFeatureLevel::SM5) && (!IsMetalPlatform(Parameters.Platform) || RHIGetShaderLanguageVersion(Parameters.Platform) >= 2) && !IsMetalMRTPlatform(Parameters.Platform);
+			return IsFeatureLevelSupported(Parameters.Platform, ERHIFeatureLevel::SM5) && !IsMetalMRTPlatform(Parameters.Platform);
 		}
 
 		static void ModifyCompilationEnvironment(const FGlobalShaderPermutationParameters& Parameters, FShaderCompilerEnvironment& OutEnvironment)
@@ -509,7 +509,7 @@ namespace GPUFFT
 		static bool ShouldCompilePermutation(const FGlobalShaderPermutationParameters& Parameters)
 		{
 			// @todo MetalMRT: Metal MRT can't cope with the threadgroup storage requirements for these shaders right now
-			return IsFeatureLevelSupported(Parameters.Platform, ERHIFeatureLevel::SM5) && (!IsMetalPlatform(Parameters.Platform) || RHIGetShaderLanguageVersion(Parameters.Platform) >= 2) && !IsMetalMRTPlatform(Parameters.Platform);
+			return IsFeatureLevelSupported(Parameters.Platform, ERHIFeatureLevel::SM5) && !IsMetalMRTPlatform(Parameters.Platform);
 		}
 
 		static void ModifyCompilationEnvironment(const FGlobalShaderPermutationParameters& Parameters, FShaderCompilerEnvironment& OutEnvironment)
@@ -595,7 +595,7 @@ namespace GPUFFT
 		static bool ShouldCompilePermutation(const FGlobalShaderPermutationParameters& Parameters)
 		{
 			// @todo MetalMRT: Metal MRT can't cope with the threadgroup storage requirements for these shaders right now
-			return IsFeatureLevelSupported(Parameters.Platform, ERHIFeatureLevel::SM5) && (!IsMetalPlatform(Parameters.Platform) || RHIGetShaderLanguageVersion(Parameters.Platform) >= 2) && !IsMetalMRTPlatform(Parameters.Platform);
+			return IsFeatureLevelSupported(Parameters.Platform, ERHIFeatureLevel::SM5) && !IsMetalMRTPlatform(Parameters.Platform);
 		}
 
 		static void ModifyCompilationEnvironment(const FGlobalShaderPermutationParameters& Parameters, FShaderCompilerEnvironment& OutEnvironment)
@@ -680,7 +680,7 @@ namespace GPUFFT
 		static bool ShouldCompilePermutation(const FGlobalShaderPermutationParameters& Parameters)
 		{
 			// @todo MetalMRT: Metal MRT can't cope with the threadgroup storage requirements for these shaders right now
-			return IsFeatureLevelSupported(Parameters.Platform, ERHIFeatureLevel::SM5) && (!IsMetalPlatform(Parameters.Platform) || RHIGetShaderLanguageVersion(Parameters.Platform) >= 2) && !IsMetalMRTPlatform(Parameters.Platform);
+			return IsFeatureLevelSupported(Parameters.Platform, ERHIFeatureLevel::SM5) && !IsMetalMRTPlatform(Parameters.Platform);
 		}
 
 		static void ModifyCompilationEnvironment(const FGlobalShaderPermutationParameters& Parameters, FShaderCompilerEnvironment& OutEnvironment)
@@ -839,7 +839,7 @@ namespace GPUFFT
 		static bool ShouldCompilePermutation(const FGlobalShaderPermutationParameters& Parameters)
 		{
 			// @todo MetalMRT: Metal MRT can't cope with the threadgroup storage requirements for these shaders right now
-			return IsFeatureLevelSupported(Parameters.Platform, ERHIFeatureLevel::SM5) && (!IsMetalPlatform(Parameters.Platform) || RHIGetShaderLanguageVersion(Parameters.Platform) >= 2) && !IsMetalMRTPlatform(Parameters.Platform);
+			return IsFeatureLevelSupported(Parameters.Platform, ERHIFeatureLevel::SM5) && !IsMetalMRTPlatform(Parameters.Platform);
 		}
 
 
@@ -876,7 +876,7 @@ namespace GPUFFT
 		static bool ShouldCompilePermutation(const FGlobalShaderPermutationParameters& Parameters)
 		{
 			// @todo MetalMRT: Metal MRT can't cope with the threadgroup storage requirements for these shaders right now
-			return IsFeatureLevelSupported(Parameters.Platform, ERHIFeatureLevel::SM5) && (!IsMetalPlatform(Parameters.Platform) || RHIGetShaderLanguageVersion(Parameters.Platform) >= 2) && !IsMetalMRTPlatform(Parameters.Platform);
+			return IsFeatureLevelSupported(Parameters.Platform, ERHIFeatureLevel::SM5) && !IsMetalMRTPlatform(Parameters.Platform);
 		}
 
 
@@ -1049,7 +1049,7 @@ namespace GPUFFT
 		static bool ShouldCompilePermutation(const FGlobalShaderPermutationParameters& Parameters)
 		{
 			// @todo MetalMRT: Metal MRT can't cope with the threadgroup storage requirements for these shaders right now
-			return IsFeatureLevelSupported(Parameters.Platform, ERHIFeatureLevel::SM5) && (!IsMetalPlatform(Parameters.Platform) || RHIGetShaderLanguageVersion(Parameters.Platform) >= 2) && !IsMetalMRTPlatform(Parameters.Platform);
+			return IsFeatureLevelSupported(Parameters.Platform, ERHIFeatureLevel::SM5) && !IsMetalMRTPlatform(Parameters.Platform);
 		}
 
 
diff --git a/Engine/Source/Runtime/Renderer/Private/PostProcess/PostProcessFFTBloom.cpp b/Engine/Source/Runtime/Renderer/Private/PostProcess/PostProcessFFTBloom.cpp
--- a/Engine/Source/Runtime/Renderer/Private/PostProcess/PostProcessFFTBloom.cpp
+++ b/Engine/Source/Runtime/Renderer/Private/PostProcess/PostProcessFFTBloom.cpp
@@ -42,7 +42,7 @@ public:
 	static bool ShouldCompilePermutation(const FGlobalShaderPermutationParameters& Parameters)
 	{
 		// @todo MetalMRT: Metal MRT can't cope with the threadgroup storage requirements for these shaders right now
-		return IsFeatureLevelSupported(Parameters.Platform, ERHIFeatureLevel::SM5) && (!IsMetalPlatform(Parameters.Platform) || RHIGetShaderLanguageVersion(Parameters.Platform) >= 2) && !IsMetalMRTPlatform(Parameters.Platform);
+		return IsFeatureLevelSupported(Parameters.Platform, ERHIFeatureLevel::SM5) && !IsMetalMRTPlatform(Parameters.Platform);
 	}
 
 	static void ModifyCompilationEnvironment(const FGlobalShaderPermutationParameters& Parameters, FShaderCompilerEnvironment& OutEnvironment)
@@ -139,7 +139,7 @@ public:
 	static bool ShouldCompilePermutation(const FGlobalShaderPermutationParameters& Parameters)
 	{
 		// @todo MetalMRT: Metal MRT can't cope with the threadgroup storage requirements for these shaders right now
-		return IsFeatureLevelSupported(Parameters.Platform, ERHIFeatureLevel::SM5) && (!IsMetalPlatform(Parameters.Platform) || RHIGetShaderLanguageVersion(Parameters.Platform) >= 2) && !IsMetalMRTPlatform(Parameters.Platform);
+		return IsFeatureLevelSupported(Parameters.Platform, ERHIFeatureLevel::SM5) && !IsMetalMRTPlatform(Parameters.Platform);
 	}
 
 	static void ModifyCompilationEnvironment(const FGlobalShaderPermutationParameters& Parameters, FShaderCompilerEnvironment& OutEnvironment)
@@ -228,7 +228,7 @@ public:
 	static bool ShouldCompilePermutation(const FGlobalShaderPermutationParameters& Parameters)
 	{
 		// @todo MetalMRT: Metal MRT can't cope with the threadgroup storage requirements for these shaders right now
-		return IsFeatureLevelSupported(Parameters.Platform, ERHIFeatureLevel::SM5) && (!IsMetalPlatform(Parameters.Platform) || RHIGetShaderLanguageVersion(Parameters.Platform) >= 2) && !IsMetalMRTPlatform(Parameters.Platform);
+		return IsFeatureLevelSupported(Parameters.Platform, ERHIFeatureLevel::SM5) && !IsMetalMRTPlatform(Parameters.Platform);
 	}
 
 	static void ModifyCompilationEnvironment(const FGlobalShaderPermutationParameters& Parameters, FShaderCompilerEnvironment& OutEnvironment)
@@ -326,7 +326,7 @@ public:
 	static bool ShouldCompilePermutation(const FGlobalShaderPermutationParameters& Parameters)
 	{
 		// @todo MetalMRT: Metal MRT can't cope with the threadgroup storage requirements for these shaders right now
-		return IsFeatureLevelSupported(Parameters.Platform, ERHIFeatureLevel::SM5) && (!IsMetalPlatform(Parameters.Platform) || RHIGetShaderLanguageVersion(Parameters.Platform) >= 2) && !IsMetalMRTPlatform(Parameters.Platform);
+		return IsFeatureLevelSupported(Parameters.Platform, ERHIFeatureLevel::SM5) && !IsMetalMRTPlatform(Parameters.Platform);
 	}
 
 	static void ModifyCompilationEnvironment(const FGlobalShaderPermutationParameters& Parameters, FShaderCompilerEnvironment& OutEnvironment)
