Missing function override meant that all volumetic fog layers were being rendered into layer 0 on Metal - we don't have geometry shaders so the vertex shader does the layer selection.

diff --git a/Engine/Source/Runtime/Renderer/Private/VolumetricFogVoxelization.cpp b/Engine/Source/Runtime/Renderer/Private/VolumetricFogVoxelization.cpp
--- a/Engine/Source/Runtime/Renderer/Private/VolumetricFogVoxelization.cpp
+++ b/Engine/Source/Runtime/Renderer/Private/VolumetricFogVoxelization.cpp
@@ -207,6 +207,24 @@ public:
 
 	const FShaderParameter& GetVoxelizationPassIndexParameter() const { return VoxelizationPassIndex; }
 
+	void GetShaderBindings(
+						   const FScene* Scene,
+						   ERHIFeatureLevel::Type FeatureLevel,
+						   const FPrimitiveSceneProxy* PrimitiveSceneProxy,
+						   const FMaterialRenderProxy& MaterialRenderProxy,
+						   const FMaterial& Material,
+						   const TUniformBufferRef<FViewUniformShaderParameters>& ViewUniformBuffer,
+						   FUniformBufferRHIParamRef PassUniformBufferValue,
+						   const FVoxelizeVolumeShaderElementData& ShaderElementData,
+						   FMeshDrawSingleShaderBindings& ShaderBindings) const
+	{
+		FMeshMaterialShader::GetShaderBindings(Scene, FeatureLevel, PrimitiveSceneProxy, MaterialRenderProxy, Material, ViewUniformBuffer, PassUniformBufferValue, ShaderElementData, ShaderBindings);
+		if (!RHISupportsGeometryShaders(Scene->GetShaderPlatform()))
+		{
+			ShaderBindings.Add(VoxelizationPassIndex, ShaderElementData.VoxelizationPassIndex);
+		}
+	}
+	
 protected:
 
 	FShaderParameter VoxelizationPassIndex;
