Improved GenerateDynamicMeshDrawCommands checks, so they print more information and variable values.

diff --git a/Engine/Source/Runtime/Renderer/Private/MeshDrawCommands.cpp b/Engine/Source/Runtime/Renderer/Private/MeshDrawCommands.cpp
--- a/Engine/Source/Runtime/Renderer/Private/MeshDrawCommands.cpp
+++ b/Engine/Source/Runtime/Renderer/Private/MeshDrawCommands.cpp
@@ -345,9 +345,9 @@ void GenerateDynamicMeshDrawCommands(
 	EMeshPass::Type PassType,
 	FMeshPassProcessor* PassMeshProcessor,
 	const TArray<FMeshBatchAndRelevance, SceneRenderingAllocator>& DynamicMeshElements,
-	int32 NumDynamicMeshElements,
+	int32 MaxNumDynamicMeshElements,
 	const TArray<const FStaticMesh*, SceneRenderingAllocator>& DynamicMeshCommandBuildRequests,
-	int32 NumDynamicMeshCommandBuildRequestElements,
+	int32 MaxNumBuildRequestElements,
 	FMeshCommandOneFrameArray& VisibleCommands,
 	FDynamicMeshDrawCommandStorage& MeshDrawCommandStorage
 )
@@ -361,34 +361,40 @@ void GenerateDynamicMeshDrawCommands(
 	);
 	PassMeshProcessor->SetDrawListContext(&DynamicPassMeshDrawListContext);
 
-	const int32 NumVisibleCommandsBeforeDynamicBatches = VisibleCommands.Num();
-	const int32 NumDynamicMeshBatches = DynamicMeshElements.Num();
-
-	for (int32 MeshIndex = 0; MeshIndex < NumDynamicMeshBatches; MeshIndex++)
 	{
-		const FMeshBatchAndRelevance& MeshAndRelevance = DynamicMeshElements[MeshIndex];
-		check(!MeshAndRelevance.Mesh->bRequiresPerElementVisibility);
-		const uint64 BatchElementMask = ~0ull;
+		const int32 NumCommandsBefore = VisibleCommands.Num();
+		const int32 NumDynamicMeshBatches = DynamicMeshElements.Num();
 
-		PassMeshProcessor->AddMeshBatch(*MeshAndRelevance.Mesh, BatchElementMask, MeshAndRelevance.PrimitiveSceneProxy);
-	}
+		for (int32 MeshIndex = 0; MeshIndex < NumDynamicMeshBatches; MeshIndex++)
+		{
+			const FMeshBatchAndRelevance& MeshAndRelevance = DynamicMeshElements[MeshIndex];
+			check(!MeshAndRelevance.Mesh->bRequiresPerElementVisibility);
+			const uint64 BatchElementMask = ~0ull;
 
-	// Check if we didn't generate more mesh draw commands, than we assumed.
-	check(VisibleCommands.Num() - NumVisibleCommandsBeforeDynamicBatches <= NumDynamicMeshElements);
+			PassMeshProcessor->AddMeshBatch(*MeshAndRelevance.Mesh, BatchElementMask, MeshAndRelevance.PrimitiveSceneProxy);
+		}
 
-	const int32 NumVisibleCommandsBeforeDynamicRequests = VisibleCommands.Num();
-	const int32 NumStaticMeshBatches = DynamicMeshCommandBuildRequests.Num();
+		const int32 NumCommandsGenerated = VisibleCommands.Num() - NumCommandsBefore;
+		checkf(NumCommandsGenerated <= MaxNumDynamicMeshElements,
+			TEXT("Generated %d mesh draw commands for DynamicMeshElements, while preallocating resources only for %d of them."), NumCommandsGenerated, MaxNumDynamicMeshElements);
+	}
 
-	for (int32 MeshIndex = 0; MeshIndex < NumStaticMeshBatches; MeshIndex++)
 	{
-		const FStaticMesh* StaticMeshBatch = DynamicMeshCommandBuildRequests[MeshIndex];
-		const uint64 BatchElementMask = StaticMeshBatch->bRequiresPerElementVisibility ? View.StaticMeshBatchVisibility[StaticMeshBatch->BatchVisibilityId] : ~0ull;
+		const int32 NumCommandsBefore = VisibleCommands.Num();
+		const int32 NumStaticMeshBatches = DynamicMeshCommandBuildRequests.Num();
 
-		PassMeshProcessor->AddMeshBatch(*StaticMeshBatch, BatchElementMask, StaticMeshBatch->PrimitiveSceneInfo->Proxy, StaticMeshBatch->Id);
-	}
+		for (int32 MeshIndex = 0; MeshIndex < NumStaticMeshBatches; MeshIndex++)
+		{
+			const FStaticMesh* StaticMeshBatch = DynamicMeshCommandBuildRequests[MeshIndex];
+			const uint64 BatchElementMask = StaticMeshBatch->bRequiresPerElementVisibility ? View.StaticMeshBatchVisibility[StaticMeshBatch->BatchVisibilityId] : ~0ull;
 
-	// Check if we didn't generate more mesh draw commands, than we assumed.
-	check(VisibleCommands.Num() - NumVisibleCommandsBeforeDynamicRequests <= NumDynamicMeshCommandBuildRequestElements);
+			PassMeshProcessor->AddMeshBatch(*StaticMeshBatch, BatchElementMask, StaticMeshBatch->PrimitiveSceneInfo->Proxy, StaticMeshBatch->Id);
+		}
+
+		const int32 NumCommandsGenerated = VisibleCommands.Num() - NumCommandsBefore;
+		checkf(NumCommandsGenerated <= MaxNumBuildRequestElements,
+			TEXT("Generated %d mesh draw commands for DynamicMeshCommandBuildRequests, while preallocating resources only for %d of them."), NumCommandsGenerated, MaxNumBuildRequestElements);
+	}
 }
 
 /**
@@ -402,9 +408,9 @@ void GenerateMobileBasePassDynamicMeshDrawCommands(
 	FMeshPassProcessor* PassMeshProcessor,
 	FMeshPassProcessor* MobilePassCSMPassMeshProcessor,
 	const TArray<FMeshBatchAndRelevance, SceneRenderingAllocator>& DynamicMeshElements,
-	int32 NumDynamicMeshElements,
+	int32 MaxNumDynamicMeshElements,
 	const TArray<const FStaticMesh*, SceneRenderingAllocator>& DynamicMeshCommandBuildRequests,
-	int32 NumDynamicMeshCommandBuildRequestElements,
+	int32 MaxNumBuildRequestElements,
 	FMeshCommandOneFrameArray& VisibleCommands,
 	FDynamicMeshDrawCommandStorage& MeshDrawCommandStorage
 )
@@ -421,52 +427,58 @@ void GenerateMobileBasePassDynamicMeshDrawCommands(
 
 	const FMobileCSMVisibilityInfo& MobileCSMVisibilityInfo = View.MobileCSMVisibilityInfo;
 	
-	const int32 NumVisibleCommandsBeforeDynamicBatches = VisibleCommands.Num();
-	const int32 NumDynamicMeshBatches = DynamicMeshElements.Num();
-
-	for (int32 MeshIndex = 0; MeshIndex < NumDynamicMeshBatches; MeshIndex++)
 	{
-		const FMeshBatchAndRelevance& MeshAndRelevance = DynamicMeshElements[MeshIndex];
-		check(!MeshAndRelevance.Mesh->bRequiresPerElementVisibility);
-		const uint64 BatchElementMask = ~0ull;
+		const int32 NumCommandsBefore = VisibleCommands.Num();
+		const int32 NumDynamicMeshBatches = DynamicMeshElements.Num();
 
-		const int32 PrimitiveIndex = MeshAndRelevance.PrimitiveSceneProxy->GetPrimitiveSceneInfo()->GetIndex();
-		if (MobileCSMVisibilityInfo.bMobileDynamicCSMInUse 
-			&& (MobileCSMVisibilityInfo.bAlwaysUseCSM || MobileCSMVisibilityInfo.MobilePrimitiveCSMReceiverVisibilityMap[PrimitiveIndex]))
+		for (int32 MeshIndex = 0; MeshIndex < NumDynamicMeshBatches; MeshIndex++)
 		{
-			MobilePassCSMPassMeshProcessor->AddMeshBatch(*MeshAndRelevance.Mesh, BatchElementMask, MeshAndRelevance.PrimitiveSceneProxy);
-		}
-		else
-		{
-			PassMeshProcessor->AddMeshBatch(*MeshAndRelevance.Mesh, BatchElementMask, MeshAndRelevance.PrimitiveSceneProxy);
-		}
-	}
+			const FMeshBatchAndRelevance& MeshAndRelevance = DynamicMeshElements[MeshIndex];
+			check(!MeshAndRelevance.Mesh->bRequiresPerElementVisibility);
+			const uint64 BatchElementMask = ~0ull;
 
-	// Check if we didn't generate more mesh draw commands, than we assumed.
-	check(VisibleCommands.Num() - NumVisibleCommandsBeforeDynamicBatches <= NumDynamicMeshElements);
+			const int32 PrimitiveIndex = MeshAndRelevance.PrimitiveSceneProxy->GetPrimitiveSceneInfo()->GetIndex();
+			if (MobileCSMVisibilityInfo.bMobileDynamicCSMInUse
+				&& (MobileCSMVisibilityInfo.bAlwaysUseCSM || MobileCSMVisibilityInfo.MobilePrimitiveCSMReceiverVisibilityMap[PrimitiveIndex]))
+			{
+				MobilePassCSMPassMeshProcessor->AddMeshBatch(*MeshAndRelevance.Mesh, BatchElementMask, MeshAndRelevance.PrimitiveSceneProxy);
+			}
+			else
+			{
+				PassMeshProcessor->AddMeshBatch(*MeshAndRelevance.Mesh, BatchElementMask, MeshAndRelevance.PrimitiveSceneProxy);
+			}
+		}
 
-	const int32 NumVisibleCommandsBeforeDynamicRequests = VisibleCommands.Num();
-	const int32 NumStaticMeshBatches = DynamicMeshCommandBuildRequests.Num();
+		const int32 NumCommandsGenerated = VisibleCommands.Num() - NumCommandsBefore;
+		checkf(NumCommandsGenerated <= MaxNumDynamicMeshElements,
+			TEXT("Generated %d mesh draw commands for DynamicMeshElements, while preallocating resources only for %d of them."), NumCommandsGenerated, MaxNumDynamicMeshElements);
+	}
 
-	for (int32 MeshIndex = 0; MeshIndex < NumStaticMeshBatches; MeshIndex++)
 	{
-		const FStaticMesh* StaticMeshBatch = DynamicMeshCommandBuildRequests[MeshIndex];
-		const uint64 BatchElementMask = StaticMeshBatch->bRequiresPerElementVisibility ? View.StaticMeshBatchVisibility[StaticMeshBatch->BatchVisibilityId] : ~0ull;
+		const int32 NumCommandsBefore = VisibleCommands.Num();
+		const int32 NumStaticMeshBatches = DynamicMeshCommandBuildRequests.Num();
 
-		const int32 PrimitiveIndex = StaticMeshBatch->PrimitiveSceneInfo->Proxy->GetPrimitiveSceneInfo()->GetIndex();
-		if (MobileCSMVisibilityInfo.bMobileDynamicCSMInUse
-			&& (MobileCSMVisibilityInfo.bAlwaysUseCSM || MobileCSMVisibilityInfo.MobilePrimitiveCSMReceiverVisibilityMap[PrimitiveIndex]))
+		for (int32 MeshIndex = 0; MeshIndex < NumStaticMeshBatches; MeshIndex++)
 		{
-			MobilePassCSMPassMeshProcessor->AddMeshBatch(*StaticMeshBatch, BatchElementMask, StaticMeshBatch->PrimitiveSceneInfo->Proxy, StaticMeshBatch->Id);
-		}
-		else
-		{
-			PassMeshProcessor->AddMeshBatch(*StaticMeshBatch, BatchElementMask, StaticMeshBatch->PrimitiveSceneInfo->Proxy, StaticMeshBatch->Id);
+			const FStaticMesh* StaticMeshBatch = DynamicMeshCommandBuildRequests[MeshIndex];
+			const uint64 BatchElementMask = StaticMeshBatch->bRequiresPerElementVisibility ? View.StaticMeshBatchVisibility[StaticMeshBatch->BatchVisibilityId] : ~0ull;
+
+			const int32 PrimitiveIndex = StaticMeshBatch->PrimitiveSceneInfo->Proxy->GetPrimitiveSceneInfo()->GetIndex();
+			if (MobileCSMVisibilityInfo.bMobileDynamicCSMInUse
+				&& (MobileCSMVisibilityInfo.bAlwaysUseCSM || MobileCSMVisibilityInfo.MobilePrimitiveCSMReceiverVisibilityMap[PrimitiveIndex]))
+			{
+				MobilePassCSMPassMeshProcessor->AddMeshBatch(*StaticMeshBatch, BatchElementMask, StaticMeshBatch->PrimitiveSceneInfo->Proxy, StaticMeshBatch->Id);
+			}
+			else
+			{
+				PassMeshProcessor->AddMeshBatch(*StaticMeshBatch, BatchElementMask, StaticMeshBatch->PrimitiveSceneInfo->Proxy, StaticMeshBatch->Id);
+			}
 		}
-	}
 
-	// Check if we didn't generate more mesh draw commands, than we assumed.
-	check(VisibleCommands.Num() - NumVisibleCommandsBeforeDynamicRequests <= NumDynamicMeshCommandBuildRequestElements);
+		const int32 NumCommandsGenerated = VisibleCommands.Num() - NumCommandsBefore;
+		checkf(NumCommandsGenerated <= MaxNumBuildRequestElements,
+			TEXT("Generated %d mesh draw commands for DynamicMeshCommandBuildRequests, while preallocating resources only for %d of them."), NumCommandsGenerated, MaxNumBuildRequestElements);
+	}
 }
 
 /**
