Removed legacy Drawing Policies and Static Mesh Draw Lists.  These are now replaced by FMeshDrawCommand / FMeshPassProcessor everywhere in the renderer.

diff --git a/Engine/Source/Runtime/Renderer/Private/BasePassRendering.h b/Engine/Source/Runtime/Renderer/Private/BasePassRendering.h
--- a/Engine/Source/Runtime/Renderer/Private/BasePassRendering.h
+++ b/Engine/Source/Runtime/Renderer/Private/BasePassRendering.h
@@ -15,7 +15,6 @@
 #include "RHIStaticStates.h"
 #include "SceneManagement.h"
 #include "Materials/Material.h"
-#include "DrawingPolicy.h"
 #include "PostProcess/SceneRenderTargets.h"
 #include "LightMapRendering.h"
 #include "VelocityRendering.h"
@@ -188,8 +187,6 @@ protected:
 		LightMapPolicyType::VertexParametersType::Bind(Initializer.ParameterMap);
 		BindBasePassUniformBuffer(Initializer.ParameterMap, PassUniformBuffer);
 		ReflectionCaptureBuffer.Bind(Initializer.ParameterMap, TEXT("ReflectionCapture"));
-		InstancedEyeIndexParameter.Bind(Initializer.ParameterMap, TEXT("InstancedEyeIndex"));
-		IsInstancedStereoParameter.Bind(Initializer.ParameterMap, TEXT("bIsInstancedStereo"));
 	}
 
 public:
@@ -205,40 +202,9 @@ public:
 		bool bShaderHasOutdatedParameters = FMeshMaterialShader::Serialize(Ar);
 		LightMapPolicyType::VertexParametersType::Serialize(Ar);
 		Ar << ReflectionCaptureBuffer;
-		Ar << InstancedEyeIndexParameter;
-		Ar << IsInstancedStereoParameter;
 		return bShaderHasOutdatedParameters;
 	}
 
-	void SetParameters(
-		FRHICommandList& RHICmdList, 
-		const FMaterialRenderProxy* MaterialRenderProxy,
-		const FVertexFactory* VertexFactory,
-		const FMaterial& InMaterialResource,
-		const FViewInfo& View,
-		const FDrawingPolicyRenderState& DrawRenderState,
-		bool bIsInstancedStereo
-		)
-	{
-		const FVertexShaderRHIParamRef ShaderRHI = GetVertexShader();
-
-		FMeshMaterialShader::SetParameters(RHICmdList, ShaderRHI, MaterialRenderProxy, InMaterialResource, View, DrawRenderState.GetViewUniformBuffer(), DrawRenderState.GetPassUniformBuffer());
-
-		SetUniformBufferParameter(RHICmdList, ShaderRHI, ReflectionCaptureBuffer, View.ReflectionCaptureUniformBuffer);
-
-		if (IsInstancedStereoParameter.IsBound())
-		{
-			SetShaderValue(RHICmdList, ShaderRHI, IsInstancedStereoParameter, bIsInstancedStereo);
-		}
-
-		if (InstancedEyeIndexParameter.IsBound())
-		{
-			SetShaderValue(RHICmdList, ShaderRHI, InstancedEyeIndexParameter, 0);
-		}
-	}
-
-	void SetMesh(FRHICommandList& RHICmdList, const FVertexFactory* VertexFactory,const FSceneView& View,const FPrimitiveSceneProxy* Proxy, const FMeshBatch& Mesh, const FMeshBatchElement& BatchElement, const FDrawingPolicyRenderState& DrawRenderState);
-
 	void GetShaderBindings(
 		const FScene* Scene,
 		ERHIFeatureLevel::Type FeatureLevel,
@@ -263,13 +229,7 @@ public:
 		FMeshDrawSingleShaderBindings& ShaderBindings,
 		FVertexInputStreamArray& VertexStreams) const;
 
-	void SetInstancedEyeIndex(FRHICommandList& RHICmdList, const uint32 EyeIndex);
-	
 	FShaderUniformBufferParameter ReflectionCaptureBuffer;
-
-private:
-	FShaderParameter InstancedEyeIndexParameter;
-	FShaderParameter IsInstancedStereoParameter;
 };
 
 
@@ -467,23 +427,6 @@ public:
 	}
 	TBasePassPixelShaderPolicyParamType() {}
 
-	void SetParameters(
-		FRHICommandList& RHICmdList, 
-		const FMaterialRenderProxy* MaterialRenderProxy, 
-		const FMaterial& MaterialResource, 
-		const FViewInfo* View, 
-		const FDrawingPolicyRenderState& DrawRenderState,
-		EBlendMode BlendMode)
-	{
-		const FPixelShaderRHIParamRef ShaderRHI = GetPixelShader();
-
-		FMeshMaterialShader::SetParameters(RHICmdList, ShaderRHI, MaterialRenderProxy, MaterialResource, *View, DrawRenderState.GetViewUniformBuffer(), DrawRenderState.GetPassUniformBuffer());
-
-		SetUniformBufferParameter(RHICmdList, ShaderRHI, ReflectionCaptureBuffer, View->ReflectionCaptureUniformBuffer);
-	}
-
-	void SetMesh(FRHICommandList& RHICmdList, const FVertexFactory* VertexFactory,const FSceneView& View,const FPrimitiveSceneProxy* Proxy,const FMeshBatchElement& BatchElement, const FDrawingPolicyRenderState& DrawRenderState, EBlendMode BlendMode);
-
 	void GetShaderBindings(
 		const FScene* Scene,
 		ERHIFeatureLevel::Type FeatureLevel,
@@ -652,25 +595,6 @@ void GetBasePassShaders<FUniformLightMapPolicy>(
 	TBasePassPixelShaderPolicyParamType<FUniformLightMapPolicy>*& PixelShader
 	);
 	
-class FBasePassDrawingPolicy : public FMeshDrawingPolicy
-{
-public:
-	FBasePassDrawingPolicy(
-		const FVertexFactory* InVertexFactory,
-		const FMaterialRenderProxy* InMaterialRenderProxy, 
-		const FMaterial& InMaterialResource, 
-		const FMeshDrawingPolicyOverrideSettings& InOverrideSettings,
-		bool bInEnableReceiveDecalOutput) : FMeshDrawingPolicy(InVertexFactory, InMaterialRenderProxy, InMaterialResource, InOverrideSettings), bEnableReceiveDecalOutput(bInEnableReceiveDecalOutput)
-	{}
-
-	void ApplyDitheredLODTransitionState(FDrawingPolicyRenderState& DrawRenderState, const FViewInfo& ViewInfo, const FStaticMesh& Mesh, const bool InAllowStencilDither);
-
-protected:
-
-	/** Whether or not outputing the receive decal boolean */
-	uint32 bEnableReceiveDecalOutput : 1;
-};
-
 class FBasePassMeshProcessor : public FMeshPassProcessor
 {
 public:
diff --git a/Engine/Source/Runtime/Renderer/Private/BasePassRendering.inl b/Engine/Source/Runtime/Renderer/Private/BasePassRendering.inl
--- a/Engine/Source/Runtime/Renderer/Private/BasePassRendering.inl
+++ b/Engine/Source/Runtime/Renderer/Private/BasePassRendering.inl
@@ -19,13 +19,6 @@ struct FMeshBatch;
 struct FMeshBatchElement;
 struct FMeshDrawingRenderState;
 
-template<typename LightMapPolicyType>
-inline void TBasePassVertexShaderPolicyParamType<LightMapPolicyType>::SetMesh(FRHICommandList& RHICmdList, const FVertexFactory* VertexFactory, const FSceneView& View, const FPrimitiveSceneProxy* Proxy, const FMeshBatch& Mesh, const FMeshBatchElement& BatchElement, const FDrawingPolicyRenderState& DrawRenderState)
-{
-	FVertexShaderRHIParamRef VertexShaderRHI = GetVertexShader();
-	FMeshMaterialShader::SetMesh(RHICmdList, VertexShaderRHI, VertexFactory, View, Proxy, BatchElement, DrawRenderState);
-}
-
 template<typename LightMapPolicyType>
 void TBasePassVertexShaderPolicyParamType<LightMapPolicyType>::GetShaderBindings(
 	const FScene* Scene,
@@ -74,21 +67,6 @@ void TBasePassVertexShaderPolicyParamType<LightMapPolicyType>::GetElementShaderB
 	FMeshMaterialShader::GetElementShaderBindings(Scene, ViewIfDynamicMeshCommand, VertexFactory, bShaderRequiresPositionOnlyStream, FeatureLevel, PrimitiveSceneProxy, MeshBatch, BatchElement, ShaderElementData, ShaderBindings, VertexStreams);
 }
 
-template<typename LightMapPolicyType>
-void TBasePassVertexShaderPolicyParamType<LightMapPolicyType>::SetInstancedEyeIndex(FRHICommandList& RHICmdList, const uint32 EyeIndex)
-{
-	if (InstancedEyeIndexParameter.IsBound())
-	{
-		SetShaderValue(RHICmdList, GetVertexShader(), InstancedEyeIndexParameter, EyeIndex);
-	}
-}
-
-template<typename LightMapPolicyType>
-void TBasePassPixelShaderPolicyParamType<LightMapPolicyType>::SetMesh(FRHICommandList& RHICmdList, const FVertexFactory* VertexFactory, const FSceneView& View, const FPrimitiveSceneProxy* Proxy, const FMeshBatchElement& BatchElement, const FDrawingPolicyRenderState& DrawRenderState, EBlendMode BlendMode)
-{
-	FMeshMaterialShader::SetMesh(RHICmdList, GetPixelShader(), VertexFactory, View, Proxy, BatchElement, DrawRenderState);
-}
-
 template<typename LightMapPolicyType>
 void TBasePassPixelShaderPolicyParamType<LightMapPolicyType>::GetShaderBindings(
 	const FScene* Scene,
diff --git a/Engine/Source/Runtime/Renderer/Private/CompositionLighting/PostProcessDeferredDecals.h b/Engine/Source/Runtime/Renderer/Private/CompositionLighting/PostProcessDeferredDecals.h
--- a/Engine/Source/Runtime/Renderer/Private/CompositionLighting/PostProcessDeferredDecals.h
+++ b/Engine/Source/Runtime/Renderer/Private/CompositionLighting/PostProcessDeferredDecals.h
@@ -13,7 +13,6 @@
 #include "PostProcess/SceneRenderTargets.h"
 #include "PostProcess/RenderingCompositionGraph.h"
 #include "DecalRenderingCommon.h"
-#include "DrawingPolicy.h"
 
 // ePId_Input0: SceneColor (not needed for DBuffer decals)
 // derives from TRenderingCompositePassBase<InputCount, OutputCount> 
diff --git a/Engine/Source/Runtime/Renderer/Private/CompositionLighting/PostProcessMeshDecals.cpp b/Engine/Source/Runtime/Renderer/Private/CompositionLighting/PostProcessMeshDecals.cpp
--- a/Engine/Source/Runtime/Renderer/Private/CompositionLighting/PostProcessMeshDecals.cpp
+++ b/Engine/Source/Runtime/Renderer/Private/CompositionLighting/PostProcessMeshDecals.cpp
@@ -13,7 +13,6 @@
 #include "RHIStaticStates.h"
 #include "PostProcess/SceneRenderTargets.h"
 #include "MaterialShaderType.h"
-#include "DrawingPolicy.h"
 #include "MeshMaterialShader.h"
 #include "ShaderBaseClasses.h"
 #include "DepthRendering.h"
@@ -61,18 +60,6 @@ protected:
 	{
 		return FMeshDecalAccumulatePolicy::ShouldCompilePermutation(Platform,Material,VertexFactoryType);
 	}
-
-public:
-	
-	void SetParameters(FRHICommandList& RHICmdList, const FVertexFactory* VertexFactory,const FMaterialRenderProxy* MaterialRenderProxy,const FSceneView* View,const FDrawingPolicyRenderState& DrawRenderState)
-	{
-		FMeshMaterialShader::SetParameters(RHICmdList, GetVertexShader(), MaterialRenderProxy, *MaterialRenderProxy->GetMaterial(View->GetFeatureLevel()), *View, DrawRenderState.GetViewUniformBuffer(), DrawRenderState.GetPassUniformBuffer());
-	}
-
-	void SetMesh(FRHICommandList& RHICmdList, const FVertexFactory* VertexFactory,const FSceneView& View,const FPrimitiveSceneProxy* Proxy,const FMeshBatchElement& BatchElement,const FDrawingPolicyRenderState& DrawRenderState)
-	{
-		FMeshMaterialShader::SetMesh(RHICmdList, GetVertexShader(),VertexFactory,View,Proxy,BatchElement,DrawRenderState);
-	}
 };
 
 
@@ -155,21 +142,6 @@ public:
 	{
 	}
 
-	void SetParameters(
-		FRHICommandList& RHICmdList, 
-		const FMaterialRenderProxy* MaterialRenderProxy,
-		const FSceneView& View,
-		const FDrawingPolicyRenderState& DrawRenderState
-		)
-	{
-		FMeshMaterialShader::SetParameters(RHICmdList, GetPixelShader(), MaterialRenderProxy, *MaterialRenderProxy->GetMaterial(View.GetFeatureLevel()), View, DrawRenderState.GetViewUniformBuffer(), DrawRenderState.GetPassUniformBuffer());
-	}
-
-	void SetMesh(FRHICommandList& RHICmdList, const FVertexFactory* VertexFactory,const FSceneView& View,const FPrimitiveSceneProxy* Proxy,const FMeshBatchElement& BatchElement,const FDrawingPolicyRenderState& DrawRenderState)
-	{
-		FMeshMaterialShader::SetMesh(RHICmdList, GetPixelShader(),VertexFactory,View,Proxy,BatchElement,DrawRenderState);
-	}
-
 	virtual bool Serialize(FArchive& Ar) override
 	{
 		bool bShaderHasOutdatedParameters = FMeshMaterialShader::Serialize(Ar);
diff --git a/Engine/Source/Runtime/Renderer/Private/ConvertToUniformMesh.cpp b/Engine/Source/Runtime/Renderer/Private/ConvertToUniformMesh.cpp
--- a/Engine/Source/Runtime/Renderer/Private/ConvertToUniformMesh.cpp
+++ b/Engine/Source/Runtime/Renderer/Private/ConvertToUniformMesh.cpp
@@ -11,7 +11,6 @@
 #include "MeshBatch.h"
 #include "MaterialShaderType.h"
 #include "MaterialShader.h"
-#include "DrawingPolicy.h"
 #include "MeshMaterialShader.h"
 #include "ScenePrivate.h"
 #include "DistanceFieldLightingShared.h"
@@ -40,19 +39,6 @@ protected:
 			&& (FCString::Strstr(VertexFactoryType->GetName(), TEXT("LocalVertexFactory")) != NULL
 				|| FCString::Strstr(VertexFactoryType->GetName(), TEXT("InstancedStaticMeshVertexFactory")) != NULL);
 	}
-
-public:
-	
-	void SetParameters(FRHICommandList& RHICmdList, const FVertexFactory* VertexFactory,const FMaterialRenderProxy* MaterialRenderProxy,const FSceneView* View)
-	{
-		TUniformBufferRef<FSceneTexturesUniformParameters> PassUniformBufferValue = CreateSceneTextureUniformBufferSingleDraw(RHICmdList, ESceneTextureSetupMode::None, View->FeatureLevel);
-		FMeshMaterialShader::SetParameters(RHICmdList, GetVertexShader(), MaterialRenderProxy, *MaterialRenderProxy->GetMaterial(View->GetFeatureLevel()), *View, View->ViewUniformBuffer, PassUniformBufferValue);
-	}
-
-	void SetMesh(FRHICommandList& RHICmdList, const FVertexFactory* VertexFactory,const FSceneView& View,const FPrimitiveSceneProxy* Proxy,const FMeshBatchElement& BatchElement,const FDrawingPolicyRenderState& DrawRenderState)
-	{
-		FMeshMaterialShader::SetMesh(RHICmdList, GetVertexShader(),VertexFactory,View,Proxy,BatchElement,DrawRenderState);
-	}
 };
 
 IMPLEMENT_MATERIAL_SHADER_TYPE(,FConvertToUniformMeshVS,TEXT("/Engine/Private/ConvertToUniformMesh.usf"),TEXT("ConvertToUniformMeshVS"),SF_Vertex); 
@@ -132,19 +118,6 @@ protected:
 		GetUniformMeshStreamOutLayout(ElementList);
 		RasterizedStream = -1;
 	}
-
-public:
-	
-	void SetParameters(FRHICommandList& RHICmdList, const FVertexFactory* VertexFactory,const FMaterialRenderProxy* MaterialRenderProxy,const FSceneView* View)
-	{
-		TUniformBufferRef<FSceneTexturesUniformParameters> PassUniformBufferValue = CreateSceneTextureUniformBufferSingleDraw(RHICmdList, ESceneTextureSetupMode::None, View->FeatureLevel);
-		FMeshMaterialShader::SetParameters(RHICmdList, (FGeometryShaderRHIParamRef)GetGeometryShader(), MaterialRenderProxy, *MaterialRenderProxy->GetMaterial(View->GetFeatureLevel()), *View, View->ViewUniformBuffer, PassUniformBufferValue);
-	}
-
-	void SetMesh(FRHICommandList& RHICmdList, const FVertexFactory* VertexFactory,const FSceneView& View,const FPrimitiveSceneProxy* Proxy,const FMeshBatchElement& BatchElement,const FDrawingPolicyRenderState& DrawRenderState)
-	{
-		FMeshMaterialShader::SetMesh(RHICmdList, (FGeometryShaderRHIParamRef)GetGeometryShader(),VertexFactory,View,Proxy,BatchElement,DrawRenderState);
-	}
 };
 
 IMPLEMENT_MATERIAL_SHADER_TYPE(,FConvertToUniformMeshGS,TEXT("/Engine/Private/ConvertToUniformMesh.usf"),TEXT("ConvertToUniformMeshGS"),SF_Geometry);
diff --git a/Engine/Source/Runtime/Renderer/Private/CustomDepthRendering.cpp b/Engine/Source/Runtime/Renderer/Private/CustomDepthRendering.cpp
--- a/Engine/Source/Runtime/Renderer/Private/CustomDepthRendering.cpp
+++ b/Engine/Source/Runtime/Renderer/Private/CustomDepthRendering.cpp
@@ -4,118 +4,13 @@
 	CustomDepthRendering.cpp: CustomDepth rendering implementation.
 =============================================================================*/
 
-#include "CustomDepthRendering.h"
 #include "SceneUtils.h"
-#include "DrawingPolicy.h"
 #include "DepthRendering.h"
 #include "SceneRendering.h"
 #include "SceneCore.h"
 #include "ScenePrivate.h"
 #include "MeshPassProcessor.inl"
 
-/*-----------------------------------------------------------------------------
-	FCustomDepthPrimSet
------------------------------------------------------------------------------*/
-
-bool FCustomDepthPrimSet::DrawPrims(FRHICommandListImmediate& RHICmdList, const FViewInfo& View, FDrawingPolicyRenderState& DrawRenderState, bool bWriteCustomStencilValues)
-{
-	bool bDirty=false;
-
-	if(Prims.Num())
-	{
-		SCOPED_DRAW_EVENT(RHICmdList, CustomDepth);
-
-		for (int32 PrimIdx = 0; PrimIdx < Prims.Num(); PrimIdx++)
-		{
-			FPrimitiveSceneProxy* PrimitiveSceneProxy = Prims[PrimIdx];
-			const FPrimitiveSceneInfo* PrimitiveSceneInfo = PrimitiveSceneProxy->GetPrimitiveSceneInfo();
-
-			if (View.PrimitiveVisibilityMap[PrimitiveSceneInfo->GetIndex()])
-			{
-				const FPrimitiveViewRelevance& ViewRelevance = View.PrimitiveViewRelevanceMap[PrimitiveSceneInfo->GetIndex()];
-
-				FDepthDrawingPolicyFactory::ContextType Context(DDM_AllOpaque, false);
-
-				FDrawingPolicyRenderState DrawRenderStateLocal(DrawRenderState);
-				if (bWriteCustomStencilValues)
-				{
-					const uint32 CustomDepthStencilValue = PrimitiveSceneProxy->GetCustomDepthStencilValue();
-					const static FDepthStencilStateRHIParamRef StencilStates[EStencilMask::SM_Count] =
-					{
-						TStaticDepthStencilState<true, CF_DepthNearOrEqual, true, CF_Always, SO_Keep, SO_Keep, SO_Replace, false, CF_Always, SO_Keep, SO_Keep, SO_Keep, 255, 255>::GetRHI(),
-						TStaticDepthStencilState<true, CF_DepthNearOrEqual, true, CF_Always, SO_Keep, SO_Replace, SO_Replace, false, CF_Always, SO_Keep, SO_Keep, SO_Keep, 255, 255>::GetRHI(),
-						TStaticDepthStencilState<true, CF_DepthNearOrEqual, true, CF_Always, SO_Keep, SO_Replace, SO_Replace, false, CF_Always, SO_Keep, SO_Keep, SO_Keep, 255, 1>::GetRHI(),
-						TStaticDepthStencilState<true, CF_DepthNearOrEqual, true, CF_Always, SO_Keep, SO_Replace, SO_Replace, false, CF_Always, SO_Keep, SO_Keep, SO_Keep, 255, 2>::GetRHI(),
-						TStaticDepthStencilState<true, CF_DepthNearOrEqual, true, CF_Always, SO_Keep, SO_Replace, SO_Replace, false, CF_Always, SO_Keep, SO_Keep, SO_Keep, 255, 4>::GetRHI(),
-						TStaticDepthStencilState<true, CF_DepthNearOrEqual, true, CF_Always, SO_Keep, SO_Replace, SO_Replace, false, CF_Always, SO_Keep, SO_Keep, SO_Keep, 255, 8>::GetRHI(),
-						TStaticDepthStencilState<true, CF_DepthNearOrEqual, true, CF_Always, SO_Keep, SO_Replace, SO_Replace, false, CF_Always, SO_Keep, SO_Keep, SO_Keep, 255, 16>::GetRHI(),
-						TStaticDepthStencilState<true, CF_DepthNearOrEqual, true, CF_Always, SO_Keep, SO_Replace, SO_Replace, false, CF_Always, SO_Keep, SO_Keep, SO_Keep, 255, 32>::GetRHI(),
-						TStaticDepthStencilState<true, CF_DepthNearOrEqual, true, CF_Always, SO_Keep, SO_Replace, SO_Replace, false, CF_Always, SO_Keep, SO_Keep, SO_Keep, 255, 64>::GetRHI(),
-						TStaticDepthStencilState<true, CF_DepthNearOrEqual, true, CF_Always, SO_Keep, SO_Replace, SO_Replace, false, CF_Always, SO_Keep, SO_Keep, SO_Keep, 255, 128>::GetRHI()
-					};
-					checkSlow(EStencilMask::SM_Count == ARRAY_COUNT(StencilStates));
-
-					DrawRenderStateLocal.SetDepthStencilState(StencilStates[(int32)PrimitiveSceneProxy->GetStencilWriteMask()]);
-					DrawRenderStateLocal.SetStencilRef(CustomDepthStencilValue);
-
-					if (View.GetFeatureLevel() <= ERHIFeatureLevel::ES3_1)
-					{
-						// On mobile platforms write custom stencil value to color target
-						Context.MobileColorValue = CustomDepthStencilValue/255.0f;
-					}
-				}
-
-				// Note: As for custom depth rendering the order doesn't matter we actually could iterate View.DynamicMeshElements without this indirection	
-				{
-					// range in View.DynamicMeshElements[]
-					FInt32Range range = View.GetDynamicMeshElementRange(PrimitiveSceneInfo->GetIndex());
-
-					for (int32 MeshBatchIndex = range.GetLowerBoundValue(); MeshBatchIndex < range.GetUpperBoundValue(); MeshBatchIndex++)
-					{
-						const FMeshBatchAndRelevance& MeshBatchAndRelevance = View.DynamicMeshElements[MeshBatchIndex];
-
-						checkSlow(MeshBatchAndRelevance.PrimitiveSceneProxy == PrimitiveSceneInfo->Proxy);
-
-						const FMeshBatch& MeshBatch = *MeshBatchAndRelevance.Mesh;
-
-						FDepthDrawingPolicyFactory::DrawDynamicMesh(RHICmdList, View, Context, MeshBatch, true, DrawRenderStateLocal, MeshBatchAndRelevance.PrimitiveSceneProxy, MeshBatch.BatchHitProxyId, View.IsInstancedStereoPass(), false);
-					}
-				}
-
-				if (ViewRelevance.bStaticRelevance)
-				{
-					for (int32 StaticMeshIdx = 0; StaticMeshIdx < PrimitiveSceneInfo->StaticMeshes.Num(); StaticMeshIdx++)
-					{
-						const FStaticMesh& StaticMesh = PrimitiveSceneInfo->StaticMeshes[StaticMeshIdx];
-
-						if (View.StaticMeshVisibilityMap[StaticMesh.Id])
-						{
-							FDrawingPolicyRenderState DrawRenderStateLocal2(DrawRenderStateLocal);
-							FMeshDrawingPolicy::OnlyApplyDitheredLODTransitionState(DrawRenderStateLocal2, View, StaticMesh, false);
-
-							bDirty |= FDepthDrawingPolicyFactory::DrawStaticMesh(
-								RHICmdList, 
-								View,
-								Context,
-								StaticMesh,
-								StaticMesh.bRequiresPerElementVisibility ? View.StaticMeshBatchVisibility[StaticMesh.BatchVisibilityId] : ((1ull << StaticMesh.Elements.Num()) - 1),
-								true,
-								DrawRenderStateLocal2,
-								PrimitiveSceneProxy,
-								StaticMesh.BatchHitProxyId, 
-								View.IsInstancedStereoPass(),
-								false
-								);
-						}
-					}
-				}
-			}
-		}
-	}
-
-	return bDirty;
-}
-
 class FCustomDepthPassMeshProcessor : public FMeshPassProcessor
 {
 public:
diff --git a/Engine/Source/Runtime/Renderer/Private/DebugViewModeRendering.cpp b/Engine/Source/Runtime/Renderer/Private/DebugViewModeRendering.cpp
--- a/Engine/Source/Runtime/Renderer/Private/DebugViewModeRendering.cpp
+++ b/Engine/Source/Runtime/Renderer/Private/DebugViewModeRendering.cpp
@@ -48,57 +48,6 @@ IMPLEMENT_MATERIAL_SHADER_TYPE(,FDebugViewModeVS,TEXT("/Engine/Private/DebugView
 IMPLEMENT_MATERIAL_SHADER_TYPE(,FDebugViewModeHS,TEXT("/Engine/Private/DebugViewModeVertexShader.usf"),TEXT("MainHull"),SF_Hull);	
 IMPLEMENT_MATERIAL_SHADER_TYPE(,FDebugViewModeDS,TEXT("/Engine/Private/DebugViewModeVertexShader.usf"),TEXT("MainDomain"),SF_Domain);
 
-#if 0 // todo : implement FMissingShaderPS
-/**
-* Pixel shader that renders texture streamer wanted mips accuracy.
-*/
-class FMissingShaderPS : public FGlobalShader
-{
-	DECLARE_SHADER_TYPE(FMissingShaderPS,Global);
-
-public:
-
-	static bool ShouldCompilePermutation(const FGlobalShaderPermutationParameters& Parameters) { return AllowDebugViewVSDSHS(Parameters.Platform); }
-
-	FMissingShaderPS(const ShaderMetaType::CompiledShaderInitializerType& Initializer):	FGlobalShader(Initializer) {}
-	FMissingShaderPS() {}
-
-	virtual bool Serialize(FArchive& Ar) override { return FGlobalShader::Serialize(Ar); }
-
-	static void ModifyCompilationEnvironment(const FGlobalShaderPermutationParameters& Parameters, FShaderCompilerEnvironment& OutEnvironment)
-	{
-		OutEnvironment.SetDefine(TEXT("UNDEFINED_VALUE"), UndefinedStreamingAccuracyIntensity);
-	}
-
-	virtual void SetParameters(
-		FRHICommandList& RHICmdList, 
-		int32 NumVSInstructions, 
-		int32 NumPSInstructions, 
-		const FMaterialRenderProxy* MaterialRenderProxy,
-		const FMaterial& Material,
-		const FSceneView& View,
-		const FDrawingPolicyRenderState& DrawRenderState
-		) override {}
-
-	virtual void SetMesh(
-		FRHICommandList& RHICmdList, 
-		const FVertexFactory* VertexFactory,
-		const FSceneView& View,
-		const FPrimitiveSceneProxy* Proxy,
-		int32 VisualizeLODIndex,
-		const FMeshBatchElement& BatchElement, 
-		const FDrawingPolicyRenderState& DrawRenderState
-		) override{}
-
-	virtual void SetMesh(FRHICommandList& RHICmdList, const FSceneView& View) override {}
-
-	virtual FShader* GetShader() override { return static_cast<FShader*>(this); }
-};
-
-IMPLEMENT_SHADER_TYPE(,FMissingShaderPS,TEXT("/Engine/Private/MissingShaderPixelShader.usf"),TEXT("Main"),SF_Pixel);
-
-#endif // todo : implement FMissingShaderPS
-
 ENGINE_API bool GetDebugViewMaterial(const UMaterialInterface* InMaterialInterface, EDebugViewShaderMode InDebugViewMode, const FMaterialRenderProxy*& OutMaterialRenderProxy, const FMaterial*& OutMaterial);
 
 void FDeferredShadingSceneRenderer::DoDebugViewModePostProcessing(FRHICommandListImmediate& RHICmdList, const FViewInfo& View, TRefCountPtr<IPooledRenderTarget>& VelocityRT)
diff --git a/Engine/Source/Runtime/Renderer/Private/DebugViewModeRendering.h b/Engine/Source/Runtime/Renderer/Private/DebugViewModeRendering.h
--- a/Engine/Source/Runtime/Renderer/Private/DebugViewModeRendering.h
+++ b/Engine/Source/Runtime/Renderer/Private/DebugViewModeRendering.h
@@ -79,8 +79,6 @@ protected:
 
 	FDebugViewModeVS(const FMeshMaterialShaderType::CompiledShaderInitializerType& Initializer) : FMeshMaterialShader(Initializer)
 	{
-		IsInstancedStereoParameter.Bind(Initializer.ParameterMap, TEXT("bIsInstancedStereo"));
-		InstancedEyeIndexParameter.Bind(Initializer.ParameterMap, TEXT("InstancedEyeIndex"));
 		PassUniformBuffer.Bind(Initializer.ParameterMap, FSceneTexturesUniformParameters::StaticStructMetadata.GetShaderVariableName());
 	}
 
@@ -93,26 +91,6 @@ public:
 		return AllowDebugViewVSDSHS(Platform) && Material->GetFriendlyName().Contains(TEXT("DebugViewMode"));
 	}
 
-	void SetParameters(FRHICommandList& RHICmdList, const FMaterialRenderProxy* MaterialRenderProxy,const FMaterial& Material,const FSceneView& View,const FDrawingPolicyRenderState& DrawRenderState)
-	{
-		FMeshMaterialShader::SetParameters(RHICmdList, GetVertexShader(),MaterialRenderProxy,Material,View,DrawRenderState.GetViewUniformBuffer(),DrawRenderState.GetPassUniformBuffer());
-
-		if (IsInstancedStereoParameter.IsBound())
-		{
-			SetShaderValue(RHICmdList, GetVertexShader(), IsInstancedStereoParameter, false);
-		}
-
-		if (InstancedEyeIndexParameter.IsBound())
-		{
-			SetShaderValue(RHICmdList, GetVertexShader(), InstancedEyeIndexParameter, 0);
-		}
-	}
-
-	void SetMesh(FRHICommandList& RHICmdList, const FVertexFactory* VertexFactory,const FSceneView& View,const FPrimitiveSceneProxy* Proxy,const FMeshBatchElement& BatchElement,const FDrawingPolicyRenderState& DrawRenderState)
-	{
-		FMeshMaterialShader::SetMesh(RHICmdList, GetVertexShader(),VertexFactory,View,Proxy,BatchElement,DrawRenderState);
-	}
-
 	void GetShaderBindings(
 		const FScene* Scene,
 		ERHIFeatureLevel::Type FeatureLevel,
@@ -125,8 +103,6 @@ public:
 		FMeshDrawSingleShaderBindings& ShaderBindings) const
 	{
 		FMeshMaterialShader::GetShaderBindings(Scene, FeatureLevel, PrimitiveSceneProxy, MaterialRenderProxy, Material, ViewUniformBuffer, PassUniformBufferValue, ShaderElementData, ShaderBindings);
-		ShaderBindings.Add(IsInstancedStereoParameter, false);
-		ShaderBindings.Add(InstancedEyeIndexParameter, 0);
 	}
 
 	static void SetCommonDefinitions(EShaderPlatform Platform, const FMaterial* Material, FShaderCompilerEnvironment& OutEnvironment)
@@ -149,19 +125,6 @@ public:
 		SetCommonDefinitions(Platform, Material, OutEnvironment);
 		FMeshMaterialShader::ModifyCompilationEnvironment(Platform, OutEnvironment);
 	}
-
-	virtual bool Serialize(FArchive& Ar)
-	{
-		const bool Result = FMeshMaterialShader::Serialize(Ar);
-		Ar << IsInstancedStereoParameter;
-		Ar << InstancedEyeIndexParameter;
-		return Result;
-	}
-
-private:
-
-	FShaderParameter IsInstancedStereoParameter;
-	FShaderParameter InstancedEyeIndexParameter;
 };
 
 /**
@@ -271,41 +234,6 @@ public:
 	) const = 0;
 };
 
-/**
- * Interface for debug viewmode pixel shaders. Devired classes can be of global shader or material shaders.
- */ 
-class IDebugViewModePSInterface
-{
-public:
-
-	virtual ~IDebugViewModePSInterface() {}
-
-	virtual void SetParameters(
-		FRHICommandList& RHICmdList, 
-		int32 NumVSInstructions, 
-		int32 NumPSInstructions, 
-		const FMaterialRenderProxy* MaterialRenderProxy,
-		const FMaterial& Material,
-		const FSceneView& View,
-		const FDrawingPolicyRenderState& DrawRenderState
-		) = 0;
-
-	virtual void SetMesh(
-		FRHICommandList& RHICmdList, 
-		const FVertexFactory* VertexFactory,
-		const FSceneView& View,
-		const FPrimitiveSceneProxy* Proxy,
-		int32 VisualizeLODIndex,
-		const FMeshBatchElement& BatchElement, 
-		const FDrawingPolicyRenderState& DrawRenderState
-		) = 0;
-
-	// Used for custom rendering like decals.
-	virtual void SetMesh(FRHICommandList& RHICmdList, const FSceneView& View) = 0;
-
-	virtual FShader* GetShader() = 0;
-};
-
 class FDebugViewModeMeshProcessor : public FMeshPassProcessor
 {
 public:
diff --git a/Engine/Source/Runtime/Renderer/Private/DeferredShadingRenderer.cpp b/Engine/Source/Runtime/Renderer/Private/DeferredShadingRenderer.cpp
--- a/Engine/Source/Runtime/Renderer/Private/DeferredShadingRenderer.cpp
+++ b/Engine/Source/Runtime/Renderer/Private/DeferredShadingRenderer.cpp
@@ -1013,7 +1013,6 @@ void FDeferredShadingSceneRenderer::Render(FRHICommandListImmediate& RHICmdList)
 		GetVirtualTextureSystem()->Update(RHICmdList, FeatureLevel);
 	}
 
-	TGuardValue<bool> LockDrawLists(GDrawListsLocked, true);
 #if !UE_BUILD_SHIPPING
 	if (CVarStallInitViews.GetValueOnRenderThread() > 0.0f)
 	{
@@ -1185,11 +1184,6 @@ void FDeferredShadingSceneRenderer::Render(FRHICommandListImmediate& RHICmdList)
 				InitViewsPossiblyAfterPrepass(RHICmdList, ILCTaskData, UpdateViewCustomDataEvents);
 				PrepareDistanceFieldScene(RHICmdList, false);
 
-				if (!UseMeshDrawCommandPipeline())
-				{
-					PostInitViewCustomData(UpdateViewCustomDataEvents);
-				}
-
 				GEngine->GetPreRenderDelegate().Broadcast();
 
 				{
diff --git a/Engine/Source/Runtime/Renderer/Private/DeferredShadingRenderer.h b/Engine/Source/Runtime/Renderer/Private/DeferredShadingRenderer.h
--- a/Engine/Source/Runtime/Renderer/Private/DeferredShadingRenderer.h
+++ b/Engine/Source/Runtime/Renderer/Private/DeferredShadingRenderer.h
@@ -54,17 +54,11 @@ public:
 	/** Propagates LPVs for all views */
 	void UpdateLPVs(FRHICommandListImmediate& RHICmdList);
 
-	/**
-	 * Renders the dynamic scene's prepass for a particular view
-	 * @return true if anything was rendered
-	 */
-	bool RenderPrePassViewDynamic(FRHICommandList& RHICmdList, const FViewInfo& View, const FDrawingPolicyRenderState& DrawRenderState);
-
 	/**
 	 * Renders the scene's prepass for a particular view
 	 * @return true if anything was rendered
 	 */
-	bool RenderPrePassView(FRHICommandList& RHICmdList, const FViewInfo& View, const FDrawingPolicyRenderState& DrawRenderState);
+	void RenderPrePassView(FRHICommandList& RHICmdList, const FViewInfo& View, const FDrawingPolicyRenderState& DrawRenderState);
 
 	/**
 	 * Renders the scene's prepass for a particular view in parallel
@@ -162,7 +156,7 @@ private:
 	*/
 	bool PreRenderPrePass(FRHICommandListImmediate& RHICmdList);
 
-	void RenderPrePassEditorPrimitives(FRHICommandList& RHICmdList, const FViewInfo& View, const FDrawingPolicyRenderState& DrawRenderState, FDepthDrawingPolicyFactory::ContextType Context);
+	void RenderPrePassEditorPrimitives(FRHICommandList& RHICmdList, const FViewInfo& View, const FDrawingPolicyRenderState& DrawRenderState, EDepthDrawingMode DepthDrawingMode, bool bRespectUseAsOccluderFlag);
 
 	/**
 	 * Renders the scene's prepass and occlusion queries.
@@ -378,9 +372,6 @@ private:
 
 	void SetupVolumetricFog();
 
-	/** Will update the view specific custom data. */
-	void PostInitViewCustomData(FGraphEventArray& OutUpdateEvents);
-
 	void RenderLocalLightsForVolumetricFog(
 		FRDGBuilder& GraphBuilder,
 		FViewInfo& View,
@@ -507,8 +498,6 @@ private:
 	bool DispatchRayTracingWorldUpdates(FRHICommandListImmediate& RHICmdList);
 	FRHIRayTracingPipelineState* BindRayTracingPipeline(FRHICommandList& RHICmdList, const FViewInfo& View, FRayTracingShaderRHIParamRef RayGenShader, FRayTracingShaderRHIParamRef MissShader, FRayTracingShaderRHIParamRef DefaultChsShader);
 #endif // RHI_RAYTRACING
-
-	friend class FTranslucentPrimSet;
 };
 
 DECLARE_CYCLE_STAT_EXTERN(TEXT("PrePass"), STAT_CLM_PrePass, STATGROUP_CommandListMarkers, );
\ No newline at end of file
diff --git a/Engine/Source/Runtime/Renderer/Private/DepthRendering.cpp b/Engine/Source/Runtime/Renderer/Private/DepthRendering.cpp
--- a/Engine/Source/Runtime/Renderer/Private/DepthRendering.cpp
+++ b/Engine/Source/Runtime/Renderer/Private/DepthRendering.cpp
@@ -16,7 +16,6 @@
 #include "MeshMaterialShaderType.h"
 #include "MeshMaterialShader.h"
 #include "SceneRendering.h"
-#include "StaticMeshDrawList.h"
 #include "DeferredShadingRenderer.h"
 #include "ScenePrivate.h"
 #include "OneColorShader.h"
@@ -171,73 +170,6 @@ void GetDepthPassShaders(
 IMPLEMENT_GetDepthPassShaders( true );
 IMPLEMENT_GetDepthPassShaders( false );
 
-FDepthDrawingPolicy::FDepthDrawingPolicy(
-	const FVertexFactory* InVertexFactory,
-	const FMaterialRenderProxy* InMaterialRenderProxy,
-	const FMaterial& InMaterialResource,
-	const FMeshDrawingPolicyOverrideSettings& InOverrideSettings,
-	ERHIFeatureLevel::Type InFeatureLevel,
-	float InMobileColorValue) 
-	: FMeshDrawingPolicy(InVertexFactory, InMaterialRenderProxy, InMaterialResource, InOverrideSettings)
-{
-	const bool bUsesMobileColorValue = (InMobileColorValue != 0.0f);
-	MobileColorValue = InMobileColorValue;
-	
-	bNeedsPixelShader = bUsesMobileColorValue || (!InMaterialResource.WritesEveryPixel() || InMaterialResource.MaterialUsesPixelDepthOffset() || InMaterialResource.IsTranslucencyWritingCustomDepth());
-
-	if (!bNeedsPixelShader)
-	{
-		PixelShader = nullptr;
-	}
-
-	const EMaterialTessellationMode TessellationMode = InMaterialResource.GetTessellationMode();
-	if (RHISupportsTessellation(GShaderPlatformForFeatureLevel[InFeatureLevel])
-		&& InVertexFactory->GetType()->SupportsTessellationShaders() 
-		&& TessellationMode != MTM_NoTessellation)
-	{
-		ShaderPipeline = nullptr;
-		VertexShader = InMaterialResource.GetShader<TDepthOnlyVS<false> >(VertexFactory->GetType());
-		HullShader = InMaterialResource.GetShader<FDepthOnlyHS>(VertexFactory->GetType());
-		DomainShader = InMaterialResource.GetShader<FDepthOnlyDS>(VertexFactory->GetType());
-		if (bNeedsPixelShader)
-		{
-			PixelShader = InMaterialResource.GetShader<FDepthOnlyPS>(InVertexFactory->GetType());
-		}
-	}
-	else
-	{
-		HullShader = nullptr;
-		DomainShader = nullptr;
-		bool bUseShaderPipelines = UseShaderPipelines();
-		if (bNeedsPixelShader)
-		{
-			ShaderPipeline = bUseShaderPipelines ? InMaterialResource.GetShaderPipeline(&DepthPipeline, InVertexFactory->GetType(), false) : nullptr;
-		}
-		else
-		{
-			ShaderPipeline = bUseShaderPipelines ? InMaterialResource.GetShaderPipeline(&DepthNoPixelPipeline, InVertexFactory->GetType(), false) : nullptr;
-		}
-
-		if (ShaderPipeline)
-		{
-			VertexShader = ShaderPipeline->GetShader<TDepthOnlyVS<false> >();
-			if (bNeedsPixelShader)
-			{
-				PixelShader = ShaderPipeline->GetShader<FDepthOnlyPS>();
-			}
-		}
-		else
-		{
-			VertexShader = InMaterialResource.GetShader<TDepthOnlyVS<false> >(VertexFactory->GetType());
-			if (bNeedsPixelShader)
-			{
-				PixelShader = InMaterialResource.GetShader<FDepthOnlyPS>(InVertexFactory->GetType());
-			}
-		}
-	}
-	BaseVertexShader = VertexShader;
-}
-
 void SetDepthPassDitheredLODTransitionState(const FSceneView* SceneView, const FMeshBatch& RESTRICT Mesh, int32 StaticMeshId, FDrawingPolicyRenderState& DrawRenderState)
 {
 	if (SceneView && StaticMeshId >= 0 && Mesh.bDitheredLODTransition)
@@ -270,550 +202,6 @@ void SetDepthPassDitheredLODTransitionState(const FSceneView* SceneView, const F
 	}
 }
 
-void ApplyDitheredLODTransitionStateInternal(FDrawingPolicyRenderState& DrawRenderState, const FViewInfo& ViewInfo, const FStaticMesh& Mesh, const bool InAllowStencilDither)
-{
-	DrawRenderState.SetDitheredLODTransitionAlpha(0.0f);
-	FDepthStencilStateRHIParamRef DepthStencilState = nullptr;
-	uint32 StencilRef = 0;
-
-	if (InAllowStencilDither)
-	{
-		DepthStencilState = TStaticDepthStencilState<>::GetRHI();
-	}
-
-	if (Mesh.bDitheredLODTransition)
-	{
-		if (ViewInfo.StaticMeshFadeOutDitheredLODMap[Mesh.Id])
-		{
-			if (InAllowStencilDither)
-			{
-				DepthStencilState = TStaticDepthStencilState<true, CF_DepthNearOrEqual,
-					true, CF_Equal, SO_Keep, SO_Keep, SO_Keep,
-					false, CF_Always, SO_Keep, SO_Keep, SO_Keep,
-					STENCIL_SANDBOX_MASK, STENCIL_SANDBOX_MASK
-					>::GetRHI();
-				StencilRef = STENCIL_SANDBOX_MASK;
-			}
-			else
-			{
-				DrawRenderState.SetDitheredLODTransitionAlpha(ViewInfo.GetTemporalLODTransition());
-			}
-		}
-		else if (ViewInfo.StaticMeshFadeInDitheredLODMap[Mesh.Id])
-		{
-			if (InAllowStencilDither)
-			{
-				DepthStencilState = TStaticDepthStencilState<true, CF_DepthNearOrEqual,
-					true, CF_Equal, SO_Keep, SO_Keep, SO_Keep,
-					false, CF_Always, SO_Keep, SO_Keep, SO_Keep,
-					STENCIL_SANDBOX_MASK, STENCIL_SANDBOX_MASK
-					>::GetRHI();
-			}
-			else
-			{
-				DrawRenderState.SetDitheredLODTransitionAlpha(ViewInfo.GetTemporalLODTransition() - 1.0f);
-			}
-		}
-	}
-
-	if (DepthStencilState)
-	{
-		DrawRenderState.SetDepthStencilState(DepthStencilState);
-		DrawRenderState.SetStencilRef(StencilRef);
-	}
-}
-
-void FDepthDrawingPolicy::ApplyDitheredLODTransitionState(FDrawingPolicyRenderState& DrawRenderState, const FViewInfo& ViewInfo, const FStaticMesh& Mesh, const bool InAllowStencilDither)
-{
-	ApplyDitheredLODTransitionStateInternal(DrawRenderState, ViewInfo, Mesh, InAllowStencilDither);
-}
-
-void FPositionOnlyDepthDrawingPolicy::ApplyDitheredLODTransitionState(FDrawingPolicyRenderState& DrawRenderState, const FViewInfo& ViewInfo, const FStaticMesh& Mesh, const bool InAllowStencilDither)
-{
-	ApplyDitheredLODTransitionStateInternal(DrawRenderState, ViewInfo, Mesh, InAllowStencilDither);
-}
-
-
-void FDepthDrawingPolicy::SetInstancedEyeIndex(FRHICommandList& RHICmdList, const uint32 EyeIndex) const
-{
-	VertexShader->SetInstancedEyeIndex(RHICmdList, EyeIndex);
-}
-
-void FDepthDrawingPolicy::SetSharedState(FRHICommandList& RHICmdList, const FDrawingPolicyRenderState& DrawRenderState, const FSceneView* View, const FDepthDrawingPolicy::ContextDataType PolicyContext) const
-{
-	// Set the depth-only shader parameters for the material.
-	VertexShader->SetParameters(RHICmdList, MaterialRenderProxy, *MaterialResource, *View, DrawRenderState, PolicyContext.bIsInstancedStereo, PolicyContext.bIsInstancedStereoEmulated);
-	if(HullShader && DomainShader)
-	{
-		HullShader->SetParameters(RHICmdList, MaterialRenderProxy,*View,DrawRenderState.GetViewUniformBuffer(),DrawRenderState.GetPassUniformBuffer());
-		DomainShader->SetParameters(RHICmdList, MaterialRenderProxy,*View,DrawRenderState.GetViewUniformBuffer(),DrawRenderState.GetPassUniformBuffer());
-	}
-
-	if (bNeedsPixelShader)
-	{
-		PixelShader->SetParameters(RHICmdList, MaterialRenderProxy, *MaterialResource, View, DrawRenderState, MobileColorValue);
-	}
-
-	// Set the shared mesh resources.
-	FMeshDrawingPolicy::SetSharedState(RHICmdList, DrawRenderState, View, PolicyContext);
-}
-
-/** 
-* Create bound shader state using the vertex decl from the mesh draw policy
-* as well as the shaders needed to draw the mesh
-* @return new bound shader state object
-*/
-FBoundShaderStateInput FDepthDrawingPolicy::GetBoundShaderStateInput(ERHIFeatureLevel::Type InFeatureLevel) const
-{
-	return FBoundShaderStateInput(
-		FMeshDrawingPolicy::GetVertexDeclaration(), 
-		VertexShader->GetVertexShader(),
-		GETSAFERHISHADER_HULL(HullShader), 
-		GETSAFERHISHADER_DOMAIN(DomainShader),
-		bNeedsPixelShader ? PixelShader->GetPixelShader() : NULL,
-		NULL);
-}
-
-void FDepthDrawingPolicy::SetMeshRenderState(
-	FRHICommandList& RHICmdList, 
-	const FSceneView& View,
-	const FPrimitiveSceneProxy* PrimitiveSceneProxy,
-	const FMeshBatch& Mesh,
-	int32 BatchElementIndex,
-	const FDrawingPolicyRenderState& DrawRenderState,
-	const ElementDataType& ElementData,
-	const ContextDataType PolicyContext
-	) const
-{
-	const FMeshBatchElement& BatchElement = Mesh.Elements[BatchElementIndex];
-	VertexShader->SetMesh(RHICmdList, VertexFactory,View,PrimitiveSceneProxy,BatchElement,DrawRenderState);
-	SetPrimitiveIdStream(RHICmdList, View, PrimitiveSceneProxy, BatchElement.PrimitiveIdMode, BatchElement.DynamicPrimitiveShaderDataIndex);
-	if(HullShader && DomainShader)
-	{
-		HullShader->SetMesh(RHICmdList, VertexFactory,View,PrimitiveSceneProxy,BatchElement,DrawRenderState);
-		DomainShader->SetMesh(RHICmdList, VertexFactory,View,PrimitiveSceneProxy,BatchElement,DrawRenderState);
-	}
-
-	if (bNeedsPixelShader)
-	{
-		PixelShader->SetMesh(RHICmdList, VertexFactory,View,PrimitiveSceneProxy,BatchElement,DrawRenderState);
-	}
-}
-
-int32 CompareDrawingPolicy(const FDepthDrawingPolicy& A,const FDepthDrawingPolicy& B)
-{
-	COMPAREDRAWINGPOLICYMEMBERS(VertexShader);
-	COMPAREDRAWINGPOLICYMEMBERS(HullShader);
-	COMPAREDRAWINGPOLICYMEMBERS(DomainShader);
-	COMPAREDRAWINGPOLICYMEMBERS(bNeedsPixelShader);
-	COMPAREDRAWINGPOLICYMEMBERS(PixelShader);
-	COMPAREDRAWINGPOLICYMEMBERS(VertexFactory);
-	COMPAREDRAWINGPOLICYMEMBERS(MaterialRenderProxy);
-	COMPAREDRAWINGPOLICYMEMBERS(MobileColorValue);
-	return 0;
-}
-
-FPositionOnlyDepthDrawingPolicy::FPositionOnlyDepthDrawingPolicy(
-	const FVertexFactory* InVertexFactory,
-	const FMaterialRenderProxy* InMaterialRenderProxy,
-	const FMaterial& InMaterialResource,
-	const FMeshDrawingPolicyOverrideSettings& InOverrideSettings
-	) 
-	: FMeshDrawingPolicy(InVertexFactory, InMaterialRenderProxy, InMaterialResource, InOverrideSettings)
-{
-	ShaderPipeline = UseShaderPipelines() ? InMaterialResource.GetShaderPipeline(&DepthPosOnlyNoPixelPipeline, VertexFactory->GetType()) : nullptr;
-	VertexShader = ShaderPipeline
-		? ShaderPipeline->GetShader<TDepthOnlyVS<true> >()
-		: InMaterialResource.GetShader<TDepthOnlyVS<true> >(InVertexFactory->GetType());
-	bUsePositionOnlyVS = true;
-	BaseVertexShader = VertexShader;
-}
-
-void FPositionOnlyDepthDrawingPolicy::SetSharedState(FRHICommandList& RHICmdList, const FDrawingPolicyRenderState& DrawRenderState, const FSceneView* View, FPositionOnlyDepthDrawingPolicy::ContextDataType PolicyContext) const
-{
-	// Set the depth-only shader parameters for the material.
-	VertexShader->SetParameters(RHICmdList, MaterialRenderProxy, *MaterialResource, *View, DrawRenderState, PolicyContext.bIsInstancedStereo, PolicyContext.bIsInstancedStereoEmulated);
-
-	// Set the shared mesh resources.
-	VertexFactory->SetPositionStream(RHICmdList);
-}
-
-/** 
-* Create bound shader state using the vertex decl from the mesh draw policy
-* as well as the shaders needed to draw the mesh
-* @return new bound shader state object
-*/
-FBoundShaderStateInput FPositionOnlyDepthDrawingPolicy::GetBoundShaderStateInput(ERHIFeatureLevel::Type InFeatureLevel) const
-{
-	FVertexDeclarationRHIParamRef VertexDeclaration;
-	VertexDeclaration = VertexFactory->GetPositionDeclaration();
-
-	checkSlow(MaterialRenderProxy->GetMaterial(InFeatureLevel)->GetBlendMode() == BLEND_Opaque);
-	return FBoundShaderStateInput(VertexDeclaration, VertexShader->GetVertexShader(), FHullShaderRHIRef(), FDomainShaderRHIRef(), FPixelShaderRHIRef(), FGeometryShaderRHIRef());
-}
-
-void FPositionOnlyDepthDrawingPolicy::SetMeshRenderState(
-	FRHICommandList& RHICmdList, 
-	const FSceneView& View,
-	const FPrimitiveSceneProxy* PrimitiveSceneProxy,
-	const FMeshBatch& Mesh,
-	int32 BatchElementIndex,
-	const FDrawingPolicyRenderState& DrawRenderState,
-	const ElementDataType& ElementData,
-	const ContextDataType PolicyContext
-	) const
-{
-	const FMeshBatchElement& BatchElement = Mesh.Elements[BatchElementIndex];
-	VertexShader->SetMesh(RHICmdList, VertexFactory, View, PrimitiveSceneProxy, BatchElement, DrawRenderState);
-	SetPrimitiveIdStream(RHICmdList, View, PrimitiveSceneProxy, BatchElement.PrimitiveIdMode, BatchElement.DynamicPrimitiveShaderDataIndex);
-}
-
-void FPositionOnlyDepthDrawingPolicy::SetInstancedEyeIndex(FRHICommandList& RHICmdList, const uint32 EyeIndex) const
-{
-	VertexShader->SetInstancedEyeIndex(RHICmdList, EyeIndex);
-}
-
-int32 CompareDrawingPolicy(const FPositionOnlyDepthDrawingPolicy& A, const FPositionOnlyDepthDrawingPolicy& B)
-{
-	COMPAREDRAWINGPOLICYMEMBERS(VertexShader);
-	COMPAREDRAWINGPOLICYMEMBERS(VertexFactory);
-	COMPAREDRAWINGPOLICYMEMBERS(MaterialRenderProxy);
-	return 0;
-}
-
-void FDepthDrawingPolicyFactory::AddStaticMesh(FScene* Scene, FStaticMesh* StaticMesh)
-{
-	const FMaterialRenderProxy* MaterialRenderProxy = StaticMesh->MaterialRenderProxy;
-	const FMaterial* Material = MaterialRenderProxy->GetMaterial(Scene->GetFeatureLevel());
-	const EBlendMode BlendMode = Material->GetBlendMode();
-	const auto FeatureLevel = Scene->GetFeatureLevel();
-
-	FMeshDrawingPolicyOverrideSettings OverrideSettings = ComputeMeshOverrideSettings(*StaticMesh);
-	OverrideSettings.MeshOverrideFlags |= Material->IsTwoSided() ? EDrawingPolicyOverrideFlags::TwoSided : EDrawingPolicyOverrideFlags::None;
-
-	if (!Material->WritesEveryPixel() || Material->MaterialUsesPixelDepthOffset())
-	{	
-		FDepthDrawingPolicy DrawingPolicy(StaticMesh->VertexFactory,
-			MaterialRenderProxy,
-			*Material,
-			OverrideSettings,
-			FeatureLevel,
-			0.0f // MobileColorValue
-			);
-
-		// only draw if required
-		Scene->MaskedDepthDrawList.AddMesh(
-			StaticMesh,
-			FDepthDrawingPolicy::ElementDataType(),
-			DrawingPolicy,
-			FeatureLevel
-			);
-	}
-	else
-	{
-		if (StaticMesh->VertexFactory->SupportsPositionOnlyStream() 
-			&& !Material->MaterialModifiesMeshPosition_RenderThread())
-		{
-			OverrideSettings.MeshOverrideFlags |= Material->IsWireframe() ? EDrawingPolicyOverrideFlags::Wireframe : EDrawingPolicyOverrideFlags::None;
-
-			const FMaterialRenderProxy* DefaultProxy = UMaterial::GetDefaultMaterial(MD_Surface)->GetRenderProxy();
-			FPositionOnlyDepthDrawingPolicy DrawingPolicy(StaticMesh->VertexFactory,
-				DefaultProxy,
-				*DefaultProxy->GetMaterial(Scene->GetFeatureLevel()),
-				OverrideSettings
-				);
-
-			// Add the static mesh to the position-only depth draw list.
-			Scene->PositionOnlyDepthDrawList.AddMesh(
-				StaticMesh,
-				FPositionOnlyDepthDrawingPolicy::ElementDataType(),
-				DrawingPolicy,
-				FeatureLevel
-				);
-		}
-		else
-		{
-			if (!Material->MaterialModifiesMeshPosition_RenderThread())
-			{
-				// Override with the default material for everything but opaque two sided materials
-				MaterialRenderProxy = UMaterial::GetDefaultMaterial(MD_Surface)->GetRenderProxy();
-			}
-
-			FDepthDrawingPolicy DrawingPolicy(StaticMesh->VertexFactory,
-				MaterialRenderProxy,
-				*MaterialRenderProxy->GetMaterial(Scene->GetFeatureLevel()),
-				OverrideSettings,
-				FeatureLevel,
-				0.0f // MobileColorValue
-				);
-
-			// Add the static mesh to the opaque depth-only draw list.
-			Scene->DepthDrawList.AddMesh(
-				StaticMesh,
-				FDepthDrawingPolicy::ElementDataType(),
-				DrawingPolicy,
-				FeatureLevel
-				);
-		}
-	}
-}
-
-bool FDepthDrawingPolicyFactory::DrawMesh(
-	FRHICommandList& RHICmdList, 
-	const FViewInfo& View,
-	ContextType DrawingContext,
-	const FMeshBatch& Mesh,
-	const uint64& BatchElementMask,
-	const FDrawingPolicyRenderState& DrawRenderState,
-	bool bPreFog,
-	const FPrimitiveSceneProxy* PrimitiveSceneProxy,
-	FHitProxyId HitProxyId, 
-	const bool bIsInstancedStereo, 
-	const bool bIsInstancedStereoEmulated
-	)
-{
-	const FMaterialRenderProxy* MaterialRenderProxy = Mesh.MaterialRenderProxy;
-	const FMaterial* Material = MaterialRenderProxy->GetMaterial(View.GetFeatureLevel());
-	bool bDirty = false;
-
-	//Do a per-FMeshBatch check on top of the proxy check in RenderPrePass to handle the case where a proxy that is relevant 
-	//to the depth only pass has to submit multiple FMeshElements but only some of them should be used as occluders.
-	if ((Mesh.bUseAsOccluder || !DrawingContext.bRespectUseAsOccluderFlag || DrawingContext.DepthDrawingMode == DDM_AllOpaque)
-		&& ShouldIncludeDomainInMeshPass(Material->GetMaterialDomain()))
-	{
-		const EBlendMode BlendMode = Material->GetBlendMode();
-		const bool bUsesMobileColorValue = (DrawingContext.MobileColorValue != 0.0f);
-
-		// Check to see if the primitive is currently fading in or out using the screen door effect.  If it is,
-		// then we can't assume the object is opaque as it may be forcibly masked.
-		const FSceneViewState* SceneViewState = static_cast<const FSceneViewState*>( View.State );
-
-		FMeshDrawingPolicyOverrideSettings OverrideSettings = ComputeMeshOverrideSettings(Mesh);
-		OverrideSettings.MeshOverrideFlags |= Material->IsTwoSided() ? EDrawingPolicyOverrideFlags::TwoSided : EDrawingPolicyOverrideFlags::None;
-
-		if ( BlendMode == BLEND_Opaque 
-			&& Mesh.VertexFactory->SupportsPositionOnlyStream() 
-			&& !Material->MaterialModifiesMeshPosition_RenderThread()
-			&& Material->WritesEveryPixel()
-			&& !bUsesMobileColorValue
-			)
-		{
-			//render opaque primitives that support a separate position-only vertex buffer
-			const FMaterialRenderProxy* DefaultProxy = UMaterial::GetDefaultMaterial(MD_Surface)->GetRenderProxy();
-
-			OverrideSettings.MeshOverrideFlags |= Material->IsWireframe() ? EDrawingPolicyOverrideFlags::Wireframe : EDrawingPolicyOverrideFlags::None;
-
-			FPositionOnlyDepthDrawingPolicy DrawingPolicy(
-				Mesh.VertexFactory, 
-				DefaultProxy, 
-				*DefaultProxy->GetMaterial(View.GetFeatureLevel()), 
-				OverrideSettings
-				);
-
-			FDrawingPolicyRenderState DrawRenderStateLocal(DrawRenderState);
-			DrawingPolicy.SetupPipelineState(DrawRenderStateLocal, View);
-			CommitGraphicsPipelineState(RHICmdList, DrawingPolicy, DrawRenderStateLocal, DrawingPolicy.GetBoundShaderStateInput(View.GetFeatureLevel()), DrawingPolicy.GetMaterialRenderProxy());
-			DrawingPolicy.SetSharedState(RHICmdList, DrawRenderStateLocal, &View, FPositionOnlyDepthDrawingPolicy::ContextDataType(bIsInstancedStereo, bIsInstancedStereoEmulated));
-
-			int32 BatchElementIndex = 0;
-			uint64 Mask = BatchElementMask;
-			do
-			{
-				if(Mask & 1)
-				{
-					// We draw instanced static meshes twice when rendering with instanced stereo. Once for each eye.
-					const bool bIsInstancedMesh = Mesh.Elements[BatchElementIndex].bIsInstancedMesh;
-					const uint32 InstancedStereoDrawCount = (bIsInstancedStereo && bIsInstancedMesh) ? 2 : 1;
-					for (uint32 DrawCountIter = 0; DrawCountIter < InstancedStereoDrawCount; ++DrawCountIter)
-					{
-						DrawingPolicy.SetInstancedEyeIndex(RHICmdList, DrawCountIter);
-
-						TDrawEvent<FRHICommandList> MeshEvent;
-						BeginMeshDrawEvent(RHICmdList, PrimitiveSceneProxy, Mesh, MeshEvent, EnumHasAnyFlags(EShowMaterialDrawEventTypes(GShowMaterialDrawEventTypes), EShowMaterialDrawEventTypes::DepthPositionOnly));
-
-						DrawingPolicy.SetMeshRenderState(RHICmdList, View, PrimitiveSceneProxy, Mesh, BatchElementIndex, DrawRenderStateLocal, FPositionOnlyDepthDrawingPolicy::ElementDataType(), FPositionOnlyDepthDrawingPolicy::ContextDataType());
-						DrawingPolicy.DrawMesh(RHICmdList, View, Mesh, BatchElementIndex, bIsInstancedStereo);
-					}
-				}
-				Mask >>= 1;
-				BatchElementIndex++;
-			} while(Mask);
-
-			bDirty = true;
-		}
-		else if (!IsTranslucentBlendMode(BlendMode) || Material->IsTranslucencyWritingCustomDepth())
-		{
-			const bool bMaterialMasked = !Material->WritesEveryPixel() || Material->IsTranslucencyWritingCustomDepth();
-
-			bool bDraw = true;
-
-			switch(DrawingContext.DepthDrawingMode)
-			{
-			case DDM_AllOpaque:
-				break;
-			case DDM_AllOccluders: 
-				break;
-			case DDM_NonMaskedOnly: 
-				bDraw = !bMaterialMasked;
-				break;
-			default:
-				check(!"Unrecognized DepthDrawingMode");
-			}
-
-			if(bDraw)
-			{
-				if (!bMaterialMasked && !Material->MaterialModifiesMeshPosition_RenderThread())
-				{
-					// Override with the default material for opaque materials that are not two sided
-					MaterialRenderProxy = UMaterial::GetDefaultMaterial(MD_Surface)->GetRenderProxy();
-				}
-
-				FDepthDrawingPolicy DrawingPolicy(
-					Mesh.VertexFactory, 
-					MaterialRenderProxy, 
-					*MaterialRenderProxy->GetMaterial(View.GetFeatureLevel()), 
-					OverrideSettings,
-					View.GetFeatureLevel(),
-					DrawingContext.MobileColorValue
-					);
-
-				FDrawingPolicyRenderState DrawRenderStateLocal(DrawRenderState);
-				DrawingPolicy.SetupPipelineState(DrawRenderStateLocal, View);
-				CommitGraphicsPipelineState(RHICmdList, DrawingPolicy, DrawRenderStateLocal, DrawingPolicy.GetBoundShaderStateInput(View.GetFeatureLevel()), DrawingPolicy.GetMaterialRenderProxy());
-				DrawingPolicy.SetSharedState(RHICmdList, DrawRenderStateLocal, &View, FDepthDrawingPolicy::ContextDataType(bIsInstancedStereo, bIsInstancedStereoEmulated));
-
-				int32 BatchElementIndex = 0;
-				uint64 Mask = BatchElementMask;
-				do
-				{
-					if(Mask & 1)
-					{
-						// We draw instanced static meshes twice when rendering with instanced stereo. Once for each eye.
-						const bool bIsInstancedMesh = Mesh.Elements[BatchElementIndex].bIsInstancedMesh;
-						const uint32 InstancedStereoDrawCount = (bIsInstancedStereo && bIsInstancedMesh) ? 2 : 1;
-						for (uint32 DrawCountIter = 0; DrawCountIter < InstancedStereoDrawCount; ++DrawCountIter)
-						{
-							DrawingPolicy.SetInstancedEyeIndex(RHICmdList, DrawCountIter);
-
-							TDrawEvent<FRHICommandList> MeshEvent;
-							BeginMeshDrawEvent(RHICmdList, PrimitiveSceneProxy, Mesh, MeshEvent, EnumHasAnyFlags(EShowMaterialDrawEventTypes(GShowMaterialDrawEventTypes), EShowMaterialDrawEventTypes::Depth));
-
-							DrawingPolicy.SetMeshRenderState(RHICmdList, View, PrimitiveSceneProxy, Mesh, BatchElementIndex, DrawRenderStateLocal, FMeshDrawingPolicy::ElementDataType(), FDepthDrawingPolicy::ContextDataType());
-							DrawingPolicy.DrawMesh(RHICmdList, View, Mesh, BatchElementIndex, bIsInstancedStereo);
-						}
-					}
-					Mask >>= 1;
-					BatchElementIndex++;
-				} while(Mask);
-
-				bDirty = true;
-			}
-		}
-	}
-
-	return bDirty;
-}
-
-bool FDepthDrawingPolicyFactory::DrawDynamicMesh(
-	FRHICommandList& RHICmdList, 
-	const FViewInfo& View,
-	ContextType DrawingContext,
-	const FMeshBatch& Mesh,
-	bool bPreFog,
-	const FDrawingPolicyRenderState& DrawRenderState,
-	const FPrimitiveSceneProxy* PrimitiveSceneProxy,
-	FHitProxyId HitProxyId, 
-	const bool bIsInstancedStereo, 
-	const bool bIsInstancedStereoEmulated
-	)
-{
-	return DrawMesh(
-		RHICmdList, 
-		View,
-		DrawingContext,
-		Mesh,
-		Mesh.Elements.Num()==1 ? 1 : (1<<Mesh.Elements.Num())-1,	// 1 bit set for each mesh element
-		DrawRenderState,
-		bPreFog,
-		PrimitiveSceneProxy,
-		HitProxyId, 
-		bIsInstancedStereo, 
-		bIsInstancedStereoEmulated
-		);
-}
-
-bool FDepthDrawingPolicyFactory::DrawStaticMesh(
-	FRHICommandList& RHICmdList, 
-	const FViewInfo& View,
-	ContextType DrawingContext,
-	const FStaticMesh& StaticMesh,
-	const uint64& BatchElementMask,
-	bool bPreFog,
-	const FDrawingPolicyRenderState& DrawRenderState,
-	const FPrimitiveSceneProxy* PrimitiveSceneProxy,
-	FHitProxyId HitProxyId, 
-	const bool bIsInstancedStereo,
-	const bool bIsInstancedStereoEmulated
-	)
-{
-	bool bDirty = false;
-
-	const FMaterial* Material = StaticMesh.MaterialRenderProxy->GetMaterial(View.GetFeatureLevel());
-	const EMaterialShadingModel ShadingModel = Material->GetShadingModel();
-	bDirty |= DrawMesh(
-		RHICmdList, 
-		View,
-		DrawingContext,
-		StaticMesh,
-		BatchElementMask,
-		DrawRenderState,
-		bPreFog,
-		PrimitiveSceneProxy,
-		HitProxyId, 
-		bIsInstancedStereo, 
-		bIsInstancedStereoEmulated
-		);
-
-	return bDirty;
-}
-
-bool FDeferredShadingSceneRenderer::RenderPrePassViewDynamic(FRHICommandList& RHICmdList, const FViewInfo& View, const FDrawingPolicyRenderState& DrawRenderState)
-{
-	FDepthDrawingPolicyFactory::ContextType Context(EarlyZPassMode, true);
-
-	for (int32 MeshBatchIndex = 0; MeshBatchIndex < View.DynamicMeshElements.Num(); MeshBatchIndex++)
-	{
-		const FMeshBatchAndRelevance& MeshBatchAndRelevance = View.DynamicMeshElements[MeshBatchIndex];
-
-		if (MeshBatchAndRelevance.GetHasOpaqueOrMaskedMaterial() && MeshBatchAndRelevance.GetRenderInMainPass())
-		{
-			const FMeshBatch& MeshBatch = *MeshBatchAndRelevance.Mesh;
-			const FPrimitiveSceneProxy* PrimitiveSceneProxy = MeshBatchAndRelevance.PrimitiveSceneProxy;
-			bool bShouldUseAsOccluder = true;
-
-			if (EarlyZPassMode < DDM_AllOpaque)
-			{
-				extern float GMinScreenRadiusForDepthPrepass;
-				//@todo - move these proxy properties into FMeshBatchAndRelevance so we don't have to dereference the proxy in order to reject a mesh
-				const float LODFactorDistanceSquared = (PrimitiveSceneProxy->GetBounds().Origin - View.ViewMatrices.GetViewOrigin()).SizeSquared() * FMath::Square(View.LODDistanceFactor);
-
-				// Only render primitives marked as occluders
-				bShouldUseAsOccluder = PrimitiveSceneProxy->ShouldUseAsOccluder()
-					// Only render static objects unless movable are requested
-					&& (!PrimitiveSceneProxy->IsMovable() || bEarlyZPassMovable)
-					&& (FMath::Square(PrimitiveSceneProxy->GetBounds().SphereRadius) > GMinScreenRadiusForDepthPrepass * GMinScreenRadiusForDepthPrepass * LODFactorDistanceSquared);
-			}
-
-			if (bShouldUseAsOccluder)
-			{
-				FDepthDrawingPolicyFactory::DrawDynamicMesh(RHICmdList, View, Context, MeshBatch, true, DrawRenderState, PrimitiveSceneProxy, MeshBatch.BatchHitProxyId, View.IsInstancedStereoPass());
-			}
-		}
-	}
-
-	return true;
-}
-
 static void SetupPrePassView(FRHICommandList& RHICmdList, const FViewInfo& View, const FSceneRenderer* SceneRenderer, const bool bIsEditorPrimitivePass = false)
 {
 	RHICmdList.SetScissorRect(false, 0, 0, 0, 0);
@@ -863,96 +251,15 @@ static void RenderHiddenAreaMaskView(FRHICommandList& RHICmdList, FGraphicsPipel
 	}
 }
 
-bool FDeferredShadingSceneRenderer::RenderPrePassView(FRHICommandList& RHICmdList, const FViewInfo& View, const FDrawingPolicyRenderState& DrawRenderState)
+void FDeferredShadingSceneRenderer::RenderPrePassView(FRHICommandList& RHICmdList, const FViewInfo& View, const FDrawingPolicyRenderState& DrawRenderState)
 {
-	bool bDirty = false;
-
 	SetupPrePassView(RHICmdList, View, this);
 
-	if (UseMeshDrawCommandPipeline())
-	{
-		View.ParallelMeshDrawCommandPasses[EMeshPass::DepthPass].DispatchDraw(nullptr, RHICmdList);
-
-		return true;
-	}
-
-	// Draw the static occluder primitives using a depth drawing policy.
-	{
-		// Draw opaque occluders which support a separate position-only
-		// vertex buffer to minimize vertex fetch bandwidth, which is
-		// often the bottleneck during the depth only pass.
-		SCOPED_DRAW_EVENT(RHICmdList, PosOnlyOpaque);
-		bDirty |= Scene->PositionOnlyDepthDrawList.DrawVisible(RHICmdList, View, DrawRenderState, View.StaticMeshOccluderMap, View.StaticMeshBatchVisibility);
-	}
-	{
-		// Draw opaque occluders, using double speed z where supported.
-		SCOPED_DRAW_EVENT(RHICmdList, Opaque);
-		bDirty |= Scene->DepthDrawList.DrawVisible(RHICmdList, View, DrawRenderState, View.StaticMeshOccluderMap, View.StaticMeshBatchVisibility);
-	}
-
-	if (EarlyZPassMode >= DDM_AllOccluders)
-	{
-		// Draw opaque occluders with masked materials
-		SCOPED_DRAW_EVENT(RHICmdList, Masked);
-		bDirty |= Scene->MaskedDepthDrawList.DrawVisible(RHICmdList, View, DrawRenderState, View.StaticMeshOccluderMap, View.StaticMeshBatchVisibility);
-	}
-
-	{
-		SCOPED_DRAW_EVENT(RHICmdList, Dynamic);
-		bDirty |= RenderPrePassViewDynamic(RHICmdList, View, DrawRenderState);
-	}
-
-	return bDirty;
+	View.ParallelMeshDrawCommandPasses[EMeshPass::DepthPass].DispatchDraw(nullptr, RHICmdList);
 }
 
-class FRenderPrepassDynamicDataThreadTask : public FRenderTask
-{
-	FDeferredShadingSceneRenderer& ThisRenderer;
-	FRHICommandList& RHICmdList;
-	const FViewInfo& View;
-	FDrawingPolicyRenderState DrawRenderState;
-
-public:
-
-	FRenderPrepassDynamicDataThreadTask(
-		FDeferredShadingSceneRenderer& InThisRenderer,
-		FRHICommandList& InRHICmdList,
-		const FViewInfo& InView,
-		const FDrawingPolicyRenderState& InDrawRenderState
-		)
-		: ThisRenderer(InThisRenderer)
-		, RHICmdList(InRHICmdList)
-		, View(InView)
-		, DrawRenderState(InDrawRenderState)
-	{
-	}
-
-	FORCEINLINE TStatId GetStatId() const
-	{
-		RETURN_QUICK_DECLARE_CYCLE_STAT(FRenderPrepassDynamicDataThreadTask, STATGROUP_TaskGraphTasks);
-	}
-
-	static ESubsequentsMode::Type GetSubsequentsMode() { return ESubsequentsMode::TrackSubsequents; }
-
-	void DoTask(ENamedThreads::Type CurrentThread, const FGraphEventRef& MyCompletionGraphEvent)
-	{
-		check(RHICmdList.IsInsideRenderPass());
-		ThisRenderer.RenderPrePassViewDynamic(RHICmdList, View, DrawRenderState);
-		RHICmdList.EndRenderPass();
-		RHICmdList.HandleRTThreadTaskCompletion(MyCompletionGraphEvent);
-	}
-};
-
 DECLARE_CYCLE_STAT(TEXT("Prepass"), STAT_CLP_Prepass, STATGROUP_ParallelCommandListMarkers);
 
-static int32 GStartPrepassParallelTranslatesImmediately = 0;
-static FAutoConsoleVariableRef CVarStartPrepassParallelTranslatesImmediately(
-	TEXT("r.StartPrepassParallelTranslatesImmediately"),
-	GStartPrepassParallelTranslatesImmediately,
-	TEXT("Experimental option, allows the parallel translate tasks for the prepass to start immediately. Makes the most sense with r.DoInitViewsLightingAfterPrepass = 1."),
-	ECVF_RenderThreadSafe
-);
-
 class FPrePassParallelCommandListSet : public FParallelCommandListSet
 {
 public:
@@ -988,78 +295,13 @@ bool FDeferredShadingSceneRenderer::RenderPrePassViewParallel(const FViewInfo& V
 			CVarRHICmdFlushRenderThreadTasksPrePass.GetValueOnRenderThread() == 0 && CVarRHICmdFlushRenderThreadTasks.GetValueOnRenderThread() == 0,
 			DrawRenderState);
 
-		if (UseMeshDrawCommandPipeline())
-		{
-			View.ParallelMeshDrawCommandPasses[EMeshPass::DepthPass].DispatchDraw(&ParallelCommandListSet, ParentCmdList);
-
-			if (bDoPrePre)
-			{
-				AfterTasksAreStarted();
-				bDepthWasCleared = PreRenderPrePass(ParentCmdList);
-			}
-		}
-		else
-		{
-			// Draw the static occluder primitives using a depth drawing policy.
-			// Draw opaque occluders which support a separate position-only
-			// vertex buffer to minimize vertex fetch bandwidth, which is
-			// often the bottleneck during the depth only pass.
-			Scene->PositionOnlyDepthDrawList.DrawVisibleParallel(View.StaticMeshOccluderMap, View.StaticMeshBatchVisibility, ParallelCommandListSet);
-
-			// Draw opaque occluders, using double speed z where supported.
-			Scene->DepthDrawList.DrawVisibleParallel(View.StaticMeshOccluderMap, View.StaticMeshBatchVisibility, ParallelCommandListSet);
-
-			// Draw opaque occluders with masked materials
-			if (EarlyZPassMode >= DDM_AllOccluders)
-			{
-				Scene->MaskedDepthDrawList.DrawVisibleParallel(View.StaticMeshOccluderMap, View.StaticMeshBatchVisibility, ParallelCommandListSet);
-			}
-
-			if (!GStartPrepassParallelTranslatesImmediately)
-			{
-				// we do this step here (awkwardly) so that the above tasks can be in flight while we get the particles (which must be dynamic) setup.
-				if (bDoPrePre)
-				{
-					AfterTasksAreStarted();
-					bDepthWasCleared = PreRenderPrePass(ParentCmdList);
-				}
-				// Dynamic
-				FRHICommandList* CmdList = ParallelCommandListSet.NewParallelCommandList();
-
-				FGraphEventRef AnyThreadCompletionEvent = TGraphTask<FRenderPrepassDynamicDataThreadTask>::CreateTask(ParallelCommandListSet.GetPrereqs(), ENamedThreads::GetRenderThread())
-					.ConstructAndDispatchWhenReady(*this, *CmdList, View, ParallelCommandListSet.DrawRenderState);
-
-				ParallelCommandListSet.AddParallelCommandList(CmdList, AnyThreadCompletionEvent);
-			}
-			else
-			{
-				if (bDoPrePre)
-				{
-					bDepthWasCleared = PreRenderPrePass(ParentCmdList);
-				}
-			}
-		}
-	}
+		View.ParallelMeshDrawCommandPasses[EMeshPass::DepthPass].DispatchDraw(&ParallelCommandListSet, ParentCmdList);
 
-	if (GStartPrepassParallelTranslatesImmediately && !UseMeshDrawCommandPipeline())
-	{
-		// we do this step here (awkwardly) so that the above tasks can be in flight while we get the particles (which must be dynamic) setup.
 		if (bDoPrePre)
 		{
 			AfterTasksAreStarted();
+			bDepthWasCleared = PreRenderPrePass(ParentCmdList);
 		}
-		FPrePassParallelCommandListSet ParallelCommandListSet(View, this, ParentCmdList,
-			CVarRHICmdPrePassDeferredContexts.GetValueOnRenderThread() > 0,
-			CVarRHICmdFlushRenderThreadTasksPrePass.GetValueOnRenderThread() == 0 && CVarRHICmdFlushRenderThreadTasks.GetValueOnRenderThread() == 0,
-			DrawRenderState);
-
-		// Dynamic
-		FRHICommandList* CmdList = ParallelCommandListSet.NewParallelCommandList();
-
-		FGraphEventRef AnyThreadCompletionEvent = TGraphTask<FRenderPrepassDynamicDataThreadTask>::CreateTask(ParallelCommandListSet.GetPrereqs(), ENamedThreads::GetRenderThread())
-			.ConstructAndDispatchWhenReady(*this, *CmdList, View, ParallelCommandListSet.DrawRenderState);
-
-		ParallelCommandListSet.AddParallelCommandList(CmdList, AnyThreadCompletionEvent);
 	}
 
 	return bDepthWasCleared;
@@ -1175,7 +417,7 @@ bool FDeferredShadingSceneRenderer::PreRenderPrePass(FRHICommandListImmediate& R
 	return bDepthWasCleared;
 }
 
-void FDeferredShadingSceneRenderer::RenderPrePassEditorPrimitives(FRHICommandList& RHICmdList, const FViewInfo& View, const FDrawingPolicyRenderState& DrawRenderState, FDepthDrawingPolicyFactory::ContextType Context) 
+void FDeferredShadingSceneRenderer::RenderPrePassEditorPrimitives(FRHICommandList& RHICmdList, const FViewInfo& View, const FDrawingPolicyRenderState& DrawRenderState, EDepthDrawingMode DepthDrawingMode, bool bRespectUseAsOccluderFlag) 
 {
 	SetupPrePassView(RHICmdList, View, this, true);
 
@@ -1188,66 +430,50 @@ void FDeferredShadingSceneRenderer::RenderPrePassEditorPrimitives(FRHICommandLis
 		const bool bNeedToSwitchVerticalAxis = RHINeedsToSwitchVerticalAxis(ShaderPlatform);
 		const FScene* LocalScene = Scene;
 
-		if (UseMeshDrawCommandPipeline())
-		{
-			DrawDynamicMeshPass(View, RHICmdList,
-				[&View, &DrawRenderState, LocalScene, &Context](FDynamicPassMeshDrawListContext* DynamicMeshPassContext)
-				{
-					FDepthPassMeshProcessor PassMeshProcessor(
-						LocalScene,
-						&View,
-						DrawRenderState,
-						Context.bRespectUseAsOccluderFlag,
-						Context.DepthDrawingMode,
-						false,
-						DynamicMeshPassContext);
-
-					const uint64 DefaultBatchElementMask = ~0ull;
+		DrawDynamicMeshPass(View, RHICmdList,
+			[&View, &DrawRenderState, LocalScene, DepthDrawingMode, bRespectUseAsOccluderFlag](FDynamicPassMeshDrawListContext* DynamicMeshPassContext)
+			{
+				FDepthPassMeshProcessor PassMeshProcessor(
+					LocalScene,
+					&View,
+					DrawRenderState,
+					bRespectUseAsOccluderFlag,
+					DepthDrawingMode,
+					false,
+					DynamicMeshPassContext);
+
+				const uint64 DefaultBatchElementMask = ~0ull;
 					
-					for (int32 MeshIndex = 0; MeshIndex < View.ViewMeshElements.Num(); MeshIndex++)
-					{
-						const FMeshBatch& MeshBatch = View.ViewMeshElements[MeshIndex];
-						PassMeshProcessor.AddMeshBatch(MeshBatch, DefaultBatchElementMask, nullptr);
-					}
-				});
-		}
-		else
-		{
-			// Draw the base pass for the view's batched mesh elements.
-			bDirty |= DrawViewElements<FDepthDrawingPolicyFactory>(RHICmdList, View, DrawRenderState, Context, SDPG_World, true) || bDirty;
-		}
+				for (int32 MeshIndex = 0; MeshIndex < View.ViewMeshElements.Num(); MeshIndex++)
+				{
+					const FMeshBatch& MeshBatch = View.ViewMeshElements[MeshIndex];
+					PassMeshProcessor.AddMeshBatch(MeshBatch, DefaultBatchElementMask, nullptr);
+				}
+			});
 
 		// Draw the view's batched simple elements(lines, sprites, etc).
 		bDirty |= View.BatchedViewElements.Draw(RHICmdList, DrawRenderState, FeatureLevel, bNeedToSwitchVerticalAxis, View, false) || bDirty;
 
-		if (UseMeshDrawCommandPipeline())
-		{
-			DrawDynamicMeshPass(View, RHICmdList,
-				[&View, &DrawRenderState, LocalScene, &Context](FDynamicPassMeshDrawListContext* DynamicMeshPassContext)
-				{
-					FDepthPassMeshProcessor PassMeshProcessor(
-						LocalScene,
-						&View,
-						DrawRenderState,
-						Context.bRespectUseAsOccluderFlag,
-						Context.DepthDrawingMode,
-						false,
-						DynamicMeshPassContext);
-
-					const uint64 DefaultBatchElementMask = ~0ull;
+		DrawDynamicMeshPass(View, RHICmdList,
+			[&View, &DrawRenderState, LocalScene, DepthDrawingMode, bRespectUseAsOccluderFlag](FDynamicPassMeshDrawListContext* DynamicMeshPassContext)
+			{
+				FDepthPassMeshProcessor PassMeshProcessor(
+					LocalScene,
+					&View,
+					DrawRenderState,
+					bRespectUseAsOccluderFlag,
+					DepthDrawingMode,
+					false,
+					DynamicMeshPassContext);
+
+				const uint64 DefaultBatchElementMask = ~0ull;
 					
-					for (int32 MeshIndex = 0; MeshIndex < View.TopViewMeshElements.Num(); MeshIndex++)
-					{
-						const FMeshBatch& MeshBatch = View.TopViewMeshElements[MeshIndex];
-						PassMeshProcessor.AddMeshBatch(MeshBatch, DefaultBatchElementMask, nullptr);
-					}
-				});
-		}
-		else
-		{
-			// Draw foreground objects last
-			bDirty |= DrawViewElements<FDepthDrawingPolicyFactory>(RHICmdList, View, DrawRenderState, Context, SDPG_Foreground, true) || bDirty;
-		}
+				for (int32 MeshIndex = 0; MeshIndex < View.TopViewMeshElements.Num(); MeshIndex++)
+				{
+					const FMeshBatch& MeshBatch = View.TopViewMeshElements[MeshIndex];
+					PassMeshProcessor.AddMeshBatch(MeshBatch, DefaultBatchElementMask, nullptr);
+				}
+			});
 
 		// Draw the view's batched simple elements(lines, sprites, etc).
 		bDirty |= View.TopBatchedViewElements.Draw(RHICmdList, DrawRenderState, FeatureLevel, bNeedToSwitchVerticalAxis, View, false) || bDirty;
@@ -1298,7 +524,6 @@ bool FDeferredShadingSceneRenderer::RenderPrePass(FRHICommandListImmediate& RHIC
 	SCOPE_CYCLE_COUNTER(STAT_DepthDrawTime);
 	SCOPED_GPU_STAT(RHICmdList, Prepass);
 
-	bool bDirty = false;
 	bool bDidPrePre = false;
 	FSceneRenderTargets& SceneContext = FSceneRenderTargets::Get(RHICmdList);
 
@@ -1347,12 +572,11 @@ bool FDeferredShadingSceneRenderer::RenderPrePass(FRHICommandListImmediate& RHIC
 				{
 					check(RHICmdList.IsOutsideRenderPass());
 					bDepthWasCleared = RenderPrePassViewParallel(View, RHICmdList, DrawRenderState, AfterTasksAreStarted, !bDidPrePre) || bDepthWasCleared;
-					bDirty = true; // assume dirty since we are not going to wait
 					bDidPrePre = true;
 				}
 				else
 				{
-					bDirty |= RenderPrePassView(RHICmdList, View, DrawRenderState);
+					RenderPrePassView(RHICmdList, View, DrawRenderState);
 				}
 			}
 
@@ -1361,7 +585,7 @@ bool FDeferredShadingSceneRenderer::RenderPrePass(FRHICommandListImmediate& RHIC
 			{
 				SceneContext.BeginRenderingPrePass(RHICmdList, false);
 			}
-			RenderPrePassEditorPrimitives(RHICmdList, View, DrawRenderState, FDepthDrawingPolicyFactory::ContextType(EarlyZPassMode, true));
+			RenderPrePassEditorPrimitives(RHICmdList, View, DrawRenderState, EarlyZPassMode, true);
 			if (bParallel)
 			{
 				RHICmdList.EndRenderPass();
diff --git a/Engine/Source/Runtime/Renderer/Private/DepthRendering.h b/Engine/Source/Runtime/Renderer/Private/DepthRendering.h
--- a/Engine/Source/Runtime/Renderer/Private/DepthRendering.h
+++ b/Engine/Source/Runtime/Renderer/Private/DepthRendering.h
@@ -9,11 +9,9 @@
 #include "CoreMinimal.h"
 #include "RHI.h"
 #include "HitProxies.h"
-#include "DrawingPolicy.h"
 #include "ShaderBaseClasses.h"
 #include "MeshPassProcessor.h"
 
-class FDepthOnlyPS;
 class FPrimitiveSceneProxy;
 class FScene;
 class FStaticMesh;
@@ -33,224 +31,6 @@ enum EDepthDrawingMode
 
 extern const TCHAR* GetDepthDrawingModeString(EDepthDrawingMode Mode);
 
-template<bool>
-class TDepthOnlyVS;
-
-class FDepthOnlyPS;
-
-class FViewInfo;
-
-/**
- * Used to write out depth for opaque and masked materials during the depth-only pass.
- */
-class FDepthDrawingPolicy : public FMeshDrawingPolicy
-{
-public:
-
-	struct ContextDataType : public FMeshDrawingPolicy::ContextDataType
-	{
-		explicit ContextDataType( const bool InbIsInstancedStereo) : FMeshDrawingPolicy::ContextDataType(InbIsInstancedStereo), bIsInstancedStereoEmulated(false) {};
-		ContextDataType(const bool InbIsInstancedStereo, const bool InbIsInstancedStereoEmulated) : FMeshDrawingPolicy::ContextDataType(InbIsInstancedStereo), bIsInstancedStereoEmulated(InbIsInstancedStereoEmulated) {};
-		ContextDataType() : bIsInstancedStereoEmulated(false) {};
-		
-		bool bIsInstancedStereoEmulated;
-	};
-
-	FDepthDrawingPolicy(
-		const FVertexFactory* InVertexFactory,
-		const FMaterialRenderProxy* InMaterialRenderProxy,
-		const FMaterial& InMaterialResource,
-		const FMeshDrawingPolicyOverrideSettings& InOverrideSettings,
-		ERHIFeatureLevel::Type InFeatureLevel,
-		float MobileColorValue
-		);
-
-	// FMeshDrawingPolicy interface.
-
-	void ApplyDitheredLODTransitionState(FDrawingPolicyRenderState& DrawRenderState, const FViewInfo& ViewInfo, const FStaticMesh& Mesh, const bool InAllowStencilDither);
-
-	FDrawingPolicyMatchResult Matches(const FDepthDrawingPolicy& Other, bool bForReals = false) const
-	{
-		DRAWING_POLICY_MATCH_BEGIN
-		 	DRAWING_POLICY_MATCH(FMeshDrawingPolicy::Matches(Other, bForReals)) &&
-			DRAWING_POLICY_MATCH(bNeedsPixelShader == Other.bNeedsPixelShader) &&
-			DRAWING_POLICY_MATCH(VertexShader == Other.VertexShader) &&
-			DRAWING_POLICY_MATCH(PixelShader == Other.PixelShader) &&
-			DRAWING_POLICY_MATCH(MobileColorValue == Other.MobileColorValue);
-		DRAWING_POLICY_MATCH_END
-	}
-
-	void SetInstancedEyeIndex(FRHICommandList& RHICmdList, const uint32 EyeIndex) const;
-
-	void SetSharedState(FRHICommandList& RHICmdList, const FDrawingPolicyRenderState& DrawRenderState, const FSceneView* View, const FDepthDrawingPolicy::ContextDataType PolicyContext) const;
-
-	/** 
-	* Create bound shader state using the vertex decl from the mesh draw policy
-	* as well as the shaders needed to draw the mesh
-	* @param DynamicStride - optional stride for dynamic vertex data
-	* @return new bound shader state object
-	*/
-	FBoundShaderStateInput GetBoundShaderStateInput(ERHIFeatureLevel::Type InFeatureLevel) const;
-
-	void SetMeshRenderState(
-		FRHICommandList& RHICmdList, 
-		const FSceneView& View,
-		const FPrimitiveSceneProxy* PrimitiveSceneProxy,
-		const FMeshBatch& Mesh,
-		int32 BatchElementIndex,
-		const FDrawingPolicyRenderState& DrawRenderState,
-		const ElementDataType& ElementData,
-		const ContextDataType PolicyContext
-		) const;
-
-	friend int32 CompareDrawingPolicy(const FDepthDrawingPolicy& A,const FDepthDrawingPolicy& B);
-
-private:
-	bool bNeedsPixelShader;
-	class FDepthOnlyHS *HullShader;
-	class FDepthOnlyDS *DomainShader;
-
-	FShaderPipeline* ShaderPipeline;
-	TDepthOnlyVS<false>* VertexShader;
-	FDepthOnlyPS* PixelShader;
-	float MobileColorValue;
-};
-
-/**
- * Writes out depth for opaque materials on meshes which support a position-only vertex buffer.
- * Using the position-only vertex buffer saves vertex fetch bandwidth during the z prepass.
- */
-class FPositionOnlyDepthDrawingPolicy : public FMeshDrawingPolicy
-{
-public:
-
-	struct ContextDataType : public FMeshDrawingPolicy::ContextDataType
-	{
-		explicit ContextDataType(const bool InbIsInstancedStereo) : FMeshDrawingPolicy::ContextDataType(InbIsInstancedStereo), bIsInstancedStereoEmulated(false) {};
-		ContextDataType(const bool InbIsInstancedStereo, const bool InbIsInstancedStereoEmulated) : FMeshDrawingPolicy::ContextDataType(InbIsInstancedStereo), bIsInstancedStereoEmulated(InbIsInstancedStereoEmulated) {};
-		ContextDataType() : bIsInstancedStereoEmulated(false) {};
-		bool bIsInstancedStereoEmulated;
-	};
-
-	FPositionOnlyDepthDrawingPolicy(
-		const FVertexFactory* InVertexFactory,
-		const FMaterialRenderProxy* InMaterialRenderProxy,
-		const FMaterial& InMaterialResource,
-		const FMeshDrawingPolicyOverrideSettings& InOverrideSettings
-		);
-
-	// FMeshDrawingPolicy interface.
-
-	void ApplyDitheredLODTransitionState(FDrawingPolicyRenderState& DrawRenderState, const FViewInfo& ViewInfo, const FStaticMesh& Mesh, const bool InAllowStencilDither);
-
-	FDrawingPolicyMatchResult Matches(const FPositionOnlyDepthDrawingPolicy& Other, bool bForReals = false) const
-	{
-		DRAWING_POLICY_MATCH_BEGIN
-			DRAWING_POLICY_MATCH(FMeshDrawingPolicy::Matches(Other, bForReals)) &&
-			DRAWING_POLICY_MATCH(VertexShader == Other.VertexShader);
-		DRAWING_POLICY_MATCH_END
-	}
-
-	void SetSharedState(FRHICommandList& RHICmdList, const FDrawingPolicyRenderState& DrawRenderState, const FSceneView* View, const FPositionOnlyDepthDrawingPolicy::ContextDataType PolicyContext) const;
-
-	/**
-	* Create bound shader state using the vertex decl from the mesh draw policy
-	* as well as the shaders needed to draw the mesh
-	* @return new bound shader state object
-	*/
-	FBoundShaderStateInput GetBoundShaderStateInput(ERHIFeatureLevel::Type InFeatureLevel) const;
-
-	void SetMeshRenderState(
-		FRHICommandList& RHICmdList, 
-		const FSceneView& View,
-		const FPrimitiveSceneProxy* PrimitiveSceneProxy,
-		const FMeshBatch& Mesh,
-		int32 BatchElementIndex,
-		const FDrawingPolicyRenderState& DrawRenderState,
-		const ElementDataType& ElementData,
-		const ContextDataType PolicyContext
-		) const;
-
-	void SetInstancedEyeIndex(FRHICommandList& RHICmdList, const uint32 EyeIndex) const;
-
-	friend int32 CompareDrawingPolicy(const FPositionOnlyDepthDrawingPolicy& A,const FPositionOnlyDepthDrawingPolicy& B);
-
-private:
-	FShaderPipeline* ShaderPipeline;
-	TDepthOnlyVS<true> * VertexShader;
-};
-
-/**
- * A drawing policy factory for the depth drawing policy.
- */
-class FDepthDrawingPolicyFactory
-{
-public:
-
-	enum { bAllowSimpleElements = false };
-	struct ContextType
-	{
-		EDepthDrawingMode DepthDrawingMode;
-		float MobileColorValue;
-
-		/** Uses of FDepthDrawingPolicyFactory that are not the depth pre-pass will not want the bUseAsOccluder flag to interfere. */
-		bool bRespectUseAsOccluderFlag;
-
-		ContextType(EDepthDrawingMode InDepthDrawingMode, bool bInRespectUseAsOccluderFlag) :
-			DepthDrawingMode(InDepthDrawingMode),
-			MobileColorValue(0.0f),
-			bRespectUseAsOccluderFlag(bInRespectUseAsOccluderFlag)
-		{}
-	};
-
-	static void AddStaticMesh(FScene* Scene,FStaticMesh* StaticMesh);
-	static bool DrawDynamicMesh(
-		FRHICommandList& RHICmdList,
-		const FViewInfo& View,
-		ContextType DrawingContext,
-		const FMeshBatch& Mesh,
-		bool bPreFog,
-		const FDrawingPolicyRenderState& DrawRenderState,
-		const FPrimitiveSceneProxy* PrimitiveSceneProxy,
-		FHitProxyId HitProxyId,
-		const bool bIsInstancedStereo = false,
-		const bool bIsInstancedStereoEmulated = false
-		);
-
-	static bool DrawStaticMesh(
-		FRHICommandList& RHICmdList, 
-		const FViewInfo& View,
-		ContextType DrawingContext,
-		const FStaticMesh& StaticMesh,
-		const uint64& BatchElementMask,
-		bool bPreFog,
-		const FDrawingPolicyRenderState& DrawRenderState,
-		const FPrimitiveSceneProxy* PrimitiveSceneProxy,
-		FHitProxyId HitProxyId, 
-		const bool bIsInstancedStereo = false,
-		const bool bIsInstancedStereoEmulated = false
-		);
-
-private:
-	/**
-	* Render a dynamic or static mesh using a depth draw policy
-	* @return true if the mesh rendered
-	*/
-	static bool DrawMesh(
-		FRHICommandList& RHICmdList, 
-		const FViewInfo& View,
-		ContextType DrawingContext,
-		const FMeshBatch& Mesh,
-		const uint64& BatchElementMask,
-		const FDrawingPolicyRenderState& DrawRenderState,
-		bool bPreFog,
-		const FPrimitiveSceneProxy* PrimitiveSceneProxy,
-		FHitProxyId HitProxyId, 
-		const bool bIsInstancedStereo = false, 
-		const bool bIsInstancedStereoEmulated = false
-		);
-};
-
 class FDepthOnlyShaderElementData : public FMeshMaterialShaderElementData
 {
 public:
@@ -276,18 +56,9 @@ protected:
 	TDepthOnlyVS(const FMeshMaterialShaderType::CompiledShaderInitializerType& Initializer) :
 		FMeshMaterialShader(Initializer)
 	{
-		InstancedEyeIndexParameter.Bind(Initializer.ParameterMap, TEXT("InstancedEyeIndex"));
-		IsInstancedStereoParameter.Bind(Initializer.ParameterMap, TEXT("bIsInstancedStereo"));
-		IsInstancedStereoEmulatedParameter.Bind(Initializer.ParameterMap, TEXT("bIsInstancedStereoEmulated"));
 		BindSceneTextureUniformBufferDependentOnShadingPath(Initializer, PassUniformBuffer, PassUniformBuffer);
 	}
 
-private:
-
-	FShaderParameter InstancedEyeIndexParameter;
-	FShaderParameter IsInstancedStereoParameter;
-	FShaderParameter IsInstancedStereoEmulatedParameter;
-
 public:
 
 	static bool ShouldCompilePermutation(EShaderPlatform Platform,const FMaterial* Material,const FVertexFactoryType* VertexFactoryType)
@@ -302,43 +73,6 @@ public:
 		return (Material->IsSpecialEngineMaterial() || !Material->WritesEveryPixel() || Material->MaterialMayModifyMeshPosition() || Material->IsTranslucencyWritingCustomDepth());
 	}
 
-	virtual bool Serialize(FArchive& Ar) override
-	{
-		const bool result = FMeshMaterialShader::Serialize(Ar);
-		Ar << InstancedEyeIndexParameter;
-		Ar << IsInstancedStereoParameter;
-		Ar << IsInstancedStereoEmulatedParameter;
-		return result;
-	}
-
-	void SetParameters(
-		FRHICommandList& RHICmdList, 
-		const FMaterialRenderProxy* MaterialRenderProxy, 
-		const FMaterial& MaterialResource, 
-		const FSceneView& View, 
-		const FDrawingPolicyRenderState& DrawRenderState,
-		const bool bIsInstancedStereo, 
-		const bool bIsInstancedStereoEmulated
-		)
-	{
-		FMeshMaterialShader::SetParameters(RHICmdList, GetVertexShader(),MaterialRenderProxy,MaterialResource,View,DrawRenderState.GetViewUniformBuffer(),DrawRenderState.GetPassUniformBuffer());
-		
-		if (IsInstancedStereoParameter.IsBound())
-		{
-			SetShaderValue(RHICmdList, GetVertexShader(), IsInstancedStereoParameter, bIsInstancedStereo);
-		}
-
-		if (IsInstancedStereoEmulatedParameter.IsBound())
-		{
-			SetShaderValue(RHICmdList, GetVertexShader(), IsInstancedStereoEmulatedParameter, bIsInstancedStereoEmulated);
-		}
-
-		if (InstancedEyeIndexParameter.IsBound())
-		{
-			SetShaderValue(RHICmdList, GetVertexShader(), InstancedEyeIndexParameter, 0);
-		}
-	}
-
 	void GetShaderBindings(
 		const FScene* Scene,
 		ERHIFeatureLevel::Type FeatureLevel,
@@ -351,21 +85,6 @@ public:
 		FMeshDrawSingleShaderBindings& ShaderBindings) const
 	{
 		FMeshMaterialShader::GetShaderBindings(Scene, FeatureLevel, PrimitiveSceneProxy, MaterialRenderProxy, Material, ViewUniformBuffer, PassUniformBufferValue, ShaderElementData, ShaderBindings);
-
-		ensure(!IsInstancedStereoParameter.IsBound() && !IsInstancedStereoEmulatedParameter.IsBound() && !InstancedEyeIndexParameter.IsBound());
-	}
-
-	void SetMesh(FRHICommandList& RHICmdList, const FVertexFactory* VertexFactory,const FSceneView& View,const FPrimitiveSceneProxy* Proxy,const FMeshBatchElement& BatchElement,const FDrawingPolicyRenderState& DrawRenderState)
-	{
-		FMeshMaterialShader::SetMesh(RHICmdList, GetVertexShader(),VertexFactory,View,Proxy,BatchElement,DrawRenderState);
-	}
-
-	void SetInstancedEyeIndex(FRHICommandList& RHICmdList, const uint32 EyeIndex)
-	{
-		if (EyeIndex > 0 && InstancedEyeIndexParameter.IsBound())
-		{
-			SetShaderValue(RHICmdList, GetVertexShader(), InstancedEyeIndexParameter, EyeIndex);
-		}
 	}
 };
 
@@ -437,18 +156,6 @@ public:
 
 	FDepthOnlyPS() {}
 
-	void SetParameters(FRHICommandList& RHICmdList, const FMaterialRenderProxy* MaterialRenderProxy,const FMaterial& MaterialResource,const FSceneView* View, const FDrawingPolicyRenderState& DrawRenderState, float InMobileColorValue)
-	{
-		FMeshMaterialShader::SetParameters(RHICmdList, GetPixelShader(),MaterialRenderProxy,MaterialResource,*View,DrawRenderState.GetViewUniformBuffer(),DrawRenderState.GetPassUniformBuffer());
-
-		SetShaderValue(RHICmdList, GetPixelShader(), MobileColorValue, InMobileColorValue);
-	}
-
-	void SetMesh(FRHICommandList& RHICmdList, const FVertexFactory* VertexFactory,const FSceneView& View,const FPrimitiveSceneProxy* Proxy,const FMeshBatchElement& BatchElement,const FDrawingPolicyRenderState& DrawRenderState)
-	{
-		FMeshMaterialShader::SetMesh(RHICmdList, GetPixelShader(),VertexFactory,View,Proxy,BatchElement,DrawRenderState);
-	}
-
 	void GetShaderBindings(
 		const FScene* Scene,
 		ERHIFeatureLevel::Type FeatureLevel,
