Removed legacy Drawing Policies and Static Mesh Draw Lists.  These are now replaced by FMeshDrawCommand / FMeshPassProcessor everywhere in the renderer.

diff --git a/Engine/Source/Runtime/Renderer/Private/DistortionRendering.cpp b/Engine/Source/Runtime/Renderer/Private/DistortionRendering.cpp
--- a/Engine/Source/Runtime/Renderer/Private/DistortionRendering.cpp
+++ b/Engine/Source/Runtime/Renderer/Private/DistortionRendering.cpp
@@ -14,7 +14,6 @@
 #include "PostProcess/SceneRenderTargets.h"
 #include "GlobalShader.h"
 #include "MaterialShaderType.h"
-#include "DrawingPolicy.h"
 #include "MeshMaterialShader.h"
 #include "ShaderBaseClasses.h"
 #include "SceneRendering.h"
@@ -279,18 +278,6 @@ protected:
 	{
 		return Material && IsTranslucentBlendMode(Material->GetBlendMode()) && Material->IsDistorted();
 	}
-
-public:
-	
-	void SetParameters(FRHICommandList& RHICmdList, const FVertexFactory* VertexFactory,const FMaterialRenderProxy* MaterialRenderProxy,const FSceneView* View, const FDrawingPolicyRenderState& DrawRenderState)
-	{
-		FMeshMaterialShader::SetParameters(RHICmdList, GetVertexShader(), MaterialRenderProxy, *MaterialRenderProxy->GetMaterial(View->GetFeatureLevel()), *View, DrawRenderState.GetViewUniformBuffer(), DrawRenderState.GetPassUniformBuffer());
-	}
-
-	void SetMesh(FRHICommandList& RHICmdList, const FVertexFactory* VertexFactory,const FSceneView& View,const FPrimitiveSceneProxy* Proxy,const FMeshBatchElement& BatchElement,const FDrawingPolicyRenderState& DrawRenderState)
-	{
-		FMeshMaterialShader::SetMesh(RHICmdList, GetVertexShader(),VertexFactory,View,Proxy,BatchElement,DrawRenderState);
-	}
 };
 
 
@@ -379,21 +366,6 @@ public:
 	{
 	}
 
-	void SetParameters(
-		FRHICommandList& RHICmdList,
-		const FMaterialRenderProxy* MaterialRenderProxy,
-		const FSceneView& View,
-		const FDrawingPolicyRenderState& DrawRenderState
-		)
-	{
-		FMeshMaterialShader::SetParameters(RHICmdList, GetPixelShader(), MaterialRenderProxy, *MaterialRenderProxy->GetMaterial(View.GetFeatureLevel()), View, DrawRenderState.GetViewUniformBuffer(), DrawRenderState.GetPassUniformBuffer());
-	}
-
-	void SetMesh(FRHICommandList& RHICmdList, const FVertexFactory* VertexFactory,const FSceneView& View,const FPrimitiveSceneProxy* Proxy,const FMeshBatchElement& BatchElement,const FDrawingPolicyRenderState& DrawRenderState)
-	{
-		FMeshMaterialShader::SetMesh(RHICmdList, GetPixelShader(),VertexFactory,View,Proxy,BatchElement,DrawRenderState);
-	}
-
 	virtual bool Serialize(FArchive& Ar) override
 	{
 		bool bShaderHasOutdatedParameters = FMeshMaterialShader::Serialize(Ar);
@@ -404,478 +376,6 @@ public:
 //** distortion accumulate pixel shader type implementation */
 IMPLEMENT_MATERIAL_SHADER_TYPE(, FDistortionMeshPS,TEXT("/Engine/Private/DistortAccumulatePS.usf"),TEXT("Main"),SF_Pixel);
 
-/*-----------------------------------------------------------------------------
-TDistortionMeshDrawingPolicy
------------------------------------------------------------------------------*/
-
-/**
-* Distortion mesh drawing policy
-*/
-class FDistortionMeshDrawingPolicy : public FMeshDrawingPolicy
-{
-public:
-	/** context type */
-	typedef FMeshDrawingPolicy::ElementDataType ElementDataType;
-
-	/**
-	* Constructor
-	* @param InIndexBuffer - index buffer for rendering
-	* @param InVertexFactory - vertex factory for rendering
-	* @param InMaterialRenderProxy - material instance for rendering
-	* @param bInOverrideWithShaderComplexity - whether to override with shader complexity
-	*/
-	FDistortionMeshDrawingPolicy(
-		const FVertexFactory* InVertexFactory,
-		const FMaterialRenderProxy* InMaterialRenderProxy,
-		const FMaterial& MaterialResouce,
-		bool bInitializeOffsets,
-		const FMeshDrawingPolicyOverrideSettings& InOverrideSettings,
-		ERHIFeatureLevel::Type InFeatureLevel
-		);
-
-	// FMeshDrawingPolicy interface.
-
-	/**
-	* Match two draw policies
-	* @param Other - draw policy to compare
-	* @return true if the draw policies are a match
-	*/
-	FDrawingPolicyMatchResult Matches(const FDistortionMeshDrawingPolicy& Other, bool bForReals = false) const;
-
-	/**
-	* Executes the draw commands which can be shared between any meshes using this drawer.
-	* @param CI - The command interface to execute the draw commands on.
-	* @param View - The view of the scene being drawn.
-	*/
-	void SetSharedState(FRHICommandList& RHICmdList, const FDrawingPolicyRenderState& DrawRenderState, const FViewInfo* View, const ContextDataType PolicyContext) const;
-
-	/** 
-	* Create bound shader state using the vertex decl from the mesh draw policy
-	* as well as the shaders needed to draw the mesh
-	* @return new bound shader state object
-	*/
-	FBoundShaderStateInput GetBoundShaderStateInput(ERHIFeatureLevel::Type InFeatureLevel) const;
-
-	/**
-	* Sets the render states for drawing a mesh.
-	* @param PrimitiveSceneProxy - The primitive drawing the dynamic mesh.  If this is a view element, this will be NULL.
-	* @param Mesh - mesh element with data needed for rendering
-	* @param ElementData - context specific data for mesh rendering
-	*/
-	void SetMeshRenderState(
-		FRHICommandList& RHICmdList,
-		const FViewInfo& View,
-		const FPrimitiveSceneProxy* PrimitiveSceneProxy,
-		const FMeshBatch& Mesh,
-		int32 BatchElementIndex,
-		const FDrawingPolicyRenderState& DrawRenderState,
-		const ElementDataType& ElementData,
-		const ContextDataType PolicyContext
-		) const;
-
-private:
-	/** vertex shader based on policy type */
-	FDistortionMeshVS* VertexShader;
-
-	FDistortionMeshHS* HullShader;
-	FDistortionMeshDS* DomainShader;
-
-	/** whether we are initializing offsets or accumulating them */
-	bool bInitializeOffsets;
-	/** pixel shader based on policy type */
-	FDistortionMeshPS* DistortPixelShader;
-	/** pixel shader used to initialize offsets */
-//later	FShaderComplexityAccumulatePixelShader* InitializePixelShader;
-};
-
-/**
-* Constructor
-* @param InIndexBuffer - index buffer for rendering
-* @param InVertexFactory - vertex factory for rendering
-* @param InMaterialRenderProxy - material instance for rendering
-*/
-FDistortionMeshDrawingPolicy::FDistortionMeshDrawingPolicy(
-	const FVertexFactory* InVertexFactory,
-	const FMaterialRenderProxy* InMaterialRenderProxy,
-	const FMaterial& InMaterialResource,
-	bool bInInitializeOffsets,
-	const FMeshDrawingPolicyOverrideSettings& InOverrideSettings,
-	ERHIFeatureLevel::Type InFeatureLevel
-	)
-:	FMeshDrawingPolicy(InVertexFactory,InMaterialRenderProxy,InMaterialResource,InOverrideSettings)
-,	bInitializeOffsets(bInInitializeOffsets)
-{
-	HullShader = NULL;
-	DomainShader = NULL;
-
-	const EMaterialTessellationMode MaterialTessellationMode = MaterialResource->GetTessellationMode();
-	if (RHISupportsTessellation(GShaderPlatformForFeatureLevel[InFeatureLevel])
-		&& InVertexFactory->GetType()->SupportsTessellationShaders()
-		&& MaterialTessellationMode != MTM_NoTessellation)
-	{
-		HullShader = InMaterialResource.GetShader<FDistortionMeshHS>(VertexFactory->GetType());
-		DomainShader = InMaterialResource.GetShader<FDistortionMeshDS>(VertexFactory->GetType());
-	}
-
-	VertexShader = InMaterialResource.GetShader<FDistortionMeshVS>(InVertexFactory->GetType());
-
-	if (bInitializeOffsets)
-	{
-//later		InitializePixelShader = GetGlobalShaderMap(View.ShaderMap)->GetShader<FShaderComplexityAccumulatePixelShader>();
-		DistortPixelShader = NULL;
-	}
-	else
-	{
-		DistortPixelShader = InMaterialResource.GetShader<FDistortionMeshPS>(InVertexFactory->GetType());
-//later		InitializePixelShader = NULL;
-	}
-	BaseVertexShader = VertexShader;
-}
-
-/**
-* Match two draw policies
-* @param Other - draw policy to compare
-* @return true if the draw policies are a match
-*/
-FDrawingPolicyMatchResult FDistortionMeshDrawingPolicy::Matches(
-	const FDistortionMeshDrawingPolicy& Other, bool bForReals
-	) const
-{
-	DRAWING_POLICY_MATCH_BEGIN
-		DRAWING_POLICY_MATCH(FMeshDrawingPolicy::Matches(Other, bForReals)) &&
-		DRAWING_POLICY_MATCH(VertexShader == Other.VertexShader) &&
-		DRAWING_POLICY_MATCH(HullShader == Other.HullShader) &&
-		DRAWING_POLICY_MATCH(DomainShader == Other.DomainShader) &&
-		DRAWING_POLICY_MATCH(bInitializeOffsets == Other.bInitializeOffsets) &&
-		DRAWING_POLICY_MATCH(DistortPixelShader == Other.DistortPixelShader); //&&
-	//later	InitializePixelShader == Other.InitializePixelShader;
-	DRAWING_POLICY_MATCH_END
-}
-
-/**
-* Executes the draw commands which can be shared between any meshes using this drawer.
-* @param CI - The command interface to execute the draw commands on.
-* @param View - The view of the scene being drawn.
-*/
-void FDistortionMeshDrawingPolicy::SetSharedState(
-	FRHICommandList& RHICmdList,
-	const FDrawingPolicyRenderState& DrawRenderState,
-	const FViewInfo* View,
-	const ContextDataType PolicyContext
-	) const
-{
-	// Set shared mesh resources
-	FMeshDrawingPolicy::SetSharedState(RHICmdList, DrawRenderState, View, PolicyContext);
-	// Set the translucent shader parameters for the material instance
-	VertexShader->SetParameters(RHICmdList, VertexFactory, MaterialRenderProxy, View, DrawRenderState);
-
-	if(HullShader && DomainShader)
-	{
-		HullShader->SetParameters(RHICmdList, MaterialRenderProxy, *View, DrawRenderState.GetViewUniformBuffer(), DrawRenderState.GetPassUniformBuffer());
-		DomainShader->SetParameters(RHICmdList, MaterialRenderProxy, *View, DrawRenderState.GetViewUniformBuffer(), DrawRenderState.GetPassUniformBuffer());
-	}
-	
-	if (!bInitializeOffsets)
-	{
-		DistortPixelShader->SetParameters(RHICmdList, MaterialRenderProxy, *View, DrawRenderState);
-	}
-}
-
-/**
-* Create bound shader state using the vertex decl from the mesh draw policy
-* as well as the shaders needed to draw the mesh
-* @return new bound shader state object
-*/
-FBoundShaderStateInput FDistortionMeshDrawingPolicy::GetBoundShaderStateInput(ERHIFeatureLevel::Type InFeatureLevel) const
-{
-	FPixelShaderRHIParamRef PixelShaderRHIRef = NULL;
-
-	if (bInitializeOffsets)
-	{
-//later		PixelShaderRHIRef = InitializePixelShader->GetPixelShader();
-	}
-	else
-	{
-		PixelShaderRHIRef = DistortPixelShader->GetPixelShader();
-	}
-
-	return FBoundShaderStateInput(
-		FMeshDrawingPolicy::GetVertexDeclaration(),
-		VertexShader->GetVertexShader(),
-		GETSAFERHISHADER_HULL(HullShader),
-		GETSAFERHISHADER_DOMAIN(DomainShader),
-		PixelShaderRHIRef,
-		FGeometryShaderRHIRef());
-}
-
-/**
-* Sets the render states for drawing a mesh.
-* @param PrimitiveSceneProxy - The primitive drawing the dynamic mesh.  If this is a view element, this will be NULL.
-* @param Mesh - mesh element with data needed for rendering
-* @param ElementData - context specific data for mesh rendering
-*/
-void FDistortionMeshDrawingPolicy::SetMeshRenderState(
-	FRHICommandList& RHICmdList,
-	const FViewInfo& View,
-	const FPrimitiveSceneProxy* PrimitiveSceneProxy,
-	const FMeshBatch& Mesh,
-	int32 BatchElementIndex,
-	const FDrawingPolicyRenderState& DrawRenderState,
-	const ElementDataType& ElementData,
-	const ContextDataType PolicyContext
-	) const
-{
-
-	const FMeshBatchElement& BatchElement = Mesh.Elements[BatchElementIndex];
-	// Set transforms
-	VertexShader->SetMesh(RHICmdList, VertexFactory,View,PrimitiveSceneProxy,BatchElement,DrawRenderState);
-	SetPrimitiveIdStream(RHICmdList, View, PrimitiveSceneProxy, BatchElement.PrimitiveIdMode, BatchElement.DynamicPrimitiveShaderDataIndex);
-
-	if(HullShader && DomainShader)
-	{
-		HullShader->SetMesh(RHICmdList, VertexFactory,View,PrimitiveSceneProxy,BatchElement,DrawRenderState);
-		DomainShader->SetMesh(RHICmdList, VertexFactory,View,PrimitiveSceneProxy,BatchElement,DrawRenderState);
-	}
-
-	// Don't set pixel shader constants if we are overriding with the shader complexity pixel shader
-	if (!bInitializeOffsets)
-	{
-		DistortPixelShader->SetMesh(RHICmdList, VertexFactory, View, PrimitiveSceneProxy, BatchElement, DrawRenderState);
-	}
-}
-
-/*-----------------------------------------------------------------------------
-TDistortionMeshDrawingPolicyFactory
------------------------------------------------------------------------------*/
-
-/**
-* Distortion mesh draw policy factory.
-* Creates the policies needed for rendering a mesh based on its material
-*/
-class TDistortionMeshDrawingPolicyFactory
-{
-public:
-	enum { bAllowSimpleElements = false };
-	typedef bool ContextType;
-
-	/**
-	* Render a dynamic mesh using a distortion mesh draw policy
-	* @return true if the mesh rendered
-	*/
-	static bool DrawDynamicMesh(
-		FRHICommandList& RHICmdList,
-		const FViewInfo& View,
-		ContextType DrawingContext,
-		const FMeshBatch& Mesh,
-		bool bPreFog,
-		const FDrawingPolicyRenderState& DrawRenderState,
-		const FPrimitiveSceneProxy* PrimitiveSceneProxy,
-		FHitProxyId HitProxyId
-		);
-
-	/**
-	* Render a dynamic mesh using a distortion mesh draw policy
-	* @return true if the mesh rendered
-	*/
-	static bool DrawStaticMesh(
-		FRHICommandList& RHICmdList,
-		const FViewInfo* View,
-		ContextType DrawingContext,
-		const FStaticMesh& StaticMesh,
-		uint64 BatchElementMask,
-		const FDrawingPolicyRenderState& DrawRenderState,
-		const FPrimitiveSceneProxy* PrimitiveSceneProxy,
-		FHitProxyId HitProxyId
-		);
-};
-
-/**
-* Render a dynamic mesh using a distortion mesh draw policy
-* @return true if the mesh rendered
-*/
-bool TDistortionMeshDrawingPolicyFactory::DrawDynamicMesh(
-	FRHICommandList& RHICmdList,
-	const FViewInfo& View,
-	ContextType bInitializeOffsets,
-	const FMeshBatch& Mesh,
-	bool bPreFog,
-	const FDrawingPolicyRenderState& DrawRenderState,
-	const FPrimitiveSceneProxy* PrimitiveSceneProxy,
-	FHitProxyId HitProxyId
-	)
-{
-	const auto FeatureLevel = View.GetFeatureLevel();
-	bool bDistorted = Mesh.MaterialRenderProxy && Mesh.MaterialRenderProxy->GetMaterial(FeatureLevel)->IsDistorted() && ShouldIncludeDomainInMeshPass(Mesh.MaterialRenderProxy->GetMaterial(FeatureLevel)->GetMaterialDomain());
-
-	//reconstruct bBackface from the View
-	const bool bBackFace = XOR(View.bReverseCulling, !!(DrawRenderState.GetViewOverrideFlags() & EDrawingPolicyOverrideFlags::ReverseCullMode));
-	if(bDistorted && !bBackFace)
-	{
-		// draw dynamic mesh element using distortion mesh policy
-		FDistortionMeshDrawingPolicy DrawingPolicy(
-			Mesh.VertexFactory,
-			Mesh.MaterialRenderProxy,
-			*Mesh.MaterialRenderProxy->GetMaterial(FeatureLevel),
-			bInitializeOffsets,
-			ComputeMeshOverrideSettings(Mesh),
-			FeatureLevel
-			);
-
-		FDrawingPolicyRenderState DrawRenderStateLocal(DrawRenderState);
-		DrawRenderStateLocal.SetDitheredLODTransitionAlpha(Mesh.DitheredLODTransitionAlpha);
-		DrawingPolicy.SetupPipelineState(DrawRenderStateLocal, View);
-		CommitGraphicsPipelineState(RHICmdList, DrawingPolicy, DrawRenderStateLocal, DrawingPolicy.GetBoundShaderStateInput(View.GetFeatureLevel()), DrawingPolicy.GetMaterialRenderProxy());
-		DrawingPolicy.SetSharedState(RHICmdList, DrawRenderStateLocal, &View, typename FDistortionMeshDrawingPolicy::ContextDataType());
-
-		for (int32 BatchElementIndex = 0; BatchElementIndex < Mesh.Elements.Num(); BatchElementIndex++)
-		{
-			TDrawEvent<FRHICommandList> MeshEvent;
-			BeginMeshDrawEvent(RHICmdList, PrimitiveSceneProxy, Mesh, MeshEvent, EnumHasAnyFlags(EShowMaterialDrawEventTypes(GShowMaterialDrawEventTypes), EShowMaterialDrawEventTypes::DistortionDynamic));
-
-			DrawingPolicy.SetMeshRenderState(RHICmdList, View,PrimitiveSceneProxy,Mesh,BatchElementIndex, DrawRenderStateLocal,typename FDistortionMeshDrawingPolicy::ElementDataType(), typename FDistortionMeshDrawingPolicy::ContextDataType());
-			DrawingPolicy.DrawMesh(RHICmdList,View,Mesh,BatchElementIndex);
-		}
-
-		return true;
-	}
-	else
-	{
-		return false;
-	}
-}
-
-/**
-* Render a dynamic mesh using a distortion mesh draw policy
-* @return true if the mesh rendered
-*/
-bool TDistortionMeshDrawingPolicyFactory::DrawStaticMesh(
-	FRHICommandList& RHICmdList,
-	const FViewInfo* View,
-	ContextType bInitializeOffsets,
-	const FStaticMesh& StaticMesh,
-	uint64 BatchElementMask,
-	const FDrawingPolicyRenderState& DrawRenderState,
-	const FPrimitiveSceneProxy* PrimitiveSceneProxy,
-	FHitProxyId HitProxyId
-	)
-{
-	const auto FeatureLevel = View->GetFeatureLevel();
-	bool bDistorted = StaticMesh.MaterialRenderProxy && StaticMesh.MaterialRenderProxy->GetMaterial(FeatureLevel)->IsDistorted();
-
-	const bool bBackFace = XOR(View->bReverseCulling, !!(DrawRenderState.GetViewOverrideFlags() & EDrawingPolicyOverrideFlags::ReverseCullMode));
-	if(bDistorted && !bBackFace)
-	{
-		// draw static mesh element using distortion mesh policy
-		FDistortionMeshDrawingPolicy DrawingPolicy(
-			StaticMesh.VertexFactory,
-			StaticMesh.MaterialRenderProxy,
-			*StaticMesh.MaterialRenderProxy->GetMaterial(FeatureLevel),
-			bInitializeOffsets,
-			ComputeMeshOverrideSettings(StaticMesh),
-			FeatureLevel
-			);
-
-		FDrawingPolicyRenderState DrawRenderStateLocal(DrawRenderState);
-		DrawingPolicy.ApplyDitheredLODTransitionState(DrawRenderStateLocal, *View, StaticMesh, false);
-		DrawingPolicy.SetupPipelineState(DrawRenderStateLocal, *View);
-		CommitGraphicsPipelineState(RHICmdList, DrawingPolicy, DrawRenderStateLocal, DrawingPolicy.GetBoundShaderStateInput(View->GetFeatureLevel()), DrawingPolicy.GetMaterialRenderProxy());
-		DrawingPolicy.SetSharedState(RHICmdList, DrawRenderStateLocal, View, typename FDistortionMeshDrawingPolicy::ContextDataType());
-		int32 BatchElementIndex = 0;
-		do
-		{
-			if(BatchElementMask & 1)
-			{
-				TDrawEvent<FRHICommandList> MeshEvent;
-				BeginMeshDrawEvent(RHICmdList, PrimitiveSceneProxy, StaticMesh, MeshEvent, EnumHasAnyFlags(EShowMaterialDrawEventTypes(GShowMaterialDrawEventTypes), EShowMaterialDrawEventTypes::DistortionStatic));
-
-
-				DrawingPolicy.SetMeshRenderState(RHICmdList, *View,PrimitiveSceneProxy,StaticMesh,BatchElementIndex,DrawRenderStateLocal,
-					typename FDistortionMeshDrawingPolicy::ElementDataType(),
-					typename FDistortionMeshDrawingPolicy::ContextDataType()
-					);
-				DrawingPolicy.DrawMesh(RHICmdList, *View, StaticMesh, BatchElementIndex);
-			}
-			BatchElementMask >>= 1;
-			BatchElementIndex++;
-		} while(BatchElementMask);
-
-		return true;
-	}
-	else
-	{
-		return false;
-	}
-}
-
-/*-----------------------------------------------------------------------------
-	FDistortionPrimSet
------------------------------------------------------------------------------*/
-
-bool FDistortionPrimSet::DrawAccumulatedOffsets(FRHICommandListImmediate& RHICmdList, const FViewInfo& View, const FDrawingPolicyRenderState& DrawRenderState, bool bInitializeOffsets)
-{
-	bool bDirty = false;
-
-	QUICK_SCOPE_CYCLE_COUNTER(STAT_FDistortionPrimSet_DrawAccumulatedOffsets);
-
-	if( Prims.Num() )
-	{
-		QUICK_SCOPE_CYCLE_COUNTER(STAT_FDistortionPrimSet_DrawAccumulatedOffsets_Prims);
-
-		// Draw scene prims
-		for (int32 PrimIdx = 0; PrimIdx < Prims.Num(); PrimIdx++)
-		{
-			FPrimitiveSceneProxy* PrimitiveSceneProxy = Prims[PrimIdx];
-			const FPrimitiveViewRelevance& ViewRelevance = View.PrimitiveViewRelevanceMap[PrimitiveSceneProxy->GetPrimitiveSceneInfo()->GetIndex()];
-			FPrimitiveSceneInfo* PrimitiveSceneInfo = PrimitiveSceneProxy->GetPrimitiveSceneInfo();
-
-			TDistortionMeshDrawingPolicyFactory::ContextType Context(bInitializeOffsets);
-
-			// Note: As for distortion rendering the order doesn't matter we actually could iterate View.DynamicMeshElements without this indirection
-			{
-
-				// range in View.DynamicMeshElements[]
-				FInt32Range range = View.GetDynamicMeshElementRange(PrimitiveSceneInfo->GetIndex());
-
-				for (int32 MeshBatchIndex = range.GetLowerBoundValue(); MeshBatchIndex < range.GetUpperBoundValue(); MeshBatchIndex++)
-				{
-					const FMeshBatchAndRelevance& MeshBatchAndRelevance = View.DynamicMeshElements[MeshBatchIndex];
-
-					checkSlow(MeshBatchAndRelevance.PrimitiveSceneProxy == PrimitiveSceneProxy);
-
-					const FMeshBatch& MeshBatch = *MeshBatchAndRelevance.Mesh;
-					bDirty |= TDistortionMeshDrawingPolicyFactory::DrawDynamicMesh(RHICmdList, View, Context, MeshBatch, false, DrawRenderState, MeshBatchAndRelevance.PrimitiveSceneProxy, MeshBatch.BatchHitProxyId);
-				}
-			}
-
-			// Render static scene prim
-			if (ViewRelevance.bStaticRelevance)
-			{
-				// Render static meshes from static scene prim
-				for (int32 StaticMeshIdx = 0; StaticMeshIdx < PrimitiveSceneInfo->StaticMeshes.Num(); StaticMeshIdx++)
-				{
-					FStaticMesh& StaticMesh = PrimitiveSceneInfo->StaticMeshes[StaticMeshIdx];
-					if (View.StaticMeshVisibilityMap[StaticMesh.Id]
-						// Only render static mesh elements using translucent materials
-						&& StaticMesh.IsTranslucent(View.GetFeatureLevel()))
-					{
-						bDirty |= TDistortionMeshDrawingPolicyFactory::DrawStaticMesh(
-							RHICmdList,
-							&View,
-							bInitializeOffsets,
-							StaticMesh,
-							StaticMesh.bRequiresPerElementVisibility ? View.StaticMeshBatchVisibility[StaticMesh.BatchVisibilityId] : ((1ull << StaticMesh.Elements.Num()) - 1),
-							DrawRenderState,
-							PrimitiveSceneProxy,
-							StaticMesh.BatchHitProxyId
-						);
-					}
-				}
-			}
-		}
-	}
-	return bDirty;
-}
-
 int32 FSceneRenderer::GetRefractionQuality(const FSceneViewFamily& ViewFamily)
 {
 	static const auto ICVar = IConsoleManager::Get().FindTConsoleVariableDataInt(TEXT("r.RefractionQuality"));
@@ -987,9 +487,7 @@ bool SubmitDistortionMeshDrawCommands(FRHICommandListImmediate& RHICmdList, cons
 
 	{
 		QUICK_SCOPE_CYCLE_COUNTER(STAT_FDistortionPrimSet_DrawAccumulatedOffsets_Prims);
-
 		View.ParallelMeshDrawCommandPasses[EMeshPass::Distortion].DispatchDraw(nullptr, RHICmdList);
-
 		bDirty |= View.ParallelMeshDrawCommandPasses[EMeshPass::Distortion].HasAnyDraw();
 	}
 
@@ -1010,7 +508,7 @@ void FSceneRenderer::RenderDistortion(FRHICommandListImmediate& RHICmdList)
 	for(int32 ViewIndex = 0;ViewIndex < Views.Num();ViewIndex++)
 	{
 		const FViewInfo& View = Views[ViewIndex];
-		if( View.DistortionPrimSet.NumPrims() > 0)
+		if (View.bHasDistortionPrimitives)
 		{
 			bRender=true;
 			break;
@@ -1086,14 +584,7 @@ void FSceneRenderer::RenderDistortion(FRHICommandListImmediate& RHICmdList)
 
 
 					// Draw distortion meshes to accumulate their offsets.
-					if (UseMeshDrawCommandPipeline())
-					{
-						bDirty |= SubmitDistortionMeshDrawCommands(RHICmdList, View, DrawRenderState);
-					}
-					else
-					{
-						bDirty |= View.DistortionPrimSet.DrawAccumulatedOffsets(RHICmdList, View, DrawRenderState, false);
-					}
+					bDirty |= SubmitDistortionMeshDrawCommands(RHICmdList, View, DrawRenderState);
 				}
 			}
 			RHICmdList.EndRenderPass();
diff --git a/Engine/Source/Runtime/Renderer/Private/DistortionRendering.h b/Engine/Source/Runtime/Renderer/Private/DistortionRendering.h
--- a/Engine/Source/Runtime/Renderer/Private/DistortionRendering.h
+++ b/Engine/Source/Runtime/Renderer/Private/DistortionRendering.h
@@ -8,7 +8,6 @@
 
 #include "CoreMinimal.h"
 #include "RendererInterface.h"
-#include "DrawingPolicy.h"
 #include "MeshPassProcessor.h"
 
 class FPrimitiveSceneProxy;
diff --git a/Engine/Source/Runtime/Renderer/Private/DynamicPrimitiveDrawing.h b/Engine/Source/Runtime/Renderer/Private/DynamicPrimitiveDrawing.h
--- a/Engine/Source/Runtime/Renderer/Private/DynamicPrimitiveDrawing.h
+++ b/Engine/Source/Runtime/Renderer/Private/DynamicPrimitiveDrawing.h
@@ -10,60 +10,6 @@
 #include "SceneManagement.h"
 #include "SceneRendering.h"
 
-/**
-* Draws a range of view's elements with the specified drawing policy factory type.
-* @param View - The view to draw the meshes for.
-* @param DrawingContext - The drawing policy type specific context for the drawing.
-* @param DPGIndex World or Foreground DPG index for draw order
-* @param FirstIndex - Element range
-* @param LastIndex - Element range
-*/
-template<class DrawingPolicyFactoryType>
-void DrawViewElementsInner(
-	FRHICommandList& RHICmdList,
-	const FViewInfo& View,
-	const FDrawingPolicyRenderState& DrawRenderState,
-	const typename DrawingPolicyFactoryType::ContextType& DrawingContext,
-	uint8 DPGIndex,
-	bool bPreFog,
-	int32 FirstIndex,
-	int32 LastIndex
-	);
-
-/**
- * Draws a view's elements with the specified drawing policy factory type.
- * @param View - The view to draw the meshes for.
- * @param DrawingContext - The drawing policy type specific context for the drawing.
- * @param DPGIndex World or Foreground DPG index for draw order
- * @param bPreFog - true if the draw call is occurring before fog has been rendered.
- */
-template<class DrawingPolicyFactoryType>
-bool DrawViewElements(
-	FRHICommandList& RHICmdList,
-	const FViewInfo& View,
-	const FDrawingPolicyRenderState& DrawRenderState,
-	const typename DrawingPolicyFactoryType::ContextType& DrawingContext,
-	uint8 DPGIndex,
-	bool bPreFog
-	);
-
-/**
-* Draws a view's elements with the specified drawing policy factory type.
-* @param View - The view to draw the meshes for.
-* @param DrawingContext - The drawing policy type specific context for the drawing.
-* @param DPGIndex World or Foreground DPG index for draw order
-* @param bPreFog - true if the draw call is occurring before fog has been rendered.
-* @param ParentCmdList - cmdlist to put the wait and execute task on
-* @param Width - parallel width
-*/
-template<class DrawingPolicyFactoryType>
-void DrawViewElementsParallel(
-	const typename DrawingPolicyFactoryType::ContextType& DrawingContext,
-	uint8 DPGIndex,
-	bool bPreFog,
-	FParallelCommandListSet& ParallelCommandListSet
-	);
-
 /** A primitive draw interface which adds the drawn elements to the view's batched elements. */
 class FViewElementPDI : public FPrimitiveDrawInterface
 {
diff --git a/Engine/Source/Runtime/Renderer/Private/DynamicPrimitiveDrawing.inl b/Engine/Source/Runtime/Renderer/Private/DynamicPrimitiveDrawing.inl
--- a/Engine/Source/Runtime/Renderer/Private/DynamicPrimitiveDrawing.inl
+++ b/Engine/Source/Runtime/Renderer/Private/DynamicPrimitiveDrawing.inl
@@ -27,166 +27,6 @@ struct FMeshBatchElement;
 struct TStatId;
 template<typename TTask> class TGraphTask;
 
-template<class DrawingPolicyFactoryType>
-void DrawViewElementsInner(
-	FRHICommandList& RHICmdList,
-	const FViewInfo& View,
-	const FDrawingPolicyRenderState& DrawRenderState,
-	const typename DrawingPolicyFactoryType::ContextType& DrawingContext,
-	uint8 DPGIndex,
-	bool bPreFog,
-	int32 FirstIndex,
-	int32 LastIndex
-	)
-{
-	check(RHICmdList.IsInsideRenderPass());
-	// Get the correct element list based on dpg index
-	const TIndirectArray<FMeshBatch>& ViewMeshElementList = ( DPGIndex == SDPG_Foreground ? View.TopViewMeshElements : View.ViewMeshElements );
-	// Draw the view's mesh elements.
-	check(LastIndex < ViewMeshElementList.Num());
-	for (int32 MeshIndex = FirstIndex; MeshIndex <= LastIndex; MeshIndex++)
-	{
-		const FMeshBatch& Mesh = ViewMeshElementList[MeshIndex];
-		const auto FeatureLevel = View.GetFeatureLevel();
-		check(Mesh.MaterialRenderProxy);
-		check(Mesh.MaterialRenderProxy->GetMaterial(FeatureLevel));
-		const bool bIsTwoSided = Mesh.MaterialRenderProxy->GetMaterial(FeatureLevel)->IsTwoSided();
-		int32 bBackFace = bIsTwoSided ? 1 : 0;
-		do
-		{
-			FDrawingPolicyRenderState DrawRenderStateLocal(DrawRenderState);
-			DrawRenderStateLocal.ModifyViewOverrideFlags() ^= (bBackFace != 0) ? EDrawingPolicyOverrideFlags::ReverseCullMode : EDrawingPolicyOverrideFlags::None;
-
-			DrawingPolicyFactoryType::DrawDynamicMesh(
-				RHICmdList, 
-				View,
-				DrawingContext,
-				Mesh,
-				bPreFog,
-				DrawRenderStateLocal,
-				NULL,
-				Mesh.BatchHitProxyId
-				);
-			--bBackFace;
-		} while( bBackFace >= 0 );
-	}
-}
-
-template<typename DrawingPolicyFactoryType>
-class FDrawViewElementsAnyThreadTask : public FRenderTask
-{
-	FRHICommandList& RHICmdList;
-	const FViewInfo& View;
-	FDrawingPolicyRenderState DrawRenderState;
-	const typename DrawingPolicyFactoryType::ContextType& DrawingContext;
-	uint8 DPGIndex;
-	bool bPreFog;
-
-
-	const int32 FirstIndex;
-	const int32 LastIndex;
-
-public:
-
-	FDrawViewElementsAnyThreadTask(
-		FRHICommandList& InRHICmdList,
-		const FViewInfo& InView,
-		const FDrawingPolicyRenderState& InDrawRenderState,
-		const typename DrawingPolicyFactoryType::ContextType& InDrawingContext,
-		uint8 InDPGIndex,
-		bool InbPreFog,
-		int32 InFirstIndex,
-		int32 InLastIndex
-		)
-		: RHICmdList(InRHICmdList)
-		, View(InView)
-		, DrawRenderState(InDrawRenderState)
-		, DrawingContext(InDrawingContext)
-		, DPGIndex(InDPGIndex)
-		, bPreFog(InbPreFog)
-		, FirstIndex(InFirstIndex)
-		, LastIndex(InLastIndex)
-	{
-	}
-
-	FORCEINLINE TStatId GetStatId() const
-	{
-		RETURN_QUICK_DECLARE_CYCLE_STAT(FDrawViewElementsAnyThreadTask, STATGROUP_TaskGraphTasks);
-	}
-
-	static ESubsequentsMode::Type GetSubsequentsMode() { return ESubsequentsMode::TrackSubsequents; }
-
-	void DoTask(ENamedThreads::Type CurrentThread, const FGraphEventRef& MyCompletionGraphEvent)
-	{
-		check(RHICmdList.IsInsideRenderPass());
-		DrawViewElementsInner<DrawingPolicyFactoryType>(RHICmdList, View, DrawRenderState, DrawingContext, DPGIndex, bPreFog, FirstIndex, LastIndex);
-		RHICmdList.EndRenderPass();
-		RHICmdList.HandleRTThreadTaskCompletion(MyCompletionGraphEvent);
-	}
-};
-
-template<class DrawingPolicyFactoryType>
-void DrawViewElementsParallel(
-	const typename DrawingPolicyFactoryType::ContextType& DrawingContext,
-	uint8 DPGIndex,
-	bool bPreFog,
-	FParallelCommandListSet& ParallelCommandListSet
-	)
-{
-	// Get the correct element list based on dpg index
-	const TIndirectArray<FMeshBatch>& ViewMeshElementList = (DPGIndex == SDPG_Foreground ? ParallelCommandListSet.View.TopViewMeshElements : ParallelCommandListSet.View.ViewMeshElements);
-
-	{
-		int32 NumPrims = ViewMeshElementList.Num();
-		int32 EffectiveThreads = FMath::Min<int32>(FMath::DivideAndRoundUp(NumPrims, ParallelCommandListSet.MinDrawsPerCommandList), ParallelCommandListSet.Width);
-
-		int32 Start = 0;
-		if (EffectiveThreads)
-		{
-
-			int32 NumPer = NumPrims / EffectiveThreads;
-			int32 Extra = NumPrims - NumPer * EffectiveThreads;
-
-
-			for (int32 ThreadIndex = 0; ThreadIndex < EffectiveThreads; ThreadIndex++)
-			{
-				int32 Last = Start + (NumPer - 1) + (ThreadIndex < Extra);
-				check(Last >= Start);
-
-				{
-					FRHICommandList* CmdList = ParallelCommandListSet.NewParallelCommandList();
-
-					FGraphEventRef AnyThreadCompletionEvent = TGraphTask<FDrawViewElementsAnyThreadTask<DrawingPolicyFactoryType> >::CreateTask(ParallelCommandListSet.GetPrereqs(), ENamedThreads::GetRenderThread())
-						.ConstructAndDispatchWhenReady(*CmdList, ParallelCommandListSet.View, ParallelCommandListSet.DrawRenderState, DrawingContext, DPGIndex, bPreFog, Start, Last);
-
-					ParallelCommandListSet.AddParallelCommandList(CmdList, AnyThreadCompletionEvent, Last - Start + 1);
-				}
-				Start = Last + 1;
-			}
-		}
-	}
-}
-
-template<class DrawingPolicyFactoryType>
-bool DrawViewElements(
-	FRHICommandList& RHICmdList,
-	const FViewInfo& View,
-	const FDrawingPolicyRenderState& DrawRenderState,
-	const typename DrawingPolicyFactoryType::ContextType& DrawingContext,
-	uint8 DPGIndex,
-	bool bPreFog
-	)
-{
-	// Get the correct element list based on dpg index
-	const TIndirectArray<FMeshBatch>& ViewMeshElementList = (DPGIndex == SDPG_Foreground ? View.TopViewMeshElements : View.ViewMeshElements);
-	if (ViewMeshElementList.Num() != 0)
-	{
-		DrawViewElementsInner<DrawingPolicyFactoryType>(RHICmdList, View, DrawRenderState, DrawingContext, DPGIndex, bPreFog, 0, ViewMeshElementList.Num() - 1);
-		return true;
-	}
-	return false;
-}
-
 inline FViewElementPDI::FViewElementPDI(FViewInfo* InViewInfo,FHitProxyConsumer* InHitProxyConsumer,TArray<FPrimitiveUniformShaderParameters>* InDynamicPrimitiveShaderData):
 	FPrimitiveDrawInterface(InViewInfo),
 	ViewInfo(InViewInfo),
diff --git a/Engine/Source/Runtime/Renderer/Private/GPUScene.cpp b/Engine/Source/Runtime/Renderer/Private/GPUScene.cpp
--- a/Engine/Source/Runtime/Renderer/Private/GPUScene.cpp
+++ b/Engine/Source/Runtime/Renderer/Private/GPUScene.cpp
@@ -347,23 +347,6 @@ void UploadDynamicPrimitiveShaderDataForView(FRHICommandList& RHICmdList, FScene
 		View.CachedViewUniformShaderParameters->LightmapSceneData = Scene.GPUScene.LightmapDataBuffer.SRV;
 		View.CachedViewUniformShaderParameters->PrimitiveSceneData = ViewPrimitiveShaderDataBuffer.SRV;
 		View.ViewUniformBuffer.UpdateUniformBufferImmediate(*View.CachedViewUniformShaderParameters);
-
-		if (!UseMeshDrawCommandPipeline())
-		{
-			QUICK_SCOPE_CYCLE_COUNTER(STAT_PrimitiveIdBufferEmulation);
-
-			//@todo MeshCommandPipeline - remove PrimitiveIdBufferEmulation once drawing policies are gone
-			View.OneFramePrimitiveIdBufferEmulation.Release();
-			View.OneFramePrimitiveIdBufferEmulation.Initialize(sizeof(uint32), NumPrimitiveEntries, PF_R32_UINT);
-			uint32* LockedPrimitiveIds = (uint32*)RHILockVertexBuffer(View.OneFramePrimitiveIdBufferEmulation.Buffer, 0, View.OneFramePrimitiveIdBufferEmulation.NumBytes, RLM_WriteOnly);
-
-			for (int32 i = 0; i < NumPrimitiveEntries; i++)
-			{
-				LockedPrimitiveIds[i] = i;
-			}
-
-			RHIUnlockVertexBuffer(View.OneFramePrimitiveIdBufferEmulation.Buffer);
-		}
 	}
 }
 
diff --git a/Engine/Source/Runtime/Renderer/Private/LightMapDensityRendering.h b/Engine/Source/Runtime/Renderer/Private/LightMapDensityRendering.h
--- a/Engine/Source/Runtime/Renderer/Private/LightMapDensityRendering.h
+++ b/Engine/Source/Runtime/Renderer/Private/LightMapDensityRendering.h
@@ -15,7 +15,6 @@
 #include "Shader.h"
 #include "UnrealEngine.h"
 #include "MeshMaterialShaderType.h"
-#include "DrawingPolicy.h"
 #include "MeshMaterialShader.h"
 #include "ShaderBaseClasses.h"
 #include "SceneRendering.h"
@@ -87,21 +86,6 @@ public:
 		return bShaderHasOutdatedParameters;
 	}
 
-	void SetParameters(		
-		FRHICommandList& RHICmdList, 
-		const FMaterialRenderProxy* MaterialRenderProxy,
-		const FSceneView& View, 
-		const FDrawingPolicyRenderState& DrawRenderState
-		)
-	{
-		FMeshMaterialShader::SetParameters(RHICmdList, GetVertexShader(), MaterialRenderProxy, *MaterialRenderProxy->GetMaterial(View.GetFeatureLevel()), View, DrawRenderState.GetViewUniformBuffer(), DrawRenderState.GetPassUniformBuffer());
-	}
-
-	void SetMesh(FRHICommandList& RHICmdList, const FVertexFactory* VertexFactory,const FSceneView& View,const FPrimitiveSceneProxy* Proxy,const FMeshBatchElement& BatchElement,const FDrawingPolicyRenderState& DrawRenderState)
-	{
-		FMeshMaterialShader::SetMesh(RHICmdList, GetVertexShader(),VertexFactory,View,Proxy,BatchElement,DrawRenderState);
-	}
-
 	void GetShaderBindings(
 		const FScene* Scene,
 		ERHIFeatureLevel::Type FeatureLevel,
@@ -222,29 +206,6 @@ public:
 	}
 	TLightMapDensityPS() {}
 
-	void SetParameters(FRHICommandList& RHICmdList, const FMaterialRenderProxy* MaterialRenderProxy,const FSceneView* View,const FDrawingPolicyRenderState& DrawRenderState)
-	{
-		FMeshMaterialShader::SetParameters(RHICmdList, GetPixelShader(), MaterialRenderProxy, *MaterialRenderProxy->GetMaterial(View->GetFeatureLevel()), *View, DrawRenderState.GetViewUniformBuffer(), DrawRenderState.GetPassUniformBuffer());
-	}
-
-	void SetMesh(FRHICommandList& RHICmdList, const FVertexFactory* VertexFactory,const FPrimitiveSceneProxy* PrimitiveSceneProxy,const FMeshBatchElement& BatchElement,const FSceneView& View,const FDrawingPolicyRenderState& DrawRenderState,FVector& InBuiltLightingAndSelectedFlags,FVector2D& InLightMapResolutionScale, bool bTextureMapped)
-	{
-		FMeshMaterialShader::SetMesh(RHICmdList, GetPixelShader(),VertexFactory,View,PrimitiveSceneProxy,BatchElement,DrawRenderState);
-
-		SetShaderValue(RHICmdList, GetPixelShader(),BuiltLightingAndSelectedFlags,InBuiltLightingAndSelectedFlags,0);
-		SetShaderValue(RHICmdList, GetPixelShader(),LightMapResolutionScale,InLightMapResolutionScale,0);
-		if (LightMapDensityDisplayOptions.IsBound())
-		{
-			FVector4 OptionsParameter(
-				GEngine->bRenderLightMapDensityGrayscale ? GEngine->RenderLightMapDensityGrayscaleScale : 0.0f,
-				GEngine->bRenderLightMapDensityGrayscale ? 0.0f : GEngine->RenderLightMapDensityColorScale,
-				(bTextureMapped == true) ? 1.0f : 0.0f,
-				(bTextureMapped == false) ? 1.0f : 0.0f
-				);
-			SetShaderValue(RHICmdList, GetPixelShader(),LightMapDensityDisplayOptions,OptionsParameter,0);
-		}
-	}
-
 	void GetShaderBindings(
 		const FScene* Scene,
 		ERHIFeatureLevel::Type FeatureLevel,
diff --git a/Engine/Source/Runtime/Renderer/Private/LightMapRendering.cpp b/Engine/Source/Runtime/Renderer/Private/LightMapRendering.cpp
--- a/Engine/Source/Runtime/Renderer/Private/LightMapRendering.cpp
+++ b/Engine/Source/Runtime/Renderer/Private/LightMapRendering.cpp
@@ -104,39 +104,6 @@ void FSelfShadowedVolumetricLightmapPolicy::GetPixelShaderBindings(
 	FSelfShadowedTranslucencyPolicy::GetPixelShaderBindings(PrimitiveSceneProxy, ShaderElementData.SelfShadowTranslucencyUniformBuffer, PixelShaderParameters, ShaderBindings);
 }
 
-void FUniformLightMapPolicy::SetMesh(
-	FRHICommandList& RHICmdList,
-	const FSceneView& View,
-	const FPrimitiveSceneProxy* PrimitiveSceneProxy,
-	const VertexParametersType* VertexShaderParameters,
-	const PixelParametersType* PixelShaderParameters,
-	FShader* VertexShader,
-	FShader* PixelShader,
-	const FVertexFactory* VertexFactory,
-	const FMaterialRenderProxy* MaterialRenderProxy,
-	const FLightCacheInterface* LCI
-	) const
-{
-	FUniformBufferRHIParamRef PrecomputedLightingBuffer = nullptr;
-	FUniformBufferRHIParamRef LightmapResourceClusterBuffer = nullptr;
-	FUniformBufferRHIParamRef IndirectLightingCacheBuffer = nullptr;
-
-	SetupLCIUniformBuffers(PrimitiveSceneProxy, LCI, PrecomputedLightingBuffer, LightmapResourceClusterBuffer, IndirectLightingCacheBuffer);
-
-	if (VertexShaderParameters)
-	{
-		SetUniformBufferParameter(RHICmdList, VertexShader->GetVertexShader(), VertexShaderParameters->PrecomputedLightingBufferParameter, PrecomputedLightingBuffer);
-		SetUniformBufferParameter(RHICmdList, VertexShader->GetVertexShader(), VertexShaderParameters->IndirectLightingCacheParameter, IndirectLightingCacheBuffer);
-		SetUniformBufferParameter(RHICmdList, VertexShader->GetVertexShader(), VertexShaderParameters->LightmapResourceCluster, LightmapResourceClusterBuffer);
-	}
-	if (PixelShaderParameters)
-	{
-		SetUniformBufferParameter(RHICmdList, PixelShader->GetPixelShader(), PixelShaderParameters->PrecomputedLightingBufferParameter, PrecomputedLightingBuffer);
-		SetUniformBufferParameter(RHICmdList, PixelShader->GetPixelShader(), PixelShaderParameters->IndirectLightingCacheParameter, IndirectLightingCacheBuffer);
-		SetUniformBufferParameter(RHICmdList, PixelShader->GetPixelShader(), PixelShaderParameters->LightmapResourceCluster, LightmapResourceClusterBuffer);
-	}
-}
-
 void FUniformLightMapPolicy::GetVertexShaderBindings(
 	const FPrimitiveSceneProxy* PrimitiveSceneProxy,
 	const ElementDataType& ShaderElementData,
diff --git a/Engine/Source/Runtime/Renderer/Private/LightMapRendering.h b/Engine/Source/Runtime/Renderer/Private/LightMapRendering.h
--- a/Engine/Source/Runtime/Renderer/Private/LightMapRendering.h
+++ b/Engine/Source/Runtime/Renderer/Private/LightMapRendering.h
@@ -11,7 +11,6 @@
 #include "RenderResource.h"
 #include "UniformBuffer.h"
 #include "ShaderParameters.h"
-#include "DrawingPolicy.h"
 #include "ShadowRendering.h"
 #include "LightMap.h"
 
@@ -243,12 +242,6 @@ public:
 	{
 		return true;
 	}
-
-	friend int32 CompareDrawingPolicy(const FSelfShadowedTranslucencyPolicy&,const FSelfShadowedTranslucencyPolicy&)
-	{
-		return 0;
-	}
-
 };
 
 /**
@@ -761,19 +754,6 @@ public:
 
 	FUniformLightMapPolicy(ELightMapPolicyType InIndirectPolicy) : IndirectPolicy(InIndirectPolicy) {}
 
-	void SetMesh(
-		FRHICommandList& RHICmdList,
-		const FSceneView& View,
-		const FPrimitiveSceneProxy* PrimitiveSceneProxy,
-		const VertexParametersType* VertexShaderParameters,
-		const PixelParametersType* PixelShaderParameters,
-		FShader* VertexShader,
-		FShader* PixelShader,
-		const FVertexFactory* VertexFactory,
-		const FMaterialRenderProxy* MaterialRenderProxy,
-		const FLightCacheInterface* LCI
-		) const;
-
 	static void GetVertexShaderBindings(
 		const FPrimitiveSceneProxy* PrimitiveSceneProxy,
 		const ElementDataType& ShaderElementData,
@@ -800,12 +780,6 @@ public:
 		return A.IndirectPolicy == B.IndirectPolicy;
 	}
 
-	friend int32 CompareDrawingPolicy(const FUniformLightMapPolicy& A,const FUniformLightMapPolicy& B)
-	{
-		COMPAREDRAWINGPOLICYMEMBERS(IndirectPolicy);
-		return  0;
-	}
-
 	ELightMapPolicyType GetIndirectPolicy() const { return IndirectPolicy; }
 
 private:
diff --git a/Engine/Source/Runtime/Renderer/Private/MaterialTexCoordScalesRendering.h b/Engine/Source/Runtime/Renderer/Private/MaterialTexCoordScalesRendering.h
--- a/Engine/Source/Runtime/Renderer/Private/MaterialTexCoordScalesRendering.h
+++ b/Engine/Source/Runtime/Renderer/Private/MaterialTexCoordScalesRendering.h
@@ -24,7 +24,6 @@ struct FMeshDrawingRenderState;
 * Pixel shader that renders texcoord scales.
 * The shader is only compiled with the local vertex factory to prevent multiple compilation.
 * Nothing from the factory is actually used, but the shader must still derive from FMeshMaterialShader.
-* This is required to call FMeshMaterialShader::SetMesh and bind primitive related data.
 */
 class FMaterialTexCoordScalePS : public FDebugViewModePS
 {
diff --git a/Engine/Source/Runtime/Renderer/Private/MeshDrawCommands.cpp b/Engine/Source/Runtime/Renderer/Private/MeshDrawCommands.cpp
--- a/Engine/Source/Runtime/Renderer/Private/MeshDrawCommands.cpp
+++ b/Engine/Source/Runtime/Renderer/Private/MeshDrawCommands.cpp
@@ -95,7 +95,7 @@ void UpdateTranslucentMeshSortKeys(
 		}
 
 		// Patch distance inside translucent mesh sort key.
-		UTranslucentMeshSortKey TranslucentMeshSortKey;
+		FTranslucentMeshSortKey TranslucentMeshSortKey;
 		TranslucentMeshSortKey.PackedData = VisibleCommand.SortKey.PackedData;
 		TranslucentMeshSortKey.Fields.Distance = (uint32)~BitInvertIfNegativeFloat(*((uint32*)&Distance));
 		VisibleCommand.SortKey.PackedData = TranslucentMeshSortKey.PackedData;		
@@ -514,7 +514,7 @@ void ApplyViewOverridesToMeshDrawCommands(
 				const FMeshDrawCommand& MeshCommand = *VisibleMeshDrawCommand.MeshDrawCommand;
 				NewMeshCommand = MeshCommand;
 
-				const ERasterizerCullMode LocalCullMode = bRenderSceneTwoSided ? CM_None : bReverseCulling ? FMeshDrawingPolicy::InverseCullMode(VisibleMeshDrawCommand.MeshCullMode) : VisibleMeshDrawCommand.MeshCullMode;
+				const ERasterizerCullMode LocalCullMode = bRenderSceneTwoSided ? CM_None : bReverseCulling ? FMeshPassProcessor::InverseCullMode(VisibleMeshDrawCommand.MeshCullMode) : VisibleMeshDrawCommand.MeshCullMode;
 				NewMeshCommand.PipelineState.RasterizerState = GetStaticRasterizerState<true>(VisibleMeshDrawCommand.MeshFillMode, LocalCullMode);
 
 				if (BasePassDepthStencilAccess != DefaultBasePassDepthStencilAccess && PassType == EMeshPass::BasePass)
diff --git a/Engine/Source/Runtime/Renderer/Private/MeshPassProcessor.cpp b/Engine/Source/Runtime/Renderer/Private/MeshPassProcessor.cpp
--- a/Engine/Source/Runtime/Renderer/Private/MeshPassProcessor.cpp
+++ b/Engine/Source/Runtime/Renderer/Private/MeshPassProcessor.cpp
@@ -12,6 +12,7 @@
 #include "ScenePrivate.h"
 #include "SceneInterface.h"
 #include "MeshPassProcessor.inl"
+#include "PipelineStateCache.h"
 
 TSet<FGraphicsMinimalPipelineStateInitializer> FGraphicsMinimalPipelineStateId::GlobalTable;
 FCriticalSection GlobalTableCriticalSection;
@@ -619,6 +620,7 @@ void FMeshDrawCommand::SetDrawParametersAndFinalize(const FMeshBatch& MeshBatch,
 	NumInstances = BatchElement.NumInstances * InstanceFactor;
 	BaseVertexIndex = BatchElement.BaseVertexIndex;
 	NumVertices = BatchElement.MaxVertexIndex - BatchElement.MinVertexIndex + 1;
+	IndirectArgsBuffer = BatchElement.IndirectArgsBuffer;
 
 	int32 SegmentIndex = MeshBatch.SegmentIndex + BatchElementIndex;
 	RayTracedSegmentIndex = (SegmentIndex < UINT8_MAX) ? uint8(MeshBatch.SegmentIndex + BatchElementIndex) : UINT8_MAX;
@@ -729,7 +731,7 @@ void FMeshDrawCommand::SubmitDraw(
 #if WANTS_DRAW_MESH_EVENTS
 	TDrawEvent<FRHICommandList> MeshEvent;
 
-	if (GShowMaterialDrawEventTypes)
+	if (GShowMaterialDrawEvents)
 	{
 		if (MeshDrawCommand.NumInstances > 1)
 		{
@@ -782,15 +784,26 @@ void FMeshDrawCommand::SubmitDraw(
 
 	if (MeshDrawCommand.IndexBuffer)
 	{
-		RHICmdList.DrawIndexedPrimitive(
-			MeshDrawCommand.IndexBuffer,
-			MeshDrawCommand.BaseVertexIndex,
-			0,
-			MeshDrawCommand.NumVertices,
-			MeshDrawCommand.FirstIndex,
-			MeshDrawCommand.NumPrimitives,
-			MeshDrawCommand.NumInstances
-		);
+		if (MeshDrawCommand.IndirectArgsBuffer)
+		{
+			RHICmdList.DrawIndexedPrimitiveIndirect(
+				MeshDrawCommand.IndexBuffer, 
+				MeshDrawCommand.IndirectArgsBuffer, 
+				0
+				);
+		}
+		else
+		{
+			RHICmdList.DrawIndexedPrimitive(
+				MeshDrawCommand.IndexBuffer,
+				MeshDrawCommand.BaseVertexIndex,
+				0,
+				MeshDrawCommand.NumVertices,
+				MeshDrawCommand.FirstIndex,
+				MeshDrawCommand.NumPrimitives,
+				MeshDrawCommand.NumInstances
+			);
+		}
 	}
 	else
 	{
@@ -800,6 +813,7 @@ void FMeshDrawCommand::SubmitDraw(
 			MeshDrawCommand.NumInstances
 		);
 	}
+
 }
 
 void SubmitMeshDrawCommands(
@@ -822,10 +836,10 @@ void SubmitMeshDrawCommandsRange(
 	FRHICommandList& RHICmdList)
 {
 	FMeshDrawCommandStateCache StateCache;
+	INC_DWORD_STAT_BY(STAT_MeshDrawCalls, NumMeshDrawCommands);
 
 	for (int32 DrawCommandIndex = StartIndex; DrawCommandIndex < StartIndex + NumMeshDrawCommands; DrawCommandIndex++)
 	{
-		INC_DWORD_STAT(STAT_MeshDrawCalls);
 		SCOPED_CONDITIONAL_DRAW_EVENTF(RHICmdList, MeshEvent, GEmitMeshDrawEvent != 0, TEXT("Mesh Draw"));
 
 		const FVisibleMeshDrawCommand& VisibleMeshDrawCommand = VisibleMeshDrawCommands[DrawCommandIndex];
@@ -868,6 +882,34 @@ FMeshPassProcessor::FMeshPassProcessor(const FScene* InScene, ERHIFeatureLevel::
 	static_assert(sizeof(FMeshPassMask::Data) * 8 >= EMeshPass::Num, "FMeshPassMask::Data is too small to fit all mesh passes.");
 }
 
+enum class EDrawingPolicyOverrideFlags
+{
+	None = 0,
+	TwoSided = 1 << 0,
+	DitheredLODTransition = 1 << 1,
+	Wireframe = 1 << 2,
+	ReverseCullMode = 1 << 3,
+};
+ENUM_CLASS_FLAGS(EDrawingPolicyOverrideFlags);
+
+struct FMeshDrawingPolicyOverrideSettings
+{
+	EDrawingPolicyOverrideFlags	MeshOverrideFlags = EDrawingPolicyOverrideFlags::None;
+	EPrimitiveType				MeshPrimitiveType = PT_TriangleList;
+};
+
+FORCEINLINE_DEBUGGABLE FMeshDrawingPolicyOverrideSettings ComputeMeshOverrideSettings(const FMeshBatch& Mesh)
+{
+	FMeshDrawingPolicyOverrideSettings OverrideSettings;
+	OverrideSettings.MeshPrimitiveType = (EPrimitiveType)Mesh.Type;
+
+	OverrideSettings.MeshOverrideFlags |= Mesh.bDisableBackfaceCulling ? EDrawingPolicyOverrideFlags::TwoSided : EDrawingPolicyOverrideFlags::None;
+	OverrideSettings.MeshOverrideFlags |= Mesh.bDitheredLODTransition ? EDrawingPolicyOverrideFlags::DitheredLODTransition : EDrawingPolicyOverrideFlags::None;
+	OverrideSettings.MeshOverrideFlags |= Mesh.bWireframe ? EDrawingPolicyOverrideFlags::Wireframe : EDrawingPolicyOverrideFlags::None;
+	OverrideSettings.MeshOverrideFlags |= Mesh.ReverseCulling ? EDrawingPolicyOverrideFlags::ReverseCullMode : EDrawingPolicyOverrideFlags::None;
+	return OverrideSettings;
+}
+
 ERasterizerFillMode FMeshPassProcessor::ComputeMeshFillMode(const FMeshBatch& Mesh, const FMaterial& InMaterialResource) const
 {
 	const FMeshDrawingPolicyOverrideSettings InOverrideSettings = ComputeMeshOverrideSettings(Mesh);
@@ -891,7 +933,7 @@ ERasterizerCullMode FMeshPassProcessor::ComputeMeshCullMode(const FMeshBatch& Me
 void FMeshPassProcessor::SetDrawCommandEvent(const FPrimitiveSceneProxy* PrimitiveSceneProxy, const FMaterial& RESTRICT MaterialResource, FMeshDrawCommand& MeshDrawCommand) const
 {
 #if WANTS_DRAW_MESH_EVENTS
-	if (GShowMaterialDrawEventTypes)
+	if (GShowMaterialDrawEvents)
 	{
 		if (PrimitiveSceneProxy)
 		{
diff --git a/Engine/Source/Runtime/Renderer/Private/MobileBasePass.cpp b/Engine/Source/Runtime/Renderer/Private/MobileBasePass.cpp
--- a/Engine/Source/Runtime/Renderer/Private/MobileBasePass.cpp
+++ b/Engine/Source/Runtime/Renderer/Private/MobileBasePass.cpp
@@ -7,7 +7,6 @@
 #include "MobileBasePassRendering.h"
 #include "TranslucentRendering.h"
 #include "DynamicPrimitiveDrawing.h"
-#include "StaticMeshDrawList.h"
 #include "ScenePrivate.h"
 #include "ShaderPlatformQualitySettings.h"
 #include "MaterialShaderQualitySettings.h"
diff --git a/Engine/Source/Runtime/Renderer/Private/MobileBasePassRendering.cpp b/Engine/Source/Runtime/Renderer/Private/MobileBasePassRendering.cpp
--- a/Engine/Source/Runtime/Renderer/Private/MobileBasePassRendering.cpp
+++ b/Engine/Source/Runtime/Renderer/Private/MobileBasePassRendering.cpp
@@ -6,7 +6,6 @@
 
 #include "MobileBasePassRendering.h"
 #include "DynamicPrimitiveDrawing.h"
-#include "StaticMeshDrawList.h"
 #include "ScenePrivate.h"
 #include "ShaderPlatformQualitySettings.h"
 #include "MaterialShaderQualitySettings.h"
diff --git a/Engine/Source/Runtime/Renderer/Private/MobileDistortionPass.cpp b/Engine/Source/Runtime/Renderer/Private/MobileDistortionPass.cpp
--- a/Engine/Source/Runtime/Renderer/Private/MobileDistortionPass.cpp
+++ b/Engine/Source/Runtime/Renderer/Private/MobileDistortionPass.cpp
@@ -9,6 +9,7 @@ MobileDistortionPass.cpp - Mobile specific rendering of primtives with refractio
 #include "DynamicPrimitiveDrawing.h"
 #include "PostProcess/PostProcessing.h"
 #include "PostProcess/SceneFilterRendering.h"
+#include "PipelineStateCache.h"
 
 bool IsMobileDistortionActive(const FViewInfo& View)
 {
diff --git a/Engine/Source/Runtime/Renderer/Private/MobileShadingRenderer.cpp b/Engine/Source/Runtime/Renderer/Private/MobileShadingRenderer.cpp
--- a/Engine/Source/Runtime/Renderer/Private/MobileShadingRenderer.cpp
+++ b/Engine/Source/Runtime/Renderer/Private/MobileShadingRenderer.cpp
@@ -228,19 +228,6 @@ void FMobileSceneRenderer::InitViews(FRHICommandListImmediate& RHICmdList)
 	OnStartRender(RHICmdList);
 }
 
-void FMobileSceneRenderer::PostInitViewCustomData()
-{
-	QUICK_SCOPE_CYCLE_COUNTER(STAT_UpdateViewCustomData);
-
-	for (FViewInfo& ViewInfo : Views)
-	{
-		for (const FPrimitiveSceneInfo* PrimitiveSceneInfo : ViewInfo.PrimitivesWithCustomData)
-		{
-			PrimitiveSceneInfo->Proxy->PostInitViewCustomData(ViewInfo, ViewInfo.GetCustomData(PrimitiveSceneInfo->GetIndex()));
-		}
-	}
-}
-
 /** 
 * Renders the view family. 
 */
diff --git a/Engine/Source/Runtime/Renderer/Private/MobileTranslucentRendering.cpp b/Engine/Source/Runtime/Renderer/Private/MobileTranslucentRendering.cpp
--- a/Engine/Source/Runtime/Renderer/Private/MobileTranslucentRendering.cpp
+++ b/Engine/Source/Runtime/Renderer/Private/MobileTranslucentRendering.cpp
@@ -15,7 +15,6 @@
 #include "PostProcess/SceneRenderTargets.h"
 #include "GlobalShader.h"
 #include "SceneRenderTargetParameters.h"
-#include "DrawingPolicy.h"
 #include "SceneRendering.h"
 #include "LightMapRendering.h"
 #include "MaterialShaderType.h"
@@ -220,22 +219,6 @@ public:
 		const bool result = FMeshMaterialShader::Serialize(Ar);
 		return result;
 	}
-
-	void SetParameters(
-		FRHICommandList& RHICmdList,
-		const FMaterialRenderProxy* MaterialRenderProxy,
-		const FMaterial& MaterialResource,
-		const FSceneView& View, 
-		const FDrawingPolicyRenderState& DrawRenderState)
-	{
-
-		FMeshMaterialShader::SetParameters(RHICmdList, GetVertexShader(), MaterialRenderProxy, MaterialResource, View, DrawRenderState.GetViewUniformBuffer(), DrawRenderState.GetPassUniformBuffer());
-	}
-
-	void SetMesh(FRHICommandList& RHICmdList, const FVertexFactory* VertexFactory, const FSceneView& View, const FPrimitiveSceneProxy* Proxy, const FMeshBatchElement& BatchElement, const FDrawingPolicyRenderState& DrawRenderState)
-	{
-		FMeshMaterialShader::SetMesh(RHICmdList, GetVertexShader(), VertexFactory, View, Proxy, BatchElement, DrawRenderState);
-	}
 };
 
 IMPLEMENT_MATERIAL_SHADER_TYPE(, FOpacityOnlyVS, TEXT("/Engine/Private/MobileOpacityShaders.usf"), TEXT("MainVS"), SF_Vertex);
@@ -267,16 +250,6 @@ public:
 
 	FOpacityOnlyPS() {}
 
-	void SetParameters(FRHICommandList& RHICmdList, const FMaterialRenderProxy* MaterialRenderProxy, const FMaterial& MaterialResource, const FSceneView* View, const FDrawingPolicyRenderState& DrawRenderState)
-	{
-		FMeshMaterialShader::SetParameters(RHICmdList, GetPixelShader(), MaterialRenderProxy, MaterialResource, *View, DrawRenderState.GetViewUniformBuffer(), DrawRenderState.GetPassUniformBuffer());
-	}
-
-	void SetMesh(FRHICommandList& RHICmdList, const FVertexFactory* VertexFactory, const FSceneView& View, const FPrimitiveSceneProxy* Proxy, const FMeshBatchElement& BatchElement, const FDrawingPolicyRenderState& DrawRenderState)
-	{
-		FMeshMaterialShader::SetMesh(RHICmdList, GetPixelShader(), VertexFactory, View, Proxy, BatchElement, DrawRenderState);
-	}
-
 	virtual bool Serialize(FArchive& Ar) override
 	{
 		bool bShaderHasOutdatedParameters = FMeshMaterialShader::Serialize(Ar);
diff --git a/Engine/Source/Runtime/Renderer/Private/PostProcess/PostProcessCompositeEditorPrimitives.cpp b/Engine/Source/Runtime/Renderer/Private/PostProcess/PostProcessCompositeEditorPrimitives.cpp
--- a/Engine/Source/Runtime/Renderer/Private/PostProcess/PostProcessCompositeEditorPrimitives.cpp
+++ b/Engine/Source/Runtime/Renderer/Private/PostProcess/PostProcessCompositeEditorPrimitives.cpp
@@ -351,43 +351,6 @@ static void RenderEditorPrimitives(FRHICommandListImmediate& RHICmdList, const F
 	View.BatchedViewElements.Draw(RHICmdList, DrawRenderState, FeatureLevel, bNeedToSwitchVerticalAxis, View, false, 1.0f);
 }
 
-//@todo MeshCommandPipeline remove together with mobile base pass draw policies.
-template<typename TBasePass>
-static void RenderEditorPrimitivesWithDrawPolicies(FRHICommandListImmediate& RHICmdList, const FViewInfo& View, FDrawingPolicyRenderState& DrawRenderState)
-{
-	// Always depth test against other editor primitives
-	DrawRenderState.SetDepthStencilState(TStaticDepthStencilState<
-		true, CF_DepthNearOrEqual,
-		true, CF_Always, SO_Keep, SO_Keep, SO_Replace,
-		false, CF_Always, SO_Keep, SO_Keep, SO_Keep,
-		0xFF, GET_STENCIL_BIT_MASK(RECEIVE_DECAL, 1) | STENCIL_LIGHTING_CHANNELS_MASK(0x7)>::GetRHI());
-
-	typename TBasePass::ContextType Context;
-
-	for (int32 MeshBatchIndex = 0; MeshBatchIndex < View.DynamicEditorMeshElements.Num(); MeshBatchIndex++)
-	{
-		const FMeshBatchAndRelevance& MeshBatchAndRelevance = View.DynamicEditorMeshElements[MeshBatchIndex];
-
-		if (MeshBatchAndRelevance.GetHasOpaqueOrMaskedMaterial() || View.Family->EngineShowFlags.Wireframe)
-		{
-			const FMeshBatch& MeshBatch = *MeshBatchAndRelevance.Mesh;
-			TBasePass::DrawDynamicMesh(RHICmdList, View, Context, MeshBatch, true, DrawRenderState, MeshBatchAndRelevance.PrimitiveSceneProxy, MeshBatch.BatchHitProxyId);
-		}
-	}
-
-	View.EditorSimpleElementCollector.DrawBatchedElements(RHICmdList, DrawRenderState, View, EBlendModeFilter::OpaqueAndMasked, SDPG_World);
-
-	// Draw the base pass for the view's batched mesh elements.
-	DrawViewElements<TBasePass>(RHICmdList, View, DrawRenderState, Context, SDPG_World, false);
-
-	const auto FeatureLevel = View.GetFeatureLevel();
-	const auto ShaderPlatform = GShaderPlatformForFeatureLevel[FeatureLevel];
-	const bool bNeedToSwitchVerticalAxis = RHINeedsToSwitchVerticalAxis(ShaderPlatform);
-
-	// Draw the view's batched simple elements(lines, sprites, etc).
-	View.BatchedViewElements.Draw(RHICmdList, DrawRenderState, FeatureLevel, bNeedToSwitchVerticalAxis, View, false, 1.0f);
-}
-
 static void RenderForegroundEditorPrimitives(FRHICommandListImmediate& RHICmdList, const FViewInfo& View, FDrawingPolicyRenderState& DrawRenderState)
 {
 	FSceneRenderTargets& SceneContext = FSceneRenderTargets::Get(RHICmdList);
@@ -455,35 +418,6 @@ static void RenderForegroundEditorPrimitives(FRHICommandListImmediate& RHICmdLis
 	}
 }
 
-//@todo MeshCommandPipeline remove together with mobile base pass draw policies.
-template<typename TBasePass>
-static void RenderForegroundEditorPrimitivesWithDrawPolicies(FRHICommandListImmediate& RHICmdList, const FViewInfo& View, FDrawingPolicyRenderState& DrawRenderState)
-{
-	FSceneRenderTargets& SceneContext = FSceneRenderTargets::Get(RHICmdList);
-
-	const auto FeatureLevel = View.GetFeatureLevel();
-	const auto ShaderPlatform = GShaderPlatformForFeatureLevel[FeatureLevel];
-	const bool bNeedToSwitchVerticalAxis = RHINeedsToSwitchVerticalAxis(ShaderPlatform);
-
-	// Draw a first time the foreground primitive without depth test to over right depth from non-foreground editor primitives.
-	{
-		DrawRenderState.SetDepthStencilState(TStaticDepthStencilState<true, CF_Always>::GetRHI());
-		View.EditorSimpleElementCollector.DrawBatchedElements(RHICmdList, DrawRenderState, View, EBlendModeFilter::OpaqueAndMasked, SDPG_Foreground);
-		DrawViewElements<TBasePass>(RHICmdList, View, DrawRenderState, typename TBasePass::ContextType(), SDPG_Foreground, false);
-
-		View.TopBatchedViewElements.Draw(RHICmdList, DrawRenderState, FeatureLevel, bNeedToSwitchVerticalAxis, View, false);
-	}
-
-	// Draw a second time the foreground primitive with depth test to have proper depth test between foreground primitives.
-	{
-		DrawRenderState.SetDepthStencilState(TStaticDepthStencilState<true, CF_DepthNearOrEqual>::GetRHI());
-		View.EditorSimpleElementCollector.DrawBatchedElements(RHICmdList, DrawRenderState, View, EBlendModeFilter::OpaqueAndMasked, SDPG_Foreground);
-		DrawViewElements<TBasePass>(RHICmdList, View, DrawRenderState, typename TBasePass::ContextType(), SDPG_Foreground, false);
-
-		View.TopBatchedViewElements.Draw(RHICmdList, DrawRenderState, FeatureLevel, bNeedToSwitchVerticalAxis, View, false);
-	}
-}
-
 template <uint32 MSAASampleCount>
 static void SetCompositePrimitivesShaderTempl(const FRenderingCompositePassContext& Context, bool bComposeAnyNonNullDepth)
 {
diff --git a/Engine/Source/Runtime/Renderer/Private/PrimitiveSceneInfo.cpp b/Engine/Source/Runtime/Renderer/Private/PrimitiveSceneInfo.cpp
--- a/Engine/Source/Runtime/Renderer/Private/PrimitiveSceneInfo.cpp
+++ b/Engine/Source/Runtime/Renderer/Private/PrimitiveSceneInfo.cpp
@@ -305,16 +305,12 @@ void FPrimitiveSceneInfo::AddStaticMeshes(FRHICommandListImmediate& RHICmdList,
 			Mesh.BatchVisibilityId = Scene->StaticMeshBatchVisibility.AddUninitialized().Index;
 			Scene->StaticMeshBatchVisibility[Mesh.BatchVisibilityId] = true;
 		}
-
-		if (bAddToStaticDrawLists)
-		{
-			// By this point, the index buffer render resource must be initialized
-			// Add the static mesh to the appropriate draw lists.
-			Mesh.AddToDrawLists(RHICmdList, Scene);
-		}
 	}
 
-	CacheMeshDrawCommands(RHICmdList);
+	if (bAddToStaticDrawLists)
+	{
+		CacheMeshDrawCommands(RHICmdList);
+	}
 
 #if RHI_RAYTRACING
 	if (IsRayTracingEnabled())
@@ -607,16 +603,6 @@ void FPrimitiveSceneInfo::UpdateStaticMeshes(FRHICommandListImmediate& RHICmdLis
 		}
 	}
 
-	// Remove the primitive's static meshes from the draw lists they're currently in, and re-add them to the appropriate draw lists.
-	for (int32 MeshIndex = 0; MeshIndex < StaticMeshes.Num(); MeshIndex++)
-	{
-		StaticMeshes[MeshIndex].RemoveFromDrawLists();
-		if (bReAddToDrawLists)
-		{
-			StaticMeshes[MeshIndex].AddToDrawLists(RHICmdList, Scene);
-		}
-	}
-
 	RemoveCachedMeshDrawCommands();
 	if (bReAddToDrawLists)
 	{
diff --git a/Engine/Source/Runtime/Renderer/Private/RendererScene.cpp b/Engine/Source/Runtime/Renderer/Private/RendererScene.cpp
--- a/Engine/Source/Runtime/Renderer/Private/RendererScene.cpp
+++ b/Engine/Source/Runtime/Renderer/Private/RendererScene.cpp
@@ -33,7 +33,6 @@
 #include "SceneCore.h"
 #include "PrimitiveSceneInfo.h"
 #include "LightSceneInfo.h"
-#include "StaticMeshDrawList.h"
 #include "LightMapRendering.h"
 #include "AtmosphereRendering.h"
 #include "BasePassRendering.h"
@@ -76,8 +75,6 @@ TGlobalResource< FGlobalDistanceCullFadeUniformBuffer > GDistanceCullFadedInUnif
 /** Global primitive uniform buffer resource containing dither faded in */
 TGlobalResource< FGlobalDitherUniformBuffer > GDitherFadedInUniformBuffer;
 
-SIZE_T FStaticMeshDrawListBase::TotalBytesUsed = 0;
-
 static FThreadSafeCounter FSceneViewState_UniqueID;
 
 /**
@@ -2797,40 +2794,6 @@ void FScene::UpdateStaticDrawLists()
 		});
 }
 
-void FScene::UpdateStaticDrawListsForMaterials_RenderThread(FRHICommandListImmediate& RHICmdList, const TArray<const FMaterial*>& Materials)
-{
-	SCOPE_CYCLE_COUNTER(STAT_Scene_UpdateStaticDrawListsForMaterials_RT);
-
-	// Warning: if any static draw lists are missed here, there will be a crash when trying to render with shaders that have been deleted!
-	TArray<FPrimitiveSceneInfo*> PrimitivesToUpdate;
-	auto SceneFeatureLevel = GetFeatureLevel();
-
-	PositionOnlyDepthDrawList.GetUsedPrimitivesBasedOnMaterials(SceneFeatureLevel, Materials, PrimitivesToUpdate);
-	DepthDrawList.GetUsedPrimitivesBasedOnMaterials(SceneFeatureLevel, Materials, PrimitivesToUpdate);
-	MaskedDepthDrawList.GetUsedPrimitivesBasedOnMaterials(SceneFeatureLevel, Materials, PrimitivesToUpdate);
-	WholeSceneShadowDepthDrawList.GetUsedPrimitivesBasedOnMaterials(SceneFeatureLevel, Materials, PrimitivesToUpdate);
-	WholeSceneReflectiveShadowMapDrawList.GetUsedPrimitivesBasedOnMaterials(SceneFeatureLevel, Materials, PrimitivesToUpdate);
-
-	for (int32 PrimitiveIndex = 0; PrimitiveIndex < PrimitivesToUpdate.Num(); PrimitiveIndex++)
-	{
-		FPrimitiveSceneInfo* Primitive = PrimitivesToUpdate[PrimitiveIndex];
-
-		Primitive->RemoveStaticMeshes();
-		Primitive->AddStaticMeshes(RHICmdList);
-	}
-}
-
-void FScene::UpdateStaticDrawListsForMaterials(const TArray<const FMaterial*>& Materials)
-{
-	ENQUEUE_UNIQUE_RENDER_COMMAND_TWOPARAMETER(
-		FUpdateDrawListsForMaterials,
-		FScene*,Scene,this,
-		TArray<const FMaterial*>,Materials,Materials,
-	{
-		Scene->UpdateStaticDrawListsForMaterials_RenderThread(RHICmdList, Materials);
-	});
-}
-
 /**
  * @return		true if hit proxies should be rendered in this scene.
  */
@@ -2962,95 +2925,6 @@ void FScene::DumpUnbuiltLightInteractions( FOutputDevice& Ar ) const
 	}
 }
 
-/**
- * Logs the provided draw list stats.
- */
-static void LogDrawListStats(FDrawListStats Stats, const TCHAR* DrawListName)
-{
-	if (Stats.NumDrawingPolicies == 0 || Stats.NumMeshes == 0)
-	{
-		UE_LOG(LogRenderer,Log,TEXT("%s: empty"), DrawListName);
-	}
-	else
-	{
-		FString MatchFailedReasons;
-		for (auto& It : Stats.SingleMeshPolicyMatchFailedReasons)
-		{
-			TArray<FStringFormatArg> Args;
-			Args.Emplace(It.Value);
-			Args.Emplace(*It.Key);
-
-			MatchFailedReasons.Append(FString::Format(TEXT("      - {0} ({1})\n"), Args));
-		}
-
-		FString VertexFactoryFreq;
-		for (auto& It : Stats.SingleMeshPolicyVertexFactoryFrequency)
-		{
-			TArray<FStringFormatArg> Args;
-			auto KeyStr = It.Key.ToString();
-
-			Args.Emplace(It.Value);
-			Args.Emplace(*KeyStr);
-
-			VertexFactoryFreq.Append(FString::Format(TEXT("      - {0} ({1})\n"), Args));
-		}
-
-		UE_LOG(LogRenderer,Log,
-			TEXT("%s: %d policies %d meshes\n")
-			TEXT("  - %d median meshes/policy\n")
-			TEXT("  - %f mean meshes/policy\n")
-			TEXT("  - %d max meshes/policy\n")
-			TEXT("  - %d policies with one mesh\n")
-			TEXT("    One mesh policy closest match failure reason:\n%s\n")
-			TEXT("    One mesh policy vertex factory frequencies:\n%s"),
-			DrawListName,
-			Stats.NumDrawingPolicies,
-			Stats.NumMeshes,
-			Stats.MedianMeshesPerDrawingPolicy,
-			(float)Stats.NumMeshes / (float)Stats.NumDrawingPolicies,
-			Stats.MaxMeshesPerDrawingPolicy,
-			Stats.NumSingleMeshDrawingPolicies,
-			*MatchFailedReasons,
-			*VertexFactoryFreq
-			);
-	}
-}
-
-void FScene::DumpStaticMeshDrawListStats() const
-{
-	UE_LOG(LogRenderer,Log,TEXT("Static mesh draw lists for %s:"),
-		World ? *World->GetFullName() : TEXT("[no world]")
-		);
-#define DUMP_DRAW_LIST(Name) LogDrawListStats(Name.GetStats(), TEXT(#Name))
-	DUMP_DRAW_LIST(PositionOnlyDepthDrawList);
-	DUMP_DRAW_LIST(DepthDrawList);
-	DUMP_DRAW_LIST(MaskedDepthDrawList);
-	DUMP_DRAW_LIST(WholeSceneShadowDepthDrawList);
-#undef DUMP_DRAW_LIST
-}
-
-/**
- * Dumps stats for all scenes to the log.
- */
-static void DumpDrawListStats()
-{
-	for (TObjectIterator<UWorld> It; It; ++It)
-	{
-		UWorld* World = *It;
-		if (World && World->Scene)
-		{
-			World->Scene->DumpStaticMeshDrawListStats();
-		}
-	}
-}
-
-static FAutoConsoleCommand GDumpDrawListStatsCmd(
-	TEXT("r.DumpDrawListStats"),
-	TEXT("Dumps static mesh draw list statistics for all scenes associated with ")
-	TEXT("world objects."),
-	FConsoleCommandDelegate::CreateStatic(&DumpDrawListStats)
-	);
-
 /**
  * Exports the scene.
  *
@@ -3073,23 +2947,6 @@ void FScene::ApplyWorldOffset(FVector InOffset)
 	});
 }
 
-// StaticMeshDrawList elements shifting
-template<typename T>
-static void StaticMeshDrawListApplyWorldOffset(T& InList, FVector InOffset)
-{
-	InList.ApplyWorldOffset(InOffset);
-}
-
-// StaticMeshDrawList elements shifting: specialization for an arrays
-template<typename T, int32 N>
-static void StaticMeshDrawListApplyWorldOffset(T(&InList)[N], FVector InOffset)
-{
-	for (int32 i = 0; i < N; i++)
-	{
-		InList[i].ApplyWorldOffset(InOffset);
-	}
-}
-
 void FScene::ApplyWorldOffset_RenderThread(FVector InOffset)
 {
 	QUICK_SCOPE_CYCLE_COUNTER(STAT_SceneApplyWorldOffset);
@@ -3182,12 +3039,6 @@ void FScene::ApplyWorldOffset_RenderThread(FVector InOffset)
 		}
 	}
 	
-	// StaticMeshDrawLists
-	StaticMeshDrawListApplyWorldOffset(PositionOnlyDepthDrawList, InOffset);
-	StaticMeshDrawListApplyWorldOffset(DepthDrawList, InOffset);
-	StaticMeshDrawListApplyWorldOffset(MaskedDepthDrawList, InOffset);
-	StaticMeshDrawListApplyWorldOffset(WholeSceneShadowDepthDrawList, InOffset);
-
 	VelocityData.ApplyOffset(InOffset);
 }
 
@@ -3458,11 +3309,6 @@ void UpdateStaticMeshesForMaterials(const TArray<const FMaterial*>& MaterialReso
 
 void FRendererModule::UpdateStaticDrawListsForMaterials(const TArray<const FMaterial*>& Materials)
 {
-	for (TSet<FSceneInterface*>::TConstIterator SceneIt(AllocatedScenes); SceneIt; ++SceneIt)
-	{
-		(*SceneIt)->UpdateStaticDrawListsForMaterials(Materials);
-	}
-
 	// Update static meshes for a given set of materials in order to recache cached mesh draw commands.
 	UpdateStaticMeshesForMaterials(Materials);
 }
diff --git a/Engine/Source/Runtime/Renderer/Private/RequiredTextureResolutionRendering.h b/Engine/Source/Runtime/Renderer/Private/RequiredTextureResolutionRendering.h
--- a/Engine/Source/Runtime/Renderer/Private/RequiredTextureResolutionRendering.h
+++ b/Engine/Source/Runtime/Renderer/Private/RequiredTextureResolutionRendering.h
@@ -17,7 +17,6 @@ RequiredTextureResolutionRendering.h: Declarations used for the viewmode.
 * Pixel shader that renders texcoord scales.
 * The shader is only compiled with the local vertex factory to prevent multiple compilation.
 * Nothing from the factory is actually used, but the shader must still derive from FMeshMaterialShader.
-* This is required to call FMeshMaterialShader::SetMesh and bind primitive related data.
 */
 class FRequiredTextureResolutionPS : public FDebugViewModePS
 {
diff --git a/Engine/Source/Runtime/Renderer/Private/SceneCore.cpp b/Engine/Source/Runtime/Renderer/Private/SceneCore.cpp
--- a/Engine/Source/Runtime/Renderer/Private/SceneCore.cpp
+++ b/Engine/Source/Runtime/Renderer/Private/SceneCore.cpp
@@ -339,63 +339,6 @@ int32 FStaticMeshRelevance::GetStaticMeshCommandInfoIndex(EMeshPass::Type MeshPa
 	FStaticMesh
 -----------------------------------------------------------------------------*/
 
-void FStaticMesh::LinkDrawList(FStaticMesh::FDrawListElementLink* Link)
-{
-	check(IsInRenderingThread());
-	check(!DrawListLinks.Contains(Link));
-	DrawListLinks.Add(Link);
-}
-
-void FStaticMesh::UnlinkDrawList(FStaticMesh::FDrawListElementLink* Link)
-{
-	check(IsInRenderingThread());
-	verify(DrawListLinks.RemoveSingleSwap(Link) == 1);
-}
-
-void FStaticMesh::AddToDrawLists(FRHICommandListImmediate& RHICmdList, FScene* Scene)
-{
-	const auto FeatureLevel = Scene->GetFeatureLevel();
-	
-	if (!PrimitiveSceneInfo->Proxy->ShouldRenderInMainPass() || !ShouldIncludeDomainInMeshPass(MaterialRenderProxy->GetMaterial(FeatureLevel)->GetMaterialDomain()))
-	{
-		return;
-	}
-
-	if (CastShadow)
-	{
-		FShadowDepthDrawingPolicyFactory::AddStaticMesh(Scene, this);
-	}
-}
-
-void FStaticMesh::RemoveFromDrawLists(bool bMeshIsBeingDestroyed)
-{
-	if (bMeshIsBeingDestroyed)
-	{
-		// This is cheaper than calling RemoveFromDrawLists, since it 
-		// doesn't unlink meshes which are about to be destroyed
-		for (int32 i = 0; i < DrawListLinks.Num(); i++)
-		{
-			DrawListLinks[i]->Remove(false);
-		}
-	}
-	else
-	{
-		// Remove the mesh from all draw lists.
-		while(DrawListLinks.Num())
-		{
-			TRefCountPtr<FStaticMesh::FDrawListElementLink> Link = DrawListLinks[0];
-			const int32 OriginalNumLinks = DrawListLinks.Num();
-			// This will call UnlinkDrawList.
-			Link->Remove(true);
-			check(DrawListLinks.Num() == OriginalNumLinks - 1);
-			if(DrawListLinks.Num())
-			{
-				check(DrawListLinks[0] != Link);
-			}
-		}
-	}
-}
-
 FStaticMesh::~FStaticMesh()
 {
 	FScene* Scene = PrimitiveSceneInfo->Scene;
@@ -406,8 +349,6 @@ FStaticMesh::~FStaticMesh()
 	{
 		Scene->StaticMeshBatchVisibility.RemoveAt(BatchVisibilityId);
 	}
-
-	RemoveFromDrawLists(true);
 }
 
 /** Initialization constructor. */
diff --git a/Engine/Source/Runtime/Renderer/Private/SceneCore.h b/Engine/Source/Runtime/Renderer/Private/SceneCore.h
--- a/Engine/Source/Runtime/Renderer/Private/SceneCore.h
+++ b/Engine/Source/Runtime/Renderer/Private/SceneCore.h
@@ -10,13 +10,11 @@
 #include "Templates/RefCounting.h"
 #include "HitProxies.h"
 #include "MeshBatch.h"
-#include "DrawingPolicy.h"
 #include "MeshPassProcessor.h"
 
 class FLightSceneInfo;
 class FPrimitiveSceneInfo;
 class FScene;
-class FStaticMeshDrawListBase;
 class UExponentialHeightFogComponent;
 
 /**
@@ -129,17 +127,6 @@ class FStaticMesh : public FMeshBatch
 {
 public:
 
-	/**
-	 * An interface to a draw list's reference to this static mesh.
-	 * used to remove the static mesh from the draw list without knowing the draw list type.
-	 */
-	class FDrawListElementLink : public FRefCountedObject
-	{
-	public:
-		virtual bool IsInDrawList(const class FStaticMeshDrawListBase* DrawList) const = 0;
-		virtual void Remove(const bool bUnlinkMesh) = 0;
-	};
-
 	/** The render info for the primitive which created this mesh. */
 	FPrimitiveSceneInfo* PrimitiveSceneInfo;
 
@@ -165,22 +152,7 @@ public:
 
 	~FStaticMesh();
 
-	/** Adds a link from the mesh to its entry in a draw list. */
-	void LinkDrawList(FDrawListElementLink* Link);
-
-	/** Removes a link from the mesh to its entry in a draw list. */
-	void UnlinkDrawList(FDrawListElementLink* Link);
-
-	/** Adds the static mesh to the appropriate draw lists in a scene. */
-	void AddToDrawLists(FRHICommandListImmediate& RHICmdList, FScene* Scene);
-
-	/** Removes the static mesh from all draw lists. */
-	void RemoveFromDrawLists(bool bMeshIsBeingDestroyed = false);
-
 private:
-	/** Links to the draw lists this mesh is an element of. */
-	TArray<TRefCountPtr<FDrawListElementLink> > DrawListLinks;
-
 	/** Private copy constructor. */
 	FStaticMesh(const FStaticMesh& InStaticMesh):
 		FMeshBatch(InStaticMesh),
diff --git a/Engine/Source/Runtime/Renderer/Private/SceneHitProxyRendering.cpp b/Engine/Source/Runtime/Renderer/Private/SceneHitProxyRendering.cpp
--- a/Engine/Source/Runtime/Renderer/Private/SceneHitProxyRendering.cpp
+++ b/Engine/Source/Runtime/Renderer/Private/SceneHitProxyRendering.cpp
@@ -42,16 +42,6 @@ class FHitProxyVS : public FMeshMaterialShader
 
 public:
 
-	void SetParameters(FRHICommandList& RHICmdList, const FMaterialRenderProxy* MaterialRenderProxy,const FSceneView& View,const FDrawingPolicyRenderState& DrawRenderState)
-	{
-		FMeshMaterialShader::SetParameters(RHICmdList, GetVertexShader(), MaterialRenderProxy, *MaterialRenderProxy->GetMaterial(View.GetFeatureLevel()), View, DrawRenderState.GetViewUniformBuffer(), DrawRenderState.GetPassUniformBuffer());
-	}
-
-	void SetMesh(FRHICommandList& RHICmdList, const FVertexFactory* VertexFactory,const FSceneView& View,const FPrimitiveSceneProxy* Proxy,const FMeshBatchElement& BatchElement,const FDrawingPolicyRenderState& DrawRenderState)
-	{
-		FMeshMaterialShader::SetMesh(RHICmdList, GetVertexShader(),VertexFactory,View,Proxy,BatchElement,DrawRenderState);
-	}
-
 	virtual bool Serialize(FArchive& Ar) override
 	{
 		bool bShaderHasOutdatedParameters = FMeshMaterialShader::Serialize(Ar);
@@ -149,21 +139,6 @@ public:
 
 	FHitProxyPS() {}
 
-	void SetParameters(FRHICommandList& RHICmdList, const FMaterialRenderProxy* MaterialRenderProxy,const FSceneView& View,const FDrawingPolicyRenderState& DrawRenderState)
-	{
-		FMeshMaterialShader::SetParameters(RHICmdList, GetPixelShader(), MaterialRenderProxy, *MaterialRenderProxy->GetMaterial(View.GetFeatureLevel()), View, DrawRenderState.GetViewUniformBuffer(), DrawRenderState.GetPassUniformBuffer());
-	}
-
-	void SetMesh(FRHICommandList& RHICmdList, const FVertexFactory* VertexFactory,const FSceneView& View,const FPrimitiveSceneProxy* Proxy,const FMeshBatchElement& BatchElement,const FDrawingPolicyRenderState& DrawRenderState)
-	{
-		FMeshMaterialShader::SetMesh(RHICmdList, GetPixelShader(),VertexFactory,View,Proxy,BatchElement,DrawRenderState);
-	}
-
-	void SetHitProxyId(FRHICommandList& RHICmdList, FHitProxyId HitProxyIdValue)
-	{
-		SetShaderValue(RHICmdList, GetPixelShader(), HitProxyId, HitProxyIdValue.GetColor().ReinterpretAsLinear());
-	}
-
 	void GetShaderBindings(
 		const FScene* Scene,
 		ERHIFeatureLevel::Type FeatureLevel,
@@ -307,7 +282,6 @@ static void DoRenderHitProxies(FRHICommandListImmediate& RHICmdList, const FScen
 		DrawRenderState.SetDepthStencilState(TStaticDepthStencilState<true, CF_DepthNearOrEqual>::GetRHI());
 		DrawRenderState.SetBlendState(TStaticBlendState<>::GetRHI());
 
-		const bool bPreFog = true;
 		const bool bHitTesting = true;
 
 		// Adjust the visibility map for this view
diff --git a/Engine/Source/Runtime/Renderer/Private/ScenePrivate.h b/Engine/Source/Runtime/Renderer/Private/ScenePrivate.h
--- a/Engine/Source/Runtime/Renderer/Private/ScenePrivate.h
+++ b/Engine/Source/Runtime/Renderer/Private/ScenePrivate.h
@@ -34,7 +34,6 @@
 #include "ShadowRendering.h"
 #include "TextureLayout.h"
 #include "SceneRendering.h"
-#include "StaticMeshDrawList.h"
 #include "LightMapRendering.h"
 #include "VelocityRendering.h"
 #include "BasePassRendering.h"
@@ -2437,19 +2436,6 @@ public:
 
 	FCachedPassMeshDrawList CachedDrawLists[EMeshPass::Num];
 
-	/** position-only opaque depth draw list */
-	TStaticMeshDrawList<FPositionOnlyDepthDrawingPolicy> PositionOnlyDepthDrawList;
-	/** opaque depth draw list */
-	TStaticMeshDrawList<FDepthDrawingPolicy> DepthDrawList;
-	/** masked depth draw list */
-	TStaticMeshDrawList<FDepthDrawingPolicy> MaskedDepthDrawList;
-
-	/** Draw list used for rendering whole scene shadow depths. */
-	TStaticMeshDrawList<FShadowDepthDrawingPolicy<false> > WholeSceneShadowDepthDrawList;
-
-	/** Draw list used for rendering whole scene reflective shadow maps.  */
-	TStaticMeshDrawList<FShadowDepthDrawingPolicy<true> > WholeSceneReflectiveShadowMapDrawList;
-
 	/**
 	 * The following arrays are densely packed primitive data needed by various
 	 * rendering passes. PrimitiveSceneInfo->PackedIndex maintains the index
@@ -2704,7 +2690,6 @@ public:
 	virtual void UpdateSpeedTreeWind(double CurrentTime) override;
 	virtual FUniformBufferRHIParamRef GetSpeedTreeUniformBuffer(const FVertexFactory* VertexFactory) const override;
 	virtual void DumpUnbuiltLightInteractions( FOutputDevice& Ar ) const override;
-	virtual void DumpStaticMeshDrawListStats() const override;
 	virtual void UpdateParameterCollections(const TArray<FMaterialParameterCollectionInstanceResource*>& InParameterCollections) override;
 
 	/** Determines whether the scene has atmospheric fog and sun light. */
@@ -2727,9 +2712,6 @@ public:
 	/** Updates all static draw lists. */
 	virtual void UpdateStaticDrawLists() override;
 
-	/** Updates static draw lists for the given set of materials. */
-	virtual void UpdateStaticDrawListsForMaterials(const TArray<const FMaterial*>& Materials) override;
-
 	virtual void Release() override;
 	virtual UWorld* GetWorld() const override { return World; }
 
@@ -2940,9 +2922,6 @@ private:
 	/** Updates all static draw lists. */
 	void UpdateStaticDrawLists_RenderThread(FRHICommandListImmediate& RHICmdList);
 
-	/** Updates static draw lists for the given materials. */
-	void UpdateStaticDrawListsForMaterials_RenderThread(FRHICommandListImmediate& RHICmdList, const TArray<const FMaterial*>& Materials);
-
 	/**
 	 * Shifts scene data by provided delta
 	 * Called on world origin changes
diff --git a/Engine/Source/Runtime/Renderer/Private/SceneRendering.cpp b/Engine/Source/Runtime/Renderer/Private/SceneRendering.cpp
--- a/Engine/Source/Runtime/Renderer/Private/SceneRendering.cpp
+++ b/Engine/Source/Runtime/Renderer/Private/SceneRendering.cpp
@@ -14,7 +14,6 @@
 #include "Components/ReflectionCaptureComponent.h"
 #include "Components/SceneCaptureComponent2D.h"
 #include "Components/SceneCaptureComponentCube.h"
-#include "StaticMeshDrawList.h"
 #include "DeferredShadingRenderer.h"
 #include "DynamicPrimitiveDrawing.h"
 #include "RenderTargetTemp.h"
@@ -55,17 +54,6 @@
 	Globals
 -----------------------------------------------------------------------------*/
 
-static TAutoConsoleVariable<int32> CVarMeshDrawCommandPipeline(
-	TEXT("r.MeshDrawCommandPipeline"),
-	1,
-	TEXT("Whether to use the Mesh Draw Command pipeline, which replaces Drawing Polices and StaticMeshDrawLists.  The Mesh Draw Command pipeline supports Cached Mesh Draw Commands and Dynamic Instancing."),
-	ECVF_RenderThreadSafe);
-
-bool UseMeshDrawCommandPipeline()
-{
-	return CVarMeshDrawCommandPipeline.GetValueOnRenderThread() > 0;
-}
-
 static TAutoConsoleVariable<int32> CVarCachedMeshDrawCommands(
 	TEXT("r.MeshDrawCommands.UseCachedCommands"),
 	1,
@@ -619,7 +607,6 @@ FParallelCommandListSet::FParallelCommandListSet(
 	bSpewBalance = !!CVarRHICmdSpewParallelListBalance.GetValueOnRenderThread();
 	int32 IntBalance = CVarRHICmdBalanceParallelLists.GetValueOnRenderThread();
 	bBalanceCommands = !!IntBalance;
-	bBalanceCommandsWithLastFrame = IntBalance > 1;
 	CommandLists.Reserve(Width * 8);
 	Events.Reserve(Width * 8);
 	NumDrawsIfKnown.Reserve(Width * 8);
@@ -868,7 +855,8 @@ void FViewInfo::Init()
 
 	ViewState = (FSceneViewState*)State;
 	bIsSnapshot = false;
-
+	bHasCustomDepthPrimitives = false;
+	bHasDistortionPrimitives = false;
 	bAllowStencilDither = false;
 
 	ForwardLightingResources = nullptr;
@@ -1551,9 +1539,6 @@ FViewInfo* FViewInfo::CreateSnapshot() const
 	FRWBufferStructured NullOneFramePrimitiveShaderDataBuffer;
 	FMemory::Memcpy(Result->OneFramePrimitiveShaderDataBuffer, NullOneFramePrimitiveShaderDataBuffer);
 
-	FReadBuffer NullReadBuffer;
-	FMemory::Memcpy(Result->OneFramePrimitiveIdBufferEmulation, NullReadBuffer);
-
 	TStaticArray<FParallelMeshDrawCommandPass, EMeshPass::Num> NullParallelMeshDrawCommandPasses;
 	FMemory::Memcpy(Result->ParallelMeshDrawCommandPasses, NullParallelMeshDrawCommandPasses);
 
@@ -1583,7 +1568,6 @@ void FViewInfo::DestroyAllSnapshots()
 		Snapshot->CachedViewUniformShaderParameters.Reset();
 		Snapshot->DynamicPrimitiveShaderData.Empty();
 		Snapshot->OneFramePrimitiveShaderDataBuffer.Release();
-		Snapshot->OneFramePrimitiveIdBufferEmulation.Release();
 		FreeViewInfoSnapshots.Add(Snapshot);
 	}
 	ViewInfoSnapshots.Reset();
@@ -2747,7 +2731,7 @@ void FSceneRenderer::RenderCustomDepthPass(FRHICommandListImmediate& RHICmdList)
 		for(int32 ViewIndex = 0; ViewIndex < Views.Num(); ++ViewIndex)
 		{
 			const FViewInfo& View = Views[ViewIndex];
-			if(View.CustomDepthSet.NumPrims())
+			if (View.bHasCustomDepthPrimitives)
 			{
 				bPrimitives = true;
 				break;
@@ -2832,26 +2816,7 @@ void FSceneRenderer::RenderCustomDepthPass(FRHICommandListImmediate& RHICmdList)
 
 				Scene->UniformBuffers.CustomDepthViewUniformBuffer.UpdateUniformBufferImmediate(CustomDepthViewUniformBufferParameters);
 	
-				if (UseMeshDrawCommandPipeline())
-				{
-					View.ParallelMeshDrawCommandPasses[EMeshPass::CustomDepth].DispatchDraw(nullptr, RHICmdList);
-				}
-				else
-				{
-					FDrawingPolicyRenderState DrawRenderState(View, PassUniformBuffer);
-
-					DrawRenderState.SetBlendState(TStaticBlendState<>::GetRHI());
-					DrawRenderState.SetViewUniformBuffer(Scene->UniformBuffers.CustomDepthViewUniformBuffer);
-
-					const bool bWriteCustomStencilValues = SceneContext.IsCustomDepthPassWritingStencil();
-
-					if (!bWriteCustomStencilValues)
-					{
-						DrawRenderState.SetDepthStencilState(TStaticDepthStencilState<true, CF_DepthNearOrEqual>::GetRHI());
-					}
-
-					View.CustomDepthSet.DrawPrims(RHICmdList, View, DrawRenderState, bWriteCustomStencilValues);
-				}
+				View.ParallelMeshDrawCommandPasses[EMeshPass::CustomDepth].DispatchDraw(nullptr, RHICmdList);
 			}
 		}
 
@@ -3138,7 +3103,6 @@ static void RenderViewFamily_RenderThread(FRHICommandListImmediate& RHICmdList,
 			QUICK_SCOPE_CYCLE_COUNTER(STAT_RenderViewFamily_RenderThread_MemStats);
 
 			// Update scene memory stats that couldn't be tracked continuously
-			SET_MEMORY_STAT(STAT_StaticDrawListMemory, FStaticMeshDrawListBase::TotalBytesUsed);
 			SET_MEMORY_STAT(STAT_RenderingSceneMemory, SceneRenderer->Scene->GetSizeBytes());
 
 			SIZE_T ViewStateMemory = 0;
diff --git a/Engine/Source/Runtime/Renderer/Private/SceneRendering.h b/Engine/Source/Runtime/Renderer/Private/SceneRendering.h
--- a/Engine/Source/Runtime/Renderer/Private/SceneRendering.h
+++ b/Engine/Source/Runtime/Renderer/Private/SceneRendering.h
@@ -25,7 +25,6 @@
 #include "GlobalShader.h"
 #include "PrimitiveViewRelevance.h"
 #include "DistortionRendering.h"
-#include "CustomDepthRendering.h"
 #include "HeightfieldLighting.h"
 #include "GlobalDistanceFieldParameters.h"
 #include "Templates/UniquePtr.h"
@@ -300,89 +299,50 @@ public:
 class FTranslucentPrimSet
 {
 public:
-	/** contains a scene prim and its sort key */
-	struct FTranslucentSortedPrim
+	/** contains a scene prim */
+	struct FTranslucentPrim
 	{
 		/** Default constructor. */
-		FTranslucentSortedPrim() {}
+		FTranslucentPrim() {}
 
-		// @param InPass (first we sort by this)
-		// @param InSortPriority SHRT_MIN .. SHRT_MAX (then we sort by this)
-		// @param InSortKey from UPrimitiveComponent::TranslucencySortPriority e.g. SortByDistance/SortAlongAxis (then by this)
-		FTranslucentSortedPrim(FPrimitiveSceneInfo* InPrimitiveSceneInfo, ETranslucencyPass::Type InPass, int16 InSortPriority, float InSortKey)
+		// @param InPass
+		FTranslucentPrim(FPrimitiveSceneInfo* InPrimitiveSceneInfo, ETranslucencyPass::Type InPass)
 			:	PrimitiveSceneInfo(InPrimitiveSceneInfo)
-			,	SortKey(InSortKey)
-		{
-			SetSortOrder(InPass, InSortPriority);
-		}
-
-		void SetSortOrder(ETranslucencyPass::Type InPass, int16 InSortPriority)
-		{
-			uint32 UpperShort = (uint32)InPass;
-			// 0 .. 0xffff
-			int32 SortPriorityWithoutSign = (int32)InSortPriority - (int32)SHRT_MIN;
-			uint32 LowerShort = SortPriorityWithoutSign;
-
-			check(LowerShort <= 0xffff);
-
-			// top 8 bits are currently unused
-			SortOrder = (UpperShort << 16) | LowerShort;
-		}
+			,	Pass(InPass)
+		{}
 
-		//
 		FPrimitiveSceneInfo* PrimitiveSceneInfo;
-		// single 32bit sort order containing Pass and SortPriority (first we sort by this)
-		uint32 SortOrder;
-		// from UPrimitiveComponent::TranslucencySortPriority (then by this)
-		float SortKey;
+		ETranslucencyPass::Type Pass;
 	};
 
 	/**
 	* Insert a primitive to the translucency rendering list[s]
 	*/
-	
 	static void PlaceScenePrimitive(FPrimitiveSceneInfo* PrimitiveSceneInfo, const FViewInfo& ViewInfo, const FPrimitiveViewRelevance& ViewRelevance,
-		FTranslucentSortedPrim* InArrayStart, int32& InOutArrayNum, FTranslucenyPrimCount& OutCount);
-
-	/**
-	* Sort any primitives that were added to the set back-to-front
-	*/
-	void SortPrimitives();
+		FTranslucentPrim* InArrayStart, int32& InOutArrayNum, FTranslucenyPrimCount& OutCount);
 
 	/** 
 	* @return number of prims to render
 	*/
 	int32 NumPrims() const
 	{
-		return SortedPrims.Num();
+		return Prims.Num();
 	}
 
 	/**
 	* Adds primitives originally created with PlaceScenePrimitive
 	*/
-	void AppendScenePrimitives(FTranslucentSortedPrim* Elements, int32 Num, const FTranslucenyPrimCount& TranslucentPrimitiveCountPerPass);
+	void AppendScenePrimitives(FTranslucentPrim* Elements, int32 Num, const FTranslucenyPrimCount& TranslucentPrimitiveCountPerPass);
 
-	// belongs to SortedPrims
-	FTranslucenyPrimCount SortedPrimsNum;
+	FTranslucenyPrimCount PrimsNum;
 
 private:
 
-	/** sortkey compare class */
-	struct FCompareFTranslucentSortedPrim
-	{
-		FORCEINLINE bool operator()( const FTranslucentSortedPrim& A, const FTranslucentSortedPrim& B ) const
-		{
-			// If priorities are equal sort normally from back to front
-			// otherwise lower sort priorities should render first
-			return ( A.SortOrder == B.SortOrder ) ? ( B.SortKey < A.SortKey ) : ( A.SortOrder < B.SortOrder );
-		}
-	};
-
 	/** list of translucent primitives, sorted after calling Sort() */
-	TArray<FTranslucentSortedPrim,SceneRenderingAllocator> SortedPrims;
+	TArray<FTranslucentPrim,SceneRenderingAllocator> Prims;
 };
 
-template <> struct TIsPODType<FTranslucentPrimSet::FTranslucentSortedPrim> { enum { Value = true }; };
+template <> struct TIsPODType<FTranslucentPrimSet::FTranslucentPrim> { enum { Value = true }; };
 
 /** A batched occlusion primitive. */
 struct FOcclusionPrimitive
@@ -509,7 +469,6 @@ public:
 	bool bBalanceCommands;
 	// see r.RHICmdSpewParallelListBalance
 	bool bSpewBalance;
-	bool bBalanceCommandsWithLastFrame;
 public:
 	TArray<FRHICommandList*,SceneRenderingAllocator> CommandLists;
 	TArray<FGraphEventRef,SceneRenderingAllocator> Events;
@@ -924,31 +883,17 @@ public:
 	/** A map from static mesh ID to a boolean visibility value. */
 	FSceneBitArray StaticMeshVisibilityMap;
 
-	/** A map from static mesh ID to a boolean occluder value. */
-	FSceneBitArray StaticMeshOccluderMap;
-
-	/** A map from static mesh ID to a boolean velocity visibility value. */
-	FSceneBitArray StaticMeshVelocityMap;
-
-	/** A map from static mesh ID to a boolean shadow depth visibility value. */
-	FSceneBitArray StaticMeshShadowDepthMap;
-
 	/** A map from static mesh ID to a boolean dithered LOD fade out value. */
 	FSceneBitArray StaticMeshFadeOutDitheredLODMap;
 
 	/** A map from static mesh ID to a boolean dithered LOD fade in value. */
 	FSceneBitArray StaticMeshFadeInDitheredLODMap;
 
-#if WITH_EDITOR
-	/** A map from static mesh ID to editor selection visibility (whether or not it is selected AND should be drawn).  */
-	FSceneBitArray StaticMeshEditorSelectionMap;
-#endif
-
 	/** Will only contain relevant primitives for view and/or shadow */
 	TArray<FLODMask, SceneRenderingAllocator> PrimitivesLODMask;
 
 	/** An array of batch element visibility masks, valid only for meshes
-	 set visible in either StaticMeshVisibilityMap or StaticMeshShadowDepthMap. */
+	 set visible in StaticMeshVisibilityMap. */
 	TArray<uint64,SceneRenderingAllocator> StaticMeshBatchVisibility;
 
 	/** The dynamic primitives visible in this view. */
@@ -975,14 +920,11 @@ public:
 	/** Set of translucent prims for this view */
 	FTranslucentPrimSet TranslucentPrimSet;
 
-	/** Set of distortion prims for this view */
-	FDistortionPrimSet DistortionPrimSet;
-	
 	/** Set of mesh decal prims for this view */
 	FMeshDecalPrimSet MeshDecalPrimSet;
 	
-	/** Set of CustomDepth prims for this view */
-	FCustomDepthPrimSet CustomDepthSet;
+	bool bHasDistortionPrimitives;
+	bool bHasCustomDepthPrimitives;
 
 	/** Primitives with a volumetric material. */
 	FVolumetricPrimSet VolumetricPrimSet;
@@ -1024,7 +966,6 @@ public:
 	TArray<FPrimitiveUniformShaderParameters> DynamicPrimitiveShaderData;
 
 	FRWBufferStructured OneFramePrimitiveShaderDataBuffer;
-	FReadBuffer OneFramePrimitiveIdBufferEmulation;
 
 	TStaticArray<FParallelMeshDrawCommandPass, EMeshPass::Num> ParallelMeshDrawCommandPasses;
 	
@@ -1036,7 +977,6 @@ public:
 
 	// Primitive CustomData
 	TArray<const FPrimitiveSceneInfo*, SceneRenderingAllocator> PrimitivesWithCustomData;	// Size == Amount of Primitive With Custom Data
-	FSceneBitArray UpdatedPrimitivesWithCustomData;
 	TArray<FMemStackBase, SceneRenderingAllocator> PrimitiveCustomDataMemStack; // Size == 1 global stack + 1 per visibility thread (if multithread)
 
 	/** Parameters for exponential height fog. */
@@ -1813,9 +1753,6 @@ protected:
 	/** Copy scene color from the mobile multi-view render target array to side by side stereo scene color */
 	void CopyMobileMultiViewSceneColor(FRHICommandListImmediate& RHICmdList);
 
-	/** Will update the view custom data. */
-	void PostInitViewCustomData();
-
 	/** Whether GPU particle collisions simulation is allowed. */
 	bool IsGPUParticleCollisionEnabled(const FViewInfo& View);
 	void SortMobileBasePassAfterShadowInit(FExclusiveDepthStencil::Type BasePassDepthStencilAccess, FViewVisibleCommandsPerView& ViewCommandsPerView);
@@ -1949,7 +1886,6 @@ private:
 
 extern FFastVramConfig GFastVRamConfig;
 
-RENDERER_API extern bool UseMeshDrawCommandPipeline();
 extern bool UseCachedMeshDrawCommands();
 extern bool IsDynamicInstancingEnabled();
 
