Removed legacy Drawing Policies and Static Mesh Draw Lists.  These are now replaced by FMeshDrawCommand / FMeshPassProcessor everywhere in the renderer.

diff --git a/Engine/Source/Runtime/Renderer/Private/SceneVisibility.cpp b/Engine/Source/Runtime/Renderer/Private/SceneVisibility.cpp
--- a/Engine/Source/Runtime/Renderer/Private/SceneVisibility.cpp
+++ b/Engine/Source/Runtime/Renderer/Private/SceneVisibility.cpp
@@ -35,6 +35,7 @@
 #include "Engine/LODActor.h"
 #include "GPUScene.h"
 #include "TranslucentRendering.h"
+#include "Async/ParallelFor.h"
 
 /*------------------------------------------------------------------------------
 	Globals
@@ -216,14 +217,6 @@ static TAutoConsoleVariable<int32> CVarParallelInitViews(
 	ECVF_RenderThreadSafe
 	);          
 
-static TAutoConsoleVariable<int32> CVarParallelPostInitViewCustomData(
-	TEXT("r.ParallelViewsCustomDataUpdate"),
-	1,
-	TEXT("Toggles parallel views custom data update. 0 = off; 1 = on"),
-	ECVF_RenderThreadSafe
-);
-
-
 float GLightMaxDrawDistanceScale = 1.0f;
 static FAutoConsoleVariableRef CVarLightMaxDrawDistanceScale(
 	TEXT("r.LightMaxDrawDistanceScale"),
@@ -1665,13 +1658,9 @@ namespace EMarkMaskBits
 {
 	enum Type
 	{
-		StaticMeshShadowDepthMapMask = 0x1,
 		StaticMeshVisibilityMapMask = 0x2,
-		StaticMeshVelocityMapMask = 0x4,
-		StaticMeshOccluderMapMask = 0x8,
 		StaticMeshFadeOutDitheredLODMapMask = 0x10,
 		StaticMeshFadeInDitheredLODMapMask = 0x20,
-		StaticMeshEditorSelectedMask = 0x40,
 	};
 }
 
@@ -1682,7 +1671,6 @@ struct FDrawCommandRelevancePacket
 {
 	FDrawCommandRelevancePacket()
 	{
-		bUseMeshDrawCommandPipeline = UseMeshDrawCommandPipeline();
 		bUseCachedMeshDrawCommands = UseCachedMeshDrawCommands();
 
 		for (int32 PassIndex = 0; PassIndex < EMeshPass::Num; ++PassIndex)
@@ -1694,7 +1682,6 @@ struct FDrawCommandRelevancePacket
 	FPassDrawCommandArray VisibleCachedDrawCommands[EMeshPass::Num];
 	FPassDrawCommandBuildRequestArray DynamicBuildRequests[EMeshPass::Num];
 	int32 NumDynamicBuildRequestElements[EMeshPass::Num];
-	bool bUseMeshDrawCommandPipeline;
 	bool bUseCachedMeshDrawCommands;
 
 	void AddCommandsForMesh(
@@ -1706,41 +1693,38 @@ struct FDrawCommandRelevancePacket
 		bool bCanCache, 
 		EMeshPass::Type PassType)
 	{
-		if (bUseMeshDrawCommandPipeline)
-		{
-			const EShadingPath ShadingPath = Scene->GetShadingPath();
-			const bool bUseCachedMeshCommand = bUseCachedMeshDrawCommands
-				&& !!(FPassProcessorManager::GetPassFlags(ShadingPath, PassType) & EMeshPassFlags::CachedMeshCommands)
-				&& StaticMeshRelevance.bSupportsCachingMeshDrawCommands
-				&& bCanCache;
+		const EShadingPath ShadingPath = Scene->GetShadingPath();
+		const bool bUseCachedMeshCommand = bUseCachedMeshDrawCommands
+			&& !!(FPassProcessorManager::GetPassFlags(ShadingPath, PassType) & EMeshPassFlags::CachedMeshCommands)
+			&& StaticMeshRelevance.bSupportsCachingMeshDrawCommands
+			&& bCanCache;
 
-			if (bUseCachedMeshCommand)
+		if (bUseCachedMeshCommand)
+		{
+			const int32 StaticMeshCommandInfoIndex = StaticMeshRelevance.GetStaticMeshCommandInfoIndex(PassType);
+			if (StaticMeshCommandInfoIndex >= 0)
 			{
-				const int32 StaticMeshCommandInfoIndex = StaticMeshRelevance.GetStaticMeshCommandInfoIndex(PassType);
-				if (StaticMeshCommandInfoIndex >= 0)
-				{
-					const FCachedMeshDrawCommandInfo& CachedMeshDrawCommand = InPrimitiveSceneInfo->StaticMeshCommandInfos[StaticMeshCommandInfoIndex];
-					const FCachedPassMeshDrawList& SceneDrawList = Scene->CachedDrawLists[PassType];
+				const FCachedMeshDrawCommandInfo& CachedMeshDrawCommand = InPrimitiveSceneInfo->StaticMeshCommandInfos[StaticMeshCommandInfoIndex];
+				const FCachedPassMeshDrawList& SceneDrawList = Scene->CachedDrawLists[PassType];
 
-					FVisibleMeshDrawCommand NewVisibleMeshDrawCommand;
+				FVisibleMeshDrawCommand NewVisibleMeshDrawCommand;
 
-					NewVisibleMeshDrawCommand.Setup(
-						&SceneDrawList.MeshDrawCommands[CachedMeshDrawCommand.CommandIndex],
-						PrimitiveIndex,
-						CachedMeshDrawCommand.StateBucketId,
-						CachedMeshDrawCommand.MeshFillMode,
-						CachedMeshDrawCommand.MeshCullMode,
-						CachedMeshDrawCommand.SortKey);
+				NewVisibleMeshDrawCommand.Setup(
+					&SceneDrawList.MeshDrawCommands[CachedMeshDrawCommand.CommandIndex],
+					PrimitiveIndex,
+					CachedMeshDrawCommand.StateBucketId,
+					CachedMeshDrawCommand.MeshFillMode,
+					CachedMeshDrawCommand.MeshCullMode,
+					CachedMeshDrawCommand.SortKey);
 
-					VisibleCachedDrawCommands[(uint32)PassType].AddElement(NewVisibleMeshDrawCommand);
-				}
-			}
-			else
-			{
-				NumDynamicBuildRequestElements[PassType] += StaticMeshRelevance.NumElements;
-				DynamicBuildRequests[PassType].AddElement(&StaticMesh);
+				VisibleCachedDrawCommands[(uint32)PassType].AddElement(NewVisibleMeshDrawCommand);
 			}
 		}
+		else
+		{
+			NumDynamicBuildRequestElements[PassType] += StaticMeshRelevance.NumElements;
+			DynamicBuildRequests[PassType].AddElement(&StaticMesh);
+		}
 	}
 };
 
@@ -1765,12 +1749,12 @@ struct FRelevancePacket
 	FRelevancePrimSet<int32> TranslucentSelfShadowPrimitives;
 	FRelevancePrimSet<FPrimitiveSceneInfo*> VisibleDynamicPrimitives;
 	FMeshPassMask VisibleDynamicMeshesPassMask;
-	FRelevancePrimSet<FTranslucentPrimSet::FTranslucentSortedPrim, ETranslucencyPass::TPT_MAX> TranslucencyPrims;
+	FRelevancePrimSet<FTranslucentPrimSet::FTranslucentPrim, ETranslucencyPass::TPT_MAX> TranslucencyPrims;
 	// belongs to TranslucencyPrims
 	FTranslucenyPrimCount TranslucencyPrimCount;
-	FRelevancePrimSet<FPrimitiveSceneProxy*> DistortionPrimSet;
 	FRelevancePrimSet<FMeshDecalPrimSet::KeyType> MeshDecalPrimSet;
-	FRelevancePrimSet<FPrimitiveSceneProxy*> CustomDepthSet;
+	bool bHasDistortionPrimitives;
+	bool bHasCustomDepthPrimitives;
 	FRelevancePrimSet<FPrimitiveSceneInfo*> LazyUpdatePrimitives;
 	FRelevancePrimSet<FPrimitiveSceneInfo*> DirtyIndirectLightingCacheBufferPrimitives;
 	FRelevancePrimSet<FPrimitiveSceneInfo*> VisibleEditorPrimitives;
@@ -1845,6 +1829,8 @@ struct FRelevancePacket
 		, OutHasDynamicMeshElementsMasks(InOutHasDynamicMeshElementsMasks)
 		, OutHasDynamicEditorMeshElementsMasks(InOutHasDynamicEditorMeshElementsMasks)
 		, MarkMasks(InMarkMasks)
+		, bHasDistortionPrimitives(false)
+		, bHasCustomDepthPrimitives(false)
 		, PrimitiveCustomDataMemStack(InPrimitiveCustomDataMemStack)
 		, OutHasViewCustomDataMasks(InOutHasViewCustomDataMasks)
 		, CombinedShadingModelMask(0)
@@ -2024,8 +2010,7 @@ struct FRelevancePacket
 
 				if (ViewRelevance.bDistortionRelevance)
 				{
-					// Add to set of dynamic distortion primitives
-					DistortionPrimSet.AddPrim(PrimitiveSceneInfo->Proxy);
+					bHasDistortionPrimitives = true;
 				}
 			}
 
@@ -2042,8 +2027,7 @@ struct FRelevancePacket
 
 			if (ViewRelevance.bRenderCustomDepth)
 			{
-				// Add to set of dynamic distortion primitives
-				CustomDepthSet.AddPrim(PrimitiveSceneInfo->Proxy);
+				bHasCustomDepthPrimitives = true;
 			}
 
 			extern bool GUseTranslucencyShadowDepths;
@@ -2201,8 +2185,6 @@ struct FRelevancePacket
 
 					if (ViewRelevance.bShadowRelevance && bDrawShadowDepth && StaticMeshRelevance.CastShadow)
 					{
-						// Mark static mesh as visible in shadows.
-						MarkMask |= EMarkMaskBits::StaticMeshShadowDepthMapMask;
 						bNeedsBatchVisibility = true;
 					}
 
@@ -2211,7 +2193,6 @@ struct FRelevancePacket
 						if (StaticMeshRelevance.bUseForDepthPass && bDrawDepthOnly)
 						{
 							DrawCommandPacket.AddCommandsForMesh(PrimitiveIndex, PrimitiveSceneInfo, StaticMeshRelevance, StaticMesh, Scene, bCanCache, EMeshPass::DepthPass);
-							MarkMask |= EMarkMaskBits::StaticMeshOccluderMapMask;
 						}
 
 						// Mark static mesh as visible for rendering
@@ -2260,7 +2241,6 @@ struct FRelevancePacket
 								&& FVelocityRendering::PrimitiveHasVelocityForView(View, Bounds.BoxSphereBounds, PrimitiveSceneInfo))
 							{
 								DrawCommandPacket.AddCommandsForMesh(PrimitiveIndex, PrimitiveSceneInfo, StaticMeshRelevance, StaticMesh, Scene, bCanCache, EMeshPass::Velocity);
-								MarkMask |= EMarkMaskBits::StaticMeshVelocityMapMask;
 							}
 
 							++NumVisibleStaticMeshElements;
@@ -2308,8 +2288,6 @@ struct FRelevancePacket
 					if(ViewRelevance.bDrawRelevance && ViewRelevance.bEditorStaticSelectionRelevance)
 					{
 						DrawCommandPacket.AddCommandsForMesh(PrimitiveIndex, PrimitiveSceneInfo, StaticMeshRelevance, StaticMesh, Scene, bCanCache, EMeshPass::EditorSelection);
-
-						MarkMask |= EMarkMaskBits::StaticMeshEditorSelectedMask;
 					}
 #endif
 					if (MarkMask)
@@ -2348,9 +2326,9 @@ struct FRelevancePacket
 		VisibleDynamicPrimitives.AppendTo(WriteView.VisibleDynamicPrimitives);
 		VisibleDynamicMeshesPassMask.AppendTo(WriteView.VisibleDynamicMeshesPassMask);
 		WriteView.TranslucentPrimSet.AppendScenePrimitives(TranslucencyPrims.Prims, TranslucencyPrims.NumPrims, TranslucencyPrimCount);
-		DistortionPrimSet.AppendTo(WriteView.DistortionPrimSet);
 		MeshDecalPrimSet.AppendTo(WriteView.MeshDecalPrimSet.Prims);
-		CustomDepthSet.AppendTo(WriteView.CustomDepthSet);
+		WriteView.bHasDistortionPrimitives |= bHasDistortionPrimitives;
+		WriteView.bHasCustomDepthPrimitives |= bHasCustomDepthPrimitives;
 		DirtyIndirectLightingCacheBufferPrimitives.AppendTo(WriteView.DirtyIndirectLightingCacheBufferPrimitives);
 		VolumetricPrimSet.AppendTo(WriteView.VolumetricPrimSet);
 
@@ -2410,7 +2388,6 @@ static void ComputeAndMarkRelevanceForViewParallel(
 	const FMarkRelevantStaticMeshesForViewData ViewData(View);
 
 	int32 NumMesh = View.StaticMeshVisibilityMap.Num();
-	check(View.StaticMeshShadowDepthMap.Num() == NumMesh && View.StaticMeshVelocityMap.Num() == NumMesh && View.StaticMeshOccluderMap.Num() == NumMesh);
 	uint8* RESTRICT MarkMasks = (uint8*)FMemStack::Get().Alloc(NumMesh + 31 , 8); // some padding to simplify the high speed transpose
 	FMemory::Memzero(MarkMasks, NumMesh + 31);
 
@@ -2508,33 +2485,20 @@ static void ComputeAndMarkRelevanceForViewParallel(
 	}
 
 	QUICK_SCOPE_CYCLE_COUNTER(STAT_ComputeAndMarkRelevanceForViewParallel_TransposeMeshBits);
-	check(View.StaticMeshVelocityMap.Num() == NumMesh && 
-		View.StaticMeshShadowDepthMap.Num() == NumMesh && 
-		View.StaticMeshVisibilityMap.Num() == NumMesh && 
-		View.StaticMeshOccluderMap.Num() == NumMesh &&
+	check(View.StaticMeshVisibilityMap.Num() == NumMesh && 
 		View.StaticMeshFadeOutDitheredLODMap.Num() == NumMesh && 
 		View.StaticMeshFadeInDitheredLODMap.Num() == NumMesh
 		);
 	uint32* RESTRICT StaticMeshVisibilityMap_Words = View.StaticMeshVisibilityMap.GetData();
-	uint32* RESTRICT StaticMeshVelocityMap_Words = View.StaticMeshVelocityMap.GetData();
-	uint32* RESTRICT StaticMeshShadowDepthMap_Words = View.StaticMeshShadowDepthMap.GetData();
-	uint32* RESTRICT StaticMeshOccluderMap_Words = View.StaticMeshOccluderMap.GetData();
 	uint32* RESTRICT StaticMeshFadeOutDitheredLODMap_Words = View.StaticMeshFadeOutDitheredLODMap.GetData();
 	uint32* RESTRICT StaticMeshFadeInDitheredLODMap_Words = View.StaticMeshFadeInDitheredLODMap.GetData();
-#if WITH_EDITOR
-	uint32* RESTRICT StaticMeshEditorSelectionMap_Words = View.StaticMeshEditorSelectionMap.GetData();
-#endif
 	const uint64* RESTRICT MarkMasks64 = (const uint64* RESTRICT)MarkMasks;
 	const uint8* RESTRICT MarkMasks8 = MarkMasks;
 	for (int32 BaseIndex = 0; BaseIndex < NumMesh; BaseIndex += 32)
 	{
 		uint32 StaticMeshVisibilityMap_Word = 0;
-		uint32 StaticMeshVelocityMap_Word = 0;
-		uint32 StaticMeshShadowDepthMap_Word = 0;
-		uint32 StaticMeshOccluderMap_Word = 0;
 		uint32 StaticMeshFadeOutDitheredLODMap_Word = 0;
 		uint32 StaticMeshFadeInDitheredLODMap_Word = 0;
-		uint32 StaticMeshEditorSelectionMap_Word = 0;
 		uint32 Mask = 1;
 		bool bAny = false;
 		for (int32 QWordIndex = 0; QWordIndex < 4; QWordIndex++)
@@ -2545,14 +2509,8 @@ static void ComputeAndMarkRelevanceForViewParallel(
 				{
 					uint8 MaskMask = *MarkMasks8;
 					StaticMeshVisibilityMap_Word |= (MaskMask & EMarkMaskBits::StaticMeshVisibilityMapMask) ? Mask : 0;
-					StaticMeshVelocityMap_Word |= (MaskMask & EMarkMaskBits::StaticMeshVelocityMapMask) ? Mask : 0;
-					StaticMeshShadowDepthMap_Word |= (MaskMask & EMarkMaskBits::StaticMeshShadowDepthMapMask) ? Mask : 0;
-					StaticMeshOccluderMap_Word |= (MaskMask & EMarkMaskBits::StaticMeshOccluderMapMask) ? Mask : 0;
 					StaticMeshFadeOutDitheredLODMap_Word |= (MaskMask & EMarkMaskBits::StaticMeshFadeOutDitheredLODMapMask) ? Mask : 0;
 					StaticMeshFadeInDitheredLODMap_Word |= (MaskMask & EMarkMaskBits::StaticMeshFadeInDitheredLODMapMask) ? Mask : 0;
-#if WITH_EDITOR
-					StaticMeshEditorSelectionMap_Word |= (MaskMask & EMarkMaskBits::StaticMeshEditorSelectedMask) ? Mask : 0;
-#endif
 				}
 				bAny = true;
 			}
@@ -2564,26 +2522,14 @@ static void ComputeAndMarkRelevanceForViewParallel(
 		}
 		if (bAny)
 		{
-			checkSlow(!*StaticMeshVisibilityMap_Words && !*StaticMeshVelocityMap_Words && !*StaticMeshShadowDepthMap_Words && !*StaticMeshOccluderMap_Words && !*StaticMeshFadeOutDitheredLODMap_Words && !*StaticMeshFadeInDitheredLODMap_Words);
+			checkSlow(!*StaticMeshVisibilityMap_Words && !*StaticMeshFadeOutDitheredLODMap_Words && !*StaticMeshFadeInDitheredLODMap_Words);
 			*StaticMeshVisibilityMap_Words = StaticMeshVisibilityMap_Word;
-			*StaticMeshVelocityMap_Words = StaticMeshVelocityMap_Word;
-			*StaticMeshShadowDepthMap_Words = StaticMeshShadowDepthMap_Word;
-			*StaticMeshOccluderMap_Words = StaticMeshOccluderMap_Word;
 			*StaticMeshFadeOutDitheredLODMap_Words = StaticMeshFadeOutDitheredLODMap_Word;
 			*StaticMeshFadeInDitheredLODMap_Words = StaticMeshFadeInDitheredLODMap_Word;
-#if WITH_EDITOR
-			*StaticMeshEditorSelectionMap_Words = StaticMeshEditorSelectionMap_Word;
-#endif
 		}
 		StaticMeshVisibilityMap_Words++;
-		StaticMeshVelocityMap_Words++;
-		StaticMeshShadowDepthMap_Words++;
-		StaticMeshOccluderMap_Words++;
 		StaticMeshFadeOutDitheredLODMap_Words++;
 		StaticMeshFadeInDitheredLODMap_Words++;
-#if WITH_EDITOR
-		StaticMeshEditorSelectionMap_Words++;
-#endif
 	}
 }
 
@@ -3309,13 +3255,9 @@ void FSceneRenderer::ComputeViewVisibility(FRHICommandListImmediate& RHICmdList,
 		View.PrimitiveFadeUniformBuffers.AddZeroed(Scene->Primitives.Num());
 		View.PrimitiveFadeUniformBufferMap.Init(false, Scene->Primitives.Num());
 		View.StaticMeshVisibilityMap.Init(false,Scene->StaticMeshes.GetMaxIndex());
-		View.StaticMeshOccluderMap.Init(false,Scene->StaticMeshes.GetMaxIndex());
 		View.StaticMeshFadeOutDitheredLODMap.Init(false,Scene->StaticMeshes.GetMaxIndex());
 		View.StaticMeshFadeInDitheredLODMap.Init(false,Scene->StaticMeshes.GetMaxIndex());
-		View.StaticMeshVelocityMap.Init(false,Scene->StaticMeshes.GetMaxIndex());
-		View.StaticMeshShadowDepthMap.Init(false,Scene->StaticMeshes.GetMaxIndex());
 		View.StaticMeshBatchVisibility.AddZeroed(Scene->StaticMeshBatchVisibility.GetMaxIndex());
-		View.UpdatedPrimitivesWithCustomData.Init(false, Scene->Primitives.Num());
 		View.PrimitivesLODMask.Init(FLODMask(), Scene->Primitives.Num());
 
 		View.PrimitivesCustomData.Init(nullptr, Scene->Primitives.Num());
@@ -3329,10 +3271,6 @@ void FSceneRenderer::ComputeViewVisibility(FRHICommandListImmediate& RHICmdList,
 
 		View.VisibleLightInfos.Empty(Scene->Lights.GetMaxIndex());
 
-#if WITH_EDITOR
-		View.StaticMeshEditorSelectionMap.Init(false, Scene->StaticMeshes.GetMaxIndex());
-#endif
-
 		// The dirty list allocation must take into account the max possible size because when GILCUpdatePrimTaskEnabled is true,
 		// the indirect lighting cache will be update on by threaded job, which can not do reallocs on the buffer (since it uses the SceneRenderingAllocator).
 		View.DirtyIndirectLightingCacheBufferPrimitives.Reserve(Scene->Primitives.Num());
@@ -3599,15 +3537,12 @@ void FSceneRenderer::ComputeViewVisibility(FRHICommandListImmediate& RHICmdList,
 	// Gather FMeshBatches from scene proxies
 	GatherDynamicMeshElements(Views, Scene, ViewFamily, HasDynamicMeshElementsMasks, HasDynamicEditorMeshElementsMasks, HasViewCustomDataMasks, MeshCollector);
 
-	if (UseMeshDrawCommandPipeline())
+	for (int32 ViewIndex = 0; ViewIndex < Views.Num(); ViewIndex++)
 	{
-		for (int32 ViewIndex = 0; ViewIndex < Views.Num(); ViewIndex++)
-		{
-			FViewInfo& View = Views[ViewIndex];
-			FViewCommands& ViewCommands = ViewCommandsPerView[ViewIndex];
+		FViewInfo& View = Views[ViewIndex];
+		FViewCommands& ViewCommands = ViewCommandsPerView[ViewIndex];
 
-			SetupMeshPass(View, BasePassDepthStencilAccess, ViewCommands);
-		}
+		SetupMeshPass(View, BasePassDepthStencilAccess, ViewCommands);
 	}
 
 	INC_DWORD_STAT_BY(STAT_ProcessedPrimitives,NumProcessedPrimitives);
@@ -3625,7 +3560,6 @@ void FSceneRenderer::PostVisibilityFrameSetup(FILCUpdatePrimTaskData& OutILCTask
 		{		
 			FViewInfo& View = Views[ViewIndex];
 
-			View.TranslucentPrimSet.SortPrimitives();
 			View.MeshDecalPrimSet.SortPrimitives();
 
 			if (View.State)
@@ -3905,11 +3839,6 @@ bool FDeferredShadingSceneRenderer::InitViews(FRHICommandListImmediate& RHICmdLi
 		InitViewsPossiblyAfterPrepass(RHICmdList, ILCTaskData, UpdateViewCustomDataEvents);
 	}
 
-	if (!UseMeshDrawCommandPipeline())
-	{
-		PostInitViewCustomData(UpdateViewCustomDataEvents);
-	}
-
 	{
 		QUICK_SCOPE_CYCLE_COUNTER(STAT_InitViews_InitRHIResources);
 		// initialize per-view uniform buffer.
@@ -3953,105 +3882,6 @@ bool FDeferredShadingSceneRenderer::InitViews(FRHICommandListImmediate& RHICmdLi
 	return bDoInitViewAftersPrepass;
 }
 
-class FPostInitViewCustomDataTask
-{
-private:
-	FViewInfo* ViewInfo;
-	const int32 PrimitiveStartIndex;
-	const int32 PrimitiveCount;
-
-public:
-	FPostInitViewCustomDataTask(FViewInfo* InViewInfo, int32 InPrimitiveStartIndex, int32 InPrimitiveCount)
-		: ViewInfo(InViewInfo)
-		, PrimitiveStartIndex(InPrimitiveStartIndex)
-		, PrimitiveCount(InPrimitiveCount)
-	{}
-
-	FORCEINLINE TStatId GetStatId() const
-	{
-		RETURN_QUICK_DECLARE_CYCLE_STAT(FPostInitViewCustomDataTask, STATGROUP_TaskGraphTasks);
-	}
-
-	ENamedThreads::Type GetDesiredThread()
-	{
-		return ENamedThreads::AnyHiPriThreadHiPriTask;
-	}
-
-	static ESubsequentsMode::Type GetSubsequentsMode() { return ESubsequentsMode::TrackSubsequents; }
-
-	void DoTask(ENamedThreads::Type CurrentThread, const FGraphEventRef& MyCompletionGraphEvent)
-	{
-		for (int32 i = PrimitiveStartIndex; i < PrimitiveStartIndex + PrimitiveCount; ++i)
-		{
-			if (ViewInfo->PrimitivesWithCustomData.IsValidIndex(i))
-			{
-				check(ViewInfo->UpdatedPrimitivesWithCustomData.IsValidIndex(i));
-
-				if (!ViewInfo->UpdatedPrimitivesWithCustomData[i])
-				{
-					const FPrimitiveSceneInfo* PrimitiveSceneInfo = ViewInfo->PrimitivesWithCustomData[i];
-					check(PrimitiveSceneInfo != nullptr);
-
-					PrimitiveSceneInfo->Proxy->PostInitViewCustomData(*ViewInfo, ViewInfo->GetCustomData(PrimitiveSceneInfo->GetIndex()));
-					ViewInfo->UpdatedPrimitivesWithCustomData[i] = true;
-				}
-			}
-		}
-	}
-};
-
-void FDeferredShadingSceneRenderer::PostInitViewCustomData(FGraphEventArray& OutUpdateEvents)
-{
-	if (FApp::ShouldUseThreadingForPerformance() && CVarParallelPostInitViewCustomData.GetValueOnRenderThread() > 0)
-	{
-		QUICK_SCOPE_CYCLE_COUNTER(STAT_PostInitViewCustomData_AsyncTask);
-
-		const int32 MaxPrimitiveUpdateTaskCount = 10;
-		const int32 MinPrimitiveCountByTask = 100;
-		
-		for (FViewInfo& ViewInfo : Views)
-		{
-			if (ViewInfo.PrimitivesWithCustomData.Num() > 0)
-			{
-				const int32 BatchSize = FMath::Max(FMath::Max(FMath::RoundToInt((float)ViewInfo.PrimitivesWithCustomData.Num() / (float)MaxPrimitiveUpdateTaskCount), 1), MinPrimitiveCountByTask);
-
-				int32 UpdateCountLeft = ViewInfo.PrimitivesWithCustomData.Num();
-				int32 StartIndex = 0;
-				int32 CurrentBatchSize = BatchSize;
-
-				while (UpdateCountLeft > 0)
-				{
-					if (UpdateCountLeft - CurrentBatchSize < 0)
-					{
-						CurrentBatchSize = UpdateCountLeft;
-					}
-
-					OutUpdateEvents.Add(TGraphTask<FPostInitViewCustomDataTask>::CreateTask(nullptr, ENamedThreads::GetRenderThread()).ConstructAndDispatchWhenReady(&ViewInfo, StartIndex, CurrentBatchSize));
-
-					StartIndex += CurrentBatchSize;
-					UpdateCountLeft -= CurrentBatchSize;
-				}
-			}
-		}
-	}
-	else
-	{
-		QUICK_SCOPE_CYCLE_COUNTER(STAT_PostInitViewCustomData);
-
-		for (FViewInfo& ViewInfo : Views)
-		{
-			for (const FPrimitiveSceneInfo* PrimitiveSceneInfo : ViewInfo.PrimitivesWithCustomData)
-			{
-				if (!ViewInfo.UpdatedPrimitivesWithCustomData[PrimitiveSceneInfo->GetIndex()])
-				{
-					PrimitiveSceneInfo->Proxy->PostInitViewCustomData(ViewInfo, ViewInfo.GetCustomData(PrimitiveSceneInfo->GetIndex()));
-					ViewInfo.UpdatedPrimitivesWithCustomData[PrimitiveSceneInfo->GetIndex()] = true;
-				}
-			}
-		}
-	}
-}
-
 void FDeferredShadingSceneRenderer::SetupSceneReflectionCaptureBuffer(FRHICommandListImmediate& RHICmdList)
 {
 	FReflectionCaptureShaderData SamplePositionsBuffer;
diff --git a/Engine/Source/Runtime/Renderer/Private/ShaderBaseClasses.cpp b/Engine/Source/Runtime/Renderer/Private/ShaderBaseClasses.cpp
--- a/Engine/Source/Runtime/Renderer/Private/ShaderBaseClasses.cpp
+++ b/Engine/Source/Runtime/Renderer/Private/ShaderBaseClasses.cpp
@@ -83,8 +83,6 @@ FMaterialShader::FMaterialShader(const FMaterialShaderType::CompiledShaderInitia
 	}
 
 	SceneTextureParameters.Bind(Initializer);
-
-	InstanceOffset.Bind(Initializer.ParameterMap, TEXT("InstanceOffset"));
 }
 
 FUniformBufferRHIParamRef FMaterialShader::GetParameterCollectionBuffer(const FGuid& Id, const FSceneInterface* SceneInterface) const
@@ -105,7 +103,7 @@ FUniformBufferRHIParamRef FMaterialShader::GetParameterCollectionBuffer(const FG
 }
 
 #if !(UE_BUILD_TEST || UE_BUILD_SHIPPING || !WITH_EDITOR)
-void FMaterialShader::VerifyExpressionAndShaderMaps(const FMaterialRenderProxy* MaterialRenderProxy, const FMaterial& Material, const FUniformExpressionCache* UniformExpressionCache)
+void FMaterialShader::VerifyExpressionAndShaderMaps(const FMaterialRenderProxy* MaterialRenderProxy, const FMaterial& Material, const FUniformExpressionCache* UniformExpressionCache) const
 {
 	// Validate that the shader is being used for a material that matches the uniform expression set the shader was compiled for.
 	const FUniformExpressionSet& MaterialUniformExpressionSet = Material.GetRenderingThreadShaderMap()->GetUniformExpressionSet();
@@ -379,6 +377,11 @@ void FMaterialShader::GetShaderBindings(
 
 	const FUniformExpressionCache& UniformExpressionCache = MaterialRenderProxy.UniformExpressionCache[FeatureLevel];
 	check(UniformExpressionCache.bUpToDate && UniformExpressionCache.UniformBuffer);
+
+#if !(UE_BUILD_TEST || UE_BUILD_SHIPPING || !WITH_EDITOR)
+	VerifyExpressionAndShaderMaps(&MaterialRenderProxy, Material, &UniformExpressionCache);
+#endif
+
 	ShaderBindings.Add(MaterialUniformBuffer, UniformExpressionCache.UniformBuffer);
 
 	{
@@ -496,8 +499,6 @@ bool FMaterialShader::Serialize(FArchive& Ar)
 	Ar << PageTable;
 	Ar << PageTableSampler;
 
-	Ar << InstanceOffset;
-
 	return bShaderHasOutdatedParameters;
 }
 
@@ -511,63 +512,6 @@ uint32 FMaterialShader::GetAllocatedSize() const
 	;
 }
 
-
-template< typename ShaderRHIParamRef >
-void FMeshMaterialShader::SetMesh(
-	FRHICommandList& RHICmdList,
-	const ShaderRHIParamRef ShaderRHI,
-	const FVertexFactory* VertexFactory,
-	const FSceneView& View,
-	const FPrimitiveSceneProxy* Proxy,
-	const FMeshBatchElement& BatchElement,
-	const FDrawingPolicyRenderState& DrawRenderState,
-	uint32 DataFlags )
-{
-	// Set the mesh for the vertex factory
-	VertexFactoryParameters.SetMesh(RHICmdList, this,VertexFactory,View,BatchElement, DataFlags);
-		
-	if (BatchElement.PrimitiveUniformBuffer)
-	{
-		SetUniformBufferParameter(RHICmdList, ShaderRHI,GetUniformBufferParameter<FPrimitiveUniformShaderParameters>(),BatchElement.PrimitiveUniformBuffer);
-	}
-	else if (BatchElement.PrimitiveUniformBufferResource)
-	{
-		SetUniformBufferParameter(RHICmdList, ShaderRHI,GetUniformBufferParameter<FPrimitiveUniformShaderParameters>(),*BatchElement.PrimitiveUniformBufferResource);
-	}
-
-	TShaderUniformBufferParameter<FDistanceCullFadeUniformShaderParameters> LODParameter = GetUniformBufferParameter<FDistanceCullFadeUniformShaderParameters>();
-	if( LODParameter.IsBound() )
-	{
-		SetUniformBufferParameter(RHICmdList, ShaderRHI,LODParameter,GetPrimitiveFadeUniformBufferParameter(View, Proxy));
-	}
-
-	TShaderUniformBufferParameter<FDitherUniformShaderParameters> DitherParameter = GetUniformBufferParameter<FDitherUniformShaderParameters>();
-	if (DitherParameter.IsBound())
-	{
-		// Dither is currently disabled in mesh draw policy rendering path.
-		SetUniformBufferParameter(RHICmdList, ShaderRHI, DitherParameter, GDitherFadedInUniformBuffer.GetUniformBufferRHI());
-	}
-}
-
-#define IMPLEMENT_MESH_MATERIAL_SHADER_SetMesh( ShaderRHIParamRef ) \
-	template RENDERER_API void FMeshMaterialShader::SetMesh< ShaderRHIParamRef >( \
-		FRHICommandList& RHICmdList,					 \
-		const ShaderRHIParamRef ShaderRHI,				 \
-		const FVertexFactory* VertexFactory,			 \
-		const FSceneView& View,							 \
-		const FPrimitiveSceneProxy* Proxy,				 \
-		const FMeshBatchElement& BatchElement,			 \
-		const FDrawingPolicyRenderState& DrawRenderState,\
-		uint32 DataFlags								 \
-	);
-
-IMPLEMENT_MESH_MATERIAL_SHADER_SetMesh( FVertexShaderRHIParamRef );
-IMPLEMENT_MESH_MATERIAL_SHADER_SetMesh( FHullShaderRHIParamRef );
-IMPLEMENT_MESH_MATERIAL_SHADER_SetMesh( FDomainShaderRHIParamRef );
-IMPLEMENT_MESH_MATERIAL_SHADER_SetMesh( FGeometryShaderRHIParamRef );
-IMPLEMENT_MESH_MATERIAL_SHADER_SetMesh( FPixelShaderRHIParamRef );
-IMPLEMENT_MESH_MATERIAL_SHADER_SetMesh( FComputeShaderRHIParamRef );
-
 void FMeshMaterialShader::GetShaderBindings(
 	const FScene* Scene,
 	ERHIFeatureLevel::Type FeatureLevel,
@@ -601,7 +545,6 @@ void FMeshMaterialShader::GetElementShaderBindings(
 {
 	checkSlow(ShaderBindings.Frequency == GetType()->GetFrequency());
 
-	//@todo MeshCommandPipeline - remove from FMaterialShader
 	VertexFactoryParameters.GetElementShaderBindings(Scene, ViewIfDynamicMeshCommand, this, bShaderRequiresPositionOnlyStream, FeatureLevel, VertexFactory, BatchElement, ShaderBindings, VertexStreams);
 		
 	if (UseGPUScene(GMaxRHIShaderPlatform, FeatureLevel) && VertexFactory->GetPrimitiveIdStreamIndex(bShaderRequiresPositionOnlyStream) >= 0)
@@ -636,24 +579,3 @@ uint32 FMeshMaterialShader::GetAllocatedSize() const
 	return FMaterialShader::GetAllocatedSize()
 		+ VertexFactoryParameters.GetAllocatedSize();
 }
-
-//@todo MeshCommandPipeline Remove it together with draw policies.
-FUniformBufferRHIParamRef FMeshMaterialShader::GetPrimitiveFadeUniformBufferParameter(const FSceneView& View, const FPrimitiveSceneProxy* Proxy)
-{
-	FUniformBufferRHIParamRef FadeUniformBuffer = NULL;
-	if( Proxy != NULL )
-	{
-		const FPrimitiveSceneInfo* PrimitiveSceneInfo = Proxy->GetPrimitiveSceneInfo();
-		int32 PrimitiveIndex = PrimitiveSceneInfo->GetIndex();
-
-		// This cast should always be safe. Check it :)
-		checkSlow(View.bIsViewInfo);
-		const FViewInfo& ViewInfo = (const FViewInfo&)View;
-		FadeUniformBuffer = ViewInfo.PrimitiveFadeUniformBuffers[PrimitiveIndex];
-	}
-	if (FadeUniformBuffer == NULL)
-	{
-		FadeUniformBuffer = GDistanceCullFadedInUniformBuffer.GetUniformBufferRHI();
-	}
-	return FadeUniformBuffer;
-}
diff --git a/Engine/Source/Runtime/Renderer/Private/ShaderBaseClasses.h b/Engine/Source/Runtime/Renderer/Private/ShaderBaseClasses.h
--- a/Engine/Source/Runtime/Renderer/Private/ShaderBaseClasses.h
+++ b/Engine/Source/Runtime/Renderer/Private/ShaderBaseClasses.h
@@ -10,7 +10,6 @@
 #include "UniformBuffer.h"
 #include "Shader.h"
 #include "MeshMaterialShader.h"
-#include "DrawingPolicy.h"
 
 class FPrimitiveSceneProxy;
 struct FMeshBatchElement;
@@ -72,21 +71,6 @@ public:
 	}
 
 	FBaseHS() {}
-
-	void SetParameters(
-		FRHICommandList& RHICmdList, 
-		const FMaterialRenderProxy* MaterialRenderProxy,
-		const FSceneView& View,
-		const TUniformBufferRef<FViewUniformShaderParameters>& ViewUniformBuffer,
-		FUniformBufferRHIParamRef PassUniformBufferValue)
-	{
-		FMeshMaterialShader::SetParameters(RHICmdList, (FHullShaderRHIParamRef)GetHullShader(), MaterialRenderProxy, *MaterialRenderProxy->GetMaterial(View.GetFeatureLevel()), View, ViewUniformBuffer, PassUniformBufferValue);
-	}
-
-	void SetMesh(FRHICommandList& RHICmdList, const FVertexFactory* VertexFactory,const FSceneView& View,const FPrimitiveSceneProxy* Proxy,const FMeshBatchElement& BatchElement, const FDrawingPolicyRenderState& DrawRenderState)
-	{
-		FMeshMaterialShader::SetMesh(RHICmdList, (FHullShaderRHIParamRef)GetHullShader(),VertexFactory,View,Proxy,BatchElement,DrawRenderState);
-	}
 };
 
 /** Base Domain shader for drawing policy rendering */
@@ -127,20 +111,5 @@ public:
 	}
 
 	FBaseDS() {}
-
-	void SetParameters(
-		FRHICommandList& RHICmdList, 
-		const FMaterialRenderProxy* MaterialRenderProxy,
-		const FSceneView& View,
-		const TUniformBufferRef<FViewUniformShaderParameters>& ViewUniformBuffer,
-		FUniformBufferRHIParamRef PassUniformBufferValue)
-	{
-		FMeshMaterialShader::SetParameters(RHICmdList, (FDomainShaderRHIParamRef)GetDomainShader(), MaterialRenderProxy, *MaterialRenderProxy->GetMaterial(View.GetFeatureLevel()), View, ViewUniformBuffer, PassUniformBufferValue);
-	}
-
-	void SetMesh(FRHICommandList& RHICmdList, const FVertexFactory* VertexFactory,const FSceneView& View,const FPrimitiveSceneProxy* Proxy,const FMeshBatchElement& BatchElement, const FDrawingPolicyRenderState& DrawRenderState)
-	{
-		FMeshMaterialShader::SetMesh(RHICmdList, (FDomainShaderRHIParamRef)GetDomainShader(),VertexFactory,View,Proxy,BatchElement,DrawRenderState);
-	}
 };
 
diff --git a/Engine/Source/Runtime/Renderer/Private/ShadowDepthRendering.cpp b/Engine/Source/Runtime/Renderer/Private/ShadowDepthRendering.cpp
--- a/Engine/Source/Runtime/Renderer/Private/ShadowDepthRendering.cpp
+++ b/Engine/Source/Runtime/Renderer/Private/ShadowDepthRendering.cpp
@@ -26,7 +26,6 @@
 #include "GlobalShader.h"
 #include "MaterialShaderType.h"
 #include "MaterialShader.h"
-#include "DrawingPolicy.h"
 #include "MeshMaterialShader.h"
 #include "ShaderBaseClasses.h"
 #include "ShadowRendering.h"
@@ -556,30 +555,6 @@ void OverrideWithDefaultMaterialForShadowDepth(
 	}
 }
 
-/*-----------------------------------------------------------------------------
-FShadowDepthDrawingPolicy
------------------------------------------------------------------------------*/
-
-template <bool bRenderingReflectiveShadowMaps>
-void FShadowDepthDrawingPolicy<bRenderingReflectiveShadowMaps>::UpdateElementState(FShadowStaticMeshElement& State, ERHIFeatureLevel::Type InFeatureLevel)
-{
-	FMeshDrawingPolicyOverrideSettings OverrideSettings = ComputeMeshOverrideSettings(*State.Mesh);
-	OverrideSettings.MeshOverrideFlags |= State.bIsTwoSided ? EDrawingPolicyOverrideFlags::TwoSided : EDrawingPolicyOverrideFlags::None;
-
-	// can be optimized
-	*this = FShadowDepthDrawingPolicy(
-		State.MaterialResource,
-		bDirectionalLight,
-		bOnePassPointLightShadow,
-		bPreShadow,
-		OverrideSettings,
-		InFeatureLevel,
-		State.Mesh->VertexFactory,
-		State.RenderProxy,
-		State.Mesh->ReverseCulling);
-}
-
-
 template <bool bRenderingReflectiveShadowMaps>
 void GetShadowDepthPassShaders(
 	const FMaterial& Material,
@@ -715,464 +690,6 @@ void GetShadowDepthPassShaders(
 	}
 }
 
-template <bool bRenderingReflectiveShadowMaps>
-FShadowDepthDrawingPolicy<bRenderingReflectiveShadowMaps>::FShadowDepthDrawingPolicy(
-	const FMaterial* InMaterialResource,
-	bool bInDirectionalLight,
-	bool bInOnePassPointLightShadow,
-	bool bInPreShadow,
-	const FMeshDrawingPolicyOverrideSettings& InOverrideSettings,
-	ERHIFeatureLevel::Type InFeatureLevel,
-	const FVertexFactory* InVertexFactory,
-	const FMaterialRenderProxy* InMaterialRenderProxy,
-	bool bInReverseCulling
-) :
-	FMeshDrawingPolicy(InVertexFactory, InMaterialRenderProxy, *InMaterialResource, InOverrideSettings),
-	GeometryShader(NULL),
-	FeatureLevel(InFeatureLevel),
-	bDirectionalLight(bInDirectionalLight),
-	bReverseCulling(bInReverseCulling),
-	bOnePassPointLightShadow(bInOnePassPointLightShadow),
-	bPreShadow(bInPreShadow)
-{
-	check(!bInOnePassPointLightShadow || !bRenderingReflectiveShadowMaps);
-
-	InstanceFactor = !bOnePassPointLightShadow || RHISupportsGeometryShaders(GShaderPlatformForFeatureLevel[FeatureLevel]) ? 1 : 6;
-
-	if (!InVertexFactory)
-	{
-		// dummy object, needs call to UpdateElementState() to be fully initialized
-		return;
-	}
-
-	// Use perspective correct shadow depths for shadow types which typically render low poly meshes into the shadow depth buffer.
-	// Depth will be interpolated to the pixel shader and written out, which disables HiZ and double speed Z.
-	// Directional light shadows use an ortho projection and can use the non-perspective correct path without artifacts.
-	// One pass point lights don't output a linear depth, so they are already perspective correct.
-	const bool bUsePerspectiveCorrectShadowDepths = !bInDirectionalLight && !bInOnePassPointLightShadow;
-
-	HullShader = NULL;
-	DomainShader = NULL;
-
-	FVertexFactoryType* VFType = InVertexFactory->GetType();
-
-	const bool bInitializeTessellationShaders =
-		MaterialResource->GetTessellationMode() != MTM_NoTessellation
-		&& RHISupportsTessellation(GShaderPlatformForFeatureLevel[InFeatureLevel])
-		&& VFType->SupportsTessellationShaders();
-
-	bUsePositionOnlyVS = !bRenderingReflectiveShadowMaps
-		&& VertexFactory->SupportsPositionOnlyStream()
-		&& MaterialResource->WritesEveryPixel(true)
-		&& !MaterialResource->MaterialModifiesMeshPosition_RenderThread();
-
-	// Vertex related shaders
-	if (bOnePassPointLightShadow)
-	{
-		if (bUsePositionOnlyVS)
-		{
-			VertexShader = MaterialResource->GetShader<TShadowDepthVS<VertexShadowDepth_OnePassPointLight, false, true, true> >(VFType);
-		}
-		else
-		{
-			VertexShader = MaterialResource->GetShader<TShadowDepthVS<VertexShadowDepth_OnePassPointLight, false, false, true> >(VFType);
-		}
-		if (RHISupportsGeometryShaders(GShaderPlatformForFeatureLevel[InFeatureLevel]))
-		{
-			// Use the geometry shader which will clone output triangles to all faces of the cube map
-			GeometryShader = MaterialResource->GetShader<FOnePassPointShadowDepthGS>(VFType);
-		}
-		if (bInitializeTessellationShaders)
-		{
-			HullShader = MaterialResource->GetShader<TShadowDepthHS<VertexShadowDepth_OnePassPointLight, false> >(VFType);
-			DomainShader = MaterialResource->GetShader<TShadowDepthDS<VertexShadowDepth_OnePassPointLight, false> >(VFType);
-		}
-	}
-	else if (bUsePerspectiveCorrectShadowDepths)
-	{
-		if (bRenderingReflectiveShadowMaps)
-		{
-			VertexShader = MaterialResource->GetShader<TShadowDepthVS<VertexShadowDepth_PerspectiveCorrect, true, false> >(VFType);
-		}
-		else
-		{
-			if (bUsePositionOnlyVS)
-			{
-				VertexShader = MaterialResource->GetShader<TShadowDepthVS<VertexShadowDepth_PerspectiveCorrect, false, true> >(VFType);
-			}
-			else
-			{
-				VertexShader = MaterialResource->GetShader<TShadowDepthVS<VertexShadowDepth_PerspectiveCorrect, false, false> >(VFType);
-			}
-		}
-		if (bInitializeTessellationShaders)
-		{
-			HullShader = MaterialResource->GetShader<TShadowDepthHS<VertexShadowDepth_PerspectiveCorrect, bRenderingReflectiveShadowMaps> >(VFType);
-			DomainShader = MaterialResource->GetShader<TShadowDepthDS<VertexShadowDepth_PerspectiveCorrect, bRenderingReflectiveShadowMaps> >(VFType);
-		}
-	}
-	else
-	{
-		if (bRenderingReflectiveShadowMaps)
-		{
-			VertexShader = MaterialResource->GetShader<TShadowDepthVS<VertexShadowDepth_OutputDepth, true, false> >(VFType);
-			if (bInitializeTessellationShaders)
-			{
-				HullShader = MaterialResource->GetShader<TShadowDepthHS<VertexShadowDepth_OutputDepth, true> >(VFType);
-				DomainShader = MaterialResource->GetShader<TShadowDepthDS<VertexShadowDepth_OutputDepth, true> >(VFType);
-			}
-		}
-		else
-		{
-			if (bUsePositionOnlyVS)
-			{
-				VertexShader = MaterialResource->GetShader<TShadowDepthVS<VertexShadowDepth_OutputDepth, false, true> >(VFType);
-			}
-			else
-			{
-				VertexShader = MaterialResource->GetShader<TShadowDepthVS<VertexShadowDepth_OutputDepth, false, false> >(VFType);
-			}
-			if (bInitializeTessellationShaders)
-			{
-				HullShader = MaterialResource->GetShader<TShadowDepthHS<VertexShadowDepth_OutputDepth, false> >(VFType);
-				DomainShader = MaterialResource->GetShader<TShadowDepthDS<VertexShadowDepth_OutputDepth, false> >(VFType);
-			}
-		}
-	}
-
-	// Pixel shaders
-	if (MaterialResource->WritesEveryPixel(true) && !bUsePerspectiveCorrectShadowDepths && !bRenderingReflectiveShadowMaps && VertexFactory->SupportsNullPixelShader())
-	{
-		// No pixel shader necessary.
-		PixelShader = NULL;
-	}
-	else
-	{
-		if (bUsePerspectiveCorrectShadowDepths)
-		{
-			PixelShader = (TShadowDepthBasePS<bRenderingReflectiveShadowMaps> *)MaterialResource->GetShader<TShadowDepthPS<PixelShadowDepth_PerspectiveCorrect, bRenderingReflectiveShadowMaps> >(VFType, false);
-		}
-		else if (bOnePassPointLightShadow)
-		{
-			PixelShader = (TShadowDepthBasePS<bRenderingReflectiveShadowMaps> *)MaterialResource->GetShader<TShadowDepthPS<PixelShadowDepth_OnePassPointLight, false> >(VFType, false);
-		}
-		else
-		{
-			PixelShader = (TShadowDepthBasePS<bRenderingReflectiveShadowMaps> *)MaterialResource->GetShader<TShadowDepthPS<PixelShadowDepth_NonPerspectiveCorrect, bRenderingReflectiveShadowMaps> >(VFType, false);
-		}
-	}
-	BaseVertexShader = VertexShader;
-}
-
-static void SetViewFlagsForShadowPass(FDrawingPolicyRenderState& DrawRenderState, const FSceneView& View, ERHIFeatureLevel::Type FeatureLevel, bool bInIsTwoSided, bool isReflectiveShadowmap, bool isOnePassPointLightShadow)
-{
-	// @TODO: only render directional light shadows as two sided, and only when blocking is enabled (required by geometry volume injection)
-	bool bEffectivelyTwoSided = isReflectiveShadowmap ? true : bInIsTwoSided;
-	// Invert culling order when mobile HDR == false.
-	auto ShaderPlatform = GShaderPlatformForFeatureLevel[FeatureLevel];
-	static auto* MobileHDRCvar = IConsoleManager::Get().FindTConsoleVariableDataInt(TEXT("r.MobileHDR"));
-	check(MobileHDRCvar);
-	const bool bPlatformReversesCulling = (RHINeedsToSwitchVerticalAxis(ShaderPlatform) && MobileHDRCvar->GetValueOnAnyThread() == 0);
-
-	EDrawingPolicyOverrideFlags& ViewOverrideFlags = DrawRenderState.ModifyViewOverrideFlags();
-
-	ViewOverrideFlags = (View.bRenderSceneTwoSided || bEffectivelyTwoSided) ?
-		ViewOverrideFlags | EDrawingPolicyOverrideFlags::TwoSided :
-		ViewOverrideFlags & ~EDrawingPolicyOverrideFlags::TwoSided;
-	ViewOverrideFlags = XOR(View.bReverseCulling, XOR(bPlatformReversesCulling, isOnePassPointLightShadow)) ?
-		ViewOverrideFlags | EDrawingPolicyOverrideFlags::ReverseCullMode :
-		ViewOverrideFlags & ~EDrawingPolicyOverrideFlags::ReverseCullMode;
-}
-
-template <bool bRenderingReflectiveShadowMaps>
-void FShadowDepthDrawingPolicy<bRenderingReflectiveShadowMaps>::SetSharedState(FRHICommandList& RHICmdList, const FDrawingPolicyRenderState& DrawRenderState, const FSceneView* View, const ContextDataType PolicyContext) const
-{
-	checkSlow(bDirectionalLight == PolicyContext.ShadowInfo->bDirectionalLight && bPreShadow == PolicyContext.ShadowInfo->bPreShadow);
-
-	VertexShader->SetParameters(RHICmdList, VertexShader->GetVertexShader(), MaterialRenderProxy, *MaterialResource, *View, DrawRenderState.GetViewUniformBuffer(), DrawRenderState.GetPassUniformBuffer());
-
-	if (GeometryShader)
-	{
-		GeometryShader->SetParameters(RHICmdList, GeometryShader->GetGeometryShader(), MaterialRenderProxy, *MaterialResource, *View, DrawRenderState.GetViewUniformBuffer(), DrawRenderState.GetPassUniformBuffer());
-	}
-
-	if (HullShader && DomainShader)
-	{
-		HullShader->SetParameters(RHICmdList, MaterialRenderProxy, *View, DrawRenderState.GetViewUniformBuffer(), DrawRenderState.GetPassUniformBuffer());
-		DomainShader->SetParameters(RHICmdList, MaterialRenderProxy, *View, DrawRenderState.GetViewUniformBuffer(), DrawRenderState.GetPassUniformBuffer());
-	}
-
-	if (PixelShader)
-	{
-		PixelShader->SetParameters(RHICmdList, PixelShader->GetPixelShader(), MaterialRenderProxy, *MaterialResource, *View, DrawRenderState.GetViewUniformBuffer(), DrawRenderState.GetPassUniformBuffer());
-	}
-
-	// Set the shared mesh resources.
-	if (bUsePositionOnlyVS)
-	{
-		VertexFactory->SetPositionStream(RHICmdList);
-	}
-	else
-	{
-		FMeshDrawingPolicy::SetSharedState(RHICmdList, DrawRenderState, View, PolicyContext);
-	}
-}
-
-/**
-* Create bound shader state using the vertex decl from the mesh draw policy
-* as well as the shaders needed to draw the mesh
-* @return new bound shader state object
-*/
-template <bool bRenderingReflectiveShadowMaps>
-FBoundShaderStateInput FShadowDepthDrawingPolicy<bRenderingReflectiveShadowMaps>::GetBoundShaderStateInput(ERHIFeatureLevel::Type InFeatureLevel) const
-{
-	FVertexDeclarationRHIRef VertexDeclaration;
-	if (bUsePositionOnlyVS)
-	{
-		VertexDeclaration = VertexFactory->GetPositionDeclaration();
-	}
-	else
-	{
-		VertexDeclaration = FMeshDrawingPolicy::GetVertexDeclaration();
-	}
-
-	return FBoundShaderStateInput(
-		VertexDeclaration,
-		VertexShader->GetVertexShader(),
-		GETSAFERHISHADER_HULL(HullShader),
-		GETSAFERHISHADER_DOMAIN(DomainShader),
-		GETSAFERHISHADER_PIXEL(PixelShader),
-		GETSAFERHISHADER_GEOMETRY(GeometryShader));
-}
-
-template <bool bRenderingReflectiveShadowMaps>
-void FShadowDepthDrawingPolicy<bRenderingReflectiveShadowMaps>::SetMeshRenderState(
-	FRHICommandList& RHICmdList,
-	const FSceneView& View,
-	const FPrimitiveSceneProxy* PrimitiveSceneProxy,
-	const FMeshBatch& Mesh,
-	int32 BatchElementIndex,
-	const FDrawingPolicyRenderState& DrawRenderState,
-	const ElementDataType& ElementData,
-	const ContextDataType PolicyContext
-) const
-{
-	const FMeshBatchElement& BatchElement = Mesh.Elements[BatchElementIndex];
-
-	VertexShader->SetMesh(RHICmdList, VertexShader->GetVertexShader(), VertexFactory, View, PrimitiveSceneProxy, BatchElement, DrawRenderState);
-	SetPrimitiveIdStream(RHICmdList, View, PrimitiveSceneProxy, BatchElement.PrimitiveIdMode, BatchElement.DynamicPrimitiveShaderDataIndex);
-
-	if (HullShader && DomainShader)
-	{
-		HullShader->SetMesh(RHICmdList, VertexFactory, View, PrimitiveSceneProxy, BatchElement, DrawRenderState);
-		DomainShader->SetMesh(RHICmdList, VertexFactory, View, PrimitiveSceneProxy, BatchElement, DrawRenderState);
-	}
-	if (GeometryShader)
-	{
-		GeometryShader->SetMesh(RHICmdList, GeometryShader->GetGeometryShader(), VertexFactory, View, PrimitiveSceneProxy, BatchElement, DrawRenderState);
-	}
-	if (PixelShader)
-	{
-		PixelShader->SetMesh(RHICmdList, PixelShader->GetPixelShader(), VertexFactory, View, PrimitiveSceneProxy, BatchElement, DrawRenderState);
-	}
-}
-
-template <bool bRenderingReflectiveShadowMaps>
-int32 CompareDrawingPolicy(const FShadowDepthDrawingPolicy<bRenderingReflectiveShadowMaps>& A, const FShadowDepthDrawingPolicy<bRenderingReflectiveShadowMaps>& B)
-{
-	COMPAREDRAWINGPOLICYMEMBERS(VertexShader);
-	COMPAREDRAWINGPOLICYMEMBERS(HullShader);
-	COMPAREDRAWINGPOLICYMEMBERS(DomainShader);
-	COMPAREDRAWINGPOLICYMEMBERS(GeometryShader);
-	COMPAREDRAWINGPOLICYMEMBERS(PixelShader);
-	COMPAREDRAWINGPOLICYMEMBERS(VertexFactory);
-	COMPAREDRAWINGPOLICYMEMBERS(MaterialRenderProxy);
-	COMPAREDRAWINGPOLICYMEMBERS(bDirectionalLight);
-	COMPAREDRAWINGPOLICYMEMBERS(MeshPrimitiveType);
-	COMPAREDRAWINGPOLICYMEMBERS(bOnePassPointLightShadow);
-	COMPAREDRAWINGPOLICYMEMBERS(bUsePositionOnlyVS);
-	COMPAREDRAWINGPOLICYMEMBERS(bPreShadow);
-	return 0;
-}
-
-void FShadowDepthDrawingPolicyFactory::AddStaticMesh(FScene* Scene, FStaticMesh* StaticMesh)
-{
-	if (StaticMesh->CastShadow)
-	{
-		const auto FeatureLevel = Scene->GetFeatureLevel();
-		const FMaterialRenderProxy* MaterialRenderProxy = StaticMesh->MaterialRenderProxy;
-		const FMaterial* Material = MaterialRenderProxy->GetMaterial(FeatureLevel);
-		const EBlendMode BlendMode = Material->GetBlendMode();
-		const EMaterialShadingModel ShadingModel = Material->GetShadingModel();
-
-		const bool bLightPropagationVolume = UseLightPropagationVolumeRT(FeatureLevel);
-		const bool bTwoSided = Material->IsTwoSided() || StaticMesh->PrimitiveSceneInfo->Proxy->CastsShadowAsTwoSided();
-		const bool bLitOpaque = !IsTranslucentBlendMode(BlendMode) && ShadingModel != MSM_Unlit;
-
-		FMeshDrawingPolicyOverrideSettings OverrideSettings = ComputeMeshOverrideSettings(*StaticMesh);
-		OverrideSettings.MeshOverrideFlags |= bTwoSided ? EDrawingPolicyOverrideFlags::TwoSided : EDrawingPolicyOverrideFlags::None;
-
-		if (bLightPropagationVolume && ((!IsTranslucentBlendMode(BlendMode) && ShadingModel != MSM_Unlit) || Material->ShouldInjectEmissiveIntoLPV() || Material->ShouldBlockGI()))
-		{
-			// Add the static mesh to the shadow's subject draw list.
-			if (StaticMesh->PrimitiveSceneInfo->Proxy->AffectsDynamicIndirectLighting())
-			{
-				Scene->WholeSceneReflectiveShadowMapDrawList.AddMesh(
-					StaticMesh,
-					FShadowDepthDrawingPolicy<true>::ElementDataType(),
-					FShadowDepthDrawingPolicy<true>(
-						Material,
-						true,
-						false,
-						false,
-						OverrideSettings,
-						FeatureLevel,
-						StaticMesh->VertexFactory,
-						MaterialRenderProxy,
-						StaticMesh->ReverseCulling),
-					FeatureLevel
-				);
-			}
-		}
-		if (bLitOpaque)
-		{
-			OverrideWithDefaultMaterialForShadowDepth(MaterialRenderProxy, Material, false, FeatureLevel);
-
-			// Add the static mesh to the shadow's subject draw list.
-			Scene->WholeSceneShadowDepthDrawList.AddMesh(
-				StaticMesh,
-				FShadowDepthDrawingPolicy<false>::ElementDataType(),
-				FShadowDepthDrawingPolicy<false>(
-					Material,
-					true,
-					false,
-					false,
-					OverrideSettings,
-					FeatureLevel,
-					StaticMesh->VertexFactory,
-					MaterialRenderProxy,
-					StaticMesh->ReverseCulling),
-				FeatureLevel
-			);
-		}
-	}
-}
-
-bool FShadowDepthDrawingPolicyFactory::DrawDynamicMesh(
-	FRHICommandList& RHICmdList,
-	const FSceneView& View,
-	ContextType Context,
-	const FMeshBatch& Mesh,
-	bool bPreFog,
-	const FDrawingPolicyRenderState& DrawRenderState,
-	const FPrimitiveSceneProxy* PrimitiveSceneProxy,
-	FHitProxyId HitProxyId
-)
-{
-	bool bDirty = false;
-
-	// Use a per-FMeshBatch check on top of the per-primitive check because dynamic primitives can submit multiple FMeshElements.
-	if (Mesh.CastShadow)
-	{
-		const auto FeatureLevel = View.GetFeatureLevel();
-		const FMaterialRenderProxy* MaterialRenderProxy = Mesh.MaterialRenderProxy;
-		const FMaterial* Material = MaterialRenderProxy->GetMaterial(FeatureLevel);
-		const EBlendMode BlendMode = Material->GetBlendMode();
-		const EMaterialShadingModel ShadingModel = Material->GetShadingModel();
-
-		const bool bLocalOnePassPointLightShadow = Context.ShadowInfo->bOnePassPointLightShadow;
-		const bool bReflectiveShadowmap = Context.ShadowInfo->bReflectiveShadowmap && !bLocalOnePassPointLightShadow;
-
-		bool bProcess = Material->ShouldCastDynamicShadows() && ShouldIncludeDomainInMeshPass(Material->GetMaterialDomain());
-
-		if (bReflectiveShadowmap && Material->ShouldInjectEmissiveIntoLPV())
-		{
-			bProcess = true;
-		}
-
-		if (bProcess)
-		{
-			const bool bLocalDirectionalLight = Context.ShadowInfo->bDirectionalLight;
-			const bool bPreShadow = Context.ShadowInfo->bPreShadow;
-			const bool bTwoSided = Material->IsTwoSided() || PrimitiveSceneProxy->CastsShadowAsTwoSided();
-			const FShadowDepthDrawingPolicyContext PolicyContext(Context.ShadowInfo);
-
-			FMeshDrawingPolicyOverrideSettings OverrideSettings = ComputeMeshOverrideSettings(Mesh);
-			OverrideSettings.MeshOverrideFlags |= bTwoSided ? EDrawingPolicyOverrideFlags::TwoSided : EDrawingPolicyOverrideFlags::None;
-
-			OverrideWithDefaultMaterialForShadowDepth(MaterialRenderProxy, Material, bReflectiveShadowmap, FeatureLevel);
-
-			if (bReflectiveShadowmap)
-			{
-				FShadowDepthDrawingPolicy<true> DrawingPolicy(
-					MaterialRenderProxy->GetMaterial(FeatureLevel),
-					bLocalDirectionalLight,
-					bLocalOnePassPointLightShadow,
-					bPreShadow,
-					OverrideSettings,
-					FeatureLevel,
-					Mesh.VertexFactory,
-					MaterialRenderProxy,
-					Mesh.ReverseCulling
-				);
-
-				//TODO MaybeRemovable if ShadowDepth never support LOD Transitions
-				FDrawingPolicyRenderState DrawRenderStateLocal(DrawRenderState);
-				DrawRenderStateLocal.SetDitheredLODTransitionAlpha(Mesh.DitheredLODTransitionAlpha);
-				SetViewFlagsForShadowPass(DrawRenderStateLocal, View, View.GetFeatureLevel(), DrawingPolicy.IsTwoSided(), true, bLocalOnePassPointLightShadow);
-				DrawingPolicy.SetupPipelineState(DrawRenderStateLocal, View);
-				CommitGraphicsPipelineState(RHICmdList, DrawingPolicy, DrawRenderStateLocal, DrawingPolicy.GetBoundShaderStateInput(View.GetFeatureLevel()), DrawingPolicy.GetMaterialRenderProxy());
-				DrawingPolicy.SetSharedState(RHICmdList, DrawRenderStateLocal, &View, PolicyContext);
-
-				for (int32 BatchElementIndex = 0, Num = Mesh.Elements.Num(); BatchElementIndex < Num; BatchElementIndex++)
-				{
-					TDrawEvent<FRHICommandList> MeshEvent;
-					BeginMeshDrawEvent(RHICmdList, PrimitiveSceneProxy, Mesh, MeshEvent, EnumHasAnyFlags(EShowMaterialDrawEventTypes(GShowMaterialDrawEventTypes), EShowMaterialDrawEventTypes::ShadowDepthRsm));
-
-					DrawingPolicy.SetMeshRenderState(RHICmdList, View, PrimitiveSceneProxy, Mesh, BatchElementIndex, DrawRenderStateLocal, FMeshDrawingPolicy::ElementDataType(), PolicyContext);
-					DrawingPolicy.DrawMesh(RHICmdList, View, Mesh, BatchElementIndex);
-				}
-			}
-			else
-			{
-				FShadowDepthDrawingPolicy<false> DrawingPolicy(
-					MaterialRenderProxy->GetMaterial(FeatureLevel),
-					bLocalDirectionalLight,
-					bLocalOnePassPointLightShadow,
-					bPreShadow,
-					OverrideSettings,
-					FeatureLevel,
-					Mesh.VertexFactory,
-					MaterialRenderProxy,
-					Mesh.ReverseCulling
-				);
-
-				//TODO MaybeRemovable if ShadowDepth never support LOD Transitions
-				FDrawingPolicyRenderState DrawRenderStateLocal(DrawRenderState);
-				DrawRenderStateLocal.SetDitheredLODTransitionAlpha(Mesh.DitheredLODTransitionAlpha);
-				SetViewFlagsForShadowPass(DrawRenderStateLocal, View, View.GetFeatureLevel(), DrawingPolicy.IsTwoSided(), false, bLocalOnePassPointLightShadow);
-				DrawingPolicy.SetupPipelineState(DrawRenderStateLocal, View);
-				CommitGraphicsPipelineState(RHICmdList, DrawingPolicy, DrawRenderStateLocal, DrawingPolicy.GetBoundShaderStateInput(View.GetFeatureLevel()), DrawingPolicy.GetMaterialRenderProxy());
-				DrawingPolicy.SetSharedState(RHICmdList, DrawRenderStateLocal, &View, PolicyContext);
-
-				for (int32 BatchElementIndex = 0; BatchElementIndex < Mesh.Elements.Num(); BatchElementIndex++)
-				{
-					TDrawEvent<FRHICommandList> MeshEvent;
-					BeginMeshDrawEvent(RHICmdList, PrimitiveSceneProxy, Mesh, MeshEvent, EnumHasAnyFlags(EShowMaterialDrawEventTypes(GShowMaterialDrawEventTypes), EShowMaterialDrawEventTypes::ShadowDepth));
-
-					DrawingPolicy.SetMeshRenderState(RHICmdList, View, PrimitiveSceneProxy, Mesh, BatchElementIndex, DrawRenderStateLocal, FMeshDrawingPolicy::ElementDataType(), PolicyContext);
-					DrawingPolicy.DrawMesh(RHICmdList, View, Mesh, BatchElementIndex);
-				}
-			}
-
-
-			bDirty = true;
-		}
-	}
-
-	return bDirty;
-}
-
 /*-----------------------------------------------------------------------------
 FProjectedShadowInfo
 -----------------------------------------------------------------------------*/
@@ -1240,218 +757,6 @@ void FProjectedShadowInfo::ClearDepth(FRHICommandList& RHICmdList, class FSceneR
 	}
 }
 
-template <bool bReflectiveShadowmap>
-void DrawMeshElements(FRHICommandList& RHICmdList, FShadowDepthDrawingPolicy<bReflectiveShadowmap>& SharedDrawingPolicy, const FShadowStaticMeshElement& State, const FViewInfo& View, FShadowDepthDrawingPolicyContext PolicyContext, const FDrawingPolicyRenderState& DrawRenderState, const FStaticMesh* Mesh)
-{
-#if UE_BUILD_DEBUG
-	// During shadow setup we should have already overridden materials with default material where needed.
-	// Make sure of it!
-	CheckShadowDepthMaterials(State.RenderProxy, State.MaterialResource, bReflectiveShadowmap, View.GetFeatureLevel());
-#endif
-
-#if UE_BUILD_DEBUG
-	FMeshDrawingPolicyOverrideSettings OverrideSettings = ComputeMeshOverrideSettings(*State.Mesh);
-	OverrideSettings.MeshOverrideFlags |= State.bIsTwoSided ? EDrawingPolicyOverrideFlags::TwoSided : EDrawingPolicyOverrideFlags::None;
-
-	FShadowDepthDrawingPolicy<bReflectiveShadowmap> DebugPolicy(
-		State.MaterialResource,
-		SharedDrawingPolicy.bDirectionalLight,
-		SharedDrawingPolicy.bOnePassPointLightShadow,
-		SharedDrawingPolicy.bPreShadow,
-		OverrideSettings,
-		View.GetFeatureLevel(),
-		State.Mesh->VertexFactory,
-		State.RenderProxy,
-		State.Mesh->ReverseCulling);
-	// Verify that SharedDrawingPolicy can be used to draw this mesh without artifacts by checking the comparison functions that static draw lists use
-	checkSlow(DebugPolicy.Matches(SharedDrawingPolicy).Result());
-	checkSlow(CompareDrawingPolicy(DebugPolicy, SharedDrawingPolicy) == 0);
-#endif
-
-	//TODO MaybeRemovable if ShadowDepth never support LOD Transitions
-	FDrawingPolicyRenderState DrawRenderStateLocal(DrawRenderState);
-	SharedDrawingPolicy.ApplyDitheredLODTransitionState(DrawRenderStateLocal, View, *Mesh, false);
-
-	// Render only those batch elements that match the current LOD
-	uint64 BatchElementMask = Mesh->bRequiresPerElementVisibility ? View.StaticMeshBatchVisibility[Mesh->BatchVisibilityId] : ((1ull << Mesh->Elements.Num()) - 1);
-	int32 BatchElementIndex = 0;
-	do
-	{
-		if (BatchElementMask & 1)
-		{
-			TDrawEvent<FRHICommandList> MeshEvent;
-			BeginMeshDrawEvent(RHICmdList, Mesh->PrimitiveSceneInfo->Proxy, *Mesh, MeshEvent, EnumHasAnyFlags(EShowMaterialDrawEventTypes(GShowMaterialDrawEventTypes), EShowMaterialDrawEventTypes::ShadowDepthStatic));
-
-			SharedDrawingPolicy.SetMeshRenderState(RHICmdList, View, Mesh->PrimitiveSceneInfo->Proxy, *Mesh, BatchElementIndex, DrawRenderStateLocal, FMeshDrawingPolicy::ElementDataType(), PolicyContext);
-			SharedDrawingPolicy.DrawMesh(RHICmdList, View, *Mesh, BatchElementIndex);
-			INC_DWORD_STAT(STAT_ShadowDynamicPathDrawCalls);
-		}
-
-		BatchElementMask >>= 1;
-		BatchElementIndex++;
-	} while (BatchElementMask);
-}
-
-template <bool bReflectiveShadowmap>
-void DrawShadowMeshElements(FRHICommandList& RHICmdList, const FViewInfo& View, const FDrawingPolicyRenderState& DrawRenderState, const FProjectedShadowInfo& ShadowInfo)
-{
-	FShadowDepthDrawingPolicyContext PolicyContext(&ShadowInfo);
-	const FShadowStaticMeshElement& FirstShadowMesh = ShadowInfo.StaticSubjectMeshElements[0];
-	const FMaterial* FirstMaterialResource = FirstShadowMesh.MaterialResource;
-	auto FeatureLevel = View.GetFeatureLevel();
-
-	FShadowDepthDrawingPolicy<bReflectiveShadowmap> SharedDrawingPolicy(
-		FirstMaterialResource,
-		ShadowInfo.bDirectionalLight,
-		ShadowInfo.bOnePassPointLightShadow,
-		ShadowInfo.bPreShadow,
-		ComputeMeshOverrideSettings(*FirstShadowMesh.Mesh),
-		FeatureLevel);
-
-	FShadowStaticMeshElement OldState;
-
-	FDrawingPolicyRenderState DrawRenderStateLocal(DrawRenderState);
-
-	uint32 ElementCount = ShadowInfo.StaticSubjectMeshElements.Num();
-	for (uint32 ElementIndex = 0; ElementIndex < ElementCount; ++ElementIndex)
-	{
-		const FShadowStaticMeshElement& ShadowMesh = ShadowInfo.StaticSubjectMeshElements[ElementIndex];
-
-		if (!View.StaticMeshShadowDepthMap[ShadowMesh.Mesh->Id])
-		{
-			// not visible
-			continue;
-		}
-
-		FShadowStaticMeshElement CurrentState(ShadowMesh.RenderProxy, ShadowMesh.MaterialResource, ShadowMesh.Mesh, ShadowMesh.bIsTwoSided);
-
-		// Only call draw shared when the vertex factory or material have changed
-		if (OldState.DoesDeltaRequireADrawSharedCall(CurrentState))
-		{
-			OldState = CurrentState;
-
-			SharedDrawingPolicy.UpdateElementState(CurrentState, FeatureLevel);
-			DrawRenderStateLocal.SetDitheredLODTransitionAlpha(ShadowMesh.Mesh->DitheredLODTransitionAlpha);
-			SetViewFlagsForShadowPass(DrawRenderStateLocal, View, View.GetFeatureLevel(), SharedDrawingPolicy.IsTwoSided(), bReflectiveShadowmap, ShadowInfo.bOnePassPointLightShadow);
-			SharedDrawingPolicy.SetupPipelineState(DrawRenderStateLocal, View);
-			CommitGraphicsPipelineState(RHICmdList, SharedDrawingPolicy, DrawRenderStateLocal, SharedDrawingPolicy.GetBoundShaderStateInput(View.GetFeatureLevel()), SharedDrawingPolicy.GetMaterialRenderProxy());
-			SharedDrawingPolicy.SetSharedState(RHICmdList, DrawRenderStateLocal, &View, PolicyContext);
-		}
-
-		DrawMeshElements(RHICmdList, SharedDrawingPolicy, OldState, View, PolicyContext, DrawRenderStateLocal, ShadowMesh.Mesh);
-	}
-}
-
-void FProjectedShadowInfo::RenderDepthDynamic(FRHICommandList& RHICmdList, FSceneRenderer* SceneRenderer, const FViewInfo* FoundView, const FDrawingPolicyRenderState& DrawRenderState)
-{
-	// Draw the subject's dynamic elements.
-	SCOPE_CYCLE_COUNTER(STAT_WholeSceneDynamicShadowDepthsTime);
-
-	FShadowDepthDrawingPolicyFactory::ContextType Context(this);
-
-	for (int32 MeshBatchIndex = 0; MeshBatchIndex < DynamicSubjectMeshElements.Num(); MeshBatchIndex++)
-	{
-		const FMeshBatchAndRelevance& MeshBatchAndRelevance = DynamicSubjectMeshElements[MeshBatchIndex];
-		const FMeshBatch& MeshBatch = *MeshBatchAndRelevance.Mesh;
-		FShadowDepthDrawingPolicyFactory::DrawDynamicMesh(RHICmdList, *FoundView, Context, MeshBatch, true, DrawRenderState, MeshBatchAndRelevance.PrimitiveSceneProxy, MeshBatch.BatchHitProxyId);
-	}
-}
-
-class FDrawShadowMeshElementsThreadTask : public FRenderTask
-{
-	FProjectedShadowInfo& ThisShadow;
-	FRHICommandList& RHICmdList;
-	const FViewInfo& View;
-	FDrawingPolicyRenderState DrawRenderState;
-	bool bReflective;
-
-public:
-
-	FDrawShadowMeshElementsThreadTask(
-		FProjectedShadowInfo& InThisShadow,
-		FRHICommandList& InRHICmdList,
-		const FViewInfo& InView,
-		const FDrawingPolicyRenderState& InDrawRenderState,
-		bool InbReflective
-	)
-		: ThisShadow(InThisShadow)
-		, RHICmdList(InRHICmdList)
-		, View(InView)
-		, DrawRenderState(InDrawRenderState)
-		, bReflective(InbReflective)
-	{
-	}
-
-	FORCEINLINE TStatId GetStatId() const
-	{
-		RETURN_QUICK_DECLARE_CYCLE_STAT(FDrawShadowMeshElementsThreadTask, STATGROUP_TaskGraphTasks);
-	}
-
-	static ESubsequentsMode::Type GetSubsequentsMode() { return ESubsequentsMode::TrackSubsequents; }
-
-	void DoTask(ENamedThreads::Type CurrentThread, const FGraphEventRef& MyCompletionGraphEvent)
-	{
-		SCOPE_CYCLE_COUNTER(STAT_WholeSceneStaticShadowDepthsTime);
-
-		checkSlow(RHICmdList.IsInsideRenderPass());
-
-		if (bReflective)
-		{
-			// reflective shadow map
-			DrawShadowMeshElements<true>(RHICmdList, View, DrawRenderState, ThisShadow);
-		}
-		else
-		{
-			// normal shadow map
-			DrawShadowMeshElements<false>(RHICmdList, View, DrawRenderState, ThisShadow);
-		}
-
-		RHICmdList.EndRenderPass();
-		RHICmdList.HandleRTThreadTaskCompletion(MyCompletionGraphEvent);
-	}
-};
-
-class FRenderDepthDynamicThreadTask : public FRenderTask
-{
-	FProjectedShadowInfo& ThisShadow;
-	FRHICommandList& RHICmdList;
-	const FViewInfo& View;
-	FDrawingPolicyRenderState DrawRenderState;
-	FSceneRenderer* SceneRenderer;
-
-public:
-
-	FRenderDepthDynamicThreadTask(
-		FProjectedShadowInfo& InThisShadow,
-		FRHICommandList& InRHICmdList,
-		const FViewInfo& InView,
-		const FDrawingPolicyRenderState& InDrawRenderState,
-		FSceneRenderer* InSceneRenderer
-	)
-		: ThisShadow(InThisShadow)
-		, RHICmdList(InRHICmdList)
-		, View(InView)
-		, DrawRenderState(InDrawRenderState)
-		, SceneRenderer(InSceneRenderer)
-	{
-	}
-
-	FORCEINLINE TStatId GetStatId() const
-	{
-		RETURN_QUICK_DECLARE_CYCLE_STAT(FRenderDepthDynamicThreadTask, STATGROUP_TaskGraphTasks);
-	}
-
-	static ESubsequentsMode::Type GetSubsequentsMode() { return ESubsequentsMode::TrackSubsequents; }
-
-	void DoTask(ENamedThreads::Type CurrentThread, const FGraphEventRef& MyCompletionGraphEvent)
-	{
-		check(RHICmdList.IsInsideRenderPass());
-		ThisShadow.RenderDepthDynamic(RHICmdList, SceneRenderer, &View, DrawRenderState);
-		RHICmdList.EndRenderPass();
-		RHICmdList.HandleRTThreadTaskCompletion(MyCompletionGraphEvent);
-	}
-};
-
 void FProjectedShadowInfo::SetStateForView(FRHICommandList& RHICmdList) const
 {
 	check(bAllocated);
@@ -1528,8 +833,6 @@ public:
 		, ProjectedShadowInfo(InProjectedShadowInfo)
 		, BeginShadowRenderPass(InBeginShadowRenderPass)
 	{
-		// It is difficult to get accurate counts from the last frame. Cascades interfere with each other.
-		bBalanceCommandsWithLastFrame = false;
 		bBalanceCommands = false;
 	}
 
@@ -1817,9 +1120,6 @@ void FProjectedShadowInfo::RenderDepthInner(FRHICommandListImmediate& RHICmdList
 		}
 	}
 
-	FShadowDepthDrawingPolicyContext StackPolicyContext(this);
-	FShadowDepthDrawingPolicyContext* PolicyContext(&StackPolicyContext);
-
 	if (bDoParallelDispatch)
 	{
 		check(IsInRenderingThread());
@@ -1830,62 +1130,14 @@ void FProjectedShadowInfo::RenderDepthInner(FRHICommandListImmediate& RHICmdList
 		bool bFlush = CVarRHICmdFlushRenderThreadTasksShadowPass.GetValueOnRenderThread() > 0
 			|| CVarRHICmdFlushRenderThreadTasks.GetValueOnRenderThread() > 0
 			// Scene uniform buffers are reused between cascades
-			|| (UseMeshDrawCommandPipeline() && GetShadowDepthType() == CSMShadowDepthType);
+			|| GetShadowDepthType() == CSMShadowDepthType;
 		FScopedCommandListWaitForTasks Flusher(bFlush);
-		if (!bFlush && !UseMeshDrawCommandPipeline())
-		{
-			/** CAUTION, this is assumed to be a POD type. We allocate the on the scene allocator and NEVER CALL A DESTRUCTOR.
-			If you want to add non-pod data, not a huge problem, we just need to track and destruct them at the end of the scene.
-			**/
-			check(IsInRenderingThread() && FMemStack::Get().GetNumMarks() == 1); // we do not want this popped before the end of the scene and it better be the scene allocator
-			PolicyContext = new (FMemStack::Get()) FShadowDepthDrawingPolicyContext(this);
-		}
 
 		// Dispatch commands
 		{
 			FShadowParallelCommandListSet ParallelCommandListSet(*ShadowDepthView, SceneRenderer, RHICmdList, CVarRHICmdShadowDeferredContexts.GetValueOnRenderThread() > 0, !bFlush, DrawRenderState, *this, BeginShadowRenderPass);
 
-			if (UseMeshDrawCommandPipeline())
-			{
-				ShadowDepthPass.DispatchDraw(&ParallelCommandListSet, RHICmdList);
-			}
-			else
-			{
-				// Draw the subject's static elements using static draw lists
-				if (bIsWholeSceneDirectionalShadow)
-				{
-					SCOPE_CYCLE_COUNTER(STAT_WholeSceneStaticDrawListShadowDepthsTime);
-
-					if (bReflectiveShadowmap)
-					{
-						SceneRenderer->Scene->WholeSceneReflectiveShadowMapDrawList.DrawVisibleParallel(*PolicyContext, StaticMeshWholeSceneShadowDepthMap, StaticMeshWholeSceneShadowBatchVisibility, ParallelCommandListSet);
-					}
-					else
-					{
-						// Use the scene's shadow depth draw list with this shadow's visibility map
-						SceneRenderer->Scene->WholeSceneShadowDepthDrawList.DrawVisibleParallel(*PolicyContext, StaticMeshWholeSceneShadowDepthMap, StaticMeshWholeSceneShadowBatchVisibility, ParallelCommandListSet);
-					}
-				}
-				// Draw the subject's static elements using manual state filtering
-				else if (StaticSubjectMeshElements.Num() > 0)
-				{
-					FRHICommandList* CmdList = ParallelCommandListSet.NewParallelCommandList();
-
-					FGraphEventRef AnyThreadCompletionEvent = TGraphTask<FDrawShadowMeshElementsThreadTask>::CreateTask(ParallelCommandListSet.GetPrereqs(), ENamedThreads::GetRenderThread())
-						.ConstructAndDispatchWhenReady(*this, *CmdList, *ShadowDepthView, DrawRenderState, bReflectiveShadowmap && !bOnePassPointLightShadow);
-
-					ParallelCommandListSet.AddParallelCommandList(CmdList, AnyThreadCompletionEvent, StaticSubjectMeshElements.Num());
-				}
-				if (DynamicSubjectMeshElements.Num())
-				{
-					FRHICommandList* CmdList = ParallelCommandListSet.NewParallelCommandList();
-
-					FGraphEventRef AnyThreadCompletionEvent = TGraphTask<FRenderDepthDynamicThreadTask>::CreateTask(ParallelCommandListSet.GetPrereqs(), ENamedThreads::GetRenderThread())
-						.ConstructAndDispatchWhenReady(*this, *CmdList, *ShadowDepthView, DrawRenderState, SceneRenderer);
-
-					ParallelCommandListSet.AddParallelCommandList(CmdList, AnyThreadCompletionEvent, DynamicSubjectMeshElements.Num());
-				}
-			}
+			ShadowDepthPass.DispatchDraw(&ParallelCommandListSet, RHICmdList);
 		}
 
 		// Renderpass must be closed once we get here.
@@ -1896,45 +1148,7 @@ void FProjectedShadowInfo::RenderDepthInner(FRHICommandListImmediate& RHICmdList
 		// We must have already opened the renderpass by the time we get here.
 		check(RHICmdList.IsInsideRenderPass());
 
-		if (UseMeshDrawCommandPipeline())
-		{
-			ShadowDepthPass.DispatchDraw(nullptr, RHICmdList);
-		}
-		else
-		{
-			// Draw the subject's static elements using static draw lists
-			if (bIsWholeSceneDirectionalShadow)
-			{
-				SCOPE_CYCLE_COUNTER(STAT_WholeSceneStaticDrawListShadowDepthsTime);
-
-				if (bReflectiveShadowmap)
-				{
-					SceneRenderer->Scene->WholeSceneReflectiveShadowMapDrawList.DrawVisible(RHICmdList, *ShadowDepthView, *PolicyContext, DrawRenderState, StaticMeshWholeSceneShadowDepthMap, StaticMeshWholeSceneShadowBatchVisibility);
-				}
-				else
-				{
-					// Use the scene's shadow depth draw list with this shadow's visibility map
-					SceneRenderer->Scene->WholeSceneShadowDepthDrawList.DrawVisible(RHICmdList, *ShadowDepthView, *PolicyContext, DrawRenderState, StaticMeshWholeSceneShadowDepthMap, StaticMeshWholeSceneShadowBatchVisibility);
-				}
-			}
-			// Draw the subject's static elements using manual state filtering
-			else if (StaticSubjectMeshElements.Num() > 0)
-			{
-				SCOPE_CYCLE_COUNTER(STAT_WholeSceneStaticShadowDepthsTime);
-
-				if (bReflectiveShadowmap && !bOnePassPointLightShadow)
-				{
-					// reflective shadow map
-					DrawShadowMeshElements<true>(RHICmdList, *ShadowDepthView, DrawRenderState, *this);
-				}
-				else
-				{
-					// normal shadow map
-					DrawShadowMeshElements<false>(RHICmdList, *ShadowDepthView, DrawRenderState, *this);
-				}
-			}
-			RenderDepthDynamic(RHICmdList, SceneRenderer, ShadowDepthView, DrawRenderState);
-		}
+		ShadowDepthPass.DispatchDraw(nullptr, RHICmdList);
 
 		// Renderpass must still be open when we reach here
 		check(RHICmdList.IsInsideRenderPass());
@@ -2036,25 +1250,6 @@ void FProjectedShadowInfo::SetupShadowDepthView(FRHICommandListImmediate& RHICmd
 	ShadowDepthView = DepthPassView;
 }
 
-void FProjectedShadowInfo::SortSubjectMeshElements()
-{
-	// Note: this should match the criteria in FProjectedShadowInfo::RenderDepth for deciding when to call SetSharedState on a static mesh element for best performance
-	struct FCompareFShadowStaticMeshElement
-	{
-		FORCEINLINE bool operator()(const FShadowStaticMeshElement& A, const FShadowStaticMeshElement& B) const
-		{
-			if (A.Mesh->VertexFactory != B.Mesh->VertexFactory) return A.Mesh->VertexFactory < B.Mesh->VertexFactory;
-			if (A.RenderProxy != B.RenderProxy) return A.RenderProxy < B.RenderProxy;
-			if (A.bIsTwoSided != B.bIsTwoSided) return A.bIsTwoSided < B.bIsTwoSided;
-			if (A.Mesh->ReverseCulling != B.Mesh->ReverseCulling) return A.Mesh->ReverseCulling < B.Mesh->ReverseCulling;
-
-			return false;
-		}
-	};
-
-	StaticSubjectMeshElements.Sort(FCompareFShadowStaticMeshElement());
-}
-
 void FProjectedShadowInfo::GetShadowTypeNameForDrawEvent(FString& TypeName) const
 {
 	const FName ParentName = ParentSceneInfo ? ParentSceneInfo->Proxy->GetOwnerName() : NAME_None;
diff --git a/Engine/Source/Runtime/Renderer/Private/ShadowRendering.cpp b/Engine/Source/Runtime/Renderer/Private/ShadowRendering.cpp
--- a/Engine/Source/Runtime/Renderer/Private/ShadowRendering.cpp
+++ b/Engine/Source/Runtime/Renderer/Private/ShadowRendering.cpp
@@ -609,103 +609,14 @@ void FProjectedShadowInfo::SetupProjectionStencilMask(
 			RHICmdList.SetViewport(0, 0, 0, SceneRender->InstancedStereoWidth, View->ViewRect.Max.Y, 1);
 		}
 
-		if (UseMeshDrawCommandPipeline())
+		const FShadowMeshDrawCommandPass& ProjectionStencilingPass = ProjectionStencilingPasses[ViewIndex];
+		if (ProjectionStencilingPass.VisibleMeshDrawCommands.Num() > 0)
 		{
-			const FShadowMeshDrawCommandPass& ProjectionStencilingPass = ProjectionStencilingPasses[ViewIndex];
-			if (ProjectionStencilingPass.VisibleMeshDrawCommands.Num() > 0)
-			{
-				FVertexBufferRHIParamRef PrimitiveIdsBuffer = UseGPUScene(GMaxRHIShaderPlatform, View->GetFeatureLevel()) ? ProjectionStencilingPass.PrimitiveIdBuffer.VertexBuffer->VertexBufferRHI : nullptr;
-				const int32 BasePrimitiveIdsOffset = ProjectionStencilingPass.PrimitiveIdBuffer.VertexOffset;
-				const bool bDynamicInstancing = IsDynamicInstancingEnabled() && UseGPUScene(GMaxRHIShaderPlatform, View->GetFeatureLevel());
-
-				SubmitMeshDrawCommands(ProjectionStencilingPass.VisibleMeshDrawCommands, PrimitiveIdsBuffer, BasePrimitiveIdsOffset, bDynamicInstancing, RHICmdList);
-			}
-		}
-		else
-		{
-			// Set stencil to one.
-			DrawRenderState.SetDepthStencilState(
-				TStaticDepthStencilState<
-				false, CF_DepthNearOrEqual,
-				true, CF_Always, SO_Keep, SO_Keep, SO_Replace,
-				false, CF_Always, SO_Keep, SO_Keep, SO_Keep,
-				0xff, 0xff
-				>::GetRHI());
-			DrawRenderState.SetStencilRef(1);
-
-			// Pre-shadows mask by receiver elements, self-shadow mask by subject elements.
-			// Note that self-shadow pre-shadows still mask by receiver elements.
-			const TArray<FMeshBatchAndRelevance, SceneRenderingAllocator>& DynamicMeshElements = bPreShadow ? DynamicReceiverMeshElements : DynamicSubjectMeshElements;
-
-			FDepthDrawingPolicyFactory::ContextType Context(DDM_AllOccluders, false);
-
-			for (int32 MeshBatchIndex = 0; MeshBatchIndex < DynamicMeshElements.Num(); MeshBatchIndex++)
-			{
-				const FMeshBatchAndRelevance& MeshBatchAndRelevance = DynamicMeshElements[MeshBatchIndex];
-				const FMeshBatch& MeshBatch = *MeshBatchAndRelevance.Mesh;
-
-				FDepthDrawingPolicyFactory::DrawDynamicMesh(RHICmdList, *View, Context, MeshBatch, true, DrawRenderState, MeshBatchAndRelevance.PrimitiveSceneProxy, MeshBatch.BatchHitProxyId, false, bIsInstancedStereoEmulated);
-			}
-
-			// Pre-shadows mask by receiver elements, self-shadow mask by subject elements.
-			// Note that self-shadow pre-shadows still mask by receiver elements.
-			const PrimitiveArrayType& MaskPrimitives = bPreShadow ? ReceiverPrimitives : DynamicSubjectPrimitives;
-
-			for (int32 PrimitiveIndex = 0, PrimitiveCount = MaskPrimitives.Num(); PrimitiveIndex < PrimitiveCount; PrimitiveIndex++)
-			{
-				const FPrimitiveSceneInfo* ReceiverPrimitiveSceneInfo = MaskPrimitives[PrimitiveIndex];
+			FVertexBufferRHIParamRef PrimitiveIdsBuffer = UseGPUScene(GMaxRHIShaderPlatform, View->GetFeatureLevel()) ? ProjectionStencilingPass.PrimitiveIdBuffer.VertexBuffer->VertexBufferRHI : nullptr;
+			const int32 BasePrimitiveIdsOffset = ProjectionStencilingPass.PrimitiveIdBuffer.VertexOffset;
+			const bool bDynamicInstancing = IsDynamicInstancingEnabled() && UseGPUScene(GMaxRHIShaderPlatform, View->GetFeatureLevel());
 
-				if (View->PrimitiveVisibilityMap[ReceiverPrimitiveSceneInfo->GetIndex()])
-				{
-					const FPrimitiveViewRelevance& ViewRelevance = View->PrimitiveViewRelevanceMap[ReceiverPrimitiveSceneInfo->GetIndex()];
-
-					if (ViewRelevance.bRenderInMainPass && ViewRelevance.bStaticRelevance)
-					{
-						for (int32 StaticMeshIdx = 0; StaticMeshIdx < ReceiverPrimitiveSceneInfo->StaticMeshes.Num(); StaticMeshIdx++)
-						{
-							const FStaticMesh& StaticMesh = ReceiverPrimitiveSceneInfo->StaticMeshes[StaticMeshIdx];
-
-							if (View->StaticMeshVisibilityMap[StaticMesh.Id])
-							{
-								FDepthDrawingPolicyFactory::DrawStaticMesh(
-									RHICmdList, 
-									*View,
-									FDepthDrawingPolicyFactory::ContextType(DDM_AllOccluders, false),
-									StaticMesh,
-									StaticMesh.bRequiresPerElementVisibility ? View->StaticMeshBatchVisibility[StaticMesh.BatchVisibilityId] : ((1ull << StaticMesh.Elements.Num() )- 1),
-									true,
-									DrawRenderState,
-									ReceiverPrimitiveSceneInfo->Proxy,
-									StaticMesh.BatchHitProxyId, 
-									false, 
-									bIsInstancedStereoEmulated
-									);
-							}
-						}
-					}
-				}
-			}
-
-			if (bSelfShadowOnly && !bPreShadow)
-			{
-				for (int32 ElementIndex = 0; ElementIndex < StaticSubjectMeshElements.Num(); ++ElementIndex)
-				{
-					const FStaticMesh& StaticMesh = *StaticSubjectMeshElements[ElementIndex].Mesh;
-					FDepthDrawingPolicyFactory::DrawStaticMesh(
-						RHICmdList, 
-						*View,
-						FDepthDrawingPolicyFactory::ContextType(DDM_AllOccluders, false),
-						StaticMesh,
-						StaticMesh.bRequiresPerElementVisibility ? View->StaticMeshBatchVisibility[StaticMesh.BatchVisibilityId] : ((1ull << StaticMesh.Elements.Num() )- 1),
-						true,
-						DrawRenderState,
-						StaticMesh.PrimitiveSceneInfo->Proxy,
-						StaticMesh.BatchHitProxyId, 
-						false, 
-						bIsInstancedStereoEmulated
-						);
-				}
-			}
+			SubmitMeshDrawCommands(ProjectionStencilingPass.VisibleMeshDrawCommands, PrimitiveIdsBuffer, BasePrimitiveIdsOffset, bDynamicInstancing, RHICmdList);
 		}
 
 		// Restore viewport
diff --git a/Engine/Source/Runtime/Renderer/Private/ShadowRendering.h b/Engine/Source/Runtime/Renderer/Private/ShadowRendering.h
--- a/Engine/Source/Runtime/Renderer/Private/ShadowRendering.h
+++ b/Engine/Source/Runtime/Renderer/Private/ShadowRendering.h
@@ -21,7 +21,6 @@
 #include "SceneManagement.h"
 #include "ScenePrivateBase.h"
 #include "SceneCore.h"
-#include "DrawingPolicy.h"
 #include "GlobalShader.h"
 #include "SystemTextures.h"
 #include "PostProcess/SceneRenderTargets.h"
@@ -37,167 +36,12 @@ class FPrimitiveSceneProxy;
 class FProjectedShadowInfo;
 class FScene;
 class FSceneRenderer;
-class FShadowStaticMeshElement;
 class FViewInfo;
 
 /** Renders a cone with a spherical cap, used for rendering spot lights in deferred passes. */
 extern void DrawStencilingCone(const FMatrix& ConeToWorld, float ConeAngle, float SphereRadius, const FVector& PreViewTranslation);
 
 template <bool bRenderingReflectiveShadowMaps> class TShadowDepthBasePS;
-class FShadowStaticMeshElement;
-
-/**
- * The shadow depth drawing policy's context data.
- */
-struct FShadowDepthDrawingPolicyContext : FMeshDrawingPolicy::ContextDataType
-{
-	/** CAUTION, this is assumed to be a POD type. We allocate the on the scene allocator and NEVER CALL A DESTRUCTOR.
-		If you want to add non-pod data, not a huge problem, we just need to track and destruct them at the end of the scene.
-	**/
-	/** The projected shadow info for which we are rendering shadow depths. */
-	const FProjectedShadowInfo* ShadowInfo;
-
-	/** Initialization constructor. */
-	explicit FShadowDepthDrawingPolicyContext(const FProjectedShadowInfo* InShadowInfo)
-		: ShadowInfo(InShadowInfo)
-	{}
-};
-
-/**
- * Outputs no color, but can be used to write the mesh's depth values to the depth buffer.
- */
-template <bool bRenderingReflectiveShadowMaps>
-class FShadowDepthDrawingPolicy : public FMeshDrawingPolicy
-{
-public:
-	typedef FShadowDepthDrawingPolicyContext ContextDataType;
-
-	FShadowDepthDrawingPolicy(
-		const FMaterial* InMaterialResource,
-		bool bInDirectionalLight,
-		bool bInOnePassPointLightShadow,
-		bool bInPreShadow,
-		const FMeshDrawingPolicyOverrideSettings& InOverrideSettings,
-		ERHIFeatureLevel::Type InFeatureLevel,
-		const FVertexFactory* InVertexFactory = 0,
-		const FMaterialRenderProxy* InMaterialRenderProxy = 0,
-		bool bReverseCulling = false
-		);
-
-	void UpdateElementState(FShadowStaticMeshElement& State, ERHIFeatureLevel::Type FeatureLevel);
-
-	FShadowDepthDrawingPolicy& operator = (const FShadowDepthDrawingPolicy& Other)
-	{ 
-		VertexShader = Other.VertexShader;
-		GeometryShader = Other.GeometryShader;
-		HullShader = Other.HullShader;
-		DomainShader = Other.DomainShader;
-		PixelShader = Other.PixelShader;
-		bDirectionalLight = Other.bDirectionalLight;
-		bReverseCulling = Other.bReverseCulling;
-		bOnePassPointLightShadow = Other.bOnePassPointLightShadow;
-		bUsePositionOnlyVS = Other.bUsePositionOnlyVS;
-		bPreShadow = Other.bPreShadow;
-		FeatureLevel = Other.FeatureLevel;
-		(FMeshDrawingPolicy&)*this = (const FMeshDrawingPolicy&)Other;
-		return *this; 
-	}
-
-	//~ Begin FMeshDrawingPolicy Interface.
-	FDrawingPolicyMatchResult Matches(const FShadowDepthDrawingPolicy& Other, bool bForReals = false) const
-	{
-		DRAWING_POLICY_MATCH_BEGIN
-			DRAWING_POLICY_MATCH(FMeshDrawingPolicy::Matches(Other, bForReals)) &&
-			DRAWING_POLICY_MATCH(VertexShader == Other.VertexShader) &&
-			DRAWING_POLICY_MATCH(GeometryShader == Other.GeometryShader) &&
-			DRAWING_POLICY_MATCH(HullShader == Other.HullShader) &&
-			DRAWING_POLICY_MATCH(DomainShader == Other.DomainShader) &&
-			DRAWING_POLICY_MATCH(PixelShader == Other.PixelShader) &&
-			DRAWING_POLICY_MATCH(bDirectionalLight == Other.bDirectionalLight) &&
-			DRAWING_POLICY_MATCH(bReverseCulling == Other.bReverseCulling) &&
-			DRAWING_POLICY_MATCH(bOnePassPointLightShadow == Other.bOnePassPointLightShadow) &&
-			DRAWING_POLICY_MATCH(bUsePositionOnlyVS == Other.bUsePositionOnlyVS) &&
-			DRAWING_POLICY_MATCH(bPreShadow == Other.bPreShadow) &&
-			DRAWING_POLICY_MATCH(FeatureLevel == Other.FeatureLevel);
-		DRAWING_POLICY_MATCH_END
-	}
-	void SetSharedState(FRHICommandList& RHICmdList, const FDrawingPolicyRenderState& DrawRenderState, const FSceneView* View, const ContextDataType PolicyContext) const;
-
-	/** 
-	 * Create bound shader state using the vertex decl from the mesh draw policy
-	 * as well as the shaders needed to draw the mesh
-	 * @return new bound shader state object
-	 */
-	FBoundShaderStateInput GetBoundShaderStateInput(ERHIFeatureLevel::Type InFeatureLevel) const;
-
-	void SetMeshRenderState(
-		FRHICommandList& RHICmdList, 
-		const FSceneView& View,
-		const FPrimitiveSceneProxy* PrimitiveSceneProxy,
-		const FMeshBatch& Mesh,
-		int32 BatchElementIndex,
-		const FDrawingPolicyRenderState& DrawRenderState,
-		const ElementDataType& ElementData,
-		const ContextDataType PolicyContext
-		) const;
-
-	template<bool T2>
-	friend int32 CompareDrawingPolicy(const FShadowDepthDrawingPolicy<T2>& A,const FShadowDepthDrawingPolicy<T2>& B);
-
-	bool IsReversingCulling() const
-	{
-		return bReverseCulling;
-	}
-	
-private:
-
-	class FShadowDepthVS* VertexShader;
-	class FOnePassPointShadowDepthGS* GeometryShader;
-	TShadowDepthBasePS<bRenderingReflectiveShadowMaps>* PixelShader;
-	class FBaseHS* HullShader;
-	class FBaseDS* DomainShader;
-	ERHIFeatureLevel::Type FeatureLevel;
-
-public:
-
-	uint32 bDirectionalLight:1;
-	uint32 bReverseCulling:1;
-	uint32 bOnePassPointLightShadow:1;
-	uint32 bUsePositionOnlyVS:1;
-	uint32 bPreShadow:1;
-};
-
-/**
- * A drawing policy factory for the shadow depth drawing policy.
- */
-class FShadowDepthDrawingPolicyFactory
-{
-public:
-
-	enum { bAllowSimpleElements = false };
-
-	struct ContextType
-	{
-		const FProjectedShadowInfo* ShadowInfo;
-
-		ContextType(const FProjectedShadowInfo* InShadowInfo) :
-			ShadowInfo(InShadowInfo)
-		{}
-	};
-
-	static void AddStaticMesh(FScene* Scene,FStaticMesh* StaticMesh);
-
-	static bool DrawDynamicMesh(
-		FRHICommandList& RHICmdList, 
-		const FSceneView& View,
-		ContextType Context,
-		const FMeshBatch& Mesh,
-		bool bPreFog,
-		const FDrawingPolicyRenderState& DrawRenderState,
-		const FPrimitiveSceneProxy* PrimitiveSceneProxy,
-		FHitProxyId HitProxyId
-		);
-};
 
 /** 
  * Overrides a material used for shadow depth rendering with the default material when appropriate.
@@ -211,49 +55,6 @@ void OverrideWithDefaultMaterialForShadowDepth(
 	ERHIFeatureLevel::Type InFeatureLevel
 	);
 
-/** A single static mesh element for shadow depth rendering. */
-class FShadowStaticMeshElement
-{
-public:
-
-	FShadowStaticMeshElement()
-		: RenderProxy(0)
-		, MaterialResource(0)
-		, Mesh(0)
-		, bIsTwoSided(false)
-	{
-	}
-
-	FShadowStaticMeshElement(const FMaterialRenderProxy* InRenderProxy, const FMaterial* InMaterialResource, const FStaticMesh* InMesh, bool bInIsTwoSided) :
-		RenderProxy(InRenderProxy),
-		MaterialResource(InMaterialResource),
-		Mesh(InMesh),
-		bIsTwoSided(bInIsTwoSided)
-	{}
-
-	bool DoesDeltaRequireADrawSharedCall(const FShadowStaticMeshElement& rhs) const
-	{
-		checkSlow(rhs.RenderProxy);
-		checkSlow(rhs.Mesh);
-
-		// Note: this->RenderProxy or this->Mesh can be 0
-		// but in this case rhs.RenderProxy should not be 0
-		// so it will early out and there will be no crash on Mesh->VertexFactory
-		checkSlow(!RenderProxy || rhs.RenderProxy);
-
-		return RenderProxy != rhs.RenderProxy
-			|| bIsTwoSided != rhs.bIsTwoSided
-			|| Mesh->VertexFactory != rhs.Mesh->VertexFactory
-			|| Mesh->ReverseCulling != rhs.Mesh->ReverseCulling;
-	}
-
-	/** Store the FMaterialRenderProxy pointer since it may be different from the one that FStaticMesh stores. */
-	const FMaterialRenderProxy* RenderProxy;
-	const FMaterial* MaterialResource;
-	const FStaticMesh* Mesh;
-	bool bIsTwoSided;
-};
-
 enum EShadowDepthRenderMode
 {
 	/** The render mode used by regular shadows */
@@ -518,9 +319,6 @@ public:
 	/** Whether turn on back-lighting transmission. */
 	uint32 bTransmission : 1;
 
-	TBitArray<SceneRenderingBitArrayAllocator> StaticMeshWholeSceneShadowDepthMap;
-	TArray<uint64,SceneRenderingAllocator> StaticMeshWholeSceneShadowBatchVisibility;
-
 	/** View projection matrices for each cubemap face, used by one pass point light shadows. */
 	TArray<FMatrix> OnePassShadowViewProjectionMatrices;
 
@@ -684,9 +482,6 @@ public:
 		return bWholeSceneShadow && ( LightSceneInfo->Proxy->GetLightType() == LightType_Point || LightSceneInfo->Proxy->GetLightType() == LightType_Rect );
 	}
 
-	/** Sorts StaticSubjectMeshElements based on state so that rendering the static elements will set as little state as possible. */
-	void SortSubjectMeshElements();
-
 	// 0 if Setup...() wasn't called yet
 	const FLightSceneInfo& GetLightSceneInfo() const { return *LightSceneInfo; }
 	const FLightSceneInfoCompact& GetLightSceneInfoCompact() const { return LightSceneInfoCompact; }
@@ -722,9 +517,6 @@ private:
 	/** Subject primitives with translucent relevance. */
 	PrimitiveArrayType SubjectTranslucentPrimitives;
 
-	/** Static shadow casting elements. */
-	TArray<FShadowStaticMeshElement,SceneRenderingAllocator> StaticSubjectMeshElements;
-
 	/** Dynamic mesh elements for subject primitives. */
 	TArray<FMeshBatchAndRelevance,SceneRenderingAllocator> DynamicSubjectMeshElements;
 	/** Dynamic mesh elements for receiver primitives. */
@@ -784,16 +576,8 @@ private:
 	/** Will return if we should draw the static mesh for the shadow, and will perform lazy init of primitive if it wasn't visible */
 	bool ShouldDrawStaticMeshes(FViewInfo& InCurrentView, bool bInCustomDataRelevance, FPrimitiveSceneInfo* InPrimitiveSceneInfo);
 
-	/**
-	* Renders the dynamic shadow subject depth, to a particular hacked view
-	*/
-	friend class FRenderDepthDynamicThreadTask;
-	void RenderDepthDynamic(FRHICommandList& RHICmdList, class FSceneRenderer* SceneRenderer, const FViewInfo* FoundView, const FDrawingPolicyRenderState& DrawRenderState);
-
 	void GetShadowTypeNameForDrawEvent(FString& TypeName) const;
 
-	template <bool bReflectiveShadowmap> friend void DrawShadowMeshElements(FRHICommandList& RHICmdList, const FViewInfo& View, const FDrawingPolicyRenderState& DrawRenderState, const FProjectedShadowInfo& ShadowInfo);
-
 	/** Updates object buffers needed by ray traced distance field shadows. */
 	int32 UpdateShadowCastingObjectBuffers() const;
 
@@ -821,7 +605,6 @@ private:
 	template <bool bRenderingReflectiveShadowMaps> friend class TShadowDepthBasePS;
 	friend class FShadowVolumeBoundProjectionVS;
 	friend class FShadowProjectionPS;
-	friend class FShadowDepthDrawingPolicyFactory;
 };
 
 /**
diff --git a/Engine/Source/Runtime/Renderer/Private/ShadowSetup.cpp b/Engine/Source/Runtime/Renderer/Private/ShadowSetup.cpp
--- a/Engine/Source/Runtime/Renderer/Private/ShadowSetup.cpp
+++ b/Engine/Source/Runtime/Renderer/Private/ShadowSetup.cpp
@@ -30,6 +30,7 @@
 #include "RendererModule.h"
 #include "LightPropagationVolumeSettings.h"
 #include "CapsuleShadowRendering.h"
+#include "Async/ParallelFor.h"
 
 static float GMinScreenRadiusForShadowCaster = 0.01f;
 static FAutoConsoleVariableRef CVarMinScreenRadiusForShadowCaster(
@@ -809,40 +810,37 @@ void FProjectedShadowInfo::AddCachedMeshDrawCommandsForPass(
 	TArray<const FStaticMesh*, SceneRenderingAllocator>& MeshCommandBuildRequests,
 	int32& NumMeshCommandBuildRequestElements)
 {
-	if (UseMeshDrawCommandPipeline())
-	{
-		const EShadingPath ShadingPath = Scene->GetShadingPath();
-		const bool bUseCachedMeshCommand = UseCachedMeshDrawCommands()
-			&& !!(FPassProcessorManager::GetPassFlags(ShadingPath, PassType) & EMeshPassFlags::CachedMeshCommands)
-			&& StaticMeshRelevance.bSupportsCachingMeshDrawCommands;
+	const EShadingPath ShadingPath = Scene->GetShadingPath();
+	const bool bUseCachedMeshCommand = UseCachedMeshDrawCommands()
+		&& !!(FPassProcessorManager::GetPassFlags(ShadingPath, PassType) & EMeshPassFlags::CachedMeshCommands)
+		&& StaticMeshRelevance.bSupportsCachingMeshDrawCommands;
 
-		if (bUseCachedMeshCommand)
+	if (bUseCachedMeshCommand)
+	{
+		const int32 StaticMeshCommandInfoIndex = StaticMeshRelevance.GetStaticMeshCommandInfoIndex(PassType);
+		if (StaticMeshCommandInfoIndex >= 0)
 		{
-			const int32 StaticMeshCommandInfoIndex = StaticMeshRelevance.GetStaticMeshCommandInfoIndex(PassType);
-			if (StaticMeshCommandInfoIndex >= 0)
-			{
-				const FCachedMeshDrawCommandInfo& CachedMeshDrawCommand = InPrimitiveSceneInfo->StaticMeshCommandInfos[StaticMeshCommandInfoIndex];
-				const FCachedPassMeshDrawList& SceneDrawList = Scene->CachedDrawLists[PassType];
-
-				FVisibleMeshDrawCommand NewVisibleMeshDrawCommand;
+			const FCachedMeshDrawCommandInfo& CachedMeshDrawCommand = InPrimitiveSceneInfo->StaticMeshCommandInfos[StaticMeshCommandInfoIndex];
+			const FCachedPassMeshDrawList& SceneDrawList = Scene->CachedDrawLists[PassType];
 
-				NewVisibleMeshDrawCommand.Setup(
-					&SceneDrawList.MeshDrawCommands[CachedMeshDrawCommand.CommandIndex],
-					PrimitiveIndex,
-					CachedMeshDrawCommand.StateBucketId,
-					CachedMeshDrawCommand.MeshFillMode,
-					CachedMeshDrawCommand.MeshCullMode,
-					FMeshDrawCommandSortKey::Default);
+			FVisibleMeshDrawCommand NewVisibleMeshDrawCommand;
 
-				VisibleMeshCommands.Add(NewVisibleMeshDrawCommand);
-			}
-		}
-		else
-		{
-			NumMeshCommandBuildRequestElements += StaticMeshRelevance.NumElements;
-			MeshCommandBuildRequests.Add(&StaticMesh);
+			NewVisibleMeshDrawCommand.Setup(
+				&SceneDrawList.MeshDrawCommands[CachedMeshDrawCommand.CommandIndex],
+				PrimitiveIndex,
+				CachedMeshDrawCommand.StateBucketId,
+				CachedMeshDrawCommand.MeshFillMode,
+				CachedMeshDrawCommand.MeshCullMode,
+				FMeshDrawCommandSortKey::Default);
+
+			VisibleMeshCommands.Add(NewVisibleMeshDrawCommand);
 		}
 	}
+	else
+	{
+		NumMeshCommandBuildRequestElements += StaticMeshRelevance.NumElements;
+		MeshCommandBuildRequests.Add(&StaticMesh);
+	}
 }
 
 bool FProjectedShadowInfo::ShouldDrawStaticMeshes(FViewInfo& InCurrentView, bool InbCustomDataRelevance, FPrimitiveSceneInfo* InPrimitiveSceneInfo)
@@ -923,35 +921,23 @@ bool FProjectedShadowInfo::ShouldDrawStaticMeshes(FViewInfo& InCurrentView, bool
 
 				if ((StaticMeshRelevance.CastShadow || (bSelfShadowOnly && StaticMeshRelevance.bUseForDepthPass)) && ShadowLODToRender.ContainsLOD(StaticMeshRelevance.LODIndex))
 				{
-					if (UseMeshDrawCommandPipeline())
-					{
-						if (GetShadowDepthType() == CSMShadowDepthType && bCanCache)
-						{
-							AddCachedMeshDrawCommandsForPass(
-								PrimitiveId,
-								InPrimitiveSceneInfo,
-								StaticMeshRelevance,
-								StaticMesh,
-								InPrimitiveSceneInfo->Scene,
-								EMeshPass::CSMShadowDepth,
-								ShadowDepthPassVisibleCommands,
-								SubjectMeshCommandBuildRequests,
-								NumSubjectMeshCommandBuildRequestElements);
-						}
-						else
-						{
-							NumSubjectMeshCommandBuildRequestElements += StaticMeshRelevance.NumElements;
-							SubjectMeshCommandBuildRequests.Add(&StaticMesh);
-						}
-					}
-					else if (StaticMeshRelevance.CastShadow)
+					if (GetShadowDepthType() == CSMShadowDepthType && bCanCache)
 					{
-						StaticMeshWholeSceneShadowDepthMap[StaticMeshRelevance.Id] = true;
+						AddCachedMeshDrawCommandsForPass(
+							PrimitiveId,
+							InPrimitiveSceneInfo,
+							StaticMeshRelevance,
+							StaticMesh,
+							InPrimitiveSceneInfo->Scene,
+							EMeshPass::CSMShadowDepth,
+							ShadowDepthPassVisibleCommands,
+							SubjectMeshCommandBuildRequests,
+							NumSubjectMeshCommandBuildRequestElements);
 					}
-
-					if (StaticMeshRelevance.bRequiresPerElementVisibility)
+					else
 					{
-						StaticMeshWholeSceneShadowBatchVisibility[StaticMesh.BatchVisibilityId] = StaticMesh.VertexFactory->GetStaticBatchElementVisibility(InCurrentView, &StaticMesh);
+						NumSubjectMeshCommandBuildRequestElements += StaticMeshRelevance.NumElements;
+						SubjectMeshCommandBuildRequests.Add(&StaticMesh);
 					}
 
 					bDrawingStaticMeshes = true;
@@ -967,15 +953,8 @@ bool FProjectedShadowInfo::ShouldDrawStaticMeshes(FViewInfo& InCurrentView, bool
 
 				if ((StaticMeshRelevance.CastShadow || (bSelfShadowOnly && StaticMeshRelevance.bUseForDepthPass)) && ShadowLODToRender.ContainsLOD(StaticMeshRelevance.LODIndex))
 				{
-					if (UseMeshDrawCommandPipeline())
-					{
-						NumSubjectMeshCommandBuildRequestElements += StaticMeshRelevance.NumElements;
-						SubjectMeshCommandBuildRequests.Add(&StaticMesh);
-					}
-					else if (StaticMeshRelevance.CastShadow)
-					{
-						InCurrentView.StaticMeshShadowDepthMap[StaticMeshRelevance.Id] = true;
-					}
+					NumSubjectMeshCommandBuildRequestElements += StaticMeshRelevance.NumElements;
+					SubjectMeshCommandBuildRequests.Add(&StaticMesh);
 
 					if (StaticMeshRelevance.bRequiresPerElementVisibility)
 					{
@@ -1092,20 +1071,8 @@ void FProjectedShadowInfo::AddSubjectPrimitive(FPrimitiveSceneInfo* PrimitiveSce
 			*(PrimitiveSceneInfo->ComponentLastRenderTime) = CurrentWorldTime;
 			if (PrimitiveSceneInfo->NeedsUniformBufferUpdate() || PrimitiveSceneInfo->NeedsUpdateStaticMeshes())
 			{
-				if (GDrawListsLocked && PrimitiveSceneInfo->NeedsUpdateStaticMeshes())
-				{
-					QUICK_SCOPE_CYCLE_COUNTER(STAT_FProjectedShadowInfo_AddSubjectPrimitive_FlushPrepass);
-					FRHICommandListExecutor::GetImmediateCommandList().ImmediateFlush(EImmediateFlushType::WaitForOutstandingTasksOnly);
-					FParallelCommandListSet::WaitForTasks();
-					TGuardValue<bool> LockDrawLists(GDrawListsLocked, false);
-					PrimitiveSceneInfo->ConditionalUpdateStaticMeshes(FRHICommandListExecutor::GetImmediateCommandList());
-					PrimitiveSceneInfo->ConditionalUpdateUniformBuffer(FRHICommandListExecutor::GetImmediateCommandList());
-				}
-				else
-				{
-					PrimitiveSceneInfo->ConditionalUpdateStaticMeshes(FRHICommandListExecutor::GetImmediateCommandList());
-					PrimitiveSceneInfo->ConditionalUpdateUniformBuffer(FRHICommandListExecutor::GetImmediateCommandList());
-				}
+				PrimitiveSceneInfo->ConditionalUpdateStaticMeshes(FRHICommandListExecutor::GetImmediateCommandList());
+				PrimitiveSceneInfo->ConditionalUpdateUniformBuffer(FRHICommandListExecutor::GetImmediateCommandList());
 			}
 		}
 
@@ -1142,30 +1109,7 @@ void FProjectedShadowInfo::AddSubjectPrimitive(FPrimitiveSceneInfo* PrimitiveSce
 
 			if (bDrawingStaticMeshes)
 			{
-				if (!bWholeSceneDirectionalShadow && !UseMeshDrawCommandPipeline())
-				{
-					// Add the primitive's static mesh elements to the draw lists.
-					for (int32 MeshIndex = 0; MeshIndex < PrimitiveSceneInfo->StaticMeshes.Num(); MeshIndex++)
-					{
-						FStaticMesh& StaticMesh = PrimitiveSceneInfo->StaticMeshes[MeshIndex];
-
-						if (StaticMesh.CastShadow)
-						{
-							const FMaterialRenderProxy* MaterialRenderProxy = StaticMesh.MaterialRenderProxy;
-							const FMaterial* Material = MaterialRenderProxy->GetMaterial(FeatureLevel);
-							const EBlendMode BlendMode = Material->GetBlendMode();
-							const EMaterialShadingModel ShadingModel = Material->GetShadingModel();
-
-							if(Material->ShouldCastDynamicShadows() || (bReflectiveShadowmap && Material->ShouldInjectEmissiveIntoLPV())) 
-							{
-								const bool bTwoSided = Material->IsTwoSided() || PrimitiveSceneInfo->Proxy->CastsShadowAsTwoSided();
-								OverrideWithDefaultMaterialForShadowDepth(MaterialRenderProxy, Material, bReflectiveShadowmap, FeatureLevel);
-								StaticSubjectMeshElements.Add(FShadowStaticMeshElement(MaterialRenderProxy, Material, &StaticMesh,bTwoSided));
-							}
-						}
-					}
-				}
-				else if (bRecordShadowSubjectsForMobileShading)
+				if (bRecordShadowSubjectsForMobileShading)
 				{
 					DependentView->VisibleLightInfos[GetLightSceneInfo().Id].MobileCSMSubjectPrimitives.AddSubjectPrimitive(PrimitiveSceneInfo, PrimitiveId);
 				}
@@ -1193,7 +1137,6 @@ void FProjectedShadowInfo::AddSubjectPrimitive(FPrimitiveSceneInfo* PrimitiveSce
 bool FProjectedShadowInfo::HasSubjectPrims() const
 {
 	return DynamicSubjectPrimitives.Num() > 0
-		|| StaticSubjectMeshElements.Num() > 0
 		|| ShadowDepthPass.HasAnyDraw()
 		|| SubjectMeshCommandBuildRequests.Num() > 0;
 }
@@ -1370,7 +1313,7 @@ void FProjectedShadowInfo::ApplyViewOverridesToMeshDrawCommands(const FViewInfo&
 			const FMeshDrawCommand& MeshCommand = *VisibleMeshDrawCommand.MeshDrawCommand;
 			NewMeshCommand = MeshCommand;
 
-			const ERasterizerCullMode LocalCullMode = View.bRenderSceneTwoSided ? CM_None : View.bReverseCulling ? FMeshDrawingPolicy::InverseCullMode(VisibleMeshDrawCommand.MeshCullMode) : VisibleMeshDrawCommand.MeshCullMode;
+			const ERasterizerCullMode LocalCullMode = View.bRenderSceneTwoSided ? CM_None : View.bReverseCulling ? FMeshPassProcessor::InverseCullMode(VisibleMeshDrawCommand.MeshCullMode) : VisibleMeshDrawCommand.MeshCullMode;
 			NewMeshCommand.PipelineState.RasterizerState = GetStaticRasterizerState<true>(VisibleMeshDrawCommand.MeshFillMode, LocalCullMode);
 
 			NewMeshCommand.Finalize(true);
@@ -1458,11 +1401,8 @@ void FProjectedShadowInfo::GatherDynamicMeshElements(FSceneRenderer& Renderer, F
 		PassUniformBuffer = MobileShadowDepthPassUniformBuffer;
 	}
 
-	if (UseMeshDrawCommandPipeline())
-	{
-		SetupMeshDrawCommandsForShadowDepth(Renderer, PassUniformBuffer);
-		SetupMeshDrawCommandsForProjectionStenciling(Renderer);
-	}
+	SetupMeshDrawCommandsForShadowDepth(Renderer, PassUniformBuffer);
+	SetupMeshDrawCommandsForProjectionStenciling(Renderer);
 }
 
 void FProjectedShadowInfo::GatherDynamicMeshElementsArray(
@@ -1549,7 +1489,6 @@ void FProjectedShadowInfo::ClearTransientArrays()
 	SubjectTranslucentPrimitives.Empty();
 	DynamicSubjectPrimitives.Empty();
 	ReceiverPrimitives.Empty();
-	StaticSubjectMeshElements.Empty();
 	DynamicSubjectMeshElements.Empty();
 	DynamicReceiverMeshElements.Empty();
 	DynamicSubjectTranslucentMeshElements.Empty();
@@ -3187,15 +3126,6 @@ void FSceneRenderer::GatherShadowPrimitives(
 
 	if (PreShadows.Num() || ViewDependentWholeSceneShadows.Num())
 	{
-		for (int32 ShadowIndex = 0, Num = ViewDependentWholeSceneShadows.Num(); ShadowIndex < Num;ShadowIndex++)
-		{
-			FProjectedShadowInfo* ProjectedShadowInfo = ViewDependentWholeSceneShadows[ShadowIndex];
-			checkSlow(ProjectedShadowInfo->DependentView);
-			// Initialize the whole scene shadow's depth map with the shadow independent depth map from the view
-			ProjectedShadowInfo->StaticMeshWholeSceneShadowDepthMap.Init(false,Scene->StaticMeshes.GetMaxIndex());
-			ProjectedShadowInfo->StaticMeshWholeSceneShadowBatchVisibility.AddZeroed(Scene->StaticMeshBatchVisibility.GetMaxIndex());
-		}
-
 		TArray<FGatherShadowPrimitivesPacket*,SceneRenderingAllocator> Packets;
 
 		if (GUseOctreeForShadowCulling)
@@ -3317,14 +3247,6 @@ void FSceneRenderer::GatherShadowPrimitives(
 				Packet->~FGatherShadowPrimitivesPacket();
 			}
 		}
-
-		for (int32 ShadowIndex = 0, Num = PreShadows.Num(); ShadowIndex < Num; ShadowIndex++)
-		{
-			FProjectedShadowInfo* ProjectedShadowInfo = PreShadows[ShadowIndex];
-			//@todo - sort other shadow types' subject mesh elements?
-			// Probably needed for good performance with non-dominant whole scene shadows (spotlightmovable)
-			ProjectedShadowInfo->SortSubjectMeshElements();
-		}
 	}
 }
 
diff --git a/Engine/Source/Runtime/Renderer/Private/TranslucentLighting.cpp b/Engine/Source/Runtime/Renderer/Private/TranslucentLighting.cpp
--- a/Engine/Source/Runtime/Renderer/Private/TranslucentLighting.cpp
+++ b/Engine/Source/Runtime/Renderer/Private/TranslucentLighting.cpp
@@ -29,7 +29,6 @@
 #include "MaterialShaderType.h"
 #include "MaterialShader.h"
 #include "MeshMaterialShaderType.h"
-#include "DrawingPolicy.h"
 #include "MeshMaterialShader.h"
 #include "ShadowRendering.h"
 #include "SceneRendering.h"
@@ -266,7 +265,7 @@ public:
 		Ar << ShadowParameters;
 		return bShaderHasOutdatedParameters;
 	}
-
+	/*
 	void SetParameters(
 		FRHICommandList& RHICmdList, 
 		const FMaterialRenderProxy* MaterialRenderProxy,
@@ -283,7 +282,7 @@ public:
 	{
 		FMeshMaterialShader::SetMesh(RHICmdList, GetVertexShader(),VertexFactory,View,Proxy,BatchElement,DrawRenderState);
 	}
-
+	*/
 private:
 	FShadowDepthShaderParameters ShadowParameters;
 };
@@ -339,7 +338,7 @@ public:
 	}
 
 	FTranslucencyShadowDepthPS() {}
-
+	/*
 	void SetParameters(
 		FRHICommandList& RHICmdList, 
 		const FMaterialRenderProxy* MaterialRenderProxy,
@@ -368,7 +367,7 @@ public:
 	{
 		FMeshMaterialShader::SetMesh(RHICmdList, GetPixelShader(),VertexFactory,View,Proxy,BatchElement,DrawRenderState);
 	}
-
+	*/
 	virtual bool Serialize(FArchive& Ar) override
 	{
 		bool bShaderHasOutdatedParameters = FMeshMaterialShader::Serialize(Ar);
@@ -405,6 +404,9 @@ public:
 IMPLEMENT_MATERIAL_SHADER_TYPE(template<>,TTranslucencyShadowDepthPS<TranslucencyShadowDepth_PerspectiveCorrect>,TEXT("/Engine/Private/TranslucentShadowDepthShaders.usf"),TEXT("MainOpacityPS"),SF_Pixel);
 IMPLEMENT_MATERIAL_SHADER_TYPE(template<>,TTranslucencyShadowDepthPS<TranslucencyShadowDepth_Standard>,TEXT("/Engine/Private/TranslucentShadowDepthShaders.usf"),TEXT("MainOpacityPS"),SF_Pixel);
 
+//@todo MeshCommandPipeline
+#if 0
+
 /**
  * Drawing policy used to create Fourier opacity maps
  */
@@ -585,6 +587,8 @@ public:
 	}
 };
 
+#endif
+
 void FProjectedShadowInfo::RenderTranslucencyDepths(FRHICommandList& RHICmdList, FSceneRenderer* SceneRenderer)
 {
 	check(RHICmdList.IsInsideRenderPass());
@@ -638,6 +642,8 @@ void FProjectedShadowInfo::RenderTranslucencyDepths(FRHICommandList& RHICmdList,
 			CW_RGBA, BO_Add, BF_One, BF_One, BO_Add, BF_One, BF_One,
 			CW_RGBA, BO_Add, BF_One, BF_One, BO_Add, BF_One, BF_One>::GetRHI());
 
+//@todo MeshCommandPipeline
+#if 0
 		FTranslucencyShadowDepthDrawingPolicyFactory::ContextType DrawingContext(this,bDirectionalLight);
 
 		for (int32 MeshBatchIndex = 0; MeshBatchIndex < DynamicSubjectTranslucentMeshElements.Num(); MeshBatchIndex++)
@@ -675,6 +681,7 @@ void FProjectedShadowInfo::RenderTranslucencyDepths(FRHICommandList& RHICmdList,
 				}
 			}
 		}
+#endif
 	}
 }
 
