Removed legacy Drawing Policies and Static Mesh Draw Lists.  These are now replaced by FMeshDrawCommand / FMeshPassProcessor everywhere in the renderer.

diff --git a/Engine/Source/Runtime/Renderer/Private/TranslucentRendering.cpp b/Engine/Source/Runtime/Renderer/Private/TranslucentRendering.cpp
--- a/Engine/Source/Runtime/Renderer/Private/TranslucentRendering.cpp
+++ b/Engine/Source/Runtime/Renderer/Private/TranslucentRendering.cpp
@@ -263,32 +263,6 @@ private:
 
 IMPLEMENT_SHADER_TYPE(, FCopySceneColorPS, TEXT("/Engine/Private/TranslucentLightingShaders.usf"), TEXT("CopySceneColorMain"), SF_Pixel);
 
-// @todo MeshCommandPipeline - remove when UseMeshDrawCommandPipeline (and draw policies) will be removed.
-float CalculateTranslucentSortKey(const FPrimitiveSceneInfo* RESTRICT PrimitiveSceneInfo, const FSceneView& View)
-{
-	float SortKey = 0.0f;
-	if (View.TranslucentSortPolicy == ETranslucentSortPolicy::SortByDistance)
-	{
-		//sort based on distance to the view position, view rotation is not a factor
-		SortKey = (PrimitiveSceneInfo->Proxy->GetBounds().Origin - View.ViewMatrices.GetViewOrigin()).Size();
-		// UE4_TODO: also account for DPG in the sort key.
-	}
-	else if (View.TranslucentSortPolicy == ETranslucentSortPolicy::SortAlongAxis)
-	{
-		// Sort based on enforced orthogonal distance
-		const FVector CameraToObject = PrimitiveSceneInfo->Proxy->GetBounds().Origin - View.ViewMatrices.GetViewOrigin();
-		SortKey = FVector::DotProduct(CameraToObject, View.TranslucentSortAxis);
-	}
-	else
-	{
-		// Sort based on projected Z distance
-		check(View.TranslucentSortPolicy == ETranslucentSortPolicy::SortByProjectedZ);
-		SortKey = View.ViewMatrices.GetViewMatrix().TransformPosition(PrimitiveSceneInfo->Proxy->GetBounds().Origin).Z;
-	}
-
-	return SortKey;
-}
-
 FMeshDrawCommandSortKey CalculateStaticTranslucentMeshSortKey(const FPrimitiveSceneProxy* RESTRICT PrimitiveSceneProxy, uint16 MeshIdInPrimitive)
 {
 	uint16 SortKeyPriority = 0;
@@ -299,7 +273,7 @@ FMeshDrawCommandSortKey CalculateStaticTranslucentMeshSortKey(const FPrimitiveSc
 		SortKeyPriority = (uint16)((int32)PrimitiveSceneInfo->Proxy->GetTranslucencySortPriority() - (int32)SHRT_MIN);
 	}
 
-	UTranslucentMeshSortKey TranslucentMeshSortKey;
+	FTranslucentMeshSortKey TranslucentMeshSortKey;
 	TranslucentMeshSortKey.Fields.MeshIdInPrimitive = MeshIdInPrimitive;
 	TranslucentMeshSortKey.Fields.Priority = SortKeyPriority;
 	TranslucentMeshSortKey.Fields.Distance = 0; // View specific, so will be filled later inside VisibleMeshCommands.
@@ -310,46 +284,39 @@ FMeshDrawCommandSortKey CalculateStaticTranslucentMeshSortKey(const FPrimitiveSc
 	return SortKey;
 }
 
-void FTranslucentPrimSet::AppendScenePrimitives(FTranslucentSortedPrim* Elements, int32 Num, const FTranslucenyPrimCount& TranslucentPrimitiveCountPerPass)
+void FTranslucentPrimSet::AppendScenePrimitives(FTranslucentPrim* Elements, int32 Num, const FTranslucenyPrimCount& TranslucentPrimitiveCountPerPass)
 {
-	SortedPrims.Append(Elements, Num);
-	SortedPrimsNum.Append(TranslucentPrimitiveCountPerPass);
+	Prims.Append(Elements, Num);
+	PrimsNum.Append(TranslucentPrimitiveCountPerPass);
 }
 
 void FTranslucentPrimSet::PlaceScenePrimitive(FPrimitiveSceneInfo* PrimitiveSceneInfo, const FViewInfo& ViewInfo, const FPrimitiveViewRelevance& ViewRelevance,
-	FTranslucentPrimSet::FTranslucentSortedPrim *InArrayStart, int32& InOutArrayNum, FTranslucenyPrimCount& OutCount)
+	FTranslucentPrimSet::FTranslucentPrim *InArrayStart, int32& InOutArrayNum, FTranslucenyPrimCount& OutCount)
 {
-	const float SortKey = CalculateTranslucentSortKey(PrimitiveSceneInfo, ViewInfo);
 	const auto FeatureLevel = ViewInfo.GetFeatureLevel();
 
 	if (ViewInfo.Family->AllowTranslucencyAfterDOF())
 	{
 		if (ViewRelevance.bNormalTranslucencyRelevance)
 		{
-			new(&InArrayStart[InOutArrayNum++]) FTranslucentSortedPrim(PrimitiveSceneInfo, ETranslucencyPass::TPT_StandardTranslucency, PrimitiveSceneInfo->Proxy->GetTranslucencySortPriority(), SortKey);
+			new(&InArrayStart[InOutArrayNum++]) FTranslucentPrim(PrimitiveSceneInfo, ETranslucencyPass::TPT_StandardTranslucency);
 			OutCount.Add(ETranslucencyPass::TPT_StandardTranslucency, ViewRelevance.bUsesSceneColorCopy, ViewRelevance.bDisableOffscreenRendering);
 		}
 
 		if (ViewRelevance.bSeparateTranslucencyRelevance)
 		{
-			new(&InArrayStart[InOutArrayNum++]) FTranslucentSortedPrim(PrimitiveSceneInfo, ETranslucencyPass::TPT_TranslucencyAfterDOF, PrimitiveSceneInfo->Proxy->GetTranslucencySortPriority(), SortKey);
+			new(&InArrayStart[InOutArrayNum++]) FTranslucentPrim(PrimitiveSceneInfo, ETranslucencyPass::TPT_TranslucencyAfterDOF);
 			OutCount.Add(ETranslucencyPass::TPT_TranslucencyAfterDOF, ViewRelevance.bUsesSceneColorCopy, ViewRelevance.bDisableOffscreenRendering);
 		}
 	}
 	else // Otherwise, everything is rendered in a single bucket. This is not related to whether DOF is currently enabled or not.
 	{
 		// When using all translucency, Standard and AfterDOF are sorted together instead of being rendered like 2 buckets.
-		new(&InArrayStart[InOutArrayNum++]) FTranslucentSortedPrim(PrimitiveSceneInfo, ETranslucencyPass::TPT_AllTranslucency, PrimitiveSceneInfo->Proxy->GetTranslucencySortPriority(), SortKey);
+		new(&InArrayStart[InOutArrayNum++]) FTranslucentPrim(PrimitiveSceneInfo, ETranslucencyPass::TPT_AllTranslucency);
 		OutCount.Add(ETranslucencyPass::TPT_AllTranslucency, ViewRelevance.bUsesSceneColorCopy, ViewRelevance.bDisableOffscreenRendering);
 	}
 }
 
-void FTranslucentPrimSet::SortPrimitives()
-{
-	// sort prims based on the specified criteria (usually depth)
-	SortedPrims.Sort(FCompareFTranslucentSortedPrim());
-}
-
 extern int32 GLightShaftRenderAfterDOF;
 
 bool FSceneRenderer::ShouldRenderTranslucency(ETranslucencyPass::Type TranslucencyPass) const
@@ -379,7 +346,7 @@ bool FSceneRenderer::ShouldRenderTranslucency(ETranslucencyPass::Type Translucen
 
 	for (const FViewInfo& View : Views)
 	{
-		if (View.TranslucentPrimSet.SortedPrimsNum.Num(TranslucencyPass) > 0)
+		if (View.TranslucentPrimSet.PrimsNum.Num(TranslucencyPass) > 0)
 		{
 			return true;
 		}
@@ -588,7 +555,7 @@ void FDeferredShadingSceneRenderer::ConditionalResolveSceneColorForTranslucentMa
 		bool bNeedsResolve = false;
 		for (int32 TranslucencyPass = 0; TranslucencyPass < ETranslucencyPass::TPT_MAX && !bNeedsResolve; ++TranslucencyPass)
 		{
-			bNeedsResolve |= View.TranslucentPrimSet.SortedPrimsNum.UseSceneColorCopy((ETranslucencyPass::Type)TranslucencyPass);
+			bNeedsResolve |= View.TranslucentPrimSet.PrimsNum.UseSceneColorCopy((ETranslucencyPass::Type)TranslucencyPass);
 		}
 
 		if (bNeedsResolve)
@@ -981,7 +948,7 @@ void FDeferredShadingSceneRenderer::RenderTranslucency(FRHICommandListImmediate&
 		FDrawingPolicyRenderState DrawRenderState(View, BasePassUniformBuffer);
 
 		// If downsampling we need to render in the separate buffer. Otherwise we also need to render offscreen to apply TPT_TranslucencyAfterDOF
-		if (RenderInSeparateTranslucency(SceneContext, TranslucencyPass, View.TranslucentPrimSet.SortedPrimsNum.DisableOffscreenRendering(TranslucencyPass)))
+		if (RenderInSeparateTranslucency(SceneContext, TranslucencyPass, View.TranslucentPrimSet.PrimsNum.DisableOffscreenRendering(TranslucencyPass)))
 		{
 			checkSlow(RHICmdList.IsOutsideRenderPass());
 
diff --git a/Engine/Source/Runtime/Renderer/Private/TranslucentRendering.h b/Engine/Source/Runtime/Renderer/Private/TranslucentRendering.h
--- a/Engine/Source/Runtime/Renderer/Private/TranslucentRendering.h
+++ b/Engine/Source/Runtime/Renderer/Private/TranslucentRendering.h
@@ -16,14 +16,13 @@
 
 bool UseNearestDepthNeighborUpsampleForSeparateTranslucency(const FSceneRenderTargets& SceneContext);
 
-extern float CalculateTranslucentSortKey(const FPrimitiveSceneInfo* RESTRICT PrimitiveSceneInfo, const FSceneView& View);
 extern FMeshDrawCommandSortKey CalculateStaticTranslucentMeshSortKey(const FPrimitiveSceneProxy* RESTRICT PrimitiveSceneProxy, uint16 MeshIdInPrimitive);
 EMeshPass::Type TranslucencyPassToMeshPass(ETranslucencyPass::Type TranslucencyPass);
 
 /**
 * Translucent mesh sort key format.
 */
-union UTranslucentMeshSortKey
+union FTranslucentMeshSortKey
 {
 	uint64 PackedData;
 
@@ -34,36 +33,3 @@ union UTranslucentMeshSortKey
 		uint64 Priority				: 16; // First order by priority.
 	} Fields;
 };
-
-/**
-* Translucent draw policy factory.
-* Creates the policies needed for rendering a mesh based on its material
-*/
-class FMobileTranslucencyDrawingPolicyFactory
-{
-public:
-	enum { bAllowSimpleElements = true };
-	struct ContextType 
-	{
-		ETranslucencyPass::Type TranslucencyPass;
-
-		ContextType(ETranslucencyPass::Type InTranslucencyPass)
-		: TranslucencyPass(InTranslucencyPass)
-		{}
-	};
-
-	/**
-	* Render a dynamic mesh using a translucent draw policy 
-	* @return true if the mesh rendered
-	*/
-	static bool DrawDynamicMesh(
-		FRHICommandList& RHICmdList, 
-		const FViewInfo& View,
-		ContextType DrawingContext,
-		const FMeshBatch& Mesh,
-		bool bPreFog,
-		const FDrawingPolicyRenderState& DrawRenderState,
-		const FPrimitiveSceneProxy* PrimitiveSceneProxy,
-		FHitProxyId HitProxyId
-		);
-};
\ No newline at end of file
diff --git a/Engine/Source/Runtime/Renderer/Private/VelocityRendering.cpp b/Engine/Source/Runtime/Renderer/Private/VelocityRendering.cpp
--- a/Engine/Source/Runtime/Renderer/Private/VelocityRendering.cpp
+++ b/Engine/Source/Runtime/Renderer/Private/VelocityRendering.cpp
@@ -13,7 +13,6 @@
 #include "MeshMaterialShader.h"
 #include "ShaderBaseClasses.h"
 #include "SceneRendering.h"
-#include "StaticMeshDrawList.h"
 #include "DeferredShadingRenderer.h"
 #include "ScenePrivate.h"
 #include "PostProcess/ScreenSpaceReflections.h"
@@ -67,34 +66,6 @@ class FVelocityVS : public FMeshMaterialShader
 
 public:
 
-	void SetParameters(FRHICommandList& RHICmdList, const FVertexFactory* VertexFactory, const FMaterialRenderProxy* MaterialRenderProxy, const FViewInfo& View, const FDrawingPolicyRenderState& DrawRenderState, const bool bIsInstancedStereo)
-	{
-		if (IsInstancedStereoParameter.IsBound())
-		{
-			SetShaderValue(RHICmdList, GetVertexShader(), IsInstancedStereoParameter, bIsInstancedStereo);
-		}
-
-		if (InstancedEyeIndexParameter.IsBound())
-		{
-			SetShaderValue(RHICmdList, GetVertexShader(), InstancedEyeIndexParameter, 0);
-		}
-
-		FMeshMaterialShader::SetParameters(RHICmdList, GetVertexShader(), MaterialRenderProxy, *MaterialRenderProxy->GetMaterial(View.GetFeatureLevel()), View, DrawRenderState.GetViewUniformBuffer(), DrawRenderState.GetPassUniformBuffer());
-	}
-
-	void SetMesh(FRHICommandList& RHICmdList, const FVertexFactory* VertexFactory, const FMeshBatch& Mesh, int32 BatchElementIndex, const FDrawingPolicyRenderState& DrawRenderState, const FViewInfo& View, const FPrimitiveSceneProxy* Proxy)
-	{
-		FMeshMaterialShader::SetMesh(RHICmdList, GetVertexShader(), VertexFactory, View, Proxy, Mesh.Elements[BatchElementIndex], DrawRenderState);
-	}
-
-	void SetInstancedEyeIndex(FRHICommandList& RHICmdList, const uint32 EyeIndex)
-	{
-		if (InstancedEyeIndexParameter.IsBound())
-		{
-			SetShaderValue(RHICmdList, GetVertexShader(), InstancedEyeIndexParameter, EyeIndex);
-		}
-	}
-
 	bool SupportsVelocity() const
 	{
 		return GPUSkinCachePreviousPositionBuffer.IsBound() ||
@@ -125,8 +96,6 @@ protected:
 		PrevTransform1.Bind(Initializer.ParameterMap, TEXT("PrevTransform1"));
 		PrevTransform2.Bind(Initializer.ParameterMap, TEXT("PrevTransform2"));
 		PrevTransformBuffer.Bind(Initializer.ParameterMap, TEXT("PrevTransformBuffer"));
-		InstancedEyeIndexParameter.Bind(Initializer.ParameterMap, TEXT("InstancedEyeIndex"));
-		IsInstancedStereoParameter.Bind(Initializer.ParameterMap, TEXT("bIsInstancedStereo"));
 		PassUniformBuffer.Bind(Initializer.ParameterMap, FSceneTexturesUniformParameters::StaticStructMetadata.GetShaderVariableName());
 	}
 
@@ -141,8 +110,6 @@ protected:
 		Ar << PrevTransform1;
 		Ar << PrevTransform2;
 		Ar << PrevTransformBuffer;
-		Ar << InstancedEyeIndexParameter;
-		Ar << IsInstancedStereoParameter;
 
 		return bShaderHasOutdatedParameters;
 	}
@@ -153,8 +120,6 @@ private:
 	FShaderParameter PrevTransform1;
 	FShaderParameter PrevTransform2;
 	FShaderResourceParameter PrevTransformBuffer;
-	FShaderParameter InstancedEyeIndexParameter;
-	FShaderParameter IsInstancedStereoParameter;
 };
 
 
@@ -184,13 +149,6 @@ class FVelocityDS : public FBaseDS
 {
 	DECLARE_SHADER_TYPE(FVelocityDS,MeshMaterial);
 
-public:
-
-	void SetParameters(FRHICommandList& RHICmdList, const FMaterialRenderProxy* MaterialRenderProxy,const FViewInfo& View, const FDrawingPolicyRenderState& DrawRenderState)
-	{
-		FMeshMaterialShader::SetParameters(RHICmdList, (FDomainShaderRHIParamRef)GetDomainShader(), MaterialRenderProxy, *MaterialRenderProxy->GetMaterial(View.GetFeatureLevel()), View, DrawRenderState.GetViewUniformBuffer(), DrawRenderState.GetPassUniformBuffer());
-	}
-
 protected:
 	FVelocityDS(const ShaderMetaType::CompiledShaderInitializerType& Initializer):
 		FBaseDS(Initializer)
@@ -240,24 +198,6 @@ public:
 	{
 		PassUniformBuffer.Bind(Initializer.ParameterMap, FSceneTexturesUniformParameters::StaticStructMetadata.GetShaderVariableName());
 	}
-
-	void SetParameters(FRHICommandList& RHICmdList, const FVertexFactory* VertexFactory,const FMaterialRenderProxy* MaterialRenderProxy,const FViewInfo& View, const FDrawingPolicyRenderState& DrawRenderState)
-	{
-		FMeshMaterialShader::SetParameters(RHICmdList, GetPixelShader(), MaterialRenderProxy, *MaterialRenderProxy->GetMaterial(View.GetFeatureLevel()), View, DrawRenderState.GetViewUniformBuffer(), DrawRenderState.GetPassUniformBuffer());
-	}
-
-	void SetMesh(FRHICommandList& RHICmdList, const FVertexFactory* VertexFactory, const FMeshBatch& Mesh,int32 BatchElementIndex,const FDrawingPolicyRenderState& DrawRenderState,const FViewInfo& View, const FPrimitiveSceneProxy* Proxy)
-	{
-		const FPixelShaderRHIParamRef ShaderRHI = GetPixelShader();
-		
-		FMeshMaterialShader::SetMesh(RHICmdList, ShaderRHI, VertexFactory, View, Proxy, Mesh.Elements[BatchElementIndex], DrawRenderState);
-	}
-
-	virtual bool Serialize(FArchive& Ar) override
-	{
-		bool bShaderHasOutdatedParameters = FMeshMaterialShader::Serialize(Ar);
-		return bShaderHasOutdatedParameters;
-	}
 };
 
 IMPLEMENT_MATERIAL_SHADER_TYPE(,FVelocityPS,TEXT("/Engine/Private/VelocityShader.usf"),TEXT("MainPixelShader"),SF_Pixel);
diff --git a/Engine/Source/Runtime/Renderer/Private/VolumetricFogVoxelization.cpp b/Engine/Source/Runtime/Renderer/Private/VolumetricFogVoxelization.cpp
--- a/Engine/Source/Runtime/Renderer/Private/VolumetricFogVoxelization.cpp
+++ b/Engine/Source/Runtime/Renderer/Private/VolumetricFogVoxelization.cpp
@@ -179,25 +179,6 @@ class FVoxelizeVolumeVS : public FMeshMaterialShader
 
 public:
 	
-	void SetParameters(
-		FRHICommandList& RHICmdList, 
-		const FVertexFactory* VertexFactory,
-		const FMaterialRenderProxy* MaterialRenderProxy,
-		const FViewInfo& View, 
-		const FDrawingPolicyRenderState& DrawRenderState)
-	{
-		FMeshMaterialShader::SetParameters(RHICmdList, GetVertexShader(), MaterialRenderProxy, *MaterialRenderProxy->GetMaterial(View.GetFeatureLevel()), View, DrawRenderState.GetViewUniformBuffer(), DrawRenderState.GetPassUniformBuffer());
-	}
-
-	void SetMesh(FRHICommandList& RHICmdList, const FVertexFactory* VertexFactory,const FSceneView& View,const FPrimitiveSceneProxy* Proxy,const FMeshBatchElement& BatchElement,const FDrawingPolicyRenderState& DrawRenderState,int32 VoxelizationPassIndexValue)
-	{
-		FMeshMaterialShader::SetMesh(RHICmdList, GetVertexShader(),VertexFactory,View,Proxy,BatchElement,DrawRenderState);
-		if (!RHISupportsGeometryShaders(View.GetShaderPlatform()))
-		{
-			SetShaderValue(RHICmdList, GetVertexShader(), VoxelizationPassIndex, VoxelizationPassIndexValue);
-		}
-	}
-
 	virtual bool Serialize(FArchive& Ar)
 	{		
 		bool bShaderHasOutdatedParameters = FMeshMaterialShader::Serialize(Ar);
@@ -205,18 +186,16 @@ public:
 		return bShaderHasOutdatedParameters;
 	}
 
-	const FShaderParameter& GetVoxelizationPassIndexParameter() const { return VoxelizationPassIndex; }
-
 	void GetShaderBindings(
-						   const FScene* Scene,
-						   ERHIFeatureLevel::Type FeatureLevel,
-						   const FPrimitiveSceneProxy* PrimitiveSceneProxy,
-						   const FMaterialRenderProxy& MaterialRenderProxy,
-						   const FMaterial& Material,
-						   const TUniformBufferRef<FViewUniformShaderParameters>& ViewUniformBuffer,
-						   FUniformBufferRHIParamRef PassUniformBufferValue,
-						   const FVoxelizeVolumeShaderElementData& ShaderElementData,
-						   FMeshDrawSingleShaderBindings& ShaderBindings) const
+		const FScene* Scene,
+		ERHIFeatureLevel::Type FeatureLevel,
+		const FPrimitiveSceneProxy* PrimitiveSceneProxy,
+		const FMaterialRenderProxy& MaterialRenderProxy,
+		const FMaterial& Material,
+		const TUniformBufferRef<FViewUniformShaderParameters>& ViewUniformBuffer,
+		FUniformBufferRHIParamRef PassUniformBufferValue,
+		const FVoxelizeVolumeShaderElementData& ShaderElementData,
+		FMeshDrawSingleShaderBindings& ShaderBindings) const
 	{
 		FMeshMaterialShader::GetShaderBindings(Scene, FeatureLevel, PrimitiveSceneProxy, MaterialRenderProxy, Material, ViewUniformBuffer, PassUniformBufferValue, ShaderElementData, ShaderBindings);
 		if (!RHISupportsGeometryShaders(Scene->GetShaderPlatform()))
@@ -303,22 +282,6 @@ protected:
 
 public:
 	
-	void SetParameters(
-		FRHICommandList& RHICmdList, 
-		const FVertexFactory* VertexFactory,
-		const FMaterialRenderProxy* MaterialRenderProxy,
-		const FViewInfo& View, 
-		const FDrawingPolicyRenderState& DrawRenderState)
-	{
-		FMeshMaterialShader::SetParameters(RHICmdList, (FGeometryShaderRHIParamRef)GetGeometryShader(), MaterialRenderProxy, *MaterialRenderProxy->GetMaterial(View.GetFeatureLevel()), View, DrawRenderState.GetViewUniformBuffer(), DrawRenderState.GetPassUniformBuffer());
-	}
-
-	void SetMesh(FRHICommandList& RHICmdList, const FVertexFactory* VertexFactory,const FSceneView& View,const FPrimitiveSceneProxy* Proxy,const FMeshBatchElement& BatchElement,const FDrawingPolicyRenderState& DrawRenderState,int32 VoxelizationPassIndexValue)
-	{
-		FMeshMaterialShader::SetMesh(RHICmdList, (FGeometryShaderRHIParamRef)GetGeometryShader(),VertexFactory,View,Proxy,BatchElement,DrawRenderState);
-		SetShaderValue(RHICmdList, GetGeometryShader(), VoxelizationPassIndex, VoxelizationPassIndexValue);
-	}
-
 	virtual bool Serialize(FArchive& Ar)
 	{		
 		bool bShaderHasOutdatedParameters = FMeshMaterialShader::Serialize(Ar);
@@ -403,23 +366,6 @@ protected:
 			&& DoesPlatformSupportVolumetricFogVoxelization(Platform)
 			&& Material->GetMaterialDomain() == MD_Volume;
 	}
-
-public:
-	
-	void SetParameters(
-		FRHICommandList& RHICmdList, 
-		const FVertexFactory* VertexFactory,
-		const FMaterialRenderProxy* MaterialRenderProxy,
-		const FViewInfo& View, 
-		const FDrawingPolicyRenderState& DrawRenderState)
-	{
-		FMeshMaterialShader::SetParameters(RHICmdList, GetPixelShader(), MaterialRenderProxy, *MaterialRenderProxy->GetMaterial(View.GetFeatureLevel()), View, DrawRenderState.GetViewUniformBuffer(), DrawRenderState.GetPassUniformBuffer());
-	}
-
-	void SetMesh(FRHICommandList& RHICmdList, const FVertexFactory* VertexFactory,const FSceneView& View,const FPrimitiveSceneProxy* Proxy,const FMeshBatchElement& BatchElement,const FDrawingPolicyRenderState& DrawRenderState)
-	{
-		FMeshMaterialShader::SetMesh(RHICmdList, GetPixelShader(),VertexFactory,View,Proxy,BatchElement,DrawRenderState);
-	}
 };
 
 template<EVoxelizeShapeMode Mode>
diff --git a/Engine/Source/Runtime/Renderer/Public/MaterialShader.h b/Engine/Source/Runtime/Renderer/Public/MaterialShader.h
--- a/Engine/Source/Runtime/Renderer/Public/MaterialShader.h
+++ b/Engine/Source/Runtime/Renderer/Public/MaterialShader.h
@@ -171,21 +171,9 @@ public:
 	virtual bool Serialize(FArchive& Ar) override;
 	virtual uint32 GetAllocatedSize() const override;
 
-	void SetInstanceParameters(FRHICommandList& RHICmdList, uint32 InInstanceOffset, uint32 InInstanceCount) const
-	{
-		//@todo MeshCommandPipeline - remove when switching fully to the mesh draw command pipeline.  These are now handled by the vertex factories.
-		static const auto* CVar = IConsoleManager::Get().FindTConsoleVariableDataInt(TEXT("foliage.InstanceRuns"));
-		if (CVar->GetValueOnAnyThread() != 0)
-		{
-			bool const bZeroInstanceOffset = RHISupportsAbsoluteVertexID(GMaxRHIShaderPlatform);
-			SetShaderValue(RHICmdList, GetVertexShader(), InstanceOffset, bZeroInstanceOffset ? 0 : InInstanceOffset);
-		}
-	}
-
 protected:
 
 	FSceneTextureShaderParameters SceneTextureParameters;
-	FShaderParameter InstanceOffset;
 
 private:
 
@@ -214,6 +202,6 @@ private:
 	static FAutoConsoleVariableRef CVarAllowCachedUniformExpressions;
 
 #if !(UE_BUILD_TEST || UE_BUILD_SHIPPING || !WITH_EDITOR)
-	void VerifyExpressionAndShaderMaps(const FMaterialRenderProxy* MaterialRenderProxy, const FMaterial& Material, const FUniformExpressionCache* UniformExpressionCache);
+	void VerifyExpressionAndShaderMaps(const FMaterialRenderProxy* MaterialRenderProxy, const FMaterial& Material, const FUniformExpressionCache* UniformExpressionCache) const;
 #endif
 };
diff --git a/Engine/Source/Runtime/Renderer/Public/MeshMaterialShader.h b/Engine/Source/Runtime/Renderer/Public/MeshMaterialShader.h
--- a/Engine/Source/Runtime/Renderer/Public/MeshMaterialShader.h
+++ b/Engine/Source/Runtime/Renderer/Public/MeshMaterialShader.h
@@ -56,33 +56,6 @@ public:
 	}
 #endif
 
-	template< typename ShaderRHIParamRef >
-	void SetPassUniformBuffer(
-		FRHICommandList& RHICmdList,
-		const ShaderRHIParamRef ShaderRHI,
-		FUniformBufferRHIParamRef PassUniformBufferValue)
-	{
-		SetUniformBufferParameter(RHICmdList, ShaderRHI, PassUniformBuffer, PassUniformBufferValue);
-	}
-
-	template< typename ShaderRHIParamRef >
-	void SetParameters(
-		FRHICommandList& RHICmdList,
-		const ShaderRHIParamRef ShaderRHI,
-		const FMaterialRenderProxy* MaterialRenderProxy,
-		const FMaterial& Material,
-		const FSceneView& View,
-		const TUniformBufferRef<FViewUniformShaderParameters>& ViewUniformBuffer,
-		FUniformBufferRHIParamRef PassUniformBufferValue)
-	{
-		SetUniformBufferParameter(RHICmdList, ShaderRHI, PassUniformBuffer, PassUniformBufferValue);
-
-		checkfSlow(!(PassUniformBuffer.IsBound() && SceneTextureParameters.IsBound()) || SceneTextureParameters.IsSameUniformParameter(PassUniformBuffer), TEXT("If the pass uniform buffer is bound, it should contain SceneTexturesStruct: %s"), GetType()->GetName());
-
-		SetViewParameters(RHICmdList, ShaderRHI, View, ViewUniformBuffer);
-		FMaterialShader::SetParametersInner(RHICmdList, ShaderRHI, MaterialRenderProxy, Material, View);
-	}
-
 	void GetShaderBindings(
 		const FScene* Scene,
 		ERHIFeatureLevel::Type FeatureLevel,
@@ -94,18 +67,6 @@ public:
 		const FMeshMaterialShaderElementData& ShaderElementData,
 		FMeshDrawSingleShaderBindings& ShaderBindings) const;
 
-	template< typename ShaderRHIParamRef >
-	void SetMesh(
-		FRHICommandList& RHICmdList,
-		const ShaderRHIParamRef ShaderRHI,
-		const FVertexFactory* VertexFactory,
-		const FSceneView& View,
-		const FPrimitiveSceneProxy* Proxy,
-		const FMeshBatchElement& BatchElement,
-		const FDrawingPolicyRenderState& DrawRenderState,
-		uint32 DataFlags = 0
-	);
-
 	void GetElementShaderBindings(
 		const FScene* Scene, 
 		const FSceneView* ViewIfDynamicMeshCommand, 
@@ -119,12 +80,6 @@ public:
 		FMeshDrawSingleShaderBindings& ShaderBindings,
 		FVertexInputStreamArray& VertexStreams) const;
 
-	/**
-	 * Retrieves the fade uniform buffer parameter from a FSceneViewState for the primitive
-	 * This code was moved from SetMesh() to work around the template first-use vs first-seen differences between MSVC and others
-	 */
-	FUniformBufferRHIParamRef GetPrimitiveFadeUniformBufferParameter(const FSceneView& View, const FPrimitiveSceneProxy* Proxy);
-
 	// FShader interface.
 	virtual const FVertexFactoryParameterRef* GetVertexFactoryParameterRef() const override { return &VertexFactoryParameters; }
 	virtual bool Serialize(FArchive& Ar) override;
diff --git a/Engine/Source/Runtime/Renderer/Public/MeshPassProcessor.h b/Engine/Source/Runtime/Renderer/Public/MeshPassProcessor.h
--- a/Engine/Source/Runtime/Renderer/Public/MeshPassProcessor.h
+++ b/Engine/Source/Runtime/Renderer/Public/MeshPassProcessor.h
@@ -338,6 +338,7 @@ public:
 	uint32 NumInstances;
 	uint32 BaseVertexIndex;
 	uint32 NumVertices;
+	FVertexBufferRHIParamRef IndirectArgsBuffer;
 	uint8 RayTracedSegmentIndex;
 
 	int8 PrimitiveIdStreamIndex;
@@ -364,7 +365,8 @@ public:
 			&& NumPrimitives == Rhs.NumPrimitives
 			&& NumInstances == Rhs.NumInstances
 			&& BaseVertexIndex == Rhs.BaseVertexIndex
-			&& NumVertices == Rhs.NumVertices;
+			&& NumVertices == Rhs.NumVertices
+			&& IndirectArgsBuffer == Rhs.IndirectArgsBuffer;
 	}
 
 	/** Sets shaders on the mesh draw command and allocates room for the shader bindings. */
@@ -717,10 +719,127 @@ public:
 #endif // RHI_RAYTRACING
 };
 
-class FDefaultSubPolicyEx
+struct FDrawingPolicyRenderState
 {
+	FDrawingPolicyRenderState(const FSceneView& SceneView, FUniformBufferRHIParamRef InPassUniformBuffer = nullptr) : 
+		  BlendState(nullptr)
+		, DepthStencilState(nullptr)
+		, DepthStencilAccess(FExclusiveDepthStencil::DepthRead_StencilRead)
+		, ViewUniformBuffer(SceneView.ViewUniformBuffer)
+		, PassUniformBuffer(InPassUniformBuffer)
+		, StencilRef(0)
+	{
+	}
+
+	FDrawingPolicyRenderState(const TUniformBufferRef<FViewUniformShaderParameters>& InViewUniformBuffer, FUniformBufferRHIParamRef InPassUniformBuffer) : 
+		  BlendState(nullptr)
+		, DepthStencilState(nullptr)
+		, DepthStencilAccess(FExclusiveDepthStencil::DepthRead_StencilRead)
+		, ViewUniformBuffer(InViewUniformBuffer)
+		, PassUniformBuffer(InPassUniformBuffer)
+		, StencilRef(0)
+	{
+	}
+
+	FDrawingPolicyRenderState() :
+		BlendState(nullptr)
+		, DepthStencilState(nullptr)
+		, ViewUniformBuffer()
+		, PassUniformBuffer(nullptr)
+		, StencilRef(0)
+	{
+	}
+
+	FORCEINLINE_DEBUGGABLE FDrawingPolicyRenderState(const FDrawingPolicyRenderState& DrawRenderState) :
+		  BlendState(DrawRenderState.BlendState)
+		, DepthStencilState(DrawRenderState.DepthStencilState)
+		, DepthStencilAccess(DrawRenderState.DepthStencilAccess)
+		, ViewUniformBuffer(DrawRenderState.ViewUniformBuffer)
+		, PassUniformBuffer(DrawRenderState.PassUniformBuffer)
+		, StencilRef(DrawRenderState.StencilRef)
+	{
+	}
+
+	~FDrawingPolicyRenderState()
+	{
+	}
+
 public:
-	void GetPixelShaderBindings(const FPrimitiveSceneProxy* RESTRICT PrimitiveSceneProxy, const FMeshMaterialShader* ShaderParameters, FMeshDrawSingleShaderBindings& Bindings) const {}
+	FORCEINLINE_DEBUGGABLE void SetBlendState(FBlendStateRHIParamRef InBlendState)
+	{
+		BlendState = InBlendState;
+	}
+
+	FORCEINLINE_DEBUGGABLE const FBlendStateRHIParamRef GetBlendState() const
+	{
+		return BlendState;
+	}
+
+	FORCEINLINE_DEBUGGABLE void SetDepthStencilState(FDepthStencilStateRHIParamRef InDepthStencilState)
+	{
+		DepthStencilState = InDepthStencilState;
+		StencilRef = 0;
+	}
+
+	FORCEINLINE_DEBUGGABLE void SetStencilRef(uint32 InStencilRef)
+		{
+		StencilRef = InStencilRef;
+	}
+
+	FORCEINLINE_DEBUGGABLE const FDepthStencilStateRHIParamRef GetDepthStencilState() const
+	{
+		return DepthStencilState;
+	}
+
+	FORCEINLINE_DEBUGGABLE void SetDepthStencilAccess(FExclusiveDepthStencil::Type InDepthStencilAccess)
+	{
+		DepthStencilAccess = InDepthStencilAccess;
+	}
+
+	FORCEINLINE_DEBUGGABLE FExclusiveDepthStencil::Type GetDepthStencilAccess() const
+	{
+		return DepthStencilAccess;
+	}
+
+	FORCEINLINE_DEBUGGABLE void SetViewUniformBuffer(const TUniformBufferRef<FViewUniformShaderParameters>& InViewUniformBuffer)
+	{
+		ViewUniformBuffer = InViewUniformBuffer;
+	}
+
+	FORCEINLINE_DEBUGGABLE const TUniformBufferRef<FViewUniformShaderParameters>& GetViewUniformBuffer() const
+	{
+		return ViewUniformBuffer;
+	}
+
+	FORCEINLINE_DEBUGGABLE void SetPassUniformBuffer(FUniformBufferRHIParamRef InPassUniformBuffer)
+	{
+		PassUniformBuffer = InPassUniformBuffer;
+	}
+
+	FORCEINLINE_DEBUGGABLE FUniformBufferRHIParamRef GetPassUniformBuffer() const
+	{
+		return PassUniformBuffer;
+	}
+
+	FORCEINLINE_DEBUGGABLE uint32 GetStencilRef() const
+	{
+		return StencilRef;
+	}
+
+	FORCEINLINE_DEBUGGABLE void ApplyToPSO(FGraphicsPipelineStateInitializer& GraphicsPSOInit) const
+	{
+		GraphicsPSOInit.BlendState = BlendState;
+		GraphicsPSOInit.DepthStencilState = DepthStencilState;
+	}
+
+private:
+	FBlendStateRHIParamRef			BlendState;
+	FDepthStencilStateRHIParamRef	DepthStencilState;
+	FExclusiveDepthStencil::Type	DepthStencilAccess;
+
+	TUniformBufferRef<FViewUniformShaderParameters>	ViewUniformBuffer;
+	FUniformBufferRHIParamRef		PassUniformBuffer;
+	uint32							StencilRef;
 };
 
 /** 
