Renamed FStaticMesh -> FStaticMeshBatch to better describe what the class does (cached FMeshBatch) and reduce confusion with UStaticMesh

diff --git a/Engine/Source/Runtime/Engine/Private/SceneManagement.cpp b/Engine/Source/Runtime/Engine/Private/SceneManagement.cpp
--- a/Engine/Source/Runtime/Engine/Private/SceneManagement.cpp
+++ b/Engine/Source/Runtime/Engine/Private/SceneManagement.cpp
@@ -538,7 +538,7 @@ int8 ComputeStaticMeshLOD( const FStaticMeshRenderData* RenderData, const FVecto
 	return MinLOD;
 }
 
-FLODMask ComputeLODForMeshes(const TArray<class FStaticMeshRelevance>& StaticMeshRelevances, const FSceneView& View, const FVector4& Origin, float SphereRadius, int32 ForcedLODLevel, float& OutScreenRadiusSquared, float ScreenSizeScale, bool bDitheredLODTransition)
+FLODMask ComputeLODForMeshes(const TArray<class FStaticMeshBatchRelevance>& StaticMeshRelevances, const FSceneView& View, const FVector4& Origin, float SphereRadius, int32 ForcedLODLevel, float& OutScreenRadiusSquared, float ScreenSizeScale, bool bDitheredLODTransition)
 {
 	FLODMask LODToRender;
 	const FSceneView& LODView = GetLODView(View);
@@ -553,7 +553,7 @@ FLODMask ComputeLODForMeshes(const TArray<class FStaticMeshRelevance>& StaticMes
 		int32 MinLOD = 127, MaxLOD = 0;
 		for (int32 MeshIndex = 0; MeshIndex < StaticMeshRelevances.Num(); ++MeshIndex)
 		{
-			const FStaticMeshRelevance& Mesh = StaticMeshRelevances[MeshIndex];
+			const FStaticMeshBatchRelevance& Mesh = StaticMeshRelevances[MeshIndex];
 			if (Mesh.ScreenSize > 0.0f)
 			{
 				MinLOD = FMath::Min(MinLOD, (int32)Mesh.LODIndex);
@@ -574,7 +574,7 @@ FLODMask ComputeLODForMeshes(const TArray<class FStaticMeshRelevance>& StaticMes
 
 				for (int32 MeshIndex = NumMeshes - 1; MeshIndex >= 0; --MeshIndex)
 				{
-					const FStaticMeshRelevance& Mesh = StaticMeshRelevances[MeshIndex];
+					const FStaticMeshBatchRelevance& Mesh = StaticMeshRelevances[MeshIndex];
 					if (Mesh.ScreenSize > 0.0f)
 					{
 						float MeshScreenSize = Mesh.ScreenSize * ScreenSizeScale;
@@ -604,7 +604,7 @@ FLODMask ComputeLODForMeshes(const TArray<class FStaticMeshRelevance>& StaticMes
 
 			for (int32 MeshIndex = NumMeshes - 1; MeshIndex >= 0; --MeshIndex)
 			{
-				const FStaticMeshRelevance& Mesh = StaticMeshRelevances[MeshIndex];
+				const FStaticMeshBatchRelevance& Mesh = StaticMeshRelevances[MeshIndex];
 
 				float MeshScreenSize = Mesh.ScreenSize * ScreenSizeScale;
 
diff --git a/Engine/Source/Runtime/Engine/Public/SceneManagement.h b/Engine/Source/Runtime/Engine/Public/SceneManagement.h
--- a/Engine/Source/Runtime/Engine/Public/SceneManagement.h
+++ b/Engine/Source/Runtime/Engine/Public/SceneManagement.h
@@ -2664,7 +2664,7 @@ struct FLODMask
 		return DitheredLODIndices[0] != DitheredLODIndices[1];
 	}
 };
-FLODMask ENGINE_API ComputeLODForMeshes(const TArray<class FStaticMeshRelevance>& StaticMeshRelevances, const FSceneView& View, const FVector4& Origin, float SphereRadius, int32 ForcedLODLevel, float& OutScreenRadiusSquared, float ScreenSizeScale = 1.0f, bool bDitheredLODTransition = true);
+FLODMask ENGINE_API ComputeLODForMeshes(const TArray<class FStaticMeshBatchRelevance>& StaticMeshRelevances, const FSceneView& View, const FVector4& Origin, float SphereRadius, int32 ForcedLODLevel, float& OutScreenRadiusSquared, float ScreenSizeScale = 1.0f, bool bDitheredLODTransition = true);
 FLODMask ENGINE_API ComputeFastLODForMeshes(const TArray<float>& ScreenSizes, const FSceneView& View, const FVector4& Origin, float SphereRadius, int32 ForcedLODLevel, float& OutScreenRadiusSquared, float ScreenSizeScale = 1.0f, bool bDitheredLODTransition = true);
 
 class FSharedSamplerState : public FRenderResource
diff --git a/Engine/Source/Runtime/Landscape/Private/LandscapeRender.cpp b/Engine/Source/Runtime/Landscape/Private/LandscapeRender.cpp
--- a/Engine/Source/Runtime/Landscape/Private/LandscapeRender.cpp
+++ b/Engine/Source/Runtime/Landscape/Private/LandscapeRender.cpp
@@ -2265,8 +2265,8 @@ bool FLandscapeComponentSceneProxy::IsUsingCustomWholeSceneShadowLODRules() cons
 
 bool FLandscapeComponentSceneProxy::CanUseMeshBatchForShadowCascade(int8 InLODIndex, float InShadowMapTextureResolution, float InShadowMapCascadeSize) const
 {
-	const FStaticMesh* MeshBatch = nullptr;
-	const TArray<FStaticMesh>& PrimitiveStaticMeshes = GetPrimitiveSceneInfo()->StaticMeshes;
+	const FStaticMeshBatch* MeshBatch = nullptr;
+	const TArray<FStaticMeshBatch>& PrimitiveStaticMeshes = GetPrimitiveSceneInfo()->StaticMeshes;
 
 	check(PrimitiveStaticMeshes.IsValidIndex(InLODIndex));
 	check(PrimitiveStaticMeshes[InLODIndex].CastShadow);
diff --git a/Engine/Source/Runtime/Renderer/Private/CompositionLighting/PostProcessMeshDecals.cpp b/Engine/Source/Runtime/Renderer/Private/CompositionLighting/PostProcessMeshDecals.cpp
--- a/Engine/Source/Runtime/Renderer/Private/CompositionLighting/PostProcessMeshDecals.cpp
+++ b/Engine/Source/Runtime/Renderer/Private/CompositionLighting/PostProcessMeshDecals.cpp
@@ -403,7 +403,7 @@ void DrawDecalMeshCommands(FRenderingCompositePassContext& Context, EDecalRender
 					// Render static meshes from static scene prim
 					for (int32 StaticMeshIdx = 0, Count = PrimitiveSceneInfo->StaticMeshes.Num(); StaticMeshIdx < Count; StaticMeshIdx++)
 					{
-						const FStaticMesh& StaticMesh = PrimitiveSceneInfo->StaticMeshes[StaticMeshIdx];
+						const FStaticMeshBatch& StaticMesh = PrimitiveSceneInfo->StaticMeshes[StaticMeshIdx];
 
 						// Only render static mesh elements using decal materials
 						if (View.StaticMeshVisibilityMap[StaticMesh.Id] && StaticMesh.IsDecal(View.FeatureLevel))
diff --git a/Engine/Source/Runtime/Renderer/Private/DeferredShadingRenderer.cpp b/Engine/Source/Runtime/Renderer/Private/DeferredShadingRenderer.cpp
--- a/Engine/Source/Runtime/Renderer/Private/DeferredShadingRenderer.cpp
+++ b/Engine/Source/Runtime/Renderer/Private/DeferredShadingRenderer.cpp
@@ -693,7 +693,7 @@ bool FDeferredShadingSceneRenderer::GatherRayTracingWorldInstances(FRHICommandLi
 							else
 							{
 								const FPrimitiveSceneInfo* PrimitiveInfo = Scene->Primitives[PrimitiveIndex];
-								const FStaticMesh& StaticMesh = PrimitiveInfo->StaticMeshes[StaticMeshMdcIndices[StaticMeshIndex].StaticMeshIndex];
+								const FStaticMeshBatch& StaticMesh = PrimitiveInfo->StaticMeshes[StaticMeshMdcIndices[StaticMeshIndex].StaticMeshIndex];
 
 								uint64 BatchVisibility = 1;
 								if (StaticMesh.bRequiresPerElementVisibility)
diff --git a/Engine/Source/Runtime/Renderer/Private/DepthRendering.h b/Engine/Source/Runtime/Renderer/Private/DepthRendering.h
--- a/Engine/Source/Runtime/Renderer/Private/DepthRendering.h
+++ b/Engine/Source/Runtime/Renderer/Private/DepthRendering.h
@@ -14,7 +14,7 @@
 
 class FPrimitiveSceneProxy;
 class FScene;
-class FStaticMesh;
+class FStaticMeshBatch;
 class FViewInfo;
 
 enum EDepthDrawingMode
diff --git a/Engine/Source/Runtime/Renderer/Private/MeshDrawCommands.cpp b/Engine/Source/Runtime/Renderer/Private/MeshDrawCommands.cpp
--- a/Engine/Source/Runtime/Renderer/Private/MeshDrawCommands.cpp
+++ b/Engine/Source/Runtime/Renderer/Private/MeshDrawCommands.cpp
@@ -346,7 +346,7 @@ void GenerateDynamicMeshDrawCommands(
 	FMeshPassProcessor* PassMeshProcessor,
 	const TArray<FMeshBatchAndRelevance, SceneRenderingAllocator>& DynamicMeshElements,
 	int32 MaxNumDynamicMeshElements,
-	const TArray<const FStaticMesh*, SceneRenderingAllocator>& DynamicMeshCommandBuildRequests,
+	const TArray<const FStaticMeshBatch*, SceneRenderingAllocator>& DynamicMeshCommandBuildRequests,
 	int32 MaxNumBuildRequestElements,
 	FMeshCommandOneFrameArray& VisibleCommands,
 	FDynamicMeshDrawCommandStorage& MeshDrawCommandStorage
@@ -385,7 +385,7 @@ void GenerateDynamicMeshDrawCommands(
 
 		for (int32 MeshIndex = 0; MeshIndex < NumStaticMeshBatches; MeshIndex++)
 		{
-			const FStaticMesh* StaticMeshBatch = DynamicMeshCommandBuildRequests[MeshIndex];
+			const FStaticMeshBatch* StaticMeshBatch = DynamicMeshCommandBuildRequests[MeshIndex];
 			const uint64 BatchElementMask = StaticMeshBatch->bRequiresPerElementVisibility ? View.StaticMeshBatchVisibility[StaticMeshBatch->BatchVisibilityId] : ~0ull;
 
 			PassMeshProcessor->AddMeshBatch(*StaticMeshBatch, BatchElementMask, StaticMeshBatch->PrimitiveSceneInfo->Proxy, StaticMeshBatch->Id);
@@ -409,7 +409,7 @@ void GenerateMobileBasePassDynamicMeshDrawCommands(
 	FMeshPassProcessor* MobilePassCSMPassMeshProcessor,
 	const TArray<FMeshBatchAndRelevance, SceneRenderingAllocator>& DynamicMeshElements,
 	int32 MaxNumDynamicMeshElements,
-	const TArray<const FStaticMesh*, SceneRenderingAllocator>& DynamicMeshCommandBuildRequests,
+	const TArray<const FStaticMeshBatch*, SceneRenderingAllocator>& DynamicMeshCommandBuildRequests,
 	int32 MaxNumBuildRequestElements,
 	FMeshCommandOneFrameArray& VisibleCommands,
 	FDynamicMeshDrawCommandStorage& MeshDrawCommandStorage
@@ -460,7 +460,7 @@ void GenerateMobileBasePassDynamicMeshDrawCommands(
 
 		for (int32 MeshIndex = 0; MeshIndex < NumStaticMeshBatches; MeshIndex++)
 		{
-			const FStaticMesh* StaticMeshBatch = DynamicMeshCommandBuildRequests[MeshIndex];
+			const FStaticMeshBatch* StaticMeshBatch = DynamicMeshCommandBuildRequests[MeshIndex];
 			const uint64 BatchElementMask = StaticMeshBatch->bRequiresPerElementVisibility ? View.StaticMeshBatchVisibility[StaticMeshBatch->BatchVisibilityId] : ~0ull;
 
 			const int32 PrimitiveIndex = StaticMeshBatch->PrimitiveSceneInfo->Proxy->GetPrimitiveSceneInfo()->GetIndex();
@@ -746,7 +746,7 @@ void FParallelMeshDrawCommandPass::DispatchPassSetup(
 	FMeshPassProcessor* MeshPassProcessor,
 	const TArray<FMeshBatchAndRelevance, SceneRenderingAllocator>& DynamicMeshElements,
 	int32 NumDynamicMeshElements,
-	TArray<const FStaticMesh*, SceneRenderingAllocator>& InOutDynamicMeshCommandBuildRequests,
+	TArray<const FStaticMeshBatch*, SceneRenderingAllocator>& InOutDynamicMeshCommandBuildRequests,
 	int32 NumDynamicMeshCommandBuildRequestElements,
 	FMeshCommandOneFrameArray& InOutMeshDrawCommands,
 	FMeshPassProcessor* MobileBasePassCSMMeshPassProcessor,
diff --git a/Engine/Source/Runtime/Renderer/Private/MeshDrawCommands.h b/Engine/Source/Runtime/Renderer/Private/MeshDrawCommands.h
--- a/Engine/Source/Runtime/Renderer/Private/MeshDrawCommands.h
+++ b/Engine/Source/Runtime/Renderer/Private/MeshDrawCommands.h
@@ -10,7 +10,7 @@ MeshDrawCommands.h: Mesh draw commands.
 #include "TranslucencyPass.h"
 
 struct FMeshBatchAndRelevance;
-class FStaticMesh;
+class FStaticMeshBatch;
 class FParallelCommandListSet;
 
 /**	
@@ -59,8 +59,8 @@ public:
 	int32 NumDynamicMeshCommandBuildRequestElements;
 	FMeshCommandOneFrameArray MeshDrawCommands;
 	FMeshCommandOneFrameArray MobileBasePassCSMMeshDrawCommands;
-	TArray<const FStaticMesh*, SceneRenderingAllocator> DynamicMeshCommandBuildRequests;
-	TArray<const FStaticMesh*, SceneRenderingAllocator> MobileBasePassCSMDynamicMeshCommandBuildRequests;
+	TArray<const FStaticMeshBatch*, SceneRenderingAllocator> DynamicMeshCommandBuildRequests;
+	TArray<const FStaticMeshBatch*, SceneRenderingAllocator> MobileBasePassCSMDynamicMeshCommandBuildRequests;
 	FDynamicMeshDrawCommandStorage MeshDrawCommandStorage;
 
 	// Preallocated resources.
@@ -107,7 +107,7 @@ public:
 		FMeshPassProcessor* MeshPassProcessor,
 		const TArray<FMeshBatchAndRelevance, SceneRenderingAllocator>& DynamicMeshElements,
 		int32 NumDynamicMeshElements,
-		TArray<const FStaticMesh*, SceneRenderingAllocator>& InOutDynamicMeshCommandBuildRequests,
+		TArray<const FStaticMeshBatch*, SceneRenderingAllocator>& InOutDynamicMeshCommandBuildRequests,
 		int32 NumDynamicMeshCommandBuildRequestElements,
 		FMeshCommandOneFrameArray& InOutMeshDrawCommands,
 		FMeshPassProcessor* MobileBasePassCSMMeshPassProcessor = nullptr, // Required only for the mobile base pass.
diff --git a/Engine/Source/Runtime/Renderer/Private/MeshPassProcessor.cpp b/Engine/Source/Runtime/Renderer/Private/MeshPassProcessor.cpp
--- a/Engine/Source/Runtime/Renderer/Private/MeshPassProcessor.cpp
+++ b/Engine/Source/Runtime/Renderer/Private/MeshPassProcessor.cpp
@@ -988,7 +988,7 @@ FCachedPassMeshDrawListContext::FCachedPassMeshDrawListContext(FCachedMeshDrawCo
 
 FMeshDrawCommand& FCachedPassMeshDrawListContext::AddCommand(const FMeshDrawCommand& Initializer)
 {
-	// Only one FMeshDrawCommand supported per FStaticMesh in a pass
+	// Only one FMeshDrawCommand supported per FStaticMeshBatch in a pass
 	check(CommandInfo.CommandIndex == -1);
 	CommandInfo.CommandIndex = DrawList.MeshDrawCommands.Add(Initializer);
 	return DrawList.MeshDrawCommands[CommandInfo.CommandIndex];
diff --git a/Engine/Source/Runtime/Renderer/Private/PrimitiveSceneInfo.cpp b/Engine/Source/Runtime/Renderer/Private/PrimitiveSceneInfo.cpp
--- a/Engine/Source/Runtime/Renderer/Private/PrimitiveSceneInfo.cpp
+++ b/Engine/Source/Runtime/Renderer/Private/PrimitiveSceneInfo.cpp
@@ -56,7 +56,7 @@ public:
 			FPrimitiveSceneProxy* PrimitiveSceneProxy = PrimitiveSceneInfo->Proxy;
 			PrimitiveSceneProxy->VerifyUsedMaterial(Mesh.MaterialRenderProxy);
 
-			FStaticMesh* StaticMesh = new(PrimitiveSceneInfo->StaticMeshes) FStaticMesh(
+			FStaticMeshBatch* StaticMesh = new(PrimitiveSceneInfo->StaticMeshes) FStaticMeshBatch(
 				PrimitiveSceneInfo,
 				Mesh,
 				CurrentHitProxy ? CurrentHitProxy->Id : FHitProxyId()
@@ -66,7 +66,7 @@ public:
 
 			const bool bSupportsCachingMeshDrawCommands = SupportsCachingMeshDrawCommands(StaticMesh->VertexFactory, PrimitiveSceneProxy);
 
-			FStaticMeshRelevance* StaticMeshRelevance = new(PrimitiveSceneInfo->StaticMeshRelevances) FStaticMeshRelevance(
+			FStaticMeshBatchRelevance* StaticMeshRelevance = new(PrimitiveSceneInfo->StaticMeshRelevances) FStaticMeshBatchRelevance(
 				*StaticMesh, 
 				ScreenSize, 
 				bSupportsCachingMeshDrawCommands
@@ -167,7 +167,7 @@ void FPrimitiveSceneInfo::CacheMeshDrawCommands(FRHICommandListImmediate& RHICmd
 	int32 MeshWithCachedCommandsNum = 0;
 	for (int32 MeshIndex = 0; MeshIndex < StaticMeshes.Num(); MeshIndex++)
 	{
-		const FStaticMesh& Mesh = StaticMeshes[MeshIndex];
+		const FStaticMeshBatch& Mesh = StaticMeshes[MeshIndex];
 		if (SupportsCachingMeshDrawCommands(Mesh.VertexFactory, Proxy))
 		{
 			++MeshWithCachedCommandsNum;
@@ -190,8 +190,8 @@ void FPrimitiveSceneInfo::CacheMeshDrawCommands(FRHICommandListImmediate& RHICmd
 
 		for (int32 MeshIndex = 0; MeshIndex < StaticMeshes.Num(); MeshIndex++)
 		{
-			FStaticMeshRelevance& MeshRelevance = StaticMeshRelevances[MeshIndex];
-			FStaticMesh& Mesh = StaticMeshes[MeshIndex];
+			FStaticMeshBatchRelevance& MeshRelevance = StaticMeshRelevances[MeshIndex];
+			FStaticMeshBatch& Mesh = StaticMeshes[MeshIndex];
 
 			check(MeshRelevance.CommandInfosMask.IsEmpty());
 			MeshRelevance.CommandInfosBase = StaticMeshCommandInfos.Num();
@@ -268,7 +268,7 @@ void FPrimitiveSceneInfo::RemoveCachedMeshDrawCommands()
 
 	for (int32 MeshIndex = 0; MeshIndex < StaticMeshRelevances.Num(); ++MeshIndex)
 	{
-		FStaticMeshRelevance& MeshRelevance = StaticMeshRelevances[MeshIndex];
+		FStaticMeshBatchRelevance& MeshRelevance = StaticMeshRelevances[MeshIndex];
 
 		MeshRelevance.CommandInfosBase = 0;
 		MeshRelevance.CommandInfosMask.Reset();
@@ -290,8 +290,8 @@ void FPrimitiveSceneInfo::AddStaticMeshes(FRHICommandListImmediate& RHICmdList,
 
 	for(int32 MeshIndex = 0;MeshIndex < StaticMeshes.Num();MeshIndex++)
 	{
-		FStaticMeshRelevance& MeshRelevance = StaticMeshRelevances[MeshIndex];
-		FStaticMesh& Mesh = StaticMeshes[MeshIndex];
+		FStaticMeshBatchRelevance& MeshRelevance = StaticMeshRelevances[MeshIndex];
+		FStaticMeshBatch& Mesh = StaticMeshes[MeshIndex];
 
 		// Add the static mesh to the scene's static mesh list.
 		FSparseArrayAllocationInfo SceneArrayAllocation = Scene->StaticMeshes.AddUninitialized();
@@ -319,7 +319,7 @@ void FPrimitiveSceneInfo::AddStaticMeshes(FRHICommandListImmediate& RHICmdList,
 
 		for (int32 MeshIndex = 0; MeshIndex < StaticMeshes.Num(); MeshIndex++)
 		{
-			FStaticMesh& Mesh = StaticMeshes[MeshIndex];
+			FStaticMeshBatch& Mesh = StaticMeshes[MeshIndex];
 			MaxLOD = MaxLOD < Mesh.LODIndex ? Mesh.LODIndex : MaxLOD;
 		}
 
@@ -330,8 +330,8 @@ void FPrimitiveSceneInfo::AddStaticMeshes(FRHICommandListImmediate& RHICmdList,
 
 			for (int32 MeshIndex = 0; MeshIndex < StaticMeshes.Num(); MeshIndex++)
 			{
-				FStaticMeshRelevance& MeshRelevance = StaticMeshRelevances[MeshIndex];
-				FStaticMesh& Mesh = StaticMeshes[MeshIndex];
+				FStaticMeshBatchRelevance& MeshRelevance = StaticMeshRelevances[MeshIndex];
+				FStaticMeshBatch& Mesh = StaticMeshes[MeshIndex];
 
 				if (Proxy->ScreenSizes[MeshRelevance.LODIndex] != 0.0f)
 				{
@@ -887,7 +887,7 @@ void FPrimitiveSceneInfo::GetStaticMeshesLODRange(int8& OutMinLOD, int8& OutMaxL
 
 	for (int32 MeshIndex = 0; MeshIndex < StaticMeshRelevances.Num(); ++MeshIndex)
 	{
-		const FStaticMeshRelevance& MeshRelevance = StaticMeshRelevances[MeshIndex];
+		const FStaticMeshBatchRelevance& MeshRelevance = StaticMeshRelevances[MeshIndex];
 		OutMinLOD = FMath::Min(OutMinLOD, MeshRelevance.LODIndex);
 		OutMaxLOD = FMath::Max(OutMaxLOD, MeshRelevance.LODIndex);
 	}
diff --git a/Engine/Source/Runtime/Renderer/Private/SceneCore.cpp b/Engine/Source/Runtime/Renderer/Private/SceneCore.cpp
--- a/Engine/Source/Runtime/Renderer/Private/SceneCore.cpp
+++ b/Engine/Source/Runtime/Renderer/Private/SceneCore.cpp
@@ -315,7 +315,7 @@ void FLightPrimitiveInteraction::FlushCachedShadowMapData()
 	}
 }
 
-int32 FStaticMeshRelevance::GetStaticMeshCommandInfoIndex(EMeshPass::Type MeshPass) const
+int32 FStaticMeshBatchRelevance::GetStaticMeshCommandInfoIndex(EMeshPass::Type MeshPass) const
 {
 	int32 CommandInfoIndex = CommandInfosBase;
 
@@ -336,10 +336,10 @@ int32 FStaticMeshRelevance::GetStaticMeshCommandInfoIndex(EMeshPass::Type MeshPa
 }
 
 /*-----------------------------------------------------------------------------
-	FStaticMesh
+	FStaticMeshBatch
 -----------------------------------------------------------------------------*/
 
-FStaticMesh::~FStaticMesh()
+FStaticMeshBatch::~FStaticMeshBatch()
 {
 	FScene* Scene = PrimitiveSceneInfo->Scene;
 	// Remove this static mesh from the scene's list.
diff --git a/Engine/Source/Runtime/Renderer/Private/SceneCore.h b/Engine/Source/Runtime/Renderer/Private/SceneCore.h
--- a/Engine/Source/Runtime/Renderer/Private/SceneCore.h
+++ b/Engine/Source/Runtime/Renderer/Private/SceneCore.h
@@ -123,7 +123,7 @@ private:
  * A mesh which is defined by a primitive at scene segment construction time and never changed.
  * Lights are attached and detached as the segment containing the mesh is added or removed from a scene.
  */
-class FStaticMesh : public FMeshBatch
+class FStaticMeshBatch : public FMeshBatch
 {
 public:
 
@@ -137,7 +137,7 @@ public:
 	int32 BatchVisibilityId;
 
 	// Constructor/destructor.
-	FStaticMesh(
+	FStaticMeshBatch(
 		FPrimitiveSceneInfo* InPrimitiveSceneInfo,
 		const FMeshBatch& InMesh,
 		FHitProxyId InHitProxyId
@@ -150,11 +150,11 @@ public:
 		BatchHitProxyId = InHitProxyId;
 	}
 
-	~FStaticMesh();
+	~FStaticMeshBatch();
 
 private:
 	/** Private copy constructor. */
-	FStaticMesh(const FStaticMesh& InStaticMesh):
+	FStaticMeshBatch(const FStaticMeshBatch& InStaticMesh):
 		FMeshBatch(InStaticMesh),
 		PrimitiveSceneInfo(InStaticMesh.PrimitiveSceneInfo),
 		Id(InStaticMesh.Id),
@@ -163,12 +163,12 @@ private:
 };
 
 /**
- * FStaticMesh data which is InitViews specific. Stored separately for cache efficiency.
+ * FStaticMeshBatch data which is InitViews specific. Stored separately for cache efficiency.
  */
-class FStaticMeshRelevance
+class FStaticMeshBatchRelevance
 {
 public:
-	FStaticMeshRelevance(const FStaticMesh& StaticMesh, float InScreenSize, bool InbSupportsCachingMeshDrawCommands)
+	FStaticMeshBatchRelevance(const FStaticMeshBatch& StaticMesh, float InScreenSize, bool InbSupportsCachingMeshDrawCommands)
 		: Id(StaticMesh.Id)
 		, ScreenSize(InScreenSize)
 		, CommandInfosBase(0)
diff --git a/Engine/Source/Runtime/Renderer/Private/SceneHitProxyRendering.h b/Engine/Source/Runtime/Renderer/Private/SceneHitProxyRendering.h
--- a/Engine/Source/Runtime/Renderer/Private/SceneHitProxyRendering.h
+++ b/Engine/Source/Runtime/Renderer/Private/SceneHitProxyRendering.h
@@ -13,7 +13,7 @@
 
 class FPrimitiveSceneProxy;
 class FScene;
-class FStaticMesh;
+class FStaticMeshBatch;
 
 #if WITH_EDITOR
 
diff --git a/Engine/Source/Runtime/Renderer/Private/ScenePrivate.h b/Engine/Source/Runtime/Renderer/Private/ScenePrivate.h
--- a/Engine/Source/Runtime/Renderer/Private/ScenePrivate.h
+++ b/Engine/Source/Runtime/Renderer/Private/ScenePrivate.h
@@ -2552,9 +2552,9 @@ public:
 	FTextureLayout PreshadowCacheLayout;
 
 	/** The static meshes in the scene. */
-	TSparseArray<FStaticMesh*> StaticMeshes;
+	TSparseArray<FStaticMeshBatch*> StaticMeshes;
 
-	/** This sparse array is used just to track free indices for FStaticMesh::BatchVisibilityId. */
+	/** This sparse array is used just to track free indices for FStaticMeshBatch::BatchVisibilityId. */
 	TSparseArray<bool> StaticMeshBatchVisibility;
 
 	/** The exponential fog components in the scene. */
diff --git a/Engine/Source/Runtime/Renderer/Private/SceneRendering.h b/Engine/Source/Runtime/Renderer/Private/SceneRendering.h
--- a/Engine/Source/Runtime/Renderer/Private/SceneRendering.h
+++ b/Engine/Source/Runtime/Renderer/Private/SceneRendering.h
@@ -834,7 +834,7 @@ public:
 
 	TStaticArray<FMeshCommandOneFrameArray, EMeshPass::Num> MeshCommands;
 	TStaticArray<int32, EMeshPass::Num> NumDynamicMeshCommandBuildRequestElements;
-	TStaticArray<TArray<const FStaticMesh*, SceneRenderingAllocator>, EMeshPass::Num> DynamicMeshCommandBuildRequests;
+	TStaticArray<TArray<const FStaticMeshBatch*, SceneRenderingAllocator>, EMeshPass::Num> DynamicMeshCommandBuildRequests;
 };
 
 typedef TArray<FViewCommands, TInlineAllocator<4>> FViewVisibleCommandsPerView;
diff --git a/Engine/Source/Runtime/Renderer/Private/SceneVisibility.cpp b/Engine/Source/Runtime/Renderer/Private/SceneVisibility.cpp
--- a/Engine/Source/Runtime/Renderer/Private/SceneVisibility.cpp
+++ b/Engine/Source/Runtime/Renderer/Private/SceneVisibility.cpp
@@ -1665,7 +1665,7 @@ namespace EMarkMaskBits
 }
 
 typedef TChunkedArray<FVisibleMeshDrawCommand> FPassDrawCommandArray;
-typedef TChunkedArray<const FStaticMesh*> FPassDrawCommandBuildRequestArray;
+typedef TChunkedArray<const FStaticMeshBatch*> FPassDrawCommandBuildRequestArray;
 
 struct FDrawCommandRelevancePacket
 {
@@ -1687,8 +1687,8 @@ struct FDrawCommandRelevancePacket
 	void AddCommandsForMesh(
 		int32 PrimitiveIndex, 
 		const FPrimitiveSceneInfo* InPrimitiveSceneInfo,
-		const FStaticMeshRelevance& RESTRICT StaticMeshRelevance, 
-		const FStaticMesh& RESTRICT StaticMesh, 
+		const FStaticMeshBatchRelevance& RESTRICT StaticMeshRelevance, 
+		const FStaticMeshBatch& RESTRICT StaticMesh, 
 		const FScene* RESTRICT Scene, 
 		bool bCanCache, 
 		EMeshPass::Type PassType)
@@ -2133,8 +2133,8 @@ struct FRelevancePacket
 			const int32 NumStaticMeshes = PrimitiveSceneInfo->StaticMeshRelevances.Num();
 			for(int32 MeshIndex = 0;MeshIndex < NumStaticMeshes;MeshIndex++)
 			{
-				const FStaticMeshRelevance& StaticMeshRelevance = PrimitiveSceneInfo->StaticMeshRelevances[MeshIndex];
-				const FStaticMesh& StaticMesh = PrimitiveSceneInfo->StaticMeshes[MeshIndex];
+				const FStaticMeshBatchRelevance& StaticMeshRelevance = PrimitiveSceneInfo->StaticMeshRelevances[MeshIndex];
+				const FStaticMeshBatch& StaticMesh = PrimitiveSceneInfo->StaticMeshes[MeshIndex];
 
 				if (LODToRender.ContainsLOD(StaticMeshRelevance.LODIndex))
 				{
@@ -4102,7 +4102,7 @@ void FLODSceneTree::UpdateVisibilityStates(FViewInfo& View)
 			}
 
 			FHLODSceneNodeVisibilityState& NodeVisibility = VisibilityStates.FindOrAdd(SceneInfo->PrimitiveComponentId);
-			const TArray<FStaticMeshRelevance>& NodeMeshRelevances = SceneInfo->StaticMeshRelevances;
+			const TArray<FStaticMeshBatchRelevance>& NodeMeshRelevances = SceneInfo->StaticMeshRelevances;
 
 			// Ignore already updated nodes, or those that we can't work with
 			if (NodeVisibility.UpdateCount == UpdateCount || !NodeMeshRelevances.IsValidIndex(0))
diff --git a/Engine/Source/Runtime/Renderer/Private/ShadowRendering.h b/Engine/Source/Runtime/Renderer/Private/ShadowRendering.h
--- a/Engine/Source/Runtime/Renderer/Private/ShadowRendering.h
+++ b/Engine/Source/Runtime/Renderer/Private/ShadowRendering.h
@@ -524,7 +524,7 @@ private:
 	/** Dynamic mesh elements for translucent subject primitives. */
 	TArray<FMeshBatchAndRelevance,SceneRenderingAllocator> DynamicSubjectTranslucentMeshElements;
 
-	TArray<const FStaticMesh*, SceneRenderingAllocator> SubjectMeshCommandBuildRequests;
+	TArray<const FStaticMeshBatch*, SceneRenderingAllocator> SubjectMeshCommandBuildRequests;
 
 	/** Number of elements of DynamicSubjectMeshElements meshes. */
 	int32 NumDynamicSubjectMeshElements;
@@ -565,12 +565,12 @@ private:
 	void AddCachedMeshDrawCommandsForPass(
 		int32 PrimitiveIndex,
 		const FPrimitiveSceneInfo* InPrimitiveSceneInfo,
-		const FStaticMeshRelevance& RESTRICT StaticMeshRelevance,
-		const FStaticMesh& StaticMesh,
+		const FStaticMeshBatchRelevance& RESTRICT StaticMeshRelevance,
+		const FStaticMeshBatch& StaticMesh,
 		const FScene* Scene,
 		EMeshPass::Type PassType,
 		FMeshCommandOneFrameArray& VisibleMeshCommands,
-		TArray<const FStaticMesh*, SceneRenderingAllocator>& MeshCommandBuildRequests,
+		TArray<const FStaticMeshBatch*, SceneRenderingAllocator>& MeshCommandBuildRequests,
 		int32& NumMeshCommandBuildRequestElements);
 
 	/** Will return if we should draw the static mesh for the shadow, and will perform lazy init of primitive if it wasn't visible */
diff --git a/Engine/Source/Runtime/Renderer/Private/ShadowSetup.cpp b/Engine/Source/Runtime/Renderer/Private/ShadowSetup.cpp
--- a/Engine/Source/Runtime/Renderer/Private/ShadowSetup.cpp
+++ b/Engine/Source/Runtime/Renderer/Private/ShadowSetup.cpp
@@ -802,12 +802,12 @@ void FProjectedShadowInfo::SetupWholeSceneProjection(
 void FProjectedShadowInfo::AddCachedMeshDrawCommandsForPass(
 	int32 PrimitiveIndex,
 	const FPrimitiveSceneInfo* InPrimitiveSceneInfo,
-	const FStaticMeshRelevance& RESTRICT StaticMeshRelevance,
-	const FStaticMesh& StaticMesh,
+	const FStaticMeshBatchRelevance& RESTRICT StaticMeshRelevance,
+	const FStaticMeshBatch& StaticMesh,
 	const FScene* Scene,
 	EMeshPass::Type PassType,
 	FMeshCommandOneFrameArray& VisibleMeshCommands,
-	TArray<const FStaticMesh*, SceneRenderingAllocator>& MeshCommandBuildRequests,
+	TArray<const FStaticMeshBatch*, SceneRenderingAllocator>& MeshCommandBuildRequests,
 	int32& NumMeshCommandBuildRequestElements)
 {
 	const EShadingPath ShadingPath = Scene->GetShadingPath();
@@ -916,8 +916,8 @@ bool FProjectedShadowInfo::ShouldDrawStaticMeshes(FViewInfo& InCurrentView, bool
 
 			for (int32 MeshIndex = 0; MeshIndex < InPrimitiveSceneInfo->StaticMeshRelevances.Num(); MeshIndex++)
 			{
-				const FStaticMeshRelevance& StaticMeshRelevance = InPrimitiveSceneInfo->StaticMeshRelevances[MeshIndex];
-				const FStaticMesh& StaticMesh = InPrimitiveSceneInfo->StaticMeshes[MeshIndex];
+				const FStaticMeshBatchRelevance& StaticMeshRelevance = InPrimitiveSceneInfo->StaticMeshRelevances[MeshIndex];
+				const FStaticMeshBatch& StaticMesh = InPrimitiveSceneInfo->StaticMeshes[MeshIndex];
 
 				if ((StaticMeshRelevance.CastShadow || (bSelfShadowOnly && StaticMeshRelevance.bUseForDepthPass)) && ShadowLODToRender.ContainsLOD(StaticMeshRelevance.LODIndex))
 				{
@@ -948,8 +948,8 @@ bool FProjectedShadowInfo::ShouldDrawStaticMeshes(FViewInfo& InCurrentView, bool
 		{
 			for (int32 MeshIndex = 0; MeshIndex < InPrimitiveSceneInfo->StaticMeshRelevances.Num(); MeshIndex++)
 			{
-				const FStaticMeshRelevance& StaticMeshRelevance = InPrimitiveSceneInfo->StaticMeshRelevances[MeshIndex];
-				const FStaticMesh& StaticMesh = InPrimitiveSceneInfo->StaticMeshes[MeshIndex];
+				const FStaticMeshBatchRelevance& StaticMeshRelevance = InPrimitiveSceneInfo->StaticMeshRelevances[MeshIndex];
+				const FStaticMeshBatch& StaticMesh = InPrimitiveSceneInfo->StaticMeshes[MeshIndex];
 
 				if ((StaticMeshRelevance.CastShadow || (bSelfShadowOnly && StaticMeshRelevance.bUseForDepthPass)) && ShadowLODToRender.ContainsLOD(StaticMeshRelevance.LODIndex))
 				{
@@ -1266,7 +1266,7 @@ void FProjectedShadowInfo::SetupMeshDrawCommandsForProjectionStenciling(FSceneRe
 					{
 						for (int32 StaticMeshIdx = 0; StaticMeshIdx < ReceiverPrimitiveSceneInfo->StaticMeshes.Num(); StaticMeshIdx++)
 						{
-							const FStaticMesh& StaticMesh = ReceiverPrimitiveSceneInfo->StaticMeshes[StaticMeshIdx];
+							const FStaticMeshBatch& StaticMesh = ReceiverPrimitiveSceneInfo->StaticMeshes[StaticMeshIdx];
 
 							if (View.StaticMeshVisibilityMap[StaticMesh.Id])
 							{
@@ -1282,7 +1282,7 @@ void FProjectedShadowInfo::SetupMeshDrawCommandsForProjectionStenciling(FSceneRe
 			{
 				for (int32 MeshBatchIndex = 0; MeshBatchIndex < SubjectMeshCommandBuildRequests.Num(); ++MeshBatchIndex)
 				{
-					const FStaticMesh& StaticMesh = *SubjectMeshCommandBuildRequests[MeshBatchIndex];
+					const FStaticMeshBatch& StaticMesh = *SubjectMeshCommandBuildRequests[MeshBatchIndex];
 					const uint64 BatchElementMask = StaticMesh.bRequiresPerElementVisibility ? View.StaticMeshBatchVisibility[StaticMesh.BatchVisibilityId] : ~0ull;
 					DepthPassMeshProcessor.AddMeshBatch(StaticMesh, BatchElementMask, StaticMesh.PrimitiveSceneInfo->Proxy);
 				}
diff --git a/Engine/Source/Runtime/Renderer/Private/ShadowSetupMobile.cpp b/Engine/Source/Runtime/Renderer/Private/ShadowSetupMobile.cpp
--- a/Engine/Source/Runtime/Renderer/Private/ShadowSetupMobile.cpp
+++ b/Engine/Source/Runtime/Renderer/Private/ShadowSetupMobile.cpp
@@ -34,7 +34,7 @@ static TAutoConsoleVariable<int32> CVarsCsmShaderCullingMethod(
 	TEXT("Combine with 16 to change primitive bounding test to spheres instead of box. (i.e. 18 == combined casters + sphere test)")
 	,ECVF_RenderThreadSafe);
 
-static bool CouldStaticMeshEverReceiveCSMFromStationaryLight(ERHIFeatureLevel::Type FeatureLevel, const FPrimitiveSceneInfo* PrimitiveSceneInfo, const FStaticMesh& StaticMesh)
+static bool CouldStaticMeshEverReceiveCSMFromStationaryLight(ERHIFeatureLevel::Type FeatureLevel, const FPrimitiveSceneInfo* PrimitiveSceneInfo, const FStaticMeshBatch& StaticMesh)
 {
 	// test if static shadows are allowed in the first place:
 	static auto* CVarMobileAllowDistanceFieldShadows = IConsoleManager::Get().FindTConsoleVariableDataInt(TEXT("r.Mobile.AllowDistanceFieldShadows"));
@@ -59,7 +59,7 @@ static bool EnableStaticMeshCSMVisibilityState(bool bMovableLight, const FPrimit
 	INC_DWORD_STAT_BY(STAT_CSMStaticPrimitiveReceivers, 1);
 	for (int32 MeshIndex = 0; MeshIndex < PrimitiveSceneInfo->StaticMeshes.Num(); MeshIndex++)
 	{
-		const FStaticMesh& StaticMesh = PrimitiveSceneInfo->StaticMeshes[MeshIndex];
+		const FStaticMeshBatch& StaticMesh = PrimitiveSceneInfo->StaticMeshes[MeshIndex];
 
 		bool bHasCSMApplicableShadowInteraction = View.StaticMeshVisibilityMap[StaticMesh.Id] && StaticMesh.LCI;
 		bHasCSMApplicableShadowInteraction = bHasCSMApplicableShadowInteraction && StaticMesh.LCI->GetShadowMapInteraction().GetType() == SMIT_Texture;
diff --git a/Engine/Source/Runtime/Renderer/Private/TranslucentLighting.cpp b/Engine/Source/Runtime/Renderer/Private/TranslucentLighting.cpp
--- a/Engine/Source/Runtime/Renderer/Private/TranslucentLighting.cpp
+++ b/Engine/Source/Runtime/Renderer/Private/TranslucentLighting.cpp
@@ -563,7 +563,7 @@ public:
 		FRHICommandList& RHICmdList, 
 		const FViewInfo& View,
 		ContextType DrawingContext,
-		const FStaticMesh& StaticMesh,
+		const FStaticMeshBatch& StaticMesh,
 		bool bPreFog,
 		const FDrawingPolicyRenderState& DrawRenderState,
 		const FPrimitiveSceneProxy* PrimitiveSceneProxy,
diff --git a/Engine/Source/Runtime/Renderer/Private/VelocityRendering.h b/Engine/Source/Runtime/Renderer/Private/VelocityRendering.h
--- a/Engine/Source/Runtime/Renderer/Private/VelocityRendering.h
+++ b/Engine/Source/Runtime/Renderer/Private/VelocityRendering.h
@@ -15,7 +15,7 @@
 class FPrimitiveSceneInfo;
 class FPrimitiveSceneProxy;
 class FScene;
-class FStaticMesh;
+class FStaticMeshBatch;
 class FViewInfo;
 
 /** Get the cvar clamped state */
diff --git a/Engine/Source/Runtime/Renderer/Private/VolumetricFogVoxelization.cpp b/Engine/Source/Runtime/Renderer/Private/VolumetricFogVoxelization.cpp
--- a/Engine/Source/Runtime/Renderer/Private/VolumetricFogVoxelization.cpp
+++ b/Engine/Source/Runtime/Renderer/Private/VolumetricFogVoxelization.cpp
@@ -659,7 +659,7 @@ void FDeferredShadingSceneRenderer::VoxelizeFogVolumePrimitives(
 						{
 							for (int32 StaticMeshIdx = 0; StaticMeshIdx < PrimitiveSceneInfo->StaticMeshes.Num(); StaticMeshIdx++)
 							{
-								const FStaticMesh& StaticMesh = PrimitiveSceneInfo->StaticMeshes[StaticMeshIdx];
+								const FStaticMeshBatch& StaticMesh = PrimitiveSceneInfo->StaticMeshes[StaticMeshIdx];
 
 								if (View.StaticMeshVisibilityMap[StaticMesh.Id])
 								{
diff --git a/Engine/Source/Runtime/Renderer/Public/MeshPassProcessor.h b/Engine/Source/Runtime/Renderer/Public/MeshPassProcessor.h
--- a/Engine/Source/Runtime/Renderer/Public/MeshPassProcessor.h
+++ b/Engine/Source/Runtime/Renderer/Public/MeshPassProcessor.h
@@ -612,7 +612,7 @@ public:
 class FCachedPassMeshDrawList
 {
 public:
-	/** Indices held by FStaticMesh::CachedMeshDrawCommands must be stable */
+	/** Indices held by FStaticMeshBatch::CachedMeshDrawCommands must be stable */
 	TSparseArray<FMeshDrawCommand> MeshDrawCommands;
 };
 
diff --git a/Engine/Source/Runtime/Renderer/Public/PrimitiveSceneInfo.h b/Engine/Source/Runtime/Renderer/Public/PrimitiveSceneInfo.h
--- a/Engine/Source/Runtime/Renderer/Public/PrimitiveSceneInfo.h
+++ b/Engine/Source/Runtime/Renderer/Public/PrimitiveSceneInfo.h
@@ -215,10 +215,10 @@ public:
 	TArray<class FCachedMeshDrawCommandInfo> StaticMeshCommandInfos;
 
 	/** The primitive's static mesh relevances. Must be in sync with StaticMeshes. Kept separately from StaticMeshes for cache efficiency inside InitViews. */
-	TArray<class FStaticMeshRelevance> StaticMeshRelevances;
+	TArray<class FStaticMeshBatchRelevance> StaticMeshRelevances;
 
 	/** The primitive's static meshes. */
-	TArray<class FStaticMesh> StaticMeshes;
+	TArray<class FStaticMeshBatch> StaticMeshes;
 
 	/** The identifier for the primitive in Scene->PrimitiveOctree. */
 	FOctreeElementId OctreeId;
