Renamed FDrawingPolicyRenderState -> FMeshPassProcessorRenderState

diff --git a/Engine/Source/Developer/MaterialBaking/Private/MaterialRenderItem.cpp b/Engine/Source/Developer/MaterialBaking/Private/MaterialRenderItem.cpp
--- a/Engine/Source/Developer/MaterialBaking/Private/MaterialRenderItem.cpp
+++ b/Engine/Source/Developer/MaterialBaking/Private/MaterialRenderItem.cpp
@@ -19,7 +19,7 @@ FMeshMaterialRenderItem::FMeshMaterialRenderItem(const FMaterialData* InMaterial
 	LCI = new FMeshRenderInfo(InMeshSettings->LightMap, nullptr, nullptr, InMeshSettings->LightmapResourceCluster);
 }
 
-bool FMeshMaterialRenderItem::Render_RenderThread(FRHICommandListImmediate& RHICmdList, FDrawingPolicyRenderState& DrawRenderState, const FCanvas* Canvas)
+bool FMeshMaterialRenderItem::Render_RenderThread(FRHICommandListImmediate& RHICmdList, FMeshPassProcessorRenderState& DrawRenderState, const FCanvas* Canvas)
 {
 	return false;
 }
@@ -68,7 +68,7 @@ bool FMeshMaterialRenderItem::Render_GameThread(const FCanvas* Canvas, FRenderTh
 		RenderScope.EnqueueRenderCommand(
 			[Parameters](FRHICommandListImmediate& RHICmdList)
 		{
-			FDrawingPolicyRenderState DrawRenderState(*Parameters.View);
+			FMeshPassProcessorRenderState DrawRenderState(*Parameters.View);
 
 			// disable depth test & writes
 			DrawRenderState.SetBlendState(TStaticBlendState<CW_RGBA>::GetRHI());
@@ -101,7 +101,7 @@ void FMeshMaterialRenderItem::GenerateRenderData()
 	}
 }
 
-void FMeshMaterialRenderItem::QueueMaterial(FRHICommandListImmediate& RHICmdList, FDrawingPolicyRenderState& DrawRenderState, const FSceneView* View)
+void FMeshMaterialRenderItem::QueueMaterial(FRHICommandListImmediate& RHICmdList, FMeshPassProcessorRenderState& DrawRenderState, const FSceneView* View)
 {
 	FDynamicMeshBuilder DynamicMeshBuilder(View->GetFeatureLevel(), MAX_STATIC_TEXCOORDS, MeshSettings->LightMapIndex);
 	DynamicMeshBuilder.AddVertices(Vertices);
diff --git a/Engine/Source/Developer/MaterialBaking/Private/MaterialRenderItem.h b/Engine/Source/Developer/MaterialBaking/Private/MaterialRenderItem.h
--- a/Engine/Source/Developer/MaterialBaking/Private/MaterialRenderItem.h
+++ b/Engine/Source/Developer/MaterialBaking/Private/MaterialRenderItem.h
@@ -12,7 +12,7 @@ class FSceneView;
 class FRHICommandListImmediate;
 struct FMaterialData;
 struct FMeshData;
-struct FDrawingPolicyRenderState;
+struct FMeshPassProcessorRenderState;
 
 class FMeshMaterialRenderItem : public FCanvasBaseRenderItem
 {
@@ -20,7 +20,7 @@ public:
 	FMeshMaterialRenderItem(const FMaterialData* InMaterialSettings, const FMeshData* InMeshSettings, EMaterialProperty InMaterialProperty);
 
 	/** Begin FCanvasBaseRenderItem overrides */
-	virtual bool Render_RenderThread(FRHICommandListImmediate& RHICmdList, FDrawingPolicyRenderState& DrawRenderState, const FCanvas* Canvas) final;
+	virtual bool Render_RenderThread(FRHICommandListImmediate& RHICmdList, FMeshPassProcessorRenderState& DrawRenderState, const FCanvas* Canvas) final;
 	virtual bool Render_GameThread(const FCanvas* Canvas, FRenderThreadScope& RenderScope) final;
 	/** End FCanvasBaseRenderItem overrides */
 
@@ -28,7 +28,7 @@ public:
 	void GenerateRenderData();
 protected:
 	/** Enqueues the current material to be rendered */
-	void QueueMaterial(FRHICommandListImmediate& RHICmdList, FDrawingPolicyRenderState& DrawRenderState, const FSceneView* View);
+	void QueueMaterial(FRHICommandListImmediate& RHICmdList, FMeshPassProcessorRenderState& DrawRenderState, const FSceneView* View);
 	
 	/** Helper functions to populate render data using either mesh data or a simple quad */
 	void PopulateWithQuadData();
diff --git a/Engine/Source/Developer/MaterialUtilities/Private/MeshRendering.cpp b/Engine/Source/Developer/MaterialUtilities/Private/MeshRendering.cpp
--- a/Engine/Source/Developer/MaterialUtilities/Private/MeshRendering.cpp
+++ b/Engine/Source/Developer/MaterialUtilities/Private/MeshRendering.cpp
@@ -472,7 +472,7 @@ public:
 		return 2;
 	}
 
-	static void RenderMaterial(FRHICommandListImmediate& RHICmdList, FDrawingPolicyRenderState& DrawRenderState, const class FSceneView& View, FRenderData& Data)
+	static void RenderMaterial(FRHICommandListImmediate& RHICmdList, FMeshPassProcessorRenderState& DrawRenderState, const class FSceneView& View, FRenderData& Data)
 	{
 		// Check if material is TwoSided - single-sided materials should be rendered with normal and reverse
 		// triangle corner orders, to avoid problems with inside-out meshes or mesh parts. Note:
@@ -528,7 +528,7 @@ public:
 		GetRendererModule().DrawTileMesh(RHICmdList, DrawRenderState, View, MeshElement, false /*bIsHitTesting*/, FHitProxyId());
 	}
 
-	virtual bool Render_RenderThread(FRHICommandListImmediate& RHICmdList, FDrawingPolicyRenderState& DrawRenderState, const FCanvas* Canvas)
+	virtual bool Render_RenderThread(FRHICommandListImmediate& RHICmdList, FMeshPassProcessorRenderState& DrawRenderState, const FCanvas* Canvas)
 	{
 		checkSlow(Data);
 		// current render target set for the canvas
@@ -602,7 +602,7 @@ public:
 		RenderScope.EnqueueRenderCommand(
 			[Parameters](FRHICommandListImmediate& RHICmdList)
 		{
-			FDrawingPolicyRenderState DrawRenderState(*Parameters.View);
+			FMeshPassProcessorRenderState DrawRenderState(*Parameters.View);
 
 			// disable depth test & writes
 			DrawRenderState.SetDepthStencilState(TStaticDepthStencilState<false, CF_Always>::GetRHI());
diff --git a/Engine/Source/Editor/UnrealEd/Private/Lightmass/LightmassLandscapeRender.cpp b/Engine/Source/Editor/UnrealEd/Private/Lightmass/LightmassLandscapeRender.cpp
--- a/Engine/Source/Editor/UnrealEd/Private/Lightmass/LightmassLandscapeRender.cpp
+++ b/Engine/Source/Editor/UnrealEd/Private/Lightmass/LightmassLandscapeRender.cpp
@@ -134,7 +134,7 @@ void RenderLandscapeMaterialForLightmass(const FLandscapeStaticLightingMesh* Lan
 
 					FSceneView View(ViewInitOptions);
 
-					FDrawingPolicyRenderState DrawRenderState(View);
+					FMeshPassProcessorRenderState DrawRenderState(View);
 
 					// disable depth test & writes
 					DrawRenderState.SetDepthStencilState(TStaticDepthStencilState<false, CF_Always>::GetRHI());
diff --git a/Engine/Source/Runtime/Engine/Private/BatchedElements.cpp b/Engine/Source/Runtime/Engine/Private/BatchedElements.cpp
--- a/Engine/Source/Runtime/Engine/Private/BatchedElements.cpp
+++ b/Engine/Source/Runtime/Engine/Private/BatchedElements.cpp
@@ -822,7 +822,7 @@ FSceneView FBatchedElements::CreateProxySceneView(const FMatrix& ProjectionMatri
 	return FSceneView(ProxyViewInitOptions);
 }
 
-bool FBatchedElements::Draw(FRHICommandList& RHICmdList, const FDrawingPolicyRenderState& DrawRenderState, ERHIFeatureLevel::Type FeatureLevel, bool bNeedToSwitchVerticalAxis, const FMatrix& Transform, uint32 ViewportSizeX, uint32 ViewportSizeY, bool bHitTesting, float Gamma, const FSceneView* View, EBlendModeFilter::Type Filter) const
+bool FBatchedElements::Draw(FRHICommandList& RHICmdList, const FMeshPassProcessorRenderState& DrawRenderState, ERHIFeatureLevel::Type FeatureLevel, bool bNeedToSwitchVerticalAxis, const FMatrix& Transform, uint32 ViewportSizeX, uint32 ViewportSizeY, bool bHitTesting, float Gamma, const FSceneView* View, EBlendModeFilter::Type Filter) const
 {
 	if ( View )
 	{
@@ -841,7 +841,7 @@ bool FBatchedElements::Draw(FRHICommandList& RHICmdList, const FDrawingPolicyRen
 	}
 }
 
-bool FBatchedElements::Draw(FRHICommandList& RHICmdList, const FDrawingPolicyRenderState& DrawRenderState, ERHIFeatureLevel::Type FeatureLevel, bool bNeedToSwitchVerticalAxis, const FSceneView& View, bool bHitTesting, float Gamma /* = 1.0f */, EBlendModeFilter::Type Filter /* = EBlendModeFilter::All */) const
+bool FBatchedElements::Draw(FRHICommandList& RHICmdList, const FMeshPassProcessorRenderState& DrawRenderState, ERHIFeatureLevel::Type FeatureLevel, bool bNeedToSwitchVerticalAxis, const FSceneView& View, bool bHitTesting, float Gamma /* = 1.0f */, EBlendModeFilter::Type Filter /* = EBlendModeFilter::All */) const
 {
 	const FMatrix& Transform = View.ViewMatrices.GetViewProjectionMatrix();
 	const uint32 ViewportSizeX = View.UnscaledViewRect.Width();
diff --git a/Engine/Source/Runtime/Engine/Private/SceneManagement.cpp b/Engine/Source/Runtime/Engine/Private/SceneManagement.cpp
--- a/Engine/Source/Runtime/Engine/Private/SceneManagement.cpp
+++ b/Engine/Source/Runtime/Engine/Private/SceneManagement.cpp
@@ -168,7 +168,7 @@ void FSimpleElementCollector::RegisterDynamicResource(FDynamicPrimitiveResource*
 	DynamicResource->InitPrimitiveResource();
 }
 
-void FSimpleElementCollector::DrawBatchedElements(FRHICommandList& RHICmdList, const FDrawingPolicyRenderState& DrawRenderState, const FSceneView& InView, EBlendModeFilter::Type Filter, ESceneDepthPriorityGroup DepthPriorityGroup) const
+void FSimpleElementCollector::DrawBatchedElements(FRHICommandList& RHICmdList, const FMeshPassProcessorRenderState& DrawRenderState, const FSceneView& InView, EBlendModeFilter::Type Filter, ESceneDepthPriorityGroup DepthPriorityGroup) const
 {
 	// Mobile HDR does not execute post process, so does not need to render flipped
 	const bool bNeedToSwitchVerticalAxis = RHINeedsToSwitchVerticalAxis(InView.GetShaderPlatform()) && !bIsMobileHDR;
diff --git a/Engine/Source/Runtime/Engine/Private/TileRendering.cpp b/Engine/Source/Runtime/Engine/Private/TileRendering.cpp
--- a/Engine/Source/Runtime/Engine/Private/TileRendering.cpp
+++ b/Engine/Source/Runtime/Engine/Private/TileRendering.cpp
@@ -348,7 +348,7 @@ void FCanvasTileRendererItem::InitTileBuffers(FLocalVertexFactory* VertexFactory
 	});
 }
 
-bool FCanvasTileRendererItem::Render_RenderThread(FRHICommandListImmediate& RHICmdList, FDrawingPolicyRenderState& DrawRenderState, const FCanvas* Canvas)
+bool FCanvasTileRendererItem::Render_RenderThread(FRHICommandListImmediate& RHICmdList, FMeshPassProcessorRenderState& DrawRenderState, const FCanvas* Canvas)
 {
 	float CurrentRealTime = 0.f;
 	float CurrentWorldTime = 0.f;
@@ -479,7 +479,7 @@ bool FCanvasTileRendererItem::Render_GameThread(const FCanvas* Canvas, FRenderTh
 	{
 		SCOPED_DRAW_EVENTF(RHICmdList, CanvasDrawTile, *DrawTileParameters.RenderData->MaterialRenderProxy->GetMaterial(GMaxRHIFeatureLevel)->GetFriendlyName());
 
-		FDrawingPolicyRenderState DrawRenderState(*DrawTileParameters.View);
+		FMeshPassProcessorRenderState DrawRenderState(*DrawTileParameters.View);
 
 		// disable depth test & writes
 		DrawRenderState.SetDepthStencilState(TStaticDepthStencilState<false, CF_Always>::GetRHI());
diff --git a/Engine/Source/Runtime/Engine/Private/TriangleRendering.cpp b/Engine/Source/Runtime/Engine/Private/TriangleRendering.cpp
--- a/Engine/Source/Runtime/Engine/Private/TriangleRendering.cpp
+++ b/Engine/Source/Runtime/Engine/Private/TriangleRendering.cpp
@@ -276,7 +276,7 @@ void FCanvasTriangleRendererItem::InitTriangleBuffers(FLocalVertexFactory* Verte
 	});
 };
 
-bool FCanvasTriangleRendererItem::Render_RenderThread(FRHICommandListImmediate& RHICmdList, FDrawingPolicyRenderState& DrawRenderState, const FCanvas* Canvas)
+bool FCanvasTriangleRendererItem::Render_RenderThread(FRHICommandListImmediate& RHICmdList, FMeshPassProcessorRenderState& DrawRenderState, const FCanvas* Canvas)
 {
 	float CurrentRealTime = 0.f;
 	float CurrentWorldTime = 0.f;
@@ -407,7 +407,7 @@ bool FCanvasTriangleRendererItem::Render_GameThread(const FCanvas* Canvas, FRend
 	RenderScope.EnqueueRenderCommand(
 		[Parameters](FRHICommandListImmediate& RHICmdList)	
 	{
-		FDrawingPolicyRenderState DrawRenderState(*Parameters.View);
+		FMeshPassProcessorRenderState DrawRenderState(*Parameters.View);
 
 		// disable depth test & writes
 		DrawRenderState.SetDepthStencilState(TStaticDepthStencilState<false, CF_Always>::GetRHI());
diff --git a/Engine/Source/Runtime/Engine/Private/UserInterface/Canvas.cpp b/Engine/Source/Runtime/Engine/Private/UserInterface/Canvas.cpp
--- a/Engine/Source/Runtime/Engine/Private/UserInterface/Canvas.cpp
+++ b/Engine/Source/Runtime/Engine/Private/UserInterface/Canvas.cpp
@@ -375,7 +375,7 @@ FMatrix FCanvas::CalcProjectionMatrix(uint32 ViewSizeX, uint32 ViewSizeY, float
 	}
 }
 
-bool FCanvasBatchedElementRenderItem::Render_RenderThread(FRHICommandListImmediate& RHICmdList, FDrawingPolicyRenderState& DrawRenderState, const FCanvas* Canvas)
+bool FCanvasBatchedElementRenderItem::Render_RenderThread(FRHICommandListImmediate& RHICmdList, FMeshPassProcessorRenderState& DrawRenderState, const FCanvas* Canvas)
 {
 	checkSlow(Data);
 	bool bDirty = false;
@@ -462,7 +462,7 @@ bool FCanvasBatchedElementRenderItem::Render_GameThread(const FCanvas* Canvas, F
 		{
 			FSceneView SceneView = FBatchedElements::CreateProxySceneView(DrawParameters.RenderData->Transform.GetMatrix(), FIntRect(0, 0, DrawParameters.ViewportSizeX, DrawParameters.ViewportSizeY));
 
-			FDrawingPolicyRenderState DrawRenderState(SceneView);
+			FMeshPassProcessorRenderState DrawRenderState(SceneView);
 
 			// disable depth test & writes
 			DrawRenderState.SetDepthStencilState(TStaticDepthStencilState<false, CF_Always>::GetRHI());
@@ -690,7 +690,7 @@ void FCanvas::Flush_RenderThread(FRHICommandListImmediate& RHICmdList, bool bFor
 	
 	RHICmdList.BeginRenderPass(RPInfo, TEXT("CanvasRenderThread"));
 	{
-		FDrawingPolicyRenderState DrawRenderState;
+		FMeshPassProcessorRenderState DrawRenderState;
 		// disable depth test & writes
 		DrawRenderState.SetDepthStencilState(TStaticDepthStencilState<false, CF_Always>::GetRHI());
 
diff --git a/Engine/Source/Runtime/Engine/Public/BatchedElements.h b/Engine/Source/Runtime/Engine/Public/BatchedElements.h
--- a/Engine/Source/Runtime/Engine/Public/BatchedElements.h
+++ b/Engine/Source/Runtime/Engine/Public/BatchedElements.h
@@ -16,7 +16,7 @@
 #include "StaticBoundShaderState.h"
 
 struct FBatchedPoint;
-struct FDrawingPolicyRenderState;
+struct FMeshPassProcessorRenderState;
 
 namespace EBlendModeFilter
 {
@@ -204,7 +204,7 @@ public:
 	 * @param DepthTexture	DepthTexture for manual depth testing with editor compositing in the pixel shader
 	 */
 	UE_DEPRECATED(4.14, "Deprecated. Use the FBatchedElements::Draw method that takes a non-optional FSceneView parameter instead")
-	bool Draw(FRHICommandList& RHICmdList, const FDrawingPolicyRenderState& DrawRenderState, ERHIFeatureLevel::Type FeatureLevel, bool bNeedToSwitchVerticalAxis, const FMatrix& Transform, uint32 ViewportSizeX, uint32 ViewportSizeY, bool bHitTesting, float Gamma = 1.0f, const FSceneView* View = nullptr, EBlendModeFilter::Type Filter = EBlendModeFilter::All) const;
+	bool Draw(FRHICommandList& RHICmdList, const FMeshPassProcessorRenderState& DrawRenderState, ERHIFeatureLevel::Type FeatureLevel, bool bNeedToSwitchVerticalAxis, const FMatrix& Transform, uint32 ViewportSizeX, uint32 ViewportSizeY, bool bHitTesting, float Gamma = 1.0f, const FSceneView* View = nullptr, EBlendModeFilter::Type Filter = EBlendModeFilter::All) const;
 	
 	/**
 	 * Draws the batch
@@ -214,7 +214,7 @@ public:
 	 * @param Gamma			Optional gamma override
 	 * @param DepthTexture	DepthTexture for manual depth testing with editor compositing in the pixel shader
 	 */
-	bool Draw(FRHICommandList& RHICmdList, const FDrawingPolicyRenderState& DrawRenderState, ERHIFeatureLevel::Type FeatureLevel, bool bNeedToSwitchVerticalAxis, const FSceneView& View, bool bHitTesting, float Gamma = 1.0f, EBlendModeFilter::Type Filter = EBlendModeFilter::All) const;
+	bool Draw(FRHICommandList& RHICmdList, const FMeshPassProcessorRenderState& DrawRenderState, ERHIFeatureLevel::Type FeatureLevel, bool bNeedToSwitchVerticalAxis, const FSceneView& View, bool bHitTesting, float Gamma = 1.0f, EBlendModeFilter::Type Filter = EBlendModeFilter::All) const;
 
 	/**
 	 * Creates a proxy FSceneView for operations that are not tied directly to a scene but still require batched elements to be drawn.
diff --git a/Engine/Source/Runtime/Engine/Public/CanvasTypes.h b/Engine/Source/Runtime/Engine/Public/CanvasTypes.h
--- a/Engine/Source/Runtime/Engine/Public/CanvasTypes.h
+++ b/Engine/Source/Runtime/Engine/Public/CanvasTypes.h
@@ -787,7 +787,7 @@ public:
 	* @param RHICmdList - command list to use
 	* @return true if anything rendered
 	*/
-	virtual bool Render_RenderThread(FRHICommandListImmediate& RHICmdList, FDrawingPolicyRenderState& DrawRenderState, const FCanvas* Canvas) = 0;
+	virtual bool Render_RenderThread(FRHICommandListImmediate& RHICmdList, FMeshPassProcessorRenderState& DrawRenderState, const FCanvas* Canvas) = 0;
 	
 	/**
 	* Renders the canvas item
@@ -866,7 +866,7 @@ public:
 	* @param RHICmdList - command list to use
 	* @return true if anything rendered
 	*/
-	virtual bool Render_RenderThread(FRHICommandListImmediate& RHICmdList, FDrawingPolicyRenderState& DrawRenderState, const FCanvas* Canvas) override;
+	virtual bool Render_RenderThread(FRHICommandListImmediate& RHICmdList, FMeshPassProcessorRenderState& DrawRenderState, const FCanvas* Canvas) override;
 	
 	/**
 	* Renders the canvas item.
@@ -994,7 +994,7 @@ public:
 	* @param RHICmdList - command list to use
 	* @return true if anything rendered
 	*/
-	virtual bool Render_RenderThread(FRHICommandListImmediate& RHICmdList, FDrawingPolicyRenderState& DrawRenderState, const FCanvas* Canvas) override;
+	virtual bool Render_RenderThread(FRHICommandListImmediate& RHICmdList, FMeshPassProcessorRenderState& DrawRenderState, const FCanvas* Canvas) override;
 
 	/**
 	* Renders the canvas item.
@@ -1150,7 +1150,7 @@ public:
 	* @param RHICmdList - command list to use
 	* @return true if anything rendered
 	*/
-	virtual bool Render_RenderThread(FRHICommandListImmediate& RHICmdList, FDrawingPolicyRenderState& DrawRenderState, const FCanvas* Canvas) override;
+	virtual bool Render_RenderThread(FRHICommandListImmediate& RHICmdList, FMeshPassProcessorRenderState& DrawRenderState, const FCanvas* Canvas) override;
 
 	/**
 	* Renders the canvas item.
diff --git a/Engine/Source/Runtime/Engine/Public/DebugViewModeInterface.h b/Engine/Source/Runtime/Engine/Public/DebugViewModeInterface.h
--- a/Engine/Source/Runtime/Engine/Public/DebugViewModeInterface.h
+++ b/Engine/Source/Runtime/Engine/Public/DebugViewModeInterface.h
@@ -17,7 +17,7 @@ DebugViewModeInterface.h: Contains definitions for rendering debug viewmodes.
 class FDebugViewModePS;
 class FMaterial;
 class FVertexFactoryType;
-struct FDrawingPolicyRenderState;
+struct FMeshPassProcessorRenderState;
 
 class ENGINE_API FDebugViewModeInterface
 {
diff --git a/Engine/Source/Runtime/Engine/Public/SceneManagement.h b/Engine/Source/Runtime/Engine/Public/SceneManagement.h
--- a/Engine/Source/Runtime/Engine/Public/SceneManagement.h
+++ b/Engine/Source/Runtime/Engine/Public/SceneManagement.h
@@ -1851,7 +1851,7 @@ public:
 		return 0;
 	}
 
-	void DrawBatchedElements(FRHICommandList& RHICmdList, const FDrawingPolicyRenderState& DrawRenderState, const FSceneView& InView, EBlendModeFilter::Type Filter, ESceneDepthPriorityGroup DPG) const;
+	void DrawBatchedElements(FRHICommandList& RHICmdList, const FMeshPassProcessorRenderState& DrawRenderState, const FSceneView& InView, EBlendModeFilter::Type Filter, ESceneDepthPriorityGroup DPG) const;
 
 	bool HasPrimitives(ESceneDepthPriorityGroup DPG) const
 	{
diff --git a/Engine/Source/Runtime/Landscape/Private/LandscapeGrass.cpp b/Engine/Source/Runtime/Landscape/Private/LandscapeGrass.cpp
--- a/Engine/Source/Runtime/Landscape/Private/LandscapeGrass.cpp
+++ b/Engine/Source/Runtime/Landscape/Private/LandscapeGrass.cpp
@@ -343,7 +343,7 @@ private:
 		FVector2D ViewOffset,
 		float PassOffsetX);
 
-	FDrawingPolicyRenderState PassDrawRenderState;
+	FMeshPassProcessorRenderState PassDrawRenderState;
 };
 
 FLandscapeGrassWeightMeshProcessor::FLandscapeGrassWeightMeshProcessor(const FScene* Scene, const FSceneView* InViewIfDynamicMeshCommand, FMeshPassDrawListContext* InDrawListContext)
diff --git a/Engine/Source/Runtime/RenderCore/Public/RendererInterface.h b/Engine/Source/Runtime/RenderCore/Public/RendererInterface.h
--- a/Engine/Source/Runtime/RenderCore/Public/RendererInterface.h
+++ b/Engine/Source/Runtime/RenderCore/Public/RendererInterface.h
@@ -688,7 +688,7 @@ public:
 	virtual void InitializeSystemTextures(FRHICommandListImmediate& RHICmdList) = 0;
 
 	/** Draws a tile mesh element with the specified view. */
-	virtual void DrawTileMesh(FRHICommandListImmediate& RHICmdList, struct FDrawingPolicyRenderState& DrawRenderState, const FSceneView& View, FMeshBatch& Mesh, bool bIsHitTesting, const class FHitProxyId& HitProxyId) = 0;
+	virtual void DrawTileMesh(FRHICommandListImmediate& RHICmdList, struct FMeshPassProcessorRenderState& DrawRenderState, const FSceneView& View, FMeshBatch& Mesh, bool bIsHitTesting, const class FHitProxyId& HitProxyId) = 0;
 
 	/** Render thread side, use TRefCountPtr<IPooledRenderTarget>, allows to use sharing and VisualizeTexture */
 	// TODO(RDG): Kill that guy.
diff --git a/Engine/Source/Runtime/Renderer/Private/BasePassRendering.cpp b/Engine/Source/Runtime/Renderer/Private/BasePassRendering.cpp
--- a/Engine/Source/Runtime/Renderer/Private/BasePassRendering.cpp
+++ b/Engine/Source/Runtime/Renderer/Private/BasePassRendering.cpp
@@ -130,7 +130,7 @@ IMPLEMENT_BASEPASS_LIGHTMAPPED_SHADER_TYPE( TUniformLightMapPolicy<LMP_DISTANCE_
 
 DECLARE_GPU_STAT(Basepass);
 
-void SetBasePassDitheredLODTransitionState(const FSceneView* SceneView, const FMeshBatch& RESTRICT Mesh, int32 StaticMeshId, FDrawingPolicyRenderState& DrawRenderState)
+void SetBasePassDitheredLODTransitionState(const FSceneView* SceneView, const FMeshBatch& RESTRICT Mesh, int32 StaticMeshId, FMeshPassProcessorRenderState& DrawRenderState)
 {
 	if (SceneView && StaticMeshId >= 0 && Mesh.bDitheredLODTransition)
 	{
@@ -163,7 +163,7 @@ void SetBasePassDitheredLODTransitionState(const FSceneView* SceneView, const FM
 	}
 }
 
-void SetTranslucentRenderState(FDrawingPolicyRenderState& DrawRenderState, const FMaterial& Material)
+void SetTranslucentRenderState(FMeshPassProcessorRenderState& DrawRenderState, const FMaterial& Material)
 {
 	switch (Material.GetBlendMode())
 	{
@@ -243,7 +243,7 @@ FMeshDrawCommandSortKey CalculateStaticBaseMeshSortKey(EDepthDrawingMode EarlyZP
 	return SortKey;
 }
 
-void SetDepthStencilStateForBasePass(FDrawingPolicyRenderState& DrawRenderState, ERHIFeatureLevel::Type FeatureLevel, const FMeshBatch& Mesh, const FPrimitiveSceneProxy* PrimitiveSceneProxy, bool bEnableReceiveDecalOutput, bool bUseDebugViewPS, FDepthStencilStateRHIParamRef LodFadeOverrideDepthStencilState)
+void SetDepthStencilStateForBasePass(FMeshPassProcessorRenderState& DrawRenderState, ERHIFeatureLevel::Type FeatureLevel, const FMeshBatch& Mesh, const FPrimitiveSceneProxy* PrimitiveSceneProxy, bool bEnableReceiveDecalOutput, bool bUseDebugViewPS, FDepthStencilStateRHIParamRef LodFadeOverrideDepthStencilState)
 {
 	static IConsoleVariable* EarlyZPassOnlyMaterialMaskingCVar = IConsoleManager::Get().FindConsoleVariable(TEXT("r.EarlyZPassOnlyMaterialMasking"));
 	bool bMaskInEarlyPass = (EarlyZPassOnlyMaterialMaskingCVar && Mesh.MaterialRenderProxy->GetMaterial(FeatureLevel)->IsMasked() && EarlyZPassOnlyMaterialMaskingCVar->GetInt());
@@ -298,7 +298,7 @@ void SetDepthStencilStateForBasePass(FDrawingPolicyRenderState& DrawRenderState,
 	}
 }
 
-void SetupBasePassState(FExclusiveDepthStencil::Type BasePassDepthStencilAccess, const bool bShaderComplexity, FDrawingPolicyRenderState& DrawRenderState)
+void SetupBasePassState(FExclusiveDepthStencil::Type BasePassDepthStencilAccess, const bool bShaderComplexity, FMeshPassProcessorRenderState& DrawRenderState)
 {
 	DrawRenderState.SetDepthStencilAccess(BasePassDepthStencilAccess);
 
@@ -637,7 +637,7 @@ bool FDeferredShadingSceneRenderer::RenderBasePass(FRHICommandListImmediate& RHI
 				TUniformBufferRef<FOpaqueBasePassUniformParameters> BasePassUniformBuffer;
 				CreateOpaqueBasePassUniformBuffer(RHICmdList, View, ForwardScreenSpaceShadowMask, BasePassUniformBuffer);
 
-				FDrawingPolicyRenderState DrawRenderState(View, BasePassUniformBuffer);
+				FMeshPassProcessorRenderState DrawRenderState(View, BasePassUniformBuffer);
 
 				SetupBasePassState(BasePassDepthStencilAccess, ViewFamily.EngineShowFlags.ShaderComplexity, DrawRenderState);
 
@@ -671,7 +671,7 @@ bool FDeferredShadingSceneRenderer::RenderBasePass(FRHICommandListImmediate& RHI
 				TUniformBufferRef<FOpaqueBasePassUniformParameters> BasePassUniformBuffer;
 				CreateOpaqueBasePassUniformBuffer(RHICmdList, View, ForwardScreenSpaceShadowMask, BasePassUniformBuffer);
 
-				FDrawingPolicyRenderState DrawRenderState(View, BasePassUniformBuffer);
+				FMeshPassProcessorRenderState DrawRenderState(View, BasePassUniformBuffer);
 
 				SetupBasePassState(BasePassDepthStencilAccess, ViewFamily.EngineShowFlags.ShaderComplexity, DrawRenderState);
 
@@ -734,7 +734,7 @@ public:
 		bool bInCreateSceneContext,
 		const FSceneRenderer* InSceneRenderer,
 		FExclusiveDepthStencil::Type InBasePassDepthStencilAccess,
-		const FDrawingPolicyRenderState& InDrawRenderState)
+		const FMeshPassProcessorRenderState& InDrawRenderState)
 		: FParallelCommandListSet(GET_STATID(STAT_CLP_Basepass), InView, InSceneRenderer, InParentCmdList, bInParallelExecute, bInCreateSceneContext, InDrawRenderState)
 		, BasePassDepthStencilAccess(InBasePassDepthStencilAccess)
 	{
@@ -753,7 +753,7 @@ public:
 	}
 };
 
-void FDeferredShadingSceneRenderer::RenderBasePassViewParallel(FViewInfo& View, FRHICommandListImmediate& ParentCmdList, FExclusiveDepthStencil::Type BasePassDepthStencilAccess, const FDrawingPolicyRenderState& InDrawRenderState)
+void FDeferredShadingSceneRenderer::RenderBasePassViewParallel(FViewInfo& View, FRHICommandListImmediate& ParentCmdList, FExclusiveDepthStencil::Type BasePassDepthStencilAccess, const FMeshPassProcessorRenderState& InDrawRenderState)
 {
 	check(ParentCmdList.IsOutsideRenderPass());
 
@@ -785,9 +785,9 @@ bool HasEditorPrimitivesForDPG(const FViewInfo& View, ESceneDepthPriorityGroup D
 	return bHasPrimitives;
 }
 
-void FDeferredShadingSceneRenderer::RenderEditorPrimitives(FRHICommandListImmediate& RHICmdList, const FViewInfo& View, FExclusiveDepthStencil::Type BasePassDepthStencilAccess, const FDrawingPolicyRenderState& InDrawRenderState, bool& bOutDirty) 
+void FDeferredShadingSceneRenderer::RenderEditorPrimitives(FRHICommandListImmediate& RHICmdList, const FViewInfo& View, FExclusiveDepthStencil::Type BasePassDepthStencilAccess, const FMeshPassProcessorRenderState& InDrawRenderState, bool& bOutDirty) 
 {
-	FDrawingPolicyRenderState DrawRenderState(InDrawRenderState);
+	FMeshPassProcessorRenderState DrawRenderState(InDrawRenderState);
 	SetupBasePassState(BasePassDepthStencilAccess, ViewFamily.EngineShowFlags.ShaderComplexity, DrawRenderState);
 	SetupBasePassView(RHICmdList, View, this, true);
 
@@ -806,7 +806,7 @@ void FDeferredShadingSceneRenderer::RenderEditorPrimitives(FRHICommandListImmedi
 			// need to setup view again after reconfiguring render targets
 			SetupBasePassView(RHICmdList, View, this, true);
 
-			FDrawingPolicyRenderState NoDepthTestDrawRenderState(DrawRenderState);
+			FMeshPassProcessorRenderState NoDepthTestDrawRenderState(DrawRenderState);
 			NoDepthTestDrawRenderState.SetDepthStencilState(TStaticDepthStencilState<true, CF_Always>::GetRHI());
 			NoDepthTestDrawRenderState.SetDepthStencilAccess(FExclusiveDepthStencil::DepthWrite_StencilWrite);
 			RenderEditorPrimitivesForDPG(RHICmdList, View, BasePassDepthStencilAccess, NoDepthTestDrawRenderState, SDPG_Foreground, bOutDirty);
@@ -823,7 +823,7 @@ void FDeferredShadingSceneRenderer::RenderEditorPrimitives(FRHICommandListImmedi
 	}
 }
 
-void FDeferredShadingSceneRenderer::RenderEditorPrimitivesForDPG(FRHICommandList& RHICmdList, const FViewInfo& View, FExclusiveDepthStencil::Type BasePassDepthStencilAccess, const FDrawingPolicyRenderState& DrawRenderState, ESceneDepthPriorityGroup DepthPriorityGroup, bool& bOutDirty)
+void FDeferredShadingSceneRenderer::RenderEditorPrimitivesForDPG(FRHICommandList& RHICmdList, const FViewInfo& View, FExclusiveDepthStencil::Type BasePassDepthStencilAccess, const FMeshPassProcessorRenderState& DrawRenderState, ESceneDepthPriorityGroup DepthPriorityGroup, bool& bOutDirty)
 {
 	View.SimpleElementCollector.DrawBatchedElements(RHICmdList, DrawRenderState, View, EBlendModeFilter::OpaqueAndMasked, DepthPriorityGroup);
 
@@ -884,10 +884,10 @@ void FDeferredShadingSceneRenderer::RenderEditorPrimitivesForDPG(FRHICommandList
 	}
 }
 
-bool FDeferredShadingSceneRenderer::RenderBasePassView(FRHICommandListImmediate& RHICmdList, FViewInfo& View, FExclusiveDepthStencil::Type BasePassDepthStencilAccess, const FDrawingPolicyRenderState& InDrawRenderState)
+bool FDeferredShadingSceneRenderer::RenderBasePassView(FRHICommandListImmediate& RHICmdList, FViewInfo& View, FExclusiveDepthStencil::Type BasePassDepthStencilAccess, const FMeshPassProcessorRenderState& InDrawRenderState)
 {
 	bool bDirty = false; 
-	FDrawingPolicyRenderState DrawRenderState(InDrawRenderState);
+	FMeshPassProcessorRenderState DrawRenderState(InDrawRenderState);
 	SetupBasePassView(RHICmdList, View, this);
 
 	View.ParallelMeshDrawCommandPasses[EMeshPass::BasePass].DispatchDraw(nullptr, RHICmdList);
@@ -935,7 +935,7 @@ void FBasePassMeshProcessor::Process(
 		);
 
 
-	FDrawingPolicyRenderState DrawRenderState(PassDrawRenderState);
+	FMeshPassProcessorRenderState DrawRenderState(PassDrawRenderState);
 
 	const bool bEnableReceiveDecalOutput = Scene != nullptr;
 	SetDepthStencilStateForBasePass(DrawRenderState, FeatureLevel, MeshBatch, PrimitiveSceneProxy, bEnableReceiveDecalOutput, false, nullptr);
@@ -1498,7 +1498,7 @@ void FBasePassMeshProcessor::AddMeshBatch(const FMeshBatch& RESTRICT MeshBatch,
 	}
 }
 
-FBasePassMeshProcessor::FBasePassMeshProcessor(const FScene* Scene, ERHIFeatureLevel::Type InFeatureLevel, const FSceneView* InViewIfDynamicMeshCommand, const FDrawingPolicyRenderState& InDrawRenderState, FMeshPassDrawListContext* InDrawListContext, ETranslucencyPass::Type InTranslucencyPassType)
+FBasePassMeshProcessor::FBasePassMeshProcessor(const FScene* Scene, ERHIFeatureLevel::Type InFeatureLevel, const FSceneView* InViewIfDynamicMeshCommand, const FMeshPassProcessorRenderState& InDrawRenderState, FMeshPassDrawListContext* InDrawListContext, ETranslucencyPass::Type InTranslucencyPassType)
 	: FMeshPassProcessor(Scene, InFeatureLevel, InViewIfDynamicMeshCommand, InDrawListContext)
 	, PassDrawRenderState(InDrawRenderState)
 	, TranslucencyPassType(InTranslucencyPassType)
@@ -1514,7 +1514,7 @@ FMeshPassProcessor* CreateBasePassProcessor(const FScene* Scene, const FSceneVie
 {
 	FExclusiveDepthStencil::Type DefaultBasePassDepthStencilAccess = GetDefaultBasePassDepthStencilAccess(EShadingPath::Deferred);
 
-	FDrawingPolicyRenderState PassDrawRenderState(Scene->UniformBuffers.ViewUniformBuffer, Scene->UniformBuffers.OpaqueBasePassUniformBuffer);
+	FMeshPassProcessorRenderState PassDrawRenderState(Scene->UniformBuffers.ViewUniformBuffer, Scene->UniformBuffers.OpaqueBasePassUniformBuffer);
 	SetupBasePassState(DefaultBasePassDepthStencilAccess, false, PassDrawRenderState);
 
 	return new(FMemStack::Get()) FBasePassMeshProcessor(Scene, Scene->GetFeatureLevel(), InViewIfDynamicMeshCommand, PassDrawRenderState, InDrawListContext);
@@ -1522,7 +1522,7 @@ FMeshPassProcessor* CreateBasePassProcessor(const FScene* Scene, const FSceneVie
 
 FMeshPassProcessor* CreateTranslucencyStandardPassProcessor(const FScene* Scene, const FSceneView* InViewIfDynamicMeshCommand, FMeshPassDrawListContext* InDrawListContext)
 {
-	FDrawingPolicyRenderState PassDrawRenderState(Scene->UniformBuffers.ViewUniformBuffer, Scene->UniformBuffers.TranslucentBasePassUniformBuffer);
+	FMeshPassProcessorRenderState PassDrawRenderState(Scene->UniformBuffers.ViewUniformBuffer, Scene->UniformBuffers.TranslucentBasePassUniformBuffer);
 	PassDrawRenderState.SetDepthStencilState(TStaticDepthStencilState<false, CF_DepthNearOrEqual>::GetRHI());
 
 	return new(FMemStack::Get()) FBasePassMeshProcessor(Scene, Scene->GetFeatureLevel(), InViewIfDynamicMeshCommand, PassDrawRenderState, InDrawListContext, ETranslucencyPass::TPT_StandardTranslucency);
@@ -1530,7 +1530,7 @@ FMeshPassProcessor* CreateTranslucencyStandardPassProcessor(const FScene* Scene,
 
 FMeshPassProcessor* CreateTranslucencyAfterDOFProcessor(const FScene* Scene, const FSceneView* InViewIfDynamicMeshCommand, FMeshPassDrawListContext* InDrawListContext)
 {
-	FDrawingPolicyRenderState PassDrawRenderState(Scene->UniformBuffers.ViewUniformBuffer, Scene->UniformBuffers.TranslucentBasePassUniformBuffer);
+	FMeshPassProcessorRenderState PassDrawRenderState(Scene->UniformBuffers.ViewUniformBuffer, Scene->UniformBuffers.TranslucentBasePassUniformBuffer);
 	PassDrawRenderState.SetDepthStencilState(TStaticDepthStencilState<false, CF_DepthNearOrEqual>::GetRHI());
 
 	return new(FMemStack::Get()) FBasePassMeshProcessor(Scene, Scene->GetFeatureLevel(), InViewIfDynamicMeshCommand, PassDrawRenderState, InDrawListContext, ETranslucencyPass::TPT_TranslucencyAfterDOF);
@@ -1538,7 +1538,7 @@ FMeshPassProcessor* CreateTranslucencyAfterDOFProcessor(const FScene* Scene, con
 
 FMeshPassProcessor* CreateTranslucencyAllPassProcessor(const FScene* Scene, const FSceneView* InViewIfDynamicMeshCommand, FMeshPassDrawListContext* InDrawListContext)
 {
-	FDrawingPolicyRenderState PassDrawRenderState(Scene->UniformBuffers.ViewUniformBuffer, Scene->UniformBuffers.TranslucentBasePassUniformBuffer);
+	FMeshPassProcessorRenderState PassDrawRenderState(Scene->UniformBuffers.ViewUniformBuffer, Scene->UniformBuffers.TranslucentBasePassUniformBuffer);
 	PassDrawRenderState.SetDepthStencilState(TStaticDepthStencilState<false, CF_DepthNearOrEqual>::GetRHI());
 
 	return new(FMemStack::Get()) FBasePassMeshProcessor(Scene, Scene->GetFeatureLevel(), InViewIfDynamicMeshCommand, PassDrawRenderState, InDrawListContext, ETranslucencyPass::TPT_AllTranslucency);
diff --git a/Engine/Source/Runtime/Renderer/Private/BasePassRendering.h b/Engine/Source/Runtime/Renderer/Private/BasePassRendering.h
--- a/Engine/Source/Runtime/Renderer/Private/BasePassRendering.h
+++ b/Engine/Source/Runtime/Renderer/Private/BasePassRendering.h
@@ -599,11 +599,11 @@ class FBasePassMeshProcessor : public FMeshPassProcessor
 {
 public:
 
-	FBasePassMeshProcessor(const FScene* InScene, ERHIFeatureLevel::Type InFeatureLevel, const FSceneView* InViewIfDynamicMeshCommand, const FDrawingPolicyRenderState& InDrawRenderState, FMeshPassDrawListContext* InDrawListContext, ETranslucencyPass::Type InTranslucencyPassType = ETranslucencyPass::TPT_MAX);
+	FBasePassMeshProcessor(const FScene* InScene, ERHIFeatureLevel::Type InFeatureLevel, const FSceneView* InViewIfDynamicMeshCommand, const FMeshPassProcessorRenderState& InDrawRenderState, FMeshPassDrawListContext* InDrawListContext, ETranslucencyPass::Type InTranslucencyPassType = ETranslucencyPass::TPT_MAX);
 
 	virtual void AddMeshBatch(const FMeshBatch& RESTRICT MeshBatch, uint64 BatchElementMask, const FPrimitiveSceneProxy* RESTRICT PrimitiveSceneProxy, int32 StaticMeshId = -1) override final;
 
-	FDrawingPolicyRenderState PassDrawRenderState;
+	FMeshPassProcessorRenderState PassDrawRenderState;
 
 private:
 
@@ -642,6 +642,6 @@ private:
 	EDepthDrawingMode EarlyZPassMode;
 };
 
-extern void SetDepthStencilStateForBasePass(FDrawingPolicyRenderState& DrawRenderState, ERHIFeatureLevel::Type FeatureLevel, const FMeshBatch& Mesh, const FPrimitiveSceneProxy* PrimitiveSceneProxy, bool bEnableReceiveDecalOutput, bool bUseDebugViewPS, FDepthStencilStateRHIParamRef LodFadeOverrideDepthStencilState);
-extern void SetupBasePassState(FExclusiveDepthStencil::Type BasePassDepthStencilAccess, const bool bShaderComplexity, FDrawingPolicyRenderState& DrawRenderState);
+extern void SetDepthStencilStateForBasePass(FMeshPassProcessorRenderState& DrawRenderState, ERHIFeatureLevel::Type FeatureLevel, const FMeshBatch& Mesh, const FPrimitiveSceneProxy* PrimitiveSceneProxy, bool bEnableReceiveDecalOutput, bool bUseDebugViewPS, FDepthStencilStateRHIParamRef LodFadeOverrideDepthStencilState);
+extern void SetupBasePassState(FExclusiveDepthStencil::Type BasePassDepthStencilAccess, const bool bShaderComplexity, FMeshPassProcessorRenderState& DrawRenderState);
 extern FExclusiveDepthStencil::Type GetDefaultBasePassDepthStencilAccess(EShadingPath ShadingPath);
\ No newline at end of file
diff --git a/Engine/Source/Runtime/Renderer/Private/CompositionLighting/PostProcessMeshDecals.cpp b/Engine/Source/Runtime/Renderer/Private/CompositionLighting/PostProcessMeshDecals.cpp
--- a/Engine/Source/Runtime/Renderer/Private/CompositionLighting/PostProcessMeshDecals.cpp
+++ b/Engine/Source/Runtime/Renderer/Private/CompositionLighting/PostProcessMeshDecals.cpp
@@ -199,7 +199,7 @@ private:
 		ERasterizerFillMode MeshFillMode,
 		ERasterizerCullMode MeshCullMode);
 
-	FDrawingPolicyRenderState PassDrawRenderState;
+	FMeshPassProcessorRenderState PassDrawRenderState;
 	const EDecalRenderStage PassDecalStage;
 	const FDecalRenderingCommon::ERenderTargetMode RenderTargetMode;
 };
diff --git a/Engine/Source/Runtime/Renderer/Private/ConvertToUniformMesh.cpp b/Engine/Source/Runtime/Renderer/Private/ConvertToUniformMesh.cpp
--- a/Engine/Source/Runtime/Renderer/Private/ConvertToUniformMesh.cpp
+++ b/Engine/Source/Runtime/Renderer/Private/ConvertToUniformMesh.cpp
@@ -139,7 +139,7 @@ private:
 		ERasterizerFillMode MeshFillMod1e,
 		ERasterizerCullMode MeshCullMode);
 
-	FDrawingPolicyRenderState PassDrawRenderState;
+	FMeshPassProcessorRenderState PassDrawRenderState;
 };
 
 FConvertToUniformMeshProcessor::FConvertToUniformMeshProcessor(const FScene* Scene, const FViewInfo* InViewIfDynamicMeshCommand, FMeshPassDrawListContext* InDrawListContext)
diff --git a/Engine/Source/Runtime/Renderer/Private/CustomDepthRendering.cpp b/Engine/Source/Runtime/Renderer/Private/CustomDepthRendering.cpp
--- a/Engine/Source/Runtime/Renderer/Private/CustomDepthRendering.cpp
+++ b/Engine/Source/Runtime/Renderer/Private/CustomDepthRendering.cpp
@@ -32,7 +32,7 @@ private:
 		float MobileColorValue,
 		bool bUsesMobileColorValue);
 
-	FDrawingPolicyRenderState PassDrawRenderState;
+	FMeshPassProcessorRenderState PassDrawRenderState;
 };
 
 FCustomDepthPassMeshProcessor::FCustomDepthPassMeshProcessor(const FScene* Scene, const FSceneView* InViewIfDynamicMeshCommand, FMeshPassDrawListContext* InDrawListContext)
diff --git a/Engine/Source/Runtime/Renderer/Private/DebugViewModeRendering.cpp b/Engine/Source/Runtime/Renderer/Private/DebugViewModeRendering.cpp
--- a/Engine/Source/Runtime/Renderer/Private/DebugViewModeRendering.cpp
+++ b/Engine/Source/Runtime/Renderer/Private/DebugViewModeRendering.cpp
@@ -304,7 +304,7 @@ void FDebugViewModeMeshProcessor::AddMeshBatch(const FMeshBatch& RESTRICT MeshBa
 	const ERasterizerFillMode MeshFillMode = ComputeMeshFillMode(MeshBatch, *BatchMaterial);
 	const ERasterizerCullMode MeshCullMode = ComputeMeshCullMode(MeshBatch, *BatchMaterial);
 
-	FDrawingPolicyRenderState DrawRenderState;
+	FMeshPassProcessorRenderState DrawRenderState;
 	DrawRenderState.SetViewUniformBuffer(ViewUniformBuffer);
 	DrawRenderState.SetPassUniformBuffer(PassUniformBuffer);
 
diff --git a/Engine/Source/Runtime/Renderer/Private/DeferredShadingRenderer.h b/Engine/Source/Runtime/Renderer/Private/DeferredShadingRenderer.h
--- a/Engine/Source/Runtime/Renderer/Private/DeferredShadingRenderer.h
+++ b/Engine/Source/Runtime/Renderer/Private/DeferredShadingRenderer.h
@@ -58,28 +58,28 @@ public:
 	 * Renders the scene's prepass for a particular view
 	 * @return true if anything was rendered
 	 */
-	void RenderPrePassView(FRHICommandList& RHICmdList, const FViewInfo& View, const FDrawingPolicyRenderState& DrawRenderState);
+	void RenderPrePassView(FRHICommandList& RHICmdList, const FViewInfo& View, const FMeshPassProcessorRenderState& DrawRenderState);
 
 	/**
 	 * Renders the scene's prepass for a particular view in parallel
 	 * @return true if the depth was cleared
 	 */
-	bool RenderPrePassViewParallel(const FViewInfo& View, FRHICommandListImmediate& ParentCmdList, const FDrawingPolicyRenderState& DrawRenderState, TFunctionRef<void()> AfterTasksAreStarted, bool bDoPrePre);
+	bool RenderPrePassViewParallel(const FViewInfo& View, FRHICommandListImmediate& ParentCmdList, const FMeshPassProcessorRenderState& DrawRenderState, TFunctionRef<void()> AfterTasksAreStarted, bool bDoPrePre);
 
 	/** Culls local lights to a grid in frustum space.  Needed for forward shading or translucency using the Surface lighting mode. */
 	void ComputeLightGrid(FRHICommandListImmediate& RHICmdList, bool bNeedLightGrid);
 
 	/** Renders the basepass for a given View, in parallel */
-	void RenderBasePassViewParallel(FViewInfo& View, FRHICommandListImmediate& ParentCmdList, FExclusiveDepthStencil::Type BasePassDepthStencilAccess, const FDrawingPolicyRenderState& InDrawRenderState);
+	void RenderBasePassViewParallel(FViewInfo& View, FRHICommandListImmediate& ParentCmdList, FExclusiveDepthStencil::Type BasePassDepthStencilAccess, const FMeshPassProcessorRenderState& InDrawRenderState);
 
 	/** Renders the basepass for a given View. */
-	bool RenderBasePassView(FRHICommandListImmediate& RHICmdList, FViewInfo& View, FExclusiveDepthStencil::Type BasePassDepthStencilAccess, const FDrawingPolicyRenderState& InDrawRenderState);
+	bool RenderBasePassView(FRHICommandListImmediate& RHICmdList, FViewInfo& View, FExclusiveDepthStencil::Type BasePassDepthStencilAccess, const FMeshPassProcessorRenderState& InDrawRenderState);
 
 	/** Renders editor primitives for a given View. */
-	void RenderEditorPrimitives(FRHICommandListImmediate& RHICmdList, const FViewInfo& View, FExclusiveDepthStencil::Type BasePassDepthStencilAccess, const FDrawingPolicyRenderState& DrawRenderState, bool& bOutDirty);
+	void RenderEditorPrimitives(FRHICommandListImmediate& RHICmdList, const FViewInfo& View, FExclusiveDepthStencil::Type BasePassDepthStencilAccess, const FMeshPassProcessorRenderState& DrawRenderState, bool& bOutDirty);
 
 	/** Renders editor primitives for a given View. */
-	void RenderEditorPrimitivesForDPG(FRHICommandList& RHICmdList, const FViewInfo& View, FExclusiveDepthStencil::Type BasePassDepthStencilAccess, const FDrawingPolicyRenderState& DrawRenderState, ESceneDepthPriorityGroup DepthPriorityGroup, bool& bOutDirty);
+	void RenderEditorPrimitivesForDPG(FRHICommandList& RHICmdList, const FViewInfo& View, FExclusiveDepthStencil::Type BasePassDepthStencilAccess, const FMeshPassProcessorRenderState& DrawRenderState, ESceneDepthPriorityGroup DepthPriorityGroup, bool& bOutDirty);
 
 	/** 
 	* Renders the scene's base pass 
@@ -156,7 +156,7 @@ private:
 	*/
 	bool PreRenderPrePass(FRHICommandListImmediate& RHICmdList);
 
-	void RenderPrePassEditorPrimitives(FRHICommandList& RHICmdList, const FViewInfo& View, const FDrawingPolicyRenderState& DrawRenderState, EDepthDrawingMode DepthDrawingMode, bool bRespectUseAsOccluderFlag);
+	void RenderPrePassEditorPrimitives(FRHICommandList& RHICmdList, const FViewInfo& View, const FMeshPassProcessorRenderState& DrawRenderState, EDepthDrawingMode DepthDrawingMode, bool bRespectUseAsOccluderFlag);
 
 	/**
 	 * Renders the scene's prepass and occlusion queries.
@@ -424,8 +424,8 @@ private:
 	void UpdateGlobalDistanceFieldObjectBuffers(FRHICommandListImmediate& RHICmdList);
 	void PrepareDistanceFieldScene(FRHICommandListImmediate& RHICmdList, bool bSplitDispatch);
 
-	void RenderViewTranslucency(FRHICommandListImmediate& RHICmdList, const FViewInfo& View, const FDrawingPolicyRenderState& DrawRenderState, ETranslucencyPass::Type TranslucenyPass);
-	void RenderViewTranslucencyParallel(FRHICommandListImmediate& RHICmdList, const FViewInfo& View, const FDrawingPolicyRenderState& DrawRenderState, ETranslucencyPass::Type TranslucencyPass);
+	void RenderViewTranslucency(FRHICommandListImmediate& RHICmdList, const FViewInfo& View, const FMeshPassProcessorRenderState& DrawRenderState, ETranslucencyPass::Type TranslucenyPass);
+	void RenderViewTranslucencyParallel(FRHICommandListImmediate& RHICmdList, const FViewInfo& View, const FMeshPassProcessorRenderState& DrawRenderState, ETranslucencyPass::Type TranslucencyPass);
 
 	void CopySceneCaptureComponentToTarget(FRHICommandListImmediate& RHICmdList);
 
diff --git a/Engine/Source/Runtime/Renderer/Private/DepthRendering.cpp b/Engine/Source/Runtime/Renderer/Private/DepthRendering.cpp
--- a/Engine/Source/Runtime/Renderer/Private/DepthRendering.cpp
+++ b/Engine/Source/Runtime/Renderer/Private/DepthRendering.cpp
@@ -170,7 +170,7 @@ void GetDepthPassShaders(
 IMPLEMENT_GetDepthPassShaders( true );
 IMPLEMENT_GetDepthPassShaders( false );
 
-void SetDepthPassDitheredLODTransitionState(const FSceneView* SceneView, const FMeshBatch& RESTRICT Mesh, int32 StaticMeshId, FDrawingPolicyRenderState& DrawRenderState)
+void SetDepthPassDitheredLODTransitionState(const FSceneView* SceneView, const FMeshBatch& RESTRICT Mesh, int32 StaticMeshId, FMeshPassProcessorRenderState& DrawRenderState)
 {
 	if (SceneView && StaticMeshId >= 0 && Mesh.bDitheredLODTransition)
 	{
@@ -251,7 +251,7 @@ static void RenderHiddenAreaMaskView(FRHICommandList& RHICmdList, FGraphicsPipel
 	}
 }
 
-void FDeferredShadingSceneRenderer::RenderPrePassView(FRHICommandList& RHICmdList, const FViewInfo& View, const FDrawingPolicyRenderState& DrawRenderState)
+void FDeferredShadingSceneRenderer::RenderPrePassView(FRHICommandList& RHICmdList, const FViewInfo& View, const FMeshPassProcessorRenderState& DrawRenderState)
 {
 	SetupPrePassView(RHICmdList, View, this);
 
@@ -263,7 +263,7 @@ DECLARE_CYCLE_STAT(TEXT("Prepass"), STAT_CLP_Prepass, STATGROUP_ParallelCommandL
 class FPrePassParallelCommandListSet : public FParallelCommandListSet
 {
 public:
-	FPrePassParallelCommandListSet(const FViewInfo& InView, const FSceneRenderer* InSceneRenderer, FRHICommandListImmediate& InParentCmdList, bool bInParallelExecute, bool bInCreateSceneContext, const FDrawingPolicyRenderState& InDrawRenderState)
+	FPrePassParallelCommandListSet(const FViewInfo& InView, const FSceneRenderer* InSceneRenderer, FRHICommandListImmediate& InParentCmdList, bool bInParallelExecute, bool bInCreateSceneContext, const FMeshPassProcessorRenderState& InDrawRenderState)
 		: FParallelCommandListSet(GET_STATID(STAT_CLP_Prepass), InView, InSceneRenderer, InParentCmdList, bInParallelExecute, bInCreateSceneContext, InDrawRenderState)
 	{
 		// Do not copy-paste. this is a very unusual FParallelCommandListSet because it is a prepass and we want to do some work after starting some tasks
@@ -283,7 +283,7 @@ public:
 	}
 };
 
-bool FDeferredShadingSceneRenderer::RenderPrePassViewParallel(const FViewInfo& View, FRHICommandListImmediate& ParentCmdList, const FDrawingPolicyRenderState& DrawRenderState, TFunctionRef<void()> AfterTasksAreStarted, bool bDoPrePre)
+bool FDeferredShadingSceneRenderer::RenderPrePassViewParallel(const FViewInfo& View, FRHICommandListImmediate& ParentCmdList, const FMeshPassProcessorRenderState& DrawRenderState, TFunctionRef<void()> AfterTasksAreStarted, bool bDoPrePre)
 {
 	bool bDepthWasCleared = false;
 
@@ -417,7 +417,7 @@ bool FDeferredShadingSceneRenderer::PreRenderPrePass(FRHICommandListImmediate& R
 	return bDepthWasCleared;
 }
 
-void FDeferredShadingSceneRenderer::RenderPrePassEditorPrimitives(FRHICommandList& RHICmdList, const FViewInfo& View, const FDrawingPolicyRenderState& DrawRenderState, EDepthDrawingMode DepthDrawingMode, bool bRespectUseAsOccluderFlag) 
+void FDeferredShadingSceneRenderer::RenderPrePassEditorPrimitives(FRHICommandList& RHICmdList, const FViewInfo& View, const FMeshPassProcessorRenderState& DrawRenderState, EDepthDrawingMode DepthDrawingMode, bool bRespectUseAsOccluderFlag) 
 {
 	SetupPrePassView(RHICmdList, View, this, true);
 
@@ -480,7 +480,7 @@ void FDeferredShadingSceneRenderer::RenderPrePassEditorPrimitives(FRHICommandLis
 	}
 }
 
-void SetupDepthPassState(FDrawingPolicyRenderState& DrawRenderState)
+void SetupDepthPassState(FMeshPassProcessorRenderState& DrawRenderState)
 {
 	// Disable color writes, enable depth tests and writes.
 	DrawRenderState.SetBlendState(TStaticBlendState<CW_NONE>::GetRHI());
@@ -560,7 +560,7 @@ bool FDeferredShadingSceneRenderer::RenderPrePass(FRHICommandListImmediate& RHIC
 			TUniformBufferRef<FSceneTexturesUniformParameters> PassUniformBuffer;
 			CreateDepthPassUniformBuffer(RHICmdList, View, PassUniformBuffer);
 
-			FDrawingPolicyRenderState DrawRenderState(View, PassUniformBuffer);
+			FMeshPassProcessorRenderState DrawRenderState(View, PassUniformBuffer);
 
 			if (View.ShouldRenderView())
 			{
@@ -711,7 +711,7 @@ void FDepthPassMeshProcessor::Process(
 		false
 		);
 
-	FDrawingPolicyRenderState DrawRenderState(PassDrawRenderState);
+	FMeshPassProcessorRenderState DrawRenderState(PassDrawRenderState);
 
 	SetDepthPassDitheredLODTransitionState(ViewIfDynamicMeshCommand, MeshBatch, StaticMeshId, DrawRenderState);
 
@@ -813,7 +813,7 @@ void FDepthPassMeshProcessor::AddMeshBatch(const FMeshBatch& RESTRICT MeshBatch,
 
 FDepthPassMeshProcessor::FDepthPassMeshProcessor(const FScene* Scene,
 	const FSceneView* InViewIfDynamicMeshCommand,
-	const FDrawingPolicyRenderState& InPassDrawRenderState,
+	const FMeshPassProcessorRenderState& InPassDrawRenderState,
 	const bool InbRespectUseAsOccluderFlag,
 	const EDepthDrawingMode InEarlyZPassMode,
 	const bool InbEarlyZPassMovable,
@@ -836,7 +836,7 @@ FMeshPassProcessor* CreateDepthPassProcessor(const FScene* Scene, const FSceneVi
 	bool bEarlyZPassMovable;
 	GetEarlyZPassMode(Scene->GetShaderPlatform(), EarlyZPassMode, bEarlyZPassMovable);
 
-	FDrawingPolicyRenderState DepthPassState;
+	FMeshPassProcessorRenderState DepthPassState;
 	SetupDepthPassState(DepthPassState);
 	return new(FMemStack::Get()) FDepthPassMeshProcessor(Scene, InViewIfDynamicMeshCommand, DepthPassState, true, EarlyZPassMode, bEarlyZPassMovable, InDrawListContext);
 }
diff --git a/Engine/Source/Runtime/Renderer/Private/DepthRendering.h b/Engine/Source/Runtime/Renderer/Private/DepthRendering.h
--- a/Engine/Source/Runtime/Renderer/Private/DepthRendering.h
+++ b/Engine/Source/Runtime/Renderer/Private/DepthRendering.h
@@ -205,7 +205,7 @@ public:
 
 	FDepthPassMeshProcessor(const FScene* Scene, 
 		const FSceneView* InViewIfDynamicMeshCommand, 
-		const FDrawingPolicyRenderState& InPassDrawRenderState, 
+		const FMeshPassProcessorRenderState& InPassDrawRenderState, 
 		const bool InbRespectUseAsOccluderFlag,
 		const EDepthDrawingMode InEarlyZPassMode,
 		const bool InbEarlyZPassMovable,
@@ -226,11 +226,11 @@ private:
 		ERasterizerFillMode MeshFillMode,
 		ERasterizerCullMode MeshCullMode);
 
-	FDrawingPolicyRenderState PassDrawRenderState;
+	FMeshPassProcessorRenderState PassDrawRenderState;
 
 	const bool bRespectUseAsOccluderFlag;
 	const EDepthDrawingMode EarlyZPassMode;
 	const bool bEarlyZPassMovable;
 };
 
-extern void SetupDepthPassState(FDrawingPolicyRenderState& DrawRenderState);
+extern void SetupDepthPassState(FMeshPassProcessorRenderState& DrawRenderState);
diff --git a/Engine/Source/Runtime/Renderer/Private/DistortionRendering.cpp b/Engine/Source/Runtime/Renderer/Private/DistortionRendering.cpp
--- a/Engine/Source/Runtime/Renderer/Private/DistortionRendering.cpp
+++ b/Engine/Source/Runtime/Renderer/Private/DistortionRendering.cpp
@@ -479,7 +479,7 @@ static void DrawDistortionMergePass(FRHICommandListImmediate& RHICmdList, FScene
 		EDRF_UseTriangleOptimization);
 }
 
-bool SubmitDistortionMeshDrawCommands(FRHICommandListImmediate& RHICmdList, const FViewInfo& View, const FDrawingPolicyRenderState& DrawRenderState)
+bool SubmitDistortionMeshDrawCommands(FRHICommandListImmediate& RHICmdList, const FViewInfo& View, const FMeshPassProcessorRenderState& DrawRenderState)
 {
 	QUICK_SCOPE_CYCLE_COUNTER(STAT_FDistortionPrimSet_DrawAccumulatedOffsets);
 
@@ -569,7 +569,7 @@ void FSceneRenderer::RenderDistortion(FRHICommandListImmediate& RHICmdList)
 					SetupDistortionPassUniformBuffer(RHICmdList, View, DistortionPassParameters);
 					Scene->UniformBuffers.DistortionPassUniformBuffer.UpdateUniformBufferImmediate(DistortionPassParameters);
 
-					FDrawingPolicyRenderState DrawRenderState(View, Scene->UniformBuffers.DistortionPassUniformBuffer);
+					FMeshPassProcessorRenderState DrawRenderState(View, Scene->UniformBuffers.DistortionPassUniformBuffer);
 
 					// test against depth and write stencil mask
 					DrawRenderState.SetDepthStencilState(TStaticDepthStencilState<
@@ -772,7 +772,7 @@ void FDistortionMeshProcessor::Process(
 		ShaderElementData);
 }
 
-FDistortionMeshProcessor::FDistortionMeshProcessor(const FScene* Scene, const FSceneView* InViewIfDynamicMeshCommand, const FDrawingPolicyRenderState& InPassDrawRenderState, FMeshPassDrawListContext* InDrawListContext)
+FDistortionMeshProcessor::FDistortionMeshProcessor(const FScene* Scene, const FSceneView* InViewIfDynamicMeshCommand, const FMeshPassProcessorRenderState& InPassDrawRenderState, FMeshPassDrawListContext* InDrawListContext)
 	: FMeshPassProcessor(Scene, Scene->GetFeatureLevel(), InViewIfDynamicMeshCommand, InDrawListContext)
 	, PassDrawRenderState(InPassDrawRenderState)
 {
@@ -780,7 +780,7 @@ FDistortionMeshProcessor::FDistortionMeshProcessor(const FScene* Scene, const FS
 
 FMeshPassProcessor* CreateDistortionPassProcessor(const FScene* Scene, const FSceneView* InViewIfDynamicMeshCommand, FMeshPassDrawListContext* InDrawListContext)
 {
-	FDrawingPolicyRenderState DistortionPassState;
+	FMeshPassProcessorRenderState DistortionPassState;
 	DistortionPassState.SetViewUniformBuffer(Scene->UniformBuffers.ViewUniformBuffer);
 	DistortionPassState.SetPassUniformBuffer(Scene->UniformBuffers.DistortionPassUniformBuffer);
 	
@@ -801,7 +801,7 @@ FMeshPassProcessor* CreateDistortionPassProcessor(const FScene* Scene, const FSc
 
 FMeshPassProcessor* CreateMobileDistortionPassProcessor(const FScene* Scene, const FSceneView* InViewIfDynamicMeshCommand, FMeshPassDrawListContext* InDrawListContext)
 {
-	FDrawingPolicyRenderState DistortionPassState;
+	FMeshPassProcessorRenderState DistortionPassState;
 	DistortionPassState.SetViewUniformBuffer(Scene->UniformBuffers.ViewUniformBuffer);
 	DistortionPassState.SetPassUniformBuffer(Scene->UniformBuffers.MobileDistortionPassUniformBuffer);
 
diff --git a/Engine/Source/Runtime/Renderer/Private/DistortionRendering.h b/Engine/Source/Runtime/Renderer/Private/DistortionRendering.h
--- a/Engine/Source/Runtime/Renderer/Private/DistortionRendering.h
+++ b/Engine/Source/Runtime/Renderer/Private/DistortionRendering.h
@@ -11,7 +11,7 @@
 #include "MeshPassProcessor.h"
 
 class FPrimitiveSceneProxy;
-struct FDrawingPolicyRenderState;
+struct FMeshPassProcessorRenderState;
 
 /** 
 * Set of distortion scene prims  
@@ -26,7 +26,7 @@ public:
 	* @param DPGIndex - current DPG used to draw items
 	* @return true if anything was drawn
 	*/
-	bool DrawAccumulatedOffsets(FRHICommandListImmediate& RHICmdList, const class FViewInfo& View, const FDrawingPolicyRenderState& DrawRenderState, bool bInitializeOffsets);
+	bool DrawAccumulatedOffsets(FRHICommandListImmediate& RHICmdList, const class FViewInfo& View, const FMeshPassProcessorRenderState& DrawRenderState, bool bInitializeOffsets);
 
 	/**
 	* Adds a new primitives to the list of distortion prims
@@ -76,11 +76,11 @@ class FDistortionMeshProcessor : public FMeshPassProcessor
 {
 public:
 
-	FDistortionMeshProcessor(const FScene* Scene, const FSceneView* InViewIfDynamicMeshCommand, const FDrawingPolicyRenderState& InPassDrawRenderState, FMeshPassDrawListContext* InDrawListContext);
+	FDistortionMeshProcessor(const FScene* Scene, const FSceneView* InViewIfDynamicMeshCommand, const FMeshPassProcessorRenderState& InPassDrawRenderState, FMeshPassDrawListContext* InDrawListContext);
 
 	virtual void AddMeshBatch(const FMeshBatch& RESTRICT MeshBatch, uint64 BatchElementMask, const FPrimitiveSceneProxy* RESTRICT PrimitiveSceneProxy, int32 StaticMeshId = -1) override final;
 
-	FDrawingPolicyRenderState PassDrawRenderState;
+	FMeshPassProcessorRenderState PassDrawRenderState;
 
 private:
 	void Process(
diff --git a/Engine/Source/Runtime/Renderer/Private/EditorPrimitivesRendering.cpp b/Engine/Source/Runtime/Renderer/Private/EditorPrimitivesRendering.cpp
--- a/Engine/Source/Runtime/Renderer/Private/EditorPrimitivesRendering.cpp
+++ b/Engine/Source/Runtime/Renderer/Private/EditorPrimitivesRendering.cpp
@@ -10,7 +10,7 @@
 #include "MobileBasePassRendering.h"
 #include "MeshPassProcessor.inl"
 
-FEditorPrimitivesBasePassMeshProcessor::FEditorPrimitivesBasePassMeshProcessor(const FScene* Scene, ERHIFeatureLevel::Type InFeatureLevel, const FSceneView* InViewIfDynamicMeshCommand, const FDrawingPolicyRenderState& InDrawRenderState, bool bInTranslucentBasePass, FMeshPassDrawListContext* InDrawListContext) 
+FEditorPrimitivesBasePassMeshProcessor::FEditorPrimitivesBasePassMeshProcessor(const FScene* Scene, ERHIFeatureLevel::Type InFeatureLevel, const FSceneView* InViewIfDynamicMeshCommand, const FMeshPassProcessorRenderState& InDrawRenderState, bool bInTranslucentBasePass, FMeshPassDrawListContext* InDrawListContext) 
 	: FMeshPassProcessor(Scene, InFeatureLevel, InViewIfDynamicMeshCommand, InDrawListContext)
 	, PassDrawRenderState(InDrawRenderState)
 	, bTranslucentBasePass(bInTranslucentBasePass)
@@ -73,11 +73,11 @@ void FEditorPrimitivesBasePassMeshProcessor::ProcessDeferredShadingPath(const FM
 		BasePassShaders.PixelShader
 		);
 
-	FDrawingPolicyRenderState DrawRenderState(PassDrawRenderState);
+	FMeshPassProcessorRenderState DrawRenderState(PassDrawRenderState);
 
 	if (bTranslucentBasePass)
 	{
-		extern void SetTranslucentRenderState(FDrawingPolicyRenderState& DrawRenderState, const FMaterial& Material);
+		extern void SetTranslucentRenderState(FMeshPassProcessorRenderState& DrawRenderState, const FMaterial& Material);
 		SetTranslucentRenderState(DrawRenderState, Material);
 	}
 
@@ -127,7 +127,7 @@ void FEditorPrimitivesBasePassMeshProcessor::ProcessMobileShadingPath(const FMes
 		BasePassShaders.VertexShader, 
 		BasePassShaders.PixelShader);
 	
-	FDrawingPolicyRenderState DrawRenderState(PassDrawRenderState);
+	FMeshPassProcessorRenderState DrawRenderState(PassDrawRenderState);
 
 	if (bTranslucentBasePass)
 	{
diff --git a/Engine/Source/Runtime/Renderer/Private/EditorPrimitivesRendering.h b/Engine/Source/Runtime/Renderer/Private/EditorPrimitivesRendering.h
--- a/Engine/Source/Runtime/Renderer/Private/EditorPrimitivesRendering.h
+++ b/Engine/Source/Runtime/Renderer/Private/EditorPrimitivesRendering.h
@@ -8,11 +8,11 @@
 class FEditorPrimitivesBasePassMeshProcessor : public FMeshPassProcessor
 {
 public:
-	FEditorPrimitivesBasePassMeshProcessor(const FScene* InScene, ERHIFeatureLevel::Type InFeatureLevel, const FSceneView* InViewIfDynamicMeshCommand, const FDrawingPolicyRenderState& InDrawRenderState, bool bTranslucentBasePass, FMeshPassDrawListContext* InDrawListContext);
+	FEditorPrimitivesBasePassMeshProcessor(const FScene* InScene, ERHIFeatureLevel::Type InFeatureLevel, const FSceneView* InViewIfDynamicMeshCommand, const FMeshPassProcessorRenderState& InDrawRenderState, bool bTranslucentBasePass, FMeshPassDrawListContext* InDrawListContext);
 
 	virtual void AddMeshBatch(const FMeshBatch& RESTRICT MeshBatch, uint64 BatchElementMask, const FPrimitiveSceneProxy* RESTRICT PrimitiveSceneProxy, int32 StaticMeshId = -1) override final;
 
-	const FDrawingPolicyRenderState PassDrawRenderState;
+	const FMeshPassProcessorRenderState PassDrawRenderState;
 	const bool bTranslucentBasePass;
 
 private:
diff --git a/Engine/Source/Runtime/Renderer/Private/LightMapDensityRendering.cpp b/Engine/Source/Runtime/Renderer/Private/LightMapDensityRendering.cpp
--- a/Engine/Source/Runtime/Renderer/Private/LightMapDensityRendering.cpp
+++ b/Engine/Source/Runtime/Renderer/Private/LightMapDensityRendering.cpp
@@ -78,7 +78,7 @@ bool FDeferredShadingSceneRenderer::RenderLightMapDensities(FRHICommandListImmed
 			SetupLightmapDensityPassUniformBuffer(RHICmdList, View, LightmapDensityPassParameters);
 			Scene->UniformBuffers.LightmapDensityPassUniformBuffer.UpdateUniformBufferImmediate(LightmapDensityPassParameters);
 
-			FDrawingPolicyRenderState DrawRenderState(View, Scene->UniformBuffers.LightmapDensityPassUniformBuffer);
+			FMeshPassProcessorRenderState DrawRenderState(View, Scene->UniformBuffers.LightmapDensityPassUniformBuffer);
 
 			// Opaque blending, depth tests and writes.
 			DrawRenderState.SetBlendState(TStaticBlendState<>::GetRHI());
diff --git a/Engine/Source/Runtime/Renderer/Private/LightMapDensityRendering.h b/Engine/Source/Runtime/Renderer/Private/LightMapDensityRendering.h
--- a/Engine/Source/Runtime/Renderer/Private/LightMapDensityRendering.h
+++ b/Engine/Source/Runtime/Renderer/Private/LightMapDensityRendering.h
@@ -278,5 +278,5 @@ private:
 		ERasterizerFillMode MeshFillMode,
 		ERasterizerCullMode MeshCullMode);
 
-	FDrawingPolicyRenderState PassDrawRenderState;
+	FMeshPassProcessorRenderState PassDrawRenderState;
 };
\ No newline at end of file
diff --git a/Engine/Source/Runtime/Renderer/Private/MeshDrawCommands.cpp b/Engine/Source/Runtime/Renderer/Private/MeshDrawCommands.cpp
--- a/Engine/Source/Runtime/Renderer/Private/MeshDrawCommands.cpp
+++ b/Engine/Source/Runtime/Renderer/Private/MeshDrawCommands.cpp
@@ -519,7 +519,7 @@ void ApplyViewOverridesToMeshDrawCommands(
 
 				if (BasePassDepthStencilAccess != DefaultBasePassDepthStencilAccess && PassType == EMeshPass::BasePass)
 				{
-					FDrawingPolicyRenderState PassDrawRenderState;
+					FMeshPassProcessorRenderState PassDrawRenderState;
 					SetupBasePassState(BasePassDepthStencilAccess, false, PassDrawRenderState);
 					NewMeshCommand.PipelineState.DepthStencilState = PassDrawRenderState.GetDepthStencilState();
 				}
diff --git a/Engine/Source/Runtime/Renderer/Private/MobileBasePass.cpp b/Engine/Source/Runtime/Renderer/Private/MobileBasePass.cpp
--- a/Engine/Source/Runtime/Renderer/Private/MobileBasePass.cpp
+++ b/Engine/Source/Runtime/Renderer/Private/MobileBasePass.cpp
@@ -354,7 +354,7 @@ ELightMapPolicyType MobileBasePass::SelectMeshLightmapPolicy(
 	return SelectedLightmapPolicy;
 }
 
-void MobileBasePass::SetOpaqueRenderState(FDrawingPolicyRenderState& DrawRenderState, const FPrimitiveSceneProxy* PrimitiveSceneProxy, const FMaterial& Material, bool bEnableReceiveDecalOutput)
+void MobileBasePass::SetOpaqueRenderState(FMeshPassProcessorRenderState& DrawRenderState, const FPrimitiveSceneProxy* PrimitiveSceneProxy, const FMaterial& Material, bool bEnableReceiveDecalOutput)
 {
 	bool bEncodedHDR = GetMobileHDRMode() == EMobileHDRMode::EnabledRGBE && Material.GetMaterialDomain() != MD_UI;
 	if (bEncodedHDR)
@@ -380,7 +380,7 @@ void MobileBasePass::SetOpaqueRenderState(FDrawingPolicyRenderState& DrawRenderS
 	}
 }
 
-void MobileBasePass::SetTranslucentRenderState(FDrawingPolicyRenderState& DrawRenderState, const FMaterial& Material)
+void MobileBasePass::SetTranslucentRenderState(FMeshPassProcessorRenderState& DrawRenderState, const FMaterial& Material)
 {
 	bool bEncodedHDR = GetMobileHDRMode() == EMobileHDRMode::EnabledRGBE && Material.GetMaterialDomain() != MD_UI;
 	static const auto CVarMonoscopicFarField = IConsoleManager::Get().FindTConsoleVariableDataInt(TEXT("vr.MonoscopicFarField"));
@@ -552,7 +552,7 @@ void TMobileBasePassPSPolicyParamType<FUniformLightMapPolicy>::GetShaderBindings
 	}
 }
 
-FMobileBasePassMeshProcessor::FMobileBasePassMeshProcessor(const FScene* Scene, ERHIFeatureLevel::Type InFeatureLevel, const FSceneView* InViewIfDynamicMeshCommand, const FDrawingPolicyRenderState& InDrawRenderState, FMeshPassDrawListContext* InDrawListContext, bool bInCanReceiveCSM, ETranslucencyPass::Type InTranslucencyPassType)
+FMobileBasePassMeshProcessor::FMobileBasePassMeshProcessor(const FScene* Scene, ERHIFeatureLevel::Type InFeatureLevel, const FSceneView* InViewIfDynamicMeshCommand, const FMeshPassProcessorRenderState& InDrawRenderState, FMeshPassDrawListContext* InDrawListContext, bool bInCanReceiveCSM, ETranslucencyPass::Type InTranslucencyPassType)
 	: FMeshPassProcessor(Scene, InFeatureLevel, InViewIfDynamicMeshCommand, InDrawListContext)
 	, PassDrawRenderState(InDrawRenderState)
 	, TranslucencyPassType(InTranslucencyPassType)
@@ -632,7 +632,7 @@ void FMobileBasePassMeshProcessor::Process(
 		BasePassShaders.VertexShader, 
 		BasePassShaders.PixelShader);
 
-	FDrawingPolicyRenderState DrawRenderState(PassDrawRenderState);
+	FMeshPassProcessorRenderState DrawRenderState(PassDrawRenderState);
 	if (bTranslucentBasePass)
 	{
 		MobileBasePass::SetTranslucentRenderState(DrawRenderState, MaterialResource);
@@ -681,7 +681,7 @@ FMeshPassProcessor* CreateMobileBasePassProcessor(const FScene* Scene, const FSc
 {
 	FExclusiveDepthStencil::Type DefaultBasePassDepthStencilAccess = GetDefaultBasePassDepthStencilAccess(EShadingPath::Mobile);
 	
-	FDrawingPolicyRenderState PassDrawRenderState(Scene->UniformBuffers.ViewUniformBuffer, Scene->UniformBuffers.MobileOpaqueBasePassUniformBuffer);
+	FMeshPassProcessorRenderState PassDrawRenderState(Scene->UniformBuffers.ViewUniformBuffer, Scene->UniformBuffers.MobileOpaqueBasePassUniformBuffer);
 	PassDrawRenderState.SetBlendState(TStaticBlendStateWriteMask<CW_RGBA>::GetRHI());
 	PassDrawRenderState.SetDepthStencilAccess(DefaultBasePassDepthStencilAccess);
 	PassDrawRenderState.SetDepthStencilState(TStaticDepthStencilState<true, CF_DepthNearOrEqual>::GetRHI());
@@ -693,7 +693,7 @@ FMeshPassProcessor* CreateMobileBasePassCSMProcessor(const FScene* Scene, const
 {
 	FExclusiveDepthStencil::Type DefaultBasePassDepthStencilAccess = GetDefaultBasePassDepthStencilAccess(EShadingPath::Mobile);
 
-	FDrawingPolicyRenderState PassDrawRenderState(Scene->UniformBuffers.ViewUniformBuffer, Scene->UniformBuffers.MobileOpaqueBasePassUniformBuffer);
+	FMeshPassProcessorRenderState PassDrawRenderState(Scene->UniformBuffers.ViewUniformBuffer, Scene->UniformBuffers.MobileOpaqueBasePassUniformBuffer);
 	PassDrawRenderState.SetBlendState(TStaticBlendStateWriteMask<CW_RGBA>::GetRHI());
 	PassDrawRenderState.SetDepthStencilAccess(DefaultBasePassDepthStencilAccess);
 	PassDrawRenderState.SetDepthStencilState(TStaticDepthStencilState<true, CF_DepthNearOrEqual>::GetRHI());
@@ -703,7 +703,7 @@ FMeshPassProcessor* CreateMobileBasePassCSMProcessor(const FScene* Scene, const
 
 FMeshPassProcessor* CreateMobileTranslucencyStandardPassProcessor(const FScene* Scene, const FSceneView* InViewIfDynamicMeshCommand, FMeshPassDrawListContext* InDrawListContext)
 {
-	FDrawingPolicyRenderState PassDrawRenderState(Scene->UniformBuffers.ViewUniformBuffer, Scene->UniformBuffers.MobileTranslucentBasePassUniformBuffer);
+	FMeshPassProcessorRenderState PassDrawRenderState(Scene->UniformBuffers.ViewUniformBuffer, Scene->UniformBuffers.MobileTranslucentBasePassUniformBuffer);
 	PassDrawRenderState.SetDepthStencilState(TStaticDepthStencilState<false, CF_DepthNearOrEqual>::GetRHI());
 	PassDrawRenderState.SetDepthStencilAccess(FExclusiveDepthStencil::DepthRead_StencilRead);
 
@@ -712,7 +712,7 @@ FMeshPassProcessor* CreateMobileTranslucencyStandardPassProcessor(const FScene*
 
 FMeshPassProcessor* CreateMobileTranslucencyAfterDOFProcessor(const FScene* Scene, const FSceneView* InViewIfDynamicMeshCommand, FMeshPassDrawListContext* InDrawListContext)
 {
-	FDrawingPolicyRenderState PassDrawRenderState(Scene->UniformBuffers.ViewUniformBuffer, Scene->UniformBuffers.MobileTranslucentBasePassUniformBuffer);
+	FMeshPassProcessorRenderState PassDrawRenderState(Scene->UniformBuffers.ViewUniformBuffer, Scene->UniformBuffers.MobileTranslucentBasePassUniformBuffer);
 	PassDrawRenderState.SetDepthStencilState(TStaticDepthStencilState<false, CF_DepthNearOrEqual>::GetRHI());
 	PassDrawRenderState.SetDepthStencilAccess(FExclusiveDepthStencil::DepthRead_StencilRead);
 
@@ -721,7 +721,7 @@ FMeshPassProcessor* CreateMobileTranslucencyAfterDOFProcessor(const FScene* Scen
 
 FMeshPassProcessor* CreateMobileTranslucencyAllPassProcessor(const FScene* Scene, const FSceneView* InViewIfDynamicMeshCommand, FMeshPassDrawListContext* InDrawListContext)
 {
-	FDrawingPolicyRenderState PassDrawRenderState(Scene->UniformBuffers.ViewUniformBuffer, Scene->UniformBuffers.MobileTranslucentBasePassUniformBuffer);
+	FMeshPassProcessorRenderState PassDrawRenderState(Scene->UniformBuffers.ViewUniformBuffer, Scene->UniformBuffers.MobileTranslucentBasePassUniformBuffer);
 	PassDrawRenderState.SetDepthStencilState(TStaticDepthStencilState<false, CF_DepthNearOrEqual>::GetRHI());
 	PassDrawRenderState.SetDepthStencilAccess(FExclusiveDepthStencil::DepthRead_StencilRead);
 
diff --git a/Engine/Source/Runtime/Renderer/Private/MobileBasePassRendering.cpp b/Engine/Source/Runtime/Renderer/Private/MobileBasePassRendering.cpp
--- a/Engine/Source/Runtime/Renderer/Private/MobileBasePassRendering.cpp
+++ b/Engine/Source/Runtime/Renderer/Private/MobileBasePassRendering.cpp
@@ -217,7 +217,7 @@ void FMobileSceneRenderer::RenderMobileBasePass(FRHICommandListImmediate& RHICmd
 
 		// editor primitives
 		{
-			FDrawingPolicyRenderState DrawRenderState(View, Scene->UniformBuffers.MobileOpaqueBasePassUniformBuffer);
+			FMeshPassProcessorRenderState DrawRenderState(View, Scene->UniformBuffers.MobileOpaqueBasePassUniformBuffer);
 			DrawRenderState.SetBlendState(TStaticBlendStateWriteMask<CW_RGBA>::GetRHI());
 			DrawRenderState.SetDepthStencilAccess(GetDefaultBasePassDepthStencilAccess(EShadingPath::Mobile));
 			DrawRenderState.SetDepthStencilState(TStaticDepthStencilState<true, CF_DepthNearOrEqual>::GetRHI());
@@ -226,7 +226,7 @@ void FMobileSceneRenderer::RenderMobileBasePass(FRHICommandListImmediate& RHICmd
 	}
 }
 
-void FMobileSceneRenderer::RenderMobileEditorPrimitives(FRHICommandList& RHICmdList, const FViewInfo& View, const FDrawingPolicyRenderState& DrawRenderState)
+void FMobileSceneRenderer::RenderMobileEditorPrimitives(FRHICommandList& RHICmdList, const FViewInfo& View, const FMeshPassProcessorRenderState& DrawRenderState)
 {
 	QUICK_SCOPE_CYCLE_COUNTER(STAT_EditorDynamicPrimitiveDrawTime);
 	SCOPED_DRAW_EVENT(RHICmdList, DynamicEd);
diff --git a/Engine/Source/Runtime/Renderer/Private/MobileBasePassRendering.h b/Engine/Source/Runtime/Renderer/Private/MobileBasePassRendering.h
--- a/Engine/Source/Runtime/Renderer/Private/MobileBasePassRendering.h
+++ b/Engine/Source/Runtime/Renderer/Private/MobileBasePassRendering.h
@@ -351,8 +351,8 @@ namespace MobileBasePass
 		
 	bool StaticCanReceiveCSM(const FLightSceneInfo* LightSceneInfo, const FPrimitiveSceneProxy* PrimitiveSceneProxy);
 
-	void SetOpaqueRenderState(FDrawingPolicyRenderState& DrawRenderState, const FPrimitiveSceneProxy* PrimitiveSceneProxy, const FMaterial& Material, bool bEnableReceiveDecalOutput);
-	void SetTranslucentRenderState(FDrawingPolicyRenderState& DrawRenderState, const FMaterial& Material);
+	void SetOpaqueRenderState(FMeshPassProcessorRenderState& DrawRenderState, const FPrimitiveSceneProxy* PrimitiveSceneProxy, const FMaterial& Material, bool bEnableReceiveDecalOutput);
+	void SetTranslucentRenderState(FMeshPassProcessorRenderState& DrawRenderState, const FMaterial& Material);
 };
 
 
@@ -436,14 +436,14 @@ public:
 		const FScene* InScene, 
 		ERHIFeatureLevel::Type InFeatureLevel, 
 		const FSceneView* InViewIfDynamicMeshCommand, 
-		const FDrawingPolicyRenderState& InDrawRenderState, 
+		const FMeshPassProcessorRenderState& InDrawRenderState, 
 		FMeshPassDrawListContext* InDrawListContext,
 		bool bInCanReceiveCSM,
 		ETranslucencyPass::Type InTranslucencyPassType = ETranslucencyPass::TPT_MAX);
 
 	virtual void AddMeshBatch(const FMeshBatch& RESTRICT MeshBatch, uint64 BatchElementMask, const FPrimitiveSceneProxy* RESTRICT PrimitiveSceneProxy, int32 StaticMeshId = -1) override final;
 
-	FDrawingPolicyRenderState PassDrawRenderState;
+	FMeshPassProcessorRenderState PassDrawRenderState;
 
 private:
 	void Process(
diff --git a/Engine/Source/Runtime/Renderer/Private/MobileTranslucentRendering.cpp b/Engine/Source/Runtime/Renderer/Private/MobileTranslucentRendering.cpp
--- a/Engine/Source/Runtime/Renderer/Private/MobileTranslucentRendering.cpp
+++ b/Engine/Source/Runtime/Renderer/Private/MobileTranslucentRendering.cpp
@@ -317,10 +317,10 @@ bool FMobileSceneRenderer::RenderInverseOpacity(FRHICommandListImmediate& RHICmd
 class FMobileInverseOpacityMeshProcessor : public FMeshPassProcessor
 {
 public:
-	const FDrawingPolicyRenderState PassDrawRenderState;
+	const FMeshPassProcessorRenderState PassDrawRenderState;
 
 public:
-	FMobileInverseOpacityMeshProcessor(const FScene* InScene, ERHIFeatureLevel::Type InFeatureLevel, const FSceneView* InViewIfDynamicMeshCommand, const FDrawingPolicyRenderState& InDrawRenderState, FMeshPassDrawListContext* InDrawListContext)
+	FMobileInverseOpacityMeshProcessor(const FScene* InScene, ERHIFeatureLevel::Type InFeatureLevel, const FSceneView* InViewIfDynamicMeshCommand, const FMeshPassProcessorRenderState& InDrawRenderState, FMeshPassDrawListContext* InDrawListContext)
 		: FMeshPassProcessor(InScene, InFeatureLevel, InViewIfDynamicMeshCommand, InDrawListContext)
 		, PassDrawRenderState(InDrawRenderState)
 	{
@@ -360,7 +360,7 @@ private:
 		InverseOpacityShaders.VertexShader = Material.GetShader<FOpacityOnlyVS>(VertexFactory->GetType());
 		InverseOpacityShaders.PixelShader = Material.GetShader<FOpacityOnlyPS>(VertexFactory->GetType());
 
-		FDrawingPolicyRenderState DrawRenderState(PassDrawRenderState);
+		FMeshPassProcessorRenderState DrawRenderState(PassDrawRenderState);
 		MobileBasePass::SetTranslucentRenderState(DrawRenderState, Material);
 
 		ERasterizerFillMode MeshFillMode = ComputeMeshFillMode(MeshBatch, Material);
@@ -391,7 +391,7 @@ private:
 // This pass is registered only when we render to scene capture, see UpdateSceneCaptureContentMobile_RenderThread()
 FMeshPassProcessor* CreateMobileInverseOpacityPassProcessor(const FScene* Scene, const FSceneView* InViewIfDynamicMeshCommand, FMeshPassDrawListContext* InDrawListContext)
 {
-	FDrawingPolicyRenderState PassDrawRenderState(Scene->UniformBuffers.ViewUniformBuffer, Scene->UniformBuffers.MobileTranslucentBasePassUniformBuffer);
+	FMeshPassProcessorRenderState PassDrawRenderState(Scene->UniformBuffers.ViewUniformBuffer, Scene->UniformBuffers.MobileTranslucentBasePassUniformBuffer);
 	PassDrawRenderState.SetDepthStencilState(TStaticDepthStencilState<false, CF_DepthNearOrEqual>::GetRHI());
 	PassDrawRenderState.SetBlendState(TStaticBlendState<CW_ALPHA, BO_Add, BF_DestColor, BF_Zero, BO_Add, BF_Zero, BF_InverseSourceAlpha>::GetRHI());
 
diff --git a/Engine/Source/Runtime/Renderer/Private/PostProcess/PostProcessCompositeEditorPrimitives.cpp b/Engine/Source/Runtime/Renderer/Private/PostProcess/PostProcessCompositeEditorPrimitives.cpp
--- a/Engine/Source/Runtime/Renderer/Private/PostProcess/PostProcessCompositeEditorPrimitives.cpp
+++ b/Engine/Source/Runtime/Renderer/Private/PostProcess/PostProcessCompositeEditorPrimitives.cpp
@@ -300,7 +300,7 @@ static void SetPopulateSceneDepthForEditorPrimitivesShaderTempl(const FRendering
 	PixelShader->SetParameters(Context);
 }
 
-static void RenderEditorPrimitives(FRHICommandListImmediate& RHICmdList, const FViewInfo& View, FDrawingPolicyRenderState& DrawRenderState)
+static void RenderEditorPrimitives(FRHICommandListImmediate& RHICmdList, const FViewInfo& View, FMeshPassProcessorRenderState& DrawRenderState)
 {
 	// Always depth test against other editor primitives
 	DrawRenderState.SetDepthStencilState(TStaticDepthStencilState<
@@ -351,7 +351,7 @@ static void RenderEditorPrimitives(FRHICommandListImmediate& RHICmdList, const F
 	View.BatchedViewElements.Draw(RHICmdList, DrawRenderState, FeatureLevel, bNeedToSwitchVerticalAxis, View, false, 1.0f);
 }
 
-static void RenderForegroundEditorPrimitives(FRHICommandListImmediate& RHICmdList, const FViewInfo& View, FDrawingPolicyRenderState& DrawRenderState)
+static void RenderForegroundEditorPrimitives(FRHICommandListImmediate& RHICmdList, const FViewInfo& View, FMeshPassProcessorRenderState& DrawRenderState)
 {
 	FSceneRenderTargets& SceneContext = FSceneRenderTargets::Get(RHICmdList);
 
@@ -570,7 +570,7 @@ void FRCPassPostProcessCompositeEditorPrimitives::Process(FRenderingCompositePas
 				BasePassUniformBuffer = MobileBasePassUniformBuffer;
 			}
 
-			FDrawingPolicyRenderState DrawRenderState(EditorView, BasePassUniformBuffer);
+			FMeshPassProcessorRenderState DrawRenderState(EditorView, BasePassUniformBuffer);
 			DrawRenderState.SetDepthStencilAccess(FExclusiveDepthStencil::DepthWrite_StencilWrite);
 			DrawRenderState.SetBlendState(TStaticBlendStateWriteMask<CW_RGBA>::GetRHI());
 
diff --git a/Engine/Source/Runtime/Renderer/Private/PostProcess/PostProcessCompositeEditorPrimitives.h b/Engine/Source/Runtime/Renderer/Private/PostProcess/PostProcessCompositeEditorPrimitives.h
--- a/Engine/Source/Runtime/Renderer/Private/PostProcess/PostProcessCompositeEditorPrimitives.h
+++ b/Engine/Source/Runtime/Renderer/Private/PostProcess/PostProcessCompositeEditorPrimitives.h
@@ -7,7 +7,7 @@
 #include "PostProcess/RenderingCompositionGraph.h"
 
 class FViewInfo;
-struct FDrawingPolicyRenderState;
+struct FMeshPassProcessorRenderState;
 
 #if WITH_EDITOR
 
diff --git a/Engine/Source/Runtime/Renderer/Private/RayTracing/RayTracingMaterialHitShaders.cpp b/Engine/Source/Runtime/Renderer/Private/RayTracing/RayTracingMaterialHitShaders.cpp
--- a/Engine/Source/Runtime/Renderer/Private/RayTracing/RayTracingMaterialHitShaders.cpp
+++ b/Engine/Source/Runtime/Renderer/Private/RayTracing/RayTracingMaterialHitShaders.cpp
@@ -217,7 +217,7 @@ void FRayTracingMeshProcessor::Process(
 	}
 	
 	// #dxr_todo: use something else instead of OpaqueBasePassUniformBuffer. We will have a more clear 
-	FDrawingPolicyRenderState PassDrawRenderState(Scene->UniformBuffers.ViewUniformBuffer, Scene->UniformBuffers.OpaqueBasePassUniformBuffer);
+	FMeshPassProcessorRenderState PassDrawRenderState(Scene->UniformBuffers.ViewUniformBuffer, Scene->UniformBuffers.OpaqueBasePassUniformBuffer);
 	PassDrawRenderState.SetBlendState(TStaticBlendState<CW_RGBA, BO_Add, BF_One, BF_One, BO_Add, BF_Zero, BF_One>::GetRHI());
 	PassDrawRenderState.SetDepthStencilState(TStaticDepthStencilState<false, CF_DepthNearOrEqual>::GetRHI());
 
diff --git a/Engine/Source/Runtime/Renderer/Private/Renderer.cpp b/Engine/Source/Runtime/Renderer/Private/Renderer.cpp
--- a/Engine/Source/Runtime/Renderer/Private/Renderer.cpp
+++ b/Engine/Source/Runtime/Renderer/Private/Renderer.cpp
@@ -88,7 +88,7 @@ void FRendererModule::InitializeSystemTextures(FRHICommandListImmediate& RHICmdL
 	GSystemTextures.InitializeTextures(RHICmdList, GMaxRHIFeatureLevel);
 }
 
-void FRendererModule::DrawTileMesh(FRHICommandListImmediate& RHICmdList, FDrawingPolicyRenderState& DrawRenderState, const FSceneView& SceneView, FMeshBatch& Mesh, bool bIsHitTesting, const FHitProxyId& HitProxyId)
+void FRendererModule::DrawTileMesh(FRHICommandListImmediate& RHICmdList, FMeshPassProcessorRenderState& DrawRenderState, const FSceneView& SceneView, FMeshBatch& Mesh, bool bIsHitTesting, const FHitProxyId& HitProxyId)
 {
 	if (!GUsingNullRHI)
 	{
diff --git a/Engine/Source/Runtime/Renderer/Private/RendererModule.h b/Engine/Source/Runtime/Renderer/Private/RendererModule.h
--- a/Engine/Source/Runtime/Renderer/Private/RendererModule.h
+++ b/Engine/Source/Runtime/Renderer/Private/RendererModule.h
@@ -47,7 +47,7 @@ public:
 	virtual void ReallocateSceneRenderTargets() override;
 	virtual void SceneRenderTargetsSetBufferSize(uint32 SizeX, uint32 SizeY) override;
 	virtual void InitializeSystemTextures(FRHICommandListImmediate& RHICmdList);
-	virtual void DrawTileMesh(FRHICommandListImmediate& RHICmdList, FDrawingPolicyRenderState& DrawRenderState, const FSceneView& View, FMeshBatch& Mesh, bool bIsHitTesting, const FHitProxyId& HitProxyId) override;
+	virtual void DrawTileMesh(FRHICommandListImmediate& RHICmdList, FMeshPassProcessorRenderState& DrawRenderState, const FSceneView& View, FMeshBatch& Mesh, bool bIsHitTesting, const FHitProxyId& HitProxyId) override;
 	virtual void RenderTargetPoolFindFreeElement(FRHICommandListImmediate& RHICmdList, const FPooledRenderTargetDesc& Desc, TRefCountPtr<IPooledRenderTarget> &Out, const TCHAR* InDebugName) override;
 	virtual void TickRenderTargetPool() override;
 	virtual void DebugLogOnCrash() override;
diff --git a/Engine/Source/Runtime/Renderer/Private/SceneHitProxyRendering.cpp b/Engine/Source/Runtime/Renderer/Private/SceneHitProxyRendering.cpp
--- a/Engine/Source/Runtime/Renderer/Private/SceneHitProxyRendering.cpp
+++ b/Engine/Source/Runtime/Renderer/Private/SceneHitProxyRendering.cpp
@@ -270,7 +270,7 @@ static void DoRenderHitProxies(FRHICommandListImmediate& RHICmdList, const FScen
 		SetupSceneTextureUniformParameters(SceneContext, View.FeatureLevel, ESceneTextureSetupMode::None, SceneTextureParameters);
 		SceneRenderer->Scene->UniformBuffers.HitProxyPassUniformBuffer.UpdateUniformBufferImmediate(SceneTextureParameters);
 
-		FDrawingPolicyRenderState DrawRenderState(View, SceneRenderer->Scene->UniformBuffers.HitProxyPassUniformBuffer);
+		FMeshPassProcessorRenderState DrawRenderState(View, SceneRenderer->Scene->UniformBuffers.HitProxyPassUniformBuffer);
 
 		// Set the device viewport for the view.
 		RHICmdList.SetViewport(View.ViewRect.Min.X, View.ViewRect.Min.Y, 0.0f, View.ViewRect.Max.X, View.ViewRect.Max.Y, 1.0f);
@@ -477,7 +477,7 @@ static void DoRenderHitProxies(FRHICommandListImmediate& RHICmdList, const FScen
 		RHICmdList.BeginRenderPass(RPInfo, TEXT("HitProxies"));
 		{
 			FSceneView SceneView = FBatchedElements::CreateProxySceneView(PixelToView, FIntRect(0, 0, RenderTargetSize.X, RenderTargetSize.Y));
-			FDrawingPolicyRenderState DrawRenderState(SceneView);
+			FMeshPassProcessorRenderState DrawRenderState(SceneView);
 
 			DrawRenderState.SetDepthStencilState(TStaticDepthStencilState<false, CF_Always>::GetRHI());
 			DrawRenderState.SetBlendState(TStaticBlendState<>::GetRHI());
@@ -671,7 +671,7 @@ void FHitProxyMeshProcessor::Process(
 		ShaderElementData);
 }
 
-FHitProxyMeshProcessor::FHitProxyMeshProcessor(const FScene* Scene, const FSceneView* InViewIfDynamicMeshCommand, bool InbAllowTranslucentPrimitivesInHitProxy, const FDrawingPolicyRenderState& InRenderState, FMeshPassDrawListContext* InDrawListContext)
+FHitProxyMeshProcessor::FHitProxyMeshProcessor(const FScene* Scene, const FSceneView* InViewIfDynamicMeshCommand, bool InbAllowTranslucentPrimitivesInHitProxy, const FMeshPassProcessorRenderState& InRenderState, FMeshPassDrawListContext* InDrawListContext)
 	: FMeshPassProcessor(Scene, Scene->GetFeatureLevel(), InViewIfDynamicMeshCommand, InDrawListContext)
 	, PassDrawRenderState(InRenderState)
 	, bAllowTranslucentPrimitivesInHitProxy(InbAllowTranslucentPrimitivesInHitProxy)
@@ -680,7 +680,7 @@ FHitProxyMeshProcessor::FHitProxyMeshProcessor(const FScene* Scene, const FScene
 
 FMeshPassProcessor* CreateHitProxyPassProcessor(const FScene* Scene, const FSceneView* InViewIfDynamicMeshCommand, FMeshPassDrawListContext* InDrawListContext)
 {
-	FDrawingPolicyRenderState PassDrawRenderState(Scene->UniformBuffers.ViewUniformBuffer, Scene->UniformBuffers.HitProxyPassUniformBuffer);
+	FMeshPassProcessorRenderState PassDrawRenderState(Scene->UniformBuffers.ViewUniformBuffer, Scene->UniformBuffers.HitProxyPassUniformBuffer);
 	PassDrawRenderState.SetDepthStencilState(TStaticDepthStencilState<true, CF_DepthNearOrEqual>::GetRHI());
 	PassDrawRenderState.SetBlendState(TStaticBlendState<>::GetRHI());
 	return new(FMemStack::Get()) FHitProxyMeshProcessor(Scene, InViewIfDynamicMeshCommand, true, PassDrawRenderState, InDrawListContext);
@@ -688,7 +688,7 @@ FMeshPassProcessor* CreateHitProxyPassProcessor(const FScene* Scene, const FScen
 
 FMeshPassProcessor* CreateHitProxyOpaqueOnlyPassProcessor(const FScene* Scene, const FSceneView* InViewIfDynamicMeshCommand, FMeshPassDrawListContext* InDrawListContext)
 {
-	FDrawingPolicyRenderState PassDrawRenderState(Scene->UniformBuffers.ViewUniformBuffer, Scene->UniformBuffers.HitProxyPassUniformBuffer);
+	FMeshPassProcessorRenderState PassDrawRenderState(Scene->UniformBuffers.ViewUniformBuffer, Scene->UniformBuffers.HitProxyPassUniformBuffer);
 	PassDrawRenderState.SetDepthStencilState(TStaticDepthStencilState<true, CF_DepthNearOrEqual>::GetRHI());
 	PassDrawRenderState.SetBlendState(TStaticBlendState<>::GetRHI());
 	return new(FMemStack::Get()) FHitProxyMeshProcessor(Scene, InViewIfDynamicMeshCommand, false, PassDrawRenderState, InDrawListContext);
diff --git a/Engine/Source/Runtime/Renderer/Private/SceneHitProxyRendering.h b/Engine/Source/Runtime/Renderer/Private/SceneHitProxyRendering.h
--- a/Engine/Source/Runtime/Renderer/Private/SceneHitProxyRendering.h
+++ b/Engine/Source/Runtime/Renderer/Private/SceneHitProxyRendering.h
@@ -21,11 +21,11 @@ class FHitProxyMeshProcessor : public FMeshPassProcessor
 {
 public:
 
-	FHitProxyMeshProcessor(const FScene* Scene, const FSceneView* InViewIfDynamicMeshCommand, bool InbAllowTranslucentPrimitivesInHitProxy, const FDrawingPolicyRenderState& InRenderState, FMeshPassDrawListContext* InDrawListContext);
+	FHitProxyMeshProcessor(const FScene* Scene, const FSceneView* InViewIfDynamicMeshCommand, bool InbAllowTranslucentPrimitivesInHitProxy, const FMeshPassProcessorRenderState& InRenderState, FMeshPassDrawListContext* InDrawListContext);
 
 	virtual void AddMeshBatch(const FMeshBatch& RESTRICT MeshBatch, uint64 BatchElementMask, const FPrimitiveSceneProxy* RESTRICT PrimitiveSceneProxy, int32 StaticMeshId = -1) override final;
 
-	FDrawingPolicyRenderState PassDrawRenderState;
+	FMeshPassProcessorRenderState PassDrawRenderState;
 
 
 private:
@@ -51,7 +51,7 @@ public:
 
 	virtual void AddMeshBatch(const FMeshBatch& RESTRICT MeshBatch, uint64 BatchElementMask, const FPrimitiveSceneProxy* RESTRICT PrimitiveSceneProxy, int32 StaticMeshId = -1) override final;
 
-	FDrawingPolicyRenderState PassDrawRenderState;
+	FMeshPassProcessorRenderState PassDrawRenderState;
 
 private:
 	void Process(
diff --git a/Engine/Source/Runtime/Renderer/Private/SceneRendering.cpp b/Engine/Source/Runtime/Renderer/Private/SceneRendering.cpp
--- a/Engine/Source/Runtime/Renderer/Private/SceneRendering.cpp
+++ b/Engine/Source/Runtime/Renderer/Private/SceneRendering.cpp
@@ -590,7 +590,7 @@ FParallelCommandListSet::FParallelCommandListSet(
 	FRHICommandListImmediate& InParentCmdList, 
 	bool bInParallelExecute, 
 	bool bInCreateSceneContext,
-	const FDrawingPolicyRenderState& InDrawRenderState)
+	const FMeshPassProcessorRenderState& InDrawRenderState)
 	: View(InView)
 	, SceneRenderer(InSceneRenderer)
 	, DrawRenderState(InDrawRenderState)
diff --git a/Engine/Source/Runtime/Renderer/Private/SceneRendering.h b/Engine/Source/Runtime/Renderer/Private/SceneRendering.h
--- a/Engine/Source/Runtime/Renderer/Private/SceneRendering.h
+++ b/Engine/Source/Runtime/Renderer/Private/SceneRendering.h
@@ -457,7 +457,7 @@ class FParallelCommandListSet
 public:
 	const FViewInfo& View;
 	const FSceneRenderer* SceneRenderer;
-	FDrawingPolicyRenderState DrawRenderState;
+	FMeshPassProcessorRenderState DrawRenderState;
 	FRHICommandListImmediate& ParentCmdList;
 	const FRHIGPUMask GPUMask; // Copy of the Parent GPUMask at creation (since it could change).
 	FSceneRenderTargets* Snapshot;
@@ -489,7 +489,7 @@ public:
 		FRHICommandListImmediate& InParentCmdList, 
 		bool bInParallelExecute, 
 		bool bInCreateSceneContext, 
-		const FDrawingPolicyRenderState& InDrawRenderState);
+		const FMeshPassProcessorRenderState& InDrawRenderState);
 
 	virtual ~FParallelCommandListSet();
 	int32 NumParallelCommandLists() const
@@ -1721,7 +1721,7 @@ protected:
 	/** Renders the opaque base pass for mobile. */
 	void RenderMobileBasePass(FRHICommandListImmediate& RHICmdList, const TArrayView<const FViewInfo*> PassViews);
 
-	void RenderMobileEditorPrimitives(FRHICommandList& RHICmdList, const FViewInfo& View, const FDrawingPolicyRenderState& DrawRenderState);
+	void RenderMobileEditorPrimitives(FRHICommandList& RHICmdList, const FViewInfo& View, const FMeshPassProcessorRenderState& DrawRenderState);
 
 	/** Render modulated shadow projections in to the scene, loops over any unrendered shadows until all are processed.*/
 	void RenderModulatedShadowProjections(FRHICommandListImmediate& RHICmdList);
diff --git a/Engine/Source/Runtime/Renderer/Private/ShadowDepthRendering.cpp b/Engine/Source/Runtime/Renderer/Private/ShadowDepthRendering.cpp
--- a/Engine/Source/Runtime/Renderer/Private/ShadowDepthRendering.cpp
+++ b/Engine/Source/Runtime/Renderer/Private/ShadowDepthRendering.cpp
@@ -771,7 +771,7 @@ void FProjectedShadowInfo::SetStateForView(FRHICommandList& RHICmdList) const
 	);
 }
 
-void SetStateForShadowDepth(bool bReflectiveShadowmap, bool bOnePassPointLightShadow, FDrawingPolicyRenderState& DrawRenderState)
+void SetStateForShadowDepth(bool bReflectiveShadowmap, bool bOnePassPointLightShadow, FMeshPassProcessorRenderState& DrawRenderState)
 {
 	if (bReflectiveShadowmap && !bOnePassPointLightShadow)
 	{
@@ -826,7 +826,7 @@ public:
 		FRHICommandListImmediate& InParentCmdList,
 		bool bInParallelExecute,
 		bool bInCreateSceneContext,
-		const FDrawingPolicyRenderState& InDrawRenderState,
+		const FMeshPassProcessorRenderState& InDrawRenderState,
 		FProjectedShadowInfo& InProjectedShadowInfo,
 		FBeginShadowRenderPassFunction InBeginShadowRenderPass)
 		: FParallelCommandListSet(GET_STATID(STAT_CLP_Shadow), InView, InSceneRenderer, InParentCmdList, bInParallelExecute, bInCreateSceneContext, InDrawRenderState)
@@ -947,7 +947,7 @@ public:
 
 IMPLEMENT_SHADER_TYPE(, FCopyShadowMaps2DPS, TEXT("/Engine/Private/CopyShadowMaps.usf"), TEXT("Copy2DDepthPS"), SF_Pixel);
 
-void FProjectedShadowInfo::CopyCachedShadowMap(FRHICommandList& RHICmdList, const FDrawingPolicyRenderState& DrawRenderState, FSceneRenderer* SceneRenderer, const FViewInfo& View)
+void FProjectedShadowInfo::CopyCachedShadowMap(FRHICommandList& RHICmdList, const FMeshPassProcessorRenderState& DrawRenderState, FSceneRenderer* SceneRenderer, const FViewInfo& View)
 {
 	check(CacheMode == SDCM_MovablePrimitivesOnly);
 	const FCachedShadowMapData& CachedShadowMapData = SceneRenderer->Scene->CachedShadowMaps.FindChecked(GetLightSceneInfo().Id);
@@ -1099,7 +1099,7 @@ void FProjectedShadowInfo::RenderDepthInner(FRHICommandListImmediate& RHICmdList
 		PassUniformBuffer = SceneRenderer->Scene->UniformBuffers.MobileCSMShadowDepthPassUniformBuffer;
 	}
 
-	FDrawingPolicyRenderState DrawRenderState(*ShadowDepthView, PassUniformBuffer);
+	FMeshPassProcessorRenderState DrawRenderState(*ShadowDepthView, PassUniformBuffer);
 	SetStateForShadowDepth(bReflectiveShadowmap, bOnePassPointLightShadow, DrawRenderState);
 	SetStateForView(RHICmdList);
 
@@ -1847,7 +1847,7 @@ FShadowDepthPassMeshProcessor::FShadowDepthPassMeshProcessor(
 	FShadowDepthType InShadowDepthType,
 	FMeshPassDrawListContext* InDrawListContext)
 	: FMeshPassProcessor(Scene, Scene->GetFeatureLevel(), InViewIfDynamicMeshCommand, InDrawListContext)
-	, PassDrawRenderState(FDrawingPolicyRenderState(InViewUniformBuffer, InPassUniformBuffer))
+	, PassDrawRenderState(FMeshPassProcessorRenderState(InViewUniformBuffer, InPassUniformBuffer))
 	, ShadowDepthType(InShadowDepthType)
 {
 	SetStateForShadowDepth(ShadowDepthType.bReflectiveShadowmap, ShadowDepthType.bOnePassPointLightShadow, PassDrawRenderState);
diff --git a/Engine/Source/Runtime/Renderer/Private/ShadowRendering.cpp b/Engine/Source/Runtime/Renderer/Private/ShadowRendering.cpp
--- a/Engine/Source/Runtime/Renderer/Private/ShadowRendering.cpp
+++ b/Engine/Source/Runtime/Renderer/Private/ShadowRendering.cpp
@@ -591,7 +591,7 @@ void FProjectedShadowInfo::SetupProjectionStencilMask(
 	bool bMobileModulatedProjections, 
 	bool bCameraInsideShadowFrustum) const
 {
-	FDrawingPolicyRenderState DrawRenderState(*View);
+	FMeshPassProcessorRenderState DrawRenderState(*View);
 
 	// Depth test wo/ writes, no color writing.
 	DrawRenderState.SetDepthStencilState(TStaticDepthStencilState<false, CF_DepthNearOrEqual>::GetRHI());
diff --git a/Engine/Source/Runtime/Renderer/Private/ShadowRendering.h b/Engine/Source/Runtime/Renderer/Private/ShadowRendering.h
--- a/Engine/Source/Runtime/Renderer/Private/ShadowRendering.h
+++ b/Engine/Source/Runtime/Renderer/Private/ShadowRendering.h
@@ -109,7 +109,7 @@ public:
 
 	virtual void AddMeshBatch(const FMeshBatch& RESTRICT MeshBatch, uint64 BatchElementMask, const FPrimitiveSceneProxy* RESTRICT PrimitiveSceneProxy, int32 StaticMeshId = -1) override final;
 
-	FDrawingPolicyRenderState PassDrawRenderState;
+	FMeshPassProcessorRenderState PassDrawRenderState;
 
 private:
 
@@ -372,7 +372,7 @@ public:
 	void SetStateForView(FRHICommandList& RHICmdList) const;
 
 	/** Set state for depth rendering */
-	void SetStateForDepth(FDrawingPolicyRenderState& DrawRenderState) const;
+	void SetStateForDepth(FMeshPassProcessorRenderState& DrawRenderState) const;
 
 	void ClearDepth(FRHICommandList& RHICmdList, class FSceneRenderer* SceneRenderer, int32 NumColorTextures, FTextureRHIParamRef* ColorTextures, FTextureRHIParamRef DepthTexture, bool bPerformClear);
 
@@ -545,7 +545,7 @@ private:
 	 */
 	float ShaderDepthBias;
 
-	void CopyCachedShadowMap(FRHICommandList& RHICmdList, const FDrawingPolicyRenderState& DrawRenderState, FSceneRenderer* SceneRenderer, const FViewInfo& View);
+	void CopyCachedShadowMap(FRHICommandList& RHICmdList, const FMeshPassProcessorRenderState& DrawRenderState, FSceneRenderer* SceneRenderer, const FViewInfo& View);
 
 	/**
 	* Renders the shadow subject depth, to a particular hacked view
diff --git a/Engine/Source/Runtime/Renderer/Private/ShadowSetup.cpp b/Engine/Source/Runtime/Renderer/Private/ShadowSetup.cpp
--- a/Engine/Source/Runtime/Renderer/Private/ShadowSetup.cpp
+++ b/Engine/Source/Runtime/Renderer/Private/ShadowSetup.cpp
@@ -1198,7 +1198,7 @@ void FProjectedShadowInfo::SetupMeshDrawCommandsForProjectionStenciling(FSceneRe
 
 			FDynamicPassMeshDrawListContext ProjectionStencilingContext(DynamicMeshDrawCommandStorage, ProjectionStencilingPass.VisibleMeshDrawCommands);
 
-			FDrawingPolicyRenderState DrawRenderState;
+			FMeshPassProcessorRenderState DrawRenderState;
 			DrawRenderState.SetBlendState(TStaticBlendState<CW_NONE>::GetRHI());
 
 			if (bMobileModulatedShadowsAllowSelfShadow)
diff --git a/Engine/Source/Runtime/Renderer/Private/TranslucentLighting.cpp b/Engine/Source/Runtime/Renderer/Private/TranslucentLighting.cpp
--- a/Engine/Source/Runtime/Renderer/Private/TranslucentLighting.cpp
+++ b/Engine/Source/Runtime/Renderer/Private/TranslucentLighting.cpp
@@ -271,14 +271,14 @@ public:
 		const FMaterialRenderProxy* MaterialRenderProxy,
 		const FSceneView& View,
 		const FProjectedShadowInfo* ShadowInfo, 
-		const FDrawingPolicyRenderState& DrawRenderState
+		const FMeshPassProcessorRenderState& DrawRenderState
 		)
 	{
 		FMeshMaterialShader::SetParameters(RHICmdList, GetVertexShader(), MaterialRenderProxy, *MaterialRenderProxy->GetMaterial(View.GetFeatureLevel()), View, DrawRenderState.GetViewUniformBuffer(), DrawRenderState.GetPassUniformBuffer());
 		ShadowParameters.SetVertexShader(RHICmdList, this, View, ShadowInfo, MaterialRenderProxy);
 	}
 
-	void SetMesh(FRHICommandList& RHICmdList, const FVertexFactory* VertexFactory,const FSceneView& View,const FPrimitiveSceneProxy* Proxy,const FMeshBatchElement& BatchElement,const FDrawingPolicyRenderState& DrawRenderState)
+	void SetMesh(FRHICommandList& RHICmdList, const FVertexFactory* VertexFactory,const FSceneView& View,const FPrimitiveSceneProxy* Proxy,const FMeshBatchElement& BatchElement,const FMeshPassProcessorRenderState& DrawRenderState)
 	{
 		FMeshMaterialShader::SetMesh(RHICmdList, GetVertexShader(),VertexFactory,View,Proxy,BatchElement,DrawRenderState);
 	}
@@ -344,7 +344,7 @@ public:
 		const FMaterialRenderProxy* MaterialRenderProxy,
 		const FSceneView& View,
 		const FProjectedShadowInfo* ShadowInfo, 
-		const FDrawingPolicyRenderState& DrawRenderState
+		const FMeshPassProcessorRenderState& DrawRenderState
 		)
 	{
 		const FPixelShaderRHIParamRef ShaderRHI = GetPixelShader();
@@ -363,7 +363,7 @@ public:
 		SetUniformBufferParameterImmediate(RHICmdList, ShaderRHI, GetUniformBufferParameter<FTranslucentSelfShadowUniformParameters>(), TranslucentSelfShadowUniformParameters);
 	}
 
-	void SetMesh(FRHICommandList& RHICmdList, const FVertexFactory* VertexFactory,const FSceneView& View,const FPrimitiveSceneProxy* Proxy,const FMeshBatchElement& BatchElement,const FDrawingPolicyRenderState& DrawRenderState)
+	void SetMesh(FRHICommandList& RHICmdList, const FVertexFactory* VertexFactory,const FSceneView& View,const FPrimitiveSceneProxy* Proxy,const FMeshBatchElement& BatchElement,const FMeshPassProcessorRenderState& DrawRenderState)
 	{
 		FMeshMaterialShader::SetMesh(RHICmdList, GetPixelShader(),VertexFactory,View,Proxy,BatchElement,DrawRenderState);
 	}
@@ -447,7 +447,7 @@ public:
 		BaseVertexShader = VertexShader;
 	}
 
-	void SetSharedState(FRHICommandList& RHICmdList, const FDrawingPolicyRenderState& DrawRenderState, const FSceneView* View, const ContextDataType PolicyContext) const
+	void SetSharedState(FRHICommandList& RHICmdList, const FMeshPassProcessorRenderState& DrawRenderState, const FSceneView* View, const ContextDataType PolicyContext) const
 	{
 		// Set the shared mesh resources.
 		FMeshDrawingPolicy::SetSharedState(RHICmdList, DrawRenderState, View, PolicyContext);
@@ -478,7 +478,7 @@ public:
 		const FPrimitiveSceneProxy* PrimitiveSceneProxy,
 		const FMeshBatch& Mesh,
 		int32 BatchElementIndex,
-		const FDrawingPolicyRenderState& DrawRenderState,
+		const FMeshPassProcessorRenderState& DrawRenderState,
 		const ElementDataType& ElementData,
 		const ContextDataType PolicyContext
 		) const
@@ -516,7 +516,7 @@ public:
 		ContextType DrawingContext,
 		const FMeshBatch& Mesh,
 		bool bPreFog,
-		const FDrawingPolicyRenderState& DrawRenderState,
+		const FMeshPassProcessorRenderState& DrawRenderState,
 		const FPrimitiveSceneProxy* PrimitiveSceneProxy,
 		FHitProxyId HitProxyId
 		)
@@ -535,7 +535,7 @@ public:
 			{
 				FTranslucencyShadowDepthDrawingPolicy DrawingPolicy(Mesh.VertexFactory, MaterialRenderProxy, *MaterialRenderProxy->GetMaterial(FeatureLevel), ComputeMeshOverrideSettings(Mesh), DrawingContext.bDirectionalLight);
 
-				FDrawingPolicyRenderState DrawRenderStateLocal(DrawRenderState);
+				FMeshPassProcessorRenderState DrawRenderStateLocal(DrawRenderState);
 				DrawRenderStateLocal.SetDitheredLODTransitionAlpha(Mesh.DitheredLODTransitionAlpha);
 				DrawingPolicy.SetupPipelineState(DrawRenderStateLocal, View);
 				CommitGraphicsPipelineState(RHICmdList, DrawingPolicy, DrawRenderStateLocal, DrawingPolicy.GetBoundShaderStateInput(View.GetFeatureLevel()), DrawingPolicy.GetMaterialRenderProxy());
@@ -565,7 +565,7 @@ public:
 		ContextType DrawingContext,
 		const FStaticMeshBatch& StaticMesh,
 		bool bPreFog,
-		const FDrawingPolicyRenderState& DrawRenderState,
+		const FMeshPassProcessorRenderState& DrawRenderState,
 		const FPrimitiveSceneProxy* PrimitiveSceneProxy,
 		FHitProxyId HitProxyId
 		)
@@ -604,7 +604,7 @@ void FProjectedShadowInfo::RenderTranslucencyDepths(FRHICommandList& RHICmdList,
 	SetupSceneTextureUniformParameters(SceneContext, ShadowDepthView->FeatureLevel, ESceneTextureSetupMode::None, SceneTextureParameters);
 	TUniformBufferRef<FSceneTexturesUniformParameters> PassUniformBuffer = TUniformBufferRef<FSceneTexturesUniformParameters>::CreateUniformBufferImmediate(SceneTextureParameters, UniformBuffer_SingleFrame);	
 
-	FDrawingPolicyRenderState DrawRenderState(*ShadowDepthView, PassUniformBuffer);
+	FMeshPassProcessorRenderState DrawRenderState(*ShadowDepthView, PassUniformBuffer);
 	{
 #if WANTS_DRAW_MESH_EVENTS
 		FString EventName;
diff --git a/Engine/Source/Runtime/Renderer/Private/TranslucentRendering.cpp b/Engine/Source/Runtime/Renderer/Private/TranslucentRendering.cpp
--- a/Engine/Source/Runtime/Renderer/Private/TranslucentRendering.cpp
+++ b/Engine/Source/Runtime/Renderer/Private/TranslucentRendering.cpp
@@ -363,7 +363,7 @@ class FTranslucencyPassParallelCommandListSet : public FParallelCommandListSet
 	bool bRenderInSeparateTranslucency;
 
 public:
-	FTranslucencyPassParallelCommandListSet(const FViewInfo& InView, const FSceneRenderer* InSceneRenderer, FRHICommandListImmediate& InParentCmdList, bool bInParallelExecute, bool bInCreateSceneContext, const FDrawingPolicyRenderState& InDrawRenderState, ETranslucencyPass::Type InTranslucencyPass, bool InRenderInSeparateTranslucency)
+	FTranslucencyPassParallelCommandListSet(const FViewInfo& InView, const FSceneRenderer* InSceneRenderer, FRHICommandListImmediate& InParentCmdList, bool bInParallelExecute, bool bInCreateSceneContext, const FMeshPassProcessorRenderState& InDrawRenderState, ETranslucencyPass::Type InTranslucencyPass, bool InRenderInSeparateTranslucency)
 		: FParallelCommandListSet(GET_STATID(STAT_CLP_Translucency), InView, InSceneRenderer, InParentCmdList, bInParallelExecute, bInCreateSceneContext, InDrawRenderState)
 		, TranslucencyPass(InTranslucencyPass)
 		, bRenderInSeparateTranslucency(InRenderInSeparateTranslucency)
@@ -409,7 +409,7 @@ static TAutoConsoleVariable<int32> CVarParallelTranslucency(
 	ECVF_RenderThreadSafe
 );
 
-void RenderViewTranslucencyInner(FRHICommandListImmediate& RHICmdList, const FViewInfo& View, const FDrawingPolicyRenderState& DrawRenderState, ETranslucencyPass::Type TranslucencyPass, FTranslucencyPassParallelCommandListSet* ParallelCommandListSet)
+void RenderViewTranslucencyInner(FRHICommandListImmediate& RHICmdList, const FViewInfo& View, const FMeshPassProcessorRenderState& DrawRenderState, ETranslucencyPass::Type TranslucencyPass, FTranslucencyPassParallelCommandListSet* ParallelCommandListSet)
 {
 	SCOPED_GPU_MASK(RHICmdList, View.GPUMask);
 
@@ -492,14 +492,14 @@ void RenderViewTranslucencyInner(FRHICommandListImmediate& RHICmdList, const FVi
 	}
 }
 
-void FDeferredShadingSceneRenderer::RenderViewTranslucency(FRHICommandListImmediate& RHICmdList, const FViewInfo& View, const FDrawingPolicyRenderState& DrawRenderState, ETranslucencyPass::Type TranslucencyPass)
+void FDeferredShadingSceneRenderer::RenderViewTranslucency(FRHICommandListImmediate& RHICmdList, const FViewInfo& View, const FMeshPassProcessorRenderState& DrawRenderState, ETranslucencyPass::Type TranslucencyPass)
 {
 	check(RHICmdList.IsInsideRenderPass());
 
 	RenderViewTranslucencyInner(RHICmdList, View, DrawRenderState, TranslucencyPass, nullptr);
 }
 
-void FDeferredShadingSceneRenderer::RenderViewTranslucencyParallel(FRHICommandListImmediate& RHICmdList, const FViewInfo& View, const FDrawingPolicyRenderState& DrawRenderState, ETranslucencyPass::Type TranslucencyPass)
+void FDeferredShadingSceneRenderer::RenderViewTranslucencyParallel(FRHICommandListImmediate& RHICmdList, const FViewInfo& View, const FMeshPassProcessorRenderState& DrawRenderState, ETranslucencyPass::Type TranslucencyPass)
 {
 	check(RHICmdList.IsOutsideRenderPass());
 
@@ -945,7 +945,7 @@ void FDeferredShadingSceneRenderer::RenderTranslucency(FRHICommandListImmediate&
 
 		TUniformBufferRef<FTranslucentBasePassUniformParameters> BasePassUniformBuffer;
 		CreateTranslucentBasePassUniformBuffer(RHICmdList, View, SceneColorCopy, ESceneTextureSetupMode::All, BasePassUniformBuffer, ViewIndex);
-		FDrawingPolicyRenderState DrawRenderState(View, BasePassUniformBuffer);
+		FMeshPassProcessorRenderState DrawRenderState(View, BasePassUniformBuffer);
 
 		// If downsampling we need to render in the separate buffer. Otherwise we also need to render offscreen to apply TPT_TranslucencyAfterDOF
 		if (RenderInSeparateTranslucency(SceneContext, TranslucencyPass, View.TranslucentPrimSet.PrimsNum.DisableOffscreenRendering(TranslucencyPass)))
diff --git a/Engine/Source/Runtime/Renderer/Private/VelocityRendering.cpp b/Engine/Source/Runtime/Renderer/Private/VelocityRendering.cpp
--- a/Engine/Source/Runtime/Renderer/Private/VelocityRendering.cpp
+++ b/Engine/Source/Runtime/Renderer/Private/VelocityRendering.cpp
@@ -261,7 +261,7 @@ static void BeginVelocityRendering(FRHICommandList& RHICmdList, TRefCountPtr<IPo
 	}
 }
 
-static void SetVelocitiesState(FRHICommandList& RHICmdList, const FViewInfo& View, const FSceneRenderer* SceneRender, FDrawingPolicyRenderState& DrawRenderState, TRefCountPtr<IPooledRenderTarget>& VelocityRT)
+static void SetVelocitiesState(FRHICommandList& RHICmdList, const FViewInfo& View, const FSceneRenderer* SceneRender, FMeshPassProcessorRenderState& DrawRenderState, TRefCountPtr<IPooledRenderTarget>& VelocityRT)
 {
 	const FIntPoint BufferSize = FSceneRenderTargets::Get(RHICmdList).GetBufferSizeXY();
 	const FIntPoint VelocityBufferSize = BufferSize;		// full resolution so we can reuse the existing full res z buffer
@@ -316,7 +316,7 @@ public:
 		FRHICommandListImmediate& InParentCmdList,
 		bool bInParallelExecute,
 		bool bInCreateSceneContext,
-		const FDrawingPolicyRenderState& InDrawRenderState,
+		const FMeshPassProcessorRenderState& InDrawRenderState,
 		TRefCountPtr<IPooledRenderTarget>& InVelocityRT)
 		: FParallelCommandListSet(GET_STATID(STAT_CLP_Velocity), InView, InSceneRenderer, InParentCmdList, bInParallelExecute, bInCreateSceneContext, InDrawRenderState)
 		, VelocityRT(InVelocityRT)
@@ -363,7 +363,7 @@ void FDeferredShadingSceneRenderer::RenderVelocitiesInnerParallel(FRHICommandLis
 			SetupSceneTextureUniformParameters(SceneContext, View.FeatureLevel, ESceneTextureSetupMode::None, SceneTextureParameters);
 			TUniformBufferRef<FSceneTexturesUniformParameters> PassUniformBuffer = TUniformBufferRef<FSceneTexturesUniformParameters>::CreateUniformBufferImmediate(SceneTextureParameters, UniformBuffer_SingleFrame);
 
-			FDrawingPolicyRenderState DrawRenderState(View, PassUniformBuffer);
+			FMeshPassProcessorRenderState DrawRenderState(View, PassUniformBuffer);
 
 			FVelocityPassParallelCommandListSet ParallelCommandListSet(View,
 				this,
@@ -391,7 +391,7 @@ void FDeferredShadingSceneRenderer::RenderVelocitiesInner(FRHICommandListImmedia
 		SetupSceneTextureUniformParameters(SceneContext, View.FeatureLevel, ESceneTextureSetupMode::None, SceneTextureParameters);
 		TUniformBufferRef<FSceneTexturesUniformParameters> PassUniformBuffer = TUniformBufferRef<FSceneTexturesUniformParameters>::CreateUniformBufferImmediate(SceneTextureParameters, UniformBuffer_SingleFrame);	
 
-		FDrawingPolicyRenderState DrawRenderState(View, PassUniformBuffer);
+		FMeshPassProcessorRenderState DrawRenderState(View, PassUniformBuffer);
 
 		if (View.ShouldRenderView())
 		{
@@ -704,7 +704,7 @@ void FVelocityMeshProcessor::Process(
 		ShaderElementData);
 }
 
-FVelocityMeshProcessor::FVelocityMeshProcessor(const FScene* Scene, const FSceneView* InViewIfDynamicMeshCommand, const FDrawingPolicyRenderState& InPassDrawRenderState, FMeshPassDrawListContext* InDrawListContext)
+FVelocityMeshProcessor::FVelocityMeshProcessor(const FScene* Scene, const FSceneView* InViewIfDynamicMeshCommand, const FMeshPassProcessorRenderState& InPassDrawRenderState, FMeshPassDrawListContext* InDrawListContext)
 	: FMeshPassProcessor(Scene, Scene->GetFeatureLevel(), InViewIfDynamicMeshCommand, InDrawListContext)
 {
 	PassDrawRenderState = InPassDrawRenderState;
@@ -714,7 +714,7 @@ FVelocityMeshProcessor::FVelocityMeshProcessor(const FScene* Scene, const FScene
 
 FMeshPassProcessor* CreateVelocityPassProcessor(const FScene* Scene, const FSceneView* InViewIfDynamicMeshCommand, FMeshPassDrawListContext* InDrawListContext)
 {
-	FDrawingPolicyRenderState VelocityPassState;
+	FMeshPassProcessorRenderState VelocityPassState;
 	VelocityPassState.SetBlendState(TStaticBlendState<CW_RGBA>::GetRHI());
 	VelocityPassState.SetDepthStencilState(TStaticDepthStencilState<false, CF_DepthNearOrEqual>::GetRHI());
 
diff --git a/Engine/Source/Runtime/Renderer/Private/VelocityRendering.h b/Engine/Source/Runtime/Renderer/Private/VelocityRendering.h
--- a/Engine/Source/Runtime/Renderer/Private/VelocityRendering.h
+++ b/Engine/Source/Runtime/Renderer/Private/VelocityRendering.h
@@ -50,11 +50,11 @@ class FVelocityMeshProcessor : public FMeshPassProcessor
 {
 public:
 
-	FVelocityMeshProcessor(const FScene* Scene, const FSceneView* InViewIfDynamicMeshCommand, const FDrawingPolicyRenderState& InPassDrawRenderState, FMeshPassDrawListContext* InDrawListContext);
+	FVelocityMeshProcessor(const FScene* Scene, const FSceneView* InViewIfDynamicMeshCommand, const FMeshPassProcessorRenderState& InPassDrawRenderState, FMeshPassDrawListContext* InDrawListContext);
 
 	virtual void AddMeshBatch(const FMeshBatch& RESTRICT MeshBatch, uint64 BatchElementMask, const FPrimitiveSceneProxy* RESTRICT PrimitiveSceneProxy, int32 StaticMeshId = -1) override final;
 
-	FDrawingPolicyRenderState PassDrawRenderState;
+	FMeshPassProcessorRenderState PassDrawRenderState;
 
 private:
 	void Process(
diff --git a/Engine/Source/Runtime/Renderer/Private/VolumetricFogVoxelization.cpp b/Engine/Source/Runtime/Renderer/Private/VolumetricFogVoxelization.cpp
--- a/Engine/Source/Runtime/Renderer/Private/VolumetricFogVoxelization.cpp
+++ b/Engine/Source/Runtime/Renderer/Private/VolumetricFogVoxelization.cpp
@@ -425,7 +425,7 @@ private:
 		ERasterizerFillMode MeshFillMode,
 		ERasterizerCullMode MeshCullMode);
 
-	FDrawingPolicyRenderState PassDrawRenderState;
+	FMeshPassProcessorRenderState PassDrawRenderState;
 };
 
 FVoxelizeVolumeMeshProcessor::FVoxelizeVolumeMeshProcessor(const FScene* Scene, const FViewInfo* InViewIfDynamicMeshCommand, FMeshPassDrawListContext* InDrawListContext)
@@ -618,7 +618,7 @@ void FDeferredShadingSceneRenderer::VoxelizeFogVolumePrimitives(
 			Scene->UniformBuffers.VoxelizeVolumeViewUniformBuffer.UpdateUniformBufferImmediate(ViewVoxelizeParameters);
 			Scene->UniformBuffers.VoxelizeVolumePassUniformBuffer.UpdateUniformBufferImmediate(VoxelizeVolumePassParameters);
 
-			FDrawingPolicyRenderState DrawRenderState(View, Scene->UniformBuffers.VoxelizeVolumePassUniformBuffer);
+			FMeshPassProcessorRenderState DrawRenderState(View, Scene->UniformBuffers.VoxelizeVolumePassUniformBuffer);
 			DrawRenderState.SetViewUniformBuffer(Scene->UniformBuffers.VoxelizeVolumeViewUniformBuffer);
 
 			DrawDynamicMeshPass(View, RHICmdList,
diff --git a/Engine/Source/Runtime/Renderer/Public/MeshMaterialShader.h b/Engine/Source/Runtime/Renderer/Public/MeshMaterialShader.h
--- a/Engine/Source/Runtime/Renderer/Public/MeshMaterialShader.h
+++ b/Engine/Source/Runtime/Renderer/Public/MeshMaterialShader.h
@@ -16,7 +16,7 @@
 class FPrimitiveSceneProxy;
 struct FMeshBatchElement;
 struct FMeshDrawingRenderState;
-struct FDrawingPolicyRenderState;
+struct FMeshPassProcessorRenderState;
 
 template<typename TBufferStruct> class TUniformBufferRef;
 
diff --git a/Engine/Source/Runtime/Renderer/Public/MeshPassProcessor.h b/Engine/Source/Runtime/Renderer/Public/MeshPassProcessor.h
--- a/Engine/Source/Runtime/Renderer/Public/MeshPassProcessor.h
+++ b/Engine/Source/Runtime/Renderer/Public/MeshPassProcessor.h
@@ -680,48 +680,12 @@ enum class EMeshPassFeatures
 };
 ENUM_CLASS_FLAGS(EMeshPassFeatures);
 
-class FDefaultSubPolicy
-{
-public:
-	struct ElementDataType {};
-
-	void GetShaderBindings(
-		const FPrimitiveSceneProxy* PrimitiveSceneProxy,
-		ElementDataType ElementData,
-#if !(UE_BUILD_SHIPPING || UE_BUILD_TEST)
-		int32 VisualizeElementIndex,
-#endif
-		const FMeshMaterialShader* VertexShaderParameters,
-		const FMeshMaterialShader* PixelShaderParameters,
-		FMeshDrawSingleShaderBindings& VertexBindings,
-		FMeshDrawSingleShaderBindings& PixelBindings
-	) const
-	{}
-
-	void GetGeometryShaderBindings(
-		const FPrimitiveSceneProxy* PrimitiveSceneProxy,
-		ElementDataType ElementData,
-		const FMeshMaterialShader* GeometryShaderParameters,
-		FMeshDrawSingleShaderBindings& GeometryBindings
-	) const
-	{
-	}
-
-#if RHI_RAYTRACING
-	void GetRayHitGroupShaderBindings(
-		const FPrimitiveSceneProxy* PrimitiveSceneProxy,
-		ElementDataType ElementData,
-		const FMeshMaterialShader* RayHitGroupShaderParameters,
-		FMeshDrawSingleShaderBindings& RayHitGroupBindings
-	) const
-	{
-	}
-#endif // RHI_RAYTRACING
-};
-
-struct FDrawingPolicyRenderState
+/**
+ * A set of render state overrides passed into a Mesh Pass Processor, so it can be configured from the outside.
+ */
+struct FMeshPassProcessorRenderState
 {
-	FDrawingPolicyRenderState(const FSceneView& SceneView, FUniformBufferRHIParamRef InPassUniformBuffer = nullptr) : 
+	FMeshPassProcessorRenderState(const FSceneView& SceneView, FUniformBufferRHIParamRef InPassUniformBuffer = nullptr) : 
 		  BlendState(nullptr)
 		, DepthStencilState(nullptr)
 		, DepthStencilAccess(FExclusiveDepthStencil::DepthRead_StencilRead)
@@ -731,7 +695,7 @@ struct FDrawingPolicyRenderState
 	{
 	}
 
-	FDrawingPolicyRenderState(const TUniformBufferRef<FViewUniformShaderParameters>& InViewUniformBuffer, FUniformBufferRHIParamRef InPassUniformBuffer) : 
+	FMeshPassProcessorRenderState(const TUniformBufferRef<FViewUniformShaderParameters>& InViewUniformBuffer, FUniformBufferRHIParamRef InPassUniformBuffer) : 
 		  BlendState(nullptr)
 		, DepthStencilState(nullptr)
 		, DepthStencilAccess(FExclusiveDepthStencil::DepthRead_StencilRead)
@@ -741,7 +705,7 @@ struct FDrawingPolicyRenderState
 	{
 	}
 
-	FDrawingPolicyRenderState() :
+	FMeshPassProcessorRenderState() :
 		BlendState(nullptr)
 		, DepthStencilState(nullptr)
 		, ViewUniformBuffer()
@@ -750,7 +714,7 @@ struct FDrawingPolicyRenderState
 	{
 	}
 
-	FORCEINLINE_DEBUGGABLE FDrawingPolicyRenderState(const FDrawingPolicyRenderState& DrawRenderState) :
+	FORCEINLINE_DEBUGGABLE FMeshPassProcessorRenderState(const FMeshPassProcessorRenderState& DrawRenderState) :
 		  BlendState(DrawRenderState.BlendState)
 		, DepthStencilState(DrawRenderState.DepthStencilState)
 		, DepthStencilAccess(DrawRenderState.DepthStencilAccess)
@@ -760,7 +724,7 @@ struct FDrawingPolicyRenderState
 	{
 	}
 
-	~FDrawingPolicyRenderState()
+	~FMeshPassProcessorRenderState()
 	{
 	}
 
@@ -882,7 +846,7 @@ public:
 		const FPrimitiveSceneProxy* RESTRICT PrimitiveSceneProxy,
 		const FMaterialRenderProxy& RESTRICT MaterialRenderProxy,
 		const FMaterial& RESTRICT MaterialResource,
-		const FDrawingPolicyRenderState& RESTRICT DrawRenderState,
+		const FMeshPassProcessorRenderState& RESTRICT DrawRenderState,
 		PassShadersType PassShaders,
 		ERasterizerFillMode MeshFillMode,
 		ERasterizerCullMode MeshCullMode,
@@ -898,7 +862,7 @@ public:
 		const FPrimitiveSceneProxy* RESTRICT PrimitiveSceneProxy,
 		const FMaterialRenderProxy& RESTRICT MaterialRenderProxy,
 		const FMaterial& RESTRICT MaterialResource,
-		const FDrawingPolicyRenderState& RESTRICT DrawRenderState,
+		const FMeshPassProcessorRenderState& RESTRICT DrawRenderState,
 		PassShadersType PassShaders,
 		ERasterizerFillMode MeshFillMode,
 		ERasterizerCullMode MeshCullMode,
diff --git a/Engine/Source/Runtime/Renderer/Public/MeshPassProcessor.inl b/Engine/Source/Runtime/Renderer/Public/MeshPassProcessor.inl
--- a/Engine/Source/Runtime/Renderer/Public/MeshPassProcessor.inl
+++ b/Engine/Source/Runtime/Renderer/Public/MeshPassProcessor.inl
@@ -13,7 +13,7 @@ void FMeshPassProcessor::BuildMeshDrawCommands(
 	const FPrimitiveSceneProxy* RESTRICT PrimitiveSceneProxy,
 	const FMaterialRenderProxy& RESTRICT MaterialRenderProxy,
 	const FMaterial& RESTRICT MaterialResource,
-	const FDrawingPolicyRenderState& RESTRICT DrawRenderState,
+	const FMeshPassProcessorRenderState& RESTRICT DrawRenderState,
 	PassShadersType PassShaders,
 	ERasterizerFillMode MeshFillMode,
 	ERasterizerCullMode MeshCullMode,
@@ -162,7 +162,7 @@ void FMeshPassProcessor::BuildRayTracingDrawCommands(
 	const FPrimitiveSceneProxy* RESTRICT PrimitiveSceneProxy,
 	const FMaterialRenderProxy& RESTRICT MaterialRenderProxy,
 	const FMaterial& RESTRICT MaterialResource,
-	const FDrawingPolicyRenderState& RESTRICT DrawRenderState,
+	const FMeshPassProcessorRenderState& RESTRICT DrawRenderState,
 	PassShadersType PassShaders,
 	ERasterizerFillMode MeshFillMode,
 	ERasterizerCullMode MeshCullMode,
