Fixed static assert for EMeshPass::NumBits
Added GetAllocatedSize to FMeshDrawCommand

diff --git a/Engine/Source/Runtime/Renderer/Private/MeshPassProcessor.cpp b/Engine/Source/Runtime/Renderer/Private/MeshPassProcessor.cpp
--- a/Engine/Source/Runtime/Renderer/Private/MeshPassProcessor.cpp
+++ b/Engine/Source/Runtime/Renderer/Private/MeshPassProcessor.cpp
@@ -878,8 +878,6 @@ FMeshPassProcessor::FMeshPassProcessor(const FScene* InScene, ERHIFeatureLevel::
 	, ViewIfDynamicMeshCommand(InViewIfDynamicMeshCommand)
 	, DrawListContext(InDrawListContext)
 {
-	static_assert(EMeshPass::Num <= EMeshPass::NumBits * 8, "EMeshPass::NumBits is too small to fit all mesh passes.");
-	static_assert(sizeof(FMeshPassMask::Data) * 8 >= EMeshPass::Num, "FMeshPassMask::Data is too small to fit all mesh passes.");
 }
 
 enum class EDrawingPolicyOverrideFlags
diff --git a/Engine/Source/Runtime/Renderer/Public/MeshPassProcessor.h b/Engine/Source/Runtime/Renderer/Public/MeshPassProcessor.h
--- a/Engine/Source/Runtime/Renderer/Public/MeshPassProcessor.h
+++ b/Engine/Source/Runtime/Renderer/Public/MeshPassProcessor.h
@@ -42,6 +42,7 @@ namespace EMeshPass
 		NumBits = 5,
 	};
 }
+static_assert(EMeshPass::Num <= (1 << EMeshPass::NumBits), "EMeshPass::Num will not fit in EMeshPass::NumBits");
 
 inline const TCHAR* GetMeshPassName(EMeshPass::Type MeshPass)
 {
@@ -91,6 +92,8 @@ public:
 	uint32 Data;
 };
 
+static_assert(sizeof(FMeshPassMask::Data) * 8 >= EMeshPass::Num, "FMeshPassMask::Data is too small to fit all mesh passes.");
+
 /** Uniquely represents a FGraphicsMinimalPipelineStateInitializer for fast compares. */
 class FGraphicsMinimalPipelineStateId
 {
@@ -240,6 +243,17 @@ public:
 	/** Returns whether this set of shader bindings can be merged into an instanced draw call with another. */
 	bool MatchesForDynamicInstancing(const FMeshDrawShaderBindings& Rhs) const;
 
+	SIZE_T GetAllocatedSize() const
+	{
+		SIZE_T Bytes = ShaderLayouts.GetAllocatedSize();
+		if (Size > ARRAY_COUNT(InlineStorage))
+		{
+			Bytes += Size;
+		}
+
+		return Bytes;
+	}
+
 private:
 
 	TArray<FMeshDrawShaderBindingsLayout, TInlineAllocator<2>> ShaderLayouts;
@@ -303,10 +317,10 @@ private:
 /** 
  * FMeshDrawCommand fully describes a mesh pass draw call, captured just above the RHI.  
 		FMeshDrawCommand should contain only data needed to draw.  For InitViews payloads, use FVisibleMeshDrawCommand.
- * FMeshDrawCommand's are cached at Primitive AddToScene time for vertex factories that support it (no per-frame or per-view shader binding changes).
+ * FMeshDrawCommands are cached at Primitive AddToScene time for vertex factories that support it (no per-frame or per-view shader binding changes).
  * Dynamic Instancing operates at the FMeshDrawCommand level for robustness.  
 		Adding per-command shader bindings will reduce the efficiency of Dynamic Instancing, but rendering will always be correct.
- * Any resources referenced by a command must be kept alive for the lifetime of the command.
+ * Any resources referenced by a command must be kept alive for the lifetime of the command.  FMeshDrawCommand is not responsible for lifetime management of resources.
 		For uniform buffers referenced by cached FMeshDrawCommand's, RHIUpdateUniformBuffer makes it possible to access per-frame data in the shader without changing bindings.
  */
 class FMeshDrawCommand
@@ -421,6 +435,10 @@ public:
 #endif
 	}
 
+	SIZE_T GetAllocatedSize() const
+	{
+		return ShaderBindings.GetAllocatedSize() + VertexStreams.GetAllocatedSize();
+	}
 
 #if MESH_DRAW_COMMAND_DEBUG_DATA
 private:
@@ -602,11 +620,11 @@ public:
 	int32 StateBucketId;
 
 	// Needed for easier debugging and faster removal of cached mesh draw commands.
-	EMeshPass::Type MeshPass : EMeshPass::NumBits;
+	EMeshPass::Type MeshPass : EMeshPass::NumBits + 1;
 
 	// Needed for view overrides
-	ERasterizerFillMode MeshFillMode : ERasterizerFillMode_NumBits;
-	ERasterizerCullMode MeshCullMode : ERasterizerCullMode_NumBits;
+	ERasterizerFillMode MeshFillMode : ERasterizerFillMode_NumBits + 1;
+	ERasterizerCullMode MeshCullMode : ERasterizerCullMode_NumBits + 1;
 };
 
 class FCachedPassMeshDrawList
