DR - Move Raytracing Geometries to the SceneInfo which for now does disable ShowFlags and specifically FSceneProxy::IsShown(View)

diff --git a/Engine/Plugins/FX/Niagara/Source/Niagara/Public/NiagaraComponent.h b/Engine/Plugins/FX/Niagara/Source/Niagara/Public/NiagaraComponent.h
--- a/Engine/Plugins/FX/Niagara/Source/Niagara/Public/NiagaraComponent.h
+++ b/Engine/Plugins/FX/Niagara/Source/Niagara/Public/NiagaraComponent.h
@@ -409,10 +409,6 @@ public:
 #if RHI_RAYTRACING
 	virtual void GetRayTracingGeometryInstances(TArray<FRayTracingGeometryInstanceCollection>& OutInstanceCollections) override;
 	virtual bool IsRayTracingRelevant() const override { return true; }
-	virtual bool IsRayTracingDrawRelevant(const FSceneView* View) const override
-	{
-		return ShouldRenderInMainPass() && View->Family->EngineShowFlags.Particles && IsShown(View);
-	}
 #endif
 
 private:
diff --git a/Engine/Source/Runtime/Engine/Private/SceneManagement.cpp b/Engine/Source/Runtime/Engine/Private/SceneManagement.cpp
--- a/Engine/Source/Runtime/Engine/Private/SceneManagement.cpp
+++ b/Engine/Source/Runtime/Engine/Private/SceneManagement.cpp
@@ -627,98 +627,6 @@ FLODMask ComputeLODForMeshes(const TArray<class FStaticMeshBatchRelevance>& Stat
 	return LODToRender;
 }
 
-FLODMask ComputeFastLODForMeshes(const TArray<float>& ScreenSizes, const FSceneView& View, const FVector4& Origin, float SphereRadius, int32 ForcedLODLevel, float& OutScreenRadiusSquared, float ScreenSizeScale, bool bDitheredLODTransition)
-{
-	FLODMask LODToRender;
-	const FSceneView& LODView = GetLODView(View);
-
-	const int32 NumMeshes = ScreenSizes.Num();
-
-	// Handle forced LOD level first
-	if (ForcedLODLevel >= 0)
-	{
-		OutScreenRadiusSquared = 0.0f;
-
-		int32 MinLOD = 127, MaxLOD = 0;
-		for (int32 MeshIndex = 0; MeshIndex < ScreenSizes.Num(); ++MeshIndex)
-		{
-			if (ScreenSizes[MeshIndex] == 0.0f)
-			{
-				continue;
-			}
-
-			MinLOD = FMath::Min(MinLOD, MeshIndex);
-			MaxLOD = FMath::Max(MaxLOD, MeshIndex);
-		}
-		LODToRender.SetLOD(FMath::Clamp(ForcedLODLevel, MinLOD, MaxLOD));
-	}
-	else if (LODView.Family->EngineShowFlags.LOD && NumMeshes)
-	{
-		if (bDitheredLODTransition)
-		{
-			for (int32 SampleIndex = 0; SampleIndex < 2; SampleIndex++)
-			{
-				int32 MinLODFound = INT_MAX;
-				bool bFoundLOD = false;
-				OutScreenRadiusSquared = ComputeTemporalLODBoundsScreenRadiusSquared(Origin, SphereRadius, LODView, SampleIndex);
-
-				for (int32 MeshIndex = NumMeshes - 1; MeshIndex >= 0; --MeshIndex)
-				{
-					if (ScreenSizes[MeshIndex] > 0.0f)
-					{
-						float MeshScreenSize = ScreenSizes[MeshIndex] * ScreenSizeScale;
-
-						if (FMath::Square(MeshScreenSize * 0.5f) >= OutScreenRadiusSquared)
-						{
-							LODToRender.SetLODSample(MeshIndex, SampleIndex);
-							bFoundLOD = true;
-							break;
-						}
-
-						MinLODFound = FMath::Min<int32>(MinLODFound, MeshIndex);
-					}
-				}
-				// If no LOD was found matching the screen size, use the lowest in the array instead of LOD 0, to handle non-zero MinLOD
-				if (!bFoundLOD)
-				{
-					LODToRender.SetLODSample(MinLODFound, SampleIndex);
-				}
-			}
-		}
-		else
-		{
-			int32 MinLODFound = INT_MAX;
-			bool bFoundLOD = false;
-			OutScreenRadiusSquared = ComputeBoundsScreenRadiusSquared(Origin, SphereRadius, LODView);
-
-			for (int32 MeshIndex = NumMeshes - 1; MeshIndex >= 0; --MeshIndex)
-			{
-				if (ScreenSizes[MeshIndex] == 0.0f)
-				{
-					continue;
-				}
-
-				float MeshScreenSize = ScreenSizes[MeshIndex] * ScreenSizeScale;
-
-				if (FMath::Square(MeshScreenSize * 0.5f) >= OutScreenRadiusSquared)
-				{
-					LODToRender.SetLOD(MeshIndex);
-					bFoundLOD = true;
-					break;
-				}
-
-				MinLODFound = FMath::Min<int32>(MinLODFound, MeshIndex);
-			}
-			// If no LOD was found matching the screen size, use the lowest in the array instead of LOD 0, to handle non-zero MinLOD
-			if (!bFoundLOD)
-			{
-				LODToRender.SetLOD(MinLODFound);
-			}
-		}
-	}
-	return LODToRender;
-}
-
 FMobileDirectionalLightShaderParameters::FMobileDirectionalLightShaderParameters()
 {
 	FMemory::Memzero(*this);
diff --git a/Engine/Source/Runtime/Engine/Private/SkeletalMesh.cpp b/Engine/Source/Runtime/Engine/Private/SkeletalMesh.cpp
--- a/Engine/Source/Runtime/Engine/Private/SkeletalMesh.cpp
+++ b/Engine/Source/Runtime/Engine/Private/SkeletalMesh.cpp
@@ -3967,7 +3967,7 @@ void FSkeletalMeshSceneProxy::GetDynamicElementsSection(const TArray<const FScen
 }
 
 #if RHI_RAYTRACING
-FRayTracingGeometryRHIRef FSkeletalMeshSceneProxy::GetRayTracingGeometryInstance(int LodLevel) const
+FRayTracingGeometryRHIRef FSkeletalMeshSceneProxy::GetDynamicRayTracingGeometryInstance() const
 {
 	if (MeshObject->GetRayTracingGeometry())
 	{
diff --git a/Engine/Source/Runtime/Engine/Private/SplineMeshSceneProxy.cpp b/Engine/Source/Runtime/Engine/Private/SplineMeshSceneProxy.cpp
--- a/Engine/Source/Runtime/Engine/Private/SplineMeshSceneProxy.cpp
+++ b/Engine/Source/Runtime/Engine/Private/SplineMeshSceneProxy.cpp
@@ -102,11 +102,3 @@ bool FSplineMeshSceneProxy::GetWireframeMeshElement(int32 LODIndex, int32 BatchI
 	}
 	return false;
 }
-
-#if RHI_RAYTRACING
-FRayTracingGeometryRHIRef FSplineMeshSceneProxy::GetRayTracingGeometryInstance(int LodLevel) const
-{
-	// #dxr_todo: return actual geometry when FSplineMeshVertexFactory is supported
-	return nullptr;
-}
-#endif // RHI_RAYTRACING
diff --git a/Engine/Source/Runtime/Engine/Private/SplineMeshSceneProxy.h b/Engine/Source/Runtime/Engine/Private/SplineMeshSceneProxy.h
--- a/Engine/Source/Runtime/Engine/Private/SplineMeshSceneProxy.h
+++ b/Engine/Source/Runtime/Engine/Private/SplineMeshSceneProxy.h
@@ -172,9 +172,7 @@ public:
 
 #if RHI_RAYTRACING
 	virtual bool IsRayTracingRelevant() const override final { return false; }
-	virtual bool IsRayTracingDrawRelevant(const FSceneView* View) const override final { return false; }
 	virtual bool IsRayTracingStaticRelevant() const override final { return false; }
-	FRayTracingGeometryRHIRef GetRayTracingGeometryInstance(int LodLevel) const override;
 #endif // RHI_RAYTRACING
 
 	virtual FPrimitiveViewRelevance GetViewRelevance(const FSceneView* View) const override
diff --git a/Engine/Source/Runtime/Engine/Private/StaticMeshRender.cpp b/Engine/Source/Runtime/Engine/Private/StaticMeshRender.cpp
--- a/Engine/Source/Runtime/Engine/Private/StaticMeshRender.cpp
+++ b/Engine/Source/Runtime/Engine/Private/StaticMeshRender.cpp
@@ -1306,13 +1306,6 @@ void FStaticMeshSceneProxy::GetDynamicMeshElements(const TArray<const FSceneView
 #endif // STATICMESH_ENABLE_DEBUG_RENDERING
 }
 
-#if RHI_RAYTRACING
-FRayTracingGeometryRHIRef FStaticMeshSceneProxy::GetRayTracingGeometryInstance(int LodLevel) const
-{
-	return RayTracingGeometries[LodLevel];
-}
-#endif // RHI_RAYTRACING
-
 void FStaticMeshSceneProxy::GetLCIs(FLCIArray& LCIs)
 {
 	for (int32 LODIndex = 0; LODIndex < LODs.Num(); ++LODIndex)
diff --git a/Engine/Source/Runtime/Engine/Public/PrimitiveSceneProxy.h b/Engine/Source/Runtime/Engine/Public/PrimitiveSceneProxy.h
--- a/Engine/Source/Runtime/Engine/Public/PrimitiveSceneProxy.h
+++ b/Engine/Source/Runtime/Engine/Public/PrimitiveSceneProxy.h
@@ -189,11 +189,14 @@ public:
 
 #if RHI_RAYTRACING
 	virtual bool IsRayTracingRelevant() const { return false; }
-	virtual bool IsRayTracingDrawRelevant(const FSceneView* View) const { return false; }
 	virtual bool IsRayTracingStaticRelevant() const { return false; }
-	/** Gathers static ray tracing instances from this proxy. */
-	ENGINE_API virtual FRayTracingGeometryRHIRef GetRayTracingGeometryInstance(int LodLevel) const { return nullptr; }
+	/** Gathers dynamic ray tracing instances from this proxy. */
+	ENGINE_API virtual FRayTracingGeometryRHIRef GetDynamicRayTracingGeometryInstance() const { return nullptr; }
 	virtual void GetRayTracingGeometryInstances(TArray<FRayTracingGeometryInstanceCollection>& OutInstanceCollections) {}
+	TArray<FRayTracingGeometryRHIRef>&& MoveRayTracingGeometries()
+	{
+		return static_cast<TArray<FRayTracingGeometryRHIRef>&&>(RayTracingGeometries);
+	}
 #endif // RHI_RAYTRACING
 
 	/** Collects occluder geometry for software occlusion culling */
@@ -729,6 +732,10 @@ protected:
 		OwnerName = InOwnerName;
 	}
 
+#if RHI_RAYTRACING
+	TArray<FRayTracingGeometryRHIRef> RayTracingGeometries;
+#endif
+
 private:
 #if !(UE_BUILD_SHIPPING || UE_BUILD_TEST)
 	FLinearColor WireframeColor;
@@ -1023,17 +1030,6 @@ private:
 	TArray<UMaterialInterface*> UsedMaterialsForVerification;
 #endif
 
-public:
-#if RHI_RAYTRACING
-	TArray<float> ScreenSizes;
-	struct FStaticMeshOrCommandIndex
-	{
-		int32 StaticMeshIndex;
-		int32 CommandIndex;
-	};
-	TArray<TArray<FStaticMeshOrCommandIndex, TInlineAllocator<2>>> RayTracingLodIndexToMeshDrawCommandIndicies;
-#endif
-private:
 	/**
 	 * Updates the primitive proxy's cached transforms, and calls OnUpdateTransform to notify it of the change.
 	 * Called in the thread that owns the proxy; game or rendering.
diff --git a/Engine/Source/Runtime/Engine/Public/SkeletalMeshTypes.h b/Engine/Source/Runtime/Engine/Public/SkeletalMeshTypes.h
--- a/Engine/Source/Runtime/Engine/Public/SkeletalMeshTypes.h
+++ b/Engine/Source/Runtime/Engine/Public/SkeletalMeshTypes.h
@@ -225,17 +225,12 @@ public:
 
 #if RHI_RAYTRACING
 	virtual bool IsRayTracingRelevant() const override final { return true; }
-	virtual bool IsRayTracingDrawRelevant(const FSceneView* View) const override
-	{
-		return ShouldRenderInMainPass() && View->Family->EngineShowFlags.SkeletalMeshes && IsShown(View);
-	}
 
 	virtual bool IsRayTracingStaticRelevant() const override
 	{
 		return bRenderStatic;
 	}
-
-	FRayTracingGeometryRHIRef GetRayTracingGeometryInstance(int LodLevel) const override;
+	FRayTracingGeometryRHIRef GetDynamicRayTracingGeometryInstance() const final override;
 #endif // RHI_RAYTRACING
 
 	virtual FPrimitiveViewRelevance GetViewRelevance(const FSceneView* View) const override;
diff --git a/Engine/Source/Runtime/Engine/Public/StaticMeshResources.h b/Engine/Source/Runtime/Engine/Public/StaticMeshResources.h
--- a/Engine/Source/Runtime/Engine/Public/StaticMeshResources.h
+++ b/Engine/Source/Runtime/Engine/Public/StaticMeshResources.h
@@ -603,18 +603,11 @@ public:
 
 #if RHI_RAYTRACING
 	virtual bool IsRayTracingRelevant() const override { return true; }
-	virtual bool IsRayTracingDrawRelevant(const FSceneView* View) const override
-	{
-		return ShouldRenderInMainPass() && View->Family->EngineShowFlags.StaticMeshes && IsShown(View);
-	}
-
 	virtual bool IsRayTracingStaticRelevant() const override 
 	{ 
 		const bool bAllowStaticLighting = FReadOnlyCVARCache::Get().bAllowStaticLighting;
 		return IsStaticPathAvailable() && !HasViewDependentDPG() && !(bAllowStaticLighting && HasStaticLighting() && !HasValidSettingsForStaticLighting());
 	}
-
-	FRayTracingGeometryRHIRef GetRayTracingGeometryInstance(int LodLevel) const override;
 #endif // RHI_RAYTRACING
 
 	virtual void GetLCIs(FLCIArray& LCIs) override;
@@ -699,9 +692,6 @@ protected:
 	FStaticMeshOccluderData* OccluderData;
 
 	TIndirectArray<FLODInfo> LODs;
-#if RHI_RAYTRACING
-	TArray<FRayTracingGeometryRHIRef> RayTracingGeometries;
-#endif
 
 	const FDistanceFieldVolumeData* DistanceFieldData;	
 
diff --git a/Engine/Source/Runtime/Renderer/Private/DeferredShadingRenderer.cpp b/Engine/Source/Runtime/Renderer/Private/DeferredShadingRenderer.cpp
--- a/Engine/Source/Runtime/Renderer/Private/DeferredShadingRenderer.cpp
+++ b/Engine/Source/Runtime/Renderer/Private/DeferredShadingRenderer.cpp
@@ -620,10 +620,9 @@ bool FDeferredShadingSceneRenderer::GatherRayTracingWorldInstances(FRHICommandLi
 				BroadIndex++;
 			}
 
-			FPrimitiveSceneProxy* SceneProxy = Scene->PrimitiveSceneProxies[PrimitiveIndex];
-			checkSlow(Scene->TypeOffsetTable[BroadIndex].PrimitiveSceneProxyType == SceneProxy->GetTypeHash());
+			FPrimitiveSceneInfo* SceneInfo = Scene->Primitives[PrimitiveIndex];
 
-			if (!SceneProxy->IsRayTracingRelevant())
+			if (!SceneInfo->bIsRayTracingRelevant)
 			{
 				//skip over unsupported SceneProxies (warning don't make IsRayTracingRelevant data dependent other than the vtable)
 				PrimitiveIndex = Scene->TypeOffsetTable[BroadIndex].Offset - 1;
@@ -639,14 +638,16 @@ bool FDeferredShadingSceneRenderer::GatherRayTracingWorldInstances(FRHICommandLi
 					continue;
 				}
 
-				if (View.bIsReflectionCapture && !SceneProxy->IsVisibleInReflectionCaptures())
+				if (View.bIsReflectionCapture && !SceneInfo->bIsVisibleInReflectionCaptures)
 				{
 					continue;
 				}
 
-				if (SceneProxy->IsRayTracingDrawRelevant(&View))
+				//#dxr_todo The Raytracing codepath does not support Showflags since data moved to the SceneInfo. 
+				//Touching the SceneProxy to determine this would simply cost too much
+				if (SceneInfo->bShouldRenderInMainPass)
 				{
-					if (SceneProxy->IsRayTracingStaticRelevant())
+					if (SceneInfo->bIsRayTracingStaticRelevant)
 					{
 						static const auto ICVarStaticMeshLODDistanceScale = IConsoleManager::Get().FindConsoleVariable(TEXT("r.StaticMeshLODDistanceScale"));
 						float LODScale = ICVarStaticMeshLODDistanceScale->GetFloat() * View.LODDistanceFactor;
@@ -656,22 +657,23 @@ bool FDeferredShadingSceneRenderer::GatherRayTracingWorldInstances(FRHICommandLi
 						FLODMask LODToRender;
 						float MeshScreenSizeSquared = 0;
 						int32 ForcedLODLevel = GetCVarForceLOD();
-						if (SceneProxy->IsUsingCustomLODRules())
+						if (SceneInfo->bIsUsingCustomLODRules)
 						{
+							FPrimitiveSceneProxy* SceneProxy = Scene->PrimitiveSceneProxies[PrimitiveIndex];
 							LODToRender = SceneProxy->GetCustomLOD(View, View.LODDistanceFactor, ForcedLODLevel, MeshScreenSizeSquared);
 						}
 						else
 						{
-							LODToRender = ComputeFastLODForMeshes(SceneProxy->ScreenSizes, View, Bounds.BoxSphereBounds.Origin, Bounds.BoxSphereBounds.SphereRadius, ForcedLODLevel, MeshScreenSizeSquared, LODScale, false);
+							LODToRender = ComputeLODForMeshes(SceneInfo->StaticMeshRelevances, View, Bounds.BoxSphereBounds.Origin, Bounds.BoxSphereBounds.SphereRadius, ForcedLODLevel, MeshScreenSizeSquared, LODScale, false);
 						}
 
-						FRayTracingGeometryRHIRef RayTracingGeometryInstance = SceneProxy->GetRayTracingGeometryInstance(LODToRender.GetRayTracedLOD());
+						FRayTracingGeometryRHIRef RayTracingGeometryInstance = SceneInfo->GetStaticRayTracingGeometryInstance(LODToRender.GetRayTracedLOD());
 						if (!RayTracingGeometryInstance.IsValid())
 						{
 							continue;
 						}
 
-						const auto& StaticMeshMdcIndices = SceneProxy->RayTracingLodIndexToMeshDrawCommandIndicies[LODToRender.GetRayTracedLOD()];
+						const auto& StaticMeshMdcIndices = SceneInfo->RayTracingLodIndexToMeshDrawCommandIndicies[LODToRender.GetRayTracedLOD()];
 						for (int StaticMeshIndex = 0; StaticMeshIndex < StaticMeshMdcIndices.Num(); StaticMeshIndex++)
 						{
 							const int32 CommandIndex = StaticMeshMdcIndices[StaticMeshIndex].CommandIndex;
@@ -692,8 +694,7 @@ bool FDeferredShadingSceneRenderer::GatherRayTracingWorldInstances(FRHICommandLi
 							}
 							else
 							{
-								const FPrimitiveSceneInfo* PrimitiveInfo = Scene->Primitives[PrimitiveIndex];
-								const FStaticMeshBatch& StaticMesh = PrimitiveInfo->StaticMeshes[StaticMeshMdcIndices[StaticMeshIndex].StaticMeshIndex];
+								const FStaticMeshBatch& StaticMesh = SceneInfo->StaticMeshes[StaticMeshMdcIndices[StaticMeshIndex].StaticMeshIndex];
 
 								uint64 BatchVisibility = 1;
 								if (StaticMesh.bRequiresPerElementVisibility)
@@ -707,6 +708,8 @@ bool FDeferredShadingSceneRenderer::GatherRayTracingWorldInstances(FRHICommandLi
 									View.RaytraycingVisibleMeshDrawCommands
 								);
 								FRayTracingMeshProcessor RayTracingMeshProcessor(Scene, View.GetFeatureLevel(), &View, &DynamicPassMeshDrawListContext);
+
+								FPrimitiveSceneProxy* SceneProxy = Scene->PrimitiveSceneProxies[PrimitiveIndex];
 								RayTracingMeshProcessor.AddMeshBatch(StaticMesh, BatchVisibility, SceneProxy);
 							}
 						}
@@ -735,9 +738,10 @@ bool FDeferredShadingSceneRenderer::GatherRayTracingWorldInstances(FRHICommandLi
 			if (RayTracedMeshElementsMask != 0)
 			{
 				{
-					FRayTracingGeometryRHIRef RayTracingGeometryInstance = SceneProxy->GetRayTracingGeometryInstance(0);
+					FPrimitiveSceneProxy* SceneProxy = Scene->PrimitiveSceneProxies[PrimitiveIndex];
+					FRayTracingGeometryRHIRef RayTracingGeometryInstance = SceneProxy->GetDynamicRayTracingGeometryInstance();
 					if (RayTracingGeometryInstance.IsValid())
-					{
+					{					
 						RayTracingCollector.SetPrimitive(SceneProxy, FHitProxyId());
 						SceneProxy->GetDynamicMeshElements(ViewFamily.Views, ViewFamily, RayTracedMeshElementsMask, RayTracingCollector);
 
@@ -782,6 +786,7 @@ bool FDeferredShadingSceneRenderer::GatherRayTracingWorldInstances(FRHICommandLi
 				}
 
 				{
+					FPrimitiveSceneProxy* SceneProxy = Scene->PrimitiveSceneProxies[PrimitiveIndex];
 					TArray<FRayTracingGeometryInstanceCollection> InstanceCollections;
 					SceneProxy->GetRayTracingGeometryInstances(InstanceCollections);
 					if (InstanceCollections.Num() > 0)
diff --git a/Engine/Source/Runtime/Renderer/Private/PrimitiveSceneInfo.cpp b/Engine/Source/Runtime/Renderer/Private/PrimitiveSceneInfo.cpp
--- a/Engine/Source/Runtime/Renderer/Private/PrimitiveSceneInfo.cpp
+++ b/Engine/Source/Runtime/Renderer/Private/PrimitiveSceneInfo.cpp
@@ -113,6 +113,12 @@ FPrimitiveSceneInfo::FPrimitiveSceneInfo(UPrimitiveComponent* InComponent,FScene
 	NumMobileMovablePointLights(0),
 	bIsUsingCustomLODRules(Proxy->IsUsingCustomLODRules()),
 	bIsUsingCustomWholeSceneShadowLODRules(Proxy->IsUsingCustomWholeSceneShadowLODRules()),
+#if RHI_RAYTRACING
+	bShouldRenderInMainPass(InComponent->SceneProxy->ShouldRenderInMainPass()),
+	bIsVisibleInReflectionCaptures(InComponent->SceneProxy->IsVisibleInReflectionCaptures()),
+	bIsRayTracingRelevant(InComponent->SceneProxy->IsRayTracingRelevant()),
+	bIsRayTracingStaticRelevant(InComponent->SceneProxy->IsRayTracingStaticRelevant()),
+#endif
 	PackedIndex(INDEX_NONE),
 	ComponentForDebuggingOnly(InComponent),
 	bNeedsStaticMeshUpdate(false),
@@ -151,6 +157,10 @@ FPrimitiveSceneInfo::FPrimitiveSceneInfo(UPrimitiveComponent* InComponent,FScene
 	}
 
 	FMemory::Memzero(CachedReflectionCaptureProxies);
+
+#if RHI_RAYTRACING
+	RayTracingGeometries = InComponent->SceneProxy->MoveRayTracingGeometries();
+#endif
 }
 
 FPrimitiveSceneInfo::~FPrimitiveSceneInfo()
@@ -160,6 +170,20 @@ FPrimitiveSceneInfo::~FPrimitiveSceneInfo()
 	RemoveCachedMeshDrawCommands();
 }
 
+#if RHI_RAYTRACING
+FRayTracingGeometryRHIRef FPrimitiveSceneInfo::GetStaticRayTracingGeometryInstance(int LodLevel)
+{
+	if (RayTracingGeometries.Num() > LodLevel)
+	{
+		return RayTracingGeometries[LodLevel];
+	}
+	else
+	{
+		return nullptr;
+	}
+}
+#endif
+
 void FPrimitiveSceneInfo::CacheMeshDrawCommands(FRHICommandListImmediate& RHICmdList)
 {
 	check(StaticMeshCommandInfos.Num() == 0);
@@ -325,24 +349,13 @@ void FPrimitiveSceneInfo::AddStaticMeshes(FRHICommandListImmediate& RHICmdList,
 
 		if (StaticMeshes.Num() > 0)
 		{
-			Proxy->ScreenSizes.AddDefaulted(MaxLOD + 1);
-			Proxy->RayTracingLodIndexToMeshDrawCommandIndicies.AddDefaulted(MaxLOD + 1);
+			RayTracingLodIndexToMeshDrawCommandIndicies.AddDefaulted(MaxLOD + 1);
 
 			for (int32 MeshIndex = 0; MeshIndex < StaticMeshes.Num(); MeshIndex++)
 			{
 				FStaticMeshBatchRelevance& MeshRelevance = StaticMeshRelevances[MeshIndex];
 				FStaticMeshBatch& Mesh = StaticMeshes[MeshIndex];
 
-				if (Proxy->ScreenSizes[MeshRelevance.LODIndex] != 0.0f)
-				{
-					check(Proxy->ScreenSizes[MeshRelevance.LODIndex] == MeshRelevance.ScreenSize);
-				}
-				else
-				{
-					check(MeshRelevance.ScreenSize != 0.0f);
-					Proxy->ScreenSizes[MeshRelevance.LODIndex] = MeshRelevance.ScreenSize;
-				}
-
 				if (Mesh.Elements.Num() > 0)
 				{
 					const int32 RayTracingStaticMeshCommandInfoIndex = MeshRelevance.GetStaticMeshCommandInfoIndex(EMeshPass::RayTracing);
@@ -350,11 +363,11 @@ void FPrimitiveSceneInfo::AddStaticMeshes(FRHICommandListImmediate& RHICmdList,
 					if (SupportsCachingMeshDrawCommands(Mesh.VertexFactory, Proxy) && RayTracingStaticMeshCommandInfoIndex != -1)
 					{
 						const int32 CommandIndex = StaticMeshCommandInfos[RayTracingStaticMeshCommandInfoIndex].CommandIndex;
-						Proxy->RayTracingLodIndexToMeshDrawCommandIndicies[Mesh.LODIndex].Add({ MeshIndex, CommandIndex });
+						RayTracingLodIndexToMeshDrawCommandIndicies[Mesh.LODIndex].Add({ MeshIndex, CommandIndex });
 					}
 					else
 					{
-						Proxy->RayTracingLodIndexToMeshDrawCommandIndicies[Mesh.LODIndex].Add({ MeshIndex, -1 });
+						RayTracingLodIndexToMeshDrawCommandIndicies[Mesh.LODIndex].Add({ MeshIndex, -1 });
 					}
 				}
 			}
diff --git a/Engine/Source/Runtime/Renderer/Public/PrimitiveSceneInfo.h b/Engine/Source/Runtime/Renderer/Public/PrimitiveSceneInfo.h
--- a/Engine/Source/Runtime/Renderer/Public/PrimitiveSceneInfo.h
+++ b/Engine/Source/Runtime/Renderer/Public/PrimitiveSceneInfo.h
@@ -282,6 +282,20 @@ public:
 	/** This indicate that we should call the GetCustomWholeSceneShadowLOD function on the proxy instead of the generic implementation. */
 	bool bIsUsingCustomWholeSceneShadowLODRules : 1;
 
+#if RHI_RAYTRACING
+	bool bShouldRenderInMainPass : 1;
+	bool bIsVisibleInReflectionCaptures : 1;
+	bool bIsRayTracingRelevant : 1;
+	bool bIsRayTracingStaticRelevant : 1;
+
+	struct FStaticMeshOrCommandIndex
+	{
+		int32 StaticMeshIndex;
+		int32 CommandIndex;
+	};
+	TArray<TArray<FStaticMeshOrCommandIndex, TInlineAllocator<2>>> RayTracingLodIndexToMeshDrawCommandIndicies;
+#endif
+
 	/** Initialization constructor. */
 	FPrimitiveSceneInfo(UPrimitiveComponent* InPrimitive,FScene* InScene);
 
@@ -421,6 +435,10 @@ public:
 	/** Cache per-primitive reflection captures used for mobile/forward rendering */
 	void CacheReflectionCaptures();
 
+#if RHI_RAYTRACING
+	RENDERER_API FRayTracingGeometryRHIRef GetStaticRayTracingGeometryInstance(int LodLevel);
+#endif
+
 private:
 
 	/** Let FScene have direct access to the Id. */
@@ -465,6 +483,10 @@ private:
 
 	/** Removes cached mesh draw commands for all meshes. */
 	void RemoveCachedMeshDrawCommands();
+
+#if RHI_RAYTRACING
+	TArray<FRayTracingGeometryRHIRef> RayTracingGeometries;
+#endif
 };
 
 /** Defines how the primitive is stored in the scene's primitive octree. */
