Re-implement instanced stereo rendering for the HLR

diff --git a/Engine/Shaders/Private/BasePassPixelShader.usf b/Engine/Shaders/Private/BasePassPixelShader.usf
--- a/Engine/Shaders/Private/BasePassPixelShader.usf
+++ b/Engine/Shaders/Private/BasePassPixelShader.usf
@@ -728,7 +728,7 @@ void FPixelShaderInOut_MainPS(
 	inout FPixelShaderOut Out)
 {
 #if INSTANCED_STEREO
-	ResolvedView = ResolveView(GetEyeIndex(Interpolants.PackedEyeIndex));
+	ResolvedView = ResolveView(Interpolants.EyeIndex);
 #else
 	ResolvedView = ResolveView();
 #endif
@@ -1004,7 +1004,7 @@ void FPixelShaderInOut_MainPS(
 			uint GridIndex = 0;
 
 			#if INSTANCED_STEREO
-				const uint EyeIndex = GetEyeIndex(Interpolants.PackedEyeIndex);
+				const uint EyeIndex = Interpolants.EyeIndex;
 			#else
 				const uint EyeIndex = 0;
 			#endif
diff --git a/Engine/Shaders/Private/BasePassVertexShader.usf b/Engine/Shaders/Private/BasePassVertexShader.usf
--- a/Engine/Shaders/Private/BasePassVertexShader.usf
+++ b/Engine/Shaders/Private/BasePassVertexShader.usf
@@ -37,7 +37,7 @@ void Main(
 	)
 {	
 #if INSTANCED_STEREO
-	const uint EyeIndex = VertexFactoryGetEyeIndex(InstanceId);
+	const uint EyeIndex = GetEyeIndex(InstanceId);
 	ResolvedView = ResolveView(EyeIndex);
 	#if !MULTI_VIEW
 		OutClipDistance = 0.0;
@@ -87,7 +87,7 @@ void Main(
 
 	#if INSTANCED_STEREO && !MULTI_VIEW
 	BRANCH
-	if (bIsInstancedStereo)
+	if (IsInstancedStereo())
 	{
 		// Clip at the center of the screen
 		OutClipDistance = dot(Output.Position, EyeClipEdge[EyeIndex]);
@@ -112,9 +112,9 @@ void Main(
 
 #if INSTANCED_STEREO
 	#if USING_TESSELLATION	
-		Output.FactoryInterpolants.InterpolantsVSToPS.PackedEyeIndex = PackEyeIndex(EyeIndex, bIsInstancedStereo);
+		Output.FactoryInterpolants.InterpolantsVSToPS.EyeIndex = EyeIndex;
 	#else
-		Output.FactoryInterpolants.PackedEyeIndex = PackEyeIndex(EyeIndex, bIsInstancedStereo);
+		Output.FactoryInterpolants.EyeIndex = EyeIndex;
 	#endif
 #endif
 
diff --git a/Engine/Shaders/Private/DebugViewModeCommon.ush b/Engine/Shaders/Private/DebugViewModeCommon.ush
--- a/Engine/Shaders/Private/DebugViewModeCommon.ush
+++ b/Engine/Shaders/Private/DebugViewModeCommon.ush
@@ -20,7 +20,7 @@ struct FDebugVSToPS
 	float3 TangentToWorld1 : TEXCOORD4;
 	float3 TangentToWorld2 : TEXCOORD5;
 #if INSTANCED_STEREO
-	nointerpolation uint PackedEyeIndex : PACKED_EYE_INDEX;
+	nointerpolation uint EyeIndex : EYE_INDEX;
 #endif
 };
 
@@ -34,7 +34,7 @@ struct FDebugPSIn
 	float3 TangentToWorld1 : TEXCOORD4;
 	float3 TangentToWorld2 : TEXCOORD5;
 #if INSTANCED_STEREO
-	nointerpolation uint PackedEyeIndex : PACKED_EYE_INDEX;
+	nointerpolation uint EyeIndex : EYE_INDEX;
 #endif
 	uint SvPrimitiveID : SV_PrimitiveID;
 };
diff --git a/Engine/Shaders/Private/DebugViewModeVertexShader.usf b/Engine/Shaders/Private/DebugViewModeVertexShader.usf
--- a/Engine/Shaders/Private/DebugViewModeVertexShader.usf
+++ b/Engine/Shaders/Private/DebugViewModeVertexShader.usf
@@ -114,7 +114,7 @@ void Main(
 	)
 {
 #if INSTANCED_STEREO
-	const uint EyeIndex = VertexFactoryGetEyeIndex(InstanceId);
+	const uint EyeIndex = GetEyeIndex(InstanceId);
 	ResolvedView = ResolveView(EyeIndex);
 	#if !MULTI_VIEW
 		OutClipDistance = 0.0;
@@ -152,7 +152,7 @@ void Main(
 
 	#if INSTANCED_STEREO && !MULTI_VIEW
 	BRANCH 
-	if (bIsInstancedStereo)  
+	if (IsInstancedStereo())  
 	{
 		// Clip at the center of the screen
 		OutClipDistance = dot(Output.Position, EyeClipEdge[EyeIndex]);
@@ -194,9 +194,9 @@ void Main(
 
 	#if INSTANCED_STEREO
 		#if USING_TESSELLATION	
-			Output.FactoryInterpolants.InterpolantsVSToPS.PackedEyeIndex = PackEyeIndex(EyeIndex, bIsInstancedStereo);
+			Output.FactoryInterpolants.InterpolantsVSToPS.EyeIndex = EyeIndex;
 		#else
-			Output.PackedEyeIndex = PackEyeIndex(EyeIndex, bIsInstancedStereo);
+			Output.EyeIndex = EyeIndex;
 		#endif
 	#endif
 
diff --git a/Engine/Shaders/Private/DepthOnlyPixelShader.usf b/Engine/Shaders/Private/DepthOnlyPixelShader.usf
--- a/Engine/Shaders/Private/DepthOnlyPixelShader.usf
+++ b/Engine/Shaders/Private/DepthOnlyPixelShader.usf
@@ -32,7 +32,7 @@ void Main(
 {
 #if !MATERIALBLENDING_SOLID || OUTPUT_PIXEL_DEPTH_OFFSET
 	#if INSTANCED_STEREO
-		ResolvedView = ResolveView(GetEyeIndex(FactoryInterpolants.PackedEyeIndex));
+		ResolvedView = ResolveView(FactoryInterpolants.EyeIndex);
 	#else
 		ResolvedView = ResolveView();
 	#endif
diff --git a/Engine/Shaders/Private/DepthOnlyVertexShader.usf b/Engine/Shaders/Private/DepthOnlyVertexShader.usf
--- a/Engine/Shaders/Private/DepthOnlyVertexShader.usf
+++ b/Engine/Shaders/Private/DepthOnlyVertexShader.usf
@@ -8,10 +8,6 @@
 #include "/Engine/Generated/Material.ush"
 #include "/Engine/Generated/VertexFactory.ush"
 
-#if INSTANCED_STEREO
-bool bIsInstancedStereoEmulated;
-#endif
-
 #define IOSZBIASWORKAROUND (IOS && (ES2_PROFILE || ES3_1_PROFILE))
 #define USE_RAW_WORLD_POSITION ((!MATERIALBLENDING_SOLID || OUTPUT_PIXEL_DEPTH_OFFSET) && USE_WORLD_POSITION_EXCLUDING_SHADER_OFFSETS)
 
@@ -111,17 +107,7 @@ void Main(
 	)
 {
 #if INSTANCED_STEREO
-	uint EyeIndex;
-	BRANCH
-	if (bIsInstancedStereoEmulated)
-	{
-		EyeIndex = View.StereoPassIndex;
-	}
-	else
-	{
-		EyeIndex = VertexFactoryGetEyeIndex(InstanceId);
-	}
-
+	uint EyeIndex = GetEyeIndex(InstanceId);
 	ResolvedView = ResolveView(EyeIndex);
 	#if !MULTI_VIEW
 		OutClipDistance = 0.0;
@@ -164,7 +150,7 @@ void Main(
 
 	#if INSTANCED_STEREO && !MULTI_VIEW
 	BRANCH 
-	if (bIsInstancedStereo || bIsInstancedStereoEmulated)  
+	if (IsInstancedStereo())  
 	{
 		// Clip at the center of the screen
 		OutClipDistance = dot(Output.Position, EyeClipEdge[EyeIndex]);
@@ -201,9 +187,9 @@ void Main(
 
 #if INSTANCED_STEREO
 	#if USING_TESSELLATION	
-		Output.FactoryInterpolants.InterpolantsVSToPS.PackedEyeIndex = PackEyeIndex(EyeIndex, bIsInstancedStereo);
+		Output.FactoryInterpolants.InterpolantsVSToPS.EyeIndex = EyeIndex;
 	#else
-		Output.FactoryInterpolants.PackedEyeIndex = PackEyeIndex(EyeIndex, bIsInstancedStereo);
+		Output.FactoryInterpolants.EyeIndex = EyeIndex;
 	#endif
 #endif
 
diff --git a/Engine/Shaders/Private/FlatTessellation.ush b/Engine/Shaders/Private/FlatTessellation.ush
--- a/Engine/Shaders/Private/FlatTessellation.ush
+++ b/Engine/Shaders/Private/FlatTessellation.ush
@@ -49,7 +49,7 @@ struct FFlatTessellationHSToDS
 		FHullShaderConstantOutput O = (FHullShaderConstantOutput)0;
 
 		#if INSTANCED_STEREO
-			ResolvedView = ResolveView(GetEyeIndex(I[0].PassSpecificData.FactoryInterpolants.InterpolantsVSToPS.PackedEyeIndex));
+			ResolvedView = ResolveView(I[0].PassSpecificData.FactoryInterpolants.InterpolantsVSToPS.EyeIndex);
 		#else
 			ResolvedView = ResolveView();
 		#endif
@@ -91,7 +91,7 @@ struct FFlatTessellationHSToDS
 	FFlatTessellationHSToDS MainHull( InputPatch<FPassSpecificVSToDS, IN_CONTROL_POINTS> I, uint ControlPointID : SV_OutputControlPointID )
 	{
 		#if INSTANCED_STEREO
-			ResolvedView = ResolveView(GetEyeIndex(I[0].FactoryInterpolants.InterpolantsVSToPS.PackedEyeIndex));
+			ResolvedView = ResolveView(I[0].FactoryInterpolants.InterpolantsVSToPS.EyeIndex);
 		#else
 			ResolvedView = ResolveView();
 		#endif
@@ -144,7 +144,7 @@ struct FFlatTessellationHSToDS
 		)
 	{
 		#if INSTANCED_STEREO
-			ResolvedView = ResolveView(GetEyeIndex(I[0].PassSpecificData.FactoryInterpolants.InterpolantsVSToPS.PackedEyeIndex));
+			ResolvedView = ResolveView(I[0].PassSpecificData.FactoryInterpolants.InterpolantsVSToPS.EyeIndex);
 			#if !MULTI_VIEW
 				OutClipDistance = 0.0;
 			#endif
@@ -168,7 +168,7 @@ struct FFlatTessellationHSToDS
 							
 		// Forward the eye index along	
 		#if INSTANCED_STEREO
-		Interp.FactoryInterpolants.InterpolantsVSToPS.PackedEyeIndex = I[0].PassSpecificData.FactoryInterpolants.InterpolantsVSToPS.PackedEyeIndex;
+		Interp.FactoryInterpolants.InterpolantsVSToPS.EyeIndex = I[0].PassSpecificData.FactoryInterpolants.InterpolantsVSToPS.EyeIndex;
 		#endif
 																		
 		FMaterialTessellationParameters MaterialParameters = GetMaterialTessellationParameters(Interp.FactoryInterpolants, WorldPosition.xyz);
@@ -221,9 +221,9 @@ struct FFlatTessellationHSToDS
 
 		#if INSTANCED_STEREO && !MULTI_VIEW
 		BRANCH
-		if (IsInstancedStereo(I[0].PassSpecificData.FactoryInterpolants.InterpolantsVSToPS.PackedEyeIndex))
+		if (IsInstancedStereo())
 		{
-			const uint EyeIndex = GetEyeIndex(I[0].PassSpecificData.FactoryInterpolants.InterpolantsVSToPS.PackedEyeIndex);
+			const uint EyeIndex = I[0].PassSpecificData.FactoryInterpolants.InterpolantsVSToPS.EyeIndex;
 
 			// Clip at the center of the screen
 			OutClipDistance = dot(Result.Position, EyeClipEdge[EyeIndex]);
diff --git a/Engine/Shaders/Private/GeometryCacheVertexFactory.ush b/Engine/Shaders/Private/GeometryCacheVertexFactory.ush
--- a/Engine/Shaders/Private/GeometryCacheVertexFactory.ush
+++ b/Engine/Shaders/Private/GeometryCacheVertexFactory.ush
@@ -87,7 +87,7 @@ struct FVertexFactoryInterpolantsVSToPS
 #endif
 
 #if INSTANCED_STEREO
-	nointerpolation uint PackedEyeIndex : PACKED_EYE_INDEX;
+	nointerpolation uint EyeIndex : PACKED_EYE_INDEX;
 #endif
 };
 
@@ -267,7 +267,7 @@ FVertexFactoryInterpolantsVSToPS VertexFactoryGetInterpolantsVSToPS(FVertexFacto
 	Interpolants.Color = Intermediates.Color;
 
 #if INSTANCED_STEREO
-	Interpolants.PackedEyeIndex = 0;
+	Interpolants.EyeIndex = 0;
 #endif
 
 	return Interpolants;
@@ -383,13 +383,6 @@ float3x3 VertexFactoryGetTangentToLocalDS(FVertexFactoryInterpolantsVSToDS Inter
 
 #endif // #if USING_TESSELLATION
 
-#if INSTANCED_STEREO
-uint VertexFactoryGetEyeIndex(uint InstanceId)
-{
-	return InstanceId & 1;
-}
-#endif
-
 float4 VertexFactoryGetTranslatedPrimitiveVolumeBounds(FVertexFactoryInterpolantsVSToPS Interpolants)
 {
 	return 0;
diff --git a/Engine/Shaders/Private/GpuSkinVertexFactory.ush b/Engine/Shaders/Private/GpuSkinVertexFactory.ush
--- a/Engine/Shaders/Private/GpuSkinVertexFactory.ush
+++ b/Engine/Shaders/Private/GpuSkinVertexFactory.ush
@@ -150,7 +150,7 @@ struct FVertexFactoryInterpolantsVSToPS
 #endif
 
 #if INSTANCED_STEREO
-	nointerpolation uint PackedEyeIndex : PACKED_EYE_INDEX;
+	nointerpolation uint EyeIndex : PACKED_EYE_INDEX;
 #endif
 };
 
@@ -614,7 +614,7 @@ FVertexFactoryInterpolantsVSToPS VertexFactoryGetInterpolantsVSToPS(FVertexFacto
 #endif
 
 #if INSTANCED_STEREO
-	Interpolants.PackedEyeIndex = 0;
+	Interpolants.EyeIndex = 0;
 #endif
 
 	return Interpolants;
@@ -732,13 +732,6 @@ float3x3 VertexFactoryGetTangentToLocalDS(FVertexFactoryInterpolantsVSToDS Inter
 
 #endif // #if USING_TESSELLATION
 
-#if INSTANCED_STEREO
-uint VertexFactoryGetEyeIndex(uint InstanceId)
-{
-	return InstanceId & 1;
-}
-#endif
-
 float4 VertexFactoryGetTranslatedPrimitiveVolumeBounds(FVertexFactoryInterpolantsVSToPS Interpolants)
 {
 	return 0;
diff --git a/Engine/Shaders/Private/InstancedStereo.ush b/Engine/Shaders/Private/InstancedStereo.ush
--- a/Engine/Shaders/Private/InstancedStereo.ush
+++ b/Engine/Shaders/Private/InstancedStereo.ush
@@ -14,25 +14,6 @@
 // see GenerateInstancedStereoCode()
 #include "/Engine/Generated/GeneratedInstancedStereo.ush"
 
-// Packing the eye index and instanced stereo flag into a single uint to save an interpolator
-uint PackEyeIndex(uint EyeIndex, bool bIsInstancedStereo)
-{
-	uint Result = 0;
-	Result |= (uint)bIsInstancedStereo << 31;
-	Result |= EyeIndex;
-	return Result;
-}
-
-uint GetEyeIndex(uint PackedIndex)
-{
-	return PackedIndex & 0x1;
-}
-
-bool IsInstancedStereo(uint PackedIndex)
-{
-	return (PackedIndex >> 1) > 0;
-}
-
 static ViewState ResolvedView;
 
 ViewState ResolveView()
@@ -43,7 +24,6 @@ ViewState ResolveView()
 #if INSTANCED_STEREO
 static const float EyeOffsetScale[2] = { -1.0, 1.0 };
 static const float4 EyeClipEdge[2] = { float4(-1.0, 0.0, 0.0, 1.0), float4(1.0, 0.0, 0.0, 1.0) };
-bool bIsInstancedStereo;
 #endif
 
 // Generated for Metal in GeneratedInstancedStereo.usf
@@ -60,3 +40,30 @@ ViewState ResolveView(uint ViewIndex)
 	}
 }
 #endif
+
+bool IsInstancedStereo()
+{
+#if INSTANCED_STEREO
+	return InstancedView_StereoPassIndex > 0;
+#else
+	return false;
+#endif
+}
+
+uint GetEyeIndex(uint InstanceId)
+{
+#if INSTANCED_STEREO
+	return IsInstancedStereo() ? InstanceId & 1 : 0;
+#else
+	return 0;
+#endif
+}
+
+uint GetInstanceId(uint InstanceId)
+{
+#if INSTANCED_STEREO
+	return IsInstancedStereo() ? InstanceId / 2 : InstanceId;
+#else
+	return InstanceId;
+#endif
+}
diff --git a/Engine/Shaders/Private/LandscapeVertexFactory.ush b/Engine/Shaders/Private/LandscapeVertexFactory.ush
--- a/Engine/Shaders/Private/LandscapeVertexFactory.ush
+++ b/Engine/Shaders/Private/LandscapeVertexFactory.ush
@@ -109,7 +109,7 @@ struct FVertexFactoryInterpolantsVSToPS
 #endif
 
 #if INSTANCED_STEREO
-	nointerpolation uint PackedEyeIndex : PACKED_EYE_INDEX;
+	nointerpolation uint EyeIndex : PACKED_EYE_INDEX;
 #endif
 };
 
@@ -795,7 +795,7 @@ FVertexFactoryInterpolantsVSToPS VertexFactoryGetInterpolantsVSToPS(FVertexFacto
 #endif
 
 #if INSTANCED_STEREO
-	Interpolants.PackedEyeIndex = 0;
+	Interpolants.EyeIndex = 0;
 #endif
 	
 	return Interpolants;
@@ -877,13 +877,6 @@ void VertexFactoryUnpackInterpolants(FVertexFactoryInterpolantsVSToPS Interpolan
 }
 #endif // NUM_VF_PACKED_INTERPOLANTS > 0
 
-#if INSTANCED_STEREO
-uint VertexFactoryGetEyeIndex(uint InstanceId)
-{
-	return InstanceId & 1;
-}
-#endif
-
 float4 VertexFactoryGetTranslatedPrimitiveVolumeBounds(FVertexFactoryInterpolantsVSToPS Interpolants)
 {
 	return 0;
diff --git a/Engine/Shaders/Private/LocalVertexFactory.ush b/Engine/Shaders/Private/LocalVertexFactory.ush
--- a/Engine/Shaders/Private/LocalVertexFactory.ush
+++ b/Engine/Shaders/Private/LocalVertexFactory.ush
@@ -15,12 +15,6 @@
 #include "/Engine/Generated/UniformBuffers/PrecomputedLightingBuffer.ush"
 
 #if USE_INSTANCING
-
-	// Instanced stereo sets the eye index explicitly for instanced geometry
-	#if INSTANCED_STEREO
-		uint InstancedEyeIndex;
-	#endif
-
 	#ifndef USE_DITHERED_LOD_TRANSITION_FOR_INSTANCED
 		#error "USE_DITHERED_LOD_TRANSITION_FOR_INSTANCED should have been defined"
 	#endif
@@ -272,11 +266,12 @@ float4x4 GetInstanceTransform(FVertexFactoryIntermediates Intermediates)
 float4x4 GetInstanceTransform(FPositionOnlyVertexFactoryInput Input)
 {
 	#if !USE_INSTANCING_EMULATED && MANUAL_VERTEX_FETCH
+	uint InstanceId = GetInstanceId(Input.InstanceId);
 	return float4x4(
-		float4(VertexFetch_InstanceTransformBuffer[3 * (Input.InstanceId + InstanceOffset) + 0].xyz, 0.0f),
-		float4(VertexFetch_InstanceTransformBuffer[3 * (Input.InstanceId + InstanceOffset) + 1].xyz, 0.0f),
-		float4(VertexFetch_InstanceTransformBuffer[3 * (Input.InstanceId + InstanceOffset) + 2].xyz, 0.0f),
-		float4(VertexFetch_InstanceOriginBuffer[(Input.InstanceId + InstanceOffset)].xyz, 1.0f));
+		float4(VertexFetch_InstanceTransformBuffer[3 * (InstanceId + InstanceOffset) + 0].xyz, 0.0f),
+		float4(VertexFetch_InstanceTransformBuffer[3 * (InstanceId + InstanceOffset) + 1].xyz, 0.0f),
+		float4(VertexFetch_InstanceTransformBuffer[3 * (InstanceId + InstanceOffset) + 2].xyz, 0.0f),
+		float4(VertexFetch_InstanceOriginBuffer[(InstanceId + InstanceOffset)].xyz, 1.0f));
 	#elif !USE_INSTANCING_EMULATED
 	return float4x4(
 		float4(Input.InstanceTransform1.xyz, 0.0f),
@@ -679,11 +674,12 @@ FVertexFactoryIntermediates GetVertexFactoryIntermediates(FVertexFactoryInput In
 #endif
 
 #if USE_INSTANCING && !USE_INSTANCING_EMULATED && MANUAL_VERTEX_FETCH
-	Intermediates.InstanceTransform1 = VertexFetch_InstanceTransformBuffer[3 * (Input.InstanceId + InstanceOffset) + 0];
-	Intermediates.InstanceTransform2 = VertexFetch_InstanceTransformBuffer[3 * (Input.InstanceId + InstanceOffset) + 1];
-	Intermediates.InstanceTransform3 = VertexFetch_InstanceTransformBuffer[3 * (Input.InstanceId + InstanceOffset) + 2];
-	Intermediates.InstanceOrigin = VertexFetch_InstanceOriginBuffer[(Input.InstanceId + InstanceOffset)];
-	Intermediates.InstanceLightmapAndShadowMapUVBias = VertexFetch_InstanceLightmapBuffer[(Input.InstanceId + InstanceOffset)];
+	uint InstanceId = GetInstanceId(Input.InstanceId);
+	Intermediates.InstanceTransform1 = VertexFetch_InstanceTransformBuffer[3 * (InstanceId + InstanceOffset) + 0];
+	Intermediates.InstanceTransform2 = VertexFetch_InstanceTransformBuffer[3 * (InstanceId + InstanceOffset) + 1];
+	Intermediates.InstanceTransform3 = VertexFetch_InstanceTransformBuffer[3 * (InstanceId + InstanceOffset) + 2];
+	Intermediates.InstanceOrigin = VertexFetch_InstanceOriginBuffer[(InstanceId + InstanceOffset)];
+	Intermediates.InstanceLightmapAndShadowMapUVBias = VertexFetch_InstanceLightmapBuffer[(InstanceId + InstanceOffset)];
 #elif USE_INSTANCING && !USE_INSTANCING_EMULATED
 	Intermediates.InstanceTransform1 = Input.InstanceTransform1;
 	Intermediates.InstanceTransform2 = Input.InstanceTransform2;
@@ -831,7 +827,7 @@ FVertexFactoryInterpolantsVSToPS VertexFactoryGetInterpolantsVSToPS(FVertexFacto
 #endif
 
 #if INSTANCED_STEREO
-	Interpolants.PackedEyeIndex = 0;
+	Interpolants.EyeIndex = 0;
 #endif
 
 	SetPrimitiveId(Interpolants, Intermediates.PrimitiveId);
@@ -1017,17 +1013,6 @@ float4 VertexFactoryGetInstanceHitProxyId(FVertexFactoryInput Input, FVertexFact
 }
 #endif	// USE_INSTANCING
 
-#if INSTANCED_STEREO
-uint VertexFactoryGetEyeIndex(uint InstanceId)
-{
-	#if USE_INSTANCING
-		return InstancedEyeIndex;
-	#else
-		return InstanceId & 1;
-	#endif
-}
-#endif
-
 float4 VertexFactoryGetTranslatedPrimitiveVolumeBounds(FVertexFactoryInterpolantsVSToPS Interpolants)
 {
 	float4 ObjectWorldPositionAndRadius = GetPrimitiveData(GetPrimitiveId(Interpolants)).ObjectWorldPositionAndRadius;
diff --git a/Engine/Shaders/Private/LocalVertexFactoryCommon.ush b/Engine/Shaders/Private/LocalVertexFactoryCommon.ush
--- a/Engine/Shaders/Private/LocalVertexFactoryCommon.ush
+++ b/Engine/Shaders/Private/LocalVertexFactoryCommon.ush
@@ -32,7 +32,7 @@ struct FVertexFactoryInterpolantsVSToPS
 #endif
 
 #if INSTANCED_STEREO
-	nointerpolation uint PackedEyeIndex : PACKED_EYE_INDEX;
+	nointerpolation uint EyeIndex : PACKED_EYE_INDEX;
 #endif
 #if VF_USE_PRIMITIVE_SCENE_DATA
 	nointerpolation uint PrimitiveId : PRIMITIVE_ID;
@@ -193,7 +193,7 @@ struct FVertexFactoryInterpolantsVSToPS
 #endif
 
 #if INSTANCED_STEREO
-	nointerpolation uint PackedEyeIndex : PACKED_EYE_INDEX;
+	nointerpolation uint EyeIndex : PACKED_EYE_INDEX;
 #endif
 };
 
diff --git a/Engine/Shaders/Private/MaterialTexCoordScalesPixelShader.usf b/Engine/Shaders/Private/MaterialTexCoordScalesPixelShader.usf
--- a/Engine/Shaders/Private/MaterialTexCoordScalesPixelShader.usf
+++ b/Engine/Shaders/Private/MaterialTexCoordScalesPixelShader.usf
@@ -102,7 +102,7 @@ void Main(
 	)
 {
 #if INSTANCED_STEREO
-	ResolvedView = ResolveView(GetEyeIndex(DebugInputs.PackedEyeIndex));
+	ResolvedView = ResolveView(DebugInputs.EyeIndex);
 #else
 	ResolvedView = ResolveView();
 #endif
diff --git a/Engine/Shaders/Private/MeshDecals.usf b/Engine/Shaders/Private/MeshDecals.usf
--- a/Engine/Shaders/Private/MeshDecals.usf
+++ b/Engine/Shaders/Private/MeshDecals.usf
@@ -152,7 +152,7 @@ void FPixelShaderInOut_MainPS(
 	inout FPixelShaderOut Out)
 {
 #if INSTANCED_STEREO
-	ResolvedView = ResolveView(GetEyeIndex(FactoryInterpolants.PackedEyeIndex));
+	ResolvedView = ResolveView(FactoryInterpolants.EyeIndex);
 #else
 	ResolvedView = ResolveView();
 #endif
diff --git a/Engine/Shaders/Private/MeshParticleVertexFactory.ush b/Engine/Shaders/Private/MeshParticleVertexFactory.ush
--- a/Engine/Shaders/Private/MeshParticleVertexFactory.ush
+++ b/Engine/Shaders/Private/MeshParticleVertexFactory.ush
@@ -129,7 +129,7 @@ struct FVertexFactoryInterpolantsVSToPS
 #endif
 
 #if INSTANCED_STEREO
-	nointerpolation uint PackedEyeIndex : PACKED_EYE_INDEX;
+	nointerpolation uint EyeIndex : PACKED_EYE_INDEX;
 #endif
 };
 
@@ -567,7 +567,7 @@ FVertexFactoryInterpolantsVSToPS VertexFactoryGetInterpolantsVSToPS(FVertexFacto
 #endif
 
 #if INSTANCED_STEREO
-	Interpolants.PackedEyeIndex = 0;
+	Interpolants.EyeIndex = 0;
 #endif
 
 	return Interpolants;
@@ -675,23 +675,6 @@ float4 VertexFactoryGetPreviousWorldPosition(FVertexFactoryInput Input, FVertexF
 	}
 #endif // #if USING_TESSELLATION
 
-#if INSTANCED_STEREO
-
-// Instanced stereo sets the eye index explicitly for instanced geometry
-#if PARTICLE_MESH_INSTANCED
-uint InstancedEyeIndex;
-#endif
-
-uint VertexFactoryGetEyeIndex(uint InstanceId)
-{
-	#if PARTICLE_MESH_INSTANCED
-		return InstancedEyeIndex & bIsInstancedStereo;
-	#else
-		return InstanceId & 1;
-	#endif
-}
-#endif
-
 float4 VertexFactoryGetTranslatedPrimitiveVolumeBounds(FVertexFactoryInterpolantsVSToPS Interpolants)
 {
 	return 0;
diff --git a/Engine/Shaders/Private/MeshTexCoordSizeAccuracyPixelShader.usf b/Engine/Shaders/Private/MeshTexCoordSizeAccuracyPixelShader.usf
--- a/Engine/Shaders/Private/MeshTexCoordSizeAccuracyPixelShader.usf
+++ b/Engine/Shaders/Private/MeshTexCoordSizeAccuracyPixelShader.usf
@@ -71,7 +71,7 @@ void Main(
 	)
 {
 #if INSTANCED_STEREO
-	ResolvedView = ResolveView(GetEyeIndex(DebugInputs.PackedEyeIndex));
+	ResolvedView = ResolveView(DebugInputs.EyeIndex);
 #else
 	ResolvedView = ResolveView();
 #endif
diff --git a/Engine/Shaders/Private/MissingShaderPixelShader.usf b/Engine/Shaders/Private/MissingShaderPixelShader.usf
--- a/Engine/Shaders/Private/MissingShaderPixelShader.usf
+++ b/Engine/Shaders/Private/MissingShaderPixelShader.usf
@@ -18,7 +18,7 @@ void Main(
 	)
 {
 #if INSTANCED_STEREO
-	ResolvedView = ResolveView(GetEyeIndex(DebugInputs.PackedEyeIndex));
+	ResolvedView = ResolveView(DebugInputs.EyeIndex);
 #else
 	ResolvedView = ResolveView();
 #endif
diff --git a/Engine/Shaders/Private/NiagaraMeshVertexFactory.ush b/Engine/Shaders/Private/NiagaraMeshVertexFactory.ush
--- a/Engine/Shaders/Private/NiagaraMeshVertexFactory.ush
+++ b/Engine/Shaders/Private/NiagaraMeshVertexFactory.ush
@@ -114,7 +114,7 @@ struct FVertexFactoryInterpolantsVSToPS
 #endif
 
 #if INSTANCED_STEREO
-	nointerpolation uint PackedEyeIndex : PACKED_EYE_INDEX;
+	nointerpolation uint EyeIndex : PACKED_EYE_INDEX;
 #endif
 };
 
@@ -783,7 +783,7 @@ FVertexFactoryInterpolantsVSToPS VertexFactoryGetInterpolantsVSToPS(FVertexFacto
 #endif
 
 #if INSTANCED_STEREO
-	Interpolants.PackedEyeIndex = 0;
+	Interpolants.EyeIndex = 0;
 #endif
 
 	return Interpolants;
@@ -900,14 +900,3 @@ float4 VertexFactoryGetPreviousWorldPosition(FVertexFactoryInput Input, FVertexF
 		return TangentToWorld;
 	}
 #endif // #if USING_TESSELLATION
-
-#if INSTANCED_STEREO
-
-// Instanced stereo sets the eye index explicitly for instanced geometry
-uint InstancedEyeIndex;
-
-uint VertexFactoryGetEyeIndex(uint InstanceId)
-{
-	return InstancedEyeIndex & bIsInstancedStereo;
-}
-#endif
diff --git a/Engine/Shaders/Private/NiagaraRibbonVertexFactory.ush b/Engine/Shaders/Private/NiagaraRibbonVertexFactory.ush
--- a/Engine/Shaders/Private/NiagaraRibbonVertexFactory.ush
+++ b/Engine/Shaders/Private/NiagaraRibbonVertexFactory.ush
@@ -213,7 +213,7 @@ struct FVertexFactoryInterpolantsVSToPS
 #endif
 
 #if INSTANCED_STEREO
-	nointerpolation uint PackedEyeIndex : PACKED_EYE_INDEX;
+	nointerpolation uint EyeIndex : PACKED_EYE_INDEX;
 #endif
 };
 
@@ -715,7 +715,7 @@ FVertexFactoryInterpolantsVSToPS VertexFactoryGetInterpolantsVSToPS(FVertexFacto
 	Interpolants.TangentToWorld2 = float4(TangentToWorld[2], sign(determinant(TangentToLocal)));
 
 #if INSTANCED_STEREO
-	Interpolants.PackedEyeIndex = 0;
+	Interpolants.EyeIndex = 0;
 #endif
 
 	return Interpolants;
@@ -739,13 +739,6 @@ float3x3 VertexFactoryGetTangentToLocal( FVertexFactoryInput Input, FVertexFacto
 	return CalcTangentBasis(Input, Intermediates);
 }
 
-#if INSTANCED_STEREO
-uint VertexFactoryGetEyeIndex(uint InstanceId)
-{
-	return InstanceId & 1;
-}
-#endif
-
 float4 VertexFactoryGetTranslatedPrimitiveVolumeBounds(FVertexFactoryInterpolantsVSToPS Interpolants)
 {
 	return 0;
diff --git a/Engine/Shaders/Private/NiagaraSpriteVertexFactory.ush b/Engine/Shaders/Private/NiagaraSpriteVertexFactory.ush
--- a/Engine/Shaders/Private/NiagaraSpriteVertexFactory.ush
+++ b/Engine/Shaders/Private/NiagaraSpriteVertexFactory.ush
@@ -123,7 +123,7 @@ struct FVertexFactoryInterpolantsVSToPS
 #endif
 
 #if INSTANCED_STEREO
-	nointerpolation uint PackedEyeIndex : PACKED_EYE_INDEX;
+	nointerpolation uint EyeIndex : PACKED_EYE_INDEX;
 #endif
 };
 
@@ -970,7 +970,7 @@ FVertexFactoryInterpolantsVSToPS VertexFactoryGetInterpolantsVSToPS(FVertexFacto
 
 
 #if INSTANCED_STEREO
-	Interpolants.PackedEyeIndex = 0;
+	Interpolants.EyeIndex = 0;
 #endif
 
 	return Interpolants;
@@ -1080,17 +1080,6 @@ float3x3 VertexFactoryGetTangentToLocal(FVertexFactoryInput Input, FVertexFactor
 	return Intermediates.TangentToLocal;
 }
 
-#if INSTANCED_STEREO
-
-// Instanced stereo sets the eye index explicitly for instanced geometry
-uint InstancedEyeIndex;	 
-
-uint VertexFactoryGetEyeIndex(uint InstanceId)
-{
-	return InstancedEyeIndex & bIsInstancedStereo;
-}
-#endif
-
 float4 VertexFactoryGetTranslatedPrimitiveVolumeBounds(FVertexFactoryInterpolantsVSToPS Interpolants)
 {
 	return 0;
diff --git a/Engine/Shaders/Private/PNTriangles.ush b/Engine/Shaders/Private/PNTriangles.ush
--- a/Engine/Shaders/Private/PNTriangles.ush
+++ b/Engine/Shaders/Private/PNTriangles.ush
@@ -57,7 +57,7 @@ float4 ComputeControlPoint(float4 PositionA, float4 PositionB, float4 NormalA)
 		FHullShaderConstantOutput O = (FHullShaderConstantOutput)0;
 
 		#if INSTANCED_STEREO
-			ResolvedView = ResolveView(GetEyeIndex(I[0].PassSpecificData.FactoryInterpolants.InterpolantsVSToPS.PackedEyeIndex));
+			ResolvedView = ResolveView(I[0].PassSpecificData.FactoryInterpolants.InterpolantsVSToPS.EyeIndex);
 		#else
 			ResolvedView = ResolveView();
 		#endif
@@ -109,7 +109,7 @@ float4 ComputeControlPoint(float4 PositionA, float4 PositionB, float4 NormalA)
 	FPNTessellationHSToDS MainHull( InputPatch<FPassSpecificVSToDS, IN_CONTROL_POINTS> I, uint ControlPointID : SV_OutputControlPointID )
 	{
 		#if INSTANCED_STEREO
-			ResolvedView = ResolveView(GetEyeIndex(I[0].FactoryInterpolants.InterpolantsVSToPS.PackedEyeIndex));
+			ResolvedView = ResolveView(I[0].FactoryInterpolants.InterpolantsVSToPS.EyeIndex);
 		#else
 			ResolvedView = ResolveView();
 		#endif
@@ -182,7 +182,7 @@ float4 ComputeControlPoint(float4 PositionA, float4 PositionB, float4 NormalA)
 		)
 	{
 		#if INSTANCED_STEREO
-			ResolvedView = ResolveView(GetEyeIndex(I[0].PassSpecificData.FactoryInterpolants.InterpolantsVSToPS.PackedEyeIndex));
+			ResolvedView = ResolveView(I[0].PassSpecificData.FactoryInterpolants.InterpolantsVSToPS.EyeIndex);
 			#if !MULTI_VIEW
 				OutClipDistance = 0.0;
 			#endif
@@ -234,7 +234,7 @@ float4 ComputeControlPoint(float4 PositionA, float4 PositionB, float4 NormalA)
 
 		// Forward the eye index along
 		#if INSTANCED_STEREO
-		Interp.FactoryInterpolants.InterpolantsVSToPS.PackedEyeIndex = I[0].PassSpecificData.FactoryInterpolants.InterpolantsVSToPS.PackedEyeIndex;
+		Interp.FactoryInterpolants.InterpolantsVSToPS.EyeIndex = I[0].PassSpecificData.FactoryInterpolants.InterpolantsVSToPS.EyeIndex;
 		#endif
 
 		// get our base material params that we will potentially override
@@ -287,9 +287,9 @@ float4 ComputeControlPoint(float4 PositionA, float4 PositionB, float4 NormalA)
 
 		#if INSTANCED_STEREO && !MULTI_VIEW
 		BRANCH
-		if (IsInstancedStereo(I[0].PassSpecificData.FactoryInterpolants.InterpolantsVSToPS.PackedEyeIndex))
+		if (IsInstancedStereo())
 		{
-			const uint EyeIndex = GetEyeIndex(I[0].PassSpecificData.FactoryInterpolants.InterpolantsVSToPS.PackedEyeIndex);
+			const uint EyeIndex = I[0].PassSpecificData.FactoryInterpolants.InterpolantsVSToPS.EyeIndex;
 
 			// Clip at the center of the screen
 			OutClipDistance = dot(Result.Position, EyeClipEdge[EyeIndex]);
diff --git a/Engine/Shaders/Private/ParticleBeamTrailVertexFactory.ush b/Engine/Shaders/Private/ParticleBeamTrailVertexFactory.ush
--- a/Engine/Shaders/Private/ParticleBeamTrailVertexFactory.ush
+++ b/Engine/Shaders/Private/ParticleBeamTrailVertexFactory.ush
@@ -36,7 +36,7 @@ struct FVertexFactoryInterpolantsVSToPS
 #endif	//USE_DYNAMIC_PARAMETERS
 
 #if INSTANCED_STEREO
-	nointerpolation uint PackedEyeIndex : PACKED_EYE_INDEX;
+	nointerpolation uint EyeIndex : PACKED_EYE_INDEX;
 #endif
 };
 
@@ -207,7 +207,7 @@ FVertexFactoryInterpolantsVSToPS VertexFactoryGetInterpolantsVSToPS(FVertexFacto
 	Interpolants.TangentToWorld2 = float4(TangentToWorld[2], sign(determinant(TangentToLocal)));
 
 #if INSTANCED_STEREO
-	Interpolants.PackedEyeIndex = 0;
+	Interpolants.EyeIndex = 0;
 #endif
 
 	return Interpolants;
@@ -229,13 +229,6 @@ float3x3 VertexFactoryGetTangentToLocal( FVertexFactoryInput Input, FVertexFacto
 	return CalcTangentBasis(Input);
 }
 
-#if INSTANCED_STEREO
-uint VertexFactoryGetEyeIndex(uint InstanceId)
-{
-	return InstanceId & 1;
-}
-#endif
-
 float4 VertexFactoryGetTranslatedPrimitiveVolumeBounds(FVertexFactoryInterpolantsVSToPS Interpolants)
 {
 	return 0;
diff --git a/Engine/Shaders/Private/ParticleGPUSpriteVertexFactory.ush b/Engine/Shaders/Private/ParticleGPUSpriteVertexFactory.ush
--- a/Engine/Shaders/Private/ParticleGPUSpriteVertexFactory.ush
+++ b/Engine/Shaders/Private/ParticleGPUSpriteVertexFactory.ush
@@ -96,7 +96,7 @@ struct FVertexFactoryInterpolantsVSToPS
 #endif
 
 #if INSTANCED_STEREO
-	nointerpolation uint PackedEyeIndex : PACKED_EYE_INDEX;
+	nointerpolation uint EyeIndex : PACKED_EYE_INDEX;
 #endif
 };
 
@@ -661,7 +661,7 @@ FVertexFactoryInterpolantsVSToPS VertexFactoryGetInterpolantsVSToPS( FVertexFact
 #endif
 
 #if INSTANCED_STEREO
-	Interpolants.PackedEyeIndex = 0;
+	Interpolants.EyeIndex = 0;
 #endif
 
 	return Interpolants;
@@ -678,17 +678,6 @@ float4 VertexFactoryGetPreviousWorldPosition( FVertexFactoryInput Input, FVertex
 	return float4(Intermediates.VertexWorldPosition,1);
 }
 
-#if INSTANCED_STEREO
-
-// Instanced stereo sets the eye index explicitly for instanced geometry
-uint InstancedEyeIndex;	 
-
-uint VertexFactoryGetEyeIndex(uint InstanceId)
-{
-	return InstancedEyeIndex & bIsInstancedStereo;
-}
-#endif
-
 float4 VertexFactoryGetTranslatedPrimitiveVolumeBounds(FVertexFactoryInterpolantsVSToPS Interpolants)
 {
 #if USE_PARTICLE_POSITION
diff --git a/Engine/Shaders/Private/ParticleSpriteVertexFactory.ush b/Engine/Shaders/Private/ParticleSpriteVertexFactory.ush
--- a/Engine/Shaders/Private/ParticleSpriteVertexFactory.ush
+++ b/Engine/Shaders/Private/ParticleSpriteVertexFactory.ush
@@ -85,7 +85,7 @@ struct FVertexFactoryInterpolantsVSToPS
 #endif
 
 #if INSTANCED_STEREO
-	nointerpolation uint PackedEyeIndex : PACKED_EYE_INDEX;
+	nointerpolation uint EyeIndex : PACKED_EYE_INDEX;
 #endif
 };
 
@@ -650,7 +650,7 @@ FVertexFactoryInterpolantsVSToPS VertexFactoryGetInterpolantsVSToPS(FVertexFacto
 #endif
 
 #if INSTANCED_STEREO
-	Interpolants.PackedEyeIndex = 0;
+	Interpolants.EyeIndex = 0;
 #endif
 
 	return Interpolants;
@@ -672,17 +672,6 @@ float3x3 VertexFactoryGetTangentToLocal(FVertexFactoryInput Input, FVertexFactor
 	return Intermediates.TangentToLocal;
 }
 
-#if INSTANCED_STEREO
-
-// Instanced stereo sets the eye index explicitly for instanced geometry
-uint InstancedEyeIndex;	 
-
-uint VertexFactoryGetEyeIndex(uint InstanceId)
-{
-	return InstancedEyeIndex & bIsInstancedStereo;
-}
-#endif
-
 float4 VertexFactoryGetTranslatedPrimitiveVolumeBounds(FVertexFactoryInterpolantsVSToPS Interpolants)
 {
 #if USE_PARTICLE_POSITION
diff --git a/Engine/Shaders/Private/PositionOnlyDepthVertexShader.usf b/Engine/Shaders/Private/PositionOnlyDepthVertexShader.usf
--- a/Engine/Shaders/Private/PositionOnlyDepthVertexShader.usf
+++ b/Engine/Shaders/Private/PositionOnlyDepthVertexShader.usf
@@ -8,10 +8,6 @@
 #include "/Engine/Generated/Material.ush"
 #include "/Engine/Generated/VertexFactory.ush"
 
-#if INSTANCED_STEREO
-bool bIsInstancedStereoEmulated;
-#endif
-
 void Main(
 	FPositionOnlyVertexFactoryInput Input,
 	out float4 OutPosition : SV_POSITION
@@ -29,17 +25,7 @@ void Main(
 	)
 {
 #if INSTANCED_STEREO
-	uint EyeIndex;
-	BRANCH
-	if (bIsInstancedStereoEmulated)
-	{
-		EyeIndex = View.StereoPassIndex;
-	}
-	else
-	{
-		EyeIndex = VertexFactoryGetEyeIndex(InstanceId);
-	}
-
+	uint EyeIndex = GetEyeIndex(InstanceId);
 	ResolvedView = ResolveView(EyeIndex);
 
 	#if !MULTI_VIEW
@@ -65,7 +51,7 @@ void Main(
 
 	#if INSTANCED_STEREO && !MULTI_VIEW
 	BRANCH 
-	if (bIsInstancedStereo || bIsInstancedStereoEmulated)  
+	if (IsInstancedStereo())  
 	{
 		// Clip at the center of the screen
 		OutClipDistance = dot(OutPosition, EyeClipEdge[EyeIndex]);
diff --git a/Engine/Shaders/Private/PrimitiveDistanceAccuracyPixelShader.usf b/Engine/Shaders/Private/PrimitiveDistanceAccuracyPixelShader.usf
--- a/Engine/Shaders/Private/PrimitiveDistanceAccuracyPixelShader.usf
+++ b/Engine/Shaders/Private/PrimitiveDistanceAccuracyPixelShader.usf
@@ -19,7 +19,7 @@ void Main(
 	)
 {
 #if INSTANCED_STEREO
-	ResolvedView = ResolveView(GetEyeIndex(DebugInputs.PackedEyeIndex));
+	ResolvedView = ResolveView(DebugInputs.EyeIndex);
 #else
 	ResolvedView = ResolveView();
 #endif
diff --git a/Engine/Shaders/Private/RequiredTextureResolutionPixelShader.usf b/Engine/Shaders/Private/RequiredTextureResolutionPixelShader.usf
--- a/Engine/Shaders/Private/RequiredTextureResolutionPixelShader.usf
+++ b/Engine/Shaders/Private/RequiredTextureResolutionPixelShader.usf
@@ -70,7 +70,7 @@ void Main(
 	)
 {
 #if INSTANCED_STEREO
-	ResolvedView = ResolveView(GetEyeIndex(DebugInputs.PackedEyeIndex));
+	ResolvedView = ResolveView(DebugInputs.EyeIndex);
 #else
 	ResolvedView = ResolveView();
 #endif
diff --git a/Engine/Shaders/Private/VectorFieldVisualizationVertexFactory.ush b/Engine/Shaders/Private/VectorFieldVisualizationVertexFactory.ush
--- a/Engine/Shaders/Private/VectorFieldVisualizationVertexFactory.ush
+++ b/Engine/Shaders/Private/VectorFieldVisualizationVertexFactory.ush
@@ -32,7 +32,7 @@ struct FVertexFactoryInterpolantsVSToPS
 	float2 DummyTexCoord : TEXCOORD0;
 
 #if INSTANCED_STEREO
-	nointerpolation uint PackedEyeIndex : PACKED_EYE_INDEX;
+	nointerpolation uint EyeIndex : EYE_INDEX;
 #endif
 };
 
@@ -153,7 +153,7 @@ FVertexFactoryInterpolantsVSToPS VertexFactoryGetInterpolantsVSToPS( FVertexFact
 	Interpolants.DummyTexCoord = float2(0,0);
 
 #if INSTANCED_STEREO
-	Interpolants.PackedEyeIndex = 0;
+	Interpolants.EyeIndex = 0;
 #endif
 
 	return Interpolants;
@@ -185,13 +185,6 @@ float3x3 VertexFactoryGetTangentToLocal( FVertexFactoryInput Input, FVertexFacto
 	return TangentToLocal;
 }
 
-#if INSTANCED_STEREO
-uint VertexFactoryGetEyeIndex(uint InstanceId)
-{
-	return 0;
-}
-#endif
-
 float4 VertexFactoryGetTranslatedPrimitiveVolumeBounds(FVertexFactoryInterpolantsVSToPS Interpolants)
 {
 	return 0;
diff --git a/Engine/Shaders/Private/VelocityShader.usf b/Engine/Shaders/Private/VelocityShader.usf
--- a/Engine/Shaders/Private/VelocityShader.usf
+++ b/Engine/Shaders/Private/VelocityShader.usf
@@ -123,7 +123,7 @@ void MainVertexShader(
 	)
 {
 #if INSTANCED_STEREO
-	const uint EyeIndex = VertexFactoryGetEyeIndex(InstanceId);
+	const uint EyeIndex = GetEyeIndex(InstanceId);
 	ResolvedView = ResolveView(EyeIndex);
 	#if !MULTI_VIEW
 		OutClipDistance = 0.0;
@@ -190,7 +190,7 @@ void MainVertexShader(
 
 	#if INSTANCED_STEREO && !MULTI_VIEW
 	BRANCH
-	if (bIsInstancedStereo)
+	if (IsInstancedStereo())
 	{
 		// Clip at the center of the screen
 		OutClipDistance = dot(Output.Position, EyeClipEdge[EyeIndex]);
@@ -222,9 +222,9 @@ void MainVertexShader(
 
 #if INSTANCED_STEREO
 	#if USING_TESSELLATION	
-		Output.FactoryInterpolants.InterpolantsVSToPS.PackedEyeIndex = PackEyeIndex(EyeIndex, bIsInstancedStereo);
+		Output.FactoryInterpolants.InterpolantsVSToPS.EyeIndex = EyeIndex;
 	#else
-		Output.Interps.FactoryInterpolants.PackedEyeIndex = PackEyeIndex(EyeIndex, bIsInstancedStereo);
+		Output.Interps.FactoryInterpolants.EyeIndex = EyeIndex;
 	#endif
 #endif
 }
@@ -245,7 +245,7 @@ void MainPixelShader(
 	)
 {
 #if INSTANCED_STEREO
-	ResolvedView = ResolveView(GetEyeIndex(Input.FactoryInterpolants.PackedEyeIndex));
+	ResolvedView = ResolveView(Input.FactoryInterpolants.EyeIndex);
 #else
 	ResolvedView = ResolveView();
 #endif
diff --git a/Engine/Source/Runtime/Landscape/Private/LandscapeGrass.cpp b/Engine/Source/Runtime/Landscape/Private/LandscapeGrass.cpp
--- a/Engine/Source/Runtime/Landscape/Private/LandscapeGrass.cpp
+++ b/Engine/Source/Runtime/Landscape/Private/LandscapeGrass.cpp
@@ -247,12 +247,11 @@ public:
 		const FPrimitiveSceneProxy* PrimitiveSceneProxy,
 		const FMaterialRenderProxy& MaterialRenderProxy,
 		const FMaterial& Material,
-		const TUniformBufferRef<FViewUniformShaderParameters>& ViewUniformBuffer,
-		FUniformBufferRHIParamRef PassUniformBufferValue,
+		const FMeshPassProcessorRenderState& DrawRenderState,
 		const FLandscapeGrassWeightShaderElementData& ShaderElementData,
 		FMeshDrawSingleShaderBindings& ShaderBindings) const
 	{
-		FMeshMaterialShader::GetShaderBindings(Scene, FeatureLevel, PrimitiveSceneProxy, MaterialRenderProxy, Material, ViewUniformBuffer, PassUniformBufferValue, ShaderElementData, ShaderBindings);
+		FMeshMaterialShader::GetShaderBindings(Scene, FeatureLevel, PrimitiveSceneProxy, MaterialRenderProxy, Material, DrawRenderState, ShaderElementData, ShaderBindings);
 
 		ShaderBindings.Add(RenderOffsetParameter, ShaderElementData.RenderOffset);
 	}
@@ -294,12 +293,11 @@ public:
 		const FPrimitiveSceneProxy* PrimitiveSceneProxy,
 		const FMaterialRenderProxy& MaterialRenderProxy,
 		const FMaterial& Material,
-		const TUniformBufferRef<FViewUniformShaderParameters>& ViewUniformBuffer,
-		FUniformBufferRHIParamRef PassUniformBufferValue,
+		const FMeshPassProcessorRenderState& DrawRenderState,
 		const FLandscapeGrassWeightShaderElementData& ShaderElementData,
 		FMeshDrawSingleShaderBindings& ShaderBindings) const
 	{
-		FMeshMaterialShader::GetShaderBindings(Scene, FeatureLevel, PrimitiveSceneProxy, MaterialRenderProxy, Material, ViewUniformBuffer, PassUniformBufferValue, ShaderElementData, ShaderBindings);
+		FMeshMaterialShader::GetShaderBindings(Scene, FeatureLevel, PrimitiveSceneProxy, MaterialRenderProxy, Material, DrawRenderState, ShaderElementData, ShaderBindings);
 
 		ShaderBindings.Add(OutputPassParameter, ShaderElementData.OutputPass);
 	}
diff --git a/Engine/Source/Runtime/RHI/Public/RHI.h b/Engine/Source/Runtime/RHI/Public/RHI.h
--- a/Engine/Source/Runtime/RHI/Public/RHI.h
+++ b/Engine/Source/Runtime/RHI/Public/RHI.h
@@ -123,11 +123,8 @@ RHI_API bool RHIGetPreviewFeatureLevel(ERHIFeatureLevel::Type& PreviewFeatureLev
 
 inline bool RHISupportsInstancedStereo(const EShaderPlatform Platform)
 {
-	//@todo MeshCommandPipeline - add vr.InstancedStereo render path.
-	return false;
-
 	// Only D3D SM5, PS4 and Metal SM5 supports Instanced Stereo
-	//return (Platform == EShaderPlatform::SP_PCD3D_SM5 || Platform == EShaderPlatform::SP_PS4 || Platform == EShaderPlatform::SP_METAL_SM5 || Platform == EShaderPlatform::SP_METAL_SM5_NOTESS);
+	return (Platform == EShaderPlatform::SP_PCD3D_SM5 || Platform == EShaderPlatform::SP_PS4 || Platform == EShaderPlatform::SP_METAL_SM5 || Platform == EShaderPlatform::SP_METAL_SM5_NOTESS);
 }
 
 inline bool RHISupportsMultiView(const EShaderPlatform Platform)
diff --git a/Engine/Source/Runtime/Renderer/Private/BasePassRendering.cpp b/Engine/Source/Runtime/Renderer/Private/BasePassRendering.cpp
--- a/Engine/Source/Runtime/Renderer/Private/BasePassRendering.cpp
+++ b/Engine/Source/Runtime/Renderer/Private/BasePassRendering.cpp
@@ -961,6 +961,9 @@ void FBasePassMeshProcessor::Process(
 	TBasePassShaderElementData<LightMapPolicyType> ShaderElementData(LightMapElementData);
 	ShaderElementData.InitializeMeshMaterialData(ViewIfDynamicMeshCommand, PrimitiveSceneProxy, MeshBatch, StaticMeshId, true);
 
+	const bool bIsInstancedStereo = ViewIfDynamicMeshCommand ? ViewIfDynamicMeshCommand->IsInstancedStereoPass() : Scene->bStaticDrawInstancedStereo;
+	const int32 InstanceFactor = bIsInstancedStereo ? 2 : 1;
+
 	BuildMeshDrawCommands(
 		MeshBatch,
 		BatchElementMask,
@@ -971,7 +974,7 @@ void FBasePassMeshProcessor::Process(
 		BasePassShaders,
 		MeshFillMode,
 		MeshCullMode,
-		1,
+		InstanceFactor,
 		SortKey,
 		EMeshPassFeatures::Default,
 		ShaderElementData);
@@ -1515,6 +1518,7 @@ FMeshPassProcessor* CreateBasePassProcessor(const FScene* Scene, const FSceneVie
 	FExclusiveDepthStencil::Type DefaultBasePassDepthStencilAccess = GetDefaultBasePassDepthStencilAccess(EShadingPath::Deferred);
 
 	FMeshPassProcessorRenderState PassDrawRenderState(Scene->UniformBuffers.ViewUniformBuffer, Scene->UniformBuffers.OpaqueBasePassUniformBuffer);
+	PassDrawRenderState.SetInstancedViewUniformBuffer(Scene->UniformBuffers.InstancedViewUniformBuffer);
 	SetupBasePassState(DefaultBasePassDepthStencilAccess, false, PassDrawRenderState);
 
 	return new(FMemStack::Get()) FBasePassMeshProcessor(Scene, Scene->GetFeatureLevel(), InViewIfDynamicMeshCommand, PassDrawRenderState, InDrawListContext);
@@ -1523,6 +1527,7 @@ FMeshPassProcessor* CreateBasePassProcessor(const FScene* Scene, const FSceneVie
 FMeshPassProcessor* CreateTranslucencyStandardPassProcessor(const FScene* Scene, const FSceneView* InViewIfDynamicMeshCommand, FMeshPassDrawListContext* InDrawListContext)
 {
 	FMeshPassProcessorRenderState PassDrawRenderState(Scene->UniformBuffers.ViewUniformBuffer, Scene->UniformBuffers.TranslucentBasePassUniformBuffer);
+	PassDrawRenderState.SetInstancedViewUniformBuffer(Scene->UniformBuffers.InstancedViewUniformBuffer);
 	PassDrawRenderState.SetDepthStencilState(TStaticDepthStencilState<false, CF_DepthNearOrEqual>::GetRHI());
 
 	return new(FMemStack::Get()) FBasePassMeshProcessor(Scene, Scene->GetFeatureLevel(), InViewIfDynamicMeshCommand, PassDrawRenderState, InDrawListContext, ETranslucencyPass::TPT_StandardTranslucency);
@@ -1531,6 +1536,7 @@ FMeshPassProcessor* CreateTranslucencyStandardPassProcessor(const FScene* Scene,
 FMeshPassProcessor* CreateTranslucencyAfterDOFProcessor(const FScene* Scene, const FSceneView* InViewIfDynamicMeshCommand, FMeshPassDrawListContext* InDrawListContext)
 {
 	FMeshPassProcessorRenderState PassDrawRenderState(Scene->UniformBuffers.ViewUniformBuffer, Scene->UniformBuffers.TranslucentBasePassUniformBuffer);
+	PassDrawRenderState.SetInstancedViewUniformBuffer(Scene->UniformBuffers.InstancedViewUniformBuffer);
 	PassDrawRenderState.SetDepthStencilState(TStaticDepthStencilState<false, CF_DepthNearOrEqual>::GetRHI());
 
 	return new(FMemStack::Get()) FBasePassMeshProcessor(Scene, Scene->GetFeatureLevel(), InViewIfDynamicMeshCommand, PassDrawRenderState, InDrawListContext, ETranslucencyPass::TPT_TranslucencyAfterDOF);
@@ -1539,6 +1545,7 @@ FMeshPassProcessor* CreateTranslucencyAfterDOFProcessor(const FScene* Scene, con
 FMeshPassProcessor* CreateTranslucencyAllPassProcessor(const FScene* Scene, const FSceneView* InViewIfDynamicMeshCommand, FMeshPassDrawListContext* InDrawListContext)
 {
 	FMeshPassProcessorRenderState PassDrawRenderState(Scene->UniformBuffers.ViewUniformBuffer, Scene->UniformBuffers.TranslucentBasePassUniformBuffer);
+	PassDrawRenderState.SetInstancedViewUniformBuffer(Scene->UniformBuffers.InstancedViewUniformBuffer);
 	PassDrawRenderState.SetDepthStencilState(TStaticDepthStencilState<false, CF_DepthNearOrEqual>::GetRHI());
 
 	return new(FMemStack::Get()) FBasePassMeshProcessor(Scene, Scene->GetFeatureLevel(), InViewIfDynamicMeshCommand, PassDrawRenderState, InDrawListContext, ETranslucencyPass::TPT_AllTranslucency);
diff --git a/Engine/Source/Runtime/Renderer/Private/BasePassRendering.h b/Engine/Source/Runtime/Renderer/Private/BasePassRendering.h
--- a/Engine/Source/Runtime/Renderer/Private/BasePassRendering.h
+++ b/Engine/Source/Runtime/Renderer/Private/BasePassRendering.h
@@ -211,8 +211,7 @@ public:
 		const FPrimitiveSceneProxy* PrimitiveSceneProxy,
 		const FMaterialRenderProxy& MaterialRenderProxy,
 		const FMaterial& Material,
-		const TUniformBufferRef<FViewUniformShaderParameters>& ViewUniformBuffer,
-		FUniformBufferRHIParamRef PassUniformBufferValue,
+		const FMeshPassProcessorRenderState& DrawRenderState,
 		const TBasePassShaderElementData<LightMapPolicyType>& ShaderElementData,
 		FMeshDrawSingleShaderBindings& ShaderBindings) const;
 
@@ -433,8 +432,7 @@ public:
 		const FPrimitiveSceneProxy* PrimitiveSceneProxy,
 		const FMaterialRenderProxy& MaterialRenderProxy,
 		const FMaterial& Material,
-		const TUniformBufferRef<FViewUniformShaderParameters>& ViewUniformBuffer,
-		FUniformBufferRHIParamRef PassUniformBufferValue,
+		const FMeshPassProcessorRenderState& DrawRenderState,
 		const TBasePassShaderElementData<LightMapPolicyType>& ShaderElementData,
 		FMeshDrawSingleShaderBindings& ShaderBindings) const;
 
diff --git a/Engine/Source/Runtime/Renderer/Private/BasePassRendering.inl b/Engine/Source/Runtime/Renderer/Private/BasePassRendering.inl
--- a/Engine/Source/Runtime/Renderer/Private/BasePassRendering.inl
+++ b/Engine/Source/Runtime/Renderer/Private/BasePassRendering.inl
@@ -26,12 +26,11 @@ void TBasePassVertexShaderPolicyParamType<LightMapPolicyType>::GetShaderBindings
 	const FPrimitiveSceneProxy* PrimitiveSceneProxy,
 	const FMaterialRenderProxy& MaterialRenderProxy,
 	const FMaterial& Material,
-	const TUniformBufferRef<FViewUniformShaderParameters>& ViewUniformBuffer,
-	FUniformBufferRHIParamRef PassUniformBufferValue,
+	const FMeshPassProcessorRenderState& DrawRenderState,
 	const TBasePassShaderElementData<LightMapPolicyType>& ShaderElementData,
 	FMeshDrawSingleShaderBindings& ShaderBindings) const
 {
-	FMeshMaterialShader::GetShaderBindings(Scene, FeatureLevel, PrimitiveSceneProxy, MaterialRenderProxy, Material, ViewUniformBuffer, PassUniformBufferValue, ShaderElementData, ShaderBindings);
+	FMeshMaterialShader::GetShaderBindings(Scene, FeatureLevel, PrimitiveSceneProxy, MaterialRenderProxy, Material, DrawRenderState, ShaderElementData, ShaderBindings);
 
 	if (Scene)
 	{
@@ -74,12 +73,11 @@ void TBasePassPixelShaderPolicyParamType<LightMapPolicyType>::GetShaderBindings(
 	const FPrimitiveSceneProxy* PrimitiveSceneProxy,
 	const FMaterialRenderProxy& MaterialRenderProxy,
 	const FMaterial& Material,
-	const TUniformBufferRef<FViewUniformShaderParameters>& ViewUniformBuffer,
-	FUniformBufferRHIParamRef PassUniformBufferValue,
+	const FMeshPassProcessorRenderState& DrawRenderState,
 	const TBasePassShaderElementData<LightMapPolicyType>& ShaderElementData,
 	FMeshDrawSingleShaderBindings& ShaderBindings) const
 {
-	FMeshMaterialShader::GetShaderBindings(Scene, FeatureLevel, PrimitiveSceneProxy, MaterialRenderProxy, Material, ViewUniformBuffer, PassUniformBufferValue, ShaderElementData, ShaderBindings);
+	FMeshMaterialShader::GetShaderBindings(Scene, FeatureLevel, PrimitiveSceneProxy, MaterialRenderProxy, Material, DrawRenderState, ShaderElementData, ShaderBindings);
 
 	if (Scene)
 	{
diff --git a/Engine/Source/Runtime/Renderer/Private/CompositionLighting/PostProcessMeshDecals.cpp b/Engine/Source/Runtime/Renderer/Private/CompositionLighting/PostProcessMeshDecals.cpp
--- a/Engine/Source/Runtime/Renderer/Private/CompositionLighting/PostProcessMeshDecals.cpp
+++ b/Engine/Source/Runtime/Renderer/Private/CompositionLighting/PostProcessMeshDecals.cpp
@@ -216,6 +216,7 @@ FMeshDecalMeshProcessor::FMeshDecalMeshProcessor(const FScene* Scene,
 {
 	PassDrawRenderState.SetDepthStencilState(TStaticDepthStencilState<false, CF_DepthNearOrEqual>::GetRHI());
 	PassDrawRenderState.SetViewUniformBuffer(Scene->UniformBuffers.ViewUniformBuffer);
+	PassDrawRenderState.SetInstancedViewUniformBuffer(Scene->UniformBuffers.InstancedViewUniformBuffer);
 	PassDrawRenderState.SetPassUniformBuffer(Scene->UniformBuffers.MeshDecalPassUniformBuffer);
 }
 
diff --git a/Engine/Source/Runtime/Renderer/Private/ConvertToUniformMesh.cpp b/Engine/Source/Runtime/Renderer/Private/ConvertToUniformMesh.cpp
--- a/Engine/Source/Runtime/Renderer/Private/ConvertToUniformMesh.cpp
+++ b/Engine/Source/Runtime/Renderer/Private/ConvertToUniformMesh.cpp
@@ -149,6 +149,7 @@ FConvertToUniformMeshProcessor::FConvertToUniformMeshProcessor(const FScene* Sce
 	PassDrawRenderState.SetDepthStencilState(TStaticDepthStencilState<false, CF_Always>::GetRHI());
 
 	PassDrawRenderState.SetViewUniformBuffer(Scene->UniformBuffers.ViewUniformBuffer);
+	PassDrawRenderState.SetInstancedViewUniformBuffer(Scene->UniformBuffers.InstancedViewUniformBuffer);
 	PassDrawRenderState.SetPassUniformBuffer(Scene->UniformBuffers.ConvertToUniformMeshPassUniformBuffer);
 }
 
diff --git a/Engine/Source/Runtime/Renderer/Private/DebugViewModeRendering.h b/Engine/Source/Runtime/Renderer/Private/DebugViewModeRendering.h
--- a/Engine/Source/Runtime/Renderer/Private/DebugViewModeRendering.h
+++ b/Engine/Source/Runtime/Renderer/Private/DebugViewModeRendering.h
@@ -97,12 +97,11 @@ public:
 		const FPrimitiveSceneProxy* PrimitiveSceneProxy,
 		const FMaterialRenderProxy& MaterialRenderProxy,
 		const FMaterial& Material,
-		const TUniformBufferRef<FViewUniformShaderParameters>& ViewUniformBuffer,
-		FUniformBufferRHIParamRef PassUniformBufferValue,
+		const FMeshPassProcessorRenderState& DrawRenderState,
 		const FDebugViewModeShaderElementData& ShaderElementData,
 		FMeshDrawSingleShaderBindings& ShaderBindings) const
 	{
-		FMeshMaterialShader::GetShaderBindings(Scene, FeatureLevel, PrimitiveSceneProxy, MaterialRenderProxy, Material, ViewUniformBuffer, PassUniformBufferValue, ShaderElementData, ShaderBindings);
+		FMeshMaterialShader::GetShaderBindings(Scene, FeatureLevel, PrimitiveSceneProxy, MaterialRenderProxy, Material, DrawRenderState, ShaderElementData, ShaderBindings);
 	}
 
 	static void SetCommonDefinitions(EShaderPlatform Platform, const FMaterial* Material, FShaderCompilerEnvironment& OutEnvironment)
diff --git a/Engine/Source/Runtime/Renderer/Private/DepthRendering.cpp b/Engine/Source/Runtime/Renderer/Private/DepthRendering.cpp
--- a/Engine/Source/Runtime/Renderer/Private/DepthRendering.cpp
+++ b/Engine/Source/Runtime/Renderer/Private/DepthRendering.cpp
@@ -718,6 +718,9 @@ void FDepthPassMeshProcessor::Process(
 	FDepthOnlyShaderElementData ShaderElementData(0.0f);
 	ShaderElementData.InitializeMeshMaterialData(ViewIfDynamicMeshCommand, PrimitiveSceneProxy, MeshBatch, StaticMeshId, true);
 
+	const bool bIsInstancedStereo = ViewIfDynamicMeshCommand ? ViewIfDynamicMeshCommand->IsInstancedStereoPass() : Scene->bStaticDrawInstancedStereo;
+	const int32 InstanceFactor = bIsInstancedStereo ? 2 : 1;
+
 	BuildMeshDrawCommands(
 		MeshBatch,
 		BatchElementMask,
@@ -728,7 +731,7 @@ void FDepthPassMeshProcessor::Process(
 		DepthPassShaders,
 		MeshFillMode,
 		MeshCullMode,
-		1,
+		InstanceFactor,
 		FMeshDrawCommandSortKey::Default,
 		bPositionOnly ? EMeshPassFeatures::PositionOnly : EMeshPassFeatures::Default,
 		ShaderElementData);
@@ -825,6 +828,7 @@ FDepthPassMeshProcessor::FDepthPassMeshProcessor(const FScene* Scene,
 {
 	PassDrawRenderState = InPassDrawRenderState;
 	PassDrawRenderState.SetViewUniformBuffer(Scene->UniformBuffers.ViewUniformBuffer);
+	PassDrawRenderState.SetInstancedViewUniformBuffer(Scene->UniformBuffers.InstancedViewUniformBuffer);
 	PassDrawRenderState.SetPassUniformBuffer(Scene->UniformBuffers.DepthPassUniformBuffer);
 }
 
diff --git a/Engine/Source/Runtime/Renderer/Private/DepthRendering.h b/Engine/Source/Runtime/Renderer/Private/DepthRendering.h
--- a/Engine/Source/Runtime/Renderer/Private/DepthRendering.h
+++ b/Engine/Source/Runtime/Renderer/Private/DepthRendering.h
@@ -79,12 +79,11 @@ public:
 		const FPrimitiveSceneProxy* PrimitiveSceneProxy,
 		const FMaterialRenderProxy& MaterialRenderProxy,
 		const FMaterial& Material,
-		const TUniformBufferRef<FViewUniformShaderParameters>& ViewUniformBuffer,
-		FUniformBufferRHIParamRef PassUniformBufferValue,
+		const FMeshPassProcessorRenderState& DrawRenderState,
 		const FDepthOnlyShaderElementData& ShaderElementData,
 		FMeshDrawSingleShaderBindings& ShaderBindings) const
 	{
-		FMeshMaterialShader::GetShaderBindings(Scene, FeatureLevel, PrimitiveSceneProxy, MaterialRenderProxy, Material, ViewUniformBuffer, PassUniformBufferValue, ShaderElementData, ShaderBindings);
+		FMeshMaterialShader::GetShaderBindings(Scene, FeatureLevel, PrimitiveSceneProxy, MaterialRenderProxy, Material, DrawRenderState, ShaderElementData, ShaderBindings);
 	}
 };
 
@@ -162,12 +161,11 @@ public:
 		const FPrimitiveSceneProxy* PrimitiveSceneProxy,
 		const FMaterialRenderProxy& MaterialRenderProxy,
 		const FMaterial& Material,
-		const TUniformBufferRef<FViewUniformShaderParameters>& ViewUniformBuffer,
-		FUniformBufferRHIParamRef PassUniformBufferValue,
+		const FMeshPassProcessorRenderState& DrawRenderState,
 		const FDepthOnlyShaderElementData& ShaderElementData,
 		FMeshDrawSingleShaderBindings& ShaderBindings) const
 	{
-		FMeshMaterialShader::GetShaderBindings(Scene, FeatureLevel, PrimitiveSceneProxy, MaterialRenderProxy, Material, ViewUniformBuffer, PassUniformBufferValue, ShaderElementData, ShaderBindings);
+		FMeshMaterialShader::GetShaderBindings(Scene, FeatureLevel, PrimitiveSceneProxy, MaterialRenderProxy, Material, DrawRenderState, ShaderElementData, ShaderBindings);
 
 		ShaderBindings.Add(MobileColorValue, ShaderElementData.MobileColorValue);
 	}
diff --git a/Engine/Source/Runtime/Renderer/Private/DistortionRendering.cpp b/Engine/Source/Runtime/Renderer/Private/DistortionRendering.cpp
--- a/Engine/Source/Runtime/Renderer/Private/DistortionRendering.cpp
+++ b/Engine/Source/Runtime/Renderer/Private/DistortionRendering.cpp
@@ -782,6 +782,7 @@ FMeshPassProcessor* CreateDistortionPassProcessor(const FScene* Scene, const FSc
 {
 	FMeshPassProcessorRenderState DistortionPassState;
 	DistortionPassState.SetViewUniformBuffer(Scene->UniformBuffers.ViewUniformBuffer);
+	DistortionPassState.SetInstancedViewUniformBuffer(Scene->UniformBuffers.InstancedViewUniformBuffer);
 	DistortionPassState.SetPassUniformBuffer(Scene->UniformBuffers.DistortionPassUniformBuffer);
 	
 	// test against depth and write stencil mask
@@ -803,6 +804,7 @@ FMeshPassProcessor* CreateMobileDistortionPassProcessor(const FScene* Scene, con
 {
 	FMeshPassProcessorRenderState DistortionPassState;
 	DistortionPassState.SetViewUniformBuffer(Scene->UniformBuffers.ViewUniformBuffer);
+	DistortionPassState.SetInstancedViewUniformBuffer(Scene->UniformBuffers.InstancedViewUniformBuffer);
 	DistortionPassState.SetPassUniformBuffer(Scene->UniformBuffers.MobileDistortionPassUniformBuffer);
 
 	// We don't have depth, render all pixels, pixel shader will sample SceneDepth from SceneColor.A and discard if occluded
diff --git a/Engine/Source/Runtime/Renderer/Private/LightMapDensityRendering.h b/Engine/Source/Runtime/Renderer/Private/LightMapDensityRendering.h
--- a/Engine/Source/Runtime/Renderer/Private/LightMapDensityRendering.h
+++ b/Engine/Source/Runtime/Renderer/Private/LightMapDensityRendering.h
@@ -92,12 +92,11 @@ public:
 		const FPrimitiveSceneProxy* PrimitiveSceneProxy,
 		const FMaterialRenderProxy& MaterialRenderProxy,
 		const FMaterial& Material,
-		const TUniformBufferRef<FViewUniformShaderParameters>& ViewUniformBuffer,
-		FUniformBufferRHIParamRef PassUniformBufferValue,
+		const FMeshPassProcessorRenderState& DrawRenderState,
 		const TLightMapDensityElementData<LightMapPolicyType>& ShaderElementData,
 		FMeshDrawSingleShaderBindings& ShaderBindings) const
 	{
-		FMeshMaterialShader::GetShaderBindings(Scene, FeatureLevel, PrimitiveSceneProxy, MaterialRenderProxy, Material, ViewUniformBuffer, PassUniformBufferValue, ShaderElementData, ShaderBindings);
+		FMeshMaterialShader::GetShaderBindings(Scene, FeatureLevel, PrimitiveSceneProxy, MaterialRenderProxy, Material, DrawRenderState, ShaderElementData, ShaderBindings);
 
 		LightMapPolicyType::GetVertexShaderBindings(
 			PrimitiveSceneProxy,
@@ -212,12 +211,11 @@ public:
 		const FPrimitiveSceneProxy* PrimitiveSceneProxy,
 		const FMaterialRenderProxy& MaterialRenderProxy,
 		const FMaterial& Material,
-		const TUniformBufferRef<FViewUniformShaderParameters>& ViewUniformBuffer,
-		FUniformBufferRHIParamRef PassUniformBufferValue,
+		const FMeshPassProcessorRenderState& DrawRenderState,
 		const TLightMapDensityElementData<LightMapPolicyType>& ShaderElementData,
 		FMeshDrawSingleShaderBindings& ShaderBindings) const
 	{
-		FMeshMaterialShader::GetShaderBindings(Scene, FeatureLevel, PrimitiveSceneProxy, MaterialRenderProxy, Material, ViewUniformBuffer, PassUniformBufferValue, ShaderElementData, ShaderBindings);
+		FMeshMaterialShader::GetShaderBindings(Scene, FeatureLevel, PrimitiveSceneProxy, MaterialRenderProxy, Material, DrawRenderState, ShaderElementData, ShaderBindings);
 
 		LightMapPolicyType::GetPixelShaderBindings(
 			PrimitiveSceneProxy,
diff --git a/Engine/Source/Runtime/Renderer/Private/MeshDrawCommands.cpp b/Engine/Source/Runtime/Renderer/Private/MeshDrawCommands.cpp
--- a/Engine/Source/Runtime/Renderer/Private/MeshDrawCommands.cpp
+++ b/Engine/Source/Runtime/Renderer/Private/MeshDrawCommands.cpp
@@ -236,7 +236,8 @@ void BuildMeshDrawCommandPrimitiveIdBuffer(
 	FMeshCommandOneFrameArray& TempVisibleMeshDrawCommands,
 	int32& MaxInstances,
 	int32& VisibleMeshDrawCommandsNum,
-	int32& NewPassVisibleMeshDrawCommandsNum
+	int32& NewPassVisibleMeshDrawCommandsNum,
+	uint32 MaxInstanceFactor = 1
 	)
 {
 	QUICK_SCOPE_CYCLE_COUNTER(STAT_BuildMeshDrawCommandPrimitiveIdBuffer);
@@ -245,6 +246,7 @@ void BuildMeshDrawCommandPrimitiveIdBuffer(
 	const FVisibleMeshDrawCommand* RESTRICT PassVisibleMeshDrawCommands = VisibleMeshDrawCommands.GetData();
 	const int32 NumDrawCommands = VisibleMeshDrawCommands.Num();
 
+	uint32 PrimitiveIdIndex = 0;
 	int32* RESTRICT PrimitiveIds = (int32*)PrimitiveIdBuffer.Buffer;
 
 	if (IsDynamicInstancingEnabled())
@@ -270,7 +272,7 @@ void BuildMeshDrawCommandPrimitiveIdBuffer(
 				else
 				{
 					FVisibleMeshDrawCommand NewVisibleMeshDrawCommand = VisibleMeshDrawCommand;
-					NewVisibleMeshDrawCommand.PrimitiveIdBufferOffset = DrawCommandIndex;
+					NewVisibleMeshDrawCommand.PrimitiveIdBufferOffset = PrimitiveIdIndex;
 					TempVisibleMeshDrawCommands.Emplace(MoveTemp(NewVisibleMeshDrawCommand));
 				}
 			}
@@ -298,7 +300,7 @@ void BuildMeshDrawCommandPrimitiveIdBuffer(
 						VisibleMeshDrawCommand.MeshCullMode,
 						VisibleMeshDrawCommand.SortKey);
 
-					NewVisibleMeshDrawCommand.PrimitiveIdBufferOffset = DrawCommandIndex;
+					NewVisibleMeshDrawCommand.PrimitiveIdBufferOffset = PrimitiveIdIndex;
 					TempVisibleMeshDrawCommands.Emplace(MoveTemp(NewVisibleMeshDrawCommand));
 
 					CurrentDynamicallyInstancedMeshCommandNumInstances = &NewCommand.NumInstances;
@@ -307,13 +309,19 @@ void BuildMeshDrawCommandPrimitiveIdBuffer(
 				{
 					CurrentDynamicallyInstancedMeshCommandNumInstances = nullptr;
 					FVisibleMeshDrawCommand NewVisibleMeshDrawCommand = VisibleMeshDrawCommand;
-					NewVisibleMeshDrawCommand.PrimitiveIdBufferOffset = DrawCommandIndex;
+					NewVisibleMeshDrawCommand.PrimitiveIdBufferOffset = PrimitiveIdIndex;
 					TempVisibleMeshDrawCommands.Emplace(MoveTemp(NewVisibleMeshDrawCommand));
 				}
 			}
 
-			//@todo - refactor into memcpy
-			PrimitiveIds[DrawCommandIndex] = VisibleMeshDrawCommand.DrawPrimitiveId;
+			//@todo - refactor into instance step rate in the RHI
+			uint32 InstanceFactor = VisibleMeshDrawCommand.MeshDrawCommand->InstanceFactor;
+			InstanceFactor = ensure(InstanceFactor <= MaxInstanceFactor) ? InstanceFactor : MaxInstanceFactor;
+			for (uint32 InstanceFactorIndex = 0; InstanceFactorIndex < InstanceFactor; InstanceFactorIndex++, PrimitiveIdIndex++)
+			{
+				//@todo - refactor into memcpy
+				PrimitiveIds[PrimitiveIdIndex] = VisibleMeshDrawCommand.DrawPrimitiveId;
+			}
 		}
 
 		// Setup instancing stats for logging.
@@ -331,7 +339,12 @@ void BuildMeshDrawCommandPrimitiveIdBuffer(
 		for (int32 DrawCommandIndex = 0; DrawCommandIndex < NumDrawCommands; DrawCommandIndex++)
 		{
 			const FVisibleMeshDrawCommand& VisibleMeshDrawCommand = VisibleMeshDrawCommands[DrawCommandIndex];
-			PrimitiveIds[DrawCommandIndex] = VisibleMeshDrawCommand.DrawPrimitiveId;
+			uint32 InstanceFactor = VisibleMeshDrawCommand.MeshDrawCommand->InstanceFactor;
+			InstanceFactor = ensure(InstanceFactor <= MaxInstanceFactor) ? InstanceFactor : MaxInstanceFactor;
+			for (uint32 InstanceFactorIndex = 0; InstanceFactorIndex < InstanceFactor; InstanceFactorIndex++, PrimitiveIdIndex++)
+			{
+				PrimitiveIds[PrimitiveIdIndex] = VisibleMeshDrawCommand.DrawPrimitiveId;
+			}
 		}
 	}
 }
@@ -708,11 +721,14 @@ void SortPassMeshDrawCommands(
 		int32 MaxInstances = 1;
 		int32 VisibleMeshDrawCommandsNum = 0;
 		int32 NewPassVisibleMeshDrawCommandsNum = 0;
+		uint32 MaxInstanceFactor = 1;
 
 		// Preallocate context memory and resources.
 		if (bUseGPUScene)
 		{
-			PrimitiveIdBuffer = DynamicVertexBuffer.Allocate(NumDrawCommands * sizeof(int32));
+			//@todo - make the instance factor a pass parameter instead of a draw command parameter
+			MaxInstanceFactor = VisibleMeshDrawCommands[0].MeshDrawCommand->InstanceFactor;
+			PrimitiveIdBuffer = DynamicVertexBuffer.Allocate(NumDrawCommands * MaxInstanceFactor * sizeof(int32));
 
 			if (IsDynamicInstancingEnabled())
 			{
@@ -731,7 +747,8 @@ void SortPassMeshDrawCommands(
 				NewPassVisibleMeshDrawCommands,
 				MaxInstances,
 				VisibleMeshDrawCommandsNum,
-				NewPassVisibleMeshDrawCommandsNum
+				NewPassVisibleMeshDrawCommandsNum,
+				MaxInstanceFactor
 			);
 		}
 	}
@@ -784,6 +801,9 @@ void FParallelMeshDrawCommandPass::DispatchPassSetup(
 	TaskContext.ViewMatrix = View.ViewMatrices.GetViewMatrix();
 	TaskContext.PrimitiveBounds = &Scene->PrimitiveBounds;
 
+	//@todo - make the instance factor a pass parameter instead of a draw command parameter
+	TaskContext.MaxInstanceFactor = InOutMeshDrawCommands[0].MeshDrawCommand->InstanceFactor;
+
 	switch (PassType)
 	{
 		case EMeshPass::TranslucencyStandard: TaskContext.TranslucencyPass = ETranslucencyPass::TPT_StandardTranslucency; break;
@@ -807,7 +827,7 @@ void FParallelMeshDrawCommandPass::DispatchPassSetup(
 	// Preallocate resources on rendering thread based on MaxNumDraws.
 	if (TaskContext.bDynamicInstancing)
 	{
-		TaskContext.PrimitiveIdBuffer = FGlobalDynamicVertexBuffer::Get().Allocate(MaxNumDraws * sizeof(int32));
+		TaskContext.PrimitiveIdBuffer = FGlobalDynamicVertexBuffer::Get().Allocate(MaxNumDraws * TaskContext.MaxInstanceFactor * sizeof(int32));
 	}
 	TaskContext.MeshDrawCommands.Reserve(MaxNumDraws);
 	TaskContext.TempVisibleMeshDrawCommands.Reserve(MaxNumDraws);
diff --git a/Engine/Source/Runtime/Renderer/Private/MeshDrawCommands.h b/Engine/Source/Runtime/Renderer/Private/MeshDrawCommands.h
--- a/Engine/Source/Runtime/Renderer/Private/MeshDrawCommands.h
+++ b/Engine/Source/Runtime/Renderer/Private/MeshDrawCommands.h
@@ -37,6 +37,7 @@ public:
 		, VisibleMeshDrawCommandsNum(0)
 		, NewPassVisibleMeshDrawCommandsNum(0)
 		, MaxInstances(1)
+		, MaxInstanceFactor(1)
 	{
 	}
 
@@ -79,6 +80,9 @@ public:
 	int32 VisibleMeshDrawCommandsNum;
 	int32 NewPassVisibleMeshDrawCommandsNum;
 	int32 MaxInstances;
+
+	// Assertion check
+	uint32 MaxInstanceFactor;
 };
 
 /**
diff --git a/Engine/Source/Runtime/Renderer/Private/MeshPassProcessor.cpp b/Engine/Source/Runtime/Renderer/Private/MeshPassProcessor.cpp
--- a/Engine/Source/Runtime/Renderer/Private/MeshPassProcessor.cpp
+++ b/Engine/Source/Runtime/Renderer/Private/MeshPassProcessor.cpp
@@ -608,7 +608,7 @@ void FMeshDrawCommand::SetRayTracingShaders(const FMeshProcessorShaders& Shaders
 }
 #endif // RHI_RAYTRACING
 
-void FMeshDrawCommand::SetDrawParametersAndFinalize(const FMeshBatch& MeshBatch, int32 BatchElementIndex, int32 InstanceFactor, bool bDoSetupPsoStateForRasterization)
+void FMeshDrawCommand::SetDrawParametersAndFinalize(const FMeshBatch& MeshBatch, int32 BatchElementIndex, int32 BatchInstanceFactor, bool bDoSetupPsoStateForRasterization)
 {
 	const FMeshBatchElement& BatchElement = MeshBatch.Elements[BatchElementIndex];
 
@@ -617,7 +617,8 @@ void FMeshDrawCommand::SetDrawParametersAndFinalize(const FMeshBatch& MeshBatch,
 	IndexBuffer = BatchElement.IndexBuffer ? BatchElement.IndexBuffer->IndexBufferRHI : nullptr;
 	FirstIndex = BatchElement.FirstIndex;
 	NumPrimitives = BatchElement.NumPrimitives;
-	NumInstances = BatchElement.NumInstances * InstanceFactor;
+	NumInstances = BatchElement.NumInstances;
+	InstanceFactor = BatchInstanceFactor;
 	BaseVertexIndex = BatchElement.BaseVertexIndex;
 	NumVertices = BatchElement.MaxVertexIndex - BatchElement.MinVertexIndex + 1;
 	IndirectArgsBuffer = BatchElement.IndirectArgsBuffer;
@@ -733,7 +734,8 @@ void FMeshDrawCommand::SubmitDraw(
 
 	if (GShowMaterialDrawEvents)
 	{
-		if (MeshDrawCommand.NumInstances > 1)
+		const uint32 Instances = MeshDrawCommand.NumInstances * MeshDrawCommand.InstanceFactor;
+		if (Instances > 1)
 		{
 			BEGIN_DRAW_EVENTF(
 				RHICmdList,
@@ -741,7 +743,7 @@ void FMeshDrawCommand::SubmitDraw(
 				MeshEvent,
 				TEXT("%s %u instances"),
 				*MeshDrawCommand.DrawEventName,
-				MeshDrawCommand.NumInstances);
+				Instances);
 		}
 		else
 		{
@@ -801,7 +803,7 @@ void FMeshDrawCommand::SubmitDraw(
 				MeshDrawCommand.NumVertices,
 				MeshDrawCommand.FirstIndex,
 				MeshDrawCommand.NumPrimitives,
-				MeshDrawCommand.NumInstances
+				MeshDrawCommand.NumInstances * MeshDrawCommand.InstanceFactor
 			);
 		}
 	}
@@ -810,7 +812,7 @@ void FMeshDrawCommand::SubmitDraw(
 		RHICmdList.DrawPrimitive(
 			MeshDrawCommand.BaseVertexIndex + MeshDrawCommand.FirstIndex,
 			MeshDrawCommand.NumPrimitives,
-			MeshDrawCommand.NumInstances
+			MeshDrawCommand.NumInstances * MeshDrawCommand.InstanceFactor
 		);
 	}
 
@@ -986,7 +988,7 @@ FCachedPassMeshDrawListContext::FCachedPassMeshDrawListContext(FCachedMeshDrawCo
 
 FMeshDrawCommand& FCachedPassMeshDrawListContext::AddCommand(const FMeshDrawCommand& Initializer)
 {
-	// Only one FMeshDrawCommand supported per FStaticMeshBatch in a pass
+	// Only one FMeshDrawCommand supported per FStaticMesh in a pass
 	check(CommandInfo.CommandIndex == -1);
 	CommandInfo.CommandIndex = DrawList.MeshDrawCommands.Add(Initializer);
 	return DrawList.MeshDrawCommands[CommandInfo.CommandIndex];
diff --git a/Engine/Source/Runtime/Renderer/Private/MobileBasePass.cpp b/Engine/Source/Runtime/Renderer/Private/MobileBasePass.cpp
--- a/Engine/Source/Runtime/Renderer/Private/MobileBasePass.cpp
+++ b/Engine/Source/Runtime/Renderer/Private/MobileBasePass.cpp
@@ -445,12 +445,11 @@ void TMobileBasePassPSPolicyParamType<FUniformLightMapPolicy>::GetShaderBindings
 	const FPrimitiveSceneProxy* PrimitiveSceneProxy,
 	const FMaterialRenderProxy& MaterialRenderProxy,
 	const FMaterial& Material,
-	const TUniformBufferRef<FViewUniformShaderParameters>& ViewUniformBuffer,
-	FUniformBufferRHIParamRef PassUniformBufferValue,
+	const FMeshPassProcessorRenderState& DrawRenderState,
 	const TMobileBasePassShaderElementData<FUniformLightMapPolicy>& ShaderElementData,
 	FMeshDrawSingleShaderBindings& ShaderBindings) const
 {
-	FMeshMaterialShader::GetShaderBindings(Scene, FeatureLevel, PrimitiveSceneProxy, MaterialRenderProxy, Material, ViewUniformBuffer, PassUniformBufferValue, ShaderElementData, ShaderBindings);
+	FMeshMaterialShader::GetShaderBindings(Scene, FeatureLevel, PrimitiveSceneProxy, MaterialRenderProxy, Material, DrawRenderState, ShaderElementData, ShaderBindings);
 
 	FUniformLightMapPolicy::GetPixelShaderBindings(
 		PrimitiveSceneProxy,
diff --git a/Engine/Source/Runtime/Renderer/Private/MobileBasePassRendering.h b/Engine/Source/Runtime/Renderer/Private/MobileBasePassRendering.h
--- a/Engine/Source/Runtime/Renderer/Private/MobileBasePassRendering.h
+++ b/Engine/Source/Runtime/Renderer/Private/MobileBasePassRendering.h
@@ -131,12 +131,11 @@ public:
 		const FPrimitiveSceneProxy* PrimitiveSceneProxy,
 		const FMaterialRenderProxy& MaterialRenderProxy,
 		const FMaterial& Material,
-		const TUniformBufferRef<FViewUniformShaderParameters>& ViewUniformBuffer,
-		FUniformBufferRHIParamRef PassUniformBufferValue,
+		const FMeshPassProcessorRenderState& DrawRenderState,
 		const TMobileBasePassShaderElementData<LightMapPolicyType>& ShaderElementData,
 		FMeshDrawSingleShaderBindings& ShaderBindings) const
 	{
-		FMeshMaterialShader::GetShaderBindings(Scene, FeatureLevel, PrimitiveSceneProxy, MaterialRenderProxy, Material, ViewUniformBuffer, PassUniformBufferValue, ShaderElementData, ShaderBindings);
+		FMeshMaterialShader::GetShaderBindings(Scene, FeatureLevel, PrimitiveSceneProxy, MaterialRenderProxy, Material, DrawRenderState, ShaderElementData, ShaderBindings);
 
 		LightMapPolicyType::GetVertexShaderBindings(
 			PrimitiveSceneProxy,
@@ -292,8 +291,7 @@ public:
 		const FPrimitiveSceneProxy* PrimitiveSceneProxy,
 		const FMaterialRenderProxy& MaterialRenderProxy,
 		const FMaterial& Material,
-		const TUniformBufferRef<FViewUniformShaderParameters>& ViewUniformBuffer,
-		FUniformBufferRHIParamRef PassUniformBufferValue,
+		const FMeshPassProcessorRenderState& DrawRenderState,
 		const TMobileBasePassShaderElementData<LightMapPolicyType>& ShaderElementData,
 		FMeshDrawSingleShaderBindings& ShaderBindings) const;
 
diff --git a/Engine/Source/Runtime/Renderer/Private/RayTracing/RayTracingDynamicGeometry.cpp b/Engine/Source/Runtime/Renderer/Private/RayTracing/RayTracingDynamicGeometry.cpp
--- a/Engine/Source/Runtime/Renderer/Private/RayTracing/RayTracingDynamicGeometry.cpp
+++ b/Engine/Source/Runtime/Renderer/Private/RayTracing/RayTracingDynamicGeometry.cpp
@@ -47,12 +47,11 @@ public:
 		const FPrimitiveSceneProxy* PrimitiveSceneProxy,
 		const FMaterialRenderProxy& MaterialRenderProxy,
 		const FMaterial& Material,
-		const TUniformBufferRef<FViewUniformShaderParameters>& ViewUniformBuffer,
-		FUniformBufferRHIParamRef PassUniformBufferValue,
+		const FMeshPassProcessorRenderState& DrawRenderState,
 		const FMeshMaterialShaderElementData& ShaderElementData,
 		FMeshDrawSingleShaderBindings& ShaderBindings) const
 	{
-		FMeshMaterialShader::GetShaderBindings(Scene, FeatureLevel, PrimitiveSceneProxy, MaterialRenderProxy, Material, ViewUniformBuffer, PassUniformBufferValue, ShaderElementData, ShaderBindings);
+		FMeshMaterialShader::GetShaderBindings(Scene, FeatureLevel, PrimitiveSceneProxy, MaterialRenderProxy, Material, DrawRenderState, ShaderElementData, ShaderBindings);
 	}
 
 	void GetElementShaderBindings(
@@ -112,7 +111,8 @@ void FRayTracingDynamicGeometryCollection::AddDynamicMeshBatchForGeometryUpdate(
 	ShaderElementData.InitializeMeshMaterialData(View, PrimitiveSceneProxy, MeshBatch, -1, false);
 
 	FMeshDrawSingleShaderBindings SingleShaderBindings = ShaderBindings.GetSingleShaderBindings(SF_Compute);
-	Shader->GetShaderBindings(Scene, Scene->GetFeatureLevel(), PrimitiveSceneProxy, MaterialRenderProxy, Material, Scene->UniformBuffers.ViewUniformBuffer, Scene->UniformBuffers.OpaqueBasePassUniformBuffer, ShaderElementData, SingleShaderBindings);
+	FMeshPassProcessorRenderState DrawRenderState(Scene->UniformBuffers.ViewUniformBuffer, Scene->UniformBuffers.OpaqueBasePassUniformBuffer);
+	Shader->GetShaderBindings(Scene, Scene->GetFeatureLevel(), PrimitiveSceneProxy, MaterialRenderProxy, Material, DrawRenderState, ShaderElementData, SingleShaderBindings);
 
 	FVertexInputStreamArray DummyArray;
 	Shader->GetElementShaderBindings(Scene, View, MeshBatch.VertexFactory, false, Scene->GetFeatureLevel(), PrimitiveSceneProxy, MeshBatch, MeshBatch.Elements[0], ShaderElementData, SingleShaderBindings, DummyArray);
diff --git a/Engine/Source/Runtime/Renderer/Private/RayTracing/RayTracingMaterialHitShaders.cpp b/Engine/Source/Runtime/Renderer/Private/RayTracing/RayTracingMaterialHitShaders.cpp
--- a/Engine/Source/Runtime/Renderer/Private/RayTracing/RayTracingMaterialHitShaders.cpp
+++ b/Engine/Source/Runtime/Renderer/Private/RayTracing/RayTracingMaterialHitShaders.cpp
@@ -41,12 +41,11 @@ public:
 		const FPrimitiveSceneProxy* PrimitiveSceneProxy,
 		const FMaterialRenderProxy& MaterialRenderProxy,
 		const FMaterial& Material,
-		const TUniformBufferRef<FViewUniformShaderParameters>& ViewUniformBuffer,
-		FUniformBufferRHIParamRef PassUniformBufferValue,
+		const FMeshPassProcessorRenderState& DrawRenderState,
 		const TBasePassShaderElementData<FUniformLightMapPolicy>& ShaderElementData,
 		FMeshDrawSingleShaderBindings& ShaderBindings) const
 	{
-		FMeshMaterialShader::GetShaderBindings(Scene, FeatureLevel, PrimitiveSceneProxy, MaterialRenderProxy, Material, ViewUniformBuffer, PassUniformBufferValue, ShaderElementData, ShaderBindings);
+		FMeshMaterialShader::GetShaderBindings(Scene, FeatureLevel, PrimitiveSceneProxy, MaterialRenderProxy, Material, DrawRenderState, ShaderElementData, ShaderBindings);
 		
 		FUniformLightMapPolicy::GetPixelShaderBindings(
 			PrimitiveSceneProxy,
diff --git a/Engine/Source/Runtime/Renderer/Private/RendererScene.cpp b/Engine/Source/Runtime/Renderer/Private/RendererScene.cpp
--- a/Engine/Source/Runtime/Renderer/Private/RendererScene.cpp
+++ b/Engine/Source/Runtime/Renderer/Private/RendererScene.cpp
@@ -723,6 +723,9 @@ void FPersistentUniformBuffers::Initialize()
 	FViewUniformShaderParameters ViewUniformBufferParameters;
 	ViewUniformBuffer = TUniformBufferRef<FViewUniformShaderParameters>::CreateUniformBufferImmediate(ViewUniformBufferParameters, UniformBuffer_MultiFrame, EUniformBufferValidation::None);
 
+	FInstancedViewUniformShaderParameters InstancedViewUniformBufferParameters;
+	InstancedViewUniformBuffer = TUniformBufferRef<FInstancedViewUniformShaderParameters>::CreateUniformBufferImmediate(InstancedViewUniformBufferParameters, UniformBuffer_MultiFrame, EUniformBufferValidation::None);
+
 	FSceneTexturesUniformParameters DepthPassParameters;
 	DepthPassUniformBuffer = TUniformBufferRef<FSceneTexturesUniformParameters>::CreateUniformBufferImmediate(DepthPassParameters, UniformBuffer_MultiFrame, EUniformBufferValidation::None);
 
@@ -803,6 +806,16 @@ bool FPersistentUniformBuffers::UpdateViewUniformBuffer(const FViewInfo& View)
 	if (CachedView != &View)
 	{
 		ViewUniformBuffer.UpdateUniformBufferImmediate(*View.CachedViewUniformShaderParameters);
+
+		if (View.IsInstancedStereoPass() && View.Family->Views.Num() > 0)
+		{
+			// When drawing the left eye in a stereo scene, copy the right eye view values into the instanced view uniform buffer.
+			const EStereoscopicPass StereoPassIndex = (View.StereoPass != eSSP_FULL) ? eSSP_RIGHT_EYE : eSSP_FULL;
+
+			const FViewInfo& InstancedView = static_cast<const FViewInfo&>(View.Family->GetStereoEyeView(StereoPassIndex));
+			InstancedViewUniformBuffer.UpdateUniformBufferImmediate(reinterpret_cast<FInstancedViewUniformShaderParameters&>(*InstancedView.CachedViewUniformShaderParameters));
+		}
+
 		CachedView = &View;
 		return true;
 	}
@@ -866,6 +879,10 @@ FScene::FScene(UWorld* InWorld, bool bInRequiresHitProxies, bool bInIsEditorScen
 	static auto* ShaderPipelinesCvar = IConsoleManager::Get().FindTConsoleVariableDataInt(TEXT("r.ShaderPipelines"));
 	StaticDrawShaderPipelines = ShaderPipelinesCvar->GetValueOnAnyThread();
 
+	// Query instanced stereo and multi-view state
+	static const auto InstancedStereoCvar = IConsoleManager::Get().FindTConsoleVariableDataInt(TEXT("vr.InstancedStereo"));
+	bStaticDrawInstancedStereo = RHISupportsInstancedStereo(GShaderPlatformForFeatureLevel[InFeatureLevel]) && GEngine->IsStereoscopic3D() && InstancedStereoCvar->GetValueOnAnyThread();
+
 	if (World->FXSystem)
 	{
 		FFXSystemInterface::Destroy(World->FXSystem);
@@ -2840,17 +2857,21 @@ void FScene::ConditionalMarkStaticMeshElementsForUpdate()
 {
 	static auto* EarlyZPassCvar = IConsoleManager::Get().FindTConsoleVariableDataInt(TEXT("r.EarlyZPass"));
 	static auto* ShaderPipelinesCvar = IConsoleManager::Get().FindTConsoleVariableDataInt(TEXT("r.ShaderPipelines"));
+	static auto* InstancedStereoCvar = IConsoleManager::Get().FindTConsoleVariableDataInt(TEXT("vr.InstancedStereo"));
 
 	bool bMobileHDR = IsMobileHDR();
 	bool bMobileHDR32bpp = IsMobileHDR32bpp();
 	int32 DesiredStaticDrawListsEarlyZPassMode = EarlyZPassCvar->GetValueOnRenderThread();
 	int32 DesiredStaticDrawShaderPipelines = ShaderPipelinesCvar->GetValueOnRenderThread();
+	bool bDesiredStaticDrawInstancedStereo = RHISupportsInstancedStereo(GShaderPlatformForFeatureLevel[FeatureLevel]) &&
+		GEngine->IsStereoscopic3D() && InstancedStereoCvar->GetValueOnAnyThread();
 
 	if (bScenesPrimitivesNeedStaticMeshElementUpdate
 		|| bStaticDrawListsMobileHDR != bMobileHDR
 		|| bStaticDrawListsMobileHDR32bpp != bMobileHDR32bpp
 		|| StaticDrawShaderPipelines != DesiredStaticDrawShaderPipelines
-		|| StaticDrawListsEarlyZPassMode != DesiredStaticDrawListsEarlyZPassMode)
+		|| StaticDrawListsEarlyZPassMode != DesiredStaticDrawListsEarlyZPassMode
+		|| bStaticDrawInstancedStereo != bDesiredStaticDrawInstancedStereo)
 	{
 		// Mark all primitives as needing an update
 		// Note: Only visible primitives will actually update their static mesh elements
@@ -2864,6 +2885,7 @@ void FScene::ConditionalMarkStaticMeshElementsForUpdate()
 		bStaticDrawListsMobileHDR32bpp = bMobileHDR32bpp;
 		StaticDrawListsEarlyZPassMode = DesiredStaticDrawListsEarlyZPassMode;
 		StaticDrawShaderPipelines = DesiredStaticDrawShaderPipelines;
+		bStaticDrawInstancedStereo = bDesiredStaticDrawInstancedStereo;
 	}
 }
 
diff --git a/Engine/Source/Runtime/Renderer/Private/SceneHitProxyRendering.cpp b/Engine/Source/Runtime/Renderer/Private/SceneHitProxyRendering.cpp
--- a/Engine/Source/Runtime/Renderer/Private/SceneHitProxyRendering.cpp
+++ b/Engine/Source/Runtime/Renderer/Private/SceneHitProxyRendering.cpp
@@ -145,12 +145,11 @@ public:
 		const FPrimitiveSceneProxy* PrimitiveSceneProxy,
 		const FMaterialRenderProxy& MaterialRenderProxy,
 		const FMaterial& Material,
-		const TUniformBufferRef<FViewUniformShaderParameters>& ViewUniformBuffer,
-		FUniformBufferRHIParamRef PassUniformBufferValue,
+		const FMeshPassProcessorRenderState& DrawRenderState,
 		const FHitProxyShaderElementData& ShaderElementData,
 		FMeshDrawSingleShaderBindings& ShaderBindings) const
 	{
-		FMeshMaterialShader::GetShaderBindings(Scene, FeatureLevel, PrimitiveSceneProxy, MaterialRenderProxy, Material, ViewUniformBuffer, PassUniformBufferValue, ShaderElementData, ShaderBindings);
+		FMeshMaterialShader::GetShaderBindings(Scene, FeatureLevel, PrimitiveSceneProxy, MaterialRenderProxy, Material, DrawRenderState, ShaderElementData, ShaderBindings);
 
 
 		FHitProxyId hitProxyId = ShaderElementData.BatchHitProxyId;
@@ -681,6 +680,7 @@ FHitProxyMeshProcessor::FHitProxyMeshProcessor(const FScene* Scene, const FScene
 FMeshPassProcessor* CreateHitProxyPassProcessor(const FScene* Scene, const FSceneView* InViewIfDynamicMeshCommand, FMeshPassDrawListContext* InDrawListContext)
 {
 	FMeshPassProcessorRenderState PassDrawRenderState(Scene->UniformBuffers.ViewUniformBuffer, Scene->UniformBuffers.HitProxyPassUniformBuffer);
+	PassDrawRenderState.SetInstancedViewUniformBuffer(Scene->UniformBuffers.InstancedViewUniformBuffer);
 	PassDrawRenderState.SetDepthStencilState(TStaticDepthStencilState<true, CF_DepthNearOrEqual>::GetRHI());
 	PassDrawRenderState.SetBlendState(TStaticBlendState<>::GetRHI());
 	return new(FMemStack::Get()) FHitProxyMeshProcessor(Scene, InViewIfDynamicMeshCommand, true, PassDrawRenderState, InDrawListContext);
@@ -689,6 +689,7 @@ FMeshPassProcessor* CreateHitProxyPassProcessor(const FScene* Scene, const FScen
 FMeshPassProcessor* CreateHitProxyOpaqueOnlyPassProcessor(const FScene* Scene, const FSceneView* InViewIfDynamicMeshCommand, FMeshPassDrawListContext* InDrawListContext)
 {
 	FMeshPassProcessorRenderState PassDrawRenderState(Scene->UniformBuffers.ViewUniformBuffer, Scene->UniformBuffers.HitProxyPassUniformBuffer);
+	PassDrawRenderState.SetInstancedViewUniformBuffer(Scene->UniformBuffers.InstancedViewUniformBuffer);
 	PassDrawRenderState.SetDepthStencilState(TStaticDepthStencilState<true, CF_DepthNearOrEqual>::GetRHI());
 	PassDrawRenderState.SetBlendState(TStaticBlendState<>::GetRHI());
 	return new(FMemStack::Get()) FHitProxyMeshProcessor(Scene, InViewIfDynamicMeshCommand, false, PassDrawRenderState, InDrawListContext);
@@ -831,6 +832,7 @@ FEditorSelectionMeshProcessor::FEditorSelectionMeshProcessor(const FScene* Scene
 	PassDrawRenderState.SetDepthStencilState(TStaticDepthStencilState<true, CF_DepthNearOrEqual, true, CF_Always, SO_Keep, SO_Keep, SO_Replace>::GetRHI());
 	PassDrawRenderState.SetBlendState(TStaticBlendStateWriteMask<CW_NONE, CW_NONE, CW_NONE, CW_NONE>::GetRHI());
 	PassDrawRenderState.SetViewUniformBuffer(Scene->UniformBuffers.ViewUniformBuffer);
+	PassDrawRenderState.SetInstancedViewUniformBuffer(Scene->UniformBuffers.InstancedViewUniformBuffer);
 	PassDrawRenderState.SetPassUniformBuffer(Scene->UniformBuffers.EditorSelectionPassUniformBuffer);
 }
 
diff --git a/Engine/Source/Runtime/Renderer/Private/ScenePrivate.h b/Engine/Source/Runtime/Renderer/Private/ScenePrivate.h
--- a/Engine/Source/Runtime/Renderer/Private/ScenePrivate.h
+++ b/Engine/Source/Runtime/Renderer/Private/ScenePrivate.h
@@ -2338,6 +2338,7 @@ public:
 	bool UpdateViewUniformBuffer(const FViewInfo& View);
 
 	TUniformBufferRef<FViewUniformShaderParameters> ViewUniformBuffer;
+	TUniformBufferRef<FInstancedViewUniformShaderParameters> InstancedViewUniformBuffer;
 	TUniformBufferRef<FSceneTexturesUniformParameters> DepthPassUniformBuffer;
 	TUniformBufferRef<FOpaqueBasePassUniformParameters> OpaqueBasePassUniformBuffer;
 	TUniformBufferRef<FTranslucentBasePassUniformParameters> TranslucentBasePassUniformBuffer;
@@ -2494,6 +2495,9 @@ public:
 	/** Whether the ShaderPipelines were enabled when the static draw lists were built. */
 	int32 StaticDrawShaderPipelines;
 
+	/** Whether instanced stereoscopy was enabled when the static draw lists were built. */
+	bool bStaticDrawInstancedStereo;
+
 	/** True if a change to SkyLight / Lighting has occurred that requires static draw lists to be updated. */
 	bool bScenesPrimitivesNeedStaticMeshElementUpdate;
 
diff --git a/Engine/Source/Runtime/Renderer/Private/SceneVisibility.cpp b/Engine/Source/Runtime/Renderer/Private/SceneVisibility.cpp
--- a/Engine/Source/Runtime/Renderer/Private/SceneVisibility.cpp
+++ b/Engine/Source/Runtime/Renderer/Private/SceneVisibility.cpp
@@ -3540,8 +3540,12 @@ void FSceneRenderer::ComputeViewVisibility(FRHICommandListImmediate& RHICmdList,
 	for (int32 ViewIndex = 0; ViewIndex < Views.Num(); ViewIndex++)
 	{
 		FViewInfo& View = Views[ViewIndex];
-		FViewCommands& ViewCommands = ViewCommandsPerView[ViewIndex];
+		if (!View.ShouldRenderView())
+		{
+			continue;
+		}
 
+		FViewCommands& ViewCommands = ViewCommandsPerView[ViewIndex];
 		SetupMeshPass(View, BasePassDepthStencilAccess, ViewCommands);
 	}
 
diff --git a/Engine/Source/Runtime/Renderer/Private/ShaderBaseClasses.cpp b/Engine/Source/Runtime/Renderer/Private/ShaderBaseClasses.cpp
--- a/Engine/Source/Runtime/Renderer/Private/ShaderBaseClasses.cpp
+++ b/Engine/Source/Runtime/Renderer/Private/ShaderBaseClasses.cpp
@@ -518,16 +518,16 @@ void FMeshMaterialShader::GetShaderBindings(
 	const FPrimitiveSceneProxy* PrimitiveSceneProxy,
 	const FMaterialRenderProxy& MaterialRenderProxy,
 	const FMaterial& Material,
-	const TUniformBufferRef<FViewUniformShaderParameters>& ViewUniformBuffer,
-	FUniformBufferRHIParamRef PassUniformBufferValue,
+	const FMeshPassProcessorRenderState& DrawRenderState,
 	const FMeshMaterialShaderElementData& ShaderElementData,
 	FMeshDrawSingleShaderBindings& ShaderBindings) const
 {
 	FMaterialShader::GetShaderBindings(Scene, FeatureLevel, MaterialRenderProxy, Material, ShaderBindings);
-	ShaderBindings.Add(PassUniformBuffer, PassUniformBufferValue);
-	ShaderBindings.Add(GetUniformBufferParameter<FViewUniformShaderParameters>(), ViewUniformBuffer);
+	ShaderBindings.Add(PassUniformBuffer, DrawRenderState.GetPassUniformBuffer());
+	ShaderBindings.Add(GetUniformBufferParameter<FViewUniformShaderParameters>(), DrawRenderState.GetViewUniformBuffer());
 	ShaderBindings.Add(GetUniformBufferParameter<FDistanceCullFadeUniformShaderParameters>(), ShaderElementData.FadeUniformBuffer);
 	ShaderBindings.Add(GetUniformBufferParameter<FDitherUniformShaderParameters>(), ShaderElementData.DitherUniformBuffer);
+	ShaderBindings.Add(GetUniformBufferParameter<FInstancedViewUniformShaderParameters>(), DrawRenderState.GetInstancedViewUniformBuffer());
 }
 
 void FMeshMaterialShader::GetElementShaderBindings(
diff --git a/Engine/Source/Runtime/Renderer/Private/ShadowDepthRendering.cpp b/Engine/Source/Runtime/Renderer/Private/ShadowDepthRendering.cpp
--- a/Engine/Source/Runtime/Renderer/Private/ShadowDepthRendering.cpp
+++ b/Engine/Source/Runtime/Renderer/Private/ShadowDepthRendering.cpp
@@ -159,12 +159,11 @@ public:
 		const FPrimitiveSceneProxy* PrimitiveSceneProxy,
 		const FMaterialRenderProxy& MaterialRenderProxy,
 		const FMaterial& Material,
-		const TUniformBufferRef<FViewUniformShaderParameters>& ViewUniformBuffer,
-		FUniformBufferRHIParamRef PassUniformBufferValue,
+		const FMeshPassProcessorRenderState& DrawRenderState,
 		const FShadowDepthShaderElementData& ShaderElementData,
 		FMeshDrawSingleShaderBindings& ShaderBindings) const
 	{
-		FMeshMaterialShader::GetShaderBindings(Scene, FeatureLevel, PrimitiveSceneProxy, MaterialRenderProxy, Material, ViewUniformBuffer, PassUniformBufferValue, ShaderElementData, ShaderBindings);
+		FMeshMaterialShader::GetShaderBindings(Scene, FeatureLevel, PrimitiveSceneProxy, MaterialRenderProxy, Material, DrawRenderState, ShaderElementData, ShaderBindings);
 
 		ShaderBindings.Add(LayerId, ShaderElementData.LayerId);
 	}
diff --git a/Engine/Source/Runtime/Renderer/Private/VelocityRendering.cpp b/Engine/Source/Runtime/Renderer/Private/VelocityRendering.cpp
--- a/Engine/Source/Runtime/Renderer/Private/VelocityRendering.cpp
+++ b/Engine/Source/Runtime/Renderer/Private/VelocityRendering.cpp
@@ -688,6 +688,9 @@ void FVelocityMeshProcessor::Process(
 	FMeshMaterialShaderElementData ShaderElementData;
 	ShaderElementData.InitializeMeshMaterialData(ViewIfDynamicMeshCommand, PrimitiveSceneProxy, MeshBatch, StaticMeshId, false);
 
+	const bool bIsInstancedStereo = ViewIfDynamicMeshCommand ? ViewIfDynamicMeshCommand->IsInstancedStereoPass() : Scene->bStaticDrawInstancedStereo;
+	const int32 InstanceFactor = bIsInstancedStereo ? 2 : 1;
+
 	BuildMeshDrawCommands(
 		MeshBatch,
 		BatchElementMask,
@@ -698,7 +701,7 @@ void FVelocityMeshProcessor::Process(
 		VelocityPassShaders,
 		MeshFillMode,
 		MeshCullMode,
-		1,
+		InstanceFactor,
 		FMeshDrawCommandSortKey::Default,
 		EMeshPassFeatures::Default,
 		ShaderElementData);
@@ -709,6 +712,7 @@ FVelocityMeshProcessor::FVelocityMeshProcessor(const FScene* Scene, const FScene
 {
 	PassDrawRenderState = InPassDrawRenderState;
 	PassDrawRenderState.SetViewUniformBuffer(Scene->UniformBuffers.ViewUniformBuffer);
+	PassDrawRenderState.SetInstancedViewUniformBuffer(Scene->UniformBuffers.InstancedViewUniformBuffer);
 	PassDrawRenderState.SetPassUniformBuffer(Scene->UniformBuffers.VelocityPassUniformBuffer);
 }
 
diff --git a/Engine/Source/Runtime/Renderer/Private/VolumetricFogVoxelization.cpp b/Engine/Source/Runtime/Renderer/Private/VolumetricFogVoxelization.cpp
--- a/Engine/Source/Runtime/Renderer/Private/VolumetricFogVoxelization.cpp
+++ b/Engine/Source/Runtime/Renderer/Private/VolumetricFogVoxelization.cpp
@@ -178,7 +178,7 @@ class FVoxelizeVolumeVS : public FMeshMaterialShader
 	}
 
 public:
-	
+
 	virtual bool Serialize(FArchive& Ar)
 	{		
 		bool bShaderHasOutdatedParameters = FMeshMaterialShader::Serialize(Ar);
@@ -192,18 +192,17 @@ public:
 		const FPrimitiveSceneProxy* PrimitiveSceneProxy,
 		const FMaterialRenderProxy& MaterialRenderProxy,
 		const FMaterial& Material,
-		const TUniformBufferRef<FViewUniformShaderParameters>& ViewUniformBuffer,
-		FUniformBufferRHIParamRef PassUniformBufferValue,
+		const FMeshPassProcessorRenderState& DrawRenderState,
 		const FVoxelizeVolumeShaderElementData& ShaderElementData,
 		FMeshDrawSingleShaderBindings& ShaderBindings) const
 	{
-		FMeshMaterialShader::GetShaderBindings(Scene, FeatureLevel, PrimitiveSceneProxy, MaterialRenderProxy, Material, ViewUniformBuffer, PassUniformBufferValue, ShaderElementData, ShaderBindings);
+		FMeshMaterialShader::GetShaderBindings(Scene, FeatureLevel, PrimitiveSceneProxy, MaterialRenderProxy, Material, DrawRenderState, ShaderElementData, ShaderBindings);
 		if (!RHISupportsGeometryShaders(Scene->GetShaderPlatform()))
 		{
 			ShaderBindings.Add(VoxelizationPassIndex, ShaderElementData.VoxelizationPassIndex);
 		}
 	}
-	
+
 protected:
 
 	FShaderParameter VoxelizationPassIndex;
@@ -295,12 +294,11 @@ public:
 		const FPrimitiveSceneProxy* PrimitiveSceneProxy,
 		const FMaterialRenderProxy& MaterialRenderProxy,
 		const FMaterial& Material,
-		const TUniformBufferRef<FViewUniformShaderParameters>& ViewUniformBuffer,
-		FUniformBufferRHIParamRef PassUniformBufferValue,
+		const FMeshPassProcessorRenderState& DrawRenderState,
 		const FVoxelizeVolumeShaderElementData& ShaderElementData,
 		FMeshDrawSingleShaderBindings& ShaderBindings) const
 	{
-		FMeshMaterialShader::GetShaderBindings(Scene, FeatureLevel, PrimitiveSceneProxy, MaterialRenderProxy, Material, ViewUniformBuffer, PassUniformBufferValue, ShaderElementData, ShaderBindings);
+		FMeshMaterialShader::GetShaderBindings(Scene, FeatureLevel, PrimitiveSceneProxy, MaterialRenderProxy, Material, DrawRenderState, ShaderElementData, ShaderBindings);
 
 		ShaderBindings.Add(VoxelizationPassIndex, ShaderElementData.VoxelizationPassIndex);
 	}
diff --git a/Engine/Source/Runtime/Renderer/Public/MeshMaterialShader.h b/Engine/Source/Runtime/Renderer/Public/MeshMaterialShader.h
--- a/Engine/Source/Runtime/Renderer/Public/MeshMaterialShader.h
+++ b/Engine/Source/Runtime/Renderer/Public/MeshMaterialShader.h
@@ -62,8 +62,7 @@ public:
 		const FPrimitiveSceneProxy* PrimitiveSceneProxy,
 		const FMaterialRenderProxy& MaterialRenderProxy,
 		const FMaterial& Material,
-		const TUniformBufferRef<FViewUniformShaderParameters>& ViewUniformBuffer,
-		FUniformBufferRHIParamRef PassUniformBufferValue,
+		const FMeshPassProcessorRenderState& DrawRenderState,
 		const FMeshMaterialShaderElementData& ShaderElementData,
 		FMeshDrawSingleShaderBindings& ShaderBindings) const;
 
diff --git a/Engine/Source/Runtime/Renderer/Public/MeshPassProcessor.h b/Engine/Source/Runtime/Renderer/Public/MeshPassProcessor.h
--- a/Engine/Source/Runtime/Renderer/Public/MeshPassProcessor.h
+++ b/Engine/Source/Runtime/Renderer/Public/MeshPassProcessor.h
@@ -350,6 +350,7 @@ public:
 	uint32 FirstIndex;
 	uint32 NumPrimitives;
 	uint32 NumInstances;
+	uint32 InstanceFactor;
 	uint32 BaseVertexIndex;
 	uint32 NumVertices;
 	FVertexBufferRHIParamRef IndirectArgsBuffer;
@@ -378,6 +379,7 @@ public:
 			&& FirstIndex == Rhs.FirstIndex
 			&& NumPrimitives == Rhs.NumPrimitives
 			&& NumInstances == Rhs.NumInstances
+			&& InstanceFactor == Rhs.InstanceFactor
 			&& BaseVertexIndex == Rhs.BaseVertexIndex
 			&& NumVertices == Rhs.NumVertices
 			&& IndirectArgsBuffer == Rhs.IndirectArgsBuffer;
@@ -399,7 +401,7 @@ public:
 	}
 
 	/** Called when the mesh draw command is complete. */
-	RENDERER_API void SetDrawParametersAndFinalize(const FMeshBatch& MeshBatch, int32 BatchElementIndex, int32 InstanceFactor, bool bDoSetupPsoStateForRasterization);
+	RENDERER_API void SetDrawParametersAndFinalize(const FMeshBatch& MeshBatch, int32 BatchElementIndex, int32 BatchInstanceFactor, bool bDoSetupPsoStateForRasterization);
 
 	void Finalize(bool bDoSetupPsoStateForRasterization)
 	{
@@ -708,6 +710,7 @@ struct FMeshPassProcessorRenderState
 		, DepthStencilState(nullptr)
 		, DepthStencilAccess(FExclusiveDepthStencil::DepthRead_StencilRead)
 		, ViewUniformBuffer(SceneView.ViewUniformBuffer)
+		, InstancedViewUniformBuffer()
 		, PassUniformBuffer(InPassUniformBuffer)
 		, StencilRef(0)
 	{
@@ -718,6 +721,7 @@ struct FMeshPassProcessorRenderState
 		, DepthStencilState(nullptr)
 		, DepthStencilAccess(FExclusiveDepthStencil::DepthRead_StencilRead)
 		, ViewUniformBuffer(InViewUniformBuffer)
+		, InstancedViewUniformBuffer()
 		, PassUniformBuffer(InPassUniformBuffer)
 		, StencilRef(0)
 	{
@@ -727,6 +731,7 @@ struct FMeshPassProcessorRenderState
 		BlendState(nullptr)
 		, DepthStencilState(nullptr)
 		, ViewUniformBuffer()
+		, InstancedViewUniformBuffer()
 		, PassUniformBuffer(nullptr)
 		, StencilRef(0)
 	{
@@ -737,6 +742,7 @@ struct FMeshPassProcessorRenderState
 		, DepthStencilState(DrawRenderState.DepthStencilState)
 		, DepthStencilAccess(DrawRenderState.DepthStencilAccess)
 		, ViewUniformBuffer(DrawRenderState.ViewUniformBuffer)
+		, InstancedViewUniformBuffer(DrawRenderState.InstancedViewUniformBuffer)
 		, PassUniformBuffer(DrawRenderState.PassUniformBuffer)
 		, StencilRef(DrawRenderState.StencilRef)
 	{
@@ -793,6 +799,16 @@ public:
 		return ViewUniformBuffer;
 	}
 
+	FORCEINLINE_DEBUGGABLE void SetInstancedViewUniformBuffer(const TUniformBufferRef<FInstancedViewUniformShaderParameters>& InViewUniformBuffer)
+	{
+		InstancedViewUniformBuffer = InViewUniformBuffer;
+	}
+
+	FORCEINLINE_DEBUGGABLE const TUniformBufferRef<FInstancedViewUniformShaderParameters>& GetInstancedViewUniformBuffer() const
+	{
+		return InstancedViewUniformBuffer.IsValid() ? InstancedViewUniformBuffer : reinterpret_cast<const TUniformBufferRef<FInstancedViewUniformShaderParameters>&>(ViewUniformBuffer);
+	}
+
 	FORCEINLINE_DEBUGGABLE void SetPassUniformBuffer(FUniformBufferRHIParamRef InPassUniformBuffer)
 	{
 		PassUniformBuffer = InPassUniformBuffer;
@@ -820,6 +836,7 @@ private:
 	FExclusiveDepthStencil::Type	DepthStencilAccess;
 
 	TUniformBufferRef<FViewUniformShaderParameters>	ViewUniformBuffer;
+	TUniformBufferRef<FInstancedViewUniformShaderParameters> InstancedViewUniformBuffer;
 	FUniformBufferRHIParamRef		PassUniformBuffer;
 	uint32							StencilRef;
 };
diff --git a/Engine/Source/Runtime/Renderer/Public/MeshPassProcessor.inl b/Engine/Source/Runtime/Renderer/Public/MeshPassProcessor.inl
--- a/Engine/Source/Runtime/Renderer/Public/MeshPassProcessor.inl
+++ b/Engine/Source/Runtime/Renderer/Public/MeshPassProcessor.inl
@@ -64,31 +64,31 @@ void FMeshPassProcessor::BuildMeshDrawCommands(
 	if (PassShaders.VertexShader)
 	{
 		FMeshDrawSingleShaderBindings ShaderBindings = SharedMeshDrawCommand.ShaderBindings.GetSingleShaderBindings(SF_Vertex);
-		PassShaders.VertexShader->GetShaderBindings(Scene, FeatureLevel, PrimitiveSceneProxy, MaterialRenderProxy, MaterialResource, DrawRenderState.GetViewUniformBuffer(), DrawRenderState.GetPassUniformBuffer(), ShaderElementData, ShaderBindings);
+		PassShaders.VertexShader->GetShaderBindings(Scene, FeatureLevel, PrimitiveSceneProxy, MaterialRenderProxy, MaterialResource, DrawRenderState, ShaderElementData, ShaderBindings);
 	}
 
 	if (PassShaders.HullShader)
 	{
 		FMeshDrawSingleShaderBindings ShaderBindings = SharedMeshDrawCommand.ShaderBindings.GetSingleShaderBindings(SF_Hull);
-		PassShaders.HullShader->GetShaderBindings(Scene, FeatureLevel, PrimitiveSceneProxy, MaterialRenderProxy, MaterialResource, DrawRenderState.GetViewUniformBuffer(), DrawRenderState.GetPassUniformBuffer(), ShaderElementData, ShaderBindings);
+		PassShaders.HullShader->GetShaderBindings(Scene, FeatureLevel, PrimitiveSceneProxy, MaterialRenderProxy, MaterialResource, DrawRenderState, ShaderElementData, ShaderBindings);
 	}
 
 	if (PassShaders.DomainShader)
 	{
 		FMeshDrawSingleShaderBindings ShaderBindings = SharedMeshDrawCommand.ShaderBindings.GetSingleShaderBindings(SF_Domain);
-		PassShaders.DomainShader->GetShaderBindings(Scene, FeatureLevel, PrimitiveSceneProxy, MaterialRenderProxy, MaterialResource, DrawRenderState.GetViewUniformBuffer(), DrawRenderState.GetPassUniformBuffer(), ShaderElementData, ShaderBindings);
+		PassShaders.DomainShader->GetShaderBindings(Scene, FeatureLevel, PrimitiveSceneProxy, MaterialRenderProxy, MaterialResource, DrawRenderState, ShaderElementData, ShaderBindings);
 	}
 
 	if (PassShaders.PixelShader)
 	{
 		FMeshDrawSingleShaderBindings ShaderBindings = SharedMeshDrawCommand.ShaderBindings.GetSingleShaderBindings(SF_Pixel);
-		PassShaders.PixelShader->GetShaderBindings(Scene, FeatureLevel, PrimitiveSceneProxy, MaterialRenderProxy, MaterialResource, DrawRenderState.GetViewUniformBuffer(), DrawRenderState.GetPassUniformBuffer(), ShaderElementData, ShaderBindings);
+		PassShaders.PixelShader->GetShaderBindings(Scene, FeatureLevel, PrimitiveSceneProxy, MaterialRenderProxy, MaterialResource, DrawRenderState, ShaderElementData, ShaderBindings);
 	}
 
 	if (PassShaders.GeometryShader)
 	{
 		FMeshDrawSingleShaderBindings ShaderBindings = SharedMeshDrawCommand.ShaderBindings.GetSingleShaderBindings(SF_Geometry);
-		PassShaders.GeometryShader->GetShaderBindings(Scene, FeatureLevel, PrimitiveSceneProxy, MaterialRenderProxy, MaterialResource, DrawRenderState.GetViewUniformBuffer(), DrawRenderState.GetPassUniformBuffer(), ShaderElementData, ShaderBindings);
+		PassShaders.GeometryShader->GetShaderBindings(Scene, FeatureLevel, PrimitiveSceneProxy, MaterialRenderProxy, MaterialResource, DrawRenderState, ShaderElementData, ShaderBindings);
 	}
 
 	SharedMeshDrawCommand.SetDebugData(PrimitiveSceneProxy, &MaterialResource, &MaterialRenderProxy, PassShaders.GetUntypedShaders());
@@ -190,7 +190,7 @@ void FMeshPassProcessor::BuildRayTracingDrawCommands(
 	if (PassShaders.RayHitGroupShader)
 	{
 		FMeshDrawSingleShaderBindings ShaderBindings = SharedMeshDrawCommand.ShaderBindings.GetSingleShaderBindings(SF_RayHitGroup);
-		PassShaders.RayHitGroupShader->GetShaderBindings(Scene, FeatureLevel, PrimitiveSceneProxy, MaterialRenderProxy, MaterialResource, DrawRenderState.GetViewUniformBuffer(), DrawRenderState.GetPassUniformBuffer(), ShaderElementData, ShaderBindings);
+		PassShaders.RayHitGroupShader->GetShaderBindings(Scene, FeatureLevel, PrimitiveSceneProxy, MaterialRenderProxy, MaterialResource, DrawRenderState, ShaderElementData, ShaderBindings);
 	}
 #endif
 
