Disable auto-instancing on unsuported platforms

diff --git a/Engine/Source/Runtime/Renderer/Private/MeshPassProcessor.cpp b/Engine/Source/Runtime/Renderer/Private/MeshPassProcessor.cpp
--- a/Engine/Source/Runtime/Renderer/Private/MeshPassProcessor.cpp
+++ b/Engine/Source/Runtime/Renderer/Private/MeshPassProcessor.cpp
@@ -1010,8 +1010,9 @@ void FCachedPassMeshDrawListContext::FinalizeCommand(
 	MeshDrawCommand.SetDrawParametersAndFinalize(MeshBatch, BatchElementIndex, InstanceFactor, bDoSetupPsoStateForRasterization);
 
 	check(CommandInfo.CommandIndex != -1);
+	const bool bCanUseGPUScene = UseGPUScene(GMaxRHIShaderPlatform, GMaxRHIFeatureLevel);
 
-	//if (bDoSetupPsoStateForRasterization)
+	if (bCanUseGPUScene /* && bDoSetupPsoStateForRasterization*/)
 	{
 		FSetElementId SetId = Scene.CachedMeshDrawCommandStateBuckets.FindId(MeshDrawCommand);
 
diff --git a/Engine/Source/Runtime/Renderer/Private/PrimitiveSceneInfo.cpp b/Engine/Source/Runtime/Renderer/Private/PrimitiveSceneInfo.cpp
--- a/Engine/Source/Runtime/Renderer/Private/PrimitiveSceneInfo.cpp
+++ b/Engine/Source/Runtime/Renderer/Private/PrimitiveSceneInfo.cpp
@@ -264,25 +264,30 @@ void FPrimitiveSceneInfo::CacheMeshDrawCommands(FRHICommandListImmediate& RHICmd
 
 void FPrimitiveSceneInfo::RemoveCachedMeshDrawCommands()
 {
+	const bool bCanUseGPUScene = UseGPUScene(GMaxRHIShaderPlatform, GMaxRHIFeatureLevel);
+	
 	for (int32 CommandIndex = 0; CommandIndex < StaticMeshCommandInfos.Num(); ++CommandIndex)
 	{
 		const FCachedMeshDrawCommandInfo& CachedCommand = StaticMeshCommandInfos[CommandIndex];
 		if (CachedCommand.CommandIndex != -1)
 		{
 			FCachedPassMeshDrawList& PassDrawList = Scene->CachedDrawLists[CachedCommand.MeshPass];
-
-			const FSetElementId StateBucketId = FSetElementId::FromInteger(CachedCommand.StateBucketId);
-			checkSlow(StateBucketId.IsValidId());
-			FMeshDrawCommandStateBucket& StateBucket = Scene->CachedMeshDrawCommandStateBuckets[StateBucketId];
-			if (CachedCommand.StateBucketId != -1)
+			
+			if (bCanUseGPUScene)
 			{
-				if (StateBucket.Num == 1)
-				{
-					Scene->CachedMeshDrawCommandStateBuckets.Remove(StateBucketId);
-				}
-				else
+				const FSetElementId StateBucketId = FSetElementId::FromInteger(CachedCommand.StateBucketId);
+				checkSlow(StateBucketId.IsValidId());
+				FMeshDrawCommandStateBucket& StateBucket = Scene->CachedMeshDrawCommandStateBuckets[StateBucketId];
+				if (CachedCommand.StateBucketId != -1)
 				{
-					StateBucket.Num--;
+					if (StateBucket.Num == 1)
+					{
+						Scene->CachedMeshDrawCommandStateBuckets.Remove(StateBucketId);
+					}
+					else
+					{
+						StateBucket.Num--;
+					}
 				}
 			}
 
diff --git a/Engine/Source/Runtime/Renderer/Private/SceneRendering.cpp b/Engine/Source/Runtime/Renderer/Private/SceneRendering.cpp
--- a/Engine/Source/Runtime/Renderer/Private/SceneRendering.cpp
+++ b/Engine/Source/Runtime/Renderer/Private/SceneRendering.cpp
@@ -73,7 +73,8 @@ static TAutoConsoleVariable<int32> CVarMeshDrawCommandsDynamicInstancing(
 
 bool IsDynamicInstancingEnabled()
 {
-	return CVarMeshDrawCommandsDynamicInstancing.GetValueOnRenderThread() > 0;
+	return CVarMeshDrawCommandsDynamicInstancing.GetValueOnRenderThread() > 0
+		&& UseGPUScene(GMaxRHIShaderPlatform, GMaxRHIFeatureLevel); // always disabled on platforms that do not support auto-instacing
 }
 
 int32 GDumpInstancingStats = 0;
