Fixed MaxInstanceFactor calculation. It can't be based on first mesh draw command inside DispatchPassSetup, as at this point this array may be empty (as dynamic commands are generated later).

diff --git a/Engine/Source/Runtime/Renderer/Private/MeshDrawCommands.cpp b/Engine/Source/Runtime/Renderer/Private/MeshDrawCommands.cpp
--- a/Engine/Source/Runtime/Renderer/Private/MeshDrawCommands.cpp
+++ b/Engine/Source/Runtime/Renderer/Private/MeshDrawCommands.cpp
@@ -801,8 +801,8 @@ void FParallelMeshDrawCommandPass::DispatchPassSetup(
 	TaskContext.ViewMatrix = View.ViewMatrices.GetViewMatrix();
 	TaskContext.PrimitiveBounds = &Scene->PrimitiveBounds;
 
-	//@todo - make the instance factor a pass parameter instead of a draw command parameter
-	TaskContext.MaxInstanceFactor = InOutMeshDrawCommands[0].MeshDrawCommand->InstanceFactor;
+	const bool bIsInstancedStereo = View.IsInstancedStereoPass() && Scene->bStaticDrawInstancedStereo;
+	TaskContext.MaxInstanceFactor = bIsInstancedStereo ? 2 : 1;
 
 	switch (PassType)
 	{
